Change PPLL to use a visibility buffer raster.

diff --git a/Engine/Shaders/Private/HairStrands/HairStrandsVisibilityCommon.ush b/Engine/Shaders/Private/HairStrands/HairStrandsVisibilityCommon.ush
--- a/Engine/Shaders/Private/HairStrands/HairStrandsVisibilityCommon.ush
+++ b/Engine/Shaders/Private/HairStrands/HairStrandsVisibilityCommon.ush
@@ -184,6 +184,11 @@ uint GetInvalidHairPrimitiveId()
 	return 0xFFFFFFFF;
 }
 
+void PatchPackedHairVisCoverage(inout FPackedHairVis Out, uint Coverage8bit)
+{
+	Out.Depth_Coverage8bit = (Out.Depth_Coverage8bit & (0x00FFFFFFu << 8)) | (Coverage8bit & 0xFFu);
+}
+
 struct FHairVis
 {
 	float	Depth;
@@ -372,70 +377,40 @@ float TransmittanceToCoverage(float InTransmittance, float InCoverageThreashold)
 ////////////////////////////////////////////////////////////////////////////////////////////////////////////////////////////////
 // Hair PPLL
 
-uint GetCoverage8Bits(FPPLLNodeData In)
-{
-	return (In.Tangent_Coverage8bit >> 24) & 0xFF;
-}
-
-uint GetCoverage16Bits(FPPLLNodeData In)
-{
-	return (In.Specular_LightChannel_Backlit_Coverage16bit >> 16) & 0xFFFF;
-}
-
-FPPLLNodeData CreatePPLLNodeData(
+FPackedHairVisPPLL PackHairVisPPLL(
 	float	Depth, 
+	float	Coverage,
 	uint	PrimitiveId,
-	uint	MacroGroupId,
-	float3	Tangent,
-	float	CoverageFullPrecision, 
-	uint	Coverage8bit,
-	float3	BaseColor,
-	float	Roughness, 
-	float	Specular, 
-	float	Backlit, 
-	uint	LightChannelMask, 
-	float2	Velocity, 
+	uint	MaterialId,
 	uint	NextNodeIndex)
 {
-	FHairSample S;
+	FHairVis S;
 	S.Depth = Depth;
+	S.Coverage8bit = To8bitCoverage(Coverage); // To16bitCoverage(CoverageFullPrecision)
 	S.PrimitiveId = PrimitiveId;
-	S.MacroGroupId = MacroGroupId;
-	S.Tangent = Tangent;
-	S.Coverage8bit = Coverage8bit;
-	S.BaseColor = BaseColor;
-	S.Roughness = Roughness;
-	S.Specular = 0;			// Custom encoding for PPLL
-	S.Backlit = 0;			// Custom encoding for PPLL
-	S.LightChannelMask = 0; // Custom encoding for PPLL
-	
-	const FPackedHairSample PackedS = PackHairSample(S);
-	const uint Coverage16bit = To16bitCoverage(CoverageFullPrecision);
-
-	FPPLLNodeData Out;
-	Out.Depth = PackedS.Depth;
-	Out.PrimitiveID_MacroGroupID = PackedS.PrimitiveID_MacroGroupID;
-	Out.Tangent_Coverage8bit = PackedS.Tangent_Coverage8bit;
-	Out.BaseColor_Roughness = PackedS.BaseColor_Roughness;
-	Out.Specular_LightChannel_Backlit_Coverage16bit = clamp(uint(Specular * 0xFF), 0, 0xFF) |  (clamp(uint(Backlit * 0x1F), 0, 0x1F)<<11) | ((LightChannelMask & 0x7)<<8) |  ((0xFFFF & Coverage16bit) << 16);
-	Out.NextNodeIndex = NextNodeIndex;
-	Out.PackedVelocity = PackVelocity(Velocity);
+	S.MaterialId = MaterialId;
+	const FPackedHairVis PackedS = PackHairVis(S);
 
+	FPackedHairVisPPLL Out;
+	Out.Depth_Coverage8bit = PackedS.Depth_Coverage8bit;
+	Out.PrimitiveID_MaterialID = PackedS.PrimitiveID_MaterialID;
+	Out.NextNodeIndex = NextNodeIndex;
 	return Out;
 }
 
-FPackedHairSample ConvertToPackedHairSample(FPPLLNodeData In, uint InNewCoverage8bit)
+FHairVis UnpackHairVisPPLL(FPackedHairVisPPLL In)
 {
-	const uint Specular_LightChannel = In.Specular_LightChannel_Backlit_Coverage16bit & 0x7FF;
-	const float Backlit = float((In.Specular_LightChannel_Backlit_Coverage16bit >> 11) & 0x1F) / float(0x1F);
-
-	FPackedHairSample Out;
-	Out.Depth = In.Depth;
-	Out.PrimitiveID_MacroGroupID = In.PrimitiveID_MacroGroupID;
-	Out.Tangent_Coverage8bit = (In.Tangent_Coverage8bit & 0x00FFFFFF) | ((InNewCoverage8bit & 0xFF)<<24);
-	Out.BaseColor_Roughness = In.BaseColor_Roughness;
-	Out.Specular_LightChannels_Backlit = Specular_LightChannel | (clamp(uint(Backlit * 0xFF), 0, 0xFF) << 16);
+	FPackedHairVis Temp;
+	Temp.Depth_Coverage8bit = In.Depth_Coverage8bit;
+	Temp.PrimitiveID_MaterialID = In.PrimitiveID_MaterialID;
+	return UnpackHairVis(Temp);
+}
 
+FPackedHairVis ConvertToPackedHairVis(FPackedHairVisPPLL In)
+{
+	FPackedHairVis Out;
+	Out.Depth_Coverage8bit = In.Depth_Coverage8bit;
+	Out.PrimitiveID_MaterialID = In.PrimitiveID_MaterialID;
 	return Out;
 }
 
diff --git a/Engine/Shaders/Private/HairStrands/HairStrandsVisibilityCommonStruct.ush b/Engine/Shaders/Private/HairStrands/HairStrandsVisibilityCommonStruct.ush
--- a/Engine/Shaders/Private/HairStrands/HairStrandsVisibilityCommonStruct.ush
+++ b/Engine/Shaders/Private/HairStrands/HairStrandsVisibilityCommonStruct.ush
@@ -2,18 +2,6 @@
 
 #pragma once
 
-// Footprint: 224bits | 28bytes
-struct FPPLLNodeData
-{
-	float Depth;
-	uint Specular_LightChannel_Backlit_Coverage16bit;
-	uint PrimitiveID_MacroGroupID;
-	uint Tangent_Coverage8bit;
-	uint BaseColor_Roughness;
-	uint NextNodeIndex;
-	uint PackedVelocity;
-};
-
 struct FPackedHairSample
 {
 	float Depth;
@@ -27,4 +15,11 @@ struct FPackedHairVis
 {
 	uint Depth_Coverage8bit;
 	uint PrimitiveID_MaterialID;
+};
+
+struct FPackedHairVisPPLL
+{
+	uint Depth_Coverage8bit;
+	uint PrimitiveID_MaterialID;
+	uint NextNodeIndex;
 };
\ No newline at end of file
diff --git a/Engine/Shaders/Private/HairStrands/HairStrandsVisibilityCompaction.usf b/Engine/Shaders/Private/HairStrands/HairStrandsVisibilityCompaction.usf
--- a/Engine/Shaders/Private/HairStrands/HairStrandsVisibilityCompaction.usf
+++ b/Engine/Shaders/Private/HairStrands/HairStrandsVisibilityCompaction.usf
@@ -27,26 +27,10 @@ uint VelocityType;
 
 RWTexture2D<uint> OutCompactNodeCounter;
 RWTexture2D<uint> OutCompactNodeIndex;
-#if PERMUTATION_PPLL == 0
-	RWStructuredBuffer<FPackedHairVis> OutCompactNodeVis;
-#else
-	RWStructuredBuffer<FPackedHairSample> OutCompactNodeData;
-#endif
+RWStructuredBuffer<FPackedHairVis> OutCompactNodeVis;
 RWBuffer<uint2> OutCompactNodeCoord;
 RWTexture2D<float> OutCoverageTexture;
 
-#ifndef PERMUTATION_VELOCITY
-	#error Velocity permutation needs to be defined
-#endif
-
-#define VELOCITY_TYPE_NONE 0
-#define VELOCITY_TYPE_AVG 1  
-#define VELOCITY_TYPE_CLOSEST 2
-#define VELOCITY_TYPE_MAX 3
-#if PERMUTATION_VELOCITY != VELOCITY_TYPE_NONE
-	RWTexture2D<float2> OutVelocityTexture;
-#endif
-
 groupshared uint AllocationNodeCount;
 groupshared uint AllocationNodeOffset;
 
@@ -67,27 +51,15 @@ struct FSampleSetDesc
 
 Texture2D<uint> PPLLCounter;
 Texture2D<uint> PPLLNodeIndex;
-StructuredBuffer<FPPLLNodeData> PPLLNodeData;
+StructuredBuffer<FPackedHairVisPPLL> PPLLNodeData;
 
 
 [numthreads(TILE_PIXEL_SIZE_X, TILE_PIXEL_SIZE_Y, 1)]
 void MainCS(uint GroupIndex : SV_GroupIndex, uint3 DispatchThreadId : SV_DispatchThreadID)
 {
-	uint2 PixelCoord = DispatchThreadId.xy + ResolutionOffset;
+	const uint2 PixelCoord = DispatchThreadId.xy + ResolutionOffset;
 	const uint FirstNodeIndex = PPLLNodeIndex[PixelCoord];
-
-	float ClosestDepth = 0; // Inverse-Z
 	float PixelCoverage = 0;
-
-#if PERMUTATION_VELOCITY != VELOCITY_TYPE_NONE
-	float2 AverageVelocity = 0;
-	uint ValidVelocityCount = 0;
-
-	float2 ClosestVelocity = 0;
-
-	float2 MaxVelocity = 0;
-	float MaxVelocityMagnitude2 = 0;
-#endif
 	
 	const bool bExecute = FirstNodeIndex != 0xFFFFFFFF && all(PixelCoord < uint2(OutputResolution));
 
@@ -107,16 +79,15 @@ void MainCS(uint GroupIndex : SV_GroupIndex, uint3 DispatchThreadId : SV_Dispatc
 		uint NodeIndex = FirstNodeIndex;
 		while (NodeIndex != 0xFFFFFFFF)
 		{
-			const FPPLLNodeData PPLLNode = PPLLNodeData[NodeIndex];
-			const uint Coverage8bit		 = GetCoverage8Bits(PPLLNode);
-			const float Coverage		 = From8bitCoverage(Coverage8bit);
+			const FPackedHairVisPPLL PackedHairVis = PPLLNodeData[NodeIndex];
+			const FHairVis HairVis = UnpackHairVisPPLL(PackedHairVis);
 
-			TotalTransmittance *= 1.0f - Coverage;
+			TotalTransmittance *= 1.0f - HairVis.Coverage;
 
 			int InsertPos;
 			for (InsertPos = 0; InsertPos < PPLL_MAX_RENDER_NODE_COUNT; ++InsertPos)
 			{
-				if (PPLLNode.Depth > SortedDepth[InsertPos])
+				if (HairVis.Depth > SortedDepth[InsertPos])
 				{
 					// The new node is closer
 					break;
@@ -141,12 +112,12 @@ void MainCS(uint GroupIndex : SV_GroupIndex, uint3 DispatchThreadId : SV_Dispatc
 
 				// Insert new node
 				SortedIndex[InsertPos] = NodeIndex;
-				SortedDepth[InsertPos] = PPLLNode.Depth;
+				SortedDepth[InsertPos] = HairVis.Depth;
 
 				PixelNodeCount++;
 			}
 
-			NodeIndex = PPLLNode.NextNodeIndex;
+			NodeIndex = PackedHairVis.NextNodeIndex;
 		}
 	}
 
@@ -170,7 +141,7 @@ void MainCS(uint GroupIndex : SV_GroupIndex, uint3 DispatchThreadId : SV_Dispatc
 	GroupMemoryBarrierWithGroupSync();
 	uint OffsetInStorage = AllocationNodeOffset + PixelDataSubOffsetInStorage;
 
-	if(bExecute)
+	if (bExecute)
 	{
 		//////////
 		// Second pass: compute total coverage for validated nodes we are going to render
@@ -180,10 +151,10 @@ void MainCS(uint GroupIndex : SV_GroupIndex, uint3 DispatchThreadId : SV_Dispatc
 		for (uint i = 0; i < PixelRenderNodecount; ++i)
 		{
 			const uint NodePPLLIndex	 = SortedIndex[i];
-			const FPPLLNodeData PPLLNode = PPLLNodeData[NodePPLLIndex];
-			const uint Coverage16bit	 = GetCoverage16Bits(PPLLNode);
 
-			const float Coverage = From16bitCoverage(Coverage16bit);
+			const FPackedHairVisPPLL PackedHairVis = PPLLNodeData[NodePPLLIndex];
+			const FHairVis HairVis = UnpackHairVisPPLL(PackedHairVis);
+			const float Coverage = HairVis.Coverage; // This should be 16bits, instead of 8bits for better precision
 
 			// Update current node coverage as a function of previous nodes coverage
 			SortedCoverage[i] = TotalSortedTransmittance * Coverage;
@@ -200,68 +171,20 @@ void MainCS(uint GroupIndex : SV_GroupIndex, uint3 DispatchThreadId : SV_Dispatc
 		for (uint j = 0; j < PixelRenderNodecount; ++j)
 		{
 			const uint NodePPLLIndex = SortedIndex[j];
-			uint Coverage16bit = 0;
-			FPPLLNodeData PPLLNode = PPLLNodeData[NodePPLLIndex];
+			FPackedHairVis PackedHairVis = ConvertToPackedHairVis(PPLLNodeData[NodePPLLIndex]);
 
 			// Coverage8bit is a weight normalising to 1 the contribution of all the compacted samples. Because later it is weighted by PixelCoverage.
 			// Patch the coverage on the out node
 			const float PatchedCoverage8bit = To8bitCoverage(SortedCoverage[j] / float(ValidPixelSampleTotalCoverage));
-			FPackedHairSample OutNodeData = ConvertToPackedHairSample(PPLLNode, PatchedCoverage8bit);
-
-			ClosestDepth = max(ClosestDepth, OutNodeData.Depth); // Inverse-Z
-
+			PatchPackedHairVisCoverage(PackedHairVis, PatchedCoverage8bit);
+			
 			// TODOAGREGATE if last, Create FPackedHairSample from aggregated  data
-			OutCompactNodeData[OffsetInStorage + j] = OutNodeData;
+			OutCompactNodeVis[OffsetInStorage + j] = PackedHairVis;
 			OutCompactNodeCoord[OffsetInStorage + j] = PixelCoord;
 
-
-		#if PERMUTATION_VELOCITY != VELOCITY_TYPE_NONE
-			const float2 Velocity = UnpackVelocity(PPLLNode.PackedVelocity);
-		
-			if (VelocityType == VELOCITY_TYPE_AVG)
-			{
-				AverageVelocity += Velocity;
-				++ValidVelocityCount;
-			}
-			else if (VelocityType == VELOCITY_TYPE_CLOSEST)
-			{
-				if (ClosestDepth < OutNodeData.Depth) // Inverse-Z
-				{
-					ClosestVelocity = Velocity;
-					ClosestDepth = OutNodeData.Depth;
-				}
-			}
-			else if (VelocityType == VELOCITY_TYPE_MAX)
-			{
-				const float2 CurrentVelocity = Velocity;
-				const float CurrentVelocityMagnitude2 = dot(CurrentVelocity, CurrentVelocity);
-				if (CurrentVelocityMagnitude2 > MaxVelocityMagnitude2)
-				{
-					MaxVelocity = Velocity;
-					MaxVelocityMagnitude2 = CurrentVelocityMagnitude2;
-				}
-			}
-		#endif		
 		}
 
-		PixelCoverage = TransmittanceToCoverage(TotalTransmittance, CoverageThreshold);
-
-	#if PERMUTATION_VELOCITY != VELOCITY_TYPE_NONE
-		float2 VelocityOutput = 0;
-		if (VelocityType == VELOCITY_TYPE_AVG)
-		{
-			VelocityOutput = AverageVelocity / max(ValidVelocityCount, 1u);
-		}
-		else if (VelocityType == VELOCITY_TYPE_CLOSEST)
-		{
-			VelocityOutput = ClosestVelocity;
-		}
-		else if (VelocityType == VELOCITY_TYPE_MAX)
-		{
-			VelocityOutput = MaxVelocity;
-		}
-		OutVelocityTexture[PixelCoord] = EncodeVelocityToTexture(float3(VelocityOutput, 0.0)).xy;
-	#endif		
+		PixelCoverage = TransmittanceToCoverage(TotalTransmittance, CoverageThreshold);	
 	}
 
 	FNodeDesc NodeDesc;
diff --git a/Engine/Shaders/Private/HairStrands/HairStrandsVisibilityPPLLDebug.usf b/Engine/Shaders/Private/HairStrands/HairStrandsVisibilityPPLLDebug.usf
--- a/Engine/Shaders/Private/HairStrands/HairStrandsVisibilityPPLLDebug.usf
+++ b/Engine/Shaders/Private/HairStrands/HairStrandsVisibilityPPLLDebug.usf
@@ -9,10 +9,16 @@ float PPLLMaxTotalListElementCount;
 
 Texture2D<uint> PPLLCounter;
 Texture2D<uint> PPLLNodeIndex;
-StructuredBuffer<FPPLLNodeData> PPLLNodeData;
+StructuredBuffer<FPackedHairVisPPLL> PPLLNodeData;
 
 RWTexture2D<float4> SceneColorTextureUAV;
 
+FFontColor GetOccupancyColor(float In)
+{
+	float3 Color = lerp(float3(0, 1, 0), float3(1, 0, 0), saturate(In));
+	return InitFontColor(Color);
+}
+
 [numthreads(8, 8, 1)]
 void VisibilityDebugPPLLCS(uint GroupIndex : SV_GroupIndex, uint3 DispatchThreadId : SV_DispatchThreadID)
 {
@@ -26,7 +32,7 @@ void VisibilityDebugPPLLCS(uint GroupIndex : SV_GroupIndex, uint3 DispatchThread
 		do
 		{
 			StrandCount++;
-			FPPLLNodeData Node = PPLLNodeData[NextNodeIndex];
+			FPackedHairVisPPLL Node = PPLLNodeData[NextNodeIndex];
 			NextNodeIndex = Node.NextNodeIndex;
 		} while (NextNodeIndex != 0xFFFFFFFF);
 
@@ -61,76 +67,29 @@ void VisibilityDebugPPLLCS(uint GroupIndex : SV_GroupIndex, uint3 DispatchThread
 	{
 		ShaderPrintFilter(true);
 
+		const float PPLLNodeSizeInMb = (4 + 4 + 4) / 1000000.f;
+		const float NodePercentage = float(PPLLListElementcount) / float(PPLLMaxTotalListElementCount);
+
 		float2 StartPos = (BarOrigin + BarSize*float2(0.0f, 2.0f)) * View.ViewSizeAndInvSize.zw;
 		float2 Pos = StartPos;
+		Pos = Print(Pos, TEXT("Mean Node Per Pixel : "));
+		Pos = ShaderPrint(Pos, PPLLMeanListElementCountPerPixel, FontEmerald);
+		Pos = ShaderPrintNewline(Pos); Pos.x = StartPos.x;
+		
+		Pos = Print(Pos, TEXT("Nodes Used (%)      : "));
+		Pos = ShaderPrint(Pos, NodePercentage * 100.f, GetOccupancyColor(NodePercentage));
+		Pos = ShaderPrintNewline(Pos); Pos.x = StartPos.x;
 
-		Pos = ShaderPrintSymbol(Pos, _M_);
-		Pos = ShaderPrintSymbol(Pos, _E_);
-		Pos = ShaderPrintSymbol(Pos, _A_);
-		Pos = ShaderPrintSymbol(Pos, _N_);
-		Pos = ShaderPrintSymbol(Pos, _SPC_);
-		Pos = ShaderPrintSymbol(Pos, _P_);
-		Pos = ShaderPrintSymbol(Pos, _P_);
-		Pos = ShaderPrintSymbol(Pos, _L_);
-		Pos = ShaderPrintSymbol(Pos, _L_);
-		Pos = ShaderPrintSymbol(Pos, _SPC_);
-		Pos = ShaderPrintSymbol(Pos, _N_);
-		Pos = ShaderPrintSymbol(Pos, _O_);
-		Pos = ShaderPrintSymbol(Pos, _D_);
-		Pos = ShaderPrintSymbol(Pos, _E_);
-		Pos = ShaderPrintSymbol(Pos, _SPC_);
-		Pos = ShaderPrintSymbol(Pos, _P_);
-		Pos = ShaderPrintSymbol(Pos, _E_);
-		Pos = ShaderPrintSymbol(Pos, _R_);
-		Pos = ShaderPrintSymbol(Pos, _SPC_);
-		Pos = ShaderPrintSymbol(Pos, _P_);
-		Pos = ShaderPrintSymbol(Pos, _I_);
-		Pos = ShaderPrintSymbol(Pos, _X_);
-		Pos = ShaderPrintSymbol(Pos, _E_);
-		Pos = ShaderPrintSymbol(Pos, _L_);
-		Pos = ShaderPrintSymbol(Pos, _SPC_);
-		Pos = ShaderPrint(Pos, PPLLMeanListElementCountPerPixel);
+		Pos = Print(Pos, TEXT("Nodes Used          : "));
+		Pos = ShaderPrint(Pos, PPLLListElementcount, FontSilver);
 		Pos = ShaderPrintNewline(Pos); Pos.x = StartPos.x;
 
-		Pos = ShaderPrintSymbol(Pos, _P_);
-		Pos = ShaderPrintSymbol(Pos, _P_);
-		Pos = ShaderPrintSymbol(Pos, _L_);
-		Pos = ShaderPrintSymbol(Pos, _L_);
-		Pos = ShaderPrintSymbol(Pos, _SPC_);
-		Pos = ShaderPrintSymbol(Pos, _N_);
-		Pos = ShaderPrintSymbol(Pos, _O_);
-		Pos = ShaderPrintSymbol(Pos, _D_);
-		Pos = ShaderPrintSymbol(Pos, _E_);
-		Pos = ShaderPrintSymbol(Pos, _S_);
-		Pos = ShaderPrintSymbol(Pos, _SPC_);
-		Pos = ShaderPrintSymbol(Pos, _U_);
-		Pos = ShaderPrintSymbol(Pos, _S_);
-		Pos = ShaderPrintSymbol(Pos, _E_);
-		Pos = ShaderPrintSymbol(Pos, _D_);
-		Pos = ShaderPrintSymbol(Pos, _SPC_);
-		Pos = ShaderPrintSymbol(Pos, _SPC_);
-		Pos = ShaderPrint(Pos, PPLLListElementcount);
+		Pos = Print(Pos, TEXT("Allocated Nodes (Mb): "));
+		Pos = ShaderPrint(Pos, PPLLMaxTotalListElementCount * PPLLNodeSizeInMb, FontSilver);
 		Pos = ShaderPrintNewline(Pos); Pos.x = StartPos.x;
 
-		
-		Pos = ShaderPrintSymbol(Pos, _SPC_);
-		Pos = ShaderPrintSymbol(Pos, _SPC_);
-		Pos = ShaderPrintSymbol(Pos, _SPC_);
-		Pos = ShaderPrintSymbol(Pos, _SPC_);
-		Pos = ShaderPrintSymbol(Pos, _SPC_);
-		Pos = ShaderPrintSymbol(Pos, _SPC_);
-		Pos = ShaderPrintSymbol(Pos, _SPC_);
-		Pos = ShaderPrintSymbol(Pos, _SPC_);
-		Pos = ShaderPrintSymbol(Pos, _SPC_);
-		Pos = ShaderPrintSymbol(Pos, _SPC_);
-		Pos = ShaderPrintSymbol(Pos, _SPC_);
-		Pos = ShaderPrintSymbol(Pos, _T_);
-		Pos = ShaderPrintSymbol(Pos, _O_);
-		Pos = ShaderPrintSymbol(Pos, _T_);
-		Pos = ShaderPrintSymbol(Pos, _A_);
-		Pos = ShaderPrintSymbol(Pos, _L_);
-		Pos = ShaderPrintSymbol(Pos, _SPC_);
-		Pos = ShaderPrint(Pos, PPLLMaxTotalListElementCount);
-		Pos = ShaderPrintNewline(Pos);
+		Pos = Print(Pos, TEXT("Total               : "));
+		Pos = ShaderPrint(Pos, PPLLMaxTotalListElementCount, FontOrange);
+		Pos = ShaderPrintNewline(Pos); Pos.x = StartPos.x;
 	}
 }
diff --git a/Engine/Shaders/Private/HairStrands/HairStrandsVisibilityPS.usf b/Engine/Shaders/Private/HairStrands/HairStrandsVisibilityPS.usf
--- a/Engine/Shaders/Private/HairStrands/HairStrandsVisibilityPS.usf
+++ b/Engine/Shaders/Private/HairStrands/HairStrandsVisibilityPS.usf
@@ -44,11 +44,7 @@ void MainVisibility(
 	centroid FVertexFactoryInterpolantsVSToPS Interpolants
 #endif
 	, in float4 SvPosition	: SV_Position
-#if HAIR_RENDER_MODE == RENDER_MODE_PPLL
-	, centroid in float WorldStrandRadius : TEXCOORD8
-	, const float4 PackedVelocityA : TEXCOORD9
-	, const float4 PackedVelocityC : TEXCOORD12
-#elif HAIR_RENDER_MODE == RENDER_MODE_TRANSMITTANCE || HAIR_RENDER_MODE == RENDER_MODE_TRANSMITTANCE_AND_HAIRCOUNT
+#if HAIR_RENDER_MODE == RENDER_MODE_PPLL || HAIR_RENDER_MODE == RENDER_MODE_TRANSMITTANCE || HAIR_RENDER_MODE == RENDER_MODE_TRANSMITTANCE_AND_HAIRCOUNT
 	, centroid in float WorldStrandRadius : TEXCOORD8
 #endif
 	  OPTIONAL_IsFrontFace
@@ -96,68 +92,11 @@ void MainVisibility(
 	OutColor1 = float2(Coverage, 1);
 #endif
 
-#if HAIR_RENDER_MODE == RENDER_MODE_PPLL
-  #if HAIR_STRAND_MESH_FACTORY
-	float3 Tangent = normalize(Interpolants.TangentToWorld2.xyz);
-	const uint PrimitiveId = Interpolants.HairPrimitiveId;
-	const bool bUseScatterSceneLighting = UseScatterSceneLighting();
-  #else
-	float3 Tangent = float3(0, 0, 1);
-	const uint PrimitiveId = 0;
-	const bool bUseScatterSceneLighting = false;
-  #endif
-	const uint MacroGroupId = HairVisibilityPass_HairMacroGroupIndex;
-	const uint Coverage8bit = To8bitCoverage(Coverage);
-
-	half3 BaseColor = 0;
-	float Roughness = 0;
-	float Specular  = 0;
-	float Backlit	= 1;
-	float2 Velocity = 0;
-	float2 EncodedVelocity = 0;
-	#if SUPPORT_MATERIAL_PROPERTY
-	{
-		FMaterialPixelParameters MaterialParameters = GetMaterialPixelParameters(Interpolants, SvPosition);
-		FPixelMaterialInputs PixelMaterialInputs;
-	
-		float4 ScreenPosition = SvPositionToResolvedScreenPosition(SvPosition);
-		float3 TranslatedWorldPosition = SvPositionToResolvedTranslatedWorld(SvPosition);
-		CalcMaterialParametersEx(MaterialParameters, PixelMaterialInputs, SvPosition, ScreenPosition, bIsFrontFace, TranslatedWorldPosition, TranslatedWorldPosition);
-
-#if STRATA_ENABLED && MATERIAL_IS_STRATA
-		FStrataBSDF HairBSDF = PixelMaterialInputs.FrontMaterial.InlinedBSDF;
-
-		BaseColor		 = HAIR_BASECOLOR(HairBSDF);
-		Roughness		 = HAIR_ROUGHNESS(HairBSDF);
-		half  Metallic	 = HAIR_SCATTER(HairBSDF);
-		Specular		 = HAIR_SPECULAR(HairBSDF);
-		Backlit			 = HAIR_BACKLIT(HairBSDF);
-		float MaterialAO = GetMaterialAmbientOcclusion(PixelMaterialInputs);
-
-		Tangent = MaterialParameters.SharedLocalBases.Normals[BSDF_GETSHAREDLOCALBASISID(HairBSDF)];	// Strata normals are all world space already
-#else
-		BaseColor		 = GetMaterialBaseColor(PixelMaterialInputs);
-		Roughness		 = GetMaterialRoughness(PixelMaterialInputs);
-		half  Metallic	 = GetMaterialMetallic(PixelMaterialInputs);
-		Specular		 = GetMaterialSpecular(PixelMaterialInputs);
-		Backlit			 = GetMaterialCustomData0(MaterialParameters);
-		float MaterialAO = GetMaterialAmbientOcclusion(PixelMaterialInputs);
-
-		Tangent = mul(normalize(GetMaterialNormalRaw(PixelMaterialInputs)), MaterialParameters.TangentToWorld);
-#endif
-	}
-	{
-		Velocity = Calculate3DVelocity(PackedVelocityA, PackedVelocityC).xy;
-		EncodedVelocity = EncodeVelocityToTexture(float3(Velocity, 0.0)).xy;
-	}
-	#endif
-#endif
-
-#if HAIR_RENDER_MODE == RENDER_MODE_PPLL
+#if HAIR_RENDER_MODE == RENDER_MODE_PPLL && HAIR_STRAND_MESH_FACTORY
+	const uint HairPrimitiveId = Interpolants.HairPrimitiveId;
 	const uint2 PixelCoord = SvPosition.xy;
-
-	float2 NDC = SvPosition.xy / SvPosition.w;
-	float2 ScreenUV = NDC * ResolvedView.ScreenPositionScaleBias.xy + ResolvedView.ScreenPositionScaleBias.wz;
+	const float2 NDC = SvPosition.xy / SvPosition.w;
+	const float2 ScreenUV = NDC * ResolvedView.ScreenPositionScaleBias.xy + ResolvedView.ScreenPositionScaleBias.wz;
 
 	// Allocate space for a new node
 	uint NodeIndex;
@@ -168,19 +107,11 @@ void MainVisibility(
 		// If we can append new node, add it and make it point to the node previously in head of list
 		uint NextNodeIndex;
 		InterlockedExchange(HairVisibilityPass.PPLLNodeIndex[PixelCoord], NodeIndex, NextNodeIndex);
-		const FPPLLNodeData Node = CreatePPLLNodeData(
+		const FPackedHairVisPPLL Node = PackHairVisPPLL(
 			SvPosition.z, 
-			PrimitiveId, 
-			MacroGroupId, 
-			Tangent, 
 			Coverage, 
-			Coverage8bit, 
-			BaseColor, 
-			Roughness, 
-			Specular, 
-			Backlit, 
-			HairVisibilityPass_LightChannelMask, 
-			Velocity, 
+			HairPrimitiveId,
+			HairVisibilityPass_HairMaterialId,
 			NextNodeIndex);
 		HairVisibilityPass.PPLLNodeData[NodeIndex] = Node;
 	}
diff --git a/Engine/Source/Runtime/Renderer/Private/HairStrands/HairStrandsDebug.cpp b/Engine/Source/Runtime/Renderer/Private/HairStrands/HairStrandsDebug.cpp
--- a/Engine/Source/Runtime/Renderer/Private/HairStrands/HairStrandsDebug.cpp
+++ b/Engine/Source/Runtime/Renderer/Private/HairStrands/HairStrandsDebug.cpp
@@ -1376,6 +1376,13 @@ static void InternalRenderHairStrandsDebugInfo(
 	{
 		if (HairData.DebugData.IsPPLLDataValid()) // Check if PPLL rendering is used and its debug view is enabled.
 		{
+			if (!ShaderPrint::IsEnabled(View))
+			{
+				ShaderPrint::SetEnabled(true);
+			}
+
+			ShaderPrint::RequestSpaceForCharacters(256);
+
 			const FIntPoint PPLLResolution = HairData.DebugData.PPLLNodeIndexTexture->Desc.Extent;
 			FHairVisibilityDebugPPLLCS::FParameters* PassParameters = GraphBuilder.AllocParameters<FHairVisibilityDebugPPLLCS::FParameters>();
 			PassParameters->PPLLMeanListElementCountPerPixel = GetHairStrandsMeanSamplePerPixel();
@@ -1390,8 +1397,7 @@ static void InternalRenderHairStrandsDebugInfo(
 			FHairVisibilityDebugPPLLCS::FPermutationDomain PermutationVector;
 			TShaderMapRef<FHairVisibilityDebugPPLLCS> ComputeShader(View.ShaderMap, PermutationVector);
 			FIntVector TextureSize = SceneColorTexture->Desc.GetSize(); TextureSize.Z = 1;
-			FComputeShaderUtils::AddPass(GraphBuilder, RDG_EVENT_NAME("HairStrands::PPLLDebug"), ComputeShader, PassParameters,
-				FIntVector::DivideAndRoundUp(TextureSize, FIntVector(8, 8, 1)));
+			FComputeShaderUtils::AddPass(GraphBuilder, RDG_EVENT_NAME("HairStrands::PPLLDebug"), ComputeShader, PassParameters, FIntVector::DivideAndRoundUp(TextureSize, FIntVector(8, 8, 1)));
 		}
 	}
 
diff --git a/Engine/Source/Runtime/Renderer/Private/HairStrands/HairStrandsUtils.cpp b/Engine/Source/Runtime/Renderer/Private/HairStrands/HairStrandsUtils.cpp
--- a/Engine/Source/Runtime/Renderer/Private/HairStrands/HairStrandsUtils.cpp
+++ b/Engine/Source/Runtime/Renderer/Private/HairStrands/HairStrandsUtils.cpp
@@ -26,7 +26,7 @@ static float GStrandHairStableRasterizationScale = 1.0f; // For no AA without TA
 static FAutoConsoleVariableRef CVarStrandHairStableRasterizationScale(TEXT("r.HairStrands.StableRasterizationScale"), GStrandHairStableRasterizationScale, TEXT("Rasterization scale to snap strand to pixel for 'stable' hair option. This value can't go below 1."), ECVF_Scalability | ECVF_RenderThreadSafe);
 
 static float GStrandHairVelocityRasterizationScale = 1.5f; // Tuned based on heavy motion example (e.g., head shaking)
-static FAutoConsoleVariableRef CVarStrandHairMaxRasterizationScale(TEXT("r.HairStrands.VelocityRasterizationScale"), GStrandHairVelocityRasterizationScale, TEXT("Rasterization scale to snap strand to pixel under high velocity"), ECVF_Scalability | ECVF_RenderThreadSafe);
+static FAutoConsoleVariableRef CVarStrandHairVelocityRasterizationScale(TEXT("r.HairStrands.VelocityRasterizationScale"), GStrandHairVelocityRasterizationScale, TEXT("Rasterization scale to snap strand to pixel under high velocity"), ECVF_Scalability | ECVF_RenderThreadSafe);
 
 static float GStrandHairShadowRasterizationScale = 1.0f;
 static FAutoConsoleVariableRef CVarStrandHairShadowRasterizationScale(TEXT("r.HairStrands.ShadowRasterizationScale"), GStrandHairShadowRasterizationScale, TEXT("Rasterization scale to snap strand to pixel in shadow view"));
diff --git a/Engine/Source/Runtime/Renderer/Private/HairStrands/HairStrandsVisibility.cpp b/Engine/Source/Runtime/Renderer/Private/HairStrands/HairStrandsVisibility.cpp
--- a/Engine/Source/Runtime/Renderer/Private/HairStrands/HairStrandsVisibility.cpp
+++ b/Engine/Source/Runtime/Renderer/Private/HairStrands/HairStrandsVisibility.cpp
@@ -1083,7 +1083,7 @@ BEGIN_GLOBAL_SHADER_PARAMETER_STRUCT(FVisibilityPassUniformParameters, )
 	SHADER_PARAMETER(uint32, MaxPPLLNodeCount)
 	SHADER_PARAMETER_RDG_TEXTURE_UAV(RWTexture2D<uint>, PPLLCounter)
 	SHADER_PARAMETER_RDG_TEXTURE_UAV(RWTexture2D<uint>, PPLLNodeIndex)
-	SHADER_PARAMETER_RDG_BUFFER_UAV(RWStructuredBuffer<FPPLLNodeData>, PPLLNodeData)
+	SHADER_PARAMETER_RDG_BUFFER_UAV(RWStructuredBuffer<FPackedHairVisPPLL>, PPLLNodeData)
 END_GLOBAL_SHADER_PARAMETER_STRUCT()
 IMPLEMENT_STATIC_UNIFORM_BUFFER_STRUCT(FVisibilityPassUniformParameters, "HairVisibilityPass", SceneTextures);
 
@@ -1094,16 +1094,12 @@ BEGIN_SHADER_PARAMETER_STRUCT(FVisibilityPassParameters, )
 	RENDER_TARGET_BINDING_SLOTS()
 END_SHADER_PARAMETER_STRUCT()
 
-// Example: 28bytes * 8spp = 224bytes per pixel = 442Mb @ 1080p
-struct PPLLNodeData
+// Example: 12bytes * 8spp = 96bytes per pixel = 192Mb @ 1080p
+struct FPackedHairVisPPLL
 {
 	uint32 Depth;
 	uint32 PrimitiveId_MacroGroupId;
-	uint32 Tangent_Coverage;
-	uint32 BaseColor_Roughness;
-	uint32 Specular;
 	uint32 NextNodeIndex;
-	uint32 PackedVelocity;
 };
 
 TRDGUniformBufferRef<FVisibilityPassUniformParameters> CreatePassDummyTextures(FRDGBuilder& GraphBuilder)
@@ -1113,7 +1109,7 @@ TRDGUniformBufferRef<FVisibilityPassUniformParameters> CreatePassDummyTextures(F
 	FRDGTextureDesc Desc = FRDGTextureDesc::Create2D(FIntPoint(1,1), PF_R32_UINT, FClearValueBinding::None, TexCreate_UAV | TexCreate_ShaderResource);
 	UniformParameters->PPLLCounter		= GraphBuilder.CreateUAV(GraphBuilder.CreateTexture(Desc, TEXT("Hair.VisibilityPPLLNodeCounter")));
 	UniformParameters->PPLLNodeIndex	= GraphBuilder.CreateUAV(GraphBuilder.CreateTexture(Desc, TEXT("Hair.VisibilityPPLLNodeIndex")));
-	UniformParameters->PPLLNodeData		= GraphBuilder.CreateUAV(GraphBuilder.CreateBuffer(FRDGBufferDesc::CreateStructuredDesc(sizeof(PPLLNodeData), 1), TEXT("Hair.DummyPPLLNodeData")));
+	UniformParameters->PPLLNodeData		= GraphBuilder.CreateUAV(GraphBuilder.CreateBuffer(FRDGBufferDesc::CreateStructuredDesc(sizeof(FPackedHairVisPPLL), 1), TEXT("Hair.DummyPPLLNodeData")));
 
 	return GraphBuilder.CreateUniformBuffer(UniformParameters);
 }
@@ -1603,11 +1599,10 @@ class FHairVisibilityPrimitiveIdCompactionCS : public FGlobalShader
 	SHADER_USE_PARAMETER_STRUCT(FHairVisibilityPrimitiveIdCompactionCS, FGlobalShader);
 
 	class FGroupSize	: SHADER_PERMUTATION_SPARSE_INT("PERMUTATION_GROUPSIZE", 32, 64);
-	class FVelocity		: SHADER_PERMUTATION_INT("PERMUTATION_VELOCITY", 2);
 	class FTile			: SHADER_PERMUTATION_BOOL("PERMUTATION_TILE");
 	class FPPLL 		: SHADER_PERMUTATION_SPARSE_INT("PERMUTATION_PPLL", 0, 8, 16, 32); // See GetPPLLMaxRenderNodePerPixel
 	class FMSAACount 	: SHADER_PERMUTATION_SPARSE_INT("PERMUTATION_MSAACOUNT", 1, 2, 4, 8);
-	using FPermutationDomain = TShaderPermutationDomain<FGroupSize, FVelocity, FTile, FPPLL, FMSAACount>;
+	using FPermutationDomain = TShaderPermutationDomain<FGroupSize, FTile, FPPLL, FMSAACount>;
 
 	BEGIN_SHADER_PARAMETER_STRUCT(FParameters, )
 		SHADER_PARAMETER(FIntPoint, OutputResolution)
@@ -1617,7 +1612,6 @@ class FHairVisibilityPrimitiveIdCompactionCS : public FGlobalShader
 		SHADER_PARAMETER(float, DepthTheshold)
 		SHADER_PARAMETER(float, CosTangentThreshold)
 		SHADER_PARAMETER(float, CoverageThreshold)
-		SHADER_PARAMETER(uint32, VelocityType)
 
 		SHADER_PARAMETER(FIntPoint, TileCountXY)
 		SHADER_PARAMETER(uint32, TileSize)
@@ -1641,7 +1635,6 @@ class FHairVisibilityPrimitiveIdCompactionCS : public FGlobalShader
 		SHADER_PARAMETER_RDG_TEXTURE_UAV(RWTexture2D, OutCompactNodeIndex)
 		SHADER_PARAMETER_RDG_TEXTURE_UAV(RWTexture2D, OutCoverageTexture)
 		SHADER_PARAMETER_RDG_BUFFER_UAV(RWStructuredBuffer, OutCompactNodeVis)
-		SHADER_PARAMETER_RDG_BUFFER_UAV(RWStructuredBuffer, OutCompactNodeData)
 		SHADER_PARAMETER_RDG_BUFFER_UAV(RWBuffer, OutCompactNodeCoord)
 		SHADER_PARAMETER_RDG_TEXTURE_UAV(RWTexture2D, OutVelocityTexture)
 		
@@ -1687,10 +1680,9 @@ static void AddHairVisibilityPrimitiveIdCompactionPass(
 	FHairVisibilityPrimitiveIdCompactionCS::FParameters* PassParameters,
 	FRDGTextureRef& OutCompactCounter,
 	FRDGTextureRef& OutCompactNodeIndex,
-	FRDGBufferRef& OutCompactNodeVis, // Or OutCompactNodeData for PPLL
+	FRDGBufferRef& OutCompactNodeVis,
 	FRDGBufferRef& OutCompactNodeCoord,
 	FRDGTextureRef& OutCoverageTexture,
-	FRDGTextureRef OutVelocityTexture,
 	FRDGBufferRef& OutIndirectArgsBuffer,
 	uint32& OutMaxRenderNodeCount)
 {
@@ -1740,26 +1732,12 @@ static void AddHairVisibilityPrimitiveIdCompactionPass(
 	const uint32 MaxRenderNodeCount = GetTotalSampleCountForAllocation(EffectiveResolution);
 	const bool bUseTile = TileData.IsValid();
 
-	if (bUsePPLL)
-	{
-		// PPLL output directly the node data
-		OutCompactNodeVis = GraphBuilder.CreateBuffer(FRDGBufferDesc::CreateStructuredDesc(sizeof(HairStrandsVisibilityInternal::NodeData), MaxRenderNodeCount), TEXT("Hair.VisibilityNodeData"));
-	}
-	else
-	{
-		OutCompactNodeVis = GraphBuilder.CreateBuffer(FRDGBufferDesc::CreateStructuredDesc(sizeof(HairStrandsVisibilityInternal::NodeVis), MaxRenderNodeCount), TEXT("Hair.VisibilityNodeVis"));
-	}
-
-	{
-		// Pixel coord of the node. Stored as 2*R16_UINT
-		OutCompactNodeCoord = GraphBuilder.CreateBuffer(FRDGBufferDesc::CreateBufferDesc(sizeof(uint32), MaxRenderNodeCount), TEXT("Hair.VisibilityNodeCoord"));
-	}
+	OutCompactNodeVis = GraphBuilder.CreateBuffer(FRDGBufferDesc::CreateStructuredDesc(sizeof(HairStrandsVisibilityInternal::NodeVis), MaxRenderNodeCount), TEXT("Hair.VisibilityNodeVis"));
+	// Pixel coord of the node. Stored as 2*R16_UINT
+	OutCompactNodeCoord = GraphBuilder.CreateBuffer(FRDGBufferDesc::CreateBufferDesc(sizeof(uint32), MaxRenderNodeCount), TEXT("Hair.VisibilityNodeCoord"));
 
-	const bool bWriteOutVelocity = OutVelocityTexture != nullptr && bUsePPLL; // Velocity write out is only support with PPLL
-	const uint32 VelocityPermutation = bWriteOutVelocity ? FMath::Clamp(GHairVelocityType + 1, 0, 3) : 0;
 	FHairVisibilityPrimitiveIdCompactionCS::FPermutationDomain PermutationVector;
 	PermutationVector.Set<FHairVisibilityPrimitiveIdCompactionCS::FGroupSize>(bUseTile ? FHairStrandsTiles::GroupSize : GetVendorOptimalGroupSize1D());
-	PermutationVector.Set<FHairVisibilityPrimitiveIdCompactionCS::FVelocity>(VelocityPermutation > 0 ? 1 : 0);
 	PermutationVector.Set<FHairVisibilityPrimitiveIdCompactionCS::FTile>(bUseTile);
 	PermutationVector.Set<FHairVisibilityPrimitiveIdCompactionCS::FPPLL>(bUsePPLL ? PPLLMaxRenderNodePerPixel : 0);
 	PermutationVector.Set<FHairVisibilityPrimitiveIdCompactionCS::FMSAACount>(MSAASampleCount);
@@ -1767,7 +1745,6 @@ static void AddHairVisibilityPrimitiveIdCompactionPass(
 
 	PassParameters->ResolutionOffset = FIntPoint(0,0);
 	PassParameters->OutputResolution = Resolution;
-	PassParameters->VelocityType = VelocityPermutation;
 	PassParameters->MaxNodeCount = MaxRenderNodeCount;
 	PassParameters->bSortSampleByDepth = GHairStrandsSortHairSampleByDepth > 0 ? 1 : 0;
 	PassParameters->CoverageThreshold = GetHairStrandsFullCoverageThreshold();
@@ -1777,18 +1754,10 @@ static void AddHairVisibilityPrimitiveIdCompactionPass(
 	PassParameters->ViewUniformBuffer = View.ViewUniformBuffer;
 	PassParameters->OutCompactNodeCounter = GraphBuilder.CreateUAV(OutCompactCounter);
 	PassParameters->OutCompactNodeIndex = GraphBuilder.CreateUAV(OutCompactNodeIndex);
-	if (bUsePPLL)
-		PassParameters->OutCompactNodeData = GraphBuilder.CreateUAV(OutCompactNodeVis);
-	else
-		PassParameters->OutCompactNodeVis = GraphBuilder.CreateUAV(OutCompactNodeVis);
+	PassParameters->OutCompactNodeVis = GraphBuilder.CreateUAV(OutCompactNodeVis);
 	PassParameters->OutCompactNodeCoord = GraphBuilder.CreateUAV(OutCompactNodeCoord, FHairStrandsVisibilityData::NodeCoordFormat);
 	PassParameters->OutCoverageTexture = GraphBuilder.CreateUAV(OutCoverageTexture);
 
-	if (bWriteOutVelocity)
-	{
-		PassParameters->OutVelocityTexture = GraphBuilder.CreateUAV(OutVelocityTexture);
-	}
-
 	const FHairStrandsTiles::ETileType TileType = FHairStrandsTiles::ETileType::HairAll;
 	if (bUseTile)
 	{
@@ -2275,12 +2244,14 @@ static void AddHairVisibilityPPLLPass(
 		OutVisibilityPPLLNodeIndex = GraphBuilder.CreateTexture(Desc, TEXT("Hair.VisibilityPPLLNodeIndex"));
 	}
 
-	const FIntRect HairRect = ComputeVisibleHairStrandsMacroGroupsRect(ViewInfo->ViewRect, MacroGroupDatas);
+	// Don't use CPU bounds projection to drive PPLL buffer allocation (ComputeVisibleHairStrandsMacroGroupsRect), 
+	// as it often underestimate allocation needs and causes flickering
+	const FIntRect HairRect = ViewInfo->ViewRect;
 	const FIntPoint EffectiveResolution(HairRect.Width(), HairRect.Height());
 
 	const uint32 PPLLMaxTotalListElementCount = GetTotalSampleCountForAllocation(EffectiveResolution);
 	{
-		OutVisibilityPPLLNodeData = GraphBuilder.CreateBuffer(FRDGBufferDesc::CreateStructuredDesc(sizeof(PPLLNodeData), PPLLMaxTotalListElementCount), TEXT("Hair.VisibilityPPLLNodeData"));
+		OutVisibilityPPLLNodeData = GraphBuilder.CreateBuffer(FRDGBufferDesc::CreateStructuredDesc(sizeof(FPackedHairVisPPLL), PPLLMaxTotalListElementCount), TEXT("Hair.VisibilityPPLLNodeData"));
 	}
 	const uint32 ClearValue0[4] = { 0,0,0,0 };
 	const uint32 ClearValueInvalid[4] = { 0xFFFFFFFF, 0xFFFFFFFF, 0xFFFFFFFF, 0xFFFFFFFF };
@@ -3617,7 +3588,6 @@ void DrawHitProxies(
 		VisNodeData,
 		VisNodeCoord,
 		CoverageTexture,
-		nullptr, // Velocity output is only needed for PPLL
 		IndirectArgsBuffer,
 		OutMaxNodeCount);
 
@@ -3948,7 +3918,6 @@ void RenderHairStrandsVisibilityBuffer(
 						CompactNodeVis,
 						CompactNodeCoord,
 						CoverageTexture,
-						nullptr, // Velocity output is only needed for PPLL
 						IndirectArgsBuffer,
 						VisibilityData.MaxNodeCount);
 
@@ -4091,6 +4060,7 @@ void RenderHairStrandsVisibilityBuffer(
 
 					FRDGBufferRef CompactNodeCoord;
 					FRDGBufferRef IndirectArgsBuffer;
+					FRDGTextureRef ResolveMaskTexture = nullptr;
 					AddHairVisibilityPrimitiveIdCompactionPass(
 						true, // bUsePPLL
 						GraphBuilder,
@@ -4102,13 +4072,59 @@ void RenderHairStrandsVisibilityBuffer(
 						PassParameters,
 						NodeCounter,
 						CompactNodeIndex,
-						CompactNodeData,
+						CompactNodeVis,
 						CompactNodeCoord,
 						CoverageTexture,
-						SceneVelocityTexture,
 						IndirectArgsBuffer,
 						VisibilityData.MaxNodeCount);
 
+					{
+						const bool bUpdateSampleCoverage = GHairStrandsSortHairSampleByDepth > 0;
+
+						// Evaluate material based on the visiblity pass result
+						// Output both complete sample data + per-sample velocity
+						FMaterialPassOutput PassOutput = AddHairMaterialPass(
+							GraphBuilder,
+							Scene,
+							&View,
+							bUpdateSampleCoverage,
+							MacroGroupDatas,
+							InstanceCullingManager,
+							VisibilityData.NodeGroupSize,
+							CompactNodeIndex,
+							CompactNodeVis,
+							CompactNodeCoord,
+							NodeCounter,
+							IndirectArgsBuffer);
+
+						// Merge per-sample velocity into the scene velocity buffer
+						AddHairVelocityPass(
+							GraphBuilder,
+							View,
+							MacroGroupDatas,
+							VisibilityData.TileData,
+							CoverageTexture,
+							CompactNodeIndex,
+							CompactNodeVis,
+							PassOutput.NodeVelocity,
+							SceneVelocityTexture,
+							ResolveMaskTexture);
+
+						if (bUpdateSampleCoverage)
+						{
+							PassOutput.NodeData = AddUpdateSampleCoveragePass(
+								GraphBuilder,
+								&View,
+								CompactNodeIndex,
+								PassOutput.NodeData);
+						}
+
+						CompactNodeData = PassOutput.NodeData;
+
+						VisibilityData.SampleLightingViewportResolution = PassOutput.SampleLightingTexture->Desc.Extent;
+						VisibilityData.SampleLightingTexture			= PassOutput.SampleLightingTexture;
+					}
+
 					VisibilityData.MaxSampleCount = GetMaxSamplePerPixel();
 					VisibilityData.NodeIndex = CompactNodeIndex;
 					VisibilityData.CoverageTexture = CoverageTexture;
@@ -4119,7 +4135,6 @@ void RenderHairStrandsVisibilityBuffer(
 					VisibilityData.NodeCount = NodeCounter;
 				}
 
-
 				if (bRunColorAndDepthPatching)
 				{
 					AddHairGbufferPatchPass(
@@ -4145,9 +4160,13 @@ void RenderHairStrandsVisibilityBuffer(
 					CompactNodeData,
 					HairOnlyDepthTexture);
 
-				// Allocate buffer for storing all the light samples
-				VisibilityData.SampleLightingTexture = AddClearLightSamplePass(GraphBuilder, &View, VisibilityData.MaxNodeCount, NodeCounter);
-				VisibilityData.SampleLightingViewportResolution = VisibilityData.SampleLightingTexture->Desc.Extent;
+				AddHairOnlyHZBPass(
+					GraphBuilder,
+					View,
+					HairOnlyDepthTexture,
+					VisibilityData.HairOnlyDepthHZBParameters,
+					VisibilityData.HairOnlyDepthClosestHZBTexture,
+					VisibilityData.HairOnlyDepthFurthestHZBTexture);
 
 			#if WITH_EDITOR
 				// Extract texture for debug visualization
