Added a separated start distance for volumetric fog on exponential fog.
The volumetric fog view distance is now taken into account as distance after the start distance.

# jordan.walker, patrick.kelly

diff --git a/Engine/Shaders/Private/HeightFogCommon.ush b/Engine/Shaders/Private/HeightFogCommon.ush
--- a/Engine/Shaders/Private/HeightFogCommon.ush
+++ b/Engine/Shaders/Private/HeightFogCommon.ush
@@ -231,29 +231,29 @@ float4 CombineVolumetricFog(float4 GlobalFog, float3 VolumeUV, uint EyeIndex, fl
 	float4 VolumetricFogLookup = float4(0, 0, 0, 1);
 
 #if PERMUTATION_SUPPORT_VOLUMETRIC_FOG
-	float StartDistance = 0.0f;
+	float VolFogStartDistance = 0.0f;
 	if (FogStruct.ApplyVolumetricFog > 0)
 	{
 	#if INSTANCED_STEREO
 		if (EyeIndex == 0)
 		{
-			StartDistance = FogStruct.ExponentialFogParameters.w;
+			VolFogStartDistance = FogStruct.VolumetricFogStartDistance;
 			VolumetricFogLookup = Texture3DSampleLevel(FogStruct.IntegratedLightScattering, SharedIntegratedLightScatteringSampler, VolumeUV, 0);
 		}
 		else
 		{
-			StartDistance = FogStructISR.ExponentialFogParameters.w;
+			VolFogStartDistance = FogStructISR.VolumetricFogStartDistance;
 			VolumetricFogLookup = Texture3DSampleLevel(FogStructISR.IntegratedLightScattering, SharedIntegratedLightScatteringSampler, VolumeUV, 0);
 		}
 	#else
-		StartDistance = FogStruct.ExponentialFogParameters.w;
+		VolFogStartDistance = FogStruct.VolumetricFogStartDistance;
 		VolumetricFogLookup = Texture3DSampleLevel(FogStruct.IntegratedLightScattering, SharedIntegratedLightScatteringSampler, VolumeUV, 0);
 	#endif
 	}
 
 	// Do not apply the Froxel volumetric texture in front of the fog start distance. (the soft fading occur in FinalIntegrationCS).
-	// We go with a quickly increasing step function because the soft fade in occur in FinalIntegrationCS.
-	VolumetricFogLookup = lerp(float4(0, 0, 0, 1), VolumetricFogLookup, saturate((SceneDepth - StartDistance) * 100000000.0f));
+	// We go with a quickly increasing step function because the soft fade in from start distance occurs in FinalIntegrationCS.
+	VolumetricFogLookup = lerp(float4(0, 0, 0, 1), VolumetricFogLookup, saturate((SceneDepth - VolFogStartDistance) * 100000000.0f));
 #endif
 
 	// Visualize depth distribution
diff --git a/Engine/Source/Runtime/Engine/Classes/Components/ExponentialHeightFogComponent.h b/Engine/Source/Runtime/Engine/Classes/Components/ExponentialHeightFogComponent.h
--- a/Engine/Source/Runtime/Engine/Classes/Components/ExponentialHeightFogComponent.h
+++ b/Engine/Source/Runtime/Engine/Classes/Components/ExponentialHeightFogComponent.h
@@ -172,11 +172,17 @@ class ENGINE_API UExponentialHeightFogComponent : public USceneComponent
 	float VolumetricFogExtinctionScale;
 
 	/** 
-	 * Distance over which volumetric fog should be computed.  Larger values extend the effect into the distance but expose under-sampling artifacts in details.
+	 * Distance over which volumetric fog should be computed, after the start distance.  Larger values extend the effect into the distance but expose under-sampling artifacts in details.
 	 */
 	UPROPERTY(EditAnywhere, BlueprintReadOnly, Category = VolumetricFog, meta=(DisplayName = "View Distance", UIMin = "1000", UIMax = "10000"))
 	float VolumetricFogDistance;
 
+	/** 
+	 * Distance from the camera that the volumetric fog will start, in world units. 
+	 */
+	UPROPERTY(BlueprintReadOnly, interp, Category= VolumetricFog, meta=(DisplayName = "Start Distance", UIMin = "0", UIMax = "5000"))
+	float VolumetricFogStartDistance;
+
 	/** 
 	 * Distance over which volumetric fog will fade in from the start distance.
 	 */
diff --git a/Engine/Source/Runtime/Renderer/Private/FogRendering.cpp b/Engine/Source/Runtime/Renderer/Private/FogRendering.cpp
--- a/Engine/Source/Runtime/Renderer/Private/FogRendering.cpp
+++ b/Engine/Source/Runtime/Renderer/Private/FogRendering.cpp
@@ -92,6 +92,7 @@ void SetupFogUniformParameters(FRDGBuilder& GraphBuilder, const FViewInfo& View,
 			OutParameters.ApplyVolumetricFog = 0.0f;
 		}
 		OutParameters.IntegratedLightScatteringSampler = TStaticSamplerState<SF_Bilinear, AM_Clamp, AM_Clamp, AM_Clamp>::GetRHI();
+		OutParameters.VolumetricFogStartDistance = View.VolumetricFogStartDistance;
 		OutParameters.VolumetricFogNearFadeInDistanceInv = View.VolumetricFogNearFadeInDistanceInv;
 	}
 }
@@ -258,6 +259,8 @@ void FSceneRenderer::InitFogConstants()
 					View.DirectionalInscatteringColor = FogInfo.DirectionalInscatteringColor * SunLight->Proxy->GetColor().GetLuminance();
 				}
 				View.bUseDirectionalInscattering = SunLight != nullptr;
+				View.bEnableVolumetricFog = FogInfo.bEnableVolumetricFog;
+				View.VolumetricFogStartDistance = FogInfo.VolumetricFogStartDistance;
 				View.VolumetricFogNearFadeInDistanceInv = FogInfo.VolumetricFogNearFadeInDistance > 0.0f ? (1.0f / FogInfo.VolumetricFogNearFadeInDistance) : 100000000.0f;
 			}
 		}
@@ -332,9 +335,12 @@ static void RenderViewFog(
 	GraphicsPSOInit.bDepthBounds = GSupportsDepthBoundsTest && CVarFogUseDepthBounds.GetValueOnAnyThread();
 	if (GraphicsPSOInit.bDepthBounds)
 	{
+		float ExpFogStartDistance = View.ExponentialFogParameters.W;
+		float VolFogStartDistance = View.bEnableVolumetricFog ? View.VolumetricFogStartDistance : ExpFogStartDistance;
+
 		// The fog can be set to start at a certain euclidean distance.
-		// clamp the value to be behind the near plane z
-		float FogStartDistance = FMath::Max(30.0f, View.ExponentialFogParameters.W);
+		// clamp the value to be behind the near plane z, according to the smallest distance between volumetric fog and height fog (if they are enabled). 
+		float FogStartDistance = FMath::Max(30.0f, FMath::Min(ExpFogStartDistance, VolFogStartDistance));
 
 		// Here we compute the nearest z value the fog can start
 		// to skip shader execution on pixels that are closer.
diff --git a/Engine/Source/Runtime/Renderer/Private/FogRendering.h b/Engine/Source/Runtime/Renderer/Private/FogRendering.h
--- a/Engine/Source/Runtime/Renderer/Private/FogRendering.h
+++ b/Engine/Source/Runtime/Renderer/Private/FogRendering.h
@@ -22,6 +22,7 @@ BEGIN_GLOBAL_SHADER_PARAMETER_STRUCT(FFogUniformParameters,)
 	SHADER_PARAMETER(FVector2f, SinCosInscatteringColorCubemapRotation)
 	SHADER_PARAMETER(FVector3f, FogInscatteringTextureParameters)
 	SHADER_PARAMETER(float, ApplyVolumetricFog)
+	SHADER_PARAMETER(float, VolumetricFogStartDistance)
 	SHADER_PARAMETER(float, VolumetricFogNearFadeInDistanceInv)
 	SHADER_PARAMETER_TEXTURE(TextureCube, FogInscatteringColorCubemap)
 	SHADER_PARAMETER_SAMPLER(SamplerState, FogInscatteringColorSampler)
diff --git a/Engine/Source/Runtime/Renderer/Private/SceneCore.cpp b/Engine/Source/Runtime/Renderer/Private/SceneCore.cpp
--- a/Engine/Source/Runtime/Renderer/Private/SceneCore.cpp
+++ b/Engine/Source/Runtime/Renderer/Private/SceneCore.cpp
@@ -418,10 +418,11 @@ FExponentialHeightFogSceneInfo::FExponentialHeightFogSceneInfo(const UExponentia
 	VolumetricFogEmissive.G = FMath::Max(VolumetricFogEmissive.G * UnitScale, 0.0f);
 	VolumetricFogEmissive.B = FMath::Max(VolumetricFogEmissive.B * UnitScale, 0.0f);
 	VolumetricFogExtinctionScale = FMath::Max(InComponent->VolumetricFogExtinctionScale, 0.0f);
-	VolumetricFogDistance = FMath::Max(InComponent->VolumetricFogDistance, 0.0f);
+	VolumetricFogDistance = FMath::Max(InComponent->VolumetricFogStartDistance + InComponent->VolumetricFogDistance, 0.0f);
 	VolumetricFogStaticLightingScatteringIntensity = FMath::Max(InComponent->VolumetricFogStaticLightingScatteringIntensity, 0.0f);
 	bOverrideLightColorsWithFogInscatteringColors = InComponent->bOverrideLightColorsWithFogInscatteringColors;
 
+	VolumetricFogStartDistance = InComponent->VolumetricFogStartDistance;
 	VolumetricFogNearFadeInDistance = InComponent->VolumetricFogNearFadeInDistance;
 
 	SkyAtmosphereAmbientContributionColorScale = InComponent->SkyAtmosphereAmbientContributionColorScale;
diff --git a/Engine/Source/Runtime/Renderer/Private/SceneCore.h b/Engine/Source/Runtime/Renderer/Private/SceneCore.h
--- a/Engine/Source/Runtime/Renderer/Private/SceneCore.h
+++ b/Engine/Source/Runtime/Renderer/Private/SceneCore.h
@@ -287,6 +287,7 @@ public:
 	float VolumetricFogDistance;
 	float VolumetricFogStaticLightingScatteringIntensity;
 	bool bOverrideLightColorsWithFogInscatteringColors;
+	float VolumetricFogStartDistance;
 	float VolumetricFogNearFadeInDistance;
 
 	FLinearColor SkyAtmosphereAmbientContributionColorScale;
diff --git a/Engine/Source/Runtime/Renderer/Private/SceneRendering.cpp b/Engine/Source/Runtime/Renderer/Private/SceneRendering.cpp
--- a/Engine/Source/Runtime/Renderer/Private/SceneRendering.cpp
+++ b/Engine/Source/Runtime/Renderer/Private/SceneRendering.cpp
@@ -910,6 +910,8 @@ void FViewInfo::Init()
 	SinCosInscatteringColorCubemapRotation = FVector2f::ZeroVector;
 	FogInscatteringColorCubemap = nullptr;
 	FogInscatteringTextureParameters = FVector::ZeroVector;
+	VolumetricFogStartDistance = false;
+	VolumetricFogStartDistance = 0.0f;
 	VolumetricFogNearFadeInDistanceInv = 100000000.0f;
 
 	SkyAtmosphereCameraAerialPerspectiveVolume = nullptr;
diff --git a/Engine/Source/Runtime/Renderer/Private/SceneRendering.h b/Engine/Source/Runtime/Renderer/Private/SceneRendering.h
--- a/Engine/Source/Runtime/Renderer/Private/SceneRendering.h
+++ b/Engine/Source/Runtime/Renderer/Private/SceneRendering.h
@@ -1308,6 +1308,8 @@ public:
 	float FogMaxOpacity;
 	FVector4f ExponentialFogParameters3;
 	FVector4f SkyAtmosphereAmbientContributionColorScale;
+	bool bEnableVolumetricFog;
+	float VolumetricFogStartDistance;
 	float VolumetricFogNearFadeInDistanceInv;
 	FVector2f SinCosInscatteringColorCubemapRotation;
 
diff --git a/Engine/Source/Runtime/Renderer/Private/VolumetricFog.cpp b/Engine/Source/Runtime/Renderer/Private/VolumetricFog.cpp
--- a/Engine/Source/Runtime/Renderer/Private/VolumetricFog.cpp
+++ b/Engine/Source/Runtime/Renderer/Private/VolumetricFog.cpp
@@ -762,7 +762,7 @@ bool ShouldRenderVolumetricFog(const FScene* Scene, const FSceneViewFamily& View
 		&& Scene->ExponentialFogs[0].VolumetricFogDistance > 0;
 }
 
-FVector GetVolumetricFogGridZParams(float HeighFogStartdistance, float NearPlane, float FarPlane, int32 GridSizeZ)
+FVector GetVolumetricFogGridZParams(float VolumetricFogStartDistance, float NearPlane, float FarPlane, int32 GridSizeZ)
 {
 	// S = distribution scale
 	// B, O are solved for given the z distances of the first+last slice, and the # of slices.
@@ -771,7 +771,7 @@ FVector GetVolumetricFogGridZParams(float HeighFogStartdistance, float NearPlane
 
 	// Don't spend lots of resolution right in front of the near plane
 
-	NearPlane = FMath::Max(NearPlane, double(HeighFogStartdistance));
+	NearPlane = FMath::Max(NearPlane, double(VolumetricFogStartDistance));
 
 	double NearOffset = .095 * 100.0;
 	// Space out the slices so they aren't all clustered at the near plane
@@ -838,7 +838,7 @@ void SetupVolumetricFogGlobalData(const FViewInfo& View, FVolumetricFogGlobalDat
 	Parameters.GridSizeInt = VolumetricFogGridSize;
 	Parameters.GridSize = FVector3f(VolumetricFogGridSize);
 
-	FVector ZParams = GetVolumetricFogGridZParams(FogInfo.StartDistance, View.NearClippingDistance, FogInfo.VolumetricFogDistance, VolumetricFogGridSize.Z);
+	FVector ZParams = GetVolumetricFogGridZParams(FogInfo.VolumetricFogStartDistance, View.NearClippingDistance, FogInfo.VolumetricFogDistance, VolumetricFogGridSize.Z);
 	Parameters.GridZParams = (FVector3f)ZParams;
 
 	Parameters.SVPosToVolumeUV = FVector2f::UnitVector / (FVector2f(VolumetricFogGridSize.X, VolumetricFogGridSize.Y) * VolumetricFogGridPixelSize);
@@ -867,7 +867,7 @@ void FViewInfo::SetupVolumetricFogUniformBufferParameters(FViewUniformShaderPara
 
 		ViewUniformShaderParameters.VolumetricFogInvGridSize = FVector3f(1.0f / VolumetricFogGridSize.X, 1.0f / VolumetricFogGridSize.Y, 1.0f / VolumetricFogGridSize.Z);
 
-		const FVector ZParams = GetVolumetricFogGridZParams(FogInfo.StartDistance, NearClippingDistance, FogInfo.VolumetricFogDistance, VolumetricFogGridSize.Z);
+		const FVector ZParams = GetVolumetricFogGridZParams(FogInfo.VolumetricFogStartDistance, NearClippingDistance, FogInfo.VolumetricFogDistance, VolumetricFogGridSize.Z);
 		ViewUniformShaderParameters.VolumetricFogGridZParams = (FVector3f)ZParams;
 
 		ViewUniformShaderParameters.VolumetricFogSVPosToVolumeUV = FVector2f::UnitVector / (FVector2f(VolumetricFogGridSize.X, VolumetricFogGridSize.Y) * VolumetricFogGridPixelSize);
@@ -941,7 +941,7 @@ void FDeferredShadingSceneRenderer::ComputeVolumetricFog(FRDGBuilder& GraphBuild
 
 		int32 VolumetricFogGridPixelSize;
 		const FIntVector VolumetricFogGridSize = GetVolumetricFogGridSize(View.ViewRect.Size(), VolumetricFogGridPixelSize);
-		const FVector GridZParams = GetVolumetricFogGridZParams(FogInfo.StartDistance, View.NearClippingDistance, FogInfo.VolumetricFogDistance, VolumetricFogGridSize.Z);
+		const FVector GridZParams = GetVolumetricFogGridZParams(FogInfo.VolumetricFogStartDistance, View.NearClippingDistance, FogInfo.VolumetricFogDistance, VolumetricFogGridSize.Z);
 
 		FVolumetricFogIntegrationParameterData IntegrationData;
 		IntegrationData.FrameJitterOffsetValues.Empty(16);
