Delete unused fog debug cvars

diff --git a/Engine/Source/Runtime/Renderer/Private/FogRendering.cpp b/Engine/Source/Runtime/Renderer/Private/FogRendering.cpp
--- a/Engine/Source/Runtime/Renderer/Private/FogRendering.cpp
+++ b/Engine/Source/Runtime/Renderer/Private/FogRendering.cpp
@@ -10,26 +10,6 @@
 
 DECLARE_GPU_STAT(Fog);
 
-#if UE_ENABLE_DEBUG_DRAWING
-static TAutoConsoleVariable<float> CVarFogStartDistance(
-	TEXT("r.FogStartDistance"),
-	-1.0f,
-	TEXT("Allows to override the FogStartDistance setting (needs ExponentialFog in the level).\n")
-	TEXT(" <0: use default settings (default: -1)\n")
-	TEXT(">=0: override settings by the given value (in world units)"),
-	ECVF_Cheat | ECVF_RenderThreadSafe);
-
-static TAutoConsoleVariable<float> CVarFogDensity(
-	TEXT("r.FogDensity"),
-	-1.0f,
-	TEXT("Allows to override the FogDensity setting (needs ExponentialFog in the level).\n")
-	TEXT("Using a strong value allows to quickly see which pixel are affected by fog.\n")
-	TEXT("Using a start distance allows to cull pixels are can speed up rendering.\n")
-	TEXT(" <0: use default settings (default: -1)\n")
-	TEXT(">=0: override settings by the given value (0:off, 1=very dense fog)"),
-	ECVF_Cheat | ECVF_RenderThreadSafe);
-#endif
-
 static TAutoConsoleVariable<int32> CVarFog(
 	TEXT("r.Fog"),
 	1,
@@ -182,18 +162,6 @@ TGlobalResource<FFogVertexDeclaration> GFogVertexDeclaration;
 
 void FSceneRenderer::InitFogConstants()
 {
-	// console command override
-	float FogDensityOverride = -1.0f;
-	float FogStartDistanceOverride = -1.0f;
-
-#if !(UE_BUILD_SHIPPING || UE_BUILD_TEST)
-	{
-		// console variable overrides
-		FogDensityOverride = CVarFogDensity.GetValueOnAnyThread();
-		FogStartDistanceOverride = CVarFogStartDistance.GetValueOnAnyThread();
-	}
-#endif // !(UE_BUILD_SHIPPING || UE_BUILD_TEST)
-
 	for(int32 ViewIndex = 0;ViewIndex < Views.Num();ViewIndex++)
 	{
 		FViewInfo& View = Views[ViewIndex];
