Avoid compiling non-nanite shadow depth VS/PS on platforms that don't support nanite/VSMs.

diff --git a/Engine/Source/Runtime/Renderer/Private/ShadowDepthRendering.cpp b/Engine/Source/Runtime/Renderer/Private/ShadowDepthRendering.cpp
--- a/Engine/Source/Runtime/Renderer/Private/ShadowDepthRendering.cpp
+++ b/Engine/Source/Runtime/Renderer/Private/ShadowDepthRendering.cpp
@@ -221,6 +221,14 @@ public:
 	int32 bUseGpuSceneInstancing;
 };
 
+static bool UseNonNaniteVirtualShadowMaps(const EShaderPlatform Platform)
+{
+	return IsFeatureLevelSupported(Platform, ERHIFeatureLevel::SM5) &&
+		UseGPUScene(Platform) &&
+		GEnableNonNaniteVSM != 0 &&
+		DoesPlatformSupportNanite(Platform);
+}
+
 /**
 * A vertex shader for rendering the depth of a mesh.
 */
@@ -327,8 +335,7 @@ public:
 			return false;
 		}
 
-		if (ShaderMode == VertexShadowDepth_VirtualShadowMap &&
-			(!IsFeatureLevelSupported(Platform, ERHIFeatureLevel::SM5) || !UseGPUScene(Platform) || GEnableNonNaniteVSM == 0))
+		if (ShaderMode == VertexShadowDepth_VirtualShadowMap && UseNonNaniteVirtualShadowMaps(Platform))
 		{
 			return false;
 		}
@@ -445,8 +452,7 @@ public:
 			return false;
 		}
 
-		if (ShaderMode == PixelShadowDepth_VirtualShadowMap &&
-			(!IsFeatureLevelSupported(Platform, ERHIFeatureLevel::SM5) || !UseGPUScene(Platform) || GEnableNonNaniteVSM == 0))
+		if (ShaderMode == PixelShadowDepth_VirtualShadowMap && UseNonNaniteVirtualShadowMaps(Platform))
 		{
 			return false;
 		}
