Avoid setting up the full ray tracing pipeline when Lumen uses lighting from surface cache and no other HWRT effects are enabled

diff --git a/Engine/Source/Runtime/Renderer/Private/DeferredShadingRenderer.cpp b/Engine/Source/Runtime/Renderer/Private/DeferredShadingRenderer.cpp
--- a/Engine/Source/Runtime/Renderer/Private/DeferredShadingRenderer.cpp
+++ b/Engine/Source/Runtime/Renderer/Private/DeferredShadingRenderer.cpp
@@ -1764,79 +1764,79 @@ void FDeferredShadingSceneRenderer::WaitForRayTracingScene(FRDGBuilder& GraphBui
 				NumTotalBindings, MergedBindings,
 				bCopyDataToInlineStorage);
 
-			if (!bIsPathTracing)
+			// Move the ray tracing binding container ownership to the command list, so that memory will be
+			// released on the RHI thread timeline, after the commands that reference it are processed.
+			RHICmdList.EnqueueLambda([Ptrs = MoveTemp(ReferenceView.RayTracingMaterialBindings)](FRHICommandListImmediate&)
 			{
-				TArray<FRHIRayTracingShader*> DeferredMaterialRayGenShaders;
-				if (!IsForwardShadingEnabled(ShaderPlatform))
+				for (auto Ptr : Ptrs)
 				{
-					for (const FViewInfo& View : Views)
-					{
-						PrepareRayTracingReflectionsDeferredMaterial(View, *Scene, DeferredMaterialRayGenShaders);
-						PrepareRayTracingDeferredReflectionsDeferredMaterial(View, *Scene, DeferredMaterialRayGenShaders);
-						PrepareRayTracingGlobalIlluminationDeferredMaterial(View, DeferredMaterialRayGenShaders);
-						if (DoesPlatformSupportLumenGI(ShaderPlatform))
-						{
-							PrepareLumenHardwareRayTracingReflectionsDeferredMaterial(View, DeferredMaterialRayGenShaders);
-							PrepareLumenHardwareRayTracingRadianceCacheDeferredMaterial(View, DeferredMaterialRayGenShaders);
-							PrepareLumenHardwareRayTracingScreenProbeGatherDeferredMaterial(View, DeferredMaterialRayGenShaders);
-							PrepareLumenHardwareRayTracingVisualizeDeferredMaterial(View, DeferredMaterialRayGenShaders);
-						}
-					}
+					delete Ptr;
 				}
-				DeduplicateRayGenerationShaders(DeferredMaterialRayGenShaders);
+			});
 
-				if (DeferredMaterialRayGenShaders.Num())
-				{
-					ReferenceView.RayTracingMaterialGatherPipeline = BindRayTracingDeferredMaterialGatherPipeline(RHICmdList, ReferenceView, DeferredMaterialRayGenShaders);
-				}
+			if (!bIsPathTracing)
+			{
+				SetupRayTracingLightingMissShader(RHICmdList, ReferenceView);
+			}
+		}
 
-				// Add Lumen hardware ray tracing materials
-				TArray<FRHIRayTracingShader*> LumenHardwareRayTracingRayGenShaders;
-				if (DoesPlatformSupportLumenGI(ShaderPlatform))
+		if (!bIsPathTracing)
+		{
+			TArray<FRHIRayTracingShader*> DeferredMaterialRayGenShaders;
+			if (!IsForwardShadingEnabled(ShaderPlatform))
+			{
+				for (const FViewInfo& View : Views)
 				{
-					for (const FViewInfo& View : Views)
+					PrepareRayTracingReflectionsDeferredMaterial(View, *Scene, DeferredMaterialRayGenShaders);
+					PrepareRayTracingDeferredReflectionsDeferredMaterial(View, *Scene, DeferredMaterialRayGenShaders);
+					PrepareRayTracingGlobalIlluminationDeferredMaterial(View, DeferredMaterialRayGenShaders);
+					if (DoesPlatformSupportLumenGI(ShaderPlatform))
 					{
-						PrepareLumenHardwareRayTracingVisualizeLumenMaterial(View, LumenHardwareRayTracingRayGenShaders);
-						PrepareLumenHardwareRayTracingRadianceCacheLumenMaterial(View, LumenHardwareRayTracingRayGenShaders);
-						PrepareLumenHardwareRayTracingTranslucencyVolumeLumenMaterial(View, LumenHardwareRayTracingRayGenShaders);
-						PrepareLumenHardwareRayTracingRadiosityLumenMaterial(View, LumenHardwareRayTracingRayGenShaders);
-						PrepareLumenHardwareRayTracingReflectionsLumenMaterial(View, LumenHardwareRayTracingRayGenShaders);
-						PrepareLumenHardwareRayTracingScreenProbeGatherLumenMaterial(View, LumenHardwareRayTracingRayGenShaders);
-						PrepareLumenHardwareRayTracingDirectLightingLumenMaterial(View, LumenHardwareRayTracingRayGenShaders);
+						PrepareLumenHardwareRayTracingReflectionsDeferredMaterial(View, DeferredMaterialRayGenShaders);
+						PrepareLumenHardwareRayTracingRadianceCacheDeferredMaterial(View, DeferredMaterialRayGenShaders);
+						PrepareLumenHardwareRayTracingScreenProbeGatherDeferredMaterial(View, DeferredMaterialRayGenShaders);
+						PrepareLumenHardwareRayTracingVisualizeDeferredMaterial(View, DeferredMaterialRayGenShaders);
 					}
 				}
-				DeduplicateRayGenerationShaders(DeferredMaterialRayGenShaders);
+			}
+			DeduplicateRayGenerationShaders(DeferredMaterialRayGenShaders);
 
-				if (LumenHardwareRayTracingRayGenShaders.Num())
-				{
-					ReferenceView.LumenHardwareRayTracingMaterialPipeline = BindLumenHardwareRayTracingMaterialPipeline(RHICmdList, ReferenceView, LumenHardwareRayTracingRayGenShaders, ReferenceView.LumenHardwareRayTracingHitDataBuffer);
-				}
+			if (DeferredMaterialRayGenShaders.Num())
+			{
+				ReferenceView.RayTracingMaterialGatherPipeline = BindRayTracingDeferredMaterialGatherPipeline(RHICmdList, ReferenceView, DeferredMaterialRayGenShaders);
 			}
 
-			// Move the ray tracing binding container ownership to the command list, so that memory will be
-			// released on the RHI thread timeline, after the commands that reference it are processed.
-			RHICmdList.EnqueueLambda([Ptrs = MoveTemp(ReferenceView.RayTracingMaterialBindings)](FRHICommandListImmediate&)
+			// Add Lumen hardware ray tracing materials
+			TArray<FRHIRayTracingShader*> LumenHardwareRayTracingRayGenShaders;
+			if (DoesPlatformSupportLumenGI(ShaderPlatform))
 			{
-				for (auto Ptr : Ptrs)
+				for (const FViewInfo& View : Views)
 				{
-					delete Ptr;
+					PrepareLumenHardwareRayTracingVisualizeLumenMaterial(View, LumenHardwareRayTracingRayGenShaders);
+					PrepareLumenHardwareRayTracingRadianceCacheLumenMaterial(View, LumenHardwareRayTracingRayGenShaders);
+					PrepareLumenHardwareRayTracingTranslucencyVolumeLumenMaterial(View, LumenHardwareRayTracingRayGenShaders);
+					PrepareLumenHardwareRayTracingRadiosityLumenMaterial(View, LumenHardwareRayTracingRayGenShaders);
+					PrepareLumenHardwareRayTracingReflectionsLumenMaterial(View, LumenHardwareRayTracingRayGenShaders);
+					PrepareLumenHardwareRayTracingScreenProbeGatherLumenMaterial(View, LumenHardwareRayTracingRayGenShaders);
+					PrepareLumenHardwareRayTracingDirectLightingLumenMaterial(View, LumenHardwareRayTracingRayGenShaders);
 				}
-			});
-
-			// Send ray tracing resources from reference view to all others.
-			for (int32 ViewIndex = 1; ViewIndex < Views.Num(); ++ViewIndex)
-			{
-				FViewInfo& View = Views[ViewIndex];
-				View.RayTracingMaterialGatherPipeline = ReferenceView.RayTracingMaterialGatherPipeline;
-				View.LumenHardwareRayTracingMaterialPipeline = ReferenceView.LumenHardwareRayTracingMaterialPipeline;
 			}
+			DeduplicateRayGenerationShaders(LumenHardwareRayTracingRayGenShaders);
 
-			if (!bIsPathTracing)
+			if (LumenHardwareRayTracingRayGenShaders.Num())
 			{
-				SetupRayTracingLightingMissShader(RHICmdList, ReferenceView);
+				ReferenceView.LumenHardwareRayTracingMaterialPipeline = BindLumenHardwareRayTracingMaterialPipeline(RHICmdList, ReferenceView, LumenHardwareRayTracingRayGenShaders, ReferenceView.LumenHardwareRayTracingHitDataBuffer);
 			}
 		}
 
+		// Send ray tracing resources from reference view to all others.
+		for (int32 ViewIndex = 1; ViewIndex < Views.Num(); ++ViewIndex)
+		{
+			FViewInfo& View = Views[ViewIndex];
+			View.RayTracingMaterialGatherPipeline = ReferenceView.RayTracingMaterialGatherPipeline;
+			View.LumenHardwareRayTracingMaterialPipeline = ReferenceView.LumenHardwareRayTracingMaterialPipeline;
+		}
+
 		if (RayTracingDynamicGeometryUpdateEndTransition)
 		{
 			RHICmdList.EndTransition(RayTracingDynamicGeometryUpdateEndTransition);
diff --git a/Engine/Source/Runtime/Renderer/Private/RayTracing/RayTracingReflections.cpp b/Engine/Source/Runtime/Renderer/Private/RayTracing/RayTracingReflections.cpp
--- a/Engine/Source/Runtime/Renderer/Private/RayTracing/RayTracingReflections.cpp
+++ b/Engine/Source/Runtime/Renderer/Private/RayTracing/RayTracingReflections.cpp
@@ -484,6 +484,11 @@ void FDeferredShadingSceneRenderer::PrepareRayTracingReflectionsDeferredMaterial
 {
 	// Declare all RayGen shaders that are used with deferred material sorts
 
+	if (!ShouldRenderRayTracingReflections(View))
+	{
+		return;
+	}
+
 	if (CVarRayTracingReflectionsExperimentalDeferred.GetValueOnRenderThread())
 	{
 		// If deferred reflections technique is used, then we only need to gather its shaders and skip the rest.
diff --git a/Engine/Source/Runtime/Renderer/Private/RayTracing/RayTracingShadows.cpp b/Engine/Source/Runtime/Renderer/Private/RayTracing/RayTracingShadows.cpp
--- a/Engine/Source/Runtime/Renderer/Private/RayTracing/RayTracingShadows.cpp
+++ b/Engine/Source/Runtime/Renderer/Private/RayTracing/RayTracingShadows.cpp
@@ -190,9 +190,7 @@ void FDeferredShadingSceneRenderer::PrepareRayTracingShadows(const FViewInfo& Vi
 	// Declare all RayGen shaders that require material closest hit shaders to be bound
 	// Ray tracing shadows shaders should be properly configured even if r.RayTracing.Shadows is 0 because lights can have raytracing shadows enabled independently of that CVar
 
-	const bool bRayTracingShadows = ShouldRenderRayTracingEffect(ERayTracingPipelineCompatibilityFlags::FullPipeline);
-
-	if (!bRayTracingShadows)
+	if (!ShouldRenderRayTracingShadows())
 	{
 		return;
 	}
