Quantize VSM metadata allocations to reduce buffer size churn

diff --git a/Engine/Source/Runtime/Renderer/Private/VirtualShadowMaps/VirtualShadowMapArray.cpp b/Engine/Source/Runtime/Renderer/Private/VirtualShadowMaps/VirtualShadowMapArray.cpp
--- a/Engine/Source/Runtime/Renderer/Private/VirtualShadowMaps/VirtualShadowMapArray.cpp
+++ b/Engine/Source/Runtime/Renderer/Private/VirtualShadowMaps/VirtualShadowMapArray.cpp
@@ -1013,6 +1013,12 @@ static FRDGTextureRef CreateDebugVisualizationTexture(FRDGBuilder& GraphBuilder,
 	return Texture;
 }
 
+static uint32 GetShadowMapsToAllocate(uint32 NumShadowMaps)
+{
+	// Round up to powers of two to be friendlier to the buffer pool
+	return FMath::RoundUpToPowerOfTwo(FMath::Max(64U, NumShadowMaps));
+}
+
 void FVirtualShadowMapArray::BuildPageAllocations(
 	FRDGBuilder& GraphBuilder,
 	const FMinimalSceneTextures& SceneTextures,
@@ -1146,14 +1152,15 @@ void FVirtualShadowMapArray::BuildPageAllocations(
 		StatsBufferRDG = GraphBuilder.CreateBuffer(FRDGBufferDesc::CreateStructuredDesc(sizeof(uint32), NumStats), TEXT("Shadow.Virtual.StatsBuffer"));
 		AddClearUAVPass(GraphBuilder, GraphBuilder.CreateUAV(StatsBufferRDG), 0);
 	}
-		
-	// Create and clear the requested page flags
-	const uint32 NumPageFlags = ShadowMaps.Num() * FVirtualShadowMap::PageTableSize;
-	FRDGBufferRef PageRequestFlagsRDG = GraphBuilder.CreateBuffer(FRDGBufferDesc::CreateStructuredDesc(sizeof(uint32), NumPageFlags), TEXT("Shadow.Virtual.PageRequestFlags"));
+	
+	// We potentially over-allocate these to avoid too many different allocation sizes each frame
+	const uint32 NumShadowMapsToAllocate = GetShadowMapsToAllocate(ShadowMaps.Num());
+	const uint32 NumPageFlagsToAllocate = NumShadowMapsToAllocate * FVirtualShadowMap::PageTableSize;
+
+	// Create and clear the requested page flags	
+	FRDGBufferRef PageRequestFlagsRDG = GraphBuilder.CreateBuffer(FRDGBufferDesc::CreateStructuredDesc(sizeof(uint32), NumPageFlagsToAllocate), TEXT("Shadow.Virtual.PageRequestFlags"));
 	AddClearUAVPass(GraphBuilder, GraphBuilder.CreateUAV(PageRequestFlagsRDG), 0);
-		
-	// TODO: Remove/move to next frame and make temporary OR replace with direct page table manipulation?
-	DynamicCasterPageFlagsRDG = GraphBuilder.CreateBuffer(FRDGBufferDesc::CreateStructuredDesc(sizeof(uint32), NumPageFlags), TEXT("Shadow.Virtual.DynamicCasterPageFlags"));
+	DynamicCasterPageFlagsRDG = GraphBuilder.CreateBuffer(FRDGBufferDesc::CreateStructuredDesc(sizeof(uint32), NumPageFlagsToAllocate), TEXT("Shadow.Virtual.DynamicCasterPageFlags"));
 	AddClearUAVPass(GraphBuilder, GraphBuilder.CreateUAV(DynamicCasterPageFlagsRDG), 0);
 
 	DirtyPageFlagsRDG = GraphBuilder.CreateBuffer(FRDGBufferDesc::CreateStructuredDesc(sizeof(uint32), NumPageFlags), TEXT("Shadow.Virtual.DirtyPageFlags"));
@@ -1167,8 +1174,9 @@ void FVirtualShadowMapArray::BuildPageAllocations(
 	// Clear to zero, technically only need to clear first Scene.GPUScene.GetNumInstances()  + 1 uints
 	AddClearUAVPass(GraphBuilder, GraphBuilder.CreateUAV(InvalidatingInstancesRDG), 0);
 
-	const uint32 NumPageRects = UniformParameters.NumShadowMaps * FVirtualShadowMap::MaxMipLevels;
-	PageRectBoundsRDG = GraphBuilder.CreateBuffer(FRDGBufferDesc::CreateStructuredDesc(sizeof(FIntVector4), NumPageRects), TEXT("Shadow.Virtual.PageRectBounds"));
+	const uint32 NumPageRects = ShadowMaps.Num() * FVirtualShadowMap::MaxMipLevels;
+	const uint32 NumPageRectsToAllocate = NumShadowMapsToAllocate * FVirtualShadowMap::MaxMipLevels;
+	PageRectBoundsRDG = GraphBuilder.CreateBuffer(FRDGBufferDesc::CreateStructuredDesc(sizeof(FIntVector4), NumPageRectsToAllocate), TEXT("Shadow.Virtual.PageRectBounds"));
 	{
 		FInitPageRectBoundsCS::FParameters* PassParameters = GraphBuilder.AllocParameters< FInitPageRectBoundsCS::FParameters >();
 		PassParameters->VirtualShadowMap = GetUniformBuffer(GraphBuilder);
@@ -1249,7 +1257,7 @@ void FVirtualShadowMapArray::BuildPageAllocations(
 					{
 						FComputeShaderUtils::AddPass(
 							GraphBuilder,
-							RDG_EVENT_NAME("GeneratePageFlagsFromPixels(%s)", ToString(InputType)),
+							RDG_EVENT_NAME("GeneratePageFlagsFromPixels(%s,NumShadowMaps=%d)", ToString(InputType), ShadowMaps.Num()),
 							ComputeShader,
 							PassParameters,
 							FIntVector(GridSize.X, GridSize.Y, 1));
@@ -1287,9 +1295,9 @@ void FVirtualShadowMapArray::BuildPageAllocations(
 		}
 	}
 
-	PageTableRDG = GraphBuilder.CreateBuffer(FRDGBufferDesc::CreateStructuredDesc(sizeof(uint32), NumPageFlags), TEXT("Shadow.Virtual.PageTable"));		
+	PageTableRDG = GraphBuilder.CreateBuffer(FRDGBufferDesc::CreateStructuredDesc(sizeof(uint32), NumPageFlagsToAllocate), TEXT("Shadow.Virtual.PageTable"));		
 	// Note: these are passed to the rendering and are not identical to the PageRequest flags coming in from GeneratePageFlagsFromPixels 
-	PageFlagsRDG = GraphBuilder.CreateBuffer(FRDGBufferDesc::CreateStructuredDesc(sizeof(uint32), NumPageFlags), TEXT("Shadow.Virtual.PageFlags"));
+	PageFlagsRDG = GraphBuilder.CreateBuffer(FRDGBufferDesc::CreateStructuredDesc(sizeof(uint32), NumPageFlagsToAllocate), TEXT("Shadow.Virtual.PageFlags"));
 
 	// One additional element as the last element is used as an atomic counter
 	FRDGBufferRef FreePhysicalPagesRDG = GraphBuilder.CreateBuffer(FRDGBufferDesc::CreateStructuredDesc(sizeof(int32), GetMaxPhysicalPages() + 1), TEXT("Shadow.Virtual.FreePhysicalPages"));
