Fix merge issue

diff --git a/Engine/Source/Runtime/Renderer/Private/VirtualShadowMaps/VirtualShadowMapArray.cpp b/Engine/Source/Runtime/Renderer/Private/VirtualShadowMaps/VirtualShadowMapArray.cpp
--- a/Engine/Source/Runtime/Renderer/Private/VirtualShadowMaps/VirtualShadowMapArray.cpp
+++ b/Engine/Source/Runtime/Renderer/Private/VirtualShadowMaps/VirtualShadowMapArray.cpp
@@ -1163,7 +1163,7 @@ void FVirtualShadowMapArray::BuildPageAllocations(
 	DynamicCasterPageFlagsRDG = GraphBuilder.CreateBuffer(FRDGBufferDesc::CreateStructuredDesc(sizeof(uint32), NumPageFlagsToAllocate), TEXT("Shadow.Virtual.DynamicCasterPageFlags"));
 	AddClearUAVPass(GraphBuilder, GraphBuilder.CreateUAV(DynamicCasterPageFlagsRDG), 0);
 
-	DirtyPageFlagsRDG = GraphBuilder.CreateBuffer(FRDGBufferDesc::CreateStructuredDesc(sizeof(uint32), NumPageFlags), TEXT("Shadow.Virtual.DirtyPageFlags"));
+	DirtyPageFlagsRDG = GraphBuilder.CreateBuffer(FRDGBufferDesc::CreateStructuredDesc(sizeof(uint32), NumPageFlagsToAllocate), TEXT("Shadow.Virtual.DirtyPageFlags"));
 	AddClearUAVPass(GraphBuilder, GraphBuilder.CreateUAV(DirtyPageFlagsRDG), 0);
 		
 	// Record the number of instances the buffer has capactiy for, should anything change (it shouldn't!)
