Fix VSM permutation logic and use existing function

diff --git a/Engine/Source/Runtime/Renderer/Private/ShadowDepthRendering.cpp b/Engine/Source/Runtime/Renderer/Private/ShadowDepthRendering.cpp
--- a/Engine/Source/Runtime/Renderer/Private/ShadowDepthRendering.cpp
+++ b/Engine/Source/Runtime/Renderer/Private/ShadowDepthRendering.cpp
@@ -221,14 +221,6 @@ public:
 	int32 bUseGpuSceneInstancing;
 };
 
-static bool UseNonNaniteVirtualShadowMaps(const EShaderPlatform Platform)
-{
-	return IsFeatureLevelSupported(Platform, ERHIFeatureLevel::SM5) &&
-		UseGPUScene(Platform) &&
-		GEnableNonNaniteVSM != 0 &&
-		DoesPlatformSupportNanite(Platform);
-}
-
 /**
 * A vertex shader for rendering the depth of a mesh.
 */
@@ -335,7 +327,7 @@ public:
 			return false;
 		}
 
-		if (ShaderMode == VertexShadowDepth_VirtualShadowMap && UseNonNaniteVirtualShadowMaps(Platform))
+		if (ShaderMode == VertexShadowDepth_VirtualShadowMap && !UseNonNaniteVirtualShadowMaps(Platform, GetMaxSupportedFeatureLevel(Platform)))
 		{
 			return false;
 		}
@@ -452,7 +444,7 @@ public:
 			return false;
 		}
 
-		if (ShaderMode == PixelShadowDepth_VirtualShadowMap && UseNonNaniteVirtualShadowMaps(Platform))
+		if (ShaderMode == PixelShadowDepth_VirtualShadowMap && !UseNonNaniteVirtualShadowMaps(Platform, GetMaxSupportedFeatureLevel(Platform)))
 		{
 			return false;
 		}
