Fix VSM debug view being the same for the right eye.

diff --git a/Engine/Shaders/Private/VirtualShadowMaps/VirtualShadowMapProjection.usf b/Engine/Shaders/Private/VirtualShadowMaps/VirtualShadowMapProjection.usf
--- a/Engine/Shaders/Private/VirtualShadowMaps/VirtualShadowMapProjection.usf
+++ b/Engine/Shaders/Private/VirtualShadowMaps/VirtualShadowMapProjection.usf
@@ -628,7 +628,7 @@ void VirtualShadowMapProjection(
 #if VISUALIZE_OUTPUT
 	if (InputType == INPUT_TYPE_GBUFFER)
 	{
-		OutputVisualize( PixelPos, VisualizeResult );
+		OutputVisualize( PixelPos - uint2( ProjectionRect.xy ), VisualizeResult );
 	}
 #endif // VISUALIZE_OUTPUT
 }
diff --git a/Engine/Source/Runtime/Renderer/Private/DeferredShadingRenderer.cpp b/Engine/Source/Runtime/Renderer/Private/DeferredShadingRenderer.cpp
--- a/Engine/Source/Runtime/Renderer/Private/DeferredShadingRenderer.cpp
+++ b/Engine/Source/Runtime/Renderer/Private/DeferredShadingRenderer.cpp
@@ -3157,7 +3157,7 @@ void FDeferredShadingSceneRenderer::Render(FRDGBuilder& GraphBuilder)
 
 	if (VirtualShadowMapArray.IsEnabled())
 	{
-		VirtualShadowMapArray.RenderDebugInfo(GraphBuilder);
+		VirtualShadowMapArray.RenderDebugInfo(GraphBuilder, Views);
 		if (Views.Num() > 0)
 		{
 			VirtualShadowMapArray.PrintStats(GraphBuilder, Views[0]);
@@ -3337,7 +3337,7 @@ void FDeferredShadingSceneRenderer::Render(FRDGBuilder& GraphBuilder)
 					const FPerViewPipelineState& ViewPipelineState = GetViewPipelineState(View);
 					const bool bAnyLumenActive = ViewPipelineState.DiffuseIndirectMethod == EDiffuseIndirectMethod::Lumen || ViewPipelineState.ReflectionsMethod == EReflectionsMethod::Lumen;
 
-					AddPostProcessingPasses(GraphBuilder, View, bAnyLumenActive, PostProcessingInputs, NaniteResults, InstanceCullingManager, &VirtualShadowMapArray, LumenFrameTemporaries);
+					AddPostProcessingPasses(GraphBuilder, View, ViewIndex, bAnyLumenActive, PostProcessingInputs, NaniteResults, InstanceCullingManager, &VirtualShadowMapArray, LumenFrameTemporaries);
 				}
 			}
 		}
diff --git a/Engine/Source/Runtime/Renderer/Private/LightRendering.cpp b/Engine/Source/Runtime/Renderer/Private/LightRendering.cpp
--- a/Engine/Source/Runtime/Renderer/Private/LightRendering.cpp
+++ b/Engine/Source/Runtime/Renderer/Private/LightRendering.cpp
@@ -912,12 +912,14 @@ void FDeferredShadingSceneRenderer::RenderLights(
 				if( VirtualShadowMapArray.IsAllocated() && CVarVirtualShadowOnePassProjection.GetValueOnRenderThread() )
 				{
 					// TODO: This needs to move into the view loop in clustered deferred shading pass
-					for (const FViewInfo& View : Views)
+					for (int32 ViewIndex = 0; ViewIndex < Views.Num(); ++ViewIndex)
 					{
+						const FViewInfo& View = Views[ViewIndex];
+
 						ShadowMaskBits = RenderVirtualShadowMapProjectionOnePass(
 							GraphBuilder,
 							SceneTextures,
-							View,
+							View, ViewIndex,
 							VirtualShadowMapArray,
 							EVirtualShadowMapProjectionInputType::GBuffer);
 
@@ -926,7 +928,7 @@ void FDeferredShadingSceneRenderer::RenderLights(
 							HairStrandsShadowMaskBits = RenderVirtualShadowMapProjectionOnePass(
 							GraphBuilder,
 							SceneTextures,
-							View,
+							View, ViewIndex,
 							VirtualShadowMapArray,
 							EVirtualShadowMapProjectionInputType::HairStrands);
 						}
diff --git a/Engine/Source/Runtime/Renderer/Private/PostProcess/PostProcessing.cpp b/Engine/Source/Runtime/Renderer/Private/PostProcess/PostProcessing.cpp
--- a/Engine/Source/Runtime/Renderer/Private/PostProcess/PostProcessing.cpp
+++ b/Engine/Source/Runtime/Renderer/Private/PostProcess/PostProcessing.cpp
@@ -186,7 +186,7 @@ bool ComposeSeparateTranslucencyInTSR(const FViewInfo& View);
 
 void AddPostProcessingPasses(
 	FRDGBuilder& GraphBuilder,
-	const FViewInfo& View,
+	const FViewInfo& View, int32 ViewIndex,
 	bool bAnyLumenActive,
 	const FPostProcessingInputs& Inputs,
 	const Nanite::FRasterResults* NaniteRasterResults,
@@ -1352,7 +1352,7 @@ void AddPostProcessingPasses(
 
 		if (EngineShowFlags.VisualizeVirtualShadowMap && VirtualShadowMapArray != nullptr)
 		{
-			VirtualShadowMapArray->AddVisualizePass(GraphBuilder, View, SceneColor);
+			VirtualShadowMapArray->AddVisualizePass(GraphBuilder, View, ViewIndex, SceneColor);
 		}
 
 		{
diff --git a/Engine/Source/Runtime/Renderer/Private/PostProcess/PostProcessing.h b/Engine/Source/Runtime/Renderer/Private/PostProcess/PostProcessing.h
--- a/Engine/Source/Runtime/Renderer/Private/PostProcess/PostProcessing.h
+++ b/Engine/Source/Runtime/Renderer/Private/PostProcess/PostProcessing.h
@@ -43,7 +43,7 @@ struct FPostProcessingInputs
 
 void AddPostProcessingPasses(
 	FRDGBuilder& GraphBuilder,
-	const FViewInfo& View,
+	const FViewInfo& View, int32 ViewIndex,
 	bool bAnyLumenActive,
 	const FPostProcessingInputs& Inputs,
 	const Nanite::FRasterResults* NaniteRasterResults,
diff --git a/Engine/Source/Runtime/Renderer/Private/ShadowRendering.cpp b/Engine/Source/Runtime/Renderer/Private/ShadowRendering.cpp
--- a/Engine/Source/Runtime/Renderer/Private/ShadowRendering.cpp
+++ b/Engine/Source/Runtime/Renderer/Private/ShadowRendering.cpp
@@ -2114,6 +2114,8 @@ void FSceneRenderer::RenderShadowProjections(
 
 			for (int32 ViewIndex = 0; ViewIndex < Views.Num(); ViewIndex++)
 			{
+				RDG_EVENT_SCOPE_CONDITIONAL(GraphBuilder, Views.Num() > 1, "View%d", ViewIndex);
+
 				FViewInfo& View = Views[ViewIndex];
 				FIntRect ScissorRect;
 				if (!LightSceneInfo->Proxy->GetScissorRect(ScissorRect, View, View.ViewRect))
@@ -2129,7 +2131,7 @@ void FSceneRenderer::RenderShadowProjections(
 						RenderVirtualShadowMapProjection(
 							GraphBuilder,
 							SceneTextures,
-							View,
+							View, ViewIndex,
 							VirtualShadowMapArray,
 							ScissorRect,
 							EVirtualShadowMapProjectionInputType::GBuffer,
@@ -2142,7 +2144,7 @@ void FSceneRenderer::RenderShadowProjections(
 						RenderVirtualShadowMapProjection(
 							GraphBuilder,
 							SceneTextures,
-							View,
+							View, ViewIndex,
 							VirtualShadowMapArray,
 							ScissorRect,
 							EVirtualShadowMapProjectionInputType::GBuffer,
@@ -2158,7 +2160,7 @@ void FSceneRenderer::RenderShadowProjections(
 							RenderVirtualShadowMapProjection(
 								GraphBuilder,
 								SceneTextures,
-								View,
+								View, ViewIndex,
 								VirtualShadowMapArray,
 								ScissorRect,
 								EVirtualShadowMapProjectionInputType::HairStrands,
@@ -2171,7 +2173,7 @@ void FSceneRenderer::RenderShadowProjections(
 							RenderVirtualShadowMapProjection(
 								GraphBuilder,
 								SceneTextures,
-								View,
+								View, ViewIndex,
 								VirtualShadowMapArray,
 								ScissorRect,
 								EVirtualShadowMapProjectionInputType::HairStrands,
diff --git a/Engine/Source/Runtime/Renderer/Private/VirtualShadowMaps/VirtualShadowMapArray.cpp b/Engine/Source/Runtime/Renderer/Private/VirtualShadowMaps/VirtualShadowMapArray.cpp
--- a/Engine/Source/Runtime/Renderer/Private/VirtualShadowMaps/VirtualShadowMapArray.cpp
+++ b/Engine/Source/Runtime/Renderer/Private/VirtualShadowMaps/VirtualShadowMapArray.cpp
@@ -4,6 +4,7 @@
 =============================================================================*/
 #include "VirtualShadowMapArray.h"
 #include "VirtualShadowMapVisualizationData.h"
+#include "VirtualShadowMapDefinitions.h"
 #include "../BasePassRendering.h"
 #include "../ScreenPass.h"
 #include "Components/LightComponent.h"
@@ -1053,19 +1054,22 @@ void FVirtualShadowMapArray::BuildPageAllocations(
 	{
 		FVirtualShadowMapVisualizationData& VisualizationData = GetVirtualShadowMapVisualizationData();
 	
-		// TODO: Support more than one view/debug output
-		const FName& VisualizationMode = Views[0].CurrentVirtualShadowMapVisualizationMode;
-		FIntPoint Extent = Views[0].ViewRect.Max;
-		if (VisualizationData.Update(VisualizationMode))
+		for (const FViewInfo& View : Views)
 		{
-			// TODO - automatically enable the show flag when set from command line?
-			//EngineShowFlags.SetVisualizeVirtualShadowMap(true);
-		}
+			const FName& VisualizationMode = View.CurrentVirtualShadowMapVisualizationMode;
+			// for stereo views that aren't multi-view, don't account for the left
+			FIntPoint Extent = View.ViewRect.Max - View.ViewRect.Min;
+			if (VisualizationData.Update(VisualizationMode))
+			{
+				// TODO - automatically enable the show flag when set from command line?
+				//EngineShowFlags.SetVisualizeVirtualShadowMap(true);
+			}
 
-		if (VisualizationData.IsActive() && EngineShowFlags.VisualizeVirtualShadowMap)
-		{
-			bDebugOutputEnabled = true;
-			DebugVisualizationOutput = CreateDebugVisualizationTexture(GraphBuilder, Extent);
+			if (VisualizationData.IsActive() && EngineShowFlags.VisualizeVirtualShadowMap)
+			{
+				bDebugOutputEnabled = true;
+				DebugVisualizationOutput.Add(CreateDebugVisualizationTexture(GraphBuilder, Extent));
+			}
 		}
 	}
 #endif //!UE_BUILD_SHIPPING
@@ -1196,6 +1200,8 @@ void FVirtualShadowMapArray::BuildPageAllocations(
 
 	for (int32 ViewIndex = 0; ViewIndex < Views.Num(); ++ViewIndex)
 	{
+		RDG_EVENT_SCOPE_CONDITIONAL(GraphBuilder, Views.Num() > 1, "View%d", ViewIndex);
+
 		const FViewInfo &View = Views[ViewIndex];
 
 		// This view contained no local lights (that were stored in the light grid), and no directional lights, so nothing to do.
@@ -1565,41 +1571,50 @@ class FDebugVisualizeVirtualSmCS : public FVirtualPageManagementShader
 IMPLEMENT_GLOBAL_SHADER(FDebugVisualizeVirtualSmCS, "/Engine/Private/VirtualShadowMaps/VirtualShadowMapDebug.usf", "DebugVisualizeVirtualSmCS", SF_Compute);
 
 
-void FVirtualShadowMapArray::RenderDebugInfo(FRDGBuilder& GraphBuilder)
+void FVirtualShadowMapArray::RenderDebugInfo(FRDGBuilder& GraphBuilder, TArrayView<FViewInfo> Views)
 {
 	check(IsEnabled());
 			
-	if (DebugVisualizationOutput == nullptr || !VisualizeLight.IsValid())
+	if (DebugVisualizationOutput.IsEmpty() || !VisualizeLight.IsValid())
 	{
 		return;
 	}
 
 	const FVirtualShadowMapVisualizationData& VisualizationData = GetVirtualShadowMapVisualizationData();
+	if (VisualizationData.GetActiveModeID() != VIRTUAL_SHADOW_MAP_VISUALIZE_CLIPMAP_VIRTUAL_SPACE)
+	{
+		return;
+	}
 
 	int32 BorderWidth = 2;
 
-	FIntPoint DebugTargetExtent = DebugVisualizationOutput->Desc.Extent;
+	for (int32 ViewIndex = 0; ViewIndex < Views.Num(); ViewIndex++)
+	{
+		FViewInfo& View = Views[ViewIndex];
 
-	FDebugVisualizeVirtualSmCS::FParameters* PassParameters = GraphBuilder.AllocParameters<FDebugVisualizeVirtualSmCS::FParameters>();
-	PassParameters->ProjectionParameters = GetSamplingParameters(GraphBuilder);
-	
-	PassParameters->DebugTargetWidth = DebugTargetExtent.X;
-	PassParameters->DebugTargetHeight = DebugTargetExtent.Y;
-	PassParameters->BorderWidth = BorderWidth;
-	PassParameters->VisualizeModeId = VisualizationData.GetActiveModeID();
-	PassParameters->VirtualShadowMapId = VisualizeLight.GetVirtualShadowMapId();
+		FIntPoint DebugTargetExtent = DebugVisualizationOutput[ViewIndex]->Desc.Extent;
 
-	PassParameters->OutVisualize = GraphBuilder.CreateUAV(DebugVisualizationOutput);
+		FDebugVisualizeVirtualSmCS::FParameters* PassParameters = GraphBuilder.AllocParameters<FDebugVisualizeVirtualSmCS::FParameters>();
+		PassParameters->ProjectionParameters = GetSamplingParameters(GraphBuilder);
 
-	auto ComputeShader = GetGlobalShaderMap(GMaxRHIFeatureLevel)->GetShader<FDebugVisualizeVirtualSmCS>();
+		PassParameters->DebugTargetWidth = DebugTargetExtent.X;
+		PassParameters->DebugTargetHeight = DebugTargetExtent.Y;
+		PassParameters->BorderWidth = BorderWidth;
+		PassParameters->VisualizeModeId = VisualizationData.GetActiveModeID();
+		PassParameters->VirtualShadowMapId = VisualizeLight.GetVirtualShadowMapId();
 
-	FComputeShaderUtils::AddPass(
-		GraphBuilder,
-		RDG_EVENT_NAME("DebugVisualizeVirtualShadowMap"),
-		ComputeShader,
-		PassParameters,
-		FComputeShaderUtils::GetGroupCount(DebugTargetExtent, FVirtualPageManagementShader::DefaultCSGroupXY)
-	);
+		PassParameters->OutVisualize = GraphBuilder.CreateUAV(DebugVisualizationOutput[ViewIndex]);
+
+		auto ComputeShader = GetGlobalShaderMap(GMaxRHIFeatureLevel)->GetShader<FDebugVisualizeVirtualSmCS>();
+
+		FComputeShaderUtils::AddPass(
+			GraphBuilder,
+			RDG_EVENT_NAME("DebugVisualizeVirtualShadowMap"),
+			ComputeShader,
+			PassParameters,
+			FComputeShaderUtils::GetGroupCount(DebugTargetExtent, FVirtualPageManagementShader::DefaultCSGroupXY)
+		);
+	}
 }
 
 
@@ -2771,10 +2786,10 @@ uint32 FVirtualShadowMapArray::AddRenderViews(const FProjectedShadowInfo* Projec
 	return uint32(NumMaps);
 }
 
-void FVirtualShadowMapArray::AddVisualizePass(FRDGBuilder& GraphBuilder, const FViewInfo& View, FScreenPassTexture Output)
+void FVirtualShadowMapArray::AddVisualizePass(FRDGBuilder& GraphBuilder, const FViewInfo& View, int32 ViewIndex, FScreenPassTexture Output)
 {
 #if !UE_BUILD_SHIPPING
-	if (!IsAllocated() || DebugVisualizationOutput == nullptr)
+	if (!IsAllocated() || DebugVisualizationOutput.IsEmpty())
 	{
 		return;
 	}
@@ -2783,13 +2798,13 @@ void FVirtualShadowMapArray::AddVisualizePass(FRDGBuilder& GraphBuilder, const F
 	if (VisualizationData.IsActive() && VisualizeLight.IsValid())
 	{	
 		FCopyRectPS::FParameters* Parameters = GraphBuilder.AllocParameters<FCopyRectPS::FParameters>();
-		Parameters->InputTexture = DebugVisualizationOutput;
+		Parameters->InputTexture = DebugVisualizationOutput[ViewIndex];
 		Parameters->InputSampler = TStaticSamplerState<SF_Bilinear, AM_Clamp, AM_Clamp, AM_Clamp>::GetRHI();
 		Parameters->RenderTargets[0] = FRenderTargetBinding(Output.Texture, ERenderTargetLoadAction::ENoAction);
 
 		TShaderMapRef<FCopyRectPS> PixelShader(View.ShaderMap);
 
-		FScreenPassTextureViewport InputViewport(DebugVisualizationOutput->Desc.Extent);
+		FScreenPassTextureViewport InputViewport(DebugVisualizationOutput[ViewIndex]->Desc.Extent);
 		FScreenPassTextureViewport OutputViewport(Output);
 
 		// See CVarVisualizeLayout documentation
@@ -2799,7 +2814,6 @@ void FVirtualShadowMapArray::AddVisualizePass(FRDGBuilder& GraphBuilder, const F
 			const int32 TileWidth  = View.UnscaledViewRect.Width() / 3;
 			const int32 TileHeight = View.UnscaledViewRect.Height() / 3;
 
-			OutputViewport.Rect.Min = FIntPoint(0, 0);
 			OutputViewport.Rect.Max = OutputViewport.Rect.Min + FIntPoint(TileWidth, TileHeight);
 		}
 		else if (VisualizeLayout == 2)	// Split screen
diff --git a/Engine/Source/Runtime/Renderer/Private/VirtualShadowMaps/VirtualShadowMapArray.h b/Engine/Source/Runtime/Renderer/Private/VirtualShadowMaps/VirtualShadowMapArray.h
--- a/Engine/Source/Runtime/Renderer/Private/VirtualShadowMaps/VirtualShadowMapArray.h
+++ b/Engine/Source/Runtime/Renderer/Private/VirtualShadowMaps/VirtualShadowMapArray.h
@@ -252,7 +252,7 @@ public:
 	 */
 	void RenderVirtualShadowMapsNonNanite(FRDGBuilder& GraphBuilder, const TArray<FProjectedShadowInfo*, SceneRenderingAllocator>& VirtualSmMeshCommandPasses, FScene& Scene, TArrayView<FViewInfo> Views);
 
-	void RenderDebugInfo(FRDGBuilder& GraphBuilder);
+	void RenderDebugInfo(FRDGBuilder& GraphBuilder, TArrayView<FViewInfo> Views);
 	
 	void PrintStats(FRDGBuilder& GraphBuilder, const FViewInfo& View);
 
@@ -274,7 +274,7 @@ public:
 	uint32 AddRenderViews(const FProjectedShadowInfo* ProjectedShadowInfo, float LODScaleFactor, bool bSetHzbParams, bool bUpdateHZBMetaData, bool bClampToNearPlane, TArray<Nanite::FPackedView, SceneRenderingAllocator>& OutVirtualShadowViews);
 
 	// Add visualization composite pass, if enabled
-	void AddVisualizePass(FRDGBuilder& GraphBuilder, const FViewInfo& View, FScreenPassTexture Output);
+	void AddVisualizePass(FRDGBuilder& GraphBuilder, const FViewInfo& View, int32 ViewIndex, FScreenPassTexture Output);
 
 	//
 	bool UseHzbOcclusion() const { return bUseHzbOcclusion; }
@@ -331,7 +331,7 @@ public:
 	FRDGBufferRef StatsBufferRDG = nullptr;
 
 	// Debug visualization
-	FRDGTextureRef DebugVisualizationOutput = nullptr;
+	TArray<FRDGTextureRef> DebugVisualizationOutput;
 	FVirtualShadowMapVisualizeLightSearch VisualizeLight;
 
 private:
diff --git a/Engine/Source/Runtime/Renderer/Private/VirtualShadowMaps/VirtualShadowMapProjection.cpp b/Engine/Source/Runtime/Renderer/Private/VirtualShadowMaps/VirtualShadowMapProjection.cpp
--- a/Engine/Source/Runtime/Renderer/Private/VirtualShadowMaps/VirtualShadowMapProjection.cpp
+++ b/Engine/Source/Runtime/Renderer/Private/VirtualShadowMaps/VirtualShadowMapProjection.cpp
@@ -285,7 +285,7 @@ static float GetNormalBiasForShader()
 static void RenderVirtualShadowMapProjectionCommon(
 	FRDGBuilder& GraphBuilder,
 	const FMinimalSceneTextures& SceneTextures,
-	const FViewInfo& View,
+	const FViewInfo& View, int32 ViewIndex,
 	FVirtualShadowMapArray& VirtualShadowMapArray,
 	const FIntRect ProjectionRect,
 	EVirtualShadowMapProjectionInputType InputType,
@@ -350,7 +350,7 @@ static void RenderVirtualShadowMapProjectionCommon(
 	
 	bool bDebugOutput = false;
 #if !UE_BUILD_SHIPPING
-	if ( VirtualShadowMapArray.DebugVisualizationOutput && InputType == EVirtualShadowMapProjectionInputType::GBuffer && VirtualShadowMapArray.VisualizeLight.IsValid() )
+	if ( !VirtualShadowMapArray.DebugVisualizationOutput.IsEmpty() && InputType == EVirtualShadowMapProjectionInputType::GBuffer && VirtualShadowMapArray.VisualizeLight.IsValid() )
 	{
 		const FVirtualShadowMapVisualizationData& VisualizationData = GetVirtualShadowMapVisualizationData();
 
@@ -358,7 +358,7 @@ static void RenderVirtualShadowMapProjectionCommon(
 		PassParameters->VisualizeModeId = VisualizationData.GetActiveModeID();
 		PassParameters->VisualizeVirtualShadowMapId = VirtualShadowMapArray.VisualizeLight.GetVirtualShadowMapId();
 		PassParameters->PhysicalPageMetaData = GraphBuilder.CreateSRV( VirtualShadowMapArray.PhysicalPageMetaDataRDG );
-		PassParameters->OutVisualize = GraphBuilder.CreateUAV( VirtualShadowMapArray.DebugVisualizationOutput );
+		PassParameters->OutVisualize = GraphBuilder.CreateUAV( VirtualShadowMapArray.DebugVisualizationOutput[ViewIndex] );
 	}
 #endif
 
@@ -390,7 +390,7 @@ static void RenderVirtualShadowMapProjectionCommon(
 FRDGTextureRef RenderVirtualShadowMapProjectionOnePass(
 	FRDGBuilder& GraphBuilder,
 	const FMinimalSceneTextures& SceneTextures,
-	const FViewInfo& View,
+	const FViewInfo& View, int32 ViewIndex,
 	FVirtualShadowMapArray& VirtualShadowMapArray,
 	EVirtualShadowMapProjectionInputType InputType)
 {
@@ -407,7 +407,7 @@ FRDGTextureRef RenderVirtualShadowMapProjectionOnePass(
 	RenderVirtualShadowMapProjectionCommon(
 		GraphBuilder,
 		SceneTextures,
-		View,
+		View, ViewIndex,
 		VirtualShadowMapArray,
 		ProjectionRect,
 		InputType,
@@ -440,7 +440,7 @@ static FRDGTextureRef CreateShadowMaskTexture(FRDGBuilder& GraphBuilder, FIntPoi
 void RenderVirtualShadowMapProjection(
 	FRDGBuilder& GraphBuilder,
 	const FMinimalSceneTextures& SceneTextures,
-	const FViewInfo& View,
+	const FViewInfo& View, int32 ViewIndex,
 	FVirtualShadowMapArray& VirtualShadowMapArray,
 	const FIntRect ScissorRect,
 	EVirtualShadowMapProjectionInputType InputType,
@@ -452,7 +452,7 @@ void RenderVirtualShadowMapProjection(
 	RenderVirtualShadowMapProjectionCommon(
 		GraphBuilder,
 		SceneTextures,
-		View,
+		View, ViewIndex,
 		VirtualShadowMapArray,
 		ScissorRect,
 		InputType,
@@ -471,7 +471,7 @@ void RenderVirtualShadowMapProjection(
 void RenderVirtualShadowMapProjection(
 	FRDGBuilder& GraphBuilder,
 	const FMinimalSceneTextures& SceneTextures,
-	const FViewInfo& View,
+	const FViewInfo& View, int32 ViewIndex,
 	FVirtualShadowMapArray& VirtualShadowMapArray,
 	const FIntRect ScissorRect,
 	EVirtualShadowMapProjectionInputType InputType,
@@ -483,7 +483,7 @@ void RenderVirtualShadowMapProjection(
 	RenderVirtualShadowMapProjectionCommon(
 		GraphBuilder,
 		SceneTextures,
-		View,
+		View, ViewIndex,
 		VirtualShadowMapArray,		
 		ScissorRect,
 		InputType,
diff --git a/Engine/Source/Runtime/Renderer/Private/VirtualShadowMaps/VirtualShadowMapProjection.h b/Engine/Source/Runtime/Renderer/Private/VirtualShadowMaps/VirtualShadowMapProjection.h
--- a/Engine/Source/Runtime/Renderer/Private/VirtualShadowMaps/VirtualShadowMapProjection.h
+++ b/Engine/Source/Runtime/Renderer/Private/VirtualShadowMaps/VirtualShadowMapProjection.h
@@ -22,7 +22,7 @@ const TCHAR* ToString(EVirtualShadowMapProjectionInputType In);
 void RenderVirtualShadowMapProjection(
 	FRDGBuilder& GraphBuilder,
 	const FMinimalSceneTextures& SceneTextures,
-	const FViewInfo& View,
+	const FViewInfo& View, int32 ViewIndex,
 	FVirtualShadowMapArray& VirtualShadowMapArray,
 	const FIntRect ScissorRect,
 	EVirtualShadowMapProjectionInputType InputType,
@@ -32,7 +32,7 @@ void RenderVirtualShadowMapProjection(
 void RenderVirtualShadowMapProjection(
 	FRDGBuilder& GraphBuilder,
 	const FMinimalSceneTextures& SceneTextures,
-	const FViewInfo& View,
+	const FViewInfo& View, int32 ViewIndex,
 	FVirtualShadowMapArray& VirtualShadowMapArray,
 	const FIntRect ScissorRect,
 	EVirtualShadowMapProjectionInputType InputType,
@@ -42,7 +42,7 @@ void RenderVirtualShadowMapProjection(
 FRDGTextureRef RenderVirtualShadowMapProjectionOnePass(
 	FRDGBuilder& GraphBuilder,
 	const FMinimalSceneTextures& SceneTextures,
-	const FViewInfo& View,
+	const FViewInfo& View, int32 ViewIndex,
 	FVirtualShadowMapArray& VirtualShadowMapArray,
 	EVirtualShadowMapProjectionInputType InputType);
 
