Added lighting channels support for mobile deferred shading

diff --git a/Engine/Shaders/Private/MobileDeferredShading.usf b/Engine/Shaders/Private/MobileDeferredShading.usf
--- a/Engine/Shaders/Private/MobileDeferredShading.usf
+++ b/Engine/Shaders/Private/MobileDeferredShading.usf
@@ -4,24 +4,30 @@
 
 #include "Common.ush"
 #include "SHCommon.ush"
+
+#if MATERIAL_SHADER
 #include "/Engine/Generated/Material.ush"
+#endif
 
 // Reroute MobileSceneTextures uniform buffer references to the base pass uniform buffer
 #define MobileSceneTextures		MobileBasePass.SceneTextures
 #define ForwardLightData		MobileBasePass.Forward
 #define PlanarReflectionStruct	MobileBasePass.PlanarReflection
+#ifndef GLOBAL_PREINTEGRATEDGF
+#define PreIntegratedGF			ReflectionStruct.PreIntegratedGF
+#define PreIntegratedGFSampler	ReflectionStruct.PreIntegratedGFSampler
+#endif
 
 #include "DynamicLightingCommon.ush"
 #include "ReflectionEnvironmentShared.ush"
 #include "LightGridCommon.ush"
 #include "MobileShadingModels.ush"
 #include "IESLightProfilesCommon.ush"
+#if MATERIAL_SHADER
 #include "LightFunctionCommon.ush"
+#endif
 #include "PlanarReflectionShared.ush"
 
-#if APPLY_SKY_REFLECTION
-#define ENABLE_SKY_LIGHT 1
-#endif
 #define REFLECTION_COMPOSITE_USE_BLENDED_REFLECTION_CAPTURES 1
 #define REFLECTION_COMPOSITE_SUPPORT_SKYLIGHT_BLEND 1
 #define REFLECTION_COMPOSITE_HAS_SPHERE_CAPTURES 1
@@ -39,11 +45,6 @@
 #define IS_SPOT_LIGHT 0 
 #endif 
 
-#ifndef USE_CLUSTERED
-#define USE_CLUSTERED 0
-#endif 
-
-
 struct FMobileLightData
 {
 	float3 TranslatedWorldPosition;
@@ -87,33 +88,37 @@ void FetchGBuffer(in float2 UV, out float4 GBufferA, out float4 GBufferB, out fl
 #endif
 }
 
-FGBufferData DecodeGBufferMobile(
-	float4 InGBufferA,
-	float4 InGBufferB,
-	float4 InGBufferC,
-	float4 InGBufferD)
+FGBufferData FetchAndDecodeGBuffer(in float2 UV)
 {
+	float SceneDepth = 0;
+	float4 GBufferA = 0;
+	float4 GBufferB = 0;
+	float4 GBufferC = 0;
+	float4 GBufferD = 0;
+	FetchGBuffer(UV, GBufferA, GBufferB, GBufferC, GBufferD, SceneDepth);
+	
 	FGBufferData GBuffer = (FGBufferData)0;
-	GBuffer.WorldNormal = OctahedronToUnitVector( InGBufferA.xy * 2.0f - 1.0f );
-	GBuffer.PrecomputedShadowFactors = InGBufferA.z;
-	GBuffer.PerObjectGBufferData = InGBufferA.a;  
-	GBuffer.Metallic	= InGBufferB.r;
-	GBuffer.Specular	= InGBufferB.g;
-	GBuffer.Roughness	= max(0.015625, InGBufferB.b);
+	GBuffer.WorldNormal = OctahedronToUnitVector( GBufferA.xy * 2.0f - 1.0f );
+	GBuffer.PrecomputedShadowFactors = GBufferA.z;
+	GBuffer.PerObjectGBufferData = GBufferA.a;  
+	GBuffer.Metallic	= GBufferB.r;
+	GBuffer.Specular	= GBufferB.g;
+	GBuffer.Roughness	= max(0.015625, GBufferB.b);
 	// Note: must match GetShadingModelId standalone function logic
 	// Also Note: SimpleElementPixelShader directly sets SV_Target2 ( GBufferB ) to indicate unlit.
 	// An update there will be required if this layout changes.
-	GBuffer.ShadingModelID = DecodeShadingModelId(InGBufferB.a);
-	GBuffer.SelectiveOutputMask = DecodeSelectiveOutputMask(InGBufferB.a);
-	GBuffer.BaseColor = DecodeBaseColor(InGBufferC.rgb);
+	GBuffer.ShadingModelID = DecodeShadingModelId(GBufferB.a);
+	GBuffer.SelectiveOutputMask = DecodeSelectiveOutputMask(GBufferB.a);
+	GBuffer.BaseColor = DecodeBaseColor(GBufferC.rgb);
 #if ALLOW_STATIC_LIGHTING
 	GBuffer.GBufferAO = 1;
-	GBuffer.IndirectIrradiance = DecodeIndirectIrradiance(InGBufferC.a);
+	GBuffer.IndirectIrradiance = DecodeIndirectIrradiance(GBufferC.a);
 #else
-	GBuffer.GBufferAO = InGBufferC.a;
+	GBuffer.GBufferAO = GBufferC.a;
 	GBuffer.IndirectIrradiance = 1;
 #endif
-	GBuffer.CustomData = HasCustomGBufferData(GBuffer.ShadingModelID) ? InGBufferD : 0;
+	GBuffer.CustomData = HasCustomGBufferData(GBuffer.ShadingModelID) ? GBufferD : 0;
+	GBuffer.Depth = SceneDepth;
 	return GBuffer;
 }
 
@@ -237,8 +242,73 @@ void GetLightGridLocalLighting(
 		GetDirectLighting(LightData, ShadingModelContext, GBuffer, TranslatedWorldPosition, CameraVector, MobileLightAccumulator);
 	}
 }
+
+half3 ReflectionEnvironmentSkyLighting(
+	FMobileShadingModelContext ShadingModelContext,
+	FGBufferData GBuffer,
+	float3 TranslatedWorldPosition,
+	half3 ReflectionVector,
+	uint GridIndex)
+{
+	half3 Result = 0;
+
+#if ENABLE_CLUSTERED_REFLECTION
+	uint NumCulledEntryIndex = (ForwardLightData.NumGridCells + GridIndex) * NUM_CULLED_LIGHTS_GRID_STRIDE;
+	uint NumLocalReflectionCaptures = min(ForwardLightData.NumCulledLightsGrid[NumCulledEntryIndex + 0], ForwardLightData.NumReflectionCaptures);
+	uint DataStartIndex = ForwardLightData.NumCulledLightsGrid[NumCulledEntryIndex + 1];
+
+	half3 SpecularIBL = CompositeReflectionCapturesAndSkylightTWS(
+		1.0f,
+		TranslatedWorldPosition,
+		ReflectionVector,//RayDirection,
+		GBuffer.Roughness,
+		GBuffer.IndirectIrradiance,
+		1.0f,
+		0.0f,
+		NumLocalReflectionCaptures,
+		DataStartIndex,
+		0,
+		true);
+
+	BRANCH
+	if (abs(dot(PlanarReflectionStruct.ReflectionPlane.xyz, 1)) > .0001f)
+	{
+		half4 PlanarReflection = GetPlanarReflection(TranslatedWorldPosition, GBuffer.WorldNormal, GBuffer.Roughness);
+		// Planar reflections win over reflection environment
+		SpecularIBL = lerp(SpecularIBL, PlanarReflection.rgb, PlanarReflection.a);
+	}
+
+	Result += SpecularIBL * ShadingModelContext.EnvBrdf;
+#endif
+
+#if ENABLE_SKY_LIGHT
+	#if !ENABLE_CLUSTERED_REFLECTION
+		half SkyAverageBrightness = 1.0f;
+		half3 SpecularIBL = GetSkyLightReflection(ReflectionVector, GBuffer.Roughness, SkyAverageBrightness);
+	#if ALLOW_STATIC_LIGHTING
+		SpecularIBL *= ComputeMixingWeight(GBuffer.IndirectIrradiance, SkyAverageBrightness, GBuffer.Roughness);
+	#endif	
+		Result += SpecularIBL * ShadingModelContext.EnvBrdf;
+	#endif
+
+	half3 SkyDiffuseLighting = GetSkySHDiffuseSimple(GBuffer.WorldNormal);
+	half3 DiffuseColorForSky = ShadingModelContext.DiffuseColor;
+	
+	#if MOBILE_SHADINGMODEL_SUPPORT
+	if (GBuffer.ShadingModelID == SHADINGMODELID_SUBSURFACE || GBuffer.ShadingModelID == SHADINGMODELID_PREINTEGRATED_SKIN)
+	{
+		half3 SubsurfaceColor = ExtractSubsurfaceColor(GBuffer);
+		// Add subsurface energy to diffuse
+		DiffuseColorForSky += SubsurfaceColor;
+	}
+	#endif
+	Result += SkyDiffuseLighting * half3(View.SkyLightColor.rgb) * DiffuseColorForSky * GBuffer.GBufferAO;
+#endif
+		
+	return Result;
+}
   
-void MobileDirectLightPS(
+void MobileDirectionalLightPS(
 	noperspective float4 UVAndScreenPos : TEXCOORD0, 
 	float4 SvPosition : SV_POSITION, 
 #if USE_GLES_FBF_DEFERRED
@@ -253,19 +323,9 @@ void MobileDirectLightPS(
 {
 	ResolvedView = ResolveView();
 
-	FGBufferData GBuffer = (FGBufferData)0;
-	float SceneDepth = 0; 
-	{
-		float4 GBufferA = 0; 
-		float4 GBufferB = 0; 
-		float4 GBufferC = 0; 
-		float4 GBufferD = 0;
-		FetchGBuffer(UVAndScreenPos.xy, GBufferA, GBufferB, GBufferC, GBufferD, SceneDepth);
-		GBuffer = DecodeGBufferMobile(GBufferA, GBufferB, GBufferC, GBufferD);
-	}
-	
+	FGBufferData GBuffer = FetchAndDecodeGBuffer(UVAndScreenPos.xy);
 	float2 ScreenPos = UVAndScreenPos.zw;
-	float3 TranslatedWorldPosition = mul(float4(ScreenPos * SceneDepth, SceneDepth, 1), PrimaryView.ScreenToTranslatedWorld).xyz;
+	float3 TranslatedWorldPosition = mul(float4(ScreenPos * GBuffer.Depth, GBuffer.Depth, 1), PrimaryView.ScreenToTranslatedWorld).xyz;
 
 	half3 CameraVector = normalize(PrimaryView.TranslatedWorldCameraOrigin - TranslatedWorldPosition);
 	half NoV = max(0, dot(GBuffer.WorldNormal, CameraVector));
@@ -285,8 +345,8 @@ void MobileDirectLightPS(
 #if APPLY_CSM && !ENABLE_DISTANCE_FIELD
 	// TODO: we don't need to fully compute ScreenPosition here
 	float  ShadowPositionZ = 0;
-	float4 ScreenPosition = SvPositionToScreenPosition(float4(SvPosition.xyz,SceneDepth));
-	float ShadowMap = MobileDirectionalLightCSM(ScreenPosition.xy, SceneDepth, NoL, ShadowPositionZ);
+	float4 ScreenPosition = SvPositionToScreenPosition(float4(SvPosition.xyz, GBuffer.Depth));
+	float ShadowMap = MobileDirectionalLightCSM(ScreenPosition.xy, GBuffer.Depth, NoL, ShadowPositionZ);
 	Shadow = min(ShadowMap, Shadow);
 #endif
 
@@ -299,9 +359,9 @@ void MobileDirectLightPS(
 	InitShadingModelContext(ShadingModelContext, GBuffer, CameraVector);
 	
 	float2 LocalPosition = SvPosition.xy - View.ViewRectMin.xy;
-	uint GridIndex = ComputeLightGridCellIndex(uint2(LocalPosition.x, LocalPosition.y), SceneDepth);
+	uint GridIndex = ComputeLightGridCellIndex(uint2(LocalPosition.x, LocalPosition.y), GBuffer.Depth);
 	// Local lights
-#if USE_CLUSTERED
+#if ENABLE_CLUSTERED_LIGHTS
 	{
 		const uint EyeIndex = 0;
 		const FCulledLightsGridData CulledLightGridData = GetCulledLightsGrid(GridIndex, EyeIndex);
@@ -312,59 +372,14 @@ void MobileDirectLightPS(
 	FMobileDirectLighting Lighting = MobileIntegrateBxDF(ShadingModelContext, GBuffer, NoL, NoH, VoH, CameraVector, L, Shadow);
 	MobileLightAccumulator_Add(MobileLightAccumulator, Lighting.Diffuse, Lighting.Specular * MobileDirectionalLight.DirectionalLightDistanceFadeMADAndSpecularScale.z, MobileDirectionalLight.DirectionalLightColor.rgb);
 
-#if APPLY_REFLECTION
-	uint NumCulledEntryIndex = (ForwardLightData.NumGridCells + GridIndex) * NUM_CULLED_LIGHTS_GRID_STRIDE;
-	uint NumLocalReflectionCaptures = min(ForwardLightData.NumCulledLightsGrid[NumCulledEntryIndex + 0], ForwardLightData.NumReflectionCaptures);
-	uint DataStartIndex = ForwardLightData.NumCulledLightsGrid[NumCulledEntryIndex + 1];
-
-	float3 SpecularIBL = CompositeReflectionCapturesAndSkylightTWS(
-		1.0f,
-		TranslatedWorldPosition,
-		ReflectionVector,//RayDirection,
-		GBuffer.Roughness,
-		GBuffer.IndirectIrradiance,
-		1.0f,
-		0.0f,
-		NumLocalReflectionCaptures,
-		DataStartIndex,
-		0,
-		true);
-
-	BRANCH
-	if (abs(dot(PlanarReflectionStruct.ReflectionPlane.xyz, 1)) > .0001f)
-	{
-		half4 PlanarReflection = GetPlanarReflection(TranslatedWorldPosition, GBuffer.WorldNormal, GBuffer.Roughness);
-		// Planar reflections win over reflection environment
-		SpecularIBL = lerp(SpecularIBL, PlanarReflection.rgb, PlanarReflection.a);
-	}
-	  
-	MobileLightAccumulator_Add(MobileLightAccumulator, 0.0f, SpecularIBL * ShadingModelContext.EnvBrdf, 1.0f);
-#elif APPLY_SKY_REFLECTION
-	float SkyAverageBrightness = 1.0f;
-	float3 SpecularIBL = GetSkyLightReflection(ReflectionVector, GBuffer.Roughness, SkyAverageBrightness);
-#if ALLOW_STATIC_LIGHTING
-	SpecularIBL *= ComputeMixingWeight(GBuffer.IndirectIrradiance, SkyAverageBrightness, GBuffer.Roughness);
-#endif
-	MobileLightAccumulator_Add(MobileLightAccumulator, 0.0f, SpecularIBL * ShadingModelContext.EnvBrdf, 1.0f);
+#if ENABLE_CLUSTERED_REFLECTION || ENABLE_SKY_LIGHT
+	// If we have a single directional light, apply relfection and sky contrubution here
+	half3 ReflectionAndSky = ReflectionEnvironmentSkyLighting(ShadingModelContext, GBuffer, TranslatedWorldPosition, ReflectionVector, GridIndex);
+	MobileLightAccumulator_Add(MobileLightAccumulator, ReflectionAndSky, 0.0, 1.0);
 #endif
-	half3 SkyDiffuseLighting = GetSkySHDiffuseSimple(GBuffer.WorldNormal);
-	half3 DiffuseColorForSky = ShadingModelContext.DiffuseColor;
 
-#if MOBILE_SHADINGMODEL_SUPPORT
-	if (GBuffer.ShadingModelID == SHADINGMODELID_SUBSURFACE || GBuffer.ShadingModelID == SHADINGMODELID_PREINTEGRATED_SKIN)
-	{
-		half3 SubsurfaceColor = ExtractSubsurfaceColor(GBuffer);
-		// Add subsurface energy to diffuse
-		DiffuseColorForSky += SubsurfaceColor;
-	}
-#endif
-
-	MobileLightAccumulator_Add(MobileLightAccumulator, SkyDiffuseLighting, 0.0f, half3(View.SkyLightColor.rgb) * DiffuseColorForSky * GBuffer.GBufferAO);
-	
 	half3 Color = MobileLightAccumulator_GetResult(MobileLightAccumulator);
-
 	half LightAttenuation = ComputeLightFunctionMultiplier(TranslatedWorldPosition);
-
 	// MobileHDR applies PreExposure in tonemapper
 	LightAttenuation *= View.PreExposure;
 
@@ -389,22 +404,12 @@ void MobileRadialLightPS(
 #endif
 )
 {
-	FGBufferData GBuffer = (FGBufferData)0;
-	float SceneDepth = 0; 
-	{
-		float2 ScreenUV = InScreenPosition.xy / InScreenPosition.w * View.ScreenPositionScaleBias.xy + View.ScreenPositionScaleBias.wz;
-		float4 GBufferA = 0;  
-		float4 GBufferB = 0; 
-		float4 GBufferC = 0; 
-		float4 GBufferD = 0;
-		FetchGBuffer(ScreenUV, GBufferA, GBufferB, GBufferC, GBufferD, SceneDepth);
-		GBuffer = DecodeGBufferMobile(GBufferA, GBufferB, GBufferC, GBufferD);
-	}
-	
+	float2 ScreenUV = InScreenPosition.xy / InScreenPosition.w * View.ScreenPositionScaleBias.xy + View.ScreenPositionScaleBias.wz;
+	FGBufferData GBuffer = FetchAndDecodeGBuffer(ScreenUV);
 	// With a perspective projection, the clip space position is NDC * Clip.w
 	// With an orthographic projection, clip space is the same as NDC
-	float2 ClipPosition = InScreenPosition.xy / InScreenPosition.w * (View.ViewToClip[3][3] < 1.0f ? SceneDepth : 1.0f);
-	float3 TranslatedWorldPosition = mul(float4(ClipPosition, SceneDepth, 1), PrimaryView.ScreenToTranslatedWorld).xyz;
+	float2 ClipPosition = InScreenPosition.xy / InScreenPosition.w * (View.ViewToClip[3][3] < 1.0f ? GBuffer.Depth : 1.0f);
+	float3 TranslatedWorldPosition = mul(float4(ClipPosition, GBuffer.Depth, 1), PrimaryView.ScreenToTranslatedWorld).xyz;
 	half3 CameraVector = normalize(PrimaryView.TranslatedWorldCameraOrigin - TranslatedWorldPosition);
 	half NoV = max(0, dot(GBuffer.WorldNormal, CameraVector));
 		
@@ -441,3 +446,41 @@ void MobileRadialLightPS(
 	OutColor.a = 1;
 #endif
 }
+
+void MobileReflectionEnvironmentSkyLightingPS(
+	noperspective float4 UVAndScreenPos : TEXCOORD0,
+	float4 SvPosition : SV_POSITION,
+#if USE_GLES_FBF_DEFERRED
+	out HALF4_TYPE OutProxyAdditive : SV_Target0,
+	out HALF4_TYPE OutGBufferA : SV_Target1,
+	out HALF4_TYPE OutGBufferB : SV_Target2,
+	out HALF4_TYPE OutGBufferC : SV_Target3
+#else
+	out HALF4_TYPE OutColor : SV_Target0
+#endif
+)
+{
+	ResolvedView = ResolveView();
+	FGBufferData GBuffer = FetchAndDecodeGBuffer(UVAndScreenPos.xy);
+	float2 ScreenPos = UVAndScreenPos.zw;
+	float3 TranslatedWorldPosition = mul(float4(ScreenPos * GBuffer.Depth, GBuffer.Depth, 1), PrimaryView.ScreenToTranslatedWorld).xyz;
+	half3 CameraVector = normalize(PrimaryView.TranslatedWorldCameraOrigin - TranslatedWorldPosition);
+	half NoV = max(0, dot(GBuffer.WorldNormal, CameraVector));
+	half3 ReflectionVector = GBuffer.WorldNormal * (NoV * 2.0) - CameraVector;
+
+	FMobileShadingModelContext ShadingModelContext = (FMobileShadingModelContext)0;
+	InitShadingModelContext(ShadingModelContext, GBuffer, CameraVector);
+
+	float2 LocalPosition = SvPosition.xy - View.ViewRectMin.xy;
+	uint GridIndex = ComputeLightGridCellIndex(uint2(LocalPosition.x, LocalPosition.y), GBuffer.Depth);
+	half3 ReflectionAndSky = ReflectionEnvironmentSkyLighting(ShadingModelContext, GBuffer, TranslatedWorldPosition, ReflectionVector, GridIndex);
+
+	ReflectionAndSky *= View.PreExposure;
+
+#if defined(MOBILE_DEFERRED_SHADING) && USE_GLES_FBF_DEFERRED == 1
+	OutProxyAdditive.rgb = ReflectionAndSky.rgb;
+#else
+	OutColor.rgb = ReflectionAndSky.rgb;
+	OutColor.a = 1;
+#endif
+}
\ No newline at end of file
diff --git a/Engine/Source/Runtime/Renderer/Private/DepthRendering.cpp b/Engine/Source/Runtime/Renderer/Private/DepthRendering.cpp
--- a/Engine/Source/Runtime/Renderer/Private/DepthRendering.cpp
+++ b/Engine/Source/Runtime/Renderer/Private/DepthRendering.cpp
@@ -723,6 +723,7 @@ void SetMobileDepthPassRenderState(const FPrimitiveSceneProxy* RESTRICT Primitiv
 		const FMaterial& MaterialResource = *MeshBatch.MaterialRenderProxy->GetMaterialNoFallback(ERHIFeatureLevel::ES3_1);
 		uint8 ShadingModel = MaterialResource.GetShadingModels().IsLit() ? MSM_DefaultLit : MSM_Unlit;
 		StencilValue |= GET_STENCIL_MOBILE_SM_MASK(ShadingModel);
+		StencilValue |= STENCIL_LIGHTING_CHANNELS_MASK(PrimitiveSceneProxy ? PrimitiveSceneProxy->GetLightingChannelStencilValue() : 0x00);
 	}
 
 	DrawRenderState.SetStencilRef(StencilValue);
diff --git a/Engine/Source/Runtime/Renderer/Private/MobileBasePass.cpp b/Engine/Source/Runtime/Renderer/Private/MobileBasePass.cpp
--- a/Engine/Source/Runtime/Renderer/Private/MobileBasePass.cpp
+++ b/Engine/Source/Runtime/Renderer/Private/MobileBasePass.cpp
@@ -338,14 +338,15 @@ void MobileBasePass::SetOpaqueRenderState(FMeshPassProcessorRenderState& DrawRen
 	if (bEnableReceiveDecalOutput)
 	{
 		uint8 ReceiveDecals = (PrimitiveSceneProxy && !PrimitiveSceneProxy->ReceivesDecals() ? 0x01 : 0x00);
-		StencilValue|= GET_STENCIL_BIT_MASK(RECEIVE_DECAL, ReceiveDecals);
+		StencilValue |= GET_STENCIL_BIT_MASK(RECEIVE_DECAL, ReceiveDecals);
 	}
 	
 	if (bUsesDeferredShading)
 	{
 		// store into [1-3] bits
 		uint8 ShadingModel = Material.GetShadingModels().IsLit() ? MSM_DefaultLit : MSM_Unlit;
-		StencilValue|= GET_STENCIL_MOBILE_SM_MASK(ShadingModel);
+		StencilValue |= GET_STENCIL_MOBILE_SM_MASK(ShadingModel);
+		StencilValue |= STENCIL_LIGHTING_CHANNELS_MASK(PrimitiveSceneProxy ? PrimitiveSceneProxy->GetLightingChannelStencilValue() : 0x00);
 	}
 		
 	if (bEnableReceiveDecalOutput || bUsesDeferredShading)
diff --git a/Engine/Source/Runtime/Renderer/Private/MobileDeferredShadingPass.cpp b/Engine/Source/Runtime/Renderer/Private/MobileDeferredShadingPass.cpp
--- a/Engine/Source/Runtime/Renderer/Private/MobileDeferredShadingPass.cpp
+++ b/Engine/Source/Runtime/Renderer/Private/MobileDeferredShadingPass.cpp
@@ -29,17 +29,23 @@ BEGIN_SHADER_PARAMETER_STRUCT(FMobileDeferredPassParameters, )
 	RENDER_TARGET_BINDING_SLOTS()
 END_SHADER_PARAMETER_STRUCT()
 
-class FMobileDirectLightFunctionPS : public FMaterialShader
+class FMobileDirectionalLightFunctionPS : public FMaterialShader
 {
-	DECLARE_SHADER_TYPE(FMobileDirectLightFunctionPS, Material);
-	SHADER_USE_PARAMETER_STRUCT_WITH_LEGACY_BASE(FMobileDirectLightFunctionPS, FMaterialShader)
-
-	class FUseClustred			: SHADER_PERMUTATION_BOOL("USE_CLUSTERED");
-	class FApplySkyReflection	: SHADER_PERMUTATION_BOOL("APPLY_SKY_REFLECTION");
-	class FApplyCSM				: SHADER_PERMUTATION_BOOL("APPLY_CSM");
-	class FApplyReflection		: SHADER_PERMUTATION_BOOL("APPLY_REFLECTION");
-	class FShadowQuality		: SHADER_PERMUTATION_INT("MOBILE_SHADOW_QUALITY", 4);
-	using FPermutationDomain = TShaderPermutationDomain< FUseClustred, FApplySkyReflection, FApplyCSM, FApplyReflection, FShadowQuality>;
+	DECLARE_SHADER_TYPE(FMobileDirectionalLightFunctionPS, Material);
+	SHADER_USE_PARAMETER_STRUCT_WITH_LEGACY_BASE(FMobileDirectionalLightFunctionPS, FMaterialShader)
+
+	class FEnableClustredLights		: SHADER_PERMUTATION_BOOL("ENABLE_CLUSTERED_LIGHTS");
+	class FEnableClustredReflection	: SHADER_PERMUTATION_BOOL("ENABLE_CLUSTERED_REFLECTION");
+	class FEnableSkyLight			: SHADER_PERMUTATION_BOOL("ENABLE_SKY_LIGHT");
+	class FApplyCSM					: SHADER_PERMUTATION_BOOL("APPLY_CSM");
+	class FShadowQuality			: SHADER_PERMUTATION_INT("MOBILE_SHADOW_QUALITY", 4);
+	
+	using FPermutationDomain = TShaderPermutationDomain< 
+		FEnableClustredLights, 
+		FEnableClustredReflection, 
+		FEnableSkyLight,
+		FApplyCSM, 
+		FShadowQuality>;
 
 	BEGIN_SHADER_PARAMETER_STRUCT(FParameters, )
 		SHADER_PARAMETER_STRUCT_REF(FMobileDirectionalLightShaderParameters, MobileDirectionalLight)
@@ -48,8 +54,6 @@ class FMobileDirectLightFunctionPS : public FMaterialShader
 		SHADER_PARAMETER(FMatrix44f, TranslatedWorldToLight)
 		SHADER_PARAMETER(FVector4f, LightFunctionParameters)
 		SHADER_PARAMETER(FVector3f, LightFunctionParameters2)
-		SHADER_PARAMETER_TEXTURE(Texture2D, PreIntegratedGF)
-		SHADER_PARAMETER_SAMPLER(SamplerState, PreIntegratedGFSampler)
 		SHADER_PARAMETER_TEXTURE(Texture2D, ScreenSpaceShadowMaskTexture)
 		SHADER_PARAMETER_SAMPLER(SamplerState, ScreenSpaceShadowMaskSampler)
 	END_SHADER_PARAMETER_STRUCT()
@@ -62,6 +66,7 @@ class FMobileDirectLightFunctionPS : public FMaterialShader
 		OutEnvironment.SetDefine(TEXT("SUPPORTS_TEXTURECUBE_ARRAY"), 1);
 		OutEnvironment.SetDefine(TEXT("USE_LIGHT_FUNCTION"), Parameters.MaterialParameters.bIsDefaultMaterial ? 0 : 1);
 		OutEnvironment.SetDefine(TEXT("ENABLE_DISTANCE_FIELD"), IsMobileDistanceFieldEnabled(Parameters.Platform));
+		OutEnvironment.SetDefine(TEXT("MATERIAL_SHADER"), 1);
 	}
 
 	static FPermutationDomain RemapPermutationVector(FPermutationDomain PermutationVector)
@@ -92,23 +97,23 @@ class FMobileDirectLightFunctionPS : public FMaterialShader
 		return true;
 	}
 
-	static FPermutationDomain BuildPermutationVector(const FViewInfo& View, bool bDirectionalLight)
+	static FPermutationDomain BuildPermutationVector(const FViewInfo& View, bool bInlineReflectionAndSky, bool bDynamicShadows, bool bSkyLight)
 	{
-		bool bUseClustered = bDirectionalLight && GMobileUseClusteredDeferredShading != 0;
-		bool bApplySky = View.Family->EngineShowFlags.SkyLighting;
-		int32 ShadowQuality = bDirectionalLight ? (int32)GetShadowQuality() : 0;
-		int32 NumReflectionCaptures = View.NumBoxReflectionCaptures + View.NumSphereReflectionCaptures;
-		
+		bool bUseClusteredLights = GMobileUseClusteredDeferredShading != 0;
+		bool bClustredReflection = bInlineReflectionAndSky && (View.NumBoxReflectionCaptures + View.NumSphereReflectionCaptures) > 0;
+		bool bEnableSkyLight = bInlineReflectionAndSky && bSkyLight;
+		int32 ShadowQuality = bDynamicShadows ? (int32)GetShadowQuality() : 0;
+				
 		FPermutationDomain PermutationVector;
-		PermutationVector.Set<FMobileDirectLightFunctionPS::FUseClustred>(bUseClustered);
-		PermutationVector.Set<FMobileDirectLightFunctionPS::FApplySkyReflection>(bApplySky);
-		PermutationVector.Set<FMobileDirectLightFunctionPS::FApplyCSM>(ShadowQuality > 0);
-		PermutationVector.Set<FMobileDirectLightFunctionPS::FApplyReflection>(NumReflectionCaptures > 0);
-		PermutationVector.Set<FMobileDirectLightFunctionPS::FShadowQuality>(FMath::Clamp(ShadowQuality - 1, 0, 3));
+		PermutationVector.Set<FMobileDirectionalLightFunctionPS::FEnableClustredLights>(bUseClusteredLights);
+		PermutationVector.Set<FMobileDirectionalLightFunctionPS::FEnableClustredReflection>(bClustredReflection);
+		PermutationVector.Set<FMobileDirectionalLightFunctionPS::FEnableSkyLight>(bEnableSkyLight);
+		PermutationVector.Set<FMobileDirectionalLightFunctionPS::FApplyCSM>(ShadowQuality > 0);
+		PermutationVector.Set<FMobileDirectionalLightFunctionPS::FShadowQuality>(FMath::Clamp(ShadowQuality - 1, 0, 3));
 		return PermutationVector;
 	}
 
-	static void SetParameters(FRHICommandList& RHICmdList, const TShaderRef<FMobileDirectLightFunctionPS>& Shader, const FViewInfo& View, const FMaterialRenderProxy* Proxy, const FMaterial& Material, const FParameters& Parameters)
+	static void SetParameters(FRHICommandList& RHICmdList, const TShaderRef<FMobileDirectionalLightFunctionPS>& Shader, const FViewInfo& View, const FMaterialRenderProxy* Proxy, const FMaterial& Material, const FParameters& Parameters)
 	{
 		FMaterialShader* MaterialShader = Shader.GetShader();
 		FRHIPixelShader* ShaderRHI = Shader.GetPixelShader();
@@ -117,7 +122,7 @@ class FMobileDirectLightFunctionPS : public FMaterialShader
 	}
 };
 
-IMPLEMENT_MATERIAL_SHADER_TYPE(, FMobileDirectLightFunctionPS, TEXT("/Engine/Private/MobileDeferredShading.usf"), TEXT("MobileDirectLightPS"), SF_Pixel);
+IMPLEMENT_MATERIAL_SHADER_TYPE(, FMobileDirectionalLightFunctionPS, TEXT("/Engine/Private/MobileDeferredShading.usf"), TEXT("MobileDirectionalLightPS"), SF_Pixel);
 
 /**
  * A pixel shader for projecting a light function onto the scene.
@@ -158,6 +163,8 @@ public:
 	{
 		FMaterialShader::ModifyCompilationEnvironment(Parameters, OutEnvironment);
 		OutEnvironment.SetDefine(TEXT("USE_LIGHT_FUNCTION"), Parameters.MaterialParameters.bIsDefaultMaterial ? 0 : 1);
+		OutEnvironment.SetDefine(TEXT("MATERIAL_SHADER"), 1);
+		OutEnvironment.SetDefine(TEXT("GLOBAL_PREINTEGRATEDGF"), 1);
 	}
 
 	static void SetParameters(FRHICommandList& RHICmdList, const TShaderRef<FMobileRadialLightFunctionPS>& Shader, const FViewInfo& View, const FMaterialRenderProxy* Proxy, const FMaterial& Material, const FParameters& Parameters)
@@ -172,6 +179,50 @@ public:
 
 IMPLEMENT_MATERIAL_SHADER_TYPE(,FMobileRadialLightFunctionPS, TEXT("/Engine/Private/MobileDeferredShading.usf"), TEXT("MobileRadialLightPS"), SF_Pixel);
 
+
+/**
+ * A pixel shader for reflection env and sky lighting. 
+ */
+class FMobileReflectionEnvironmentSkyLightingPS : public FGlobalShader
+{
+public:
+	DECLARE_GLOBAL_SHADER(FMobileReflectionEnvironmentSkyLightingPS);
+	SHADER_USE_PARAMETER_STRUCT(FMobileReflectionEnvironmentSkyLightingPS, FGlobalShader);
+	
+	class FEnableClustredReflection	: SHADER_PERMUTATION_BOOL("ENABLE_CLUSTERED_REFLECTION");
+	class FEnableSkyLight			: SHADER_PERMUTATION_BOOL("ENABLE_SKY_LIGHT");
+	using FPermutationDomain = TShaderPermutationDomain<FEnableClustredReflection, FEnableSkyLight>;
+
+	BEGIN_SHADER_PARAMETER_STRUCT(FParameters, )
+		SHADER_PARAMETER_STRUCT_REF(FViewUniformShaderParameters, View)
+		SHADER_PARAMETER_STRUCT_REF(FReflectionUniformParameters, ReflectionsParameters)
+		SHADER_PARAMETER_STRUCT_REF(FReflectionCaptureShaderData, ReflectionCaptureData)
+	END_SHADER_PARAMETER_STRUCT()
+
+	static bool ShouldCompilePermutation(const FGlobalShaderPermutationParameters& Parameters)
+	{
+		if (!IsMobilePlatform(Parameters.Platform) ||
+			!IsMobileDeferredShadingEnabled(Parameters.Platform))
+		{
+			return false;
+		}
+		return true;
+	}
+};
+
+IMPLEMENT_GLOBAL_SHADER(FMobileReflectionEnvironmentSkyLightingPS, "/Engine/Private/MobileDeferredShading.usf", "MobileReflectionEnvironmentSkyLightingPS", SF_Pixel);
+
+constexpr uint32 GetLightingChannel(uint32 LightingChannelMask)
+{
+	return (LightingChannelMask & 0x1) ? 0u : ((LightingChannelMask & 0x2) ? 1u : 2u);
+}
+
+constexpr uint8 GetLightingChannelStencilValue(uint32 LightingChannel)
+{
+	// LightingChannel_0 has an inverted bit in the stencil. 0 - means LightingChannel_0 is enabled. See FPrimitiveSceneProxy::GetLightingChannelStencilValue()
+	return (LightingChannel == 0u ? 0u : (1u << LightingChannel));
+}
+
 struct FCachedLightMaterial
 {
 	const FMaterial* Material;
@@ -207,19 +258,79 @@ static void GetLightMaterial(const FCachedLightMaterial& DefaultLightMaterial, c
 	OutShader = MaterialShaderMap->GetShader<ShaderType>(PermutationId);
 }
 
-static void RenderDirectLight(FRHICommandListImmediate& RHICmdList, const FScene& Scene, const FViewInfo& View, const FCachedLightMaterial& DefaultLightMaterial)
+void RenderReflectionEnvironmentSkyLighting(FRHICommandListImmediate& RHICmdList, const FScene& Scene, const FViewInfo& View)
 {
-	FLightSceneInfo* DirectionalLight = nullptr;
-	for (int32 ChannelIdx = 0; ChannelIdx < UE_ARRAY_COUNT(Scene.MobileDirectionalLights) && !DirectionalLight; ChannelIdx++)
+	// Skylights with static lighting already had their diffuse contribution baked into lightmaps
+	const bool bSkyLight = Scene.SkyLight && !Scene.SkyLight->bHasStaticLighting && View.Family->EngineShowFlags.SkyLighting;
+	const bool bClustredReflection = (View.NumBoxReflectionCaptures + View.NumSphereReflectionCaptures) > 0;
+	if (!bSkyLight && !bClustredReflection)
 	{
-		DirectionalLight = Scene.MobileDirectionalLights[ChannelIdx];
+		return;
 	}
+		
+	SCOPED_DRAW_EVENT(RHICmdList, ReflectionEnvironmentSkyLighting);
 
-	FString LightNameWithLevel = TEXT("DirectionalLight");
-	if (DirectionalLight)
-	{
-		FSceneRenderer::GetLightNameForDrawEvent(DirectionalLight->Proxy, LightNameWithLevel);
-	}
+	FGraphicsPipelineStateInitializer GraphicsPSOInit;
+	RHICmdList.ApplyCachedRenderTargets(GraphicsPSOInit);
+	// Add to emissive in SceneColor
+	GraphicsPSOInit.BlendState = TStaticBlendState<CW_RGB, BO_Add, BF_One, BF_One>::GetRHI();
+	GraphicsPSOInit.RasterizerState = TStaticRasterizerState<>::GetRHI();
+	GraphicsPSOInit.DepthStencilState = TStaticDepthStencilState<
+		false, CF_Always,
+		true, CF_Equal, SO_Keep, SO_Keep, SO_Keep,
+		false, CF_Always, SO_Keep, SO_Keep, SO_Keep,
+		GET_STENCIL_MOBILE_SM_MASK(0xff), 0x00>::GetRHI();
+
+	TShaderMapRef<FPostProcessVS> VertexShader(View.ShaderMap);
+
+	FMobileReflectionEnvironmentSkyLightingPS::FPermutationDomain PermutationVector;
+	PermutationVector.Set<FMobileReflectionEnvironmentSkyLightingPS::FEnableClustredReflection>(bClustredReflection);
+	PermutationVector.Set<FMobileReflectionEnvironmentSkyLightingPS::FEnableSkyLight>(bSkyLight);
+	TShaderMapRef<FMobileReflectionEnvironmentSkyLightingPS> PixelShader(View.ShaderMap, PermutationVector);
+	
+	GraphicsPSOInit.BoundShaderState.VertexDeclarationRHI = GFilterVertexDeclaration.VertexDeclarationRHI;
+	GraphicsPSOInit.BoundShaderState.VertexShaderRHI = VertexShader.GetVertexShader();
+	GraphicsPSOInit.BoundShaderState.PixelShaderRHI = PixelShader.GetPixelShader();
+	GraphicsPSOInit.PrimitiveType = PT_TriangleList;
+
+	uint8 StencilRef = GET_STENCIL_MOBILE_SM_MASK(MSM_DefaultLit);
+	SetGraphicsPipelineState(RHICmdList, GraphicsPSOInit, StencilRef);
+
+	FMobileReflectionEnvironmentSkyLightingPS::FParameters PassParameters;
+	PassParameters.View = GetShaderBinding(View.ViewUniformBuffer);
+	PassParameters.ReflectionCaptureData = GetShaderBinding(View.ReflectionCaptureUniformBuffer);
+	FReflectionUniformParameters ReflectionUniformParameters;
+	SetupReflectionUniformParameters(View, ReflectionUniformParameters);
+	PassParameters.ReflectionsParameters = CreateUniformBufferImmediate(ReflectionUniformParameters, UniformBuffer_SingleDraw);
+	SetShaderParameters(RHICmdList, PixelShader, PixelShader.GetPixelShader(), PassParameters);
+	
+	const FIntPoint TargetSize = GetSceneTextureExtent();
+
+	DrawRectangle(
+		RHICmdList,
+		0, 0,
+		View.ViewRect.Width(), View.ViewRect.Height(),
+		View.ViewRect.Min.X, View.ViewRect.Min.Y,
+		View.ViewRect.Width(), View.ViewRect.Height(),
+		FIntPoint(View.ViewRect.Width(), View.ViewRect.Height()),
+		TargetSize,
+		VertexShader);
+}
+
+template<uint32 LightingChannelIdx>
+static void SetDirectionalLightDepthStencilState(FGraphicsPipelineStateInitializer& GraphicsPSOInit)
+{
+	GraphicsPSOInit.DepthStencilState = TStaticDepthStencilState<
+		false, CF_Always,
+		true, CF_Equal, SO_Keep, SO_Keep, SO_Keep,
+		false, CF_Always, SO_Keep, SO_Keep, SO_Keep,
+		GET_STENCIL_MOBILE_SM_MASK(0xff) | STENCIL_LIGHTING_CHANNELS_MASK(1u << LightingChannelIdx), 0x00>::GetRHI();
+}
+
+static void RenderDirectionalLight(FRHICommandListImmediate& RHICmdList, const FScene& Scene, const FViewInfo& View, const FCachedLightMaterial& DefaultLightMaterial, const FLightSceneInfo& DirectionalLight, uint32 LightingChannel, bool bInlineReflectionAndSky)
+{
+	FString LightNameWithLevel;
+	FSceneRenderer::GetLightNameForDrawEvent(DirectionalLight.Proxy, LightNameWithLevel);
 	SCOPED_DRAW_EVENTF(RHICmdList, DirectionalLight, TEXT("%s"), *LightNameWithLevel);
 
 	FGraphicsPipelineStateInitializer GraphicsPSOInit;
@@ -227,24 +338,38 @@ static void RenderDirectLight(FRHICommandListImmediate& RHICmdList, const FScene
 	// Add to emissive in SceneColor
 	GraphicsPSOInit.BlendState = TStaticBlendState<CW_RGB, BO_Add, BF_One, BF_One>::GetRHI();
 	GraphicsPSOInit.RasterizerState = TStaticRasterizerState<>::GetRHI();
+
+	uint8 LightingChannelStencilValue = GetLightingChannelStencilValue(LightingChannel);
 	// Shade only MSM_DefaultLit pixels
-	uint8 StencilRef = GET_STENCIL_MOBILE_SM_MASK(MSM_DefaultLit);
-	GraphicsPSOInit.DepthStencilState = TStaticDepthStencilState<
-										false, CF_Always,
-										true, CF_Equal, SO_Keep, SO_Keep, SO_Keep,		
-										false, CF_Always, SO_Keep, SO_Keep, SO_Keep,
-										GET_STENCIL_MOBILE_SM_MASK(0x7), 0x00>::GetRHI(); // 4 bits for shading models
+	uint8 StencilRef = GET_STENCIL_MOBILE_SM_MASK(MSM_DefaultLit) | STENCIL_LIGHTING_CHANNELS_MASK(LightingChannelStencilValue);
+	if (LightingChannel == 1u)
+	{
+		SetDirectionalLightDepthStencilState<1u>(GraphicsPSOInit);
+	}
+	else if (LightingChannel == 2u)
+	{
+		SetDirectionalLightDepthStencilState<2u>(GraphicsPSOInit);
+	}
+	else
+	{
+		SetDirectionalLightDepthStencilState<0u>(GraphicsPSOInit);
+	}
 	
 	TShaderMapRef<FPostProcessVS> VertexShader(View.ShaderMap);
 	
 	const FMaterialRenderProxy* LightFunctionMaterialProxy = nullptr;
-	if (View.Family->EngineShowFlags.LightFunctions && DirectionalLight)
+	if (View.Family->EngineShowFlags.LightFunctions)
 	{
-		LightFunctionMaterialProxy = DirectionalLight->Proxy->GetLightFunctionMaterial();
+		LightFunctionMaterialProxy = DirectionalLight.Proxy->GetLightFunctionMaterial();
 	}
-	FMobileDirectLightFunctionPS::FPermutationDomain PermutationVector = FMobileDirectLightFunctionPS::BuildPermutationVector(View, DirectionalLight != nullptr);
+
+	// Skylights with static lighting already had their diffuse contribution baked into lightmaps
+	const bool bSkyLight = Scene.SkyLight && !Scene.SkyLight->bHasStaticLighting && View.Family->EngineShowFlags.SkyLighting;
+	const bool bDynamicShadows = DirectionalLight.Proxy->CastsDynamicShadow() && (LightingChannel == 0u) && View.Family->EngineShowFlags.DynamicShadows;
+
+	FMobileDirectionalLightFunctionPS::FPermutationDomain PermutationVector = FMobileDirectionalLightFunctionPS::BuildPermutationVector(View, bInlineReflectionAndSky, bDynamicShadows, bSkyLight);
 	FCachedLightMaterial LightMaterial;
-	TShaderRef<FMobileDirectLightFunctionPS> PixelShader;
+	TShaderRef<FMobileDirectionalLightFunctionPS> PixelShader;
 	GetLightMaterial(DefaultLightMaterial, LightFunctionMaterialProxy, PermutationVector.ToDimensionValueId(), LightMaterial, PixelShader);
 	
 	GraphicsPSOInit.BoundShaderState.VertexDeclarationRHI = GFilterVertexDeclaration.VertexDeclarationRHI;
@@ -253,15 +378,13 @@ static void RenderDirectLight(FRHICommandListImmediate& RHICmdList, const FScene
 	GraphicsPSOInit.PrimitiveType = PT_TriangleList;
 	SetGraphicsPipelineState(RHICmdList, GraphicsPSOInit, StencilRef);
 
-	FMobileDirectLightFunctionPS::FParameters PassParameters;
-	PassParameters.MobileDirectionalLight = Scene.UniformBuffers.MobileDirectionalLightUniformBuffers[1];
+	FMobileDirectionalLightFunctionPS::FParameters PassParameters;
+	PassParameters.MobileDirectionalLight = Scene.UniformBuffers.MobileDirectionalLightUniformBuffers[LightingChannel+1];
 	PassParameters.ReflectionCaptureData = GetShaderBinding(View.ReflectionCaptureUniformBuffer);
 	FReflectionUniformParameters ReflectionUniformParameters;
 	SetupReflectionUniformParameters(View, ReflectionUniformParameters);
 	PassParameters.ReflectionsParameters = CreateUniformBufferImmediate(ReflectionUniformParameters, UniformBuffer_SingleDraw);
 	PassParameters.LightFunctionParameters = FVector4f(1.0f, 1.0f, 0.0f, 0.0f);
-	PassParameters.PreIntegratedGF = GSystemTextures.PreintegratedGF->GetRHI();
-	PassParameters.PreIntegratedGFSampler = TStaticSamplerState<SF_Bilinear, AM_Clamp, AM_Clamp, AM_Clamp>::GetRHI();
 
 	if (IsMobileDistanceFieldEnabled(View.GetShaderPlatform()) && GScreenSpaceShadowMaskTextureMobileOutputs.ScreenSpaceShadowMaskTextureMobile.IsValid())
 	{
@@ -274,17 +397,16 @@ static void RenderDirectLight(FRHICommandListImmediate& RHICmdList, const FScene
 		PassParameters.ScreenSpaceShadowMaskSampler = TStaticSamplerState<SF_Point, AM_Clamp, AM_Clamp, AM_Clamp>::GetRHI();
 	}
 	
-	if (DirectionalLight)
 	{
-		const bool bUseMovableLight = DirectionalLight && !DirectionalLight->Proxy->HasStaticShadowing();
-		PassParameters.LightFunctionParameters2 = FVector3f(DirectionalLight->Proxy->GetLightFunctionFadeDistance(), DirectionalLight->Proxy->GetLightFunctionDisabledBrightness(), bUseMovableLight ? 1.0f : 0.0f);
-		const FVector Scale = DirectionalLight->Proxy->GetLightFunctionScale();
+		const bool bUseMovableLight = !DirectionalLight.Proxy->HasStaticShadowing();
+		PassParameters.LightFunctionParameters2 = FVector3f(DirectionalLight.Proxy->GetLightFunctionFadeDistance(), DirectionalLight.Proxy->GetLightFunctionDisabledBrightness(), bUseMovableLight ? 1.0f : 0.0f);
+		const FVector Scale = DirectionalLight.Proxy->GetLightFunctionScale();
 		// Switch x and z so that z of the user specified scale affects the distance along the light direction
 		const FVector InverseScale = FVector(1.f / Scale.Z, 1.f / Scale.Y, 1.f / Scale.X);
-		const FMatrix WorldToLight = DirectionalLight->Proxy->GetWorldToLight() * FScaleMatrix(FVector(InverseScale));
+		const FMatrix WorldToLight = DirectionalLight.Proxy->GetWorldToLight() * FScaleMatrix(FVector(InverseScale));
 		PassParameters.TranslatedWorldToLight = FMatrix44f(FTranslationMatrix(-View.ViewMatrices.GetPreViewTranslation()) * WorldToLight);
 	}
-	FMobileDirectLightFunctionPS::SetParameters(RHICmdList, PixelShader, View, LightMaterial.MaterialProxy, *LightMaterial.Material, PassParameters);
+	FMobileDirectionalLightFunctionPS::SetParameters(RHICmdList, PixelShader, View, LightMaterial.MaterialProxy, *LightMaterial.Material, PassParameters);
 	
 	const FIntPoint TargetSize = GetSceneTextureExtent();
 	
@@ -299,44 +421,88 @@ static void RenderDirectLight(FRHICommandListImmediate& RHICmdList, const FScene
 		VertexShader);
 }
 
-static void SetLocalLightRasterizerAndDepthState_StencilMask(FGraphicsPipelineStateInitializer& GraphicsPSOInit, const FViewInfo& View)
+static void RenderDirectionalLights(FRHICommandListImmediate& RHICmdList, const FScene& Scene, const FViewInfo& View, const FCachedLightMaterial& DefaultLightMaterial)
 {
-	// Render backfaces with depth and stencil tests
-	// and clear stencil to zero for next light mask
-	GraphicsPSOInit.RasterizerState = View.bReverseCulling ? TStaticRasterizerState<FM_Solid, CM_CW>::GetRHI() : TStaticRasterizerState<FM_Solid, CM_CCW>::GetRHI();
-	GraphicsPSOInit.DepthStencilState = TStaticDepthStencilState<
-		false, CF_LessEqual,
-		false, CF_Equal, SO_Keep, SO_Keep, SO_Keep,		
-		true, CF_Equal, SO_Zero, SO_Keep, SO_Zero,
-		GET_STENCIL_MOBILE_SM_MASK(0x7) | STENCIL_SANDBOX_MASK, STENCIL_SANDBOX_MASK>::GetRHI();
+	uint32 NumLights = 0;
+	for (uint32 ChannelIdx = 0; ChannelIdx < UE_ARRAY_COUNT(Scene.MobileDirectionalLights); ChannelIdx++)
+	{
+		NumLights += (Scene.MobileDirectionalLights[ChannelIdx] ? 1 : 0);
+	}
+	// We can merge reflection and skylight pass with a sole directional light pass
+	const bool bInlineReflectionAndSky = (NumLights == 1);
+
+	for (uint32 ChannelIdx = 0; ChannelIdx < UE_ARRAY_COUNT(Scene.MobileDirectionalLights); ChannelIdx++)
+	{
+		FLightSceneInfo* DirectionalLight = Scene.MobileDirectionalLights[ChannelIdx];
+		if (DirectionalLight)
+		{
+			RenderDirectionalLight(RHICmdList, Scene, View, DefaultLightMaterial, *DirectionalLight, ChannelIdx, bInlineReflectionAndSky);
+		}
+	}
+
+	if (!bInlineReflectionAndSky)
+	{
+		RenderReflectionEnvironmentSkyLighting(RHICmdList, Scene, View);
+	}
 }
 
+template<uint32 LightingChannel, bool bWithStencilCulling>
 static void SetLocalLightRasterizerAndDepthState(FGraphicsPipelineStateInitializer& GraphicsPSOInit, const FViewInfo& View, const FSphere& LightBounds)
 {
-	const bool bCameraInsideLightGeometry = ((FVector)View.ViewMatrices.GetViewOrigin() - LightBounds.Center).SizeSquared() < FMath::Square(LightBounds.W * 1.05f + View.NearClippingDistance * 2.0f)
-		// Always draw backfaces in ortho
-		//@todo - accurate ortho camera / light intersection
-		|| !View.IsPerspectiveProjection();
-
-	if (bCameraInsideLightGeometry)
+	if (bWithStencilCulling)
 	{
-		// Render backfaces with depth tests disabled since the camera is inside (or close to inside) the light geometry
+		// Render backfaces with depth and stencil tests
+		// and clear stencil to zero for next light mask
 		GraphicsPSOInit.RasterizerState = View.bReverseCulling ? TStaticRasterizerState<FM_Solid, CM_CW>::GetRHI() : TStaticRasterizerState<FM_Solid, CM_CCW>::GetRHI();
 		GraphicsPSOInit.DepthStencilState = TStaticDepthStencilState<
-			false, CF_Always,
-			true, CF_Equal, SO_Keep, SO_Keep, SO_Keep,		
-			false, CF_Always, SO_Keep, SO_Keep, SO_Keep,
-			GET_STENCIL_MOBILE_SM_MASK(0x7), 0x00>::GetRHI();
+			false, CF_LessEqual,
+			false, CF_Equal, SO_Keep, SO_Keep, SO_Keep,
+			true, CF_Equal, SO_Zero, SO_Keep, SO_Zero,
+			GET_STENCIL_MOBILE_SM_MASK(0xff) | STENCIL_LIGHTING_CHANNELS_MASK(1u << LightingChannel) | STENCIL_SANDBOX_MASK,
+			STENCIL_SANDBOX_MASK
+		>::GetRHI();
 	}
 	else
 	{
-		// Render frontfaces with depth tests on to get the speedup from HiZ since the camera is outside the light geometry
-		GraphicsPSOInit.RasterizerState = View.bReverseCulling ? TStaticRasterizerState<FM_Solid, CM_CCW>::GetRHI() : TStaticRasterizerState<FM_Solid, CM_CW>::GetRHI();
-		GraphicsPSOInit.DepthStencilState = TStaticDepthStencilState<
-			false, CF_DepthNearOrEqual,
-			true, CF_Equal, SO_Keep, SO_Keep, SO_Keep,		
-			false, CF_Always, SO_Keep, SO_Keep, SO_Keep,
-			GET_STENCIL_MOBILE_SM_MASK(0x7), 0x00>::GetRHI();
+
+		const bool bCameraInsideLightGeometry = ((FVector)View.ViewMatrices.GetViewOrigin() - LightBounds.Center).SizeSquared() < FMath::Square(LightBounds.W * 1.05f + View.NearClippingDistance * 2.0f)
+			// Always draw backfaces in ortho
+			//@todo - accurate ortho camera / light intersection
+			|| !View.IsPerspectiveProjection();
+
+		if (bCameraInsideLightGeometry)
+		{
+			// Render backfaces with depth tests disabled since the camera is inside (or close to inside) the light geometry
+			GraphicsPSOInit.RasterizerState = View.bReverseCulling ? TStaticRasterizerState<FM_Solid, CM_CW>::GetRHI() : TStaticRasterizerState<FM_Solid, CM_CCW>::GetRHI();
+			GraphicsPSOInit.DepthStencilState = TStaticDepthStencilState<
+				false, CF_Always,
+				true, CF_Equal, SO_Keep, SO_Keep, SO_Keep,
+				false, CF_Always, SO_Keep, SO_Keep, SO_Keep,
+				GET_STENCIL_MOBILE_SM_MASK(0xff) | STENCIL_LIGHTING_CHANNELS_MASK(1u << LightingChannel), 0x00>::GetRHI();
+		}
+		else
+		{
+			// Render frontfaces with depth tests on to get the speedup from HiZ since the camera is outside the light geometry
+			GraphicsPSOInit.RasterizerState = View.bReverseCulling ? TStaticRasterizerState<FM_Solid, CM_CCW>::GetRHI() : TStaticRasterizerState<FM_Solid, CM_CW>::GetRHI();
+			GraphicsPSOInit.DepthStencilState = TStaticDepthStencilState<
+				false, CF_DepthNearOrEqual,
+				true, CF_Equal, SO_Keep, SO_Keep, SO_Keep,
+				false, CF_Always, SO_Keep, SO_Keep, SO_Keep,
+				GET_STENCIL_MOBILE_SM_MASK(0xff) | STENCIL_LIGHTING_CHANNELS_MASK(1u << LightingChannel), 0x00>::GetRHI();
+		}
+	}
+}
+
+template<uint32 LightingChannel>
+static void SetLocalLightRasterizerAndDepthState(FGraphicsPipelineStateInitializer& GraphicsPSOInit, const FViewInfo& View, const FSphere& LightBounds)
+{
+	if (GMobileUseLightStencilCulling != 0)
+	{
+		SetLocalLightRasterizerAndDepthState<LightingChannel, true>(GraphicsPSOInit, View, LightBounds);
+	}
+	else
+	{
+		SetLocalLightRasterizerAndDepthState<LightingChannel, false>(GraphicsPSOInit, View, LightBounds);
 	}
 }
 
@@ -385,11 +551,12 @@ static void RenderLocalLight(
 	const FLightSceneInfo& LightSceneInfo, 
 	const FCachedLightMaterial& DefaultLightMaterial)
 {
-	if (!LightSceneInfo.ShouldRenderLight(View))
+	uint8 LightingChannelMask = LightSceneInfo.Proxy->GetLightingChannelMask();
+	if (!LightSceneInfo.ShouldRenderLight(View) || LightingChannelMask == 0)
 	{
 		return;
 	}
-
+	
 	const uint8 LightType = LightSceneInfo.Proxy->GetLightType();
 	const bool bIsSpotLight = LightType == LightType_Spot;
 	const bool bIsPointLight = LightType == LightType_Point;
@@ -401,7 +568,7 @@ static void RenderLocalLight(
 	FString LightNameWithLevel;
 	FSceneRenderer::GetLightNameForDrawEvent(LightSceneInfo.Proxy, LightNameWithLevel);
 	SCOPED_DRAW_EVENTF(RHICmdList, LocalLight, TEXT("%s"), *LightNameWithLevel);
-	
+
 	if (GMobileUseLightStencilCulling != 0)
 	{
 		RenderLocalLight_StencilMask(RHICmdList, Scene, View, LightSceneInfo);
@@ -414,20 +581,28 @@ static void RenderLocalLight(
 		IESTextureResource = LightSceneInfo.Proxy->GetIESTextureResource();
 		bUseIESTexture = true;
 	}
-		
+
 	FGraphicsPipelineStateInitializer GraphicsPSOInit;
 	RHICmdList.ApplyCachedRenderTargets(GraphicsPSOInit);
 	GraphicsPSOInit.BlendState = TStaticBlendState<CW_RGBA, BO_Add, BF_One, BF_One, BO_Add, BF_One, BF_One>::GetRHI();
 	GraphicsPSOInit.PrimitiveType = PT_TriangleList;
 	const FSphere LightBounds = LightSceneInfo.Proxy->GetBoundingSphere();
-	
-	if (GMobileUseLightStencilCulling != 0)
+
+	uint32 LightingChannel = GetLightingChannel(LightingChannelMask);
+	uint8 LightingChannelStencilValue = GetLightingChannelStencilValue(LightingChannel);
+
+	// TODO: support multi-channel ligths?
+	if (LightingChannel == 1u)
+	{
+		SetLocalLightRasterizerAndDepthState<1u>(GraphicsPSOInit, View, LightBounds);
+	}
+	else if (LightingChannel == 2u)
 	{
-		SetLocalLightRasterizerAndDepthState_StencilMask(GraphicsPSOInit, View);
+		SetLocalLightRasterizerAndDepthState<2u>(GraphicsPSOInit, View, LightBounds);
 	}
 	else
 	{
-		SetLocalLightRasterizerAndDepthState(GraphicsPSOInit, View, LightBounds);
+		SetLocalLightRasterizerAndDepthState<0u>(GraphicsPSOInit, View, LightBounds);
 	}
 
 	FDeferredLightVS::FPermutationDomain PermutationVectorVS;
@@ -449,9 +624,9 @@ static void RenderLocalLight(
 	GraphicsPSOInit.BoundShaderState.VertexDeclarationRHI = GetVertexDeclarationFVector4();
 	GraphicsPSOInit.BoundShaderState.VertexShaderRHI = VertexShader.GetVertexShader();
 	GraphicsPSOInit.BoundShaderState.PixelShaderRHI = PixelShader.GetPixelShader();
-
+		
 	// Shade only MSM_DefaultLit pixels
-	uint8 StencilRef = GET_STENCIL_MOBILE_SM_MASK(MSM_DefaultLit);
+	uint8 StencilRef = GET_STENCIL_MOBILE_SM_MASK(MSM_DefaultLit) | STENCIL_LIGHTING_CHANNELS_MASK(LightingChannelStencilValue);
 	SetGraphicsPipelineState(RHICmdList, GraphicsPSOInit, StencilRef);
 
 	FDeferredLightVS::FParameters ParametersVS = FDeferredLightVS::GetParameters(View, &LightSceneInfo);
@@ -483,23 +658,6 @@ static void RenderLocalLight(
 	}
 }
 
-static void SetupSimpleLightPSO(
-	FRHICommandListImmediate& RHICmdList, 
-	const FViewInfo& View,
-	const TShaderMapRef<FDeferredLightVS>& VertexShader,
-	const TShaderRef<FMobileRadialLightFunctionPS>& PixelShader, 
-	FGraphicsPipelineStateInitializer& GraphicsPSOInit)
-{
-	RHICmdList.ApplyCachedRenderTargets(GraphicsPSOInit);
-	// Use additive blending for color
-	GraphicsPSOInit.BlendState = TStaticBlendState<CW_RGBA, BO_Add, BF_One, BF_One, BO_Add, BF_One, BF_One>::GetRHI();
-	GraphicsPSOInit.PrimitiveType = PT_TriangleList;
-	GraphicsPSOInit.BoundShaderState.VertexDeclarationRHI = GetVertexDeclarationFVector4();
-	GraphicsPSOInit.BoundShaderState.VertexShaderRHI = VertexShader.GetVertexShader();
-	GraphicsPSOInit.BoundShaderState.PixelShaderRHI = PixelShader.GetPixelShader();
-	SetLocalLightRasterizerAndDepthState_StencilMask(GraphicsPSOInit, View);
-}
-
 static void RenderSimpleLights(
 	FRHICommandListImmediate& RHICmdList, 
 	const FScene& Scene, 
@@ -520,19 +678,26 @@ static void RenderSimpleLights(
 	FDeferredLightVS::FPermutationDomain PermutationVectorVS;
 	PermutationVectorVS.Set<FDeferredLightVS::FRadialLight>(true);
 	TShaderMapRef<FDeferredLightVS> VertexShader(View.ShaderMap, PermutationVectorVS);
-	TShaderRef<FMobileRadialLightFunctionPS> PixelShaders;
+	TShaderRef<FMobileRadialLightFunctionPS> PixelShader;
 	{
 		const FMaterialShaderMap* MaterialShaderMap = DefaultMaterial.Material->GetRenderingThreadShaderMap();
 		FMobileRadialLightFunctionPS::FPermutationDomain PermutationVector;
 		PermutationVector.Set<FMobileRadialLightFunctionPS::FSpotLightDim>(false);
 		PermutationVector.Set<FMobileRadialLightFunctionPS::FIESProfileDim>(false);
-		PixelShaders = MaterialShaderMap->GetShader<FMobileRadialLightFunctionPS>(PermutationVector);
+		PixelShader = MaterialShaderMap->GetShader<FMobileRadialLightFunctionPS>(PermutationVector);
 	}
 
 	// Setup PSOs we going to use for light rendering 
 	FGraphicsPipelineStateInitializer GraphicsPSOLight;
 	{
-		SetupSimpleLightPSO(RHICmdList, View, VertexShader, PixelShaders, GraphicsPSOLight);
+		RHICmdList.ApplyCachedRenderTargets(GraphicsPSOLight);
+		// Use additive blending for color
+		GraphicsPSOLight.BlendState = TStaticBlendState<CW_RGBA, BO_Add, BF_One, BF_One, BO_Add, BF_One, BF_One>::GetRHI();
+		GraphicsPSOLight.PrimitiveType = PT_TriangleList;
+		GraphicsPSOLight.BoundShaderState.VertexDeclarationRHI = GetVertexDeclarationFVector4();
+		GraphicsPSOLight.BoundShaderState.VertexShaderRHI = VertexShader.GetVertexShader();
+		GraphicsPSOLight.BoundShaderState.PixelShaderRHI = PixelShader.GetPixelShader();
+		SetLocalLightRasterizerAndDepthState<0u, true>(GraphicsPSOLight, View, FSphere());
 	}
 	// Setup stencil mask PSO
 	FGraphicsPipelineStateInitializer GraphicsPSOLightMask;
@@ -585,7 +750,7 @@ static void RenderSimpleLights(
 
 		{
 			SetGraphicsPipelineState(RHICmdList, GraphicsPSOLight, StencilRef);
-			FMobileRadialLightFunctionPS::SetParameters(RHICmdList, PixelShaders, View, DefaultMaterial.MaterialProxy, *DefaultMaterial.Material, PassParameters);
+			FMobileRadialLightFunctionPS::SetParameters(RHICmdList, PixelShader, View, DefaultMaterial.MaterialProxy, *DefaultMaterial.Material, PassParameters);
 		}
 
 		// Apply the point or spot light with some approximately bounding geometry,
@@ -612,7 +777,7 @@ void MobileDeferredShadingPass(
 	DefaultMaterial.Material = DefaultMaterial.MaterialProxy->GetMaterialNoFallback(ERHIFeatureLevel::ES3_1);
 	check(DefaultMaterial.Material);
 
-	RenderDirectLight(RHICmdList, Scene, View, DefaultMaterial);
+	RenderDirectionalLights(RHICmdList, Scene, View, DefaultMaterial);
 
 	if (GMobileUseClusteredDeferredShading == 0)
 	{
