Fix crash when setting runtime virtual texture to a new volume.
Crash was caused by virtual texture pointer being cleared on game thread before use on render thread.
Issue was probably introduced by 18816214 and will need fixing in 5.01
#rnx

diff --git a/Engine/Source/Runtime/Renderer/Private/VT/RuntimeVirtualTextureSceneProxy.cpp b/Engine/Source/Runtime/Renderer/Private/VT/RuntimeVirtualTextureSceneProxy.cpp
--- a/Engine/Source/Runtime/Renderer/Private/VT/RuntimeVirtualTextureSceneProxy.cpp
+++ b/Engine/Source/Runtime/Renderer/Private/VT/RuntimeVirtualTextureSceneProxy.cpp
@@ -91,10 +91,10 @@ FRuntimeVirtualTextureSceneProxy::FRuntimeVirtualTextureSceneProxy(URuntimeVirtu
 
 		// Store the ProducerHandle and SpaceID immediately after virtual texture is initialized.
 		ENQUEUE_RENDER_COMMAND(GetProducerHandle)(
-			[this](FRHICommandList& RHICmdList)
+			[this, VirtualTexturePtr = VirtualTexture](FRHICommandList& RHICmdList)
 			{
-				ProducerHandle = VirtualTexture->GetProducerHandle();
-				SpaceID = VirtualTexture->GetAllocatedVirtualTexture()->GetSpaceID();
+				ProducerHandle = VirtualTexturePtr->GetProducerHandle();
+				SpaceID = VirtualTexturePtr->GetAllocatedVirtualTexture()->GetSpaceID();
 			});
 	}
 }
