Strata - move visualisation code into its own file

diff --git a/Engine/Source/Runtime/Renderer/Private/Strata/Strata.cpp b/Engine/Source/Runtime/Renderer/Private/Strata/Strata.cpp
--- a/Engine/Source/Runtime/Renderer/Private/Strata/Strata.cpp
+++ b/Engine/Source/Runtime/Renderer/Private/Strata/Strata.cpp
@@ -9,7 +9,6 @@
 #include "RendererInterface.h"
 #include "UniformBuffer.h"
 #include "SceneTextureParameters.h"
-#include "ScreenPass.h"
 #include "ShaderCompiler.h"
 
 //PRAGMA_DISABLE_OPTIMIZATION
@@ -33,12 +32,6 @@ static TAutoConsoleVariable<int32> CVarStrataBytePerPixel(
 	TEXT("Strata allocated byte per pixel to store materials data. Higher value means more complex material can be represented."),
 	ECVF_ReadOnly | ECVF_RenderThreadSafe);
 
-static TAutoConsoleVariable<int32> CVarStrataClassificationDebug(
-	TEXT("r.Strata.Classification.Debug"),
-	0,
-	TEXT("Enable strata classification visualization: 1 shows simple material tiles in green and complex material tiles in red."),
-	ECVF_RenderThreadSafe);
-
 static TAutoConsoleVariable<int32> CVarStrataRoughDiffuse(
 	TEXT("r.Strata.RoughDiffuse"),
 	1,
@@ -66,12 +59,6 @@ static TAutoConsoleVariable<int32> CVarStrataTileOverflow(
 	TEXT("Scale the number of Strata tile for overflowing tiles containing multi-BSDFs pixels. (0: 0%, 1: 100%. Default 1.0f)."),
 	ECVF_RenderThreadSafe);
 
-static TAutoConsoleVariable<int32> CVarStrataDebugMode(
-	TEXT("r.Strata.DebugMode"),
-	1,
-	TEXT("Strata debug view mode."),
-	ECVF_RenderThreadSafe);
-
 IMPLEMENT_GLOBAL_SHADER_PARAMETER_STRUCT(FStrataGlobalUniformParameters, "Strata");
 
 void FStrataViewData::Reset()
@@ -415,157 +402,6 @@ TRDGUniformBufferRef<FStrataGlobalUniformParameters> BindStrataGlobalUniformPara
 	return View.StrataViewData.StrataGlobalUniformParameters;
 }
 
-////////////////////////////////////////////////////////////////////////// 
-// Debug
-
-#define VISUALIZE_MATERIAL_PASS_COUNT 3
-#define MULTIPASS_ENABLE 0
-
-class FMaterialPrintInfoCS : public FGlobalShader
-{
-	DECLARE_GLOBAL_SHADER(FMaterialPrintInfoCS);
-	SHADER_USE_PARAMETER_STRUCT(FMaterialPrintInfoCS, FGlobalShader);
-
-	using FPermutationDomain = TShaderPermutationDomain<>;
-
-	BEGIN_SHADER_PARAMETER_STRUCT(FParameters, )
-		SHADER_PARAMETER(uint32, BSDFIndex)
-		SHADER_PARAMETER_STRUCT_REF(FViewUniformShaderParameters, ViewUniformBuffer)
-		SHADER_PARAMETER_RDG_UNIFORM_BUFFER(FStrataGlobalUniformParameters, Strata)
-		SHADER_PARAMETER_STRUCT_INCLUDE(FSceneTextureParameters, SceneTextures)
-		SHADER_PARAMETER_STRUCT_INCLUDE(ShaderPrint::FShaderParameters, ShaderPrintParameters)
-		SHADER_PARAMETER_RDG_BUFFER_UAV(RWBuffer<uint>, RWPositionOffsetBuffer)
-	END_SHADER_PARAMETER_STRUCT()
-
-	static FPermutationDomain RemapPermutation(FPermutationDomain PermutationVector)
-	{
-		return PermutationVector;
-	}
-
-	static bool CanRunStrataVizualizeMaterial(EShaderPlatform Platform)
-	{
-		// On some consoles, this ALU heavy shader (and with optimisation disables for the sake of low compilation time) would spill registers. So only keep it for the editor.
-		return IsPCPlatform(Platform);
-	}
-
-	static bool ShouldCompilePermutation(const FGlobalShaderPermutationParameters& Parameters)
-	{
-		return GetMaxSupportedFeatureLevel(Parameters.Platform) >= ERHIFeatureLevel::SM5 && Strata::IsStrataEnabled() && CanRunStrataVizualizeMaterial(Parameters.Platform);
-	}
-
-	static void ModifyCompilationEnvironment(const FGlobalShaderPermutationParameters& Parameters, FShaderCompilerEnvironment& OutEnvironment)
-	{
-		FGlobalShader::ModifyCompilationEnvironment(Parameters, OutEnvironment);
-
-		// Stay debug and skip optimizations to reduce compilation time on this long shader.
-		OutEnvironment.CompilerFlags.Add(CFLAG_Debug);
-		OutEnvironment.SetDefine(TEXT("SHADER_MATERIALPRINT"), 1);
-		OutEnvironment.SetDefine(TEXT("MULTIPASS_ENABLE"), MULTIPASS_ENABLE);
-	}
-};
-IMPLEMENT_GLOBAL_SHADER(FMaterialPrintInfoCS, "/Engine/Private/Strata/StrataVisualize.usf", "MaterialPrintInfoCS", SF_Compute);
-
-class FVisualizeMaterialPS : public FGlobalShader
-{
-	DECLARE_GLOBAL_SHADER(FVisualizeMaterialPS);
-	SHADER_USE_PARAMETER_STRUCT(FVisualizeMaterialPS, FGlobalShader);
-
-	using FPermutationDomain = TShaderPermutationDomain<>;
-
-	BEGIN_SHADER_PARAMETER_STRUCT(FParameters, )
-		SHADER_PARAMETER(uint32, ViewMode)
-		SHADER_PARAMETER_STRUCT_REF(FViewUniformShaderParameters, ViewUniformBuffer)
-		SHADER_PARAMETER_RDG_UNIFORM_BUFFER(FStrataGlobalUniformParameters, Strata)
-		SHADER_PARAMETER_STRUCT_INCLUDE(FSceneTextureParameters, SceneTextures)
-		SHADER_PARAMETER_STRUCT_INCLUDE(ShaderPrint::FShaderParameters, ShaderPrintParameters)
-		RENDER_TARGET_BINDING_SLOTS()
-	END_SHADER_PARAMETER_STRUCT()
-
-	static FPermutationDomain RemapPermutation(FPermutationDomain PermutationVector)
-	{
-		return PermutationVector;
-	}
-
-	static bool CanRunStrataVizualizeMaterial(EShaderPlatform Platform)
-	{
-		// On some consoles, this ALU heavy shader (and with optimisation disables for the sake of low compilation time) would spill registers. So only keep it for the editor.
-		return IsPCPlatform(Platform);
-	}
-
-	static bool ShouldCompilePermutation(const FGlobalShaderPermutationParameters& Parameters)
-	{
-		return GetMaxSupportedFeatureLevel(Parameters.Platform) >= ERHIFeatureLevel::SM5 && Strata::IsStrataEnabled() && CanRunStrataVizualizeMaterial(Parameters.Platform);
-	}
-
-	static void ModifyCompilationEnvironment(const FGlobalShaderPermutationParameters& Parameters, FShaderCompilerEnvironment& OutEnvironment)
-	{
-		FGlobalShader::ModifyCompilationEnvironment(Parameters, OutEnvironment);
-
-		// Stay debug and skip optimizations to reduce compilation time on this long shader.
-		OutEnvironment.CompilerFlags.Add(CFLAG_Debug);
-		OutEnvironment.SetDefine(TEXT("SHADER_MATERIALVISUALIZE"), 1);
-	}
-};
-IMPLEMENT_GLOBAL_SHADER(FVisualizeMaterialPS, "/Engine/Private/Strata/StrataVisualize.usf", "VisualizeMaterialPS", SF_Pixel);
-
-static void AddVisualizeMaterialPasses(FRDGBuilder& GraphBuilder, const FViewInfo& View, FRDGTextureRef SceneColorTexture, EShaderPlatform Platform)
-{
-	FRHIBlendState* PreMultipliedColorTransmittanceBlend = TStaticBlendState<CW_RGB, BO_Add, BF_One, BF_SourceAlpha, BO_Add, BF_Zero, BF_One>::GetRHI();
-	if (View.Family->EngineShowFlags.VisualizeStrataMaterial)
-	{
-		if (!ShaderPrint::IsEnabled(View)) { ShaderPrint::SetEnabled(true); }
-		ShaderPrint::RequestSpaceForLines(64);
-		ShaderPrint::RequestSpaceForCharacters(1024);
-
-		FSceneTextureParameters SceneTextureParameters = GetSceneTextureParameters(GraphBuilder);
-
-		// Print Material info
-		{
-			uint32 BSDFIndex = 0;
-			FRDGBufferUAVRef PrintOffsetBufferUAV = nullptr;
-
-			#if MULTIPASS_ENABLE
-			FRDGBufferRef PrintOffsetBuffer = GraphBuilder.CreateBuffer(FRDGBufferDesc::CreateBufferDesc(4, 2), TEXT("Strata.DebugPrintPositionOffset"));
-			PrintOffsetBufferUAV = GraphBuilder.CreateUAV(PrintOffsetBuffer, PF_R32_UINT);
-			AddClearUAVPass(GraphBuilder, PrintOffsetBufferUAV, 50u);
-			const uint32 MaxBSDFCount = 8;
-
-			for (BSDFIndex; BSDFIndex < MaxBSDFCount; ++BSDFIndex)
-			#endif
-			{
-				FMaterialPrintInfoCS::FParameters* PassParameters = GraphBuilder.AllocParameters<FMaterialPrintInfoCS::FParameters>();
-				PassParameters->BSDFIndex = BSDFIndex;
-				PassParameters->RWPositionOffsetBuffer = PrintOffsetBufferUAV;
-				PassParameters->ViewUniformBuffer = View.ViewUniformBuffer;
-				PassParameters->Strata = Strata::BindStrataGlobalUniformParameters(View);
-				PassParameters->SceneTextures = GetSceneTextureParameters(GraphBuilder);
-				ShaderPrint::SetParameters(GraphBuilder, View.ShaderPrintData, PassParameters->ShaderPrintParameters);
-
-				TShaderMapRef<FMaterialPrintInfoCS> ComputeShader(View.ShaderMap);
-				FComputeShaderUtils::AddPass(GraphBuilder, RDG_EVENT_NAME("Strata::VisualizeMaterial(Print)"), ComputeShader, PassParameters, FIntVector(1, 1, 1));
-			}
-		}
-
-		// Draw material debug
-		const uint32 ViewMode = FMath::Max(0, CVarStrataDebugMode.GetValueOnRenderThread());
-		if (ViewMode > 1)
-		{
-			FVisualizeMaterialPS::FParameters* PassParameters = GraphBuilder.AllocParameters<FVisualizeMaterialPS::FParameters>();
-			PassParameters->ViewUniformBuffer = View.ViewUniformBuffer;
-			PassParameters->ViewMode = ViewMode;
-			PassParameters->Strata = Strata::BindStrataGlobalUniformParameters(View);
-			PassParameters->SceneTextures = GetSceneTextureParameters(GraphBuilder);
-			PassParameters->RenderTargets[0] = FRenderTargetBinding(SceneColorTexture, ERenderTargetLoadAction::ELoad);
-			ShaderPrint::SetParameters(GraphBuilder, View.ShaderPrintData, PassParameters->ShaderPrintParameters);
-
-			FVisualizeMaterialPS::FPermutationDomain PermutationVector;
-			TShaderMapRef<FVisualizeMaterialPS> PixelShader(View.ShaderMap, PermutationVector);
-
-			FPixelShaderUtils::AddFullscreenPass<FVisualizeMaterialPS>(GraphBuilder, View.ShaderMap, RDG_EVENT_NAME("Strata::VisualizeMaterial(Draw)"), PixelShader, PassParameters, View.ViewRect, PreMultipliedColorTransmittanceBlend);
-		}
-	}
-}
-
 ////////////////////////////////////////////////////////////////////////////////////////////////////////////////////////////////////////////////
 
 class FStrataClearMaterialBufferCS : public FGlobalShader
@@ -865,7 +701,7 @@ constexpr uint32 StencilBit_Fast_1	  = 0x07u | StencilBit_Fast;
 constexpr uint32 StencilBit_Single_1  = 0x07u | StencilBit_Single;
 constexpr uint32 StencilBit_Complex_1 = 0x07u | StencilBit_Complex; 
 
-static void AddStrataInternalClassificationTilePass(
+void AddStrataInternalClassificationTilePass(
 	FRDGBuilder& GraphBuilder,
 	const FViewInfo& View,
 	const FRDGTextureRef* DepthTexture,
@@ -1197,48 +1033,4 @@ static void AddStrataClearMaterialBufferPass(
 		FComputeShaderUtils::GetGroupCount(TiledViewBufferResolution, GroupSize));
 }
 
-bool ShouldRenderStrataDebugPasses(const FViewInfo& View)
-{
-	return IsStrataEnabled() &&
-		(
-		   (FVisualizeMaterialPS::CanRunStrataVizualizeMaterial(View.GetShaderPlatform()) && View.Family && View.Family->EngineShowFlags.VisualizeStrataMaterial)
-		|| (CVarStrataClassificationDebug.GetValueOnAnyThread() > 0)
-		|| ShouldRenderStrataRoughRefractionRnD()
-		);
-}
-
-FScreenPassTexture AddStrataDebugPasses(FRDGBuilder& GraphBuilder, const FViewInfo& View, FScreenPassTexture& ScreenPassSceneColor)
-{
-	check(IsStrataEnabled());
-	EShaderPlatform Platform = View.GetShaderPlatform();
-
-	if (FVisualizeMaterialPS::CanRunStrataVizualizeMaterial(Platform))
-	{
-		RDG_EVENT_SCOPE(GraphBuilder, "Strata::VisualizeMaterial");
-		AddVisualizeMaterialPasses(GraphBuilder, View, ScreenPassSceneColor.Texture, Platform);
-	}
-
-	const int32 StrataClassificationDebug = CVarStrataClassificationDebug.GetValueOnAnyThread();
-	if (StrataClassificationDebug > 0)
-	{
-		RDG_EVENT_SCOPE(GraphBuilder, "Strata::VisualizeClassification");
-		const bool bDebugPass = true;
-		if (StrataClassificationDebug > 1)
-		{
-			AddStrataInternalClassificationTilePass(GraphBuilder, View, nullptr, &ScreenPassSceneColor.Texture, EStrataTileType::EOpaqueRoughRefraction, bDebugPass);
-			AddStrataInternalClassificationTilePass(GraphBuilder, View, nullptr, &ScreenPassSceneColor.Texture, EStrataTileType::ESSSWithoutOpaqueRoughRefraction, bDebugPass);
-		}
-		else
-		{
-			AddStrataInternalClassificationTilePass(GraphBuilder, View, nullptr, &ScreenPassSceneColor.Texture, EStrataTileType::ESimple, bDebugPass);
-			AddStrataInternalClassificationTilePass(GraphBuilder, View, nullptr, &ScreenPassSceneColor.Texture, EStrataTileType::ESingle, bDebugPass);
-			AddStrataInternalClassificationTilePass(GraphBuilder, View, nullptr, &ScreenPassSceneColor.Texture, EStrataTileType::EComplex, bDebugPass);
-		}
-	}
-
-	StrataRoughRefractionRnD(GraphBuilder, View, ScreenPassSceneColor);
-
-	return MoveTemp(ScreenPassSceneColor);
-}
-
-} // namespace Strata
\ No newline at end of file
+} // namespace Strata
diff --git a/Engine/Source/Runtime/Renderer/Private/Strata/StrataVisualize.cpp b/Engine/Source/Runtime/Renderer/Private/Strata/StrataVisualize.cpp
new file mode 100644
--- /dev/null
+++ b/Engine/Source/Runtime/Renderer/Private/Strata/StrataVisualize.cpp
@@ -0,0 +1,225 @@
+// Copyright Epic Games, Inc. All Rights Reserved.
+
+#include "Strata.h"
+#include "SceneTextureParameters.h"
+#include "SceneRendering.h"
+#include "ScreenPass.h"
+#include "ShaderCompiler.h"
+#include "PixelShaderUtils.h"
+
+
+static TAutoConsoleVariable<int32> CVarStrataClassificationDebug(
+	TEXT("r.Strata.Classification.Debug"),
+	0,
+	TEXT("Enable strata classification visualization: 1 shows simple material tiles in green and complex material tiles in red."),
+	ECVF_RenderThreadSafe);
+
+static TAutoConsoleVariable<int32> CVarStrataDebugMode(
+	TEXT("r.Strata.DebugMode"),
+	1,
+	TEXT("Strata debug view mode."),
+	ECVF_RenderThreadSafe);
+
+namespace Strata
+{
+// Forward declarations
+void AddStrataInternalClassificationTilePass(
+	FRDGBuilder& GraphBuilder,
+	const FViewInfo& View,
+	const FRDGTextureRef* DepthTexture,
+	const FRDGTextureRef* ColorTexture,
+	EStrataTileType TileMaterialType,
+	const bool bDebug);
+
+#define MULTIPASS_ENABLE 0
+
+class FMaterialPrintInfoCS : public FGlobalShader
+{
+	DECLARE_GLOBAL_SHADER(FMaterialPrintInfoCS);
+	SHADER_USE_PARAMETER_STRUCT(FMaterialPrintInfoCS, FGlobalShader);
+
+	using FPermutationDomain = TShaderPermutationDomain<>;
+
+	BEGIN_SHADER_PARAMETER_STRUCT(FParameters, )
+		SHADER_PARAMETER(uint32, BSDFIndex)
+		SHADER_PARAMETER_STRUCT_REF(FViewUniformShaderParameters, ViewUniformBuffer)
+		SHADER_PARAMETER_RDG_UNIFORM_BUFFER(FStrataGlobalUniformParameters, Strata)
+		SHADER_PARAMETER_STRUCT_INCLUDE(FSceneTextureParameters, SceneTextures)
+		SHADER_PARAMETER_STRUCT_INCLUDE(ShaderPrint::FShaderParameters, ShaderPrintParameters)
+		SHADER_PARAMETER_RDG_BUFFER_UAV(RWBuffer<uint>, RWPositionOffsetBuffer)
+	END_SHADER_PARAMETER_STRUCT()
+
+	static FPermutationDomain RemapPermutation(FPermutationDomain PermutationVector)
+	{
+		return PermutationVector;
+	}
+
+	static bool CanRunStrataVizualizeMaterial(EShaderPlatform Platform)
+	{
+		// On some consoles, this ALU heavy shader (and with optimisation disables for the sake of low compilation time) would spill registers. So only keep it for the editor.
+		return IsPCPlatform(Platform);
+	}
+
+	static bool ShouldCompilePermutation(const FGlobalShaderPermutationParameters& Parameters)
+	{
+		return GetMaxSupportedFeatureLevel(Parameters.Platform) >= ERHIFeatureLevel::SM5 && Strata::IsStrataEnabled() && CanRunStrataVizualizeMaterial(Parameters.Platform);
+	}
+
+	static void ModifyCompilationEnvironment(const FGlobalShaderPermutationParameters& Parameters, FShaderCompilerEnvironment& OutEnvironment)
+	{
+		FGlobalShader::ModifyCompilationEnvironment(Parameters, OutEnvironment);
+
+		// Stay debug and skip optimizations to reduce compilation time on this long shader.
+		OutEnvironment.CompilerFlags.Add(CFLAG_Debug);
+		OutEnvironment.SetDefine(TEXT("SHADER_MATERIALPRINT"), 1);
+		OutEnvironment.SetDefine(TEXT("MULTIPASS_ENABLE"), MULTIPASS_ENABLE);
+	}
+};
+IMPLEMENT_GLOBAL_SHADER(FMaterialPrintInfoCS, "/Engine/Private/Strata/StrataVisualize.usf", "MaterialPrintInfoCS", SF_Compute);
+
+class FVisualizeMaterialPS : public FGlobalShader
+{
+	DECLARE_GLOBAL_SHADER(FVisualizeMaterialPS);
+	SHADER_USE_PARAMETER_STRUCT(FVisualizeMaterialPS, FGlobalShader);
+
+	using FPermutationDomain = TShaderPermutationDomain<>;
+
+	BEGIN_SHADER_PARAMETER_STRUCT(FParameters, )
+		SHADER_PARAMETER(uint32, ViewMode)
+		SHADER_PARAMETER_STRUCT_REF(FViewUniformShaderParameters, ViewUniformBuffer)
+		SHADER_PARAMETER_RDG_UNIFORM_BUFFER(FStrataGlobalUniformParameters, Strata)
+		SHADER_PARAMETER_STRUCT_INCLUDE(FSceneTextureParameters, SceneTextures)
+		SHADER_PARAMETER_STRUCT_INCLUDE(ShaderPrint::FShaderParameters, ShaderPrintParameters)
+		RENDER_TARGET_BINDING_SLOTS()
+	END_SHADER_PARAMETER_STRUCT()
+
+	static FPermutationDomain RemapPermutation(FPermutationDomain PermutationVector)
+	{
+		return PermutationVector;
+	}
+
+	static bool CanRunStrataVizualizeMaterial(EShaderPlatform Platform)
+	{
+		// On some consoles, this ALU heavy shader (and with optimisation disables for the sake of low compilation time) would spill registers. So only keep it for the editor.
+		return IsPCPlatform(Platform);
+	}
+
+	static bool ShouldCompilePermutation(const FGlobalShaderPermutationParameters& Parameters)
+	{
+		return GetMaxSupportedFeatureLevel(Parameters.Platform) >= ERHIFeatureLevel::SM5 && Strata::IsStrataEnabled() && CanRunStrataVizualizeMaterial(Parameters.Platform);
+	}
+
+	static void ModifyCompilationEnvironment(const FGlobalShaderPermutationParameters& Parameters, FShaderCompilerEnvironment& OutEnvironment)
+	{
+		FGlobalShader::ModifyCompilationEnvironment(Parameters, OutEnvironment);
+
+		// Stay debug and skip optimizations to reduce compilation time on this long shader.
+		OutEnvironment.CompilerFlags.Add(CFLAG_Debug);
+		OutEnvironment.SetDefine(TEXT("SHADER_MATERIALVISUALIZE"), 1);
+	}
+};
+IMPLEMENT_GLOBAL_SHADER(FVisualizeMaterialPS, "/Engine/Private/Strata/StrataVisualize.usf", "VisualizeMaterialPS", SF_Pixel);
+
+static void AddVisualizeMaterialPasses(FRDGBuilder& GraphBuilder, const FViewInfo& View, FRDGTextureRef SceneColorTexture, EShaderPlatform Platform)
+{
+	FRHIBlendState* PreMultipliedColorTransmittanceBlend = TStaticBlendState<CW_RGB, BO_Add, BF_One, BF_SourceAlpha, BO_Add, BF_Zero, BF_One>::GetRHI();
+	if (View.Family->EngineShowFlags.VisualizeStrataMaterial)
+	{
+		if (!ShaderPrint::IsEnabled(View)) { ShaderPrint::SetEnabled(true); }
+		ShaderPrint::RequestSpaceForLines(64);
+		ShaderPrint::RequestSpaceForCharacters(1024);
+
+		FSceneTextureParameters SceneTextureParameters = GetSceneTextureParameters(GraphBuilder);
+
+		// Print Material info
+		{
+			uint32 BSDFIndex = 0;
+			FRDGBufferUAVRef PrintOffsetBufferUAV = nullptr;
+
+#if MULTIPASS_ENABLE
+			FRDGBufferRef PrintOffsetBuffer = GraphBuilder.CreateBuffer(FRDGBufferDesc::CreateBufferDesc(4, 2), TEXT("Strata.DebugPrintPositionOffset"));
+			PrintOffsetBufferUAV = GraphBuilder.CreateUAV(PrintOffsetBuffer, PF_R32_UINT);
+			AddClearUAVPass(GraphBuilder, PrintOffsetBufferUAV, 50u);
+			const uint32 MaxBSDFCount = 8;
+
+			for (BSDFIndex; BSDFIndex < MaxBSDFCount; ++BSDFIndex)
+#endif
+			{
+				FMaterialPrintInfoCS::FParameters* PassParameters = GraphBuilder.AllocParameters<FMaterialPrintInfoCS::FParameters>();
+				PassParameters->BSDFIndex = BSDFIndex;
+				PassParameters->RWPositionOffsetBuffer = PrintOffsetBufferUAV;
+				PassParameters->ViewUniformBuffer = View.ViewUniformBuffer;
+				PassParameters->Strata = Strata::BindStrataGlobalUniformParameters(View);
+				PassParameters->SceneTextures = GetSceneTextureParameters(GraphBuilder);
+				ShaderPrint::SetParameters(GraphBuilder, View.ShaderPrintData, PassParameters->ShaderPrintParameters);
+
+				TShaderMapRef<FMaterialPrintInfoCS> ComputeShader(View.ShaderMap);
+				FComputeShaderUtils::AddPass(GraphBuilder, RDG_EVENT_NAME("Strata::VisualizeMaterial(Print)"), ComputeShader, PassParameters, FIntVector(1, 1, 1));
+			}
+		}
+
+		// Draw material debug
+		const uint32 ViewMode = FMath::Max(0, CVarStrataDebugMode.GetValueOnRenderThread());
+		if (ViewMode > 1)
+		{
+			FVisualizeMaterialPS::FParameters* PassParameters = GraphBuilder.AllocParameters<FVisualizeMaterialPS::FParameters>();
+			PassParameters->ViewUniformBuffer = View.ViewUniformBuffer;
+			PassParameters->ViewMode = ViewMode;
+			PassParameters->Strata = Strata::BindStrataGlobalUniformParameters(View);
+			PassParameters->SceneTextures = GetSceneTextureParameters(GraphBuilder);
+			PassParameters->RenderTargets[0] = FRenderTargetBinding(SceneColorTexture, ERenderTargetLoadAction::ELoad);
+			ShaderPrint::SetParameters(GraphBuilder, View.ShaderPrintData, PassParameters->ShaderPrintParameters);
+
+			FVisualizeMaterialPS::FPermutationDomain PermutationVector;
+			TShaderMapRef<FVisualizeMaterialPS> PixelShader(View.ShaderMap, PermutationVector);
+
+			FPixelShaderUtils::AddFullscreenPass<FVisualizeMaterialPS>(GraphBuilder, View.ShaderMap, RDG_EVENT_NAME("Strata::VisualizeMaterial(Draw)"), PixelShader, PassParameters, View.ViewRect, PreMultipliedColorTransmittanceBlend);
+		}
+	}
+}
+
+bool ShouldRenderStrataDebugPasses(const FViewInfo& View)
+{
+	return IsStrataEnabled() &&
+		(
+			(FVisualizeMaterialPS::CanRunStrataVizualizeMaterial(View.GetShaderPlatform()) && View.Family && View.Family->EngineShowFlags.VisualizeStrataMaterial)
+			|| (CVarStrataClassificationDebug.GetValueOnAnyThread() > 0)
+			|| ShouldRenderStrataRoughRefractionRnD()
+			);
+}
+
+FScreenPassTexture AddStrataDebugPasses(FRDGBuilder& GraphBuilder, const FViewInfo& View, FScreenPassTexture& ScreenPassSceneColor)
+{
+	check(IsStrataEnabled());
+	EShaderPlatform Platform = View.GetShaderPlatform();
+
+	if (FVisualizeMaterialPS::CanRunStrataVizualizeMaterial(Platform))
+	{
+		RDG_EVENT_SCOPE(GraphBuilder, "Strata::VisualizeMaterial");
+		AddVisualizeMaterialPasses(GraphBuilder, View, ScreenPassSceneColor.Texture, Platform);
+	}
+
+	const int32 StrataClassificationDebug = CVarStrataClassificationDebug.GetValueOnAnyThread();
+	if (StrataClassificationDebug > 0)
+	{
+		RDG_EVENT_SCOPE(GraphBuilder, "Strata::VisualizeClassification");
+		const bool bDebugPass = true;
+		if (StrataClassificationDebug > 1)
+		{
+			AddStrataInternalClassificationTilePass(GraphBuilder, View, nullptr, &ScreenPassSceneColor.Texture, EStrataTileType::EOpaqueRoughRefraction, bDebugPass);
+			AddStrataInternalClassificationTilePass(GraphBuilder, View, nullptr, &ScreenPassSceneColor.Texture, EStrataTileType::ESSSWithoutOpaqueRoughRefraction, bDebugPass);
+		}
+		else
+		{
+			AddStrataInternalClassificationTilePass(GraphBuilder, View, nullptr, &ScreenPassSceneColor.Texture, EStrataTileType::ESimple, bDebugPass);
+			AddStrataInternalClassificationTilePass(GraphBuilder, View, nullptr, &ScreenPassSceneColor.Texture, EStrataTileType::ESingle, bDebugPass);
+			AddStrataInternalClassificationTilePass(GraphBuilder, View, nullptr, &ScreenPassSceneColor.Texture, EStrataTileType::EComplex, bDebugPass);
+		}
+	}
+
+	StrataRoughRefractionRnD(GraphBuilder, View, ScreenPassSceneColor);
+
+	return MoveTemp(ScreenPassSceneColor);
+}
+
+} // namespace Strata
