Strata - renaming of a cvar for all debug to be under r.strata.debug.

diff --git a/Engine/Source/Runtime/Renderer/Private/Strata/StrataVisualize.cpp b/Engine/Source/Runtime/Renderer/Private/Strata/StrataVisualize.cpp
--- a/Engine/Source/Runtime/Renderer/Private/Strata/StrataVisualize.cpp
+++ b/Engine/Source/Runtime/Renderer/Private/Strata/StrataVisualize.cpp
@@ -14,8 +14,8 @@ static TAutoConsoleVariable<int32> CVarStrataClassificationDebug(
 	TEXT("Enable strata classification visualization: 1 shows simple material tiles in green and complex material tiles in red."),
 	ECVF_RenderThreadSafe);
 
-static TAutoConsoleVariable<int32> CVarStrataDebugMode(
-	TEXT("r.Strata.DebugMode"),
+static TAutoConsoleVariable<int32> CVarStrataDebugVisualizeMode(
+	TEXT("r.Strata.Debug.VisualizeMode"),
 	1,
 	TEXT("Strata debug view mode."),
 	ECVF_RenderThreadSafe);
@@ -159,7 +159,7 @@ static void AddVisualizeMaterialPasses(FRDGBuilder& GraphBuilder, const FViewInf
 		}
 
 		// Draw material debug
-		const uint32 ViewMode = FMath::Max(0, CVarStrataDebugMode.GetValueOnRenderThread());
+		const uint32 ViewMode = FMath::Max(0, CVarStrataDebugVisualizeMode.GetValueOnRenderThread());
 		if (ViewMode > 1)
 		{
 			FVisualizeMaterialPS::FParameters* PassParameters = GraphBuilder.AllocParameters<FVisualizeMaterialPS::FParameters>();
