Renderer dependency cleanup focused on removing RenderGraph.h from SceneRenderTargetParameters.h

diff --git a/Engine/Plugins/Experimental/GPULightmass/Source/GPULightmass/Private/IrradianceCaching.h b/Engine/Plugins/Experimental/GPULightmass/Source/GPULightmass/Private/IrradianceCaching.h
--- a/Engine/Plugins/Experimental/GPULightmass/Source/GPULightmass/Private/IrradianceCaching.h
+++ b/Engine/Plugins/Experimental/GPULightmass/Source/GPULightmass/Private/IrradianceCaching.h
@@ -4,6 +4,7 @@
 #include "GlobalShader.h"
 #include "ShaderParameterStruct.h"
 #include "SceneView.h"
+#include "SceneRenderTargetParameters.h"
 
 BEGIN_UNIFORM_BUFFER_STRUCT(FIrradianceCachingParameters, )
 	SHADER_PARAMETER(uint32, HashTableSize)
diff --git a/Engine/Plugins/FX/Niagara/Source/Niagara/Public/NiagaraRenderViewDataManager.h b/Engine/Plugins/FX/Niagara/Source/Niagara/Public/NiagaraRenderViewDataManager.h
--- a/Engine/Plugins/FX/Niagara/Source/Niagara/Public/NiagaraRenderViewDataManager.h
+++ b/Engine/Plugins/FX/Niagara/Source/Niagara/Public/NiagaraRenderViewDataManager.h
@@ -7,6 +7,7 @@
 #include "RenderResource.h"
 #include "ShaderParameterStruct.h"
 #include "ShaderParameterMacros.h"
+#include "SceneRenderTargetParameters.h"
 
 BEGIN_SHADER_PARAMETER_STRUCT(FNiagaraSceneTextureParameters, )
 	SHADER_PARAMETER_RDG_UNIFORM_BUFFER(FSceneTextureUniformParameters, SceneTextures)
diff --git a/Engine/Plugins/FX/Niagara/Source/NiagaraShader/Public/NiagaraShader.h b/Engine/Plugins/FX/Niagara/Source/NiagaraShader/Public/NiagaraShader.h
--- a/Engine/Plugins/FX/Niagara/Source/NiagaraShader/Public/NiagaraShader.h
+++ b/Engine/Plugins/FX/Niagara/Source/NiagaraShader/Public/NiagaraShader.h
@@ -16,7 +16,9 @@
 #include "NiagaraScriptBase.h"
 #include "NiagaraShared.h"
 #include "NiagaraShaderType.h"
+#if UE_ENABLE_INCLUDE_ORDER_DEPRECATED_IN_5_1
 #include "SceneRenderTargetParameters.h"
+#endif
 
 class UClass;
 
diff --git a/Engine/Source/Runtime/Engine/Private/Particles/ParticleGpuSimulation.cpp b/Engine/Source/Runtime/Engine/Private/Particles/ParticleGpuSimulation.cpp
--- a/Engine/Source/Runtime/Engine/Private/Particles/ParticleGpuSimulation.cpp
+++ b/Engine/Source/Runtime/Engine/Private/Particles/ParticleGpuSimulation.cpp
@@ -47,7 +47,6 @@
 #include "VectorField/VectorField.h"
 #include "Misc/CoreDelegates.h"
 #include "PipelineStateCache.h"
-#include "SceneRenderTargetParameters.h"
 #include "MeshMaterialShader.h"
 
 DECLARE_CYCLE_STAT(TEXT("GPUSpriteEmitterInstance Init GT"), STAT_GPUSpriteEmitterInstance_Init, STATGROUP_Particles);
diff --git a/Engine/Source/Runtime/Engine/Private/TileRendering.cpp b/Engine/Source/Runtime/Engine/Private/TileRendering.cpp
--- a/Engine/Source/Runtime/Engine/Private/TileRendering.cpp
+++ b/Engine/Source/Runtime/Engine/Private/TileRendering.cpp
@@ -18,6 +18,7 @@
 #include "SceneUtils.h"
 #include "EngineModule.h"
 #include "MeshPassProcessor.h"
+#include "RenderGraphUtils.h"
 
 DECLARE_GPU_STAT_NAMED(CanvasDrawTiles, TEXT("CanvasDrawTiles"));
 
diff --git a/Engine/Source/Runtime/Engine/Private/TriangleRendering.cpp b/Engine/Source/Runtime/Engine/Private/TriangleRendering.cpp
--- a/Engine/Source/Runtime/Engine/Private/TriangleRendering.cpp
+++ b/Engine/Source/Runtime/Engine/Private/TriangleRendering.cpp
@@ -18,6 +18,7 @@
 #include "SceneUtils.h"
 #include "EngineModule.h"
 #include "MeshPassProcessor.h"
+#include "RenderGraphUtils.h"
 
 DECLARE_GPU_STAT_NAMED(CanvasDrawTriangles, TEXT("CanvasDrawTriangles"));
 
diff --git a/Engine/Source/Runtime/Engine/Private/UserInterface/Canvas.cpp b/Engine/Source/Runtime/Engine/Private/UserInterface/Canvas.cpp
--- a/Engine/Source/Runtime/Engine/Private/UserInterface/Canvas.cpp
+++ b/Engine/Source/Runtime/Engine/Private/UserInterface/Canvas.cpp
@@ -28,6 +28,7 @@
 #include "Fonts/FontMeasure.h"
 #include "EngineModule.h"
 #include "CanvasRender.h"
+#include "RenderGraphUtils.h"
 
 DEFINE_LOG_CATEGORY_STATIC(LogCanvas, Log, All);
 
diff --git a/Engine/Source/Runtime/Landscape/Private/LandscapeGrass.cpp b/Engine/Source/Runtime/Landscape/Private/LandscapeGrass.cpp
--- a/Engine/Source/Runtime/Landscape/Private/LandscapeGrass.cpp
+++ b/Engine/Source/Runtime/Landscape/Private/LandscapeGrass.cpp
@@ -62,6 +62,7 @@
 #include "SimpleMeshDrawCommandPass.h"
 #include "UObject/FortniteMainBranchObjectVersion.h"
 #include "MaterialCachedData.h"
+#include "SceneRenderTargetParameters.h"
 
 #define LOCTEXT_NAMESPACE "Landscape"
 
diff --git a/Engine/Source/Runtime/Landscape/Private/LandscapePhysicalMaterial.cpp b/Engine/Source/Runtime/Landscape/Private/LandscapePhysicalMaterial.cpp
--- a/Engine/Source/Runtime/Landscape/Private/LandscapePhysicalMaterial.cpp
+++ b/Engine/Source/Runtime/Landscape/Private/LandscapePhysicalMaterial.cpp
@@ -18,6 +18,8 @@
 #include "RHIResources.h"
 #include "SceneRenderTargetParameters.h"
 #include "SimpleMeshDrawCommandPass.h"
+#include "RHIGPUReadback.h"
+#include "RenderGraphUtils.h"
 
 DECLARE_GPU_STAT_NAMED(LandscapePhysicalMaterial_Draw, TEXT("LandscapePhysicalMaterial"));
 
diff --git a/Engine/Source/Runtime/RenderCore/Private/RenderGraphUtils.cpp b/Engine/Source/Runtime/RenderCore/Private/RenderGraphUtils.cpp
--- a/Engine/Source/Runtime/RenderCore/Private/RenderGraphUtils.cpp
+++ b/Engine/Source/Runtime/RenderCore/Private/RenderGraphUtils.cpp
@@ -9,6 +9,7 @@
 #include "GlobalShader.h"
 #include "PixelShaderUtils.h"
 #include "RenderGraphResourcePool.h"
+#include "RHIGPUReadback.h"
 
 void ClearUnusedGraphResourcesImpl(
 	const FShaderParameterBindings& ShaderBindings,
diff --git a/Engine/Source/Runtime/RenderCore/Public/RenderGraphUtils.h b/Engine/Source/Runtime/RenderCore/Public/RenderGraphUtils.h
--- a/Engine/Source/Runtime/RenderCore/Public/RenderGraphUtils.h
+++ b/Engine/Source/Runtime/RenderCore/Public/RenderGraphUtils.h
@@ -7,10 +7,14 @@
 #include "Shader.h"
 #include "ShaderParameterStruct.h"
 #include "ShaderParameterMacros.h"
+#if UE_ENABLE_INCLUDE_ORDER_DEPRECATED_IN_5_1
 #include "RHIGPUReadback.h"
+#endif
 #include "PipelineStateCache.h"
 
 class FGlobalShaderMap;
+class FRHIGPUTextureReadback;
+class FRHIGPUBufferReadback;
 
 /** Returns whether the resource was produced by a prior pass. */
 inline bool HasBeenProduced(FRDGViewableResource* Resource)
diff --git a/Engine/Source/Runtime/Renderer/Private/DebugViewModeRendering.h b/Engine/Source/Runtime/Renderer/Private/DebugViewModeRendering.h
--- a/Engine/Source/Runtime/Renderer/Private/DebugViewModeRendering.h
+++ b/Engine/Source/Runtime/Renderer/Private/DebugViewModeRendering.h
@@ -17,6 +17,7 @@ DebugViewModeRendering.h: Contains definitions for rendering debug viewmodes.
 #include "MeshPassProcessor.h"
 #include "DebugViewModeInterface.h"
 #include "GPUSkinCache.h"
+#include "SceneRenderTargetParameters.h"
 
 class FPrimitiveSceneProxy;
 struct FMeshBatchElement;
diff --git a/Engine/Source/Runtime/Renderer/Private/DistortionRendering.cpp b/Engine/Source/Runtime/Renderer/Private/DistortionRendering.cpp
--- a/Engine/Source/Runtime/Renderer/Private/DistortionRendering.cpp
+++ b/Engine/Source/Runtime/Renderer/Private/DistortionRendering.cpp
@@ -58,6 +58,10 @@ static TAutoConsoleVariable<float> CVarRefractionRoughnessToMipLevelFactor(
 	TEXT("Factor to translate the roughness factor into a mip level."),
 	ECVF_Scalability | ECVF_RenderThreadSafe);
 
+BEGIN_GLOBAL_SHADER_PARAMETER_STRUCT(FDistortionPassUniformParameters, RENDERER_API)
+	SHADER_PARAMETER_STRUCT(FSceneTextureUniformParameters, SceneTextures)
+	SHADER_PARAMETER(FVector4f, DistortionParams)
+END_GLOBAL_SHADER_PARAMETER_STRUCT()
 
 IMPLEMENT_STATIC_UNIFORM_BUFFER_STRUCT(FDistortionPassUniformParameters, "DistortionPass", SceneTextures);
 
diff --git a/Engine/Source/Runtime/Renderer/Private/DistortionRendering.h b/Engine/Source/Runtime/Renderer/Private/DistortionRendering.h
--- a/Engine/Source/Runtime/Renderer/Private/DistortionRendering.h
+++ b/Engine/Source/Runtime/Renderer/Private/DistortionRendering.h
@@ -59,18 +59,6 @@ private:
 	TArray<FPrimitiveSceneProxy*, SceneRenderingAllocator> Prims;
 };
 
-
-BEGIN_GLOBAL_SHADER_PARAMETER_STRUCT(FDistortionPassUniformParameters, RENDERER_API)
-	SHADER_PARAMETER_STRUCT(FSceneTextureUniformParameters, SceneTextures)
-	SHADER_PARAMETER(FVector4f, DistortionParams)
-END_GLOBAL_SHADER_PARAMETER_STRUCT()
-
-BEGIN_GLOBAL_SHADER_PARAMETER_STRUCT(FMobileDistortionPassUniformParameters, RENDERER_API)
-	SHADER_PARAMETER_STRUCT(FMobileSceneTextureUniformParameters, SceneTextures)
-	SHADER_PARAMETER(FVector4f, DistortionParams)
-END_GLOBAL_SHADER_PARAMETER_STRUCT()
-
-
 extern void SetupDistortionParams(FVector4f& DistortionParams, const FViewInfo& View);
 
 class FDistortionMeshProcessor : public FMeshPassProcessor
diff --git a/Engine/Source/Runtime/Renderer/Private/LightMapDensityRendering.cpp b/Engine/Source/Runtime/Renderer/Private/LightMapDensityRendering.cpp
--- a/Engine/Source/Runtime/Renderer/Private/LightMapDensityRendering.cpp
+++ b/Engine/Source/Runtime/Renderer/Private/LightMapDensityRendering.cpp
@@ -32,6 +32,15 @@ IMPLEMENT_DENSITY_LIGHTMAPPED_SHADER_TYPE(TUniformLightMapPolicy<LMP_HQ_LIGHTMAP
 
 #endif
 
+BEGIN_GLOBAL_SHADER_PARAMETER_STRUCT(FLightmapDensityPassUniformParameters, )
+	SHADER_PARAMETER_STRUCT(FSceneTextureUniformParameters, SceneTextures)
+	SHADER_PARAMETER(FVector4f, LightMapDensity)
+	SHADER_PARAMETER(FVector4f, DensitySelectedColor) // The color to apply to selected objects.
+	SHADER_PARAMETER(FVector4f, VertexMappedColor) // The color to apply to vertex mapped objects.
+	SHADER_PARAMETER_TEXTURE(Texture2D, GridTexture) // The "Grid" texture to visualize resolution.
+	SHADER_PARAMETER_SAMPLER(SamplerState, GridTextureSampler)
+END_GLOBAL_SHADER_PARAMETER_STRUCT()
+
 IMPLEMENT_STATIC_UNIFORM_BUFFER_STRUCT(FLightmapDensityPassUniformParameters, "LightmapDensityPass", SceneTextures);
 
 void SetupLightmapDensityPassUniformBuffer(FRDGBuilder& GraphBuilder, ERHIFeatureLevel::Type FeatureLevel, FLightmapDensityPassUniformParameters& LightmapDensityPassParameters)
diff --git a/Engine/Source/Runtime/Renderer/Private/LightMapDensityRendering.h b/Engine/Source/Runtime/Renderer/Private/LightMapDensityRendering.h
--- a/Engine/Source/Runtime/Renderer/Private/LightMapDensityRendering.h
+++ b/Engine/Source/Runtime/Renderer/Private/LightMapDensityRendering.h
@@ -20,15 +20,6 @@
 #include "SceneRendering.h"
 #include "Engine/LightMapTexture2D.h"
 
-BEGIN_GLOBAL_SHADER_PARAMETER_STRUCT(FLightmapDensityPassUniformParameters, )
-	SHADER_PARAMETER_STRUCT(FSceneTextureUniformParameters, SceneTextures)
-	SHADER_PARAMETER(FVector4f, LightMapDensity)
-	SHADER_PARAMETER(FVector4f, DensitySelectedColor) // The color to apply to selected objects.
-	SHADER_PARAMETER(FVector4f, VertexMappedColor) // The color to apply to vertex mapped objects.
-	SHADER_PARAMETER_TEXTURE(Texture2D, GridTexture) // The "Grid" texture to visualize resolution.
-	SHADER_PARAMETER_SAMPLER(SamplerState, GridTextureSampler)
-END_GLOBAL_SHADER_PARAMETER_STRUCT()
-
 template<typename LightMapPolicyType>
 class TLightMapDensityElementData : public FMeshMaterialShaderElementData
 {
diff --git a/Engine/Source/Runtime/Renderer/Private/LightRendering.cpp b/Engine/Source/Runtime/Renderer/Private/LightRendering.cpp
--- a/Engine/Source/Runtime/Renderer/Private/LightRendering.cpp
+++ b/Engine/Source/Runtime/Renderer/Private/LightRendering.cpp
@@ -180,6 +180,88 @@ bool ShouldRenderRayTracingShadowsForLight(const FLightSceneInfoCompact& LightIn
 }
 #endif // RHI_RAYTRACING
 
+FDrawFullScreenRectangleParameters FDeferredLightVS::GetFullScreenRectParameters(
+	float X, float Y,
+	float SizeX, float SizeY,
+	float U, float V,
+	float SizeU, float SizeV,
+	FIntPoint InTargetSize,
+	FIntPoint InTextureSize)
+{
+	FDrawFullScreenRectangleParameters Out;
+	Out.PosScaleBias = FVector4f(SizeX, SizeY, X, Y);
+	Out.UVScaleBias = FVector4f(SizeU, SizeV, U, V);
+	Out.InvTargetSizeAndTextureSize = FVector4f(
+		1.0f / InTargetSize.X, 1.0f / InTargetSize.Y,
+		1.0f / InTextureSize.X, 1.0f / InTextureSize.Y);
+	return Out;
+}
+
+FDeferredLightVS::FParameters FDeferredLightVS::GetParameters(const FViewInfo& View,
+	float X, float Y,
+	float SizeX, float SizeY,
+	float U, float V,
+	float SizeU, float SizeV,
+	FIntPoint TargetSize,
+	FIntPoint TextureSize,
+	bool bBindViewUniform)
+{
+	FParameters Out;
+	if (bBindViewUniform)
+	{
+		Out.View = View.ViewUniformBuffer;
+	}
+	Out.Geometry = FStencilingGeometryShaderParameters::GetParameters(FVector4f(0, 0, 0, 0));
+	Out.FullScreenRect = GetFullScreenRectParameters(X, Y, SizeX, SizeY, U, V, SizeU, SizeV, TargetSize, TextureSize);
+	return Out;
+}
+
+FDeferredLightVS::FParameters FDeferredLightVS::GetParameters(const FViewInfo& View, bool bBindViewUniform)
+{
+	FParameters Out;
+	if (bBindViewUniform)
+	{
+		Out.View = View.ViewUniformBuffer;
+	}
+	Out.Geometry = FStencilingGeometryShaderParameters::GetParameters(FVector4f(0, 0, 0, 0));
+	Out.FullScreenRect = GetFullScreenRectParameters(
+		0, 0,
+		View.ViewRect.Width(), View.ViewRect.Height(),
+		View.ViewRect.Min.X, View.ViewRect.Min.Y,
+		View.ViewRect.Width(), View.ViewRect.Height(),
+		View.ViewRect.Size(),
+		GetSceneTextureExtent());
+
+	return Out;
+}
+
+FDeferredLightVS::FParameters FDeferredLightVS::GetParameters(const FViewInfo& View, const FSphere& LightBounds, bool bBindViewUniform)
+{
+	FVector4f StencilingSpherePosAndScale;
+	StencilingGeometry::GStencilSphereVertexBuffer.CalcTransform(StencilingSpherePosAndScale, LightBounds, View.ViewMatrices.GetPreViewTranslation());
+
+	FParameters Out;
+	if (bBindViewUniform)
+	{
+		Out.View = View.ViewUniformBuffer;
+	}
+	Out.Geometry = FStencilingGeometryShaderParameters::GetParameters((FVector4f)StencilingSpherePosAndScale); // LWC_TODO: Precision loss
+	Out.FullScreenRect = GetFullScreenRectParameters(0, 0, 0, 0, 0, 0, 0, 0, FIntPoint(1, 1), FIntPoint(1, 1));
+	return Out;
+}
+
+FDeferredLightVS::FParameters FDeferredLightVS::GetParameters(const FViewInfo& View, const FLightSceneInfo* LightSceneInfo, bool bBindViewUniform)
+{
+	FParameters Out;
+	if (bBindViewUniform)
+	{
+		Out.View = View.ViewUniformBuffer;
+	}
+	Out.Geometry = FStencilingGeometryShaderParameters::GetParameters(View, LightSceneInfo);
+	Out.FullScreenRect = GetFullScreenRectParameters(0, 0, 0, 0, 0, 0, 0, 0, FIntPoint(1, 1), FIntPoint(1, 1));
+	return Out;
+}
+
 static void RenderLight(
 	FRDGBuilder& GraphBuilder,
 	const FScene* Scene,
diff --git a/Engine/Source/Runtime/Renderer/Private/LightRendering.h b/Engine/Source/Runtime/Renderer/Private/LightRendering.h
--- a/Engine/Source/Runtime/Renderer/Private/LightRendering.h
+++ b/Engine/Source/Runtime/Renderer/Private/LightRendering.h
@@ -514,81 +514,22 @@ public:
 		float U, float V,
 		float SizeU, float SizeV,
 		FIntPoint InTargetSize,
-		FIntPoint InTextureSize)
-	{
-		FDrawFullScreenRectangleParameters Out;
-		Out.PosScaleBias = FVector4f(SizeX, SizeY, X, Y);
-		Out.UVScaleBias  = FVector4f(SizeU, SizeV, U, V);
-		Out.InvTargetSizeAndTextureSize = FVector4f(
-			1.0f / InTargetSize.X, 1.0f / InTargetSize.Y,
-			1.0f / InTextureSize.X, 1.0f / InTextureSize.Y);
-		return Out;
-	}
+		FIntPoint InTextureSize);
 
-	static FParameters GetParameters(const FViewInfo& View, 
+	static FParameters GetParameters(const FViewInfo& View,
 		float X, float Y,
 		float SizeX, float SizeY,
 		float U, float V,
 		float SizeU, float SizeV,
 		FIntPoint TargetSize,
 		FIntPoint TextureSize,
-		bool bBindViewUniform = true)
-	{
-		FParameters Out;
-		if (bBindViewUniform)
-		{
-			Out.View = View.ViewUniformBuffer;
-		}
-		Out.Geometry = FStencilingGeometryShaderParameters::GetParameters(FVector4f(0,0,0,0));
-		Out.FullScreenRect = GetFullScreenRectParameters(X, Y, SizeX, SizeY, U, V, SizeU, SizeV, TargetSize, TextureSize);
-		return Out;
-	}
+		bool bBindViewUniform = true);
 
-	static FParameters GetParameters(const FViewInfo& View, bool bBindViewUniform = true)
-	{
-		FParameters Out;
-		if (bBindViewUniform)
-		{
-			Out.View = View.ViewUniformBuffer;
-		}
-		Out.Geometry = FStencilingGeometryShaderParameters::GetParameters(FVector4f(0, 0, 0, 0));
-		Out.FullScreenRect = GetFullScreenRectParameters(
-			0, 0,
-			View.ViewRect.Width(), View.ViewRect.Height(),
-			View.ViewRect.Min.X, View.ViewRect.Min.Y,
-			View.ViewRect.Width(), View.ViewRect.Height(),
-			View.ViewRect.Size(),
-			GetSceneTextureExtent());
+	static FParameters GetParameters(const FViewInfo& View, bool bBindViewUniform = true);
 
-		return Out;
-	}
+	static FParameters GetParameters(const FViewInfo& View, const FSphere& LightBounds, bool bBindViewUniform = true);
 
-	static FParameters GetParameters(const FViewInfo& View, const FSphere& LightBounds, bool bBindViewUniform = true)
-	{
-		FVector4f StencilingSpherePosAndScale;
-		StencilingGeometry::GStencilSphereVertexBuffer.CalcTransform(StencilingSpherePosAndScale, LightBounds, View.ViewMatrices.GetPreViewTranslation());
-
-		FParameters Out;
-		if (bBindViewUniform)
-		{
-			Out.View = View.ViewUniformBuffer;
-		}
-		Out.Geometry = FStencilingGeometryShaderParameters::GetParameters((FVector4f)StencilingSpherePosAndScale); // LWC_TODO: Precision loss
-		Out.FullScreenRect = GetFullScreenRectParameters(0, 0, 0, 0, 0, 0, 0, 0, FIntPoint(1, 1), FIntPoint(1, 1));
-		return Out;
-	}
-
-	static FParameters GetParameters(const FViewInfo& View, const FLightSceneInfo* LightSceneInfo, bool bBindViewUniform = true)
-	{
-		FParameters Out;
-		if (bBindViewUniform)
-		{
-			Out.View = View.ViewUniformBuffer;
-		}
-		Out.Geometry = FStencilingGeometryShaderParameters::GetParameters(View, LightSceneInfo);
-		Out.FullScreenRect = GetFullScreenRectParameters(0, 0, 0, 0, 0, 0, 0, 0, FIntPoint(1, 1), FIntPoint(1, 1));
-		return Out;
-	}
+	static FParameters GetParameters(const FViewInfo& View, const FLightSceneInfo* LightSceneInfo, bool bBindViewUniform = true);
 };
 
 enum class FLightOcclusionType : uint8
diff --git a/Engine/Source/Runtime/Renderer/Private/MobileDistortionPass.cpp b/Engine/Source/Runtime/Renderer/Private/MobileDistortionPass.cpp
--- a/Engine/Source/Runtime/Renderer/Private/MobileDistortionPass.cpp
+++ b/Engine/Source/Runtime/Renderer/Private/MobileDistortionPass.cpp
@@ -12,6 +12,11 @@ MobileDistortionPass.cpp - Mobile specific rendering of primtives with refractio
 #include "PipelineStateCache.h"
 #include "DistortionRendering.h"
 
+BEGIN_GLOBAL_SHADER_PARAMETER_STRUCT(FMobileDistortionPassUniformParameters, RENDERER_API)
+	SHADER_PARAMETER_STRUCT(FMobileSceneTextureUniformParameters, SceneTextures)
+	SHADER_PARAMETER(FVector4f, DistortionParams)
+END_GLOBAL_SHADER_PARAMETER_STRUCT()
+
 IMPLEMENT_STATIC_UNIFORM_BUFFER_STRUCT(FMobileDistortionPassUniformParameters, "MobileDistortionPass", SceneTextures);
 
 bool IsMobileDistortionActive(const FViewInfo& View)
diff --git a/Engine/Source/Runtime/Renderer/Private/SceneRendering.h b/Engine/Source/Runtime/Renderer/Private/SceneRendering.h
--- a/Engine/Source/Runtime/Renderer/Private/SceneRendering.h
+++ b/Engine/Source/Runtime/Renderer/Private/SceneRendering.h
@@ -51,6 +51,7 @@ class FRaytracingLightDataPacked;
 class FRayTracingLocalShaderBindingWriter;
 class FVirtualShadowMapClipmap;
 class FShadowProjectionPassParameters;
+class FSceneTextureShaderParameters;
 
 struct FCloudRenderContext;
 struct FSceneWithoutWaterTextures;
diff --git a/Engine/Source/Runtime/Renderer/Private/Strata/Strata.h b/Engine/Source/Runtime/Renderer/Private/Strata/Strata.h
--- a/Engine/Source/Runtime/Renderer/Private/Strata/Strata.h
+++ b/Engine/Source/Runtime/Renderer/Private/Strata/Strata.h
@@ -9,6 +9,7 @@
 #include "UnifiedBuffer.h"
 #include "RHIUtilities.h"
 #include "StrataDefinitions.h"
+#include "RenderGraphUtils.h"
 
 // Forward declarations.
 class FScene;
diff --git a/Engine/Source/Runtime/Renderer/Public/MaterialShader.h b/Engine/Source/Runtime/Renderer/Public/MaterialShader.h
--- a/Engine/Source/Runtime/Renderer/Public/MaterialShader.h
+++ b/Engine/Source/Runtime/Renderer/Public/MaterialShader.h
@@ -15,7 +15,9 @@
 #include "MaterialShared.h"
 #include "GlobalShader.h"
 #include "MaterialShaderType.h"
+#if UE_ENABLE_INCLUDE_ORDER_DEPRECATED_IN_5_1
 #include "SceneRenderTargetParameters.h"
+#endif
 #include "ShaderParameterUtils.h"
 
 template<typename TBufferStruct> class TUniformBufferRef;
diff --git a/Engine/Source/Runtime/Renderer/Public/SceneRenderTargetParameters.h b/Engine/Source/Runtime/Renderer/Public/SceneRenderTargetParameters.h
--- a/Engine/Source/Runtime/Renderer/Public/SceneRenderTargetParameters.h
+++ b/Engine/Source/Runtime/Renderer/Public/SceneRenderTargetParameters.h
@@ -2,7 +2,16 @@
 
 #pragma once
 
+#if UE_ENABLE_INCLUDE_ORDER_DEPRECATED_IN_5_1
 #include "RenderGraph.h"
+#endif
+#include "RenderGraphDefinitions.h"
+#include "ShaderParameterMacros.h"
+#include "RendererInterface.h"
+
+class FRDGBuilder;
+struct IPooledRenderTarget;
+struct FSceneTextures;
 
 /** A uniform buffer containing common scene textures used by materials or global shaders. */
 BEGIN_GLOBAL_SHADER_PARAMETER_STRUCT(FSceneTextureUniformParameters, RENDERER_API)
diff --git a/Engine/Source/Runtime/Renderer/Public/SimpleMeshDrawCommandPass.h b/Engine/Source/Runtime/Renderer/Public/SimpleMeshDrawCommandPass.h
--- a/Engine/Source/Runtime/Renderer/Public/SimpleMeshDrawCommandPass.h
+++ b/Engine/Source/Runtime/Renderer/Public/SimpleMeshDrawCommandPass.h
@@ -10,6 +10,7 @@
 #include "SceneUtils.h"
 #include "MeshPassProcessor.h"
 #include "InstanceCulling/InstanceCullingContext.h"
+#include "RenderGraphBuilder.h"
 
 class FGPUScene;
 class FInstanceCullingDrawParams;
diff --git a/Engine/Source/Runtime/SlateRHIRenderer/Private/SlateRHIRenderingPolicy.cpp b/Engine/Source/Runtime/SlateRHIRenderer/Private/SlateRHIRenderingPolicy.cpp
--- a/Engine/Source/Runtime/SlateRHIRenderer/Private/SlateRHIRenderingPolicy.cpp
+++ b/Engine/Source/Runtime/SlateRHIRenderer/Private/SlateRHIRenderingPolicy.cpp
@@ -26,6 +26,8 @@
 #include "DeviceProfiles/DeviceProfile.h"
 #include "DeviceProfiles/DeviceProfileManager.h"
 #include "Types/SlateConstants.h"
+#include "RenderGraphResources.h"
+#include "SceneRenderTargetParameters.h"
 
 extern void UpdateNoiseTextureParameters(FViewUniformShaderParameters& ViewUniformShaderParameters);
 
