Use non-immediate command list for distance field pass to let it run in parallel.

diff --git a/Engine/Source/Runtime/Renderer/Private/DistanceFieldObjectCulling.cpp b/Engine/Source/Runtime/Renderer/Private/DistanceFieldObjectCulling.cpp
--- a/Engine/Source/Runtime/Renderer/Private/DistanceFieldObjectCulling.cpp
+++ b/Engine/Source/Runtime/Renderer/Private/DistanceFieldObjectCulling.cpp
@@ -302,7 +302,7 @@ void ScatterTilesToObjects(
 		bCountingPass ? RDG_EVENT_NAME("CountTileObjectIntersections") : RDG_EVENT_NAME("CullTilesToObjects"),
 		PassParameters,
 		ERDGPassFlags::Raster,
-		[PassParameters, VertexShader, PixelShader, &View, TileListGroupSize, ObjectIndirectArguments](FRHICommandListImmediate& RHICmdList)
+		[PassParameters, VertexShader, PixelShader, &View, TileListGroupSize, ObjectIndirectArguments](FRHICommandList& RHICmdList)
 		{
 			FGraphicsPipelineStateInitializer GraphicsPSOInit;
 			RHICmdList.ApplyCachedRenderTargets(GraphicsPSOInit);
