Fix Strata visualization multi-pass.

diff --git a/Engine/Shaders/Private/Strata/StrataVisualize.usf b/Engine/Shaders/Private/Strata/StrataVisualize.usf
--- a/Engine/Shaders/Private/Strata/StrataVisualize.usf
+++ b/Engine/Shaders/Private/Strata/StrataVisualize.usf
@@ -82,10 +82,6 @@ void DrawPixelFootprint(float3 P, float3 dPdx, float3 dPdy, uint2 PixelCoord)
 
 #if SHADER_MATERIALPRINT
 
-#ifndef MULTIPASS_ENABLE 
-#define MULTIPASS_ENABLE 0
-#endif
-
 uint BSDFIndex;
 RWBuffer<uint> RWPositionOffsetBuffer;
 
@@ -301,7 +297,6 @@ void PrintBSDF(inout FShaderPrintContext Context, inout FStrataAddressing Strata
 void PrintPixel(uint2 InCoord, float3 InWorldPosition, float3 V)
 {
 	FShaderPrintContext Context;
-	#if MULTIPASS_ENABLE
 	if (BSDFIndex > 0)
 	{
 		float2 PrintOffset;
@@ -310,7 +305,6 @@ void PrintPixel(uint2 InCoord, float3 InWorldPosition, float3 V)
 		Context = InitShaderPrintContext(true, PrintOffset);
 	}
 	else
-	#endif
 	{
 		Context = InitShaderPrintContext(true, uint2(50, 50));
 	}
@@ -359,41 +353,21 @@ void PrintPixel(uint2 InCoord, float3 InWorldPosition, float3 V)
 	}
 
 	// BSDFs
-#if MULTIPASS_ENABLE
 	if (BSDFIndex < Header.BSDFCount)
 	{
 		if (BSDFIndex > 0)
 		{
-			StrataSeekBSDF(StrataAddressing, UnpackBSDFOffset(Strata.BSDFOffsetTexture[InCoord], Header.BSDFCount), BSDFIndex);
+			const uint AddressOffset = UnpackBSDFOffsetAtIndex(Strata.BSDFOffsetTexture[InCoord], BSDFIndex, Header.BSDFCount);
+			StrataSeekBSDF(StrataAddressing, AddressOffset);
 		}
 		PrintBSDF(Context, StrataAddressing, Header, InWorldPosition, V);
 
 		RWPositionOffsetBuffer[0] = asuint(Context.Pos.x);
 		RWPositionOffsetBuffer[1] = asuint(Context.Pos.y);
 	}
-#else // MULTIPASS_ENABLE
-	if (BSDFIndex == 0 && Header.BSDFCount > 0)
-	{
-	#if 1
-		#if COMPILER_DXC == 1
-		LOOP
-		for (uint BSDFIndex = 0; BSDFIndex < Header.BSDFCount; ++BSDFIndex)
-		#else
-		Header.BSDFCount = min(Header.BSDFCount, STRATA_VIZ_MAX_BSDF);
-		UNROLL_N(STRATA_VIZ_MAX_BSDF) // FXC requires that we specify an upper bound here. Not DXC.
-		for (uint LocalBSDFIndex = 0; LocalBSDFIndex < Header.BSDFCount; ++LocalBSDFIndex)
-		#endif
-	#endif
-		{
-			PrintBSDF(Context, StrataAddressing, Header, InWorldPosition, V);
-		}
-	}
-#endif // MULTIPASS_ENABLE
 
 	// Memory footprint
-#if MULTIPASS_ENABLE
 	if (BSDFIndex+1 == Header.BSDFCount)
-#endif
 	{
 		const uint FootPrint_PostBSDFs = StrataAddressing.ReadBytes;
 
diff --git a/Engine/Source/Runtime/Renderer/Private/Strata/StrataVisualize.cpp b/Engine/Source/Runtime/Renderer/Private/Strata/StrataVisualize.cpp
--- a/Engine/Source/Runtime/Renderer/Private/Strata/StrataVisualize.cpp
+++ b/Engine/Source/Runtime/Renderer/Private/Strata/StrataVisualize.cpp
@@ -31,8 +31,6 @@ void AddStrataInternalClassificationTilePass(
 	EStrataTileType TileMaterialType,
 	const bool bDebug);
 
-#define MULTIPASS_ENABLE 0
-
 class FMaterialPrintInfoCS : public FGlobalShader
 {
 	DECLARE_GLOBAL_SHADER(FMaterialPrintInfoCS);
@@ -72,7 +70,6 @@ class FMaterialPrintInfoCS : public FGlobalShader
 		// Stay debug and skip optimizations to reduce compilation time on this long shader.
 		OutEnvironment.CompilerFlags.Add(CFLAG_Debug);
 		OutEnvironment.SetDefine(TEXT("SHADER_MATERIALPRINT"), 1);
-		OutEnvironment.SetDefine(TEXT("MULTIPASS_ENABLE"), MULTIPASS_ENABLE);
 	}
 };
 IMPLEMENT_GLOBAL_SHADER(FMaterialPrintInfoCS, "/Engine/Private/Strata/StrataVisualize.usf", "MaterialPrintInfoCS", SF_Compute);
@@ -133,17 +130,14 @@ static void AddVisualizeMaterialPasses(FRDGBuilder& GraphBuilder, const FViewInf
 
 		// Print Material info
 		{
-			uint32 BSDFIndex = 0;
 			FRDGBufferUAVRef PrintOffsetBufferUAV = nullptr;
 
-#if MULTIPASS_ENABLE
 			FRDGBufferRef PrintOffsetBuffer = GraphBuilder.CreateBuffer(FRDGBufferDesc::CreateBufferDesc(4, 2), TEXT("Strata.DebugPrintPositionOffset"));
 			PrintOffsetBufferUAV = GraphBuilder.CreateUAV(PrintOffsetBuffer, PF_R32_UINT);
 			AddClearUAVPass(GraphBuilder, PrintOffsetBufferUAV, 50u);
 			const uint32 MaxBSDFCount = 8;
 
-			for (BSDFIndex; BSDFIndex < MaxBSDFCount; ++BSDFIndex)
-#endif
+			for (uint32 BSDFIndex=0; BSDFIndex < MaxBSDFCount; ++BSDFIndex)
 			{
 				FMaterialPrintInfoCS::FParameters* PassParameters = GraphBuilder.AllocParameters<FMaterialPrintInfoCS::FParameters>();
 				PassParameters->BSDFIndex = BSDFIndex;
@@ -154,7 +148,7 @@ static void AddVisualizeMaterialPasses(FRDGBuilder& GraphBuilder, const FViewInf
 				ShaderPrint::SetParameters(GraphBuilder, View.ShaderPrintData, PassParameters->ShaderPrintParameters);
 
 				TShaderMapRef<FMaterialPrintInfoCS> ComputeShader(View.ShaderMap);
-				FComputeShaderUtils::AddPass(GraphBuilder, RDG_EVENT_NAME("Strata::VisualizeMaterial(Print)"), ComputeShader, PassParameters, FIntVector(1, 1, 1));
+				FComputeShaderUtils::AddPass(GraphBuilder, RDG_EVENT_NAME("Strata::VisualizeMaterial(Print, BSDF=%d)", BSDFIndex), ComputeShader, PassParameters, FIntVector(1, 1, 1));
 			}
 		}
 
