Strata - advanced vis debug
- partial checkin to continue on another stream.
- Write to the end of the strata texture array is working.

diff --git a/Engine/Shaders/Private/BasePassPixelShader.usf b/Engine/Shaders/Private/BasePassPixelShader.usf
--- a/Engine/Shaders/Private/BasePassPixelShader.usf
+++ b/Engine/Shaders/Private/BasePassPixelShader.usf
@@ -1732,6 +1732,10 @@ void FPixelShaderInOut_MainPS(
 		FStrataOpaqueRoughRefractionData OpaqueRoughRefractionData = (FStrataOpaqueRoughRefractionData)0;
 		FStrataIntegrationSettings Settings = InitStrataIntegrationSettings(false /*bForceFullyRough*/, StrataStruct.bRoughDiffuse, StrataStruct.PeelLayersAboveDepth);
 
+#if STRATA_ADVANCED_DEBUG_ENABLED
+		Settings.SliceStoringDebugStrataTree = all(View.CursorPosition == PixelPos) ? Strata.SliceStoringDebugStrataTree : Settings.SliceStoringDebugStrataTree;
+#endif
+
 		// Generate the strata material data to write out
 		FStrataAddressing StrataAddressing = GetStrataPixelDataByteOffset(PixelPos, uint2(ResolvedView.BufferSizeAndInvSize.xy), StrataStruct.MaxBytesPerPixel);
 		FRWStrataMaterialContainer RWStrataMaterialContainer = InitialiseRWStrataMaterialContainer(StrataStruct.MaterialTextureArrayUAVWithoutRTs);
diff --git a/Engine/Shaders/Private/Strata/Strata.ush b/Engine/Shaders/Private/Strata/Strata.ush
--- a/Engine/Shaders/Private/Strata/Strata.ush
+++ b/Engine/Shaders/Private/Strata/Strata.ush
@@ -813,6 +813,60 @@ FStrataTree GetInitialisedStrataTree()
 	return StrataTree;
 }
 
+void StrataPackOutStrataTree(
+	inout RWTexture2DArray<uint> ExtraMaterialDataUAV,
+	in FStrataTree StrataTree,
+	in int SliceStoringDebugStrataTree)
+{
+	int OutCoord = 0;
+#define WRITE_UINT(Data) ExtraMaterialDataUAV[uint3(OutCoord++, 0, SliceStoringDebugStrataTree)] = Data;
+#define WRITE_FLOAT(Data) ExtraMaterialDataUAV[uint3(OutCoord++, 0, SliceStoringDebugStrataTree)] = asuint(Data);
+#define WRITE_FLOAT3(Data) WRITE_FLOAT(Data.x);WRITE_FLOAT(Data.y);WRITE_FLOAT(Data.z);
+#define WRITE_FLOAT4(Data) WRITE_FLOAT(Data.x);WRITE_FLOAT(Data.y);WRITE_FLOAT(Data.z);;WRITE_FLOAT(Data.w);
+
+	WRITE_UINT(StrataTree.BSDFCount);
+	WRITE_UINT(StrataTree.OperatorCount);
+
+	LOOP
+	for(int i=0; i< StrataTree.BSDFCount; ++i)
+	{
+		WRITE_UINT(StrataTree.BSDFs[i].State);
+		WRITE_UINT(StrataTree.BSDFs[i].OperatorIndex);
+		WRITE_FLOAT3(asuint(StrataTree.BSDFs[i].LuminanceWeightV));
+		WRITE_FLOAT3(asuint(StrataTree.BSDFs[i].LuminanceWeightL));
+		WRITE_UINT(StrataTree.BSDFs[i].bIsBottom);
+		WRITE_UINT(StrataTree.BSDFs[i].bIsTop);
+
+#if STRATA_INLINE_SHADING
+		WRITE_FLOAT(StrataTree.BSDFs[i].Coverage);
+
+		WRITE_FLOAT3(StrataTree.BSDFs[i].Emissive);
+		WRITE_FLOAT(StrataTree.BSDFs[i].ThicknessCm);
+		WRITE_FLOAT(StrataTree.BSDFs[i].TmpMFP);
+
+		WRITE_FLOAT(StrataTree.BSDFs[i].TopLayerDataWeight);
+#endif
+
+		// Skipping specific single layer water data
+
+		WRITE_FLOAT4(StrataTree.BSDFs[i].VGPRs[0]);
+		WRITE_FLOAT4(StrataTree.BSDFs[i].VGPRs[1]);
+		WRITE_FLOAT4(StrataTree.BSDFs[i].VGPRs[2]);
+		WRITE_FLOAT4(StrataTree.BSDFs[i].VGPRs[3]);
+		WRITE_FLOAT4(StrataTree.BSDFs[i].VGPRs[4]);
+	}
+
+#undef WRITE_UINT
+#undef WRITE_FLOAT
+#undef WRITE_FLOAT3
+#undef WRITE_FLOAT4
+}
+
+void StrataUnpackInStrataTree()
+{
+}
+
+
 struct FStrataData
 {
 	FStrataBSDF InlinedBSDF;// Used for parameter blending or special BSDF such as hair, cloud or water.
@@ -2593,6 +2647,7 @@ struct FStrataIntegrationSettings
 {
 	bool bForceFullyRough;
 	bool bRoughDiffuseEnabled;
+	int  SliceStoringDebugStrataTree;
 	int  PeelLayersAboveDepth;
 };
 
@@ -2602,6 +2657,7 @@ FStrataIntegrationSettings InitStrataIntegrationSettings(bool bForceFullyRough,
 	Out.bForceFullyRough = bForceFullyRough;
 	Out.bRoughDiffuseEnabled = bRoughDiffuseEnabled;
 	Out.PeelLayersAboveDepth = PeelLayersAboveDepth;
+	Out.SliceStoringDebugStrataTree = -1;
 	return Out;
 }
 
@@ -2724,6 +2780,13 @@ void PackStrataOut(
 		const float3 NullLightDir = float3(0.0, 0.0, 1.0);
 		StrataUpdateTree(NullStrataAddressing, StrataPixelHeader, StrataTree, V, NullLightDir, Settings);
 
+	#if STRATA_ADVANCED_DEBUG_ENABLED
+		if (Settings.SliceStoringDebugStrataTree > -1)
+		{
+			StrataPackOutStrataTree(ExtraMaterialDataUAV, StrataTree, Settings.SliceStoringDebugStrataTree);
+		}
+	#endif
+
 		// Process the rough refraction data (this code is only used, and not culled out, if effectively used in base pass for instance)
 		{
 			uint RootOperatorIndex = Strata.OperatorIndex;
diff --git a/Engine/Source/Runtime/Engine/Classes/Engine/RendererSettings.h b/Engine/Source/Runtime/Engine/Classes/Engine/RendererSettings.h
--- a/Engine/Source/Runtime/Engine/Classes/Engine/RendererSettings.h
+++ b/Engine/Source/Runtime/Engine/Classes/Engine/RendererSettings.h
@@ -931,6 +931,15 @@ class ENGINE_API URendererSettings : public UDeveloperSettings
 		ConfigRestartRequired = true))
 		uint32 StrataOpaqueMaterialRoughRefraction : 1;
 
+	/**
+	"Enable advanced strata material debug visualization shaders. Base pas shaders can output such advanced data."
+	*/
+	UPROPERTY(config, EditAnywhere, Category = Strata, meta = (
+		ConsoleVariable = "r.Strata.Debug.AdvancedVisualizationShaders", DisplayName = "Strata advanced visualization shaders",
+		ToolTip = "Enable advanced strata material debug visualization shaders. Base pass shaders can output such advanced data.",
+		ConfigRestartRequired = true))
+		uint32 StrataDebugAdvancedVisualizationShaders : 1;
+
 	/**
 	"Enable rough diffuse material."
 	*/
diff --git a/Engine/Source/Runtime/Engine/Private/ShaderCompiler/ShaderCompiler.cpp b/Engine/Source/Runtime/Engine/Private/ShaderCompiler/ShaderCompiler.cpp
--- a/Engine/Source/Runtime/Engine/Private/ShaderCompiler/ShaderCompiler.cpp
+++ b/Engine/Source/Runtime/Engine/Private/ShaderCompiler/ShaderCompiler.cpp
@@ -5896,6 +5896,10 @@ void GlobalBeginCompileShader(
 		static IConsoleVariable* CVarOpaqueRoughRefrac = IConsoleManager::Get().FindConsoleVariable(TEXT("r.Strata.OpaqueMaterialRoughRefraction"));
 		const bool bStrataOpaqueRoughRefrac = bStrata && CVarOpaqueRoughRefrac && CVarOpaqueRoughRefrac->GetInt() != 0;
 		Input.Environment.SetDefine(TEXT("STRATA_OPAQUE_ROUGH_REFRACTION_ENABLED"), bStrataOpaqueRoughRefrac ? 1 : 0);
+
+		static IConsoleVariable* CVarAdvDebug = IConsoleManager::Get().FindConsoleVariable(TEXT("r.Strata.Debug.AdvancedVisualizationShaders"));
+		const bool bStrataAdvDebug = bStrata && CVarAdvDebug && CVarAdvDebug->GetInt() != 0;
+		Input.Environment.SetDefine(TEXT("STRATA_ADVANCED_DEBUG_ENABLED"), bStrataAdvDebug ? 1 : 0);
 	}
 
 	{
diff --git a/Engine/Source/Runtime/RenderCore/Private/Shader.cpp b/Engine/Source/Runtime/RenderCore/Private/Shader.cpp
--- a/Engine/Source/Runtime/RenderCore/Private/Shader.cpp
+++ b/Engine/Source/Runtime/RenderCore/Private/Shader.cpp
@@ -1808,6 +1808,12 @@ void ShaderMapAppendKeyString(EShaderPlatform Platform, FString& KeyString)
 		{
 			KeyString += FString::Printf(TEXT("_ROUGHDIFF"));
 		}
+
+		static const auto CVarAdvancedDebug = IConsoleManager::Get().FindTConsoleVariableDataInt(TEXT("r.Strata.Debug.AdvancedVisualizationShaders"));
+		if (bStrataEnabled && CVarAdvancedDebug && CVarAdvancedDebug->GetValueOnAnyThread() > 0)
+		{
+			KeyString += FString::Printf(TEXT("_ADVDEBUG"));
+		}
 	}
 
 	{
diff --git a/Engine/Source/Runtime/Renderer/Private/Strata/Strata.cpp b/Engine/Source/Runtime/Renderer/Private/Strata/Strata.cpp
--- a/Engine/Source/Runtime/Renderer/Private/Strata/Strata.cpp
+++ b/Engine/Source/Runtime/Renderer/Private/Strata/Strata.cpp
@@ -254,7 +254,8 @@ void InitialiseStrataFrameSceneData(FRDGBuilder& GraphBuilder, FSceneRenderer& S
 
 		// Separated subsurface and rough refraction textures
 		{
-			const bool bIsStrataOpaqueMaterialRoughRefractionEnabled= IsStrataOpaqueMaterialRoughRefractionEnabled();
+			const bool bIsStrataOpaqueMaterialRoughRefractionEnabled = IsStrataOpaqueMaterialRoughRefractionEnabled();
+			const bool bIsStrataAdvancedVisualizationShadersEnabled = IsStrataAdvancedVisualizationShadersEnabled();
 			const FIntPoint OpaqueRoughRefractionSceneExtent		= bIsStrataOpaqueMaterialRoughRefractionEnabled ? SceneTextureExtent : FIntPoint(4, 4);
 			
 			Out.OpaqueRoughRefractionTexture = GraphBuilder.CreateTexture(
@@ -304,6 +305,7 @@ void InitialiseStrataFrameSceneData(FRDGBuilder& GraphBuilder, FSceneRenderer& S
 	Out.bRoughDiffuse = CVarStrataRoughDiffuse.GetValueOnRenderThread() > 0 ? 1u : 0u;
 
 	Out.PeelLayersAboveDepth = FMath::Max(CVarStrataDebugPeelLayersAboveDepth.GetValueOnRenderThread(), 0);
+	Out.SliceStoringDebugStrataTree = SliceCount - 1 - STRATA_BASE_PASS_MRT_OUTPUT_COUNT;	// The UAV skips the first slices set as render target
 
 	if (IsStrataEnabled())
 	{
@@ -330,6 +332,7 @@ void BindStrataBasePassUniformParameters(FRDGBuilder& GraphBuilder, const FViewI
 		OutStrataUniformParameters.bRoughDiffuse = StrataSceneData->bRoughDiffuse ? 1u : 0u;
 		OutStrataUniformParameters.MaxBytesPerPixel = StrataSceneData->MaxBytesPerPixel;
 		OutStrataUniformParameters.PeelLayersAboveDepth = StrataSceneData->PeelLayersAboveDepth;
+		OutStrataUniformParameters.SliceStoringDebugStrataTree = StrataSceneData->SliceStoringDebugStrataTree;
 		OutStrataUniformParameters.MaterialTextureArrayUAVWithoutRTs = StrataSceneData->MaterialTextureArrayUAVWithoutRTs;
 		OutStrataUniformParameters.SSSTextureUAV = StrataSceneData->SSSTextureUAV;
 		OutStrataUniformParameters.OpaqueRoughRefractionTextureUAV = StrataSceneData->OpaqueRoughRefractionTextureUAV;
@@ -349,6 +352,7 @@ void BindStrataBasePassUniformParameters(FRDGBuilder& GraphBuilder, const FViewI
 		OutStrataUniformParameters.bRoughDiffuse = 0u;
 		OutStrataUniformParameters.MaxBytesPerPixel = 0;
 		OutStrataUniformParameters.PeelLayersAboveDepth = 0;
+		OutStrataUniformParameters.SliceStoringDebugStrataTree = -1;
 		OutStrataUniformParameters.MaterialTextureArrayUAVWithoutRTs = DummyWritableTextureArrayUAV;
 		OutStrataUniformParameters.SSSTextureUAV = DummyWritableSSSTextureUAV;
 		OutStrataUniformParameters.OpaqueRoughRefractionTextureUAV = DummyWritableRefracTextureUAV;
@@ -363,6 +367,7 @@ static void BindStrataGlobalUniformParameters(FRDGBuilder& GraphBuilder, FStrata
 		OutStrataUniformParameters.bRoughDiffuse = StrataSceneData->bRoughDiffuse ? 1u : 0u;
 		OutStrataUniformParameters.MaxBytesPerPixel = StrataSceneData->MaxBytesPerPixel;
 		OutStrataUniformParameters.PeelLayersAboveDepth = StrataSceneData->PeelLayersAboveDepth;
+		OutStrataUniformParameters.SliceStoringDebugStrataTree = StrataSceneData->SliceStoringDebugStrataTree;
 		OutStrataUniformParameters.TileSize = STRATA_TILE_SIZE;
 		OutStrataUniformParameters.TileSizeLog2 = STRATA_TILE_SIZE_DIV_AS_SHIFT;
 		OutStrataUniformParameters.TileCount = StrataViewData->TileCount_Primary;
@@ -382,6 +387,7 @@ static void BindStrataGlobalUniformParameters(FRDGBuilder& GraphBuilder, FStrata
 		OutStrataUniformParameters.bRoughDiffuse = 0;
 		OutStrataUniformParameters.MaxBytesPerPixel = 0;
 		OutStrataUniformParameters.PeelLayersAboveDepth = 0;
+		OutStrataUniformParameters.SliceStoringDebugStrataTree = -1;
 		OutStrataUniformParameters.TileSize = 0;
 		OutStrataUniformParameters.TileSizeLog2 = 0;
 		OutStrataUniformParameters.TileCount = 0;
diff --git a/Engine/Source/Runtime/Renderer/Private/Strata/Strata.h b/Engine/Source/Runtime/Renderer/Private/Strata/Strata.h
--- a/Engine/Source/Runtime/Renderer/Private/Strata/Strata.h
+++ b/Engine/Source/Runtime/Renderer/Private/Strata/Strata.h
@@ -21,6 +21,7 @@ BEGIN_SHADER_PARAMETER_STRUCT(FStrataBasePassUniformParameters, )
 	SHADER_PARAMETER(uint32, MaxBytesPerPixel)
 	SHADER_PARAMETER(uint32, bRoughDiffuse)
 	SHADER_PARAMETER(uint32, PeelLayersAboveDepth)
+	SHADER_PARAMETER(int32,  SliceStoringDebugStrataTree)
 	SHADER_PARAMETER_RDG_TEXTURE_UAV(RWTexture2DArray<uint>, MaterialTextureArrayUAVWithoutRTs)
 	SHADER_PARAMETER_RDG_TEXTURE_UAV(RWTexture2D<uint2>, SSSTextureUAV)
 	SHADER_PARAMETER_RDG_TEXTURE_UAV(RWTexture2D<float3>, OpaqueRoughRefractionTextureUAV)
@@ -42,6 +43,7 @@ BEGIN_GLOBAL_SHADER_PARAMETER_STRUCT(FStrataGlobalUniformParameters, RENDERER_AP
 	SHADER_PARAMETER(uint32, MaxBytesPerPixel)
 	SHADER_PARAMETER(uint32, bRoughDiffuse)
 	SHADER_PARAMETER(uint32, PeelLayersAboveDepth)
+	SHADER_PARAMETER(int32,  SliceStoringDebugStrataTree)
 	SHADER_PARAMETER(uint32, TileSize)
 	SHADER_PARAMETER(uint32, TileSizeLog2)
 	SHADER_PARAMETER(FIntPoint, TileCount)
@@ -72,6 +74,7 @@ struct FStrataSceneData
 	uint32 MaxBytesPerPixel;
 	bool bRoughDiffuse;
 	int32 PeelLayersAboveDepth;
+	int32 SliceStoringDebugStrataTree;
 
 	// Resources allocated and updated each frame
 
@@ -154,6 +157,7 @@ void AddStrataOpaqueRoughRefractionPasses(
 	FSceneTextures& SceneTextures,
 	TArrayView<const FViewInfo> Views);
 
+bool IsStrataAdvancedVisualizationShadersEnabled();
 bool ShouldRenderStrataDebugPasses(const FViewInfo& View);
 FScreenPassTexture AddStrataDebugPasses(FRDGBuilder& GraphBuilder, const FViewInfo& View, FScreenPassTexture& ScreenPassSceneColor);
 
diff --git a/Engine/Source/Runtime/Renderer/Private/Strata/StrataVisualize.cpp b/Engine/Source/Runtime/Renderer/Private/Strata/StrataVisualize.cpp
--- a/Engine/Source/Runtime/Renderer/Private/Strata/StrataVisualize.cpp
+++ b/Engine/Source/Runtime/Renderer/Private/Strata/StrataVisualize.cpp
@@ -20,6 +20,12 @@ static TAutoConsoleVariable<int32> CVarStrataDebugVisualizeMode(
 	TEXT("Strata debug view mode."),
 	ECVF_RenderThreadSafe);
 
+static TAutoConsoleVariable<int32> CVarStrataDebugAdvancedVisualizationShaders(
+	TEXT("r.Strata.Debug.AdvancedVisualizationShaders"),
+	0,
+	TEXT("Enable advanced strata material debug visualization shaders. Base pass shaders can output such advanced data."),
+	ECVF_ReadOnly | ECVF_RenderThreadSafe);
+
 namespace Strata
 {
 // Forward declarations
@@ -31,6 +37,11 @@ void AddStrataInternalClassificationTilePass(
 	EStrataTileType TileMaterialType,
 	const bool bDebug);
 
+bool IsStrataAdvancedVisualizationShadersEnabled()
+{
+	return CVarStrataDebugAdvancedVisualizationShaders.GetValueOnRenderThread() > 0;
+}
+
 class FMaterialPrintInfoCS : public FGlobalShader
 {
 	DECLARE_GLOBAL_SHADER(FMaterialPrintInfoCS);
@@ -126,9 +137,12 @@ static void AddVisualizeMaterialPasses(FRDGBuilder& GraphBuilder, const FViewInf
 		ShaderPrint::RequestSpaceForLines(64);
 		ShaderPrint::RequestSpaceForCharacters(1024);
 
+		const uint32 ViewMode = FMath::Max(0, CVarStrataDebugVisualizeMode.GetValueOnRenderThread());
+
 		FSceneTextureParameters SceneTextureParameters = GetSceneTextureParameters(GraphBuilder);
 
 		// Print Material info
+		if(ViewMode >= 1 && ViewMode != 3)
 		{
 			FRDGBufferUAVRef PrintOffsetBufferUAV = nullptr;
 
@@ -153,8 +167,7 @@ static void AddVisualizeMaterialPasses(FRDGBuilder& GraphBuilder, const FViewInf
 		}
 
 		// Draw material debug
-		const uint32 ViewMode = FMath::Max(0, CVarStrataDebugVisualizeMode.GetValueOnRenderThread());
-		if (ViewMode > 1)
+		if (ViewMode >= 2 && ViewMode != 3)
 		{
 			FVisualizeMaterialPS::FParameters* PassParameters = GraphBuilder.AllocParameters<FVisualizeMaterialPS::FParameters>();
 			PassParameters->ViewUniformBuffer = View.ViewUniformBuffer;
@@ -169,6 +182,12 @@ static void AddVisualizeMaterialPasses(FRDGBuilder& GraphBuilder, const FViewInf
 
 			FPixelShaderUtils::AddFullscreenPass<FVisualizeMaterialPS>(GraphBuilder, View.ShaderMap, RDG_EVENT_NAME("Strata::VisualizeMaterial(Draw)"), PixelShader, PassParameters, View.ViewRect, PreMultipliedColorTransmittanceBlend);
 		}
+
+		// Draw each material layer indepentely
+		if (ViewMode == 3)
+		{
+			
+		}
 	}
 }
 
@@ -179,6 +198,7 @@ bool ShouldRenderStrataDebugPasses(const FViewInfo& View)
 			(FVisualizeMaterialPS::CanRunStrataVizualizeMaterial(View.GetShaderPlatform()) && View.Family && View.Family->EngineShowFlags.VisualizeStrataMaterial)
 			|| (CVarStrataClassificationDebug.GetValueOnAnyThread() > 0)
 			|| ShouldRenderStrataRoughRefractionRnD()
+			|| IsStrataAdvancedVisualizationShadersEnabled()
 			);
 }
 
