Use ShouldRenderRayTracingTranslucency() in GetRayTracingTranslucencyOptions() to consistently compute whether RT translucency should be used

diff --git a/Engine/Source/Runtime/Renderer/Private/RayTracing/RayTracingTranslucency.cpp b/Engine/Source/Runtime/Renderer/Private/RayTracing/RayTracingTranslucency.cpp
--- a/Engine/Source/Runtime/Renderer/Private/RayTracing/RayTracingTranslucency.cpp
+++ b/Engine/Source/Runtime/Renderer/Private/RayTracing/RayTracingTranslucency.cpp
@@ -112,7 +112,7 @@ FRayTracingPrimaryRaysOptions GetRayTracingTranslucencyOptions(const FViewInfo&
 {
 	FRayTracingPrimaryRaysOptions Options;
 
-	Options.bEnabled = ShouldRenderRayTracingEffect(CVarRayTracingTranslucency.GetValueOnRenderThread() != 0, ERayTracingPipelineCompatibilityFlags::FullPipeline, &View);
+	Options.bEnabled = ShouldRenderRayTracingTranslucency(View);
 	Options.SamplerPerPixel = GRayTracingTranslucencySamplesPerPixel >= 0 ? GRayTracingTranslucencySamplesPerPixel : View.FinalPostProcessSettings.RayTracingTranslucencySamplesPerPixel;
 	Options.ApplyHeightFog = GRayTracingTranslucencyHeightFog;
 	Options.PrimaryRayBias = GRayTracingTranslucencyPrimaryRayBias;
