Strata - Advanced debug progress.
 - This debug mode is mode 3.
 - StrataTree can be written/read and used in later pass to print/draw on screen.

diff --git a/Engine/Shaders/Private/Strata/Strata.ush b/Engine/Shaders/Private/Strata/Strata.ush
--- a/Engine/Shaders/Private/Strata/Strata.ush
+++ b/Engine/Shaders/Private/Strata/Strata.ush
@@ -832,8 +832,8 @@ void StrataPackOutStrataTree(
 	{
 		WRITE_UINT(StrataTree.BSDFs[i].State);
 		WRITE_UINT(StrataTree.BSDFs[i].OperatorIndex);
-		WRITE_FLOAT3(asuint(StrataTree.BSDFs[i].LuminanceWeightV));
-		WRITE_FLOAT3(asuint(StrataTree.BSDFs[i].LuminanceWeightL));
+		WRITE_FLOAT3(StrataTree.BSDFs[i].LuminanceWeightV);
+		WRITE_FLOAT3(StrataTree.BSDFs[i].LuminanceWeightL);
 		WRITE_UINT(StrataTree.BSDFs[i].bIsBottom);
 		WRITE_UINT(StrataTree.BSDFs[i].bIsTop);
 
@@ -862,8 +862,54 @@ void StrataPackOutStrataTree(
 #undef WRITE_FLOAT4
 }
 
-void StrataUnpackInStrataTree()
+void StrataUnpackInStrataTree(
+	Texture2DArray<uint> StrataMaterialData,
+	inout FStrataTree StrataTree,
+	in int SliceStoringDebugStrataTree)
 {
+	StrataTree = GetInitialisedStrataTree();
+	int OutCoord = 0;
+#define READ_UINT(Data) Data = StrataMaterialData[uint3(OutCoord++, 0, SliceStoringDebugStrataTree)];
+#define READ_FLOAT(Data) Data = asfloat(StrataMaterialData[uint3(OutCoord++, 0, SliceStoringDebugStrataTree)]);
+#define READ_FLOAT3(Data) READ_FLOAT(Data.x);READ_FLOAT(Data.y);READ_FLOAT(Data.z);
+#define READ_FLOAT4(Data) READ_FLOAT(Data.x);READ_FLOAT(Data.y);READ_FLOAT(Data.z);READ_FLOAT(Data.w);
+
+	READ_UINT(StrataTree.BSDFCount);
+	READ_UINT(StrataTree.OperatorCount);
+
+	LOOP
+	for (int i = 0; i < StrataTree.BSDFCount; ++i)
+	{
+		READ_UINT(StrataTree.BSDFs[i].State);
+		READ_UINT(StrataTree.BSDFs[i].OperatorIndex);
+		READ_FLOAT3(StrataTree.BSDFs[i].LuminanceWeightV);
+		READ_FLOAT3(StrataTree.BSDFs[i].LuminanceWeightL);
+		READ_UINT(StrataTree.BSDFs[i].bIsBottom);
+		READ_UINT(StrataTree.BSDFs[i].bIsTop);
+
+	#if STRATA_INLINE_SHADING
+		READ_FLOAT(StrataTree.BSDFs[i].Coverage);
+
+		READ_FLOAT3(StrataTree.BSDFs[i].Emissive);
+		READ_FLOAT(StrataTree.BSDFs[i].ThicknessCm);
+		READ_FLOAT(StrataTree.BSDFs[i].TmpMFP);
+
+		READ_FLOAT(StrataTree.BSDFs[i].TopLayerDataWeight);
+	#endif
+
+		// Skipping specific single layer water data
+
+		READ_FLOAT4(StrataTree.BSDFs[i].VGPRs[0]);
+		READ_FLOAT4(StrataTree.BSDFs[i].VGPRs[1]);
+		READ_FLOAT4(StrataTree.BSDFs[i].VGPRs[2]);
+		READ_FLOAT4(StrataTree.BSDFs[i].VGPRs[3]);
+		READ_FLOAT4(StrataTree.BSDFs[i].VGPRs[4]);
+	}
+
+#undef READ_UINT
+#undef READ_FLOAT
+#undef READ_FLOAT3
+#undef READ_FLOAT4
 }
 
 
diff --git a/Engine/Shaders/Private/Strata/StrataVisualize.usf b/Engine/Shaders/Private/Strata/StrataVisualize.usf
--- a/Engine/Shaders/Private/Strata/StrataVisualize.usf
+++ b/Engine/Shaders/Private/Strata/StrataVisualize.usf
@@ -3,7 +3,12 @@
 #include "/Engine/Private/Common.ush"
 #include "/Engine/Private/DeferredShadingCommon.ush"
 
+#ifdef SHADER_DEBUGSTRATATREE
+#define STRATA_INLINE_SHADING 1
+#else
 #define STRATA_INLINE_SHADING 0
+#endif
+
 #define STRATA_SSS_MATERIAL_OVERRIDE 0
 #define STRATA_SSS_TRANSMISSION 0  // Disable for now, as the SSS profile texture need to be bound
 #include "/Engine/Private/Strata/StrataEvaluation.ush"
@@ -446,3 +451,91 @@ void VisualizeMaterialPS(
 	}
 }
 #endif // SHADER_MATERIALVISUALIZE
+
+///////////////////////////////////////////////////////////////////////////////////////////////////
+// Material Visualize
+
+#if SHADER_DEBUGSTRATATREE
+
+[numthreads(1, 1, 1)]
+void MaterialDebugStrataTreeCS(uint3 DispatchThreadId : SV_DispatchThreadID)
+{
+	const uint2 PixelPos = View.CursorPosition;
+
+	FStrataAddressing StrataAddressing = GetStrataPixelDataByteOffset(PixelPos, uint2(View.BufferSizeAndInvSize.xy), Strata.MaxBytesPerPixel);
+	FStrataPixelHeader Header = UnpackStrataHeaderIn(Strata.MaterialTextureArray, StrataAddressing, Strata.TopLayerTexture);
+	if (Header.BSDFCount > 0)
+	{
+		FStrataTree StrataTree = GetInitialisedStrataTree();
+		StrataUnpackInStrataTree(
+			Strata.MaterialTextureArray,
+			StrataTree,
+			Strata.SliceStoringDebugStrataTree + STRATA_BASE_PASS_MRT_OUTPUT_COUNT);
+
+
+		FShaderPrintContext Context;
+		Context = InitShaderPrintContext(true, uint2(50, 50));
+
+
+		Print(Context, StrataTree.BSDFCount, FontWhite, 2, 0);
+		Print(Context, TEXT("BSDFs"), FontWhite);
+		Newline(Context);
+		Print(Context, StrataTree.OperatorCount, FontWhite, 2, 0);
+		Print(Context, TEXT("Operators"), FontWhite);
+		Newline(Context);
+	}
+}
+
+float4 GetMaterialOnSphere(
+	in float2 PixelPos,
+	in float2 Pos,
+	in float  Size,
+	in FStrataBSDF BSDF,
+	inout float4 OutColor)
+{
+	if (all(PixelPos > Pos) && all(PixelPos < (Pos + Size)))
+	{
+		OutColor = float4(StrataGetBSDFDiffuseColor(BSDF), 0);
+	}
+	return OutColor;
+}
+
+void MaterialDebugStrataTreePS(
+	float4 SVPos : SV_POSITION,
+	out float4 OutColor : SV_Target0)
+{
+	OutColor = float4(0.0f, 0.0f, 0.0f, 1.0f);
+
+	FStrataAddressing StrataAddressing = GetStrataPixelDataByteOffset(View.CursorPosition, uint2(View.BufferSizeAndInvSize.xy), Strata.MaxBytesPerPixel);
+	FStrataPixelHeader Header = UnpackStrataHeaderIn(Strata.MaterialTextureArray, StrataAddressing, Strata.TopLayerTexture);
+	if (Header.BSDFCount > 0)
+	{
+		OutColor = float4(0.0f, 0.0f, 0.0f, 0.5f);
+
+		FStrataTree StrataTree = GetInitialisedStrataTree();
+		StrataUnpackInStrataTree(
+			Strata.MaterialTextureArray,
+			StrataTree,
+			Strata.SliceStoringDebugStrataTree + STRATA_BASE_PASS_MRT_OUTPUT_COUNT);
+
+		const float2 PixelPos = SVPos.xy;
+
+		float2 Pos = float2(100.0f, 100.0f);
+		const float SphereSize = 100.0f;
+		for (int i = 0; i < StrataTree.BSDFCount; i++)
+		{
+			OutColor = GetMaterialOnSphere(
+				PixelPos,
+				Pos,
+				SphereSize,
+				StrataTree.BSDFs[i],
+				OutColor);
+		
+			Pos.y += SphereSize + 10.0f;
+		}
+	}
+
+	OutColor = float4(pow(OutColor.rgb,1.0/2.2), OutColor.a);
+}
+
+#endif // SHADER_DEBUGSTRATATREE
diff --git a/Engine/Source/Runtime/Renderer/Private/Strata/Strata.cpp b/Engine/Source/Runtime/Renderer/Private/Strata/Strata.cpp
--- a/Engine/Source/Runtime/Renderer/Private/Strata/Strata.cpp
+++ b/Engine/Source/Runtime/Renderer/Private/Strata/Strata.cpp
@@ -255,7 +255,6 @@ void InitialiseStrataFrameSceneData(FRDGBuilder& GraphBuilder, FSceneRenderer& S
 		// Separated subsurface and rough refraction textures
 		{
 			const bool bIsStrataOpaqueMaterialRoughRefractionEnabled = IsStrataOpaqueMaterialRoughRefractionEnabled();
-			const bool bIsStrataAdvancedVisualizationShadersEnabled = IsStrataAdvancedVisualizationShadersEnabled();
 			const FIntPoint OpaqueRoughRefractionSceneExtent		= bIsStrataOpaqueMaterialRoughRefractionEnabled ? SceneTextureExtent : FIntPoint(4, 4);
 			
 			Out.OpaqueRoughRefractionTexture = GraphBuilder.CreateTexture(
diff --git a/Engine/Source/Runtime/Renderer/Private/Strata/Strata.h b/Engine/Source/Runtime/Renderer/Private/Strata/Strata.h
--- a/Engine/Source/Runtime/Renderer/Private/Strata/Strata.h
+++ b/Engine/Source/Runtime/Renderer/Private/Strata/Strata.h
@@ -157,7 +157,6 @@ void AddStrataOpaqueRoughRefractionPasses(
 	FSceneTextures& SceneTextures,
 	TArrayView<const FViewInfo> Views);
 
-bool IsStrataAdvancedVisualizationShadersEnabled();
 bool ShouldRenderStrataDebugPasses(const FViewInfo& View);
 FScreenPassTexture AddStrataDebugPasses(FRDGBuilder& GraphBuilder, const FViewInfo& View, FScreenPassTexture& ScreenPassSceneColor);
 
diff --git a/Engine/Source/Runtime/Renderer/Private/Strata/StrataVisualize.cpp b/Engine/Source/Runtime/Renderer/Private/Strata/StrataVisualize.cpp
--- a/Engine/Source/Runtime/Renderer/Private/Strata/StrataVisualize.cpp
+++ b/Engine/Source/Runtime/Renderer/Private/Strata/StrataVisualize.cpp
@@ -42,6 +42,12 @@ bool IsStrataAdvancedVisualizationShadersEnabled()
 	return CVarStrataDebugAdvancedVisualizationShaders.GetValueOnRenderThread() > 0;
 }
 
+static bool StrataDebugVisualizationCanRunOnPlatform(EShaderPlatform Platform)
+{
+	// On some consoles, this ALU heavy shader (and with optimisation disables for the sake of low compilation time) would spill registers. So only keep it for the editor.
+	return GetMaxSupportedFeatureLevel(Platform) >= ERHIFeatureLevel::SM5 && IsPCPlatform(Platform);
+}
+
 class FMaterialPrintInfoCS : public FGlobalShader
 {
 	DECLARE_GLOBAL_SHADER(FMaterialPrintInfoCS);
@@ -63,15 +69,9 @@ class FMaterialPrintInfoCS : public FGlobalShader
 		return PermutationVector;
 	}
 
-	static bool CanRunStrataVizualizeMaterial(EShaderPlatform Platform)
-	{
-		// On some consoles, this ALU heavy shader (and with optimisation disables for the sake of low compilation time) would spill registers. So only keep it for the editor.
-		return IsPCPlatform(Platform);
-	}
-
 	static bool ShouldCompilePermutation(const FGlobalShaderPermutationParameters& Parameters)
 	{
-		return GetMaxSupportedFeatureLevel(Parameters.Platform) >= ERHIFeatureLevel::SM5 && Strata::IsStrataEnabled() && CanRunStrataVizualizeMaterial(Parameters.Platform);
+		return Strata::IsStrataEnabled() && StrataDebugVisualizationCanRunOnPlatform(Parameters.Platform) && EnumHasAllFlags(Parameters.Flags, EShaderPermutationFlags::HasEditorOnlyData);
 	}
 
 	static void ModifyCompilationEnvironment(const FGlobalShaderPermutationParameters& Parameters, FShaderCompilerEnvironment& OutEnvironment)
@@ -106,15 +106,43 @@ class FVisualizeMaterialPS : public FGlobalShader
 		return PermutationVector;
 	}
 
-	static bool CanRunStrataVizualizeMaterial(EShaderPlatform Platform)
+	static bool ShouldCompilePermutation(const FGlobalShaderPermutationParameters& Parameters)
+	{
+		return Strata::IsStrataEnabled() && StrataDebugVisualizationCanRunOnPlatform(Parameters.Platform) && EnumHasAllFlags(Parameters.Flags, EShaderPermutationFlags::HasEditorOnlyData);
+	}
+
+	static void ModifyCompilationEnvironment(const FGlobalShaderPermutationParameters& Parameters, FShaderCompilerEnvironment& OutEnvironment)
+	{
+		FGlobalShader::ModifyCompilationEnvironment(Parameters, OutEnvironment);
+
+		// Stay debug and skip optimizations to reduce compilation time on this long shader.
+		OutEnvironment.CompilerFlags.Add(CFLAG_Debug);
+		OutEnvironment.SetDefine(TEXT("SHADER_MATERIALVISUALIZE"), 1);
+	}
+};
+IMPLEMENT_GLOBAL_SHADER(FVisualizeMaterialPS, "/Engine/Private/Strata/StrataVisualize.usf", "VisualizeMaterialPS", SF_Pixel);
+
+class FMaterialDebugStrataTreeCS : public FGlobalShader
+{
+	DECLARE_GLOBAL_SHADER(FMaterialDebugStrataTreeCS);
+	SHADER_USE_PARAMETER_STRUCT(FMaterialDebugStrataTreeCS, FGlobalShader);
+
+	using FPermutationDomain = TShaderPermutationDomain<>;
+
+	BEGIN_SHADER_PARAMETER_STRUCT(FParameters, )
+		SHADER_PARAMETER_STRUCT_REF(FViewUniformShaderParameters, ViewUniformBuffer)
+		SHADER_PARAMETER_RDG_UNIFORM_BUFFER(FStrataGlobalUniformParameters, Strata)
+		SHADER_PARAMETER_STRUCT_INCLUDE(ShaderPrint::FShaderParameters, ShaderPrintParameters)
+	END_SHADER_PARAMETER_STRUCT()
+
+	static FPermutationDomain RemapPermutation(FPermutationDomain PermutationVector)
 	{
-		// On some consoles, this ALU heavy shader (and with optimisation disables for the sake of low compilation time) would spill registers. So only keep it for the editor.
-		return IsPCPlatform(Platform);
+		return PermutationVector;
 	}
 
 	static bool ShouldCompilePermutation(const FGlobalShaderPermutationParameters& Parameters)
 	{
-		return GetMaxSupportedFeatureLevel(Parameters.Platform) >= ERHIFeatureLevel::SM5 && Strata::IsStrataEnabled() && CanRunStrataVizualizeMaterial(Parameters.Platform);
+		return Strata::IsStrataEnabled() && StrataDebugVisualizationCanRunOnPlatform(Parameters.Platform) && EnumHasAllFlags(Parameters.Flags, EShaderPermutationFlags::HasEditorOnlyData);
 	}
 
 	static void ModifyCompilationEnvironment(const FGlobalShaderPermutationParameters& Parameters, FShaderCompilerEnvironment& OutEnvironment)
@@ -123,10 +151,44 @@ class FVisualizeMaterialPS : public FGlobalShader
 
 		// Stay debug and skip optimizations to reduce compilation time on this long shader.
 		OutEnvironment.CompilerFlags.Add(CFLAG_Debug);
-		OutEnvironment.SetDefine(TEXT("SHADER_MATERIALVISUALIZE"), 1);
+		OutEnvironment.SetDefine(TEXT("SHADER_DEBUGSTRATATREE"), 1);
 	}
 };
-IMPLEMENT_GLOBAL_SHADER(FVisualizeMaterialPS, "/Engine/Private/Strata/StrataVisualize.usf", "VisualizeMaterialPS", SF_Pixel);
+IMPLEMENT_GLOBAL_SHADER(FMaterialDebugStrataTreeCS, "/Engine/Private/Strata/StrataVisualize.usf", "MaterialDebugStrataTreeCS", SF_Compute);
+
+class FMaterialDebugStrataTreePS : public FGlobalShader
+{
+	DECLARE_GLOBAL_SHADER(FMaterialDebugStrataTreePS);
+	SHADER_USE_PARAMETER_STRUCT(FMaterialDebugStrataTreePS, FGlobalShader);
+
+	using FPermutationDomain = TShaderPermutationDomain<>;
+
+	BEGIN_SHADER_PARAMETER_STRUCT(FParameters, )
+		SHADER_PARAMETER_STRUCT_REF(FViewUniformShaderParameters, ViewUniformBuffer)
+		SHADER_PARAMETER_RDG_UNIFORM_BUFFER(FStrataGlobalUniformParameters, Strata)
+		RENDER_TARGET_BINDING_SLOTS()
+	END_SHADER_PARAMETER_STRUCT()
+
+	static FPermutationDomain RemapPermutation(FPermutationDomain PermutationVector)
+	{
+		return PermutationVector;
+	}
+
+	static bool ShouldCompilePermutation(const FGlobalShaderPermutationParameters& Parameters)
+	{
+		return Strata::IsStrataEnabled() && StrataDebugVisualizationCanRunOnPlatform(Parameters.Platform) && EnumHasAllFlags(Parameters.Flags, EShaderPermutationFlags::HasEditorOnlyData);
+	}
+
+	static void ModifyCompilationEnvironment(const FGlobalShaderPermutationParameters& Parameters, FShaderCompilerEnvironment& OutEnvironment)
+	{
+		FGlobalShader::ModifyCompilationEnvironment(Parameters, OutEnvironment);
+
+		// Stay debug and skip optimizations to reduce compilation time on this long shader.
+		OutEnvironment.CompilerFlags.Add(CFLAG_Debug);
+		OutEnvironment.SetDefine(TEXT("SHADER_DEBUGSTRATATREE"), 1);
+	}
+};
+IMPLEMENT_GLOBAL_SHADER(FMaterialDebugStrataTreePS, "/Engine/Private/Strata/StrataVisualize.usf", "MaterialDebugStrataTreePS", SF_Pixel);
 
 static void AddVisualizeMaterialPasses(FRDGBuilder& GraphBuilder, const FViewInfo& View, FRDGTextureRef SceneColorTexture, EShaderPlatform Platform)
 {
@@ -183,23 +245,47 @@ static void AddVisualizeMaterialPasses(FRDGBuilder& GraphBuilder, const FViewInf
 			FPixelShaderUtils::AddFullscreenPass<FVisualizeMaterialPS>(GraphBuilder, View.ShaderMap, RDG_EVENT_NAME("Strata::VisualizeMaterial(Draw)"), PixelShader, PassParameters, View.ViewRect, PreMultipliedColorTransmittanceBlend);
 		}
 
-		// Draw each material layer indepentely
+		// Draw each material layer independently
 		if (ViewMode == 3)
 		{
-			
+			if (IsStrataAdvancedVisualizationShadersEnabled())
+			{
+				{
+					FMaterialDebugStrataTreeCS::FParameters* PassParameters = GraphBuilder.AllocParameters<FMaterialDebugStrataTreeCS::FParameters>();
+					PassParameters->ViewUniformBuffer = View.ViewUniformBuffer;
+					PassParameters->Strata = Strata::BindStrataGlobalUniformParameters(View);
+					ShaderPrint::SetParameters(GraphBuilder, View.ShaderPrintData, PassParameters->ShaderPrintParameters);
+
+					TShaderMapRef<FMaterialDebugStrataTreeCS> ComputeShader(View.ShaderMap);
+					FComputeShaderUtils::AddPass(GraphBuilder, RDG_EVENT_NAME("Strata::StrataAdvancedVisualization(Print)"), ComputeShader, PassParameters, FIntVector(1, 1, 1));
+				}
+
+				{
+					FMaterialDebugStrataTreePS::FParameters* PassParameters = GraphBuilder.AllocParameters<FMaterialDebugStrataTreePS::FParameters>();
+					PassParameters->ViewUniformBuffer = View.ViewUniformBuffer;
+					PassParameters->Strata = Strata::BindStrataGlobalUniformParameters(View);
+					PassParameters->RenderTargets[0] = FRenderTargetBinding(SceneColorTexture, ERenderTargetLoadAction::ELoad);
+
+					FMaterialDebugStrataTreePS::FPermutationDomain PermutationVector;
+					TShaderMapRef<FMaterialDebugStrataTreePS> PixelShader(View.ShaderMap, PermutationVector);
+
+					FPixelShaderUtils::AddFullscreenPass<FMaterialDebugStrataTreePS>(GraphBuilder, View.ShaderMap, RDG_EVENT_NAME("Strata::StrataAdvancedVisualization(Draw)"), PixelShader, PassParameters, View.ViewRect, PreMultipliedColorTransmittanceBlend);
+				}
+			}
+			else
+			{
+				// TODO warn the user that this mode is not available with delegate SceneRenderer.OnGetOnScreenMessages.AddLambda([](FScreenMessageWriter& ScreenMessageWriter)->void
+			}
 		}
 	}
 }
 
 bool ShouldRenderStrataDebugPasses(const FViewInfo& View)
 {
-	return IsStrataEnabled() &&
-		(
-			(FVisualizeMaterialPS::CanRunStrataVizualizeMaterial(View.GetShaderPlatform()) && View.Family && View.Family->EngineShowFlags.VisualizeStrataMaterial)
-			|| (CVarStrataClassificationDebug.GetValueOnAnyThread() > 0)
-			|| ShouldRenderStrataRoughRefractionRnD()
-			|| IsStrataAdvancedVisualizationShadersEnabled()
-			);
+	return IsStrataEnabled() && (
+		(StrataDebugVisualizationCanRunOnPlatform(View.GetShaderPlatform()) && ( (View.Family && View.Family->EngineShowFlags.VisualizeStrataMaterial) || CVarStrataClassificationDebug.GetValueOnAnyThread() > 0 ))
+		|| ShouldRenderStrataRoughRefractionRnD()
+		);
 }
 
 FScreenPassTexture AddStrataDebugPasses(FRDGBuilder& GraphBuilder, const FViewInfo& View, FScreenPassTexture& ScreenPassSceneColor)
@@ -207,7 +293,7 @@ FScreenPassTexture AddStrataDebugPasses(FRDGBuilder& GraphBuilder, const FViewIn
 	check(IsStrataEnabled());
 	EShaderPlatform Platform = View.GetShaderPlatform();
 
-	if (FVisualizeMaterialPS::CanRunStrataVizualizeMaterial(Platform))
+	if (StrataDebugVisualizationCanRunOnPlatform(Platform))
 	{
 		RDG_EVENT_SCOPE(GraphBuilder, "Strata::VisualizeMaterial");
 		AddVisualizeMaterialPasses(GraphBuilder, View, ScreenPassSceneColor.Texture, Platform);
