Have r.HDR.Display.OutputDevice enum values defined in a single place and replaced hardcoded numbers by enum values.

diff --git a/Engine/Plugins/Media/BinkMedia/Source/BinkMediaPlayer/Private/BinkMovieStreamer.cpp b/Engine/Plugins/Media/BinkMedia/Source/BinkMediaPlayer/Private/BinkMovieStreamer.cpp
--- a/Engine/Plugins/Media/BinkMedia/Source/BinkMediaPlayer/Private/BinkMovieStreamer.cpp
+++ b/Engine/Plugins/Media/BinkMedia/Source/BinkMediaPlayer/Private/BinkMovieStreamer.cpp
@@ -136,23 +136,23 @@ bool FBinkMovieStreamer::Tick(float DeltaTime)
 			static const auto CVarDisplayOutputDevice = IConsoleManager::Get().FindTConsoleVariableDataInt(TEXT("r.HDR.Display.OutputDevice"));
 			if (GRHISupportsHDROutput && CVarHDROutputEnabled->GetValueOnRenderThread() != 0)
 			{
-				int outDev = CVarDisplayOutputDevice->GetValueOnRenderThread();
+				EDisplayOutputFormat outDev = static_cast<EDisplayOutputFormat>(CVarDisplayOutputDevice->GetValueOnRenderThread());
 				switch (outDev)
 				{
 				// LDR
-				case 0:
-				case 1:
-				case 2:
+				case EDisplayOutputFormat::SDR_sRGB:
+				case EDisplayOutputFormat::SDR_Rec709:
+				case EDisplayOutputFormat::SDR_ExplicitGammaMapping:
 					BinkPluginSetHdrSettings(bnk, 1, 1.0f, 80);
 					break;
 				// 1k nits
-				case 3:
-				case 5:
+				case EDisplayOutputFormat::HDR_ACES_1000nit_ST2084:
+				case EDisplayOutputFormat::HDR_ACES_1000nit_ScRGB:
 					BinkPluginSetHdrSettings(bnk, 1, 1.0f, 1000);
 					break;
 				// 2k nits
-				case 4:
-				case 6:
+				case EDisplayOutputFormat::HDR_ACES_2000nit_ST2084:
+				case EDisplayOutputFormat::HDR_ACES_2000nit_ScRGB:
 					BinkPluginSetHdrSettings(bnk, 1, 1.0f, 2000);
 					break;
 				// no tonemap
diff --git a/Engine/Shaders/Private/CompositeUIPixelShader.usf b/Engine/Shaders/Private/CompositeUIPixelShader.usf
--- a/Engine/Shaders/Private/CompositeUIPixelShader.usf
+++ b/Engine/Shaders/Private/CompositeUIPixelShader.usf
@@ -26,14 +26,14 @@ void GenerateLUTPS(
 
 	float3 LinearColor = ST2084ToLinear(Neutral);
 
-	if( OutputDevice == 3 || OutputDevice == 5 )
+	if( OutputDevice == TONEMAPPER_OUTPUT_ACES1000nitST2084 || OutputDevice == TONEMAPPER_OUTPUT_ACES1000nitScRGB)
 	{		
 		// ACES 1000nit transform with PQ/2084 encoding, user specified gamut 
 		float3 ODTColor = ACESOutputTransforms1000(LinearColor);
 		ODTColor = mul(AP1_2_Output, ODTColor);
 		OutColor.rgb = LinearToST2084(ODTColor);
 	}	
-	else if( OutputDevice == 4 || OutputDevice == 6 )
+	else if( OutputDevice == TONEMAPPER_OUTPUT_ACES2000nitST2084 || OutputDevice == TONEMAPPER_OUTPUT_ACES2000nitScRGB)
 	{		
 		// ACES 2000nit transform with PQ/2084 encoding, user specified gamut 
 		float3 ODTColor = ACESOutputTransforms2000(LinearColor);
diff --git a/Engine/Shaders/Private/DeviceEncodingOnlyCommon.ush b/Engine/Shaders/Private/DeviceEncodingOnlyCommon.ush
deleted file mode 100644
diff --git a/Engine/Shaders/Private/PostProcessCombineLUTs.usf b/Engine/Shaders/Private/PostProcessCombineLUTs.usf
--- a/Engine/Shaders/Private/PostProcessCombineLUTs.usf
+++ b/Engine/Shaders/Private/PostProcessCombineLUTs.usf
@@ -210,7 +210,7 @@ float4 CombineLUTsCommon(float2 InUV, uint InLayerIndex)
 	float3 LUTEncodedColor = Neutral.rgb;
 	float3 LinearColor;
 	// Decode texture values as ST-2084 (Dolby PQ)
-	if (GetOutputDevice() >= 3)
+	if (GetOutputDevice() >= TONEMAPPER_OUTPUT_ACES1000nitST2084)
 	{
 		// Since ST2084 returns linear values in nits, divide by a scale factor to convert
 		// the reference nit result to be 1.0 in linear.
@@ -343,7 +343,7 @@ float4 CombineLUTsCommon(float2 InUV, uint InLayerIndex)
 	BRANCH
 
 	// sRGB, user specified gamut
-	if( GetOutputDevice() == 0 )
+	if( GetOutputDevice() == TONEMAPPER_OUTPUT_sRGB)
 	{		
 		// Convert from sRGB to specified output gamut	
 		//float3 OutputGamutColor = mul( AP1_2_Output, mul( sRGB_2_AP1, FilmColor ) );
@@ -357,7 +357,7 @@ float4 CombineLUTsCommon(float2 InUV, uint InLayerIndex)
 	}
 
 	// Rec 709, user specified gamut
-	else if( GetOutputDevice() == 1 )
+	else if( GetOutputDevice() == TONEMAPPER_OUTPUT_Rec709)
 	{
 		// Convert from sRGB to specified output gamut
 		float3 OutputGamutColor = mul( AP1_2_Output, mul( sRGB_2_AP1, FilmColor ) );
@@ -367,7 +367,7 @@ float4 CombineLUTsCommon(float2 InUV, uint InLayerIndex)
 	}
 
 	// ACES 1000nit transform with PQ/2084 encoding, user specified gamut 
-	else if( GetOutputDevice() == 3 || GetOutputDevice() == 5 )
+	else if( GetOutputDevice() == TONEMAPPER_OUTPUT_ACES1000nitST2084 || GetOutputDevice() == TONEMAPPER_OUTPUT_ACES1000nitScRGB)
 	{		
 		// 1000 nit ODT
 		float3 ODTColor = ACESOutputTransforms1000( GradedColor );
@@ -380,7 +380,7 @@ float4 CombineLUTsCommon(float2 InUV, uint InLayerIndex)
 	}
 
 	// ACES 2000nit transform with PQ/2084 encoding, user specified gamut 
-	else if( GetOutputDevice() == 4 || GetOutputDevice() == 6 )
+	else if( GetOutputDevice() == TONEMAPPER_OUTPUT_ACES2000nitST2084 || GetOutputDevice() == TONEMAPPER_OUTPUT_ACES2000nitScRGB)
 	{		
 		// 2000 nit ODT
 		float3 ODTColor = ACESOutputTransforms2000( GradedColor );
@@ -392,25 +392,25 @@ float4 CombineLUTsCommon(float2 InUV, uint InLayerIndex)
 		OutDeviceColor = LinearToST2084( ODTColor );
 	}	
 	
-	else if( GetOutputDevice() == 7 )
+	else if( GetOutputDevice() == TONEMAPPER_OUTPUT_LinearEXR)
 	{
 			float3 OutputGamutColor = mul( AP1_2_Output, mul( sRGB_2_AP1, GradedColor ) );
 			OutDeviceColor = LinearToST2084( OutputGamutColor );
 	}
 	// Linear HDR, including all color correction, but no tone curve
-	else if( GetOutputDevice() == 8 ) 
+	else if( GetOutputDevice() == TONEMAPPER_OUTPUT_NoToneCurve)
 	{
 			OutDeviceColor = GradedColor;
 	}
 	// "Linear" including all color correction and the tone curve, but no device gamma
-	else if (GetOutputDevice() == 9)
+	else if (GetOutputDevice() == TONEMAPPER_OUTPUT_WithToneCurve)
 	{
 		float3 OutputGamutColor = mul(AP1_2_Output, mul(sRGB_2_AP1, FilmColorNoGamma));
 
 		OutDeviceColor = OutputGamutColor;
 	}
 
-	// OutputDevice == 2
+	// OutputDevice == TONEMAPPER_OUTPUT_ExplicitGammaMapping
 	// Gamma 2.2, user specified gamut
 	else
 	{
diff --git a/Engine/Shaders/Private/PostProcessDeviceEncodingOnly.usf b/Engine/Shaders/Private/PostProcessDeviceEncodingOnly.usf
--- a/Engine/Shaders/Private/PostProcessDeviceEncodingOnly.usf
+++ b/Engine/Shaders/Private/PostProcessDeviceEncodingOnly.usf
@@ -15,18 +15,6 @@
 #include "TonemapCommon.ush"
 #include "ScreenPass.ush"
 
-/** Values of DIM_OUTPUT_DEVICE (matching C++'s ETonemapperOutputDevice) */
-#define TONEMAPPER_OUTPUT_sRGB 0
-#define TONEMAPPER_OUTPUT_Rec709 1
-#define TONEMAPPER_OUTPUT_ExplicitGammaMapping 2
-#define TONEMAPPER_OUTPUT_ACES1000nitST2084 3
-#define TONEMAPPER_OUTPUT_ACES2000nitST2084 4
-#define TONEMAPPER_OUTPUT_ACES1000nitScRGB 5
-#define TONEMAPPER_OUTPUT_ACES2000nitScRGB 6
-#define TONEMAPPER_OUTPUT_LinearEXR 7
-#define TONEMAPPER_OUTPUT_NoToneCurve 8
-#define TONEMAPPER_OUTPUT_WithToneCurve 9
-
 #ifndef DIM_OUTPUT_DEVICE
 	#define DIM_OUTPUT_DEVICE (TONEMAPPER_OUTPUT_sRGB)
 #endif
diff --git a/Engine/Shaders/Private/PostProcessTonemap.usf b/Engine/Shaders/Private/PostProcessTonemap.usf
--- a/Engine/Shaders/Private/PostProcessTonemap.usf
+++ b/Engine/Shaders/Private/PostProcessTonemap.usf
@@ -19,27 +19,10 @@
 #include "PixelQuadMessagePassing.ush"
 #include "ScreenPass.ush"
 
-
-/** Values of DIM_OUTPUT_DEVICE (matching C++'s ETonemapperOutputDevice) */
-#define TONEMAPPER_OUTPUT_sRGB 0
-#define TONEMAPPER_OUTPUT_Rec709 1
-#define TONEMAPPER_OUTPUT_ExplicitGammaMapping 2
-#define TONEMAPPER_OUTPUT_ACES1000nitST2084 3
-#define TONEMAPPER_OUTPUT_ACES2000nitST2084 4
-#define TONEMAPPER_OUTPUT_ACES1000nitScRGB 5
-#define TONEMAPPER_OUTPUT_ACES2000nitScRGB 6
-#define TONEMAPPER_OUTPUT_LinearEXR 7
-#define TONEMAPPER_OUTPUT_NoToneCurve 8
-#define TONEMAPPER_OUTPUT_WithToneCurve 9
-
 #ifndef DIM_OUTPUT_DEVICE
 	#define DIM_OUTPUT_DEVICE (TONEMAPPER_OUTPUT_sRGB)
 #endif
 
-
-
-
-
 // 0 / 1(slower, visualize multiple internals side by side for the same image portion, see DebugTile)
 #define DEBUG_TONEMAPPER 0
 
diff --git a/Engine/Shaders/Private/TonemapCommon.ush b/Engine/Shaders/Private/TonemapCommon.ush
--- a/Engine/Shaders/Private/TonemapCommon.ush
+++ b/Engine/Shaders/Private/TonemapCommon.ush
@@ -7,6 +7,18 @@
 #include "ACES.ush"
 #include "GammaCorrectionCommon.ush"
 
+/** Values of DIM_OUTPUT_DEVICE (matching C++'s EDisplayOutputFormat) */
+#define TONEMAPPER_OUTPUT_sRGB 0
+#define TONEMAPPER_OUTPUT_Rec709 1
+#define TONEMAPPER_OUTPUT_ExplicitGammaMapping 2
+#define TONEMAPPER_OUTPUT_ACES1000nitST2084 3
+#define TONEMAPPER_OUTPUT_ACES2000nitST2084 4
+#define TONEMAPPER_OUTPUT_ACES1000nitScRGB 5
+#define TONEMAPPER_OUTPUT_ACES2000nitScRGB 6
+#define TONEMAPPER_OUTPUT_LinearEXR 7
+#define TONEMAPPER_OUTPUT_NoToneCurve 8
+#define TONEMAPPER_OUTPUT_WithToneCurve 9
+
 // usually 1/2.2, the .y is used for inverse gamma when "gamma only" mode is not used
 half3 InverseGamma;
 
@@ -336,7 +348,7 @@ float3 ST2084ToScRGB(float3 Color, uint OutputDevice)
 	float MaxODTNits = 1000.0f;
 	float MinODTNits = 0.0001f;
 
-	if (OutputDevice == 4 || OutputDevice == 6)
+	if (OutputDevice == TONEMAPPER_OUTPUT_ACES2000nitST2084 || OutputDevice == TONEMAPPER_OUTPUT_ACES2000nitScRGB)
 	{
 		// 2000 nit display
 		MaxODTNits = 2000.0f;
diff --git a/Engine/Source/Runtime/Core/Public/GenericPlatform/GenericPlatformMisc.h b/Engine/Source/Runtime/Core/Public/GenericPlatform/GenericPlatformMisc.h
--- a/Engine/Source/Runtime/Core/Public/GenericPlatform/GenericPlatformMisc.h
+++ b/Engine/Source/Runtime/Core/Public/GenericPlatform/GenericPlatformMisc.h
@@ -430,6 +430,24 @@ enum class EInputOutputFlags : uint8
 
 ENUM_CLASS_FLAGS(EInputOutputFlags);
 
+/**
+ * Defines the type of format the backbuffer expects. You must update values in TonemapCommon.ush when changing this enum
+ */
+enum class EDisplayOutputFormat
+{
+	SDR_sRGB = 0,
+	SDR_Rec709 = 1,
+	SDR_ExplicitGammaMapping = 2,
+	HDR_ACES_1000nit_ST2084 = 3,
+	HDR_ACES_2000nit_ST2084 = 4,
+	HDR_ACES_1000nit_ScRGB = 5,
+	HDR_ACES_2000nit_ScRGB = 6,
+	HDR_LinearEXR = 7,
+	HDR_LinearNoToneCurve = 8,
+	HDR_LinearWithToneCurve = 9,
+	MAX
+};
+
 
 /**
  * Different types of Context Switch stats
diff --git a/Engine/Source/Runtime/Core/Public/IOS/IOSPlatformMisc.h b/Engine/Source/Runtime/Core/Public/IOS/IOSPlatformMisc.h
--- a/Engine/Source/Runtime/Core/Public/IOS/IOSPlatformMisc.h
+++ b/Engine/Source/Runtime/Core/Public/IOS/IOSPlatformMisc.h
@@ -197,7 +197,7 @@ struct CORE_API FIOSPlatformMisc : public FApplePlatformMisc
 	FORCEINLINE static void ChooseHDRDeviceAndColorGamut(uint32 DeviceId, uint32 DisplayNitLevel, int32& OutputDevice, int32& ColorGamut)
 	{
 		// Linear output to Apple's specific format.
-		OutputDevice = 7;
+		OutputDevice = (int32)EDisplayOutputFormat::HDR_LinearEXR;
 		ColorGamut = 0;
 	}
 
diff --git a/Engine/Source/Runtime/Core/Public/Mac/MacPlatformMisc.h b/Engine/Source/Runtime/Core/Public/Mac/MacPlatformMisc.h
--- a/Engine/Source/Runtime/Core/Public/Mac/MacPlatformMisc.h
+++ b/Engine/Source/Runtime/Core/Public/Mac/MacPlatformMisc.h
@@ -234,7 +234,7 @@ struct CORE_API FMacPlatformMisc : public FApplePlatformMisc
 	FORCEINLINE static void ChooseHDRDeviceAndColorGamut(uint32 DeviceId, uint32 DisplayNitLevel, int32& OutputDevice, int32& ColorGamut)
 	{
 		// ScRGB, 1000 or 2000 nits, DCI-P3
-		OutputDevice = DisplayNitLevel == 1000 ? 5 : 6;
+		OutputDevice = DisplayNitLevel == 1000 ? (int32)EDisplayOutputFormat::HDR_ACES_1000nit_ScRGB : (int32)EDisplayOutputFormat::HDR_ACES_2000nit_ScRGB;
 		ColorGamut = 1;
 	}
 };
diff --git a/Engine/Source/Runtime/Core/Public/Windows/WindowsPlatformMisc.h b/Engine/Source/Runtime/Core/Public/Windows/WindowsPlatformMisc.h
--- a/Engine/Source/Runtime/Core/Public/Windows/WindowsPlatformMisc.h
+++ b/Engine/Source/Runtime/Core/Public/Windows/WindowsPlatformMisc.h
@@ -260,7 +260,7 @@ struct CORE_API FWindowsPlatformMisc
 		if (DeviceId == 0x1002 /*AMD*/ || DeviceId == 0x10DE /*NVIDIA*/)
 		{
 			// ScRGB, 1000 or 2000 nits, Rec2020
-			OutputDevice = (DisplayNitLevel == 1000) ? 5 : 6;
+			OutputDevice = (DisplayNitLevel == 1000) ? (int32)EDisplayOutputFormat::HDR_ACES_1000nit_ScRGB : (int32)EDisplayOutputFormat::HDR_ACES_2000nit_ScRGB;
 			ColorGamut = 2;
 		}
 	}
diff --git a/Engine/Source/Runtime/D3D12RHI/Private/Windows/WindowsD3D12Viewport.cpp b/Engine/Source/Runtime/D3D12RHI/Private/Windows/WindowsD3D12Viewport.cpp
--- a/Engine/Source/Runtime/D3D12RHI/Private/Windows/WindowsD3D12Viewport.cpp
+++ b/Engine/Source/Runtime/D3D12RHI/Private/Windows/WindowsD3D12Viewport.cpp
@@ -426,9 +426,9 @@ void FD3D12Viewport::EnableHDR()
 	if ( GRHISupportsHDROutput && IsHDREnabled() )
 	{
 		static const auto CVarHDROutputDevice = IConsoleManager::Get().FindTConsoleVariableDataInt(TEXT("r.HDR.Display.OutputDevice"));
-		const EDisplayFormat OutputDevice = EDisplayFormat(CVarHDROutputDevice->GetValueOnAnyThread());
+		const EDisplayOutputFormat OutputDevice = EDisplayOutputFormat(CVarHDROutputDevice->GetValueOnAnyThread());
 
-		const float DisplayMaxOutputNits = (OutputDevice == DF_ACES2000_ST_2084 || OutputDevice == DF_ACES2000_ScRGB) ? 2000.f : 1000.f;
+		const float DisplayMaxOutputNits = (OutputDevice == EDisplayOutputFormat::HDR_ACES_2000nit_ST2084 || OutputDevice == EDisplayOutputFormat::HDR_ACES_2000nit_ScRGB) ? 2000.f : 1000.f;
 		const float DisplayMinOutputNits = 0.0f;	// Min output of the display
 		const float DisplayMaxCLL = 0.0f;			// Max content light level in lumens (0.0 == unknown)
 		const float DisplayFALL = 0.0f;				// Frame average light level (0.0 == unknown)
@@ -455,7 +455,7 @@ void FD3D12Viewport::ShutdownHDR()
 	{
 		// Default SDR display data
 		const EDisplayGamut DisplayGamut = DG_Rec709;
-		const EDisplayFormat OutputDevice = DF_sRGB;
+		const EDisplayOutputFormat OutputDevice = EDisplayOutputFormat::SDR_sRGB;
 
 		// Note: These values aren't actually used.
 		const float DisplayMaxOutputNits = 100.0f;	// Max output of the display
@@ -536,7 +536,7 @@ static const FString GetDXGIColorSpaceString(DXGI_COLOR_SPACE_TYPE ColorSpace)
 	return FString::FromInt(ColorSpace);
 };
 
-void FD3D12Viewport::EnsureColorSpace(EDisplayGamut DisplayGamut, EDisplayFormat OutputDevice)
+void FD3D12Viewport::EnsureColorSpace(EDisplayGamut DisplayGamut, EDisplayOutputFormat OutputDevice)
 {
 	ensure(SwapChain4.GetReference());
 
@@ -547,20 +547,20 @@ void FD3D12Viewport::EnsureColorSpace(EDisplayGamut DisplayGamut, EDisplayFormat
 	switch (OutputDevice)
 	{
 		// Gamma 2.2
-	case DF_sRGB:
-	case DF_Rec709:
+	case EDisplayOutputFormat::SDR_sRGB:
+	case EDisplayOutputFormat::SDR_Rec709:
 		NewColorSpace = bPrimaries2020 ? DXGI_COLOR_SPACE_RGB_FULL_G22_NONE_P2020 : DXGI_COLOR_SPACE_RGB_FULL_G22_NONE_P709;
 		break;
 
 		// Gamma ST.2084
-	case DF_ACES1000_ST_2084:
-	case DF_ACES2000_ST_2084:
+	case EDisplayOutputFormat::HDR_ACES_1000nit_ST2084:
+	case EDisplayOutputFormat::HDR_ACES_2000nit_ST2084:
 		NewColorSpace = DXGI_COLOR_SPACE_RGB_FULL_G2084_NONE_P2020;
 		break;
 
 		// Gamma 1.0 (Linear)
-	case DF_ACES1000_ScRGB:
-	case DF_ACES2000_ScRGB:
+	case EDisplayOutputFormat::HDR_ACES_1000nit_ScRGB:
+	case EDisplayOutputFormat::HDR_ACES_2000nit_ScRGB:
 		// Linear. Still supports expanded color space with values >1.0f and <0.0f.
 		// The actual range is determined by the pixel format (e.g. a UNORM format can only ever have 0-1).
 		NewColorSpace = DXGI_COLOR_SPACE_RGB_FULL_G10_NONE_P709;
diff --git a/Engine/Source/Runtime/D3D12RHI/Public/D3D12Viewport.h b/Engine/Source/Runtime/D3D12RHI/Public/D3D12Viewport.h
--- a/Engine/Source/Runtime/D3D12RHI/Public/D3D12Viewport.h
+++ b/Engine/Source/Runtime/D3D12RHI/Public/D3D12Viewport.h
@@ -253,17 +253,6 @@ private:
 		DG_ACEScg
 	};
 
-	enum EDisplayFormat
-	{
-		DF_sRGB,
-		DF_Rec709,
-		DF_ExplicitGammaMapping,
-		DF_ACES1000_ST_2084,
-		DF_ACES2000_ST_2084,
-		DF_ACES1000_ScRGB,
-		DF_ACES2000_ScRGB,
-	};
-
 	struct DisplayChromacities
 	{
 		float RedX, RedY;
@@ -283,7 +272,7 @@ private:
 
 #if PLATFORM_WINDOWS || PLATFORM_HOLOLENS
 	/** Ensure the correct color space is set on the swap chain */
-	void EnsureColorSpace(EDisplayGamut DisplayGamut, EDisplayFormat OutputDevice);
+	void EnsureColorSpace(EDisplayGamut DisplayGamut, EDisplayOutputFormat OutputDevice);
 
 	/** 
 	 * Set HDR meta data. 
diff --git a/Engine/Source/Runtime/MovieSceneCapture/Private/CompositionGraphCaptureProtocol.cpp b/Engine/Source/Runtime/MovieSceneCapture/Private/CompositionGraphCaptureProtocol.cpp
--- a/Engine/Source/Runtime/MovieSceneCapture/Private/CompositionGraphCaptureProtocol.cpp
+++ b/Engine/Source/Runtime/MovieSceneCapture/Private/CompositionGraphCaptureProtocol.cpp
@@ -40,7 +40,7 @@ struct FFrameCaptureViewExtension : public FSceneViewExtensionBase
 		if (CaptureGamut == HCGM_Linear)
 		{
 			CVarDumpGamut->Set(1);
-			CVarDumpDevice->Set(7);
+			CVarDumpDevice->Set((int32)EDisplayOutputFormat::HDR_LinearEXR);
 		}
 		else
 		{
diff --git a/Engine/Source/Runtime/MovieSceneCapture/Private/ImageSequenceProtocol.cpp b/Engine/Source/Runtime/MovieSceneCapture/Private/ImageSequenceProtocol.cpp
--- a/Engine/Source/Runtime/MovieSceneCapture/Private/ImageSequenceProtocol.cpp
+++ b/Engine/Source/Runtime/MovieSceneCapture/Private/ImageSequenceProtocol.cpp
@@ -215,7 +215,7 @@ bool UImageSequenceProtocol_EXR::SetupImpl()
 	if (CaptureGamut == HCGM_Linear)
 	{
 		CVarDumpGamut->Set(1);
-		CVarDumpDevice->Set(7);
+		CVarDumpDevice->Set((int32)EDisplayOutputFormat::HDR_LinearEXR);
 	}
 	else
 	{
diff --git a/Engine/Source/Runtime/Renderer/Private/PostProcess/PostProcessDeviceEncodingOnly.cpp b/Engine/Source/Runtime/Renderer/Private/PostProcess/PostProcessDeviceEncodingOnly.cpp
--- a/Engine/Source/Runtime/Renderer/Private/PostProcess/PostProcessDeviceEncodingOnly.cpp
+++ b/Engine/Source/Runtime/Renderer/Private/PostProcess/PostProcessDeviceEncodingOnly.cpp
@@ -24,7 +24,7 @@ namespace
 	namespace DeviceEncodingOnlyPermutation
 	{
 		// Desktop renderer permutation dimensions.
-		class FDeviceEncodingOnlyOutputDeviceDim : SHADER_PERMUTATION_ENUM_CLASS("DIM_OUTPUT_DEVICE", EDeviceEncodingOnlyOutputDevice);
+		class FDeviceEncodingOnlyOutputDeviceDim : SHADER_PERMUTATION_ENUM_CLASS("DIM_OUTPUT_DEVICE", EDisplayOutputFormat);
 
 		using FDesktopDomain = TShaderPermutationDomain<FDeviceEncodingOnlyOutputDeviceDim>;
 	
@@ -38,24 +38,24 @@ FDeviceEncodingOnlyOutputDeviceParameters GetDeviceEncodingOnlyOutputDeviceParam
 	static TConsoleVariableData<int32>* CVarOutputDevice = IConsoleManager::Get().FindTConsoleVariableDataInt(TEXT("r.HDR.Display.OutputDevice"));
 	static TConsoleVariableData<float>* CVarOutputGamma = IConsoleManager::Get().FindTConsoleVariableDataFloat(TEXT("r.TonemapperGamma"));
 
-	EDeviceEncodingOnlyOutputDevice OutputDeviceValue;
+	EDisplayOutputFormat OutputDeviceValue;
 
 	if (Family.SceneCaptureSource == SCS_FinalColorHDR)
 	{
-		OutputDeviceValue = EDeviceEncodingOnlyOutputDevice::LinearNoToneCurve;
+		OutputDeviceValue = EDisplayOutputFormat::HDR_LinearNoToneCurve;
 	}
 	else if (Family.SceneCaptureSource == SCS_FinalToneCurveHDR)
 	{
-		OutputDeviceValue = EDeviceEncodingOnlyOutputDevice::LinearWithToneCurve;
+		OutputDeviceValue = EDisplayOutputFormat::HDR_LinearWithToneCurve;
 	}
 	else if (Family.bIsHDR)
 	{
-		OutputDeviceValue = EDeviceEncodingOnlyOutputDevice::ACES1000nitST2084;
+		OutputDeviceValue = EDisplayOutputFormat::HDR_ACES_1000nit_ST2084;
 	}
 	else
 	{
-		OutputDeviceValue = static_cast<EDeviceEncodingOnlyOutputDevice>(CVarOutputDevice->GetValueOnRenderThread());
-		OutputDeviceValue = static_cast<EDeviceEncodingOnlyOutputDevice>(FMath::Clamp(static_cast<int32>(OutputDeviceValue), 0, static_cast<int32>(EDeviceEncodingOnlyOutputDevice::MAX) - 1));
+		OutputDeviceValue = static_cast<EDisplayOutputFormat>(CVarOutputDevice->GetValueOnRenderThread());
+		OutputDeviceValue = static_cast<EDisplayOutputFormat>(FMath::Clamp(static_cast<int32>(OutputDeviceValue), 0, static_cast<int32>(EDisplayOutputFormat::MAX) - 1));
 	}
 
 	float Gamma = CVarOutputGamma->GetValueOnRenderThread();
@@ -66,9 +66,9 @@ FDeviceEncodingOnlyOutputDeviceParameters GetDeviceEncodingOnlyOutputDeviceParam
 	}
 
 	// Enforce user-controlled ramp over sRGB or Rec709
-	if (Gamma > 0.0f && (OutputDeviceValue == EDeviceEncodingOnlyOutputDevice::sRGB || OutputDeviceValue == EDeviceEncodingOnlyOutputDevice::Rec709))
+	if (Gamma > 0.0f && (OutputDeviceValue == EDisplayOutputFormat::SDR_sRGB || OutputDeviceValue == EDisplayOutputFormat::SDR_Rec709))
 	{
-		OutputDeviceValue = EDeviceEncodingOnlyOutputDevice::ExplicitGammaMapping;
+		OutputDeviceValue = EDisplayOutputFormat::SDR_ExplicitGammaMapping;
 	}
 
 	FVector InvDisplayGammaValue;
@@ -170,13 +170,13 @@ FScreenPassTexture AddDeviceEncodingOnlyPass(FRDGBuilder& GraphBuilder, const FV
 		OutputDesc.Flags |= GFastVRamConfig.Tonemap;
 
 		const FDeviceEncodingOnlyOutputDeviceParameters OutputDeviceParameters = GetDeviceEncodingOnlyOutputDeviceParameters(*View.Family);
-		const EDeviceEncodingOnlyOutputDevice OutputDevice = static_cast<EDeviceEncodingOnlyOutputDevice>(OutputDeviceParameters.OutputDevice);
+		const EDisplayOutputFormat OutputDevice = static_cast<EDisplayOutputFormat>(OutputDeviceParameters.OutputDevice);
 
-		if (OutputDevice == EDeviceEncodingOnlyOutputDevice::LinearEXR)
+		if (OutputDevice == EDisplayOutputFormat::HDR_LinearEXR)
 		{
 			OutputDesc.Format = PF_A32B32G32R32F;
 		}
-		if (OutputDevice == EDeviceEncodingOnlyOutputDevice::LinearNoToneCurve || OutputDevice == EDeviceEncodingOnlyOutputDevice::LinearWithToneCurve)
+		if (OutputDevice == EDisplayOutputFormat::HDR_LinearNoToneCurve || OutputDevice == EDisplayOutputFormat::HDR_LinearWithToneCurve)
 		{
 			OutputDesc.Format = PF_FloatRGBA;
 		}
@@ -216,7 +216,7 @@ FScreenPassTexture AddDeviceEncodingOnlyPass(FRDGBuilder& GraphBuilder, const FV
 
 	// Generate permutation vector for the desktop tonemapper.
 	DeviceEncodingOnlyPermutation::FDesktopDomain DesktopPermutationVector;
-	DesktopPermutationVector.Set<DeviceEncodingOnlyPermutation::FDeviceEncodingOnlyOutputDeviceDim>(EDeviceEncodingOnlyOutputDevice(CommonParameters.OutputDevice.OutputDevice));
+	DesktopPermutationVector.Set<DeviceEncodingOnlyPermutation::FDeviceEncodingOnlyOutputDeviceDim>(EDisplayOutputFormat(CommonParameters.OutputDevice.OutputDevice));
 
 	// Override output might not support UAVs.
 	const bool bComputePass = (Output.Texture->Desc.Flags & TexCreate_UAV) == TexCreate_UAV ? View.bUseComputePasses : false;
diff --git a/Engine/Source/Runtime/Renderer/Private/PostProcess/PostProcessDeviceEncodingOnly.h b/Engine/Source/Runtime/Renderer/Private/PostProcess/PostProcessDeviceEncodingOnly.h
--- a/Engine/Source/Runtime/Renderer/Private/PostProcess/PostProcessDeviceEncodingOnly.h
+++ b/Engine/Source/Runtime/Renderer/Private/PostProcess/PostProcessDeviceEncodingOnly.h
@@ -7,23 +7,6 @@
 #include "OverridePassSequence.h"
 #include "Math/Halton.h"
 
-// You must update values in PostProcessTonemap.usf when changing this enum.
-enum class EDeviceEncodingOnlyOutputDevice
-{
-	sRGB,
-	Rec709,
-	ExplicitGammaMapping,
-	ACES1000nitST2084,
-	ACES2000nitST2084,
-	ACES1000nitScRGB,
-	ACES2000nitScRGB,
-	LinearEXR,
-	LinearNoToneCurve,
-	LinearWithToneCurve,
-
-	MAX
-};
-
 BEGIN_SHADER_PARAMETER_STRUCT(FDeviceEncodingOnlyOutputDeviceParameters, )
 	SHADER_PARAMETER(FVector3f, InverseGamma)
 	SHADER_PARAMETER(uint32, OutputDevice)
diff --git a/Engine/Source/Runtime/Renderer/Private/PostProcess/PostProcessTonemap.cpp b/Engine/Source/Runtime/Renderer/Private/PostProcess/PostProcessTonemap.cpp
--- a/Engine/Source/Runtime/Renderer/Private/PostProcess/PostProcessTonemap.cpp
+++ b/Engine/Source/Runtime/Renderer/Private/PostProcess/PostProcessTonemap.cpp
@@ -188,7 +188,7 @@ FCommonDomain BuildCommonPermutationDomain(const FViewInfo& View, bool bGammaOnl
 // Desktop renderer permutation dimensions.
 class FTonemapperColorFringeDim       : SHADER_PERMUTATION_BOOL("USE_COLOR_FRINGE");
 class FTonemapperGrainQuantizationDim : SHADER_PERMUTATION_BOOL("USE_GRAIN_QUANTIZATION");
-class FTonemapperOutputDeviceDim      : SHADER_PERMUTATION_ENUM_CLASS("DIM_OUTPUT_DEVICE", ETonemapperOutputDevice);
+class FTonemapperOutputDeviceDim      : SHADER_PERMUTATION_ENUM_CLASS("DIM_OUTPUT_DEVICE", EDisplayOutputFormat);
 
 using FDesktopDomain = TShaderPermutationDomain<
 	FCommonDomain,
@@ -223,13 +223,13 @@ FDesktopDomain RemapPermutation(FDesktopDomain PermutationVector, ERHIFeatureLev
 
 	// Mobile supports only sRGB and LinearNoToneCurve output
 	if (FeatureLevel <= ERHIFeatureLevel::ES3_1 &&
-		PermutationVector.Get<FTonemapperOutputDeviceDim>() != ETonemapperOutputDevice::LinearNoToneCurve)
+		PermutationVector.Get<FTonemapperOutputDeviceDim>() != EDisplayOutputFormat::HDR_LinearNoToneCurve)
 	{
-		PermutationVector.Set<FTonemapperOutputDeviceDim>(ETonemapperOutputDevice::sRGB);
+		PermutationVector.Set<FTonemapperOutputDeviceDim>(EDisplayOutputFormat::SDR_sRGB);
 	}
 
 	// Disable grain quantization for LinearNoToneCurve and LinearWithToneCurve output device
-	if (PermutationVector.Get<FTonemapperOutputDeviceDim>() == ETonemapperOutputDevice::LinearNoToneCurve || PermutationVector.Get<FTonemapperOutputDeviceDim>() == ETonemapperOutputDevice::LinearWithToneCurve)
+	if (PermutationVector.Get<FTonemapperOutputDeviceDim>() == EDisplayOutputFormat::HDR_LinearNoToneCurve || PermutationVector.Get<FTonemapperOutputDeviceDim>() == EDisplayOutputFormat::HDR_LinearWithToneCurve)
 		PermutationVector.Set<FTonemapperGrainQuantizationDim>(false);
 	else
 		PermutationVector.Set<FTonemapperGrainQuantizationDim>(true);
@@ -275,24 +275,24 @@ FTonemapperOutputDeviceParameters GetTonemapperOutputDeviceParameters(const FSce
 	static TConsoleVariableData<int32>* CVarOutputDevice = IConsoleManager::Get().FindTConsoleVariableDataInt(TEXT("r.HDR.Display.OutputDevice"));
 	static TConsoleVariableData<float>* CVarOutputGamma = IConsoleManager::Get().FindTConsoleVariableDataFloat(TEXT("r.TonemapperGamma"));
 
-	ETonemapperOutputDevice OutputDeviceValue;
+	EDisplayOutputFormat OutputDeviceValue;
 
 	if (Family.SceneCaptureSource == SCS_FinalColorHDR)
 	{
-		OutputDeviceValue = ETonemapperOutputDevice::LinearNoToneCurve;
+		OutputDeviceValue = EDisplayOutputFormat::HDR_LinearNoToneCurve;
 	}
 	else if (Family.SceneCaptureSource == SCS_FinalToneCurveHDR)
 	{
-		OutputDeviceValue = ETonemapperOutputDevice::LinearWithToneCurve;
+		OutputDeviceValue = EDisplayOutputFormat::HDR_LinearWithToneCurve;
 	}
 	else if (Family.bIsHDR)
 	{
-		OutputDeviceValue = ETonemapperOutputDevice::ACES1000nitST2084;
+		OutputDeviceValue = EDisplayOutputFormat::HDR_ACES_1000nit_ST2084;
 	}
 	else
 	{
-		OutputDeviceValue = static_cast<ETonemapperOutputDevice>(CVarOutputDevice->GetValueOnRenderThread());
-		OutputDeviceValue = static_cast<ETonemapperOutputDevice>(FMath::Clamp(static_cast<int32>(OutputDeviceValue), 0, static_cast<int32>(ETonemapperOutputDevice::MAX) - 1));
+		OutputDeviceValue = static_cast<EDisplayOutputFormat>(CVarOutputDevice->GetValueOnRenderThread());
+		OutputDeviceValue = static_cast<EDisplayOutputFormat>(FMath::Clamp(static_cast<int32>(OutputDeviceValue), 0, static_cast<int32>(EDisplayOutputFormat::MAX) - 1));
 	}
 
 	float Gamma = CVarOutputGamma->GetValueOnRenderThread();
@@ -303,9 +303,9 @@ FTonemapperOutputDeviceParameters GetTonemapperOutputDeviceParameters(const FSce
 	}
 
 	// Enforce user-controlled ramp over sRGB or Rec709
-	if (Gamma > 0.0f && (OutputDeviceValue == ETonemapperOutputDevice::sRGB || OutputDeviceValue == ETonemapperOutputDevice::Rec709))
+	if (Gamma > 0.0f && (OutputDeviceValue == EDisplayOutputFormat::SDR_sRGB || OutputDeviceValue == EDisplayOutputFormat::SDR_Rec709))
 	{
-		OutputDeviceValue = ETonemapperOutputDevice::ExplicitGammaMapping;
+		OutputDeviceValue = EDisplayOutputFormat::SDR_ExplicitGammaMapping;
 	}
 
 	FVector3f InvDisplayGammaValue;
@@ -610,13 +610,13 @@ FScreenPassTexture AddTonemapPass(FRDGBuilder& GraphBuilder, const FViewInfo& Vi
 		OutputDesc.ClearValue = FClearValueBinding(FLinearColor(0, 0, 0, 0));
 
 		const FTonemapperOutputDeviceParameters OutputDeviceParameters = GetTonemapperOutputDeviceParameters(*View.Family);
-		const ETonemapperOutputDevice OutputDevice = static_cast<ETonemapperOutputDevice>(OutputDeviceParameters.OutputDevice);
+		const EDisplayOutputFormat OutputDevice = static_cast<EDisplayOutputFormat>(OutputDeviceParameters.OutputDevice);
 
-		if (OutputDevice == ETonemapperOutputDevice::LinearEXR)
+		if (OutputDevice == EDisplayOutputFormat::HDR_LinearEXR)
 		{
 			OutputDesc.Format = PF_A32B32G32R32F;
 		}
-		else if (OutputDevice == ETonemapperOutputDevice::LinearNoToneCurve || OutputDevice == ETonemapperOutputDevice::LinearWithToneCurve)
+		else if (OutputDevice == EDisplayOutputFormat::HDR_LinearNoToneCurve || OutputDevice == EDisplayOutputFormat::HDR_LinearWithToneCurve)
 		{
 			OutputDesc.Format = PF_FloatRGBA;
 		}
@@ -895,7 +895,7 @@ FScreenPassTexture AddTonemapPass(FRDGBuilder& GraphBuilder, const FViewInfo& Vi
 			DesktopPermutationVector.Set<TonemapperPermutation::FTonemapperColorFringeDim>(PostProcessSettings.SceneFringeIntensity > 0.01f);
 		}
 
-		DesktopPermutationVector.Set<TonemapperPermutation::FTonemapperOutputDeviceDim>(ETonemapperOutputDevice(CommonParameters.OutputDevice.OutputDevice));
+		DesktopPermutationVector.Set<TonemapperPermutation::FTonemapperOutputDeviceDim>(EDisplayOutputFormat(CommonParameters.OutputDevice.OutputDevice));
 
 		DesktopPermutationVector = TonemapperPermutation::RemapPermutation(DesktopPermutationVector, View.GetFeatureLevel());
 	}
diff --git a/Engine/Source/Runtime/Renderer/Private/PostProcess/PostProcessTonemap.h b/Engine/Source/Runtime/Renderer/Private/PostProcess/PostProcessTonemap.h
--- a/Engine/Source/Runtime/Renderer/Private/PostProcess/PostProcessTonemap.h
+++ b/Engine/Source/Runtime/Renderer/Private/PostProcess/PostProcessTonemap.h
@@ -11,23 +11,6 @@
 
 bool SupportsFilmGrain(EShaderPlatform Platform);
 
-// You must update values in PostProcessTonemap.usf when changing this enum.
-enum class ETonemapperOutputDevice
-{
-	sRGB,
-	Rec709,
-	ExplicitGammaMapping,
-	ACES1000nitST2084,
-	ACES2000nitST2084,
-	ACES1000nitScRGB,
-	ACES2000nitScRGB,
-	LinearEXR,
-	LinearNoToneCurve,
-	LinearWithToneCurve,
-
-	MAX
-};
-
 BEGIN_SHADER_PARAMETER_STRUCT(FTonemapperOutputDeviceParameters, )
 	SHADER_PARAMETER(FVector3f, InverseGamma)
 	SHADER_PARAMETER(uint32, OutputDevice)
diff --git a/Engine/Source/Runtime/SlateRHIRenderer/Private/SlateRHIRenderer.cpp b/Engine/Source/Runtime/SlateRHIRenderer/Private/SlateRHIRenderer.cpp
--- a/Engine/Source/Runtime/SlateRHIRenderer/Private/SlateRHIRenderer.cpp
+++ b/Engine/Source/Runtime/SlateRHIRenderer/Private/SlateRHIRenderer.cpp
@@ -608,7 +608,7 @@ public:
 		if (Gamma > 0.0f)
 		{
 			// Enforce user-controlled ramp over sRGB or Rec709
-			OutputDeviceValue = FMath::Max(OutputDeviceValue, 2);
+			OutputDeviceValue = FMath::Max(OutputDeviceValue, (int32)EDisplayOutputFormat::SDR_ExplicitGammaMapping);
 		}
 
 		SetShaderValue(RHICmdList, RHICmdList.GetBoundPixelShader(), OutputDevice, OutputDeviceValue);
@@ -755,7 +755,7 @@ void FSlateRHIRenderer::DrawWindow_RenderThread(FRHICommandListImmediate& RHICmd
 		static const auto CVarHDROutputDevice = IConsoleManager::Get().FindTConsoleVariableDataInt(TEXT("r.HDR.Display.OutputDevice"));
 		static const auto CVarHDROutputGamut = IConsoleManager::Get().FindTConsoleVariableDataInt(TEXT("r.HDR.Display.ColorGamut"));
 
-		const int32 HDROutputDevice = CVarHDROutputDevice ? CVarHDROutputDevice->GetValueOnRenderThread() : 0;
+		const int32 HDROutputDevice = CVarHDROutputDevice ? CVarHDROutputDevice->GetValueOnRenderThread() : (int32)EDisplayOutputFormat::SDR_sRGB;
 		const int32 HDROutputGamut = CVarHDROutputGamut ? CVarHDROutputGamut->GetValueOnRenderThread() : 0;
 
 		bool bLUTStale = ViewportInfo.ColorSpaceLUTOutputDevice != HDROutputDevice || ViewportInfo.ColorSpaceLUTOutputGamut != HDROutputGamut;
@@ -1039,7 +1039,7 @@ void FSlateRHIRenderer::DrawWindow_RenderThread(FRHICommandListImmediate& RHICmd
 						TShaderMapRef<FScreenVS> VertexShader(ShaderMap);
 
 						FRHITexture* UITargetRTMaskTexture = RHISupportsRenderTargetWriteMask(GMaxRHIShaderPlatform) ? ViewportInfo.UITargetRTMask->GetRHI() : nullptr;
-						if (HDROutputDevice == 5 || HDROutputDevice == 6)
+						if (HDROutputDevice == (int32)EDisplayOutputFormat::HDR_ACES_1000nit_ScRGB || HDROutputDevice == (int32)EDisplayOutputFormat::HDR_ACES_2000nit_ScRGB)
 						{
 							// ScRGB encoding
 							TShaderMapRef<FCompositePS<1>> PixelShader(ShaderMap);
diff --git a/Engine/Source/Runtime/SlateRHIRenderer/Private/SlateShaders.cpp b/Engine/Source/Runtime/SlateRHIRenderer/Private/SlateShaders.cpp
--- a/Engine/Source/Runtime/SlateRHIRenderer/Private/SlateShaders.cpp
+++ b/Engine/Source/Runtime/SlateRHIRenderer/Private/SlateShaders.cpp
@@ -110,7 +110,7 @@ void FSlateInstancedVertexDeclaration::InitRHI()
 void FSlateElementPS::ModifyCompilationEnvironment(const FGlobalShaderPermutationParameters& Parameters, FShaderCompilerEnvironment& OutEnvironment)
 {
 	static const auto CVar = IConsoleManager::Get().FindTConsoleVariableDataInt(TEXT("r.HDR.Display.OutputDevice"));
-	OutEnvironment.SetDefine(TEXT("USE_709"), CVar ? (CVar->GetValueOnGameThread() == 1) : 1);
+	OutEnvironment.SetDefine(TEXT("USE_709"), CVar ? (CVar->GetValueOnGameThread() == (int32)EDisplayOutputFormat::SDR_Rec709) : 1);
 }
 
 
diff --git a/Engine/Source/Runtime/SlateRHIRenderer/Private/SlateShaders.h b/Engine/Source/Runtime/SlateRHIRenderer/Private/SlateShaders.h
--- a/Engine/Source/Runtime/SlateRHIRenderer/Private/SlateShaders.h
+++ b/Engine/Source/Runtime/SlateRHIRenderer/Private/SlateShaders.h
@@ -560,8 +560,6 @@ public:
 
 	void SetParameters(FRHICommandList& RHICmdList, FRHITexture* SceneTextureRHI)
 	{
-		static const auto CVarOutputDevice = IConsoleManager::Get().FindTConsoleVariableDataInt(TEXT("r.HDR.Display.OutputDevice"));
-
 		SetTextureParameter(RHICmdList, RHICmdList.GetBoundPixelShader(), SceneTexture, SceneSampler, TStaticSamplerState<SF_Point>::GetRHI(), SceneTextureRHI);
 		
 		static auto CVarHDRNITLevel = IConsoleManager::Get().FindConsoleVariable(TEXT("Editor.HDRNITLevel"));
diff --git a/Engine/Source/Runtime/VulkanRHI/Private/VulkanSwapChain.cpp b/Engine/Source/Runtime/VulkanRHI/Private/VulkanSwapChain.cpp
--- a/Engine/Source/Runtime/VulkanRHI/Private/VulkanSwapChain.cpp
+++ b/Engine/Source/Runtime/VulkanRHI/Private/VulkanSwapChain.cpp
@@ -175,18 +175,17 @@ FVulkanSwapChain::FVulkanSwapChain(VkInstance InInstance, FVulkanDevice& InDevic
 		if (Formats[Index].colorSpace != RequestedColorSpace)
 		{
 			static const auto CVarHDROutputDevice = IConsoleManager::Get().FindTConsoleVariableDataInt(TEXT("r.HDR.Display.OutputDevice"));
-			int32 OutputDevice = CVarHDROutputDevice ? CVarHDROutputDevice->GetValueOnAnyThread() : 0;
-			// The possible values are documented in PostProcessTonemap.cpp, where the cvar is defined. They match the ETonemapperOutputDevice enum, which is defined in a header we cannot include.
+			EDisplayOutputFormat OutputDevice = CVarHDROutputDevice ? (EDisplayOutputFormat)CVarHDROutputDevice->GetValueOnAnyThread() : EDisplayOutputFormat::SDR_sRGB;
 			switch (OutputDevice)
 			{
-			case 0:
+			case EDisplayOutputFormat::SDR_sRGB:
 				RequestedColorSpace = VK_COLOR_SPACE_SRGB_NONLINEAR_KHR;
 				break;
-			case 1:
+			case EDisplayOutputFormat::SDR_Rec709:
 				RequestedColorSpace = VK_COLOR_SPACE_BT709_NONLINEAR_EXT;
 				break;
-			case 3:
-			case 4:
+			case EDisplayOutputFormat::HDR_ACES_1000nit_ST2084:
+			case EDisplayOutputFormat::HDR_ACES_2000nit_ST2084:
 				RequestedColorSpace = VK_COLOR_SPACE_HDR10_ST2084_EXT;
 				break;
 			default:
diff --git a/Engine/Source/Runtime/Windows/D3D11RHI/Private/Windows/WindowsD3D11Device.cpp b/Engine/Source/Runtime/Windows/D3D11RHI/Private/Windows/WindowsD3D11Device.cpp
--- a/Engine/Source/Runtime/Windows/D3D11RHI/Private/Windows/WindowsD3D11Device.cpp
+++ b/Engine/Source/Runtime/Windows/D3D11RHI/Private/Windows/WindowsD3D11Device.cpp
@@ -490,7 +490,7 @@ void FD3D11DynamicRHI::EnableHDR()
 	{
 		const int32 OutputDevice = CVarHDROutputDevice->GetValueOnAnyThread();
 
-		const float DisplayMaxOutputNits = (OutputDevice == 4 || OutputDevice == 6) ? 2000.f : 1000.f;
+		const float DisplayMaxOutputNits = (OutputDevice == (int32)EDisplayOutputFormat::HDR_ACES_2000nit_ST2084 || OutputDevice == (int32)EDisplayOutputFormat::HDR_ACES_2000nit_ScRGB) ? 2000.f : 1000.f;
 		const float DisplayMinOutputNits = 0.0f;	// Min output of the display
 		const float DisplayMaxCLL = 0.0f;			// Max content light level in lumens (0.0 == unknown)
 		const float DisplayFALL = 0.0f;				// Frame average light level (0.0 == unknown)
