PathTracer: improve filtering of which scene proxies should impact path tracer invalidation

diff --git a/Engine/Source/Runtime/Renderer/Private/RendererScene.cpp b/Engine/Source/Runtime/Renderer/Private/RendererScene.cpp
--- a/Engine/Source/Runtime/Renderer/Private/RendererScene.cpp
+++ b/Engine/Source/Runtime/Renderer/Private/RendererScene.cpp
@@ -4352,11 +4352,15 @@ inline void FScene::UpdateRayTracingGroupBounds_UpdatePrimitives(const Experimen
 
 static inline bool IsPrimitiveRelevantToPathTracing(FPrimitiveSceneInfo* PrimitiveSceneInfo)
 {
+#if RHI_RAYTRACING
 	// returns true if the primitive is likely to impact the path traced image
-	return PrimitiveSceneInfo->Proxy != nullptr &&
-		PrimitiveSceneInfo->Proxy->IsVisibleInRayTracing() &&
-		PrimitiveSceneInfo->Proxy->ShouldRenderInMainPass() &&
-		PrimitiveSceneInfo->Proxy->IsDrawnInGame();
+	return PrimitiveSceneInfo->bIsRayTracingRelevant &&
+		   PrimitiveSceneInfo->bIsVisibleInRayTracing &&
+		   PrimitiveSceneInfo->bDrawInGame &&
+		   PrimitiveSceneInfo->bShouldRenderInMainPass;
+#else
+	return false;
+#endif
 }
 
 void FScene::UpdateAllPrimitiveSceneInfos(FRDGBuilder& GraphBuilder, bool bAsyncCreateLPIs)
