Fix a bug of stationary local lights don't work on mobile.

diff --git a/Engine/Source/Runtime/Renderer/Private/SceneCore.cpp b/Engine/Source/Runtime/Renderer/Private/SceneCore.cpp
--- a/Engine/Source/Runtime/Renderer/Private/SceneCore.cpp
+++ b/Engine/Source/Runtime/Renderer/Private/SceneCore.cpp
@@ -234,7 +234,7 @@ FLightPrimitiveInteraction::FLightPrimitiveInteraction(
 		PrevPrimitiveLink = PrimitiveSceneInfo->Proxy->IsMeshShapeOftenMoving() ? &LightSceneInfo->DynamicInteractionOftenMovingPrimitiveList : &LightSceneInfo->DynamicInteractionStaticPrimitiveList;
 
 		// mobile movable spotlights / point lights
-		if (PrimitiveSceneInfo->Scene->GetShadingPath() == EShadingPath::Mobile && LightSceneInfo->Proxy->IsMovable())
+		if (PrimitiveSceneInfo->Scene->GetShadingPath() == EShadingPath::Mobile && LightSceneInfo->ShouldRenderLightViewIndependent())
 		{
 			const uint8 LightType = LightSceneInfo->Proxy->GetLightType();
 
