Fix a bug of fetching depth in mobile separate tranlucent pass with vulkan subpass.

diff --git a/Engine/Source/Runtime/Renderer/Private/MobileSeparateTranslucencyPass.cpp b/Engine/Source/Runtime/Renderer/Private/MobileSeparateTranslucencyPass.cpp
--- a/Engine/Source/Runtime/Renderer/Private/MobileSeparateTranslucencyPass.cpp
+++ b/Engine/Source/Runtime/Renderer/Private/MobileSeparateTranslucencyPass.cpp
@@ -7,6 +7,15 @@ MobileSeparateTranslucencyPass.cpp - Mobile specific separate translucency pass
 #include "MobileSeparateTranslucencyPass.h"
 #include "TranslucentRendering.h"
 #include "DynamicPrimitiveDrawing.h"
+#include "MobileBasePassRendering.h"
+#include "ScenePrivate.h"
+
+BEGIN_SHADER_PARAMETER_STRUCT(FMobileSeparateTranslucencyPassParameters, )
+	SHADER_PARAMETER_STRUCT_REF(FViewUniformShaderParameters, View)
+	SHADER_PARAMETER_RDG_UNIFORM_BUFFER(FMobileBasePassUniformParameters, MobileBasePass)
+	SHADER_PARAMETER_STRUCT_INCLUDE(FInstanceCullingDrawParams, InstanceCullingDrawParams)
+	RENDER_TARGET_BINDING_SLOTS()
+END_SHADER_PARAMETER_STRUCT()
 
 bool IsMobileSeparateTranslucencyActive(const FViewInfo* Views, int32 NumViews)
 {
@@ -25,22 +34,29 @@ bool IsMobileSeparateTranslucencyActive(const FViewInfo& View)
 	return View.ParallelMeshDrawCommandPasses[EMeshPass::TranslucencyAfterDOF].HasAnyDraw();
 }
 
-void AddMobileSeparateTranslucencyPass(FRDGBuilder& GraphBuilder, const FViewInfo& View, const FMobileSeparateTranslucencyInputs& Inputs)
+void AddMobileSeparateTranslucencyPass(FRDGBuilder& GraphBuilder, FScene* Scene, const FViewInfo& View, const FMobileSeparateTranslucencyInputs& Inputs)
 {
 	// GPUCULL_TODO: View.ParallelMeshDrawCommandPasses[EMeshPass::TranslucencyAfterDOF].BuildRenderingCommands(GraphBuilder, Scene->GPUScene);
 
-	FRenderTargetParameters* PassParameters = GraphBuilder.AllocParameters<FRenderTargetParameters>();
+	FMobileSeparateTranslucencyPassParameters* PassParameters = GraphBuilder.AllocParameters<FMobileSeparateTranslucencyPassParameters>();
 	PassParameters->RenderTargets[0] = FRenderTargetBinding(Inputs.SceneColor.Texture, ERenderTargetLoadAction::ELoad);
 	PassParameters->RenderTargets.DepthStencil = FDepthStencilBinding(Inputs.SceneDepth.Texture, ERenderTargetLoadAction::ELoad, ERenderTargetLoadAction::ELoad, FExclusiveDepthStencil::DepthRead_StencilRead);
+	PassParameters->RenderTargets.SubpassHint = ESubpassHint::DepthReadSubpass;
+	PassParameters->View = View.ViewUniformBuffer;
+	EMobileSceneTextureSetupMode SetupMode = EMobileSceneTextureSetupMode::SceneDepth | EMobileSceneTextureSetupMode::SceneDepthAux | EMobileSceneTextureSetupMode::CustomDepth;
+	PassParameters->MobileBasePass = CreateMobileBasePassUniformBuffer(GraphBuilder, View, EMobileBasePass::Translucent, SetupMode);
+
+	const_cast<FViewInfo&>(View).ParallelMeshDrawCommandPasses[EMeshPass::TranslucencyAfterDOF].BuildRenderingCommands(GraphBuilder, Scene->GPUScene, PassParameters->InstanceCullingDrawParams);
 
 	GraphBuilder.AddPass(
 		RDG_EVENT_NAME("SeparateTranslucency %dx%d", View.ViewRect.Width(), View.ViewRect.Height()),
 		PassParameters,
 		ERDGPassFlags::Raster,
-		[&View](FRHICommandList& RHICmdList)
+		[&View, PassParameters](FRHICommandList& RHICmdList)
 	{
 		// Set the view family's render target/viewport.
+		RHICmdList.NextSubpass();
 		RHICmdList.SetViewport(View.ViewRect.Min.X, View.ViewRect.Min.Y, 0.0f, View.ViewRect.Max.X, View.ViewRect.Max.Y, 1.0f);
-		View.ParallelMeshDrawCommandPasses[EMeshPass::TranslucencyAfterDOF].DispatchDraw(nullptr, RHICmdList);
+		View.ParallelMeshDrawCommandPasses[EMeshPass::TranslucencyAfterDOF].DispatchDraw(nullptr, RHICmdList, &PassParameters->InstanceCullingDrawParams);
 	});
 }
\ No newline at end of file
diff --git a/Engine/Source/Runtime/Renderer/Private/MobileSeparateTranslucencyPass.h b/Engine/Source/Runtime/Renderer/Private/MobileSeparateTranslucencyPass.h
--- a/Engine/Source/Runtime/Renderer/Private/MobileSeparateTranslucencyPass.h
+++ b/Engine/Source/Runtime/Renderer/Private/MobileSeparateTranslucencyPass.h
@@ -11,7 +11,7 @@ struct FMobileSeparateTranslucencyInputs
 	FScreenPassTexture SceneDepth;
 };
 
-void AddMobileSeparateTranslucencyPass(FRDGBuilder& GraphBuilder, const FViewInfo& View, const FMobileSeparateTranslucencyInputs& Inputs);
+void AddMobileSeparateTranslucencyPass(FRDGBuilder& GraphBuilder, FScene* Scene, const FViewInfo& View, const FMobileSeparateTranslucencyInputs& Inputs);
 
 // Returns whether separate translucency is enabled and there primitives to draw in the view
 bool IsMobileSeparateTranslucencyActive(const FViewInfo& View);
diff --git a/Engine/Source/Runtime/Renderer/Private/PostProcess/PostProcessing.cpp b/Engine/Source/Runtime/Renderer/Private/PostProcess/PostProcessing.cpp
--- a/Engine/Source/Runtime/Renderer/Private/PostProcess/PostProcessing.cpp
+++ b/Engine/Source/Runtime/Renderer/Private/PostProcess/PostProcessing.cpp
@@ -2186,7 +2186,7 @@ void AddMobilePostProcessingPasses(FRDGBuilder& GraphBuilder, FScene* Scene, con
 			SeparateTranslucencyInputs.SceneColor = SceneColor;
 			SeparateTranslucencyInputs.SceneDepth = SceneDepth;
 
-			AddMobileSeparateTranslucencyPass(GraphBuilder, View, SeparateTranslucencyInputs);
+			AddMobileSeparateTranslucencyPass(GraphBuilder, Scene, View, SeparateTranslucencyInputs);
 		}
 
 		AddPostProcessMaterialPass(BL_BeforeTonemapping, false);
