Disable directional albedo LUT when PT is enabled but the platform does not support it.

diff --git a/Engine/Source/Runtime/Renderer/Private/ShadingEnergyConservation.cpp b/Engine/Source/Runtime/Renderer/Private/ShadingEnergyConservation.cpp
--- a/Engine/Source/Runtime/Renderer/Private/ShadingEnergyConservation.cpp
+++ b/Engine/Source/Runtime/Renderer/Private/ShadingEnergyConservation.cpp
@@ -225,7 +225,7 @@ void Init(FRDGBuilder& GraphBuilder, FViewInfo& View)
 	// Build/bind table if energy conservation is enabled or if strata is enabled in order to have 
 	// the correct tables built & bound. Even if we are not using energy conservation, we want to 
 	// have access to directional albedo information for env. lighting for instance)
-	const bool bBindEnergyData = (View.ViewState != nullptr) && (bMaterialEnergyConservationEnabled || Strata::IsStrataEnabled() || View.Family->EngineShowFlags.PathTracing) && (bIsEnergyPreservationEnabled || bIsEnergyConservationEnabled);
+	const bool bBindEnergyData = (View.ViewState != nullptr) && (bMaterialEnergyConservationEnabled || Strata::IsStrataEnabled() || (View.Family->EngineShowFlags.PathTracing && RHI_RAYTRACING)) && (bIsEnergyPreservationEnabled || bIsEnergyConservationEnabled);
 	if (bBindEnergyData)
 	{
 		#if UE_EDITOR
