Mirror ShaderPrint packed structured from shader to C++ to ensure allocated size is identical.

diff --git a/Engine/Source/Runtime/Renderer/Private/ShaderPrint.cpp b/Engine/Source/Runtime/Renderer/Private/ShaderPrint.cpp
--- a/Engine/Source/Runtime/Renderer/Private/ShaderPrint.cpp
+++ b/Engine/Source/Runtime/Renderer/Private/ShaderPrint.cpp
@@ -80,11 +80,12 @@ namespace ShaderPrint
 	// Struct & Functions
 	
 	// Structure used by shader buffers to store values and symbols
-	struct ShaderPrintItem
+	struct FPackedShaderPrintItem
 	{
-		FVector2D ScreenPos;
-		int32 Value;
-		int32 Type;
+		uint32 ScreenPos16bits;
+		int32  Value;
+		uint32 TypeAndColor;
+		uint32 Metadata;
 	};
 
 	// Empty buffer for binding when ShaderPrint is disabled
@@ -96,7 +97,7 @@ namespace ShaderPrint
 			if (!Buffer.IsValid())
 			{
 				FRHICommandList* UnusedCmdList = new FRHICommandList(FRHIGPUMask::All());
-				GetPooledFreeBuffer(*UnusedCmdList, FRDGBufferDesc::CreateStructuredDesc(sizeof(ShaderPrintItem), 1), Buffer, TEXT("ShaderPrint.EmptyValueBuffer"));
+				GetPooledFreeBuffer(*UnusedCmdList, FRDGBufferDesc::CreateStructuredDesc(sizeof(FPackedShaderPrintItem), 1), Buffer, TEXT("ShaderPrint.EmptyValueBuffer"));
 				delete UnusedCmdList;
 				UnusedCmdList = nullptr;
 			}
@@ -246,7 +247,7 @@ namespace ShaderPrint
 		SHADER_USE_PARAMETER_STRUCT(FShaderInitValueBufferCS, FGlobalShader);
 
 		BEGIN_SHADER_PARAMETER_STRUCT(FParameters, )
-			SHADER_PARAMETER_RDG_BUFFER_UAV(RWStructuredBuffer<ShaderPrintItem>, RWValuesBuffer)
+			SHADER_PARAMETER_RDG_BUFFER_UAV(RWStructuredBuffer<FPackedShaderPrintItem>, RWValuesBuffer)
 		END_SHADER_PARAMETER_STRUCT()
 
 		static bool ShouldCompilePermutation(FGlobalShaderPermutationParameters const& Parameters)
@@ -266,8 +267,8 @@ namespace ShaderPrint
 
 		BEGIN_SHADER_PARAMETER_STRUCT(FParameters, )
 			SHADER_PARAMETER_STRUCT_REF(FShaderPrintCommonParameters, Common)
-			SHADER_PARAMETER_RDG_BUFFER_SRV(StructuredBuffer<ShaderPrintItem>, ValuesBuffer)
-			SHADER_PARAMETER_RDG_BUFFER_UAV(RWStructuredBuffer<ShaderPrintItem>, RWSymbolsBuffer)
+			SHADER_PARAMETER_RDG_BUFFER_SRV(StructuredBuffer<FPackedShaderPrintItem>, ValuesBuffer)
+			SHADER_PARAMETER_RDG_BUFFER_UAV(RWStructuredBuffer<FPackedShaderPrintItem>, RWSymbolsBuffer)
 			SHADER_PARAMETER_RDG_BUFFER_UAV(RWStructuredBuffer<uint>, RWIndirectDispatchArgsBuffer)
 		END_SHADER_PARAMETER_STRUCT()
 
@@ -311,8 +312,8 @@ namespace ShaderPrint
 		BEGIN_SHADER_PARAMETER_STRUCT(FParameters, )
 			SHADER_PARAMETER(uint32, FrameIndex)
 			SHADER_PARAMETER_STRUCT_REF(FShaderPrintCommonParameters, Common)
-			SHADER_PARAMETER_RDG_BUFFER_SRV(StructuredBuffer<ShaderPrintItem>, ValuesBuffer)
-			SHADER_PARAMETER_RDG_BUFFER_UAV(RWStructuredBuffer<ShaderPrintItem>, RWSymbolsBuffer)
+			SHADER_PARAMETER_RDG_BUFFER_SRV(StructuredBuffer<FPackedShaderPrintItem>, ValuesBuffer)
+			SHADER_PARAMETER_RDG_BUFFER_UAV(RWStructuredBuffer<FPackedShaderPrintItem>, RWSymbolsBuffer)
 			SHADER_PARAMETER_RDG_BUFFER_UAV(RWStructuredBuffer<uint>, RWStateBuffer)
 			RDG_BUFFER_ACCESS(IndirectDispatchArgsBuffer, ERHIAccess::IndirectArgs)
 		END_SHADER_PARAMETER_STRUCT()
@@ -334,7 +335,7 @@ namespace ShaderPrint
 
 		BEGIN_SHADER_PARAMETER_STRUCT(FParameters, )
 			SHADER_PARAMETER_STRUCT_REF(FShaderPrintCommonParameters, Common)
-			SHADER_PARAMETER_RDG_BUFFER_SRV(StructuredBuffer<ShaderPrintItem>, SymbolsBuffer)
+			SHADER_PARAMETER_RDG_BUFFER_SRV(StructuredBuffer<FPackedShaderPrintItem>, SymbolsBuffer)
 			SHADER_PARAMETER_RDG_BUFFER_UAV(RWStructuredBuffer<uint>, RWIndirectDrawArgsBuffer)
 		END_SHADER_PARAMETER_STRUCT()
 
@@ -361,7 +362,7 @@ namespace ShaderPrint
 			RENDER_TARGET_BINDING_SLOTS()
 			SHADER_PARAMETER_STRUCT_REF(FShaderPrintCommonParameters, Common)
 			SHADER_PARAMETER_TEXTURE(Texture2D, MiniFontTexture)
-			SHADER_PARAMETER_RDG_BUFFER_SRV(StructuredBuffer<ShaderPrintItem>, SymbolsBuffer)
+			SHADER_PARAMETER_RDG_BUFFER_SRV(StructuredBuffer<FPackedShaderPrintItem>, SymbolsBuffer)
 			RDG_BUFFER_ACCESS(IndirectDrawArgsBuffer, ERHIAccess::IndirectArgs)
 		END_SHADER_PARAMETER_STRUCT()
 
@@ -609,7 +610,7 @@ namespace ShaderPrint
 		{
 			// Initialize output buffer and store in the view info
 			// Values buffer contains Count + 1 elements. The first element is only used as a counter.
-			View.ShaderPrintData.ShaderPrintValueBuffer = GraphBuilder.CreateBuffer(FRDGBufferDesc::CreateStructuredDesc(sizeof(ShaderPrintItem), View.ShaderPrintData.MaxValueCount + 1), TEXT("ShaderPrint.ValueBuffer"));
+			View.ShaderPrintData.ShaderPrintValueBuffer = GraphBuilder.CreateBuffer(FRDGBufferDesc::CreateStructuredDesc(sizeof(FPackedShaderPrintItem), View.ShaderPrintData.MaxValueCount + 1), TEXT("ShaderPrint.ValueBuffer"));
 
 			// State buffer is retrieved from the view state, or created if it does not exist
 			View.ShaderPrintData.ShaderPrintStateBuffer = nullptr;
@@ -658,7 +659,7 @@ namespace ShaderPrint
 	
 		// Initialize graph managed resources
 		// Symbols buffer contains Count + 1 elements. The first element is only used as a counter.
-		FRDGBufferRef SymbolBuffer = GraphBuilder.CreateBuffer(FRDGBufferDesc::CreateStructuredDesc(sizeof(ShaderPrintItem), GetMaxSymbolCount() + 1), TEXT("ShaderPrint.SymbolBuffer"));
+		FRDGBufferRef SymbolBuffer = GraphBuilder.CreateBuffer(FRDGBufferDesc::CreateStructuredDesc(sizeof(FPackedShaderPrintItem), GetMaxSymbolCount() + 1), TEXT("ShaderPrint.SymbolBuffer"));
 		FRDGBufferRef IndirectDispatchArgsBuffer = GraphBuilder.CreateBuffer(FRDGBufferDesc::CreateIndirectDesc(4), TEXT("ShaderPrint.IndirectDispatchArgs"));
 		FRDGBufferRef IndirectDrawArgsBuffer = GraphBuilder.CreateBuffer(FRDGBufferDesc::CreateIndirectDesc(5), TEXT("ShaderPrint.IndirectDrawArgs"));
 
