Fix buffer type mismatches in ShaderPrint.

#rnx

diff --git a/Engine/Shaders/Private/ShaderPrintDraw.usf b/Engine/Shaders/Private/ShaderPrintDraw.usf
--- a/Engine/Shaders/Private/ShaderPrintDraw.usf
+++ b/Engine/Shaders/Private/ShaderPrintDraw.usf
@@ -16,7 +16,7 @@ StructuredBuffer<FPackedShaderPrintItem> SymbolsBuffer;
 
 RWStructuredBuffer<FPackedShaderPrintItem> RWValuesBuffer;
 RWStructuredBuffer<FPackedShaderPrintItem> RWSymbolsBuffer;
-RWBuffer<uint> RWStateBuffer;
+RWStructuredBuffer<uint> RWStateBuffer;
 
 RWBuffer<uint> RWIndirectDispatchArgsBuffer;
 RWBuffer<uint> RWIndirectDrawArgsBuffer;
diff --git a/Engine/Source/Runtime/Renderer/Private/ShaderPrint.cpp b/Engine/Source/Runtime/Renderer/Private/ShaderPrint.cpp
--- a/Engine/Source/Runtime/Renderer/Private/ShaderPrint.cpp
+++ b/Engine/Source/Runtime/Renderer/Private/ShaderPrint.cpp
@@ -269,7 +269,7 @@ namespace ShaderPrint
 			SHADER_PARAMETER_STRUCT_REF(FShaderPrintCommonParameters, Common)
 			SHADER_PARAMETER_RDG_BUFFER_SRV(StructuredBuffer<FPackedShaderPrintItem>, ValuesBuffer)
 			SHADER_PARAMETER_RDG_BUFFER_UAV(RWStructuredBuffer<FPackedShaderPrintItem>, RWSymbolsBuffer)
-			SHADER_PARAMETER_RDG_BUFFER_UAV(RWStructuredBuffer<uint>, RWIndirectDispatchArgsBuffer)
+			SHADER_PARAMETER_RDG_BUFFER_UAV(RWBuffer<uint>, RWIndirectDispatchArgsBuffer)
 		END_SHADER_PARAMETER_STRUCT()
 
 		static bool ShouldCompilePermutation(FGlobalShaderPermutationParameters const& Parameters)
@@ -336,7 +336,7 @@ namespace ShaderPrint
 		BEGIN_SHADER_PARAMETER_STRUCT(FParameters, )
 			SHADER_PARAMETER_STRUCT_REF(FShaderPrintCommonParameters, Common)
 			SHADER_PARAMETER_RDG_BUFFER_SRV(StructuredBuffer<FPackedShaderPrintItem>, SymbolsBuffer)
-			SHADER_PARAMETER_RDG_BUFFER_UAV(RWStructuredBuffer<uint>, RWIndirectDrawArgsBuffer)
+			SHADER_PARAMETER_RDG_BUFFER_UAV(RWBuffer<uint>, RWIndirectDrawArgsBuffer)
 		END_SHADER_PARAMETER_STRUCT()
 
 		static bool ShouldCompilePermutation(FGlobalShaderPermutationParameters const& Parameters)
@@ -679,7 +679,7 @@ namespace ShaderPrint
 			SHADER::FParameters* PassParameters = GraphBuilder.AllocParameters<SHADER::FParameters>();
 			PassParameters->Common = View.ShaderPrintData.UniformBuffer;
 			PassParameters->ValuesBuffer = ValueBuffer;
-			PassParameters->RWSymbolsBuffer = GraphBuilder.CreateUAV(SymbolBuffer, EPixelFormat::PF_R32_UINT);
+			PassParameters->RWSymbolsBuffer = GraphBuilder.CreateUAV(SymbolBuffer);
 			PassParameters->RWIndirectDispatchArgsBuffer = GraphBuilder.CreateUAV(IndirectDispatchArgsBuffer, EPixelFormat::PF_R32_UINT);
 
 			FComputeShaderUtils::AddPass(
@@ -698,8 +698,8 @@ namespace ShaderPrint
 			PassParameters->FrameIndex = View.Family ? View.Family->FrameNumber : 0u;
 			PassParameters->Common = View.ShaderPrintData.UniformBuffer;
 			PassParameters->ValuesBuffer = ValueBuffer;
-			PassParameters->RWSymbolsBuffer = GraphBuilder.CreateUAV(SymbolBuffer, EPixelFormat::PF_R32_UINT);
-			PassParameters->RWStateBuffer = GraphBuilder.CreateUAV(View.ShaderPrintData.ShaderPrintStateBuffer, EPixelFormat::PF_R32_UINT);
+			PassParameters->RWSymbolsBuffer = GraphBuilder.CreateUAV(SymbolBuffer);
+			PassParameters->RWStateBuffer = GraphBuilder.CreateUAV(View.ShaderPrintData.ShaderPrintStateBuffer);
 			PassParameters->IndirectDispatchArgsBuffer = IndirectDispatchArgsBuffer;
 
 			FComputeShaderUtils::AddPass(
@@ -719,7 +719,7 @@ namespace ShaderPrint
 			PassParameters->FrameIndex = View.Family ? View.Family->FrameNumber : 0u;
 			PassParameters->FrameThreshold = 300u;
 			PassParameters->Common = View.ShaderPrintData.UniformBuffer;
-			PassParameters->RWStateBuffer = GraphBuilder.CreateUAV(View.ShaderPrintData.ShaderPrintStateBuffer, EPixelFormat::PF_R32_UINT);
+			PassParameters->RWStateBuffer = GraphBuilder.CreateUAV(View.ShaderPrintData.ShaderPrintStateBuffer);
 
 			FComputeShaderUtils::AddPass(
 				GraphBuilder,
