Strata
- fix for shaders compilation complications with FXC
- fix for missing parameters that were not showing with DXC.

diff --git a/Engine/Shaders/Private/Strata/StrataVisualize.usf b/Engine/Shaders/Private/Strata/StrataVisualize.usf
--- a/Engine/Shaders/Private/Strata/StrataVisualize.usf
+++ b/Engine/Shaders/Private/Strata/StrataVisualize.usf
@@ -513,7 +513,7 @@ void MaterialDebugStrataTreeCS(uint3 DispatchThreadId : SV_DispatchThreadID)
 float4 GetBackgroundCheckboardColor(float2 UV)
 {
 	const float Count = 16.0f;
-	int2 UVi = int2(UV * Count);
+	uint2 UVi = uint2(UV * Count);
 	bool bIsBlack = (UVi.x % 2) > 0 ? true : false;
 	bIsBlack = (UVi.y % 2) > 0 ? !bIsBlack : bIsBlack;
 	const float3 Color = (bIsBlack ? 0.04 : 0.18);
@@ -691,6 +691,14 @@ void MaterialDebugStrataTreePS(
 
 		float2 Pos = float2(100.0f, 100.0f);
 		const float SphereSize = 200.0f;
+
+		// STRATA_TODO: read BSDF and Operators "on the fly".
+#if COMPILER_DXC == 1
+		LOOP
+#else
+		StrataTree.BSDFCount = min(StrataTree.BSDFCount, 2); 
+		STRATA_UNROLL_N(2)
+#endif
 		for (int i = 0; i < StrataTree.BSDFCount; i++)
 		{
 			GetMaterialOnSphere(
diff --git a/Engine/Source/Runtime/Renderer/Private/Strata/StrataVisualize.cpp b/Engine/Source/Runtime/Renderer/Private/Strata/StrataVisualize.cpp
--- a/Engine/Source/Runtime/Renderer/Private/Strata/StrataVisualize.cpp
+++ b/Engine/Source/Runtime/Renderer/Private/Strata/StrataVisualize.cpp
@@ -7,6 +7,7 @@
 #include "ShaderCompiler.h"
 #include "PixelShaderUtils.h"
 #include "BasePassRendering.h"
+#include "IndirectLightRendering.h"
 
 
 static TAutoConsoleVariable<int32> CVarStrataClassificationDebug(
@@ -152,7 +153,6 @@ class FMaterialDebugStrataTreeCS : public FGlobalShader
 
 		// Stay debug and skip optimizations to reduce compilation time on this long shader.
 		OutEnvironment.CompilerFlags.Add(CFLAG_Debug);
-		OutEnvironment.CompilerFlags.Add(CFLAG_ForceDXC);
 		OutEnvironment.SetDefine(TEXT("SHADER_DEBUGSTRATATREE_TEST_CS"), 1);
 	}
 };
@@ -171,6 +171,7 @@ class FMaterialDebugStrataTreePS : public FGlobalShader
 		SHADER_PARAMETER_STRUCT_REF(FReflectionUniformParameters, ReflectionStruct)
 		SHADER_PARAMETER_STRUCT_REF(FReflectionCaptureShaderData, ReflectionCapture)
 		SHADER_PARAMETER_RDG_UNIFORM_BUFFER(FForwardLightData, ForwardLightData)
+		SHADER_PARAMETER_STRUCT_INCLUDE(FSkyDiffuseLightingParameters, SkyDiffuseLighting)
 		RENDER_TARGET_BINDING_SLOTS()
 	END_SHADER_PARAMETER_STRUCT()
 
@@ -274,6 +275,10 @@ static void AddVisualizeMaterialPasses(FRDGBuilder& GraphBuilder, const FViewInf
 					PassParameters->ForwardLightData = View.ForwardLightingResources.ForwardLightUniformBuffer;
 					PassParameters->RenderTargets[0] = FRenderTargetBinding(SceneColorTexture, ERenderTargetLoadAction::ELoad);
 
+					const float DynamicBentNormalAO = 0.0f;
+					FSkyLightSceneProxy* NullSkyLight = nullptr;
+					PassParameters->SkyDiffuseLighting = GetSkyDiffuseLightingParameters(NullSkyLight, DynamicBentNormalAO);
+
 					FMaterialDebugStrataTreePS::FPermutationDomain PermutationVector;
 					TShaderMapRef<FMaterialDebugStrataTreePS> PixelShader(View.ShaderMap, PermutationVector);
 
