Path Tracer: Fix constant invalidation in MRQ renders when camera DOF is animated

diff --git a/Engine/Source/Runtime/Renderer/Private/PathTracing.cpp b/Engine/Source/Runtime/Renderer/Private/PathTracing.cpp
--- a/Engine/Source/Runtime/Renderer/Private/PathTracing.cpp
+++ b/Engine/Source/Runtime/Renderer/Private/PathTracing.cpp
@@ -328,8 +328,6 @@ struct FPathTracingConfig
 			PathTracingData.MaxPathIntensity != Other.PathTracingData.MaxPathIntensity ||
 			PathTracingData.FilterWidth != Other.PathTracingData.FilterWidth ||
 			PathTracingData.AbsorptionScale != Other.PathTracingData.AbsorptionScale ||
-			PathTracingData.CameraFocusDistance != Other.PathTracingData.CameraFocusDistance ||
-			PathTracingData.CameraLensRadius != Other.PathTracingData.CameraLensRadius ||
 			PathTracingData.EnableAtmosphere != Other.PathTracingData.EnableAtmosphere ||
 			PathTracingData.EnableFog != Other.PathTracingData.EnableFog ||
 			PathTracingData.MaxRaymarchSteps != Other.PathTracingData.MaxRaymarchSteps ||
@@ -342,6 +340,13 @@ struct FPathTracingConfig
 			LockedSamplingPattern != Other.LockedSamplingPattern ||
 			UseMultiGPU != Other.UseMultiGPU;
 	}
+
+	bool IsDOFDifferent(const FPathTracingConfig& Other) const
+	{
+		return PathTracingData.CameraFocusDistance != Other.PathTracingData.CameraFocusDistance ||
+			   PathTracingData.CameraLensRadius != Other.PathTracingData.CameraLensRadius;
+
+	}
 };
 
 struct FAtmosphereConfig
@@ -1721,7 +1726,13 @@ void FDeferredShadingSceneRenderer::RenderPathTracing(
 		Scene->PathTracingSkylightPdf.SafeRelease();
 		// reset last color here as well in case we don't reach PrepareSkyLightTexture
 		Scene->PathTracingSkylightColor = Scene->SkyLight->GetEffectiveLightColor();
-		View.ViewState->PathTracingInvalidate();
+		if (!View.bIsOfflineRender)
+		{
+			// reset accumulation, unless this is an offline render, in which case it is ok for the color to evolve
+			// across temporal samples
+			View.ViewState->PathTracingInvalidate();
+		}
+		
 	}
 
 	// prepare atmosphere optical depth lookup texture (if needed)
@@ -1776,7 +1787,11 @@ void FDeferredShadingSceneRenderer::RenderPathTracing(
 	// If the scene has changed in some way (camera move, object movement, etc ...)
 	// we must invalidate the ViewState to start over from scratch
 	// NOTE: only check things like hair position changes for interactive viewports, for offline renders we don't want any chance of mid-render invalidation
-	if (FirstTime || Config.IsDifferent(PathTracingState->LastConfig) || (!View.bIsOfflineRender && HairStrands::HasPositionsChanged(GraphBuilder, View)))
+	// NOTE: same for DOF changes, these parameters could be animated which should not automatically invalidate a render in progress
+	if (FirstTime ||
+		Config.IsDifferent(PathTracingState->LastConfig) ||
+		(!View.bIsOfflineRender && Config.IsDOFDifferent(PathTracingState->LastConfig)) ||
+		(!View.bIsOfflineRender && HairStrands::HasPositionsChanged(GraphBuilder, View)))
 	{
 		// remember the options we used for next time
 		PathTracingState->LastConfig = Config;
