Fix viewports being handled incorrectly if there are multiple ISR pairs.

InstancedStereoWidth is already a view property, but is currently simply set to the widest pair.
This CL properly computes InstancedStereoWidth for each pair individually.

diff --git a/Engine/Source/Runtime/Renderer/Private/SceneRendering.cpp b/Engine/Source/Runtime/Renderer/Private/SceneRendering.cpp
--- a/Engine/Source/Runtime/Renderer/Private/SceneRendering.cpp
+++ b/Engine/Source/Runtime/Renderer/Private/SceneRendering.cpp
@@ -2481,7 +2481,6 @@ FViewFamilyInfo::~FViewFamilyInfo()
 FSceneRenderer::FSceneRenderer(TArrayView<const FSceneViewFamily*> InViewFamilies, FHitProxyConsumer* HitProxyConsumer)
 :	Scene(CheckPointer(InViewFamilies[0]->Scene)->GetRenderScene())
 ,	ActiveViewFamily(nullptr)
-,	InstancedStereoWidth(0)
 ,	GPUSceneDynamicContext(CheckPointer(Scene)->GPUScene)
 {
 	check(Scene != NULL);
@@ -3103,7 +3102,7 @@ void FSceneRenderer::ComputeFamilySize()
 	float MaxFamilyX = 0;
 	float MaxFamilyY = 0;
 
-	for (const FViewInfo& View : Views)
+	for (FViewInfo& View : Views)
 	{
 		float FinalViewMaxX = (float)View.ViewRect.Max.X;
 		float FinalViewMaxY = (float)View.ViewRect.Max.Y;
@@ -3132,12 +3131,8 @@ void FSceneRenderer::ComputeFamilySize()
 		MaxFamilyX = FMath::Max(MaxFamilyX, FinalViewMaxX);
 		MaxFamilyY = FMath::Max(MaxFamilyY, FinalViewMaxY);
 
-		InstancedStereoWidth = FPlatformMath::Max(InstancedStereoWidth, static_cast<uint32>(View.ViewRect.Max.X));
-	}
-
-	for (FViewInfo& View : Views)
-	{
-		View.InstancedStereoWidth = InstancedStereoWidth;
+		const FViewInfo* InstancedView = View.GetInstancedView();
+		View.InstancedStereoWidth = InstancedView ? InstancedView->ViewRect.Max.X : View.ViewRect.Max.X;
 	}
 
 	// We render to the actual position of the viewports so with black borders we need the max.
@@ -5009,7 +5004,7 @@ void FSceneRenderer::SetStereoViewport(FRHICommandList& RHICmdList, const FViewI
 		}
 		else
 		{
-			RHICmdList.SetViewport(View.ViewRect.Min.X * ViewportScale, View.ViewRect.Min.Y * ViewportScale, 0.0f, InstancedStereoWidth * ViewportScale, View.ViewRect.Max.Y * ViewportScale, 1.0f);
+			RHICmdList.SetViewport(View.ViewRect.Min.X * ViewportScale, View.ViewRect.Min.Y * ViewportScale, 0.0f, View.InstancedStereoWidth * ViewportScale, View.ViewRect.Max.Y * ViewportScale, 1.0f);
 		}
 	}
 	else
diff --git a/Engine/Source/Runtime/Renderer/Private/SceneRendering.h b/Engine/Source/Runtime/Renderer/Private/SceneRendering.h
--- a/Engine/Source/Runtime/Renderer/Private/SceneRendering.h
+++ b/Engine/Source/Runtime/Renderer/Private/SceneRendering.h
@@ -1492,6 +1492,11 @@ public:
 	 */
 	int32 GPUSceneViewId;
 
+	/**
+	 * The width in pixels of the ISR pair being rendered. This may be different than FamilySizeX if
+	 * we're using adaptive resolution stereo rendering. In that case, FamilySizeX represents the maximum size of
+	 * the family to ensure the backing render targets don't change between frames as the view size varies.
+	 */
 	uint32 InstancedStereoWidth = 0;
 
 #if WITH_EDITOR
@@ -1896,13 +1901,6 @@ public:
 	/** Feature level being rendered */
 	ERHIFeatureLevel::Type FeatureLevel;
 	EShaderPlatform ShaderPlatform;
-	
-	/** 
-	 * The width in pixels of the stereo view family being rendered. This may be different than FamilySizeX if
-	 * we're using adaptive resolution stereo rendering. In that case, FamilySizeX represents the maximum size of 
-	 * the family to ensure the backing render targets don't change between frames as the view size varies.
-	 */
-	uint32 InstancedStereoWidth;
 
 public:
 
diff --git a/Engine/Source/Runtime/Renderer/Private/TranslucentRendering.cpp b/Engine/Source/Runtime/Renderer/Private/TranslucentRendering.cpp
--- a/Engine/Source/Runtime/Renderer/Private/TranslucentRendering.cpp
+++ b/Engine/Source/Runtime/Renderer/Private/TranslucentRendering.cpp
@@ -811,12 +811,12 @@ bool FSceneRenderer::ShouldRenderTranslucency(ETranslucencyPass::Type Translucen
 	return false;
 }
 
-FScreenPassTextureViewport FSeparateTranslucencyDimensions::GetInstancedStereoViewport(const FViewInfo& View, float InstancedStereoWidth) const
+FScreenPassTextureViewport FSeparateTranslucencyDimensions::GetInstancedStereoViewport(const FViewInfo& View) const
 {
 	FIntRect ViewRect = View.ViewRect;
 	if (View.IsInstancedStereoPass() && !View.bIsMultiViewEnabled)
 	{
-		ViewRect.Max.X = ViewRect.Min.X + InstancedStereoWidth;
+		ViewRect.Max.X = ViewRect.Min.X + View.InstancedStereoWidth;
 	}
 	ViewRect = GetScaledRect(ViewRect, Scale);
 	return FScreenPassTextureViewport(Extent, ViewRect);
@@ -1399,7 +1399,7 @@ void FDeferredShadingSceneRenderer::RenderTranslucencyInner(
 
 			FIntRect ScaledViewRect = GetScaledRect(View.ViewRect, SeparateTranslucencyDimensions.Scale);
 
-			const FScreenPassTextureViewport SeparateTranslucencyViewport = SeparateTranslucencyDimensions.GetInstancedStereoViewport(View, InstancedStereoWidth);
+			const FScreenPassTextureViewport SeparateTranslucencyViewport = SeparateTranslucencyDimensions.GetInstancedStereoViewport(View);
 			const bool bCompositeBackToSceneColor = IsMainTranslucencyPass(TranslucencyPass) || EnumHasAnyFlags(TranslucencyView, ETranslucencyView::UnderWater);
 			const bool bLumenGIEnabled = GetViewPipelineState(View).DiffuseIndirectMethod == EDiffuseIndirectMethod::Lumen;
 
@@ -1603,7 +1603,7 @@ void FDeferredShadingSceneRenderer::RenderTranslucency(
 				continue;
 			}
 
-			const FScreenPassTextureViewport SeparateTranslucencyViewport = SeparateTranslucencyDimensions.GetInstancedStereoViewport(View, InstancedStereoWidth);
+			const FScreenPassTextureViewport SeparateTranslucencyViewport = SeparateTranslucencyDimensions.GetInstancedStereoViewport(View);
 			AddDownsampleDepthPass(
 				GraphBuilder, View,
 				FScreenPassTexture(SceneTextures.Depth.Resolve, View.ViewRect),
@@ -1664,7 +1664,7 @@ void FDeferredShadingSceneRenderer::RenderTranslucency(
 			SharedDepthTexture.Target != SceneTextures.Depth.Target);
 		if (bUpscaleResponsiveAA)
 		{
-			const FScreenPassTextureViewport SeparateTranslucencyViewport = SeparateTranslucencyDimensions.GetInstancedStereoViewport(View, InstancedStereoWidth);
+			const FScreenPassTextureViewport SeparateTranslucencyViewport = SeparateTranslucencyDimensions.GetInstancedStereoViewport(View);
 			AddUpsampleResponsiveAAPass(
 				GraphBuilder,
 				View,
diff --git a/Engine/Source/Runtime/Renderer/Private/TranslucentRendering.h b/Engine/Source/Runtime/Renderer/Private/TranslucentRendering.h
--- a/Engine/Source/Runtime/Renderer/Private/TranslucentRendering.h
+++ b/Engine/Source/Runtime/Renderer/Private/TranslucentRendering.h
@@ -17,7 +17,7 @@ struct FSeparateTranslucencyDimensions
 		return FScreenPassTextureViewport(Extent, GetScaledRect(ViewRect, Scale));
 	}
 
-	FScreenPassTextureViewport GetInstancedStereoViewport(const FViewInfo& View, float InstancedStereoWidth) const;
+	FScreenPassTextureViewport GetInstancedStereoViewport(const FViewInfo& View) const;
 
 	// Extent of the separate translucency targets, if downsampled.
 	FIntPoint Extent = FIntPoint::ZeroValue;
