Use MinDrawDistance in ray tracing cullling:
* Objects that are in front of camera will be culled if they are closer than MinDrawDistance.
* Controlled by r.RayTracing.Culling.UseMinDrawDistance (disabled by default for now).

diff --git a/Engine/Source/Runtime/Renderer/Private/RayTracing/RayTracingInstanceCulling.cpp b/Engine/Source/Runtime/Renderer/Private/RayTracing/RayTracingInstanceCulling.cpp
--- a/Engine/Source/Runtime/Renderer/Private/RayTracing/RayTracingInstanceCulling.cpp
+++ b/Engine/Source/Runtime/Renderer/Private/RayTracing/RayTracingInstanceCulling.cpp
@@ -33,6 +33,12 @@ static TAutoConsoleVariable<float> CVarRayTracingCullingAngle(
 	TEXT("Do camera culling for objects behind the camera with a projected angle smaller than this threshold in ray tracing effects (default = 5 degrees )"),
 	ECVF_RenderThreadSafe);
 
+static TAutoConsoleVariable<int32> CVarRayTracingCullingUseMinDrawDistance(
+	TEXT("r.RayTracing.Culling.UseMinDrawDistance"),
+	0,
+	TEXT("Use min draw distance for culling"),
+	ECVF_RenderThreadSafe);
+
 static TAutoConsoleVariable<int32> CVarRayTracingCullingGroupIds(
 	TEXT("r.RayTracing.Culling.UseGroupIds"),
 	0,
@@ -68,6 +74,7 @@ void FRayTracingCullingParameters::Init(FViewInfo& View)
 	bCullByRadiusOrDistance = CullInRayTracing == 3;
 	bIsRayTracingFarField = Lumen::UseFarField(*View.Family);
 	bCullUsingGroupIds = CVarRayTracingCullingGroupIds.GetValueOnRenderThread() != 0;
+	bCullMinDrawDistance = CVarRayTracingCullingUseMinDrawDistance.GetValueOnRenderThread() != 0;
 }
 
 namespace RayTracing
@@ -96,7 +103,7 @@ bool CullPrimitiveByFlags(const FRayTracingCullingParameters& CullingParameters,
 
 // Tests if the given primitive should be culled, given the pre-calculated inputs for the primitive, returns true if the primitive SHOULD be culled
 template<bool bCullByRadiusOrDistance>
-bool CullBounds(const FRayTracingCullingParameters& CullingParameters, const FBoxSphereBounds& RESTRICT ObjectBounds, bool bIsFarFieldPrimitive)
+bool CullBounds(const FRayTracingCullingParameters& CullingParameters, const FBoxSphereBounds& RESTRICT ObjectBounds, float MinDrawDistance, bool bIsFarFieldPrimitive)
 {
 	const float ObjectRadius = ObjectBounds.SphereRadius;
 	const FVector ObjectCenter = ObjectBounds.Origin + 0.5 * ObjectBounds.BoxExtent;
@@ -112,8 +119,13 @@ bool CullBounds(const FRayTracingCullingParameters& CullingParameters, const FBo
 	else
 	{
 		const bool bIsFarEnoughToCull = CameraToObjectCenterLengthSq > FMath::Square(CullingParameters.CullingRadius + ObjectRadius);
+		const bool bIsNearEnoughToCull = CameraToObjectCenterLengthSq < FMath::Square(MinDrawDistance);
 
-		if (bCullByRadiusOrDistance && bIsFarEnoughToCull)
+		if (CullingParameters.bCullMinDrawDistance && bIsNearEnoughToCull)
+		{
+			return true;
+		}
+		else if (bCullByRadiusOrDistance && bIsFarEnoughToCull)
 		{
 			return true;
 		}
@@ -135,21 +147,21 @@ bool CullBounds(const FRayTracingCullingParameters& CullingParameters, const FBo
 	return false;
 }
 
-bool ShouldCullBounds(const FRayTracingCullingParameters& CullingParameters, const FBoxSphereBounds& RESTRICT ObjectBounds, bool bIsFarFieldPrimitive)
+bool ShouldCullBounds(const FRayTracingCullingParameters& CullingParameters, const FBoxSphereBounds& RESTRICT ObjectBounds, float MinDrawDistance, bool bIsFarFieldPrimitive)
 {
 	if (CullingParameters.CullInRayTracing > 0)
 	{
-		const bool bConsiderCulling = CullingParameters.bCullAllObjects || ShouldConsiderCulling(CullingParameters, ObjectBounds);
+		const bool bConsiderCulling = CullingParameters.bCullAllObjects || ShouldConsiderCulling(CullingParameters, ObjectBounds, MinDrawDistance);
 
 		if (bConsiderCulling)
 		{
 			if (CullingParameters.bCullByRadiusOrDistance)
 			{
-				return CullBounds<true>(CullingParameters, ObjectBounds, bIsFarFieldPrimitive);
+				return CullBounds<true>(CullingParameters, ObjectBounds, MinDrawDistance, bIsFarFieldPrimitive);
 			}
 			else
 			{
-				return CullBounds<false>(CullingParameters, ObjectBounds, bIsFarFieldPrimitive);
+				return CullBounds<false>(CullingParameters, ObjectBounds, MinDrawDistance, bIsFarFieldPrimitive);
 			}
 		}
 	}
@@ -218,13 +230,14 @@ void FRayTracingCullPrimitiveInstancesClosure::operator()() const
 	const bool bUseGroupBounds = CullingParameters->bCullUsingGroupIds && GroupId.IsValid();
 	const FBoxSphereBounds* const GroupBounds = (bUseGroupBounds) ? &Scene->PrimitiveRayTracingGroups.GetByElementId(GroupId).Value.Bounds : nullptr;
 	const FBoxSphereBounds& PrimitiveBounds = (GroupBounds) ? *GroupBounds : Scene->PrimitiveBounds[PrimitiveIndex].BoxSphereBounds;
+	const float MinDrawDistance = GroupBounds ? Scene->PrimitiveRayTracingGroups.GetByElementId(GroupId).Value.MinDrawDistance : Scene->PrimitiveBounds[PrimitiveIndex].MinDrawDistance;
 
 	if (!RayTracing::ShouldSkipPerInstanceCullingForPrimitive(*CullingParameters, PrimitiveBounds, SceneInfo->CachedRayTracingInstanceWorldBounds[SceneInfo->SmallestRayTracingInstanceWorldBoundsIndex], bIsFarFieldPrimitive))
 	{
 		for (int32 InstanceIndex = 0; InstanceIndex < SceneInfo->CachedRayTracingInstanceWorldBounds.Num(); InstanceIndex++)
 		{
 			const FBoxSphereBounds& InstanceBounds = (GroupBounds) ? *GroupBounds : SceneInfo->CachedRayTracingInstanceWorldBounds[InstanceIndex];
-			if (RayTracing::ShouldCullBounds(*CullingParameters, InstanceBounds, bIsFarFieldPrimitive))
+			if (RayTracing::ShouldCullBounds(*CullingParameters, InstanceBounds, MinDrawDistance, bIsFarFieldPrimitive))
 			{
 				OutInstanceActivationMask[InstanceIndex / 32] &= ~(1 << (InstanceIndex % 32));
 			}
diff --git a/Engine/Source/Runtime/Renderer/Private/RayTracing/RayTracingInstanceCulling.h b/Engine/Source/Runtime/Renderer/Private/RayTracing/RayTracingInstanceCulling.h
--- a/Engine/Source/Runtime/Renderer/Private/RayTracing/RayTracingInstanceCulling.h
+++ b/Engine/Source/Runtime/Renderer/Private/RayTracing/RayTracingInstanceCulling.h
@@ -15,17 +15,19 @@ namespace RayTracing
 {
 
 // Tests if a primitive with the given inputs should be considered for culling.  Does NOT test configuration values.  The assumption is that the config values have already been tested.
-FORCEINLINE bool ShouldConsiderCulling(const FRayTracingCullingParameters& CullingParameters, const FBoxSphereBounds& ObjectBounds)
+FORCEINLINE bool ShouldConsiderCulling(const FRayTracingCullingParameters& CullingParameters, const FBoxSphereBounds& ObjectBounds, float MinDrawDistance)
 {
 	FVector CameraToObjectCenter = ObjectBounds.Origin - CullingParameters.ViewOrigin;
-	return FVector::DotProduct(CullingParameters.ViewDirection, CameraToObjectCenter) < -ObjectBounds.SphereRadius;
+	bool bBehindCamera = FVector::DotProduct(CullingParameters.ViewDirection, CameraToObjectCenter) < -ObjectBounds.SphereRadius;
+	bool bNearEnoughToCull = CullingParameters.bCullMinDrawDistance && (FVector::DistSquared(ObjectBounds.Origin, CullingParameters.ViewOrigin) < FMath::Square(MinDrawDistance));
+	return bBehindCamera || bNearEnoughToCull;
 }
 
 // Returns true if the primitive should be culled out due to it's ray tracing flags
 bool CullPrimitiveByFlags(const FRayTracingCullingParameters& CullingParameters, const FScene* RESTRICT Scene, int32 PrimitiveIndex);
 
 // Completely test if the bounds should be culled for ray tracing. This includes all configuration values.
-bool ShouldCullBounds(const FRayTracingCullingParameters& CullingParameters, const FBoxSphereBounds& ObjectBounds, bool bIsFarFieldPrimitive);
+bool ShouldCullBounds(const FRayTracingCullingParameters& CullingParameters, const FBoxSphereBounds& ObjectBounds, float MinDrawDistance, bool bIsFarFieldPrimitive);
 
 bool ShouldSkipPerInstanceCullingForPrimitive(const FRayTracingCullingParameters& CullingParameters, FBoxSphereBounds ObjectBounds, FBoxSphereBounds SmallestInstanceBounds, bool bIsFarFieldPrimitive);
 }
diff --git a/Engine/Source/Runtime/Renderer/Private/RendererScene.cpp b/Engine/Source/Runtime/Renderer/Private/RendererScene.cpp
--- a/Engine/Source/Runtime/Renderer/Private/RendererScene.cpp
+++ b/Engine/Source/Runtime/Renderer/Private/RendererScene.cpp
@@ -4275,10 +4275,12 @@ void FScene::UpdateRayTracingGroupBounds_AddPrimitives(const Experimental::TRobi
 			if (bInMap)
 			{
 				Group->Bounds = Group->Bounds + PrimitiveSceneInfo->Proxy->GetBounds();
+				Group->MinDrawDistance = FMath::Max(Group->MinDrawDistance, PrimitiveSceneInfo->Proxy->GetMinDrawDistance());
 			}
 			else
 			{
 				Group->Bounds = PrimitiveSceneInfo->Proxy->GetBounds();
+				Group->MinDrawDistance = PrimitiveSceneInfo->Proxy->GetMinDrawDistance();
 			}
 			Group->Primitives.Add(PrimitiveSceneInfo);
 		}
diff --git a/Engine/Source/Runtime/Renderer/Private/ScenePrivate.h b/Engine/Source/Runtime/Renderer/Private/ScenePrivate.h
--- a/Engine/Source/Runtime/Renderer/Private/ScenePrivate.h
+++ b/Engine/Source/Runtime/Renderer/Private/ScenePrivate.h
@@ -2795,6 +2795,7 @@ public:
 	struct FRayTracingCullingGroup
 	{
 		FBoxSphereBounds Bounds;
+		float MinDrawDistance = 0.0f;
 		TArray<FPrimitiveSceneInfo*> Primitives;
 	};
 	Experimental::TRobinHoodHashMap<int32, FRayTracingCullingGroup> PrimitiveRayTracingGroups;
diff --git a/Engine/Source/Runtime/Renderer/Private/SceneRendering.h b/Engine/Source/Runtime/Renderer/Private/SceneRendering.h
--- a/Engine/Source/Runtime/Renderer/Private/SceneRendering.h
+++ b/Engine/Source/Runtime/Renderer/Private/SceneRendering.h
@@ -1124,6 +1124,7 @@ struct FRayTracingCullingParameters
 	bool bCullByRadiusOrDistance;
 	bool bIsRayTracingFarField;
 	bool bCullUsingGroupIds;
+	bool bCullMinDrawDistance;
 
 	void Init(FViewInfo& View);
 };
diff --git a/Engine/Source/Runtime/Renderer/Private/SceneVisibility.cpp b/Engine/Source/Runtime/Renderer/Private/SceneVisibility.cpp
--- a/Engine/Source/Runtime/Renderer/Private/SceneVisibility.cpp
+++ b/Engine/Source/Runtime/Renderer/Private/SceneVisibility.cpp
@@ -628,12 +628,13 @@ FORCEINLINE bool ShouldCullForRayTracing(const FScene* RESTRICT Scene, FViewInfo
 	if (RayTracingCullingParameters.bCullUsingGroupIds && GroupId.IsValid())
 	{
 		const FBoxSphereBounds& GroupBounds = Scene->PrimitiveRayTracingGroups.GetByElementId(GroupId).Value.Bounds;
-		return RayTracing::ShouldCullBounds(RayTracingCullingParameters, GroupBounds, bIsFarFieldPrimitive);
+		const float GroupMinDrawDistance = Scene->PrimitiveRayTracingGroups.GetByElementId(GroupId).Value.MinDrawDistance;
+		return RayTracing::ShouldCullBounds(RayTracingCullingParameters, GroupBounds, GroupMinDrawDistance, bIsFarFieldPrimitive);
 	}
 	else
 	{
 		const FPrimitiveBounds& RESTRICT Bounds = Scene->PrimitiveBounds[PrimitiveIndex];
-		return RayTracing::ShouldCullBounds(RayTracingCullingParameters, Bounds.BoxSphereBounds, bIsFarFieldPrimitive);
+		return RayTracing::ShouldCullBounds(RayTracingCullingParameters, Bounds.BoxSphereBounds, Bounds.MinDrawDistance, bIsFarFieldPrimitive);
 	}
 };
 #endif //RHI_RAYTRACING
