Strata visualise now account for resolution scaling (when rendering and sampling the strata material).

diff --git a/Engine/Shaders/Private/BasePassPixelShader.usf b/Engine/Shaders/Private/BasePassPixelShader.usf
--- a/Engine/Shaders/Private/BasePassPixelShader.usf
+++ b/Engine/Shaders/Private/BasePassPixelShader.usf
@@ -1733,7 +1733,7 @@ void FPixelShaderInOut_MainPS(
 		FStrataIntegrationSettings Settings = InitStrataIntegrationSettings(false /*bForceFullyRough*/, StrataStruct.bRoughDiffuse, StrataStruct.PeelLayersAboveDepth);
 
 	#if STRATA_ADVANCED_DEBUG_ENABLED
-		Settings.SliceStoringDebugStrataTree = all(View.CursorPosition == PixelPos) ? StrataStruct.SliceStoringDebugStrataTree : Settings.SliceStoringDebugStrataTree;
+		Settings.SliceStoringDebugStrataTree = all(uint2(float2(View.CursorPosition) * View.ViewResolutionFraction) == PixelPos) ? StrataStruct.SliceStoringDebugStrataTree : Settings.SliceStoringDebugStrataTree;
 	#endif
 
 		// Generate the strata material data to write out
diff --git a/Engine/Shaders/Private/Strata/StrataVisualize.usf b/Engine/Shaders/Private/Strata/StrataVisualize.usf
--- a/Engine/Shaders/Private/Strata/StrataVisualize.usf
+++ b/Engine/Shaders/Private/Strata/StrataVisualize.usf
@@ -396,7 +396,7 @@ void PrintPixel(uint2 InCoord, float3 InWorldPosition, float3 V)
 [numthreads(1, 1, 1)]
 void MaterialPrintInfoCS(uint3 DispatchThreadId : SV_DispatchThreadID)
 {
-	const uint2 PixelPos = View.CursorPosition;
+	const uint2 PixelPos = (float2(View.CursorPosition) * View.ViewResolutionFraction);
 	const float2 BufferUV = SvPositionToBufferUV(float4(PixelPos + 0.5f, 0, 0));
 	const float2 ScreenPosition = SvPositionToScreenPosition(float4(PixelPos, 0.5f, 1)).xy;
 	const float DeviceZ = SampleDeviceZFromSceneTextures(BufferUV);
@@ -460,7 +460,7 @@ void VisualizeMaterialPS(
 [numthreads(1, 1, 1)]
 void MaterialDebugStrataTreeCS(uint3 DispatchThreadId : SV_DispatchThreadID)
 {
-	const uint2 PixelPos = View.CursorPosition;
+	const uint2 PixelPos = (float2(View.CursorPosition) * View.ViewResolutionFraction);
 
 	FStrataAddressing StrataAddressing = GetStrataPixelDataByteOffset(PixelPos, uint2(View.BufferSizeAndInvSize.xy), Strata.MaxBytesPerPixel);
 	FStrataPixelHeader Header = UnpackStrataHeaderIn(Strata.MaterialTextureArray, StrataAddressing, Strata.TopLayerTexture);
@@ -683,7 +683,7 @@ void MaterialDebugStrataTreePS(
 {
 	OutColor = float4(0.0f, 0.0f, 0.0f, 1.0f);
 
-	FStrataAddressing StrataAddressing = GetStrataPixelDataByteOffset(View.CursorPosition, uint2(View.BufferSizeAndInvSize.xy), Strata.MaxBytesPerPixel);
+	FStrataAddressing StrataAddressing = GetStrataPixelDataByteOffset((float2(View.CursorPosition) * View.ViewResolutionFraction), uint2(View.BufferSizeAndInvSize.xy), Strata.MaxBytesPerPixel);
 	FStrataPixelHeader Header = UnpackStrataHeaderIn(Strata.MaterialTextureArray, StrataAddressing, Strata.TopLayerTexture);
 	if (Header.BSDFCount > 0)
 	{
diff --git a/Engine/Source/Runtime/Engine/Private/SceneView.cpp b/Engine/Source/Runtime/Engine/Private/SceneView.cpp
--- a/Engine/Source/Runtime/Engine/Private/SceneView.cpp
+++ b/Engine/Source/Runtime/Engine/Private/SceneView.cpp
@@ -2424,6 +2424,8 @@ void FSceneView::SetupViewRectUniformBufferParameters(FViewUniformShaderParamete
 
 	ViewUniformShaderParameters.ScreenToViewSpace.Z = -((ViewUniformShaderParameters.ViewRectMin.X * ViewUniformShaderParameters.ViewSizeAndInvSize.Z * 2 * InvFovFixX) + InvFovFixX);
 	ViewUniformShaderParameters.ScreenToViewSpace.W = (ViewUniformShaderParameters.ViewRectMin.Y * ViewUniformShaderParameters.ViewSizeAndInvSize.W * 2 * InvFovFixY) + InvFovFixY;
+
+	ViewUniformShaderParameters.ViewResolutionFraction = EffectiveViewRect.Width() / (float)UnscaledViewRect.Width();
 }
 
 void FSceneView::SetupCommonViewUniformBufferParameters(
diff --git a/Engine/Source/Runtime/Engine/Public/SceneView.h b/Engine/Source/Runtime/Engine/Public/SceneView.h
--- a/Engine/Source/Runtime/Engine/Public/SceneView.h
+++ b/Engine/Source/Runtime/Engine/Public/SceneView.h
@@ -798,6 +798,7 @@ enum ETranslucencyVolumeCascade
 	VIEW_UNIFORM_BUFFER_MEMBER_ARRAY(FIntVector4, PhysicsFieldTargets, [MAX_PHYSICS_FIELD_TARGETS]) \
 	VIEW_UNIFORM_BUFFER_MEMBER(uint32, InstanceSceneDataSOAStride) \
 	VIEW_UNIFORM_BUFFER_MEMBER(uint32, GPUSceneViewId) \
+	VIEW_UNIFORM_BUFFER_MEMBER(float, ViewResolutionFraction) \
 
 #define VIEW_UNIFORM_BUFFER_MEMBER(type, identifier) \
 	SHADER_PARAMETER(type, identifier)
diff --git a/Engine/Source/Runtime/Renderer/Private/Strata/StrataVisualize.cpp b/Engine/Source/Runtime/Renderer/Private/Strata/StrataVisualize.cpp
--- a/Engine/Source/Runtime/Renderer/Private/Strata/StrataVisualize.cpp
+++ b/Engine/Source/Runtime/Renderer/Private/Strata/StrataVisualize.cpp
@@ -196,8 +196,9 @@ class FMaterialDebugStrataTreePS : public FGlobalShader
 };
 IMPLEMENT_GLOBAL_SHADER(FMaterialDebugStrataTreePS, "/Engine/Private/Strata/StrataVisualize.usf", "MaterialDebugStrataTreePS", SF_Pixel);
 
-static void AddVisualizeMaterialPasses(FRDGBuilder& GraphBuilder, const FViewInfo& View, FRDGTextureRef SceneColorTexture, EShaderPlatform Platform)
+static void AddVisualizeMaterialPasses(FRDGBuilder& GraphBuilder, const FViewInfo& View, FScreenPassTexture& ScreenPassSceneColor, EShaderPlatform Platform)
 {
+	FRDGTextureRef SceneColorTexture = ScreenPassSceneColor.Texture;
 	FRHIBlendState* PreMultipliedColorTransmittanceBlend = TStaticBlendState<CW_RGB, BO_Add, BF_One, BF_SourceAlpha, BO_Add, BF_Zero, BF_One>::GetRHI();
 	if (View.Family->EngineShowFlags.VisualizeStrataMaterial)
 	{
@@ -248,7 +249,7 @@ static void AddVisualizeMaterialPasses(FRDGBuilder& GraphBuilder, const FViewInf
 			FVisualizeMaterialPS::FPermutationDomain PermutationVector;
 			TShaderMapRef<FVisualizeMaterialPS> PixelShader(View.ShaderMap, PermutationVector);
 
-			FPixelShaderUtils::AddFullscreenPass<FVisualizeMaterialPS>(GraphBuilder, View.ShaderMap, RDG_EVENT_NAME("Strata::VisualizeMaterial(Draw)"), PixelShader, PassParameters, View.ViewRect, PreMultipliedColorTransmittanceBlend);
+			FPixelShaderUtils::AddFullscreenPass<FVisualizeMaterialPS>(GraphBuilder, View.ShaderMap, RDG_EVENT_NAME("Strata::VisualizeMaterial(Draw)"), PixelShader, PassParameters, ScreenPassSceneColor.ViewRect, PreMultipliedColorTransmittanceBlend);
 		}
 
 		// Draw each material layer independently
@@ -282,7 +283,7 @@ static void AddVisualizeMaterialPasses(FRDGBuilder& GraphBuilder, const FViewInf
 					FMaterialDebugStrataTreePS::FPermutationDomain PermutationVector;
 					TShaderMapRef<FMaterialDebugStrataTreePS> PixelShader(View.ShaderMap, PermutationVector);
 
-					FPixelShaderUtils::AddFullscreenPass<FMaterialDebugStrataTreePS>(GraphBuilder, View.ShaderMap, RDG_EVENT_NAME("Strata::StrataAdvancedVisualization(Draw)"), PixelShader, PassParameters, View.ViewRect, PreMultipliedColorTransmittanceBlend);
+					FPixelShaderUtils::AddFullscreenPass<FMaterialDebugStrataTreePS>(GraphBuilder, View.ShaderMap, RDG_EVENT_NAME("Strata::StrataAdvancedVisualization(Draw)"), PixelShader, PassParameters, ScreenPassSceneColor.ViewRect, PreMultipliedColorTransmittanceBlend);
 				}
 			}
 			else
@@ -309,7 +310,7 @@ FScreenPassTexture AddStrataDebugPasses(FRDGBuilder& GraphBuilder, const FViewIn
 	if (StrataDebugVisualizationCanRunOnPlatform(Platform))
 	{
 		RDG_EVENT_SCOPE(GraphBuilder, "Strata::VisualizeMaterial");
-		AddVisualizeMaterialPasses(GraphBuilder, View, ScreenPassSceneColor.Texture, Platform);
+		AddVisualizeMaterialPasses(GraphBuilder, View, ScreenPassSceneColor, Platform);
 	}
 
 	const int32 StrataClassificationDebug = CVarStrataClassificationDebug.GetValueOnAnyThread();
