Apply minimum width to page table to avoid bug on one platform.
Further RHI investigation needed to find issue with small page table texture.
#rnx

diff --git a/Engine/Source/Runtime/Renderer/Private/VT/VirtualTextureSpace.cpp b/Engine/Source/Runtime/Renderer/Private/VT/VirtualTextureSpace.cpp
--- a/Engine/Source/Runtime/Renderer/Private/VT/VirtualTextureSpace.cpp
+++ b/Engine/Source/Runtime/Renderer/Private/VT/VirtualTextureSpace.cpp
@@ -133,7 +133,7 @@ uint32 FVirtualTextureSpace::GetSizeInBytes() const
 	uint32 TotalSize = 0u;
 	for (uint32 TextureIndex = 0u; TextureIndex < GetNumPageTableTextures(); ++TextureIndex)
 	{
-		const uint32 PageTableWidth = Align(Allocator.GetAllocatedWidth(), VIRTUALTEXTURE_MIN_PAGETABLE_SIZE);
+		const uint32 PageTableWidth = FMath::Max(Align(Allocator.GetAllocatedWidth(), VIRTUALTEXTURE_MIN_PAGETABLE_SIZE), MinPageTableWidth);
 		const uint32 PageTableHeight = Align(Allocator.GetAllocatedHeight(), VIRTUALTEXTURE_MIN_PAGETABLE_SIZE);
 		const uint32 NumPageTableLevels = FMath::FloorLog2(FMath::Max(PageTableWidth, PageTableHeight)) + 1u;
 		const SIZE_T TextureSize = CalcTextureSize(PageTableWidth, PageTableHeight, TexturePixelFormat[TextureIndex], NumPageTableLevels);
@@ -269,7 +269,7 @@ void FVirtualTextureSpace::AllocateTextures(FRDGBuilder& GraphBuilder)
 		const TCHAR* TextureNames[] = { TEXT("VirtualPageTable_0"), TEXT("VirtualPageTable_1") };
 		static_assert(UE_ARRAY_COUNT(TextureNames) == TextureCapacity, "");
 
-		CachedPageTableWidth = Align(Allocator.GetAllocatedWidth(), VIRTUALTEXTURE_MIN_PAGETABLE_SIZE);
+		CachedPageTableWidth = FMath::Max(Align(Allocator.GetAllocatedWidth(), VIRTUALTEXTURE_MIN_PAGETABLE_SIZE), MinPageTableWidth);
 		CachedPageTableHeight = Align(Allocator.GetAllocatedHeight(), VIRTUALTEXTURE_MIN_PAGETABLE_SIZE);
 		CachedNumPageTableLevels = FMath::FloorLog2(FMath::Max(CachedPageTableWidth, CachedPageTableHeight)) + 1u;
 
diff --git a/Engine/Source/Runtime/Renderer/Private/VT/VirtualTextureSpace.h b/Engine/Source/Runtime/Renderer/Private/VT/VirtualTextureSpace.h
--- a/Engine/Source/Runtime/Renderer/Private/VT/VirtualTextureSpace.h
+++ b/Engine/Source/Runtime/Renderer/Private/VT/VirtualTextureSpace.h
@@ -104,6 +104,7 @@ public:
 
 private:
 	static const uint32 TextureCapacity = (VIRTUALTEXTURE_SPACE_MAXLAYERS + LayersPerPageTableTexture - 1u) / LayersPerPageTableTexture;
+	static const uint32 MinPageTableWidth = 128u; // UE-149293
 
 	struct FTextureEntry
 	{
