Automatically change to spatial upscale when subsurface scattering visualization is enabled to fix visualization regression.

diff --git a/Engine/Source/Runtime/Renderer/Private/SceneRendering.cpp b/Engine/Source/Runtime/Renderer/Private/SceneRendering.cpp
--- a/Engine/Source/Runtime/Renderer/Private/SceneRendering.cpp
+++ b/Engine/Source/Runtime/Renderer/Private/SceneRendering.cpp
@@ -2897,7 +2897,8 @@ void FSceneRenderer::PrepareViewRectsForRendering(FRHICommandListImmediate& RHIC
 			// Tenmporal upsample is supported only if TAA is turned on.
 			if (View.PrimaryScreenPercentageMethod == EPrimaryScreenPercentageMethod::TemporalUpscale &&
 				(!IsTemporalAccumulationBasedMethod(View.AntiAliasingMethod) ||
-				 ActiveViewFamily->EngineShowFlags.VisualizeBuffer))
+				 ActiveViewFamily->EngineShowFlags.VisualizeBuffer || 
+				 ActiveViewFamily->EngineShowFlags.VisualizeSSS))
 			{
 				View.PrimaryScreenPercentageMethod = EPrimaryScreenPercentageMethod::SpatialUpscale;
 			}
