HairStrands: Allow a negative threshold to disable the path tracer invalidation to work around cases where the hair simulation does not fully settle

Also remove the clamp so the threshold can be increased past 1cm without fully disabling it.

diff --git a/Engine/Source/Runtime/Renderer/Private/HairStrands/HairStrandsVisibility.cpp b/Engine/Source/Runtime/Renderer/Private/HairStrands/HairStrandsVisibility.cpp
--- a/Engine/Source/Runtime/Renderer/Private/HairStrands/HairStrandsVisibility.cpp
+++ b/Engine/Source/Runtime/Renderer/Private/HairStrands/HairStrandsVisibility.cpp
@@ -93,7 +93,7 @@ static int32 GHairStrandsLightSampleFormat = 1;
 static FAutoConsoleVariableRef CVarHairStrandsLightSampleFormat(TEXT("r.HairStrands.LightSampleFormat"), GHairStrandsLightSampleFormat, TEXT("Define the format used for storing the lighting of hair samples (0: RGBA-16bits, 1: RGB-11.11.10bits)"));
 
 static float GHairStrands_InvalidationPosition_Threshold = 0.05f;
-static FAutoConsoleVariableRef CVarHairStrands_InvalidationPosition_Threshold(TEXT("r.HairStrands.PathTracing.InvalidationThreshold"), GHairStrands_InvalidationPosition_Threshold, TEXT("Define the minimal distance to invalidate path tracer output when groom changes (in cm, default: 0.5mm)"));
+static FAutoConsoleVariableRef CVarHairStrands_InvalidationPosition_Threshold(TEXT("r.HairStrands.PathTracing.InvalidationThreshold"), GHairStrands_InvalidationPosition_Threshold, TEXT("Define the minimal distance to invalidate path tracer output when groom changes (in cm, default: 0.5mm)\nSet to a negative value to disable this feature"));
 
 static int32 GHairStrands_InvalidationPosition_Debug = 0;
 static FAutoConsoleVariableRef CVarHairStrands_InvalidationPosition_Debug(TEXT("r.HairStrands.PathTracing.InvalidationDebug"), GHairStrands_InvalidationPosition_Debug, TEXT("Enable bounding box drawing for groom element causing path tracer invalidation"));
@@ -3460,7 +3460,7 @@ static void AddHairStrandsHasPositionChangedPass(
 	FHairStrandsPositionChangedCS::FParameters* Parameters = GraphBuilder.AllocParameters<FHairStrandsPositionChangedCS::FParameters>();
 	Parameters->VertexCount = VertexCount;
 	Parameters->DispatchCountX = DispatchCount.X;
-	Parameters->PositionThreshold2 = FMath::Square(FMath::Clamp(GHairStrands_InvalidationPosition_Threshold, 0, 1.0f));
+	Parameters->PositionThreshold2 = FMath::Square(GHairStrands_InvalidationPosition_Threshold);
 	Parameters->bDrawInvalidElement = GHairStrands_InvalidationPosition_Debug > 0 ? 1u : 0u;
 	Parameters->HairStrandsVF_bIsCullingEnable = 0u;
 	Parameters->HairStrandsVF_CullingIndexBuffer = GraphBuilder.CreateSRV(GSystemTextures.GetDefaultBuffer(GraphBuilder, 4, 0u), PF_R32_UINT);
@@ -3603,6 +3603,11 @@ bool HasPositionsChanged(FRDGBuilder& GraphBuilder, const FViewInfo& View)
 		return false;
 	}
 
+	if (GHairStrands_InvalidationPosition_Threshold < 0)
+	{
+		return false;
+	}
+
 	FHairStrandsViewStateData* HairStrandsViewStateData = const_cast<FHairStrandsViewStateData*>(&View.ViewState->HairStrandsViewStateData);
 	if (!HairStrandsViewStateData->IsInit())
 	{
