Support shading models in a mobile deferred shading path. Disabled by default
Requrements:
iOS: A8+
Android VK: more than 4 color attachements, more than 4 input attachments
Android GL: Unsuported - PLS seem to be limited to 128bits on most devices

diff --git a/Engine/Shaders/Private/MobileBasePassPixelShader.usf b/Engine/Shaders/Private/MobileBasePassPixelShader.usf
--- a/Engine/Shaders/Private/MobileBasePassPixelShader.usf
+++ b/Engine/Shaders/Private/MobileBasePassPixelShader.usf
@@ -67,11 +67,14 @@
 // SceneDepthAux always enabled for IOS, on other platforms only enabled when in forward and MobileHDR=true
 #define USE_SCENE_DEPTH_AUX ((METAL_PROFILE || (!DEFERRED_SHADING_PATH && OUTPUT_MOBILE_HDR) || (DEFERRED_SHADING_PATH && USE_GLES_FBF_DEFERRED)) && (MATERIALBLENDING_SOLID || MATERIALBLENDING_MASKED) && !COMPILER_SWITCH)
 
-#if DEFERRED_SHADING_PATH
-	#define SV_TargetDepthAux SV_Target4
+#if DEFERRED_SHADING_PATH && MOBILE_EXTENDED_GBUFFER
+	#define SV_TargetDepthAux SV_Target5	
+#elif DEFERRED_SHADING_PATH
+	#define SV_TargetDepthAux SV_Target4	
 #else
 	#define SV_TargetDepthAux SV_Target1
 #endif
+
 #define TRANSLUCENCY_NON_DIRECTIONAL (MATERIALBLENDING_ANY_TRANSLUCENT && (TRANSLUCENCY_LIGHTING_VOLUMETRIC_NONDIRECTIONAL || TRANSLUCENCY_LIGHTING_VOLUMETRIC_PERVERTEX_NONDIRECTIONAL) && !MATERIAL_SHADINGMODEL_SINGLELAYERWATER)
 
 #include "SHCommon.ush"
@@ -184,14 +187,17 @@ void Main(
 	, in float4 SvPosition : SV_Position
 	OPTIONAL_IsFrontFace
 #if DEFERRED_SHADING_PATH
-#if USE_GLES_FBF_DEFERRED
+	#if USE_GLES_FBF_DEFERRED
 	, out HALF4_TYPE OutProxy : SV_Target0
-#else
+	#else
 	, out HALF4_TYPE OutColor : SV_Target0
-#endif
+	#endif
 	, out HALF4_TYPE OutGBufferA : SV_Target1
 	, out HALF4_TYPE OutGBufferB : SV_Target2
-	, out HALF4_TYPE OutGBufferC : SV_Target3  
+	, out HALF4_TYPE OutGBufferC : SV_Target3
+	#if MOBILE_EXTENDED_GBUFFER
+	, out HALF4_TYPE OutGBufferD : SV_Target4
+	#endif
 #else
 	, out HALF4_TYPE OutColor : SV_Target0
 #endif
@@ -203,8 +209,13 @@ void Main(
 #endif
 	)
 {  
-#if DEFERRED_SHADING_PATH && USE_GLES_FBF_DEFERRED
-	half4 OutColor;	
+#if DEFERRED_SHADING_PATH 
+	#if USE_GLES_FBF_DEFERRED
+		half4 OutColor;
+	#endif
+	#if !MOBILE_EXTENDED_GBUFFER
+		half4 OutGBufferD;
+	#endif
 #endif
 
 #if MOBILE_MULTI_VIEW
@@ -389,7 +400,6 @@ void Main(
 	half Shadow = GetPrimaryPrecomputedShadowMask(LightmapVTPageTableResult, Interpolants, MaterialParameters).r;
 
 #if DEFERRED_SHADING_PATH
-	float4 OutGBufferD;
 	float4 OutGBufferE;
 	float4 OutGBufferF;
 	float4 OutGBufferVelocity = 0;
diff --git a/Engine/Shaders/Private/MobileDeferredShading.usf b/Engine/Shaders/Private/MobileDeferredShading.usf
--- a/Engine/Shaders/Private/MobileDeferredShading.usf
+++ b/Engine/Shaders/Private/MobileDeferredShading.usf
@@ -13,10 +13,9 @@
 #define MobileSceneTextures		MobileBasePass.SceneTextures
 #define ForwardLightData		MobileBasePass.Forward
 #define PlanarReflectionStruct	MobileBasePass.PlanarReflection
-#ifndef GLOBAL_PREINTEGRATEDGF
+#define ReflectionStruct		MobileBasePass.ReflectionsParameters
 #define PreIntegratedGF			ReflectionStruct.PreIntegratedGF
 #define PreIntegratedGFSampler	ReflectionStruct.PreIntegratedGFSampler
-#endif
 
 #include "MobileLightingCommon.ush"
 #include "IESLightProfilesCommon.ush"
@@ -38,29 +37,37 @@
 
 void FetchGBuffer(in float2 UV, out float4 GBufferA, out float4 GBufferB, out float4 GBufferC, out float4 GBufferD, out float SceneDepth)
 {
+	GBufferD = 0;
+
 #if VULKAN_PROFILE
 	GBufferA = VulkanSubpassFetch1(); 
 	GBufferB = VulkanSubpassFetch2(); 
-	GBufferC = VulkanSubpassFetch3(); 
-	GBufferD = 0;
+	GBufferC = VulkanSubpassFetch3();
+	#if MOBILE_EXTENDED_GBUFFER
+		GBufferD = VulkanSubpassFetch4();
+	#endif
 	SceneDepth = ConvertFromDeviceZ(VulkanSubpassDepthFetch());
 #elif METAL_PROFILE
 	GBufferA = SubpassFetchRGBA_1(); 
 	GBufferB = SubpassFetchRGBA_2(); 
 	GBufferC = SubpassFetchRGBA_3(); 
-	GBufferD = 0; 
+	#if MOBILE_EXTENDED_GBUFFER
+		GBufferD = SubpassFetchRGBA_4();
+	#endif
 	SceneDepth = ConvertFromDeviceZ(SubpassFetchR_4());
 #elif USE_GLES_FBF_DEFERRED
 	GBufferA = GLSubpassFetch1(); 
 	GBufferB = GLSubpassFetch2(); 
 	GBufferC = GLSubpassFetch3();  
-	GBufferD = 0;   
+	GBufferD = 0; // PLS is limited to 128bits
 	SceneDepth = ConvertFromDeviceZ(DepthbufferFetchES2());
 #else
 	GBufferA = Texture2DSampleLevel(MobileSceneTextures.GBufferATexture, MobileSceneTextures.GBufferATextureSampler, UV, 0); 
 	GBufferB = Texture2DSampleLevel(MobileSceneTextures.GBufferBTexture, MobileSceneTextures.GBufferBTextureSampler, UV, 0);
 	GBufferC = Texture2DSampleLevel(MobileSceneTextures.GBufferCTexture, MobileSceneTextures.GBufferCTextureSampler, UV, 0);
-	GBufferD = 0;
+	#if MOBILE_EXTENDED_GBUFFER
+		GBufferD = Texture2DSampleLevel(MobileSceneTextures.GBufferDTexture, MobileSceneTextures.GBufferDTextureSampler, UV, 0);
+	#endif
 	SceneDepth = ConvertFromDeviceZ(Texture2DSampleLevel(MobileSceneTextures.SceneDepthTexture, MobileSceneTextures.SceneDepthTextureSampler, UV, 0).r);
 #endif
 }
@@ -306,15 +313,15 @@ void MobileRadialLightPS(
 }
 
 void MobileReflectionEnvironmentSkyLightingPS(
-	noperspective float4 UVAndScreenPos : TEXCOORD0,
-	float4 SvPosition : SV_POSITION,
+	noperspective float4 UVAndScreenPos : TEXCOORD0
+	, float4 SvPosition : SV_POSITION
 #if USE_GLES_FBF_DEFERRED
-	out HALF4_TYPE OutProxyAdditive : SV_Target0,
-	out HALF4_TYPE OutGBufferA : SV_Target1,
-	out HALF4_TYPE OutGBufferB : SV_Target2,
-	out HALF4_TYPE OutGBufferC : SV_Target3
+	, out HALF4_TYPE OutProxyAdditive : SV_Target0
+	, out HALF4_TYPE OutGBufferA : SV_Target1
+	, out HALF4_TYPE OutGBufferB : SV_Target2
+	, out HALF4_TYPE OutGBufferC : SV_Target3
 #else
-	out HALF4_TYPE OutColor : SV_Target0
+	, out HALF4_TYPE OutColor : SV_Target0
 #endif
 )
 {
diff --git a/Engine/Shaders/Private/MobileLightingCommon.ush b/Engine/Shaders/Private/MobileLightingCommon.ush
--- a/Engine/Shaders/Private/MobileLightingCommon.ush
+++ b/Engine/Shaders/Private/MobileLightingCommon.ush
@@ -9,19 +9,19 @@
 #include "LightGridCommon.ush"
 
 #ifndef TRANSLUCENCY_NON_DIRECTIONAL
-#define  TRANSLUCENCY_NON_DIRECTIONAL 0
-#endif // !TRANSLUCENCY_NON_DIRECTIONAL
+#define TRANSLUCENCY_NON_DIRECTIONAL 0
+#endif
 
 #ifndef MOBILE_USE_CSM_BRANCH
 #define MOBILE_USE_CSM_BRANCH 0
 #endif
 
 #ifndef LQ_TEXTURE_LIGHTMAP
-#define  LQ_TEXTURE_LIGHTMAP 0
+#define LQ_TEXTURE_LIGHTMAP 0
 #endif
 
 #define REFLECTION_COMPOSITE_USE_BLENDED_REFLECTION_CAPTURES 1
-#define REFLECTION_COMPOSITE_SUPPORT_SKYLIGHT_BLEND 1
+#define REFLECTION_COMPOSITE_SUPPORT_SKYLIGHT_BLEND 0 // Adds additional sampler
 #define REFLECTION_COMPOSITE_HAS_SPHERE_CAPTURES 1
 #define REFLECTION_COMPOSITE_HAS_BOX_CAPTURES 1
 #include "ReflectionEnvironmentComposite.ush"
diff --git a/Engine/Shaders/Private/MobileShadingModels.ush b/Engine/Shaders/Private/MobileShadingModels.ush
--- a/Engine/Shaders/Private/MobileShadingModels.ush
+++ b/Engine/Shaders/Private/MobileShadingModels.ush
@@ -51,10 +51,11 @@
 #define DEFERRED_SHADING_PATH 0
 #endif
 
-// Enable shading model support for deferred shading once the gbuffer has capacity.
-#define MOBILE_SHADINGMODEL_SUPPORT_DEFERRED 0
+#ifndef ENABLE_SHADINGMODEL_SUPPORT_MOBILE_DEFERRED
+#define ENABLE_SHADINGMODEL_SUPPORT_MOBILE_DEFERRED 0
+#endif
 
-#define MOBILE_SHADINGMODEL_SUPPORT (!MOBILE_DEFERRED_SHADING || MOBILE_SHADINGMODEL_SUPPORT_DEFERRED)
+#define MOBILE_SHADINGMODEL_SUPPORT (!MOBILE_DEFERRED_SHADING || ENABLE_SHADINGMODEL_SUPPORT_MOBILE_DEFERRED)
 
 /*------------------------------------------------------------------------------
 	Mobile Shading Models
diff --git a/Engine/Source/Runtime/Engine/Private/ShaderCompiler/ShaderCompiler.cpp b/Engine/Source/Runtime/Engine/Private/ShaderCompiler/ShaderCompiler.cpp
--- a/Engine/Source/Runtime/Engine/Private/ShaderCompiler/ShaderCompiler.cpp
+++ b/Engine/Source/Runtime/Engine/Private/ShaderCompiler/ShaderCompiler.cpp
@@ -5479,6 +5479,7 @@ void GlobalBeginCompileShader(
 
 			bool bGLESDeferredShading = Target.Platform == SP_OPENGL_ES3_1_ANDROID;
 			Input.Environment.SetDefine(TEXT("USE_GLES_FBF_DEFERRED"), bGLESDeferredShading ? 1 : 0);
+			Input.Environment.SetDefine(TEXT("MOBILE_EXTENDED_GBUFFER"), MobileEnableExtenedGBuffer((EShaderPlatform)Target.Platform) ? 1 : 0);
 		}
 	}
 
diff --git a/Engine/Source/Runtime/RenderCore/Private/RenderUtils.cpp b/Engine/Source/Runtime/RenderCore/Private/RenderUtils.cpp
--- a/Engine/Source/Runtime/RenderCore/Private/RenderUtils.cpp
+++ b/Engine/Source/Runtime/RenderCore/Private/RenderUtils.cpp
@@ -1228,6 +1228,17 @@ RENDERCORE_API bool MobileUsesShadowMaskTexture(const FStaticShaderPlatform Plat
 	return IsMobileDistanceFieldEnabled(Platform) || (!IsMobileDeferredShadingEnabled(Platform) && IsMobileMovableSpotlightShadowsEnabled(Platform) && MobileForwardEnableLocalLights(Platform));
 }
 
+// Whether to support more than 4 color attachments for GBuffer. Required for shading models with a custom data
+RENDERCORE_API bool MobileEnableExtenedGBuffer(FStaticShaderPlatform ShaderPlatform)
+{
+	// Android GLES: uses PLS for deferred shading and limited to 128 bits
+	// Vulkan requires:
+		// maxDescriptorSetInputAttachments > 4
+		// maxColorAttachments > 4
+	// iOS: A8+
+	return (ShaderPlatform != SP_OPENGL_ES3_1_ANDROID) && false;
+}
+
 RENDERCORE_API bool MobileBasePassAlwaysUsesCSM(const FStaticShaderPlatform Platform)
 {
 	static TConsoleVariableData<int32>* CVar = IConsoleManager::Get().FindTConsoleVariableDataInt(TEXT("r.Mobile.Shadow.CSMShaderCullingMethod"));
diff --git a/Engine/Source/Runtime/RenderCore/Private/Shader.cpp b/Engine/Source/Runtime/RenderCore/Private/Shader.cpp
--- a/Engine/Source/Runtime/RenderCore/Private/Shader.cpp
+++ b/Engine/Source/Runtime/RenderCore/Private/Shader.cpp
@@ -1567,7 +1567,7 @@ void ShaderMapAppendKeyString(EShaderPlatform Platform, FString& KeyString)
 			{
 				if (MobileShadingPathCVar->GetInt() != 0)
 				{
-					KeyString += TEXT("_MobDSh");
+					KeyString += (MobileEnableExtenedGBuffer(Platform) ? TEXT("_MobDShEx") : TEXT("_MobDSh"));
 				}
 				else
 				{
diff --git a/Engine/Source/Runtime/RenderCore/Public/RenderUtils.h b/Engine/Source/Runtime/RenderCore/Public/RenderUtils.h
--- a/Engine/Source/Runtime/RenderCore/Public/RenderUtils.h
+++ b/Engine/Source/Runtime/RenderCore/Public/RenderUtils.h
@@ -511,6 +511,8 @@ RENDERCORE_API bool MobileEnableClusteredReflections(const FStaticShaderPlatform
 
 RENDERCORE_API bool MobileUsesShadowMaskTexture(const FStaticShaderPlatform Platform);
 
+RENDERCORE_API bool MobileEnableExtenedGBuffer(const FStaticShaderPlatform Platform);
+
 RENDERCORE_API bool MobileBasePassAlwaysUsesCSM(const FStaticShaderPlatform Platform);
 
 RENDERCORE_API bool SupportsGen4TAA(const FStaticShaderPlatform Platform);
diff --git a/Engine/Source/Runtime/Renderer/Private/DepthRendering.cpp b/Engine/Source/Runtime/Renderer/Private/DepthRendering.cpp
--- a/Engine/Source/Runtime/Renderer/Private/DepthRendering.cpp
+++ b/Engine/Source/Runtime/Renderer/Private/DepthRendering.cpp
@@ -719,9 +719,11 @@ void SetMobileDepthPassRenderState(const FPrimitiveSceneProxy* RESTRICT Primitiv
 
 	if (bUsesDeferredShading)
 	{
+		extern uint8 GetMobileShadingModelStencilValue(FMaterialShadingModelField ShadingModel);
+
 		// store into [1-3] bits
-		const FMaterial& MaterialResource = *MeshBatch.MaterialRenderProxy->GetMaterialNoFallback(ERHIFeatureLevel::ES3_1);
-		uint8 ShadingModel = MaterialResource.GetShadingModels().IsLit() ? MSM_DefaultLit : MSM_Unlit;
+		const FMaterial& MaterialResource = MeshBatch.MaterialRenderProxy->GetIncompleteMaterialWithFallback(ERHIFeatureLevel::ES3_1);
+		uint8 ShadingModel = GetMobileShadingModelStencilValue(MaterialResource.GetShadingModels());
 		StencilValue |= GET_STENCIL_MOBILE_SM_MASK(ShadingModel);
 		StencilValue |= STENCIL_LIGHTING_CHANNELS_MASK(PrimitiveSceneProxy ? PrimitiveSceneProxy->GetLightingChannelStencilValue() : 0x00);
 	}
diff --git a/Engine/Source/Runtime/Renderer/Private/MobileBasePass.cpp b/Engine/Source/Runtime/Renderer/Private/MobileBasePass.cpp
--- a/Engine/Source/Runtime/Renderer/Private/MobileBasePass.cpp
+++ b/Engine/Source/Runtime/Renderer/Private/MobileBasePass.cpp
@@ -14,6 +14,21 @@
 #include "MeshPassProcessor.inl"
 #include "Engine/TextureCube.h"
 
+uint8 GetMobileShadingModelStencilValue(FMaterialShadingModelField ShadingModel)
+{
+	if (ShadingModel.HasOnlyShadingModel(MSM_DefaultLit))
+	{
+		return 1u;
+	}
+	else if (ShadingModel.HasOnlyShadingModel(MSM_Unlit))
+	{
+		return 0u;
+	}
+	
+	// mark everyhing as MSM_DefaultLit if extended GBuffer is not supported
+	return MobileEnableExtenedGBuffer(GMaxRHIShaderPlatform) ? 2u : 1u;
+}
+
 bool MobileUsesNoLightMapPermutation(const FMeshMaterialShaderPermutationParameters& Parameters)
 {
 	static const auto AllowStaticLightingVar = IConsoleManager::Get().FindTConsoleVariableDataInt(TEXT("r.AllowStaticLighting"));
@@ -331,8 +346,7 @@ void MobileBasePass::SetOpaqueRenderState(FMeshPassProcessorRenderState& DrawRen
 	
 	if (bUsesDeferredShading)
 	{
-		// store into [1-3] bits
-		uint8 ShadingModel = Material.GetShadingModels().IsLit() ? MSM_DefaultLit : MSM_Unlit;
+		uint8 ShadingModel = GetMobileShadingModelStencilValue(Material.GetShadingModels());
 		StencilValue |= GET_STENCIL_MOBILE_SM_MASK(ShadingModel);
 		StencilValue |= STENCIL_LIGHTING_CHANNELS_MASK(PrimitiveSceneProxy ? PrimitiveSceneProxy->GetLightingChannelStencilValue() : 0x00);
 	}
diff --git a/Engine/Source/Runtime/Renderer/Private/MobileDeferredShadingPass.cpp b/Engine/Source/Runtime/Renderer/Private/MobileDeferredShadingPass.cpp
--- a/Engine/Source/Runtime/Renderer/Private/MobileDeferredShadingPass.cpp
+++ b/Engine/Source/Runtime/Renderer/Private/MobileDeferredShadingPass.cpp
@@ -35,13 +35,15 @@ class FMobileDirectionalLightFunctionPS : public FMaterialShader
 	DECLARE_SHADER_TYPE(FMobileDirectionalLightFunctionPS, Material);
 	SHADER_USE_PARAMETER_STRUCT_WITH_LEGACY_BASE(FMobileDirectionalLightFunctionPS, FMaterialShader)
 
-	class FEnableClustredLights		: SHADER_PERMUTATION_BOOL("ENABLE_CLUSTERED_LIGHTS");
-	class FEnableClustredReflection	: SHADER_PERMUTATION_BOOL("ENABLE_CLUSTERED_REFLECTION");
-	class FEnableSkyLight			: SHADER_PERMUTATION_BOOL("ENABLE_SKY_LIGHT");
-	class FApplyCSM					: SHADER_PERMUTATION_BOOL("APPLY_CSM");
-	class FShadowQuality			: SHADER_PERMUTATION_INT("MOBILE_SHADOW_QUALITY", 4);
+	class FEnableShadingModelSupport	: SHADER_PERMUTATION_BOOL("ENABLE_SHADINGMODEL_SUPPORT_MOBILE_DEFERRED");
+	class FEnableClustredLights			: SHADER_PERMUTATION_BOOL("ENABLE_CLUSTERED_LIGHTS");
+	class FEnableClustredReflection		: SHADER_PERMUTATION_BOOL("ENABLE_CLUSTERED_REFLECTION");
+	class FEnableSkyLight				: SHADER_PERMUTATION_BOOL("ENABLE_SKY_LIGHT");
+	class FApplyCSM						: SHADER_PERMUTATION_BOOL("APPLY_CSM");
+	class FShadowQuality				: SHADER_PERMUTATION_INT("MOBILE_SHADOW_QUALITY", 4);
 	
-	using FPermutationDomain = TShaderPermutationDomain< 
+	using FPermutationDomain = TShaderPermutationDomain<
+		FEnableShadingModelSupport,
 		FEnableClustredLights, 
 		FEnableClustredReflection, 
 		FEnableSkyLight,
@@ -50,7 +52,6 @@ class FMobileDirectionalLightFunctionPS : public FMaterialShader
 
 	BEGIN_SHADER_PARAMETER_STRUCT(FParameters, )
 		SHADER_PARAMETER_STRUCT_REF(FMobileDirectionalLightShaderParameters, MobileDirectionalLight)
-		SHADER_PARAMETER_STRUCT_REF(FReflectionUniformParameters, ReflectionsParameters)
 		SHADER_PARAMETER_STRUCT_REF(FReflectionCaptureShaderData, ReflectionCaptureData)
 		SHADER_PARAMETER(FMatrix44f, TranslatedWorldToLight)
 		SHADER_PARAMETER(FVector4f, LightFunctionParameters)
@@ -81,6 +82,12 @@ class FMobileDirectionalLightFunctionPS : public FMaterialShader
 		{
 			PermutationVector.Set<FShadowQuality>(0);
 		}
+
+		if (!MobileEnableExtenedGBuffer(Platform))
+		{
+			PermutationVector.Set<FEnableShadingModelSupport>(false);
+		}
+
 		return PermutationVector;
 	}
 
@@ -103,7 +110,7 @@ class FMobileDirectionalLightFunctionPS : public FMaterialShader
 		return true;
 	}
 
-	static FPermutationDomain BuildPermutationVector(const FViewInfo& View, bool bInlineReflectionAndSky, bool bDynamicShadows, bool bSkyLight)
+	static FPermutationDomain BuildPermutationVector(const FViewInfo& View, bool bInlineReflectionAndSky, bool bShadingModelSupport, bool bDynamicShadows, bool bSkyLight)
 	{
 		bool bUseClusteredLights = GMobileUseClusteredDeferredShading != 0;
 		bool bClustredReflection = bInlineReflectionAndSky && (View.NumBoxReflectionCaptures + View.NumSphereReflectionCaptures) > 0;
@@ -112,6 +119,7 @@ class FMobileDirectionalLightFunctionPS : public FMaterialShader
 		int32 ShadowQuality = bDynamicShadows && !bMobileUsesShadowMaskTexture ? (int32)GetShadowQuality() : 0;
 				
 		FPermutationDomain PermutationVector;
+		PermutationVector.Set<FMobileDirectionalLightFunctionPS::FEnableShadingModelSupport>(bShadingModelSupport);
 		PermutationVector.Set<FMobileDirectionalLightFunctionPS::FEnableClustredLights>(bUseClusteredLights);
 		PermutationVector.Set<FMobileDirectionalLightFunctionPS::FEnableClustredReflection>(bClustredReflection);
 		PermutationVector.Set<FMobileDirectionalLightFunctionPS::FEnableSkyLight>(bEnableSkyLight);
@@ -140,10 +148,11 @@ public:
 	DECLARE_SHADER_TYPE(FMobileRadialLightFunctionPS,Material);
 	SHADER_USE_PARAMETER_STRUCT_WITH_LEGACY_BASE(FMobileRadialLightFunctionPS, FMaterialShader)
 
-	class FSpotLightDim			: SHADER_PERMUTATION_BOOL("IS_SPOT_LIGHT");
-	class FIESProfileDim		: SHADER_PERMUTATION_BOOL("USE_IES_PROFILE");
-	class FSpotLightShadowDim	: SHADER_PERMUTATION_BOOL("SUPPORT_SPOTLIGHTS_SHADOW");
-	using FPermutationDomain = TShaderPermutationDomain<FSpotLightDim, FIESProfileDim, FSpotLightShadowDim>;
+	class FEnableShadingModelSupport: SHADER_PERMUTATION_BOOL("ENABLE_SHADINGMODEL_SUPPORT_MOBILE_DEFERRED");
+	class FSpotLightDim				: SHADER_PERMUTATION_BOOL("IS_SPOT_LIGHT");
+	class FIESProfileDim			: SHADER_PERMUTATION_BOOL("USE_IES_PROFILE");
+	class FSpotLightShadowDim		: SHADER_PERMUTATION_BOOL("SUPPORT_SPOTLIGHTS_SHADOW");
+	using FPermutationDomain = TShaderPermutationDomain<FEnableShadingModelSupport, FSpotLightDim, FIESProfileDim, FSpotLightShadowDim>;
 
 	BEGIN_SHADER_PARAMETER_STRUCT(FParameters, )
 		SHADER_PARAMETER(FMatrix44f, TranslatedWorldToLight)
@@ -153,8 +162,6 @@ public:
 		SHADER_PARAMETER_STRUCT_INCLUDE(FMobileMovableLocalLightShadowParameters, MobileMovableLocalLightShadow)
 		SHADER_PARAMETER_TEXTURE(Texture2D, IESTexture)
 		SHADER_PARAMETER_SAMPLER(SamplerState, IESTextureSampler)
-		SHADER_PARAMETER_TEXTURE(Texture2D, PreIntegratedGF)
-		SHADER_PARAMETER_SAMPLER(SamplerState, PreIntegratedGFSampler)
 	END_SHADER_PARAMETER_STRUCT()
 
 	static bool ShouldCompilePermutation(const FMaterialShaderPermutationParameters& Parameters)
@@ -183,6 +190,11 @@ public:
 			PermutationVector.Set<FSpotLightShadowDim>(false);
 		}
 
+		if (!MobileEnableExtenedGBuffer(Platform))
+		{
+			PermutationVector.Set<FEnableShadingModelSupport>(false);
+		}
+
 		return PermutationVector;
 	}
 
@@ -191,7 +203,6 @@ public:
 		FMaterialShader::ModifyCompilationEnvironment(Parameters, OutEnvironment);
 		OutEnvironment.SetDefine(TEXT("USE_LIGHT_FUNCTION"), Parameters.MaterialParameters.bIsDefaultMaterial ? 0 : 1);
 		OutEnvironment.SetDefine(TEXT("MATERIAL_SHADER"), 1);
-		OutEnvironment.SetDefine(TEXT("GLOBAL_PREINTEGRATEDGF"), 1);
 		OutEnvironment.SetDefine(TEXT("USE_SHADOWMASKTEXTURE"), 0);
 		OutEnvironment.SetDefine(TEXT("ENABLE_CLUSTERED_LIGHTS"), 0);
 	}
@@ -218,13 +229,14 @@ public:
 	DECLARE_GLOBAL_SHADER(FMobileReflectionEnvironmentSkyLightingPS);
 	SHADER_USE_PARAMETER_STRUCT(FMobileReflectionEnvironmentSkyLightingPS, FGlobalShader);
 	
-	class FEnableClustredReflection	: SHADER_PERMUTATION_BOOL("ENABLE_CLUSTERED_REFLECTION");
-	class FEnableSkyLight			: SHADER_PERMUTATION_BOOL("ENABLE_SKY_LIGHT");
-	using FPermutationDomain = TShaderPermutationDomain<FEnableClustredReflection, FEnableSkyLight>;
+	class FEnableShadingModelSupport	: SHADER_PERMUTATION_BOOL("ENABLE_SHADINGMODEL_SUPPORT_MOBILE_DEFERRED");
+	class FEnableClustredReflection		: SHADER_PERMUTATION_BOOL("ENABLE_CLUSTERED_REFLECTION");
+	class FEnableSkyLight				: SHADER_PERMUTATION_BOOL("ENABLE_SKY_LIGHT");
+	
+	using FPermutationDomain = TShaderPermutationDomain<FEnableShadingModelSupport, FEnableClustredReflection, FEnableSkyLight>;
 
 	BEGIN_SHADER_PARAMETER_STRUCT(FParameters, )
 		SHADER_PARAMETER_STRUCT_REF(FViewUniformShaderParameters, View)
-		SHADER_PARAMETER_STRUCT_REF(FReflectionUniformParameters, ReflectionsParameters)
 		SHADER_PARAMETER_STRUCT_REF(FReflectionCaptureShaderData, ReflectionCaptureData)
 	END_SHADER_PARAMETER_STRUCT()
 
@@ -235,12 +247,27 @@ public:
 		{
 			return false;
 		}
+
+		FPermutationDomain PermutationVector(Parameters.PermutationId);
+		if (!MobileEnableExtenedGBuffer(Parameters.Platform) && PermutationVector.Get<FEnableShadingModelSupport>())
+		{
+			return false;
+		}
+
 		return true;
 	}
+
+	static void ModifyCompilationEnvironment(const FGlobalShaderPermutationParameters& Parameters, FShaderCompilerEnvironment& OutEnvironment)
+	{
+		FGlobalShader::ModifyCompilationEnvironment(Parameters, OutEnvironment);
+		OutEnvironment.SetDefine(TEXT("SUPPORTS_TEXTURECUBE_ARRAY"), 1);
+	}
 };
 
 IMPLEMENT_GLOBAL_SHADER(FMobileReflectionEnvironmentSkyLightingPS, "/Engine/Private/MobileDeferredShading.usf", "MobileReflectionEnvironmentSkyLightingPS", SF_Pixel);
 
+extern uint8 GetMobileShadingModelStencilValue(FMaterialShadingModelField ShadingModel);
+
 constexpr uint32 GetLightingChannel(uint32 LightingChannelMask)
 {
 	return (LightingChannelMask & 0x1) ? 0u : ((LightingChannelMask & 0x2) ? 1u : 2u);
@@ -310,40 +337,51 @@ void RenderReflectionEnvironmentSkyLighting(FRHICommandListImmediate& RHICmdList
 		false, CF_Always, SO_Keep, SO_Keep, SO_Keep,
 		GET_STENCIL_MOBILE_SM_MASK(0xff), 0x00>::GetRHI();
 
-	TShaderMapRef<FPostProcessVS> VertexShader(View.ShaderMap);
-
-	FMobileReflectionEnvironmentSkyLightingPS::FPermutationDomain PermutationVector;
-	PermutationVector.Set<FMobileReflectionEnvironmentSkyLightingPS::FEnableClustredReflection>(bClustredReflection);
-	PermutationVector.Set<FMobileReflectionEnvironmentSkyLightingPS::FEnableSkyLight>(bSkyLight);
-	TShaderMapRef<FMobileReflectionEnvironmentSkyLightingPS> PixelShader(View.ShaderMap, PermutationVector);
-	
-	GraphicsPSOInit.BoundShaderState.VertexDeclarationRHI = GFilterVertexDeclaration.VertexDeclarationRHI;
-	GraphicsPSOInit.BoundShaderState.VertexShaderRHI = VertexShader.GetVertexShader();
-	GraphicsPSOInit.BoundShaderState.PixelShaderRHI = PixelShader.GetPixelShader();
-	GraphicsPSOInit.PrimitiveType = PT_TriangleList;
-
-	uint8 StencilRef = GET_STENCIL_MOBILE_SM_MASK(MSM_DefaultLit);
-	SetGraphicsPipelineState(RHICmdList, GraphicsPSOInit, StencilRef);
-
 	FMobileReflectionEnvironmentSkyLightingPS::FParameters PassParameters;
 	PassParameters.View = GetShaderBinding(View.ViewUniformBuffer);
 	PassParameters.ReflectionCaptureData = GetShaderBinding(View.ReflectionCaptureUniformBuffer);
-	FReflectionUniformParameters ReflectionUniformParameters;
-	SetupReflectionUniformParameters(View, ReflectionUniformParameters);
-	PassParameters.ReflectionsParameters = CreateUniformBufferImmediate(ReflectionUniformParameters, UniformBuffer_SingleDraw);
-	SetShaderParameters(RHICmdList, PixelShader, PixelShader.GetPixelShader(), PassParameters);
-	
-	const FIntPoint TargetSize = GetSceneTextureExtent();
-
-	DrawRectangle(
-		RHICmdList,
-		0, 0,
-		View.ViewRect.Width(), View.ViewRect.Height(),
-		View.ViewRect.Min.X, View.ViewRect.Min.Y,
-		View.ViewRect.Width(), View.ViewRect.Height(),
-		FIntPoint(View.ViewRect.Width(), View.ViewRect.Height()),
-		TargetSize,
-		VertexShader);
+
+	TShaderMapRef<FPostProcessVS> VertexShader(View.ShaderMap);
+
+	// Do two passes, first masking DefautLit, second masking all other shading models
+	int32 NumPasses = MobileEnableExtenedGBuffer(Scene.GetShaderPlatform()) ? 2 : 1;
+	uint8 PassShadingModelStencilValue[2] =
+	{
+		GetMobileShadingModelStencilValue(MSM_DefaultLit),
+		GetMobileShadingModelStencilValue(FMaterialShadingModelField())
+	};
+			
+	for (int32 PassIndex = 0; PassIndex < NumPasses; PassIndex++)
+	{
+		const bool bEnableShadingModelSupport = (PassIndex > 0);
+		
+		FMobileReflectionEnvironmentSkyLightingPS::FPermutationDomain PermutationVector;
+		PermutationVector.Set<FMobileReflectionEnvironmentSkyLightingPS::FEnableShadingModelSupport>(bEnableShadingModelSupport);
+		PermutationVector.Set<FMobileReflectionEnvironmentSkyLightingPS::FEnableClustredReflection>(bClustredReflection);
+		PermutationVector.Set<FMobileReflectionEnvironmentSkyLightingPS::FEnableSkyLight>(bSkyLight);
+		TShaderMapRef<FMobileReflectionEnvironmentSkyLightingPS> PixelShader(View.ShaderMap, PermutationVector);
+		
+		GraphicsPSOInit.BoundShaderState.VertexDeclarationRHI = GFilterVertexDeclaration.VertexDeclarationRHI;
+		GraphicsPSOInit.BoundShaderState.VertexShaderRHI = VertexShader.GetVertexShader();
+		GraphicsPSOInit.BoundShaderState.PixelShaderRHI = PixelShader.GetPixelShader();
+		GraphicsPSOInit.PrimitiveType = PT_TriangleList;
+
+		uint8 StencilRef = GET_STENCIL_MOBILE_SM_MASK(PassShadingModelStencilValue[PassIndex]);
+		SetGraphicsPipelineState(RHICmdList, GraphicsPSOInit, StencilRef);
+		SetShaderParameters(RHICmdList, PixelShader, PixelShader.GetPixelShader(), PassParameters);
+
+		const FIntPoint TargetSize = GetSceneTextureExtent();
+
+		DrawRectangle(
+			RHICmdList,
+			0, 0,
+			View.ViewRect.Width(), View.ViewRect.Height(),
+			View.ViewRect.Min.X, View.ViewRect.Min.Y,
+			View.ViewRect.Width(), View.ViewRect.Height(),
+			FIntPoint(View.ViewRect.Width(), View.ViewRect.Height()),
+			TargetSize,
+			VertexShader);
+	}
 }
 
 template<uint32 LightingChannelIdx>
@@ -369,8 +407,6 @@ static void RenderDirectionalLight(FRHICommandListImmediate& RHICmdList, const F
 	GraphicsPSOInit.RasterizerState = TStaticRasterizerState<>::GetRHI();
 
 	uint8 LightingChannelStencilValue = GetLightingChannelStencilValue(LightingChannel);
-	// Shade only MSM_DefaultLit pixels
-	uint8 StencilRef = GET_STENCIL_MOBILE_SM_MASK(MSM_DefaultLit) | STENCIL_LIGHTING_CHANNELS_MASK(LightingChannelStencilValue);
 	if (LightingChannel == 1u)
 	{
 		SetDirectionalLightDepthStencilState<1u>(GraphicsPSOInit);
@@ -392,27 +428,9 @@ static void RenderDirectionalLight(FRHICommandListImmediate& RHICmdList, const F
 		LightFunctionMaterialProxy = DirectionalLight.Proxy->GetLightFunctionMaterial();
 	}
 
-	// Skylights with static lighting already had their diffuse contribution baked into lightmaps
-	const bool bSkyLight = Scene.SkyLight && !Scene.SkyLight->bHasStaticLighting && View.Family->EngineShowFlags.SkyLighting;
-	const bool bDynamicShadows = DirectionalLight.Proxy->CastsDynamicShadow() && (LightingChannel == 0u) && View.Family->EngineShowFlags.DynamicShadows;
-
-	FMobileDirectionalLightFunctionPS::FPermutationDomain PermutationVector = FMobileDirectionalLightFunctionPS::BuildPermutationVector(View, bInlineReflectionAndSky, bDynamicShadows, bSkyLight);
-	FCachedLightMaterial LightMaterial;
-	TShaderRef<FMobileDirectionalLightFunctionPS> PixelShader;
-	GetLightMaterial(DefaultLightMaterial, LightFunctionMaterialProxy, PermutationVector.ToDimensionValueId(), LightMaterial, PixelShader);
-	
-	GraphicsPSOInit.BoundShaderState.VertexDeclarationRHI = GFilterVertexDeclaration.VertexDeclarationRHI;
-	GraphicsPSOInit.BoundShaderState.VertexShaderRHI = VertexShader.GetVertexShader();
-	GraphicsPSOInit.BoundShaderState.PixelShaderRHI = PixelShader.GetPixelShader();
-	GraphicsPSOInit.PrimitiveType = PT_TriangleList;
-	SetGraphicsPipelineState(RHICmdList, GraphicsPSOInit, StencilRef);
-
 	FMobileDirectionalLightFunctionPS::FParameters PassParameters;
-	PassParameters.MobileDirectionalLight = Scene.UniformBuffers.MobileDirectionalLightUniformBuffers[LightingChannel+1];
+	PassParameters.MobileDirectionalLight = Scene.UniformBuffers.MobileDirectionalLightUniformBuffers[LightingChannel + 1];
 	PassParameters.ReflectionCaptureData = GetShaderBinding(View.ReflectionCaptureUniformBuffer);
-	FReflectionUniformParameters ReflectionUniformParameters;
-	SetupReflectionUniformParameters(View, ReflectionUniformParameters);
-	PassParameters.ReflectionsParameters = CreateUniformBufferImmediate(ReflectionUniformParameters, UniformBuffer_SingleDraw);
 	PassParameters.LightFunctionParameters = FVector4f(1.0f, 1.0f, 0.0f, 0.0f);
 
 	const bool bMobileUsesShadowMaskTexture = MobileUsesShadowMaskTexture(View.GetShaderPlatform());
@@ -427,7 +445,7 @@ static void RenderDirectionalLight(FRHICommandListImmediate& RHICmdList, const F
 		PassParameters.ScreenSpaceShadowMaskTexture = GSystemTextures.WhiteDummy->GetRHI();
 		PassParameters.ScreenSpaceShadowMaskSampler = TStaticSamplerState<SF_Point, AM_Clamp, AM_Clamp, AM_Clamp>::GetRHI();
 	}
-	
+
 	{
 		const bool bUseMovableLight = !DirectionalLight.Proxy->HasStaticShadowing();
 		PassParameters.LightFunctionParameters2 = FVector3f(DirectionalLight.Proxy->GetLightFunctionFadeDistance(), DirectionalLight.Proxy->GetLightFunctionDisabledBrightness(), bUseMovableLight ? 1.0f : 0.0f);
@@ -437,19 +455,58 @@ static void RenderDirectionalLight(FRHICommandListImmediate& RHICmdList, const F
 		const FMatrix WorldToLight = DirectionalLight.Proxy->GetWorldToLight() * FScaleMatrix(FVector(InverseScale));
 		PassParameters.TranslatedWorldToLight = FMatrix44f(FTranslationMatrix(-View.ViewMatrices.GetPreViewTranslation()) * WorldToLight);
 	}
-	FMobileDirectionalLightFunctionPS::SetParameters(RHICmdList, PixelShader, View, LightMaterial.MaterialProxy, *LightMaterial.Material, PassParameters);
-	
-	const FIntPoint TargetSize = GetSceneTextureExtent();
+
+	// Skylights with static lighting already had their diffuse contribution baked into lightmaps
+	const bool bSkyLight = Scene.SkyLight && !Scene.SkyLight->bHasStaticLighting && View.Family->EngineShowFlags.SkyLighting;
+	const bool bDynamicShadows = DirectionalLight.Proxy->CastsDynamicShadow() && (LightingChannel == 0u) && View.Family->EngineShowFlags.DynamicShadows;
+
+	// Do two passes, first masking DefautLit, second masking all other shading models
+	int32 NumPasses = MobileEnableExtenedGBuffer(Scene.GetShaderPlatform()) ? 2 : 1;
+	uint8 PassShadingModelStencilValue[2] =
+	{
+		GetMobileShadingModelStencilValue(MSM_DefaultLit),
+		GetMobileShadingModelStencilValue(FMaterialShadingModelField())
+	};
 	
-	DrawRectangle(
-		RHICmdList, 
-		0, 0, 
-		View.ViewRect.Width(), View.ViewRect.Height(), 
-		View.ViewRect.Min.X, View.ViewRect.Min.Y, 
-		View.ViewRect.Width(), View.ViewRect.Height(),
-		FIntPoint(View.ViewRect.Width(), View.ViewRect.Height()), 
-		TargetSize, 
-		VertexShader);
+	for (int32 PassIndex = 0; PassIndex < NumPasses; ++PassIndex)
+	{
+		const bool bEnableShadingModelSupport = (PassIndex > 0);
+		
+		FMobileDirectionalLightFunctionPS::FPermutationDomain PermutationVector = FMobileDirectionalLightFunctionPS::BuildPermutationVector(
+			View,
+			bInlineReflectionAndSky,
+			bEnableShadingModelSupport,
+			bDynamicShadows,
+			bSkyLight
+		);
+		FCachedLightMaterial LightMaterial;
+		TShaderRef<FMobileDirectionalLightFunctionPS> PixelShader;
+		GetLightMaterial(DefaultLightMaterial, LightFunctionMaterialProxy, PermutationVector.ToDimensionValueId(), LightMaterial, PixelShader);
+
+		uint8 StencilRef = 
+			GET_STENCIL_MOBILE_SM_MASK(PassShadingModelStencilValue[PassIndex]) | 
+			STENCIL_LIGHTING_CHANNELS_MASK(LightingChannelStencilValue);
+
+		GraphicsPSOInit.BoundShaderState.VertexDeclarationRHI = GFilterVertexDeclaration.VertexDeclarationRHI;
+		GraphicsPSOInit.BoundShaderState.VertexShaderRHI = VertexShader.GetVertexShader();
+		GraphicsPSOInit.BoundShaderState.PixelShaderRHI = PixelShader.GetPixelShader();
+		GraphicsPSOInit.PrimitiveType = PT_TriangleList;
+		SetGraphicsPipelineState(RHICmdList, GraphicsPSOInit, StencilRef);
+
+		FMobileDirectionalLightFunctionPS::SetParameters(RHICmdList, PixelShader, View, LightMaterial.MaterialProxy, *LightMaterial.Material, PassParameters);
+
+		const FIntPoint TargetSize = GetSceneTextureExtent();
+
+		DrawRectangle(
+			RHICmdList,
+			0, 0,
+			View.ViewRect.Width(), View.ViewRect.Height(),
+			View.ViewRect.Min.X, View.ViewRect.Min.Y,
+			View.ViewRect.Width(), View.ViewRect.Height(),
+			FIntPoint(View.ViewRect.Width(), View.ViewRect.Height()),
+			TargetSize,
+			VertexShader);
+	}
 }
 
 static void RenderDirectionalLights(FRHICommandListImmediate& RHICmdList, const FScene& Scene, const FViewInfo& View, const FCachedLightMaterial& DefaultLightMaterial)
@@ -641,32 +698,16 @@ static void RenderLocalLight(
 	FDeferredLightVS::FPermutationDomain PermutationVectorVS;
 	PermutationVectorVS.Set<FDeferredLightVS::FRadialLight>(true);
 	TShaderMapRef<FDeferredLightVS> VertexShader(View.ShaderMap, PermutationVectorVS);
+	FDeferredLightVS::FParameters ParametersVS = FDeferredLightVS::GetParameters(View, &LightSceneInfo);
 		
 	const FMaterialRenderProxy* LightFunctionMaterialProxy = nullptr;
 	if (View.Family->EngineShowFlags.LightFunctions)
 	{
 		LightFunctionMaterialProxy = LightSceneInfo.Proxy->GetLightFunctionMaterial();
 	}
+
 	FMobileRadialLightFunctionPS::FParameters PassParameters;
 	const bool bShouldCastShadow = LightSceneInfo.SetupMobileMovableLocalLightShadowParameters(View, VisibleLightInfos, PassParameters.MobileMovableLocalLightShadow);
-	FMobileRadialLightFunctionPS::FPermutationDomain PermutationVector;
-	PermutationVector.Set<FMobileRadialLightFunctionPS::FSpotLightDim>(bIsSpotLight);
-	PermutationVector.Set<FMobileRadialLightFunctionPS::FIESProfileDim>(bUseIESTexture);
-	PermutationVector.Set<FMobileRadialLightFunctionPS::FSpotLightShadowDim>(bShouldCastShadow);
-	FCachedLightMaterial LightMaterial;
-	TShaderRef<FMobileRadialLightFunctionPS> PixelShader;
-	GetLightMaterial(DefaultLightMaterial, LightFunctionMaterialProxy, PermutationVector.ToDimensionValueId(), LightMaterial, PixelShader);
-			
-	GraphicsPSOInit.BoundShaderState.VertexDeclarationRHI = GetVertexDeclarationFVector4();
-	GraphicsPSOInit.BoundShaderState.VertexShaderRHI = VertexShader.GetVertexShader();
-	GraphicsPSOInit.BoundShaderState.PixelShaderRHI = PixelShader.GetPixelShader();
-		
-	// Shade only MSM_DefaultLit pixels
-	uint8 StencilRef = GET_STENCIL_MOBILE_SM_MASK(MSM_DefaultLit) | STENCIL_LIGHTING_CHANNELS_MASK(LightingChannelStencilValue);
-	SetGraphicsPipelineState(RHICmdList, GraphicsPSOInit, StencilRef);
-
-	FDeferredLightVS::FParameters ParametersVS = FDeferredLightVS::GetParameters(View, &LightSceneInfo);
-	SetShaderParameters(RHICmdList, VertexShader, VertexShader.GetVertexShader(), ParametersVS);
 
 	PassParameters.DeferredLightUniforms = TUniformBufferRef<FDeferredLightUniformStruct>::CreateUniformBufferImmediate(GetDeferredLightParameters(View, LightSceneInfo), EUniformBufferUsage::UniformBuffer_SingleFrame);
 	PassParameters.IESTexture = IESTextureResource->TextureRHI;
@@ -679,17 +720,49 @@ static void RenderLocalLight(
 	const FVector InverseScale = FVector(1.f / Scale.Z, 1.f / Scale.Y, 1.f / Scale.X);
 	const FMatrix WorldToLight = LightSceneInfo.Proxy->GetWorldToLight() * FScaleMatrix(FVector(InverseScale));
 	PassParameters.TranslatedWorldToLight = FMatrix44f(FTranslationMatrix(-View.ViewMatrices.GetPreViewTranslation()) * WorldToLight);
-	PassParameters.PreIntegratedGF = GSystemTextures.PreintegratedGF->GetRHI();
-	PassParameters.PreIntegratedGFSampler = TStaticSamplerState<SF_Bilinear, AM_Clamp, AM_Clamp, AM_Clamp>::GetRHI();
-	FMobileRadialLightFunctionPS::SetParameters(RHICmdList, PixelShader, View, LightMaterial.MaterialProxy, *LightMaterial.Material, PassParameters);
 
-	if (LightType == LightType_Point)
+	// Do two passes, first masking DefautLit, second masking all other shading models
+	int32 NumPasses = MobileEnableExtenedGBuffer(Scene.GetShaderPlatform()) ? 2 : 1;
+	uint8 PassShadingModelStencilValue[2] =
 	{
-		StencilingGeometry::DrawSphere(RHICmdList);
-	}
-	else // LightType_Spot
+		GetMobileShadingModelStencilValue(MSM_DefaultLit),
+		GetMobileShadingModelStencilValue(FMaterialShadingModelField())
+	};
+
+	for (int32 PassIndex = 0; PassIndex < NumPasses; PassIndex++)
 	{
-		StencilingGeometry::DrawCone(RHICmdList);
+		const bool bEnableShadingModelSupport = (PassIndex > 0);
+
+		FMobileRadialLightFunctionPS::FPermutationDomain PermutationVector;
+		PermutationVector.Set<FMobileRadialLightFunctionPS::FEnableShadingModelSupport>(bEnableShadingModelSupport);
+		PermutationVector.Set<FMobileRadialLightFunctionPS::FSpotLightDim>(bIsSpotLight);
+		PermutationVector.Set<FMobileRadialLightFunctionPS::FIESProfileDim>(bUseIESTexture);
+		PermutationVector.Set<FMobileRadialLightFunctionPS::FSpotLightShadowDim>(bShouldCastShadow);
+		FCachedLightMaterial LightMaterial;
+		TShaderRef<FMobileRadialLightFunctionPS> PixelShader;
+		GetLightMaterial(DefaultLightMaterial, LightFunctionMaterialProxy, PermutationVector.ToDimensionValueId(), LightMaterial, PixelShader);
+
+		GraphicsPSOInit.BoundShaderState.VertexDeclarationRHI = GetVertexDeclarationFVector4();
+		GraphicsPSOInit.BoundShaderState.VertexShaderRHI = VertexShader.GetVertexShader();
+		GraphicsPSOInit.BoundShaderState.PixelShaderRHI = PixelShader.GetPixelShader();
+
+		uint8 StencilRef = 
+			GET_STENCIL_MOBILE_SM_MASK(PassShadingModelStencilValue[PassIndex]) | 
+			STENCIL_LIGHTING_CHANNELS_MASK(LightingChannelStencilValue);
+
+		SetGraphicsPipelineState(RHICmdList, GraphicsPSOInit, StencilRef);
+
+		SetShaderParameters(RHICmdList, VertexShader, VertexShader.GetVertexShader(), ParametersVS);
+		FMobileRadialLightFunctionPS::SetParameters(RHICmdList, PixelShader, View, LightMaterial.MaterialProxy, *LightMaterial.Material, PassParameters);
+
+		if (LightType == LightType_Point)
+		{
+			StencilingGeometry::DrawSphere(RHICmdList);
+		}
+		else // LightType_Spot
+		{
+			StencilingGeometry::DrawCone(RHICmdList);
+		}
 	}
 }
 
@@ -713,27 +786,7 @@ static void RenderSimpleLights(
 	FDeferredLightVS::FPermutationDomain PermutationVectorVS;
 	PermutationVectorVS.Set<FDeferredLightVS::FRadialLight>(true);
 	TShaderMapRef<FDeferredLightVS> VertexShader(View.ShaderMap, PermutationVectorVS);
-	TShaderRef<FMobileRadialLightFunctionPS> PixelShader;
-	{
-		const FMaterialShaderMap* MaterialShaderMap = DefaultMaterial.Material->GetRenderingThreadShaderMap();
-		FMobileRadialLightFunctionPS::FPermutationDomain PermutationVector;
-		PermutationVector.Set<FMobileRadialLightFunctionPS::FSpotLightDim>(false);
-		PermutationVector.Set<FMobileRadialLightFunctionPS::FIESProfileDim>(false);
-		PixelShader = MaterialShaderMap->GetShader<FMobileRadialLightFunctionPS>(PermutationVector);
-	}
 
-	// Setup PSOs we going to use for light rendering 
-	FGraphicsPipelineStateInitializer GraphicsPSOLight;
-	{
-		RHICmdList.ApplyCachedRenderTargets(GraphicsPSOLight);
-		// Use additive blending for color
-		GraphicsPSOLight.BlendState = TStaticBlendState<CW_RGBA, BO_Add, BF_One, BF_One, BO_Add, BF_One, BF_One>::GetRHI();
-		GraphicsPSOLight.PrimitiveType = PT_TriangleList;
-		GraphicsPSOLight.BoundShaderState.VertexDeclarationRHI = GetVertexDeclarationFVector4();
-		GraphicsPSOLight.BoundShaderState.VertexShaderRHI = VertexShader.GetVertexShader();
-		GraphicsPSOLight.BoundShaderState.PixelShaderRHI = PixelShader.GetPixelShader();
-		SetLocalLightRasterizerAndDepthState<0u, true>(GraphicsPSOLight, View, FSphere());
-	}
 	// Setup stencil mask PSO
 	FGraphicsPipelineStateInitializer GraphicsPSOLightMask;
 	{
@@ -751,6 +804,40 @@ static void RenderSimpleLights(
 		GraphicsPSOLightMask.BoundShaderState.VertexShaderRHI = VertexShader.GetVertexShader();
 		GraphicsPSOLightMask.BoundShaderState.PixelShaderRHI = nullptr;
 	}
+
+	const FMaterialShaderMap* MaterialShaderMap = DefaultMaterial.Material->GetRenderingThreadShaderMap();
+
+	// Setup PSOs we going to use for light rendering 
+	// Do two passes, first masking DefautLit, second masking all other shading models
+	int32 NumPasses = MobileEnableExtenedGBuffer(Scene.GetShaderPlatform()) ? 2 : 1;
+	uint8 PassShadingModelStencilValue[2] =
+	{
+		GetMobileShadingModelStencilValue(MSM_DefaultLit),
+		GetMobileShadingModelStencilValue(FMaterialShadingModelField())
+	};
+	TShaderRef<FMobileRadialLightFunctionPS> PassPixelShaders[2];
+	FGraphicsPipelineStateInitializer GraphicsPSOLight[2];
+
+	for (int32 PassIndex = 0; PassIndex < NumPasses; ++PassIndex)
+	{
+		const bool bEnableShadingModelSupport = (PassIndex > 0);
+		
+		RHICmdList.ApplyCachedRenderTargets(GraphicsPSOLight[PassIndex]);
+		// Use additive blending for color
+		GraphicsPSOLight[PassIndex].BlendState = TStaticBlendState<CW_RGBA, BO_Add, BF_One, BF_One, BO_Add, BF_One, BF_One>::GetRHI();
+		GraphicsPSOLight[PassIndex].PrimitiveType = PT_TriangleList;
+		GraphicsPSOLight[PassIndex].BoundShaderState.VertexDeclarationRHI = GetVertexDeclarationFVector4();
+		GraphicsPSOLight[PassIndex].BoundShaderState.VertexShaderRHI = VertexShader.GetVertexShader();
+		SetLocalLightRasterizerAndDepthState<0u, true>(GraphicsPSOLight[PassIndex], View, FSphere());
+
+		TShaderRef<FMobileRadialLightFunctionPS> PixelShader;
+		FMobileRadialLightFunctionPS::FPermutationDomain PermutationVector;
+		PermutationVector.Set<FMobileRadialLightFunctionPS::FEnableShadingModelSupport>(bEnableShadingModelSupport);
+		PermutationVector.Set<FMobileRadialLightFunctionPS::FSpotLightDim>(false);
+		PermutationVector.Set<FMobileRadialLightFunctionPS::FIESProfileDim>(false);
+		PassPixelShaders[PassIndex] = MaterialShaderMap->GetShader<FMobileRadialLightFunctionPS>(PermutationVector);
+		GraphicsPSOLight[PassIndex].BoundShaderState.PixelShaderRHI = PassPixelShaders[PassIndex].GetPixelShader();
+	}
 		
 	for (int32 LightIndex = 0; LightIndex < SimpleLights.InstanceData.Num(); LightIndex++)
 	{
@@ -767,10 +854,8 @@ static void RenderSimpleLights(
 
 		// Render light mask
 		SetGraphicsPipelineState(RHICmdList, GraphicsPSOLightMask, 1);
-
 		FDeferredLightVS::FParameters ParametersVS = FDeferredLightVS::GetParameters(View, LightBounds);
 		SetShaderParameters(RHICmdList, VertexShader, VertexShader.GetVertexShader(), ParametersVS);
-
 		StencilingGeometry::DrawSphere(RHICmdList);
 
 		// Render light
@@ -780,17 +865,16 @@ static void RenderSimpleLights(
 		PassParameters.IESTexture = GWhiteTexture->TextureRHI;
 		PassParameters.IESTextureSampler = GWhiteTexture->SamplerStateRHI;
 
-		// Shade only MSM_DefaultLit pixels
-		uint8 StencilRef = GET_STENCIL_MOBILE_SM_MASK(MSM_DefaultLit);
-
+		for (int32 PassIndex = 0; PassIndex < NumPasses; ++PassIndex)
 		{
-			SetGraphicsPipelineState(RHICmdList, GraphicsPSOLight, StencilRef);
-			FMobileRadialLightFunctionPS::SetParameters(RHICmdList, PixelShader, View, DefaultMaterial.MaterialProxy, *DefaultMaterial.Material, PassParameters);
-		}
+			uint8 StencilRef = GET_STENCIL_MOBILE_SM_MASK(PassShadingModelStencilValue[PassIndex]);
+			SetGraphicsPipelineState(RHICmdList, GraphicsPSOLight[PassIndex], StencilRef);
+			FMobileRadialLightFunctionPS::SetParameters(RHICmdList, PassPixelShaders[PassIndex], View, DefaultMaterial.MaterialProxy, *DefaultMaterial.Material, PassParameters);
 
-		// Apply the point or spot light with some approximately bounding geometry,
-		// So we can get speedups from depth testing and not processing pixels outside of the light's influence.
-		StencilingGeometry::DrawSphere(RHICmdList);
+			// Apply the point or spot light with some approximately bounding geometry,
+			// So we can get speedups from depth testing and not processing pixels outside of the light's influence.
+			StencilingGeometry::DrawSphere(RHICmdList);
+		}
 	}
 }
 
diff --git a/Engine/Source/Runtime/Renderer/Private/MobileShadingRenderer.cpp b/Engine/Source/Runtime/Renderer/Private/MobileShadingRenderer.cpp
--- a/Engine/Source/Runtime/Renderer/Private/MobileShadingRenderer.cpp
+++ b/Engine/Source/Runtime/Renderer/Private/MobileShadingRenderer.cpp
@@ -1297,7 +1297,7 @@ void MobileDeferredCopyBuffer(FRHICommandListImmediate& RHICmdList, const FViewI
 
 void FMobileSceneRenderer::RenderDeferred(FRDGBuilder& GraphBuilder, const FSortedLightSetSceneInfo& SortedLightSet, FRDGTextureRef ViewFamilyTexture, FSceneTextures& SceneTextures)
 {
-	TArray<FRDGTextureRef, TInlineAllocator<5>> ColorTargets;
+	TArray<FRDGTextureRef, TInlineAllocator<6>> ColorTargets;
 
 	// If we are using GL and don't have FBF support, use PLS
 	bool bUsingPixelLocalStorage = IsAndroidOpenGLESPlatform(ShaderPlatform) && GSupportsPixelLocalStorage;
@@ -1312,7 +1312,13 @@ void FMobileSceneRenderer::RenderDeferred(FRDGBuilder& GraphBuilder, const FSort
 		ColorTargets.Add(SceneTextures.GBufferA);
 		ColorTargets.Add(SceneTextures.GBufferB);
 		ColorTargets.Add(SceneTextures.GBufferC);
-		if(bRequiresSceneDepthAux)
+		
+		if (MobileEnableExtenedGBuffer(ShaderPlatform))
+		{
+			ColorTargets.Add(SceneTextures.GBufferD);
+		}
+
+		if (bRequiresSceneDepthAux)
 		{
 			ColorTargets.Add(SceneTextures.DepthAux.Target);
 		}
diff --git a/Engine/Source/Runtime/Renderer/Private/PostProcess/SceneRenderTargets.h b/Engine/Source/Runtime/Renderer/Private/PostProcess/SceneRenderTargets.h
--- a/Engine/Source/Runtime/Renderer/Private/PostProcess/SceneRenderTargets.h
+++ b/Engine/Source/Runtime/Renderer/Private/PostProcess/SceneRenderTargets.h
@@ -42,10 +42,10 @@
 #define STENCIL_LIGHTING_CHANNELS_MASK(Value) uint8(((Value) & 0x7) << STENCIL_LIGHTING_CHANNELS_BIT_ID)
 
 // Mobile specific
-// Sky material mask
-#define STENCIL_MOBILE_SKY_MASK GET_STENCIL_BIT_MASK(SANDBOX,1)
-// Store shading model into stencil [1-3] bits
-#define GET_STENCIL_MOBILE_SM_MASK(Value) uint8(((Value) & 0x7) << 1)
+// Store shading model into stencil [1-2] bits
+#define GET_STENCIL_MOBILE_SM_MASK(Value) uint8(((Value) & 0x3) << 1)
+// Sky material mask - bit 3
+#define STENCIL_MOBILE_SKY_MASK uint8(1 << 3)
 
 class FSceneRenderTargets
 {
diff --git a/Engine/Source/Runtime/VulkanRHI/Private/VulkanRenderpass.cpp b/Engine/Source/Runtime/VulkanRHI/Private/VulkanRenderpass.cpp
--- a/Engine/Source/Runtime/VulkanRHI/Private/VulkanRenderpass.cpp
+++ b/Engine/Source/Runtime/VulkanRHI/Private/VulkanRenderpass.cpp
@@ -93,7 +93,7 @@ public:
 
 			check(RTLayout.GetDepthStencilAttachmentReference());
 
-			// depth as Input0
+			// Depth as Input0
 			InputAttachments1[0].SetAttachment(DepthStencilAttachmentOG, VK_IMAGE_ASPECT_DEPTH_BIT);
 			InputAttachments1[0].layout = VK_IMAGE_LAYOUT_DEPTH_STENCIL_READ_ONLY_OPTIMAL;
 
@@ -131,7 +131,7 @@ public:
 
 			//const VkAttachmentReference* ColorRef = RTLayout.GetColorAttachmentReferences();
 			//uint32 NumColorAttachments = RTLayout.GetNumColorAttachments();
-			check(RTLayout.GetNumColorAttachments() == 4); //current layout is SceneColor, GBufferA/B/C
+			//check(RTLayout.GetNumColorAttachments() == 5); //current layout is SceneColor, GBufferA/B/C/D
 
 			// 1. Write to SceneColor and GBuffer, input DepthStencil
 			{
@@ -148,7 +148,7 @@ public:
 				SubpassDesc.SetMultiViewMask(MultiviewMask);
 #endif
 				
-				// depth as Input0
+				// Depth as Input0
 				TSubpassDependencyClass& SubpassDep = SubpassDependencies[NumDependencies++];
 				SubpassDep.srcSubpass = 0;
 				SubpassDep.dstSubpass = 1;
@@ -165,23 +165,26 @@ public:
 				SubpassDesc.SetColorAttachments(ColorAttachmentReferences, 1); // SceneColor only
 				SubpassDesc.SetDepthStencilAttachment(&DepthStencilAttachment);
 
-				// GBuffer as Input2/3/4
+				// Depth as Input0
 				InputAttachments2[0].attachment = DepthStencilAttachment.attachment;
 				InputAttachments2[0].layout = DepthStencilAttachment.layout;
 				InputAttachments2[0].SetAspect(VK_IMAGE_ASPECT_DEPTH_BIT);
 
+				// SceneColor write only
 				InputAttachments2[1].attachment = VK_ATTACHMENT_UNUSED;
 				InputAttachments2[1].layout = VK_IMAGE_LAYOUT_UNDEFINED;
 				InputAttachments2[1].SetAspect(0);
-
-				for (int32 i = 2; i < 5; ++i)
+				
+				// GBufferA/B/C/D as Input2/3/4/5
+				int32 NumColorInputs = ColorAttachmentReferences.Num() - 1;
+				for (int32 i = 2; i < (NumColorInputs + 2); ++i)
 				{
 					InputAttachments2[i].attachment = ColorAttachmentReferences[i - 1].attachment;
 					InputAttachments2[i].layout = VK_IMAGE_LAYOUT_SHADER_READ_ONLY_OPTIMAL;
 					InputAttachments2[i].SetAspect(VK_IMAGE_ASPECT_COLOR_BIT);
 				}
 
-				SubpassDesc.SetInputAttachments(InputAttachments2, 5);
+				SubpassDesc.SetInputAttachments(InputAttachments2, NumColorInputs + 2);
 #if VULKAN_SUPPORTS_RENDERPASS2
 				if (bApplyFragmentShadingRate)
 				{
