Fix a bug of depthAux corruption on mobile opengl.

diff --git a/Engine/Source/Runtime/Renderer/Private/DecalRenderingCommon.cpp b/Engine/Source/Runtime/Renderer/Private/DecalRenderingCommon.cpp
--- a/Engine/Source/Runtime/Renderer/Private/DecalRenderingCommon.cpp
+++ b/Engine/Source/Runtime/Renderer/Private/DecalRenderingCommon.cpp
@@ -784,21 +784,21 @@ namespace DecalRendering
 				if (DecalBlendDesc.bWriteEmissive)
 				{
 					// Treat blend as emissive
-					return TStaticBlendState<CW_RGB, BO_Add, BF_SourceAlpha, BF_One>::GetRHI();
+					return TStaticBlendState<CW_RGB, BO_Add, BF_SourceAlpha, BF_One, BO_Add, BF_Zero, BF_One, CW_NONE>::GetRHI();
 				}
 				else
 				{
 					// Treat blend as non-emissive
-					return TStaticBlendState<CW_RGB, BO_Add, BF_SourceAlpha, BF_InverseSourceAlpha>::GetRHI();
+					return TStaticBlendState<CW_RGB, BO_Add, BF_SourceAlpha, BF_InverseSourceAlpha, BO_Add, BF_Zero, BF_One, CW_NONE>::GetRHI();
 				}
 			}
 			else if (DecalBlendDesc.BlendMode == BLEND_AlphaComposite)
 			{
-				return TStaticBlendState<CW_RGB, BO_Add, BF_One, BF_InverseSourceAlpha>::GetRHI();
+				return TStaticBlendState<CW_RGB, BO_Add, BF_One, BF_InverseSourceAlpha, BO_Add, BF_Zero, BF_One, CW_NONE>::GetRHI();
 			}
 			else if (DecalBlendDesc.BlendMode == BLEND_Modulate)
 			{
-				return TStaticBlendState<CW_RGB, BO_Add, BF_DestColor, BF_InverseSourceAlpha>::GetRHI();
+				return TStaticBlendState<CW_RGB, BO_Add, BF_DestColor, BF_InverseSourceAlpha, BO_Add, BF_Zero, BF_One, CW_NONE>::GetRHI();
 			}
 		}
 		else
diff --git a/Engine/Source/Runtime/Renderer/Private/MobileBasePass.cpp b/Engine/Source/Runtime/Renderer/Private/MobileBasePass.cpp
--- a/Engine/Source/Runtime/Renderer/Private/MobileBasePass.cpp
+++ b/Engine/Source/Runtime/Renderer/Private/MobileBasePass.cpp
@@ -388,7 +388,14 @@ void MobileBasePass::SetTranslucentRenderState(FMeshPassProcessorRenderState& Dr
 	if (Material.GetShadingModels().HasShadingModel(MSM_ThinTranslucent))
 	{
 		// the mobile thin translucent fallback uses a similar mode as BLEND_Translucent, but multiplies color by 1 insead of SrcAlpha.
-		DrawRenderState.SetBlendState(TStaticBlendState<CW_RGBA, BO_Add, BF_One, BF_InverseSourceAlpha, BO_Add, BF_Zero, BF_InverseSourceAlpha>::GetRHI());
+		DrawRenderState.SetBlendState(TStaticBlendState<CW_RGBA, BO_Add, BF_One, BF_InverseSourceAlpha, BO_Add, BF_Zero, BF_InverseSourceAlpha,
+														CW_NONE, BO_Add, BF_One, BF_Zero, BO_Add, BF_One, BF_Zero,
+														CW_NONE, BO_Add, BF_One, BF_Zero, BO_Add, BF_One, BF_Zero,
+														CW_NONE, BO_Add, BF_One, BF_Zero, BO_Add, BF_One, BF_Zero,
+														CW_NONE, BO_Add, BF_One, BF_Zero, BO_Add, BF_One, BF_Zero,
+														CW_NONE, BO_Add, BF_One, BF_Zero, BO_Add, BF_One, BF_Zero,
+														CW_NONE, BO_Add, BF_One, BF_Zero, BO_Add, BF_One, BF_Zero,
+														CW_NONE, BO_Add, BF_One, BF_Zero, BO_Add, BF_One, BF_Zero>::GetRHI());
 	}
 	else
 	{
@@ -397,34 +404,83 @@ void MobileBasePass::SetTranslucentRenderState(FMeshPassProcessorRenderState& Dr
 		case BLEND_Translucent:
 			if (Material.ShouldWriteOnlyAlpha())
 			{
-				DrawRenderState.SetBlendState(TStaticBlendState<CW_ALPHA, BO_Add, BF_Zero, BF_Zero, BO_Add, BF_One, BF_Zero>::GetRHI());
+				DrawRenderState.SetBlendState(TStaticBlendState<CW_ALPHA, BO_Add, BF_Zero, BF_Zero, BO_Add, BF_One, BF_Zero,
+																CW_NONE, BO_Add, BF_One, BF_Zero, BO_Add, BF_One, BF_Zero,
+																CW_NONE, BO_Add, BF_One, BF_Zero, BO_Add, BF_One, BF_Zero,
+																CW_NONE, BO_Add, BF_One, BF_Zero, BO_Add, BF_One, BF_Zero,
+																CW_NONE, BO_Add, BF_One, BF_Zero, BO_Add, BF_One, BF_Zero,
+																CW_NONE, BO_Add, BF_One, BF_Zero, BO_Add, BF_One, BF_Zero,
+																CW_NONE, BO_Add, BF_One, BF_Zero, BO_Add, BF_One, BF_Zero,
+																CW_NONE, BO_Add, BF_One, BF_Zero, BO_Add, BF_One, BF_Zero>::GetRHI());
 			}
 			else
 			{
-				DrawRenderState.SetBlendState(TStaticBlendState<CW_RGBA, BO_Add, BF_SourceAlpha, BF_InverseSourceAlpha, BO_Add, BF_Zero, BF_InverseSourceAlpha>::GetRHI());
+				DrawRenderState.SetBlendState(TStaticBlendState<CW_RGBA, BO_Add, BF_SourceAlpha, BF_InverseSourceAlpha, BO_Add, BF_Zero, BF_InverseSourceAlpha,
+																CW_NONE, BO_Add, BF_One, BF_Zero, BO_Add, BF_One, BF_Zero,
+																CW_NONE, BO_Add, BF_One, BF_Zero, BO_Add, BF_One, BF_Zero,
+																CW_NONE, BO_Add, BF_One, BF_Zero, BO_Add, BF_One, BF_Zero,
+																CW_NONE, BO_Add, BF_One, BF_Zero, BO_Add, BF_One, BF_Zero,
+																CW_NONE, BO_Add, BF_One, BF_Zero, BO_Add, BF_One, BF_Zero,
+																CW_NONE, BO_Add, BF_One, BF_Zero, BO_Add, BF_One, BF_Zero,
+																CW_NONE, BO_Add, BF_One, BF_Zero, BO_Add, BF_One, BF_Zero>::GetRHI());
 			}
 			break;
 		case BLEND_Additive:
 			// Add to the existing scene color
-			DrawRenderState.SetBlendState(TStaticBlendState<CW_RGBA, BO_Add, BF_One, BF_One, BO_Add, BF_Zero, BF_InverseSourceAlpha>::GetRHI());
+			DrawRenderState.SetBlendState(TStaticBlendState<CW_RGBA, BO_Add, BF_One, BF_One, BO_Add, BF_Zero, BF_InverseSourceAlpha,
+															CW_NONE, BO_Add, BF_One, BF_Zero, BO_Add, BF_One, BF_Zero,
+															CW_NONE, BO_Add, BF_One, BF_Zero, BO_Add, BF_One, BF_Zero,
+															CW_NONE, BO_Add, BF_One, BF_Zero, BO_Add, BF_One, BF_Zero,
+															CW_NONE, BO_Add, BF_One, BF_Zero, BO_Add, BF_One, BF_Zero,
+															CW_NONE, BO_Add, BF_One, BF_Zero, BO_Add, BF_One, BF_Zero,
+															CW_NONE, BO_Add, BF_One, BF_Zero, BO_Add, BF_One, BF_Zero,
+															CW_NONE, BO_Add, BF_One, BF_Zero, BO_Add, BF_One, BF_Zero>::GetRHI());
 			break;
 		case BLEND_Modulate:
 			// Modulate with the existing scene color, preserve destination alpha.
-			DrawRenderState.SetBlendState(TStaticBlendState<CW_RGB, BO_Add, BF_DestColor, BF_Zero>::GetRHI());
+			DrawRenderState.SetBlendState(TStaticBlendState<CW_RGB, BO_Add, BF_DestColor, BF_Zero, BO_Add, BF_Zero, BF_One,
+															CW_NONE, BO_Add, BF_One, BF_Zero, BO_Add, BF_One, BF_Zero,
+															CW_NONE, BO_Add, BF_One, BF_Zero, BO_Add, BF_One, BF_Zero,
+															CW_NONE, BO_Add, BF_One, BF_Zero, BO_Add, BF_One, BF_Zero,
+															CW_NONE, BO_Add, BF_One, BF_Zero, BO_Add, BF_One, BF_Zero,
+															CW_NONE, BO_Add, BF_One, BF_Zero, BO_Add, BF_One, BF_Zero,
+															CW_NONE, BO_Add, BF_One, BF_Zero, BO_Add, BF_One, BF_Zero,
+															CW_NONE, BO_Add, BF_One, BF_Zero, BO_Add, BF_One, BF_Zero>::GetRHI());
 			break;
 		case BLEND_AlphaComposite:
 			// Blend with existing scene color. New color is already pre-multiplied by alpha.
-			DrawRenderState.SetBlendState(TStaticBlendState<CW_RGBA, BO_Add, BF_One, BF_InverseSourceAlpha, BO_Add, BF_Zero, BF_InverseSourceAlpha>::GetRHI());
+			DrawRenderState.SetBlendState(TStaticBlendState<CW_RGBA, BO_Add, BF_One, BF_InverseSourceAlpha, BO_Add, BF_Zero, BF_InverseSourceAlpha,
+															CW_NONE, BO_Add, BF_One, BF_Zero, BO_Add, BF_One, BF_Zero,
+															CW_NONE, BO_Add, BF_One, BF_Zero, BO_Add, BF_One, BF_Zero,
+															CW_NONE, BO_Add, BF_One, BF_Zero, BO_Add, BF_One, BF_Zero,
+															CW_NONE, BO_Add, BF_One, BF_Zero, BO_Add, BF_One, BF_Zero,
+															CW_NONE, BO_Add, BF_One, BF_Zero, BO_Add, BF_One, BF_Zero,
+															CW_NONE, BO_Add, BF_One, BF_Zero, BO_Add, BF_One, BF_Zero,
+															CW_NONE, BO_Add, BF_One, BF_Zero, BO_Add, BF_One, BF_Zero>::GetRHI());
 			break;
 		case BLEND_AlphaHoldout:
 			// Blend by holding out the matte shape of the source alpha
-			DrawRenderState.SetBlendState(TStaticBlendState<CW_RGBA, BO_Add, BF_Zero, BF_InverseSourceAlpha, BO_Add, BF_One, BF_InverseSourceAlpha>::GetRHI());
+			DrawRenderState.SetBlendState(TStaticBlendState<CW_RGBA, BO_Add, BF_Zero, BF_InverseSourceAlpha, BO_Add, BF_One, BF_InverseSourceAlpha,
+															CW_NONE, BO_Add, BF_One, BF_Zero, BO_Add, BF_One, BF_Zero,
+															CW_NONE, BO_Add, BF_One, BF_Zero, BO_Add, BF_One, BF_Zero,
+															CW_NONE, BO_Add, BF_One, BF_Zero, BO_Add, BF_One, BF_Zero,
+															CW_NONE, BO_Add, BF_One, BF_Zero, BO_Add, BF_One, BF_Zero,
+															CW_NONE, BO_Add, BF_One, BF_Zero, BO_Add, BF_One, BF_Zero,
+															CW_NONE, BO_Add, BF_One, BF_Zero, BO_Add, BF_One, BF_Zero,
+															CW_NONE, BO_Add, BF_One, BF_Zero, BO_Add, BF_One, BF_Zero>::GetRHI());
 			break;
 		default:
 			if (Material.GetShadingModels().HasShadingModel(MSM_SingleLayerWater))
 			{
 				// Single layer water is an opaque marerial rendered as translucent on Mobile. We force pre-multiplied alpha to achieve water depth based transmittance.
-				DrawRenderState.SetBlendState(TStaticBlendState<CW_RGBA, BO_Add, BF_One, BF_InverseSourceAlpha, BO_Add, BF_Zero, BF_InverseSourceAlpha>::GetRHI());
+				DrawRenderState.SetBlendState(TStaticBlendState<CW_RGBA, BO_Add, BF_One, BF_InverseSourceAlpha, BO_Add, BF_Zero, BF_InverseSourceAlpha,
+																CW_NONE, BO_Add, BF_One, BF_Zero, BO_Add, BF_One, BF_Zero,
+																CW_NONE, BO_Add, BF_One, BF_Zero, BO_Add, BF_One, BF_Zero,
+																CW_NONE, BO_Add, BF_One, BF_Zero, BO_Add, BF_One, BF_Zero,
+																CW_NONE, BO_Add, BF_One, BF_Zero, BO_Add, BF_One, BF_Zero,
+																CW_NONE, BO_Add, BF_One, BF_Zero, BO_Add, BF_One, BF_Zero,
+																CW_NONE, BO_Add, BF_One, BF_Zero, BO_Add, BF_One, BF_Zero,
+																CW_NONE, BO_Add, BF_One, BF_Zero, BO_Add, BF_One, BF_Zero>::GetRHI());
 			}
 			else
 			{
diff --git a/Engine/Source/Runtime/Renderer/Private/MobileFogRendering.cpp b/Engine/Source/Runtime/Renderer/Private/MobileFogRendering.cpp
--- a/Engine/Source/Runtime/Renderer/Private/MobileFogRendering.cpp
+++ b/Engine/Source/Runtime/Renderer/Private/MobileFogRendering.cpp
@@ -145,7 +145,14 @@ void FMobileSceneRenderer::RenderFog(FRHICommandListImmediate& RHICmdList, const
 		GraphicsPSOInit.DepthStencilState = TStaticDepthStencilState<false, CF_Always>::GetRHI();
 	}
 	
-	GraphicsPSOInit.BlendState = TStaticBlendState<CW_RGB, BO_Add, BF_One, BF_SourceAlpha>::GetRHI();
+	GraphicsPSOInit.BlendState = TStaticBlendState<CW_RGB, BO_Add, BF_One, BF_SourceAlpha, BO_Add, BF_Zero, BF_One,
+													CW_NONE, BO_Add, BF_One, BF_Zero, BO_Add, BF_One, BF_Zero,
+													CW_NONE, BO_Add, BF_One, BF_Zero, BO_Add, BF_One, BF_Zero,
+													CW_NONE, BO_Add, BF_One, BF_Zero, BO_Add, BF_One, BF_Zero,
+													CW_NONE, BO_Add, BF_One, BF_Zero, BO_Add, BF_One, BF_Zero,
+													CW_NONE, BO_Add, BF_One, BF_Zero, BO_Add, BF_One, BF_Zero,
+													CW_NONE, BO_Add, BF_One, BF_Zero, BO_Add, BF_One, BF_Zero,
+													CW_NONE, BO_Add, BF_One, BF_Zero, BO_Add, BF_One, BF_Zero>::GetRHI();
 		
 	TShaderMapRef<FMobileFogVS> VertexShader(View.ShaderMap);
 
diff --git a/Engine/Source/Runtime/Renderer/Private/MobileShadingRenderer.cpp b/Engine/Source/Runtime/Renderer/Private/MobileShadingRenderer.cpp
--- a/Engine/Source/Runtime/Renderer/Private/MobileShadingRenderer.cpp
+++ b/Engine/Source/Runtime/Renderer/Private/MobileShadingRenderer.cpp
@@ -1673,7 +1673,7 @@ void FMobileSceneRenderer::PreTonemapMSAA(FRHICommandListImmediate& RHICmdList,
 
 	FGraphicsPipelineStateInitializer GraphicsPSOInit;
 	RHICmdList.ApplyCachedRenderTargets(GraphicsPSOInit);
-	GraphicsPSOInit.BlendState = TStaticBlendState<>::GetRHI();
+	GraphicsPSOInit.BlendState = TStaticBlendState<CW_RGBA, BO_Add, BF_One, BF_Zero, BO_Add, BF_One, BF_Zero, CW_NONE>::GetRHI();
 	GraphicsPSOInit.RasterizerState = TStaticRasterizerState<>::GetRHI();
 	GraphicsPSOInit.DepthStencilState = TStaticDepthStencilState<false, CF_Always>::GetRHI();
 
diff --git a/Engine/Source/Runtime/Renderer/Private/ShadowRendering.cpp b/Engine/Source/Runtime/Renderer/Private/ShadowRendering.cpp
--- a/Engine/Source/Runtime/Renderer/Private/ShadowRendering.cpp
+++ b/Engine/Source/Runtime/Renderer/Private/ShadowRendering.cpp
@@ -732,7 +732,7 @@ FRHIBlendState* FProjectedShadowInfo::GetBlendStateForProjection(
 			if (bMobileModulatedProjections)
 			{
 				// Color modulate shadows, ignore alpha.
-				BlendState = TStaticBlendState<CW_RGB, BO_Add, BF_Zero, BF_SourceColor, BO_Add, BF_Zero, BF_One>::GetRHI();
+				BlendState = TStaticBlendState<CW_RGB, BO_Add, BF_Zero, BF_SourceColor, BO_Add, BF_Zero, BF_One, CW_NONE>::GetRHI();
 			}
 			else
 			{
@@ -894,7 +894,7 @@ void FProjectedShadowInfo::SetupProjectionStencilMask(
 
 	// Depth test wo/ writes, no color writing.
 	DrawRenderState.SetDepthStencilState(TStaticDepthStencilState<false, CF_DepthNearOrEqual>::GetRHI());
-	DrawRenderState.SetBlendState(TStaticBlendState<CW_NONE>::GetRHI());
+	DrawRenderState.SetBlendState(TStaticBlendState<CW_NONE, BO_Add, BF_Zero, BF_SourceColor, BO_Add, BF_Zero, BF_One, CW_NONE>::GetRHI());
 
 	// If this is a preshadow, mask the projection by the receiver primitives.
 	if (bPreShadow || bSelfShadowOnly)
