Fix crash when building lighting with ISR (UE-148100).

diff --git a/Engine/Source/Runtime/Renderer/Private/Renderer.cpp b/Engine/Source/Runtime/Renderer/Private/Renderer.cpp
--- a/Engine/Source/Runtime/Renderer/Private/Renderer.cpp
+++ b/Engine/Source/Runtime/Renderer/Private/Renderer.cpp
@@ -109,7 +109,7 @@ void FRendererModule::InitializeSystemTextures(FRHICommandListImmediate& RHICmdL
 }
 
 BEGIN_SHADER_PARAMETER_STRUCT(FDrawTileMeshPassParameters, )
-	SHADER_PARAMETER_STRUCT_REF(FViewUniformShaderParameters, View)
+	SHADER_PARAMETER_STRUCT_INCLUDE(FViewShaderParameters, View)
 	SHADER_PARAMETER_RDG_UNIFORM_BUFFER(FInstanceCullingGlobalUniforms, InstanceCulling)
 	SHADER_PARAMETER_STRUCT_REF(FReflectionCaptureShaderData, ReflectionCapture)
 	SHADER_PARAMETER_RDG_UNIFORM_BUFFER(FDebugViewModePassUniformParameters, DebugViewMode)
@@ -235,7 +235,7 @@ void FRendererModule::DrawTileMesh(FCanvasRenderContext& RenderContext, FMeshPas
 
 		auto* PassParameters = GraphBuilder.AllocParameters<FDrawTileMeshPassParameters>();
 		PassParameters->RenderTargets[0] = FRenderTargetBinding(RenderContext.GetRenderTarget(), ERenderTargetLoadAction::ELoad);
-		PassParameters->View = View.ViewUniformBuffer;
+		PassParameters->View = View.GetShaderParameters();
 		PassParameters->ReflectionCapture = EmptyReflectionCaptureUniformBuffer;
 		PassParameters->InstanceCulling = FInstanceCullingContext::CreateDummyInstanceCullingUniformBuffer(GraphBuilder);
 
