Remove obsolete mobile base pass parameters

diff --git a/Engine/Source/Runtime/Core/Private/HAL/ConsoleManager.cpp b/Engine/Source/Runtime/Core/Private/HAL/ConsoleManager.cpp
--- a/Engine/Source/Runtime/Core/Private/HAL/ConsoleManager.cpp
+++ b/Engine/Source/Runtime/Core/Private/HAL/ConsoleManager.cpp
@@ -2591,12 +2591,6 @@ static TAutoConsoleVariable<int32> CVarMobileUseCSMShaderBranch(
 		 "1: Use a single shader pemutation with a branch in a shader to apply CSM (only with r.AllowStaticLighting=0)"),
 	ECVF_RenderThreadSafe | ECVF_ReadOnly);
 
-static TAutoConsoleVariable<float> CVarsCSMDebugHint(
-	TEXT("r.Mobile.Shadow.CSMDebugHint"),
-	0.0f,
-	TEXT(""),
-	ECVF_RenderThreadSafe | ECVF_ReadOnly);
-
 static TAutoConsoleVariable<int32> CVarMobileAllowDistanceFieldShadows(
 	TEXT("r.Mobile.AllowDistanceFieldShadows"),
 	1,
diff --git a/Engine/Source/Runtime/Renderer/Private/MobileBasePass.cpp b/Engine/Source/Runtime/Renderer/Private/MobileBasePass.cpp
--- a/Engine/Source/Runtime/Renderer/Private/MobileBasePass.cpp
+++ b/Engine/Source/Runtime/Renderer/Private/MobileBasePass.cpp
@@ -539,67 +539,7 @@ void TMobileBasePassPSPolicyParamType<FUniformLightMapPolicy>::GetShaderBindings
 
 	if (Scene)
 	{
-		// test for HQ reflection parameter existence
-		if (HQReflectionCubemaps[0].IsBound() || HQReflectionCubemaps[1].IsBound() || HQReflectionCubemaps[2].IsBound())
-		{
-			static const int32 MaxNumReflections = FPrimitiveSceneInfo::MaxCachedReflectionCaptureProxies;
-			static_assert(MaxNumReflections == 3, "Update reflection array initializations to match MaxCachedReflectionCaptureProxies");
-			// set reflection parameters
-			FTexture* ReflectionCubemapTextures[MaxNumReflections] = { GBlackTextureCube, GBlackTextureCube, GBlackTextureCube };
-			FVector4f CapturePositions[MaxNumReflections] = { FVector4f(0, 0, 0, 0), FVector4f(0, 0, 0, 0), FVector4f(0, 0, 0, 0) };
-			FVector4f CaptureTilePositions[MaxNumReflections] = { FVector4f(0, 0, 0, 0), FVector4f(0, 0, 0, 0), FVector4f(0, 0, 0, 0) };
-			FVector4f ReflectionParams(0.0f, 0.0f, 0.0f, 0.0f);
-			FVector4f ReflectanceMaxValueRGBMParams(0.0f, 0.0f, 0.0f, 0.0f);
-			FMatrix44f CaptureBoxTransformArray[MaxNumReflections] = { FMatrix44f(EForceInit::ForceInitToZero), FMatrix44f(EForceInit::ForceInitToZero), FMatrix44f(EForceInit::ForceInitToZero) };
-			FVector4f CaptureBoxScalesArray[MaxNumReflections] = { FVector4f(EForceInit::ForceInitToZero), FVector4f(EForceInit::ForceInitToZero), FVector4f(EForceInit::ForceInitToZero) };
-			FPrimitiveSceneInfo* PrimitiveSceneInfo = PrimitiveSceneProxy ? PrimitiveSceneProxy->GetPrimitiveSceneInfo() : nullptr;
-			if (PrimitiveSceneInfo)
-			{
-				for (int32 i = 0; i < MaxNumReflections; i++)
-				{
-					const FReflectionCaptureProxy* ReflectionProxy = PrimitiveSceneInfo->CachedReflectionCaptureProxies[i];
-					if (ReflectionProxy)
-					{
-						CapturePositions[i] = ReflectionProxy->RelativePosition;
-						CapturePositions[i].W = ReflectionProxy->InfluenceRadius;
-						CaptureTilePositions[i] = FVector4f(ReflectionProxy->TilePosition, 0);
-						if (ReflectionProxy->EncodedHDRCubemap)
-						{
-							ReflectionCubemapTextures[i] = ReflectionProxy->EncodedHDRCubemap->GetResource();
-						}
-						//To keep ImageBasedReflectionLighting coherence with PC, use AverageBrightness instead of InvAverageBrightness to calculate the IBL contribution
-						ReflectionParams[i] = ReflectionProxy->EncodedHDRAverageBrightness;
-
-						ReflectanceMaxValueRGBMParams[i] = ReflectionProxy->MaxValueRGBM;
-						if (ReflectionProxy->Shape == EReflectionCaptureShape::Box)
-						{
-							CaptureBoxTransformArray[i] = ReflectionProxy->BoxTransform;
-							CaptureBoxScalesArray[i] = FVector4f(ReflectionProxy->BoxScales, ReflectionProxy->BoxTransitionDistance);
-						}
-					}
-					else if (Scene->SkyLight != nullptr && Scene->SkyLight->ProcessedTexture != nullptr)
-					{
-						// NegativeInfluence to signal the shader we are defaulting to SkyLight if there are no ReflectionComponents in the Level
-						CapturePositions[i].W = -1.0f;
-						ReflectionCubemapTextures[i] = Scene->SkyLight->ProcessedTexture;
-						ReflectionParams[3] = FMath::FloorLog2(Scene->SkyLight->ProcessedTexture->GetSizeX());
-						break;
-					}
-				}
-			}
-
-			for (int32 i = 0; i < MaxNumReflections; i++)
-			{
-				ShaderBindings.AddTexture(HQReflectionCubemaps[i], HQReflectionSamplers[i], ReflectionCubemapTextures[i]->SamplerStateRHI, ReflectionCubemapTextures[i]->TextureRHI);
-			}
-			ShaderBindings.Add(HQReflectionInvAverageBrigtnessParams, ReflectionParams);
-			ShaderBindings.Add(HQReflectanceMaxValueRGBMParams, ReflectanceMaxValueRGBMParams);
-			ShaderBindings.Add(HQReflectionPositionsAndRadii, CapturePositions);
-			ShaderBindings.Add(HQReflectionTilePositions, CaptureTilePositions);
-			ShaderBindings.Add(HQReflectionCaptureBoxTransformArray, CaptureBoxTransformArray);
-			ShaderBindings.Add(HQReflectionCaptureBoxScalesArray, CaptureBoxScalesArray);
-		}
-		else if (ReflectionParameter.IsBound())
+		if (ReflectionParameter.IsBound())
 		{
 			FRHIUniformBuffer* ReflectionUB = GDefaultMobileReflectionCaptureUniformBuffer.GetUniformBufferRHI();
 			FPrimitiveSceneInfo* PrimitiveSceneInfo = PrimitiveSceneProxy ? PrimitiveSceneProxy->GetPrimitiveSceneInfo() : nullptr;
@@ -627,13 +567,6 @@ void TMobileBasePassPSPolicyParamType<FUniformLightMapPolicy>::GetShaderBindings
 		ShaderBindings.Add(MobileDirectionLightBufferParam, Scene->UniformBuffers.MobileDirectionalLightUniformBuffers[UniformBufferIndex]);
 	}
 
-	if (CSMDebugHintParams.IsBound())
-	{
-		static const auto CVarsCSMDebugHint = IConsoleManager::Get().FindTConsoleVariableDataFloat(TEXT("r.Mobile.Shadow.CSMDebugHint"));
-		float CSMDebugValue = CVarsCSMDebugHint->GetValueOnRenderThread();
-		ShaderBindings.Add(CSMDebugHintParams, CSMDebugValue);
-	}
-
 	if (UseCSMParameter.IsBound())
 	{
 		ShaderBindings.Add(UseCSMParameter, ShaderElementData.bCanReceiveCSM ? 1 : 0);
diff --git a/Engine/Source/Runtime/Renderer/Private/MobileBasePassRendering.h b/Engine/Source/Runtime/Renderer/Private/MobileBasePassRendering.h
--- a/Engine/Source/Runtime/Renderer/Private/MobileBasePassRendering.h
+++ b/Engine/Source/Runtime/Renderer/Private/MobileBasePassRendering.h
@@ -248,21 +248,7 @@ public:
 		
 		MobileDirectionLightBufferParam.Bind(Initializer.ParameterMap, FMobileDirectionalLightShaderParameters::StaticStructMetadata.GetShaderVariableName());
 		ReflectionParameter.Bind(Initializer.ParameterMap, FMobileReflectionCaptureShaderParameters::StaticStructMetadata.GetShaderVariableName());
-
-		HQReflectionCubemaps[0].Bind(Initializer.ParameterMap, TEXT("ReflectionCubemap0"));
-		HQReflectionSamplers[0].Bind(Initializer.ParameterMap, TEXT("ReflectionCubemapSampler0"));
-		HQReflectionCubemaps[1].Bind(Initializer.ParameterMap, TEXT("ReflectionCubemap1"));
-		HQReflectionSamplers[1].Bind(Initializer.ParameterMap, TEXT("ReflectionCubemapSampler1"));
-		HQReflectionCubemaps[2].Bind(Initializer.ParameterMap, TEXT("ReflectionCubemap2"));
-		HQReflectionSamplers[2].Bind(Initializer.ParameterMap, TEXT("ReflectionCubemapSampler2"));
-		HQReflectionInvAverageBrigtnessParams.Bind(Initializer.ParameterMap, TEXT("ReflectionAverageBrigtness"));
-		HQReflectanceMaxValueRGBMParams.Bind(Initializer.ParameterMap, TEXT("ReflectanceMaxValueRGBM"));
-		HQReflectionPositionsAndRadii.Bind(Initializer.ParameterMap, TEXT("ReflectionPositionsAndRadii"));
-		HQReflectionTilePositions.Bind(Initializer.ParameterMap, TEXT("ReflectionTilePositions"));
-		HQReflectionCaptureBoxTransformArray.Bind(Initializer.ParameterMap, TEXT("CaptureBoxTransformArray"));
-		HQReflectionCaptureBoxScalesArray.Bind(Initializer.ParameterMap, TEXT("CaptureBoxScalesArray"));
-						
-		CSMDebugHintParams.Bind(Initializer.ParameterMap, TEXT("CSMDebugHint"));
+				
 		UseCSMParameter.Bind(Initializer.ParameterMap, TEXT("UseCSM"));
 	}
 
@@ -271,18 +257,6 @@ public:
 private:
 	LAYOUT_FIELD(FShaderUniformBufferParameter, MobileDirectionLightBufferParam);
 	LAYOUT_FIELD(FShaderUniformBufferParameter, ReflectionParameter);
-
-	// HQ reflection bound as loose params
-	LAYOUT_ARRAY(FShaderResourceParameter, HQReflectionCubemaps, 3);
-	LAYOUT_ARRAY(FShaderResourceParameter, HQReflectionSamplers, 3);
-	LAYOUT_FIELD(FShaderParameter, HQReflectionInvAverageBrigtnessParams);
-	LAYOUT_FIELD(FShaderParameter, HQReflectanceMaxValueRGBMParams);
-	LAYOUT_FIELD(FShaderParameter, HQReflectionPositionsAndRadii);
-	LAYOUT_FIELD(FShaderParameter, HQReflectionTilePositions);
-	LAYOUT_FIELD(FShaderParameter, HQReflectionCaptureBoxTransformArray);
-	LAYOUT_FIELD(FShaderParameter, HQReflectionCaptureBoxScalesArray);
-
-	LAYOUT_FIELD(FShaderParameter, CSMDebugHintParams);
 	LAYOUT_FIELD(FShaderParameter, UseCSMParameter);
 	
 public:
