Fixed: Rendering artifacts on a specific mobile devices with deferred shading enabled

diff --git a/Engine/Source/Runtime/Renderer/Private/MobileShadingRenderer.cpp b/Engine/Source/Runtime/Renderer/Private/MobileShadingRenderer.cpp
--- a/Engine/Source/Runtime/Renderer/Private/MobileShadingRenderer.cpp
+++ b/Engine/Source/Runtime/Renderer/Private/MobileShadingRenderer.cpp
@@ -216,6 +216,13 @@ static void SetupGBufferFlags(FSceneTexturesConfig& SceneTexturesConfig, bool bR
 	SceneTexturesConfig.GBufferC.Flags|= AddFlags;
 	SceneTexturesConfig.GBufferD.Flags|= AddFlags;
 	SceneTexturesConfig.GBufferE.Flags|= AddFlags;
+	
+	// Mobile uses FBF/subpassLoad to fetch data from GBuffer, and FBF does not always work with sRGB targets 
+	SceneTexturesConfig.GBufferA.Flags &= (~TexCreate_SRGB);
+	SceneTexturesConfig.GBufferB.Flags &= (~TexCreate_SRGB);
+	SceneTexturesConfig.GBufferC.Flags &= (~TexCreate_SRGB);
+	SceneTexturesConfig.GBufferD.Flags &= (~TexCreate_SRGB);
+	SceneTexturesConfig.GBufferE.Flags &= (~TexCreate_SRGB);
 }
 
 FMobileSceneRenderer::FMobileSceneRenderer(TArrayView<const FSceneViewFamily*> InViewFamilies, FHitProxyConsumer* HitProxyConsumer)
@@ -486,7 +493,7 @@ void FMobileSceneRenderer::InitViews(FRDGBuilder& GraphBuilder, FSceneTexturesCo
 	
 	if (bDeferredShading) 
 	{
-		SetupGBufferFlags(SceneTexturesConfig, bRequiresMultiPass);
+		SetupGBufferFlags(SceneTexturesConfig, bRequiresMultiPass || GraphBuilder.IsDumpingFrame());
 	}
 
 	// Update the pixel projected reflection extent according to the settings in the PlanarReflectionComponent.
diff --git a/Engine/Source/Runtime/Renderer/Private/SceneTextures.cpp b/Engine/Source/Runtime/Renderer/Private/SceneTextures.cpp
--- a/Engine/Source/Runtime/Renderer/Private/SceneTextures.cpp
+++ b/Engine/Source/Runtime/Renderer/Private/SceneTextures.cpp
@@ -772,16 +772,6 @@ void FSceneTextures::InitializeViewFamily(FRDGBuilder& GraphBuilder, FViewFamily
 	{
 		ETextureCreateFlags FlagsToAdd = TexCreate_None;
 		
-		if (Config.ShadingPath == EShadingPath::Mobile)
-		{
-			FlagsToAdd |= TexCreate_InputAttachmentRead;
-
-			if (!Config.bKeepDepthContent)
-			{
-				FlagsToAdd |= TexCreate_Memoryless;
-			}
-		}
-		
 		if (Config.GBufferA.Index >= 0)
 		{
 			const FRDGTextureDesc Desc(FRDGTextureDesc::Create2D(Config.Extent, Config.GBufferA.Format, FClearValueBinding::Transparent, Config.GBufferA.Flags | FlagsToAdd | GFastVRamConfig.GBufferA));
