Fix a bug of rendering decal on mac mobile preview

diff --git a/Engine/Source/Runtime/Renderer/Private/MobileShadingRenderer.cpp b/Engine/Source/Runtime/Renderer/Private/MobileShadingRenderer.cpp
--- a/Engine/Source/Runtime/Renderer/Private/MobileShadingRenderer.cpp
+++ b/Engine/Source/Runtime/Renderer/Private/MobileShadingRenderer.cpp
@@ -1575,7 +1575,7 @@ bool FMobileSceneRenderer::RequiresMultiPass(FRHICommandListImmediate& RHICmdLis
 	}
 
 	// All iOS support frame_buffer_fetch
-	if (IsMetalMobilePlatform(ShaderPlatform))
+	if (IsMetalMobilePlatform(ShaderPlatform) && GSupportsShaderFramebufferFetch)
 	{
 		return false;
 	}
