Fix a bug of writing stencil in PPMaterial on mobile.

#p4v-preflight-copy 19900274

diff --git a/Engine/Source/Runtime/Renderer/Private/PostProcess/PostProcessMaterial.cpp b/Engine/Source/Runtime/Renderer/Private/PostProcess/PostProcessMaterial.cpp
--- a/Engine/Source/Runtime/Renderer/Private/PostProcess/PostProcessMaterial.cpp
+++ b/Engine/Source/Runtime/Renderer/Private/PostProcess/PostProcessMaterial.cpp
@@ -660,7 +660,8 @@ FScreenPassTexture AddPostProcessMaterialPass(
 		View,
 		OutputViewport,
 		SceneColorViewport,
-		FScreenPassPipelineState(VertexShader, PixelShader, BlendState, DepthStencilState, MaterialStencilRef),
+		// Uses default depth stencil on mobile since the stencil test is done in pixel shader.
+		FScreenPassPipelineState(VertexShader, PixelShader, BlendState, bMobilePlatform ? DefaultDepthStencilState : DepthStencilState, MaterialStencilRef),
 		PostProcessMaterialParameters,
 		ScreenPassFlags,
 		[&View, VertexShader, PixelShader, MaterialRenderProxy, Material, PostProcessMaterialParameters, MaterialStencilRef](FRHICommandList& RHICmdList)
