Hack fix to bypass issue with recent view family refactor.

diff --git a/Engine/Source/Runtime/Engine/Public/SceneView.h b/Engine/Source/Runtime/Engine/Public/SceneView.h
--- a/Engine/Source/Runtime/Engine/Public/SceneView.h
+++ b/Engine/Source/Runtime/Engine/Public/SceneView.h
@@ -1749,6 +1749,8 @@ public:
 	*/
 	bool bIsFirstViewInMultipleViewFamily = true;
 
+	bool bIsSceneTexturesInitialized = false;
+
 	/** 
 	 * Which component of the scene rendering should be output to the final render target.
 	 * If SCS_FinalColorLDR this indicates do nothing.
diff --git a/Engine/Source/Runtime/Renderer/Private/SceneRendering.h b/Engine/Source/Runtime/Renderer/Private/SceneRendering.h
--- a/Engine/Source/Runtime/Renderer/Private/SceneRendering.h
+++ b/Engine/Source/Runtime/Renderer/Private/SceneRendering.h
@@ -1874,8 +1874,6 @@ public:
 
 	bool bShadowDepthRenderCompleted;
 
-	bool bIsSceneTexturesInitialized = false;
-
 	/** Lights added if wholescenepointlight shadow would have been rendered (ignoring r.SupportPointLightWholeSceneShadows). Used for warning about unsupported features. */	
 	TArray<FString, SceneRenderingAllocator> UsedWholeScenePointLightNames;
 
