Force release of pages from the LRU virtual texture page pool if we are requesting a locked page.
This reduces case where new virtual texture producer can't map the root page.
Also only update adaptive VT indirection texture when the allocated VT root page is mapped.

diff --git a/Engine/Source/Runtime/Renderer/Private/VT/AdaptiveVirtualTexture.cpp b/Engine/Source/Runtime/Renderer/Private/VT/AdaptiveVirtualTexture.cpp
--- a/Engine/Source/Runtime/Renderer/Private/VT/AdaptiveVirtualTexture.cpp
+++ b/Engine/Source/Runtime/Renderer/Private/VT/AdaptiveVirtualTexture.cpp
@@ -32,7 +32,7 @@ static TAutoConsoleVariable<int32> CVarAVTMaxPageResidency(
 
 static TAutoConsoleVariable<int32> CVarAVTAgeToFree(
 	TEXT("r.VT.AVT.AgeToFree"),
-	300,
+	60,
 	TEXT("Number of frames for an allocation to be unused before it is considered for free"),
 	ECVF_RenderThreadSafe
 );
@@ -546,6 +546,16 @@ void FAdaptiveVirtualTexture::Allocate(FVirtualTextureSystem* InSystem, uint32 I
 		// Mark allocation as used.
 		const uint32 Key = (InFrame << 4) | InNewLevel;
 		LRUHeap.Update(Key, InAllocationIndex);
+
+		// Queue indirection texture update unless this allocation slot is already marked as pending.
+		if (SlotsPendingRootPageMap.Find(InAllocationIndex) == INDEX_NONE)
+		{
+			const uint32 vAddress = NewAllocatedVT->GetVirtualAddress();
+			const uint32 vAddressX = FMath::ReverseMortonCode2(vAddress);
+			const uint32 vAddressY = FMath::ReverseMortonCode2(vAddress >> 1);
+			const uint32 PackedIndirectionValue = (1 << 28) | (InNewLevel << 24) | (vAddressY << 12) | vAddressX;
+			TextureUpdates.Add(FIndirectionTextureUpdate{ X, Y, PackedIndirectionValue });
+		}
 	}
 	else
 	{
@@ -562,6 +572,9 @@ void FAdaptiveVirtualTexture::Allocate(FVirtualTextureSystem* InSystem, uint32 I
 			AllocationSlots[InAllocationIndex].AllocatedVT = NewAllocatedVT;
 		}
 
+		// Add to pending for later indirection texture update.
+		SlotsPendingRootPageMap.Add(InAllocationIndex);
+
 		// Add to allocation structures.
 		GridIndexMap.Add(GetGridIndexHash(InGridIndex), InAllocationIndex);
 		AllocatedVTMap.Add(GetAllocatedVTHash(NewAllocatedVT), InAllocationIndex);
@@ -571,13 +584,6 @@ void FAdaptiveVirtualTexture::Allocate(FVirtualTextureSystem* InSystem, uint32 I
 
 		NumAllocated++;
 	}
-
-	// Queue indirection texture update.
-	const uint32 vAddress = NewAllocatedVT->GetVirtualAddress();
-	const uint32 vAddressX = FMath::ReverseMortonCode2(vAddress);
-	const uint32 vAddressY = FMath::ReverseMortonCode2(vAddress >> 1);
-	const uint32 PackedIndirectionValue = (1 << 28) | (InNewLevel << 24) | (vAddressY << 12) | vAddressX;
-	TextureUpdates.Add(FIndirectionTextureUpdate{ X, Y, PackedIndirectionValue });
 }
 
 void FAdaptiveVirtualTexture::Free(FVirtualTextureSystem* InSystem, uint32 InAllocationIndex, uint32 InFrame)
@@ -592,6 +598,7 @@ void FAdaptiveVirtualTexture::Free(FVirtualTextureSystem* InSystem, uint32 InAll
 	AllocatedVTMap.Remove(GetAllocatedVTHash(OldAllocatedVT), InAllocationIndex);
 	AllocationSlots[InAllocationIndex] = FAllocation(0, nullptr);
 	FreeSlots.Add(InAllocationIndex);
+	SlotsPendingRootPageMap.RemoveAllSwap([InAllocationIndex](int32& V) { return V == InAllocationIndex; });
 
 	NumAllocated--;
 	check(NumAllocated >= 0);
@@ -687,6 +694,30 @@ void FAdaptiveVirtualTexture::UpdateAllocations(FVirtualTextureSystem* InSystem,
 		}
 	}
 
+	// Check if any pending allocation slots are now ready.
+	// Pending slots are ones where the virtual texture locked root page(s) remain unmapped.
+	// If the root page is unmapped then we may return bad data from a sample.
+	for (int32 Index = 0; Index < SlotsPendingRootPageMap.Num(); ++Index)
+	{
+		const int32 AllocationSlotIndex = SlotsPendingRootPageMap[Index];
+		FAllocation const& Allocation = AllocationSlots[AllocationSlotIndex];
+		FAllocatedVirtualTexture* AllocatedVT = Allocation.AllocatedVT;
+		if (!InSystem->IsPendingRootPageMap(AllocatedVT))
+		{
+			SlotsPendingRootPageMap.RemoveAtSwap(Index--);
+
+			// Ready for use so that we can now queue the indirection texture update.
+			const uint32 vAddress = AllocatedVT->GetVirtualAddress();
+			const uint32 vAddressX = FMath::ReverseMortonCode2(vAddress);
+			const uint32 vAddressY = FMath::ReverseMortonCode2(vAddress >> 1);
+			const uint32 vLevel = AllocatedVT->GetMaxLevel();
+			const uint32 PackedIndirectionValue = (1 << 28) | (vLevel << 24) | (vAddressY << 12) | vAddressX;
+			const uint32 X = Allocation.GridIndex % GridSize.X;
+			const uint32 Y = Allocation.GridIndex / GridSize.X;
+			TextureUpdates.Add(FIndirectionTextureUpdate{ X, Y, PackedIndirectionValue });
+		}
+	}
+
 	// Update indirection texture
 	if (TextureUpdates.Num())
 	{
diff --git a/Engine/Source/Runtime/Renderer/Private/VT/AdaptiveVirtualTexture.h b/Engine/Source/Runtime/Renderer/Private/VT/AdaptiveVirtualTexture.h
--- a/Engine/Source/Runtime/Renderer/Private/VT/AdaptiveVirtualTexture.h
+++ b/Engine/Source/Runtime/Renderer/Private/VT/AdaptiveVirtualTexture.h
@@ -13,7 +13,7 @@ class FVirtualTextureSystem;
 
 /**
  * Concrete implementation of an adaptive virtual texture.
- * This allocates multiple virtual textures within the same space: one each for a grid of UV ranges, and an additional persistent one for the low mips.
+ * This allocates multiple virtual textures within the same space: one each for a grid of UV ranges, and an additional persistent one for the low resolution mips.
  * We then use an additional page table indirection texture in the shader to select the correct page table address range for our sampled UV.
  * We use the virtual texture feedback to decide when to increase or decrease the resolution of each UV range.
  * When we change resolution for a range we directly remap the page table entires. This removes the cost and any visual glitch from regenerating the pages.
@@ -107,6 +107,8 @@ private:
 	FHashTable GridIndexMap;
 	/** Map from AllocatedVT pointer to allocation slot index. */
 	FHashTable AllocatedVTMap;
+	/** Indices to AllocationSlots array for newly allocated virtual textures that are pending their root page before we can use them. */
+	TArray<int32> SlotsPendingRootPageMap;
 	/** Binary heap to track least recently used entries in AllocationSlots array. Used to decide what slots to evict next. */
 	FBinaryHeap<uint32, uint32> LRUHeap;
 
diff --git a/Engine/Source/Runtime/Renderer/Private/VT/VirtualTextureSystem.cpp b/Engine/Source/Runtime/Renderer/Private/VT/VirtualTextureSystem.cpp
--- a/Engine/Source/Runtime/Renderer/Private/VT/VirtualTextureSystem.cpp
+++ b/Engine/Source/Runtime/Renderer/Private/VT/VirtualTextureSystem.cpp
@@ -2390,7 +2390,7 @@ void FVirtualTextureSystem::SubmitRequests(FRDGBuilder& GraphBuilder, ERHIFeatur
 					{
 						FVirtualTexturePhysicalSpace* RESTRICT PhysicalSpace = Producer.GetPhysicalSpaceForPhysicalGroup(ProducerPhysicalGroupIndex);
 						FTexturePagePool& RESTRICT PagePool = PhysicalSpace->GetPagePool();
-						if (PagePool.AnyFreeAvailable(Frame, PageFreeThreshold))
+						if (PagePool.AnyFreeAvailable(Frame, bLockTile ? 0 : PageFreeThreshold))
 						{
 							const uint32 pAddress = PagePool.Alloc(this, Frame, ProducerHandle, ProducerPhysicalGroupIndex, TileToLoad.Local_vAddress, TileToLoad.Local_vLevel, bLockTile);
 							check(pAddress != ~0u);
@@ -2672,6 +2672,11 @@ float FVirtualTextureSystem::GetGlobalMipBias() const
 	return UTexture2D::GetGlobalMipMapLODBias() + MaxResidencyMipMapBias;
 }
 
+bool FVirtualTextureSystem::IsPendingRootPageMap(IAllocatedVirtualTexture* AllocatedVT) const
+{
+	return AllocatedVTsToMap.Find(AllocatedVT) != INDEX_NONE;
+}
+
 #if !UE_BUILD_SHIPPING
 
 void FVirtualTextureSystem::GetOnScreenMessages(FCoreDelegates::FSeverityMessageMap& OutMessages)
diff --git a/Engine/Source/Runtime/Renderer/Private/VT/VirtualTextureSystem.h b/Engine/Source/Runtime/Renderer/Private/VT/VirtualTextureSystem.h
--- a/Engine/Source/Runtime/Renderer/Private/VT/VirtualTextureSystem.h
+++ b/Engine/Source/Runtime/Renderer/Private/VT/VirtualTextureSystem.h
@@ -86,6 +86,8 @@ public:
 
 	float GetGlobalMipBias() const;
 
+	bool IsPendingRootPageMap(IAllocatedVirtualTexture* AllocatedVT) const;
+
 private:
 	friend class FFeedbackAnalysisTask;
 	friend class FAddRequestedTilesTask;
