Use new FViewFamilyInfo class for RVT view family.
Fixes crash when rendering RVT.
#rnx

diff --git a/Engine/Source/Runtime/Renderer/Private/VT/RuntimeVirtualTextureRender.cpp b/Engine/Source/Runtime/Renderer/Private/VT/RuntimeVirtualTextureRender.cpp
--- a/Engine/Source/Runtime/Renderer/Private/VT/RuntimeVirtualTextureRender.cpp
+++ b/Engine/Source/Runtime/Renderer/Private/VT/RuntimeVirtualTextureRender.cpp
@@ -1111,7 +1111,7 @@ namespace RuntimeVirtualTexture
 		//todo[vt]: Have specific shader variations and setup for different output texture configs
 		FSceneViewFamily::ConstructionValues ViewFamilyInit(nullptr, Scene, FEngineShowFlags(ESFIM_Game));
 		ViewFamilyInit.SetTime(FGameTime());
-		FSceneViewFamily& ViewFamily = *GraphBuilder.AllocObject<FSceneViewFamily>(ViewFamilyInit);
+		FSceneViewFamily& ViewFamily = *GraphBuilder.AllocObject<FViewFamilyInfo>(ViewFamilyInit);
 
 		FSceneViewInitOptions ViewInitOptions;
 		ViewInitOptions.ViewFamily = &ViewFamily;
