Fix a bug of stationary local lights with casting shadow disable don't work on mobile.

diff --git a/Engine/Source/Runtime/Renderer/Private/LightSceneInfo.cpp b/Engine/Source/Runtime/Renderer/Private/LightSceneInfo.cpp
--- a/Engine/Source/Runtime/Renderer/Private/LightSceneInfo.cpp
+++ b/Engine/Source/Runtime/Renderer/Private/LightSceneInfo.cpp
@@ -78,7 +78,7 @@ void FLightSceneInfo::AddToScene()
 	const FLightSceneInfoCompact& LightSceneInfoCompact = Scene->Lights[Id];
 
 	bool bIsValidLightTypeMobile = false;
-	if (Scene->GetShadingPath() == EShadingPath::Mobile && Proxy->IsMovable())
+	if (Scene->GetShadingPath() == EShadingPath::Mobile && ShouldRenderLightViewIndependent())
 	{
 		const uint8 LightType = Proxy->GetLightType();
 		bIsValidLightTypeMobile = LightType == LightType_Rect || LightType == LightType_Point || LightType == LightType_Spot;
