Fixed: missing reflections when clustered option is enabled on mobile

diff --git a/Engine/Shaders/Private/Common.ush b/Engine/Shaders/Private/Common.ush
--- a/Engine/Shaders/Private/Common.ush
+++ b/Engine/Shaders/Private/Common.ush
@@ -137,7 +137,7 @@ const static float3 Max111110BitsFloat3 = float3(Max11BitsFloat, Max11BitsFloat,
 #endif
 
 #ifndef SUPPORTS_TEXTURECUBE_ARRAY
-#define SUPPORTS_TEXTURECUBE_ARRAY (FEATURE_LEVEL >= FEATURE_LEVEL_SM5)
+#define SUPPORTS_TEXTURECUBE_ARRAY 1
 #endif
 
 #if SUPPORTS_TEXTURECUBE_ARRAY == 0
diff --git a/Engine/Source/Runtime/Renderer/Private/MobileDeferredShadingPass.cpp b/Engine/Source/Runtime/Renderer/Private/MobileDeferredShadingPass.cpp
--- a/Engine/Source/Runtime/Renderer/Private/MobileDeferredShadingPass.cpp
+++ b/Engine/Source/Runtime/Renderer/Private/MobileDeferredShadingPass.cpp
@@ -65,7 +65,6 @@ class FMobileDirectionalLightFunctionPS : public FMaterialShader
 		FMaterialShader::ModifyCompilationEnvironment(Parameters, OutEnvironment);
 		FForwardLightingParameters::ModifyCompilationEnvironment(Parameters.Platform, OutEnvironment);
 		OutEnvironment.SetDefine(TEXT(PREPROCESSOR_TO_STRING(MAX_MOBILE_SHADOWCASCADES)), GetMobileMaxShadowCascades());
-		OutEnvironment.SetDefine(TEXT("SUPPORTS_TEXTURECUBE_ARRAY"), 1);
 		OutEnvironment.SetDefine(TEXT("USE_LIGHT_FUNCTION"), Parameters.MaterialParameters.bIsDefaultMaterial ? 0 : 1);
 		OutEnvironment.SetDefine(TEXT("USE_SHADOWMASKTEXTURE"), MobileUsesShadowMaskTexture(Parameters.Platform) ? 1u : 0u);
 		OutEnvironment.SetDefine(TEXT("MATERIAL_SHADER"), 1);
@@ -260,7 +259,6 @@ public:
 	static void ModifyCompilationEnvironment(const FGlobalShaderPermutationParameters& Parameters, FShaderCompilerEnvironment& OutEnvironment)
 	{
 		FGlobalShader::ModifyCompilationEnvironment(Parameters, OutEnvironment);
-		OutEnvironment.SetDefine(TEXT("SUPPORTS_TEXTURECUBE_ARRAY"), 1);
 	}
 };
 
