Make ShaderPrint easier to use from outside renderer library.
Aim is to allow use without needing internals of FViewInfo.
Remaining step is to fully allow rendering final results from self managed FShaderPrintData.

diff --git a/Engine/Shaders/Private/ShaderPrintDrawPrimitive.usf b/Engine/Shaders/Private/ShaderPrintDrawPrimitive.usf
--- a/Engine/Shaders/Private/ShaderPrintDrawPrimitive.usf
+++ b/Engine/Shaders/Private/ShaderPrintDrawPrimitive.usf
@@ -15,6 +15,7 @@ float2 OutputInvResolution;
 float2 OriginalViewRectMin;
 float2 OriginalViewSize;
 float2 OriginalBufferInvSize;
+float4x4 TranslatedWorldToClip;
 
 SamplerState DepthSampler;
 
@@ -46,7 +47,7 @@ void ShaderDrawDebugVS(
 		// world offset
 		VertexPositionTWS.xyz += TranslatedWorldOffsetConversion;
 
-		Position = mul(VertexPositionTWS, PrimaryView.TranslatedWorldToClip);
+		Position = mul(VertexPositionTWS, TranslatedWorldToClip);
 	}
 	ColorAlpha = VertexColor;
 }
diff --git a/Engine/Source/Runtime/Renderer/Private/ShaderPrint.cpp b/Engine/Source/Runtime/Renderer/Private/ShaderPrint.cpp
--- a/Engine/Source/Runtime/Renderer/Private/ShaderPrint.cpp
+++ b/Engine/Source/Runtime/Renderer/Private/ShaderPrint.cpp
@@ -10,9 +10,10 @@
 #include "GlobalShader.h"
 #include "PipelineStateCache.h"
 #include "RenderGraphBuilder.h"
+#include "ScenePrivate.h"
 #include "SceneRendering.h"
+#include "ScreenPass.h"
 #include "SystemTextures.h"
-#include "ScenePrivate.h"
 
 namespace ShaderPrint
 {
@@ -122,9 +123,14 @@ namespace ShaderPrint
 
 	// Get symbol buffer size
 	// This is some multiple of the value buffer size to allow for maximum value->symbol expansion
+	static uint32 GetMaxSymbolCountFromValueCount(uint32 MaxValueCount)
+	{
+		return MaxValueCount * 12u;
+	}
+
 	static uint32 GetMaxSymbolCount()
 	{
-		return GetMaxValueCount() * 12u;
+		return GetMaxSymbolCountFromValueCount(GetMaxValueCount());
 	}
 
 	static bool IsDrawLocked()
@@ -140,27 +146,44 @@ namespace ShaderPrint
 	
 	// Fill the uniform buffer parameters
 	// Return a uniform buffer with values filled and with single frame lifetime
-	static TUniformBufferRef<ShaderPrint::FShaderPrintCommonParameters> CreateUniformBuffer(const FShaderPrintData& Data)
+	static TUniformBufferRef<ShaderPrint::FShaderPrintCommonParameters> CreateUniformBuffer(const FShaderPrintSetup& Setup)
 	{		
 		FShaderPrintCommonParameters Out;
-		Out.FontSize = Data.FontSize;
-		Out.FontSpacing = Data.FontSpacing;
-		Out.Resolution = Data.OutputRect.Size();
-		Out.CursorCoord = Data.CursorCoord;
-		Out.MaxValueCount = Data.MaxValueCount;
-		Out.MaxSymbolCount = Data.MaxSymbolCount;
-		Out.MaxStateCount = Data.MaxStateCount;
-		Out.MaxLineCount = Data.MaxLineCount;
-		Out.TranslatedWorldOffset = FVector3f(Data.TranslatedWorldOffset);
+
+		const FVector2D ViewSize(FMath::Max(Setup.ViewRect.Size().X, 1), FMath::Max(Setup.ViewRect.Size().Y, 1));
+		const float FontWidth = float(Setup.FontSize.X) * Setup.DPIScale / ViewSize.X;
+		const float FontHeight = float(Setup.FontSize.Y) * Setup.DPIScale / ViewSize.Y;
+		const float SpaceWidth = float(Setup.FontSpacing.X) * Setup.DPIScale / ViewSize.X;
+		const float SpaceHeight = float(Setup.FontSpacing.Y) * Setup.DPIScale / ViewSize.Y;
+		
+		Out.FontSize = FVector2f(FontWidth, FontHeight);
+		Out.FontSpacing = FVector2f(FontWidth + SpaceWidth, FontHeight + SpaceHeight);
+		Out.Resolution = Setup.ViewRect.Size();
+		Out.CursorCoord = Setup.CursorCoord;
+		Out.MaxValueCount = Setup.MaxValueCount;
+		Out.MaxSymbolCount = GetMaxSymbolCountFromValueCount(Setup.MaxValueCount);
+		Out.MaxStateCount = Setup.MaxStateCount;
+		Out.MaxLineCount = Setup.MaxLineCount;
+		Out.TranslatedWorldOffset = FVector3f(Setup.PreViewTranslation);
+
 		return TUniformBufferRef<ShaderPrint::FShaderPrintCommonParameters>::CreateUniformBufferImmediate(Out, UniformBuffer_SingleFrame);
 	}
 
 	//////////////////////////////////////////////////////////////////////////////////////////////////
 	// Accessors
 
+	// Fill the FShaderParameters parameters
+	void SetParameters(FRDGBuilder& GraphBuilder, const FShaderPrintData& Data, FShaderParameters& OutParameters)
+	{
+		OutParameters.Common = Data.UniformBuffer;
+		OutParameters.ShaderPrint_StateBuffer = GraphBuilder.CreateSRV(Data.ShaderPrintStateBuffer);
+		OutParameters.ShaderPrint_RWValuesBuffer = GraphBuilder.CreateUAV(Data.ShaderPrintValueBuffer);
+		OutParameters.ShaderPrint_RWLinesBuffer = GraphBuilder.CreateUAV(Data.ShaderPrintLineBuffer);
+	}
+
 	void SetParameters(FRDGBuilder& GraphBuilder, FShaderParameters& OutParameters)
 	{
-		if (GDefaultView != nullptr)
+		if (ensure(GDefaultView != nullptr))
 		{
 			SetParameters(GraphBuilder, GDefaultView->ShaderPrintData, OutParameters);
 		}
@@ -172,15 +195,6 @@ namespace ShaderPrint
 		SetParameters(GraphBuilder, View.ShaderPrintData, OutParameters);
 	}
 
-	// Fill the FShaderParameters parameters
-	void SetParameters(FRDGBuilder& GraphBuilder, const FShaderPrintData& Data, FShaderParameters& OutParameters)
-	{
-		OutParameters.Common = Data.UniformBuffer;
-		OutParameters.ShaderPrint_StateBuffer = GraphBuilder.CreateSRV(Data.ShaderPrintStateBuffer);
-		OutParameters.ShaderPrint_RWValuesBuffer = GraphBuilder.CreateUAV(Data.ShaderPrintValueBuffer);
-		OutParameters.ShaderPrint_RWLinesBuffer = GraphBuilder.CreateUAV(Data.ShaderPrintLineBuffer);
-	}
-
 	// Supported platforms
 	bool IsSupported(EShaderPlatform InShaderPlatform)
 	{
@@ -219,7 +233,7 @@ namespace ShaderPrint
 		return CVarEnable.GetValueOnAnyThread() != 0 || GShaderPrintEnableOverride;
 	}
 
-	bool IsEnabled(const FViewInfo& View)
+	bool IsEnabled(const FSceneView& View)
 	{
 		return IsEnabled() && IsSupported(View.GetShaderPlatform());
 	}
@@ -441,6 +455,7 @@ namespace ShaderPrint
 
 		BEGIN_SHADER_PARAMETER_STRUCT(FParameters, )
 			SHADER_PARAMETER(FVector3f, TranslatedWorldOffsetConversion)
+			SHADER_PARAMETER(FMatrix44f, TranslatedWorldToClip)
 			SHADER_PARAMETER_STRUCT_REF(FViewUniformShaderParameters, View)
 			SHADER_PARAMETER_SRV(StructuredBuffer, LockedShaderDrawDebugPrimitive)
 			SHADER_PARAMETER_RDG_BUFFER_SRV(StructuredBuffer, ShaderDrawDebugPrimitive)
@@ -503,154 +518,183 @@ namespace ShaderPrint
 	END_SHADER_PARAMETER_STRUCT()
 
 	//////////////////////////////////////////////////////////////////////////////////////////////////
-	// Drawing/Rendering API
+	// Setup render data
 
-	void BeginView(FRDGBuilder& GraphBuilder, FViewInfo& View)
+	FShaderPrintSetup::FShaderPrintSetup(FIntRect InViewRect)
 	{
-		TRACE_CPUPROFILER_EVENT_SCOPE(ShaderPrint::BeginView);
+		ViewRect = InViewRect;
+		CursorCoord = FIntPoint(-1, -1);
+		PreViewTranslation = FVector::ZeroVector;
+		DPIScale = 1.f;
 
-		View.ShaderPrintData = FShaderPrintData();
-		if (!IsSupported(View.GetShaderPlatform()))
-		{
-			return;
-		}
+		FontSize = FIntPoint(FMath::Max(CVarFontSize.GetValueOnRenderThread(), 1), FMath::Max(CVarFontSize.GetValueOnRenderThread(), 1));
+		FontSpacing = FIntPoint(FMath::Max(CVarFontSpacingX.GetValueOnRenderThread(), 1), FMath::Max(CVarFontSpacingY.GetValueOnRenderThread(), 1));
+		MaxValueCount = GetMaxValueCount();
+		MaxStateCount = GetMaxWidgetCount();
+		MaxLineCount = GetMaxLineCount();
+	}
 
-		const FIntPoint ViewSize(FMath::Max(View.UnconstrainedViewRect.Size().X, 1), FMath::Max(View.UnconstrainedViewRect.Size().Y, 1));
+	FShaderPrintSetup::FShaderPrintSetup(FSceneView const& View)
+	{
+		ViewRect = View.UnconstrainedViewRect;
+		CursorCoord = View.CursorPos;
+		PreViewTranslation = View.ViewMatrices.GetPreViewTranslation();
+		DPIScale = View.Family->DebugDPIScale;
 
-		float FontSize = float(FMath::Max(CVarFontSize.GetValueOnRenderThread(), 1)) * View.Family->DebugDPIScale;
-		float FontSpacingX = float(FMath::Max(CVarFontSpacingX.GetValueOnRenderThread(), 1)) * View.Family->DebugDPIScale;
-		float FontSpacingY = float(FMath::Max(CVarFontSpacingY.GetValueOnRenderThread(), 1)) * View.Family->DebugDPIScale;
+		FontSize = FIntPoint(FMath::Max(CVarFontSize.GetValueOnRenderThread(), 1), FMath::Max(CVarFontSize.GetValueOnRenderThread(), 1));
+		FontSpacing = FIntPoint(FMath::Max(CVarFontSpacingX.GetValueOnRenderThread(), 1), FMath::Max(CVarFontSpacingY.GetValueOnRenderThread(), 1));
+		MaxValueCount = GetMaxValueCount();
+		MaxStateCount = GetMaxWidgetCount();
+		MaxLineCount = GetMaxLineCount();
+	}
 
-		const float FontWidth = FontSize / float(ViewSize.X);
-		const float FontHeight = FontSize / float(ViewSize.Y);
-		const float SpaceWidth = FontSpacingX / float(ViewSize.X);
-		const float SpaceHeight = FontSpacingY / float(ViewSize.Y);
+	FShaderPrintData CreateShaderPrintData(FRDGBuilder& GraphBuilder, FShaderPrintSetup const& InSetup, FSceneViewState* InViewState)
+	{
+		FShaderPrintData ShaderPrintData;
 
-		View.ShaderPrintData.FontSize = FVector2f(FontWidth, FontHeight);
-		View.ShaderPrintData.FontSpacing = FVector2f(SpaceWidth + FontWidth, SpaceHeight + FontHeight);		
-		View.ShaderPrintData.OutputRect = View.UnconstrainedViewRect;
-		View.ShaderPrintData.MaxValueCount = GetMaxValueCount();
-		View.ShaderPrintData.MaxSymbolCount = GetMaxSymbolCount();
-		View.ShaderPrintData.MaxStateCount = GetMaxWidgetCount();
-		View.ShaderPrintData.MaxLineCount = GetMaxLineCount();
-		View.ShaderPrintData.CursorCoord = View.CursorPos;
-		View.ShaderPrintData.TranslatedWorldOffset = View.ViewMatrices.GetPreViewTranslation();
-		GCharacterRequestCount = 0;
-		GWidgetRequestCount = 0;
-		GLineRequestCount = 0;
+		// Common uniform buffer
+		ShaderPrintData.UniformBuffer = CreateUniformBuffer(InSetup);
 
-		// Early out if system is disabled
-		// Note that we still bind a dummy ShaderPrintValueBuffer 
-		// This is in case some debug shader code is still active (we don't want an unbound buffer!)
+		// Early out if system is disabled.
+		// Note that we still bind dummy buffers.
+		// This is in case some debug shader code is still active and accessing the buffer.
 		if (!IsEnabled())
 		{
-			View.ShaderPrintData.UniformBuffer = CreateUniformBuffer(View.ShaderPrintData);
-			View.ShaderPrintData.ShaderPrintValueBuffer = GraphBuilder.RegisterExternalBuffer(GEmptyBuffer->Buffer);
-			View.ShaderPrintData.ShaderPrintStateBuffer = GraphBuilder.RegisterExternalBuffer(GEmptyBuffer->Buffer);
-			View.ShaderPrintData.ShaderPrintLineBuffer  = GraphBuilder.RegisterExternalBuffer(GEmptyBuffer->Buffer);
-			return;
-		}
-
-		// Invalid to call begin twice for the same view.
-		check(GDefaultView != &View);
-		if (GDefaultView == nullptr)
-		{
-			GDefaultView = &View;
+			ShaderPrintData.ShaderPrintValueBuffer = GraphBuilder.RegisterExternalBuffer(GEmptyBuffer->Buffer);
+			ShaderPrintData.ShaderPrintStateBuffer = GraphBuilder.RegisterExternalBuffer(GEmptyBuffer->Buffer);
+			ShaderPrintData.ShaderPrintLineBuffer = GraphBuilder.RegisterExternalBuffer(GEmptyBuffer->Buffer);
+			return ShaderPrintData;
 		}
 
-		// Primitives/Lines
-		// ----------------
-		{			
-			FSceneViewState* ViewState = View.ViewState;
-			const bool bLockBufferThisFrame = IsDrawLocked() && ViewState && !ViewState->ShaderPrintStateData.bIsLocked;
-			ERDGBufferFlags Flags = bLockBufferThisFrame ? ERDGBufferFlags::MultiFrame : ERDGBufferFlags::None;
-
-			View.ShaderPrintData.ShaderPrintLineBuffer = GraphBuilder.CreateBuffer(FRDGBufferDesc::CreateStructuredDesc(4, 8 * View.ShaderPrintData.MaxLineCount), TEXT("ShaderPrint.LineBuffer"), Flags);
-
-			// Clear buffer counter
-			{
-				FShaderDrawDebugClearCS::FParameters* Parameters = GraphBuilder.AllocParameters<FShaderDrawDebugClearCS::FParameters>();
-				Parameters->RWElementBuffer = GraphBuilder.CreateUAV(View.ShaderPrintData.ShaderPrintLineBuffer);
-
-				TShaderMapRef<FShaderDrawDebugClearCS> ComputeShader(View.ShaderMap);
-				ClearUnusedGraphResources(ComputeShader, Parameters);
-				GraphBuilder.AddPass(
-					RDG_EVENT_NAME("ShaderPrint::BeginView(Clear lines)"),
-					Parameters,
-					ERDGPassFlags::Compute,
-					[Parameters, ComputeShader](FRHICommandList& RHICmdList)
-					{
-						FComputeShaderUtils::Dispatch(RHICmdList, ComputeShader, *Parameters, FIntVector(1, 1, 1));
-					});
-			}
-
-			if (IsDrawLocked() && ViewState && !ViewState->ShaderPrintStateData.bIsLocked)
-			{
-				ViewState->ShaderPrintStateData.LineBuffer = GraphBuilder.ConvertToExternalBuffer(View.ShaderPrintData.ShaderPrintLineBuffer);
-				ViewState->ShaderPrintStateData.PreViewTranslation = View.ViewMatrices.GetPreViewTranslation();
-				ViewState->ShaderPrintStateData.bIsLocked = true;
-			}
-
-			if (!IsDrawLocked() && ViewState && ViewState->ShaderPrintStateData.bIsLocked)
-			{
-				ViewState->ShaderPrintStateData.LineBuffer = nullptr;
-				ViewState->ShaderPrintStateData.PreViewTranslation = FVector::ZeroVector;
-				ViewState->ShaderPrintStateData.bIsLocked = false;
-			}
-		}
+		FGlobalShaderMap* GlobalShaderMap = GetGlobalShaderMap(GMaxRHIFeatureLevel);
 
-		// Characters & Widgets
-		// --------------------
+		// Characters/Widgets
 		{
 			// Initialize output buffer and store in the view info
 			// Values buffer contains Count + 1 elements. The first element is only used as a counter.
-			View.ShaderPrintData.ShaderPrintValueBuffer = GraphBuilder.CreateBuffer(FRDGBufferDesc::CreateStructuredDesc(sizeof(FPackedShaderPrintItem), View.ShaderPrintData.MaxValueCount + 1), TEXT("ShaderPrint.ValueBuffer"));
+			ShaderPrintData.ShaderPrintValueBuffer = GraphBuilder.CreateBuffer(FRDGBufferDesc::CreateStructuredDesc(sizeof(FPackedShaderPrintItem), InSetup.MaxValueCount + 1), TEXT("ShaderPrint.ValueBuffer"));
 
 			// State buffer is retrieved from the view state, or created if it does not exist
-			View.ShaderPrintData.ShaderPrintStateBuffer = nullptr;
-			if (FSceneViewState* ViewState = View.ViewState)
+			if (InViewState != nullptr)
 			{
-				if (ViewState->ShaderPrintStateData.StateBuffer)
+				if (InViewState->ShaderPrintStateData.StateBuffer)
 				{
-					View.ShaderPrintData.ShaderPrintStateBuffer = GraphBuilder.RegisterExternalBuffer(ViewState->ShaderPrintStateData.StateBuffer);
+					ShaderPrintData.ShaderPrintStateBuffer = GraphBuilder.RegisterExternalBuffer(InViewState->ShaderPrintStateData.StateBuffer);
 				}
 				else
 				{
-					View.ShaderPrintData.ShaderPrintStateBuffer = GraphBuilder.CreateBuffer(FRDGBufferDesc::CreateStructuredDesc(sizeof(uint32), (3*View.ShaderPrintData.MaxStateCount) + 1), TEXT("ShaderPrint.StateBuffer"));
-					AddClearUAVPass(GraphBuilder, GraphBuilder.CreateUAV(View.ShaderPrintData.ShaderPrintStateBuffer, PF_R32_UINT), 0u);
-					ViewState->ShaderPrintStateData.StateBuffer = GraphBuilder.ConvertToExternalBuffer(View.ShaderPrintData.ShaderPrintStateBuffer);
+					ShaderPrintData.ShaderPrintStateBuffer = GraphBuilder.CreateBuffer(FRDGBufferDesc::CreateStructuredDesc(sizeof(uint32), (3 * InSetup.MaxStateCount) + 1), TEXT("ShaderPrint.StateBuffer"));
+					AddClearUAVPass(GraphBuilder, GraphBuilder.CreateUAV(ShaderPrintData.ShaderPrintStateBuffer, PF_R32_UINT), 0u);
+					InViewState->ShaderPrintStateData.StateBuffer = GraphBuilder.ConvertToExternalBuffer(ShaderPrintData.ShaderPrintStateBuffer);
 				}
 			}
 			else
 			{
-				View.ShaderPrintData.MaxStateCount = 0;
-				View.ShaderPrintData.ShaderPrintStateBuffer = GraphBuilder.RegisterExternalBuffer(GEmptyBuffer->Buffer);
+				ShaderPrintData.ShaderPrintStateBuffer = GraphBuilder.RegisterExternalBuffer(GEmptyBuffer->Buffer);
 			}
 
 			// Clear the output buffer internal counter ready for use
-			const ERHIFeatureLevel::Type FeatureLevel = View.GetFeatureLevel();
-			FGlobalShaderMap* GlobalShaderMap = GetGlobalShaderMap(FeatureLevel);
 			TShaderMapRef< FShaderInitValueBufferCS > ComputeShader(GlobalShaderMap);
 
 			auto* PassParameters = GraphBuilder.AllocParameters<FShaderInitValueBufferCS::FParameters>();
-			PassParameters->RWValuesBuffer = GraphBuilder.CreateUAV(View.ShaderPrintData.ShaderPrintValueBuffer);
+			PassParameters->RWValuesBuffer = GraphBuilder.CreateUAV(ShaderPrintData.ShaderPrintValueBuffer);
 
 			FComputeShaderUtils::AddPass(
 				GraphBuilder,
-				RDG_EVENT_NAME("ShaderPrint::BeginView(Clear characters)"),
+				RDG_EVENT_NAME("ShaderPrint::CreateShaderPrintData (Clear characters)"),
 				ComputeShader,
 				PassParameters,
 				FIntVector(1, 1, 1));
 		}
 
-		// Create common parameters uniform buffer
-		View.ShaderPrintData.UniformBuffer = CreateUniformBuffer(View.ShaderPrintData);
+		// Primitives/Lines
+		{
+			const bool bLockBufferThisFrame = IsDrawLocked() && InViewState != nullptr && !InViewState->ShaderPrintStateData.bIsLocked;
+			ERDGBufferFlags Flags = bLockBufferThisFrame ? ERDGBufferFlags::MultiFrame : ERDGBufferFlags::None;
+
+			ShaderPrintData.ShaderPrintLineBuffer = GraphBuilder.CreateBuffer(FRDGBufferDesc::CreateStructuredDesc(4, 8 * InSetup.MaxLineCount), TEXT("ShaderPrint.LineBuffer"), Flags);
+
+			// Clear buffer counter
+			{
+				TShaderMapRef<FShaderDrawDebugClearCS> ComputeShader(GlobalShaderMap);
+
+				FShaderDrawDebugClearCS::FParameters* PassParameters = GraphBuilder.AllocParameters<FShaderDrawDebugClearCS::FParameters>();
+				PassParameters->RWElementBuffer = GraphBuilder.CreateUAV(ShaderPrintData.ShaderPrintLineBuffer);
+				ClearUnusedGraphResources(ComputeShader, PassParameters);
+
+				GraphBuilder.AddPass(
+					RDG_EVENT_NAME("ShaderPrint::CreateShaderPrintData (Clear lines)"),
+					PassParameters,
+					ERDGPassFlags::Compute,
+					[PassParameters, ComputeShader](FRHICommandList& RHICmdList)
+					{
+						FComputeShaderUtils::Dispatch(RHICmdList, ComputeShader, *PassParameters, FIntVector(1, 1, 1));
+					});
+			}
+
+			if (InViewState != nullptr)
+			{
+				if (IsDrawLocked() && !InViewState->ShaderPrintStateData.bIsLocked)
+				{
+					InViewState->ShaderPrintStateData.LineBuffer = GraphBuilder.ConvertToExternalBuffer(ShaderPrintData.ShaderPrintLineBuffer);
+					InViewState->ShaderPrintStateData.PreViewTranslation = InSetup.PreViewTranslation;
+					InViewState->ShaderPrintStateData.bIsLocked = true;
+				}
+
+				if (!IsDrawLocked() && InViewState->ShaderPrintStateData.bIsLocked)
+				{
+					InViewState->ShaderPrintStateData.LineBuffer = nullptr;
+					InViewState->ShaderPrintStateData.PreViewTranslation = FVector::ZeroVector;
+					InViewState->ShaderPrintStateData.bIsLocked = false;
+				}
+			}
+		}
+
+		return ShaderPrintData;
 	}
 
-	static void InternalDrawView_Characters(FRDGBuilder& GraphBuilder, const FViewInfo& View, FScreenPassTexture OutputTexture)
+	FShaderPrintData CreateShaderPrintData(FRDGBuilder& GraphBuilder, FShaderPrintSetup const& InSetup)
+	{
+		return CreateShaderPrintData(GraphBuilder, InSetup, nullptr);
+	}
+
+	//////////////////////////////////////////////////////////////////////////////////////////////////
+	// Drawing/Rendering API
+
+	void BeginView(FRDGBuilder& GraphBuilder, FViewInfo& View)
+	{
+		TRACE_CPUPROFILER_EVENT_SCOPE(ShaderPrint::BeginView);
+
+		View.ShaderPrintData = FShaderPrintData();
+		if (!IsSupported(View.GetShaderPlatform()))
+		{
+			return;
+		}
+
+		// Invalid to call begin twice for the same view.
+		ensure(GDefaultView != &View);
+		if (GDefaultView == nullptr)
+		{
+			GDefaultView = &View;
+		}
+
+		// Create the render data and store on the view.
+		FShaderPrintSetup ShaderPrintSetup(View);
+		View.ShaderPrintData = CreateShaderPrintData(GraphBuilder, ShaderPrintSetup, View.ViewState);
+
+		// Reset counter which is read on the next BeginView().
+		GCharacterRequestCount = 0;
+		GWidgetRequestCount = 0;
+		GLineRequestCount = 0;
+	}
+
+	static void InternalDrawView_Characters(
+		FRDGBuilder& GraphBuilder, 
+		FShaderPrintData const& ShaderPrintData, 
+		FIntRect ViewRect, 
+		int32 FrameNumber, 
+		FScreenPassTexture OutputTexture)
 	{
-		const FIntRect Viewport = OutputTexture.ViewRect;
-	
 		// Initialize graph managed resources
 		// Symbols buffer contains Count + 1 elements. The first element is only used as a counter.
 		FRDGBufferRef SymbolBuffer = GraphBuilder.CreateBuffer(FRDGBufferDesc::CreateStructuredDesc(sizeof(FPackedShaderPrintItem), GetMaxSymbolCount() + 1), TEXT("ShaderPrint.SymbolBuffer"));
@@ -658,12 +702,11 @@ namespace ShaderPrint
 		FRDGBufferRef IndirectDrawArgsBuffer = GraphBuilder.CreateBuffer(FRDGBufferDesc::CreateIndirectDesc(5), TEXT("ShaderPrint.IndirectDrawArgs"));
 
 		// Non graph managed resources
-		FRDGBufferSRVRef ValueBuffer = GraphBuilder.CreateSRV(View.ShaderPrintData.ShaderPrintValueBuffer);
-		FRDGBufferSRVRef StateBuffer = GraphBuilder.CreateSRV(View.ShaderPrintData.ShaderPrintStateBuffer);
+		FRDGBufferSRVRef ValueBuffer = GraphBuilder.CreateSRV(ShaderPrintData.ShaderPrintValueBuffer);
+		FRDGBufferSRVRef StateBuffer = GraphBuilder.CreateSRV(ShaderPrintData.ShaderPrintStateBuffer);
 		FTextureRHIRef FontTexture = GSystemTextures.AsciiTexture->GetRHI();
 
-		const ERHIFeatureLevel::Type FeatureLevel = View.GetFeatureLevel();
-		FGlobalShaderMap* GlobalShaderMap = GetGlobalShaderMap(FeatureLevel);
+		FGlobalShaderMap* GlobalShaderMap = GetGlobalShaderMap(GMaxRHIFeatureLevel);
 
 		// BuildIndirectDispatchArgs
 		{
@@ -671,7 +714,7 @@ namespace ShaderPrint
 			TShaderMapRef<SHADER> ComputeShader(GlobalShaderMap);
 
 			SHADER::FParameters* PassParameters = GraphBuilder.AllocParameters<SHADER::FParameters>();
-			PassParameters->Common = View.ShaderPrintData.UniformBuffer;
+			PassParameters->Common = ShaderPrintData.UniformBuffer;
 			PassParameters->ValuesBuffer = ValueBuffer;
 			PassParameters->RWSymbolsBuffer = GraphBuilder.CreateUAV(SymbolBuffer);
 			PassParameters->RWIndirectDispatchArgsBuffer = GraphBuilder.CreateUAV(IndirectDispatchArgsBuffer, EPixelFormat::PF_R32_UINT);
@@ -689,11 +732,11 @@ namespace ShaderPrint
 			TShaderMapRef<SHADER> ComputeShader(GlobalShaderMap);
 
 			SHADER::FParameters* PassParameters = GraphBuilder.AllocParameters<SHADER::FParameters>();
-			PassParameters->FrameIndex = View.Family ? View.Family->FrameNumber : 0u;
-			PassParameters->Common = View.ShaderPrintData.UniformBuffer;
+			PassParameters->FrameIndex = FrameNumber;
+			PassParameters->Common = ShaderPrintData.UniformBuffer;
 			PassParameters->ValuesBuffer = ValueBuffer;
 			PassParameters->RWSymbolsBuffer = GraphBuilder.CreateUAV(SymbolBuffer);
-			PassParameters->RWStateBuffer = GraphBuilder.CreateUAV(View.ShaderPrintData.ShaderPrintStateBuffer);
+			PassParameters->RWStateBuffer = GraphBuilder.CreateUAV(ShaderPrintData.ShaderPrintStateBuffer);
 			PassParameters->IndirectDispatchArgsBuffer = IndirectDispatchArgsBuffer;
 
 			FComputeShaderUtils::AddPass(
@@ -710,10 +753,10 @@ namespace ShaderPrint
 			TShaderMapRef<SHADER> ComputeShader(GlobalShaderMap);
 
 			SHADER::FParameters* PassParameters = GraphBuilder.AllocParameters<SHADER::FParameters>();
-			PassParameters->FrameIndex = View.Family ? View.Family->FrameNumber : 0u;
+			PassParameters->FrameIndex = FrameNumber;
 			PassParameters->FrameThreshold = 300u;
-			PassParameters->Common = View.ShaderPrintData.UniformBuffer;
-			PassParameters->RWStateBuffer = GraphBuilder.CreateUAV(View.ShaderPrintData.ShaderPrintStateBuffer);
+			PassParameters->Common = ShaderPrintData.UniformBuffer;
+			PassParameters->RWStateBuffer = GraphBuilder.CreateUAV(ShaderPrintData.ShaderPrintStateBuffer);
 
 			FComputeShaderUtils::AddPass(
 				GraphBuilder,
@@ -728,7 +771,7 @@ namespace ShaderPrint
 			TShaderMapRef<SHADER> ComputeShader(GlobalShaderMap);
 
 			SHADER::FParameters* PassParameters = GraphBuilder.AllocParameters<SHADER::FParameters>();
-			PassParameters->Common = View.ShaderPrintData.UniformBuffer;
+			PassParameters->Common = ShaderPrintData.UniformBuffer;
 			PassParameters->SymbolsBuffer = GraphBuilder.CreateSRV(SymbolBuffer);
 			PassParameters->RWIndirectDrawArgsBuffer = GraphBuilder.CreateUAV(IndirectDrawArgsBuffer, EPixelFormat::PF_R32_UINT);
 
@@ -747,7 +790,7 @@ namespace ShaderPrint
 
 			SHADER::FParameters* PassParameters = GraphBuilder.AllocParameters<SHADER::FParameters>();
 			PassParameters->RenderTargets[0] = FRenderTargetBinding(OutputTexture.Texture, ERenderTargetLoadAction::ELoad);
-			PassParameters->Common = View.ShaderPrintData.UniformBuffer;
+			PassParameters->Common = ShaderPrintData.UniformBuffer;
 			PassParameters->MiniFontTexture = FontTexture;
 			PassParameters->SymbolsBuffer = GraphBuilder.CreateSRV(SymbolBuffer);
 			PassParameters->IndirectDrawArgsBuffer = IndirectDrawArgsBuffer;
@@ -756,7 +799,7 @@ namespace ShaderPrint
 				RDG_EVENT_NAME("ShaderPrint::DrawSymbols"),
 				PassParameters,
 				ERDGPassFlags::Raster,
-				[VertexShader, PixelShader, PassParameters, Viewport](FRHICommandList& RHICmdList)
+				[VertexShader, PixelShader, PassParameters, ViewRect](FRHICommandList& RHICmdList)
 			{
 				
 				FGraphicsPipelineStateInitializer GraphicsPSOInit;
@@ -770,7 +813,7 @@ namespace ShaderPrint
 				GraphicsPSOInit.BoundShaderState.PixelShaderRHI = PixelShader.GetPixelShader();
 				SetGraphicsPipelineState(RHICmdList, GraphicsPSOInit, 0);
 
-				RHICmdList.SetViewport(Viewport.Min.X, Viewport.Min.Y, 0.0f, Viewport.Max.X, Viewport.Max.Y, 1.0f);
+				RHICmdList.SetViewport(ViewRect.Min.X, ViewRect.Min.Y, 0.0f, ViewRect.Max.X, ViewRect.Max.Y, 1.0f);
 				SetShaderParameters(RHICmdList, VertexShader, VertexShader.GetVertexShader(), *PassParameters);
 				SetShaderParameters(RHICmdList, PixelShader, PixelShader.GetPixelShader(), *PassParameters);
 
@@ -781,19 +824,23 @@ namespace ShaderPrint
 
 	static void InternalDrawView_Lines(
 		FRDGBuilder& GraphBuilder,
-		const FViewInfo& View,
-		const FVector& TranslatedWorldOffsetConversion,
-		FRDGBufferRef DataBuffer,
+		FRDGBufferRef ShaderPrintLineBuffer,
+		FIntRect ViewRect,
+		FIntRect UnscaledViewRect,
+		FMatrix const& TranslatedWorldToClip,
+		FVector TranslatedWorldOffsetConversion,
 		FRDGTextureRef OutputTexture,
 		FRDGTextureRef DepthTexture)
 	{
+		FGlobalShaderMap* GlobalShaderMap = GetGlobalShaderMap(GMaxRHIFeatureLevel);
+
 		FRDGBufferRef IndirectBuffer = GraphBuilder.CreateBuffer(FRDGBufferDesc::CreateIndirectDesc<FRHIDrawIndirectParameters>(1), TEXT("ShaderDraw.IndirectBuffer"), ERDGBufferFlags::None);
 		{
 			FShaderDrawDebugCopyCS::FParameters* Parameters = GraphBuilder.AllocParameters<FShaderDrawDebugCopyCS::FParameters>();
-			Parameters->ElementBuffer = GraphBuilder.CreateSRV(DataBuffer);
+			Parameters->ElementBuffer = GraphBuilder.CreateSRV(ShaderPrintLineBuffer);
 			Parameters->RWIndirectArgs = GraphBuilder.CreateUAV(IndirectBuffer, PF_R32_UINT);
 
-			TShaderMapRef<FShaderDrawDebugCopyCS> ComputeShader(View.ShaderMap);
+			TShaderMapRef<FShaderDrawDebugCopyCS> ComputeShader(GlobalShaderMap);
 			ClearUnusedGraphResources(ComputeShader, Parameters);
 			GraphBuilder.AddPass(
 				RDG_EVENT_NAME("ShaderPrint::CopyLineArgs"),
@@ -805,20 +852,20 @@ namespace ShaderPrint
 				});
 		}
 
-		TShaderMapRef<FShaderDrawDebugVS> VertexShader(View.ShaderMap);
-		TShaderMapRef<FShaderDrawDebugPS> PixelShader(View.ShaderMap);
+		TShaderMapRef<FShaderDrawDebugVS> VertexShader(GlobalShaderMap);
+		TShaderMapRef<FShaderDrawDebugPS> PixelShader(GlobalShaderMap);
 
 		FShaderDrawVSPSParameters* PassParameters = GraphBuilder.AllocParameters<FShaderDrawVSPSParameters >();
-		PassParameters->VS.View = View.ViewUniformBuffer;
 		PassParameters->VS.TranslatedWorldOffsetConversion = FVector3f(TranslatedWorldOffsetConversion);
+		PassParameters->VS.TranslatedWorldToClip = FMatrix44f(TranslatedWorldToClip);
 		PassParameters->PS.RenderTargets[0] = FRenderTargetBinding(OutputTexture, ERenderTargetLoadAction::ELoad);
-		PassParameters->PS.OutputInvResolution = FVector2f(1.f / View.UnscaledViewRect.Width(), 1.f / View.UnscaledViewRect.Height());
-		PassParameters->PS.OriginalViewRectMin = FVector2f(View.ViewRect.Min);
-		PassParameters->PS.OriginalViewSize = FVector2f(View.ViewRect.Width(), View.ViewRect.Height());
+		PassParameters->PS.OutputInvResolution = FVector2f(1.f / UnscaledViewRect.Width(), 1.f / UnscaledViewRect.Height());
+		PassParameters->PS.OriginalViewRectMin = FVector2f(ViewRect.Min);
+		PassParameters->PS.OriginalViewSize = FVector2f(ViewRect.Width(), ViewRect.Height());
 		PassParameters->PS.OriginalBufferInvSize = FVector2f(1.f / DepthTexture->Desc.Extent.X, 1.f / DepthTexture->Desc.Extent.Y);
 		PassParameters->PS.DepthTexture = DepthTexture;
 		PassParameters->PS.DepthSampler = TStaticSamplerState<SF_Point, AM_Clamp, AM_Clamp, AM_Clamp>::GetRHI();
-		PassParameters->VS.ShaderDrawDebugPrimitive = GraphBuilder.CreateSRV(DataBuffer);
+		PassParameters->VS.ShaderDrawDebugPrimitive = GraphBuilder.CreateSRV(ShaderPrintLineBuffer);
 		PassParameters->VS.IndirectBuffer = IndirectBuffer;
 
 		ValidateShaderParameters(PixelShader, PassParameters->PS);
@@ -826,7 +873,7 @@ namespace ShaderPrint
 		ValidateShaderParameters(VertexShader, PassParameters->VS);
 		ClearUnusedGraphResources(VertexShader, &PassParameters->VS, { IndirectBuffer });
 
-		const FIntRect Viewport = View.UnscaledViewRect;
+		const FIntRect Viewport = UnscaledViewRect;
 		GraphBuilder.AddPass(
 			RDG_EVENT_NAME("ShaderPrint::DrawLines"),
 			PassParameters,
@@ -867,7 +914,7 @@ namespace ShaderPrint
 		// Lines
 		{
 			FRDGBufferRef DataBuffer = View.ShaderPrintData.ShaderPrintLineBuffer;
-			InternalDrawView_Lines(GraphBuilder, View, FVector::ZeroVector, DataBuffer, OutputTexture.Texture, DepthTexture.Texture);
+			InternalDrawView_Lines(GraphBuilder, DataBuffer, View.ViewRect, View.UnscaledViewRect, View.ViewMatrices.GetTranslatedViewProjectionMatrix(), FVector::ZeroVector, OutputTexture.Texture, DepthTexture.Texture);
 		}
 
 		// Locked lines
@@ -875,12 +922,13 @@ namespace ShaderPrint
 		{
 			const FVector LockedToCurrentTranslatedOffset = View.ViewMatrices.GetPreViewTranslation() - View.ViewState->ShaderPrintStateData.PreViewTranslation;
 			FRDGBufferRef DataBuffer = GraphBuilder.RegisterExternalBuffer(View.ViewState->ShaderPrintStateData.LineBuffer);
-			InternalDrawView_Lines(GraphBuilder, View, LockedToCurrentTranslatedOffset, DataBuffer, OutputTexture.Texture, DepthTexture.Texture);
+			InternalDrawView_Lines(GraphBuilder, DataBuffer, View.ViewRect, View.UnscaledViewRect, View.ViewMatrices.GetTranslatedViewProjectionMatrix(), LockedToCurrentTranslatedOffset, OutputTexture.Texture, DepthTexture.Texture);
 		}
 
 		// Characters
 		{
-			InternalDrawView_Characters(GraphBuilder, View, OutputTexture);
+			const int32 FrameNumber = View.Family ? View.Family->FrameNumber : 0u;
+			InternalDrawView_Characters(GraphBuilder, View.ShaderPrintData, View.ViewRect, FrameNumber, OutputTexture);
 		}
 	}
 
diff --git a/Engine/Source/Runtime/Renderer/Private/ShaderPrint.h b/Engine/Source/Runtime/Renderer/Private/ShaderPrint.h
--- a/Engine/Source/Runtime/Renderer/Private/ShaderPrint.h
+++ b/Engine/Source/Runtime/Renderer/Private/ShaderPrint.h
@@ -10,11 +10,9 @@
 
 #pragma once
 
-#include "RenderGraph.h"
-#include "Templates/RefCounting.h"
-#include "ScreenPass.h"
-#include "ShaderPrintParameters.h"
+#include "CoreMinimal.h"
 
+class FRDGBuilder;
 class FViewInfo;
 struct FScreenPassTexture;
 
diff --git a/Engine/Source/Runtime/Renderer/Public/ShaderPrintParameters.h b/Engine/Source/Runtime/Renderer/Public/ShaderPrintParameters.h
--- a/Engine/Source/Runtime/Renderer/Public/ShaderPrintParameters.h
+++ b/Engine/Source/Runtime/Renderer/Public/ShaderPrintParameters.h
@@ -4,33 +4,34 @@
 // This means that the buffer needs to be bound for the shader you wish to debug.
 // It would be ideal if that was automatic (maybe by having a fixed bind point for the buffer and binding it for the entire view).
 // But for now you need to manually add binding information to your FShader class.
-// To do this either:
-// (i) Use SHADER_PARAMETER_STRUCT_INCLUDE(ShaderPrint::FShaderParameters) in your FShader::FParameters declaration and call SetParameters().
-
-// Also it seems that we can only bind a RWBuffer to compute shaders right now. Fixing this would allow us to use this system from all shader stages.
+// To do this use SHADER_PARAMETER_STRUCT_INCLUDE(ShaderPrint::FShaderParameters) in your FShader::FParameters declaration.
+// Then call a variant of SetParameters().
 
 #pragma once
 
 #include "CoreMinimal.h"
+#include "RenderGraphDefinitions.h"
 #include "ShaderParameterMacros.h"
 #include "ShaderParameters.h"
 
-class FViewInfo;
+class FSceneView;
 struct FShaderPrintData;
+class FRDGBuilder;
+class FViewInfo;
 
 namespace ShaderPrint
 {
 	// ShaderPrint uniform buffer layout
 	BEGIN_GLOBAL_SHADER_PARAMETER_STRUCT(FShaderPrintCommonParameters, RENDERER_API)
-		SHADER_PARAMETER(FVector2f, FontSize)
-		SHADER_PARAMETER(FVector2f, FontSpacing)
 		SHADER_PARAMETER(FIntPoint, Resolution)
 		SHADER_PARAMETER(FIntPoint, CursorCoord)
+		SHADER_PARAMETER(FVector3f, TranslatedWorldOffset)
+		SHADER_PARAMETER(FVector2f, FontSize)
+		SHADER_PARAMETER(FVector2f, FontSpacing)
 		SHADER_PARAMETER(uint32, MaxValueCount)
 		SHADER_PARAMETER(uint32, MaxSymbolCount)
 		SHADER_PARAMETER(uint32, MaxStateCount)
 		SHADER_PARAMETER(uint32, MaxLineCount)
-		SHADER_PARAMETER(FVector3f, TranslatedWorldOffset)
 	END_GLOBAL_SHADER_PARAMETER_STRUCT()
 
 	// ShaderPrint parameter struct declaration
@@ -47,16 +48,18 @@ namespace ShaderPrint
 	// Have we enabled the ShaderPrint system?
 	RENDERER_API bool IsEnabled();
 
-	// Call this to know if a view can render this debug information
-	RENDERER_API bool IsEnabled(const FViewInfo& View);
+	// Call this to know if a view can render this debug information.
+	// This should be checked before using a permutation that requires the shader draw parameters.
+	RENDERER_API bool IsEnabled(const FSceneView& View);
 
-	// Returns true if the default view exists and has shader debug rendering enabled (this needs to be checked before using a permutation that requires the shader draw parameters)
+	// Returns true if the default view exists and has shader debug rendering enabled.
+	// This should be checked before using a permutation that requires the shader draw parameters.
 	RENDERER_API bool IsDefaultViewEnabled();
 
-	// Enable/disable shader print
+	// Enable/disable shader print.
 	RENDERER_API void SetEnabled(bool bInEnabled);
 
-	// Set characters font size
+	// Set characters font size.
 	RENDERER_API void SetFontSize(int32 InFontSize);
 
 	/**
@@ -68,29 +71,57 @@ namespace ShaderPrint
 	RENDERER_API void RequestSpaceForCharacters(uint32 MaxElementCount);
 	RENDERER_API void RequestSpaceForLines(uint32 MaxElementCount);
 
-	// Fill the FShaderParameters
-	RENDERER_API void SetParameters(FRDGBuilder& GraphBuilder, FShaderParameters& OutParameters);
-	RENDERER_API void SetParameters(FRDGBuilder& GraphBuilder, const FViewInfo & View, FShaderParameters& OutParameters);
+	/** Structure containing setup for shader print capturing. */
+	struct RENDERER_API FShaderPrintSetup
+	{
+		/** Construct with view and system defaults. */
+		FShaderPrintSetup(FSceneView const& InView);
+		/** Construct with view rectangle and system defaults. */
+		FShaderPrintSetup(FIntRect InViewRect);
+
+		/** Expected viewport rectangle. */
+		FIntRect ViewRect;
+		/** Cursor pixel position within viewport. Can be used for isolating a pixel to debug. */
+		FIntPoint CursorCoord;
+		/** PreView translation used for storing line positions in translated world space. */
+		FVector PreViewTranslation;
+		/** DPI scale to take into account when drawing font. */
+		float DPIScale;
+		/** Font size in pixels. */
+		FIntPoint FontSize;
+		/** Font spacing in pixels (not including font size). */
+		FIntPoint FontSpacing;
+		/** Initial size of character buffer. Will also be increased by RequestSpaceForCharacters(). */
+		uint32 MaxValueCount;
+		/** Initial size of widget buffer. */
+		uint32 MaxStateCount;
+		/** Initial size of line buffer. Will also be increased by RequestSpaceForLines(). */
+		uint32 MaxLineCount;
+	};
+
+	/** Create the shader print render data. This allocates and clears the render buffers. */
+	RENDERER_API FShaderPrintData CreateShaderPrintData(FRDGBuilder& GraphBuilder, FShaderPrintSetup const& InSetup);
+
+	/** Fill the FShaderParameters with an explicit FShaderPrintData managed by the calling code. */
 	RENDERER_API void SetParameters(FRDGBuilder& GraphBuilder, const FShaderPrintData& Data, FShaderParameters& OutParameters);
+	/** Fill the FShaderParameters with the opaque FShaderPrintData from the current default view. */
+	RENDERER_API void SetParameters(FRDGBuilder& GraphBuilder, FShaderParameters& OutParameters);
+	
+	/** Fill the FShaderParameters with opaque FShaderPrintData already bound to the specified view. FViewInfo implies this is private for renderer. */
+	void SetParameters(FRDGBuilder& GraphBuilder, const FViewInfo& View, FShaderParameters& OutParameters);
 }
 
-struct FShaderPrintData
+/** 
+ * Structure containing shader print render data.
+ * This is automatically created, setup and rendered for each view.
+ * Also it is possible for client code to create and own this. 
+ * If this is client managed then the client also needs to send to the shader print renderer.
+ */
+struct RENDERER_API FShaderPrintData
 {
-	FVector2f FontSpacing;
-	FVector2f FontSize;
-	FIntRect OutputRect;
-	FIntPoint CursorCoord = FIntPoint(-1, -1);
-	uint32 MaxValueCount = 0;
-	uint32 MaxSymbolCount = 0;
-	uint32 MaxStateCount = 0;
-	uint32 MaxLineCount = 0u;
-	FVector TranslatedWorldOffset = FVector::ZeroVector;
-
-	bool IsEnabled() const { return MaxValueCount > 0 || MaxSymbolCount > 0 || MaxLineCount > 0; }
-	bool IsValid() const { return ShaderPrintValueBuffer != nullptr; }
+	TUniformBufferRef<ShaderPrint::FShaderPrintCommonParameters> UniformBuffer;
 
 	FRDGBufferRef ShaderPrintValueBuffer = nullptr;
 	FRDGBufferRef ShaderPrintStateBuffer = nullptr;
 	FRDGBufferRef ShaderPrintLineBuffer = nullptr;
-	TUniformBufferRef<ShaderPrint::FShaderPrintCommonParameters> UniformBuffer;
 };
