Fix crash in FViewInfo::SetupUniformBufferParameters (when GetRendererModule().InitializeSystemTextures hasn't been called)

#rnx

diff --git a/Engine/Source/Runtime/Renderer/Private/SceneRendering.cpp b/Engine/Source/Runtime/Renderer/Private/SceneRendering.cpp
--- a/Engine/Source/Runtime/Renderer/Private/SceneRendering.cpp
+++ b/Engine/Source/Runtime/Renderer/Private/SceneRendering.cpp
@@ -1931,12 +1931,19 @@ void FViewInfo::SetupUniformBufferParameters(
 	// Rect light. atlas
 	{
 		FRHITexture* AtlasTexture = RectLightAtlas::GetRectLightAtlasTexture();
-		if (!AtlasTexture) AtlasTexture = GSystemTextures.BlackDummy->GetRHI();
-		const FIntVector AtlasSize = AtlasTexture->GetSizeXYZ();
-		ViewUniformShaderParameters.RectLightAtlasTexture = AtlasTexture;
-		ViewUniformShaderParameters.RectLightAtlasSampler = TStaticSamplerState<SF_Bilinear, AM_Clamp, AM_Clamp, AM_Clamp>::GetRHI();
-		ViewUniformShaderParameters.RectLightAtlasMaxMipLevel = AtlasTexture->GetNumMips()-1;
-		ViewUniformShaderParameters.RectLightAtlasSizeAndInvSize = FVector4f(AtlasSize.X, AtlasSize.Y, 1.0f / AtlasSize.X, 1.0f / AtlasSize.Y);
+		if (!AtlasTexture && GSystemTextures.BlackDummy.IsValid())
+		{
+			AtlasTexture = GSystemTextures.BlackDummy->GetRHI();
+		}
+				
+		if (AtlasTexture)
+		{
+			const FIntVector AtlasSize = AtlasTexture->GetSizeXYZ();
+			ViewUniformShaderParameters.RectLightAtlasTexture = AtlasTexture;
+			ViewUniformShaderParameters.RectLightAtlasSampler = TStaticSamplerState<SF_Bilinear, AM_Clamp, AM_Clamp, AM_Clamp>::GetRHI();
+			ViewUniformShaderParameters.RectLightAtlasMaxMipLevel = AtlasTexture->GetNumMips() - 1;
+			ViewUniformShaderParameters.RectLightAtlasSizeAndInvSize = FVector4f(AtlasSize.X, AtlasSize.Y, 1.0f / AtlasSize.X, 1.0f / AtlasSize.Y);
+		}
 		ViewUniformShaderParameters.RectLightAtlasTexture = OrBlack2DIfNull(ViewUniformShaderParameters.RectLightAtlasTexture);
 	}
 
