Fix sky lighting being incorrect.

The underlying issue was using dynamic sky light data while not valid/initialized.

diff --git a/Engine/Shaders/Private/HairStrands/HairStrandsEnvironmentLighting.usf b/Engine/Shaders/Private/HairStrands/HairStrandsEnvironmentLighting.usf
--- a/Engine/Shaders/Private/HairStrands/HairStrandsEnvironmentLighting.usf
+++ b/Engine/Shaders/Private/HairStrands/HairStrandsEnvironmentLighting.usf
@@ -86,6 +86,7 @@
 ////////////////////////////////////////////////////////////////////////////////////////////////////////////////////////////////////////////////////////////////
 // Parameters
 
+uint				bDynamicSkyLight;
 uint				bHasStaticLighting;
 Texture3D<float4>	HairEnergyLUTTexture;
 
@@ -752,10 +753,16 @@ float3 ReflectionEnvironmentSH(
 	float3 Lighting = 0;
 
 	// Sky lighting (already pre-exposed)
-	Lighting += SkyLightDiffuse(GBuffer, AmbientOcclusion, BufferUV, ScreenPosition, BentNormal, DiffuseColor);
+	if (bDynamicSkyLight)
+	{
+		Lighting += SkyLightDiffuse(GBuffer, AmbientOcclusion, BufferUV, ScreenPosition, BentNormal, DiffuseColor);
+	}
 
 	// Static lighting (already pre-exposed)
-	Lighting += GetStaticLighting(GBuffer, V, AbsoluteWorldPosition);
+	if (bHasStaticLighting)
+	{
+		Lighting += GetStaticLighting(GBuffer, V, AbsoluteWorldPosition);
+	}
 #endif // PERMUTATION_LUMEN
 
 	// Transform NaNs to black, transform negative colors to black.
diff --git a/Engine/Source/Runtime/Renderer/Private/HairStrands/HairStrandsEnvironment.cpp b/Engine/Source/Runtime/Renderer/Private/HairStrands/HairStrandsEnvironment.cpp
--- a/Engine/Source/Runtime/Renderer/Private/HairStrands/HairStrandsEnvironment.cpp
+++ b/Engine/Source/Runtime/Renderer/Private/HairStrands/HairStrandsEnvironment.cpp
@@ -240,6 +240,7 @@ class FHairEnvironmentLightingPS : public FGlobalShader
 		SHADER_PARAMETER_STRUCT_INCLUDE(ShaderPrint::FShaderParameters, ShaderPrintParameters)
 		SHADER_PARAMETER_STRUCT_INCLUDE(FHairStrandsDebugData::FWriteParameters, DebugData)
 
+		SHADER_PARAMETER(uint32, bDynamicSkyLight)
 		SHADER_PARAMETER(uint32, bHasStaticLighting)
 		SHADER_PARAMETER(uint32, MultipleScatterSampleCount)
 		SHADER_PARAMETER(float,  HairDualScatteringRoughnessOverride)
@@ -368,6 +369,7 @@ static void AddHairStrandsEnvironmentLightingPassPS(
 
 	ParametersPS->SceneColorTexture = SceneColorTexture;
 	ParametersPS->HairEnergyLUTTexture = GetHairLUT(GraphBuilder, View, HairLUTType_MeanEnergy);
+	ParametersPS->bDynamicSkyLight = bDynamicSkyLight ? 1 : 0;
 	ParametersPS->bHasStaticLighting = bHasStaticLighting ? 1 : 0;
 	ParametersPS->SkyDiffuseLighting = GetSkyDiffuseLightingParameters(Scene->SkyLight, 1.0f /*DynamicBentNormalAO*/);
 	ParametersPS->HairDistanceThreshold = FMath::Max(GHairStrandsSkyLightingDistanceThreshold, 1.f);
