Hack fix for crashes related to multi-view-family rendering refactor.  Need to handle FViewInfo::SetupUniformBufferParameters called on manually constructed FViewInfo structures in a couple code paths outside the scene renderer (FRendererModule::DrawTileMesh, FNiagaraGpuComputeDispatch::ProcessPendingTicksFlush).  I plan a better code fix which will involve avoiding access to the global variable FSceneTexturesConfig, which isn't multi-view-family safe, and instead computing the data when it isn't available from the view family.

#rnx

diff --git a/Engine/Source/Runtime/Engine/Public/SceneView.h b/Engine/Source/Runtime/Engine/Public/SceneView.h
--- a/Engine/Source/Runtime/Engine/Public/SceneView.h
+++ b/Engine/Source/Runtime/Engine/Public/SceneView.h
@@ -1750,6 +1750,7 @@ public:
 	bool bIsFirstViewInMultipleViewFamily = true;
 
 	bool bIsSceneTexturesInitialized = false;
+	bool bIsViewFamilyInfo = false;
 
 	/** 
 	 * Which component of the scene rendering should be output to the final render target.
diff --git a/Engine/Source/Runtime/Renderer/Private/SceneRendering.cpp b/Engine/Source/Runtime/Renderer/Private/SceneRendering.cpp
--- a/Engine/Source/Runtime/Renderer/Private/SceneRendering.cpp
+++ b/Engine/Source/Runtime/Renderer/Private/SceneRendering.cpp
@@ -2476,6 +2476,7 @@ FViewFamilyInfo::FViewFamilyInfo(const FSceneViewFamily& InViewFamily)
 ,	bShadowDepthRenderCompleted(false)
 ,	FamilySize(0, 0)
 {
+	bIsViewFamilyInfo = true;
 }
 
 FViewFamilyInfo::~FViewFamilyInfo()
diff --git a/Engine/Source/Runtime/Renderer/Private/SceneRendering.h b/Engine/Source/Runtime/Renderer/Private/SceneRendering.h
--- a/Engine/Source/Runtime/Renderer/Private/SceneRendering.h
+++ b/Engine/Source/Runtime/Renderer/Private/SceneRendering.h
@@ -2638,7 +2638,8 @@ FRDGTextureRef CreateHalfResolutionDepthCheckerboardMinMax(FRDGBuilder& GraphBui
 
 inline const FSceneTexturesConfig& FViewInfo::GetSceneTexturesConfig() const
 {
-	return ((const FViewFamilyInfo*)Family)->SceneTexturesConfig;
+	// TODO:  We are refactoring away use of the FSceneTexturesConfig::Get() global singleton, but need this workaround for now to avoid crashes
+	return Family->bIsViewFamilyInfo ? ((const FViewFamilyInfo*)Family)->SceneTexturesConfig : FSceneTexturesConfig::Get();
 }
 
 inline const FSceneTextures& FViewInfo::GetSceneTextures() const
