Fix for macOS VRTemplate project crashes when adding a const 4 in a new material.: Use Resolve scene textures in CopySceneToCubeFace() as Target can cause MSAA texture type binding issues in Metal.   Part 2 of this fix has already been addressed in CL 19857000.

diff --git a/Engine/Source/Runtime/Renderer/Private/ReflectionEnvironmentCapture.cpp b/Engine/Source/Runtime/Renderer/Private/ReflectionEnvironmentCapture.cpp
--- a/Engine/Source/Runtime/Renderer/Private/ReflectionEnvironmentCapture.cpp
+++ b/Engine/Source/Runtime/Renderer/Private/ReflectionEnvironmentCapture.cpp
@@ -545,9 +545,9 @@ void CaptureSceneToScratchCubemap(
 
 		PassParameters->View = View.ViewUniformBuffer;
 		PassParameters->SceneColorSampler = TStaticSamplerState<SF_Bilinear, AM_Clamp, AM_Clamp, AM_Clamp>::GetRHI();
-		PassParameters->SceneColorTexture = SceneTextures.Color.Target;
+		PassParameters->SceneColorTexture = SceneTextures.Color.Resolve;
 		PassParameters->SceneDepthSampler = TStaticSamplerState<SF_Point, AM_Clamp, AM_Clamp, AM_Clamp>::GetRHI();
-		PassParameters->SceneDepthTexture = SceneTextures.Depth.Target;
+		PassParameters->SceneDepthTexture = SceneTextures.Depth.Resolve;
 
 		const int32 EffectiveSize = CubemapSize;
 		const FIntPoint SceneTextureExtent = SceneTextures.Config.Extent;
