45s to 35s when booting the editor with cold cache by using PrecompileWithDXC on DOFGatherPass

#rnx

diff --git a/Engine/Source/Runtime/Renderer/Private/PostProcess/DiaphragmDOF.cpp b/Engine/Source/Runtime/Renderer/Private/PostProcess/DiaphragmDOF.cpp
--- a/Engine/Source/Runtime/Renderer/Private/PostProcess/DiaphragmDOF.cpp
+++ b/Engine/Source/Runtime/Renderer/Private/PostProcess/DiaphragmDOF.cpp
@@ -1064,14 +1064,10 @@ class FDiaphragmDOFGatherCS : public FDiaphragmDOFShader
 
 	static void ModifyCompilationEnvironment(const FGlobalShaderPermutationParameters& Parameters, FShaderCompilerEnvironment& OutEnvironment)
 	{
-		// The gathering pass shader code gives a really hard time to the HLSL compiler. To improve
-		// iteration time on the shader, only pass down a /O1 instead of /O3.
-		if (Parameters.Platform == SP_PCD3D_SM5)
-		{
-			OutEnvironment.CompilerFlags.Add(CFLAG_StandardOptimization);
-		}
-
 		FDiaphragmDOFShader::ModifyCompilationEnvironment(Parameters, OutEnvironment);
+
+		// This shader takes a very long time to compile with FXC, so we pre-compile it with DXC first and then forward the optimized HLSL to FXC.
+		OutEnvironment.CompilerFlags.Add(CFLAG_PrecompileWithDXC);
 	}
 
 	
