Fixed new instances of reflection capture uniforms not using mobile ub

diff --git a/Engine/Source/Runtime/Renderer/Private/MobileShadingRenderer.cpp b/Engine/Source/Runtime/Renderer/Private/MobileShadingRenderer.cpp
--- a/Engine/Source/Runtime/Renderer/Private/MobileShadingRenderer.cpp
+++ b/Engine/Source/Runtime/Renderer/Private/MobileShadingRenderer.cpp
@@ -172,7 +172,7 @@ static void RenderOpaqueFX(
 BEGIN_SHADER_PARAMETER_STRUCT(FMobileRenderPassParameters, RENDERER_API)
 	SHADER_PARAMETER_STRUCT_INCLUDE(FViewShaderParameters, View)
 	SHADER_PARAMETER_RDG_UNIFORM_BUFFER(FMobileBasePassUniformParameters, MobileBasePass)
-	SHADER_PARAMETER_STRUCT_REF(FReflectionCaptureShaderData, ReflectionCapture)
+	SHADER_PARAMETER_STRUCT_REF(FMobileReflectionCaptureShaderData, ReflectionCapture)
 	RDG_BUFFER_ACCESS_ARRAY(DrawIndirectArgsBuffers)
 	RDG_BUFFER_ACCESS_ARRAY(InstanceIdOffsetBuffers)
 	RENDER_TARGET_BINDING_SLOTS()
@@ -1081,7 +1081,7 @@ void FMobileSceneRenderer::RenderForward(FRDGBuilder& GraphBuilder, FRDGTextureR
 		FMobileRenderPassParameters* PassParameters = GraphBuilder.AllocParameters<FMobileRenderPassParameters>();
 		PassParameters->View = View.GetShaderParameters();
 		PassParameters->MobileBasePass = CreateMobileBasePassUniformBuffer(GraphBuilder, View, EMobileBasePass::Opaque, SetupMode, MobileBasePassTextures);
-		PassParameters->ReflectionCapture = View.ReflectionCaptureUniformBuffer;
+		PassParameters->ReflectionCapture = View.MobileReflectionCaptureUniformBuffer;
 		PassParameters->RenderTargets = BasePassRenderTargets;
 	
 		BuildInstanceCullingDrawParams(GraphBuilder, View, PassParameters);
@@ -1200,7 +1200,7 @@ void FMobileSceneRenderer::RenderForwardMultiPass(FRDGBuilder& GraphBuilder, FMo
 	FMobileRenderPassParameters* SecondPassParameters = GraphBuilder.AllocParameters<FMobileRenderPassParameters>();
 	*SecondPassParameters = *PassParameters;
 	SecondPassParameters->MobileBasePass = CreateMobileBasePassUniformBuffer(GraphBuilder, View, EMobileBasePass::Translucent, SetupMode);
-	SecondPassParameters->ReflectionCapture = View.ReflectionCaptureUniformBuffer;
+	SecondPassParameters->ReflectionCapture = View.MobileReflectionCaptureUniformBuffer;
 	SecondPassParameters->RenderTargets = BasePassRenderTargets;
 
 	GraphBuilder.AddPass(
@@ -1375,7 +1375,7 @@ void FMobileSceneRenderer::RenderDeferred(FRDGBuilder& GraphBuilder, const FSort
 		auto* PassParameters = GraphBuilder.AllocParameters<FMobileRenderPassParameters>();
 		PassParameters->View = View.GetShaderParameters();
 		PassParameters->MobileBasePass = CreateMobileBasePassUniformBuffer(GraphBuilder, View, EMobileBasePass::Opaque, SetupMode, MobileBasePassTextures);
-		PassParameters->ReflectionCapture = View.ReflectionCaptureUniformBuffer;
+		PassParameters->ReflectionCapture = View.MobileReflectionCaptureUniformBuffer;
 		PassParameters->RenderTargets = BasePassRenderTargets;
 		
 		BuildInstanceCullingDrawParams(GraphBuilder, View, PassParameters);
@@ -1473,7 +1473,7 @@ void FMobileSceneRenderer::RenderDeferredMultiPass(FRDGBuilder& GraphBuilder, cl
 	auto* SecondPassParameters = GraphBuilder.AllocParameters<FMobileRenderPassParameters>();
 	*SecondPassParameters = *PassParameters;
 	SecondPassParameters->MobileBasePass = CreateMobileBasePassUniformBuffer(GraphBuilder, View, EMobileBasePass::Translucent, SetupMode);
-	SecondPassParameters->ReflectionCapture = View.ReflectionCaptureUniformBuffer;
+	SecondPassParameters->ReflectionCapture = View.MobileReflectionCaptureUniformBuffer;
 	SecondPassParameters->RenderTargets = BasePassRenderTargets;
 
 	GraphBuilder.AddPass(
@@ -1497,7 +1497,7 @@ void FMobileSceneRenderer::RenderDeferredMultiPass(FRDGBuilder& GraphBuilder, cl
 	auto* ThirdPassParameters = GraphBuilder.AllocParameters<FMobileRenderPassParameters>();
 	*ThirdPassParameters = *PassParameters;
 	ThirdPassParameters->MobileBasePass = CreateMobileBasePassUniformBuffer(GraphBuilder, View, EMobileBasePass::Translucent, SetupMode);
-	ThirdPassParameters->ReflectionCapture = View.ReflectionCaptureUniformBuffer;
+	ThirdPassParameters->ReflectionCapture = View.MobileReflectionCaptureUniformBuffer;
 	ThirdPassParameters->RenderTargets = BasePassRenderTargets;
 
 	GraphBuilder.AddPass(
diff --git a/Engine/Source/Runtime/Renderer/Private/SceneVisibility.cpp b/Engine/Source/Runtime/Renderer/Private/SceneVisibility.cpp
--- a/Engine/Source/Runtime/Renderer/Private/SceneVisibility.cpp
+++ b/Engine/Source/Runtime/Renderer/Private/SceneVisibility.cpp
@@ -5008,9 +5008,8 @@ void FSceneRenderer::SetupSceneReflectionCaptureBuffer(FRHICommandListImmediate&
 
 	TUniformBufferRef<FMobileReflectionCaptureShaderData> MobileReflectionCaptureUniformBuffer;
 	TUniformBufferRef<FReflectionCaptureShaderData> ReflectionCaptureUniformBuffer;
-	bool bIsMobileClusteredReflections = IsMobilePlatform(ShaderPlatform) && MobileEnableClusteredReflections(ShaderPlatform);
 
-	if (bIsMobileClusteredReflections)
+	if (IsMobilePlatform(ShaderPlatform))
 	{
 		CreateReflectionCaptureUniformBuffer(SortedCaptures, MobileReflectionCaptureUniformBuffer);
 	}
@@ -5023,7 +5022,7 @@ void FSceneRenderer::SetupSceneReflectionCaptureBuffer(FRHICommandListImmediate&
 	{
 		FViewInfo& View = Views[ViewIndex];
 
-		if (bIsMobileClusteredReflections)
+		if (IsMobilePlatform(ShaderPlatform))
 		{
 			View.MobileReflectionCaptureUniformBuffer = MobileReflectionCaptureUniformBuffer;
 		}
