Change shading energy table to use dynamic loop instead of unroll.

This  reduces compilation time with fxc from 16s to <1s.

diff --git a/Engine/Shaders/Private/ShadingEnergyConservationTable.usf b/Engine/Shaders/Private/ShadingEnergyConservationTable.usf
--- a/Engine/Shaders/Private/ShadingEnergyConservationTable.usf
+++ b/Engine/Shaders/Private/ShadingEnergyConservationTable.usf
@@ -14,6 +14,8 @@
 // "Energy" here refers to the directional albedo. Energy loss is simply 1-Energy since we assume the raw BxDF should reflect all
 // incoming light.
 
+// Number of samples used for integrating directional albedo 
+uint NumSamples;
 // x: directional albedo for rough diffuse
 RWTexture2D<float>  Output1Texture2D;
 // x: directional albedo of pure GGX lobe
@@ -40,7 +42,8 @@ void BuildEnergyTableCS(uint3 DispatchThreadId : SV_DispatchThreadID)
 
 	RandomSequence RandSequence;
 	float2 Result = 0;
-	const int NumSamples = 1 << 14; // use plenty of samples -- this is computed on startup once and at fairly low resolution
+	// use plenty of samples -- this is computed on startup once and at fairly low resolution
+	LOOP
 	for (int Sample = 0; Sample < NumSamples; Sample++)
 	{
 		// draw a random sample
diff --git a/Engine/Source/Runtime/Renderer/Private/ShadingEnergyConservation.cpp b/Engine/Source/Runtime/Renderer/Private/ShadingEnergyConservation.cpp
--- a/Engine/Source/Runtime/Renderer/Private/ShadingEnergyConservation.cpp
+++ b/Engine/Source/Runtime/Renderer/Private/ShadingEnergyConservation.cpp
@@ -192,6 +192,7 @@ class FBuildShadingEnergyConservationTableCS : public FGlobalShader
 	}
 
 	BEGIN_SHADER_PARAMETER_STRUCT(FParameters, )
+		SHADER_PARAMETER(uint32, NumSamples)
 		SHADER_PARAMETER_RDG_TEXTURE_UAV(RWTexture2D, Output1Texture2D)
 		SHADER_PARAMETER_RDG_TEXTURE_UAV(RWTexture2D, Output2Texture2D)
 		SHADER_PARAMETER_RDG_TEXTURE_UAV(RWTexture3D, OutputTexture3D)
@@ -255,12 +256,15 @@ void Init(FRDGBuilder& GraphBuilder, FViewInfo& View)
 				ClothEnergyTexture		= GraphBuilder.CreateTexture(FRDGTextureDesc::Create2D(FIntPoint(Size, Size),        SpecFormat, FClearValueBinding::None, TexCreate_ShaderResource | TexCreate_UAV), TEXT("Shading.ClothSpecEnergy"), ERDGTextureFlags::MultiFrame);
 				DiffuseEnergyTexture	= GraphBuilder.CreateTexture(FRDGTextureDesc::Create2D(FIntPoint(Size, Size),        DiffFormat, FClearValueBinding::None, TexCreate_ShaderResource | TexCreate_UAV), TEXT("Shading.DiffuseEnergy"),   ERDGTextureFlags::MultiFrame);
 			
+				const uint32 NumSamples = 1u << 14u;
+
 				// GGX
 				{
 					FBuildShadingEnergyConservationTableCS::FPermutationDomain PermutationVector;
 					PermutationVector.Set<FBuildShadingEnergyConservationTableCS::FEnergyTableDim>(FBuildShadingEnergyConservationTableCS::EEnergyTableType::GGXSpecular);
 					TShaderMapRef<FBuildShadingEnergyConservationTableCS> ComputeShader(View.ShaderMap, PermutationVector);
 					FBuildShadingEnergyConservationTableCS::FParameters* PassParameters = GraphBuilder.AllocParameters<FBuildShadingEnergyConservationTableCS::FParameters>();
+					PassParameters->NumSamples = NumSamples;
 					PassParameters->Output2Texture2D = GraphBuilder.CreateUAV(FRDGTextureUAVDesc(GGXSpecEnergyTexture, 0));
 					FComputeShaderUtils::AddPass(
 						GraphBuilder,
@@ -276,6 +280,7 @@ void Init(FRDGBuilder& GraphBuilder, FViewInfo& View)
 					PermutationVector.Set<FBuildShadingEnergyConservationTableCS::FEnergyTableDim>(FBuildShadingEnergyConservationTableCS::EEnergyTableType::GGXGlass);
 					TShaderMapRef<FBuildShadingEnergyConservationTableCS> ComputeShader(View.ShaderMap, PermutationVector);
 					FBuildShadingEnergyConservationTableCS::FParameters* PassParameters = GraphBuilder.AllocParameters<FBuildShadingEnergyConservationTableCS::FParameters>();
+					PassParameters->NumSamples = NumSamples;
 					PassParameters->OutputTexture3D = GraphBuilder.CreateUAV(FRDGTextureUAVDesc(GGXGlassEnergyTexture, 0));
 					FComputeShaderUtils::AddPass(
 						GraphBuilder,
@@ -291,6 +296,7 @@ void Init(FRDGBuilder& GraphBuilder, FViewInfo& View)
 					PermutationVector.Set<FBuildShadingEnergyConservationTableCS::FEnergyTableDim>(FBuildShadingEnergyConservationTableCS::EEnergyTableType::Cloth);
 					TShaderMapRef<FBuildShadingEnergyConservationTableCS> ComputeShader(View.ShaderMap, PermutationVector);
 					FBuildShadingEnergyConservationTableCS::FParameters* PassParameters = GraphBuilder.AllocParameters<FBuildShadingEnergyConservationTableCS::FParameters>();
+					PassParameters->NumSamples = NumSamples;
 					PassParameters->Output2Texture2D = GraphBuilder.CreateUAV(FRDGTextureUAVDesc(ClothEnergyTexture, 0));
 					FComputeShaderUtils::AddPass(
 						GraphBuilder,
@@ -306,6 +312,7 @@ void Init(FRDGBuilder& GraphBuilder, FViewInfo& View)
 					PermutationVector.Set<FBuildShadingEnergyConservationTableCS::FEnergyTableDim>(FBuildShadingEnergyConservationTableCS::EEnergyTableType::Diffuse);
 					TShaderMapRef<FBuildShadingEnergyConservationTableCS> ComputeShader(View.ShaderMap, PermutationVector);
 					FBuildShadingEnergyConservationTableCS::FParameters* PassParameters = GraphBuilder.AllocParameters<FBuildShadingEnergyConservationTableCS::FParameters>();
+					PassParameters->NumSamples = NumSamples;
 					PassParameters->Output1Texture2D = GraphBuilder.CreateUAV(FRDGTextureUAVDesc(DiffuseEnergyTexture, 0));
 					FComputeShaderUtils::AddPass(
 						GraphBuilder,
