Fix the crash when building lighting with a reflection capture in the scene.

diff --git a/Engine/Source/Runtime/Renderer/Private/ReflectionEnvironment.cpp b/Engine/Source/Runtime/Renderer/Private/ReflectionEnvironment.cpp
--- a/Engine/Source/Runtime/Renderer/Private/ReflectionEnvironment.cpp
+++ b/Engine/Source/Runtime/Renderer/Private/ReflectionEnvironment.cpp
@@ -213,7 +213,19 @@ const FCaptureComponentSceneState* FReflectionCaptureCache::Find(const UReflecti
 	if (!Component)	// Intentionally not IsValid(Component), as this often occurs when Component is explicitly PendingKill.
 		return nullptr;
 
-	return Find(Component->MapBuildDataId);
+	const FCaptureComponentSceneState* SceneState = Find(Component->MapBuildDataId);
+	if (SceneState)
+	{
+		return SceneState;
+	}
+
+	const FGuid* Guid = RegisteredComponentMapBuildDataIds.Find(Component);
+	if (Guid)
+	{
+		return Find(*Guid);
+	}
+
+	return nullptr;
 }
 
 FCaptureComponentSceneState* FReflectionCaptureCache::Find(const UReflectionCaptureComponent* Component)
@@ -221,7 +233,18 @@ FCaptureComponentSceneState* FReflectionCaptureCache::Find(const UReflectionCapt
 	if (!Component)	// Intentionally not IsValid(Component), as this often occurs when Component is explicitly PendingKill.
 		return nullptr;
 
-	return Find(Component->MapBuildDataId);
+	FCaptureComponentSceneState* SceneState = Find(Component->MapBuildDataId);
+	if (SceneState)
+	{
+		return SceneState;
+	}
+
+	const FGuid* Guid = RegisteredComponentMapBuildDataIds.Find(Component);
+	if (Guid)
+	{
+		return Find(*Guid);
+	}
+	return nullptr;
 }
 
 
@@ -253,6 +276,7 @@ FCaptureComponentSceneState& FReflectionCaptureCache::Add(const UReflectionCaptu
 	else
 	{
 		FReflectionCaptureCacheEntry& Entry = CaptureData.Add(Component->MapBuildDataId, { 1, Value });
+		RegisterComponentMapBuildDataId(Component);
 		return Entry.SceneState;
 	}
 }
@@ -261,6 +285,7 @@ FCaptureComponentSceneState* FReflectionCaptureCache::AddReference(const UReflec
 {
 	check(IsValid(Component))
 
+	RemapRegisteredComponentMapBuildDataId(Component);
 	FReflectionCaptureCacheEntry* Found = CaptureData.Find(Component->MapBuildDataId);
 	if (Found == nullptr)
 		return nullptr;
@@ -282,6 +307,7 @@ int32 FReflectionCaptureCache::GetKeys(TSet<FGuid>& OutKeys) const
 void FReflectionCaptureCache::Empty()
 {
 	CaptureData.Empty();
+	RegisteredComponentMapBuildDataIds.Empty();
 }
 
 bool FReflectionCaptureCache::Remove(const UReflectionCaptureComponent* Component)
@@ -289,6 +315,7 @@ bool FReflectionCaptureCache::Remove(const UReflectionCaptureComponent* Componen
 	if (!Component)	// Intentionally not IsValid(Component), as this often occurs when Component is explicitly PendingKill.
 		return false;
 
+	RemapRegisteredComponentMapBuildDataId(Component);
 	FReflectionCaptureCacheEntry* Found = CaptureData.Find(Component->MapBuildDataId);
 	if (Found == nullptr)
 		return false;
@@ -301,10 +328,39 @@ bool FReflectionCaptureCache::Remove(const UReflectionCaptureComponent* Componen
 	else
 	{
 		CaptureData.Remove(Component->MapBuildDataId);
+		UnregisterComponentMapBuildDataId(Component);
 		return true;
 	}
 }
 
+void FReflectionCaptureCache::RegisterComponentMapBuildDataId(const UReflectionCaptureComponent* Component)
+{
+	RegisteredComponentMapBuildDataIds.Add(Component, Component->MapBuildDataId);
+}
+
+void FReflectionCaptureCache::RemapRegisteredComponentMapBuildDataId(const UReflectionCaptureComponent* Component)
+{
+	check(Component);
+
+	// Remap old guid to new guid when the component is the same pointer.
+	FGuid* OldBuildId = RegisteredComponentMapBuildDataIds.Find(Component);
+	if (OldBuildId &&
+		*OldBuildId != Component->MapBuildDataId)
+	{
+		FReflectionCaptureCacheEntry Entry = *CaptureData.Find(*OldBuildId);
+		CaptureData.Remove(*OldBuildId);
+		CaptureData.Shrink();
+		CaptureData.Add(Component->MapBuildDataId, Entry);
+
+		RegisteredComponentMapBuildDataIds.Add(Component, Component->MapBuildDataId);
+	}
+}
+
+void FReflectionCaptureCache::UnregisterComponentMapBuildDataId(const UReflectionCaptureComponent* Component)
+{
+	RegisteredComponentMapBuildDataIds.Remove(Component);
+}
+
 void FReflectionEnvironmentSceneData::SetGameThreadTrackingData(int32 MaxAllocatedCubemaps, int32 CaptureSize)
 {
 	MaxAllocatedReflectionCubemapsGameThread = MaxAllocatedCubemaps;
diff --git a/Engine/Source/Runtime/Renderer/Private/ScenePrivate.h b/Engine/Source/Runtime/Renderer/Private/ScenePrivate.h
--- a/Engine/Source/Runtime/Renderer/Private/ScenePrivate.h
+++ b/Engine/Source/Runtime/Renderer/Private/ScenePrivate.h
@@ -1632,6 +1632,15 @@ public:
 	void Empty();
 
 protected:
+
+	void RemapRegisteredComponentMapBuildDataId(const UReflectionCaptureComponent* Component);
+	void RegisterComponentMapBuildDataId(const UReflectionCaptureComponent* Component);
+	void UnregisterComponentMapBuildDataId(const UReflectionCaptureComponent* Component);
+
+	// Different map build data id of a capture might share the same capture component while editing (e.g., when they move).
+	// need to replace it with the new one.
+	TMap<const UReflectionCaptureComponent*, FGuid> RegisteredComponentMapBuildDataIds;
+
 	TMap<FGuid, FReflectionCaptureCacheEntry> CaptureData;
 };
 
