Fix ray tracing pipeline initialization when r.RayTracing.Shadows=0 but there is a light in the scene that explicitly enables RT shadows.

diff --git a/Engine/Source/Runtime/Renderer/Private/DeferredShadingRenderer.cpp b/Engine/Source/Runtime/Renderer/Private/DeferredShadingRenderer.cpp
--- a/Engine/Source/Runtime/Renderer/Private/DeferredShadingRenderer.cpp
+++ b/Engine/Source/Runtime/Renderer/Private/DeferredShadingRenderer.cpp
@@ -1458,7 +1458,7 @@ bool FDeferredShadingSceneRenderer::SetupRayTracingPipelineStates(FRHICommandLis
 				{
 					PrepareRayTracingReflections(View, *Scene, RayGenShaders);
 					PrepareSingleLayerWaterRayTracingReflections(View, *Scene, RayGenShaders);
-					PrepareRayTracingShadows(View, RayGenShaders);
+					PrepareRayTracingShadows(View, *Scene, RayGenShaders);
 					PrepareRayTracingAmbientOcclusion(View, RayGenShaders);
 					PrepareRayTracingSkyLight(View, *Scene, RayGenShaders);
 					PrepareRayTracingGlobalIllumination(View, RayGenShaders);
diff --git a/Engine/Source/Runtime/Renderer/Private/DeferredShadingRenderer.h b/Engine/Source/Runtime/Renderer/Private/DeferredShadingRenderer.h
--- a/Engine/Source/Runtime/Renderer/Private/DeferredShadingRenderer.h
+++ b/Engine/Source/Runtime/Renderer/Private/DeferredShadingRenderer.h
@@ -1021,7 +1021,7 @@ private:
 	static void PrepareRayTracingReflections(const FViewInfo& View, const FScene& Scene, TArray<FRHIRayTracingShader*>& OutRayGenShaders);
 	static void PrepareRayTracingDeferredReflections(const FViewInfo& View, const FScene& Scene, TArray<FRHIRayTracingShader*>& OutRayGenShaders);
 	static void PrepareSingleLayerWaterRayTracingReflections(const FViewInfo& View, const FScene& Scene, TArray<FRHIRayTracingShader*>& OutRayGenShaders);
-	static void PrepareRayTracingShadows(const FViewInfo& View, TArray<FRHIRayTracingShader*>& OutRayGenShaders);
+	static void PrepareRayTracingShadows(const FViewInfo& View, const FScene& Scene, TArray<FRHIRayTracingShader*>& OutRayGenShaders);
 	static void PrepareRayTracingAmbientOcclusion(const FViewInfo& View, TArray<FRHIRayTracingShader*>& OutRayGenShaders);
 	static void PrepareRayTracingSkyLight(const FViewInfo& View, const FScene& Scene, TArray<FRHIRayTracingShader*>& OutRayGenShaders);
 	static void PrepareRayTracingGlobalIllumination(const FViewInfo& View, TArray<FRHIRayTracingShader*>& OutRayGenShaders);
diff --git a/Engine/Source/Runtime/Renderer/Private/RayTracing/RayTracingShadows.cpp b/Engine/Source/Runtime/Renderer/Private/RayTracing/RayTracingShadows.cpp
--- a/Engine/Source/Runtime/Renderer/Private/RayTracing/RayTracingShadows.cpp
+++ b/Engine/Source/Runtime/Renderer/Private/RayTracing/RayTracingShadows.cpp
@@ -185,12 +185,13 @@ float GetRaytracingMaxNormalBias()
 	return FMath::Max(0.01f, GRayTracingMaxNormalBias);
 }
 
-void FDeferredShadingSceneRenderer::PrepareRayTracingShadows(const FViewInfo& View, TArray<FRHIRayTracingShader*>& OutRayGenShaders)
+void FDeferredShadingSceneRenderer::PrepareRayTracingShadows(const FViewInfo& View, const FScene& Scene, TArray<FRHIRayTracingShader*>& OutRayGenShaders)
 {
-	// Declare all RayGen shaders that require material closest hit shaders to be bound
 	// Ray tracing shadows shaders should be properly configured even if r.RayTracing.Shadows is 0 because lights can have raytracing shadows enabled independently of that CVar
+	// We have to check if ray tracing is enabled on any of the scene lights. The RayTracedLights array is populated using ShouldRenderRayTracingShadowsForLight() helper, 
+	// which handles various override conditions.
 
-	if (!ShouldRenderRayTracingShadows())
+	if (Scene.RayTracedLights.IsEmpty())
 	{
 		return;
 	}
