Fix rendering artifacts with real time reflections and volumetric clouds.

#rnx

diff --git a/Engine/Source/Runtime/Renderer/Private/ReflectionEnvironmentRealTimeCapture.cpp b/Engine/Source/Runtime/Renderer/Private/ReflectionEnvironmentRealTimeCapture.cpp
--- a/Engine/Source/Runtime/Renderer/Private/ReflectionEnvironmentRealTimeCapture.cpp
+++ b/Engine/Source/Runtime/Renderer/Private/ReflectionEnvironmentRealTimeCapture.cpp
@@ -476,7 +476,7 @@ void FScene::AllocateAndCaptureFrameSkyEnvMap(
 			{
 				RDG_EVENT_SCOPE(GraphBuilder, "Capture Face=%d", CubeFace);
 
-				SkyRC.RenderTargets[0] = FRenderTargetBinding(SkyCubeTexture, ERenderTargetLoadAction::ENoAction, 0, CubeFace);
+				SkyRC.RenderTargets[0] = FRenderTargetBinding(SkyCubeTexture, ERenderTargetLoadAction::ELoad, 0, CubeFace);
 
 				const FMatrix CubeViewRotationMatrix = CalcCubeFaceViewRotationMatrix((ECubeFace)CubeFace);
 
