Fix crash when forcing all ray tracing effects off via r.RayTracing.ForceAllRayTracingEffects=0 while a scene contains a light with RT shadows

Setting ForceAllRayTracingEffects=0 is expected to work as if all ray tracing effects (including r.RayTracing.Shadows) were set to 0 and none of the lights should use ray tracing.

When ray tracing pipeline is being initialized in FDeferredShadingSceneRenderer::SetupRayTracingPipelineStates(), the required ray generation shaders are not added by PrepareRayTracingShadows() because it only checks the global RT shadow setting, but not per-light override. This can lead to ray tracing pipeline not getting initialized because the code believes that there are no active RT effects, which leads to a crash later during light rendering.

This change makes the per-light RT shadow state determination respect the global override.

diff --git a/Engine/Source/Runtime/Renderer/Private/LightRendering.cpp b/Engine/Source/Runtime/Renderer/Private/LightRendering.cpp
--- a/Engine/Source/Runtime/Renderer/Private/LightRendering.cpp
+++ b/Engine/Source/Runtime/Renderer/Private/LightRendering.cpp
@@ -167,16 +167,18 @@ bool ShouldRenderRayTracingShadows()
 
 bool ShouldRenderRayTracingShadowsForLight(const FLightSceneProxy& LightProxy)
 {
+	const bool bShadowRayTracingAllowed = ShouldRenderRayTracingEffect(true, ERayTracingPipelineCompatibilityFlags::FullPipeline, nullptr);
 	return (LightProxy.CastsRaytracedShadow() == ECastRayTracedShadow::Enabled || (ShouldRenderRayTracingShadows() && LightProxy.CastsRaytracedShadow() == ECastRayTracedShadow::UseProjectSetting))
 		&& ShouldRenderRayTracingShadowsForLightType((ELightComponentType)LightProxy.GetLightType())
-		&& IsRayTracingEnabled();
+		&& bShadowRayTracingAllowed;
 }
 
 bool ShouldRenderRayTracingShadowsForLight(const FLightSceneInfoCompact& LightInfo)
 {
+	const bool bShadowRayTracingAllowed = ShouldRenderRayTracingEffect(true, ERayTracingPipelineCompatibilityFlags::FullPipeline, nullptr);
 	return (LightInfo.CastRaytracedShadow == ECastRayTracedShadow::Enabled || (ShouldRenderRayTracingShadows() && LightInfo.CastRaytracedShadow == ECastRayTracedShadow::UseProjectSetting))
 		&& ShouldRenderRayTracingShadowsForLightType((ELightComponentType)LightInfo.LightType)
-		&& IsRayTracingEnabled();
+		&& bShadowRayTracingAllowed;
 }
 #endif // RHI_RAYTRACING
 
