Path Tracer: Implement alpha holdout blendmode properly

This should render black, with alpha controlled by the opacity.

diff --git a/Engine/Shaders/Private/PathTracing/Material/PathTracingMaterialSampling.ush b/Engine/Shaders/Private/PathTracing/Material/PathTracingMaterialSampling.ush
--- a/Engine/Shaders/Private/PathTracing/Material/PathTracingMaterialSampling.ush
+++ b/Engine/Shaders/Private/PathTracing/Material/PathTracingMaterialSampling.ush
@@ -90,9 +90,8 @@ float3 GetMaterialTransparency(FPathTracingPayload Payload, float3 RayDirection)
 		case RAY_TRACING_BLEND_MODE_MODULATE:
 			return Payload.Radiance;
 		case RAY_TRACING_BLEND_MODE_ALPHA_COMPOSITE:
-			return 1 - Payload.Opacity;
 		case RAY_TRACING_BLEND_MODE_ALPHA_HOLDOUT:
-			return 1;
+			return 1 - Payload.Opacity;
 		default: 
 			return 0;
 	}
@@ -123,7 +122,7 @@ float3 EstimateMaterialAlbedo(FPathTracingPayload Payload)
 			break;
 		case RAY_TRACING_BLEND_MODE_MODULATE:
 		case RAY_TRACING_BLEND_MODE_ALPHA_HOLDOUT:
-			// In these modes, the light does not interact with the material at all, only passes through
+			// In these modes, the light does not interact with the material at all
 			return 0;
 		default:
 			break;
@@ -209,7 +208,7 @@ void AccumulateAlbedoNormal(inout FPathTracingPayload Payload, float3 PathThroug
 			break;
 		case RAY_TRACING_BLEND_MODE_MODULATE:
 		case RAY_TRACING_BLEND_MODE_ALPHA_HOLDOUT:
-			// In these modes, the light does not interact with the material at all, only passes through
+			// In these modes, the light does not interact with the material at all
 			return;
 		default:
 			break;
diff --git a/Engine/Shaders/Private/PathTracing/PathTracingCore.ush b/Engine/Shaders/Private/PathTracing/PathTracingCore.ush
--- a/Engine/Shaders/Private/PathTracing/PathTracingCore.ush
+++ b/Engine/Shaders/Private/PathTracing/PathTracingCore.ush
@@ -323,6 +323,11 @@ FPathTracingPayload TraceTransparentRay(inout FPathState PathState, bool bIsCame
 			HitPayload.Radiance *= HitPayload.Opacity;
 		}
 
+		if (HitPayload.BlendingMode == RAY_TRACING_BLEND_MODE_ALPHA_HOLDOUT && bIsCameraRay)
+		{
+			PathState.BackgroundVisibility += Luminance(PathState.PathThroughput * HitPayload.Opacity);
+		}
+
 		// add in surface emission (except for modulate blend mode which uses Radiance as the transparency multiplier)
 		if (bIncludeEmission && HitPayload.BlendingMode != RAY_TRACING_BLEND_MODE_MODULATE)
 		{
diff --git a/Engine/Shaders/Private/PathTracing/PathTracingMaterialHitShader.usf b/Engine/Shaders/Private/PathTracing/PathTracingMaterialHitShader.usf
--- a/Engine/Shaders/Private/PathTracing/PathTracingMaterialHitShader.usf
+++ b/Engine/Shaders/Private/PathTracing/PathTracingMaterialHitShader.usf
@@ -335,11 +335,11 @@ RAY_TRACING_ENTRY_ANY_HIT(PathTracingMaterialAHS,
 {
 #if USE_MATERIAL_ANY_HIT_SHADER
 
-#if MATERIALBLENDING_SOLID
+#if MATERIALBLENDING_SOLID || MATERIALBLENDING_ALPHAHOLDOUT
 	// Continue traversal by simply returning from this shader
 	// Usually this path is inactive because opaque materials don't have AHS bound
 	return;
-#elif MATERIALBLENDING_ALPHAHOLDOUT || MATERIALBLENDING_ADDITIVE
+#elif MATERIALBLENDING_ADDITIVE
 	if (PackedPayload.IsVisibilityRay())
 	{
 		// path tracer wants to treat these cases as fully transparent
diff --git a/Engine/Source/Runtime/Renderer/Private/PathTracing.cpp b/Engine/Source/Runtime/Renderer/Private/PathTracing.cpp
--- a/Engine/Source/Runtime/Renderer/Private/PathTracing.cpp
+++ b/Engine/Source/Runtime/Renderer/Private/PathTracing.cpp
@@ -903,7 +903,21 @@ IMPLEMENT_MATERIAL_SHADER_TYPE(template <>, FPathTracingMaterialCHS_AHS_IS, TEXT
 IMPLEMENT_MATERIAL_SHADER_TYPE(template <>, FGPULightmassCHS              , TEXT("/Engine/Private/PathTracing/PathTracingMaterialHitShader.usf"), TEXT("closesthit=PathTracingMaterialCHS"), SF_RayHitGroup);
 IMPLEMENT_MATERIAL_SHADER_TYPE(template <>, FGPULightmassCHS_AHS          , TEXT("/Engine/Private/PathTracing/PathTracingMaterialHitShader.usf"), TEXT("closesthit=PathTracingMaterialCHS anyhit=PathTracingMaterialAHS"), SF_RayHitGroup);
 
-
+static bool NeedsAnyHitShader(const FMaterial& RESTRICT MaterialResource)
+{
+	if (MaterialResource.IsMasked())
+	{
+		return true;
+	}
+	switch (MaterialResource.GetBlendMode())
+	{
+		case BLEND_Opaque:
+		case BLEND_AlphaHoldout:
+			return false;
+		default:
+			return true;
+	}
+}
 
 bool FRayTracingMeshProcessor::ProcessPathTracing(
 	const FMeshBatch& RESTRICT MeshBatch,
@@ -921,7 +935,7 @@ bool FRayTracingMeshProcessor::ProcessPathTracing(
 	{
 		case ERayTracingMeshCommandsMode::PATH_TRACING:
 		{
-			if (MaterialResource.IsMasked() || MaterialResource.GetBlendMode() != BLEND_Opaque)
+			if (NeedsAnyHitShader(MaterialResource))
 			{
 				if (bUseProceduralPrimitive)
 					ShaderTypes.AddShaderType<FPathTracingMaterialCHS_AHS_IS>();
@@ -939,7 +953,7 @@ bool FRayTracingMeshProcessor::ProcessPathTracing(
 		}
 		case ERayTracingMeshCommandsMode::LIGHTMAP_TRACING:
 		{
-			if (MaterialResource.IsMasked() || MaterialResource.GetBlendMode() != BLEND_Opaque)
+			if (NeedsAnyHitShader(MaterialResource))
 			{
 				ShaderTypes.AddShaderType<FGPULightmassCHS_AHS>();
 			}
