Fix a bug of stencil testing in PPMaterial if MobileSupportFetchBindedCustomStencilBuffer

diff --git a/Engine/Source/Runtime/Renderer/Private/PostProcess/PostProcessMaterial.cpp b/Engine/Source/Runtime/Renderer/Private/PostProcess/PostProcessMaterial.cpp
--- a/Engine/Source/Runtime/Renderer/Private/PostProcess/PostProcessMaterial.cpp
+++ b/Engine/Source/Runtime/Renderer/Private/PostProcess/PostProcessMaterial.cpp
@@ -538,6 +538,7 @@ FScreenPassTexture AddPostProcessMaterialPass(
 	const uint32 MaterialStencilRef = Material->GetStencilRefValue();
 
 	const bool bMobilePlatform = IsMobilePlatform(View.GetShaderPlatform());
+	const bool bSupportFetchBindedCustomStencilBuffer = !bMobilePlatform || IsMobileSupportFetchBindedCustomStencilBuffer(View.GetShaderPlatform());
 
 	FPostProcessMaterialParameters* PostProcessMaterialParameters = GraphBuilder.AllocParameters<FPostProcessMaterialParameters>();
 	PostProcessMaterialParameters->SceneTextures = Inputs.SceneTextures;
@@ -559,7 +560,7 @@ FScreenPassTexture AddPostProcessMaterialPass(
 	// The target color will be decoded if bForceIntermediateTarget is true in any case, but we might still need to decode the input color
 	PostProcessMaterialParameters->bMetalMSAAHDRDecode = Inputs.bMetalMSAAHDRDecode ? 1 : 0;
 
-	if (DepthStencilTexture && (!bMobilePlatform || IsMobileSupportFetchBindedCustomStencilBuffer(View.GetShaderPlatform())))
+	if (DepthStencilTexture && bSupportFetchBindedCustomStencilBuffer)
 	{
 		PostProcessMaterialParameters->RenderTargets.DepthStencil = FDepthStencilBinding(
 			DepthStencilTexture,
@@ -661,7 +662,7 @@ FScreenPassTexture AddPostProcessMaterialPass(
 		OutputViewport,
 		SceneColorViewport,
 		// Uses default depth stencil on mobile since the stencil test is done in pixel shader.
-		FScreenPassPipelineState(VertexShader, PixelShader, BlendState, bMobilePlatform ? DefaultDepthStencilState : DepthStencilState, MaterialStencilRef),
+		FScreenPassPipelineState(VertexShader, PixelShader, BlendState, !bSupportFetchBindedCustomStencilBuffer ? DefaultDepthStencilState : DepthStencilState, MaterialStencilRef),
 		PostProcessMaterialParameters,
 		ScreenPassFlags,
 		[&View, VertexShader, PixelShader, MaterialRenderProxy, Material, PostProcessMaterialParameters, MaterialStencilRef](FRHICommandList& RHICmdList)
