Fix crash when building lighting with a reflection capture.
* Capture in editor(was +2) and loading from file(+1) added different reference count. The cubemap in use is not marked as delete correctly.
* When a capture is rebuilt, all reference should change to the new capture.

diff --git a/Engine/Source/Runtime/Renderer/Private/ReflectionEnvironment.cpp b/Engine/Source/Runtime/Renderer/Private/ReflectionEnvironment.cpp
--- a/Engine/Source/Runtime/Renderer/Private/ReflectionEnvironment.cpp
+++ b/Engine/Source/Runtime/Renderer/Private/ReflectionEnvironment.cpp
@@ -285,12 +285,17 @@ FCaptureComponentSceneState* FReflectionCaptureCache::AddReference(const UReflec
 {
 	check(IsValid(Component))
 
-	RemapRegisteredComponentMapBuildDataId(Component);
+	bool Remap = RemapRegisteredComponentMapBuildDataId(Component);
 	FReflectionCaptureCacheEntry* Found = CaptureData.Find(Component->MapBuildDataId);
 	if (Found == nullptr)
 		return nullptr;
 
-	Found->RefCount += 1;
+	// Should not add reference count if this is caused by capture rebuilt
+	if (!Remap)
+	{
+		Found->RefCount += 1;
+	}
+
 	return &(Found->SceneState);
 }
 
@@ -338,7 +343,7 @@ void FReflectionCaptureCache::RegisterComponentMapBuildDataId(const UReflectionC
 	RegisteredComponentMapBuildDataIds.Add(Component, Component->MapBuildDataId);
 }
 
-void FReflectionCaptureCache::RemapRegisteredComponentMapBuildDataId(const UReflectionCaptureComponent* Component)
+bool FReflectionCaptureCache::RemapRegisteredComponentMapBuildDataId(const UReflectionCaptureComponent* Component)
 {
 	check(Component);
 
@@ -347,13 +352,30 @@ void FReflectionCaptureCache::RemapRegisteredComponentMapBuildDataId(const URefl
 	if (OldBuildId &&
 		*OldBuildId != Component->MapBuildDataId)
 	{
-		FReflectionCaptureCacheEntry Entry = *CaptureData.Find(*OldBuildId);
-		CaptureData.Remove(*OldBuildId);
+		FGuid OldBuildIdCopy = *OldBuildId;
+
+		// Remap all pointers that point to the old guid to the new one.
+		int32 ReferenceCount = 0;
+		for (TPair<const UReflectionCaptureComponent*, FGuid>& item : RegisteredComponentMapBuildDataIds)
+		{
+			if (item.Value == OldBuildIdCopy)
+			{
+				item.Value = Component->MapBuildDataId;
+				ReferenceCount++;
+			}
+		}
+
+		FReflectionCaptureCacheEntry Entry = *CaptureData.Find(OldBuildIdCopy);
+		Entry.RefCount = ReferenceCount;
+
+		CaptureData.Remove(OldBuildIdCopy);
 		CaptureData.Shrink();
 		CaptureData.Add(Component->MapBuildDataId, Entry);
 
-		RegisteredComponentMapBuildDataIds.Add(Component, Component->MapBuildDataId);
+		return true;
 	}
+
+	return false;
 }
 
 void FReflectionCaptureCache::UnregisterComponentMapBuildDataId(const UReflectionCaptureComponent* Component)
diff --git a/Engine/Source/Runtime/Renderer/Private/ReflectionEnvironmentCapture.cpp b/Engine/Source/Runtime/Renderer/Private/ReflectionEnvironmentCapture.cpp
--- a/Engine/Source/Runtime/Renderer/Private/ReflectionEnvironmentCapture.cpp
+++ b/Engine/Source/Runtime/Renderer/Private/ReflectionEnvironmentCapture.cpp
@@ -1222,12 +1222,11 @@ void CaptureSceneIntoScratchCubemap(
 	}
 }
 
-void CopyToSceneArray(FRDGBuilder& GraphBuilder, FScene* Scene, FRDGTexture* FilteredCubeTexture, FReflectionCaptureProxy* ReflectionProxy)
+void CopyToSceneArray(FRDGBuilder& GraphBuilder, FScene* Scene, FRDGTexture* FilteredCubeTexture, FReflectionCaptureProxy* ReflectionProxy, int32 CaptureIndex)
 {
 	RDG_EVENT_SCOPE(GraphBuilder, "CopyToSceneArray");
 
 	const int32 NumMips = GetNumMips(Scene->ReflectionSceneData.CubemapArray.GetCubemapSize());
-	const int32 CaptureIndex = FindOrAllocateCubemapIndex(Scene, ReflectionProxy->Component);
 
 	FRDGTexture* DestCubeTexture = GraphBuilder.RegisterExternalTexture(Scene->ReflectionSceneData.CubemapArray.GetRenderTarget());
 
@@ -1352,7 +1351,7 @@ void FScene::CaptureOrUploadReflectionCapture(UReflectionCaptureComponent* Captu
 			ENQUEUE_RENDER_COMMAND(FilterCommand)(
 				[Scene = this, FeatureLevel = FeatureLevel, ReflectionCubemapTexture = ReflectionCubemapTexture.Get(), ReflectionCaptureSize, CaptureComponent, ReflectionProxy](FRHICommandListImmediate& RHICmdList)
 			{
-				FindOrAllocateCubemapIndex(Scene, CaptureComponent);
+				const int32 CubemapIndex = FindOrAllocateCubemapIndex(Scene, CaptureComponent);
 				FCaptureComponentSceneState& FoundState = Scene->ReflectionSceneData.AllocatedReflectionCaptureState.FindChecked(CaptureComponent);
 
 				FRDGBuilder GraphBuilder(RHICmdList);
@@ -1367,7 +1366,7 @@ void FScene::CaptureOrUploadReflectionCapture(UReflectionCaptureComponent* Captu
 
 				if (FeatureLevel == ERHIFeatureLevel::SM5)
 				{
-					CopyToSceneArray(GraphBuilder, Scene, FilteredSceneCubemapTexture, ReflectionProxy);
+					CopyToSceneArray(GraphBuilder, Scene, FilteredSceneCubemapTexture, ReflectionProxy, CubemapIndex);
 				}
 
 				GraphBuilder.Execute();
diff --git a/Engine/Source/Runtime/Renderer/Private/ScenePrivate.h b/Engine/Source/Runtime/Renderer/Private/ScenePrivate.h
--- a/Engine/Source/Runtime/Renderer/Private/ScenePrivate.h
+++ b/Engine/Source/Runtime/Renderer/Private/ScenePrivate.h
@@ -1633,7 +1633,7 @@ public:
 
 protected:
 
-	void RemapRegisteredComponentMapBuildDataId(const UReflectionCaptureComponent* Component);
+	bool RemapRegisteredComponentMapBuildDataId(const UReflectionCaptureComponent* Component);
 	void RegisterComponentMapBuildDataId(const UReflectionCaptureComponent* Component);
 	void UnregisterComponentMapBuildDataId(const UReflectionCaptureComponent* Component);
 
