Fix deprecated GetSceneTextureExtent()

diff --git a/Engine/Source/Runtime/Renderer/Private/Strata/Strata.cpp b/Engine/Source/Runtime/Renderer/Private/Strata/Strata.cpp
--- a/Engine/Source/Runtime/Renderer/Private/Strata/Strata.cpp
+++ b/Engine/Source/Runtime/Renderer/Private/Strata/Strata.cpp
@@ -148,7 +148,7 @@ FIntPoint GetStrataTextureResolution(const FIntPoint& InResolution)
 
 static void BindStrataGlobalUniformParameters(FRDGBuilder& GraphBuilder, FStrataViewData* StrataViewData, FStrataGlobalUniformParameters& OutStrataUniformParameters);
 
-static void InitialiseStrataViewData(FRDGBuilder& GraphBuilder, FViewInfo& View, FStrataSceneData& SceneData)
+static void InitialiseStrataViewData(FRDGBuilder& GraphBuilder, FViewInfo& View, const FSceneTexturesConfig& SceneTexturesConfig, FStrataSceneData& SceneData)
 {
 	// Sanity check: the scene data should already exist 
 	check(SceneData.MaterialTextureArray != nullptr);
@@ -207,8 +207,9 @@ static void InitialiseStrataViewData(FRDGBuilder& GraphBuilder, FViewInfo& View,
 
 		// BSDF tiles
 		{
-			const FIntPoint BufferSize = GetSceneTextureExtent();
-			const FIntPoint BufferSize_Extended = GetStrataTextureTileResolution(GetSceneTextureExtent(), EStrataTileSpace::StrataTileSpace_Primary | EStrataTileSpace::StrataTileSpace_Overflow);
+			
+			const FIntPoint BufferSize = SceneTexturesConfig.Extent;
+			const FIntPoint BufferSize_Extended = GetStrataTextureTileResolution(SceneTexturesConfig.Extent, EStrataTileSpace::StrataTileSpace_Primary | EStrataTileSpace::StrataTileSpace_Overflow);
 
 			const FIntPoint BaseTileOffset_Overflow = FIntPoint(0, FMath::DivideAndRoundUp(BufferSize.Y, STRATA_TILE_SIZE));
 
@@ -339,7 +340,7 @@ void InitialiseStrataFrameSceneData(FRDGBuilder& GraphBuilder, FSceneRenderer& S
 	// Initialized view data
 	for (int32 ViewIndex = 0; ViewIndex < SceneRenderer.Views.Num(); ViewIndex++)
 	{
-		Strata::InitialiseStrataViewData(GraphBuilder, SceneRenderer.Views[ViewIndex], Out);
+		Strata::InitialiseStrataViewData(GraphBuilder, SceneRenderer.Views[ViewIndex], SceneRenderer.GetActiveSceneTexturesConfig(), Out);
 	}
 }
 
