Ray tracing now takes into account the VisibleInSceneCapturesOnly flag

diff --git a/Engine/Source/Runtime/Renderer/Private/DeferredShadingRenderer.cpp b/Engine/Source/Runtime/Renderer/Private/DeferredShadingRenderer.cpp
--- a/Engine/Source/Runtime/Renderer/Private/DeferredShadingRenderer.cpp
+++ b/Engine/Source/Runtime/Renderer/Private/DeferredShadingRenderer.cpp
@@ -687,6 +687,11 @@ bool FDeferredShadingSceneRenderer::GatherRayTracingWorldInstancesForView(FRDGBu
 				continue;
 			}
 
+			if (!View.bIsSceneCapture && SceneInfo->bIsVisibleInSceneCapturesOnly)
+			{
+				continue;
+			}
+
 			// Marked visible and used after point, check if streaming then mark as used in the TLAS (so it can be streamed in)
 			if (EnumHasAnyFlags(Scene->PrimitiveRayTracingFlags[PrimitiveIndex], ERayTracingPrimitiveFlags::Streaming))
 			{
diff --git a/Engine/Source/Runtime/Renderer/Private/PrimitiveSceneInfo.cpp b/Engine/Source/Runtime/Renderer/Private/PrimitiveSceneInfo.cpp
--- a/Engine/Source/Runtime/Renderer/Private/PrimitiveSceneInfo.cpp
+++ b/Engine/Source/Runtime/Renderer/Private/PrimitiveSceneInfo.cpp
@@ -174,6 +174,7 @@ FPrimitiveSceneInfo::FPrimitiveSceneInfo(UPrimitiveComponent* InComponent,FScene
 #if RHI_RAYTRACING
 	bDrawInGame(Proxy->IsDrawnInGame()),
 	bIsVisibleInSceneCaptures(!InComponent->SceneProxy->IsHiddenInSceneCapture()),
+	bIsVisibleInSceneCapturesOnly(InComponent->SceneProxy->IsVisibleInSceneCaptureOnly()),
 	bIsRayTracingRelevant(InComponent->SceneProxy->IsRayTracingRelevant()),
 	bIsRayTracingStaticRelevant(InComponent->SceneProxy->IsRayTracingStaticRelevant()),
 	bIsVisibleInRayTracing(InComponent->SceneProxy->IsVisibleInRayTracing()),
diff --git a/Engine/Source/Runtime/Renderer/Public/PrimitiveSceneInfo.h b/Engine/Source/Runtime/Renderer/Public/PrimitiveSceneInfo.h
--- a/Engine/Source/Runtime/Renderer/Public/PrimitiveSceneInfo.h
+++ b/Engine/Source/Runtime/Renderer/Public/PrimitiveSceneInfo.h
@@ -390,6 +390,7 @@ public:
 #if RHI_RAYTRACING
 	bool bDrawInGame : 1;
 	bool bIsVisibleInSceneCaptures : 1;
+	bool bIsVisibleInSceneCapturesOnly : 1;
 	bool bIsRayTracingRelevant : 1;
 	bool bIsRayTracingStaticRelevant : 1;
 	bool bIsVisibleInRayTracing : 1;
