Fix groom crash due to deferring optional material buffer.

diff --git a/Engine/Source/Runtime/Renderer/Private/HairStrands/HairStrandsVisibility.cpp b/Engine/Source/Runtime/Renderer/Private/HairStrands/HairStrandsVisibility.cpp
--- a/Engine/Source/Runtime/Renderer/Private/HairStrands/HairStrandsVisibility.cpp
+++ b/Engine/Source/Runtime/Renderer/Private/HairStrands/HairStrandsVisibility.cpp
@@ -2829,7 +2829,10 @@ void AddMeshDrawTransitionPass(
 			ExternalAccessQueue.Add(VFInput.Strands.TangentBuffer.Buffer);
 			ExternalAccessQueue.Add(VFInput.Strands.Attribute0Buffer.Buffer);
 			ExternalAccessQueue.Add(VFInput.Strands.Attribute1Buffer.Buffer);
-			ExternalAccessQueue.Add(VFInput.Strands.MaterialBuffer.Buffer);
+			if (VFInput.Strands.MaterialBuffer.Buffer != nullptr)
+			{
+				ExternalAccessQueue.Add(VFInput.Strands.MaterialBuffer.Buffer);
+			}
 			ExternalAccessQueue.Add(VFInput.Strands.PositionOffsetBuffer.Buffer);
 			ExternalAccessQueue.Add(VFInput.Strands.PrevPositionOffsetBuffer.Buffer);
 
