Share translucency light volume textures between both views of a stereo pair (UE-148569).

diff --git a/Engine/Source/Runtime/Renderer/Private/DeferredShadingRenderer.cpp b/Engine/Source/Runtime/Renderer/Private/DeferredShadingRenderer.cpp
--- a/Engine/Source/Runtime/Renderer/Private/DeferredShadingRenderer.cpp
+++ b/Engine/Source/Runtime/Renderer/Private/DeferredShadingRenderer.cpp
@@ -2659,7 +2659,7 @@ void FDeferredShadingSceneRenderer::Render(FRDGBuilder& GraphBuilder)
 
 	if (bRenderDeferredLighting && GbEnableAsyncComputeTranslucencyLightingVolumeClear && GSupportsEfficientAsyncCompute)
 	{
-		InitTranslucencyLightingVolumeTextures(GraphBuilder, Views, ERDGPassFlags::AsyncCompute, TranslucencyLightingVolumeTextures);
+		TranslucencyLightingVolumeTextures.Init(GraphBuilder, Views, ERDGPassFlags::AsyncCompute);
 	}
 
 #if RHI_RAYTRACING
@@ -2925,7 +2925,7 @@ void FDeferredShadingSceneRenderer::Render(FRDGBuilder& GraphBuilder)
 		// Clear the translucent lighting volumes before we accumulate
 		if ((GbEnableAsyncComputeTranslucencyLightingVolumeClear && GSupportsEfficientAsyncCompute) == false)
 		{
-			InitTranslucencyLightingVolumeTextures(GraphBuilder, Views, ERDGPassFlags::Compute, TranslucencyLightingVolumeTextures);
+			TranslucencyLightingVolumeTextures.Init(GraphBuilder, Views, ERDGPassFlags::Compute);
 		}
 
 #if RHI_RAYTRACING
diff --git a/Engine/Source/Runtime/Renderer/Private/TranslucentLighting.cpp b/Engine/Source/Runtime/Renderer/Private/TranslucentLighting.cpp
--- a/Engine/Source/Runtime/Renderer/Private/TranslucentLighting.cpp
+++ b/Engine/Source/Runtime/Renderer/Private/TranslucentLighting.cpp
@@ -760,30 +760,52 @@ public:
 
 IMPLEMENT_GLOBAL_SHADER(FClearTranslucentLightingVolumeCS, "/Engine/Private/TranslucentLightInjectionShaders.usf", "ClearTranslucentLightingVolumeCS", SF_Compute);
 
-void InitTranslucencyLightingVolumeTextures(FRDGBuilder& GraphBuilder, TArrayView<const FViewInfo> Views, ERDGPassFlags PassFlags, FTranslucencyLightingVolumeTextures& OutTextures)
+int32 FTranslucencyLightingVolumeTextures::GetIndex(const FViewInfo& View, int32 CascadeIndex) const
+{
+	// if we only have one view or one stereo pair we can just use primary index
+	if (Directional.Num() == TVC_MAX)
+	{
+		return (View.PrimaryViewIndex * TVC_MAX) + CascadeIndex;
+	}
+	else
+	{
+		// support uncommon but possible (in theory) situations, like a stereo pair and also multiple views
+		return (ViewsToTexturePairs[View.PrimaryViewIndex] * TVC_MAX) + CascadeIndex;
+	}
+}
+
+void FTranslucencyLightingVolumeTextures::Init(FRDGBuilder& GraphBuilder, TArrayView<const FViewInfo> Views, ERDGPassFlags PassFlags)
 {
 	check(PassFlags == ERDGPassFlags::Compute || PassFlags == ERDGPassFlags::AsyncCompute);
 
 	RDG_GPU_STAT_SCOPE(GraphBuilder, TranslucentLighting);
 
-	OutTextures.VolumeDim = GetTranslucencyLightingVolumeDim();
-	const FIntVector TranslucencyLightingVolumeDim(OutTextures.VolumeDim);
+	VolumeDim = GetTranslucencyLightingVolumeDim();
+	const FIntVector TranslucencyLightingVolumeDim(VolumeDim);
 
+	// calculate the number of textures needed given that for each stereo pair the primary view's textures will be shared between the "eyes"
+	const int32 ViewCount = Views.Num();
+	uint32 NumViewsWithTextures = 0;
+	ViewsToTexturePairs.SetNumZeroed(Views.Num());
+	for (int32 ViewIndex = 0, NumViews = Views.Num(); ViewIndex < NumViews; ++ViewIndex)
+	{
+		ViewsToTexturePairs[ViewIndex] = NumViewsWithTextures;
+		NumViewsWithTextures += (ViewIndex == Views[ViewIndex].PrimaryViewIndex) ? 1 : 0;	// this will add 0 for those views who aren't primary
+	}
+	check(NumViewsWithTextures > 0);
 	{
 		// TODO: We can skip the and TLV allocations when rendering in forward shading mode
 		const ETextureCreateFlags TranslucencyTargetFlags = TexCreate_ShaderResource | TexCreate_RenderTargetable | TexCreate_ReduceMemoryWithTilingMode | TexCreate_UAV;
 
-		const int32 ViewCount = Views.Num();
-		check(ViewCount > 0);
-
-		OutTextures.Ambient.SetNum(ViewCount * TVC_MAX);
-		OutTextures.Directional.SetNum(ViewCount * TVC_MAX);
+		Ambient.SetNum(NumViewsWithTextures * TVC_MAX);
+		Directional.SetNum(NumViewsWithTextures * TVC_MAX);
 
 		for (int32 ViewIndex = 0; ViewIndex < ViewCount; ++ViewIndex)
 		{
 			for (int32 CascadeIndex = 0; CascadeIndex < TVC_MAX; ++CascadeIndex)
 			{
-				const uint32 TextureIndex = FTranslucencyLightingVolumeTextures::GetIndex(ViewIndex, CascadeIndex);
+				const uint32 TextureIndex = FTranslucencyLightingVolumeTextures::GetIndex(Views[ViewIndex], CascadeIndex);
+				check(TextureIndex <= NumViewsWithTextures * TVC_MAX);
 
 				const FRDGEventName& AmbientName = *GraphBuilder.AllocObject<FRDGEventName>(RDG_EVENT_NAME("TranslucentVolumeAmbient%d", TextureIndex));
 				const FRDGEventName& DirectionalName = *GraphBuilder.AllocObject<FRDGEventName>(RDG_EVENT_NAME("TranslucentVolumeDirectional%d", TextureIndex));
@@ -804,8 +826,8 @@ void InitTranslucencyLightingVolumeTextures(FRDGBuilder& GraphBuilder, TArrayVie
 						TranslucencyTargetFlags),
 					DirectionalName.GetTCHAR());
 
-				OutTextures.Ambient[TextureIndex] = AmbientTexture;
-				OutTextures.Directional[TextureIndex] = DirectionalTexture;
+				Ambient[TextureIndex] = AmbientTexture;
+				Directional[TextureIndex] = DirectionalTexture;
 			}
 		}
 	}
@@ -814,17 +836,17 @@ void InitTranslucencyLightingVolumeTextures(FRDGBuilder& GraphBuilder, TArrayVie
 
 	TShaderMapRef<FClearTranslucentLightingVolumeCS> ComputeShader(Views[0].ShaderMap);
 
-	for (int32 ViewIndex = 0; ViewIndex < Views.Num(); ++ViewIndex)
+	for (uint32 TexturePairIndex = 0; TexturePairIndex < NumViewsWithTextures; ++TexturePairIndex)
 	{
 		auto* PassParameters = GraphBuilder.AllocParameters<FClearTranslucentLightingVolumeCS::FParameters>();
-		PassParameters->RWAmbient0 = GraphBuilder.CreateUAV(OutTextures.Ambient[FTranslucencyLightingVolumeTextures::GetIndex(ViewIndex, 0)]);
-		PassParameters->RWAmbient1 = GraphBuilder.CreateUAV(OutTextures.Ambient[FTranslucencyLightingVolumeTextures::GetIndex(ViewIndex, 1)]);
-		PassParameters->RWDirectional0 = GraphBuilder.CreateUAV(OutTextures.Directional[FTranslucencyLightingVolumeTextures::GetIndex(ViewIndex, 0)]);
-		PassParameters->RWDirectional1 = GraphBuilder.CreateUAV(OutTextures.Directional[FTranslucencyLightingVolumeTextures::GetIndex(ViewIndex, 1)]);
+		PassParameters->RWAmbient0 = GraphBuilder.CreateUAV(Ambient[TexturePairIndex * TVC_MAX]);
+		PassParameters->RWAmbient1 = GraphBuilder.CreateUAV(Ambient[TexturePairIndex * TVC_MAX + 1]);
+		PassParameters->RWDirectional0 = GraphBuilder.CreateUAV(Directional[TexturePairIndex * TVC_MAX]);
+		PassParameters->RWDirectional1 = GraphBuilder.CreateUAV(Directional[TexturePairIndex * TVC_MAX + 1]);
 
 		FComputeShaderUtils::AddPass(
 			GraphBuilder,
-			RDG_EVENT_NAME("ClearTranslucencyLightingVolumeCompute %d", OutTextures.VolumeDim),
+			RDG_EVENT_NAME("ClearTranslucencyLightingVolumeCompute %d", VolumeDim),
 			PassFlags,
 			ComputeShader,
 			PassParameters,
@@ -832,13 +854,13 @@ void InitTranslucencyLightingVolumeTextures(FRDGBuilder& GraphBuilder, TArrayVie
 	}
 }
 
-FTranslucencyLightingVolumeParameters GetTranslucencyLightingVolumeParameters(FRDGBuilder& GraphBuilder, const FTranslucencyLightingVolumeTextures& Textures, uint32 ViewIndex)
+FTranslucencyLightingVolumeParameters GetTranslucencyLightingVolumeParameters(FRDGBuilder& GraphBuilder, const FTranslucencyLightingVolumeTextures& Textures, const FViewInfo& View)
 {
 	FTranslucencyLightingVolumeParameters Parameters;
 	if (Textures.IsValid())
 	{
-		const uint32 InnerIndex = FTranslucencyLightingVolumeTextures::GetIndex(ViewIndex, TVC_Inner);
-		const uint32 OuterIndex = FTranslucencyLightingVolumeTextures::GetIndex(ViewIndex, TVC_Outer);
+		const uint32 InnerIndex = Textures.GetIndex(View, TVC_Inner);
+		const uint32 OuterIndex = Textures.GetIndex(View, TVC_Outer);
 
 		Parameters.TranslucencyLightingVolumeAmbientInner = Textures.Ambient[InnerIndex];
 		Parameters.TranslucencyLightingVolumeAmbientOuter = Textures.Ambient[OuterIndex];
@@ -900,7 +922,7 @@ void InjectTranslucencyLightingVolumeAmbientCubemap(
 
 		for (int32 VolumeCascadeIndex = 0; VolumeCascadeIndex < TVC_MAX; ++VolumeCascadeIndex)
 		{
-			FRDGTextureRef VolumeAmbientTexture = Textures.Ambient[FTranslucencyLightingVolumeTextures::GetIndex(ViewIndex, VolumeCascadeIndex)];
+			FRDGTextureRef VolumeAmbientTexture = Textures.GetAmbientTexture(View, VolumeCascadeIndex);
 
 			for (const FFinalPostProcessSettings::FCubemapEntry& CubemapEntry : View.FinalPostProcessSettings.ContributingCubemaps)
 			{
@@ -1066,7 +1088,8 @@ static void InjectTranslucentLightArray(
 	// Inject into each volume cascade. Operate on one cascade at a time to reduce render target switches.
 	for (uint32 VolumeCascadeIndex = 0; VolumeCascadeIndex < TVC_MAX; VolumeCascadeIndex++)
 	{
-		const uint32 TextureIndex = FTranslucencyLightingVolumeTextures::GetIndex(ViewIndex, VolumeCascadeIndex);
+		// for stereo case, using PrimaryViewIndex essentially shares the lighting volume textures
+		const uint32 TextureIndex = Textures.GetIndex(View, VolumeCascadeIndex);
 		FRDGTextureRef VolumeAmbientTexture = Textures.Ambient[TextureIndex];
 		FRDGTextureRef VolumeDirectionalTexture = Textures.Directional[TextureIndex];
 
@@ -1314,7 +1337,7 @@ void InjectSimpleTranslucencyLightingVolumeArray(
 		// Operate on one cascade at a time to reduce render target switches
 		for (int32 VolumeCascadeIndex = 0; VolumeCascadeIndex < TVC_MAX; VolumeCascadeIndex++)
 		{
-			const uint32 TextureIndex = FTranslucencyLightingVolumeTextures::GetIndex(ViewIndex, VolumeCascadeIndex);
+			const uint32 TextureIndex = Textures.GetIndex(View, VolumeCascadeIndex);
 
 			RDG_EVENT_SCOPE(GraphBuilder, "Cascade%d", VolumeCascadeIndex);
 			FRDGTextureRef VolumeAmbientTexture = Textures.Ambient[TextureIndex];
@@ -1408,7 +1431,7 @@ void FilterTranslucencyLightingVolume(
 
 		for (int32 VolumeCascadeIndex = 0; VolumeCascadeIndex < TVC_MAX; VolumeCascadeIndex++)
 		{
-			const uint32 TextureIndex = FTranslucencyLightingVolumeTextures::GetIndex(ViewIndex, VolumeCascadeIndex);
+			const uint32 TextureIndex = Textures.GetIndex(View, VolumeCascadeIndex);
 
 			FRDGTextureRef InputVolumeAmbientTexture = Textures.Ambient[TextureIndex];
 			FRDGTextureRef InputVolumeDirectionalTexture = Textures.Directional[TextureIndex];
diff --git a/Engine/Source/Runtime/Renderer/Private/TranslucentLighting.h b/Engine/Source/Runtime/Renderer/Private/TranslucentLighting.h
--- a/Engine/Source/Runtime/Renderer/Private/TranslucentLighting.h
+++ b/Engine/Source/Runtime/Renderer/Private/TranslucentLighting.h
@@ -13,10 +13,11 @@ class FProjectedShadowInfo;
 
 struct FTranslucencyLightingVolumeTextures
 {
-	static int32 GetIndex(int32 ViewIndex, int32 CascadeIndex)
-	{
-		return (ViewIndex * TVC_MAX) + CascadeIndex;
-	}
+	/** Maps view to a texture index, accounting for sharing between some views */
+	int32 GetIndex(const FViewInfo& View, int32 CascadeIndex) const;
+
+	/** Initializes the translucent volume textures and clears them using a compute pass. PassFlags should be ERDGPassFlags::Compute or ERDGPassFlags::AsyncCompute. */
+	void Init(FRDGBuilder& GraphBuilder, TArrayView<const FViewInfo> Views, ERDGPassFlags PassFlags);
 
 	bool IsValid() const
 	{
@@ -24,23 +25,33 @@ struct FTranslucencyLightingVolumeTextures
 		return VolumeDim != 0;
 	}
 
+	FRDGTextureRef GetAmbientTexture(const FViewInfo& View, int32 CascadeIndex) const
+	{
+		return Ambient[GetIndex(View, CascadeIndex)];
+	}
+
+	FRDGTextureRef GetDirectionalTexture(const FViewInfo& View, int32 CascadeIndex) const
+	{
+		return Directional[GetIndex(View, CascadeIndex)];
+	}
+
 	TArray<FRDGTextureRef, TInlineAllocator<TVC_MAX>> Ambient;
 	TArray<FRDGTextureRef, TInlineAllocator<TVC_MAX>> Directional;
 	int32 VolumeDim = 0;
+	
+	/** Mapping between the view index and texture pair - needed because stereo views shares textures. */
+	TArray<int32, TInlineAllocator<2>> ViewsToTexturePairs;
 };
 
-/** Initializes the translucent volume textures and clears them using a compute pass. PassFlags should be ERDGPassFlags::Compute or ERDGPassFlags::AsyncCompute. */
-void InitTranslucencyLightingVolumeTextures(FRDGBuilder& GraphBuilder, TArrayView<const FViewInfo> Views, ERDGPassFlags PassFlags, FTranslucencyLightingVolumeTextures& OutTextures);
-
 BEGIN_SHADER_PARAMETER_STRUCT(FTranslucencyLightingVolumeParameters, )
-	SHADER_PARAMETER_RDG_TEXTURE(Texture3D, TranslucencyLightingVolumeAmbientInner)
-	SHADER_PARAMETER_RDG_TEXTURE(Texture3D, TranslucencyLightingVolumeAmbientOuter)
-	SHADER_PARAMETER_RDG_TEXTURE(Texture3D, TranslucencyLightingVolumeDirectionalInner)
-	SHADER_PARAMETER_RDG_TEXTURE(Texture3D, TranslucencyLightingVolumeDirectionalOuter)
+SHADER_PARAMETER_RDG_TEXTURE(Texture3D, TranslucencyLightingVolumeAmbientInner)
+SHADER_PARAMETER_RDG_TEXTURE(Texture3D, TranslucencyLightingVolumeAmbientOuter)
+SHADER_PARAMETER_RDG_TEXTURE(Texture3D, TranslucencyLightingVolumeDirectionalInner)
+SHADER_PARAMETER_RDG_TEXTURE(Texture3D, TranslucencyLightingVolumeDirectionalOuter)
 END_SHADER_PARAMETER_STRUCT()
 
 /** Initializes translucency volume lighting shader parameters from an optional textures struct. If null or uninitialized, fallback textures are used. */
-FTranslucencyLightingVolumeParameters GetTranslucencyLightingVolumeParameters(FRDGBuilder& GraphBuilder, const FTranslucencyLightingVolumeTextures& Textures, uint32 ViewIndex);
+FTranslucencyLightingVolumeParameters GetTranslucencyLightingVolumeParameters(FRDGBuilder& GraphBuilder, const FTranslucencyLightingVolumeTextures& Textures, const FViewInfo& View);
 
 void InjectTranslucencyLightingVolume(
 	FRDGBuilder& GraphBuilder,
diff --git a/Engine/Source/Runtime/Renderer/Private/TranslucentRendering.cpp b/Engine/Source/Runtime/Renderer/Private/TranslucentRendering.cpp
--- a/Engine/Source/Runtime/Renderer/Private/TranslucentRendering.cpp
+++ b/Engine/Source/Runtime/Renderer/Private/TranslucentRendering.cpp
@@ -1042,7 +1042,7 @@ TRDGUniformBufferRef<FTranslucentBasePassUniformParameters> CreateTranslucentBas
 	}
 
 	// Translucency Lighting Volume
-	BasePassParameters.TranslucencyLightingVolume = GetTranslucencyLightingVolumeParameters(GraphBuilder, TranslucencyLightingVolumeTextures, ViewIndex);
+	BasePassParameters.TranslucencyLightingVolume = GetTranslucencyLightingVolumeParameters(GraphBuilder, TranslucencyLightingVolumeTextures, View);
 	BasePassParameters.LumenParameters = GetLumenTranslucencyLightingParameters(GraphBuilder, View.LumenTranslucencyGIVolume, View.LumenFrontLayerTranslucency);
 
 	const bool bLumenGIHandlingSkylight = bLumenGIEnabled
