This CL adds Continuous LOD support for Compute Rasterized strand based hair - it includes:
- a read only project setting to enable Continuous Decimation Reordering: r.HairStrands.ContinuousDecimationReordering - which is done when platform data is built in the editor - and is disabled by default.
- code and shader changes to use the reordered strand data at runtime in the Compute Raster rendering path
- a proof of concept reordering strategy (random shuffle) - with better strategies planned
- a partial implementation of Temporal Layering - a form of temporal sample accumulation which also uses the Continuous Decimation Reordering - with a full implemention planned in the next CL

diff --git a/Engine/Plugins/Runtime/HairStrands/Source/HairStrandsCore/Private/GroomBuilder.cpp b/Engine/Plugins/Runtime/HairStrands/Source/HairStrandsCore/Private/GroomBuilder.cpp
--- a/Engine/Plugins/Runtime/HairStrands/Source/HairStrandsCore/Private/GroomBuilder.cpp
+++ b/Engine/Plugins/Runtime/HairStrands/Source/HairStrandsCore/Private/GroomBuilder.cpp
@@ -35,7 +35,10 @@ static FAutoConsoleVariableRef CVarHairGroupIndexBuilder_MaxVoxelResolution(TEXT
 FString FGroomBuilder::GetVersion()
 {
 	// Important to update the version when groom building changes
-	return TEXT("2r2");
+	if (IsHairStrandContinuousDecimationReorderingEnabled())
+		return TEXT("3r0");
+	else
+		return TEXT("2r2");
 }
 
 namespace FHairStrandsDecimation
@@ -44,6 +47,7 @@ namespace FHairStrandsDecimation
 		const FHairStrandsDatas& InData,
 		float CurveDecimationPercentage,
 		float VertexDecimationPercentage,
+		bool bContinuousDecimationReordering,
 		FHairStrandsDatas& OutData);
 }
 
@@ -1750,6 +1754,7 @@ void FGroomBuilder::BuildData(
 	const float CurveDecimation		= FMath::Clamp(InSettings.DecimationSettings.CurveDecimation, 0.f, 1.f);
 	const float VertexDecimation	= FMath::Clamp(InSettings.DecimationSettings.VertexDecimation, 0.f, 1.f);
 	const float HairToGuideDensity	= FMath::Clamp(InSettings.InterpolationSettings.HairToGuideDensity, 0.f, 1.f);	
+	const bool bContinuousDecimationReordering = IsHairStrandContinuousDecimationReorderingEnabled(); //this is a read only project setting which will modify the platform data that is stored in the DDC
 
 	{
 		OutGroupInfo.GroupID = InHairDescriptionGroup.Info.GroupID;
@@ -1763,11 +1768,11 @@ void FGroomBuilder::BuildData(
 			HairStrandsBuilder::BuildInternalData(OutRen, !InHairDescriptionGroup.bHasUVData);
 
 			// Decimate
-			if (CurveDecimation < 1 || VertexDecimation < 1)
+			if (CurveDecimation < 1 || VertexDecimation < 1 || bContinuousDecimationReordering)
 			{
 				FHairStrandsDatas FullData = OutRen;
 				OutRen.Reset();
-				FHairStrandsDecimation::Decimate(FullData, CurveDecimation, VertexDecimation, OutRen);
+				FHairStrandsDecimation::Decimate(FullData, CurveDecimation, VertexDecimation, bContinuousDecimationReordering, OutRen);
 			}
 			OutGroupInfo.NumCurves			= OutRen.GetNumCurves();
 			OutGroupInfo.NumCurveVertices	= OutRen.GetNumPoints();
@@ -1789,7 +1794,7 @@ void FGroomBuilder::BuildData(
 			else
 			{
 				OutSim.Reset();
-				FHairStrandsDecimation::Decimate(OutRen, HairToGuideDensity, 1, OutSim);
+				FHairStrandsDecimation::Decimate(OutRen, HairToGuideDensity, 1, false, OutSim);
 			}
 
 			OutGroupInfo.NumGuides			= OutSim.GetNumCurves();
@@ -1924,6 +1929,7 @@ void Decimate(
 	const FHairStrandsDatas& InData, 
 	float CurveDecimationPercentage, 
 	float VertexDecimationPercentage, 
+	bool bContinuousDecimationReordering,
 	FHairStrandsDatas& OutData)
 {
 	CurveDecimationPercentage = FMath::Clamp(CurveDecimationPercentage, 0.f, 1.f);
@@ -1941,6 +1947,23 @@ void Decimate(
 	uint32 OutTotalPointCount = 0;
 	FRandomStream Random;
 
+	TArray<uint32> CurveRandomizedIndex;
+
+	// this is the proof of concept version - which just does a shuffle - TODO: improved voxel based version which caters for both long hair and short hair
+	if (bContinuousDecimationReordering)
+	{
+		CurveRandomizedIndex.SetNum(InCurveCount);
+
+		Random.Initialize(0xdeedbeed);
+
+		//shuffle
+		for (uint32 CurveIndex = 0; CurveIndex < InCurveCount; CurveIndex++)
+		{
+			uint32 RandIndex = Random.RandRange(0, InCurveCount - 1);
+			CurveRandomizedIndex.Swap(CurveIndex, RandIndex);
+		}
+	}
+
 	const float CurveBucketSize = float(InCurveCount) / float(OutCurveCount);
 	int32 LastCurveIndex = -1;
 	for (uint32 BucketIndex = 0; BucketIndex < OutCurveCount; BucketIndex++)
@@ -1953,8 +1976,10 @@ void Decimate(
 			const uint32 CurveIndex = Random.RandRange(MinBucket, FMath::FloorToInt(MinBucket + AdjustedBucketSize)-1);
 			CurveIndices[BucketIndex] = CurveIndex;
 			LastCurveIndex = CurveIndex;
+			
+			const uint32 EffectiveCurveIndex = bContinuousDecimationReordering ? CurveRandomizedIndex[CurveIndex] : CurveIndex;
 
-			const uint32 InPointCount = InData.StrandsCurves.CurvesCount[CurveIndex];
+			const uint32 InPointCount = InData.StrandsCurves.CurvesCount[EffectiveCurveIndex];
 			const uint32 OutPointCount = DecimatePointCount(InPointCount, VertexDecimationPercentage);
 			OutTotalPointCount += OutPointCount;
 		}
@@ -1968,8 +1993,10 @@ void Decimate(
 	for (uint32 OutCurveIndex = 0; OutCurveIndex < OutCurveCount; ++OutCurveIndex)
 	{
 		const uint32 InCurveIndex	= CurveIndices[OutCurveIndex];
-		const uint32 InPointOffset	= InData.StrandsCurves.CurvesOffset[InCurveIndex];
-		const uint32 InPointCount	= InData.StrandsCurves.CurvesCount[InCurveIndex];
+		const uint32 EffectiveCurveIndex = bContinuousDecimationReordering ? CurveRandomizedIndex[InCurveIndex] : InCurveIndex;
+
+		const uint32 InPointOffset	= InData.StrandsCurves.CurvesOffset[EffectiveCurveIndex];
+		const uint32 InPointCount	= InData.StrandsCurves.CurvesCount[EffectiveCurveIndex];
 
 		// Decimation using area metric
 	#if 1
@@ -2029,13 +2056,13 @@ void Decimate(
 	#endif
 
 		OutData.StrandsCurves.CurvesCount[OutCurveIndex]  = OutPointCount;
-		OutData.StrandsCurves.CurvesRootUV[OutCurveIndex] = InData.StrandsCurves.CurvesRootUV[InCurveIndex];
+		OutData.StrandsCurves.CurvesRootUV[OutCurveIndex] = InData.StrandsCurves.CurvesRootUV[EffectiveCurveIndex];
 		OutData.StrandsCurves.CurvesOffset[OutCurveIndex] = OutPointOffset;
-		OutData.StrandsCurves.CurvesLength[OutCurveIndex] = InData.StrandsCurves.CurvesLength[InCurveIndex];
+		OutData.StrandsCurves.CurvesLength[OutCurveIndex] = InData.StrandsCurves.CurvesLength[EffectiveCurveIndex];
 		if (bHasPrecomputedWeights)
 		{
-			OutData.StrandsCurves.CurvesClosestGuideIDs[OutCurveIndex] = InData.StrandsCurves.CurvesClosestGuideIDs[InCurveIndex];
-			OutData.StrandsCurves.CurvesClosestGuideWeights[OutCurveIndex] = InData.StrandsCurves.CurvesClosestGuideWeights[InCurveIndex];
+			OutData.StrandsCurves.CurvesClosestGuideIDs[OutCurveIndex] = InData.StrandsCurves.CurvesClosestGuideIDs[EffectiveCurveIndex];
+			OutData.StrandsCurves.CurvesClosestGuideWeights[OutCurveIndex] = InData.StrandsCurves.CurvesClosestGuideWeights[EffectiveCurveIndex];
 		}
 		OutData.StrandsCurves.MaxLength = InData.StrandsCurves.MaxLength;
 		OutData.StrandsCurves.MaxRadius = InData.StrandsCurves.MaxRadius;
diff --git a/Engine/Plugins/Runtime/HairStrands/Source/HairStrandsCore/Private/GroomManager.cpp b/Engine/Plugins/Runtime/HairStrands/Source/HairStrandsCore/Private/GroomManager.cpp
--- a/Engine/Plugins/Runtime/HairStrands/Source/HairStrandsCore/Private/GroomManager.cpp
+++ b/Engine/Plugins/Runtime/HairStrands/Source/HairStrandsCore/Private/GroomManager.cpp
@@ -715,7 +715,9 @@ static void RunHairLODSelection(
 		// is not resquested (i.e., LODIndex<0), the MinLOD is applied after ViewLODIndex has been determined in the codeblock below
 		const int32 LODCount = Instance->HairGroupPublicData->GetLODVisibilities().Num();
 		const float MinLOD = FMath::Max(0, GHairStrandsMinLOD);
-		float LODIndex = Instance->Debug.LODForcedIndex >= 0 ? FMath::Max(Instance->Debug.LODForcedIndex, MinLOD) : -1.0f;
+		
+		// If continuous LOD is enabled, we bypass all other type of geometric representation, and only use LOD0
+		float LODIndex = IsHairVisibilityComputeRasterContinuousLODEnabled() ? 0.0 : (Instance->Debug.LODForcedIndex >= 0 ? FMath::Max(Instance->Debug.LODForcedIndex, MinLOD) : -1.0f);
 		float LODViewIndex = -1;
 		{
 			float MaxScreenSize = 0.f;
diff --git a/Engine/Shaders/Private/HairStrands/HairStrandsVisibilityCompactionComputeRaster.usf b/Engine/Shaders/Private/HairStrands/HairStrandsVisibilityCompactionComputeRaster.usf
--- a/Engine/Shaders/Private/HairStrands/HairStrandsVisibilityCompactionComputeRaster.usf
+++ b/Engine/Shaders/Private/HairStrands/HairStrandsVisibilityCompactionComputeRaster.usf
@@ -159,7 +159,6 @@ void MainCS(
 				OutNodeVis.MaterialId	= Vis_PrimId_MatId[OutIndex].y;
 				OutNodeVis.Depth		= Vis_Depth_Cov[OutIndex].x;
 				OutNodeVis.Coverage8bit = 0xff;
-				//OutNodeVis.Coverage8bit = To8bitCoverage(Vis_Depth_Cov[OutIndex].y / TotalCoverage); 
 				const uint StoreIndex = NodeDesc.Offset + OutIndex;
 				OutCompactNodeVis[StoreIndex] = PackHairVis(OutNodeVis);
 				OutCompactNodeCoord[StoreIndex] = PixelCoord;
@@ -170,8 +169,6 @@ void MainCS(
 			float ViewTransmittance = pow(2.71828, LogViewTransmittance);
 			float CoverageThreshold = 0.98;
 			PixelCoverage = saturate(min(1.0 - ViewTransmittance, 1.0) / CoverageThreshold);
-
-//			PixelCoverage = TransmittanceToCoverage(ViewTransmittanceTexture.Load(uint3(PixelCoord, 0)), CoverageThreshold);
 		}
 	}
 
diff --git a/Engine/Shaders/Private/HairStrands/HairStrandsVisibilityRasterCompute.usf b/Engine/Shaders/Private/HairStrands/HairStrandsVisibilityRasterCompute.usf
--- a/Engine/Shaders/Private/HairStrands/HairStrandsVisibilityRasterCompute.usf
+++ b/Engine/Shaders/Private/HairStrands/HairStrandsVisibilityRasterCompute.usf
@@ -81,9 +81,12 @@ Buffer<uint4>	HairStrandsVF_PositionBuffer;
 Buffer<float4>	HairStrandsVF_PositionOffsetBuffer;
 
 Texture2D<float> SceneDepthTexture;
-float			TileLODScale;
+
 int2			TileRes;
-float			TileBaseDepth;
+
+uint 			VertexCount;
+uint 			VertexStart;
+float			SampleWeight;
 
 float3 NDCToPixelCoord(float4 InDC)
 {
@@ -123,7 +126,7 @@ void CalcScreenPosAndRad(in uint4 PB, in float3 PBO, out float4 SP, out float Ra
 
 	SP = float4(NDCToPixelCoord(NDC).xy, NDC.zw);
 
-	Rad = CP.WorldRadius * OutputResolutionf.x;
+	Rad = CP.WorldRadius*2000.0; // OutputResolutionf.x; //TODO: figure this out correctly?
 }
 
 ///////////////////////////////////////////////////////////////////////////
@@ -217,14 +220,19 @@ void BinningCS(uint2 DispatchThreadID : SV_DispatchThreadID, uint GroupThreadID
 	ResolvedView = ResolveView();
 	if (GroupThreadID == 0)
 	{
-		uint VertexCount = HairStrandsVF_VertexCount;
 #if PERMUTATION_CULLING	
 		if (HairStrandsVF_bIsCullingEnable)
 		{
-			VertexCount = HairStrandsVF_CullingIndirectBuffer[3];
+			group_BatchNum = (HairStrandsVF_CullingIndirectBuffer[3] + 1023) / 1024;
 		}
+		else
+		{
+			group_BatchNum = (VertexCount + 1023) / 1024;
+		}
+#else
+		group_BatchNum = (VertexCount + 1023) / 1024;
 #endif
-		group_BatchNum = (VertexCount+1023)/1024;
+		
 		group_LoopNum = (group_BatchNum + (NumBinners - 1)) * RcpNumBinners;
 	}
 
@@ -243,9 +251,19 @@ void BinningCS(uint2 DispatchThreadID : SV_DispatchThreadID, uint GroupThreadID
 
 		uint SegIndex[3] = { 0xffffffff, 0xffffffff, 0xffffffff };
 
-		uint PrimID = BatchIndex * 1024 + GroupThreadID;
 
 #if PERMUTATION_CULLING
+		uint PrimID = BatchIndex * 1024 + GroupThreadID;
+		if (HairStrandsVF_bIsCullingEnable)
+		{
+			group_BatchNum = (HairStrandsVF_CullingIndirectBuffer[3] + 1023) / 1024;
+			if (bSegValid) bSegValid = (PrimID < (HairStrandsVF_CullingIndirectBuffer[3])) ? true : false;
+		}
+		else
+		{
+			if (bSegValid) bSegValid = (PrimID < (VertexCount)) ? true : false;
+		}
+
 		if (bSegValid && HairStrandsVF_bIsCullingEnable)
 		{
 			const uint VertexCountAfterCulling = HairStrandsVF_CullingIndirectBuffer[3];
@@ -274,20 +292,15 @@ void BinningCS(uint2 DispatchThreadID : SV_DispatchThreadID, uint GroupThreadID
 				}
 			}
 		}
+#else
+		uint PrimID = BatchIndex * 1024 + GroupThreadID + VertexStart;
+		if (bSegValid) bSegValid = (PrimID < (VertexStart + VertexCount)) ? true : false;
 #endif
 
 		uint BinID = GroupID - (uint(float(GroupID + 0.5) * RcpNumBinners)* NumBinners);
 		uint BinID2 = BinID + NumBinners;
 
-#if PERMUTATION_TEMPORAL_LAYERING
-		//remap linear frame index to improve TAA result (still a bit exprimental)
-		uint remapped = (View.StateFrameIndex) & 0x3;
-		remapped = ((remapped & 0x1) << 1) | ((remapped & 0x2) >> 1);
 
-		if (bSegValid) bSegValid = ((PrimID < HairStrandsVF_VertexCount) && ((PrimID & 0x3) == (remapped))) ? true : false;
-#else
-		if (bSegValid) bSegValid = (PrimID < HairStrandsVF_VertexCount) ? true : false;
-#endif
 		if (bSegValid) //currently hardcoded to a maximum of 3 tile bins - this only works corectly when segment length does not exceed TileSize - TODO: add support for N bins
 		{
 			float3 SP0;
@@ -454,12 +467,9 @@ void RasterCS(uint DispatchThreadID : SV_DispatchThreadID, uint GroupThreadID :
 	{
 		group_TileNum = VisTileArgs[0];
 		group_LoopNum = (float(group_TileNum) + float(NumRasterizers - 1))* RcpNumRasterizers;
-//		group_LoopNum = (float(1024) + float(NumRasterizers - 1)) * RcpNumRasterizers;
 
 		group_PositionOffset = HairStrandsVF_PositionOffsetBuffer[0].xyz;
 
-		group_ooTileLODScale = 1.0;// max(min(1.0 / (TileLODScale * 25.0), 16.0), 1.0);
-
 		/* no longer in use - keep for ref? Moving these values to group shared memory did seem to reduce VGPRs - more experimentation needed
 		group_RadScale = (((HairStrandsVF_TipScale - HairStrandsVF_RootScale) * HairStrandsVF_Radius * OutputResolutionf.x) / 63.0) / 255.0;
 		group_RadOffset = (HairStrandsVF_RootScale * HairStrandsVF_Radius * OutputResolutionf.x)/63.0;
@@ -610,11 +620,7 @@ void RasterCS(uint DispatchThreadID : SV_DispatchThreadID, uint GroupThreadID :
 					{
 						const float Rad = lerp(Rad0, Rad1, J / float(NumSteps)) / SP.w;
 
-#if PERMUTATION_TEMPORAL_LAYERING
-						InterlockedAdd(group_SubTile[LinearIndex].z, min(Rad, 1.0) * 1.0 * 8000.0 * group_ooTileLODScale);
-#else
-						InterlockedAdd(group_SubTile[LinearIndex].z, min(Rad, 1.0) * 1.0 * 1000.0 * group_ooTileLODScale);
-#endif
+						InterlockedAdd(group_SubTile[LinearIndex].z, min(Rad, 0.5) * 2.0 * 1000.0 * SampleWeight);
 					}
 				}					
 			}
diff --git a/Engine/Shaders/Private/HairStrands/HairStrandsVisibilityTile.usf b/Engine/Shaders/Private/HairStrands/HairStrandsVisibilityTile.usf
--- a/Engine/Shaders/Private/HairStrands/HairStrandsVisibilityTile.usf
+++ b/Engine/Shaders/Private/HairStrands/HairStrandsVisibilityTile.usf
@@ -40,7 +40,7 @@ void TileMainCS(uint2 DispatchThreadId : SV_DispatchThreadID, uint LinearIndex :
 		{
 			const uint Coverage = InputUintTexture.Load(uint3(PixelCoord, 0));
 			fHair = Coverage >= IntCoverageThreshold;
-			pHair = Coverage > 0 && Coverage < IntCoverageThreshold;
+			pHair = Coverage > 0;//&& Coverage < IntCoverageThreshold; //this is now consistent with non-int branch below
 		}
 		else
 		{
diff --git a/Engine/Shaders/Private/HairStrands/HairStrandsVoxelRasterCompute.usf b/Engine/Shaders/Private/HairStrands/HairStrandsVoxelRasterCompute.usf
--- a/Engine/Shaders/Private/HairStrands/HairStrandsVoxelRasterCompute.usf
+++ b/Engine/Shaders/Private/HairStrands/HairStrandsVoxelRasterCompute.usf
@@ -41,6 +41,10 @@ StructuredBuffer<FPackedVirtualVoxelNodeDesc> VirtualVoxelParams_NodeDescBuffer;
 
 RWTexture3D<uint> OutPageTexture;
 
+uint			VertexCount;
+uint			VertexStart;
+float SampleWeight;
+
 ///////////////////////////////////////////////////////////////////////////
 // Define
 
@@ -72,11 +76,10 @@ void MainCS(uint2 DispatchThreadID : SV_DispatchThreadID)
 #if PERMUTATION_CULLING == 1
 	uint VertexIndex0 = GetHairStrandsVertexFetchIndex_Culled(DispatchThreadID, GROUP_SIZE);
 #else
-	uint VertexIndex0 = GetHairStrandsVertexFetchIndex(DispatchThreadID, GROUP_SIZE, DispatchCountX);
+	uint VertexIndex0 = VertexStart + GetHairStrandsVertexFetchIndex(DispatchThreadID, GROUP_SIZE, DispatchCountX);
 #endif
-
-	uint VertexIndex1 = min(VertexIndex0 + 1, HairStrandsVF_VertexCount - 1);
-	if (VertexIndex0 >= HairStrandsVF_VertexCount)
+	uint VertexIndex1 = min(VertexIndex0 + 1, VertexStart + VertexCount - 1);
+	if (VertexIndex0 >= VertexStart + VertexCount)
 		return;
 
 	float LodRadiusScale0 = 1;
@@ -209,7 +212,7 @@ void MainCS(uint2 DispatchThreadID : SV_DispatchThreadID)
 				const float VoxelFixPointScale = GetVoxelDensityFixPointScale();
 				
 				const float HairCoverage = lerp(HairCoverage0, HairCoverage1, U);
-				uint RawData = HairCoverage * VoxelFixPointScale * HairStrandsVF_Density;
+				uint RawData = HairCoverage * VoxelFixPointScale * HairStrandsVF_Density * SampleWeight;
 				InterlockedAdd(OutPageTexture[VoxelPageCoord], RawData);
 
 				// Groom having raytraced geometry will cast-shadow on opaque geomtry with their RT-geometry, not with their 
diff --git a/Engine/Source/Runtime/Renderer/Private/HairStrands/HairStrandsInterface.cpp b/Engine/Source/Runtime/Renderer/Private/HairStrands/HairStrandsInterface.cpp
--- a/Engine/Source/Runtime/Renderer/Private/HairStrands/HairStrandsInterface.cpp
+++ b/Engine/Source/Runtime/Renderer/Private/HairStrands/HairStrandsInterface.cpp
@@ -64,6 +64,36 @@ static TAutoConsoleVariable<int32> CVarHairStrandsNonVisibleShadowCasting(
 	TEXT("Enable shadow casting for hair strands even when culled out from the primary view"),
 	ECVF_RenderThreadSafe);
 
+static TAutoConsoleVariable<int32> CVarHairStrandsContinuousDecimationReordering(
+	TEXT("r.HairStrands.ContinuousDecimationReordering"), 0,
+	TEXT("Enable strand reordering to allow Continuous LOD. Experimental"),
+	ECVF_RenderThreadSafe | ECVF_ReadOnly);
+
+static TAutoConsoleVariable<int32> CVarHairStrandsVisibilityComputeRaster(
+	TEXT("r.HairStrands.Visibility.ComputeRaster"), 0,
+	TEXT("Hair Visiblity uses raster compute. Experimental"),
+	ECVF_RenderThreadSafe);
+
+static TAutoConsoleVariable<int32> CVarHairStrandsVisibilityComputeRaster_ContinuousLOD(
+	TEXT("r.HairStrands.Visibility.ComputeRaster.ContinuousLOD"), 1,
+	TEXT("Enable Continuos LOD when using compute rasterization. Experimental"),
+	ECVF_RenderThreadSafe);
+
+static TAutoConsoleVariable<int32> CVarHairStrandsVisibilityComputeRaster_TemporalLayering(
+	TEXT("r.HairStrands.Visibility.ComputeRaster.TemporalLayering"), 0,
+	TEXT("Enable Experimental WIP Temporal Layering (requires TAA changes to work well)"),
+	ECVF_RenderThreadSafe);
+
+static TAutoConsoleVariable<int32> CVarHairStrandsVisibilityComputeRaster_TemporalLayering_LayerCount(
+	TEXT("r.HairStrands.Visibility.ComputeRaster.TemporalLayering.LayerCount"), 8,
+	TEXT("Enable Temporal Layering Override Index  8)"),
+	ECVF_RenderThreadSafe);
+
+static TAutoConsoleVariable<int32> CVarHairStrandsVisibilityComputeRaster_TemporalLayering_OverrideIndex(
+	TEXT("r.HairStrands.Visibility.ComputeRaster.TemporalLayering.OverrideIndex"), -1,
+	TEXT("Enable Temporal Layering Override Index (default: -1 = no override)"),
+	ECVF_RenderThreadSafe);
+
 ////////////////////////////////////////////////////////////////////////////////////////////////////////////////////////
 // Hair strands instance ref. counting for debug purpose only
 uint32 FHairStrandsInstance::GetRefCount() const
@@ -423,6 +453,91 @@ bool IsHairStrandsVisibleInShadows(const FViewInfo& View, const FHairStrandsInst
 	return bIsVisibleInShadow;
 }
 
+bool IsHairStrandContinuousDecimationReorderingEnabled()
+{
+	return CVarHairStrandsContinuousDecimationReordering.GetValueOnAnyThread() > 0;
+}
+
+bool IsHairVisibilityComputeRasterEnabled()
+{
+	return CVarHairStrandsVisibilityComputeRaster.GetValueOnAnyThread() > 0;
+}
+
+bool IsHairVisibilityComputeRasterContinuousLODEnabled()
+{
+	return IsHairStrandContinuousDecimationReorderingEnabled() && IsHairVisibilityComputeRasterEnabled() && (CVarHairStrandsVisibilityComputeRaster_ContinuousLOD.GetValueOnAnyThread() > 0);
+}
+
+float GetHairVisibilityComputeRasterContinuousLODScale(float ScreenSize)
+{
+	return FMath::Pow(FMath::Clamp(ScreenSize, 1.0f / 8.0f, 1.0f), 2.0f);
+}
+
+bool IsHairVisibilityComputeRasterTemporalLayeringEnabled()
+{
+	return IsHairStrandContinuousDecimationReorderingEnabled() && IsHairVisibilityComputeRasterEnabled() && (CVarHairStrandsVisibilityComputeRaster_TemporalLayering.GetValueOnAnyThread() > 0);
+}
+
+int32 GetHairVisibilityComputeRasterTemporalLayerCount()
+{
+	const uint32 LayerCount = FMath::Clamp(CVarHairStrandsVisibilityComputeRaster_TemporalLayering_LayerCount.GetValueOnAnyThread(), 1, 32);
+	return LayerCount;
+}
+
+uint32 GetHairVisibilityComputeRasterVertexStart(uint32 TemporalAASampleIndex, uint32 InVertexCount)
+{
+	uint32 VertexStart = 0;
+
+	if (IsHairVisibilityComputeRasterTemporalLayeringEnabled())
+	{
+		const uint32 LayerCount = GetHairVisibilityComputeRasterTemporalLayerCount();
+		const uint32 OverrideIndex = CVarHairStrandsVisibilityComputeRaster_TemporalLayering_OverrideIndex.GetValueOnAnyThread();
+
+		const uint32 VertexCount = InVertexCount / LayerCount;
+
+		VertexStart = (((OverrideIndex >= 0) ? OverrideIndex : TemporalAASampleIndex) % LayerCount) * VertexCount;
+	}
+
+	return VertexStart;
+}
+
+uint32 GetHairVisibilityComputeRasterVertexCount(float ScreenSize, uint32 InVertexCount)
+{
+	uint32 VertexCount = InVertexCount;
+
+	if (IsHairVisibilityComputeRasterTemporalLayeringEnabled())
+	{
+		VertexCount /= GetHairVisibilityComputeRasterTemporalLayerCount();
+	}
+
+	if (IsHairVisibilityComputeRasterContinuousLODEnabled())
+	{
+		VertexCount *= GetHairVisibilityComputeRasterContinuousLODScale(ScreenSize);
+	}
+
+	return VertexCount;
+}
+
+float GetHairVisibilityComputeRasterSampleWeight(float ScreenSize)
+{
+	float SampleWeight = 1.0;
+
+	if (IsHairVisibilityComputeRasterTemporalLayeringEnabled())
+	{
+		// sample weight should increase if layer count increases
+		SampleWeight *= GetHairVisibilityComputeRasterTemporalLayerCount();
+	}
+
+	if (IsHairVisibilityComputeRasterContinuousLODEnabled())
+	{
+		//sample weight should increase if decimation scale is lower
+		SampleWeight /= GetHairVisibilityComputeRasterContinuousLODScale(ScreenSize);
+	}
+
+	return SampleWeight;
+}
+
+
 ///////////////////////////////////////////////////////////////////////////////////////////////////
 // Bookmark API
 THairStrandsBookmarkFunction  GHairStrandsBookmarkFunction = nullptr;
diff --git a/Engine/Source/Runtime/Renderer/Private/HairStrands/HairStrandsVisibility.cpp b/Engine/Source/Runtime/Renderer/Private/HairStrands/HairStrandsVisibility.cpp
--- a/Engine/Source/Runtime/Renderer/Private/HairStrands/HairStrandsVisibility.cpp
+++ b/Engine/Source/Runtime/Renderer/Private/HairStrands/HairStrandsVisibility.cpp
@@ -17,6 +17,7 @@
 #include "ScenePrivate.h"
 #include "SceneTextureReductions.h"
 #include "PixelShaderUtils.h"
+#include "SceneManagement.h"
 
 DECLARE_GPU_STAT(HairStrandsVisibility);
 
@@ -53,23 +54,17 @@ static FAutoConsoleVariableRef CVarGHairVisibilityPPLL_MeanSamplePerPixel(TEXT("
 static float GHairStrandsViewHairCountDepthDistanceThreshold = 30.f;
 static FAutoConsoleVariableRef CVarHairStrandsViewHairCountDepthDistanceThreshold(TEXT("r.HairStrands.Visibility.HairCount.DistanceThreshold"), GHairStrandsViewHairCountDepthDistanceThreshold, TEXT("Distance threshold defining if opaque depth get injected into the 'view-hair-count' buffer."));
 
-int32 GHairVisibilityComputeRaster = 0;
 int32 GHairVisibilityComputeRaster_Culling = 0;
-int32 GHairVisibilityComputeRaster_CLOD = 0;
 int32 GHairVisibilityComputeRaster_MaxTiles = 8192;
 int32 GHairVisibilityComputeRaster_TileSize = 32;
 int32 GHairVisibilityComputeRaster_NumBinners = 32;
 int32 GHairVisibilityComputeRaster_NumRasterizers = 256;
-int32 GHairVisibilityComputeRaster_TemporalLayering = 0;
 
-static FAutoConsoleVariableRef CVarHairStrandsVisibilityComputeRaster(TEXT("r.HairStrands.Visibility.ComputeRaster"), GHairVisibilityComputeRaster, TEXT("Hair Visiblity uses raster compute. Experimental"), ECVF_Scalability | ECVF_RenderThreadSafe);
 static FAutoConsoleVariableRef CVarHairVisibilityComputeRaster_Culling(TEXT("r.HairStrands.Visibility.ComputeRaster.Culling"), GHairVisibilityComputeRaster_Culling, TEXT("Use culling buffers with compute rasterization."));
-static FAutoConsoleVariableRef CVarHairVisibilityComputeRaster_CLOD(TEXT("r.HairStrands.Visibility.ComputeRaster.CLOD"), GHairVisibilityComputeRaster_CLOD, TEXT("Enable Continuous LOD when using compute rasterization. Experimental - needs further improvement"));
-static FAutoConsoleVariableRef CHairVisibilityComputeRaster_MaxTiles(TEXT("r.HairStrands.Visibility.ComputeRaster.MaxTiles"), GHairVisibilityComputeRaster_MaxTiles, TEXT("Maximum number of tiles used for compute rasterization. 8192 is default"));
-static FAutoConsoleVariableRef CHairVisibilityComputeRaster_TileSize(TEXT("r.HairStrands.Visibility.ComputeRaster.TileSize"), GHairVisibilityComputeRaster_TileSize, TEXT("Tile size used for compute rasterization. Experimental - only size of 32 currently supported"));
-static FAutoConsoleVariableRef CHairVisibilityComputeRaster_NumBinners(TEXT("r.HairStrands.Visibility.ComputeRaster.NumBinners"), GHairVisibilityComputeRaster_NumBinners, TEXT("Number of Binners used in Binning compute rasterization pass. 32 is default"));
-static FAutoConsoleVariableRef CHairVisibilityComputeRaster_NumRasterizers(TEXT("r.HairStrands.Visibility.ComputeRaster.NumRasterizers"), GHairVisibilityComputeRaster_NumRasterizers, TEXT("Number of Rasterizers used compute rasterization. 256 is default"));
-static FAutoConsoleVariableRef CHairVisibilityComputeRaster_TemporalLayering(TEXT("r.HairStrands.Visibility.ComputeRaster.TemporalLayering"), GHairVisibilityComputeRaster_TemporalLayering, TEXT("Enable Experimental WIP Temporal Layering (requires TAA changes to work well)"));
+static FAutoConsoleVariableRef CVarHairVisibilityComputeRaster_MaxTiles(TEXT("r.HairStrands.Visibility.ComputeRaster.MaxTiles"), GHairVisibilityComputeRaster_MaxTiles, TEXT("Maximum number of tiles used for compute rasterization. 8192 is default"));
+static FAutoConsoleVariableRef CVarHairVisibilityComputeRaster_TileSize(TEXT("r.HairStrands.Visibility.ComputeRaster.TileSize"), GHairVisibilityComputeRaster_TileSize, TEXT("Tile size used for compute rasterization. Experimental - only size of 32 currently supported"));
+static FAutoConsoleVariableRef CVarHairVisibilityComputeRaster_NumBinners(TEXT("r.HairStrands.Visibility.ComputeRaster.NumBinners"), GHairVisibilityComputeRaster_NumBinners, TEXT("Number of Binners used in Binning compute rasterization pass. 32 is default"));
+static FAutoConsoleVariableRef CVarHairVisibilityComputeRaster_NumRasterizers(TEXT("r.HairStrands.Visibility.ComputeRaster.NumRasterizers"), GHairVisibilityComputeRaster_NumRasterizers, TEXT("Number of Rasterizers used compute rasterization. 256 is default"));
 
 static float GHairStrandsFullCoverageThreshold = 0.98f;
 static FAutoConsoleVariableRef CVarHairStrandsFullCoverageThreshold(TEXT("r.HairStrands.Visibility.FullCoverageThreshold"), GHairStrandsFullCoverageThreshold, TEXT("Define the coverage threshold at which a pixel is considered fully covered."));
@@ -138,13 +133,15 @@ inline bool DoesSupportRasterCompute()
 	return true; //GRHISupportsAtomicUInt64; //TODO: re-enable support for 64 bit visibility buffers
 }
 
+bool IsHairVisibilityComputeRasterEnabled();
+
 inline EHairVisibilityRenderMode GetHairVisibilityRenderMode()
 {
 	if (GHairVisibilityPPLL > 0)
 	{
 		return HairVisibilityRenderMode_PPLL;
 	}
-	else if (GHairVisibilityComputeRaster > 0 && DoesSupportRasterCompute())
+	else if (IsHairVisibilityComputeRasterEnabled() && DoesSupportRasterCompute())
 	{
 		return HairVisibilityRenderMode_ComputeRaster;
 	}
@@ -289,9 +286,11 @@ float GetHairStrandsFullCoverageThreshold()
 	return FMath::Clamp(GHairStrandsFullCoverageThreshold, 0.1f, 1.f);
 }
 
+bool IsHairVisibilityComputeRasterTemporalLayeringEnabled();
+
 uint32 GetHairStrandsIntCoverageThreshold()
 {
-	return ((GHairVisibilityComputeRaster) ? ((GHairVisibilityComputeRaster_TemporalLayering) ? 16000u : 4000u) : 1000u);
+	return (IsHairVisibilityComputeRasterEnabled() ? ((IsHairVisibilityComputeRasterTemporalLayeringEnabled()) ? 1000u : 1000u) : 1000u);
 }
 
 ///////////////////////////////////////////////////////////////////////////////////////////////////
@@ -2895,8 +2894,7 @@ class FVisiblityRasterComputeBinningCS : public FGlobalShader
 
 	class FGroupSize : SHADER_PERMUTATION_SPARSE_INT("PERMUTATION_GROUP_SIZE", 32, 64);
 	class FCulling : SHADER_PERMUTATION_BOOL("PERMUTATION_CULLING");
-	class FTemporalLayering : SHADER_PERMUTATION_BOOL("PERMUTATION_TEMPORAL_LAYERING");
-	using FPermutationDomain = TShaderPermutationDomain<FGroupSize, FCulling, FTemporalLayering>;
+	using FPermutationDomain = TShaderPermutationDomain<FGroupSize, FCulling>;
 
 	BEGIN_SHADER_PARAMETER_STRUCT(FParameters, )
 		SHADER_PARAMETER(uint32, MacroGroupId)
@@ -2930,8 +2928,8 @@ class FVisiblityRasterComputeBinningCS : public FGlobalShader
 		SHADER_PARAMETER_RDG_BUFFER_UAV(RWByteAddressBuffer, OutVisTileData)
 		SHADER_PARAMETER_STRUCT_REF(FViewUniformShaderParameters, ViewUniformBuffer)
 		RDG_TEXTURE_ACCESS(VisTileBinningGridTex, ERHIAccess::UAVCompute)
-		SHADER_PARAMETER(float, TileLODScale)
-		SHADER_PARAMETER(float, TileBaseDepth)
+		SHADER_PARAMETER(uint32, VertexCount)
+		SHADER_PARAMETER(uint32, VertexStart)
 		SHADER_PARAMETER(FIntPoint, TileRes)
 		SHADER_PARAMETER(uint32, TileSize)
 		SHADER_PARAMETER(float, RcpTileSize)
@@ -2957,8 +2955,7 @@ class FVisiblityRasterComputeRasterizeCS : public FGlobalShader
 	SHADER_USE_PARAMETER_STRUCT(FVisiblityRasterComputeRasterizeCS, FGlobalShader);
 
 	class FGroupSize : SHADER_PERMUTATION_SPARSE_INT("PERMUTATION_GROUP_SIZE", 32, 64);
-	class FTemporalLayering : SHADER_PERMUTATION_BOOL("PERMUTATION_TEMPORAL_LAYERING");
-	using FPermutationDomain = TShaderPermutationDomain<FGroupSize, FTemporalLayering>;
+	using FPermutationDomain = TShaderPermutationDomain<FGroupSize>;
 
 	BEGIN_SHADER_PARAMETER_STRUCT(FParameters, )
 		SHADER_PARAMETER(uint32, MacroGroupId)
@@ -2985,7 +2982,7 @@ class FVisiblityRasterComputeRasterizeCS : public FGlobalShader
 		SHADER_PARAMETER_RDG_BUFFER_SRV(Buffer, HairStrandsVF_CullingRadiusScaleBuffer)
 		RDG_BUFFER_ACCESS(IndirectBufferArgs, ERHIAccess::IndirectArgs)
 		SHADER_PARAMETER_RDG_TEXTURE(Texture2D, SceneDepthTexture)
-		SHADER_PARAMETER(float, TileLODScale)
+		SHADER_PARAMETER(float, SampleWeight)
 		SHADER_PARAMETER(uint32, TileSize)
 		SHADER_PARAMETER(float, RcpTileSize)
 		SHADER_PARAMETER(uint32, SqrTileSize)
@@ -3014,6 +3011,11 @@ public:
 
 IMPLEMENT_GLOBAL_SHADER(FVisiblityRasterComputeRasterizeCS, "/Engine/Private/HairStrands/HairStrandsVisibilityRasterCompute.usf", "RasterCS", SF_Compute);
 
+uint32 GetHairVisibilityComputeRasterVertexStart(uint32 TemporalAASampleIndex, uint32 InVertexCount);
+uint32 GetHairVisibilityComputeRasterVertexCount(float ScreenSize, uint32 InVertexCount);
+float GetHairVisibilityComputeRasterSampleWeight(float ScreenSize);
+bool IsHairStrandContinuousDecimationReorderingEnabled();
+
 static FRasterComputeOutput AddVisibilityComputeRasterPass(
 	FRDGBuilder& GraphBuilder,
 	const FViewInfo& ViewInfo,
@@ -3099,26 +3101,22 @@ static FRasterComputeOutput AddVisibilityComputeRasterPass(
 	SetUpViewHairRenderInfo(ViewInfo, bEnableMSAA, ViewInfo.CachedViewUniformShaderParameters->HairRenderInfo, ViewInfo.CachedViewUniformShaderParameters->HairRenderInfoBits, ViewInfo.CachedViewUniformShaderParameters->HairComponents);
 	ViewUniformShaderParameters = TUniformBufferRef<FViewUniformShaderParameters>::CreateUniformBufferImmediate(*ViewInfo.CachedViewUniformShaderParameters, UniformBuffer_SingleFrame);
 
-	const uint32 FrameIdMode8 = ViewInfo.ViewState ? (ViewInfo.ViewState->GetFrameIndex() % 8) : 0;
 	const uint32 GroupSize = GetVendorOptimalGroupSize1D();
-
-	const bool bTemporalLayering = GHairVisibilityComputeRaster_TemporalLayering ? true : false;
+	
+	const uint32 TemporalAASampleIndex = ViewInfo.ViewState ? ViewInfo.ViewState->GetCurrentTemporalAASampleIndex() : 0;
 
 	FVisiblityRasterComputeBinningCS::FPermutationDomain BinningPermutationVector_CullingOff;
 	BinningPermutationVector_CullingOff.Set<FVisiblityRasterComputeBinningCS::FGroupSize>(GroupSize);
 	BinningPermutationVector_CullingOff.Set<FVisiblityRasterComputeBinningCS::FCulling>(false);
-	BinningPermutationVector_CullingOff.Set<FVisiblityRasterComputeBinningCS::FTemporalLayering>(bTemporalLayering);
 	TShaderMapRef<FVisiblityRasterComputeBinningCS> ComputeShaderBinning_CullingOff(ViewInfo.ShaderMap, BinningPermutationVector_CullingOff);
 
 	FVisiblityRasterComputeBinningCS::FPermutationDomain BinningPermutationVector_CullingOn;
 	BinningPermutationVector_CullingOn.Set<FVisiblityRasterComputeBinningCS::FGroupSize>(GroupSize);
 	BinningPermutationVector_CullingOn.Set<FVisiblityRasterComputeBinningCS::FCulling>(true);
-	BinningPermutationVector_CullingOn.Set<FVisiblityRasterComputeBinningCS::FTemporalLayering>(bTemporalLayering);
 	TShaderMapRef<FVisiblityRasterComputeBinningCS> ComputeShaderBinning_CullingOn(ViewInfo.ShaderMap, BinningPermutationVector_CullingOn);
 
 	FVisiblityRasterComputeRasterizeCS::FPermutationDomain RasterPermutationVector;
 	RasterPermutationVector.Set<FVisiblityRasterComputeRasterizeCS::FGroupSize>(GroupSize);
-	RasterPermutationVector.Set<FVisiblityRasterComputeRasterizeCS::FTemporalLayering>(bTemporalLayering);
 	TShaderMapRef<FVisiblityRasterComputeRasterizeCS> ComputeShaderRaster(ViewInfo.ShaderMap, RasterPermutationVector);
 
 	TShaderMapRef<FVisiblityRasterComputePrepareDepthGridCS> ComputeShaderPrepareDepthGrid(ViewInfo.ShaderMap);	
@@ -3138,6 +3136,8 @@ static FRasterComputeOutput AddVisibilityComputeRasterPass(
 	{
 		const FHairStrandsMacroGroupData::TPrimitiveInfos& PrimitiveSceneInfos = MacroGroup.PrimitivesInfos;
 
+		const FSphere MacroBoundsSphere = MacroGroup.Bounds.GetSphere();
+
 		for (const FHairStrandsMacroGroupData::PrimitiveInfo& PrimitiveInfo : PrimitiveSceneInfos)
 		{
 			check(PrimitiveInfo.Mesh && PrimitiveInfo.Mesh->Elements.Num() > 0);
@@ -3147,8 +3147,13 @@ static FRasterComputeOutput AddVisibilityComputeRasterPass(
 			const FHairGroupPublicData::FVertexFactoryInput& VFInput = HairGroupPublicData->VFInput;
 			const FMatrix44d LocalToWorldPrimitiveTransform = VFInput.LocalToWorldTransform.ToMatrixWithScale(); 
 
-			const FVector4f ScreenPosition = FVector4f(ViewInfo.ViewMatrices.GetViewProjectionMatrix().TransformPosition(LocalToWorldPrimitiveTransform.GetOrigin() + FVector4d(0, 0, 165, 0))); //hardcoded - TODO find a way to get correct offset to origin
 			const bool bCullingEnable = GHairVisibilityComputeRaster_Culling ? HairGroupPublicData->GetCullingResultAvailable() : false;
+			
+			const float ScreenSize = ComputeBoundsScreenSize(FVector4(MacroBoundsSphere.Center, 1), MacroBoundsSphere.W, ViewInfo);
+
+			const uint32 VertexStart = GetHairVisibilityComputeRasterVertexStart(TemporalAASampleIndex, VFInput.Strands.VertexCount);
+			const uint32 VertexCount = GetHairVisibilityComputeRasterVertexCount(ScreenSize, VFInput.Strands.VertexCount);
+			const float SampleWeight = GetHairVisibilityComputeRasterSampleWeight(ScreenSize);
 
 			//reset buffers
 			uint32 IndexGridClearValues[4] = { 0x0,0x0,0x0,0x0 };
@@ -3180,8 +3185,8 @@ static FRasterComputeOutput AddVisibilityComputeRasterPass(
 			BinningParameters->HairStrandsVF_Density = VFInput.Strands.HairDensity;
 			BinningParameters->HairStrandsVF_LocalToWorldPrimitiveTransform = FMatrix44f(LocalToWorldPrimitiveTransform); // LWC_TODO: Precision loss
 
-			BinningParameters->TileBaseDepth = ScreenPosition.Z / ScreenPosition.W;
-			BinningParameters->TileLODScale = GHairVisibilityComputeRaster_CLOD ? 1.0 / ScreenPosition.W : 1.0;
+			BinningParameters->VertexCount = VertexCount;
+			BinningParameters->VertexStart = VertexStart;
 			BinningParameters->TileRes = TileGridRes;
 			BinningParameters->VisTileBinningGridTex = OutVisTileBinningGrid;
 			BinningParameters->TileSize = TileSize;
@@ -3209,7 +3214,6 @@ static FRasterComputeOutput AddVisibilityComputeRasterPass(
 			RasterParameters->ResolutionMultiplier = Out.ResolutionMultiplier;
 			RasterParameters->MacroGroupId = MacroGroup.MacroGroupId;
 			RasterParameters->MaxRasterCount = TileSize; //TODO: implement tile clipping in raster pass to exisiting avoid edge case
-			RasterParameters->FrameIdMod8 = FrameIdMode8;
 			RasterParameters->HairMaterialId = PrimitiveInfo.MaterialId;
 			RasterParameters->ViewUniformBuffer = ViewUniformShaderParameters;
 			RasterParameters->SceneDepthTexture = SceneDepthTexture;
@@ -3235,7 +3239,7 @@ static FRasterComputeOutput AddVisibilityComputeRasterPass(
 			RasterParameters->VisTileIndirect = GraphBuilder.CreateSRV(OutVisTileIndirect);
 			RasterParameters->IndirectBufferArgs = OutVisTileIndirect;
 			RasterParameters->TileRes = TileGridRes;
-			RasterParameters->TileLODScale = GHairVisibilityComputeRaster_CLOD ? 1.0 / ScreenPosition.W : 1.0;
+			RasterParameters->SampleWeight = SampleWeight;
 			RasterParameters->TileSize = TileSize;
 			RasterParameters->RcpTileSize = 1.0 / TileSize;
 			RasterParameters->SqrTileSize = TileSize * TileSize;
@@ -3622,7 +3626,7 @@ bool HasPositionsChanged(FRDGBuilder& GraphBuilder, const FViewInfo& View)
 	{
 		FHairGroupPublicData* HairGroupPublicData = GetHairData(Batch.Mesh);
 		check(HairGroupPublicData);
-		const int32 LODIndex = FMath::FloorToInt(HairGroupPublicData->LODIndex);
+		const int32 LODIndex = HairGroupPublicData->GetIntLODIndex();
 		const bool bHasSimulationOrSkinning = 
 			HairGroupPublicData->GetGeometryType(LODIndex) == EHairGeometryType::Strands &&
 			(HairGroupPublicData->IsSimulationEnable(LODIndex) || HairGroupPublicData->GetBindingType(LODIndex) == EHairBindingType::Skinning);
@@ -3709,10 +3713,11 @@ void RenderHairStrandsVisibilityBuffer(
 			check(SceneDepthTexture);
 			const FIntPoint Resolution = SceneDepthTexture->Desc.Extent;
 
-			const bool bRunColorAndDepthPatching = SceneGBufferBTexture && SceneColorTexture;
 			const EHairVisibilityRenderMode RenderMode = GetHairVisibilityRenderMode();
 			check(RenderMode == HairVisibilityRenderMode_MSAA_Visibility || RenderMode == HairVisibilityRenderMode_PPLL || RenderMode == HairVisibilityRenderMode_ComputeRaster);
 
+			const bool bRunColorAndDepthPatching = SceneGBufferBTexture && SceneColorTexture && RenderMode != HairVisibilityRenderMode_ComputeRaster;
+
 			FRDGTextureRef HairOnlyDepthTexture = GraphBuilder.CreateTexture(SceneDepthTexture->Desc, TEXT("Hair.HairOnlyDepthTexture"));
 			FRDGTextureRef CoverageTexture = nullptr;
 			FRDGTextureRef CompactNodeIndex = nullptr;
@@ -4202,4 +4207,4 @@ void RenderHairStrandsVisibilityBuffer(
 		#endif
 		}
 	}
-}
\ No newline at end of file
+}
diff --git a/Engine/Source/Runtime/Renderer/Private/HairStrands/HairStrandsVoxelization.cpp b/Engine/Source/Runtime/Renderer/Private/HairStrands/HairStrandsVoxelization.cpp
--- a/Engine/Source/Runtime/Renderer/Private/HairStrands/HairStrandsVoxelization.cpp
+++ b/Engine/Source/Runtime/Renderer/Private/HairStrands/HairStrandsVoxelization.cpp
@@ -1247,6 +1247,9 @@ class FVoxelRasterComputeCS : public FGlobalShader
 		SHADER_PARAMETER_RDG_BUFFER_SRV(Buffer, VoxelizationViewInfoBuffer)
 		RDG_BUFFER_ACCESS(IndirectBufferArgs, ERHIAccess::IndirectArgs)
 		SHADER_PARAMETER_RDG_TEXTURE_UAV(RWTexture3D, OutPageTexture)
+		SHADER_PARAMETER(uint32, VertexCount)
+		SHADER_PARAMETER(uint32, VertexStart)
+		SHADER_PARAMETER(float, SampleWeight)
 	END_SHADER_PARAMETER_STRUCT()
 
 public:
@@ -1260,6 +1263,11 @@ public:
 
 IMPLEMENT_GLOBAL_SHADER(FVoxelRasterComputeCS, "/Engine/Private/HairStrands/HairStrandsVoxelRasterCompute.usf", "MainCS", SF_Compute);
 
+uint32 GetHairVisibilityComputeRasterVertexStart(uint32 TemporalAASampleIndex, uint32 InVertexCount);
+uint32 GetHairVisibilityComputeRasterVertexCount(float ScreenSize, uint32 InVertexCount);
+float GetHairVisibilityComputeRasterSampleWeight(float ScreenSize);
+bool IsHairStrandContinuousDecimationReorderingEnabled();
+
 static void AddVirtualVoxelizationComputeRasterPass(
 	FRDGBuilder& GraphBuilder,
 	const FViewInfo* ViewInfo,
@@ -1274,6 +1282,8 @@ static void AddVirtualVoxelizationComputeRasterPass(
 
 	if (ViewInfo)
 	{
+		const uint32 TemporalAASampleIndex = ViewInfo->ViewState ? ViewInfo->ViewState->GetCurrentTemporalAASampleIndex() : 0;
+
 		const FHairStrandsMacroGroupData::TPrimitiveInfos& PrimitiveSceneInfos = MacroGroup.PrimitivesInfos;
 
 		FRDGBufferSRVRef VoxelizationViewInfoBufferSRV = GraphBuilder.CreateSRV(VoxelResources.VoxelizationViewInfoBuffer);
@@ -1305,6 +1315,10 @@ static void AddVirtualVoxelizationComputeRasterPass(
 					continue;
 				}
 
+				const uint32 VertexStart = GetHairVisibilityComputeRasterVertexStart(TemporalAASampleIndex, VFInput.Strands.VertexCount);
+				const uint32 VertexCount = GetHairVisibilityComputeRasterVertexCount(HairGroupPublicData->DebugScreenSize, VFInput.Strands.VertexCount);
+				const float SampleWeight = GetHairVisibilityComputeRasterSampleWeight(HairGroupPublicData->DebugScreenSize);
+
 				FTransform LocalToTranslatedWorldTransform = VFInput.LocalToWorldTransform;
 				LocalToTranslatedWorldTransform.AddToTranslation(TranslatedWorldOffset);
 
@@ -1327,7 +1341,13 @@ static void AddVirtualVoxelizationComputeRasterPass(
 				PassParameters->HairStrandsVF_Density = VFInput.Strands.HairDensity;
 				PassParameters->HairStrandsVF_LocalToTranslatedWorldPrimitiveTransform = FMatrix44f(LocalToTranslatedWorldTransform.ToMatrixWithScale());
 				PassParameters->HairStrandsVF_bHasRaytracedGeometry = VFInput.Strands.bUseRaytracingGeometry ? 1u : 0u;
-				const bool bCullingEnable = HairGroupPublicData->GetCullingResultAvailable();
+
+				PassParameters->VertexCount = VertexCount;
+				PassParameters->VertexStart = VertexStart;
+				PassParameters->SampleWeight = SampleWeight;
+			
+				const bool bCullingEnable = !IsHairStrandContinuousDecimationReorderingEnabled() && HairGroupPublicData->GetCullingResultAvailable();
+
 				PassParameters->HairStrandsVF_bIsCullingEnable = bCullingEnable ? 1 : 0;
 
 				if (bCullingEnable)
@@ -1342,7 +1362,7 @@ static void AddVirtualVoxelizationComputeRasterPass(
 				}
 				else
 				{
-					const FIntVector DispatchCount = ComputeDispatchCount(PassParameters->HairStrandsVF_VertexCount, GroupSize);
+					const FIntVector DispatchCount = ComputeDispatchCount(PassParameters->VertexCount, GroupSize);
 					PassParameters->DispatchCountX = DispatchCount.X;
 					FComputeShaderUtils::AddPass(GraphBuilder, RDG_EVENT_NAME("HairStrands::VoxelComputeRaster(culling=off)"), ComputeShader_CullingOff, PassParameters, DispatchCount);
 				}
diff --git a/Engine/Source/Runtime/Renderer/Public/HairStrandsInterface.h b/Engine/Source/Runtime/Renderer/Public/HairStrandsInterface.h
--- a/Engine/Source/Runtime/Renderer/Public/HairStrandsInterface.h
+++ b/Engine/Source/Runtime/Renderer/Public/HairStrandsInterface.h
@@ -415,6 +415,12 @@ RENDERER_API bool IsHairStrandsSimulationEnable();
 // Return true if hair binding is enabled (i.e., hair can be attached to skeletal mesh)
 RENDERER_API bool IsHairStrandsBindingEnable();
 
+// Return true if strand reordering for compute raster continuous LOD is enabled
+RENDERER_API bool IsHairStrandContinuousDecimationReorderingEnabled();
+
+// Return true if continuous LOD is enabled - implies computer raster and continuous decimation reordering is true
+RENDERER_API bool IsHairVisibilityComputeRasterContinuousLODEnabled();
+
 ////////////////////////////////////////////////////////////////////////////////////////////////////////////////////////
 
 typedef TArray<FRHIUnorderedAccessView*> FBufferTransitionQueue;
