Fix distance field shadows artifacts in the right eye (UE-150939).

diff --git a/Engine/Source/Runtime/Renderer/Private/ShadowRendering.cpp b/Engine/Source/Runtime/Renderer/Private/ShadowRendering.cpp
--- a/Engine/Source/Runtime/Renderer/Private/ShadowRendering.cpp
+++ b/Engine/Source/Runtime/Renderer/Private/ShadowRendering.cpp
@@ -2209,7 +2209,7 @@ void FSceneRenderer::RenderShadowProjections(
 				{
 					FProjectedShadowInfo* ProjectedShadowInfo = DistanceFieldShadows[ShadowIndex];
 
-					if (Views.Num() == 1 || ProjectedShadowInfo->DependentView == &View || !ProjectedShadowInfo->DependentView)
+					if (Views.Num() == 1 || !ProjectedShadowInfo->DependentView || ProjectedShadowInfo->DependentView == &View || ProjectedShadowInfo->DependentView == View.GetPrimaryView())
 					{
 						ProjectedShadowInfo->RenderRayTracedDistanceFieldProjection(
 							GraphBuilder,
@@ -2251,7 +2251,7 @@ void FSceneRenderer::BeginAsyncDistanceFieldShadowProjections(FRDGBuilder& Graph
 					ScissorRect = View.ViewRect;
 				}
 
-				if (ScissorRect.Area() > 0 && (Views.Num() == 1 || ProjectedShadowInfo->DependentView == &View || !ProjectedShadowInfo->DependentView))
+				if (ScissorRect.Area() > 0 && (Views.Num() == 1 || ProjectedShadowInfo->DependentView == &View || ProjectedShadowInfo->DependentView == View.GetPrimaryView() || !ProjectedShadowInfo->DependentView))
 				{
 					// Kick off distance field shadow calculation in async compute
 					// Don't need store result reference because it is internally cached by FProjectedShadowInfo
