Adds an Editor Primitives stat gpu bucket

diff --git a/Engine/Source/Runtime/Renderer/Private/EditorPrimitivesRendering.cpp b/Engine/Source/Runtime/Renderer/Private/EditorPrimitivesRendering.cpp
--- a/Engine/Source/Runtime/Renderer/Private/EditorPrimitivesRendering.cpp
+++ b/Engine/Source/Runtime/Renderer/Private/EditorPrimitivesRendering.cpp
@@ -10,6 +10,9 @@
 #include "MobileBasePassRendering.h"
 #include "MeshPassProcessor.inl"
 
+
+DEFINE_GPU_STAT(EditorPrimitives);
+
 FEditorPrimitivesBasePassMeshProcessor::FEditorPrimitivesBasePassMeshProcessor(const FScene* Scene, ERHIFeatureLevel::Type InFeatureLevel, const FSceneView* InViewIfDynamicMeshCommand, const FMeshPassProcessorRenderState& InDrawRenderState, bool bInTranslucentBasePass, FMeshPassDrawListContext* InDrawListContext) 
 	: FMeshPassProcessor(Scene, InFeatureLevel, InViewIfDynamicMeshCommand, InDrawListContext)
 	, PassDrawRenderState(InDrawRenderState)
diff --git a/Engine/Source/Runtime/Renderer/Private/EditorPrimitivesRendering.h b/Engine/Source/Runtime/Renderer/Private/EditorPrimitivesRendering.h
--- a/Engine/Source/Runtime/Renderer/Private/EditorPrimitivesRendering.h
+++ b/Engine/Source/Runtime/Renderer/Private/EditorPrimitivesRendering.h
@@ -23,4 +23,6 @@ private:
 
 	bool ProcessDeferredShadingPath(const FMeshBatch& MeshBatch, uint64 BatchElementMask, const FMaterial& Material, const FMaterialRenderProxy& MaterialRenderProxy, const FPrimitiveSceneProxy* PrimitiveSceneProxy, int32 StaticMeshId);
 	bool ProcessMobileShadingPath(const FMeshBatch& MeshBatch, uint64 BatchElementMask, const FMaterial& Material, const FMaterialRenderProxy& MaterialRenderProxy, const FPrimitiveSceneProxy* PrimitiveSceneProxy, int32 StaticMeshId);
-};
\ No newline at end of file
+};
+
+DECLARE_GPU_STAT_NAMED_EXTERN(EditorPrimitives, TEXT("Editor Primitives"));
diff --git a/Engine/Source/Runtime/Renderer/Private/PostProcess/PostProcessCompositeEditorPrimitives.cpp b/Engine/Source/Runtime/Renderer/Private/PostProcess/PostProcessCompositeEditorPrimitives.cpp
--- a/Engine/Source/Runtime/Renderer/Private/PostProcess/PostProcessCompositeEditorPrimitives.cpp
+++ b/Engine/Source/Runtime/Renderer/Private/PostProcess/PostProcessCompositeEditorPrimitives.cpp
@@ -327,6 +327,7 @@ FScreenPassTexture AddEditorPrimitivePass(
 	// Load the color target if it already exists.
 	const FScreenPassTextureViewport EditorPrimitivesViewport(EditorPrimitiveColor, Inputs.SceneColor.ViewRect);
 
+	RDG_GPU_STAT_SCOPE(GraphBuilder, EditorPrimitives);
 	RDG_EVENT_SCOPE(GraphBuilder, "CompositeEditorPrimitives %dx%d MSAA=%d",
 		EditorPrimitivesViewport.Rect.Width(),
 		EditorPrimitivesViewport.Rect.Height(),
diff --git a/Engine/Source/Runtime/Renderer/Private/PostProcess/PostProcessSelectionOutline.cpp b/Engine/Source/Runtime/Renderer/Private/PostProcess/PostProcessSelectionOutline.cpp
--- a/Engine/Source/Runtime/Renderer/Private/PostProcess/PostProcessSelectionOutline.cpp
+++ b/Engine/Source/Runtime/Renderer/Private/PostProcess/PostProcessSelectionOutline.cpp
@@ -4,6 +4,7 @@
 
 #include "PostProcess/PostProcessSelectionOutline.h"
 #include "PostProcess/PostProcessCompositeEditorPrimitives.h"
+#include "EditorPrimitivesRendering.h"
 #include "SceneTextureParameters.h"
 #include "CanvasTypes.h"
 #include "RenderTargetTemp.h"
@@ -57,6 +58,7 @@ FScreenPassTexture AddSelectionOutlinePass(FRDGBuilder& GraphBuilder, const FVie
 	const bool bNaniteEnabled = NaniteRasterResults != nullptr;
 
 	RDG_EVENT_SCOPE(GraphBuilder, "EditorSelectionOutlines");
+	RDG_GPU_STAT_SCOPE(GraphBuilder, EditorPrimitives);
 
 	const uint32 NumSamples = View.GetSceneTexturesConfig().NumSamples;
 
