Fixes PCD3D_SM6 on a GTX 1080

diff --git a/Engine/Shaders/Private/VirtualShadowMaps/VirtualShadowMapProjection.usf b/Engine/Shaders/Private/VirtualShadowMaps/VirtualShadowMapProjection.usf
--- a/Engine/Shaders/Private/VirtualShadowMaps/VirtualShadowMapProjection.usf
+++ b/Engine/Shaders/Private/VirtualShadowMaps/VirtualShadowMapProjection.usf
@@ -1,9 +1,5 @@
 // Copyright Epic Games, Inc. All Rights Reserved.
 
-/*=============================================================================
-	VirtualShadowMapProjection.usf: 
-=============================================================================*/
-
 #define SUPPORT_CONTACT_SHADOWS 1
 #include "../Common.ush"
 #include "../SceneTexturesCommon.ush"
diff --git a/Engine/Source/Runtime/Renderer/Private/VirtualShadowMaps/VirtualShadowMapProjection.cpp b/Engine/Source/Runtime/Renderer/Private/VirtualShadowMaps/VirtualShadowMapProjection.cpp
--- a/Engine/Source/Runtime/Renderer/Private/VirtualShadowMaps/VirtualShadowMapProjection.cpp
+++ b/Engine/Source/Runtime/Renderer/Private/VirtualShadowMaps/VirtualShadowMapProjection.cpp
@@ -257,8 +257,11 @@ class FVirtualShadowMapProjectionCS : public FGlobalShader
 			OutEnvironment.CompilerFlags.Add(CFLAG_WaveOperations);
 		}
 
-		OutEnvironment.CompilerFlags.Add( CFLAG_Wave32 );
-		OutEnvironment.CompilerFlags.Add( CFLAG_AllowRealTypes );
+		OutEnvironment.CompilerFlags.Add(CFLAG_Wave32);
+		if (FDataDrivenShaderPlatformInfo::GetSupportsRealTypes(Parameters.Platform) == ERHIFeatureSupport::RuntimeGuaranteed)
+		{
+			OutEnvironment.CompilerFlags.Add(CFLAG_AllowRealTypes);
+		}
 	}
 
 	static bool ShouldCompilePermutation(const FGlobalShaderPermutationParameters& Parameters)
