Support multiple raytracing scene layers in RayTracingInstanceBufferUtil
- LayerIndex can be specified in FRayTracingGeometryInstance.
- Apply offset to FRayTracingInstanceDescriptorInput.OutputDescriptorIndex so that the instance buffer is divided into different one subrange per layer.

diff --git a/Engine/Plugins/Experimental/GPULightmass/Source/GPULightmass/Private/LightmapRenderer.cpp b/Engine/Plugins/Experimental/GPULightmass/Source/GPULightmass/Private/LightmapRenderer.cpp
--- a/Engine/Plugins/Experimental/GPULightmass/Source/GPULightmass/Private/LightmapRenderer.cpp
+++ b/Engine/Plugins/Experimental/GPULightmass/Source/GPULightmass/Private/LightmapRenderer.cpp
@@ -1063,6 +1063,7 @@ bool FSceneRenderState::SetupRayTracingScene(int32 LODIndex)
 
 			FRayTracingSceneWithGeometryInstances SceneWithGeometryInstances = CreateRayTracingSceneWithGeometryInstances(
 				RayTracingGeometryInstances,
+				1,
 				RAY_TRACING_NUM_SHADER_SLOTS,
 				1);
 
diff --git a/Engine/Source/Runtime/RHI/Public/RHIResources.h b/Engine/Source/Runtime/RHI/Public/RHIResources.h
--- a/Engine/Source/Runtime/RHI/Public/RHIResources.h
+++ b/Engine/Source/Runtime/RHI/Public/RHIResources.h
@@ -2556,6 +2556,8 @@ struct FRayTracingGeometryInstance
 	// If binary AND of instance mask with ray mask is zero, then the instance is considered not intersected / invisible.
 	uint8 Mask = 0xFF;
 
+	uint8 LayerIndex = 0;
+
 	// Flags to control triangle back face culling, whether to allow any-hit shaders, etc.
 	ERayTracingInstanceFlags Flags = ERayTracingInstanceFlags::None;
 };
diff --git a/Engine/Source/Runtime/Renderer/Private/RayTracing/RayTracingInstanceBufferUtil.cpp b/Engine/Source/Runtime/Renderer/Private/RayTracing/RayTracingInstanceBufferUtil.cpp
--- a/Engine/Source/Runtime/Renderer/Private/RayTracing/RayTracingInstanceBufferUtil.cpp
+++ b/Engine/Source/Runtime/Renderer/Private/RayTracing/RayTracingInstanceBufferUtil.cpp
@@ -21,6 +21,7 @@
 
 FRayTracingSceneWithGeometryInstances CreateRayTracingSceneWithGeometryInstances(
 	TArrayView<FRayTracingGeometryInstance> Instances,
+	uint8 NumLayers,
 	uint32 NumShaderSlotsPerGeometrySegment,
 	uint32 NumMissShaderSlots)
 {
@@ -40,7 +41,7 @@ FRayTracingSceneWithGeometryInstances CreateRayTracingSceneWithGeometryInstances
 	Initializer.PerInstanceGeometries.SetNumUninitialized(NumSceneInstances);
 	Initializer.BaseInstancePrefixSum.SetNumUninitialized(NumSceneInstances);
 	Initializer.SegmentPrefixSum.SetNumUninitialized(NumSceneInstances);
-	Initializer.NumNativeInstancesPerLayer.SetNumZeroed(1);
+	Initializer.NumNativeInstancesPerLayer.SetNumZeroed(NumLayers);
 	Initializer.NumTotalSegments = 0;
 
 	Experimental::TSherwoodMap<FRHIRayTracingGeometry*, uint32> UniqueGeometries;
@@ -105,8 +106,13 @@ FRayTracingSceneWithGeometryInstances CreateRayTracingSceneWithGeometryInstances
 			Output.NumNativeGPUInstances += InstanceDesc.NumTransforms;
 		}
 
-		Initializer.BaseInstancePrefixSum[InstanceIndex] = Initializer.NumNativeInstancesPerLayer[0];
-		Initializer.NumNativeInstancesPerLayer[0] += InstanceDesc.NumTransforms;
+		checkf(InstanceDesc.LayerIndex < NumLayers, 
+			TEXT("FRayTracingGeometryInstance is assigned to layer %d but raytracing scene being created only has %d layers."),
+			InstanceDesc.LayerIndex, NumLayers);
+
+		// Can't support same instance in multiple layers because BaseInstancePrefixSum would be different per layer
+		Initializer.BaseInstancePrefixSum[InstanceIndex] = Initializer.NumNativeInstancesPerLayer[InstanceDesc.LayerIndex];
+		Initializer.NumNativeInstancesPerLayer[InstanceDesc.LayerIndex] += InstanceDesc.NumTransforms;
 	}
 
 	Output.Scene = RHICreateRayTracingScene(MoveTemp(Initializer));
@@ -114,6 +120,14 @@ FRayTracingSceneWithGeometryInstances CreateRayTracingSceneWithGeometryInstances
 	return MoveTemp(Output);
 }
 
+FRayTracingSceneWithGeometryInstances CreateRayTracingSceneWithGeometryInstances(
+	TArrayView<FRayTracingGeometryInstance> Instances,
+	uint32 NumShaderSlotsPerGeometrySegment,
+	uint32 NumMissShaderSlots)
+{
+	return CreateRayTracingSceneWithGeometryInstances(Instances, 1, NumShaderSlotsPerGeometrySegment, NumMissShaderSlots);
+}
+
 void FillRayTracingInstanceUploadBuffer(
 	FRayTracingSceneRHIRef RayTracingSceneRHI,
 	FVector PreViewTranslation,
@@ -126,6 +140,16 @@ void FillRayTracingInstanceUploadBuffer(
 	TArrayView<FVector4f> OutTransformData)
 {
 	const FRayTracingSceneInitializer2& SceneInitializer = RayTracingSceneRHI->GetInitializer();
+	const uint32 NumLayers = SceneInitializer.NumNativeInstancesPerLayer.Num();
+
+	TArray<uint32> LayerBaseIndices;
+	LayerBaseIndices.SetNumUninitialized(NumLayers);
+	LayerBaseIndices[0] = 0;
+
+	for (uint32 LayerIndex = 1; LayerIndex < NumLayers; ++LayerIndex)
+	{
+		LayerBaseIndices[LayerIndex] = LayerBaseIndices[LayerIndex - 1] + SceneInitializer.NumNativeInstancesPerLayer[LayerIndex - 1];
+	}
 
 	int32 NumInactiveNativeInstances = 0;
 
@@ -139,6 +163,7 @@ void FillRayTracingInstanceUploadBuffer(
 			Instances,
 			InstanceGeometryIndices,
 			BaseUploadBufferOffsets,
+			LayerBaseIndices,
 			PreViewTranslation,
 			&SceneInitializer,
 			&NumInactiveNativeInstances
@@ -162,6 +187,7 @@ void FillRayTracingInstanceUploadBuffer(
 			checkf(bGpuSceneInstance + bGpuInstance + bCpuInstance == 1, TEXT("Instance can only get transforms from one of GPUScene, GPUTransformsSRV, or Transforms array."));
 
 			const uint32 AccelerationStructureIndex = InstanceGeometryIndices[SceneInstanceIndex];
+			const uint32 LayerBaseIndex = LayerBaseIndices[SceneInstance.LayerIndex];
 			const uint32 BaseInstanceIndex = SceneInitializer.BaseInstancePrefixSum[SceneInstanceIndex];
 			const uint32 BaseTransformIndex = bCpuInstance ? BaseUploadBufferOffsets[SceneInstanceIndex] : 0;
 
@@ -186,7 +212,7 @@ void FillRayTracingInstanceUploadBuffer(
 				InstanceDesc.InstanceMaskAndFlags = SceneInstance.Mask | ((uint32)SceneInstance.Flags << 8);
 				InstanceDesc.InstanceContributionToHitGroupIndex = SceneInitializer.SegmentPrefixSum[SceneInstanceIndex] * SceneInitializer.ShaderSlotsPerGeometrySegment;
 				InstanceDesc.InstanceId = bUseUniqueUserData ? SceneInstance.UserData[TransformIndex] : SceneInstance.DefaultUserData;
-				InstanceDesc.OutputDescriptorIndex = BaseInstanceIndex + TransformIndex;
+				InstanceDesc.OutputDescriptorIndex = LayerBaseIndex + BaseInstanceIndex + TransformIndex;
 
 				if (bGpuSceneInstance)
 				{
@@ -231,7 +257,15 @@ void FillRayTracingInstanceUploadBuffer(
 #endif
 		});
 
-	SET_DWORD_STAT(STAT_RayTracingInstances, SceneInitializer.NumNativeInstancesPerLayer[0] - NumInactiveNativeInstances);
+#if STATS
+	uint32 TotalNumNativeInstances = 0;
+	for (uint32 NumNativeInstances : SceneInitializer.NumNativeInstancesPerLayer)
+	{
+		TotalNumNativeInstances += NumNativeInstances;
+	}
+
+	SET_DWORD_STAT(STAT_RayTracingInstances, TotalNumNativeInstances - NumInactiveNativeInstances);
+#endif
 }
 
 struct FRayTracingBuildInstanceBufferCS : public FGlobalShader
diff --git a/Engine/Source/Runtime/Renderer/Private/RayTracing/RayTracingScene.cpp b/Engine/Source/Runtime/Renderer/Private/RayTracing/RayTracingScene.cpp
--- a/Engine/Source/Runtime/Renderer/Private/RayTracing/RayTracingScene.cpp
+++ b/Engine/Source/Runtime/Renderer/Private/RayTracing/RayTracingScene.cpp
@@ -32,6 +32,7 @@ void FRayTracingScene::Create(FRDGBuilder& GraphBuilder, const FGPUScene& GPUSce
 
 	FRayTracingSceneWithGeometryInstances SceneWithGeometryInstances = CreateRayTracingSceneWithGeometryInstances(
 		Instances,
+		1,
 		RAY_TRACING_NUM_SHADER_SLOTS,
 		RAY_TRACING_NUM_MISS_SHADER_SLOTS);
 
diff --git a/Engine/Source/Runtime/Renderer/Private/Tests/RayTracingTestbed.cpp b/Engine/Source/Runtime/Renderer/Private/Tests/RayTracingTestbed.cpp
--- a/Engine/Source/Runtime/Renderer/Private/Tests/RayTracingTestbed.cpp
+++ b/Engine/Source/Runtime/Renderer/Private/Tests/RayTracingTestbed.cpp
@@ -118,6 +118,7 @@ bool RunRayTracingTestbed_RenderThread(const FString& Parameters)
 
 	FRayTracingSceneWithGeometryInstances RayTracingScene = CreateRayTracingSceneWithGeometryInstances(
 		Instances,
+		1,
 		RAY_TRACING_NUM_SHADER_SLOTS,
 		1);
 
diff --git a/Engine/Source/Runtime/Renderer/Public/RayTracingInstanceBufferUtil.h b/Engine/Source/Runtime/Renderer/Public/RayTracingInstanceBufferUtil.h
--- a/Engine/Source/Runtime/Renderer/Public/RayTracingInstanceBufferUtil.h
+++ b/Engine/Source/Runtime/Renderer/Public/RayTracingInstanceBufferUtil.h
@@ -20,6 +20,8 @@ class FGPUScene;
 * This way native GPUScene and CPU instance descriptors can be built in a single dispatch per type.
 * Followed by one dispatch per GPU instance (since GPU transforms of each GPU instance are stored in separate buffers).
 * 
+* If the raytracing scene contains multiple layers, the instance buffer is divided into multiple subranges as expected by the RHI.
+* 
 */
 
 struct FRayTracingInstanceDescriptorInput
@@ -55,6 +57,13 @@ struct FRayTracingSceneWithGeometryInstances
 
 // Helper function to create FRayTracingSceneRHI using array of high level instances
 // Also outputs data required to build the instance buffer
+RENDERER_API FRayTracingSceneWithGeometryInstances CreateRayTracingSceneWithGeometryInstances(
+	TArrayView<FRayTracingGeometryInstance> Instances,
+	uint8 NumLayers,
+	uint32 NumShaderSlotsPerGeometrySegment,
+	uint32 NumMissShaderSlots);
+
+UE_DEPRECATED(5.1, "Specify NumLayers instead.")
 RENDERER_API FRayTracingSceneWithGeometryInstances CreateRayTracingSceneWithGeometryInstances(
 	TArrayView<FRayTracingGeometryInstance> Instances,
 	uint32 NumShaderSlotsPerGeometrySegment,
