Fix crash when r.RayTracing.Shadows is enabled dynamically

Previously FScene contained an array of RT-shadowed lights which was updated when a light is added/removed.
However, this is incorrect as the light RT shadow state depends on CVars that may be changed dynamically.

The solution is to check if the scene currently contains any RT-shadowed lights during rendering. This is more expensive but robust.
The array of RT-shadowed lights is removed as it was not used beyond checking if it's empty.

diff --git a/Engine/Source/Runtime/Renderer/Private/DeferredShadingRenderer.cpp b/Engine/Source/Runtime/Renderer/Private/DeferredShadingRenderer.cpp
--- a/Engine/Source/Runtime/Renderer/Private/DeferredShadingRenderer.cpp
+++ b/Engine/Source/Runtime/Renderer/Private/DeferredShadingRenderer.cpp
@@ -1967,6 +1967,7 @@ void FDeferredShadingSceneRenderer::Render(FRDGBuilder& GraphBuilder)
 			Scene->RefreshRayTracingInstances();
 		}
 
+		Scene->UpdateRayTracedLights();
 	}
 #endif
 	for (int32 ViewIndex = 0; ViewIndex < Views.Num(); ViewIndex++)
@@ -3408,7 +3409,7 @@ bool AnyRayTracingPassEnabled(const FScene* Scene, const FViewInfo& View)
 		|| ShouldRenderRayTracingTranslucency(View)
 		|| ShouldRenderRayTracingSkyLight(Scene->SkyLight)
 		|| ShouldRenderRayTracingShadows()
-		|| Scene->RayTracedLights.Num() > 0
+		|| Scene->bHasRayTracedLights
 		|| ShouldRenderPluginRayTracingGlobalIllumination(View)
         || Lumen::AnyLumenHardwareRayTracingPassEnabled(Scene, View)
 		|| HasRayTracedOverlay(*View.Family);
diff --git a/Engine/Source/Runtime/Renderer/Private/PrimitiveSceneInfo.cpp b/Engine/Source/Runtime/Renderer/Private/PrimitiveSceneInfo.cpp
--- a/Engine/Source/Runtime/Renderer/Private/PrimitiveSceneInfo.cpp
+++ b/Engine/Source/Runtime/Renderer/Private/PrimitiveSceneInfo.cpp
@@ -693,6 +693,30 @@ void FScene::RefreshRayTracingInstances()
 	FPrimitiveSceneInfo::UpdateCachedRayTracingInstances(this, Primitives);
 }
 
+void FScene::UpdateRayTracedLights()
+{
+	// Whether a light can use ray traced shadows depends on CVars that may be changed at runtime.
+	// It is not enough to check the light shadow mode when light is added to the scene. This must be done during rendering.
+
+	bHasRayTracedLights = false;
+
+	if (!IsRayTracingEnabled())
+	{
+		return;
+	}
+
+	// We currently don't need a full list of RT lights, only whether there are any RT lights at all.
+	for (auto LightIt = Lights.CreateConstIterator(); LightIt; ++LightIt)
+	{
+		const FLightSceneInfoCompact& LightSceneInfoCompact = *LightIt;
+		if (ShouldRenderRayTracingShadowsForLight(LightSceneInfoCompact))
+		{
+			bHasRayTracedLights = true;
+			break;
+		}
+	}
+}
+
 void FPrimitiveSceneInfo::UpdateCachedRayTracingInstances(FScene* Scene, const TArrayView<FPrimitiveSceneInfo*>& SceneInfos)
 {
 	if (IsRayTracingEnabled())
diff --git a/Engine/Source/Runtime/Renderer/Private/RayTracing/RayTracingShadows.cpp b/Engine/Source/Runtime/Renderer/Private/RayTracing/RayTracingShadows.cpp
--- a/Engine/Source/Runtime/Renderer/Private/RayTracing/RayTracingShadows.cpp
+++ b/Engine/Source/Runtime/Renderer/Private/RayTracing/RayTracingShadows.cpp
@@ -188,10 +188,10 @@ float GetRaytracingMaxNormalBias()
 void FDeferredShadingSceneRenderer::PrepareRayTracingShadows(const FViewInfo& View, const FScene& Scene, TArray<FRHIRayTracingShader*>& OutRayGenShaders)
 {
 	// Ray tracing shadows shaders should be properly configured even if r.RayTracing.Shadows is 0 because lights can have raytracing shadows enabled independently of that CVar
-	// We have to check if ray tracing is enabled on any of the scene lights. The RayTracedLights array is populated using ShouldRenderRayTracingShadowsForLight() helper, 
+	// We have to check if ray tracing is enabled on any of the scene lights. The Scene.bHasRayTracedLights is computed using ShouldRenderRayTracingShadowsForLight() helper, 
 	// which handles various override conditions.
 
-	if (Scene.RayTracedLights.IsEmpty())
+	if (Scene.bHasRayTracedLights == false)
 	{
 		return;
 	}
diff --git a/Engine/Source/Runtime/Renderer/Private/RendererScene.cpp b/Engine/Source/Runtime/Renderer/Private/RendererScene.cpp
--- a/Engine/Source/Runtime/Renderer/Private/RendererScene.cpp
+++ b/Engine/Source/Runtime/Renderer/Private/RendererScene.cpp
@@ -1974,13 +1974,7 @@ void FScene::AddLightSceneInfo_RenderThread(FLightSceneInfo* LightSceneInfo)
 
 	ProcessAtmosphereLightAddition_RenderThread(LightSceneInfo);
 
-#if RHI_RAYTRACING
-	if (ShouldRenderRayTracingShadowsForLight(LightSceneInfo))
-	{
-		RayTracedLights.Add(LightSceneInfo);
-	}
 	InvalidatePathTracedOutput();
-#endif 
 
 	// Add the light to the scene.
 	LightSceneInfo->AddToScene();
@@ -3200,13 +3194,6 @@ void FScene::RemoveLightSceneInfo_RenderThread(FLightSceneInfo* LightSceneInfo)
 
 		ProcessAtmosphereLightRemoval_RenderThread(LightSceneInfo);
 
-
-#if RHI_RAYTRACING
-		if (ShouldRenderRayTracingShadowsForLight(LightSceneInfo))
-		{
-			RayTracedLights.Remove(LightSceneInfo);
-		}
-#endif 
 		// Remove the light from the scene.
 		LightSceneInfo->RemoveFromScene();
 
diff --git a/Engine/Source/Runtime/Renderer/Private/ScenePrivate.h b/Engine/Source/Runtime/Renderer/Private/ScenePrivate.h
--- a/Engine/Source/Runtime/Renderer/Private/ScenePrivate.h
+++ b/Engine/Source/Runtime/Renderer/Private/ScenePrivate.h
@@ -2887,7 +2887,9 @@ public:
 #if RHI_RAYTRACING
 	/** Persistently-allocated ray tracing scene data. */
 	FRayTracingScene RayTracingScene;
-	TArray<FLightSceneInfo*, TInlineAllocator<4>> RayTracedLights;
+
+	bool bHasRayTracedLights = false;
+	void UpdateRayTracedLights();
 #endif // RHI_RAYTRACING
 
 	/** Distance field object scene data. */
