Exclude primitives that are meant to be hidden in game from ray tracing scene when game view is enabled or in PIE mode.

diff --git a/Engine/Source/Runtime/Renderer/Private/DeferredShadingRenderer.cpp b/Engine/Source/Runtime/Renderer/Private/DeferredShadingRenderer.cpp
--- a/Engine/Source/Runtime/Renderer/Private/DeferredShadingRenderer.cpp
+++ b/Engine/Source/Runtime/Renderer/Private/DeferredShadingRenderer.cpp
@@ -647,6 +647,8 @@ bool FDeferredShadingSceneRenderer::GatherRayTracingWorldInstancesForView(FRDGBu
 	TArray<FPrimitiveSceneInfo*> DirtyCachedRayTracingPrimitives;
 	DirtyCachedRayTracingPrimitives.Reserve(Scene->PrimitiveSceneProxies.Num());
 
+	const bool bGameView = View.bIsGameView || View.Family->EngineShowFlags.Game;
+
 	int32 VisiblePrimitives = 0;
 	bool bPerformRayTracing = View.State != nullptr && !View.bIsReflectionCapture && View.bAllowRayTracing;
 	if (bPerformRayTracing)
@@ -692,6 +694,13 @@ bool FDeferredShadingSceneRenderer::GatherRayTracingWorldInstancesForView(FRDGBu
 				continue;
 			}
 
+			// Some primitives should only be visible editor mode, however far field geometry 
+			// and hidden shadow casters must still always be added to the RT scene.
+			if (bGameView && !SceneInfo->bDrawInGame && !SceneInfo->bRayTracingFarField)
+			{
+				continue;
+			}
+
 			// Marked visible and used after point, check if streaming then mark as used in the TLAS (so it can be streamed in)
 			if (EnumHasAnyFlags(Scene->PrimitiveRayTracingFlags[PrimitiveIndex], ERayTracingPrimitiveFlags::Streaming))
 			{
diff --git a/Engine/Source/Runtime/Renderer/Private/PrimitiveSceneInfo.cpp b/Engine/Source/Runtime/Renderer/Private/PrimitiveSceneInfo.cpp
--- a/Engine/Source/Runtime/Renderer/Private/PrimitiveSceneInfo.cpp
+++ b/Engine/Source/Runtime/Renderer/Private/PrimitiveSceneInfo.cpp
@@ -173,6 +173,7 @@ FPrimitiveSceneInfo::FPrimitiveSceneInfo(UPrimitiveComponent* InComponent,FScene
 	bVisibleInRealTimeSkyCapture(InComponent->SceneProxy->IsVisibleInRealTimeSkyCaptures()),
 #if RHI_RAYTRACING
 	bDrawInGame(Proxy->IsDrawnInGame()),
+	bRayTracingFarField(Proxy->IsRayTracingFarField()),
 	bIsVisibleInSceneCaptures(!InComponent->SceneProxy->IsHiddenInSceneCapture()),
 	bIsVisibleInSceneCapturesOnly(InComponent->SceneProxy->IsVisibleInSceneCaptureOnly()),
 	bIsRayTracingRelevant(InComponent->SceneProxy->IsRayTracingRelevant()),
diff --git a/Engine/Source/Runtime/Renderer/Public/PrimitiveSceneInfo.h b/Engine/Source/Runtime/Renderer/Public/PrimitiveSceneInfo.h
--- a/Engine/Source/Runtime/Renderer/Public/PrimitiveSceneInfo.h
+++ b/Engine/Source/Runtime/Renderer/Public/PrimitiveSceneInfo.h
@@ -389,6 +389,7 @@ public:
 
 #if RHI_RAYTRACING
 	bool bDrawInGame : 1;
+	bool bRayTracingFarField : 1;
 	bool bIsVisibleInSceneCaptures : 1;
 	bool bIsVisibleInSceneCapturesOnly : 1;
 	bool bIsRayTracingRelevant : 1;
