Strata - fixed SSAO final pass to work with strata.

diff --git a/Engine/Source/Runtime/Renderer/Private/CompositionLighting/CompositionLighting.cpp b/Engine/Source/Runtime/Renderer/Private/CompositionLighting/CompositionLighting.cpp
--- a/Engine/Source/Runtime/Renderer/Private/CompositionLighting/CompositionLighting.cpp
+++ b/Engine/Source/Runtime/Renderer/Private/CompositionLighting/CompositionLighting.cpp
@@ -471,12 +471,21 @@ static FScreenPassTexture AddPostProcessingAmbientOcclusion(
 				CommonParameters.HZBInput);
 	}
 
+	FScreenPassTexture SetupTexture = CommonParameters.GBufferA;
+	if (Strata::IsStrataEnabled())
+	{
+		// For Strata, we invalidate the setup texture for the final pass:
+		//	- We do not need GBufferA, the Strata TopLayer texture will fill in for that.
+		//	- Setting it to nullptr will make the AddAmbientOcclusionPass use a valid viewport from SceneTextures.
+		SetupTexture.Texture = nullptr;
+	}
+
 	FScreenPassTexture FinalOutput =
 		AddAmbientOcclusionFinalPass(
 			GraphBuilder,
 			View,
 			CommonParameters,
-			CommonParameters.GBufferA,
+			SetupTexture,
 			AmbientOcclusionInMip1,
 			AmbientOcclusionPassMip1,
 			CommonParameters.HZBInput,
