Strata - fixed SLW to not cause issue on when shader specifies output format wiht Strata.
The separated main dir light buffer is no longer bound as render target during the base pass. It is used for the DFShadow only.
The main dir light contribution flow thgouth the Strata mateiral buffer.

diff --git a/Engine/Shaders/Private/BasePassPixelShader.usf b/Engine/Shaders/Private/BasePassPixelShader.usf
--- a/Engine/Shaders/Private/BasePassPixelShader.usf
+++ b/Engine/Shaders/Private/BasePassPixelShader.usf
@@ -1705,6 +1705,9 @@ void FPixelShaderInOut_MainPS(
 	// Need to reset color to make sure strata material are only lit using Strata lighting passes.
 	// Except for water doing some specialised and simplified lighting during the base pass
 	Color = 0;
+#else
+	// Store the potential separated main dir light contribution into the water BSDF
+	SLW_SEPARATEDMAINDIRLIGHT(StrataPixelHeader.StrataTree.BSDFs[0]) = SeparatedWaterMainDirLightLuminance * View.PreExposure;
 #endif
 
 	// STRATA_TODO GBUFFER_HAS_DIFFUSE_SAMPLE_OCCLUSION and ALLOW_STATIC_LIGHTING cases
@@ -1984,20 +1987,6 @@ void FPixelShaderInOut_MainPS(
 				Out.MRT[GBUFFER_HAS_VELOCITY ? 6 : 5] = 0;
 			}
 
-		#if SINGLE_LAYER_WATER_SEPARATED_MAIN_LIGHT && STRATA_ENABLED
-			if (bSeparateWaterMainDirLightLuminance)
-			{
-				// STRATA_TODO: we should have a define set to the correct output index.
-				#if PIXELSHADEROUTPUT_MRT2
-					Out.MRT[2] = float4(SeparatedWaterMainDirLightLuminance * View.PreExposure, 0.0f);
-				#elif PIXELSHADEROUTPUT_MRT1
-					Out.MRT[1] = float4(SeparatedWaterMainDirLightLuminance * View.PreExposure, 0.0f);
-				#else
-					Out.MRT[0] = float4(SeparatedWaterMainDirLightLuminance * View.PreExposure, 0.0f);
-				#endif
-			}
-		#endif
-
 		#if SINGLE_LAYER_WATER_SEPARATED_MAIN_LIGHT && !STRATA_ENABLED
 			if (GBuffer.ShadingModelID == SHADINGMODELID_SINGLELAYERWATER)
 			{
diff --git a/Engine/Shaders/Private/SingleLayerWaterComposite.usf b/Engine/Shaders/Private/SingleLayerWaterComposite.usf
--- a/Engine/Shaders/Private/SingleLayerWaterComposite.usf
+++ b/Engine/Shaders/Private/SingleLayerWaterComposite.usf
@@ -54,6 +54,7 @@ SingleLayerWaterCompositeOutput SingleLayerWaterComposite(float2 BufferUV, float
 	Output.Clip = 1.0f; // no clipping
 
 	float3 TransmittanceToScene = 1.0f;
+	float3 SeparatedMainDirLightContribution = 0.0f;
 
 #if STRATA_ENABLED
 
@@ -79,6 +80,12 @@ SingleLayerWaterCompositeOutput SingleLayerWaterComposite(float2 BufferUV, float
 			const float Roughness = SLW_ROUGHNESS(SLWaterBSDF);
 			const float SafeRoughness = MakeRoughnessSafe(SLW_ROUGHNESS(SLWaterBSDF));
 
+			if (UseSeparatedMainDirLightTexture > 0.0f)
+			{
+				// Also compose the separated main direction light luminance (separated to be able to combine it with DistanceFieldShadow)
+				SeparatedMainDirLightContribution = SLW_SEPARATEDMAINDIRLIGHT(SLWaterBSDF) * SeparatedMainDirLightTexture[PixelPos].rgb;
+			}
+
 			// Compute the sky reflection contribution
 			float3 Reflection = 0.0f;
 
@@ -178,6 +185,12 @@ SingleLayerWaterCompositeOutput SingleLayerWaterComposite(float2 BufferUV, float
 			// BRDF/Fresnel is already applied for the under water part in the scene.
 			float3 EnvBrdfValue = EnvBRDF(GBuffer.SpecularColor, GBuffer.Roughness, NoV);
 
+			if (UseSeparatedMainDirLightTexture > 0.0f)
+			{
+				// Also compose the separated main direction light luminance (separated to be able to combine it with DistanceFieldShadow)
+				SeparatedMainDirLightContribution = SeparatedMainDirLightTexture[PixelPos].rgb;
+			}
+
 #endif // STRATA_ENABLED
 
 			Reflection *= EnvBrdfValue;
@@ -199,13 +212,7 @@ SingleLayerWaterCompositeOutput SingleLayerWaterComposite(float2 BufferUV, float
 				TransmittanceToScene *= WaterMediumTransmittance;
 			}
 
-			Output.LuminanceTransmittance = float4(Reflection, dot(OneThird3, TransmittanceToScene));
-
-			if (UseSeparatedMainDirLightTexture)
-			{
-				// Also compose the separated main direction light luminance (separated to be able to combine it with DistanceFieldShadow)
-				Output.LuminanceTransmittance.rgb += SeparatedMainDirLightTexture[PixelPos].rgb;
-			}
+			Output.LuminanceTransmittance = float4(Reflection + SeparatedMainDirLightContribution, dot(OneThird3, TransmittanceToScene));
 		}
 		else
 		{
diff --git a/Engine/Shaders/Private/Strata/Strata.ush b/Engine/Shaders/Private/Strata/Strata.ush
--- a/Engine/Shaders/Private/Strata/Strata.ush
+++ b/Engine/Shaders/Private/Strata/Strata.ush
@@ -431,6 +431,7 @@ struct FStrataBSDF
 #define SLW_SPECULAR(X)					X.VGPRs[1].y
 #define SLW_ROUGHNESS(X)				X.VGPRs[1].z
 #define SLW_TOPMATERIALOPACITY(X)		X.VGPRs[1].w
+#define SLW_SEPARATEDMAINDIRLIGHT(X)	X.VGPRs[2].xyz
 #if STRATA_INLINE_SINGLELAYERWATER
 #define SLW_WATERALBEDO(X)				X.InlineVGPRs[0].xyz
 #define SLW_WATEREXTINCTION(X)			X.InlineVGPRs[1].xyz
@@ -1241,6 +1242,7 @@ FStrataData GetStrataSingleLayerWaterBSDF(
 	SLW_SPECULAR					(StrataData.InlinedBSDF) = Specular;
 	SLW_ROUGHNESS					(StrataData.InlinedBSDF) = Roughness;
 	SLW_TOPMATERIALOPACITY			(StrataData.InlinedBSDF) = TopMaterialOpacity;
+	SLW_SEPARATEDMAINDIRLIGHT		(StrataData.InlinedBSDF) = float3(0.0f, 0.0f, 0.0f);
 #if STRATA_INLINE_SINGLELAYERWATER
 	SLW_WATERALBEDO					(StrataData.InlinedBSDF) = WaterAlbedo;
 	SLW_WATEREXTINCTION				(StrataData.InlinedBSDF) = WaterExtinction;
@@ -3344,6 +3346,7 @@ void PackStrataOut(
 					{
 						STRATA_STORE_UINT1(PackColorLinearToGamma2AlphaLinear(float4(SLW_BASECOLOR(BSDF), 0.0f)));
 						STRATA_STORE_UINT1(PackRGBA8(float4(SLW_METALLIC(BSDF), SLW_SPECULAR(BSDF), SLW_ROUGHNESS(BSDF), SLW_TOPMATERIALOPACITY(BSDF))));
+						STRATA_STORE_UINT1(PackR11G11B10F(SLW_SEPARATEDMAINDIRLIGHT(BSDF)));	// STRATA_TODO make this an optional load/store base on if DistanceFieldShadow are enabled
 					}
 					break;
 					}
@@ -3644,7 +3647,8 @@ FStrataBSDF UnpackStrataBSDFIn(FStrataMaterialContainer StrataBuffer, inout FStr
 		SLW_SPECULAR(OutBSDF) = Data.y;
 		SLW_ROUGHNESS(OutBSDF) = Data.z;
 		SLW_TOPMATERIALOPACITY(OutBSDF) = Data.w;
-		// 8 bytes
+		StrataLoad_R11G11B10F(StrataBuffer, StrataAddressing, SLW_SEPARATEDMAINDIRLIGHT(OutBSDF));
+		// 12 bytes
 	}
 	break;
 	}
diff --git a/Engine/Source/Runtime/Engine/Private/Materials/HLSLMaterialTranslator.cpp b/Engine/Source/Runtime/Engine/Private/Materials/HLSLMaterialTranslator.cpp
--- a/Engine/Source/Runtime/Engine/Private/Materials/HLSLMaterialTranslator.cpp
+++ b/Engine/Source/Runtime/Engine/Private/Materials/HLSLMaterialTranslator.cpp
@@ -10115,6 +10115,7 @@ bool FHLSLMaterialTranslator::StrataGenerateDerivedMaterialOperatorData()
 				}
 				case STRATA_BSDF_TYPE_SINGLELAYERWATER:
 				{
+					StrataMaterialRequestedSizeByte += UintByteSize;
 					StrataMaterialRequestedSizeByte += UintByteSize;
 					StrataMaterialRequestedSizeByte += UintByteSize;
 					break;
diff --git a/Engine/Source/Runtime/Engine/Private/ShaderCompiler/ShaderGenerationUtil.cpp b/Engine/Source/Runtime/Engine/Private/ShaderCompiler/ShaderGenerationUtil.cpp
--- a/Engine/Source/Runtime/Engine/Private/ShaderCompiler/ShaderGenerationUtil.cpp
+++ b/Engine/Source/Runtime/Engine/Private/ShaderCompiler/ShaderGenerationUtil.cpp
@@ -1786,7 +1786,7 @@ static void DetermineUsedMaterialSlots(
 		SetStandardGBufferSlots(Slots, bWriteEmissive, bHasTangent, bHasVelocity, bHasStaticLighting, bIsStrataMaterial);
 		if (Mat.SINGLE_LAYER_WATER_DF_SHADOW_ENABLED)
 		{
-			Slots[GBS_SeparatedMainDirLight] = true;
+			Slots[GBS_SeparatedMainDirLight] = !bIsStrataMaterial; // This is stored in the Strata material buffer, not in a separate RT.
 		}
 	}
 
diff --git a/Engine/Source/Runtime/RenderCore/Private/GBufferInfo.cpp b/Engine/Source/Runtime/RenderCore/Private/GBufferInfo.cpp
--- a/Engine/Source/Runtime/RenderCore/Private/GBufferInfo.cpp
+++ b/Engine/Source/Runtime/RenderCore/Private/GBufferInfo.cpp
@@ -255,10 +255,6 @@ FGBufferInfo RENDERCORE_API FetchLegacyGBufferInfo(const FGBufferParams& Params)
 		{
 			TargetGBufferE = Info.NumTargets++;
 		}
-		if (Params.bHasSingleLayerWaterSeparatedMainLight)
-		{
-			TargetSeparatedMainDirLight = Info.NumTargets++;
-		}
 
 		// this value isn't correct, becuase it doesn't resepect the scene color format cvar, but it's ignored anyways
 		// so it's ok for now
@@ -292,14 +288,7 @@ FGBufferInfo RENDERCORE_API FetchLegacyGBufferInfo(const FGBufferParams& Params)
 			Info.Slots[GBS_PrecomputedShadowFactor].Packing[3] = FGBufferPacking(TargetGBufferE, 3, 3);
 		}
 
-		// Special water output
-		if (Params.bHasSingleLayerWaterSeparatedMainLight)
-		{
-			Info.Slots[GBS_SeparatedMainDirLight] = FGBufferItem(GBS_SeparatedMainDirLight, GBC_Raw_Float_11_11_10, GBCH_Both);
-			Info.Slots[GBS_SeparatedMainDirLight].Packing[0] = FGBufferPacking(TargetSeparatedMainDirLight, 0, 0);
-			Info.Slots[GBS_SeparatedMainDirLight].Packing[1] = FGBufferPacking(TargetSeparatedMainDirLight, 1, 1);
-			Info.Slots[GBS_SeparatedMainDirLight].Packing[2] = FGBufferPacking(TargetSeparatedMainDirLight, 2, 2);
-		}
+		// We do not setup GBS_SeparatedMainDirLight as the data will flow through the Strata material buffer.
 
 		return Info;
 	}
diff --git a/Engine/Source/Runtime/Renderer/Private/SingleLayerWaterRendering.cpp b/Engine/Source/Runtime/Renderer/Private/SingleLayerWaterRendering.cpp
--- a/Engine/Source/Runtime/Renderer/Private/SingleLayerWaterRendering.cpp
+++ b/Engine/Source/Runtime/Renderer/Private/SingleLayerWaterRendering.cpp
@@ -280,7 +280,8 @@ static FSceneWithoutWaterTextures AddCopySceneWithoutWaterPass(
 	check(SceneColorTexture);
 	check(SceneDepthTexture);
 
-	const bool bCopyColor = !SingleLayerWaterUsesSimpleShading(Views[0].GetShaderPlatform());
+	EShaderPlatform ShaderPlatform = Views[0].GetShaderPlatform();
+	const bool bCopyColor = !SingleLayerWaterUsesSimpleShading(ShaderPlatform);
 
 	const FRDGTextureDesc& SceneColorDesc = SceneColorTexture->Desc;
 	const FRDGTextureDesc& SceneDepthDesc = SceneColorTexture->Desc;
@@ -298,8 +299,14 @@ static FSceneWithoutWaterTextures AddCopySceneWithoutWaterPass(
 	const FRDGTextureDesc DepthDesc(FRDGTextureDesc::Create2D(RefractionResolution, ViewFamily.EngineShowFlags.SingleLayerWaterRefractionFullPrecision ? PF_R32_FLOAT : PF_R16F, SceneDepthDesc.ClearValue, TexCreate_ShaderResource | TexCreate_RenderTargetable));
 	FRDGTextureRef SceneDepthWithoutSingleLayerWaterTexture = GraphBuilder.CreateTexture(DepthDesc, TEXT("SLW.SceneDepthWithout"));
 
-	const FRDGTextureDesc SeparatedMainDirLightDesc(FRDGTextureDesc::Create2D(SceneColorDesc.Extent, PF_FloatR11G11B10, SceneDepthDesc.ClearValue, TexCreate_ShaderResource | TexCreate_RenderTargetable));
+	const FRDGTextureDesc SeparatedMainDirLightDesc(FRDGTextureDesc::Create2D(SceneColorDesc.Extent, PF_FloatR11G11B10, FClearValueBinding(FLinearColor::White), TexCreate_ShaderResource | TexCreate_RenderTargetable));
 	FRDGTextureRef SeparatedMainDirLightTexture = GraphBuilder.CreateTexture(SeparatedMainDirLightDesc, TEXT("SLW.SeparatedMainDirLight"));
+	if (IsWaterDistanceFieldShadowEnabled_Runtime(ShaderPlatform) && Strata::IsStrataEnabled())
+	{
+		// This clear is needed with strata because that texture will be modulated by DFShadows.
+		// STRATA_TODO: when strata is enabled, we can change RenderRayTracedDistanceFieldProjection to have a bForceNoBlending instead if bForceRGBModulation and remove that clear.
+		AddClearRenderTargetPass(GraphBuilder, SeparatedMainDirLightTexture);
+	}
 
 	FSceneWithoutWaterTextures Textures;
 	Textures.RefractionDownsampleFactor = float(RefractionDownsampleFactor);
@@ -814,7 +821,8 @@ void FDeferredShadingSceneRenderer::RenderSingleLayerWaterInner(
 
 	TStaticArray<FTextureRenderTargetBinding, MaxSimultaneousRenderTargets> BasePassTextures;
 	uint32 BasePassTextureCount = SceneTextures.GetGBufferRenderTargets(BasePassTextures);
-	if(IsWaterDistanceFieldShadowEnabled_Runtime(Scene->GetShaderPlatform()))
+	if(IsWaterDistanceFieldShadowEnabled_Runtime(Scene->GetShaderPlatform()) 
+		&& !Strata::IsStrataEnabled())	// We do not bind that texture if Strata is enabled as the data will go through the Strata material buffer.
 	{
 		const bool bNeverClear = true;
 		BasePassTextures[BasePassTextureCount++] = FTextureRenderTargetBinding(SceneWithoutWaterTextures.SeparatedMainDirLightTexture, bNeverClear);
