Remove BSDF offsets/tile overflow computation when Lumen or Strata Vis are disabled.

diff --git a/Engine/Source/Runtime/Renderer/Private/Strata/Strata.cpp b/Engine/Source/Runtime/Renderer/Private/Strata/Strata.cpp
--- a/Engine/Source/Runtime/Renderer/Private/Strata/Strata.cpp
+++ b/Engine/Source/Runtime/Renderer/Private/Strata/Strata.cpp
@@ -10,6 +10,7 @@
 #include "UniformBuffer.h"
 #include "SceneTextureParameters.h"
 #include "ShaderCompiler.h"
+#include "Lumen/Lumen.h"
 
 //PRAGMA_DISABLE_OPTIMIZATION
 
@@ -148,7 +149,7 @@ FIntPoint GetStrataTextureResolution(const FIntPoint& InResolution)
 
 static void BindStrataGlobalUniformParameters(FRDGBuilder& GraphBuilder, FStrataViewData* StrataViewData, FStrataGlobalUniformParameters& OutStrataUniformParameters);
 
-static void InitialiseStrataViewData(FRDGBuilder& GraphBuilder, FViewInfo& View, const FSceneTexturesConfig& SceneTexturesConfig, FStrataSceneData& SceneData)
+static void InitialiseStrataViewData(FRDGBuilder& GraphBuilder, FViewInfo& View, const FSceneTexturesConfig& SceneTexturesConfig, bool bNeedBSDFOffets, FStrataSceneData& SceneData)
 {
 	// Sanity check: the scene data should already exist 
 	check(SceneData.MaterialTextureArray != nullptr);
@@ -206,8 +207,8 @@ static void InitialiseStrataViewData(FRDGBuilder& GraphBuilder, FViewInfo& View,
 		}
 
 		// BSDF tiles
-		{
-			
+		if (bNeedBSDFOffets)
+		{			
 			const FIntPoint BufferSize = SceneTexturesConfig.Extent;
 			const FIntPoint BufferSize_Extended = GetStrataTextureTileResolution(SceneTexturesConfig.Extent, EStrataTileSpace::StrataTileSpace_Primary | EStrataTileSpace::StrataTileSpace_Overflow);
 
@@ -226,6 +227,17 @@ static void InitialiseStrataViewData(FRDGBuilder& GraphBuilder, FViewInfo& View,
 			Out.BSDFTileDispatchIndirectBuffer = GraphBuilder.CreateBuffer(FRDGBufferDesc::CreateIndirectDesc<FRHIDispatchIndirectParameters>(1), TEXT("Strata.StrataBSDFTileDispatchIndirectBuffer"));
 			Out.BSDFTileCountBuffer = GraphBuilder.CreateBuffer(FRDGBufferDesc::CreateBufferDesc(4, 1), TEXT("Strata.BSDFTileCount"));
 		}
+		else
+		{
+			Out.TileCount = GetStrataTextureTileResolution(ViewResolution, EStrataTileSpace::StrataTileSpace_Primary);
+			Out.TileOffset = FIntPoint(0,0);
+			Out.OverflowTileCount = FIntPoint(0, 0);
+			Out.OverflowTileOffset = FIntPoint(0, 0);
+			Out.BSDFTileTexture = nullptr;
+			Out.BSDFTilePerThreadDispatchIndirectBuffer = nullptr;
+			Out.BSDFTileDispatchIndirectBuffer = nullptr;
+			Out.BSDFTileCountBuffer = nullptr;
+		}
 	}
 
 	// Create the readable uniform buffers
@@ -237,6 +249,11 @@ static void InitialiseStrataViewData(FRDGBuilder& GraphBuilder, FViewInfo& View,
 	}
 }
 
+static bool NeedBSDFOffsets(const FScene* Scene, const FViewInfo& View)
+{
+	return  ShouldRenderLumenDiffuseGI(Scene, View) || ShouldRenderLumenReflections(View) || Strata::ShouldRenderStrataDebugPasses(View);
+}
+
 void InitialiseStrataFrameSceneData(FRDGBuilder& GraphBuilder, FSceneRenderer& SceneRenderer)
 {
 	FStrataSceneData& Out = SceneRenderer.Scene->StrataSceneData;
@@ -249,6 +266,12 @@ void InitialiseStrataFrameSceneData(FRDGBuilder& GraphBuilder, FSceneRenderer& S
 		OutMaterialBufferSizeXY.Y = FMath::DivideAndRoundUp(InBufferSizeXY.Y, STRATA_TILE_SIZE) * STRATA_TILE_SIZE;
 	};
 
+	bool bNeedBSDFOffsets = false;
+	for (const FViewInfo& View : SceneRenderer.Views)
+	{
+		bNeedBSDFOffsets = bNeedBSDFOffsets || NeedBSDFOffsets(SceneRenderer.Scene, View);
+	}
+
 	FIntPoint MaterialBufferSizeXY;
 	UpdateMaterialBufferToTiledResolution(FIntPoint(1, 1), MaterialBufferSizeXY);
 	if (IsStrataEnabled())
@@ -297,6 +320,7 @@ void InitialiseStrataFrameSceneData(FRDGBuilder& GraphBuilder, FSceneRenderer& S
 		}
 
 		// BSDF offsets
+		if (bNeedBSDFOffsets)
 		{
 			Out.BSDFOffsetTexture = GraphBuilder.CreateTexture(FRDGTextureDesc::Create2D(SceneTextureExtent, PF_R32_UINT, FClearValueBinding::None, TexCreate_UAV | TexCreate_ShaderResource), TEXT("Strata.BSDFOffsets"));
 			AddClearUAVPass(GraphBuilder, GraphBuilder.CreateUAV(Out.BSDFOffsetTexture), 0u);
@@ -340,7 +364,7 @@ void InitialiseStrataFrameSceneData(FRDGBuilder& GraphBuilder, FSceneRenderer& S
 	// Initialized view data
 	for (int32 ViewIndex = 0; ViewIndex < SceneRenderer.Views.Num(); ViewIndex++)
 	{
-		Strata::InitialiseStrataViewData(GraphBuilder, SceneRenderer.Views[ViewIndex], SceneRenderer.GetActiveSceneTexturesConfig(), Out);
+		Strata::InitialiseStrataViewData(GraphBuilder, SceneRenderer.Views[ViewIndex], SceneRenderer.GetActiveSceneTexturesConfig(), bNeedBSDFOffsets, Out);
 	}
 }
 
@@ -400,7 +424,16 @@ static void BindStrataGlobalUniformParameters(FRDGBuilder& GraphBuilder, FStrata
 		OutStrataUniformParameters.OpaqueRoughRefractionTexture = StrataSceneData->OpaqueRoughRefractionTexture;
 		OutStrataUniformParameters.BSDFTileTexture = StrataViewData->BSDFTileTexture;
 		OutStrataUniformParameters.BSDFOffsetTexture = StrataSceneData->BSDFOffsetTexture;
-		OutStrataUniformParameters.BSDFTileCountBuffer = GraphBuilder.CreateSRV(StrataViewData->BSDFTileCountBuffer, PF_R32_UINT);
+		OutStrataUniformParameters.BSDFTileCountBuffer = StrataViewData->BSDFTileCountBuffer ? GraphBuilder.CreateSRV(StrataViewData->BSDFTileCountBuffer, PF_R32_UINT) : nullptr;
+
+		if (OutStrataUniformParameters.BSDFOffsetTexture == nullptr)
+		{
+			const FRDGSystemTextures& SystemTextures = FRDGSystemTextures::Get(GraphBuilder);
+			FRDGBufferSRVRef DefaultBuffer = GraphBuilder.CreateSRV(GSystemTextures.GetDefaultBuffer(GraphBuilder, 4, 0u), PF_R32_UINT);
+			OutStrataUniformParameters.BSDFOffsetTexture = SystemTextures.Black;
+			OutStrataUniformParameters.BSDFTileTexture = SystemTextures.Black;
+			OutStrataUniformParameters.BSDFTileCountBuffer = DefaultBuffer;
+		}
 	}
 	else
 	{
@@ -1013,6 +1046,7 @@ void AddStrataMaterialClassificationPass(FRDGBuilder& GraphBuilder, const FMinim
 		}
 
 		// Compute BSDF tile index and material read offset
+		if (StrataSceneData->BSDFOffsetTexture)
 		{
 			FRDGBufferUAVRef RWBSDFTileCountBuffer = GraphBuilder.CreateUAV(StrataViewData->BSDFTileCountBuffer, PF_R32_UINT);
 			AddClearUAVPass(GraphBuilder, RWBSDFTileCountBuffer, 0u);
@@ -1048,6 +1082,7 @@ void AddStrataMaterialClassificationPass(FRDGBuilder& GraphBuilder, const FMinim
 		}
 
 		// Tile indirect dispatch args conversion
+		if (StrataSceneData->BSDFOffsetTexture)
 		{
 			TShaderMapRef<FStrataBSDFTilePrepareArgsPassCS> ComputeShader(View.ShaderMap);
 			FStrataBSDFTilePrepareArgsPassCS::FParameters* PassParameters = GraphBuilder.AllocParameters<FStrataBSDFTilePrepareArgsPassCS::FParameters>();
