Don't apply far field offset to CachedRayTracingInstanceWorldBounds.
- raytracing instance culling doesn't apply far field offset to view origin (FRayTracingCullPrimitiveInstancesClosure::operator()).
- also don't need to refresh cached raytracing instances when far field offset changes anymore.

diff --git a/Engine/Source/Runtime/Renderer/Private/DeferredShadingRenderer.cpp b/Engine/Source/Runtime/Renderer/Private/DeferredShadingRenderer.cpp
--- a/Engine/Source/Runtime/Renderer/Private/DeferredShadingRenderer.cpp
+++ b/Engine/Source/Runtime/Renderer/Private/DeferredShadingRenderer.cpp
@@ -1955,14 +1955,6 @@ void FDeferredShadingSceneRenderer::Render(FRDGBuilder& GraphBuilder)
 			Scene->RefreshRayTracingMeshCommandCache();
 		}
 
-		static FVector FarFieldReferencePosLast = Lumen::GetFarFieldReferencePos();
-		FVector FarFieldReferencePos = Lumen::GetFarFieldReferencePos();
-		if (FarFieldReferencePosLast.Z != FarFieldReferencePos.Z)
-		{
-			FarFieldReferencePosLast = FarFieldReferencePos;
-			Scene->RefreshRayTracingInstances();
-		}
-
 		Scene->UpdateRayTracedLights();
 	}
 #endif
diff --git a/Engine/Source/Runtime/Renderer/Private/PrimitiveSceneInfo.cpp b/Engine/Source/Runtime/Renderer/Private/PrimitiveSceneInfo.cpp
--- a/Engine/Source/Runtime/Renderer/Private/PrimitiveSceneInfo.cpp
+++ b/Engine/Source/Runtime/Renderer/Private/PrimitiveSceneInfo.cpp
@@ -952,17 +952,10 @@ void FPrimitiveSceneInfo::UpdateCachedRayTracingInstance(FPrimitiveSceneInfo* Sc
 		SceneInfo->CachedRayTracingInstance.NumTransforms = CachedRayTracingInstance.NumTransforms;
 		SceneInfo->CachedRayTracingInstance.BaseInstanceSceneDataOffset = SceneInfo->GetInstanceSceneDataOffset();
 
-		// Apply local offset to far-field object
-		FMatrix LocalToWorld = SceneInfo->Proxy->GetLocalToWorld();
-		if (SceneInfo->Proxy->IsRayTracingFarField())
-		{
-			LocalToWorld = LocalToWorld.ConcatTranslation(Lumen::GetFarFieldReferencePos());
-		}
-
 		SceneInfo->CachedRayTracingInstanceWorldBounds.Empty();
 		SceneInfo->CachedRayTracingInstanceWorldBounds.AddUninitialized(CachedRayTracingInstance.NumTransforms);
 
-		SceneInfo->UpdateCachedRayTracingInstanceWorldBounds(LocalToWorld);
+		SceneInfo->UpdateCachedRayTracingInstanceWorldBounds(SceneInfo->Proxy->GetLocalToWorld());
 
 		SceneInfo->CachedRayTracingInstance.GeometryRHI = CachedRayTracingInstance.Geometry->RayTracingGeometryRHI;
 
