Multi-view multi-gpu, scene rendering perf improvements for Virtual Production.
* Fence wait for cross GPU resource transfers can now be deferred.  Moved the wait for view family render target transfers to the end of all view family rendering.  Saves a few ms GPU in stalls in the Valley project.  The fence wait is marked with a different GPU stat scope, so it can be differentiated from the transfer workload.
* Ray tracing scene update now runs once, for the first view family only.  Saves around 1.2 ms GPU when three views are present.
* Added logic to sort the largest view families to run first, which tend to be the inner frustum, and are usually the bottleneck for the entire frame.  This improves scheduling, saving around 1.0 ms GPU.
* Avoid CPU stalls where Render Thread waits on RHI Thread, by only running "SetFlushResourcesRHI()" on the first view family.  Can save multiple ms when CPU bound.

This is part of ongoing multi-view-family scene render refactor work, where these improvements were made possible by the refactor.

#rnx

diff --git a/Engine/Plugins/Runtime/nDisplay/Source/DisplayCluster/Private/Game/EngineClasses/Basics/DisplayClusterViewportClient.cpp b/Engine/Plugins/Runtime/nDisplay/Source/DisplayCluster/Private/Game/EngineClasses/Basics/DisplayClusterViewportClient.cpp
--- a/Engine/Plugins/Runtime/nDisplay/Source/DisplayCluster/Private/Game/EngineClasses/Basics/DisplayClusterViewportClient.cpp
+++ b/Engine/Plugins/Runtime/nDisplay/Source/DisplayCluster/Private/Game/EngineClasses/Basics/DisplayClusterViewportClient.cpp
@@ -75,6 +75,37 @@ static FAutoConsoleVariableRef CVarDisplayClusterSingleRender(
 	ECVF_RenderThreadSafe
 );
 
+int32 GDisplayClusterSortViews = 1;
+static FAutoConsoleVariableRef CVarDisplayClusterSortViews(
+	TEXT("DC.SortViews"),
+	GDisplayClusterSortViews,
+	TEXT("Enable sorting of views by decreasing pixel count and decreasing GPU index.  Adds determinism, and tends to run inners first, which helps with scheduling, improving perf (default: enabled)."),
+	ECVF_RenderThreadSafe
+);
+
+struct FCompareViewFamilyBySizeAndGPU
+{
+	FORCEINLINE bool operator()(const FSceneViewFamilyContext& A, const FSceneViewFamilyContext& B) const
+	{
+		FIntPoint SizeA = A.RenderTarget->GetSizeXY();
+		FIntPoint SizeB = B.RenderTarget->GetSizeXY();
+		int32 AreaA = SizeA.X * SizeA.Y;
+		int32 AreaB = SizeB.X * SizeB.Y;
+
+		if (AreaA != AreaB)
+		{
+			// Decreasing area
+			return AreaA > AreaB;
+		}
+
+		int32 GPUIndexA = A.Views[0]->GPUMask.GetFirstIndex();
+		int32 GPUIndexB = B.Views[0]->GPUMask.GetFirstIndex();
+
+		// Decreasing GPU index
+		return GPUIndexA > GPUIndexB;
+	}
+};
+
 UDisplayClusterViewportClient::UDisplayClusterViewportClient(FVTableHelper& Helper)
 	: Super(Helper)
 {
@@ -418,18 +449,11 @@ void UDisplayClusterViewportClient::Draw(FViewport* InViewport, FCanvas* SceneCa
 		return Super::Draw(InViewport, SceneCanvas);
 	}
 
-	//Experimental code from render team, now always disabled
-	const bool bIsRenderedImmediatelyAfterAnotherViewFamily = false;
-	bool bIsFirstViewInMultipleViewFamily = true;
-
 	// Gather all view families first
 	TArray<FSceneViewFamilyContext*> ViewFamilies;
 
 	for (FDisplayClusterRenderFrame::FFrameRenderTarget& DCRenderTarget : RenderFrame.RenderTargets)
 	{
-		// Special flag, allow clear RTT surface only for first family
-		bool bAdditionalViewFamily = false;
-
 		for (FDisplayClusterRenderFrame::FFrameViewFamily& DCViewFamily : DCRenderTarget.ViewFamilies)
 		{
 			// Create the view family for rendering the world scene to the viewport's render target
@@ -437,14 +461,11 @@ void UDisplayClusterViewportClient::Draw(FViewport* InViewport, FCanvas* SceneCa
 				DCRenderTarget,
 				MyWorld->Scene,
 				EngineShowFlags,
-				bAdditionalViewFamily
+				false				// bAdditionalViewFamily  (filled in later, after list of families is known, and optionally reordered)
 			)));
 			FSceneViewFamilyContext& ViewFamily = *ViewFamilies.Last();
 			bool bIsFamilyVisible = false;
 
-			// Disable clean op for all next families on this render target
-			bAdditionalViewFamily = true;
-
 			// Configure family
 			RenderFrame.ViewportManager->ConfigureViewFamily(DCRenderTarget, DCViewFamily, ViewFamily);
 
@@ -764,22 +785,12 @@ void UDisplayClusterViewportClient::Draw(FViewport* InViewport, FCanvas* SceneCa
 				// Draw the player views.
 				if (!bDisableWorldRendering && PlayerViewMap.Num() > 0 && FSlateApplication::Get().GetPlatformApplication()->IsAllowedToRender()) //-V560
 				{
-					ViewFamily.bIsRenderedImmediatelyAfterAnotherViewFamily = bIsRenderedImmediatelyAfterAnotherViewFamily;
-					ViewFamily.bIsMultipleViewFamily = true;
-					ViewFamily.bIsFirstViewInMultipleViewFamily = bIsFirstViewInMultipleViewFamily;
-					bIsFirstViewInMultipleViewFamily = false;
-
 					// If we reach here, the view family should be rendered
 					bIsFamilyVisible = true;
 				}
 			}
 
-			if (bIsFamilyVisible)
-			{
-				// Disable clean op for all next families on this render target
-				bAdditionalViewFamily = true;
-			}
-			else
+			if (!bIsFamilyVisible)
 			{
 				// Family didn't end up visible, remove last view family from the array
 				delete ViewFamilies.Pop();
@@ -790,6 +801,29 @@ void UDisplayClusterViewportClient::Draw(FViewport* InViewport, FCanvas* SceneCa
 	// We gathered all the view families, now render them
 	if (!ViewFamilies.IsEmpty())
 	{
+		if (ViewFamilies.Num() > 1)
+		{
+#if WITH_MGPU
+			if (GDisplayClusterSortViews)
+			{
+				ViewFamilies.StableSort(FCompareViewFamilyBySizeAndGPU());
+			}
+#endif  // WITH_MGPU
+
+			// Initialize some flags for which view family is which, now that any view family reordering has been handled.
+			ViewFamilies[0]->bAdditionalViewFamily = false;
+			ViewFamilies[0]->bIsFirstViewInMultipleViewFamily = true;
+			ViewFamilies[0]->bIsMultipleViewFamily = true;
+
+			for (int32 FamilyIndex = 1; FamilyIndex < ViewFamilies.Num(); FamilyIndex++)
+			{
+				FSceneViewFamily& ViewFamily = *ViewFamilies[FamilyIndex];
+				ViewFamily.bAdditionalViewFamily = true;
+				ViewFamily.bIsFirstViewInMultipleViewFamily = false;
+				ViewFamily.bIsMultipleViewFamily = true;
+			}
+		}
+
 		if (GDisplayClusterSingleRender)
 		{
 			GetRendererModule().BeginRenderingViewFamilies(
diff --git a/Engine/Plugins/Runtime/nDisplay/Source/DisplayCluster/Private/Render/Viewport/DisplayClusterViewportManager_EditorPreview.cpp b/Engine/Plugins/Runtime/nDisplay/Source/DisplayCluster/Private/Render/Viewport/DisplayClusterViewportManager_EditorPreview.cpp
--- a/Engine/Plugins/Runtime/nDisplay/Source/DisplayCluster/Private/Render/Viewport/DisplayClusterViewportManager_EditorPreview.cpp
+++ b/Engine/Plugins/Runtime/nDisplay/Source/DisplayCluster/Private/Render/Viewport/DisplayClusterViewportManager_EditorPreview.cpp
@@ -120,9 +120,6 @@ bool FDisplayClusterViewportManager::RenderInEditor(class FDisplayClusterRenderF
 	FSceneInterface* PreviewScene = CurrentWorld->Scene;
 	FEngineShowFlags EngineShowFlags = FEngineShowFlags(EShowFlagInitMode::ESFIM_Game);
 
-	//Experimental code from render team, now always disabled
-	const bool bIsRenderedImmediatelyAfterAnotherViewFamily = false;
-
 	int32 ViewportIndex = 0;
 	bool bViewportsRenderPassDone = false;
 
@@ -220,8 +217,6 @@ bool FDisplayClusterViewportManager::RenderInEditor(class FDisplayClusterRenderF
 					ViewFamily.EngineShowFlags.ScreenPercentage = false;
 					ViewFamily.SetScreenPercentageInterface(new FLegacyScreenPercentageDriver(ViewFamily, 1.0f));
 
-					ViewFamily.bIsRenderedImmediatelyAfterAnotherViewFamily = bIsRenderedImmediatelyAfterAnotherViewFamily;
-
 					FCanvas Canvas(RenderTargetIt.RenderTargetPtr, nullptr, PreviewScene->GetWorld(), ERHIFeatureLevel::SM5, FCanvas::CDM_DeferDrawing /*FCanvas::CDM_ImmediateDrawing*/, 1.0f);
 					Canvas.Clear(FLinearColor::Black);
 
diff --git a/Engine/Source/Runtime/D3D12RHI/Private/D3D12CommandContext.cpp b/Engine/Source/Runtime/D3D12RHI/Private/D3D12CommandContext.cpp
--- a/Engine/Source/Runtime/D3D12RHI/Private/D3D12CommandContext.cpp
+++ b/Engine/Source/Runtime/D3D12RHI/Private/D3D12CommandContext.cpp
@@ -1264,9 +1264,41 @@ void FD3D12CommandContextRedirector::RHITransferResources(const TArrayView<const
 		{
 			if (!Param.bPullData)
 			{
-				// The dest waits for the src to be at this place in the frame before using the data.
-				GPUFence->GpuWait(Param.DestGPUIndex, ED3D12CommandQueueType::Direct, GPUFence->GetLastSignaledFence(), Param.SrcGPUIndex);
+				if (Param.DelayedFence)
+				{
+					check(Param.DelayedFence->DeviceGPUIndex == INDEX_NONE);
+					Param.DelayedFence->LastSignal = GPUFence->GetLastSignaledFence();
+					Param.DelayedFence->DeviceGPUIndex = Param.DestGPUIndex;
+					Param.DelayedFence->FenceGPUIndex = Param.SrcGPUIndex;
+				}
+				else
+				{
+					// The dest waits for the src to be at this place in the frame before using the data.
+					GPUFence->GpuWait(Param.DestGPUIndex, ED3D12CommandQueueType::Direct, GPUFence->GetLastSignaledFence(), Param.SrcGPUIndex);
+				}
+			}
+		}
+	}
+#endif // WITH_MGPU
+}
+
+void FD3D12CommandContextRedirector::RHITransferResourceWait(const TArrayView<FTransferResourceFenceData* const> FenceDatas)
+{
+#if WITH_MGPU
+	if (FenceDatas.Num() > 0)
+	{
+		FD3D12Fence* GPUFence = GetParentAdapter()->GetStagingFence();
+
+		RHISubmitCommandsHint();
+
+		for (FTransferResourceFenceData* FenceData : FenceDatas)
+		{
+			// It's possible to allocate fence data, but not actually do the transfer, in which case we treat this command as a nop.
+			if (FenceData->DeviceGPUIndex != INDEX_NONE)
+			{
+				GPUFence->GpuWait(FenceData->DeviceGPUIndex, ED3D12CommandQueueType::Direct, FenceData->LastSignal, FenceData->FenceGPUIndex);
 			}
+			delete FenceData;
 		}
 	}
 #endif // WITH_MGPU
diff --git a/Engine/Source/Runtime/D3D12RHI/Private/D3D12CommandContext.h b/Engine/Source/Runtime/D3D12RHI/Private/D3D12CommandContext.h
--- a/Engine/Source/Runtime/D3D12RHI/Private/D3D12CommandContext.h
+++ b/Engine/Source/Runtime/D3D12RHI/Private/D3D12CommandContext.h
@@ -525,6 +525,7 @@ public:
 	virtual void RHIEndTransitions(TArrayView<const FRHITransition*> Transitions) final override;
 
 	virtual void RHITransferResources(const TArrayView<const FTransferResourceParams> Params) final override;
+	virtual void RHITransferResourceWait(const TArrayView<FTransferResourceFenceData* const> FenceDatas) final override;
 
 	FORCEINLINE virtual void RHICopyToStagingBuffer(FRHIBuffer* SourceBuffer, FRHIStagingBuffer* DestinationStagingBuffer, uint32 Offset, uint32 NumBytes) final override
 	{
diff --git a/Engine/Source/Runtime/Engine/Public/SceneView.h b/Engine/Source/Runtime/Engine/Public/SceneView.h
--- a/Engine/Source/Runtime/Engine/Public/SceneView.h
+++ b/Engine/Source/Runtime/Engine/Public/SceneView.h
@@ -1736,9 +1736,6 @@ public:
 	bool bForceCopyCrossGPU;
 #endif
 
-	/** Whether this view is rendered in BeginRenderingViewFamily() immediately after another render. */
-	bool bIsRenderedImmediatelyAfterAnotherViewFamily = false;
-
 	/** Whether this view is one of multiple view families rendered in a single frame.  Affects occlusion query synchronization logic. */
 	bool bIsMultipleViewFamily = false;
 	
diff --git a/Engine/Source/Runtime/RHI/Public/RHICommandList.h b/Engine/Source/Runtime/RHI/Public/RHICommandList.h
--- a/Engine/Source/Runtime/RHI/Public/RHICommandList.h
+++ b/Engine/Source/Runtime/RHI/Public/RHICommandList.h
@@ -886,6 +886,18 @@ FRHICOMMAND_MACRO(FRHICommandTransferResources)
 	RHI_API void Execute(FRHICommandListBase& CmdList);
 };
 
+FRHICOMMAND_MACRO(FRHICommandTransferResourceWait)
+{
+	TArray<FTransferResourceFenceData*, TInlineAllocator<4>> FenceDatas;
+
+	FORCEINLINE_DEBUGGABLE FRHICommandTransferResourceWait(TArrayView<FTransferResourceFenceData* const> InFenceDatas)
+		: FenceDatas(InFenceDatas)
+	{
+	}
+
+	RHI_API void Execute(FRHICommandListBase & CmdList);
+};
+
 #endif // WITH_MGPU
 
 FRHICOMMAND_MACRO(FRHICommandSetStencilRef)
@@ -2946,6 +2958,20 @@ PRAGMA_ENABLE_DEPRECATION_WARNINGS
 #endif // WITH_MGPU
 	}
 
+	FORCEINLINE_DEBUGGABLE void TransferResourceWait(const TArrayView<FTransferResourceFenceData* const> FenceDatas)
+	{
+#if WITH_MGPU
+		if (Bypass())
+		{
+			GetComputeContext().RHITransferResourceWait(FenceDatas);
+		}
+		else
+		{
+			ALLOC_COMMAND(FRHICommandTransferResourceWait)(FenceDatas);
+		}
+#endif // WITH_MGPU
+	}
+
 #if WITH_MGPU
 	FORCEINLINE_DEBUGGABLE void WaitForTemporalEffect(const FName& EffectName)
 	{
diff --git a/Engine/Source/Runtime/RHI/Public/RHICommandListCommandExecutes.inl b/Engine/Source/Runtime/RHI/Public/RHICommandListCommandExecutes.inl
--- a/Engine/Source/Runtime/RHI/Public/RHICommandListCommandExecutes.inl
+++ b/Engine/Source/Runtime/RHI/Public/RHICommandListCommandExecutes.inl
@@ -115,6 +115,12 @@ void FRHICommandTransferResources::Execute(FRHICommandListBase& CmdList)
 	INTERNAL_DECORATOR_COMPUTE(RHITransferResources)(Params);
 }
 
+void FRHICommandTransferResourceWait::Execute(FRHICommandListBase& CmdList)
+{
+	RHISTAT(TransferResourceWait);
+	INTERNAL_DECORATOR_COMPUTE(RHITransferResourceWait)(FenceDatas);
+}
+
 #endif // WITH_MGPU
 
 void FRHICommandSetStencilRef::Execute(FRHICommandListBase& CmdList)
diff --git a/Engine/Source/Runtime/RHI/Public/RHIContext.h b/Engine/Source/Runtime/RHI/Public/RHIContext.h
--- a/Engine/Source/Runtime/RHI/Public/RHIContext.h
+++ b/Engine/Source/Runtime/RHI/Public/RHIContext.h
@@ -115,6 +115,22 @@ private:
 UE_DEPRECATED(5.0, "Please rename to FUniformBufferStaticBindings")
 typedef FUniformBufferStaticBindings FUniformBufferGlobalBindings;
 
+struct FTransferResourceFenceData
+{
+	uint64	LastSignal = 0;
+	int32	DeviceGPUIndex = INDEX_NONE;
+	int32	FenceGPUIndex = INDEX_NONE;
+};
+
+FORCEINLINE FTransferResourceFenceData* RHICreateTransferResourceFenceData()
+{
+#if WITH_MGPU
+	return new FTransferResourceFenceData;
+#else
+	return nullptr;
+#endif
+}
+
 /** Parameters for RHITransferResources, used to copy memory between GPUs */
 struct FTransferResourceParams
 {
@@ -155,6 +171,14 @@ struct FTransferResourceParams
 	// Whether the GPUs must handshake before and after the transfer. Required if the texture rect is being written to in several render passes.
 	// Otherwise, minimal synchronization will be used.
 	bool bLockStepGPUs = true;
+	/**
+	  * Optional pointer where fence data can be written if you want to delay waiting on the GPU fence for a resource transfer.
+	  * Should be created via "RHICreateTransferResourceFenceData", and must later be consumed via "TransferResourceWait" command.
+	  * Note that it is valid to consume the fence data, even if you don't end up implementing a transfer that uses it -- it will
+	  * behave as a nop in that case.  That can simplify cases where the transfer may be conditional, and you don't want to worry
+	  * about whether it occurred or not, but need to reserve the possibility.
+	  */
+	FTransferResourceFenceData* DelayedFence = nullptr;
 };
 
 /** Context that is capable of doing Compute work.  Can be async or compute on the gfx pipe. */
@@ -326,6 +350,17 @@ public:
 		/* empty default implementation */
 	}
 
+	virtual void RHITransferResourceWait(const TArrayView<FTransferResourceFenceData* const> FenceDatas)
+	{
+		/* default noop implementation */
+#if WITH_MGPU
+		for (FTransferResourceFenceData* FenceData : FenceDatas)
+		{
+			delete FenceData;
+		}
+#endif
+	}
+
 	virtual void RHIBuildAccelerationStructures(const TArrayView<const FRayTracingGeometryBuildParams> Params, const FRHIBufferRange& ScratchBufferRange)
 	{
 		checkNoEntry();
diff --git a/Engine/Source/Runtime/Renderer/Private/DeferredShadingRenderer.cpp b/Engine/Source/Runtime/Renderer/Private/DeferredShadingRenderer.cpp
--- a/Engine/Source/Runtime/Renderer/Private/DeferredShadingRenderer.cpp
+++ b/Engine/Source/Runtime/Renderer/Private/DeferredShadingRenderer.cpp
@@ -1680,9 +1680,9 @@ static void ReleaseRaytracingResources(FRDGBuilder& GraphBuilder, TArrayView<FVi
 void FDeferredShadingSceneRenderer::WaitForRayTracingScene(FRDGBuilder& GraphBuilder, FRDGBufferRef DynamicGeometryScratchBuffer)
 {
 	bool bAnyRayTracingPassEnabled = false;
-	for (int32 ViewIndex = 0; ViewIndex < Views.Num(); ViewIndex++)
+	for (int32 ViewIndex = 0; ViewIndex < AllFamilyViews.Num(); ViewIndex++)
 	{
-		bAnyRayTracingPassEnabled |= AnyRayTracingPassEnabled(Scene, Views[ViewIndex]);
+		bAnyRayTracingPassEnabled |= AnyRayTracingPassEnabled(Scene, AllFamilyViews[ViewIndex]);
 	}
 
 	if (!bAnyRayTracingPassEnabled)
@@ -1697,7 +1697,7 @@ void FDeferredShadingSceneRenderer::WaitForRayTracingScene(FRDGBuilder& GraphBui
 	SetupRayTracingPipelineStates(GraphBuilder.RHICmdList);
 
 	bool bAnyInlineRayTracingPassEnabled = false;
-	for (const FViewInfo& View : Views)
+	for (const FViewInfo& View : AllFamilyViews)
 	{
 		bAnyInlineRayTracingPassEnabled |= Lumen::AnyLumenHardwareInlineRayTracingPassEnabled(Scene, View);
 	}
@@ -1705,7 +1705,7 @@ void FDeferredShadingSceneRenderer::WaitForRayTracingScene(FRDGBuilder& GraphBui
 	if (bAnyInlineRayTracingPassEnabled)
 	{
 		const int32 ReferenceViewIndex = 0;
-		FViewInfo& ReferenceView = Views[ReferenceViewIndex];
+		FViewInfo& ReferenceView = AllFamilyViews[ReferenceViewIndex];
 
 		SetupLumenHardwareRayTracingHitGroupBuffer(ReferenceView);
 	}
@@ -1719,7 +1719,7 @@ void FDeferredShadingSceneRenderer::WaitForRayTracingScene(FRDGBuilder& GraphBui
 		[this, PassParams](FRHICommandListImmediate& RHICmdList)
 	{
 		const int32 ReferenceViewIndex = 0;
-		FViewInfo& ReferenceView = Views[ReferenceViewIndex];
+		FViewInfo& ReferenceView = AllFamilyViews[ReferenceViewIndex];
 
 		const bool bIsPathTracing = ActiveViewFamily->EngineShowFlags.PathTracing;
 
@@ -1792,7 +1792,7 @@ void FDeferredShadingSceneRenderer::WaitForRayTracingScene(FRDGBuilder& GraphBui
 			TArray<FRHIRayTracingShader*> DeferredMaterialRayGenShaders;
 			if (!IsForwardShadingEnabled(ShaderPlatform))
 			{
-				for (const FViewInfo& View : Views)
+				for (const FViewInfo& View : AllFamilyViews)
 				{
 					PrepareRayTracingReflectionsDeferredMaterial(View, *Scene, DeferredMaterialRayGenShaders);
 					PrepareRayTracingDeferredReflectionsDeferredMaterial(View, *Scene, DeferredMaterialRayGenShaders);
@@ -1817,7 +1817,7 @@ void FDeferredShadingSceneRenderer::WaitForRayTracingScene(FRDGBuilder& GraphBui
 			TArray<FRHIRayTracingShader*> LumenHardwareRayTracingRayGenShaders;
 			if (DoesPlatformSupportLumenGI(ShaderPlatform))
 			{
-				for (const FViewInfo& View : Views)
+				for (const FViewInfo& View : AllFamilyViews)
 				{
 					PrepareLumenHardwareRayTracingVisualizeLumenMaterial(View, LumenHardwareRayTracingRayGenShaders);
 					PrepareLumenHardwareRayTracingRadianceCacheLumenMaterial(View, LumenHardwareRayTracingRayGenShaders);
@@ -1837,9 +1837,9 @@ void FDeferredShadingSceneRenderer::WaitForRayTracingScene(FRDGBuilder& GraphBui
 		}
 
 		// Send ray tracing resources from reference view to all others.
-		for (int32 ViewIndex = 1; ViewIndex < Views.Num(); ++ViewIndex)
+		for (int32 ViewIndex = 1; ViewIndex < AllFamilyViews.Num(); ++ViewIndex)
 		{
-			FViewInfo& View = Views[ViewIndex];
+			FViewInfo& View = AllFamilyViews[ViewIndex];
 			View.RayTracingMaterialGatherPipeline = ReferenceView.RayTracingMaterialGatherPipeline;
 			View.LumenHardwareRayTracingMaterialPipeline = ReferenceView.LumenHardwareRayTracingMaterialPipeline;
 		}
@@ -2172,7 +2172,7 @@ void FDeferredShadingSceneRenderer::Render(FRDGBuilder& GraphBuilder)
 	{
 		RDG_GPU_STAT_SCOPE(GraphBuilder, GPUSceneUpdate);
 
-		if (!ActiveViewFamily->bIsRenderedImmediatelyAfterAnotherViewFamily)
+		if (IsFirstViewFamily())
 		{
 			GraphBuilder.SetFlushResourcesRHI();
 		}
@@ -2656,9 +2656,17 @@ void FDeferredShadingSceneRenderer::Render(FRDGBuilder& GraphBuilder)
 	}
 
 #if RHI_RAYTRACING
-	// Async AS builds can potentially overlap with BasePass
-	FRDGBufferRef DynamicGeometryScratchBuffer;
-	DispatchRayTracingWorldUpdates(GraphBuilder, DynamicGeometryScratchBuffer);
+	// Async AS builds can potentially overlap with BasePass.  We only need to update ray tracing scene for the first view family,
+	// if multiple are rendered in a single scene render call.
+	FRDGBufferRef DynamicGeometryScratchBuffer = nullptr;
+	if (IsFirstViewFamily())
+	{
+		DispatchRayTracingWorldUpdates(GraphBuilder, DynamicGeometryScratchBuffer);
+	}
+	else
+	{
+		bRayTracingSceneReady = true;
+	}
 #endif
 
 	{
@@ -2789,7 +2797,7 @@ void FDeferredShadingSceneRenderer::Render(FRDGBuilder& GraphBuilder)
 
 #if RHI_RAYTRACING
 		// Lumen scene lighting requires ray tracing scene to be ready if HWRT shadows are desired
-		if (Lumen::UseHardwareRayTracedSceneLighting(*ActiveViewFamily))
+		if (!bRayTracingSceneReady && Lumen::UseHardwareRayTracedSceneLighting(*ActiveViewFamily))
 		{
 			WaitForRayTracingScene(GraphBuilder, DynamicGeometryScratchBuffer);
 			bRayTracingSceneReady = true;
diff --git a/Engine/Source/Runtime/Renderer/Private/SceneRendering.cpp b/Engine/Source/Runtime/Renderer/Private/SceneRendering.cpp
--- a/Engine/Source/Runtime/Renderer/Private/SceneRendering.cpp
+++ b/Engine/Source/Runtime/Renderer/Private/SceneRendering.cpp
@@ -3063,6 +3063,21 @@ FRHIGPUMask FSceneRenderer::ComputeGPUMasks(FRHICommandListImmediate& RHICmdList
 
 #if WITH_MGPU
 DECLARE_GPU_STAT_NAMED(CrossGPUTransfers, TEXT("Cross GPU Transfer"));
+DECLARE_GPU_STAT_NAMED(CrossGPUSync, TEXT("Cross GPU Sync"));
+
+struct FCrossGPUTransfer
+{
+	FIntRect TransferRect;
+	int32 SrcGPUIndex;
+	int32 DestGPUIndex;
+	FTransferResourceFenceData* DelayedFence;
+
+	FCrossGPUTransfer(const FIntRect& InTransferRect, uint32 InSrcGPUIndex, uint32 InDestGPUIndex)
+		: TransferRect(InTransferRect), SrcGPUIndex(InSrcGPUIndex), DestGPUIndex(InDestGPUIndex), DelayedFence(nullptr)
+	{
+		// Empty
+	}
+};
 #endif // WITH_MGPU
 
 void FSceneRenderer::DoCrossGPUTransfers(FRDGBuilder& GraphBuilder, FRHIGPUMask RenderTargetGPUMask, FRDGTextureRef ViewFamilyTexture)
@@ -3074,35 +3089,89 @@ void FSceneRenderer::DoCrossGPUTransfers(FRDGBuilder& GraphBuilder, FRHIGPUMask
 		RDG_GPU_MASK_SCOPE(GraphBuilder, FRHIGPUMask::All());
 		RDG_GPU_STAT_SCOPE(GraphBuilder, CrossGPUTransfers);
 
-		// A readback pass is the closest analog to what this is doing. There isn't a way to express cross-GPU transfers via the RHI barrier API.
-		AddReadbackTexturePass(GraphBuilder, RDG_EVENT_NAME("CrossGPUTransfers"), ViewFamilyTexture, 
-			[this, RenderTargetGPUMask, ViewFamilyTexture] (FRHICommandListImmediate& RHICmdList)
-		{
-			const FIntPoint Extent = ViewFamilyTexture->Desc.Extent;
+		// Need to use this structure as an intermediate, because the RHI texture reference isn't available yet,
+		// and must be fetched inside the pass.
+		TArray<FCrossGPUTransfer, SceneRenderingAllocator> Transfers;
+
+		const FIntPoint Extent = ViewFamilyTexture->Desc.Extent;
 
-			for (int32 ViewIndex = 0; ViewIndex < Views.Num(); ++ViewIndex)
+		for (int32 ViewIndex = 0; ViewIndex < Views.Num(); ++ViewIndex)
+		{
+			const FViewInfo& ViewInfo = Views[ViewIndex];
+			if (ViewInfo.bAllowCrossGPUTransfer && ViewInfo.GPUMask != RenderTargetGPUMask)
 			{
-				const FViewInfo& ViewInfo = Views[ViewIndex];
-				if (ViewInfo.bAllowCrossGPUTransfer && ViewInfo.GPUMask != RenderTargetGPUMask)
+				// Clamp the view rect by the rendertarget rect to prevent issues when resizing the viewport.
+				const FIntRect TransferRect(ViewInfo.UnscaledViewRect.Min.ComponentMin(Extent), ViewInfo.UnscaledViewRect.Max.ComponentMin(Extent));
+				if (TransferRect.Width() > 0 && TransferRect.Height() > 0)
 				{
-					// Clamp the view rect by the rendertarget rect to prevent issues when resizing the viewport.
-					const FIntRect TransferRect(ViewInfo.UnscaledViewRect.Min.ComponentMin(Extent), ViewInfo.UnscaledViewRect.Max.ComponentMin(Extent));
-					if (TransferRect.Width() > 0 && TransferRect.Height() > 0)
+					for (uint32 RenderTargetGPUIndex : RenderTargetGPUMask)
 					{
-						for (uint32 RenderTargetGPUIndex : RenderTargetGPUMask)
+						if (!ViewInfo.GPUMask.Contains(RenderTargetGPUIndex))
 						{
-							if (!ViewInfo.GPUMask.Contains(RenderTargetGPUIndex))
-							{
-								FTransferResourceParams Param(static_cast<FRHITexture2D*>(ViewFamilyTexture->GetRHI()), TransferRect, ViewInfo.GPUMask.GetFirstIndex(), RenderTargetGPUIndex, false, false);
-								RHICmdList.TransferResources(TArrayView<const FTransferResourceParams>(&Param, 1));
-							}
+							Transfers.Add(FCrossGPUTransfer(TransferRect, ViewInfo.GPUMask.GetFirstIndex(), RenderTargetGPUIndex));
 						}
 					}
 				}
 			}
+		}
+
+		if (Transfers.Num() > 0)
+		{
+			// If not the last view family, defer cross transfer fence wait
+			if (ActiveViewFamily != &ViewFamilies.Last())
+			{
+				ActiveViewFamily->CrossGPUTransferFences.SetNumUninitialized(Transfers.Num());
+				for (int32 TransferIndex = 0; TransferIndex < Transfers.Num(); TransferIndex++)
+				{
+					FTransferResourceFenceData* FenceData = RHICreateTransferResourceFenceData();
+
+					ActiveViewFamily->CrossGPUTransferFences[TransferIndex] = FenceData;
+					Transfers[TransferIndex].DelayedFence = FenceData;
+				}
+			}
+
+			// A readback pass is the closest analog to what this is doing. There isn't a way to express cross-GPU transfers via the RHI barrier API.
+			AddReadbackTexturePass(GraphBuilder, RDG_EVENT_NAME("CrossGPUTransfers"), ViewFamilyTexture,
+				[this, ViewFamilyTexture, LocalTransfers = MoveTemp(Transfers)](FRHICommandListImmediate& RHICmdList)
+				{
+					TArray<FTransferResourceParams> TransferParams;
+					for (const FCrossGPUTransfer& Transfer : LocalTransfers)
+					{
+						TransferParams.Add(FTransferResourceParams(ViewFamilyTexture->GetRHI(), Transfer.SrcGPUIndex, Transfer.DestGPUIndex, false, false));
+						TransferParams.Last().DelayedFence = Transfer.DelayedFence;
+					}
+
+					RHICmdList.TransferResources(TransferParams);
+				});
+		}
+	}
+#endif // WITH_MGPU
+}
+
+void FSceneRenderer::FlushCrossGPUFences(FRDGBuilder& GraphBuilder)
+{
+#if WITH_MGPU
+	TArray<FTransferResourceFenceData*> FenceDatas;
+	for (int32 FamilyIndex = 0; FamilyIndex < ViewFamilies.Num(); FamilyIndex++)
+	{
+		FViewFamilyInfo& ViewFamily = ViewFamilies[FamilyIndex];
+		for (FTransferResourceFenceData* FenceData : ViewFamily.CrossGPUTransferFences)
+		{
+			FenceDatas.Emplace(FenceData);
+		}
+	}
+
+	if (FenceDatas.Num() > 0)
+	{
+		RDG_GPU_STAT_SCOPE(GraphBuilder, CrossGPUSync);
+
+		AddPass(GraphBuilder, RDG_EVENT_NAME("CrossGPUTransferSync"),
+			[LocalFenceDatas = MoveTemp(FenceDatas)](FRHICommandListImmediate& RHICmdList)
+		{
+			RHICmdList.TransferResourceWait(LocalFenceDatas);
 		});
 	}
-#endif
+#endif // WITH_MGPU
 }
 
 
@@ -4230,6 +4299,13 @@ static void RenderViewFamily_RenderThread(FRHICommandListImmediate& RHICmdList,
 				// Render the scene.
 				SceneRenderer->Render(GraphBuilder);
 			}
+
+			// If this is the last view family, flush pending cross GPU transfer fences
+			if (&ViewFamily == &SceneRenderer->ViewFamilies.Last())
+			{
+				SceneRenderer->FlushCrossGPUFences(GraphBuilder);
+			}
+
 			GraphBuilder.Execute();
 		}
 
diff --git a/Engine/Source/Runtime/Renderer/Private/SceneRendering.h b/Engine/Source/Runtime/Renderer/Private/SceneRendering.h
--- a/Engine/Source/Runtime/Renderer/Private/SceneRendering.h
+++ b/Engine/Source/Runtime/Renderer/Private/SceneRendering.h
@@ -1897,6 +1897,11 @@ public:
 
 	FSceneTexturesConfig SceneTexturesConfig;
 
+#if WITH_MGPU
+	/** Fences for cross GPU transfers */
+	TArray<FTransferResourceFenceData*> CrossGPUTransferFences;
+#endif
+
 	/** Get scene textures associated with this view family -- asserts or checks that they have been initialized */
 	inline FSceneTextures& GetSceneTextures()
 	{
@@ -1975,6 +1980,7 @@ public:
 
 	/** Update the rendertarget with each view results.*/
 	void DoCrossGPUTransfers(FRDGBuilder& GraphBuilder, FRHIGPUMask RenderTargetGPUMask, FRDGTextureRef ViewFamilyTexture);
+	void FlushCrossGPUFences(FRDGBuilder& GraphBuilder);
 
 	bool DoOcclusionQueries() const;
 
@@ -1987,6 +1993,14 @@ public:
 
 	bool ShouldDumpMeshDrawCommandInstancingStats() const { return bDumpMeshDrawCommandInstancingStats; }
 
+	/**
+	  * Returns whether the active view family is the first being rendered in a given scene render call.  This is different from the
+	  * user facing bIsFirstViewInMultipleViewFamily flag, which defines whether it's the first view family rendered in the entire frame.
+	  * It's possible to call the scene renderer multiple times in a frame with individual view families, or pass in several view families
+	  * at once, and some code paths need to know the difference.
+	  */
+	bool IsFirstViewFamily() const { return ActiveViewFamily == ViewFamilies.GetData(); }
+
 	/** bound shader state for occlusion test prims */
 	static FGlobalBoundShaderState OcclusionTestBoundShaderState;
 	
