Reverse strata tile scheduling from complex to simple in order to have long wave earlier when rendering lighting.

diff --git a/Engine/Source/Runtime/Renderer/Private/LightRendering.cpp b/Engine/Source/Runtime/Renderer/Private/LightRendering.cpp
--- a/Engine/Source/Runtime/Renderer/Private/LightRendering.cpp
+++ b/Engine/Source/Runtime/Renderer/Private/LightRendering.cpp
@@ -2215,12 +2215,12 @@ static void RenderLight(
 		//   so that the geometry get correctly stencil culled on complex/simple part of the screen
 		if (Strata::IsStrataEnabled())
 		{
-			// Simple tiles
+			// Complex tiles
 			{
-				const EStrataTileType TileType = EStrataTileType::ESimple;
+				const EStrataTileType TileType = EStrataTileType::EComplex;
 				PermutationVector.Set<FDeferredLightPS::FStrataTileType>(TileType);
 				TShaderMapRef< FDeferredLightPS > PixelShader(View.ShaderMap, PermutationVector);
-				InternalRenderLight(GraphBuilder, Scene, View, LightSceneInfo, PixelShader, PassParameters, TileType, TEXT("Light::StandardDeferred(Simple)"));
+				InternalRenderLight(GraphBuilder, Scene, View, LightSceneInfo, PixelShader, PassParameters, TileType, TEXT("Light::StandardDeferred(Complex)"));
 			}
 			// Single tiles
 			{
@@ -2229,12 +2229,12 @@ static void RenderLight(
 				TShaderMapRef< FDeferredLightPS > PixelShader(View.ShaderMap, PermutationVector);
 				InternalRenderLight(GraphBuilder, Scene, View, LightSceneInfo, PixelShader, PassParameters, TileType, TEXT("Light::StandardDeferred(Single)"));
 			}
-			// Complex tiles
+			// Simple tiles
 			{
-				const EStrataTileType TileType = EStrataTileType::EComplex;
+				const EStrataTileType TileType = EStrataTileType::ESimple;
 				PermutationVector.Set<FDeferredLightPS::FStrataTileType>(TileType);
 				TShaderMapRef< FDeferredLightPS > PixelShader(View.ShaderMap, PermutationVector);
-				InternalRenderLight(GraphBuilder, Scene, View, LightSceneInfo, PixelShader, PassParameters, TileType, TEXT("Light::StandardDeferred(Complex)"));
+				InternalRenderLight(GraphBuilder, Scene, View, LightSceneInfo, PixelShader, PassParameters, TileType, TEXT("Light::StandardDeferred(Simple)"));
 			}
 		}
 		else
@@ -2588,9 +2588,9 @@ void FDeferredShadingSceneRenderer::RenderSimpleLightsStandardDeferred(
 
 		if (Strata::IsStrataEnabled())
 		{
-			InternalRenderSimpleLightsStandardDeferred(GraphBuilder, Scene, View, ViewIndex, NumViews, SceneTextures, SimpleLights, EStrataTileType::ESimple);
-			InternalRenderSimpleLightsStandardDeferred(GraphBuilder, Scene, View, ViewIndex, NumViews, SceneTextures, SimpleLights, EStrataTileType::ESingle);
 			InternalRenderSimpleLightsStandardDeferred(GraphBuilder, Scene, View, ViewIndex, NumViews, SceneTextures, SimpleLights, EStrataTileType::EComplex);
+			InternalRenderSimpleLightsStandardDeferred(GraphBuilder, Scene, View, ViewIndex, NumViews, SceneTextures, SimpleLights, EStrataTileType::ESingle);
+			InternalRenderSimpleLightsStandardDeferred(GraphBuilder, Scene, View, ViewIndex, NumViews, SceneTextures, SimpleLights, EStrataTileType::ESimple);
 		}
 		else
 		{
