Update dynamic primitives to use payload encoding function

diff --git a/Engine/Source/Runtime/Renderer/Private/VirtualShadowMaps/VirtualShadowMapCacheManager.cpp b/Engine/Source/Runtime/Renderer/Private/VirtualShadowMaps/VirtualShadowMapCacheManager.cpp
--- a/Engine/Source/Runtime/Renderer/Private/VirtualShadowMaps/VirtualShadowMapCacheManager.cpp
+++ b/Engine/Source/Runtime/Renderer/Private/VirtualShadowMaps/VirtualShadowMapCacheManager.cpp
@@ -149,16 +149,37 @@ void FVirtualShadowMapCacheEntry::UpdateLocal(int32 VirtualShadowMapId, const FW
 	bCurrentRendered = false;
 }
 
+static inline uint32 EncodeInstanceInvalidationPayload(bool bInvalidateStaticPage, int32 ClipmapVirtualShadowMapId = INDEX_NONE)
+{
+	uint32 Payload = 0;
+
+	if (bInvalidateStaticPage)
+	{
+		Payload = Payload | 0x2;
+	}
+
+	if (ClipmapVirtualShadowMapId != INDEX_NONE)
+	{
+		// Do a single clipmap level
+		Payload = Payload | 0x1;
+		Payload = Payload | (((uint32)ClipmapVirtualShadowMapId) << 2);
+	}
+
+	return Payload;
+}
+
 FVirtualShadowMapArrayCacheManager::FInvalidatingPrimitiveCollector::FInvalidatingPrimitiveCollector(FVirtualShadowMapArrayCacheManager* InVirtualShadowMapArrayCacheManager)
 	: AlreadyAddedPrimitives(false, InVirtualShadowMapArrayCacheManager->Scene->Primitives.Num())
 	, Scene(*InVirtualShadowMapArrayCacheManager->Scene)
 	, GPUScene(InVirtualShadowMapArrayCacheManager->Scene->GPUScene)
 	, Manager(*InVirtualShadowMapArrayCacheManager)
 {
+	bool bPossiblyCachedAsStatic = false;	// TODO
+
 	// Add and clear pending invalidations enqueued on the GPU Scene from dynamic primitives added since last invalidation
 	for (const FGPUScene::FInstanceRange& Range : GPUScene.DynamicPrimitiveInstancesToInvalidate)
 	{
-		LoadBalancer.Add(Range.InstanceSceneDataOffset, Range.NumInstanceSceneDataEntries, 0U);
+		LoadBalancer.Add(Range.InstanceSceneDataOffset, Range.NumInstanceSceneDataEntries, EncodeInstanceInvalidationPayload(bPossiblyCachedAsStatic));
 #if VSM_LOG_INVALIDATIONS
 		RangesStr.Appendf(TEXT("[%6d, %6d), "), Range.InstanceSceneDataOffset, Range.InstanceSceneDataOffset + Range.NumInstanceSceneDataEntries);
 #endif
@@ -178,7 +199,8 @@ FVirtualShadowMapArrayCacheManager::FInvalidatingPrimitiveCollector::FInvalidati
 				if (SmCacheEntry.IsValid())
 				{
 					// Lowest bit indicates whether to run the clipmap loop, add 1 to ID so != 0 <==> single level processing
-					LoadBalancer.Add(Range.InstanceSceneDataOffset, Range.NumInstanceSceneDataEntries, (uint32(SmCacheEntry->CurrentVirtualShadowMapId + 1) << 1U) | 0U);
+					LoadBalancer.Add(Range.InstanceSceneDataOffset, Range.NumInstanceSceneDataEntries,
+						EncodeInstanceInvalidationPayload(bPossiblyCachedAsStatic, SmCacheEntry->CurrentVirtualShadowMapId));
 				}
 			}
 
@@ -191,25 +213,6 @@ FVirtualShadowMapArrayCacheManager::FInvalidatingPrimitiveCollector::FInvalidati
 	}
 }
 
-static inline uint32 EncodeInstanceInvalidationPayload(bool bInvalidateStaticPage, int32 ClipmapVirtualShadowMapId = INDEX_NONE)
-{
-	uint32 Payload = 0;
-
-	if (bInvalidateStaticPage)
-	{
-		Payload = Payload | 0x2;
-	}
-
-	if (ClipmapVirtualShadowMapId != INDEX_NONE)
-	{
-		// Do a single clipmap level
-		Payload = Payload | 0x1;
-		Payload = Payload | (((uint32)ClipmapVirtualShadowMapId) << 2);
-	}
-
-	return Payload;
-}
-
 void FVirtualShadowMapArrayCacheManager::FInvalidatingPrimitiveCollector::Add(const FPrimitiveSceneInfo * PrimitiveSceneInfo)
 {
 	int32 PrimitiveID = PrimitiveSceneInfo->GetIndex();
