Fixes crash in TSR when switching r.TSR.History.HighFrequencyOnly

diff --git a/Engine/Source/Runtime/Renderer/Private/PostProcess/TemporalSuperResolution.cpp b/Engine/Source/Runtime/Renderer/Private/PostProcess/TemporalSuperResolution.cpp
--- a/Engine/Source/Runtime/Renderer/Private/PostProcess/TemporalSuperResolution.cpp
+++ b/Engine/Source/Runtime/Renderer/Private/PostProcess/TemporalSuperResolution.cpp
@@ -116,6 +116,7 @@ BEGIN_SHADER_PARAMETER_STRUCT(FTSRHistoryArrayIndices, )
 	SHADER_PARAMETER(int32, LowFrequency)
 	SHADER_PARAMETER(int32, HighFrequency)
 	SHADER_PARAMETER(int32, Translucency)
+	SHADER_PARAMETER(int32, Size)
 END_SHADER_PARAMETER_STRUCT()
 
 BEGIN_SHADER_PARAMETER_STRUCT(FTSRHistoryTextures, )
@@ -154,6 +155,33 @@ BEGIN_SHADER_PARAMETER_STRUCT(FTSRPrevHistoryParameters, )
 	SHADER_PARAMETER(float, HistoryPreExposureCorrection)
 END_SHADER_PARAMETER_STRUCT()
 
+enum class ETSRHistoryFormatBits : uint32
+{
+	LowFrequency = 1 << 0,
+	Translucency = 1 << 1,
+};
+ENUM_CLASS_FLAGS(ETSRHistoryFormatBits);
+
+FTSRHistoryArrayIndices TranslateHistoryFormatBitsToArrayIndices(ETSRHistoryFormatBits HistoryFormatBits)
+{
+	FTSRHistoryArrayIndices ArrayIndices;
+	ArrayIndices.Size = 2;
+	ArrayIndices.HighFrequency = 0;
+	ArrayIndices.Translucency = 1;
+
+	if (EnumHasAnyFlags(HistoryFormatBits, ETSRHistoryFormatBits::LowFrequency))
+	{
+		ArrayIndices.LowFrequency = ArrayIndices.Size;
+		ArrayIndices.Size += 1;
+	}
+	else
+	{
+		ArrayIndices.LowFrequency = -1;
+	}
+
+	return ArrayIndices;
+}
+
 FTSRHistorySRVs CreateSRVs(FRDGBuilder& GraphBuilder, const FTSRHistoryTextures& Textures)
 {
 	FTSRHistorySRVs SRVs;
@@ -614,8 +642,19 @@ ITemporalUpscaler::FOutputs AddTemporalSuperResolutionPasses(
 	// whether TSR passes can run on async compute.
 	ERDGPassFlags ComputePassFlags = (GSupportsEfficientAsyncCompute && CVarTSRAsyncCompute.GetValueOnRenderThread() != 0) ? ERDGPassFlags::AsyncCompute : ERDGPassFlags::Compute;
 
+	const bool bHistoryHighFrequencyOnly = CVarTSRHistoryHighFrequencyOnly.GetValueOnRenderThread() != 0;
+
+	ETSRHistoryFormatBits HistoryFormatBits;
+	{
+		HistoryFormatBits = ETSRHistoryFormatBits::Translucency;
+		if (!bHistoryHighFrequencyOnly)
+		{
+			HistoryFormatBits |= ETSRHistoryFormatBits::LowFrequency;
+		}
+	}
+
 	// Whether to use camera cut shader permutation or not.
-	bool bCameraCut = !InputHistory.IsValid() || View.bCameraCut;
+	bool bCameraCut = !InputHistory.IsValid() || View.bCameraCut || ETSRHistoryFormatBits(InputHistory.FormatBit) != HistoryFormatBits;
 
 	FTSRUpdateHistoryCS::EQuality UpdateHistoryQuality = FTSRUpdateHistoryCS::EQuality(FMath::Clamp(CVarTSRHistoryUpdateQuality.GetValueOnRenderThread(), 0, int32(FTSRUpdateHistoryCS::EQuality::MAX) - 1));
 
@@ -625,8 +664,6 @@ ITemporalUpscaler::FOutputs AddTemporalSuperResolutionPasses(
 
 	const bool bRejectSeparateTranslucency = false; // bIsSeperateTranslucyTexturesValid && CVarTSRTranslucencyPreviousFrameRejection.GetValueOnRenderThread() != 0;
 
-	const bool bHistoryHighFrequencyOnly = CVarTSRHistoryHighFrequencyOnly.GetValueOnRenderThread() != 0;
-
 	// Whether the UpdateHistory pass also output half res history to reduce memory bottleneck up the FTSRDecimateHistoryCS.
 	const bool bHalfResHistory = bHistoryHighFrequencyOnly;
 
@@ -903,16 +940,15 @@ ITemporalUpscaler::FOutputs AddTemporalSuperResolutionPasses(
 	// Create new history.
 	FTSRHistoryTextures History;
 	{
+		History.ArrayIndices = TranslateHistoryFormatBitsToArrayIndices(HistoryFormatBits);
+
 		FRDGTextureDesc ArrayDesc = FRDGTextureDesc::Create2DArray(
 			HistoryExtent,
 			(CVarTSRR11G11B10History.GetValueOnRenderThread() != 0 && !bSupportsAlpha) ? PF_FloatR11G11B10 : PF_FloatRGBA,
 			FClearValueBinding::None,
 			TexCreate_ShaderResource | TexCreate_UAV,
-			/* ArraySize = */ bHistoryHighFrequencyOnly ? 2 : 3);
+			History.ArrayIndices.Size);
 
-		History.ArrayIndices.LowFrequency = bHistoryHighFrequencyOnly ? -1 : 2;
-		History.ArrayIndices.HighFrequency = 0;
-		History.ArrayIndices.Translucency = 1;
 		History.ColorArray = GraphBuilder.CreateTexture(ArrayDesc, TEXT("TSR.History.ColorArray"));
 	}
 	{
@@ -966,7 +1002,7 @@ ITemporalUpscaler::FOutputs AddTemporalSuperResolutionPasses(
 
 	// Setup the previous frame history
 	FTSRHistorySRVs PrevHistorySRVs;
-	if (InputHistory.IsValid())
+	if (!bCameraCut)
 	{
 		FTSRHistoryTextures PrevHistory;
 		PrevHistory.ArrayIndices = History.ArrayIndices;
@@ -1568,6 +1604,7 @@ ITemporalUpscaler::FOutputs AddTemporalSuperResolutionPasses(
 		FTSRHistory& OutputHistory = View.ViewState->PrevFrameViewInfo.TSRHistory;
 		OutputHistory.InputViewportRect = InputRect;
 		OutputHistory.OutputViewportRect = FIntRect(FIntPoint(0, 0), HistorySize);
+		OutputHistory.FormatBit = uint32(HistoryFormatBits);
 
 		// Extract filterable history
 		GraphBuilder.QueueTextureExtraction(History.ColorArray, &OutputHistory.ColorArray);
diff --git a/Engine/Source/Runtime/Renderer/Private/SceneRendering.h b/Engine/Source/Runtime/Renderer/Private/SceneRendering.h
--- a/Engine/Source/Runtime/Renderer/Private/SceneRendering.h
+++ b/Engine/Source/Runtime/Renderer/Private/SceneRendering.h
@@ -906,6 +906,7 @@ struct FTSRHistory
 	FIntRect InputViewportRect;
 	FIntRect OutputViewportRect;
 
+	uint32 FormatBit = 0;
 
 	void SafeRelease()
 	{
