Add optional async support for Strata classification pass (enabled by default).

Save 0.08ms on simple case on last gen platform.

diff --git a/Engine/Source/Runtime/Renderer/Private/Strata/Strata.cpp b/Engine/Source/Runtime/Renderer/Private/Strata/Strata.cpp
--- a/Engine/Source/Runtime/Renderer/Private/Strata/Strata.cpp
+++ b/Engine/Source/Runtime/Renderer/Private/Strata/Strata.cpp
@@ -66,6 +66,12 @@ static TAutoConsoleVariable<int32> CVarStrataDebugPeelLayersAboveDepth(
 	TEXT("Strata debug control to progressively peel off materials layer by layer."),
 	ECVF_RenderThreadSafe);
 
+static TAutoConsoleVariable<int32> CVarStrataAsyncClassification(
+	TEXT("r.Strata.AsyncClassification"),
+	1,
+	TEXT("Run Strata material classification in async (with shadow)."),
+	ECVF_RenderThreadSafe);
+
 IMPLEMENT_GLOBAL_SHADER_PARAMETER_STRUCT(FStrataGlobalUniformParameters, "Strata");
 
 void FStrataViewData::Reset()
@@ -984,6 +990,9 @@ void AddStrataMaterialClassificationPass(FRDGBuilder& GraphBuilder, const FMinim
 		return;
 	}
 
+	// Optionally run tile classification in async compute
+	const ERDGPassFlags PassFlags = CVarStrataAsyncClassification.GetValueOnRenderThread() > 0 ? ERDGPassFlags::AsyncCompute : ERDGPassFlags::Compute;
+
 	for (int32 i = 0; i < Views.Num(); ++i)
 	{
 		RDG_EVENT_SCOPE_CONDITIONAL(GraphBuilder, Views.Num() > 1, "View%d", i);
@@ -1025,6 +1034,7 @@ void AddStrataMaterialClassificationPass(FRDGBuilder& GraphBuilder, const FMinim
 			FComputeShaderUtils::AddPass(
 				GraphBuilder,
 				RDG_EVENT_NAME("Strata::MaterialTileClassification(%s%s)", bWaveOps ? TEXT("Wave") : TEXT("SharedMemory"), bClear ? TEXT(", Clear") : TEXT("")),
+				PassFlags,
 				ComputeShader,
 				PassParameters,
 				FComputeShaderUtils::GetGroupCount(PassParameters->ViewResolution, GroupSize));
@@ -1040,6 +1050,7 @@ void AddStrataMaterialClassificationPass(FRDGBuilder& GraphBuilder, const FMinim
 			FComputeShaderUtils::AddPass(
 				GraphBuilder,
 				RDG_EVENT_NAME("Strata::MaterialTilePrepareArgs"),
+				PassFlags,
 				ComputeShader,
 				PassParameters,
 				FIntVector(1,1,1));
@@ -1075,6 +1086,7 @@ void AddStrataMaterialClassificationPass(FRDGBuilder& GraphBuilder, const FMinim
 			FComputeShaderUtils::AddPass(
 				GraphBuilder,
 				RDG_EVENT_NAME("Strata::BSDFTileAndOffsets(%s)", bWaveOps ? TEXT("Wave") : TEXT("SharedMemory")),
+				PassFlags,
 				ComputeShader,
 				PassParameters,
 				PassParameters->TileIndirectBuffer,
@@ -1097,6 +1109,7 @@ void AddStrataMaterialClassificationPass(FRDGBuilder& GraphBuilder, const FMinim
 			FComputeShaderUtils::AddPass(
 				GraphBuilder,
 				RDG_EVENT_NAME("Strata::BSDFTilePrepareArgs"),
+				PassFlags,
 				ComputeShader,
 				PassParameters,
 				FIntVector(1, 1, 1));
