Fix for "Distance Cull Fade Node on Material does not work with Skeletal Mesh"

#rnx

diff --git a/Engine/Source/Runtime/Renderer/Private/ShaderBaseClasses.cpp b/Engine/Source/Runtime/Renderer/Private/ShaderBaseClasses.cpp
--- a/Engine/Source/Runtime/Renderer/Private/ShaderBaseClasses.cpp
+++ b/Engine/Source/Runtime/Renderer/Private/ShaderBaseClasses.cpp
@@ -32,12 +32,12 @@ void FMeshMaterialShaderElementData::InitializeMeshMaterialData(const FSceneView
 	FadeUniformBuffer = GDistanceCullFadedInUniformBuffer.GetUniformBufferRHI();
 	DitherUniformBuffer = GDitherFadedInUniformBuffer.GetUniformBufferRHI();
 
-	if (SceneView && StaticMeshId >= 0)
+	if (SceneView)
 	{
 		checkSlow(SceneView->bIsViewInfo);
 		const FViewInfo* ViewInfo = (FViewInfo*)SceneView;
 
-		if (MeshBatch.bDitheredLODTransition && !(bAllowStencilDither && ViewInfo->bAllowStencilDither))
+		if (StaticMeshId >= 0 && MeshBatch.bDitheredLODTransition && !(bAllowStencilDither && ViewInfo->bAllowStencilDither))
 		{
 			if (ViewInfo->StaticMeshFadeOutDitheredLODMap[StaticMeshId])
 			{
