Strata - draw/dispatch complex wave first in order to have longer wave dispatch first, potentially increasing gpu work overlap and reducing long empy work tail.
All of them for good measure.

diff --git a/Engine/Source/Runtime/Renderer/Private/ClusteredDeferredShadingPass.cpp b/Engine/Source/Runtime/Renderer/Private/ClusteredDeferredShadingPass.cpp
--- a/Engine/Source/Runtime/Renderer/Private/ClusteredDeferredShadingPass.cpp
+++ b/Engine/Source/Runtime/Renderer/Private/ClusteredDeferredShadingPass.cpp
@@ -294,10 +294,10 @@ void FDeferredShadingSceneRenderer::AddClusteredDeferredShadingPass(
 					SceneTextures,
 					SortedLightsSet,
 					EClusterPassInputType::Strata,
-					EStrataTileType::ESimple,
+					EStrataTileType::EComplex,
 					ShadowMaskBits,
 					VirtualShadowMapArray,
-					nullptr, 
+					nullptr,
 					StrataSceneData);
 
 				InternalAddClusteredDeferredShadingPass(
@@ -318,10 +318,10 @@ void FDeferredShadingSceneRenderer::AddClusteredDeferredShadingPass(
 					SceneTextures,
 					SortedLightsSet,
 					EClusterPassInputType::Strata,
-					EStrataTileType::EComplex,
+					EStrataTileType::ESimple,
 					ShadowMaskBits,
 					VirtualShadowMapArray,
-					nullptr,
+					nullptr, 
 					StrataSceneData);
 			}
 			else
diff --git a/Engine/Source/Runtime/Renderer/Private/IndirectLightRendering.cpp b/Engine/Source/Runtime/Renderer/Private/IndirectLightRendering.cpp
--- a/Engine/Source/Runtime/Renderer/Private/IndirectLightRendering.cpp
+++ b/Engine/Source/Runtime/Renderer/Private/IndirectLightRendering.cpp
@@ -1221,9 +1221,9 @@ void FDeferredShadingSceneRenderer::RenderDiffuseIndirectAndAmbientOcclusion(
 		{
 			if (Strata::IsStrataEnabled() && ViewPipelineState.DiffuseIndirectMethod == EDiffuseIndirectMethod::Lumen)
 			{
-				ApplyDiffuseIndirect(EStrataTileType::ESimple);
-				ApplyDiffuseIndirect(EStrataTileType::ESingle);
 				ApplyDiffuseIndirect(EStrataTileType::EComplex);
+				ApplyDiffuseIndirect(EStrataTileType::ESingle);
+				ApplyDiffuseIndirect(EStrataTileType::ESimple);
 			}
 			else
 			{
@@ -1728,7 +1728,7 @@ void FDeferredShadingSceneRenderer::RenderDeferredReflectionsAndSkyLighting(
 					bSkyLight,
 					bDynamicSkyLight,
 					bApplySkyShadowing,
-					EStrataTileType::ESimple);
+					EStrataTileType::EComplex);
 
 				AddSkyReflectionPass(
 					GraphBuilder,
@@ -1756,7 +1756,7 @@ void FDeferredShadingSceneRenderer::RenderDeferredReflectionsAndSkyLighting(
 					bSkyLight,
 					bDynamicSkyLight,
 					bApplySkyShadowing,
-					EStrataTileType::EComplex);
+					EStrataTileType::ESimple);
 			}
 			else
 			{
diff --git a/Engine/Source/Runtime/Renderer/Private/Strata/Strata.cpp b/Engine/Source/Runtime/Renderer/Private/Strata/Strata.cpp
--- a/Engine/Source/Runtime/Renderer/Private/Strata/Strata.cpp
+++ b/Engine/Source/Runtime/Renderer/Private/Strata/Strata.cpp
@@ -928,9 +928,9 @@ void AddStrataStencilPass(
 		RDG_EVENT_SCOPE_CONDITIONAL(GraphBuilder, Views.Num() > 1, "View%d", i);
 
 		const FViewInfo& View = Views[i];
-		AddStrataInternalClassificationTilePass(GraphBuilder, View, &SceneTextures.Depth.Target, nullptr, EStrataTileType::ESimple);
-		AddStrataInternalClassificationTilePass(GraphBuilder, View, &SceneTextures.Depth.Target, nullptr, EStrataTileType::ESingle);
 		AddStrataInternalClassificationTilePass(GraphBuilder, View, &SceneTextures.Depth.Target, nullptr, EStrataTileType::EComplex);
+		AddStrataInternalClassificationTilePass(GraphBuilder, View, &SceneTextures.Depth.Target, nullptr, EStrataTileType::ESingle);
+		AddStrataInternalClassificationTilePass(GraphBuilder, View, &SceneTextures.Depth.Target, nullptr, EStrataTileType::ESimple);
 	}
 }
 
diff --git a/Engine/Source/Runtime/Renderer/Private/Strata/StrataVisualize.cpp b/Engine/Source/Runtime/Renderer/Private/Strata/StrataVisualize.cpp
--- a/Engine/Source/Runtime/Renderer/Private/Strata/StrataVisualize.cpp
+++ b/Engine/Source/Runtime/Renderer/Private/Strata/StrataVisualize.cpp
@@ -325,9 +325,9 @@ FScreenPassTexture AddStrataDebugPasses(FRDGBuilder& GraphBuilder, const FViewIn
 		}
 		else
 		{
-			AddStrataInternalClassificationTilePass(GraphBuilder, View, nullptr, &ScreenPassSceneColor.Texture, EStrataTileType::ESimple, bDebugPass);
-			AddStrataInternalClassificationTilePass(GraphBuilder, View, nullptr, &ScreenPassSceneColor.Texture, EStrataTileType::ESingle, bDebugPass);
 			AddStrataInternalClassificationTilePass(GraphBuilder, View, nullptr, &ScreenPassSceneColor.Texture, EStrataTileType::EComplex, bDebugPass);
+			AddStrataInternalClassificationTilePass(GraphBuilder, View, nullptr, &ScreenPassSceneColor.Texture, EStrataTileType::ESingle, bDebugPass);
+			AddStrataInternalClassificationTilePass(GraphBuilder, View, nullptr, &ScreenPassSceneColor.Texture, EStrataTileType::ESimple, bDebugPass);
 		}
 	}
 
