- Fix for OOB access of PrimitiveFadeUniformBufferMap

#rnx

diff --git a/Engine/Source/Runtime/Renderer/Private/ShaderBaseClasses.cpp b/Engine/Source/Runtime/Renderer/Private/ShaderBaseClasses.cpp
--- a/Engine/Source/Runtime/Renderer/Private/ShaderBaseClasses.cpp
+++ b/Engine/Source/Runtime/Renderer/Private/ShaderBaseClasses.cpp
@@ -53,7 +53,7 @@ void FMeshMaterialShaderElementData::InitializeMeshMaterialData(const FSceneView
 		{
 			int32 const PrimitiveIndex = PrimitiveSceneProxy->GetPrimitiveSceneInfo()->GetIndex();
 
-			if (ViewInfo->PrimitiveFadeUniformBufferMap[PrimitiveIndex])
+			if (ViewInfo->PrimitiveFadeUniformBufferMap.IsValidIndex(PrimitiveIndex) && ViewInfo->PrimitiveFadeUniformBufferMap[PrimitiveIndex])
 			{
 				FadeUniformBuffer = ViewInfo->PrimitiveFadeUniformBuffers[PrimitiveIndex];
 			}
