Fix a bug of bloom intensity been applied twice.

diff --git a/Engine/Shaders/Private/PostProcessMobile.usf b/Engine/Shaders/Private/PostProcessMobile.usf
--- a/Engine/Shaders/Private/PostProcessMobile.usf
+++ b/Engine/Shaders/Private/PostProcessMobile.usf
@@ -1052,7 +1052,6 @@ void SunMergeVS_Mobile(
 
 float4 SunColorVignetteIntensity;
 float3 BloomColor;
-float BloomIntensity;
 
 Texture2D BloomDirtMaskTexture;
 SamplerState BloomDirtMaskSampler;
@@ -1110,8 +1109,6 @@ void SunMergePS_Mobile(
 		OutColor.rgb += SunColorVignetteIntensity.rgb * Sun;
 	#endif
 
-	OutColor.rgb *= BloomIntensity;
-
 	OutColor.a = 1.0f;
 }
 
diff --git a/Engine/Source/Runtime/Renderer/Private/PostProcess/PostProcessMobile.cpp b/Engine/Source/Runtime/Renderer/Private/PostProcess/PostProcessMobile.cpp
--- a/Engine/Source/Runtime/Renderer/Private/PostProcess/PostProcessMobile.cpp
+++ b/Engine/Source/Runtime/Renderer/Private/PostProcess/PostProcessMobile.cpp
@@ -1032,7 +1032,6 @@ public:
 		SHADER_PARAMETER(FVector4f, BloomDirtMaskTint)
 		SHADER_PARAMETER(FVector4f, SunColorVignetteIntensity)
 		SHADER_PARAMETER(FVector3f, BloomColor)
-		SHADER_PARAMETER(float, BloomIntensity)
 		SHADER_PARAMETER_STRUCT_REF(FViewUniformShaderParameters, View)
 		SHADER_PARAMETER_RDG_TEXTURE(Texture2D, SunBlurTexture)
 		SHADER_PARAMETER_SAMPLER(SamplerState, SunBlurSampler)
@@ -1140,7 +1139,6 @@ FScreenPassTexture AddMobileSunMergePass(FRDGBuilder& GraphBuilder, const FViewI
 	PSShaderParameters->BloomDirtMaskTint = Settings.BloomDirtMaskTint * Settings.BloomDirtMaskIntensity;
 	PSShaderParameters->SunColorVignetteIntensity = SunColorVignetteIntensityParam;
 	PSShaderParameters->BloomColor = FVector3f(BloomColor.R, BloomColor.G, BloomColor.B);
-	PSShaderParameters->BloomIntensity = Settings.BloomIntensity;
 	PSShaderParameters->SunBlurTexture = Inputs.SunBlur.Texture;
 	PSShaderParameters->SunBlurSampler = TStaticSamplerState<SF_Bilinear, AM_Clamp, AM_Clamp, AM_Clamp>::GetRHI();
 	PSShaderParameters->BloomSetup_BloomTexture = Inputs.BloomSetup_Bloom.Texture;
