Make FLightSceneInfo not derive FRenderResource as this is no longer needed (post RDG conversion of capsure shadows, CL 16971682)

diff --git a/Engine/Source/Runtime/Renderer/Private/LightSceneInfo.cpp b/Engine/Source/Runtime/Renderer/Private/LightSceneInfo.cpp
--- a/Engine/Source/Runtime/Renderer/Private/LightSceneInfo.cpp
+++ b/Engine/Source/Runtime/Renderer/Private/LightSceneInfo.cpp
@@ -65,12 +65,6 @@ FLightSceneInfo::FLightSceneInfo(FLightSceneProxy* InProxy, bool InbVisible)
 	// Only visible lights can be added in game
 	check(bVisible || GIsEditor);
 
-	BeginInitResource(this);
-}
-
-FLightSceneInfo::~FLightSceneInfo()
-{
-	ReleaseResource();
 }
 
 void FLightSceneInfo::AddToScene()
diff --git a/Engine/Source/Runtime/Renderer/Private/LightSceneInfo.h b/Engine/Source/Runtime/Renderer/Private/LightSceneInfo.h
--- a/Engine/Source/Runtime/Renderer/Private/LightSceneInfo.h
+++ b/Engine/Source/Runtime/Renderer/Private/LightSceneInfo.h
@@ -7,7 +7,6 @@
 #pragma once
 
 #include "CoreMinimal.h"
-#include "RenderResource.h"
 #include "Math/GenericOctreePublic.h"
 #include "SceneManagement.h"
 #include "Math/GenericOctree.h"
@@ -176,7 +175,7 @@ public:
  * The information used to render a light.  This is the rendering thread's mirror of the game thread's ULightComponent.
  * FLightSceneInfo is internal to the renderer module and contains internal scene state.
  */
-class FLightSceneInfo : public FRenderResource
+class FLightSceneInfo
 {
 	friend class FLightPrimitiveInteraction;
 
@@ -251,7 +250,6 @@ public:
 
 	/** Initialization constructor. */
 	FLightSceneInfo(FLightSceneProxy* InProxy, bool InbVisible);
-	virtual ~FLightSceneInfo();
 
 	/** Adds the light to the scene. */
 	void AddToScene();
