Fix for translucent shadow stencil being broken with Strata due to stencil not being in a correct start from the start.

diff --git a/Engine/Source/Runtime/Renderer/Private/Strata/Strata.cpp b/Engine/Source/Runtime/Renderer/Private/Strata/Strata.cpp
--- a/Engine/Source/Runtime/Renderer/Private/Strata/Strata.cpp
+++ b/Engine/Source/Runtime/Renderer/Private/Strata/Strata.cpp
@@ -789,9 +789,9 @@ uint32 TileTypeDispatchIndirectArgOffset(const EStrataTileType Type)
 
 // Add additionnaly bits for filling/clearing stencil to ensure that the 'Strata' bits are not corrupted by the stencil shadows 
 // when generating shadow mask. Withouth these 'trailing' bits, the incr./decr. operation would change/corrupt the 'Strata' bits
-constexpr uint32 StencilBit_Fast_1	  = 0x07u | StencilBit_Fast;
-constexpr uint32 StencilBit_Single_1  = 0x07u | StencilBit_Single;
-constexpr uint32 StencilBit_Complex_1 = 0x07u | StencilBit_Complex; 
+constexpr uint32 StencilBit_Fast_1	  = StencilBit_Fast;
+constexpr uint32 StencilBit_Single_1  = StencilBit_Single;
+constexpr uint32 StencilBit_Complex_1 = StencilBit_Complex; 
 
 void AddStrataInternalClassificationTilePass(
 	FRDGBuilder& GraphBuilder,
