Uses a Texture2DArray for Debug.TSR.* textures

diff --git a/Engine/Shaders/Private/TemporalSuperResolution/TSRCommon.ush b/Engine/Shaders/Private/TemporalSuperResolution/TSRCommon.ush
--- a/Engine/Shaders/Private/TemporalSuperResolution/TSRCommon.ush
+++ b/Engine/Shaders/Private/TemporalSuperResolution/TSRCommon.ush
@@ -305,7 +305,7 @@ float HistoryPreExposureCorrection;
 uint bCameraCut;
 
 #if DEBUG_OUTPUT
-	RWTexture2D<float4> DebugOutput;
+	RWTexture2DArray<float4> DebugOutput;
 #endif
 
 
diff --git a/Engine/Shaders/Private/TemporalSuperResolution/TSRCompareHistory.usf b/Engine/Shaders/Private/TemporalSuperResolution/TSRCompareHistory.usf
--- a/Engine/Shaders/Private/TemporalSuperResolution/TSRCompareHistory.usf
+++ b/Engine/Shaders/Private/TemporalSuperResolution/TSRCompareHistory.usf
@@ -137,7 +137,7 @@ void MainCS(
 	#if DEBUG_OUTPUT
 	{
 		tsr_short2 OutputPixelPos = InvalidateOutputPixelPos(tsr_ushort2(InputPixelPos), InputPixelPosMax);
-		DebugOutput[OutputPixelPos] = Debug;
+		DebugOutput[tsr_short3(OutputPixelPos, 0)] = Debug;
 	}
 	#endif
 }
diff --git a/Engine/Shaders/Private/TemporalSuperResolution/TSRCompareTranslucency.usf b/Engine/Shaders/Private/TemporalSuperResolution/TSRCompareTranslucency.usf
--- a/Engine/Shaders/Private/TemporalSuperResolution/TSRCompareTranslucency.usf
+++ b/Engine/Shaders/Private/TemporalSuperResolution/TSRCompareTranslucency.usf
@@ -259,7 +259,7 @@ void MainCS(
 
 		#if DEBUG_OUTPUT
 		{
-			DebugOutput[InputPixelPos] = Debug;
+			DebugOutput[tsr_short3(InputPixelPos, 0)] = Debug;
 		}
 		#endif
 	}
diff --git a/Engine/Shaders/Private/TemporalSuperResolution/TSRDebugHistory.usf b/Engine/Shaders/Private/TemporalSuperResolution/TSRDebugHistory.usf
--- a/Engine/Shaders/Private/TemporalSuperResolution/TSRDebugHistory.usf
+++ b/Engine/Shaders/Private/TemporalSuperResolution/TSRDebugHistory.usf
@@ -108,7 +108,7 @@ void MainCS(
 	#if DEBUG_OUTPUT && 1
 	if (all(DispatchThreadId < HistoryInfo_ViewportMax * SUB_PIXEL_GRID_SIZE))
 	{
-		DebugOutput[DispatchThreadId] = Debug;
+		DebugOutput[tsr_short3(DispatchThreadId, 0)] = Debug;
 	}
 	#endif
 }
diff --git a/Engine/Shaders/Private/TemporalSuperResolution/TSRDecimateHistory.usf b/Engine/Shaders/Private/TemporalSuperResolution/TSRDecimateHistory.usf
--- a/Engine/Shaders/Private/TemporalSuperResolution/TSRDecimateHistory.usf
+++ b/Engine/Shaders/Private/TemporalSuperResolution/TSRDecimateHistory.usf
@@ -737,7 +737,7 @@ void MainCS(
 
 		#if DEBUG_OUTPUT
 		{
-			DebugOutput[OutputPixelPos] = Debug;
+			DebugOutput[tsr_short3(OutputPixelPos, 0)] = Debug;
 		}
 		#endif
 	}
diff --git a/Engine/Shaders/Private/TemporalSuperResolution/TSRDilateRejection.usf b/Engine/Shaders/Private/TemporalSuperResolution/TSRDilateRejection.usf
--- a/Engine/Shaders/Private/TemporalSuperResolution/TSRDilateRejection.usf
+++ b/Engine/Shaders/Private/TemporalSuperResolution/TSRDilateRejection.usf
@@ -65,7 +65,7 @@ void MainCS(
 
 	#if DEBUG_OUTPUT
 	{
-		DebugOutput[RejectionPixelPos] = Debug;
+		DebugOutput[tsr_short3(RejectionPixelPos, 0)] = Debug;
 	}
 	#endif
 }
diff --git a/Engine/Shaders/Private/TemporalSuperResolution/TSRDilateVelocity.usf b/Engine/Shaders/Private/TemporalSuperResolution/TSRDilateVelocity.usf
--- a/Engine/Shaders/Private/TemporalSuperResolution/TSRDilateVelocity.usf
+++ b/Engine/Shaders/Private/TemporalSuperResolution/TSRDilateVelocity.usf
@@ -281,7 +281,7 @@ void MainCS(
 
 		#if DEBUG_OUTPUT && 0 // running out of UAVs
 		{
-			DebugOutput[OutputPixelPos] = Debug;
+			DebugOutput[tsr_short3(OutputPixelPos, 0)] = Debug;
 		}
 		#endif
 	}
diff --git a/Engine/Shaders/Private/TemporalSuperResolution/TSRFilterAntiAliasing.usf b/Engine/Shaders/Private/TemporalSuperResolution/TSRFilterAntiAliasing.usf
--- a/Engine/Shaders/Private/TemporalSuperResolution/TSRFilterAntiAliasing.usf
+++ b/Engine/Shaders/Private/TemporalSuperResolution/TSRFilterAntiAliasing.usf
@@ -119,7 +119,7 @@ void MainCS(
 		
 		#if DEBUG_OUTPUT
 		{
-			DebugOutput[OutputPixelPos] = Debug;
+			DebugOutput[tsr_short3(OutputPixelPos, 0)] = Debug;
 		}
 		#endif
 	}
diff --git a/Engine/Shaders/Private/TemporalSuperResolution/TSRFilterFrequencies.usf b/Engine/Shaders/Private/TemporalSuperResolution/TSRFilterFrequencies.usf
--- a/Engine/Shaders/Private/TemporalSuperResolution/TSRFilterFrequencies.usf
+++ b/Engine/Shaders/Private/TemporalSuperResolution/TSRFilterFrequencies.usf
@@ -347,7 +347,7 @@ void MainCS(
 
 	#if DEBUG_OUTPUT
 	{
-		DebugOutput[OutputPixelPos - LowFrequencyInfo_ViewportMin] = Debug;
+		DebugOutput[tsr_short3(OutputPixelPos - LowFrequencyInfo_ViewportMin, 0)] = Debug;
 	}
 	#endif
 }
diff --git a/Engine/Shaders/Private/TemporalSuperResolution/TSRPostfilterRejection.usf b/Engine/Shaders/Private/TemporalSuperResolution/TSRPostfilterRejection.usf
--- a/Engine/Shaders/Private/TemporalSuperResolution/TSRPostfilterRejection.usf
+++ b/Engine/Shaders/Private/TemporalSuperResolution/TSRPostfilterRejection.usf
@@ -133,7 +133,7 @@ void MainCS(
 	#if DEBUG_OUTPUT
 	{
 		tsr_short2 OutputPixelPos = InvalidateOutputPixelPos(RejectionPixelPos, HistoryRejectionViewport.zw);
-		DebugOutput[OutputPixelPos] = Debug;
+		DebugOutput[tsr_short3(OutputPixelPos, 0)] = Debug;
 	}
 	#endif
 }
diff --git a/Engine/Shaders/Private/TemporalSuperResolution/TSRResolveHistory.usf b/Engine/Shaders/Private/TemporalSuperResolution/TSRResolveHistory.usf
--- a/Engine/Shaders/Private/TemporalSuperResolution/TSRResolveHistory.usf
+++ b/Engine/Shaders/Private/TemporalSuperResolution/TSRResolveHistory.usf
@@ -242,7 +242,7 @@ void MainCS(
 
 	#if DEBUG_OUTPUT
 	{
-		DebugOutput[OutputPixelPosition] = Debug;
+		DebugOutput[tsr_short3(OutputPixelPosition, 0)] = Debug;
 	}
 	#endif
 }
diff --git a/Engine/Shaders/Private/TemporalSuperResolution/TSRSpatialAntiAliasing.usf b/Engine/Shaders/Private/TemporalSuperResolution/TSRSpatialAntiAliasing.usf
--- a/Engine/Shaders/Private/TemporalSuperResolution/TSRSpatialAntiAliasing.usf
+++ b/Engine/Shaders/Private/TemporalSuperResolution/TSRSpatialAntiAliasing.usf
@@ -510,10 +510,10 @@ void MainCS(
 
 		#if DEBUG_OUTPUT
 		{
-			DebugOutput[dpv_lo(OutputPixelPos)] = dpv_lo(Debug);
+			DebugOutput[tsr_short3(dpv_lo(OutputPixelPos), 0)] = dpv_lo(Debug);
 
 			#if CONFIG_ENABLE_DUAL_PIXEL_VECTORIZATION
-				DebugOutput[dpv_hi(OutputPixelPos)] = dpv_hi(Debug);
+				DebugOutput[tsr_short3(dpv_hi(OutputPixelPos), 0)] = dpv_hi(Debug);
 			#endif
 		}
 		#endif
diff --git a/Engine/Shaders/Private/TemporalSuperResolution/TSRUpdateHistory.usf b/Engine/Shaders/Private/TemporalSuperResolution/TSRUpdateHistory.usf
--- a/Engine/Shaders/Private/TemporalSuperResolution/TSRUpdateHistory.usf
+++ b/Engine/Shaders/Private/TemporalSuperResolution/TSRUpdateHistory.usf
@@ -9,6 +9,8 @@
 
 #define DEBUG_LOW_FREQUENCY_REJECTION 0
 
+#define DEBUG_ARRAY_SIZE 8
+
 
 //------------------------------------------------------- CONFIG
 
@@ -310,8 +312,14 @@ void MainCS(
 	#else
 		uint GroupWaveIndex = GetGroupWaveIndex(GroupThreadIndex, TILE_SIZE * TILE_SIZE);
 	#endif
-
-	float4x2 Debug = 0.0;
+	
+	#if DEBUG_OUTPUT
+	float4x2 Debug[DEBUG_ARRAY_SIZE];
+	for (uint DebugId = 0; DebugId < DEBUG_ARRAY_SIZE; DebugId++)
+	{
+		Debug[DebugId] = 0.0;
+	}
+	#endif
 
 	tsr_short2x2 HistoryPixelPos = dpv_add(
 		tsr_short2(GroupId) * tsr_short2(TILE_SIZE, TILE_SIZE),
@@ -1152,14 +1160,6 @@ void MainCS(
 			dpv_scale(PrevHighFrequencyColor, CommonWeight * BlendPrevHistory * InvPrevHistoryWeight) +
 			dpv_scale(FilteredInputColor, CommonWeight * BlendFilteredInput));
 				
-				
-		//Debug[0] = CommonWeight * BlendFilteredInput;
-		Debug[0] = OutputValidity2;
-		//Debug[0] = CurrentMomentWeight;
-		//Debug[0] = FilteredInputColor[0];
-		//Debug[1] = FilteredInputColor[1];
-		//Debug[2] = FilteredInputColor[2];
-
 		// Quantize validity for the 8bit encoding to avoid numerical shift between color and validity.
 		FinalHistoryValidity = ceil(tsr_half(255.0) * OutputValidity) * rcp(tsr_half(255.0));
 	}
@@ -1848,13 +1848,14 @@ void MainCS(
 			#endif
 		}
 		#endif
-
+		
 		#if DEBUG_OUTPUT
+		for (uint DebugId = 0; DebugId < DEBUG_ARRAY_SIZE; DebugId++)
 		{
-			DebugOutput[dpv_lo(LocalHistoryPixelPos)] = dpv_lo(Debug);
+			DebugOutput[tsr_short3(dpv_lo(LocalHistoryPixelPos), DebugId)] = dpv_lo(Debug[DebugId]);
 
 			#if CONFIG_ENABLE_DUAL_PIXEL_VECTORIZATION
-				DebugOutput[dpv_hi(LocalHistoryPixelPos)] = dpv_hi(Debug);
+				DebugOutput[tsr_short3(dpv_hi(LocalHistoryPixelPos), DebugId)] = dpv_hi(Debug[DebugId]);
 			#endif
 		}
 		#endif
diff --git a/Engine/Source/Runtime/Renderer/Private/PostProcess/TemporalSuperResolution.cpp b/Engine/Source/Runtime/Renderer/Private/PostProcess/TemporalSuperResolution.cpp
--- a/Engine/Source/Runtime/Renderer/Private/PostProcess/TemporalSuperResolution.cpp
+++ b/Engine/Source/Runtime/Renderer/Private/PostProcess/TemporalSuperResolution.cpp
@@ -87,6 +87,11 @@ TAutoConsoleVariable<float> CVarTSRVelocityExtrapolation(
 
 #if COMPILE_TSR_DEBUG_PASSES
 
+TAutoConsoleVariable<int32> CVarTSRDebugArraySize(
+	TEXT("r.TSR.Debug.ArraySize"), 1,
+	TEXT("Size of array for the TSR.Debug.* RDG textures"),
+	ECVF_RenderThreadSafe);
+
 TAutoConsoleVariable<int32> CVarTSRSetupDebugPasses(
 	TEXT("r.TSR.Debug.SetupExtraPasses"), 0,
 	TEXT("Whether to enable the debug passes"),
@@ -287,7 +292,7 @@ class FTSRDilateVelocityCS : public FTSRShader
 		SHADER_PARAMETER_RDG_TEXTURE_UAV(RWTexture2D, VelocityTileOutput0)
 		SHADER_PARAMETER_RDG_TEXTURE_UAV(RWTexture2D, VelocityTileOutput1)
 
-		SHADER_PARAMETER_RDG_TEXTURE_UAV(RWTexture2D, DebugOutput)
+		SHADER_PARAMETER_RDG_TEXTURE_UAV(RWTexture2DArray, DebugOutput)
 	END_SHADER_PARAMETER_STRUCT()
 }; // class FTSRDilateVelocityCS
 
@@ -327,7 +332,7 @@ class FTSRDecimateHistoryCS : public FTSRShader
 		SHADER_PARAMETER_RDG_TEXTURE_UAV(RWTexture2D, ParallaxRejectionMaskOutput)
 		SHADER_PARAMETER_RDG_TEXTURE_UAV(RWTexture2D, HistorySubpixelDetailsOutput)
 
-		SHADER_PARAMETER_RDG_TEXTURE_UAV(RWTexture2D, DebugOutput)
+		SHADER_PARAMETER_RDG_TEXTURE_UAV(RWTexture2DArray, DebugOutput)
 	END_SHADER_PARAMETER_STRUCT()
 }; // class FTSRDecimateHistoryCS
 
@@ -351,7 +356,7 @@ class FTSRCompareTranslucencyCS : public FTSRShader
 		SHADER_PARAMETER_RDG_TEXTURE(Texture2D, PrevTranslucencyTexture)
 
 		SHADER_PARAMETER_RDG_TEXTURE_UAV(RWTexture2D, TranslucencyRejectionOutput)
-		SHADER_PARAMETER_RDG_TEXTURE_UAV(RWTexture2D, DebugOutput)
+		SHADER_PARAMETER_RDG_TEXTURE_UAV(RWTexture2DArray, DebugOutput)
 	END_SHADER_PARAMETER_STRUCT()
 }; // class FTSRCompareTranslucencyCS
 
@@ -375,7 +380,7 @@ class FTSRFilterFrequenciesCS : public FTSRShader
 		SHADER_PARAMETER_RDG_TEXTURE_UAV(RWTexture2D, FilteredInputOutput)
 		SHADER_PARAMETER_RDG_TEXTURE_UAV(RWTexture2D, FilteredPredictionSceneColorOutput)
 		SHADER_PARAMETER_RDG_TEXTURE_UAV(RWTexture2D, InputSceneColorLdrLumaOutput)
-		SHADER_PARAMETER_RDG_TEXTURE_UAV(RWTexture2D, DebugOutput)
+		SHADER_PARAMETER_RDG_TEXTURE_UAV(RWTexture2DArray, DebugOutput)
 	END_SHADER_PARAMETER_STRUCT()
 }; // class FTSRFilterFrequenciesCS
 
@@ -393,7 +398,7 @@ class FTSRCompareHistoryCS : public FTSRShader
 		SHADER_PARAMETER_RDG_TEXTURE(Texture2D, FilteredInputTexture)
 		SHADER_PARAMETER_RDG_TEXTURE(Texture2D, FilteredPredictionSceneColorTexture)
 		SHADER_PARAMETER_RDG_TEXTURE_UAV(RWTexture2D, HistoryRejectionOutput)
-		SHADER_PARAMETER_RDG_TEXTURE_UAV(RWTexture2D, DebugOutput)
+		SHADER_PARAMETER_RDG_TEXTURE_UAV(RWTexture2DArray, DebugOutput)
 	END_SHADER_PARAMETER_STRUCT()
 }; // class FTSRCompareHistoryCS
 
@@ -409,7 +414,7 @@ class FTSRPostfilterRejectionCS : public FTSRShader
 		SHADER_PARAMETER(FIntRect, HistoryRejectionViewport)
 		SHADER_PARAMETER_RDG_TEXTURE(Texture2D, HistoryRejectionTexture)
 		SHADER_PARAMETER_RDG_TEXTURE_UAV(RWTexture2D, HistoryRejectionOutput)
-		SHADER_PARAMETER_RDG_TEXTURE_UAV(RWTexture2D, DebugOutput)
+		SHADER_PARAMETER_RDG_TEXTURE_UAV(RWTexture2DArray, DebugOutput)
 	END_SHADER_PARAMETER_STRUCT()
 }; // class FTSRPostfilterRejectionCS
 
@@ -422,7 +427,7 @@ class FTSRDilateRejectionCS : public FTSRShader
 		SHADER_PARAMETER_STRUCT_INCLUDE(FTSRCommonParameters, CommonParameters)
 		SHADER_PARAMETER_RDG_TEXTURE(Texture2D, HistoryRejectionTexture)
 		SHADER_PARAMETER_RDG_TEXTURE_UAV(RWTexture2D, DilatedHistoryRejectionOutput)
-		SHADER_PARAMETER_RDG_TEXTURE_UAV(RWTexture2D, DebugOutput)
+		SHADER_PARAMETER_RDG_TEXTURE_UAV(RWTexture2DArray, DebugOutput)
 	END_SHADER_PARAMETER_STRUCT()
 }; // class FTSRDilateRejectionCS
 
@@ -441,7 +446,7 @@ class FTSRSpatialAntiAliasingCS : public FTSRShader
 		SHADER_PARAMETER_RDG_TEXTURE(Texture2D, InputSceneColorLdrLumaTexture)
 		SHADER_PARAMETER_RDG_TEXTURE_UAV(RWTexture2D, AntiAliasingOutput)
 		SHADER_PARAMETER_RDG_TEXTURE_UAV(RWTexture2D, NoiseFilteringOutput)
-		SHADER_PARAMETER_RDG_TEXTURE_UAV(RWTexture2D, DebugOutput)
+		SHADER_PARAMETER_RDG_TEXTURE_UAV(RWTexture2DArray, DebugOutput)
 	END_SHADER_PARAMETER_STRUCT()
 
 	static bool ShouldCompilePermutation(const FGlobalShaderPermutationParameters& Parameters)
@@ -473,7 +478,7 @@ class FTSRFilterAntiAliasingCS : public FTSRShader
 		SHADER_PARAMETER_STRUCT_INCLUDE(FTSRCommonParameters, CommonParameters)
 		SHADER_PARAMETER_RDG_TEXTURE(Texture2D, AntiAliasingTexture)
 		SHADER_PARAMETER_RDG_TEXTURE_UAV(RWTexture2D, AntiAliasingOutput)
-		SHADER_PARAMETER_RDG_TEXTURE_UAV(RWTexture2D, DebugOutput)
+		SHADER_PARAMETER_RDG_TEXTURE_UAV(RWTexture2DArray, DebugOutput)
 	END_SHADER_PARAMETER_STRUCT()
 }; // class FTSRFilterAntiAliasingCS
 
@@ -527,7 +532,7 @@ class FTSRUpdateHistoryCS : public FTSRShader
 		SHADER_PARAMETER_RDG_TEXTURE_UAV(RWTexture2D, SceneColorOutputMip0)
 		SHADER_PARAMETER_RDG_TEXTURE_UAV(RWTexture2D, SceneColorOutputMip1)
 		SHADER_PARAMETER_STRUCT(FTSRHistoryUAVs, HistoryOutput)
-		SHADER_PARAMETER_RDG_TEXTURE_UAV(RWTexture2D, DebugOutput)
+		SHADER_PARAMETER_RDG_TEXTURE_UAV(RWTexture2DArray, DebugOutput)
 	END_SHADER_PARAMETER_STRUCT()
 
 	static void ModifyCompilationEnvironment(const FGlobalShaderPermutationParameters& Parameters, FShaderCompilerEnvironment& OutEnvironment)
@@ -556,7 +561,7 @@ class FTSRResolveHistoryCS : public FTSRShader
 
 		SHADER_PARAMETER_RDG_TEXTURE_UAV(RWTexture2D, SceneColorOutputMip0)
 		SHADER_PARAMETER_RDG_TEXTURE_UAV(RWTexture2D, SceneColorOutputMip1)
-		SHADER_PARAMETER_RDG_TEXTURE_UAV(RWTexture2D, DebugOutput)
+		SHADER_PARAMETER_RDG_TEXTURE_UAV(RWTexture2DArray, DebugOutput)
 	END_SHADER_PARAMETER_STRUCT()
 
 	static void ModifyCompilationEnvironment(const FGlobalShaderPermutationParameters& Parameters, FShaderCompilerEnvironment& OutEnvironment)
@@ -577,7 +582,7 @@ class FTSRDebugHistoryCS : public FTSRShader
 		SHADER_PARAMETER_STRUCT_INCLUDE(FTSRCommonParameters, CommonParameters)
 		SHADER_PARAMETER_STRUCT(FTSRHistorySRVs, History)
 		SHADER_PARAMETER_STRUCT(FTSRHistorySRVs, PrevHistory)
-		SHADER_PARAMETER_RDG_TEXTURE_UAV(RWTexture2D, DebugOutput)
+		SHADER_PARAMETER_RDG_TEXTURE_UAV(RWTexture2DArray, DebugOutput)
 	END_SHADER_PARAMETER_STRUCT()
 }; // class FTSRDebugHistoryCS
 
@@ -796,11 +801,18 @@ ITemporalUpscaler::FOutputs AddTemporalSuperResolutionPasses(
 
 	auto CreateDebugUAV = [&](const FIntPoint& Extent, const TCHAR* DebugName)
 	{
-		FRDGTextureDesc DebugDesc = FRDGTextureDesc::Create2D(
+#if COMPILE_TSR_DEBUG_PASSES
+		uint16 ArraySize = uint16(FMath::Clamp(CVarTSRDebugArraySize.GetValueOnRenderThread(), 1, GMaxTextureArrayLayers));
+#else
+		const uint16 ArraySize = 1;
+#endif
+
+		FRDGTextureDesc DebugDesc = FRDGTextureDesc::Create2DArray(
 			Extent,
 			PF_FloatRGBA,
 			FClearValueBinding::None,
-			/* InFlags = */ TexCreate_ShaderResource | TexCreate_UAV);
+			/* InFlags = */ TexCreate_ShaderResource | TexCreate_UAV,
+			ArraySize);
 
 		FRDGTextureRef DebugTexture = GraphBuilder.CreateTexture(DebugDesc, DebugName);
 
