Scene Captures are going to be used as textures and sampled in the main scene, so they must remain unflipped.

diff --git a/Engine/Source/Runtime/Renderer/Private/PostProcess/PostProcessing.cpp b/Engine/Source/Runtime/Renderer/Private/PostProcess/PostProcessing.cpp
--- a/Engine/Source/Runtime/Renderer/Private/PostProcess/PostProcessing.cpp
+++ b/Engine/Source/Runtime/Renderer/Private/PostProcess/PostProcessing.cpp
@@ -1676,6 +1676,12 @@ static bool IsGaussianActive(const FViewInfo& View)
 	return true;
 }
 
+// Scene Captures are going to be used as textures and sampled in the main scene, so they must remain unflipped.
+bool IsMobilePostProcessingFlipEnabled(const FViewInfo& View)
+{
+	return RHINeedsToSwitchVerticalAxis(View.GetShaderPlatform()) && !View.bIsSceneCapture;
+}
+
 void AddMobilePostProcessingPasses(FRDGBuilder& GraphBuilder, FScene* Scene, const FViewInfo& View, const FMobilePostProcessingInputs& Inputs, FInstanceCullingManager& InstanceCullingManager)
 {
 	CSV_SCOPED_TIMING_STAT_EXCLUSIVE(RenderPostProcessing);
@@ -1825,7 +1831,7 @@ void AddMobilePostProcessingPasses(FRDGBuilder& GraphBuilder, FScene* Scene, con
 
 		PostProcessMaterialInputs.CustomDepthTexture = CustomDepth.Texture;
 
-		PostProcessMaterialInputs.bFlipYAxis = RHINeedsToSwitchVerticalAxis(View.GetShaderPlatform()) && bLastPass;
+		PostProcessMaterialInputs.bFlipYAxis = IsMobilePostProcessingFlipEnabled(View) && bLastPass;
 
 		PostProcessMaterialInputs.bMetalMSAAHDRDecode = bMetalMSAAHDRDecode;
 
@@ -2296,7 +2302,7 @@ void AddMobilePostProcessingPasses(FRDGBuilder& GraphBuilder, FScene* Scene, con
 		TonemapperInputs.EyeAdaptationTexture = nullptr;
 		TonemapperInputs.ColorGradingTexture = ColorGradingTexture;
 		TonemapperInputs.bWriteAlphaChannel = View.AntiAliasingMethod == AAM_FXAA || IsPostProcessingWithAlphaChannelSupported() || bUseMobileDof || IsMobilePropagateAlphaEnabled(View.GetShaderPlatform());
-		TonemapperInputs.bFlipYAxis = RHINeedsToSwitchVerticalAxis(View.GetShaderPlatform()) && !PassSequence.IsEnabled(EPass::PostProcessMaterialAfterTonemapping);
+		TonemapperInputs.bFlipYAxis = IsMobilePostProcessingFlipEnabled(View) && !PassSequence.IsEnabled(EPass::PostProcessMaterialAfterTonemapping);
 		TonemapperInputs.bOutputInHDR = bHDRTonemapperOutput;
 		TonemapperInputs.bGammaOnly = bDoGammaOnly;
 		TonemapperInputs.bMetalMSAAHDRDecode = bMetalMSAAHDRDecode;
