2 static analysis fixes and re-add needed initialization code that was mistakenly removed

diff --git a/Engine/Plugins/Runtime/HairStrands/Source/HairStrandsCore/Private/GroomBuilder.cpp b/Engine/Plugins/Runtime/HairStrands/Source/HairStrandsCore/Private/GroomBuilder.cpp
--- a/Engine/Plugins/Runtime/HairStrands/Source/HairStrandsCore/Private/GroomBuilder.cpp
+++ b/Engine/Plugins/Runtime/HairStrands/Source/HairStrandsCore/Private/GroomBuilder.cpp
@@ -1956,6 +1956,11 @@ void Decimate(
 
 		Random.Initialize(0xdeedbeed);
 
+		//initialize value
+		for (uint32 CurveIndex = 0; CurveIndex < InCurveCount; CurveIndex++)
+		{
+			CurveRandomizedIndex[CurveIndex] = CurveIndex;
+		}
 		//shuffle
 		for (uint32 CurveIndex = 0; CurveIndex < InCurveCount; CurveIndex++)
 		{
diff --git a/Engine/Source/Runtime/Renderer/Private/HairStrands/HairStrandsInterface.cpp b/Engine/Source/Runtime/Renderer/Private/HairStrands/HairStrandsInterface.cpp
--- a/Engine/Source/Runtime/Renderer/Private/HairStrands/HairStrandsInterface.cpp
+++ b/Engine/Source/Runtime/Renderer/Private/HairStrands/HairStrandsInterface.cpp
@@ -491,7 +491,7 @@ uint32 GetHairVisibilityComputeRasterVertexStart(uint32 TemporalAASampleIndex, u
 	if (IsHairVisibilityComputeRasterTemporalLayeringEnabled())
 	{
 		const uint32 LayerCount = GetHairVisibilityComputeRasterTemporalLayerCount();
-		const uint32 OverrideIndex = CVarHairStrandsVisibilityComputeRaster_TemporalLayering_OverrideIndex.GetValueOnAnyThread();
+		const int32 OverrideIndex = CVarHairStrandsVisibilityComputeRaster_TemporalLayering_OverrideIndex.GetValueOnAnyThread();
 
 		const uint32 VertexCount = InVertexCount / LayerCount;
 
diff --git a/Engine/Source/Runtime/Renderer/Private/HairStrands/HairStrandsVisibility.cpp b/Engine/Source/Runtime/Renderer/Private/HairStrands/HairStrandsVisibility.cpp
--- a/Engine/Source/Runtime/Renderer/Private/HairStrands/HairStrandsVisibility.cpp
+++ b/Engine/Source/Runtime/Renderer/Private/HairStrands/HairStrandsVisibility.cpp
@@ -286,11 +286,9 @@ float GetHairStrandsFullCoverageThreshold()
 	return FMath::Clamp(GHairStrandsFullCoverageThreshold, 0.1f, 1.f);
 }
 
-bool IsHairVisibilityComputeRasterTemporalLayeringEnabled();
-
 uint32 GetHairStrandsIntCoverageThreshold()
 {
-	return (IsHairVisibilityComputeRasterEnabled() ? ((IsHairVisibilityComputeRasterTemporalLayeringEnabled()) ? 1000u : 1000u) : 1000u);
+	return 1000u;
 }
 
 ///////////////////////////////////////////////////////////////////////////////////////////////////
