VulkanRHI: Match the behavior of other IndrectDraw arg construction in ShaderDebug (character, widget) and cap it to the max value.  Fixes rare Vulkan crash because of incorrect instance count in line indirect draw.

diff --git a/Engine/Shaders/Private/ShaderPrintDrawPrimitive.usf b/Engine/Shaders/Private/ShaderPrintDrawPrimitive.usf
--- a/Engine/Shaders/Private/ShaderPrintDrawPrimitive.usf
+++ b/Engine/Shaders/Private/ShaderPrintDrawPrimitive.usf
@@ -95,12 +95,15 @@ void ShaderDrawDebugClearCS(uint3 DispatchThreadId : SV_DispatchThreadID)
 
 StructuredBuffer<uint> ElementBuffer;
 RWBuffer<uint> RWIndirectArgs;
+uint MaxInstanceCount;
 
 [numthreads(1, 1, 1)]
 void ShaderDrawDebugCopyCS(uint3 DispatchThreadId : SV_DispatchThreadID)
 {
+	uint IndirectArgInstanceCount = min(MaxInstanceCount, ElementBuffer[0]);
+
 	RWIndirectArgs[0] = 2; // Line primitive
-	RWIndirectArgs[1] = ElementBuffer[0];
+	RWIndirectArgs[1] = IndirectArgInstanceCount;
 	RWIndirectArgs[2] = 0;
 	RWIndirectArgs[3] = 0;
 }
diff --git a/Engine/Source/Runtime/Renderer/Private/ShaderPrint.cpp b/Engine/Source/Runtime/Renderer/Private/ShaderPrint.cpp
--- a/Engine/Source/Runtime/Renderer/Private/ShaderPrint.cpp
+++ b/Engine/Source/Runtime/Renderer/Private/ShaderPrint.cpp
@@ -121,6 +121,13 @@ namespace ShaderPrint
 		return IsEnabled() ? MaxValueCount : 0;
 	}
 
+	// Returns the element count for a whole line
+	static uint32 GetLineElementCount()
+	{
+		// x,y,z (3) and a packed color (1) for both the start and end points (2x)
+		return 8;
+	}
+
 	// Get symbol buffer size
 	// This is some multiple of the value buffer size to allow for maximum value->symbol expansion
 	static uint32 GetMaxSymbolCountFromValueCount(uint32 MaxValueCount)
@@ -431,6 +438,7 @@ namespace ShaderPrint
 		BEGIN_SHADER_PARAMETER_STRUCT(FParameters, )
 			SHADER_PARAMETER_RDG_BUFFER_SRV(StructuredBuffer, ElementBuffer)
 			SHADER_PARAMETER_RDG_BUFFER_UAV(RWBuffer, RWIndirectArgs)
+			SHADER_PARAMETER(uint32, MaxInstanceCount)
 		END_SHADER_PARAMETER_STRUCT()
 
 		static bool ShouldCompilePermutation(const FGlobalShaderPermutationParameters& Parameters)
@@ -612,7 +620,7 @@ namespace ShaderPrint
 			const bool bLockBufferThisFrame = IsDrawLocked() && InViewState != nullptr && !InViewState->ShaderPrintStateData.bIsLocked;
 			ERDGBufferFlags Flags = bLockBufferThisFrame ? ERDGBufferFlags::MultiFrame : ERDGBufferFlags::None;
 
-			ShaderPrintData.ShaderPrintLineBuffer = GraphBuilder.CreateBuffer(FRDGBufferDesc::CreateStructuredDesc(4, 8 * InSetup.MaxLineCount), TEXT("ShaderPrint.LineBuffer"), Flags);
+			ShaderPrintData.ShaderPrintLineBuffer = GraphBuilder.CreateBuffer(FRDGBufferDesc::CreateStructuredDesc(4, GetLineElementCount() * InSetup.MaxLineCount), TEXT("ShaderPrint.LineBuffer"), Flags);
 
 			// Clear buffer counter
 			{
@@ -836,9 +844,12 @@ namespace ShaderPrint
 
 		FRDGBufferRef IndirectBuffer = GraphBuilder.CreateBuffer(FRDGBufferDesc::CreateIndirectDesc<FRHIDrawIndirectParameters>(1), TEXT("ShaderDraw.IndirectBuffer"), ERDGBufferFlags::None);
 		{
+			const uint32 MaxLineCount = ShaderPrintLineBuffer->Desc.NumElements / GetLineElementCount();
+
 			FShaderDrawDebugCopyCS::FParameters* Parameters = GraphBuilder.AllocParameters<FShaderDrawDebugCopyCS::FParameters>();
 			Parameters->ElementBuffer = GraphBuilder.CreateSRV(ShaderPrintLineBuffer);
 			Parameters->RWIndirectArgs = GraphBuilder.CreateUAV(IndirectBuffer, PF_R32_UINT);
+			Parameters->MaxInstanceCount = MaxLineCount;
 
 			TShaderMapRef<FShaderDrawDebugCopyCS> ComputeShader(GlobalShaderMap);
 			ClearUnusedGraphResources(ComputeShader, Parameters);
