Refactor of mobile custom depth/stencil to line up with pc on all mobile platform.

diff --git a/Engine/Plugins/Experimental/RemoteSession/Source/RemoteSession/Private/Channels/RemoteSessionARCameraChannel.cpp b/Engine/Plugins/Experimental/RemoteSession/Source/RemoteSession/Private/Channels/RemoteSessionARCameraChannel.cpp
--- a/Engine/Plugins/Experimental/RemoteSession/Source/RemoteSession/Private/Channels/RemoteSessionARCameraChannel.cpp
+++ b/Engine/Plugins/Experimental/RemoteSession/Source/RemoteSession/Private/Channels/RemoteSessionARCameraChannel.cpp
@@ -74,7 +74,6 @@ public:
 	{
 		FMaterialShader::ModifyCompilationEnvironment(Parameters, OutEnvironment);
 		OutEnvironment.SetDefine(TEXT("POST_PROCESS_MATERIAL"), 1);
-		OutEnvironment.SetDefine(TEXT("POST_PROCESS_MATERIAL_MOBILE"), 0);
 		OutEnvironment.SetDefine(TEXT("POST_PROCESS_MATERIAL_BEFORE_TONEMAP"), (Parameters.MaterialParameters.BlendableLocation != BL_AfterTonemapping) ? 1 : 0);
 	}
 };
diff --git a/Engine/Plugins/FX/Niagara/Shaders/Private/NiagaraDataInterfaceGBufferTemplate.ush b/Engine/Plugins/FX/Niagara/Shaders/Private/NiagaraDataInterfaceGBufferTemplate.ush
--- a/Engine/Plugins/FX/Niagara/Shaders/Private/NiagaraDataInterfaceGBufferTemplate.ush
+++ b/Engine/Plugins/FX/Niagara/Shaders/Private/NiagaraDataInterfaceGBufferTemplate.ush
@@ -222,11 +222,7 @@ void DecodeCustomStencil_{ParameterName}(in float2 ScreenUV, in bool bApplyViewp
 	OutValue = GBufferData.CustomStencil;
 #elif NDIGBUFFER_IS_MOBILE
 	IsValid = true;
-	#if MOBILE_SUPPORT_FETCH_BINDED_CUSTOM_STENCIL_BUFFER
-		OutValue = MobileSceneTextures.CustomStencilTexture.Load(int3(Parameters.SvPosition.xy, 0)).x;
-	#else
-		OutValue = Texture2DSampleLevel(MobileSceneTextures.MobileCustomStencilTexture, MobileSceneTextures.MobileCustomStencilTextureSampler, ScreenUV, 0).r * 255.0f;
-	#endif
+	OutValue = MobileSceneTextures.CustomStencilTexture.Load(int3(Parameters.SvPosition.xy, 0)) STENCIL_COMPONENT_SWIZZLE;
 #else
 	IsValid = false;
 	OutValue = 0;
diff --git a/Engine/Plugins/Runtime/AR/AppleAR/AppleARKit/Source/AppleARKit/Private/AppleARKitVideoOverlay.cpp b/Engine/Plugins/Runtime/AR/AppleAR/AppleARKit/Source/AppleARKit/Private/AppleARKitVideoOverlay.cpp
--- a/Engine/Plugins/Runtime/AR/AppleAR/AppleARKit/Source/AppleARKit/Private/AppleARKitVideoOverlay.cpp
+++ b/Engine/Plugins/Runtime/AR/AppleAR/AppleARKit/Source/AppleARKit/Private/AppleARKitVideoOverlay.cpp
@@ -140,7 +140,6 @@ public:
 	{
 		FMaterialShader::ModifyCompilationEnvironment(Parameters, OutEnvironment);
 		OutEnvironment.SetDefine(TEXT("POST_PROCESS_MATERIAL"), 1);
-		OutEnvironment.SetDefine(TEXT("POST_PROCESS_MATERIAL_MOBILE"), bIsMobileRenderer);
 		OutEnvironment.SetDefine(TEXT("POST_PROCESS_MATERIAL_BEFORE_TONEMAP"), (Parameters.MaterialParameters.BlendableLocation != BL_AfterTonemapping) ? 1 : 0);
 		OutEnvironment.SetDefine(TEXT("POST_PROCESS_AR_PASSTHROUGH"), 1);
 	}
diff --git a/Engine/Plugins/Runtime/AR/Google/GoogleARCore/Source/GoogleARCoreRendering/Private/GoogleARCorePassthroughCameraRenderer.cpp b/Engine/Plugins/Runtime/AR/Google/GoogleARCore/Source/GoogleARCoreRendering/Private/GoogleARCorePassthroughCameraRenderer.cpp
--- a/Engine/Plugins/Runtime/AR/Google/GoogleARCore/Source/GoogleARCoreRendering/Private/GoogleARCorePassthroughCameraRenderer.cpp
+++ b/Engine/Plugins/Runtime/AR/Google/GoogleARCore/Source/GoogleARCoreRendering/Private/GoogleARCorePassthroughCameraRenderer.cpp
@@ -109,7 +109,6 @@ public:
 	{
 		FMaterialShader::ModifyCompilationEnvironment(Parameters, OutEnvironment);
 		OutEnvironment.SetDefine(TEXT("POST_PROCESS_MATERIAL"), 1);
-		OutEnvironment.SetDefine(TEXT("POST_PROCESS_MATERIAL_MOBILE"), 1);
 		OutEnvironment.SetDefine(TEXT("POST_PROCESS_MATERIAL_BEFORE_TONEMAP"), (Parameters.MaterialParameters.BlendableLocation != BL_AfterTonemapping) ? 1 : 0);
 	}
 };
diff --git a/Engine/Shaders/Private/DepthOnlyPixelShader.usf b/Engine/Shaders/Private/DepthOnlyPixelShader.usf
--- a/Engine/Shaders/Private/DepthOnlyPixelShader.usf
+++ b/Engine/Shaders/Private/DepthOnlyPixelShader.usf
@@ -8,13 +8,8 @@
 #include "/Engine/Generated/Material.ush"
 #include "/Engine/Generated/VertexFactory.ush"
 
-#if OUTPUT_MOBILE_COLOR_VALUE
-// used to write custom stencil values to color target
-float MobileColorValue;
-#endif
-
 void Main(
-#if !MATERIALBLENDING_SOLID || OUTPUT_PIXEL_DEPTH_OFFSET || OUTPUT_MOBILE_COLOR_VALUE
+#if !MATERIALBLENDING_SOLID || OUTPUT_PIXEL_DEPTH_OFFSET
 	in INPUT_POSITION_QUALIFIERS float4 SvPosition : SV_Position,
 #endif
 
@@ -27,12 +22,7 @@ void Main(
 	OPTIONAL_IsFrontFace
 	OPTIONAL_OutDepthConservative,
 #endif
-#if OUTPUT_MOBILE_COLOR_VALUE
-	out HALF_TYPE OutCustomDepth : SV_Target0,
-	out HALF_TYPE OutCustomStencil : SV_Target1
-#else
 	out float4 OutColor : SV_Target0
-#endif
 #if MATERIALBLENDING_MASKED_USING_COVERAGE
 	, out uint OutCoverage : SV_Coverage
 #endif
@@ -75,10 +65,5 @@ void Main(
 	#endif
 #endif
 
-#if OUTPUT_MOBILE_COLOR_VALUE
-	OutCustomDepth = SvPosition.z;
-	OutCustomStencil = MobileColorValue;
-#else
 	OutColor = 0;
-#endif
 }
diff --git a/Engine/Shaders/Private/MaterialTemplate.ush b/Engine/Shaders/Private/MaterialTemplate.ush
--- a/Engine/Shaders/Private/MaterialTemplate.ush
+++ b/Engine/Shaders/Private/MaterialTemplate.ush
@@ -2189,13 +2189,8 @@ MaterialFloat4 MobileSceneTextureLookup(inout FMaterialPixelParameters Parameter
 	}
 	else if (SceneTextureId == PPI_CustomStencil)
 	{
-#if MOBILE_SUPPORT_FETCH_BINDED_CUSTOM_STENCIL_BUFFER
-		MaterialFloat Stencil = MobileSceneTextures.CustomStencilTexture.Load(int3(Parameters.SvPosition.xy, 0)).x;
-#else
-		MaterialFloat Stencil = Texture2DSample(MobileSceneTextures.MobileCustomStencilTexture, MobileSceneTextures.MobileCustomStencilTextureSampler, UV).r*255.0;
-		Stencil = floor(Stencil + 0.5);
-#endif
-	return MaterialFloat4(Stencil.rrr, 0.f);
+		MaterialFloat Stencil = MobileSceneTextures.CustomStencilTexture.Load(int3(Parameters.SvPosition.xy, 0)) STENCIL_COMPONENT_SWIZZLE;
+		return MaterialFloat4(Stencil, Stencil, Stencil, 0.f);
 	}
 #endif// SCENE_TEXTURES_DISABLED
 
diff --git a/Engine/Shaders/Private/PostProcessMaterialShaders.usf b/Engine/Shaders/Private/PostProcessMaterialShaders.usf
--- a/Engine/Shaders/Private/PostProcessMaterialShaders.usf
+++ b/Engine/Shaders/Private/PostProcessMaterialShaders.usf
@@ -11,18 +11,10 @@
 #define POST_PROCESS_MATERIAL 0
 #endif
 
-#ifndef POST_PROCESS_MATERIAL_MOBILE
-#define POST_PROCESS_MATERIAL_MOBILE 0
-#endif
-
 #if (POST_PROCESS_MATERIAL == 0)
 #error POST_PROCESS_MATERIAL must be defined to non-zero in the shader compilation environment.
 #endif
 
-#ifndef MOBILE_SUPPORT_FETCH_BINDED_CUSTOM_STENCIL_BUFFER
-#define MOBILE_SUPPORT_FETCH_BINDED_CUSTOM_STENCIL_BUFFER 0
-#endif 
-
 // Must match ESceneTextureId
 #define PPI_PostProcessInput0 14
 #define PPI_PostProcessInput1 15
@@ -116,7 +108,7 @@ FMaterialVertexParameters GetPostProcessMaterialVSParameters(float2 UV)
 	return Result;
 }
 
-#if (FEATURE_LEVEL > FEATURE_LEVEL_ES3_1 || MOBILE_SUPPORT_FETCH_BINDED_CUSTOM_STENCIL_BUFFER)
+#if (FEATURE_LEVEL > FEATURE_LEVEL_ES3_1)
 
 #if VERTEXSHADER
 
@@ -214,8 +206,6 @@ void MainPS(
 	const float UnlitBSDFTransmittance = UNLIT_TRANSMITTANCE(UnlitBSDF).x;
 	const float UnlitBSDFOpacity = 1.0f - UnlitBSDFTransmittance;
 	Alpha = UnlitBSDFOpacity;
-  #elif (FEATURE_LEVEL <= FEATURE_LEVEL_ES3_1)
-	Alpha = Parameters.BackupSceneColorAlpha;
   #endif
 
 	const float3 Emissive = BSDF_GETEMISSIVE(UnlitBSDF);
@@ -275,12 +265,7 @@ void MainPS_VideoOverlay(
 
 #endif
 
-#else // FEATURE_LEVEL > FEATURE_LEVEL_ES3_1 || MOBILE_SUPPORT_FETCH_BINDED_CUSTOM_STENCIL_BUFFER
-
-Texture2D MobileCustomStencilTexture;
-SamplerState MobileCustomStencilTextureSampler;
-
-int MobileStencilValueRef; // Use integer cause it has to be negative to make the less function to pass the test as always
+#else // FEATURE_LEVEL > FEATURE_LEVEL_ES3_1
 
 uint bFlipYAxis;
 
@@ -321,19 +306,6 @@ void MainVS(
 #endif
 }
 
-#define MOBILE_STENCIL_COMPARE_LESS				0
-#define MOBILE_STENCIL_COMPARE_LESSEQUAL		1
-#define MOBILE_STENCIL_COMPARE_GREATE			2
-#define MOBILE_STENCIL_COMPARE_GREATEEQUAL		3
-#define MOBILE_STENCIL_COMPARE_EQUAL			4
-#define MOBILE_STENCIL_COMPARE_NOTEQUAL			5
-#define MOBILE_STENCIL_COMPARE_NEVER			6
-#define MOBILE_STENCIL_COMPARE_ALWAYS			7
-
-#ifndef MOBILE_STENCIL_COMPARE_FUNCTION
-#define MOBILE_STENCIL_COMPARE_FUNCTION MOBILE_STENCIL_COMPARE_NEVER
-#endif
-
 void MainPS(
 	in FPostProcessMaterialVSToPS Input,
 	out half4 OutColor : SV_Target0
@@ -372,31 +344,6 @@ void MainPS(
 		
 	float ClipValue = 1.0f;
 
-#if MOBILE_STENCIL_COMPARE_FUNCTION != MOBILE_STENCIL_COMPARE_NEVER && MOBILE_STENCIL_COMPARE_FUNCTION != MOBILE_STENCIL_COMPARE_ALWAYS
-	int MobileStencilValue = floor(Texture2DSample(MobileCustomStencilTexture, MobileCustomStencilTextureSampler, BufferUV).r*255.0f + 0.5f);
-#endif
-
-#if MOBILE_STENCIL_COMPARE_FUNCTION == MOBILE_STENCIL_COMPARE_LESS
-	ClipValue = MobileStencilValueRef < MobileStencilValue ? 1.0f : -1.0f;
-#elif MOBILE_STENCIL_COMPARE_FUNCTION == MOBILE_STENCIL_COMPARE_LESSEQUAL
-	ClipValue = MobileStencilValueRef <= MobileStencilValue ? 1.0f : -1.0f;
-#elif MOBILE_STENCIL_COMPARE_FUNCTION == MOBILE_STENCIL_COMPARE_GREATE
-	ClipValue = MobileStencilValueRef > MobileStencilValue ? 1.0f : -1.0f;
-#elif MOBILE_STENCIL_COMPARE_FUNCTION == MOBILE_STENCIL_COMPARE_GREATEEQUAL
-	ClipValue = MobileStencilValueRef >= MobileStencilValue ? 1.0f : -1.0f;
-#elif MOBILE_STENCIL_COMPARE_FUNCTION == MOBILE_STENCIL_COMPARE_EQUAL
-	ClipValue = MobileStencilValueRef == MobileStencilValue ? 1.0f : -1.0f;
-#elif MOBILE_STENCIL_COMPARE_FUNCTION == MOBILE_STENCIL_COMPARE_NOTEQUAL
-	ClipValue = MobileStencilValueRef != MobileStencilValue ? 1.0f : -1.0f;
-#elif MOBILE_STENCIL_COMPARE_FUNCTION == MOBILE_STENCIL_COMPARE_NEVER
-#elif MOBILE_STENCIL_COMPARE_FUNCTION == MOBILE_STENCIL_COMPARE_ALWAYS
-	ClipValue = MobileStencilValueRef == 256 ? 1.0f : -1.0f;
-#endif
-
-#if MOBILE_STENCIL_COMPARE_FUNCTION != MOBILE_STENCIL_COMPARE_NEVER
-	clip(ClipValue);
-#endif
-
 	float DeviceZ = LookupDeviceZ(BufferUV);
 	SvPosition.z = DeviceZ;
 	SvPosition.z = max(SvPosition.z, 1e-18);
@@ -424,4 +371,4 @@ void MainPS(
 #endif
 	OutColor = FullSceneColor;
 }
-#endif // FEATURE_LEVEL > FEATURE_LEVEL_ES3_1 || MOBILE_SUPPORT_FETCH_BINDED_CUSTOM_STENCIL_BUFFER
\ No newline at end of file
+#endif // FEATURE_LEVEL > FEATURE_LEVEL_ES3_1
\ No newline at end of file
diff --git a/Engine/Shaders/Private/SceneCapturePixelShader.usf b/Engine/Shaders/Private/SceneCapturePixelShader.usf
--- a/Engine/Shaders/Private/SceneCapturePixelShader.usf
+++ b/Engine/Shaders/Private/SceneCapturePixelShader.usf
@@ -5,7 +5,7 @@
 =============================================================================*/
 
 // Sample SceneDepthAux texture on mobile
-#define POST_PROCESS_MATERIAL_MOBILE 1
+#define POST_PROCESS_MATERIAL 1
 
 #include "Common.ush"
 #include "DeferredShadingCommon.ush"
diff --git a/Engine/Shaders/Private/SceneTexturesCommon.ush b/Engine/Shaders/Private/SceneTexturesCommon.ush
--- a/Engine/Shaders/Private/SceneTexturesCommon.ush
+++ b/Engine/Shaders/Private/SceneTexturesCommon.ush
@@ -131,7 +131,7 @@ uint CalcSceneCustomStencil(uint2 PixelPos)
 #if SCENE_TEXTURES_DISABLED
 	return 0;
 #else
-	return SceneTexturesStruct.CustomStencilTexture.Load(uint3(PixelPos, 0)).r;
+	return SceneTexturesStruct.CustomStencilTexture.Load(uint3(PixelPos, 0)) STENCIL_COMPONENT_SWIZZLE;
 #endif
 }
 
@@ -184,7 +184,7 @@ float LookupDeviceZ( float2 ScreenUV )
 #elif FORCE_DEPTH_TEXTURE_READS || COMPILER_SWITCH
 	// native Depth buffer lookup
 	return Texture2DSampleLevel(MobileSceneTextures.SceneDepthTexture, MobileSceneTextures.SceneDepthTextureSampler, ScreenUV, 0).r;
-#elif (POST_PROCESS_MATERIAL || POST_PROCESS_MATERIAL_MOBILE || COMPUTESHADER) && !POST_PROCESS_AR_PASSTHROUGH
+#elif (POST_PROCESS_MATERIAL || COMPUTESHADER) && !POST_PROCESS_AR_PASSTHROUGH
 	#if MOBILE_DEFERRED_SHADING
 		return Texture2DSampleLevel(MobileSceneTextures.SceneDepthTexture, MobileSceneTextures.SceneDepthTextureSampler, ScreenUV, 0).r;
 	#else
diff --git a/Engine/Source/Runtime/Engine/Classes/Materials/MaterialRelevance.h b/Engine/Source/Runtime/Engine/Classes/Materials/MaterialRelevance.h
--- a/Engine/Source/Runtime/Engine/Classes/Materials/MaterialRelevance.h
+++ b/Engine/Source/Runtime/Engine/Classes/Materials/MaterialRelevance.h
@@ -37,7 +37,7 @@ struct ENGINE_API FMaterialRelevance
 			uint8 bUsesSkyMaterial : 1;
 			uint8 bUsesSingleLayerWaterMaterial : 1;
 			uint8 bHasVolumeMaterialDomain : 1;
-			uint8 bUsesCustomDepthStencil : 1;
+			uint8 CustomDepthStencilUsageMask : 2;
 			uint8 bUsesDistanceCullFade : 1;
 			uint8 bDisableDepthTest : 1;
 			uint8 bUsesAnisotropy : 1;
diff --git a/Engine/Source/Runtime/Engine/Private/Materials/MaterialInterface.cpp b/Engine/Source/Runtime/Engine/Private/Materials/MaterialInterface.cpp
--- a/Engine/Source/Runtime/Engine/Private/Materials/MaterialInterface.cpp
+++ b/Engine/Source/Runtime/Engine/Private/Materials/MaterialInterface.cpp
@@ -366,7 +366,7 @@ FMaterialRelevance UMaterialInterface::GetRelevance_Internal(const UMaterial* Ma
 		FMaterialRelevance MaterialRelevance;
 
 		MaterialRelevance.ShadingModelMask = GetShadingModels().GetShadingModelField();
-		MaterialRelevance.bUsesCustomDepthStencil = MaterialResource->UsesCustomDepthStencil_GameThread();
+		MaterialRelevance.CustomDepthStencilUsageMask = MaterialResource->GetCustomDepthStencilUsageMask_GameThread();
 
 		if(bDecal)
 		{
diff --git a/Engine/Source/Runtime/Engine/Private/Materials/MaterialShared.cpp b/Engine/Source/Runtime/Engine/Private/Materials/MaterialShared.cpp
--- a/Engine/Source/Runtime/Engine/Private/Materials/MaterialShared.cpp
+++ b/Engine/Source/Runtime/Engine/Private/Materials/MaterialShared.cpp
@@ -971,9 +971,15 @@ bool FMaterial::MaterialUsesSceneDepthLookup_GameThread() const
 	return GameThreadShaderMap.GetReference() ? GameThreadShaderMap->UsesSceneDepthLookup() : false;
 }
 
-bool FMaterial::UsesCustomDepthStencil_GameThread() const
+uint8 FMaterial::GetCustomDepthStencilUsageMask_GameThread() const
 {
-	return GameThreadShaderMap.GetReference() ? (GameThreadShaderMap->UsesSceneTexture(PPI_CustomDepth) || GameThreadShaderMap->UsesSceneTexture(PPI_CustomStencil)) : false;
+	uint8 CustomDepthStencilUsageMask = 0;
+	if (GameThreadShaderMap.GetReference())
+	{
+		CustomDepthStencilUsageMask |= GameThreadShaderMap->UsesSceneTexture(PPI_CustomDepth) ? 1 : 0;
+		CustomDepthStencilUsageMask |= GameThreadShaderMap->UsesSceneTexture(PPI_CustomStencil) ? 1 << 1 : 0;
+	}
+	return CustomDepthStencilUsageMask;
 }
 
 uint8 FMaterial::GetRuntimeVirtualTextureOutputAttibuteMask_RenderThread() const
diff --git a/Engine/Source/Runtime/Engine/Private/ShaderCompiler/ShaderGenerationUtil.cpp b/Engine/Source/Runtime/Engine/Private/ShaderCompiler/ShaderGenerationUtil.cpp
--- a/Engine/Source/Runtime/Engine/Private/ShaderCompiler/ShaderGenerationUtil.cpp
+++ b/Engine/Source/Runtime/Engine/Private/ShaderCompiler/ShaderGenerationUtil.cpp
@@ -1360,17 +1360,10 @@ static FString CreateGBufferDecodeFunctionVariation(const FGBufferInfo& BufferIn
 	{
 		FullStr += FString::Printf(TEXT("\tfloat CustomNativeDepth = Texture2DSampleLevel(SceneTexturesStruct.CustomDepthTexture, SceneTexturesStruct_CustomDepthTextureSampler, %s, 0).r;\n"), CoordName.GetCharArray().GetData());
 
-		if (FEATURE_LEVEL >= ERHIFeatureLevel::SM5)
-		{
-			// BufferToSceneTextureScale is necessary when translucent materials are rendered in a render target 
-			// that has a different resolution than the scene color textures, e.g. r.SeparateTranslucencyScreenPercentage < 100.
-			FullStr += FString::Printf(TEXT("\tint2 IntUV = (int2)trunc(%s * View.BufferSizeAndInvSize.xy * View.BufferToSceneTextureScale.xy);\n"), CoordName.GetCharArray().GetData());
-			FullStr += TEXT("\tuint CustomStencil = SceneTexturesStruct.CustomStencilTexture.Load(int3(IntUV, 0)) STENCIL_COMPONENT_SWIZZLE;\n");
-		}
-		else
-		{
-			FullStr += TEXT("\tuint CustomStencil = 0;\n");
-		}
+		// BufferToSceneTextureScale is necessary when translucent materials are rendered in a render target 
+		// that has a different resolution than the scene color textures, e.g. r.SeparateTranslucencyScreenPercentage < 100.
+		FullStr += FString::Printf(TEXT("\tint2 IntUV = (int2)trunc(%s * View.BufferSizeAndInvSize.xy * View.BufferToSceneTextureScale.xy);\n"), CoordName.GetCharArray().GetData());
+		FullStr += TEXT("\tuint CustomStencil = SceneTexturesStruct.CustomStencilTexture.Load(int3(IntUV, 0)) STENCIL_COMPONENT_SWIZZLE;\n");
 
 		FullStr += FString::Printf(TEXT("\tfloat SceneDepth = CalcSceneDepth(%s);\n"), CoordName.GetCharArray().GetData());
 		FullStr += TEXT("\tfloat4 AnisotropicData = Texture2DSampleLevel(SceneTexturesStruct.GBufferFTexture, SceneTexturesStruct_GBufferFTextureSampler, UV, 0).xyzw;\n");
diff --git a/Engine/Source/Runtime/Engine/Public/MaterialShared.h b/Engine/Source/Runtime/Engine/Public/MaterialShared.h
--- a/Engine/Source/Runtime/Engine/Public/MaterialShared.h
+++ b/Engine/Source/Runtime/Engine/Public/MaterialShared.h
@@ -2000,8 +2000,8 @@ public:
 	ENGINE_API bool MaterialUsesSceneDepthLookup_RenderThread() const;
 	ENGINE_API bool MaterialUsesSceneDepthLookup_GameThread() const;
 
-	/** Does the material use CustomDepth or CustomStencil lookup */
-	ENGINE_API bool UsesCustomDepthStencil_GameThread() const;
+	/** The material usage mask of CustomDepth and CustomStencil*/
+	ENGINE_API uint8 GetCustomDepthStencilUsageMask_GameThread() const;
 
 	/** Note: This function is only intended for use in deciding whether or not shader permutations are required before material translation occurs. */
 	ENGINE_API bool MaterialMayModifyMeshPosition() const;
diff --git a/Engine/Source/Runtime/Landscape/Private/LandscapeRender.cpp b/Engine/Source/Runtime/Landscape/Private/LandscapeRender.cpp
--- a/Engine/Source/Runtime/Landscape/Private/LandscapeRender.cpp
+++ b/Engine/Source/Runtime/Landscape/Private/LandscapeRender.cpp
@@ -3504,8 +3504,7 @@ public:
 			FName(TEXT("FAnisotropyPS")),
 			FName(TEXT("TDepthOnlyVS<false>")),
 			FName(TEXT("TDepthOnlyVS<true>")),
-			FName(TEXT("FDepthOnlyPS<true>")),
-			FName(TEXT("FDepthOnlyPS<false>")),
+			FName(TEXT("FDepthOnlyPS")),
 			// UE-44519, masked material with landscape layers requires FHitProxy shaders.
 			FName(TEXT("FHitProxyVS")),
 			FName(TEXT("FHitProxyPS")),
diff --git a/Engine/Source/Runtime/OpenGLDrv/Private/OpenGLTexture.cpp b/Engine/Source/Runtime/OpenGLDrv/Private/OpenGLTexture.cpp
--- a/Engine/Source/Runtime/OpenGLDrv/Private/OpenGLTexture.cpp
+++ b/Engine/Source/Runtime/OpenGLDrv/Private/OpenGLTexture.cpp
@@ -1666,21 +1666,29 @@ FOpenGLShaderResourceView::FOpenGLShaderResourceView(FOpenGLTexture* InTexture,
 		? Texture->GetFormat()
 		: CreateInfo.Format;
 	
-	const bool bFormatsMatch = Format == Texture->GetFormat();
+	const bool bFormatsMatch = Format == Texture->GetFormat() || Format == PF_X24_G8;
 
-	if (Format == PF_X24_G8)
+	RunOnGLRenderContextThread([this, bFormatsMatch, Format]()
 	{
-		UE_LOG(LogRHI, Fatal, TEXT("Cannot create a stencil SRV (PF_X24_G8) when texture views are unsupported."));
-	}
-	else
-	{
-		RunOnGLRenderContextThread([this, bFormatsMatch]()
+		VERIFY_GL_SCOPE();
+		checkf(bFormatsMatch, TEXT("SRVs cannot modify the pixel format of a texture when texture views are unsupported."));
+		Resource = Texture->GetResource();
+
+		// Handle the custom stencil SRV
+		if (Format == PF_X24_G8)
 		{
-			VERIFY_GL_SCOPE();
-			checkf(bFormatsMatch, TEXT("SRVs cannot modify the pixel format of a texture when texture views are unsupported."));
-			Resource = Texture->GetResource();
-		});
-	}
+			// Use a texture stage that's not likely to be used for draws, to avoid waiting
+			FOpenGLContextState& ContextState = FOpenGLDynamicRHI::Get().GetContextStateForCurrentContext();
+			FOpenGLDynamicRHI::Get().CachedSetupTextureStage(ContextState, FOpenGL::GetMaxCombinedTextureImageUnits() - 1, Target, Resource, LimitMip, Texture->GetNumMips());
+
+			//set the texture to return the stencil index, and then force the components to match D3D
+			glTexParameteri(Target, GL_DEPTH_STENCIL_TEXTURE_MODE, GL_STENCIL_INDEX);
+			glTexParameteri(Target, GL_TEXTURE_SWIZZLE_R, GL_ZERO);
+			glTexParameteri(Target, GL_TEXTURE_SWIZZLE_G, GL_RED);
+			glTexParameteri(Target, GL_TEXTURE_SWIZZLE_B, GL_ZERO);
+			glTexParameteri(Target, GL_TEXTURE_SWIZZLE_A, GL_ZERO);
+		}
+	});
 }
 
 FShaderResourceViewRHIRef FOpenGLDynamicRHI::RHICreateShaderResourceView(FRHITexture* Texture, const FRHITextureSRVCreateInfo& CreateInfo)
diff --git a/Engine/Source/Runtime/RHI/Private/RHI.cpp b/Engine/Source/Runtime/RHI/Private/RHI.cpp
--- a/Engine/Source/Runtime/RHI/Private/RHI.cpp
+++ b/Engine/Source/Runtime/RHI/Private/RHI.cpp
@@ -2410,7 +2410,6 @@ void FGenericDataDrivenShaderPlatformInfo::SetDefaultValues()
 	bSupportsWaterIndirectDraw = true;
 	bSupportsAsyncPipelineCompilation = true;
 	bSupportsManualVertexFetch = true;
-	bMobileSupportFetchBindedCustomStencilBuffer = false;
 	bSupportsVolumeTextureAtomics = true;
 }
 
@@ -2500,7 +2499,6 @@ void FGenericDataDrivenShaderPlatformInfo::ParseDataDrivenShaderInfo(const FConf
 	GET_SECTION_BOOL_HELPER(bSupportsMobileDistanceField);
 	GET_SECTION_BOOL_HELPER(bSupportsFFTBloom);
 	GET_SECTION_BOOL_HELPER(bSupportsVertexShaderLayer);
-	GET_SECTION_BOOL_HELPER(bMobileSupportFetchBindedCustomStencilBuffer);
 	GET_SECTION_BOOL_HELPER(bSupportsBindless);
 	GET_SECTION_BOOL_HELPER(bSupportsVolumeTextureAtomics);
 	GET_SECTION_BOOL_HELPER(bSupportsROV);
diff --git a/Engine/Source/Runtime/RHI/Public/RHIDefinitions.h b/Engine/Source/Runtime/RHI/Public/RHIDefinitions.h
--- a/Engine/Source/Runtime/RHI/Public/RHIDefinitions.h
+++ b/Engine/Source/Runtime/RHI/Public/RHIDefinitions.h
@@ -408,7 +408,6 @@ class RHI_API FGenericDataDrivenShaderPlatformInfo
 	uint32 bSupportsInlineRayTracing : 1;
 	uint32 bSupportsRayTracingShaders : 1;
 	uint32 bSupportsVertexShaderLayer : 1;
-	uint32 bMobileSupportFetchBindedCustomStencilBuffer : 1;
 	uint32 bSupportsBindless : 1;
 	uint32 bSupportsVolumeTextureAtomics : 1;
 	uint32 bSupportsROV : 1;
@@ -921,12 +920,6 @@ public:
 		check(IsValid(Platform));
 		return Infos[Platform].bSupportsVertexShaderLayer;
 	}
-	
-	static FORCEINLINE_DEBUGGABLE const bool GetMobileSupportFetchBindedCustomStencilBuffer(const FStaticShaderPlatform Platform)
-	{
-		check(IsValid(Platform));
-		return Infos[Platform].bMobileSupportFetchBindedCustomStencilBuffer;
-	}
 
 	static FORCEINLINE_DEBUGGABLE const bool GetSupportsBindless(const FStaticShaderPlatform Platform)
 	{
diff --git a/Engine/Source/Runtime/Renderer/Private/CustomDepthRendering.cpp b/Engine/Source/Runtime/Renderer/Private/CustomDepthRendering.cpp
--- a/Engine/Source/Runtime/Renderer/Private/CustomDepthRendering.cpp
+++ b/Engine/Source/Runtime/Renderer/Private/CustomDepthRendering.cpp
@@ -26,14 +26,6 @@ static TAutoConsoleVariable<int32> CVarCustomDepthOrder(
 	TEXT("  2: Default (Before Base Pass if DBuffer enabled.)\n"),
 	ECVF_RenderThreadSafe);
 
-static TAutoConsoleVariable<int32> CVarMobileCustomDepthDownSample(
-	TEXT("r.Mobile.CustomDepthDownSample"),
-	0,
-	TEXT("Perform Mobile CustomDepth at HalfRes \n ")
-	TEXT("0: Off (default)\n ")
-	TEXT("1: On \n "),
-	ECVF_RenderThreadSafe);
-
 static TAutoConsoleVariable<int32> CVarCustomDepthTemporalAAJitter(
 	TEXT("r.CustomDepthTemporalAAJitter"),
 	1,
@@ -60,38 +52,12 @@ ECustomDepthMode GetCustomDepthMode()
 	return ECustomDepthMode::Disabled;
 }
 
-bool IsMobileSeparateDepthStencilRenderTargets(EShaderPlatform Platform)
-{
-	return IsMobilePlatform(Platform) && !FDataDrivenShaderPlatformInfo::GetMobileSupportFetchBindedCustomStencilBuffer(Platform);
-}
-
-bool IsMobileSupportFetchBindedCustomStencilBuffer(EShaderPlatform Platform)
-{
-	return IsMobilePlatform(Platform) && FDataDrivenShaderPlatformInfo::GetMobileSupportFetchBindedCustomStencilBuffer(Platform);
-}
-
-bool IsCustomDepthPassWritingStencil(EShaderPlatform Platform)
-{
-	switch (GetCustomDepthMode())
-	{
-		case ECustomDepthMode::Disabled:
-		{
-			return false;
-		}
-		case ECustomDepthMode::Enabled:
-		{
-			return IsMobileSeparateDepthStencilRenderTargets(Platform);
-		}
-	}
-	return true;
-}
-
-uint32 GetCustomDepthDownsampleFactor(EShaderPlatform Platform)
+bool IsCustomDepthPassWritingStencil()
 {
-	return IsMobileSeparateDepthStencilRenderTargets(Platform) && CVarMobileCustomDepthDownSample.GetValueOnRenderThread() > 0 ? 2 : 1;
+	return GetCustomDepthMode() == ECustomDepthMode::EnabledWithStencil;
 }
 
-FCustomDepthTextures FCustomDepthTextures::Create(FRDGBuilder& GraphBuilder, FIntPoint Extent, ERHIFeatureLevel::Type FeatureLevel, uint32 DownsampleFactor)
+FCustomDepthTextures FCustomDepthTextures::Create(FRDGBuilder& GraphBuilder, FIntPoint CustomDepthExtent)
 {
 	const ECustomDepthMode CustomDepthMode = GetCustomDepthMode();
 
@@ -100,40 +66,16 @@ FCustomDepthTextures FCustomDepthTextures::Create(FRDGBuilder& GraphBuilder, FIn
 		return {};
 	}
 
-	EShaderPlatform Platform = GShaderPlatformForFeatureLevel[FeatureLevel];
-	const bool bWritesCustomStencil = IsCustomDepthPassWritingStencil(Platform);
-	const FIntPoint CustomDepthExtent = FIntPoint::DivideAndRoundUp(Extent, DownsampleFactor);
+	const bool bWritesCustomStencil = IsCustomDepthPassWritingStencil();
 
 	FCustomDepthTextures CustomDepthTextures;
-	
-	if (IsMobileSeparateDepthStencilRenderTargets(Platform))
-	{
-		const float DepthFar = (float)ERHIZBuffer::FarPlane;
-		const FClearValueBinding DepthFarColor = FClearValueBinding(FLinearColor(DepthFar, DepthFar, DepthFar, DepthFar));
-
-		ETextureCreateFlags MobileCustomDepthFlags = TexCreate_RenderTargetable | TexCreate_ShaderResource;
-		ETextureCreateFlags MobileCustomStencilFlags = MobileCustomDepthFlags;
-
-		if (!bWritesCustomStencil)
-		{
-			MobileCustomStencilFlags |= TexCreate_Memoryless;
-		}
-
-		const FRDGTextureDesc MobileCustomDepthDesc = FRDGTextureDesc::Create2D(CustomDepthExtent, PF_R16F, DepthFarColor, MobileCustomDepthFlags);
-		const FRDGTextureDesc MobileCustomStencilDesc = FRDGTextureDesc::Create2D(CustomDepthExtent, PF_G8, FClearValueBinding::Transparent, MobileCustomStencilFlags);
-
-		CustomDepthTextures.MobileDepth = GraphBuilder.CreateTexture(MobileCustomDepthDesc, TEXT("MobileCustomDepth"));
-		CustomDepthTextures.MobileStencil = GraphBuilder.CreateTexture(MobileCustomStencilDesc, TEXT("MobileCustomStencil"));
-	}
 
 	const FRDGTextureDesc CustomDepthDesc = FRDGTextureDesc::Create2D(CustomDepthExtent, PF_DepthStencil, FClearValueBinding::DepthFar, GFastVRamConfig.CustomDepth | TexCreate_NoFastClear | TexCreate_DepthStencilTargetable | TexCreate_ShaderResource);
 
 	CustomDepthTextures.Depth = GraphBuilder.CreateTexture(CustomDepthDesc, TEXT("CustomDepth"));
 
-	CustomDepthTextures.Stencil = GraphBuilder.CreateSRV(FRDGTextureSRVDesc::CreateWithPixelFormat(CustomDepthTextures.Depth, PF_X24_G8));
 	CustomDepthTextures.DepthAction = ERenderTargetLoadAction::EClear;
 	CustomDepthTextures.StencilAction = bWritesCustomStencil ? ERenderTargetLoadAction::EClear : ERenderTargetLoadAction::ENoAction;
-	CustomDepthTextures.DownsampleFactor = DownsampleFactor;
 
 	return CustomDepthTextures;
 }
@@ -175,7 +117,7 @@ static FViewShaderParameters CreateViewShaderParametersWithoutJitter(const FView
 	return Parameters;
 }
 
-bool FSceneRenderer::RenderCustomDepthPass(FRDGBuilder& GraphBuilder, const FCustomDepthTextures& CustomDepthTextures, const FSceneTextureShaderParameters& SceneTextures)
+bool FSceneRenderer::RenderCustomDepthPass(FRDGBuilder& GraphBuilder, FCustomDepthTextures& CustomDepthTextures, const FSceneTextureShaderParameters& SceneTextures)
 {
 	if (!CustomDepthTextures.IsValid())
 	{
@@ -212,25 +154,11 @@ bool FSceneRenderer::RenderCustomDepthPass(FRDGBuilder& GraphBuilder, const FCus
 			const ERenderTargetLoadAction DepthLoadAction = GetLoadActionIfProduced(CustomDepthTextures.Depth, CustomDepthTextures.DepthAction);
 			const ERenderTargetLoadAction StencilLoadAction = GetLoadActionIfProduced(CustomDepthTextures.Depth, CustomDepthTextures.StencilAction);
 
-			if (IsMobileSeparateDepthStencilRenderTargets(View.GetShaderPlatform()))
-			{
-				PassParameters->RenderTargets[0] = FRenderTargetBinding(CustomDepthTextures.MobileDepth, DepthLoadAction);
-				PassParameters->RenderTargets[1] = FRenderTargetBinding(CustomDepthTextures.MobileStencil, StencilLoadAction);
-
-				PassParameters->RenderTargets.DepthStencil = FDepthStencilBinding(
-					CustomDepthTextures.Depth,
-					DepthLoadAction,
-					DepthLoadAction,
-					FExclusiveDepthStencil::DepthWrite_StencilWrite);
-			}
-			else
-			{
-				PassParameters->RenderTargets.DepthStencil = FDepthStencilBinding(
-					CustomDepthTextures.Depth,
-					DepthLoadAction,
-					StencilLoadAction,
-					FExclusiveDepthStencil::DepthWrite_StencilWrite);
-			}
+			PassParameters->RenderTargets.DepthStencil = FDepthStencilBinding(
+				CustomDepthTextures.Depth,
+				DepthLoadAction,
+				StencilLoadAction,
+				FExclusiveDepthStencil::DepthWrite_StencilWrite);
 
 			View.ParallelMeshDrawCommandPasses[EMeshPass::CustomDepth].BuildRenderingCommands(GraphBuilder, Scene->GPUScene, PassParameters->InstanceCullingDrawParams);
 
@@ -238,9 +166,9 @@ bool FSceneRenderer::RenderCustomDepthPass(FRDGBuilder& GraphBuilder, const FCus
 				RDG_EVENT_NAME("CustomDepth"),
 				PassParameters,
 				ERDGPassFlags::Raster,
-				[this, &View, PassParameters, DownsampleFactor = CustomDepthTextures.DownsampleFactor](FRHICommandList& RHICmdList)
+				[this, &View, PassParameters](FRHICommandList& RHICmdList)
 			{
-				SetStereoViewport(RHICmdList, View, 1.0f / static_cast<float>(DownsampleFactor));
+				SetStereoViewport(RHICmdList, View, 1.0f);
 				View.ParallelMeshDrawCommandPasses[EMeshPass::CustomDepth].DispatchDraw(nullptr, RHICmdList, &PassParameters->InstanceCullingDrawParams);
 			});
 
@@ -248,6 +176,22 @@ bool FSceneRenderer::RenderCustomDepthPass(FRDGBuilder& GraphBuilder, const FCus
 		}
 	}
 
+	if (bCustomDepthRendered)
+	{
+		const FSceneTexturesConfig& Config = FSceneTexturesConfig::Get();
+		FRDGTextureRef CustomDepth = CustomDepthTextures.Depth;
+
+		// TextureView is not supported in GLES, so we can't lookup CustomDepth and CustomStencil from a single texture
+		// Do a copy of the CustomDepthStencil texture if both CustomDepth and CustomStencil are sampled in a shader.
+		if (IsOpenGLPlatform(ShaderPlatform) && Config.bSamplesCustomDepthAndStencil)
+		{
+			CustomDepth = GraphBuilder.CreateTexture(CustomDepthTextures.Depth->Desc, TEXT("CustomDepthCopy"));
+			AddCopyTexturePass(GraphBuilder, CustomDepthTextures.Depth, CustomDepth);
+		}
+
+		CustomDepthTextures.Stencil = GraphBuilder.CreateSRV(FRDGTextureSRVDesc::CreateWithPixelFormat(CustomDepth, PF_X24_G8));
+	}
+
 	return bCustomDepthRendered;
 }
 
@@ -267,7 +211,7 @@ private:
 		const FMaterialRenderProxy& MaterialRenderProxy,
 		const FMaterial& Material);
 
-	template<bool bPositionOnly, bool bUsesMobileColorValue>
+	template<bool bPositionOnly>
 	bool Process(
 		const FMeshBatch& MeshBatch,
 		uint64 BatchElementMask,
@@ -276,8 +220,7 @@ private:
 		const FMaterialRenderProxy& RESTRICT MaterialRenderProxy,
 		const FMaterial& RESTRICT MaterialResource,
 		ERasterizerFillMode MeshFillMode,
-		ERasterizerCullMode MeshCullMode,
-		float MobileColorValue);
+		ERasterizerCullMode MeshCullMode);
 
 	FMeshPassProcessorRenderState PassDrawRenderState;
 };
@@ -320,13 +263,11 @@ bool FCustomDepthPassMeshProcessor::TryAddMeshBatch(
 {
 	// Determine the mesh's material and blend mode.
 	const EBlendMode BlendMode = Material.GetBlendMode();
-	EShaderPlatform Platform = GShaderPlatformForFeatureLevel[FeatureLevel];
 	const FMeshDrawingPolicyOverrideSettings OverrideSettings = ComputeMeshOverrideSettings(MeshBatch);
 	const ERasterizerFillMode MeshFillMode = ComputeMeshFillMode(MeshBatch, Material, OverrideSettings);
 	const ERasterizerCullMode MeshCullMode = ComputeMeshCullMode(MeshBatch, Material, OverrideSettings);
 	const bool bIsTranslucent = IsTranslucentBlendMode(BlendMode);
-	const bool bWriteCustomStencilValues = IsCustomDepthPassWritingStencil(Platform);
-	float MobileColorValue = 0.0f;
+	const bool bWriteCustomStencilValues = IsCustomDepthPassWritingStencil();
 
 	if (bWriteCustomStencilValues)
 	{
@@ -348,27 +289,21 @@ bool FCustomDepthPassMeshProcessor::TryAddMeshBatch(
 
 		PassDrawRenderState.SetDepthStencilState(StencilStates[(int32)PrimitiveSceneProxy->GetStencilWriteMask()]);
 		PassDrawRenderState.SetStencilRef(CustomDepthStencilValue);
-
-		// On mobile platforms write custom stencil value to color target
-		MobileColorValue = CustomDepthStencilValue / 255.0f;
 	}
 	else
 	{
 		PassDrawRenderState.SetDepthStencilState(TStaticDepthStencilState<true, CF_DepthNearOrEqual>::GetRHI());
 	}
 
-	const bool bUsesMobileColorValue = IsMobileSeparateDepthStencilRenderTargets(Platform);
-
 	bool bResult = true;
 	if (BlendMode == BLEND_Opaque
 		&& MeshBatch.VertexFactory->SupportsPositionOnlyStream()
 		&& !Material.MaterialModifiesMeshPosition_RenderThread()
-		&& Material.WritesEveryPixel()
-		&& !bUsesMobileColorValue)
+		&& Material.WritesEveryPixel())
 	{
 		const FMaterialRenderProxy& DefaultProxy = *UMaterial::GetDefaultMaterial(MD_Surface)->GetRenderProxy();
 		const FMaterial& DefaultMaterial = *DefaultProxy.GetMaterialNoFallback(FeatureLevel);
-		bResult = Process<true, false>(MeshBatch, BatchElementMask, StaticMeshId, PrimitiveSceneProxy, DefaultProxy, DefaultMaterial, MeshFillMode, MeshCullMode, MobileColorValue);
+		bResult = Process<true>(MeshBatch, BatchElementMask, StaticMeshId, PrimitiveSceneProxy, DefaultProxy, DefaultMaterial, MeshFillMode, MeshCullMode);
 	}
 	else if (!IsTranslucentBlendMode(BlendMode) || Material.IsTranslucencyWritingCustomDepth())
 	{
@@ -385,20 +320,13 @@ bool FCustomDepthPassMeshProcessor::TryAddMeshBatch(
 			check(EffectiveMaterial);
 		}
 
-		if (bUsesMobileColorValue)
-		{
-			bResult = Process<false, true>(MeshBatch, BatchElementMask, StaticMeshId, PrimitiveSceneProxy, *EffectiveMaterialRenderProxy, *EffectiveMaterial, MeshFillMode, MeshCullMode, MobileColorValue);
-		}
-		else
-		{
-			bResult = Process<false, false>(MeshBatch, BatchElementMask, StaticMeshId, PrimitiveSceneProxy, *EffectiveMaterialRenderProxy, *EffectiveMaterial, MeshFillMode, MeshCullMode, MobileColorValue);
-		}
+		bResult = Process<false>(MeshBatch, BatchElementMask, StaticMeshId, PrimitiveSceneProxy, *EffectiveMaterialRenderProxy, *EffectiveMaterial, MeshFillMode, MeshCullMode);
 	}
 
 	return bResult;
 }
 
-template<bool bPositionOnly, bool bUsesMobileColorValue>
+template<bool bPositionOnly>
 bool FCustomDepthPassMeshProcessor::Process(
 	const FMeshBatch& RESTRICT MeshBatch,
 	uint64 BatchElementMask,
@@ -407,17 +335,16 @@ bool FCustomDepthPassMeshProcessor::Process(
 	const FMaterialRenderProxy& RESTRICT MaterialRenderProxy,
 	const FMaterial& RESTRICT MaterialResource,
 	ERasterizerFillMode MeshFillMode,
-	ERasterizerCullMode MeshCullMode,
-	float MobileColorValue)
+	ERasterizerCullMode MeshCullMode)
 {
 	const FVertexFactory* VertexFactory = MeshBatch.VertexFactory;
 
 	TMeshProcessorShaders<
 		TDepthOnlyVS<bPositionOnly>,
-		FDepthOnlyPS<bUsesMobileColorValue>> DepthPassShaders;
+		FDepthOnlyPS> DepthPassShaders;
 
 	FShaderPipelineRef ShaderPipeline;
-	if (!GetDepthPassShaders<bPositionOnly, bUsesMobileColorValue>(
+	if (!GetDepthPassShaders<bPositionOnly>(
 		MaterialResource,
 		VertexFactory->GetType(),
 		FeatureLevel,
@@ -429,7 +356,7 @@ bool FCustomDepthPassMeshProcessor::Process(
 		return false;
 	}
 
-	FDepthOnlyShaderElementData ShaderElementData(MobileColorValue);
+	FMeshMaterialShaderElementData ShaderElementData;
 	ShaderElementData.InitializeMeshMaterialData(ViewIfDynamicMeshCommand, PrimitiveSceneProxy, MeshBatch, StaticMeshId, false);
 
 	const FMeshDrawCommandSortKey SortKey = CalculateMeshStaticSortKey(DepthPassShaders.VertexShader, DepthPassShaders.PixelShader);
diff --git a/Engine/Source/Runtime/Renderer/Private/CustomDepthRendering.h b/Engine/Source/Runtime/Renderer/Private/CustomDepthRendering.h
--- a/Engine/Source/Runtime/Renderer/Private/CustomDepthRendering.h
+++ b/Engine/Source/Runtime/Renderer/Private/CustomDepthRendering.h
@@ -31,23 +31,14 @@ enum class ECustomDepthMode : uint32
 // The custom depth mode currently configured.
 extern ECustomDepthMode GetCustomDepthMode();
 
-// Returns the requested downsample factor for custom depth textures.
-extern uint32 GetCustomDepthDownsampleFactor(EShaderPlatform Platform);
-
 inline bool IsCustomDepthPassEnabled()
 {
 	return GetCustomDepthMode() != ECustomDepthMode::Disabled;
 }
 
-// Check if the mobile uses separate depth and stencil render targets for the custom renderer
-extern bool IsMobileSeparateDepthStencilRenderTargets(EShaderPlatform Platform);
-
-// Check if the mobile support to directly fetch from the depthStencil buffer
-extern bool IsMobileSupportFetchBindedCustomStencilBuffer(EShaderPlatform Platform);
-
 struct FCustomDepthTextures
 {
-	static FCustomDepthTextures Create(FRDGBuilder& GraphBuilder, FIntPoint Extent, ERHIFeatureLevel::Type FeatureLevel, uint32 DownsampleFactor);
+	static FCustomDepthTextures Create(FRDGBuilder& GraphBuilder, FIntPoint CustomDepthExtent);
 
 	bool IsValid() const
 	{
@@ -57,12 +48,7 @@ struct FCustomDepthTextures
 	FRDGTextureRef Depth{};
 	FRDGTextureSRVRef Stencil{};
 
-	FRDGTextureRef MobileDepth{};
-	FRDGTextureRef MobileStencil{};
-
 	// Actions to use when initially rendering to custom depth / stencil.
 	ERenderTargetLoadAction DepthAction = ERenderTargetLoadAction::EClear;
 	ERenderTargetLoadAction StencilAction = ERenderTargetLoadAction::EClear;
-
-	uint32 DownsampleFactor = 1;
 };
diff --git a/Engine/Source/Runtime/Renderer/Private/DepthRendering.cpp b/Engine/Source/Runtime/Renderer/Private/DepthRendering.cpp
--- a/Engine/Source/Runtime/Renderer/Private/DepthRendering.cpp
+++ b/Engine/Source/Runtime/Renderer/Private/DepthRendering.cpp
@@ -135,13 +135,11 @@ DECLARE_GPU_DRAWCALL_STAT(Prepass);
 IMPLEMENT_MATERIAL_SHADER_TYPE(template<>,TDepthOnlyVS<true>,TEXT("/Engine/Private/PositionOnlyDepthVertexShader.usf"),TEXT("Main"),SF_Vertex);
 IMPLEMENT_MATERIAL_SHADER_TYPE(template<>,TDepthOnlyVS<false>,TEXT("/Engine/Private/DepthOnlyVertexShader.usf"),TEXT("Main"),SF_Vertex);
 
-IMPLEMENT_MATERIAL_SHADER_TYPE(template<>,FDepthOnlyPS<true>,TEXT("/Engine/Private/DepthOnlyPixelShader.usf"),TEXT("Main"),SF_Pixel);
-IMPLEMENT_MATERIAL_SHADER_TYPE(template<>,FDepthOnlyPS<false>,TEXT("/Engine/Private/DepthOnlyPixelShader.usf"),TEXT("Main"),SF_Pixel);
+IMPLEMENT_MATERIAL_SHADER_TYPE(,FDepthOnlyPS,TEXT("/Engine/Private/DepthOnlyPixelShader.usf"),TEXT("Main"),SF_Pixel);
 
 IMPLEMENT_SHADERPIPELINE_TYPE_VS(DepthNoPixelPipeline, TDepthOnlyVS<false>, true);
 IMPLEMENT_SHADERPIPELINE_TYPE_VS(DepthPosOnlyNoPixelPipeline, TDepthOnlyVS<true>, true);
-IMPLEMENT_SHADERPIPELINE_TYPE_VSPS(DepthNoColorOutputPipeline, TDepthOnlyVS<false>, FDepthOnlyPS<false>, true);
-IMPLEMENT_SHADERPIPELINE_TYPE_VSPS(DepthWithColorOutputPipeline, TDepthOnlyVS<false>, FDepthOnlyPS<true>, true);
+IMPLEMENT_SHADERPIPELINE_TYPE_VSPS(DepthPipeline, TDepthOnlyVS<false>, FDepthOnlyPS, true);
 
 static bool IsDepthPassWaitForTasksEnabled()
 {
@@ -154,19 +152,19 @@ static FORCEINLINE bool UseShaderPipelines(ERHIFeatureLevel::Type InFeatureLevel
 	return RHISupportsShaderPipelines(GShaderPlatformForFeatureLevel[InFeatureLevel]) && CVar && CVar->GetValueOnAnyThread() != 0;
 }
 
-template <bool bPositionOnly, bool bUsesMobileColorValue>
+template <bool bPositionOnly>
 bool GetDepthPassShaders(
 	const FMaterial& Material,
 	FVertexFactoryType* VertexFactoryType,
 	ERHIFeatureLevel::Type FeatureLevel,
 	TShaderRef<TDepthOnlyVS<bPositionOnly>>& VertexShader,
-	TShaderRef<FDepthOnlyPS<bUsesMobileColorValue>>& PixelShader,
+	TShaderRef<FDepthOnlyPS>& PixelShader,
 	FShaderPipelineRef& ShaderPipeline)
 {
 	FMaterialShaderTypes ShaderTypes;
 	ShaderTypes.AddShaderType<TDepthOnlyVS<bPositionOnly>>();
 
-	if (bPositionOnly && !bUsesMobileColorValue)
+	if (bPositionOnly)
 	{
 		ShaderTypes.PipelineType = &DepthPosOnlyNoPixelPipeline;
 		/*ShaderPipeline = UseShaderPipelines(FeatureLevel) ? Material.GetShaderPipeline(&DepthPosOnlyNoPixelPipeline, VertexFactoryType) : FShaderPipelineRef();
@@ -177,22 +175,11 @@ bool GetDepthPassShaders(
 	}
 	else
 	{
-		const bool bNeedsPixelShader = bUsesMobileColorValue || !Material.WritesEveryPixel() || Material.MaterialUsesPixelDepthOffset_RenderThread() || Material.IsTranslucencyWritingCustomDepth();
+		const bool bNeedsPixelShader = !Material.WritesEveryPixel() || Material.MaterialUsesPixelDepthOffset_RenderThread() || Material.IsTranslucencyWritingCustomDepth();
 		if (bNeedsPixelShader)
 		{
-			ShaderTypes.AddShaderType<FDepthOnlyPS<bUsesMobileColorValue>>();
-		}
-
-		if (bNeedsPixelShader)
-		{
-			if (bUsesMobileColorValue)
-			{
-				ShaderTypes.PipelineType = &DepthWithColorOutputPipeline;
-			}
-			else
-			{
-				ShaderTypes.PipelineType = &DepthNoColorOutputPipeline;
-			}
+			ShaderTypes.AddShaderType<FDepthOnlyPS>();
+			ShaderTypes.PipelineType = &DepthPipeline;
 		}
 		else
 		{
@@ -212,19 +199,18 @@ bool GetDepthPassShaders(
 	return true;
 }
 
-#define IMPLEMENT_GetDepthPassShaders( bPositionOnly, bUsesMobileColorValue ) \
-	template bool GetDepthPassShaders< bPositionOnly, bUsesMobileColorValue >( \
+#define IMPLEMENT_GetDepthPassShaders( bPositionOnly ) \
+	template bool GetDepthPassShaders< bPositionOnly >( \
 		const FMaterial& Material, \
 		FVertexFactoryType* VertexFactoryType, \
 		ERHIFeatureLevel::Type FeatureLevel, \
 		TShaderRef<TDepthOnlyVS<bPositionOnly>>& VertexShader, \
-		TShaderRef<FDepthOnlyPS<bUsesMobileColorValue>>& PixelShader, \
+		TShaderRef<FDepthOnlyPS>& PixelShader, \
 		FShaderPipelineRef& ShaderPipeline \
 	);
 
-IMPLEMENT_GetDepthPassShaders( true, false );
-IMPLEMENT_GetDepthPassShaders( false, false );
-IMPLEMENT_GetDepthPassShaders( false, true );
+IMPLEMENT_GetDepthPassShaders( true );
+IMPLEMENT_GetDepthPassShaders( false );
 
 void SetDepthPassDitheredLODTransitionState(const FSceneView* SceneView, const FMeshBatch& RESTRICT Mesh, int32 StaticMeshId, FMeshPassProcessorRenderState& DrawRenderState)
 {
@@ -747,11 +733,11 @@ bool FDepthPassMeshProcessor::Process(
 
 	TMeshProcessorShaders<
 		TDepthOnlyVS<bPositionOnly>,
-		FDepthOnlyPS<false>> DepthPassShaders;
+		FDepthOnlyPS> DepthPassShaders;
 
 	FShaderPipelineRef ShaderPipeline;
 
-	if (!GetDepthPassShaders<bPositionOnly, false>(
+	if (!GetDepthPassShaders<bPositionOnly>(
 		MaterialResource,
 		VertexFactory->GetType(),
 		FeatureLevel,
@@ -775,7 +761,7 @@ bool FDepthPassMeshProcessor::Process(
 		SetMobileDepthPassRenderState(PrimitiveSceneProxy, DrawRenderState, MeshBatch, IsMobileDeferredShadingEnabled(GetFeatureLevelShaderPlatform(FeatureLevel)));
 	}
 
-	FDepthOnlyShaderElementData ShaderElementData(0.0f);
+	FMeshMaterialShaderElementData ShaderElementData;
 	ShaderElementData.InitializeMeshMaterialData(ViewIfDynamicMeshCommand, PrimitiveSceneProxy, MeshBatch, StaticMeshId, true);
 
 	const FMeshDrawCommandSortKey SortKey = CalculateDepthPassMeshStaticSortKey(BlendMode, DepthPassShaders.VertexShader.GetShader(), DepthPassShaders.PixelShader.GetShader());
diff --git a/Engine/Source/Runtime/Renderer/Private/DepthRendering.h b/Engine/Source/Runtime/Renderer/Private/DepthRendering.h
--- a/Engine/Source/Runtime/Renderer/Private/DepthRendering.h
+++ b/Engine/Source/Runtime/Renderer/Private/DepthRendering.h
@@ -57,17 +57,6 @@ extern FDepthPassInfo GetDepthPassInfo(const FScene* Scene);
 
 void AddDitheredStencilFillPass(FRDGBuilder& GraphBuilder, TConstArrayView<FViewInfo> Views, FRDGTextureRef DepthTexture, const FDepthPassInfo& DepthPass);
 
-class FDepthOnlyShaderElementData : public FMeshMaterialShaderElementData
-{
-public:
-	FDepthOnlyShaderElementData(float InMobileColorValue)
-		: MobileColorValue(InMobileColorValue)
-	{
-	}
-
-	float MobileColorValue;
-};
-
 /**
  * A vertex shader for rendering the depth of a mesh.
  */
@@ -112,7 +101,7 @@ public:
 		const FMaterialRenderProxy& MaterialRenderProxy,
 		const FMaterial& Material,
 		const FMeshPassProcessorRenderState& DrawRenderState,
-		const FDepthOnlyShaderElementData& ShaderElementData,
+		const FMeshMaterialShaderElementData& ShaderElementData,
 		FMeshDrawSingleShaderBindings& ShaderBindings) const
 	{
 		FMeshMaterialShader::GetShaderBindings(Scene, FeatureLevel, PrimitiveSceneProxy, MaterialRenderProxy, Material, DrawRenderState, ShaderElementData, ShaderBindings);
@@ -122,7 +111,6 @@ public:
 /**
 * A pixel shader for rendering the depth of a mesh.
 */
-template <bool bUsesMobileColorValue>
 class FDepthOnlyPS : public FMeshMaterialShader
 {
 	DECLARE_SHADER_TYPE(FDepthOnlyPS,MeshMaterial);
@@ -135,16 +123,13 @@ public:
 		}
 		
 		return
-			// Compile for materials that are masked, avoid generating permutation for other platforms if bUsesMobileColorValue is true
-			((!Parameters.MaterialParameters.bWritesEveryPixel || Parameters.MaterialParameters.bHasPixelDepthOffsetConnected) && (!bUsesMobileColorValue || IsMobilePlatform(Parameters.Platform)))
-			// Mobile uses material pixel shader to write custom stencil to color target
-			|| (IsMobilePlatform(Parameters.Platform) && (Parameters.MaterialParameters.bIsDefaultMaterial || Parameters.MaterialParameters.bMaterialMayModifyMeshPosition));
+			// Compile for materials that are masked
+			(!Parameters.MaterialParameters.bWritesEveryPixel || Parameters.MaterialParameters.bHasPixelDepthOffsetConnected);
 	}
 
 	FDepthOnlyPS(const ShaderMetaType::CompiledShaderInitializerType& Initializer):
 		FMeshMaterialShader(Initializer)
 	{
-		MobileColorValue.Bind(Initializer.ParameterMap, TEXT("MobileColorValue"));
 	}
 
 	static void ModifyCompilationEnvironment(const FMaterialShaderPermutationParameters& Parameters, FShaderCompilerEnvironment& OutEnvironment)
@@ -152,14 +137,6 @@ public:
 		FMeshMaterialShader::ModifyCompilationEnvironment(Parameters, OutEnvironment);
 		
 		OutEnvironment.SetDefine(TEXT("ALLOW_DEBUG_VIEW_MODES"), AllowDebugViewmodes(Parameters.Platform));
-		if (IsMobilePlatform(Parameters.Platform))
-		{
-			OutEnvironment.SetDefine(TEXT("OUTPUT_MOBILE_COLOR_VALUE"), bUsesMobileColorValue ? 1u : 0u);
-		}
-		else
-		{
-			OutEnvironment.SetDefine(TEXT("OUTPUT_MOBILE_COLOR_VALUE"), 0u);
-		}
 		OutEnvironment.SetDefine(TEXT("SCENE_TEXTURES_DISABLED"), 1u);
 	}
 
@@ -172,24 +149,20 @@ public:
 		const FMaterialRenderProxy& MaterialRenderProxy,
 		const FMaterial& Material,
 		const FMeshPassProcessorRenderState& DrawRenderState,
-		const FDepthOnlyShaderElementData& ShaderElementData,
+		const FMeshMaterialShaderElementData& ShaderElementData,
 		FMeshDrawSingleShaderBindings& ShaderBindings) const
 	{
 		FMeshMaterialShader::GetShaderBindings(Scene, FeatureLevel, PrimitiveSceneProxy, MaterialRenderProxy, Material, DrawRenderState, ShaderElementData, ShaderBindings);
-
-		ShaderBindings.Add(MobileColorValue, ShaderElementData.MobileColorValue);
 	}
-
-	LAYOUT_FIELD(FShaderParameter, MobileColorValue);
 };
 
-template <bool bPositionOnly, bool bUsesMobileColorValue>
+template <bool bPositionOnly>
 bool GetDepthPassShaders(
 	const FMaterial& Material,
 	FVertexFactoryType* VertexFactoryType,
 	ERHIFeatureLevel::Type FeatureLevel,
 	TShaderRef<TDepthOnlyVS<bPositionOnly>>& VertexShader,
-	TShaderRef<FDepthOnlyPS<bUsesMobileColorValue>>& PixelShader,
+	TShaderRef<FDepthOnlyPS>& PixelShader,
 	FShaderPipelineRef& ShaderPipeline);
 
 class FDepthPassMeshProcessor : public FMeshPassProcessor
diff --git a/Engine/Source/Runtime/Renderer/Private/DistortionRendering.cpp b/Engine/Source/Runtime/Renderer/Private/DistortionRendering.cpp
--- a/Engine/Source/Runtime/Renderer/Private/DistortionRendering.cpp
+++ b/Engine/Source/Runtime/Renderer/Private/DistortionRendering.cpp
@@ -207,7 +207,7 @@ public:
 		if (IsMobilePlatform(Parameters.Platform))
 		{
 			// use same path for scene textures as post-process material
-			OutEnvironment.SetDefine(TEXT("POST_PROCESS_MATERIAL_MOBILE"), 1);
+			OutEnvironment.SetDefine(TEXT("POST_PROCESS_MATERIAL"), 1);
 		}
 
 		// Skip the material clip if depth test should not be done
diff --git a/Engine/Source/Runtime/Renderer/Private/MobileShadingRenderer.cpp b/Engine/Source/Runtime/Renderer/Private/MobileShadingRenderer.cpp
--- a/Engine/Source/Runtime/Renderer/Private/MobileShadingRenderer.cpp
+++ b/Engine/Source/Runtime/Renderer/Private/MobileShadingRenderer.cpp
@@ -106,49 +106,59 @@ TGlobalResource<FGlobalDynamicReadBuffer> FMobileSceneRenderer::DynamicReadBuffe
 
 extern bool IsMobileEyeAdaptationEnabled(const FViewInfo& View);
 
-static bool UsesCustomDepthStencilLookup(const FViewInfo& View)
+struct FMobileCustomDepthStencilUsage
 {
 	bool bUsesCustomDepthStencil = false;
+	// whether both CustomDepth and CustomStencil are sampled as a textures
+	bool bSamplesCustomDepthAndStencil = false;
+};
 
-	// Find out whether CustomDepth/Stencil used in translucent materials
-	if (View.bUsesCustomDepthStencilInTranslucentMaterials && CVarMobileCustomDepthForTranslucency.GetValueOnAnyThread() != 0)
-	{
-		bUsesCustomDepthStencil = true;
-	}
-	else
+static FMobileCustomDepthStencilUsage GetCustomDepthStencilUsage(const FViewInfo& View)
+{
+	FMobileCustomDepthStencilUsage CustomDepthStencilUsage;
+
+	// Find out whether there are primitives will render in custom depth pass or just always render custom depth
+	if ((View.bHasCustomDepthPrimitives || GetCustomDepthMode() == ECustomDepthMode::EnabledWithStencil))
 	{
-		// Find out whether post-process materials use CustomDepth/Stencil lookups
-		const FBlendableManager& BlendableManager = View.FinalPostProcessSettings.BlendableManager;
-		FBlendableEntry* BlendableIt = nullptr;
+		// Find out whether CustomDepth/Stencil used in translucent materials
+		if (CVarMobileCustomDepthForTranslucency.GetValueOnAnyThread() != 0)
+		{
+			CustomDepthStencilUsage.bUsesCustomDepthStencil = View.bUsesCustomDepth || View.bUsesCustomStencil;
+			CustomDepthStencilUsage.bSamplesCustomDepthAndStencil = View.bUsesCustomDepth && View.bUsesCustomStencil;
+		}
 
-		while (FPostProcessMaterialNode* DataPtr = BlendableManager.IterateBlendables<FPostProcessMaterialNode>(BlendableIt))
+		if (!CustomDepthStencilUsage.bSamplesCustomDepthAndStencil)
 		{
-			if (DataPtr->IsValid())
+			// Find out whether post-process materials use CustomDepth/Stencil lookups
+			const FBlendableManager& BlendableManager = View.FinalPostProcessSettings.BlendableManager;
+			FBlendableEntry* BlendableIt = nullptr;
+			while (FPostProcessMaterialNode* DataPtr = BlendableManager.IterateBlendables<FPostProcessMaterialNode>(BlendableIt))
 			{
-				FMaterialRenderProxy* Proxy = DataPtr->GetMaterialInterface()->GetRenderProxy();
-				check(Proxy);
-
-				const FMaterial& Material = Proxy->GetIncompleteMaterialWithFallback(View.GetFeatureLevel());
-				if (Material.IsStencilTestEnabled())
+				if (DataPtr->IsValid())
 				{
-					bUsesCustomDepthStencil = true;
-					break;
-				}
-
-				const FMaterialShaderMap* MaterialShaderMap = Material.GetRenderingThreadShaderMap();
-				if (MaterialShaderMap->UsesSceneTexture(PPI_CustomDepth) || MaterialShaderMap->UsesSceneTexture(PPI_CustomStencil))
-				{
-					bUsesCustomDepthStencil = true;
-					break;
+					FMaterialRenderProxy* Proxy = DataPtr->GetMaterialInterface()->GetRenderProxy();
+					check(Proxy);
+
+					const FMaterial& Material = Proxy->GetIncompleteMaterialWithFallback(View.GetFeatureLevel());
+					const FMaterialShaderMap* MaterialShaderMap = Material.GetRenderingThreadShaderMap();
+					bool bUsesCustomDepth = MaterialShaderMap->UsesSceneTexture(PPI_CustomDepth);
+					bool bUsesCustomStencil = MaterialShaderMap->UsesSceneTexture(PPI_CustomStencil);
+					if (Material.IsStencilTestEnabled() || bUsesCustomDepth || bUsesCustomStencil)
+					{
+						CustomDepthStencilUsage.bUsesCustomDepthStencil |= true;
+					}
+
+					if (bUsesCustomDepth && bUsesCustomStencil)
+					{
+						CustomDepthStencilUsage.bSamplesCustomDepthAndStencil |= true;
+						break;
+					}
 				}
 			}
 		}
 	}
-
-	// Find out whether there are primitives will render in custom depth pass or just always render custom depth
-	bUsesCustomDepthStencil &= (View.bHasCustomDepthPrimitives || GetCustomDepthMode() == ECustomDepthMode::EnabledWithStencil);
-
-	return bUsesCustomDepthStencil;
+	
+	return CustomDepthStencilUsage;
 }
 
 static void RenderOpaqueFX(
@@ -514,6 +524,19 @@ void FMobileSceneRenderer::InitViews(FRDGBuilder& GraphBuilder, FSceneTexturesCo
 		SceneTexturesConfig.bPreciseDepthAux = true;
 	}
 	
+	// Find out whether custom depth pass should be rendered.
+	{
+		bool bCouldUseCustomDepthStencil = !bGammaSpace && (!Scene->World || (Scene->World->WorldType != EWorldType::EditorPreview && Scene->World->WorldType != EWorldType::Inactive));
+		FMobileCustomDepthStencilUsage CustomDepthStencilUsage;
+		for (int32 ViewIndex = 0; ViewIndex < Views.Num(); ViewIndex++)
+		{
+			CustomDepthStencilUsage = GetCustomDepthStencilUsage(Views[ViewIndex]);
+			Views[ViewIndex].bCustomDepthStencilValid = bCouldUseCustomDepthStencil && CustomDepthStencilUsage.bUsesCustomDepthStencil;
+			bShouldRenderCustomDepth |= Views[ViewIndex].bCustomDepthStencilValid;
+			SceneTexturesConfig.bSamplesCustomDepthAndStencil |= bShouldRenderCustomDepth && CustomDepthStencilUsage.bSamplesCustomDepthAndStencil;
+		}
+	}
+	
 	// Finalize and set the scene textures config.
 	FSceneTexturesConfig::Set(SceneTexturesConfig);
 
@@ -531,16 +554,6 @@ void FMobileSceneRenderer::InitViews(FRDGBuilder& GraphBuilder, FSceneTexturesCo
 	{
 		ReleaseMobileShadowProjectionOutputs();
 	}
-
-	// Find out whether custom depth pass should be rendered.
-	{
-		bool bCouldUseCustomDepthStencil = !bGammaSpace && (!Scene->World || (Scene->World->WorldType != EWorldType::EditorPreview && Scene->World->WorldType != EWorldType::Inactive));
-		for (int32 ViewIndex = 0; ViewIndex < Views.Num(); ViewIndex++)
-		{
-			Views[ViewIndex].bCustomDepthStencilValid = bCouldUseCustomDepthStencil && UsesCustomDepthStencilLookup(Views[ViewIndex]);
-			bShouldRenderCustomDepth |= Views[ViewIndex].bCustomDepthStencilValid;
-		}
-	}
 	
 	const bool bDynamicShadows = ActiveViewFamily->EngineShowFlags.DynamicShadows;
 	
diff --git a/Engine/Source/Runtime/Renderer/Private/PostProcess/PostProcessMaterial.cpp b/Engine/Source/Runtime/Renderer/Private/PostProcess/PostProcessMaterial.cpp
--- a/Engine/Source/Runtime/Renderer/Private/PostProcess/PostProcessMaterial.cpp
+++ b/Engine/Source/Runtime/Renderer/Private/PostProcess/PostProcessMaterial.cpp
@@ -62,7 +62,7 @@ enum class EMaterialCustomDepthPolicy : uint32
 	Enabled
 };
 
-EMaterialCustomDepthPolicy GetMaterialCustomDepthPolicy(const FMaterial* Material, ERHIFeatureLevel::Type FeatureLevel)
+EMaterialCustomDepthPolicy GetMaterialCustomDepthPolicy(const FMaterial* Material)
 {
 	check(Material);
 
@@ -149,7 +149,7 @@ FRHIBlendState* GetMaterialBlendState(const FMaterial* Material)
 	return BlendStates[Material->GetBlendMode()];
 }
 
-bool PostProcessStencilTest(uint32 StencilValue, uint32 StencilComp, uint32 StencilRef)
+bool PostProcessStencilTest(const uint32 StencilValue, const uint32 StencilComp, const uint32 StencilRef)
 {
 	bool bStencilTestPassed = true;
 
@@ -189,23 +189,11 @@ public:
 	using FParameters = FPostProcessMaterialParameters;
 	SHADER_USE_PARAMETER_STRUCT_WITH_LEGACY_BASE(FPostProcessMaterialShader, FMaterialShader);
 
-	class FMobileDimension : SHADER_PERMUTATION_BOOL("POST_PROCESS_MATERIAL_MOBILE");
-	using FPermutationDomain = TShaderPermutationDomain<FMobileDimension>;
-
 	static bool ShouldCompilePermutation(const FMaterialShaderPermutationParameters& Parameters)
 	{
 		if (Parameters.MaterialParameters.MaterialDomain == MD_PostProcess)
 		{
-			const FPermutationDomain PermutationVector(Parameters.PermutationId);
-
-			if (PermutationVector.Get<FMobileDimension>())
-			{
-				return IsMobilePlatform(Parameters.Platform) && IsMobileHDR();
-			}
-			else
-			{
-				return IsFeatureLevelSupported(Parameters.Platform, ERHIFeatureLevel::SM5) || IsMobileSupportFetchBindedCustomStencilBuffer(Parameters.Platform);
-			}
+			return !IsMobilePlatform(Parameters.Platform) || IsMobileHDR();
 		}
 		return false;
 	}
@@ -218,18 +206,11 @@ public:
 		EBlendableLocation Location = EBlendableLocation(Parameters.MaterialParameters.BlendableLocation);
 		OutEnvironment.SetDefine(TEXT("POST_PROCESS_MATERIAL_BEFORE_TONEMAP"), (Location == BL_AfterTonemapping || Location == BL_ReplacingTonemapper) ? 0 : 1);
 
-		FPermutationDomain PermutationVector(Parameters.PermutationId);
-		if (PermutationVector.Get<FMobileDimension>())
+		if (IsMobilePlatform(Parameters.Platform))
 		{
 			OutEnvironment.SetDefine(TEXT("POST_PROCESS_MATERIAL_BEFORE_TONEMAP"), (Parameters.MaterialParameters.BlendableLocation != BL_AfterTonemapping) ? 1 : 0);
 		}
 
-
-		if (IsMobileSupportFetchBindedCustomStencilBuffer(Parameters.Platform))
-		{
-			OutEnvironment.SetDefine(TEXT("MOBILE_SUPPORT_FETCH_BINDED_CUSTOM_STENCIL_BUFFER"), 1);
-		}
-
 		// PostProcessMaterial can both read & write Strata data
 		OutEnvironment.SetDefine(TEXT("STRATA_INLINE_SHADING"), 1);
 		OutEnvironment.SetDefine(TEXT("STRATA_DEFERRED_SHADING"), 1);
@@ -279,10 +260,6 @@ public:
 	static void ModifyCompilationEnvironment(const FMaterialShaderPermutationParameters& Parameters, FShaderCompilerEnvironment& OutEnvironment)
 	{
 		FPostProcessMaterialShader::ModifyCompilationEnvironment(Parameters, OutEnvironment);
-
-		uint32 StencilCompareFunction = Parameters.MaterialParameters.bIsStencilTestEnabled ? Parameters.MaterialParameters.StencilCompare : EMaterialStencilCompare::MSC_Never;
-
-		OutEnvironment.SetDefine(TEXT("MOBILE_STENCIL_COMPARE_FUNCTION"), StencilCompareFunction);
 	}
 };
 
@@ -320,12 +297,8 @@ void GetMaterialInfo(
 {
 	FMaterialShaderTypes ShaderTypes;
 	{
-		const bool bUseMobilePath = IsMobileSeparateDepthStencilRenderTargets(GShaderPlatformForFeatureLevel[InFeatureLevel]);
-		FPostProcessMaterialShader::FPermutationDomain PermutationVector;
-		PermutationVector.Set<FPostProcessMaterialShader::FMobileDimension>(bUseMobilePath);
-
-		ShaderTypes.AddShaderType< FPostProcessMaterialVS>(PermutationVector.ToDimensionValueId());
-		ShaderTypes.AddShaderType< FPostProcessMaterialPS>(PermutationVector.ToDimensionValueId());
+		ShaderTypes.AddShaderType< FPostProcessMaterialVS>();
+		ShaderTypes.AddShaderType< FPostProcessMaterialPS>();
 	}
 
 	const FMaterialRenderProxy* MaterialProxy = InMaterialInterface->GetRenderProxy();
@@ -448,7 +421,7 @@ FScreenPassTexture AddPostProcessMaterialPass(
 	FRDGTextureRef DepthStencilTexture = nullptr;
 
 	// Allocate custom depth stencil texture(s) and depth stencil state.
-	const EMaterialCustomDepthPolicy CustomStencilPolicy = GetMaterialCustomDepthPolicy(Material, FeatureLevel);
+	const EMaterialCustomDepthPolicy CustomStencilPolicy = GetMaterialCustomDepthPolicy(Material);
 
 	if (CustomStencilPolicy == EMaterialCustomDepthPolicy::Enabled && HasBeenProduced(Inputs.CustomDepthTexture))
 	{
@@ -538,7 +511,6 @@ FScreenPassTexture AddPostProcessMaterialPass(
 	const uint32 MaterialStencilRef = Material->GetStencilRefValue();
 
 	const bool bMobilePlatform = IsMobilePlatform(View.GetShaderPlatform());
-	const bool bSupportFetchBindedCustomStencilBuffer = !bMobilePlatform || IsMobileSupportFetchBindedCustomStencilBuffer(View.GetShaderPlatform());
 
 	FPostProcessMaterialParameters* PostProcessMaterialParameters = GraphBuilder.AllocParameters<FPostProcessMaterialParameters>();
 	PostProcessMaterialParameters->SceneTextures = Inputs.SceneTextures;
@@ -552,15 +524,12 @@ FScreenPassTexture AddPostProcessMaterialPass(
 		PostProcessMaterialParameters->EyeAdaptationTexture = GetEyeAdaptationTexture(GraphBuilder, View);
 	}
 	PostProcessMaterialParameters->PostProcessOutput = GetScreenPassTextureViewportParameters(OutputViewport);
-	PostProcessMaterialParameters->MobileCustomStencilTexture = DepthStencilTexture;
-	PostProcessMaterialParameters->MobileCustomStencilTextureSampler = TStaticSamplerState<SF_Point, AM_Clamp, AM_Clamp, AM_Clamp>::GetRHI();
-	PostProcessMaterialParameters->MobileStencilValueRef = MaterialStencilRef;
 	PostProcessMaterialParameters->RenderTargets[0] = Output.GetRenderTargetBinding();
 
 	// The target color will be decoded if bForceIntermediateTarget is true in any case, but we might still need to decode the input color
 	PostProcessMaterialParameters->bMetalMSAAHDRDecode = Inputs.bMetalMSAAHDRDecode ? 1 : 0;
 
-	if (DepthStencilTexture && bSupportFetchBindedCustomStencilBuffer)
+	if (DepthStencilTexture)
 	{
 		PostProcessMaterialParameters->RenderTargets.DepthStencil = FDepthStencilBinding(
 			DepthStencilTexture,
@@ -568,31 +537,6 @@ FScreenPassTexture AddPostProcessMaterialPass(
 			ERenderTargetLoadAction::ELoad,
 			FExclusiveDepthStencil::DepthRead_StencilRead);
 	}
-	else if (!DepthStencilTexture && IsMobileSeparateDepthStencilRenderTargets(View.GetShaderPlatform()) && Material->IsStencilTestEnabled()) // we have to set a default texture for MobileStencilTexture and override the MobileStencilValueRef to make all function to pass the stencil test
-	{
-		PostProcessMaterialParameters->MobileCustomStencilTexture = GSystemTextures.GetBlackDummy(GraphBuilder);
-		
-		switch (Material->GetStencilCompare())
-		{
-		case EMaterialStencilCompare::MSC_Less:
-			PostProcessMaterialParameters->MobileStencilValueRef = -1;
-			break;
-		case EMaterialStencilCompare::MSC_LessEqual:
-		case EMaterialStencilCompare::MSC_GreaterEqual:
-		case EMaterialStencilCompare::MSC_Equal:
-			PostProcessMaterialParameters->MobileStencilValueRef = 0;
-			break;
-		case EMaterialStencilCompare::MSC_Greater:
-		case EMaterialStencilCompare::MSC_NotEqual:
-			PostProcessMaterialParameters->MobileStencilValueRef = 1;
-			break;
-		case EMaterialStencilCompare::MSC_Always:
-			PostProcessMaterialParameters->MobileStencilValueRef = 256;
-			break;
-		default:
-			break;
-		}
-	}
 
 	PostProcessMaterialParameters->PostProcessInput_BilinearSampler = TStaticSamplerState<SF_Bilinear, AM_Clamp, AM_Clamp, AM_Clamp>::GetRHI();;
 
@@ -635,14 +579,16 @@ FScreenPassTexture AddPostProcessMaterialPass(
 	{
 		bool bFailStencil = true;
 
+		const uint32 StencilComp = Material->GetStencilCompare();
+
 		// Always check against clear value, since a material might want to perform operations against that value
-		uint32 StencilClearValue = Inputs.CustomDepthTexture ? Inputs.CustomDepthTexture->Desc.ClearValue.Value.DSValue.Stencil : 0;
-		bFailStencil &= PostProcessStencilTest(StencilClearValue, Material->GetStencilCompare(), PostProcessMaterialParameters->MobileStencilValueRef);
+		const uint32 StencilClearValue = Inputs.CustomDepthTexture ? Inputs.CustomDepthTexture->Desc.ClearValue.Value.DSValue.Stencil : 0;
+		bFailStencil &= PostProcessStencilTest(StencilClearValue, StencilComp, MaterialStencilRef);
 
 
 		for (const uint32& Value : View.CustomDepthStencilValues)
 		{
-			bFailStencil &= PostProcessStencilTest(Value, Material->GetStencilCompare(), PostProcessMaterialParameters->MobileStencilValueRef);
+			bFailStencil &= PostProcessStencilTest(Value, StencilComp, MaterialStencilRef);
 
 			if (!bFailStencil)
 			{
@@ -662,10 +608,10 @@ FScreenPassTexture AddPostProcessMaterialPass(
 		OutputViewport,
 		SceneColorViewport,
 		// Uses default depth stencil on mobile since the stencil test is done in pixel shader.
-		FScreenPassPipelineState(VertexShader, PixelShader, BlendState, !bSupportFetchBindedCustomStencilBuffer ? DefaultDepthStencilState : DepthStencilState, MaterialStencilRef),
+		FScreenPassPipelineState(VertexShader, PixelShader, BlendState, DepthStencilState, MaterialStencilRef),
 		PostProcessMaterialParameters,
 		ScreenPassFlags,
-		[&View, VertexShader, PixelShader, MaterialRenderProxy, Material, PostProcessMaterialParameters, MaterialStencilRef](FRHICommandList& RHICmdList)
+		[&View, VertexShader, PixelShader, MaterialRenderProxy, Material, PostProcessMaterialParameters](FRHICommandList& RHICmdList)
 	{
 		FPostProcessMaterialVS::SetParameters(RHICmdList, VertexShader, View, MaterialRenderProxy, *Material, *PostProcessMaterialParameters);
 		FPostProcessMaterialPS::SetParameters(RHICmdList, PixelShader, View, MaterialRenderProxy, *Material, *PostProcessMaterialParameters);
diff --git a/Engine/Source/Runtime/Renderer/Private/PostProcess/PostProcessMaterial.h b/Engine/Source/Runtime/Renderer/Private/PostProcess/PostProcessMaterial.h
--- a/Engine/Source/Runtime/Renderer/Private/PostProcess/PostProcessMaterial.h
+++ b/Engine/Source/Runtime/Renderer/Private/PostProcess/PostProcessMaterial.h
@@ -43,11 +43,8 @@ BEGIN_SHADER_PARAMETER_STRUCT(FPostProcessMaterialParameters, )
 	SHADER_PARAMETER_STRUCT(FScreenPassTextureViewportParameters, PostProcessOutput)
 	SHADER_PARAMETER_STRUCT_ARRAY(FScreenPassTextureInput, PostProcessInput, [kPostProcessMaterialInputCountMax])
 	SHADER_PARAMETER_SAMPLER(SamplerState, PostProcessInput_BilinearSampler)
-	SHADER_PARAMETER_RDG_TEXTURE(Texture2D, MobileCustomStencilTexture)
-	SHADER_PARAMETER_SAMPLER(SamplerState, MobileCustomStencilTextureSampler)
 	SHADER_PARAMETER_RDG_TEXTURE(Texture2D, EyeAdaptationTexture)
 	SHADER_PARAMETER_RDG_BUFFER_SRV(Buffer<float4>, EyeAdaptationBuffer)
-	SHADER_PARAMETER(int32, MobileStencilValueRef)
 	SHADER_PARAMETER(uint32, bFlipYAxis)
 	SHADER_PARAMETER(uint32, bMetalMSAAHDRDecode)
 	RENDER_TARGET_BINDING_SLOTS()
diff --git a/Engine/Source/Runtime/Renderer/Private/SceneRendering.cpp b/Engine/Source/Runtime/Renderer/Private/SceneRendering.cpp
--- a/Engine/Source/Runtime/Renderer/Private/SceneRendering.cpp
+++ b/Engine/Source/Runtime/Renderer/Private/SceneRendering.cpp
@@ -964,7 +964,8 @@ void FViewInfo::Init()
 	bHasDistortionPrimitives = false;
 	bAllowStencilDither = false;
 	bCustomDepthStencilValid = false;
-	bUsesCustomDepthStencilInTranslucentMaterials = false;
+	bUsesCustomDepth = false;
+	bUsesCustomStencil = false;
 
 	NumBoxReflectionCaptures = 0;
 	NumSphereReflectionCaptures = 0;
diff --git a/Engine/Source/Runtime/Renderer/Private/SceneRendering.h b/Engine/Source/Runtime/Renderer/Private/SceneRendering.h
--- a/Engine/Source/Runtime/Renderer/Private/SceneRendering.h
+++ b/Engine/Source/Runtime/Renderer/Private/SceneRendering.h
@@ -1370,7 +1370,8 @@ public:
 	/** Whether the view has any materials that read from scene depth. */
 	uint32 bUsesSceneDepth : 1;
 	uint32 bCustomDepthStencilValid : 1;
-	uint32 bUsesCustomDepthStencilInTranslucentMaterials : 1;
+	uint32 bUsesCustomDepth : 1;
+	uint32 bUsesCustomStencil : 1;
 
 	/** Whether fog should only be computed on rendered opaque pixels or not. */
 	uint32 bFogOnlyOnRenderedOpaque : 1;
@@ -2258,7 +2259,7 @@ protected:
 	/** Updates state for the end of the frame. */
 	void RenderFinish(FRDGBuilder& GraphBuilder, FRDGTextureRef ViewFamilyTexture);
 
-	bool RenderCustomDepthPass(FRDGBuilder& GraphBuilder, const FCustomDepthTextures& CustomDepthTextures, const FSceneTextureShaderParameters& SceneTextures);
+	bool RenderCustomDepthPass(FRDGBuilder& GraphBuilder, FCustomDepthTextures& CustomDepthTextures, const FSceneTextureShaderParameters& SceneTextures);
 
 	void OnStartRender(FRHICommandListImmediate& RHICmdList);
 
diff --git a/Engine/Source/Runtime/Renderer/Private/SceneTextures.cpp b/Engine/Source/Runtime/Renderer/Private/SceneTextures.cpp
--- a/Engine/Source/Runtime/Renderer/Private/SceneTextures.cpp
+++ b/Engine/Source/Runtime/Renderer/Private/SceneTextures.cpp
@@ -553,7 +553,6 @@ void FSceneTexturesConfig::InitializeViewFamily(FViewFamilyInfo& ViewFamily)
 	Config.ColorFormat = PF_Unknown;
 	Config.ColorClearValue = FClearValueBinding::Black;
 	Config.DepthClearValue = FClearValueBinding::DepthFar;
-	Config.CustomDepthDownsampleFactor = GetCustomDepthDownsampleFactor(Config.ShaderPlatform);
 	Config.bRequireMultiView = ViewFamily.bRequireMultiView;
 	Config.bIsUsingGBuffers = IsUsingGBuffers(Config.ShaderPlatform);
 
@@ -718,7 +717,7 @@ void FMinimalSceneTextures::InitializeViewFamily(FRDGBuilder& GraphBuilder, FVie
 	}
 
 	// Custom Depth
-	SceneTextures.CustomDepth = FCustomDepthTextures::Create(GraphBuilder, Config.Extent, Config.FeatureLevel, Config.CustomDepthDownsampleFactor);
+	SceneTextures.CustomDepth = FCustomDepthTextures::Create(GraphBuilder, Config.Extent);
 
 	ViewFamily.bIsSceneTexturesInitialized = true;
 }
@@ -940,8 +939,6 @@ void FSceneTextureExtracts::QueueExtractions(FRDGBuilder& GraphBuilder, const FS
 	{
 		SetupMode |= ESceneTextureSetupMode::CustomDepth;
 		ExtractIfProduced(SceneTextures.CustomDepth.Depth, CustomDepth);
-		ExtractIfProduced(SceneTextures.CustomDepth.MobileDepth, MobileCustomDepth);
-		ExtractIfProduced(SceneTextures.CustomDepth.MobileStencil, MobileCustomStencil);
 	}
 
 	// Create and extract a scene texture uniform buffer for RHI code outside of the main render graph instance. This
@@ -971,8 +968,6 @@ void FSceneTextureExtracts::Release()
 {
 	Depth = {};
 	CustomDepth = {};
-	MobileCustomDepth = {};
-	MobileCustomStencil = {};
 	UniformBuffer = {};
 	MobileUniformBuffer = {};
 }
@@ -1114,7 +1109,6 @@ void SetupMobileSceneTextureUniformParameters(
 	FMobileSceneTextureUniformParameters& SceneTextureParameters)
 {
 	const FRDGSystemTextures& SystemTextures = FRDGSystemTextures::Get(GraphBuilder);
-	bool bMobileSupportFetchBindedCustomStencilBuffer = FDataDrivenShaderPlatformInfo::GetMobileSupportFetchBindedCustomStencilBuffer(GShaderPlatformForFeatureLevel[GMaxRHIFeatureLevel]);
 
 	SceneTextureParameters.SceneColorTexture = SystemTextures.Black;
 	SceneTextureParameters.SceneColorTextureSampler = TStaticSamplerState<>::GetRHI();
@@ -1124,8 +1118,6 @@ void SetupMobileSceneTextureUniformParameters(
 	SceneTextureParameters.CustomDepthTexture = SystemTextures.Black;
 	SceneTextureParameters.CustomDepthTextureSampler = TStaticSamplerState<>::GetRHI();
 	SceneTextureParameters.CustomStencilTexture = SystemTextures.StencilDummySRV;
-	SceneTextureParameters.MobileCustomStencilTexture = SystemTextures.Black;
-	SceneTextureParameters.MobileCustomStencilTextureSampler = TStaticSamplerState<>::GetRHI();
 	SceneTextureParameters.SceneVelocityTexture = SystemTextures.Black;
 	SceneTextureParameters.SceneVelocityTextureSampler = TStaticSamplerState<>::GetRHI();
 	SceneTextureParameters.GBufferATexture = SystemTextures.Black;
@@ -1189,24 +1181,9 @@ void SetupMobileSceneTextureUniformParameters(
 		{
 			const FCustomDepthTextures& CustomDepthTextures = SceneTextures->CustomDepth;
 
-			if (bMobileSupportFetchBindedCustomStencilBuffer)
-			{
-				bool bCustomDepthProduced = HasBeenProduced(CustomDepthTextures.Depth);
-				SceneTextureParameters.CustomDepthTexture = bCustomDepthProduced ? CustomDepthTextures.Depth : SystemTextures.DepthDummy;
-				SceneTextureParameters.CustomStencilTexture = bCustomDepthProduced ? CustomDepthTextures.Stencil : SystemTextures.StencilDummySRV;
-			}				
-			else
-			{
-				if (HasBeenProduced(CustomDepthTextures.MobileDepth))
-				{
-					SceneTextureParameters.CustomDepthTexture = CustomDepthTextures.MobileDepth;
-				}
-
-				if (HasBeenProduced(CustomDepthTextures.MobileStencil) && !EnumHasAnyFlags(CustomDepthTextures.MobileStencil->Desc.Flags, TexCreate_Memoryless))
-				{
-					SceneTextureParameters.MobileCustomStencilTexture = CustomDepthTextures.MobileStencil;
-				}
-			}
+			bool bCustomDepthProduced = HasBeenProduced(CustomDepthTextures.Depth);
+			SceneTextureParameters.CustomDepthTexture = bCustomDepthProduced ? CustomDepthTextures.Depth : SystemTextures.DepthDummy;
+			SceneTextureParameters.CustomStencilTexture = bCustomDepthProduced ? CustomDepthTextures.Stencil : SystemTextures.StencilDummySRV;
 		}
 
 		if (EnumHasAnyFlags(SetupMode, EMobileSceneTextureSetupMode::SceneVelocity))
diff --git a/Engine/Source/Runtime/Renderer/Private/SceneTextures.h b/Engine/Source/Runtime/Renderer/Private/SceneTextures.h
--- a/Engine/Source/Runtime/Renderer/Private/SceneTextures.h
+++ b/Engine/Source/Runtime/Renderer/Private/SceneTextures.h
@@ -52,6 +52,7 @@ struct RENDERER_API FSceneTexturesConfig
 		, bIsUsingGBuffers{}
 		, bKeepDepthContent{1}
 		, bPreciseDepthAux{}
+		, bSamplesCustomDepthAndStencil{}
 	{}
 
 	// Extractions to queue for after execution of the render graph.
@@ -70,7 +71,6 @@ struct RENDERER_API FSceneTexturesConfig
 
 	// Downsample factors to divide against the full resolution texture extent.
 	uint32 SmallDepthDownsampleFactor = 2;
-	uint32 CustomDepthDownsampleFactor = 1;
 
 	// Number of MSAA samples used by color and depth targets.
 	uint32 NumSamples = 1;
@@ -106,6 +106,9 @@ struct RENDERER_API FSceneTexturesConfig
 	// (Mobile) True if SceneDepthAux should use a precise pixel format
 	uint32 bPreciseDepthAux : 1;
 
+	// (Mobile) True if both CustomDepth and CustomStencil are sampled in a shader
+	uint32 bSamplesCustomDepthAndStencil : 1;
+
 private:
 	static FSceneTexturesConfig GlobalInstance;
 };
diff --git a/Engine/Source/Runtime/Renderer/Private/SceneVisibility.cpp b/Engine/Source/Runtime/Renderer/Private/SceneVisibility.cpp
--- a/Engine/Source/Runtime/Renderer/Private/SceneVisibility.cpp
+++ b/Engine/Source/Runtime/Renderer/Private/SceneVisibility.cpp
@@ -2273,7 +2273,8 @@ struct FRelevancePacket
 	bool bUsesLightingChannels;
 	bool bTranslucentSurfaceLighting;
 	bool bUsesSceneDepth;
-	bool bUsesCustomDepthStencil;
+	bool bUsesCustomDepth;
+	bool bUsesCustomStencil;
 	bool bSceneHasSkyMaterial;
 	bool bHasSingleLayerWaterMaterial;
 	bool bHasTranslucencySeparateModulation;
@@ -2309,7 +2310,8 @@ struct FRelevancePacket
 		, bUsesLightingChannels(false)
 		, bTranslucentSurfaceLighting(false)
 		, bUsesSceneDepth(false)
-		, bUsesCustomDepthStencil(false)
+		, bUsesCustomDepth(false)
+		, bUsesCustomStencil(false)
 		, bSceneHasSkyMaterial(false)
 		, bHasSingleLayerWaterMaterial(false)
 		, bHasTranslucencySeparateModulation(false)
@@ -2452,7 +2454,8 @@ struct FRelevancePacket
 			bUsesLightingChannels |= ViewRelevance.bUsesLightingChannels;
 			bTranslucentSurfaceLighting |= ViewRelevance.bTranslucentSurfaceLighting;
 			bUsesSceneDepth |= ViewRelevance.bUsesSceneDepth;
-			bUsesCustomDepthStencil |= ViewRelevance.bUsesCustomDepthStencil;
+			bUsesCustomDepth |= (ViewRelevance.CustomDepthStencilUsageMask & 1) > 0;
+			bUsesCustomStencil |= (ViewRelevance.CustomDepthStencilUsageMask & (1 << 1)) > 0;
 			bSceneHasSkyMaterial |= ViewRelevance.bUsesSkyMaterial;
 			bHasSingleLayerWaterMaterial |= ViewRelevance.bUsesSingleLayerWaterMaterial;
 			bHasTranslucencySeparateModulation |= ViewRelevance.bSeparateTranslucencyModulate;
@@ -2913,7 +2916,8 @@ struct FRelevancePacket
 		WriteView.bHasDistortionPrimitives |= bHasDistortionPrimitives;
 		WriteView.bHasCustomDepthPrimitives |= bHasCustomDepthPrimitives;
 		WriteView.CustomDepthStencilValues.Append(CustomDepthStencilValues);
-		WriteView.bUsesCustomDepthStencilInTranslucentMaterials |= bUsesCustomDepthStencil;
+		WriteView.bUsesCustomDepth |= bUsesCustomDepth;
+		WriteView.bUsesCustomStencil |= bUsesCustomStencil;
 		DirtyIndirectLightingCacheBufferPrimitives.AppendTo(WriteView.DirtyIndirectLightingCacheBufferPrimitives);
 
 		WriteView.MeshDecalBatches.Append(MeshDecalBatches);
diff --git a/Engine/Source/Runtime/Renderer/Public/SceneRenderTargetParameters.h b/Engine/Source/Runtime/Renderer/Public/SceneRenderTargetParameters.h
--- a/Engine/Source/Runtime/Renderer/Public/SceneRenderTargetParameters.h
+++ b/Engine/Source/Runtime/Renderer/Public/SceneRenderTargetParameters.h
@@ -81,8 +81,6 @@ BEGIN_GLOBAL_SHADER_PARAMETER_STRUCT(FMobileSceneTextureUniformParameters, RENDE
 	SHADER_PARAMETER_RDG_TEXTURE(Texture2D, CustomDepthTexture)
 	SHADER_PARAMETER_SAMPLER(SamplerState, CustomDepthTextureSampler)
 	SHADER_PARAMETER_RDG_TEXTURE_SRV(Texture2D<uint2>, CustomStencilTexture)
-	SHADER_PARAMETER_RDG_TEXTURE(Texture2D, MobileCustomStencilTexture)
-	SHADER_PARAMETER_SAMPLER(SamplerState, MobileCustomStencilTextureSampler)
 	SHADER_PARAMETER_RDG_TEXTURE(Texture2D, SceneVelocityTexture)
 	SHADER_PARAMETER_SAMPLER(SamplerState, SceneVelocityTextureSampler)
 	// GBuffer
@@ -192,8 +190,6 @@ private:
 
 	// Contains the custom depth targets.
 	TRefCountPtr<IPooledRenderTarget> CustomDepth;
-	TRefCountPtr<IPooledRenderTarget> MobileCustomDepth;
-	TRefCountPtr<IPooledRenderTarget> MobileCustomStencil;
 
 	// Contains RHI scene texture uniform buffers referencing the extracted textures.
 	TUniformBufferRef<FSceneTextureUniformParameters> UniformBuffer;
