Remove SRGB flag when creating RVT PF_G16 texture.
Fixes issue with texture copy to physical texture for world height RVT.
Issue was triggered by change in RHICopyTexture behavior on some platforms.
#rnx

diff --git a/Engine/Source/Runtime/Renderer/Private/VT/VirtualTexturePhysicalSpace.cpp b/Engine/Source/Runtime/Renderer/Private/VT/VirtualTexturePhysicalSpace.cpp
--- a/Engine/Source/Runtime/Renderer/Private/VT/VirtualTexturePhysicalSpace.cpp
+++ b/Engine/Source/Runtime/Renderer/Private/VT/VirtualTexturePhysicalSpace.cpp
@@ -118,9 +118,11 @@ void FVirtualTexturePhysicalSpace::InitRHI()
 		const EPixelFormat FormatUAV = GetUnorderedAccessViewFormat(FormatSRV);
 		const bool bCreateAliasedUAV = (FormatUAV != PF_Unknown) && (FormatUAV != FormatSRV);
 		
-		const bool b16BitFormat = (FormatSRV == EPixelFormat::PF_R5G6B5_UNORM) || (FormatSRV == EPixelFormat::PF_B5G5R5A1_UNORM);
-		// Not all mobile RHIs support sRGB texture views/aliasing, use only linear targets on mobile, and the 16bit format transfer the color space manually in shader
-		ETextureCreateFlags VT_SRGB = (b16BitFormat || GetFeatureLevel() <= ERHIFeatureLevel::ES3_1) ? TexCreate_None : TexCreate_SRGB;
+		// Not all mobile RHIs support sRGB views/aliasing, use only linear targets on mobile.
+		const bool bFeatureLevelSupportsSRGB = GetFeatureLevel() > ERHIFeatureLevel::ES3_1;
+		// Not all formats support sRGB.
+		const bool bFormatSupportsSRGB = FormatSRV != EPixelFormat::PF_R5G6B5_UNORM && FormatSRV != EPixelFormat::PF_B5G5R5A1_UNORM && FormatSRV != EPixelFormat::PF_G16;
+		ETextureCreateFlags VT_SRGB = (bFeatureLevelSupportsSRGB && bFormatSupportsSRGB) ? TexCreate_SRGB : TexCreate_None;
 
 		// Allocate physical texture from the render target pool
 		const uint32 TextureSize = GetTextureSize();
