Removed depth texture binding in CopySceneToCubeFace pass on Mobile platforms

diff --git a/Engine/Source/Runtime/Renderer/Private/ReflectionEnvironmentCapture.cpp b/Engine/Source/Runtime/Renderer/Private/ReflectionEnvironmentCapture.cpp
--- a/Engine/Source/Runtime/Renderer/Private/ReflectionEnvironmentCapture.cpp
+++ b/Engine/Source/Runtime/Renderer/Private/ReflectionEnvironmentCapture.cpp
@@ -546,8 +546,12 @@ void CaptureSceneToScratchCubemap(
 		PassParameters->View = View.ViewUniformBuffer;
 		PassParameters->SceneColorSampler = TStaticSamplerState<SF_Bilinear, AM_Clamp, AM_Clamp, AM_Clamp>::GetRHI();
 		PassParameters->SceneColorTexture = SceneTextures.Color.Resolve;
-		PassParameters->SceneDepthSampler = TStaticSamplerState<SF_Point, AM_Clamp, AM_Clamp, AM_Clamp>::GetRHI();
-		PassParameters->SceneDepthTexture = SceneTextures.Depth.Resolve;
+        
+        if (!IsMobilePlatform(SceneRenderer->ShaderPlatform))
+        {
+            PassParameters->SceneDepthSampler = TStaticSamplerState<SF_Point, AM_Clamp, AM_Clamp, AM_Clamp>::GetRHI();
+            PassParameters->SceneDepthTexture = SceneTextures.Depth.Resolve;
+        }
 
 		const int32 EffectiveSize = CubemapSize;
 		const FIntPoint SceneTextureExtent = SceneTextures.Config.Extent;
