Limit dynamic shadow quality to a 1-3 range for a mobile deferred path
Add permutation for planar reflections

diff --git a/Engine/Shaders/Private/MobileDeferredShading.usf b/Engine/Shaders/Private/MobileDeferredShading.usf
--- a/Engine/Shaders/Private/MobileDeferredShading.usf
+++ b/Engine/Shaders/Private/MobileDeferredShading.usf
@@ -180,8 +180,7 @@ void MobileDirectionalLightPS(
 	half NoL = 0.0f;
 	AccumulateDirectionalLighting(GBuffer, ShadingModelContext, CameraVector, ScreenPosition, SvPosition, Shadow, DynamicShadowFactors, NoL, MobileLightAccumulator);
 
-
-#if ENABLE_CLUSTERED_REFLECTION || ENABLE_SKY_LIGHT
+#if ENABLE_CLUSTERED_REFLECTION || ENABLE_SKY_LIGHT || ENABLE_PLANAR_REFLECTION
 	// If we have a single directional light, apply relfection and sky contrubution here
 	ReflectionEnvironmentSkyLighting(ShadingModelContext, GBuffer, TranslatedWorldPosition, ReflectionVector, GridIndex, MobileLightAccumulator);
 #endif
diff --git a/Engine/Shaders/Private/MobileLightingCommon.ush b/Engine/Shaders/Private/MobileLightingCommon.ush
--- a/Engine/Shaders/Private/MobileLightingCommon.ush
+++ b/Engine/Shaders/Private/MobileLightingCommon.ush
@@ -20,6 +20,10 @@
 #define LQ_TEXTURE_LIGHTMAP 0
 #endif
 
+#ifndef ENABLE_PLANAR_REFLECTION
+#define ENABLE_PLANAR_REFLECTION 0
+#endif 
+
 #define REFLECTION_COMPOSITE_USE_BLENDED_REFLECTION_CAPTURES 1
 #define REFLECTION_COMPOSITE_SUPPORT_SKYLIGHT_BLEND 0 // Adds additional sampler
 #define REFLECTION_COMPOSITE_HAS_SPHERE_CAPTURES 1
@@ -95,9 +99,14 @@ return ShadowMap;
 // Add fading CSM plane:
 #define FADE_CSM 1
 
+#ifndef ENABLE_MOBILE_CSM
+#define ENABLE_MOBILE_CSM 1
+#endif
+
 half MobileDirectionalLightCSM(float2 ScreenPosition, float SceneDepth, half NoL, inout float ShadowPositionZ)
 {
 	half ShadowMap = 1;
+#if ENABLE_MOBILE_CSM
 	ShadowPositionZ = 0;
 	FPCFSamplerSettings Settings;
 	Settings.ShadowDepthTexture = MobileDirectionalLight.DirectionalLightShadowTexture;
@@ -144,6 +153,7 @@ half MobileDirectionalLightCSM(float2 ScreenPosition, float SceneDepth, half NoL
 			ShadowMap = lerp(ShadowMap, 1.0, Fade * Fade);
 		#endif
 	}
+#endif //ENABLE_MOBILE_CSM
 	return ShadowMap;
 }
 
@@ -182,7 +192,7 @@ void AccumulateDirectionalLighting(FGBufferData GBuffer, FMobileShadingModelCont
 		{		
 			Shadow = ApplyPrecomputedShadowMask(ShadowMap, Shadow);
 		}
-	#elif MOBILE_DEFERRED_SHADING && APPLY_CSM
+	#elif MOBILE_DEFERRED_SHADING
 		Shadow = min(ShadowMap, Shadow);
 	#endif
 #endif
@@ -332,7 +342,7 @@ void AccumulateReflection(FGBufferData GBuffer
 		, GridIndex
 	);
 
-#if MOBILE_DEFERRED_LIGHTING || MATERIAL_PLANAR_FORWARD_REFLECTIONS
+#if ENABLE_PLANAR_REFLECTION || MATERIAL_PLANAR_FORWARD_REFLECTIONS
 	BRANCH
 	if (abs(dot(PlanarReflectionStruct.ReflectionPlane.xyz, 1)) > .0001f)
 	{
diff --git a/Engine/Source/Runtime/Renderer/Private/MobileDeferredShadingPass.cpp b/Engine/Source/Runtime/Renderer/Private/MobileDeferredShadingPass.cpp
--- a/Engine/Source/Runtime/Renderer/Private/MobileDeferredShadingPass.cpp
+++ b/Engine/Source/Runtime/Renderer/Private/MobileDeferredShadingPass.cpp
@@ -38,16 +38,18 @@ class FMobileDirectionalLightFunctionPS : public FMaterialShader
 	class FEnableShadingModelSupport	: SHADER_PERMUTATION_BOOL("ENABLE_SHADINGMODEL_SUPPORT_MOBILE_DEFERRED");
 	class FEnableClustredLights			: SHADER_PERMUTATION_BOOL("ENABLE_CLUSTERED_LIGHTS");
 	class FEnableClustredReflection		: SHADER_PERMUTATION_BOOL("ENABLE_CLUSTERED_REFLECTION");
+	class FEnablePlanarReflection		: SHADER_PERMUTATION_BOOL("ENABLE_PLANAR_REFLECTION");
 	class FEnableSkyLight				: SHADER_PERMUTATION_BOOL("ENABLE_SKY_LIGHT");
-	class FApplyCSM						: SHADER_PERMUTATION_BOOL("APPLY_CSM");
-	class FShadowQuality				: SHADER_PERMUTATION_INT("MOBILE_SHADOW_QUALITY", 4);
+	class FEnableCSM					: SHADER_PERMUTATION_BOOL("ENABLE_MOBILE_CSM");
+	class FShadowQuality				: SHADER_PERMUTATION_RANGE_INT("MOBILE_SHADOW_QUALITY", 1, 3); // not using Quality=0
 	
 	using FPermutationDomain = TShaderPermutationDomain<
 		FEnableShadingModelSupport,
 		FEnableClustredLights, 
 		FEnableClustredReflection, 
+		FEnablePlanarReflection,
 		FEnableSkyLight,
-		FApplyCSM, 
+		FEnableCSM, 
 		FShadowQuality>;
 
 	BEGIN_SHADER_PARAMETER_STRUCT(FParameters, )
@@ -74,12 +76,12 @@ class FMobileDirectionalLightFunctionPS : public FMaterialShader
 	{
 		if (MobileUsesShadowMaskTexture(Platform))
 		{
-			PermutationVector.Set<FApplyCSM>(false);
+			PermutationVector.Set<FEnableCSM>(false);
 		}
 
-		if (PermutationVector.Get<FApplyCSM>() == false)
+		if (PermutationVector.Get<FEnableCSM>() == false)
 		{
-			PermutationVector.Set<FShadowQuality>(0);
+			PermutationVector.Set<FShadowQuality>(1);
 		}
 
 		if (!MobileUsesGBufferCustomData(Platform))
@@ -109,7 +111,7 @@ class FMobileDirectionalLightFunctionPS : public FMaterialShader
 		return true;
 	}
 
-	static FPermutationDomain BuildPermutationVector(const FViewInfo& View, bool bInlineReflectionAndSky, bool bShadingModelSupport, bool bDynamicShadows, bool bSkyLight)
+	static FPermutationDomain BuildPermutationVector(const FViewInfo& View, bool bInlineReflectionAndSky, bool bShadingModelSupport, bool bDynamicShadows, bool bSkyLight, bool bPlanarReflection)
 	{
 		bool bUseClusteredLights = GMobileUseClusteredDeferredShading != 0;
 		bool bClustredReflection = bInlineReflectionAndSky && (View.NumBoxReflectionCaptures + View.NumSphereReflectionCaptures) > 0;
@@ -121,9 +123,10 @@ class FMobileDirectionalLightFunctionPS : public FMaterialShader
 		PermutationVector.Set<FMobileDirectionalLightFunctionPS::FEnableShadingModelSupport>(bShadingModelSupport);
 		PermutationVector.Set<FMobileDirectionalLightFunctionPS::FEnableClustredLights>(bUseClusteredLights);
 		PermutationVector.Set<FMobileDirectionalLightFunctionPS::FEnableClustredReflection>(bClustredReflection);
+		PermutationVector.Set<FMobileDirectionalLightFunctionPS::FEnablePlanarReflection>(bPlanarReflection);
 		PermutationVector.Set<FMobileDirectionalLightFunctionPS::FEnableSkyLight>(bEnableSkyLight);
-		PermutationVector.Set<FMobileDirectionalLightFunctionPS::FApplyCSM>(ShadowQuality > 0);
-		PermutationVector.Set<FMobileDirectionalLightFunctionPS::FShadowQuality>(FMath::Clamp(ShadowQuality - 1, 0, 3));
+		PermutationVector.Set<FMobileDirectionalLightFunctionPS::FEnableCSM>(ShadowQuality > 0);
+		PermutationVector.Set<FMobileDirectionalLightFunctionPS::FShadowQuality>(FMath::Clamp(ShadowQuality, 1, 3));
 		return PermutationVector;
 	}
 
@@ -230,9 +233,15 @@ public:
 	
 	class FEnableShadingModelSupport	: SHADER_PERMUTATION_BOOL("ENABLE_SHADINGMODEL_SUPPORT_MOBILE_DEFERRED");
 	class FEnableClustredReflection		: SHADER_PERMUTATION_BOOL("ENABLE_CLUSTERED_REFLECTION");
+	class FEnablePlanarReflection		: SHADER_PERMUTATION_BOOL("ENABLE_PLANAR_REFLECTION");
 	class FEnableSkyLight				: SHADER_PERMUTATION_BOOL("ENABLE_SKY_LIGHT");
 	
-	using FPermutationDomain = TShaderPermutationDomain<FEnableShadingModelSupport, FEnableClustredReflection, FEnableSkyLight>;
+	using FPermutationDomain = TShaderPermutationDomain<
+		FEnableShadingModelSupport, 
+		FEnableClustredReflection, 
+		FEnablePlanarReflection,
+		FEnableSkyLight
+	>;
 
 	BEGIN_SHADER_PARAMETER_STRUCT(FParameters, )
 		SHADER_PARAMETER_STRUCT_REF(FViewUniformShaderParameters, View)
@@ -317,7 +326,8 @@ void RenderReflectionEnvironmentSkyLighting(FRHICommandListImmediate& RHICmdList
 	// Skylights with static lighting already had their diffuse contribution baked into lightmaps
 	const bool bSkyLight = Scene.SkyLight && !Scene.SkyLight->bHasStaticLighting && View.Family->EngineShowFlags.SkyLighting;
 	const bool bClustredReflection = (View.NumBoxReflectionCaptures + View.NumSphereReflectionCaptures) > 0;
-	if (!bSkyLight && !bClustredReflection)
+	const bool bPlanarReflection = Scene.GetForwardPassGlobalPlanarReflection() != nullptr;
+	if (!(bSkyLight || bClustredReflection || bPlanarReflection))
 	{
 		return;
 	}
@@ -356,6 +366,7 @@ void RenderReflectionEnvironmentSkyLighting(FRHICommandListImmediate& RHICmdList
 		FMobileReflectionEnvironmentSkyLightingPS::FPermutationDomain PermutationVector;
 		PermutationVector.Set<FMobileReflectionEnvironmentSkyLightingPS::FEnableShadingModelSupport>(bEnableShadingModelSupport);
 		PermutationVector.Set<FMobileReflectionEnvironmentSkyLightingPS::FEnableClustredReflection>(bClustredReflection);
+		PermutationVector.Set<FMobileReflectionEnvironmentSkyLightingPS::FEnablePlanarReflection>(bPlanarReflection);
 		PermutationVector.Set<FMobileReflectionEnvironmentSkyLightingPS::FEnableSkyLight>(bSkyLight);
 		TShaderMapRef<FMobileReflectionEnvironmentSkyLightingPS> PixelShader(View.ShaderMap, PermutationVector);
 		
@@ -457,6 +468,7 @@ static void RenderDirectionalLight(FRHICommandListImmediate& RHICmdList, const F
 	// Skylights with static lighting already had their diffuse contribution baked into lightmaps
 	const bool bSkyLight = Scene.SkyLight && !Scene.SkyLight->bHasStaticLighting && View.Family->EngineShowFlags.SkyLighting;
 	const bool bDynamicShadows = DirectionalLight.Proxy->CastsDynamicShadow() && (LightingChannel == 0u) && View.Family->EngineShowFlags.DynamicShadows;
+	const bool bPlanarReflection = Scene.GetForwardPassGlobalPlanarReflection() != nullptr;
 
 	// Do two passes, first masking DefautLit, second masking all other shading models
 	int32 NumPasses = MobileUsesGBufferCustomData(Scene.GetShaderPlatform()) ? 2 : 1;
@@ -475,7 +487,8 @@ static void RenderDirectionalLight(FRHICommandListImmediate& RHICmdList, const F
 			bInlineReflectionAndSky,
 			bEnableShadingModelSupport,
 			bDynamicShadows,
-			bSkyLight
+			bSkyLight,
+			bPlanarReflection
 		);
 		FCachedLightMaterial LightMaterial;
 		TShaderRef<FMobileDirectionalLightFunctionPS> PixelShader;
