Updated strata waveop code to only use RHI info and be vendor independent. (works post 20161673 which fixes waveop on a RHI)

diff --git a/Engine/Source/Runtime/Renderer/Private/Strata/Strata.cpp b/Engine/Source/Runtime/Renderer/Private/Strata/Strata.cpp
--- a/Engine/Source/Runtime/Renderer/Private/Strata/Strata.cpp
+++ b/Engine/Source/Runtime/Renderer/Private/Strata/Strata.cpp
@@ -998,11 +998,7 @@ void AddStrataMaterialClassificationPass(FRDGBuilder& GraphBuilder, const FMinim
 		RDG_EVENT_SCOPE_CONDITIONAL(GraphBuilder, Views.Num() > 1, "View%d", i);
 
 		const FViewInfo& View = Views[i];
-		bool bWaveOps = GRHISupportsWaveOperations && FDataDrivenShaderPlatformInfo::GetSupportsWaveOperations(View.GetShaderPlatform());
-	#if PLATFORM_WINDOWS
-		// Tile reduction requires 64-wide wave
-		bWaveOps = bWaveOps && !IsRHIDeviceNVIDIA();
-	#endif
+		bool bWaveOps = GRHISupportsWaveOperations && GRHIMaximumWaveSize >= 64 && FDataDrivenShaderPlatformInfo::GetSupportsWaveOperations(View.GetShaderPlatform());
 		
 		const FStrataViewData* StrataViewData = &View.StrataViewData;
 		const FStrataSceneData* StrataSceneData = View.StrataViewData.SceneData;
