removed unused global var

diff --git a/Engine/Source/Runtime/Renderer/Private/DeferredShadingRenderer.cpp b/Engine/Source/Runtime/Renderer/Private/DeferredShadingRenderer.cpp
--- a/Engine/Source/Runtime/Renderer/Private/DeferredShadingRenderer.cpp
+++ b/Engine/Source/Runtime/Renderer/Private/DeferredShadingRenderer.cpp
@@ -400,7 +400,6 @@ DECLARE_CYCLE_STAT(TEXT("RenderFinish"), STAT_CLM_RenderFinish, STATGROUP_Comman
 DECLARE_CYCLE_STAT(TEXT("AfterFrame"), STAT_CLM_AfterFrame, STATGROUP_CommandListMarkers);
 DECLARE_CYCLE_STAT(TEXT("Wait RayTracing Add Mesh Batch"), STAT_WaitRayTracingAddMesh, STATGROUP_SceneRendering);
 
-FGraphEventRef FDeferredShadingSceneRenderer::TranslucencyTimestampQuerySubmittedFence[FOcclusionQueryHelpers::MaxBufferedOcclusionFrames + 1];
 FGlobalDynamicIndexBuffer FDeferredShadingSceneRenderer::DynamicIndexBufferForInitViews;
 FGlobalDynamicIndexBuffer FDeferredShadingSceneRenderer::DynamicIndexBufferForInitShadows;
 FGlobalDynamicVertexBuffer FDeferredShadingSceneRenderer::DynamicVertexBufferForInitViews;
diff --git a/Engine/Source/Runtime/Renderer/Private/DeferredShadingRenderer.h b/Engine/Source/Runtime/Renderer/Private/DeferredShadingRenderer.h
--- a/Engine/Source/Runtime/Renderer/Private/DeferredShadingRenderer.h
+++ b/Engine/Source/Runtime/Renderer/Private/DeferredShadingRenderer.h
@@ -421,7 +421,6 @@ private:
 		return bAnyViewHasGIMethodSupportingDFAO;
 	}
 
-	static FGraphEventRef TranslucencyTimestampQuerySubmittedFence[FOcclusionQueryHelpers::MaxBufferedOcclusionFrames + 1];
 	static FGlobalDynamicIndexBuffer DynamicIndexBufferForInitViews;
 	static FGlobalDynamicIndexBuffer DynamicIndexBufferForInitShadows;
 	static FGlobalDynamicVertexBuffer DynamicVertexBufferForInitViews;
