Added more checks when removing primitives from DistanceField scene.

diff --git a/Engine/Source/Runtime/Renderer/Private/RendererScene.cpp b/Engine/Source/Runtime/Renderer/Private/RendererScene.cpp
--- a/Engine/Source/Runtime/Renderer/Private/RendererScene.cpp
+++ b/Engine/Source/Runtime/Renderer/Private/RendererScene.cpp
@@ -645,6 +645,11 @@ void FDistanceFieldSceneData::RemovePrimitive(FPrimitiveSceneInfo* InPrimitive)
 	}
 
 	checkf(!PendingAddOperations.Contains(InPrimitive), TEXT("Primitive is being removed from the scene, but didn't remove from Distance Field Scene properly - a crash will occur when processing PendingAddOperations.  This can happen if the proxy's properties have changed without recreating its render state."));
+	checkf(!PendingUpdateOperations.Contains(InPrimitive), TEXT("Primitive is being removed from the scene, but didn't remove from Distance Field Scene properly - a crash will occur when processing PendingUpdateOperations.  This can happen if the proxy's properties have changed without recreating its render state."));
+	checkf(!PendingThrottledOperations.Contains(InPrimitive), TEXT("Primitive is being removed from the scene, but didn't remove from Distance Field Scene properly - a crash will occur when processing PendingThrottledOperations.  This can happen if the proxy's properties have changed without recreating its render state."));
+	
+	checkf(!PendingHeightFieldAddOps.Contains(InPrimitive), TEXT("Primitive is being removed from the scene, but didn't remove from Distance Field Scene properly - a crash will occur when processing PendingHeightFieldAddOps.  This can happen if the proxy's properties have changed without recreating its render state."));
+	checkf(!PendingHeightFieldUpdateOps.Contains(InPrimitive), TEXT("Primitive is being removed from the scene, but didn't remove from Distance Field Scene properly - a crash will occur when processing PendingHeightFieldUpdateOps.  This can happen if the proxy's properties have changed without recreating its render state."));
 }
 
 void FDistanceFieldSceneData::Release()
