Clear IndicesToUpdateInHeightFieldObjectBuffers after processing it.
- It was accumulating new indices resulting in unecessary updates.
Store PendingHeightFieldAdd/UpdateOps in TSet instead of TArray so that Contains/Remove calls are O(1).

diff --git a/Engine/Source/Runtime/Renderer/Private/DistanceFieldObjectManagement.cpp b/Engine/Source/Runtime/Renderer/Private/DistanceFieldObjectManagement.cpp
--- a/Engine/Source/Runtime/Renderer/Private/DistanceFieldObjectManagement.cpp
+++ b/Engine/Source/Runtime/Renderer/Private/DistanceFieldObjectManagement.cpp
@@ -904,20 +904,13 @@ void FSceneRenderer::UpdateGlobalHeightFieldObjectBuffers(FRDGBuilder& GraphBuil
 
 		if (DistanceFieldSceneData.PendingHeightFieldAddOps.Num() || DistanceFieldSceneData.PendingHeightFieldUpdateOps.Num())
 		{
-			const int32 NumAddOps = DistanceFieldSceneData.PendingHeightFieldAddOps.Num();
-			const int32 NumUpdateOps = DistanceFieldSceneData.PendingHeightFieldUpdateOps.Num();
-			const int32 NumUploadOps = NumAddOps + NumUpdateOps;
-			const int32 OriginalNumObjects = DistanceFieldSceneData.NumHeightFieldObjectsInBuffer;
-
-			for (int32 Idx = 0; Idx < NumAddOps; ++Idx)
+			for (FPrimitiveSceneInfo* PrimitiveSceneInfo : DistanceFieldSceneData.PendingHeightFieldAddOps)
 			{
-				FPrimitiveSceneInfo* PrimitiveSceneInfo = DistanceFieldSceneData.PendingHeightFieldAddOps[Idx];
 				ProcessHeightFieldPrimitiveUpdate(true, Scene, UpdateTrackingBounds, PrimitiveSceneInfo, DistanceFieldSceneData.IndicesToUpdateInHeightFieldObjectBuffers);
 			}
 
-			for (int32 Idx = 0; Idx < NumUpdateOps; ++Idx)
+			for (FPrimitiveSceneInfo* PrimitiveSceneInfo : DistanceFieldSceneData.PendingHeightFieldUpdateOps)
 			{
-				FPrimitiveSceneInfo* PrimitiveSceneInfo = DistanceFieldSceneData.PendingHeightFieldUpdateOps[Idx];
 				ProcessHeightFieldPrimitiveUpdate(false, Scene, UpdateTrackingBounds, PrimitiveSceneInfo, DistanceFieldSceneData.IndicesToUpdateInHeightFieldObjectBuffers);
 			}
 
@@ -1030,6 +1023,8 @@ void FSceneRenderer::UpdateGlobalHeightFieldObjectBuffers(FRDGBuilder& GraphBuil
 
 					DistanceFieldSceneData.UploadHeightFieldDataBuffer.ResourceUploadTo(GraphBuilder, HeightfieldObjectDataBuffer);
 					DistanceFieldSceneData.UploadHeightFieldBoundsBuffer.ResourceUploadTo(GraphBuilder, HeightfieldObjectBoundsBuffer);
+
+					DistanceFieldSceneData.IndicesToUpdateInHeightFieldObjectBuffers.Reset();
 				}
 			}
 		}
@@ -1140,9 +1135,8 @@ void FSceneRenderer::AddOrRemoveSceneHeightFieldPrimitives(const DistanceField::
 
 	if (!bSkipAdd)
 	{
-		for (int32 Idx = 0; Idx < SceneData.PendingHeightFieldAddOps.Num(); ++Idx)
+		for (FPrimitiveSceneInfo* Primitive : SceneData.PendingHeightFieldAddOps)
 		{
-			FPrimitiveSceneInfo* Primitive = SceneData.PendingHeightFieldAddOps[Idx];
 			const int32 HFIdx = SceneData.HeightfieldPrimitives.Add(Primitive);
 			Primitive->DistanceFieldInstanceIndices.Empty(1);
 			Primitive->DistanceFieldInstanceIndices.Add(HFIdx);
diff --git a/Engine/Source/Runtime/Renderer/Private/ScenePrivate.h b/Engine/Source/Runtime/Renderer/Private/ScenePrivate.h
--- a/Engine/Source/Runtime/Renderer/Private/ScenePrivate.h
+++ b/Engine/Source/Runtime/Renderer/Private/ScenePrivate.h
@@ -2083,8 +2083,8 @@ public:
 	TArray<FPrimitiveRemoveInfo> PendingRemoveOperations;
 	TArray<FRenderBounds> PrimitiveModifiedBounds[GDF_Num];
 
-	TArray<FPrimitiveSceneInfo*> PendingHeightFieldAddOps;
-	TArray<FPrimitiveSceneInfo*> PendingHeightFieldUpdateOps;
+	TSet<FPrimitiveSceneInfo*> PendingHeightFieldAddOps;
+	TSet<FPrimitiveSceneInfo*> PendingHeightFieldUpdateOps;
 	TArray<FHeightFieldPrimitiveRemoveInfo> PendingHeightFieldRemoveOps;
 
 	int32 HeightFieldAtlasGeneration;
