Fixes VisualizeLumenScene doing a full-res texture copy of the back buffer when disabled

diff --git a/Engine/Source/Runtime/Renderer/Private/PostProcess/PostProcessing.cpp b/Engine/Source/Runtime/Renderer/Private/PostProcess/PostProcessing.cpp
--- a/Engine/Source/Runtime/Renderer/Private/PostProcess/PostProcessing.cpp
+++ b/Engine/Source/Runtime/Renderer/Private/PostProcess/PostProcessing.cpp
@@ -248,7 +248,7 @@ void AddPostProcessingPasses(
 	const bool bVisualizeGBufferDumpToFile = IsVisualizeGBufferDumpToFileEnabled(View);
 	const bool bVisualizeGBufferDumpToPIpe = IsVisualizeGBufferDumpToPipeEnabled(View);
 	const bool bOutputInHDR = IsPostProcessingOutputInHDR();
-	const int32 VisualizeMode = GetLumenVisualizeMode(View);
+	const int32 LumenVisualizeMode = GetLumenVisualizeMode(View);
 	const bool bPostProcessingEnabled = IsPostProcessingEnabled(View);
 
 	const FPaniniProjectionConfig PaniniConfig(View);
@@ -343,7 +343,7 @@ void AddPostProcessingPasses(
 	PassSequence.SetEnabled(EPass::VisualizePostProcessStack, false);
 #endif
 	PassSequence.SetEnabled(EPass::VisualizeStrata, Strata::ShouldRenderStrataDebugPasses(View));
-	PassSequence.SetEnabled(EPass::VisualizeLumenScene, VisualizeMode != VISUALIZE_MODE_OVERVIEW && bPostProcessingEnabled);
+	PassSequence.SetEnabled(EPass::VisualizeLumenScene, LumenVisualizeMode >= 0 && LumenVisualizeMode != VISUALIZE_MODE_OVERVIEW && bPostProcessingEnabled);
 #if WITH_EDITOR
 	PassSequence.SetEnabled(EPass::VisualizeSkyAtmosphere, Scene && View.Family && View.Family->EngineShowFlags.VisualizeSkyAtmosphere&& ShouldRenderSkyAtmosphereDebugPasses(Scene, View.Family->EngineShowFlags));
 	PassSequence.SetEnabled(EPass::VisualizeLevelInstance, GIsEditor && EngineShowFlags.EditingLevelInstance && EngineShowFlags.VisualizeLevelInstanceEditing && !bVisualizeHDR);
@@ -359,7 +359,7 @@ void AddPostProcessingPasses(
 	PassSequence.SetEnabled(EPass::VisualizeGBufferHints, EngineShowFlags.GBufferHints);
 	PassSequence.SetEnabled(EPass::VisualizeSubsurface, EngineShowFlags.VisualizeSSS);
 	PassSequence.SetEnabled(EPass::VisualizeGBufferOverview, bVisualizeGBufferOverview || bVisualizeGBufferDumpToFile || bVisualizeGBufferDumpToPIpe);
-	PassSequence.SetEnabled(EPass::VisualizeLumenSceneOverview, VisualizeMode == VISUALIZE_MODE_OVERVIEW && bPostProcessingEnabled);
+	PassSequence.SetEnabled(EPass::VisualizeLumenSceneOverview, LumenVisualizeMode == VISUALIZE_MODE_OVERVIEW && bPostProcessingEnabled);
 	PassSequence.SetEnabled(EPass::VisualizeHDR, EngineShowFlags.VisualizeHDR);
 	PassSequence.SetEnabled(EPass::VisualizeMotionVectors, EngineShowFlags.VisualizeMotionVectors);
 #if WITH_EDITOR
