Allows TSR to output half res when motion blur is disabled.

diff --git a/Engine/Source/Runtime/Renderer/Private/PostProcess/PostProcessing.cpp b/Engine/Source/Runtime/Renderer/Private/PostProcess/PostProcessing.cpp
--- a/Engine/Source/Runtime/Renderer/Private/PostProcess/PostProcessing.cpp
+++ b/Engine/Source/Runtime/Renderer/Private/PostProcess/PostProcessing.cpp
@@ -569,16 +569,6 @@ void AddPostProcessingPasses(
 		FVelocityFlattenTextures VelocityFlattenTextures;
 		if (TAAConfig != EMainTAAPassConfig::Disabled)
 		{
-			// Whether we allow the temporal AA pass to downsample scene color. It may choose not to based on internal context,
-			// in which case the output half resolution texture will remain null.
-			const bool bAllowSceneDownsample =
-				IsTemporalAASceneDownsampleAllowed(View) &&
-				// We can only merge if the normal downsample pass would happen immediately after.
-				bNeedPostMotionBlurHalfRes &&
-				!bMotionBlurEnabled && !bVisualizeMotionBlur &&
-				// TemporalAA is only able to match the low quality mode (box filter).
-				DownsampleQuality == EDownsampleQuality::Low;
-
 			const ITemporalUpscaler* UpscalerToUse = (TAAConfig == EMainTAAPassConfig::ThirdParty) ? View.Family->GetTemporalUpscalerInterface() : ITemporalUpscaler::GetDefaultTemporalUpscaler();
 			check(UpscalerToUse);
 			const TCHAR* UpscalerName = UpscalerToUse->GetDebugName();
@@ -594,10 +584,24 @@ void AddPostProcessingPasses(
 				View.GetSecondaryViewRectSize().X, View.GetSecondaryViewRectSize().Y);
 
 			ITemporalUpscaler::FPassInputs UpscalerPassInputs;
-
-			UpscalerPassInputs.bAllowDownsampleSceneColor = bAllowSceneDownsample;
-			UpscalerPassInputs.bGenerateOutputMip1 = bMotionBlurNeedsHalfResInput;
-			UpscalerPassInputs.bGenerateVelocityFlattenTextures = PassSequence.IsEnabled(EPass::MotionBlur) && FVelocityFlattenTextures::AllowExternal(View) && !bVisualizeMotionBlur;
+			if (PassSequence.IsEnabled(EPass::MotionBlur))
+			{
+				if (bVisualizeMotionBlur)
+				{
+					// NOP
+				}
+				else
+				{
+					UpscalerPassInputs.bGenerateOutputMip1 = bMotionBlurNeedsHalfResInput;
+					UpscalerPassInputs.bGenerateVelocityFlattenTextures = FVelocityFlattenTextures::AllowExternal(View) && !bVisualizeMotionBlur;
+				}
+			}
+			else
+			{
+				UpscalerPassInputs.bAllowDownsampleSceneColor =
+					(bNeedPostMotionBlurHalfRes || bNeedPostMotionBlurQuarterRes) &&
+					DownsampleQuality == EDownsampleQuality::Low;;
+			}
 			UpscalerPassInputs.DownsampleOverrideFormat = DownsampleOverrideFormat;
 			UpscalerPassInputs.SceneColorTexture = SceneColor.Texture;
 			UpscalerPassInputs.SceneDepthTexture = SceneDepth.Texture;
diff --git a/Engine/Source/Runtime/Renderer/Private/PostProcess/TemporalAA.cpp b/Engine/Source/Runtime/Renderer/Private/PostProcess/TemporalAA.cpp
--- a/Engine/Source/Runtime/Renderer/Private/PostProcess/TemporalAA.cpp
+++ b/Engine/Source/Runtime/Renderer/Private/PostProcess/TemporalAA.cpp
@@ -55,12 +55,6 @@ TAutoConsoleVariable<float> CVarTemporalAAHistorySP(
 	TEXT("Size of temporal AA's history."),
 	ECVF_RenderThreadSafe);
 
-TAutoConsoleVariable<int32> CVarTemporalAAAllowDownsampling(
-	TEXT("r.TemporalAA.AllowDownsampling"),
-	1,
-	TEXT("Allows half-resolution color buffer to be produced during TAA. Only possible when motion blur is off and when using compute shaders for post processing."),
-	ECVF_RenderThreadSafe);
-
 static TAutoConsoleVariable<int32> CVarUseTemporalAAUpscaler(
 	TEXT("r.TemporalAA.Upscaler"),
 	1,
@@ -371,11 +365,6 @@ static_assert(UE_ARRAY_COUNT(kTAAPassNames) == int32(ETAAPassConfig::MAX), "Miss
 static_assert(UE_ARRAY_COUNT(kTAAQualityNames) == int32(ETAAQuality::MAX), "Missing TAA quality name.");
 } //! namespace
 
-bool IsTemporalAASceneDownsampleAllowed(const FViewInfo& View)
-{
-	return CVarTemporalAAAllowDownsampling.GetValueOnRenderThread() != 0;
-}
-
 FVector ComputePixelFormatQuantizationError(EPixelFormat PixelFormat)
 {
 	FVector Error;
diff --git a/Engine/Source/Runtime/Renderer/Private/PostProcess/TemporalAA.h b/Engine/Source/Runtime/Renderer/Private/PostProcess/TemporalAA.h
--- a/Engine/Source/Runtime/Renderer/Private/PostProcess/TemporalAA.h
+++ b/Engine/Source/Runtime/Renderer/Private/PostProcess/TemporalAA.h
@@ -47,8 +47,6 @@ enum class ETAAPassConfig
 	MAX
 };
 
-bool IsTemporalAASceneDownsampleAllowed(const FViewInfo& View);
-
 static FORCEINLINE bool IsTAAUpsamplingConfig(ETAAPassConfig Pass)
 {
 	return Pass == ETAAPassConfig::MainUpsampling || Pass == ETAAPassConfig::DiaphragmDOFUpsampling || Pass == ETAAPassConfig::MainSuperSampling;
diff --git a/Engine/Source/Runtime/Renderer/Private/PostProcess/TemporalSuperResolution.cpp b/Engine/Source/Runtime/Renderer/Private/PostProcess/TemporalSuperResolution.cpp
--- a/Engine/Source/Runtime/Renderer/Private/PostProcess/TemporalSuperResolution.cpp
+++ b/Engine/Source/Runtime/Renderer/Private/PostProcess/TemporalSuperResolution.cpp
@@ -1454,7 +1454,9 @@ ITemporalUpscaler::FOutputs AddTemporalSuperResolutionPasses(
 
 	// Allocate output
 	FRDGTextureRef SceneColorOutputTexture;
+	FRDGTextureRef SceneColorOutputHalfResTexture = nullptr;
 	{
+		check(!(PassInputs.bGenerateOutputMip1 && PassInputs.bAllowDownsampleSceneColor));
 		FRDGTextureDesc Desc = FRDGTextureDesc::Create2D(
 			OutputExtent,
 			ColorFormat,
@@ -1463,6 +1465,12 @@ ITemporalUpscaler::FOutputs AddTemporalSuperResolutionPasses(
 			/* NumMips = */ PassInputs.bGenerateOutputMip1 ? 2 : 1);
 
 		SceneColorOutputTexture = GraphBuilder.CreateTexture(Desc, TEXT("TSR.Output"));
+
+		if (PassInputs.bAllowDownsampleSceneColor)
+		{
+			Desc.Extent /= 2;
+			SceneColorOutputHalfResTexture = GraphBuilder.CreateTexture(Desc, TEXT("TSR.HalfResOutput"));
+		}
 	}
 
 	// Update temporal history.
@@ -1501,7 +1509,7 @@ ITemporalUpscaler::FOutputs AddTemporalSuperResolutionPasses(
 		PassParameters->InvWeightClampingPixelSpeed = 1.0f / CVarTSRWeightClampingPixelSpeed.GetValueOnRenderThread();
 		PassParameters->InputToHistoryFactor = float(HistorySize.X) / float(InputRect.Width());
 		PassParameters->ResponsiveStencilMask = CVarTSREnableResponiveAA.GetValueOnRenderThread() ? (STENCIL_TEMPORAL_RESPONSIVE_AA_MASK) : 0;
-		PassParameters->bGenerateOutputMip1 = (PassInputs.bGenerateOutputMip1 && HistorySize == OutputRect.Size()) ? 1 : 0;
+		PassParameters->bGenerateOutputMip1 = ((PassInputs.bGenerateOutputMip1 || PassInputs.bAllowDownsampleSceneColor) && HistorySize == OutputRect.Size()) ? 1 : 0;
 		PassParameters->bHasSeparateTranslucency = bHasSeparateTranslucency;
 
 		PassParameters->PrevHistoryParameters = PrevHistoryParameters;
@@ -1517,13 +1525,21 @@ ITemporalUpscaler::FOutputs AddTemporalSuperResolutionPasses(
 			PassParameters->SceneColorOutputMip0 = GraphBuilder.CreateUAV(FRDGTextureUAVDesc(SceneColorOutputTexture, /* InMipLevel = */ 0));
 		}
 		
-		if (PassParameters->bGenerateOutputMip1)
+		if (!PassParameters->bGenerateOutputMip1)
+		{
+			PassParameters->SceneColorOutputMip1 = CreateDummyUAV(GraphBuilder, PF_FloatR11G11B10);
+		}
+		else if (PassInputs.bGenerateOutputMip1)
 		{
 			PassParameters->SceneColorOutputMip1 = GraphBuilder.CreateUAV(FRDGTextureUAVDesc(SceneColorOutputTexture, /* InMipLevel = */ 1));
 		}
+		else if (PassInputs.bAllowDownsampleSceneColor)
+		{
+			PassParameters->SceneColorOutputMip1 = GraphBuilder.CreateUAV(FRDGTextureUAVDesc(SceneColorOutputHalfResTexture));
+		}
 		else
 		{
-			PassParameters->SceneColorOutputMip1 = CreateDummyUAV(GraphBuilder, PF_FloatR11G11B10);
+			unimplemented();
 		}
 		PassParameters->DebugOutput = CreateDebugUAV(HistoryExtent, TEXT("Debug.TSR.UpdateHistory"));
 
@@ -1538,7 +1554,7 @@ ITemporalUpscaler::FOutputs AddTemporalSuperResolutionPasses(
 				kUpdateQualityNames[int32(PermutationVector.Get<FTSRUpdateHistoryCS::FQualityDim>())],
 				History.ColorArray->Desc.Format == PF_FloatR11G11B10 ? TEXT("R11G11B10") : TEXT(""),
 				PermutationVector.Get<FTSRHighFrequencyOnlyDim>() ? TEXT("") : TEXT(" LowFrequency"),
-				PassInputs.bGenerateOutputMip1 ? TEXT(" OutputMip1") : TEXT(""),
+				PassParameters->bGenerateOutputMip1 ? TEXT(" OutputMip1") : TEXT(""),
 				HistorySize.X, HistorySize.Y),
 			ComputePassFlags,
 			ComputeShader,
@@ -1581,16 +1597,20 @@ ITemporalUpscaler::FOutputs AddTemporalSuperResolutionPasses(
 				FScreenTransform::ETextureBasis::ViewportUV, FScreenTransform::ETextureBasis::TexelPosition));
 		PassParameters->OutputViewRectMin = OutputRect.Min;
 		PassParameters->OutputViewRectMax = OutputRect.Max;
-		PassParameters->bGenerateOutputMip1 = PassInputs.bGenerateOutputMip1 ? 1 : 0;
+		PassParameters->bGenerateOutputMip1 = (PassInputs.bGenerateOutputMip1 || PassInputs.bAllowDownsampleSceneColor) ? 1 : 0;
 		PassParameters->HistoryValidityMultiply = float(HistorySize.X * HistorySize.Y) / float(OutputRect.Width() * OutputRect.Height());
 
 		PassParameters->History = CreateSRVs(GraphBuilder, History);
 		
 		PassParameters->SceneColorOutputMip0 = GraphBuilder.CreateUAV(FRDGTextureUAVDesc(SceneColorOutputTexture, /* InMipLevel = */ 0));
-		if (PassParameters->bGenerateOutputMip1)
+		if (PassInputs.bGenerateOutputMip1)
 		{
 			PassParameters->SceneColorOutputMip1 = GraphBuilder.CreateUAV(FRDGTextureUAVDesc(SceneColorOutputTexture, /* InMipLevel = */ 1));
 		}
+		else if (PassInputs.bAllowDownsampleSceneColor)
+		{
+			PassParameters->SceneColorOutputMip1 = GraphBuilder.CreateUAV(FRDGTextureUAVDesc(SceneColorOutputHalfResTexture));
+		}
 		else
 		{
 			PassParameters->SceneColorOutputMip1 = CreateDummyUAV(GraphBuilder, PF_FloatR11G11B10);
@@ -1654,6 +1674,12 @@ ITemporalUpscaler::FOutputs AddTemporalSuperResolutionPasses(
 	ITemporalUpscaler::FOutputs Outputs;
 	Outputs.FullRes.Texture = SceneColorOutputTexture;
 	Outputs.FullRes.ViewRect = OutputRect;
+	if (PassInputs.bAllowDownsampleSceneColor)
+	{
+		Outputs.HalfRes.Texture = SceneColorOutputHalfResTexture;
+		Outputs.HalfRes.ViewRect.Min = OutputRect.Min / 2;
+		Outputs.HalfRes.ViewRect.Max = Outputs.HalfRes.ViewRect.Min + FIntPoint::DivideAndRoundUp(OutputRect.Size(), 2);
+	}
 	Outputs.VelocityFlattenTextures = VelocityFlattenTextures;
 	return Outputs;
 } // AddTemporalSuperResolutionPasses()
