fixed undefined order of static destruction for FSceneRenderer::OcclusionSubmittedFence

diff --git a/Engine/Source/Runtime/Renderer/Private/Renderer.cpp b/Engine/Source/Runtime/Renderer/Private/Renderer.cpp
--- a/Engine/Source/Runtime/Renderer/Private/Renderer.cpp
+++ b/Engine/Source/Runtime/Renderer/Private/Renderer.cpp
@@ -83,6 +83,9 @@ void FRendererModule::ShutdownModule()
 
 	// Free up global resources in Lumen
 	Lumen::Shutdown();
+
+	void CleanupOcclusionSubmittedFence();
+	CleanupOcclusionSubmittedFence();
 }
 
 void FRendererModule::OnWorldCleanup(UWorld* World, bool bSessionEnded, bool bCleanupResources, bool bWorldChanged)
diff --git a/Engine/Source/Runtime/Renderer/Private/SceneRendering.cpp b/Engine/Source/Runtime/Renderer/Private/SceneRendering.cpp
--- a/Engine/Source/Runtime/Renderer/Private/SceneRendering.cpp
+++ b/Engine/Source/Runtime/Renderer/Private/SceneRendering.cpp
@@ -401,8 +401,18 @@ static TAutoConsoleVariable<int32> CVarNaniteShowUnsupportedError(
 #endif
 
 static FParallelCommandListSet* GOutstandingParallelCommandListSet = nullptr;
+
 FOcclusionSubmittedFenceState FSceneRenderer::OcclusionSubmittedFence[FOcclusionQueryHelpers::MaxBufferedOcclusionFrames];
 
+// cleanup OcclusionSubmittedFence to avoid undefined order of destruction that can destroy it after its allocator
+void CleanupOcclusionSubmittedFence()
+{
+	for (FOcclusionSubmittedFenceState& FenceState : FSceneRenderer::OcclusionSubmittedFence)
+	{
+		FenceState.Fence = nullptr;
+	}
+}
+
 extern int32 GetTranslucencyLightingVolumeDim();
 
 DECLARE_CYCLE_STAT(TEXT("DeferredShadingSceneRenderer ViewExtensionPostRenderView"), STAT_FDeferredShadingSceneRenderer_ViewExtensionPostRenderView, STATGROUP_SceneRendering);
