Use Translucency Sort Priority to sort mesh decals

diff --git a/Engine/Source/Runtime/Renderer/Private/CompositionLighting/PostProcessMeshDecals.cpp b/Engine/Source/Runtime/Renderer/Private/CompositionLighting/PostProcessMeshDecals.cpp
--- a/Engine/Source/Runtime/Renderer/Private/CompositionLighting/PostProcessMeshDecals.cpp
+++ b/Engine/Source/Runtime/Renderer/Private/CompositionLighting/PostProcessMeshDecals.cpp
@@ -290,7 +290,11 @@ bool FMeshDecalMeshProcessor::Process(
 	FMeshMaterialShaderElementData ShaderElementData;
 	ShaderElementData.InitializeMeshMaterialData(ViewIfDynamicMeshCommand, PrimitiveSceneProxy, MeshBatch, StaticMeshId, true);
 
-	const FMeshDrawCommandSortKey SortKey = CalculateMeshStaticSortKey(MeshDecalPassShaders.VertexShader, MeshDecalPassShaders.PixelShader);
+	// Use BasePass sort key layout but replace the "Masked" highest priority bits with TranslucencySortPriority.
+	FMeshDrawCommandSortKey SortKey;
+	SortKey.BasePass.VertexShaderHash = (MeshDecalPassShaders.VertexShader.IsValid() ? MeshDecalPassShaders.VertexShader->GetSortKey() : 0) & 0xFFFF;
+	SortKey.BasePass.PixelShaderHash = MeshDecalPassShaders.PixelShader.IsValid() ? MeshDecalPassShaders.PixelShader->GetSortKey() : 0;
+	SortKey.BasePass.Masked = PrimitiveSceneProxy->GetTranslucencySortPriority();
 
 	BuildMeshDrawCommands(
 		MeshBatch,
