Fix a bug that the indirect lighting of tranclucency is too bright on mobile.

diff --git a/Engine/Source/Runtime/Renderer/Private/LightMapRendering.cpp b/Engine/Source/Runtime/Renderer/Private/LightMapRendering.cpp
--- a/Engine/Source/Runtime/Renderer/Private/LightMapRendering.cpp
+++ b/Engine/Source/Runtime/Renderer/Private/LightMapRendering.cpp
@@ -251,8 +251,8 @@ void InterpolateVolumetricLightmap(
 	OutInterpolation.IndirectLightingSHCoefficients1[2] = FVector4f(LQSH.B.V[4], LQSH.B.V[5], LQSH.B.V[6], LQSH.B.V[7]) * INV_PI;
 	OutInterpolation.IndirectLightingSHCoefficients2 = FVector4f(LQSH.R.V[8], LQSH.G.V[8], LQSH.B.V[8], 0.0f) * INV_PI;
 
-	OutInterpolation.IndirectLightingSHSingleCoefficient = FVector4f(AmbientVector.X, AmbientVector.Y, AmbientVector.Z)
-		* FSHVector2::ConstantBasisIntegral * .5f;
+	//The IndirectLightingSHSingleCoefficient should be DotSH1(AmbientVector, CalcDiffuseTransferSH1(1)) / PI, and CalcDiffuseTransferSH1(1) / PI equals to 1 / (2 * sqrt(PI)) and FSHVector2::ConstantBasisIntegral is 2 * sqrt(PI)
+	OutInterpolation.IndirectLightingSHSingleCoefficient = FVector4f(AmbientVector.X, AmbientVector.Y, AmbientVector.Z) / FSHVector2::ConstantBasisIntegral;
 
 	if (VolumetricLightmapData.BrickData.SkyBentNormal.Data.Num() > 0)
 	{
@@ -341,8 +341,8 @@ void GetIndirectLightingCacheParameters(
 				Parameters.IndirectLightingSHCoefficients1[i] = LightingAllocation->SingleSamplePacked1[i];
 			}
 			Parameters.IndirectLightingSHCoefficients2 = LightingAllocation->SingleSamplePacked2;
-			Parameters.IndirectLightingSHSingleCoefficient = FVector4f(LightingAllocation->SingleSamplePacked0[0].X, LightingAllocation->SingleSamplePacked0[1].X, LightingAllocation->SingleSamplePacked0[2].X)
-					* FSHVector2::ConstantBasisIntegral * .5f; //@todo - why is .5f needed to match directional?
+			//The IndirectLightingSHSingleCoefficient should be DotSH1(LightingAllocation->SingleSamplePacked0[0], CalcDiffuseTransferSH1(1)) / PI, and CalcDiffuseTransferSH1(1) / PI equals to 1 / (2 * sqrt(PI)) and FSHVector2::ConstantBasisIntegral is 2 * sqrt(PI)
+			Parameters.IndirectLightingSHSingleCoefficient = FVector4f(LightingAllocation->SingleSamplePacked0[0].X, LightingAllocation->SingleSamplePacked0[1].X, LightingAllocation->SingleSamplePacked0[2].X) / FSHVector2::ConstantBasisIntegral;
 		}
 		else
 		{
