SceneDepthWithoutSingleLayerWaterTexture is now 32bit on all platforms

#rnx

diff --git a/Engine/Source/Runtime/Renderer/Private/SingleLayerWaterRendering.cpp b/Engine/Source/Runtime/Renderer/Private/SingleLayerWaterRendering.cpp
--- a/Engine/Source/Runtime/Renderer/Private/SingleLayerWaterRendering.cpp
+++ b/Engine/Source/Runtime/Renderer/Private/SingleLayerWaterRendering.cpp
@@ -296,7 +296,7 @@ static FSceneWithoutWaterTextures AddCopySceneWithoutWaterPass(
 		SceneColorWithoutSingleLayerWaterTexture = GraphBuilder.CreateTexture(ColorDesc, TEXT("SLW.SceneColorWithout"));
 	}
 
-	const FRDGTextureDesc DepthDesc(FRDGTextureDesc::Create2D(RefractionResolution, ViewFamily.EngineShowFlags.SingleLayerWaterRefractionFullPrecision ? PF_R32_FLOAT : PF_R16F, SceneDepthDesc.ClearValue, TexCreate_ShaderResource | TexCreate_RenderTargetable));
+	const FRDGTextureDesc DepthDesc(FRDGTextureDesc::Create2D(RefractionResolution, PF_R32_FLOAT, SceneDepthDesc.ClearValue, TexCreate_ShaderResource | TexCreate_RenderTargetable));
 	FRDGTextureRef SceneDepthWithoutSingleLayerWaterTexture = GraphBuilder.CreateTexture(DepthDesc, TEXT("SLW.SceneDepthWithout"));
 
 	const FRDGTextureDesc SeparatedMainDirLightDesc(FRDGTextureDesc::Create2D(SceneColorDesc.Extent, PF_FloatR11G11B10, FClearValueBinding(FLinearColor::White), TexCreate_ShaderResource | TexCreate_RenderTargetable));
