Fix for compilation strata compilation error.

diff --git a/Engine/Source/Runtime/Renderer/Private/Strata/Strata.cpp b/Engine/Source/Runtime/Renderer/Private/Strata/Strata.cpp
--- a/Engine/Source/Runtime/Renderer/Private/Strata/Strata.cpp
+++ b/Engine/Source/Runtime/Renderer/Private/Strata/Strata.cpp
@@ -599,7 +599,7 @@ class FStrataMaterialTileClassificationPassCS : public FGlobalShader
 		SHADER_PARAMETER_RDG_BUFFER_UAV(RWBuffer, ComplexTileListDataBuffer)
 		SHADER_PARAMETER_RDG_BUFFER_UAV(RWBuffer, OpaqueRoughRefractionTileListDataBuffer)
 		SHADER_PARAMETER_RDG_BUFFER_UAV(RWBuffer, SSSWithoutOpaqueRoughRefractionTileListDataBuffer)
-		SHADER_PARAMETER_RDG_TEXTURE_UAV(RWTextureArray2D<uint>, SSSTextureUAV)
+		SHADER_PARAMETER_RDG_TEXTURE_UAV(RWTexture2DArray<uint>, SSSTextureUAV)
 		SHADER_PARAMETER_RDG_TEXTURE(Texture2D<float3>, OpaqueRoughRefractionTexture)
 	END_SHADER_PARAMETER_STRUCT()
 
