Fix for crash when FHZBOcclusionTester::MapResults is called before FHZBOcclusionTester::Submit

diff --git a/Engine/Source/Runtime/Renderer/Private/SceneOcclusion.cpp b/Engine/Source/Runtime/Renderer/Private/SceneOcclusion.cpp
--- a/Engine/Source/Runtime/Renderer/Private/SceneOcclusion.cpp
+++ b/Engine/Source/Runtime/Renderer/Private/SceneOcclusion.cpp
@@ -796,8 +796,6 @@ void FHZBOcclusionTester::InitDynamicRHI()
 {
 	if (GetFeatureLevel() >= ERHIFeatureLevel::SM5)
 	{
-		FRHICommandListImmediate& RHICmdList = FRHICommandListExecutor::GetImmediateCommandList();
-		FPooledRenderTargetDesc Desc( FPooledRenderTargetDesc::Create2DDesc( FIntPoint( SizeX, SizeY ), PF_B8G8R8A8, FClearValueBinding::None, TexCreate_CPUReadback | TexCreate_HideInVisualizeTexture, TexCreate_None, false ) );
 		ResultsReadback.Reset(new FRHIGPUTextureReadback(TEXT("HZBGPUReadback")));
 	}
 }
@@ -832,8 +830,11 @@ void FHZBOcclusionTester::MapResults(FRHICommandListImmediate& RHICmdList)
 		int32 ResultBufferHeight = 0;
 		ResultsBufferRowPitch = 0;
 		ResultsBuffer = reinterpret_cast<const uint8*>(ResultsReadback->Lock(ResultsBufferRowPitch, &ResultBufferHeight));
-		check(ResultsBufferRowPitch >= SizeX);
-		check(ResultBufferHeight >= SizeY);
+		if (ResultsBuffer)
+		{
+			check(ResultsBufferRowPitch >= SizeX);
+			check(ResultBufferHeight >= SizeY);
+		}
 
 		// RHIMapStagingSurface will block until the results are ready (from the previous frame) so we need to consider this RT idle time
 		GRenderThreadIdle[ERenderThreadIdleTypes::WaitingForGPUQuery] += FPlatformTime::Cycles() - IdleStart;
@@ -841,9 +842,9 @@ void FHZBOcclusionTester::MapResults(FRHICommandListImmediate& RHICmdList)
 	}
 	
 	// Can happen because of device removed, we might crash later but this occlusion culling system can behave gracefully.
+	// It also happens if Submit has not been called yet, which can occur if there are no primitives in the scene.
 	if( ResultsBuffer == nullptr )
 	{
-		// First frame
 		static uint8 FirstFrameBuffer[] = { 255 };
 		ResultsBuffer = FirstFrameBuffer;
 		SetInvalidFrameNumber();
