Removed cases of AddCopyToResolveTargetPass which either weren't necessary or weren't performing multisample resolves.

diff --git a/Engine/Plugins/Experimental/Water/Source/Runtime/Private/WaterInfoRendering.cpp b/Engine/Plugins/Experimental/Water/Source/Runtime/Private/WaterInfoRendering.cpp
--- a/Engine/Plugins/Experimental/Water/Source/Runtime/Private/WaterInfoRendering.cpp
+++ b/Engine/Plugins/Experimental/Water/Source/Runtime/Private/WaterInfoRendering.cpp
@@ -329,10 +329,9 @@ static void UpdateWaterInfoRendering_RenderThread(
 			RDG_RHI_EVENT_SCOPE(GraphBuilder, RenderWaterInfoDepth);
 			DepthRenderer->Render(GraphBuilder);
 		}
-		
-		FRDGTextureRef ResolveTexture = RegisterExternalTexture(GraphBuilder, OutputTexture->TextureRHI, TEXT("WaterDepthResolve"));
-		FResolveParams ResolveParams;
-		AddCopyToResolveTargetPass(GraphBuilder, TargetTexture, ResolveTexture, ResolveParams);
+
+		FRDGTextureRef ShaderResourceTexture = RegisterExternalTexture(GraphBuilder, OutputTexture->TextureRHI, TEXT("WaterDepthTexture"));
+		AddCopyTexturePass(GraphBuilder, TargetTexture, ShaderResourceTexture);
 
 		if (DepthRenderer->Scene->GetShadingPath() == EShadingPath::Mobile)
 		{
@@ -342,7 +341,7 @@ static void UpdateWaterInfoRendering_RenderThread(
 			CopySceneCaptureComponentToTarget(
 				GraphBuilder,
 				SceneTextures,
-				ResolveTexture,
+				ShaderResourceTexture,
 				*DepthRenderer->ActiveViewFamily,
 				DepthRenderer->Views,
 				bNeedsFlippedRenderTarget);
@@ -406,9 +405,8 @@ static void UpdateWaterInfoRendering_RenderThread(
 		FRDGTextureRef FinalizedTexture = GraphBuilder.CreateTexture(TargetTexture->Desc, TEXT("WaterInfoFinalized"));
 		FinalizeWaterInfo(GraphBuilder, SceneTextures, *ColorRenderer->ActiveViewFamily, ColorRenderer->Views[0], MergeTargetTexture, FinalizedTexture, Params);
 		
-		FRDGTextureRef ResolveOutputTexture = RegisterExternalTexture(GraphBuilder, OutputTexture->TextureRHI, TEXT("WaterInfoResolve"));
-		FResolveParams ResolveParams;
-		AddCopyToResolveTargetPass(GraphBuilder, FinalizedTexture, ResolveOutputTexture, ResolveParams);
+		FRDGTextureRef ShaderResourceTexture = RegisterExternalTexture(GraphBuilder, OutputTexture->TextureRHI, TEXT("WaterInfoResolve"));
+		AddCopyTexturePass(GraphBuilder, FinalizedTexture, ShaderResourceTexture);
 		GraphBuilder.Execute();
 		
 		SetWaterBodiesWithinWaterInfoPass(ColorRenderer, false);
diff --git a/Engine/Plugins/Media/MediaIOFramework/Source/MediaIOCore/Private/MediaCapture.cpp b/Engine/Plugins/Media/MediaIOFramework/Source/MediaIOCore/Private/MediaCapture.cpp
--- a/Engine/Plugins/Media/MediaIOFramework/Source/MediaIOCore/Private/MediaCapture.cpp
+++ b/Engine/Plugins/Media/MediaIOFramework/Source/MediaIOCore/Private/MediaCapture.cpp
@@ -270,7 +270,7 @@ namespace UE::MediaCaptureData
 		FRDGBuilder& GraphBuilder;
 		const FRDGTextureRef& SourceRGBTexture;
 		bool bRequiresFormatConversion = false;
-		FResolveParams ResolveParams;
+		FRHICopyTextureInfo CopyInfo;
 		FVector2D SizeU = FVector2D::ZeroVector;
 		FVector2D SizeV = FVector2D::ZeroVector;
 	};
@@ -374,40 +374,30 @@ namespace UE::MediaCaptureData
 			return true;
 		}
 
-		static void GetResolveParameters(const FCaptureFrameArgs& Args, const FTexture2DRHIRef& SourceTexture, FResolveParams& OutResolveParams, FVector2D& OutSizeU, FVector2D& OutSizeV)
+		static void GetCopyInfo(const FCaptureFrameArgs& Args, const FTexture2DRHIRef& SourceTexture, FRHICopyTextureInfo& OutCopyInfo, FVector2D& OutSizeU, FVector2D& OutSizeV)
 		{
 			// Default to no crop
 			OutSizeU = { 0.0f, 1.0f };
 			OutSizeV = { 0.0f, 1.0f };
-			OutResolveParams.Rect = FResolveRect(0, 0, SourceTexture->GetSizeX(), SourceTexture->GetSizeY());
+			OutCopyInfo.Size = FIntVector(SourceTexture->GetSizeX(), SourceTexture->GetSizeY(), 1);
 			if (Args.MediaCapture->DesiredCaptureOptions.Crop != EMediaCaptureCroppingType::None)
 			{
 				switch (Args.MediaCapture->DesiredCaptureOptions.Crop)
 				{
 				case EMediaCaptureCroppingType::Center:
-					OutResolveParams.Rect = FResolveRect((SourceTexture->GetSizeX() - Args.DesiredSize.X) / 2, (SourceTexture->GetSizeY() - Args.DesiredSize.Y) / 2, 0, 0);
-					OutResolveParams.Rect.X2 = OutResolveParams.Rect.X1 + Args.DesiredSize.X;
-					OutResolveParams.Rect.Y2 = OutResolveParams.Rect.Y1 + Args.DesiredSize.Y;
+					OutCopyInfo.SourcePosition = FIntVector((SourceTexture->GetSizeX() - Args.DesiredSize.X) / 2, (SourceTexture->GetSizeY() - Args.DesiredSize.Y) / 2, 0);
 					break;
 				case EMediaCaptureCroppingType::TopLeft:
-					OutResolveParams.Rect = FResolveRect(0, 0, Args.DesiredSize.X, Args.DesiredSize.Y);
 					break;
 				case EMediaCaptureCroppingType::Custom:
-					OutResolveParams.Rect = FResolveRect(Args.MediaCapture->DesiredCaptureOptions.CustomCapturePoint.X, Args.MediaCapture->DesiredCaptureOptions.CustomCapturePoint.Y, 0, 0);
-					OutResolveParams.Rect.X2 = OutResolveParams.Rect.X1 + Args.DesiredSize.X;
-					OutResolveParams.Rect.Y2 = OutResolveParams.Rect.Y1 + Args.DesiredSize.Y;
+					OutCopyInfo.SourcePosition = FIntVector(Args.MediaCapture->DesiredCaptureOptions.CustomCapturePoint.X, Args.MediaCapture->DesiredCaptureOptions.CustomCapturePoint.Y, 0);
 					break;
 				}
 
-				OutResolveParams.DestRect.X1 = 0;
-				OutResolveParams.DestRect.X2 = Args.DesiredSize.X;
-				OutResolveParams.DestRect.Y1 = 0;
-				OutResolveParams.DestRect.Y2 = Args.DesiredSize.Y;
-
-				OutSizeU.X = (float)OutResolveParams.Rect.X1 / (float)SourceTexture->GetSizeX();
-				OutSizeU.Y = (float)OutResolveParams.Rect.X2 / (float)SourceTexture->GetSizeX();
-				OutSizeV.X = (float)OutResolveParams.Rect.Y1 / (float)SourceTexture->GetSizeY();
-				OutSizeV.Y = (float)OutResolveParams.Rect.Y2 / (float)SourceTexture->GetSizeY();
+				OutSizeU.X = (float)(OutCopyInfo.SourcePosition.X)                      / (float)SourceTexture->GetSizeX();
+				OutSizeU.Y = (float)(OutCopyInfo.SourcePosition.X + OutCopyInfo.Size.X) / (float)SourceTexture->GetSizeX();
+				OutSizeV.X = (float)(OutCopyInfo.SourcePosition.Y)                      / (float)SourceTexture->GetSizeY();
+				OutSizeV.Y = (float)(OutCopyInfo.SourcePosition.Y + OutCopyInfo.Size.Y) / (float)SourceTexture->GetSizeY();
 			}
 		}
 
@@ -417,7 +407,7 @@ namespace UE::MediaCaptureData
 			bool bRequiresFormatConversion = ConversionPassArgs.bRequiresFormatConversion;
 
 			// Rectangle area to use from source
-			const FIntRect ViewRect(ConversionPassArgs.ResolveParams.Rect.X1, ConversionPassArgs.ResolveParams.Rect.Y1, ConversionPassArgs.ResolveParams.Rect.X2, ConversionPassArgs.ResolveParams.Rect.Y2);
+			const FIntRect ViewRect(ConversionPassArgs.CopyInfo.GetSourceRect());
 
 			//Dummy ViewFamily/ViewInfo created to use built in Draw Screen/Texture Pass
 			FSceneViewFamily ViewFamily(FSceneViewFamily::ConstructionValues(nullptr, nullptr, FEngineShowFlags(ESFIM_Game))
@@ -434,7 +424,7 @@ namespace UE::MediaCaptureData
 			// If no conversion was required, go through a simple copy
 			if (Args.MediaCapture->ConversionOperation == EMediaCaptureConversionOperation::NONE && !bRequiresFormatConversion)
 			{
-				AddCopyToResolveTargetPass(ConversionPassArgs.GraphBuilder, ConversionPassArgs.SourceRGBTexture, OutputResource, ConversionPassArgs.ResolveParams);
+				AddCopyTexturePass(ConversionPassArgs.GraphBuilder, ConversionPassArgs.SourceRGBTexture, OutputResource, ConversionPassArgs.CopyInfo);
 			}
 			else
 			{
@@ -536,10 +526,10 @@ namespace UE::MediaCaptureData
 				TRACE_CPUPROFILER_EVENT_SCOPE(UMediaCapture::FrameCapture);
 				SCOPE_CYCLE_COUNTER(STAT_MediaCapture_RenderThread_FrameCapture);
 
-				FResolveParams ResolveParams;
+				FRHICopyTextureInfo CopyInfo;
 				FVector2D SizeU;
 				FVector2D SizeV;
-				GetResolveParameters(Args, SourceTexture, ResolveParams, SizeU, SizeV);
+				GetCopyInfo(Args, SourceTexture, CopyInfo, SizeU, SizeV);
 
 				FRDGBuilder GraphBuilder(Args.RHICmdList);
 
@@ -558,12 +548,12 @@ namespace UE::MediaCaptureData
 						TRACE_CPUPROFILER_EVENT_SCOPE(UMediaCapture::CustomCapture);
 
 						Args.MediaCapture->OnCustomCapture_RenderingThread(GraphBuilder, CapturingFrame->CaptureBaseData, CapturingFrame->UserData
-							, SourceRGBTexture, OutputResource, ResolveParams, SizeU, SizeV);
+							, SourceRGBTexture, OutputResource, CopyInfo, SizeU, SizeV);
 					}
 					else
 					{
 						TRACE_CPUPROFILER_EVENT_SCOPE(UMediaCapture::FormatConversion);
-						AddConversionPass(Args, { GraphBuilder, SourceRGBTexture, bRequiresFormatConversion, ResolveParams, SizeU, SizeV }, OutputResource);
+						AddConversionPass(Args, { GraphBuilder, SourceRGBTexture, bRequiresFormatConversion, CopyInfo, SizeU, SizeV }, OutputResource);
 					}
 				}
 
diff --git a/Engine/Plugins/Media/MediaIOFramework/Source/MediaIOCore/Public/MediaCapture.h b/Engine/Plugins/Media/MediaIOFramework/Source/MediaIOCore/Public/MediaCapture.h
--- a/Engine/Plugins/Media/MediaIOFramework/Source/MediaIOCore/Public/MediaCapture.h
+++ b/Engine/Plugins/Media/MediaIOFramework/Source/MediaIOCore/Public/MediaCapture.h
@@ -335,7 +335,7 @@ protected:
 	 * The texture is valid for the duration of the callback.
 	 * Only called when the conversion operation is custom and the output resource type is Texture.
 	 */
-	virtual void OnCustomCapture_RenderingThread(FRDGBuilder& GraphBuilder, const FCaptureBaseData& InBaseData, TSharedPtr<FMediaCaptureUserData, ESPMode::ThreadSafe> InUserData, FRDGTextureRef InSourceTexture, FRDGTextureRef OutputTexture, FResolveParams& ResolveParams, FVector2D CropU, FVector2D CropV)
+	virtual void OnCustomCapture_RenderingThread(FRDGBuilder& GraphBuilder, const FCaptureBaseData& InBaseData, TSharedPtr<FMediaCaptureUserData, ESPMode::ThreadSafe> InUserData, FRDGTextureRef InSourceTexture, FRDGTextureRef OutputTexture, const FRHICopyTextureInfo& CopyInfo, FVector2D CropU, FVector2D CropV)
 	{
 	}
 
@@ -346,7 +346,7 @@ protected:
 	 * The texture is valid for the duration of the callback.
 	 * Only called when the conversion operation is custom and the output resource type is Buffer.
 	 */
-	virtual void OnCustomCapture_RenderingThread(FRDGBuilder& GraphBuilder, const FCaptureBaseData& InBaseData, TSharedPtr<FMediaCaptureUserData, ESPMode::ThreadSafe> InUserData, FRDGTextureRef InSourceTexture, FRDGBufferRef OutputBuffer, FResolveParams& ResolveParams, FVector2D CropU, FVector2D CropV)
+	virtual void OnCustomCapture_RenderingThread(FRDGBuilder& GraphBuilder, const FCaptureBaseData& InBaseData, TSharedPtr<FMediaCaptureUserData, ESPMode::ThreadSafe> InUserData, FRDGTextureRef InSourceTexture, FRDGBufferRef OutputBuffer, const FRHICopyTextureInfo& CopyInfo, FVector2D CropU, FVector2D CropV)
 	{
 	}
 
diff --git a/Engine/Plugins/VirtualProduction/Rivermax/RivermaxMedia/Source/RivermaxMedia/Private/RivermaxMediaCapture.cpp b/Engine/Plugins/VirtualProduction/Rivermax/RivermaxMedia/Source/RivermaxMedia/Private/RivermaxMediaCapture.cpp
--- a/Engine/Plugins/VirtualProduction/Rivermax/RivermaxMedia/Source/RivermaxMedia/Private/RivermaxMediaCapture.cpp
+++ b/Engine/Plugins/VirtualProduction/Rivermax/RivermaxMedia/Source/RivermaxMedia/Private/RivermaxMediaCapture.cpp
@@ -309,7 +309,7 @@ FRDGBufferDesc URivermaxMediaCapture::GetCustomBufferDescription(const FIntPoint
 	return FRDGBufferDesc::CreateStructuredDesc(BytesPerElement, ElementsPerRow * InDesiredSize.Y);;
 }
 
-void URivermaxMediaCapture::OnCustomCapture_RenderingThread(FRDGBuilder& GraphBuilder, const FCaptureBaseData& InBaseData, TSharedPtr<FMediaCaptureUserData, ESPMode::ThreadSafe> InUserData, FRDGTextureRef InSourceTexture, FRDGBufferRef OutputBuffer, FResolveParams& ResolveParams, FVector2D CropU, FVector2D CropV)
+void URivermaxMediaCapture::OnCustomCapture_RenderingThread(FRDGBuilder& GraphBuilder, const FCaptureBaseData& InBaseData, TSharedPtr<FMediaCaptureUserData, ESPMode::ThreadSafe> InUserData, FRDGTextureRef InSourceTexture, FRDGBufferRef OutputBuffer, const FRHICopyTextureInfo& CopyInfo, FVector2D CropU, FVector2D CropV)
 {
 	using namespace UE::RivermaxShaders;
 	URivermaxMediaOutput* RivermaxOutput = CastChecked<URivermaxMediaOutput>(MediaOutput);
@@ -317,7 +317,7 @@ void URivermaxMediaCapture::OnCustomCapture_RenderingThread(FRDGBuilder& GraphBu
 	FGlobalShaderMap* GlobalShaderMap = GetGlobalShaderMap(GMaxRHIFeatureLevel);
 
 	// Rectangle area to use from source. This is used when source render target is bigger than output resolution
-	const FIntRect ViewRect(ResolveParams.Rect.X1, ResolveParams.Rect.Y1, ResolveParams.Rect.X2, ResolveParams.Rect.Y2);
+	const FIntRect ViewRect(CopyInfo.GetSourceRect());
 	constexpr bool bDoLinearToSRGB = false;
 	const FIntPoint SourceSize = { InSourceTexture->Desc.Extent.X, InSourceTexture->Desc.Extent.Y };
 	const FIntVector GroupCount = FComputeShaderUtils::GetGroupCount(DesiredOutputSize, FComputeShaderUtils::kGolden2DGroupSize);
diff --git a/Engine/Plugins/VirtualProduction/Rivermax/RivermaxMedia/Source/RivermaxMedia/Public/RivermaxMediaCapture.h b/Engine/Plugins/VirtualProduction/Rivermax/RivermaxMedia/Source/RivermaxMedia/Public/RivermaxMediaCapture.h
--- a/Engine/Plugins/VirtualProduction/Rivermax/RivermaxMedia/Source/RivermaxMedia/Public/RivermaxMediaCapture.h
+++ b/Engine/Plugins/VirtualProduction/Rivermax/RivermaxMedia/Source/RivermaxMedia/Public/RivermaxMediaCapture.h
@@ -47,7 +47,7 @@ protected:
 	virtual FIntPoint GetCustomOutputSize(const FIntPoint& InSize) const override;
 	virtual EMediaCaptureResourceType GetCustomOutputResourceType() const override;
 	virtual FRDGBufferDesc GetCustomBufferDescription(const FIntPoint& InDesiredSize) const override;
-	virtual void OnCustomCapture_RenderingThread(FRDGBuilder& GraphBuilder, const FCaptureBaseData& InBaseData, TSharedPtr<FMediaCaptureUserData, ESPMode::ThreadSafe> InUserData, FRDGTextureRef InSourceTexture, FRDGBufferRef OutputBuffer, FResolveParams& ResolveParams, FVector2D CropU, FVector2D CropV) override;
+	virtual void OnCustomCapture_RenderingThread(FRDGBuilder& GraphBuilder, const FCaptureBaseData& InBaseData, TSharedPtr<FMediaCaptureUserData, ESPMode::ThreadSafe> InUserData, FRDGTextureRef InSourceTexture, FRDGBufferRef OutputBuffer, const FRHICopyTextureInfo& CopyInfo, FVector2D CropU, FVector2D CropV) override;
 	//~ End UMediaCapture interface
 
 	//~ Begin IRivermaxOutputStreamListener interface
diff --git a/Engine/Source/Runtime/RHI/Public/RHI.h b/Engine/Source/Runtime/RHI/Public/RHI.h
--- a/Engine/Source/Runtime/RHI/Public/RHI.h
+++ b/Engine/Source/Runtime/RHI/Public/RHI.h
@@ -1939,6 +1939,16 @@ struct FResolveParams
 
 struct FRHICopyTextureInfo
 {
+	FIntRect GetSourceRect() const
+	{
+		return FIntRect(SourcePosition.X, SourcePosition.Y, SourcePosition.X + Size.X, SourcePosition.Y + Size.Y);
+	}
+	
+	FIntRect GetDestRect() const
+	{
+		return FIntRect(DestPosition.X, DestPosition.Y, DestPosition.X + Size.X, DestPosition.Y + Size.Y);
+	}
+
 	// Number of texels to copy. By default it will copy the whole resource if no size is specified.
 	FIntVector Size = FIntVector::ZeroValue;
 
diff --git a/Engine/Source/Runtime/Renderer/Private/PostProcess/PostProcessMaterial.cpp b/Engine/Source/Runtime/Renderer/Private/PostProcess/PostProcessMaterial.cpp
--- a/Engine/Source/Runtime/Renderer/Private/PostProcess/PostProcessMaterial.cpp
+++ b/Engine/Source/Runtime/Renderer/Private/PostProcess/PostProcessMaterial.cpp
@@ -344,23 +344,6 @@ TGlobalResource<FPostProcessMaterialVertexDeclaration> GPostProcessMaterialVerte
 
 static void AddCopyAndFlipTexturePass(FRDGBuilder& GraphBuilder, const FViewInfo& View, FRDGTextureRef SrcTexture, FRDGTextureRef DestTexture)
 {
-	if (IsOpenGLPlatform(GShaderPlatformForFeatureLevel[View.GetFeatureLevel()]))
-	{
-		// The OpenGL RHI can copy and flip at the same time by using an upside down destination rectangle.
-		FResolveParams ResolveParams;
-		FMemory::Memzero(ResolveParams);
-		ResolveParams.Rect.X1 = 0;
-		ResolveParams.Rect.X2 = SrcTexture->Desc.Extent.X;
-		ResolveParams.Rect.Y1 = 0;
-		ResolveParams.Rect.Y2 = SrcTexture->Desc.Extent.Y;
-		ResolveParams.DestRect.X1 = 0;
-		ResolveParams.DestRect.X2 = DestTexture->Desc.Extent.X;
-		ResolveParams.DestRect.Y1 = DestTexture->Desc.Extent.Y - 1;
-		ResolveParams.DestRect.Y2 = -1;
-		AddCopyToResolveTargetPass(GraphBuilder, SrcTexture, DestTexture, ResolveParams);
-		return;
-	}
-
 	// Other RHIs can't flip and copy at the same time, so we'll use a pixel shader to perform the copy, together with the FlipYAxis
 	// flag on the screen pass. This path is only taken when using the mobile preview feature in the editor with
 	// r.Mobile.ForceRHISwitchVerticalAxis set to 1, so we don't care about it being sub-optimal.
diff --git a/Engine/Source/Runtime/Renderer/Private/ReflectionEnvironmentCapture.cpp b/Engine/Source/Runtime/Renderer/Private/ReflectionEnvironmentCapture.cpp
--- a/Engine/Source/Runtime/Renderer/Private/ReflectionEnvironmentCapture.cpp
+++ b/Engine/Source/Runtime/Renderer/Private/ReflectionEnvironmentCapture.cpp
@@ -1239,7 +1239,13 @@ void CopyToSceneArray(FRDGBuilder& GraphBuilder, FScene* Scene, FRDGTexture* Fil
 	{
 		for (int32 CubeFace = 0; CubeFace < CubeFace_MAX; CubeFace++)
 		{
-			AddCopyToResolveTargetPass(GraphBuilder, FilteredCubeTexture, DestCubeTexture, FResolveParams({}, (ECubeFace)CubeFace, MipIndex, 0, CaptureIndex));
+			FRHICopyTextureInfo CopyInfo;
+			CopyInfo.SourceMipIndex   = MipIndex;
+			CopyInfo.DestMipIndex     = MipIndex;
+			CopyInfo.SourceSliceIndex = CubeFace;
+			CopyInfo.DestSliceIndex   = CaptureIndex * CubeFace_MAX + CubeFace;
+
+			AddCopyTexturePass(GraphBuilder, FilteredCubeTexture, DestCubeTexture, CopyInfo);
 		}
 	}
 }
diff --git a/Engine/Source/Runtime/Renderer/Private/SceneCaptureRendering.cpp b/Engine/Source/Runtime/Renderer/Private/SceneCaptureRendering.cpp
--- a/Engine/Source/Runtime/Renderer/Private/SceneCaptureRendering.cpp
+++ b/Engine/Source/Runtime/Renderer/Private/SceneCaptureRendering.cpp
@@ -243,7 +243,6 @@ void CopySceneCaptureComponentToTarget(
 
 				if (bNeedsFlippedRenderTarget)
 				{
-					CopyCaptureToTargetSetViewportFn(RHICmdList);
 					DrawRectangle(
 						RHICmdList,
 						View.ViewRect.Min.X, View.ViewRect.Min.Y,
@@ -257,7 +256,6 @@ void CopySceneCaptureComponentToTarget(
 				}
 				else
 				{
-					CopyCaptureToTargetSetViewportFn(RHICmdList);
 					DrawRectangle(
 						RHICmdList,
 						View.ViewRect.Min.X, View.ViewRect.Min.Y,
@@ -280,7 +278,7 @@ static void UpdateSceneCaptureContentDeferred_RenderThread(
 	FRenderTarget* RenderTarget, 
 	FTexture* RenderTargetTexture, 
 	const FString& EventName, 
-	const FResolveParams& ResolveParams,
+	const FRHICopyTextureInfo& CopyInfo,
 	bool bGenerateMips,
 	const FGenerateMipsParams& GenerateMipsParams,
 	bool bClearRenderTarget,
@@ -306,23 +304,27 @@ static void UpdateSceneCaptureContentDeferred_RenderThread(
 #endif
 
 		FRDGTextureRef TargetTexture = RegisterExternalTexture(GraphBuilder, RenderTarget->GetRenderTargetTexture(), TEXT("SceneCaptureTarget"));
+		FRDGTextureRef ShaderResourceTexture = RegisterExternalTexture(GraphBuilder, RenderTargetTexture->TextureRHI, TEXT("SceneCaptureTexture"));
+
 		if (bClearRenderTarget)
 		{
 			AddClearRenderTargetPass(GraphBuilder, TargetTexture, FLinearColor::Black, SceneRenderer->Views[0].UnscaledViewRect);
 		}
 
-		if (ResolveParams.DestRect.IsValid())
+		const FIntRect CopyDestRect = CopyInfo.GetDestRect();
+
+		if (!CopyDestRect.IsEmpty())
 		{
-			CopyCaptureToTargetSetViewportFn = [ResolveParams](FRHICommandList& RHICmdList)
+			CopyCaptureToTargetSetViewportFn = [CopyDestRect](FRHICommandList& RHICmdList)
 			{
 				RHICmdList.SetScissorRect(false, 0, 0, 0, 0);
 				RHICmdList.SetViewport
 				(
-					float(ResolveParams.DestRect.X1), 
-					float(ResolveParams.DestRect.Y1), 
-					0.0f, 
-					float(ResolveParams.DestRect.X2), 
-					float(ResolveParams.DestRect.Y2), 
+					float(CopyDestRect.Min.X),
+					float(CopyDestRect.Min.Y),
+					0.0f,
+					float(CopyDestRect.Max.X),
+					float(CopyDestRect.Max.Y),
 					1.0f
 				);
 			};
@@ -347,13 +349,12 @@ static void UpdateSceneCaptureContentDeferred_RenderThread(
 			SceneRenderer->Render(GraphBuilder);
 		}
 
-			if (bGenerateMips)
-			{
+		if (bGenerateMips)
+		{
 			FGenerateMips::Execute(GraphBuilder, TargetTexture, GenerateMipsParams);
-			}
+		}
 
-		FRDGTextureRef ResolveTexture = RegisterExternalTexture(GraphBuilder, RenderTargetTexture->TextureRHI, TEXT("SceneCaptureResolve"));
-		AddCopyToResolveTargetPass(GraphBuilder, TargetTexture, ResolveTexture, ResolveParams);
+		AddCopyTexturePass(GraphBuilder, TargetTexture, ShaderResourceTexture, CopyInfo);
 
 		GraphBuilder.Execute();
 	}
@@ -367,7 +368,7 @@ void UpdateSceneCaptureContentMobile_RenderThread(
 	FRenderTarget* RenderTarget,
 	FTexture* RenderTargetTexture,
 	const FString& EventName,
-	const FResolveParams& ResolveParams,
+	const FRHICopyTextureInfo& CopyInfo,
 	bool bGenerateMips,
 	const FGenerateMipsParams& GenerateMipsParams,
 	bool bDisableFlipCopyGLES)
@@ -427,26 +428,28 @@ void UpdateSceneCaptureContentMobile_RenderThread(
 			AddClearRenderTargetPass(GraphBuilder, FlippedOutputTexture, FLinearColor::Black, ViewRect);
 		}
 
-		if (ResolveParams.DestRect.IsValid())
+		const FIntRect CopyDestRect = CopyInfo.GetDestRect();
+
+		if (!CopyDestRect.IsEmpty())
 		{
-			CopyCaptureToTargetSetViewportFn = [ResolveParams, bNeedsFlippedFinalColor, ViewRect, FlippedOutputTexture](FRHICommandList& RHICmdList)
+			CopyCaptureToTargetSetViewportFn = [CopyDestRect, bNeedsFlippedFinalColor, ViewRect, FlippedOutputTexture](FRHICommandList& RHICmdList)
 			{
 				RHICmdList.SetScissorRect(false, 0, 0, 0, 0);
 
 				if (bNeedsFlippedFinalColor)
 				{
-					FResolveRect DestRect = ResolveParams.DestRect;
-					int32 TileYID = DestRect.Y1 / ViewRect.Height();
+					FIntRect DestRect = CopyDestRect;
+					int32 TileYID = DestRect.Min.Y / ViewRect.Height();
 					int32 TileYCount = (FlippedOutputTexture->Desc.GetSize().Y / ViewRect.Height()) - 1;
-					DestRect.Y1 = (TileYCount - TileYID) * ViewRect.Height();
-					DestRect.Y2 = DestRect.Y1 + ViewRect.Height();
+					DestRect.Min.Y = (TileYCount - TileYID) * ViewRect.Height();
+					DestRect.Max.Y = DestRect.Min.Y + ViewRect.Height();
 					RHICmdList.SetViewport
 					(
-						float(DestRect.X1),
-						float(DestRect.Y1),
+						float(DestRect.Min.X),
+						float(DestRect.Min.Y),
 						0.0f,
-						float(DestRect.X2),
-						float(DestRect.Y2),
+						float(DestRect.Max.X),
+						float(DestRect.Max.Y),
 						1.0f
 					);
 				}
@@ -454,11 +457,11 @@ void UpdateSceneCaptureContentMobile_RenderThread(
 				{
 					RHICmdList.SetViewport
 					(
-						float(ResolveParams.DestRect.X1),
-						float(ResolveParams.DestRect.Y1),
+						float(CopyDestRect.Min.X),
+						float(CopyDestRect.Min.Y),
 						0.0f,
-						float(ResolveParams.DestRect.X2),
-						float(ResolveParams.DestRect.Y2),
+						float(CopyDestRect.Max.X),
+						float(CopyDestRect.Max.Y),
 						1.0f
 					);
 				}
@@ -593,7 +596,7 @@ static void UpdateSceneCaptureContent_RenderThread(
 	FRenderTarget* RenderTarget,
 	FTexture* RenderTargetTexture,
 	const FString& EventName,
-	const FResolveParams& ResolveParams,
+	const FRHICopyTextureInfo& CopyInfo,
 	bool bGenerateMips,
 	const FGenerateMipsParams& GenerateMipsParams,
 	const bool bDisableFlipCopyLDRGLES,
@@ -612,7 +615,7 @@ static void UpdateSceneCaptureContent_RenderThread(
 				RenderTarget,
 				RenderTargetTexture,
 				EventName,
-				ResolveParams,
+				CopyInfo,
 				bGenerateMips,
 				GenerateMipsParams,
 				bDisableFlipCopyLDRGLES);
@@ -626,7 +629,7 @@ static void UpdateSceneCaptureContent_RenderThread(
 				RenderTarget,
 				RenderTargetTexture,
 				EventName,
-				ResolveParams,
+				CopyInfo,
 				bGenerateMips,
 				GenerateMipsParams,
 				bClearRenderTarget,
@@ -1101,7 +1104,7 @@ void FScene::UpdateSceneCaptureContents(USceneCaptureComponent2D* CaptureCompone
 					TextureRenderTargetResource->SetActiveGPUMask(FRHIGPUMask::All());
 				}
 
-				FResolveParams ResolveParams;
+				FRHICopyTextureInfo CopyInfo;
 
 				if (bEnableOrthographicTiling)
 				{
@@ -1109,15 +1112,15 @@ void FScene::UpdateSceneCaptureContents(USceneCaptureComponent2D* CaptureCompone
 					const uint32 RTSizeY = TextureRenderTargetResource->GetSizeY() / NumYTiles;
 					const uint32 TileX = TileID % NumXTiles;
 					const uint32 TileY = TileID / NumXTiles;
-					ResolveParams.DestRect.X1 = TileX * RTSizeX;
-					ResolveParams.DestRect.Y1 = TileY * RTSizeY;
-					ResolveParams.DestRect.X2 = ResolveParams.DestRect.X1 + RTSizeX;
-					ResolveParams.DestRect.Y2 = ResolveParams.DestRect.Y1 + RTSizeY;
+					CopyInfo.DestPosition.X = TileX * RTSizeX;
+					CopyInfo.DestPosition.Y = TileY * RTSizeY;
+					CopyInfo.Size.X = RTSizeX;
+					CopyInfo.Size.Y = RTSizeY;
 				}
 
 				RectLightAtlas::FAtlasTextureInvalidationScope Invalidation(TexturePtrNotDeferenced);
 
-				UpdateSceneCaptureContent_RenderThread(RHICmdList, SceneRenderer, TextureRenderTargetResource, TextureRenderTargetResource, EventName, ResolveParams, bGenerateMips, GenerateMipsParams, bDisableFlipCopyGLES, !bEnableOrthographicTiling, bOrthographicCamera);
+				UpdateSceneCaptureContent_RenderThread(RHICmdList, SceneRenderer, TextureRenderTargetResource, TextureRenderTargetResource, EventName, CopyInfo, bGenerateMips, GenerateMipsParams, bDisableFlipCopyGLES, !bEnableOrthographicTiling, bOrthographicCamera);
 			}
 		);
 	}
@@ -1244,7 +1247,10 @@ void FScene::UpdateSceneCaptureContents(USceneCaptureComponentCube* CaptureCompo
 				ENQUEUE_RENDER_COMMAND(CaptureCommand)(
 					[SceneRenderer, TextureRenderTarget, EventName, TargetFace](FRHICommandListImmediate& RHICmdList)
 					{
-						UpdateSceneCaptureContent_RenderThread(RHICmdList, SceneRenderer, TextureRenderTarget, TextureRenderTarget, EventName, FResolveParams(FResolveRect(), TargetFace), false, FGenerateMipsParams(), false, true, false);
+						FRHICopyTextureInfo CopyInfo;
+						CopyInfo.SourceSliceIndex = TargetFace;
+						CopyInfo.DestSliceIndex   = TargetFace; 
+						UpdateSceneCaptureContent_RenderThread(RHICmdList, SceneRenderer, TextureRenderTarget, TextureRenderTarget, EventName, CopyInfo, false, FGenerateMipsParams(), false, true, false);
 					}
 				);
 			}
diff --git a/Engine/Source/Runtime/Renderer/Private/SceneRendering.cpp b/Engine/Source/Runtime/Renderer/Private/SceneRendering.cpp
--- a/Engine/Source/Runtime/Renderer/Private/SceneRendering.cpp
+++ b/Engine/Source/Runtime/Renderer/Private/SceneRendering.cpp
@@ -5149,7 +5149,6 @@ public:
 };
 
 TGlobalResource<FDummySceneColorResolveBuffer> GResolveDummyVertexBuffer;
-extern int32 GAllowCustomMSAAResolves;
 
 BEGIN_SHADER_PARAMETER_STRUCT(FResolveSceneColorParameters, )
 	RDG_TEXTURE_ACCESS(SceneColor, ERHIAccess::SRVGraphics)
@@ -5169,152 +5168,139 @@ void AddResolveSceneColorPass(FRDGBuilder& GraphBuilder, const FViewInfo& View,
 		return;
 	}
 
-	if (!GAllowCustomMSAAResolves)
+	FRDGTextureSRVRef SceneColorFMask = nullptr;
+
+	if (GRHISupportsExplicitFMask)
 	{
-		FResolveRect ResolveRect(View.ViewRect);
-		if (View.IsInstancedStereoPass())
-		{
-			ResolveRect.X1 = 0;
-			ResolveRect.X2 = View.InstancedStereoWidth;
-		}
-		AddCopyToResolveTargetPass(GraphBuilder, SceneColor.Target, SceneColor.Resolve, ResolveRect);
+		SceneColorFMask = GraphBuilder.CreateSRV(FRDGTextureSRVDesc::CreateForMetaData(SceneColor.Target, ERDGTextureMetaDataAccess::FMask));
 	}
-	else
+
+	FResolveSceneColorParameters* PassParameters = GraphBuilder.AllocParameters<FResolveSceneColorParameters>();
+	PassParameters->SceneColor = SceneColor.Target;
+	PassParameters->SceneColorFMask = SceneColorFMask;
+	PassParameters->RenderTargets[0] = FRenderTargetBinding(SceneColor.Resolve, SceneColor.Resolve->HasBeenProduced() ? ERenderTargetLoadAction::ELoad : ERenderTargetLoadAction::ENoAction);
+
+	FRDGTextureRef SceneColorTargetable = SceneColor.Target;
+
+	GraphBuilder.AddPass(
+		RDG_EVENT_NAME("ResolveSceneColor"),
+		PassParameters,
+		ERDGPassFlags::Raster,
+		[&View, SceneColorTargetable, SceneColorFMask, NumSamples](FRHICommandList& RHICmdList)
 	{
-		FRDGTextureSRVRef SceneColorFMask = nullptr;
+		FRHITexture* SceneColorTargetableRHI = SceneColorTargetable->GetRHI();
+		SceneColorTargetable->MarkResourceAsUsed();
 
-		if (GRHISupportsExplicitFMask)
+		FRHIShaderResourceView* SceneColorFMaskRHI = nullptr;
+		if (SceneColorFMask)
 		{
-			SceneColorFMask = GraphBuilder.CreateSRV(FRDGTextureSRVDesc::CreateForMetaData(SceneColor.Target, ERDGTextureMetaDataAccess::FMask));
+			SceneColorFMask->MarkResourceAsUsed();
+			SceneColorFMaskRHI = SceneColorFMask->GetRHI();
 		}
 
-		FResolveSceneColorParameters* PassParameters = GraphBuilder.AllocParameters<FResolveSceneColorParameters>();
-		PassParameters->SceneColor = SceneColor.Target;
-		PassParameters->SceneColorFMask = SceneColorFMask;
-		PassParameters->RenderTargets[0] = FRenderTargetBinding(SceneColor.Resolve, SceneColor.Resolve->HasBeenProduced() ? ERenderTargetLoadAction::ELoad : ERenderTargetLoadAction::ENoAction);
+		FGraphicsPipelineStateInitializer GraphicsPSOInit;
+		RHICmdList.ApplyCachedRenderTargets(GraphicsPSOInit);
 
-		FRDGTextureRef SceneColorTargetable = SceneColor.Target;
+		GraphicsPSOInit.BlendState = TStaticBlendState<>::GetRHI();
+		GraphicsPSOInit.RasterizerState = TStaticRasterizerState<>::GetRHI();
+		GraphicsPSOInit.DepthStencilState = TStaticDepthStencilState<false, CF_Always>::GetRHI();
 
-		GraphBuilder.AddPass(
-			RDG_EVENT_NAME("ResolveSceneColor"),
-			PassParameters,
-			ERDGPassFlags::Raster,
-			[&View, SceneColorTargetable, SceneColorFMask, NumSamples](FRHICommandList& RHICmdList)
-		{
-			FRHITexture* SceneColorTargetableRHI = SceneColorTargetable->GetRHI();
-			SceneColorTargetable->MarkResourceAsUsed();
+		const FIntPoint SceneColorExtent = SceneColorTargetable->Desc.Extent;
 
-			FRHIShaderResourceView* SceneColorFMaskRHI = nullptr;
-			if (SceneColorFMask)
-			{
-				SceneColorFMask->MarkResourceAsUsed();
-				SceneColorFMaskRHI = SceneColorFMask->GetRHI();
-			}
+		// Resolve views individually. In the case of adaptive resolution, the view family will be much larger than the views individually.
+		RHICmdList.SetViewport(0.0f, 0.0f, 0.0f, SceneColorExtent.X, SceneColorExtent.Y, 1.0f);
+		RHICmdList.SetScissorRect(true, View.IsInstancedStereoPass() ? 0 : View.ViewRect.Min.X, View.ViewRect.Min.Y,
+			View.IsInstancedStereoPass() ? View.InstancedStereoWidth : View.ViewRect.Max.X, View.ViewRect.Max.Y);
 
-			FGraphicsPipelineStateInitializer GraphicsPSOInit;
-			RHICmdList.ApplyCachedRenderTargets(GraphicsPSOInit);
+		int32 ResolveWidth = CVarWideCustomResolve.GetValueOnRenderThread();
 
-			GraphicsPSOInit.BlendState = TStaticBlendState<>::GetRHI();
-			GraphicsPSOInit.RasterizerState = TStaticRasterizerState<>::GetRHI();
-			GraphicsPSOInit.DepthStencilState = TStaticDepthStencilState<false, CF_Always>::GetRHI();
+		if (NumSamples <= 1)
+		{
+			ResolveWidth = 0;
+		}
 
-			const FIntPoint SceneColorExtent = SceneColorTargetable->Desc.Extent;
+		if (ResolveWidth != 0)
+		{
+			ResolveFilterWide(RHICmdList, GraphicsPSOInit, View.FeatureLevel, SceneColorTargetableRHI, SceneColorFMaskRHI, FIntPoint(0, 0), NumSamples, ResolveWidth, GResolveDummyVertexBuffer.VertexBufferRHI);
+		}
+		else
+		{
+			TShaderMapRef<FHdrCustomResolveVS> VertexShader(View.ShaderMap);
+			GraphicsPSOInit.BoundShaderState.VertexDeclarationRHI = GetVertexDeclarationFVector4();
+			GraphicsPSOInit.BoundShaderState.VertexShaderRHI = VertexShader.GetVertexShader();
+			GraphicsPSOInit.PrimitiveType = PT_TriangleList;
 
-			// Resolve views individually. In the case of adaptive resolution, the view family will be much larger than the views individually.
-			RHICmdList.SetViewport(0.0f, 0.0f, 0.0f, SceneColorExtent.X, SceneColorExtent.Y, 1.0f);
-			RHICmdList.SetScissorRect(true, View.IsInstancedStereoPass() ? 0 : View.ViewRect.Min.X, View.ViewRect.Min.Y,
-				View.IsInstancedStereoPass() ? View.InstancedStereoWidth : View.ViewRect.Max.X, View.ViewRect.Max.Y);
+			if (SceneColorFMaskRHI)
+			{
+				if (NumSamples == 2)
+				{
+					TShaderMapRef<FHdrCustomResolveFMask2xPS> PixelShader(View.ShaderMap);
+					GraphicsPSOInit.BoundShaderState.PixelShaderRHI = PixelShader.GetPixelShader();
 
-			int32 ResolveWidth = CVarWideCustomResolve.GetValueOnRenderThread();
+					SetGraphicsPipelineState(RHICmdList, GraphicsPSOInit, 0);
+					PixelShader->SetParameters(RHICmdList, SceneColorTargetableRHI, SceneColorFMaskRHI);
+				}
+				else if (NumSamples == 4)
+				{
+					TShaderMapRef<FHdrCustomResolveFMask4xPS> PixelShader(View.ShaderMap);
+					GraphicsPSOInit.BoundShaderState.PixelShaderRHI = PixelShader.GetPixelShader();
 
-			if (NumSamples <= 1)
-			{
-				ResolveWidth = 0;
-			}
+					SetGraphicsPipelineState(RHICmdList, GraphicsPSOInit, 0);
+					PixelShader->SetParameters(RHICmdList, SceneColorTargetableRHI, SceneColorFMaskRHI);
+				}
+				else if (NumSamples == 8)
+				{
+					TShaderMapRef<FHdrCustomResolveFMask8xPS> PixelShader(View.ShaderMap);
+					GraphicsPSOInit.BoundShaderState.PixelShaderRHI = PixelShader.GetPixelShader();
 
-			if (ResolveWidth != 0)
-			{
-				ResolveFilterWide(RHICmdList, GraphicsPSOInit, View.FeatureLevel, SceneColorTargetableRHI, SceneColorFMaskRHI, FIntPoint(0, 0), NumSamples, ResolveWidth, GResolveDummyVertexBuffer.VertexBufferRHI);
+					SetGraphicsPipelineState(RHICmdList, GraphicsPSOInit, 0);
+					PixelShader->SetParameters(RHICmdList, SceneColorTargetableRHI, SceneColorFMaskRHI);
+				}
+				else
+				{
+					// Everything other than 2,4,8 samples is not implemented.
+					checkNoEntry();
+				}
 			}
 			else
 			{
-				TShaderMapRef<FHdrCustomResolveVS> VertexShader(View.ShaderMap);
-				GraphicsPSOInit.BoundShaderState.VertexDeclarationRHI = GetVertexDeclarationFVector4();
-				GraphicsPSOInit.BoundShaderState.VertexShaderRHI = VertexShader.GetVertexShader();
-				GraphicsPSOInit.PrimitiveType = PT_TriangleList;
-
-				if (SceneColorFMaskRHI)
+				if (NumSamples == 2)
 				{
-					if (NumSamples == 2)
-					{
-						TShaderMapRef<FHdrCustomResolveFMask2xPS> PixelShader(View.ShaderMap);
-						GraphicsPSOInit.BoundShaderState.PixelShaderRHI = PixelShader.GetPixelShader();
+					TShaderMapRef<FHdrCustomResolve2xPS> PixelShader(View.ShaderMap);
+					GraphicsPSOInit.BoundShaderState.PixelShaderRHI = PixelShader.GetPixelShader();
 
-						SetGraphicsPipelineState(RHICmdList, GraphicsPSOInit, 0);
-						PixelShader->SetParameters(RHICmdList, SceneColorTargetableRHI, SceneColorFMaskRHI);
-					}
-					else if (NumSamples == 4)
-					{
-						TShaderMapRef<FHdrCustomResolveFMask4xPS> PixelShader(View.ShaderMap);
-						GraphicsPSOInit.BoundShaderState.PixelShaderRHI = PixelShader.GetPixelShader();
+					SetGraphicsPipelineState(RHICmdList, GraphicsPSOInit, 0);
+					PixelShader->SetParameters(RHICmdList, SceneColorTargetableRHI);
+				}
+				else if (NumSamples == 4)
+				{
+					TShaderMapRef<FHdrCustomResolve4xPS> PixelShader(View.ShaderMap);
+					GraphicsPSOInit.BoundShaderState.PixelShaderRHI = PixelShader.GetPixelShader();
 
-						SetGraphicsPipelineState(RHICmdList, GraphicsPSOInit, 0);
-						PixelShader->SetParameters(RHICmdList, SceneColorTargetableRHI, SceneColorFMaskRHI);
-					}
-					else if (NumSamples == 8)
-					{
-						TShaderMapRef<FHdrCustomResolveFMask8xPS> PixelShader(View.ShaderMap);
-						GraphicsPSOInit.BoundShaderState.PixelShaderRHI = PixelShader.GetPixelShader();
+					SetGraphicsPipelineState(RHICmdList, GraphicsPSOInit, 0);
+					PixelShader->SetParameters(RHICmdList, SceneColorTargetableRHI);
+				}
+				else if (NumSamples == 8)
+				{
+					TShaderMapRef<FHdrCustomResolve8xPS> PixelShader(View.ShaderMap);
+					GraphicsPSOInit.BoundShaderState.PixelShaderRHI = PixelShader.GetPixelShader();
 
-						SetGraphicsPipelineState(RHICmdList, GraphicsPSOInit, 0);
-						PixelShader->SetParameters(RHICmdList, SceneColorTargetableRHI, SceneColorFMaskRHI);
-					}
-					else
-					{
-						// Everything other than 2,4,8 samples is not implemented.
-						checkNoEntry();
-					}
+					SetGraphicsPipelineState(RHICmdList, GraphicsPSOInit, 0);
+					PixelShader->SetParameters(RHICmdList, SceneColorTargetableRHI);
 				}
 				else
 				{
-					if (NumSamples == 2)
-					{
-						TShaderMapRef<FHdrCustomResolve2xPS> PixelShader(View.ShaderMap);
-						GraphicsPSOInit.BoundShaderState.PixelShaderRHI = PixelShader.GetPixelShader();
-
-						SetGraphicsPipelineState(RHICmdList, GraphicsPSOInit, 0);
-						PixelShader->SetParameters(RHICmdList, SceneColorTargetableRHI);
-					}
-					else if (NumSamples == 4)
-					{
-						TShaderMapRef<FHdrCustomResolve4xPS> PixelShader(View.ShaderMap);
-						GraphicsPSOInit.BoundShaderState.PixelShaderRHI = PixelShader.GetPixelShader();
-
-						SetGraphicsPipelineState(RHICmdList, GraphicsPSOInit, 0);
-						PixelShader->SetParameters(RHICmdList, SceneColorTargetableRHI);
-					}
-					else if (NumSamples == 8)
-					{
-						TShaderMapRef<FHdrCustomResolve8xPS> PixelShader(View.ShaderMap);
-						GraphicsPSOInit.BoundShaderState.PixelShaderRHI = PixelShader.GetPixelShader();
-
-						SetGraphicsPipelineState(RHICmdList, GraphicsPSOInit, 0);
-						PixelShader->SetParameters(RHICmdList, SceneColorTargetableRHI);
-					}
-					else
-					{
-						// Everything other than 2,4,8 samples is not implemented.
-						checkNoEntry();
-					}
+					// Everything other than 2,4,8 samples is not implemented.
+					checkNoEntry();
 				}
-
-				RHICmdList.SetStreamSource(0, GResolveDummyVertexBuffer.VertexBufferRHI, 0);
-				RHICmdList.DrawPrimitive(0, 1, 1);
 			}
 
-			RHICmdList.SetScissorRect(false, 0, 0, 0, 0);
-		});
-	}
+			RHICmdList.SetStreamSource(0, GResolveDummyVertexBuffer.VertexBufferRHI, 0);
+			RHICmdList.DrawPrimitive(0, 1, 1);
+		}
+
+		RHICmdList.SetScissorRect(false, 0, 0, 0, 0);
+	});
 }
 
 void AddResolveSceneColorPass(FRDGBuilder& GraphBuilder, TArrayView<const FViewInfo> Views, FRDGTextureMSAA SceneColor)
@@ -5353,81 +5339,74 @@ void AddResolveSceneDepthPass(FRDGBuilder& GraphBuilder, const FViewInfo& View,
 		ResolveRect.X2 = View.InstancedStereoWidth;
 	}
 
-	if (!GAllowCustomMSAAResolves)
-	{
-		AddCopyToResolveTargetPass(GraphBuilder, SceneDepth.Target, SceneDepth.Resolve, ResolveRect);
-	}
-	else
-	{
-		const FIntPoint DepthExtent = SceneDepth.Resolve->Desc.Extent;
+	const FIntPoint DepthExtent = SceneDepth.Resolve->Desc.Extent;
 
-		FResolveSceneDepthParameters* PassParameters = GraphBuilder.AllocParameters<FResolveSceneDepthParameters>();
-		PassParameters->SceneDepth = SceneDepth.Target;
-		PassParameters->RenderTargets.DepthStencil = FDepthStencilBinding(SceneDepth.Resolve, ERenderTargetLoadAction::ENoAction, ERenderTargetLoadAction::ENoAction, FExclusiveDepthStencil::DepthWrite_StencilWrite);
+	FResolveSceneDepthParameters* PassParameters = GraphBuilder.AllocParameters<FResolveSceneDepthParameters>();
+	PassParameters->SceneDepth = SceneDepth.Target;
+	PassParameters->RenderTargets.DepthStencil = FDepthStencilBinding(SceneDepth.Resolve, ERenderTargetLoadAction::ENoAction, ERenderTargetLoadAction::ENoAction, FExclusiveDepthStencil::DepthWrite_StencilWrite);
 
-		FRDGTextureRef SourceTexture = SceneDepth.Target;
+	FRDGTextureRef SourceTexture = SceneDepth.Target;
 
-		GraphBuilder.AddPass(
-			RDG_EVENT_NAME("ResolveSceneDepth"),
-			PassParameters,
-			ERDGPassFlags::Raster,
-			[&View, SourceTexture, NumSamples, DepthExtent, ResolveRect](FRHICommandList& RHICmdList)
-		{
-			FRHITexture* SourceTextureRHI = SourceTexture->GetRHI();
-			SourceTexture->MarkResourceAsUsed();
+	GraphBuilder.AddPass(
+		RDG_EVENT_NAME("ResolveSceneDepth"),
+		PassParameters,
+		ERDGPassFlags::Raster,
+		[&View, SourceTexture, NumSamples, DepthExtent, ResolveRect](FRHICommandList& RHICmdList)
+	{
+		FRHITexture* SourceTextureRHI = SourceTexture->GetRHI();
+		SourceTexture->MarkResourceAsUsed();
 
-			FGraphicsPipelineStateInitializer GraphicsPSOInit;
-			RHICmdList.ApplyCachedRenderTargets(GraphicsPSOInit);
-			GraphicsPSOInit.BlendState = TStaticBlendState<>::GetRHI();
-			GraphicsPSOInit.RasterizerState = TStaticRasterizerState<FM_Solid, CM_None>::GetRHI();
-			GraphicsPSOInit.DepthStencilState = TStaticDepthStencilState<true, CF_Always>::GetRHI();
+		FGraphicsPipelineStateInitializer GraphicsPSOInit;
+		RHICmdList.ApplyCachedRenderTargets(GraphicsPSOInit);
+		GraphicsPSOInit.BlendState = TStaticBlendState<>::GetRHI();
+		GraphicsPSOInit.RasterizerState = TStaticRasterizerState<FM_Solid, CM_None>::GetRHI();
+		GraphicsPSOInit.DepthStencilState = TStaticDepthStencilState<true, CF_Always>::GetRHI();
 
-			RHICmdList.SetViewport(0.0f, 0.0f, 0.0f, DepthExtent.X, DepthExtent.Y, 1.0f);
+		RHICmdList.SetViewport(0.0f, 0.0f, 0.0f, DepthExtent.X, DepthExtent.Y, 1.0f);
 
-			TShaderMapRef<FResolveVS> ResolveVertexShader(View.ShaderMap);
-			TShaderMapRef<FResolveDepthPS> ResolvePixelShaderAny(View.ShaderMap);
-			TShaderMapRef<FResolveDepth2XPS> ResolvePixelShader2X(View.ShaderMap);
-			TShaderMapRef<FResolveDepth4XPS> ResolvePixelShader4X(View.ShaderMap);
-			TShaderMapRef<FResolveDepth8XPS> ResolvePixelShader8X(View.ShaderMap);
+		TShaderMapRef<FResolveVS> ResolveVertexShader(View.ShaderMap);
+		TShaderMapRef<FResolveDepthPS> ResolvePixelShaderAny(View.ShaderMap);
+		TShaderMapRef<FResolveDepth2XPS> ResolvePixelShader2X(View.ShaderMap);
+		TShaderMapRef<FResolveDepth4XPS> ResolvePixelShader4X(View.ShaderMap);
+		TShaderMapRef<FResolveDepth8XPS> ResolvePixelShader8X(View.ShaderMap);
 
-			int32 TextureIndex = -1;
-			FRHIPixelShader* ResolvePixelShader;
-			switch (NumSamples)
-			{
-			case 2:
-				TextureIndex = ResolvePixelShader2X->UnresolvedSurface.GetBaseIndex();
-				ResolvePixelShader = ResolvePixelShader2X.GetPixelShader();
-				break;
-			case 4:
-				TextureIndex = ResolvePixelShader4X->UnresolvedSurface.GetBaseIndex();
-				ResolvePixelShader = ResolvePixelShader4X.GetPixelShader();
-				break;
-			case 8:
-				TextureIndex = ResolvePixelShader8X->UnresolvedSurface.GetBaseIndex();
-				ResolvePixelShader = ResolvePixelShader8X.GetPixelShader();
-				break;
-			default:
-				ensureMsgf(false, TEXT("Unsupported depth resolve for samples: %i.  Dynamic loop method isn't supported on all platforms.  Please add specific case."), NumSamples);
-				TextureIndex = ResolvePixelShaderAny->UnresolvedSurface.GetBaseIndex();
-				ResolvePixelShader = ResolvePixelShaderAny.GetPixelShader();
-				break;
-			}
+		int32 TextureIndex = -1;
+		FRHIPixelShader* ResolvePixelShader;
+		switch (NumSamples)
+		{
+		case 2:
+			TextureIndex = ResolvePixelShader2X->UnresolvedSurface.GetBaseIndex();
+			ResolvePixelShader = ResolvePixelShader2X.GetPixelShader();
+			break;
+		case 4:
+			TextureIndex = ResolvePixelShader4X->UnresolvedSurface.GetBaseIndex();
+			ResolvePixelShader = ResolvePixelShader4X.GetPixelShader();
+			break;
+		case 8:
+			TextureIndex = ResolvePixelShader8X->UnresolvedSurface.GetBaseIndex();
+			ResolvePixelShader = ResolvePixelShader8X.GetPixelShader();
+			break;
+		default:
+			ensureMsgf(false, TEXT("Unsupported depth resolve for samples: %i.  Dynamic loop method isn't supported on all platforms.  Please add specific case."), NumSamples);
+			TextureIndex = ResolvePixelShaderAny->UnresolvedSurface.GetBaseIndex();
+			ResolvePixelShader = ResolvePixelShaderAny.GetPixelShader();
+			break;
+		}
 
-			GraphicsPSOInit.BoundShaderState.VertexDeclarationRHI = GEmptyVertexDeclaration.VertexDeclarationRHI;
-			GraphicsPSOInit.BoundShaderState.VertexShaderRHI = ResolveVertexShader.GetVertexShader();
-			GraphicsPSOInit.BoundShaderState.PixelShaderRHI = ResolvePixelShader;
-			GraphicsPSOInit.PrimitiveType = PT_TriangleStrip;
+		GraphicsPSOInit.BoundShaderState.VertexDeclarationRHI = GEmptyVertexDeclaration.VertexDeclarationRHI;
+		GraphicsPSOInit.BoundShaderState.VertexShaderRHI = ResolveVertexShader.GetVertexShader();
+		GraphicsPSOInit.BoundShaderState.PixelShaderRHI = ResolvePixelShader;
+		GraphicsPSOInit.PrimitiveType = PT_TriangleStrip;
 
-			SetGraphicsPipelineState(RHICmdList, GraphicsPSOInit, 0);
-			RHICmdList.SetBlendFactor(FLinearColor::White);
-			RHICmdList.SetShaderTexture(ResolvePixelShader, TextureIndex, SourceTextureRHI);
+		SetGraphicsPipelineState(RHICmdList, GraphicsPSOInit, 0);
+		RHICmdList.SetBlendFactor(FLinearColor::White);
+		RHICmdList.SetShaderTexture(ResolvePixelShader, TextureIndex, SourceTextureRHI);
 
-			ResolveVertexShader->SetParameters(RHICmdList, ResolveRect, ResolveRect, DepthExtent.X, DepthExtent.Y);
+		ResolveVertexShader->SetParameters(RHICmdList, ResolveRect, ResolveRect, DepthExtent.X, DepthExtent.Y);
 
-			RHICmdList.SetStreamSource(0, nullptr, 0);
-			RHICmdList.DrawPrimitive(0, 2, 1);
-		});
-	}
+		RHICmdList.SetStreamSource(0, nullptr, 0);
+		RHICmdList.DrawPrimitive(0, 2, 1);
+	});
 }
 
 void AddResolveSceneDepthPass(FRDGBuilder& GraphBuilder, TArrayView<const FViewInfo> Views, FRDGTextureMSAA SceneDepth)
diff --git a/Engine/Source/Runtime/Renderer/Private/SceneTextures.cpp b/Engine/Source/Runtime/Renderer/Private/SceneTextures.cpp
--- a/Engine/Source/Runtime/Renderer/Private/SceneTextures.cpp
+++ b/Engine/Source/Runtime/Renderer/Private/SceneTextures.cpp
@@ -75,14 +75,6 @@ static TAutoConsoleVariable<int32> CVarDefaultBackBufferPixelFormat(
 	TEXT(" 4: 10bit RGB, 2bit Alpha\n"),
 	ECVF_ReadOnly);
 
-int32 GAllowCustomMSAAResolves = 1;
-static FAutoConsoleVariableRef CVarAllowCustomResolves(
-   TEXT("r.MSAA.AllowCustomResolves"),
-   GAllowCustomMSAAResolves,
-   TEXT("Whether to use builtin HW resolve or allow custom shader MSAA resolves"),
-   ECVF_RenderThreadSafe
-   );
-
 IMPLEMENT_STATIC_UNIFORM_BUFFER_SLOT(SceneTextures);
 IMPLEMENT_STATIC_UNIFORM_BUFFER_STRUCT(FSceneTextureUniformParameters, "SceneTexturesStruct", SceneTextures);
 IMPLEMENT_STATIC_UNIFORM_BUFFER_STRUCT(FMobileSceneTextureUniformParameters, "MobileSceneTextures", SceneTextures);
diff --git a/Engine/Source/Runtime/Renderer/Private/TranslucentRendering.cpp b/Engine/Source/Runtime/Renderer/Private/TranslucentRendering.cpp
--- a/Engine/Source/Runtime/Renderer/Private/TranslucentRendering.cpp
+++ b/Engine/Source/Runtime/Renderer/Private/TranslucentRendering.cpp
@@ -451,7 +451,7 @@ static FRDGTextureRef AddCopySceneColorPass(FRDGBuilder& GraphBuilder, TArrayVie
 			RDG_GPU_MASK_SCOPE(GraphBuilder, View.GPUMask);
 			RDG_EVENT_SCOPE_CONDITIONAL(GraphBuilder, Views.Num() > 1, "View%d", ViewIndex);
 
-			AddCopyToResolveTargetPass(GraphBuilder, SceneColor.Target, SceneColor.Resolve, FResolveRect(View.ViewRect));
+			AddResolveSceneColorPass(GraphBuilder, View, SceneColor);
 
 			const FIntPoint SceneColorExtent = SceneColor.Target->Desc.Extent;
 
