Fixed double reflections when Lumen GI is enabled with Standalone Ray Traced Reflections, Lumen GI was outputting rough specular

diff --git a/Engine/Source/Runtime/Renderer/Private/IndirectLightRendering.cpp b/Engine/Source/Runtime/Renderer/Private/IndirectLightRendering.cpp
--- a/Engine/Source/Runtime/Renderer/Private/IndirectLightRendering.cpp
+++ b/Engine/Source/Runtime/Renderer/Private/IndirectLightRendering.cpp
@@ -890,14 +890,14 @@ void FDeferredShadingSceneRenderer::RenderDiffuseIndirectAndAmbientOcclusion(
 					nullptr,
 					LumenReflectionCompositeParameters);
 			}
+			else
+			{
+				DenoiserOutputs.Textures[1] = SystemTextures.Black;
+				DenoiserOutputs.Textures[2] = SystemTextures.Black;
+			}
 
 			// Lumen needs its own depth history because things like Translucency velocities write to depth
 			StoreLumenDepthHistory(GraphBuilder, SceneTextures, View);
-
-			if (!DenoiserOutputs.Textures[2])
-			{
-				DenoiserOutputs.Textures[2] = DenoiserOutputs.Textures[1];
-			}
 		}
 		else if (ViewPipelineState.DiffuseIndirectMethod == EDiffuseIndirectMethod::Plugin)
 		{
