Fix Lumen GI on Volumetric Fog - fog is moved to after the new place that Lumen Translucency Volume GI is calculated as of cl 19855871

diff --git a/Engine/Source/Runtime/Renderer/Private/DeferredShadingRenderer.cpp b/Engine/Source/Runtime/Renderer/Private/DeferredShadingRenderer.cpp
--- a/Engine/Source/Runtime/Renderer/Private/DeferredShadingRenderer.cpp
+++ b/Engine/Source/Runtime/Renderer/Private/DeferredShadingRenderer.cpp
@@ -2812,12 +2812,6 @@ void FDeferredShadingSceneRenderer::Render(FRDGBuilder& GraphBuilder)
 			LLM_SCOPE_BYTAG(Lumen);
 			RenderLumenSceneLighting(GraphBuilder, LumenFrameTemporaries);
 		}
-
-		// If forward shading is enabled, we computed fog earlier already
-		if (!IsForwardShadingEnabled(ShaderPlatform))
-		{
-			ComputeVolumetricFog(GraphBuilder, SceneTextures);
-		}
 	}
 	// End shadow and fog after base pass
 
@@ -2982,6 +2976,12 @@ void FDeferredShadingSceneRenderer::Render(FRDGBuilder& GraphBuilder)
 		RenderDeferredReflectionsAndSkyLightingHair(GraphBuilder);
 	}
 
+	// Volumetric fog after Lumen GI and shadow depths
+	if (!IsForwardShadingEnabled(ShaderPlatform) && !bHasRayTracedOverlay)
+	{
+		ComputeVolumetricFog(GraphBuilder, SceneTextures);
+	}
+
 	if (bShouldRenderVolumetricCloud && IsVolumetricRenderTargetEnabled() && !bHasHalfResCheckerboardMinMaxDepth && !bHasRayTracedOverlay)
 	{
 		HalfResolutionDepthCheckerboardMinMaxTexture = CreateHalfResolutionDepthCheckerboardMinMax(GraphBuilder, Views, SceneTextures.Depth.Resolve);
