Tentative fix in hair strands rendering when rendering hair proxy are rendered on non-initialized-yet instances.

diff --git a/Engine/Source/Runtime/Renderer/Private/HairStrands/HairStrandsUtils.cpp b/Engine/Source/Runtime/Renderer/Private/HairStrands/HairStrandsUtils.cpp
--- a/Engine/Source/Runtime/Renderer/Private/HairStrands/HairStrandsUtils.cpp
+++ b/Engine/Source/Runtime/Renderer/Private/HairStrands/HairStrandsUtils.cpp
@@ -270,6 +270,9 @@ FIntRect ComputeVisibleHairStrandsMacroGroupsRect(const FIntRect& ViewRect, cons
 		{
 			TotalRect.Union(Data.ScreenRect);
 		}
+
+		// In case bounds are not initialized correct for some reason, return view rect
+		if (TotalRect.Min.X >= TotalRect.Max.X || TotalRect.Min.Y >= TotalRect.Max.Y) TotalRect = ViewRect;
 	}
 	else
 	{
