Fix a bug of missing directional lighting if there is no local light in view on mobile.

diff --git a/Engine/Source/Runtime/Renderer/Private/MobileShadingRenderer.cpp b/Engine/Source/Runtime/Renderer/Private/MobileShadingRenderer.cpp
--- a/Engine/Source/Runtime/Renderer/Private/MobileShadingRenderer.cpp
+++ b/Engine/Source/Runtime/Renderer/Private/MobileShadingRenderer.cpp
@@ -817,7 +817,7 @@ void FMobileSceneRenderer::Render(FRDGBuilder& GraphBuilder)
 		bool bShadowedLightsInClustered = bRequiresShadowProjections && !bDeferredShading;
 		GatherAndSortLights(SortedLightSet, bShadowedLightsInClustered);
 		int32 NumReflectionCaptures = Views[0].NumBoxReflectionCaptures + Views[0].NumSphereReflectionCaptures;
-		bool bCullLightsToGrid = ((bEnableClusteredReflections && NumReflectionCaptures > 0) || (bEnableClusteredLocalLights && SortedLightSet.ClusteredSupportedEnd > SortedLightSet.SimpleLightsEnd));
+		bool bCullLightsToGrid = ((bEnableClusteredReflections && NumReflectionCaptures > 0) || bEnableClusteredLocalLights);
 		if (bCullLightsToGrid)
 		{
 			ComputeLightGrid(GraphBuilder, true, SortedLightSet);
