Minor clean up of reflection capture copy.

diff --git a/Engine/Source/Runtime/Renderer/Private/ReflectionEnvironment.cpp b/Engine/Source/Runtime/Renderer/Private/ReflectionEnvironment.cpp
--- a/Engine/Source/Runtime/Renderer/Private/ReflectionEnvironment.cpp
+++ b/Engine/Source/Runtime/Renderer/Private/ReflectionEnvironment.cpp
@@ -501,13 +501,9 @@ void FReflectionEnvironmentCubemapArray::ResizeCubemapArrayGPU(uint32 InMaxCubem
 					FRHICopyTextureInfo CopyInfo;
 					CopyInfo.SourceSliceIndex = SourceCubemapIndex * CubeFace_MAX + Face;
 					CopyInfo.DestSliceIndex   = DestCubemapIndex   * CubeFace_MAX + Face;
+					CopyInfo.NumMips          = NumMips;
 
-					for (int32 Mip = 0; Mip < NumMips; Mip++)
-					{
-						CopyInfo.SourceMipIndex = CopyInfo.DestMipIndex = Mip;
-
-						RHICmdList.CopyTexture(OldReflectionEnvs->GetRHI(), ReflectionEnvs->GetRHI(), CopyInfo);
-					}
+					RHICmdList.CopyTexture(OldReflectionEnvs->GetRHI(), ReflectionEnvs->GetRHI(), CopyInfo);
 				}
 			}
 		}
