Replaced AddCopyToResolveTarget pass with render target resolve.

diff --git a/Engine/Source/Runtime/RenderCore/Public/RenderGraphUtils.h b/Engine/Source/Runtime/RenderCore/Public/RenderGraphUtils.h
--- a/Engine/Source/Runtime/RenderCore/Public/RenderGraphUtils.h
+++ b/Engine/Source/Runtime/RenderCore/Public/RenderGraphUtils.h
@@ -698,6 +698,7 @@ RENDERCORE_API void AddDrawTexturePass(
 /** Adds a render graph pass to resolve from one texture to another. Uses RHICopyToResolveTarget under the hood.
  *  The formats of the two textures don't need to match.
  */
+UE_DEPRECATED(5.1, "AddCopyToResolveTargetPass is deprecated. Use AddCopyTexturePass to perform texture copies, or FRenderTargetBinding to perform resolves.")
 RENDERCORE_API void AddCopyToResolveTargetPass(
 	FRDGBuilder& GraphBuilder,
 	FRDGTextureRef InputTexture,
diff --git a/Engine/Source/Runtime/Renderer/Private/LightGridInjection.cpp b/Engine/Source/Runtime/Renderer/Private/LightGridInjection.cpp
--- a/Engine/Source/Runtime/Renderer/Private/LightGridInjection.cpp
+++ b/Engine/Source/Runtime/Renderer/Private/LightGridInjection.cpp
@@ -900,12 +900,16 @@ void FDeferredShadingSceneRenderer::RenderForwardShadowProjections(
 			}
 		}
 
-		AddCopyToResolveTargetPass(GraphBuilder, ForwardScreenSpaceShadowMask.Target, ForwardScreenSpaceShadowMask.Resolve, FResolveParams());
+		auto* PassParameters = GraphBuilder.AllocParameters<FRenderTargetParameters>();
+		PassParameters->RenderTargets[0] = FRenderTargetBinding(ForwardScreenSpaceShadowMask.Target, ForwardScreenSpaceShadowMask.Resolve, ERenderTargetLoadAction::ELoad);
 		OutForwardScreenSpaceShadowMask = ForwardScreenSpaceShadowMask.Resolve;
+
 		if (bIsHairEnable)
 		{
-			AddCopyToResolveTargetPass(GraphBuilder, ForwardScreenSpaceShadowMaskSubPixel.Target, ForwardScreenSpaceShadowMaskSubPixel.Resolve, FResolveParams());
+			PassParameters->RenderTargets[1] = FRenderTargetBinding(ForwardScreenSpaceShadowMaskSubPixel.Target, ForwardScreenSpaceShadowMaskSubPixel.Resolve, ERenderTargetLoadAction::ELoad);
 			OutForwardScreenSpaceShadowMaskSubPixel = ForwardScreenSpaceShadowMaskSubPixel.Resolve;
 		}
+
+		GraphBuilder.AddPass(RDG_EVENT_NAME("ResolveScreenSpaceShadowMask"), PassParameters, ERDGPassFlags::Raster, [](FRHICommandList&) {});
 	}
 }
