Correctly disable Eye Adaptation pass in a mobile post-processing

diff --git a/Engine/Source/Runtime/Renderer/Private/PostProcess/PostProcessing.cpp b/Engine/Source/Runtime/Renderer/Private/PostProcess/PostProcessing.cpp
--- a/Engine/Source/Runtime/Renderer/Private/PostProcess/PostProcessing.cpp
+++ b/Engine/Source/Runtime/Renderer/Private/PostProcess/PostProcessing.cpp
@@ -1859,10 +1859,7 @@ void AddMobilePostProcessingPasses(FRDGBuilder& GraphBuilder, FScene* Scene, con
 			
 		bool bUseBloom = View.FinalPostProcessSettings.BloomIntensity > 0.0f;
 
-		bool bUseBasicEyeAdaptation = bUseEyeAdaptation && (AutoExposureMethod == EAutoExposureMethod::AEM_Basic) &&
-			// Skip if we don't have any exposure range to generate (eye adaptation will clamp).
-			View.FinalPostProcessSettings.AutoExposureMinBrightness < View.FinalPostProcessSettings.AutoExposureMaxBrightness;
-
+		bool bUseBasicEyeAdaptation = bUseEyeAdaptation && (AutoExposureMethod == EAutoExposureMethod::AEM_Basic);
 		bool bUseHistogramEyeAdaptation = bUseEyeAdaptation && (AutoExposureMethod == EAutoExposureMethod::AEM_Histogram) &&
 			// Skip if we don't have any exposure range to generate (eye adaptation will clamp).
 			View.FinalPostProcessSettings.AutoExposureMinBrightness < View.FinalPostProcessSettings.AutoExposureMaxBrightness;
