Add per-pass traversal statistics for inline ray tracing for supporting platforms:
* Added RaytracingTraversalStatistics to create, capture and print TraceRayInline traversal statistics.
* In inline raytracing shaders you only need to call TraceRayInlineAccumulateStatistics() which will gather the results of the most recent TraceRayInline call.
* New debug visualization mode 'Traversal Statistics' that will print to the screen the traversal statistics for primary rays.

diff --git a/Engine/Shaders/Private/RayTracing/RayTracingDebugTraversal.usf b/Engine/Shaders/Private/RayTracing/RayTracingDebugTraversal.usf
--- a/Engine/Shaders/Private/RayTracing/RayTracingDebugTraversal.usf
+++ b/Engine/Shaders/Private/RayTracing/RayTracingDebugTraversal.usf
@@ -55,5 +55,9 @@ void RayTracingDebugTraversalCS(uint3 DispatchThread : SV_DispatchThreadID)
 		break;
 	}
 
+#if PRINT_TRAVERSAL_STATISTICS
+	TraceRayInlineAccumulateStatistics();
+#endif
+
 	Output[PixelCoord] = Result;
 }
diff --git a/Engine/Shaders/Private/RayTracing/TraceRayInlineCommon.ush b/Engine/Shaders/Private/RayTracing/TraceRayInlineCommon.ush
--- a/Engine/Shaders/Private/RayTracing/TraceRayInlineCommon.ush
+++ b/Engine/Shaders/Private/RayTracing/TraceRayInlineCommon.ush
@@ -8,8 +8,39 @@ struct FTraceRayInlineTraversalStatistics
 	uint ClusterIntersectionCount;
 	uint TriangleIntersectionCount;
 	uint IterationCount;
+
+	uint WaveNodeIntersectionCount;
+	uint WaveTriangleIntersectionCount;
+
+	uint RayCount;
+	uint WavefrontCount;
+
+	uint ActiveLaneCount;
+	uint TotalLaneCount;
 };
 
+static FTraceRayInlineTraversalStatistics TraceRayInlineTraversalStatistics_Sample;
+RWStructuredBuffer<FTraceRayInlineTraversalStatistics> TraceRayInlineTraversalStatistics_Accumulator;
+
+void TraceRayInlineAccumulateStatistics()
+{
+#if ENABLE_TRACE_RAY_INLINE_TRAVERSAL_STATISTICS
+	if (WaveIsFirstLane())
+	{
+		InterlockedAdd(TraceRayInlineTraversalStatistics_Accumulator[0].RayCount, TraceRayInlineTraversalStatistics_Sample.RayCount);
+		InterlockedAdd(TraceRayInlineTraversalStatistics_Accumulator[0].WavefrontCount, TraceRayInlineTraversalStatistics_Sample.WavefrontCount);
+
+		InterlockedAdd(TraceRayInlineTraversalStatistics_Accumulator[0].ActiveLaneCount, TraceRayInlineTraversalStatistics_Sample.ActiveLaneCount);
+		InterlockedAdd(TraceRayInlineTraversalStatistics_Accumulator[0].TotalLaneCount, TraceRayInlineTraversalStatistics_Sample.TotalLaneCount);
+
+		InterlockedAdd(TraceRayInlineTraversalStatistics_Accumulator[0].IterationCount, TraceRayInlineTraversalStatistics_Sample.IterationCount);
+
+		InterlockedAdd(TraceRayInlineTraversalStatistics_Accumulator[0].WaveNodeIntersectionCount, TraceRayInlineTraversalStatistics_Sample.WaveNodeIntersectionCount);
+		InterlockedAdd(TraceRayInlineTraversalStatistics_Accumulator[0].WaveTriangleIntersectionCount, TraceRayInlineTraversalStatistics_Sample.WaveTriangleIntersectionCount);
+	}
+#endif
+}
+
 struct FTraceRayInlineResult
 {
 	float HitT;
diff --git a/Engine/Shaders/Private/RayTracing/TraceRayInlineStatistics.usf b/Engine/Shaders/Private/RayTracing/TraceRayInlineStatistics.usf
new file mode 100644
--- /dev/null
+++ b/Engine/Shaders/Private/RayTracing/TraceRayInlineStatistics.usf
@@ -0,0 +1,71 @@
+// Copyright Epic Games, Inc. All Rights Reserved.
+
+#include "../Common.ush"
+#include "../ShaderPrint.ush"
+#include "TraceRayInlineCommon.ush"
+
+StructuredBuffer<FTraceRayInlineTraversalStatistics> TraversalStatistics;
+
+[numthreads(1, 1, 1)]
+void TraceRayInlinePrintStatisticsCS()
+{
+	const float TopMargin = 0.05f;
+	const float HeadlineX = 0.75f;
+	const float ItemX = 0.76f;
+
+	float2 ScreenPos = float2(HeadlineX, TopMargin);
+
+	const uint PassIndex = 0; 
+	const float RayCount = TraversalStatistics[PassIndex].RayCount;
+	const float WavefrontCount = TraversalStatistics[PassIndex].WavefrontCount;
+
+	ScreenPos.x = HeadlineX;
+	ScreenPos = ShaderPrintText(ScreenPos, TEXT("Pass Traversal Statistics"), FontOrange);
+	ScreenPos = ShaderPrintNewline(ScreenPos);
+
+	ScreenPos.x = ItemX;
+	ScreenPos = ShaderPrintText(ScreenPos, TEXT("Ray Count            "), FontYellow);
+	ScreenPos = ShaderPrint(ScreenPos, TraversalStatistics[PassIndex].RayCount);
+	ScreenPos = ShaderPrintNewline(ScreenPos);
+
+	ScreenPos.x = ItemX;
+	ScreenPos = ShaderPrintText(ScreenPos, TEXT("Wavefront Count      "), FontYellow);
+	ScreenPos = ShaderPrint(ScreenPos, TraversalStatistics[PassIndex].WavefrontCount);
+	ScreenPos = ShaderPrintNewline(ScreenPos);
+
+	// Node Intersection
+	ScreenPos.x = ItemX;
+	ScreenPos = ShaderPrintText(ScreenPos, TEXT("Tri/Ray              "), FontYellow);
+	float TriangleRay = TraversalStatistics[PassIndex].WaveTriangleIntersectionCount / RayCount;
+	ScreenPos = ShaderPrint(ScreenPos, TriangleRay);
+	ScreenPos = ShaderPrintNewline(ScreenPos);
+
+	ScreenPos.x = ItemX;
+	ScreenPos = ShaderPrintText(ScreenPos, TEXT("Node/Ray             "), FontYellow);
+	float NodeRay = TraversalStatistics[PassIndex].WaveNodeIntersectionCount / RayCount;
+	ScreenPos = ShaderPrint(ScreenPos, NodeRay);
+	ScreenPos = ShaderPrintNewline(ScreenPos);
+
+	ScreenPos.x = ItemX;
+	ScreenPos = ShaderPrintText(ScreenPos, TEXT("Iteration/Wave       "), FontYellow);
+	float IterationWave = TraversalStatistics[PassIndex].IterationCount / WavefrontCount;
+	ScreenPos = ShaderPrint(ScreenPos, IterationWave);
+	ScreenPos = ShaderPrintNewline(ScreenPos);
+
+	ScreenPos.x = ItemX;
+	ScreenPos = ShaderPrintText(ScreenPos, TEXT("Iteration/Ray        "), FontYellow);
+	float IterationRay = TriangleRay + NodeRay;
+	ScreenPos = ShaderPrint(ScreenPos, IterationRay);
+	ScreenPos = ShaderPrintNewline(ScreenPos);
+
+	ScreenPos.x = ItemX;
+	ScreenPos = ShaderPrintText(ScreenPos, TEXT("Iteration Divergence "), FontYellow);
+	ScreenPos = ShaderPrint(ScreenPos, IterationWave / IterationRay);
+	ScreenPos = ShaderPrintNewline(ScreenPos);
+
+	ScreenPos.x = ItemX;
+	ScreenPos = ShaderPrintText(ScreenPos, TEXT("Occupancy (%)        "), FontYellow);
+	float Occupancy = 100.0 * TraversalStatistics[PassIndex].ActiveLaneCount / TraversalStatistics[PassIndex].TotalLaneCount;
+	ScreenPos = ShaderPrint(ScreenPos, Occupancy);
+	ScreenPos = ShaderPrintNewline(ScreenPos);
+}
diff --git a/Engine/Shaders/Shared/RaytracingDebugDefinitions.h b/Engine/Shaders/Shared/RaytracingDebugDefinitions.h
--- a/Engine/Shaders/Shared/RaytracingDebugDefinitions.h
+++ b/Engine/Shaders/Shared/RaytracingDebugDefinitions.h
@@ -37,3 +37,4 @@
 #define RAY_TRACING_DEBUG_VIZ_TRAVERSAL_CLUSTER			27
 #define RAY_TRACING_DEBUG_VIZ_TRAVERSAL_TRIANGLE		28
 #define RAY_TRACING_DEBUG_VIZ_TRAVERSAL_ALL				29
+#define RAY_TRACING_DEBUG_VIZ_TRAVERSAL_STATISTICS		30
diff --git a/Engine/Source/Runtime/RHI/Private/RHI.cpp b/Engine/Source/Runtime/RHI/Private/RHI.cpp
--- a/Engine/Source/Runtime/RHI/Private/RHI.cpp
+++ b/Engine/Source/Runtime/RHI/Private/RHI.cpp
@@ -2440,6 +2440,7 @@ void FGenericDataDrivenShaderPlatformInfo::ParseDataDrivenShaderInfo(const FConf
 	GET_SECTION_BOOL_HELPER(bSupportsRayTracingShaders);
 	GET_SECTION_BOOL_HELPER(bSupportsInlineRayTracing);
 	GET_SECTION_BOOL_HELPER(bSupportsRayTracingProceduralPrimitive);
+	GET_SECTION_BOOL_HELPER(bSupportsRayTracingTraversalStatistics);
 	GET_SECTION_BOOL_HELPER(bSupportsRayTracingIndirectInstanceData);
 	GET_SECTION_BOOL_HELPER(bSupportsPathTracing);
 	GET_SECTION_BOOL_HELPER(bSupportsHighEndRayTracingReflections);
diff --git a/Engine/Source/Runtime/RHI/Public/RHIDefinitions.h b/Engine/Source/Runtime/RHI/Public/RHIDefinitions.h
--- a/Engine/Source/Runtime/RHI/Public/RHIDefinitions.h
+++ b/Engine/Source/Runtime/RHI/Public/RHIDefinitions.h
@@ -357,6 +357,7 @@ class RHI_API FGenericDataDrivenShaderPlatformInfo
 	uint32 bSupportsRenderTargetWriteMask: 1;
 	uint32 bSupportsRayTracing: 1;
 	uint32 bSupportsRayTracingProceduralPrimitive : 1;
+	uint32 bSupportsRayTracingTraversalStatistics : 1;
 	uint32 bSupportsRayTracingIndirectInstanceData : 1; // Whether instance transforms can be copied from the GPU to the TLAS instances buffer
 	uint32 bSupportsHighEndRayTracingReflections : 1; // Whether fully-featured RT reflections can be used on the platform (with multi-bounce, translucency, etc.)
 	uint32 bSupportsPathTracing : 1; // Whether real-time path tracer is supported on this platform (avoids compiling unnecessary shaders)
@@ -621,6 +622,12 @@ public:
 		return Infos[Platform].bSupportsRayTracing && Infos[Platform].bSupportsRayTracingProceduralPrimitive;
 	}
 
+	static FORCEINLINE_DEBUGGABLE const bool GetSupportsRayTracingTraversalStatistics(const FStaticShaderPlatform Platform)
+	{
+		check(IsValid(Platform));
+		return Infos[Platform].bSupportsRayTracing && Infos[Platform].bSupportsRayTracingTraversalStatistics;
+	}
+
 	static FORCEINLINE_DEBUGGABLE const bool GetSupportsRayTracingIndirectInstanceData(const FStaticShaderPlatform Platform)
 	{
 		check(IsValid(Platform));
diff --git a/Engine/Source/Runtime/Renderer/Private/RayTracing/RayTracingDebug.cpp b/Engine/Source/Runtime/Renderer/Private/RayTracing/RayTracingDebug.cpp
--- a/Engine/Source/Runtime/Renderer/Private/RayTracing/RayTracingDebug.cpp
+++ b/Engine/Source/Runtime/Renderer/Private/RayTracing/RayTracingDebug.cpp
@@ -14,6 +14,7 @@
 #include "RaytracingDebugDefinitions.h"
 #include "RayTracing/RayTracingLighting.h"
 #include "RayTracing/RaytracingOptions.h"
+#include "RayTracing/RayTracingTraversalStatistics.h"
 
 #define LOCTEXT_NAMESPACE "RayTracingDebugVisualizationMenuCommands"
 
@@ -123,6 +124,7 @@ class FRayTracingDebugTraversalCS : public FGlobalShader
 		SHADER_PARAMETER_SRV(RaytracingAccelerationStructure, TLAS)
 		SHADER_PARAMETER_STRUCT_REF(FViewUniformShaderParameters, ViewUniformBuffer)
 		SHADER_PARAMETER_STRUCT_REF(FNaniteUniformParameters, NaniteUniformBuffer)
+		SHADER_PARAMETER_STRUCT_INCLUDE(RaytracingTraversalStatistics::FShaderParameters, TraversalStatistics)
 
 		SHADER_PARAMETER(uint32, VisualizationMode)
 		SHADER_PARAMETER(float, TraversalBoxScale)
@@ -133,7 +135,8 @@ class FRayTracingDebugTraversalCS : public FGlobalShader
 	END_SHADER_PARAMETER_STRUCT()
 
 	class FSupportProceduralPrimitive : SHADER_PERMUTATION_BOOL("ENABLE_TRACE_RAY_INLINE_PROCEDURAL_PRIMITIVE");
-	using FPermutationDomain = TShaderPermutationDomain<FSupportProceduralPrimitive>;
+	class FPrintTraversalStatistics : SHADER_PERMUTATION_BOOL("PRINT_TRAVERSAL_STATISTICS");
+	using FPermutationDomain = TShaderPermutationDomain<FSupportProceduralPrimitive, FPrintTraversalStatistics>;
 
 	static void ModifyCompilationEnvironment(const FGlobalShaderPermutationParameters& Parameters, FShaderCompilerEnvironment& OutEnvironment)
 	{
@@ -150,6 +153,14 @@ class FRayTracingDebugTraversalCS : public FGlobalShader
 
 	static bool ShouldCompilePermutation(const FGlobalShaderPermutationParameters& Parameters)
 	{
+		FPermutationDomain PermutationVector(Parameters.PermutationId);
+		bool bTraversalStats = PermutationVector.Get<FPrintTraversalStatistics>();
+		bool bSupportsTraversalStats = FDataDrivenShaderPlatformInfo::GetSupportsRayTracingTraversalStatistics(Parameters.Platform);
+		if (bTraversalStats && !bSupportsTraversalStats)
+		{
+			return false;
+		}
+
 		return IsRayTracingEnabledForProject(Parameters.Platform) && RHISupportsRayTracing(Parameters.Platform) && RHISupportsInlineRayTracing(Parameters.Platform);
 	}
 
@@ -169,7 +180,8 @@ static bool IsRayTracingDebugTraversalMode(uint32 DebugVisualizationMode)
 	return DebugVisualizationMode == RAY_TRACING_DEBUG_VIZ_TRAVERSAL_NODE ||
 		DebugVisualizationMode == RAY_TRACING_DEBUG_VIZ_TRAVERSAL_CLUSTER ||
 		DebugVisualizationMode == RAY_TRACING_DEBUG_VIZ_TRAVERSAL_TRIANGLE ||
-		DebugVisualizationMode == RAY_TRACING_DEBUG_VIZ_TRAVERSAL_ALL;
+		DebugVisualizationMode == RAY_TRACING_DEBUG_VIZ_TRAVERSAL_ALL ||
+		DebugVisualizationMode == RAY_TRACING_DEBUG_VIZ_TRAVERSAL_STATISTICS;
 }
 
 void FDeferredShadingSceneRenderer::PrepareRayTracingDebug(const FSceneViewFamily& ViewFamily, TArray<FRHIRayTracingShader*>& OutRayGenShaders)
@@ -218,6 +230,7 @@ void FDeferredShadingSceneRenderer::RenderRayTracingDebug(FRDGBuilder& GraphBuil
 		RayTracingDebugVisualizationModes.Emplace(FName(*LOCTEXT("Traversal Cluster", "Traversal Cluster").ToString()),							RAY_TRACING_DEBUG_VIZ_TRAVERSAL_CLUSTER);
 		RayTracingDebugVisualizationModes.Emplace(FName(*LOCTEXT("Traversal Triangle", "Traversal Triangle").ToString()),						RAY_TRACING_DEBUG_VIZ_TRAVERSAL_TRIANGLE);
 		RayTracingDebugVisualizationModes.Emplace(FName(*LOCTEXT("Traversal All", "Traversal All").ToString()),									RAY_TRACING_DEBUG_VIZ_TRAVERSAL_ALL);
+		RayTracingDebugVisualizationModes.Emplace(FName(*LOCTEXT("Traversal Statistics", "Traversal Statistics").ToString()),					RAY_TRACING_DEBUG_VIZ_TRAVERSAL_STATISTICS);
 	}
 
 	uint32 DebugVisualizationMode;
@@ -245,6 +258,9 @@ void FDeferredShadingSceneRenderer::RenderRayTracingDebug(FRDGBuilder& GraphBuil
 
 	if (IsRayTracingDebugTraversalMode(DebugVisualizationMode) && ShouldRenderRayTracingEffect(ERayTracingPipelineCompatibilityFlags::Inline))
 	{
+		const bool bPrintTraversalStats = FDataDrivenShaderPlatformInfo::GetSupportsRayTracingTraversalStatistics(GMaxRHIShaderPlatform)
+			&& (DebugVisualizationMode == RAY_TRACING_DEBUG_VIZ_TRAVERSAL_STATISTICS);
+
 		FRayTracingDebugTraversalCS::FParameters* PassParameters = GraphBuilder.AllocParameters<FRayTracingDebugTraversalCS::FParameters>();
 		PassParameters->Output = GraphBuilder.CreateUAV(SceneColorTexture);
 		PassParameters->TLAS = Scene->RayTracingScene.GetLayerSRVChecked(ERayTracingSceneLayer::Base);
@@ -258,6 +274,13 @@ void FDeferredShadingSceneRenderer::RenderRayTracingDebug(FRDGBuilder& GraphBuil
 
 		PassParameters->RTDebugVisualizationNaniteCutError = 0.0f;
 
+		RaytracingTraversalStatistics::FTraceRayInlineStatisticsData TraversalData;
+		if (bPrintTraversalStats)
+		{
+			RaytracingTraversalStatistics::Init(GraphBuilder, TraversalData);
+			RaytracingTraversalStatistics::SetParameters(GraphBuilder, TraversalData, PassParameters->TraversalStatistics);
+		}
+
 		FIntRect ViewRect = View.ViewRect;
 
 		RDG_GPU_STAT_SCOPE(GraphBuilder, RayTracingDebug);
@@ -267,10 +290,17 @@ void FDeferredShadingSceneRenderer::RenderRayTracingDebug(FRDGBuilder& GraphBuil
 
 		FRayTracingDebugTraversalCS::FPermutationDomain PermutationVector;
 		PermutationVector.Set<FRayTracingDebugTraversalCS::FSupportProceduralPrimitive>((bool)GVisualizeProceduralPrimitives);
+		PermutationVector.Set<FRayTracingDebugTraversalCS::FPrintTraversalStatistics>(bPrintTraversalStats);
 		
 		TShaderRef<FRayTracingDebugTraversalCS> ComputeShader = View.ShaderMap->GetShader<FRayTracingDebugTraversalCS>(PermutationVector);
 
 		FComputeShaderUtils::AddPass(GraphBuilder, RDG_EVENT_NAME("RayTracingDebug"), ComputeShader, PassParameters, GroupCount);
+
+		if (bPrintTraversalStats)
+		{
+			RaytracingTraversalStatistics::AddPrintPass(GraphBuilder, View, TraversalData);
+		}
+
 		return;
 	}
 
diff --git a/Engine/Source/Runtime/Renderer/Private/RayTracing/RayTracingTraversalStatistics.cpp b/Engine/Source/Runtime/Renderer/Private/RayTracing/RayTracingTraversalStatistics.cpp
new file mode 100644
--- /dev/null
+++ b/Engine/Source/Runtime/Renderer/Private/RayTracing/RayTracingTraversalStatistics.cpp
@@ -0,0 +1,90 @@
+// Copyright Epic Games, Inc. All Rights Reserved.
+
+#include "RayTracingTraversalStatistics.h"
+
+#if RHI_RAYTRACING
+
+#include "RenderGraphBuilder.h"
+#include "RenderGraphUtils.h"
+#include "ShaderParameterMacros.h"
+#include "ShaderParameters.h"
+#include "ShaderParameterStruct.h"
+#include "ShaderPrint.h"
+#include "ShaderPrintParameters.h"
+#include "SceneRendering.h"
+#include "GlobalShader.h"
+
+// Needs to match the size of FTraceRayInlineTraversalStatistics in TraceRayInlineCommon.ush
+#define TRACE_RAY_INLINE_TRAVERSAL_STATISTICS_SIZE_IN_BYTES 40
+
+class FTraceRayInlinePrintStatisticsCS : public FGlobalShader
+{
+	DECLARE_GLOBAL_SHADER(FTraceRayInlinePrintStatisticsCS);
+	SHADER_USE_PARAMETER_STRUCT(FTraceRayInlinePrintStatisticsCS, FGlobalShader);
+
+	BEGIN_SHADER_PARAMETER_STRUCT(FParameters, )
+		SHADER_PARAMETER_STRUCT_INCLUDE(ShaderPrint::FShaderParameters, ShaderPrintStruct)
+		SHADER_PARAMETER_RDG_BUFFER_SRV(StructuredBuffer, TraversalStatistics)
+	END_SHADER_PARAMETER_STRUCT()
+
+	static bool ShouldCompilePermutation(const FGlobalShaderPermutationParameters& Parameters)
+	{
+		return IsRayTracingEnabledForProject(Parameters.Platform) && RHISupportsRayTracing(Parameters.Platform) && RHISupportsInlineRayTracing(Parameters.Platform) && FDataDrivenShaderPlatformInfo::GetSupportsRayTracingTraversalStatistics(Parameters.Platform);
+	}
+
+	static void ModifyCompilationEnvironment(const FGlobalShaderPermutationParameters& Parameters, FShaderCompilerEnvironment& OutEnvironment)
+	{
+		FGlobalShader::ModifyCompilationEnvironment(Parameters, OutEnvironment);		
+	}
+};
+
+IMPLEMENT_GLOBAL_SHADER(FTraceRayInlinePrintStatisticsCS, "/Engine/Private/RayTracing/TraceRayInlineStatistics.usf", "TraceRayInlinePrintStatisticsCS", SF_Compute);
+
+namespace RaytracingTraversalStatistics
+{
+	bool IsEnabled()
+	{
+		return GRHISupportsInlineRayTracing && FDataDrivenShaderPlatformInfo::GetSupportsRayTracingTraversalStatistics(GMaxRHIShaderPlatform);
+	}
+
+	void Init(FRDGBuilder& GraphBuilder, FTraceRayInlineStatisticsData& OutTraversalData)
+	{
+		if (!IsEnabled())
+		{
+			return;
+		}
+
+		OutTraversalData.TraversalStatisticsBuffer = GraphBuilder.CreateBuffer(
+			FRDGBufferDesc::CreateStructuredDesc(TRACE_RAY_INLINE_TRAVERSAL_STATISTICS_SIZE_IN_BYTES, 1),
+			TEXT("TraversalStatisticsBuffer"));
+
+		AddClearUAVPass(GraphBuilder, GraphBuilder.CreateUAV(OutTraversalData.TraversalStatisticsBuffer), 0);
+	}
+
+	void SetParameters(FRDGBuilder& GraphBuilder, const FTraceRayInlineStatisticsData& TraversalData, FShaderParameters& OutParameters)
+	{		
+		if (IsEnabled() && TraversalData.TraversalStatisticsBuffer)
+		{
+			OutParameters.TraceRayInlineTraversalStatistics_Accumulator = GraphBuilder.CreateUAV(TraversalData.TraversalStatisticsBuffer);
+		}		
+	}
+
+	void AddPrintPass(FRDGBuilder& GraphBuilder, const FViewInfo& View, const FTraceRayInlineStatisticsData& TraversalData)
+	{
+		if (IsEnabled() && TraversalData.TraversalStatisticsBuffer)
+		{
+			FTraceRayInlinePrintStatisticsCS::FParameters* Parameters = GraphBuilder.AllocParameters<FTraceRayInlinePrintStatisticsCS::FParameters>();
+			ShaderPrint::SetParameters(GraphBuilder, View, Parameters->ShaderPrintStruct);
+			Parameters->TraversalStatistics = GraphBuilder.CreateSRV(TraversalData.TraversalStatisticsBuffer);
+			TShaderRef<FTraceRayInlinePrintStatisticsCS> ComputeShader = View.ShaderMap->GetShader<FTraceRayInlinePrintStatisticsCS>();
+			FComputeShaderUtils::AddPass(
+				GraphBuilder,
+				RDG_EVENT_NAME("TraceRayInlinePrintStatisticsCS"),
+				ComputeShader,
+				Parameters,
+				FIntVector(1, 1, 1));
+		}
+	}
+}
+
+#endif // RHI_RAYTRACING
diff --git a/Engine/Source/Runtime/Renderer/Private/RayTracing/RayTracingTraversalStatistics.h b/Engine/Source/Runtime/Renderer/Private/RayTracing/RayTracingTraversalStatistics.h
new file mode 100644
--- /dev/null
+++ b/Engine/Source/Runtime/Renderer/Private/RayTracing/RayTracingTraversalStatistics.h
@@ -0,0 +1,31 @@
+// Copyright Epic Games, Inc. All Rights Reserved.
+
+#pragma once
+
+#include "CoreMinimal.h"
+
+#if RHI_RAYTRACING
+
+#include "RenderGraphDefinitions.h"
+#include "ShaderParameterMacros.h"
+#include "ShaderParameters.h"
+
+class FViewInfo;
+
+namespace RaytracingTraversalStatistics
+{
+	BEGIN_SHADER_PARAMETER_STRUCT(FShaderParameters, )
+		SHADER_PARAMETER_RDG_BUFFER_UAV(RWStructuredBuffer, TraceRayInlineTraversalStatistics_Accumulator)
+	END_SHADER_PARAMETER_STRUCT()
+
+	struct FTraceRayInlineStatisticsData
+	{
+		FRDGBufferRef TraversalStatisticsBuffer;
+	};
+	
+	void Init(FRDGBuilder& GraphBuilder, FTraceRayInlineStatisticsData& OutTraversalData);
+	void SetParameters(FRDGBuilder& GraphBuilder, const FTraceRayInlineStatisticsData& TraversalData, FShaderParameters& OutParameters);
+	void AddPrintPass(FRDGBuilder& GraphBuilder, const FViewInfo& View, const FTraceRayInlineStatisticsData& TraversalData);
+}
+
+#endif // RHI_RAYTRACING
