Fix for underwater height fog being applied twice if the camera is above the water. Also added a cvar to override this behavior.

#rnx

diff --git a/Engine/Source/Runtime/Renderer/Private/FogRendering.cpp b/Engine/Source/Runtime/Renderer/Private/FogRendering.cpp
--- a/Engine/Source/Runtime/Renderer/Private/FogRendering.cpp
+++ b/Engine/Source/Runtime/Renderer/Private/FogRendering.cpp
@@ -31,6 +31,12 @@ static TAutoConsoleVariable<float> CVarUpsampleJitterMultiplier(
 	TEXT("Multiplier for random offset value used to jitter the sample position of the 3D fog volume to hide fog pixelization due to sampling from a lower resolution texture."),
 	ECVF_RenderThreadSafe | ECVF_Scalability);
 
+static TAutoConsoleVariable<bool> CVarUnderwaterFogWhenCameraIsAboveWater(
+	TEXT("r.Water.SingleLayer.UnderwaterFogWhenCameraIsAboveWater"), 
+	false, 
+	TEXT("Renders height fog behind the water surface even when the camera is above water. This avoids artifacts when entering and exiting the water with strong height fog in the scene but causes artifacts when looking at the water surface from a distance."),
+	ECVF_RenderThreadSafe);
+
 IMPLEMENT_GLOBAL_SHADER_PARAMETER_STRUCT(FFogUniformParameters, "FogStruct");
 
 void SetupFogUniformParameters(FRDGBuilder& GraphBuilder, const FViewInfo& View, FFogUniformParameters& OutParameters)
@@ -414,7 +420,7 @@ void FDeferredShadingSceneRenderer::RenderUnderWaterFog(
 		for (int32 ViewIndex = 0; ViewIndex < Views.Num(); ViewIndex++)
 		{
 			const FViewInfo& View = Views[ViewIndex];
-			if (View.IsPerspectiveProjection())
+			if (View.IsPerspectiveProjection() && (View.IsUnderwater() || CVarUnderwaterFogWhenCameraIsAboveWater.GetValueOnRenderThread()))
 			{
 				RDG_EVENT_SCOPE_CONDITIONAL(GraphBuilder, Views.Num() > 1, "View%d", ViewIndex);
 				RDG_GPU_MASK_SCOPE(GraphBuilder, View.GPUMask);
