Fix incorrect output with ShaderPrint when using screen percentage.

diff --git a/Engine/Source/Runtime/Renderer/Private/ShaderPrint.cpp b/Engine/Source/Runtime/Renderer/Private/ShaderPrint.cpp
--- a/Engine/Source/Runtime/Renderer/Private/ShaderPrint.cpp
+++ b/Engine/Source/Runtime/Renderer/Private/ShaderPrint.cpp
@@ -958,16 +958,19 @@ namespace ShaderPrint
 
 		RDG_EVENT_SCOPE(GraphBuilder, "ShaderPrint::DrawView");
 
+		const FIntRect SourceViewRect = View.ViewRect;
+		const FIntRect OutputViewRect = OutputTexture.ViewRect;
+		
 		// Lines
 		{
 			FRDGBufferRef DataBuffer = View.ShaderPrintData.ShaderPrintPrimitiveBuffer;
-			InternalDrawView_Primitives(GraphBuilder, View.ShaderPrintData, DataBuffer, View.ViewRect, View.UnscaledViewRect, View.ViewMatrices.GetTranslatedViewProjectionMatrix(), FVector::ZeroVector, true /*bLines*/, false /*bLocked*/, OutputTexture.Texture, DepthTexture.Texture);
+			InternalDrawView_Primitives(GraphBuilder, View.ShaderPrintData, DataBuffer, SourceViewRect, OutputViewRect, View.ViewMatrices.GetTranslatedViewProjectionMatrix(), FVector::ZeroVector, true /*bLines*/, false /*bLocked*/, OutputTexture.Texture, DepthTexture.Texture);
 		}
 
 		// Triangles
 		{
 			FRDGBufferRef DataBuffer = View.ShaderPrintData.ShaderPrintPrimitiveBuffer;
-			InternalDrawView_Primitives(GraphBuilder, View.ShaderPrintData, DataBuffer, View.ViewRect, View.UnscaledViewRect, View.ViewMatrices.GetTranslatedViewProjectionMatrix(), FVector::ZeroVector, false /*bLines*/, false /*bLocked*/, OutputTexture.Texture, DepthTexture.Texture);
+			InternalDrawView_Primitives(GraphBuilder, View.ShaderPrintData, DataBuffer, SourceViewRect, OutputViewRect, View.ViewMatrices.GetTranslatedViewProjectionMatrix(), FVector::ZeroVector, false /*bLines*/, false /*bLocked*/, OutputTexture.Texture, DepthTexture.Texture);
 		}
 
 		// Locked Lines/Triangles
@@ -975,14 +978,14 @@ namespace ShaderPrint
 		{
 			const FVector LockedToCurrentTranslatedOffset = View.ViewMatrices.GetPreViewTranslation() - View.ViewState->ShaderPrintStateData.PreViewTranslation;
 			FRDGBufferRef DataBuffer = GraphBuilder.RegisterExternalBuffer(View.ViewState->ShaderPrintStateData.PrimitiveBuffer);
-			InternalDrawView_Primitives(GraphBuilder, View.ShaderPrintData, DataBuffer, View.ViewRect, View.UnscaledViewRect, View.ViewMatrices.GetTranslatedViewProjectionMatrix(), LockedToCurrentTranslatedOffset, true  /*bLines*/, true/*bLocked*/, OutputTexture.Texture, DepthTexture.Texture);
-			InternalDrawView_Primitives(GraphBuilder, View.ShaderPrintData, DataBuffer, View.ViewRect, View.UnscaledViewRect, View.ViewMatrices.GetTranslatedViewProjectionMatrix(), LockedToCurrentTranslatedOffset, false /*bLines*/, true/*bLocked*/, OutputTexture.Texture, DepthTexture.Texture);
+			InternalDrawView_Primitives(GraphBuilder, View.ShaderPrintData, DataBuffer, SourceViewRect, OutputViewRect, View.ViewMatrices.GetTranslatedViewProjectionMatrix(), LockedToCurrentTranslatedOffset, true  /*bLines*/, true/*bLocked*/, OutputTexture.Texture, DepthTexture.Texture);
+			InternalDrawView_Primitives(GraphBuilder, View.ShaderPrintData, DataBuffer, SourceViewRect, OutputViewRect, View.ViewMatrices.GetTranslatedViewProjectionMatrix(), LockedToCurrentTranslatedOffset, false /*bLines*/, true/*bLocked*/, OutputTexture.Texture, DepthTexture.Texture);
 		}
 
 		// Characters
 		{
 			const int32 FrameNumber = View.Family ? View.Family->FrameNumber : 0u;
-			InternalDrawView_Characters(GraphBuilder, View.ShaderPrintData, View.ViewRect, FrameNumber, OutputTexture);
+			InternalDrawView_Characters(GraphBuilder, View.ShaderPrintData, OutputViewRect, FrameNumber, OutputTexture);
 		}
 	}
 
