Fix for RHI validation error in GPUScene.

diff --git a/Engine/Source/Runtime/Renderer/Private/GPUScene.cpp b/Engine/Source/Runtime/Renderer/Private/GPUScene.cpp
--- a/Engine/Source/Runtime/Renderer/Private/GPUScene.cpp
+++ b/Engine/Source/Runtime/Renderer/Private/GPUScene.cpp
@@ -852,17 +852,14 @@ void FGPUScene::UseExternalAccessMode(FRDGExternalAccessQueue& ExternalAccessQue
 		return;
 	}
 
-	if (!ExternalAccessQueue.Contains(BufferState.InstanceSceneDataBuffer))
-	{
-		ExternalAccessQueue.Add(BufferState.InstanceSceneDataBuffer, Access, Pipelines);
-		ExternalAccessQueue.Add(BufferState.InstancePayloadDataBuffer, Access, Pipelines);
-		ExternalAccessQueue.Add(BufferState.PrimitiveBuffer, Access, Pipelines);
-		ExternalAccessQueue.Add(BufferState.LightmapDataBuffer, Access, Pipelines);
+	ExternalAccessQueue.AddUnique(BufferState.InstanceSceneDataBuffer, Access, Pipelines);
+	ExternalAccessQueue.AddUnique(BufferState.InstancePayloadDataBuffer, Access, Pipelines);
+	ExternalAccessQueue.AddUnique(BufferState.PrimitiveBuffer, Access, Pipelines);
+	ExternalAccessQueue.AddUnique(BufferState.LightmapDataBuffer, Access, Pipelines);
 
-		if (BufferState.InstanceBVHBuffer)
-		{
-			ExternalAccessQueue.Add(BufferState.InstanceBVHBuffer, Access, Pipelines);
-		}
+	if (BufferState.InstanceBVHBuffer)
+	{
+		ExternalAccessQueue.AddUnique(BufferState.InstanceBVHBuffer, Access, Pipelines);
 	}
 
 	bInExternalAccessMode = true;
