Support Callable/Miss shaders in MeshPassProcessor.
- useful to gather shader bindings to evaluate materials in callable shaders.
- similar to how it supports compute and hit group shaders.

diff --git a/Engine/Source/Runtime/Renderer/Private/MeshPassProcessor.cpp b/Engine/Source/Runtime/Renderer/Private/MeshPassProcessor.cpp
--- a/Engine/Source/Runtime/Renderer/Private/MeshPassProcessor.cpp
+++ b/Engine/Source/Runtime/Renderer/Private/MeshPassProcessor.cpp
@@ -281,7 +281,7 @@ void FMeshDrawShaderBindings::SetRayTracingShaderBindingsForHitGroup(
 	FRayTracingLocalShaderBindingWriter* BindingWriter,
 	uint32 InstanceIndex, 
 	uint32 SegmentIndex,
-	uint32 HitGroupIndex,
+	uint32 HitGroupIndexInPipeline,
 	uint32 ShaderSlot) const
 {
 	check(ShaderLayouts.Num() == 1);
@@ -333,7 +333,7 @@ void FMeshDrawShaderBindings::SetRayTracingShaderBindingsForHitGroup(
 	Bindings.InstanceIndex = InstanceIndex;
 	Bindings.SegmentIndex = SegmentIndex;
 	Bindings.ShaderSlot = ShaderSlot;
-	Bindings.ShaderIndexInPipeline = HitGroupIndex;
+	Bindings.ShaderIndexInPipeline = HitGroupIndexInPipeline;
 	Bindings.UserData = UserData;
 
 	for (int32 UniformBufferIndex = 0; UniformBufferIndex < NumUniformBufferParameters; UniformBufferIndex++)
@@ -356,6 +356,12 @@ void FMeshDrawShaderBindings::SetRayTracingShaderBindingsForHitGroup(
 	}
 }
 
+void FMeshDrawShaderBindings::SetRayTracingShaderBindings(FRayTracingLocalShaderBindingWriter* BindingWriter, uint32 ShaderIndexInPipeline, uint32 ShaderSlot) const
+{
+	check(ShaderLayouts.Num() == 1);
+	return SetRayTracingShaderBindingsForHitGroup(BindingWriter, 0, 0, ShaderIndexInPipeline, ShaderSlot);
+}
+
 #endif // RHI_RAYTRACING
 
 FGraphicsMinimalPipelineStateId FGraphicsMinimalPipelineStateId::GetPersistentId(const FGraphicsMinimalPipelineStateInitializer& InPipelineState)
@@ -481,12 +487,13 @@ FMeshDrawShaderBindings::~FMeshDrawShaderBindings()
 
 void FMeshDrawShaderBindings::Initialize(FMeshProcessorShaders Shaders)
 {
-	const int32 NumShaderFrequencies = (Shaders.VertexShader.IsValid() ? 1 : 0) +
+	const int32 NumShaderFrequencies = 
+		(Shaders.VertexShader.IsValid() ? 1 : 0) +
 		(Shaders.PixelShader.IsValid() ? 1 : 0) +
 		(Shaders.GeometryShader.IsValid() ? 1 : 0) +
 		(Shaders.ComputeShader.IsValid() ? 1 : 0)
 #if RHI_RAYTRACING
-		+ (Shaders.RayHitGroupShader.IsValid() ? 1 : 0)
+		+ (Shaders.RayTracingShader.IsValid() ? 1 : 0)
 #endif
 		;
 
@@ -526,12 +533,14 @@ void FMeshDrawShaderBindings::Initialize(FMeshProcessorShaders Shaders)
 	}
 
 #if RHI_RAYTRACING
-	if (Shaders.RayHitGroupShader.IsValid())
+	if (Shaders.RayTracingShader.IsValid())
 	{
-		ShaderLayouts.Add(FMeshDrawShaderBindingsLayout(Shaders.RayHitGroupShader));
+		ShaderLayouts.Add(FMeshDrawShaderBindingsLayout(Shaders.RayTracingShader));
 		ShaderBindingDataSize += ShaderLayouts.Last().GetDataSizeBytes();
-		check(ShaderFrequencyBits < (1 << SF_RayHitGroup));
-		ShaderFrequencyBits |= (1 << SF_RayHitGroup);
+
+		const EShaderFrequency Frequency = Shaders.RayTracingShader->GetFrequency();
+		check(ShaderFrequencyBits < (1 << Frequency));
+		ShaderFrequencyBits |= (1 << Frequency);
 	}
 #endif
 
@@ -768,9 +777,9 @@ void FMeshDrawCommand::SetShaders(FRHIVertexDeclaration* VertexDeclaration, cons
 #if RHI_RAYTRACING
 void FRayTracingMeshCommand::SetShaders(const FMeshProcessorShaders& Shaders)
 {
-	check(Shaders.RayHitGroupShader.IsValid())
-	MaterialShaderIndex = Shaders.RayHitGroupShader.GetRayTracingHitGroupLibraryIndex();
-	MaterialShader = Shaders.RayHitGroupShader.GetRayTracingShader();
+	check(Shaders.RayTracingShader.IsValid())
+	MaterialShaderIndex = Shaders.RayTracingShader.GetRayTracingHitGroupLibraryIndex();
+	MaterialShader = Shaders.RayTracingShader.GetRayTracingShader();
 	ShaderBindings.Initialize(Shaders);
 }
 #endif // RHI_RAYTRACING
diff --git a/Engine/Source/Runtime/Renderer/Private/PathTracing.cpp b/Engine/Source/Runtime/Renderer/Private/PathTracing.cpp
--- a/Engine/Source/Runtime/Renderer/Private/PathTracing.cpp
+++ b/Engine/Source/Runtime/Renderer/Private/PathTracing.cpp
@@ -981,7 +981,7 @@ bool FRayTracingMeshProcessor::ProcessPathTracing(
 		FMeshMaterialShader,
 		FMeshMaterialShader,
 		FMeshMaterialShader> RayTracingShaders;
-	if (!Shaders.TryGetShader(SF_RayHitGroup, RayTracingShaders.RayHitGroupShader))
+	if (!Shaders.TryGetShader(SF_RayHitGroup, RayTracingShaders.RayTracingShader))
 	{
 		return false;
 	}
diff --git a/Engine/Source/Runtime/Renderer/Private/PrimitiveSceneInfo.cpp b/Engine/Source/Runtime/Renderer/Private/PrimitiveSceneInfo.cpp
--- a/Engine/Source/Runtime/Renderer/Private/PrimitiveSceneInfo.cpp
+++ b/Engine/Source/Runtime/Renderer/Private/PrimitiveSceneInfo.cpp
@@ -798,7 +798,7 @@ void CacheRayTracingPrimitive(
 		for (const FStaticMeshBatch& Mesh : SceneInfo->StaticMeshes)
 		{
 			// Why do we pass a full mask here when the dynamic case only uses a mask of 1?
-			// Also note that the code below assumes only a single command was generated per batch.
+			// Also note that the code below assumes only a single command was generated per batch (see SupportsCachingMeshDrawCommands(...))
 			const uint64 BatchElementMask = ~0ull;
 			RayTracingMeshProcessor.AddMeshBatch(Mesh, BatchElementMask, SceneInfo->Proxy);
 
@@ -851,7 +851,7 @@ void CacheRayTracingPrimitive(
 		for (const FMeshBatch& Mesh : CachedRayTracingInstance.Materials)
 		{
 			// Why do we pass a full mask here when the dynamic case only uses a mask of 1?
-			// Also note that the code below assumes only a single command was generated per batch.
+			// Also note that the code below assumes only a single command was generated per batch (see SupportsCachingMeshDrawCommands(...))
 			const uint64 BatchElementMask = ~0ull;
 			RayTracingMeshProcessor.AddMeshBatch(Mesh, BatchElementMask, SceneInfo->Proxy);
 
diff --git a/Engine/Source/Runtime/Renderer/Private/RayTracing/RayTracingMaterialHitShaders.cpp b/Engine/Source/Runtime/Renderer/Private/RayTracing/RayTracingMaterialHitShaders.cpp
--- a/Engine/Source/Runtime/Renderer/Private/RayTracing/RayTracingMaterialHitShaders.cpp
+++ b/Engine/Source/Runtime/Renderer/Private/RayTracing/RayTracingMaterialHitShaders.cpp
@@ -384,7 +384,7 @@ bool FRayTracingMeshProcessor::Process(
 
 	if (GRHISupportsRayTracingShaders)
 	{
-		if (!GetRayTracingMeshProcessorShaders(LightMapPolicy, MeshBatch.VertexFactory, MaterialResource, RayTracingShaders.RayHitGroupShader))
+		if (!GetRayTracingMeshProcessorShaders(LightMapPolicy, MeshBatch.VertexFactory, MaterialResource, RayTracingShaders.RayTracingShader))
 		{
 			return false;
 		}
diff --git a/Engine/Source/Runtime/Renderer/Private/RayTracing/RayTracingMaterialHitShaders.h b/Engine/Source/Runtime/Renderer/Private/RayTracing/RayTracingMaterialHitShaders.h
--- a/Engine/Source/Runtime/Renderer/Private/RayTracing/RayTracingMaterialHitShaders.h
+++ b/Engine/Source/Runtime/Renderer/Private/RayTracing/RayTracingMaterialHitShaders.h
@@ -81,10 +81,10 @@ protected:
 		VertexFactory->GetStreams(ERHIFeatureLevel::SM5, EVertexInputStreamType::Default, VertexStreams);
 
 		int32 DataOffset = 0;
-		if (PassShaders.RayHitGroupShader.IsValid())
+		if (PassShaders.RayTracingShader.IsValid())
 		{
 			FMeshDrawSingleShaderBindings ShaderBindings = SharedCommand.ShaderBindings.GetSingleShaderBindings(SF_RayHitGroup, DataOffset);
-			PassShaders.RayHitGroupShader->GetShaderBindings(Scene, FeatureLevel, PrimitiveSceneProxy, MaterialRenderProxy, MaterialResource, DrawRenderState, ShaderElementData, ShaderBindings);
+			PassShaders.RayTracingShader->GetShaderBindings(Scene, FeatureLevel, PrimitiveSceneProxy, MaterialRenderProxy, MaterialResource, DrawRenderState, ShaderElementData, ShaderBindings);
 		}
 
 		const int32 NumElements = MeshBatch.Elements.Num();
@@ -97,10 +97,10 @@ protected:
 				FRayTracingMeshCommand& RayTracingMeshCommand = CommandContext->AddCommand(SharedCommand);
 
 				DataOffset = 0;
-				if (PassShaders.RayHitGroupShader.IsValid())
+				if (PassShaders.RayTracingShader.IsValid())
 				{
 					FMeshDrawSingleShaderBindings RayHitGroupShaderBindings = RayTracingMeshCommand.ShaderBindings.GetSingleShaderBindings(SF_RayHitGroup, DataOffset);
-					FMeshMaterialShader::GetElementShaderBindings(PassShaders.RayHitGroupShader, Scene, ViewIfDynamicMeshCommand, VertexFactory, EVertexInputStreamType::Default, FeatureLevel, PrimitiveSceneProxy, MeshBatch, BatchElement, ShaderElementData, RayHitGroupShaderBindings, VertexStreams);
+					FMeshMaterialShader::GetElementShaderBindings(PassShaders.RayTracingShader, Scene, ViewIfDynamicMeshCommand, VertexFactory, EVertexInputStreamType::Default, FeatureLevel, PrimitiveSceneProxy, MeshBatch, BatchElement, ShaderElementData, RayHitGroupShaderBindings, VertexStreams);
 				}
 
 				RayTracingMeshCommand.GeometrySegmentIndex = uint32(MeshBatch.SegmentIndex) + BatchElementIndex;
diff --git a/Engine/Source/Runtime/Renderer/Public/MeshPassProcessor.h b/Engine/Source/Runtime/Renderer/Public/MeshPassProcessor.h
--- a/Engine/Source/Runtime/Renderer/Public/MeshPassProcessor.h
+++ b/Engine/Source/Runtime/Renderer/Public/MeshPassProcessor.h
@@ -584,7 +584,7 @@ struct FMeshProcessorShaders
 	mutable TShaderRef<FShader> GeometryShader;
 	mutable TShaderRef<FShader> ComputeShader;
 #if RHI_RAYTRACING
-	mutable TShaderRef<FShader> RayHitGroupShader;
+	mutable TShaderRef<FShader> RayTracingShader;
 #endif
 
 	TShaderRef<FShader> GetShader(EShaderFrequency Frequency) const
@@ -606,9 +606,14 @@ struct FMeshProcessorShaders
 			return ComputeShader;
 		}
 #if RHI_RAYTRACING
-		else if (Frequency == SF_RayHitGroup)
+		else if (Frequency == SF_RayHitGroup || Frequency == SF_RayCallable || Frequency == SF_RayMiss)
 		{
-			return RayHitGroupShader;
+			if (RayTracingShader.IsValid() && Frequency != RayTracingShader->GetFrequency())
+			{
+				checkf(0, TEXT("Requested raytracing shader frequency (%d) doesn't match assigned shader frequency (%d)."), Frequency, RayTracingShader->GetFrequency());
+				return TShaderRef<FShader>();
+			}
+			return RayTracingShader;
 		}
 #endif // RHI_RAYTRACING
 
@@ -819,7 +824,8 @@ public:
 	RENDERER_API void SetOnCommandList(FRHIComputeCommandList& RHICmdList, FRHIComputeShader* Shader, class FShaderBindingState* StateCacheShaderBindings = nullptr) const;
 
 #if RHI_RAYTRACING
-	RENDERER_API void SetRayTracingShaderBindingsForHitGroup(FRayTracingLocalShaderBindingWriter* BindingWriter, uint32 InstanceIndex, uint32 SegmentIndex, uint32 HitGroupIndex, uint32 ShaderSlot) const;
+	RENDERER_API void SetRayTracingShaderBindingsForHitGroup(FRayTracingLocalShaderBindingWriter* BindingWriter, uint32 InstanceIndex, uint32 SegmentIndex, uint32 HitGroupIndexInPipeline, uint32 ShaderSlot) const;
+	RENDERER_API void SetRayTracingShaderBindings(FRayTracingLocalShaderBindingWriter* BindingWriter, uint32 ShaderIndexInPipeline, uint32 ShaderSlot) const;
 #endif // RHI_RAYTRACING
 
 	/** Returns whether this set of shader bindings can be merged into an instanced draw call with another. */
@@ -1692,7 +1698,7 @@ private:
 	TArray<Experimental::FHashType> DeferredCommandHashes;
 };
 
-template<typename VertexType, typename PixelType, typename GeometryType = FMeshMaterialShader, typename RayHitGroupType = FMeshMaterialShader, typename ComputeType = FMeshMaterialShader>
+template<typename VertexType, typename PixelType, typename GeometryType = FMeshMaterialShader, typename RayTracingType = FMeshMaterialShader, typename ComputeType = FMeshMaterialShader>
 struct TMeshProcessorShaders
 {
 	TShaderRef<VertexType> VertexShader;
@@ -1700,7 +1706,7 @@ struct TMeshProcessorShaders
 	TShaderRef<GeometryType> GeometryShader;
 	TShaderRef<ComputeType> ComputeShader;
 #if RHI_RAYTRACING
-	TShaderRef<RayHitGroupType> RayHitGroupShader;
+	TShaderRef<RayTracingType> RayTracingShader;
 #endif
 
 	TMeshProcessorShaders() = default;
@@ -1713,7 +1719,7 @@ struct TMeshProcessorShaders
 		Shaders.GeometryShader = GeometryShader;
 		Shaders.ComputeShader = ComputeShader;
 #if RHI_RAYTRACING
-		Shaders.RayHitGroupShader = RayHitGroupShader;
+		Shaders.RayTracingShader = RayTracingShader;
 #endif
 		return Shaders;
 	}
