Missing include for RHI_RAYTRACING

diff --git a/Engine/Source/Runtime/Renderer/Private/RayTracing/RayTracingTraversalStatistics.h b/Engine/Source/Runtime/Renderer/Private/RayTracing/RayTracingTraversalStatistics.h
--- a/Engine/Source/Runtime/Renderer/Private/RayTracing/RayTracingTraversalStatistics.h
+++ b/Engine/Source/Runtime/Renderer/Private/RayTracing/RayTracingTraversalStatistics.h
@@ -3,6 +3,7 @@
 #pragma once
 
 #include "CoreMinimal.h"
+#include "RHIDefinitions.h"
 
 #if RHI_RAYTRACING
 
