Deprecated GRHISupportsResolveCubemapFaces.

diff --git a/Engine/Source/Runtime/Apple/AGXRHI/Private/Mac/MacPlatformAGXConfig.cpp b/Engine/Source/Runtime/Apple/AGXRHI/Private/Mac/MacPlatformAGXConfig.cpp
--- a/Engine/Source/Runtime/Apple/AGXRHI/Private/Mac/MacPlatformAGXConfig.cpp
+++ b/Engine/Source/Runtime/Apple/AGXRHI/Private/Mac/MacPlatformAGXConfig.cpp
@@ -527,7 +527,6 @@ GRHISupportsRayTracingDispatchIndirect                 = MAC_PLATFORM_AGXRHI_SUP
 GRHISupportsRayTracingPSOAdditions                     = MAC_PLATFORM_AGXRHI_SUPPORTS_RAY_TRACING_PSO_ADDITIONS;
 GRHISupportsRayTracingShaders                          = MAC_PLATFORM_AGXRHI_SUPPORTS_RAY_TRACING_SHADERS;
 GRHISupportsRectTopology                               = MAC_PLATFORM_AGXRHI_SUPPORTS_RECT_TOPOLOGY;
-GRHISupportsResolveCubemapFaces                        = MAC_PLATFORM_AGXRHI_SUPPORTS_RESOLVE_CUBEMAP_FACES;
 GRHISupportsResummarizeHTile                           = MAC_PLATFORM_AGXRHI_SUPPORTS_RESUMMARIZE_HTILE;
 GRHISupportsRHIOnTaskThread                            = MAC_PLATFORM_AGXRHI_SUPPORTS_RHI_ON_TASK_THREAD;
 GRHISupportsRHIThread                                  = MAC_PLATFORM_AGXRHI_SUPPORTS_RHI_THREAD;
diff --git a/Engine/Source/Runtime/D3D12RHI/Private/D3D12RHI.cpp b/Engine/Source/Runtime/D3D12RHI/Private/D3D12RHI.cpp
--- a/Engine/Source/Runtime/D3D12RHI/Private/D3D12RHI.cpp
+++ b/Engine/Source/Runtime/D3D12RHI/Private/D3D12RHI.cpp
@@ -197,7 +197,6 @@ FD3D12DynamicRHI::FD3D12DynamicRHI(const TArray<TSharedPtr<FD3D12Adapter>>& Chos
 	GMaxTextureMipCount = FMath::Min<int32>(MAX_TEXTURE_MIP_COUNT, GMaxTextureMipCount);
 	GMaxShadowDepthBufferSizeX = GMaxTextureDimensions;
 	GMaxShadowDepthBufferSizeY = GMaxTextureDimensions;
-	GRHISupportsResolveCubemapFaces = true;
 	GRHISupportsArrayIndexFromAnyShader = true;
 
 	GRHIMaxDispatchThreadGroupsPerDimension.X = D3D12_CS_DISPATCH_MAX_THREAD_GROUPS_PER_DIMENSION;
diff --git a/Engine/Source/Runtime/RHI/Private/RHI.cpp b/Engine/Source/Runtime/RHI/Private/RHI.cpp
--- a/Engine/Source/Runtime/RHI/Private/RHI.cpp
+++ b/Engine/Source/Runtime/RHI/Private/RHI.cpp
@@ -1426,7 +1426,6 @@ bool GRHISupportsDirectGPUMemoryLock = false;
 bool GRHISupportsMultithreadedShaderCreation = true;
 
 bool GRHISupportsMSAADepthSampleAccess = false;
-bool GRHISupportsResolveCubemapFaces = false;
 
 bool GRHISupportsBackBufferWithCustomDepthStencil = true;
 
diff --git a/Engine/Source/Runtime/RHI/Public/RHI.h b/Engine/Source/Runtime/RHI/Public/RHI.h
--- a/Engine/Source/Runtime/RHI/Public/RHI.h
+++ b/Engine/Source/Runtime/RHI/Public/RHI.h
@@ -606,6 +606,7 @@ RHI_API void RHISetCurrentNumDrawCallPtr(FRHIDrawCallsStatPtr InNumDrawCallsRHIP
 extern RHI_API bool GRHISupportsBaseVertexIndex;
 
 /** True if the RHI supports copying cubemap faces using CopyToResolveTarget */
+UE_DEPRECATED(5.1, "CopyToResoveTarget is deprecated.")
 extern RHI_API bool GRHISupportsResolveCubemapFaces;
 
 /** Whether or not the RHI can handle a non-zero FirstInstance to DrawIndexedPrimitive and friends - extra SetStreamSource calls will be needed if this is false */
diff --git a/Engine/Source/Runtime/Renderer/Private/ReflectionEnvironmentCapture.cpp b/Engine/Source/Runtime/Renderer/Private/ReflectionEnvironmentCapture.cpp
--- a/Engine/Source/Runtime/Renderer/Private/ReflectionEnvironmentCapture.cpp
+++ b/Engine/Source/Runtime/Renderer/Private/ReflectionEnvironmentCapture.cpp
@@ -218,11 +218,6 @@ public:
 
 IMPLEMENT_GLOBAL_SHADER(FCopyCubemapToCubeFacePS, "/Engine/Private/ReflectionEnvironmentShaders.usf", "CopyCubemapToCubeFaceColorPS", SF_Pixel);
 
-bool DoGPUArrayCopy()
-{
-	return GRHISupportsResolveCubemapFaces;
-}
-
 class FReflectionCubemapTexture : public FRenderThreadStructBase
 {
 public:
@@ -795,7 +790,7 @@ void FScene::AllocateReflectionCaptures(const TArray<UReflectionCaptureComponent
 			const int32 ReflectionCaptureSize = GetReflectionCaptureSizeForArrayCount(UReflectionCaptureComponent::GetReflectionCaptureSize(), DesiredMaxCubemaps);
 			bool bNeedsUpdateAllCaptures = ReflectionSceneData.DoesAllocatedDataNeedUpdate(DesiredMaxCubemaps, ReflectionCaptureSize);
 
-			if (DoGPUArrayCopy() && bNeedsUpdateAllCaptures)
+			if (bNeedsUpdateAllCaptures)
 			{
 				// If we're not in the editor, we discard the CPU-side reflection capture data after loading to save memory, so we can't resize if the resolution changes. If this happens, we assert
 				check(GIsEditor || ReflectionCaptureSize == ReflectionSceneData.CubemapArray.GetCubemapSize() || ReflectionSceneData.CubemapArray.GetCubemapSize() == 0);
@@ -1280,10 +1275,7 @@ void FScene::CaptureOrUploadReflectionCapture(UReflectionCaptureComponent* Captu
 				{
 					UploadReflectionCapture_RenderingThread(Scene, CaptureData, CaptureComponent);
 
-					if (DoGPUArrayCopy())
-					{
-						CaptureData->OnDataUploadedToGPUFinal();
-					}
+					CaptureData->OnDataUploadedToGPUFinal();
 				}
 				else
 				{
diff --git a/Engine/Source/Runtime/Windows/D3D11RHI/Private/D3D11Device.cpp b/Engine/Source/Runtime/Windows/D3D11RHI/Private/D3D11Device.cpp
--- a/Engine/Source/Runtime/Windows/D3D11RHI/Private/D3D11Device.cpp
+++ b/Engine/Source/Runtime/Windows/D3D11RHI/Private/D3D11Device.cpp
@@ -211,7 +211,6 @@ FD3D11DynamicRHI::FD3D11DynamicRHI(IDXGIFactory1* InDXGIFactory1, D3D_FEATURE_LE
 	GMaxShadowDepthBufferSizeX = GMaxTextureDimensions;
 	GMaxShadowDepthBufferSizeY = GMaxTextureDimensions;
 	GSupportsTimestampRenderQueries = true;
-	GRHISupportsResolveCubemapFaces = true;
 
 	GRHIMaxDispatchThreadGroupsPerDimension.X = D3D11_CS_DISPATCH_MAX_THREAD_GROUPS_PER_DIMENSION;
 	GRHIMaxDispatchThreadGroupsPerDimension.Y = D3D11_CS_DISPATCH_MAX_THREAD_GROUPS_PER_DIMENSION;
