Multi-view-family:  need to call pre-render view extensions for all view families.

#rnx

diff --git a/Engine/Source/Runtime/Renderer/Private/SceneRendering.cpp b/Engine/Source/Runtime/Renderer/Private/SceneRendering.cpp
--- a/Engine/Source/Runtime/Renderer/Private/SceneRendering.cpp
+++ b/Engine/Source/Runtime/Renderer/Private/SceneRendering.cpp
@@ -4229,7 +4229,12 @@ static void RenderViewFamily_RenderThread(FRHICommandListImmediate& RHICmdList,
 	LLM_SCOPE(ELLMTag::SceneRender);
 
 	// We need to execute the pre-render view extensions before we do any view dependent work.
-	FSceneRenderer::ViewExtensionPreRender_RenderThread(RHICmdList, SceneRenderer);
+	for (FViewFamilyInfo& ViewFamily : SceneRenderer->ViewFamilies)
+	{
+		SceneRenderer->SetActiveViewFamily(ViewFamily);
+
+		FSceneRenderer::ViewExtensionPreRender_RenderThread(RHICmdList, SceneRenderer);
+	}
 
 	SceneRenderer->RenderThreadBegin(RHICmdList);
 
