Path Tracer: Fix moveable lights not invalidating the path tracer

diff --git a/Engine/Source/Runtime/Renderer/Private/RendererScene.cpp b/Engine/Source/Runtime/Renderer/Private/RendererScene.cpp
--- a/Engine/Source/Runtime/Renderer/Private/RendererScene.cpp
+++ b/Engine/Source/Runtime/Renderer/Private/RendererScene.cpp
@@ -3036,6 +3036,8 @@ void FScene::UpdateLightTransform_RenderThread(FLightSceneInfo* LightSceneInfo,
 		// Update the light's transform and position.
 		LightSceneInfo->Proxy->SetTransform(Parameters.LightToWorld,Parameters.Position);
 
+		InvalidatePathTracedOutput();
+
 		// Also update the LightSceneInfoCompact
 		if( bHasId )
 		{
