Strata - fix for editor primitives not showing up when strata visualize is enabled.

diff --git a/Engine/Source/Runtime/Renderer/Private/SceneRendering.cpp b/Engine/Source/Runtime/Renderer/Private/SceneRendering.cpp
--- a/Engine/Source/Runtime/Renderer/Private/SceneRendering.cpp
+++ b/Engine/Source/Runtime/Renderer/Private/SceneRendering.cpp
@@ -3872,7 +3872,7 @@ void FSceneRenderer::OnStartRender(FRHICommandListImmediate& RHICmdList)
 
 bool FSceneRenderer::ShouldCompositeEditorPrimitives(const FViewInfo& View)
 {
-	if (View.Family->EngineShowFlags.VisualizeHDR || View.Family->EngineShowFlags.VisualizeStrataMaterial || View.Family->UseDebugViewPS())
+	if (View.Family->EngineShowFlags.VisualizeHDR || View.Family->UseDebugViewPS())
 	{
 		// certain visualize modes get obstructed too much
 		return false;
