Added light grid culling debug pass.

diff --git a/Engine/Shaders/Private/LightGridInjection.usf b/Engine/Shaders/Private/LightGridInjection.usf
--- a/Engine/Shaders/Private/LightGridInjection.usf
+++ b/Engine/Shaders/Private/LightGridInjection.usf
@@ -158,6 +158,8 @@ bool IsAabbOutsideInfiniteAcuteConeApprox(float3 ConeVertex, float3 ConeAxis, fl
 	return AabbOutsidePlane(D, AabbExt, Plane);
 }
 
+#if SHADER_LIGHT_GRID_INJECTION_CS
+
 [numthreads(THREADGROUP_SIZE, THREADGROUP_SIZE, THREADGROUP_SIZE)]
 void LightGridInjectionCS(
 	uint3 GroupId : SV_GroupID,
@@ -193,7 +195,6 @@ void LightGridInjectionCS(
 			float4 LightPositionAndRadius = LightViewSpacePositionAndRadius[LocalLightIndex];
 			float3 ViewSpaceLightPosition = LightPositionAndRadius.xyz;
 			float LightRadius = LightPositionAndRadius.w;
-			float4 LightPositionAndInvRadius = ForwardLocalLightBuffer[LocalLightBaseIndex + 0];
 #else // !ENABLE_LIGHT_CULLING_VIEW_SPACE_BUILD_DATA
 			float4 LightPositionAndInvRadius = ForwardLocalLightBuffer[LocalLightBaseIndex + 0];
 			float LightRadius = 1.0f / LightPositionAndInvRadius.w;
@@ -313,6 +314,12 @@ void LightGridInjectionCS(
 	}
 }
 
+#endif // SHADER_LIGHT_GRID_INJECTION_CS
+
+
+
+#if SHADER_LIGHT_GRID_COMPACT_CS
+
 RWBuffer<uint> RWNextCulledLightData;
 Buffer<uint> StartOffsetGrid;
 Buffer<uint> CulledLightLinks;
@@ -405,4 +412,100 @@ void LightGridCompactCS(
 
 	// Compact reflection captures
 	CompactReverseLinkedList(NumGridCells + GridIndex, NumReflectionCaptures, bFirstThread, bThreadValid);
-}
\ No newline at end of file
+}
+
+#endif // SHADER_LIGHT_GRID_COMPACT_CS
+
+
+
+#ifdef SHADER_DEBUG_LIGHT_GRID_PS
+
+#define SUPPORT_CONTACT_SHADOWS 0
+#include "DeferredLightingCommon.ush"
+#include "LightGridCommon.ush"
+#include "/Engine/Private/MiniFontCommon.ush"
+#include "/Engine/Private/ShaderPrint.ush"
+#include "ColorMap.ush"
+
+Texture2D<float4> DepthTexture;
+uint DebugMode;
+
+void DebugLightGridPS(
+	in float4 SVPos : SV_POSITION,
+	out float4 OutLuminanceTransmittance : SV_Target0)
+{
+	ResolvedView = ResolveView();
+
+	const uint EyeIndex = 0;
+	const uint2 PixelPos = SVPos.xy;
+	const uint2 LocalPixelPos = (SVPos.xy - ResolvedView.ViewRectMin.xy);
+
+	// Using View.ViewResolutionFraction to maintain a good look under any dynamic resolution factor.
+	const uint2 PixelPosDynRes = float2(PixelPos) * View.ViewResolutionFraction;
+	const uint2 LocalPixelPosDynRes = float2(LocalPixelPos) * View.ViewResolutionFraction;
+	const uint2 LocalPixelPosDynRes2 = float2(LocalPixelPos+1) * View.ViewResolutionFraction;
+
+	OutLuminanceTransmittance = float4(0, 0, 0, 1);
+
+	if (any(LocalPixelPosDynRes >= uint2(ResolvedView.ViewSizeAndInvSize.xy)))
+	{
+		return;
+	}
+
+	const FLightGridData GridData = GetLightGridData(EyeIndex);
+	uint DebugNumLocalLights = 0;
+
+	uint2 StartPos = uint2(50, 100);
+	FShaderPrintContext Context = InitShaderPrintContext(all(PixelPos == StartPos), StartPos);
+	bool SpecifyDebugSlice = AddCheckbox(Context, TEXT("Specify Slice to debug"), false, GetDefaultFontColor());
+	Newline(Context);
+	float DebugSlice = 0;
+	FFontColor EnableSliceDebugFont = FontDarkGrey; if (SpecifyDebugSlice) { EnableSliceDebugFont = FontWhite; }
+	DebugSlice = AddSlider(Context, TEXT("Slice to debug"), 0.0f, EnableSliceDebugFont, 0.0f, GridData.CulledGridSize.z-1);
+	Newline(Context);
+
+	if (DebugMode == 1 && !SpecifyDebugSlice)
+	{
+		const float SceneDepth = ConvertFromDeviceZ(DepthTexture.Load(int3(PixelPosDynRes, 0)).r);
+
+		const uint GridIndex = ComputeLightGridCellIndex(LocalPixelPosDynRes, SceneDepth, EyeIndex);
+		const FCulledLightsGridData CulledLightsGrid = GetCulledLightsGrid(GridIndex, EyeIndex);
+
+		DebugNumLocalLights = CulledLightsGrid.NumLocalLights;
+	}
+	else
+	{
+		// We do not want to fetch light data from the last slice since in this case the culling will be pretty bad (super large depth turned into a bounding sphere...)
+		// See ComputeCellNearViewDepthFromZSlice.
+		const int StartGridIndexZ = SpecifyDebugSlice ? uint(DebugSlice) : 0;
+		const int EndGridIndexZ   = SpecifyDebugSlice ? StartGridIndexZ+1: (DebugMode == 2 ? GridData.CulledGridSize.z - 1 : GridData.CulledGridSize.z);
+		LOOP
+		for (int SliceIt = StartGridIndexZ; SliceIt < EndGridIndexZ; SliceIt++)
+		{
+			const uint GridIndex = ComputeLightGridCellIndex(uint3(LocalPixelPosDynRes >> GridData.LightGridPixelSizeShift, SliceIt), EyeIndex);
+			const FCulledLightsGridData CulledLightsGrid = GetCulledLightsGrid(GridIndex, EyeIndex);
+		
+			DebugNumLocalLights = max(DebugNumLocalLights, CulledLightsGrid.NumLocalLights);
+		}
+	}
+
+	const uint2 TileTopLeftPixelCoord = (LocalPixelPosDynRes  >> GridData.LightGridPixelSizeShift) << GridData.LightGridPixelSizeShift;
+	const uint2 TileTopLeftPixelCoord2= (LocalPixelPosDynRes2 >> GridData.LightGridPixelSizeShift) << GridData.LightGridPixelSizeShift;
+	const uint TileSize = (0x1 << GridData.LightGridPixelSizeShift) - 1;
+
+	if (DebugNumLocalLights > 0)
+	{
+		OutLuminanceTransmittance.a   = 0.0f;
+		OutLuminanceTransmittance.rgb = GetHSVDebugColor(float(DebugNumLocalLights) / 8.0f);
+		PrintSmallUint(PixelPos, OutLuminanceTransmittance.rgb, float3(0.1, 0.1, 0.1), (TileTopLeftPixelCoord + TileSize/2) / View.ViewResolutionFraction, float(DebugNumLocalLights));
+	}
+
+	if ((TileTopLeftPixelCoord.x <= LocalPixelPosDynRes.x && TileTopLeftPixelCoord2.x >= LocalPixelPosDynRes.x)
+		|| (TileTopLeftPixelCoord.y <= LocalPixelPosDynRes.y && TileTopLeftPixelCoord2.y >= LocalPixelPosDynRes.y))
+	{
+		OutLuminanceTransmittance.rgba = float4(0.25, 0.25, 0.25, 0.0);
+	}
+}
+
+#endif // SHADER_DEBUG_LIGHT_GRID_PS
+
diff --git a/Engine/Shaders/Private/ShaderPrint.ush b/Engine/Shaders/Private/ShaderPrint.ush
--- a/Engine/Shaders/Private/ShaderPrint.ush
+++ b/Engine/Shaders/Private/ShaderPrint.ush
@@ -166,6 +166,8 @@ FFontColor GetDefaultFontColor()							{ FFontColor Out; Out.Color = float3(1,1,
 
 // Predefined colors
 #define FontWhite		InitFontColor(1, 1, 1)
+#define FontGrey		InitFontColor(0.5, 0.5, 0.5)
+#define FontDarkGrey	InitFontColor(0.25, 0.25, 0.25)
 #define FontBlack		InitFontColor(0, 0, 0)
 #define FontRed			InitFontColor(1, 0, 0)
 #define FontGreen		InitFontColor(0, 1, 0)
diff --git a/Engine/Source/Runtime/Renderer/Private/DeferredShadingRenderer.h b/Engine/Source/Runtime/Renderer/Private/DeferredShadingRenderer.h
--- a/Engine/Source/Runtime/Renderer/Private/DeferredShadingRenderer.h
+++ b/Engine/Source/Runtime/Renderer/Private/DeferredShadingRenderer.h
@@ -240,6 +240,11 @@ public:
 	 */
 	void GatherLightsAndComputeLightGrid(FRDGBuilder& GraphBuilder, bool bNeedLightGrid, FSortedLightSetSceneInfo &SortedLightSet);
 
+	/** 
+	 * Debug light grid content on screen.
+	 */
+	void DebugLightGrid(FRDGBuilder& GraphBuilder, FSceneTextures& SceneTextures, bool bNeedLightGrid);
+
 	void RenderBasePass(
 		FRDGBuilder& GraphBuilder,
 		FSceneTextures& SceneTextures,
diff --git a/Engine/Source/Runtime/Renderer/Private/LightGridInjection.cpp b/Engine/Source/Runtime/Renderer/Private/LightGridInjection.cpp
--- a/Engine/Source/Runtime/Renderer/Private/LightGridInjection.cpp
+++ b/Engine/Source/Runtime/Renderer/Private/LightGridInjection.cpp
@@ -28,6 +28,9 @@ LightGridInjection.cpp
 #include "VolumetricCloudRendering.h"
 #include "Components/LightComponent.h"
 #include "Engine/MapBuildDataRegistry.h"
+#include "PixelShaderUtils.h"
+#include "ShaderPrint.h"
+#include "ShaderPrintParameters.h"
 
 int32 GLightGridPixelSize = 64;
 FAutoConsoleVariableRef CVarLightGridPixelSize(
@@ -45,6 +48,18 @@ FAutoConsoleVariableRef CVarLightGridSizeZ(
 	ECVF_Scalability | ECVF_RenderThreadSafe
 );
 
+int32 GForwardLightGridDebug = 0;
+FAutoConsoleVariableRef CVarLightGridDebug(
+	TEXT("r.Forward.LightGridDebug"),
+	GForwardLightGridDebug,
+	TEXT("Whether to display on screen culledlight per tile.\n")
+	TEXT(" 0: off (default)\n")
+	TEXT(" 1: on - showing light count onto the depth buffer\n")
+	TEXT(" 2: on - showing max light count per tile accoung for each slice but the last one (culling there is too conservative)\n")
+	TEXT(" 3: on - showing max light count per tile accoung for each slice and the last one \n"),
+	ECVF_RenderThreadSafe
+);
+
 int32 GMaxCulledLightsPerCell = 32;
 FAutoConsoleVariableRef CVarMaxCulledLightsPerCell(
 	TEXT("r.Forward.MaxCulledLightsPerCell"),
@@ -176,6 +191,7 @@ public:
 	static void ModifyCompilationEnvironment(const FGlobalShaderPermutationParameters& Parameters, FShaderCompilerEnvironment& OutEnvironment)
 	{
 		FGlobalShader::ModifyCompilationEnvironment(Parameters, OutEnvironment);
+		OutEnvironment.SetDefine(TEXT("SHADER_LIGHT_GRID_INJECTION_CS"), 1);
 		OutEnvironment.SetDefine(TEXT("THREADGROUP_SIZE"), LightGridInjectionGroupSize);
 		FForwardLightingParameters::ModifyCompilationEnvironment(Parameters.Platform, OutEnvironment);
 		OutEnvironment.SetDefine(TEXT("LIGHT_LINK_STRIDE"), LightLinkStride);
@@ -213,6 +229,7 @@ public:
 	static void ModifyCompilationEnvironment(const FGlobalShaderPermutationParameters& Parameters, FShaderCompilerEnvironment& OutEnvironment)
 	{
 		FGlobalShader::ModifyCompilationEnvironment(Parameters, OutEnvironment);
+		OutEnvironment.SetDefine(TEXT("SHADER_LIGHT_GRID_COMPACT_CS"), 1);
 		OutEnvironment.SetDefine(TEXT("THREADGROUP_SIZE"), LightGridInjectionGroupSize);
 		FForwardLightingParameters::ModifyCompilationEnvironment(Parameters.Platform, OutEnvironment);
 		OutEnvironment.SetDefine(TEXT("LIGHT_LINK_STRIDE"), LightLinkStride);
@@ -913,3 +930,75 @@ void FDeferredShadingSceneRenderer::RenderForwardShadowProjections(
 		GraphBuilder.AddPass(RDG_EVENT_NAME("ResolveScreenSpaceShadowMask"), PassParameters, ERDGPassFlags::Raster, [](FRHICommandList&) {});
 	}
 }
+
+class FDebugLightGridPS : public FGlobalShader
+{
+	DECLARE_GLOBAL_SHADER(FDebugLightGridPS);
+	SHADER_USE_PARAMETER_STRUCT(FDebugLightGridPS, FGlobalShader);
+
+	using FPermutationDomain = TShaderPermutationDomain<>;
+
+	BEGIN_SHADER_PARAMETER_STRUCT(FParameters, )
+		SHADER_PARAMETER_STRUCT_REF(FViewUniformShaderParameters, ViewUniformBuffer)
+		SHADER_PARAMETER_RDG_UNIFORM_BUFFER(FForwardLightData, Forward)
+		SHADER_PARAMETER_STRUCT_INCLUDE(ShaderPrint::FShaderParameters, ShaderPrintParameters)
+		SHADER_PARAMETER_TEXTURE(Texture2D, MiniFontTexture)
+		SHADER_PARAMETER_RDG_TEXTURE(Texture2D, DepthTexture)
+		SHADER_PARAMETER(uint32, DebugMode)
+		RENDER_TARGET_BINDING_SLOTS()
+	END_SHADER_PARAMETER_STRUCT()
+
+	static FPermutationDomain RemapPermutation(FPermutationDomain PermutationVector)
+	{
+		return PermutationVector;
+	}
+
+	static bool ShouldCompilePermutation(const FGlobalShaderPermutationParameters& Parameters)
+	{
+		return EnumHasAllFlags(Parameters.Flags, EShaderPermutationFlags::HasEditorOnlyData);
+	}
+
+	static void ModifyCompilationEnvironment(const FGlobalShaderPermutationParameters& Parameters, FShaderCompilerEnvironment& OutEnvironment)
+	{
+		FGlobalShader::ModifyCompilationEnvironment(Parameters, OutEnvironment);
+
+		// Stay debug and skip optimizations to reduce compilation time on this long shader.
+		OutEnvironment.CompilerFlags.Add(CFLAG_Debug);
+		OutEnvironment.CompilerFlags.Add(CFLAG_AllowTypedUAVLoads);
+		OutEnvironment.SetDefine(TEXT("SHADER_DEBUG_LIGHT_GRID_PS"), 1);
+		FForwardLightingParameters::ModifyCompilationEnvironment(Parameters.Platform, OutEnvironment);
+	}
+};
+IMPLEMENT_GLOBAL_SHADER(FDebugLightGridPS, "/Engine/Private/LightGridInjection.usf", "DebugLightGridPS", SF_Pixel);
+
+bool ShouldVisualizeLightGrid()
+{
+	return GForwardLightGridDebug > 0;
+}
+
+FScreenPassTexture AddVisualizeLightGridPass(FRDGBuilder& GraphBuilder, const FViewInfo& View, FScreenPassTexture& ScreenPassSceneColor, FRDGTextureRef SceneDepthTexture)
+{
+	RDG_EVENT_SCOPE(GraphBuilder, "VisualizeLightGrid");
+
+	if (!ShaderPrint::IsEnabled(View)) { ShaderPrint::SetEnabled(true); }
+	ShaderPrint::RequestSpaceForLines(128);
+	ShaderPrint::RequestSpaceForCharacters(128);
+
+	FDebugLightGridPS::FPermutationDomain PermutationVector;
+	TShaderMapRef<FDebugLightGridPS> PixelShader(View.ShaderMap, PermutationVector);
+	FDebugLightGridPS::FParameters* PassParameters = GraphBuilder.AllocParameters<FDebugLightGridPS::FParameters>();
+	PassParameters->ViewUniformBuffer = View.ViewUniformBuffer;
+	PassParameters->Forward = View.ForwardLightingResources.ForwardLightUniformBuffer;
+	ShaderPrint::SetParameters(GraphBuilder, View.ShaderPrintData, PassParameters->ShaderPrintParameters);
+	PassParameters->DepthTexture = SceneDepthTexture ? SceneDepthTexture : GSystemTextures.GetMaxFP16Depth(GraphBuilder);
+	PassParameters->MiniFontTexture = GetMiniFontTexture();
+	PassParameters->RenderTargets[0] = FRenderTargetBinding(ScreenPassSceneColor.Texture, ERenderTargetLoadAction::ELoad);
+	PassParameters->DebugMode = GForwardLightGridDebug;
+
+	FRHIBlendState* PreMultipliedColorTransmittanceBlend = TStaticBlendState<CW_RGB, BO_Add, BF_One, BF_SourceAlpha, BO_Add, BF_Zero, BF_One>::GetRHI();
+
+	FPixelShaderUtils::AddFullscreenPass<FDebugLightGridPS>(GraphBuilder, View.ShaderMap, RDG_EVENT_NAME("DebugLightGridCS"), PixelShader, PassParameters,
+		ScreenPassSceneColor.ViewRect, PreMultipliedColorTransmittanceBlend);
+
+	return MoveTemp(ScreenPassSceneColor);
+}
diff --git a/Engine/Source/Runtime/Renderer/Private/PostProcess/PostProcessing.cpp b/Engine/Source/Runtime/Renderer/Private/PostProcess/PostProcessing.cpp
--- a/Engine/Source/Runtime/Renderer/Private/PostProcess/PostProcessing.cpp
+++ b/Engine/Source/Runtime/Renderer/Private/PostProcess/PostProcessing.cpp
@@ -63,6 +63,9 @@ bool IsMobileEyeAdaptationEnabled(const FViewInfo& View);
 
 bool IsValidBloomSetupVariation(bool bUseBloom, bool bUseSun, bool bUseDof, bool bUseEyeAdaptation);
 
+extern FScreenPassTexture AddVisualizeLightGridPass(FRDGBuilder& GraphBuilder, const FViewInfo& View, FScreenPassTexture& ScreenPassSceneColor, FRDGTextureRef SceneDepthTexture);
+extern bool ShouldVisualizeLightGrid();
+
 namespace
 {
 TAutoConsoleVariable<float> CVarDepthOfFieldNearBlurSizeThreshold(
@@ -265,6 +268,7 @@ void AddPostProcessingPasses(
 		VisualizeLightCulling,
 		VisualizePostProcessStack,
 		VisualizeStrata,
+		VisualizeLightGrid,
 		VisualizeSkyAtmosphere,
 		VisualizeLevelInstance,
 		SelectionOutline,
@@ -312,6 +316,7 @@ void AddPostProcessingPasses(
 		TEXT("VisualizeLightCulling"),
 		TEXT("VisualizePostProcessStack"),
 		TEXT("VisualizeStrata"),
+		TEXT("VisualizeLightGrid"),
 		TEXT("VisualizeSkyAtmosphere"),
 		TEXT("VisualizeLevelInstance"),
 		TEXT("SelectionOutline"),
@@ -344,6 +349,8 @@ void AddPostProcessingPasses(
 #endif
 	PassSequence.SetEnabled(EPass::VisualizeStrata, Strata::ShouldRenderStrataDebugPasses(View));
 	PassSequence.SetEnabled(EPass::VisualizeLumenScene, LumenVisualizeMode >= 0 && LumenVisualizeMode != VISUALIZE_MODE_OVERVIEW && bPostProcessingEnabled);
+	PassSequence.SetEnabled(EPass::VisualizeLightGrid, ShouldVisualizeLightGrid());
+
 #if WITH_EDITOR
 	PassSequence.SetEnabled(EPass::VisualizeSkyAtmosphere, Scene && View.Family && View.Family->EngineShowFlags.VisualizeSkyAtmosphere&& ShouldRenderSkyAtmosphereDebugPasses(Scene, View.Family->EngineShowFlags));
 	PassSequence.SetEnabled(EPass::VisualizeLevelInstance, GIsEditor && EngineShowFlags.EditingLevelInstance && EngineShowFlags.VisualizeLevelInstanceEditing && !bVisualizeHDR);
@@ -1078,6 +1085,13 @@ void AddPostProcessingPasses(
 	}
 #endif
 
+	if (PassSequence.IsEnabled(EPass::VisualizeLightGrid))
+	{
+		FScreenPassRenderTarget OverrideOutput;
+		PassSequence.AcceptOverrideIfLastPass(EPass::VisualizeLightGrid, OverrideOutput);
+		SceneColor = AddVisualizeLightGridPass(GraphBuilder, View, SceneColor, SceneDepth.Texture);
+	}
+
 	if (PassSequence.IsEnabled(EPass::VisualizeStrata))
 	{
 		FScreenPassRenderTarget OverrideOutput;
