Fixed skyatmosphere rendered over opaque meshes already containing fog when in forward.

diff --git a/Engine/Source/Runtime/Renderer/Private/DeferredShadingRenderer.cpp b/Engine/Source/Runtime/Renderer/Private/DeferredShadingRenderer.cpp
--- a/Engine/Source/Runtime/Renderer/Private/DeferredShadingRenderer.cpp
+++ b/Engine/Source/Runtime/Renderer/Private/DeferredShadingRenderer.cpp
@@ -3043,7 +3043,7 @@ void FDeferredShadingSceneRenderer::Render(FRDGBuilder& GraphBuilder)
 	}
 
 	// Draw the sky atmosphere
-	if (!bHasRayTracedOverlay && bShouldRenderSkyAtmosphere)
+	if (!bHasRayTracedOverlay && bShouldRenderSkyAtmosphere && !IsForwardShadingEnabled(ShaderPlatform))
 	{
 		SCOPE_CYCLE_COUNTER(STAT_FDeferredShadingSceneRenderer_RenderSkyAtmosphere);
 		RenderSkyAtmosphere(GraphBuilder, SceneTextures);
