Moved some decal rendering logic to CalculateDecalFadeAlpha.
- will be used to support decals in raytracing techniques.

diff --git a/Engine/Source/Runtime/Renderer/Private/CompositionLighting/PostProcessMeshDecals.cpp b/Engine/Source/Runtime/Renderer/Private/CompositionLighting/PostProcessMeshDecals.cpp
--- a/Engine/Source/Runtime/Renderer/Private/CompositionLighting/PostProcessMeshDecals.cpp
+++ b/Engine/Source/Runtime/Renderer/Private/CompositionLighting/PostProcessMeshDecals.cpp
@@ -206,7 +206,7 @@ bool FMeshDecalMeshProcessor::TryAddMeshBatch(
 		if (Material.GetRenderingThreadShaderMap())
 		{
 			const EShaderPlatform ShaderPlatform = ViewIfDynamicMeshCommand->GetShaderPlatform();
-			const FDecalBlendDesc DecalBlendDesc = DecalRendering::ComputeDecalBlendDesc(ShaderPlatform, &Material);
+			const FDecalBlendDesc DecalBlendDesc = DecalRendering::ComputeDecalBlendDesc(ShaderPlatform, Material);
 
 			const bool bShouldRender =
 				DecalRendering::IsCompatibleWithRenderStage(DecalBlendDesc, PassDecalStage) &&
diff --git a/Engine/Source/Runtime/Renderer/Private/DecalRenderingCommon.cpp b/Engine/Source/Runtime/Renderer/Private/DecalRenderingCommon.cpp
--- a/Engine/Source/Runtime/Renderer/Private/DecalRenderingCommon.cpp
+++ b/Engine/Source/Runtime/Renderer/Private/DecalRenderingCommon.cpp
@@ -72,38 +72,38 @@ namespace DecalRendering
 		}
 	}
 
-	FDecalBlendDesc ComputeDecalBlendDesc(EShaderPlatform Platform, FMaterial const* Material)
+	FDecalBlendDesc ComputeDecalBlendDesc(EShaderPlatform Platform, const FMaterial& Material)
 	{
 		FDecalBlendDesc Desc;
 		if (Strata::IsStrataEnabled())
 		{
-			check(Material->IsStrataMaterial());
+			check(Material.IsStrataMaterial());
 
 			const bool bUseDiffuseAlbedoAndF0 =
-				Material->HasMaterialPropertyConnected(EMaterialProperty::MP_DiffuseColor) ||	// This is used for Strata Slab using (DiffuseAlbedo | F0) parameterization
-				Material->HasMaterialPropertyConnected(EMaterialProperty::MP_SpecularColor);	// This is used for Strata Slab using (DiffuseAlbedo | F0) parameterization
+				Material.HasMaterialPropertyConnected(EMaterialProperty::MP_DiffuseColor) ||	// This is used for Strata Slab using (DiffuseAlbedo | F0) parameterization
+				Material.HasMaterialPropertyConnected(EMaterialProperty::MP_SpecularColor);	// This is used for Strata Slab using (DiffuseAlbedo | F0) parameterization
 
-			Desc.BlendMode = Material->GetBlendMode();
-			Desc.bWriteBaseColor = Material->HasMaterialPropertyConnected(EMaterialProperty::MP_BaseColor) || bUseDiffuseAlbedoAndF0;
-			Desc.bWriteNormal = Material->HasMaterialPropertyConnected(EMaterialProperty::MP_Normal);
+			Desc.BlendMode = Material.GetBlendMode();
+			Desc.bWriteBaseColor = Material.HasMaterialPropertyConnected(EMaterialProperty::MP_BaseColor) || bUseDiffuseAlbedoAndF0;
+			Desc.bWriteNormal = Material.HasMaterialPropertyConnected(EMaterialProperty::MP_Normal);
 			Desc.bWriteRoughnessSpecularMetallic =
 				bUseDiffuseAlbedoAndF0 ||
-				Material->HasMaterialPropertyConnected(EMaterialProperty::MP_Metallic) ||
-				Material->HasMaterialPropertyConnected(EMaterialProperty::MP_Specular) ||
-				Material->HasMaterialPropertyConnected(EMaterialProperty::MP_Roughness);
+				Material.HasMaterialPropertyConnected(EMaterialProperty::MP_Metallic) ||
+				Material.HasMaterialPropertyConnected(EMaterialProperty::MP_Specular) ||
+				Material.HasMaterialPropertyConnected(EMaterialProperty::MP_Roughness);
 			Desc.bWriteEmissive=
-				Material->HasMaterialPropertyConnected(EMaterialProperty::MP_EmissiveColor);
+				Material.HasMaterialPropertyConnected(EMaterialProperty::MP_EmissiveColor);
 			Desc.bWriteAmbientOcclusion =
-				Material->HasMaterialPropertyConnected(EMaterialProperty::MP_AmbientOcclusion);
+				Material.HasMaterialPropertyConnected(EMaterialProperty::MP_AmbientOcclusion);
 		}
 		else
 		{
-			Desc.BlendMode = Material->GetBlendMode();
-			Desc.bWriteBaseColor = Material->HasBaseColorConnected();
-			Desc.bWriteNormal = Material->HasNormalConnected();
-			Desc.bWriteRoughnessSpecularMetallic = Material->HasRoughnessConnected() || Material->HasSpecularConnected() || Material->HasMetallicConnected();
-			Desc.bWriteEmissive = Material->HasEmissiveColorConnected();
-			Desc.bWriteAmbientOcclusion = Material->HasAmbientOcclusionConnected();
+			Desc.BlendMode = Material.GetBlendMode();
+			Desc.bWriteBaseColor = Material.HasBaseColorConnected();
+			Desc.bWriteNormal = Material.HasNormalConnected();
+			Desc.bWriteRoughnessSpecularMetallic = Material.HasRoughnessConnected() || Material.HasSpecularConnected() || Material.HasMetallicConnected();
+			Desc.bWriteEmissive = Material.HasEmissiveColorConnected();
+			Desc.bWriteAmbientOcclusion = Material.HasAmbientOcclusionConnected();
 		}
 		FinalizeBlendDesc(Platform, Desc);
 		return Desc;
diff --git a/Engine/Source/Runtime/Renderer/Private/DecalRenderingCommon.h b/Engine/Source/Runtime/Renderer/Private/DecalRenderingCommon.h
--- a/Engine/Source/Runtime/Renderer/Private/DecalRenderingCommon.h
+++ b/Engine/Source/Runtime/Renderer/Private/DecalRenderingCommon.h
@@ -78,7 +78,7 @@ enum class EDecalRasterizerState : uint8
 namespace DecalRendering
 {
 	/** Build the packed decal description from a decal material. */
-	FDecalBlendDesc ComputeDecalBlendDesc(EShaderPlatform Platform, FMaterial const* Material);
+	FDecalBlendDesc ComputeDecalBlendDesc(EShaderPlatform Platform, const FMaterial& Material);
 
 	/** Build the packed decal description from a decal material. */
 	FDecalBlendDesc ComputeDecalBlendDesc(EShaderPlatform Platform, FMaterialShaderParameters const& MaterialShaderParameters);
diff --git a/Engine/Source/Runtime/Renderer/Private/DecalRenderingShared.cpp b/Engine/Source/Runtime/Renderer/Private/DecalRenderingShared.cpp
--- a/Engine/Source/Runtime/Renderer/Private/DecalRenderingShared.cpp
+++ b/Engine/Source/Runtime/Renderer/Private/DecalRenderingShared.cpp
@@ -29,7 +29,7 @@ FTransientDecalRenderData::FTransientDecalRenderData(const FScene& InScene, cons
 	MaterialProxy = InDecalProxy->DecalMaterial->GetRenderProxy();
 	MaterialResource = &MaterialProxy->GetMaterialWithFallback(InScene.GetFeatureLevel(), MaterialProxy);
 	check(MaterialProxy && MaterialResource);
-	DecalBlendDesc = DecalRendering::ComputeDecalBlendDesc(InScene.GetShaderPlatform(), MaterialResource);
+	DecalBlendDesc = DecalRendering::ComputeDecalBlendDesc(InScene.GetShaderPlatform(), *MaterialResource);
 }
 
 /**
@@ -209,6 +209,30 @@ IMPLEMENT_MATERIAL_SHADER_TYPE(, FDeferredDecalAmbientOcclusionPS, TEXT("/Engine
 
 namespace DecalRendering
 {
+	float GetDecalFadeScreenSizeMultiplier()
+	{
+		return CVarDecalFadeScreenSizeMultiplier.GetValueOnRenderThread();
+	}
+
+	float CalculateDecalFadeAlpha(float DecalFadeScreenSize, const FMatrix& ComponentToWorldMatrix, const FViewInfo& View, float FadeMultiplier)
+	{
+		check(View.IsPerspectiveProjection());
+
+		float Distance = (View.ViewMatrices.GetViewOrigin() - ComponentToWorldMatrix.GetOrigin()).Size();
+		float Radius = ComponentToWorldMatrix.GetMaximumAxisScale();
+		float CurrentScreenSize = ((Radius / Distance) * FadeMultiplier);
+
+		// fading coefficient needs to increase with increasing field of view and decrease with increasing resolution
+		// FadeCoeffScale is an empirically determined constant to bring us back roughly to fraction of screen size for FadeScreenSize
+		const float FadeCoeffScale = 600.0f;
+		float FOVFactor = ((2.0f / View.ViewMatrices.GetProjectionMatrix().M[0][0]) / View.ViewRect.Width()) * FadeCoeffScale;
+		float FadeCoeff = DecalFadeScreenSize * FOVFactor;
+		float FadeRange = FadeCoeff * 0.5f;
+
+		float Alpha = (CurrentScreenSize - FadeCoeff) / FadeRange;
+		return FMath::Clamp(Alpha, 0.0f, 1.0f);
+	}
+
 	bool BuildVisibleDecalList(const FScene& Scene, const FViewInfo& View, EDecalRenderStage DecalRenderStage, FTransientDecalRenderDataList* OutVisibleDecals)
 	{
 		QUICK_SCOPE_CYCLE_COUNTER(BuildVisibleDecalList);
@@ -225,7 +249,7 @@ namespace DecalRendering
 			return false;
 		}
 
-		const float FadeMultiplier = CVarDecalFadeScreenSizeMultiplier.GetValueOnRenderThread();
+		const float FadeMultiplier = GetDecalFadeScreenSizeMultiplier();
 		const EShaderPlatform ShaderPlatform = View.GetShaderPlatform();
 
 		const bool bIsPerspectiveProjection = View.IsPerspectiveProjection();
@@ -267,19 +291,7 @@ namespace DecalRendering
 				{
 					if (bIsPerspectiveProjection && Data.DecalProxy->FadeScreenSize != 0.0f)
 					{
-						float Distance = (View.ViewMatrices.GetViewOrigin() - ComponentToWorldMatrix.GetOrigin()).Size();
-						float Radius = ComponentToWorldMatrix.GetMaximumAxisScale();
-						float CurrentScreenSize = ((Radius / Distance) * FadeMultiplier);
-
-						// fading coefficient needs to increase with increasing field of view and decrease with increasing resolution
-						// FadeCoeffScale is an empirically determined constant to bring us back roughly to fraction of screen size for FadeScreenSize
-						const float FadeCoeffScale = 600.0f;
-						float FOVFactor = ((2.0f/View.ViewMatrices.GetProjectionMatrix().M[0][0]) / View.ViewRect.Width()) * FadeCoeffScale;
-						float FadeCoeff = Data.DecalProxy->FadeScreenSize * FOVFactor;
-						float FadeRange = FadeCoeff * 0.5f;
-
-						float Alpha = (CurrentScreenSize - FadeCoeff) / FadeRange;
-						Data.FadeAlpha = FMath::Min(Alpha, 1.0f);
+						Data.FadeAlpha = CalculateDecalFadeAlpha(Data.DecalProxy->FadeScreenSize, ComponentToWorldMatrix, View, FadeMultiplier);
 					}
 
 					const bool bShouldRender = Data.FadeAlpha > 0.0f;
diff --git a/Engine/Source/Runtime/Renderer/Private/DecalRenderingShared.h b/Engine/Source/Runtime/Renderer/Private/DecalRenderingShared.h
--- a/Engine/Source/Runtime/Renderer/Private/DecalRenderingShared.h
+++ b/Engine/Source/Runtime/Renderer/Private/DecalRenderingShared.h
@@ -34,6 +34,8 @@ typedef TArray<FTransientDecalRenderData, SceneRenderingAllocator> FTransientDec
  */
 namespace DecalRendering
 {
+	float GetDecalFadeScreenSizeMultiplier();
+	float CalculateDecalFadeAlpha(float DecalFadeScreenSize, const FMatrix& ComponentToWorldMatrix, const FViewInfo& View, float FadeMultiplier);
 	FMatrix ComputeComponentToClipMatrix(const FViewInfo& View, const FMatrix& DecalComponentToWorld);
 	void SetVertexShaderOnly(FRHICommandList& RHICmdList, FGraphicsPipelineStateInitializer& GraphicsPSOInit, const FViewInfo& View, const FMatrix& FrustumComponentToClip);
 	bool BuildVisibleDecalList(const FScene& Scene, const FViewInfo& View, EDecalRenderStage DecalRenderStage, FTransientDecalRenderDataList* OutVisibleDecals);
