Fix error creating RTPSO in platforms that don't support callable shaders.

diff --git a/Engine/Source/Runtime/Renderer/Private/RayTracing/RayTracingMaterialHitShaders.cpp b/Engine/Source/Runtime/Renderer/Private/RayTracing/RayTracingMaterialHitShaders.cpp
--- a/Engine/Source/Runtime/Renderer/Private/RayTracing/RayTracingMaterialHitShaders.cpp
+++ b/Engine/Source/Runtime/Renderer/Private/RayTracing/RayTracingMaterialHitShaders.cpp
@@ -642,15 +642,26 @@ FRayTracingPipelineState* FDeferredShadingSceneRenderer::BindRayTracingMaterialP
 
 	Initializer.SetHitGroupTable(RayTracingHitGroupLibrary);
 
+	const bool bCallableShadersSupported = RHISupportsRayTracingCallableShaders(View.Family->GetShaderPlatform());
 	TArray<FRHIRayTracingShader*> RayTracingCallableShaderLibrary;
+	FRHIRayTracingShader* DefaultCallableShader = nullptr;
 
-	if (bEnableMaterials)
+	if (bCallableShadersSupported)
 	{
-		FRHIRayTracingShader* DefaultCallableShader = View.ShaderMap->GetShader<FDefaultCallableShader>().GetRayTracingShader();
-		FShaderMapResource::GetRayTracingCallableShaderLibrary(RayTracingCallableShaderLibrary, DefaultCallableShader);
-	}
+		DefaultCallableShader = View.ShaderMap->GetShader<FDefaultCallableShader>().GetRayTracingShader();
+		check(DefaultCallableShader != nullptr);
+
+		if (bEnableMaterials)
+		{
+			FShaderMapResource::GetRayTracingCallableShaderLibrary(RayTracingCallableShaderLibrary, DefaultCallableShader);
+		}
+		else
+		{
+			RayTracingCallableShaderLibrary.Add(DefaultCallableShader);
+		}
 
-	Initializer.SetCallableTable(RayTracingCallableShaderLibrary);
+		Initializer.SetCallableTable(RayTracingCallableShaderLibrary);
+	}
 
 	FRayTracingPipelineState* FallbackPipelineState = GRayTracingNonBlockingPipelineCreation && View.ViewState
 		? PipelineStateCache::GetRayTracingPipelineState(View.ViewState->LastRayTracingMaterialPipelineSignature)
@@ -661,7 +672,8 @@ FRayTracingPipelineState* FDeferredShadingSceneRenderer::BindRayTracingMaterialP
 		&& FallbackPipelineState
 		&& View.ViewState
 		&& IsCompatibleFallbackPipelineSignature(View.ViewState->LastRayTracingMaterialPipelineSignature, Initializer)
-		&& PipelineContainsHitShaders(FallbackPipelineState, RequiredHitShaders))
+		&& PipelineContainsHitShaders(FallbackPipelineState, RequiredHitShaders)
+		&& (!bCallableShadersSupported || FindRayTracingCallableShaderIndex(FallbackPipelineState, DefaultCallableShader, false)))
 	{
 		PipelineCacheFlags |= ERayTracingPipelineCacheFlags::NonBlocking;
 	}
