DisplayCluster:  Virtual Shadow Map cache can optionally be allocated per view, and enabled that feature for Virtual Production.  Significant performance win by avoiding constant cache thrashing when rendering multiple families.  Frame time goes from ~45 ms to ~26 ms for Valley test scene (Node 2), GPU bound in both cases.  Separate cache per view isn't implemented for split screen views (single view family with multiple views), but Virtual Production doesn't use that, and support for that could be added in the future.

diff --git a/Engine/Plugins/Runtime/nDisplay/Source/DisplayCluster/Private/Game/EngineClasses/Basics/DisplayClusterViewportClient.cpp b/Engine/Plugins/Runtime/nDisplay/Source/DisplayCluster/Private/Game/EngineClasses/Basics/DisplayClusterViewportClient.cpp
--- a/Engine/Plugins/Runtime/nDisplay/Source/DisplayCluster/Private/Game/EngineClasses/Basics/DisplayClusterViewportClient.cpp
+++ b/Engine/Plugins/Runtime/nDisplay/Source/DisplayCluster/Private/Game/EngineClasses/Basics/DisplayClusterViewportClient.cpp
@@ -587,6 +587,8 @@ void UDisplayClusterViewportClient::Draw(FViewport* InViewport, FCanvas* SceneCa
 
 					View->CameraConstrainedViewRect = View->UnscaledViewRect;
 
+					// Enable per-view virtual shadow map caching
+					View->State->AddVirtualShadowMapCache();
 
 					{
 						// Save the location of the view.
diff --git a/Engine/Plugins/VirtualProduction/Switchboard/Source/Switchboard/switchboard/devices/ndisplay/plugin_ndisplay.py b/Engine/Plugins/VirtualProduction/Switchboard/Source/Switchboard/switchboard/devices/ndisplay/plugin_ndisplay.py
--- a/Engine/Plugins/VirtualProduction/Switchboard/Source/Switchboard/switchboard/devices/ndisplay/plugin_ndisplay.py
+++ b/Engine/Plugins/VirtualProduction/Switchboard/Source/Switchboard/switchboard/devices/ndisplay/plugin_ndisplay.py
@@ -899,9 +899,6 @@ class DevicenDisplay(DeviceUnreal):
         # Always tell Chaos to be deterministic
         dp_cvars.append('p.Chaos.Solver.Deterministic=1')
 
-        # Always disable virtual shadow map caching until it supports multi-viewfamily.
-        dp_cvars.append('r.Shadow.Virtual.Cache=0')
-
         # Add user set dp cvars, overriding any of the forced ones.
         user_dp_cvars = self.csettings['ndisplay_dp_cvars'].get_value(self.name)
         user_dp_cvars = [cvar.strip() for cvar in user_dp_cvars if len(cvar.strip()) and len(cvar.split('=')) == 2]
diff --git a/Engine/Source/Runtime/Engine/Private/LocalPlayer.cpp b/Engine/Source/Runtime/Engine/Private/LocalPlayer.cpp
--- a/Engine/Source/Runtime/Engine/Private/LocalPlayer.cpp
+++ b/Engine/Source/Runtime/Engine/Private/LocalPlayer.cpp
@@ -803,7 +803,7 @@ bool ULocalPlayer::CalcSceneViewInitOptions(
 		const UWorld* CurrentWorld = GetWorld();
 		const ERHIFeatureLevel::Type FeatureLevel = CurrentWorld ? CurrentWorld->FeatureLevel.GetValue() : GMaxRHIFeatureLevel;
 
-		ViewStates[ViewIndex].Allocate(FeatureLevel);
+		ViewStates[ViewIndex].Allocate(FeatureLevel, CurrentWorld);
 	}
 
 	ViewInitOptions.SceneViewStateInterface = ViewStates[ViewIndex].GetReference();
diff --git a/Engine/Source/Runtime/Engine/Private/RendererSupport.cpp b/Engine/Source/Runtime/Engine/Private/RendererSupport.cpp
--- a/Engine/Source/Runtime/Engine/Private/RendererSupport.cpp
+++ b/Engine/Source/Runtime/Engine/Private/RendererSupport.cpp
@@ -42,6 +42,9 @@ static void ClearReferencesToRendererModuleClasses(
 	TMap<const FShaderPipelineType*, FHashedName>& ShaderPipelineTypeNames,
 	TMap<FVertexFactoryType*, FHashedName>& VertexFactoryTypeNames)
 {
+	// Destroy scene view states -- must be called before destroying scenes, as scenes may have references to view states
+	FSceneViewStateReference::DestroyAll();
+
 	// Destroy all renderer scenes
 	for (TObjectIterator<UWorld> WorldIt; WorldIt; ++WorldIt)
 	{
@@ -93,7 +96,6 @@ static void ClearReferencesToRendererModuleClasses(
 	}
 
 	// Destroy misc renderer module classes and remove references
-	FSceneViewStateReference::DestroyAll();
 	FSlateApplication::Get().InvalidateAllViewports();
 
 	// Invalidate cached shader type data
diff --git a/Engine/Source/Runtime/Engine/Private/SceneView.cpp b/Engine/Source/Runtime/Engine/Private/SceneView.cpp
--- a/Engine/Source/Runtime/Engine/Private/SceneView.cpp
+++ b/Engine/Source/Runtime/Engine/Private/SceneView.cpp
@@ -347,17 +347,32 @@ FSceneViewStateReference::~FSceneViewStateReference()
 	Destroy();
 }
 
-void FSceneViewStateReference::Allocate(ERHIFeatureLevel::Type FeatureLevel)
+void FSceneViewStateReference::Allocate(ERHIFeatureLevel::Type FeatureLevel, const UWorld* InWorld)
 {
 	check(!Reference);
-	Reference = GetRendererModule().AllocateViewState(FeatureLevel);
+	World = InWorld;
+
+	if (World && World->Scene)
+	{
+		Reference = World->Scene->AllocateViewState();
+	}
+	else
+	{
+		Reference = GetRendererModule().AllocateViewState(FeatureLevel);
+	}
+
 	GlobalListLink = TLinkedList<FSceneViewStateReference*>(this);
 	GlobalListLink.LinkHead(GetSceneViewStateList());
 }
 
+void FSceneViewStateReference::Allocate(ERHIFeatureLevel::Type FeatureLevel)
+{
+	Allocate(FeatureLevel, nullptr);
+}
+
 void FSceneViewStateReference::Allocate()
 {
-	Allocate(GMaxRHIFeatureLevel);
+	Allocate(GMaxRHIFeatureLevel, nullptr);
 }
 
 void FSceneViewStateReference::Destroy()
@@ -386,7 +401,14 @@ void FSceneViewStateReference::AllocateAll(ERHIFeatureLevel::Type FeatureLevel)
 	for(TLinkedList<FSceneViewStateReference*>::TIterator ViewStateIt(FSceneViewStateReference::GetSceneViewStateList());ViewStateIt;ViewStateIt.Next())
 	{
 		FSceneViewStateReference* ViewStateReference = *ViewStateIt;
-		ViewStateReference->Reference = GetRendererModule().AllocateViewState(FeatureLevel);
+		if (ViewStateReference->World && ViewStateReference->World->Scene)
+		{
+			ViewStateReference->Reference = ViewStateReference->World->Scene->AllocateViewState();
+		}
+		else
+		{
+			ViewStateReference->Reference = GetRendererModule().AllocateViewState(FeatureLevel);
+		}
 	}
 }
 
diff --git a/Engine/Source/Runtime/Engine/Public/SceneInterface.h b/Engine/Source/Runtime/Engine/Public/SceneInterface.h
--- a/Engine/Source/Runtime/Engine/Public/SceneInterface.h
+++ b/Engine/Source/Runtime/Engine/Public/SceneInterface.h
@@ -35,6 +35,7 @@ class UTextureCube;
 class FViewInfo;
 class FSceneRenderer;
 class FInstanceCullingManager;
+class FSceneViewStateInterface;
 struct FHairStrandsInstance;
 struct FLightRenderParameters;
 
@@ -564,8 +565,12 @@ public:
 	virtual bool RequestGPUSceneUpdate(FPrimitiveSceneInfo& PrimitiveSceneInfo, EPrimitiveDirtyState PrimitiveDirtyState) { return false; }
 
 protected:
+	virtual FSceneViewStateInterface* AllocateViewState() = 0;
+
 	virtual ~FSceneInterface() {}
 
 	/** This scene's feature level */
 	ERHIFeatureLevel::Type FeatureLevel;
+
+	friend class FSceneViewStateReference;
 };
diff --git a/Engine/Source/Runtime/Engine/Public/SceneManagement.h b/Engine/Source/Runtime/Engine/Public/SceneManagement.h
--- a/Engine/Source/Runtime/Engine/Public/SceneManagement.h
+++ b/Engine/Source/Runtime/Engine/Public/SceneManagement.h
@@ -51,6 +51,7 @@ class USkyLightComponent;
 struct FDynamicMeshVertex;
 class ULightMapVirtualTexture2D;
 class FGPUScenePrimitiveCollector;
+class FVirtualShadowMapArrayCacheManager;
 
 DECLARE_LOG_CATEGORY_EXTERN(LogBufferVisualization, Log, All);
 DECLARE_LOG_CATEGORY_EXTERN(LogNaniteVisualization, Log, All);
@@ -253,6 +254,12 @@ public:
 	virtual uint32 GetPathTracingSampleCount() const = 0;
 #endif
 
+	/**
+	* Adds a per-view virtual shadow map cache, which can help performance, at a cost in memory.  Does nothing if one is already present.
+	* Requires the View State to have been originally allocated with a UWorld with a valid Scene pointer.
+	*/
+	virtual void AddVirtualShadowMapCache() {}
+
 protected:
 	// Don't allow direct deletion of the view state, Destroy should be called instead.
 	virtual ~FSceneViewStateInterface() {}
@@ -265,6 +272,9 @@ private:
 	FSceneViewStateInterface*	ViewParent;
 	/** Reference counts the number of children parented to this state. */
 	int32							NumChildren;
+
+	virtual FVirtualShadowMapArrayCacheManager* GetVirtualShadowMapCache() const { return nullptr; }
+	friend class FScene;
 };
 
 class FFrozenSceneViewMatricesGuard
diff --git a/Engine/Source/Runtime/Engine/Public/SceneTypes.h b/Engine/Source/Runtime/Engine/Public/SceneTypes.h
--- a/Engine/Source/Runtime/Engine/Public/SceneTypes.h
+++ b/Engine/Source/Runtime/Engine/Public/SceneTypes.h
@@ -98,7 +98,13 @@ public:
 
 	ENGINE_API virtual ~FSceneViewStateReference();
 
-	/** Allocates the Scene view state. */
+	/**
+	 * Allocates the Scene view state.  The first Allocate function accepts a UWorld, rather than an FScene, as the "RecompileRenderer" feature
+	 * involves destroying and reconstructing the FScene.  We need to pull the updated scene pointer from the UWorld in that case, but nothing
+	 * else in the UWorld is accessed.
+	 */
+	ENGINE_API void Allocate(ERHIFeatureLevel::Type FeatureLevel, const UWorld* World);
+
 	ENGINE_API void Allocate(ERHIFeatureLevel::Type FeatureLevel);
 
 	UE_DEPRECATED(5.0, "Allocate must be called with an appropriate RHI Feature Level")
@@ -122,6 +128,7 @@ public:
 	}
 
 private:
+	const UWorld* World;
 	FSceneViewStateInterface* Reference;
 	TLinkedList<FSceneViewStateReference*> GlobalListLink;
 
diff --git a/Engine/Source/Runtime/Renderer/Private/DeferredShadingRenderer.cpp b/Engine/Source/Runtime/Renderer/Private/DeferredShadingRenderer.cpp
--- a/Engine/Source/Runtime/Renderer/Private/DeferredShadingRenderer.cpp
+++ b/Engine/Source/Runtime/Renderer/Private/DeferredShadingRenderer.cpp
@@ -2085,7 +2085,8 @@ void FDeferredShadingSceneRenderer::Render(FRDGBuilder& GraphBuilder)
 
 	// Important that this uses consistent logic throughout the frame, so evaluate once and pass in the flag from here
 	// NOTE: Must be done after  system texture initialization
-	ActiveViewFamily->VirtualShadowMapArray.Initialize(GraphBuilder, Scene->VirtualShadowMapArrayCacheManager, UseVirtualShadowMaps(ShaderPlatform, FeatureLevel));
+	// TODO: This doesn't take into account the potential for split screen views with separate shadow caches
+	ActiveViewFamily->VirtualShadowMapArray.Initialize(GraphBuilder, Scene->GetVirtualShadowMapCache(Views[0]), UseVirtualShadowMaps(ShaderPlatform, FeatureLevel));
 
 	// if DDM_AllOpaqueNoVelocity was used, then velocity should have already been rendered as well
 	const bool bIsEarlyDepthComplete = (DepthPass.EarlyZPassMode == DDM_AllOpaque || DepthPass.EarlyZPassMode == DDM_AllOpaqueNoVelocity);
@@ -2826,14 +2827,17 @@ void FDeferredShadingSceneRenderer::Render(FRDGBuilder& GraphBuilder)
 		Nanite::GStreamingManager.SubmitFrameStreamingRequests(GraphBuilder);
 	}
 
-	if (Scene->VirtualShadowMapArrayCacheManager)
 	{
-		// Do this even if VSMs are disabled this frame to clean up any previously extracted data
-		Scene->VirtualShadowMapArrayCacheManager->ExtractFrameData(
-			GraphBuilder,				
-			ActiveViewFamily->VirtualShadowMapArray,
-			*this,
-			ActiveViewFamily->EngineShowFlags.VirtualShadowMapCaching);
+		FVirtualShadowMapArrayCacheManager* CacheManager = ActiveViewFamily->VirtualShadowMapArray.CacheManager;
+		if (CacheManager)
+		{
+			// Do this even if VSMs are disabled this frame to clean up any previously extracted data
+			CacheManager->ExtractFrameData(
+				GraphBuilder,
+				ActiveViewFamily->VirtualShadowMapArray,
+				*this,
+				ActiveViewFamily->EngineShowFlags.VirtualShadowMapCaching);
+		}
 	}
 
 	// If not all depth is written during the prepass, kick off async compute cloud after basepass
diff --git a/Engine/Source/Runtime/Renderer/Private/GPUScene.cpp b/Engine/Source/Runtime/Renderer/Private/GPUScene.cpp
--- a/Engine/Source/Runtime/Renderer/Private/GPUScene.cpp
+++ b/Engine/Source/Runtime/Renderer/Private/GPUScene.cpp
@@ -1421,7 +1421,7 @@ void FGPUScene::UploadDynamicPrimitiveShaderDataForViewInternal(FRDGBuilder& Gra
 		ensure(UploadIdStart < DynamicPrimitivesOffset);
 		ensure(InstanceIdStart != INDEX_NONE);
 
-		if (bIsShadowView && Scene != nullptr && Scene->VirtualShadowMapArrayCacheManager != nullptr)
+		if (bIsShadowView && Scene != nullptr && Scene->GetVirtualShadowMapCache(View) != nullptr)
 		{
 			// Enqueue cache invalidations for all dynamic primitives' instances, as they will be removed this frame and are not associated
 			// with any particular FPrimitiveSceneInfo. Will occur on the next call to UpdateAllPrimitiveSceneInfos
diff --git a/Engine/Source/Runtime/Renderer/Private/LightShaftRendering.cpp b/Engine/Source/Runtime/Renderer/Private/LightShaftRendering.cpp
--- a/Engine/Source/Runtime/Renderer/Private/LightShaftRendering.cpp
+++ b/Engine/Source/Runtime/Renderer/Private/LightShaftRendering.cpp
@@ -651,13 +651,14 @@ void FDeferredShadingSceneRenderer::RenderLightShaftBloom(
 	}
 }
 
-void FSceneViewState::TrimHistoryRenderTargets(const FScene* Scene)
+// InScene is passed in, as the Scene pointer in the class itself may be null, if it was allocated without a scene.
+void FSceneViewState::TrimHistoryRenderTargets(const FScene* InScene)
 {
 	for (TMap<const ULightComponent*, TUniquePtr<FTemporalAAHistory>>::TIterator It(LightShaftBloomHistoryRTs); It; ++It)
 	{
 		bool bLightIsUsed = false;
 
-		for (auto LightIt = Scene->Lights.CreateConstIterator(); LightIt; ++LightIt)
+		for (auto LightIt = InScene->Lights.CreateConstIterator(); LightIt; ++LightIt)
 		{
 			const FLightSceneInfo* const LightSceneInfo = LightIt->LightSceneInfo;
 
diff --git a/Engine/Source/Runtime/Renderer/Private/MobileShadingRenderer.cpp b/Engine/Source/Runtime/Renderer/Private/MobileShadingRenderer.cpp
--- a/Engine/Source/Runtime/Renderer/Private/MobileShadingRenderer.cpp
+++ b/Engine/Source/Runtime/Renderer/Private/MobileShadingRenderer.cpp
@@ -775,7 +775,8 @@ void FMobileSceneRenderer::Render(FRDGBuilder& GraphBuilder)
 
 	// Important that this uses consistent logic throughout the frame, so evaluate once and pass in the flag from here
 	// NOTE: Must be done after  system texture initialization
-	ActiveViewFamily->VirtualShadowMapArray.Initialize(GraphBuilder, Scene->VirtualShadowMapArrayCacheManager, UseVirtualShadowMaps(ShaderPlatform, FeatureLevel));
+	// TODO: This doesn't take into account the potential for split screen views with separate shadow caches
+	ActiveViewFamily->VirtualShadowMapArray.Initialize(GraphBuilder, Scene->GetVirtualShadowMapCache(Views[0]), UseVirtualShadowMaps(ShaderPlatform, FeatureLevel));
 
 	// Find the visible primitives and prepare targets and buffers for rendering
 	InitViews(GraphBuilder, SceneTexturesConfig, InstanceCullingManager);
diff --git a/Engine/Source/Runtime/Renderer/Private/RendererScene.cpp b/Engine/Source/Runtime/Renderer/Private/RendererScene.cpp
--- a/Engine/Source/Runtime/Renderer/Private/RendererScene.cpp
+++ b/Engine/Source/Runtime/Renderer/Private/RendererScene.cpp
@@ -318,6 +318,7 @@ FSceneViewState::FSceneViewState(ERHIFeatureLevel::Type FeatureLevel)
 	SetFeatureLevel(FeatureLevel);
 	
 	UniqueID = FSceneViewState_UniqueID.Increment();
+	Scene = nullptr;
 	OcclusionFrameCounter = 0;
 	LastRenderTime = -FLT_MAX;
 	LastRenderTimeDelta = 0.0f;
@@ -383,6 +384,8 @@ FSceneViewState::FSceneViewState(ERHIFeatureLevel::Type FeatureLevel)
 	GatherPointsBuffer = nullptr;
 	GatherPointsResolution = FIntVector(0, 0, 0);
 #endif
+
+	ViewVirtualShadowMapCache = nullptr;
 }
 
 void DestroyRenderResource(FRenderResource* RenderResource)
@@ -419,8 +422,40 @@ FSceneViewState::~FSceneViewState()
 
 	HairStrandsViewStateData.Release();
 	ShaderPrintStateData.Release();
+
+	if (ViewVirtualShadowMapCache)
+	{
+		delete ViewVirtualShadowMapCache;
+		ViewVirtualShadowMapCache = nullptr;
+	}
+
+	if (Scene)
+	{
+		Scene->RemoveViewState(this);
+	}
 }
 
+FSceneViewStateInterface* FScene::AllocateViewState()
+{
+	FSceneViewState* Result = new FSceneViewState(FeatureLevel);
+	Result->Scene = this;
+	ViewStates.Add(Result);
+	return Result;
+}
+
+void FScene::RemoveViewState(FSceneViewStateInterface* ViewState)
+{
+	for (int32 ViewStateIndex = 0; ViewStateIndex < ViewStates.Num(); ViewStateIndex++)
+	{
+		if (ViewStates[ViewStateIndex] == ViewState)
+		{
+			ViewStates.RemoveAt(ViewStateIndex);
+			break;
+		}
+	}
+}
+
+
 #if WITH_EDITOR
 
 FPixelInspectorData::FPixelInspectorData()
@@ -718,6 +753,61 @@ FFXSystemInterface* FScene::GetFXSystem()
 	return FXSystem;
 }
 
+void FSceneViewState::AddVirtualShadowMapCache()
+{
+	// Can only allocate a virtual shadow map cache if we have a scene pointer available
+	if (Scene)
+	{
+		// Don't allocate if one already exists
+		if (ViewVirtualShadowMapCache == nullptr)
+		{
+			FVirtualShadowMapArrayCacheManager* ViewCache = new FVirtualShadowMapArrayCacheManager(Scene);
+
+			// Need to add reference to virtual shadow map cache in render thread
+			ENQUEUE_RENDER_COMMAND(LinkVirtualShadowMapCache)(
+				[this, ViewCache](FRHICommandListImmediate& RHICmdList)
+				{
+					this->ViewVirtualShadowMapCache = ViewCache;
+				});
+		}
+	}
+}
+
+FVirtualShadowMapArrayCacheManager* FSceneViewState::GetVirtualShadowMapCache() const
+{
+	return ViewVirtualShadowMapCache;
+}
+
+FVirtualShadowMapArrayCacheManager* FScene::GetVirtualShadowMapCache(FSceneView& View) const
+{
+	FVirtualShadowMapArrayCacheManager* Result = DefaultVirtualShadowMapCache;
+	if (View.State)
+	{
+		FVirtualShadowMapArrayCacheManager* ViewCache = View.State->GetVirtualShadowMapCache();
+		if (ViewCache)
+		{
+			Result = ViewCache;
+		}
+	}
+	return Result;
+}
+
+void FScene::GetAllVirtualShadowMapCacheManagers(TArray<FVirtualShadowMapArrayCacheManager*, SceneRenderingAllocator>& OutCacheManagers) const
+{
+	OutCacheManagers.Empty();
+	if (DefaultVirtualShadowMapCache)
+	{
+		OutCacheManagers.Add(DefaultVirtualShadowMapCache);
+	}
+	for (const FSceneViewState* ViewState : ViewStates)
+	{
+		if (ViewState->ViewVirtualShadowMapCache)
+		{
+			OutCacheManagers.Add(ViewState->ViewVirtualShadowMapCache);
+		}
+	}
+}
+
 void FScene::UpdateParameterCollections(const TArray<FMaterialParameterCollectionInstanceResource*>& InParameterCollections)
 {
 	ENQUEUE_RENDER_COMMAND(UpdateParameterCollectionsCommand)(
@@ -1262,7 +1352,9 @@ FScene::FScene(UWorld* InWorld, bool bInRequiresHitProxies, bool bInIsEditorScen
 
 	LumenSceneData = new FLumenSceneData(GShaderPlatformForFeatureLevel[InFeatureLevel], InWorld->WorldType);
 
-	VirtualShadowMapArrayCacheManager = new FVirtualShadowMapArrayCacheManager(this);
+	// We use a default Virtual Shadow Map cache, if one hasn't been allocated for a specific view.  GPU resources for
+	// a cache aren't allocated until the cache is actually used, so this shouldn't waste any GPU memory when not in use.
+	DefaultVirtualShadowMapCache = new FVirtualShadowMapArrayCacheManager(this);
 }
 
 FScene::~FScene()
@@ -1281,11 +1373,13 @@ FScene::~FScene()
 	checkf(RemovedPrimitiveSceneInfos.Num() == 0, TEXT("All pending primitive removal operations are expected to be flushed when the scene is destroyed. Remaining operations are likely to cause a memory leak."));
 	checkf(AddedPrimitiveSceneInfos.Num() == 0, TEXT("All pending primitive addition operations are expected to be flushed when the scene is destroyed. Remaining operations are likely to cause a memory leak."));
 	checkf(Primitives.Num() == 0, TEXT("All primitives are expected to be removed before the scene is destroyed. Remaining primitives are likely to cause a memory leak."));
+	checkf(ViewStates.Num() == 0, TEXT("All scene view states are expected to be removed before the scene is destroyed."));
 
-	if (VirtualShadowMapArrayCacheManager)
+	// Delete default cache
+	if (DefaultVirtualShadowMapCache)
 	{
-		delete VirtualShadowMapArrayCacheManager;
-		VirtualShadowMapArrayCacheManager = nullptr;
+		delete DefaultVirtualShadowMapCache;
+		DefaultVirtualShadowMapCache = nullptr;
 	}
 
 	if (LumenSceneData)
@@ -4390,9 +4484,12 @@ void FScene::UpdateAllPrimitiveSceneInfos(FRDGBuilder& GraphBuilder, bool bAsync
 
 	RemovedLocalPrimitiveSceneInfos.Sort(FPrimitiveArraySortKey());
 
-	if (VirtualShadowMapArrayCacheManager && VirtualShadowMapArrayCacheManager->IsValid())
+	TArray<FVirtualShadowMapArrayCacheManager*, SceneRenderingAllocator> VirtualShadowCacheManagers;
+	GetAllVirtualShadowMapCacheManagers(VirtualShadowCacheManagers);
+
+	for (FVirtualShadowMapArrayCacheManager* CacheManager : VirtualShadowCacheManagers)
 	{
-		FVirtualShadowMapArrayCacheManager::FInvalidatingPrimitiveCollector InvalidatingPrimitiveCollector(VirtualShadowMapArrayCacheManager);
+		FVirtualShadowMapArrayCacheManager::FInvalidatingPrimitiveCollector InvalidatingPrimitiveCollector(CacheManager);
 
 		// All removed primitives must invalidate their footprints in the VSM before leaving
 		for (const FPrimitiveSceneInfo* PrimitiveSceneInfo : RemovedLocalPrimitiveSceneInfos)
@@ -4410,7 +4507,7 @@ void FScene::UpdateAllPrimitiveSceneInfos(FRDGBuilder& GraphBuilder, bool bAsync
 			InvalidatingPrimitiveCollector.Add(Transform.Key->GetPrimitiveSceneInfo());
 		}
 
-		VirtualShadowMapArrayCacheManager->ProcessRemovedOrUpdatedPrimitives(GraphBuilder, GPUScene, InvalidatingPrimitiveCollector);
+		CacheManager->ProcessRemovedOrUpdatedPrimitives(GraphBuilder, GPUScene, InvalidatingPrimitiveCollector);
 	}
 	TArray<FPrimitiveSceneInfo*> AddedLocalPrimitiveSceneInfos;
 	AddedLocalPrimitiveSceneInfos.Reserve(AddedPrimitiveSceneInfos.Num());
@@ -5059,7 +5156,10 @@ void FScene::UpdateAllPrimitiveSceneInfos(FRDGBuilder& GraphBuilder, bool bAsync
 	}
 
 	// handle scene changes
-	VirtualShadowMapArrayCacheManager->OnSceneChange();
+	for (FVirtualShadowMapArrayCacheManager* CacheManager : VirtualShadowCacheManagers)
+	{
+		CacheManager->OnSceneChange();
+	}
 
 	{
 		CSV_SCOPED_TIMING_STAT_EXCLUSIVE(UpdatePrimitiveInstances);
@@ -5436,6 +5536,10 @@ public:
 	virtual bool HasAnyLights() const override { return false; }
 
 protected:
+	virtual FSceneViewStateInterface* AllocateViewState() override
+	{
+		return new FSceneViewState(FeatureLevel);
+	}
 
 private:
 	UWorld* World;
diff --git a/Engine/Source/Runtime/Renderer/Private/SceneOcclusion.cpp b/Engine/Source/Runtime/Renderer/Private/SceneOcclusion.cpp
--- a/Engine/Source/Runtime/Renderer/Private/SceneOcclusion.cpp
+++ b/Engine/Source/Runtime/Renderer/Private/SceneOcclusion.cpp
@@ -133,13 +133,14 @@ static FGlobalBoundShaderState GOcclusionTestBoundShaderState;
  * Returns an array of visibility data for the given view position, or NULL if none exists. 
  * The data bits are indexed by VisibilityId of each primitive in the scene.
  * This method decompresses data if necessary and caches it based on the bucket and chunk index in the view state.
+ * InScene is passed in, as the Scene pointer in the class itself may be null, if it was allocated without a scene.
  */
-const uint8* FSceneViewState::GetPrecomputedVisibilityData(FViewInfo& View, const FScene* Scene)
+const uint8* FSceneViewState::GetPrecomputedVisibilityData(FViewInfo& View, const FScene* InScene)
 {
 	const uint8* PrecomputedVisibilityData = NULL;
-	if (Scene->PrecomputedVisibilityHandler && GAllowPrecomputedVisibility && View.Family->EngineShowFlags.PrecomputedVisibility)
+	if (InScene->PrecomputedVisibilityHandler && GAllowPrecomputedVisibility && View.Family->EngineShowFlags.PrecomputedVisibility)
 	{
-		const FPrecomputedVisibilityHandler& Handler = *Scene->PrecomputedVisibilityHandler;
+		const FPrecomputedVisibilityHandler& Handler = *InScene->PrecomputedVisibilityHandler;
 		FViewElementPDI VisibilityCellsPDI(&View, nullptr, nullptr);
 
 		// Draw visibility cell bounds for debugging if enabled
@@ -181,7 +182,7 @@ const uint8* FSceneViewState::GetPrecomputedVisibilityData(FViewInfo& View, cons
 			{
 				// Reuse a cached decompressed chunk if possible
 				if (CachedVisibilityChunk
-					&& CachedVisibilityHandlerId == Scene->PrecomputedVisibilityHandler->GetId()
+					&& CachedVisibilityHandlerId == InScene->PrecomputedVisibilityHandler->GetId()
 					&& CachedVisibilityBucketIndex == PrecomputedVisibilityBucketIndex
 					&& CachedVisibilityChunkIndex == CurrentCell.ChunkIndex)
 				{
@@ -193,7 +194,7 @@ const uint8* FSceneViewState::GetPrecomputedVisibilityData(FViewInfo& View, cons
 					const FCompressedVisibilityChunk& CompressedChunk = Handler.PrecomputedVisibilityCellBuckets[PrecomputedVisibilityBucketIndex].CellDataChunks[CurrentCell.ChunkIndex];
 					CachedVisibilityBucketIndex = PrecomputedVisibilityBucketIndex;
 					CachedVisibilityChunkIndex = CurrentCell.ChunkIndex;
-					CachedVisibilityHandlerId = Scene->PrecomputedVisibilityHandler->GetId();
+					CachedVisibilityHandlerId = InScene->PrecomputedVisibilityHandler->GetId();
 
 					if (CompressedChunk.bCompressed)
 					{
diff --git a/Engine/Source/Runtime/Renderer/Private/ScenePrivate.h b/Engine/Source/Runtime/Renderer/Private/ScenePrivate.h
--- a/Engine/Source/Runtime/Renderer/Private/ScenePrivate.h
+++ b/Engine/Source/Runtime/Renderer/Private/ScenePrivate.h
@@ -701,6 +701,17 @@ public:
 	};
 
 	uint32 UniqueID;
+
+	/**
+	 * The scene pointer may be NULL -- it's only filled in if a UWorld with a valid Scene pointer was passed into the original call
+	 * to FSceneViewStateReference::Allocate.  Currently, the Scene pointer in the class is only used for constructing and cleaning up
+	 * the ViewVirtualShadowMapCache member, if the user chooses to enable per-view shadow caching via AddVirtualShadowMapCache().
+	 * If the Scene isn't present, AddVirtualShadowMapCache() has no effect.
+	 * 
+	 * TODO:  deprecate support for allocating an FSceneViewState without a valid Scene...
+	 */
+	FScene* Scene;
+
 	typedef TMap<FSceneViewState::FProjectedShadowKey, FRHIPooledRenderQuery> ShadowKeyOcclusionQueryMap;
 	TArray<ShadowKeyOcclusionQueryMap, TInlineAllocator<FOcclusionQueryHelpers::MaxBufferedOcclusionFrames> > ShadowOcclusionQueryMaps;
 
@@ -1061,6 +1072,8 @@ public:
 
 	FShadingEnergyConservationStateData ShadingEnergyConservationData;
 
+	FVirtualShadowMapArrayCacheManager* ViewVirtualShadowMapCache;
+
 	// call after OnFrameRenderingSetup()
 	virtual uint32 GetCurrentTemporalAASampleIndex() const
 	{
@@ -1125,7 +1138,8 @@ public:
 		}
 	}
 
-	void TrimHistoryRenderTargets(const FScene* Scene);
+	// InScene is passed in, as the Scene pointer in the class itself may be null, if it was allocated without a scene.
+	void TrimHistoryRenderTargets(const FScene* InScene);
 
 	/**
 	 * Calculates and stores the scale factor to apply to motion vectors based on the current game
@@ -1152,8 +1166,9 @@ public:
 	 * Returns an array of visibility data for the given view, or NULL if none exists. 
 	 * The data bits are indexed by VisibilityId of each primitive in the scene.
 	 * This method decompresses data if necessary and caches it based on the bucket and chunk index in the view state.
+	 * InScene is passed in, as the Scene pointer in the class itself may be null, if it was allocated without a scene.
 	 */
-	const uint8* GetPrecomputedVisibilityData(FViewInfo& View, const FScene* Scene);
+	const uint8* GetPrecomputedVisibilityData(FViewInfo& View, const FScene* InScene);
 
 	/**
 	 * Cleans out old entries from the primitive occlusion history, and resets unused pending occlusion queries.
@@ -1513,6 +1528,9 @@ public:
 		return SequencerState;
 	}
 
+	virtual void AddVirtualShadowMapCache() override;
+	virtual FVirtualShadowMapArrayCacheManager* GetVirtualShadowMapCache() const override;
+
 	/** Information about visibility/occlusion states in past frames for individual primitives. */
 	TSet<FPrimitiveOcclusionHistory,FPrimitiveOcclusionHistoryKeyFuncs> PrimitiveOcclusionHistorySet;
 };
@@ -2669,6 +2687,9 @@ public:
 	/** An optional FX system associated with the scene. */
 	class FFXSystemInterface* FXSystem;
 
+	/** List of view states associated with the scene. */
+	TArray<FSceneViewState*> ViewStates;
+
 	FPersistentUniformBuffers UniformBuffers;
 
 	/** Instancing state buckets.  These are stored on the scene as they are precomputed at FPrimitiveSceneInfo::AddToScene time. */
@@ -2909,7 +2930,13 @@ public:
 	/** Preshadows that are currently cached in the PreshadowCache render target. */
 	TArray<TRefCountPtr<FProjectedShadowInfo> > CachedPreshadows;
 
-	FVirtualShadowMapArrayCacheManager *VirtualShadowMapArrayCacheManager;
+	/**
+	 * Virtual shadow maps can use a default cache stored in the FScene, or define a separate cache per view (AddVirtualShadowMapCache)
+	 * for improved performance.  A linked list of caches for the FScene is maintained to allow scene updates to be propagated to all
+	 * caches as needed.  If a cache is never used, the GPU resources for it are never allocated, so the overhead of the default cache
+	 * should be minimal when not in use.
+	 */
+	FVirtualShadowMapArrayCacheManager* DefaultVirtualShadowMapCache;
 
 	/** Texture layout that tracks current allocations in the PreshadowCache render target. */
 	FTextureLayout PreshadowCacheLayout;
@@ -3119,6 +3146,9 @@ public:
 
 	virtual bool RequestGPUSceneUpdate(FPrimitiveSceneInfo& PrimitiveSceneInfo, EPrimitiveDirtyState PrimitiveDirtyState) override;
 
+	FVirtualShadowMapArrayCacheManager* GetVirtualShadowMapCache(FSceneView& View) const;
+	void GetAllVirtualShadowMapCacheManagers(TArray<FVirtualShadowMapArrayCacheManager*, SceneRenderingAllocator>& OutCacheManagers) const;
+
 	bool HasSkyAtmosphere() const
 	{
 		return (SkyAtmosphere != nullptr);
@@ -3373,6 +3403,11 @@ public:
 protected:
 
 private:
+	/**
+	 * Internal view state allocation interface, used by FSceneViewState
+	 */
+	virtual FSceneViewStateInterface* AllocateViewState() override;
+	void RemoveViewState(FSceneViewStateInterface*);
 
 	/**
 	 * Ensures the packed primitive arrays contain the same number of elements.
@@ -3528,6 +3563,8 @@ private:
 
 	/** Whether world settings has bForceNoPrecomputedLighting set */
 	bool bForceNoPrecomputedLighting;
+
+	friend class FSceneViewState;
 };
 
 inline bool ShouldIncludeDomainInMeshPass(EMaterialDomain Domain)
diff --git a/Engine/Source/Runtime/Renderer/Private/ShadowDepthRendering.cpp b/Engine/Source/Runtime/Renderer/Private/ShadowDepthRendering.cpp
--- a/Engine/Source/Runtime/Renderer/Private/ShadowDepthRendering.cpp
+++ b/Engine/Source/Runtime/Renderer/Private/ShadowDepthRendering.cpp
@@ -1642,7 +1642,7 @@ void FSceneRenderer::RenderVirtualShadowMaps(FRDGBuilder& GraphBuilder, bool bNa
 		return;
 	}
 
-	FVirtualShadowMapArrayCacheManager *CacheManager = Scene->VirtualShadowMapArrayCacheManager;
+	FVirtualShadowMapArrayCacheManager *CacheManager = ActiveViewFamily->VirtualShadowMapArray.CacheManager;
 
 	// TODO: Separate out the decision about nanite using HZB and stuff like HZB culling invalidations?
 	const bool bVSMUseHZB = VirtualShadowMapArray.UseHzbOcclusion();
diff --git a/Engine/Source/Runtime/Renderer/Private/ShadowSetup.cpp b/Engine/Source/Runtime/Renderer/Private/ShadowSetup.cpp
--- a/Engine/Source/Runtime/Renderer/Private/ShadowSetup.cpp
+++ b/Engine/Source/Runtime/Renderer/Private/ShadowSetup.cpp
@@ -3731,10 +3731,11 @@ void FSceneRenderer::CreateWholeSceneProjectedShadow(
 						0 // no border
 					);
 					
+					FVirtualShadowMapArrayCacheManager* CacheManager = ActiveViewFamily->VirtualShadowMapArray.CacheManager;
 					TSharedPtr<FVirtualShadowMapPerLightCacheEntry> VirtualSmPerLightCacheEntry;
-					if (Scene->VirtualShadowMapArrayCacheManager && Scene->VirtualShadowMapArrayCacheManager->IsValid())
+					if (CacheManager->IsValid())
 					{
-						VirtualSmPerLightCacheEntry = Scene->VirtualShadowMapArrayCacheManager->FindCreateLightCacheEntry(LightSceneInfo->Id);
+						VirtualSmPerLightCacheEntry = CacheManager->FindCreateLightCacheEntry(LightSceneInfo->Id);
 					}
 
 					int32 NumMaps = ProjectedShadowInitializer.bOnePassPointLightShadow ? 6 : 1;
@@ -4910,7 +4911,6 @@ void FSceneRenderer::AddViewDependentWholeSceneShadowsForView(
 
 				TSharedPtr<FVirtualShadowMapClipmap> VirtualShadowMapClipmap = TSharedPtr<FVirtualShadowMapClipmap>(new FVirtualShadowMapClipmap(
 					ActiveViewFamily->VirtualShadowMapArray,
-					Scene->VirtualShadowMapArrayCacheManager,
 					LightSceneInfo,
 					WorldToLight,
 					View.ViewMatrices,
diff --git a/Engine/Source/Runtime/Renderer/Private/VirtualShadowMaps/VirtualShadowMapCacheManager.cpp b/Engine/Source/Runtime/Renderer/Private/VirtualShadowMaps/VirtualShadowMapCacheManager.cpp
--- a/Engine/Source/Runtime/Renderer/Private/VirtualShadowMaps/VirtualShadowMapCacheManager.cpp
+++ b/Engine/Source/Runtime/Renderer/Private/VirtualShadowMaps/VirtualShadowMapCacheManager.cpp
@@ -274,42 +274,46 @@ FVirtualShadowMapArrayCacheManager::FVirtualShadowMapArrayCacheManager(FScene* I
 	// Handle message with status sent back from GPU
 	StatusFeedbackSocket = GPUMessage::RegisterHandler(TEXT("Shadow.Virtual.StatusFeedback"), [this](GPUMessage::FReader Message)
 	{
-		// Get the frame that the message was sent.
-		uint32 FrameNumber = Message.Read<uint32>(0);
-		// Goes negative on underflow
-		int32 NumPagesFree = Message.Read<int32>(0);
-
-		if (NumPagesFree < 0)
+		// Only process status messages that came from this specific cache manager
+		if (Message.MessageId == this->StatusFeedbackSocket.GetMessageId())
 		{
-			static const auto* CVarResolutionLodBiasLocalPtr = IConsoleManager::Get().FindTConsoleVariableDataFloat(TEXT("r.Shadow.Virtual.ResolutionLodBiasLocal"));
-			const float LodBiasLocal = CVarResolutionLodBiasLocalPtr->GetValueOnRenderThread();
+			// Get the frame that the message was sent.
+			uint32 FrameNumber = Message.Read<uint32>(0);
+			// Goes negative on underflow
+			int32 NumPagesFree = Message.Read<int32>(0);
+
+			if (NumPagesFree < 0)
+			{
+				static const auto* CVarResolutionLodBiasLocalPtr = IConsoleManager::Get().FindTConsoleVariableDataFloat(TEXT("r.Shadow.Virtual.ResolutionLodBiasLocal"));
+				const float LodBiasLocal = CVarResolutionLodBiasLocalPtr->GetValueOnRenderThread();
 
-			static const auto* CVarResolutionLodBiasDirectionalPtr = IConsoleManager::Get().FindTConsoleVariableDataFloat(TEXT("r.Shadow.Virtual.ResolutionLodBiasDirectional"));
-			const float LodBiasDirectional = CVarResolutionLodBiasDirectionalPtr->GetValueOnRenderThread();
+				static const auto* CVarResolutionLodBiasDirectionalPtr = IConsoleManager::Get().FindTConsoleVariableDataFloat(TEXT("r.Shadow.Virtual.ResolutionLodBiasDirectional"));
+				const float LodBiasDirectional = CVarResolutionLodBiasDirectionalPtr->GetValueOnRenderThread();
 
-			static const auto* CVarMaxPhysicalPagesPtr = IConsoleManager::Get().FindTConsoleVariableDataInt(TEXT("r.Shadow.Virtual.MaxPhysicalPages"));
-			const int32 MaxPhysicalPages = CVarMaxPhysicalPagesPtr->GetValueOnRenderThread();
+				static const auto* CVarMaxPhysicalPagesPtr = IConsoleManager::Get().FindTConsoleVariableDataInt(TEXT("r.Shadow.Virtual.MaxPhysicalPages"));
+				const int32 MaxPhysicalPages = CVarMaxPhysicalPagesPtr->GetValueOnRenderThread();
 
 #if !UE_BUILD_SHIPPING
-			if (!bLoggedPageOverflow)
+				if (!bLoggedPageOverflow)
+				{
+					UE_LOG(LogRenderer, Warning, TEXT("Virtual Shadow Map Page Pool overflow (%d page allocations were not served), this will produce visual artifacts (missing shadow), increase the page pool limit or reduce resolution bias to avoid.\n")
+						TEXT(" See r.Shadow.Virtual.MaxPhysicalPages (%d), r.Shadow.Virtual.ResolutionLodBiasLocal (%.2f), and r.Shadow.Virtual.ResolutionLodBiasDirectional (%.2f)"),
+						-NumPagesFree,
+						MaxPhysicalPages,
+						LodBiasLocal,
+						LodBiasDirectional);
+					bLoggedPageOverflow = true;
+				}
+				LastOverflowFrame = Scene->GetFrameNumber();
+#endif
+			}
+#if !UE_BUILD_SHIPPING
+			else
 			{
-				UE_LOG(LogRenderer, Warning, TEXT("Virtual Shadow Map Page Pool overflow (%d page allocations were not served), this will produce visual artifacts (missing shadow), increase the page pool limit or reduce resolution bias to avoid.\n")
-					TEXT(" See r.Shadow.Virtual.MaxPhysicalPages (%d), r.Shadow.Virtual.ResolutionLodBiasLocal (%.2f), and r.Shadow.Virtual.ResolutionLodBiasDirectional (%.2f)"),
-					-NumPagesFree,
-					MaxPhysicalPages,
-					LodBiasLocal,
-					LodBiasDirectional);
-				bLoggedPageOverflow = true;
+				bLoggedPageOverflow = false;
 			}
-			LastOverflowFrame = Scene->GetFrameNumber();
 #endif
 		}
-#if !UE_BUILD_SHIPPING
-		else
-		{
-			bLoggedPageOverflow = false;
-		}
-#endif
 	});
 
 #if !UE_BUILD_SHIPPING
@@ -1017,4 +1021,4 @@ void FVirtualShadowMapArrayCacheManager::ProcessGPUInstanceInvalidations(FRDGBui
 		// Drop the InvalidatingInstancesBuffer to make sure we don't redundantly process the associated invalidations if ProcessRemovedOrUpdatedPrimitives is called multiple times.
 		PrevBuffers.InvalidatingInstancesBuffer.SafeRelease();
 	}
-}
\ No newline at end of file
+}
diff --git a/Engine/Source/Runtime/Renderer/Private/VirtualShadowMaps/VirtualShadowMapClipmap.cpp b/Engine/Source/Runtime/Renderer/Private/VirtualShadowMaps/VirtualShadowMapClipmap.cpp
--- a/Engine/Source/Runtime/Renderer/Private/VirtualShadowMaps/VirtualShadowMapClipmap.cpp
+++ b/Engine/Source/Runtime/Renderer/Private/VirtualShadowMaps/VirtualShadowMapClipmap.cpp
@@ -63,7 +63,6 @@ static float GetLevelRadius(int32 Level)
 
 FVirtualShadowMapClipmap::FVirtualShadowMapClipmap(
 	FVirtualShadowMapArray& VirtualShadowMapArray,
-	FVirtualShadowMapArrayCacheManager* VirtualShadowMapArrayCacheManager,
 	const FLightSceneInfo& InLightSceneInfo,
 	const FMatrix& WorldToLightRotationMatrix,
 	const FViewMatrices& CameraViewMatrices,
@@ -74,6 +73,8 @@ FVirtualShadowMapClipmap::FVirtualShadowMapClipmap(
 {
 	check(WorldToLightRotationMatrix.GetOrigin() == FVector(0, 0, 0));	// Should not contain translation or scaling
 
+	FVirtualShadowMapArrayCacheManager* VirtualShadowMapArrayCacheManager = VirtualShadowMapArray.CacheManager;
+
 	const bool bCacheValid = VirtualShadowMapArrayCacheManager && VirtualShadowMapArrayCacheManager->IsValid();
 
 	const FMatrix FaceMatrix(
diff --git a/Engine/Source/Runtime/Renderer/Private/VirtualShadowMaps/VirtualShadowMapClipmap.h b/Engine/Source/Runtime/Renderer/Private/VirtualShadowMaps/VirtualShadowMapClipmap.h
--- a/Engine/Source/Runtime/Renderer/Private/VirtualShadowMaps/VirtualShadowMapClipmap.h
+++ b/Engine/Source/Runtime/Renderer/Private/VirtualShadowMaps/VirtualShadowMapClipmap.h
@@ -19,7 +19,6 @@ class FVirtualShadowMapClipmap : FRefCountedObject
 public:	
 	FVirtualShadowMapClipmap(
 		FVirtualShadowMapArray& VirtualShadowMapArray,
-		FVirtualShadowMapArrayCacheManager* VirtualShadowMapArrayCacheManager,
 		const FLightSceneInfo& InLightSceneInfo,
 		const FMatrix& WorldToLightRotationMatrix,
 		const FViewMatrices& CameraViewMatrices,
