Fix assert on light count when there are no lights in the scene

diff --git a/Engine/Source/Runtime/Renderer/Private/RayTracing/RayTracingLighting.cpp b/Engine/Source/Runtime/Renderer/Private/RayTracing/RayTracingLighting.cpp
--- a/Engine/Source/Runtime/Renderer/Private/RayTracing/RayTracingLighting.cpp
+++ b/Engine/Source/Runtime/Renderer/Private/RayTracing/RayTracingLighting.cpp
@@ -303,7 +303,7 @@ TRDGUniformBufferRef<FRaytracingLightDataPacked> CreateRayTracingLightData(
 	FRDGUploadData<FRTLightingData> LightDataArray(GraphBuilder, FMath::Max(LightIndices.Num(), 1));
 	SetupRaytracingLightDataPacked(GraphBuilder, Lights, LightIndices, View, *LightData, LightDataArray);
 
-	check(LightData->Count == LightDataArray.Num());
+	check(LightData->Count == LightIndices.Num());
 
 	static_assert(sizeof(FRTLightingData) % sizeof(FUintVector4) == 0, "sizeof(FRTLightingData) must be a multiple of sizeof(FUintVector4)");
 	const uint32 NumUintVector4Elements = LightDataArray.GetTotalSize() / sizeof(FUintVector4);
