Fix ASan hit in DynamicBVH tests

Avoid taking a reference to a container that may be resized during the calling process

diff --git a/Engine/Source/Runtime/Renderer/Private/DynamicBVH.cpp b/Engine/Source/Runtime/Renderer/Private/DynamicBVH.cpp
--- a/Engine/Source/Runtime/Renderer/Private/DynamicBVH.cpp
+++ b/Engine/Source/Runtime/Renderer/Private/DynamicBVH.cpp
@@ -269,7 +269,8 @@ void TestBVH_Ordered()
 	for( int i = 0; i < 1024; i++ )
 	{
 		uint32 Index = FMath::Rand() % Num;
-		BVH.Update( BVH.GetBounds( Index ), Index );
+		FBounds3f Bounds = BVH.GetBounds( Index );
+		BVH.Update( Bounds, Index );
 		//BVH.Optimize(1);
 	}
 
@@ -278,7 +279,8 @@ void TestBVH_Ordered()
 	for( int i = 0; i < 65536; i++ )
 	{
 		uint32 Index = FMath::Rand() % Num;
-		BVH.Update( BVH.GetBounds( Index ), Index );
+		FBounds3f Bounds = BVH.GetBounds( Index );
+		BVH.Update( Bounds, Index );
 		//BVH.Optimize(1);
 	}
 
@@ -299,7 +301,8 @@ void TestBVH_Optimize()
 	for( int i = 0; i < 1024; i++ )
 	{
 		uint32 Index = FMath::Rand() % Num;
-		BVH.Update( BVH.GetBounds( Index ), Index );
+		FBounds3f Bounds = BVH.GetBounds( Index );
+		BVH.Update( Bounds, Index );
 		//BVH.Optimize(1);
 	}
 
@@ -308,7 +311,8 @@ void TestBVH_Optimize()
 	for( int i = 0; i < 65536; i++ )
 	{
 		uint32 Index = FMath::Rand() % Num;
-		BVH.Update( BVH.GetBounds( Index ), Index );
+		FBounds3f Bounds = BVH.GetBounds( Index );
+		BVH.Update( Bounds, Index );
 		//BVH.Optimize(1);
 	}
 
@@ -353,4 +357,4 @@ bool FTestBVH::RunTest( const FString& Parameters )
 	TestBVH_Build();
 
 	return true;
-}
\ No newline at end of file
+}
