Add warning when the maximum number of active ray tracing lights >= RAY_TRACING_LIGHT_COUNT_MAXIMUM

Promote OnGetOnScreenMessages delegate from private to protected so that all derived renderers from SceneRenderer can add screen messages without modifying the header.

diff --git a/Engine/Source/Runtime/Renderer/Private/DeferredShadingRenderer.cpp b/Engine/Source/Runtime/Renderer/Private/DeferredShadingRenderer.cpp
--- a/Engine/Source/Runtime/Renderer/Private/DeferredShadingRenderer.cpp
+++ b/Engine/Source/Runtime/Renderer/Private/DeferredShadingRenderer.cpp
@@ -1504,6 +1504,8 @@
 
 	ReferenceView.RayTracingSubSurfaceProfileSRV = RHICreateShaderResourceView(ReferenceView.RayTracingSubSurfaceProfileTexture, 0);
 
+	uint32 NumOfSkippedRayTracingLights = 0;
+
 	for (int32 ViewIndex = 0; ViewIndex < Views.Num(); ++ViewIndex)
 	{
 		FViewInfo& View = Views[ViewIndex];
@@ -1524,9 +1526,23 @@
 		else
 		{
 			// This light data is a function of the camera position, so must be computed per view.
-			View.RayTracingLightDataUniformBuffer = CreateRayTracingLightData(GraphBuilder, Scene->Lights, View);
+			View.RayTracingLightDataUniformBuffer = CreateRayTracingLightData(GraphBuilder, Scene->Lights, View, NumOfSkippedRayTracingLights);
 		}
 	}
+
+#if !UE_BUILD_SHIPPING
+	if (!bIsPathTracing && NumOfSkippedRayTracingLights > 0)
+	{
+		OnGetOnScreenMessages.AddLambda([NumOfSkippedRayTracingLights](FScreenMessageWriter& ScreenMessageWriter)->void
+			{
+				FString String = FString::Printf(
+					TEXT("%d light(s) skipped. Active Ray Tracing light count > RAY_TRACING_LIGHT_COUNT_MAXIMUM (%d)."),
+					NumOfSkippedRayTracingLights,
+					RAY_TRACING_LIGHT_COUNT_MAXIMUM);
+				ScreenMessageWriter.DrawLine(FText::FromString(String), 10, FColor::Yellow);
+			});
+	}
+#endif
 
 	return true;
 }
diff --git a/Engine/Source/Runtime/Renderer/Private/RayTracing/RayTracingLighting.cpp b/Engine/Source/Runtime/Renderer/Private/RayTracing/RayTracingLighting.cpp
--- a/Engine/Source/Runtime/Renderer/Private/RayTracing/RayTracingLighting.cpp
+++ b/Engine/Source/Runtime/Renderer/Private/RayTracing/RayTracingLighting.cpp
@@ -64,10 +64,12 @@
 static void SelectRaytracingLights(
 	const TSparseArray<FLightSceneInfoCompact, TAlignedSparseArrayAllocator<alignof(FLightSceneInfoCompact)>>& Lights,
 	const FViewInfo& View,
-	TArray<int32>& OutSelectedLights
+	TArray<int32>& OutSelectedLights,
+	uint32& NumOfSkippedRayTracingLights
 )
 {
 	OutSelectedLights.Empty();
+	NumOfSkippedRayTracingLights = 0;
 
 	for (auto Light : Lights)
 	{
@@ -75,9 +77,14 @@
 		const bool bAffectReflection = Light.LightSceneInfo->Proxy->AffectReflection();
 		if (bHasStaticLighting || !bAffectReflection) continue;
 
-		OutSelectedLights.Add(Light.LightSceneInfo->Id);
-
-		if (OutSelectedLights.Num() >= RAY_TRACING_LIGHT_COUNT_MAXIMUM) break;
+		if (OutSelectedLights.Num() < RAY_TRACING_LIGHT_COUNT_MAXIMUM)
+		{
+			OutSelectedLights.Add(Light.LightSceneInfo->Id);
+		}
+		else
+		{
+			NumOfSkippedRayTracingLights++;
+		};
 	}
 }
 
@@ -286,14 +293,15 @@
 TRDGUniformBufferRef<FRaytracingLightDataPacked> CreateRayTracingLightData(
 	FRDGBuilder& GraphBuilder,
 	const TSparseArray<FLightSceneInfoCompact, TAlignedSparseArrayAllocator<alignof(FLightSceneInfoCompact)>>& Lights,
-	const FViewInfo& View)
+	const FViewInfo& View,
+	uint32& NumOfSkippedRayTracingLights)
 {
 	auto* LightData = GraphBuilder.AllocParameters<FRaytracingLightDataPacked>();
 	LightData->CellCount = GetCellsPerDim();
 	LightData->CellScale = CVarRayTracingLightingCellSize.GetValueOnRenderThread() / 2.0f;
 
 	TArray<int32> LightIndices;
-	SelectRaytracingLights(Lights, View, LightIndices);
+	SelectRaytracingLights(Lights, View, LightIndices, NumOfSkippedRayTracingLights);
 
 	// Create light culling volume
 	FRDGBufferRef LightCullVolume;
diff --git a/Engine/Source/Runtime/Renderer/Private/RayTracing/RayTracingLighting.h b/Engine/Source/Runtime/Renderer/Private/RayTracing/RayTracingLighting.h
--- a/Engine/Source/Runtime/Renderer/Private/RayTracing/RayTracingLighting.h
+++ b/Engine/Source/Runtime/Renderer/Private/RayTracing/RayTracingLighting.h
@@ -58,6 +58,7 @@
 TRDGUniformBufferRef<FRaytracingLightDataPacked> CreateRayTracingLightData(
 	FRDGBuilder& GraphBuilder,
 	const TSparseArray<FLightSceneInfoCompact, TAlignedSparseArrayAllocator<alignof(FLightSceneInfoCompact)>>& Lights,
-	const FViewInfo& View);
+	const FViewInfo& View,
+	uint32& NumOfSkippedRayTracingLights);
 
 #endif
