Support Callable Shaders in FRayTracingScene.
- FRHIRayTracingScene is created with NumCallableShaderSlots.
- Application can add entries to SBT using CallableCommands array.
- ShaderBindings are generated from CallableCommands in BindRayTracingMaterialPipeline(...) once the RTPSO is available so CallableShaderIndex is known.
- Shader Bindings are then merged and set during WaitForRayTracingScene.

diff --git a/Engine/Source/Runtime/Renderer/Private/DeferredShadingRenderer.cpp b/Engine/Source/Runtime/Renderer/Private/DeferredShadingRenderer.cpp
--- a/Engine/Source/Runtime/Renderer/Private/DeferredShadingRenderer.cpp
+++ b/Engine/Source/Runtime/Renderer/Private/DeferredShadingRenderer.cpp
@@ -1732,57 +1732,79 @@ void FDeferredShadingSceneRenderer::WaitForRayTracingScene(FRDGBuilder& GraphBui
 	{
 		check(ReferenceView.RayTracingMaterialPipeline || ReferenceView.RayTracingMaterialBindings.Num() == 0);
 
-		if (ReferenceView.RayTracingMaterialPipeline && ReferenceView.RayTracingMaterialBindings.Num())
+		if (ReferenceView.RayTracingMaterialPipeline && (ReferenceView.RayTracingMaterialBindings.Num() || ReferenceView.RayTracingCallableBindings.Num()))
 		{
-			FTaskGraphInterface::Get().WaitUntilTaskCompletes(ReferenceView.RayTracingMaterialBindingsTask, ENamedThreads::GetRenderThread_Local());
-
 			// Gather bindings from all chunks and submit them all as a single batch to allow RHI to bind all shader parameters in parallel.
 
-			uint32 NumTotalBindings = 0;
-
-			for (FRayTracingLocalShaderBindingWriter* BindingWriter : ReferenceView.RayTracingMaterialBindings)
+			auto MergeAndSetBindings =
+				[
+					&RHICmdList,
+					RayTracingScene = ReferenceView.GetRayTracingSceneChecked(),
+					Pipeline = ReferenceView.RayTracingMaterialPipeline
+				](TConstArrayView<FRayTracingLocalShaderBindingWriter*> Bindings, ERayTracingBindingType BindingType)
 			{
-				const FRayTracingLocalShaderBindingWriter::FChunk* Chunk = BindingWriter->GetFirstChunk();
-				while (Chunk)
+				uint32 NumTotalBindings = 0;
+
+				for (FRayTracingLocalShaderBindingWriter* BindingWriter : Bindings)
 				{
-					NumTotalBindings += Chunk->Num;
-					Chunk = Chunk->Next;
+					const FRayTracingLocalShaderBindingWriter::FChunk* Chunk = BindingWriter->GetFirstChunk();
+					while (Chunk)
+					{
+						NumTotalBindings += Chunk->Num;
+						Chunk = Chunk->Next;
+					}
 				}
-			}
 
-			const uint32 MergedBindingsSize = sizeof(FRayTracingLocalShaderBindings) * NumTotalBindings;
-			FRayTracingLocalShaderBindings* MergedBindings = (FRayTracingLocalShaderBindings*)(RHICmdList.Bypass()
-				? FMemStack::Get().Alloc(MergedBindingsSize, alignof(FRayTracingLocalShaderBindings))
-				: RHICmdList.Alloc(MergedBindingsSize, alignof(FRayTracingLocalShaderBindings)));
+				if (NumTotalBindings == 0)
+				{
+					return;
+				}
 
-			uint32 MergedBindingIndex = 0;
-			for (FRayTracingLocalShaderBindingWriter* BindingWriter : ReferenceView.RayTracingMaterialBindings)
-			{
-				const FRayTracingLocalShaderBindingWriter::FChunk* Chunk = BindingWriter->GetFirstChunk();
-				while (Chunk)
+				const uint32 MergedBindingsSize = sizeof(FRayTracingLocalShaderBindings) * NumTotalBindings;
+				FRayTracingLocalShaderBindings* MergedBindings = (FRayTracingLocalShaderBindings*)(RHICmdList.Bypass()
+					? FMemStack::Get().Alloc(MergedBindingsSize, alignof(FRayTracingLocalShaderBindings))
+					: RHICmdList.Alloc(MergedBindingsSize, alignof(FRayTracingLocalShaderBindings)));
+
+				uint32 MergedBindingIndex = 0;
+				for (FRayTracingLocalShaderBindingWriter* BindingWriter : Bindings)
 				{
-					const uint32 Num = Chunk->Num;
-					for (uint32_t i = 0; i < Num; ++i)
+					const FRayTracingLocalShaderBindingWriter::FChunk* Chunk = BindingWriter->GetFirstChunk();
+					while (Chunk)
 					{
-						MergedBindings[MergedBindingIndex] = Chunk->Bindings[i];
-						MergedBindingIndex++;
+						const uint32 Num = Chunk->Num;
+						for (uint32_t i = 0; i < Num; ++i)
+						{
+							MergedBindings[MergedBindingIndex] = Chunk->Bindings[i];
+							MergedBindingIndex++;
+						}
+						Chunk = Chunk->Next;
 					}
-					Chunk = Chunk->Next;
 				}
-			}
 
-			const bool bCopyDataToInlineStorage = false; // Storage is already allocated from RHICmdList, no extra copy necessary
-			RHICmdList.SetRayTracingHitGroups(
-				ReferenceView.GetRayTracingSceneChecked(),
-				ReferenceView.RayTracingMaterialPipeline,
-				NumTotalBindings, MergedBindings,
-				bCopyDataToInlineStorage);
+				const bool bCopyDataToInlineStorage = false; // Storage is already allocated from RHICmdList, no extra copy necessary
+				RHICmdList.SetRayTracingBindings(
+					RayTracingScene,
+					Pipeline,
+					NumTotalBindings, MergedBindings,
+					BindingType,
+					bCopyDataToInlineStorage);
+			};
+
+			FTaskGraphInterface::Get().WaitUntilTaskCompletes(ReferenceView.RayTracingMaterialBindingsTask, ENamedThreads::GetRenderThread_Local());
+			MergeAndSetBindings(ReferenceView.RayTracingMaterialBindings, ERayTracingBindingType::HitGroup);
+
+			FTaskGraphInterface::Get().WaitUntilTaskCompletes(ReferenceView.RayTracingCallableBindingsTask, ENamedThreads::GetRenderThread_Local());
+			MergeAndSetBindings(ReferenceView.RayTracingCallableBindings, ERayTracingBindingType::CallableShader);
 
 			// Move the ray tracing binding container ownership to the command list, so that memory will be
 			// released on the RHI thread timeline, after the commands that reference it are processed.
-			RHICmdList.EnqueueLambda([Ptrs = MoveTemp(ReferenceView.RayTracingMaterialBindings)](FRHICommandListImmediate&)
+			RHICmdList.EnqueueLambda([PtrsA = MoveTemp(ReferenceView.RayTracingMaterialBindings), PtrsB = MoveTemp(ReferenceView.RayTracingCallableBindings)](FRHICommandListImmediate&)
 			{
-				for (auto Ptr : Ptrs)
+				for (auto Ptr : PtrsA)
+				{
+					delete Ptr;
+				}
+				for (auto Ptr : PtrsB)
 				{
 					delete Ptr;
 				}
diff --git a/Engine/Source/Runtime/Renderer/Private/MeshPassProcessor.cpp b/Engine/Source/Runtime/Renderer/Private/MeshPassProcessor.cpp
--- a/Engine/Source/Runtime/Renderer/Private/MeshPassProcessor.cpp
+++ b/Engine/Source/Runtime/Renderer/Private/MeshPassProcessor.cpp
@@ -782,6 +782,18 @@
 	MaterialShader = Shaders.RayTracingShader.GetRayTracingShader();
 	ShaderBindings.Initialize(Shaders);
 }
+
+void FRayTracingShaderCommand::SetShader(const TShaderRef<FShader>& InShader)
+{
+	check(InShader->GetFrequency() == SF_RayCallable || InShader->GetFrequency() == SF_RayMiss);
+	ShaderIndex = InShader.GetRayTracingCallableShaderLibraryIndex();
+	Shader = InShader.GetRayTracingShader();
+
+	FMeshProcessorShaders Shaders;
+	Shaders.RayTracingShader = InShader;
+
+	ShaderBindings.Initialize(Shaders);
+}
 #endif // RHI_RAYTRACING
 
 void FMeshDrawCommand::SetDrawParametersAndFinalize(
diff --git a/Engine/Source/Runtime/Renderer/Private/RayTracing/RayTracingInstanceBufferUtil.cpp b/Engine/Source/Runtime/Renderer/Private/RayTracing/RayTracingInstanceBufferUtil.cpp
--- a/Engine/Source/Runtime/Renderer/Private/RayTracing/RayTracingInstanceBufferUtil.cpp
+++ b/Engine/Source/Runtime/Renderer/Private/RayTracing/RayTracingInstanceBufferUtil.cpp
@@ -23,7 +23,8 @@
 	TArrayView<FRayTracingGeometryInstance> Instances,
 	uint8 NumLayers,
 	uint32 NumShaderSlotsPerGeometrySegment,
-	uint32 NumMissShaderSlots)
+	uint32 NumMissShaderSlots,
+	uint32 NumCallableShaderSlots)
 {
 	const uint32 NumSceneInstances = Instances.Num();
 
@@ -38,6 +39,7 @@
 	Initializer.DebugName = FName(TEXT("FRayTracingScene"));
 	Initializer.ShaderSlotsPerGeometrySegment = NumShaderSlotsPerGeometrySegment;
 	Initializer.NumMissShaderSlots = NumMissShaderSlots;
+	Initializer.NumCallableShaderSlots = NumCallableShaderSlots;
 	Initializer.PerInstanceGeometries.SetNumUninitialized(NumSceneInstances);
 	Initializer.BaseInstancePrefixSum.SetNumUninitialized(NumSceneInstances);
 	Initializer.SegmentPrefixSum.SetNumUninitialized(NumSceneInstances);
diff --git a/Engine/Source/Runtime/Renderer/Private/RayTracing/RayTracingMaterialHitShaders.cpp b/Engine/Source/Runtime/Renderer/Private/RayTracing/RayTracingMaterialHitShaders.cpp
--- a/Engine/Source/Runtime/Renderer/Private/RayTracing/RayTracingMaterialHitShaders.cpp
+++ b/Engine/Source/Runtime/Renderer/Private/RayTracing/RayTracingMaterialHitShaders.cpp
@@ -703,98 +703,153 @@
 
 	FViewInfo& ReferenceView = Views[0];
 
-	const uint32 NumTotalMeshCommands = ReferenceView.VisibleRayTracingMeshCommands.Num();
-	const uint32 TargetCommandsPerTask = 4096; // Granularity chosen based on profiling Infiltrator scene to balance wall time speedup and total CPU thread time.
-	const uint32 NumTasks = FMath::Max(1u, FMath::DivideAndRoundUp(NumTotalMeshCommands, TargetCommandsPerTask));
-	const uint32 CommandsPerTask = FMath::DivideAndRoundUp(NumTotalMeshCommands, NumTasks); // Evenly divide commands between tasks (avoiding potential short last task)
+	// material hit groups
+	{
+		const uint32 NumTotalMeshCommands = ReferenceView.VisibleRayTracingMeshCommands.Num();
+		const uint32 TargetCommandsPerTask = 4096; // Granularity chosen based on profiling Infiltrator scene to balance wall time speedup and total CPU thread time.
+		const uint32 NumTasks = FMath::Max(1u, FMath::DivideAndRoundUp(NumTotalMeshCommands, TargetCommandsPerTask));
+		const uint32 CommandsPerTask = FMath::DivideAndRoundUp(NumTotalMeshCommands, NumTasks); // Evenly divide commands between tasks (avoiding potential short last task)
 
-	FGraphEventArray TaskList;
-	TaskList.Reserve(NumTasks);
-	View.RayTracingMaterialBindings.SetNum(NumTasks);
-	
-	for (uint32 TaskIndex = 0; TaskIndex < NumTasks; ++TaskIndex)
-	{
-		const uint32 FirstTaskCommandIndex = TaskIndex * CommandsPerTask;
-		const FVisibleRayTracingMeshCommand* MeshCommands = ReferenceView.VisibleRayTracingMeshCommands.GetData() + FirstTaskCommandIndex;
-		const uint32 NumCommands = FMath::Min(CommandsPerTask, NumTotalMeshCommands - FirstTaskCommandIndex);
+		FGraphEventArray TaskList;
+		TaskList.Reserve(NumTasks);
+		View.RayTracingMaterialBindings.SetNum(NumTasks);
 
-		FRayTracingLocalShaderBindingWriter* BindingWriter = new FRayTracingLocalShaderBindingWriter();
-		View.RayTracingMaterialBindings[TaskIndex] = BindingWriter;
+		for (uint32 TaskIndex = 0; TaskIndex < NumTasks; ++TaskIndex)
+		{
+			const uint32 FirstTaskCommandIndex = TaskIndex * CommandsPerTask;
+			const FVisibleRayTracingMeshCommand* MeshCommands = ReferenceView.VisibleRayTracingMeshCommands.GetData() + FirstTaskCommandIndex;
+			const uint32 NumCommands = FMath::Min(CommandsPerTask, NumTotalMeshCommands - FirstTaskCommandIndex);
 
-		TaskList.Add(FFunctionGraphTask::CreateAndDispatchWhenReady(
-		[PipelineState, BindingWriter, MeshCommands, NumCommands, bEnableMaterials, bEnableShadowMaterials,
-			DefaultClosestHitMaterialIndex, OpaqueShadowMaterialIndex, HiddenMaterialIndex, TaskIndex]()
-		{
-			TRACE_CPUPROFILER_EVENT_SCOPE(BindRayTracingMaterialPipelineTask);
+			FRayTracingLocalShaderBindingWriter* BindingWriter = new FRayTracingLocalShaderBindingWriter();
+			View.RayTracingMaterialBindings[TaskIndex] = BindingWriter;
 
-			for (uint32 CommandIndex = 0; CommandIndex < NumCommands; ++CommandIndex)
-			{
-				const FVisibleRayTracingMeshCommand VisibleMeshCommand = MeshCommands[CommandIndex];
-				const FRayTracingMeshCommand& MeshCommand = *VisibleMeshCommand.RayTracingMeshCommand;
+			TaskList.Add(FFunctionGraphTask::CreateAndDispatchWhenReady(
+				[PipelineState, BindingWriter, MeshCommands, NumCommands, bEnableMaterials, bEnableShadowMaterials,
+				DefaultClosestHitMaterialIndex, OpaqueShadowMaterialIndex, HiddenMaterialIndex, TaskIndex]()
+				{
+					TRACE_CPUPROFILER_EVENT_SCOPE(BindRayTracingMaterialPipelineTask);
 
-				int32 HitGroupIndex = DefaultClosestHitMaterialIndex; // Force the same shader to be used on all geometry unless materials are enabled
+					for (uint32 CommandIndex = 0; CommandIndex < NumCommands; ++CommandIndex)
+					{
+						const FVisibleRayTracingMeshCommand VisibleMeshCommand = MeshCommands[CommandIndex];
+						const FRayTracingMeshCommand& MeshCommand = *VisibleMeshCommand.RayTracingMeshCommand;
 
-				if (bEnableMaterials)
-				{
-					const int32 FoundIndex = FindRayTracingHitGroupIndex(PipelineState, MeshCommand.MaterialShader, false);
-					if (FoundIndex != INDEX_NONE)
-					{
-						HitGroupIndex = FoundIndex;
+						int32 HitGroupIndex = DefaultClosestHitMaterialIndex; // Force the same shader to be used on all geometry unless materials are enabled
+
+						if (bEnableMaterials)
+						{
+							const int32 FoundIndex = FindRayTracingHitGroupIndex(PipelineState, MeshCommand.MaterialShader, false);
+							if (FoundIndex != INDEX_NONE)
+							{
+								HitGroupIndex = FoundIndex;
+							}
+							else
+							{
+								HitGroupIndex = DefaultClosestHitMaterialIndex;
+							}
+						}
+
+						// Bind primary material shader
+
+						{
+							MeshCommand.ShaderBindings.SetRayTracingShaderBindingsForHitGroup(BindingWriter,
+								VisibleMeshCommand.InstanceIndex,
+								MeshCommand.GeometrySegmentIndex,
+								HitGroupIndex,
+								RAY_TRACING_SHADER_SLOT_MATERIAL);
+						}
+
+						// Bind shadow shader
+
+						if (MeshCommand.bCastRayTracedShadows)
+						{
+							if (MeshCommand.bOpaque || !bEnableShadowMaterials)
+							{
+								FRayTracingLocalShaderBindings& Binding = BindingWriter->AddWithExternalParameters();
+								Binding.InstanceIndex = VisibleMeshCommand.InstanceIndex;
+								Binding.SegmentIndex = MeshCommand.GeometrySegmentIndex;
+								Binding.ShaderSlot = RAY_TRACING_SHADER_SLOT_SHADOW;
+								Binding.ShaderIndexInPipeline = OpaqueShadowMaterialIndex;
+							}
+							else
+							{
+								// Masked materials require full material evaluation with any-hit shader.
+								// Full CHS is bound, however material evaluation is skipped for shadow rays using a dynamic branch on a ray payload flag.
+								MeshCommand.ShaderBindings.SetRayTracingShaderBindingsForHitGroup(BindingWriter,
+									VisibleMeshCommand.InstanceIndex,
+									MeshCommand.GeometrySegmentIndex,
+									HitGroupIndex,
+									RAY_TRACING_SHADER_SLOT_SHADOW);
+							}
+						}
+						else
+						{
+							FRayTracingLocalShaderBindings& Binding = BindingWriter->AddWithExternalParameters();
+							Binding.InstanceIndex = VisibleMeshCommand.InstanceIndex;
+							Binding.SegmentIndex = MeshCommand.GeometrySegmentIndex;
+							Binding.ShaderSlot = RAY_TRACING_SHADER_SLOT_SHADOW;
+							Binding.ShaderIndexInPipeline = HiddenMaterialIndex;
+						}
 					}
-					else
-					{
-						HitGroupIndex = DefaultClosestHitMaterialIndex;
-					}
-				}
+				},
+				TStatId(), nullptr, ENamedThreads::AnyThread));
+		}
 
-				// Bind primary material shader
-
-				{
-					MeshCommand.ShaderBindings.SetRayTracingShaderBindingsForHitGroup(BindingWriter,
-						VisibleMeshCommand.InstanceIndex,
-						MeshCommand.GeometrySegmentIndex,
-						HitGroupIndex,
-						RAY_TRACING_SHADER_SLOT_MATERIAL);
-				}
-
-				// Bind shadow shader
-
-				if (MeshCommand.bCastRayTracedShadows)
-				{
-					if (MeshCommand.bOpaque || !bEnableShadowMaterials)
-					{
-						FRayTracingLocalShaderBindings& Binding = BindingWriter->AddWithExternalParameters();
-						Binding.InstanceIndex = VisibleMeshCommand.InstanceIndex;
-						Binding.SegmentIndex = MeshCommand.GeometrySegmentIndex;
-						Binding.ShaderSlot = RAY_TRACING_SHADER_SLOT_SHADOW;
-						Binding.ShaderIndexInPipeline = OpaqueShadowMaterialIndex;
-					}
-					else
-					{
-						// Masked materials require full material evaluation with any-hit shader.
-						// Full CHS is bound, however material evaluation is skipped for shadow rays using a dynamic branch on a ray payload flag.
-						MeshCommand.ShaderBindings.SetRayTracingShaderBindingsForHitGroup(BindingWriter,
-							VisibleMeshCommand.InstanceIndex,
-							MeshCommand.GeometrySegmentIndex,
-							HitGroupIndex,
-							RAY_TRACING_SHADER_SLOT_SHADOW);
-					}
-				}
-				else
-				{
-					FRayTracingLocalShaderBindings& Binding = BindingWriter->AddWithExternalParameters();
-					Binding.InstanceIndex = VisibleMeshCommand.InstanceIndex;
-					Binding.SegmentIndex = MeshCommand.GeometrySegmentIndex;
-					Binding.ShaderSlot = RAY_TRACING_SHADER_SLOT_SHADOW;
-					Binding.ShaderIndexInPipeline = HiddenMaterialIndex;
-				}
-			}
-		},
-		TStatId(), nullptr, ENamedThreads::AnyThread));
+		View.RayTracingMaterialBindingsTask = FFunctionGraphTask::CreateAndDispatchWhenReady([]() {}, TStatId(), &TaskList, ENamedThreads::AnyHiPriThreadHiPriTask);
 	}
 
-	View.RayTracingMaterialBindingsTask = FFunctionGraphTask::CreateAndDispatchWhenReady([]() {}, TStatId(), &TaskList, ENamedThreads::AnyHiPriThreadHiPriTask);
+	if(bCallableShadersSupported)
+	{
+		const int32 DefaultCallableShaderIndex = FindRayTracingCallableShaderIndex(PipelineState, DefaultCallableShader, true);
 
+		const uint32 TargetCommandsPerTask = 4096;
+
+		const uint32 NumTotalCallableCommands = Scene->RayTracingScene.CallableCommands.Num();
+		const uint32 NumTasks = FMath::Max(1u, FMath::DivideAndRoundUp(NumTotalCallableCommands, TargetCommandsPerTask));
+		const uint32 CommandsPerTask = FMath::DivideAndRoundUp(NumTotalCallableCommands, NumTasks); // Evenly divide commands between tasks (avoiding potential short last task)
+
+		FGraphEventArray TaskList;
+		TaskList.Reserve(NumTasks);
+		View.RayTracingCallableBindings.SetNum(NumTasks);
+
+		for (uint32 TaskIndex = 0; TaskIndex < NumTasks; ++TaskIndex)
+		{
+			const uint32 TaskBaseCommandIndex = TaskIndex * CommandsPerTask;
+			const FRayTracingShaderCommand* TaskCallableCommands = Scene->RayTracingScene.CallableCommands.GetData() + TaskBaseCommandIndex;
+			const uint32 NumCommands = FMath::Min(CommandsPerTask, NumTotalCallableCommands - TaskBaseCommandIndex);
+
+			FRayTracingLocalShaderBindingWriter* BindingWriter = new FRayTracingLocalShaderBindingWriter();
+			View.RayTracingCallableBindings[TaskIndex] = BindingWriter;
+
+			TaskList.Add(FFunctionGraphTask::CreateAndDispatchWhenReady(
+				[PipelineState, BindingWriter, TaskCallableCommands, NumCommands, bEnableMaterials, DefaultCallableShaderIndex, TaskIndex]()
+				{
+					TRACE_CPUPROFILER_EVENT_SCOPE(BindRayTracingMaterialPipelineTask);
+
+					for (uint32 CommandIndex = 0; CommandIndex < NumCommands; ++CommandIndex)
+					{
+						const FRayTracingShaderCommand& CallableCommand = TaskCallableCommands[CommandIndex];
+
+						int32 CallableShaderIndex = DefaultCallableShaderIndex; // Force the same shader to be used on all geometry unless materials are enabled
+
+						if (bEnableMaterials)
+						{
+							const int32 FoundIndex = FindRayTracingCallableShaderIndex(PipelineState, CallableCommand.Shader, false);
+							if (FoundIndex != INDEX_NONE)
+							{
+								CallableShaderIndex = FoundIndex;
+							}
+						}
+
+						CallableCommand.ShaderBindings.SetRayTracingShaderBindings(BindingWriter, CallableShaderIndex, CallableCommand.SlotInScene);
+					}
+				},
+				TStatId(), nullptr, ENamedThreads::AnyThread));
+		}
+
+		View.RayTracingCallableBindingsTask = FFunctionGraphTask::CreateAndDispatchWhenReady([]() {}, TStatId(), &TaskList, ENamedThreads::AnyHiPriThreadHiPriTask);
+	}
+
 	return PipelineState;
 }
 
diff --git a/Engine/Source/Runtime/Renderer/Private/RayTracing/RayTracingScene.cpp b/Engine/Source/Runtime/Renderer/Private/RayTracing/RayTracingScene.cpp
--- a/Engine/Source/Runtime/Renderer/Private/RayTracing/RayTracingScene.cpp
+++ b/Engine/Source/Runtime/Renderer/Private/RayTracing/RayTracingScene.cpp
@@ -40,7 +40,8 @@
 		Instances,
 		NumLayers,
 		RAY_TRACING_NUM_SHADER_SLOTS,
-		RAY_TRACING_NUM_MISS_SHADER_SLOTS);
+		RAY_TRACING_NUM_MISS_SHADER_SLOTS,
+		NumCallableShaderSlots);
 
 	RayTracingSceneRHI = SceneWithGeometryInstances.Scene;
 
@@ -73,7 +74,7 @@
 	{
 		const uint64 ScratchAlignment = GRHIRayTracingScratchBufferAlignment;
 		FRDGBufferDesc ScratchBufferDesc;
-		ScratchBufferDesc.Usage = EBufferUsageFlags::UnorderedAccess | EBufferUsageFlags::StructuredBuffer;
+		ScratchBufferDesc.Usage = EBufferUsageFlags::RayTracingScratch | EBufferUsageFlags::StructuredBuffer;
 		ScratchBufferDesc.BytesPerElement = uint32(ScratchAlignment);
 		ScratchBufferDesc.NumElements = uint32(FMath::DivideAndRoundUp(SizeInfo.BuildScratchSize, ScratchAlignment));
 
@@ -275,6 +276,8 @@
 	WaitForTasks();
 
 	Instances.Reset();
+	NumCallableShaderSlots = 0;
+	CallableCommands.Reset();
 	GeometriesToBuild.Reset();
 	UsedCoarseMeshStreamingHandles.Reset();
 
@@ -288,6 +291,9 @@
 	Reset();
 
 	Instances.Empty();
+	CallableCommands.Empty();
+	GeometriesToBuild.Empty();
+	UsedCoarseMeshStreamingHandles.Empty();
 	RayTracingSceneBuffer = nullptr;
 	RayTracingSceneRHI = nullptr;
 }
diff --git a/Engine/Source/Runtime/Renderer/Private/RayTracing/RayTracingScene.h b/Engine/Source/Runtime/Renderer/Private/RayTracing/RayTracingScene.h
--- a/Engine/Source/Runtime/Renderer/Private/RayTracing/RayTracingScene.h
+++ b/Engine/Source/Runtime/Renderer/Private/RayTracing/RayTracingScene.h
@@ -13,6 +13,7 @@
 #include "Misc/MemStack.h"
 #include "Containers/ArrayView.h"
 #include "PrimitiveSceneProxy.h"
+#include "MeshPassProcessor.h"
 
 class FGPUScene;
 class FRHIRayTracingScene;
@@ -79,10 +80,11 @@
 
 	// Persistent storage for ray tracing instance descriptors.
 	// Cleared every frame without releasing memory to avoid large heap allocations.
-	// This must be filled before calling BeginCreateTask().
+	// This must be filled before calling Create().
 	TArray<FRayTracingGeometryInstance> Instances;
 
-	uint32 NumTotalSegments = 0;
+	uint32 NumCallableShaderSlots = 0;
+	TArray<FRayTracingShaderCommand> CallableCommands;
 
 	// Geometries which still have a pending build request but are used this frame and require a force build.
 	TArray<const FRayTracingGeometry*> GeometriesToBuild;
diff --git a/Engine/Source/Runtime/Renderer/Private/SceneRendering.h b/Engine/Source/Runtime/Renderer/Private/SceneRendering.h
--- a/Engine/Source/Runtime/Renderer/Private/SceneRendering.h
+++ b/Engine/Source/Runtime/Renderer/Private/SceneRendering.h
@@ -1484,6 +1484,9 @@ public:
 	TArray<FRayTracingLocalShaderBindingWriter*>	RayTracingMaterialBindings; // One per binding task
 	FGraphEventRef									RayTracingMaterialBindingsTask;
 
+	TArray<FRayTracingLocalShaderBindingWriter*>	RayTracingCallableBindings; // One per binding task
+	FGraphEventRef									RayTracingCallableBindingsTask;
+
 	// Common resources used for lighting in ray tracing effects
 	TRefCountPtr<FRHITexture>						RayTracingSubSurfaceProfileTexture;
 	FShaderResourceViewRHIRef						RayTracingSubSurfaceProfileSRV;
diff --git a/Engine/Source/Runtime/Renderer/Public/MeshPassProcessor.h b/Engine/Source/Runtime/Renderer/Public/MeshPassProcessor.h
--- a/Engine/Source/Runtime/Renderer/Public/MeshPassProcessor.h
+++ b/Engine/Source/Runtime/Renderer/Public/MeshPassProcessor.h
@@ -2064,7 +2064,7 @@
 	FRHIRayTracingShader* MaterialShader = nullptr;
 	uint32 MaterialShaderIndex = UINT_MAX;
 
-	uint32 GeometrySegmentIndex = ~0u;
+	uint32 GeometrySegmentIndex = UINT_MAX;
 	uint8 InstanceMask = 0xFF;
 
 	bool bCastRayTracedShadows = true;
@@ -2170,3 +2170,17 @@
 	uint32 GeometrySegmentIndex;
 	uint32 RayTracingInstanceIndex;
 };
+
+class FRayTracingShaderCommand
+{
+public:
+	FMeshDrawShaderBindings ShaderBindings;
+
+	FRHIRayTracingShader* Shader = nullptr;
+	uint32 ShaderIndex = UINT_MAX;
+
+	uint32 SlotInScene = UINT_MAX;
+
+	/** Sets ray tracing shader on the command and allocates room for the shader bindings. */
+	RENDERER_API void SetShader(const TShaderRef<FShader>& Shader);
+};
diff --git a/Engine/Source/Runtime/Renderer/Public/RayTracingInstanceBufferUtil.h b/Engine/Source/Runtime/Renderer/Public/RayTracingInstanceBufferUtil.h
--- a/Engine/Source/Runtime/Renderer/Public/RayTracingInstanceBufferUtil.h
+++ b/Engine/Source/Runtime/Renderer/Public/RayTracingInstanceBufferUtil.h
@@ -61,7 +61,8 @@
 	TArrayView<FRayTracingGeometryInstance> Instances,
 	uint8 NumLayers,
 	uint32 NumShaderSlotsPerGeometrySegment,
-	uint32 NumMissShaderSlots);
+	uint32 NumMissShaderSlots,
+	uint32 NumCallableShaderSlots = 0);
 
 UE_DEPRECATED(5.1, "Specify NumLayers instead.")
 RENDERER_API FRayTracingSceneWithGeometryInstances CreateRayTracingSceneWithGeometryInstances(
