Fix hair strands depth missing velocity depth output.

diff --git a/Engine/Shaders/Private/HairStrands/HairStrandsMaterialCommon.ush b/Engine/Shaders/Private/HairStrands/HairStrandsMaterialCommon.ush
--- a/Engine/Shaders/Private/HairStrands/HairStrandsMaterialCommon.ush
+++ b/Engine/Shaders/Private/HairStrands/HairStrandsMaterialCommon.ush
@@ -43,7 +43,7 @@
 struct FEvaluateOutput
 {
 	FHairSample NodeData;
-	float2 NodeVelocity;
+	float4 NodeVelocity;
 };
 
 struct FHairVertexData
@@ -147,13 +147,13 @@
 	
 
 	#if 1
-	float2 EncodedVelocity = 0;
+	float4 EncodedVelocity = 0;
 	{
 		const float4 ScreenPos		= mul(float4(Velocity_World_P.xyz, 1), ResolvedView.TranslatedWorldToClip);
 		const float4 PrevScreenPos	= mul(float4(Velocity_World_P_Prev.xyz, 1), ResolvedView.PrevTranslatedWorldToClip);
 
-		const float2 Velocity	= Calculate3DVelocity(ScreenPos, PrevScreenPos).xy;
-		EncodedVelocity			= EncodeVelocityToTexture(float3(Velocity, 0.0)).xy;
+		const float3 Velocity	= Calculate3DVelocity(ScreenPos, PrevScreenPos);
+		EncodedVelocity			= EncodeVelocityToTexture(Velocity);
 	}
 	#endif
 
diff --git a/Engine/Shaders/Private/HairStrands/HairStrandsVelocity.usf b/Engine/Shaders/Private/HairStrands/HairStrandsVelocity.usf
--- a/Engine/Shaders/Private/HairStrands/HairStrandsVelocity.usf
+++ b/Engine/Shaders/Private/HairStrands/HairStrandsVelocity.usf
@@ -20,7 +20,7 @@
 uint bNeedClear;
 
 Texture2D<uint> NodeIndex;
-Buffer<float2> NodeVelocity;
+Buffer<float4> NodeVelocity;
 StructuredBuffer<FPackedHairVis> NodeVis; 
 Texture2D<float> CoverageTexture;
 
@@ -46,9 +46,9 @@
 RWTexture2D<uint> OutResolveMaskTexture;
 
 #if PERMUTATION_OUTPUT_FORMAT == 0
-float2 ToOutput(float2 In) { return In; }
+float2 ToOutput(float4 In) { return In.xy; }
 #else
-float4 ToOutput(float2 In) { return float4(In, 0, 0); }
+float4 ToOutput(float4 In) { return In; }
 #endif
 
 [numthreads(TILE_PIXEL_SIZE_X, TILE_PIXEL_SIZE_Y, 1)]
@@ -70,21 +70,21 @@
 	const FNodeDesc NodeDesc = DecodeNodeDesc(NodeIndex.Load(uint3(PixelCoord, 0)));
 	bool bNeedFastResolve = false;
 	bool bValidVelocity = false;
-	float2 OutEncodedVelocity = 0;
+	float4 OutEncodedVelocity = 0;
 	if (NodeDesc.Count > 0)
 	{
 		// Store final sort node data 
 		#if PERMUTATION_VELOCITY == VELOCITY_TYPE_AVG
-		float2 AverageVelocity = 0;
+		float3 AverageVelocity = 0;
 		#endif
 
 		#if PERMUTATION_VELOCITY == VELOCITY_TYPE_CLOSEST
-		float2 ClosestEncodedVelocity = 0;
+		float4 ClosestEncodedVelocity = 0;
 		float ClosestDepth = 0; // Inverse-Z
 		#endif		
 
 		#if PERMUTATION_VELOCITY == VELOCITY_TYPE_MAX
-		float2 MaxEncodedVelocity = 0;
+		float4 MaxEncodedVelocity = 0;
 		float MaxVelocityMagnitude2 = 0;
 		#endif		
 
@@ -95,10 +95,10 @@
 			for (uint NodeIt = 0; NodeIt < NodeDesc.Count; ++NodeIt)
 			{
 				const uint SampleIndex = NodeDesc.Offset + NodeIt;
-				const float2 EncodedVelocity = NodeVelocity[SampleIndex];
+				const float4 EncodedVelocity = NodeVelocity[SampleIndex];
 
 				#if PERMUTATION_VELOCITY == VELOCITY_TYPE_AVG
-				AverageVelocity += DecodeVelocityFromTexture(float4(EncodedVelocity, 0.0, 0.0)).xy;
+				AverageVelocity += DecodeVelocityFromTexture(EncodedVelocity);
 				#endif			
 
 				#if PERMUTATION_VELOCITY == VELOCITY_TYPE_CLOSEST
@@ -111,7 +111,7 @@
 				#endif
 
 				#if PERMUTATION_VELOCITY == VELOCITY_TYPE_MAX
-				const float2 CurrentVelocity = DecodeVelocityFromTexture(float4(EncodedVelocity, 0.0, 0.0)).xy;
+				const float4 CurrentVelocity = DecodeVelocityFromTexture(EncodedVelocity);
 				const float CurrentVelocityMagnitude2 = dot(CurrentVelocity, CurrentVelocity);
 				if (CurrentVelocityMagnitude2 > MaxVelocityMagnitude2)
 				{
@@ -123,11 +123,11 @@
 				// If the velocity is above a certain threshold, the pixel will be resolve with a fast resolved. 
 				// This will result into a sharper, but more noisy output. However it sill avoid getting smearing 
 				// from TAA.
-				bNeedFastResolve = bNeedFastResolve || NeedFastResolve(EncodedVelocity, VelocityThreshold);
+				bNeedFastResolve = bNeedFastResolve || NeedFastResolve(EncodedVelocity.xy, VelocityThreshold);
 			}
 
 			#if PERMUTATION_VELOCITY == VELOCITY_TYPE_AVG
-			OutEncodedVelocity = EncodeVelocityToTexture(float3(AverageVelocity / max(NodeDesc.Count, 1u), 0.0)).xy;
+			OutEncodedVelocity = EncodeVelocityToTexture(float3(AverageVelocity / max(NodeDesc.Count, 1u)));
 			#endif		
 
 			#if PERMUTATION_VELOCITY == VELOCITY_TYPE_CLOSEST
diff --git a/Engine/Source/Runtime/Renderer/Private/HairStrands/HairStrandsVisibility.cpp b/Engine/Source/Runtime/Renderer/Private/HairStrands/HairStrandsVisibility.cpp
--- a/Engine/Source/Runtime/Renderer/Private/HairStrands/HairStrandsVisibility.cpp
+++ b/Engine/Source/Runtime/Renderer/Private/HairStrands/HairStrandsVisibility.cpp
@@ -705,7 +705,7 @@
 	SHADER_PARAMETER_RDG_BUFFER_SRV(StructuredBuffer<FPackedHairVis>, NodeVis)
 	SHADER_PARAMETER_RDG_BUFFER_SRV(Buffer<uint>, IndirectArgs)
 	SHADER_PARAMETER_RDG_BUFFER_UAV(RWStructuredBuffer<FPackedHairSample>, OutNodeData)
-	SHADER_PARAMETER_RDG_BUFFER_UAV(RWBuffer<float2>, OutNodeVelocity)
+	SHADER_PARAMETER_RDG_BUFFER_UAV(RWBuffer<float4>, OutNodeVelocity)
 END_GLOBAL_SHADER_PARAMETER_STRUCT()
 IMPLEMENT_STATIC_UNIFORM_BUFFER_STRUCT(FVisibilityMaterialPassUniformParameters, "MaterialPassParameters", SceneTextures);
 
@@ -770,7 +770,6 @@
 ///////////////////////////////////////////////////////////////////////////////////////////////////
 struct FMaterialPassOutput
 {
-	static const EPixelFormat VelocityFormat = PF_G16R16;
 	FRDGBufferRef NodeData = nullptr;
 	FRDGBufferRef NodeVelocity = nullptr;
 	FRDGTextureRef SampleLightingTexture = nullptr;
@@ -795,9 +794,13 @@
 
 	const uint32 MaxNodeCount = CompactNodeVis->Desc.NumElements;
 
+	// Sanity check
+	const EPixelFormat VelocityFormat = FVelocityRendering::GetFormat(ViewInfo->GetShaderPlatform());
+	check(VelocityFormat == PF_A16B16G16R16 || VelocityFormat == PF_G16R16);
+		
 	FMaterialPassOutput Output;
 	Output.NodeData				 = GraphBuilder.CreateBuffer(FRDGBufferDesc::CreateStructuredDesc(sizeof(HairStrandsVisibilityInternal::NodeData), MaxNodeCount), TEXT("Hair.CompactNodeData"));
-	Output.NodeVelocity			 = GraphBuilder.CreateBuffer(FRDGBufferDesc::CreateBufferDesc(4, CompactNodeVis->Desc.NumElements), TEXT("Hair.CompactNodeVelocity"));
+	Output.NodeVelocity			 = GraphBuilder.CreateBuffer(FRDGBufferDesc::CreateBufferDesc(VelocityFormat == PF_G16R16 ? 4 : 8, CompactNodeVis->Desc.NumElements), TEXT("Hair.CompactNodeVelocity"));
 	Output.SampleLightingTexture = AddClearLightSamplePass(GraphBuilder, ViewInfo, MaxNodeCount, CompactNodeCounter);
 
 	const uint32 ResolutionDim = FMath::CeilToInt(FMath::Sqrt(static_cast<float>(MaxNodeCount)));
@@ -848,7 +851,7 @@
 			UniformParameters->NodeCoord = GraphBuilder.CreateSRV(CompactNodeCoord, FHairStrandsVisibilityData::NodeCoordFormat);
 			UniformParameters->IndirectArgs = GraphBuilder.CreateSRV(IndirectArgBuffer);
 			UniformParameters->OutNodeData = GraphBuilder.CreateUAV(FRDGBufferUAVDesc(Output.NodeData));
-			UniformParameters->OutNodeVelocity = GraphBuilder.CreateUAV(FRDGBufferUAVDesc(Output.NodeVelocity, FMaterialPassOutput::VelocityFormat));
+			UniformParameters->OutNodeVelocity = GraphBuilder.CreateUAV(FRDGBufferUAVDesc(Output.NodeVelocity, VelocityFormat));
 
 			PassParameters->UniformBuffer = GraphBuilder.CreateUniformBuffer(UniformParameters);
 		}
@@ -1015,6 +1018,7 @@
 	const FIntPoint Resolution = OutVelocityTexture->Desc.Extent;
 	OutResolveMaskTexture = GraphBuilder.CreateTexture(FRDGTextureDesc::Create2D(Resolution, PF_R8_UINT, FClearValueBinding::None, TexCreate_UAV), TEXT("Hair.VelocityResolveMaskTexture"));
 
+	const EPixelFormat VelocityFormat = FVelocityRendering::GetFormat(View.GetShaderPlatform());
 	check(OutVelocityTexture->Desc.Format == PF_G16R16 || OutVelocityTexture->Desc.Format == PF_A16B16G16R16);
 	const bool bTwoChannelsOutput = OutVelocityTexture->Desc.Format == PF_G16R16;
 
@@ -1030,7 +1034,7 @@
 	PassParameters->CoverageThreshold = GetHairWriteVelocityCoverageThreshold();
 	PassParameters->NodeIndex = NodeIndex;
 	PassParameters->NodeVis = GraphBuilder.CreateSRV(NodeVis);
-	PassParameters->NodeVelocity = GraphBuilder.CreateSRV(NodeVelocity, FMaterialPassOutput::VelocityFormat);
+	PassParameters->NodeVelocity = GraphBuilder.CreateSRV(NodeVelocity, VelocityFormat);
 	PassParameters->CoverageTexture = CoverageTexture;
 	PassParameters->OutVelocityTexture = GraphBuilder.CreateUAV(OutVelocityTexture);
 	PassParameters->OutResolveMaskTexture = GraphBuilder.CreateUAV(OutResolveMaskTexture);
