Fully remove velocity data for a primitive when it stops moving.
Without this when static components are moved in editor they continue to write velocity after they stop moving.
This can be extra bad if the component was snapped to a final location so that it writes a bad velocity.

diff --git a/Engine/Source/Runtime/Renderer/Private/RendererScene.cpp b/Engine/Source/Runtime/Renderer/Private/RendererScene.cpp
--- a/Engine/Source/Runtime/Renderer/Private/RendererScene.cpp
+++ b/Engine/Source/Runtime/Renderer/Private/RendererScene.cpp
@@ -4347,7 +4347,7 @@
 	else if (PrimitiveSceneInfo->bRegisteredWithVelocityData)
 	{
 		PrimitiveSceneInfo->bRegisteredWithVelocityData = false;
-		VelocityData.RemoveFromScene(PrimitiveSceneInfo->PrimitiveComponentId);
+		VelocityData.RemoveFromScene(PrimitiveSceneInfo->PrimitiveComponentId, true);
 	}
 }
 
@@ -4710,7 +4710,7 @@
 				if (PrimitiveSceneInfo->bRegisteredWithVelocityData)
 				{
 					// Remove primitive's motion blur information.
-					VelocityData.RemoveFromScene(PrimitiveSceneInfo->PrimitiveComponentId);
+					VelocityData.RemoveFromScene(PrimitiveSceneInfo->PrimitiveComponentId, false);
 				}
 
 				// Unlink the primitive from its shadow parent.
diff --git a/Engine/Source/Runtime/Renderer/Private/ScenePrivate.h b/Engine/Source/Runtime/Renderer/Private/ScenePrivate.h
--- a/Engine/Source/Runtime/Renderer/Private/ScenePrivate.h
+++ b/Engine/Source/Runtime/Renderer/Private/ScenePrivate.h
@@ -2470,13 +2470,20 @@
 		}
 	}
 
-	void RemoveFromScene(FPrimitiveComponentId PrimitiveComponentId)
+	void RemoveFromScene(FPrimitiveComponentId PrimitiveComponentId, bool bImmediate)
 	{
-		FComponentVelocityData* VelocityData = ComponentData.Find(PrimitiveComponentId);
+		if (bImmediate)
+		{
+			ComponentData.Remove(PrimitiveComponentId);
+		}
+		else
+		{
+			FComponentVelocityData* VelocityData = ComponentData.Find(PrimitiveComponentId);
 
-		if (VelocityData)
-		{
-			VelocityData->PrimitiveSceneInfo = nullptr;
+			if (VelocityData)
+			{
+				VelocityData->PrimitiveSceneInfo = nullptr;
+			}
 		}
 	}
 
