RHI validation fix:  avoid redundant RHI frame begin/end when cubemap capture occurs during UEngine::Tick, and other Tick related functions where the RHI frame will already be active.  A boolean "bInsideTick" flag is passed to UReflectionCaptureComponent::UpdateReflectionCaptureContents, specifying whether the cubemap capture is happening during Tick.

The four places this applies include UEditorEngine::Tick, UGameEngine::Tick, FAdvancedPreviewScene::Tick, and FPreviewScene::UpdateCaptureContents.  The last example is called from FBlueprintEditor::Tick and FThumbnailPreviewScene::Tick.

The counter-example is "UEditorEngine::BuildReflectionCaptures", which doesn't appear to be called anywhere from a Tick related function.  It's called in startup functions, various commands, and as part of static lighting generation.  If we run into further validation errors, we could add a flag to that function to differentiate those cases.  Even though there's only one counter-example, I figured it was best to err on the side of not changing behavior of client code that may also be using the function, so by default "bInsideTick" is set to false.

#rnx

diff --git a/Engine/Source/Editor/AdvancedPreviewScene/Private/AdvancedPreviewScene.cpp b/Engine/Source/Editor/AdvancedPreviewScene/Private/AdvancedPreviewScene.cpp
--- a/Engine/Source/Editor/AdvancedPreviewScene/Private/AdvancedPreviewScene.cpp
+++ b/Engine/Source/Editor/AdvancedPreviewScene/Private/AdvancedPreviewScene.cpp
@@ -250,7 +250,8 @@
 		InstancedSkyMaterial->SetScalarParameterValueEditorOnly(FName("CubemapRotation"), Profile.LightingRigRotation / 360.0f);
 		InstancedSkyMaterial->PostEditChange();
 
-		UReflectionCaptureComponent::UpdateReflectionCaptureContents(PreviewWorld);
+		const bool bInsideTick = true;
+		UReflectionCaptureComponent::UpdateReflectionCaptureContents(PreviewWorld, nullptr, false, false, bInsideTick);
 		PreviewWorld->UpdateAllSkyCaptures();
 
 		PreviousRotation = Profile.LightingRigRotation;
diff --git a/Engine/Source/Editor/UnrealEd/Private/EditorEngine.cpp b/Engine/Source/Editor/UnrealEd/Private/EditorEngine.cpp
--- a/Engine/Source/Editor/UnrealEd/Private/EditorEngine.cpp
+++ b/Engine/Source/Editor/UnrealEd/Private/EditorEngine.cpp
@@ -1680,13 +1680,14 @@
 	}
 
 	static bool bFirstTick = true;
+	const bool bInsideTick = true;
 
 	// Skip updating reflection captures on the first update as the level will not be ready to display
 	if (!bFirstTick)
 	{
 		// Update sky light first because sky diffuse will be visible in reflection capture indirect specular
 		USkyLightComponent::UpdateSkyCaptureContents(EditorContext.World());
-		UReflectionCaptureComponent::UpdateReflectionCaptureContents(EditorContext.World());
+		UReflectionCaptureComponent::UpdateReflectionCaptureContents(EditorContext.World(), nullptr, false, false, bInsideTick);
 	}
 
 	EmitDynamicResolutionEvent(EDynamicResolutionStateEvent::BeginFrame);
@@ -1795,7 +1796,7 @@
 					{
 						// Update sky light first because sky diffuse will be visible in reflection capture indirect specular
 						USkyLightComponent::UpdateSkyCaptureContents(PlayWorld);
-						UReflectionCaptureComponent::UpdateReflectionCaptureContents(PlayWorld);
+						UReflectionCaptureComponent::UpdateReflectionCaptureContents(PlayWorld, nullptr, false, false, bInsideTick);
 					}
 
 					FKismetDebugUtilities::NotifyDebuggerOfEndOfGameFrame(PieContext.World());
diff --git a/Engine/Source/Runtime/Engine/Classes/Components/ReflectionCaptureComponent.h b/Engine/Source/Runtime/Engine/Classes/Components/ReflectionCaptureComponent.h
--- a/Engine/Source/Runtime/Engine/Classes/Components/ReflectionCaptureComponent.h
+++ b/Engine/Source/Runtime/Engine/Classes/Components/ReflectionCaptureComponent.h
@@ -107,8 +107,11 @@
 	/** Gets the radius that bounds the shape's influence, used for culling. */
 	virtual float GetInfluenceBoundingRadius() const PURE_VIRTUAL(UReflectionCaptureComponent::GetInfluenceBoundingRadius,return 0;);
 
-	/** Called each tick to recapture and queued reflection captures. */
-	ENGINE_API static void UpdateReflectionCaptureContents(UWorld* WorldToUpdate, const TCHAR* CaptureReason = nullptr, bool bVerifyOnlyCapturing = false, bool bCapturingForMobile = false);
+	/**
+	  * Generally called each tick to recapture any queued reflection captures.  In some cases, it's also called from Editor utility functions or commands.
+	  * Set "bInsideTick" to true if called from inside a Tick function, indicating a render frame is already active, and the capture doesn't need to start one.
+	  */
+	ENGINE_API static void UpdateReflectionCaptureContents(UWorld* WorldToUpdate, const TCHAR* CaptureReason = nullptr, bool bVerifyOnlyCapturing = false, bool bCapturingForMobile = false, bool bInsideTick = false);
 
 	ENGINE_API class FReflectionCaptureMapBuildData* GetMapBuildData() const;
 
diff --git a/Engine/Source/Runtime/Engine/Private/Components/ReflectionCaptureComponent.cpp b/Engine/Source/Runtime/Engine/Private/Components/ReflectionCaptureComponent.cpp
--- a/Engine/Source/Runtime/Engine/Private/Components/ReflectionCaptureComponent.cpp
+++ b/Engine/Source/Runtime/Engine/Private/Components/ReflectionCaptureComponent.cpp
@@ -1058,7 +1058,7 @@
 	}
 }
 
-void UReflectionCaptureComponent::UpdateReflectionCaptureContents(UWorld* WorldToUpdate, const TCHAR* CaptureReason, bool bVerifyOnlyCapturing, bool bCapturingForMobile)
+void UReflectionCaptureComponent::UpdateReflectionCaptureContents(UWorld* WorldToUpdate, const TCHAR* CaptureReason, bool bVerifyOnlyCapturing, bool bCapturingForMobile, bool bInsideTick)
 {
 	if (WorldToUpdate && WorldToUpdate->Scene
 		// Don't capture and read back capture contents if we are currently doing async shader compiling
@@ -1115,7 +1115,7 @@
 			}
 		}
 
-		WorldToUpdate->Scene->AllocateReflectionCaptures(WorldCombinedCaptures, CaptureReason, bVerifyOnlyCapturing, bCapturingForMobile);
+		WorldToUpdate->Scene->AllocateReflectionCaptures(WorldCombinedCaptures, CaptureReason, bVerifyOnlyCapturing, bCapturingForMobile, bInsideTick);
 	}
 }
 
diff --git a/Engine/Source/Runtime/Engine/Private/GameEngine.cpp b/Engine/Source/Runtime/Engine/Private/GameEngine.cpp
--- a/Engine/Source/Runtime/Engine/Private/GameEngine.cpp
+++ b/Engine/Source/Runtime/Engine/Private/GameEngine.cpp
@@ -1845,9 +1845,11 @@
 			// This won't work with actual level streaming though
 			if (Context.World()->AreAlwaysLoadedLevelsLoaded())
 			{
+				const bool bInsideTick = true;
+
 				// Update sky light first because it's considered direct lighting, sky diffuse will be visible in reflection capture indirect specular
 				USkyLightComponent::UpdateSkyCaptureContents(Context.World());
-				UReflectionCaptureComponent::UpdateReflectionCaptureContents(Context.World());
+				UReflectionCaptureComponent::UpdateReflectionCaptureContents(Context.World(), nullptr, false, false, bInsideTick);
 			}
 		}
 
diff --git a/Engine/Source/Runtime/Engine/Private/PreviewScene.cpp b/Engine/Source/Runtime/Engine/Private/PreviewScene.cpp
--- a/Engine/Source/Runtime/Engine/Private/PreviewScene.cpp
+++ b/Engine/Source/Runtime/Engine/Private/PreviewScene.cpp
@@ -196,8 +196,12 @@
 
 void FPreviewScene::UpdateCaptureContents()
 {
+	// This function is called from FAdvancedPreviewScene::Tick, FBlueprintEditor::Tick, and FThumbnailPreviewScene::Tick,
+	// so assume we are inside a Tick function.
+	const bool bInsideTick = true;
+
 	USkyLightComponent::UpdateSkyCaptureContents(PreviewWorld);
-	UReflectionCaptureComponent::UpdateReflectionCaptureContents(PreviewWorld);
+	UReflectionCaptureComponent::UpdateReflectionCaptureContents(PreviewWorld, nullptr, false, false, bInsideTick);
 }
 
 void FPreviewScene::ClearLineBatcher()
diff --git a/Engine/Source/Runtime/Engine/Public/SceneInterface.h b/Engine/Source/Runtime/Engine/Public/SceneInterface.h
--- a/Engine/Source/Runtime/Engine/Public/SceneInterface.h
+++ b/Engine/Source/Runtime/Engine/Public/SceneInterface.h
@@ -171,9 +171,15 @@
 	 * Allocates reflection captures in the scene's reflection cubemap array and updates them by recapturing the scene.
 	 * Existing captures will only be updated.  Must be called from the game thread.
 	 */
-	virtual void AllocateReflectionCaptures(const TArray<UReflectionCaptureComponent*>& NewCaptures, const TCHAR* CaptureReason, bool bVerifyOnlyCapturing, bool bCapturingForMobile) {}
+	virtual void AllocateReflectionCaptures(const TArray<UReflectionCaptureComponent*>& NewCaptures, const TCHAR* CaptureReason, bool bVerifyOnlyCapturing, bool bCapturingForMobile, bool bInsideTick) {}
 	virtual void ReleaseReflectionCubemap(UReflectionCaptureComponent* CaptureComponent) {}
 
+	UE_DEPRECATED(5.1, "This method now accepts a bInsideTick parameter, which specifies whether it's called during the frame Tick.")
+	inline void AllocateReflectionCaptures(const TArray<UReflectionCaptureComponent*>& NewCaptures, const TCHAR* CaptureReason, bool bVerifyOnlyCapturing, bool bCapturingForMobile)
+	{
+		AllocateReflectionCaptures(NewCaptures, CaptureReason, bVerifyOnlyCapturing, bCapturingForMobile, false);
+	}
+
 	/** 
 	 * Updates the contents of the given sky capture by rendering the scene. 
 	 * This must be called on the game thread.
diff --git a/Engine/Source/Runtime/Renderer/Private/ReflectionEnvironmentCapture.cpp b/Engine/Source/Runtime/Renderer/Private/ReflectionEnvironmentCapture.cpp
--- a/Engine/Source/Runtime/Renderer/Private/ReflectionEnvironmentCapture.cpp
+++ b/Engine/Source/Runtime/Renderer/Private/ReflectionEnvironmentCapture.cpp
@@ -742,7 +742,7 @@
  * Allocates reflection captures in the scene's reflection cubemap array and updates them by recapturing the scene.
  * Existing captures will only be uploaded.  Must be called from the game thread.
  */
-void FScene::AllocateReflectionCaptures(const TArray<UReflectionCaptureComponent*>& NewCaptures, const TCHAR* CaptureReason, bool bVerifyOnlyCapturing, bool bCapturingForMobile)
+void FScene::AllocateReflectionCaptures(const TArray<UReflectionCaptureComponent*>& NewCaptures, const TCHAR* CaptureReason, bool bVerifyOnlyCapturing, bool bCapturingForMobile, bool bInsideTick)
 {
 	if (NewCaptures.Num() > 0)
 	{
@@ -827,7 +827,7 @@
 					});
 
 				// Recapture all reflection captures now that we have reallocated the cubemap array
-				UpdateAllReflectionCaptures(CaptureReason, ReflectionCaptureSize, bVerifyOnlyCapturing, bCapturingForMobile);
+				UpdateAllReflectionCaptures(CaptureReason, ReflectionCaptureSize, bVerifyOnlyCapturing, bCapturingForMobile, bInsideTick);
 			}
 			else
 			{
@@ -852,7 +852,7 @@
 
 					if (bAllocated)
 					{
-						CaptureOrUploadReflectionCapture(CurrentComponent, ReflectionCaptureSize, bVerifyOnlyCapturing, bCapturingForMobile);
+						CaptureOrUploadReflectionCapture(CurrentComponent, ReflectionCaptureSize, bVerifyOnlyCapturing, bCapturingForMobile, bInsideTick);
 					}
 				}
 
@@ -878,7 +878,7 @@
 }
 
 /** Updates the contents of all reflection captures in the scene.  Must be called from the game thread. */
-void FScene::UpdateAllReflectionCaptures(const TCHAR* CaptureReason, int32 ReflectionCaptureSize, bool bVerifyOnlyCapturing, bool bCapturingForMobile)
+void FScene::UpdateAllReflectionCaptures(const TCHAR* CaptureReason, int32 ReflectionCaptureSize, bool bVerifyOnlyCapturing, bool bCapturingForMobile, bool bInsideTick)
 {
 	if (IsReflectionEnvironmentAvailable(GetFeatureLevel()))
 	{
@@ -902,7 +902,7 @@
 
 			CaptureIndex++;
 			UReflectionCaptureComponent* CurrentComponent = *It;
-			CaptureOrUploadReflectionCapture(CurrentComponent, ReflectionCaptureSize, bVerifyOnlyCapturing, bCapturingForMobile);
+			CaptureOrUploadReflectionCapture(CurrentComponent, ReflectionCaptureSize, bVerifyOnlyCapturing, bCapturingForMobile, bInsideTick);
 		}
 
 		EndReflectionCaptureSlowTask(NumCapturesForStatus);
@@ -1088,7 +1088,8 @@
 	bool bLowerHemisphereIsBlack, 
 	bool bCaptureEmissiveOnly,
 	const FLinearColor& LowerHemisphereColor,
-	bool bCapturingForMobile
+	bool bCapturingForMobile,
+	bool bInsideTick
 	)
 {
 	int32 SupersampleCaptureFactor = FMath::Clamp(GSupersampleCaptureFactor, MinSupersampleCaptureFactor, MaxSupersampleCaptureFactor);
@@ -1116,7 +1117,7 @@
 
 	for (int32 CubeFace = 0; CubeFace < CubeFace_MAX; CubeFace++)
 	{
-		if( !bCapturingForSkyLight )
+		if( !bCapturingForSkyLight && !bInsideTick )
 		{
 			// Alert the RHI that we're rendering a new frame
 			// Not really a new frame, but it will allow pooling mechanisms to update, like the uniform buffer pool
@@ -1209,11 +1210,11 @@
 		FSceneRenderer* SceneRenderer = FSceneRenderer::CreateSceneRenderer(&ViewFamily, nullptr);
 
 		ENQUEUE_RENDER_COMMAND(CaptureCommand)(
-			[SceneRenderer, &ReflectionCubemapTexture, CubeFace, CubemapSize, bCapturingForSkyLight, bLowerHemisphereIsBlack, LowerHemisphereColor, bCapturingForMobile](FRHICommandListImmediate& RHICmdList)
+			[SceneRenderer, &ReflectionCubemapTexture, CubeFace, CubemapSize, bCapturingForSkyLight, bLowerHemisphereIsBlack, LowerHemisphereColor, bCapturingForMobile, bInsideTick](FRHICommandListImmediate& RHICmdList)
 		{
 			CaptureSceneToScratchCubemap(RHICmdList, SceneRenderer, ReflectionCubemapTexture, (ECubeFace)CubeFace, CubemapSize, bCapturingForSkyLight, bLowerHemisphereIsBlack, LowerHemisphereColor, bCapturingForMobile);
 
-			if (!bCapturingForSkyLight)
+			if (!bCapturingForSkyLight && !bInsideTick)
 			{
 				RHICmdList.EndFrame();
 			}
@@ -1249,7 +1250,7 @@
  * Updates the contents of the given reflection capture by rendering the scene. 
  * This must be called on the game thread.
  */
-void FScene::CaptureOrUploadReflectionCapture(UReflectionCaptureComponent* CaptureComponent, int32 ReflectionCaptureSize, bool bVerifyOnlyCapturing, bool bCapturingForMobile)
+void FScene::CaptureOrUploadReflectionCapture(UReflectionCaptureComponent* CaptureComponent, int32 ReflectionCaptureSize, bool bVerifyOnlyCapturing, bool bCapturingForMobile, bool bInsideTick)
 {
 	if (IsReflectionEnvironmentAvailable(GetFeatureLevel()))
 	{
@@ -1328,7 +1329,7 @@
 				// Reflection Captures are a form of static lighting, so only capture scene elements that are static
 				// However if the project has static lighting disabled, Reflection Captures can still be made to work by capturing Movable lights
 				bool const bCaptureStaticSceneOnly = CVarReflectionCaptureStaticSceneOnly.GetValueOnGameThread() != 0 && bAllowStaticLighting;
-				CaptureSceneIntoScratchCubemap(this, *ReflectionCubemapTexture, CaptureComponent->GetComponentLocation() + CaptureComponent->CaptureOffset, ReflectionCaptureSize, false, bCaptureStaticSceneOnly, 0, false, false, FLinearColor(), bCapturingForMobile);
+				CaptureSceneIntoScratchCubemap(this, *ReflectionCubemapTexture, CaptureComponent->GetComponentLocation() + CaptureComponent->CaptureOffset, ReflectionCaptureSize, false, bCaptureStaticSceneOnly, 0, false, false, FLinearColor(), bCapturingForMobile, bInsideTick);
 			}
 			else if (CaptureComponent->ReflectionSourceType == EReflectionSourceType::SpecifiedCubemap)
 			{
@@ -1459,7 +1460,7 @@
 		{
 			const bool bStaticSceneOnly = CaptureComponent->Mobility == EComponentMobility::Static;
 			const bool bCapturingForMobile = false;
-			CaptureSceneIntoScratchCubemap(this, *ReflectionCubemapTexture, CaptureComponent->GetComponentLocation(), CubemapResolution, true, bStaticSceneOnly, CaptureComponent->SkyDistanceThreshold, bLowerHemisphereIsBlack, bCaptureEmissiveOnly, LowerHemisphereColor, bCapturingForMobile);
+			CaptureSceneIntoScratchCubemap(this, *ReflectionCubemapTexture, CaptureComponent->GetComponentLocation(), CubemapResolution, true, bStaticSceneOnly, CaptureComponent->SkyDistanceThreshold, bLowerHemisphereIsBlack, bCaptureEmissiveOnly, LowerHemisphereColor, bCapturingForMobile, false);
 		}
 		else if (CaptureComponent->SourceType == SLS_SpecifiedCubemap)
 		{
diff --git a/Engine/Source/Runtime/Renderer/Private/ScenePrivate.h b/Engine/Source/Runtime/Renderer/Private/ScenePrivate.h
--- a/Engine/Source/Runtime/Renderer/Private/ScenePrivate.h
+++ b/Engine/Source/Runtime/Renderer/Private/ScenePrivate.h
@@ -3096,7 +3096,7 @@ public:
 	virtual void UpdateSceneCaptureContents(class USceneCaptureComponent2D* CaptureComponent) override;
 	virtual void UpdateSceneCaptureContents(class USceneCaptureComponentCube* CaptureComponent) override;
 	virtual void UpdatePlanarReflectionContents(UPlanarReflectionComponent* CaptureComponent, FSceneRenderer& MainSceneRenderer) override;
-	virtual void AllocateReflectionCaptures(const TArray<UReflectionCaptureComponent*>& NewCaptures, const TCHAR* CaptureReason, bool bVerifyOnlyCapturing, bool bCapturingForMobile) override;
+	virtual void AllocateReflectionCaptures(const TArray<UReflectionCaptureComponent*>& NewCaptures, const TCHAR* CaptureReason, bool bVerifyOnlyCapturing, bool bCapturingForMobile, bool bInsideTick) override;
 	virtual void UpdateSkyCaptureContents(const USkyLightComponent* CaptureComponent, bool bCaptureEmissiveOnly, UTextureCube* SourceCubemap, FTexture* OutProcessedTexture, float& OutAverageBrightness, FSHVectorRGB3& OutIrradianceEnvironmentMap, TArray<FFloat16Color>* OutRadianceMap) override; 
 	virtual void AllocateAndCaptureFrameSkyEnvMap(FRDGBuilder& GraphBuilder, FSceneRenderer& SceneRenderer, FViewInfo& MainView, bool bShouldRenderSkyAtmosphere, bool bShouldRenderVolumetricCloud, FInstanceCullingManager& InstanceCullingManager) override;
 	virtual void ValidateSkyLightRealTimeCapture(FRDGBuilder& GraphBuilder, const FViewInfo& View, FRDGTextureRef SceneColorTexture) override;
@@ -3476,10 +3476,10 @@ private:
 	* Updates the contents of the given reflection capture by rendering the scene. 
 	* This must be called on the game thread.
 	*/
-	void CaptureOrUploadReflectionCapture(UReflectionCaptureComponent* CaptureComponent, int32 ReflectionCaptureSize, bool bVerifyOnlyCapturing, bool bCapturingForMobile);
+	void CaptureOrUploadReflectionCapture(UReflectionCaptureComponent* CaptureComponent, int32 ReflectionCaptureSize, bool bVerifyOnlyCapturing, bool bCapturingForMobile, bool bInsideTick);
 
 	/** Updates the contents of all reflection captures in the scene.  Must be called from the game thread. */
-	void UpdateAllReflectionCaptures(const TCHAR* CaptureReason, int32 ReflectionCaptureSize, bool bVerifyOnlyCapturing, bool bCapturingForMobile);
+	void UpdateAllReflectionCaptures(const TCHAR* CaptureReason, int32 ReflectionCaptureSize, bool bVerifyOnlyCapturing, bool bCapturingForMobile, bool bInsideTick);
 
 	/** Updates all static draw lists. */
 	void UpdateStaticDrawLists_RenderThread(FRHICommandListImmediate& RHICmdList);
