Fix crash when using material with UseHairStrands option enabled on a non-groom geometry.

diff --git a/Engine/Source/Runtime/Renderer/Private/SceneVisibility.cpp b/Engine/Source/Runtime/Renderer/Private/SceneVisibility.cpp
--- a/Engine/Source/Runtime/Renderer/Private/SceneVisibility.cpp
+++ b/Engine/Source/Runtime/Renderer/Private/SceneVisibility.cpp
@@ -3337,9 +3337,9 @@
 	}
 	
 	const bool bIsHairStrandsCompatible = ViewRelevance.bHairStrands && IsHairStrandsEnabled(EHairStrandsShaderType::All, View.GetShaderPlatform());
-	if (bIsHairStrandsCompatible && HairStrands::IsHairVisible(MeshBatch))
+	if (bIsHairStrandsCompatible)
 	{
-		if (HairStrands::IsHairStrandsVF(MeshBatch.Mesh))
+		if (HairStrands::IsHairStrandsVF(MeshBatch.Mesh) && HairStrands::IsHairVisible(MeshBatch))
 		{
 			View.HairStrandsMeshElements.AddUninitialized(1);
 			FMeshBatchAndRelevance& BatchAndProxy = View.HairStrandsMeshElements.Last();
