Path Tracer: Avoid invalidating the path tracer when a moveable light transform has been set but did not actually change

diff --git a/Engine/Source/Runtime/Renderer/Private/RendererScene.cpp b/Engine/Source/Runtime/Renderer/Private/RendererScene.cpp
--- a/Engine/Source/Runtime/Renderer/Private/RendererScene.cpp
+++ b/Engine/Source/Runtime/Renderer/Private/RendererScene.cpp
@@ -3128,11 +3128,16 @@ void FScene::UpdateLightTransform_RenderThread(FLightSceneInfo* LightSceneInfo,
 			LightSceneInfo->RemoveFromScene();
 		}
 
+		// Invalidate the path tracer if the transform actually changed
+		// NOTE: Position is derived from the Matrix, so there is no need to check it separately
+		if( !Parameters.LightToWorld.Equals(LightSceneInfo->Proxy->LightToWorld, SMALL_NUMBER) )
+		{
+			InvalidatePathTracedOutput();
+		}
+
 		// Update the light's transform and position.
 		LightSceneInfo->Proxy->SetTransform(Parameters.LightToWorld,Parameters.Position);
 
-		InvalidatePathTracedOutput();
-
 		// Also update the LightSceneInfoCompact
 		if( bHasId )
 		{
