Deprecated raw UAV access in VT API.

diff --git a/Engine/Source/Runtime/RenderCore/Public/VirtualTexturing.h b/Engine/Source/Runtime/RenderCore/Public/VirtualTexturing.h
--- a/Engine/Source/Runtime/RenderCore/Public/VirtualTexturing.h
+++ b/Engine/Source/Runtime/RenderCore/Public/VirtualTexturing.h
@@ -229,9 +229,16 @@
 /** Describes a location to write a single layer of a VT tile */
 struct FVTProduceTargetLayer
 {
+PRAGMA_DISABLE_DEPRECATION_WARNINGS
+	FVTProduceTargetLayer() = default;
+	FVTProduceTargetLayer(const FVTProduceTargetLayer&) = default;
+	FVTProduceTargetLayer& operator=(const FVTProduceTargetLayer&) = default;
+PRAGMA_ENABLE_DEPRECATION_WARNINGS
+
 	/** The texture to write to. */
 	FRHITexture* TextureRHI = nullptr;
-	/** The UAV to write to. This may be nullptr if no suitable UAV can be created for the texture format.  */
+	
+	UE_DEPRECATED(5.1, "UnorderedAccessViewRHI is deprecated. Register the pooled render target with RDG instead.")
 	FRHIUnorderedAccessView* UnorderedAccessViewRHI = nullptr;
 	/**
 	 * Pooled render target. For FRDGBuilder::RegisterExternalTexture() which only accepts pooled render targets.
diff --git a/Engine/Source/Runtime/Renderer/Private/VT/RuntimeVirtualTextureProducer.cpp b/Engine/Source/Runtime/Renderer/Private/VT/RuntimeVirtualTextureProducer.cpp
--- a/Engine/Source/Runtime/Renderer/Private/VT/RuntimeVirtualTextureProducer.cpp
+++ b/Engine/Source/Runtime/Renderer/Private/VT/RuntimeVirtualTextureProducer.cpp
@@ -75,7 +75,6 @@
 	for (int LayerIndex = 0; LayerIndex < RuntimeVirtualTexture::MaxTextureLayers; ++LayerIndex)
 	{
 		RenderPageBatchDesc.Targets[LayerIndex].Texture = Tiles[0].Targets[LayerIndex].TextureRHI != nullptr ? Tiles[0].Targets[LayerIndex].TextureRHI->GetTexture2D() : nullptr;
-		RenderPageBatchDesc.Targets[LayerIndex].UAV = Tiles[0].Targets[LayerIndex].UnorderedAccessViewRHI;
 		RenderPageBatchDesc.Targets[LayerIndex].PooledRenderTarget = Tiles[0].Targets[LayerIndex].PooledRenderTarget;
 	}
 
@@ -123,7 +122,6 @@
 			for (int LayerIndex = 0; LayerIndex < RuntimeVirtualTexture::MaxTextureLayers; ++LayerIndex)
 			{
 				RenderPageBatchDesc.Targets[LayerIndex].Texture = Tiles[0].Targets[LayerIndex].TextureRHI != nullptr ? Tiles[0].Targets[LayerIndex].TextureRHI->GetTexture2D() : nullptr;
-				RenderPageBatchDesc.Targets[LayerIndex].UAV = Tiles[0].Targets[LayerIndex].UnorderedAccessViewRHI;
 				RenderPageBatchDesc.Targets[LayerIndex].PooledRenderTarget = Tiles[0].Targets[LayerIndex].PooledRenderTarget;
 			}
 		}
diff --git a/Engine/Source/Runtime/Renderer/Private/VT/RuntimeVirtualTextureRender.cpp b/Engine/Source/Runtime/Renderer/Private/VT/RuntimeVirtualTextureRender.cpp
--- a/Engine/Source/Runtime/Renderer/Private/VT/RuntimeVirtualTextureRender.cpp
+++ b/Engine/Source/Runtime/Renderer/Private/VT/RuntimeVirtualTextureRender.cpp
@@ -1086,15 +1086,12 @@
 		bool bClearTextures,
 		bool bIsThumbnails,
 		FRHITexture2D* OutputTexture0,
-		FRHIUnorderedAccessView* OutputUAV0,
 		IPooledRenderTarget* OutputTarget0,
 		FBox2D const& DestBox0,
 		FRHITexture2D* OutputTexture1,
-		FRHIUnorderedAccessView* OutputUAV1,
 		IPooledRenderTarget* OutputTarget1,
 		FBox2D const& DestBox1,
 		FRHITexture2D* OutputTexture2, 
-		FRHIUnorderedAccessView* OutputUAV2,
 		IPooledRenderTarget* OutputTarget2,
 		FBox2D const& DestBox2,
 		FTransform const& UVToWorld,
@@ -1267,9 +1264,9 @@
 					InDesc.MaterialType,
 					InDesc.bClearTextures,
 					InDesc.bIsThumbnails,
-					InDesc.Targets[0].Texture, InDesc.Targets[0].UAV, InDesc.Targets[0].PooledRenderTarget, PageDesc.DestBox[0],
-					InDesc.Targets[1].Texture, InDesc.Targets[1].UAV, InDesc.Targets[1].PooledRenderTarget, PageDesc.DestBox[1],
-					InDesc.Targets[2].Texture, InDesc.Targets[2].UAV, InDesc.Targets[2].PooledRenderTarget, PageDesc.DestBox[2],
+					InDesc.Targets[0].Texture, InDesc.Targets[0].PooledRenderTarget, PageDesc.DestBox[0],
+					InDesc.Targets[1].Texture, InDesc.Targets[1].PooledRenderTarget, PageDesc.DestBox[1],
+					InDesc.Targets[2].Texture, InDesc.Targets[2].PooledRenderTarget, PageDesc.DestBox[2],
 					InDesc.UVToWorld,
 					InDesc.WorldBounds,
 					PageDesc.UVRange,
diff --git a/Engine/Source/Runtime/Renderer/Private/VT/VirtualTexturePhysicalSpace.cpp b/Engine/Source/Runtime/Renderer/Private/VT/VirtualTexturePhysicalSpace.cpp
--- a/Engine/Source/Runtime/Renderer/Private/VT/VirtualTexturePhysicalSpace.cpp
+++ b/Engine/Source/Runtime/Renderer/Private/VT/VirtualTexturePhysicalSpace.cpp
@@ -151,11 +151,6 @@
 
 		SRVCreateInfo.SRGBOverride = SRGBO_Default;
 		TextureSRV_SRGB[Layer] = RHICreateShaderResourceView(TextureRHI, SRVCreateInfo);
-
-		if (bCreateAliasedUAV)
-		{
-			TextureUAV[Layer] = PooledRenderTarget[Layer]->GetUAV();
-		}
 	}
 }
 
@@ -166,7 +161,6 @@
 		GRenderTargetPool.FreeUnusedResource(PooledRenderTarget[Layer]);
 		TextureSRV[Layer].SafeRelease();
 		TextureSRV_SRGB[Layer].SafeRelease();
-		TextureUAV[Layer].SafeRelease();
 	}
 }
 
diff --git a/Engine/Source/Runtime/Renderer/Private/VT/VirtualTexturePhysicalSpace.h b/Engine/Source/Runtime/Renderer/Private/VT/VirtualTexturePhysicalSpace.h
--- a/Engine/Source/Runtime/Renderer/Private/VT/VirtualTexturePhysicalSpace.h
+++ b/Engine/Source/Runtime/Renderer/Private/VT/VirtualTexturePhysicalSpace.h
@@ -97,11 +97,6 @@
 		return bSRGB ? TextureSRV_SRGB[Layer] : TextureSRV[Layer];
 	}
 
-	FRHIUnorderedAccessView* GetPhysicalTextureUAV(int32 Layer) const
-	{
-		return TextureUAV[Layer];
-	}
-	
 	TRefCountPtr<IPooledRenderTarget> GetPhysicalTexturePooledRenderTarget(int32 Layer) const
 	{
 		check(PooledRenderTarget[Layer].IsValid());
@@ -126,7 +121,6 @@
 	TRefCountPtr<IPooledRenderTarget> PooledRenderTarget[VIRTUALTEXTURE_SPACE_MAXLAYERS];
 	FShaderResourceViewRHIRef TextureSRV[VIRTUALTEXTURE_SPACE_MAXLAYERS];
 	FShaderResourceViewRHIRef TextureSRV_SRGB[VIRTUALTEXTURE_SPACE_MAXLAYERS];
-	FUnorderedAccessViewRHIRef TextureUAV[VIRTUALTEXTURE_SPACE_MAXLAYERS];
 
 	uint32 TextureSizeInTiles;
 	uint32 NumRefs;
diff --git a/Engine/Source/Runtime/Renderer/Private/VT/VirtualTextureSystem.cpp b/Engine/Source/Runtime/Renderer/Private/VT/VirtualTextureSystem.cpp
--- a/Engine/Source/Runtime/Renderer/Private/VT/VirtualTextureSystem.cpp
+++ b/Engine/Source/Runtime/Renderer/Private/VT/VirtualTextureSystem.cpp
@@ -2225,7 +2225,6 @@
 					if (Producer->GetPhysicalGroupIndexForTextureLayer(ProducerLayerIndex) == ProducerPhysicalGroupIndex)
 					{
 						ProduceTarget[ProducerLayerIndex].TextureRHI = PhysicalSpace->GetPhysicalTexture(PhysicalLocalTextureIndex);
-						ProduceTarget[ProducerLayerIndex].UnorderedAccessViewRHI = PhysicalSpace->GetPhysicalTextureUAV(PhysicalLocalTextureIndex);
 						ProduceTarget[ProducerLayerIndex].PooledRenderTarget = PhysicalSpace->GetPhysicalTexturePooledRenderTarget(PhysicalLocalTextureIndex);
 						ProduceTarget[ProducerLayerIndex].pPageLocation = PhysicalSpace->GetPhysicalLocation(pAddress);
 						LayerMask |= 1 << ProducerLayerIndex;
@@ -2401,7 +2400,6 @@
 								if (Producer.GetPhysicalGroupIndexForTextureLayer(ProducerLayerIndex) == ProducerPhysicalGroupIndex)
 								{
 									ProduceTarget[ProducerLayerIndex].TextureRHI = PhysicalSpace->GetPhysicalTexture(PhysicalLocalTextureIndex);
-									ProduceTarget[ProducerLayerIndex].UnorderedAccessViewRHI = PhysicalSpace->GetPhysicalTextureUAV(PhysicalLocalTextureIndex);
 									ProduceTarget[ProducerLayerIndex].PooledRenderTarget = PhysicalSpace->GetPhysicalTexturePooledRenderTarget(PhysicalLocalTextureIndex);
 									ProduceTarget[ProducerLayerIndex].pPageLocation = PhysicalSpace->GetPhysicalLocation(pAddress);
 									
@@ -2450,7 +2448,6 @@
 								if (Producer.GetPhysicalGroupIndexForTextureLayer(ProducerLayerIndex) == ProducerPhysicalGroupIndex)
 								{
 									ProduceTarget[ProducerLayerIndex].TextureRHI = PhysicalSpace->GetPhysicalTexture(PhysicalLocalTextureIndex);
-									ProduceTarget[ProducerLayerIndex].UnorderedAccessViewRHI = PhysicalSpace->GetPhysicalTextureUAV(PhysicalLocalTextureIndex);
 									ProduceTarget[ProducerLayerIndex].PooledRenderTarget = PhysicalSpace->GetPhysicalTexturePooledRenderTarget(PhysicalLocalTextureIndex);
 									ProduceTarget[ProducerLayerIndex].pPageLocation = PhysicalSpace->GetPhysicalLocation(pAddress);
 									PhysicalLocalTextureIndex++;
diff --git a/Engine/Source/Runtime/Renderer/Public/VT/RuntimeVirtualTextureRender.h b/Engine/Source/Runtime/Renderer/Public/VT/RuntimeVirtualTextureRender.h
--- a/Engine/Source/Runtime/Renderer/Public/VT/RuntimeVirtualTextureRender.h
+++ b/Engine/Source/Runtime/Renderer/Public/VT/RuntimeVirtualTextureRender.h
@@ -34,7 +34,8 @@
 	{
 		/** Physical texture to render to. */
 		FRHITexture* Texture = nullptr;
-		/** Unordered access view of physical texture to render to. If this exists we can render to it directly instead of using RHICopyTexture(). */
+
+		UE_DEPRECATED(5.1, "UAV is deprecated. Register the pooled render target with RDG instead.")
 		FRHIUnorderedAccessView* UAV = nullptr;
 
 		struct IPooledRenderTarget* PooledRenderTarget = nullptr;
