Ported IrradianceEnvironmentMap to RDG.

diff --git a/Engine/Plugins/Experimental/GPULightmass/Source/GPULightmass/Private/Scene/Scene.cpp b/Engine/Plugins/Experimental/GPULightmass/Source/GPULightmass/Private/Scene/Scene.cpp
--- a/Engine/Plugins/Experimental/GPULightmass/Source/GPULightmass/Private/Scene/Scene.cpp
+++ b/Engine/Plugins/Experimental/GPULightmass/Source/GPULightmass/Private/Scene/Scene.cpp
@@ -29,7 +29,6 @@
 
 #define LOCTEXT_NAMESPACE "StaticLightingSystem"
 
-extern RENDERER_API void SetupSkyIrradianceEnvironmentMapConstantsFromSkyIrradiance(FVector4f* OutSkyIrradianceEnvironmentMap, const FSHVectorRGB3 SkyIrradiance);
 extern ENGINE_API bool GCompressLightmaps;
 
 extern float GetTerrainExpandPatchCount(float LightMapRes, int32& X, int32& Y, int32 ComponentSize, int32 LightmapSize, int32& DesiredSize, uint32 LightingLOD);
diff --git a/Engine/Shaders/Public/ShaderVersion.ush b/Engine/Shaders/Public/ShaderVersion.ush
--- a/Engine/Shaders/Public/ShaderVersion.ush
+++ b/Engine/Shaders/Public/ShaderVersion.ush
@@ -1,5 +1,5 @@
-// This file is automatically generated by the console command r.InvalidateCachedShaders
-// Each time the console command is executed it generates a new GUID. As this file is included
-// in Platform.ush (which should be included in any shader) it allows to invalidate the shader DDC.
-// 
-// GUID = 65A1E683-92F7-4EAA-BA09-F696E39FC508
+// This file is automatically generated by the console command r.InvalidateCachedShaders
+// Each time the console command is executed it generates a new GUID. As this file is included
+// in Platform.ush (which should be included in any shader) it allows to invalidate the shader DDC.
+// 
+// GUID = 3FBEFF0E-9091-4DE4-9A24-D61106A33EDA
diff --git a/Engine/Source/Runtime/Renderer/Private/ReflectionEnvironment.cpp b/Engine/Source/Runtime/Renderer/Private/ReflectionEnvironment.cpp
--- a/Engine/Source/Runtime/Renderer/Private/ReflectionEnvironment.cpp
+++ b/Engine/Source/Runtime/Renderer/Private/ReflectionEnvironment.cpp
@@ -537,3 +537,90 @@
 IMPLEMENT_STATIC_UNIFORM_BUFFER_SLOT(ReflectionCapture);
 IMPLEMENT_STATIC_AND_SHADER_UNIFORM_BUFFER_STRUCT(FReflectionCaptureShaderData, "ReflectionCaptureSM5", ReflectionCapture);
 IMPLEMENT_STATIC_AND_SHADER_UNIFORM_BUFFER_STRUCT_EX(FMobileReflectionCaptureShaderData, "ReflectionCaptureES31", ReflectionCapture, FShaderParametersMetadata::EUsageFlags::NoEmulatedUniformBuffer);
+
+void SetupSkyIrradianceEnvironmentMapConstantsFromSkyIrradiance(FVector4f* OutSkyIrradianceEnvironmentMap, const FSHVectorRGB3 SkyIrradiance)
+{
+	const float SqrtPI = FMath::Sqrt(PI);
+	const float Coefficient0 = 1.0f / (2 * SqrtPI);
+	const float Coefficient1 = FMath::Sqrt(3.f) / (3 * SqrtPI);
+	const float Coefficient2 = FMath::Sqrt(15.f) / (8 * SqrtPI);
+	const float Coefficient3 = FMath::Sqrt(5.f) / (16 * SqrtPI);
+	const float Coefficient4 = .5f * Coefficient2;
+
+	// Pack the SH coefficients in a way that makes applying the lighting use the least shader instructions
+	// This has the diffuse convolution coefficients baked in
+	// See "Stupid Spherical Harmonics (SH) Tricks"
+	OutSkyIrradianceEnvironmentMap[0].X = -Coefficient1 * SkyIrradiance.R.V[3];
+	OutSkyIrradianceEnvironmentMap[0].Y = -Coefficient1 * SkyIrradiance.R.V[1];
+	OutSkyIrradianceEnvironmentMap[0].Z = Coefficient1 * SkyIrradiance.R.V[2];
+	OutSkyIrradianceEnvironmentMap[0].W = Coefficient0 * SkyIrradiance.R.V[0] - Coefficient3 * SkyIrradiance.R.V[6];
+
+	OutSkyIrradianceEnvironmentMap[1].X = -Coefficient1 * SkyIrradiance.G.V[3];
+	OutSkyIrradianceEnvironmentMap[1].Y = -Coefficient1 * SkyIrradiance.G.V[1];
+	OutSkyIrradianceEnvironmentMap[1].Z = Coefficient1 * SkyIrradiance.G.V[2];
+	OutSkyIrradianceEnvironmentMap[1].W = Coefficient0 * SkyIrradiance.G.V[0] - Coefficient3 * SkyIrradiance.G.V[6];
+
+	OutSkyIrradianceEnvironmentMap[2].X = -Coefficient1 * SkyIrradiance.B.V[3];
+	OutSkyIrradianceEnvironmentMap[2].Y = -Coefficient1 * SkyIrradiance.B.V[1];
+	OutSkyIrradianceEnvironmentMap[2].Z = Coefficient1 * SkyIrradiance.B.V[2];
+	OutSkyIrradianceEnvironmentMap[2].W = Coefficient0 * SkyIrradiance.B.V[0] - Coefficient3 * SkyIrradiance.B.V[6];
+
+	OutSkyIrradianceEnvironmentMap[3].X = Coefficient2 * SkyIrradiance.R.V[4];
+	OutSkyIrradianceEnvironmentMap[3].Y = -Coefficient2 * SkyIrradiance.R.V[5];
+	OutSkyIrradianceEnvironmentMap[3].Z = 3 * Coefficient3 * SkyIrradiance.R.V[6];
+	OutSkyIrradianceEnvironmentMap[3].W = -Coefficient2 * SkyIrradiance.R.V[7];
+
+	OutSkyIrradianceEnvironmentMap[4].X = Coefficient2 * SkyIrradiance.G.V[4];
+	OutSkyIrradianceEnvironmentMap[4].Y = -Coefficient2 * SkyIrradiance.G.V[5];
+	OutSkyIrradianceEnvironmentMap[4].Z = 3 * Coefficient3 * SkyIrradiance.G.V[6];
+	OutSkyIrradianceEnvironmentMap[4].W = -Coefficient2 * SkyIrradiance.G.V[7];
+
+	OutSkyIrradianceEnvironmentMap[5].X = Coefficient2 * SkyIrradiance.B.V[4];
+	OutSkyIrradianceEnvironmentMap[5].Y = -Coefficient2 * SkyIrradiance.B.V[5];
+	OutSkyIrradianceEnvironmentMap[5].Z = 3 * Coefficient3 * SkyIrradiance.B.V[6];
+	OutSkyIrradianceEnvironmentMap[5].W = -Coefficient2 * SkyIrradiance.B.V[7];
+
+	OutSkyIrradianceEnvironmentMap[6].X = Coefficient4 * SkyIrradiance.R.V[8];
+	OutSkyIrradianceEnvironmentMap[6].Y = Coefficient4 * SkyIrradiance.G.V[8];
+	OutSkyIrradianceEnvironmentMap[6].Z = Coefficient4 * SkyIrradiance.B.V[8];
+	OutSkyIrradianceEnvironmentMap[6].W = 1;
+}
+
+void UpdateSkyIrradianceGpuBuffer(FRHICommandListImmediate& RHICmdList, const FEngineShowFlags& EngineShowFlags, const FSkyLightSceneProxy* SkyLight, TRefCountPtr<FRDGPooledBuffer>& Buffer)
+{
+	TRACE_CPUPROFILER_EVENT_SCOPE(UpdateSkyIrradianceGpuBuffer);
+
+	const bool bUploadIrradiance =
+		SkyLight
+		// Skylights with static lighting already had their diffuse contribution baked into lightmaps
+		&& !SkyLight->bHasStaticLighting
+		&& EngineShowFlags.SkyLighting
+		&& !SkyLight->bRealTimeCaptureEnabled; // When bRealTimeCaptureEnabled is true, the buffer will be setup on GPU directly in this case
+
+	if (AllocatePooledBuffer(FRDGBufferDesc::CreateStructuredDesc(sizeof(FVector4f), SKY_IRRADIANCE_ENVIRONMENT_MAP_VEC4_COUNT), Buffer, TEXT("SkyIrradianceEnvironmentMap"), ERDGPooledBufferAlignment::None))
+	{
+		if (!bUploadIrradiance)
+		{
+			// Ensure that sky irradiance SH buffer contains sensible initial values (zero init).
+			// If there is no sky in the level, then nothing else may fill this buffer.
+			void* DataPtr = RHICmdList.LockBuffer(Buffer->GetRHI(), 0, Buffer->GetSize(), RLM_WriteOnly);
+			FPlatformMemory::Memset(DataPtr, 0, Buffer->GetSize());
+			RHICmdList.UnlockBuffer(Buffer->GetRHI());
+		}
+
+		RHICmdList.Transition(FRHITransitionInfo(Buffer->GetRHI(), ERHIAccess::Unknown, ERHIAccess::SRVMask));
+	}
+
+	if (bUploadIrradiance)
+	{
+		FVector4f UploadData[SKY_IRRADIANCE_ENVIRONMENT_MAP_VEC4_COUNT];
+		SetupSkyIrradianceEnvironmentMapConstantsFromSkyIrradiance(UploadData, SkyLight->IrradianceEnvironmentMap);
+
+		const float SkyLightAverageBrightness = SkyLight->AverageBrightness;
+		UploadData[7] = FVector4f(SkyLightAverageBrightness, SkyLightAverageBrightness, SkyLightAverageBrightness, SkyLightAverageBrightness);
+
+		void* DataPtr = RHICmdList.LockBuffer(Buffer->GetRHI(), 0, Buffer->GetSize(), RLM_WriteOnly);
+		FPlatformMemory::Memcpy(DataPtr, &UploadData, sizeof(UploadData));
+		RHICmdList.UnlockBuffer(Buffer->GetRHI());
+	}
+}
diff --git a/Engine/Source/Runtime/Renderer/Private/ReflectionEnvironment.h b/Engine/Source/Runtime/Renderer/Private/ReflectionEnvironment.h
--- a/Engine/Source/Runtime/Renderer/Private/ReflectionEnvironment.h
+++ b/Engine/Source/Runtime/Renderer/Private/ReflectionEnvironment.h
@@ -28,3 +28,7 @@
 
 extern void SetupReflectionUniformParameters(const class FViewInfo& View, FReflectionUniformParameters& OutParameters);
 TUniformBufferRef<FReflectionUniformParameters> CreateReflectionUniformBuffer(const class FViewInfo& View, EUniformBufferUsage Usage);
+
+RENDERER_API void SetupSkyIrradianceEnvironmentMapConstantsFromSkyIrradiance(FVector4f* OutSkyIrradianceEnvironmentMap, const FSHVectorRGB3 SkyIrradiance);
+
+extern void UpdateSkyIrradianceGpuBuffer(FRHICommandListImmediate& RHICmdList, const FEngineShowFlags& EngineShowFlags, const FSkyLightSceneProxy* SkyLight, TRefCountPtr<FRDGPooledBuffer>& Buffer);
diff --git a/Engine/Source/Runtime/Renderer/Private/ReflectionEnvironmentRealTimeCapture.cpp b/Engine/Source/Runtime/Renderer/Private/ReflectionEnvironmentRealTimeCapture.cpp
--- a/Engine/Source/Runtime/Renderer/Private/ReflectionEnvironmentRealTimeCapture.cpp
+++ b/Engine/Source/Runtime/Renderer/Private/ReflectionEnvironmentRealTimeCapture.cpp
@@ -132,7 +132,7 @@
 	BEGIN_SHADER_PARAMETER_STRUCT(FParameters, )
 		SHADER_PARAMETER_RDG_TEXTURE_SRV(TextureCube, SourceCubemapTexture)
 		SHADER_PARAMETER_SAMPLER(SamplerState, SourceCubemapSampler)
-		SHADER_PARAMETER_UAV(RWStructuredBuffer, OutIrradianceEnvMapSH)
+		SHADER_PARAMETER_RDG_BUFFER_UAV(RWStructuredBuffer, OutIrradianceEnvMapSH)
 		SHADER_PARAMETER(float, UniformSampleSolidAngle)
 		SHADER_PARAMETER(uint32, MipIndex)
 	END_SHADER_PARAMETER_STRUCT()
@@ -834,14 +834,9 @@
 
 	auto RenderCubeFaces_DiffuseIrradiance = [&](TRefCountPtr<IPooledRenderTarget>& SourceCubemap)
 	{
-		// ComputeDiffuseIrradiance using N uniform samples
-		AddPass(GraphBuilder, RDG_EVENT_NAME("TransitionToUAV"), [SkyIrradianceEnvironmentMap = SkyIrradianceEnvironmentMap.Buffer](FRHIComputeCommandList& RHICmdList)
-		{
-			RHICmdList.Transition({ SkyIrradianceEnvironmentMap, ERHIAccess::Unknown, ERHIAccess::UAVCompute });
-		});
-
 		FRDGTextureRef SourceCubemapTexture = GraphBuilder.RegisterExternalTexture(SourceCubemap);
 		FRDGTextureSRVRef SourceCubemapTextureSRV = GraphBuilder.CreateSRV(FRDGTextureSRVDesc::Create(SourceCubemapTexture));
+		FRDGBuffer* SkyIrradianceEnvironmentMapRDG = GraphBuilder.RegisterExternalBuffer(SkyIrradianceEnvironmentMap);
 
 		TShaderMapRef<FComputeSkyEnvMapDiffuseIrradianceCS> ComputeShader(GetGlobalShaderMap(FeatureLevel));
 
@@ -851,7 +846,7 @@
 		FComputeSkyEnvMapDiffuseIrradianceCS::FParameters* PassParameters = GraphBuilder.AllocParameters<FComputeSkyEnvMapDiffuseIrradianceCS::FParameters>();
 		PassParameters->SourceCubemapSampler = TStaticSamplerState<SF_Point>::GetRHI();
 		PassParameters->SourceCubemapTexture = SourceCubemapTextureSRV;
-		PassParameters->OutIrradianceEnvMapSH = SkyIrradianceEnvironmentMap.UAV;
+		PassParameters->OutIrradianceEnvMapSH = GraphBuilder.CreateUAV(SkyIrradianceEnvironmentMapRDG);
 		PassParameters->UniformSampleSolidAngle = UniformSampleSolidAngle;
 
 		// For 64 uniform samples on the unit sphere, we roughly have 10 samples per face.
@@ -863,10 +858,7 @@
 		const FIntVector NumGroups = FIntVector(1, 1, 1);
 		FComputeShaderUtils::AddPass(GraphBuilder, RDG_EVENT_NAME("ComputeSkyEnvMapDiffuseIrradianceCS"), ComputeShader, PassParameters, NumGroups);
 
-		AddPass(GraphBuilder, RDG_EVENT_NAME("TransitionToSRV"), [SkyIrradianceEnvironmentMap = SkyIrradianceEnvironmentMap.Buffer](FRHICommandList& RHICmdList)
-		{
-			RHICmdList.Transition({ SkyIrradianceEnvironmentMap, ERHIAccess::UAVCompute, ERHIAccess::SRVMask });
-		});
+		ExternalAccessQueue.Add(SkyIrradianceEnvironmentMapRDG, ERHIAccess::SRVMask);
 	};
 
 	const uint32 LastMipLevel = CubeMipCount - 1;
diff --git a/Engine/Source/Runtime/Renderer/Private/ScenePrivate.h b/Engine/Source/Runtime/Renderer/Private/ScenePrivate.h
--- a/Engine/Source/Runtime/Renderer/Private/ScenePrivate.h
+++ b/Engine/Source/Runtime/Renderer/Private/ScenePrivate.h
@@ -2821,7 +2821,7 @@ public:
 	FSkyLightSceneProxy* SkyLight;
 
 	/** Contains the sky env map irradiance as spherical harmonics. */
-	FRWBufferStructured SkyIrradianceEnvironmentMap;
+	TRefCountPtr<FRDGPooledBuffer> SkyIrradianceEnvironmentMap;
 
 	/** The SkyView LUT used when rendering sky material sampling this lut into the realtime capture sky env map. It must be generated at the skylight position*/
 	TRefCountPtr<IPooledRenderTarget> RealTimeReflectionCaptureSkyAtmosphereViewLutTexture;
diff --git a/Engine/Source/Runtime/Renderer/Private/SceneRendering.cpp b/Engine/Source/Runtime/Renderer/Private/SceneRendering.cpp
--- a/Engine/Source/Runtime/Renderer/Private/SceneRendering.cpp
+++ b/Engine/Source/Runtime/Renderer/Private/SceneRendering.cpp
@@ -1106,54 +1106,6 @@
 }
 #endif
 
-RENDERER_API void SetupSkyIrradianceEnvironmentMapConstantsFromSkyIrradiance(FVector4f* OutSkyIrradianceEnvironmentMap, const FSHVectorRGB3 SkyIrradiance)
-{
-	const float SqrtPI = FMath::Sqrt(PI);
-	const float Coefficient0 = 1.0f / (2 * SqrtPI);
-	const float Coefficient1 = FMath::Sqrt(3.f) / (3 * SqrtPI);
-	const float Coefficient2 = FMath::Sqrt(15.f) / (8 * SqrtPI);
-	const float Coefficient3 = FMath::Sqrt(5.f) / (16 * SqrtPI);
-	const float Coefficient4 = .5f * Coefficient2;
-
-	// Pack the SH coefficients in a way that makes applying the lighting use the least shader instructions
-	// This has the diffuse convolution coefficients baked in
-	// See "Stupid Spherical Harmonics (SH) Tricks"
-	OutSkyIrradianceEnvironmentMap[0].X = -Coefficient1 * SkyIrradiance.R.V[3];
-	OutSkyIrradianceEnvironmentMap[0].Y = -Coefficient1 * SkyIrradiance.R.V[1];
-	OutSkyIrradianceEnvironmentMap[0].Z = Coefficient1 * SkyIrradiance.R.V[2];
-	OutSkyIrradianceEnvironmentMap[0].W = Coefficient0 * SkyIrradiance.R.V[0] - Coefficient3 * SkyIrradiance.R.V[6];
-
-	OutSkyIrradianceEnvironmentMap[1].X = -Coefficient1 * SkyIrradiance.G.V[3];
-	OutSkyIrradianceEnvironmentMap[1].Y = -Coefficient1 * SkyIrradiance.G.V[1];
-	OutSkyIrradianceEnvironmentMap[1].Z = Coefficient1 * SkyIrradiance.G.V[2];
-	OutSkyIrradianceEnvironmentMap[1].W = Coefficient0 * SkyIrradiance.G.V[0] - Coefficient3 * SkyIrradiance.G.V[6];
-
-	OutSkyIrradianceEnvironmentMap[2].X = -Coefficient1 * SkyIrradiance.B.V[3];
-	OutSkyIrradianceEnvironmentMap[2].Y = -Coefficient1 * SkyIrradiance.B.V[1];
-	OutSkyIrradianceEnvironmentMap[2].Z = Coefficient1 * SkyIrradiance.B.V[2];
-	OutSkyIrradianceEnvironmentMap[2].W = Coefficient0 * SkyIrradiance.B.V[0] - Coefficient3 * SkyIrradiance.B.V[6];
-
-	OutSkyIrradianceEnvironmentMap[3].X = Coefficient2 * SkyIrradiance.R.V[4];
-	OutSkyIrradianceEnvironmentMap[3].Y = -Coefficient2 * SkyIrradiance.R.V[5];
-	OutSkyIrradianceEnvironmentMap[3].Z = 3 * Coefficient3 * SkyIrradiance.R.V[6];
-	OutSkyIrradianceEnvironmentMap[3].W = -Coefficient2 * SkyIrradiance.R.V[7];
-
-	OutSkyIrradianceEnvironmentMap[4].X = Coefficient2 * SkyIrradiance.G.V[4];
-	OutSkyIrradianceEnvironmentMap[4].Y = -Coefficient2 * SkyIrradiance.G.V[5];
-	OutSkyIrradianceEnvironmentMap[4].Z = 3 * Coefficient3 * SkyIrradiance.G.V[6];
-	OutSkyIrradianceEnvironmentMap[4].W = -Coefficient2 * SkyIrradiance.G.V[7];
-
-	OutSkyIrradianceEnvironmentMap[5].X = Coefficient2 * SkyIrradiance.B.V[4];
-	OutSkyIrradianceEnvironmentMap[5].Y = -Coefficient2 * SkyIrradiance.B.V[5];
-	OutSkyIrradianceEnvironmentMap[5].Z = 3 * Coefficient3 * SkyIrradiance.B.V[6];
-	OutSkyIrradianceEnvironmentMap[5].W = -Coefficient2 * SkyIrradiance.B.V[7];
-
-	OutSkyIrradianceEnvironmentMap[6].X = Coefficient4 * SkyIrradiance.R.V[8];
-	OutSkyIrradianceEnvironmentMap[6].Y = Coefficient4 * SkyIrradiance.G.V[8];
-	OutSkyIrradianceEnvironmentMap[6].Z = Coefficient4 * SkyIrradiance.B.V[8];
-	OutSkyIrradianceEnvironmentMap[6].W = 1;
-}
-
 void UpdateNoiseTextureParameters(FViewUniformShaderParameters& ViewUniformShaderParameters)
 {
 	if (GSystemTextures.PerlinNoiseGradient.GetReference())
@@ -1791,9 +1743,9 @@
 	}
 	else
 	{
-		if (Scene && Scene->SkyIrradianceEnvironmentMap.SRV)
+		if (Scene && Scene->SkyIrradianceEnvironmentMap)
 		{
-			ViewUniformShaderParameters.SkyIrradianceEnvironmentMap = Scene->SkyIrradianceEnvironmentMap.SRV;
+			ViewUniformShaderParameters.SkyIrradianceEnvironmentMap = Scene->SkyIrradianceEnvironmentMap->GetSRV();
 		}
 		else
 		{
@@ -5443,57 +5395,3 @@
 
 	return SmallDepthTexture;
 }
-
-void FSceneRenderer::UpdateSkyIrradianceGpuBuffer(FRHICommandListImmediate& RHICmdList)
-{
-	if (Scene == nullptr)
-	{
-		return;
-	}
-
-	if (!Scene->SkyIrradianceEnvironmentMap.Buffer)
-	{
-		Scene->SkyIrradianceEnvironmentMap.Initialize(TEXT("SkyIrradianceEnvironmentMap"), sizeof(FVector4f), SKY_IRRADIANCE_ENVIRONMENT_MAP_VEC4_COUNT);
-	}
-
-	TRACE_CPUPROFILER_EVENT_SCOPE(UpdateSkyIrradianceGpuBuffer);
-
-	FVector4f OutSkyIrradianceEnvironmentMap[SKY_IRRADIANCE_ENVIRONMENT_MAP_VEC4_COUNT];
-	// Make sure there's no padding since we're going to cast to FVector4f*
-	checkSlow(sizeof(OutSkyIrradianceEnvironmentMap) == sizeof(FVector4f) * SKY_IRRADIANCE_ENVIRONMENT_MAP_VEC4_COUNT);
-	check(Scene);
-	const bool bUploadIrradiance = 
-		Scene->SkyLight
-		// Skylights with static lighting already had their diffuse contribution baked into lightmaps
-		&& !Scene->SkyLight->bHasStaticLighting
-		&& ActiveViewFamily->EngineShowFlags.SkyLighting
-		&& !Scene->SkyLight->bRealTimeCaptureEnabled; // When bRealTimeCaptureEnabled is tru, the buffer will be setup on GPU directly in this case
-
-	if (bUploadIrradiance)
-	{
-		const FSHVectorRGB3& SkyIrradiance = Scene->SkyLight->IrradianceEnvironmentMap;
-		SetupSkyIrradianceEnvironmentMapConstantsFromSkyIrradiance(OutSkyIrradianceEnvironmentMap, SkyIrradiance);
-
-		const float SkyLightAverageBrightness = Scene->SkyLight->AverageBrightness;
-		const FVector4f SkyAverageBrightness = FVector4f(SkyLightAverageBrightness, SkyLightAverageBrightness, SkyLightAverageBrightness, SkyLightAverageBrightness);
-
-		// Set the captured environment map data
-		FVector4f* DataPtr = (FVector4f*)RHICmdList.LockBuffer(Scene->SkyIrradianceEnvironmentMap.Buffer, 0, Scene->SkyIrradianceEnvironmentMap.NumBytes, RLM_WriteOnly);
-		checkSlow(Scene->SkyIrradianceEnvironmentMap.NumBytes == sizeof(OutSkyIrradianceEnvironmentMap));
-		FPlatformMemory::Memcpy(DataPtr    , &OutSkyIrradianceEnvironmentMap,	sizeof(FVector4f) * 7);
-		FPlatformMemory::Memcpy(DataPtr + 7, &SkyAverageBrightness,				sizeof(FVector4f) * 1);
-		RHICmdList.UnlockBuffer(Scene->SkyIrradianceEnvironmentMap.Buffer);
-	}
-	else if (Scene->SkyIrradianceEnvironmentMap.NumBytes == 0)
-	{
-		// Ensure that sky irradiance SH buffer contains sensible initial values (zero init).
-		// If there is no sky in the level, then nothing else may fill this buffer.
-		void* DataPtr = RHICmdList.LockBuffer(Scene->SkyIrradianceEnvironmentMap.Buffer, 0, Scene->SkyIrradianceEnvironmentMap.NumBytes, RLM_WriteOnly);
-		FPlatformMemory::Memset(DataPtr, 0, Scene->SkyIrradianceEnvironmentMap.NumBytes);
-		RHICmdList.UnlockBuffer(Scene->SkyIrradianceEnvironmentMap.Buffer);
-	}
-
-	// This buffer is now going to be read for rendering.
-	RHICmdList.Transition(FRHITransitionInfo(Scene->SkyIrradianceEnvironmentMap.UAV, ERHIAccess::Unknown, ERHIAccess::SRVMask));
-}
-
diff --git a/Engine/Source/Runtime/Renderer/Private/SceneRendering.h b/Engine/Source/Runtime/Renderer/Private/SceneRendering.h
--- a/Engine/Source/Runtime/Renderer/Private/SceneRendering.h
+++ b/Engine/Source/Runtime/Renderer/Private/SceneRendering.h
@@ -2040,9 +2040,6 @@ public:
 
 	static int32 GetRefractionQuality(const FSceneViewFamily& ViewFamily);
 
-	/** Create/Update the scene view irradiance buffer from CPU data or empty if generated fully on GPU. */
-	void UpdateSkyIrradianceGpuBuffer(FRHICommandListImmediate& RHICmdList);
-
 	/** Common function to render a sky using shared LUT resources from any view point (if not using the SkyView and AerialPerspective textures). */
 	void RenderSkyAtmosphereInternal(
 		FRDGBuilder& GraphBuilder,
diff --git a/Engine/Source/Runtime/Renderer/Private/SceneVisibility.cpp b/Engine/Source/Runtime/Renderer/Private/SceneVisibility.cpp
--- a/Engine/Source/Runtime/Renderer/Private/SceneVisibility.cpp
+++ b/Engine/Source/Runtime/Renderer/Private/SceneVisibility.cpp
@@ -4932,7 +4932,10 @@ void FDeferredShadingSceneRenderer::InitViews(FRDGBuilder& GraphBuilder, const F
 	CreateIndirectCapsuleShadows();
 
 	// This must happen before we start initialising and using views.
-	UpdateSkyIrradianceGpuBuffer(RHICmdList);
+	if (Scene)
+	{
+		UpdateSkyIrradianceGpuBuffer(RHICmdList, ActiveViewFamily->EngineShowFlags, Scene->SkyLight, Scene->SkyIrradianceEnvironmentMap);
+	}
 
 	DispatchToRHIThread();
 
