Suppress static analysis error for potentially leaked virtual shadow map cache pointer.  Allocation is stored to the FSceneViewState class and freed later with that class.

#rnx

diff --git a/Engine/Source/Runtime/Renderer/Private/RendererScene.cpp b/Engine/Source/Runtime/Renderer/Private/RendererScene.cpp
--- a/Engine/Source/Runtime/Renderer/Private/RendererScene.cpp
+++ b/Engine/Source/Runtime/Renderer/Private/RendererScene.cpp
@@ -769,7 +769,7 @@ void FSceneViewState::AddVirtualShadowMapCache()
 				{
 					this->ViewVirtualShadowMapCache = ViewCache;
 				});
-		}
+		} //-V773
 	}
 }
 
