Fix non-unity
#rnx

diff --git a/Engine/Plugins/Experimental/Dataflow/Source/DataflowEditor/Public/Dataflow/DataflowEditorActions.h b/Engine/Plugins/Experimental/Dataflow/Source/DataflowEditor/Public/Dataflow/DataflowEditorActions.h
--- a/Engine/Plugins/Experimental/Dataflow/Source/DataflowEditor/Public/Dataflow/DataflowEditorActions.h
+++ b/Engine/Plugins/Experimental/Dataflow/Source/DataflowEditor/Public/Dataflow/DataflowEditorActions.h
@@ -7,7 +7,9 @@
 #include "Styling/AppStyle.h"
 #include "Dataflow/DataflowNodeParameters.h"
 
+class IDetailsView;
 class UDataflowEdNode;
+typedef TSet<class UObject*> FGraphPanelSelectionSet;
 
 class FDataflowEditorCommandsImpl : public TCommands<FDataflowEditorCommandsImpl>
 {
diff --git a/Engine/Plugins/Media/PixelStreamingVCam/Source/PixelStreamingVCam/Private/VCamPixelStreamingSession.cpp b/Engine/Plugins/Media/PixelStreamingVCam/Source/PixelStreamingVCam/Private/VCamPixelStreamingSession.cpp
--- a/Engine/Plugins/Media/PixelStreamingVCam/Source/PixelStreamingVCam/Private/VCamPixelStreamingSession.cpp
+++ b/Engine/Plugins/Media/PixelStreamingVCam/Source/PixelStreamingVCam/Private/VCamPixelStreamingSession.cpp
@@ -6,6 +6,7 @@
 #include "Editor/EditorPerformanceSettings.h"
 #include "Containers/UnrealString.h"
 #include "Async/Async.h"
+#include "Modules/ModuleManager.h"
 
 DEFINE_LOG_CATEGORY(LogVCamOutputProvider);
 namespace VCamPixelStreamingSession
diff --git a/Engine/Source/Runtime/CoreUObject/Private/UObject/PackageResourceManager.cpp b/Engine/Source/Runtime/CoreUObject/Private/UObject/PackageResourceManager.cpp
--- a/Engine/Source/Runtime/CoreUObject/Private/UObject/PackageResourceManager.cpp
+++ b/Engine/Source/Runtime/CoreUObject/Private/UObject/PackageResourceManager.cpp
@@ -3,6 +3,7 @@
 #include "UObject/PackageResourceManager.h"
 
 #include "Async/AsyncFileHandle.h"
+#include "Misc/CommandLine.h"
 #include "Misc/ConfigCacheIni.h"
 #include "Misc/PackageSegment.h"
 #include "Misc/PreloadableFile.h"
diff --git a/Engine/Source/Runtime/Renderer/Private/ReflectionEnvironment.h b/Engine/Source/Runtime/Renderer/Private/ReflectionEnvironment.h
--- a/Engine/Source/Runtime/Renderer/Private/ReflectionEnvironment.h
+++ b/Engine/Source/Runtime/Renderer/Private/ReflectionEnvironment.h
@@ -10,6 +10,11 @@
 #include "RHIDefinitions.h"
 #include "ShaderParameters.h"
 #include "UniformBuffer.h"
+#include "Math/SHMath.h"
+
+class FRDGPooledBuffer;
+class FSkyLightSceneProxy;
+struct FEngineShowFlags;
 
 extern bool IsReflectionEnvironmentAvailable(ERHIFeatureLevel::Type InFeatureLevel);
 extern bool IsReflectionCaptureAvailable();
@@ -31,4 +36,4 @@ TUniformBufferRef<FReflectionUniformParameters> CreateReflectionUniformBuffer(co
 
 RENDERER_API void SetupSkyIrradianceEnvironmentMapConstantsFromSkyIrradiance(FVector4f* OutSkyIrradianceEnvironmentMap, const FSHVectorRGB3 SkyIrradiance);
 
-extern void UpdateSkyIrradianceGpuBuffer(FRHICommandListImmediate& RHICmdList, const FEngineShowFlags& EngineShowFlags, const FSkyLightSceneProxy* SkyLight, TRefCountPtr<FRDGPooledBuffer>& Buffer);
\ No newline at end of file
+extern void UpdateSkyIrradianceGpuBuffer(FRHICommandListImmediate& RHICmdList, const FEngineShowFlags& EngineShowFlags, const FSkyLightSceneProxy* SkyLight, TRefCountPtr<FRDGPooledBuffer>& Buffer);
