Rewrites dynamic translucency resolution on top of the DynamicRenderScaling API

diff --git a/Engine/Source/Runtime/Renderer/Private/RendererScene.cpp b/Engine/Source/Runtime/Renderer/Private/RendererScene.cpp
--- a/Engine/Source/Runtime/Renderer/Private/RendererScene.cpp
+++ b/Engine/Source/Runtime/Renderer/Private/RendererScene.cpp
@@ -308,11 +308,6 @@ TRefCountPtr<IPooledRenderTarget> FPersistentSkyAtmosphereData::GetCurrentCamera
 /** Default constructor. */
 FSceneViewState::FSceneViewState(ERHIFeatureLevel::Type FeatureLevel)
 	: OcclusionQueryPool(RHICreateRenderQueryPool(RQT_Occlusion))
-	, TimerQueryPool(RHICreateRenderQueryPool(RQT_AbsoluteTime, FLatentGPUTimer::NumBufferedFrames * 2 * 2 * 2))
-	, TranslucencyTimer(TimerQueryPool)
-	, SeparateTranslucencyTimer(TimerQueryPool)
-	, SeparateTranslucencyModulateTimer(TimerQueryPool)
-	, PostMotionBlurTranslucencyTimer(TimerQueryPool)
 {
 	// Set FeatureLevel to a valid value, so we get Init/ReleaseDynamicRHI calls on FeatureLevel changes
 	SetFeatureLevel(FeatureLevel);
@@ -5679,131 +5674,3 @@ void FRendererModule::InvalidatePathTracedOutput()
 		Scene->InvalidatePathTracedOutput();
 	}
 }
-
-//////////////////////////////////////////////////////////////////////////
-
-FLatentGPUTimer::FLatentGPUTimer(FRenderQueryPoolRHIRef InTimerQueryPool, int32 InAvgSamples)
-: TimerQueryPool(InTimerQueryPool)
-, AvgSamples(InAvgSamples)
-, TotalTime(0.0f)
-, SampleIndex(0)
-, QueryIndex(0)
-{
-	TimeSamples.AddZeroed(AvgSamples);
-}
-
-bool FLatentGPUTimer::Tick(FRHICommandListImmediate& RHICmdList)
-{
-	if (GSupportsTimestampRenderQueries == false)
-	{
-		return false;
-	}
-
-	QueryIndex = (QueryIndex + 1) % NumBufferedFrames;
-
-	if (StartQueries[QueryIndex].GetQuery() && EndQueries[QueryIndex].GetQuery())
-	{
-		if (IsRunningRHIInSeparateThread())
-		{
-			// Block until the RHI thread has processed the previous query commands, if necessary
-			// Stat disabled since we buffer 2 frames minimum, it won't actually block
-			//SCOPE_CYCLE_COUNTER(STAT_TranslucencyTimestampQueryFence_Wait);
-			int32 BlockFrame = NumBufferedFrames - 1;
-			FRHICommandListExecutor::WaitOnRHIThreadFence(QuerySubmittedFences[BlockFrame]);
-			QuerySubmittedFences[BlockFrame] = nullptr;
-		}
-
-		uint64 StartMicroseconds;
-		uint64 EndMicroseconds;
-		bool bStartSuccess;
-		bool bEndSuccess;
-
-		{
-			// Block on the GPU until we have the timestamp query results, if necessary
-			// Stat disabled since we buffer 2 frames minimum, it won't actually block
-			//SCOPE_CYCLE_COUNTER(STAT_TranslucencyTimestampQuery_Wait);
-			bStartSuccess = RHICmdList.GetRenderQueryResult(StartQueries[QueryIndex].GetQuery(), StartMicroseconds, true);
-			bEndSuccess = RHICmdList.GetRenderQueryResult(EndQueries[QueryIndex].GetQuery(), EndMicroseconds, true);
-		}
-
-		TotalTime -= TimeSamples[SampleIndex];
-		float LastFrameTranslucencyDurationMS = TimeSamples[SampleIndex];
-		if (bStartSuccess && bEndSuccess)
-		{
-			LastFrameTranslucencyDurationMS = (EndMicroseconds - StartMicroseconds) / 1000.0f;
-		}
-
-		TimeSamples[SampleIndex] = LastFrameTranslucencyDurationMS;
-		TotalTime += LastFrameTranslucencyDurationMS;
-		SampleIndex = (SampleIndex + 1) % AvgSamples;
-
-		return bStartSuccess && bEndSuccess;
-	}
-
-	return false;
-}
-
-void FLatentGPUTimer::Begin(FRHICommandListImmediate& RHICmdList)
-{
-	if (GSupportsTimestampRenderQueries == false)
-	{
-		return;
-	}
-	
-	if (!StartQueries[QueryIndex].GetQuery())
-	{		
-		StartQueries[QueryIndex] = TimerQueryPool->AllocateQuery();
-	}
-
-	RHICmdList.EndRenderQuery(StartQueries[QueryIndex].GetQuery());
-}
-
-void FLatentGPUTimer::End(FRHICommandListImmediate& RHICmdList)
-{
-	if (GSupportsTimestampRenderQueries == false)
-	{
-		return;
-	}
-	
-	if (!EndQueries[QueryIndex].GetQuery())
-	{
-		EndQueries[QueryIndex] = TimerQueryPool->AllocateQuery();
-	}
-
-	RHICmdList.EndRenderQuery(EndQueries[QueryIndex].GetQuery());
-	// Hint to the RHI to submit commands up to this point to the GPU if possible.  Can help avoid CPU stalls next frame waiting
-	// for these query results on some platforms.
-	RHICmdList.SubmitCommandsHint();
-
-	if (IsRunningRHIInSeparateThread())
-	{
-		int32 NumFrames = NumBufferedFrames;
-		for (int32 Dest = 1; Dest < NumFrames; Dest++)
-		{
-			QuerySubmittedFences[Dest] = QuerySubmittedFences[Dest - 1];
-		}
-		// Start an RHI thread fence so we can be sure the RHI thread has processed the EndRenderQuery before we ask for results
-		QuerySubmittedFences[0] = RHICmdList.RHIThreadFence();
-		RHICmdList.ImmediateFlush(EImmediateFlushType::DispatchToRHIThread);
-	}
-}
-
-void FLatentGPUTimer::Release()
-{
-	for (int32 i = 0; i < NumBufferedFrames; ++i)
-	{
-		StartQueries[i].ReleaseQuery();
-		EndQueries[i].ReleaseQuery();
-	}
-}
-
-float FLatentGPUTimer::GetTimeMS()
-{
-	return TimeSamples[SampleIndex];
-}
-
-float FLatentGPUTimer::GetAverageTimeMS()
-{
-	return TotalTime / AvgSamples;
-}
-
diff --git a/Engine/Source/Runtime/Renderer/Private/ScenePrivate.h b/Engine/Source/Runtime/Renderer/Private/ScenePrivate.h
--- a/Engine/Source/Runtime/Renderer/Private/ScenePrivate.h
+++ b/Engine/Source/Runtime/Renderer/Private/ScenePrivate.h
@@ -506,53 +506,6 @@ private:
 extern FOcclusionRandomStream GOcclusionRandomStream;
 
 
-/**
-Helper class to time sections of the GPU work.
-Buffers multiple frames to avoid waiting on the GPU so times are a little lagged.
-*/
-class FLatentGPUTimer
-{
-	FRenderQueryPoolRHIRef TimerQueryPool;
-public:
-	static const int32 NumBufferedFrames = FOcclusionQueryHelpers::MaxBufferedOcclusionFrames + 1;
-
-	FLatentGPUTimer(FRenderQueryPoolRHIRef InTimerQueryPool, int32 InAvgSamples = 30);
-	~FLatentGPUTimer()
-	{
-		Release();
-	}
-
-	void Release();
-
-	/** Retrieves the most recently ready query results. */
-	bool Tick(FRHICommandListImmediate& RHICmdList);
-	/** Kicks off the query for the start of the rendering you're timing. */
-	void Begin(FRHICommandListImmediate& RHICmdList);
-	/** Kicks off the query for the end of the rendering you're timing. */
-	void End(FRHICommandListImmediate& RHICmdList);
-
-	/** Returns the most recent time in ms. */
-	float GetTimeMS();
-	/** Gets the average time in ms. Average is tracked over AvgSamples. */
-	float GetAverageTimeMS();
-
-private:
-
-	int32 GetQueryIndex();
-
-	//Average Tracking;
-	int32 AvgSamples;
-	TArray<float> TimeSamples;
-	float TotalTime;
-	int32 SampleIndex;
-
-	int32 QueryIndex;
-	bool QueriesInFlight[NumBufferedFrames];
-	FRHIPooledRenderQuery StartQueries[NumBufferedFrames];
-	FRHIPooledRenderQuery EndQueries[NumBufferedFrames];
-	FGraphEventRef QuerySubmittedFences[NumBufferedFrames];
-};
-
 /** HLOD tree persistent fading and visibility state */
 class FHLODVisibilityState
 {
@@ -1034,13 +987,6 @@ public:
 	TRefCountPtr<IPooledRenderTarget> GlobalDistanceFieldPageTableLayerTextures[GDF_Num];
 	TRefCountPtr<IPooledRenderTarget> GlobalDistanceFieldMipTexture;
 
-	/** Timestamp queries around separate translucency, used for auto-downsampling. */
-	FRenderQueryPoolRHIRef TimerQueryPool;
-	FLatentGPUTimer TranslucencyTimer;
-	FLatentGPUTimer SeparateTranslucencyTimer;
-	FLatentGPUTimer SeparateTranslucencyModulateTimer;
-	FLatentGPUTimer PostMotionBlurTranslucencyTimer;
-
 	/** This is float since it is derived off of UWorld::RealTimeSeconds, which is relative to BeginPlay time. */
 	float LastAutoDownsampleChangeTime;
 	float SmoothedHalfResTranslucencyGPUDuration;
diff --git a/Engine/Source/Runtime/Renderer/Private/SceneVisibility.cpp b/Engine/Source/Runtime/Renderer/Private/SceneVisibility.cpp
--- a/Engine/Source/Runtime/Renderer/Private/SceneVisibility.cpp
+++ b/Engine/Source/Runtime/Renderer/Private/SceneVisibility.cpp
@@ -5127,7 +5127,7 @@ void FDeferredShadingSceneRenderer::InitViewsAfterPrepass(FRDGBuilder& GraphBuil
 		UpdatePrimitiveIndirectLightingCacheBuffers();
 	}
 
-	SeparateTranslucencyDimensions = UpdateTranslucencyTimers(RHICmdList, Views);
+	SeparateTranslucencyDimensions = UpdateSeparateTranslucencyDimensions(*ActiveViewFamily);
 
 	SetupSceneReflectionCaptureBuffer(RHICmdList);
 
diff --git a/Engine/Source/Runtime/Renderer/Private/TranslucentRendering.cpp b/Engine/Source/Runtime/Renderer/Private/TranslucentRendering.cpp
--- a/Engine/Source/Runtime/Renderer/Private/TranslucentRendering.cpp
+++ b/Engine/Source/Runtime/Renderer/Private/TranslucentRendering.cpp
@@ -34,31 +34,29 @@
 	TEXT("<0 is treated like 100."),
 	ECVF_Scalability | ECVF_Default);
 
-static TAutoConsoleVariable<int32> CVarSeparateTranslucencyAutoDownsample(
-	TEXT("r.SeparateTranslucencyAutoDownsample"),
-	0,
-	TEXT("Whether to automatically downsample separate translucency based on last frame's GPU time.\n")
-	TEXT("Automatic downsampling is only used when r.SeparateTranslucencyScreenPercentage is 100"),
-	ECVF_Scalability | ECVF_Default);
+static TAutoConsoleVariable<float> CVarTranslucencyMinScreenPercentage(
+	TEXT("r.Translucency.DynamicRes.MinScreenPercentage"),
+	DynamicRenderScaling::FractionToPercentage(DynamicRenderScaling::FHeuristicSettings::kDefaultMinResolutionFraction),
+	TEXT("Minimal screen percentage for translucency."),
+	ECVF_RenderThreadSafe | ECVF_Default);
 
-static TAutoConsoleVariable<float> CVarSeparateTranslucencyDurationDownsampleThreshold(
-	TEXT("r.SeparateTranslucencyDurationDownsampleThreshold"),
-	1.5f,
-	TEXT("When smoothed full-res translucency GPU duration is larger than this value (ms), the entire pass will be downsampled by a factor of 2 in each dimension."),
-	ECVF_Scalability | ECVF_Default);
+static TAutoConsoleVariable<float> CVarTranslucencyTimeBudget(
+	TEXT("r.Translucency.DynamicRes.TimeBudget"),
+	DynamicRenderScaling::FHeuristicSettings::kBudgetMsDisabled,
+	TEXT("Frame's time budget for translucency rendering in milliseconds."),
+	ECVF_RenderThreadSafe | ECVF_Default);
 
-static TAutoConsoleVariable<float> CVarSeparateTranslucencyDurationUpsampleThreshold(
-	TEXT("r.SeparateTranslucencyDurationUpsampleThreshold"),
-	.5f,
-	TEXT("When smoothed half-res translucency GPU duration is smaller than this value (ms), the entire pass will be restored to full resolution.\n")
-	TEXT("This should be around 1/4 of r.SeparateTranslucencyDurationDownsampleThreshold to avoid toggling downsampled state constantly."),
-	ECVF_Scalability | ECVF_Default);
+static TAutoConsoleVariable<float> CVarTranslucencyTargetedHeadRoomPercentage(
+	TEXT("r.Translucency.DynamicRes.TargetedHeadRoomPercentage"),
+	DynamicRenderScaling::FractionToPercentage(DynamicRenderScaling::FHeuristicSettings::kDefaultTargetedHeadRoom),
+	TEXT("Targeted GPU headroom for translucency (in percent from r.DynamicRes.DynamicRes.TimeBudget)."),
+	ECVF_RenderThreadSafe | ECVF_Default);
 
-static TAutoConsoleVariable<float> CVarSeparateTranslucencyMinDownsampleChangeTime(
-	TEXT("r.SeparateTranslucencyMinDownsampleChangeTime"),
-	1.0f,
-	TEXT("Minimum time in seconds between changes to automatic downsampling state, used to prevent rapid swapping between half and full res."),
-	ECVF_Scalability | ECVF_Default);
+static TAutoConsoleVariable<float> CVarTranslucencyChangeThreshold(
+	TEXT("r.Translucency.DynamicRes.ChangePercentageThreshold"),
+	DynamicRenderScaling::FractionToPercentage(DynamicRenderScaling::FHeuristicSettings::kDefaultChangeThreshold),
+	TEXT("Minimal increase percentage threshold to alow when changing resolution of translucency."),
+	ECVF_RenderThreadSafe | ECVF_Default);
 
 int32 GSeparateTranslucencyUpsampleMode = 1;
 static FAutoConsoleVariableRef CVarSeparateTranslucencyUpsampleMode(
@@ -79,6 +77,22 @@
 	TEXT("Toggles parallel translucency rendering. Parallel rendering must be enabled for this to have an effect."),
 	ECVF_RenderThreadSafe);
 
+
+DynamicRenderScaling::FHeuristicSettings GetDynamicTranslucencyResolutionSettings()
+{
+	DynamicRenderScaling::FHeuristicSettings BucketSetting;
+	BucketSetting.Model = DynamicRenderScaling::EHeuristicModel::Quadratic;
+	BucketSetting.bModelScalesWithPrimaryScreenPercentage = true;
+	BucketSetting.MinResolutionFraction = DynamicRenderScaling::GetPercentageCVarToFraction(CVarTranslucencyMinScreenPercentage);
+	BucketSetting.BudgetMs              = CVarTranslucencyTimeBudget.GetValueOnAnyThread();
+	BucketSetting.ChangeThreshold       = DynamicRenderScaling::GetPercentageCVarToFraction(CVarTranslucencyChangeThreshold);
+	BucketSetting.TargetedHeadRoom      = DynamicRenderScaling::GetPercentageCVarToFraction(CVarTranslucencyTargetedHeadRoomPercentage);
+	return BucketSetting;
+}
+
+DynamicRenderScaling::FBudget GDynamicTranslucencyResolution(TEXT("DynamicTranslucencyResolution"), &GetDynamicTranslucencyResolutionSettings);
+
+
 static const TCHAR* kTranslucencyPassName[] = {
 	TEXT("BeforeDistortion"),
 	TEXT("AfterDOF"),
@@ -204,186 +218,24 @@
 	return true;
 }
 
-static void AddBeginTranslucencyTimerPass(FRDGBuilder& GraphBuilder, const FViewInfo& View)
+FSeparateTranslucencyDimensions UpdateSeparateTranslucencyDimensions(const FViewFamilyInfo& ActiveViewFamily)
 {
-#if STATS
-	if (View.ViewState)
+	float TranslucencyResolutionFraction = FMath::Clamp(CVarSeparateTranslucencyScreenPercentage.GetValueOnRenderThread() / 100.0f, 0.0f, 1.0f);
+	float MaxTranslucencyResolutionFraction	= TranslucencyResolutionFraction;
+
+	if (GDynamicTranslucencyResolution.GetSettings().IsEnabled())
 	{
-		AddPass(GraphBuilder, RDG_EVENT_NAME("BeginTimer"), [&View](FRHICommandListImmediate& RHICmdList)
-		{
-			View.ViewState->TranslucencyTimer.Begin(RHICmdList);
-		});
-	}
-#endif
-}
-
-static void AddEndTranslucencyTimerPass(FRDGBuilder& GraphBuilder, const FViewInfo& View)
-{
-#if STATS
-	if (View.ViewState)
-	{
-		AddPass(GraphBuilder, RDG_EVENT_NAME("EndTimer"), [&View](FRHICommandListImmediate& RHICmdList)
-		{
-			View.ViewState->TranslucencyTimer.End(RHICmdList);
-		});
-	}
-#endif
-}
-
-static bool HasSeparateTranslucencyTimer(const FViewInfo& View)
-{
-	return View.ViewState && GSupportsTimestampRenderQueries
-#if !STATS
-		&& (CVarSeparateTranslucencyAutoDownsample.GetValueOnRenderThread() != 0)
-#endif
-		;
-}
-
-static void AddBeginSeparateTranslucencyTimerPass(FRDGBuilder& GraphBuilder, const FViewInfo& View, ETranslucencyPass::Type TranslucencyPass)
-{
-	if (HasSeparateTranslucencyTimer(View))
-	{
-		AddPass(GraphBuilder, RDG_EVENT_NAME("BeginTimer"), [&View, TranslucencyPass](FRHICommandListImmediate& RHICmdList)
-		{
-			switch(TranslucencyPass)
-			{
-			case ETranslucencyPass::TPT_TranslucencyAfterDOF:
-				View.ViewState->SeparateTranslucencyTimer.Begin(RHICmdList);
-				break;
-			case ETranslucencyPass::TPT_TranslucencyAfterDOFModulate:
-				View.ViewState->SeparateTranslucencyModulateTimer.Begin(RHICmdList);
-				break;
-			case ETranslucencyPass::TPT_TranslucencyAfterMotionBlur:
-				View.ViewState->PostMotionBlurTranslucencyTimer.Begin(RHICmdList);
-				break;
-			default:
-				break;
-			}
-		});
-	}
-}
-
-static void AddEndSeparateTranslucencyTimerPass(FRDGBuilder& GraphBuilder, const FViewInfo& View, ETranslucencyPass::Type TranslucencyPass)
-{
-	if (HasSeparateTranslucencyTimer(View))
-	{
-		AddPass(GraphBuilder, RDG_EVENT_NAME("EndTimer"), [&View, TranslucencyPass](FRHICommandListImmediate& RHICmdList)
-		{
-			switch(TranslucencyPass)
-			{
-			case ETranslucencyPass::TPT_TranslucencyAfterDOF:
-				View.ViewState->SeparateTranslucencyTimer.End(RHICmdList);
-				break;
-			case ETranslucencyPass::TPT_TranslucencyAfterDOFModulate:			
-				View.ViewState->SeparateTranslucencyModulateTimer.End(RHICmdList);
-				break;
-			case ETranslucencyPass::TPT_TranslucencyAfterMotionBlur:
-				View.ViewState->PostMotionBlurTranslucencyTimer.End(RHICmdList);
-				break;
-			default:
-				break;
-			}
-		});
-	}
-}
-
-FSeparateTranslucencyDimensions UpdateTranslucencyTimers(FRHICommandListImmediate& RHICmdList, TArrayView<const FViewInfo> Views)
-{
-	bool bAnyViewWantsDownsampledSeparateTranslucency = false;
-
-	const bool bSeparateTranslucencyAutoDownsample = CVarSeparateTranslucencyAutoDownsample.GetValueOnRenderThread() != 0;
-	const bool bStatsEnabled = STATS != 0;
-
-	if (GSupportsTimestampRenderQueries && (bSeparateTranslucencyAutoDownsample || bStatsEnabled))
-	{
-		for (int32 ViewIndex = 0; ViewIndex < Views.Num(); ViewIndex++)
-		{
-			const FViewInfo& View = Views[ViewIndex];
-			SCOPED_GPU_MASK(RHICmdList, View.GPUMask);
-			FSceneViewState* ViewState = View.ViewState;
-
-			if (ViewState)
-			{
-				//We always tick the separate trans timer but only need the other timer for stats
-				bool bSeparateTransTimerSuccess = ViewState->SeparateTranslucencyTimer.Tick(RHICmdList);
-				bool bSeparateTransModulateTimerSuccess = ViewState->SeparateTranslucencyModulateTimer.Tick(RHICmdList);
-
-				if (STATS)
-				{
-					ViewState->TranslucencyTimer.Tick(RHICmdList);
-					//Stats are fed the most recent available time and so are lagged a little. 
-					float MostRecentTotalTime = ViewState->TranslucencyTimer.GetTimeMS() +
-						ViewState->SeparateTranslucencyTimer.GetTimeMS() +
-						ViewState->SeparateTranslucencyModulateTimer.GetTimeMS();
-					SET_FLOAT_STAT(STAT_TranslucencyGPU, MostRecentTotalTime);
-				}
-
-				if (bSeparateTranslucencyAutoDownsample && bSeparateTransTimerSuccess)
-				{
-					float LastFrameTranslucencyDurationMS = ViewState->SeparateTranslucencyTimer.GetTimeMS() + ViewState->SeparateTranslucencyModulateTimer.GetTimeMS();
-					const bool bOriginalShouldAutoDownsampleTranslucency = ViewState->bShouldAutoDownsampleTranslucency;
-
-					if (ViewState->bShouldAutoDownsampleTranslucency)
-					{
-						ViewState->SmoothedFullResTranslucencyGPUDuration = 0;
-						const float LerpAlpha = ViewState->SmoothedHalfResTranslucencyGPUDuration == 0 ? 1.0f : .1f;
-						ViewState->SmoothedHalfResTranslucencyGPUDuration = FMath::Lerp(ViewState->SmoothedHalfResTranslucencyGPUDuration, LastFrameTranslucencyDurationMS, LerpAlpha);
-
-						// Don't re-asses switching for some time after the last switch
-						if (View.Family->Time.GetRealTimeSeconds() - ViewState->LastAutoDownsampleChangeTime > CVarSeparateTranslucencyMinDownsampleChangeTime.GetValueOnRenderThread())
-						{
-							// Downsample if the smoothed time is larger than the threshold
-							ViewState->bShouldAutoDownsampleTranslucency = ViewState->SmoothedHalfResTranslucencyGPUDuration > CVarSeparateTranslucencyDurationUpsampleThreshold.GetValueOnRenderThread();
-
-							if (!ViewState->bShouldAutoDownsampleTranslucency)
-							{
-								// Do 'log LogRenderer verbose' to get these
-								UE_LOG(LogRenderer, Verbose, TEXT("Upsample: %.1fms < %.1fms"), ViewState->SmoothedHalfResTranslucencyGPUDuration, CVarSeparateTranslucencyDurationUpsampleThreshold.GetValueOnRenderThread());
-							}
-						}
-					}
-					else
-					{
-						ViewState->SmoothedHalfResTranslucencyGPUDuration = 0;
-						const float LerpAlpha = ViewState->SmoothedFullResTranslucencyGPUDuration == 0 ? 1.0f : .1f;
-						ViewState->SmoothedFullResTranslucencyGPUDuration = FMath::Lerp(ViewState->SmoothedFullResTranslucencyGPUDuration, LastFrameTranslucencyDurationMS, LerpAlpha);
-
-						if (View.Family->Time.GetRealTimeSeconds() - ViewState->LastAutoDownsampleChangeTime > CVarSeparateTranslucencyMinDownsampleChangeTime.GetValueOnRenderThread())
-						{
-							// Downsample if the smoothed time is larger than the threshold
-							ViewState->bShouldAutoDownsampleTranslucency = ViewState->SmoothedFullResTranslucencyGPUDuration > CVarSeparateTranslucencyDurationDownsampleThreshold.GetValueOnRenderThread();
-
-							if (ViewState->bShouldAutoDownsampleTranslucency)
-							{
-								UE_LOG(LogRenderer, Verbose, TEXT("Downsample: %.1fms > %.1fms"), ViewState->SmoothedFullResTranslucencyGPUDuration, CVarSeparateTranslucencyDurationDownsampleThreshold.GetValueOnRenderThread());
-							}
-						}
-					}
-
-					if (bOriginalShouldAutoDownsampleTranslucency != ViewState->bShouldAutoDownsampleTranslucency)
-					{
-						ViewState->LastAutoDownsampleChangeTime = View.Family->Time.GetRealTimeSeconds();
-					}
-
-					bAnyViewWantsDownsampledSeparateTranslucency = bAnyViewWantsDownsampledSeparateTranslucency || ViewState->bShouldAutoDownsampleTranslucency;
-				}
-			}
-		}
+		TranslucencyResolutionFraction = ActiveViewFamily.DynamicResolutionFractions[GDynamicTranslucencyResolution];
+		MaxTranslucencyResolutionFraction = ActiveViewFamily.DynamicResolutionUpperBounds[GDynamicTranslucencyResolution];
 	}
 
-	float EffectiveScale = FMath::Clamp(CVarSeparateTranslucencyScreenPercentage.GetValueOnRenderThread() / 100.0f, 0.0f, 1.0f);
+	UE_LOG(LogRenderer, Display, TEXT("TranslucencyResolutionFraction=%f"), TranslucencyResolutionFraction);
 
-	// 'r.SeparateTranslucencyScreenPercentage' CVar wins over automatic downsampling
-	if (FMath::IsNearlyEqual(EffectiveScale, 1.0f) && bAnyViewWantsDownsampledSeparateTranslucency)
-	{
-		EffectiveScale = 0.5f;
-	}
-
 	FSeparateTranslucencyDimensions Dimensions;
-	const FSceneTexturesConfig& Config = GetViewFamily(Views).SceneTexturesConfig;
-	Dimensions.Extent = GetScaledExtent(Config.Extent, EffectiveScale);
-	Dimensions.NumSamples = Config.NumSamples;
-	Dimensions.Scale = EffectiveScale;
+	// TODO: this should be MaxTranslucencyResolutionFraction instead of TranslucencyResolutionFraction to keep the size of render target stable, but the SvPositionToBuffer() is broken in material.
+	Dimensions.Extent = GetScaledExtent(ActiveViewFamily.SceneTexturesConfig.Extent, TranslucencyResolutionFraction);
+	Dimensions.NumSamples = ActiveViewFamily.SceneTexturesConfig.NumSamples;
+	Dimensions.Scale = TranslucencyResolutionFraction;
 	return Dimensions;
 }
 
@@ -596,6 +448,9 @@
 		SeparateTranslucencyTexture = GraphBuilder.RegisterExternalTexture(GSystemTextures.BlackAlphaOneDummy);
 	}
 
+	RDG_GPU_STAT_SCOPE(GraphBuilder, Translucency);
+	DynamicRenderScaling::FRDGScope DynamicTranslucencyResolutionScope(GraphBuilder, GDynamicTranslucencyResolution);
+
 	const TCHAR* OpName = nullptr;
 	FRHIBlendState* BlendState = nullptr;
 	FRDGTextureRef NewSceneColor = nullptr;
@@ -1355,7 +1210,6 @@
 	}
 
 	RDG_EVENT_SCOPE(GraphBuilder, "%s", TranslucencyPassToString(TranslucencyPass));
-	RDG_GPU_STAT_SCOPE(GraphBuilder, Translucency);
 	RDG_WAIT_FOR_TASKS_CONDITIONAL(GraphBuilder, IsTranslucencyWaitForTasksEnabled());
 
 	const bool bIsModulate = TranslucencyPass == ETranslucencyPass::TPT_TranslucencyAfterDOFModulate;
@@ -1405,8 +1259,6 @@
 				SeparateTranslucencyDepthTexture = SharedDepthTexture;
 			}
 
-			AddBeginSeparateTranslucencyTimerPass(GraphBuilder, View, TranslucencyPass);
-
 			const ERenderTargetLoadAction SeparateTranslucencyColorLoadAction = NumProcessedViews == 0 || View.Family->bMultiGPUForkAndJoin
 				? ERenderTargetLoadAction::EClear
 				: ERenderTargetLoadAction::ELoad;
@@ -1484,7 +1336,6 @@
 				}
 			}
 
-			AddEndSeparateTranslucencyTimerPass(GraphBuilder, View, TranslucencyPass);
 			++NumProcessedViews;
 		}
 	}
@@ -1503,8 +1354,6 @@
 			RDG_GPU_MASK_SCOPE(GraphBuilder, View.GPUMask);
 			RDG_EVENT_SCOPE_CONDITIONAL(GraphBuilder, Views.Num() > 1, "View%d", ViewIndex);
 
-			AddBeginTranslucencyTimerPass(GraphBuilder, View);
-
 			const ERenderTargetLoadAction SceneColorLoadAction = ERenderTargetLoadAction::ELoad;
 			const FScreenPassTextureViewport Viewport(SceneTextures.Color.Target, View.ViewRect);
 			const float ViewportScale = 1.0f;
@@ -1532,8 +1381,6 @@
 			{
 				OIT::AddOITComposePass(GraphBuilder, View, OITData, SceneTextures.Color.Target);
 			}
-
-			AddEndTranslucencyTimerPass(GraphBuilder, View);
 		}
 	}
 }
@@ -1551,6 +1398,9 @@
 		return;
 	}
 
+	RDG_GPU_STAT_SCOPE(GraphBuilder, Translucency);
+	DynamicRenderScaling::FRDGScope DynamicTranslucencyResolutionScope(GraphBuilder, GDynamicTranslucencyResolution);
+
 	FRDGTextureRef SceneColorCopyTexture = nullptr;
 
 	if (EnumHasAnyFlags(ViewsToRender, ETranslucencyView::AboveWater))
diff --git a/Engine/Source/Runtime/Renderer/Private/TranslucentRendering.h b/Engine/Source/Runtime/Renderer/Private/TranslucentRendering.h
--- a/Engine/Source/Runtime/Renderer/Private/TranslucentRendering.h
+++ b/Engine/Source/Runtime/Renderer/Private/TranslucentRendering.h
@@ -134,8 +134,8 @@
 /** Returns the union of all translucency views to render. */
 ETranslucencyView GetTranslucencyViews(TArrayView<const FViewInfo> Views);
 
-/** Call once per frame to update GPU timers for stats and dynamic resolution scaling. */
-FSeparateTranslucencyDimensions UpdateTranslucencyTimers(FRHICommandListImmediate& RHICmdList, TArrayView<const FViewInfo> Views);
+/** Computes the translucency dimensions. */
+FSeparateTranslucencyDimensions UpdateSeparateTranslucencyDimensions(const FViewFamilyInfo& ActiveViewFamily);
 
 /** Returns whether the view family is requesting to render translucency. */
 bool ShouldRenderTranslucency(const FSceneViewFamily& ViewFamily);
