Initial support for deferred decals in path tracer.
- Decals materials are evaluated using callable shaders in PathTracingKernel.
- Decals are culled using a 2D grid similar to the existing light grid.
- In order to correctly handle decal blending order, decals are sorted using the same logic as the rasterizer on CPU. The compute shader that builds the decal grid maintains the correct order.
- Decal materials are wrapped in FRayTracingDecalMaterialShader. The instance parameters of each decal are bound using uniform buffers.

diff --git a/Engine/Shaders/Private/PathTracing/PathTracingCore.ush b/Engine/Shaders/Private/PathTracing/PathTracingCore.ush
--- a/Engine/Shaders/Private/PathTracing/PathTracingCore.ush
+++ b/Engine/Shaders/Private/PathTracing/PathTracingCore.ush
@@ -5,6 +5,7 @@
 #define RANDSEQ_UNROLL_SOBOL PATH_TRACER_USE_COMPACTION // unrolling sobol sampler is best when we can guarantee high occupancy
 
 #define USE_PATH_TRACING_LIGHT_GRID	1
+#define USE_RAY_TRACING_DECAL_GRID	1
 #define USE_ENERGY_CONSERVATION 1
 
 #include "../Common.ush"
@@ -14,6 +15,7 @@
 
 #include "PathTracingCommon.ush"
 #include "../RayTracing/RayTracingHitGroupCommon.ush"
+#include "../RayTracing/RayTracingDecalGrid.ush"
 
 #include "../ShadingModels.ush"
 #include "./Utilities/PathTracingRandomSequence.ush"
@@ -46,6 +48,8 @@
 uint EnableEmissive;
 uint SamplerType;
 uint VisualizeLightGrid;
+uint VisualizeDecalGrid;
+uint EnableDecals;
 float MaxPathIntensity;
 float MaxNormalBias;
 float FilterWidth;
@@ -316,6 +320,63 @@
 		// Unpack the payload
 		FPathTracingPayload HitPayload = UnpackPathTracingPayloadMaterial(PackedPayload, PathState.Ray);
 
+#if PLATFORM_SUPPORTS_CALLABLE_SHADERS
+		if (EnableDecals)
+		{
+			float3 TranslatedWorldPos = HitPayload.TranslatedWorldPos;
+			float3 WorldNormal = HitPayload.WorldNormal;
+
+			FDecalLoopCount DecalLoopCount = DecalGridLookup(TranslatedWorldPos);
+
+			for (uint Index = 0, Num = DecalLoopCount.NumDecals; Index < Num; ++Index)
+			{
+				uint DecalId = GetDecalId(Index, DecalLoopCount);
+
+				FDecalCallableShaderParams Params;
+				Params.TranslatedWorldPosition_MetallicSpecularRoughnessAO.rgb = TranslatedWorldPos;
+				Params.WorldNormal.rgb = HitPayload.WorldNormal;
+				Params.Color = float4(HitPayload.BaseColor, 0.0f);
+				Params.Emissive = float3(0.0f, 0.0f, 0.0f);
+
+				CallShader(DecalId, Params);
+
+				const float DecalAlpha = Params.Color.a;
+
+				float DecalMetallic = Params.TranslatedWorldPosition_MetallicSpecularRoughnessAO.r;
+				if (DecalMetallic >= 0)
+				{
+					HitPayload.Metallic = lerp(HitPayload.Metallic, DecalMetallic, DecalAlpha);
+				}
+
+				float DecalSpecular = Params.TranslatedWorldPosition_MetallicSpecularRoughnessAO.g;
+				if (DecalSpecular >= 0)
+				{
+					HitPayload.Specular = lerp(HitPayload.Specular, DecalSpecular, DecalAlpha);
+				}
+
+				float3 DecalSpecularColor = ComputeF0(HitPayload.Specular, Params.Color.rgb, HitPayload.Metallic);
+				float3 DecalDiffuseColor = Params.Color.rgb - Params.Color.rgb * HitPayload.Metallic;
+
+				HitPayload.BaseColor = lerp(HitPayload.BaseColor, Params.Color.rgb, DecalAlpha);
+				HitPayload.DiffuseColor = lerp(HitPayload.DiffuseColor, DecalDiffuseColor, DecalAlpha);
+				HitPayload.SpecularColor = lerp(HitPayload.SpecularColor, DecalSpecularColor, DecalAlpha);
+
+				const float DecalRoughness = Params.TranslatedWorldPosition_MetallicSpecularRoughnessAO.b;
+				if (DecalRoughness >= 0.0f)
+				{
+					HitPayload.Roughness = lerp(HitPayload.Roughness, DecalRoughness, DecalAlpha);
+				}
+
+				HitPayload.WorldNormal = lerp(HitPayload.WorldNormal, Params.WorldNormal.rgb, DecalAlpha);
+
+				HitPayload.Radiance += Params.Emissive.rgb * DecalAlpha;
+
+				// TODO: DecalAO?
+				// TODO: Payload.Opacity?
+			}
+		}
+#endif
+
 		if (HitPayload.BlendingMode == RAY_TRACING_BLEND_MODE_ADDITIVE ||
 			HitPayload.BlendingMode == RAY_TRACING_BLEND_MODE_TRANSLUCENT)
 		{
@@ -959,6 +1020,16 @@
 		return false;
 	}
 
+	// visualize the # of decals
+	if (VisualizeDecalGrid)
+	{
+		FDecalLoopCount DecalLoopCount = DecalGridLookup(Payload.TranslatedWorldPos);
+		PathState.Radiance = DecalGridVisualize(DecalLoopCount, VisualizeDecalGrid);
+		PathState.Radiance *= abs(dot(Payload.WorldNormal, PathState.Ray.Direction));
+		PathState.Radiance /= View.PreExposure; // cancel out exposure
+		return false;
+	}
+
 	// Choose a random number for both Light sampling and BxDF sampling
 	float4 RandSample = RandomSequence_GenerateSample4D(PathState.RandSequence);
 
diff --git a/Engine/Shaders/Private/RayTracing/RayTracingBuildDecalGrid.usf b/Engine/Shaders/Private/RayTracing/RayTracingBuildDecalGrid.usf
new file mode 100644
--- /dev/null
+++ b/Engine/Shaders/Private/RayTracing/RayTracingBuildDecalGrid.usf
@@ -0,0 +1,62 @@
+// Copyright Epic Games, Inc. All Rights Reserved.
+
+#include "../Common.ush"
+#include "/Engine/Shared/RayTracingDefinitions.h"
+#include "/Engine/Shared/RayTracingTypes.h"
+#include "/Engine/Shared/PathTracingDefinitions.h"
+
+StructuredBuffer<FRayTracingDecal> SceneDecals;
+uint SceneDecalCount;
+
+float3 SceneDecalsTranslatedBoundMin;
+float3 SceneDecalsTranslatedBoundMax;
+uint DecalGridResolution;
+uint DecalGridMaxCount;
+uint DecalGridAxis;
+
+RWTexture2D<uint> RWDecalGrid;
+RWBuffer<uint> RWDecalGridData;
+
+[numthreads(THREADGROUPSIZE_X, THREADGROUPSIZE_Y, 1)]
+void BuildDecalGridCS(uint2 DispatchThreadId : SV_DispatchThreadID)
+{
+	if (any(DispatchThreadId >= DecalGridResolution))
+	{
+		return;
+	}
+	// figure out dimension of the 3d grid and the current ID to be filled
+	uint3 VoxelId = 0, VoxelRes = 1;
+	switch (DecalGridAxis)
+	{
+	case 0: VoxelId.yz = DispatchThreadId; VoxelRes.yz = DecalGridResolution; break;
+	case 1: VoxelId.xz = DispatchThreadId; VoxelRes.xz = DecalGridResolution; break;
+	case 2: VoxelId.xy = DispatchThreadId; VoxelRes.xy = DecalGridResolution; break;
+	}
+	// get bounding box of current voxel
+	float3 Lo = lerp(SceneDecalsTranslatedBoundMin, SceneDecalsTranslatedBoundMax, float3(VoxelId + 0) / float3(VoxelRes));
+	float3 Hi = lerp(SceneDecalsTranslatedBoundMin, SceneDecalsTranslatedBoundMax, float3(VoxelId + 1) / float3(VoxelRes));
+
+	uint NumVisible = 0;
+	uint MaxVisible = min(DecalGridMaxCount, RAY_TRACING_DECAL_COUNT_MAXIMUM);
+	uint Offset = DecalGridMaxCount * (DispatchThreadId.x + DecalGridResolution * DispatchThreadId.y);
+
+	for (uint DecalIndex = 0; DecalIndex < SceneDecalCount; DecalIndex++)
+	{
+		FRayTracingDecal Decal = SceneDecals[DecalIndex];
+
+		if (any(Hi < Decal.TranslatedBoundMin) || any(Lo > Decal.TranslatedBoundMax))
+		{
+			continue;
+		}
+
+		RWDecalGridData[Offset + NumVisible] = Decal.CallableSlotIndex;
+		NumVisible++;
+		// TODO: handle overflow better?
+		if (NumVisible >= MaxVisible)
+		{
+			break;
+		}
+	}
+
+	RWDecalGrid[DispatchThreadId] = NumVisible;
+}
diff --git a/Engine/Shaders/Private/RayTracing/RayTracingCommon.ush b/Engine/Shaders/Private/RayTracing/RayTracingCommon.ush
--- a/Engine/Shaders/Private/RayTracing/RayTracingCommon.ush
+++ b/Engine/Shaders/Private/RayTracing/RayTracingCommon.ush
@@ -857,6 +857,23 @@
 #endif // !COMPUTESHADER
 
 ///////////////////////////////////////////////////////////////////////////////////////////////////
+// Decal parameters
+
+#if !COMPUTESHADER
+
+struct FDecalCallableShaderParams
+{
+	// In order to reduce params size reuse input params to store output
+	// Input: TranslatedWorldPosition / Output: MetallicSpecularRoughnessAO
+	float4 TranslatedWorldPosition_MetallicSpecularRoughnessAO;
+	float4 Color;
+	float3 WorldNormal;
+	float3 Emissive;
+};
+
+#endif // !COMPUTESHADER
+
+///////////////////////////////////////////////////////////////////////////////////////////////////
 // Tracing util functions
 
 #if !COMPUTESHADER
diff --git a/Engine/Shaders/Private/RayTracing/RayTracingDecalGrid.ush b/Engine/Shaders/Private/RayTracing/RayTracingDecalGrid.ush
new file mode 100644
--- /dev/null
+++ b/Engine/Shaders/Private/RayTracing/RayTracingDecalGrid.ush
@@ -0,0 +1,109 @@
+// Copyright Epic Games, Inc. All Rights Reserved.
+
+#ifndef USE_RAY_TRACING_DECAL_GRID
+#define USE_RAY_TRACING_DECAL_GRID 0
+#endif
+
+#include "../Common.ush"
+#include "../Visualization.ush"
+#include "/Engine/Shared/RayTracingDefinitions.h"
+
+struct FDecalLoopCount {
+	uint NumDecals;
+#if USE_RAY_TRACING_DECAL_GRID
+	uint DecalGridOffset;
+#endif
+};
+
+#if USE_RAY_TRACING_DECAL_GRID
+
+FDecalLoopCount DecalGridLookup(float3 TranslatedWorldPos)
+{
+	FDecalLoopCount Result;
+	Result.NumDecals = 0;
+	Result.DecalGridOffset = ~0u;
+
+	if (all(RayTracingDecals.TranslatedBoundMin <= TranslatedWorldPos) && all(TranslatedWorldPos <= RayTracingDecals.TranslatedBoundMax))
+	{
+		float3 P = TranslatedWorldPos - RayTracingDecals.TranslatedBoundMin;
+		float3 D = RayTracingDecals.TranslatedBoundMax - RayTracingDecals.TranslatedBoundMin;
+		int2 UV = 0;
+		switch (RayTracingDecals.GridAxis)
+		{
+		case 0: UV = int2(floor(RayTracingDecals.GridResolution * P.yz / D.yz)); break;
+		case 1: UV = int2(floor(RayTracingDecals.GridResolution * P.xz / D.xz)); break;
+		case 2: UV = int2(floor(RayTracingDecals.GridResolution * P.xy / D.xy)); break;
+		}
+		Result.DecalGridOffset = RayTracingDecals.GridMaxCount * (UV.x + UV.y * RayTracingDecals.GridResolution);
+		Result.NumDecals += RayTracingDecals.Grid.Load(int3(UV, 0));
+	}
+	return Result;
+}
+
+float3 DecalGridVisualize(FDecalLoopCount LoopCount, int VisualizeMode)
+{
+	const float3 OutsideColor = 0.18;
+	const float3 EmptyColor = 0.36;
+
+	if (LoopCount.DecalGridOffset != ~0u)
+	{
+		switch (VisualizeMode)
+		{
+		case 2:
+		{
+			// color by unique decal list
+			uint H = 0;
+			for (int Index = 0, Num = LoopCount.NumDecals; Index < Num; Index++)
+			{
+				H = MurmurMix(H + RayTracingDecals.GridData[LoopCount.DecalGridOffset + Index]);
+			}
+			return IntToColor(H);
+		}
+		default:
+		{
+			// default mode - color by decal count
+			uint N = LoopCount.NumDecals;
+			if (N < 1)
+			{
+				return EmptyColor;
+			}
+			float Max = RayTracingDecals.GridMaxCount;
+			float t = saturate(float(N) / Max);
+			return t * saturate(2.0 - abs(float3(4, 2, 0) - 4.0 * t));
+		}
+		}
+	}
+	// outside decal grid bounds
+	return OutsideColor;
+}
+
+uint GetDecalId(uint Index, FDecalLoopCount LoopCount)
+{
+	if (Index >= 0)
+	{
+		return RayTracingDecals.GridData[LoopCount.DecalGridOffset + Index];
+	}
+	return Index;
+}
+
+#else
+
+FDecalLoopCount DecalGridLookup(float3 TranslatedWorldPos)
+{
+	uint NumDecals = min(RayTracingDecals.Count, RAY_TRACING_DECAL_COUNT_MAXIMUM);
+	FDecalLoopCount Result;
+	Result.NumDecals = NumDecals;
+	return Result;
+}
+
+float3 DecalGridVisualize(FDecalLoopCount LoopCount, int VisualizeMode)
+{
+	return 0.5;
+}
+
+uint GetDecalId(uint Index, FDecalLoopCount LoopCount)
+{
+	return Index;
+}
+
+#endif
diff --git a/Engine/Shaders/Private/RayTracing/RayTracingDecalMaterialShader.usf b/Engine/Shaders/Private/RayTracing/RayTracingDecalMaterialShader.usf
new file mode 100644
--- /dev/null
+++ b/Engine/Shaders/Private/RayTracing/RayTracingDecalMaterialShader.usf
@@ -0,0 +1,112 @@
+// Copyright Epic Games, Inc. All Rights Reserved.
+
+#define PATH_TRACING 1
+
+#include "/Engine/Private/Common.ush"
+#include "/Engine/Private/RayTracing/RayTracingCommon.ush"
+
+#define DecalTilePosition RayTracingDecalParameters.DecalTilePosition
+#define WorldToDecal RayTracingDecalParameters.WorldToDecal
+#define SvPositionToDecal WorldToDecal
+#define DecalToWorld RayTracingDecalParameters.DecalToWorld
+#define DecalToWorldInvScale RayTracingDecalParameters.DecalToWorldInvScale
+#define DecalOrientation RayTracingDecalParameters.DecalOrientation
+#define DecalParams RayTracingDecalParameters.DecalParams
+#define DecalWriteFlags RayTracingDecalParameters.DecalWriteFlags
+
+#define DECAL_WRITE_BASE_COLOR_FLAG						(1 << 0)
+#define DECAL_WRITE_NORMAL_FLAG							(1 << 1)
+#define DECAL_WRITE_ROUGHNESS_SPECULAR_METALLIC_FLAG	(1 << 2)
+#define DECAL_WRITE_EMISSIVE_FLAG						(1 << 3)
+#define DECAL_WRITE_AMBIENT_OCCLUSION_FLAG				(1 << 4)
+
+// DECAL_PRIMITIVE informs material templates which functions to expose when rendering decals.
+#define DECAL_PRIMITIVE 1
+#include "/Engine/Generated/Material.ush"   
+
+RAY_TRACING_ENTRY_CALLABLE(RayTracingDecalMaterialShader,
+	FDecalCallableShaderParams, Params)
+{
+	ResolvedView = ResolveView();
+
+	const float3 TranslatedWorldPosition = Params.TranslatedWorldPosition_MetallicSpecularRoughnessAO.xyz;
+
+	float4 DecalPos = mul(float4(TranslatedWorldPosition.xyz, 1), WorldToDecal);
+
+	if (any(abs(DecalPos.xyz) > 1))
+	{
+		// clip content outside the decal
+		return;
+	}
+
+	float3 CameraVector = normalize(PrimaryView.TranslatedWorldCameraOrigin - TranslatedWorldPosition.xyz);
+
+	// can be optimized
+	float3 DecalVector = DecalPos.xyz * 0.5f + 0.5f;
+
+	// Swizzle so that DecalVector.xy are perpendicular to the projection direction and DecalVector.z is distance along the projection direction
+	float3 SwizzlePos = DecalVector.zyx;
+
+	// By default, map textures using the vectors perpendicular to the projection direction
+	float2 DecalUVs = SwizzlePos.xy;
+
+	FMaterialPixelParameters MaterialParameters = MakeInitializedMaterialPixelParameters();
+#if NUM_MATERIAL_TEXCOORDS
+	for (int CoordinateIndex = 0; CoordinateIndex < NUM_MATERIAL_TEXCOORDS; CoordinateIndex++)
+	{
+		MaterialParameters.TexCoords[CoordinateIndex] = DecalUVs;
+	}
+#endif
+	MaterialParameters.TwoSidedSign = 1;
+	MaterialParameters.VertexColor = 1;
+	MaterialParameters.CameraVector = CameraVector;
+	MaterialParameters.SvPosition = 0.0f; // TODO
+	MaterialParameters.ScreenPosition = 0.0f; // TODO
+	MaterialParameters.LightVector = SwizzlePos;
+
+	MaterialParameters.AbsoluteWorldPosition = MaterialParameters.WorldPosition_NoOffsets = LWCSubtract(TranslatedWorldPosition.xyz, PrimaryView.PreViewTranslation);
+	MaterialParameters.WorldPosition_CamRelative = MaterialParameters.WorldPosition_NoOffsets_CamRelative = TranslatedWorldPosition.xyz;
+
+	FPixelMaterialInputs PixelMaterialInputs;
+	CalcPixelMaterialInputs(MaterialParameters, PixelMaterialInputs);
+
+	const float DecalFading = saturate(4 - 4 * abs(SwizzlePos.z * 2 - 1)) * DecalParams.x;
+	Params.Color.a = GetMaterialOpacity(PixelMaterialInputs) * DecalFading;
+
+	if (DecalWriteFlags & DECAL_WRITE_BASE_COLOR_FLAG)
+	{
+		Params.Color.rgb = GetMaterialBaseColor(PixelMaterialInputs);
+	}
+
+	if (DecalWriteFlags & DECAL_WRITE_NORMAL_FLAG)
+	{
+		Params.WorldNormal.rgb = MaterialParameters.WorldNormal;
+	}
+
+	if (DecalWriteFlags & DECAL_WRITE_ROUGHNESS_SPECULAR_METALLIC_FLAG)
+	{
+		Params.TranslatedWorldPosition_MetallicSpecularRoughnessAO.r = GetMaterialMetallic(PixelMaterialInputs);
+		Params.TranslatedWorldPosition_MetallicSpecularRoughnessAO.g = GetMaterialSpecular(PixelMaterialInputs);
+		Params.TranslatedWorldPosition_MetallicSpecularRoughnessAO.b = GetMaterialRoughness(PixelMaterialInputs);
+	}
+	else
+	{
+		Params.TranslatedWorldPosition_MetallicSpecularRoughnessAO.r = -1;
+		Params.TranslatedWorldPosition_MetallicSpecularRoughnessAO.g = -1;
+		Params.TranslatedWorldPosition_MetallicSpecularRoughnessAO.b = -1;
+	}
+
+	if (DecalWriteFlags & DECAL_WRITE_AMBIENT_OCCLUSION_FLAG)
+	{
+		Params.TranslatedWorldPosition_MetallicSpecularRoughnessAO.a = GetMaterialAmbientOcclusion(PixelMaterialInputs);
+	}
+	else
+	{
+		Params.TranslatedWorldPosition_MetallicSpecularRoughnessAO.a = -1;
+	}
+
+	if (DecalWriteFlags & DECAL_WRITE_EMISSIVE_FLAG)
+	{
+		Params.Emissive = GetMaterialEmissive(PixelMaterialInputs);
+	}
+}
diff --git a/Engine/Shaders/Public/Platform.ush b/Engine/Shaders/Public/Platform.ush
--- a/Engine/Shaders/Public/Platform.ush
+++ b/Engine/Shaders/Public/Platform.ush
@@ -172,6 +172,10 @@
 #define PLATFORM_SUPPORTS_SRV_UB_MACRO(...)
 #endif
 
+#ifndef PLATFORM_SUPPORTS_CALLABLE_SHADERS
+#define PLATFORM_SUPPORTS_CALLABLE_SHADERS 0
+#endif
+
 // Whether the platforms support official SM6 wave intrinsics
 // [REDACTED_PROFILE_URL]/DirectXShaderCompiler/wiki/Wave-Intrinsics
 #ifndef PLATFORM_SUPPORTS_SM6_0_WAVE_OPERATIONS
diff --git a/Engine/Shaders/Shared/RayTracingDefinitions.h b/Engine/Shaders/Shared/RayTracingDefinitions.h
--- a/Engine/Shaders/Shared/RayTracingDefinitions.h
+++ b/Engine/Shaders/Shared/RayTracingDefinitions.h
@@ -32,6 +32,7 @@
 #define RAY_TRACING_NUM_MISS_SHADER_SLOTS		2
 
 #define RAY_TRACING_LIGHT_COUNT_MAXIMUM		256
+#define RAY_TRACING_DECAL_COUNT_MAXIMUM		256
 
 #define RAY_TRACING_MAX_ALLOWED_RECURSION_DEPTH 1   // Only allow ray tracing from RayGen shader
 #define RAY_TRACING_MAX_ALLOWED_ATTRIBUTE_SIZE  8   // Sizeof 2 floats (barycentrics)
diff --git a/Engine/Shaders/Shared/RayTracingTypes.h b/Engine/Shaders/Shared/RayTracingTypes.h
--- a/Engine/Shaders/Shared/RayTracingTypes.h
+++ b/Engine/Shaders/Shared/RayTracingTypes.h
@@ -52,6 +52,18 @@
 
 static_assert(sizeof(FPathTracingPackedPathState) == 84, "Packed Path State size should be minimized");
 
+// C++ representation of a decal for ray tracing
+struct FRayTracingDecal
+{
+	FVector3f TranslatedBoundMin;
+	uint32 Pad0;
+	FVector3f TranslatedBoundMax;
+	uint32 CallableSlotIndex;
+	// keep structure aligned
+};
+
+static_assert(sizeof(FRayTracingDecal) == 32, "Ray tracing decal structure should be aligned to 32 bytes for optimal access on the GPU");
+
 #else
 
 // HLSL side of the structs above
@@ -90,4 +102,12 @@
 	uint2     PackedRoughnessSigma;
 };
 
+struct FRayTracingDecal
+{
+	float3 TranslatedBoundMin;
+	uint Pad0;
+	float3 TranslatedBoundMax;
+	uint CallableSlotIndex;
+};
+
 #endif
diff --git a/Engine/Source/Developer/Windows/ShaderFormatD3D/Private/D3DShaderCompiler.cpp b/Engine/Source/Developer/Windows/ShaderFormatD3D/Private/D3DShaderCompiler.cpp
--- a/Engine/Source/Developer/Windows/ShaderFormatD3D/Private/D3DShaderCompiler.cpp
+++ b/Engine/Source/Developer/Windows/ShaderFormatD3D/Private/D3DShaderCompiler.cpp
@@ -1206,6 +1206,7 @@
 
 	if (bUseDXC)
 	{
+		AdditionalDefines.SetDefine(TEXT("PLATFORM_SUPPORTS_CALLABLE_SHADERS"), 1);
 		AdditionalDefines.SetDefine(TEXT("PLATFORM_SUPPORTS_SM6_0_WAVE_OPERATIONS"), 1);
 		AdditionalDefines.SetDefine(TEXT("PLATFORM_SUPPORTS_STATIC_SAMPLERS"), 1);
 		AdditionalDefines.SetDefine(TEXT("PLATFORM_SUPPORTS_DIAGNOSTIC_BUFFER"), 1);
diff --git a/Engine/Source/Runtime/Renderer/Private/CompositionLighting/PostProcessDeferredDecals.cpp b/Engine/Source/Runtime/Renderer/Private/CompositionLighting/PostProcessDeferredDecals.cpp
--- a/Engine/Source/Runtime/Renderer/Private/CompositionLighting/PostProcessDeferredDecals.cpp
+++ b/Engine/Source/Runtime/Renderer/Private/CompositionLighting/PostProcessDeferredDecals.cpp
@@ -476,7 +476,7 @@
 			for (uint32 DecalIndex = DecalIndexBegin; DecalIndex < DecalIndexEnd; ++DecalIndex)
 			{
 				const FTransientDecalRenderData& DecalData = (*SortedDecals)[DecalIndex];
-				const FDeferredDecalProxy& DecalProxy = *DecalData.DecalProxy;
+				const FDeferredDecalProxy& DecalProxy = DecalData.Proxy;
 				const FMatrix ComponentToWorldMatrix = DecalProxy.ComponentTrans.ToMatrixWithScale();
 				const FMatrix FrustumComponentToClip = DecalRendering::ComputeComponentToClipMatrix(View, ComponentToWorldMatrix);
 				const bool bStencilThisDecal = IsStencilOptimizationAvailable(DecalRenderStage);
@@ -508,7 +508,7 @@
 					GraphicsPSOInit.DepthStencilState = GetDecalDepthState(StencilRef, DecalDepthState);
 				}
 
-				GraphicsPSOInit.BlendState = DecalRendering::GetDecalBlendState(DecalData.DecalBlendDesc, DecalRenderStage, RenderTargetMode);
+				GraphicsPSOInit.BlendState = DecalRendering::GetDecalBlendState(DecalData.BlendDesc, DecalRenderStage, RenderTargetMode);
 				GraphicsPSOInit.PrimitiveType = PT_TriangleList;
 
 				DecalRendering::SetShader(RHICmdList, GraphicsPSOInit, StencilRef, View, DecalData, DecalRenderStage, FrustumComponentToClip);
@@ -532,11 +532,11 @@
 
 			uint32 SortedDecalIndex = 1;
 			uint32 LastSortedDecalIndex = 0;
-			EDecalRenderTargetMode LastRenderTargetMode = DecalRendering::GetRenderTargetMode((*SortedDecals)[0].DecalBlendDesc, DecalRenderStage);
+			EDecalRenderTargetMode LastRenderTargetMode = DecalRendering::GetRenderTargetMode((*SortedDecals)[0].BlendDesc, DecalRenderStage);
 
 			for (; SortedDecalIndex < SortedDecalCount; ++SortedDecalIndex)
 			{
-				const EDecalRenderTargetMode RenderTargetMode = DecalRendering::GetRenderTargetMode((*SortedDecals)[SortedDecalIndex].DecalBlendDesc, DecalRenderStage);
+				const EDecalRenderTargetMode RenderTargetMode = DecalRendering::GetRenderTargetMode((*SortedDecals)[SortedDecalIndex].BlendDesc, DecalRenderStage);
 
 				if (LastRenderTargetMode != RenderTargetMode)
 				{
diff --git a/Engine/Source/Runtime/Renderer/Private/DecalRenderingShared.cpp b/Engine/Source/Runtime/Renderer/Private/DecalRenderingShared.cpp
--- a/Engine/Source/Runtime/Renderer/Private/DecalRenderingShared.cpp
+++ b/Engine/Source/Runtime/Renderer/Private/DecalRenderingShared.cpp
@@ -21,15 +21,15 @@
 	TEXT("  Smaller means decals fade less aggressively.")
 	);
 
-FTransientDecalRenderData::FTransientDecalRenderData(const FScene& InScene, const FDeferredDecalProxy* InDecalProxy, float InConservativeRadius)
-	: DecalProxy(InDecalProxy)
+FTransientDecalRenderData::FTransientDecalRenderData(const FScene& InScene, const FDeferredDecalProxy& InDecalProxy, float InConservativeRadius)
+	: Proxy(InDecalProxy)
 	, FadeAlpha(1.0f)
 	, ConservativeRadius(InConservativeRadius)
 {
-	MaterialProxy = InDecalProxy->DecalMaterial->GetRenderProxy();
+	MaterialProxy = Proxy.DecalMaterial->GetRenderProxy();
 	MaterialResource = &MaterialProxy->GetMaterialWithFallback(InScene.GetFeatureLevel(), MaterialProxy);
 	check(MaterialProxy && MaterialResource);
-	DecalBlendDesc = DecalRendering::ComputeDecalBlendDesc(InScene.GetShaderPlatform(), *MaterialResource);
+	BlendDesc = DecalRendering::ComputeDecalBlendDesc(InScene.GetShaderPlatform(), *MaterialResource);
 }
 
 /**
@@ -233,6 +233,43 @@
 		return FMath::Clamp(Alpha, 0.0f, 1.0f);
 	}
 
+	void SortDecalList(FTransientDecalRenderDataList& Decals)
+	{
+		// Sort by sort order to allow control over composited result
+		// Then sort decals by state to reduce render target switches
+		// Also sort by component since Sort() is not stable
+		struct FCompareFTransientDecalRenderData
+		{
+			FORCEINLINE bool operator()(const FTransientDecalRenderData& A, const FTransientDecalRenderData& B) const
+			{
+				if (B.Proxy.SortOrder != A.Proxy.SortOrder)
+				{ 
+					return A.Proxy.SortOrder < B.Proxy.SortOrder;
+				}
+				if (B.BlendDesc.bWriteNormal != A.BlendDesc.bWriteNormal)
+				{
+					// bWriteNormal here has priority because we want to render decals that output normals before those could read normals.
+					// Also this is the only flag that can trigger a change of EDecalRenderTargetMode inside a single EDecalRenderStage, and we batch according to this.
+					return B.BlendDesc.bWriteNormal < A.BlendDesc.bWriteNormal; // < so that those outputting normal are first.
+				}
+				if (B.BlendDesc.Packed != A.BlendDesc.Packed)
+				{
+					// Sorting by the FDecalBlendDesc contents will reduce blend state changes.
+					return (int32)B.BlendDesc.Packed < (int32)A.BlendDesc.Packed;
+				}
+				if (B.MaterialProxy != A.MaterialProxy)
+				{
+					// Batch decals with the same material together
+					return B.MaterialProxy < A.MaterialProxy;
+				}
+				return (PTRINT)B.Proxy.Component < (PTRINT)A.Proxy.Component;
+			}
+		};
+
+		// Sort decals by blend mode to reduce render target switches
+		Decals.Sort(FCompareFTransientDecalRenderData());
+	}
+
 	bool BuildVisibleDecalList(const FScene& Scene, const FViewInfo& View, EDecalRenderStage DecalRenderStage, FTransientDecalRenderDataList* OutVisibleDecals)
 	{
 		QUICK_SCOPE_CYCLE_COUNTER(BuildVisibleDecalList);
@@ -285,13 +322,13 @@
 
 			if (bIsShown)
 			{
-				FTransientDecalRenderData Data(Scene, DecalProxy, ConservativeRadius);
+				FTransientDecalRenderData Data(Scene, *DecalProxy, ConservativeRadius);
 			
-				if (IsCompatibleWithRenderStage(Data.DecalBlendDesc, DecalRenderStage))
+				if (IsCompatibleWithRenderStage(Data.BlendDesc, DecalRenderStage))
 				{
-					if (bIsPerspectiveProjection && Data.DecalProxy->FadeScreenSize != 0.0f)
+					if (bIsPerspectiveProjection && Data.Proxy.FadeScreenSize != 0.0f)
 					{
-						Data.FadeAlpha = CalculateDecalFadeAlpha(Data.DecalProxy->FadeScreenSize, ComponentToWorldMatrix, View, FadeMultiplier);
+						Data.FadeAlpha = CalculateDecalFadeAlpha(Data.Proxy.FadeScreenSize, ComponentToWorldMatrix, View, FadeMultiplier);
 					}
 
 					const bool bShouldRender = Data.FadeAlpha > 0.0f;
@@ -315,40 +352,8 @@
 
 		if (OutVisibleDecals->Num() > 0)
 		{
-			// Sort by sort order to allow control over composited result
-			// Then sort decals by state to reduce render target switches
-			// Also sort by component since Sort() is not stable
-			struct FCompareFTransientDecalRenderData
-			{
-				FORCEINLINE bool operator()(const FTransientDecalRenderData& A, const FTransientDecalRenderData& B) const
-				{
-					if (B.DecalProxy->SortOrder != A.DecalProxy->SortOrder)
-					{ 
-						return A.DecalProxy->SortOrder < B.DecalProxy->SortOrder;
-					}
-					if (B.DecalBlendDesc.bWriteNormal != A.DecalBlendDesc.bWriteNormal)
-					{
-						// bWriteNormal here has priority because we want to render decals that output normals before those could read normals.
-						// Also this is the only flag that can trigger a change of EDecalRenderTargetMode inside a single EDecalRenderStage, and we batch according to this.
-						return B.DecalBlendDesc.bWriteNormal < A.DecalBlendDesc.bWriteNormal; // < so that those outputting normal are first.
-					}
-					if (B.DecalBlendDesc.Packed != A.DecalBlendDesc.Packed)
-					{
-						// Sorting by the FDecalBlendDesc contents will reduce blend state changes.
-						return (int32)B.DecalBlendDesc.Packed < (int32)A.DecalBlendDesc.Packed;
-					}
-					if (B.MaterialProxy != A.MaterialProxy)
-					{
-						// Batch decals with the same material together
-						return B.MaterialProxy < A.MaterialProxy;
-					}
-					return (PTRINT)B.DecalProxy->Component < (PTRINT)A.DecalProxy->Component;
-				}
-			};
+			SortDecalList(*OutVisibleDecals);
 
-			// Sort decals by blend mode to reduce render target switches
-			OutVisibleDecals->Sort(FCompareFTransientDecalRenderData());
-
 			return true;
 		}
 
@@ -406,7 +411,7 @@
 
 		// Set pixel shader parameters.
 		{
-			PixelShader->SetParameters(RHICmdList, View, DecalData.MaterialProxy, *DecalData.DecalProxy, DecalData.FadeAlpha);
+			PixelShader->SetParameters(RHICmdList, View, DecalData.MaterialProxy, DecalData.Proxy, DecalData.FadeAlpha);
 
 			auto& PrimitivePS = PixelShader->GetUniformBufferParameter<FPrimitiveUniformShaderParameters>();
 
diff --git a/Engine/Source/Runtime/Renderer/Private/DecalRenderingShared.h b/Engine/Source/Runtime/Renderer/Private/DecalRenderingShared.h
--- a/Engine/Source/Runtime/Renderer/Private/DecalRenderingShared.h
+++ b/Engine/Source/Runtime/Renderer/Private/DecalRenderingShared.h
@@ -19,12 +19,12 @@
 {
 	const FMaterialRenderProxy* MaterialProxy;
 	const FMaterial* MaterialResource;
-	const FDeferredDecalProxy* DecalProxy;
-	FDecalBlendDesc DecalBlendDesc;
+	const FDeferredDecalProxy& Proxy;
+	FDecalBlendDesc BlendDesc;
 	float FadeAlpha;
 	float ConservativeRadius;
 
-	FTransientDecalRenderData(const FScene& InScene, const FDeferredDecalProxy* InDecalProxy, float InConservativeRadius);
+	FTransientDecalRenderData(const FScene& InScene, const FDeferredDecalProxy& InDecalProxy, float InConservativeRadius);
 };
 	
 typedef TArray<FTransientDecalRenderData, SceneRenderingAllocator> FTransientDecalRenderDataList;
@@ -38,6 +38,7 @@
 	float CalculateDecalFadeAlpha(float DecalFadeScreenSize, const FMatrix& ComponentToWorldMatrix, const FViewInfo& View, float FadeMultiplier);
 	FMatrix ComputeComponentToClipMatrix(const FViewInfo& View, const FMatrix& DecalComponentToWorld);
 	void SetVertexShaderOnly(FRHICommandList& RHICmdList, FGraphicsPipelineStateInitializer& GraphicsPSOInit, const FViewInfo& View, const FMatrix& FrustumComponentToClip);
+	void SortDecalList(FTransientDecalRenderDataList& Decals);
 	bool BuildVisibleDecalList(const FScene& Scene, const FViewInfo& View, EDecalRenderStage DecalRenderStage, FTransientDecalRenderDataList* OutVisibleDecals);
 	void SetShader(FRHICommandList& RHICmdList, FGraphicsPipelineStateInitializer& GraphicsPSOInit, uint32 StencilRef, const FViewInfo& View, const FTransientDecalRenderData& DecalData, EDecalRenderStage DecalRenderStage, const FMatrix& FrustumComponentToClip);
 };
diff --git a/Engine/Source/Runtime/Renderer/Private/DeferredShadingRenderer.cpp b/Engine/Source/Runtime/Renderer/Private/DeferredShadingRenderer.cpp
--- a/Engine/Source/Runtime/Renderer/Private/DeferredShadingRenderer.cpp
+++ b/Engine/Source/Runtime/Renderer/Private/DeferredShadingRenderer.cpp
@@ -37,6 +37,7 @@
 #include "GPUScene.h"
 #include "RayTracing/RayTracingMaterialHitShaders.h"
 #include "RayTracing/RayTracingLighting.h"
+#include "RayTracing/RayTracingDecals.h"
 #include "RayTracing/RayTracingScene.h"
 #include "RayTracingDynamicGeometryCollection.h"
 #include "RayTracingSkinnedGeometry.h"
@@ -2190,6 +2191,12 @@ void FDeferredShadingSceneRenderer::Render(FRDGBuilder& GraphBuilder)
 	// Prepare the scene for rendering this frame.
 	GatherRayTracingWorldInstancesForView(GraphBuilder, ReferenceView, RayTracingScene);
 
+	if (ActiveViewFamily->EngineShowFlags.PathTracing)
+	{
+		ReferenceView.RayTracingDecalUniformBuffer = CreateRayTracingDecalData(GraphBuilder, *Scene, ReferenceView, RayTracingScene.NumCallableShaderSlots);
+		RayTracingScene.NumCallableShaderSlots += Scene->Decals.Num();
+	}
+
 #endif // RHI_RAYTRACING
 
 	// Dynamic vertex and index buffers need to be committed before rendering.
diff --git a/Engine/Source/Runtime/Renderer/Private/MobileDecalRendering.cpp b/Engine/Source/Runtime/Renderer/Private/MobileDecalRendering.cpp
--- a/Engine/Source/Runtime/Renderer/Private/MobileDecalRendering.cpp
+++ b/Engine/Source/Runtime/Renderer/Private/MobileDecalRendering.cpp
@@ -73,7 +73,7 @@
 		for (int32 DecalIndex = 0; DecalIndex < SortedDecalCount; DecalIndex++)
 		{
 			const FTransientDecalRenderData& DecalData = SortedDecals[DecalIndex];
-			const FDeferredDecalProxy& DecalProxy = *DecalData.DecalProxy;
+			const FDeferredDecalProxy& DecalProxy = DecalData.Proxy;
 			const FMatrix ComponentToWorldMatrix = DecalProxy.ComponentTrans.ToMatrixWithScale();
 			const FMatrix FrustumComponentToClip = DecalRendering::ComputeComponentToClipMatrix(View, ComponentToWorldMatrix);
 
@@ -105,7 +105,7 @@
 					GET_STENCIL_BIT_MASK(RECEIVE_DECAL, 1), 0x00>::GetRHI();
 			}
 
-			GraphicsPSOInit.BlendState = DecalRendering::GetDecalBlendState(DecalData.DecalBlendDesc, DecalRenderStage, RenderTargetMode);
+			GraphicsPSOInit.BlendState = DecalRendering::GetDecalBlendState(DecalData.BlendDesc, DecalRenderStage, RenderTargetMode);
 
 			// Set shader params
 			DecalRendering::SetShader(RHICmdList, GraphicsPSOInit, 0, View, DecalData, DecalRenderStage, FrustumComponentToClip);
diff --git a/Engine/Source/Runtime/Renderer/Private/PathTracing.cpp b/Engine/Source/Runtime/Renderer/Private/PathTracing.cpp
--- a/Engine/Source/Runtime/Renderer/Private/PathTracing.cpp
+++ b/Engine/Source/Runtime/Renderer/Private/PathTracing.cpp
@@ -23,6 +23,7 @@
 #include "RayTracingDefinitions.h"
 #include "PathTracingDefinitions.h"
 #include "RayTracing/RayTracingMaterialHitShaders.h"
+#include "RayTracing/RayTracingDecals.h"
 #include "RenderCore/Public/GenerateMips.h"
 #include "HairStrands/HairStrandsData.h"
 #include "Modules/ModuleManager.h"
@@ -259,6 +260,16 @@
 	ECVF_RenderThreadSafe
 );
 
+TAutoConsoleVariable<int32> CVarPathTracingDecalGridVisualize(
+	TEXT("r.PathTracing.DecalGrid.Visualize"),
+	0,
+	TEXT("Enables a visualization mode of the decal grid density where red indicates the maximum decal count has been reached (default = 0)\n")
+	TEXT("0: off (default)\n")
+	TEXT("1: decal count heatmap (red - close to overflow, increase r.RayTracing.DecalGrid.MaxCount)\n")
+	TEXT("2: unique decal lists (colors are a function of which decals occupy each cell)\n"),
+	ECVF_RenderThreadSafe
+);
+
 TAutoConsoleVariable<int32> CVarPathTracingDenoiser(
 	TEXT("r.PathTracing.Denoiser"),
 	-1,
@@ -284,8 +295,10 @@
 	SHADER_PARAMETER(uint32, EnableEmissive)
 	SHADER_PARAMETER(uint32, SamplerType)
 	SHADER_PARAMETER(uint32, VisualizeLightGrid)
+	SHADER_PARAMETER(uint32, VisualizeDecalGrid)
 	SHADER_PARAMETER(uint32, EnableAtmosphere)
 	SHADER_PARAMETER(uint32, EnableFog)
+	SHADER_PARAMETER(uint32, EnableDecals)
 	SHADER_PARAMETER(int32, MaxRaymarchSteps)
 	SHADER_PARAMETER(float, MaxPathIntensity)
 	SHADER_PARAMETER(float, MaxNormalBias)
@@ -325,11 +338,13 @@
 			PathTracingData.EnableDirectLighting != Other.PathTracingData.EnableDirectLighting ||
 			PathTracingData.EnableEmissive != Other.PathTracingData.EnableEmissive ||
 			PathTracingData.VisualizeLightGrid != Other.PathTracingData.VisualizeLightGrid ||
+			PathTracingData.VisualizeDecalGrid != Other.PathTracingData.VisualizeDecalGrid ||
 			PathTracingData.MaxPathIntensity != Other.PathTracingData.MaxPathIntensity ||
 			PathTracingData.FilterWidth != Other.PathTracingData.FilterWidth ||
 			PathTracingData.AbsorptionScale != Other.PathTracingData.AbsorptionScale ||
 			PathTracingData.EnableAtmosphere != Other.PathTracingData.EnableAtmosphere ||
 			PathTracingData.EnableFog != Other.PathTracingData.EnableFog ||
+			PathTracingData.EnableDecals != Other.PathTracingData.EnableDecals ||
 			PathTracingData.MaxRaymarchSteps != Other.PathTracingData.MaxRaymarchSteps ||
 			ViewRect != Other.ViewRect ||
 			LightShowFlags != Other.LightShowFlags ||
@@ -455,6 +470,7 @@
 	int32 EnableEmissive = CVarPathTracingEnableEmissive.GetValueOnRenderThread();
 	PathTracingData.EnableEmissive = EnableEmissive < 0 ? View.FinalPostProcessSettings.PathTracingEnableEmissive : EnableEmissive;
 	PathTracingData.VisualizeLightGrid = CVarPathTracingLightGridVisualize.GetValueOnRenderThread();
+	PathTracingData.VisualizeDecalGrid = CVarPathTracingDecalGridVisualize.GetValueOnRenderThread();
 	float FilterWidth = CVarPathTracingFilterWidth.GetValueOnRenderThread();
 	if (FilterWidth < 0)
 	{
@@ -477,6 +493,7 @@
 		ShouldRenderSkyAtmosphere(Scene, View.Family->EngineShowFlags) && 
 		View.SkyAtmosphereUniformShaderParameters &&
 		View.FinalPostProcessSettings.PathTracingEnableReferenceAtmosphere;
+
 	PathTracingData.EnableFog = ShouldRenderFog(*View.Family)
 		&& Scene->ExponentialFogs.Num() > 0
 		&& Scene->ExponentialFogs[0].bEnableVolumetricFog
@@ -484,6 +501,9 @@
 		&& Scene->ExponentialFogs[0].VolumetricFogExtinctionScale > 0
 		&& (Scene->ExponentialFogs[0].FogData[0].Density > 0 ||
 			Scene->ExponentialFogs[0].FogData[1].Density > 0);
+
+	PathTracingData.EnableDecals = View.Family->EngineShowFlags.Decals && Scene->Decals.Num() > 0;
+
 	PathTracingData.MaxRaymarchSteps = CVarPathTracingMaxRaymarchSteps.GetValueOnRenderThread();
 }
 
@@ -697,6 +717,9 @@
 		SHADER_PARAMETER_RDG_TEXTURE(Texture2DArray, IESTexture)
 		SHADER_PARAMETER_SAMPLER(SamplerState, IESTextureSampler) // Shared sampler for all IES profiles
 
+		// scene decals
+		SHADER_PARAMETER_RDG_UNIFORM_BUFFER(FRayTracingDecals, DecalParameters)
+
 		// Used by multi-GPU rendering and TDR-avoidance tiling
 		SHADER_PARAMETER(FIntPoint, TilePixelOffset)
 		SHADER_PARAMETER(FIntPoint, TileTextureOffset)
@@ -1946,6 +1969,8 @@
 							PassParameters->SceneVisibleLightCount = 0;
 						}
 
+						PassParameters->DecalParameters = View.RayTracingDecalUniformBuffer;
+
 						PassParameters->RadianceTexture = GraphBuilder.CreateUAV(RadianceTexture);
 						PassParameters->AlbedoTexture = GraphBuilder.CreateUAV(AlbedoTexture);
 						PassParameters->NormalTexture = GraphBuilder.CreateUAV(NormalTexture);
diff --git a/Engine/Source/Runtime/Renderer/Private/RayTracing/RayTracingDecals.cpp b/Engine/Source/Runtime/Renderer/Private/RayTracing/RayTracingDecals.cpp
new file mode 100644
--- /dev/null
+++ b/Engine/Source/Runtime/Renderer/Private/RayTracing/RayTracingDecals.cpp
@@ -0,0 +1,415 @@
+// Copyright Epic Games, Inc. All Rights Reserved.
+
+#include "RayTracingDecals.h"
+
+#include "RHI/Public/RHIDefinitions.h"
+
+#if RHI_RAYTRACING
+
+#include "RendererPrivate.h"
+#include "RayTracingMaterialHitShaders.h"
+#include "DecalRenderingShared.h"
+#include "RayTracingTypes.h"
+#include "RayTracingDefinitions.h"
+
+TAutoConsoleVariable<int32> CVarDecalGridResolution(
+	TEXT("r.RayTracing.DecalGrid.Resolution"),
+	256,
+	TEXT("Controls the resolution of the 2D decal grid used to cull irrelevant decal from calculations (default = 256)\n"),
+	ECVF_RenderThreadSafe
+);
+
+TAutoConsoleVariable<int32> CVarDecalGridMaxCount(
+	TEXT("r.RayTracing.DecalGrid.MaxCount"),
+	128,
+	TEXT("Controls the maximum number of decals per cell in the 2D decal grid. The minimum of this value and the number of decal in the scene is used. (default = 128)\n"),
+	ECVF_RenderThreadSafe
+);
+
+enum class EDecalsWriteFlags : uint32
+{
+	None						= 0,
+	BaseColor					= (1 << 0),
+	Normal						= (1 << 1),
+	RoughnessSpecularMetallic	= (1 << 2),
+	Emissive					= (1 << 3),
+	AmbientOcclusion			= (1 << 4),
+};
+
+IMPLEMENT_GLOBAL_SHADER_PARAMETER_STRUCT(FRayTracingDecals, "RayTracingDecals");
+
+BEGIN_GLOBAL_SHADER_PARAMETER_STRUCT(FDecalParametersRayTracing, )
+	SHADER_PARAMETER(FMatrix44f, WorldToDecal)
+	SHADER_PARAMETER(FMatrix44f, DecalToWorld)
+	SHADER_PARAMETER(FVector3f, DecalTilePosition)
+	SHADER_PARAMETER(FVector3f, DecalToWorldInvScale)
+	SHADER_PARAMETER(FVector3f, DecalOrientation)
+	SHADER_PARAMETER(FVector2f, DecalParams)
+	SHADER_PARAMETER(uint32, DecalWriteFlags)
+END_GLOBAL_SHADER_PARAMETER_STRUCT()
+
+IMPLEMENT_GLOBAL_SHADER_PARAMETER_STRUCT(FDecalParametersRayTracing, "RayTracingDecalParameters");
+
+static TUniformBufferRef<FDecalParametersRayTracing> CreateDecalParametersBuffer(
+	const FViewInfo& View,
+	EShaderPlatform Platform,
+	const FTransientDecalRenderData& DecalData,
+	EUniformBufferUsage Usage)
+{
+	FDecalParametersRayTracing Parameters;
+
+	const FLargeWorldRenderPosition AbsoluteOrigin(View.ViewMatrices.GetInvViewMatrix().GetOrigin());
+	const FVector3f TilePosition = AbsoluteOrigin.GetTile();
+	const FMatrix WorldToDecalMatrix = DecalData.Proxy.ComponentTrans.ToInverseMatrixWithScale();
+	const FMatrix44f RelativeWorldToDecalMatrix = FLargeWorldRenderScalar::MakeFromRelativeWorldMatrix(AbsoluteOrigin.GetTileOffset(), WorldToDecalMatrix);
+	const FMatrix DecalToWorldMatrix = DecalData.Proxy.ComponentTrans.ToMatrixWithScale();
+	const FMatrix44f RelativeDecalToWorldMatrix = FLargeWorldRenderScalar::MakeToRelativeWorldMatrix(AbsoluteOrigin.GetTileOffset(), DecalToWorldMatrix);
+	const FVector3f OrientationVector = (FVector3f)DecalData.Proxy.ComponentTrans.GetUnitAxis(EAxis::X);
+
+	Parameters.DecalTilePosition = TilePosition;
+	Parameters.WorldToDecal = FMatrix44f(FTranslationMatrix(-View.ViewMatrices.GetPreViewTranslation()) * WorldToDecalMatrix);
+	Parameters.DecalToWorld = RelativeDecalToWorldMatrix;
+	Parameters.DecalToWorldInvScale = RelativeDecalToWorldMatrix.GetScaleVector().Reciprocal();
+	Parameters.DecalOrientation = OrientationVector;
+
+	float LifetimeAlpha = 1.0f;
+
+	// Certain engine captures (e.g. environment reflection) don't have a tick. Default to fully opaque.
+	if (View.Family->Time.GetWorldTimeSeconds())
+	{
+		LifetimeAlpha = FMath::Clamp(FMath::Min(View.Family->Time.GetWorldTimeSeconds() * -DecalData.Proxy.InvFadeDuration + DecalData.Proxy.FadeStartDelayNormalized, View.Family->Time.GetWorldTimeSeconds() * DecalData.Proxy.InvFadeInDuration + DecalData.Proxy.FadeInStartDelayNormalized), 0.0f, 1.0f);
+	}
+
+	Parameters.DecalParams = FVector2f(DecalData.FadeAlpha, LifetimeAlpha);
+
+	Parameters.DecalWriteFlags = (uint32)EDecalsWriteFlags::None;
+#define SET_FLAG(BOOL_VAL, WRITE_FLAG) Parameters.DecalWriteFlags |= uint32(DecalData.BlendDesc.BOOL_VAL ? EDecalsWriteFlags::WRITE_FLAG : EDecalsWriteFlags::None)
+	SET_FLAG(bWriteBaseColor, BaseColor);
+	SET_FLAG(bWriteNormal, Normal);
+	SET_FLAG(bWriteRoughnessSpecularMetallic, RoughnessSpecularMetallic);
+	SET_FLAG(bWriteEmissive, Emissive);
+	SET_FLAG(bWriteAmbientOcclusion, AmbientOcclusion);
+#undef SET_FLAG
+
+	return CreateUniformBufferImmediate(Parameters, Usage);
+}
+
+class FRayTracingBuildDecalGridCS : public FGlobalShader
+{
+	DECLARE_GLOBAL_SHADER(FRayTracingBuildDecalGridCS)
+	SHADER_USE_PARAMETER_STRUCT(FRayTracingBuildDecalGridCS, FGlobalShader)
+
+	static bool ShouldCompilePermutation(const FGlobalShaderPermutationParameters& Parameters)
+	{
+		return ShouldCompileRayTracingShadersForProject(Parameters.Platform);
+	}
+
+	static void ModifyCompilationEnvironment(const FGlobalShaderPermutationParameters& Parameters, FShaderCompilerEnvironment& OutEnvironment)
+	{
+		OutEnvironment.CompilerFlags.Add(CFLAG_WarningsAsErrors);
+		OutEnvironment.CompilerFlags.Add(CFLAG_AllowTypedUAVLoads);
+		OutEnvironment.SetDefine(TEXT("THREADGROUPSIZE_X"), FComputeShaderUtils::kGolden2DGroupSize);
+		OutEnvironment.SetDefine(TEXT("THREADGROUPSIZE_Y"), FComputeShaderUtils::kGolden2DGroupSize);
+	}
+
+	BEGIN_SHADER_PARAMETER_STRUCT(FParameters, )
+		SHADER_PARAMETER_RDG_BUFFER_SRV(StructuredBuffer<FRayTracingDecal>, SceneDecals)
+		SHADER_PARAMETER(uint32, SceneDecalCount)
+		SHADER_PARAMETER(FVector3f, SceneDecalsTranslatedBoundMin)
+		SHADER_PARAMETER(FVector3f, SceneDecalsTranslatedBoundMax)
+		SHADER_PARAMETER(unsigned, DecalGridResolution)
+		SHADER_PARAMETER(unsigned, DecalGridMaxCount)
+		SHADER_PARAMETER(int, DecalGridAxis)
+
+		SHADER_PARAMETER_RDG_TEXTURE_UAV(RWTexture2D, RWDecalGrid)
+		SHADER_PARAMETER_RDG_BUFFER_UAV(RWBuffer<uint>, RWDecalGridData)
+	END_SHADER_PARAMETER_STRUCT()
+};
+
+IMPLEMENT_SHADER_TYPE(, FRayTracingBuildDecalGridCS, TEXT("/Engine/Private/RayTracing/RayTracingBuildDecalGrid.usf"), TEXT("BuildDecalGridCS"), SF_Compute);
+
+class FRayTracingDecalMaterialShader : public FMaterialShader
+{
+	DECLARE_SHADER_TYPE(FRayTracingDecalMaterialShader, Material);
+public:
+	FRayTracingDecalMaterialShader() = default;
+	FRayTracingDecalMaterialShader(const ShaderMetaType::CompiledShaderInitializerType& Initializer)
+		: FMaterialShader(Initializer)
+	{
+		DecalParameter.Bind(Initializer.ParameterMap, TEXT("RaytracingDecalParameters"));
+	}
+
+	static bool ShouldCompilePermutation(const FMaterialShaderPermutationParameters& Parameters)
+	{
+		if (!ShouldCompileRayTracingCallableShadersForProject(Parameters.Platform))
+		{
+			// is raytracing enabled at all?
+			return false;
+		}
+		if (Parameters.MaterialParameters.MaterialDomain != MD_DeferredDecal)
+		{
+			// only compile callable shader permutation of deferred decal material
+			return false;
+		}
+
+		return true;
+	}
+
+	static void ModifyCompilationEnvironment(const FMaterialShaderPermutationParameters& Parameters, FShaderCompilerEnvironment& OutEnvironment)
+	{
+		FMaterialShader::ModifyCompilationEnvironment(Parameters, OutEnvironment);
+	}
+
+	static bool ValidateCompiledResult(EShaderPlatform Platform, const FShaderParameterMap& ParameterMap, TArray<FString>& OutError)
+	{
+		if (ParameterMap.ContainsParameterAllocation(FSceneTextureUniformParameters::StaticStructMetadata.GetShaderVariableName()))
+		{
+			OutError.Add(TEXT("Ray tracing callable shaders cannot read from the SceneTexturesStruct."));
+			return false;
+		}
+
+		for (const auto& It : ParameterMap.GetParameterMap())
+		{
+			const FParameterAllocation& ParamAllocation = It.Value;
+			if (ParamAllocation.Type != EShaderParameterType::UniformBuffer)
+			{
+				OutError.Add(FString::Printf(TEXT("Invalid ray tracing callable shader parameter '%s'. Only uniform buffers and loose data parameters are supported."), *(It.Key)));
+				return false;
+			}
+		}
+
+		return true;
+	}
+
+	void GetShaderBindings(
+		const FScene* Scene,
+		ERHIFeatureLevel::Type FeatureLevel,
+		const FMaterialRenderProxy& MaterialRenderProxy,
+		const FMaterial& Material,
+		const FViewInfo& View,
+		const TUniformBufferRef<FDecalParametersRayTracing>& DecalParameters,
+		FMeshDrawSingleShaderBindings& ShaderBindings) const
+	{
+		FMaterialShader::GetShaderBindings(Scene, FeatureLevel, MaterialRenderProxy, Material, ShaderBindings);
+
+		ShaderBindings.Add(GetUniformBufferParameter<FViewUniformShaderParameters>(), View.ViewUniformBuffer);
+		ShaderBindings.Add(DecalParameter, DecalParameters);
+	}
+
+private:
+	LAYOUT_FIELD(FShaderUniformBufferParameter, DecalParameter);
+};
+
+IMPLEMENT_MATERIAL_SHADER_TYPE(, FRayTracingDecalMaterialShader, TEXT("/Engine/Private/RayTracing/RayTracingDecalMaterialShader.usf"), TEXT("RayTracingDecalMaterialShader"), SF_RayCallable);
+
+void BuildDecalGrid(FRDGBuilder& GraphBuilder, uint32 NumDecals, FRDGBufferSRVRef DecalsSRV, FRayTracingDecals& OutParameters)
+{
+	// if we have some decals -- build a decal grid
+	if (NumDecals > 0)
+	{
+		const uint32 Resolution = FMath::RoundUpToPowerOfTwo(CVarDecalGridResolution.GetValueOnRenderThread());
+		const uint32 MaxCount = FMath::Clamp(CVarDecalGridMaxCount.GetValueOnRenderThread(), 1, FMath::Min((int32)NumDecals, RAY_TRACING_DECAL_COUNT_MAXIMUM));
+		OutParameters.GridResolution = Resolution;
+		OutParameters.GridMaxCount = MaxCount;
+
+		// pick the shortest axis
+		FVector3f Diag = OutParameters.TranslatedBoundMax - OutParameters.TranslatedBoundMin;
+		if (Diag.X < Diag.Y && Diag.X < Diag.Z)
+		{
+			OutParameters.GridAxis = 0;
+		}
+		else if (Diag.Y < Diag.Z)
+		{
+			OutParameters.GridAxis = 1;
+		}
+		else
+		{
+			OutParameters.GridAxis = 2;
+		}
+
+		FRDGTextureDesc DecalGridDesc = FRDGTextureDesc::Create2D(
+			FIntPoint(Resolution, Resolution),
+			PF_R32_UINT,
+			FClearValueBinding::None,
+			TexCreate_ShaderResource | TexCreate_UAV);
+		FRDGTexture* DecalGridTexture = GraphBuilder.CreateTexture(DecalGridDesc, TEXT("PathTracer.DecalGrid"), ERDGTextureFlags::None);
+
+		EPixelFormat DecalGridDataFormat = PF_R32_UINT;
+		size_t DecalGridDataNumBytes = sizeof(uint32);
+		if (NumDecals <= (MAX_uint8 + 1))
+		{
+			DecalGridDataFormat = PF_R8_UINT;
+			DecalGridDataNumBytes = sizeof(uint8);
+		}
+		else if (NumDecals <= (MAX_uint16 + 1))
+		{
+			DecalGridDataFormat = PF_R16_UINT;
+			DecalGridDataNumBytes = sizeof(uint16);
+		}
+		FRDGBufferDesc DecalGridDataDesc = FRDGBufferDesc::CreateBufferDesc(DecalGridDataNumBytes, MaxCount * Resolution * Resolution);
+		FRDGBuffer* DecalGridData = GraphBuilder.CreateBuffer(DecalGridDataDesc, TEXT("PathTracer.DecalGridData"));
+
+		auto* PassParameters = GraphBuilder.AllocParameters<FRayTracingBuildDecalGridCS::FParameters>();
+		PassParameters->RWDecalGrid = GraphBuilder.CreateUAV(DecalGridTexture);
+		PassParameters->RWDecalGridData = GraphBuilder.CreateUAV(DecalGridData, DecalGridDataFormat);
+		
+		PassParameters->SceneDecals = DecalsSRV;
+		PassParameters->SceneDecalCount = OutParameters.Count;
+		PassParameters->SceneDecalsTranslatedBoundMin = OutParameters.TranslatedBoundMin;
+		PassParameters->SceneDecalsTranslatedBoundMax = OutParameters.TranslatedBoundMax;
+		PassParameters->DecalGridResolution = OutParameters.GridResolution;
+		PassParameters->DecalGridMaxCount = OutParameters.GridMaxCount;
+		PassParameters->DecalGridAxis = OutParameters.GridAxis;
+
+		TShaderMapRef<FRayTracingBuildDecalGridCS> ComputeShader(GetGlobalShaderMap(GMaxRHIFeatureLevel));
+		FComputeShaderUtils::AddPass(
+			GraphBuilder,
+			RDG_EVENT_NAME("Decal Grid Create (%u decals)", NumDecals),
+			ComputeShader,
+			PassParameters,
+			FComputeShaderUtils::GetGroupCount(FIntPoint(Resolution, Resolution), FComputeShaderUtils::kGolden2DGroupSize));
+
+		// hookup to the actual rendering pass
+		OutParameters.Grid = DecalGridTexture;
+		OutParameters.GridData = GraphBuilder.CreateSRV(DecalGridData, DecalGridDataFormat);
+	}
+	else
+	{
+		// light grid is not needed - just hookup dummy data
+
+		FRDGBufferDesc DecalGridDataDesc = FRDGBufferDesc::CreateBufferDesc(sizeof(uint32), 1);
+		FRDGBuffer* DecalGridData = GraphBuilder.CreateBuffer(DecalGridDataDesc, TEXT("PathTracer.DecalGridData"));
+		AddClearUAVPass(GraphBuilder, GraphBuilder.CreateUAV(DecalGridData, PF_R32_UINT), 0);
+
+		OutParameters.GridResolution = 0;
+		OutParameters.GridMaxCount = 0;
+		OutParameters.GridAxis = 0;
+		OutParameters.Grid = GSystemTextures.GetBlackDummy(GraphBuilder);
+		OutParameters.GridData = GraphBuilder.CreateSRV(DecalGridData, PF_R32_UINT);
+	}
+}
+
+FTransientDecalRenderDataList GetSortedDecals(const TSparseArray<FDeferredDecalProxy*>& Decals, FScene& Scene, const FViewInfo& View)
+{
+	const float FadeMultiplier = DecalRendering::GetDecalFadeScreenSizeMultiplier();
+	const bool bIsPerspectiveProjection = View.IsPerspectiveProjection();
+
+	FTransientDecalRenderDataList SortedDecals;
+	SortedDecals.Reserve(Decals.Num());
+
+	for (auto DecalProxy : Decals)
+	{
+		if (!DecalProxy->DecalMaterial || !DecalProxy->DecalMaterial->IsValidLowLevelFast())
+		{
+			continue;
+		}
+
+		bool bIsShown = true;
+
+		if (!DecalProxy->IsShown(&View))
+		{
+			bIsShown = false;
+		}
+
+		if (bIsShown)
+		{
+			FTransientDecalRenderData Data(Scene, *DecalProxy, 0.0f);
+
+			if (bIsPerspectiveProjection && Data.Proxy.FadeScreenSize != 0.0f)
+			{
+				const FMatrix ComponentToWorldMatrix = Data.Proxy.ComponentTrans.ToMatrixWithScale();
+				Data.FadeAlpha = DecalRendering::CalculateDecalFadeAlpha(Data.Proxy.FadeScreenSize, ComponentToWorldMatrix, View, FadeMultiplier);
+			}
+
+			const bool bShouldRender = Data.FadeAlpha > 0.0f;
+
+			if (bShouldRender)
+			{
+				SortedDecals.Add(Data);
+			}
+		}
+	}
+
+	DecalRendering::SortDecalList(SortedDecals);
+
+	return SortedDecals;
+}
+
+TRDGUniformBufferRef<FRayTracingDecals> CreateRayTracingDecalData(FRDGBuilder& GraphBuilder, FScene& Scene, const FViewInfo& View, uint32 BaseCallableSlotIndex)
+{
+	if (!RHISupportsRayTracingCallableShaders(View.Family->GetShaderPlatform()))
+	{
+		auto* OutParameters = GraphBuilder.AllocParameters<FRayTracingDecals>();
+		OutParameters->Count = 0;
+		OutParameters->TranslatedBoundMin = FVector3f::ZeroVector;
+		OutParameters->TranslatedBoundMax = FVector3f::ZeroVector;
+
+		BuildDecalGrid(GraphBuilder, 0, nullptr, *OutParameters);
+
+		return GraphBuilder.CreateUniformBuffer(OutParameters);
+	}
+
+	const EShaderPlatform ShaderPlatform = View.Family->GetShaderPlatform();
+
+	FTransientDecalRenderDataList SortedDecals = GetSortedDecals(Scene.Decals, Scene, View);
+
+	// Allocate from the graph builder so that we don't need to copy the data again when queuing the upload
+	auto& RayTracingDecals = GraphBuilder.AllocArray<FRayTracingDecal>();
+	RayTracingDecals.Reserve(SortedDecals.Num());
+
+	const float Inf = std::numeric_limits<float>::infinity();
+	FVector3f TranslatedBoundMin(+Inf, +Inf, +Inf);
+	FVector3f TranslatedBoundMax(-Inf, -Inf, -Inf);
+
+	for (int32 Index = 0; Index < SortedDecals.Num(); ++Index)
+	{
+		const auto& DecalData = SortedDecals[Index];
+
+		const FMaterialShaderMap* MaterialShaderMap = DecalData.MaterialResource->GetRenderingThreadShaderMap();
+
+		auto CallableShader = MaterialShaderMap->GetShader<FRayTracingDecalMaterialShader>();
+
+		if (!CallableShader.IsValid())
+		{
+			continue;
+		}
+
+		TUniformBufferRef<FDecalParametersRayTracing> DecalParameters = CreateDecalParametersBuffer(View, ShaderPlatform, DecalData, EUniformBufferUsage::UniformBuffer_SingleFrame);
+
+		Scene.RayTracingScene.UniformBuffers.Add(DecalParameters); // Hold uniform buffer ref in RayTracingScene since FMeshDrawSingleShaderBindings doesn't
+
+		FRayTracingShaderCommand& Command = Scene.RayTracingScene.CallableCommands.AddDefaulted_GetRef();
+		Command.SetShader(CallableShader);
+		Command.SlotInScene = BaseCallableSlotIndex + Index;
+
+		int32 DataOffset = 0;
+		FMeshDrawSingleShaderBindings SingleShaderBindings = Command.ShaderBindings.GetSingleShaderBindings(SF_RayCallable, DataOffset);
+		CallableShader->GetShaderBindings(&Scene, Scene.GetFeatureLevel(), *DecalData.MaterialProxy, *DecalData.MaterialResource, View, DecalParameters, SingleShaderBindings);
+
+		const FBox BoxBounds = DecalData.Proxy.GetBounds().GetBox();
+
+		FRayTracingDecal& DestDecal = RayTracingDecals.AddDefaulted_GetRef();
+		DestDecal.TranslatedBoundMin = FVector3f(BoxBounds.Min + View.ViewMatrices.GetPreViewTranslation());
+		DestDecal.TranslatedBoundMax = FVector3f(BoxBounds.Max + View.ViewMatrices.GetPreViewTranslation());
+		DestDecal.CallableSlotIndex = BaseCallableSlotIndex + Index;
+
+		TranslatedBoundMin = FVector3f::Min(TranslatedBoundMin, DestDecal.TranslatedBoundMin);
+		TranslatedBoundMax = FVector3f::Max(TranslatedBoundMax, DestDecal.TranslatedBoundMax);
+	}
+
+	// Upload the buffer of decals to the GPU
+	const uint32 NumCopyDecals = FMath::Max(1u, (uint32)RayTracingDecals.Num()); // need at least one since zero-sized buffers are not allowed
+	const size_t DataSize = sizeof(FRayTracingDecal) * RayTracingDecals.Num();
+	FRDGBufferSRVRef DecalsSRV = GraphBuilder.CreateSRV(FRDGBufferSRVDesc(CreateStructuredBuffer(GraphBuilder, TEXT("PathTracer.DecalsBuffer"), sizeof(FRayTracingDecal), NumCopyDecals, RayTracingDecals.GetData(), DataSize, ERDGInitialDataFlags::NoCopy)));
+
+	auto* OutParameters = GraphBuilder.AllocParameters<FRayTracingDecals>();
+	OutParameters->Count = RayTracingDecals.Num();
+	OutParameters->TranslatedBoundMin = TranslatedBoundMin;
+	OutParameters->TranslatedBoundMax = TranslatedBoundMax;
+
+	BuildDecalGrid(GraphBuilder, RayTracingDecals.Num(), DecalsSRV, *OutParameters);
+	
+	return GraphBuilder.CreateUniformBuffer(OutParameters);
+}
+
+#endif // RHI_RAYTRACING
diff --git a/Engine/Source/Runtime/Renderer/Private/RayTracing/RayTracingDecals.h b/Engine/Source/Runtime/Renderer/Private/RayTracing/RayTracingDecals.h
new file mode 100644
--- /dev/null
+++ b/Engine/Source/Runtime/Renderer/Private/RayTracing/RayTracingDecals.h
@@ -0,0 +1,25 @@
+// Copyright Epic Games, Inc. All Rights Reserved.
+
+#pragma once
+
+#include "CoreMinimal.h"
+#include "RenderGraphDefinitions.h"
+
+#if RHI_RAYTRACING
+
+class FScene;
+
+BEGIN_GLOBAL_SHADER_PARAMETER_STRUCT(FRayTracingDecals, RENDERER_API)
+	SHADER_PARAMETER_RDG_TEXTURE(Texture2D<uint>, Grid)
+	SHADER_PARAMETER_RDG_BUFFER_SRV(Buffer<uint>, GridData)
+	SHADER_PARAMETER(FVector3f, TranslatedBoundMin)
+	SHADER_PARAMETER(FVector3f, TranslatedBoundMax)
+	SHADER_PARAMETER(uint32, GridResolution)
+	SHADER_PARAMETER(uint32, GridMaxCount)
+	SHADER_PARAMETER(uint32, GridAxis)
+	SHADER_PARAMETER(uint32, Count)
+END_GLOBAL_SHADER_PARAMETER_STRUCT()
+
+TRDGUniformBufferRef<FRayTracingDecals> CreateRayTracingDecalData(FRDGBuilder& GraphBuilder, FScene& Scene, const FViewInfo& View, uint32 BaseCallableSlotIndex);
+
+#endif
diff --git a/Engine/Source/Runtime/Renderer/Private/RayTracing/RayTracingScene.cpp b/Engine/Source/Runtime/Renderer/Private/RayTracing/RayTracingScene.cpp
--- a/Engine/Source/Runtime/Renderer/Private/RayTracing/RayTracingScene.cpp
+++ b/Engine/Source/Runtime/Renderer/Private/RayTracing/RayTracingScene.cpp
@@ -278,6 +278,7 @@ void FRayTracingScene::Reset()
 	Instances.Reset();
 	NumCallableShaderSlots = 0;
 	CallableCommands.Reset();
+	UniformBuffers.Reset();
 	GeometriesToBuild.Reset();
 	UsedCoarseMeshStreamingHandles.Reset();
 
@@ -292,6 +293,7 @@ void FRayTracingScene::ResetAndReleaseResources()
 
 	Instances.Empty();
 	CallableCommands.Empty();
+	UniformBuffers.Empty();
 	GeometriesToBuild.Empty();
 	UsedCoarseMeshStreamingHandles.Empty();
 	RayTracingSceneBuffer = nullptr;
diff --git a/Engine/Source/Runtime/Renderer/Private/RayTracing/RayTracingScene.h b/Engine/Source/Runtime/Renderer/Private/RayTracing/RayTracingScene.h
--- a/Engine/Source/Runtime/Renderer/Private/RayTracing/RayTracingScene.h
+++ b/Engine/Source/Runtime/Renderer/Private/RayTracing/RayTracingScene.h
@@ -86,6 +86,9 @@ public:
 	uint32 NumCallableShaderSlots = 0;
 	TArray<FRayTracingShaderCommand> CallableCommands;
 
+	// Helper array to hold references to single frame uniform buffers used in SBTs
+	TArray<FUniformBufferRHIRef> UniformBuffers;
+
 	// Geometries which still have a pending build request but are used this frame and require a force build.
 	TArray<const FRayTracingGeometry*> GeometriesToBuild;
 
diff --git a/Engine/Source/Runtime/Renderer/Private/SceneRendering.h b/Engine/Source/Runtime/Renderer/Private/SceneRendering.h
--- a/Engine/Source/Runtime/Renderer/Private/SceneRendering.h
+++ b/Engine/Source/Runtime/Renderer/Private/SceneRendering.h
@@ -49,6 +49,7 @@ struct FILCUpdatePrimTaskData;
 class FPostprocessContext;
 struct FILCUpdatePrimTaskData;
 class FRaytracingLightDataPacked;
+class FRayTracingDecals;
 class FRayTracingLocalShaderBindingWriter;
 class FVirtualShadowMapClipmap;
 class FShadowProjectionPassParameters;
@@ -1492,6 +1493,7 @@ public:
 	FShaderResourceViewRHIRef						RayTracingSubSurfaceProfileSRV;
 
 	TRDGUniformBufferRef<FRaytracingLightDataPacked>	RayTracingLightDataUniformBuffer;
+	TRDGUniformBufferRef<FRayTracingDecals>				RayTracingDecalUniformBuffer;
 #endif // RHI_RAYTRACING
 
 	/**
