Removes accidental log spam in translucency

diff --git a/Engine/Source/Runtime/Renderer/Private/TranslucentRendering.cpp b/Engine/Source/Runtime/Renderer/Private/TranslucentRendering.cpp
--- a/Engine/Source/Runtime/Renderer/Private/TranslucentRendering.cpp
+++ b/Engine/Source/Runtime/Renderer/Private/TranslucentRendering.cpp
@@ -229,8 +229,6 @@ FSeparateTranslucencyDimensions UpdateSeparateTranslucencyDimensions(const FView
 		MaxTranslucencyResolutionFraction = ActiveViewFamily.DynamicResolutionUpperBounds[GDynamicTranslucencyResolution];
 	}
 
-	UE_LOG(LogRenderer, Display, TEXT("TranslucencyResolutionFraction=%f"), TranslucencyResolutionFraction);
-
 	FSeparateTranslucencyDimensions Dimensions;
 	// TODO: this should be MaxTranslucencyResolutionFraction instead of TranslucencyResolutionFraction to keep the size of render target stable, but the SvPositionToBuffer() is broken in material.
 	Dimensions.Extent = GetScaledExtent(ActiveViewFamily.SceneTexturesConfig.Extent, TranslucencyResolutionFraction);
