Path Tracer: invalidate path traced output when decals are added, removed or transformed

diff --git a/Engine/Source/Runtime/Renderer/Private/RendererScene.cpp b/Engine/Source/Runtime/Renderer/Private/RendererScene.cpp
--- a/Engine/Source/Runtime/Renderer/Private/RendererScene.cpp
+++ b/Engine/Source/Runtime/Renderer/Private/RendererScene.cpp
@@ -2223,6 +2223,7 @@
 	if(bAdd)
 	{
 		Decals.Add(Proxy);
+		InvalidatePathTracedOutput();
 	}
 	else
 	{
@@ -2233,6 +2234,7 @@
 
 			if (CurrentProxy == Proxy)
 			{
+				InvalidatePathTracedOutput();
 				It.RemoveCurrent();
 				delete CurrentProxy;
 				break;
@@ -2336,12 +2338,18 @@
 	if(Decal->SceneProxy)
 	{
 		//Send command to the rendering thread to update the decal's transform.
+		FScene* Scene = this;
 		FDeferredDecalProxy* DecalSceneProxy = Decal->SceneProxy;
 		FTransform ComponentToWorldIncludingDecalSize = Decal->GetTransformIncludingDecalSize();
 		FBoxSphereBounds Bounds = Decal->CalcBounds(Decal->GetComponentTransform());
 		ENQUEUE_RENDER_COMMAND(UpdateTransformCommand)(
-			[DecalSceneProxy, ComponentToWorldIncludingDecalSize, Bounds](FRHICommandListImmediate& RHICmdList)
+			[DecalSceneProxy, ComponentToWorldIncludingDecalSize, Bounds, Scene](FRHICommandListImmediate& RHICmdList)
 			{
+				// Invalidate the path tracer only if the decal was sufficiently moved
+				if (!ComponentToWorldIncludingDecalSize.Equals(DecalSceneProxy->ComponentTrans, SMALL_NUMBER))
+				{
+					Scene->InvalidatePathTracedOutput();
+				}
 				// Update the primitive's transform.
 				DecalSceneProxy->SetTransformIncludingDecalSize(ComponentToWorldIncludingDecalSize, Bounds);
 			});
