Fix crash when using time-sliced sky captures.

The code was never allocating the second sky render target because it was using the wrong index in the null check (ConvolvedSkyRenderTargetReadyIndex is 0 at that point).

#rnx

diff --git a/Engine/Source/Runtime/Renderer/Private/ReflectionEnvironmentRealTimeCapture.cpp b/Engine/Source/Runtime/Renderer/Private/ReflectionEnvironmentRealTimeCapture.cpp
--- a/Engine/Source/Runtime/Renderer/Private/ReflectionEnvironmentRealTimeCapture.cpp
+++ b/Engine/Source/Runtime/Renderer/Private/ReflectionEnvironmentRealTimeCapture.cpp
@@ -353,7 +353,7 @@
 		GRenderTargetPool.FindFreeElement(GraphBuilder.RHICmdList, SkyCubeTexDesc, ConvolvedSkyRenderTarget[0], TEXT("SkyLight.ConvolvedSkyRenderTarget0"));
 		GRenderTargetPool.FindFreeElement(GraphBuilder.RHICmdList, SkyCubeTexDesc, CapturedSkyRenderTarget, TEXT("SkyLight.CapturedSkyRenderTarget"));
 	}
-	if (bTimeSlicedRealTimeCapture && (CubeResolutionInvalidated || !ConvolvedSkyRenderTarget[ConvolvedSkyRenderTargetReadyIndex].IsValid()))
+	if (bTimeSlicedRealTimeCapture && (!ConvolvedSkyRenderTarget[1].IsValid() || CubeResolutionInvalidated))
 	{
 		// Additional allocation for time slicing
 		GRenderTargetPool.FindFreeElement(GraphBuilder.RHICmdList, SkyCubeTexDesc, ConvolvedSkyRenderTarget[1], TEXT("SkyLight.ConvolvedSkyRenderTarget1"));
