Fixed assert triggering because baked depth data was invalid

diff --git a/Engine/Source/Runtime/Renderer/Private/LightGridInjection.cpp b/Engine/Source/Runtime/Renderer/Private/LightGridInjection.cpp
--- a/Engine/Source/Runtime/Renderer/Private/LightGridInjection.cpp
+++ b/Engine/Source/Runtime/Renderer/Private/LightGridInjection.cpp
@@ -580,7 +580,11 @@ void FSceneRenderer::ComputeLightGrid(FRDGBuilder& GraphBuilder, bool bCullLight
 							}
 
 							const FStaticShadowDepthMap* StaticShadowDepthMap = LightSceneInfo->Proxy->GetStaticShadowDepthMap();
-							const uint32 bStaticallyShadowedValue = LightSceneInfo->IsPrecomputedLightingValid() && StaticShadowDepthMap && StaticShadowDepthMap->Data && StaticShadowDepthMap->TextureRHI ? 1 : 0;
+							const uint32 bStaticallyShadowedValue = LightSceneInfo->IsPrecomputedLightingValid() 
+																	&& StaticShadowDepthMap 
+																	&& StaticShadowDepthMap->Data 
+																	&& !StaticShadowDepthMap->Data->WorldToLight.ContainsNaN()
+																	&& StaticShadowDepthMap->TextureRHI ? 1 : 0;
 							ForwardLightData->DirectionalLightUseStaticShadowing = bStaticallyShadowedValue;
 							if (bStaticallyShadowedValue)
 							{
diff --git a/Engine/Source/Runtime/Renderer/Private/VolumeLighting.cpp b/Engine/Source/Runtime/Renderer/Private/VolumeLighting.cpp
--- a/Engine/Source/Runtime/Renderer/Private/VolumeLighting.cpp
+++ b/Engine/Source/Runtime/Renderer/Private/VolumeLighting.cpp
@@ -143,7 +143,11 @@ void GetVolumeShadowingShaderParameters(
 	check(OutParameters.ShadowDepthTexture);
 
 	const FStaticShadowDepthMap* StaticShadowDepthMap = LightSceneInfo->Proxy->GetStaticShadowDepthMap();
-	const uint32 bStaticallyShadowedValue = LightSceneInfo->IsPrecomputedLightingValid() && StaticShadowDepthMap && StaticShadowDepthMap->Data && StaticShadowDepthMap->TextureRHI ? 1 : 0;
+	const uint32 bStaticallyShadowedValue = LightSceneInfo->IsPrecomputedLightingValid() 
+											&& StaticShadowDepthMap 
+											&& StaticShadowDepthMap->Data 
+											&& !StaticShadowDepthMap->Data->WorldToLight.ContainsNaN()
+											&& StaticShadowDepthMap->TextureRHI ? 1 : 0;
 	FRHITexture* StaticShadowDepthMapTexture = bStaticallyShadowedValue ? StaticShadowDepthMap->TextureRHI : GWhiteTexture->TextureRHI;
 	const FMatrix WorldToStaticShadow = bStaticallyShadowedValue ? StaticShadowDepthMap->Data->WorldToLight : FMatrix::Identity;
 	const FVector4f StaticShadowBufferSizeValue = bStaticallyShadowedValue ? FVector4f(StaticShadowDepthMap->Data->ShadowMapSizeX, StaticShadowDepthMap->Data->ShadowMapSizeY, 1.0f / StaticShadowDepthMap->Data->ShadowMapSizeX, 1.0f / StaticShadowDepthMap->Data->ShadowMapSizeY) : FVector4f(0, 0, 0, 0);
