Fix some items rendering to a separate velocity pass when base velocity pass is used with r.SelectiveBasePassOutputs.
This logic took primitives with a material using static lighting and moved velocity writing to a second pass.
Seems like that is some old behavior and after asking around it isn't well known why we had it.
Anyway it's not good to pay the cost of that second pass, so removing that.

diff --git a/Engine/Shaders/Private/BasePassCommon.ush b/Engine/Shaders/Private/BasePassCommon.ush
--- a/Engine/Shaders/Private/BasePassCommon.ush
+++ b/Engine/Shaders/Private/BasePassCommon.ush
@@ -48,7 +48,7 @@
 
 // If a primitive has static lighting, we assume it is not moving. If it is, it will be rerendered in an extra renderpass.
 #define SUPPORTS_WRITING_VELOCITY_TO_BASE_PASS	(FEATURE_LEVEL >= FEATURE_LEVEL_SM4 && (MATERIALBLENDING_SOLID || MATERIALBLENDING_MASKED) && !SIMPLE_FORWARD_SHADING) 
-#define WRITES_VELOCITY_TO_GBUFFER				((SUPPORTS_WRITING_VELOCITY_TO_BASE_PASS || USES_GBUFFER) && GBUFFER_HAS_VELOCITY && (!SELECTIVE_BASEPASS_OUTPUTS || !(STATICLIGHTING_TEXTUREMASK || STATICLIGHTING_SIGNEDDISTANCEFIELD || HQ_TEXTURE_LIGHTMAP || LQ_TEXTURE_LIGHTMAP || WATER_MESH_FACTORY)))
+#define WRITES_VELOCITY_TO_GBUFFER				((SUPPORTS_WRITING_VELOCITY_TO_BASE_PASS || USES_GBUFFER) && GBUFFER_HAS_VELOCITY)
 
 #define TRANSLUCENCY_ANY_PERVERTEX_LIGHTING (TRANSLUCENCY_LIGHTING_VOLUMETRIC_PERVERTEX_NONDIRECTIONAL || TRANSLUCENCY_LIGHTING_VOLUMETRIC_PERVERTEX_DIRECTIONAL)
 #define TRANSLUCENCY_ANY_VOLUMETRIC (TRANSLUCENCY_LIGHTING_VOLUMETRIC_NONDIRECTIONAL || TRANSLUCENCY_LIGHTING_VOLUMETRIC_DIRECTIONAL || TRANSLUCENCY_ANY_PERVERTEX_LIGHTING)
diff --git a/Engine/Source/Runtime/RenderCore/Private/ShaderMaterialDerivedHelpers.cpp b/Engine/Source/Runtime/RenderCore/Private/ShaderMaterialDerivedHelpers.cpp
--- a/Engine/Source/Runtime/RenderCore/Private/ShaderMaterialDerivedHelpers.cpp
+++ b/Engine/Source/Runtime/RenderCore/Private/ShaderMaterialDerivedHelpers.cpp
@@ -60,7 +60,7 @@ FShaderMaterialDerivedDefines RENDERCORE_API CalculateDerivedMaterialParameters(
 
 	// If a primitive has static lighting, we assume it is not moving. If it is, it will be rerendered in an extra renderpass.
 	Dst.SUPPORTS_WRITING_VELOCITY_TO_BASE_PASS = (FEATURE_LEVEL >= ERHIFeatureLevel::SM4_REMOVED && (Mat.MATERIALBLENDING_SOLID || Mat.MATERIALBLENDING_MASKED) && !Lightmap.SIMPLE_FORWARD_SHADING);
-	Dst.WRITES_VELOCITY_TO_GBUFFER = ((Dst.SUPPORTS_WRITING_VELOCITY_TO_BASE_PASS || Dst.USES_GBUFFER) && SrcGlobal.GBUFFER_HAS_VELOCITY && (!SrcGlobal.SELECTIVE_BASEPASS_OUTPUTS || !(Lightmap.STATICLIGHTING_TEXTUREMASK || Lightmap.STATICLIGHTING_SIGNEDDISTANCEFIELD || Lightmap.HQ_TEXTURE_LIGHTMAP || Lightmap.LQ_TEXTURE_LIGHTMAP || Lightmap.WATER_MESH_FACTORY)));
+	Dst.WRITES_VELOCITY_TO_GBUFFER = ((Dst.SUPPORTS_WRITING_VELOCITY_TO_BASE_PASS || Dst.USES_GBUFFER) && SrcGlobal.GBUFFER_HAS_VELOCITY);
 
 	Dst.TRANSLUCENCY_ANY_PERVERTEX_LIGHTING = (Mat.TRANSLUCENCY_LIGHTING_VOLUMETRIC_PERVERTEX_NONDIRECTIONAL || Mat.TRANSLUCENCY_LIGHTING_VOLUMETRIC_PERVERTEX_DIRECTIONAL);
 	Dst.TRANSLUCENCY_ANY_VOLUMETRIC = (Mat.TRANSLUCENCY_LIGHTING_VOLUMETRIC_NONDIRECTIONAL || Mat.TRANSLUCENCY_LIGHTING_VOLUMETRIC_DIRECTIONAL || Dst.TRANSLUCENCY_ANY_PERVERTEX_LIGHTING);
diff --git a/Engine/Source/Runtime/Renderer/Private/DeferredShadingRenderer.cpp b/Engine/Source/Runtime/Renderer/Private/DeferredShadingRenderer.cpp
--- a/Engine/Source/Runtime/Renderer/Private/DeferredShadingRenderer.cpp
+++ b/Engine/Source/Runtime/Renderer/Private/DeferredShadingRenderer.cpp
@@ -2375,7 +2375,6 @@
 
 	const bool bShouldRenderVelocities = ShouldRenderVelocities();
 	const bool bBasePassCanOutputVelocity = FVelocityRendering::BasePassCanOutputVelocity(FeatureLevel);
-	const bool bUseSelectiveBasePassOutputs = IsUsingSelectiveBasePassOutputs(ShaderPlatform);
 	const bool bHairStrandsEnable = HairStrandsBookmarkParameters.HasInstances() && Views.Num() > 0 && IsHairStrandsEnabled(EHairStrandsShaderType::Strands, GetViewFamily(Views).GetShaderPlatform());
 
 	{
@@ -2906,18 +2905,12 @@
 		}
 	}
 
-	// TODO: Keeping the velocities here for testing, but if that works, this pass will be remove and DDM_AllOpaqueNoVelocity will be the only option with
-	// DBuffer decals enabled.
-
-	// If bBasePassCanOutputVelocity is set, basepass fully writes the velocity buffer unless bUseSelectiveBasePassOutputs is enabled.
-	if (bShouldRenderVelocities && (!bBasePassCanOutputVelocity || bUseSelectiveBasePassOutputs) && (Scene->EarlyZPassMode != DDM_AllOpaqueNoVelocity))
+	// If we are not rendering velocities in depth or base pass then do that here.
+	if (bShouldRenderVelocities && !bBasePassCanOutputVelocity && (Scene->EarlyZPassMode != DDM_AllOpaqueNoVelocity))
 	{
-		// Render the velocities of movable objects
 		GraphBuilder.SetCommandListStat(GET_STATID(STAT_CLM_Velocity));
 		RenderVelocities(GraphBuilder, SceneTextures, EVelocityPass::Opaque, bHairStrandsEnable);
 		GraphBuilder.SetCommandListStat(GET_STATID(STAT_CLM_AfterVelocity));
-
-		// TODO: Populate velocity buffer from Nanite visibility buffer.
 	}
 
 	// Copy lighting channels out of stencil before deferred decals which overwrite those values
diff --git a/Engine/Source/Runtime/Renderer/Private/VelocityRendering.cpp b/Engine/Source/Runtime/Renderer/Private/VelocityRendering.cpp
--- a/Engine/Source/Runtime/Renderer/Private/VelocityRendering.cpp
+++ b/Engine/Source/Runtime/Renderer/Private/VelocityRendering.cpp
@@ -108,14 +108,15 @@ public:
 		const bool bHasPlatformSupport = PlatformSupportsVelocityRendering(Parameters.Platform);
 
 		/**
-		 * Any material with a vertex factory incompatible with base pass velocity generation must generate
-		 * permutations for this shader. Shaders which don't fall into this set are considered "simple" enough
-		 * to swap against the default material. This massively simplifies the calculations.
+		 * If we don't use base pass velocity then we may need to generate permutations for this shader. 
+		 * We only need to compile shaders which aren't considered "simple" enough to swap against the default material. 
+		 * This massively simplifies the calculations.
 		 */
-		const bool bIsSeparateVelocityPassRequired = (bIsDefault || bIsMasked || bIsOpaqueAndTwoSided || bMayModifyMeshes) &&
-			FVelocityRendering::IsSeparateVelocityPassRequiredByVertexFactory(Parameters.Platform, Parameters.VertexFactoryType->SupportsStaticLighting());
+		const bool bIsSeparateVelocityPassRequired = 
+			!FVelocityRendering::BasePassCanOutputVelocity(Parameters.Platform) &&
+			(bIsDefault || bIsMasked || bIsOpaqueAndTwoSided || bMayModifyMeshes);
 
-		// The material may explicitly override and request that it be rendered into the velocity pass.
+		// The material may explicitly request that it be rendered into the translucent velocity pass.
 		const bool bIsSeparateVelocityPassRequiredByMaterial = Parameters.MaterialParameters.bIsTranslucencyWritingVelocity;
 
 		return bHasPlatformSupport && (bIsSeparateVelocityPassRequired || bIsSeparateVelocityPassRequiredByMaterial);
@@ -404,17 +405,6 @@ bool FVelocityRendering::BasePassCanOutputVelocity(ERHIFeatureLevel::Type Featur
 	return BasePassCanOutputVelocity(ShaderPlatform);
 }
 
-bool FVelocityRendering::IsSeparateVelocityPassRequiredByVertexFactory(EShaderPlatform ShaderPlatform, bool bVertexFactoryUsesStaticLighting)
-{
-	// A separate pass is required if the base pass can't do it.
-	const bool bBasePassVelocityNotSupported = !BasePassCanOutputVelocity(ShaderPlatform);
-
-	// Meshes with static lighting need a separate velocity pass, but only if we are using selective render target outputs.
-	const bool bVertexFactoryRequiresSeparateVelocityPass = IsUsingSelectiveBasePassOutputs(ShaderPlatform) && bVertexFactoryUsesStaticLighting;
-
-	return bBasePassVelocityNotSupported || bVertexFactoryRequiresSeparateVelocityPass;
-}
-
 bool FVelocityRendering::IsParallelVelocity(EShaderPlatform ShaderPlatform)
 {
 	return GRHICommandList.UseParallelAlgorithms() && CVarParallelVelocity.GetValueOnRenderThread()
@@ -468,19 +458,6 @@ bool FOpaqueVelocityMeshProcessor::PrimitiveCanHaveVelocity(EShaderPlatform Shad
 		return false;
 	}
 
-	/**
-	 * Whether the vertex factory for this primitive requires that it render in the separate velocity pass, as opposed to the base pass.
-	 * In cases where the base pass is rendering opaque velocity for a particular mesh batch, we want to filter it out from this pass,
-	 * which performs a separate draw call to render velocity.
-	 */
-	const bool bIsSeparateVelocityPassRequiredByVertexFactory =
-		FVelocityRendering::IsSeparateVelocityPassRequiredByVertexFactory(ShaderPlatform, PrimitiveSceneProxy->HasStaticLighting());
-
-	if (!bIsSeparateVelocityPassRequiredByVertexFactory)
-	{
-		return false;
-	}
-
 	return true;
 }
 
diff --git a/Engine/Source/Runtime/Renderer/Private/VelocityRendering.h b/Engine/Source/Runtime/Renderer/Private/VelocityRendering.h
--- a/Engine/Source/Runtime/Renderer/Private/VelocityRendering.h
+++ b/Engine/Source/Runtime/Renderer/Private/VelocityRendering.h
@@ -50,9 +50,6 @@ struct FVelocityRendering
 	/** Returns true if the velocity can be output in the BasePass. Only valid for the current platform. */
 	static bool BasePassCanOutputVelocity(ERHIFeatureLevel::Type FeatureLevel);
 
-	/** Returns true if a separate velocity pass is required (i.e. not rendered by the base pass) given the provided vertex factory settings. */
-	static bool IsSeparateVelocityPassRequiredByVertexFactory(EShaderPlatform ShaderPlatform, bool bVertexFactoryUsesStaticLighting);
-
 	/** Returns true if the velocity pass is using parallel dispatch. */
 	static bool IsParallelVelocity(EShaderPlatform ShaderPlatform);
 
