Scene Capture:  Fix for additive / composite blending.  Need to skip render target clear when compositing.

diff --git a/Engine/Source/Runtime/Renderer/Private/SceneCaptureRendering.cpp b/Engine/Source/Runtime/Renderer/Private/SceneCaptureRendering.cpp
--- a/Engine/Source/Runtime/Renderer/Private/SceneCaptureRendering.cpp
+++ b/Engine/Source/Runtime/Renderer/Private/SceneCaptureRendering.cpp
@@ -1091,8 +1091,11 @@ void FScene::UpdateSceneCaptureContents(USceneCaptureComponent2D* CaptureCompone
 
 		UTexture* TexturePtrNotDeferenced = TextureRenderTarget;
 
+		// Compositing feature is only active when using SceneColor as the source
+		bool bIsCompositing = (CaptureComponent->CompositeMode != SCCM_Overwrite) && (CaptureComponent->CaptureSource == SCS_SceneColorHDR);
+
 		ENQUEUE_RENDER_COMMAND(CaptureCommand)(
-			[SceneRenderer, TextureRenderTargetResource, TexturePtrNotDeferenced, EventName, bGenerateMips, GenerateMipsParams, bDisableFlipCopyGLES, GameViewportRT, bEnableOrthographicTiling, bOrthographicCamera, NumXTiles, NumYTiles, TileID](FRHICommandListImmediate& RHICmdList)
+			[SceneRenderer, TextureRenderTargetResource, TexturePtrNotDeferenced, EventName, bGenerateMips, GenerateMipsParams, bDisableFlipCopyGLES, GameViewportRT, bEnableOrthographicTiling, bIsCompositing, bOrthographicCamera, NumXTiles, NumYTiles, TileID](FRHICommandListImmediate& RHICmdList)
 			{
 				if (GameViewportRT != nullptr)
 				{
@@ -1120,7 +1123,10 @@ void FScene::UpdateSceneCaptureContents(USceneCaptureComponent2D* CaptureCompone
 
 				RectLightAtlas::FAtlasTextureInvalidationScope Invalidation(TexturePtrNotDeferenced);
 
-				UpdateSceneCaptureContent_RenderThread(RHICmdList, SceneRenderer, TextureRenderTargetResource, TextureRenderTargetResource, EventName, CopyInfo, bGenerateMips, GenerateMipsParams, bDisableFlipCopyGLES, !bEnableOrthographicTiling, bOrthographicCamera);
+				// Don't clear the render target when compositing, or in a tiling mode that fills in the render target in multiple passes.
+				bool bClearRenderTarget = !bIsCompositing && !bEnableOrthographicTiling;
+
+				UpdateSceneCaptureContent_RenderThread(RHICmdList, SceneRenderer, TextureRenderTargetResource, TextureRenderTargetResource, EventName, CopyInfo, bGenerateMips, GenerateMipsParams, bDisableFlipCopyGLES, bClearRenderTarget, bOrthographicCamera);
 			}
 		);
 	}
