Fix for IsValid() crash when showing RayTracingDebug:
* When RT debug is enabled we disable translucency but some of the postprocessing materials might use that texture so we need to set correct ViewRect as well.

diff --git a/Engine/Source/Runtime/Renderer/Private/PostProcess/PostProcessing.cpp b/Engine/Source/Runtime/Renderer/Private/PostProcess/PostProcessing.cpp
--- a/Engine/Source/Runtime/Renderer/Private/PostProcess/PostProcessing.cpp
+++ b/Engine/Source/Runtime/Renderer/Private/PostProcess/PostProcessing.cpp
@@ -384,7 +384,14 @@ void AddPostProcessingPasses(
 		FPostProcessMaterialInputs PostProcessMaterialInputs;
 
 		PostProcessMaterialInputs.SetInput(EPostProcessMaterialInput::SceneColor, InSceneColor);
-		PostProcessMaterialInputs.SetInput(EPostProcessMaterialInput::SeparateTranslucency, FScreenPassTexture(PostDOFTranslucencyResources.GetColorForRead(GraphBuilder), PostDOFTranslucencyResources.ViewRect));
+
+		FIntRect ViewRect{ 0, 0, 1, 1 };
+		if (PostDOFTranslucencyResources.IsValid())
+		{
+			ViewRect = PostDOFTranslucencyResources.ViewRect;
+		}
+
+		PostProcessMaterialInputs.SetInput(EPostProcessMaterialInput::SeparateTranslucency, FScreenPassTexture(PostDOFTranslucencyResources.GetColorForRead(GraphBuilder), ViewRect));
 		PostProcessMaterialInputs.SetInput(EPostProcessMaterialInput::Velocity, Velocity);
 		PostProcessMaterialInputs.SceneTextures = GetSceneTextureShaderParameters(Inputs.SceneTextures);
 		PostProcessMaterialInputs.CustomDepthTexture = CustomDepth.Texture;
