Allows to control TSR pass to run on asynccompute at thiner granularity

diff --git a/Engine/Source/Runtime/Renderer/Private/PostProcess/TemporalSuperResolution.cpp b/Engine/Source/Runtime/Renderer/Private/PostProcess/TemporalSuperResolution.cpp
--- a/Engine/Source/Runtime/Renderer/Private/PostProcess/TemporalSuperResolution.cpp
+++ b/Engine/Source/Runtime/Renderer/Private/PostProcess/TemporalSuperResolution.cpp
@@ -56,7 +56,11 @@ TAutoConsoleVariable<int32> CVarTSRRejectionAntiAliasingQuality(
 
 TAutoConsoleVariable<int32> CVarTSRAsyncCompute(
 	TEXT("r.TSR.AsyncCompute"), 0,
-	TEXT("Whether to run TSR on async compute. Some TSR passes can overlap with previous passe.\n"),
+	TEXT("Whether to run TSR on async compute. Some TSR passes can overlap with previous passe.\n")
+	TEXT(" 0: Disabled (default);\n")
+	TEXT(" 1: Only ClearPrevTextures pass;\n")
+	TEXT(" 2: Only ClearPrevTextures through DecimateHistory passes;\n")
+	TEXT(" 3: All passes;"),
 	ECVF_RenderThreadSafe);
 
 TAutoConsoleVariable<float> CVarTSRTranslucencyHighlightLuminance(
@@ -651,7 +655,7 @@ ITemporalUpscaler::FOutputs AddTemporalSuperResolutionPasses(
 	const bool bSupportsAlpha = IsPostProcessingWithAlphaChannelSupported();
 
 	// whether TSR passes can run on async compute.
-	ERDGPassFlags ComputePassFlags = (GSupportsEfficientAsyncCompute && CVarTSRAsyncCompute.GetValueOnRenderThread() != 0) ? ERDGPassFlags::AsyncCompute : ERDGPassFlags::Compute;
+	int32 AsyncComputePasses = GSupportsEfficientAsyncCompute ? CVarTSRAsyncCompute.GetValueOnRenderThread() : 0;
 
 	const bool bHistoryHighFrequencyOnly = CVarTSRHistoryHighFrequencyOnly.GetValueOnRenderThread() != 0;
 
@@ -849,7 +853,7 @@ ITemporalUpscaler::FOutputs AddTemporalSuperResolutionPasses(
 		FComputeShaderUtils::AddPass(
 			GraphBuilder,
 			RDG_EVENT_NAME("TSR ClearPrevTextures %dx%d", InputRect.Width(), InputRect.Height()),
-			ComputePassFlags,
+			AsyncComputePasses >= 1 ? ERDGPassFlags::AsyncCompute : ERDGPassFlags::Compute,
 			ComputeShader,
 			PassParameters,
 			FComputeShaderUtils::GetGroupCount(InputRect.Size(), 8 * 2));
@@ -949,7 +953,7 @@ ITemporalUpscaler::FOutputs AddTemporalSuperResolutionPasses(
 			RDG_EVENT_NAME("TSR DilateVelocity(MotionBlurDirections=%d) %dx%d",
 				int32(PermutationVector.Get<FTSRDilateVelocityCS::FMotionBlurDirectionsDim>()),
 				InputRect.Width(), InputRect.Height()),
-			ComputePassFlags,
+			AsyncComputePasses >= 2 ? ERDGPassFlags::AsyncCompute : ERDGPassFlags::Compute,
 			ComputeShader,
 			PassParameters,
 			FComputeShaderUtils::GetGroupCount(InputRect.Size(), TileSize));
@@ -1167,7 +1171,7 @@ ITemporalUpscaler::FOutputs AddTemporalSuperResolutionPasses(
 				bHistoryHighFrequencyOnly ? TEXT("HighFrequency") : TEXT("LowFrequency"),
 				bHalfResLowFrequency ? TEXT(" HalfResShadingOutput") : TEXT(""),
 				InputRect.Width(), InputRect.Height()),
-			ComputePassFlags,
+			AsyncComputePasses >= 2 ? ERDGPassFlags::AsyncCompute : ERDGPassFlags::Compute,
 			ComputeShader,
 			PassParameters,
 			FComputeShaderUtils::GetGroupCount(InputRect.Size(), 8));
@@ -1227,7 +1231,7 @@ ITemporalUpscaler::FOutputs AddTemporalSuperResolutionPasses(
 		FComputeShaderUtils::AddPass(
 			GraphBuilder,
 			RDG_EVENT_NAME("TSR CompareTranslucency %dx%d", InputRect.Width(), InputRect.Height()),
-			ComputePassFlags,
+			AsyncComputePasses >= 3 ? ERDGPassFlags::AsyncCompute : ERDGPassFlags::Compute,
 			ComputeShader,
 			PassParameters,
 			FComputeShaderUtils::GetGroupCount(InputRect.Size(), 8));
@@ -1294,7 +1298,7 @@ ITemporalUpscaler::FOutputs AddTemporalSuperResolutionPasses(
 				RDG_EVENT_NAME("TSR FilterFrequencies(%s) %dx%d",
 					PermutationVector.Get<FTSRFilterFrequenciesCS::FOutputAALumaDim>() ? TEXT("OutputAALuma") : TEXT(""),
 					LowFrequencyRect.Width(), LowFrequencyRect.Height()),
-				ComputePassFlags,
+				AsyncComputePasses >= 3 ? ERDGPassFlags::AsyncCompute : ERDGPassFlags::Compute,
 				ComputeShader,
 				PassParameters,
 				FComputeShaderUtils::GetGroupCount(LowFrequencyRect.Size(), 16));
@@ -1330,7 +1334,7 @@ ITemporalUpscaler::FOutputs AddTemporalSuperResolutionPasses(
 			FComputeShaderUtils::AddPass(
 				GraphBuilder,
 				RDG_EVENT_NAME("TSR CompareHistory %dx%d", LowFrequencyRect.Width(), LowFrequencyRect.Height()),
-				ComputePassFlags,
+				AsyncComputePasses >= 3 ? ERDGPassFlags::AsyncCompute : ERDGPassFlags::Compute,
 				ComputeShader,
 				PassParameters,
 				FComputeShaderUtils::GetGroupCount(LowFrequencyRect.Size(), 16));
@@ -1375,7 +1379,7 @@ ITemporalUpscaler::FOutputs AddTemporalSuperResolutionPasses(
 				RDG_EVENT_NAME("TSR SpatialAntiAliasing(Quality=%d) %dx%d",
 					RejectionAntiAliasingQuality,
 					InputRect.Width(), InputRect.Height()),
-				ComputePassFlags,
+				AsyncComputePasses >= 3 ? ERDGPassFlags::AsyncCompute : ERDGPassFlags::Compute,
 				ComputeShader,
 				PassParameters,
 				FComputeShaderUtils::GetGroupCount(InputRect.Size(), 8));
@@ -1392,7 +1396,7 @@ ITemporalUpscaler::FOutputs AddTemporalSuperResolutionPasses(
 			FComputeShaderUtils::AddPass(
 				GraphBuilder,
 				RDG_EVENT_NAME("TSR FilterAntiAliasing %dx%d", InputRect.Width(), InputRect.Height()),
-				ComputePassFlags,
+				AsyncComputePasses >= 3 ? ERDGPassFlags::AsyncCompute : ERDGPassFlags::Compute,
 				ComputeShader,
 				PassParameters,
 				FComputeShaderUtils::GetGroupCount(InputRect.Size(), 8));
@@ -1429,7 +1433,7 @@ ITemporalUpscaler::FOutputs AddTemporalSuperResolutionPasses(
 		FComputeShaderUtils::AddPass(
 			GraphBuilder,
 			RDG_EVENT_NAME("TSR PostfilterRejection %dx%d", Rect.Width(), Rect.Height()),
-			ComputePassFlags,
+			AsyncComputePasses >= 3 ? ERDGPassFlags::AsyncCompute : ERDGPassFlags::Compute,
 			ComputeShader,
 			PassParameters,
 			FComputeShaderUtils::GetGroupCount(Rect.Size(), 8));
@@ -1452,7 +1456,7 @@ ITemporalUpscaler::FOutputs AddTemporalSuperResolutionPasses(
 		FComputeShaderUtils::AddPass(
 			GraphBuilder,
 			RDG_EVENT_NAME("TSR DilateRejection %dx%d", RejectionRect.Width(), RejectionRect.Height()),
-			ComputePassFlags,
+			AsyncComputePasses >= 3 ? ERDGPassFlags::AsyncCompute : ERDGPassFlags::Compute,
 			ComputeShader,
 			PassParameters,
 			FComputeShaderUtils::GetGroupCount(RejectionRect.Size(), 8));
@@ -1566,7 +1570,7 @@ ITemporalUpscaler::FOutputs AddTemporalSuperResolutionPasses(
 				PermutationVector.Get<FTSRHighFrequencyOnlyDim>() ? TEXT("") : TEXT(" LowFrequency"),
 				PassParameters->bGenerateOutputMip1 ? TEXT(" OutputMip1") : TEXT(""),
 				HistorySize.X, HistorySize.Y),
-			ComputePassFlags,
+			AsyncComputePasses >= 3 ? ERDGPassFlags::AsyncCompute : ERDGPassFlags::Compute,
 			ComputeShader,
 			PassParameters,
 			FComputeShaderUtils::GetGroupCount(HistorySize, 8));
@@ -1588,7 +1592,7 @@ ITemporalUpscaler::FOutputs AddTemporalSuperResolutionPasses(
 		FComputeShaderUtils::AddPass(
 			GraphBuilder,
 			RDG_EVENT_NAME("TSR DebugHistory %dx%d", HistorySize.X, HistorySize.Y),
-			ComputePassFlags,
+			AsyncComputePasses >= 3 ? ERDGPassFlags::AsyncCompute : ERDGPassFlags::Compute,
 			ComputeShader,
 			PassParameters,
 			FComputeShaderUtils::GetGroupCount(HistorySize * kHistoryUpscalingFactor, 8));
@@ -1634,7 +1638,7 @@ ITemporalUpscaler::FOutputs AddTemporalSuperResolutionPasses(
 		FComputeShaderUtils::AddPass(
 			GraphBuilder,
 			RDG_EVENT_NAME("TSR ResolveHistory %dx%d", OutputRect.Width(), OutputRect.Height()),
-			ComputePassFlags,
+			AsyncComputePasses >= 3 ? ERDGPassFlags::AsyncCompute : ERDGPassFlags::Compute,
 			ComputeShader,
 			PassParameters,
 			FComputeShaderUtils::GetGroupCount(OutputRect.Size(), 8));
