Added sky atmosphere transmittance on forward shaded translucent when  per pixel mode is selected on the atmospoheric dir light.

diff --git a/Engine/Shaders/Private/BasePassPixelShader.usf b/Engine/Shaders/Private/BasePassPixelShader.usf
--- a/Engine/Shaders/Private/BasePassPixelShader.usf
+++ b/Engine/Shaders/Private/BasePassPixelShader.usf
@@ -1507,9 +1507,21 @@ void FPixelShaderInOut_MainPS(
 				#endif
 
 				#if FORWARD_SHADING || TRANSLUCENCY_LIGHTING_SURFACE_FORWARDSHADING || MATERIAL_SHADINGMODEL_SINGLELAYERWATER
+
+					float3 DirectionalLightAtmosphereTransmittance = 1.0f;
+					const uint LightIndex = 0;
+					if (ResolvedView.AtmosphereLightDiscCosHalfApexAngle_PPTrans[LightIndex].y > 0.0f)
+					{
+						// Only when using forward shading, we can evaluate per pixel atmosphere transmittance.
+						const float3 PlanetCenterToTranslatedWorldPos = (LWCToFloat(LWCAdd(MaterialParameters.AbsoluteWorldPosition, ResolvedView.PreViewTranslation)) - ResolvedView.SkyPlanetTranslatedWorldCenterAndViewHeight.xyz) * CM_TO_SKY_UNIT;
+						DirectionalLightAtmosphereTransmittance = GetAtmosphereTransmittance(
+							PlanetCenterToTranslatedWorldPos, ResolvedView.AtmosphereLightDirection[LightIndex].xyz, ResolvedView.SkyAtmosphereBottomRadiusKm, ResolvedView.SkyAtmosphereTopRadiusKm,
+							View.TransmittanceLutTexture, View.TransmittanceLutTextureSampler);
+					}
+
 					FDeferredLightingSplit ForwardDirectLighting = GetForwardDirectLightingSplit(
 						GridIndex, MaterialParameters.WorldPosition_CamRelative, MaterialParameters.CameraVector, GBuffer, NearestResolvedDepthScreenUV, MaterialParameters.PrimitiveId, EyeIndex, Dither, 
-						DirectionalLightCloudShadow, DirectionalLightShadow,
+						DirectionalLightCloudShadow, DirectionalLightAtmosphereTransmittance, DirectionalLightShadow,
 						bSeparateWaterMainDirLightLuminance, SeparatedWaterMainDirLightLuminance);
 					
 					#if MATERIAL_SHADINGMODEL_THIN_TRANSLUCENT
diff --git a/Engine/Shaders/Private/ForwardLightingCommon.ush b/Engine/Shaders/Private/ForwardLightingCommon.ush
--- a/Engine/Shaders/Private/ForwardLightingCommon.ush
+++ b/Engine/Shaders/Private/ForwardLightingCommon.ush
@@ -120,7 +120,7 @@ float GetIndirectOcclusion(float2 ScreenUV, bool bHasDynamicIndirectShadowCaster
 
 FDeferredLightingSplit GetForwardDirectLightingSplit(
 	uint GridIndex, float3 TranslatedWorldPosition, float3 CameraVector, FGBufferData GBufferData, float2 ScreenUV, uint PrimitiveId, uint EyeIndex, float Dither, 
-	float InDirectionalLightCloudShadow, inout float OutDirectionalLightShadow,
+	float InDirectionalLightCloudShadow, float3 InDirectionalLightAtmosphereTransmittance, inout float OutDirectionalLightShadow,
 	bool bSeparateMainDirLightLuminance, inout float3 SeparatedMainDirLightLuminance)
 {
 	float3 WorldPosition = TranslatedWorldPosition - LWCHackToFloat(PrimaryView.PreViewTranslation);
@@ -213,6 +213,8 @@ FDeferredLightingSplit GetForwardDirectLightingSplit(
 			NewLighting.DiffuseLighting *= InDirectionalLightCloudShadow;
 			NewLighting.SpecularLighting *= InDirectionalLightCloudShadow;
 		#endif
+			NewLighting.DiffuseLighting.rgb *= InDirectionalLightAtmosphereTransmittance;
+			NewLighting.SpecularLighting.rgb *= InDirectionalLightAtmosphereTransmittance;
 
 			if (bSeparateMainDirLightLuminance)
 			{
@@ -306,18 +308,6 @@ FDeferredLightingSplit GetForwardDirectLightingSplit(
 	return DirectLighting;
 }
 
-float3 GetForwardDirectLighting(uint GridIndex, float3 TranslatedWorldPosition, float3 CameraVector, FGBufferData GBufferData, float2 ScreenUV, uint PrimitiveId, uint EyeIndex, float Dither, float InDirectionalLightCloudShadow, inout float OutDirectionalLightShadow)
-{
-	bool bSeparateMainDirLightLuminance = false;
-	float3 SeparatedMainDirLightLuminance = 0.0f;
-	FDeferredLightingSplit Ret = GetForwardDirectLightingSplit(
-		GridIndex, TranslatedWorldPosition, CameraVector, GBufferData, ScreenUV, PrimitiveId, EyeIndex, Dither, 
-		InDirectionalLightCloudShadow, OutDirectionalLightShadow,
-		bSeparateMainDirLightLuminance, SeparatedMainDirLightLuminance);
-
-	return Ret.DiffuseLighting.rgb + Ret.SpecularLighting.rgb;
-}
-
 float3 GetForwardDirectLightingForVertexLighting(uint GridIndex, float3 TranslatedWorldPosition, float SceneDepth, float3 WorldNormal, uint EyeIndex, float InDirectionalLightCloudShadow)
 {
 	float3 DirectLighting = 0;
diff --git a/Engine/Shaders/Private/MaterialTemplate.ush b/Engine/Shaders/Private/MaterialTemplate.ush
--- a/Engine/Shaders/Private/MaterialTemplate.ush
+++ b/Engine/Shaders/Private/MaterialTemplate.ush
@@ -1579,7 +1579,7 @@ float3 MaterialExpressionSkyAtmosphereLightDiskLuminance(FMaterialPixelParameter
 		// GetLightDiskLuminance does a test against the virtual planet but SkyCameraTranslatedWorldOrigin is always put safely setup above it (to never have the camera into the virtual planet with a black screen)
 		LightDiskLuminance =  GetLightDiskLuminance(PlanetCenterToCameraTranslatedWorld, ViewDir, ResolvedView.SkyAtmosphereBottomRadiusKm, ResolvedView.SkyAtmosphereTopRadiusKm,
 			View.TransmittanceLutTexture, View.TransmittanceLutTextureSampler,
-			ResolvedView.AtmosphereLightDirection[LightIndex].xyz, ResolvedView.AtmosphereLightDiscCosHalfApexAngle[LightIndex].x, ResolvedView.AtmosphereLightDiscLuminance[LightIndex].xyz);
+			ResolvedView.AtmosphereLightDirection[LightIndex].xyz, ResolvedView.AtmosphereLightDiscCosHalfApexAngle_PPTrans[LightIndex].x, ResolvedView.AtmosphereLightDiscLuminance[LightIndex].xyz);
 	}
 #endif
 	return LightDiskLuminance;
diff --git a/Engine/Shaders/Private/SkyAtmosphere.usf b/Engine/Shaders/Private/SkyAtmosphere.usf
--- a/Engine/Shaders/Private/SkyAtmosphere.usf
+++ b/Engine/Shaders/Private/SkyAtmosphere.usf
@@ -317,7 +317,7 @@ float3 GetLightDiskLuminance(float3 PlanetCenterToCamera, float3 WorldDir, uint
 		float3 LightDiskLuminance = GetLightDiskLuminance(
 			PlanetCenterToCamera, WorldDir, Atmosphere.BottomRadiusKm, Atmosphere.TopRadiusKm,
 			TransmittanceLutTexture, TransmittanceLutTextureSampler,
-			View.AtmosphereLightDirection[LightIndex].xyz, View.AtmosphereLightDiscCosHalfApexAngle[LightIndex].x, View.AtmosphereLightDiscLuminance[LightIndex].xyz);
+			View.AtmosphereLightDirection[LightIndex].xyz, View.AtmosphereLightDiscCosHalfApexAngle_PPTrans[LightIndex].x, View.AtmosphereLightDiscLuminance[LightIndex].xyz);
 
 		// Clamp to avoid crazy high values (and exposed 64000.0f luminance is already crazy high, solar system sun is 1.6x10^9). Also this removes +inf float and helps TAA.
 		const float3 MaxLightLuminance = 64000.0f;
diff --git a/Engine/Source/Runtime/Engine/Public/SceneView.h b/Engine/Source/Runtime/Engine/Public/SceneView.h
--- a/Engine/Source/Runtime/Engine/Public/SceneView.h
+++ b/Engine/Source/Runtime/Engine/Public/SceneView.h
@@ -704,7 +704,7 @@
 	VIEW_UNIFORM_BUFFER_MEMBER_ARRAY(FLinearColor, AtmosphereLightIlluminanceOnGroundPostTransmittance, [NUM_ATMOSPHERE_LIGHTS]) \
 	VIEW_UNIFORM_BUFFER_MEMBER_ARRAY(FLinearColor, AtmosphereLightIlluminanceOuterSpace, [NUM_ATMOSPHERE_LIGHTS]) \
 	VIEW_UNIFORM_BUFFER_MEMBER_ARRAY(FLinearColor, AtmosphereLightDiscLuminance, [NUM_ATMOSPHERE_LIGHTS]) \
-	VIEW_UNIFORM_BUFFER_MEMBER_ARRAY(FVector4f, AtmosphereLightDiscCosHalfApexAngle, [NUM_ATMOSPHERE_LIGHTS]) \
+	VIEW_UNIFORM_BUFFER_MEMBER_ARRAY(FVector4f, AtmosphereLightDiscCosHalfApexAngle_PPTrans, [NUM_ATMOSPHERE_LIGHTS]) \
 	VIEW_UNIFORM_BUFFER_MEMBER(FVector4f, SkyViewLutSizeAndInvSize) \
 	VIEW_UNIFORM_BUFFER_MEMBER(FVector3f, SkyCameraTranslatedWorldOrigin) \
 	VIEW_UNIFORM_BUFFER_MEMBER(FVector4f, SkyPlanetTranslatedWorldCenterAndViewHeight) \
diff --git a/Engine/Source/Runtime/Renderer/Private/SceneRendering.cpp b/Engine/Source/Runtime/Renderer/Private/SceneRendering.cpp
--- a/Engine/Source/Runtime/Renderer/Private/SceneRendering.cpp
+++ b/Engine/Source/Runtime/Renderer/Private/SceneRendering.cpp
@@ -1296,7 +1296,7 @@
 	auto ClearAtmosphereLightData = [&](uint32 Index)
 	{
 		check(Index < NUM_ATMOSPHERE_LIGHTS);
-		ViewUniformShaderParameters.AtmosphereLightDiscCosHalfApexAngle[Index] = FVector4f(1.0f);
+		ViewUniformShaderParameters.AtmosphereLightDiscCosHalfApexAngle_PPTrans[Index] = FVector4f(1.0f, 0.0f, 0.0f, 0.0f);
 		ViewUniformShaderParameters.AtmosphereLightDiscLuminance[Index] = FLinearColor::Black;
 		ViewUniformShaderParameters.AtmosphereLightIlluminanceOnGroundPostTransmittance[Index] = FLinearColor::Black;
 		ViewUniformShaderParameters.AtmosphereLightIlluminanceOnGroundPostTransmittance[Index].A = 0.0f;
@@ -1322,8 +1322,9 @@
 		// Set default atmosphere lights parameters
 		FLightSceneInfo* SunLight = Scene->AtmosphereLights[0];	// Atmospheric fog only takes into account the a single sun light with index 0.
 		const float SunLightDiskHalfApexAngleRadian = SunLight ? SunLight->Proxy->GetSunLightHalfApexAngleRadian() : FLightSceneProxy::GetSunOnEarthHalfApexAngleRadian();
+		const float UsePerPixelAtmosphereTransmittance = 0.0f; // The default sun light should not use per pixel transmitance without an atmosphere.
 
-		ViewUniformShaderParameters.AtmosphereLightDiscCosHalfApexAngle[0] = FVector4f(FMath::Cos(SunLightDiskHalfApexAngleRadian));
+		ViewUniformShaderParameters.AtmosphereLightDiscCosHalfApexAngle_PPTrans[0] = FVector4f(FMath::Cos(SunLightDiskHalfApexAngleRadian), UsePerPixelAtmosphereTransmittance, 0.0f, 0.0f);
 		//Added check so atmospheric light color and vector can use a directional light without needing an atmospheric fog actor in the scene
 		ViewUniformShaderParameters.AtmosphereLightDiscLuminance[0] = SunLight ? SunLight->Proxy->GetOuterSpaceLuminance() : FLinearColor::Black;
 		ViewUniformShaderParameters.AtmosphereLightIlluminanceOnGroundPostTransmittance[0] = SunLight ? SunLight->Proxy->GetColor() : FLinearColor::Black;
@@ -1410,7 +1411,8 @@
 			FLightSceneInfo* Light = Scene->AtmosphereLights[Index];
 			if (Light)
 			{
-				ViewUniformShaderParameters.AtmosphereLightDiscCosHalfApexAngle[Index] = FVector4f(FMath::Cos(Light->Proxy->GetSunLightHalfApexAngleRadian()));
+				const float UsePerPixelAtmosphereTransmittance = Light->Proxy->GetUsePerPixelAtmosphereTransmittance() ? 1.0f : 0.0f;
+				ViewUniformShaderParameters.AtmosphereLightDiscCosHalfApexAngle_PPTrans[Index] = FVector4f(FMath::Cos(Light->Proxy->GetSunLightHalfApexAngleRadian()), UsePerPixelAtmosphereTransmittance, 0.0f, 0.0f);
 				ViewUniformShaderParameters.AtmosphereLightDiscLuminance[Index] = Light->Proxy->GetOuterSpaceLuminance();
 				ViewUniformShaderParameters.AtmosphereLightIlluminanceOnGroundPostTransmittance[Index] = Light->Proxy->GetSunIlluminanceOnGroundPostTransmittance();
 				ViewUniformShaderParameters.AtmosphereLightIlluminanceOnGroundPostTransmittance[Index].A = 1.0f; // interactions with HeightFogComponent
@@ -1464,7 +1466,7 @@
 				FLightSceneInfo* Light = Scene->AtmosphereLights[Index];
 				if (Light)
 				{
-					ViewUniformShaderParameters.AtmosphereLightDiscCosHalfApexAngle[Index] = FVector4f(1.0f);
+					ViewUniformShaderParameters.AtmosphereLightDiscCosHalfApexAngle_PPTrans[Index] = FVector4f(1.0f, 0.0f, 0.0f, 0.0f);
 					ViewUniformShaderParameters.AtmosphereLightDiscLuminance[Index] = FLinearColor::Black;
 					ViewUniformShaderParameters.AtmosphereLightIlluminanceOnGroundPostTransmittance[Index] = Light->Proxy->GetColor();
 					ViewUniformShaderParameters.AtmosphereLightIlluminanceOnGroundPostTransmittance[Index].A = 0.0f; // no interactions with HeightFogComponent
