Fix for crash that would occur when setting the Preview Render Level to SM5 after successfully rendering SM6.

diff --git a/Engine/Source/Runtime/Renderer/Private/SceneHitProxyRendering.cpp b/Engine/Source/Runtime/Renderer/Private/SceneHitProxyRendering.cpp
--- a/Engine/Source/Runtime/Renderer/Private/SceneHitProxyRendering.cpp
+++ b/Engine/Source/Runtime/Renderer/Private/SceneHitProxyRendering.cpp
@@ -604,7 +604,7 @@ void FMobileSceneRenderer::RenderHitProxies(FRDGBuilder& GraphBuilder)
 
 void FDeferredShadingSceneRenderer::RenderHitProxies(FRDGBuilder& GraphBuilder)
 {
-	static const bool bNaniteEnabled = UseNanite(ShaderPlatform);
+	const bool bNaniteEnabled = UseNanite(ShaderPlatform);
 
 	Scene->UpdateAllPrimitiveSceneInfos(GraphBuilder);
 
