Fix crash in CitySample when groom component are being destroyed.

diff --git a/Engine/Source/Runtime/Renderer/Private/SceneVisibility.cpp b/Engine/Source/Runtime/Renderer/Private/SceneVisibility.cpp
--- a/Engine/Source/Runtime/Renderer/Private/SceneVisibility.cpp
+++ b/Engine/Source/Runtime/Renderer/Private/SceneVisibility.cpp
@@ -3346,7 +3346,7 @@
 			BatchAndProxy = MeshBatch;
 		}
 
-		if (HairStrands::IsHairCardsVF(MeshBatch.Mesh))
+		if (HairStrands::IsHairCardsVF(MeshBatch.Mesh) && ViewRelevance.bRenderInMainPass)
 		{
 			View.HairCardsMeshElements.AddUninitialized(1);
 			FMeshBatchAndRelevance& BatchAndProxy = View.HairCardsMeshElements.Last();
