Fix for water shore aliasing with lower resolution refraction.
Changed point sampling of linear depth buffer to manual bilinear interpolation on Metal and to hardware bilinear sampling on other platforms.

#rnx

diff --git a/Engine/Shaders/Private/BasePassPixelShader.usf b/Engine/Shaders/Private/BasePassPixelShader.usf
--- a/Engine/Shaders/Private/BasePassPixelShader.usf
+++ b/Engine/Shaders/Private/BasePassPixelShader.usf
@@ -1627,6 +1627,8 @@ void FPixelShaderInOut_MainPS(
 			MaterialParameters, PixelMaterialInputs, ResolvedView,
 			DirectionalLightShadow,
 			OpaqueBasePass.SceneDepthWithoutSingleLayerWaterTexture, SingleLayerWaterSceneDepthSampler, // Scene depth texture
+			OpaqueBasePass.SceneWithoutSingleLayerWaterTextureSize,
+			OpaqueBasePass.SceneWithoutSingleLayerWaterInvTextureSize,
 			Specular, NullDistortionParams,
 			SunIlluminance, WaterDiffuseIndirectIlluminance, EnvBrdf,
 			CameraIsUnderWater, WaterVisibility, EyeIndex,
@@ -1644,6 +1646,8 @@ void FPixelShaderInOut_MainPS(
 			OpaqueBasePass.SceneDepthWithoutSingleLayerWaterTexture, SingleLayerWaterSceneDepthSampler,
 			OpaqueBasePass.SceneWithoutSingleLayerWaterMinMaxUV.xy,
 			OpaqueBasePass.SceneWithoutSingleLayerWaterMinMaxUV.zw,
+			OpaqueBasePass.SceneWithoutSingleLayerWaterTextureSize,
+			OpaqueBasePass.SceneWithoutSingleLayerWaterInvTextureSize,
 			Specular, OpaqueBasePass.DistortionParams,
 			SunIlluminance, WaterDiffuseIndirectIlluminance, EnvBrdf,
 			CameraIsUnderWater, WaterVisibility, EyeIndex,
diff --git a/Engine/Shaders/Private/SingleLayerWaterCommon.ush b/Engine/Shaders/Private/SingleLayerWaterCommon.ush
--- a/Engine/Shaders/Private/SingleLayerWaterCommon.ush
+++ b/Engine/Shaders/Private/SingleLayerWaterCommon.ush
@@ -1,8 +1,10 @@
 // Copyright Epic Games, Inc. All Rights Reserved.
 
 #pragma once
+#include "TextureSampling.ush" // GetBilinearSampleLevelInfos() and GetBilinearInterpolation()
 
 // Scale for encoding depth into SceneDepthWithoutSingleLayerWaterTexture
+// Must match FDeferredShadingSceneRenderer::RenderUnderWaterFog
 #define SINGLE_LAYER_WATER_DEPTH_SCALE 100.0f
 
 #ifndef SIMPLE_SINGLE_LAYER_WATER
@@ -11,7 +13,11 @@
 
 #if SUPPORTS_INDEPENDENT_SAMPLERS
 	#define SingleLayerWaterSceneColorSampler View.SharedBilinearClampedSampler
+#if METAL_PROFILE
 	#define SingleLayerWaterSceneDepthSampler View.SharedPointClampedSampler
+#else // TODO: Dynamically use a point sampler instead if the depth texture was not downsampled.
+	#define SingleLayerWaterSceneDepthSampler View.SharedBilinearClampedSampler
+#endif // METAL_PROFILE
 #else
 	#define SingleLayerWaterSceneColorSampler OpaqueBasePass.SceneColorWithoutSingleLayerWaterSampler
 	#define SingleLayerWaterSceneDepthSampler OpaqueBasePass.SceneDepthWithoutSingleLayerWaterSampler
@@ -33,3 +39,16 @@ bool WaterRefract(float3 V, float3 N, float Eta, out float3 R)
 	R = (Eta * CosThetaI + CosThetaT) * N - Eta * V;
 	return true;
 }
+
+float WaterSampleSceneDepthWithoutWater(Texture2D SceneDepthWithoutWaterTexture, SamplerState SceneDepthWithoutWaterSampler, float2 UV, float2 SceneDepthWithoutWaterTextureSize, float2 SceneDepthWithoutWaterTextureInvSize)
+{
+#if METAL_PROFILE
+	// On Metal we cannot do hardware bilinear sampling of float32 textures, so we need to do it manually.
+	FBilinearSampleInfos SampleInfos = GetBilinearSampleLevelInfos(UV, SceneDepthWithoutWaterTextureSize, SceneDepthWithoutWaterTextureInvSize);
+	float4 PixelSceneDepth4 = SceneDepthWithoutWaterTexture.Gather(SceneDepthWithoutWaterSampler, UV).wzxy;
+	return GetBilinearInterpolation(SampleInfos, PixelSceneDepth4);
+#else
+	// A linear sampler should be passed in here if downsampling was used.
+	return SceneDepthWithoutWaterTexture.SampleLevel(SceneDepthWithoutWaterSampler, UV, 0).x;
+#endif // METAL_PROFILE
+}
\ No newline at end of file
diff --git a/Engine/Shaders/Private/SingleLayerWaterComposite.usf b/Engine/Shaders/Private/SingleLayerWaterComposite.usf
--- a/Engine/Shaders/Private/SingleLayerWaterComposite.usf
+++ b/Engine/Shaders/Private/SingleLayerWaterComposite.usf
@@ -11,10 +11,7 @@
 #include "ShadingModels.ush"
 #include "SceneTextureParameters.ush"
 #include "LightGridCommon.ush"
-
-// Scale for encoding depth into SceneDepthWithoutSingleLayerWaterTexture
-// Must match FDeferredShadingSceneRenderer::RenderUnderWaterFog
-#define SINGLE_LAYER_WATER_DEPTH_SCALE 100.0f
+#include "SingleLayerWaterCommon.ush"
 
 #define ENABLE_SKY_LIGHT 1
 
@@ -32,6 +29,8 @@ SamplerState SceneNoWaterDepthSampler;
 Texture2D<float4> SeparatedMainDirLightTexture;
 
 float4 SceneNoWaterMinMaxUV;
+float2 SceneNoWaterTextureSize;
+float2 SceneNoWaterInvTextureSize;
 float UseSeparatedMainDirLightTexture;
 
 struct SingleLayerWaterCompositeOutput
@@ -68,7 +67,14 @@ SingleLayerWaterCompositeOutput SingleLayerWaterComposite(float2 BufferUV, float
 
 		const float2 SceneNoWaterUV = clamp(BufferUV, SceneNoWaterMinMaxUV.xy, SceneNoWaterMinMaxUV.zw);
 
-		float OpaqueDepth = SceneNoWaterDepthTexture.SampleLevel(SceneNoWaterDepthSampler, SceneNoWaterUV, 0).x * SINGLE_LAYER_WATER_DEPTH_SCALE;
+		float OpaqueDepth = WaterSampleSceneDepthWithoutWater(
+			SceneNoWaterDepthTexture,
+			SceneNoWaterDepthSampler,
+			SceneNoWaterUV,
+			SceneNoWaterTextureSize,
+			SceneNoWaterInvTextureSize)
+			* SINGLE_LAYER_WATER_DEPTH_SCALE;
+
 		float WaterDepth = SceneDepth;
 
 		float DeltaDepth = CameraIsUnderWater ? WaterDepth : OpaqueDepth - WaterDepth;	// Inverted depth
@@ -150,7 +156,14 @@ SingleLayerWaterCompositeOutput SingleLayerWaterComposite(float2 BufferUV, float
 	{
 		const float2 SceneNoWaterUV = clamp(BufferUV, SceneNoWaterMinMaxUV.xy, SceneNoWaterMinMaxUV.zw);
 
-		float OpaqueDepth = SceneNoWaterDepthTexture.SampleLevel(SceneNoWaterDepthSampler, SceneNoWaterUV, 0).x * SINGLE_LAYER_WATER_DEPTH_SCALE;
+		float OpaqueDepth = WaterSampleSceneDepthWithoutWater(
+			SceneNoWaterDepthTexture,
+			SceneNoWaterDepthSampler,
+			SceneNoWaterUV,
+			SceneNoWaterTextureSize,
+			SceneNoWaterInvTextureSize)
+			* SINGLE_LAYER_WATER_DEPTH_SCALE;
+
 		float WaterDepth = GBuffer.Depth;
 
 		float DeltaDepth = CameraIsUnderWater ? WaterDepth : OpaqueDepth - WaterDepth;	// Inverted depth
diff --git a/Engine/Shaders/Private/SingleLayerWaterShading.ush b/Engine/Shaders/Private/SingleLayerWaterShading.ush
--- a/Engine/Shaders/Private/SingleLayerWaterShading.ush
+++ b/Engine/Shaders/Private/SingleLayerWaterShading.ush
@@ -27,12 +27,16 @@ WaterVolumeLightingOutput EvaluateWaterVolumeLighting(
 	float DirectionalLightShadow,
 #if SINGLE_LAYER_WATER_SIMPLE_FORWARD
 	Texture2D SceneDepthWithoutSingleLayerWaterTexture, SamplerState SceneDepthWithoutSingleLayerWaterSampler,
+	float2 SceneDepthWithoutSingleLayerWaterTextureSize,
+	float2 SceneDepthWithoutSingleLayerWaterTextureInvSize,
 #endif
 #if !SIMPLE_SINGLE_LAYER_WATER
 	Texture2D SceneColorWithoutSingleLayerWaterTexture, SamplerState SceneColorWithoutSingleLayerWaterSampler,
 	Texture2D SceneDepthWithoutSingleLayerWaterTexture, SamplerState SceneDepthWithoutSingleLayerWaterSampler,
 	float2 SceneWithoutSingleLayerWaterMinUV,
 	float2 SceneWithoutSingleLayerWaterMaxUV,
+	float2 SceneDepthWithoutSingleLayerWaterTextureSize,
+	float2 SceneDepthWithoutSingleLayerWaterTextureInvSize,
 #endif // SIMPLE_SINGLE_LAYER_WATER
 	float Specular, const float4 DistortionParams,
 	float3 SunIlluminance, float3 AmbiantIlluminance, float3 EnvBrdf,
@@ -52,7 +56,15 @@ WaterVolumeLightingOutput EvaluateWaterVolumeLighting(
 #if SIMPLE_SINGLE_LAYER_WATER
 
 #if SINGLE_LAYER_WATER_SIMPLE_FORWARD 
-		float PixelSceneDepth = SceneDepthWithoutSingleLayerWaterTexture.SampleLevel(SceneDepthWithoutSingleLayerWaterSampler, ViewportUV, 0).x * SINGLE_LAYER_WATER_DEPTH_SCALE;
+
+		float PixelSceneDepth = WaterSampleSceneDepthWithoutWater(
+			SceneDepthWithoutSingleLayerWaterTexture,
+			SceneDepthWithoutSingleLayerWaterSampler,
+			ViewportUV,
+			SceneDepthWithoutSingleLayerWaterTextureSize,
+			SceneDepthWithoutSingleLayerWaterTextureInvSize)
+			* SINGLE_LAYER_WATER_DEPTH_SCALE;
+
 		float PixelSceneDeviceZ = ConvertToDeviceZ(PixelSceneDepth);
 #else
 		float PixelSceneDeviceZ = LookupDeviceZ(ViewportUV);
@@ -69,7 +81,13 @@ WaterVolumeLightingOutput EvaluateWaterVolumeLighting(
 
 		ViewportUV = clamp(ViewportUV, SceneWithoutSingleLayerWaterMinUV, SceneWithoutSingleLayerWaterMaxUV);
 
-		float PixelSceneDepth = SceneDepthWithoutSingleLayerWaterTexture.SampleLevel(SceneDepthWithoutSingleLayerWaterSampler, ViewportUV, 0).x * SINGLE_LAYER_WATER_DEPTH_SCALE;
+		float PixelSceneDepth = WaterSampleSceneDepthWithoutWater(
+			SceneDepthWithoutSingleLayerWaterTexture, 
+			SceneDepthWithoutSingleLayerWaterSampler, 
+			ViewportUV, 
+			SceneDepthWithoutSingleLayerWaterTextureSize, 
+			SceneDepthWithoutSingleLayerWaterTextureInvSize) 
+			* SINGLE_LAYER_WATER_DEPTH_SCALE;
 		
 		float2 Refraction = GetMaterialRefraction(PixelMaterialInputs); // X = IOR, Y = Refraction Bias
 		float2 BufferUVDistortion = ComputeBufferUVDistortion(
@@ -93,8 +111,8 @@ WaterVolumeLightingOutput EvaluateWaterVolumeLighting(
 		}
 		else
 		{
-			// Pick the depth furthest away to make sure we pick a depth sample under water
-			SceneDepth = max(max(SceneDepth4.x, SceneDepth4.y), max(SceneDepth4.z, SceneDepth4.w));
+			// We can use bilinear interpolation here because all depth samples are guaranteed to be farther away than the water surface depth.
+			SceneDepth = GetBilinearInterpolation(GetBilinearSampleLevelInfos(DistortedUV, SceneDepthWithoutSingleLayerWaterTextureSize, SceneDepthWithoutSingleLayerWaterTextureInvSize), SceneDepth4.wzxy);
 			ViewportUV = DistortedUV;
 		}
 
diff --git a/Engine/Shaders/Private/TextureSampling.ush b/Engine/Shaders/Private/TextureSampling.ush
--- a/Engine/Shaders/Private/TextureSampling.ush
+++ b/Engine/Shaders/Private/TextureSampling.ush
@@ -85,6 +85,18 @@ uint2 GetSamplePixelCoord(FBilinearSampleInfos SampleInfos, const uint SampleId)
 	return uint2(int2(SampleInfos.TopLeftPixelCoord)) + BilinearSamplingOffsets2x2[SampleId];
 }
 
+/** Bilinearly interpolates the four given samples. */
+float GetBilinearInterpolation(FBilinearSampleInfos SampleInfos, float4 Samples)
+{
+	float4 BilinearSampleWeights;
+	BilinearSampleWeights[0] = (1 - SampleInfos.BilinearInterp.x) * (1 - SampleInfos.BilinearInterp.y);
+	BilinearSampleWeights[1] = (SampleInfos.BilinearInterp.x) * (1 - SampleInfos.BilinearInterp.y);
+	BilinearSampleWeights[2] = (1 - SampleInfos.BilinearInterp.x) * (SampleInfos.BilinearInterp.y);
+	BilinearSampleWeights[3] = (SampleInfos.BilinearInterp.x) * (SampleInfos.BilinearInterp.y);
+
+	return dot(Samples, BilinearSampleWeights);
+}
+
 
 void Bicubic2DCatmullRom(in float2 UV, in float2 Size, in float2 InvSize, out float2 Sample[3], out half2 Weight[3])
 {
diff --git a/Engine/Source/Runtime/Renderer/Private/BasePassRendering.cpp b/Engine/Source/Runtime/Renderer/Private/BasePassRendering.cpp
--- a/Engine/Source/Runtime/Renderer/Private/BasePassRendering.cpp
+++ b/Engine/Source/Runtime/Renderer/Private/BasePassRendering.cpp
@@ -626,15 +626,20 @@ TRDGUniformBufferRef<FOpaqueBasePassUniformParameters> CreateOpaqueBasePassUnifo
 	BasePassParameters.DBuffer = GetDBufferParameters(GraphBuilder, DBufferTextures, View.GetShaderPlatform());
 
 	// Single Layer Water
+	const bool bShouldUseBilinearSamplerForDepth = SceneWithoutWaterTextures && ShouldUseBilinearSamplerForDepthWithoutSingleLayerWater(SceneWithoutWaterTextures->DepthTexture->Desc.Format);
 	BasePassParameters.SceneWithoutSingleLayerWaterMinMaxUV = FVector4f(0.0f, 0.0f, 1.0f, 1.0f);
 	BasePassParameters.SceneColorWithoutSingleLayerWaterTexture = SystemTextures.Black;
 	BasePassParameters.SceneDepthWithoutSingleLayerWaterTexture = SystemTextures.Black;
 	BasePassParameters.SceneColorWithoutSingleLayerWaterSampler = TStaticSamplerState<SF_Bilinear>::GetRHI();
-	BasePassParameters.SceneDepthWithoutSingleLayerWaterSampler = TStaticSamplerState<SF_Point>::GetRHI();
+	BasePassParameters.SceneDepthWithoutSingleLayerWaterSampler = bShouldUseBilinearSamplerForDepth ? TStaticSamplerState<SF_Bilinear>::GetRHI() : TStaticSamplerState<SF_Point>::GetRHI();
 
 	if (SceneWithoutWaterTextures)
 	{
+		FIntVector DepthTextureSize = SceneWithoutWaterTextures->DepthTexture->Desc.GetSize();
+		
 		BasePassParameters.SceneWithoutSingleLayerWaterMinMaxUV = SceneWithoutWaterTextures->Views[ViewIndex].MinMaxUV;
+		BasePassParameters.SceneWithoutSingleLayerWaterTextureSize = FVector2f(DepthTextureSize.X, DepthTextureSize.Y);
+		BasePassParameters.SceneWithoutSingleLayerWaterInvTextureSize = FVector2f(1.0f / DepthTextureSize.X, 1.0f / DepthTextureSize.Y);
 		BasePassParameters.SceneColorWithoutSingleLayerWaterTexture = SceneWithoutWaterTextures->ColorTexture;
 		BasePassParameters.SceneDepthWithoutSingleLayerWaterTexture = SceneWithoutWaterTextures->DepthTexture;
 	}
diff --git a/Engine/Source/Runtime/Renderer/Private/BasePassRendering.h b/Engine/Source/Runtime/Renderer/Private/BasePassRendering.h
--- a/Engine/Source/Runtime/Renderer/Private/BasePassRendering.h
+++ b/Engine/Source/Runtime/Renderer/Private/BasePassRendering.h
@@ -92,6 +92,8 @@ BEGIN_GLOBAL_SHADER_PARAMETER_STRUCT(FOpaqueBasePassUniformParameters,)
 	SHADER_PARAMETER_SAMPLER(SamplerState, PreIntegratedGFSampler)
 	SHADER_PARAMETER(FVector4f, SceneWithoutSingleLayerWaterMinMaxUV)
 	SHADER_PARAMETER(FVector4f, DistortionParams)
+	SHADER_PARAMETER(FVector2f, SceneWithoutSingleLayerWaterTextureSize)
+	SHADER_PARAMETER(FVector2f, SceneWithoutSingleLayerWaterInvTextureSize)
 	SHADER_PARAMETER(int32, Is24BitUnormDepthStencil)
 	// Misc
 	SHADER_PARAMETER_RDG_TEXTURE(Texture2D, EyeAdaptationTexture)
diff --git a/Engine/Source/Runtime/Renderer/Private/SingleLayerWaterRendering.cpp b/Engine/Source/Runtime/Renderer/Private/SingleLayerWaterRendering.cpp
--- a/Engine/Source/Runtime/Renderer/Private/SingleLayerWaterRendering.cpp
+++ b/Engine/Source/Runtime/Renderer/Private/SingleLayerWaterRendering.cpp
@@ -125,6 +125,15 @@ bool ShouldRenderSingleLayerWaterSkippedRenderEditorNotification(TArrayView<cons
 	return false;
 }
 
+bool ShouldUseBilinearSamplerForDepthWithoutSingleLayerWater(EPixelFormat DepthTextureFormat)
+{
+	const bool bHasDownsampling = GSingleLayerWaterRefractionDownsampleFactor > 1;
+	const bool bSupportsLinearSampling = !!(GPixelFormats[DepthTextureFormat].Capabilities & EPixelFormatCapabilities::TextureSample);
+	
+	// Linear sampling is only required if the depth texture has been downsampled.
+	return bHasDownsampling && bSupportsLinearSampling;
+}
+
 bool UseSingleLayerWaterIndirectDraw(EShaderPlatform ShaderPlatform)
 {
 	return IsFeatureLevelSupported(ShaderPlatform, ERHIFeatureLevel::SM5)
@@ -147,6 +156,8 @@ BEGIN_SHADER_PARAMETER_STRUCT(FSingleLayerWaterCommonShaderParameters, )
 	SHADER_PARAMETER_SAMPLER(SamplerState, SceneNoWaterDepthSampler)
 	SHADER_PARAMETER_RDG_TEXTURE(Texture2D, SeparatedMainDirLightTexture)
 	SHADER_PARAMETER(FVector4f, SceneNoWaterMinMaxUV)
+	SHADER_PARAMETER(FVector2f, SceneNoWaterTextureSize)
+	SHADER_PARAMETER(FVector2f, SceneNoWaterInvTextureSize)
 	SHADER_PARAMETER(float, UseSeparatedMainDirLightTexture)
 	SHADER_PARAMETER_STRUCT_INCLUDE(FSceneTextureParameters, SceneTextures)	// Water scene texture
 	SHADER_PARAMETER_STRUCT_REF(FViewUniformShaderParameters, ViewUniformBuffer)
@@ -426,13 +437,18 @@ void FDeferredShadingSceneRenderer::RenderSingleLayerWaterReflections(
 
 		auto SetCommonParameters = [&](FSingleLayerWaterCommonShaderParameters& Parameters)
 		{
+			FIntVector DepthTextureSize = SceneWithoutWaterTextures.DepthTexture ? SceneWithoutWaterTextures.DepthTexture->Desc.GetSize() : FIntVector::ZeroValue;
+			const bool bShouldUseBilinearSamplerForDepth = SceneWithoutWaterTextures.DepthTexture && ShouldUseBilinearSamplerForDepthWithoutSingleLayerWater(SceneWithoutWaterTextures.DepthTexture->Desc.Format);
+
 			Parameters.ScreenSpaceReflectionsTexture = ReflectionsColor ? ReflectionsColor : BlackDummyTexture;
 			Parameters.ScreenSpaceReflectionsSampler = TStaticSamplerState<SF_Point>::GetRHI();
 			Parameters.PreIntegratedGF = GSystemTextures.PreintegratedGF->GetRHI();
 			Parameters.PreIntegratedGFSampler = TStaticSamplerState<SF_Bilinear, AM_Clamp, AM_Clamp, AM_Clamp>::GetRHI();
 			Parameters.SceneNoWaterDepthTexture = SceneWithoutWaterTextures.DepthTexture ? SceneWithoutWaterTextures.DepthTexture : BlackDummyTexture;
-			Parameters.SceneNoWaterDepthSampler = TStaticSamplerState<SF_Point>::GetRHI();
+			Parameters.SceneNoWaterDepthSampler = bShouldUseBilinearSamplerForDepth ? TStaticSamplerState<SF_Bilinear>::GetRHI() : TStaticSamplerState<SF_Point>::GetRHI();
 			Parameters.SceneNoWaterMinMaxUV = SceneWithoutWaterTextures.Views[ViewIndex].MinMaxUV;
+			Parameters.SceneNoWaterTextureSize = SceneWithoutWaterTextures.DepthTexture ? FVector2f(DepthTextureSize.X, DepthTextureSize.Y) : FVector2f();
+			Parameters.SceneNoWaterInvTextureSize = SceneWithoutWaterTextures.DepthTexture ? FVector2f(1.0f / DepthTextureSize.X, 1.0f / DepthTextureSize.Y) : FVector2f();
 			Parameters.SeparatedMainDirLightTexture = BlackDummyTexture;
 			Parameters.UseSeparatedMainDirLightTexture = 0.0f;
 			Parameters.SceneTextures = SceneTextureParameters;
diff --git a/Engine/Source/Runtime/Renderer/Private/SingleLayerWaterRendering.h b/Engine/Source/Runtime/Renderer/Private/SingleLayerWaterRendering.h
--- a/Engine/Source/Runtime/Renderer/Private/SingleLayerWaterRendering.h
+++ b/Engine/Source/Runtime/Renderer/Private/SingleLayerWaterRendering.h
@@ -25,6 +25,7 @@ struct FSceneWithoutWaterTextures
 
 bool ShouldRenderSingleLayerWater(TArrayView<const FViewInfo> Views);
 bool ShouldRenderSingleLayerWaterSkippedRenderEditorNotification(TArrayView<const FViewInfo> Views);
+bool ShouldUseBilinearSamplerForDepthWithoutSingleLayerWater(EPixelFormat DepthTextureFormat);
 
 class FWaterTileVS : public FGlobalShader
 {
