Fix a potential crash where deleted primitives could potentially be put into the raytracing mesh command cache even though their resources were already released

This change only refreshes the command cache after having updated all primitive scene infos

diff --git a/Engine/Source/Runtime/Renderer/Private/DeferredShadingRenderer.cpp b/Engine/Source/Runtime/Renderer/Private/DeferredShadingRenderer.cpp
--- a/Engine/Source/Runtime/Renderer/Private/DeferredShadingRenderer.cpp
+++ b/Engine/Source/Runtime/Renderer/Private/DeferredShadingRenderer.cpp
@@ -1981,7 +1981,27 @@ void FDeferredShadingSceneRenderer::Render(FRDGBuilder& GraphBuilder)
 
 	GPU_MESSAGE_SCOPE(GraphBuilder);
 
+	for (int32 ViewIndex = 0; ViewIndex < Views.Num(); ViewIndex++)
+	{
+		FViewInfo& View = Views[ViewIndex];
+		RDG_GPU_MASK_SCOPE(GraphBuilder, View.GPUMask);
+
+		ShaderPrint::BeginView(GraphBuilder, View);
+		ShadingEnergyConservation::Init(GraphBuilder, View);
+	}
+	
+	ON_SCOPE_EXIT
+	{
+		for (FViewInfo& View : Views)
+		{
+			ShaderPrint::EndView(View);
+		}
+	};
+
+	Scene->UpdateAllPrimitiveSceneInfos(GraphBuilder, true);
+
 #if RHI_RAYTRACING
+	// Now that we have updated all the PrimitiveSceneInfos, update the RayTracing mesh commands cache if needed
 	{
 		ERayTracingMeshCommandsMode CurrentMode = ActiveViewFamily->EngineShowFlags.PathTracing ? ERayTracingMeshCommandsMode::PATH_TRACING : ERayTracingMeshCommandsMode::RAY_TRACING;
 		bool bNaniteCoarseMeshStreamingModeChanged = false;
@@ -2001,24 +2021,6 @@ void FDeferredShadingSceneRenderer::Render(FRDGBuilder& GraphBuilder)
 		Scene->UpdateRayTracedLights();
 	}
 #endif
-	for (int32 ViewIndex = 0; ViewIndex < Views.Num(); ViewIndex++)
-	{
-		FViewInfo& View = Views[ViewIndex];
-		RDG_GPU_MASK_SCOPE(GraphBuilder, View.GPUMask);
-
-		ShaderPrint::BeginView(GraphBuilder, View);
-		ShadingEnergyConservation::Init(GraphBuilder, View);
-	}
-	
-	ON_SCOPE_EXIT
-	{
-		for (FViewInfo& View : Views)
-		{
-			ShaderPrint::EndView(View);
-		}
-	};
-
-	Scene->UpdateAllPrimitiveSceneInfos(GraphBuilder, true);
 
 	FGPUSceneScopeBeginEndHelper GPUSceneScopeBeginEndHelper(Scene->GPUScene, GPUSceneDynamicContext, Scene);
 
