ShaderPrint: Support external clients without a view.
Clients can create, bind and submit their own shader print data.
Also supports freezing and unfreezing the data so that it can be collected over multiple frames.

diff --git a/Engine/Source/Runtime/Renderer/Private/ShaderPrint.cpp b/Engine/Source/Runtime/Renderer/Private/ShaderPrint.cpp
--- a/Engine/Source/Runtime/Renderer/Private/ShaderPrint.cpp
+++ b/Engine/Source/Runtime/Renderer/Private/ShaderPrint.cpp
@@ -83,31 +83,32 @@ namespace ShaderPrint
 	static uint32 GTriangleRequestCount = 0;
 	static bool GShaderPrintEnableOverride = false;
 	static FViewInfo* GDefaultView = nullptr;
+	static TArray<FFrozenShaderPrintData> GShaderPrintDataToRender;
 
 	//////////////////////////////////////////////////////////////////////////////////////////////////
 	// Struct & Functions
 
 	static uint32 GetMaxValueCount()
 	{
-		uint32 MaxValueCount = FMath::Max(CVarMaxCharacterCount.GetValueOnRenderThread() + int32(GCharacterRequestCount), 0);
+		uint32 MaxValueCount = FMath::Max(CVarMaxCharacterCount.GetValueOnAnyThread() + int32(GCharacterRequestCount), 0);
 		return IsEnabled() ? MaxValueCount : 0;
 	}
 
 	static uint32 GetMaxWidgetCount()
 	{
-		uint32 MaxValueCount = FMath::Max(CVarMaxWidgetCount.GetValueOnRenderThread() + int32(GWidgetRequestCount), 0);
+		uint32 MaxValueCount = FMath::Max(CVarMaxWidgetCount.GetValueOnAnyThread() + int32(GWidgetRequestCount), 0);
 		return IsEnabled() ? MaxValueCount : 0;
 	}
 
 	static uint32 GetMaxLineCount()
 	{
-		uint32 MaxValueCount = FMath::Max(CVarMaxLineCount.GetValueOnRenderThread() + int32(GLineRequestCount), 0);
+		uint32 MaxValueCount = FMath::Max(CVarMaxLineCount.GetValueOnAnyThread() + int32(GLineRequestCount), 0);
 		return IsEnabled() ? MaxValueCount : 0;
 	}
 
 	static uint32 GetMaxTriangleCount()
 	{
-		uint32 MaxValueCount = FMath::Max(CVarMaxTriangleCount.GetValueOnRenderThread() + int32(GTriangleRequestCount), 0);
+		uint32 MaxValueCount = FMath::Max(CVarMaxTriangleCount.GetValueOnAnyThread() + int32(GTriangleRequestCount), 0);
 		return IsEnabled() ? MaxValueCount : 0;
 	}
 
@@ -182,12 +183,12 @@ namespace ShaderPrint
 	// Accessors
 
 	// Fill the FShaderParameters parameters
-	void SetParameters(FRDGBuilder& GraphBuilder, const FShaderPrintData& Data, FShaderParameters& OutParameters)
+	void SetParameters(FRDGBuilder& GraphBuilder, const FShaderPrintData& InData, FShaderParameters& OutParameters)
 	{
-		OutParameters.Common = Data.UniformBuffer;
-		OutParameters.ShaderPrint_StateBuffer = GraphBuilder.CreateSRV(Data.ShaderPrintStateBuffer);
-		OutParameters.ShaderPrint_RWValuesBuffer = GraphBuilder.CreateUAV(Data.ShaderPrintValueBuffer);
-		OutParameters.ShaderPrint_RWPrimitivesBuffer = GraphBuilder.CreateUAV(Data.ShaderPrintPrimitiveBuffer);
+		OutParameters.Common = InData.UniformBuffer;
+		OutParameters.ShaderPrint_StateBuffer = GraphBuilder.CreateSRV(InData.ShaderPrintStateBuffer);
+		OutParameters.ShaderPrint_RWValuesBuffer = GraphBuilder.CreateUAV(InData.ShaderPrintValueBuffer);
+		OutParameters.ShaderPrint_RWPrimitivesBuffer = GraphBuilder.CreateUAV(InData.ShaderPrintPrimitiveBuffer);
 	}
 
 	void SetParameters(FRDGBuilder& GraphBuilder, FShaderParameters& OutParameters)
@@ -258,6 +259,11 @@ namespace ShaderPrint
 		return GDefaultView != nullptr && IsEnabled(*GDefaultView);
 	}
 
+	void SubmitShaderPrintData(FFrozenShaderPrintData& InData)
+	{
+		GShaderPrintDataToRender.Add(InData);
+	}
+
 	//////////////////////////////////////////////////////////////////////////////////////////////////
 	// Widget/Characters Shaders
 	 
@@ -547,13 +553,12 @@ namespace ShaderPrint
 
 	FShaderPrintSetup::FShaderPrintSetup(FIntRect InViewRect)
 	{
+		bEnabled = IsEnabled();
+
 		ViewRect = InViewRect;
-		CursorCoord = FIntPoint(-1, -1);
-		PreViewTranslation = FVector::ZeroVector;
-		DPIScale = 1.f;
 
-		FontSize = FIntPoint(FMath::Max(CVarFontSize.GetValueOnRenderThread(), 1), FMath::Max(CVarFontSize.GetValueOnRenderThread(), 1));
-		FontSpacing = FIntPoint(FMath::Max(CVarFontSpacingX.GetValueOnRenderThread(), 1), FMath::Max(CVarFontSpacingY.GetValueOnRenderThread(), 1));
+		FontSize = FIntPoint(FMath::Max(CVarFontSize.GetValueOnAnyThread(), 1), FMath::Max(CVarFontSize.GetValueOnAnyThread(), 1));
+		FontSpacing = FIntPoint(FMath::Max(CVarFontSpacingX.GetValueOnAnyThread(), 1), FMath::Max(CVarFontSpacingY.GetValueOnAnyThread(), 1));
 		MaxValueCount = GetMaxValueCount();
 		MaxStateCount = GetMaxWidgetCount();
 		MaxLineCount = GetMaxLineCount();
@@ -562,13 +567,15 @@ namespace ShaderPrint
 
 	FShaderPrintSetup::FShaderPrintSetup(FSceneView const& View)
 	{
+		bEnabled = IsEnabled() && IsSupported(View.GetShaderPlatform());
+
 		ViewRect = View.UnconstrainedViewRect;
 		CursorCoord = View.CursorPos;
 		PreViewTranslation = View.ViewMatrices.GetPreViewTranslation();
 		DPIScale = View.Family->DebugDPIScale;
 
-		FontSize = FIntPoint(FMath::Max(CVarFontSize.GetValueOnRenderThread(), 1), FMath::Max(CVarFontSize.GetValueOnRenderThread(), 1));
-		FontSpacing = FIntPoint(FMath::Max(CVarFontSpacingX.GetValueOnRenderThread(), 1), FMath::Max(CVarFontSpacingY.GetValueOnRenderThread(), 1));
+		FontSize = FIntPoint(FMath::Max(CVarFontSize.GetValueOnAnyThread(), 1), FMath::Max(CVarFontSize.GetValueOnAnyThread(), 1));
+		FontSpacing = FIntPoint(FMath::Max(CVarFontSpacingX.GetValueOnAnyThread(), 1), FMath::Max(CVarFontSpacingY.GetValueOnAnyThread(), 1));
 		MaxValueCount = GetMaxValueCount();
 		MaxStateCount = GetMaxWidgetCount();
 		MaxLineCount = GetMaxLineCount();
@@ -580,12 +587,13 @@ namespace ShaderPrint
 		FShaderPrintData ShaderPrintData;
 
 		// Common uniform buffer
+		ShaderPrintData.Setup = InSetup;
 		ShaderPrintData.UniformBuffer = CreateUniformBuffer(InSetup);
 
 		// Early out if system is disabled.
 		// Note that we still bind dummy buffers.
 		// This is in case some debug shader code is still active and accessing the buffer.
-		if (!IsEnabled())
+		if (!InSetup.bEnabled)
 		{
 			ShaderPrintData.ShaderPrintValueBuffer = GraphBuilder.RegisterExternalBuffer(GEmptyBuffer->Buffer);
 			ShaderPrintData.ShaderPrintStateBuffer = GraphBuilder.RegisterExternalBuffer(GEmptyBuffer->Buffer);
@@ -680,12 +688,39 @@ namespace ShaderPrint
 
 		return ShaderPrintData;
 	}
-	\
+
 	FShaderPrintData CreateShaderPrintData(FRDGBuilder& GraphBuilder, FShaderPrintSetup const& InSetup)
 	{
 		return CreateShaderPrintData(GraphBuilder, InSetup, nullptr);
 	}
 
+	FFrozenShaderPrintData FreezeShaderPrintData(FRDGBuilder& GraphBuilder, FShaderPrintData& ShaderPrintData)
+	{
+		FFrozenShaderPrintData Out;
+
+		Out.Setup = ShaderPrintData.Setup;
+
+		Out.ShaderPrintValueBuffer = GraphBuilder.ConvertToExternalBuffer(ShaderPrintData.ShaderPrintValueBuffer);
+		Out.ShaderPrintStateBuffer = GraphBuilder.ConvertToExternalBuffer(ShaderPrintData.ShaderPrintStateBuffer);
+		Out.ShaderPrintPrimitiveBuffer = GraphBuilder.ConvertToExternalBuffer(ShaderPrintData.ShaderPrintPrimitiveBuffer);
+
+		return Out;
+	}
+
+	FShaderPrintData UnFreezeShaderPrintData(FRDGBuilder& GraphBuilder, FFrozenShaderPrintData& FrozenShaderPrintData)
+	{
+		FShaderPrintData Out;
+
+		Out.Setup = FrozenShaderPrintData.Setup;
+		Out.UniformBuffer = CreateUniformBuffer(Out.Setup);
+
+		Out.ShaderPrintValueBuffer = GraphBuilder.RegisterExternalBuffer(FrozenShaderPrintData.ShaderPrintValueBuffer);
+		Out.ShaderPrintStateBuffer = GraphBuilder.RegisterExternalBuffer(FrozenShaderPrintData.ShaderPrintStateBuffer);
+		Out.ShaderPrintPrimitiveBuffer = GraphBuilder.RegisterExternalBuffer(FrozenShaderPrintData.ShaderPrintPrimitiveBuffer);
+
+		return Out;
+	}
+
 	//////////////////////////////////////////////////////////////////////////////////////////////////
 	// Drawing/Rendering API
 
@@ -693,9 +728,9 @@ namespace ShaderPrint
 	{
 		TRACE_CPUPROFILER_EVENT_SCOPE(ShaderPrint::BeginView);
 
-		View.ShaderPrintData = FShaderPrintData();
 		if (!IsSupported(View.GetShaderPlatform()))
 		{
+			View.ShaderPrintData = FShaderPrintData();
 			return;
 		}
 
@@ -942,40 +977,57 @@ namespace ShaderPrint
 			});
 	}
 
-	void DrawView(FRDGBuilder& GraphBuilder, const FViewInfo& View, const FScreenPassTexture& OutputTexture, const FScreenPassTexture& DepthTexture)
+	void InternalDrawView(FRDGBuilder& GraphBuilder, const FViewInfo& View, const FShaderPrintData& ShaderPrintData, const FScreenPassTexture& OutputTexture, const FScreenPassTexture& DepthTexture)
 	{
-		check(OutputTexture.IsValid());
+		if (!ensure(OutputTexture.IsValid()))
+		{
+			return;
+		}
 
 		RDG_EVENT_SCOPE(GraphBuilder, "ShaderPrint::DrawView");
 
 		const FIntRect SourceViewRect = View.ViewRect;
 		const FIntRect OutputViewRect = OutputTexture.ViewRect;
-		
+		const FVector PreViewTranslation = View.ViewMatrices.GetPreViewTranslation() - ShaderPrintData.Setup.PreViewTranslation;
+
 		// Lines
 		{
-			FRDGBufferRef DataBuffer = View.ShaderPrintData.ShaderPrintPrimitiveBuffer;
-			InternalDrawView_Primitives(GraphBuilder, View.ShaderPrintData, DataBuffer, SourceViewRect, OutputViewRect, View.ViewMatrices.GetTranslatedViewProjectionMatrix(), FVector::ZeroVector, true /*bLines*/, false /*bLocked*/, OutputTexture.Texture, DepthTexture.Texture);
+			FRDGBufferRef DataBuffer = ShaderPrintData.ShaderPrintPrimitiveBuffer;
+			InternalDrawView_Primitives(GraphBuilder, ShaderPrintData, DataBuffer, SourceViewRect, OutputViewRect, View.ViewMatrices.GetTranslatedViewProjectionMatrix(), PreViewTranslation, true /*bLines*/, false /*bLocked*/, OutputTexture.Texture, DepthTexture.Texture);
 		}
 
 		// Triangles
 		{
-			FRDGBufferRef DataBuffer = View.ShaderPrintData.ShaderPrintPrimitiveBuffer;
-			InternalDrawView_Primitives(GraphBuilder, View.ShaderPrintData, DataBuffer, SourceViewRect, OutputViewRect, View.ViewMatrices.GetTranslatedViewProjectionMatrix(), FVector::ZeroVector, false /*bLines*/, false /*bLocked*/, OutputTexture.Texture, DepthTexture.Texture);
+			FRDGBufferRef DataBuffer = ShaderPrintData.ShaderPrintPrimitiveBuffer;
+			InternalDrawView_Primitives(GraphBuilder, ShaderPrintData, DataBuffer, SourceViewRect, OutputViewRect, View.ViewMatrices.GetTranslatedViewProjectionMatrix(), PreViewTranslation, false /*bLines*/, false /*bLocked*/, OutputTexture.Texture, DepthTexture.Texture);
 		}
 
 		// Locked Lines/Triangles
 		if (View.ViewState && View.ViewState->ShaderPrintStateData.bIsLocked)
 		{
-			const FVector LockedToCurrentTranslatedOffset = View.ViewMatrices.GetPreViewTranslation() - View.ViewState->ShaderPrintStateData.PreViewTranslation;
+			const FVector LockedPreViewTranslation = View.ViewMatrices.GetPreViewTranslation() - View.ViewState->ShaderPrintStateData.PreViewTranslation;
 			FRDGBufferRef DataBuffer = GraphBuilder.RegisterExternalBuffer(View.ViewState->ShaderPrintStateData.PrimitiveBuffer);
-			InternalDrawView_Primitives(GraphBuilder, View.ShaderPrintData, DataBuffer, SourceViewRect, OutputViewRect, View.ViewMatrices.GetTranslatedViewProjectionMatrix(), LockedToCurrentTranslatedOffset, true  /*bLines*/, true/*bLocked*/, OutputTexture.Texture, DepthTexture.Texture);
-			InternalDrawView_Primitives(GraphBuilder, View.ShaderPrintData, DataBuffer, SourceViewRect, OutputViewRect, View.ViewMatrices.GetTranslatedViewProjectionMatrix(), LockedToCurrentTranslatedOffset, false /*bLines*/, true/*bLocked*/, OutputTexture.Texture, DepthTexture.Texture);
+			InternalDrawView_Primitives(GraphBuilder, ShaderPrintData, DataBuffer, SourceViewRect, OutputViewRect, View.ViewMatrices.GetTranslatedViewProjectionMatrix(), LockedPreViewTranslation, true  /*bLines*/, true/*bLocked*/, OutputTexture.Texture, DepthTexture.Texture);
+			InternalDrawView_Primitives(GraphBuilder, ShaderPrintData, DataBuffer, SourceViewRect, OutputViewRect, View.ViewMatrices.GetTranslatedViewProjectionMatrix(), LockedPreViewTranslation, false /*bLines*/, true/*bLocked*/, OutputTexture.Texture, DepthTexture.Texture);
 		}
 
 		// Characters
 		{
 			const int32 FrameNumber = View.Family ? View.Family->FrameNumber : 0u;
-			InternalDrawView_Characters(GraphBuilder, View.ShaderPrintData, OutputViewRect, FrameNumber, OutputTexture);
+			InternalDrawView_Characters(GraphBuilder, ShaderPrintData, OutputViewRect, FrameNumber, OutputTexture);
+		}
+	}
+
+	void DrawView(FRDGBuilder& GraphBuilder, const FViewInfo& View, const FScreenPassTexture& OutputTexture, const FScreenPassTexture& DepthTexture)
+	{
+		// Draw the shader print data for the view.
+		InternalDrawView(GraphBuilder, View, View.ShaderPrintData, OutputTexture, DepthTexture);
+
+		// Draw any externally enqueued shader print data.
+		for (FFrozenShaderPrintData& ShaderPrintDataToRender : GShaderPrintDataToRender)
+		{
+			FShaderPrintData ShaderPrintData = UnFreezeShaderPrintData(GraphBuilder, ShaderPrintDataToRender);
+			InternalDrawView(GraphBuilder, View, ShaderPrintData, OutputTexture, DepthTexture);
 		}
 	}
 
@@ -983,5 +1035,6 @@ namespace ShaderPrint
 	{
 		View.ShaderPrintData = FShaderPrintData();
 		GDefaultView = nullptr;
+		GShaderPrintDataToRender.Reset();
 	}
-}
+}
\ No newline at end of file
diff --git a/Engine/Source/Runtime/Renderer/Public/ShaderPrintParameters.h b/Engine/Source/Runtime/Renderer/Public/ShaderPrintParameters.h
--- a/Engine/Source/Runtime/Renderer/Public/ShaderPrintParameters.h
+++ b/Engine/Source/Runtime/Renderer/Public/ShaderPrintParameters.h
@@ -14,6 +14,7 @@
 #include "ShaderParameterMacros.h"
 #include "ShaderParameters.h"
 
+struct FFrozenShaderPrintData;
 class FSceneView;
 struct FShaderPrintData;
 class FRDGBuilder;
@@ -76,38 +77,50 @@ namespace ShaderPrint
 	/** Structure containing setup for shader print capturing. */
 	struct RENDERER_API FShaderPrintSetup
 	{
+		/** Construct with shader print disabled setup. */
+		FShaderPrintSetup() = default;
 		/** Construct with view and system defaults. */
 		FShaderPrintSetup(FSceneView const& InView);
 		/** Construct with view rectangle and system defaults. */
 		FShaderPrintSetup(FIntRect InViewRect);
 
+		/** The shader print system's enabled state. This is set in the constructor and should't be overriden. */
+		bool bEnabled = false;
 		/** Expected viewport rectangle. */
-		FIntRect ViewRect;
+		FIntRect ViewRect = FIntRect(0, 0, 1, 1);
 		/** Cursor pixel position within viewport. Can be used for isolating a pixel to debug. */
-		FIntPoint CursorCoord;
+		FIntPoint CursorCoord = FIntPoint(-1, -1);
 		/** PreView translation used for storing line positions in translated world space. */
-		FVector PreViewTranslation;
+		FVector PreViewTranslation = FVector::ZeroVector;
 		/** DPI scale to take into account when drawing font. */
-		float DPIScale;
+		float DPIScale = 1.f;
 		/** Font size in pixels. */
-		FIntPoint FontSize;
+		FIntPoint FontSize = 1;
 		/** Font spacing in pixels (not including font size). */
-		FIntPoint FontSpacing;
+		FIntPoint FontSpacing = 1;
 		/** Initial size of character buffer. Will also be increased by RequestSpaceForCharacters(). */
-		uint32 MaxValueCount;
+		uint32 MaxValueCount = 0;
 		/** Initial size of widget buffer. */
-		uint32 MaxStateCount;
+		uint32 MaxStateCount = 0;
 		/** Initial size of line buffer. Will also be increased by RequestSpaceForLines(). */
-		uint32 MaxLineCount;
+		uint32 MaxLineCount = 0;
 		/** Initial size of triangle buffer. Will also be increased by RequestSpaceForLines(). */
-		uint32 MaxTriangleCount;
+		uint32 MaxTriangleCount = 0;
 	};
 
 	/** Create the shader print render data. This allocates and clears the render buffers. */
 	RENDERER_API FShaderPrintData CreateShaderPrintData(FRDGBuilder& GraphBuilder, FShaderPrintSetup const& InSetup);
+	
+	/** Make the buffers in a FShaderPrintData object external to an RDG builder. Do this for later reuse, or when submiting for later rendering. */
+	RENDERER_API FFrozenShaderPrintData FreezeShaderPrintData(FRDGBuilder& GraphBuilder, FShaderPrintData& ShaderPrintData);
+	/** Import the shader print buffers into an RDG builder and recreate the FShaderPrintData object. */
+	RENDERER_API FShaderPrintData UnFreezeShaderPrintData(FRDGBuilder& GraphBuilder, FFrozenShaderPrintData& FrozenShaderPrintData);
+
+	/** Submit shader print data for display in the next rendered frame. */
+	RENDERER_API void SubmitShaderPrintData(FFrozenShaderPrintData& InData);
 
 	/** Fill the FShaderParameters with an explicit FShaderPrintData managed by the calling code. */
-	RENDERER_API void SetParameters(FRDGBuilder& GraphBuilder, const FShaderPrintData& Data, FShaderParameters& OutParameters);
+	RENDERER_API void SetParameters(FRDGBuilder& GraphBuilder, FShaderPrintData const& InData, FShaderParameters& OutParameters);
 	/** Fill the FShaderParameters with the opaque FShaderPrintData from the current default view. */
 	RENDERER_API void SetParameters(FRDGBuilder& GraphBuilder, FShaderParameters& OutParameters);
 	
@@ -119,13 +132,31 @@ namespace ShaderPrint
  * Structure containing shader print render data.
  * This is automatically created, setup and rendered for each view.
  * Also it is possible for client code to create and own this. 
- * If this is client managed then the client also needs to send to the shader print renderer.
+ * If this is client managed then the client can queue for rendering by calling:
+ * (i) FreezeShaderPrintData() to "freeze" the data which exports it from the current RDG builder context.
+ * (ii) SubmitShaderPrintData() to submit the frozen data for later thawing and rendering.
  */
 struct RENDERER_API FShaderPrintData
 {
+	ShaderPrint::FShaderPrintSetup Setup;
 	TUniformBufferRef<ShaderPrint::FShaderPrintCommonParameters> UniformBuffer;
 
 	FRDGBufferRef ShaderPrintValueBuffer = nullptr;
 	FRDGBufferRef ShaderPrintStateBuffer = nullptr;
 	FRDGBufferRef ShaderPrintPrimitiveBuffer = nullptr;
 };
+
+/**
+ * Structure containing "frozen" shader print render data.
+ * This is in a state so that it:
+ * (i) Can be thawed by the client for continued gathering of shader print glyphs, or
+ * (ii) Can be submitted for later rendering using SubmitShaderPrintData().
+ */
+struct RENDERER_API FFrozenShaderPrintData
+{
+	ShaderPrint::FShaderPrintSetup Setup;
+
+	TRefCountPtr<FRDGPooledBuffer> ShaderPrintValueBuffer;
+	TRefCountPtr<FRDGPooledBuffer> ShaderPrintStateBuffer;
+	TRefCountPtr<FRDGPooledBuffer> ShaderPrintPrimitiveBuffer;
+};
