Fix for missing view uniform buffer in ray tracing world updates.

diff --git a/Engine/Source/Runtime/Renderer/Private/DeferredShadingRenderer.cpp b/Engine/Source/Runtime/Renderer/Private/DeferredShadingRenderer.cpp
--- a/Engine/Source/Runtime/Renderer/Private/DeferredShadingRenderer.cpp
+++ b/Engine/Source/Runtime/Renderer/Private/DeferredShadingRenderer.cpp
@@ -1403,6 +1403,7 @@ BEGIN_SHADER_PARAMETER_STRUCT(FBuildAccelerationStructurePassParams, )
 	RDG_BUFFER_ACCESS(DynamicGeometryScratchBuffer, ERHIAccess::UAVCompute)
 	RDG_BUFFER_ACCESS(RayTracingSceneInstanceBuffer, ERHIAccess::SRVCompute)
 
+	SHADER_PARAMETER_STRUCT_REF(FViewUniformShaderParameters, View)
 	SHADER_PARAMETER_RDG_UNIFORM_BUFFER(FRaytracingLightDataPacked, LightDataPacked)
 END_SHADER_PARAMETER_STRUCT()
 
@@ -1617,6 +1618,7 @@ bool FDeferredShadingSceneRenderer::DispatchRayTracingWorldUpdates(FRDGBuilder&
 		FBuildAccelerationStructurePassParams* PassParams = GraphBuilder.AllocParameters<FBuildAccelerationStructurePassParams>();
 		PassParams->RayTracingSceneScratchBuffer = Scene->RayTracingScene.BuildScratchBuffer;
 		PassParams->RayTracingSceneInstanceBuffer = Scene->RayTracingScene.InstanceBuffer;
+		PassParams->View = ReferenceView.ViewUniformBuffer;
 		PassParams->DynamicGeometryScratchBuffer = OutDynamicGeometryScratchBuffer;
 		PassParams->LightDataPacked =  nullptr;
 
