Increase r.DFShadow.TwoSidedMeshDistanceBiasScale to 1.0
- since r.DistanceFields.TwoSidedSurfaceBiasExpand was lowered from 4 to 1 in CL 18850862.

diff --git a/Engine/Source/Runtime/Renderer/Private/DistanceFieldShadowing.cpp b/Engine/Source/Runtime/Renderer/Private/DistanceFieldShadowing.cpp
--- a/Engine/Source/Runtime/Renderer/Private/DistanceFieldShadowing.cpp
+++ b/Engine/Source/Runtime/Renderer/Private/DistanceFieldShadowing.cpp
@@ -71,7 +71,7 @@ FAutoConsoleVariableRef CVarShadowCullTileWorldSize(
 	ECVF_Scalability | ECVF_RenderThreadSafe
 	);
 
-float GDFShadowTwoSidedMeshDistanceBiasScale = 0.25f;
+float GDFShadowTwoSidedMeshDistanceBiasScale = 1.0f;
 FAutoConsoleVariableRef CVarShadowTwoSidedMeshDistanceBiasScale(
 	TEXT("r.DFShadow.TwoSidedMeshDistanceBiasScale"),
 	GDFShadowTwoSidedMeshDistanceBiasScale,
