Strata - Fixed translucent shadow depth compilation issues.

diff --git a/Engine/Shaders/Private/TranslucentShadowDepthShaders.usf b/Engine/Shaders/Private/TranslucentShadowDepthShaders.usf
--- a/Engine/Shaders/Private/TranslucentShadowDepthShaders.usf
+++ b/Engine/Shaders/Private/TranslucentShadowDepthShaders.usf
@@ -4,6 +4,12 @@
 
 #define SCENE_TEXTURES_DISABLED 1
 
+// Enable Strata only when a material is specifying a strata material is specified.
+#if !MATERIAL_IS_STRATA && STRATA_ENABLED
+#undef STRATA_ENABLED
+#define STRATA_ENABLED 0
+#endif
+
 #include "Common.ush"
 #include "/Engine/Generated/Material.ush"
 #include "/Engine/Generated/VertexFactory.ush"
diff --git a/Engine/Source/Runtime/Renderer/Private/TranslucentLighting.cpp b/Engine/Source/Runtime/Renderer/Private/TranslucentLighting.cpp
--- a/Engine/Source/Runtime/Renderer/Private/TranslucentLighting.cpp
+++ b/Engine/Source/Runtime/Renderer/Private/TranslucentLighting.cpp
@@ -327,6 +327,7 @@ public:
 	{
 		FTranslucencyShadowDepthPS::ModifyCompilationEnvironment(Parameters, OutEnvironment);
 		OutEnvironment.SetDefine(TEXT("PERSPECTIVE_CORRECT_DEPTH"), (uint32)(ShaderMode == TranslucencyShadowDepth_PerspectiveCorrect ? 1 : 0));
+		OutEnvironment.SetDefine(TEXT("STRATA_INLINE_SHADING"), 1);
 	}
 
 	TTranslucencyShadowDepthPS() = default;
