Skip instances part of groups before doing PerInstanceCulling.
- FRayTracingCullPrimitiveInstancesClosure was using the group bounds when culling instances inside a group so those instances were never culled anyway.
- This saves around 0.2ms in CitySample (0.7ms if not using ParallelFor).

diff --git a/Engine/Source/Runtime/Renderer/Private/DeferredShadingRenderer.cpp b/Engine/Source/Runtime/Renderer/Private/DeferredShadingRenderer.cpp
--- a/Engine/Source/Runtime/Renderer/Private/DeferredShadingRenderer.cpp
+++ b/Engine/Source/Runtime/Renderer/Private/DeferredShadingRenderer.cpp
@@ -1250,7 +1250,10 @@ bool FDeferredShadingSceneRenderer::GatherRayTracingWorldInstancesForView(FRDGBu
 				checkf(SceneInfo->CachedRayTracingInstance.GeometryRHI, TEXT("Ray tracing instance must have a valid geometry."));
 				RayTracingScene.Instances.Add(SceneInfo->CachedRayTracingInstance);
 
-				if (View.RayTracingCullingParameters.CullInRayTracing > 0 && GetRayTracingCullingPerInstance() && SceneInfo->CachedRayTracingInstance.NumTransforms > 1)
+				const Experimental::FHashElementId GroupId = Scene->PrimitiveRayTracingGroupIds[PrimitiveIndex];
+				const bool bUseGroupBounds = View.RayTracingCullingParameters.bCullUsingGroupIds && GroupId.IsValid();
+
+				if (View.RayTracingCullingParameters.CullInRayTracing > 0 && GetRayTracingCullingPerInstance() && SceneInfo->CachedRayTracingInstance.NumTransforms > 1 && !bUseGroupBounds)
 				{
 					FRayTracingGeometryInstance& NewInstance = RayTracingScene.Instances.Last();
 
diff --git a/Engine/Source/Runtime/Renderer/Private/RayTracing/RayTracingInstanceCulling.cpp b/Engine/Source/Runtime/Renderer/Private/RayTracing/RayTracingInstanceCulling.cpp
--- a/Engine/Source/Runtime/Renderer/Private/RayTracing/RayTracingInstanceCulling.cpp
+++ b/Engine/Source/Runtime/Renderer/Private/RayTracing/RayTracingInstanceCulling.cpp
@@ -226,18 +226,16 @@ void FRayTracingCullPrimitiveInstancesClosure::operator()() const
 {
 	FMemory::Memset(OutInstanceActivationMask.GetData(), 0xFF, OutInstanceActivationMask.Num() * 4);
 
-	const Experimental::FHashElementId GroupId = Scene->PrimitiveRayTracingGroupIds[PrimitiveIndex];
-	const bool bUseGroupBounds = CullingParameters->bCullUsingGroupIds && GroupId.IsValid();
-	const FBoxSphereBounds* const GroupBounds = (bUseGroupBounds) ? &Scene->PrimitiveRayTracingGroups.GetByElementId(GroupId).Value.Bounds : nullptr;
-	const FBoxSphereBounds& PrimitiveBounds = (GroupBounds) ? *GroupBounds : Scene->PrimitiveBounds[PrimitiveIndex].BoxSphereBounds;
-	const float MinDrawDistance = GroupBounds ? Scene->PrimitiveRayTracingGroups.GetByElementId(GroupId).Value.MinDrawDistance : Scene->PrimitiveBounds[PrimitiveIndex].MinDrawDistance;
+	checkf(!CullingParameters->bCullUsingGroupIds || !Scene->PrimitiveRayTracingGroupIds[PrimitiveIndex].IsValid(), TEXT("Shouldn't do instance level culling of primitives in raytracing groups."));
 
-	if (!RayTracing::ShouldSkipPerInstanceCullingForPrimitive(*CullingParameters, PrimitiveBounds, SceneInfo->CachedRayTracingInstanceWorldBounds[SceneInfo->SmallestRayTracingInstanceWorldBoundsIndex], bIsFarFieldPrimitive))
+	const FPrimitiveBounds& PrimitiveBounds = Scene->PrimitiveBounds[PrimitiveIndex];
+
+	if (!RayTracing::ShouldSkipPerInstanceCullingForPrimitive(*CullingParameters, PrimitiveBounds.BoxSphereBounds, SceneInfo->CachedRayTracingInstanceWorldBounds[SceneInfo->SmallestRayTracingInstanceWorldBoundsIndex], bIsFarFieldPrimitive))
 	{
 		for (int32 InstanceIndex = 0; InstanceIndex < SceneInfo->CachedRayTracingInstanceWorldBounds.Num(); InstanceIndex++)
 		{
-			const FBoxSphereBounds& InstanceBounds = (GroupBounds) ? *GroupBounds : SceneInfo->CachedRayTracingInstanceWorldBounds[InstanceIndex];
-			if (RayTracing::ShouldCullBounds(*CullingParameters, InstanceBounds, MinDrawDistance, bIsFarFieldPrimitive))
+			const FBoxSphereBounds& InstanceBounds = SceneInfo->CachedRayTracingInstanceWorldBounds[InstanceIndex];
+			if (RayTracing::ShouldCullBounds(*CullingParameters, InstanceBounds, PrimitiveBounds.MinDrawDistance, bIsFarFieldPrimitive))
 			{
 				OutInstanceActivationMask[InstanceIndex / 32] &= ~(1 << (InstanceIndex % 32));
 			}
