Improvements to Ray Tracing instance/primitive debug visualization.
- Use PrimitiveIndex combined with PrimitiveInstanceIndex in order to get a stable output without relying on world position.
- Apply some shading to the visualization based on world normal of the primitives.

Had to enable SBT indexing and bind hitgroups in order to use GetBaseInstanceIndex() and access vertex data.

diff --git a/Engine/Shaders/Private/RayTracing/RayTracingDebug.usf b/Engine/Shaders/Private/RayTracing/RayTracingDebug.usf
--- a/Engine/Shaders/Private/RayTracing/RayTracingDebug.usf
+++ b/Engine/Shaders/Private/RayTracing/RayTracingDebug.usf
@@ -3,6 +3,7 @@
 #include "../Common.ush"
 #include "../ShadingCommon.ush"
 #include "../ColorMap.ush"
+#include "../HashTable.ush"
 
 #include "RayTracingCommon.ush"
 #include "RayTracingHitGroupCommon.ush"
@@ -21,18 +22,6 @@ float MaxTraceDistance;
 float FarFieldMaxTraceDistance;
 float3 FarFieldReferencePos;
 
-// TODO: move MurmurMix and IntToColor into a common header, shared with Nanite/NaniteVisualize.usf
-
-uint MurmurMix(uint Hash)
-{
-	Hash ^= Hash >> 16;
-	Hash *= 0x85ebca6b;
-	Hash ^= Hash >> 13;
-	Hash *= 0xc2b2ae35;
-	Hash ^= Hash >> 16;
-	return Hash;
-}
-
 float3 IntToColor(uint Index)
 {
 	uint Hash = MurmurMix(Index);
@@ -180,6 +169,7 @@ float4 DebugRayTracingMaterial(RayDesc Ray, bool bIsFarField)
 
 struct FRayTracingDebugPayload : FMinimalPayload
 {
+	float3 WorldNormal;
 	uint InstanceHash;
 	uint TriangleHash;
 };
@@ -213,6 +203,9 @@ float4 DebugRayTracingInstance(RayDesc Ray)
 			Result.rgb = IntToColor(Payload.TriangleHash);
 			break;
 		}
+		
+		// apply some shading based on the primitive normal (map dot product result from [0,1] to [0.75,1]).
+		Result.rgb *= dot(Payload.WorldNormal, normalize(float3(1.0f, 1.0f, 1.0f))) * 0.125f + 0.875;
 	}
 
 	return Result;
@@ -255,26 +248,33 @@ RAY_TRACING_ENTRY_RAYGEN(RayTracingDebugMainRGS)
 	Output[PixelCoord] = Result;
 }
 
+float3 ComputeWorldNormal()
+{
+	FTriangleBaseAttributes Tri = LoadTriangleBaseAttributes(PrimitiveIndex());
+
+	float3 LocalEdges[2] = {
+		Tri.LocalPositions[1] - Tri.LocalPositions[0],
+		Tri.LocalPositions[2] - Tri.LocalPositions[0]
+	};
+	float3 LocalNormal = cross(LocalEdges[1], LocalEdges[0]);
+
+	float3x3 InverseTranspose3x3 = transpose(WorldToObject4x3());
+	float3 WorldNormal = normalize(mul(LocalNormal, InverseTranspose3x3));
+
+	return WorldNormal;
+}
+
 RAY_TRACING_ENTRY_CLOSEST_HIT(RayTracingDebugMainCHS,
 	FRayTracingDebugPayload, Payload,
 	FRayTracingIntersectionAttributes, Attributes)
 {
 	Payload.HitT = RayTCurrent();
 
-	// Can't use InstanceIndex() because it's not stable due to culling.
-	// Instad we can use InstanceID() (user data on the instance) combined with
-	// a hash of the instance transform. This gives us a reasonable overall instance hash
-	// that can be used for debug visualizations. 
-	// Dynamic objects can't be visualized in a stable way, as their transform is constantly updated.
+	const int ScenePrimitiveIndex = GetInstanceUserData();
+	const int ScenePrimitiveInstanceIndex = InstanceIndex() - GetBaseInstanceIndex();
+	Payload.InstanceHash = MurmurAdd(ScenePrimitiveIndex, ScenePrimitiveInstanceIndex);
 
-	Payload.InstanceHash = InstanceID();
-	float4x3 Mat = ObjectToWorld4x3();
-	for (int i = 0; i < 4; ++i)
-	{
-		for (int j = 0; j < 3; ++j)
-		{
-			Payload.InstanceHash += MurmurMix(asuint(Mat[i][j]));
-		}
-	}
 	Payload.TriangleHash = MurmurMix(Payload.InstanceHash + PrimitiveIndex());
+
+	Payload.WorldNormal = ComputeWorldNormal();
 }
diff --git a/Engine/Source/Runtime/Renderer/Private/RayTracing/RayTracingDebug.cpp b/Engine/Source/Runtime/Renderer/Private/RayTracing/RayTracingDebug.cpp
--- a/Engine/Source/Runtime/Renderer/Private/RayTracing/RayTracingDebug.cpp
+++ b/Engine/Source/Runtime/Renderer/Private/RayTracing/RayTracingDebug.cpp
@@ -170,6 +170,47 @@ class FRayTracingDebugTraversalCS : public FGlobalShader
 };
 IMPLEMENT_GLOBAL_SHADER(FRayTracingDebugTraversalCS, "/Engine/Private/RayTracing/RayTracingDebugTraversal.usf", "RayTracingDebugTraversalCS", SF_Compute);
 
+void BindRayTracingDebugCHSMaterialBindings(FRHICommandList& RHICmdList, const FViewInfo& View, FRayTracingPipelineState* PipelineState)
+{
+	const int32 NumTotalBindings = View.VisibleRayTracingMeshCommands.Num();
+
+	auto Alloc = [&](uint32 Size, uint32 Align)
+	{
+		return RHICmdList.Bypass()
+			? FMemStack::Get().Alloc(Size, Align)
+			: RHICmdList.Alloc(Size, Align);
+	};
+
+	const uint32 MergedBindingsSize = sizeof(FRayTracingLocalShaderBindings) * NumTotalBindings;
+	FRayTracingLocalShaderBindings* Bindings = (FRayTracingLocalShaderBindings*)Alloc(MergedBindingsSize, alignof(FRayTracingLocalShaderBindings));
+
+	const uint32 NumUniformBuffers = 1;
+	FRHIUniformBuffer** UniformBufferArray = (FRHIUniformBuffer**)Alloc(sizeof(FRHIUniformBuffer*) * NumUniformBuffers, alignof(FRHIUniformBuffer*));
+	UniformBufferArray[0] = View.ViewUniformBuffer.GetReference();
+
+	uint32 BindingIndex = 0;
+	for (const FVisibleRayTracingMeshCommand VisibleMeshCommand : View.VisibleRayTracingMeshCommands)
+	{
+		const FRayTracingMeshCommand& MeshCommand = *VisibleMeshCommand.RayTracingMeshCommand;
+
+		FRayTracingLocalShaderBindings Binding = {};
+		Binding.InstanceIndex = VisibleMeshCommand.InstanceIndex;
+		Binding.SegmentIndex = MeshCommand.GeometrySegmentIndex;
+		Binding.UniformBuffers = UniformBufferArray;
+		Binding.NumUniformBuffers = NumUniformBuffers;
+
+		Bindings[BindingIndex] = Binding;
+		BindingIndex++;
+	}
+
+	const bool bCopyDataToInlineStorage = false; // Storage is already allocated from RHICmdList, no extra copy necessary
+	RHICmdList.SetRayTracingHitGroups(
+		View.GetRayTracingSceneChecked(),
+		PipelineState,
+		NumTotalBindings, Bindings,
+		bCopyDataToInlineStorage);
+}
+
 static bool RequiresRayTracingDebugCHS(uint32 DebugVisualizationMode)
 {
 	return DebugVisualizationMode == RAY_TRACING_DEBUG_VIZ_INSTANCES || DebugVisualizationMode == RAY_TRACING_DEBUG_VIZ_TRIANGLES;
@@ -328,6 +369,7 @@ void FDeferredShadingSceneRenderer::RenderRayTracingDebug(FRDGBuilder& GraphBuil
 	auto RayGenShader = ShaderMap->GetShader<FRayTracingDebugRGS>();
 
 	FRayTracingPipelineState* Pipeline = View.RayTracingMaterialPipeline;
+	bool bRequiresBindings = false;
 
 	if (RequiresRayTracingDebugCHS(DebugVisualizationMode))
 	{
@@ -339,10 +381,11 @@ void FDeferredShadingSceneRenderer::RenderRayTracingDebug(FRDGBuilder& GraphBuil
 		auto ClosestHitShader = ShaderMap->GetShader<FRayTracingDebugCHS>();
 		FRHIRayTracingShader* HitGroupTable[] = { ClosestHitShader.GetRayTracingShader() };
 		Initializer.SetHitGroupTable(HitGroupTable);
-		Initializer.bAllowHitGroupIndexing = false; // Use the same hit shader for all geometry in the scene by disabling SBT indexing.
+		Initializer.bAllowHitGroupIndexing = true; // Required for stable output using GetBaseInstanceIndex().
 		Initializer.MaxPayloadSizeInBytes = RAY_TRACING_MAX_ALLOWED_PAYLOAD_SIZE;
 
 		Pipeline = PipelineStateCache::GetAndOrCreateRayTracingPipelineState(GraphBuilder.RHICmdList, Initializer);
+		bRequiresBindings = true;
 	}
 
 	FRayTracingDebugRGS::FParameters* RayGenParameters = GraphBuilder.AllocParameters<FRayTracingDebugRGS::FParameters>();
@@ -377,11 +420,16 @@ void FDeferredShadingSceneRenderer::RenderRayTracingDebug(FRDGBuilder& GraphBuil
 		RDG_EVENT_NAME("RayTracingDebug"),
 		RayGenParameters,
 		ERDGPassFlags::Compute,
-		[this, RayGenParameters, RayGenShader, &View, Pipeline, ViewRect](FRHIRayTracingCommandList& RHICmdList)
+		[this, RayGenParameters, RayGenShader, &View, Pipeline, ViewRect, bRequiresBindings](FRHIRayTracingCommandList& RHICmdList)
 	{
 		FRayTracingShaderBindingsWriter GlobalResources;
 		SetShaderParameters(GlobalResources, RayGenShader, *RayGenParameters);
 
+		if (bRequiresBindings)
+		{
+			BindRayTracingDebugCHSMaterialBindings(RHICmdList, View, Pipeline);
+		}
+
 		RHICmdList.RayTraceDispatch(Pipeline, RayGenShader.GetRayTracingShader(), View.GetRayTracingSceneChecked(), GlobalResources, ViewRect.Size().X, ViewRect.Size().Y);
 	});
 }
