Implements r.Translucency.DynamicRes.MaxScreenPercentage

diff --git a/Engine/Source/Runtime/Renderer/Private/TranslucentRendering.cpp b/Engine/Source/Runtime/Renderer/Private/TranslucentRendering.cpp
--- a/Engine/Source/Runtime/Renderer/Private/TranslucentRendering.cpp
+++ b/Engine/Source/Runtime/Renderer/Private/TranslucentRendering.cpp
@@ -49,6 +49,12 @@ static TAutoConsoleVariable<float> CVarTranslucencyMinScreenPercentage(
 	TEXT("Minimal screen percentage for translucency."),
 	ECVF_RenderThreadSafe | ECVF_Default);
 
+static TAutoConsoleVariable<float> CVarTranslucencyMaxScreenPercentage(
+	TEXT("r.Translucency.DynamicRes.MaxScreenPercentage"),
+	DynamicRenderScaling::FractionToPercentage(DynamicRenderScaling::FHeuristicSettings::kDefaultMaxResolutionFraction),
+	TEXT("Maximal screen percentage for translucency."),
+	ECVF_RenderThreadSafe | ECVF_Default);
+
 static TAutoConsoleVariable<float> CVarTranslucencyTimeBudget(
 	TEXT("r.Translucency.DynamicRes.TimeBudget"),
 	DynamicRenderScaling::FHeuristicSettings::kBudgetMsDisabled,
@@ -93,6 +99,7 @@ DynamicRenderScaling::FHeuristicSettings GetDynamicTranslucencyResolutionSetting
 	BucketSetting.Model = DynamicRenderScaling::EHeuristicModel::Quadratic;
 	BucketSetting.bModelScalesWithPrimaryScreenPercentage = CVarTranslucencyScreenPercentageBasis.GetValueOnAnyThread() != 1;
 	BucketSetting.MinResolutionFraction = DynamicRenderScaling::GetPercentageCVarToFraction(CVarTranslucencyMinScreenPercentage);
+	BucketSetting.MaxResolutionFraction = DynamicRenderScaling::GetPercentageCVarToFraction(CVarTranslucencyMaxScreenPercentage);
 	BucketSetting.BudgetMs              = CVarTranslucencyTimeBudget.GetValueOnAnyThread();
 	BucketSetting.ChangeThreshold       = DynamicRenderScaling::GetPercentageCVarToFraction(CVarTranslucencyChangeThreshold);
 	BucketSetting.TargetedHeadRoom      = DynamicRenderScaling::GetPercentageCVarToFraction(CVarTranslucencyTargetedHeadRoomPercentage);
