Multi-view-family rendering:  Fix for crash when Lumen ray tracing features are enabled.  Ray tracing scene resources initialized in the first scene renderer were not being copied to views for other scene renderers / view families.  Before CL 20540730, all view families were in the same scene renderer structure, so the loop had been running for views across all families.  This fix restores that behavior.  I missed the bug in my tests because someone had disabled Lumen in content in the scene I was testing.

#rnx

diff --git a/Engine/Source/Runtime/Renderer/Private/DeferredShadingRenderer.cpp b/Engine/Source/Runtime/Renderer/Private/DeferredShadingRenderer.cpp
--- a/Engine/Source/Runtime/Renderer/Private/DeferredShadingRenderer.cpp
+++ b/Engine/Source/Runtime/Renderer/Private/DeferredShadingRenderer.cpp
@@ -1507,16 +1507,16 @@ bool FDeferredShadingSceneRenderer::SetupRayTracingPipelineStates(FRDGBuilder& G
 
 	uint32 NumOfSkippedRayTracingLights = 0;
 
-	for (int32 ViewIndex = 0; ViewIndex < Views.Num(); ++ViewIndex)
+	for (int32 ViewIndex = 0; ViewIndex < AllFamilyViews.Num(); ++ViewIndex)
 	{
-		FViewInfo& View = Views[ViewIndex];
+		FViewInfo* View = (FViewInfo*)AllFamilyViews[ViewIndex];
 
 		// Send common ray tracing resources from reference view to all others.
-		if (ViewIndex != ReferenceViewIndex)
+		if (View != &ReferenceView)
 		{
-			View.RayTracingSubSurfaceProfileTexture = ReferenceView.RayTracingSubSurfaceProfileTexture;
-			View.RayTracingSubSurfaceProfileSRV = ReferenceView.RayTracingSubSurfaceProfileSRV;
-			View.RayTracingMaterialPipeline = ReferenceView.RayTracingMaterialPipeline;
+			View->RayTracingSubSurfaceProfileTexture = ReferenceView.RayTracingSubSurfaceProfileTexture;
+			View->RayTracingSubSurfaceProfileSRV = ReferenceView.RayTracingSubSurfaceProfileSRV;
+			View->RayTracingMaterialPipeline = ReferenceView.RayTracingMaterialPipeline;
 		}
 
 		if (bIsPathTracing)
@@ -1527,7 +1527,7 @@ bool FDeferredShadingSceneRenderer::SetupRayTracingPipelineStates(FRDGBuilder& G
 		else
 		{
 			// This light data is a function of the camera position, so must be computed per view.
-			View.RayTracingLightDataUniformBuffer = CreateRayTracingLightData(GraphBuilder, Scene->Lights, View, NumOfSkippedRayTracingLights);
+			View->RayTracingLightDataUniformBuffer = CreateRayTracingLightData(GraphBuilder, Scene->Lights, *View, NumOfSkippedRayTracingLights);
 		}
 	}
 
@@ -1881,11 +1881,14 @@ void FDeferredShadingSceneRenderer::WaitForRayTracingScene(FRDGBuilder& GraphBui
 		}
 
 		// Send ray tracing resources from reference view to all others.
-		for (int32 ViewIndex = 1; ViewIndex < Views.Num(); ++ViewIndex)
+		for (int32 ViewIndex = 0; ViewIndex < AllFamilyViews.Num(); ++ViewIndex)
 		{
-			FViewInfo& View = Views[ViewIndex];
-			View.RayTracingMaterialGatherPipeline = ReferenceView.RayTracingMaterialGatherPipeline;
-			View.LumenHardwareRayTracingMaterialPipeline = ReferenceView.LumenHardwareRayTracingMaterialPipeline;
+			FViewInfo* View = (FViewInfo*)AllFamilyViews[ViewIndex];
+			if (View != &ReferenceView)
+			{
+				View->RayTracingMaterialGatherPipeline = ReferenceView.RayTracingMaterialGatherPipeline;
+				View->LumenHardwareRayTracingMaterialPipeline = ReferenceView.LumenHardwareRayTracingMaterialPipeline;
+			}
 		}
 
 		if (RayTracingDynamicGeometryUpdateEndTransition)
