Fix a shader bug of mobile pixel projected reflection

diff --git a/Engine/Shaders/Private/PostProcessPixelProjectedReflectionMobile.usf b/Engine/Shaders/Private/PostProcessPixelProjectedReflectionMobile.usf
--- a/Engine/Shaders/Private/PostProcessPixelProjectedReflectionMobile.usf
+++ b/Engine/Shaders/Private/PostProcessPixelProjectedReflectionMobile.usf
@@ -222,8 +222,7 @@ void ReflectionPassVS(
 #if REFLECTION_PASS_PIXEL_SHADER
 
 #if PROJECTION_OUTPUT_TYPE_TEXTURE
-Texture2D<uint> ProjectionTexture;
-SamplerState ProjectionTextureSampler;
+Texture2D<uint> ProjectionTextureSRV;
 #else
 Buffer<uint> ProjectionBuffer;
 #endif
@@ -231,11 +230,7 @@ Buffer<uint> ProjectionBuffer;
 uint GetEncodeValue(int2 ReflectingPixel)
 {
 #if PROJECTION_OUTPUT_TYPE_TEXTURE
-#if COMPILER_GLSL_ES3_1 // Force to use a point sampler for Texture2D<uint> on OpenGLES platform
-	uint EncodeValue = ProjectionTexture.SampleLevel(ProjectionTextureSampler, (ReflectingPixel + float2(0.5f, 0.5f)) * BufferSizeAndInvSize.zw, 0).x;
-#else
-	uint EncodeValue = ProjectionTexture.Load(int3(ReflectingPixel, 0));
-#endif
+	uint EncodeValue = ProjectionTextureSRV.Load(int3(ReflectingPixel, 0));
 #else
 	uint EncodeValue = ProjectionBuffer[ReflectingPixel.x + ReflectingPixel.y * int(BufferSizeAndInvSize.x)];
 #endif
@@ -251,7 +246,7 @@ void DecodeProjectionBufferValue(uint EncodeValue, out int2 PixelOffset)
 	EncodeValue = EncodeValue >> 2;
 
 	// unpack whole part
-	PixelOffset.x = ((EncodeValue & 1) == 1 ? 1 : -1) * ((EncodeValue >> 1) & 2047);
+	PixelOffset.x = ((EncodeValue & 1) == 1 ? 1 : -1) * int(((EncodeValue >> 1) & 2047));
 	EncodeValue = EncodeValue >> 12;
 	PixelOffset.y = EncodeValue;
 
diff --git a/Engine/Source/Runtime/Renderer/Private/PostProcess/PostProcessPixelProjectedReflectionMobile.cpp b/Engine/Source/Runtime/Renderer/Private/PostProcess/PostProcessPixelProjectedReflectionMobile.cpp
--- a/Engine/Source/Runtime/Renderer/Private/PostProcess/PostProcessPixelProjectedReflectionMobile.cpp
+++ b/Engine/Source/Runtime/Renderer/Private/PostProcess/PostProcessPixelProjectedReflectionMobile.cpp
@@ -135,8 +135,7 @@ public:
 		SHADER_PARAMETER_STRUCT_REF(FViewUniformShaderParameters, View)
 		SHADER_PARAMETER_RDG_TEXTURE(Texture2D, SceneColorTexture)
 		SHADER_PARAMETER_SAMPLER(SamplerState, SceneColorSampler)
-		SHADER_PARAMETER_RDG_TEXTURE(Texture2D<uint>, ProjectionTexture)
-		SHADER_PARAMETER_SAMPLER(SamplerState, ProjectionTextureSampler)
+		SHADER_PARAMETER_RDG_TEXTURE_SRV(Texture2D<uint>, ProjectionTextureSRV)
 		SHADER_PARAMETER_RDG_BUFFER_SRV(Buffer<uint>, ProjectionBuffer)
 		SHADER_PARAMETER(FVector4f, ReflectionPlane)
 		SHADER_PARAMETER_EX(FVector4f, BufferSizeAndInvSize, EShaderPrecisionModifier::Half)
@@ -203,6 +202,7 @@ void FMobileSceneRenderer::RenderPixelProjectedReflection(FRDGBuilder& GraphBuil
 	const FIntPoint BufferSize = PlanarReflectionSceneProxy->RenderTarget->GetSizeXY();
 	
 	FRDGTextureRef ProjectionTexture = nullptr;
+	FRDGTextureSRVRef ProjectionTextureSRV = nullptr;
 	FRDGTextureUAVRef ProjectionTextureUAV = nullptr;
 
 	FRDGBufferRef ProjectionBuffer = nullptr;
@@ -214,6 +214,7 @@ void FMobileSceneRenderer::RenderPixelProjectedReflection(FRDGBuilder& GraphBuil
 	if (bProjectionOutputTexture)
 	{
 		ProjectionTexture = GraphBuilder.CreateTexture(FRDGTextureDesc::Create2D(BufferSize, PF_R32_UINT, FClearValueBinding::Black, TexCreate_ShaderResource | TexCreate_UAV), TEXT("ProjectionTexture"));
+		ProjectionTextureSRV = GraphBuilder.CreateSRV(FRDGTextureSRVDesc::Create(ProjectionTexture));
 		ProjectionTextureUAV = GraphBuilder.CreateUAV(ProjectionTexture);
 		uint32 ClearColor[4] = { 0xFFFFFFFF, 0xFFFFFFFF, 0xFFFFFFFF, 0xFFFFFFFF };
 		AddClearUAVPass(GraphBuilder, ProjectionTextureUAV, ClearColor);
@@ -295,8 +296,7 @@ void FMobileSceneRenderer::RenderPixelProjectedReflection(FRDGBuilder& GraphBuil
 
 		if (bProjectionOutputTexture)
 		{
-			PSShaderParameters->ProjectionTexture = ProjectionTexture;
-			PSShaderParameters->ProjectionTextureSampler = TStaticSamplerState<SF_Point, AM_Clamp, AM_Clamp, AM_Clamp>::GetRHI();
+			PSShaderParameters->ProjectionTextureSRV = ProjectionTextureSRV;
 		}
 		else
 		{
