Make sure the light debug view is always visible even when nothing request the light grid.

diff --git a/Engine/Source/Runtime/Renderer/Private/DeferredShadingRenderer.cpp b/Engine/Source/Runtime/Renderer/Private/DeferredShadingRenderer.cpp
--- a/Engine/Source/Runtime/Renderer/Private/DeferredShadingRenderer.cpp
+++ b/Engine/Source/Runtime/Renderer/Private/DeferredShadingRenderer.cpp
@@ -265,6 +265,7 @@ namespace Nanite
 	extern bool IsStatFilterActive(const FString& FilterName);
 	extern void ListStatFilters(FSceneRenderer* SceneRenderer);
 }
+extern bool ShouldVisualizeLightGrid();
 
 DECLARE_CYCLE_STAT(TEXT("InitViews Intentional Stall"), STAT_InitViews_Intentional_Stall, STATGROUP_InitViews);
 
@@ -2300,7 +2301,8 @@ void FDeferredShadingSceneRenderer::Render(FRDGBuilder& GraphBuilder)
 			VolumetricCloudWantsToSampleLocalLights(Scene, ViewFamily.EngineShowFlags) ||
 			ViewFamily.ViewMode != VMI_Lit ||
 			bAnyLumenEnabled ||
-			VirtualShadowMapArray.IsEnabled());
+			VirtualShadowMapArray.IsEnabled() ||
+			ShouldVisualizeLightGrid());
 	}
 
 	// force using occ queries for wireframe if rendering is parented or frozen in the first view
diff --git a/Engine/Source/Runtime/Renderer/Private/LightGridInjection.cpp b/Engine/Source/Runtime/Renderer/Private/LightGridInjection.cpp
--- a/Engine/Source/Runtime/Renderer/Private/LightGridInjection.cpp
+++ b/Engine/Source/Runtime/Renderer/Private/LightGridInjection.cpp
@@ -98,6 +98,11 @@ FAutoConsoleVariableRef CVarLightCullingMaxDistanceOverride(
 
 extern TAutoConsoleVariable<int32> CVarVirtualShadowOnePassProjection;
 
+bool ShouldVisualizeLightGrid()
+{
+	return GForwardLightGridDebug > 0;
+}
+
 void SetupDummyForwardLightUniformParameters(FRDGBuilder& GraphBuilder, FForwardLightData& ForwardLightData)
 {
 	const FRDGSystemTextures& SystemTextures = FRDGSystemTextures::Get(GraphBuilder);
@@ -815,7 +820,7 @@ void FDeferredShadingSceneRenderer::GatherLightsAndComputeLightGrid(FRDGBuilder&
 	for (int32 ViewIndex = 0; ViewIndex < Views.Num(); ViewIndex++)
 	{
 		const FViewInfo& View = Views[ViewIndex];
-		bAnyViewUsesForwardLighting |= View.bTranslucentSurfaceLighting || ShouldRenderVolumetricFog() || View.bHasSingleLayerWaterMaterial || VolumetricCloudWantsToSampleLocalLights(Scene, ViewFamily.EngineShowFlags);
+		bAnyViewUsesForwardLighting |= View.bTranslucentSurfaceLighting || ShouldRenderVolumetricFog() || View.bHasSingleLayerWaterMaterial || VolumetricCloudWantsToSampleLocalLights(Scene, ViewFamily.EngineShowFlags) || ShouldVisualizeLightGrid();
 		bAnyViewUsesLumen |= GetViewPipelineState(View).DiffuseIndirectMethod == EDiffuseIndirectMethod::Lumen || GetViewPipelineState(View).ReflectionsMethod == EReflectionsMethod::Lumen;
 	}
 	
@@ -975,11 +980,6 @@ class FDebugLightGridPS : public FGlobalShader
 };
 IMPLEMENT_GLOBAL_SHADER(FDebugLightGridPS, "/Engine/Private/LightGridInjection.usf", "DebugLightGridPS", SF_Pixel);
 
-bool ShouldVisualizeLightGrid()
-{
-	return GForwardLightGridDebug > 0;
-}
-
 FScreenPassTexture AddVisualizeLightGridPass(FRDGBuilder& GraphBuilder, const FViewInfo& View, FScreenPassTexture& ScreenPassSceneColor, FRDGTextureRef SceneDepthTexture)
 {
 	RDG_EVENT_SCOPE(GraphBuilder, "VisualizeLightGrid");
