SceneViewExtension: Adding SSR input pass support (which is already present in PPM as a blend location).

diff --git a/Engine/Source/Runtime/Engine/Public/SceneViewExtension.h b/Engine/Source/Runtime/Engine/Public/SceneViewExtension.h
--- a/Engine/Source/Runtime/Engine/Public/SceneViewExtension.h
+++ b/Engine/Source/Runtime/Engine/Public/SceneViewExtension.h
@@ -89,7 +89,7 @@ struct FScreenPassTexture;
 class FViewport;
 
 
-/** This is used to add more flexibility to Post Processing, so that users can subscribe to any after Post Porocessing Pass events. */
+/** This is used to add more flexibility to Post Processing, so that users can subscribe to any after Post Processing Pass events. */
 FUNC_DECLARE_DELEGATE(FAfterPassCallbackDelegate, FScreenPassTexture /*ReturnSceneColor*/, FRDGBuilder& /*GraphBuilder*/, const FSceneView& /*View*/, const FPostProcessMaterialInputs& /*Inputs*/)
 using FAfterPassCallbackDelegateArray = TArray<FAfterPassCallbackDelegate, SceneRenderingAllocator>;
 
@@ -100,6 +100,7 @@ public:
 
 	enum class EPostProcessingPass : uint32
 	{
+		SSRInput,
 		MotionBlur,
 		Tonemap,
 		FXAA,
diff --git a/Engine/Source/Runtime/Renderer/Private/PostProcess/PostProcessing.cpp b/Engine/Source/Runtime/Renderer/Private/PostProcess/PostProcessing.cpp
--- a/Engine/Source/Runtime/Renderer/Private/PostProcess/PostProcessing.cpp
+++ b/Engine/Source/Runtime/Renderer/Private/PostProcess/PostProcessing.cpp
@@ -486,9 +486,13 @@ void AddPostProcessingPasses(
 		PassSequence.SetEnabled(EPass::VisualizeDepthOfField, bVisualizeDepthOfField);
 		PassSequence.SetEnabled(EPass::VisualizeLocalExposure, EngineShowFlags.VisualizeLocalExposure);
 
+		FAfterPassCallbackDelegateArray SSRInputDelegates;
+
 		for (int32 ViewExt = 0; ViewExt < View.Family->ViewExtensions.Num(); ++ViewExt)
 		{
-			for (int32 SceneViewPassId = 0; SceneViewPassId != static_cast<int>(ISceneViewExtension::EPostProcessingPass::MAX); SceneViewPassId++)
+			View.Family->ViewExtensions[ViewExt]->SubscribeToPostProcessingPass(ISceneViewExtension::EPostProcessingPass::SSRInput, SSRInputDelegates, true);
+
+			for (int32 SceneViewPassId = static_cast<int32>(ISceneViewExtension::EPostProcessingPass::MotionBlur); SceneViewPassId != static_cast<int32>(ISceneViewExtension::EPostProcessingPass::MAX); SceneViewPassId++)
 			{
 				ISceneViewExtension::EPostProcessingPass SceneViewPass = static_cast<ISceneViewExtension::EPostProcessingPass>(SceneViewPassId);
 				EPass PostProcessingPass = TranslatePass(SceneViewPass);
@@ -652,15 +656,30 @@ void AddPostProcessingPasses(
 
 		ensure(SceneColor.ViewRect.Size() == PostTAAViewSize);
 
-		// Post Process Material Chain - SSR Input
+		// SVE/Post Process Material Chain - SSR Input
 		if (View.ViewState && !View.bStatePrevViewInfoIsReadOnly)
 		{
+			FScreenPassTexture PassOutput;
+			FPostProcessMaterialInputs InOutPassInputs = GetPostProcessMaterialInputs(SceneColor);
 			const FPostProcessMaterialChain MaterialChain = GetPostProcessMaterialChain(View, BL_SSRInput);
 
 			if (MaterialChain.Num())
+			{
+				PassOutput = AddPostProcessMaterialChain(GraphBuilder, View, InOutPassInputs, MaterialChain);
+			}
+
+			for (FAfterPassCallbackDelegate& PassCallback : SSRInputDelegates)
+			{
+				if (PassOutput.IsValid())
+				{
+					InOutPassInputs.SetInput(EPostProcessMaterialInput::SceneColor, PassOutput);
+				}
+				PassOutput = PassCallback.Execute(GraphBuilder, View, InOutPassInputs);
+			}
+
+			if (PassOutput.IsValid())
 			{
 				// Save off SSR post process output for the next frame.
-				FScreenPassTexture PassOutput = AddPostProcessMaterialChain(GraphBuilder, View, GetPostProcessMaterialInputs(SceneColor), MaterialChain);
 				GraphBuilder.QueueTextureExtraction(PassOutput.Texture, &View.ViewState->PrevFrameViewInfo.CustomSSRInput.RT[0]);
 
 				View.ViewState->PrevFrameViewInfo.CustomSSRInput.ViewportRect = PassOutput.ViewRect;
