Move VolumetricCloudRendering compute shader to Wave32. This is significantly faster on some platforms due to the divergence of the shader. Adjusting thread group size to 32 in combination with this made no measurable difference to perf, so I left that at 64.

diff --git a/Engine/Source/Runtime/Renderer/Private/VolumetricCloudRendering.cpp b/Engine/Source/Runtime/Renderer/Private/VolumetricCloudRendering.cpp
--- a/Engine/Source/Runtime/Renderer/Private/VolumetricCloudRendering.cpp
+++ b/Engine/Source/Runtime/Renderer/Private/VolumetricCloudRendering.cpp
@@ -861,6 +861,8 @@ class FRenderVolumetricCloudRenderViewCS : public FMeshMaterialShader
 		// This shader takes a very long time to compile with FXC, so we pre-compile it with DXC first and then forward the optimized HLSL to FXC.
 		OutEnvironment.CompilerFlags.Add(CFLAG_PrecompileWithDXC);
 
+		OutEnvironment.CompilerFlags.Add(CFLAG_Wave32);
+
 		FForwardLightingParameters::ModifyCompilationEnvironment(Parameters.Platform, OutEnvironment);
 	}
 };
