Fix missing GRHISupportsRayTracingShaders check before initializing HWRT pipelines

diff --git a/Engine/Source/Runtime/Renderer/Private/DeferredShadingRenderer.cpp b/Engine/Source/Runtime/Renderer/Private/DeferredShadingRenderer.cpp
--- a/Engine/Source/Runtime/Renderer/Private/DeferredShadingRenderer.cpp
+++ b/Engine/Source/Runtime/Renderer/Private/DeferredShadingRenderer.cpp
@@ -1504,6 +1504,7 @@ bool FDeferredShadingSceneRenderer::SetupRayTracingPipelineStates(FRDGBuilder& G
 	}
 
 	// Add deferred material gather shaders
+	if (GRHISupportsRayTracingShaders)
 	{
 		TArray<FRHIRayTracingShader*> DeferredMaterialRayGenShaders;
 		if (!IsForwardShadingEnabled(ShaderPlatform))
@@ -1532,6 +1533,7 @@ bool FDeferredShadingSceneRenderer::SetupRayTracingPipelineStates(FRDGBuilder& G
 	}
 
 	// Add Lumen hardware ray tracing materials
+	if (GRHISupportsRayTracingShaders)
 	{
 		TArray<FRHIRayTracingShader*> LumenHardwareRayTracingRayGenShaders;
 		if (DoesPlatformSupportLumenGI(ShaderPlatform))
