Rewrites TSR's shading rejection heuristics to reduce ghosting on translucency, pixel shader animations or when motion vector are wrong

diff --git a/Engine/Shaders/Private/TemporalSuperResolution/TSRCommon.ush b/Engine/Shaders/Private/TemporalSuperResolution/TSRCommon.ush
--- a/Engine/Shaders/Private/TemporalSuperResolution/TSRCommon.ush
+++ b/Engine/Shaders/Private/TemporalSuperResolution/TSRCommon.ush
@@ -208,10 +208,12 @@ static const uint kPlusIndexes3x3[5] = { 4, 1, 3, 7, 5 };
 #if CONFIG_SCENE_COLOR_APLHA
 	#define tsr_halfC tsr_half4
 	#define tsr_halfCx2 tsr_half4x2
+	#define CONFIG_CHANNEL_COUNT 4
 
 #else
 	#define tsr_halfC tsr_half3
 	#define tsr_halfCx2 tsr_half3x2
+	#define CONFIG_CHANNEL_COUNT 3
 
 #endif
 
@@ -1038,6 +1040,18 @@ tsr_half UncompressParallaxFactor(tsr_subpixel_payload CompressedParallaxFactor)
 		tsr_half(1.0);
 }
 
+CALL_SITE_DEBUGLOC
+tsr_half3x2 QuantizeForUNormRenderTarget(tsr_half3x2 Color, tsr_half E, const float3 QuantizationError)
+{
+	return Color + dpv_interleave_mono_registers(tsr_half3(QuantizationError)) * (E - tsr_half(0.5));
+}
+
+CALL_SITE_DEBUGLOC
+tsr_half3 QuantizeForUNormRenderTarget(tsr_half3 Color, tsr_half E, const float3 QuantizationError)
+{
+	return dpv_lo(QuantizeForUNormRenderTarget(dpv_interleave_mono_registers(Color), E, QuantizationError));
+}
+
 CALL_SITE_DEBUGLOC
 tsr_half3x2 QuantizeForFloatRenderTarget(tsr_half3x2 Color, tsr_half E, const float3 QuantizationError)
 {
diff --git a/Engine/Shaders/Private/TemporalSuperResolution/TSRCompareHistory.usf b/Engine/Shaders/Private/TemporalSuperResolution/TSRCompareHistory.usf
deleted file mode 100644
diff --git a/Engine/Shaders/Private/TemporalSuperResolution/TSRDebugHistory.usf b/Engine/Shaders/Private/TemporalSuperResolution/TSRDebugHistory.usf
--- a/Engine/Shaders/Private/TemporalSuperResolution/TSRDebugHistory.usf
+++ b/Engine/Shaders/Private/TemporalSuperResolution/TSRDebugHistory.usf
@@ -13,7 +13,6 @@
 
 Texture2D<tsr_halfC> History_LowFrequency;
 Texture2D<tsr_halfC> History_HighFrequency;
-Texture2D<tsr_halfC> History_HighFrequencyHalfRes;
 Texture2D<tsr_half2> History_Metadata;
 Texture2D<tsr_half3> History_Translucency;
 Texture2D<tsr_half>  History_TranslucencyAlpha;
@@ -21,7 +20,6 @@ Texture2D<tsr_subpixel_details> History_SubpixelDetails;
 
 Texture2D<tsr_halfC> PrevHistory_LowFrequency;
 Texture2D<tsr_halfC> PrevHistory_HighFrequency;
-Texture2D<tsr_halfC> PrevHistory_HighFrequencyHalfRes;
 Texture2D<tsr_half2> PrevHistory_Metadata;
 Texture2D<tsr_half3> PrevHistory_Translucency;
 Texture2D<tsr_half>  PrevHistory_TranslucencyAlpha;
diff --git a/Engine/Shaders/Private/TemporalSuperResolution/TSRDecimateHistory.usf b/Engine/Shaders/Private/TemporalSuperResolution/TSRDecimateHistory.usf
--- a/Engine/Shaders/Private/TemporalSuperResolution/TSRDecimateHistory.usf
+++ b/Engine/Shaders/Private/TemporalSuperResolution/TSRDecimateHistory.usf
@@ -9,14 +9,8 @@
 
 #define CONFIG_KERNEL 0
 
-#define CONFIG_OUTPUT_HALF_RES 0 // TODO
-
-#define CONFIG_MEDIAN_DECIMATE 1
-
 #define CONFIG_MANUAL_LDS_SPILL 1
 
-#define CONFIG_UPDATE_SUBPIXEL_DETAIL (DIM_HIGH_FREQUENCY_ONLY)
-
 #define CONFIG_VELOCITY_SAMPLES 9
 #if CONFIG_VELOCITY_SAMPLES != 1
 	#define CONFIG_VELOCITY_SAMPLE_PAIRS ((CONFIG_VELOCITY_SAMPLES - 1) / 2)
@@ -36,7 +30,6 @@
 //------------------------------------------------------- PARAMETERS
 
 float4x4 RotationalClipToPrevClip;
-float3 OutputQuantizationError;
 
 Texture2D<float2> DilatedVelocityTexture;
 Texture2D<float> ClosestDepthTexture;
@@ -44,17 +37,23 @@ Texture2D<uint> PrevUseCountTexture;
 Texture2D<uint> PrevClosestDepthTexture;
 Texture2D<tsr_subpixel_payload> ParallaxFactorTexture;
 
-Texture2D<tsr_halfC> PrevHistory_LowFrequency;
-Texture2D<tsr_halfC> PrevHistory_HighFrequency;
-Texture2D<tsr_halfC> PrevHistory_HighFrequencyHalfRes;
-Texture2D<tsr_half2> PrevHistory_Metadata;
-Texture2D<tsr_half4> PrevHistory_Translucency;
-Texture2D<tsr_subpixel_details> PrevHistory_SubpixelDetails;
+Texture2D<tsr_subpixel_details> PrevHistorySubpixelDetails;
+
+Texture2D<tsr_halfC> PrevHistoryGuide;
+float2 PrevGuideInfo_Extent;
+float2 PrevGuideInfo_ExtentInverse;
+uint2  PrevGuideInfo_ViewportMin;
+uint2  PrevGuideInfo_ViewportMax;
+float2 PrevGuideInfo_ViewportSize;
+float2 PrevGuideInfo_ViewportSizeInverse;
+float2 PrevGuideInfo_UVViewportBilinearMin;
+float2 PrevGuideInfo_UVViewportBilinearMax;
+FScreenTransform InputPixelPosToReprojectScreenPos;
+FScreenTransform ScreenPosToPrevHistoryGuideBufferUV;
+float3 HistoryGuideQuantizationError;
 
-RWTexture2D<tsr_halfC> HalfResPredictionSceneColorOutput;
-RWTexture2D<tsr_half> HalfResParallaxRejectionMaskOutput;
 
-RWTexture2D<tsr_halfC> PredictionSceneColorOutput;
+RWTexture2D<tsr_halfC> ReprojectedHistoryGuideOutput;
 
 RWTexture2D<float2> HoleFilledVelocityOutput;
 RWTexture2D<tsr_half2> HoleFilledVelocityMaskOutput;
@@ -70,60 +69,6 @@ RWTexture2D<tsr_subpixel_payload> HistorySubpixelDetailsOutput;
 #endif
 
 
-#if CONFIG_SCENE_COLOR_APLHA
-groupshared tsr_halfC SharedHistoryColor[TILE_SIZE * TILE_SIZE];
-groupshared tsr_half  SharedMask[TILE_SIZE * TILE_SIZE];
-#else
-groupshared tsr_half4 SharedHistoryColorAndMask[TILE_SIZE * TILE_SIZE];
-#endif
-
-
-//------------------------------------------------------- ENTRY POINT
-
-void LHOperator2(tsr_halfC A, tsr_halfC B, out tsr_halfC L, out tsr_halfC H)
-{
-	L = min(A, B);
-	H = max(A, B);
-}
-
-void LMHOperator3(tsr_halfC A, tsr_halfC B, tsr_halfC C, out tsr_halfC L, out tsr_halfC M, out tsr_halfC H)
-{
-	// TODO: min3, max3, med3.
-	tsr_halfC X = min(B, C);
-	tsr_halfC Y = max(B, C);
-	
-	L = min(A, X);
-	tsr_halfC Z = max(A, X);
-
-	M = min(Z, Y);
-	H = max(Z, Y);
-}
-
-tsr_halfC Median3(tsr_halfC A, tsr_halfC B, tsr_halfC C)
-{
-	// TODO: med3.
-	tsr_halfC L, M, H;
-	LMHOperator3(A, B, C, L, M, H);
-	return M;
-}
-
-tsr_halfC Median5(tsr_halfC In[5])
-{
-	LMHOperator3(In[1], In[2], In[3], /* out */ In[1], /* out */ In[2], /* out */ In[3]);
-	LHOperator2(In[0], In[4], /* out */ In[0], /* out */ In[4]);
-	LHOperator2(In[0], In[1], /* out */ In[0], /* out */ In[1]);
-	LHOperator2(In[3], In[4], /* out */ In[3], /* out */ In[4]);
-	return Median3(In[1], In[2], In[3]);
-}
-
-tsr_half3 TransformToComparisonColorSpace(tsr_half3 Color)
-{
-	return Color;
-	//return log2(Color);
-	//return Color *HdrWeight4(Color); // *float3(HdrWeightY(Color.r), HdrWeightY(Color.g), HdrWeightY(Color.b));
-}
-
-
 //------------------------------------------------------- ENTRY POINT
 
 [numthreads(TILE_SIZE, TILE_SIZE, 1)]
@@ -136,9 +81,7 @@ void MainCS(
 	// Sample current frame informations.
 	float2 EncodedScreenVelocity;
 	float DeviceZ;
-	#if CONFIG_UPDATE_SUBPIXEL_DETAIL
-		tsr_subpixel_payload ParallaxFactorBits;
-	#endif
+	tsr_subpixel_payload ParallaxFactorBits;
 	tsr_half ParallaxFactor;
 	tsr_half RejectHighFrequency = tsr_half(0.0);
 	ISOLATE
@@ -158,9 +101,7 @@ void MainCS(
 
 			DeviceZ = ClosestDepthTexture[InputPixelPos];
 			ParallaxFactor = UncompressParallaxFactor(RawParallaxFactor & tsr_subpixel_payload(SUB_PIXEL_PARALLAX_FACTOR_BIT_MASK));
-			#if CONFIG_UPDATE_SUBPIXEL_DETAIL
-				ParallaxFactorBits = RawParallaxFactor & tsr_subpixel_payload(SUB_PIXEL_PARALLAX_FACTOR_BIT_MASK);
-			#endif
+			ParallaxFactorBits = RawParallaxFactor & tsr_subpixel_payload(SUB_PIXEL_PARALLAX_FACTOR_BIT_MASK);
 
 			#if CONFIG_VELOCITY_SAMPLES == 1
 				EncodedScreenVelocity = DilatedVelocityTexture[InputPixelPos];
@@ -319,19 +260,9 @@ void MainCS(
 	// Sample the previous frame information
 	uint HistoryUseCountSamples[4];
 	uint HistoryClosestDeviceZSamples[4];
-
-	#if DIM_HIGH_FREQUENCY_ONLY && CONFIG_MEDIAN_DECIMATE
-		tsr_halfC DecimatedColorSamples[5];
-		tsr_half2 DecimatedMetadataSamples[5];
-	#elif DIM_HIGH_FREQUENCY_ONLY
-		tsr_halfC DecimatedColorSamples[1];
-		tsr_half2 DecimatedMetadataSamples[1];
-	#else
-		tsr_halfC DecimatedColorSamples[1];
-	#endif
-		
-	FSubpixelNeighborhood SubpixelNeigborhood;
 	
+	tsr_halfC ReprojectedHistoryGuideColor;
+	FSubpixelNeighborhood SubpixelNeigborhood;
 	bool bIsOffScreen;
 
 	// Issues all the history texture fetches
@@ -371,6 +302,33 @@ void MainCS(
 			}
 		}
 			
+		{
+			float2 OutputScreenPos = ApplyScreenTransform(float2(InputPixelPos), InputPixelPosToReprojectScreenPos);
+	
+			float2 PrevOutputScreenPos = OutputScreenPos - DecodeVelocityFromTexture(float4(EncodedScreenVelocity, 0.0, 0.0)).xy;
+			bIsOffScreen = (
+				(PrevOutputScreenPos.x < -1.0) ||
+				(PrevOutputScreenPos.y < -1.0) ||
+				(PrevOutputScreenPos.x > 1.0) ||
+				(PrevOutputScreenPos.y > 1.0));
+				
+			float2 KernelHistoryBufferUV = ApplyScreenTransform(PrevOutputScreenPos, ScreenPosToPrevHistoryGuideBufferUV);
+			FCatmullRomSamples Samples = GetBicubic2DCatmullRomSamples(KernelHistoryBufferUV, PrevGuideInfo_Extent, PrevGuideInfo_ExtentInverse);
+
+			ReprojectedHistoryGuideColor = tsr_half(0.0);
+			for (uint i = 0; i < Samples.Count; i++)
+			{
+				float2 SampleUV = Samples.UV[i];
+				SampleUV = clamp(
+					SampleUV,
+					PrevGuideInfo_UVViewportBilinearMin,
+					PrevGuideInfo_UVViewportBilinearMax);
+
+				ReprojectedHistoryGuideColor += PrevHistoryGuide.SampleLevel(GlobalBilinearClampedSampler, SampleUV, 0) * tsr_half(Samples.Weight[i]);
+			}
+			ReprojectedHistoryGuideColor *= tsr_half(Samples.FinalMultiplier);
+		}
+
 		{
 			// Pixel coordinate of the input pixel in the predictor viewport.
 			float2 InputPPCk = float2(InputPixelPos - InputInfo_ViewportMin) + 0.5f - InputJitter;
@@ -393,84 +351,12 @@ void MainCS(
 				PrevHistoryInfo_UVViewportBilinearMin,
 				PrevHistoryInfo_UVViewportBilinearMax);
 		
-			#if DIM_HIGH_FREQUENCY_ONLY && CONFIG_MEDIAN_DECIMATE
-			{
-				const float kTopSecretSauce = 2.0;
-
-				float2 PrevUV[5];
-				PrevUV[0] = KernelHistoryBufferUV;
-				PrevUV[1] = KernelHistoryBufferUV + (kTopSecretSauce * float2(+1.0, +1.0)) * PrevHistoryInfo_ExtentInverse;
-				PrevUV[2] = KernelHistoryBufferUV + (kTopSecretSauce * float2(-1.0, +1.0)) * PrevHistoryInfo_ExtentInverse;
-				PrevUV[3] = KernelHistoryBufferUV + (kTopSecretSauce * float2(+1.0, -1.0)) * PrevHistoryInfo_ExtentInverse;
-				PrevUV[4] = KernelHistoryBufferUV + (kTopSecretSauce * float2(-1.0, -1.0)) * PrevHistoryInfo_ExtentInverse;
-
-				tsr_halfC PrevDecimatedColorArray[5];
-			
-				UNROLL_N(5)
-				for (uint i = 0; i < 5; i++)
-				{
-					float2 UV = KernelHistoryBufferUV = clamp(
-						PrevUV[i],
-						PrevHistoryInfo_UVViewportBilinearMin,
-						PrevHistoryInfo_UVViewportBilinearMax);
-					#if 1 // saves 0.120ms
-						DecimatedMetadataSamples[i] = tsr_half(1.0).xx;
-						DecimatedColorSamples[i] = PrevHistory_HighFrequencyHalfRes[uint2(UV * PrevHistoryInfo_Extent * 0.5)];
-					#elif 1
-						DecimatedMetadataSamples[i] = PrevHistory_Metadata[uint2(UV * PrevHistoryInfo_Extent)];
-						DecimatedColorSamples[i] = PrevHistory_HighFrequency[uint2(UV * PrevHistoryInfo_Extent)];
-					#else
-						DecimatedMetadataSamples[i] = PrevHistory_Metadata.SampleLevel(GlobalBilinearClampedSampler, UV, 0);
-						DecimatedColorSamples[i] = PrevHistory_HighFrequency.SampleLevel(GlobalBilinearClampedSampler, UV, 0);
-					#endif
-				}
-			}
-			#elif DIM_HIGH_FREQUENCY_ONLY
-			{
-				DecimatedMetadataSamples[0] = PrevHistory_Metadata.SampleLevel(GlobalBilinearClampedSampler, KernelHistoryBufferUV, 0);
-				DecimatedColorSamples[0] = PrevHistory_HighFrequency.SampleLevel(GlobalBilinearClampedSampler, KernelHistoryBufferUV, 0);
-			}
-			#else
-			{
-				DecimatedColorSamples[0] = PrevHistory_LowFrequency.SampleLevel(GlobalBilinearClampedSampler, KernelHistoryBufferUV, 0);
-			}
-			#endif
-		
-			#if CONFIG_UPDATE_SUBPIXEL_DETAIL
-				float2 SubpixelBufferUV = ApplyScreenTransform(PrevScreenPos, ScreenPosToPrevSubpixelDetails);
-			#else
-				float2 SubpixelBufferUV = KernelHistoryBufferUV;
-			#endif
+			float2 SubpixelBufferUV = ApplyScreenTransform(PrevScreenPos, ScreenPosToPrevSubpixelDetails);
 
-			SubpixelNeigborhood = GatherPrevSubpixelNeighborhood(PrevHistory_SubpixelDetails, SubpixelBufferUV);
+			SubpixelNeigborhood = GatherPrevSubpixelNeighborhood(PrevHistorySubpixelDetails, SubpixelBufferUV);
 		}
 	}
 	
-	// Decimate the history color for comparison with input immediatly to quickly reduce VGPRs.
-	tsr_halfC DecimatedColor;
-	ISOLATE
-	{
-		#if DIM_HIGH_FREQUENCY_ONLY && CONFIG_MEDIAN_DECIMATE
-		{
-			UNROLL_N(5)
-			for (uint i = 0; i < 5; i++)
-			{
-				DecimatedColorSamples[i] *= SafeRcp(DecimatedMetadataSamples[i].x);
-			}
-
-			DecimatedColor = Median5(DecimatedColorSamples);
-		}
-		#elif DIM_HIGH_FREQUENCY_ONLY
-		{
-			DecimatedColor = DecimatedColorSamples[0] * SafeRcp(DecimatedMetadataSamples[0].x);
-		}
-		#else
-		{
-			DecimatedColor = DecimatedColorSamples[0];
-		}
-		#endif
-	}
-		
 
 	// Process the history samples
 	tsr_half ParallaxRejectionMask = 0.0;
@@ -622,7 +508,6 @@ void MainCS(
 		}
 		
 		// Update subpixel detail
-		#if CONFIG_UPDATE_SUBPIXEL_DETAIL
 		{
 			FinalSubpixelDetails = PrevSubpixelDetails;
 
@@ -650,59 +535,7 @@ void MainCS(
 
 			FinalSubpixelDetails = UpdatedFinalSubpixelDetails;
 		}
-		#endif // CONFIG_UPDATE_SUBPIXEL_DETAIL
-	}
-	
-	// Output half res.
-	#if CONFIG_OUTPUT_HALF_RES
-	{
-		#if CONFIG_SCENE_COLOR_APLHA
-			SharedHistoryColor[GroupThreadIndex] = DecimatedColor;
-			SharedMask[GroupThreadIndex] = ParallaxRejectionMask;
-		#else
-			SharedHistoryColorAndMask[GroupThreadIndex] = tsr_half4(DecimatedColor, ParallaxRejectionMask);
-		#endif
-
-		GroupMemoryBarrierWithGroupSync();
-
-		bool bCheckerboard = (((InputPixelPos.x ^ InputPixelPos.y) >> 1) & 0x1) != 0;
-
-		tsr_halfC HalfResDecimatedColor = DecimatedColor;
-		tsr_half HalfResParallaxRejectionMask = ParallaxRejectionMask;
-
-		UNROLL
-		for (uint i = 1; i < 4; i++)
-		{
-			uint ButterflySwap = i;
-			
-			#if CONFIG_SCENE_COLOR_APLHA
-				HalfResDecimatedColor += SharedHistoryColor[GroupThreadIndex ^ ButterflySwap];
-				HalfResParallaxRejectionMask += SharedMask[GroupThreadIndex ^ ButterflySwap];
-			#else
-				HalfResDecimatedColor += SharedHistoryColorAndMask[GroupThreadIndex ^ ButterflySwap].rgb;
-				HalfResParallaxRejectionMask += SharedHistoryColorAndMask[GroupThreadIndex ^ ButterflySwap].a;
-			#endif
-		}
-
-		HalfResDecimatedColor *= tsr_half(0.25);
-		HalfResParallaxRejectionMask *= tsr_half(0.25);
-
-		#if CONFIG_ENABLE_STOCASTIC_QUANTIZATION
-		{
-			uint2 Random = Rand3DPCG16(int3((InputPixelPos - InputInfo_ViewportMin) / 2, View.StateFrameIndexMod8)).xy;
-			tsr_half E = tsr_half(Hammersley16(0, 1, Random).x);
-
-			HalfResDecimatedColor = QuantizeForFloatRenderTarget(HalfResDecimatedColor, E, OutputQuantizationError);
-		}
-		#endif
-
-		if (all(InputPixelPos < InputInfo_ViewportMax) && all(((InputPixelPos - InputInfo_ViewportMin) % 2) == 0))
-		{
-			HalfResPredictionSceneColorOutput[InputPixelPos / 2] = HalfResDecimatedColor;
-			HalfResParallaxRejectionMaskOutput[InputPixelPos / 2] = HalfResParallaxRejectionMask;
-		}
 	}
-	#endif
 
 	// Output full res
 	ISOLATE
@@ -711,29 +544,22 @@ void MainCS(
 			Map8x8Tile2x2Lane(GroupThreadIndex) +
 			(tsr_ushort2(InputInfo_ViewportMin) + tsr_ushort2(GroupId) * tsr_ushort2(TILE_SIZE, TILE_SIZE)));
 			
-		#if CONFIG_ENABLE_STOCASTIC_QUANTIZATION
+		#if CONFIG_ENABLE_STOCASTIC_QUANTIZATION && 1
 		{
 			uint2 Random = Rand3DPCG16(int3(InputPixelPos - InputInfo_ViewportMin, View.StateFrameIndexMod8)).xy;
 			tsr_half E = tsr_half(Hammersley16(0, 1, Random).x);
 
-			DecimatedColor = QuantizeForFloatRenderTarget(DecimatedColor, E, OutputQuantizationError);
+			ReprojectedHistoryGuideColor = QuantizeForUNormRenderTarget(ReprojectedHistoryGuideColor, E, HistoryGuideQuantizationError);
 		}
 		#endif
 		
 		tsr_short2 OutputPixelPos = InvalidateOutputPixelPos(InputPixelPos, InputInfo_ViewportMax);
 
-		#if !CONFIG_OUTPUT_HALF_RES
-		{
-			PredictionSceneColorOutput[OutputPixelPos] = DecimatedColor;
-		}
-		#endif
-		
+		ReprojectedHistoryGuideOutput[OutputPixelPos] = ReprojectedHistoryGuideColor;
 		HoleFilledVelocityOutput[OutputPixelPos] = EncodedScreenVelocity;
 		HoleFilledVelocityMaskOutput[InputPixelPos] = tsr_half2(0.0, RejectHighFrequency);
 		ParallaxRejectionMaskOutput[OutputPixelPos] = ParallaxRejectionMask;
-		#if CONFIG_UPDATE_SUBPIXEL_DETAIL
-			HistorySubpixelDetailsOutput[OutputPixelPos] = FinalSubpixelDetails;
-		#endif
+		HistorySubpixelDetailsOutput[OutputPixelPos] = FinalSubpixelDetails;
 
 		#if DEBUG_OUTPUT
 		{
diff --git a/Engine/Shaders/Private/TemporalSuperResolution/TSRDilateRejection.usf b/Engine/Shaders/Private/TemporalSuperResolution/TSRDilateRejection.usf
--- a/Engine/Shaders/Private/TemporalSuperResolution/TSRDilateRejection.usf
+++ b/Engine/Shaders/Private/TemporalSuperResolution/TSRDilateRejection.usf
@@ -31,7 +31,7 @@ void MainCS(
 	// Fetch center of the kernel
 	tsr_half HistoryRejectionInfo = HistoryRejectionTexture[RejectionPixelPos];
 
-	const int KernelSize = 2;
+	const int KernelSize = 1;
 
 	// Fetch neighborhood of the kernel
 	UNROLL
diff --git a/Engine/Shaders/Private/TemporalSuperResolution/TSRFilterFrequencies.usf b/Engine/Shaders/Private/TemporalSuperResolution/TSRFilterFrequencies.usf
deleted file mode 100644
diff --git a/Engine/Shaders/Private/TemporalSuperResolution/TSRRejectShading.usf b/Engine/Shaders/Private/TemporalSuperResolution/TSRRejectShading.usf
new file mode 100644
--- /dev/null
+++ b/Engine/Shaders/Private/TemporalSuperResolution/TSRRejectShading.usf
@@ -0,0 +1,448 @@
+// Copyright Epic Games, Inc. All Rights Reserved.
+
+//------------------------------------------------------- CONFIG WAVE SIZE
+
+#define WAVE_SIZE 32
+#if !DIM_USE_WAVE_OPS
+	#define WAVE_BROADCAST_GROUPSIZE (WAVE_SIZE)
+#endif
+
+//------------------------------------------------------- INCLUDES
+
+#include "TSRCommon.ush"
+#include "/Engine/Public/LaneVectorization.ush"
+
+
+//------------------------------------------------------- CONFIG
+
+#define DEBUG_ARRAY_SIZE 8
+
+/** Size of the entire group loaded into memory including overscan. */
+#define GROUP_TILE_SIZE 16
+
+/** Size of the overscan. */
+#define TILE_OVERSCAN_PADDING 2
+
+/** Size of of the output tile. */
+#define TILE_SIZE (GROUP_TILE_SIZE - 2 * TILE_OVERSCAN_PADDING)
+
+/** Whether should remove firefliers from input and history before comparison to improve stability. */
+#define CONFIG_FIREFLY_REMOVAL 1
+
+
+//------------------------------------------------------- DERIVED
+
+/** Pixel coordinate stride between each lane in the GROUP_TILE_SIZE. */
+#define LANE_STRIDE_X 4
+#define LANE_STRIDE_Y 2
+
+/** Number of lane the GROUP_TILE_SIZE. */
+#define LANE_COUNT_X (GROUP_TILE_SIZE / LANE_STRIDE_X)
+#define LANE_COUNT_Y (GROUP_TILE_SIZE / LANE_STRIDE_Y)
+
+/** Size of the SIMD per lane that also number of input pixel loaded into each individual lanes. */
+#define SIMD_SIZE (LANE_STRIDE_X * LANE_STRIDE_Y)
+
+
+//------------------------------------------------------- TYPEDEFS
+
+#define tsr_simd_half  TLaneVector<tsr_half, 1, SIMD_SIZE>
+#define tsr_simd_halfC TLaneVector<tsr_half, CONFIG_CHANNEL_COUNT, SIMD_SIZE>
+#define tsr_simd_bool  TLaneVector<bool, 1, SIMD_SIZE>
+
+
+//------------------------------------------------------- PARAMETERS
+
+float3 HistoryGuideQuantizationError;
+
+Texture2D<tsr_halfC> InputTexture;
+Texture2D<tsr_halfC> ReprojectedHistoryGuideTexture;
+Texture2D<tsr_half> ParallaxRejectionMaskTexture;
+
+RWTexture2D<tsr_halfC> HistoryGuideOutput;
+RWTexture2D<tsr_half> HistoryRejectionOutput;
+RWTexture2D<tsr_half> InputSceneColorLdrLumaOutput;
+
+
+//------------------------------------------------------- GENERAL TEXEL SWIZZLING & NEIGHBOR ACCESS
+
+CALL_SITE_DEBUGLOC
+tsr_short2 GetSimdIndexPixelCoordinateInLane(const uint SimdIndex)
+{
+	return tsr_short2(SimdIndex % uint(LANE_STRIDE_X), SimdIndex / uint(LANE_STRIDE_X));
+}
+
+CALL_SITE_DEBUGLOC
+tsr_short2 GetLaneSimdPixelOffset(const uint LaneIndex, const uint SimdIndex)
+{
+	return (
+		tsr_short2(LaneIndex % uint(LANE_COUNT_X), LaneIndex / uint(LANE_COUNT_X)) * tsr_short2(LANE_STRIDE_X, LANE_STRIDE_Y) +
+		GetSimdIndexPixelCoordinateInLane(SimdIndex)
+	);
+}
+
+CALL_SITE_DEBUGLOC
+template<typename ScalarType, uint VectorSize, uint SimdSize>
+TLaneVector<ScalarType, VectorSize, SimdSize> AccessNeighborTexel(TLaneVector<ScalarType, VectorSize, SimdSize> Input, const tsr_short2 Offset, uint GroupThreadIndex)
+{
+	TLaneVector<ScalarType, VectorSize, SimdSize> Return;
+	
+	UNROLL_N(SIMD_SIZE)
+	for (uint SimdIndex = 0; SimdIndex < SIMD_SIZE; SimdIndex++)
+	{
+		const uint2 SimdPixelCoord = uint2(GetSimdIndexPixelCoordinateInLane(SimdIndex));
+
+		// Compute the simd index of the neighbor.
+		const uint2 NeigborSimdPixelCoord = (SimdPixelCoord + uint2(Offset)) % uint2(LANE_STRIDE_X, LANE_STRIDE_Y);
+		const uint NeigborSimdIndex = dot(NeigborSimdPixelCoord, uint2(1, LANE_STRIDE_X));
+
+		// Compute the lane rotaton of the neighbor.
+		const bool2 bNeedsLaneRotation = (SimdPixelCoord + uint2(Offset)) >= uint2(LANE_STRIDE_X, LANE_STRIDE_Y);
+		const int LaneRotation = dot(select(bNeedsLaneRotation, int2(Offset), 0), int2(1, LANE_COUNT_X));
+		
+		// Access the lement.
+		vector<ScalarType, VectorSize> SimdElement = Input.GetElement(NeigborSimdIndex);
+		vector<ScalarType, VectorSize> ReturnSimdElement = SimdElement;
+		if (LaneRotation != 0)
+		{
+			// Only uses wave intrinsics on platforms that provides RDNA's ds_swizzle because we know the wave size of this platforms. 
+			#if DIM_USE_WAVE_OPS
+				#if !PLATFORM_SUPPORTS_WAVE_ROTATE
+					//#error InitWaveRotateLaneGroup() isn't supported.
+				#endif
+				const FWaveBroadcastSettings BroadcastSettings = InitWaveRotateLaneGroup(/* LaneGroupSize = */ WAVE_SIZE, LaneRotation);
+			#else
+				const FWaveBroadcastSettings BroadcastSettings = ConvertWaveBroadcastToLDS(InitWaveRotateLaneGroup(/* LaneGroupSize = */ WAVE_SIZE, LaneRotation), GroupThreadIndex);
+			#endif
+			ReturnSimdElement = WaveBroadcast(BroadcastSettings, SimdElement);
+		}
+
+		Return.SetElement(SimdIndex, ReturnSimdElement);
+	}
+
+	return Return;
+}
+
+
+//------------------------------------------------------- CONVOLUTION ACCUMULATORS
+
+template<typename Type>
+struct TAccumulatorAvg
+{
+	Type Value;
+	tsr_simd_half Weight;
+	
+	CALL_SITE_DEBUGLOC
+	void Start(Type CenterSample)
+	{
+		Value = CenterSample;
+		Weight = tsr_simd_half::Const(1.0);
+	}
+	
+	CALL_SITE_DEBUGLOC
+	static Type PreprocessNeighborSample(Type Sample, tsr_simd_half InWeight)
+	{
+		return Sample;
+	}
+	
+	CALL_SITE_DEBUGLOC
+	void Accumulate(const tsr_short2 Offset, Type Sample, tsr_simd_half InWeight)
+	{
+		const tsr_half SampleSpatialWeight = tsr_half(rcp(abs(tsr_half(Offset.x)) + 1.0) * rcp(abs(tsr_half(Offset.y)) + 1.0));
+		
+		Value = Value + Sample * Type::Vectorize(InWeight * SampleSpatialWeight);
+		Weight = Weight + InWeight * tsr_simd_half::Const(SampleSpatialWeight);
+	}
+	
+	CALL_SITE_DEBUGLOC
+	void Finish(const uint SampleCount)
+	{
+		Value = Value * Type::Vectorize(rcp(Weight));
+	}
+};
+
+template<typename Type>
+struct TAccumulatorMin
+{
+	Type Value;
+	
+	CALL_SITE_DEBUGLOC
+	void Start(Type CenterSample)
+	{
+		Value = CenterSample;
+	}
+	
+	CALL_SITE_DEBUGLOC
+	static Type PreprocessNeighborSample(Type Sample, tsr_simd_half InWeight)
+	{
+		return Sample;
+	}
+	
+	CALL_SITE_DEBUGLOC
+	void Accumulate(const tsr_short2 Offset, Type Sample, tsr_simd_half Weight)
+	{
+		Value = min(Value, Sample);
+	}
+	
+	CALL_SITE_DEBUGLOC
+	void Finish(const uint SampleCount)
+	{
+		// NOP
+	}
+};
+
+template<typename Type>
+struct TAccumulatorMax
+{
+	Type Value;
+	
+	CALL_SITE_DEBUGLOC
+	void Start(Type CenterSample)
+	{
+		Value = CenterSample;
+	}
+	
+	CALL_SITE_DEBUGLOC
+	static Type PreprocessNeighborSample(Type Sample, tsr_simd_half InWeight)
+	{
+		return Sample;
+	}
+	
+	CALL_SITE_DEBUGLOC
+	void Accumulate(const tsr_short2 Offset, Type Sample, tsr_simd_half Weight)
+	{
+		Value = max(Value, Sample);
+	}
+	
+	CALL_SITE_DEBUGLOC
+	void Finish(const uint SampleCount)
+	{
+		// NOP
+	}
+};
+
+
+//------------------------------------------------------- CONVOLUTIONS
+
+CALL_SITE_DEBUGLOC
+template<class FAccumulator, typename FSampleType>
+FAccumulator Convole3x1(FSampleType Color, tsr_simd_half Weight, const tsr_short2 Offset, uint GroupThreadIndex)
+{
+	FAccumulator Accumulator;
+	Accumulator.Start(Color);
+	
+	FSampleType PreprocessedNeighborColor = FAccumulator::PreprocessNeighborSample(Color, Weight);
+
+	FSampleType Sample0 = AccessNeighborTexel(PreprocessedNeighborColor, +Offset, GroupThreadIndex);
+	tsr_simd_half Weight0 = AccessNeighborTexel(Weight, +Offset, GroupThreadIndex);
+	Accumulator.Accumulate(+Offset, Sample0, Weight0);
+	
+	const tsr_short2 OffsetN = tsr_short2(-int2(Offset)); // work arround a shader compiler bug
+	FSampleType Sample1 = AccessNeighborTexel(PreprocessedNeighborColor, OffsetN, GroupThreadIndex);
+	tsr_simd_half Weight1 = AccessNeighborTexel(Weight, OffsetN, GroupThreadIndex);
+	Accumulator.Accumulate(OffsetN, Sample1, Weight1);
+
+	Accumulator.Finish(3);
+	return Accumulator;
+}
+
+CALL_SITE_DEBUGLOC
+template<class FAccumulator, typename FSampleType>
+FSampleType Convolve3x3HV(FSampleType Color, tsr_simd_half Weight, uint GroupThreadIndex)
+{
+	FSampleType HorizontalColor = Convole3x1<FAccumulator, FSampleType>(Color, Weight, tsr_short2(1, 0), GroupThreadIndex).Value;
+	return Convole3x1<FAccumulator, FSampleType>(HorizontalColor, Weight, tsr_short2(0, 1), GroupThreadIndex).Value;
+}
+
+
+//------------------------------------------------------- GENERAL OPERATIONS
+
+tsr_simd_halfC ClampFireFliersWithGuide(tsr_simd_halfC ToClamp, tsr_simd_halfC Guide, tsr_simd_half Weight, uint GroupThreadIndex)
+{
+	tsr_simd_halfC ToClampMin = Convolve3x3HV<TAccumulatorMin<tsr_simd_halfC>, tsr_simd_halfC>(ToClamp, Weight, GroupThreadIndex);
+	tsr_simd_halfC ToClampMax = Convolve3x3HV<TAccumulatorMax<tsr_simd_halfC>, tsr_simd_halfC>(ToClamp, Weight, GroupThreadIndex);
+
+	tsr_simd_halfC ClampedGuide = clamp(Guide, ToClampMin, ToClampMax);
+	
+	tsr_simd_halfC ClampedGuideMin = Convolve3x3HV<TAccumulatorMin<tsr_simd_halfC>, tsr_simd_halfC>(ClampedGuide, Weight, GroupThreadIndex);
+	tsr_simd_halfC ClampedGuideMax = Convolve3x3HV<TAccumulatorMax<tsr_simd_halfC>, tsr_simd_halfC>(ClampedGuide, Weight, GroupThreadIndex);
+	
+	return clamp(ToClamp, ClampedGuideMin, ClampedGuideMax);
+}
+
+
+//------------------------------------------------------- ENTRY POINT
+
+[numthreads(WAVE_SIZE, 1, 1)]
+void MainCS(
+	uint2 GroupId : SV_GroupID,
+	uint GroupThreadIndex : SV_GroupIndex)
+{
+	tsr_simd_halfC Debug[DEBUG_ARRAY_SIZE];
+	UNROLL_N(DEBUG_ARRAY_SIZE)
+	for (uint DebugId = 0; DebugId < DEBUG_ARRAY_SIZE; DebugId++)
+	{
+		Debug[DebugId].SetAllElements(0.0);
+	}
+
+	const uint LaneIndex = GroupThreadIndex;
+	const tsr_short2 GroupPixelOffset = (
+		tsr_short2(InputPixelPosMin) +
+		tsr_short2(GroupId) * tsr_short2(TILE_SIZE, TILE_SIZE) +
+		tsr_short(-TILE_OVERSCAN_PADDING).xx);
+		
+	// Load
+	tsr_simd_halfC OriginalInput;
+	tsr_simd_halfC OriginalHistory;
+	tsr_simd_half bIsValid;
+	ISOLATE
+	{
+		UNROLL_N(SIMD_SIZE)
+		for (uint SimdIndex = 0; SimdIndex < SIMD_SIZE; SimdIndex++)
+		{
+			tsr_short2 ClampedFetchPixelPos = GroupPixelOffset + GetLaneSimdPixelOffset(LaneIndex, SimdIndex);
+			ClampedFetchPixelPos = ClampPixelOffset(ClampedFetchPixelPos, InputPixelPosMin, InputPixelPosMax);
+
+			OriginalInput.SetElement(SimdIndex, InputTexture[ClampedFetchPixelPos]);
+			OriginalHistory.SetElement(SimdIndex, ReprojectedHistoryGuideTexture[ClampedFetchPixelPos]);
+			bIsValid.SetElement(SimdIndex, ParallaxRejectionMaskTexture[ClampedFetchPixelPos] > tsr_half(PARALLAX_REJECTION_MASK_THRESHOLD) ? tsr_half(1.0) : tsr_half(0.0));
+		}
+	}
+		
+	// Compute AA Luma
+	tsr_simd_half CenterLumaLDR;
+	{
+		const tsr_half SpatialAAExposure = tsr_half(0.5);
+		UNROLL_N(SIMD_SIZE)
+		for (uint SimdIndex = 0; SimdIndex < SIMD_SIZE; SimdIndex++)
+		{
+			tsr_half CenterLuma = dot(OriginalInput.GetElement(SimdIndex).rgb, tsr_half3(0.299f, 0.587f, 0.114f));
+			CenterLumaLDR.SetElement(SimdIndex, CenterLuma / (SpatialAAExposure + CenterLuma));
+		}
+	}
+	
+	// Convert input to LDR
+	OriginalInput = OriginalInput * rcp(OriginalInput + tsr_simd_halfC::Const(1.0));
+	
+	tsr_simd_halfC ClampedInput = OriginalInput;
+	tsr_simd_halfC ClampedHistory = OriginalHistory;
+	#if CONFIG_FIREFLY_REMOVAL
+	{
+		ClampedInput   = ClampFireFliersWithGuide(OriginalInput  , /* Guide = */ OriginalHistory, bIsValid, GroupThreadIndex);
+		ClampedHistory = ClampFireFliersWithGuide(OriginalHistory, /* Guide = */ OriginalInput  , bIsValid, GroupThreadIndex);
+	}
+	#endif
+
+	// Filter clamped input and history to increase quantization of the rejection detection
+	tsr_simd_halfC FilteredInput   = Convolve3x3HV<TAccumulatorAvg<tsr_simd_halfC>, tsr_simd_halfC>(ClampedInput  , bIsValid, GroupThreadIndex);
+	tsr_simd_halfC FilteredHistory = Convolve3x3HV<TAccumulatorAvg<tsr_simd_halfC>, tsr_simd_halfC>(ClampedHistory, bIsValid, GroupThreadIndex);
+	
+	Debug[0] = OriginalInput;
+	Debug[1] = AccessNeighborTexel(OriginalInput, tsr_short2(1, 0), GroupThreadIndex);
+	Debug[2] = FilteredInput;
+
+	// Measure rejection
+	tsr_simd_half Rejection;
+	Rejection.SetAllElements(tsr_half(1.0));
+	#if 1
+	{
+		tsr_simd_halfC BackbufferError = tsr_simd_halfC::Const(MeasureBackbufferLDRQuantizationError());
+
+		tsr_simd_halfC FilteredBoxMin = Convolve3x3HV<TAccumulatorMin<tsr_simd_halfC>, tsr_simd_halfC>(FilteredInput, bIsValid, GroupThreadIndex) - BackbufferError;
+		tsr_simd_halfC FilteredBoxMax = Convolve3x3HV<TAccumulatorMax<tsr_simd_halfC>, tsr_simd_halfC>(FilteredInput, bIsValid, GroupThreadIndex) + BackbufferError;
+		tsr_simd_halfC ClampedFilteredHistory = clamp(FilteredHistory, FilteredBoxMin, FilteredBoxMax);
+		
+		tsr_simd_halfC OriginalBoxMin = Convolve3x3HV<TAccumulatorMin<tsr_simd_halfC>, tsr_simd_halfC>(OriginalInput, bIsValid, GroupThreadIndex);
+		tsr_simd_halfC OriginalBoxMax = Convolve3x3HV<TAccumulatorMax<tsr_simd_halfC>, tsr_simd_halfC>(OriginalInput, bIsValid, GroupThreadIndex);
+		tsr_simd_halfC BoxSize = (OriginalBoxMax - OriginalBoxMin + (BackbufferError + BackbufferError)) * rcp(tsr_half(1.0 + 4 * 0.5 + 4 * 0.25));
+		
+		UNROLL_N(SIMD_SIZE)
+		for (uint SimdIndex = 0; SimdIndex < SIMD_SIZE; SimdIndex++)
+		{
+			Rejection.SetElement(SimdIndex, MeasureRejectionFactor(
+				FilteredHistory.GetElement(SimdIndex), ClampedFilteredHistory.GetElement(SimdIndex),
+				FilteredInput.GetElement(SimdIndex), tsr_half(0.0), BoxSize.GetElement(SimdIndex),
+				MeasureBackbufferLDRQuantizationError()));
+		}
+		
+		Rejection = saturate(Rejection + tsr_simd_half::Const(1.0) - bIsValid);
+	}
+	#endif
+	
+	// Update the guide for next frame.
+	tsr_simd_halfC FinalGuide;
+	{
+		const tsr_half CurrentFrameSampleCount = tsr_half(1.0);
+		const tsr_half PrevLowFrequencySampleCount = tsr_half(24.0);
+		
+		// Reject the history based on the low frequencies.
+		tsr_simd_half PrevLowFrequencyRejectionFactor;
+		#if 1
+		{
+			PrevLowFrequencyRejectionFactor = saturate(Rejection * ((CurrentFrameSampleCount + PrevLowFrequencySampleCount) * rcp(PrevLowFrequencySampleCount)));
+		}
+		#endif
+		
+		tsr_simd_half BlendFinal = rcp(PrevLowFrequencyRejectionFactor * PrevLowFrequencySampleCount + CurrentFrameSampleCount) * CurrentFrameSampleCount;
+		
+		// Completely discard the history
+		UNROLL_N(SIMD_SIZE)
+		for (uint SimdIndex = 0; SimdIndex < SIMD_SIZE; SimdIndex++)
+		{
+			bool bIsOffScreen = false;
+			bool bIsDisoccluded = bIsValid.GetElement(SimdIndex) == tsr_half(0.0);
+			tsr_half TotalBlend = select(or(bIsOffScreen, bIsDisoccluded), tsr_half(1.0), saturate(tsr_half(1.0) - Rejection.GetElement(SimdIndex) * tsr_half(4.0)));
+
+			BlendFinal.SetElement(SimdIndex, max(BlendFinal.GetElement(SimdIndex), TotalBlend));
+		}
+		
+		FinalGuide = OriginalInput * tsr_simd_halfC::Vectorize(BlendFinal) - OriginalHistory * tsr_simd_halfC::Vectorize(BlendFinal - 1.0);
+	}
+
+	// Mark the top left corner of the tile for debugging purposes.
+	#if DEBUG_OUTPUT
+	if (GroupThreadIndex == 0)
+	{
+		UNROLL_N(DEBUG_ARRAY_SIZE)
+		for (uint DebugId = 0; DebugId < DEBUG_ARRAY_SIZE; DebugId++)
+		{
+			Debug[DebugId].SetElement(/* SimdIndex = */ 0, 1.0);
+		}
+	}
+	#endif
+
+	// Output data.
+	ISOLATE
+	{
+		UNROLL_N(SIMD_SIZE)
+		for (uint SimdIndex = 0; SimdIndex < SIMD_SIZE; SimdIndex++)
+		{
+			tsr_short2 LaneSimdPixelOffset = GetLaneSimdPixelOffset(LaneIndex, SimdIndex);
+			tsr_short2 OutputPixelPos = GroupPixelOffset + LaneSimdPixelOffset;
+			OutputPixelPos.x = select(all(and(LaneSimdPixelOffset >= TILE_OVERSCAN_PADDING, LaneSimdPixelOffset < TILE_SIZE + TILE_OVERSCAN_PADDING)), OutputPixelPos.x, ~0);
+		
+			tsr_halfC FinalGuideColor = FinalGuide.GetElement(SimdIndex);
+
+			#if CONFIG_ENABLE_STOCASTIC_QUANTIZATION
+			{
+				uint2 Random = Rand3DPCG16(int3(LaneSimdPixelOffset, View.StateFrameIndexMod8)).xy;
+				tsr_half E = tsr_half(Hammersley16(0, 1, Random).x);
+		
+				FinalGuideColor.rgb = QuantizeForUNormRenderTarget(FinalGuideColor.rgb, E.x, HistoryGuideQuantizationError);
+			}
+			#endif
+
+			HistoryGuideOutput[OutputPixelPos] = FinalGuideColor;
+			HistoryRejectionOutput[OutputPixelPos] = Rejection.GetElement(SimdIndex);
+			InputSceneColorLdrLumaOutput[OutputPixelPos] = CenterLumaLDR.GetElement(SimdIndex);
+			
+			#if DEBUG_OUTPUT
+			UNROLL_N(DEBUG_ARRAY_SIZE)
+			for (uint DebugId = 0; DebugId < DEBUG_ARRAY_SIZE; DebugId++)
+			{
+				DebugOutput[tsr_short3(OutputPixelPos, DebugId)] = float4(Debug[DebugId].GetElement(SimdIndex), 0.0);
+			}
+			#endif
+		}
+	}
+}
diff --git a/Engine/Shaders/Private/TemporalSuperResolution/TSRResolveHistory.usf b/Engine/Shaders/Private/TemporalSuperResolution/TSRResolveHistory.usf
--- a/Engine/Shaders/Private/TemporalSuperResolution/TSRResolveHistory.usf
+++ b/Engine/Shaders/Private/TemporalSuperResolution/TSRResolveHistory.usf
@@ -142,7 +142,7 @@ void MainCS(
 				tsr_half3 SampleTranslucencyColor = History_Translucency.SampleLevel(GlobalBilinearClampedSampler, SampleUV, 0);
 				tsr_half  SampleTranslucencyAlpha = History_TranslucencyAlpha.SampleLevel(GlobalBilinearClampedSampler, SampleUV, 0);
 				
-				tsr_half SampleHistoryValidity = DIM_HIGH_FREQUENCY_ONLY ? SampleMetadata.x : SampleMetadata.y;
+				tsr_half SampleHistoryValidity = SampleMetadata.x;
 				tsr_half InvSampleHistoryValidity = SafeRcp(SampleHistoryValidity);
 
 				#if 1
@@ -181,23 +181,7 @@ void MainCS(
 	}
 	
 	// Compute final output
-	tsr_halfC FinalOutputColor;
-	#if DIM_HIGH_FREQUENCY_ONLY
-	{
-		FinalOutputColor = HighFrequencyColor;
-	}
-	#else
-	{
-		tsr_half OutputBlend = ComputeLowFrequencyContribution(HistoryValidity);
-		
-		tsr_half LowFrequencyWeight = HdrWeight4(LowFrequencyColor);
-		tsr_half Moment1Weight = HdrWeight4(HighFrequencyColor);
-
-		tsr_half2 Weights = WeightedLerpFactors(LowFrequencyWeight, Moment1Weight, OutputBlend);
-		
-		FinalOutputColor = LowFrequencyColor * Weights[0] + HighFrequencyColor * Weights[1];
-	}
-	#endif
+	tsr_halfC FinalOutputColor = HighFrequencyColor;
 
 	tsr_short2 OutputPixelPosition = InvalidateOutputPixelPos(tsr_short2(OutputViewRectMin + DispatchThreadId), OutputViewRectMax);
 	
diff --git a/Engine/Shaders/Private/TemporalSuperResolution/TSRUpdateHistory.usf b/Engine/Shaders/Private/TemporalSuperResolution/TSRUpdateHistory.usf
--- a/Engine/Shaders/Private/TemporalSuperResolution/TSRUpdateHistory.usf
+++ b/Engine/Shaders/Private/TemporalSuperResolution/TSRUpdateHistory.usf
@@ -5,10 +5,6 @@
 
 //------------------------------------------------------- DEBUG
 
-#define DEBUG_FALLBACK_BLENDING 0
-
-#define DEBUG_LOW_FREQUENCY_REJECTION 0
-
 #define DEBUG_ARRAY_SIZE 8
 
 
@@ -16,22 +12,18 @@
 
 #if DIM_UPDATE_QUALITY == 0 // Low
 	#define CONFIG_SAMPLES 5
-	#define CONFIG_LOW_FREQUENCY_DRIFT_REJECTION 0
 	#define CONFIG_REJECTION_ANTI_ALIASING 0
 
 #elif DIM_UPDATE_QUALITY == 1 // Medium
 	#define CONFIG_SAMPLES 5
-	#define CONFIG_LOW_FREQUENCY_DRIFT_REJECTION 1
 	#define CONFIG_REJECTION_ANTI_ALIASING 1
 
 #elif DIM_UPDATE_QUALITY == 2 // High
 	#define CONFIG_SAMPLES 6
-	#define CONFIG_LOW_FREQUENCY_DRIFT_REJECTION 1
 	#define CONFIG_REJECTION_ANTI_ALIASING 1
 
 #elif DIM_UPDATE_QUALITY == 3 // Epic
 	#define CONFIG_SAMPLES 6
-	#define CONFIG_LOW_FREQUENCY_DRIFT_REJECTION 1
 	#define CONFIG_INTERPOLATE_VELOCITY 1
 	#define CONFIG_REJECTION_ANTI_ALIASING 1
 
@@ -39,18 +31,9 @@
 	#error Unknown history update quality
 #endif
 
-#if DIM_HIGH_FREQUENCY_ONLY == 1
+#if 1
 	#define CONFIG_HIGH_FREQUENCY_ONLY 1
 	#define CONFIG_HIGH_FREQUENCY_YCOCG 0
-	#define CONFIG_OUTPUT_HALFRES_HISTORY 1
-	#define CONFIG_UPDATE_SUBPIXEL_DETAIL 0
-
-#else
-	#define CONFIG_HIGH_FREQUENCY_ONLY 0
-	#define CONFIG_HIGH_FREQUENCY_YCOCG 1
-	#define CONFIG_OUTPUT_HALFRES_HISTORY 0
-	#define CONFIG_UPDATE_SUBPIXEL_DETAIL 1
-
 #endif
 
 #if 1
@@ -138,9 +121,7 @@ uint ResponsiveStencilMask;
 uint bGenerateOutputMip1;
 uint bHasSeparateTranslucency;
 
-Texture2D<tsr_halfC> PrevHistory_LowFrequency;
 Texture2D<tsr_halfC> PrevHistory_HighFrequency;
-Texture2D<tsr_halfC> PrevHistory_HighFrequencyHalfRes;
 Texture2D<tsr_half2> PrevHistory_Metadata;
 Texture2D<tsr_half3> PrevHistory_Translucency;
 Texture2D<tsr_half>  PrevHistory_TranslucencyAlpha;
@@ -149,12 +130,10 @@ Texture2D<tsr_subpixel_details> PrevHistory_SubpixelDetails;
 RWTexture2D<tsr_halfC> SceneColorOutputMip0;
 RWTexture2D<tsr_halfC> SceneColorOutputMip1;
 
-uint HistoryOutput_ArrayIndices_LowFrequency;
 uint HistoryOutput_ArrayIndices_HighFrequency;
 uint HistoryOutput_ArrayIndices_Translucency;
 
 RWTexture2DArray<tsr_halfC> HistoryOutput_ColorArray;
-RWTexture2D<tsr_halfC> HistoryOutput_HighFrequencyHalfRes;
 RWTexture2D<tsr_half2> HistoryOutput_Metadata;
 RWTexture2D<tsr_half>  HistoryOutput_TranslucencyAlpha;
 RWTexture2D<tsr_subpixel_details> HistoryOutput_SubpixelDetails;
@@ -182,9 +161,6 @@ groupshared tsr_half4 SharedArray9[GROUP_SHARED_SIZE];
 groupshared tsr_half4 SharedArray10[GROUP_SHARED_SIZE];
 groupshared tsr_half4 SharedArray11[GROUP_SHARED_SIZE];
 
-#if CONFIG_UPDATE_SUBPIXEL_DETAIL
-	groupshared tsr_subpixel_payload2 SharedParallaxFactorBits[GROUP_SHARED_SIZE];
-#endif
 
 //------------------------------------------------------- FUNCTIONS
 
@@ -364,9 +340,6 @@ void MainCS(
 	tsr_half2 OutputPixelVelocity = tsr_half(0.0).xx;
 	tsr_half2 NoiseFiltering = tsr_half(0.0).xx;
 	tsr_half2x2 HoleFilledVelocityMask = dpv_interleave_mono_registers(tsr_half(1.0).xx);
-	#if CONFIG_UPDATE_SUBPIXEL_DETAIL
-		tsr_subpixel_payload2 ParallaxFactorBits = tsr_subpixel_payload(0x0).xx;
-	#endif
 	bool2 bIsResponsiveAAPixel = false;
 	bool2 bIsOffScreen = false;
 	bool2 bIsDisoccluded = false;
@@ -518,9 +491,6 @@ void MainCS(
 				/* out */ bIsDisoccluded);
 
 			// Process input texel, forcing to pack register.
-			#if CONFIG_UPDATE_SUBPIXEL_DETAIL
-				ParallaxFactorBits    = dpv_force_interleave_registers_array(RawParallaxFactorBits) & tsr_subpixel_payload(SUB_PIXEL_PARALLAX_FACTOR_BIT_MASK);
-			#endif
 			LowFrequencyRejection = dpv_force_interleave_registers_array(RawLowFrequencyRejection);
 			SceneStencilRef       = dpv_force_interleave_registers_array(RawSceneStencilRef);
 			
@@ -587,30 +557,18 @@ void MainCS(
 	#endif
 	
 	// Reproject history
-	#if !CONFIG_HIGH_FREQUENCY_ONLY
-		tsr_halfCx2 PrevLowFrequencyColor;
-		tsr_half2 PrevLowFrequencyWeight;
-	#endif
-
 	tsr_halfCx2 PrevHighFrequencyColor;
-	#if CONFIG_HIGH_FREQUENCY_ONLY
-		tsr_half2 PrevHistoryWeight;
-	#endif
+	tsr_half2 PrevHistoryWeight;
 	tsr_half2 PrevHistoryValidity;
 	
 	tsr_halfCx2 PrevHistoryMommentMin;
 	tsr_halfCx2 PrevHistoryMommentMax;
 
 	tsr_half4x2 PrevHistoryTranslucency;
-	
-	#if CONFIG_UPDATE_SUBPIXEL_DETAIL
-		tsr_subpixel_details2 PrevSubpixelDetails;
-	#endif
 
 	PLATFORM_SPECIFIC_ISOLATE
 	{
 		// Reproject history
-		tsr_halfCx2 AccumulateLowFrequency = tsr_half(0);
 		tsr_halfCx2 AccumulateHighFrequency = tsr_half(0);
 		tsr_half2x2 AccumulateMetadata = tsr_half(0);
 		tsr_half4x2 AccumulateTranslucency = tsr_half(0);
@@ -629,9 +587,6 @@ void MainCS(
 			FCatmullRomSamples Samples1 = GetBicubic2DCatmullRomSamples(
 				dpv_hi(PrevHistoryBufferUV), PrevHistoryInfo_Extent, PrevHistoryInfo_ExtentInverse);
 			
-			#if !CONFIG_HIGH_FREQUENCY_ONLY
-				tsr_halfC RawLowFrequency[BICUBIC_CATMULL_ROM_SAMPLES][DPV_PIXEL_PER_LANE];
-			#endif
 			tsr_halfC RawHighFrequency[BICUBIC_CATMULL_ROM_SAMPLES][DPV_PIXEL_PER_LANE];
 			tsr_half2 RawMetadata[BICUBIC_CATMULL_ROM_SAMPLES][DPV_PIXEL_PER_LANE];
 			tsr_half3 RawTranslucencyColor[BICUBIC_CATMULL_ROM_SAMPLES][DPV_PIXEL_PER_LANE];
@@ -658,9 +613,6 @@ void MainCS(
 						float2 SampleUV = PixelId == 0 ? SampleUV0 : SampleUV1;
 						
 
-						#if !CONFIG_HIGH_FREQUENCY_ONLY
-							RawLowFrequency[i][PixelId]      = BilinearSampleColorHistory(PrevHistory_LowFrequency, SampleUV);
-						#endif
 						RawHighFrequency[i][PixelId]     = BilinearSampleColorHistory(PrevHistory_HighFrequency, SampleUV);
 						RawMetadata[i][PixelId]	         = PrevHistory_Metadata.SampleLevel(GlobalBilinearClampedSampler, SampleUV, 0);
 						
@@ -679,9 +631,6 @@ void MainCS(
 				UNROLL_N(BICUBIC_CATMULL_ROM_SAMPLES)
 				for (uint i = 0; i < BICUBIC_CATMULL_ROM_SAMPLES; i++)
 				{
-					#if !CONFIG_HIGH_FREQUENCY_ONLY
-						tsr_halfCx2 SampleLowFrequency = dpv_force_interleave_registers_array(RawLowFrequency[i]);
-					#endif
 					tsr_halfCx2 SampleHighFrequency = dpv_force_interleave_registers_array(RawHighFrequency[i]);
 					tsr_half2x2 SampleMetadata = dpv_force_interleave_registers_array(RawMetadata[i]);
 					tsr_half4x2 SampleTranslucency = dpv_force_interleave_registers(
@@ -690,18 +639,11 @@ void MainCS(
 					
 					tsr_half2 KernelWeight = RawKernelWeight[i];
 
-					#if CONFIG_HIGH_FREQUENCY_ONLY
-						tsr_half2 InvWeight = SafeRcp(SampleMetadata[0]);
-					#else
-						tsr_half2 InvWeight = SafeRcp(SampleMetadata[1]);
-					#endif
+					tsr_half2 InvWeight = SafeRcp(SampleMetadata[0]);
 
 					HighFrequencyMin = min(HighFrequencyMin, dpv_scale(SampleHighFrequency, InvWeight));
 					HighFrequencyMax = max(HighFrequencyMax, dpv_scale(SampleHighFrequency, InvWeight));
 
-					#if !CONFIG_HIGH_FREQUENCY_ONLY
-						AccumulateLowFrequency  += dpv_scale(SampleLowFrequency, KernelWeight);
-					#endif
 					AccumulateHighFrequency += dpv_scale(SampleHighFrequency, KernelWeight);
 					AccumulateMetadata      += dpv_scale(SampleMetadata, KernelWeight);
 					
@@ -712,35 +654,14 @@ void MainCS(
 				
 				tsr_half2 FinalMultiplier = dpv_force_interleave_registers(tsr_half(Samples0.FinalMultiplier), tsr_half(Samples1.FinalMultiplier));
 
-				#if !CONFIG_HIGH_FREQUENCY_ONLY
-					AccumulateLowFrequency  = dpv_scale(AccumulateLowFrequency, FinalMultiplier);
-				#endif
 				AccumulateHighFrequency = dpv_scale(AccumulateHighFrequency, FinalMultiplier);
 				AccumulateMetadata      = dpv_scale(AccumulateMetadata, FinalMultiplier);
 				AccumulateTranslucency  = dpv_scale(AccumulateTranslucency, FinalMultiplier);
 			}
-
-			#if CONFIG_UPDATE_SUBPIXEL_DETAIL
-			FSubpixelNeighborhood SubpixelNeighborhood0 = GatherPrevSubpixelNeighborhood(PrevHistory_SubpixelDetails, dpv_lo(PrevHistoryBufferUV));
-			FSubpixelNeighborhood SubpixelNeighborhood1 = GatherPrevSubpixelNeighborhood(PrevHistory_SubpixelDetails, dpv_hi(PrevHistoryBufferUV));
-			{
-				PrevSubpixelDetails = 0;
-				UNROLL_N(SUB_PIXEL_COUNT)
-				for (uint SubpixelId = 0; SubpixelId < SUB_PIXEL_COUNT; SubpixelId++)
-				{
-					tsr_subpixel_payload SubpixelPayload0 = GetSubpixelPayload(SubpixelNeighborhood0, SubpixelId);
-					tsr_subpixel_payload SubpixelPayload1 = GetSubpixelPayload(SubpixelNeighborhood1, SubpixelId);
-					PrevSubpixelDetails[0] |= SubpixelPayload0 << tsr_subpixel_payload(SUB_PIXEL_BIT_COUNT * SubpixelId);
-					PrevSubpixelDetails[1] |= SubpixelPayload1 << tsr_subpixel_payload(SUB_PIXEL_BIT_COUNT * SubpixelId);
-				}
-			}
-			#endif // CONFIG_UPDATE_SUBPIXEL_DETAIL
 			
 			#if CONFIG_SCENE_COLOR_APLHA
-				AccumulateLowFrequency  = -dpv_min(-AccumulateLowFrequency, tsr_half(0.0).xxxx);
 				AccumulateHighFrequency = -dpv_min(-AccumulateHighFrequency, tsr_half(0.0).xxxx);
 			#else
-				AccumulateLowFrequency  = -dpv_min(-AccumulateLowFrequency, tsr_half(0.0).xxx);
 				AccumulateHighFrequency = -dpv_min(-AccumulateHighFrequency, tsr_half(0.0).xxx);
 			#endif
 			AccumulateMetadata      = -dpv_min(-AccumulateMetadata, tsr_half(0.0).xx);
@@ -749,18 +670,11 @@ void MainCS(
 		
 		// Unpack history.
 		{
-			#if !CONFIG_HIGH_FREQUENCY_ONLY
-				PrevLowFrequencyColor = AccumulateLowFrequency;
-				PrevLowFrequencyWeight = AccumulateMetadata[0];
-			#endif
-
 			PrevHistoryMommentMin = HighFrequencyMin;
 			PrevHistoryMommentMax = HighFrequencyMax;
 
 			PrevHighFrequencyColor = AccumulateHighFrequency;
-			#if CONFIG_HIGH_FREQUENCY_ONLY
-				PrevHistoryWeight = AccumulateMetadata[0];
-			#endif
+			PrevHistoryWeight = AccumulateMetadata[0];
 			PrevHistoryValidity = AccumulateMetadata[1];
 
 			PrevHistoryTranslucency = AccumulateTranslucency;
@@ -777,11 +691,6 @@ void MainCS(
 			CorrectExposure(PrevHighFrequencyColor[0]);
 			CorrectExposure(PrevHighFrequencyColor[1]);
 			CorrectExposure(PrevHighFrequencyColor[2]);
-			#if !CONFIG_HIGH_FREQUENCY_ONLY
-				CorrectExposure(PrevLowFrequencyColor[0]);
-				CorrectExposure(PrevLowFrequencyColor[1]);
-				CorrectExposure(PrevLowFrequencyColor[2]);
-			#endif
 
 			// Fixes negative validity that can happen due to negative lobes of Catrom.
 			PrevHistoryValidity = max(PrevHistoryValidity, tsr_half(0.0).xx);
@@ -803,29 +712,11 @@ void MainCS(
 	tsr_half2 ClosestInputLuma4;
 	
 	
-	#if CONFIG_UPDATE_SUBPIXEL_DETAIL
-		bool2 bUpdateSubpixelDetails;
-		bool2x2 bInputKOSign;
-	#endif
-
 	PLATFORM_SPECIFIC_ISOLATE
 	{
 		// Vector in pixel between pixel K -> O.
 		tsr_half2x2 dInputKO = tsr_half2x2(InputPPCo - InputPPCk);
 
-		// Compresses the sign of dInputKO into SGPR to keep VGPR pressure low.
-		#if CONFIG_UPDATE_SUBPIXEL_DETAIL
-		{
-			float SizeOfOutputPixelInInputPixelWidth = 0.5 * (InputInfo_ViewportSize.x * HistoryInfo_ViewportSizeInverse.x);
-
-			bUpdateSubpixelDetails[0] = all(abs(dpv_lo(dInputKO)) < SizeOfOutputPixelInInputPixelWidth);
-			bUpdateSubpixelDetails[1] = all(abs(dpv_hi(dInputKO)) < SizeOfOutputPixelInInputPixelWidth);
-		
-			bInputKOSign[0] = dInputKO[0] < 0.0;
-			bInputKOSign[1] = dInputKO[1] < 0.0;
-		}
-		#endif
-		
 		tsr_halfC RawInputColorArray[CONFIG_SAMPLES][DPV_PIXEL_PER_LANE];
 		tsr_half4 RawTranslucencyColorArray[CONFIG_SAMPLES][DPV_PIXEL_PER_LANE];
 
@@ -1001,10 +892,6 @@ void MainCS(
 				SharedArray10[LocalGroupThreadIndex] = tsr_half4(TranslucencyMaxColor[0], TranslucencyMaxColor[1]);
 				SharedArray11[LocalGroupThreadIndex] = tsr_half4(TranslucencyMaxColor[2], TranslucencyMaxColor[3]);
 			#endif
-	
-			#if CONFIG_UPDATE_SUBPIXEL_DETAIL
-				SharedParallaxFactorBits[LocalGroupThreadIndex] = ParallaxFactorBits;
-			#endif
 		#else
 			SharedArray0[LocalGroupThreadIndex] = tsr_half4(dpv_lo(FilteredInputColor), dpv_lo(LowFrequencyRejection));
 			SharedArray1[LocalGroupThreadIndex] = tsr_half4(dpv_lo(InputMinColor), dpv_lo(InputPixelAlignement));
@@ -1035,10 +922,6 @@ void MainCS(
 				tsr_half4 RawLDS11 = SharedArray11[LocalGroupThreadIndex];
 			#endif
 
-			#if CONFIG_UPDATE_SUBPIXEL_DETAIL
-				ParallaxFactorBits = SharedParallaxFactorBits[LocalGroupThreadIndex];
-			#endif
-
 			FilteredInputColor = tsr_half3x2(RawLDS0.rg, RawLDS0.ba, RawLDS1.rg);
 			InputMinColor = tsr_half3x2(RawLDS2.rg, RawLDS2.ba, RawLDS3.rg);
 			InputMaxColor = tsr_half3x2(RawLDS4.rg, RawLDS4.ba, RawLDS5.rg);
@@ -1068,75 +951,12 @@ void MainCS(
 	}
 	#endif
 
-	tsr_half2 HighFrequencyRejection = tsr_half(1.0).xx;
-	// Rejects high frequency details if their low frequencies drift from the low frequency history
-	#if CONFIG_LOW_FREQUENCY_DRIFT_REJECTION
-	{
-		#if CONFIG_HIGH_FREQUENCY_ONLY
-			tsr_halfCx2 NormalizedPrevHighFrequencyColor = dpv_scale(PrevHighFrequencyColor, SafeRcp(PrevHistoryWeight));
-
-			tsr_halfCx2 PrevHighFrequencyYCoCg = NormalizedPrevHighFrequencyColor;
-			tsr_halfCx2 PrevYCoCg = FilteredInputColor;
-			tsr_halfCx2 ClampedPrevYCoCg = clamp(FilteredInputColor, PrevHistoryMommentMin, PrevHistoryMommentMax);
-		#else
-			tsr_halfCx2 NormalizedPrevHighFrequencyColor = dpv_scale(PrevHighFrequencyColor, SafeRcp(PrevHistoryValidity));
-
-			tsr_halfCx2 PrevHighFrequencyYCoCg = TransformColorForClampingBox(NormalizedPrevHighFrequencyColor);
-			tsr_halfCx2 PrevYCoCg = TransformColorForClampingBox(PrevLowFrequencyColor);
-			tsr_halfCx2 ClampedPrevYCoCg = TransformColorForClampingBox(clamp(PrevLowFrequencyColor, PrevHistoryMommentMin, PrevHistoryMommentMax));
-		#endif
-
-		HighFrequencyRejection = MeasureRejectionFactor(
-			PrevYCoCg, ClampedPrevYCoCg,
-			PrevHighFrequencyYCoCg, InputMinColor, InputMaxColor,
-			MeasureBackbufferLDRQuantizationErrorFromHDRLuma(max(PrevYCoCg[0], ClampedPrevYCoCg[0])));
-
-		//PrevHighFrequencyColor = dpv_scale(PrevHighFrequencyColor, HighFrequencyRejection);
-
-		#if !CONFIG_HIGH_FREQUENCY_ONLY
-			PrevHighFrequencyColor = dpv_scale(clamp(NormalizedPrevHighFrequencyColor, PrevHistoryMommentMin, PrevHistoryMommentMax), HighFrequencyRejection * PrevHistoryValidity);
-			PrevHistoryValidity *= HighFrequencyRejection;
-		#endif
-	}
-	#endif
-
 	// Contribute current frame input into the predictor for next frame.
 	const tsr_half Histeresis = rcp(tsr_half(MAX_SAMPLE_COUNT));
 	const tsr_half PredictionOnlyValidity = Histeresis * tsr_half(2.0);
 	
-	// Clamp low frequency
-	#if !CONFIG_HIGH_FREQUENCY_ONLY
-	tsr_half2 LumaMin;
-	tsr_half2 LumaMax;
-	tsr_halfCx2 ClampedLowFrequencyColor;
-	tsr_half2 MidFrequencyRejection;
-	{
-		LumaMin = InputMinColor[0];
-		LumaMax = InputMaxColor[0];
-
-		tsr_halfCx2 PrevYCoCg = TransformColorForClampingBox(PrevLowFrequencyColor);
-		tsr_halfCx2 ClampedPrevYCoCg = dpv_clamp(PrevYCoCg, InputMinColor, InputMaxColor);
-		tsr_halfCx2 InputCenterYCoCg = TransformColorForClampingBox(FilteredInputColor);
-
-		ClampedLowFrequencyColor = YCoCgToRGB(ClampedPrevYCoCg);
-		
-		MidFrequencyRejection = MeasureRejectionFactor(
-			PrevYCoCg, ClampedPrevYCoCg,
-			InputCenterYCoCg, InputMinColor, InputMaxColor,
-			MeasureBackbufferLDRQuantizationErrorFromHDRLuma(max(PrevYCoCg[0], ClampedPrevYCoCg[0])));
-
-		#if !CONFIG_CLAMP
-		{
-			ClampedLowFrequencyColor = PrevLowFrequencyColor;
-			MidFrequencyRejection = tsr_half(1.0);
-		}
-		#endif
-	}
-	#endif // !CONFIG_HIGH_FREQUENCY_ONLY
-
 	tsr_halfCx2 FinalHighFrequencyColor;
 	tsr_half2 FinalHistoryValidity;
-	#if CONFIG_HIGH_FREQUENCY_ONLY
 	{
 		#if CONFIG_HIGH_FREQUENCY_YCOCG
 			tsr_halfCx2 ClampedPrevHighFrequencyColor = YCoCgToRGB(clamp(TransformColorForClampingBox(PrevHighFrequencyColor), dpv_scale(InputMinColor, PrevHistoryWeight), dpv_scale(InputMaxColor, PrevHistoryWeight)));
@@ -1146,7 +966,7 @@ void MainCS(
 		
 		tsr_half2 TotalRejection = select(or(bIsOffScreen, bIsDisoccluded), tsr_half(0.0).xx, saturate(LowFrequencyRejection * tsr_half(4.0)));
 
-		tsr_half2 BlendClamp = TotalRejection * HoleFilledVelocityMask[1] * HighFrequencyRejection;
+		tsr_half2 BlendClamp = TotalRejection * HoleFilledVelocityMask[1];
 
 		PrevHighFrequencyColor[0] = lerp(ClampedPrevHighFrequencyColor[0], PrevHighFrequencyColor[0], BlendClamp);
 		PrevHighFrequencyColor[1] = lerp(ClampedPrevHighFrequencyColor[1], PrevHighFrequencyColor[1], BlendClamp);
@@ -1218,229 +1038,6 @@ void MainCS(
 		// Quantize validity for the 8bit encoding to avoid numerical shift between color and validity.
 		FinalHistoryValidity = ceil(tsr_half(255.0) * OutputValidity) * rcp(tsr_half(255.0));
 	}
-	#else // !CONFIG_HIGH_FREQUENCY_ONLY
-	{
-		tsr_half2 InvPrevHistoryValidity = SafeRcp(PrevHistoryValidity);
-
-		// Compute how much the history needs to be rejected, based on the completeness.
-		tsr_half2 PrevHistoryRejectionWeight = LowFrequencyRejection * HoleFilledVelocityMask[1];
-	
-		//#if !CONFIG_CLAMP
-		//	PrevHistoryRejectionWeight = 1.0;
-		//#endif
-		
-		PrevHistoryRejectionWeight = select(or(bIsOffScreen, bIsDisoccluded), tsr_half(0.0), PrevHistoryRejectionWeight);
-
-		tsr_half2 DesiredCurrentContribution = max(Histeresis * InputPixelAlignement, tsr_half(0.0));
-
-		// Determine whether the the prediction based rejection was confident enough.
-		tsr_half2 RejectionConfidentEnough = tsr_half(1); // saturate(RejectionValidity * MAX_SAMPLE_COUNT - 3.0);
-		
-		// Clamp the validity due to motion to maintain better sharpness in history reprojection under motion.
-		tsr_half2 ClampedPrevHistoryValidity;
-		#if 1
-		{
-			tsr_half2 MaxValidity = tsr_half(1.0) - (tsr_half(1.0) - PredictionOnlyValidity) * saturate(OutputPixelVelocity * tsr_half(InvWeightClampingPixelSpeed));
-
-			// Clamp up the max validity to favorise stability under motion on high contrast edges.
-			#if 1
-			{
-				tsr_half2 PrevHistoryLuma = Luma4(PrevHighFrequencyColor) * InvPrevHistoryValidity;
-				tsr_half2 FilteredLuma = Luma4(FilteredInputColor);
-
-				tsr_half2 MinValidityForStability = abs(FilteredLuma - PrevHistoryLuma) / max(FilteredLuma, PrevHistoryLuma);
-				
-				MaxValidity = max(MaxValidity, MinValidityForStability);
-			}
-			#endif
-			
-			ClampedPrevHistoryValidity = min(PrevHistoryValidity, MaxValidity);
-		}
-		#else
-		{
-			ClampedPrevHistoryValidity = PrevHistoryValidity;
-		}
-		#endif
-
-		// Compute the newly rejected validity
-		tsr_half2 RejectedValidity = (
-			min(ClampedPrevHistoryValidity, PredictionOnlyValidity - DesiredCurrentContribution) +
-			max(ClampedPrevHistoryValidity - PredictionOnlyValidity + DesiredCurrentContribution, tsr_half(0.0)) * PrevHistoryRejectionWeight);
-
-		RejectedValidity = ClampedPrevHistoryValidity * PrevHistoryRejectionWeight;
-
-		// Compute the maximum output validity.
-		tsr_half2 OutputValidity = (
-			clamp(RejectedValidity + DesiredCurrentContribution, tsr_half(0.0), PredictionOnlyValidity) +
-			clamp(RejectedValidity + DesiredCurrentContribution * PrevHistoryRejectionWeight * RejectionConfidentEnough - PredictionOnlyValidity, tsr_half(0.0), tsr_half(1.0) - PredictionOnlyValidity));
-
-		#if !CONFIG_ACCUMULATE_TRANSLUCENCY_SEPARATLY
-			FLATTEN
-			if (bIsResponsiveAAPixel[0])
-			{
-				OutputValidity[0] = tsr_half(0.0);
-			}
-
-			FLATTEN
-			if (bIsResponsiveAAPixel[1])
-			{
-				OutputValidity[1] = tsr_half(0.0);
-			}
-		#endif
-		
-		tsr_half2 PrevMomentWeight = max(OutputValidity - DesiredCurrentContribution, tsr_half(0.0));
-		tsr_half2 CurrentMomentWeight = min(DesiredCurrentContribution, OutputValidity);
-		
-		#if 1
-		{
-			tsr_half2 PrevHistoryToneWeight = HdrWeightY(Luma4(PrevHighFrequencyColor) * InvPrevHistoryValidity);
-			tsr_half2 FilteredInputToneWeight = HdrWeight4(FilteredInputColor);
-			
-			tsr_half2 BlendPrevHistory = PrevMomentWeight * PrevHistoryToneWeight;
-			tsr_half2 BlendFilteredInput = CurrentMomentWeight * FilteredInputToneWeight;
-
-			tsr_half2 CommonWeight = SafeRcp(BlendPrevHistory + BlendFilteredInput);
-
-			FinalHighFrequencyColor = (
-				dpv_scale(PrevHighFrequencyColor, CommonWeight * BlendPrevHistory * InvPrevHistoryValidity) +
-				dpv_scale(FilteredInputColor, CommonWeight * BlendFilteredInput));
-		}
-		#else
-		{
-			FinalHighFrequencyColor = PrevHighFrequencyColor * (PrevMomentWeight * InvPrevHistoryValidity) + FilteredInputColor * (CurrentMomentWeight);
-		}
-		#endif
-
-		// Quantize validity for the 8bit encoding to avoid numerical shift between color and validity.
-		#if 0
-			FinalHistoryValidity = OutputValidity;
-		#else
-			FinalHistoryValidity = ceil(tsr_half(255.0) * OutputValidity) * rcp(tsr_half(255.0));
-		#endif
-
-		#if 0 // DEBUG
-		{
-			FinalHighFrequencyColor = dpv_scale(PrevHighFrequencyColor, InvPrevHistoryValidity);
-			FinalHistoryValidity = PrevHistoryValidity;
-		}
-		#endif
-	}
-	#endif // !CONFIG_HIGH_FREQUENCY_ONLY
-
-	// Compute the low frequency history
-	#if !CONFIG_HIGH_FREQUENCY_ONLY
-	tsr_halfCx2 FinalLowFrequencyColor;
-	tsr_half2 FinalLowFrequencyWeight;
-	{
-		const tsr_half TargetHesteresisCurrentFrameWeight = rcp(tsr_half(MAX_FALLBACK_SAMPLE_COUNT));
-
-		tsr_half2 LumaHistory = Luma4(PrevLowFrequencyColor);
-		tsr_half2 LumaFiltered = Luma4(FilteredInputColor);
-
-		{
-			tsr_half2 OutputBlend = ComputeLowFrequencyContribution(FinalHistoryValidity);
-
-			//ClampedLowFrequencyColor = lerp(ClampedLowFrequencyColor, PrevLowFrequencyColor, OutputBlend);
-		}
-
-		tsr_half2 BlendFinal;
-		#if 1
-		{
-			tsr_half2 CurrentFrameSampleCount = max(InputPixelAlignement, tsr_half(0.005));
-			
-			// Use only one sample count to be extremely fast to recover in history rejection. But then stabilize imediately so sub pixel
-			// frequency can be used as soon as possible.
-			tsr_half2 PrevLowFrequencySampleCount;
-			FLATTEN
-			if (PrevLowFrequencyWeight[0] < tsr_half(1.0))
-			{
-				PrevLowFrequencySampleCount[0] = PrevLowFrequencyWeight[0];
-			}
-			else
-			{
-				PrevLowFrequencySampleCount[0] = tsr_half(MAX_FALLBACK_SAMPLE_COUNT);
-			}
-
-			FLATTEN
-			if (PrevLowFrequencyWeight[1] < tsr_half(1.0))
-			{
-				PrevLowFrequencySampleCount[1] = PrevLowFrequencyWeight[1];
-			}
-			else
-			{
-				PrevLowFrequencySampleCount[1] = tsr_half(MAX_FALLBACK_SAMPLE_COUNT);
-			}
-
-			// Reject the history based on the low frequencies.
-			#if 1
-			{
-				// float InvertedBlendFinal = PrevLowFrequencySampleCount / (CurrentFrameSampleCount + PrevLowFrequencySampleCount);
-				// float PrevLowFrequencyRejectionFactor = saturate(LowFrequencyRejection / InvertedBlendFinal);
-				
-				tsr_half2 PrevLowFrequencyRejectionFactor = saturate(LowFrequencyRejection * (CurrentFrameSampleCount + PrevLowFrequencySampleCount) * rcp(PrevLowFrequencySampleCount));
-
-				PrevLowFrequencySampleCount *= PrevLowFrequencyRejectionFactor;
-			}
-			#endif
-
-			BlendFinal = CurrentFrameSampleCount * rcp(CurrentFrameSampleCount + PrevLowFrequencySampleCount);
-
-			// Increase the blending on motion
-			#if 1
-			{
-				BlendFinal = lerp(BlendFinal, max(tsr_half(0.2), BlendFinal), saturate(OutputPixelVelocity * rcp(tsr_half(40.0))));
-			}
-			#endif
-
-			// Antiflicker
-			#if 1
-			{
-				tsr_half2 DistToClamp = min( abs(LumaHistory - LumaMin), abs(LumaHistory - LumaMax) ) * rcp(max3( LumaHistory, LumaFiltered, tsr_half(1e-4) ));
-				BlendFinal *= tsr_half(0.2) + tsr_half(0.8) * saturate(tsr_half(0.5) * DistToClamp);
-			}
-			#endif
-			
-			// Make sure to have at least some small contribution
-			#if 1
-			{
-				BlendFinal = max(BlendFinal, saturate(tsr_half(0.01) * LumaHistory * rcp(abs(LumaFiltered - LumaHistory))));
-			}
-			#endif
-
-			#if !CONFIG_ACCUMULATE_TRANSLUCENCY_SEPARATLY
-				// Responsive forces 1/4 of new frame.
-				BlendFinal[0] = bIsResponsiveAAPixel[0] ? tsr_half(1.0/4.0) : BlendFinal[0];
-				BlendFinal[1] = bIsResponsiveAAPixel[1] ? tsr_half(1.0/4.0) : BlendFinal[1];
-			#endif
-
-			// Completely discard the history
-			{
-				tsr_half2 TotalRejection = select(or(bIsOffScreen, bIsDisoccluded), tsr_half(0.0), saturate(LowFrequencyRejection * tsr_half(4.0)));
-
-				tsr_half2 TotalRejectionBlending = max(BlendFinal, HoleFilledVelocityMask[0]);
-
-				PrevLowFrequencySampleCount *= TotalRejection;
-				BlendFinal = lerp(TotalRejectionBlending, BlendFinal, TotalRejection);
-			}
-
-			FinalLowFrequencyWeight = saturate(CurrentFrameSampleCount + PrevLowFrequencySampleCount);
-			
-			#if 1
-				FinalLowFrequencyWeight = saturate(floor(tsr_half(255.0) * (CurrentFrameSampleCount + PrevLowFrequencySampleCount)) * rcp(tsr_half(255.0)));
-			#endif
-		}
-		#endif
-
-		{
-			tsr_half2 FilterWeight = HdrWeight4(FilteredInputColor);
-			tsr_half2 ClampedHistoryWeight = HdrWeight4(ClampedLowFrequencyColor);
-
-			tsr_half2x2 Weights = WeightedLerpFactors(ClampedHistoryWeight, FilterWeight, BlendFinal);
-
-			FinalLowFrequencyColor = dpv_scale(ClampedLowFrequencyColor, Weights[0]) + dpv_scale(FilteredInputColor, Weights[1]);
-		}
-	}
-	#endif // !CONFIG_HIGH_FREQUENCY_ONLY
 	
 	// Compute the low frequency history
 	tsr_half4x2 FinalTranslucencyColor;
@@ -1532,95 +1129,11 @@ void MainCS(
 	}
 	#endif
 
-	// Update subpixel detail
-	#if CONFIG_UPDATE_SUBPIXEL_DETAIL
-	tsr_subpixel_details2 FinalSubpixelDetails;
-	{
-		FinalSubpixelDetails = PrevSubpixelDetails;
-
-
-		{
-			const tsr_subpixel_details ParallaxFactorMask = tsr_subpixel_details(
-				(SUB_PIXEL_PARALLAX_FACTOR_BIT_MASK << (SUB_PIXEL_PARALLAX_FACTOR_BIT_OFFSET + 0 * SUB_PIXEL_BIT_COUNT)) | 
-				(SUB_PIXEL_PARALLAX_FACTOR_BIT_MASK << (SUB_PIXEL_PARALLAX_FACTOR_BIT_OFFSET + 1 * SUB_PIXEL_BIT_COUNT)) | 
-				(SUB_PIXEL_PARALLAX_FACTOR_BIT_MASK << (SUB_PIXEL_PARALLAX_FACTOR_BIT_OFFSET + 2 * SUB_PIXEL_BIT_COUNT)) | 
-				(SUB_PIXEL_PARALLAX_FACTOR_BIT_MASK << (SUB_PIXEL_PARALLAX_FACTOR_BIT_OFFSET + 3 * SUB_PIXEL_BIT_COUNT)) | 
-				0x0);
-			
-			// Reset parallax factor
-			FinalSubpixelDetails = select(or(bIsOffScreen, bIsDisoccluded), FinalSubpixelDetails & ~ParallaxFactorMask, FinalSubpixelDetails);
-		}
-		
-		uint2 SubpixelId = dpv_dot(uint2x2(bInputKOSign), uint2(1, SUB_PIXEL_GRID_SIZE));
-
-		tsr_subpixel_payload2 SubpixelShift = tsr_subpixel_payload2(SubpixelId) * tsr_subpixel_payload(SUB_PIXEL_BIT_COUNT);
-
-		tsr_subpixel_payload2 SubpixelPayload = (ParallaxFactorBits << tsr_subpixel_payload(SUB_PIXEL_PARALLAX_FACTOR_BIT_OFFSET));
-
-		tsr_subpixel_details2 UpdatedFinalSubpixelDetails = (FinalSubpixelDetails & (~(tsr_subpixel_payload(SUB_PIXEL_BIT_MASK) << SubpixelShift))) | (SubpixelPayload << SubpixelShift);
-
-		FLATTEN
-		if (bUpdateSubpixelDetails[0])
-		{
-			FinalSubpixelDetails.x = UpdatedFinalSubpixelDetails.x;
-		}
-		
-		FLATTEN
-		if (bUpdateSubpixelDetails[1])
-		{
-			FinalSubpixelDetails.y = UpdatedFinalSubpixelDetails.y;
-		}
-	}
-	#endif // CONFIG_UPDATE_SUBPIXEL_DETAIL
-
 	// Compute final output
 	tsr_halfCx2 FinalOutputColor;
-	#if CONFIG_HIGH_FREQUENCY_ONLY
 	{
 		FinalOutputColor = FinalHighFrequencyColor;
-	}
-	#else // !CONFIG_HIGH_FREQUENCY_ONLY
-	{
-		tsr_half2 OutputBlend = ComputeLowFrequencyContribution(FinalHistoryValidity);
 
-		#if 0 // DEBUG
-		{
-			if (View.GeneralPurposeTweak == 2.0)
-				OutputBlend = 0.0;
-			else if (View.GeneralPurposeTweak == 3.0)
-				OutputBlend = 1.0;
-		}
-		#elif 0 // DEBUG
-		{
-			OutputBlend = 0.0;
-		}
-		#elif 0 // DEBUG
-		{
-			OutputBlend = 1.0;
-		}
-		#endif
-
-		tsr_half2 LowFrequencyWeight = HdrWeight4(FinalLowFrequencyColor);
-		tsr_half2 Moment1Weight = HdrWeight4(FinalHighFrequencyColor);
-
-		tsr_half2x2 Weights = WeightedLerpFactors(LowFrequencyWeight, Moment1Weight, OutputBlend);
-
-		#if DEBUG_FALLBACK_BLENDING
-			tsr_half3 LowFrequencyColor = tsr_half3(1, 0.25, 0.25);
-			tsr_half3 HighFrequencyColor = tsr_half3(0.25, 1, 0.25);
-
-			FinalOutputColor = dpv_scale(FinalLowFrequencyColor * dpv_interleave_mono_registers(LowFrequencyColor), Weights[0]) + dpv_scale(FinalHighFrequencyColor * dpv_interleave_mono_registers(HighFrequencyColor), Weights[1]);
-		#elif DEBUG_LOW_FREQUENCY_REJECTION
-			tsr_half3 DebugColor = lerp(tsr_half3(1, 0.5, 0.5), tsr_half3(0.5, 1, 0.5), LowFrequencyRejection);
-			
-			FinalOutputColor = FinalLowFrequencyColor * Weights.x * DebugColor + FinalHighFrequencyColor * Weights.y * DebugColor;
-		#else
-			FinalOutputColor = dpv_scale(FinalLowFrequencyColor, Weights[0]) + dpv_scale(FinalHighFrequencyColor, Weights[1]);
-		#endif
-
-	}
-	#endif // !CONFIG_HIGH_FREQUENCY_ONLY
-	{
 		#if CONFIG_ACCUMULATE_TRANSLUCENCY_SEPARATLY
 			FinalOutputColor[0] = FinalOutputColor[0] * FinalTranslucencyColor[3] + FinalTranslucencyColor[0];
 			FinalOutputColor[1] = FinalOutputColor[1] * FinalTranslucencyColor[3] + FinalTranslucencyColor[1];
@@ -1637,6 +1150,16 @@ void MainCS(
 			FinalOutputColor[3] = select(FinalOutputColor[3] > 0.995, tsr_half(1.0), FinalOutputColor[3]);
 		}
 		#endif
+		
+		#if 0
+		{
+			//tsr_half2 DebugValue = tsr_half(1.0) - tsr_half2(bIsDisoccluded);
+			tsr_half2 DebugValue = LowFrequencyRejection;
+			FinalOutputColor *= 
+				dpv_scale(dpv_interleave_mono_registers(tsr_half3(1, 0.5, 0.5)), tsr_half(1.0) - DebugValue) +
+				dpv_scale(dpv_interleave_mono_registers(tsr_half3(0.5, 1, 0.5)), DebugValue);
+		}
+		#endif
 	}
 
 
@@ -1667,11 +1190,7 @@ void MainCS(
 		HalfResOutputPixelPos[1] = (LocalHistoryPixelPos[1] >> tsr_short(1)) | (((LocalHistoryPixelPos[1] & tsr_short(0x1))) * tsr_short(~0));
 
 		// Premultiply FinalHighFrequencyColor to have validity weighted average in next frame history reprojection.
-		#if CONFIG_HIGH_FREQUENCY_ONLY
-			tsr_half2 FinalHistoryWeight = max(FinalHistoryValidity, MIN_IMPLICIT_HISTORY_WEIGHT);
-		#else
-			tsr_half2 FinalHistoryWeight = FinalHistoryValidity;
-		#endif
+		tsr_half2 FinalHistoryWeight = max(FinalHistoryValidity, MIN_IMPLICIT_HISTORY_WEIGHT);
 		tsr_halfCx2 PremultipliedHighFrequencyColor = dpv_scale(FinalHighFrequencyColor, FinalHistoryWeight);
 
 		// Stocastically round up or down using the hardware RWTexture2D truncation unit to take into precision
@@ -1681,20 +1200,12 @@ void MainCS(
 			uint2 Random = Rand3DPCG16(int3(dpv_lo(LocalHistoryPixelPos), View.StateFrameIndexMod8)).xy;
 			tsr_half E = tsr_half(Hammersley16(0, 1, Random).x);
 			
-			#if !CONFIG_HIGH_FREQUENCY_ONLY
-				FinalLowFrequencyColor = QuantizeForFloatRenderTarget(FinalLowFrequencyColor, E, HistoryQuantizationError);
-			#endif
 			PremultipliedHighFrequencyColor = QuantizeForFloatRenderTarget(PremultipliedHighFrequencyColor, E, HistoryQuantizationError);
 		}
 		#endif
 
 		// Protext from NaN and +Inf when writing out the history.
 		{
-			#if !CONFIG_HIGH_FREQUENCY_ONLY
-				FinalLowFrequencyColor = -dpv_min(-FinalLowFrequencyColor, ColorNull);
-				FinalLowFrequencyColor = dpv_min(FinalLowFrequencyColor, ColorMax10BitsFloat);
-			#endif
-
 			PremultipliedHighFrequencyColor = -dpv_min(-PremultipliedHighFrequencyColor, ColorNull);
 			PremultipliedHighFrequencyColor = dpv_min(PremultipliedHighFrequencyColor, ColorMax10BitsFloat);	
 
@@ -1708,14 +1219,8 @@ void MainCS(
 		{
 			// Output final history lo pixel.
 			{
-				#if CONFIG_HIGH_FREQUENCY_ONLY
-					HistoryOutput_ColorArray[tsr_short3(dpv_lo(LocalHistoryPixelPos), HistoryOutput_ArrayIndices_HighFrequency)] = dpv_lo(PremultipliedHighFrequencyColor);
-					HistoryOutput_Metadata[dpv_lo(LocalHistoryPixelPos)]      = tsr_half2(dpv_lo(FinalHistoryWeight), dpv_lo(FinalHistoryValidity));
-				#else
-					HistoryOutput_ColorArray[tsr_short3(dpv_lo(LocalHistoryPixelPos), HistoryOutput_ArrayIndices_LowFrequency)]  = dpv_lo(FinalLowFrequencyColor);
-					HistoryOutput_ColorArray[tsr_short3(dpv_lo(LocalHistoryPixelPos), HistoryOutput_ArrayIndices_HighFrequency)] = dpv_lo(PremultipliedHighFrequencyColor);
-					HistoryOutput_Metadata[dpv_lo(LocalHistoryPixelPos)]      = tsr_half2(dpv_lo(FinalLowFrequencyWeight), dpv_lo(FinalHistoryValidity));
-				#endif
+				HistoryOutput_ColorArray[tsr_short3(dpv_lo(LocalHistoryPixelPos), HistoryOutput_ArrayIndices_HighFrequency)] = dpv_lo(PremultipliedHighFrequencyColor);
+				HistoryOutput_Metadata[dpv_lo(LocalHistoryPixelPos)]      = tsr_half2(dpv_lo(FinalHistoryWeight), dpv_lo(FinalHistoryValidity));
 				#if CONFIG_ACCUMULATE_TRANSLUCENCY_SEPARATLY
 					#if CONFIG_SCENE_COLOR_APLHA
 						HistoryOutput_ColorArray[tsr_short3(dpv_lo(LocalHistoryPixelPos), HistoryOutput_ArrayIndices_Translucency)] = dpv_lo(FinalTranslucencyColor);
@@ -1724,22 +1229,13 @@ void MainCS(
 					#endif 
 					HistoryOutput_TranslucencyAlpha[dpv_lo(LocalHistoryPixelPos)] = dpv_lo(FinalTranslucencyColor).a;
 				#endif
-				#if CONFIG_UPDATE_SUBPIXEL_DETAIL
-					HistoryOutput_SubpixelDetails[dpv_lo(LocalHistoryPixelPos)] = dpv_lo(FinalSubpixelDetails);
-				#endif
 			}
 		
 			// Output final history hi pixel.
 			#if CONFIG_ENABLE_DUAL_PIXEL_VECTORIZATION
 			{
-				#if CONFIG_HIGH_FREQUENCY_ONLY
-					HistoryOutput_ColorArray[tsr_short3(dpv_hi(LocalHistoryPixelPos), HistoryOutput_ArrayIndices_HighFrequency)] = dpv_hi(PremultipliedHighFrequencyColor);
-					HistoryOutput_Metadata[dpv_hi(LocalHistoryPixelPos)]      = tsr_half2(dpv_hi(FinalHistoryWeight), dpv_hi(FinalHistoryValidity));
-				#else
-					HistoryOutput_ColorArray[tsr_short3(dpv_hi(LocalHistoryPixelPos), HistoryOutput_ArrayIndices_LowFrequency)]  = dpv_hi(FinalLowFrequencyColor);
-					HistoryOutput_ColorArray[tsr_short3(dpv_hi(LocalHistoryPixelPos), HistoryOutput_ArrayIndices_HighFrequency)] = dpv_hi(PremultipliedHighFrequencyColor);
-					HistoryOutput_Metadata[dpv_hi(LocalHistoryPixelPos)]      = tsr_half2(dpv_hi(FinalLowFrequencyWeight), dpv_hi(FinalHistoryValidity));
-				#endif
+				HistoryOutput_ColorArray[tsr_short3(dpv_hi(LocalHistoryPixelPos), HistoryOutput_ArrayIndices_HighFrequency)] = dpv_hi(PremultipliedHighFrequencyColor);
+				HistoryOutput_Metadata[dpv_hi(LocalHistoryPixelPos)]      = tsr_half2(dpv_hi(FinalHistoryWeight), dpv_hi(FinalHistoryValidity));
 				#if CONFIG_ACCUMULATE_TRANSLUCENCY_SEPARATLY
 					#if CONFIG_SCENE_COLOR_APLHA
 						HistoryOutput_ColorArray[tsr_short3(dpv_hi(LocalHistoryPixelPos), HistoryOutput_ArrayIndices_Translucency)] = dpv_hi(FinalTranslucencyColor);
@@ -1748,78 +1244,9 @@ void MainCS(
 					#endif
 					HistoryOutput_TranslucencyAlpha[dpv_hi(LocalHistoryPixelPos)] = dpv_hi(FinalTranslucencyColor).a;
 				#endif
-				#if CONFIG_UPDATE_SUBPIXEL_DETAIL
-					HistoryOutput_SubpixelDetails[dpv_hi(LocalHistoryPixelPos)] = dpv_hi(FinalSubpixelDetails);
-				#endif
 			}
 			#endif // !CONFIG_ENABLE_DUAL_PIXEL_VECTORIZATION
 		}
-		
-		// Output half res history
-		#if CONFIG_OUTPUT_HALFRES_HISTORY
-		{
-			#if CONFIG_SCENE_COLOR_APLHA
-				tsr_half4x2 HalfResFinalHighFrequencyColor = FinalHighFrequencyColor * tsr_half(0.25);
-			#else
-				tsr_half4x2 HalfResFinalHighFrequencyColor = tsr_half4x2(FinalHighFrequencyColor[0], FinalHighFrequencyColor[1], FinalHighFrequencyColor[2], tsr_half(0.0).xx) * tsr_half(0.25);
-			#endif
-
-			#if PLATFORM_SUPPORTS_WAVE_BROADCAST && 0 // Support WaveBroadcast with halfs
-			{
-				FWaveBroadcastSettings Horizontal = InitWaveXorButterfly(/* XorButterFly = */ 0x1);
-				FWaveBroadcastSettings Vertical = InitWaveXorButterfly(/* XorButterFly = */ 0x2);
-
-				HalfResFinalHighFrequencyColor += WaveBroadcast(Horizontal, HalfResFinalHighFrequencyColor);
-				HalfResFinalHighFrequencyColor += WaveBroadcast(Vertical, HalfResFinalHighFrequencyColor);
-			}
-			#else
-			{
-				SharedArray0[LocalGroupThreadIndex] = dpv_lo(HalfResFinalHighFrequencyColor);
-				#if CONFIG_ENABLE_DUAL_PIXEL_VECTORIZATION
-					SharedArray1[LocalGroupThreadIndex] = dpv_hi(HalfResFinalHighFrequencyColor);
-				#endif
-
-				SharedArray0[LocalGroupThreadIndex] += SharedArray0[LocalGroupThreadIndex ^ 0x1];
-				#if CONFIG_ENABLE_DUAL_PIXEL_VECTORIZATION
-					SharedArray1[LocalGroupThreadIndex] += SharedArray1[LocalGroupThreadIndex ^ 0x1];
-				#endif
-					
-				SharedArray0[LocalGroupThreadIndex] += SharedArray0[LocalGroupThreadIndex ^ 0x2];
-				#if CONFIG_ENABLE_DUAL_PIXEL_VECTORIZATION
-					SharedArray1[LocalGroupThreadIndex] += SharedArray1[LocalGroupThreadIndex ^ 0x2];
-				#endif
-					
-				#if CONFIG_ENABLE_DUAL_PIXEL_VECTORIZATION
-					HalfResFinalHighFrequencyColor = dpv_interleave_registers(SharedArray0[LocalGroupThreadIndex], SharedArray1[LocalGroupThreadIndex]);
-				#else
-					HalfResFinalHighFrequencyColor = dpv_interleave_mono_registers(SharedArray0[LocalGroupThreadIndex]);
-				#endif
-
-			}
-			#endif
-			
-			//HalfResFinalHighFrequencyColor[0] = 1.0;
-			//HalfResFinalHighFrequencyColor[1] = 0.0;
-			//HalfResFinalHighFrequencyColor[2] = 0.0;
-
-			//HalfResOutputPixelPos = LocalHistoryPixelPos;
-
-			#if CONFIG_SCENE_COLOR_APLHA
-				HistoryOutput_HighFrequencyHalfRes[dpv_lo(HalfResOutputPixelPos)] = dpv_lo(HalfResFinalHighFrequencyColor);
-			
-				#if CONFIG_ENABLE_DUAL_PIXEL_VECTORIZATION
-					HistoryOutput_HighFrequencyHalfRes[dpv_hi(HalfResOutputPixelPos)] = dpv_hi(HalfResFinalHighFrequencyColor);
-				#endif
-			#else
-				HistoryOutput_HighFrequencyHalfRes[dpv_lo(HalfResOutputPixelPos)] = dpv_lo(HalfResFinalHighFrequencyColor).rgb;
-			
-				#if CONFIG_ENABLE_DUAL_PIXEL_VECTORIZATION
-					HistoryOutput_HighFrequencyHalfRes[dpv_hi(HalfResOutputPixelPos)] = dpv_hi(HalfResFinalHighFrequencyColor).rgb;
-				#endif
-			#endif
-		}
-		#endif // !CONFIG_OUTPUT_HALFRES_HISTORY
-
 
 		// Output final scene color Mip0
 		tsr_halfCx2 OutputColor;
@@ -1837,7 +1264,6 @@ void MainCS(
 		}
 		
 		// Output final scene color Mip1
-		#if !DEBUG_OUTPUT // running out of UAVs
 		{
 			tsr_halfCx2 HalfResOutput = OutputColor * tsr_half(0.25);
 			
@@ -1902,7 +1328,6 @@ void MainCS(
 				SceneColorOutputMip1[dpv_hi(HalfResOutputPixelPos)] = dpv_hi(HalfResOutput);
 			#endif
 		}
-		#endif
 		
 		#if DEBUG_OUTPUT
 		for (uint DebugId = 0; DebugId < DEBUG_ARRAY_SIZE; DebugId++)
diff --git a/Engine/Shaders/Public/LaneVectorization.ush b/Engine/Shaders/Public/LaneVectorization.ush
new file mode 100644
--- /dev/null
+++ b/Engine/Shaders/Public/LaneVectorization.ush
@@ -0,0 +1,233 @@
+// Copyright Epic Games, Inc. All Rights Reserved.
+
+/*=============================================================================
+	LaneVectorization.ush: Vectorize arbitrary number of processing per lane.
+=============================================================================*/
+  
+#pragma once
+
+#include "Platform.ush"
+
+
+template<typename ScalarType, uint VectorSize, uint SimdSize>
+struct TLaneVector
+{
+	vector<ScalarType, VectorSize> RegisterArray[SimdSize];
+
+	
+	// Access a component of the vector
+	CALL_SITE_DEBUGLOC
+	TLaneVector<ScalarType, 1, SimdSize> operator [](const uint ComponentIndex)
+	{
+		TLaneVector<ScalarType, 1, SimdSize> R;
+		UNROLL
+		for (uint SimdIndex = 0; SimdIndex < SimdSize; SimdIndex++)
+		{
+			R.SetElement(SimdIndex, GetElement(SimdIndex)[ComponentIndex]);
+		}
+		return R;
+	}
+	
+	// Access an individual element of the vector.
+	CALL_SITE_DEBUGLOC
+	vector<ScalarType, VectorSize> GetElement(const uint SimdIndex)
+	{
+		return RegisterArray[SimdIndex];
+	}
+	
+	CALL_SITE_DEBUGLOC
+	void SetElement(const uint SimdIndex, vector<ScalarType, VectorSize> Element)
+	{
+		RegisterArray[SimdIndex] = Element;
+	}
+	
+	CALL_SITE_DEBUGLOC
+	void SetAllElements(vector<ScalarType, VectorSize> Element)
+	{
+		UNROLL
+		for (uint SimdIndex = 0; SimdIndex < SimdSize; SimdIndex++)
+		{
+			RegisterArray[SimdIndex] = Element;
+		}
+	}
+	
+	
+	// ------------- binary operator +
+	CALL_SITE_DEBUGLOC
+	TLaneVector<ScalarType, VectorSize, SimdSize> operator + (TLaneVector<ScalarType, VectorSize, SimdSize> B)
+	{
+		TLaneVector<ScalarType, VectorSize, SimdSize> R;
+		UNROLL
+		for (uint SimdIndex = 0; SimdIndex < SimdSize; SimdIndex++)
+		{
+			R.SetElement(SimdIndex, GetElement(SimdIndex) + B.GetElement(SimdIndex));
+		}
+		return R;
+	}
+	
+	CALL_SITE_DEBUGLOC
+	TLaneVector<ScalarType, VectorSize, SimdSize> operator + (vector<ScalarType, VectorSize> B)
+	{
+		TLaneVector<ScalarType, VectorSize, SimdSize> R;
+		UNROLL
+		for (uint SimdIndex = 0; SimdIndex < SimdSize; SimdIndex++)
+		{
+			R.SetElement(SimdIndex, GetElement(SimdIndex) + B);
+		}
+		return R;
+	}
+
+	
+	// ------------- unary operator -
+	CALL_SITE_DEBUGLOC
+	TLaneVector<ScalarType, VectorSize, SimdSize> operator - ()
+	{
+		TLaneVector<ScalarType, VectorSize, SimdSize> R;
+		UNROLL
+		for (uint SimdIndex = 0; SimdIndex < SimdSize; SimdIndex++)
+		{
+			R.SetElement(SimdIndex, -GetElement(SimdIndex));
+		}
+		return R;
+	}
+	
+	// ------------- binary operator -
+	CALL_SITE_DEBUGLOC
+	TLaneVector<ScalarType, VectorSize, SimdSize> operator - (TLaneVector<ScalarType, VectorSize, SimdSize> B)
+	{
+		TLaneVector<ScalarType, VectorSize, SimdSize> R;
+		UNROLL
+		for (uint SimdIndex = 0; SimdIndex < SimdSize; SimdIndex++)
+		{
+			R.SetElement(SimdIndex, GetElement(SimdIndex) - B.GetElement(SimdIndex));
+		}
+		return R;
+	}
+	
+	CALL_SITE_DEBUGLOC
+	TLaneVector<ScalarType, VectorSize, SimdSize> operator - (vector<ScalarType, VectorSize> B)
+	{
+		TLaneVector<ScalarType, VectorSize, SimdSize> R;
+		UNROLL
+		for (uint SimdIndex = 0; SimdIndex < SimdSize; SimdIndex++)
+		{
+			R.SetElement(SimdIndex, GetElement(SimdIndex) - B);
+		}
+		return R;
+	}
+
+	
+	// ------------- binary operator *
+	CALL_SITE_DEBUGLOC
+	TLaneVector<ScalarType, VectorSize, SimdSize> operator * (TLaneVector<ScalarType, VectorSize, SimdSize> B)
+	{
+		TLaneVector<ScalarType, VectorSize, SimdSize> R;
+		UNROLL
+		for (uint SimdIndex = 0; SimdIndex < SimdSize; SimdIndex++)
+		{
+			R.SetElement(SimdIndex, GetElement(SimdIndex) * B.GetElement(SimdIndex));
+		}
+		return R;
+	}
+
+	CALL_SITE_DEBUGLOC
+	TLaneVector<ScalarType, VectorSize, SimdSize> operator * (vector<ScalarType, VectorSize> B)
+	{
+		TLaneVector<ScalarType, VectorSize, SimdSize> R;
+		UNROLL
+		for (uint SimdIndex = 0; SimdIndex < SimdSize; SimdIndex++)
+		{
+			R.SetElement(SimdIndex, GetElement(SimdIndex) * B);
+		}
+		return R;
+	}
+	
+
+	// ------------- constructors
+	CALL_SITE_DEBUGLOC
+	static TLaneVector<ScalarType, VectorSize, SimdSize> Vectorize(
+		TLaneVector<ScalarType, 1, SimdSize> A)
+	{
+		TLaneVector<ScalarType, VectorSize, SimdSize> R;
+		UNROLL
+		for (uint SimdIndex = 0; SimdIndex < SimdSize; SimdIndex++)
+		{
+			R.SetElement(SimdIndex, A.GetElement(SimdIndex));
+		}
+		return R;
+	}
+	
+	CALL_SITE_DEBUGLOC
+	static TLaneVector<ScalarType, VectorSize, SimdSize> Const(
+		vector<ScalarType, VectorSize> A)
+	{
+		TLaneVector<ScalarType, VectorSize, SimdSize> R;
+		UNROLL
+		for (uint SimdIndex = 0; SimdIndex < SimdSize; SimdIndex++)
+		{
+			R.SetElement(SimdIndex, A);
+		}
+		return R;
+	}
+
+};
+
+// ------------- 1 parameter
+#define TVECTOR_FUNCTION_1PARAMS(FunctionName) \
+	CALL_SITE_DEBUGLOC \
+	template<typename ScalarType, uint VectorSize, uint SimdSize> \
+	TLaneVector<ScalarType, VectorSize, SimdSize> FunctionName( \
+		TLaneVector<ScalarType, VectorSize, SimdSize> A) \
+	{ \
+		TLaneVector<ScalarType, VectorSize, SimdSize> R; \
+		UNROLL \
+		for (uint SimdIndex = 0; SimdIndex < SimdSize; SimdIndex++) \
+		{ \
+			R.SetElement(SimdIndex, FunctionName(A.GetElement(SimdIndex))); \
+		} \
+		return R; \
+	}
+	
+TVECTOR_FUNCTION_1PARAMS(rcp);
+TVECTOR_FUNCTION_1PARAMS(saturate);
+TVECTOR_FUNCTION_1PARAMS(abs);
+
+// ------------- 2 parameters
+#define TVECTOR_FUNCTION_2PARAMS(ReturnVectorSize, FunctionName) \
+	CALL_SITE_DEBUGLOC \
+	template<typename ScalarType, uint VectorSize, uint SimdSize> \
+	TLaneVector<ScalarType, ReturnVectorSize, SimdSize> FunctionName( \
+		TLaneVector<ScalarType, VectorSize, SimdSize> A, \
+		TLaneVector<ScalarType, VectorSize, SimdSize> B) \
+	{ \
+		TLaneVector<ScalarType, ReturnVectorSize, SimdSize> R; \
+		UNROLL for (uint SimdIndex = 0; SimdIndex < SimdSize; SimdIndex++) { \
+			R.SetElement(SimdIndex, FunctionName(A.GetElement(SimdIndex), B.GetElement(SimdIndex))); \
+		} \
+		return R; \
+	}
+	
+TVECTOR_FUNCTION_2PARAMS(VectorSize, min);
+TVECTOR_FUNCTION_2PARAMS(VectorSize, max);
+TVECTOR_FUNCTION_2PARAMS(/* VectorSize = */ 1, dot);
+
+
+// ------------- 3 parameters
+#define TVECTOR_FUNCTION_3PARAMS(ReturnVectorSize, FunctionName) \
+	CALL_SITE_DEBUGLOC \
+	template<typename ScalarType, uint VectorSize, uint SimdSize> \
+	TLaneVector<ScalarType, ReturnVectorSize, SimdSize> FunctionName( \
+		TLaneVector<ScalarType, VectorSize, SimdSize> A, \
+		TLaneVector<ScalarType, VectorSize, SimdSize> B, \
+		TLaneVector<ScalarType, VectorSize, SimdSize> C) \
+	{ \
+		TLaneVector<ScalarType, ReturnVectorSize, SimdSize> R; \
+		UNROLL for (uint SimdIndex = 0; SimdIndex < SimdSize; SimdIndex++) { \
+			R.SetElement(SimdIndex, FunctionName(A.GetElement(SimdIndex), B.GetElement(SimdIndex), C.GetElement(SimdIndex))); \
+		} \
+		return R; \
+	}
+	
+TVECTOR_FUNCTION_3PARAMS(VectorSize, clamp);
+TVECTOR_FUNCTION_3PARAMS(VectorSize, min3);
+TVECTOR_FUNCTION_3PARAMS(VectorSize, max3);
diff --git a/Engine/Shaders/Public/WaveBroadcastIntrinsics.ush b/Engine/Shaders/Public/WaveBroadcastIntrinsics.ush
--- a/Engine/Shaders/Public/WaveBroadcastIntrinsics.ush
+++ b/Engine/Shaders/Public/WaveBroadcastIntrinsics.ush
@@ -10,14 +10,26 @@
 #include "Platform.ush"
 
 
-#ifdef COMPILER_SUPPORTS_WAVE_SWIZZLE_GCN
-	#define PLATFORM_SUPPORTS_WAVE_BROADCAST (COMPILER_SUPPORTS_WAVE_SWIZZLE_GCN)
+#if COMPILER_SUPPORTS_WAVE_PERMUTE
+	#define PLATFORM_SUPPORTS_WAVE_READ_AT 1
+#elif FEATURE_LEVEL >= FEATURE_LEVEL_SM6
+	#define PLATFORM_SUPPORTS_WAVE_READ_AT 1
+#else
+	#define PLATFORM_SUPPORTS_WAVE_READ_AT 0
+#endif
+
+#if COMPILER_SUPPORTS_WAVE_SWIZZLE_GCN
+	#define PLATFORM_SUPPORTS_WAVE_BROADCAST 1
+#elif FEATURE_LEVEL >= FEATURE_LEVEL_SM6
+	#define PLATFORM_SUPPORTS_WAVE_BROADCAST 1
 #else
 	#define PLATFORM_SUPPORTS_WAVE_BROADCAST 0
 #endif
 
-#ifdef COMPILER_SUPPORTS_WAVE_SWIZZLE_RDNA
-	#define PLATFORM_SUPPORTS_WAVE_ROTATE (COMPILER_SUPPORTS_WAVE_SWIZZLE_RDNA)
+#if defined(COMPILER_SUPPORTS_WAVE_SWIZZLE_RDNA) && COMPILER_SUPPORTS_WAVE_SWIZZLE_RDNA
+	#define PLATFORM_SUPPORTS_WAVE_ROTATE 1
+#elif FEATURE_LEVEL >= FEATURE_LEVEL_SM6
+	#define PLATFORM_SUPPORTS_WAVE_ROTATE 1
 #else
 	#define PLATFORM_SUPPORTS_WAVE_ROTATE 0
 #endif
@@ -25,33 +37,65 @@
 
 /** Different instruction available for WaveBroadcast() */
 	#define WAVE_BROADCAST_NOP 0
-	#if PLATFORM_SUPPORTS_WAVE_BROADCAST
-		#define WAVE_BROADCAST_GCN_SWIZZLE 1
-	#endif
-	#if PLATFORM_SUPPORTS_WAVE_ROTATE
-		#define WAVE_BROADCAST_RDNA_ROTATE 2
-	#endif
+	#define WAVE_BROADCAST_READ_LANE_AT 1
+	#define WAVE_BROADCAST_READ_LDS_AT 2
+	#define WAVE_BROADCAST_GCN_SWIZZLE 3
+	#define WAVE_BROADCAST_RDNA_ROTATE 4
 
 /** Compile time structure to choose which broadcasting should be done. */
 struct FWaveBroadcastSettings
 {
 	// Broadcast operations.
 	uint Operation;
+	
+	// Lane index to read from
+	uint SourceLaneIndex;
+	
+	// index of the thread.
+	uint GroupThreadIndex;
 
 	// Informations about the broadcast for GCN's ds_swizzle.
-	#if PLATFORM_SUPPORTS_WAVE_BROADCAST
-		uint SwizzleAnd;
-		uint SwizzleOr;
-		uint SwizzleXor;
-	#endif
+	uint SwizzleAnd;
+	uint SwizzleOr;
+	uint SwizzleXor;
 	
 	// Information about the broadcast for RDNA's ds_swizzle's rotate functionality
-	#if PLATFORM_SUPPORTS_WAVE_ROTATE
-		int Rotate;
-		uint RotateFixMask;
-	#endif
+	int Rotate;
+	uint RotateFixMask;
 };
 
+/** Group shared memory . */
+#if defined(WAVE_BROADCAST_GROUPSIZE)
+groupshared uint SharedBroadcastLDS[WAVE_BROADCAST_GROUPSIZE];
+#endif
+
+/** Returns the index of the source lane for a dest lane index. */
+CALL_SITE_DEBUGLOC
+uint GetWaveBroadcastSourceLaneIndex(const FWaveBroadcastSettings Settings, uint DestLaneIndex)
+{
+	if (0)
+	{
+		 return DestLaneIndex;
+	}
+	else if (Settings.Operation == WAVE_BROADCAST_READ_LANE_AT)
+	{
+		return Settings.SourceLaneIndex;
+	}
+	else if (Settings.Operation == WAVE_BROADCAST_READ_LDS_AT)
+	{
+		return Settings.SourceLaneIndex;
+	}
+	else if (Settings.Operation == WAVE_BROADCAST_GCN_SWIZZLE)
+	{
+		return (((DestLaneIndex & Settings.SwizzleAnd) | Settings.SwizzleOr) ^ Settings.SwizzleXor) % 32;
+	}
+	else if (Settings.Operation == WAVE_BROADCAST_RDNA_ROTATE)
+	{
+		return ((Settings.RotateFixMask & DestLaneIndex) | ((~Settings.RotateFixMask) & uint(DestLaneIndex + Settings.Rotate))) % 32;
+	}
+	return DestLaneIndex;
+}
+
 /** Generic broadcast instruction. */
 CALL_SITE_DEBUGLOC
 uint WaveBroadcast(const FWaveBroadcastSettings Settings, uint x)
@@ -60,16 +104,68 @@ uint WaveBroadcast(const FWaveBroadcastSettings Settings, uint x)
 	{
 		 return x;
 	}
+	
+	#if PLATFORM_SUPPORTS_WAVE_READ_AT
+	else if (Settings.Operation == WAVE_BROADCAST_READ_LANE_AT)
+	{
+		#if COMPILER_SUPPORTS_WAVE_PERMUTE
+		{
+			return WaveReadLaneAtVarying(x, Settings.SourceLaneIndex);
+		}
+		#elif FEATURE_LEVEL >= FEATURE_LEVEL_SM6
+		{
+			return WaveReadLaneAt(x, Settings.SourceLaneIndex);
+		}
+		#else
+			#error Unimplemented
+		#endif
+	}
+	#endif
+
+	#if defined(WAVE_BROADCAST_GROUPSIZE)
+	else if (Settings.Operation == WAVE_BROADCAST_READ_LDS_AT)
+	{
+		GroupMemoryBarrierWithGroupSync();
+
+		SharedBroadcastLDS[Settings.GroupThreadIndex] = x;
+		
+		GroupMemoryBarrierWithGroupSync();
+
+		return SharedBroadcastLDS[Settings.SourceLaneIndex];
+	}
+	#endif
+
 	#if PLATFORM_SUPPORTS_WAVE_BROADCAST
 	else if (Settings.Operation == WAVE_BROADCAST_GCN_SWIZZLE)
 	{
-		return WaveLaneSwizzleGCN(x, Settings.SwizzleAnd, Settings.SwizzleOr, Settings.SwizzleXor);
+		#if COMPILER_SUPPORTS_WAVE_SWIZZLE_GCN
+		{
+			return WaveLaneSwizzleGCN(x, Settings.SwizzleAnd, Settings.SwizzleOr, Settings.SwizzleXor);
+		}
+		#elif FEATURE_LEVEL >= FEATURE_LEVEL_SM6
+		{
+			return WaveReadLaneAt(x, GetWaveBroadcastSourceLaneIndex(Settings, WaveGetLaneIndex()));
+		}
+		#else
+			#error Unimplemented
+		#endif
 	}
 	#endif
+
 	#if PLATFORM_SUPPORTS_WAVE_ROTATE
 	else if (Settings.Operation == WAVE_BROADCAST_RDNA_ROTATE)
 	{
-		return WaveLaneRotateSwizzleRDNA(x, /* rotate_amount = */ Settings.Rotate, Settings.RotateFixMask);
+		#if defined(COMPILER_SUPPORTS_WAVE_SWIZZLE_RDNA) && COMPILER_SUPPORTS_WAVE_SWIZZLE_RDNA
+		{
+			return WaveLaneRotateSwizzleRDNA(x, /* rotate_amount = */ Settings.Rotate, Settings.RotateFixMask);
+		}
+		#elif FEATURE_LEVEL >= FEATURE_LEVEL_SM6
+		{
+			return WaveReadLaneAt(x, GetWaveBroadcastSourceLaneIndex(Settings, WaveGetLaneIndex()));
+		}
+		#else
+			#error Unimplemented
+		#endif
 	}
 	#endif
 
@@ -262,7 +358,7 @@ half4 WaveBroadcast(const FWaveBroadcastSettings Settings, half4 v)
 	return asfloat16(WaveBroadcast(Settings, asuint16(v)));
 }
 
-#endif
+#endif // PLATFORM_SUPPORTS_REAL_TYPES
 
 
 
@@ -272,22 +368,74 @@ FWaveBroadcastSettings InitNopBroadcast()
 {
 	FWaveBroadcastSettings Settings;
 	Settings.Operation  = WAVE_BROADCAST_NOP;
-	#if PLATFORM_SUPPORTS_WAVE_BROADCAST
-		Settings.SwizzleAnd = 0x00;
-		Settings.SwizzleOr  = 0x00;
-		Settings.SwizzleXor = 0x00;
-	#endif
-	#if PLATFORM_SUPPORTS_WAVE_ROTATE
-		Settings.Rotate        = +0;
-		Settings.RotateFixMask = 0x00;
-	#endif
+	Settings.SourceLaneIndex  = 0;
+	Settings.GroupThreadIndex = 0;
+	Settings.SwizzleAnd = 0x00;
+	Settings.SwizzleOr  = 0x00;
+	Settings.SwizzleXor = 0x00;
+	Settings.Rotate        = +0;
+	Settings.RotateFixMask = 0x00;
 	return Settings;
 }
 
 
-#if PLATFORM_SUPPORTS_WAVE_BROADCAST
+/**  Dynamically read another lane's register
+ *
+ *	Requires PLATFORM_SUPPORTS_WAVE_READ_AT.
+ *
+ *	return src[SourceLaneIndex];
+ *
+ *  Note the SourceLaneIndex can be dynamic according to SM 6.0's WaveReadLaneAt()
+ */
+CALL_SITE_DEBUGLOC
+FWaveBroadcastSettings InitWaveReadLaneAt(uint SourceLaneIndex)
+{
+	FWaveBroadcastSettings Settings = InitNopBroadcast();
+	Settings.Operation  = WAVE_BROADCAST_READ_LANE_AT;
+	Settings.SourceLaneIndex = SourceLaneIndex;
+	return Settings;
+}
+
+/**  Dynamically read another group thread's register through LDS
+ *
+ *	Requires setting WAVE_BROADCAST_GROUPSIZE to allocate groupshared memory.
+ *
+ *	LDS[GroupThreadIndex] = src;
+ *	return LDS[SourceLaneIndex];
+ */
+CALL_SITE_DEBUGLOC
+FWaveBroadcastSettings InitWaveReadLDS(uint SourceLaneIndex, uint GroupThreadIndex)
+{
+	FWaveBroadcastSettings Settings = InitNopBroadcast();
+	Settings.Operation  = WAVE_BROADCAST_READ_LDS_AT;
+	Settings.SourceLaneIndex = SourceLaneIndex;
+	Settings.GroupThreadIndex  = GroupThreadIndex;
+	return Settings;
+}
+
+/**  Converts a broadcast operation to be performed in SM5.
+ *
+ *	Requires setting WAVE_BROADCAST_GROUPSIZE to allocate groupshared memory.
+ *
+ *	Returns the exact sames data as Settings would otherwise return, but through LDS. Example:
+ *
+ *		const FWaveBroadcastSettings BroadcastSettings = InitWaveReadLaneAt(SourceLaneIndex);
+ *
+ *		#if PLATFORM_SUPPORTS_WAVE_READ_AT
+ *			uint Desc = WaveBroadcast(BroadcastSettings, Src);
+ *		#else
+ *			uint Desc = WaveBroadcast(ConvertWaveBroadcastToLDS(BroadcastSettings, GroupThreadIndex), Src);
+ *		#endif
+ */
+CALL_SITE_DEBUGLOC
+FWaveBroadcastSettings ConvertWaveBroadcastToLDS(const FWaveBroadcastSettings Settings, uint GroupThreadIndex)
+{
+	return InitWaveReadLDS(GetWaveBroadcastSourceLaneIndex(Settings, /* DestLaneIndex = */ GroupThreadIndex), GroupThreadIndex);
+}
 
 /**  Broadcast in butterfly pattern.
+ *
+ *	Requires PLATFORM_SUPPORTS_WAVE_BROADCAST.
  *
  *	return src[laneId ^ XorButterFly];
  */
@@ -303,6 +451,8 @@ FWaveBroadcastSettings InitWaveXorButterfly(const uint XorButterFly)
 }
 
 /** Swap left lane with righ lanes within lane group (size is power of two in [2; 64]). 
+ *
+ *	Requires PLATFORM_SUPPORTS_WAVE_BROADCAST.
  *
  *  If a lane is not active, the VGPR value returned is 0.
  *
@@ -321,6 +471,8 @@ FWaveBroadcastSettings InitWaveSwapWithinLaneGroup(const uint LaneGroupSize)
 }
 
 /**  Broadcast inner lane group over a lane group (size is power of two in [2; 64]).
+ *
+ *	Requires PLATFORM_SUPPORTS_WAVE_BROADCAST.
  *
  *  If a lane is not active, the VGPR value returned is 0.
  *
@@ -346,11 +498,9 @@ FWaveBroadcastSettings InitWaveBroadcastLaneGroup(const uint LaneGroupSize, cons
 	return Settings;
 }
 
-#endif // PLATFORM_SUPPORTS_WAVE_BROADCAST
-
-#if PLATFORM_SUPPORTS_WAVE_ROTATE
-
 /**  Rotate lane (-31 to +31) over a lane group (size is power of two in [2; 32]).
+ *
+ *	Requires PLATFORM_SUPPORTS_WAVE_ROTATE.
  *
  *  If a lane is not active, the VGPR value returned is 0.
  *
@@ -374,5 +524,3 @@ FWaveBroadcastSettings InitWaveRotateLaneGroup(const uint LaneGroupSize, const i
 	Settings.RotateFixMask = RotateFixMask;
 	return Settings;
 }
-
-#endif // PLATFORM_SUPPORTS_WAVE_ROTATE
diff --git a/Engine/Source/Runtime/Renderer/Private/PostProcess/TemporalAA.cpp b/Engine/Source/Runtime/Renderer/Private/PostProcess/TemporalAA.cpp
--- a/Engine/Source/Runtime/Renderer/Private/PostProcess/TemporalAA.cpp
+++ b/Engine/Source/Runtime/Renderer/Private/PostProcess/TemporalAA.cpp
@@ -365,22 +365,48 @@ static_assert(UE_ARRAY_COUNT(kTAAPassNames) == int32(ETAAPassConfig::MAX), "Miss
 static_assert(UE_ARRAY_COUNT(kTAAQualityNames) == int32(ETAAQuality::MAX), "Missing TAA quality name.");
 } //! namespace
 
-FVector ComputePixelFormatQuantizationError(EPixelFormat PixelFormat)
+FVector3f ComputePixelFormatQuantizationError(EPixelFormat PixelFormat)
 {
-	FVector Error;
-	if (PixelFormat == PF_FloatRGBA || PixelFormat == PF_FloatR11G11B10)
+	FIntVector ColorMantissaBits = FIntVector(1, 1, 1);
+	switch (PixelFormat)
 	{
-		FIntVector HistoryColorMantissaBits = PixelFormat == PF_FloatR11G11B10 ? FIntVector(6, 6, 5) : FIntVector(10, 10, 10);
+	case PF_FloatR11G11B10:
+		ColorMantissaBits = FIntVector(6, 6, 5);
+		break;
 
-		Error.X = FMath::Pow(0.5f, HistoryColorMantissaBits.X);
-		Error.Y = FMath::Pow(0.5f, HistoryColorMantissaBits.Y);
-		Error.Z = FMath::Pow(0.5f, HistoryColorMantissaBits.Z);
-	}
-	else
-	{
-		check(0);
+	case PF_FloatRGBA:
+		ColorMantissaBits = FIntVector(10, 10, 10);
+		break;
+
+	case PF_A32B32G32R32F:
+		ColorMantissaBits = FIntVector(23, 23, 23);
+		break;
+
+	case PF_R5G6B5_UNORM:
+		ColorMantissaBits = FIntVector(5, 6, 5);
+		break;
+
+	case PF_B8G8R8A8:
+	case PF_R8G8B8A8:
+		ColorMantissaBits = FIntVector(8, 8, 8);
+		break;
+
+	case PF_A2B10G10R10:
+		ColorMantissaBits = FIntVector(10, 10, 10);
+		break;
+
+	case PF_A16B16G16R16:
+		ColorMantissaBits = FIntVector(16, 16, 16);
+		break;
+
+	default:
+		unimplemented();
 	}
 
+	FVector3f Error;
+	Error.X = FMath::Pow(0.5f, ColorMantissaBits.X);
+	Error.Y = FMath::Pow(0.5f, ColorMantissaBits.Y);
+	Error.Z = FMath::Pow(0.5f, ColorMantissaBits.Z);
 	return Error;
 }
 
diff --git a/Engine/Source/Runtime/Renderer/Private/PostProcess/TemporalSuperResolution.cpp b/Engine/Source/Runtime/Renderer/Private/PostProcess/TemporalSuperResolution.cpp
--- a/Engine/Source/Runtime/Renderer/Private/PostProcess/TemporalSuperResolution.cpp
+++ b/Engine/Source/Runtime/Renderer/Private/PostProcess/TemporalSuperResolution.cpp
@@ -31,14 +31,9 @@ TAutoConsoleVariable<int32> CVarTSRHistoryUpdateQuality(
 	TEXT("Select the quality of the history update."),
 	ECVF_Scalability | ECVF_RenderThreadSafe);
 
-TAutoConsoleVariable<int32> CVarTSRHistoryHighFrequencyOnly(
-	TEXT("r.TSR.History.HighFrequencyOnly"), 0,
-	TEXT("Experimental store only high frequency in the history for performance saving (default = 0)."),
-	ECVF_Scalability | ECVF_RenderThreadSafe);
-
-TAutoConsoleVariable<int32> CVarTSRHalfResShadingRejection(
-	TEXT("r.TSR.ShadingRejection.HalfRes"), 0,
-	TEXT("Whether the shading rejection should be done at half res. Saves performance but may introduce back some flickering (default = 0)."),
+TAutoConsoleVariable<int32> CVarTSRWaveOps(
+	TEXT("r.TSR.WaveOps"), 1,
+	TEXT("Whether to use wave ops in the shading rejection heuristics"),
 	ECVF_RenderThreadSafe);
 
 TAutoConsoleVariable<int32> CVarTSRFilterShadingRejection(
@@ -122,7 +117,6 @@ BEGIN_SHADER_PARAMETER_STRUCT(FTSRCommonParameters, )
 END_SHADER_PARAMETER_STRUCT()
 
 BEGIN_SHADER_PARAMETER_STRUCT(FTSRHistoryArrayIndices, )
-	SHADER_PARAMETER(int32, LowFrequency)
 	SHADER_PARAMETER(int32, HighFrequency)
 	SHADER_PARAMETER(int32, Translucency)
 	SHADER_PARAMETER(int32, Size)
@@ -131,26 +125,24 @@ END_SHADER_PARAMETER_STRUCT()
 BEGIN_SHADER_PARAMETER_STRUCT(FTSRHistoryTextures, )
 	SHADER_PARAMETER_STRUCT(FTSRHistoryArrayIndices, ArrayIndices)
 	SHADER_PARAMETER_RDG_TEXTURE(Texture2DArray, ColorArray)
-	SHADER_PARAMETER_RDG_TEXTURE(Texture2D, HighFrequencyHalfRes)
 	SHADER_PARAMETER_RDG_TEXTURE(Texture2D, Metadata)
 	SHADER_PARAMETER_RDG_TEXTURE(Texture2D, SubpixelDetails)
 	SHADER_PARAMETER_RDG_TEXTURE(Texture2D, TranslucencyAlpha)
+	SHADER_PARAMETER_RDG_TEXTURE(Texture2D, Guide)
 END_SHADER_PARAMETER_STRUCT()
 
 BEGIN_SHADER_PARAMETER_STRUCT(FTSRHistorySRVs, )
-	SHADER_PARAMETER_RDG_TEXTURE_SRV(Texture2D, LowFrequency)
 	SHADER_PARAMETER_RDG_TEXTURE_SRV(Texture2D, HighFrequency)
-	SHADER_PARAMETER_RDG_TEXTURE_SRV(Texture2D, HighFrequencyHalfRes)
 	SHADER_PARAMETER_RDG_TEXTURE_SRV(Texture2D, Metadata)
-	SHADER_PARAMETER_RDG_TEXTURE_SRV(Texture2D, SubpixelDetails)
+	SHADER_PARAMETER_RDG_TEXTURE(Texture2D, SubpixelDetails)
 	SHADER_PARAMETER_RDG_TEXTURE_SRV(Texture2D, Translucency)
 	SHADER_PARAMETER_RDG_TEXTURE_SRV(Texture2D, TranslucencyAlpha)
+	SHADER_PARAMETER_RDG_TEXTURE(Texture2D, Guide)
 END_SHADER_PARAMETER_STRUCT()
 
 BEGIN_SHADER_PARAMETER_STRUCT(FTSRHistoryUAVs, )
 	SHADER_PARAMETER_STRUCT(FTSRHistoryArrayIndices, ArrayIndices)
 	SHADER_PARAMETER_RDG_TEXTURE_UAV(RWTexture2DArray, ColorArray)
-	SHADER_PARAMETER_RDG_TEXTURE_UAV(RWTexture2D, HighFrequencyHalfRes)
 	SHADER_PARAMETER_RDG_TEXTURE_UAV(RWTexture2D, Metadata)
 	SHADER_PARAMETER_RDG_TEXTURE_UAV(RWTexture2D, SubpixelDetails)
 	SHADER_PARAMETER_RDG_TEXTURE_UAV(RWTexture2D, TranslucencyAlpha)
@@ -166,7 +158,6 @@ END_SHADER_PARAMETER_STRUCT()
 
 enum class ETSRHistoryFormatBits : uint32
 {
-	LowFrequency = 1 << 0,
 	Translucency = 1 << 1,
 };
 ENUM_CLASS_FLAGS(ETSRHistoryFormatBits);
@@ -178,35 +169,18 @@ FTSRHistoryArrayIndices TranslateHistoryFormatBitsToArrayIndices(ETSRHistoryForm
 	ArrayIndices.HighFrequency = 0;
 	ArrayIndices.Translucency = 1;
 
-	if (EnumHasAnyFlags(HistoryFormatBits, ETSRHistoryFormatBits::LowFrequency))
-	{
-		ArrayIndices.LowFrequency = ArrayIndices.Size;
-		ArrayIndices.Size += 1;
-	}
-	else
-	{
-		ArrayIndices.LowFrequency = -1;
-	}
-
 	return ArrayIndices;
 }
 
 FTSRHistorySRVs CreateSRVs(FRDGBuilder& GraphBuilder, const FTSRHistoryTextures& Textures)
 {
 	FTSRHistorySRVs SRVs;
-	if (Textures.ArrayIndices.LowFrequency >= 0)
-	{
-		SRVs.LowFrequency = GraphBuilder.CreateSRV(FRDGTextureSRVDesc::CreateForSlice(Textures.ColorArray, Textures.ArrayIndices.LowFrequency));
-	}
 	SRVs.HighFrequency = GraphBuilder.CreateSRV(FRDGTextureSRVDesc::CreateForSlice(Textures.ColorArray, Textures.ArrayIndices.HighFrequency));
-	if (Textures.HighFrequencyHalfRes)
-	{
-		SRVs.HighFrequencyHalfRes = GraphBuilder.CreateSRV(FRDGTextureSRVDesc::Create(Textures.HighFrequencyHalfRes));
-	}
 	SRVs.Metadata = GraphBuilder.CreateSRV(FRDGTextureSRVDesc::Create(Textures.Metadata));
-	SRVs.SubpixelDetails = GraphBuilder.CreateSRV(FRDGTextureSRVDesc::Create(Textures.SubpixelDetails));
+	SRVs.SubpixelDetails = Textures.SubpixelDetails;
 	SRVs.Translucency = GraphBuilder.CreateSRV(FRDGTextureSRVDesc::CreateForSlice(Textures.ColorArray, Textures.ArrayIndices.Translucency));
 	SRVs.TranslucencyAlpha = GraphBuilder.CreateSRV(FRDGTextureSRVDesc::Create(Textures.TranslucencyAlpha)); // TODO: move into metadata
+	SRVs.Guide = Textures.Guide;
 	return SRVs;
 }
 
@@ -215,10 +189,6 @@ FTSRHistoryUAVs CreateUAVs(FRDGBuilder& GraphBuilder, const FTSRHistoryTextures&
 	FTSRHistoryUAVs UAVs;
 	UAVs.ArrayIndices = Textures.ArrayIndices;
 	UAVs.ColorArray = GraphBuilder.CreateUAV(Textures.ColorArray);
-	if (Textures.HighFrequencyHalfRes)
-	{
-		UAVs.HighFrequencyHalfRes = GraphBuilder.CreateUAV(Textures.HighFrequencyHalfRes);
-	}
 	UAVs.Metadata = GraphBuilder.CreateUAV(Textures.Metadata);
 	UAVs.SubpixelDetails = GraphBuilder.CreateUAV(Textures.SubpixelDetails);
 	UAVs.TranslucencyAlpha = GraphBuilder.CreateUAV(Textures.TranslucencyAlpha); // TODO: move into metadata
@@ -251,8 +221,6 @@ public:
 	}
 }; // class FTemporalSuperResolutionShader
 
-class FTSRHighFrequencyOnlyDim : SHADER_PERMUTATION_BOOL("DIM_HIGH_FREQUENCY_ONLY");
-
 class FTSRClearPrevTexturesCS : public FTSRShader
 {
 	DECLARE_GLOBAL_SHADER(FTSRClearPrevTexturesCS);
@@ -305,15 +273,9 @@ class FTSRDecimateHistoryCS : public FTSRShader
 	DECLARE_GLOBAL_SHADER(FTSRDecimateHistoryCS);
 	SHADER_USE_PARAMETER_STRUCT(FTSRDecimateHistoryCS, FTSRShader);
 
-	class FOutputHalfRes : SHADER_PERMUTATION_BOOL("DIM_OUTPUT_HALF_RES");
-
-	using FPermutationDomain = TShaderPermutationDomain<FTSRHighFrequencyOnlyDim, FOutputHalfRes>;
-
 	BEGIN_SHADER_PARAMETER_STRUCT(FParameters, )
 		SHADER_PARAMETER_STRUCT_INCLUDE(FTSRCommonParameters, CommonParameters)
-		SHADER_PARAMETER_STRUCT_INCLUDE(FTSRPrevHistoryParameters, PrevHistoryParameters)
 		SHADER_PARAMETER(FMatrix44f, RotationalClipToPrevClip)
-		SHADER_PARAMETER(FVector3f, OutputQuantizationError)
 		SHADER_PARAMETER(float, WorldDepthToPixelWorldRadius)
 
 		SHADER_PARAMETER_RDG_TEXTURE(Texture2D, InputSceneColorTexture)
@@ -323,14 +285,16 @@ class FTSRDecimateHistoryCS : public FTSRShader
 		SHADER_PARAMETER_RDG_TEXTURE(Texture2D, PrevClosestDepthTexture)
 		SHADER_PARAMETER_RDG_TEXTURE(Texture2D, ParallaxFactorTexture)
 
-		SHADER_PARAMETER_STRUCT(FTSRHistorySRVs, PrevHistory)
-
-		SHADER_PARAMETER_RDG_TEXTURE_UAV(RWTexture2D, HalfResSceneColorOutput)
-		SHADER_PARAMETER_RDG_TEXTURE_UAV(RWTexture2D, HalfResPredictionSceneColorOutput)
-		SHADER_PARAMETER_RDG_TEXTURE_UAV(RWTexture2D, HalfResParallaxRejectionMaskOutput)
+		SHADER_PARAMETER_STRUCT_INCLUDE(FTSRPrevHistoryParameters, PrevHistoryParameters)
+		SHADER_PARAMETER_RDG_TEXTURE(Texture2D, PrevHistorySubpixelDetails)
 
-		SHADER_PARAMETER_RDG_TEXTURE_UAV(RWTexture2D, PredictionSceneColorOutput)
+		SHADER_PARAMETER_RDG_TEXTURE(Texture2D, PrevHistoryGuide)
+		SHADER_PARAMETER_STRUCT(FScreenPassTextureViewportParameters, PrevGuideInfo)
+		SHADER_PARAMETER(FScreenTransform, InputPixelPosToReprojectScreenPos)
+		SHADER_PARAMETER(FScreenTransform, ScreenPosToPrevHistoryGuideBufferUV)
+		SHADER_PARAMETER(FVector3f, HistoryGuideQuantizationError)
 
+		SHADER_PARAMETER_RDG_TEXTURE_UAV(RWTexture2D, ReprojectedHistoryGuideOutput)
 		SHADER_PARAMETER_RDG_TEXTURE_UAV(RWTexture2D, HoleFilledVelocityOutput)
 		SHADER_PARAMETER_RDG_TEXTURE_UAV(RWTexture2D, HoleFilledVelocityMaskOutput)
 		SHADER_PARAMETER_RDG_TEXTURE_UAV(RWTexture2D, ParallaxRejectionMaskOutput)
@@ -364,47 +328,53 @@ class FTSRCompareTranslucencyCS : public FTSRShader
 	END_SHADER_PARAMETER_STRUCT()
 }; // class FTSRCompareTranslucencyCS
 
-class FTSRFilterFrequenciesCS : public FTSRShader
+class FTSRRejectShadingCS : public FTSRShader
 {
-	DECLARE_GLOBAL_SHADER(FTSRFilterFrequenciesCS);
-	SHADER_USE_PARAMETER_STRUCT(FTSRFilterFrequenciesCS, FTSRShader);
-
-	class FOutputAALumaDim : SHADER_PERMUTATION_BOOL("DIM_OUTPUT_ANTI_ALIASING_LUMA");
+	DECLARE_GLOBAL_SHADER(FTSRRejectShadingCS);
+	SHADER_USE_PARAMETER_STRUCT(FTSRRejectShadingCS, FTSRShader);
 
-	using FPermutationDomain = TShaderPermutationDomain<FOutputAALumaDim>;
+	class FUseWaveOps : SHADER_PERMUTATION_BOOL("DIM_USE_WAVE_OPS");
+	using FPermutationDomain = TShaderPermutationDomain<FUseWaveOps>;
 
 	BEGIN_SHADER_PARAMETER_STRUCT(FParameters, )
 		SHADER_PARAMETER_STRUCT_INCLUDE(FTSRCommonParameters, CommonParameters)
-		SHADER_PARAMETER(FVector3f, OutputQuantizationError)
+		SHADER_PARAMETER(FVector3f, HistoryGuideQuantizationError)
 
 		SHADER_PARAMETER_RDG_TEXTURE(Texture2D, InputTexture)
-		SHADER_PARAMETER_RDG_TEXTURE(Texture2D, PredictionSceneColorTexture)
+		SHADER_PARAMETER_RDG_TEXTURE(Texture2D, ReprojectedHistoryGuideTexture)
 		SHADER_PARAMETER_RDG_TEXTURE(Texture2D, ParallaxRejectionMaskTexture)
 
-		SHADER_PARAMETER_RDG_TEXTURE_UAV(RWTexture2D, FilteredInputOutput)
-		SHADER_PARAMETER_RDG_TEXTURE_UAV(RWTexture2D, FilteredPredictionSceneColorOutput)
+		SHADER_PARAMETER_RDG_TEXTURE_UAV(RWTexture2D, HistoryGuideOutput)
+		SHADER_PARAMETER_RDG_TEXTURE_UAV(RWTexture2D, HistoryRejectionOutput)
 		SHADER_PARAMETER_RDG_TEXTURE_UAV(RWTexture2D, InputSceneColorLdrLumaOutput)
-		SHADER_PARAMETER_RDG_TEXTURE_UAV(RWTexture2DArray, DebugOutput)
+		SHADER_PARAMETER_RDG_TEXTURE_UAV(RWTexture2D, DebugOutput)
 	END_SHADER_PARAMETER_STRUCT()
-}; // class FTSRFilterFrequenciesCS
 
-class FTSRCompareHistoryCS : public FTSRShader
-{
-	DECLARE_GLOBAL_SHADER(FTSRCompareHistoryCS);
-	SHADER_USE_PARAMETER_STRUCT(FTSRCompareHistoryCS, FTSRShader);
+	static bool ShouldCompilePermutation(const FGlobalShaderPermutationParameters& Parameters)
+	{
+		FPermutationDomain PermutationVector(Parameters.PermutationId);
 
-	class FOutputHalfRes : SHADER_PERMUTATION_BOOL("DIM_OUTPUT_HALF_RES");
-	using FPermutationDomain = TShaderPermutationDomain<FOutputHalfRes>;
+		if (PermutationVector.Get<FUseWaveOps>() && !RHISupportsWaveOperations(Parameters.Platform))
+		{
+			return false;
+		}
 
-	BEGIN_SHADER_PARAMETER_STRUCT(FParameters, )
-		SHADER_PARAMETER_STRUCT_INCLUDE(FTSRCommonParameters, CommonParameters)
-		SHADER_PARAMETER_RDG_TEXTURE(Texture2D, ParallaxRejectionMaskTexture)
-		SHADER_PARAMETER_RDG_TEXTURE(Texture2D, FilteredInputTexture)
-		SHADER_PARAMETER_RDG_TEXTURE(Texture2D, FilteredPredictionSceneColorTexture)
-		SHADER_PARAMETER_RDG_TEXTURE_UAV(RWTexture2D, HistoryRejectionOutput)
-		SHADER_PARAMETER_RDG_TEXTURE_UAV(RWTexture2DArray, DebugOutput)
-	END_SHADER_PARAMETER_STRUCT()
-}; // class FTSRCompareHistoryCS
+		return FTSRShader::ShouldCompilePermutation(Parameters);
+	}
+
+	static void ModifyCompilationEnvironment(const FGlobalShaderPermutationParameters& Parameters, FShaderCompilerEnvironment& OutEnvironment)
+	{
+		FPermutationDomain PermutationVector(Parameters.PermutationId);
+
+		FTSRShader::ModifyCompilationEnvironment(Parameters, OutEnvironment);
+		OutEnvironment.CompilerFlags.Add(CFLAG_Wave32);
+
+		if (PermutationVector.Get<FUseWaveOps>())
+		{
+			OutEnvironment.CompilerFlags.Add(CFLAG_WaveOperations);
+		}
+	}
+}; // class FTSRRejectShadingCS
 
 class FTSRPostfilterRejectionCS : public FTSRShader
 {
@@ -502,7 +472,7 @@ class FTSRUpdateHistoryCS : public FTSRShader
 
 	class FQualityDim : SHADER_PERMUTATION_ENUM_CLASS("DIM_UPDATE_QUALITY", EQuality);
 
-	using FPermutationDomain = TShaderPermutationDomain<FTSRHighFrequencyOnlyDim, FQualityDim>;
+	using FPermutationDomain = TShaderPermutationDomain<FQualityDim>;
 
 	BEGIN_SHADER_PARAMETER_STRUCT(FParameters, )
 		SHADER_PARAMETER_STRUCT_INCLUDE(FTSRCommonParameters, CommonParameters)
@@ -557,8 +527,6 @@ class FTSRResolveHistoryCS : public FTSRShader
 	DECLARE_GLOBAL_SHADER(FTSRResolveHistoryCS);
 	SHADER_USE_PARAMETER_STRUCT(FTSRResolveHistoryCS, FTSRShader);
 
-	using FPermutationDomain = TShaderPermutationDomain<FTSRHighFrequencyOnlyDim>;
-
 	BEGIN_SHADER_PARAMETER_STRUCT(FParameters, )
 		SHADER_PARAMETER_STRUCT_INCLUDE(FTSRCommonParameters, CommonParameters)
 		SHADER_PARAMETER(FScreenTransform, DispatchThreadToHistoryPixelPos)
@@ -602,8 +570,7 @@ IMPLEMENT_GLOBAL_SHADER(FTSRClearPrevTexturesCS,     "/Engine/Private/TemporalSu
 IMPLEMENT_GLOBAL_SHADER(FTSRDilateVelocityCS,        "/Engine/Private/TemporalSuperResolution/TSRDilateVelocity.usf",        "MainCS", SF_Compute);
 IMPLEMENT_GLOBAL_SHADER(FTSRDecimateHistoryCS,       "/Engine/Private/TemporalSuperResolution/TSRDecimateHistory.usf",       "MainCS", SF_Compute);
 IMPLEMENT_GLOBAL_SHADER(FTSRCompareTranslucencyCS,   "/Engine/Private/TemporalSuperResolution/TSRCompareTranslucency.usf",   "MainCS", SF_Compute);
-IMPLEMENT_GLOBAL_SHADER(FTSRFilterFrequenciesCS,     "/Engine/Private/TemporalSuperResolution/TSRFilterFrequencies.usf",     "MainCS", SF_Compute);
-IMPLEMENT_GLOBAL_SHADER(FTSRCompareHistoryCS,        "/Engine/Private/TemporalSuperResolution/TSRCompareHistory.usf",        "MainCS", SF_Compute);
+IMPLEMENT_GLOBAL_SHADER(FTSRRejectShadingCS,         "/Engine/Private/TemporalSuperResolution/TSRRejectShading.usf",         "MainCS", SF_Compute);
 IMPLEMENT_GLOBAL_SHADER(FTSRPostfilterRejectionCS,   "/Engine/Private/TemporalSuperResolution/TSRPostfilterRejection.usf",   "MainCS", SF_Compute);
 IMPLEMENT_GLOBAL_SHADER(FTSRDilateRejectionCS,       "/Engine/Private/TemporalSuperResolution/TSRDilateRejection.usf",       "MainCS", SF_Compute);
 IMPLEMENT_GLOBAL_SHADER(FTSRSpatialAntiAliasingCS,   "/Engine/Private/TemporalSuperResolution/TSRSpatialAntiAliasing.usf",   "MainCS", SF_Compute);
@@ -619,7 +586,7 @@ DECLARE_GPU_STAT(TemporalSuperResolution)
 
 } //! namespace
 
-FVector ComputePixelFormatQuantizationError(EPixelFormat PixelFormat);
+FVector3f ComputePixelFormatQuantizationError(EPixelFormat PixelFormat);
 
 bool ComposeSeparateTranslucencyInTSR(const FViewInfo& View)
 {
@@ -651,21 +618,18 @@ ITemporalUpscaler::FOutputs AddTemporalSuperResolutionPasses(
 	const bool bSetupDebugPasses = CVarTSRSetupDebugPasses.GetValueOnRenderThread() != 0;
 #endif
 
+	// Whether to use wave ops optimizations.
+	const bool bUseWaveOps = CVarTSRWaveOps.GetValueOnRenderThread() != 0 && GRHISupportsWaveOperations && RHISupportsWaveOperations(View.GetShaderPlatform());
+
 	// Whether alpha channel is supported.
 	const bool bSupportsAlpha = IsPostProcessingWithAlphaChannelSupported();
 
 	// whether TSR passes can run on async compute.
 	int32 AsyncComputePasses = GSupportsEfficientAsyncCompute ? CVarTSRAsyncCompute.GetValueOnRenderThread() : 0;
 
-	const bool bHistoryHighFrequencyOnly = CVarTSRHistoryHighFrequencyOnly.GetValueOnRenderThread() != 0;
-
 	ETSRHistoryFormatBits HistoryFormatBits;
 	{
 		HistoryFormatBits = ETSRHistoryFormatBits::Translucency;
-		if (!bHistoryHighFrequencyOnly)
-		{
-			HistoryFormatBits |= ETSRHistoryFormatBits::LowFrequency;
-		}
 	}
 
 	// Whether to use camera cut shader permutation or not.
@@ -673,18 +637,10 @@ ITemporalUpscaler::FOutputs AddTemporalSuperResolutionPasses(
 
 	FTSRUpdateHistoryCS::EQuality UpdateHistoryQuality = FTSRUpdateHistoryCS::EQuality(FMath::Clamp(CVarTSRHistoryUpdateQuality.GetValueOnRenderThread(), 0, int32(FTSRUpdateHistoryCS::EQuality::MAX) - 1));
 
-	bool bHalfResLowFrequency = CVarTSRHalfResShadingRejection.GetValueOnRenderThread() != 0;
-
 	bool bIsSeperateTranslucyTexturesValid = PassInputs.PostDOFTranslucencyResources.IsValid();
 
 	const bool bRejectSeparateTranslucency = false; // bIsSeperateTranslucyTexturesValid && CVarTSRTranslucencyPreviousFrameRejection.GetValueOnRenderThread() != 0;
 
-	// Whether the UpdateHistory pass also output half res history to reduce memory bottleneck up the FTSRDecimateHistoryCS.
-	const bool bHalfResHistory = bHistoryHighFrequencyOnly;
-
-	// Whether the subpixel details are stored at input or output resolution
-	const bool bSubPixelDetailsAtOutputResolution = !bHistoryHighFrequencyOnly;
-
 	EPixelFormat ColorFormat = bSupportsAlpha ? PF_FloatRGBA : PF_FloatR11G11B10;
 
 	int32 RejectionAntiAliasingQuality = FMath::Clamp(CVarTSRRejectionAntiAliasingQuality.GetValueOnRenderThread(), 1, 2);
@@ -700,7 +656,7 @@ ITemporalUpscaler::FOutputs AddTemporalSuperResolutionPasses(
 		PreRejectionDownsample,
 	};
 
-	ERejectionPostFilter PostFilter = ERejectionPostFilter(FMath::Clamp(CVarTSRFilterShadingRejection.GetValueOnRenderThread(), 0, 2));
+	ERejectionPostFilter PostFilter = ERejectionPostFilter(FMath::Clamp(CVarTSRFilterShadingRejection.GetValueOnRenderThread(), 2, 2)); // TODO(TSR)
 
 	FIntPoint InputExtent = PassInputs.SceneColorTexture->Desc.Extent;
 	FIntRect InputRect = View.ViewRect;
@@ -708,12 +664,6 @@ ITemporalUpscaler::FOutputs AddTemporalSuperResolutionPasses(
 	FIntPoint LowFrequencyExtent = InputExtent;
 	FIntRect LowFrequencyRect = InputRect;
 
-	if (bHalfResLowFrequency)
-	{
-		LowFrequencyExtent = InputExtent / 2;
-		LowFrequencyRect = FIntRect(FIntPoint(0, 0), FIntPoint::DivideAndRoundUp(InputRect.Size(), 2));
-	}
-
 	FIntPoint RejectionExtent = LowFrequencyExtent / 2;
 	FIntRect RejectionRect = FIntRect(FIntPoint(0, 0), FIntPoint::DivideAndRoundUp(LowFrequencyRect.Size(), 2));
 
@@ -980,46 +930,40 @@ ITemporalUpscaler::FOutputs AddTemporalSuperResolutionPasses(
 			FClearValueBinding::None,
 			TexCreate_ShaderResource | TexCreate_UAV);
 
-		FRDGTextureDesc HalfResDesc = Desc;
-		HalfResDesc.Extent /= 2;
-		if (bHalfResHistory)
-		{
-			History.HighFrequencyHalfRes = GraphBuilder.CreateTexture(HalfResDesc, TEXT("TSR.History.HighFrequenciesHalfRes"));
-		}
-
 		Desc.Format = PF_R8G8;
 		History.Metadata = GraphBuilder.CreateTexture(Desc, TEXT("TSR.History.Metadata"));
 
 		Desc.Format = PF_R8;
 		History.TranslucencyAlpha = GraphBuilder.CreateTexture(Desc, TEXT("TSR.History.TranslucencyAlpha"));
+	}
+	{
+		FRDGTextureDesc SubpixelDetailsDesc = FRDGTextureDesc::Create2D(
+			InputExtent,
+			PF_R16_UINT,
+			FClearValueBinding::None,
+			TexCreate_ShaderResource | TexCreate_UAV);
+		History.SubpixelDetails = GraphBuilder.CreateTexture(SubpixelDetailsDesc, TEXT("TSR.History.SubpixelInfo"));
+	}
+	{
+		FRDGTextureDesc Desc = FRDGTextureDesc::Create2D(
+			InputExtent,
+			bSupportsAlpha ? PF_FloatRGBA : PF_A2B10G10R10,
+			FClearValueBinding::None,
+			TexCreate_ShaderResource | TexCreate_UAV);
 
-		if (bSubPixelDetailsAtOutputResolution)
-		{
-			Desc.Format = PF_R16_UINT;
-			History.SubpixelDetails = GraphBuilder.CreateTexture(Desc, TEXT("TSR.History.SubpixelInfo"));
-		}
-		else
-		{
-			FRDGTextureDesc SubpixelDetailsDesc = FRDGTextureDesc::Create2D(
-				InputExtent,
-				PF_R16_UINT,
-				FClearValueBinding::None,
-				TexCreate_ShaderResource | TexCreate_UAV);
-			History.SubpixelDetails = GraphBuilder.CreateTexture(SubpixelDetailsDesc, TEXT("TSR.History.SubpixelInfo"));
-		}
+		History.Guide = GraphBuilder.CreateTexture(Desc, TEXT("TSR.History.Guide"));
 	}
 
 	// Setup a dummy history
 	FTSRHistorySRVs DummyHistorySRVs;
 	{
-		DummyHistorySRVs.LowFrequency = GraphBuilder.CreateSRV(FRDGTextureSRVDesc::Create(BlackDummy));
 		DummyHistorySRVs.HighFrequency = GraphBuilder.CreateSRV(FRDGTextureSRVDesc::Create(BlackDummy));
-		DummyHistorySRVs.HighFrequencyHalfRes = GraphBuilder.CreateSRV(FRDGTextureSRVDesc::Create(BlackDummy));
 		DummyHistorySRVs.Metadata = GraphBuilder.CreateSRV(FRDGTextureSRVDesc::Create(BlackDummy));
 		DummyHistorySRVs.Translucency = GraphBuilder.CreateSRV(FRDGTextureSRVDesc::Create(BlackDummy));
 		DummyHistorySRVs.TranslucencyAlpha = GraphBuilder.CreateSRV(FRDGTextureSRVDesc::Create(WhiteDummy));
+		DummyHistorySRVs.Guide = BlackDummy;
 
-		DummyHistorySRVs.SubpixelDetails = GraphBuilder.CreateSRV(FRDGTextureSRVDesc::Create(BlackUintDummy));
+		DummyHistorySRVs.SubpixelDetails = BlackUintDummy;
 	}
 
 	// Setup the previous frame history
@@ -1031,9 +975,9 @@ ITemporalUpscaler::FOutputs AddTemporalSuperResolutionPasses(
 
 		// Register filterable history
 		PrevHistory.ColorArray = GraphBuilder.RegisterExternalTexture(InputHistory.ColorArray);
-		PrevHistory.HighFrequencyHalfRes = InputHistory.HighFrequencyHalfRes.IsValid() ? GraphBuilder.RegisterExternalTexture(InputHistory.HighFrequencyHalfRes) : DummyHistorySRVs.HighFrequencyHalfRes->Desc.Texture;
 		PrevHistory.Metadata = GraphBuilder.RegisterExternalTexture(InputHistory.Metadata);
 		PrevHistory.TranslucencyAlpha = InputHistory.TranslucencyAlpha.IsValid() ? GraphBuilder.RegisterExternalTexture(InputHistory.TranslucencyAlpha) : DummyHistorySRVs.TranslucencyAlpha->Desc.Texture;
+		PrevHistory.Guide = GraphBuilder.RegisterExternalTexture(InputHistory.Guide);
 
 		// Register non-filterable history
 		PrevHistory.SubpixelDetails = GraphBuilder.RegisterExternalTexture(InputHistory.SubpixelDetails);
@@ -1050,7 +994,7 @@ ITemporalUpscaler::FOutputs AddTemporalSuperResolutionPasses(
 	{
 		// Setup prev history parameters.
 		FScreenPassTextureViewport PrevHistoryViewport(PrevHistorySRVs.HighFrequency->Desc.Texture->Desc.Extent, InputHistory.OutputViewportRect);
-		FScreenPassTextureViewport PrevSubpixelDetailsViewport(PrevHistorySRVs.SubpixelDetails->Desc.Texture->Desc.Extent, bSubPixelDetailsAtOutputResolution ? InputHistory.OutputViewportRect : InputHistory.InputViewportRect);
+		FScreenPassTextureViewport PrevSubpixelDetailsViewport(PrevHistorySRVs.SubpixelDetails->Desc.Extent, InputHistory.InputViewportRect);
 
 		if (bCameraCut)
 		{
@@ -1070,10 +1014,7 @@ ITemporalUpscaler::FOutputs AddTemporalSuperResolutionPasses(
 	}
 
 	// Decimate input to flicker at same frequency as input.
-	FRDGTextureRef HalfResInputSceneColorTexture = nullptr;
-	FRDGTextureRef HalfResPredictionSceneColorTexture = nullptr;
-	FRDGTextureRef HalfResParallaxRejectionMaskTexture = nullptr;
-	FRDGTextureRef PredictionSceneColorTexture = nullptr;
+	FRDGTextureRef ReprojectedHistoryGuideTexture = nullptr;
 	FRDGTextureRef ParallaxRejectionMaskTexture = nullptr;
 	FRDGTextureRef HoleFilledVelocityMaskTexture = nullptr;
 	{
@@ -1094,6 +1035,16 @@ ITemporalUpscaler::FOutputs AddTemporalSuperResolutionPasses(
 			HoleFilledVelocityMaskTexture = GraphBuilder.CreateTexture(Desc, TEXT("TSR.velocity.HoleFillMask"));
 		}
 
+		{
+			FRDGTextureDesc Desc = FRDGTextureDesc::Create2D(
+				InputExtent,
+				PrevHistorySRVs.Guide->Desc.Format,
+				FClearValueBinding::None,
+				/* InFlags = */ TexCreate_ShaderResource | TexCreate_UAV);
+
+			ReprojectedHistoryGuideTexture = GraphBuilder.CreateTexture(Desc, TEXT("TSR.ReprojectedHistoryGuide"));
+		}
+
 		FTSRDecimateHistoryCS::FParameters* PassParameters = GraphBuilder.AllocParameters<FTSRDecimateHistoryCS::FParameters>();
 		PassParameters->CommonParameters = CommonParameters;
 		{
@@ -1105,7 +1056,6 @@ ITemporalUpscaler::FOutputs AddTemporalSuperResolutionPasses(
 
 			PassParameters->RotationalClipToPrevClip = FMatrix44f(RotationalInvViewProj * RotationalPrevViewProj);		// LWC_TODO: Precision loss?
 		}
-		PassParameters->OutputQuantizationError = (FVector3f)ComputePixelFormatQuantizationError(ColorFormat);
 		{
 			float TanHalfFieldOfView = View.ViewMatrices.GetInvProjectionMatrix().M[0][0];
 
@@ -1121,56 +1071,31 @@ ITemporalUpscaler::FOutputs AddTemporalSuperResolutionPasses(
 		PassParameters->ParallaxFactorTexture = ParallaxFactorTexture;
 
 		PassParameters->PrevHistoryParameters = PrevHistoryParameters;
-		PassParameters->PrevHistory = PrevHistorySRVs;
+		PassParameters->PrevHistorySubpixelDetails = PrevHistorySRVs.SubpixelDetails;
 
-		if (bHalfResLowFrequency)
 		{
-			FRDGTextureDesc Desc = FRDGTextureDesc::Create2D(
-				LowFrequencyExtent,
-				ColorFormat,
-				FClearValueBinding::None,
-				/* InFlags = */ TexCreate_ShaderResource | TexCreate_UAV);
-
-			HalfResInputSceneColorTexture = GraphBuilder.CreateTexture(Desc, TEXT("TSR.HalfResInput"));
-			HalfResPredictionSceneColorTexture = GraphBuilder.CreateTexture(Desc, TEXT("TSR.Prediction.SceneColor"));
-
-			Desc.Format = PF_R8;
-			HalfResParallaxRejectionMaskTexture = GraphBuilder.CreateTexture(Desc, TEXT("TSR.HalfResParallaxRejectionMask"));
-
-			PassParameters->HalfResSceneColorOutput = GraphBuilder.CreateUAV(HalfResInputSceneColorTexture);
-			PassParameters->HalfResPredictionSceneColorOutput = GraphBuilder.CreateUAV(HalfResPredictionSceneColorTexture);
-			PassParameters->HalfResParallaxRejectionMaskOutput = GraphBuilder.CreateUAV(HalfResParallaxRejectionMaskTexture);
-		}
-		else
-		{
-			FRDGTextureDesc Desc = FRDGTextureDesc::Create2D(
-				InputExtent,
-				ColorFormat,
-				FClearValueBinding::None,
-				/* InFlags = */ TexCreate_ShaderResource | TexCreate_UAV);
-
-			PredictionSceneColorTexture = GraphBuilder.CreateTexture(Desc, TEXT("TSR.Prediction.SceneColor"));
-
-			PassParameters->PredictionSceneColorOutput = GraphBuilder.CreateUAV(PredictionSceneColorTexture);
+			FScreenPassTextureViewport PrevHistoryColorViewport(PrevHistorySRVs.Guide->Desc.Extent, InputHistory.InputViewportRect);
+			PassParameters->PrevHistoryGuide = PrevHistorySRVs.Guide;
+			PassParameters->PrevGuideInfo = GetScreenPassTextureViewportParameters(PrevHistoryColorViewport);
+			PassParameters->InputPixelPosToReprojectScreenPos = ((FScreenTransform::Identity - InputRect.Min + 0.5f) / InputRect.Size()) * FScreenTransform::ViewportUVToScreenPos;
+			PassParameters->ScreenPosToPrevHistoryGuideBufferUV = FScreenTransform::ChangeTextureBasisFromTo(
+				PrevHistoryColorViewport,
+				FScreenTransform::ETextureBasis::ScreenPosition,
+				FScreenTransform::ETextureBasis::TextureUV);
+			PassParameters->HistoryGuideQuantizationError = ComputePixelFormatQuantizationError(PrevHistorySRVs.Guide->Desc.Format);
 		}
 
+		PassParameters->ReprojectedHistoryGuideOutput = GraphBuilder.CreateUAV(ReprojectedHistoryGuideTexture);
 		PassParameters->HoleFilledVelocityOutput = GraphBuilder.CreateUAV(HoleFilledVelocityTexture);
 		PassParameters->HoleFilledVelocityMaskOutput = GraphBuilder.CreateUAV(HoleFilledVelocityMaskTexture);
 		PassParameters->ParallaxRejectionMaskOutput = GraphBuilder.CreateUAV(ParallaxRejectionMaskTexture);
 		PassParameters->HistorySubpixelDetailsOutput = GraphBuilder.CreateUAV(History.SubpixelDetails);
 		PassParameters->DebugOutput = CreateDebugUAV(LowFrequencyExtent, TEXT("Debug.TSR.DecimateHistory"));
 
-		FTSRDecimateHistoryCS::FPermutationDomain PermutationVector;
-		PermutationVector.Set<FTSRHighFrequencyOnlyDim>(bHistoryHighFrequencyOnly);
-		PermutationVector.Set<FTSRDecimateHistoryCS::FOutputHalfRes>(bHalfResLowFrequency);
-
-		TShaderMapRef<FTSRDecimateHistoryCS> ComputeShader(View.ShaderMap, PermutationVector);
+		TShaderMapRef<FTSRDecimateHistoryCS> ComputeShader(View.ShaderMap);
 		FComputeShaderUtils::AddPass(
 			GraphBuilder,
-			RDG_EVENT_NAME("TSR DecimateHistory(%s%s) %dx%d",
-				bHistoryHighFrequencyOnly ? TEXT("HighFrequency") : TEXT("LowFrequency"),
-				bHalfResLowFrequency ? TEXT(" HalfResShadingOutput") : TEXT(""),
-				InputRect.Width(), InputRect.Height()),
+			RDG_EVENT_NAME("TSR DecimateHistory %dx%d", InputRect.Width(), InputRect.Height()),
 			AsyncComputePasses >= 2 ? ERDGPassFlags::AsyncCompute : ERDGPassFlags::Compute,
 			ComputeShader,
 			PassParameters,
@@ -1241,104 +1166,54 @@ ITemporalUpscaler::FOutputs AddTemporalSuperResolutionPasses(
 	FRDGTextureRef HistoryRejectionTexture;
 	FRDGTextureRef InputSceneColorLdrLumaTexture = nullptr;
 	{
-		// Filter out the high frquencies
-		FRDGTextureRef FilteredInputTexture;
-		FRDGTextureRef FilteredPredictionSceneColorTexture;
-		{
-			{
-				FRDGTextureDesc Desc = FRDGTextureDesc::Create2D(
-					LowFrequencyExtent,
-					ColorFormat,
-					FClearValueBinding::None,
-					/* InFlags = */ TexCreate_ShaderResource | TexCreate_UAV);
+		bool bOutputHalfRes = PostFilter != ERejectionPostFilter::PreRejectionDownsample;
+		check(!bOutputHalfRes); // TODO(TSR)
 
-				FilteredInputTexture = GraphBuilder.CreateTexture(Desc, TEXT("TSR.Filtered.SceneColor"));
-				FilteredPredictionSceneColorTexture = GraphBuilder.CreateTexture(Desc, TEXT("TSR.Filtered.Prediction.SceneColor"));
-			}
-			
-			if (RejectionAntiAliasingQuality > 0)
-			{
-				FRDGTextureDesc Desc = FRDGTextureDesc::Create2D(
-					LowFrequencyExtent,
-					PF_R8,
-					FClearValueBinding::None,
-					/* InFlags = */ TexCreate_ShaderResource | TexCreate_UAV);
-
-				InputSceneColorLdrLumaTexture = GraphBuilder.CreateTexture(Desc, TEXT("TSR.SceneColorLdrLuma"));
-			}
-
-			FTSRFilterFrequenciesCS::FParameters* PassParameters = GraphBuilder.AllocParameters<FTSRFilterFrequenciesCS::FParameters>();
-			PassParameters->CommonParameters = CommonParameters;
-			PassParameters->OutputQuantizationError = (FVector3f)ComputePixelFormatQuantizationError(FilteredInputTexture->Desc.Format);
-
-			if (bHalfResLowFrequency)
-			{
-				PassParameters->InputTexture = HalfResInputSceneColorTexture;
-				PassParameters->PredictionSceneColorTexture = HalfResPredictionSceneColorTexture;
-				PassParameters->ParallaxRejectionMaskTexture = HalfResParallaxRejectionMaskTexture;
-			}
-			else
-			{
-				PassParameters->InputTexture = PassInputs.SceneColorTexture;
-				PassParameters->PredictionSceneColorTexture = PredictionSceneColorTexture;
-				PassParameters->ParallaxRejectionMaskTexture = ParallaxRejectionMaskTexture;
-			}
-
-			PassParameters->FilteredInputOutput = GraphBuilder.CreateUAV(FilteredInputTexture);
-			PassParameters->FilteredPredictionSceneColorOutput = GraphBuilder.CreateUAV(FilteredPredictionSceneColorTexture);
-			PassParameters->InputSceneColorLdrLumaOutput = InputSceneColorLdrLumaTexture ? GraphBuilder.CreateUAV(InputSceneColorLdrLumaTexture) : nullptr;
-			PassParameters->DebugOutput = CreateDebugUAV(LowFrequencyExtent, TEXT("Debug.TSR.FilterFrequencies"));
-
-			FTSRFilterFrequenciesCS::FPermutationDomain PermutationVector;
-			PermutationVector.Set<FTSRFilterFrequenciesCS::FOutputAALumaDim>(RejectionAntiAliasingQuality > 0);
+		{
+			FRDGTextureDesc Desc = FRDGTextureDesc::Create2D(
+				LowFrequencyExtent,
+				PF_R8,
+				FClearValueBinding::None,
+				/* InFlags = */ TexCreate_ShaderResource | TexCreate_UAV);
 
-			TShaderMapRef<FTSRFilterFrequenciesCS> ComputeShader(View.ShaderMap, PermutationVector);
-			FComputeShaderUtils::AddPass(
-				GraphBuilder,
-				RDG_EVENT_NAME("TSR FilterFrequencies(%s) %dx%d",
-					PermutationVector.Get<FTSRFilterFrequenciesCS::FOutputAALumaDim>() ? TEXT("OutputAALuma") : TEXT(""),
-					LowFrequencyRect.Width(), LowFrequencyRect.Height()),
-				AsyncComputePasses >= 3 ? ERDGPassFlags::AsyncCompute : ERDGPassFlags::Compute,
-				ComputeShader,
-				PassParameters,
-				FComputeShaderUtils::GetGroupCount(LowFrequencyRect.Size(), 16));
+			InputSceneColorLdrLumaTexture = GraphBuilder.CreateTexture(Desc, TEXT("TSR.SceneColorLdrLuma"));
 		}
 
-		// Compare the low frequencies
 		{
-			bool bOutputHalfRes = PostFilter != ERejectionPostFilter::PreRejectionDownsample;
-
-			{
-				FRDGTextureDesc Desc = FRDGTextureDesc::Create2D(
-					bOutputHalfRes ? RejectionExtent : LowFrequencyExtent,
-					PF_R8,
-					FClearValueBinding::None,
-					/* InFlags = */ TexCreate_ShaderResource | TexCreate_UAV);
+			FRDGTextureDesc Desc = FRDGTextureDesc::Create2D(
+				bOutputHalfRes ? RejectionExtent : LowFrequencyExtent,
+				PF_R8,
+				FClearValueBinding::None,
+				/* InFlags = */ TexCreate_ShaderResource | TexCreate_UAV);
 
-				HistoryRejectionTexture = GraphBuilder.CreateTexture(Desc, TEXT("TSR.HistoryRejection"));
-			}
+			HistoryRejectionTexture = GraphBuilder.CreateTexture(Desc, TEXT("TSR.HistoryRejection"));
+		}
 
-			FTSRCompareHistoryCS::FParameters* PassParameters = GraphBuilder.AllocParameters<FTSRCompareHistoryCS::FParameters>();
-			PassParameters->CommonParameters = CommonParameters;
-			PassParameters->ParallaxRejectionMaskTexture = ParallaxRejectionMaskTexture;
-			PassParameters->FilteredInputTexture = FilteredInputTexture;
-			PassParameters->FilteredPredictionSceneColorTexture = FilteredPredictionSceneColorTexture;
+		FTSRRejectShadingCS::FParameters* PassParameters = GraphBuilder.AllocParameters<FTSRRejectShadingCS::FParameters>();
+		PassParameters->CommonParameters = CommonParameters;
+		PassParameters->InputTexture = PassInputs.SceneColorTexture;
+		PassParameters->ReprojectedHistoryGuideTexture = ReprojectedHistoryGuideTexture;
+		PassParameters->ParallaxRejectionMaskTexture = ParallaxRejectionMaskTexture;
+		PassParameters->HistoryGuideQuantizationError = ComputePixelFormatQuantizationError(History.Guide->Desc.Format);
 
-			PassParameters->HistoryRejectionOutput = GraphBuilder.CreateUAV(HistoryRejectionTexture);
-			PassParameters->DebugOutput = CreateDebugUAV(LowFrequencyExtent, TEXT("Debug.TSR.CompareHistory"));
+		PassParameters->HistoryGuideOutput = GraphBuilder.CreateUAV(History.Guide);
+		PassParameters->HistoryRejectionOutput = GraphBuilder.CreateUAV(HistoryRejectionTexture);
+		PassParameters->InputSceneColorLdrLumaOutput = GraphBuilder.CreateUAV(InputSceneColorLdrLumaTexture);
+		PassParameters->DebugOutput = CreateDebugUAV(LowFrequencyExtent, TEXT("Debug.TSR.RejectShading"));
 
-			FTSRCompareHistoryCS::FPermutationDomain PermutationVector;
-			PermutationVector.Set<FTSRCompareHistoryCS::FOutputHalfRes>(bOutputHalfRes);
+		FTSRRejectShadingCS::FPermutationDomain PermutationVector;
+		PermutationVector.Set<FTSRRejectShadingCS::FUseWaveOps>(bUseWaveOps && GRHIMinimumWaveSize <= 32 && GRHIMaximumWaveSize >= 32);
 
-			TShaderMapRef<FTSRCompareHistoryCS> ComputeShader(View.ShaderMap, PermutationVector);
-			FComputeShaderUtils::AddPass(
-				GraphBuilder,
-				RDG_EVENT_NAME("TSR CompareHistory %dx%d", LowFrequencyRect.Width(), LowFrequencyRect.Height()),
-				AsyncComputePasses >= 3 ? ERDGPassFlags::AsyncCompute : ERDGPassFlags::Compute,
-				ComputeShader,
-				PassParameters,
-				FComputeShaderUtils::GetGroupCount(LowFrequencyRect.Size(), 16));
-		}
+		TShaderMapRef<FTSRRejectShadingCS> ComputeShader(View.ShaderMap, PermutationVector);
+		FComputeShaderUtils::AddPass(
+			GraphBuilder,
+			RDG_EVENT_NAME("TSR RejectShading(%s) %dx%d",
+				PermutationVector.Get<FTSRRejectShadingCS::FUseWaveOps>() ? TEXT("WaveOps") : TEXT(""),
+				LowFrequencyRect.Width(), LowFrequencyRect.Height()),
+			AsyncComputePasses >= 3 ? ERDGPassFlags::AsyncCompute : ERDGPassFlags::Compute,
+			ComputeShader,
+			PassParameters,
+			FComputeShaderUtils::GetGroupCount(LowFrequencyRect.Size(), 12));
 	}
 
 	// Spatial anti-aliasing when doing history rejection.
@@ -1518,7 +1393,7 @@ ITemporalUpscaler::FOutputs AddTemporalSuperResolutionPasses(
 		PassParameters->HistoryPixelPosToScreenPos = HistoryPixelPosToViewportUV * FScreenTransform::ViewportUVToScreenPos;
 		PassParameters->HistoryPixelPosToInputPPCo = HistoryPixelPosToViewportUV * CommonParameters.InputInfo.ViewportSize + CommonParameters.InputJitter + CommonParameters.InputPixelPosMin;
 		PassParameters->HistoryPixelPosToTranslucencyPPCo = HistoryPixelPosToViewportUV * PassParameters->TranslucencyInfo.ViewportSize + CommonParameters.InputJitter * PassParameters->TranslucencyInfo.ViewportSize / CommonParameters.InputInfo.ViewportSize + SeparateTranslucencyRect.Min;
-		PassParameters->HistoryQuantizationError = (FVector3f)ComputePixelFormatQuantizationError(History.ColorArray->Desc.Format);
+		PassParameters->HistoryQuantizationError = ComputePixelFormatQuantizationError(History.ColorArray->Desc.Format);
 		PassParameters->MinTranslucencyRejection = TranslucencyRejectionTexture == nullptr ? 1.0 : 0.0;
 		PassParameters->InvWeightClampingPixelSpeed = 1.0f / CVarTSRWeightClampingPixelSpeed.GetValueOnRenderThread();
 		PassParameters->InputToHistoryFactor = float(HistorySize.X) / float(InputRect.Width());
@@ -1558,16 +1433,14 @@ ITemporalUpscaler::FOutputs AddTemporalSuperResolutionPasses(
 		PassParameters->DebugOutput = CreateDebugUAV(HistoryExtent, TEXT("Debug.TSR.UpdateHistory"));
 
 		FTSRUpdateHistoryCS::FPermutationDomain PermutationVector;
-		PermutationVector.Set<FTSRHighFrequencyOnlyDim>(bHistoryHighFrequencyOnly);
 		PermutationVector.Set<FTSRUpdateHistoryCS::FQualityDim>(UpdateHistoryQuality);
 
 		TShaderMapRef<FTSRUpdateHistoryCS> ComputeShader(View.ShaderMap, PermutationVector);
 		FComputeShaderUtils::AddPass(
 			GraphBuilder,
-			RDG_EVENT_NAME("TSR UpdateHistory(Quality=%s %s%s%s) %dx%d",
+			RDG_EVENT_NAME("TSR UpdateHistory(Quality=%s %s%s) %dx%d",
 				kUpdateQualityNames[int32(PermutationVector.Get<FTSRUpdateHistoryCS::FQualityDim>())],
 				History.ColorArray->Desc.Format == PF_FloatR11G11B10 ? TEXT("R11G11B10") : TEXT(""),
-				PermutationVector.Get<FTSRHighFrequencyOnlyDim>() ? TEXT("") : TEXT(" LowFrequency"),
 				PassParameters->bGenerateOutputMip1 ? TEXT(" OutputMip1") : TEXT(""),
 				HistorySize.X, HistorySize.Y),
 			AsyncComputePasses >= 3 ? ERDGPassFlags::AsyncCompute : ERDGPassFlags::Compute,
@@ -1631,10 +1504,7 @@ ITemporalUpscaler::FOutputs AddTemporalSuperResolutionPasses(
 		}
 		PassParameters->DebugOutput = CreateDebugUAV(HistoryExtent, TEXT("Debug.TSR.ResolveHistory"));
 
-		FTSRResolveHistoryCS::FPermutationDomain PermutationVector;
-		PermutationVector.Set<FTSRHighFrequencyOnlyDim>(bHistoryHighFrequencyOnly);
-
-		TShaderMapRef<FTSRResolveHistoryCS> ComputeShader(View.ShaderMap, PermutationVector);
+		TShaderMapRef<FTSRResolveHistoryCS> ComputeShader(View.ShaderMap);
 		FComputeShaderUtils::AddPass(
 			GraphBuilder,
 			RDG_EVENT_NAME("TSR ResolveHistory %dx%d", OutputRect.Width(), OutputRect.Height()),
@@ -1660,11 +1530,8 @@ ITemporalUpscaler::FOutputs AddTemporalSuperResolutionPasses(
 		// Extract non-filterable history
 		GraphBuilder.QueueTextureExtraction(History.SubpixelDetails, &OutputHistory.SubpixelDetails);
 
-		// Extract half res history
-		if (bHalfResHistory)
-		{
-			GraphBuilder.QueueTextureExtraction(History.HighFrequencyHalfRes, &OutputHistory.HighFrequencyHalfRes);
-		}
+		// Extract history guide
+		GraphBuilder.QueueTextureExtraction(History.Guide, &OutputHistory.Guide);
 
 		// Extract the translucency buffer to compare it with next frame
 		if (bRejectSeparateTranslucency)
diff --git a/Engine/Source/Runtime/Renderer/Private/SceneRendering.h b/Engine/Source/Runtime/Renderer/Private/SceneRendering.h
--- a/Engine/Source/Runtime/Renderer/Private/SceneRendering.h
+++ b/Engine/Source/Runtime/Renderer/Private/SceneRendering.h
@@ -903,8 +903,8 @@ struct FTSRHistory
 	// Non-filterable output resolution
 	TRefCountPtr<IPooledRenderTarget> SubpixelDetails;
 
-	// Filterable half-res output resolution
-	TRefCountPtr<IPooledRenderTarget> HighFrequencyHalfRes;
+	// Filterable input resolution representation of the output
+	TRefCountPtr<IPooledRenderTarget> Guide;
 
 	FIntRect InputViewportRect;
 	FIntRect OutputViewportRect;
