Unified all shader print buffer.

This saves one UAV slots when using ShaderPrint.

diff --git a/Engine/Shaders/Private/ShaderPrint.ush b/Engine/Shaders/Private/ShaderPrint.ush
--- a/Engine/Shaders/Private/ShaderPrint.ush
+++ b/Engine/Shaders/Private/ShaderPrint.ush
@@ -129,8 +129,8 @@
 #define SHADER_PRINT_COUNTER_COUNT           4
 
 uint GetSymbolOffset(uint InIndex)				{ return SHADER_PRINT_COUNTER_COUNT + SHADER_PRINT_UINT_COUNT_SYMBOL * InIndex; }
-uint GetPrimitiveLineOffset(uint InIndex)		{ return SHADER_PRINT_COUNTER_COUNT + SHADER_PRINT_UINT_COUNT_LINE * InIndex; }
-uint GetPrimitiveTriangleOffset(uint InIndex)	{ return SHADER_PRINT_COUNTER_COUNT + SHADER_PRINT_UINT_COUNT_LINE * ShaderPrintData.MaxLineCount + SHADER_PRINT_UINT_COUNT_TRIANGLE * InIndex; }
+uint GetPrimitiveLineOffset(uint InIndex)		{ return SHADER_PRINT_COUNTER_COUNT + SHADER_PRINT_UINT_COUNT_SYMBOL * ShaderPrintData.MaxValueCount + SHADER_PRINT_UINT_COUNT_LINE * InIndex; }
+uint GetPrimitiveTriangleOffset(uint InIndex)	{ return SHADER_PRINT_COUNTER_COUNT + SHADER_PRINT_UINT_COUNT_SYMBOL * ShaderPrintData.MaxValueCount + SHADER_PRINT_UINT_COUNT_LINE * ShaderPrintData.MaxLineCount + SHADER_PRINT_UINT_COUNT_TRIANGLE * InIndex; }
 
 void ClearCounters(RWStructuredBuffer<uint> InRWBuffer)
 {
@@ -159,21 +159,13 @@ void ClearCounters(RWStructuredBuffer<uint> InRWBuffer)
 
 StructuredBuffer<uint> ShaderPrint_StateBuffer;
 
-// Symbols
-//  ----------------------------------------------------------------------------------------------
-// ||  Line    | Triangle |  Symbol |  Empty  ||  Symb0  |  Symb1  |  Symb2  |  Symb3  |  ...  ...
-// || Counter  | Counter  | Counter | Counter ||         |         |         |         |  ...  ...
-// || (1uint)  | (1uint)  | (1uint) | (1uint) || (4uint) | (4uint) | (4uint) | (4uint) |  ...  ...
-//  ----------------------------------------------------------------------------------------------
-RWStructuredBuffer<uint> ShaderPrint_RWValuesBuffer;
-
-// Primitive buffer layout
-//  ------------------------------------------------------------------------------------------------------------------------------------------------
-// ||  Line    | Triangle |  Empty  |  Empty  ||  Line0  |  Line1  |  Line2  |  Line3  |  ...  | Triangle0 | Triangle1 | Triangle2 | Triangle3 | ...
-// || Counter  | Counter  | Counter | Counter ||         |         |         |         |  ...  |           |           |           |           | ...
-// || (1uint)  | (1uint)  | (1uint) | (1uint) || (8uint) | (8uint) | (8uint) | (8uint) |  ...  | (12uint)  | (12uint)  | (12uint)  | (12uint)  | ...
-//  ------------------------------------------------------------------------------------------------------------------------------------------------
-RWStructuredBuffer<uint> ShaderPrint_RWPrimitivesBuffer;
+// Primitive & Symbols layout
+//  ----------------------------------------------------------------------------------------------------------------------------------------------------------------------------------------------
+// ||  Line    | Triangle |  Symbol |  Empty  ||  Symb0  |  Symb1  |  Symb2  |  Symb3  |  ...  |  Line0  |  Line1  |  Line2  |  Line3  |  ...  | Triangle0 | Triangle1 | Triangle2 | Triangle3 | ...
+// || Counter  | Counter  | Counter | Counter ||         |         |         |         |  ...  |         |         |         |         |  ...  |           |           |           |           | ...
+// || (1uint)  | (1uint)  | (1uint) | (1uint) || (4uint) | (4uint) | (4uint) | (4uint) |  ...  | (8uint) | (8uint) | (8uint) | (8uint) |  ...  | (12uint)  | (12uint)  | (12uint)  | (12uint)  | ...
+//  ----------------------------------------------------------------------------------------------------------------------------------------------------------------------------------------------
+RWStructuredBuffer<uint> ShaderPrint_RWEntryBuffer;
 
 ///////////////////////////////////////////////////////////////////////////////////////////////////
 // Font Color
@@ -356,7 +348,7 @@ void ShaderPrint_Internal(in FShaderPrintItem Item)
 
 	// Buffer counter is stored in first element .Value
 	int IndexToStore = 0;
-	InterlockedAdd(ShaderPrint_RWValuesBuffer[SHADER_PRINT_COUNTER_OFFSET_SYMBOL], 1, IndexToStore);
+	InterlockedAdd(ShaderPrint_RWEntryBuffer[SHADER_PRINT_COUNTER_OFFSET_SYMBOL], 1, IndexToStore);
 
 	// Prevent writing off the buffer
 	// Note that counter still increases so need clamp when reading it in later passes
@@ -365,7 +357,7 @@ void ShaderPrint_Internal(in FShaderPrintItem Item)
 		return;
 	}
 
-	WriteSymbol(IndexToStore, Item, ShaderPrint_RWValuesBuffer);
+	WriteSymbol(IndexToStore, Item, ShaderPrint_RWEntryBuffer);
 }
 
 void ShaderPrint_Internal(in float2 ScreenPos, in int Value, in FFontColor FontColor, in uint Metadata, in int Type)
@@ -433,7 +425,7 @@ float2 ShaderPrintText_Internal(bool bIsActive, float2 InPos, FShaderPrintText I
 		// Prevent writing off the buffer
 		// Note that counter still increases so need clamp when reading it in later passes
 		int IndexToStore = 0;
-		InterlockedAdd(ShaderPrint_RWValuesBuffer[SHADER_PRINT_COUNTER_OFFSET_SYMBOL], Count, IndexToStore);
+		InterlockedAdd(ShaderPrint_RWEntryBuffer[SHADER_PRINT_COUNTER_OFFSET_SYMBOL], Count, IndexToStore);
 		if (uint(IndexToStore + Count) < ShaderPrintData.MaxValueCount)
 		{
 			for (uint i = Begin; i < End; ++i)
@@ -445,7 +437,7 @@ float2 ShaderPrintText_Internal(bool bIsActive, float2 InPos, FShaderPrintText I
 				Item.Color = InColor.Color;
 				Item.Metadata = 0u;
 
-				WriteSymbol(IndexToStore, Item, ShaderPrint_RWValuesBuffer);
+				WriteSymbol(IndexToStore, Item, ShaderPrint_RWEntryBuffer);
 
 				++IndexToStore;
 				InPos.x += ShaderPrintData.FontSpacing.x;
@@ -876,7 +868,7 @@ bool AllocateTriangleElement(uint Count, inout uint OutIndex)
 		return false;
 	}
 
-	InterlockedAdd(ShaderPrint_RWPrimitivesBuffer[SHADER_PRINT_COUNTER_OFFSET_TRIANGLE], Count, OutIndex);
+	InterlockedAdd(ShaderPrint_RWEntryBuffer[SHADER_PRINT_COUNTER_OFFSET_TRIANGLE], Count, OutIndex);
 	return (OutIndex + Count) < ShaderPrintData.MaxTriangleCount;
 }
 
@@ -885,20 +877,20 @@ void AddTriangleElement(FTriangleElement In, uint Index)
 	const uint Index12 = GetPrimitiveTriangleOffset(Index);
 
 	FPackedTriangleElement Out = PackTriangleElement(In);
-	ShaderPrint_RWPrimitivesBuffer[Index12 + 0] = Out.Packed0.x;
-	ShaderPrint_RWPrimitivesBuffer[Index12 + 1] = Out.Packed0.y;
-	ShaderPrint_RWPrimitivesBuffer[Index12 + 2] = Out.Packed0.z;
-	ShaderPrint_RWPrimitivesBuffer[Index12 + 3] = Out.Packed0.w;
-
-	ShaderPrint_RWPrimitivesBuffer[Index12 + 4] = Out.Packed1.x;
-	ShaderPrint_RWPrimitivesBuffer[Index12 + 5] = Out.Packed1.y;
-	ShaderPrint_RWPrimitivesBuffer[Index12 + 6] = Out.Packed1.z;
-	ShaderPrint_RWPrimitivesBuffer[Index12 + 7] = Out.Packed1.w;
-
-	ShaderPrint_RWPrimitivesBuffer[Index12 + 8] = Out.Packed2.x;
-	ShaderPrint_RWPrimitivesBuffer[Index12 + 9] = Out.Packed2.y;
-	ShaderPrint_RWPrimitivesBuffer[Index12 + 10]= Out.Packed2.z;
-	ShaderPrint_RWPrimitivesBuffer[Index12 + 11]= Out.Packed2.w;
+	ShaderPrint_RWEntryBuffer[Index12 + 0] = Out.Packed0.x;
+	ShaderPrint_RWEntryBuffer[Index12 + 1] = Out.Packed0.y;
+	ShaderPrint_RWEntryBuffer[Index12 + 2] = Out.Packed0.z;
+	ShaderPrint_RWEntryBuffer[Index12 + 3] = Out.Packed0.w;
+
+	ShaderPrint_RWEntryBuffer[Index12 + 4] = Out.Packed1.x;
+	ShaderPrint_RWEntryBuffer[Index12 + 5] = Out.Packed1.y;
+	ShaderPrint_RWEntryBuffer[Index12 + 6] = Out.Packed1.z;
+	ShaderPrint_RWEntryBuffer[Index12 + 7] = Out.Packed1.w;
+
+	ShaderPrint_RWEntryBuffer[Index12 + 8] = Out.Packed2.x;
+	ShaderPrint_RWEntryBuffer[Index12 + 9] = Out.Packed2.y;
+	ShaderPrint_RWEntryBuffer[Index12 + 10]= Out.Packed2.z;
+	ShaderPrint_RWEntryBuffer[Index12 + 11]= Out.Packed2.w;
 }
 
 ///////////////////////////////////////////////////////////////////////////////////////////////////
@@ -1050,7 +1042,7 @@ bool AllocateLineElement(uint Count, inout uint OutIndex)
 		return false;
 	}
 
-	InterlockedAdd(ShaderPrint_RWPrimitivesBuffer[SHADER_PRINT_COUNTER_OFFSET_LINE], Count, OutIndex);
+	InterlockedAdd(ShaderPrint_RWEntryBuffer[SHADER_PRINT_COUNTER_OFFSET_LINE], Count, OutIndex);
 	return (OutIndex + Count) < ShaderPrintData.MaxLineCount;
 }
 
@@ -1059,15 +1051,15 @@ void AddLineElement(FLineElement In, uint Index)
 	const uint Index8 = GetPrimitiveLineOffset(Index);
 
 	FPackedLineElement Out = PackLineElement(In);
-	ShaderPrint_RWPrimitivesBuffer[Index8 + 0] = Out.Packed0.x;
-	ShaderPrint_RWPrimitivesBuffer[Index8 + 1] = Out.Packed0.y;
-	ShaderPrint_RWPrimitivesBuffer[Index8 + 2] = Out.Packed0.z;
-	ShaderPrint_RWPrimitivesBuffer[Index8 + 3] = Out.Packed0.w;
-
-	ShaderPrint_RWPrimitivesBuffer[Index8 + 4] = Out.Packed1.x;
-	ShaderPrint_RWPrimitivesBuffer[Index8 + 5] = Out.Packed1.y;
-	ShaderPrint_RWPrimitivesBuffer[Index8 + 6] = Out.Packed1.z;
-	ShaderPrint_RWPrimitivesBuffer[Index8 + 7] = Out.Packed1.w;
+	ShaderPrint_RWEntryBuffer[Index8 + 0] = Out.Packed0.x;
+	ShaderPrint_RWEntryBuffer[Index8 + 1] = Out.Packed0.y;
+	ShaderPrint_RWEntryBuffer[Index8 + 2] = Out.Packed0.z;
+	ShaderPrint_RWEntryBuffer[Index8 + 3] = Out.Packed0.w;
+
+	ShaderPrint_RWEntryBuffer[Index8 + 4] = Out.Packed1.x;
+	ShaderPrint_RWEntryBuffer[Index8 + 5] = Out.Packed1.y;
+	ShaderPrint_RWEntryBuffer[Index8 + 6] = Out.Packed1.z;
+	ShaderPrint_RWEntryBuffer[Index8 + 7] = Out.Packed1.w;
 }
 
 int2 GetCursorPos()
diff --git a/Engine/Shaders/Private/ShaderPrintDraw.usf b/Engine/Shaders/Private/ShaderPrintDraw.usf
--- a/Engine/Shaders/Private/ShaderPrintDraw.usf
+++ b/Engine/Shaders/Private/ShaderPrintDraw.usf
@@ -28,7 +28,7 @@ RWBuffer<uint> RWIndirectDrawArgsBuffer;
 // Other shaders should only print debug values after this pass
 
 [numthreads(1, 1, 1)]
-void InitValueBufferCS()
+void ClearCounterCS()
 {
 	// Symbol/print Counter
 	ClearCounters(RWValuesBuffer);
diff --git a/Engine/Shaders/Private/ShaderPrintDrawPrimitive.usf b/Engine/Shaders/Private/ShaderPrintDrawPrimitive.usf
--- a/Engine/Shaders/Private/ShaderPrintDrawPrimitive.usf
+++ b/Engine/Shaders/Private/ShaderPrintDrawPrimitive.usf
@@ -107,17 +107,6 @@ void ShaderDrawDebugPS(
 }
 #endif
 
-#if GPU_DEBUG_RENDERING_CLEAR_CS
-
-RWStructuredBuffer<uint> RWElementBuffer;
-
-[numthreads(1, 1, 1)]
-void ShaderDrawDebugClearCS(uint3 DispatchThreadId : SV_DispatchThreadID)
-{
-	ClearCounters(RWElementBuffer);
-}
-#endif
-
 #if GPU_DEBUG_RENDERING_COPY_CS
 
 StructuredBuffer<uint> ElementBuffer;
diff --git a/Engine/Source/Runtime/Renderer/Private/ScenePrivate.h b/Engine/Source/Runtime/Renderer/Private/ScenePrivate.h
--- a/Engine/Source/Runtime/Renderer/Private/ScenePrivate.h
+++ b/Engine/Source/Runtime/Renderer/Private/ScenePrivate.h
@@ -576,7 +576,7 @@ struct FHLODSceneNodeVisibilityState
 struct FShaderPrintStateData
 {
 	TRefCountPtr<FRDGPooledBuffer> StateBuffer;
-	TRefCountPtr<FRDGPooledBuffer> PrimitiveBuffer;
+	TRefCountPtr<FRDGPooledBuffer> EntryBuffer;
 	FVector PreViewTranslation = FVector::ZeroVector;
 	bool bIsLocked = false;
 
@@ -585,7 +585,7 @@ struct FShaderPrintStateData
 		bIsLocked = false;
 		PreViewTranslation = FVector::ZeroVector;
 		StateBuffer = nullptr;
-		PrimitiveBuffer = nullptr;
+		EntryBuffer = nullptr;
 	}
 };
 
diff --git a/Engine/Source/Runtime/Renderer/Private/ShaderPrint.cpp b/Engine/Source/Runtime/Renderer/Private/ShaderPrint.cpp
--- a/Engine/Source/Runtime/Renderer/Private/ShaderPrint.cpp
+++ b/Engine/Source/Runtime/Renderer/Private/ShaderPrint.cpp
@@ -178,8 +178,7 @@ namespace ShaderPrint
 	{
 		OutParameters.Common = InData.UniformBuffer;
 		OutParameters.ShaderPrint_StateBuffer = GraphBuilder.CreateSRV(InData.ShaderPrintStateBuffer);
-		OutParameters.ShaderPrint_RWValuesBuffer = GraphBuilder.CreateUAV(InData.ShaderPrintValueBuffer);
-		OutParameters.ShaderPrint_RWPrimitivesBuffer = GraphBuilder.CreateUAV(InData.ShaderPrintPrimitiveBuffer);
+		OutParameters.ShaderPrint_RWEntryBuffer = GraphBuilder.CreateUAV(InData.ShaderPrintEntryBuffer);
 	}
 
 	void SetParameters(FRDGBuilder& GraphBuilder, FShaderParameters& OutParameters)
@@ -255,11 +254,11 @@ namespace ShaderPrint
 	// Widget/Characters Shaders
 	 
 	// Shader to initialize the output value buffer
-	class FShaderInitValueBufferCS : public FGlobalShader
+	class FShaderPrintClearCounterCS : public FGlobalShader
 	{
 	public:
-		DECLARE_GLOBAL_SHADER(FShaderInitValueBufferCS);
-		SHADER_USE_PARAMETER_STRUCT(FShaderInitValueBufferCS, FGlobalShader);
+		DECLARE_GLOBAL_SHADER(FShaderPrintClearCounterCS);
+		SHADER_USE_PARAMETER_STRUCT(FShaderPrintClearCounterCS, FGlobalShader);
 
 		BEGIN_SHADER_PARAMETER_STRUCT(FParameters, )
 			SHADER_PARAMETER_RDG_BUFFER_UAV(RWStructuredBuffer<uint>, RWValuesBuffer)
@@ -271,7 +270,7 @@ namespace ShaderPrint
 		}
 	};
 
-	IMPLEMENT_GLOBAL_SHADER(FShaderInitValueBufferCS, "/Engine/Private/ShaderPrintDraw.usf", "InitValueBufferCS", SF_Compute);
+	IMPLEMENT_GLOBAL_SHADER(FShaderPrintClearCounterCS, "/Engine/Private/ShaderPrintDraw.usf", "ClearCounterCS", SF_Compute);
 
 	// Shader to fill the indirect parameter arguments ready for the value->symbol compute pass
 	class FShaderBuildIndirectDispatchArgsCS : public FGlobalShader
@@ -406,30 +405,6 @@ namespace ShaderPrint
 	//////////////////////////////////////////////////////////////////////////
 	// Line Shaders
 
-	class FShaderDrawDebugClearCS : public FGlobalShader
-	{
-		DECLARE_GLOBAL_SHADER(FShaderDrawDebugClearCS);
-		SHADER_USE_PARAMETER_STRUCT(FShaderDrawDebugClearCS, FGlobalShader);
-
-		BEGIN_SHADER_PARAMETER_STRUCT(FParameters, )
-			SHADER_PARAMETER_RDG_BUFFER_UAV(RWStructuredBuffer, RWElementBuffer)
-		END_SHADER_PARAMETER_STRUCT()
-
-		static bool ShouldCompilePermutation(const FGlobalShaderPermutationParameters& Parameters)
-		{
-			return IsSupported(Parameters.Platform);
-		}
-
-		static void ModifyCompilationEnvironment(const FGlobalShaderPermutationParameters& Parameters, FShaderCompilerEnvironment& OutEnvironment)
-		{
-			FGlobalShader::ModifyCompilationEnvironment(Parameters, OutEnvironment);
-			OutEnvironment.SetDefine(TEXT("GPU_DEBUG_RENDERING"), 1);
-			OutEnvironment.SetDefine(TEXT("GPU_DEBUG_RENDERING_CLEAR_CS"), 1);
-		}
-	};
-
-	IMPLEMENT_GLOBAL_SHADER(FShaderDrawDebugClearCS, "/Engine/Private/ShaderPrintDrawPrimitive.usf", "ShaderDrawDebugClearCS", SF_Compute);
-	
 	class FShaderDrawDebugCopyCS : public FGlobalShader
 	{
 		DECLARE_GLOBAL_SHADER(FShaderDrawDebugCopyCS);
@@ -582,19 +557,23 @@ namespace ShaderPrint
 		// This is in case some debug shader code is still active and accessing the buffer.
 		if (!InSetup.bEnabled)
 		{
-			ShaderPrintData.ShaderPrintValueBuffer = GraphBuilder.RegisterExternalBuffer(GEmptyBuffer->Buffer);
+			ShaderPrintData.ShaderPrintEntryBuffer = GraphBuilder.RegisterExternalBuffer(GEmptyBuffer->Buffer);
 			ShaderPrintData.ShaderPrintStateBuffer = GraphBuilder.RegisterExternalBuffer(GEmptyBuffer->Buffer);
-			ShaderPrintData.ShaderPrintPrimitiveBuffer = GraphBuilder.RegisterExternalBuffer(GEmptyBuffer->Buffer);
 			return ShaderPrintData;
 		}
 
-		FGlobalShaderMap* GlobalShaderMap = GetGlobalShaderMap(GMaxRHIFeatureLevel);
 
-		// Characters/Widgets
+		// Characters/Widgets/Primitives/Lines
 		{
-			// Initialize output buffer and store in the view info
-			const uint32 UintElementCount = GetCountersUintSize() + GetPackedSymbolUintSize() * InSetup.MaxValueCount;
-			ShaderPrintData.ShaderPrintValueBuffer = GraphBuilder.CreateBuffer(FRDGBufferDesc::CreateStructuredDesc(4, UintElementCount), TEXT("ShaderPrint.ValueBuffer"));
+			const bool bLockBufferThisFrame = IsDrawLocked() && InViewState != nullptr && !InViewState->ShaderPrintStateData.bIsLocked;
+			ERDGBufferFlags Flags = bLockBufferThisFrame ? ERDGBufferFlags::MultiFrame : ERDGBufferFlags::None;
+
+			const uint32 UintElementCount = 
+				GetCountersUintSize() + 
+				GetPackedSymbolUintSize() * InSetup.MaxValueCount +
+				GetPackedLineUintSize() * InSetup.MaxLineCount + 
+				GetPackedTriangleUintSize() * InSetup.MaxTriangleCount;
+			ShaderPrintData.ShaderPrintEntryBuffer = GraphBuilder.CreateBuffer(FRDGBufferDesc::CreateStructuredDesc(4, UintElementCount), TEXT("ShaderPrint.EntryBuffer"), Flags);
 
 			// State buffer is retrieved from the view state, or created if it does not exist
 			if (InViewState != nullptr)
@@ -615,38 +594,17 @@ namespace ShaderPrint
 				ShaderPrintData.ShaderPrintStateBuffer = GraphBuilder.RegisterExternalBuffer(GEmptyBuffer->Buffer);
 			}
 
-			// Clear the output buffer internal counter ready for use
-			TShaderMapRef< FShaderInitValueBufferCS > ComputeShader(GlobalShaderMap);
-
-			auto* PassParameters = GraphBuilder.AllocParameters<FShaderInitValueBufferCS::FParameters>();
-			PassParameters->RWValuesBuffer = GraphBuilder.CreateUAV(ShaderPrintData.ShaderPrintValueBuffer);
-
-			FComputeShaderUtils::AddPass(
-				GraphBuilder,
-				RDG_EVENT_NAME("ShaderPrint::CreateShaderPrintData (Clear characters)"),
-				ComputeShader,
-				PassParameters,
-				FIntVector(1, 1, 1));
-		}
-
-		// Primitives/Lines
-		{
-			const bool bLockBufferThisFrame = IsDrawLocked() && InViewState != nullptr && !InViewState->ShaderPrintStateData.bIsLocked;
-			ERDGBufferFlags Flags = bLockBufferThisFrame ? ERDGBufferFlags::MultiFrame : ERDGBufferFlags::None;
-
-			const uint32 UintElementCount = GetCountersUintSize() + GetPackedLineUintSize() * InSetup.MaxLineCount + GetPackedTriangleUintSize() * InSetup.MaxTriangleCount;
-			ShaderPrintData.ShaderPrintPrimitiveBuffer = GraphBuilder.CreateBuffer(FRDGBufferDesc::CreateStructuredDesc(4, UintElementCount), TEXT("ShaderPrint.PrimitiveBuffer"), Flags);
-
-			// Clear buffer counter
+			// Clear counters
 			{
-				TShaderMapRef<FShaderDrawDebugClearCS> ComputeShader(GlobalShaderMap);
+				FGlobalShaderMap* GlobalShaderMap = GetGlobalShaderMap(GMaxRHIFeatureLevel);
+				TShaderMapRef<FShaderPrintClearCounterCS> ComputeShader(GlobalShaderMap);
 
-				FShaderDrawDebugClearCS::FParameters* PassParameters = GraphBuilder.AllocParameters<FShaderDrawDebugClearCS::FParameters>();
-				PassParameters->RWElementBuffer = GraphBuilder.CreateUAV(ShaderPrintData.ShaderPrintPrimitiveBuffer);
+				FShaderPrintClearCounterCS::FParameters* PassParameters = GraphBuilder.AllocParameters<FShaderPrintClearCounterCS::FParameters>();
+				PassParameters->RWValuesBuffer = GraphBuilder.CreateUAV(ShaderPrintData.ShaderPrintEntryBuffer);
 				ClearUnusedGraphResources(ComputeShader, PassParameters);
 
 				GraphBuilder.AddPass(
-					RDG_EVENT_NAME("ShaderPrint::CreateShaderPrintData (Clear primitives)"),
+					RDG_EVENT_NAME("ShaderPrint::ClearCounters"),
 					PassParameters,
 					ERDGPassFlags::Compute,
 					[PassParameters, ComputeShader](FRHICommandList& RHICmdList)
@@ -659,14 +617,14 @@ namespace ShaderPrint
 			{
 				if (IsDrawLocked() && !InViewState->ShaderPrintStateData.bIsLocked)
 				{
-					InViewState->ShaderPrintStateData.PrimitiveBuffer = GraphBuilder.ConvertToExternalBuffer(ShaderPrintData.ShaderPrintPrimitiveBuffer);
+					InViewState->ShaderPrintStateData.EntryBuffer = GraphBuilder.ConvertToExternalBuffer(ShaderPrintData.ShaderPrintEntryBuffer);
 					InViewState->ShaderPrintStateData.PreViewTranslation = InSetup.PreViewTranslation;
 					InViewState->ShaderPrintStateData.bIsLocked = true;
 				}
 
 				if (!IsDrawLocked() && InViewState->ShaderPrintStateData.bIsLocked)
 				{
-					InViewState->ShaderPrintStateData.PrimitiveBuffer = nullptr;
+					InViewState->ShaderPrintStateData.EntryBuffer = nullptr;
 					InViewState->ShaderPrintStateData.PreViewTranslation = FVector::ZeroVector;
 					InViewState->ShaderPrintStateData.bIsLocked = false;
 				}
@@ -687,9 +645,8 @@ namespace ShaderPrint
 
 		Out.Setup = ShaderPrintData.Setup;
 
-		Out.ShaderPrintValueBuffer = GraphBuilder.ConvertToExternalBuffer(ShaderPrintData.ShaderPrintValueBuffer);
+		Out.ShaderPrintEntryBuffer = GraphBuilder.ConvertToExternalBuffer(ShaderPrintData.ShaderPrintEntryBuffer);
 		Out.ShaderPrintStateBuffer = GraphBuilder.ConvertToExternalBuffer(ShaderPrintData.ShaderPrintStateBuffer);
-		Out.ShaderPrintPrimitiveBuffer = GraphBuilder.ConvertToExternalBuffer(ShaderPrintData.ShaderPrintPrimitiveBuffer);
 
 		return Out;
 	}
@@ -701,9 +658,8 @@ namespace ShaderPrint
 		Out.Setup = FrozenShaderPrintData.Setup;
 		Out.UniformBuffer = CreateUniformBuffer(Out.Setup);
 
-		Out.ShaderPrintValueBuffer = GraphBuilder.RegisterExternalBuffer(FrozenShaderPrintData.ShaderPrintValueBuffer);
+		Out.ShaderPrintEntryBuffer = GraphBuilder.RegisterExternalBuffer(FrozenShaderPrintData.ShaderPrintEntryBuffer);
 		Out.ShaderPrintStateBuffer = GraphBuilder.RegisterExternalBuffer(FrozenShaderPrintData.ShaderPrintStateBuffer);
-		Out.ShaderPrintPrimitiveBuffer = GraphBuilder.RegisterExternalBuffer(FrozenShaderPrintData.ShaderPrintPrimitiveBuffer);
 
 		return Out;
 	}
@@ -753,7 +709,7 @@ namespace ShaderPrint
 		FRDGBufferRef IndirectDrawArgsBuffer = GraphBuilder.CreateBuffer(FRDGBufferDesc::CreateIndirectDesc(5), TEXT("ShaderPrint.IndirectDrawArgs"));
 
 		// Non graph managed resources
-		FRDGBufferSRVRef ValueBuffer = GraphBuilder.CreateSRV(ShaderPrintData.ShaderPrintValueBuffer);
+		FRDGBufferSRVRef ValueBuffer = GraphBuilder.CreateSRV(ShaderPrintData.ShaderPrintEntryBuffer);
 		FRDGBufferSRVRef StateBuffer = GraphBuilder.CreateSRV(ShaderPrintData.ShaderPrintStateBuffer);
 		FTextureRHIRef FontTexture = GSystemTextures.AsciiTexture->GetRHI();
 
@@ -979,13 +935,13 @@ namespace ShaderPrint
 
 		// Lines
 		{
-			FRDGBufferRef DataBuffer = ShaderPrintData.ShaderPrintPrimitiveBuffer;
+			FRDGBufferRef DataBuffer = ShaderPrintData.ShaderPrintEntryBuffer;
 			InternalDrawView_Primitives(GraphBuilder, ShaderPrintData, DataBuffer, SourceViewRect, OutputViewRect, View.ViewMatrices.GetTranslatedViewProjectionMatrix(), PreViewTranslation, true /*bLines*/, false /*bLocked*/, OutputTexture.Texture, DepthTexture.Texture);
 		}
 
 		// Triangles
 		{
-			FRDGBufferRef DataBuffer = ShaderPrintData.ShaderPrintPrimitiveBuffer;
+			FRDGBufferRef DataBuffer = ShaderPrintData.ShaderPrintEntryBuffer;
 			InternalDrawView_Primitives(GraphBuilder, ShaderPrintData, DataBuffer, SourceViewRect, OutputViewRect, View.ViewMatrices.GetTranslatedViewProjectionMatrix(), PreViewTranslation, false /*bLines*/, false /*bLocked*/, OutputTexture.Texture, DepthTexture.Texture);
 		}
 
@@ -993,7 +949,7 @@ namespace ShaderPrint
 		if (View.ViewState && View.ViewState->ShaderPrintStateData.bIsLocked)
 		{
 			const FVector LockedPreViewTranslation = View.ViewMatrices.GetPreViewTranslation() - View.ViewState->ShaderPrintStateData.PreViewTranslation;
-			FRDGBufferRef DataBuffer = GraphBuilder.RegisterExternalBuffer(View.ViewState->ShaderPrintStateData.PrimitiveBuffer);
+			FRDGBufferRef DataBuffer = GraphBuilder.RegisterExternalBuffer(View.ViewState->ShaderPrintStateData.EntryBuffer);
 			InternalDrawView_Primitives(GraphBuilder, ShaderPrintData, DataBuffer, SourceViewRect, OutputViewRect, View.ViewMatrices.GetTranslatedViewProjectionMatrix(), LockedPreViewTranslation, true  /*bLines*/, true/*bLocked*/, OutputTexture.Texture, DepthTexture.Texture);
 			InternalDrawView_Primitives(GraphBuilder, ShaderPrintData, DataBuffer, SourceViewRect, OutputViewRect, View.ViewMatrices.GetTranslatedViewProjectionMatrix(), LockedPreViewTranslation, false /*bLines*/, true/*bLocked*/, OutputTexture.Texture, DepthTexture.Texture);
 		}
diff --git a/Engine/Source/Runtime/Renderer/Public/ShaderPrintParameters.h b/Engine/Source/Runtime/Renderer/Public/ShaderPrintParameters.h
--- a/Engine/Source/Runtime/Renderer/Public/ShaderPrintParameters.h
+++ b/Engine/Source/Runtime/Renderer/Public/ShaderPrintParameters.h
@@ -40,8 +40,7 @@ namespace ShaderPrint
 	BEGIN_SHADER_PARAMETER_STRUCT(FShaderParameters, )
 		SHADER_PARAMETER_STRUCT_REF(FShaderPrintCommonParameters, Common)
 		SHADER_PARAMETER_RDG_BUFFER_SRV(StructuredBuffer<uint2>, ShaderPrint_StateBuffer)
-		SHADER_PARAMETER_RDG_BUFFER_UAV(RWStructuredBuffer<ShaderPrintItem>, ShaderPrint_RWValuesBuffer)
-		SHADER_PARAMETER_RDG_BUFFER_UAV(RWStructuredBuffer<uint>, ShaderPrint_RWPrimitivesBuffer)
+		SHADER_PARAMETER_RDG_BUFFER_UAV(RWStructuredBuffer<ShaderPrintItem>, ShaderPrint_RWEntryBuffer)
 	END_SHADER_PARAMETER_STRUCT()
 
 	// Does the platform support the ShaderPrint system?
@@ -150,9 +149,8 @@ struct RENDERER_API FShaderPrintData
 	ShaderPrint::FShaderPrintSetup Setup;
 	TUniformBufferRef<ShaderPrint::FShaderPrintCommonParameters> UniformBuffer;
 
-	FRDGBufferRef ShaderPrintValueBuffer = nullptr;
+	FRDGBufferRef ShaderPrintEntryBuffer = nullptr;
 	FRDGBufferRef ShaderPrintStateBuffer = nullptr;
-	FRDGBufferRef ShaderPrintPrimitiveBuffer = nullptr;
 };
 
 /**
@@ -165,7 +163,6 @@ struct RENDERER_API FFrozenShaderPrintData
 {
 	ShaderPrint::FShaderPrintSetup Setup;
 
-	TRefCountPtr<FRDGPooledBuffer> ShaderPrintValueBuffer;
+	TRefCountPtr<FRDGPooledBuffer> ShaderPrintEntryBuffer;
 	TRefCountPtr<FRDGPooledBuffer> ShaderPrintStateBuffer;
-	TRefCountPtr<FRDGPooledBuffer> ShaderPrintPrimitiveBuffer;
 };
