Fix EmissiveLightSource not correctly encoded to be injected into the VoxelLighting.

diff --git a/Engine/Source/Runtime/Renderer/Private/DistanceFieldObjectManagement.cpp b/Engine/Source/Runtime/Renderer/Private/DistanceFieldObjectManagement.cpp
--- a/Engine/Source/Runtime/Renderer/Private/DistanceFieldObjectManagement.cpp
+++ b/Engine/Source/Runtime/Renderer/Private/DistanceFieldObjectManagement.cpp
@@ -748,7 +748,7 @@ void FDistanceFieldSceneData::UpdateDistanceFieldObjectBuffers(
 										const uint32 bIsNaniteMesh = PrimitiveSceneProxy->IsNaniteMesh() ? 1U : 0U;
 										const uint32 bEmissiveLightSource = PrimitiveSceneProxy->IsEmissiveLightSource();
 
-										const uint32 Flags = bOftenMoving | (bCastShadow << 1U) | (bIsNaniteMesh << 2U) | (bEmissiveLightSource << 4U);
+										const uint32 Flags = bOftenMoving | (bCastShadow << 1U) | (bIsNaniteMesh << 2U) | (bEmissiveLightSource << 3U);
 
 										FVector4f ObjectWorldExtentAndFlags((FVector3f)WorldSpaceMeshBounds.GetExtent(), 0.0f);
 										ObjectWorldExtentAndFlags.W = *(const float*)&Flags;
