Multi-view-family render:  Fix for a couple thread safety issues with view states.  Writes to the FScene when adding view states need to occur in the render thread.  Added separate bool to track whether a per-view VSM cache has been added to a view state.  The pointer gets written in the render thread, so there is a delay where duplicate "Add" calls in a single frame could leak the cache structure.

#rnx

diff --git a/Engine/Source/Runtime/Renderer/Private/RendererScene.cpp b/Engine/Source/Runtime/Renderer/Private/RendererScene.cpp
--- a/Engine/Source/Runtime/Renderer/Private/RendererScene.cpp
+++ b/Engine/Source/Runtime/Renderer/Private/RendererScene.cpp
@@ -385,6 +385,7 @@ FSceneViewState::FSceneViewState(ERHIFeatureLevel::Type FeatureLevel, FSceneView
 	GatherPointsResolution = FIntVector(0, 0, 0);
 #endif
 
+	bVirtualShadowMapCacheAdded = false;
 	ViewVirtualShadowMapCache = nullptr;
 }
 
@@ -427,11 +428,12 @@ FSceneViewState::~FSceneViewState()
 	{
 		delete ViewVirtualShadowMapCache;
 		ViewVirtualShadowMapCache = nullptr;
+		bVirtualShadowMapCacheAdded = false;
 	}
 
 	if (Scene)
 	{
-		Scene->RemoveViewState(this);
+		Scene->RemoveViewState_RenderThread(this);
 	}
 }
 
@@ -441,11 +443,19 @@ FSceneViewStateInterface* FScene::AllocateViewState(FSceneViewStateInterface* Sh
 
 	FSceneViewState* Result = new FSceneViewState(FeatureLevel, (FSceneViewState*)ShareOriginTarget);
 	Result->Scene = this;
-	ViewStates.Add(Result);
+
+	// Modification of scene structure needs to happen on render thread
+	FScene* Scene = this;
+	ENQUEUE_RENDER_COMMAND(SceneViewStateAdd)(
+		[Scene, Result](FRHICommandList&)
+		{
+			Scene->ViewStates.Add(Result);
+		});
+
 	return Result;
 }
 
-void FScene::RemoveViewState(FSceneViewStateInterface* ViewState)
+void FScene::RemoveViewState_RenderThread(FSceneViewStateInterface* ViewState)
 {
 	for (int32 ViewStateIndex = 0; ViewStateIndex < ViewStates.Num(); ViewStateIndex++)
 	{
@@ -760,9 +770,12 @@ void FSceneViewState::AddVirtualShadowMapCache()
 	// Can only allocate a virtual shadow map cache if we have a scene pointer available
 	if (Scene)
 	{
-		// Don't allocate if one already exists
-		if (ViewVirtualShadowMapCache == nullptr)
+		// Don't add the cache if one has already been added.  We have a separate bool since the write to the cache
+		// pointer is deferred to the render thread.
+		if (!bVirtualShadowMapCacheAdded)
 		{
+			bVirtualShadowMapCacheAdded = true;
+
 			FVirtualShadowMapArrayCacheManager* ViewCache = new FVirtualShadowMapArrayCacheManager(Scene);
 
 			// Need to add reference to virtual shadow map cache in render thread
diff --git a/Engine/Source/Runtime/Renderer/Private/ScenePrivate.h b/Engine/Source/Runtime/Renderer/Private/ScenePrivate.h
--- a/Engine/Source/Runtime/Renderer/Private/ScenePrivate.h
+++ b/Engine/Source/Runtime/Renderer/Private/ScenePrivate.h
@@ -1028,6 +1028,8 @@ public:
 
 	FShadingEnergyConservationStateData ShadingEnergyConservationData;
 
+	bool bVirtualShadowMapCacheAdded;
+
 	FVirtualShadowMapArrayCacheManager* ViewVirtualShadowMapCache;
 
 	// call after OnFrameRenderingSetup()
@@ -3367,7 +3369,7 @@ private:
 	 * Internal view state allocation interface, used by FSceneViewState
 	 */
 	virtual FSceneViewStateInterface* AllocateViewState(FSceneViewStateInterface* ShareOriginTarget) override;
-	void RemoveViewState(FSceneViewStateInterface*);
+	void RemoveViewState_RenderThread(FSceneViewStateInterface*);
 
 	/**
 	 * Ensures the packed primitive arrays contain the same number of elements.
