Add support for using the SceneTextures node to sample custom depth/stencil in Single Layer Water.

#rnx

diff --git a/Engine/Shaders/Private/MaterialTemplate.ush b/Engine/Shaders/Private/MaterialTemplate.ush
--- a/Engine/Shaders/Private/MaterialTemplate.ush
+++ b/Engine/Shaders/Private/MaterialTemplate.ush
@@ -2158,6 +2158,18 @@ MaterialFloat4 MobileSceneTextureLookup(inout FMaterialPixelParameters Parameter
 		return 0.0f;
 	}
 #else
+#if MATERIAL_SHADINGMODEL_SINGLELAYERWATER
+	if (SceneTextureId == PPI_CustomDepth)
+	{
+		MaterialFloat Depth = ConvertFromDeviceZ(Texture2DSample(MobileSceneTextures.CustomDepthTexture, MobileSceneTextures.CustomDepthTextureSampler, UV).r);
+		return MaterialFloat4(Depth.rrr, 0.f);
+	}
+	else if (SceneTextureId == PPI_CustomStencil)
+	{
+		MaterialFloat Stencil = MobileSceneTextures.CustomStencilTexture.Load(int3(Parameters.SvPosition.xy, 0)) STENCIL_COMPONENT_SWIZZLE;
+		return MaterialFloat4(Stencil, Stencil, Stencil, 0.f);
+	}
+#else // MATERIAL_SHADINGMODEL_SINGLELAYERWATER
 	if (SceneTextureId == PPI_SceneDepth)
 	{
 		MaterialFloat Depth = CalcSceneDepth(UV);
@@ -2192,7 +2204,8 @@ MaterialFloat4 MobileSceneTextureLookup(inout FMaterialPixelParameters Parameter
 		MaterialFloat Stencil = MobileSceneTextures.CustomStencilTexture.Load(int3(Parameters.SvPosition.xy, 0)) STENCIL_COMPONENT_SWIZZLE;
 		return MaterialFloat4(Stencil, Stencil, Stencil, 0.f);
 	}
-#endif// SCENE_TEXTURES_DISABLED
+#endif // MATERIAL_SHADINGMODEL_SINGLELAYERWATER
+#endif // SCENE_TEXTURES_DISABLED
 
 	return MaterialFloat4(0.0f, 0.0f, 0.0f, 0.0f);
 }
@@ -2243,6 +2256,20 @@ float4 SceneTextureLookup(float2 UV, int SceneTextureIndex, bool bFiltered)
 	}
 
 #else // !SCENE_TEXTURES_DISABLED
+#if MATERIAL_SHADINGMODEL_SINGLELAYERWATER
+	if (SceneTextureIndex == PPI_CustomDepth)
+	{
+		MaterialFloat Depth = ConvertFromDeviceZ(Texture2DSample(OpaqueBasePass.CustomDepthTexture, OpaqueBasePass.CustomDepthSampler, UV).r);
+		return MaterialFloat4(Depth.rrr, 0.f);
+	}
+	else if (SceneTextureIndex == PPI_CustomStencil)
+	{
+		MaterialFloat Stencil = OpaqueBasePass.CustomStencilTexture.Load(int3(UV * View.BufferSizeAndInvSize.xy, 0)) STENCIL_COMPONENT_SWIZZLE;
+		return MaterialFloat4(Stencil, Stencil, Stencil, 0.f);
+	}
+
+	return float4(0, 0, 0, 0);
+#else // MATERIAL_SHADINGMODEL_SINGLELAYERWATER
 #if STRATA_ENABLED
 
 	const uint2 PixelPos = uint2(UV * View.BufferSizeAndInvSize.xy);
@@ -2378,6 +2405,7 @@ float4 SceneTextureLookup(float2 UV, int SceneTextureIndex, bool bFiltered)
 			return float4(0, 0, 0, 0);
 	}
 #endif // STRATA_ENABLED
+#endif // MATERIAL_SHADINGMODEL_SINGLELAYERWATER
 #endif // SCENE_TEXTURES_DISABLED
 }
 
diff --git a/Engine/Source/Runtime/Engine/Private/Materials/HLSLMaterialTranslator.cpp b/Engine/Source/Runtime/Engine/Private/Materials/HLSLMaterialTranslator.cpp
--- a/Engine/Source/Runtime/Engine/Private/Materials/HLSLMaterialTranslator.cpp
+++ b/Engine/Source/Runtime/Engine/Private/Materials/HLSLMaterialTranslator.cpp
@@ -1157,10 +1157,10 @@ bool FHLSLMaterialTranslator::Translate()
 		{
 			if (Domain != MD_DeferredDecal && Domain != MD_PostProcess)
 			{
-				if (BlendMode == BLEND_Opaque || BlendMode == BLEND_Masked)
+				if (!MaterialShadingModels.HasShadingModel(MSM_SingleLayerWater) && (BlendMode == BLEND_Opaque || BlendMode == BLEND_Masked))
 				{
 					// In opaque pass, none of the textures are available
-					Errorf(TEXT("SceneTexture expressions cannot be used in opaque materials"));
+					Errorf(TEXT("SceneTexture expressions cannot be used in opaque materials except if used with the Single Layer Water shading model."));
 				}
 				else if (bNeedsSceneTexturePostProcessInputs)
 				{
@@ -6418,6 +6418,13 @@ int32 FHLSLMaterialTranslator::SceneTextureLookup(int32 ViewportUV, uint32 InSce
 {
 	ESceneTextureId SceneTextureId = (ESceneTextureId)InSceneTextureId;
 
+	// Guard against using unsupported textures with SLW
+	const bool bHasSingleLayerWaterSM = GetMaterialShadingModels().HasShadingModel(MSM_SingleLayerWater);
+	if (bHasSingleLayerWaterSM && SceneTextureId != PPI_CustomDepth && SceneTextureId != PPI_CustomStencil)
+	{
+		return Error(TEXT("Only custom depth and custom stencil can be sampled with SceneTexture when used with the Single Layer Water shading model."));
+	}
+
 	const bool bSupportedOnMobile = SceneTextureId == PPI_PostProcessInput0 ||
 									SceneTextureId == PPI_CustomDepth ||
 									SceneTextureId == PPI_SceneDepth ||
diff --git a/Engine/Source/Runtime/Renderer/Private/BasePassRendering.cpp b/Engine/Source/Runtime/Renderer/Private/BasePassRendering.cpp
--- a/Engine/Source/Runtime/Renderer/Private/BasePassRendering.cpp
+++ b/Engine/Source/Runtime/Renderer/Private/BasePassRendering.cpp
@@ -592,6 +592,7 @@ TRDGUniformBufferRef<FOpaqueBasePassUniformParameters> CreateOpaqueBasePassUnifo
 	const FForwardBasePassTextures& ForwardBasePassTextures,
 	const FDBufferTextures& DBufferTextures,
 	const FSceneWithoutWaterTextures* SceneWithoutWaterTextures,
+	const FCustomDepthTextures* CustomDepth,
 	bool bLumenGIEnabled)
 {
 	FOpaqueBasePassUniformParameters& BasePassParameters = *GraphBuilder.AllocParameters<FOpaqueBasePassUniformParameters>();
@@ -629,31 +630,36 @@ TRDGUniformBufferRef<FOpaqueBasePassUniformParameters> CreateOpaqueBasePassUnifo
 
 	// Single Layer Water
 	const bool bShouldUseBilinearSamplerForDepth = SceneWithoutWaterTextures && ShouldUseBilinearSamplerForDepthWithoutSingleLayerWater(SceneWithoutWaterTextures->DepthTexture->Desc.Format);
-	BasePassParameters.SceneWithoutSingleLayerWaterMinMaxUV = FVector4f(0.0f, 0.0f, 1.0f, 1.0f);
+	const bool bDepthTextureProduced = CustomDepth && HasBeenProduced(CustomDepth->Depth);
 	BasePassParameters.SceneColorWithoutSingleLayerWaterTexture = SystemTextures.Black;
-	BasePassParameters.SceneDepthWithoutSingleLayerWaterTexture = SystemTextures.Black;
 	BasePassParameters.SceneColorWithoutSingleLayerWaterSampler = TStaticSamplerState<SF_Bilinear>::GetRHI();
+	BasePassParameters.SceneDepthWithoutSingleLayerWaterTexture = SystemTextures.Black;
 	BasePassParameters.SceneDepthWithoutSingleLayerWaterSampler = bShouldUseBilinearSamplerForDepth ? TStaticSamplerState<SF_Bilinear>::GetRHI() : TStaticSamplerState<SF_Point>::GetRHI();
+	BasePassParameters.CustomDepthTexture = bDepthTextureProduced ? CustomDepth->Depth : SystemTextures.DepthDummy;
+	BasePassParameters.CustomStencilTexture = bDepthTextureProduced ? CustomDepth->Stencil : SystemTextures.StencilDummySRV;
+	BasePassParameters.CustomDepthSampler = TStaticSamplerState<SF_Point>::GetRHI();
+	BasePassParameters.SceneWithoutSingleLayerWaterMinMaxUV = FVector4f(0.0f, 0.0f, 1.0f, 1.0f);
+	SetupDistortionParams(BasePassParameters.DistortionParams, View);
+	BasePassParameters.SceneWithoutSingleLayerWaterTextureSize = FVector2f(0.0f, 0.0f);
+	BasePassParameters.SceneWithoutSingleLayerWaterInvTextureSize = FVector2f(0.0f, 0.0f);
 
 	if (SceneWithoutWaterTextures)
 	{
 		FIntVector DepthTextureSize = SceneWithoutWaterTextures->DepthTexture->Desc.GetSize();
-		
+
+		BasePassParameters.SceneColorWithoutSingleLayerWaterTexture = SceneWithoutWaterTextures->ColorTexture;
+		BasePassParameters.SceneDepthWithoutSingleLayerWaterTexture = SceneWithoutWaterTextures->DepthTexture;
 		BasePassParameters.SceneWithoutSingleLayerWaterMinMaxUV = SceneWithoutWaterTextures->Views[ViewIndex].MinMaxUV;
 		BasePassParameters.SceneWithoutSingleLayerWaterTextureSize = FVector2f(DepthTextureSize.X, DepthTextureSize.Y);
 		BasePassParameters.SceneWithoutSingleLayerWaterInvTextureSize = FVector2f(1.0f / DepthTextureSize.X, 1.0f / DepthTextureSize.Y);
-		BasePassParameters.SceneColorWithoutSingleLayerWaterTexture = SceneWithoutWaterTextures->ColorTexture;
-		BasePassParameters.SceneDepthWithoutSingleLayerWaterTexture = SceneWithoutWaterTextures->DepthTexture;
 	}
 
-	BasePassParameters.PreIntegratedGFTexture = GSystemTextures.PreintegratedGF->GetRHI();
-	BasePassParameters.PreIntegratedGFSampler = TStaticSamplerState<SF_Bilinear, AM_Clamp, AM_Clamp, AM_Clamp>::GetRHI();
-	SetupDistortionParams(BasePassParameters.DistortionParams, View);
-
 	// Strata
 	Strata::BindStrataBasePassUniformParameters(GraphBuilder, View, BasePassParameters.Strata);
 
 	// Misc
+	BasePassParameters.PreIntegratedGFTexture = GSystemTextures.PreintegratedGF->GetRHI();
+	BasePassParameters.PreIntegratedGFSampler = TStaticSamplerState<SF_Bilinear, AM_Clamp, AM_Clamp, AM_Clamp>::GetRHI();
 	BasePassParameters.EyeAdaptationTexture = GetEyeAdaptationTexture(GraphBuilder, View);
 
 	return GraphBuilder.CreateUniformBuffer(&BasePassParameters);
@@ -1274,7 +1280,7 @@ void FDeferredShadingSceneRenderer::RenderBasePassInternal(
 				FOpaqueBasePassParameters* PassParameters = GraphBuilder.AllocParameters<FOpaqueBasePassParameters>();
 				PassParameters->View = View.GetShaderParameters();
 				PassParameters->ReflectionCapture = View.ReflectionCaptureUniformBuffer;
-				PassParameters->BasePass = CreateOpaqueBasePassUniformBuffer(GraphBuilder, View, ViewIndex, ForwardBasePassTextures, DBufferTextures, nullptr, bLumenGIEnabled);
+				PassParameters->BasePass = CreateOpaqueBasePassUniformBuffer(GraphBuilder, View, ViewIndex, ForwardBasePassTextures, DBufferTextures, nullptr, nullptr, bLumenGIEnabled);
 				PassParameters->RenderTargets = BasePassRenderTargets;
 
 				const bool bShouldRenderView = View.ShouldRenderView();
@@ -1340,7 +1346,7 @@ void FDeferredShadingSceneRenderer::RenderBasePassInternal(
 				FOpaqueBasePassParameters* PassParameters = GraphBuilder.AllocParameters<FOpaqueBasePassParameters>();
 				PassParameters->View = View.GetShaderParameters();
 				PassParameters->ReflectionCapture = View.ReflectionCaptureUniformBuffer;
-				PassParameters->BasePass = CreateOpaqueBasePassUniformBuffer(GraphBuilder, View, ViewIndex, ForwardBasePassTextures, DBufferTextures, nullptr, bLumenGIEnabled);
+				PassParameters->BasePass = CreateOpaqueBasePassUniformBuffer(GraphBuilder, View, ViewIndex, ForwardBasePassTextures, DBufferTextures, nullptr, nullptr, bLumenGIEnabled);
 				PassParameters->RenderTargets = BasePassRenderTargets;
 
 				const bool bShouldRenderView = View.ShouldRenderView();
diff --git a/Engine/Source/Runtime/Renderer/Private/BasePassRendering.h b/Engine/Source/Runtime/Renderer/Private/BasePassRendering.h
--- a/Engine/Source/Runtime/Renderer/Private/BasePassRendering.h
+++ b/Engine/Source/Runtime/Renderer/Private/BasePassRendering.h
@@ -88,14 +88,17 @@ BEGIN_GLOBAL_SHADER_PARAMETER_STRUCT(FOpaqueBasePassUniformParameters,)
 	SHADER_PARAMETER_SAMPLER(SamplerState, SceneColorWithoutSingleLayerWaterSampler)
 	SHADER_PARAMETER_RDG_TEXTURE(Texture2D, SceneDepthWithoutSingleLayerWaterTexture)
 	SHADER_PARAMETER_SAMPLER(SamplerState, SceneDepthWithoutSingleLayerWaterSampler)
-	SHADER_PARAMETER_TEXTURE(Texture2D, PreIntegratedGFTexture)
-	SHADER_PARAMETER_SAMPLER(SamplerState, PreIntegratedGFSampler)
+	SHADER_PARAMETER_RDG_TEXTURE(Texture2D, CustomDepthTexture)
+	SHADER_PARAMETER_RDG_TEXTURE_SRV(Texture2D<uint2>, CustomStencilTexture)
+	SHADER_PARAMETER_SAMPLER(SamplerState, CustomDepthSampler)
 	SHADER_PARAMETER(FVector4f, SceneWithoutSingleLayerWaterMinMaxUV)
 	SHADER_PARAMETER(FVector4f, DistortionParams)
 	SHADER_PARAMETER(FVector2f, SceneWithoutSingleLayerWaterTextureSize)
 	SHADER_PARAMETER(FVector2f, SceneWithoutSingleLayerWaterInvTextureSize)
-	SHADER_PARAMETER(int32, Is24BitUnormDepthStencil)
 	// Misc
+	SHADER_PARAMETER_TEXTURE(Texture2D, PreIntegratedGFTexture)
+	SHADER_PARAMETER_SAMPLER(SamplerState, PreIntegratedGFSampler)
+	SHADER_PARAMETER(int32, Is24BitUnormDepthStencil)
 	SHADER_PARAMETER_RDG_TEXTURE(Texture2D, EyeAdaptationTexture)
 END_GLOBAL_SHADER_PARAMETER_STRUCT()
 
@@ -149,6 +152,7 @@ extern TRDGUniformBufferRef<FOpaqueBasePassUniformParameters> CreateOpaqueBasePa
 	const FForwardBasePassTextures& ForwardBasePassTextures = {},
 	const FDBufferTextures& DBufferTextures = {},
 	const FSceneWithoutWaterTextures* SceneWithoutWaterTextures = nullptr,
+	const FCustomDepthTextures* CustomDepth = nullptr,
 	bool bLumenGIEnabled = false);
 
 extern TRDGUniformBufferRef<FTranslucentBasePassUniformParameters> CreateTranslucentBasePassUniformBuffer(
diff --git a/Engine/Source/Runtime/Renderer/Private/SingleLayerWaterRendering.cpp b/Engine/Source/Runtime/Renderer/Private/SingleLayerWaterRendering.cpp
--- a/Engine/Source/Runtime/Renderer/Private/SingleLayerWaterRendering.cpp
+++ b/Engine/Source/Runtime/Renderer/Private/SingleLayerWaterRendering.cpp
@@ -870,7 +870,7 @@ void FDeferredShadingSceneRenderer::RenderSingleLayerWaterInner(
 		FSingleLayerWaterPassParameters* PassParameters = GraphBuilder.AllocParameters<FSingleLayerWaterPassParameters>();
 		PassParameters->View = View.GetShaderParameters();
 		PassParameters->ReflectionCapture = View.ReflectionCaptureUniformBuffer;
-		PassParameters->BasePass = CreateOpaqueBasePassUniformBuffer(GraphBuilder, View, ViewIndex, {}, {}, &SceneWithoutWaterTextures);
+		PassParameters->BasePass = CreateOpaqueBasePassUniformBuffer(GraphBuilder, View, ViewIndex, {}, {}, &SceneWithoutWaterTextures, &SceneTextures.CustomDepth, false);
 		PassParameters->VirtualShadowMapSamplingParameters = VirtualShadowMapArray.GetSamplingParameters(GraphBuilder);
 		PassParameters->RenderTargets = GetRenderTargetBindings(ERenderTargetLoadAction::ELoad, BasePassTexturesView);
 		PassParameters->RenderTargets.DepthStencil = FDepthStencilBinding(SceneTextures.Depth.Target, ERenderTargetLoadAction::ELoad, ERenderTargetLoadAction::ELoad, FExclusiveDepthStencil::DepthWrite_StencilWrite);
