Unbind GL backbuffer as soon as we done with it in a frame
Add more cmdbuffer stats for mobile rendering debugging

diff --git a/Engine/Source/Runtime/OpenGLDrv/Private/Android/AndroidEGL.cpp b/Engine/Source/Runtime/OpenGLDrv/Private/Android/AndroidEGL.cpp
--- a/Engine/Source/Runtime/OpenGLDrv/Private/Android/AndroidEGL.cpp
+++ b/Engine/Source/Runtime/OpenGLDrv/Private/Android/AndroidEGL.cpp
@@ -87,6 +87,7 @@ struct AndroidESPImpl
 	bool Initalized ;
 	EOpenGLCurrentContext CurrentContextType;
 	GLuint ResolveFrameBuffer;
+	GLuint DummyFrameBuffer;
 	FPlatformRect CachedWindowRect;
 
 	AndroidESPImpl();
@@ -525,6 +526,7 @@ eglDisplay(EGL_NO_DISPLAY)
 	,Window(NULL)
 	,Initalized(false)
 	,ResolveFrameBuffer(0)
+	,DummyFrameBuffer(0)
 	,NativeVisualID(0)
 	,CurrentContextType(CONTEXT_Invalid)
 	,CachedWindowRect(FPlatformRect(0,0,0,0))
@@ -612,6 +614,12 @@ void AndroidEGL::DestroyBackBuffer()
 		glDeleteFramebuffers(1, &PImplData->ResolveFrameBuffer);
 		PImplData->ResolveFrameBuffer = 0 ;
 	}
+
+	if(PImplData->DummyFrameBuffer)
+	{
+		glDeleteFramebuffers(1, &PImplData->DummyFrameBuffer);
+		PImplData->DummyFrameBuffer = 0;
+	}
 }
 
 void AndroidEGL::InitBackBuffer()
@@ -632,6 +640,12 @@ void AndroidEGL::InitBackBuffer()
 	PImplData->RenderingContext.ViewportFramebuffer = GetResolveFrameBuffer();
 	PImplData->SharedContext.ViewportFramebuffer = GetResolveFrameBuffer();
 	PImplData->SingleThreadedContext.ViewportFramebuffer = GetResolveFrameBuffer();
+	
+	// Dummy FBO we bind right after SwapBuffers to tell driver that backbuffer is no longer in use by the App
+	glGenFramebuffers(1, &PImplData->DummyFrameBuffer);
+	PImplData->RenderingContext.DummyFrameBuffer = PImplData->DummyFrameBuffer;
+	PImplData->SharedContext.DummyFrameBuffer = PImplData->DummyFrameBuffer;
+	PImplData->SingleThreadedContext.DummyFrameBuffer = PImplData->DummyFrameBuffer;
 }
 
 extern void AndroidThunkCpp_SetDesiredViewSize(int32 Width, int32 Height);
diff --git a/Engine/Source/Runtime/OpenGLDrv/Private/Android/AndroidEGL.h b/Engine/Source/Runtime/OpenGLDrv/Private/Android/AndroidEGL.h
--- a/Engine/Source/Runtime/OpenGLDrv/Private/Android/AndroidEGL.h
+++ b/Engine/Source/Runtime/OpenGLDrv/Private/Android/AndroidEGL.h
@@ -36,6 +36,7 @@ struct FPlatformOpenGLContext
 	GLuint		DefaultVertexArrayObject;
 	GLuint		BackBufferResource;
 	GLenum		BackBufferTarget;
+	GLuint		DummyFrameBuffer;
 
 	FPlatformOpenGLContext()
 	{
@@ -50,6 +51,7 @@ struct FPlatformOpenGLContext
 		DefaultVertexArrayObject = 0;
 		BackBufferResource = 0;
 		BackBufferTarget = 0;
+		DummyFrameBuffer = 0;
 	}
 };
 
diff --git a/Engine/Source/Runtime/OpenGLDrv/Private/Android/AndroidOpenGL.cpp b/Engine/Source/Runtime/OpenGLDrv/Private/Android/AndroidOpenGL.cpp
--- a/Engine/Source/Runtime/OpenGLDrv/Private/Android/AndroidOpenGL.cpp
+++ b/Engine/Source/Runtime/OpenGLDrv/Private/Android/AndroidOpenGL.cpp
@@ -221,10 +221,12 @@ void* PlatformGetWindow(FPlatformOpenGLContext* Context, void** AddParam)
 
 bool PlatformBlitToViewport(FPlatformOpenGLDevice* Device, const FOpenGLViewport& Viewport, uint32 BackbufferSizeX, uint32 BackbufferSizeY, bool bPresent,bool bLockToVsync )
 {
+	SCOPED_NAMED_EVENT(STAT_PlatformBlitToViewportTime, FColor::Red)
+	
+	FPlatformOpenGLContext* const Context = Viewport.GetGLContext();
+
 	if (FPlatformMisc::SupportsBackbufferSampling())
 	{
-		FPlatformOpenGLContext* const Context = Viewport.GetGLContext();
-
 		if (Device->TargetDirty)
 		{
 			VERIFY_GL_SCOPE();
@@ -266,7 +268,9 @@ bool PlatformBlitToViewport(FPlatformOpenGLDevice* Device, const FOpenGLViewport
 	static IConsoleVariable* CVar = IConsoleManager::Get().FindConsoleVariable(TEXT("a.UseFrameTimeStampsForPacing"));
 	const bool bForceGPUFence = CVar ? CVar->GetInt() != 0 : false;
 
-
+	// Bind a dummy FBO so driver knows we don't need backbuffer anymore
+	check(Context->DummyFrameBuffer != 0);
+	glBindFramebuffer(GL_FRAMEBUFFER, Context->DummyFrameBuffer);
 
 	return bPresent && ShouldUseGPUFencesToLimitLatency();
 }
diff --git a/Engine/Source/Runtime/Renderer/Private/MobileShadingRenderer.cpp b/Engine/Source/Runtime/Renderer/Private/MobileShadingRenderer.cpp
--- a/Engine/Source/Runtime/Renderer/Private/MobileShadingRenderer.cpp
+++ b/Engine/Source/Runtime/Renderer/Private/MobileShadingRenderer.cpp
@@ -89,6 +89,7 @@ DECLARE_GPU_STAT_NAMED(MobileSceneRender, TEXT("Mobile Scene Render"));
 DECLARE_CYCLE_STAT(TEXT("SceneStart"), STAT_CLMM_SceneStart, STATGROUP_CommandListMarkers);
 DECLARE_CYCLE_STAT(TEXT("SceneEnd"), STAT_CLMM_SceneEnd, STATGROUP_CommandListMarkers);
 DECLARE_CYCLE_STAT(TEXT("InitViews"), STAT_CLMM_InitViews, STATGROUP_CommandListMarkers);
+DECLARE_CYCLE_STAT(TEXT("AfterInitViews"), STAT_CLMM_AfterInitViews, STATGROUP_CommandListMarkers);
 DECLARE_CYCLE_STAT(TEXT("Opaque"), STAT_CLMM_Opaque, STATGROUP_CommandListMarkers);
 DECLARE_CYCLE_STAT(TEXT("Occlusion"), STAT_CLMM_Occlusion, STATGROUP_CommandListMarkers);
 DECLARE_CYCLE_STAT(TEXT("Post"), STAT_CLMM_Post, STATGROUP_CommandListMarkers);
@@ -97,7 +98,6 @@ DECLARE_CYCLE_STAT(TEXT("Shadows"), STAT_CLMM_Shadows, STATGROUP_CommandListMark
 DECLARE_CYCLE_STAT(TEXT("SceneSimulation"), STAT_CLMM_SceneSim, STATGROUP_CommandListMarkers);
 DECLARE_CYCLE_STAT(TEXT("PrePass"), STAT_CLM_MobilePrePass, STATGROUP_CommandListMarkers);
 DECLARE_CYCLE_STAT(TEXT("Velocity"), STAT_CLMM_Velocity, STATGROUP_CommandListMarkers);
-DECLARE_CYCLE_STAT(TEXT("AfterVelocity"), STAT_CLMM_AfterVelocity, STATGROUP_CommandListMarkers);
 DECLARE_CYCLE_STAT(TEXT("TranslucentVelocity"), STAT_CLMM_TranslucentVelocity, STATGROUP_CommandListMarkers);
 
 FGlobalDynamicIndexBuffer FMobileSceneRenderer::DynamicIndexBuffer;
@@ -235,6 +235,14 @@ static void SetupGBufferFlags(FSceneTexturesConfig& SceneTexturesConfig, bool bR
 	SceneTexturesConfig.GBufferE.Flags &= (~TexCreate_SRGB);
 }
 
+static void PollOcclusionQueriesPass(FRDGBuilder& GraphBuilder)
+{
+	AddPass(GraphBuilder, RDG_EVENT_NAME("PollOcclusionQueries"), [](FRHICommandListImmediate& RHICmdList)
+	{
+		RHICmdList.PollOcclusionQueries();
+	});
+}
+
 FMobileSceneRenderer::FMobileSceneRenderer(const FSceneViewFamily* InViewFamily, FHitProxyConsumer* HitProxyConsumer)
 	: FSceneRenderer(InViewFamily, HitProxyConsumer)
 	, bGammaSpace(!IsMobileHDR())
@@ -382,8 +390,7 @@ void FMobileSceneRenderer::SetupMobileBasePassAfterShadowInit(FExclusiveDepthSte
 void FMobileSceneRenderer::InitViews(FRDGBuilder& GraphBuilder, FSceneTexturesConfig& SceneTexturesConfig, FInstanceCullingManager& InstanceCullingManager)
 {
 	FRHICommandListImmediate& RHICmdList = GraphBuilder.RHICmdList;
-	RHICmdList.SetCurrentStat(GET_STATID(STAT_CLMM_InitViews));
-
+	
 	SCOPED_DRAW_EVENT(RHICmdList, InitViews);
 
 	SCOPE_CYCLE_COUNTER(STAT_InitViewsTime);
@@ -753,8 +760,6 @@ void FMobileSceneRenderer::Render(FRDGBuilder& GraphBuilder)
 	}
 
 	WaitOcclusionTests(GraphBuilder.RHICmdList);
-	FRHICommandListExecutor::GetImmediateCommandList().PollOcclusionQueries();
-	GraphBuilder.RHICmdList.ImmediateFlush(EImmediateFlushType::DispatchToRHIThread);
 
 	FSceneTexturesConfig::InitializeViewFamily(ViewFamily);
 	FSceneTexturesConfig& SceneTexturesConfig = GetActiveSceneTexturesConfig();
@@ -778,9 +783,13 @@ void FMobileSceneRenderer::Render(FRDGBuilder& GraphBuilder)
 	// TODO: This doesn't take into account the potential for split screen views with separate shadow caches
 	VirtualShadowMapArray.Initialize(GraphBuilder, Scene->GetVirtualShadowMapCache(Views[0]), UseVirtualShadowMaps(ShaderPlatform, FeatureLevel));
 
+	GraphBuilder.SetCommandListStat(GET_STATID(STAT_CLMM_InitViews));
+
 	// Find the visible primitives and prepare targets and buffers for rendering
 	InitViews(GraphBuilder, SceneTexturesConfig, InstanceCullingManager);
 
+	GraphBuilder.SetCommandListStat(GET_STATID(STAT_CLMM_AfterInitViews));
+
 	if (GRHINeedsExtraDeletionLatency || !GRHICommandList.Bypass())
 	{
 		QUICK_SCOPE_CYCLE_COUNTER(STAT_FMobileSceneRenderer_PostInitViewsFlushDel);
@@ -796,8 +805,7 @@ void FMobileSceneRenderer::Render(FRDGBuilder& GraphBuilder)
 	DynamicIndexBuffer.Commit();
 	DynamicVertexBuffer.Commit();
 	DynamicReadBuffer.Commit();
-	GraphBuilder.RHICmdList.ImmediateFlush(EImmediateFlushType::DispatchToRHIThread);
-
+	
 	GraphBuilder.SetCommandListStat(GET_STATID(STAT_CLMM_SceneSim));
 
 	FSceneTextures::InitializeViewFamily(GraphBuilder, ViewFamily);
@@ -841,22 +849,12 @@ void FMobileSceneRenderer::Render(FRDGBuilder& GraphBuilder)
 			GPUSortManager->OnPreRender(GraphBuilder);
 		}
 	}
-
-	const auto PollOcclusionQueries = [](FRDGBuilder& GraphBuilder)
-	{
-		AddPass(GraphBuilder, RDG_EVENT_NAME("PollOcclusionQueries"), [](FRHICommandListImmediate& RHICmdList)
-		{
-			RHICmdList.PollOcclusionQueries();
-			RHICmdList.ImmediateFlush(EImmediateFlushType::DispatchToRHIThread);
-		});
-	};
-
-	PollOcclusionQueries(GraphBuilder);
-
+	
 	GraphBuilder.SetCommandListStat(GET_STATID(STAT_CLMM_Shadows));
 	RenderShadowDepthMaps(GraphBuilder, InstanceCullingManager);
-
-	PollOcclusionQueries(GraphBuilder);
+	
+	PollOcclusionQueriesPass(GraphBuilder);
+	GraphBuilder.AddDispatchHint();
 
 	// Custom depth
 	// bShouldRenderCustomDepth has been initialized in InitViews on mobile platform
@@ -920,7 +918,6 @@ void FMobileSceneRenderer::Render(FRDGBuilder& GraphBuilder)
 		// Render the velocities of movable objects
 		GraphBuilder.SetCommandListStat(GET_STATID(STAT_CLMM_Velocity));
 		RenderVelocities(GraphBuilder, SceneTextures, EVelocityPass::Opaque, false);
-		GraphBuilder.SetCommandListStat(GET_STATID(STAT_CLMM_AfterVelocity));
 
 		GraphBuilder.SetCommandListStat(GET_STATID(STAT_CLMM_TranslucentVelocity));
 		RenderVelocities(GraphBuilder, SceneTextures, EVelocityPass::Translucent, false);
@@ -929,6 +926,8 @@ void FMobileSceneRenderer::Render(FRDGBuilder& GraphBuilder)
 		SceneTextures.MobileUniformBuffer = CreateMobileSceneTextureUniformBuffer(GraphBuilder, &SceneTextures, SceneTextures.MobileSetupMode);
 	}
 
+	GraphBuilder.SetCommandListStat(GET_STATID(STAT_CLMM_Post));
+
 	FRendererModule& RendererModule = static_cast<FRendererModule&>(GetRendererModule());
 	RendererModule.RenderPostOpaqueExtensions(GraphBuilder, Views, SceneTextures);
 
@@ -940,9 +939,7 @@ void FMobileSceneRenderer::Render(FRDGBuilder& GraphBuilder)
 
 		RenderPixelProjectedReflection(GraphBuilder, SceneTextures.Color.Resolve, SceneTextures.Depth.Resolve, SceneTextures.PixelProjectedReflection, PlanarReflectionSceneProxy);
 	}
-
-	GraphBuilder.SetCommandListStat(GET_STATID(STAT_CLMM_Post));
-
+	
 	if (bUseVirtualTexturing)
 	{
 		RDG_GPU_STAT_SCOPE(GraphBuilder, VirtualTextureUpdate);
@@ -977,7 +974,7 @@ void FMobileSceneRenderer::Render(FRDGBuilder& GraphBuilder)
 
 	RenderFinish(GraphBuilder, ViewFamilyTexture);
 
-	PollOcclusionQueries(GraphBuilder);
+	PollOcclusionQueriesPass(GraphBuilder);
 
 	QueueSceneTextureExtractions(GraphBuilder, SceneTextures);
 }
@@ -1139,7 +1136,6 @@ void FMobileSceneRenderer::RenderForwardSinglePass(FRDGBuilder& GraphBuilder, FM
 		RHICmdList.SetCurrentStat(GET_STATID(STAT_CLMM_Opaque));
 		RenderMobileBasePass(RHICmdList, View);
 		RenderMobileDebugView(RHICmdList, View);
-		RHICmdList.ImmediateFlush(EImmediateFlushType::DispatchToRHIThread);
 		const bool bAdrenoOcclusionMode = (CVarMobileAdrenoOcclusionMode.GetValueOnRenderThread() != 0 && IsOpenGLPlatform(ShaderPlatform));
 		if (!bAdrenoOcclusionMode)
 		{
@@ -1425,7 +1421,6 @@ void FMobileSceneRenderer::RenderDeferredSinglePass(FRDGBuilder& GraphBuilder, c
 		// Opaque and masked
 		RHICmdList.SetCurrentStat(GET_STATID(STAT_CLMM_Opaque));
 		RenderMobileBasePass(RHICmdList, View);
-		RHICmdList.ImmediateFlush(EImmediateFlushType::DispatchToRHIThread);
 		
 		// Issue occlusion queries
 		if (!bUsingPixelLocalStorage)
