Fix: dumping buffer visualization targets didn't work when they were using the A2B10G10R10 pixel format.

#rnx

diff --git a/Engine/Source/Runtime/Renderer/Private/PostProcess/PostProcessVisualizeBuffer.cpp b/Engine/Source/Runtime/Renderer/Private/PostProcess/PostProcessVisualizeBuffer.cpp
--- a/Engine/Source/Runtime/Renderer/Private/PostProcess/PostProcessVisualizeBuffer.cpp
+++ b/Engine/Source/Runtime/Renderer/Private/PostProcess/PostProcessVisualizeBuffer.cpp
@@ -196,6 +196,7 @@ TUniquePtr<FImagePixelData> ReadbackPixelData(FRHICommandListImmediate& RHICmdLi
 	}
 
 	case PF_A32B32G32R32F:
+	case PF_A2B10G10R10:
 	{
 		FReadSurfaceDataFlags ReadDataFlags(RCM_MinMax);
 		ReadDataFlags.SetLinearToGamma(false);
