Fix uniform buffer leak with instanced stereo.

#rnx

diff --git a/Engine/Source/Runtime/Renderer/Private/SceneRendering.cpp b/Engine/Source/Runtime/Renderer/Private/SceneRendering.cpp
--- a/Engine/Source/Runtime/Renderer/Private/SceneRendering.cpp
+++ b/Engine/Source/Runtime/Renderer/Private/SceneRendering.cpp
@@ -2151,6 +2151,7 @@ void FViewInfo::DestroyAllSnapshots(FParallelMeshDrawCommandPass::EWaitThread Wa
 	for (FViewInfo* Snapshot : ViewInfoSnapshots)
 	{
 		Snapshot->ViewUniformBuffer.SafeRelease();
+		Snapshot->InstancedViewUniformBuffer.SafeRelease();
 		Snapshot->CachedViewUniformShaderParameters.Reset();
 
 		for (int32 Index = 0; Index < Snapshot->ParallelMeshDrawCommandPasses.Num(); ++Index)
