- Check for cube array support before creating RDG dummy

#rnx

diff --git a/Engine/Source/Runtime/Renderer/Private/SystemTextures.cpp b/Engine/Source/Runtime/Renderer/Private/SystemTextures.cpp
--- a/Engine/Source/Runtime/Renderer/Private/SystemTextures.cpp
+++ b/Engine/Source/Runtime/Renderer/Private/SystemTextures.cpp
@@ -276,7 +276,10 @@ void FSystemTextures::InitializeCommonTextures(FRHICommandListImmediate& RHICmdL
 	// Create Cube BlackDummy textures
 	{
 		CubeBlackDummy = CreateRenderTarget(GBlackTextureCube->TextureRHI, TEXT("CubeBlackDummy"));
-		CubeArrayBlackDummy = CreateRenderTarget(GBlackCubeArrayTexture->TextureRHI, TEXT("CubeArrayBlackDummy"));
+		if (SupportsTextureCubeArray(GetFeatureLevel()))
+		{
+			CubeArrayBlackDummy = CreateRenderTarget(GBlackCubeArrayTexture->TextureRHI, TEXT("CubeArrayBlackDummy"));
+		}
 	}
 }
 
