Move some of the FRayTracingScene initialization work into an asynchronous task

* Add FRayTracingScene::BuildInitializationData, which builds FRayTracingSceneWithGeometryInstances
* Add FRayTracingScene::CreateWithInitializationData() that consumes FRayTracingSceneWithGeometryInstances data
* Run BuildInitializationData() as soon as ray tracing instance culling tasks finish
* Wait for scene initialization data in FDeferredShadingSceneRenderer::DispatchRayTracingWorldUpdates()

This saves ~0.2-0.3ms of rendering thread time in large scenes.

diff --git a/Engine/Source/Runtime/Renderer/Private/DeferredShadingRenderer.cpp b/Engine/Source/Runtime/Renderer/Private/DeferredShadingRenderer.cpp
--- a/Engine/Source/Runtime/Renderer/Private/DeferredShadingRenderer.cpp
+++ b/Engine/Source/Runtime/Renderer/Private/DeferredShadingRenderer.cpp
@@ -1387,6 +1387,13 @@ bool FDeferredShadingSceneRenderer::GatherRayTracingWorldInstancesForView(FRDGBu
 		CoarseMeshSM->AddUsedStreamingHandles(RayTracingScene.UsedCoarseMeshStreamingHandles);
 	}
 
+	View.RayTracingSceneInitTask = FFunctionGraphTask::CreateAndDispatchWhenReady(
+		[&View, &RayTracingScene]()
+		{
+			View.RayTracingSceneInitData = RayTracingScene.BuildInitializationData();
+		},
+		TStatId(), &View.RayTracingPerInstanceCullingTaskList, ENamedThreads::AnyThread);
+
 	return true;
 }
 
@@ -1647,12 +1654,14 @@ bool FDeferredShadingSceneRenderer::DispatchRayTracingWorldUpdates(FRDGBuilder&
 		GRayTracingGeometryManager.ForceBuildIfPending(GraphBuilder.RHICmdList, RayTracingScene.GeometriesToBuild);
 	}
 
-	FTaskGraphInterface::Get().WaitUntilTasksComplete(ReferenceView.RayTracingPerInstanceCullingTaskList, ENamedThreads::GetRenderThread_Local());
+	FTaskGraphInterface::Get().WaitUntilTaskCompletes(ReferenceView.RayTracingSceneInitTask, ENamedThreads::GetRenderThread_Local());
+
+	ReferenceView.RayTracingSceneInitTask = {};
 	ReferenceView.RayTracingPerInstanceCullingTaskList.Empty();
 
 	RDG_GPU_MASK_SCOPE(GraphBuilder, FRHIGPUMask::All());
 
-	RayTracingScene.Create(GraphBuilder, Scene->GPUScene, ReferenceView.ViewMatrices);
+	RayTracingScene.CreateWithInitializationData(GraphBuilder, Scene->GPUScene, ReferenceView.ViewMatrices, MoveTemp(ReferenceView.RayTracingSceneInitData));
 
 	const uint32 BLASScratchSize = Scene->GetRayTracingDynamicGeometryCollection()->ComputeScratchBufferSize();
 	if (BLASScratchSize > 0)
diff --git a/Engine/Source/Runtime/Renderer/Private/RayTracing/RayTracingInstanceBufferUtil.cpp b/Engine/Source/Runtime/Renderer/Private/RayTracing/RayTracingInstanceBufferUtil.cpp
--- a/Engine/Source/Runtime/Renderer/Private/RayTracing/RayTracingInstanceBufferUtil.cpp
+++ b/Engine/Source/Runtime/Renderer/Private/RayTracing/RayTracingInstanceBufferUtil.cpp
@@ -20,7 +20,7 @@
 #if RHI_RAYTRACING
 
 FRayTracingSceneWithGeometryInstances CreateRayTracingSceneWithGeometryInstances(
-	TArrayView<FRayTracingGeometryInstance> Instances,
+	TConstArrayView<FRayTracingGeometryInstance> Instances,
 	uint8 NumLayers,
 	uint32 NumShaderSlotsPerGeometrySegment,
 	uint32 NumMissShaderSlots,
diff --git a/Engine/Source/Runtime/Renderer/Private/RayTracing/RayTracingScene.cpp b/Engine/Source/Runtime/Renderer/Private/RayTracing/RayTracingScene.cpp
--- a/Engine/Source/Runtime/Renderer/Private/RayTracing/RayTracingScene.cpp
+++ b/Engine/Source/Runtime/Renderer/Private/RayTracing/RayTracingScene.cpp
@@ -21,27 +21,47 @@
 
 FRayTracingScene::~FRayTracingScene()
 {
+	// Make sure that all async tasks complete before we tear down any memory they might reference.
 	WaitForTasks();
 }
 
+FRayTracingSceneWithGeometryInstances FRayTracingScene::BuildInitializationData() const
+{
+	return CreateRayTracingSceneWithGeometryInstances(
+		Instances,
+		uint8(ERayTracingSceneLayer::NUM),
+		RAY_TRACING_NUM_SHADER_SLOTS,
+		RAY_TRACING_NUM_MISS_SHADER_SLOTS,
+		NumCallableShaderSlots);
+}
+
 void FRayTracingScene::Create(FRDGBuilder& GraphBuilder, const FGPUScene& GPUScene, const FViewMatrices& ViewMatrices)
 {
+	FRayTracingSceneWithGeometryInstances InitializationData = CreateRayTracingSceneWithGeometryInstances(
+		Instances,
+		uint8(ERayTracingSceneLayer::NUM),
+		RAY_TRACING_NUM_SHADER_SLOTS,
+		RAY_TRACING_NUM_MISS_SHADER_SLOTS,
+		NumCallableShaderSlots);
+
+	CreateWithInitializationData(GraphBuilder, GPUScene, ViewMatrices, MoveTemp(InitializationData));
+}
+
+void FRayTracingScene::CreateWithInitializationData(FRDGBuilder& GraphBuilder, const FGPUScene& GPUScene, const FViewMatrices& ViewMatrices, FRayTracingSceneWithGeometryInstances SceneWithGeometryInstances)
+{
 	QUICK_SCOPE_CYCLE_COUNTER(FRayTracingScene_BeginCreate);
 
 	// Round up buffer sizes to some multiple to avoid pathological growth reallocations.
 	static constexpr uint32 AllocationGranularity = 8 * 1024;
 	static constexpr uint64 BufferAllocationGranularity = 16 * 1024 * 1024;
 
+	// Make sure that all async tasks complete before we run initialization code again and create new tasks.
 	WaitForTasks();
 
-	static const uint8 NumLayers = uint32(ERayTracingSceneLayer::NUM);
+	static const uint8 NumLayers = uint8(ERayTracingSceneLayer::NUM);
 
-	FRayTracingSceneWithGeometryInstances SceneWithGeometryInstances = CreateRayTracingSceneWithGeometryInstances(
-		Instances,
-		NumLayers,
-		RAY_TRACING_NUM_SHADER_SLOTS,
-		RAY_TRACING_NUM_MISS_SHADER_SLOTS,
-		NumCallableShaderSlots);
+	checkf(SceneWithGeometryInstances.Scene.IsValid(), 
+		TEXT("Ray tracing scene RHI object is expected to have been created by BuildInitializationData() or CreateRayTracingSceneWithGeometryInstances()"));
 
 	RayTracingSceneRHI = SceneWithGeometryInstances.Scene;
 
diff --git a/Engine/Source/Runtime/Renderer/Private/RayTracing/RayTracingScene.h b/Engine/Source/Runtime/Renderer/Private/RayTracing/RayTracingScene.h
--- a/Engine/Source/Runtime/Renderer/Private/RayTracing/RayTracingScene.h
+++ b/Engine/Source/Runtime/Renderer/Private/RayTracing/RayTracingScene.h
@@ -14,6 +14,7 @@
 #include "Containers/ArrayView.h"
 #include "PrimitiveSceneProxy.h"
 #include "MeshPassProcessor.h"
+#include "RayTracingInstanceBufferUtil.h"
 
 class FGPUScene;
 class FRHIRayTracingScene;
@@ -39,9 +40,15 @@
 	FRayTracingScene();
 	~FRayTracingScene();
 
-	// Creates RayTracingSceneRHI.
+	// Allocates RayTracingSceneRHI and builds various metadata required to create the final scene.
+	FRayTracingSceneWithGeometryInstances BuildInitializationData() const;
+
 	// Allocates GPU memory to fit at least the current number of instances.
-	// Kicks off instance buffer build to parallel thread along with RDG pass
+	// Kicks off instance buffer build to parallel thread along with RDG pass.
+	// NOTE: SceneWithGeometryInstances is passed in by value because ownership of the internal data is taken over. Use MoveTemp at call site, if possible.
+	void CreateWithInitializationData(FRDGBuilder& GraphBuilder, const FGPUScene& GPUScene, const FViewMatrices& ViewMatrices, FRayTracingSceneWithGeometryInstances SceneWithGeometryInstances);
+
+	// Backwards-compatible version of Create() which internally calls CreateRayTracingSceneWithGeometryInstances().
 	void Create(FRDGBuilder& GraphBuilder, const FGPUScene& GPUScene, const FViewMatrices& ViewMatrices);
 
 	// Resets the instance list and reserves memory for this frame.
@@ -80,7 +87,7 @@
 
 	// Persistent storage for ray tracing instance descriptors.
 	// Cleared every frame without releasing memory to avoid large heap allocations.
-	// This must be filled before calling Create().
+	// This must be filled before calling CreateRayTracingSceneWithGeometryInstances() and Create().
 	TArray<FRayTracingGeometryInstance> Instances;
 
 	uint32 NumCallableShaderSlots = 0;
diff --git a/Engine/Source/Runtime/Renderer/Private/SceneRendering.h b/Engine/Source/Runtime/Renderer/Private/SceneRendering.h
--- a/Engine/Source/Runtime/Renderer/Private/SceneRendering.h
+++ b/Engine/Source/Runtime/Renderer/Private/SceneRendering.h
@@ -41,6 +41,10 @@
 #include "CanvasTypes.h"
 #include "SceneTextures.h"
 
+#if RHI_RAYTRACING
+#include "RayTracingInstanceBufferUtil.h"
+#endif // RHI_RAYTRACING
+
 // Forward declarations.
 class FScene;
 class FSceneViewState;
@@ -1282,6 +1286,10 @@ public:
 	FRayTracingCullingParameters RayTracingCullingParameters;
 	FGraphEventArray RayTracingPerInstanceCullingTaskList;
 
+	// Data required for FRayTracingScene, depends on RT instance culling tasks
+	FRayTracingSceneWithGeometryInstances RayTracingSceneInitData;
+	FGraphEventRef RayTracingSceneInitTask; // Task to asynchronously create RayTracingSceneInitData
+
 	FGraphEventArray AddRayTracingMeshBatchTaskList;
 	TArray<FRayTracingMeshCommandOneFrameArray*, SceneRenderingAllocator> VisibleRayTracingMeshCommandsPerTask;
 	TArray<FDynamicRayTracingMeshCommandStorage*, SceneRenderingAllocator> DynamicRayTracingMeshCommandStoragePerTask;
diff --git a/Engine/Source/Runtime/Renderer/Public/RayTracingInstanceBufferUtil.h b/Engine/Source/Runtime/Renderer/Public/RayTracingInstanceBufferUtil.h
--- a/Engine/Source/Runtime/Renderer/Public/RayTracingInstanceBufferUtil.h
+++ b/Engine/Source/Runtime/Renderer/Public/RayTracingInstanceBufferUtil.h
@@ -58,7 +58,7 @@
 // Helper function to create FRayTracingSceneRHI using array of high level instances
 // Also outputs data required to build the instance buffer
 RENDERER_API FRayTracingSceneWithGeometryInstances CreateRayTracingSceneWithGeometryInstances(
-	TArrayView<FRayTracingGeometryInstance> Instances,
+	TConstArrayView<FRayTracingGeometryInstance> Instances,
 	uint8 NumLayers,
 	uint32 NumShaderSlotsPerGeometrySegment,
 	uint32 NumMissShaderSlots,
