Multi-view-family render:  More robust handling of ray tracing scene update.  There is an optimization where we only update the ray tracing scene for the first view family when rendering multiple view families together...

* There was a bug where if the first view family didn't use ray tracing, the ray tracing scene would never be updated.  Now the ray tracing scene is updated on the first view family using ray tracing.  Also cleaned up multiple places we compute "bAnyRayTracingPassEnabled" by moving that flag to FDeferredSceneRenderer.
* There was another bug where view specific light ray tracing RDG buffers were created during the scene renderer that updates the ray tracing scene, potentially leading to those RDG buffers being used outside the graph where they were created.  Now the view specific light buffers are created in each scene renderer.  The code was cleaned up so the buffer creation logic only has access to the static "FSceneView", not any of the transient state in "FViewInfo", to enforce that there are no dependencies on transient state.

#rnx

diff --git a/Engine/Source/Runtime/Renderer/Private/DeferredShadingRenderer.cpp b/Engine/Source/Runtime/Renderer/Private/DeferredShadingRenderer.cpp
--- a/Engine/Source/Runtime/Renderer/Private/DeferredShadingRenderer.cpp
+++ b/Engine/Source/Runtime/Renderer/Private/DeferredShadingRenderer.cpp
@@ -378,6 +378,19 @@ FDeferredShadingSceneRenderer::FDeferredShadingSceneRenderer(const FSceneViewFam
 	, bAreLightsInLightGrid(false)
 {
 	ViewPipelineStates.SetNum(Views.Num());
+
+#if RHI_RAYTRACING
+	bAnyRayTracingPassEnabled = false;
+	for (int32 ViewIndex = 0; ViewIndex < Views.Num(); ViewIndex++)
+	{
+		bool bHasRayTracing = AnyRayTracingPassEnabled(Scene, Views[ViewIndex]);
+
+		Views[ViewIndex].bHasAnyRayTracingPass = bHasRayTracing;
+
+		bAnyRayTracingPassEnabled |= bHasRayTracing;
+	}
+	bShouldUpdateRayTracingScene = bAnyRayTracingPassEnabled;
+#endif  // RHI_RAYTRACING
 }
 
 /** 
@@ -551,12 +564,6 @@ bool FDeferredShadingSceneRenderer::GatherRayTracingWorldInstancesForView(FRDGBu
 		return false;
 	}
 
-	bool bAnyRayTracingPassEnabled = false;
-	for (int32 ViewIndex = 0; ViewIndex < Views.Num(); ViewIndex++)
-	{
-		bAnyRayTracingPassEnabled |= AnyRayTracingPassEnabled(Scene, Views[ViewIndex]);
-	}
-
 	if (!bAnyRayTracingPassEnabled)
 	{
 		return false;
@@ -1423,12 +1430,6 @@ bool FDeferredShadingSceneRenderer::SetupRayTracingPipelineStates(FRDGBuilder& G
 		return false;
 	}
 
-	bool bAnyRayTracingPassEnabled = false;
-	for (const FViewInfo& View : Views)
-	{
-		bAnyRayTracingPassEnabled |= AnyRayTracingPassEnabled(Scene, View);
-	}
-
 	if (!bAnyRayTracingPassEnabled)
 	{
 		return false;
@@ -1570,19 +1571,39 @@ bool FDeferredShadingSceneRenderer::SetupRayTracingPipelineStates(FRDGBuilder& G
 
 	ReferenceView.RayTracingSubSurfaceProfileSRV = RHICreateShaderResourceView(ReferenceView.RayTracingSubSurfaceProfileTexture, 0);
 
-	uint32 NumOfSkippedRayTracingLights = 0;
-
 	for (int32 ViewIndex = 0; ViewIndex < AllFamilyViews.Num(); ++ViewIndex)
 	{
-		FViewInfo* View = (FViewInfo*)AllFamilyViews[ViewIndex];
+		// TODO:  It would make more sense for common ray tracing resources to be in a shared structure, rather than copied into each FViewInfo.
+		//        A goal is to have the FViewInfo structure only be visible to the scene renderer that owns it, to avoid dependencies being created
+		//        that could lead to maintenance issues or interfere with paralellism goals.  For now, this works though...
+		FViewInfo* View = const_cast<FViewInfo*>(static_cast<const FViewInfo*>(AllFamilyViews[ViewIndex]));
 
 		// Send common ray tracing resources from reference view to all others.
-		if (View != &ReferenceView)
+		if (View->bHasAnyRayTracingPass && View != &ReferenceView)
 		{
 			View->RayTracingSubSurfaceProfileTexture = ReferenceView.RayTracingSubSurfaceProfileTexture;
 			View->RayTracingSubSurfaceProfileSRV = ReferenceView.RayTracingSubSurfaceProfileSRV;
 			View->RayTracingMaterialPipeline = ReferenceView.RayTracingMaterialPipeline;
 		}
+	}
+
+	return true;
+}
+
+void FDeferredShadingSceneRenderer::SetupRayTracingLightDataForViews(FRDGBuilder& GraphBuilder)
+{
+	if (!bAnyRayTracingPassEnabled)
+	{
+		return;
+	}
+
+	const bool bIsPathTracing = ViewFamily.EngineShowFlags.PathTracing;
+
+	uint32 NumOfSkippedRayTracingLights = 0;
+
+	for (int32 ViewIndex = 0; ViewIndex < Views.Num(); ++ViewIndex)
+	{
+		FViewInfo& View = Views[ViewIndex];
 
 		if (bIsPathTracing)
 		{
@@ -1592,7 +1613,7 @@ bool FDeferredShadingSceneRenderer::SetupRayTracingPipelineStates(FRDGBuilder& G
 		else
 		{
 			// This light data is a function of the camera position, so must be computed per view.
-			View->RayTracingLightDataUniformBuffer = CreateRayTracingLightData(GraphBuilder, Scene->Lights, *View, NumOfSkippedRayTracingLights);
+			View.RayTracingLightDataUniformBuffer = CreateRayTracingLightData(GraphBuilder, Scene->Lights, View, View.ShaderMap, NumOfSkippedRayTracingLights);
 		}
 	}
 
@@ -1609,20 +1630,12 @@ bool FDeferredShadingSceneRenderer::SetupRayTracingPipelineStates(FRDGBuilder& G
 			});
 	}
 #endif
-
-	return true;
 }
 
 bool FDeferredShadingSceneRenderer::DispatchRayTracingWorldUpdates(FRDGBuilder& GraphBuilder, FRDGBufferRef& OutDynamicGeometryScratchBuffer)
 {
 	OutDynamicGeometryScratchBuffer = nullptr;
 
-	bool bAnyRayTracingPassEnabled = false;
-	for (int32 ViewIndex = 0; ViewIndex < Views.Num(); ViewIndex++)
-	{
-		bAnyRayTracingPassEnabled |= AnyRayTracingPassEnabled(Scene, Views[ViewIndex]);
-	}
-
 	if (!IsRayTracingEnabled() || !bAnyRayTracingPassEnabled || Views.Num() == 0)
 	{
 		// This needs to happen even when ray tracing is not enabled
@@ -1757,12 +1770,6 @@ static void ReleaseRaytracingResources(FRDGBuilder& GraphBuilder, TArrayView<FVi
 
 void FDeferredShadingSceneRenderer::WaitForRayTracingScene(FRDGBuilder& GraphBuilder, FRDGBufferRef DynamicGeometryScratchBuffer)
 {
-	bool bAnyRayTracingPassEnabled = false;
-	for (int32 ViewIndex = 0; ViewIndex < Views.Num(); ViewIndex++)
-	{
-		bAnyRayTracingPassEnabled |= AnyRayTracingPassEnabled(Scene, Views[ViewIndex]);
-	}
-
 	if (!bAnyRayTracingPassEnabled)
 	{
 		return;
@@ -1840,8 +1847,9 @@ void FDeferredShadingSceneRenderer::WaitForRayTracingScene(FRDGBuilder& GraphBui
 		// Send ray tracing resources from reference view to all others.
 		for (int32 ViewIndex = 0; ViewIndex < AllFamilyViews.Num(); ++ViewIndex)
 		{
-			FViewInfo* View = (FViewInfo*)AllFamilyViews[ViewIndex];
-			if (View != &ReferenceView)
+			// See comment above where we copy "RayTracingSubSurfaceProfileTexture" to each view...
+			FViewInfo* View = const_cast<FViewInfo*>(static_cast<const FViewInfo*>(AllFamilyViews[ViewIndex]));
+			if (View->bHasAnyRayTracingPass && View != &ReferenceView)
 			{
 				View->RayTracingMaterialGatherPipeline = ReferenceView.RayTracingMaterialGatherPipeline;
 				View->LumenHardwareRayTracingMaterialPipeline = ReferenceView.LumenHardwareRayTracingMaterialPipeline;
@@ -2674,7 +2682,7 @@ void FDeferredShadingSceneRenderer::Render(FRDGBuilder& GraphBuilder)
 	// Async AS builds can potentially overlap with BasePass.  We only need to update ray tracing scene for the first view family,
 	// if multiple are rendered in a single scene render call.
 	FRDGBufferRef DynamicGeometryScratchBuffer = nullptr;
-	if (bIsFirstSceneRenderer)
+	if (bShouldUpdateRayTracingScene)
 	{
 		DispatchRayTracingWorldUpdates(GraphBuilder, DynamicGeometryScratchBuffer);
 	}
@@ -2809,6 +2817,11 @@ void FDeferredShadingSceneRenderer::Render(FRDGBuilder& GraphBuilder)
 		Strata::AddStrataMaterialClassificationPass(GraphBuilder, SceneTextures, Views);
 	}
 
+#if RHI_RAYTRACING
+	/** Should be called somewhere before "WaitForRayTracingScene" */
+	SetupRayTracingLightDataForViews(GraphBuilder);
+#endif // RHI_RAYTRACING
+
 	// Shadows, lumen and fog after base pass
 	if (!bHasRayTracedOverlay)
 	{
diff --git a/Engine/Source/Runtime/Renderer/Private/DeferredShadingRenderer.h b/Engine/Source/Runtime/Renderer/Private/DeferredShadingRenderer.h
--- a/Engine/Source/Runtime/Renderer/Private/DeferredShadingRenderer.h
+++ b/Engine/Source/Runtime/Renderer/Private/DeferredShadingRenderer.h
@@ -211,6 +211,11 @@ public:
 
 	FLumenCardRenderer LumenCardRenderer;
 
+#if RHI_RAYTRACING
+	bool bAnyRayTracingPassEnabled;
+	bool bShouldUpdateRayTracingScene;
+#endif
+
 	FDeferredShadingSceneRenderer(const FSceneViewFamily* InViewFamily, FHitProxyConsumer* HitProxyConsumer);
 
 	/** Determine and commit the final state of the pipeline for the view family and views. */
@@ -1007,6 +1012,7 @@ private:
 	bool GatherRayTracingWorldInstancesForView(FRDGBuilder& GraphBuilder, FViewInfo& View, FRayTracingScene& RayTracingScene);
 
 	bool SetupRayTracingPipelineStates(FRDGBuilder& GraphBuilder);
+	void SetupRayTracingLightDataForViews(FRDGBuilder& GraphBuilder);
 	bool DispatchRayTracingWorldUpdates(FRDGBuilder& GraphBuilder, FRDGBufferRef& OutDynamicGeometryScratchBuffer);
 
 	/** Functions to create ray tracing pipeline state objects for various effects */
@@ -1061,7 +1067,6 @@ private:
 	static FRHIRayTracingShader* GetRayTracingLightingMissShader(FViewInfo& View);
 
 	const FRHITransition* RayTracingDynamicGeometryUpdateEndTransition = nullptr; // Signaled when all AS for this frame are built
-
 #endif // RHI_RAYTRACING
 
 	/** Set to true if lights were injected into the light grid (this controlled by somewhat complex logic, this flag is used to cross-check). */
diff --git a/Engine/Source/Runtime/Renderer/Private/RayTracing/RayTracingLighting.cpp b/Engine/Source/Runtime/Renderer/Private/RayTracing/RayTracingLighting.cpp
--- a/Engine/Source/Runtime/Renderer/Private/RayTracing/RayTracingLighting.cpp
+++ b/Engine/Source/Runtime/Renderer/Private/RayTracing/RayTracingLighting.cpp
@@ -63,7 +63,6 @@
 
 static void SelectRaytracingLights(
 	const TSparseArray<FLightSceneInfoCompact, TAlignedSparseArrayAllocator<alignof(FLightSceneInfoCompact)>>& Lights,
-	const FViewInfo& View,
 	TArray<int32>& OutSelectedLights,
 	uint32& NumOfSkippedRayTracingLights
 )
@@ -97,7 +96,8 @@
 static void CreateRaytracingLightCullingStructure(
 	FRDGBuilder& GraphBuilder,
 	const TSparseArray<FLightSceneInfoCompact, TAlignedSparseArrayAllocator<alignof(FLightSceneInfoCompact)>>& Lights,
-	const FViewInfo& View,
+	const FSceneView& View,
+	FGlobalShaderMap* ShaderMap,
 	const TArray<int32>& LightIndices,
 	FRDGBufferRef& LightCullVolume,
 	FRDGBufferRef& LightIndicesBuffer)
@@ -155,7 +155,7 @@
 		Parameters->CellCount = CellsPerDim;
 		Parameters->CellScale = CVarRayTracingLightingCellSize.GetValueOnRenderThread() / 2.0f; // cells are based on pow2, and initial cell is 2^1, so scale is half min cell size
 
-		auto Shader = View.ShaderMap->GetShader<FSetupRayTracingLightCullData>();
+		auto Shader = ShaderMap->GetShader<FSetupRayTracingLightCullData>();
 		FComputeShaderUtils::AddPass(GraphBuilder, RDG_EVENT_NAME("LightCullingVolumeCompute"), Shader, Parameters, FIntVector(CellsPerDim, CellsPerDim, CellsPerDim));
 	}
 }
@@ -165,7 +165,7 @@
 	FRDGBuilder& GraphBuilder,
 	const TSparseArray<FLightSceneInfoCompact, TAlignedSparseArrayAllocator<alignof(FLightSceneInfoCompact)>>& Lights,
 	const TArray<int32>& LightIndices,
-	const FViewInfo& View,
+	const FSceneView& View,
 	FRaytracingLightDataPacked& OutLightData,
 	TArrayView<FRTLightingData>& OutLightDataArray)
 {
@@ -293,7 +293,8 @@
 TRDGUniformBufferRef<FRaytracingLightDataPacked> CreateRayTracingLightData(
 	FRDGBuilder& GraphBuilder,
 	const TSparseArray<FLightSceneInfoCompact, TAlignedSparseArrayAllocator<alignof(FLightSceneInfoCompact)>>& Lights,
-	const FViewInfo& View,
+	const FSceneView& View,
+	FGlobalShaderMap* ShaderMap,
 	uint32& NumOfSkippedRayTracingLights)
 {
 	auto* LightData = GraphBuilder.AllocParameters<FRaytracingLightDataPacked>();
@@ -301,12 +302,12 @@
 	LightData->CellScale = CVarRayTracingLightingCellSize.GetValueOnRenderThread() / 2.0f;
 
 	TArray<int32> LightIndices;
-	SelectRaytracingLights(Lights, View, LightIndices, NumOfSkippedRayTracingLights);
+	SelectRaytracingLights(Lights, LightIndices, NumOfSkippedRayTracingLights);
 
 	// Create light culling volume
 	FRDGBufferRef LightCullVolume;
 	FRDGBufferRef LightIndicesBuffer;
-	CreateRaytracingLightCullingStructure(GraphBuilder, Lights, View, LightIndices, LightCullVolume, LightIndicesBuffer);
+	CreateRaytracingLightCullingStructure(GraphBuilder, Lights, View, ShaderMap, LightIndices, LightCullVolume, LightIndicesBuffer);
 
 	FRDGUploadData<FRTLightingData> LightDataArray(GraphBuilder, FMath::Max(LightIndices.Num(), 1));
 	SetupRaytracingLightDataPacked(GraphBuilder, Lights, LightIndices, View, *LightData, LightDataArray);
diff --git a/Engine/Source/Runtime/Renderer/Private/RayTracing/RayTracingLighting.h b/Engine/Source/Runtime/Renderer/Private/RayTracing/RayTracingLighting.h
--- a/Engine/Source/Runtime/Renderer/Private/RayTracing/RayTracingLighting.h
+++ b/Engine/Source/Runtime/Renderer/Private/RayTracing/RayTracingLighting.h
@@ -58,7 +58,8 @@
 TRDGUniformBufferRef<FRaytracingLightDataPacked> CreateRayTracingLightData(
 	FRDGBuilder& GraphBuilder,
 	const TSparseArray<FLightSceneInfoCompact, TAlignedSparseArrayAllocator<alignof(FLightSceneInfoCompact)>>& Lights,
-	const FViewInfo& View,
+	const FSceneView& View,
+	FGlobalShaderMap* ShaderMap,
 	uint32& NumOfSkippedRayTracingLights);
 
 #endif
diff --git a/Engine/Source/Runtime/Renderer/Private/SceneRendering.cpp b/Engine/Source/Runtime/Renderer/Private/SceneRendering.cpp
--- a/Engine/Source/Runtime/Renderer/Private/SceneRendering.cpp
+++ b/Engine/Source/Runtime/Renderer/Private/SceneRendering.cpp
@@ -3065,7 +3065,7 @@ static void GetCrossGPUTransfers(FSceneRenderer* SceneRenderer, TArray<FCrossGPU
 }
 #endif // WITH_MGPU
 
-void FSceneRenderer::PreallocateCrossGPUFences(TArray<FSceneRenderer*> SceneRenderers)
+void FSceneRenderer::PreallocateCrossGPUFences(const TArray<FSceneRenderer*>& SceneRenderers)
 {
 #if WITH_MGPU
 	// We can only apply cross GPU fence wait transfer optimizations if AFR isn't enabled.  This is because we can't determine
@@ -3800,24 +3800,70 @@ void FSceneRenderer::SetupMeshPass(FViewInfo& View, FExclusiveDepthStencil::Type
 	}
 }
 
-FSceneRenderer* FSceneRenderer::CreateSceneRenderer(const FSceneViewFamily* InViewFamily, FHitProxyConsumer* HitProxyConsumer)
+void FSceneRenderer::CreateSceneRenderers(TArrayView<const FSceneViewFamily*> InViewFamilies, FHitProxyConsumer* HitProxyConsumer, TArray<FSceneRenderer*>& OutSceneRenderers)
 {
-	check(InViewFamily);
-
-	EShadingPath ShadingPath = InViewFamily->Scene->GetShadingPath();
-	FSceneRenderer* SceneRenderer = nullptr;
+	OutSceneRenderers.Empty(InViewFamilies.Num());
 
-	if (ShadingPath == EShadingPath::Deferred)
+	if (!InViewFamilies.Num())
 	{
-		SceneRenderer = new FDeferredShadingSceneRenderer(InViewFamily, HitProxyConsumer);
+		return;
 	}
-	else 
+
+	const FSceneInterface* Scene = InViewFamilies[0]->Scene;
+	check(Scene);
+
+	EShadingPath ShadingPath = Scene->GetShadingPath();
+
+#if RHI_RAYTRACING
+	// For multi-view-family scene rendering, we need to determine which scene renderer will update the ray tracing scene.
+	// This will be the first view family that uses ray tracing, and subsequent families that use ray tracing can skip
+	// that step.
+	bool bShouldUpdateRayTracingScene = true;
+#endif
+
+	for (int32 FamilyIndex = 0; FamilyIndex < InViewFamilies.Num(); FamilyIndex++)
 	{
-		check(ShadingPath == EShadingPath::Mobile);
-		SceneRenderer = new FMobileSceneRenderer(InViewFamily, HitProxyConsumer);
+		const FSceneViewFamily* InViewFamily = InViewFamilies[FamilyIndex];
+
+		check(InViewFamily);
+		check(InViewFamily->Scene == Scene);
+
+		if (ShadingPath == EShadingPath::Deferred)
+		{
+			FDeferredShadingSceneRenderer* SceneRenderer = new FDeferredShadingSceneRenderer(InViewFamily, HitProxyConsumer);
+#if RHI_RAYTRACING
+			if (SceneRenderer->bAnyRayTracingPassEnabled)
+			{
+				SceneRenderer->bShouldUpdateRayTracingScene = bShouldUpdateRayTracingScene;
+				bShouldUpdateRayTracingScene = false;
+			}
+#endif
+
+			OutSceneRenderers.Add(SceneRenderer);
+		}
+		else
+		{
+			check(ShadingPath == EShadingPath::Mobile);
+			OutSceneRenderers.Add(new FMobileSceneRenderer(InViewFamily, HitProxyConsumer));
+		}
+
+		OutSceneRenderers.Last()->bIsFirstSceneRenderer = (FamilyIndex == 0);
+		OutSceneRenderers.Last()->bIsLastSceneRenderer = (FamilyIndex == InViewFamilies.Num() - 1);
 	}
+}
+
+FSceneRenderer* FSceneRenderer::CreateSceneRenderer(const FSceneViewFamily* InViewFamily, FHitProxyConsumer* HitProxyConsumer)
+{
+	check(InViewFamily);
+
+	TArray<FSceneRenderer*> SceneRenderers;
+
+	TArray<const FSceneViewFamily*> InViewFamilies;
+	InViewFamilies.Add(InViewFamily);
 
-	return SceneRenderer;
+	CreateSceneRenderers(InViewFamilies, HitProxyConsumer, SceneRenderers);
+
+	return SceneRenderers[0];
 }
 
 void FSceneRenderer::OnStartRender(FRHICommandListImmediate& RHICmdList)
@@ -4560,16 +4606,15 @@ void FRendererModule::BeginRenderingViewFamilies(FCanvas* Canvas, TArrayView<FSc
 
 		// Construct the scene renderers.  This copies the view family attributes into its own structures.
 		TArray<FSceneRenderer*> SceneRenderers;
-		SceneRenderers.SetNumUninitialized(ViewFamilies.Num());
 
-		for (int32 FamilyIndex = 0; FamilyIndex < ViewFamilies.Num(); FamilyIndex++)
+		TArray<const FSceneViewFamily*> ViewFamiliesConst;
+		for (FSceneViewFamily* ViewFamily : ViewFamilies)
 		{
-			SceneRenderers[FamilyIndex] = FSceneRenderer::CreateSceneRenderer(ViewFamilies[FamilyIndex], Canvas->GetHitProxyConsumer());
-
-			SceneRenderers[FamilyIndex]->bIsFirstSceneRenderer = (FamilyIndex == 0);
-			SceneRenderers[FamilyIndex]->bIsLastSceneRenderer = (FamilyIndex == ViewFamilies.Num() - 1);
+			ViewFamiliesConst.Add(ViewFamily);
 		}
 
+		FSceneRenderer::CreateSceneRenderers(ViewFamiliesConst, Canvas->GetHitProxyConsumer(), SceneRenderers);
+
 		bool bShowHitProxies = false;
 		for (FSceneRenderer* SceneRenderer : SceneRenderers)
 		{
diff --git a/Engine/Source/Runtime/Renderer/Private/SceneRendering.h b/Engine/Source/Runtime/Renderer/Private/SceneRendering.h
--- a/Engine/Source/Runtime/Renderer/Private/SceneRendering.h
+++ b/Engine/Source/Runtime/Renderer/Private/SceneRendering.h
@@ -1502,6 +1502,8 @@ public:
 
 	TRDGUniformBufferRef<FRaytracingLightDataPacked>	RayTracingLightDataUniformBuffer;
 	TRDGUniformBufferRef<FRayTracingDecals>				RayTracingDecalUniformBuffer;
+
+	bool bHasAnyRayTracingPass = false;
 #endif // RHI_RAYTRACING
 
 	/**
@@ -1981,6 +1983,9 @@ public:
 	/** Creates a scene renderer based on the current feature level. */
 	RENDERER_API static FSceneRenderer* CreateSceneRenderer(const FSceneViewFamily* InViewFamily, FHitProxyConsumer* HitProxyConsumer);
 
+	/** Creates multiple scene renderers based on the current feature level.  All view families must point to the same Scene. */
+	RENDERER_API static void CreateSceneRenderers(TArrayView<const FSceneViewFamily*> InViewFamilies, FHitProxyConsumer* HitProxyConsumer, TArray<FSceneRenderer*>& OutSceneRenderers);
+
 	/** Setups FViewInfo::ViewRect according to ViewFamilly's ScreenPercentageInterface. */
 	void PrepareViewRectsForRendering(FRHICommandListImmediate& RHICmdList);
 
@@ -1994,7 +1999,7 @@ public:
 #endif
 
 	/** Logic to update render targets across all GPUs */
-	static void PreallocateCrossGPUFences(TArray<FSceneRenderer*> SceneRenderers);
+	static void PreallocateCrossGPUFences(const TArray<FSceneRenderer*>& SceneRenderers);
 	void DoCrossGPUTransfers(FRDGBuilder& GraphBuilder, FRDGTextureRef ViewFamilyTexture);
 	void FlushCrossGPUFences(FRDGBuilder& GraphBuilder);
 
