Bug fix:  ShouldRenderRayTracingEffect can now be called in the game thread as well as the render thread, so it needs to use "GetValueOnAnyThread" for CVAR access.

#rnx

diff --git a/Engine/Source/Runtime/Renderer/Private/DeferredShadingRenderer.cpp b/Engine/Source/Runtime/Renderer/Private/DeferredShadingRenderer.cpp
--- a/Engine/Source/Runtime/Renderer/Private/DeferredShadingRenderer.cpp
+++ b/Engine/Source/Runtime/Renderer/Private/DeferredShadingRenderer.cpp
@@ -3462,11 +3462,11 @@ bool ShouldRenderRayTracingEffect(bool bEffectEnabled, ERayTracingPipelineCompat
 	}
 
 	const bool bAllowPipeline = GRHISupportsRayTracingShaders && 
-								CVarRayTracingAllowPipeline.GetValueOnRenderThread() &&
+								CVarRayTracingAllowPipeline.GetValueOnAnyThread() &&
 								EnumHasAnyFlags(CompatibilityFlags, ERayTracingPipelineCompatibilityFlags::FullPipeline);
 
 	const bool bAllowInline = GRHISupportsInlineRayTracing && 
-							  CVarRayTracingAllowInline.GetValueOnRenderThread() &&
+							  CVarRayTracingAllowInline.GetValueOnAnyThread() &&
 							  EnumHasAnyFlags(CompatibilityFlags, ERayTracingPipelineCompatibilityFlags::Inline);
 
 	// Disable the effect if current machine does not support the full ray tracing pipeline and the effect can't fall back to inline mode or vice versa.
@@ -3475,7 +3475,7 @@ bool ShouldRenderRayTracingEffect(bool bEffectEnabled, ERayTracingPipelineCompat
 		return false;
 	}
 
-	const int32 OverrideMode = CVarForceAllRayTracingEffects.GetValueOnRenderThread();
+	const int32 OverrideMode = CVarForceAllRayTracingEffects.GetValueOnAnyThread();
 
 	if (OverrideMode >= 0)
 	{
