Bind virtual texture feedback buffer for DrawTileMesh
Also create a FViewFamilyInfo for DrawTileMesh so that we have a valid SceneTexturesConfig when setting up view uniforms.
Moved material baking renderdoc capture so that it captures warm up frames.

diff --git a/Engine/Source/Developer/MaterialBaking/Private/MaterialBakingModule.cpp b/Engine/Source/Developer/MaterialBaking/Private/MaterialBakingModule.cpp
--- a/Engine/Source/Developer/MaterialBaking/Private/MaterialBakingModule.cpp
+++ b/Engine/Source/Developer/MaterialBaking/Private/MaterialBakingModule.cpp
@@ -380,8 +380,6 @@
 		}
 	}
 
-	RenderCaptureInterface::FScopedCapture RenderCapture(CVarMaterialBakingRDOCCapture.GetValueOnAnyThread() == 1, TEXT("MaterialBaking"));
-
 	TRACE_CPUPROFILER_EVENT_SCOPE(FMaterialBakingModule::BakeMaterials)
 
 	checkf(MaterialSettings.Num() == MeshSettings.Num(), TEXT("Number of material settings does not match that of MeshSettings"));
@@ -621,13 +619,19 @@
  							}
 
 							// Do rendering
-							for (int WarmupIndex = 0; WarmupIndex < WarmupIterationCount; ++WarmupIndex)
 							{
-								Canvas.Clear(RenderTarget->ClearColor);
-								FCanvas::FCanvasSortElement& SortElement = Canvas.GetSortElement(Canvas.TopDepthSortKey());
-								SortElement.RenderBatchArray.Add(&RenderItem);
-								Canvas.Flush_RenderThread(RHICmdList);
-								SortElement.RenderBatchArray.Empty();
+								RenderCaptureInterface::FScopedCapture RenderCapture(CVarMaterialBakingRDOCCapture.GetValueOnAnyThread() == 1, &RHICmdList, TEXT("MaterialBaking"));
+
+								for (int WarmupIndex = 0; WarmupIndex < WarmupIterationCount; ++WarmupIndex)
+								{
+									Canvas.Clear(RenderTarget->ClearColor);
+									FCanvas::FCanvasSortElement& SortElement = Canvas.GetSortElement(Canvas.TopDepthSortKey());
+									SortElement.RenderBatchArray.Add(&RenderItem);
+									Canvas.Flush_RenderThread(RHICmdList);
+									SortElement.RenderBatchArray.Empty();
+
+									RHICmdList.ImmediateFlush(EImmediateFlushType::FlushRHIThreadFlushResources);
+								}
 							}
 
 							FTexture2DRHIRef StagingBufferRef = StagingBufferPool.CreateStagingBuffer_RenderThread(RHICmdList, RenderTargetResource->GetSizeX(), RenderTargetResource->GetSizeY(), RenderTarget->GetFormat(), RenderTarget->IsSRGB());
diff --git a/Engine/Source/Runtime/Renderer/Private/Renderer.cpp b/Engine/Source/Runtime/Renderer/Private/Renderer.cpp
--- a/Engine/Source/Runtime/Renderer/Private/Renderer.cpp
+++ b/Engine/Source/Runtime/Renderer/Private/Renderer.cpp
@@ -130,10 +130,18 @@ void FRendererModule::DrawTileMesh(FCanvasRenderContext& RenderContext, FMeshPas
 		//@todo - reuse this view for multiple tiles, this is going to be slow for each tile
 		FViewInfo& View = *RenderContext.Alloc<FViewInfo>(&SceneView);
 		View.ViewRect = View.UnscaledViewRect;
+		FViewFamilyInfo* ViewFamily = RenderContext.Alloc<FViewFamilyInfo>(*SceneView.Family);
+		ViewFamily->Views.Add(&View);
+		View.Family = ViewFamily;
+
+		// Default init of SceneTexturesConfig will take extents from FSceneTextureExtentState.
+		// We want the view extents, so explicitly set that.
+		FSceneTexturesConfig::InitializeViewFamily(*ViewFamily);
+		ViewFamily->SceneTexturesConfig.Extent = View.ViewRect.Size();
 
 		const auto FeatureLevel = View.GetFeatureLevel();
 		const EShadingPath ShadingPath = FSceneInterface::GetShadingPath(FeatureLevel);
-		const FSceneViewFamily* ViewFamily = View.Family;
+
 		FScene* Scene = nullptr;
 
 		if (ViewFamily->Scene)
@@ -210,6 +218,19 @@ void FRendererModule::DrawTileMesh(FCanvasRenderContext& RenderContext, FMeshPas
 			FRDGSystemTextures::Create(GraphBuilder);
 		}
 
+		// Materials sampling VTs need FVirtualTextureSystem to be updated before being rendered
+		const FMaterial& MeshMaterial = Mesh.MaterialRenderProxy->GetIncompleteMaterialWithFallback(FeatureLevel);
+		const bool bUseVirtualTexturing = UseVirtualTexturing(FeatureLevel) && !MeshMaterial.GetUniformVirtualTextureExpressions().IsEmpty();
+		if (bUseVirtualTexturing)
+		{
+			RDG_GPU_STAT_SCOPE(GraphBuilder, VirtualTextureUpdate);
+			FVirtualTextureSystem::Get().AllocateResources(GraphBuilder, FeatureLevel);
+			FVirtualTextureSystem::Get().CallPendingCallbacks();
+			FVirtualTextureSystem::Get().Update(GraphBuilder, FeatureLevel, Scene);
+
+			VirtualTextureFeedbackBegin(GraphBuilder, TArrayView<const FViewInfo>(&View, 1), RenderContext.GetViewportRect().Size());
+		}
+
 		View.InitRHIResources();
 		View.ForwardLightingResources.SetUniformBuffer(CreateDummyForwardLightUniformBuffer(GraphBuilder));
 
@@ -220,20 +241,6 @@ void FRendererModule::DrawTileMesh(FCanvasRenderContext& RenderContext, FMeshPas
 			EmptyReflectionCaptureUniformBuffer = TUniformBufferRef<FReflectionCaptureShaderData>::CreateUniformBufferImmediate(EmptyData, UniformBuffer_SingleFrame);
 		}
 
-		//get the blend mode of the material
-		const FMaterial& MeshMaterial = Mesh.MaterialRenderProxy->GetIncompleteMaterialWithFallback(FeatureLevel);
-		const EBlendMode MaterialBlendMode = MeshMaterial.GetBlendMode();
-
-		const bool bUseVirtualTexturing = UseVirtualTexturing(FeatureLevel);
-		// Materials sampling VTs need FVirtualTextureSystem to be updated before being rendered :
-		if (bUseVirtualTexturing && !MeshMaterial.GetUniformVirtualTextureExpressions().IsEmpty())
-		{
-			RDG_GPU_STAT_SCOPE(GraphBuilder, VirtualTextureUpdate);
-			FVirtualTextureSystem::Get().AllocateResources(GraphBuilder, FeatureLevel);
-			FVirtualTextureSystem::Get().CallPendingCallbacks();
-			FVirtualTextureSystem::Get().Update(GraphBuilder, FeatureLevel, Scene);
-		}
-
 		RDG_EVENT_SCOPE(GraphBuilder, "DrawTileMesh");
 
 		auto* PassParameters = GraphBuilder.AllocParameters<FDrawTileMeshPassParameters>();
@@ -242,6 +249,8 @@ void FRendererModule::DrawTileMesh(FCanvasRenderContext& RenderContext, FMeshPas
 		PassParameters->ReflectionCapture = EmptyReflectionCaptureUniformBuffer;
 		PassParameters->InstanceCulling = FInstanceCullingContext::CreateDummyInstanceCullingUniformBuffer(GraphBuilder);
 
+		const EBlendMode MaterialBlendMode = MeshMaterial.GetBlendMode();
+
 		// handle translucent material blend modes, not relevant in MaterialTexCoordScalesAnalysis since it outputs the scales.
 		if (ViewFamily->GetDebugViewShaderMode() == DVSM_OutputMaterialTextureScales)
 		{
@@ -400,6 +409,12 @@ void FRendererModule::DrawTileMesh(FCanvasRenderContext& RenderContext, FMeshPas
 				}
 			}
 		}
+
+		if (bUseVirtualTexturing)
+		{
+			RDG_GPU_STAT_SCOPE(GraphBuilder, VirtualTextureUpdate);
+			VirtualTextureFeedbackEnd(GraphBuilder);
+		}
 	}
 }
 
diff --git a/Engine/Source/Runtime/Renderer/Private/SceneRendering.cpp b/Engine/Source/Runtime/Renderer/Private/SceneRendering.cpp
--- a/Engine/Source/Runtime/Renderer/Private/SceneRendering.cpp
+++ b/Engine/Source/Runtime/Renderer/Private/SceneRendering.cpp
@@ -1836,7 +1836,7 @@ void FViewInfo::SetupUniformBufferParameters(
 	check(VirtualTextureFeedbackScale == 1 << FMath::FloorLog2(VirtualTextureFeedbackScale));
 	ViewUniformShaderParameters.VirtualTextureFeedbackShift = FMath::FloorLog2(VirtualTextureFeedbackScale);
 	ViewUniformShaderParameters.VirtualTextureFeedbackMask = VirtualTextureFeedbackScale - 1;
-	ViewUniformShaderParameters.VirtualTextureFeedbackStride = GetVirtualTextureFeedbackBufferSize(SceneTexturesConfig.Extent).X;
+	ViewUniformShaderParameters.VirtualTextureFeedbackStride = GetVirtualTextureFeedbackBufferSize(EffectiveBufferSize).X;
 	// Use some low(ish) discrepancy sequence to run over every pixel in the virtual texture feedback tile.
 	ViewUniformShaderParameters.VirtualTextureFeedbackJitterOffset = SampleVirtualTextureFeedbackSequence(FrameIndex);
 	// Offset the selected sample index for each frame and add an additional offset each time we iterate over a full virtual texture feedback tile to ensure we get full coverage of sample indices over time.
diff --git a/Engine/Source/Runtime/Renderer/Private/VT/VirtualTextureFeedbackBuffer.cpp b/Engine/Source/Runtime/Renderer/Private/VT/VirtualTextureFeedbackBuffer.cpp
--- a/Engine/Source/Runtime/Renderer/Private/VT/VirtualTextureFeedbackBuffer.cpp
+++ b/Engine/Source/Runtime/Renderer/Private/VT/VirtualTextureFeedbackBuffer.cpp
@@ -93,14 +93,14 @@ void FVirtualTextureFeedbackBuffer::Begin(FRDGBuilder& GraphBuilder, const FVirt
 		UAV = PooledBuffer->GetOrCreateUAV(UAVDesc);
 	}
 
-	// We can go ahead and just clear this now. No need to wait on the RDG timeline.
-	FRHICommandListImmediate& RHICmdList = GraphBuilder.RHICmdList;
-
-	// Clear virtual texture feedback to default value
-	RHICmdList.Transition(FRHITransitionInfo(UAV, ERHIAccess::Unknown, ERHIAccess::UAVCompute));
-	RHICmdList.ClearUAVUint(UAV, FUintVector4(~0u, ~0u, ~0u, ~0u));
-	RHICmdList.Transition(FRHITransitionInfo(UAV, ERHIAccess::UAVCompute, ERHIAccess::UAVGraphics));
-	RHICmdList.BeginUAVOverlap(UAV);
+	AddPass(GraphBuilder, RDG_EVENT_NAME("VirtualTextureClear"), [this](FRHICommandListImmediate& RHICmdList)
+	{
+		// Clear virtual texture feedback to default value
+		RHICmdList.Transition(FRHITransitionInfo(UAV, ERHIAccess::Unknown, ERHIAccess::UAVCompute));
+		RHICmdList.ClearUAVUint(UAV, FUintVector4(~0u, ~0u, ~0u, ~0u));
+		RHICmdList.Transition(FRHITransitionInfo(UAV, ERHIAccess::UAVCompute, ERHIAccess::UAVGraphics));
+		RHICmdList.BeginUAVOverlap(UAV);
+	});
 }
 
 void FVirtualTextureFeedbackBuffer::End(FRDGBuilder& GraphBuilder)
