Skip marking primitive as rendered in VSM if not added to any draw lists (removes invalidations from translucent, e.g., water).

diff --git a/Engine/Source/Runtime/Renderer/Private/ShadowRendering.h b/Engine/Source/Runtime/Renderer/Private/ShadowRendering.h
--- a/Engine/Source/Runtime/Renderer/Private/ShadowRendering.h
+++ b/Engine/Source/Runtime/Renderer/Private/ShadowRendering.h
@@ -626,7 +626,7 @@ public:
 	/**
 	 * Adds a primitive to the shadow's subject list.
 	 */
-	void AddSubjectPrimitive(FPrimitiveSceneInfo* PrimitiveSceneInfo, TArrayView<FViewInfo> ViewArray, bool bRecordShadowSubjectForMobileShading);
+	bool AddSubjectPrimitive(FPrimitiveSceneInfo* PrimitiveSceneInfo, TArrayView<FViewInfo> ViewArray, bool bRecordShadowSubjectForMobileShading);
 
 	uint64 AddSubjectPrimitive_AnyThread(
 		const FPrimitiveSceneInfoCompact& PrimitiveSceneInfoCompact,
diff --git a/Engine/Source/Runtime/Renderer/Private/ShadowSetup.cpp b/Engine/Source/Runtime/Renderer/Private/ShadowSetup.cpp
--- a/Engine/Source/Runtime/Renderer/Private/ShadowSetup.cpp
+++ b/Engine/Source/Runtime/Renderer/Private/ShadowSetup.cpp
@@ -1576,11 +1576,11 @@ bool FProjectedShadowInfo::ShouldDrawStaticMeshes_AnyThread(
 	return bDrawingStaticMeshes;
 }
 
-void FProjectedShadowInfo::AddSubjectPrimitive(FPrimitiveSceneInfo* PrimitiveSceneInfo, TArrayView<FViewInfo> ViewArray, bool bShouldRecordShadowSubjectsForMobileShading)
+bool FProjectedShadowInfo::AddSubjectPrimitive(FPrimitiveSceneInfo* PrimitiveSceneInfo, TArrayView<FViewInfo> ViewArray, bool bShouldRecordShadowSubjectsForMobileShading)
 {
 	// Ray traced shadows use the GPU managed distance field object buffers, no CPU culling should be used
 	check(!bRayTracedDistanceField);
-
+	bool bWasAdded = false;
 	if (!ReceiverPrimitives.Contains(PrimitiveSceneInfo)
 		// Far cascade only casts from primitives marked for it
 		&& TestPrimitiveFarCascadeConditions(PrimitiveSceneInfo->Proxy->CastsFarShadow(), PrimitiveSceneInfo->Proxy->GetBounds()))
@@ -1722,6 +1722,7 @@ void FProjectedShadowInfo::AddSubjectPrimitive(FPrimitiveSceneInfo* PrimitiveSce
 				{
 					DependentView->VisibleLightInfos[GetLightSceneInfo().Id].MobileCSMSubjectPrimitives.AddSubjectPrimitive(PrimitiveSceneInfo, PrimitiveId);
 				}
+				bWasAdded = true;
 			}
 			// EShadowMeshSelection::All is intentional because ProxySupportsGPUScene being false may indicate that some VFs (e.g., some LODs) support GPUScene while others don't
 			// thus we have to leave the final decision until the mesh batches are produced.
@@ -1734,6 +1735,7 @@ void FProjectedShadowInfo::AddSubjectPrimitive(FPrimitiveSceneInfo* PrimitiveSce
 				{
 					DependentView->VisibleLightInfos[GetLightSceneInfo().Id].MobileCSMSubjectPrimitives.AddSubjectPrimitive(PrimitiveSceneInfo, PrimitiveId);
 				}
+				bWasAdded = true;
 			}
 		}
 
@@ -1741,8 +1743,10 @@ void FProjectedShadowInfo::AddSubjectPrimitive(FPrimitiveSceneInfo* PrimitiveSce
 		if (bTranslucentShadow && bTranslucentRelevance && bShadowRelevance)
 		{
 			SubjectTranslucentPrimitives.Add(PrimitiveSceneInfo);
+			bWasAdded = true;
 		}
 	}
+	return bWasAdded;
 }
 
 bool FProjectedShadowInfo::TestPrimitiveFarCascadeConditions(bool bPrimitiveCastsFarShadow, const FBoxSphereBounds& Bounds) const
@@ -2003,7 +2007,10 @@ void FProjectedShadowInfo::FinalizeAddSubjectPrimitive(
 
 	if (Result.bAddOnRenderThread)
 	{
-		AddSubjectPrimitive(PrimitiveSceneInfo, ViewArray, Result.bShouldRecordShadowSubjectsForMobile);
+		if (AddSubjectPrimitive(PrimitiveSceneInfo, ViewArray, Result.bShouldRecordShadowSubjectsForMobile) && VirtualShadowMapClipmap.IsValid())
+		{
+			VirtualShadowMapClipmap->OnPrimitiveRendered(PrimitiveSceneInfo);
+		}
 		return;
 	}
 
@@ -3695,8 +3702,8 @@ void FSceneRenderer::CreateWholeSceneProjectedShadow(
 								FBoxSphereBounds const& Bounds = PrimitiveSceneInfo->Proxy->GetBounds();
 								if (!bTestConvexHull || IntersectsConvexHulls(LightViewFrustumConvexHulls, Bounds))
 								{
-									ProjectedShadowInfo->AddSubjectPrimitive(Interaction->GetPrimitiveSceneInfo(), Views, false);
-									if (VirtualSmCacheEntry.IsValid())
+									if (ProjectedShadowInfo->AddSubjectPrimitive(Interaction->GetPrimitiveSceneInfo(), Views, false) 
+										&& VirtualSmCacheEntry.IsValid())
 									{
 										VirtualSmCacheEntry->OnPrimitiveRendered(PrimitiveSceneInfo);
 									}
