Add experimental temporal denoising for offline path tracing that is spatially denoised. It gives a final touch to mitigate the temporal instability while preserving fine details.

Basic algorithm:

Estimate the pixel correspondences (motion vectors) to the previous frame where their local patches match the best using non-local means. Then history is combined with the current frame with a temporal low pass filter (exponential moving average & temporal bilaterial filtering). To better control the parameters, we use perception error DeltaE2000. Pixel difference larger than 10 dE is rejected for detail preserving.

How to enable:

Add `r.PathTracing.TemporalDenoiser 1` in MRQ console variables when path tracing denoising is enabled in postprocess volume.

diff --git a/Engine/Shaders/Private/ColorDifference.ush b/Engine/Shaders/Private/ColorDifference.ush
new file mode 100644
--- /dev/null
+++ b/Engine/Shaders/Private/ColorDifference.ush
@@ -0,0 +1,130 @@
+// Copyright Epic Games, Inc. All Rights Reserved.
+
+#pragma once
+
+#include "ColorSpace.ush"
+
+//----------------------------------------------------------------------------------------------------
+// reference: https://en.wikipedia.org/wiki/Color_difference
+// CIE76
+float DeltaE_CIE76(float3 Lab0, float3 Lab1)
+{
+	return length(Lab0 - Lab1);
+}
+
+// graphics arts
+static float CIE_k_L = 1;
+static float CIE_K_1 = 0.045;
+static float CIE_K_2 = 0.015;
+
+static float CIE_k_C = 1;
+static float CIE_k_H = 1;
+
+float DeltaE_CIE94(float3 Lab0, float3 Lab1)
+{
+	float Delta_L = Lab0.x - Lab1.x;
+	float C_1 = length(Lab0.yz);
+	float C_2 = length(Lab1.yz);
+	float Delta_C_ab = C_1 - C_2;
+
+	float2 Delta_ab = Lab0.yz - Lab1.yz;
+	float Delta_H_ab_sqr = length2(Delta_ab) - 2 * Delta_C_ab * Delta_C_ab;
+
+	float S_L = 1;
+	float S_C = 1 + CIE_K_1 * C_1;
+	float S_H = 1 + CIE_K_2 * C_2;
+
+	return sqrt(
+		Delta_L * Delta_L / ((CIE_k_L * S_L) * (CIE_k_L * S_L)) +
+		Delta_C_ab * Delta_C_ab / ((CIE_k_C * S_C) * (CIE_k_C * S_C)) +
+		Delta_H_ab_sqr / ((CIE_k_H * S_H) * (CIE_k_H * S_H))
+	);
+}
+
+float CosDeg(float Degree)
+{
+	return cos(Degree * PI / 180.0f);
+}
+
+float SinDeg(float Degree)
+{
+	return sin(Degree * PI / 180.0f);
+}
+
+// Verification is done with data in table 1
+// from "The CIEDE2000 Color-Difference
+// Formula: Implementation Notes,
+// Supplementary Test Data, and
+// Mathematical Observations"
+// http://www2.ece.rochester.edu/~gsharma/ciede2000/ciede2000noteCRNA.pdf
+float DeltaE_CIE2000(float3 Lab0, float3 Lab1)
+{
+	float3 Lch0 = LAB_2_LCH(Lab0);
+	float3 Lch1 = LAB_2_LCH(Lab1);
+
+	float Delta_L = Lab1.x - Lab0.x;
+	float L_bar = (Lab1.x + Lab0.x) * 0.5f;
+	float C_bar = (Lch0.y + Lch1.y) * 0.5f;
+
+	float component = sqrt(pow(C_bar, 7) / (pow(C_bar, 7) + pow(25, 7)));
+	float G = 0.5f * (1.0f - component);
+	float a_1 = Lab0.y * (1.0f + G);
+	float a_2 = Lab1.y * (1.0f + G);
+
+	float C = (Lch0.y + Lch1.y) * 0.5f;
+
+	float3 Lch0_ap = LAB_2_LCH(float3(Lab0.x, a_1, Lab0.z));
+	float3 Lch1_ap = LAB_2_LCH(float3(Lab1.x, a_2, Lab1.z));
+	float C_bar_ap = (Lch0_ap.y + Lch1_ap.y) * 0.5f;
+	float Delta_C_ap = Lch1_ap.y - Lch0_ap.y;
+
+	float h1_ap = Lch0_ap.z;
+	float h2_ap = Lch1_ap.z;
+
+	float Abs_h2_ap_minus_h1_ap = abs(h1_ap - h2_ap);
+	float Delta_h_ap = (Abs_h2_ap_minus_h1_ap <= 180) ? (h2_ap - h1_ap) : (
+		(Abs_h2_ap_minus_h1_ap > 180 && h2_ap <= h1_ap) ? (h2_ap - h1_ap + 360) :
+		(h2_ap - h1_ap - 360)
+		);
+	// Corner case: Set Delta_h_ap to zero if either C is zero.
+	if (Lch0_ap.y == 0 || Lch1_ap.y == 0)
+	{
+		Delta_h_ap = 0.0f;
+	}
+
+	// Delta H' and H'
+	float Delta_H_ap = 2 * sqrt(Lch0_ap.y * Lch1_ap.y) * SinDeg(Delta_h_ap * 0.5f);
+	float h1_ap_plus_h2_ap = h1_ap + h2_ap;
+	float H_bar_ap = (Abs_h2_ap_minus_h1_ap <= 180) ? (h1_ap_plus_h2_ap * 0.5f) : (
+		(Abs_h2_ap_minus_h1_ap > 180 && h1_ap_plus_h2_ap < 360) ? ((h1_ap_plus_h2_ap + 360) * 0.5f) :
+		((h1_ap_plus_h2_ap - 360) * 0.5f)
+		);
+
+	// Corner case: When either C'_1 or C'_2 is zero, H_bar' = h'_1 + h'_2;
+	if (Lch0_ap.y == 0 || Lch1_ap.y == 0)
+	{
+		H_bar_ap = h1_ap_plus_h2_ap;
+	}
+
+	// Set T
+	float T = 1 - 0.17f * CosDeg(H_bar_ap - 30.0f) + 0.24f * CosDeg(2.0f * H_bar_ap)
+		+ 0.32f * CosDeg(3.0f * H_bar_ap + 6.0f) - 0.20f * CosDeg(4.0f * H_bar_ap - 63.0f);
+
+	float L_bar_minus_50_power2 = (L_bar - 50.0f) * (L_bar - 50.0f);
+	float S_L = 1 + 0.015f * L_bar_minus_50_power2 / sqrt(20.0f + L_bar_minus_50_power2);
+	float S_C = 1 + 0.045f * C_bar_ap;
+	float S_H = 1 + 0.015f * C_bar_ap * T;
+	float component_bar = sqrt(pow(C_bar_ap, 7) / (pow(C_bar_ap, 7) + pow(25, 7)));
+	float R_T = -2.0f * component_bar * SinDeg(60.0f * exp(-pow((H_bar_ap - 275.0f) / 25.0f, 2.0f)));
+
+	float C_component = Delta_C_ap / (CIE_k_C * S_C);
+	float H_component = Delta_H_ap / (CIE_k_H * S_H);
+	float Delta_E_Star_00 = sqrt(
+		pow(Delta_L / (CIE_k_L * S_L), 2.0f) +
+		pow(C_component, 2.0f) +
+		pow(H_component, 2.0f) +
+		R_T * C_component * H_component
+	);
+
+	return Delta_E_Star_00;
+}
diff --git a/Engine/Shaders/Private/ColorSpace.ush b/Engine/Shaders/Private/ColorSpace.ush
--- a/Engine/Shaders/Private/ColorSpace.ush
+++ b/Engine/Shaders/Private/ColorSpace.ush
@@ -245,3 +245,23 @@ float3 LinearRGB_2_HSV(in float3 RGB)
 	float s = HCV.y / (HCV.z + 1e-10f);
 	return float3(HCV.x, s, HCV.z);
 }
+
+//------------------------------------------------------------------
+// LCH
+//------------------------------------------------------------------
+
+float3 LAB_2_LCH(float3 LAB)
+{
+	float3 LCH;
+	LCH.x = LAB.x;
+	LCH.y = length(LAB.yz);
+	float HInDegree0To360 = 0.0f;
+
+	if (LAB.z != 0 || LAB.y != 0)
+	{
+		HInDegree0To360 = atan2(LAB.z, LAB.y) * 180 / 3.1415926535897932f;
+		HInDegree0To360 += lerp(0, 360.0f, HInDegree0To360 < 0);
+	}
+	LCH.z = HInDegree0To360;
+	return LCH;
+}
\ No newline at end of file
diff --git a/Engine/Shaders/Private/PathTracing/PathTracingSpatialTemporalDenoising.usf b/Engine/Shaders/Private/PathTracing/PathTracingSpatialTemporalDenoising.usf
new file mode 100644
--- /dev/null
+++ b/Engine/Shaders/Private/PathTracing/PathTracingSpatialTemporalDenoising.usf
@@ -0,0 +1,1475 @@
+// Copyright Epic Games, Inc. All Rights Reserved.
+
+/*=============================================================================================
+PathTracingSpatialTemporalDenoising.usf: Spatial temporal denoising for path tracing
+===============================================================================================*/
+#pragma once
+
+#include "/Engine/Public/Platform.ush"
+#include "../ScreenPass.ush"
+#include "../../Private/Common.ush"
+#include "../TextureSampling.ush"
+#include "../ColorDifference.ush"
+
+#ifndef THREAD_SIZE
+#define THREAD_SIZE						8
+#define THREAD_SIZE_X					THREAD_SIZE
+#define THREAD_SIZE_Y					THREAD_SIZE
+#endif
+
+#ifndef TEMPORAL_REPROJECTION_ALIGN
+#define TEMPORAL_REPROJECTION_ALIGN		0
+#endif
+
+#ifndef TEMPORAL_REPROJECTION_BLUR
+#define TEMPORAL_REPROJECTION_BLUR		0
+#define DIRECT_COPY						0
+#endif
+
+#ifndef TEMPORAL_REPROJECTION_MERGE
+#define TEMPORAL_REPROJECTION_MERGE		0
+#endif
+
+#ifndef TEMPORAL_REPROJECTION_RESOLVE
+#define TEMPORAL_REPROJECTION_RESOLVE	0
+#endif
+
+#ifndef TEMPORAL_FEATURE_FUSION
+#define TEMPORAL_FEATURE_FUSION			0
+#endif
+
+#ifndef TEMPORAL_PREPASS
+#define TEMPORAL_PREPASS				0
+#endif
+
+#define PREPASS_PHASE_INIT				0
+#define PREPASS_PHASE_UPDATE			1
+
+#ifndef SPATIAL_DENOISING
+#define SPATIAL_DENOISING				0
+#endif
+
+#ifndef VISUALIZE_MOTIONVECTOR
+#define VISUALIZE_MOTIONVECTOR			0
+#endif
+
+#ifndef MOTION_VECTOR_SUBTRACT
+#define MOTION_VECTOR_SUBTRACT			0
+#endif
+
+#ifndef TEMPORAL_HIGHFREQ_REJECT
+#define TEMPORAL_HIGHFREQ_REJECT		0
+#endif
+
+#ifndef VISUALIZE_WARPING
+#define VISUALIZE_WARPING				0
+#endif 
+
+#ifndef TOTAL_VARIATION
+#define TOTAL_VARIATION					0
+#endif
+
+// Define the distance metrics to use
+#define METRICS_LUMINANCE				0
+#define METRICS_EUCLIDEAN				1
+#define METRICS_PERCEPTION				2
+
+#ifndef DISTANCE_METRICS
+#define DISTANCE_METRICS METRICS_LUMINANCE
+#endif
+
+// Gamma correction improves the reprojection of euclidean and luminance
+// based distance metrics
+#define APPLY_SIMPLIFIED_GAMMA (DISTANCE_METRICS == METRICS_EUCLIDEAN || \
+								DISTANCE_METRICS == METRICS_LUMINANCE)
+
+#ifndef K_NUM_OF_TEXTURES_PER_PASS
+#define K_NUM_OF_TEXTURES_PER_PASS		7
+#endif
+
+#ifndef K_NUM_OF_SHIFTS_PER_SLICE
+#define K_NUM_OF_SHIFTS_PER_SLICE		4
+#endif
+
+// the mip difference is 2 instead of one to make use of coarser mips
+#define MIP_DIFF_DELTA 2
+
+// Distance type (L1, L2, logL1, logL1^C)
+#define DISTANCE_L1						0
+#define DISTANCE_L2						1
+#define DISTANCE_LOG_L1					2
+#define DISTANCE_LOG_REGULARIZED_L1		3
+#define DISTANCE_TYPE					DISTANCE_LOG_REGULARIZED_L1
+
+SCREEN_PASS_TEXTURE_VIEWPORT(TargetViewport)
+
+struct FPixelOffset
+{
+	float2		xy;
+	float		Dist;
+	float		TV;
+};
+
+struct FPixelMaterialLightingFingerprint
+{
+	float4 Mean;
+	float4 Var;
+};
+
+void Max3(in float3 Value, out float MaxValue, out int MaxIndex)
+{
+	MaxIndex = (Value.x > Value.y) ? 0 : 1;
+	MaxIndex = Value[MaxIndex] > Value.z ? MaxIndex : 2;
+	MaxValue = Value[MaxIndex];
+}
+
+float2 GetPositionScaleFromMipLevel(uint FromMipLevel, uint ToMipLevel)
+{
+	return (float2((uint2)TargetViewport_Extent >> FromMipLevel)) * (float2((uint2)TargetViewport_Extent >> ToMipLevel));
+}
+
+#if K_NUM_OF_TEXTURES_PER_PASS <= 7 && (TEMPORAL_REPROJECTION_ALIGN||TEMPORAL_REPROJECTION_MERGE)
+
+static int2 Shifts[25] = {
+	{-2,-2},{-1,-2},{0,-2},{1,-2},{2,-2},
+	{-2,-1},{-1,-1},{0,-1},{1,-1},{2,-1},
+	{-2, 0},{-1, 0},{0, 0},{1, 0},{2, 0},
+	{-2, 1},{-1, 1},{0, 1},{1, 1},{2, 1},
+	{-2, 2},{-1, 2},{0, 2},{1, 2},{2, 2},
+};
+
+RWTexture2D<float4> RWDistanceTextures_0;
+RWTexture2D<float4> RWDistanceTextures_1;
+RWTexture2D<float4> RWDistanceTextures_2;
+RWTexture2D<float4> RWDistanceTextures_3;
+RWTexture2D<float4> RWDistanceTextures_4;
+RWTexture2D<float4> RWDistanceTextures_5;
+RWTexture2D<float4> RWDistanceTextures_6;
+
+Texture2D<float4> DistanceTextures_0;
+Texture2D<float4> DistanceTextures_1;
+Texture2D<float4> DistanceTextures_2;
+Texture2D<float4> DistanceTextures_3;
+Texture2D<float4> DistanceTextures_4;
+Texture2D<float4> DistanceTextures_5;
+Texture2D<float4> DistanceTextures_6;
+
+
+void SaveOffsets(int2 Position, int BasePixelShift, float4 Value)
+{
+
+	if (BasePixelShift < 4)
+	{
+		RWDistanceTextures_0[Position] = Value;
+	}
+	else if (BasePixelShift < 8)
+	{
+		RWDistanceTextures_1[Position] = Value;
+	}
+	else if (BasePixelShift < 12)
+	{
+		RWDistanceTextures_2[Position] = Value;
+	}
+	else if (BasePixelShift < 16)
+	{
+		RWDistanceTextures_3[Position] = Value;
+	}
+	else if (BasePixelShift < 20)
+	{
+		RWDistanceTextures_4[Position] = Value;
+	}
+	else if (BasePixelShift < 24)
+	{
+		RWDistanceTextures_5[Position] = Value;
+	}
+	else if (BasePixelShift < 28)
+	{
+		RWDistanceTextures_6[Position] = Value;
+	}
+}
+
+float GetOffsets(int2 Position, int2 Offset)
+{
+	Offset = clamp(Offset, int2(-2, -2), int2(2, 2));
+	uint Index = Offset.x + Offset.y * 5 + 12;
+	uint TextureIndexOffset = Index % 4;
+	int2 ShiftedPosition = Position + Offset;
+	//@TODO: range check
+	float Ret = 0.0f;
+	if (Index < 4)
+	{
+		Ret = DistanceTextures_0[ShiftedPosition][TextureIndexOffset];
+	}
+	else if (Index < 8)
+	{
+		Ret = DistanceTextures_1[ShiftedPosition][TextureIndexOffset];
+	}
+	else if (Index < 12)
+	{
+		Ret = DistanceTextures_2[ShiftedPosition][TextureIndexOffset];
+	}
+	else if (Index < 16)
+	{
+		Ret = DistanceTextures_3[ShiftedPosition][TextureIndexOffset];
+	}
+	else if (Index < 20)
+	{
+		Ret = DistanceTextures_4[ShiftedPosition][TextureIndexOffset];
+	}
+	else if (Index < 24)
+	{
+		Ret = DistanceTextures_5[ShiftedPosition][TextureIndexOffset];
+	}
+	else if (Index < 28)
+	{
+		Ret = DistanceTextures_6[ShiftedPosition][TextureIndexOffset];
+	}
+	return Ret;
+}
+
+struct FQuadricFittingContext
+{
+	int2			Offset;	// integer shift of the minimal value
+	float3x3	   Values;  // holding the surrounding distance to fit a bivariate polynomial.
+};
+
+float Min3(float3 value)
+{
+	return min(min(value.x, value.y), value.z);
+}
+
+float Min4(float4 value)
+{
+	return min(min(value.x, value.y), min(value.z, value.w));
+}
+
+struct FFindMinContext
+{
+	int TextureOffset;
+	int Index;
+	float Value;
+};
+
+void Min3(in float3 Value, out float MinValue, out int MinIndex)
+{
+	MinIndex = (Value.x <= Value.y) ? 0 : 1;
+	MinIndex = Value[MinIndex] <= Value.z ? MinIndex : 2;
+	MinValue = Value[MinIndex];
+}
+
+void Min4(in float4 Value, out float MinValue, out int MinIndex)
+{
+	int MinIndexXy = (Value.x <= Value.y) ? 0 : 1;
+	int MinIndexZw = (Value.z <= Value.w) ? 2 : 3;
+	MinIndex = (Value[MinIndexXy] <= Value[MinIndexZw]) ? MinIndexXy : MinIndexZw;
+	MinValue = Value[MinIndex];
+}
+
+void Max4(in float4 Value, out float MaxValue, out int MaxIndex)
+{
+	int MaxIndexXy = (Value.x < Value.y) ? 1 : 0;
+	int MaxIndexZw = (Value.z < Value.w) ? 3 : 2;
+	MaxIndex = (Value[MaxIndexXy] < Value[MaxIndexZw]) ? MaxIndexZw : MaxIndexXy;
+	MaxValue = Value[MaxIndex];
+}
+
+void UpdateMin4Index(int TextureOffset, float4 Value, inout FFindMinContext Context)
+{
+	int MinIndex;
+	float MinValue;
+	Min4(Value, MinValue, MinIndex);
+
+	if (MinValue < Context.Value)
+	{
+		Context.TextureOffset = TextureOffset;
+		Context.Index = MinIndex;
+		Context.Value = MinValue;
+	}
+}
+
+#define ESTIMATE_SUBPIXELOFFSET SUBPIXEL_OFFSET
+
+#if ESTIMATE_SUBPIXELOFFSET
+
+// Ref: Section 4 of the Appendix of "Burst photography for high dynamic range and low-light imaging on mobile cameras"
+// https://dl.acm.org/doi/10.1145/2980179.2980254, 2016
+// solve a weighted least-squares problem with weight of
+// 1 2 1
+// 2 4 2
+// 1 2 1
+
+static const float FA11[9] = {	1.0f, -2.0f,  1.0f,
+								2.0f, -4.0f,  2.0f,
+								1.0f, -2.0f,  1.0f};
+
+static const float FA22[9] = {	1.0f / 4.0f,  2.0f / 4.0f,  1.0f / 4.0f,
+							   -2.0f / 4.0f, -4.0f / 4.0f, -2.0f / 4.0f,
+								1.0f / 4.0f,  2.0f / 4.0f,  1.0f / 4.0f };
+
+
+static const float FA12[9] = {	1.0f / 4.0f,  0.0f, -1.0f / 4.0f,
+									   0.0f,  0.0f,         0.0f,
+							   -1.0f / 4.0f,  0.0f,  1.0f / 4.0f }; // FA12 == FA21
+
+static const float  FB1[9] = {  -1.0f / 8.0f,  0.0f,  1.0f / 8.0f,
+							   -2.0f / 8.0f,  0.0f,  2.0f / 8.0f,
+							   -1.0f / 8.0f,  0.0f,  1.0f / 8.0f };
+
+static const float  FB2[9] = {  -1.0f / 8.0f, -2.0f / 8.0f, -1.0f / 8.0f,
+							    0.0f,  0.0f,					   0.0f,
+							    1.0f / 8.0f,  2.0f / 8.0f,  1.0f / 8.0f };
+
+//static const float  FC[] = { -1.0f,  2.0f, -1.0f,
+//								2.0f, 12.0f,  2.0f,
+//							   -1.0f,  2.0f, -1.0f } / 16;
+
+float Dot9(float A[9], float B[9])
+{
+	return	A[0] * B[0] + A[1] * B[1] + A[2] * B[2] +
+			A[3] * B[3] + A[4] * B[4] + A[5] * B[5] +
+			A[6] * B[6] + A[7] * B[7] + A[8] * B[8];
+}
+
+// D holds the 3x3 image distance. Return the subpixel offset 
+// in range [-1, 1] x [-1, 1]
+float2 EstimateSubpixelOffset(float3x3 D)
+{
+	// Step 1. Construct the constant coefficient A, b, c
+	// for the bivariate quadratic function.
+
+	float DSub[9];
+	for (int j = -1; j <= 1; j++)
+		for (int i = -1; i <= 1; i++)
+			DSub[3 * (j + 1) + i + 1] = D[j + 1][i + 1];
+
+	// force A positive semi-definite (PSD) with max(0, Aij) and reset non-diag
+	// to zero if det is negative.
+	// note: A12 == A21
+	float4 A = { max(0.0f, Dot9(FA11, DSub)),			   Dot9(FA12, DSub),
+				 0.0f						,	max(0.0f, Dot9(FA22, DSub))};
+
+	float DetA = A[0] * A[3] - A[1] * A[1];
+	
+	if (DetA < 0.0f)
+	{
+		A[1] = 0;
+		DetA = A[0] * A[3];
+	}
+
+	float2 b = { Dot9(FB1, DSub), 
+				 Dot9(FB2, DSub)};
+
+	//float c = Dot9(FC, DSub);
+
+	// Get the sub-pixel offset
+	// \mu = -A^{-1}b, 
+	// the constant parameter s = c - \frac{\mu^{T}A\mu}{2} is ignored
+
+	float2 SubpixelOffset = 0.0f;
+
+	if (abs(DetA) > 1e-9)
+	{
+		float2 Mu = { -(A[3] * b[0] - A[1] * b[1]) / DetA,
+					  -(A[0] * b[1] - A[1] * b[0]) / DetA };
+
+		SubpixelOffset = lerp(0, Mu, length2(Mu) < 1);
+	}
+	
+	return SubpixelOffset;
+}
+#endif
+
+#define MAX_HEAP_SIZE 4
+
+#if MAX_HEAP_SIZE != 4
+	#error Support up to max heap size of 4 only
+#endif
+
+struct FMaxHeap
+{
+	float Heap[MAX_HEAP_SIZE];
+	uint Index[MAX_HEAP_SIZE];
+};
+
+FMaxHeap CreateMaxHeap()
+{
+	FMaxHeap MaxHeap;
+	
+	UNROLL
+	for (int i = 0; i < MAX_HEAP_SIZE; ++i)
+	{
+		MaxHeap.Heap[i] = MaxHalfFloat;
+		MaxHeap.Index[i] = 0;
+	}
+
+	return MaxHeap;
+}
+
+float GetMax(in float Array[MAX_HEAP_SIZE])
+{
+	return Array[0];
+}
+
+void _Swap(inout float Array[MAX_HEAP_SIZE], uint i, uint j)
+{
+	float TmpValue = Array[i];
+	Array[i] = Array[j];
+	Array[j] = TmpValue;
+}
+
+void _Swap(inout uint Array[MAX_HEAP_SIZE], uint i, uint j)
+{
+	uint TmpValue = Array[i];
+	Array[i] = Array[j];
+	Array[j] = TmpValue;
+}
+
+void Insert(inout FMaxHeap MaxHeap, float Value, int Index)
+{
+	// Insert only when the value is smaller than the max value of the max heap
+	if (Value < GetMax(MaxHeap.Heap))
+	{
+		// 1. replace the value at the head
+		MaxHeap.Heap[0] = Value;
+		MaxHeap.Index[0] = Index;
+		
+		// 2. Find max and swap to the head
+		float MaxValue;
+		int MaxIndex;
+		Max4(float4(MaxHeap.Heap[0], MaxHeap.Heap[1], MaxHeap.Heap[2], MaxHeap.Heap[3]), MaxValue, MaxIndex);
+
+		if (MaxIndex != 0)
+		{
+			_Swap(MaxHeap.Heap, 0, MaxIndex);
+			_Swap(MaxHeap.Index, 0, MaxIndex);
+		}
+	}
+}
+
+void Insert(inout FMaxHeap MaxHeap, float4 Values, int4 Indices)
+{
+	UNROLL
+	for (int i = 0; i < 4; ++i)
+	{
+		Insert(MaxHeap, Values[i], Indices[i]);
+	}
+}
+
+void Sort4(float Values[MAX_HEAP_SIZE], inout uint4 Indices)
+{
+	uint Index[MAX_HEAP_SIZE] = {0, 1, 2, 3};
+	// Optimal number of comparison is 5
+	if (Values[Index[0]] < Values[Index[1]])
+	{
+		_Swap(Index, 0, 1);
+	}
+
+	if (Values[Index[2]] < Values[Index[3]])
+	{
+		_Swap(Index, 2, 3);
+	}
+
+	if (Values[Index[0]] < Values[Index[2]])
+	{
+		_Swap(Index, 0, 2);
+	}
+
+	if (Values[Index[1]] < Values[Index[3]])
+	{
+		_Swap(Index, 1, 3);
+	}
+
+	if (Values[Index[1]] < Values[Index[2]])
+	{
+		_Swap(Index, 1, 2);
+	}
+	Indices = uint4(Index[0], Index[1], Index[2], Index[3]);
+}
+
+// based on the target position, and the shift from the last mip, we fetch
+FQuadricFittingContext GetQuadricFittingContext(int2 Position, int kth = 0)
+{
+	kth = clamp(kth, 0, MAX_HEAP_SIZE - 1);
+
+	// get the minimal distance and the surrounding
+	FFindMinContext Context = (FFindMinContext)0;
+	
+	BRANCH
+	if (kth == 0)
+	{
+		Context.TextureOffset = 3;// starting to have (0,0) as the minimal value
+
+		Min4(DistanceTextures_3[Position], Context.Value, Context.Index);
+
+		UpdateMin4Index(1, DistanceTextures_1[Position], Context);
+		UpdateMin4Index(2, DistanceTextures_2[Position], Context);
+		UpdateMin4Index(4, DistanceTextures_4[Position], Context);
+		UpdateMin4Index(0, DistanceTextures_0[Position], Context);
+		UpdateMin4Index(5, DistanceTextures_5[Position], Context);
+
+		int MinIndex = 0;
+		float MinValue = DistanceTextures_6[Position].x;
+		if (MinValue < Context.Value)
+		{
+			Context.TextureOffset = 6;
+			Context.Value = MinValue;
+			Context.Index = MinIndex;
+		}
+	}
+	else
+	{
+		//TODO: assign the first distance texture values
+		FMaxHeap MaxHeap = CreateMaxHeap();
+
+		Insert(MaxHeap, DistanceTextures_3[Position], 3 * 4 + int4(0, 1, 2, 3));
+		Insert(MaxHeap, DistanceTextures_1[Position], 1 * 4 + int4(0, 1, 2, 3));
+		Insert(MaxHeap, DistanceTextures_2[Position], 2 * 4 + int4(0, 1, 2, 3));
+		Insert(MaxHeap, DistanceTextures_4[Position], 4 * 4 + int4(0, 1, 2, 3));
+		Insert(MaxHeap, DistanceTextures_0[Position], 0 * 4 + int4(0, 1, 2, 3));
+		Insert(MaxHeap, DistanceTextures_5[Position], 5 * 4 + int4(0, 1, 2, 3));
+		Insert(MaxHeap, DistanceTextures_6[Position].x, 6 * 4);
+
+		uint4 Max4Index;
+		Sort4(MaxHeap.Heap, Max4Index);
+
+		Context.TextureOffset = Max4Index[kth] / 4;
+		Context.Index = Max4Index[kth] % 4;
+		Context.Value = GetMax(MaxHeap.Heap);
+
+	}
+	
+	FQuadricFittingContext QuadContext;
+	QuadContext.Offset = Shifts[Context.TextureOffset * 4 + Context.Index];
+
+
+#if !ESTIMATE_SUBPIXELOFFSET
+	QuadContext.Values = Context.Value;
+#else
+	int2 ShiftedPosition = Position + QuadContext.Offset;
+	QuadContext.Values[0][0] = GetOffsets(ShiftedPosition, int2(-1, -1));
+	QuadContext.Values[0][1] = GetOffsets(ShiftedPosition, int2( 0, -1));
+	QuadContext.Values[0][2] = GetOffsets(ShiftedPosition, int2( 1, -1));
+	QuadContext.Values[1][0] = GetOffsets(ShiftedPosition, int2(-1,  0));
+	QuadContext.Values[1][1] = Context.Value;
+	QuadContext.Values[1][2] = GetOffsets(ShiftedPosition, int2( 1,  0));
+	QuadContext.Values[2][0] = GetOffsets(ShiftedPosition, int2(-1,  1));
+	QuadContext.Values[2][1] = GetOffsets(ShiftedPosition, int2( 0,  1));
+	QuadContext.Values[2][2] = GetOffsets(ShiftedPosition, int2( 1,  1));
+#endif
+	return QuadContext;
+}
+
+float GetTotalVariation(int2 Position, FQuadricFittingContext Context)
+{
+	int2 ShiftedPosition = Position + Context.Offset;
+	float TotalVariation = 0.0f;
+	for (int i = -2; i < 2; ++i)
+	{
+		for (int j = -2; j < 2; ++j)
+		{
+			float vij	= GetOffsets(ShiftedPosition, int2(		i,		j));
+			float vi_1j = GetOffsets(ShiftedPosition, int2(	i + 1,		j));
+			float vij_1 = GetOffsets(ShiftedPosition, int2(		i,	j + 1));
+
+			TotalVariation += abs(vi_1j - vij) + abs(vij_1 - vij);
+		}
+	}
+
+	for (int j = -2; j < 2; ++j)
+	{
+		float vij = GetOffsets(ShiftedPosition,		int2(2,		j));
+		float vij_1 = GetOffsets(ShiftedPosition,	int2(2, j + 1));
+		TotalVariation += abs(vij - vij_1);
+	}
+
+	for (int i = -2; i < 1; ++i)
+	{
+		float vij	= GetOffsets(ShiftedPosition, int2(		i,	2));
+		float vi_1j = GetOffsets(ShiftedPosition, int2(	i + 1,	2));
+		TotalVariation += abs(vij - vi_1j);
+	}
+	return TotalVariation;
+}
+
+FPixelOffset GetQuadricFittingOffset(FQuadricFittingContext Context)
+{
+	FPixelOffset Offset;
+
+#if ESTIMATE_SUBPIXELOFFSET
+	float2 SubpixelOffset = EstimateSubpixelOffset(Context.Values);
+	Offset.xy = Context.Offset + SubpixelOffset;
+#else
+	Offset.xy = Context.Offset;
+#endif
+	
+	Offset.Dist = Context.Values[1][1];
+	return Offset;
+}
+
+#endif
+
+int2 GetShiftedPosition(int2 Position, int2 Shift)
+{
+	return Position + Shift;
+}
+
+float2 GetBufferUVFromPosition(float2 Position, int MipLevel)
+{
+	return (Position + 0.5f) / ((uint2)TargetViewport_Extent >> MipLevel);
+}
+
+// Error Normalization.
+// to have a similar error level between euclidean e.g. for color
+// C0 = (0.58f, 0.46f, 0.28f), C1 = (0.28f, 0.96f, 0.43f)
+// the color difference of euclidean with gamma corrected D_eu_g
+// D_eu_g(C0, C1) = 0.38f
+// the color difference of DeltaE_CIE2000 has
+// DeltaE(C0,C1)  = 27.71
+// To have the error on the same level, we multiply by 0.53f/ 37.85f = 1/70.78f, which is an experimental value
+// of running 1e6 random pair of colors to get the mean of different measures where
+// Mean_eu_g = 0.53f, Mean_DeltaE = 37.85f;
+static float NormalizeDeltaEAndEuclideanDistanceToSameMean = 0.0141f;
+#define NormalizeDeltaE2000(X) (NormalizeDeltaEAndEuclideanDistanceToSameMean*X)
+
+// rgb is in linear space
+float GetDeltaE(float3 LinearRGB0, float3 LinearRGB1, bool bNormalizeToEuclideanMean = true)
+{
+
+	float3 Lab0 = LinearRGB_2_LAB(LinearRGB0, XYZ_WHITE_REF_D65_NOMINAL);
+	float3 Lab1 = LinearRGB_2_LAB(LinearRGB1, XYZ_WHITE_REF_D65_NOMINAL);
+
+	float DeltaE = DeltaE_CIE2000(Lab0, Lab1);
+
+	if (bNormalizeToEuclideanMean)
+	{
+		DeltaE = NormalizeDeltaE2000(DeltaE);
+	}
+
+	return DeltaE;
+}
+
+float GetDistance(float LumA, float LumB, int2 shift)
+{
+	float distance = 0;
+
+#if DISTANCE_TYPE == DISTANCE_L1
+	return abs(LumA - LumB);
+#elif DISTANCE_TYPE == DISTANCE_L2
+	return Square(LumA - LumB);
+#elif DISTANCE_TYPE == DISTANCE_LOG_L1
+	return log2(2 + abs(LumA - LumB));
+#elif DISTANCE_TYPE == DISTANCE_LOG_REGULARIZED_L1
+	return log2(2 + abs(LumA - LumB)) * (2 - exp(-0.02 * length2(shift)));
+#else
+#error Unsupported dinstance type between two pixels, please implement.
+	return 0.0f;
+#endif
+}
+
+float GetFinalMipDistance(float LumA, float LumB, int2 shift)
+{
+	return log2(2 + abs(LumA - LumB)) * (2 - exp(-0.02 * length2(shift)));
+}
+
+float GetDistance(float3 A, float3 B, int2 shift)
+{
+
+#if APPLY_SIMPLIFIED_GAMMA
+	A = pow(abs(A), 1 / 2.2f);
+	B = pow(abs(B), 1 / 2.2f);
+#endif
+
+#if DISTANCE_METRICS == METRICS_LUMINANCE
+	float LumA = Luminance(A);
+	float LumB = Luminance(B);
+#elif DISTANCE_METRICS == METRICS_EUCLIDEAN
+	float LumA = length(A - B);
+	float LumB = 0;
+#elif DISTANCE_METRICS == METRICS_PERCEPTION
+	float LumA = GetDeltaE(A, B);
+	float LumB = 0;
+#endif
+	return GetDistance(LumA, LumB, shift);
+}
+
+float GetFinalMipDistance(float3 A, float3 B, int2 shift)
+{
+
+#if APPLY_SIMPLIFIED_GAMMA
+	A = pow(abs(A), 1 / 2.2f);
+	B = pow(abs(B), 1 / 2.2f);
+#endif
+
+#if DISTANCE_METRICS == METRICS_LUMINANCE
+	float LumA = Luminance(A);
+	float LumB = Luminance(B);
+#elif DISTANCE_METRICS == METRICS_EUCLIDEAN
+	float LumA = length(A - B);
+	float LumB = 0;
+#elif DISTANCE_METRICS == METRICS_PERCEPTION
+	float LumA = GetDeltaE(A, B);
+	float LumB = 0;
+#endif
+	return GetFinalMipDistance(LumA, LumB, shift);
+}
+
+#if TEMPORAL_REPROJECTION_ALIGN || TEMPORAL_REPROJECTION_MERGE || TEMPORAL_REPROJECTION_RESOLVE || TEMPORAL_HIGHFREQ_REJECT
+
+Texture2D<float4> PixelOffsetTexture;
+Texture2D<float4> SourceTexture; 
+Texture2D<float4> TargetTexture; 
+SamplerState SharedTextureSampler;
+uint MipLevel;
+
+FPixelOffset GetPixelOffsetFromLowerMip(float2 Position, uint MipDiff)
+{
+	uint MipPosMultipler = 1 << (MipDiff);
+	uint LowerMipLevel = MipLevel + MipDiff;
+
+	float2 LowerMipPosition = 0.0f;
+
+	BRANCH
+	if (MipDiff == 0)
+	{
+		LowerMipPosition = Position;
+	}
+	else
+	{
+		// Keep the reprojection border of the coarse level close to zero
+		LowerMipPosition = (Position - 1.0) / (float)MipPosMultipler - 0.5f;
+	}
+
+	FPixelOffset Offset = (FPixelOffset)0;
+
+	float2 Extent = (float2)((uint2)TargetViewport_Extent >> LowerMipLevel);
+	float2 UV = (LowerMipPosition + 0.5f) / Extent;
+	float4 OffsetDist = Texture2DSampleBicubic(PixelOffsetTexture, SharedTextureSampler, UV, Extent, 1.0f / Extent);
+
+	Offset.xy = MipPosMultipler * OffsetDist.xy;
+	Offset.Dist = OffsetDist.z;
+#if TOTAL_VARIATION || TEMPORAL_REPROJECTION_RESOLVE || TEMPORAL_HIGHFREQ_REJECT
+	Offset.TV = OffsetDist.w;
+#else
+	Offset.TV = 1.0f;
+#endif
+	return Offset;
+}
+
+float3 GetTextureValue(Texture2D Tex, float2 Position)
+{
+	float2 BufferUV = GetBufferUVFromPosition(Position, MipLevel);
+	return Texture2DSampleLevel(Tex, SharedTextureSampler, BufferUV, 0).rgb;
+}
+
+float3 GetTextureValueWithOffset(Texture2D Tex, float2 Position)
+{
+	// Get the postion offset from the previous miplevel texture
+	FPixelOffset Offset = GetPixelOffsetFromLowerMip(Position, MIP_DIFF_DELTA);
+
+	float2 BufferUV = GetBufferUVFromPosition(Position + Offset.xy, MipLevel);
+	return Texture2DSampleLevel(Tex, SharedTextureSampler, BufferUV, 0).rgb;
+}
+
+#endif
+
+#if TEMPORAL_PREPASS
+
+Texture2D<float4> InputTexture;
+RWStructuredBuffer<FPixelMaterialLightingFingerprint> RWVarianceMap;
+int Iteration;
+
+[numthreads(THREAD_SIZE_X, THREAD_SIZE_Y, 1)]
+void TemporalPrepassCS(uint3 DT_ID : SV_DispatchThreadID)
+{
+	int2 Position = (int2)DT_ID.xy;
+	int2 MaxPosition = int2(TargetViewport_Extent);
+	if (all(DT_ID.xy < MaxPosition))
+	{
+		int index = Position.x + Position.y * TargetViewport_Extent.x;
+		float4 Radiance = InputTexture[Position];
+
+#if PREPASS_PHASE == PREPASS_PHASE_INIT
+
+		FPixelMaterialLightingFingerprint VarianceInfo = (FPixelMaterialLightingFingerprint)0;
+		VarianceInfo.Mean = Radiance;
+		RWVarianceMap[index] = VarianceInfo;
+
+#elif PREPASS_PHASE == PREPASS_PHASE_UPDATE
+
+		FPixelMaterialLightingFingerprint VarianceInfo = RWVarianceMap[index];
+		
+		float4 PreviousMean = VarianceInfo.Mean;
+		float4 CurrentValue = (Radiance - PreviousMean) * Iteration + Radiance;
+		
+		VarianceInfo.Mean = Radiance; 
+		VarianceInfo.Mean.w = Iteration;
+		VarianceInfo.Var += (CurrentValue - VarianceInfo.Mean)*(CurrentValue - PreviousMean);
+
+		RWVarianceMap[index] = VarianceInfo;
+#else
+#endif
+	}
+}
+#endif // TEMPORAL_PREPASS
+
+#if TEMPORAL_FEATURE_FUSION
+//Use Albedo, Normal, and Radiance Texture to create another feature vector
+// such that the feature vector leads to the best motion vector estimation.
+
+Texture2D<float4> AlbedoTexture_0;
+Texture2D<float4> AlbedoTexture_1;
+
+Texture2D<float4> NormalTexture_0;
+Texture2D<float4> NormalTexture_1;
+
+Texture2D<float4> RadianceTexture_0;
+Texture2D<float4> RadianceTexture_1;
+
+StructuredBuffer<FPixelMaterialLightingFingerprint> VarianceMap_0;
+StructuredBuffer<FPixelMaterialLightingFingerprint> VarianceMap_1;
+
+
+Texture2D<float4> LastDenoisedRadiance;
+
+SamplerState SharedTextureSampler;
+
+RWTexture2D<float4> OutputTexture_0;
+RWTexture2D<float4> OutputTexture_1;
+
+#define SOURCE(x) x##_0[Position]
+#define TARGET(x) x##_1[Position]
+
+#define SOURCE_BUFFER(B) B##_0[Position.x+int(Position.y*TargetViewport_Extent.x)]
+#define TARGET_BUFFER(B) B##_1[Position.x+int(Position.y*TargetViewport_Extent.x)]
+
+float4 Fuse(float4 Albedo, float4 Normal, float4 Radiance, float4 History, FPixelMaterialLightingFingerprint Fingerprint)
+{
+	float4 combined = 0;
+	//TODO: Derive good feature fusion algorithm.
+	combined.rgb = sqrt(Fingerprint.Var.rgb / (Fingerprint.Mean.w + 1));
+
+	return combined;
+}
+
+float4 FuseSource(int2 Position)
+{
+	float4 Albedo = SOURCE(AlbedoTexture);
+	float4 Normal = SOURCE(NormalTexture);
+	float4 Radiance = SOURCE(RadianceTexture);
+	float4 History = LastDenoisedRadiance[Position];
+	FPixelMaterialLightingFingerprint Fingerprint = SOURCE_BUFFER(VarianceMap);
+
+	return Fuse(Albedo, Normal, Radiance, History,Fingerprint);
+}
+
+float4 FuseTarget(int2 Position)
+{
+	float4 Albedo = TARGET(AlbedoTexture);
+	float4 Normal = TARGET(NormalTexture);
+	float4 Radiance = TARGET(RadianceTexture);
+	float4 History = -1.0f;
+	FPixelMaterialLightingFingerprint Fingerprint = TARGET_BUFFER(VarianceMap);
+
+	return Fuse(Albedo, Normal, Radiance, History, Fingerprint);
+}
+
+[numthreads(THREAD_SIZE_X, THREAD_SIZE_Y, 1)]
+void TemporalFeatureFusionCS(uint3 DT_ID : SV_DispatchThreadID)
+{
+	int2 Position = DT_ID.xy;
+	int2 MaxPosition = (int2)TargetViewport_Extent;
+	if (all(Position < MaxPosition))
+	{
+		SOURCE(OutputTexture) = FuseSource(Position);
+		TARGET(OutputTexture) = FuseTarget(Position);
+	}
+}
+
+#endif // TEMPORAL_FEATURE_FUSION
+
+#if TEMPORAL_REPROJECTION_ALIGN
+
+[numthreads(THREAD_SIZE_X, THREAD_SIZE_Y, 1)]
+void ReprojectionAlignCS(uint3 DT_ID : SV_DispatchThreadID)
+{
+	int2 Position = DT_ID.xy;
+	int SliceId = DT_ID.z;
+	uint3 MaxPosition = uint3(((uint2)TargetViewport_Extent.xy) >> MipLevel, K_NUM_OF_TEXTURES_PER_PASS);
+
+	if (all(DT_ID.xyz < MaxPosition))
+	{
+		float3 Target = GetTextureValue(TargetTexture, Position)*View.PreExposure;
+		float4 Distances = 0.0f;
+
+		UNROLL
+			for (int i = 0; i < K_NUM_OF_SHIFTS_PER_SLICE; ++i)
+			{
+				int2 Shift = Shifts[SliceId * K_NUM_OF_SHIFTS_PER_SLICE + i];
+				int2 ShiftPosition = GetShiftedPosition(Position, Shift);
+				float3 Source = GetTextureValueWithOffset(SourceTexture, ShiftPosition)*View.PreExposure;
+
+				if (MipLevel != 0)
+				{
+					Distances[i] = GetDistance(Target, Source, Shift);
+				}
+				else
+				{
+					Distances[i] = GetFinalMipDistance(Target, Source, Shift);
+				}
+			}
+
+		SaveOffsets(Position, SliceId * K_NUM_OF_SHIFTS_PER_SLICE, Distances);
+	}
+}
+#endif // TEMPORAL_REPROJECTION_ALIGN
+
+#if TEMPORAL_REPROJECTION_BLUR
+
+Texture2D<float4> InputTexture; // Previous accumulation
+RWTexture2D<float4> OutputTexture; // new observation
+SamplerState SharedTextureSampler;
+uint MipLevel;
+
+float4 Get3x3Blur(int2 Position)
+{
+	float2 BufferUV = GetBufferUVFromPosition(Position, MipLevel);
+	float2 HalfPixel = 0.5f * (1 << MipLevel) * TargetViewport_ExtentInverse;
+
+#if 0
+	//Sampling pattern and weights
+	// 1   1
+	//   4
+	// 1   1
+	float4 Sum = Texture2DSampleLevel(InputTexture, SharedTextureSampler, BufferUV, 0) * 4;
+	Sum += Texture2DSampleLevel(InputTexture, SharedTextureSampler,
+		(BufferUV + float2(-HalfPixel.x, -HalfPixel.y)), 0);
+	Sum += Texture2DSampleLevel(InputTexture, SharedTextureSampler,
+		(BufferUV + float2(HalfPixel.x, -HalfPixel.y)), 0);
+	Sum += Texture2DSampleLevel(InputTexture, SharedTextureSampler,
+		(BufferUV + float2(-HalfPixel.x, HalfPixel.y)), 0);
+	Sum += Texture2DSampleLevel(InputTexture, SharedTextureSampler,
+		(BufferUV + float2(HalfPixel.x, HalfPixel.y)), 0);
+
+	return Sum / 8.0f;
+
+#elif DIRECT_COPY
+
+	return InputTexture[Position];
+
+#else
+
+	float4 Sum = 0;
+	Sum += Texture2DSampleLevel(InputTexture, SharedTextureSampler,
+		(BufferUV + 2.0 * float2(-HalfPixel.x, -HalfPixel.y)), 0);
+	Sum += Texture2DSampleLevel(InputTexture, SharedTextureSampler,
+		(BufferUV + 2.0 * float2(-HalfPixel.x, 0)), 0);
+	Sum += Texture2DSampleLevel(InputTexture, SharedTextureSampler,
+		(BufferUV + 2.0 * float2(-HalfPixel.x, HalfPixel.y)), 0);
+	Sum += Texture2DSampleLevel(InputTexture, SharedTextureSampler,
+		(BufferUV + 2.0 * float2(0, -HalfPixel.y)), 0);
+	Sum += Texture2DSampleLevel(InputTexture, SharedTextureSampler,
+		(BufferUV + 2.0 * float2(0, 0)), 0);
+	Sum += Texture2DSampleLevel(InputTexture, SharedTextureSampler,
+		(BufferUV + 2.0 * float2(0, HalfPixel.y)), 0);
+	Sum += Texture2DSampleLevel(InputTexture, SharedTextureSampler,
+		(BufferUV + 2.0 * float2(HalfPixel.x, -HalfPixel.y)), 0);
+	Sum += Texture2DSampleLevel(InputTexture, SharedTextureSampler,
+		(BufferUV + 2.0 * float2(HalfPixel.x, 0)), 0);
+	Sum += Texture2DSampleLevel(InputTexture, SharedTextureSampler,
+		(BufferUV + 2.0 * float2(HalfPixel.x, HalfPixel.y)), 0);
+
+	return Sum / 9.0f;
+#endif
+	
+}
+
+[numthreads(THREAD_SIZE_X, THREAD_SIZE_Y, 1)]
+void ReprojectionBlurCS(uint3 DT_ID : SV_DispatchThreadID)
+{
+
+	int2 Position = (int2)DT_ID.xy;
+	int2 MaxPosition = int2(((uint2)TargetViewport_Extent) >> MipLevel);
+	if (all(Position.xy < MaxPosition))
+	{
+		float4 Blur = Get3x3Blur(Position);
+		OutputTexture[Position] = Blur;
+	}
+}
+#endif // TEMPORAL_REPROJECTION_ALIGN
+
+#if TEMPORAL_REPROJECTION_MERGE
+
+RWTexture2D<float4> RWPixelOffsetTexture;
+uint PatchId;
+uint PatchCount;
+
+// For a three level mip chain and a delta of MIP_DIFF_DELTA.
+// fetch the patch. E.g., 
+// 0-0-0 fetch the patch with the minimal distance
+// 0-1-0 fetch the second best for the second mip level
+
+uint GetTheKthPatchBasedOnMipAndPatchId()
+{
+	uint Mip2PatchIndex = PatchId % 4;
+	uint Mip0PatchIndex = PatchId / 4;
+	const uint TopMipLevel = 0;
+	uint TheKthPatchForMipLevel = lerp(
+		lerp(0,																// All others
+			Mip2PatchIndex,	MipLevel == (TopMipLevel + MIP_DIFF_DELTA)),	// 2nd mip level
+			Mip0PatchIndex,	MipLevel == TopMipLevel					  );	// Most fine grain level
+
+	return TheKthPatchForMipLevel;
+}
+
+[numthreads(THREAD_SIZE_X, THREAD_SIZE_Y, 1)]
+void ReprojectionMergeCS(uint3 DT_ID : SV_DispatchThreadID)
+{
+	// Keep the debug here as the reproject still needs some work around borders for better  history reuse.
+	bool DebugMotionVector = false;
+
+	int2 Position = DT_ID.xy;
+	int2 MaxPosition = int2(((uint2)TargetViewport_Extent) >> MipLevel);
+	if (all(Position < MaxPosition))
+	{
+		uint Kth = GetTheKthPatchBasedOnMipAndPatchId();
+
+		FQuadricFittingContext Context = GetQuadricFittingContext(Position, Kth);
+
+		// Add up the old offset with the new one.
+		FPixelOffset OffsetFromLowerMip = GetPixelOffsetFromLowerMip(Position, MIP_DIFF_DELTA);
+
+		// the offset is from target to source
+		FPixelOffset Offset = GetQuadricFittingOffset(Context);
+
+		float4 PixelDiff = float4(Offset.xy + OffsetFromLowerMip.xy, Offset.Dist, length(Offset.xy + OffsetFromLowerMip.xy));
+
+#if TOTAL_VARIATION
+		float TotalVariation = GetTotalVariation(Position + PixelDiff.xy, Context);
+		PixelDiff.w = TotalVariation;
+#else
+		PixelDiff.w = 1.0f;
+#endif
+		if (!DebugMotionVector)
+		{
+			RWPixelOffsetTexture[Position] = PixelDiff;
+		}
+		else
+		{
+			if (MipLevel >= 6)
+			{
+				RWPixelOffsetTexture[Position] = PixelDiff;
+			}
+			else
+			{
+				RWPixelOffsetTexture[Position] = float4(OffsetFromLowerMip.xy, OffsetFromLowerMip.Dist, length(OffsetFromLowerMip.xy));
+			}
+		}
+	}
+}
+#endif // TEMPORAL_REPROJECTION_MERGE
+
+#if TEMPORAL_HIGHFREQ_REJECT
+
+RWTexture2D<float4> OutputTexture;
+float HighFrequencyCutoffDeltaE;
+
+[numthreads(THREAD_SIZE_X, THREAD_SIZE_Y, 1)]
+void TemporalHighFrequencyRejectCS(uint3 DT_ID : SV_DispatchThreadID)
+{
+	int2 Position = DT_ID.xy;
+	int2 MaxPosition = (int2)TargetViewport_Extent;
+	if (all(Position < MaxPosition))
+	{
+
+		FPixelOffset Offset = GetPixelOffsetFromLowerMip(Position, 0);
+
+		float2 Pos = Position + Offset.xy;
+
+		float2 UV = (Pos + 0.5f) * TargetViewport_ExtentInverse;
+		float4 History_Warped = Texture2DSampleBicubic(SourceTexture, SharedTextureSampler, UV, TargetViewport_Extent, TargetViewport_ExtentInverse);
+
+		float4 Target = TargetTexture[Position];
+
+		//calculate the distance between warped source and target
+		const bool bNormalizeDeltaE = true;
+		float dENormalized = GetDeltaE(History_Warped.rgb*View.PreExposure, Target.rgb*View.PreExposure, bNormalizeDeltaE);
+
+		// since the warping of the source image is based on the distance of a region with weights
+		// it is a matching based on low frequency information. High frequency detail is not well
+		// utilized. So in this CS, we try to use the high frequency information to correct the blending
+		// factor per pixel. E.g., when albedo is used, we expect a high matching.
+		// TODO: Explore different buffers.
+		OutputTexture[Position] = float4(dENormalized,
+			dENormalized >= NormalizeDeltaE2000(HighFrequencyCutoffDeltaE),
+			Luminance(History_Warped.rgb),
+			Luminance(Target.rgb));
+	}
+
+}
+#endif // TEMPORAL_HIGHFREQ_REJECT
+
+
+#if TEMPORAL_REPROJECTION_RESOLVE
+Texture2D<float4> HighFrequencyRejectMap;
+RWTexture2D<float4> OutputTexture;
+float Alpha;
+float Kappa;
+float Eta;
+float HighFrequencyCutoffDeltaE;
+uint PatchId;
+uint PatchCount;
+
+[numthreads(THREAD_SIZE_X, THREAD_SIZE_Y, 1)]
+void TemporalResolveCS(uint3 DT_ID : SV_DispatchThreadID)
+{
+	int2 Position = DT_ID.xy;
+	int2 MaxPosition = (int2)TargetViewport_Extent;
+	if (all(Position < MaxPosition))
+	{
+
+		FPixelOffset Offset = GetPixelOffsetFromLowerMip(Position, 0);
+
+		float2 Pos = Position +Offset.xy;		
+		float2 UV = (Pos + 0.5f) * TargetViewport_ExtentInverse;
+		float4 Sampling = Texture2DSampleBicubic(SourceTexture, SharedTextureSampler, UV, TargetViewport_Extent, TargetViewport_ExtentInverse);
+		
+		float m = max(Kappa * (Offset.Dist - Eta), 0.0f);
+		m = 2 * m / (1 + m);
+		float FinalAlpha = clamp(Alpha * (1 - m) * Offset.TV, 0, 1);
+		
+		float dEHighFrequency = HighFrequencyRejectMap[Position].x;
+
+		// Determine the final alpha based on perception error
+		// 1. if the high frequency difference is larger than 10dE reject
+		// 2. if the difference is smaller than the high frequency cutoff, use the weight based on low frequency
+		// 3. otherwise, we move the target color to the source by a distance at most 0.5 dE
+		// TODO: lerp in Lab space instead of linear RGB space?
+		FinalAlpha = lerp(FinalAlpha,
+			min(FinalAlpha, NormalizeDeltaE2000(0.5f) / dEHighFrequency),
+			dEHighFrequency >= NormalizeDeltaE2000(HighFrequencyCutoffDeltaE));
+
+		FinalAlpha = lerp(FinalAlpha, 0, dEHighFrequency >= NormalizeDeltaE2000(10.0f));
+
+		float4 Accumulation = 0.0f;
+		float Weight = 0.0f;
+		float4 TargetValue = TargetTexture[Position];
+
+		BRANCH
+		if  (PatchCount == 1)
+		{	
+			// Use exponential moving to low pass the signal to create more
+			// temporal stable results.
+
+			Accumulation = (1 - FinalAlpha) * TargetValue + FinalAlpha * Sampling;
+			Weight = FinalAlpha;
+		}
+		else
+		{
+			// When we have multiple patches, use bilaterial filter to combine them
+			// TODO: improve high frequency information for this mode.
+			
+			Accumulation = TargetValue;
+			Weight = (PatchId == 0) ? 1 : Accumulation.w;
+		
+			Accumulation += FinalAlpha * Sampling;
+			Weight += FinalAlpha;
+
+			if (PatchId + 1 == PatchCount)
+			{
+				Accumulation /= Weight;
+			}
+		}
+
+		OutputTexture[Position] = float4(max(Accumulation.xyz, 0.0f), Weight);
+	}
+}
+#endif // TEMPORAL_REPROJECTION_RESOLVE
+
+// TODO: Ablation test
+float GetWeightMultiplier(float SourceTV, float TargetSourceTV)
+{
+	// Use less history if total variation is larger than reference
+	float delta = -(TargetSourceTV - SourceTV);
+	float depth = 0.2f;
+	return 1 / (1 + exp(-10.0 * delta)) * depth + (1 - depth / 2.0);
+}
+
+#if MOTION_VECTOR_SUBTRACT
+RWTexture2D<float4> Minuend;
+Texture2D<float4> Subtrahend;
+uint MipLevel;
+
+[numthreads(THREAD_SIZE_X, THREAD_SIZE_Y, 1)]
+void MotionVectorSubtractCS(uint3 DT_ID : SV_DispatchThreadID)
+{
+	int2 Position = DT_ID.xy;
+	int2 MaxPosition = int2(((uint2)TargetViewport_Extent) >> MipLevel);
+	if (all(Position < MaxPosition))
+	{
+		float4 M = Minuend[Position];
+		float4 S = Subtrahend[Position];
+
+		float4 M_S = float4(M.x - S.x, M.y - S.y, M.z, GetWeightMultiplier(S.w, M.w));
+
+		// Do not estimate the TV at the border when history is not available
+		const uint NumOfIgnoredBorderPixels = 12;
+		if (any(Position > TargetViewport_Extent - NumOfIgnoredBorderPixels) ||
+			any(Position - NumOfIgnoredBorderPixels < float2(0, 0)))
+		{
+			M_S.w = 1.0f;
+		}
+
+		Minuend[Position] = M_S;
+	}
+}
+#endif // MOTION_VECTOR_SUBTRACT
+
+#if VISUALIZE_MOTIONVECTOR
+
+// h = [0, 1], s = [0 ,1], v = [0, 1]
+// h: hue, 
+// s: saturate
+// v: value
+float3 hsv2rgb(float h, float s, float v)
+{
+	h = 360.0 * h;
+
+	float C = v * s;
+	float X = C * (1 - abs( (h/60) % 2 - 1));
+	float m = v - C;
+
+	float3 RGB = 0;
+	if (h >= 0 && h < 60)
+	{
+		RGB.r = C;
+		RGB.g = X;
+	}
+	else if (h >= 60 && h < 120)
+	{
+		RGB.r = X;
+		RGB.g = C;
+	}
+	else if (h >= 120 && h < 180)
+	{
+		RGB.g = C;
+		RGB.b = X;
+	}
+	else if (h >= 180 && h < 240)
+	{
+		RGB.g = X;
+		RGB.b = C;
+	}
+	else if (h >= 240 && h < 300)
+	{
+		RGB.r = X;
+		RGB.b = C;
+	}
+	else if (h >= 300 && h < 360)
+	{
+		RGB.r = C;
+		RGB.b = X;
+	}
+
+	RGB += m;
+	return RGB;
+}
+
+//:adaptive from https://www.shadertoy.com/view/MsS3Wc by iq
+float3 hsv2rgb_smooth(float3 c)
+{
+	float3 rgb = clamp(abs( (c.x * 6.0 + float3(0.0, 4.0, 2.0)) % 6.0 - 3.0) - 1.0, 0.0, 1.0);
+
+	rgb = rgb * rgb * (3.0 - 2.0 * rgb); // cubic smoothing	
+
+	return c.z * lerp(float3(1.0, 1.0, 1.0), rgb, c.y);
+}
+
+
+Texture2D<float4> TemporalDenoisingMotionVector;
+Texture2D<float4> DenoisedTexture;
+
+#define VISUALIZE_VECTOR 1
+#define VISUALIZE_COLOR_ENCODING 1
+#define VECTOR_GRID_LENGTH 32
+#define COLOR_METER_LENGTH 50
+
+void VisualizePathTracingMotionVector(
+	in noperspective float2 UV : TEXCOORD0,
+	out float4 OutColor : SV_Target0
+)
+{
+	float2 BufferSize = View.BufferSizeAndInvSize.xy;
+	int3 TexCoord = int3(UV * BufferSize - View.ViewRectMin.xy, 0);
+
+	float4 Velocity = TemporalDenoisingMotionVector.Load(TexCoord);
+
+	// draw the denoised scene
+	OutColor.rgb = DenoisedTexture.Load(TexCoord) * View.PreExposure;
+
+#if VISUALIZE_COLOR_ENCODING
+	float s, h, v;
+	if (length(Velocity.xy) >= 0.001) // ignore small velocity color encoding
+	{
+#if 0
+		s = clamp(lengthFast(float3(Velocity.xy, 0.0f)), 0.0f, 50.0f) / 50.0f;
+		h = atan2(Velocity.y, Velocity.x) / 3.1415926f * 0.5f + 0.5f;
+		v = 1.0f;
+		OutColor = float4(hsv2rgb(h, s, v), 1.0f);// float4(s, h, Velocity.b, 1.0f);
+		OutColor.rgb = pow(1 - OutColor.rgb, 1 / 2.2f);
+#else
+		v = pow(clamp(length(float3(Velocity.xy, 0.0f)), 0.0f, 50.0f) / 50.0f, 1 / 2.2f);
+		h = atan2(Velocity.y, -Velocity.x) / 3.1415927f;
+		h = h < 0 ? ((h + 2) / 2) : h / 2.0;
+		s = 1.0f;
+		OutColor = float4(hsv2rgb_smooth(float3(h, s, v)), 1.0f);// float4(s, h, Velocity.b, 1.0f);
+		OutColor.rgb = pow(OutColor.rgb, Velocity.z);
+#endif
+	}
+	else
+	{
+		OutColor.rgb = 0.0f;
+	}
+
+	// add motion color meter
+	if (length(TexCoord.xy - COLOR_METER_LENGTH) < COLOR_METER_LENGTH)
+	{
+		float2 Vector = TexCoord.xy - COLOR_METER_LENGTH;
+#if 0
+		s = length(float3(Vector, 0)) / COLOR_METER_LENGTH;
+		h = atan2(-Vector.y, -Vector.x) / 3.1415927f * 0.5f + 0.4999f;
+		v = 1.0f;
+		OutColor = float4(hsv2rgb(h, s, v), 1.0f);// float4(s, h, Velocity.b, 1.0f);
+		OutColor.rgb = pow(1 - OutColor.rgb, 1 / 2.2f);
+#else
+		v = pow(length(float3(Vector, 0)) / COLOR_METER_LENGTH, 1 / 2.2f);
+		h = atan2(-Vector.y, Vector.x) / 3.1415927f;
+		h = h < 0 ? ((h + 2) / 2) : h / 2.0;
+		s = 1.0f;
+		OutColor = float4(hsv2rgb_smooth(float3(h, s, v)), 1.0f);// float4(s, h, Velocity.b, 1.0f);
+		OutColor.rgb = pow(OutColor.rgb, 1.0f);
+#endif
+	}
+#endif 
+
+#if VISUALIZE_VECTOR
+	
+	int2 P0 = (TexCoord.xy / VECTOR_GRID_LENGTH) * VECTOR_GRID_LENGTH + VECTOR_GRID_LENGTH / 2;
+	Velocity = TemporalDenoisingMotionVector.Load(int3(P0, 0));
+	float2 P1 = (float2)P0 + Velocity.xy;
+	float2 N = normalize(float2(Velocity.y, -Velocity.x));
+	float2 P2 = (float2)TexCoord.xy;
+	
+#if 0	// plot actual velocity
+	float Dist = abs(dot(N, P2 - P0));
+	float Lambda = dot(P2 - P0, P1 - P0) / dot(Velocity.xy, Velocity.xy);
+	float P02P1Length = length(Velocity.xy);
+
+	if (P02P1Length > 0.001)
+	{
+		if (Dist <= 3 && Lambda >= -0.1 && Lambda <= 1.1)
+		{
+			
+			if (Dist <= 1.0)
+			{
+				OutColor = lerp(float4(0.5f, 0.5f, 0.5f, 1.0f), float4(1.0f, 0.0f, 0.0f, 1.0f), Lambda);
+			}
+
+			if (length(P2-P0) <= 2)
+			{
+				OutColor = 0.6f;
+			}
+			
+			if (length(P2-P1)<=2)
+			{
+				OutColor = float4(1.0, 0.0, 0.0, 1.0f);
+			}
+			
+
+		}
+	}
+#else
+	float3 LineColor = float3(0.7, 0.7, 0.7);
+
+	// points into the movement direction, movement direction is the inverse of the reprojection velocity.
+	float2 PixelDirection = -Velocity.xy;
+	float2 PixelPosAtTileCenter = P0;
+	float2 PixelPos = P2;
+	// arrow
+	{
+		float2 PerpPixelDirection = float2(PixelDirection.y, -PixelDirection.x);
+		float2 DirectionInTile = PixelPos - PixelPosAtTileCenter;
+
+		float DistOnLine = dot(normalize(-PixelDirection), DirectionInTile) + length(PixelDirection);
+
+		bool bArrowHead = DistOnLine < 8;
+
+		float LocalThickness = 1 + (frac(DistOnLine / 8) * 8) * 0.25f;
+
+		float PerpDirectionMask = saturate(LocalThickness - abs(dot(normalize(PerpPixelDirection), DirectionInTile)));
+		float DirectionMask = saturate(length(PixelDirection) - length(DirectionInTile));
+
+		float3 LineMask = PerpDirectionMask * DirectionMask;
+		OutColor.rgb = lerp(OutColor.rgb, LineColor, LineMask);
+	}
+
+	// previous pos is a dot
+	{
+		float3 DotColor = float3(0, 1, 0);
+		// PixelPos of the previous position
+		float2 PreviousPixelPos = PixelPosAtTileCenter - PixelDirection;
+		float Dist = length(PreviousPixelPos - PixelPos);
+		OutColor.rgb = lerp(OutColor.rgb, LineColor, saturate(3 - Dist));
+		OutColor.rgb = lerp(OutColor.rgb, 0, saturate(1.5f - Dist));
+	}
+#endif
+
+#endif
+
+}
+
+#endif
+
+#if VISUALIZE_WARPING
+
+Texture2D<float4> TemporalDenoisingMotionVector;
+Texture2D<float4> DenoisedTexture;
+Texture2D<float4> SourceTexture;
+SamplerState SharedTextureSampler;
+
+void FVisualizeWarpingPS(
+	in noperspective float2 UV : TEXCOORD0,
+	out float4 OutColor : SV_Target0)
+{
+	float2 BufferSize = View.BufferSizeAndInvSize.xy;
+	int3 TexCoord = int3(UV * BufferSize - View.ViewRectMin.xy, 0);
+
+	float4 Velocity = TemporalDenoisingMotionVector.Load(TexCoord);
+
+	float2 Pos = TexCoord.xy + Velocity.xy;
+
+	float2 SourceUV = (Pos + 0.5f + View.ViewRectMin.xy) * TargetViewport_ExtentInverse;
+	float4 Sampling = Texture2DSampleBicubic(SourceTexture, SharedTextureSampler, SourceUV, TargetViewport_Extent, TargetViewport_ExtentInverse);
+
+	OutColor.rgb = Sampling.rgb * View.PreExposure;
+	OutColor.a = 1.0f;
+}
+#endif
+
+#if SPATIAL_DENOISING
+
+Texture2D<float4> InputTexture;
+Texture2D<float4> InputNormal;
+Texture2D<float4> InputAlbedo;
+RWTexture2D<float4> OutputTexture;
+
+[numthreads(THREAD_SIZE_X, THREAD_SIZE_Y, 1)]
+void SpatialDenoiserCS(uint3 DT_ID : SV_DispatchThreadID)
+{
+	int2 Position = DT_ID.xy;
+	int2 MaxPosition = (int2)TargetViewport_Extent;
+	if (all(Position < MaxPosition))
+	{
+		OutputTexture[Position] = InputTexture[Position];
+	}
+}
+#endif // SPATIAL_DENOISING
\ No newline at end of file
diff --git a/Engine/Source/Runtime/Renderer/Private/PathTracing.cpp b/Engine/Source/Runtime/Renderer/Private/PathTracing.cpp
--- a/Engine/Source/Runtime/Renderer/Private/PathTracing.cpp
+++ b/Engine/Source/Runtime/Renderer/Private/PathTracing.cpp
@@ -4,8 +4,6 @@
 #include "RHI.h"
 #include "PathTracingDenoiser.h"
 
-PathTracingDenoiserFunction* GPathTracingDenoiserFunc = nullptr;
-
 TAutoConsoleVariable<int32> CVarPathTracing(
 	TEXT("r.PathTracing"),
 	1,
@@ -28,6 +26,7 @@
 #include "HairStrands/HairStrandsData.h"
 #include "Modules/ModuleManager.h"
 #include <limits>
+#include "PathTracingSpatialTemporalDenoising.h"
 
 TAutoConsoleVariable<int32> CVarPathTracingCompaction(
 	TEXT("r.PathTracing.Compaction"),
@@ -270,15 +269,6 @@
 	ECVF_RenderThreadSafe
 );
 
-TAutoConsoleVariable<int32> CVarPathTracingDenoiser(
-	TEXT("r.PathTracing.Denoiser"),
-	-1,
-	TEXT("Enable denoising of the path traced output (if a denoiser plugin is active) (default = -1 (driven by postprocesing volume))\n")
-	TEXT("-1: inherit from PostProcessVolume\n")
-	TEXT("0: disable denoiser\n")
-	TEXT("1: enable denoiser (if a denoiser plugin is active)\n"),
-	ECVF_RenderThreadSafe
-);
 
 BEGIN_SHADER_PARAMETER_STRUCT(FPathTracingData, )
 	SHADER_PARAMETER(float, BlendFactor)
@@ -408,8 +398,15 @@
 	TRefCountPtr<IPooledRenderTarget> RadianceRT;
 	TRefCountPtr<IPooledRenderTarget> AlbedoRT;
 	TRefCountPtr<IPooledRenderTarget> NormalRT;
-	TRefCountPtr<IPooledRenderTarget> RadianceDenoisedRT;
+	TRefCountPtr<FRDGPooledBuffer> VarianceBuffer;
 
+	// Cache to improve the stability when frame denoising (SPP=r.pathtracing.SamplesPerPixel) is used in animation rendering
+	TRefCountPtr<IPooledRenderTarget> LastDenoisedRadianceRT;
+	TRefCountPtr<IPooledRenderTarget> LastRadianceRT;
+	TRefCountPtr<IPooledRenderTarget> LastNormalRT;
+	TRefCountPtr<IPooledRenderTarget> LastAlbedoRT;
+	TRefCountPtr<FRDGPooledBuffer> LastVarianceBuffer;
+
 	// Texture holding onto the precomputed atmosphere data
 	TRefCountPtr<IPooledRenderTarget> AtmosphereOpticalDepthLUT;
 	FAtmosphereConfig LastAtmosphereConfig;
@@ -1641,15 +1638,25 @@
 	}
 }
 
-void FSceneViewState::PathTracingInvalidate()
+void FSceneViewState::PathTracingInvalidate(bool InvalidateAnimationStates)
 {
 	FPathTracingState* State = PathTracingState.Get();
 	if (State)
 	{
+		
+		if(InvalidateAnimationStates)
+		{
+			State->LastDenoisedRadianceRT.SafeRelease();
+			State->LastRadianceRT.SafeRelease();
+			State->LastNormalRT.SafeRelease();
+			State->LastAlbedoRT.SafeRelease();
+			State->LastVarianceBuffer.SafeRelease();
+		}
+
 		State->RadianceRT.SafeRelease();
 		State->AlbedoRT.SafeRelease();
 		State->NormalRT.SafeRelease();
-		State->RadianceDenoisedRT.SafeRelease();
+		State->VarianceBuffer.SafeRelease();
 		State->SampleIndex = 0;
 	}
 }
@@ -1680,13 +1687,6 @@
 END_SHADER_PARAMETER_STRUCT()
 #endif
 
-BEGIN_SHADER_PARAMETER_STRUCT(FDenoiseTextureParameters, )
-	RDG_TEXTURE_ACCESS(InputTexture, ERHIAccess::CopySrc)
-	RDG_TEXTURE_ACCESS(InputAlbedo, ERHIAccess::CopySrc)
-	RDG_TEXTURE_ACCESS(InputNormal, ERHIAccess::CopySrc)
-	RDG_TEXTURE_ACCESS(OutputTexture, ERHIAccess::CopyDest)
-END_SHADER_PARAMETER_STRUCT()
-
 DECLARE_GPU_STAT_NAMED(Stat_GPU_PathTracing, TEXT("Path Tracing"));
 void FDeferredShadingSceneRenderer::RenderPathTracing(
 	FRDGBuilder& GraphBuilder,
@@ -2144,12 +2144,9 @@
 
 	// Figure out if the denoiser is enabled and needs to run
 	FRDGTexture* DenoisedRadianceTexture = nullptr;
-	int DenoiserMode = CVarPathTracingDenoiser.GetValueOnRenderThread();
-	if (DenoiserMode < 0)
-	{
-		DenoiserMode = View.FinalPostProcessSettings.PathTracingEnableDenoiser;
-	}
-	const bool IsDenoiserEnabled = DenoiserMode != 0 && GPathTracingDenoiserFunc != nullptr;
+	bool IsDenoiserEnabled = IsPathTracingDenoiserEnabled(View);
+	int DenoiserMode = GetPathTracingDenoiserMode(View);
+
 	// Request denoise if this is the last sample OR allow turning on the denoiser after the image has stopped accumulating samples
 	const bool NeedsDenoise = IsDenoiserEnabled &&
 		(((Config.PathTracingData.Iteration + 1) == MaxSPP) ||
@@ -2198,47 +2195,75 @@
 	}
 #endif
 
+	FPathTracingSpatialTemporalDenoisingContext DenoisingContext;
+	const bool EnablePathTracingDenoiserRealtimeDebug = ShouldEnablePathTracingDenoiserRealtimeDebug();
+
 	if (IsDenoiserEnabled)
-	{
-		if (PathTracingState->RadianceDenoisedRT)
+	{		
+		// 1. Prepass to estimate pixel variance
+		FRDGBuffer* CurrentVarianceBufer = nullptr;
 		{
-			// we already have a texture for this
-			DenoisedRadianceTexture = GraphBuilder.RegisterExternalTexture(PathTracingState->RadianceDenoisedRT, TEXT("PathTracer.DenoisedRadiance"));
+			DenoisingContext.RadianceTexture = RadianceTexture;
+			if (PathTracingState->VarianceBuffer)
+			{
+				DenoisingContext.VarianceBuffer = GraphBuilder.RegisterExternalBuffer(PathTracingState->VarianceBuffer, TEXT("PathTracing.VarianceBuffer"));
+			}
+
+			PathTracingSpatialTemporalDenoisingPrePass(GraphBuilder, View, Config.PathTracingData.Iteration, DenoisingContext);
+
+			CurrentVarianceBufer = DenoisingContext.VarianceBuffer;
 		}
 
-		if (NeedsDenoise)
+		// 2. Denoising pass
+		if (NeedsDenoise || EnablePathTracingDenoiserRealtimeDebug)
 		{
-			if (DenoisedRadianceTexture == nullptr)
+			DenoisingContext.AlbedoTexture = AlbedoTexture;
+			DenoisingContext.NormalTexture = NormalTexture;
+			DenoisingContext.RadianceTexture = RadianceTexture;
+			DenoisingContext.FrameIndex = PathTracingState->FrameIndex;
+			DenoisingContext.VarianceBuffer = CurrentVarianceBufer;
+
+			if (PathTracingState->LastDenoisedRadianceRT)
 			{
-				// First time through, need to make a new texture
-				FRDGTextureDesc RadianceTextureDesc = FRDGTextureDesc::Create2D(
-					View.ViewRect.Size(),
-					PF_A32B32G32R32F,
-					FClearValueBinding::None,
-					TexCreate_ShaderResource | TexCreate_UAV);
-				DenoisedRadianceTexture = GraphBuilder.CreateTexture(RadianceTextureDesc, TEXT("PathTracer.DenoisedRadiance"), ERDGTextureFlags::MultiFrame);
+				DenoisingContext.LastDenoisedRadianceTexture =
+					GraphBuilder.RegisterExternalTexture(PathTracingState->LastDenoisedRadianceRT, TEXT("PathTracing.LastPreDenoisedRadiance"));
+				DenoisingContext.LastRadianceTexture =
+					GraphBuilder.RegisterExternalTexture(PathTracingState->LastRadianceRT, TEXT("PathTracing.LastRadianceTexture"));
+				DenoisingContext.LastNormalTexture =
+					GraphBuilder.RegisterExternalTexture(PathTracingState->LastNormalRT, TEXT("PathTracing.LastNormalTexture"));
+				DenoisingContext.LastAlbedoTexture =
+					GraphBuilder.RegisterExternalTexture(PathTracingState->LastAlbedoRT, TEXT("PathTracing.LastAlbedoTexture"));
+
+				DenoisingContext.LastVarianceBuffer = PathTracingState->LastVarianceBuffer?
+					GraphBuilder.RegisterExternalBuffer(PathTracingState->LastVarianceBuffer, TEXT("PathTracing.LastVarianceBuffer")) : nullptr;
 			}
 
-			FDenoiseTextureParameters* DenoiseParameters = GraphBuilder.AllocParameters<FDenoiseTextureParameters>();
-			DenoiseParameters->InputTexture = RadianceTexture;
-			DenoiseParameters->InputAlbedo = AlbedoTexture;
-			DenoiseParameters->InputNormal = NormalTexture;
-			DenoiseParameters->OutputTexture = DenoisedRadianceTexture;
-			// Need to read GPU mask outside Pass function, as the value is not refreshed inside the pass
-			GraphBuilder.AddPass(RDG_EVENT_NAME("Path Tracer Denoiser Plugin"), DenoiseParameters, ERDGPassFlags::Readback, 
-				[DenoiseParameters, DenoiserMode, GPUMask = View.GPUMask](FRHICommandListImmediate& RHICmdList)
-				{
-					GPathTracingDenoiserFunc(RHICmdList,
-						DenoiseParameters->InputTexture->GetRHI()->GetTexture2D(),
-						DenoiseParameters->InputAlbedo->GetRHI()->GetTexture2D(),
-						DenoiseParameters->InputNormal->GetRHI()->GetTexture2D(),
-						DenoiseParameters->OutputTexture->GetRHI()->GetTexture2D(),
-						GPUMask);
-				}
-			);
+			PathTracingSpatialTemporalDenoising(GraphBuilder,
+				View,
+				DenoiserMode,
+				DenoisedRadianceTexture,
+				DenoisingContext);
 
-			GraphBuilder.QueueTextureExtraction(DenoisedRadianceTexture, &PathTracingState->RadianceDenoisedRT);
+			GraphBuilder.QueueTextureExtraction(DenoisedRadianceTexture, &PathTracingState->LastDenoisedRadianceRT);
+			GraphBuilder.QueueTextureExtraction(NormalTexture, &PathTracingState->LastNormalRT);
+			GraphBuilder.QueueTextureExtraction(AlbedoTexture, &PathTracingState->LastAlbedoRT);
+			GraphBuilder.QueueTextureExtraction(RadianceTexture, &PathTracingState->LastRadianceRT);
 		}
+
+		// 3. Update pixel variance
+		if (CurrentVarianceBufer)
+		{
+			GraphBuilder.QueueBufferExtraction(CurrentVarianceBufer, 
+				(NeedsDenoise || EnablePathTracingDenoiserRealtimeDebug) ?
+				&PathTracingState->LastVarianceBuffer:
+				&PathTracingState->VarianceBuffer);
+
+			if (NeedsDenoise || EnablePathTracingDenoiserRealtimeDebug)
+			{
+				PathTracingState->VarianceBuffer = nullptr;
+			}
+		}
+		
 	}
 	PathTracingState->LastConfig.DenoiserMode = DenoiserMode;
 
@@ -2254,12 +2279,21 @@
 
 	FScreenPassTextureViewport Viewport(SceneColorOutputTexture, View.ViewRect);
 
+	const bool IsCursorInsideView = View.CursorPos.X != -1 || View.CursorPos.Y != -1;
 	// wiper mode - reveals the render below the path tracing display
 	// NOTE: we still path trace the full resolution even while wiping the cursor so that rendering does not get out of sync
 	if (CVarPathTracingWiperMode.GetValueOnRenderThread() != 0)
 	{
 		float DPIScale = FPlatformApplicationMisc::GetDPIScaleFactorAtPoint(View.CursorPos.X, View.CursorPos.Y);
-		Viewport.Rect.Min.X = View.CursorPos.X / DPIScale;
+		
+		if (IsCursorInsideView)
+		{
+			Viewport.Rect.Min.X = View.CursorPos.X / DPIScale;
+		}
+		else
+		{
+			Viewport.Rect.Min.X = 0.5 * View.ViewRect.Min.X + 0.5 * View.ViewRect.Max.X;
+		}
 	}
 
 	TShaderMapRef<FPathTracingCompositorPS> PixelShader(View.ShaderMap);
@@ -2272,6 +2306,35 @@
 		PixelShader,
 		DisplayParameters
 	);
+
+	// Add a visualization path for denoising
+	if (NeedsDenoise || EnablePathTracingDenoiserRealtimeDebug)
+	{
+		FVisualizePathTracingDenoisingInputs Inputs;
+		Inputs.SceneColor =SceneColorOutputTexture;
+
+		FScreenPassTextureViewport MotionVectorViewport(SceneColorOutputTexture, View.ViewRect);
+		if (CVarPathTracingWiperMode.GetValueOnRenderThread() != 0)
+		{
+			float DPIScale = FPlatformApplicationMisc::GetDPIScaleFactorAtPoint(View.CursorPos.X, View.CursorPos.Y);
+			if (IsCursorInsideView)
+			{
+				MotionVectorViewport.Rect.Max.X = View.CursorPos.X / DPIScale;
+			}
+			else
+			{
+				MotionVectorViewport.Rect.Max.X = 0.5 * View.ViewRect.Min.X + 0.5 * View.ViewRect.Max.X;
+			}
+		}
+
+		Inputs.Viewport = MotionVectorViewport;
+
+		Inputs.DenoisingContext = DenoisingContext;
+		Inputs.SceneTexturesUniformBuffer = SceneTexturesUniformBuffer;
+		Inputs.DenoisedTexture = DenoisedRadianceTexture;
+
+		AddVisualizePathTracingDenoisingPass(GraphBuilder, View, Inputs);
+	}
 }
 
 #endif
diff --git a/Engine/Source/Runtime/Renderer/Private/PathTracingSpatialTemporalDenoising.cpp b/Engine/Source/Runtime/Renderer/Private/PathTracingSpatialTemporalDenoising.cpp
new file mode 100644
--- /dev/null
+++ b/Engine/Source/Runtime/Renderer/Private/PathTracingSpatialTemporalDenoising.cpp
@@ -0,0 +1,1530 @@
+// Copyright Epic Games, Inc. All Rights Reserved.
+
+#include "PathTracingSpatialTemporalDenoising.h"
+#include "PathTracing.h"
+#include "RHI.h"
+#include "PathTracingDenoiser.h"
+
+PathTracingDenoiserFunction* GPathTracingDenoiserFunc = nullptr;
+
+#if RHI_RAYTRACING
+
+#include "RendererPrivate.h"
+#include "GlobalShader.h"
+#include "DeferredShadingRenderer.h"
+#include "HAL/PlatformApplicationMisc.h"
+#include "RayTracingTypes.h"
+#include "RayTracingDefinitions.h"
+#include "PathTracingDefinitions.h"
+#include "RayTracing/RayTracingMaterialHitShaders.h"
+#include "RenderCore/Public/GenerateMips.h"
+#include "RenderGraphDefinitions.h"
+#include <limits>
+
+DEFINE_LOG_CATEGORY_STATIC(LogPathTracingDenoising, Log, All);
+
+BEGIN_SHADER_PARAMETER_STRUCT(FDenoiseTextureParameters, )
+	RDG_TEXTURE_ACCESS(InputTexture, ERHIAccess::CopySrc)
+	RDG_TEXTURE_ACCESS(InputAlbedo, ERHIAccess::CopySrc)
+	RDG_TEXTURE_ACCESS(InputNormal, ERHIAccess::CopySrc)
+	RDG_TEXTURE_ACCESS(OutputTexture, ERHIAccess::CopyDest)
+END_SHADER_PARAMETER_STRUCT()
+
+namespace {
+
+	TAutoConsoleVariable<int32> CVarPathTracingDenoiser(
+		TEXT("r.PathTracing.Denoiser"),
+		-1,
+		TEXT("Enable denoising of the path traced output (if a denoiser plugin is active) (default = -1 (driven by postprocesing volume))\n")
+		TEXT("-1: inherit from PostProcessVolume\n")
+		TEXT("0: disable denoiser\n")
+		TEXT("1: enable denoiser (if a denoiser plugin is active)\n"),
+		ECVF_RenderThreadSafe
+	);
+
+	TAutoConsoleVariable<int32> CVarPathTracingSpatialDenoiser(
+		TEXT("r.PathTracing.SpatialDenoiser"),
+		1,
+		TEXT("Enable spatial denoising of the path traced output\n")
+		TEXT("-1: inherit from PostProcessVolume\n")
+		TEXT("0: disable denoiser\n")
+		TEXT("1: enable denoiser (if a denoiser plugin is active)\n"),
+		ECVF_RenderThreadSafe
+	);
+
+	TAutoConsoleVariable<int32> CVarPathTracingTemporalDenoiser(
+		TEXT("r.PathTracing.TemporalDenoiser"),
+		0,
+		TEXT("Enable temporal denoising of the path traced output\n")
+		TEXT("-1: inherit from PostProcessVolume (TODO when out of experimental phase)\n")
+		TEXT("0: disable denoiser\n")
+		TEXT("1: enable denoiser (if a denoiser plugin is active)\n"),
+		ECVF_RenderThreadSafe
+	);
+
+	TAutoConsoleVariable<int32> CVarPathTracingTemporalDenoiserVisualizeMotionVector(
+		TEXT("r.PathTracing.TemporalDenoiser.VisualizeMotionVector"),
+		0,
+		TEXT("1: visualize the motion vector compared to the raster motion vector\n"),
+		ECVF_RenderThreadSafe
+	);
+
+	TAutoConsoleVariable<int32> CVarPathTracingTemporalDenoiserEnableSubPixelOffset(
+		TEXT("r.PathTracing.TemporalDenoiser.EnableSubPixelOffset"),
+		1,
+		TEXT("Enable subpixel offset when merging\n")
+		TEXT("-1: inherit from PostProcessVolume\n")
+		TEXT("0: disable denoiser\n")
+		TEXT("1: enable denoiser (if a denoiser plugin is active)\n"),
+		ECVF_RenderThreadSafe
+	);
+
+	TAutoConsoleVariable<float> CVarPathTracingTemporalDenoiserKappa(
+		TEXT("r.PathTracing.TemporalDenoiser.kappa"),
+		-1.0f,
+		TEXT("Scaling parameter to determine how fast the history weight falls and the cutting point to zero. Use DeltaE to derive kappa if -1\n"),
+		ECVF_RenderThreadSafe
+	);
+
+	TAutoConsoleVariable<float> CVarPathTracingTemporalDenoiserEta(
+		TEXT("r.PathTracing.TemporalDenoiser.eta"),
+		1.0f,
+		TEXT("Eta param. Error distance below this will have max history weight. Use DeltaE to derive Eta if -1\n"),
+		ECVF_RenderThreadSafe
+	);
+
+	TAutoConsoleVariable<float> CVarPathTracingTemporalDenoiserDeltaE(
+		TEXT("r.PathTracing.TemporalDenoiser.DeltaE"),
+		5.0f,
+		TEXT("Cut off the history weight to zero at CIE DeltaE for low frequency.\n")
+		TEXT("1.0 :the just noticeable difference (JND),\n")
+		TEXT("2.0 :perceptible for close look\n")
+		TEXT("10.0:Perceptible at a glance ")
+		TEXT("This works as an alternative control instead of kappa. 2 as default."),
+		ECVF_RenderThreadSafe
+	);
+
+	TAutoConsoleVariable<float> CVarPathTracingTemporalDenoiserDeltaEHighFrequency(
+		TEXT("r.PathTracing.TemporalDenoiser.DeltaE.HighFrequency"),
+		2.1f,
+		TEXT("Cut off the history weight to zero when using high frequency per pixel difference.\n")
+		TEXT("It should only be enabled when the source image is smooth or denoised."),
+		ECVF_RenderThreadSafe
+	);
+
+	TAutoConsoleVariable<float> CVarPathTracingTemporalDenoiserAlpha(
+		TEXT("r.PathTracing.TemporalDenoiser.alpha"),
+		1.0f,
+		TEXT("The weight of history in the exponential mean average\n"),
+		ECVF_RenderThreadSafe
+	);
+
+	TAutoConsoleVariable<int32> CVarPathTracingTemporalDenoiserMode(
+		TEXT("r.PathTracing.TemporalDenoiser.mode"),
+		1,
+		TEXT("0: disabled \n")
+		TEXT("1: offline rendering only\n")
+		TEXT("2: online rendering (for debug)\n"),
+		ECVF_RenderThreadSafe
+	);
+
+	TAutoConsoleVariable<int32> CVarPathTracingTemporalDenoiserSource(
+		TEXT("r.PathTracing.TemporalDenoiser.source"),
+		0,
+		TEXT("0: Denoised Radance when possible (Default) \n")
+		TEXT("1: Normal\n")
+		TEXT("2: Albedo\n")
+		TEXT("3: Raw Radiance")
+		TEXT("Otherwise: Feature Fusion (TODO)"),
+		ECVF_RenderThreadSafe
+	);
+
+	TAutoConsoleVariable<int32> CVarPathTracingTemporalDenoiserDenoiseSourceImageFirst(
+		TEXT("r.PathTracing.TemporalDenoiser.DenoiseSourceImageFirst"),
+		0,
+		TEXT("Denoise the source image with IntelImageDenoisier"),
+		ECVF_RenderThreadSafe
+	);
+
+	TAutoConsoleVariable<int32> CVarPathTracingTemporalDenoiserSubPixelOffsetStartMip(
+		TEXT("r.PathTracing.TemporalDenoiser.SubPixelOffset.StartMip"),
+		8,
+		TEXT("From 0 to this mip, we will perform subpixel offset"),
+		ECVF_RenderThreadSafe
+	);
+
+	TAutoConsoleVariable<int32> CVarPathTracingTemporalDenoiserMotionOperation(
+	   TEXT("r.PathTracing.TemporalDenoiser.MotionOperation"),
+		1,
+		TEXT("0: use the motion vector directly estimated")
+		TEXT("1: subtract between the motion"),
+		ECVF_RenderThreadSafe
+	);
+
+	TAutoConsoleVariable<int32> CVarPathTracingTemporalDenoiserVisWarp(
+		TEXT("r.PathTracing.TemporalDenoiser.VisWarp"),
+		0,
+		TEXT("0: disable")
+		TEXT("1: visualize warped source by the motion vector")
+		TEXT("2: weights, warped source, and combined"),
+		ECVF_RenderThreadSafe
+	);
+
+	TAutoConsoleVariable<int32> CVarPathTracingTemporalDenoiserDistanceMetrics(
+		TEXT("r.PathTracing.TemporalDenoiser.DistanceMetrics"),
+		2,
+		TEXT("0: Luminance based metrics for distance estimation. Color with same luminance will create error in motion and history weights estimation.")
+		TEXT("1: Direct color difference")
+		TEXT("2: Visual color difference based on CIELAB2000 color difference."),
+		ECVF_RenderThreadSafe
+	);
+
+	TAutoConsoleVariable<int32> CVarPathTracingTemporalDenoiserTotalVariation(
+		TEXT("r.PathTracing.TemporalDenoiser.TotalVariation"),
+		1,
+		TEXT("!=0: Use less history if the total variation is large in a local patch"),
+		ECVF_RenderThreadSafe
+	);
+
+	TAutoConsoleVariable<int32> CVarPathTracingTemporalDenoiserPatchCount(
+		TEXT("r.PathTracing.TemporalDenoiser.PatchCount"),
+		1,
+		TEXT("The number of similar patches found by Non-Local Mean to use for temporal denoising\n")
+		TEXT("1: default. Accumulae the one with the minimal distance exponentially.")
+		TEXT(">1 && < 16: Use bilaterial filtering to accumulate multiple patches."),
+		ECVF_RenderThreadSafe
+	);
+
+	DECLARE_GPU_STAT(PathTracingSpatialTemporalDenoising)
+}
+
+int GetPathTracingDenoiserMode(const FViewInfo& View)
+{
+	int DenoiserMode = CVarPathTracingDenoiser.GetValueOnRenderThread();
+	if (DenoiserMode < 0)
+	{
+		DenoiserMode = View.FinalPostProcessSettings.PathTracingEnableDenoiser;
+	}
+	return DenoiserMode;
+}
+
+bool IsPathTracingDenoiserEnabled(const FViewInfo& View)
+{
+	int DenoiserMode = GetPathTracingDenoiserMode(View);
+	const bool IsDenoiserEnabled = DenoiserMode != 0 && GPathTracingDenoiserFunc != nullptr;
+
+	return IsDenoiserEnabled;
+}
+
+enum class ETemporalDenoisingMode : uint32 
+{
+	ETDM_DISABLED,
+	ETDM_OFFLINE,
+	ETDM_ONLINE,
+	ETDM_MAX
+};
+
+bool ShouldEnablePathTracingDenoiserRealtimeDebug()
+{
+	int TemporalDenoiserMode = CVarPathTracingTemporalDenoiserMode.GetValueOnRenderThread();
+	TemporalDenoiserMode = FMath::Clamp(TemporalDenoiserMode, 0, static_cast<int32>(ETemporalDenoisingMode::ETDM_MAX));
+	ETemporalDenoisingMode Mode = static_cast<ETemporalDenoisingMode>(TemporalDenoiserMode);
+
+	return Mode == ETemporalDenoisingMode::ETDM_ONLINE;
+}
+
+static bool ShouldApplySpatialDenoiser()
+{
+	return CVarPathTracingSpatialDenoiser.GetValueOnRenderThread() != 0;
+}
+
+static bool ShouldApplyTemporalDenoiser(const FPathTracingSpatialTemporalDenoisingContext& DenoisingContext)
+{
+	const bool bApplyTemporalDenoiser =
+		CVarPathTracingTemporalDenoiser.GetValueOnAnyThread() == 1 &&
+		DenoisingContext.LastDenoisedRadianceTexture &&
+		DenoisingContext.LastDenoisedRadianceTexture != DenoisingContext.RadianceTexture;
+
+	return bApplyTemporalDenoiser;
+}
+
+static float GetHighFrequencyCutoffDeltaE()
+{
+	return FMath::Max(1.0f, CVarPathTracingTemporalDenoiserDeltaEHighFrequency.GetValueOnRenderThread());
+}
+
+static bool ShouldRemoveSelfSubpixelOffset(int MipLevel)
+{
+	bool bRemoveSelfSubpixelOffset = CVarPathTracingTemporalDenoiserMotionOperation.GetValueOnRenderThread() == 1 
+		&& MipLevel <= CVarPathTracingTemporalDenoiserSubPixelOffsetStartMip.GetValueOnRenderThread()
+		&& CVarPathTracingTemporalDenoiserEnableSubPixelOffset.GetValueOnAnyThread() != 0;
+
+	return bRemoveSelfSubpixelOffset;
+}
+
+static float GetErrorDistanceBasedOnDeltaE(float DeltaE)
+{
+	const float NormalizeDeltaEAndEuclideanDistanceToSameMean = 0.0141f;
+	float ErrorDistanceFromDeltaE = log2f(2 + DeltaE * NormalizeDeltaEAndEuclideanDistanceToSameMean);
+	return ErrorDistanceFromDeltaE;
+}
+
+static void GetBlendingFactor(float& Kappa, float& Eta, float& Alpha)
+{
+	Kappa = CVarPathTracingTemporalDenoiserKappa.GetValueOnRenderThread();
+	Eta = CVarPathTracingTemporalDenoiserEta.GetValueOnRenderThread();
+	Alpha = CVarPathTracingTemporalDenoiserAlpha.GetValueOnRenderThread();
+
+	if (Eta == -1)
+	{
+		// Eta is set to 1 JND (just noticeable difference).
+		Eta = GetErrorDistanceBasedOnDeltaE(1.0f);
+	}
+
+	if (Kappa == -1)
+	{
+		// Use DeltaE to determine Kappa in our modified weight formula such that
+		// the history weight falls off to zero when we reaches DeltaE.
+
+		float DeltaE =CVarPathTracingTemporalDenoiserDeltaE.GetValueOnRenderThread();
+		DeltaE = FMath::Max(DeltaE, 1.0f + 1e-6f);
+		float FallOffToZeroErrorDistance = GetErrorDistanceBasedOnDeltaE(DeltaE);
+		Kappa = 1 / (FallOffToZeroErrorDistance - Eta);
+	}
+}
+
+static bool ShouldCompilePathTracingDenoiserShadersForProject(EShaderPlatform ShaderPlatform)
+{
+	static const auto CVarLocalVariablePathTracing = IConsoleManager::Get().FindTConsoleVariableDataInt(TEXT("r.PathTracing"));
+	const int PathTracing = CVarLocalVariablePathTracing ? CVarLocalVariablePathTracing->GetValueOnAnyThread() : 0;
+
+	return ShouldCompileRayTracingShadersForProject(ShaderPlatform) &&
+		FDataDrivenShaderPlatformInfo::GetSupportsPathTracing(ShaderPlatform) &&
+		PathTracing != 0;
+}
+
+static bool ShouldGenerateVarianceMap()
+{
+	int32 PathTracingTemporalDenoiserSource = CVarPathTracingTemporalDenoiserSource.GetValueOnRenderThread();
+	return PathTracingTemporalDenoiserSource < 0 || PathTracingTemporalDenoiserSource > 3;
+}
+
+static bool ShouldVisualizePathTracingVelocityState(const FPathTracingSpatialTemporalDenoisingContext& DenoisingContext)
+{
+	bool bVisualizeMotionVector = CVarPathTracingTemporalDenoiserVisualizeMotionVector.GetValueOnRenderThread() != 0;
+
+	return bVisualizeMotionVector && DenoisingContext.MotionVector;
+}
+
+static bool ShouldVisualizeWarping(const FPathTracingSpatialTemporalDenoisingContext& DenoisingContext)
+{
+	return CVarPathTracingTemporalDenoiserVisWarp.GetValueOnRenderThread() != 0 && 
+		DenoisingContext.MotionVector && DenoisingContext.LastDenoisedRadianceTexture;
+}
+
+static int32 GetTemporalAccumulationPatchCount()
+{
+	return FMath::Clamp(CVarPathTracingTemporalDenoiserPatchCount.GetValueOnRenderThread(), 1, 16);
+}
+
+static bool ShouldUseTotalVariation(uint32 MipLevel)
+{
+	return CVarPathTracingTemporalDenoiserTotalVariation.GetValueOnRenderThread() != 0.0f &&
+		(MipLevel == 0) &&
+		CVarPathTracingTemporalDenoiserEnableSubPixelOffset.GetValueOnAnyThread() != 0.0f;
+}
+
+static bool ShouldEnableSubpixelOffset(uint32 MipLevel)
+{
+	uint32 PixelOffsetStartMip = CVarPathTracingTemporalDenoiserSubPixelOffsetStartMip.GetValueOnRenderThread();
+	bool bShouldEnableSubpixelOffset = (MipLevel <= PixelOffsetStartMip) &&
+		CVarPathTracingTemporalDenoiserEnableSubPixelOffset.GetValueOnAnyThread() != 0;
+
+	return bShouldEnableSubpixelOffset;
+}
+
+static constexpr uint32 kMipDiffDelta = 2;
+
+static constexpr uint32 kNumberOfPixelShifts = 25;
+static constexpr uint32 kNumberOfShiftsPerTexture = 4;
+static constexpr uint32 kNumberOfPasses = 1;
+static constexpr uint32 kNumberOfShiftsPerPass = kNumberOfShiftsPerTexture * kNumberOfPasses;
+static constexpr uint32 kNumberOfTexturesPerPass = (kNumberOfPixelShifts + kNumberOfShiftsPerPass - 1) / kNumberOfShiftsPerPass;
+
+static constexpr uint32 kThreadSize = 8;
+
+static_assert(kNumberOfTexturesPerPass <= 7, "The number of buffers for Pixel correspondence estimation can not exceed  seven per pass");
+
+static const TCHAR* DistanceTextureNames[7] =
+{
+	TEXT("PathTracing.EstimateMotion.Disntace.0"),
+	TEXT("PathTracing.EstimateMotion.Disntace.1"),
+	TEXT("PathTracing.EstimateMotion.Disntace.2"),
+	TEXT("PathTracing.EstimateMotion.Disntace.3"),
+	TEXT("PathTracing.EstimateMotion.Disntace.4"),
+	TEXT("PathTracing.EstimateMotion.Disntace.5"),
+	TEXT("PathTracing.EstimateMotion.Disntace.6"),
+};
+
+BEGIN_SHADER_PARAMETER_STRUCT(FDenoisingCommonParameters, )
+	// Constant variable in each pass
+	SHADER_PARAMETER_STRUCT_REF(FViewUniformShaderParameters, ViewUniformBuffer)
+	SHADER_PARAMETER_STRUCT(FScreenPassTextureViewportParameters, TargetViewport)
+	SHADER_PARAMETER_SAMPLER(SamplerState, SharedTextureSampler)
+	SHADER_PARAMETER(uint32, PatchCount)
+	SHADER_PARAMETER(uint32, NumOfMips)
+
+	// Dynamic variable in each subpass.
+	SHADER_PARAMETER(uint32, PatchId)
+	SHADER_PARAMETER(uint32, MipLevel)
+END_SHADER_PARAMETER_STRUCT();
+
+
+class FTemporalReprojectionAlignCS : public FGlobalShader
+{
+public:
+	DECLARE_GLOBAL_SHADER(FTemporalReprojectionAlignCS);
+	SHADER_USE_PARAMETER_STRUCT(FTemporalReprojectionAlignCS, FGlobalShader);
+
+	BEGIN_SHADER_PARAMETER_STRUCT(FParameters, )
+		SHADER_PARAMETER_STRUCT_INCLUDE(FDenoisingCommonParameters, DenoisingCommonParameters)
+		SHADER_PARAMETER_RDG_TEXTURE_SRV(Texture2D<float4>, PixelOffsetTexture)
+		SHADER_PARAMETER_RDG_TEXTURE_SRV(Texture2D, SourceTexture)
+		SHADER_PARAMETER_RDG_TEXTURE_SRV(Texture2D, TargetTexture)
+		SHADER_PARAMETER_RDG_TEXTURE_UAV_ARRAY(RWTexture2D<float4>, RWDistanceTextures, [kNumberOfTexturesPerPass])
+	END_SHADER_PARAMETER_STRUCT()
+
+	enum class EDistanceMetrics : uint32 {
+		METRICS_LUMINANCE,
+		METRICS_EUCLIDEAN,
+		METRICS_PERCEPTION,
+		MAX
+	};
+
+	class FDistanceMetrics : SHADER_PERMUTATION_ENUM_CLASS("DISTANCE_METRICS", EDistanceMetrics);
+	using FPermutationDomain = TShaderPermutationDomain<FDistanceMetrics>;
+
+		static bool ShouldCompilePermutation(const FGlobalShaderPermutationParameters& Parameters)
+	{
+		return ShouldCompilePathTracingDenoiserShadersForProject(Parameters.Platform);
+	}
+
+	static void ModifyCompilationEnvironment(const FShaderPermutationParameters&, FShaderCompilerEnvironment& OutEnvironment)
+	{
+		OutEnvironment.SetDefine(TEXT("TEMPORAL_REPROJECTION_ALIGN"), 1);
+		OutEnvironment.SetDefine(TEXT("K_NUM_OF_TEXTURES_PER_PASS"), kNumberOfTexturesPerPass);
+	}
+
+	static EDistanceMetrics GetDistanceMetrics()
+	{
+		uint32 DistanceMetrics = CVarPathTracingTemporalDenoiserDistanceMetrics.GetValueOnRenderThread();
+		DistanceMetrics = FMath::Clamp(DistanceMetrics, 
+			static_cast<uint32>(EDistanceMetrics::METRICS_LUMINANCE),
+			static_cast<uint32>(EDistanceMetrics::MAX) - 1);
+
+		return static_cast<EDistanceMetrics>(DistanceMetrics);
+	}
+
+	static const TCHAR* GetEventName(EDistanceMetrics DistanceMetrics)
+	{
+		static const TCHAR* const kEventNames[] = {
+			TEXT("Luminance"),
+			TEXT("Euclidean"),
+			TEXT("Perception")
+		};
+		static_assert(UE_ARRAY_COUNT(kEventNames) == int32(EDistanceMetrics::MAX), "Fix me");
+		return kEventNames[static_cast<uint32>(DistanceMetrics)];
+	}
+};
+
+IMPLEMENT_GLOBAL_SHADER(FTemporalReprojectionAlignCS, "/Engine/Private/PathTracing/PathTracingSpatialTemporalDenoising.usf", "ReprojectionAlignCS", SF_Compute);
+
+class FTemporalReprojectionBlurCS : public FGlobalShader
+{
+public:
+	DECLARE_GLOBAL_SHADER(FTemporalReprojectionBlurCS);
+	SHADER_USE_PARAMETER_STRUCT(FTemporalReprojectionBlurCS, FGlobalShader);
+
+	BEGIN_SHADER_PARAMETER_STRUCT(FParameters, )
+		SHADER_PARAMETER_STRUCT_INCLUDE(FDenoisingCommonParameters, DenoisingCommonParameters)
+		SHADER_PARAMETER_RDG_TEXTURE_SRV(Texture2D<float4>, InputTexture)
+		SHADER_PARAMETER_RDG_TEXTURE_UAV(RWTexture2D<float4>, OutputTexture)
+	END_SHADER_PARAMETER_STRUCT()
+	
+	class FDimensionDirectCopy : SHADER_PERMUTATION_BOOL("DIRECT_COPY");
+	using FPermutationDomain = TShaderPermutationDomain<FDimensionDirectCopy>;
+
+		static bool ShouldCompilePermutation(const FGlobalShaderPermutationParameters& Parameters)
+	{
+		return ShouldCompilePathTracingDenoiserShadersForProject(Parameters.Platform);
+	}
+
+	static void ModifyCompilationEnvironment(const FShaderPermutationParameters&, FShaderCompilerEnvironment& OutEnvironment)
+	{
+		OutEnvironment.SetDefine(TEXT("TEMPORAL_REPROJECTION_BLUR"), 1);
+		OutEnvironment.SetDefine(TEXT("K_NUM_OF_TEXTURES_PER_PASS"), kNumberOfTexturesPerPass);
+	}
+};
+
+IMPLEMENT_GLOBAL_SHADER(FTemporalReprojectionBlurCS, "/Engine/Private/PathTracing/PathTracingSpatialTemporalDenoising.usf", "ReprojectionBlurCS", SF_Compute);
+
+class FTemporalReprojectionMergeCS : public FGlobalShader
+{
+public:
+	DECLARE_GLOBAL_SHADER(FTemporalReprojectionMergeCS);
+	SHADER_USE_PARAMETER_STRUCT(FTemporalReprojectionMergeCS, FGlobalShader);
+
+	BEGIN_SHADER_PARAMETER_STRUCT(FParameters, )
+		SHADER_PARAMETER_STRUCT_INCLUDE(FDenoisingCommonParameters, DenoisingCommonParameters)
+		SHADER_PARAMETER_RDG_TEXTURE_SRV_ARRAY(Texture2D<float4>, DistanceTextures, [kNumberOfTexturesPerPass])
+		SHADER_PARAMETER_RDG_TEXTURE_SRV(Texture2D<float4>, PixelOffsetTexture)
+		SHADER_PARAMETER_RDG_TEXTURE_UAV(RWTexture2D<float4>, RWPixelOffsetTexture)
+	END_SHADER_PARAMETER_STRUCT()
+	
+	class FSubpixelOffset : SHADER_PERMUTATION_BOOL("SUBPIXEL_OFFSET");
+	class FTotalVariation : SHADER_PERMUTATION_BOOL("TOTAL_VARIATION");
+	using FPermutationDomain = TShaderPermutationDomain<FSubpixelOffset,FTotalVariation>;
+
+		static bool ShouldCompilePermutation(const FGlobalShaderPermutationParameters& Parameters)
+	{
+		return ShouldCompilePathTracingDenoiserShadersForProject(Parameters.Platform);
+	}
+
+	static void ModifyCompilationEnvironment(const FShaderPermutationParameters&, FShaderCompilerEnvironment& OutEnvironment)
+	{
+		OutEnvironment.SetDefine(TEXT("TEMPORAL_REPROJECTION_MERGE"), 1);
+		OutEnvironment.SetDefine(TEXT("K_NUM_OF_TEXTURES_PER_PASS"), kNumberOfTexturesPerPass);
+		OutEnvironment.CompilerFlags.Add(CFLAG_AllowTypedUAVLoads);
+	}
+
+};
+
+IMPLEMENT_GLOBAL_SHADER(FTemporalReprojectionMergeCS, "/Engine/Private/PathTracing/PathTracingSpatialTemporalDenoising.usf", "ReprojectionMergeCS", SF_Compute);
+
+class FMotionVectorSubtractCS : public FGlobalShader
+{
+public:
+	DECLARE_GLOBAL_SHADER(FMotionVectorSubtractCS);
+	SHADER_USE_PARAMETER_STRUCT(FMotionVectorSubtractCS, FGlobalShader);
+
+	BEGIN_SHADER_PARAMETER_STRUCT(FParameters, )
+		SHADER_PARAMETER_STRUCT_INCLUDE(FDenoisingCommonParameters, DenoisingCommonParameters)
+		SHADER_PARAMETER_RDG_TEXTURE_UAV(RWTexture2D, Minuend)
+		SHADER_PARAMETER_RDG_TEXTURE_SRV(Texture2D, Subtrahend)
+	END_SHADER_PARAMETER_STRUCT()
+
+		static bool ShouldCompilePermutation(const FGlobalShaderPermutationParameters& Parameters)
+	{
+		return ShouldCompilePathTracingDenoiserShadersForProject(Parameters.Platform);
+	}
+
+	static void ModifyCompilationEnvironment(const FShaderPermutationParameters&, FShaderCompilerEnvironment& OutEnvironment)
+	{
+		OutEnvironment.SetDefine(TEXT("MOTION_VECTOR_SUBTRACT"), 1);
+		OutEnvironment.CompilerFlags.Add(CFLAG_AllowTypedUAVLoads);
+	}
+};
+
+IMPLEMENT_GLOBAL_SHADER(FMotionVectorSubtractCS, "/Engine/Private/PathTracing/PathTracingSpatialTemporalDenoising.usf", "MotionVectorSubtractCS", SF_Compute);
+
+// Warp the history texture, and use local high frequency to adjust the final blending factor. 
+class FTemporalHighFrequencyRejectMapCS : public FGlobalShader
+{
+public:
+	DECLARE_GLOBAL_SHADER(FTemporalHighFrequencyRejectMapCS);
+	SHADER_USE_PARAMETER_STRUCT(FTemporalHighFrequencyRejectMapCS, FGlobalShader);
+
+	BEGIN_SHADER_PARAMETER_STRUCT(FParameters, )
+		SHADER_PARAMETER(float, HighFrequencyCutoffDeltaE)
+		SHADER_PARAMETER_STRUCT_INCLUDE(FDenoisingCommonParameters, DenoisingCommonParameters)
+		SHADER_PARAMETER_RDG_TEXTURE_SRV(Texture2D, SourceTexture)
+		SHADER_PARAMETER_RDG_TEXTURE_SRV(Texture2D, TargetTexture)
+		SHADER_PARAMETER_RDG_TEXTURE_SRV(Texture2D, PixelOffsetTexture)
+		SHADER_PARAMETER_RDG_TEXTURE_UAV(RWTexture2D, OutputTexture)
+		END_SHADER_PARAMETER_STRUCT()
+
+		static bool ShouldCompilePermutation(const FGlobalShaderPermutationParameters& Parameters)
+	{
+		return ShouldCompilePathTracingDenoiserShadersForProject(Parameters.Platform);
+	}
+
+	static void ModifyCompilationEnvironment(const FShaderPermutationParameters&, FShaderCompilerEnvironment& OutEnvironment)
+	{
+		OutEnvironment.SetDefine(TEXT("TEMPORAL_HIGHFREQ_REJECT"), 1);
+		OutEnvironment.SetDefine(TEXT("K_NUM_OF_TEXTURES_PER_PASS"), kNumberOfTexturesPerPass);
+	}
+};
+
+IMPLEMENT_GLOBAL_SHADER(FTemporalHighFrequencyRejectMapCS, "/Engine/Private/PathTracing/PathTracingSpatialTemporalDenoising.usf", "TemporalHighFrequencyRejectCS", SF_Compute);
+
+enum class EFeatureFusionCategory : uint32
+{
+	SOURCE,
+	TARGET,
+	MAX
+};
+
+static constexpr uint32 kNumOfFeatureFusionCategory = static_cast<uint32>(EFeatureFusionCategory::MAX);
+
+class FTemporalFeatureFusionCS : public FGlobalShader
+{
+public:
+	DECLARE_GLOBAL_SHADER(FTemporalFeatureFusionCS)
+	SHADER_USE_PARAMETER_STRUCT(FTemporalFeatureFusionCS, FGlobalShader)
+
+		BEGIN_SHADER_PARAMETER_STRUCT(FParameters, )
+		SHADER_PARAMETER_RDG_TEXTURE_SRV_ARRAY(Texture2D, AlbedoTexture,	[kNumOfFeatureFusionCategory])
+		SHADER_PARAMETER_RDG_TEXTURE_SRV_ARRAY(Texture2D, NormalTexture,	[kNumOfFeatureFusionCategory])
+		SHADER_PARAMETER_RDG_TEXTURE_SRV_ARRAY(Texture2D, RadianceTexture,	[kNumOfFeatureFusionCategory])
+		SHADER_PARAMETER_RDG_BUFFER_SRV_ARRAY(StructuredBuffer<FPixelMaterialLightingFingerprint>, VarianceMap,[kNumOfFeatureFusionCategory])
+		SHADER_PARAMETER_RDG_TEXTURE_SRV(Texture2D, LastDenoisedRadiance)
+
+		SHADER_PARAMETER_SAMPLER(SamplerState, SharedTextureSampler)
+		SHADER_PARAMETER_STRUCT(FScreenPassTextureViewportParameters, TargetViewport)
+
+		SHADER_PARAMETER_RDG_TEXTURE_UAV_ARRAY(RWTexture2D<float4>, OutputTexture,	[kNumOfFeatureFusionCategory])
+	END_SHADER_PARAMETER_STRUCT()
+
+	static bool ShouldCompilePermutation(const FGlobalShaderPermutationParameters& Parameters)
+	{
+		return ShouldCompilePathTracingDenoiserShadersForProject(Parameters.Platform);
+	}
+
+	static void ModifyCompilationEnvironment(const FShaderPermutationParameters&, FShaderCompilerEnvironment& OutEnvironment)
+	{
+		OutEnvironment.SetDefine(TEXT("TEMPORAL_FEATURE_FUSION"), 1);
+	}
+};
+
+IMPLEMENT_GLOBAL_SHADER(FTemporalFeatureFusionCS, "/Engine/Private/PathTracing/PathTracingSpatialTemporalDenoising.usf", "TemporalFeatureFusionCS", SF_Compute);
+
+
+class FTemporalResolveCS : public FGlobalShader
+{
+public:
+	DECLARE_GLOBAL_SHADER(FTemporalResolveCS);
+	SHADER_USE_PARAMETER_STRUCT(FTemporalResolveCS, FGlobalShader);
+
+	BEGIN_SHADER_PARAMETER_STRUCT(FParameters, )
+		SHADER_PARAMETER(float, Kappa)
+		SHADER_PARAMETER(float, Eta)
+		SHADER_PARAMETER(float, Alpha)
+		SHADER_PARAMETER(float, HighFrequencyCutoffDeltaE)
+		SHADER_PARAMETER_STRUCT_INCLUDE(FDenoisingCommonParameters, DenoisingCommonParameters)
+		SHADER_PARAMETER_RDG_TEXTURE_SRV(Texture2D, SourceTexture)
+		SHADER_PARAMETER_RDG_TEXTURE_SRV(Texture2D, TargetTexture)
+		SHADER_PARAMETER_RDG_TEXTURE_SRV(Texture2D, PixelOffsetTexture)
+		SHADER_PARAMETER_RDG_TEXTURE_SRV(Texture2D, HighFrequencyRejectMap)
+		SHADER_PARAMETER_RDG_TEXTURE_UAV(RWTexture2D, OutputTexture)
+		END_SHADER_PARAMETER_STRUCT()
+
+		static bool ShouldCompilePermutation(const FGlobalShaderPermutationParameters& Parameters)
+	{
+		return ShouldCompilePathTracingDenoiserShadersForProject(Parameters.Platform);
+	}
+
+	static void ModifyCompilationEnvironment(const FShaderPermutationParameters&, FShaderCompilerEnvironment& OutEnvironment)
+	{
+		OutEnvironment.SetDefine(TEXT("TEMPORAL_REPROJECTION_RESOLVE"), 1);
+		OutEnvironment.SetDefine(TEXT("K_NUM_OF_TEXTURES_PER_PASS"), kNumberOfTexturesPerPass);
+		OutEnvironment.CompilerFlags.Add(CFLAG_AllowTypedUAVLoads);
+	}
+};
+
+IMPLEMENT_GLOBAL_SHADER(FTemporalResolveCS, "/Engine/Private/PathTracing/PathTracingSpatialTemporalDenoising.usf", "TemporalResolveCS", SF_Compute);
+
+// Add Spatial denoiser
+class FSpatialDenoiserCS : public FGlobalShader
+{
+public:
+	DECLARE_GLOBAL_SHADER(FSpatialDenoiserCS);
+	SHADER_USE_PARAMETER_STRUCT(FSpatialDenoiserCS, FGlobalShader);
+
+	BEGIN_SHADER_PARAMETER_STRUCT(FParameters, )
+		SHADER_PARAMETER_RDG_TEXTURE_SRV(Texture2D, InputTexture)
+		SHADER_PARAMETER_RDG_TEXTURE_SRV(Texture2D, InputNormal)
+		SHADER_PARAMETER_RDG_TEXTURE_SRV(Texture2D, InputAlbedo)
+		SHADER_PARAMETER_RDG_TEXTURE_UAV(RWTexture2D, OutputTexture)
+		SHADER_PARAMETER_STRUCT(FScreenPassTextureViewportParameters, TargetViewport)
+	END_SHADER_PARAMETER_STRUCT()
+
+		static bool ShouldCompilePermutation(const FGlobalShaderPermutationParameters& Parameters)
+	{
+		return ShouldCompilePathTracingDenoiserShadersForProject(Parameters.Platform);
+	}
+
+	static void ModifyCompilationEnvironment(const FShaderPermutationParameters&, FShaderCompilerEnvironment& OutEnvironment)
+	{
+		OutEnvironment.SetDefine(TEXT("SPATIAL_DENOISING"), 1);
+	}
+};
+
+IMPLEMENT_GLOBAL_SHADER(FSpatialDenoiserCS, "/Engine/Private/PathTracing/PathTracingSpatialTemporalDenoising.usf", "SpatialDenoiserCS", SF_Compute);
+
+static void PathTracingDenoiserPlugin(FRDGBuilder& GraphBuilder,
+	const FViewInfo& View,
+	int DenoiserMode,
+	FRDGTextureRef InputTexture, 
+	FRDGTextureRef AlbedoTexture,
+	FRDGTextureRef NormalTexture,
+	FRDGTextureRef OutputTexture)
+{
+	FDenoiseTextureParameters* DenoiseParameters = GraphBuilder.AllocParameters<FDenoiseTextureParameters>();
+	DenoiseParameters->InputTexture = InputTexture;
+	DenoiseParameters->InputAlbedo = AlbedoTexture;
+	DenoiseParameters->InputNormal = NormalTexture;
+	DenoiseParameters->OutputTexture = OutputTexture;
+	// Need to read GPU mask outside Pass function, as the value is not refreshed inside the pass
+	GraphBuilder.AddPass(RDG_EVENT_NAME("Path Tracer Denoiser Plugin"), DenoiseParameters, ERDGPassFlags::Readback,
+		[DenoiseParameters, DenoiserMode, GPUMask = View.GPUMask](FRHICommandListImmediate& RHICmdList)
+	{
+		GPathTracingDenoiserFunc(RHICmdList,
+			DenoiseParameters->InputTexture->GetRHI()->GetTexture2D(),
+			DenoiseParameters->InputAlbedo->GetRHI()->GetTexture2D(),
+			DenoiseParameters->InputNormal->GetRHI()->GetTexture2D(),
+			DenoiseParameters->OutputTexture->GetRHI()->GetTexture2D(),
+			GPUMask);
+	}
+	);
+}
+
+class FMotionVectorEstimationContext
+{
+public:
+
+	const TCHAR* PixelOffsetTextureNames[2] = {
+		TEXT("PathTracing.EstimateMotion.PixelOffset.Ping"),
+		TEXT("PathTracing.EstimateMotion.PixelOffset.Pong"),
+	};
+
+	FRDGTextureRef PixelOffsetTextures[2];
+	FScreenPassTextureViewport TargetViewport;
+
+	FRDGTextureDesc TextureWithNMipsDescriptor;
+	FRDGTextureDesc TextureDescriptor;
+
+	FRDGTextureRef SourceTexture;
+	FRDGTextureRef TargetTexture;
+
+	FRDGTextureRef SourceMipTexture;
+	FRDGTextureRef TargetMipTexture;
+	FRHISamplerState* BilinearClampSampler;
+
+	FRDGTextureRef DistanceTextures[kNumberOfTexturesPerPass];
+	FDenoisingCommonParameters DenoisingCommonParameters;
+
+public:
+
+	FMotionVectorEstimationContext(FRDGTextureRef InSourceTexture, 
+		FRDGTextureRef InTargetTexture):
+	SourceTexture(InSourceTexture),
+	TargetTexture(InTargetTexture){}
+
+	bool InitContext(FRDGBuilder& GraphBuilder,
+		const FViewInfo& View)
+	{
+		TargetViewport = FScreenPassTextureViewport(View.ViewRect);
+		uint32 NumOfMips = FMath::Min(7u, 1 + FMath::FloorLog2((uint32)TargetViewport.Extent.GetMin()));
+
+		if (!ensureMsgf(NumOfMips == 7u, TEXT("The image is too small to estimate temporal reprojection NumOfMips(%d)< 7"), NumOfMips))
+		{
+			return false;
+		}
+
+		// Set up common denoising parameters for shader
+		{
+			DenoisingCommonParameters.ViewUniformBuffer = View.ViewUniformBuffer;
+			DenoisingCommonParameters.TargetViewport = GetScreenPassTextureViewportParameters(TargetViewport);
+			DenoisingCommonParameters.SharedTextureSampler = TStaticSamplerState<SF_Bilinear, AM_Clamp, AM_Clamp, AM_Clamp>::GetRHI();
+			DenoisingCommonParameters.PatchCount = GetTemporalAccumulationPatchCount();
+			DenoisingCommonParameters.NumOfMips = NumOfMips;
+		}
+
+		TextureWithNMipsDescriptor = FRDGTextureDesc::Create2D(
+			TargetViewport.Extent,
+			PF_A32B32G32R32F,
+			FClearValueBinding(),
+			TexCreate_RenderTargetable | TexCreate_ShaderResource | TexCreate_UAV,
+			NumOfMips);
+
+		TextureDescriptor = FRDGTextureDesc::Create2D(
+			TargetViewport.Extent,
+			PF_A32B32G32R32F,
+			FClearValueBinding(),
+			TexCreate_RenderTargetable | TexCreate_ShaderResource | TexCreate_UAV);
+
+		
+		PixelOffsetTextures[0] = GraphBuilder.CreateTexture(TextureWithNMipsDescriptor, PixelOffsetTextureNames[0]);
+		PixelOffsetTextures[1] = GraphBuilder.CreateTexture(TextureWithNMipsDescriptor, PixelOffsetTextureNames[1]);
+
+		FRDGTextureDesc SourceDesc = SourceTexture->Desc;
+		FRDGTextureDesc TargetDesc = TargetTexture->Desc;
+
+		// Copy and create the mipmap for both source and target texture
+		SourceMipTexture = GraphBuilder.CreateTexture(TextureWithNMipsDescriptor, TEXT("PathTracing.EstimateMotion.Source"));
+		TargetMipTexture = GraphBuilder.CreateTexture(TextureWithNMipsDescriptor, TEXT("PathTracing.EstimateMotion.Target"));
+		BilinearClampSampler = TStaticSamplerState<SF_Bilinear, AM_Clamp, AM_Clamp, AM_Clamp>::GetRHI();
+
+		{
+			FRHICopyTextureInfo CopyInfo;
+			CopyInfo.Size.X = FMath::Min(SourceDesc.Extent.X, TargetDesc.Extent.X);
+			CopyInfo.Size.Y = FMath::Min(SourceDesc.Extent.Y, TargetDesc.Extent.Y);
+			CopyInfo.Size.Z = 1;
+			CopyInfo.NumMips = FMath::Min(SourceDesc.NumMips, TargetDesc.NumMips);
+
+			AddCopyTexturePass(GraphBuilder, SourceTexture, SourceMipTexture, CopyInfo);
+			FGenerateMips::Execute(GraphBuilder, SourceMipTexture, BilinearClampSampler);
+
+			AddCopyTexturePass(GraphBuilder, TargetTexture, TargetMipTexture, CopyInfo);
+			FGenerateMips::Execute(GraphBuilder, TargetMipTexture, BilinearClampSampler);
+		}
+
+		for (int TextureId = 0; TextureId < kNumberOfTexturesPerPass; ++TextureId)
+		{
+			DistanceTextures[TextureId] = GraphBuilder.CreateTexture(TextureWithNMipsDescriptor, DistanceTextureNames[TextureId]);
+		}
+
+		return true;
+	}
+
+	uint32 GetPatchCount() const
+	{
+		return DenoisingCommonParameters.PatchCount;
+	}
+
+	uint32 GetNumOfMips() const
+	{
+		return DenoisingCommonParameters.NumOfMips;
+	}
+
+	void UpdatePatchId(uint32 PatchId)
+	{
+		DenoisingCommonParameters.PatchId = PatchId;
+	}
+
+	void UpdateMipLevel(uint32 MipLevel)
+	{
+		DenoisingCommonParameters.MipLevel = MipLevel;
+	}
+
+	uint32 GetMipLevel() const
+	{
+		return DenoisingCommonParameters.MipLevel;
+	}
+
+	FIntVector GetAlignGroupCount() const
+	{
+		return FComputeShaderUtils::GetGroupCount(
+			FIntVector(TargetViewport.Extent.X, TargetViewport.Extent.Y, kNumberOfTexturesPerPass),
+			FIntVector(kThreadSize, kThreadSize, 1));
+	}
+	
+};
+
+static void AlignTexture(FRDGBuilder& GraphBuilder, 
+	const FViewInfo& View,
+	const FMotionVectorEstimationContext& Context,
+	FRDGTextureSRVDesc PixelOffsetSRVDesc,
+	FRDGTextureRef SourceMipTexture, 
+	FRDGTextureRef TargetMipTexture)
+{
+
+	uint32 MipLevel = Context.GetMipLevel();
+
+	typedef FTemporalReprojectionAlignCS SHADER;
+
+	SHADER::FParameters* PassParameters = GraphBuilder.AllocParameters<SHADER::FParameters>();
+	PassParameters->PixelOffsetTexture = GraphBuilder.CreateSRV(PixelOffsetSRVDesc);
+	PassParameters->SourceTexture = GraphBuilder.CreateSRV(FRDGTextureSRVDesc::CreateForMipLevel(SourceMipTexture, MipLevel));
+	PassParameters->TargetTexture = GraphBuilder.CreateSRV(FRDGTextureSRVDesc::CreateForMipLevel(TargetMipTexture, MipLevel));
+	PassParameters->DenoisingCommonParameters = Context.DenoisingCommonParameters;
+
+	for (int TextureId = 0; TextureId < kNumberOfTexturesPerPass; ++TextureId)
+	{
+		PassParameters->RWDistanceTextures[TextureId] = GraphBuilder.CreateUAV(FRDGTextureUAVDesc(Context.DistanceTextures[TextureId], MipLevel));
+	}
+
+	SHADER::EDistanceMetrics DistanceMetrics = SHADER::GetDistanceMetrics();
+	SHADER::FPermutationDomain ComputeShaderPermutationVector;
+	ComputeShaderPermutationVector.Set<SHADER::FDistanceMetrics>(DistanceMetrics);
+
+	TShaderMapRef<SHADER> ComputeShader(View.ShaderMap,ComputeShaderPermutationVector);
+
+	FComputeShaderUtils::AddPass(
+		GraphBuilder,
+		RDG_EVENT_NAME("PathTracing::Denoising::Align %dx%d (%s)",
+			Context.TargetViewport.Extent.X,
+			Context.TargetViewport.Extent.Y,
+			SHADER::GetEventName(DistanceMetrics)),
+		ComputeShader,
+		PassParameters,
+		Context.GetAlignGroupCount());
+}
+
+// Determine the image blur size based on the current mip level and image size.
+// TODO: scale blur size based on viewport size
+static int32 GetBlurSize(int32 MipLevel, FIntPoint ViewportExtent)
+{
+	int BlurSize = 1;	
+
+	switch (MipLevel)
+	{
+	case 0:
+		BlurSize = 3; break;
+	case 2:
+		BlurSize = 2; break;
+	default:
+		BlurSize = 1; break;
+	};
+
+	return BlurSize;
+}
+
+static void BlurDistanceMetrics(FRDGBuilder& GraphBuilder,
+	const FViewInfo& View,
+	const FMotionVectorEstimationContext& Context)
+{	
+	struct FReprojectionBulrPassInfo
+	{
+		FReprojectionBulrPassInfo(const TCHAR* InName, FRDGTextureRef  InInput, FRDGTextureRef InOutput)
+			:Name(InName), Input(InInput), Output(InOutput)
+		{}
+
+		const TCHAR* Name;
+		FRDGTextureRef Input;
+		FRDGTextureRef Output;
+	};
+
+	FRDGTextureRef TempBuffer = GraphBuilder.CreateTexture(Context.TextureWithNMipsDescriptor, TEXT("PathTracing.EstimateMotion.TempBuffer"));
+	uint32 MipLevel = Context.GetMipLevel();
+	int BlurSize = GetBlurSize(MipLevel, Context.TargetViewport.Extent);
+
+	for (int TextureId = 0; TextureId < kNumberOfTexturesPerPass; ++TextureId)
+	{
+		const int NumOfBlurPass = 2;
+		const FReprojectionBulrPassInfo BlurPassInfos[NumOfBlurPass] =
+		{
+			{TEXT("PathTracing::Denoising::Blur0"),	Context.DistanceTextures[TextureId],						  TempBuffer},
+			{TEXT("PathTracing::Denoising::Blur1"),							 TempBuffer, Context.DistanceTextures[TextureId]}
+		};
+
+		for (int i = 0; i < BlurSize; ++i) 
+		{
+			for (int PassIndex = 0; PassIndex < NumOfBlurPass; ++PassIndex)
+			{
+
+				FReprojectionBulrPassInfo PassInfo = BlurPassInfos[PassIndex];
+				FRDGTextureSRVDesc InputSRVDesc = FRDGTextureSRVDesc::CreateForMipLevel(PassInfo.Input, MipLevel);
+				FRDGTextureUAVDesc OutputUAVDesc(PassInfo.Output, MipLevel);
+
+				typedef FTemporalReprojectionBlurCS SHADER;
+				SHADER::FParameters* PassParameters = GraphBuilder.AllocParameters<SHADER::FParameters>();
+				PassParameters->InputTexture = GraphBuilder.CreateSRV(InputSRVDesc);
+				PassParameters->OutputTexture = GraphBuilder.CreateUAV(OutputUAVDesc);
+				PassParameters->DenoisingCommonParameters = Context.DenoisingCommonParameters;
+
+				SHADER::FPermutationDomain ComputeShaderPermutationVector;
+				ComputeShaderPermutationVector.Set<SHADER::FDimensionDirectCopy>(false);
+
+				TShaderMapRef<SHADER> ComputeShader(View.ShaderMap, ComputeShaderPermutationVector);
+
+				FComputeShaderUtils::AddPass(
+					GraphBuilder,
+					RDG_EVENT_NAME("%s %dx%d Mip=%d (DistPatch %d/%d)",
+						BlurPassInfos[PassIndex].Name,
+						Context.TargetViewport.Extent.X,
+						Context.TargetViewport.Extent.Y,
+						MipLevel,
+						TextureId + 1,
+						kNumberOfTexturesPerPass),
+					ComputeShader,
+					PassParameters,
+					FComputeShaderUtils::GetGroupCount(Context.TargetViewport.Extent, 8));
+			}
+		}
+	}
+}
+
+static void MergeDistanceMetrics(FRDGBuilder& GraphBuilder,
+	const FViewInfo& View,
+	const FMotionVectorEstimationContext& Context,
+	const FRDGTextureSRVDesc& LastPixelOffsetSRVDesc,
+	FRDGTextureRef TargetPixelOffsetTexture)
+{
+	uint32 MipLevel = Context.GetMipLevel();
+
+	typedef FTemporalReprojectionMergeCS SHADER;
+	SHADER::FParameters* PassParameters = GraphBuilder.AllocParameters<SHADER::FParameters>();
+	for (int TextureId = 0; TextureId < kNumberOfTexturesPerPass; ++TextureId)
+	{
+		PassParameters->DistanceTextures[TextureId] = GraphBuilder.CreateSRV(FRDGTextureSRVDesc::CreateForMipLevel(Context.DistanceTextures[TextureId], MipLevel));
+	}
+	PassParameters->PixelOffsetTexture = GraphBuilder.CreateSRV(LastPixelOffsetSRVDesc);
+	PassParameters->RWPixelOffsetTexture = GraphBuilder.CreateUAV(FRDGTextureUAVDesc(TargetPixelOffsetTexture, MipLevel));
+	PassParameters->DenoisingCommonParameters = Context.DenoisingCommonParameters;
+
+	bool bShouldUseTotalVariation = ShouldUseTotalVariation(MipLevel);
+	bool bShouldEnableSubpixelOffset = ShouldEnableSubpixelOffset(MipLevel);
+
+	SHADER::FPermutationDomain ComputeShaderPermutationVector;
+	ComputeShaderPermutationVector.Set<SHADER::FSubpixelOffset>(bShouldEnableSubpixelOffset);
+	ComputeShaderPermutationVector.Set<SHADER::FTotalVariation>(bShouldUseTotalVariation);
+
+	TShaderMapRef<SHADER> ComputeShader(View.ShaderMap, ComputeShaderPermutationVector);
+	FComputeShaderUtils::AddPass(
+		GraphBuilder,
+		RDG_EVENT_NAME("PathTracing::Denoising::Merge %dx%d",
+			Context.TargetViewport.Extent.X,
+			Context.TargetViewport.Extent.Y),
+		ComputeShader,
+		PassParameters,
+		FComputeShaderUtils::GetGroupCount(Context.TargetViewport.Extent, 8));
+}
+
+// Estimate the motion vector (pixel correspondence offset)
+static FRDGTextureRef EstimateMotionVector(FRDGBuilder& GraphBuilder,
+	const FViewInfo& View,
+	FMotionVectorEstimationContext& Context)
+{
+	int PixelOffsetIndex = 0;
+	int NumOfMips = Context.GetNumOfMips();
+
+	for (int MipLevel = NumOfMips - 1; MipLevel >= 0; MipLevel -= kMipDiffDelta)
+	{
+		Context.UpdateMipLevel(MipLevel);
+
+		for (int Pass = 0; Pass < kNumberOfPasses; ++Pass)
+		{
+			const bool UseBlackDummpy = ((MipLevel == (NumOfMips - 1)) && Pass == 0);
+
+			// Write the calibration subpixel offset to the source so that we can remove it.
+			FRDGTextureSRVDesc SelfLastPixelOffsetSRVDesk = FRDGTextureSRVDesc::CreateForMipLevel(
+				GraphBuilder.RegisterExternalTexture(GSystemTextures.BlackDummy, Context.PixelOffsetTextureNames[1]),
+				0);
+
+			FRDGTextureRef SelfTargetPixelOffsetTexture = Context.PixelOffsetTextures[PixelOffsetIndex % 2];
+
+			FRDGTextureSRVDesc PixelOffsetSRVDesc = FRDGTextureSRVDesc::CreateForMipLevel(
+				UseBlackDummpy ?
+				GraphBuilder.RegisterExternalTexture(GSystemTextures.BlackDummy, Context.PixelOffsetTextureNames[0]) :
+				Context.PixelOffsetTextures[PixelOffsetIndex % 2], UseBlackDummpy ? 0 : (MipLevel + kMipDiffDelta));
+			
+			FRDGTextureRef TargetPixelOffsetTexture = Context.PixelOffsetTextures[(++PixelOffsetIndex) % 2];
+
+			// Calculate the subpixel offset for the source image, and write to the unused texture in the source mipmap
+			if(ShouldRemoveSelfSubpixelOffset(MipLevel))
+			{
+				AlignTexture(GraphBuilder, View, Context,
+					SelfLastPixelOffsetSRVDesk,
+					Context.SourceMipTexture,
+					Context.SourceMipTexture);
+
+				BlurDistanceMetrics(GraphBuilder, View,
+					Context);
+
+				// Merge: Find the smallest distance within [-2,2]^2
+				MergeDistanceMetrics(GraphBuilder, View,
+					Context,
+					SelfLastPixelOffsetSRVDesk,
+					SelfTargetPixelOffsetTexture);
+			}
+
+			// Calculate the subpixel offset between the source and the target
+			{
+				AlignTexture(GraphBuilder, View, Context,
+					PixelOffsetSRVDesc,
+					Context.SourceMipTexture,
+					Context.TargetMipTexture);
+
+				BlurDistanceMetrics(GraphBuilder, View,
+					Context);
+
+				// Merge: Find the smallest distance within [-2,2]^2
+				MergeDistanceMetrics(GraphBuilder, View,
+					Context,
+					PixelOffsetSRVDesc,
+					TargetPixelOffsetTexture);
+			}
+			
+			// Subtract the distance from source to target
+			if (ShouldRemoveSelfSubpixelOffset(MipLevel))
+			{
+				typedef FMotionVectorSubtractCS SHADER;
+				SHADER::FParameters* PassParameters = GraphBuilder.AllocParameters<SHADER::FParameters>();
+				PassParameters->Minuend = GraphBuilder.CreateUAV(FRDGTextureUAVDesc(TargetPixelOffsetTexture, MipLevel));
+				PassParameters->Subtrahend = GraphBuilder.CreateSRV(FRDGTextureSRVDesc::CreateForMipLevel(SelfTargetPixelOffsetTexture, MipLevel));
+				PassParameters->DenoisingCommonParameters = Context.DenoisingCommonParameters;
+
+				TShaderMapRef<SHADER> ComputeShader(View.ShaderMap);
+				FComputeShaderUtils::AddPass(
+					GraphBuilder,
+					RDG_EVENT_NAME("PathTracing::Denoising::MotionVectorDiff %dx%d (Mip=%d)",
+						Context.TargetViewport.Extent.X,
+						Context.TargetViewport.Extent.Y,
+						MipLevel),
+					ComputeShader,
+					PassParameters,
+					FComputeShaderUtils::GetGroupCount(Context.TargetViewport.Extent, 8));
+			}
+
+		}
+	}
+
+	return Context.PixelOffsetTextures[PixelOffsetIndex % 2];
+}
+
+// combine albedo, normal, and radiance into a single feature texture
+static void FuseTemporalFeature(FRDGBuilder& GraphBuilder,
+	const FViewInfo& View,
+	FPathTracingSpatialTemporalDenoisingContext& DenoisingContext, 
+	FRDGTextureRef& OutSourceTexture, 
+	FRDGTextureRef& OutTargetTexture)
+{
+	const FScreenPassTextureViewport TargetViewport(View.ViewRect);
+	const FScreenPassTextureViewportParameters TargetViewportParameters = GetScreenPassTextureViewportParameters(TargetViewport);
+
+	const FRDGTextureDesc TextureDescriptor = FRDGTextureDesc::Create2D(
+		TargetViewport.Extent,
+		PF_A32B32G32R32F,
+		FClearValueBinding(),
+		TexCreate_RenderTargetable | TexCreate_ShaderResource | TexCreate_UAV);
+
+	OutSourceTexture = GraphBuilder.CreateTexture(TextureDescriptor, TEXT("PathTracing.Denoising.Feature.Source"));
+	OutTargetTexture = GraphBuilder.CreateTexture(TextureDescriptor, TEXT("PathTracing.Denoising.Feature.Target"));
+
+	typedef FTemporalFeatureFusionCS SHADER;
+	SHADER::FParameters* PassParameters = GraphBuilder.AllocParameters<SHADER::FParameters>();
+	{
+		PassParameters->AlbedoTexture[0] = GraphBuilder.CreateSRV(FRDGTextureSRVDesc::Create(DenoisingContext.LastAlbedoTexture));
+		PassParameters->AlbedoTexture[1] = GraphBuilder.CreateSRV(FRDGTextureSRVDesc::Create(DenoisingContext.AlbedoTexture));
+		PassParameters->NormalTexture[0] = GraphBuilder.CreateSRV(FRDGTextureSRVDesc::Create(DenoisingContext.LastNormalTexture));
+		PassParameters->NormalTexture[1] = GraphBuilder.CreateSRV(FRDGTextureSRVDesc::Create(DenoisingContext.NormalTexture));
+		PassParameters->RadianceTexture[0] = GraphBuilder.CreateSRV(FRDGTextureSRVDesc::Create(DenoisingContext.LastRadianceTexture));
+		PassParameters->RadianceTexture[1] = GraphBuilder.CreateSRV(FRDGTextureSRVDesc::Create(DenoisingContext.RadianceTexture));
+		PassParameters->VarianceMap[0] = GraphBuilder.CreateSRV(DenoisingContext.LastVarianceBuffer ? 
+											DenoisingContext.LastVarianceBuffer : DenoisingContext.VarianceBuffer,EPixelFormat::PF_R32_FLOAT);
+		PassParameters->VarianceMap[1] = GraphBuilder.CreateSRV(DenoisingContext.VarianceBuffer, EPixelFormat::PF_R32_FLOAT);
+		PassParameters->OutputTexture[0] = GraphBuilder.CreateUAV(FRDGTextureUAVDesc(OutSourceTexture));
+		PassParameters->OutputTexture[1] = GraphBuilder.CreateUAV(FRDGTextureUAVDesc(OutTargetTexture));
+
+		PassParameters->LastDenoisedRadiance = GraphBuilder.CreateSRV(FRDGTextureSRVDesc::Create(DenoisingContext.LastDenoisedRadianceTexture));
+
+		PassParameters->SharedTextureSampler = TStaticSamplerState<SF_Bilinear, AM_Clamp, AM_Clamp, AM_Clamp>::GetRHI();
+		PassParameters->TargetViewport = TargetViewportParameters;
+	}
+
+	TShaderMapRef<SHADER> ComputeShader(View.ShaderMap);
+	FComputeShaderUtils::AddPass(
+		GraphBuilder,
+		RDG_EVENT_NAME("PathTracing::Denoising::FeatureFusion %dx%d",
+			TargetViewport.Extent.X,
+			TargetViewport.Extent.Y),
+		ComputeShader,
+		PassParameters,
+		FComputeShaderUtils::GetGroupCount(TargetViewport.Extent, 8));
+
+}
+
+static void SelectMotionSourceAndTargetTextures(
+	FRDGBuilder& GraphBuilder,
+	const FViewInfo& View,
+	FRDGTextureRef InSourceTexture,
+	FRDGTextureRef InTargetTexture,
+	FPathTracingSpatialTemporalDenoisingContext& DenoisingContext,
+	FRDGTextureRef& MotionEstimationSourceTexture,
+	FRDGTextureRef& MotionEstimationTargetTexture)
+{
+	int32 MotionSource = CVarPathTracingTemporalDenoiserSource.GetValueOnAnyThread();
+
+	switch (MotionSource)
+	{
+	case 3:
+		MotionEstimationSourceTexture = DenoisingContext.LastRadianceTexture;
+		MotionEstimationTargetTexture = DenoisingContext.RadianceTexture;
+		break;
+	case 2:
+		MotionEstimationSourceTexture = DenoisingContext.LastAlbedoTexture;
+		MotionEstimationTargetTexture = DenoisingContext.AlbedoTexture;
+		break;
+	case 1:
+		MotionEstimationSourceTexture = DenoisingContext.LastNormalTexture;
+		MotionEstimationTargetTexture = DenoisingContext.NormalTexture;
+		break;
+	case 0:
+		MotionEstimationSourceTexture = InSourceTexture;
+		MotionEstimationTargetTexture = InTargetTexture;
+		break;
+	default:
+		FuseTemporalFeature(GraphBuilder, View, DenoisingContext, MotionEstimationSourceTexture, MotionEstimationTargetTexture);
+		break;
+	}
+
+}
+
+// Estimate the motion vector from source to target
+static FRDGTextureRef TemporalReprojectionWithoutMotionVector(FRDGBuilder& GraphBuilder,
+	const FViewInfo& View,
+	FRDGTextureRef InSourceTexture,
+	FRDGTextureRef InTargetTexture,
+	FPathTracingSpatialTemporalDenoisingContext& DenoisingContext)
+{
+
+	FRDGTextureRef MotionEstimationSourceTexture = nullptr;
+	FRDGTextureRef MotionEstimationTargetTexture = nullptr;
+
+	SelectMotionSourceAndTargetTextures(
+		GraphBuilder, View, InSourceTexture, InTargetTexture, DenoisingContext, MotionEstimationSourceTexture, MotionEstimationTargetTexture);
+
+	// TODO: when the source and target texture size is different
+	if (!(MotionEstimationSourceTexture->Desc.Extent == MotionEstimationTargetTexture->Desc.Extent))
+	{
+		UE_LOG(LogPathTracingDenoising, Log, TEXT("TODO: Add motion vector estimation when source and dest dimensions do not match"));
+		return nullptr;
+	}
+
+	FMotionVectorEstimationContext EstimationContext(MotionEstimationSourceTexture, MotionEstimationTargetTexture);
+
+	if (!EstimationContext.InitContext(GraphBuilder, View))
+	{
+		return nullptr;
+	}
+
+	FRDGTextureRef  HighFrequencyRejectMap = GraphBuilder.CreateTexture(EstimationContext.TextureDescriptor, TEXT("PathTracing.EstimateMotion.HighFrequencyRejectMap"));
+	FRDGTextureRef  FinalAccumulation = GraphBuilder.CreateTexture(EstimationContext.TextureDescriptor, TEXT("PathTracing.EstimateMotion.FinalAccumulation"));
+	int PatchCount = EstimationContext.GetPatchCount();
+	float HighFrequencyCutoffDeltaE = GetHighFrequencyCutoffDeltaE();
+
+	FRDGTextureRef  TempAccumulation = nullptr;
+	if (PatchCount > 1)
+	{
+		TempAccumulation = GraphBuilder.CreateTexture(EstimationContext.TextureDescriptor, TEXT("PathTracing.EstimateMotion.TempAccumulation"));
+	}
+
+	// Perform accumulation on Miplevel 0
+	const uint32 AccumulationMipLevel = 0;
+	float Kappa, Eta, Alpha;
+	GetBlendingFactor(Kappa, Eta, Alpha);
+
+	for (int PatchId = 0; PatchId < PatchCount; ++PatchId)
+	{
+		EstimationContext.UpdatePatchId(PatchId);
+
+		FRDGTextureRef  PixelOffsetTexture = EstimateMotionVector(
+			GraphBuilder,
+			View,
+			EstimationContext);
+
+		DenoisingContext.MotionVector = PixelOffsetTexture; // Used for debugging
+
+		// Final resolve, accumulate the temporal information based on the pixel offset texture
+		FRDGTextureRef BlendSourceTexture = InSourceTexture;
+		FRDGTextureRef BlendTargetTexture = (PatchId == 0) ? InTargetTexture : TempAccumulation;
+		FRDGTextureRef BlendFinalTexture = (PatchId + 1 == PatchCount)? FinalAccumulation: TempAccumulation;
+		
+		EstimationContext.UpdateMipLevel(AccumulationMipLevel);
+
+		// Warp the source image so that we can adjust the blending weight based on high frequency information.
+		// E.g., if the source and target is the albedo. We can use visual perception difference of 1dE, 2dE or 10dE
+		{
+			typedef FTemporalHighFrequencyRejectMapCS SHADER;
+			SHADER::FParameters* PassParameters = GraphBuilder.AllocParameters<SHADER::FParameters>();
+			PassParameters->PixelOffsetTexture = GraphBuilder.CreateSRV(FRDGTextureSRVDesc::CreateForMipLevel(PixelOffsetTexture, AccumulationMipLevel));
+			PassParameters->SourceTexture = GraphBuilder.CreateSRV(FRDGTextureSRVDesc::CreateForMipLevel(BlendSourceTexture, AccumulationMipLevel));
+			PassParameters->TargetTexture = GraphBuilder.CreateSRV(FRDGTextureSRVDesc::CreateForMipLevel(InTargetTexture, AccumulationMipLevel));
+			PassParameters->OutputTexture = GraphBuilder.CreateUAV(FRDGTextureUAVDesc(HighFrequencyRejectMap));	// the difference of the optimal shift.
+			PassParameters->DenoisingCommonParameters = EstimationContext.DenoisingCommonParameters;
+			PassParameters->HighFrequencyCutoffDeltaE = HighFrequencyCutoffDeltaE;
+			TShaderMapRef<SHADER> ComputeShader(View.ShaderMap);
+			FComputeShaderUtils::AddPass(
+				GraphBuilder,
+				RDG_EVENT_NAME("PathTracing::Denoising::RejectionMap %dx%d",
+					EstimationContext.TargetViewport.Extent.X,
+					EstimationContext.TargetViewport.Extent.Y),
+				ComputeShader,
+				PassParameters,
+				FComputeShaderUtils::GetGroupCount(EstimationContext.TargetViewport.Extent, 8));
+		}
+
+		{
+			typedef FTemporalResolveCS SHADER;
+			SHADER::FParameters* PassParameters = GraphBuilder.AllocParameters<SHADER::FParameters>();
+			PassParameters->PixelOffsetTexture = GraphBuilder.CreateSRV(FRDGTextureSRVDesc::CreateForMipLevel(PixelOffsetTexture, AccumulationMipLevel));
+			PassParameters->SourceTexture = GraphBuilder.CreateSRV(FRDGTextureSRVDesc::CreateForMipLevel(BlendSourceTexture, AccumulationMipLevel));
+			PassParameters->TargetTexture = GraphBuilder.CreateSRV(FRDGTextureSRVDesc::CreateForMipLevel(BlendTargetTexture, AccumulationMipLevel));
+			PassParameters->HighFrequencyRejectMap = GraphBuilder.CreateSRV(FRDGTextureSRVDesc::CreateForMipLevel(HighFrequencyRejectMap, AccumulationMipLevel));
+			PassParameters->OutputTexture = GraphBuilder.CreateUAV(FRDGTextureUAVDesc(BlendFinalTexture));	// the difference of the optimal shift.
+			PassParameters->DenoisingCommonParameters = EstimationContext.DenoisingCommonParameters;
+
+			PassParameters->HighFrequencyCutoffDeltaE = HighFrequencyCutoffDeltaE;
+			PassParameters->Alpha = Alpha;
+			PassParameters->Kappa = Kappa;
+			PassParameters->Eta = Eta;
+
+			TShaderMapRef<SHADER> ComputeShader(View.ShaderMap);
+			FComputeShaderUtils::AddPass(
+				GraphBuilder,
+				RDG_EVENT_NAME("PathTracing::Denoising::TemporalResolve %dx%d",
+					EstimationContext.TargetViewport.Extent.X,
+					EstimationContext.TargetViewport.Extent.Y),
+				ComputeShader,
+				PassParameters,
+				FComputeShaderUtils::GetGroupCount(EstimationContext.TargetViewport.Extent, 8));
+		}
+	}
+
+	return FinalAccumulation;
+}
+
+struct FPixelMaterialLightingFingerprint
+{
+	FVector4 Mean;
+	FVector4 Var;
+};
+
+class FTemporalPrepassCS : public FGlobalShader
+{
+public:
+	DECLARE_GLOBAL_SHADER(FTemporalPrepassCS);
+	SHADER_USE_PARAMETER_STRUCT(FTemporalPrepassCS, FGlobalShader);
+
+	BEGIN_SHADER_PARAMETER_STRUCT(FParameters, )
+		SHADER_PARAMETER_RDG_TEXTURE_SRV(Texture2D, InputTexture)
+		SHADER_PARAMETER_RDG_BUFFER_UAV(RWStructuredBuffer<FPixelMaterialLightingFingerprint>, RWVarianceMap)
+		SHADER_PARAMETER_STRUCT(FScreenPassTextureViewportParameters, TargetViewport)
+		SHADER_PARAMETER(int, Iteration)
+	END_SHADER_PARAMETER_STRUCT()
+	
+	class FPrepassPhase : SHADER_PERMUTATION_BOOL("PREPASS_PHASE");	// 0: initialize, 1: update
+	using FPermutationDomain = TShaderPermutationDomain<FPrepassPhase>;
+
+	static bool ShouldCompilePermutation(const FGlobalShaderPermutationParameters& Parameters)
+	{
+		return ShouldCompilePathTracingDenoiserShadersForProject(Parameters.Platform);
+	}
+
+	static void ModifyCompilationEnvironment(const FShaderPermutationParameters&, FShaderCompilerEnvironment& OutEnvironment)
+	{
+		OutEnvironment.SetDefine(TEXT("TEMPORAL_PREPASS"), 1);
+	}
+};
+
+IMPLEMENT_GLOBAL_SHADER(FTemporalPrepassCS, "/Engine/Private/PathTracing/PathTracingSpatialTemporalDenoising.usf", "TemporalPrepassCS", SF_Compute);
+
+void PathTracingSpatialTemporalDenoisingPrePass(FRDGBuilder& GraphBuilder, const FViewInfo& View,
+	int IterationNumber,
+	FPathTracingSpatialTemporalDenoisingContext& SpatialTemporalDenoisingContext)
+{
+	bool bShouldGenerateVarianceMap = ShouldGenerateVarianceMap();
+	if (bShouldGenerateVarianceMap)
+	{
+		bool bUpdateVarianceMap = (IterationNumber > 0);
+		const FScreenPassTextureViewport TargetViewport(View.ViewRect);
+		const FScreenPassTextureViewportParameters TargetViewportParameters = GetScreenPassTextureViewportParameters(TargetViewport);
+
+		if (!SpatialTemporalDenoisingContext.VarianceBuffer)
+		{
+			SpatialTemporalDenoisingContext.VarianceBuffer = GraphBuilder.CreateBuffer(
+				FRDGBufferDesc::CreateStructuredDesc(sizeof(float) * 8, View.ViewRect.Area()), TEXT("PathTracing.VarianceBuffer"));
+		}
+
+		{
+			typedef FTemporalPrepassCS SHADER;
+			SHADER::FParameters* PassParameters = GraphBuilder.AllocParameters<SHADER::FParameters>();
+			{
+				PassParameters->InputTexture = GraphBuilder.CreateSRV(FRDGTextureSRVDesc::Create(SpatialTemporalDenoisingContext.RadianceTexture));
+				PassParameters->RWVarianceMap = GraphBuilder.CreateUAV(SpatialTemporalDenoisingContext.VarianceBuffer, EPixelFormat::PF_R32_FLOAT);
+				PassParameters->TargetViewport = TargetViewportParameters;
+				PassParameters->Iteration = IterationNumber;
+			}
+
+			SHADER::FPermutationDomain ComputeShaderPermutationVector;
+			ComputeShaderPermutationVector.Set<SHADER::FPrepassPhase>(bUpdateVarianceMap);
+
+			TShaderMapRef<SHADER> ComputeShader(View.ShaderMap, ComputeShaderPermutationVector);
+			FComputeShaderUtils::AddPass(
+				GraphBuilder,
+				RDG_EVENT_NAME("PathTracing::Denoising::Prepass %dx%d",
+					TargetViewport.Extent.X,
+					TargetViewport.Extent.Y),
+				ComputeShader,
+				PassParameters,
+				FComputeShaderUtils::GetGroupCount(TargetViewport.Extent, 8));
+		}
+	}
+}
+
+void PathTracingSpatialTemporalDenoising(FRDGBuilder& GraphBuilder,
+	const FViewInfo& View,
+	int DenoiserMode,
+	FRDGTexture*& SpatialTemporalDenoisedTexture,
+	FPathTracingSpatialTemporalDenoisingContext& SpatialTemporalDenoisingContext)
+{
+
+	RDG_GPU_STAT_SCOPE(GraphBuilder, PathTracingSpatialTemporalDenoising);
+	RDG_EVENT_SCOPE(GraphBuilder, "PathTracingSpatialTemporalDenoising");
+
+	FRDGTextureDesc RadianceTextureDesc = FRDGTextureDesc::Create2D(
+		View.ViewRect.Size(),
+		PF_A32B32G32R32F,
+		FClearValueBinding::None,
+		TexCreate_ShaderResource | TexCreate_UAV);
+
+	const bool bApplySpatialDenoiser = ShouldApplySpatialDenoiser();
+	FRDGTextureRef TargetTexture = SpatialTemporalDenoisingContext.RadianceTexture;
+
+	if (bApplySpatialDenoiser)
+	{
+		TargetTexture = GraphBuilder.CreateTexture(RadianceTextureDesc, TEXT("PathTracer.SpatialDenoiser.Output"));
+		PathTracingDenoiserPlugin(
+			GraphBuilder,
+			View,
+			DenoiserMode,
+			SpatialTemporalDenoisingContext.RadianceTexture,
+			SpatialTemporalDenoisingContext.AlbedoTexture,
+			SpatialTemporalDenoisingContext.NormalTexture,
+			TargetTexture);
+	}
+
+	const bool bApplyTemporalDenoiser = ShouldApplyTemporalDenoiser(SpatialTemporalDenoisingContext);
+	FRDGTextureRef TemporalDenoisedTexture = nullptr;
+
+	if (bApplyTemporalDenoiser)
+	{
+		UE_LOG(LogPathTracingDenoising, Log, TEXT("Using temporal denoising for frame %i"), SpatialTemporalDenoisingContext.FrameIndex);
+
+		// Image space temporal and spatial denoising.
+		// if we have the last available frame, do a fast reprojection to use history for frame stabilization
+		FRDGTextureRef SourceTexture = SpatialTemporalDenoisingContext.LastDenoisedRadianceTexture;
+
+		TemporalDenoisedTexture = TemporalReprojectionWithoutMotionVector(GraphBuilder,
+			View,
+			SourceTexture,
+			TargetTexture,
+			SpatialTemporalDenoisingContext);
+	}
+
+	SpatialTemporalDenoisedTexture = TemporalDenoisedTexture ? TemporalDenoisedTexture : TargetTexture;
+}
+
+class FVisualizePathTracingMotionVectorPS : public FGlobalShader
+{
+public:
+	DECLARE_GLOBAL_SHADER(FVisualizePathTracingMotionVectorPS);
+	SHADER_USE_PARAMETER_STRUCT(FVisualizePathTracingMotionVectorPS, FGlobalShader);
+
+	BEGIN_SHADER_PARAMETER_STRUCT(FParameters, )
+		SHADER_PARAMETER_STRUCT_REF(FViewUniformShaderParameters, ViewUniformBuffer)
+		SHADER_PARAMETER_RDG_UNIFORM_BUFFER(FSceneTextureUniformParameters, SceneTextures)
+		SHADER_PARAMETER_RDG_TEXTURE_SRV(Texture2D<float4>, TemporalDenoisingMotionVector)
+		SHADER_PARAMETER_RDG_TEXTURE_SRV(Texture2D<float4>, RasterMotionVector)
+		SHADER_PARAMETER_RDG_TEXTURE_SRV(Texture2D<float4>, DenoisedTexture)
+		RENDER_TARGET_BINDING_SLOTS()
+	END_SHADER_PARAMETER_STRUCT()
+
+	static bool ShouldCompilePermutation(const FGlobalShaderPermutationParameters& Parameters)
+	{
+		return ShouldCompilePathTracingDenoiserShadersForProject(Parameters.Platform);
+	}
+
+	static void ModifyCompilationEnvironment(const FShaderPermutationParameters&, FShaderCompilerEnvironment& OutEnvironment)
+	{
+		OutEnvironment.SetDefine(TEXT("VISUALIZE_MOTIONVECTOR"), 1);
+	}
+
+};
+
+IMPLEMENT_GLOBAL_SHADER(FVisualizePathTracingMotionVectorPS, "/Engine/Private/PathTracing/PathTracingSpatialTemporalDenoising.usf", "VisualizePathTracingMotionVector", SF_Pixel);
+
+class FVisualizeWarpingPS : public FGlobalShader
+{
+public:
+	DECLARE_GLOBAL_SHADER(FVisualizeWarpingPS);
+	SHADER_USE_PARAMETER_STRUCT(FVisualizeWarpingPS, FGlobalShader);
+
+	BEGIN_SHADER_PARAMETER_STRUCT(FParameters, )
+		SHADER_PARAMETER_STRUCT_REF(FViewUniformShaderParameters, ViewUniformBuffer)
+		SHADER_PARAMETER_RDG_UNIFORM_BUFFER(FSceneTextureUniformParameters, SceneTextures)
+		SHADER_PARAMETER_RDG_TEXTURE_SRV(Texture2D<float4>, TemporalDenoisingMotionVector)
+		SHADER_PARAMETER_RDG_TEXTURE_SRV(Texture2D<float4>, SourceTexture)
+		SHADER_PARAMETER_RDG_TEXTURE_SRV(Texture2D<float4>, DenoisedTexture)
+		SHADER_PARAMETER_RDG_TEXTURE_SRV(Texture2D<float4>, TargetTexture)
+		SHADER_PARAMETER_STRUCT(FScreenPassTextureViewportParameters, TargetViewport)
+		SHADER_PARAMETER_SAMPLER(SamplerState, SharedTextureSampler)
+		RENDER_TARGET_BINDING_SLOTS()
+	END_SHADER_PARAMETER_STRUCT()
+
+		static bool ShouldCompilePermutation(const FGlobalShaderPermutationParameters& Parameters)
+	{
+		return ShouldCompilePathTracingDenoiserShadersForProject(Parameters.Platform);
+	}
+
+	static void ModifyCompilationEnvironment(const FShaderPermutationParameters&, FShaderCompilerEnvironment& OutEnvironment)
+	{
+		OutEnvironment.SetDefine(TEXT("VISUALIZE_WARPING"), 1);
+	}
+
+};
+
+IMPLEMENT_GLOBAL_SHADER(FVisualizeWarpingPS, "/Engine/Private/PathTracing/PathTracingSpatialTemporalDenoising.usf", "FVisualizeWarpingPS", SF_Pixel);
+
+FScreenPassTexture AddVisualizePathTracingDenoisingPass(FRDGBuilder& GraphBuilder, const FViewInfo& View, const FVisualizePathTracingDenoisingInputs& Inputs)
+{
+	
+	if (ShouldVisualizePathTracingVelocityState(Inputs.DenoisingContext))
+	{
+		typedef FVisualizePathTracingMotionVectorPS SHADER;
+		SHADER::FParameters* PassParameters = GraphBuilder.AllocParameters<SHADER::FParameters>();
+		{
+			PassParameters->ViewUniformBuffer = View.ViewUniformBuffer;
+			PassParameters->TemporalDenoisingMotionVector = GraphBuilder.CreateSRV(FRDGTextureSRVDesc::Create(Inputs.DenoisingContext.MotionVector));
+			PassParameters->SceneTextures = Inputs.SceneTexturesUniformBuffer;
+			PassParameters->DenoisedTexture = GraphBuilder.CreateSRV(FRDGTextureSRVDesc::Create(Inputs.DenoisedTexture));
+			PassParameters->RenderTargets[0] = FRenderTargetBinding(Inputs.SceneColor, ERenderTargetLoadAction::ELoad);
+		}
+
+		TShaderMapRef<SHADER> PixelShader(View.ShaderMap);
+		AddDrawScreenPass(
+			GraphBuilder,
+			RDG_EVENT_NAME("Visualize Motion Vector (%d x %d)", View.ViewRect.Size().X, View.ViewRect.Size().Y),
+			View,
+			Inputs.Viewport,
+			Inputs.Viewport,
+			PixelShader,
+			PassParameters
+		);
+	}
+
+	if (ShouldVisualizeWarping(Inputs.DenoisingContext) && Inputs.DenoisedTexture)
+	{
+		const FScreenPassTextureViewport TargetViewport(Inputs.DenoisedTexture, View.ViewRect);
+		typedef FVisualizeWarpingPS SHADER;
+		SHADER::FParameters* PassParameters = GraphBuilder.AllocParameters<SHADER::FParameters>();
+		{
+			const FScreenPassTextureViewportParameters TargetViewportParameters = GetScreenPassTextureViewportParameters(TargetViewport);
+
+			PassParameters->ViewUniformBuffer = View.ViewUniformBuffer;
+			PassParameters->TemporalDenoisingMotionVector = GraphBuilder.CreateSRV(FRDGTextureSRVDesc::Create(Inputs.DenoisingContext.MotionVector));
+			PassParameters->SceneTextures = Inputs.SceneTexturesUniformBuffer;
+			PassParameters->DenoisedTexture = GraphBuilder.CreateSRV(FRDGTextureSRVDesc::Create(Inputs.DenoisedTexture));
+			PassParameters->SourceTexture = GraphBuilder.CreateSRV(FRDGTextureSRVDesc::Create(Inputs.DenoisingContext.LastDenoisedRadianceTexture));
+			PassParameters->SharedTextureSampler = TStaticSamplerState<SF_Bilinear, AM_Clamp, AM_Clamp, AM_Clamp>::GetRHI();
+			PassParameters->TargetViewport = TargetViewportParameters;
+			PassParameters->RenderTargets[0] = FRenderTargetBinding(Inputs.SceneColor, ERenderTargetLoadAction::ELoad);
+		}
+
+		TShaderMapRef<SHADER> PixelShader(View.ShaderMap);
+		AddDrawScreenPass(
+			GraphBuilder,
+			RDG_EVENT_NAME("Visualize Source Warping (%d x %d)", View.ViewRect.Size().X, View.ViewRect.Size().Y),
+			View,
+			Inputs.Viewport,
+			Inputs.Viewport,
+			PixelShader,
+			PassParameters
+		);
+	}
+
+	return FScreenPassTexture();
+}
+
+#endif
\ No newline at end of file
diff --git a/Engine/Source/Runtime/Renderer/Private/PathTracingSpatialTemporalDenoising.h b/Engine/Source/Runtime/Renderer/Private/PathTracingSpatialTemporalDenoising.h
new file mode 100644
--- /dev/null
+++ b/Engine/Source/Runtime/Renderer/Private/PathTracingSpatialTemporalDenoising.h
@@ -0,0 +1,77 @@
+// Copyright Epic Games, Inc. All Rights Reserved.
+
+#pragma once
+
+#if RHI_RAYTRACING
+
+#include "ShaderParameterMacros.h"
+#include "RendererPrivate.h"
+
+class FRDGBuilder;
+class FScene;
+class FViewInfo;
+class FRDGTexture;
+
+struct FPathTracingSpatialTemporalDenoisingContext
+{
+	FRDGTexture*	RadianceTexture = nullptr;
+	FRDGTexture*	AlbedoTexture = nullptr;
+	FRDGTexture*	NormalTexture = nullptr;
+	FRDGBuffer*		VarianceBuffer = nullptr;
+
+	FRDGTexture*	LastDenoisedRadianceTexture = nullptr;
+	FRDGTexture*	LastRadianceTexture = nullptr;
+	FRDGTexture*	LastNormalTexture = nullptr;
+	FRDGTexture*	LastAlbedoTexture = nullptr;
+	FRDGBuffer*		LastVarianceBuffer = nullptr;
+
+	int FrameIndex;
+
+	static constexpr uint32 kNumberOfPixelShifts = 25;
+	static constexpr uint32 kNumberOfShiftsPerTexture = 4;
+	static constexpr uint32 kNumberOfPasses = 1;
+	static constexpr uint32 kNumberOfTexturesPerPass =
+		(kNumberOfPixelShifts + (kNumberOfShiftsPerTexture * kNumberOfPasses) - 1)
+		/ (kNumberOfShiftsPerTexture * kNumberOfPasses);
+
+	// Debug information
+	FRDGTexture*	MotionVector = nullptr;
+	FRDGTexture*	WarpedSource = nullptr;
+};
+
+
+int	GetPathTracingDenoiserMode(const FViewInfo& View);
+bool IsPathTracingDenoiserEnabled(const FViewInfo& View);
+bool ShouldEnablePathTracingDenoiserRealtimeDebug();
+
+// Calculate the variance per-pixel.
+void PathTracingSpatialTemporalDenoisingPrePass(
+	FRDGBuilder& GraphBuilder,
+	const FViewInfo& View,
+	int IterationNumber,
+	FPathTracingSpatialTemporalDenoisingContext& SpatialTemporalDenoisingContext
+);
+
+void PathTracingSpatialTemporalDenoising(FRDGBuilder& GraphBuilder,
+	const FViewInfo& View,
+	int DenoiserMode,
+	FRDGTexture*& SpatialTemporalDenoisedTexture,
+	FPathTracingSpatialTemporalDenoisingContext& SpatialTemporalDenoisingContext);
+
+
+struct FVisualizePathTracingDenoisingInputs
+{
+	// [Optional] Render to the specified output. If invalid, a new texture is created and returned.
+	FScreenPassRenderTarget OverrideOutput;
+	FRDGTextureRef SceneColor;
+	FRDGTextureRef DenoisedTexture;
+	FScreenPassTextureViewport Viewport;
+
+	FPathTracingSpatialTemporalDenoisingContext DenoisingContext;
+
+	TRDGUniformBufferRef<FSceneTextureUniformParameters> SceneTexturesUniformBuffer;
+};
+
+FScreenPassTexture AddVisualizePathTracingDenoisingPass(FRDGBuilder& GraphBuilder, const FViewInfo& View, const FVisualizePathTracingDenoisingInputs& Inputs);
+
+#endif
diff --git a/Engine/Source/Runtime/Renderer/Private/ScenePrivate.h b/Engine/Source/Runtime/Renderer/Private/ScenePrivate.h
--- a/Engine/Source/Runtime/Renderer/Private/ScenePrivate.h
+++ b/Engine/Source/Runtime/Renderer/Private/ScenePrivate.h
@@ -957,7 +957,7 @@ public:
 
 #if RHI_RAYTRACING
 	// Invalidates cached results related to the path tracer so accumulated rendering can start over
-	void PathTracingInvalidate();
+	void PathTracingInvalidate(bool InvalidateAnimationStates = true);
 	virtual uint32 GetPathTracingSampleIndex() const override;
 	virtual uint32 GetPathTracingSampleCount() const override;
 
diff --git a/Engine/Source/Runtime/Renderer/Private/SceneVisibility.cpp b/Engine/Source/Runtime/Renderer/Private/SceneVisibility.cpp
--- a/Engine/Source/Runtime/Renderer/Private/SceneVisibility.cpp
+++ b/Engine/Source/Runtime/Renderer/Private/SceneVisibility.cpp
@@ -3881,6 +3881,9 @@ void FSceneRenderer::PreVisibilityFrameSetup(FRDGBuilder& GraphBuilder, const FS
 				0.18f /*degree*/, 0.1f /*cm*/);
 			const bool bIsProjMatrixDifferent = View.ViewMatrices.GetProjectionNoAAMatrix() != View.ViewState->PrevFrameViewInfo.ViewMatrices.GetProjectionNoAAMatrix();
 			
+			static const auto CVarTemporalDenoiser = IConsoleManager::Get().FindTConsoleVariableDataInt(TEXT("r.PathTracing.TemporalDenoiser.mode"));
+			const int TemporalDenoiserMode = CVarTemporalDenoiser ? CVarTemporalDenoiser->GetValueOnAnyThread() : 0;
+
 			if (View.bIsOfflineRender)
 			{
 				// In the offline context, we want precise control over when to restart the path tracer's accumulation to allow for motion blur
@@ -3888,7 +3891,8 @@ void FSceneRenderer::PreVisibilityFrameSetup(FRDGBuilder& GraphBuilder, const FS
 				// interactions with the motion blur post process effect in tiled rendering (see comment below).
 				if (View.bCameraCut || View.bForcePathTracerReset)
 				{
-					ViewState->PathTracingInvalidate();
+					const bool bClearTemporalDenoisingHistory = (TemporalDenoiserMode == 1) ? View.bCameraCut : true;
+					ViewState->PathTracingInvalidate(bClearTemporalDenoisingHistory);
 				}
 			}
 			else
@@ -3904,7 +3908,8 @@ void FSceneRenderer::PreVisibilityFrameSetup(FRDGBuilder& GraphBuilder, const FS
 					bIsCameraMove ||
 					View.bForcePathTracerReset)
 				{
-					ViewState->PathTracingInvalidate();
+					const bool bClearTemporalDenoisingHistory = (TemporalDenoiserMode == 2) ? (View.bCameraCut || bResetCamera) : true;
+					ViewState->PathTracingInvalidate(bClearTemporalDenoisingHistory);
 				}
 			}
 
