Bug fix:  code initializing ray tracing related variables in FDeferredShadingSceneRenderer need to be initialized on render thread, as they access numerous render thread related CVARs...

#rnx

diff --git a/Engine/Source/Runtime/Renderer/Private/DeferredShadingRenderer.cpp b/Engine/Source/Runtime/Renderer/Private/DeferredShadingRenderer.cpp
--- a/Engine/Source/Runtime/Renderer/Private/DeferredShadingRenderer.cpp
+++ b/Engine/Source/Runtime/Renderer/Private/DeferredShadingRenderer.cpp
@@ -380,17 +380,10 @@ FDeferredShadingSceneRenderer::FDeferredShadingSceneRenderer(const FSceneViewFam
 	ViewPipelineStates.SetNum(Views.Num());
 
 #if RHI_RAYTRACING
+	// Flags initialized later in render thread
 	bAnyRayTracingPassEnabled = false;
-	for (int32 ViewIndex = 0; ViewIndex < Views.Num(); ViewIndex++)
-	{
-		bool bHasRayTracing = AnyRayTracingPassEnabled(Scene, Views[ViewIndex]);
-
-		Views[ViewIndex].bHasAnyRayTracingPass = bHasRayTracing;
-
-		bAnyRayTracingPassEnabled |= bHasRayTracing;
-	}
-	bShouldUpdateRayTracingScene = bAnyRayTracingPassEnabled;
-#endif  // RHI_RAYTRACING
+	bShouldUpdateRayTracingScene = false;
+#endif
 }
 
 /** 
@@ -1967,6 +1960,9 @@ void FDeferredShadingSceneRenderer::Render(FRDGBuilder& GraphBuilder)
 	Scene->UpdateAllPrimitiveSceneInfos(GraphBuilder, true);
 
 #if RHI_RAYTRACING
+	// Initialize ray tracing flags, in case they weren't initialized in the CreateSceneRenderers code path
+	InitializeRayTracingFlags_RenderThread();
+
 	// Now that we have updated all the PrimitiveSceneInfos, update the RayTracing mesh commands cache if needed
 	{
 		ERayTracingMeshCommandsMode CurrentMode = ViewFamily.EngineShowFlags.PathTracing ? ERayTracingMeshCommandsMode::PATH_TRACING : ERayTracingMeshCommandsMode::RAY_TRACING;
@@ -3462,11 +3458,11 @@ bool ShouldRenderRayTracingEffect(bool bEffectEnabled, ERayTracingPipelineCompat
 	}
 
 	const bool bAllowPipeline = GRHISupportsRayTracingShaders && 
-								CVarRayTracingAllowPipeline.GetValueOnAnyThread() &&
+								CVarRayTracingAllowPipeline.GetValueOnRenderThread() &&
 								EnumHasAnyFlags(CompatibilityFlags, ERayTracingPipelineCompatibilityFlags::FullPipeline);
 
 	const bool bAllowInline = GRHISupportsInlineRayTracing && 
-							  CVarRayTracingAllowInline.GetValueOnAnyThread() &&
+							  CVarRayTracingAllowInline.GetValueOnRenderThread() &&
 							  EnumHasAnyFlags(CompatibilityFlags, ERayTracingPipelineCompatibilityFlags::Inline);
 
 	// Disable the effect if current machine does not support the full ray tracing pipeline and the effect can't fall back to inline mode or vice versa.
@@ -3475,7 +3471,7 @@ bool ShouldRenderRayTracingEffect(bool bEffectEnabled, ERayTracingPipelineCompat
 		return false;
 	}
 
-	const int32 OverrideMode = CVarForceAllRayTracingEffects.GetValueOnAnyThread();
+	const int32 OverrideMode = CVarForceAllRayTracingEffects.GetValueOnRenderThread();
 
 	if (OverrideMode >= 0)
 	{
@@ -3495,4 +3491,24 @@ bool HasRayTracedOverlay(const FSceneViewFamily& ViewFamily)
 		ViewFamily.EngineShowFlags.PathTracing ||
 		ViewFamily.EngineShowFlags.RayTracingDebug;
 }
+
+void FDeferredShadingSceneRenderer::InitializeRayTracingFlags_RenderThread()
+{
+	// This function may be called twice -- once in CreateSceneRenderers and again in Render.  We deliberately skip the logic
+	// if the flag is already set, because CreateSceneRenderers fills in the correct value for "bShouldUpdateRayTracingScene"
+	// in that case, and we don't want to overwrite it.
+	if (!bAnyRayTracingPassEnabled)
+	{
+		for (int32 ViewIndex = 0; ViewIndex < Views.Num(); ViewIndex++)
+		{
+			bool bHasRayTracing = AnyRayTracingPassEnabled(Scene, Views[ViewIndex]);
+
+			Views[ViewIndex].bHasAnyRayTracingPass = bHasRayTracing;
+
+			bAnyRayTracingPassEnabled |= bHasRayTracing;
+		}
+
+		bShouldUpdateRayTracingScene = bAnyRayTracingPassEnabled;
+	}
+}
 #endif // RHI_RAYTRACING
diff --git a/Engine/Source/Runtime/Renderer/Private/DeferredShadingRenderer.h b/Engine/Source/Runtime/Renderer/Private/DeferredShadingRenderer.h
--- a/Engine/Source/Runtime/Renderer/Private/DeferredShadingRenderer.h
+++ b/Engine/Source/Runtime/Renderer/Private/DeferredShadingRenderer.h
@@ -214,6 +214,8 @@ public:
 #if RHI_RAYTRACING
 	bool bAnyRayTracingPassEnabled;
 	bool bShouldUpdateRayTracingScene;
+
+	void InitializeRayTracingFlags_RenderThread();
 #endif
 
 	FDeferredShadingSceneRenderer(const FSceneViewFamily* InViewFamily, FHitProxyConsumer* HitProxyConsumer);
diff --git a/Engine/Source/Runtime/Renderer/Private/SceneRendering.cpp b/Engine/Source/Runtime/Renderer/Private/SceneRendering.cpp
--- a/Engine/Source/Runtime/Renderer/Private/SceneRendering.cpp
+++ b/Engine/Source/Runtime/Renderer/Private/SceneRendering.cpp
@@ -3814,13 +3814,6 @@ void FSceneRenderer::CreateSceneRenderers(TArrayView<const FSceneViewFamily*> In
 
 	EShadingPath ShadingPath = Scene->GetShadingPath();
 
-#if RHI_RAYTRACING
-	// For multi-view-family scene rendering, we need to determine which scene renderer will update the ray tracing scene.
-	// This will be the first view family that uses ray tracing, and subsequent families that use ray tracing can skip
-	// that step.
-	bool bShouldUpdateRayTracingScene = true;
-#endif
-
 	for (int32 FamilyIndex = 0; FamilyIndex < InViewFamilies.Num(); FamilyIndex++)
 	{
 		const FSceneViewFamily* InViewFamily = InViewFamilies[FamilyIndex];
@@ -3831,14 +3824,6 @@ void FSceneRenderer::CreateSceneRenderers(TArrayView<const FSceneViewFamily*> In
 		if (ShadingPath == EShadingPath::Deferred)
 		{
 			FDeferredShadingSceneRenderer* SceneRenderer = new FDeferredShadingSceneRenderer(InViewFamily, HitProxyConsumer);
-#if RHI_RAYTRACING
-			if (SceneRenderer->bAnyRayTracingPassEnabled)
-			{
-				SceneRenderer->bShouldUpdateRayTracingScene = bShouldUpdateRayTracingScene;
-				bShouldUpdateRayTracingScene = false;
-			}
-#endif
-
 			OutSceneRenderers.Add(SceneRenderer);
 		}
 		else
@@ -3850,6 +3835,35 @@ void FSceneRenderer::CreateSceneRenderers(TArrayView<const FSceneViewFamily*> In
 		OutSceneRenderers.Last()->bIsFirstSceneRenderer = (FamilyIndex == 0);
 		OutSceneRenderers.Last()->bIsLastSceneRenderer = (FamilyIndex == InViewFamilies.Num() - 1);
 	}
+
+#if RHI_RAYTRACING
+	// Update ray tracing flags.  We need to do this on the render thread, since "AnyRayTracingPassEnabled" accesses CVars which
+	// are only visible on the render thread.
+	if (ShadingPath == EShadingPath::Deferred)
+	{
+		ENQUEUE_RENDER_COMMAND(InitializeMultiRendererRayTracingFlags)(
+			[SceneRenderers = CopyTemp(OutSceneRenderers)](FRHICommandList& RHICmdList)
+		{
+			// For multi-view-family scene rendering, we need to determine which scene renderer will update the ray tracing
+			// scene.  This will be the first view family that uses ray tracing, and subsequent families that use ray
+			// tracing can skip that step.  
+			bool bShouldUpdateRayTracingScene = true;
+
+			for (int32 RendererIndex = 0; RendererIndex < SceneRenderers.Num(); RendererIndex++)
+			{
+				FDeferredShadingSceneRenderer* SceneRenderer = (FDeferredShadingSceneRenderer*)SceneRenderers[RendererIndex];
+
+				SceneRenderer->InitializeRayTracingFlags_RenderThread();
+
+				if (SceneRenderer->bAnyRayTracingPassEnabled)
+				{
+					SceneRenderer->bShouldUpdateRayTracingScene = bShouldUpdateRayTracingScene;
+					bShouldUpdateRayTracingScene = false;
+				}
+			}
+		});
+	}
+#endif  // RHI_RAYTRACING
 }
 
 FSceneRenderer* FSceneRenderer::CreateSceneRenderer(const FSceneViewFamily* InViewFamily, FHitProxyConsumer* HitProxyConsumer)
