Add r.RayTracing.ExcludeSky CVar to control whether sky geometry is added to the ray tracing scene (default: 1, sky is excluded). This allows a project to restore sky behavior to what it was in UE4.

diff --git a/Engine/Source/Runtime/Renderer/Private/DeferredShadingRenderer.cpp b/Engine/Source/Runtime/Renderer/Private/DeferredShadingRenderer.cpp
--- a/Engine/Source/Runtime/Renderer/Private/DeferredShadingRenderer.cpp
+++ b/Engine/Source/Runtime/Renderer/Private/DeferredShadingRenderer.cpp
@@ -169,6 +169,15 @@ static FAutoConsoleVariableRef CRayTracingExcludeTranslucent(
 	TEXT(" 1: Translucent objects excluded from the ray tracing scene"),
 	ECVF_RenderThreadSafe);
 
+static int32 GRayTracingExcludeSky = 1;
+static FAutoConsoleVariableRef CRayTracingExcludeSky(
+	TEXT("r.RayTracing.ExcludeSky"),
+	GRayTracingExcludeSky,
+	TEXT("A toggle that controls inclusion of sky geometry in the ray tracing scene (excluding sky can make ray tracing faster).\n")
+	TEXT(" 0: Sky objects included in the ray tracing scene\n")
+	TEXT(" 1: Sky objects excluded from the ray tracing scene (default)"),
+	ECVF_RenderThreadSafe);
+
 static TAutoConsoleVariable<int32> CVarRayTracingAsyncBuild(
 	TEXT("r.RayTracing.AsyncBuild"),
 	0,
@@ -1297,7 +1306,7 @@ bool FDeferredShadingSceneRenderer::GatherRayTracingWorldInstancesForView(FRDGBu
 
 				if ((GRayTracingExcludeDecals && RelevantPrimitive.bAnySegmentsDecal)
 					|| (GRayTracingExcludeTranslucent && RelevantPrimitive.bAllSegmentsTranslucent)
-					|| RelevantPrimitive.bIsSky)
+					|| (GRayTracingExcludeSky && RelevantPrimitive.bIsSky))
 				{
 					continue;
 				}
