Deal with case where ScreenPercentageInterface isn't set up on a ViewFamily.
This happens for DrawTileMesh which doesn't set one up, but needs to call InitializeViewFamily() to setup up some default SceneTexturesConfig settings.
#rnx

diff --git a/Engine/Source/Runtime/Renderer/Private/SceneRendering.cpp b/Engine/Source/Runtime/Renderer/Private/SceneRendering.cpp
--- a/Engine/Source/Runtime/Renderer/Private/SceneRendering.cpp
+++ b/Engine/Source/Runtime/Renderer/Private/SceneRendering.cpp
@@ -2708,11 +2708,14 @@ FIntPoint FSceneRenderer::GetDesiredInternalBufferSize(const FSceneViewFamily& V
 		return DesiredBufferSize;
 	}
 
-	DynamicRenderScaling::TMap<float> DynamicResolutionUpperBounds = ViewFamily.GetScreenPercentageInterface()->GetResolutionFractionsUpperBound();
-	float PrimaryResolutionFractionUpperBound = DynamicResolutionUpperBounds[GDynamicPrimaryResolutionFraction];
-
 	// Compute final resolution fraction.
-	float ResolutionFractionUpperBound = PrimaryResolutionFractionUpperBound * ViewFamily.SecondaryViewFraction;
+	float ResolutionFractionUpperBound = 1.f;
+	if (ISceneViewFamilyScreenPercentage const* ScreenPercentageInterface = ViewFamily.GetScreenPercentageInterface())
+	{
+		DynamicRenderScaling::TMap<float> DynamicResolutionUpperBounds = ScreenPercentageInterface->GetResolutionFractionsUpperBound();
+		const float PrimaryResolutionFractionUpperBound = DynamicResolutionUpperBounds[GDynamicPrimaryResolutionFraction];
+		ResolutionFractionUpperBound = PrimaryResolutionFractionUpperBound * ViewFamily.SecondaryViewFraction;
+	}
 
 	FIntPoint FamilySizeUpperBound(0, 0);
 
