Move per-frame RHIT flush to happen right before we start submitting RHIT work
Use 2 frames latency for occlusion queries on GL (same as VK)
This reduce time RHIT waits for work

diff --git a/Engine/Source/Runtime/OpenGLDrv/Private/Android/AndroidOpenGLFramePacer.cpp b/Engine/Source/Runtime/OpenGLDrv/Private/Android/AndroidOpenGLFramePacer.cpp
--- a/Engine/Source/Runtime/OpenGLDrv/Private/Android/AndroidOpenGLFramePacer.cpp
+++ b/Engine/Source/Runtime/OpenGLDrv/Private/Android/AndroidOpenGLFramePacer.cpp
@@ -210,6 +210,8 @@ bool FAndroidOpenGLFramePacer::SupportsFramePace(int32 QueryFramePace)
 
 bool FAndroidOpenGLFramePacer::SwapBuffers(bool bLockToVsync)
 {
+	SCOPED_NAMED_EVENT(STAT_OpenGLSwapBuffersTime, FColor::Red)
+
 #if !UE_BUILD_SHIPPING
 	if (FAndroidPlatformRHIFramePacer::CVarStallSwap.GetValueOnAnyThread() > 0.0f)
 	{
diff --git a/Engine/Source/Runtime/Renderer/Private/MobileShadingRenderer.cpp b/Engine/Source/Runtime/Renderer/Private/MobileShadingRenderer.cpp
--- a/Engine/Source/Runtime/Renderer/Private/MobileShadingRenderer.cpp
+++ b/Engine/Source/Runtime/Renderer/Private/MobileShadingRenderer.cpp
@@ -790,13 +790,9 @@ void FMobileSceneRenderer::Render(FRDGBuilder& GraphBuilder)
 
 	GraphBuilder.SetCommandListStat(GET_STATID(STAT_CLMM_AfterInitViews));
 
-	if (GRHINeedsExtraDeletionLatency || !GRHICommandList.Bypass())
+	if (!ViewFamily.bIsRenderedImmediatelyAfterAnotherViewFamily)
 	{
-		QUICK_SCOPE_CYCLE_COUNTER(STAT_FMobileSceneRenderer_PostInitViewsFlushDel);
-		// we will probably stall on occlusion queries, so might as well have the RHI thread and GPU work while we wait.
-		// Also when doing RHI thread this is the only spot that will process pending deletes
-		FRHICommandListExecutor::GetImmediateCommandList().PollOcclusionQueries();
-		FRHICommandListExecutor::GetImmediateCommandList().ImmediateFlush(EImmediateFlushType::FlushRHIThreadFlushResources);
+		GraphBuilder.SetFlushResourcesRHI();
 	}
 
 	GEngine->GetPreRenderDelegateEx().Broadcast(GraphBuilder);
@@ -1136,6 +1132,7 @@ void FMobileSceneRenderer::RenderForwardSinglePass(FRDGBuilder& GraphBuilder, FM
 		RHICmdList.SetCurrentStat(GET_STATID(STAT_CLMM_Opaque));
 		RenderMobileBasePass(RHICmdList, View);
 		RenderMobileDebugView(RHICmdList, View);
+		RHICmdList.PollOcclusionQueries();
 		const bool bAdrenoOcclusionMode = (CVarMobileAdrenoOcclusionMode.GetValueOnRenderThread() != 0 && IsOpenGLPlatform(ShaderPlatform));
 		if (!bAdrenoOcclusionMode)
 		{
@@ -1187,6 +1184,7 @@ void FMobileSceneRenderer::RenderForwardMultiPass(FRDGBuilder& GraphBuilder, FMo
 		RenderMobileDebugView(RHICmdList, View);
 		// Issue occlusion queries
 		RHICmdList.SetCurrentStat(GET_STATID(STAT_CLMM_Occlusion));
+		RHICmdList.PollOcclusionQueries();
 		RenderOcclusion(RHICmdList, View);
 		PostRenderBasePass(RHICmdList, View);
 	});
@@ -1428,6 +1426,7 @@ void FMobileSceneRenderer::RenderDeferredSinglePass(FRDGBuilder& GraphBuilder, c
 			RHICmdList.SetCurrentStat(GET_STATID(STAT_CLMM_Occlusion));
 			RenderOcclusion(RHICmdList, View);
 		}
+		RHICmdList.PollOcclusionQueries();
 
 		PostRenderBasePass(RHICmdList, View);
 		// SceneColor + GBuffer write, SceneDepth is read only
@@ -1468,6 +1467,7 @@ void FMobileSceneRenderer::RenderDeferredMultiPass(FRDGBuilder& GraphBuilder, cl
 		RenderMobileBasePass(RHICmdList, View);
 		// Issue occlusion queries
 		RHICmdList.SetCurrentStat(GET_STATID(STAT_CLMM_Occlusion));
+		RHICmdList.PollOcclusionQueries();
 		RenderOcclusion(RHICmdList, View);
 		PostRenderBasePass(RHICmdList, View);
 	});
diff --git a/Engine/Source/Runtime/Renderer/Private/SceneOcclusion.cpp b/Engine/Source/Runtime/Renderer/Private/SceneOcclusion.cpp
--- a/Engine/Source/Runtime/Renderer/Private/SceneOcclusion.cpp
+++ b/Engine/Source/Runtime/Renderer/Private/SceneOcclusion.cpp
@@ -106,7 +106,7 @@ int32 FOcclusionQueryHelpers::GetNumBufferedFrames(ERHIFeatureLevel::Type Featur
 	{
 		NumExtraMobileFrames++; // the mobile renderer just doesn't do much after the basepass, and hence it will be asking for the query results almost immediately; the results can't possibly be ready in 1 frame.
 		
-		bool bNeedsAnotherExtraMobileFrame = IsVulkanPlatform(ShaderPlatform); // || IsOpenGLPlatform(ShaderPlatform)
+		bool bNeedsAnotherExtraMobileFrame = IsVulkanPlatform(ShaderPlatform) || IsOpenGLPlatform(ShaderPlatform);
 		bNeedsAnotherExtraMobileFrame = bNeedsAnotherExtraMobileFrame || IsVulkanMobileSM5Platform(ShaderPlatform);
 		bNeedsAnotherExtraMobileFrame = bNeedsAnotherExtraMobileFrame || FDataDrivenShaderPlatformInfo::GetNeedsExtraMobileFrames(ShaderPlatform);
 		bNeedsAnotherExtraMobileFrame = bNeedsAnotherExtraMobileFrame && IsRunningRHIInSeparateThread();
