Fixed compile error from CL# 20687044

diff --git a/Engine/Source/Runtime/Renderer/Private/MobileShadingRenderer.cpp b/Engine/Source/Runtime/Renderer/Private/MobileShadingRenderer.cpp
--- a/Engine/Source/Runtime/Renderer/Private/MobileShadingRenderer.cpp
+++ b/Engine/Source/Runtime/Renderer/Private/MobileShadingRenderer.cpp
@@ -790,7 +790,7 @@ void FMobileSceneRenderer::Render(FRDGBuilder& GraphBuilder)
 
 	GraphBuilder.SetCommandListStat(GET_STATID(STAT_CLMM_AfterInitViews));
 
-	if (!ViewFamily.bIsRenderedImmediatelyAfterAnotherViewFamily)
+	if (bIsFirstSceneRenderer)
 	{
 		GraphBuilder.SetFlushResourcesRHI();
 	}
