Fix path tracing denoiser for interactive rendering.

diff --git a/Engine/Source/Runtime/Renderer/Private/PathTracing.cpp b/Engine/Source/Runtime/Renderer/Private/PathTracing.cpp
--- a/Engine/Source/Runtime/Renderer/Private/PathTracing.cpp
+++ b/Engine/Source/Runtime/Renderer/Private/PathTracing.cpp
@@ -2199,7 +2199,13 @@ void FDeferredShadingSceneRenderer::RenderPathTracing(
 	const bool EnablePathTracingDenoiserRealtimeDebug = ShouldEnablePathTracingDenoiserRealtimeDebug();
 
 	if (IsDenoiserEnabled)
-	{		
+	{	
+		if (PathTracingState->LastDenoisedRadianceRT)
+		{
+			// we already have a texture for this
+			DenoisedRadianceTexture = GraphBuilder.RegisterExternalTexture(PathTracingState->LastDenoisedRadianceRT, TEXT("PathTracer.DenoisedRadiance"));
+		}
+
 		// 1. Prepass to estimate pixel variance
 		FRDGBuffer* CurrentVarianceBufer = nullptr;
 		{
