Path Tracer: Fix potential crash in shader bindings when a decal material accesses scene textures

Make sure decals are attached before the raytracing scene is created so that the right number of slots can be reserved

Fix miss/callbable shader typo

Fix potential crash when searching for default callable shader index

diff --git a/Engine/Shaders/Private/RayTracing/RayTracingDecalMaterialShader.usf b/Engine/Shaders/Private/RayTracing/RayTracingDecalMaterialShader.usf
--- a/Engine/Shaders/Private/RayTracing/RayTracingDecalMaterialShader.usf
+++ b/Engine/Shaders/Private/RayTracing/RayTracingDecalMaterialShader.usf
@@ -1,6 +1,7 @@
 // Copyright Epic Games, Inc. All Rights Reserved.
 
 #define PATH_TRACING 1
+#define SCENE_TEXTURES_DISABLED 1 // Ray tracing shaders cannot access scene textures
 
 #include "/Engine/Private/Common.ush"
 #include "/Engine/Private/RayTracing/RayTracingCommon.ush"
diff --git a/Engine/Source/Runtime/D3D12RHI/Private/D3D12RayTracing.cpp b/Engine/Source/Runtime/D3D12RHI/Private/D3D12RayTracing.cpp
--- a/Engine/Source/Runtime/D3D12RHI/Private/D3D12RayTracing.cpp
+++ b/Engine/Source/Runtime/D3D12RHI/Private/D3D12RayTracing.cpp
@@ -5585,7 +5585,7 @@ static void SetRayTracingMissShader(
 	const uint32 UserDataOffset = offsetof(FHitGroupSystemParameters, RootConstants) + offsetof(FHitGroupSystemRootConstants, UserData);
 	ShaderTable->SetLocalShaderParameters(RecordIndex, UserDataOffset, UserData);
 
-	const FD3D12RayTracingShader* Shader = Pipeline->CallableShaders.Shaders[ShaderIndexInPipeline];
+	const FD3D12RayTracingShader* Shader = Pipeline->MissShaders.Shaders[ShaderIndexInPipeline];
 
 	FD3D12RayTracingLocalResourceBinder ResourceBinder(*Device, *ShaderTable, *(Shader->pRootSignature), RecordIndex, WorkerIndex);
 	const bool bResourcesBound = SetRayTracingShaderResources(Shader,
diff --git a/Engine/Source/Runtime/Renderer/Private/MeshPassProcessor.cpp b/Engine/Source/Runtime/Renderer/Private/MeshPassProcessor.cpp
--- a/Engine/Source/Runtime/Renderer/Private/MeshPassProcessor.cpp
+++ b/Engine/Source/Runtime/Renderer/Private/MeshPassProcessor.cpp
@@ -310,6 +310,8 @@ FRayTracingLocalShaderBindings* FMeshDrawShaderBindings::SetRayTracingShaderBind
 
 	const uint32 NumUniformBuffersToSet = MaxUniformBufferUsed + 1;
 
+	check(NumUniformBuffersToSet == NumUniformBufferParameters);
+
 	const TMemoryImageArray<FShaderLooseParameterBufferInfo>& LooseParameterBuffers = SingleShaderBindings.ParameterMapInfo.LooseParameterBuffers;
 	uint32 LooseParameterDataSize = 0;
 
diff --git a/Engine/Source/Runtime/Renderer/Private/RayTracing/RayTracingMaterialHitShaders.cpp b/Engine/Source/Runtime/Renderer/Private/RayTracing/RayTracingMaterialHitShaders.cpp
--- a/Engine/Source/Runtime/Renderer/Private/RayTracing/RayTracingMaterialHitShaders.cpp
+++ b/Engine/Source/Runtime/Renderer/Private/RayTracing/RayTracingMaterialHitShaders.cpp
@@ -673,7 +673,7 @@ FRayTracingPipelineState* FDeferredShadingSceneRenderer::CreateRayTracingMateria
 		&& View.ViewState
 		&& IsCompatibleFallbackPipelineSignature(View.ViewState->LastRayTracingMaterialPipelineSignature, Initializer)
 		&& PipelineContainsHitShaders(FallbackPipelineState, RequiredHitShaders)
-		&& (!bCallableShadersSupported || FindRayTracingCallableShaderIndex(FallbackPipelineState, DefaultCallableShader, false)))
+		&& (!bCallableShadersSupported || FindRayTracingCallableShaderIndex(FallbackPipelineState, DefaultCallableShader, false) != INDEX_NONE))
 	{
 		PipelineCacheFlags |= ERayTracingPipelineCacheFlags::NonBlocking;
 	}
