Fix path tracing denoiser compile error for android

diff --git a/Engine/Source/Runtime/Renderer/Private/PathTracingSpatialTemporalDenoising.h b/Engine/Source/Runtime/Renderer/Private/PathTracingSpatialTemporalDenoising.h
--- a/Engine/Source/Runtime/Renderer/Private/PathTracingSpatialTemporalDenoising.h
+++ b/Engine/Source/Runtime/Renderer/Private/PathTracingSpatialTemporalDenoising.h
@@ -2,6 +2,8 @@
 
 #pragma once
 
+#include "RHIDefinitions.h"
+
 #if RHI_RAYTRACING
 
 #include "ShaderParameterMacros.h"
