Remove unnecessary decal calculation that can trigger an ensure() for near zero component scales.

diff --git a/Engine/Source/Runtime/Renderer/Private/DecalRenderingShared.cpp b/Engine/Source/Runtime/Renderer/Private/DecalRenderingShared.cpp
--- a/Engine/Source/Runtime/Renderer/Private/DecalRenderingShared.cpp
+++ b/Engine/Source/Runtime/Renderer/Private/DecalRenderingShared.cpp
@@ -91,7 +91,6 @@
 		const FLargeWorldRenderPosition AbsoluteOrigin(View.ViewMatrices.GetInvViewMatrix().GetOrigin());
 		const FVector3f TilePosition = AbsoluteOrigin.GetTile();
 		const FMatrix WorldToDecalMatrix = DecalProxy.ComponentTrans.ToInverseMatrixWithScale();
-		const FMatrix44f RelativeWorldToDecalMatrix = FLargeWorldRenderScalar::MakeFromRelativeWorldMatrix(AbsoluteOrigin.GetTileOffset(), WorldToDecalMatrix);
 		const FMatrix DecalToWorldMatrix = DecalProxy.ComponentTrans.ToMatrixWithScale();
 		const FMatrix44f RelativeDecalToWorldMatrix = FLargeWorldRenderScalar::MakeToRelativeWorldMatrix(AbsoluteOrigin.GetTileOffset(), DecalToWorldMatrix);
 		const FVector3f OrientationVector = (FVector3f)DecalProxy.ComponentTrans.GetUnitAxis(EAxis::X);
diff --git a/Engine/Source/Runtime/Renderer/Private/RayTracing/RayTracingDecals.cpp b/Engine/Source/Runtime/Renderer/Private/RayTracing/RayTracingDecals.cpp
--- a/Engine/Source/Runtime/Renderer/Private/RayTracing/RayTracingDecals.cpp
+++ b/Engine/Source/Runtime/Renderer/Private/RayTracing/RayTracingDecals.cpp
@@ -61,7 +61,6 @@
 	const FLargeWorldRenderPosition AbsoluteOrigin(View.ViewMatrices.GetInvViewMatrix().GetOrigin());
 	const FVector3f TilePosition = AbsoluteOrigin.GetTile();
 	const FMatrix WorldToDecalMatrix = DecalData.Proxy.ComponentTrans.ToInverseMatrixWithScale();
-	const FMatrix44f RelativeWorldToDecalMatrix = FLargeWorldRenderScalar::MakeFromRelativeWorldMatrix(AbsoluteOrigin.GetTileOffset(), WorldToDecalMatrix);
 	const FMatrix DecalToWorldMatrix = DecalData.Proxy.ComponentTrans.ToMatrixWithScale();
 	const FMatrix44f RelativeDecalToWorldMatrix = FLargeWorldRenderScalar::MakeToRelativeWorldMatrix(AbsoluteOrigin.GetTileOffset(), DecalToWorldMatrix);
 	const FVector3f OrientationVector = (FVector3f)DecalData.Proxy.ComponentTrans.GetUnitAxis(EAxis::X);
