Don't disable tonemapper when visualizing Nanite or Virtual Shadow Maps.
- The debug views are composited after post processing.

diff --git a/Engine/Source/Editor/UnrealEd/Private/EditorViewportClient.cpp b/Engine/Source/Editor/UnrealEd/Private/EditorViewportClient.cpp
--- a/Engine/Source/Editor/UnrealEd/Private/EditorViewportClient.cpp
+++ b/Engine/Source/Editor/UnrealEd/Private/EditorViewportClient.cpp
@@ -3927,11 +3927,9 @@ void FEditorViewportClient::Draw(FViewport* InViewport, FCanvas* Canvas)
 	ViewFamily.ViewMode = CurrentViewMode;
 
 	const bool bVisualizeBufferEnabled = CurrentViewMode == VMI_VisualizeBuffer && CurrentBufferVisualizationMode != NAME_None;
-	const bool bVisualizeNaniteEnabled = CurrentViewMode == VMI_VisualizeNanite && CurrentNaniteVisualizationMode != NAME_None;
-	const bool bVisualizeVirtualShadowMapEnabled = CurrentViewMode == VMI_VisualizeVirtualShadowMap && CurrentVirtualShadowMapVisualizationMode != NAME_None;
 	const bool bRayTracingDebugEnabled = CurrentViewMode == VMI_RayTracingDebug && CurrentRayTracingDebugVisualizationMode != NAME_None;
 	const bool bVisualizeGPUSkinCache = CurrentViewMode == VMI_VisualizeGPUSkinCache && CurrentGPUSkinCacheVisualizationMode != NAME_None;
-	const bool bCanDisableTonemapper = bVisualizeBufferEnabled || bVisualizeNaniteEnabled || bVisualizeVirtualShadowMapEnabled || bVisualizeGPUSkinCache || (bRayTracingDebugEnabled && !FRayTracingDebugVisualizationMenuCommands::DebugModeShouldBeTonemapped(CurrentRayTracingDebugVisualizationMode));
+	const bool bCanDisableTonemapper = bVisualizeBufferEnabled || bVisualizeGPUSkinCache || (bRayTracingDebugEnabled && !FRayTracingDebugVisualizationMenuCommands::DebugModeShouldBeTonemapped(CurrentRayTracingDebugVisualizationMode));
 	
 	EngineShowFlagOverride(ESFIM_Editor, ViewFamily.ViewMode, ViewFamily.EngineShowFlags, bCanDisableTonemapper);
 	EngineShowFlagOrthographicOverride(IsPerspective(), ViewFamily.EngineShowFlags);
diff --git a/Engine/Source/Runtime/Renderer/Private/PostProcess/PostProcessEyeAdaptation.cpp b/Engine/Source/Runtime/Renderer/Private/PostProcess/PostProcessEyeAdaptation.cpp
--- a/Engine/Source/Runtime/Renderer/Private/PostProcess/PostProcessEyeAdaptation.cpp
+++ b/Engine/Source/Runtime/Renderer/Private/PostProcess/PostProcessEyeAdaptation.cpp
@@ -234,8 +234,6 @@ bool IsAutoExposureDebugMode(const FViewInfo& View)
 	return View.Family->UseDebugViewPS() ||
 		!EngineShowFlags.Lighting ||
 		(EngineShowFlags.VisualizeBuffer && View.CurrentBufferVisualizationMode != NAME_None) ||
-		(EngineShowFlags.VisualizeNanite && View.CurrentNaniteVisualizationMode != NAME_None) ||
-		(EngineShowFlags.VisualizeVirtualShadowMap && View.CurrentVirtualShadowMapVisualizationMode != NAME_None) ||
 		EngineShowFlags.RayTracingDebug ||
 		EngineShowFlags.VisualizeDistanceFieldAO ||
 		EngineShowFlags.VisualizeVolumetricCloudConservativeDensity ||
@@ -965,9 +963,7 @@ void FSceneViewState::UpdatePreExposure(FViewInfo& View)
 		!ViewFamily.EngineShowFlags.LODColoration &&
 		!ViewFamily.EngineShowFlags.HLODColoration &&
 		!ViewFamily.EngineShowFlags.LevelColoration &&
-		((!ViewFamily.EngineShowFlags.VisualizeBuffer) || View.CurrentBufferVisualizationMode != NAME_None) && // disable pre-exposure for the buffer visualization modes
-		((!ViewFamily.EngineShowFlags.VisualizeNanite) || View.CurrentNaniteVisualizationMode != NAME_None) && // disable pre-exposure for the Nanite visualization modes
-		((!ViewFamily.EngineShowFlags.VisualizeVirtualShadowMap) || View.CurrentVirtualShadowMapVisualizationMode != NAME_None); // disable pre-exposure for the virtual shadow map visualization modes
+		((!ViewFamily.EngineShowFlags.VisualizeBuffer) || View.CurrentBufferVisualizationMode != NAME_None); // disable pre-exposure for the buffer visualization modes
 
 	PreExposure = 1.f;
 	bUpdateLastExposure = false;
