Virtual texturing forces depth pass on.
This fixes bug when using any masked materials that sample VT.

diff --git a/Engine/Source/Runtime/Renderer/Private/DeferredShadingRenderer.cpp b/Engine/Source/Runtime/Renderer/Private/DeferredShadingRenderer.cpp
--- a/Engine/Source/Runtime/Renderer/Private/DeferredShadingRenderer.cpp
+++ b/Engine/Source/Runtime/Renderer/Private/DeferredShadingRenderer.cpp
@@ -355,20 +355,21 @@ const TCHAR* GetDepthPassReason(bool bDitheredLODTransitionsUseStencil, EShaderP
 		return TEXT("(Forced by ForwardShading)");
 	}
 
-	bool bUseNanite = UseNanite(ShaderPlatform);
-
-	if (bUseNanite)
+	if (UseNanite(ShaderPlatform))
 	{
 		return TEXT("(Forced by Nanite)");
 	}
 
-	bool bDBufferAllowed = IsUsingDBuffers(ShaderPlatform);
-
-	if (bDBufferAllowed)
+	if (IsUsingDBuffers(ShaderPlatform))
 	{
 		return TEXT("(Forced by DBuffer)");
 	}
 
+	if (UseVirtualTexturing(ShaderPlatform))
+	{
+		return TEXT("(Forced by VirtualTexture)");
+	}
+
 	if (bDitheredLODTransitionsUseStencil)
 	{
 		return TEXT("(Forced by StencilLODDither)");
diff --git a/Engine/Source/Runtime/Renderer/Private/RendererScene.cpp b/Engine/Source/Runtime/Renderer/Private/RendererScene.cpp
--- a/Engine/Source/Runtime/Renderer/Private/RendererScene.cpp
+++ b/Engine/Source/Runtime/Renderer/Private/RendererScene.cpp
@@ -3893,19 +3893,19 @@ void FScene::Release()
 bool ShouldForceFullDepthPass(EShaderPlatform ShaderPlatform)
 {
 	const bool bNaniteEnabled = UseNanite(ShaderPlatform);
-
 	const bool bDBufferAllowed = IsUsingDBuffers(ShaderPlatform);
+	const bool bVirtualTextureEnabled = UseVirtualTexturing(ShaderPlatform);
 
 	static const auto StencilLODDitherCVar = IConsoleManager::Get().FindTConsoleVariableDataInt(TEXT("r.StencilForLODDither"));
 	const bool bStencilLODDither = StencilLODDitherCVar->GetValueOnAnyThread() != 0;
 
 	static const auto AOComputeCVar = IConsoleManager::Get().FindTConsoleVariableDataInt(TEXT("r.AmbientOcclusion.Compute"));
-	const bool bAOCompute		= AOComputeCVar->GetValueOnAnyThread() > 0;
+	const bool bAOCompute = AOComputeCVar->GetValueOnAnyThread() > 0;
 
 	const bool bEarlyZMaterialMasking = MaskedInEarlyPass(ShaderPlatform);
 
 	// Note: ShouldForceFullDepthPass affects which static draw lists meshes go into, so nothing it depends on can change at runtime, unless you do a FGlobalComponentRecreateRenderStateContext to propagate the cvar change
-	return bNaniteEnabled || bAOCompute || bDBufferAllowed || bStencilLODDither || bEarlyZMaterialMasking || IsForwardShadingEnabled(ShaderPlatform) || IsUsingSelectiveBasePassOutputs(ShaderPlatform);
+	return bNaniteEnabled || bAOCompute || bDBufferAllowed || bVirtualTextureEnabled || bStencilLODDither || bEarlyZMaterialMasking || IsForwardShadingEnabled(ShaderPlatform) || IsUsingSelectiveBasePassOutputs(ShaderPlatform);
 }
 
 void FScene::UpdateEarlyZPassMode()
