Deal with crash where a primitive has been removed from scene just before update from being moved in editor.

diff --git a/Engine/Source/Runtime/Renderer/Private/RendererScene.cpp b/Engine/Source/Runtime/Renderer/Private/RendererScene.cpp
--- a/Engine/Source/Runtime/Renderer/Private/RendererScene.cpp
+++ b/Engine/Source/Runtime/Renderer/Private/RendererScene.cpp
@@ -4357,10 +4357,13 @@ void FScene::UpdatePrimitiveVelocityState_RenderThread(FPrimitiveSceneInfo* Prim
 	{
 		if (ShouldPrimitiveOutputVelocity(PrimitiveSceneInfo->Proxy, GetShaderPlatform()))
 		{
-			PrimitiveSceneInfo->bRegisteredWithVelocityData = true;
-			// We must register the initial LocalToWorld with the velocity state. 
-			int32 PrimitiveIndex = PrimitiveSceneInfo->PackedIndex;
-			VelocityData.UpdateTransform(PrimitiveSceneInfo, PrimitiveTransforms[PrimitiveIndex], PrimitiveTransforms[PrimitiveIndex]);
+			if (PrimitiveSceneInfo->IsIndexValid())
+			{
+				PrimitiveSceneInfo->bRegisteredWithVelocityData = true;
+				// We must register the initial LocalToWorld with the velocity state. 
+				int32 PrimitiveIndex = PrimitiveSceneInfo->PackedIndex;
+				VelocityData.UpdateTransform(PrimitiveSceneInfo, PrimitiveTransforms[PrimitiveIndex], PrimitiveTransforms[PrimitiveIndex]);
+			}
 		}
 	}
 	else if (PrimitiveSceneInfo->bRegisteredWithVelocityData)
