Add CVar to fix subsurface border artifact when checkerboard is turned on. It is off by default due to performance concern, 1 extra profile id access per pixel at low quality (TAA mode by default), 4 extra profile id access per pixel at high quality.

Set `r.SSS.Checkerboard.NeighborSSSValidation 1` to turn it on to have a clear border.

diff --git a/Engine/Shaders/Private/PostProcessSubsurface.usf b/Engine/Shaders/Private/PostProcessSubsurface.usf
--- a/Engine/Shaders/Private/PostProcessSubsurface.usf
+++ b/Engine/Shaders/Private/PostProcessSubsurface.usf
@@ -85,6 +85,18 @@ RWTexture2D<uint> ProfileIdTexture;
 	#define SUBSURFACE_HALFRES 1
 #endif
 
+// Validate if the checkerboard neighbor pixel is subsurface in recombine pass
+#ifndef SUBSURFACE_RECOMBINE
+	#define SUBSURFACE_RECOMBINE 0
+#endif
+
+#ifndef CHECKERBOARD_NEIGHBOR_SSS_VALIDATION
+	#define CHECKERBOARD_NEIGHBOR_SSS_VALIDATION SUBSURFACE_RECOMBINE
+	#if CHECKERBOARD_NEIGHBOR_SSS_VALIDATION
+uint CheckerboardNeighborSSSValidation;
+	#endif
+#endif
+
 //=============================================================================
 // setup for "SeparableSSS.ush"
 //=============================================================================
@@ -439,11 +451,39 @@ struct SDiffuseAndSpecular
 	float3 Specular;
 };
 
+#if CHECKERBOARD_NEIGHBOR_SSS_VALIDATION
+bool IsCheckerBoardNeighborSubsurface(float2 SceneUV, int2 PixelOffset)
+{
+	const float2 UV = SceneUV + PixelOffset * View.BufferSizeAndInvSize.zw;
+
+#if STRATA_ENABLED
+	const FStrataSubsurfaceData SSSData = LoadStataSSSData(UV);
+	const bool bHasSSSProfile = SSSData.bIsValid;
+#else
+	FScreenSpaceData ScreenSpaceData = GetScreenSpaceData(UV);
+	const bool bHasSSSProfile = UseSubsurfaceProfile(ScreenSpaceData.GBuffer.ShadingModelID);
+#endif
+
+	return bHasSSSProfile;
+}
+#endif
+
 // can be moved/shared
 half3 LookupSceneColor(float2 SceneUV, int2 PixelOffset)
 {
 	// faster
-	return SubsurfaceInput0_Texture.SampleLevel(SharedSubsurfaceSampler0, SceneUV, 0, PixelOffset).rgb;
+	float3 SceneColor = SubsurfaceInput0_Texture.SampleLevel(SharedSubsurfaceSampler0, SceneUV, 0, PixelOffset).rgb;
+
+#if CHECKERBOARD_NEIGHBOR_SSS_VALIDATION
+	// Fix border background color leaking into subsurface diffuse/specular.
+	BRANCH
+	if (CheckerboardNeighborSSSValidation != 0)
+	{
+		bool bIsSubsurface = IsCheckerBoardNeighborSubsurface(SceneUV, PixelOffset);
+		SceneColor = lerp(0, SceneColor, bIsSubsurface);
+	}
+#endif
+	return SceneColor;
 }
  
 // @param UVSceneColor for the full res rendertarget (BufferSize) e.g. SceneColor or GBuffers
diff --git a/Engine/Source/Runtime/Renderer/Private/PostProcess/PostProcessSubsurface.cpp b/Engine/Source/Runtime/Renderer/Private/PostProcess/PostProcessSubsurface.cpp
--- a/Engine/Source/Runtime/Renderer/Private/PostProcess/PostProcessSubsurface.cpp
+++ b/Engine/Source/Runtime/Renderer/Private/PostProcess/PostProcessSubsurface.cpp
@@ -95,6 +95,15 @@ namespace
 		TEXT(" 2: Automatic. Non-checkerboard lighting will be applied if we have a suitable rendertarget format\n"),
 		ECVF_RenderThreadSafe);
 
+	TAutoConsoleVariable<int32> CVarSSSCheckerboardNeighborSSSValidation(
+		TEXT("r.SSS.Checkerboard.NeighborSSSValidation"),
+		0,
+		TEXT("Enable or disable checkerboard neighbor subsurface scattering validation.\n")
+		TEXT("This validation can remove border light leakage into subsurface scattering, creating a sharpe border with correct color")
+		TEXT(" 0: Disabled (default)")
+		TEXT(" 1: Enabled. Add 1 subsurface profile id query/pixel (low quality), 4 id query/pixel (high quality) at recombine pass"),
+		ECVF_RenderThreadSafe);
+
 	TAutoConsoleVariable<int32> CVarSSSBurleyQuality(
 		TEXT("r.SSS.Burley.Quality"),
 		1,
@@ -818,6 +827,7 @@ class FSubsurfaceRecombinePS : public FSubsurfaceShader
 		SHADER_PARAMETER_STRUCT(FSubsurfaceInput, SubsurfaceInput1)
 		SHADER_PARAMETER_SAMPLER(SamplerState, SubsurfaceSampler0)
 		SHADER_PARAMETER_SAMPLER(SamplerState, SubsurfaceSampler1)
+		SHADER_PARAMETER(uint32, CheckerboardNeighborSSSValidation)
 		RENDER_TARGET_BINDING_SLOTS()
 	END_SHADER_PARAMETER_STRUCT();
 
@@ -851,6 +861,12 @@ class FSubsurfaceRecombinePS : public FSubsurfaceShader
 		return GetMaxSupportedFeatureLevel(Parameters.Platform) >= ERHIFeatureLevel::SM5;
 	}
 
+	static void ModifyCompilationEnvironment(const FGlobalShaderPermutationParameters& Parameters, FShaderCompilerEnvironment& OutEnvironment)
+	{
+		FSubsurfaceShader::ModifyCompilationEnvironment(Parameters, OutEnvironment);
+		OutEnvironment.SetDefine(TEXT("SUBSURFACE_RECOMBINE"), 1);
+	}
+
 	// Returns the Recombine quality level requested by the SSS Quality CVar setting.
 	static EQuality GetQuality(const FViewInfo& View)
 	{
@@ -889,6 +905,12 @@ class FSubsurfaceRecombinePS : public FSubsurfaceShader
 			return EQuality::Low;
 		}
 	}
+
+	static uint32 GetCheckerBoardNeighborSSSValidation(bool bCheckerBoard)
+	{
+		bool bValidation = CVarSSSCheckerboardNeighborSSSValidation.GetValueOnRenderThread() > 0 ? true : false;
+		return (bCheckerBoard && bValidation) ? 1u : 0u;
+	}
 };
 
 IMPLEMENT_GLOBAL_SHADER(FSubsurfaceRecombinePS, "/Engine/Private/PostProcessSubsurface.usf", "SubsurfaceRecombinePS", SF_Pixel);
@@ -1314,6 +1336,9 @@ void AddSubsurfaceViewPass(
 
 			FSubsurfaceRecombinePS::FParameters* PassParameters = GraphBuilder.AllocParameters<FSubsurfaceRecombinePS::FParameters>();
 			PassParameters->Subsurface = SubsurfaceCommonParameters;
+			// Add dynamic branch parameters for recombine pass only.
+			PassParameters->CheckerboardNeighborSSSValidation = FSubsurfaceRecombinePS::GetCheckerBoardNeighborSSSValidation(bCheckerboard);
+			
 			if (SubsurfaceMode != ESubsurfaceMode::Bypass)
 			{
 				PassParameters->TileParameters = GetSubsurfaceTileParameters(SubsurfaceViewport, Tiles, TileType);
@@ -1349,10 +1374,11 @@ void AddSubsurfaceViewPass(
 				*/
 			AddSubsurfaceTiledScreenPass(
 				GraphBuilder,
-				RDG_EVENT_NAME("SSS::Recombine(%s %s%s%s%s%s) %dx%d",
+				RDG_EVENT_NAME("SSS::Recombine(%s %s%s%s%s%s%s) %dx%d",
 					GetEventName(PixelShaderPermutationVector.Get<FSubsurfaceRecombinePS::FDimensionMode>()),
 					FSubsurfaceRecombinePS::GetEventName(PixelShaderPermutationVector.Get<FSubsurfaceRecombinePS::FDimensionQuality>()),
 					PixelShaderPermutationVector.Get<FSubsurfaceRecombinePS::FDimensionCheckerboard>() ? TEXT(" Checkerboard") : TEXT(""),
+					FSubsurfaceRecombinePS::GetCheckerBoardNeighborSSSValidation(bCheckerboard) ? TEXT("-Validation") : TEXT(""),
 					PixelShaderPermutationVector.Get<FSubsurfaceRecombinePS::FDimensionHalfRes>() ? TEXT(" HalfRes") : TEXT(""),
 					PixelShaderPermutationVector.Get<FSubsurfaceRecombinePS::FRunningInSeparable>() ? TEXT(" RunningInSeparable") : TEXT(""),
 					!bShouldFallbackToFullScreenPass ? TEXT(" Tiled") : TEXT(""),
