Fix SingleLayerWater in Instanced Stereo on desktop deferred (UE-156161).

diff --git a/Engine/Plugins/Experimental/Water/Shaders/Private/WaterMeshVertexFactory.ush b/Engine/Plugins/Experimental/Water/Shaders/Private/WaterMeshVertexFactory.ush
--- a/Engine/Plugins/Experimental/Water/Shaders/Private/WaterMeshVertexFactory.ush
+++ b/Engine/Plugins/Experimental/Water/Shaders/Private/WaterMeshVertexFactory.ush
@@ -325,10 +325,6 @@ FVertexFactoryInterpolantsVSToPS VertexFactoryGetInterpolantsVSToPS(FVertexFacto
 
 	Interpolants.WaveParamIndex = Intermediates.WaveParamIndex;
 
-#if INSTANCED_STEREO
-	Interpolants.EyeIndex = 0;
-#endif
-
 	SetPrimitiveId(Interpolants, Intermediates.SceneData.PrimitiveId);
 
 	return Interpolants;
diff --git a/Engine/Plugins/Experimental/Water/Source/Runtime/Private/WaterMeshSceneProxy.cpp b/Engine/Plugins/Experimental/Water/Source/Runtime/Private/WaterMeshSceneProxy.cpp
--- a/Engine/Plugins/Experimental/Water/Source/Runtime/Private/WaterMeshSceneProxy.cpp
+++ b/Engine/Plugins/Experimental/Water/Source/Runtime/Private/WaterMeshSceneProxy.cpp
@@ -213,13 +213,22 @@ void FWaterMeshSceneProxy::GetDynamicMeshElements(const TArray<const FSceneView*
 
 	TArray<FWaterQuadTree::FTraversalOutput, TInlineAllocator<4>> WaterInstanceDataPerView;
 
-	// Gather visible tiles, their lod and materials for all views 
+	bool bEncounteredISRView = false;
+	int32 InstanceFactor = 1;
+
+	// Gather visible tiles, their lod and materials for all renderable views (skip right view when stereo pair is rendered instanced)
 	for (int32 ViewIndex = 0; ViewIndex < Views.Num(); ViewIndex++)
 	{
-		if (VisibilityMap & (1 << ViewIndex))
+		const FSceneView* View = Views[ViewIndex];
+		if (!bEncounteredISRView && View->IsInstancedStereoPass())
 		{
-			const FSceneView* View = Views[ViewIndex];
+			bEncounteredISRView = true;
+			InstanceFactor = View->GetStereoPassInstanceFactor();
+		}
 
+		// skip gathering visible tiles from instanced right eye views
+		if ((VisibilityMap & (1 << ViewIndex)) && (!bEncounteredISRView || View->IsPrimarySceneView()))
+		{
 			const FVector ObserverPosition = View->ViewMatrices.GetViewOrigin();
 			
 			FWaterLODParams WaterLODParams = GetWaterLODParams(ObserverPosition);
@@ -282,7 +291,7 @@ void FWaterMeshSceneProxy::GetDynamicMeshElements(const TArray<const FSceneView*
 		return;
 	}
 
-	WaterInstanceDataBuffers->Lock(TotalInstanceCount);
+	WaterInstanceDataBuffers->Lock(TotalInstanceCount * InstanceFactor);
 
 	int32 InstanceDataOffset = 0;
 
@@ -290,7 +299,8 @@ void FWaterMeshSceneProxy::GetDynamicMeshElements(const TArray<const FSceneView*
 	int32 TraversalIndex = 0;
 	for (int32 ViewIndex = 0; ViewIndex < Views.Num(); ViewIndex++)
 	{
-		if (VisibilityMap & (1 << ViewIndex))
+		// when rendering ISR, don't process the instanced view
+		if ((VisibilityMap & (1 << ViewIndex)) && (!bEncounteredISRView || Views[ViewIndex]->IsPrimarySceneView()))
 		{
 			TRACE_CPUPROFILER_EVENT_SCOPE(BucketsPerView);
 
@@ -362,7 +372,7 @@ void FWaterMeshSceneProxy::GetDynamicMeshElements(const TArray<const FSceneView*
 							//BatchElement.bIsInstancedMesh = true;
 							BatchElement.NumInstances = InstanceCount;
 							BatchElement.UserData = (void*)WaterMeshUserDataBuffers->GetUserData(RenderGroup);
-							BatchElement.UserIndex = InstanceDataOffset;
+							BatchElement.UserIndex = InstanceDataOffset * InstanceFactor;
 
 							BatchElement.FirstIndex = 0;
 							BatchElement.NumPrimitives = WaterVertexFactories[DensityIndex]->IndexBuffer->GetIndexCount() / 3;
@@ -405,7 +415,11 @@ void FWaterMeshSceneProxy::GetDynamicMeshElements(const TArray<const FSceneView*
 
 				for (int32 StreamIdx = 0; StreamIdx < WaterInstanceDataBuffersType::NumBuffers; ++StreamIdx)
 				{
-					WaterInstanceDataBuffers->GetBufferMemory(StreamIdx)[WriteIndex] = Data.Data[StreamIdx];
+					TArrayView<FVector4f> BufferMemory = WaterInstanceDataBuffers->GetBufferMemory(StreamIdx);
+					for (int32 IdxMultipliedInstance = 0; IdxMultipliedInstance < InstanceFactor; ++IdxMultipliedInstance)
+					{
+						BufferMemory[WriteIndex * InstanceFactor + IdxMultipliedInstance] = Data.Data[StreamIdx];
+					}
 				}
 			}
 		}
diff --git a/Engine/Shaders/Private/SingleLayerWaterComposite.usf b/Engine/Shaders/Private/SingleLayerWaterComposite.usf
--- a/Engine/Shaders/Private/SingleLayerWaterComposite.usf
+++ b/Engine/Shaders/Private/SingleLayerWaterComposite.usf
@@ -55,6 +55,23 @@ SingleLayerWaterCompositeOutput SingleLayerWaterComposite(float2 BufferUV, float
 	float3 TransmittanceToScene = 1.0f;
 	float3 SeparatedMainDirLightContribution = 0.0f;
 
+	float4 EffectiveSceneNoWaterMinMaxUV = SceneNoWaterMinMaxUV;
+#if INSTANCED_STEREO && !MULTI_VIEW
+	if (IsInstancedStereo())
+	{
+		// prevent pixels from left eye from sampling "no water" texture from the right eye
+		if (BufferUV.x < 0.5f)
+		{
+			EffectiveSceneNoWaterMinMaxUV.z -= 0.5f;
+		}
+		else
+		{
+			EffectiveSceneNoWaterMinMaxUV.x += 0.5f;
+		}
+	}
+#endif
+
+
 #if STRATA_ENABLED
 
 	// Water do nto read the classification data to avoid running it a second time when water is enabled.
@@ -65,7 +82,7 @@ SingleLayerWaterCompositeOutput SingleLayerWaterComposite(float2 BufferUV, float
 		float DeviceZ = SampleDeviceZFromSceneTextures(BufferUV);
 		float SceneDepth = ConvertFromDeviceZ(DeviceZ);
 
-		const float2 SceneNoWaterUV = clamp(BufferUV, SceneNoWaterMinMaxUV.xy, SceneNoWaterMinMaxUV.zw);
+		const float2 SceneNoWaterUV = clamp(BufferUV, EffectiveSceneNoWaterMinMaxUV.xy, EffectiveSceneNoWaterMinMaxUV.zw);
 
 		float OpaqueDepth = WaterSampleSceneDepthWithoutWater(
 			SceneNoWaterDepthTexture,
@@ -154,7 +171,7 @@ SingleLayerWaterCompositeOutput SingleLayerWaterComposite(float2 BufferUV, float
 
 	if (ShadingModelID == SHADINGMODELID_SINGLELAYERWATER)
 	{
-		const float2 SceneNoWaterUV = clamp(BufferUV, SceneNoWaterMinMaxUV.xy, SceneNoWaterMinMaxUV.zw);
+		const float2 SceneNoWaterUV = clamp(BufferUV, EffectiveSceneNoWaterMinMaxUV.xy, EffectiveSceneNoWaterMinMaxUV.zw);
 
 		float OpaqueDepth = WaterSampleSceneDepthWithoutWater(
 			SceneNoWaterDepthTexture,
diff --git a/Engine/Shaders/Private/SingleLayerWaterShading.ush b/Engine/Shaders/Private/SingleLayerWaterShading.ush
--- a/Engine/Shaders/Private/SingleLayerWaterShading.ush
+++ b/Engine/Shaders/Private/SingleLayerWaterShading.ush
@@ -50,6 +50,21 @@ WaterVolumeLightingOutput EvaluateWaterVolumeLighting(
 
 	float2 ViewportUV = MaterialParameters.ViewBufferUV;
 
+#if INSTANCED_STEREO && !MULTI_VIEW
+	if (IsInstancedStereo())
+	{
+		// prevent pixels from left eye from sampling "no water" texture from the right eye
+		if (ViewportUV.x < 0.5f)
+		{
+			SceneWithoutSingleLayerWaterMaxUV.x -= 0.5f;
+		}
+		else
+		{
+			SceneWithoutSingleLayerWaterMinUV.x += 0.5f;
+		}
+	}
+#endif
+
 	if (WaterVisibility > 0.0f)
 	{
 		float WaterDepth = MaterialParameters.ScreenPosition.w; // Scene depth with Pixel-Depth-Offset taken into account is in ScreenPosition.w already
diff --git a/Engine/Shaders/Private/VolumetricRenderTarget.usf b/Engine/Shaders/Private/VolumetricRenderTarget.usf
--- a/Engine/Shaders/Private/VolumetricRenderTarget.usf
+++ b/Engine/Shaders/Private/VolumetricRenderTarget.usf
@@ -325,21 +325,53 @@ float2 FullResolutionToWaterBufferScale;
 Texture2DMS<float, MSAA_SAMPLE_COUNT> MSAADepthTexture;
 #endif
 
+#if INSTANCED_STEREO
+// When rendering instanced stereo side by side, we may use first view's texture for both views. This wraps the coords so the second view can use first view's texture.
+// (Do we need a new shader permutation here? Also, should really use ViewRect for wrapping)
+uint2 WrapCoordsForInstancedViews(uint2 Coord, uint4 ValidRect)
+{
+	return uint2(
+		(Coord.x > ValidRect.z) ? (Coord.x - ValidRect.z) : Coord.x,
+		Coord.y
+		);
+}
+float2 WrapUVsForInstancedViews(float2 UV, float4 ValidRect)
+{
+	return float2(
+		(UV.x > ValidRect.z) ? (UV.x - ValidRect.z) : UV.x,
+		UV.y
+		);
+}
+#else
+uint2 WrapCoordsForInstancedViews(uint2 Coord, uint4 ValidRect)
+{
+	return Coord;
+}
+float2 WrapUVsForInstancedViews(float2 UV, float4 ValidRect)
+{
+	return UV;
+}
+#endif
+
 float4 SafeLoadVolumetricTexture(uint2 Coord)
 {
-	return VolumetricTexture.Load(uint3(clamp(Coord, VolumetricTextureValidCoordRect.xy, VolumetricTextureValidCoordRect.zw), 0));
+	const uint2 WrappedCoords = WrapCoordsForInstancedViews(Coord, VolumetricTextureValidCoordRect);
+	return VolumetricTexture.Load(uint3(clamp(WrappedCoords, VolumetricTextureValidCoordRect.xy, VolumetricTextureValidCoordRect.zw), 0));
 }
 float4 SafeSampleVolumetricTexture(float2 UV)
 {
-	return VolumetricTexture.SampleLevel(LinearTextureSampler, clamp(UV, VolumetricTextureValidUvRect.xy, VolumetricTextureValidUvRect.zw), 0);
+	const float2 WrappedUV = WrapUVsForInstancedViews(UV, VolumetricTextureValidUvRect);
+	return VolumetricTexture.SampleLevel(LinearTextureSampler, clamp(WrappedUV, VolumetricTextureValidUvRect.xy, VolumetricTextureValidUvRect.zw), 0);
 }
 float2 SafeLoadVolumetricDepthTexture(uint2 Coord)
 {
-	return VolumetricDepthTexture.Load(uint3(clamp(Coord, VolumetricTextureValidCoordRect.xy, VolumetricTextureValidCoordRect.zw), 0)).rg;
+	const uint2 WrappedCoords = WrapCoordsForInstancedViews(Coord, VolumetricTextureValidCoordRect);
+	return VolumetricDepthTexture.Load(uint3(clamp(WrappedCoords, VolumetricTextureValidCoordRect.xy, VolumetricTextureValidCoordRect.zw), 0)).rg;
 }
 float2 SafeSampleVolumetricDepthTexture(float2 UV)
 {
-	return VolumetricDepthTexture.SampleLevel(LinearTextureSampler, clamp(UV, VolumetricTextureValidUvRect.xy, VolumetricTextureValidUvRect.zw), 0).rg;
+	const float2 WrappedUV = WrapUVsForInstancedViews(UV, VolumetricTextureValidUvRect);
+	return VolumetricDepthTexture.SampleLevel(LinearTextureSampler, clamp(WrappedUV, VolumetricTextureValidUvRect.xy, VolumetricTextureValidUvRect.zw), 0).rg;
 }
 
 float4 VolumetricTextureSizeAndInvSize;
diff --git a/Engine/Source/Runtime/Engine/Private/SceneView.cpp b/Engine/Source/Runtime/Engine/Private/SceneView.cpp
--- a/Engine/Source/Runtime/Engine/Private/SceneView.cpp
+++ b/Engine/Source/Runtime/Engine/Private/SceneView.cpp
@@ -2403,6 +2403,12 @@ bool FSceneView::IsInstancedStereoPass() const
 	return bIsInstancedStereoEnabled && IStereoRendering::IsStereoEyeView(*this) && IStereoRendering::IsAPrimaryView(*this);
 }
 
+int32 FSceneView::GetStereoPassInstanceFactor() const
+{
+	return bIsInstancedStereoEnabled && IStereoRendering::IsStereoEyeView(*this) && GEngine->StereoRenderingDevice.IsValid() ?
+		GEngine->StereoRenderingDevice->GetDesiredNumberOfViews(true) : 1;
+}
+
 FVector4f FSceneView::GetScreenPositionScaleBias(const FIntPoint& BufferSize, const FIntRect& ViewRect) const
 {
 	const float InvBufferSizeX = 1.0f / BufferSize.X;
diff --git a/Engine/Source/Runtime/Engine/Public/SceneView.h b/Engine/Source/Runtime/Engine/Public/SceneView.h
--- a/Engine/Source/Runtime/Engine/Public/SceneView.h
+++ b/Engine/Source/Runtime/Engine/Public/SceneView.h
@@ -1437,6 +1437,9 @@ public:
 	/** True if the view should render as an instanced stereo pass */
 	bool IsInstancedStereoPass() const;
 
+	/** Instance factor for a stereo pass (normally 2 for ISR views, but see IStereoRendering::GetDesiredNumberOfViews()). Returns 1 for non-instanced stereo views or regular (split screen etc) views. */
+	int32 GetStereoPassInstanceFactor() const;
+
 	/** Sets up the view rect parameters in the view's uniform shader parameters */
 	void SetupViewRectUniformBufferParameters(FViewUniformShaderParameters& ViewUniformShaderParameters, 
 		const FIntPoint& InBufferSize,
@@ -1475,9 +1478,21 @@ public:
 	/** Get the primary view associated with the secondary view. */
 	const FSceneView* GetPrimarySceneView() const;
 
+	/** Checks whether this is the primary view of a stereo pair (important in instanced stereo rendering). Will also be true for any view that isn't stereo. */
+	inline bool IsPrimarySceneView() const
+	{
+		return GetPrimarySceneView() == this;
+	}
+
 	/** Get the first secondary view associated with the primary view. */
 	const FSceneView* GetInstancedSceneView() const;
 
+	/** Checks whether this is the instanced view of a stereo pair. If the technique supports instanced rendering stereo and ISR is enabled, such views can be skipped. */
+	inline bool IsInstancedSceneView() const
+	{
+		return GetInstancedSceneView() == this;
+	}
+
 	/** Get all secondary views associated with the primary view. */
 	TArray<const FSceneView*> GetSecondaryViews() const;
 
diff --git a/Engine/Source/Runtime/Renderer/Private/IndirectLightRendering.cpp b/Engine/Source/Runtime/Renderer/Private/IndirectLightRendering.cpp
--- a/Engine/Source/Runtime/Renderer/Private/IndirectLightRendering.cpp
+++ b/Engine/Source/Runtime/Renderer/Private/IndirectLightRendering.cpp
@@ -1583,6 +1583,8 @@ void FDeferredShadingSceneRenderer::RenderDeferredReflectionsAndSkyLighting(
 	uint32 ViewIndex = 0;
 	for (FViewInfo& View : Views)
 	{
+		RDG_EVENT_SCOPE_CONDITIONAL(GraphBuilder, Views.Num() > 1, "View%d", ViewIndex);
+
 		const uint32 CurrentViewIndex = ViewIndex++;
 		const FPerViewPipelineState& ViewPipelineState = GetViewPipelineState(View);
 
diff --git a/Engine/Source/Runtime/Renderer/Private/SingleLayerWaterRendering.cpp b/Engine/Source/Runtime/Renderer/Private/SingleLayerWaterRendering.cpp
--- a/Engine/Source/Runtime/Renderer/Private/SingleLayerWaterRendering.cpp
+++ b/Engine/Source/Runtime/Renderer/Private/SingleLayerWaterRendering.cpp
@@ -368,7 +368,19 @@ static FSceneWithoutWaterTextures AddCopySceneWithoutWaterPass(
 		PermutationVector.Set<FWaterRefractionCopyPS::FDownsampleColor>(bCopyColor);
 		auto PixelShader = View.ShaderMap->GetShader<FWaterRefractionCopyPS>(PermutationVector);
 
-		const FIntRect RefractionViewRect = FIntRect(FIntPoint::DivideAndRoundDown(View.ViewRect.Min, RefractionDownsampleFactor), FIntPoint::DivideAndRoundDown(View.ViewRect.Max, RefractionDownsampleFactor));
+		// if we have a particular case of ISR where two views are laid out in side by side, we should copy both views at once
+		const bool bIsInstancedStereoSideBySide = View.bIsInstancedStereoEnabled && !View.bIsMultiViewEnabled && IStereoRendering::IsStereoEyeView(View);
+		FIntRect RectToCopy = View.ViewRect;
+		if (bIsInstancedStereoSideBySide)
+		{
+			const FViewInfo* NeighboringStereoView = View.GetInstancedView();
+			if (ensure(NeighboringStereoView))
+			{
+				RectToCopy.Union(NeighboringStereoView->ViewRect);
+			}
+		}
+
+		const FIntRect RefractionViewRect = FIntRect(FIntPoint::DivideAndRoundDown(RectToCopy.Min, RefractionDownsampleFactor), FIntPoint::DivideAndRoundDown(RectToCopy.Max, RefractionDownsampleFactor));
 
 		Textures.Views[ViewIndex].ViewRect   = RefractionViewRect;
 
@@ -422,7 +434,9 @@ void FDeferredShadingSceneRenderer::RenderSingleLayerWaterReflections(
 	{
 		FViewInfo& View = Views[ViewIndex];
 
-		if (!View.ShouldRenderView())
+		// Unfortunately, reflections cannot handle two views at once (yet?) - because of that, allow the secondary pass here.
+		// Note: not completely removing ShouldRenderView in case some other reason to not render it is valid.
+		if (!View.ShouldRenderView() && !IStereoRendering::IsASecondaryPass(View.StereoPass))
 		{
 			continue;
 		}
@@ -440,13 +454,15 @@ void FDeferredShadingSceneRenderer::RenderSingleLayerWaterReflections(
 			FIntVector DepthTextureSize = SceneWithoutWaterTextures.DepthTexture ? SceneWithoutWaterTextures.DepthTexture->Desc.GetSize() : FIntVector::ZeroValue;
 			const bool bShouldUseBilinearSamplerForDepth = SceneWithoutWaterTextures.DepthTexture && ShouldUseBilinearSamplerForDepthWithoutSingleLayerWater(SceneWithoutWaterTextures.DepthTexture->Desc.Format);
 
+			const bool bIsInstancedStereoSideBySide = View.bIsInstancedStereoEnabled && !View.bIsMultiViewEnabled && IStereoRendering::IsStereoEyeView(View);
+
 			Parameters.ScreenSpaceReflectionsTexture = ReflectionsColor ? ReflectionsColor : BlackDummyTexture;
 			Parameters.ScreenSpaceReflectionsSampler = TStaticSamplerState<SF_Point>::GetRHI();
 			Parameters.PreIntegratedGF = GSystemTextures.PreintegratedGF->GetRHI();
 			Parameters.PreIntegratedGFSampler = TStaticSamplerState<SF_Bilinear, AM_Clamp, AM_Clamp, AM_Clamp>::GetRHI();
 			Parameters.SceneNoWaterDepthTexture = SceneWithoutWaterTextures.DepthTexture ? SceneWithoutWaterTextures.DepthTexture : BlackDummyTexture;
 			Parameters.SceneNoWaterDepthSampler = bShouldUseBilinearSamplerForDepth ? TStaticSamplerState<SF_Bilinear>::GetRHI() : TStaticSamplerState<SF_Point>::GetRHI();
-			Parameters.SceneNoWaterMinMaxUV = SceneWithoutWaterTextures.Views[ViewIndex].MinMaxUV;
+			Parameters.SceneNoWaterMinMaxUV = SceneWithoutWaterTextures.Views[bIsInstancedStereoSideBySide ? View.PrimaryViewIndex : ViewIndex].MinMaxUV; // instanced view does not have rect initialized, instead the primary view covers both
 			Parameters.SceneNoWaterTextureSize = SceneWithoutWaterTextures.DepthTexture ? FVector2f(DepthTextureSize.X, DepthTextureSize.Y) : FVector2f();
 			Parameters.SceneNoWaterInvTextureSize = SceneWithoutWaterTextures.DepthTexture ? FVector2f(1.0f / DepthTextureSize.X, 1.0f / DepthTextureSize.Y) : FVector2f();
 			Parameters.SeparatedMainDirLightTexture = BlackDummyTexture;
diff --git a/Engine/Source/Runtime/Renderer/Private/VolumetricRenderTarget.cpp b/Engine/Source/Runtime/Renderer/Private/VolumetricRenderTarget.cpp
--- a/Engine/Source/Runtime/Renderer/Private/VolumetricRenderTarget.cpp
+++ b/Engine/Source/Runtime/Renderer/Private/VolumetricRenderTarget.cpp
@@ -709,7 +709,7 @@ void ComposeVolumetricRenderTargetOverSceneUnderWater(
 	for (int32 ViewIndex = 0; ViewIndex < Views.Num(); ViewIndex++)
 	{
 		FViewInfo& ViewInfo = Views[ViewIndex];
-		if (!ShouldViewRenderVolumetricCloudRenderTarget(ViewInfo))
+		if (!ShouldViewRenderVolumetricCloudRenderTarget(ViewInfo) || !ViewInfo.ShouldRenderView())
 		{
 			continue;
 		}
