Disables bloomsetup output's clear since it's a full screen pass anyway

diff --git a/Engine/Source/Runtime/Renderer/Private/PostProcess/PostProcessBloomSetup.cpp b/Engine/Source/Runtime/Renderer/Private/PostProcess/PostProcessBloomSetup.cpp
--- a/Engine/Source/Runtime/Renderer/Private/PostProcess/PostProcessBloomSetup.cpp
+++ b/Engine/Source/Runtime/Renderer/Private/PostProcess/PostProcessBloomSetup.cpp
@@ -107,7 +107,7 @@ FScreenPassTexture AddBloomSetupPass(FRDGBuilder& GraphBuilder, const FViewInfo&
 	OutputDesc.Flags |= bIsComputePass ? TexCreate_UAV : TexCreate_RenderTargetable;
 
 	const FScreenPassTextureViewport Viewport(Inputs.SceneColor);
-	const FScreenPassRenderTarget Output(GraphBuilder.CreateTexture(OutputDesc, TEXT("BloomSetup")), Viewport.Rect, View.GetOverwriteLoadAction());
+	const FScreenPassRenderTarget Output(GraphBuilder.CreateTexture(OutputDesc, TEXT("BloomSetup")), Viewport.Rect, ERenderTargetLoadAction::ENoAction);
 
 	if (bIsComputePass)
 	{
