Fixes TSR's shading rejection heuristic in split screen & VR

diff --git a/Engine/Shaders/Private/TemporalSuperResolution/TSRCommon.ush b/Engine/Shaders/Private/TemporalSuperResolution/TSRCommon.ush
--- a/Engine/Shaders/Private/TemporalSuperResolution/TSRCommon.ush
+++ b/Engine/Shaders/Private/TemporalSuperResolution/TSRCommon.ush
@@ -257,24 +257,6 @@ FScreenTransform InputPixelPosToScreenPos;
 float2 ScreenVelocityToInputPixelVelocity;
 float2 InputPixelVelocityToScreenVelocity;
 
-float2 LowFrequencyInfo_Extent;
-float2 LowFrequencyInfo_ExtentInverse;
-uint2  LowFrequencyInfo_ViewportMin;
-uint2  LowFrequencyInfo_ViewportMax;
-float2 LowFrequencyInfo_ViewportSize;
-float2 LowFrequencyInfo_ViewportSizeInverse;
-float2 LowFrequencyInfo_UVViewportBilinearMin;
-float2 LowFrequencyInfo_UVViewportBilinearMax;
-
-float2 RejectionInfo_Extent;
-float2 RejectionInfo_ExtentInverse;
-uint2  RejectionInfo_ViewportMin;
-uint2  RejectionInfo_ViewportMax;
-float2 RejectionInfo_ViewportSize;
-float2 RejectionInfo_ViewportSizeInverse;
-float2 RejectionInfo_UVViewportBilinearMin;
-float2 RejectionInfo_UVViewportBilinearMax;
-
 float2 HistoryInfo_Extent;
 float2 HistoryInfo_ExtentInverse;
 uint2  HistoryInfo_ViewportMin;
diff --git a/Engine/Shaders/Private/TemporalSuperResolution/TSRDilateRejection.usf b/Engine/Shaders/Private/TemporalSuperResolution/TSRDilateRejection.usf
--- a/Engine/Shaders/Private/TemporalSuperResolution/TSRDilateRejection.usf
+++ b/Engine/Shaders/Private/TemporalSuperResolution/TSRDilateRejection.usf
@@ -11,6 +11,8 @@
 
 //------------------------------------------------------- PARAMETERS
 
+uint4 HistoryRejectionViewport;
+
 Texture2D<tsr_half> HistoryRejectionTexture;
 RWTexture2D<tsr_half> DilatedHistoryRejectionOutput;
 
@@ -25,6 +27,7 @@ void MainCS(
 	float4 Debug = 0.0;
 	
 	tsr_short2 RejectionPixelPos = (
+		tsr_short2(HistoryRejectionViewport.xy) +
 		tsr_short2(GroupId) * tsr_short2(TILE_SIZE, TILE_SIZE) +
 		Map8x8Tile2x2Lane(GroupThreadIndex));
 
@@ -50,7 +53,7 @@ void MainCS(
 			tsr_short2 SampleInputPixelPos = RejectionPixelPos + ClampPixelOffset(
 				RejectionPixelPos,
 				Offset, Offset,
-				RejectionInfo_ViewportMin, RejectionInfo_ViewportMax - 1);
+				HistoryRejectionViewport.xy, HistoryRejectionViewport.zw - 1);
 
 			tsr_half SampleRejectionInfo = HistoryRejectionTexture[SampleInputPixelPos];
 
@@ -58,7 +61,7 @@ void MainCS(
 		}
 	}
 	
-	bool bIsValid = all(RejectionPixelPos < tsr_short2(RejectionInfo_ViewportMax));
+	bool bIsValid = all(RejectionPixelPos < tsr_short2(HistoryRejectionViewport.zw));
 	RejectionPixelPos.x = bIsValid ? RejectionPixelPos.x : ~tsr_ushort(0);
 
 	DilatedHistoryRejectionOutput[RejectionPixelPos] = HistoryRejectionInfo;
diff --git a/Engine/Shaders/Private/TemporalSuperResolution/TSRPostfilterRejection.usf b/Engine/Shaders/Private/TemporalSuperResolution/TSRPostfilterRejection.usf
--- a/Engine/Shaders/Private/TemporalSuperResolution/TSRPostfilterRejection.usf
+++ b/Engine/Shaders/Private/TemporalSuperResolution/TSRPostfilterRejection.usf
@@ -82,6 +82,7 @@ void MainCS(
 	float4 Debug = 0.0;
 
 	tsr_short2 RejectionPixelPos = (
+		tsr_short2(HistoryRejectionViewport.xy) +
 		tsr_short2(GroupId) * tsr_short2(TILE_SIZE, TILE_SIZE) +
 		Map8x8Tile2x2Lane(GroupThreadIndex));
 
diff --git a/Engine/Shaders/Private/TemporalSuperResolution/TSRRejectShading.usf b/Engine/Shaders/Private/TemporalSuperResolution/TSRRejectShading.usf
--- a/Engine/Shaders/Private/TemporalSuperResolution/TSRRejectShading.usf
+++ b/Engine/Shaders/Private/TemporalSuperResolution/TSRRejectShading.usf
@@ -418,8 +418,10 @@ void MainCS(
 		for (uint SimdIndex = 0; SimdIndex < SIMD_SIZE; SimdIndex++)
 		{
 			tsr_short2 LaneSimdPixelOffset = GetLaneSimdPixelOffset(LaneIndex, SimdIndex);
-			tsr_short2 OutputPixelPos = GroupPixelOffset + LaneSimdPixelOffset;
-			OutputPixelPos.x = select(all(and(LaneSimdPixelOffset >= TILE_OVERSCAN_PADDING, LaneSimdPixelOffset < TILE_SIZE + TILE_OVERSCAN_PADDING)), OutputPixelPos.x, ~0);
+			bool bIsNonOverscanPixel = all(and(LaneSimdPixelOffset >= TILE_OVERSCAN_PADDING, LaneSimdPixelOffset < TILE_SIZE + TILE_OVERSCAN_PADDING));
+
+			tsr_short2 OutputPixelPos = InvalidateOutputPixelPos(GroupPixelOffset + LaneSimdPixelOffset, InputInfo_ViewportMax);
+			OutputPixelPos.x = select(bIsNonOverscanPixel, OutputPixelPos.x, ~tsr_short(0));
 		
 			tsr_halfC FinalGuideColor = FinalGuide.GetElement(SimdIndex);
 
diff --git a/Engine/Shaders/Private/TemporalSuperResolution/TSRUpdateHistory.usf b/Engine/Shaders/Private/TemporalSuperResolution/TSRUpdateHistory.usf
--- a/Engine/Shaders/Private/TemporalSuperResolution/TSRUpdateHistory.usf
+++ b/Engine/Shaders/Private/TemporalSuperResolution/TSRUpdateHistory.usf
@@ -368,7 +368,7 @@ void MainCS(
 		// Issue overlapped texture fetches
 		PLATFORM_SPECIFIC_ISOLATE
 		{
-			tsr_ushort2x2 RejectionPixelPos = dpv_sub(tsr_ushort2x2(InputPixelPos), tsr_ushort2(InputPixelPosMin)) / 2;
+			tsr_ushort2x2 RejectionPixelPos = tsr_ushort2x2(InputPixelPos) / 2;
 
 			UNROLL_N(DPV_PIXEL_PER_LANE)
 			for (uint PixelId = 0; PixelId < DPV_PIXEL_PER_LANE; PixelId++)
diff --git a/Engine/Source/Runtime/Renderer/Private/PostProcess/TemporalSuperResolution.cpp b/Engine/Source/Runtime/Renderer/Private/PostProcess/TemporalSuperResolution.cpp
--- a/Engine/Source/Runtime/Renderer/Private/PostProcess/TemporalSuperResolution.cpp
+++ b/Engine/Source/Runtime/Renderer/Private/PostProcess/TemporalSuperResolution.cpp
@@ -100,8 +100,6 @@ TAutoConsoleVariable<int32> CVarTSRSetupDebugPasses(
 
 BEGIN_SHADER_PARAMETER_STRUCT(FTSRCommonParameters, )
 	SHADER_PARAMETER_STRUCT(FScreenPassTextureViewportParameters, InputInfo)
-	SHADER_PARAMETER_STRUCT(FScreenPassTextureViewportParameters, LowFrequencyInfo)
-	SHADER_PARAMETER_STRUCT(FScreenPassTextureViewportParameters, RejectionInfo)
 	SHADER_PARAMETER_STRUCT(FScreenPassTextureViewportParameters, HistoryInfo)
 
 	SHADER_PARAMETER(FIntPoint, InputPixelPosMin)
@@ -398,7 +396,7 @@ class FTSRDilateRejectionCS : public FTSRShader
 	SHADER_USE_PARAMETER_STRUCT(FTSRDilateRejectionCS, FTSRShader);
 
 	BEGIN_SHADER_PARAMETER_STRUCT(FParameters, )
-		SHADER_PARAMETER_STRUCT_INCLUDE(FTSRCommonParameters, CommonParameters)
+		SHADER_PARAMETER(FIntRect, HistoryRejectionViewport)
 		SHADER_PARAMETER_RDG_TEXTURE(Texture2D, HistoryRejectionTexture)
 		SHADER_PARAMETER_RDG_TEXTURE_UAV(RWTexture2D, DilatedHistoryRejectionOutput)
 		SHADER_PARAMETER_RDG_TEXTURE_UAV(RWTexture2DArray, DebugOutput)
@@ -661,11 +659,8 @@ ITemporalUpscaler::FOutputs AddTemporalSuperResolutionPasses(
 	FIntPoint InputExtent = PassInputs.SceneColorTexture->Desc.Extent;
 	FIntRect InputRect = View.ViewRect;
 
-	FIntPoint LowFrequencyExtent = InputExtent;
-	FIntRect LowFrequencyRect = InputRect;
-
-	FIntPoint RejectionExtent = LowFrequencyExtent / 2;
-	FIntRect RejectionRect = FIntRect(FIntPoint(0, 0), FIntPoint::DivideAndRoundUp(LowFrequencyRect.Size(), 2));
+	FIntPoint RejectionExtent = InputExtent / 2;
+	FIntRect RejectionRect = FIntRect(InputRect.Min / 2, InputRect.Min / 2 + FIntPoint::DivideAndRoundUp(InputRect.Size(), 2));
 
 	FIntPoint OutputExtent;
 	FIntRect OutputRect;
@@ -745,12 +740,6 @@ ITemporalUpscaler::FOutputs AddTemporalSuperResolutionPasses(
 		CommonParameters.ScreenVelocityToInputPixelVelocity = (FScreenTransform::Identity / CommonParameters.InputPixelPosToScreenPos).Scale;
 		CommonParameters.InputPixelVelocityToScreenVelocity = CommonParameters.InputPixelPosToScreenPos.Scale.GetAbs();
 
-		CommonParameters.LowFrequencyInfo = GetScreenPassTextureViewportParameters(FScreenPassTextureViewport(
-			LowFrequencyExtent, LowFrequencyRect));
-
-		CommonParameters.RejectionInfo = GetScreenPassTextureViewportParameters(FScreenPassTextureViewport(
-			RejectionExtent, RejectionRect));
-
 		CommonParameters.HistoryInfo = GetScreenPassTextureViewportParameters(FScreenPassTextureViewport(
 			HistoryExtent, FIntRect(FIntPoint(0, 0), HistorySize)));
 
@@ -1090,7 +1079,7 @@ ITemporalUpscaler::FOutputs AddTemporalSuperResolutionPasses(
 		PassParameters->HoleFilledVelocityMaskOutput = GraphBuilder.CreateUAV(HoleFilledVelocityMaskTexture);
 		PassParameters->ParallaxRejectionMaskOutput = GraphBuilder.CreateUAV(ParallaxRejectionMaskTexture);
 		PassParameters->HistorySubpixelDetailsOutput = GraphBuilder.CreateUAV(History.SubpixelDetails);
-		PassParameters->DebugOutput = CreateDebugUAV(LowFrequencyExtent, TEXT("Debug.TSR.DecimateHistory"));
+		PassParameters->DebugOutput = CreateDebugUAV(InputExtent, TEXT("Debug.TSR.DecimateHistory"));
 
 		TShaderMapRef<FTSRDecimateHistoryCS> ComputeShader(View.ShaderMap);
 		FComputeShaderUtils::AddPass(
@@ -1150,7 +1139,7 @@ ITemporalUpscaler::FOutputs AddTemporalSuperResolutionPasses(
 		PassParameters->PrevTranslucencyTexture = PrevTranslucencyTexture;
 
 		PassParameters->TranslucencyRejectionOutput = GraphBuilder.CreateUAV(TranslucencyRejectionTexture);
-		PassParameters->DebugOutput = CreateDebugUAV(LowFrequencyExtent, TEXT("Debug.TSR.CompareTranslucency"));
+		PassParameters->DebugOutput = CreateDebugUAV(InputExtent, TEXT("Debug.TSR.CompareTranslucency"));
 
 		TShaderMapRef<FTSRCompareTranslucencyCS> ComputeShader(View.ShaderMap);
 		FComputeShaderUtils::AddPass(
@@ -1171,7 +1160,7 @@ ITemporalUpscaler::FOutputs AddTemporalSuperResolutionPasses(
 
 		{
 			FRDGTextureDesc Desc = FRDGTextureDesc::Create2D(
-				LowFrequencyExtent,
+				InputExtent,
 				PF_R8,
 				FClearValueBinding::None,
 				/* InFlags = */ TexCreate_ShaderResource | TexCreate_UAV);
@@ -1181,7 +1170,7 @@ ITemporalUpscaler::FOutputs AddTemporalSuperResolutionPasses(
 
 		{
 			FRDGTextureDesc Desc = FRDGTextureDesc::Create2D(
-				bOutputHalfRes ? RejectionExtent : LowFrequencyExtent,
+				bOutputHalfRes ? RejectionExtent : InputExtent,
 				PF_R8,
 				FClearValueBinding::None,
 				/* InFlags = */ TexCreate_ShaderResource | TexCreate_UAV);
@@ -1199,7 +1188,7 @@ ITemporalUpscaler::FOutputs AddTemporalSuperResolutionPasses(
 		PassParameters->HistoryGuideOutput = GraphBuilder.CreateUAV(History.Guide);
 		PassParameters->HistoryRejectionOutput = GraphBuilder.CreateUAV(HistoryRejectionTexture);
 		PassParameters->InputSceneColorLdrLumaOutput = GraphBuilder.CreateUAV(InputSceneColorLdrLumaTexture);
-		PassParameters->DebugOutput = CreateDebugUAV(LowFrequencyExtent, TEXT("Debug.TSR.RejectShading"));
+		PassParameters->DebugOutput = CreateDebugUAV(InputExtent, TEXT("Debug.TSR.RejectShading"));
 
 		FTSRRejectShadingCS::FPermutationDomain PermutationVector;
 		PermutationVector.Set<FTSRRejectShadingCS::FUseWaveOps>(bUseWaveOps && GRHIMinimumWaveSize <= 32 && GRHIMaximumWaveSize >= 32);
@@ -1209,11 +1198,11 @@ ITemporalUpscaler::FOutputs AddTemporalSuperResolutionPasses(
 			GraphBuilder,
 			RDG_EVENT_NAME("TSR RejectShading(%s) %dx%d",
 				PermutationVector.Get<FTSRRejectShadingCS::FUseWaveOps>() ? TEXT("WaveOps") : TEXT(""),
-				LowFrequencyRect.Width(), LowFrequencyRect.Height()),
+				InputRect.Width(), InputRect.Height()),
 			AsyncComputePasses >= 3 ? ERDGPassFlags::AsyncCompute : ERDGPassFlags::Compute,
 			ComputeShader,
 			PassParameters,
-			FComputeShaderUtils::GetGroupCount(LowFrequencyRect.Size(), 12));
+			FComputeShaderUtils::GetGroupCount(InputRect.Size(), 12));
 	}
 
 	// Spatial anti-aliasing when doing history rejection.
@@ -1282,7 +1271,7 @@ ITemporalUpscaler::FOutputs AddTemporalSuperResolutionPasses(
 	if (PostFilter != ERejectionPostFilter::Disabled)
 	{
 		bool bOutputHalfRes = PostFilter == ERejectionPostFilter::PreRejectionDownsample;
-		FIntRect Rect = bOutputHalfRes ? FIntRect(FIntPoint(0, 0), LowFrequencyRect.Size()) : RejectionRect;
+		FIntRect Rect = bOutputHalfRes ? InputRect : RejectionRect;
 
 		FRDGTextureRef FilteredHistoryRejectionTexture;
 		{
@@ -1299,7 +1288,7 @@ ITemporalUpscaler::FOutputs AddTemporalSuperResolutionPasses(
 		PassParameters->HistoryRejectionViewport = Rect;
 		PassParameters->HistoryRejectionTexture = HistoryRejectionTexture;
 		PassParameters->HistoryRejectionOutput = GraphBuilder.CreateUAV(FilteredHistoryRejectionTexture);
-		PassParameters->DebugOutput = CreateDebugUAV(bOutputHalfRes ? LowFrequencyExtent : RejectionExtent, TEXT("Debug.TSR.PostfilterRejection"));
+		PassParameters->DebugOutput = CreateDebugUAV(bOutputHalfRes ? InputExtent : RejectionExtent, TEXT("Debug.TSR.PostfilterRejection"));
 
 		FTSRPostfilterRejectionCS::FPermutationDomain PermutationVector;
 		PermutationVector.Set<FTSRPostfilterRejectionCS::FOutputHalfRes>(PostFilter == ERejectionPostFilter::PreRejectionDownsample);
@@ -1322,7 +1311,7 @@ ITemporalUpscaler::FOutputs AddTemporalSuperResolutionPasses(
 		DilatedHistoryRejectionTexture = GraphBuilder.CreateTexture(HistoryRejectionTexture->Desc, TEXT("TSR.DilatedHistoryRejection"));
 
 		FTSRDilateRejectionCS::FParameters* PassParameters = GraphBuilder.AllocParameters<FTSRDilateRejectionCS::FParameters>();
-		PassParameters->CommonParameters = CommonParameters;
+		PassParameters->HistoryRejectionViewport = RejectionRect;
 		PassParameters->HistoryRejectionTexture = HistoryRejectionTexture;
 		PassParameters->DilatedHistoryRejectionOutput = GraphBuilder.CreateUAV(DilatedHistoryRejectionTexture);
 		PassParameters->DebugOutput = CreateDebugUAV(RejectionExtent, TEXT("Debug.TSR.DilateRejection"));
