Path Tracer: enable progress display by default and reduce the default SPP to 2048

diff --git a/Engine/Source/Runtime/Engine/Private/Scene.cpp b/Engine/Source/Runtime/Engine/Private/Scene.cpp
--- a/Engine/Source/Runtime/Engine/Private/Scene.cpp
+++ b/Engine/Source/Runtime/Engine/Private/Scene.cpp
@@ -594,7 +594,7 @@ FPostProcessSettings::FPostProcessSettings()
 	RayTracingTranslucencyRefraction = 1;
 
 	PathTracingMaxBounces = 32;
-	PathTracingSamplesPerPixel = 16384;
+	PathTracingSamplesPerPixel = 2048;
 	PathTracingFilterWidth = 3.0f;
 	PathTracingEnableEmissive = 1;
 	PathTracingMaxPathExposure = 30.0f;
diff --git a/Engine/Source/Runtime/Renderer/Private/PathTracing.cpp b/Engine/Source/Runtime/Renderer/Private/PathTracing.cpp
--- a/Engine/Source/Runtime/Renderer/Private/PathTracing.cpp
+++ b/Engine/Source/Runtime/Renderer/Private/PathTracing.cpp
@@ -227,10 +227,10 @@ TAutoConsoleVariable<int32> CVarPathTracingWiperMode(
 
 TAutoConsoleVariable<int32> CVarPathTracingProgressDisplay(
 	TEXT("r.PathTracing.ProgressDisplay"),
-	0,
-	TEXT("Enables an in-frame display of progress towards the defined sample per pixel limit. The indicator dissapears when the maximum is reached and sample accumulation has stopped (default = 0)\n")
-	TEXT("0: off (default)\n")
-	TEXT("1: on\n"),
+	1,
+	TEXT("Enables an in-frame display of progress towards the defined sample per pixel limit. The indicator dissapears when the maximum is reached and sample accumulation has stopped (default = 1)\n")
+	TEXT("0: off\n")
+	TEXT("1: on (default)\n"),
 	ECVF_RenderThreadSafe
 );
 
