Make Deferred Decals ODSC friendly.
Deferred Decals don't go through a standard mesh processor. (That might change soon).
So specific code needs adding to TryGet the shaders.

diff --git a/Engine/Source/Runtime/Renderer/Private/DecalRenderingShared.cpp b/Engine/Source/Runtime/Renderer/Private/DecalRenderingShared.cpp
--- a/Engine/Source/Runtime/Renderer/Private/DecalRenderingShared.cpp
+++ b/Engine/Source/Runtime/Renderer/Private/DecalRenderingShared.cpp
@@ -23,13 +23,14 @@ static TAutoConsoleVariable<float> CVarDecalFadeScreenSizeMultiplier(
 
 FTransientDecalRenderData::FTransientDecalRenderData(const FScene& InScene, const FDeferredDecalProxy& InDecalProxy, float InConservativeRadius)
 	: Proxy(InDecalProxy)
-	, FadeAlpha(1.0f)
+	, MaterialProxy(InDecalProxy.DecalMaterial->GetRenderProxy())
 	, ConservativeRadius(InConservativeRadius)
+	, FadeAlpha(1.0f)
 {
-	MaterialProxy = Proxy.DecalMaterial->GetRenderProxy();
-	MaterialResource = &MaterialProxy->GetMaterialWithFallback(InScene.GetFeatureLevel(), MaterialProxy);
-	check(MaterialProxy && MaterialResource);
-	BlendDesc = DecalRendering::ComputeDecalBlendDesc(InScene.GetShaderPlatform(), *MaterialResource);
+	// Build BlendDesc from a potentially incomplete material.
+	// If our shader isn't compiled yet then we will potentially render later with a different fallback material.
+	FMaterial const& MaterialResource = MaterialProxy->GetIncompleteMaterialWithFallback(InScene.GetFeatureLevel());
+	BlendDesc = DecalRendering::ComputeDecalBlendDesc(InScene.GetShaderPlatform(), MaterialResource);
 }
 
 /**
@@ -80,13 +81,11 @@ public:
 		DecalParams.Bind(Initializer.ParameterMap, TEXT("DecalParams"));
 	}
 
-	void SetParameters(FRHICommandList& RHICmdList, const FViewInfo& View, const FMaterialRenderProxy* MaterialProxy, const FDeferredDecalProxy& DecalProxy, const float FadeAlphaValue=1.0f)
+	void SetParameters(FRHICommandList& RHICmdList, const FViewInfo& View, const FDeferredDecalProxy& DecalProxy, const FMaterialRenderProxy* MaterialProxy, const FMaterial* MaterialResource, const float FadeAlphaValue = 1.0f)
 	{
 		FRHIPixelShader* ShaderRHI = RHICmdList.GetBoundPixelShader();
 
-		const FMaterialRenderProxy* MaterialProxyForRendering = MaterialProxy;
-		const FMaterial& Material = MaterialProxy->GetMaterialWithFallback(View.GetFeatureLevel(), MaterialProxyForRendering);
-		FMaterialShader::SetParameters(RHICmdList, ShaderRHI, MaterialProxyForRendering, Material, View);
+		FMaterialShader::SetParameters(RHICmdList, ShaderRHI, MaterialProxy, *MaterialResource, View);
 
 		const FLargeWorldRenderPosition AbsoluteOrigin(View.ViewMatrices.GetInvViewMatrix().GetOrigin());
 		const FVector3f TilePosition = AbsoluteOrigin.GetTile();
@@ -365,63 +364,91 @@ namespace DecalRendering
 		return ComponentToWorldMatrixTrans * View.ViewMatrices.GetTranslatedViewProjectionMatrix();
 	}
 
-	void SetShader(FRHICommandList& RHICmdList, FGraphicsPipelineStateInitializer& GraphicsPSOInit, uint32 StencilRef, const FViewInfo& View,
-		const FTransientDecalRenderData& DecalData, EDecalRenderStage DecalRenderStage, const FMatrix& FrustumComponentToClip)
+	bool TryGetDeferredDecalShaders(
+		FMaterial const& Material,
+		ERHIFeatureLevel::Type FeatureLevel,
+		EDecalRenderStage DecalRenderStage,
+		TShaderRef<FDeferredDecalPS>& OutPixelShader)
 	{
-		const FMaterialShaderMap* MaterialShaderMap = DecalData.MaterialResource->GetRenderingThreadShaderMap();
-		const EDebugViewShaderMode DebugViewMode = View.Family->GetDebugViewShaderMode();
-
-		// When in shader complexity, decals get rendered as emissive even though there might not be emissive decals.
-		// FDeferredDecalEmissivePS might not be available depending on the decal blend mode.
-		TShaderRef<FDeferredDecalPS> PixelShader;
-		if (DecalRenderStage == EDecalRenderStage::Emissive || DebugViewMode != DVSM_None)
+		FMaterialShaderTypes ShaderTypes;
+		
+		if (DecalRenderStage == EDecalRenderStage::Emissive)
 		{
-			PixelShader = TShaderRef<FDeferredDecalPS>(MaterialShaderMap->GetShader<FDeferredDecalEmissivePS>());
+			ShaderTypes.AddShaderType<FDeferredDecalEmissivePS>();
 		}
 		else if (DecalRenderStage == EDecalRenderStage::AmbientOcclusion)
 		{
-			PixelShader = TShaderRef<FDeferredDecalPS>(MaterialShaderMap->GetShader<FDeferredDecalAmbientOcclusionPS>());
+			ShaderTypes.AddShaderType<FDeferredDecalAmbientOcclusionPS>();
 		}
 		else
 		{
-			PixelShader = MaterialShaderMap->GetShader<FDeferredDecalPS>();
+			ShaderTypes.AddShaderType<FDeferredDecalPS>();
 		}
-		check(!PixelShader.IsNull());
 
-		TShaderMapRef<FDeferredDecalVS> VertexShader(View.ShaderMap);
+		FMaterialShaders Shaders;
+		if (!Material.TryGetShaders(ShaderTypes, nullptr, Shaders))
+		{
+			return false;
+		}
 
+		Shaders.TryGetPixelShader(OutPixelShader);
+		return OutPixelShader.IsValid();
+	}
+
+	FMaterialRenderProxy const* TryGetDeferredDecalMaterial(
+		FMaterialRenderProxy const* MaterialProxy, 
+		ERHIFeatureLevel::Type FeatureLevel,
+		EDecalRenderStage DecalRenderStage,
+		FMaterial const*& OutMaterialResource,
+		TShaderRef<FDeferredDecalPS>& OutPixelShader)
+	{
+		OutMaterialResource = nullptr;
+
+		while (MaterialProxy != nullptr)
 		{
-			GraphicsPSOInit.BoundShaderState.VertexDeclarationRHI = GetVertexDeclarationFVector4();
-			GraphicsPSOInit.BoundShaderState.VertexShaderRHI = VertexShader.GetVertexShader();
-			GraphicsPSOInit.BoundShaderState.PixelShaderRHI = PixelShader.GetPixelShader();
-			GraphicsPSOInit.PrimitiveType = PT_TriangleList;
+			OutMaterialResource = MaterialProxy->GetMaterialNoFallback(FeatureLevel);
+			if (OutMaterialResource != nullptr)
+			{
+				if (TryGetDeferredDecalShaders(*OutMaterialResource, FeatureLevel, DecalRenderStage, OutPixelShader))
+				{
+					break;
+				}
+			}
 
-			SetGraphicsPipelineState(RHICmdList, GraphicsPSOInit, StencilRef);
+			MaterialProxy = MaterialProxy->GetFallback(FeatureLevel);
 		}
 
+		return MaterialProxy;
+	}
+
+	void SetShader(FRHICommandList& RHICmdList, FGraphicsPipelineStateInitializer& GraphicsPSOInit, uint32 StencilRef, const FViewInfo& View,
+		const FTransientDecalRenderData& DecalData, EDecalRenderStage DecalRenderStage, const FMatrix& FrustumComponentToClip)
+	{
+		FMaterial const* MaterialResource = nullptr;
+		TShaderRef<FDeferredDecalPS> PixelShader;
+		FMaterialRenderProxy const* MaterialProxy = TryGetDeferredDecalMaterial(DecalData.MaterialProxy, View.GetFeatureLevel(), DecalRenderStage, MaterialResource, PixelShader);
+
+		TShaderMapRef<FDeferredDecalVS> VertexShader(View.ShaderMap);
+
+		GraphicsPSOInit.BoundShaderState.VertexDeclarationRHI = GetVertexDeclarationFVector4();
+		GraphicsPSOInit.BoundShaderState.VertexShaderRHI = VertexShader.GetVertexShader();
+		GraphicsPSOInit.BoundShaderState.PixelShaderRHI = PixelShader.GetPixelShader();
+		GraphicsPSOInit.PrimitiveType = PT_TriangleList;
+		SetGraphicsPipelineState(RHICmdList, GraphicsPSOInit, StencilRef);
+
 		// Set vertex shader parameters.
 		{
 			FDeferredDecalVS::FParameters ShaderParameters;
-			ShaderParameters.FrustumComponentToClip = FMatrix44f(FrustumComponentToClip);			// LWC_TODO: Precision loss?
+			ShaderParameters.FrustumComponentToClip = FMatrix44f(FrustumComponentToClip); // LWC_TODO: Precision loss?
 			ShaderParameters.PrimitiveUniformBuffer = GIdentityPrimitiveUniformBuffer.GetUniformBufferRef();
-
 			SetShaderParameters(RHICmdList, VertexShader, VertexShader.GetVertexShader(), ShaderParameters);
 		}
 
 		// Set pixel shader parameters.
 		{
-			PixelShader->SetParameters(RHICmdList, View, DecalData.MaterialProxy, DecalData.Proxy, DecalData.FadeAlpha);
-
+			PixelShader->SetParameters(RHICmdList, View, DecalData.Proxy, MaterialProxy, MaterialResource, DecalData.FadeAlpha);
 			auto& PrimitivePS = PixelShader->GetUniformBufferParameter<FPrimitiveUniformShaderParameters>();
-
-			// uncomment to track down usage of the Primitive uniform buffer
-			//	check(!PrimitiveVS.IsBound());
-			//	check(!PrimitivePS.IsBound());
-
-			if (DebugViewMode == DVSM_None)
-			{
-				SetUniformBufferParameter(RHICmdList, PixelShader.GetPixelShader(), PrimitivePS, GIdentityPrimitiveUniformBuffer);
-			}
+			SetUniformBufferParameter(RHICmdList, PixelShader.GetPixelShader(), PrimitivePS, GIdentityPrimitiveUniformBuffer);
 		}
 
 		// Set stream source after updating cached strides
@@ -435,17 +462,14 @@ namespace DecalRendering
 		GraphicsPSOInit.BoundShaderState.VertexDeclarationRHI = GetVertexDeclarationFVector4();
 		GraphicsPSOInit.BoundShaderState.VertexShaderRHI = VertexShader.GetVertexShader();
 		GraphicsPSOInit.PrimitiveType = PT_TriangleList;
-
 		SetGraphicsPipelineState(RHICmdList, GraphicsPSOInit, 0);
 
 		// Set vertex shader parameters.
 		{
 			FDeferredDecalVS::FParameters ShaderParameters;
-			ShaderParameters.FrustumComponentToClip = FMatrix44f(FrustumComponentToClip);	// LWC_TODO: Precision loss
+			ShaderParameters.FrustumComponentToClip = FMatrix44f(FrustumComponentToClip); // LWC_TODO: Precision loss
 			ShaderParameters.PrimitiveUniformBuffer = GIdentityPrimitiveUniformBuffer.GetUniformBufferRef();
-
 			SetShaderParameters(RHICmdList, VertexShader, VertexShader.GetVertexShader(), ShaderParameters);
 		}
-
 	}
 }
diff --git a/Engine/Source/Runtime/Renderer/Private/DecalRenderingShared.h b/Engine/Source/Runtime/Renderer/Private/DecalRenderingShared.h
--- a/Engine/Source/Runtime/Renderer/Private/DecalRenderingShared.h
+++ b/Engine/Source/Runtime/Renderer/Private/DecalRenderingShared.h
@@ -17,12 +17,11 @@
  */
 struct FTransientDecalRenderData
 {
-	const FMaterialRenderProxy* MaterialProxy;
-	const FMaterial* MaterialResource;
 	const FDeferredDecalProxy& Proxy;
+	const FMaterialRenderProxy* MaterialProxy;
 	FDecalBlendDesc BlendDesc;
-	float FadeAlpha;
 	float ConservativeRadius;
+	float FadeAlpha;
 
 	FTransientDecalRenderData(const FScene& InScene, const FDeferredDecalProxy& InDecalProxy, float InConservativeRadius);
 };
diff --git a/Engine/Source/Runtime/Renderer/Private/RayTracing/RayTracingDecals.cpp b/Engine/Source/Runtime/Renderer/Private/RayTracing/RayTracingDecals.cpp
--- a/Engine/Source/Runtime/Renderer/Private/RayTracing/RayTracingDecals.cpp
+++ b/Engine/Source/Runtime/Renderer/Private/RayTracing/RayTracingDecals.cpp
@@ -364,7 +364,9 @@ TRDGUniformBufferRef<FRayTracingDecals> CreateRayTracingDecalData(FRDGBuilder& G
 	{
 		const auto& DecalData = SortedDecals[Index];
 
-		const FMaterialShaderMap* MaterialShaderMap = DecalData.MaterialResource->GetRenderingThreadShaderMap();
+		const FMaterialRenderProxy* MaterialProxy = DecalData.MaterialProxy;
+		const FMaterial* MaterialResource = &MaterialProxy->GetMaterialWithFallback(View.GetFeatureLevel(), MaterialProxy);
+		const FMaterialShaderMap* MaterialShaderMap = MaterialResource->GetRenderingThreadShaderMap();
 
 		auto CallableShader = MaterialShaderMap->GetShader<FRayTracingDecalMaterialShader>();
 
@@ -383,7 +385,7 @@ TRDGUniformBufferRef<FRayTracingDecals> CreateRayTracingDecalData(FRDGBuilder& G
 
 		int32 DataOffset = 0;
 		FMeshDrawSingleShaderBindings SingleShaderBindings = Command.ShaderBindings.GetSingleShaderBindings(SF_RayCallable, DataOffset);
-		CallableShader->GetShaderBindings(&Scene, Scene.GetFeatureLevel(), *DecalData.MaterialProxy, *DecalData.MaterialResource, View, DecalParameters, SingleShaderBindings);
+		CallableShader->GetShaderBindings(&Scene, Scene.GetFeatureLevel(), *MaterialProxy, *MaterialResource, View, DecalParameters, SingleShaderBindings);
 
 		const FBox BoxBounds = DecalData.Proxy.GetBounds().GetBox();
 
