Fix Instanced view UB not being available in ISR.

diff --git a/Engine/Source/Runtime/Renderer/Private/SingleLayerWaterRendering.cpp b/Engine/Source/Runtime/Renderer/Private/SingleLayerWaterRendering.cpp
--- a/Engine/Source/Runtime/Renderer/Private/SingleLayerWaterRendering.cpp
+++ b/Engine/Source/Runtime/Renderer/Private/SingleLayerWaterRendering.cpp
@@ -160,7 +160,7 @@ BEGIN_SHADER_PARAMETER_STRUCT(FSingleLayerWaterCommonShaderParameters, )
 	SHADER_PARAMETER(FVector2f, SceneNoWaterInvTextureSize)
 	SHADER_PARAMETER(float, UseSeparatedMainDirLightTexture)
 	SHADER_PARAMETER_STRUCT_INCLUDE(FSceneTextureParameters, SceneTextures)	// Water scene texture
-	SHADER_PARAMETER_STRUCT_REF(FViewUniformShaderParameters, ViewUniformBuffer)
+	SHADER_PARAMETER_STRUCT_INCLUDE(FViewShaderParameters, View)
 	SHADER_PARAMETER_STRUCT_REF(FReflectionCaptureShaderData, ReflectionCaptureData)
 	SHADER_PARAMETER_STRUCT_REF(FReflectionUniformParameters, ReflectionsParameters)
 	SHADER_PARAMETER_RDG_UNIFORM_BUFFER(FForwardLightData, ForwardLightData)
@@ -257,7 +257,7 @@ class FWaterRefractionCopyPS : public FGlobalShader
 	SHADER_USE_PARAMETER_STRUCT(FWaterRefractionCopyPS, FGlobalShader);
 
 	BEGIN_SHADER_PARAMETER_STRUCT(FParameters, )
-		SHADER_PARAMETER_STRUCT_REF(FViewUniformShaderParameters, ViewUniformBuffer)
+		SHADER_PARAMETER_STRUCT_INCLUDE(FViewShaderParameters, View)
 		SHADER_PARAMETER_RDG_TEXTURE(Texture2D, SceneColorCopyDownsampleTexture)
 		SHADER_PARAMETER_SAMPLER(SamplerState, SceneColorCopyDownsampleSampler)
 		SHADER_PARAMETER_RDG_TEXTURE(Texture2D, SceneDepthCopyDownsampleTexture)
@@ -344,7 +344,7 @@ static FSceneWithoutWaterTextures AddCopySceneWithoutWaterPass(
 		RDG_EVENT_SCOPE_CONDITIONAL(GraphBuilder, Views.Num() > 1, "View%d", ViewIndex);
 
 		FWaterRefractionCopyPS::FParameters* PassParameters = GraphBuilder.AllocParameters<FWaterRefractionCopyPS::FParameters>();
-		PassParameters->ViewUniformBuffer = View.ViewUniformBuffer;
+		PassParameters->View = View.GetShaderParameters();
 		PassParameters->SceneColorCopyDownsampleTexture = SceneColorTexture;
 		PassParameters->SceneColorCopyDownsampleSampler = TStaticSamplerState<SF_Bilinear, AM_Clamp, AM_Clamp, AM_Clamp>::GetRHI();
 		PassParameters->SceneDepthCopyDownsampleTexture = SceneDepthTexture;
@@ -468,7 +468,7 @@ void FDeferredShadingSceneRenderer::RenderSingleLayerWaterReflections(
 			Parameters.SeparatedMainDirLightTexture = BlackDummyTexture;
 			Parameters.UseSeparatedMainDirLightTexture = 0.0f;
 			Parameters.SceneTextures = SceneTextureParameters;
-			Parameters.ViewUniformBuffer = GetShaderBinding(View.ViewUniformBuffer);
+			Parameters.View = View.GetShaderParameters();
 			Parameters.ReflectionCaptureData = View.ReflectionCaptureUniformBuffer;
 			{
 				FReflectionUniformParameters ReflectionUniformParameters;
