Fixed usage of ERayTracingPrimitiveFlags::UnsupportedProxyType

* Its value is 0, so it must not be used with EnumHasAnyFlags()
* Added few more comments in related code

diff --git a/Engine/Source/Runtime/Engine/Public/PrimitiveSceneProxy.h b/Engine/Source/Runtime/Engine/Public/PrimitiveSceneProxy.h
--- a/Engine/Source/Runtime/Engine/Public/PrimitiveSceneProxy.h
+++ b/Engine/Source/Runtime/Engine/Public/PrimitiveSceneProxy.h
@@ -130,14 +130,18 @@ extern bool CacheShadowDepthsFromPrimitivesUsingWPO();
 
 enum class ERayTracingPrimitiveFlags : uint8
 {
-	// Visibility flags
+	// Visibility flags:
+
 	// This type of geometry is not supported in ray tracing and all proxies will be excluded
 	// If a proxy decides to return UnsupportedProxyType it must be consistent across all proxies of this type
+	// This value is used for primitives that return false from FPrimitiveSceneProxy::IsRayTracingRelevant().
 	UnsupportedProxyType = 0 << 0,
+
 	// This scene proxy will be excluded, because it decides to be invisible in ray tracing (probably due to other flags)
 	Excluded = 1 << 0,
 
-	// Caching flags
+	// Caching flags:
+
 	// Fully dynamic (the ray tracing representation of this scene proxy will be polled every frame)
 	Dynamic = 1 << 1,
 	// Ray tracing mesh commmands generated from this proxy's materials can be cached
@@ -145,7 +149,8 @@ enum class ERayTracingPrimitiveFlags : uint8
 	// Instances from this proxy can be cached
 	CacheInstances = 1 << 3,
 
-	// Misc flags
+	// Misc flags:
+
 	// Static meshes with multiple LODs will want to select a LOD index based on screen size
 	ComputeLOD = 1 << 4, 
 
diff --git a/Engine/Source/Runtime/Renderer/Private/DeferredShadingRenderer.cpp b/Engine/Source/Runtime/Renderer/Private/DeferredShadingRenderer.cpp
--- a/Engine/Source/Runtime/Renderer/Private/DeferredShadingRenderer.cpp
+++ b/Engine/Source/Runtime/Renderer/Private/DeferredShadingRenderer.cpp
@@ -664,19 +664,23 @@ bool FDeferredShadingSceneRenderer::GatherRayTracingWorldInstancesForView(FRDGBu
 	{
 		TRACE_CPUPROFILER_EVENT_SCOPE(GatherRayTracingWorldInstances_RelevantPrimitives);
 
+		// Index into the TypeOffsetTable, which contains a prefix sum of primitive indices by proxy type
 		int32 BroadIndex = 0;
 
 		for (int PrimitiveIndex = 0; PrimitiveIndex < Scene->PrimitiveSceneProxies.Num(); PrimitiveIndex++)
 		{
+			// Find the next TypeOffsetTable entry that's relevant to this primitive idnex.
 			while (PrimitiveIndex >= int(Scene->TypeOffsetTable[BroadIndex].Offset))
 			{
 				BroadIndex++;
 			}
 
 			// Skip before dereferencing SceneInfo
-			if (EnumHasAnyFlags(Scene->PrimitiveRayTracingFlags[PrimitiveIndex], ERayTracingPrimitiveFlags::UnsupportedProxyType))
+			if (Scene->PrimitiveRayTracingFlags[PrimitiveIndex] == ERayTracingPrimitiveFlags::UnsupportedProxyType)
 			{
-				//skip over unsupported SceneProxies (warning don't make IsRayTracingRelevant data dependent other than the vtable)
+				// Find the index of a proxy of the next type, skipping over a batch of proxies that are the same type as current.
+				// This assumes that FPrimitiveSceneProxy::IsRayTracingRelevant() is consistent for all proxies of the same type.
+				// I.e. does not depend on members of the particular FPrimitiveSceneProxy implementation.
 				PrimitiveIndex = Scene->TypeOffsetTable[BroadIndex].Offset - 1;
 				continue;
 			}
diff --git a/Engine/Source/Runtime/Renderer/Private/RayTracing/RayTracingInstanceCulling.cpp b/Engine/Source/Runtime/Renderer/Private/RayTracing/RayTracingInstanceCulling.cpp
--- a/Engine/Source/Runtime/Renderer/Private/RayTracing/RayTracingInstanceCulling.cpp
+++ b/Engine/Source/Runtime/Renderer/Private/RayTracing/RayTracingInstanceCulling.cpp
@@ -81,7 +81,7 @@ namespace RayTracing
 {
 bool CullPrimitiveByFlags(const FRayTracingCullingParameters& CullingParameters, const FScene* RESTRICT Scene, int32 PrimitiveIndex)
 {
-	if (EnumHasAnyFlags(Scene->PrimitiveRayTracingFlags[PrimitiveIndex], ERayTracingPrimitiveFlags::UnsupportedProxyType))
+	if (Scene->PrimitiveRayTracingFlags[PrimitiveIndex] == ERayTracingPrimitiveFlags::UnsupportedProxyType)
 	{
 		return true;
 	}
