Fix a bug that the MSAA is broken on android.
Clean obsolete codes in mobile postprocess.

diff --git a/Engine/Source/Runtime/Core/Private/HAL/ConsoleManager.cpp b/Engine/Source/Runtime/Core/Private/HAL/ConsoleManager.cpp
--- a/Engine/Source/Runtime/Core/Private/HAL/ConsoleManager.cpp
+++ b/Engine/Source/Runtime/Core/Private/HAL/ConsoleManager.cpp
@@ -2650,7 +2650,7 @@ static TAutoConsoleVariable<int32> CVarMobileEnabledShadingModelsMask(
 	ECVF_ReadOnly | ECVF_RenderThreadSafe
 );
 
-static TAutoConsoleVariable<int32> CVarMobileForwardSupportLocalLights(
+static TAutoConsoleVariable<int32> CVarMobileForwardEnableLocalLights(
 	TEXT("r.Mobile.Forward.EnableLocalLights"),
 	1,
 	TEXT("Enable local lights support on mobile forward. 0 is disabled, 1 is enabled (default)"),
@@ -2984,14 +2984,6 @@ static TAutoConsoleVariable<int32> CVarMobileDesiredResY(
 	TEXT("Desired mobile Y resolution (shortest axis) (non-zero == use for Y, calculate X to retain aspect ratio)"),
 	ECVF_Default);
 
-static TAutoConsoleVariable<int32> CVarMobileTonemapperUpscale(
-	TEXT("r.MobileTonemapperUpscale"),
-	0,
-	TEXT("On mobile, whether to allow upscaling as part of the tonemapper or as a separate pass when possible"
-		 "0: separate pass (default)\n"
-		 "1: as part of the tonemapper pass\n"),
-	ECVF_Default);
-
 // this cvar can be removed in shipping to not compile shaders for development (faster)
 static TAutoConsoleVariable<int32> CVarCompileShadersForDevelopment(
 	TEXT("r.CompileShadersForDevelopment"),
diff --git a/Engine/Source/Runtime/RenderCore/Private/RenderGraphUtils.cpp b/Engine/Source/Runtime/RenderCore/Private/RenderGraphUtils.cpp
--- a/Engine/Source/Runtime/RenderCore/Private/RenderGraphUtils.cpp
+++ b/Engine/Source/Runtime/RenderCore/Private/RenderGraphUtils.cpp
@@ -192,9 +192,16 @@ RENDERCORE_API FRDGTextureMSAA CreateTextureMSAA(
 	const TCHAR* Name,
 	ETextureCreateFlags ResolveFlagsToAdd)
 {
+	bool bForceSeparateTargetAndShaderResource = Desc.NumSamples > 1 && RHISupportsSeparateMSAAAndResolveTextures(GMaxRHIShaderPlatform);
+	
+	if (!bForceSeparateTargetAndShaderResource)
+	{
+		Desc.Flags |= TexCreate_ShaderResource;
+	}
+
 	FRDGTextureMSAA Texture(GraphBuilder.CreateTexture(Desc, Name));
 
-	if (Desc.NumSamples > 1)
+	if (bForceSeparateTargetAndShaderResource)
 	{
 		Desc.NumSamples = 1;
 		ETextureCreateFlags ResolveFlags = TexCreate_ShaderResource;
diff --git a/Engine/Source/Runtime/Renderer/Private/MobileShadingRenderer.cpp b/Engine/Source/Runtime/Renderer/Private/MobileShadingRenderer.cpp
--- a/Engine/Source/Runtime/Renderer/Private/MobileShadingRenderer.cpp
+++ b/Engine/Source/Runtime/Renderer/Private/MobileShadingRenderer.cpp
@@ -431,12 +431,19 @@ void FMobileSceneRenderer::InitViews(FRDGBuilder& GraphBuilder, FSceneTexturesCo
 	PostVisibilityFrameSetup(ILCTaskData);
 
 	const FIntPoint RenderTargetSize = (ViewFamily.RenderTarget->GetRenderTargetTexture().IsValid()) ? ViewFamily.RenderTarget->GetRenderTargetTexture()->GetSizeXY() : ViewFamily.RenderTarget->GetSizeXY();
-	const bool bRequiresUpscale = ((int32)RenderTargetSize.X > FamilySize.X || (int32)RenderTargetSize.Y > FamilySize.Y) || IsMobilePropagateAlphaEnabled(ViewFamily.GetShaderPlatform());
+	const bool bRequiresUpscale = ((int32)RenderTargetSize.X > FamilySize.X || (int32)RenderTargetSize.Y > FamilySize.Y);
 	// ES requires that the back buffer and depth match dimensions.
 	// For the most part this is not the case when using scene captures. Thus scene captures always render to scene color target.
 	const bool bShouldCompositeEditorPrimitives = FSceneRenderer::ShouldCompositeEditorPrimitives(Views[0]);
 	const bool bStereoRenderingAndHMD = ViewFamily.EngineShowFlags.StereoRendering && ViewFamily.EngineShowFlags.HMDDistortion;
-	bRenderToSceneColor = !bGammaSpace || bStereoRenderingAndHMD || bRequiresUpscale || bShouldCompositeEditorPrimitives || Views[0].bIsSceneCapture || Views[0].bIsReflectionCapture;
+	bRenderToSceneColor = !bGammaSpace 
+						|| bStereoRenderingAndHMD 
+						|| bRequiresUpscale 
+						|| bShouldCompositeEditorPrimitives 
+						|| Views[0].bIsSceneCapture 
+						|| Views[0].bIsReflectionCapture 
+						// If the resolve texture is not the same as the MSAA texture, we need to render to scene color and copy to back buffer.
+						|| (NumMSAASamples > 1 && !RHISupportsSeparateMSAAAndResolveTextures(ShaderPlatform));
 	const FPlanarReflectionSceneProxy* PlanarReflectionSceneProxy = Scene ? Scene->GetForwardPassGlobalPlanarReflection() : nullptr;
 
 	bRequiresPixelProjectedPlanarRelfectionPass = IsUsingMobilePixelProjectedReflection(ShaderPlatform)
@@ -944,7 +951,7 @@ void FMobileSceneRenderer::Render(FRDGBuilder& GraphBuilder)
 	
 	if (ViewFamily.bResolveScene)
 	{
-		if (!bGammaSpace || bRenderToSceneColor)
+		if (bRenderToSceneColor)
 		{
 			// Finish rendering for each view, or the full stereo buffer if enabled
 			{
@@ -1011,12 +1018,12 @@ void FMobileSceneRenderer::RenderForward(FRDGBuilder& GraphBuilder, FRDGTextureR
 	// Verify using both MSAA sample count AND the scene color surface sample count, since on GLES you can't have MSAA color targets,
 	// so the color target would be created without MSAA, and MSAA is achieved through magical means (the framebuffer, being MSAA,
 	// tells the GPU "execute this renderpass as MSAA, and when you're done, automatically resolve and copy into this non-MSAA texture").
-	bool bMobileMSAA = NumMSAASamples > 1 && SceneTextures.Config.NumSamples > 1;
+	bool bMobileMSAA = NumMSAASamples > 1;
 
 	static const auto CVarMobileMultiView = IConsoleManager::Get().FindTConsoleVariableDataInt(TEXT("vr.MobileMultiView"));
 	const bool bIsMultiViewApplication = (CVarMobileMultiView && CVarMobileMultiView->GetValueOnAnyThread() != 0);
 	
-	if (bGammaSpace && !bRenderToSceneColor)
+	if (!bRenderToSceneColor)
 	{
 		if (bMobileMSAA)
 		{
diff --git a/Engine/Source/Runtime/Renderer/Private/PostProcess/PostProcessing.cpp b/Engine/Source/Runtime/Renderer/Private/PostProcess/PostProcessing.cpp
--- a/Engine/Source/Runtime/Renderer/Private/PostProcess/PostProcessing.cpp
+++ b/Engine/Source/Runtime/Renderer/Private/PostProcess/PostProcessing.cpp
@@ -1814,12 +1814,6 @@ void AddMobilePostProcessingPasses(FRDGBuilder& GraphBuilder, FScene* Scene, con
 	FScreenPassTexture BloomOutput;
 	FScreenPassTexture DofOutput;
 	FScreenPassTexture PostProcessSunShaftAndDof;
-
-	// temporary solution for SP_METAL using HW sRGB flag during read vs all other mob platforms using
-	// incorrect UTexture::SRGB state. (UTexture::SRGB != HW texture state)
-	bool bSRGBAwareTarget = View.Family->RenderTarget->GetDisplayGamma() == 1.0f
-		&& View.bIsSceneCapture
-		&& IsMetalMobilePlatform(View.GetShaderPlatform());
 	
 	const EAutoExposureMethod AutoExposureMethod = GetAutoExposureMethod(View);
 	const bool bUseEyeAdaptation = IsMobileEyeAdaptationEnabled(View);
@@ -1838,10 +1832,8 @@ void AddMobilePostProcessingPasses(FRDGBuilder& GraphBuilder, FScene* Scene, con
 
 	bool bUseHighResolutionScreenshotMask = IsHighResolutionScreenshotMaskEnabled(View);
 
-	static const auto VarTonemapperUpscale = IConsoleManager::Get().FindTConsoleVariableDataInt(TEXT("r.MobileTonemapperUpscale"));
-	bool bDisableUpscaleInTonemapper = !VarTonemapperUpscale || VarTonemapperUpscale->GetValueOnRenderThread() == 0;
-
-	bool bShouldPrimaryUpscale = IsMobilePropagateAlphaEnabled(View.GetShaderPlatform()) || (View.PrimaryScreenPercentageMethod == EPrimaryScreenPercentageMethod::SpatialUpscale && View.UnscaledViewRect != View.ViewRect);
+	bool bShouldPrimaryUpscale = (View.PrimaryScreenPercentageMethod == EPrimaryScreenPercentageMethod::SpatialUpscale && View.UnscaledViewRect != View.ViewRect)
+								|| PaniniConfig.IsEnabled();
 
 	PassSequence.SetEnabled(EPass::Tonemap, bUseToneMapper);
 	PassSequence.SetEnabled(EPass::HighResolutionScreenshotMask, bUseHighResolutionScreenshotMask);
@@ -1852,7 +1844,7 @@ void AddMobilePostProcessingPasses(FRDGBuilder& GraphBuilder, FScene* Scene, con
 	PassSequence.SetEnabled(EPass::SelectionOutline, false);
 	PassSequence.SetEnabled(EPass::EditorPrimitive, false);
 #endif
-	PassSequence.SetEnabled(EPass::PrimaryUpscale, PaniniConfig.IsEnabled() || (bShouldPrimaryUpscale && bDisableUpscaleInTonemapper));
+	PassSequence.SetEnabled(EPass::PrimaryUpscale, bShouldPrimaryUpscale);
 
 	PassSequence.SetEnabled(EPass::Visualize, View.Family->EngineShowFlags.ShaderComplexity);
 
@@ -2502,6 +2494,12 @@ void AddMobilePostProcessingPasses(FRDGBuilder& GraphBuilder, FScene* Scene, con
 
 		SceneColor = AddHMDDistortionPass(GraphBuilder, View, PassInputs);
 	}
+
+	// Copy the scene color to back buffer in case there is no post process, such as LDR MSAA.
+	if (SceneColor.Texture != ViewFamilyOutput.Texture)
+	{
+		AddDrawTexturePass(GraphBuilder, View, SceneColor, ViewFamilyOutput);
+	}
 }
 
 FRDGTextureRef AddProcessPlanarReflectionPass(
diff --git a/Engine/Source/Runtime/Renderer/Private/SceneTextures.cpp b/Engine/Source/Runtime/Renderer/Private/SceneTextures.cpp
--- a/Engine/Source/Runtime/Renderer/Private/SceneTextures.cpp
+++ b/Engine/Source/Runtime/Renderer/Private/SceneTextures.cpp
@@ -697,15 +697,7 @@ void FMinimalSceneTextures::InitializeViewFamily(FRDGBuilder& GraphBuilder, FVie
 							 FRDGTextureDesc::Create2DArray(Config.Extent, Config.ColorFormat, Config.ColorClearValue, Flags, 2) :
 							 FRDGTextureDesc::Create2D(Config.Extent, Config.ColorFormat, Config.ColorClearValue, Flags));
 		Desc.NumSamples = Config.NumSamples;
-		SceneTextures.Color = GraphBuilder.CreateTexture(Desc, SceneColorName);
-
-		if (Desc.NumSamples > 1)
-		{
-			Desc.NumSamples = 1;
-			Desc.Flags = TexCreate_ResolveTargetable | TexCreate_ShaderResource | GFastVRamConfig.SceneColor | sRGBFlag;
-
-			SceneTextures.Color.Resolve = GraphBuilder.CreateTexture(Desc, SceneColorName);
-		}
+		SceneTextures.Color = CreateTextureMSAA(GraphBuilder, Desc, SceneColorName, GFastVRamConfig.SceneColor | sRGBFlag);
 	}
 
 	// Custom Depth
@@ -812,14 +804,7 @@ void FSceneTextures::InitializeViewFamily(FRDGBuilder& GraphBuilder, FViewFamily
 			FRDGTextureDesc::Create2DArray(Config.Extent, DepthAuxFormat, FClearValueBinding(FarDepthColor), TexCreate_RenderTargetable | TexCreate_ShaderResource | TexCreate_InputAttachmentRead | FlagsToAdd, 2) :
 			FRDGTextureDesc::Create2D(Config.Extent, DepthAuxFormat, FClearValueBinding(FarDepthColor), TexCreate_RenderTargetable | TexCreate_ShaderResource | TexCreate_InputAttachmentRead| FlagsToAdd);
 		Desc.NumSamples = Config.NumSamples;
-		SceneTextures.DepthAux = GraphBuilder.CreateTexture(Desc, TEXT("SceneDepthAux"));
-		
-		if (Desc.NumSamples > 1)
-		{
-			Desc.NumSamples = 1;
-			Desc.Flags = TexCreate_ResolveTargetable | TexCreate_ShaderResource;
-			SceneTextures.DepthAux.Resolve = GraphBuilder.CreateTexture(Desc, TEXT("SceneDepthAux"));
-		}
+		SceneTextures.DepthAux = CreateTextureMSAA(GraphBuilder, Desc, TEXT("SceneDepthAux"));
 	}
 #if WITH_EDITOR
 	{
