PR #9138: Spellcheck (Contributed by [REDACTED_PERSON])

diff --git a/Engine/Source/Runtime/Renderer/Private/VelocityRendering.cpp b/Engine/Source/Runtime/Renderer/Private/VelocityRendering.cpp
--- a/Engine/Source/Runtime/Renderer/Private/VelocityRendering.cpp
+++ b/Engine/Source/Runtime/Renderer/Private/VelocityRendering.cpp
@@ -71,14 +71,14 @@ inline void ValidateVelocityCVars()
 			const int32 Value = CVarBasePassOutputsVelocity.GetValueOnAnyThread();
 			if (Value != -1)
 			{
-				UE_LOG(LogRenderer, Warning, TEXT("Deprectaed CVar r.BasePassOutputsVelocity is set to %d. Remove and use r.VelocityOutputPass instead."), Value);
+				UE_LOG(LogRenderer, Warning, TEXT("Deprecated CVar r.BasePassOutputsVelocity is set to %d. Remove and use r.VelocityOutputPass instead."), Value);
 			}
 		}
 		{
 			const int32 Value = CVarVertexDeformationOutputsVelocity.GetValueOnAnyThread();
 			if (Value != -1)
 			{
-				UE_LOG(LogRenderer, Warning, TEXT("Deprectaed CVar r.VertexDeformationOutputsVelocity is set to %d. Remove and use r.Velocity.EnableVertexDeformation instead."), Value);
+				UE_LOG(LogRenderer, Warning, TEXT("Deprecated CVar r.VertexDeformationOutputsVelocity is set to %d. Remove and use r.Velocity.EnableVertexDeformation instead."), Value);
 			}
 		}
 		bHasValidatedCVars = true;
