Reimplements r.TSR.History.SeparateTranslucency=0 as it is a lot more praticable with new shading rejection

diff --git a/Engine/Shaders/Private/TemporalSuperResolution/TSRRejectShading.usf b/Engine/Shaders/Private/TemporalSuperResolution/TSRRejectShading.usf
--- a/Engine/Shaders/Private/TemporalSuperResolution/TSRRejectShading.usf
+++ b/Engine/Shaders/Private/TemporalSuperResolution/TSRRejectShading.usf
@@ -29,6 +29,9 @@
 /** Whether should remove firefliers from input and history before comparison to improve stability. */
 #define CONFIG_FIREFLY_REMOVAL 1
 
+/** Whether to compose translucency to the input. */
+#define CONFIG_COMPOSE_TRANSLUCENCY (!DIM_SEPARATE_TRANSLUCENCY)
+
 
 //------------------------------------------------------- DERIVED
 
@@ -47,15 +50,20 @@
 //------------------------------------------------------- TYPEDEFS
 
 #define tsr_simd_half  TLaneVector<tsr_half, 1, SIMD_SIZE>
+#define tsr_simd_half4 TLaneVector<tsr_half, 4, SIMD_SIZE>
 #define tsr_simd_halfC TLaneVector<tsr_half, CONFIG_CHANNEL_COUNT, SIMD_SIZE>
 #define tsr_simd_bool  TLaneVector<bool, 1, SIMD_SIZE>
 
 
 //------------------------------------------------------- PARAMETERS
 
+FScreenTransform InputPixelPosToTranslucencyTextureUV;
+float2 TranslucencyTextureUVMin;
+float2 TranslucencyTextureUVMax;
 float3 HistoryGuideQuantizationError;
 
 Texture2D<tsr_halfC> InputTexture;
+Texture2D<tsr_half4> InputSceneTranslucencyTexture;
 Texture2D<tsr_halfC> ReprojectedHistoryGuideTexture;
 Texture2D<tsr_half> ParallaxRejectionMaskTexture;
 
@@ -293,6 +301,9 @@ void MainCS(
 		
 	// Load
 	tsr_simd_halfC OriginalInput;
+	#if CONFIG_COMPOSE_TRANSLUCENCY
+		tsr_simd_half4 OriginalTranslucencyInput;
+	#endif
 	tsr_simd_halfC OriginalHistory;
 	tsr_simd_half bIsValid;
 	ISOLATE
@@ -301,13 +312,37 @@ void MainCS(
 		for (uint SimdIndex = 0; SimdIndex < SIMD_SIZE; SimdIndex++)
 		{
 			tsr_short2 ClampedFetchPixelPos = GroupPixelOffset + GetLaneSimdPixelOffset(LaneIndex, SimdIndex);
+			float2 TranslucencyTextureUV = ApplyScreenTransform(float2(ClampedFetchPixelPos), InputPixelPosToTranslucencyTextureUV);
+
 			ClampedFetchPixelPos = ClampPixelOffset(ClampedFetchPixelPos, InputPixelPosMin, InputPixelPosMax);
+			TranslucencyTextureUV = clamp(TranslucencyTextureUV, TranslucencyTextureUVMin, TranslucencyTextureUVMax);
 
 			OriginalInput.SetElement(SimdIndex, InputTexture[ClampedFetchPixelPos]);
+			#if CONFIG_COMPOSE_TRANSLUCENCY
+				OriginalTranslucencyInput.SetElement(SimdIndex, InputSceneTranslucencyTexture.SampleLevel(GlobalBilinearClampedSampler, TranslucencyTextureUV, 0));
+			#endif
 			OriginalHistory.SetElement(SimdIndex, ReprojectedHistoryGuideTexture[ClampedFetchPixelPos]);
 			bIsValid.SetElement(SimdIndex, ParallaxRejectionMaskTexture[ClampedFetchPixelPos] > tsr_half(PARALLAX_REJECTION_MASK_THRESHOLD) ? tsr_half(1.0) : tsr_half(0.0));
 		}
 	}
+
+	#if CONFIG_COMPOSE_TRANSLUCENCY
+	{
+		UNROLL_N(SIMD_SIZE)
+		for (uint SimdIndex = 0; SimdIndex < SIMD_SIZE; SimdIndex++)
+		{
+			tsr_halfC Color = OriginalInput.GetElement(SimdIndex);
+			tsr_half4 Translucency = OriginalTranslucencyInput.GetElement(SimdIndex);
+
+			Color.rgb = Color.rgb * Translucency.a + Translucency.rgb;
+			#if CONFIG_SCENE_COLOR_APLHA
+				Color.a *= Translucency.a;
+			#endif 
+
+			OriginalInput.SetElement(SimdIndex, Color);
+		}
+	}
+	#endif
 		
 	// Compute AA Luma
 	tsr_simd_half CenterLumaLDR;
diff --git a/Engine/Shaders/Private/TemporalSuperResolution/TSRUpdateHistory.usf b/Engine/Shaders/Private/TemporalSuperResolution/TSRUpdateHistory.usf
--- a/Engine/Shaders/Private/TemporalSuperResolution/TSRUpdateHistory.usf
+++ b/Engine/Shaders/Private/TemporalSuperResolution/TSRUpdateHistory.usf
@@ -36,8 +36,13 @@
 	#define CONFIG_HIGH_FREQUENCY_YCOCG 0
 #endif
 
-#if 1
+#if DIM_SEPARATE_TRANSLUCENCY
 	#define CONFIG_ACCUMULATE_TRANSLUCENCY_SEPARATLY 1
+	#define CONFIG_COMPOSE_TRANSLUCENCY 0
+	#define CONFIG_REJECT_TRAMSLUCENCY_ON_RESPONSIVE_AA_ONLY 0
+#else
+	#define CONFIG_ACCUMULATE_TRANSLUCENCY_SEPARATLY 0
+	#define CONFIG_COMPOSE_TRANSLUCENCY 1
 	#define CONFIG_REJECT_TRAMSLUCENCY_ON_RESPONSIVE_AA_ONLY 0
 #endif
 
@@ -786,7 +791,7 @@ void MainCS(
 				for (uint PixelId = 0; PixelId < DPV_PIXEL_PER_LANE; PixelId++)
 				{
 					RawInputColorArray[SampleId][PixelId] = InputSceneColorTexture[dpv_access_pixel(InputSamplePixelPos, PixelId)];
-					#if CONFIG_ACCUMULATE_TRANSLUCENCY_SEPARATLY
+					#if CONFIG_ACCUMULATE_TRANSLUCENCY_SEPARATLY || CONFIG_COMPOSE_TRANSLUCENCY
 						RawTranslucencyColorArray[SampleId][PixelId] = InputSceneTranslucencyTexture[dpv_access_pixel(TranslucencySamplePixelPos, PixelId)];
 					#endif
 				}
@@ -813,6 +818,17 @@ void MainCS(
 			tsr_halfCx2 InputColor = dpv_force_interleave_registers_array(RawInputColorArray[SampleId]);
 			tsr_half4x2 TranslucencyColor = dpv_force_interleave_registers_array(RawTranslucencyColorArray[SampleId]);
 
+			#if CONFIG_COMPOSE_TRANSLUCENCY
+			{
+				InputColor[0] = InputColor[0] * TranslucencyColor[3] + TranslucencyColor[0];
+				InputColor[1] = InputColor[1] * TranslucencyColor[3] + TranslucencyColor[1];
+				InputColor[2] = InputColor[2] * TranslucencyColor[3] + TranslucencyColor[2];
+				#if CONFIG_SCENE_COLOR_APLHA
+					InputColor[3] = InputColor[3] * TranslucencyColor[3];
+				#endif
+			}
+			#endif
+
 			#if CONFIG_SCENE_COLOR_OVERFLOW
 				#if CONFIG_SCENE_COLOR_APLHA
 					InputColor = min(InputColor, dpv_interleave_mono_registers(LargestSceneColorRGBA));
diff --git a/Engine/Source/Runtime/Renderer/Private/PostProcess/TemporalSuperResolution.cpp b/Engine/Source/Runtime/Renderer/Private/PostProcess/TemporalSuperResolution.cpp
--- a/Engine/Source/Runtime/Renderer/Private/PostProcess/TemporalSuperResolution.cpp
+++ b/Engine/Source/Runtime/Renderer/Private/PostProcess/TemporalSuperResolution.cpp
@@ -31,6 +31,11 @@ TAutoConsoleVariable<int32> CVarTSRHistoryUpdateQuality(
 	TEXT("Select the quality of the history update."),
 	ECVF_Scalability | ECVF_RenderThreadSafe);
 
+TAutoConsoleVariable<int32> CVarTSRHistorySeparateTranslucency(
+	TEXT("r.TSR.History.SeparateTranslucency"), 1,
+	TEXT("Whether separate translucency should be accumulated separatly."),
+	ECVF_RenderThreadSafe);
+
 TAutoConsoleVariable<int32> CVarTSRWaveOps(
 	TEXT("r.TSR.WaveOps"), 1,
 	TEXT("Whether to use wave ops in the shading rejection heuristics"),
@@ -156,6 +161,7 @@ END_SHADER_PARAMETER_STRUCT()
 
 enum class ETSRHistoryFormatBits : uint32
 {
+	None = 0,
 	Translucency = 1 << 1,
 };
 ENUM_CLASS_FLAGS(ETSRHistoryFormatBits);
@@ -163,9 +169,14 @@ ENUM_CLASS_FLAGS(ETSRHistoryFormatBits);
 FTSRHistoryArrayIndices TranslateHistoryFormatBitsToArrayIndices(ETSRHistoryFormatBits HistoryFormatBits)
 {
 	FTSRHistoryArrayIndices ArrayIndices;
-	ArrayIndices.Size = 2;
+	ArrayIndices.Size = 1;
 	ArrayIndices.HighFrequency = 0;
-	ArrayIndices.Translucency = 1;
+	ArrayIndices.Translucency = -1;
+
+	if (EnumHasAnyFlags(HistoryFormatBits, ETSRHistoryFormatBits::Translucency))
+	{
+		ArrayIndices.Translucency = ArrayIndices.Size++;
+	}
 
 	return ArrayIndices;
 }
@@ -176,8 +187,11 @@ FTSRHistorySRVs CreateSRVs(FRDGBuilder& GraphBuilder, const FTSRHistoryTextures&
 	SRVs.HighFrequency = GraphBuilder.CreateSRV(FRDGTextureSRVDesc::CreateForSlice(Textures.ColorArray, Textures.ArrayIndices.HighFrequency));
 	SRVs.Metadata = GraphBuilder.CreateSRV(FRDGTextureSRVDesc::Create(Textures.Metadata));
 	SRVs.SubpixelDetails = Textures.SubpixelDetails;
-	SRVs.Translucency = GraphBuilder.CreateSRV(FRDGTextureSRVDesc::CreateForSlice(Textures.ColorArray, Textures.ArrayIndices.Translucency));
-	SRVs.TranslucencyAlpha = GraphBuilder.CreateSRV(FRDGTextureSRVDesc::Create(Textures.TranslucencyAlpha)); // TODO: move into metadata
+	if (Textures.ArrayIndices.Translucency >= 0)
+	{
+		SRVs.Translucency = GraphBuilder.CreateSRV(FRDGTextureSRVDesc::CreateForSlice(Textures.ColorArray, Textures.ArrayIndices.Translucency));
+		SRVs.TranslucencyAlpha = GraphBuilder.CreateSRV(FRDGTextureSRVDesc::Create(Textures.TranslucencyAlpha));
+	}
 	SRVs.Guide = Textures.Guide;
 	return SRVs;
 }
@@ -189,7 +203,10 @@ FTSRHistoryUAVs CreateUAVs(FRDGBuilder& GraphBuilder, const FTSRHistoryTextures&
 	UAVs.ColorArray = GraphBuilder.CreateUAV(Textures.ColorArray);
 	UAVs.Metadata = GraphBuilder.CreateUAV(Textures.Metadata);
 	UAVs.SubpixelDetails = GraphBuilder.CreateUAV(Textures.SubpixelDetails);
-	UAVs.TranslucencyAlpha = GraphBuilder.CreateUAV(Textures.TranslucencyAlpha); // TODO: move into metadata
+	if (Textures.TranslucencyAlpha)
+	{
+		UAVs.TranslucencyAlpha = GraphBuilder.CreateUAV(Textures.TranslucencyAlpha);
+	}
 	return UAVs;
 }
 
@@ -332,13 +349,19 @@ class FTSRRejectShadingCS : public FTSRShader
 	SHADER_USE_PARAMETER_STRUCT(FTSRRejectShadingCS, FTSRShader);
 
 	class FUseWaveOps : SHADER_PERMUTATION_BOOL("DIM_USE_WAVE_OPS");
-	using FPermutationDomain = TShaderPermutationDomain<FUseWaveOps>;
+	class FSeparateTranslucencyDim : SHADER_PERMUTATION_BOOL("DIM_SEPARATE_TRANSLUCENCY");
+
+	using FPermutationDomain = TShaderPermutationDomain<FUseWaveOps, FSeparateTranslucencyDim>;
 
 	BEGIN_SHADER_PARAMETER_STRUCT(FParameters, )
 		SHADER_PARAMETER_STRUCT_INCLUDE(FTSRCommonParameters, CommonParameters)
+		SHADER_PARAMETER(FScreenTransform, InputPixelPosToTranslucencyTextureUV)
+		SHADER_PARAMETER(FVector2f, TranslucencyTextureUVMin)
+		SHADER_PARAMETER(FVector2f, TranslucencyTextureUVMax)
 		SHADER_PARAMETER(FVector3f, HistoryGuideQuantizationError)
 
 		SHADER_PARAMETER_RDG_TEXTURE(Texture2D, InputTexture)
+		SHADER_PARAMETER_RDG_TEXTURE(Texture2D, InputSceneTranslucencyTexture)
 		SHADER_PARAMETER_RDG_TEXTURE(Texture2D, ReprojectedHistoryGuideTexture)
 		SHADER_PARAMETER_RDG_TEXTURE(Texture2D, ParallaxRejectionMaskTexture)
 
@@ -469,8 +492,9 @@ class FTSRUpdateHistoryCS : public FTSRShader
 	};
 
 	class FQualityDim : SHADER_PERMUTATION_ENUM_CLASS("DIM_UPDATE_QUALITY", EQuality);
+	class FSeparateTranslucencyDim : SHADER_PERMUTATION_BOOL("DIM_SEPARATE_TRANSLUCENCY");
 
-	using FPermutationDomain = TShaderPermutationDomain<FQualityDim>;
+	using FPermutationDomain = TShaderPermutationDomain<FQualityDim, FSeparateTranslucencyDim>;
 
 	BEGIN_SHADER_PARAMETER_STRUCT(FParameters, )
 		SHADER_PARAMETER_STRUCT_INCLUDE(FTSRCommonParameters, CommonParameters)
@@ -622,12 +646,18 @@ ITemporalUpscaler::FOutputs AddTemporalSuperResolutionPasses(
 	// Whether alpha channel is supported.
 	const bool bSupportsAlpha = IsPostProcessingWithAlphaChannelSupported();
 
+	// Whether separate translucency should be temporarily accumulated separatly.
+	const bool bAccumulateTranslucencySeparately = CVarTSRHistorySeparateTranslucency.GetValueOnRenderThread() != 0;
+
 	// whether TSR passes can run on async compute.
 	int32 AsyncComputePasses = GSupportsEfficientAsyncCompute ? CVarTSRAsyncCompute.GetValueOnRenderThread() : 0;
 
-	ETSRHistoryFormatBits HistoryFormatBits;
+	ETSRHistoryFormatBits HistoryFormatBits = ETSRHistoryFormatBits::None;
 	{
-		HistoryFormatBits = ETSRHistoryFormatBits::Translucency;
+		if (bAccumulateTranslucencySeparately)
+		{
+			HistoryFormatBits = ETSRHistoryFormatBits::Translucency;
+		}
 	}
 
 	// Whether to use camera cut shader permutation or not.
@@ -923,7 +953,10 @@ ITemporalUpscaler::FOutputs AddTemporalSuperResolutionPasses(
 		History.Metadata = GraphBuilder.CreateTexture(Desc, TEXT("TSR.History.Metadata"));
 
 		Desc.Format = PF_R8;
-		History.TranslucencyAlpha = GraphBuilder.CreateTexture(Desc, TEXT("TSR.History.TranslucencyAlpha"));
+		if (bAccumulateTranslucencySeparately)
+		{
+			History.TranslucencyAlpha = GraphBuilder.CreateTexture(Desc, TEXT("TSR.History.TranslucencyAlpha"));
+		}
 	}
 	{
 		FRDGTextureDesc SubpixelDetailsDesc = FRDGTextureDesc::Create2D(
@@ -1178,12 +1211,22 @@ ITemporalUpscaler::FOutputs AddTemporalSuperResolutionPasses(
 			HistoryRejectionTexture = GraphBuilder.CreateTexture(Desc, TEXT("TSR.HistoryRejection"));
 		}
 
+		FScreenPassTextureViewport TranslucencyViewport(
+			SeparateTranslucencyTexture->Desc.Extent, SeparateTranslucencyRect);
+
 		FTSRRejectShadingCS::FParameters* PassParameters = GraphBuilder.AllocParameters<FTSRRejectShadingCS::FParameters>();
 		PassParameters->CommonParameters = CommonParameters;
+		PassParameters->InputPixelPosToTranslucencyTextureUV =
+			((FScreenTransform::Identity + 0.5f - InputRect.Min) / InputRect.Size()) *
+			FScreenTransform::ChangeTextureBasisFromTo(TranslucencyViewport, FScreenTransform::ETextureBasis::ViewportUV, FScreenTransform::ETextureBasis::TextureUV);
+		PassParameters->TranslucencyTextureUVMin = GetScreenPassTextureViewportParameters(TranslucencyViewport).UVViewportBilinearMin;
+		PassParameters->TranslucencyTextureUVMax = GetScreenPassTextureViewportParameters(TranslucencyViewport).UVViewportBilinearMax;
+		PassParameters->HistoryGuideQuantizationError = ComputePixelFormatQuantizationError(History.Guide->Desc.Format);
+
 		PassParameters->InputTexture = PassInputs.SceneColorTexture;
+		PassParameters->InputSceneTranslucencyTexture = SeparateTranslucencyTexture;
 		PassParameters->ReprojectedHistoryGuideTexture = ReprojectedHistoryGuideTexture;
 		PassParameters->ParallaxRejectionMaskTexture = ParallaxRejectionMaskTexture;
-		PassParameters->HistoryGuideQuantizationError = ComputePixelFormatQuantizationError(History.Guide->Desc.Format);
 
 		PassParameters->HistoryGuideOutput = GraphBuilder.CreateUAV(History.Guide);
 		PassParameters->HistoryRejectionOutput = GraphBuilder.CreateUAV(HistoryRejectionTexture);
@@ -1191,13 +1234,15 @@ ITemporalUpscaler::FOutputs AddTemporalSuperResolutionPasses(
 		PassParameters->DebugOutput = CreateDebugUAV(InputExtent, TEXT("Debug.TSR.RejectShading"));
 
 		FTSRRejectShadingCS::FPermutationDomain PermutationVector;
-		PermutationVector.Set<FTSRRejectShadingCS::FUseWaveOps>(bUseWaveOps && GRHIMinimumWaveSize <= 32 && GRHIMaximumWaveSize >= 32);
+		PermutationVector.Set<FTSRRejectShadingCS::FUseWaveOps>(bUseWaveOps&& GRHIMinimumWaveSize <= 32 && GRHIMaximumWaveSize >= 32);
+		PermutationVector.Set<FTSRRejectShadingCS::FSeparateTranslucencyDim>(bAccumulateTranslucencySeparately);
 
 		TShaderMapRef<FTSRRejectShadingCS> ComputeShader(View.ShaderMap, PermutationVector);
 		FComputeShaderUtils::AddPass(
 			GraphBuilder,
-			RDG_EVENT_NAME("TSR RejectShading(%s) %dx%d",
+			RDG_EVENT_NAME("TSR RejectShading(%s%s) %dx%d",
 				PermutationVector.Get<FTSRRejectShadingCS::FUseWaveOps>() ? TEXT("WaveOps") : TEXT(""),
+				PermutationVector.Get<FTSRRejectShadingCS::FSeparateTranslucencyDim>() ? TEXT(" SeparateTranslucency") : TEXT(""),
 				InputRect.Width(), InputRect.Height()),
 			AsyncComputePasses >= 3 ? ERDGPassFlags::AsyncCompute : ERDGPassFlags::Compute,
 			ComputeShader,
@@ -1423,13 +1468,15 @@ ITemporalUpscaler::FOutputs AddTemporalSuperResolutionPasses(
 
 		FTSRUpdateHistoryCS::FPermutationDomain PermutationVector;
 		PermutationVector.Set<FTSRUpdateHistoryCS::FQualityDim>(UpdateHistoryQuality);
+		PermutationVector.Set<FTSRUpdateHistoryCS::FSeparateTranslucencyDim>(bAccumulateTranslucencySeparately);
 
 		TShaderMapRef<FTSRUpdateHistoryCS> ComputeShader(View.ShaderMap, PermutationVector);
 		FComputeShaderUtils::AddPass(
 			GraphBuilder,
-			RDG_EVENT_NAME("TSR UpdateHistory(Quality=%s %s%s) %dx%d",
+			RDG_EVENT_NAME("TSR UpdateHistory(Quality=%s%s%s%s) %dx%d",
 				kUpdateQualityNames[int32(PermutationVector.Get<FTSRUpdateHistoryCS::FQualityDim>())],
-				History.ColorArray->Desc.Format == PF_FloatR11G11B10 ? TEXT("R11G11B10") : TEXT(""),
+				PermutationVector.Get<FTSRUpdateHistoryCS::FSeparateTranslucencyDim>() ? TEXT(" SeparateTranslucency") : TEXT(""),
+				History.ColorArray->Desc.Format == PF_FloatR11G11B10 ? TEXT(" R11G11B10") : TEXT(""),
 				PassParameters->bGenerateOutputMip1 ? TEXT(" OutputMip1") : TEXT(""),
 				HistorySize.X, HistorySize.Y),
 			AsyncComputePasses >= 3 ? ERDGPassFlags::AsyncCompute : ERDGPassFlags::Compute,
@@ -1514,7 +1561,10 @@ ITemporalUpscaler::FOutputs AddTemporalSuperResolutionPasses(
 		// Extract filterable history
 		GraphBuilder.QueueTextureExtraction(History.ColorArray, &OutputHistory.ColorArray);
 		GraphBuilder.QueueTextureExtraction(History.Metadata, &OutputHistory.Metadata);
-		GraphBuilder.QueueTextureExtraction(History.TranslucencyAlpha, &OutputHistory.TranslucencyAlpha);
+		if (bAccumulateTranslucencySeparately)
+		{
+			GraphBuilder.QueueTextureExtraction(History.TranslucencyAlpha, &OutputHistory.TranslucencyAlpha);
+		}
 
 		// Extract non-filterable history
 		GraphBuilder.QueueTextureExtraction(History.SubpixelDetails, &OutputHistory.SubpixelDetails);
