Only compile TSR's shading rejection heuristic with LDS fallback on platforms that are not guarenteed to have waveops

diff --git a/Engine/Config/Mac/DataDrivenPlatformInfo.ini b/Engine/Config/Mac/DataDrivenPlatformInfo.ini
--- a/Engine/Config/Mac/DataDrivenPlatformInfo.ini
+++ b/Engine/Config/Mac/DataDrivenPlatformInfo.ini
@@ -34,7 +34,7 @@ bSupportsDistanceFields=true
 bSupportsGen5TemporalAA=false
 bSupportsLumenGI=true
 bSupportsSSDIndirect=true
-bSupportsWaveOperations=false
+bSupportsWaveOperations=Unsupported
 bSupportsGPUScene=true
 bSupportsFFTBloom=true
 bSupportsDiaphragmDOF=true
diff --git a/Engine/Config/VulkanPC/DataDrivenPlatformInfo.ini b/Engine/Config/VulkanPC/DataDrivenPlatformInfo.ini
--- a/Engine/Config/VulkanPC/DataDrivenPlatformInfo.ini
+++ b/Engine/Config/VulkanPC/DataDrivenPlatformInfo.ini
@@ -83,7 +83,7 @@ bSupportsIntrinsicWaveOnce = false
 ; //#todo-vulkanddpi
 bSupportsConservativeRasterization = false
 
-bSupportsWaveOperations = true
+bSupportsWaveOperations=RuntimeDependent
 
 ; //#todo-vulkanddpi
 bRequiresExplicit128bitRT = false
diff --git a/Engine/Config/Windows/DataDrivenPlatformInfo.ini b/Engine/Config/Windows/DataDrivenPlatformInfo.ini
--- a/Engine/Config/Windows/DataDrivenPlatformInfo.ini
+++ b/Engine/Config/Windows/DataDrivenPlatformInfo.ini
@@ -67,7 +67,7 @@ bSupportsRayTracingProceduralPrimitive=true
 bSupportsRayTracingIndirectInstanceData=true
 bSupportsPathTracing=true
 bSupportsHighEndRayTracingReflections=true
-bSupportsWaveOperations=false
+bSupportsWaveOperations=RuntimeDependent
 bSupportsGPUScene=true
 bSupportsNanite=false
 bSupportsComputeFramework=true
@@ -120,7 +120,7 @@ bSupportsByteBufferComputeShaders=true
 bSupportsLumenGI=true
 bSupportsDistanceFields=true
 bSupportsSSDIndirect=true
-bSupportsWaveOperations=true
+bSupportsWaveOperations=RuntimeGuaranteed
 bSupportsDxc=true
 bSupportsDiaphragmDOF=true
 bRequiresVendorExtensionsForAtomics=false
diff --git a/Engine/Plugins/FX/Niagara/Source/NiagaraVertexFactories/Private/NiagaraSortingGPU.cpp b/Engine/Plugins/FX/Niagara/Source/NiagaraVertexFactories/Private/NiagaraSortingGPU.cpp
--- a/Engine/Plugins/FX/Niagara/Source/NiagaraVertexFactories/Private/NiagaraSortingGPU.cpp
+++ b/Engine/Plugins/FX/Niagara/Source/NiagaraVertexFactories/Private/NiagaraSortingGPU.cpp
@@ -37,7 +37,7 @@ void FNiagaraSortKeyGenCS::ModifyCompilationEnvironment(const FGlobalShaderPermu
 	OutEnvironment.SetDefine(TEXT("SORT_CUSTOM_ASCENDING"), (uint8)ENiagaraSortMode::CustomAscending);
 	OutEnvironment.SetDefine(TEXT("SORT_CUSTOM_DESCENDING"), (uint8)ENiagaraSortMode::CustomDecending);
 
-	bool bUseWaveIntrinsics = FDataDrivenShaderPlatformInfo::GetSupportsWaveOperations(Parameters.Platform);
+	bool bUseWaveIntrinsics = FDataDrivenShaderPlatformInfo::GetSupportsWaveOperations(Parameters.Platform) != ERHIFeatureSupport::Unsupported;
 	OutEnvironment.SetDefine(TEXT("USE_WAVE_INTRINSICS"), bUseWaveIntrinsics);
 	if (bUseWaveIntrinsics)
 	{
diff --git a/Engine/Shaders/Private/TemporalSuperResolution/TSRCommon.ush b/Engine/Shaders/Private/TemporalSuperResolution/TSRCommon.ush
--- a/Engine/Shaders/Private/TemporalSuperResolution/TSRCommon.ush
+++ b/Engine/Shaders/Private/TemporalSuperResolution/TSRCommon.ush
@@ -11,7 +11,7 @@
 
 
 //------------------------------------------------------- RECOMPILE HASH
-// 88217342-DCE7-451F-8670-0A14BC8B6109
+// 2A42D2EF-728A-4AC9-825F-2DC009BE04B0
 
 
 //------------------------------------------------------- COMPILER CONFIG
diff --git a/Engine/Source/Runtime/RHI/Private/RHI.cpp b/Engine/Source/Runtime/RHI/Private/RHI.cpp
--- a/Engine/Source/Runtime/RHI/Private/RHI.cpp
+++ b/Engine/Source/Runtime/RHI/Private/RHI.cpp
@@ -2507,7 +2507,7 @@ void FGenericDataDrivenShaderPlatformInfo::ParseDataDrivenShaderInfo(const FConf
 	GET_SECTION_BOOL_HELPER(bSupportsRTIndexFromVS);
 	GET_SECTION_BOOL_HELPER(bSupportsIntrinsicWaveOnce);
 	GET_SECTION_BOOL_HELPER(bSupportsConservativeRasterization);
-	GET_SECTION_BOOL_HELPER(bSupportsWaveOperations);
+	GET_SECTION_SUPPORT_HELPER(bSupportsWaveOperations);
 	GET_SECTION_BOOL_HELPER(bRequiresExplicit128bitRT);
 	GET_SECTION_BOOL_HELPER(bSupportsGen5TemporalAA);
 	GET_SECTION_BOOL_HELPER(bTargetsTiledGPU);
diff --git a/Engine/Source/Runtime/RHI/Public/RHI.h b/Engine/Source/Runtime/RHI/Public/RHI.h
--- a/Engine/Source/Runtime/RHI/Public/RHI.h
+++ b/Engine/Source/Runtime/RHI/Public/RHI.h
@@ -262,7 +262,7 @@ inline uint32 RHIMaxMeshShaderThreadGroupSize(const FStaticShaderPlatform Platfo
  **/
 inline RHI_API bool RHISupportsWaveOperations(const FStaticShaderPlatform Platform)
 {
-	return FDataDrivenShaderPlatformInfo::GetSupportsWaveOperations(Platform);
+	return FDataDrivenShaderPlatformInfo::GetSupportsWaveOperations(Platform) != ERHIFeatureSupport::Unsupported;
 }
 
 /** True if the given shader platform supports a render target write mask */
diff --git a/Engine/Source/Runtime/RHI/Public/RHIDefinitions.h b/Engine/Source/Runtime/RHI/Public/RHIDefinitions.h
--- a/Engine/Source/Runtime/RHI/Public/RHIDefinitions.h
+++ b/Engine/Source/Runtime/RHI/Public/RHIDefinitions.h
@@ -375,7 +375,7 @@ class RHI_API FGenericDataDrivenShaderPlatformInfo
 	uint32 bSupportsSSDIndirect : 1;
 	uint32 bSupportsTemporalHistoryUpscale : 1;
 	uint32 bSupportsRTIndexFromVS : 1;
-	uint32 bSupportsWaveOperations : 1; // Whether HLSL SM6 shader wave intrinsics are supported
+	uint32 bSupportsWaveOperations : int32(ERHIFeatureSupport::NumBits);
 	uint32 bSupportsIntrinsicWaveOnce : 1;
 	uint32 bSupportsConservativeRasterization : 1;
 	uint32 bRequiresExplicit128bitRT : 1;
@@ -711,10 +711,10 @@ public:
 		return Infos[Platform].bSupportsByteBufferComputeShaders;
 	}
 
-	static FORCEINLINE_DEBUGGABLE const bool GetSupportsWaveOperations(const FStaticShaderPlatform Platform)
+	static FORCEINLINE_DEBUGGABLE const ERHIFeatureSupport GetSupportsWaveOperations(const FStaticShaderPlatform Platform)
 	{
 		check(IsValid(Platform));
-		return Infos[Platform].bSupportsWaveOperations;
+		return ERHIFeatureSupport(Infos[Platform].bSupportsWaveOperations);
 	}
 
 	static FORCEINLINE_DEBUGGABLE const bool GetSupportsTemporalHistoryUpscale(const FStaticShaderPlatform Platform)
diff --git a/Engine/Source/Runtime/Renderer/Private/PostProcess/TemporalSuperResolution.cpp b/Engine/Source/Runtime/Renderer/Private/PostProcess/TemporalSuperResolution.cpp
--- a/Engine/Source/Runtime/Renderer/Private/PostProcess/TemporalSuperResolution.cpp
+++ b/Engine/Source/Runtime/Renderer/Private/PostProcess/TemporalSuperResolution.cpp
@@ -375,9 +375,14 @@ class FTSRRejectShadingCS : public FTSRShader
 	{
 		FPermutationDomain PermutationVector(Parameters.PermutationId);
 
-		if (PermutationVector.Get<FUseWaveOps>() && !RHISupportsWaveOperations(Parameters.Platform))
+		ERHIFeatureSupport WaveOpsSupport = FDataDrivenShaderPlatformInfo::GetSupportsWaveOperations(Parameters.Platform);
+		if (PermutationVector.Get<FUseWaveOps>())
 		{
-			return false;
+			return WaveOpsSupport != ERHIFeatureSupport::Unsupported;
+		}
+		else
+		{
+			return WaveOpsSupport != ERHIFeatureSupport::RuntimeGuaranteed;
 		}
 
 		return FTSRShader::ShouldCompilePermutation(Parameters);
@@ -641,7 +646,8 @@ ITemporalUpscaler::FOutputs AddTemporalSuperResolutionPasses(
 #endif
 
 	// Whether to use wave ops optimizations.
-	const bool bUseWaveOps = CVarTSRWaveOps.GetValueOnRenderThread() != 0 && GRHISupportsWaveOperations && RHISupportsWaveOperations(View.GetShaderPlatform());
+	const ERHIFeatureSupport WaveOpsSupport = FDataDrivenShaderPlatformInfo::GetSupportsWaveOperations(View.GetShaderPlatform());
+	const bool bUseWaveOps = (CVarTSRWaveOps.GetValueOnRenderThread() != 0 && GRHISupportsWaveOperations && WaveOpsSupport == ERHIFeatureSupport::RuntimeDependent) || WaveOpsSupport == ERHIFeatureSupport::RuntimeGuaranteed;
 
 	// Whether alpha channel is supported.
 	const bool bSupportsAlpha = IsPostProcessingWithAlphaChannelSupported();
diff --git a/Engine/Source/Runtime/Renderer/Private/Strata/Strata.cpp b/Engine/Source/Runtime/Renderer/Private/Strata/Strata.cpp
--- a/Engine/Source/Runtime/Renderer/Private/Strata/Strata.cpp
+++ b/Engine/Source/Runtime/Renderer/Private/Strata/Strata.cpp
@@ -560,7 +560,7 @@ class FStrataBSDFTilePassCS : public FGlobalShader
 
 	static bool ShouldCompilePermutation(const FGlobalShaderPermutationParameters& Parameters)
 	{
-		const bool bUseWaveIntrinsics = FDataDrivenShaderPlatformInfo::GetSupportsWaveOperations(Parameters.Platform);
+		const bool bUseWaveIntrinsics = FDataDrivenShaderPlatformInfo::GetSupportsWaveOperations(Parameters.Platform) != ERHIFeatureSupport::Unsupported;
 		FPermutationDomain PermutationVector(Parameters.PermutationId);
 		if (PermutationVector.Get<FWaveOps>() && !bUseWaveIntrinsics)
 		{
@@ -613,7 +613,7 @@ class FStrataMaterialTileClassificationPassCS : public FGlobalShader
 
 	static bool ShouldCompilePermutation(const FGlobalShaderPermutationParameters& Parameters)
 	{
-		const bool bUseWaveIntrinsics = FDataDrivenShaderPlatformInfo::GetSupportsWaveOperations(Parameters.Platform);
+		const bool bUseWaveIntrinsics = FDataDrivenShaderPlatformInfo::GetSupportsWaveOperations(Parameters.Platform) != ERHIFeatureSupport::Unsupported;
 		FPermutationDomain PermutationVector(Parameters.PermutationId);
 		if (PermutationVector.Get<FWaveOps>() && !bUseWaveIntrinsics)
 		{
@@ -1009,7 +1009,7 @@ void AddStrataMaterialClassificationPass(FRDGBuilder& GraphBuilder, const FMinim
 		RDG_EVENT_SCOPE_CONDITIONAL(GraphBuilder, Views.Num() > 1, "View%d", i);
 
 		const FViewInfo& View = Views[i];
-		bool bWaveOps = GRHISupportsWaveOperations && GRHIMaximumWaveSize >= 64 && FDataDrivenShaderPlatformInfo::GetSupportsWaveOperations(View.GetShaderPlatform());
+		bool bWaveOps = GRHISupportsWaveOperations && GRHIMaximumWaveSize >= 64 && FDataDrivenShaderPlatformInfo::GetSupportsWaveOperations(View.GetShaderPlatform()) != ERHIFeatureSupport::Unsupported;
 		
 		const FStrataViewData* StrataViewData = &View.StrataViewData;
 		const FStrataSceneData* StrataSceneData = View.StrataViewData.SceneData;
