Fixes iOS build

diff --git a/Engine/Source/Runtime/Renderer/Private/PostProcess/TemporalSuperResolution.cpp b/Engine/Source/Runtime/Renderer/Private/PostProcess/TemporalSuperResolution.cpp
--- a/Engine/Source/Runtime/Renderer/Private/PostProcess/TemporalSuperResolution.cpp
+++ b/Engine/Source/Runtime/Renderer/Private/PostProcess/TemporalSuperResolution.cpp
@@ -378,11 +378,17 @@ class FTSRRejectShadingCS : public FTSRShader
 		ERHIFeatureSupport WaveOpsSupport = FDataDrivenShaderPlatformInfo::GetSupportsWaveOperations(Parameters.Platform);
 		if (PermutationVector.Get<FUseWaveOps>())
 		{
-			return WaveOpsSupport != ERHIFeatureSupport::Unsupported;
+			if (WaveOpsSupport == ERHIFeatureSupport::Unsupported)
+			{
+				return false;
+			}
 		}
 		else
 		{
-			return WaveOpsSupport != ERHIFeatureSupport::RuntimeGuaranteed;
+			if (WaveOpsSupport == ERHIFeatureSupport::RuntimeGuaranteed)
+			{
+				return false;
+			}
 		}
 
 		return FTSRShader::ShouldCompilePermutation(Parameters);
