Fix race between coarse and pixel page marking when non-nanite is disabled in coarse pages.

diff --git a/Engine/Shaders/Private/VirtualShadowMaps/VirtualShadowMapPageManagement.usf b/Engine/Shaders/Private/VirtualShadowMaps/VirtualShadowMapPageManagement.usf
--- a/Engine/Shaders/Private/VirtualShadowMaps/VirtualShadowMapPageManagement.usf
+++ b/Engine/Shaders/Private/VirtualShadowMaps/VirtualShadowMapPageManagement.usf
@@ -59,6 +59,11 @@ uint GetMipLevelLocal(int VirtualShadowMapId, float3 TranslatedWorldPosition, fl
 	return min(MipLevel, (VSM_MAX_MIP_LEVELS - 1U));
 }
 
+void MarkPageAddress(uint PageOffset, uint Flags)
+{
+	OutPageRequestFlags[PageOffset] = Flags;
+}
+
 void MarkPage(uint VirtualShadowMapId, uint MipLevel, float3 TranslatedWorldPosition, bool bUsePageDilation, float2 PageDilationOffset)
 {
 	FVirtualShadowMapProjectionShaderData ProjectionData = GetVirtualShadowMapProjectionData(VirtualShadowMapId);
@@ -83,7 +88,7 @@ void MarkPage(uint VirtualShadowMapId, uint MipLevel, float3 TranslatedWorldPosi
 	float2 PageAddressFloat = ShadowUVz.xy * CalcLevelDimsPages(MipLevel);
 	uint2 PageAddress = clamp(uint2(PageAddressFloat), 0U, MaxPageAddress);
 	uint PageOffset = CalcPageOffset(VirtualShadowMapId, MipLevel, PageAddress);
-	OutPageRequestFlags[PageOffset] = Flags;
+	MarkPageAddress(PageOffset, Flags);
 
 	// PageDilationBorderSize == 0 implies PageDilationOffset.xy == 0
 	if (bUsePageDilation)
@@ -92,13 +97,13 @@ void MarkPage(uint VirtualShadowMapId, uint MipLevel, float3 TranslatedWorldPosi
 		uint PageOffset2 = CalcPageOffset(VirtualShadowMapId, MipLevel, PageAddress2);
 		if (PageOffset2 != PageOffset)
 		{
-			OutPageRequestFlags[PageOffset2] = Flags;
+			MarkPageAddress(PageOffset2, Flags);
 		}
 		uint2 PageAddress3 = clamp(uint2(PageAddressFloat - PageDilationOffset), 0U, MaxPageAddress);
 		uint PageOffset3 = CalcPageOffset(VirtualShadowMapId, MipLevel, PageAddress3);
 		if (PageOffset3 != PageOffset)
 		{
-			OutPageRequestFlags[PageOffset3] = Flags;
+			MarkPageAddress(PageOffset3, Flags);
 		}
 	}
 }
@@ -297,17 +302,17 @@ void MarkCoarsePages(uint3 Index : SV_DispatchThreadID)
 			// NOTE: Page addresses round down/truncate normally, so grab the surrounding 4
 			int4 PageAddressLowHigh = int4(floor(PageAddressFloat - 0.5f), ceil(PageAddressFloat - 0.5f));
 
-			OutPageRequestFlags[CalcPageOffset(VirtualShadowMapId, 0, PageAddressLowHigh.xy)] = Flags;
-			OutPageRequestFlags[CalcPageOffset(VirtualShadowMapId, 0, PageAddressLowHigh.xw)] = Flags;
-			OutPageRequestFlags[CalcPageOffset(VirtualShadowMapId, 0, PageAddressLowHigh.zy)] = Flags;
-			OutPageRequestFlags[CalcPageOffset(VirtualShadowMapId, 0, PageAddressLowHigh.zw)] = Flags;
+			MarkPageAddress(CalcPageOffset(VirtualShadowMapId, 0, PageAddressLowHigh.xy), Flags);
+			MarkPageAddress(CalcPageOffset(VirtualShadowMapId, 0, PageAddressLowHigh.xw), Flags);
+			MarkPageAddress(CalcPageOffset(VirtualShadowMapId, 0, PageAddressLowHigh.zy), Flags);
+			MarkPageAddress(CalcPageOffset(VirtualShadowMapId, 0, PageAddressLowHigh.zw), Flags);
 		}
 	}
 	else if (bMarkCoarsePagesLocal != 0)
 	{
 		// Mark last mip
 		uint PageOffset = CalcPageOffset(VirtualShadowMapId, VSM_MAX_MIP_LEVELS - 1, uint2(0, 0));
-		OutPageRequestFlags[PageOffset] = Flags;
+		MarkPageAddress(PageOffset, Flags);
 	}
 }
 
diff --git a/Engine/Source/Runtime/Renderer/Private/VirtualShadowMaps/VirtualShadowMapArray.cpp b/Engine/Source/Runtime/Renderer/Private/VirtualShadowMaps/VirtualShadowMapArray.cpp
--- a/Engine/Source/Runtime/Renderer/Private/VirtualShadowMaps/VirtualShadowMapArray.cpp
+++ b/Engine/Source/Runtime/Renderer/Private/VirtualShadowMaps/VirtualShadowMapArray.cpp
@@ -1225,12 +1225,39 @@ void FVirtualShadowMapArray::BuildPageAllocations(
 		FRDGBufferRef ScreenSpaceGridBoundsRDG = nullptr;
 			
 		{
-			// It's safe to overlap these passes that all write to page request flags
-			FRDGBufferUAVRef PageRequestFlagsUAV = GraphBuilder.CreateUAV(PageRequestFlagsRDG, ERDGUnorderedAccessViewFlags::SkipBarrier);
+			// Mark coarse pages
+			// NOTE: Must do this *first*. In the case where bIncludeNonNaniteGeometry is false we need to ensure that the request
+			// can be over-written by any pixel pages that *do* want Non-Nanite geometry. We avoid writing with atomics since that
+			// is much slower.
+			// Because of this we also cannot overlap this pass with the following ones.
+			bool bMarkCoarsePagesDirectional = CVarMarkCoarsePagesDirectional.GetValueOnRenderThread() != 0;
+			bool bMarkCoarsePagesLocal = CVarMarkCoarsePagesLocal.GetValueOnRenderThread() != 0;
+			if (bMarkCoarsePagesDirectional || bMarkCoarsePagesLocal)
+			{
+				FMarkCoarsePagesCS::FParameters* PassParameters = GraphBuilder.AllocParameters< FMarkCoarsePagesCS::FParameters >();
+				PassParameters->VirtualShadowMap = GetUniformBuffer(GraphBuilder);
+				PassParameters->OutPageRequestFlags = GraphBuilder.CreateUAV(PageRequestFlagsRDG);
+				PassParameters->bMarkCoarsePagesLocal = bMarkCoarsePagesLocal ? 1 : 0;
+				PassParameters->ClipmapIndexMask = bMarkCoarsePagesDirectional ? FVirtualShadowMapClipmap::GetCoarsePageClipmapIndexMask() : 0;
+				PassParameters->bIncludeNonNaniteGeometry = CVarCoarsePagesIncludeNonNanite.GetValueOnRenderThread();
+
+				auto ComputeShader = View.ShaderMap->GetShader<FMarkCoarsePagesCS>();
+
+				FComputeShaderUtils::AddPass(
+					GraphBuilder,
+					RDG_EVENT_NAME("MarkCoarsePages"),
+					ComputeShader,
+					PassParameters,
+					FIntVector(FMath::DivideAndRoundUp(uint32(ShadowMaps.Num()), FMarkCoarsePagesCS::DefaultCSGroupX), 1, 1)
+				);
+			}
 
 			// Mark pages based on projected depth buffer pixels
 			if (CVarMarkPixelPages.GetValueOnRenderThread() != 0)
 			{
+				// It's currently safe to overlap these passes that all write to same page request flags
+				FRDGBufferUAVRef PageRequestFlagsUAV = GraphBuilder.CreateUAV(PageRequestFlagsRDG, ERDGUnorderedAccessViewFlags::SkipBarrier);
+
 				auto GeneratePageFlags = [&](const EVirtualShadowMapProjectionInputType InputType)
 				{
 					FGeneratePageFlagsFromPixelsCS::FPermutationDomain PermutationVector;
@@ -1283,28 +1310,6 @@ void FVirtualShadowMapArray::BuildPageAllocations(
 					GeneratePageFlags(EVirtualShadowMapProjectionInputType::HairStrands);
 				}
 			}
-			// Mark coarse pages
-			bool bMarkCoarsePagesDirectional = CVarMarkCoarsePagesDirectional.GetValueOnRenderThread() != 0;
-			bool bMarkCoarsePagesLocal = CVarMarkCoarsePagesLocal.GetValueOnRenderThread() != 0;
-			if (bMarkCoarsePagesDirectional || bMarkCoarsePagesLocal)
-			{
-				FMarkCoarsePagesCS::FParameters* PassParameters = GraphBuilder.AllocParameters< FMarkCoarsePagesCS::FParameters >();
-				PassParameters->VirtualShadowMap = GetUniformBuffer(GraphBuilder);
-				PassParameters->OutPageRequestFlags = PageRequestFlagsUAV;
-				PassParameters->bMarkCoarsePagesLocal = bMarkCoarsePagesLocal ? 1 : 0;
-				PassParameters->ClipmapIndexMask = bMarkCoarsePagesDirectional ? FVirtualShadowMapClipmap::GetCoarsePageClipmapIndexMask() : 0;
-				PassParameters->bIncludeNonNaniteGeometry = CVarCoarsePagesIncludeNonNanite.GetValueOnRenderThread();
-
-				auto ComputeShader = View.ShaderMap->GetShader<FMarkCoarsePagesCS>();
-
-				FComputeShaderUtils::AddPass(
-					GraphBuilder,
-					RDG_EVENT_NAME("MarkCoarsePages"),
-					ComputeShader,
-					PassParameters,
-					FIntVector(FMath::DivideAndRoundUp(uint32(ShadowMaps.Num()), FMarkCoarsePagesCS::DefaultCSGroupX), 1, 1)
-				);
-			}
 		}
 	}
 
