Fix translucency artifacts in VR (UE-155157)

diff --git a/Engine/Source/Runtime/Renderer/Private/ShadowRendering.cpp b/Engine/Source/Runtime/Renderer/Private/ShadowRendering.cpp
--- a/Engine/Source/Runtime/Renderer/Private/ShadowRendering.cpp
+++ b/Engine/Source/Runtime/Renderer/Private/ShadowRendering.cpp
@@ -1942,7 +1942,7 @@ bool FSceneRenderer::CheckForProjectedShadows( const FLightSceneInfo* LightScene
 		for(int32 ViewIndex = 0;ViewIndex < Views.Num();ViewIndex++)
 		{
 			const FViewInfo& View = Views[ViewIndex];
-			if (ProjectedShadowInfo->DependentView && ProjectedShadowInfo->DependentView != &View)
+			if (ProjectedShadowInfo->DependentView && ProjectedShadowInfo->DependentView != &View && ProjectedShadowInfo->DependentView != View.GetPrimaryView())
 			{
 				continue;
 			}
@@ -2308,18 +2308,18 @@ void FDeferredShadingSceneRenderer::RenderDeferredShadowProjections(
 				// Inject the shadowed light into the translucency lighting volumes
 				if (ProjectedShadowInfo->DependentView != nullptr)
 				{
-					int32 ViewIndex = -1;
-					for (int32 i = 0; i < Views.Num(); ++i)
+					for (int32 ViewIndex = 0; ViewIndex < Views.Num(); ++ViewIndex)
 					{
-						if (ProjectedShadowInfo->DependentView == &Views[i])
+						const FViewInfo& View = Views[ViewIndex];
+						// allow secondary views in a stereo pair to share the shadows
+						if (ProjectedShadowInfo->DependentView != &View || ProjectedShadowInfo->DependentView != View.GetPrimaryView())
 						{
-							ViewIndex = i;
-							break;
+							continue;
 						}
-					}
 
-					RDG_GPU_MASK_SCOPE(GraphBuilder, ProjectedShadowInfo->DependentView->GPUMask);
-					InjectTranslucencyLightingVolume(GraphBuilder, *ProjectedShadowInfo->DependentView, ViewIndex, Scene, *this, TranslucencyLightingVolumeTextures, VisibleLightInfos, *LightSceneInfo, ProjectedShadowInfo);
+						RDG_GPU_MASK_SCOPE(GraphBuilder, ProjectedShadowInfo->DependentView->GPUMask);
+						InjectTranslucencyLightingVolume(GraphBuilder, View, ViewIndex, Scene, *this, TranslucencyLightingVolumeTextures, VisibleLightInfos, *LightSceneInfo, ProjectedShadowInfo);
+					}
 				}
 				else
 				{
