Apply local exposure when calculating bloom.
- Otherwise bloom intensity will not properly match the bright regions of the final image.
- Can be disabled using r.Bloom.ApplyLocalExposure.
- For gaussian bloom, local exposure is applied during the setup pass.
- For FFT bloom, local exposure is applied using an extra pass before convolution.

- Implemented AddApplyLocalExposurePass which apply local exposure to an input texture.

diff --git a/Engine/Shaders/Private/EyeAdaptationCommon.ush b/Engine/Shaders/Private/EyeAdaptationCommon.ush
--- a/Engine/Shaders/Private/EyeAdaptationCommon.ush
+++ b/Engine/Shaders/Private/EyeAdaptationCommon.ush
@@ -10,14 +10,14 @@ Texture2D EyeAdaptationTexture;
 Buffer<float4> EyeAdaptationBuffer;
 
 #if SHADING_PATH_DEFERRED
-float EyeAdaptationLookupTexture2D(Texture2D InEyeAdaptation)
+float4 EyeAdaptationLookupTexture2D(Texture2D InEyeAdaptation)
 {
-	return InEyeAdaptation.Load(int3(0, 0, 0)).x;
+	return InEyeAdaptation.Load(int3(0, 0, 0));
 }
 #elif SHADING_PATH_MOBILE
-float EyeAdaptationLookupBuffer(Buffer<float4> InEyeAdaptation)
+float4 EyeAdaptationLookupBuffer(Buffer<float4> InEyeAdaptation)
 {
-	return InEyeAdaptation[0].x;
+	return InEyeAdaptation[0];
 }
 #else
 	#error Unknown shading path.
@@ -44,17 +44,17 @@ float EyeAdaptationLookup()
 	return 0.0f;
 #elif defined (EyeAdaptationStruct)
 	#if SHADING_PATH_DEFERRED
-		return EyeAdaptationLookupTexture2D(EyeAdaptationStruct.EyeAdaptationTexture);
+		return EyeAdaptationLookupTexture2D(EyeAdaptationStruct.EyeAdaptationTexture).x;
 	#elif SHADING_PATH_MOBILE
-		return EyeAdaptationLookupBuffer(EyeAdaptationStruct.EyeAdaptationBuffer);
+		return EyeAdaptationLookupBuffer(EyeAdaptationStruct.EyeAdaptationBuffer).x;
 	#else
 		#error Unknown shading path.
 	#endif
 #elif EYE_ADAPTATION_LOOSE_PARAMETERS
 	#if SHADING_PATH_DEFERRED
-		return EyeAdaptationLookupTexture2D(EyeAdaptationTexture);
+		return EyeAdaptationLookupTexture2D(EyeAdaptationTexture).x;
 	#elif SHADING_PATH_MOBILE
-		return EyeAdaptationLookupBuffer(EyeAdaptationBuffer);
+		return EyeAdaptationLookupBuffer(EyeAdaptationBuffer).x;
 	#else
 		#error Unknown shading path.
 	#endif
diff --git a/Engine/Shaders/Private/PostProcessBloom.usf b/Engine/Shaders/Private/PostProcessBloom.usf
--- a/Engine/Shaders/Private/PostProcessBloom.usf
+++ b/Engine/Shaders/Private/PostProcessBloom.usf
@@ -9,31 +9,55 @@
 #include "Common.ush"
 #include "PostProcessCommon.ush"
 #include "EyeAdaptationCommon.ush"
+#include "PostProcessHistogramCommon.ush"
 #include "ScreenPass.ush"
 
 SCREEN_PASS_TEXTURE_VIEWPORT(Input)
 
 Texture2D InputTexture;
-
 SamplerState InputSampler;
 
+Texture3D LumBilateralGrid;
+SamplerState LumBilateralGridSampler;
+Texture2D BlurredLogLum;
+SamplerState BlurredLogLumSampler;
+
 float BloomThreshold;
 
-void BloomSetupVS(
-	float4 InPosition : ATTRIBUTE0,
-	float2 InTexCoord : ATTRIBUTE1,
-	out noperspective float4 OutUVAndScreenPos : TEXCOORD0,
-	out nointerpolation float OutExposureScale : TEXCOORD1,
-	out float4 OutPosition : SV_POSITION)
+float2 GetExposureScaleMiddleGreyLumValue()
 {
-	DrawRectangle(InPosition, InTexCoord, OutPosition, OutUVAndScreenPos);
-
-	OutExposureScale = EyeAdaptationLookup();
+#if SHADING_PATH_DEFERRED
+	float2 ExposureScaleMiddleGreyLumValue = EyeAdaptationLookupTexture2D(EyeAdaptationTexture).xw;
+#elif SHADING_PATH_MOBILE
+	float2 ExposureScaleMiddleGreyLumValue = EyeAdaptationLookupBuffer(EyeAdaptationBuffer).xw;
+#else
+#	error Unknown shading path.
+#endif
+
+	// Middle grey lum value adjusted by exposure compensation
+	ExposureScaleMiddleGreyLumValue.y = log2(0.18f * ExposureScaleMiddleGreyLumValue.y * EyeAdaptation_LocalExposureMiddleGreyExposureCompensation);
+
+	return ExposureScaleMiddleGreyLumValue;
 }
 
-float4 BloomSetupCommon(float2 UV, float ExposureScale)
+float4 BloomSetupCommon(float2 UV, float2 ExposureScaleMiddleGreyLumValue)
 {
-	float4 SceneColor = Texture2DSample(InputTexture, InputSampler, UV) * View.OneOverPreExposure; 
+	float4 SceneColor = Texture2DSample(InputTexture, InputSampler, UV) * View.OneOverPreExposure;
+
+	float ExposureScale = ExposureScaleMiddleGreyLumValue.x;
+
+#if USE_LOCAL_EXPOSURE
+	{
+		float LuminanceVal = CalculateEyeAdaptationLuminance(SceneColor.rgb);
+		float LogLuminance = log2(LuminanceVal);
+		float MiddleGreyLumValue = ExposureScaleMiddleGreyLumValue.y;
+
+		float BaseLogLum = CalculateBaseLogLuminance(LogLuminance, EyeAdaptation_LocalExposureBlurredLuminanceBlend, ExposureScale, UV, LumBilateralGrid, BlurredLogLum, LumBilateralGridSampler, BlurredLogLumSampler);
+		float LocalExposure = CalculateLocalExposure(LogLuminance + log2(ExposureScale), BaseLogLum, MiddleGreyLumValue, EyeAdaptation_LocalExposureContrastScale, EyeAdaptation_LocalExposureDetailStrength);
+		
+		SceneColor.rgb *= LocalExposure;
+	}
+#endif
 
 	// clamp to avoid artifacts from exceeding fp16 through framebuffer blending of multiple very bright lights
 	SceneColor.rgb = min(float3(256 * 256, 256 * 256, 256 * 256), SceneColor.rgb);
@@ -57,9 +81,9 @@ void BloomSetupPS(
 {
 	float2 UV = ApplyScreenTransform(SvPosition.xy, SvPositionToInputTextureUV);
 	
-	float ExposureScale = EyeAdaptationLookup();
+	float2 ExposureScaleMiddleGreyLumValue = GetExposureScaleMiddleGreyLumValue();
 
-	OutColor = BloomSetupCommon(UV, ExposureScale);
+	OutColor = BloomSetupCommon(UV, ExposureScaleMiddleGreyLumValue);
 }
 
 #elif COMPUTESHADER
@@ -76,9 +100,9 @@ void BloomSetupCS(uint2 DispatchThreadId : SV_DispatchThreadID)
 		return;
 	}
 
-	float ExposureScale = EyeAdaptationLookup();
+	float2 ExposureScaleMiddleGreyLumValue = GetExposureScaleMiddleGreyLumValue();
 
-	float4 OutColor = BloomSetupCommon(UV, ExposureScale);
+	float4 OutColor = BloomSetupCommon(UV, ExposureScaleMiddleGreyLumValue);
 
 	uint2 PixelPos = DispatchThreadId + Input_ViewportMin;
 
diff --git a/Engine/Shaders/Private/PostProcessLocalExposure.usf b/Engine/Shaders/Private/PostProcessLocalExposure.usf
--- a/Engine/Shaders/Private/PostProcessLocalExposure.usf
+++ b/Engine/Shaders/Private/PostProcessLocalExposure.usf
@@ -5,14 +5,52 @@
 #include "PostProcessHistogramCommon.ush"
 
 SCREEN_PASS_TEXTURE_VIEWPORT(Input)
+SCREEN_PASS_TEXTURE_VIEWPORT(Output)
+
 Texture2D InputTexture;
 
-RWTexture2D<float> Output2DFRW;
+Texture2D EyeAdaptationTexture;
+
+Texture3D LumBilateralGrid;
+Texture2D BlurredLogLum;
+
+SamplerState TextureSampler;
+
+RWTexture2D<float> OutputFloat;
+RWTexture2D<float4> OutputFloat4;
 
 [numthreads(THREADGROUP_SIZEX, THREADGROUP_SIZEY, 1)]
 void SetupLogLuminanceCS(uint3 DispatchThreadId : SV_DispatchThreadID)
 {
 	float3 SceneColor = InputTexture.Load(uint3(Input_ViewportMin + DispatchThreadId.xy, 0)).rgb * View.OneOverPreExposure;
 	float LuminanceVal = CalculateEyeAdaptationLuminance(SceneColor);
-	Output2DFRW[DispatchThreadId.xy] = log2(LuminanceVal);
+	OutputFloat[DispatchThreadId.xy] = log2(LuminanceVal);
+}
+
+[numthreads(THREADGROUP_SIZEX, THREADGROUP_SIZEY, 1)]
+void ApplyLocalExposureCS(uint2 DispatchThreadId : SV_DispatchThreadID)
+{
+	uint2 OutputPixelPos = DispatchThreadId + Output_ViewportMin;
+
+	if (all(OutputPixelPos < Output_ViewportMax))
+	{
+		const float2 UV = (DispatchThreadId + 0.5f) * Output_ExtentInverse;
+
+		float2 ExposureScaleMiddleGrey = EyeAdaptationTexture.Load(int3(0, 0, 0)).xw;
+		float ExposureScale = ExposureScaleMiddleGrey.x;
+		float MiddleGreyLumValue = log2(0.18 * ExposureScaleMiddleGrey.y * EyeAdaptation_LocalExposureMiddleGreyExposureCompensation);
+
+		float4 SceneColor = InputTexture.Load(uint3(DispatchThreadId + Input_ViewportMin, 0));
+		SceneColor.rgb *= View.OneOverPreExposure;
+
+		float LuminanceVal = CalculateEyeAdaptationLuminance(SceneColor.rgb);
+		float LogLuminance = log2(LuminanceVal);
+
+		float BaseLogLum = CalculateBaseLogLuminance(LogLuminance, EyeAdaptation_LocalExposureBlurredLuminanceBlend, ExposureScale, UV, LumBilateralGrid, BlurredLogLum, TextureSampler, TextureSampler);
+		float LocalExposure = CalculateLocalExposure(LogLuminance + log2(ExposureScale), BaseLogLum, MiddleGreyLumValue, EyeAdaptation_LocalExposureContrastScale, EyeAdaptation_LocalExposureDetailStrength);
+
+		SceneColor.rgb *= LocalExposure;
+
+		OutputFloat4[OutputPixelPos] = SceneColor;
+	}
 }
diff --git a/Engine/Source/Runtime/Renderer/Private/PostProcess/PostProcessBloomSetup.cpp b/Engine/Source/Runtime/Renderer/Private/PostProcess/PostProcessBloomSetup.cpp
--- a/Engine/Source/Runtime/Renderer/Private/PostProcess/PostProcessBloomSetup.cpp
+++ b/Engine/Source/Runtime/Renderer/Private/PostProcess/PostProcessBloomSetup.cpp
@@ -4,6 +4,7 @@
 #include "PostProcess/PostProcessDownsample.h"
 #include "PostProcess/PostProcessFFTBloom.h"
 #include "PostProcess/PostProcessWeightedSampleSum.h"
+#include "PostProcess/PostProcessEyeAdaptation.h"
 #include "PixelShaderUtils.h"
 
 namespace
@@ -26,24 +27,32 @@ BEGIN_SHADER_PARAMETER_STRUCT(FBloomSetupParameters, )
 	SHADER_PARAMETER_STRUCT(FScreenPassTextureViewportParameters, Input)
 	SHADER_PARAMETER_RDG_TEXTURE(Texture2D, InputTexture)
 	SHADER_PARAMETER_SAMPLER(SamplerState, InputSampler)
+	SHADER_PARAMETER_RDG_TEXTURE(Texture2D, LumBilateralGrid)
+	SHADER_PARAMETER_SAMPLER(SamplerState, LumBilateralGridSampler)
+	SHADER_PARAMETER_RDG_TEXTURE(Texture2D, BlurredLogLum)
+	SHADER_PARAMETER_SAMPLER(SamplerState, BlurredLogLumSampler)
 	SHADER_PARAMETER_RDG_TEXTURE(Texture2D, EyeAdaptationTexture)
+	SHADER_PARAMETER_STRUCT(FEyeAdaptationParameters, EyeAdaptation)
 	SHADER_PARAMETER(float, BloomThreshold)
 END_SHADER_PARAMETER_STRUCT()
 
 FBloomSetupParameters GetBloomSetupParameters(
 	const FViewInfo& View,
 	const FScreenPassTextureViewport& InputViewport,
-	FRDGTextureRef InputTexture,
-	FRDGTextureRef EyeAdaptationTexture,
-	float BloomThreshold)
+	const FBloomSetupInputs& Inputs)
 {
 	FBloomSetupParameters Parameters;
 	Parameters.View = View.ViewUniformBuffer;
 	Parameters.Input = GetScreenPassTextureViewportParameters(InputViewport);
-	Parameters.InputTexture = InputTexture;
+	Parameters.InputTexture = Inputs.SceneColor.Texture;
 	Parameters.InputSampler = TStaticSamplerState<SF_Bilinear, AM_Clamp, AM_Clamp, AM_Clamp>::GetRHI();
-	Parameters.EyeAdaptationTexture = EyeAdaptationTexture;
-	Parameters.BloomThreshold = BloomThreshold;
+	Parameters.LumBilateralGrid = Inputs.LocalExposureTexture;
+	Parameters.LumBilateralGridSampler = TStaticSamplerState<SF_Bilinear, AM_Clamp, AM_Clamp, AM_Clamp>::GetRHI();
+	Parameters.BlurredLogLum = Inputs.BlurredLogLuminanceTexture;
+	Parameters.BlurredLogLumSampler = TStaticSamplerState<SF_Bilinear, AM_Clamp, AM_Clamp, AM_Clamp>::GetRHI();
+	Parameters.EyeAdaptationTexture = Inputs.EyeAdaptationTexture;
+	Parameters.EyeAdaptation = *Inputs.EyeAdaptationParameters;
+	Parameters.BloomThreshold = Inputs.Threshold;
 	return Parameters;
 }
 
@@ -53,6 +62,9 @@ public:
 	DECLARE_GLOBAL_SHADER(FBloomSetupPS);
 	SHADER_USE_PARAMETER_STRUCT(FBloomSetupPS, FGlobalShader);
 
+	class FLocalExposureDim : SHADER_PERMUTATION_BOOL("USE_LOCAL_EXPOSURE");
+	using FPermutationDomain = TShaderPermutationDomain<FLocalExposureDim>;
+
 	BEGIN_SHADER_PARAMETER_STRUCT(FParameters, )
 		SHADER_PARAMETER_STRUCT_INCLUDE(FBloomSetupParameters, BloomSetup)
 		SHADER_PARAMETER(FScreenTransform, SvPositionToInputTextureUV)
@@ -73,6 +85,9 @@ public:
 	DECLARE_GLOBAL_SHADER(FBloomSetupCS);
 	SHADER_USE_PARAMETER_STRUCT(FBloomSetupCS, FGlobalShader);
 
+	class FLocalExposureDim : SHADER_PERMUTATION_BOOL("USE_LOCAL_EXPOSURE");
+	using FPermutationDomain = TShaderPermutationDomain<FLocalExposureDim>;
+
 	BEGIN_SHADER_PARAMETER_STRUCT(FParameters, )
 		SHADER_PARAMETER_STRUCT_INCLUDE(FBloomSetupParameters, BloomSetup)
 		SHADER_PARAMETER_RDG_TEXTURE_UAV(RWTexture2D<float4>, RWOutputTexture)
@@ -98,7 +113,7 @@ FScreenPassTexture AddBloomSetupPass(FRDGBuilder& GraphBuilder, const FViewInfo&
 {
 	check(Inputs.SceneColor.IsValid());
 	check(Inputs.EyeAdaptationTexture);
-	check(Inputs.Threshold > -1.0f);
+	check(Inputs.Threshold > -1.0f || Inputs.EyeAdaptationParameters != nullptr);
 
 	const bool bIsComputePass = View.bUseComputePasses;
 
@@ -112,10 +127,13 @@ FScreenPassTexture AddBloomSetupPass(FRDGBuilder& GraphBuilder, const FViewInfo&
 	if (bIsComputePass)
 	{
 		FBloomSetupCS::FParameters* PassParameters = GraphBuilder.AllocParameters<FBloomSetupCS::FParameters>();
-		PassParameters->BloomSetup = GetBloomSetupParameters(View, Viewport, Inputs.SceneColor.Texture, Inputs.EyeAdaptationTexture, Inputs.Threshold);
+		PassParameters->BloomSetup = GetBloomSetupParameters(View, Viewport, Inputs);
 		PassParameters->RWOutputTexture = GraphBuilder.CreateUAV(Output.Texture);
 
-		TShaderMapRef<FBloomSetupCS> ComputeShader(View.ShaderMap);
+		FBloomSetupCS::FPermutationDomain PermutationVector;
+		PermutationVector.Set<FBloomSetupCS::FLocalExposureDim>(Inputs.LocalExposureTexture != nullptr);
+
+		auto ComputeShader = View.ShaderMap->GetShader<FBloomSetupCS>(PermutationVector);
 
 		FComputeShaderUtils::AddPass(
 			GraphBuilder,
@@ -127,13 +145,17 @@ FScreenPassTexture AddBloomSetupPass(FRDGBuilder& GraphBuilder, const FViewInfo&
 	else
 	{
 		FBloomSetupPS::FParameters* PassParameters = GraphBuilder.AllocParameters<FBloomSetupPS::FParameters>();
-		PassParameters->BloomSetup = GetBloomSetupParameters(View, Viewport, Inputs.SceneColor.Texture, Inputs.EyeAdaptationTexture, Inputs.Threshold);
+		PassParameters->BloomSetup = GetBloomSetupParameters(View, Viewport, Inputs);
 		PassParameters->SvPositionToInputTextureUV = (
 			FScreenTransform::ChangeTextureBasisFromTo(FScreenPassTextureViewport(Output), FScreenTransform::ETextureBasis::TexelPosition, FScreenTransform::ETextureBasis::ViewportUV) *
 			FScreenTransform::ChangeTextureBasisFromTo(FScreenPassTextureViewport(Inputs.SceneColor), FScreenTransform::ETextureBasis::ViewportUV, FScreenTransform::ETextureBasis::TextureUV));
 		PassParameters->RenderTargets[0] = Output.GetRenderTargetBinding();
 
-		TShaderMapRef<FBloomSetupPS> PixelShader(View.ShaderMap);
+		FBloomSetupPS::FPermutationDomain PermutationVector;
+		PermutationVector.Set<FBloomSetupPS::FLocalExposureDim>(Inputs.LocalExposureTexture != nullptr);
+
+		auto PixelShader = View.ShaderMap->GetShader<FBloomSetupPS>(PermutationVector);
+
 		FPixelShaderUtils::AddFullscreenPass(
 			GraphBuilder,
 			View.ShaderMap,
diff --git a/Engine/Source/Runtime/Renderer/Private/PostProcess/PostProcessBloomSetup.h b/Engine/Source/Runtime/Renderer/Private/PostProcess/PostProcessBloomSetup.h
--- a/Engine/Source/Runtime/Renderer/Private/PostProcess/PostProcessBloomSetup.h
+++ b/Engine/Source/Runtime/Renderer/Private/PostProcess/PostProcessBloomSetup.h
@@ -5,6 +5,7 @@
 #include "ScreenPass.h"
 
 class FSceneDownsampleChain;
+class FEyeAdaptationParameters;
 
 enum class EBloomQuality : uint32
 {
@@ -31,6 +32,15 @@ struct FBloomSetupInputs
 
 	// [Required]: The bloom threshold to apply. Must be >0.
 	float Threshold = 0.0f;
+
+	// [Optional] Eye adaptation parameters.
+	const FEyeAdaptationParameters* EyeAdaptationParameters = nullptr;
+
+	// [Optional] Luminance bilateral grid. If this is null, local exposure is disabled.
+	FRDGTextureRef LocalExposureTexture = nullptr;
+
+	// [Optional] Blurred luminance texture used to calculate local exposure.
+	FRDGTextureRef BlurredLogLuminanceTexture = nullptr;
 };
 
 FScreenPassTexture AddBloomSetupPass(FRDGBuilder& GraphBuilder, const FViewInfo& View, const FBloomSetupInputs& Inputs);
\ No newline at end of file
diff --git a/Engine/Source/Runtime/Renderer/Private/PostProcess/PostProcessFFTBloom.cpp b/Engine/Source/Runtime/Renderer/Private/PostProcess/PostProcessFFTBloom.cpp
--- a/Engine/Source/Runtime/Renderer/Private/PostProcess/PostProcessFFTBloom.cpp
+++ b/Engine/Source/Runtime/Renderer/Private/PostProcess/PostProcessFFTBloom.cpp
@@ -1,8 +1,8 @@
 // Copyright Epic Games, Inc. All Rights Reserved.
 
 #include "PostProcess/PostProcessFFTBloom.h"
-#include "PostProcess/PostProcessBloomSetup.h"
 #include "PostProcess/PostProcessTonemap.h"
+#include "PostProcess/PostProcessLocalExposure.h"
 #include "PostProcess/PostProcessDownsample.h"
 #include "GPUFastFourierTransform.h"
 #include "RendererModule.h"
@@ -631,7 +631,15 @@ void InitDomainAndGetKernel(
 	*OutKernelConstantsBuffer = KernelConstantsBuffer;
 }
 
-FFFTBloomOutput AddFFTBloomPass(FRDGBuilder& GraphBuilder, const FViewInfo& View, const FScreenPassTexture& InputSceneColor, float InputResolutionFraction)
+FFFTBloomOutput AddFFTBloomPass(
+	FRDGBuilder& GraphBuilder,
+	const FViewInfo& View,
+	const FScreenPassTexture& InputSceneColor,
+	float InputResolutionFraction,
+	const FEyeAdaptationParameters& EyeAdaptationParameters,
+	FRDGTextureRef EyeAdaptationTexture,
+	FRDGTextureRef LocalExposureTexture,
+	FRDGTextureRef BlurredLogLuminanceTexture)
 {
 	check(InputSceneColor.IsValid());
 
@@ -731,6 +739,15 @@ FFFTBloomOutput AddFFTBloomPass(FRDGBuilder& GraphBuilder, const FViewInfo& View
 		FFTInputSceneColor = InputSceneColor;
 	}
 
+	if (LocalExposureTexture != nullptr)
+	{
+		FScreenPassTexture Temp = FFTInputSceneColor;
+
+		FFTInputSceneColor.Texture = GraphBuilder.CreateTexture(Temp.Texture->Desc, TEXT("Bloom.FFT.Input"));
+
+		AddApplyLocalExposurePass(GraphBuilder, View, EyeAdaptationParameters, EyeAdaptationTexture, LocalExposureTexture, BlurredLogLuminanceTexture, Temp, FFTInputSceneColor, Intermediates.ComputePassFlags);
+	}
+
 	// Init the domain data update the cached kernel if needed.
 	FRDGTextureRef SpectralKernelTexture;
 	FRDGBufferRef KernelConstantsBuffer;
diff --git a/Engine/Source/Runtime/Renderer/Private/PostProcess/PostProcessFFTBloom.h b/Engine/Source/Runtime/Renderer/Private/PostProcess/PostProcessFFTBloom.h
--- a/Engine/Source/Runtime/Renderer/Private/PostProcess/PostProcessFFTBloom.h
+++ b/Engine/Source/Runtime/Renderer/Private/PostProcess/PostProcessFFTBloom.h
@@ -6,7 +6,7 @@
 #include "ScenePrivate.h"
 #include "ScreenPass.h"
 
-struct FBloomOutputs;
+class FEyeAdaptationParameters;
 
 // Returns whether FFT bloom is enabled for the view.
 bool IsFFTBloomEnabled(const FViewInfo& View);
@@ -19,4 +19,12 @@ struct FFFTBloomOutput
 	FRDGBufferRef SceneColorApplyParameters = nullptr;
 };
 
-FFFTBloomOutput AddFFTBloomPass(FRDGBuilder& GraphBuilder, const FViewInfo& View, const FScreenPassTexture& InputSceneColor, float InputResolutionFraction);
+FFFTBloomOutput AddFFTBloomPass(
+	FRDGBuilder& GraphBuilder,
+	const FViewInfo& View,
+	const FScreenPassTexture& InputSceneColor,
+	float InputResolutionFraction,
+	const FEyeAdaptationParameters& EyeAdaptationParameters,
+	FRDGTextureRef EyeAdaptationTexture,
+	FRDGTextureRef LocalExposureTexture,
+	FRDGTextureRef BlurredLogLuminanceTexture);
diff --git a/Engine/Source/Runtime/Renderer/Private/PostProcess/PostProcessLocalExposure.cpp b/Engine/Source/Runtime/Renderer/Private/PostProcess/PostProcessLocalExposure.cpp
--- a/Engine/Source/Runtime/Renderer/Private/PostProcess/PostProcessLocalExposure.cpp
+++ b/Engine/Source/Runtime/Renderer/Private/PostProcess/PostProcessLocalExposure.cpp
@@ -27,7 +27,7 @@ public:
 		SHADER_PARAMETER_STRUCT(FEyeAdaptationParameters, EyeAdaptation)
 		SHADER_PARAMETER_STRUCT(FScreenPassTextureViewportParameters, Input)
 		SHADER_PARAMETER_RDG_TEXTURE(Texture2D, InputTexture)
-		SHADER_PARAMETER_RDG_TEXTURE_UAV(RWTexture2D<float>, Output2DFRW)
+		SHADER_PARAMETER_RDG_TEXTURE_UAV(RWTexture2D<float>, OutputFloat)
 	END_SHADER_PARAMETER_STRUCT()
 
 	static bool ShouldCompilePermutation(const FGlobalShaderPermutationParameters& Parameters)
@@ -46,6 +46,48 @@ public:
 
 IMPLEMENT_GLOBAL_SHADER(FSetupLogLuminanceCS, "/Engine/Private/PostProcessLocalExposure.usf", "SetupLogLuminanceCS", SF_Compute);
 
+class FApplyLocalExposureCS : public FGlobalShader
+{
+public:
+	// Changing these numbers requires LocalExposure.usf to be recompiled
+	static const uint32 ThreadGroupSizeX = 8;
+	static const uint32 ThreadGroupSizeY = 8;
+
+	DECLARE_GLOBAL_SHADER(FApplyLocalExposureCS);
+	SHADER_USE_PARAMETER_STRUCT(FApplyLocalExposureCS, FGlobalShader);
+
+	BEGIN_SHADER_PARAMETER_STRUCT(FParameters, )
+		SHADER_PARAMETER_STRUCT_REF(FViewUniformShaderParameters, View)
+
+		SHADER_PARAMETER_STRUCT(FScreenPassTextureViewportParameters, Input)
+		SHADER_PARAMETER_STRUCT(FScreenPassTextureViewportParameters, Output)
+		SHADER_PARAMETER_RDG_TEXTURE(Texture2D, InputTexture)
+		SHADER_PARAMETER_RDG_TEXTURE_UAV(RWTexture2D<float4>, OutputFloat4)
+
+		SHADER_PARAMETER_STRUCT(FEyeAdaptationParameters, EyeAdaptation)
+		SHADER_PARAMETER_RDG_TEXTURE(Texture2D, EyeAdaptationTexture)
+		SHADER_PARAMETER_RDG_TEXTURE(Texture3D, LumBilateralGrid)
+		SHADER_PARAMETER_RDG_TEXTURE(Texture2D, BlurredLogLum)
+
+		SHADER_PARAMETER_SAMPLER(SamplerState, TextureSampler)
+	END_SHADER_PARAMETER_STRUCT()
+
+	static bool ShouldCompilePermutation(const FGlobalShaderPermutationParameters& Parameters)
+	{
+		return IsFeatureLevelSupported(Parameters.Platform, ERHIFeatureLevel::SM5);
+	}
+
+	static void ModifyCompilationEnvironment(const FGlobalShaderPermutationParameters& Parameters, FShaderCompilerEnvironment& OutEnvironment)
+	{
+		FGlobalShader::ModifyCompilationEnvironment(Parameters, OutEnvironment);
+		OutEnvironment.SetDefine(TEXT("THREADGROUP_SIZEX"), ThreadGroupSizeX);
+		OutEnvironment.SetDefine(TEXT("THREADGROUP_SIZEY"), ThreadGroupSizeY);
+		OutEnvironment.CompilerFlags.Add(CFLAG_StandardOptimization);
+	}
+};
+
+IMPLEMENT_GLOBAL_SHADER(FApplyLocalExposureCS, "/Engine/Private/PostProcessLocalExposure.usf", "ApplyLocalExposureCS", SF_Compute);
+
 } //! namespace
 
 FRDGTextureRef AddLocalExposureBlurredLogLuminancePass(
@@ -75,7 +117,7 @@ FRDGTextureRef AddLocalExposureBlurredLogLuminancePass(
 		PassParameters->EyeAdaptation = EyeAdaptationParameters;
 		PassParameters->Input = GetScreenPassTextureViewportParameters(FScreenPassTextureViewport(InputTexture));
 		PassParameters->InputTexture = InputTexture.Texture;
-		PassParameters->Output2DFRW = GraphBuilder.CreateUAV(GaussianLumSetupTexture);
+		PassParameters->OutputFloat = GraphBuilder.CreateUAV(GaussianLumSetupTexture);
 
 		FComputeShaderUtils::AddPass(
 			GraphBuilder,
@@ -102,4 +144,43 @@ FRDGTextureRef AddLocalExposureBlurredLogLuminancePass(
 	}
 
 	return GaussianTexture;
+}
+
+void AddApplyLocalExposurePass(
+	FRDGBuilder& GraphBuilder,
+	const FViewInfo& View,
+	const FEyeAdaptationParameters& EyeAdaptationParameters,
+	FRDGTextureRef EyeAdaptationTexture,
+	FRDGTextureRef LocalExposureTexture,
+	FRDGTextureRef BlurredLogLuminanceTexture,
+	FScreenPassTexture Input,
+	FScreenPassTexture Output,
+	ERDGPassFlags PassFlags)
+{
+	check(Input.IsValid() && Output.IsValid());
+	check(PassFlags == ERDGPassFlags::Compute || PassFlags == ERDGPassFlags::AsyncCompute);
+
+	RDG_EVENT_SCOPE(GraphBuilder, "LocalExposure - Apply");
+
+	auto* PassParameters = GraphBuilder.AllocParameters<FApplyLocalExposureCS::FParameters>();
+	PassParameters->View = View.ViewUniformBuffer;
+	PassParameters->Input = GetScreenPassTextureViewportParameters(FScreenPassTextureViewport(Input));
+	PassParameters->Output = GetScreenPassTextureViewportParameters(FScreenPassTextureViewport(Output));
+
+	PassParameters->InputTexture = Input.Texture;
+	PassParameters->OutputFloat4 = GraphBuilder.CreateUAV(Output.Texture);
+
+	PassParameters->EyeAdaptation = EyeAdaptationParameters;
+	PassParameters->EyeAdaptationTexture = EyeAdaptationTexture;
+	PassParameters->LumBilateralGrid = LocalExposureTexture;
+	PassParameters->BlurredLogLum = BlurredLogLuminanceTexture;
+	PassParameters->TextureSampler = TStaticSamplerState<SF_Bilinear, AM_Clamp, AM_Clamp, AM_Clamp>::GetRHI();;
+
+	FComputeShaderUtils::AddPass(
+		GraphBuilder,
+		RDG_EVENT_NAME("ApplyLocalExposure %dx%d", Output.ViewRect.Width(), Output.ViewRect.Height()),
+		PassFlags,
+		View.ShaderMap->GetShader<FApplyLocalExposureCS>(),
+		PassParameters,
+		FComputeShaderUtils::GetGroupCount(Output.ViewRect.Size(), FIntPoint(FApplyLocalExposureCS::ThreadGroupSizeX, FApplyLocalExposureCS::ThreadGroupSizeY)));
 }
\ No newline at end of file
diff --git a/Engine/Source/Runtime/Renderer/Private/PostProcess/PostProcessLocalExposure.h b/Engine/Source/Runtime/Renderer/Private/PostProcess/PostProcessLocalExposure.h
--- a/Engine/Source/Runtime/Renderer/Private/PostProcess/PostProcessLocalExposure.h
+++ b/Engine/Source/Runtime/Renderer/Private/PostProcess/PostProcessLocalExposure.h
@@ -11,3 +11,14 @@ FRDGTextureRef AddLocalExposureBlurredLogLuminancePass(
 	const FViewInfo& View,
 	const FEyeAdaptationParameters& EyeAdaptationParameters,
 	FScreenPassTexture InputTexture);
+
+void AddApplyLocalExposurePass(
+	FRDGBuilder& GraphBuilder,
+	const FViewInfo& View,
+	const FEyeAdaptationParameters& EyeAdaptationParameters,
+	FRDGTextureRef EyeAdaptationTexture,
+	FRDGTextureRef LocalExposureTexture,
+	FRDGTextureRef BlurredLogLuminanceTexture,
+	FScreenPassTexture Input,
+	FScreenPassTexture Output,
+	ERDGPassFlags PassFlags);
\ No newline at end of file
diff --git a/Engine/Source/Runtime/Renderer/Private/PostProcess/PostProcessing.cpp b/Engine/Source/Runtime/Renderer/Private/PostProcess/PostProcessing.cpp
--- a/Engine/Source/Runtime/Renderer/Private/PostProcess/PostProcessing.cpp
+++ b/Engine/Source/Runtime/Renderer/Private/PostProcess/PostProcessing.cpp
@@ -81,6 +81,12 @@ TAutoConsoleVariable<float> CVarDepthOfFieldMaxSize(
 	TEXT("Allows to clamp the gaussian depth of field radius (for better performance), default: 100"),
 	ECVF_Scalability | ECVF_RenderThreadSafe);
 
+TAutoConsoleVariable<bool> CVarBloomApplyLocalExposure(
+	TEXT("r.Bloom.ApplyLocalExposure"),
+	true,
+	TEXT("Whether to apply local exposure when calculating bloom, default: true"),
+	ECVF_Scalability | ECVF_RenderThreadSafe);
+
 TAutoConsoleVariable<int32> CVarPostProcessingPropagateAlpha(
 	TEXT("r.PostProcessing.PropagateAlpha"),
 	0,
@@ -862,16 +868,25 @@ void AddPostProcessingPasses(
 					InputResolutionFraction = 1.0f;
 				}
 
-				FFFTBloomOutput Outputs = AddFFTBloomPass(GraphBuilder, View, InputSceneColor, InputResolutionFraction);
+				FFFTBloomOutput Outputs = AddFFTBloomPass(
+					GraphBuilder, 
+					View,
+					InputSceneColor,
+					InputResolutionFraction,
+					EyeAdaptationParameters,
+					EyeAdaptationTexture,
+					CVarBloomApplyLocalExposure.GetValueOnRenderThread() ? LocalExposureTexture : nullptr,
+					LocalExposureBlurredLogLumTexture);
+
 				Bloom = Outputs.BloomTexture;
 				SceneColorApplyParameters = Outputs.SceneColorApplyParameters;
 			}
 			else
 			{
-				const bool bBloomThresholdEnabled = View.FinalPostProcessSettings.BloomThreshold > -1.0f;
+				const bool bBloomSetupRequiredEnabled = View.FinalPostProcessSettings.BloomThreshold > -1.0f || LocalExposureTexture != nullptr;
 
-				// Reuse the main scene downsample chain if a threshold isn't required for gaussian bloom.
-				if (SceneDownsampleChain.IsInitialized() && !bBloomThresholdEnabled)
+				// Reuse the main scene downsample chain if setup isn't required for gaussian bloom.
+				if (SceneDownsampleChain.IsInitialized() && !bBloomSetupRequiredEnabled)
 				{
 					LensFlareSceneDownsampleChain = &SceneDownsampleChain;
 				}
@@ -879,13 +894,16 @@ void AddPostProcessingPasses(
 				{
 					FScreenPassTexture DownsampleInput = bProcessQuarterResolution ? QuarterResSceneColor : HalfResSceneColor;
 
-					if (bBloomThresholdEnabled)
+					if (bBloomSetupRequiredEnabled)
 					{
 						const float BloomThreshold = View.FinalPostProcessSettings.BloomThreshold;
 
 						FBloomSetupInputs SetupPassInputs;
 						SetupPassInputs.SceneColor = DownsampleInput;
 						SetupPassInputs.EyeAdaptationTexture = EyeAdaptationTexture;
+						SetupPassInputs.EyeAdaptationParameters = &EyeAdaptationParameters;
+						SetupPassInputs.LocalExposureTexture = CVarBloomApplyLocalExposure.GetValueOnRenderThread() ? LocalExposureTexture : nullptr;
+						SetupPassInputs.BlurredLogLuminanceTexture = LocalExposureBlurredLogLumTexture;
 						SetupPassInputs.Threshold = BloomThreshold;
 
 						DownsampleInput = AddBloomSetupPass(GraphBuilder, View, SetupPassInputs);
