Improved ray tracing decals checks.
- Ray tracing decal grid is only created (CreateRayTracingDecalData(...)) if ShouldPrepareRayTracingDecals().
- ShouldPrepareRayTracingDecals() checks whether runtime supports the feature, there's decals in the scene, and any technique actually uses decals (currently only PathTracing).
- Added bHasRayTracingDecals to FViewInfo since RayTracingDecalUniformBuffer always has a valid uniform buffer (but might not contain any decals).
    - Techniques check bHasRayTracingDecals decide whether to use permutations that use decals.

diff --git a/Engine/Source/Runtime/Renderer/Private/DeferredShadingRenderer.cpp b/Engine/Source/Runtime/Renderer/Private/DeferredShadingRenderer.cpp
--- a/Engine/Source/Runtime/Renderer/Private/DeferredShadingRenderer.cpp
+++ b/Engine/Source/Runtime/Renderer/Private/DeferredShadingRenderer.cpp
@@ -35,6 +35,7 @@
 #include "VT/VirtualTextureFeedback.h"
 #include "VT/VirtualTextureSystem.h"
 #include "GPUScene.h"
+#include "PathTracing.h"
 #include "RayTracing/RayTracingMaterialHitShaders.h"
 #include "RayTracing/RayTracingLighting.h"
 #include "RayTracing/RayTracingDecals.h"
@@ -548,6 +549,21 @@ static void RenderOpaqueFX(
 
 #if RHI_RAYTRACING
 
+static bool ShouldPrepareRayTracingDecals(const FScene& Scene, const FSceneViewFamily& ViewFamily)
+{
+	if (!IsRayTracingEnabled() || !RHISupportsRayTracingCallableShaders(ViewFamily.GetShaderPlatform()))
+	{
+		return false;
+	}
+
+	if (Scene.Decals.Num() == 0 || GRayTracingExcludeDecals)
+	{
+		return false;
+	}
+
+	return ViewFamily.EngineShowFlags.PathTracing && PathTracing::UsesDecals(ViewFamily);
+}
+
 static void AddDebugRayTracingInstanceFlags(ERayTracingInstanceFlags& InOutFlags)
 {
 	if (GRayTracingDebugForceOpaque)
@@ -2166,11 +2182,17 @@ void FDeferredShadingSceneRenderer::Render(FRDGBuilder& GraphBuilder)
 	const int32 ReferenceViewIndex = 0;
 	FViewInfo& ReferenceView = Views[ReferenceViewIndex];
 
-	if (ViewFamily.EngineShowFlags.PathTracing)
+	if (ShouldPrepareRayTracingDecals(*Scene, *ReferenceView.Family))
 	{
 		ReferenceView.RayTracingDecalUniformBuffer = CreateRayTracingDecalData(GraphBuilder, *Scene, ReferenceView, RayTracingScene.NumCallableShaderSlots);
+		ReferenceView.bHasRayTracingDecals = true;
 		RayTracingScene.NumCallableShaderSlots += Scene->Decals.Num();
 	}
+	else
+	{
+		ReferenceView.RayTracingDecalUniformBuffer = CreateNullRayTracingDecalsUniformBuffer(GraphBuilder);
+		ReferenceView.bHasRayTracingDecals = false;
+	}
 
 	// Prepare the scene for rendering this frame.
 	GatherRayTracingWorldInstancesForView(GraphBuilder, ReferenceView, RayTracingScene);
diff --git a/Engine/Source/Runtime/Renderer/Private/PathTracing.cpp b/Engine/Source/Runtime/Renderer/Private/PathTracing.cpp
--- a/Engine/Source/Runtime/Renderer/Private/PathTracing.cpp
+++ b/Engine/Source/Runtime/Renderer/Private/PathTracing.cpp
@@ -419,6 +419,14 @@ struct FPathTracingState {
 	uint32_t FrameIndex = 0;
 };
 
+namespace PathTracing
+{
+	bool UsesDecals(const FSceneViewFamily& ViewFamily)
+	{
+		return ViewFamily.EngineShowFlags.Decals;
+	}
+}
+
 // This function prepares the portion of shader arguments that may involve invalidating the path traced state
 static void PreparePathTracingData(const FScene* Scene, const FViewInfo& View, FPathTracingData& PathTracingData)
 {
@@ -499,7 +507,7 @@ static void PreparePathTracingData(const FScene* Scene, const FViewInfo& View, F
 		&& (Scene->ExponentialFogs[0].FogData[0].Density > 0 ||
 			Scene->ExponentialFogs[0].FogData[1].Density > 0);
 
-	PathTracingData.EnableDecals = View.Family->EngineShowFlags.Decals && Scene->Decals.Num() > 0;
+	PathTracingData.EnableDecals = PathTracing::UsesDecals(*View.Family) && View.bHasRayTracingDecals;
 
 	PathTracingData.MaxRaymarchSteps = CVarPathTracingMaxRaymarchSteps.GetValueOnRenderThread();
 }
@@ -946,7 +954,6 @@ bool FRayTracingMeshProcessor::ProcessPathTracing(
 	const FMaterialRenderProxy& RESTRICT MaterialRenderProxy,
 	const FMaterial& RESTRICT MaterialResource)
 {
-
 	FMaterialShaderTypes ShaderTypes;
 
 	const bool bUseProceduralPrimitive = MeshBatch.VertexFactory->GetType()->SupportsRayTracingProceduralPrimitive() &&
diff --git a/Engine/Source/Runtime/Renderer/Private/PathTracing.h b/Engine/Source/Runtime/Renderer/Private/PathTracing.h
--- a/Engine/Source/Runtime/Renderer/Private/PathTracing.h
+++ b/Engine/Source/Runtime/Renderer/Private/PathTracing.h
@@ -16,5 +16,11 @@ END_SHADER_PARAMETER_STRUCT()
 class FRDGBuilder;
 class FScene;
 class FViewInfo;
+class FSceneViewFamily;
 
 bool PrepareSkyTexture(FRDGBuilder& GraphBuilder, FScene* Scene, const FViewInfo& View, bool SkylightEnabled, bool UseMISCompensation, FPathTracingSkylight* SkylightParameters);
+
+namespace PathTracing
+{
+	bool UsesDecals(const FSceneViewFamily& ViewFamily);
+}
diff --git a/Engine/Source/Runtime/Renderer/Private/RayTracing/RayTracingDecals.cpp b/Engine/Source/Runtime/Renderer/Private/RayTracing/RayTracingDecals.cpp
--- a/Engine/Source/Runtime/Renderer/Private/RayTracing/RayTracingDecals.cpp
+++ b/Engine/Source/Runtime/Renderer/Private/RayTracing/RayTracingDecals.cpp
@@ -334,20 +334,20 @@ FTransientDecalRenderDataList GetSortedDecals(const TSparseArray<FDeferredDecalP
 	return SortedDecals;
 }
 
-TRDGUniformBufferRef<FRayTracingDecals> CreateRayTracingDecalData(FRDGBuilder& GraphBuilder, FScene& Scene, const FViewInfo& View, uint32 BaseCallableSlotIndex)
+TRDGUniformBufferRef<FRayTracingDecals> CreateNullRayTracingDecalsUniformBuffer(FRDGBuilder& GraphBuilder)
 {
-	if (!RHISupportsRayTracingCallableShaders(View.Family->GetShaderPlatform()))
-	{
-		auto* OutParameters = GraphBuilder.AllocParameters<FRayTracingDecals>();
-		OutParameters->Count = 0;
-		OutParameters->TranslatedBoundMin = FVector3f::ZeroVector;
-		OutParameters->TranslatedBoundMax = FVector3f::ZeroVector;
-
-		BuildDecalGrid(GraphBuilder, 0, nullptr, *OutParameters);
-
-		return GraphBuilder.CreateUniformBuffer(OutParameters);
-	}
+	auto* OutParameters = GraphBuilder.AllocParameters<FRayTracingDecals>();
+	OutParameters->Count = 0;
+	OutParameters->TranslatedBoundMin = FVector3f::ZeroVector;
+	OutParameters->TranslatedBoundMax = FVector3f::ZeroVector;
+	
+	BuildDecalGrid(GraphBuilder, 0, nullptr, *OutParameters);
+	
+	return GraphBuilder.CreateUniformBuffer(OutParameters);
+}
 
+TRDGUniformBufferRef<FRayTracingDecals> CreateRayTracingDecalData(FRDGBuilder& GraphBuilder, FScene& Scene, const FViewInfo& View, uint32 BaseCallableSlotIndex)
+{
 	const EShaderPlatform ShaderPlatform = View.Family->GetShaderPlatform();
 
 	FTransientDecalRenderDataList SortedDecals = GetSortedDecals(Scene.Decals, Scene, View);
diff --git a/Engine/Source/Runtime/Renderer/Private/RayTracing/RayTracingDecals.h b/Engine/Source/Runtime/Renderer/Private/RayTracing/RayTracingDecals.h
--- a/Engine/Source/Runtime/Renderer/Private/RayTracing/RayTracingDecals.h
+++ b/Engine/Source/Runtime/Renderer/Private/RayTracing/RayTracingDecals.h
@@ -21,6 +21,7 @@ BEGIN_GLOBAL_SHADER_PARAMETER_STRUCT(FRayTracingDecals, RENDERER_API)
 	SHADER_PARAMETER(uint32, Count)
 END_GLOBAL_SHADER_PARAMETER_STRUCT()
 
+TRDGUniformBufferRef<FRayTracingDecals> CreateNullRayTracingDecalsUniformBuffer(FRDGBuilder& GraphBuilder);
 TRDGUniformBufferRef<FRayTracingDecals> CreateRayTracingDecalData(FRDGBuilder& GraphBuilder, FScene& Scene, const FViewInfo& View, uint32 BaseCallableSlotIndex);
 
 #endif
diff --git a/Engine/Source/Runtime/Renderer/Private/SceneRendering.h b/Engine/Source/Runtime/Renderer/Private/SceneRendering.h
--- a/Engine/Source/Runtime/Renderer/Private/SceneRendering.h
+++ b/Engine/Source/Runtime/Renderer/Private/SceneRendering.h
@@ -1502,6 +1502,7 @@ public:
 
 	TRDGUniformBufferRef<FRaytracingLightDataPacked>	RayTracingLightDataUniformBuffer;
 	TRDGUniformBufferRef<FRayTracingDecals>				RayTracingDecalUniformBuffer;
+	bool												bHasRayTracingDecals = false;
 
 	bool bHasAnyRayTracingPass = false;
 #endif // RHI_RAYTRACING
