Keep track of TargetBuffers we already transitioned for ray tracing dynamic geometry to prevent adding unnecessary UAV transitions

diff --git a/Engine/Source/Runtime/Renderer/Private/RayTracing/RayTracingDynamicGeometry.cpp b/Engine/Source/Runtime/Renderer/Private/RayTracing/RayTracingDynamicGeometry.cpp
--- a/Engine/Source/Runtime/Renderer/Private/RayTracing/RayTracingDynamicGeometry.cpp
+++ b/Engine/Source/Runtime/Renderer/Private/RayTracing/RayTracingDynamicGeometry.cpp
@@ -420,6 +420,7 @@ void FRayTracingDynamicGeometryCollection::DispatchUpdates(FRHIComputeCommandLis
 			TransitionsAfter.Reserve(DispatchCommands.Num());
 			OverlapUAVs.Reserve(DispatchCommands.Num());
 			const FRWBuffer* LastBuffer = nullptr;
+			TSet<const FRWBuffer*> TransitionedBuffers;
 			for (FMeshComputeDispatchCommand& Cmd : DispatchCommands)
 			{
 				if (Cmd.TargetBuffer == nullptr)
@@ -442,9 +443,15 @@ void FRayTracingDynamicGeometryCollection::DispatchUpdates(FRHIComputeCommandLis
 
 				LastBuffer = Cmd.TargetBuffer;
 
-				// Looks like the resource can get here in either UAVCompute or SRVMask mode, so we'll have to use Unknown until we can have better tracking.
-				TransitionsBefore.Add(FRHITransitionInfo(UAV, ERHIAccess::Unknown, ERHIAccess::UAVCompute));
-				TransitionsAfter.Add(FRHITransitionInfo(UAV, ERHIAccess::UAVCompute, ERHIAccess::SRVMask));
+				// In case different shaders use different TargetBuffer we want to add transition only once
+				bool bAlreadyInSet = false;
+				TransitionedBuffers.FindOrAdd(LastBuffer, &bAlreadyInSet);
+				if (!bAlreadyInSet)
+				{
+					// Looks like the resource can get here in either UAVCompute or SRVMask mode, so we'll have to use Unknown until we can have better tracking.
+					TransitionsBefore.Add(FRHITransitionInfo(UAV, ERHIAccess::Unknown, ERHIAccess::UAVCompute));
+					TransitionsAfter.Add(FRHITransitionInfo(UAV, ERHIAccess::UAVCompute, ERHIAccess::SRVMask));
+				}
 			}
 
 			TArray<FRHICommandList*> CommandLists;
