Fix missing UAV flag when applying local exposure in FFTBloom.

diff --git a/Engine/Source/Runtime/Renderer/Private/PostProcess/PostProcessFFTBloom.cpp b/Engine/Source/Runtime/Renderer/Private/PostProcess/PostProcessFFTBloom.cpp
--- a/Engine/Source/Runtime/Renderer/Private/PostProcess/PostProcessFFTBloom.cpp
+++ b/Engine/Source/Runtime/Renderer/Private/PostProcess/PostProcessFFTBloom.cpp
@@ -743,7 +743,10 @@ FFFTBloomOutput AddFFTBloomPass(
 	{
 		FScreenPassTexture Temp = FFTInputSceneColor;
 
-		FFTInputSceneColor.Texture = GraphBuilder.CreateTexture(Temp.Texture->Desc, TEXT("Bloom.FFT.Input"));
+		FRDGTextureDesc Desc = Temp.Texture->Desc;
+		Desc.Flags = TexCreate_ShaderResource | TexCreate_UAV;
+
+		FFTInputSceneColor.Texture = GraphBuilder.CreateTexture(Desc, TEXT("Bloom.FFT.Input"));
 
 		AddApplyLocalExposurePass(GraphBuilder, View, EyeAdaptationParameters, EyeAdaptationTexture, LocalExposureTexture, BlurredLogLuminanceTexture, Temp, FFTInputSceneColor, Intermediates.ComputePassFlags);
 	}
