Fix a bug of CSM caching on mobile.

diff --git a/Engine/Source/Runtime/Renderer/Private/MobileBasePassRendering.cpp b/Engine/Source/Runtime/Renderer/Private/MobileBasePassRendering.cpp
--- a/Engine/Source/Runtime/Renderer/Private/MobileBasePassRendering.cpp
+++ b/Engine/Source/Runtime/Renderer/Private/MobileBasePassRendering.cpp
@@ -268,9 +268,9 @@ void SetupMobileDirectionalLightUniformParameters(
 			const TArray<FProjectedShadowInfo*, SceneRenderingAllocator>& DirectionalLightShadowInfos = VisibleLightInfos[Light->Id].AllProjectedShadows;
 			static_assert(MAX_MOBILE_SHADOWCASCADES <= 4, "more than 4 cascades not supported by the shader and uniform buffer");
 
-			const int32 NumShadowsToCopy = FMath::Min(DirectionalLightShadowInfos.Num(), SceneView.MaxShadowCascades);
+			const int32 NumShadowsToCopy = DirectionalLightShadowInfos.Num();
 			int32_t OutShadowIndex = 0;
-			for (int32 i = 0; i < NumShadowsToCopy; ++i)
+			for (int32 i = 0; i < NumShadowsToCopy && OutShadowIndex < SceneView.MaxShadowCascades; ++i)
 			{
 				const FProjectedShadowInfo* ShadowInfo = DirectionalLightShadowInfos[i + ViewIndex * NumShadowsToCopy];
 
