Split GatherRayTracingWorldInstancesForView into several parts and run most of the work asynchronously

* This saves up to ~2/3 of the GatherRayTracingWorldInstancesForView rendering thread critical path time in large projects
* Main rendering thread cost is now FPrimitiveSceneProxy::GetDynamicRayTracingInstances(), which currently can't be moved to a task

diff --git a/Engine/Source/Runtime/Renderer/Private/DeferredShadingRenderer.cpp b/Engine/Source/Runtime/Renderer/Private/DeferredShadingRenderer.cpp
--- a/Engine/Source/Runtime/Renderer/Private/DeferredShadingRenderer.cpp
+++ b/Engine/Source/Runtime/Renderer/Private/DeferredShadingRenderer.cpp
@@ -227,14 +227,6 @@ static FAutoConsoleVariableRef CVarRayTracingDebugForceOpaque(
 	TEXT("Forces all ray tracing geometry instances to be opaque, effectively disabling any-hit shaders. This is useful for debugging and profiling. (default = 0)")
 );
 
-static int32 GNumLODTasksToInline = 10;
-FAutoConsoleVariableRef CVarNumLODTasksToInline(
-	TEXT("r.RayTracing.GatherWorldInstancingInlineThreshold"),
-	GNumLODTasksToInline,
-	TEXT(""),
-	ECVF_Scalability | ECVF_RenderThreadSafe);
-
-
 #if !UE_BUILD_SHIPPING
 static TAutoConsoleVariable<int32> CVarForceBlackVelocityBuffer(
 	TEXT("r.Test.ForceBlackVelocityBuffer"), 0,
@@ -576,128 +568,87 @@ static void AddDebugRayTracingInstanceFlags(ERayTracingInstanceFlags& InOutFlags
 	}
 }
 
-bool FDeferredShadingSceneRenderer::GatherRayTracingWorldInstancesForView(FRDGBuilder& GraphBuilder, FViewInfo& View, FRayTracingScene& RayTracingScene)
+struct FRayTracingRelevantPrimitive
 {
-	if (!IsRayTracingEnabled())
-	{
-		return false;
-	}
-
-	if (!bAnyRayTracingPassEnabled)
-	{
-		return false;
-	}
-
-	TRACE_CPUPROFILER_EVENT_SCOPE(FDeferredShadingSceneRenderer::GatherRayTracingWorldInstances);
-	SCOPE_CYCLE_COUNTER(STAT_GatherRayTracingWorldInstances);
-
-	RayTracingCollector.ClearViewMeshArrays();
-
-	FGPUScenePrimitiveCollector DummyDynamicPrimitiveCollector;
-
-	RayTracingCollector.AddViewMeshArrays(
-		&View,
-		&View.RayTracedDynamicMeshElements,
-		&View.SimpleElementCollector,
-		&DummyDynamicPrimitiveCollector,
-		ViewFamily.GetFeatureLevel(),
-		&DynamicIndexBufferForInitViews,
-		&DynamicVertexBufferForInitViews,
-		&DynamicReadBufferForInitViews
-		);
-
-	View.DynamicRayTracingMeshCommandStorage.Reserve(Scene->Primitives.Num());
-	View.VisibleRayTracingMeshCommands.Reserve(Scene->Primitives.Num());
-
-	extern TSet<IPersistentViewUniformBufferExtension*> PersistentViewUniformBufferExtensions;
-
-	for (IPersistentViewUniformBufferExtension* Extension : PersistentViewUniformBufferExtensions)
-	{
-		Extension->BeginRenderView(&View);
-	}
-
-	View.RayTracingMeshResourceCollector = MakeUnique<FRayTracingMeshResourceCollector>(
-		Scene->GetFeatureLevel(),
-		&DynamicIndexBufferForInitViews,
-		&DynamicVertexBufferForInitViews,
-		&DynamicReadBufferForInitViews);
+	FRHIRayTracingGeometry* RayTracingGeometryRHI = nullptr;
+	TArrayView<const int32> CachedRayTracingMeshCommandIndices; // Pointer to FPrimitiveSceneInfo::CachedRayTracingMeshCommandIndicesPerLOD data
+	uint64 StateHash = 0;
+	int32 PrimitiveIndex = -1;
+	int8 LODIndex = -1;
+	uint8 InstanceMask = 0;
+	bool bStatic = false;
+	bool bAllSegmentsOpaque = true;
+	bool bAnySegmentsCastShadow = false;
+	bool bAnySegmentsDecal = false;
+	bool bTwoSided = false;
+	bool bIsSky = false;
+	bool bAllSegmentsTranslucent = true;
+
+	uint64 InstancingKey() const
+	{
+		uint64 Key = StateHash;
+		Key ^= uint64(InstanceMask) << 32;
+		Key ^= bAllSegmentsOpaque ? 0x1ull << 40 : 0x0;
+		Key ^= bAnySegmentsCastShadow ? 0x1ull << 41 : 0x0;
+		Key ^= bAnySegmentsDecal ? 0x1ull << 42 : 0x0;
+		Key ^= bTwoSided ? 0x1ull << 43 : 0x0;
+		Key ^= bIsSky ? 0x1ull << 44 : 0x0;
+		Key ^= bAllSegmentsTranslucent ? 0x1ull << 45 : 0x0;
+		return Key ^ reinterpret_cast<uint64>(RayTracingGeometryRHI);
+	}
+};
+
+struct FRayTracingRelevantPrimitiveList
+{
+	// Filtered list of relevant primitives
+	TArray<FRayTracingRelevantPrimitive> RelevantPrimitives;
 
-	View.RayTracingCullingParameters.Init(View);
+	// Array of primitives that should update their cached ray tracing instances via FPrimitiveSceneInfo::UpdateCachedRaytracingData()
+	TArray<FPrimitiveSceneInfo*> DirtyCachedRayTracingPrimitives;
 
-	FRayTracingMaterialGatheringContext MaterialGatheringContext
-	{
-		Scene,
-		&View,
-		ViewFamily,
-		GraphBuilder,
-		*View.RayTracingMeshResourceCollector
-	};
+	// Used coarse mesh streaming handles during the last TLAS build
+	TArray<Nanite::CoarseMeshStreamingHandle> UsedCoarseMeshStreamingHandles;
 
-	const float CurrentWorldTime = View.Family->Time.GetWorldTimeSeconds();
+	// Indicates that this object has been fully produced (for validation)
+	bool bValid = false;
+};
 
-	struct FRelevantPrimitive
-	{
-		FRHIRayTracingGeometry* RayTracingGeometryRHI = nullptr;
-		TArrayView<const int32> CachedRayTracingMeshCommandIndices;
-		uint64 StateHash = 0;
-		int32 PrimitiveIndex = -1;
-		int8 LODIndex = -1;
-		uint8 InstanceMask = 0;
-		bool bStatic = false;
-		bool bAllSegmentsOpaque = true;
-		bool bAnySegmentsCastShadow = false;
-		bool bAnySegmentsDecal = false;
-		bool bTwoSided = false;
-		bool bIsSky = false;
-		bool bAllSegmentsTranslucent = true;
-
-		uint64 InstancingKey() const
-		{
-			uint64 Key = StateHash;
-			Key ^= uint64(InstanceMask) << 32;
-			Key ^= bAllSegmentsOpaque ? 0x1ull << 40 : 0x0;
-			Key ^= bAnySegmentsCastShadow ? 0x1ull << 41 : 0x0;
-			Key ^= bAnySegmentsDecal ? 0x1ull << 42 : 0x0;
-			Key ^= bTwoSided ? 0x1ull << 43 : 0x0;
-			Key ^= bIsSky ? 0x1ull << 44 : 0x0;
-			Key ^= bAllSegmentsTranslucent ? 0x1ull << 45 : 0x0;
-			return Key ^ reinterpret_cast<uint64>(RayTracingGeometryRHI);
-		}
-	};
+// Iterates over Scene's PrimitiveSceneProxies and extracts ones that are relevant for ray tracing. 
+// This function can run on any thread.
+static FRayTracingRelevantPrimitiveList GatherRayTracingRelevantPrimitives(const FScene& Scene, const FViewInfo& View)
+{
+	FRayTracingRelevantPrimitiveList Result;
 
 	// Unified array is used for static and dynamic primitives because we don't know ahead of time how many we'll have of each.
-	TArray<FRelevantPrimitive> RelevantPrimitives;
-	RelevantPrimitives.Reserve(Scene->PrimitiveSceneProxies.Num());
+	Result.RelevantPrimitives.Reserve(Scene.PrimitiveSceneProxies.Num());
 
-	TArray<FPrimitiveSceneInfo*> DirtyCachedRayTracingPrimitives;
-	DirtyCachedRayTracingPrimitives.Reserve(Scene->PrimitiveSceneProxies.Num());
+	Result.DirtyCachedRayTracingPrimitives.Reserve(Scene.PrimitiveSceneProxies.Num());
 
 	const bool bGameView = View.bIsGameView || View.Family->EngineShowFlags.Game;
 
-	int32 VisiblePrimitives = 0;
 	bool bPerformRayTracing = View.State != nullptr && !View.bIsReflectionCapture && View.bAllowRayTracing;
 	if (bPerformRayTracing)
 	{
-		TRACE_CPUPROFILER_EVENT_SCOPE(GatherRayTracingWorldInstances_RelevantPrimitives);
+		TRACE_CPUPROFILER_EVENT_SCOPE(GatherRayTracingRelevantPrimitives);
 
 		// Index into the TypeOffsetTable, which contains a prefix sum of primitive indices by proxy type
 		int32 BroadIndex = 0;
 
-		for (int PrimitiveIndex = 0; PrimitiveIndex < Scene->PrimitiveSceneProxies.Num(); PrimitiveIndex++)
+		for (int PrimitiveIndex = 0; PrimitiveIndex < Scene.PrimitiveSceneProxies.Num(); PrimitiveIndex++)
 		{
 			// Find the next TypeOffsetTable entry that's relevant to this primitive idnex.
-			while (PrimitiveIndex >= int(Scene->TypeOffsetTable[BroadIndex].Offset))
+			while (PrimitiveIndex >= int(Scene.TypeOffsetTable[BroadIndex].Offset))
 			{
 				BroadIndex++;
 			}
 
 			// Skip before dereferencing SceneInfo
-			if (Scene->PrimitiveRayTracingFlags[PrimitiveIndex] == ERayTracingPrimitiveFlags::UnsupportedProxyType)
+			if (Scene.PrimitiveRayTracingFlags[PrimitiveIndex] == ERayTracingPrimitiveFlags::UnsupportedProxyType)
 			{
 				// Find the index of a proxy of the next type, skipping over a batch of proxies that are the same type as current.
 				// This assumes that FPrimitiveSceneProxy::IsRayTracingRelevant() is consistent for all proxies of the same type.
 				// I.e. does not depend on members of the particular FPrimitiveSceneProxy implementation.
-				PrimitiveIndex = Scene->TypeOffsetTable[BroadIndex].Offset - 1;
+				PrimitiveIndex = Scene.TypeOffsetTable[BroadIndex].Offset - 1;
 				continue;
 			}
 
@@ -707,7 +658,7 @@ bool FDeferredShadingSceneRenderer::GatherRayTracingWorldInstancesForView(FRDGBu
 				continue;
 			}
 
-			const FPrimitiveSceneInfo* SceneInfo = Scene->Primitives[PrimitiveIndex];
+			const FPrimitiveSceneInfo* SceneInfo = Scene.Primitives[PrimitiveIndex];
 
 			// #dxr_todo: ray tracing in scene captures should re-use the persistent RT scene. (UE-112448)
 			bool bShouldRayTraceSceneCapture = GRayTracingSceneCaptures > 0
@@ -731,48 +682,40 @@ bool FDeferredShadingSceneRenderer::GatherRayTracingWorldInstancesForView(FRDGBu
 			}
 
 			// Marked visible and used after point, check if streaming then mark as used in the TLAS (so it can be streamed in)
-			if (EnumHasAnyFlags(Scene->PrimitiveRayTracingFlags[PrimitiveIndex], ERayTracingPrimitiveFlags::Streaming))
+			if (EnumHasAnyFlags(Scene.PrimitiveRayTracingFlags[PrimitiveIndex], ERayTracingPrimitiveFlags::Streaming))
 			{
 				// Is the cached data dirty?
 				if (SceneInfo->bCachedRaytracingDataDirty)
 				{
-					DirtyCachedRayTracingPrimitives.Add(Scene->Primitives[PrimitiveIndex]);
+					Result.DirtyCachedRayTracingPrimitives.Add(Scene.Primitives[PrimitiveIndex]);
 				}
 
 				check(SceneInfo->CoarseMeshStreamingHandle != INDEX_NONE);
-				RayTracingScene.UsedCoarseMeshStreamingHandles.Add(SceneInfo->CoarseMeshStreamingHandle);
+				Result.UsedCoarseMeshStreamingHandles.Add(SceneInfo->CoarseMeshStreamingHandle);
 			}
 
-			VisiblePrimitives++;
-
 			//#dxr_todo UE-68621  The Raytracing code path does not support ShowFlags since data moved to the SceneInfo. 
 			//Touching the SceneProxy to determine this would simply cost too much
 			static const auto RayTracingStaticMeshesCVar = IConsoleManager::Get().FindTConsoleVariableDataInt(TEXT("r.RayTracing.Geometry.StaticMeshes"));
 
-			FRelevantPrimitive Item;
+			FRayTracingRelevantPrimitive Item;
 			Item.PrimitiveIndex = PrimitiveIndex;
 
-			if (EnumHasAnyFlags(Scene->PrimitiveRayTracingFlags[PrimitiveIndex], ERayTracingPrimitiveFlags::StaticMesh)
+			if (EnumHasAnyFlags(Scene.PrimitiveRayTracingFlags[PrimitiveIndex], ERayTracingPrimitiveFlags::StaticMesh)
 				&& View.Family->EngineShowFlags.StaticMeshes
 				&& RayTracingStaticMeshesCVar && RayTracingStaticMeshesCVar->GetValueOnRenderThread() > 0)
 			{
 				Item.bStatic = true;
-				RelevantPrimitives.Add(Item);
+				Result.RelevantPrimitives.Add(Item);
 			}
 			else if (View.Family->EngineShowFlags.SkeletalMeshes)
 			{
 				Item.bStatic = false;
-				RelevantPrimitives.Add(Item);
+				Result.RelevantPrimitives.Add(Item);
 			}
 		}
 	}
 
-	INC_DWORD_STAT_BY(STAT_VisibleRayTracingPrimitives, VisiblePrimitives);
-
-	FPrimitiveSceneInfo::UpdateCachedRaytracingData(Scene, DirtyCachedRayTracingPrimitives);
-
-	FGraphEventArray LODTaskList;
-
 	{
 		TRACE_CPUPROFILER_EVENT_SCOPE(GatherRayTracingWorldInstances_ComputeLOD);
 
@@ -780,130 +723,157 @@ bool FDeferredShadingSceneRenderer::GatherRayTracingWorldInstancesForView(FRDGBu
 		const float LODScaleCVarValue = ICVarStaticMeshLODDistanceScale->GetFloat();
 		const int32 ForcedLODLevel = GetCVarForceLOD();
 
-		const uint32 NumTotalItems = RelevantPrimitives.Num();
-		const uint32 TargetItemsPerTask = 1024; // Granularity based on profiling Infiltrator scene
-		const uint32 NumTasks = FMath::Max(1u, FMath::DivideAndRoundUp(NumTotalItems, TargetItemsPerTask));
-		const uint32 ItemsPerTask = FMath::DivideAndRoundUp(NumTotalItems, NumTasks); // Evenly divide commands between tasks (avoiding potential short last task)
-
-		auto ComputeLOD =
-			[	&View,
-				Scene = this->Scene,
-				LODScaleCVarValue,
-				ForcedLODLevel
-			](FRelevantPrimitive* Items, uint32 NumItems)
+		for (FRayTracingRelevantPrimitive& RelevantPrimitive : Result.RelevantPrimitives)
+		{
+			if (!RelevantPrimitive.bStatic)
 			{
-				TRACE_CPUPROFILER_EVENT_SCOPE(GatherRayTracingWorldInstances_ComputeLOD_Task);
+				continue; // skip dynamic primitives
+			}
 
-				for (uint32 i = 0; i < NumItems; ++i)
-				{
-					FRelevantPrimitive& RelevantPrimitive = Items[i];
-					if (!RelevantPrimitive.bStatic)
-					{
-						continue; // skip dynamic primitives
-					}
+			const int32 PrimitiveIndex = RelevantPrimitive.PrimitiveIndex;
+			const FPrimitiveSceneInfo* SceneInfo = Scene.Primitives[PrimitiveIndex];
 
-					const int32 PrimitiveIndex = RelevantPrimitive.PrimitiveIndex;
-					const FPrimitiveSceneInfo* SceneInfo = Scene->Primitives[PrimitiveIndex];
+			int8 LODIndex = 0;
 
-					int8 LODIndex = 0;
+			if (EnumHasAnyFlags(Scene.PrimitiveRayTracingFlags[PrimitiveIndex], ERayTracingPrimitiveFlags::ComputeLOD))
+			{
+				const FPrimitiveBounds& Bounds = Scene.PrimitiveBounds[PrimitiveIndex];
+				const FPrimitiveSceneInfo* RESTRICT PrimitiveSceneInfo = Scene.Primitives[PrimitiveIndex];
 
-					if (EnumHasAnyFlags(Scene->PrimitiveRayTracingFlags[PrimitiveIndex], ERayTracingPrimitiveFlags::ComputeLOD))
-					{
-						const FPrimitiveBounds& Bounds = Scene->PrimitiveBounds[PrimitiveIndex];
-						const FPrimitiveSceneInfo* RESTRICT PrimitiveSceneInfo = Scene->Primitives[PrimitiveIndex];
+				FLODMask LODToRender;
 
-						FLODMask LODToRender;
+				const int8 CurFirstLODIdx = PrimitiveSceneInfo->Proxy->GetCurrentFirstLODIdx_RenderThread();
+				check(CurFirstLODIdx >= 0);
 
-						const int8 CurFirstLODIdx = PrimitiveSceneInfo->Proxy->GetCurrentFirstLODIdx_RenderThread();
-						check(CurFirstLODIdx >= 0);
+				float MeshScreenSizeSquared = 0;
+				float LODScale = LODScaleCVarValue * View.LODDistanceFactor;
+				LODToRender = ComputeLODForMeshes(SceneInfo->StaticMeshRelevances, View, Bounds.BoxSphereBounds.Origin, Bounds.BoxSphereBounds.SphereRadius, ForcedLODLevel, MeshScreenSizeSquared, CurFirstLODIdx, LODScale, true);
 
-						float MeshScreenSizeSquared = 0;
-						float LODScale = LODScaleCVarValue * View.LODDistanceFactor;
-						LODToRender = ComputeLODForMeshes(SceneInfo->StaticMeshRelevances, View, Bounds.BoxSphereBounds.Origin, Bounds.BoxSphereBounds.SphereRadius, ForcedLODLevel, MeshScreenSizeSquared, CurFirstLODIdx, LODScale, true);
+				LODIndex = LODToRender.GetRayTracedLOD();
+			}
 
-						LODIndex = LODToRender.GetRayTracedLOD();
-					}
+			if (!EnumHasAllFlags(Scene.PrimitiveRayTracingFlags[PrimitiveIndex], ERayTracingPrimitiveFlags::CacheInstances))
+			{
+				FRHIRayTracingGeometry* RayTracingGeometryInstance = SceneInfo->GetStaticRayTracingGeometryInstance(LODIndex);
+				if (RayTracingGeometryInstance == nullptr)
+				{
+					continue;
+				}
+
+				// Sometimes LODIndex is out of range because it is clamped by ClampToFirstLOD, like the requested LOD is being streamed in and hasn't been available
+				// According to InitViews, we should hide the static mesh instance
+				if (SceneInfo->CachedRayTracingMeshCommandIndicesPerLOD.IsValidIndex(LODIndex))
+				{
+					RelevantPrimitive.LODIndex = LODIndex;
+					RelevantPrimitive.RayTracingGeometryRHI = SceneInfo->GetStaticRayTracingGeometryInstance(LODIndex);
+
+					RelevantPrimitive.CachedRayTracingMeshCommandIndices = SceneInfo->CachedRayTracingMeshCommandIndicesPerLOD[LODIndex];
+					RelevantPrimitive.StateHash = SceneInfo->CachedRayTracingMeshCommandsHashPerLOD[LODIndex];
 
-					if (!EnumHasAllFlags(Scene->PrimitiveRayTracingFlags[PrimitiveIndex], ERayTracingPrimitiveFlags::CacheInstances))
+					for (int32 CommandIndex : RelevantPrimitive.CachedRayTracingMeshCommandIndices)
 					{
-						FRHIRayTracingGeometry* RayTracingGeometryInstance = SceneInfo->GetStaticRayTracingGeometryInstance(LODIndex);
-						if (RayTracingGeometryInstance == nullptr)
+						if (CommandIndex >= 0)
 						{
-							continue;
+							const FRayTracingMeshCommand& RayTracingMeshCommand = Scene.CachedRayTracingMeshCommands[CommandIndex];
+
+							RelevantPrimitive.InstanceMask |= RayTracingMeshCommand.InstanceMask;
+							RelevantPrimitive.bAllSegmentsOpaque &= RayTracingMeshCommand.bOpaque;
+							RelevantPrimitive.bAnySegmentsCastShadow |= RayTracingMeshCommand.bCastRayTracedShadows;
+							RelevantPrimitive.bAnySegmentsDecal |= RayTracingMeshCommand.bDecal;
+							RelevantPrimitive.bTwoSided |= RayTracingMeshCommand.bTwoSided;
+							RelevantPrimitive.bIsSky |= RayTracingMeshCommand.bIsSky;
+							RelevantPrimitive.bAllSegmentsTranslucent &= RayTracingMeshCommand.bIsTranslucent;
 						}
-
-						// Sometimes LODIndex is out of range because it is clamped by ClampToFirstLOD, like the requested LOD is being streamed in and hasn't been available
-						// According to InitViews, we should hide the static mesh instance
-						if (SceneInfo->CachedRayTracingMeshCommandIndicesPerLOD.IsValidIndex(LODIndex))
+						else
 						{
-							RelevantPrimitive.LODIndex = LODIndex;
-							RelevantPrimitive.RayTracingGeometryRHI = SceneInfo->GetStaticRayTracingGeometryInstance(LODIndex);
-
-							RelevantPrimitive.CachedRayTracingMeshCommandIndices = SceneInfo->CachedRayTracingMeshCommandIndicesPerLOD[LODIndex];
-							RelevantPrimitive.StateHash = SceneInfo->CachedRayTracingMeshCommandsHashPerLOD[LODIndex];
-
-							for (int32 CommandIndex : RelevantPrimitive.CachedRayTracingMeshCommandIndices)
-							{
-								if (CommandIndex >= 0)
-								{
-									const FRayTracingMeshCommand& RayTracingMeshCommand = Scene->CachedRayTracingMeshCommands[CommandIndex];
-
-									RelevantPrimitive.InstanceMask |= RayTracingMeshCommand.InstanceMask;
-									RelevantPrimitive.bAllSegmentsOpaque &= RayTracingMeshCommand.bOpaque;
-									RelevantPrimitive.bAnySegmentsCastShadow |= RayTracingMeshCommand.bCastRayTracedShadows;
-									RelevantPrimitive.bAnySegmentsDecal |= RayTracingMeshCommand.bDecal;
-									RelevantPrimitive.bTwoSided |= RayTracingMeshCommand.bTwoSided;
-									RelevantPrimitive.bIsSky |= RayTracingMeshCommand.bIsSky;
-									RelevantPrimitive.bAllSegmentsTranslucent &= RayTracingMeshCommand.bIsTranslucent;
-								}
-								else
-								{
-									// CommandIndex == -1 indicates that the mesh batch has been filtered by FRayTracingMeshProcessor (like the shadow depth pass batch)
-									// Do nothing in this case
-								}
-							}
-
-							RelevantPrimitive.InstanceMask |= RelevantPrimitive.bAnySegmentsCastShadow ? RAY_TRACING_MASK_SHADOW : 0;
-
-							if (EnumHasAllFlags(Scene->PrimitiveRayTracingFlags[PrimitiveIndex], ERayTracingPrimitiveFlags::FarField))
-							{
-								RelevantPrimitive.InstanceMask = RAY_TRACING_MASK_FAR_FIELD;
-							}
+							// CommandIndex == -1 indicates that the mesh batch has been filtered by FRayTracingMeshProcessor (like the shadow depth pass batch)
+							// Do nothing in this case
 						}
 					}
-				}
-			};
 
-		if (NumTasks > (uint32)GNumLODTasksToInline)
-		{
-			SCOPED_NAMED_EVENT(DispatchParallelComputeLOD, FColor::Red);
-			
-			LODTaskList.Reserve(NumTasks);
+					RelevantPrimitive.InstanceMask |= RelevantPrimitive.bAnySegmentsCastShadow ? RAY_TRACING_MASK_SHADOW : 0;
 
-			for (uint32 TaskIndex = 0; TaskIndex < NumTasks; ++TaskIndex)
-			{
-				const uint32 FirstTaskItemIndex = TaskIndex * ItemsPerTask;
-
-				LODTaskList.Add(FFunctionGraphTask::CreateAndDispatchWhenReady(
-					[	FirstTaskItemIndex,
-						Items = RelevantPrimitives.GetData() + FirstTaskItemIndex,
-						NumItems = FMath::Min(ItemsPerTask, NumTotalItems - FirstTaskItemIndex),
-						ComputeLOD
-					]()
+					if (EnumHasAllFlags(Scene.PrimitiveRayTracingFlags[PrimitiveIndex], ERayTracingPrimitiveFlags::FarField))
 					{
-						ComputeLOD(Items, NumItems);
-					},
-					TStatId(), nullptr, ENamedThreads::AnyNormalThreadHiPriTask));
+						RelevantPrimitive.InstanceMask = RAY_TRACING_MASK_FAR_FIELD;
+					}
+				}
 			}
 		}
-		else
-		{
-			SCOPED_NAMED_EVENT(ComputeLOD, FColor::Magenta);
-			
-			ComputeLOD(RelevantPrimitives.GetData(), NumTotalItems);
-		}
 	}
 
+	Result.bValid = true;
+
+	return Result;
+}
+
+bool FDeferredShadingSceneRenderer::GatherRayTracingWorldInstancesForView(FRDGBuilder& GraphBuilder, FViewInfo& View, FRayTracingScene& RayTracingScene, struct FRayTracingRelevantPrimitiveList&& RelevantPrimitiveList)
+{
+	checkf(IsRayTracingEnabled() && bAnyRayTracingPassEnabled, TEXT("GatherRayTracingWorldInstancesForView should only be called if ray tracing is used"))
+
+	TRACE_CPUPROFILER_EVENT_SCOPE(FDeferredShadingSceneRenderer::GatherRayTracingWorldInstances);
+	SCOPE_CYCLE_COUNTER(STAT_GatherRayTracingWorldInstances);
+
+	RayTracingCollector.ClearViewMeshArrays();
+
+	FGPUScenePrimitiveCollector DummyDynamicPrimitiveCollector;
+
+	RayTracingCollector.AddViewMeshArrays(
+		&View,
+		&View.RayTracedDynamicMeshElements,
+		&View.SimpleElementCollector,
+		&DummyDynamicPrimitiveCollector,
+		ViewFamily.GetFeatureLevel(),
+		&DynamicIndexBufferForInitViews,
+		&DynamicVertexBufferForInitViews,
+		&DynamicReadBufferForInitViews
+	);
+
+	View.DynamicRayTracingMeshCommandStorage.Reserve(Scene->Primitives.Num());
+	View.VisibleRayTracingMeshCommands.Reserve(Scene->Primitives.Num());
+
+	extern TSet<IPersistentViewUniformBufferExtension*> PersistentViewUniformBufferExtensions;
+
+	for (IPersistentViewUniformBufferExtension* Extension : PersistentViewUniformBufferExtensions)
+	{
+		Extension->BeginRenderView(&View);
+	}
+
+	View.RayTracingMeshResourceCollector = MakeUnique<FRayTracingMeshResourceCollector>(
+		Scene->GetFeatureLevel(),
+		&DynamicIndexBufferForInitViews,
+		&DynamicVertexBufferForInitViews,
+		&DynamicReadBufferForInitViews);
+
+	View.RayTracingCullingParameters.Init(View);
+
+	FRayTracingMaterialGatheringContext MaterialGatheringContext
+	{
+		Scene,
+		&View,
+		ViewFamily,
+		GraphBuilder,
+		*View.RayTracingMeshResourceCollector
+	};
+
+	const float CurrentWorldTime = View.Family->Time.GetWorldTimeSeconds();
+
+	// Consume output of the relevant primitive gathering task
+	RayTracingScene.UsedCoarseMeshStreamingHandles = MoveTemp(RelevantPrimitiveList.UsedCoarseMeshStreamingHandles);
+	TArray<FRayTracingRelevantPrimitive> RelevantPrimitives = MoveTemp(RelevantPrimitiveList.RelevantPrimitives);
+	TArray<FPrimitiveSceneInfo*> DirtyCachedRayTracingPrimitives = MoveTemp(RelevantPrimitiveList.DirtyCachedRayTracingPrimitives);
+
+	// Inform the coarse mesh streaming manager about all the used streamable render assets in the scene
+	Nanite::FCoarseMeshStreamingManager* CoarseMeshSM = IStreamingManager::Get().GetNaniteCoarseMeshStreamingManager();
+	if (CoarseMeshSM)
+	{
+		CoarseMeshSM->AddUsedStreamingHandles(RayTracingScene.UsedCoarseMeshStreamingHandles);
+	}
+
+	INC_DWORD_STAT_BY(STAT_VisibleRayTracingPrimitives, RelevantPrimitives.Num());
+
+	FPrimitiveSceneInfo::UpdateCachedRaytracingData(Scene, DirtyCachedRayTracingPrimitives);
+
 	{
 		TRACE_CPUPROFILER_EVENT_SCOPE(GatherRayTracingWorldInstances_DynamicElements);
 
@@ -997,7 +967,7 @@ bool FDeferredShadingSceneRenderer::GatherRayTracingWorldInstancesForView(FRDGBu
 		// Local temporary array of instances used for GetDynamicRayTracingInstances()
 		TArray<FRayTracingInstance> TempRayTracingInstances;
 
-		for (const FRelevantPrimitive& RelevantPrimitive : RelevantPrimitives)
+		for (const FRayTracingRelevantPrimitive& RelevantPrimitive : RelevantPrimitives)
 		{
 			if (RelevantPrimitive.bStatic)
 			{
@@ -1173,18 +1143,38 @@ bool FDeferredShadingSceneRenderer::GatherRayTracingWorldInstancesForView(FRDGBu
 		KickRayTracingMeshBatchTask();
 	}
 
-	//
-
+	// Task to iterate over static ray tracing instances, perform auto-instancing and culling.
+	// This adds final instances to the ray tracing scene and must be done before FRayTracingScene::BuildInitializationData().
+	struct FRayTracingSceneAddInstancesTask
 	{
-		TRACE_CPUPROFILER_EVENT_SCOPE(GatherRayTracingWorldInstances_AddInstances);
+		UE_NONCOPYABLE(FRayTracingSceneAddInstancesTask)
+
+		static ESubsequentsMode::Type GetSubsequentsMode() { return ESubsequentsMode::TrackSubsequents; }
+		TStatId                       GetStatId() const    { return TStatId(); }
+		ENamedThreads::Type           GetDesiredThread()   { return ENamedThreads::AnyThread; }
 
-		const bool bAutoInstance = CVarRayTracingAutoInstance.GetValueOnRenderThread() != 0;
+		// Inputs
 
-		if (LODTaskList.Num() > 0)
+		const FScene& Scene;
+		TArray<FRayTracingRelevantPrimitive> RelevantPrimitives;
+		const FRayTracingCullingParameters& CullingParameters;
+		int32 ExpectedMaxVisibleRayTracingMeshCommands = 0;
+
+		// Outputs
+
+		FRayTracingScene& RayTracingScene; // New instances are added into FRayTracingScene::Instances and FRayTracingScene::Allocator is used for temporary data
+		FRayTracingMeshCommandOneFrameArray& VisibleRayTracingMeshCommands; // New elements are added here by this task
+
+		FRayTracingSceneAddInstancesTask(const FScene& InScene, TArray<FRayTracingRelevantPrimitive> InRelevantPrimitives, const FRayTracingCullingParameters& InCullingParameters,
+											FRayTracingScene& InRayTracingScene, FRayTracingMeshCommandOneFrameArray& InVisibleRayTracingMeshCommands)
+			: Scene(InScene)
+			, RelevantPrimitives(MoveTemp(InRelevantPrimitives))
+			, CullingParameters(InCullingParameters)
+			, RayTracingScene(InRayTracingScene)
+			, VisibleRayTracingMeshCommands(InVisibleRayTracingMeshCommands)
 		{
-			SCOPED_NAMED_EVENT(WaitForParallelComputeLOD, FColor::Red);
-			TRACE_CPUPROFILER_EVENT_SCOPE(WaitForLODTasks);
-			FTaskGraphInterface::Get().WaitUntilTasksComplete(LODTaskList, ENamedThreads::GetRenderThread_Local());
+			VisibleRayTracingMeshCommands.Reserve(RelevantPrimitives.Num());
+			ExpectedMaxVisibleRayTracingMeshCommands = VisibleRayTracingMeshCommands.Max();
 		}
 
 		struct FAutoInstanceBatch
@@ -1237,192 +1227,210 @@ bool FDeferredShadingSceneRenderer::GatherRayTracingWorldInstancesForView(FRDGBu
 			uint32 Cursor = 0;
 		};
 
-		Experimental::TSherwoodMap<uint64, FAutoInstanceBatch> InstanceBatches;
+		void DoTask(ENamedThreads::Type CurrentThread, const FGraphEventRef& MyCompletionGraphEvent)
+		{
+			FTaskTagScope TaskTagScope(ETaskTag::EParallelRenderingThread);
 
-		InstanceBatches.Reserve(RelevantPrimitives.Num());
+			TRACE_CPUPROFILER_EVENT_SCOPE(RayTracingSceneStaticInstanceTask);
 
-		TArray<FRayTracingCullPrimitiveInstancesClosure> CullInstancesClosures;
-		if (View.RayTracingCullingParameters.CullInRayTracing > 0 && GetRayTracingCullingPerInstance())
-		{
-			CullInstancesClosures.Reserve(RelevantPrimitives.Num());
-			View.RayTracingPerInstanceCullingTaskList.Reserve(RelevantPrimitives.Num() / 256 + 1);
-		}
+			FGraphEventArray CullingTasks;
 
-		// scan relevant primitives computing hash data to look for duplicate instances
-		for (const FRelevantPrimitive& RelevantPrimitive : RelevantPrimitives)
-		{
-			const int32 PrimitiveIndex = RelevantPrimitive.PrimitiveIndex;
-			FPrimitiveSceneInfo* SceneInfo = Scene->Primitives[PrimitiveIndex];
-			ERayTracingPrimitiveFlags Flags = Scene->PrimitiveRayTracingFlags[PrimitiveIndex];
+			const bool bAutoInstance = CVarRayTracingAutoInstance.GetValueOnRenderThread() != 0;
 
-			if (EnumHasAnyFlags(Flags, ERayTracingPrimitiveFlags::CacheInstances))
+			// Instance batches by FRayTracingRelevantPrimitive::InstancingKey()
+			Experimental::TSherwoodMap<uint64, FAutoInstanceBatch> InstanceBatches;
+
+			TArray<FRayTracingCullPrimitiveInstancesClosure> CullInstancesClosures;
+			if (CullingParameters.CullInRayTracing > 0 && GetRayTracingCullingPerInstance())
 			{
-				// TODO: support GRayTracingExcludeDecals, but not in the form of RayTracingMeshCommand.bDecal as that requires looping over all cached MDCs
-				// Instead, either make r.RayTracing.ExcludeDecals read only or request a recache of all ray tracing commands during which decals are excluded
+				CullInstancesClosures.Reserve(RelevantPrimitives.Num());
+				CullingTasks.Reserve(RelevantPrimitives.Num() / 256 + 1);
+			}
 
-				const int32 NewInstanceIndex = RayTracingScene.Instances.Num();
+			// scan relevant primitives computing hash data to look for duplicate instances
+			for (const FRayTracingRelevantPrimitive& RelevantPrimitive : RelevantPrimitives)
+			{
+				const int32 PrimitiveIndex = RelevantPrimitive.PrimitiveIndex;
+				FPrimitiveSceneInfo* SceneInfo = Scene.Primitives[PrimitiveIndex];
+				ERayTracingPrimitiveFlags Flags = Scene.PrimitiveRayTracingFlags[PrimitiveIndex];
 
-				// At the moment we only support SM & ISMs on this path
-				check(EnumHasAnyFlags(Flags, ERayTracingPrimitiveFlags::CacheMeshCommands));
-				if (SceneInfo->CachedRayTracingMeshCommandIndicesPerLOD.Num() > 0 && SceneInfo->CachedRayTracingMeshCommandIndicesPerLOD[0].Num() > 0)
+				if (EnumHasAnyFlags(Flags, ERayTracingPrimitiveFlags::CacheInstances))
 				{
-					for (int32 CommandIndex : SceneInfo->CachedRayTracingMeshCommandIndicesPerLOD[0])
+					// TODO: support GRayTracingExcludeDecals, but not in the form of RayTracingMeshCommand.bDecal as that requires looping over all cached MDCs
+					// Instead, either make r.RayTracing.ExcludeDecals read only or request a recache of all ray tracing commands during which decals are excluded
+
+					const int32 NewInstanceIndex = RayTracingScene.Instances.Num();
+
+					// At the moment we only support SM & ISMs on this path
+					check(EnumHasAnyFlags(Flags, ERayTracingPrimitiveFlags::CacheMeshCommands));
+					if (SceneInfo->CachedRayTracingMeshCommandIndicesPerLOD.Num() > 0 && SceneInfo->CachedRayTracingMeshCommandIndicesPerLOD[0].Num() > 0)
 					{
-						FVisibleRayTracingMeshCommand NewVisibleMeshCommand;
+						for (int32 CommandIndex : SceneInfo->CachedRayTracingMeshCommandIndicesPerLOD[0])
+						{
+							FVisibleRayTracingMeshCommand NewVisibleMeshCommand;
 
-						NewVisibleMeshCommand.RayTracingMeshCommand = &Scene->CachedRayTracingMeshCommands[CommandIndex];
-						NewVisibleMeshCommand.InstanceIndex = NewInstanceIndex;
-						View.VisibleRayTracingMeshCommands.Add(NewVisibleMeshCommand);
+							NewVisibleMeshCommand.RayTracingMeshCommand = &Scene.CachedRayTracingMeshCommands[CommandIndex];
+							NewVisibleMeshCommand.InstanceIndex = NewInstanceIndex;
+							VisibleRayTracingMeshCommands.Add(NewVisibleMeshCommand);
+						}
 					}
-				}
 
-				checkf(SceneInfo->CachedRayTracingInstance.GeometryRHI, TEXT("Ray tracing instance must have a valid geometry."));
-				RayTracingScene.Instances.Add(SceneInfo->CachedRayTracingInstance);
+					checkf(SceneInfo->CachedRayTracingInstance.GeometryRHI, TEXT("Ray tracing instance must have a valid geometry."));
+					RayTracingScene.Instances.Add(SceneInfo->CachedRayTracingInstance);
 
-				const Experimental::FHashElementId GroupId = Scene->PrimitiveRayTracingGroupIds[PrimitiveIndex];
-				const bool bUseGroupBounds = View.RayTracingCullingParameters.bCullUsingGroupIds && GroupId.IsValid();
+					const Experimental::FHashElementId GroupId = Scene.PrimitiveRayTracingGroupIds[PrimitiveIndex];
+					const bool bUseGroupBounds = CullingParameters.bCullUsingGroupIds && GroupId.IsValid();
 
-				if (View.RayTracingCullingParameters.CullInRayTracing > 0 && GetRayTracingCullingPerInstance() && SceneInfo->CachedRayTracingInstance.NumTransforms > 1 && !bUseGroupBounds)
-				{
-					FRayTracingGeometryInstance& NewInstance = RayTracingScene.Instances.Last();
+					if (CullingParameters.CullInRayTracing > 0 && GetRayTracingCullingPerInstance() && SceneInfo->CachedRayTracingInstance.NumTransforms > 1 && !bUseGroupBounds)
+					{
+						FRayTracingGeometryInstance& NewInstance = RayTracingScene.Instances.Last();
 
-					const bool bIsFarFieldPrimitive = EnumHasAnyFlags(Flags, ERayTracingPrimitiveFlags::FarField);
+						const bool bIsFarFieldPrimitive = EnumHasAnyFlags(Flags, ERayTracingPrimitiveFlags::FarField);
 
-					TArrayView<uint32> InstanceActivationMask = RayTracingScene.Allocate<uint32>(FMath::DivideAndRoundUp(NewInstance.NumTransforms, 32u));
+						TArrayView<uint32> InstanceActivationMask = RayTracingScene.Allocate<uint32>(FMath::DivideAndRoundUp(NewInstance.NumTransforms, 32u));
 
-					NewInstance.ActivationMask = InstanceActivationMask;
+						NewInstance.ActivationMask = InstanceActivationMask;
 
-					FRayTracingCullPrimitiveInstancesClosure Closure;
-					Closure.Scene = Scene;
-					Closure.SceneInfo = SceneInfo;
-					Closure.PrimitiveIndex = PrimitiveIndex;
-					Closure.bIsFarFieldPrimitive = bIsFarFieldPrimitive;
-					Closure.CullingParameters = &View.RayTracingCullingParameters;
-					Closure.OutInstanceActivationMask = InstanceActivationMask;
+						FRayTracingCullPrimitiveInstancesClosure Closure;
+						Closure.Scene = &Scene;
+						Closure.SceneInfo = SceneInfo;
+						Closure.PrimitiveIndex = PrimitiveIndex;
+						Closure.bIsFarFieldPrimitive = bIsFarFieldPrimitive;
+						Closure.CullingParameters = &CullingParameters;
+						Closure.OutInstanceActivationMask = InstanceActivationMask;
 
-					CullInstancesClosures.Add(MoveTemp(Closure));
+						CullInstancesClosures.Add(MoveTemp(Closure));
 
-					if (CullInstancesClosures.Num() >= 256)
-					{
-						View.RayTracingPerInstanceCullingTaskList.Add(FFunctionGraphTask::CreateAndDispatchWhenReady([CullInstancesClosures = MoveTemp(CullInstancesClosures)]()
+						if (CullInstancesClosures.Num() >= 256)
 						{
-							for (auto& Closure : CullInstancesClosures)
+							CullingTasks.Add(FFunctionGraphTask::CreateAndDispatchWhenReady([CullInstancesClosures = MoveTemp(CullInstancesClosures)]()
 							{
-								Closure();
-							}
-						}, TStatId(), nullptr, ENamedThreads::AnyThread));
+								for (auto& Closure : CullInstancesClosures)
+								{
+									Closure();
+								}
+							}, TStatId(), nullptr, ENamedThreads::AnyThread));
+						}
 					}
-				}
 
-				AddDebugRayTracingInstanceFlags(RayTracingScene.Instances.Last().Flags);
-			}
-			else
-			{
-				const int8 LODIndex = RelevantPrimitive.LODIndex;
-
-				if (LODIndex < 0 || !RelevantPrimitive.bStatic)
-				{
-					continue; // skip dynamic primitives and other 
+					AddDebugRayTracingInstanceFlags(RayTracingScene.Instances.Last().Flags);
 				}
-
-				if ((GRayTracingExcludeDecals && RelevantPrimitive.bAnySegmentsDecal)
-					|| (GRayTracingExcludeTranslucent && RelevantPrimitive.bAllSegmentsTranslucent)
-					|| (GRayTracingExcludeSky && RelevantPrimitive.bIsSky))
+				else
 				{
-					continue;
-				}
-
-				// location if this is a new entry
-				const int32 NewInstanceIndex = RayTracingScene.Instances.Num();
-				const uint64 InstanceKey = RelevantPrimitive.InstancingKey();
+					const int8 LODIndex = RelevantPrimitive.LODIndex;
 
-				FAutoInstanceBatch DummyInstanceBatch = { NewInstanceIndex };
-				FAutoInstanceBatch& InstanceBatch = bAutoInstance ? InstanceBatches.FindOrAdd(InstanceKey, DummyInstanceBatch) : DummyInstanceBatch;
+					if (LODIndex < 0 || !RelevantPrimitive.bStatic)
+					{
+						continue; // skip dynamic primitives and other 
+					}
 
-				if (InstanceBatch.Index != NewInstanceIndex)
-				{
-					// Reusing a previous entry, just append to the instance list.
+					if ((GRayTracingExcludeDecals && RelevantPrimitive.bAnySegmentsDecal)
+						|| (GRayTracingExcludeTranslucent && RelevantPrimitive.bAllSegmentsTranslucent)
+						|| (GRayTracingExcludeSky && RelevantPrimitive.bIsSky))
+					{
+						continue;
+					}
 
-					FRayTracingGeometryInstance& RayTracingInstance = RayTracingScene.Instances[InstanceBatch.Index];
-					bool bReallocated = InstanceBatch.Add(RayTracingScene, SceneInfo->GetInstanceSceneDataOffset(), (uint32)PrimitiveIndex);
+					// location if this is a new entry
+					const int32 NewInstanceIndex = RayTracingScene.Instances.Num();
+					const uint64 InstanceKey = RelevantPrimitive.InstancingKey();
 
-					++RayTracingInstance.NumTransforms;
-					check(RayTracingInstance.NumTransforms == InstanceBatch.Cursor); // sanity check
+					FAutoInstanceBatch DummyInstanceBatch = { NewInstanceIndex };
+					FAutoInstanceBatch& InstanceBatch = bAutoInstance ? InstanceBatches.FindOrAdd(InstanceKey, DummyInstanceBatch) : DummyInstanceBatch;
 
-					if (bReallocated)
+					if (InstanceBatch.Index != NewInstanceIndex)
 					{
-						RayTracingInstance.InstanceSceneDataOffsets = InstanceBatch.InstanceSceneDataOffsets;
-						RayTracingInstance.UserData = InstanceBatch.UserData;
-					}
-				}
-				else
-				{
-					// Starting new instance batch
+						// Reusing a previous entry, just append to the instance list.
 
-					for (int32 CommandIndex : RelevantPrimitive.CachedRayTracingMeshCommandIndices)
-					{
-						if (CommandIndex >= 0)
-						{
-							FVisibleRayTracingMeshCommand NewVisibleMeshCommand;
+						FRayTracingGeometryInstance& RayTracingInstance = RayTracingScene.Instances[InstanceBatch.Index];
+						bool bReallocated = InstanceBatch.Add(RayTracingScene, SceneInfo->GetInstanceSceneDataOffset(), (uint32)PrimitiveIndex);
 
-							NewVisibleMeshCommand.RayTracingMeshCommand = &Scene->CachedRayTracingMeshCommands[CommandIndex];
-							NewVisibleMeshCommand.InstanceIndex = NewInstanceIndex;
-							View.VisibleRayTracingMeshCommands.Add(NewVisibleMeshCommand);
-						}
-						else
+						++RayTracingInstance.NumTransforms;
+						check(RayTracingInstance.NumTransforms == InstanceBatch.Cursor); // sanity check
+
+						if (bReallocated)
 						{
-							// CommandIndex == -1 indicates that the mesh batch has been filtered by FRayTracingMeshProcessor (like the shadow depth pass batch)
-							// Do nothing in this case
+							RayTracingInstance.InstanceSceneDataOffsets = InstanceBatch.InstanceSceneDataOffsets;
+							RayTracingInstance.UserData = InstanceBatch.UserData;
 						}
 					}
+					else
+					{
+						// Starting new instance batch
 
-					FRayTracingGeometryInstance& RayTracingInstance = RayTracingScene.Instances.AddDefaulted_GetRef();
+						for (int32 CommandIndex : RelevantPrimitive.CachedRayTracingMeshCommandIndices)
+						{
+							if (CommandIndex >= 0)
+							{
+								FVisibleRayTracingMeshCommand NewVisibleMeshCommand;
 
-					RayTracingInstance.GeometryRHI = RelevantPrimitive.RayTracingGeometryRHI;
-					checkf(RayTracingInstance.GeometryRHI, TEXT("Ray tracing instance must have a valid geometry."));
+								NewVisibleMeshCommand.RayTracingMeshCommand = &Scene.CachedRayTracingMeshCommands[CommandIndex];
+								NewVisibleMeshCommand.InstanceIndex = NewInstanceIndex;
+								VisibleRayTracingMeshCommands.Add(NewVisibleMeshCommand);
+							}
+							else
+							{
+								// CommandIndex == -1 indicates that the mesh batch has been filtered by FRayTracingMeshProcessor (like the shadow depth pass batch)
+								// Do nothing in this case
+							}
+						}
 
-					InstanceBatch.Add(RayTracingScene, SceneInfo->GetInstanceSceneDataOffset(), (uint32)PrimitiveIndex);
-					RayTracingInstance.InstanceSceneDataOffsets = InstanceBatch.InstanceSceneDataOffsets;
-					RayTracingInstance.UserData = InstanceBatch.UserData;
-					RayTracingInstance.NumTransforms = 1;
+						FRayTracingGeometryInstance& RayTracingInstance = RayTracingScene.Instances.AddDefaulted_GetRef();
 
-					RayTracingInstance.Mask = RelevantPrimitive.InstanceMask; // When no cached command is found, InstanceMask == 0 and the instance is effectively filtered out
+						RayTracingInstance.GeometryRHI = RelevantPrimitive.RayTracingGeometryRHI;
+						checkf(RayTracingInstance.GeometryRHI, TEXT("Ray tracing instance must have a valid geometry."));
 
-					if (RelevantPrimitive.bAllSegmentsOpaque)
-					{
-						RayTracingInstance.Flags |= ERayTracingInstanceFlags::ForceOpaque;
-					}
-					if (RelevantPrimitive.bTwoSided)
-					{
-						RayTracingInstance.Flags |= ERayTracingInstanceFlags::TriangleCullDisable;
+						InstanceBatch.Add(RayTracingScene, SceneInfo->GetInstanceSceneDataOffset(), (uint32)PrimitiveIndex);
+						RayTracingInstance.InstanceSceneDataOffsets = InstanceBatch.InstanceSceneDataOffsets;
+						RayTracingInstance.UserData = InstanceBatch.UserData;
+						RayTracingInstance.NumTransforms = 1;
+
+						RayTracingInstance.Mask = RelevantPrimitive.InstanceMask; // When no cached command is found, InstanceMask == 0 and the instance is effectively filtered out
+
+						if (RelevantPrimitive.bAllSegmentsOpaque)
+						{
+							RayTracingInstance.Flags |= ERayTracingInstanceFlags::ForceOpaque;
+						}
+						if (RelevantPrimitive.bTwoSided)
+						{
+							RayTracingInstance.Flags |= ERayTracingInstanceFlags::TriangleCullDisable;
+						}
+						AddDebugRayTracingInstanceFlags(RayTracingInstance.Flags);
 					}
-					AddDebugRayTracingInstanceFlags(RayTracingInstance.Flags);
 				}
 			}
-		}
 
-		View.RayTracingPerInstanceCullingTaskList.Add(FFunctionGraphTask::CreateAndDispatchWhenReady([CullInstancesClosures = MoveTemp(CullInstancesClosures)]()
-		{
-			for (auto& Closure : CullInstancesClosures)
+			CullingTasks.Add(FFunctionGraphTask::CreateAndDispatchWhenReady([CullInstancesClosures = MoveTemp(CullInstancesClosures)]()
 			{
-				Closure();
+				for (auto& Closure : CullInstancesClosures)
+				{
+					Closure();
+				}
+			}, TStatId(), nullptr, ENamedThreads::AnyThread));
+
+			for (FGraphEventRef& CullingTask : CullingTasks)
+			{
+				MyCompletionGraphEvent->DontCompleteUntil(CullingTask);
 			}
-		}, TStatId(), nullptr, ENamedThreads::AnyThread));
-	}
 
-	// Inform the coarse mesh streaming manager about all the used streamable render assets in the scene
-	Nanite::FCoarseMeshStreamingManager* CoarseMeshSM = IStreamingManager::Get().GetNaniteCoarseMeshStreamingManager();
-	if (CoarseMeshSM)
-	{
-		CoarseMeshSM->AddUsedStreamingHandles(RayTracingScene.UsedCoarseMeshStreamingHandles);
-	}
 
+			checkf(VisibleRayTracingMeshCommands.Max() == ExpectedMaxVisibleRayTracingMeshCommands, TEXT("VisibleRayTracingMeshCommands should not be reallocated during FRayTracingSceneAddInstancesTask"));
+		}
+	};
+
+	FGraphEventRef AddInstancesTask = TGraphTask<FRayTracingSceneAddInstancesTask>::CreateTask().ConstructAndDispatchWhenReady(
+		*Scene, MoveTemp(RelevantPrimitives), View.RayTracingCullingParameters, // inputs 
+		RayTracingScene, View.VisibleRayTracingMeshCommands // outputs
+	);
+
+	// Scene init task can run only when all pre-init tasks are complete (including culling tasks that are spawned while adding instances)
 	View.RayTracingSceneInitTask = FFunctionGraphTask::CreateAndDispatchWhenReady(
 		[&View, &RayTracingScene]()
 		{
+			TRACE_CPUPROFILER_EVENT_SCOPE(RayTracingSceneInitTask);
 			View.RayTracingSceneInitData = RayTracingScene.BuildInitializationData();
 		},
-		TStatId(), &View.RayTracingPerInstanceCullingTaskList, ENamedThreads::AnyThread);
+		TStatId(), AddInstancesTask, ENamedThreads::AnyThread);
 
 	return true;
 }
@@ -1693,7 +1701,6 @@ bool FDeferredShadingSceneRenderer::DispatchRayTracingWorldUpdates(FRDGBuilder&
 	FTaskGraphInterface::Get().WaitUntilTaskCompletes(ReferenceView.RayTracingSceneInitTask, ENamedThreads::GetRenderThread_Local());
 
 	ReferenceView.RayTracingSceneInitTask = {};
-	ReferenceView.RayTracingPerInstanceCullingTaskList.Empty();
 
 	RDG_GPU_MASK_SCOPE(GraphBuilder, FRHIGPUMask::All());
 
@@ -2174,8 +2181,7 @@ void FDeferredShadingSceneRenderer::Render(FRDGBuilder& GraphBuilder)
 
 #if RHI_RAYTRACING
 
-	// Gather mesh instances, shaders, resources, parameters, etc. and build ray tracing acceleration structure
-
+	// Prepare the scene for rendering this frame.
 	FRayTracingScene& RayTracingScene = Scene->RayTracingScene;
 	RayTracingScene.Reset(); // Resets the internal arrays, but does not release any resources.
 
@@ -2194,9 +2200,18 @@ void FDeferredShadingSceneRenderer::Render(FRDGBuilder& GraphBuilder)
 		ReferenceView.bHasRayTracingDecals = false;
 	}
 
-	// Prepare the scene for rendering this frame.
-	GatherRayTracingWorldInstancesForView(GraphBuilder, ReferenceView, RayTracingScene);
-
+	// Asynchronously create a list of primitives relevant to ray tracing scene
+	FRayTracingRelevantPrimitiveList RayTracingRelevantPrimitiveList;
+	FGraphEventRef RayTracingRelevantPrimitiveListTask;
+	if (bAnyRayTracingPassEnabled)
+	{
+		RayTracingRelevantPrimitiveListTask = FFunctionGraphTask::CreateAndDispatchWhenReady(
+			[Scene = this->Scene, &ReferenceView, &RayTracingRelevantPrimitiveList]()
+		{
+			FTaskTagScope TaskTagScope(ETaskTag::EParallelRenderingThread);
+			RayTracingRelevantPrimitiveList = GatherRayTracingRelevantPrimitives(*Scene, ReferenceView);
+		}, TStatId(), nullptr, ENamedThreads::AnyNormalThreadHiPriTask);
+	}
 #endif // RHI_RAYTRACING
 
 	// Dynamic vertex and index buffers need to be committed before rendering.
@@ -2262,6 +2277,23 @@ void FDeferredShadingSceneRenderer::Render(FRDGBuilder& GraphBuilder)
 		FVirtualTextureSystem::Get().Update(GraphBuilder, FeatureLevel, Scene);
 	}
 
+#if RHI_RAYTRACING
+	if (bAnyRayTracingPassEnabled)
+	{
+		// Wait until RayTracingRelevantPrimitiveList is ready
+		if (RayTracingRelevantPrimitiveListTask.IsValid())
+		{
+			RayTracingRelevantPrimitiveListTask->Wait();
+			RayTracingRelevantPrimitiveListTask.SafeRelease();
+		}
+
+		// Prepare ray tracing scene instance list
+		checkf(RayTracingRelevantPrimitiveList.bValid, TEXT("Ray tracing relevant primitive list is expected to have been created before GatherRayTracingWorldInstancesForView() is called."));
+
+		GatherRayTracingWorldInstancesForView(GraphBuilder, ReferenceView, RayTracingScene, MoveTemp(RayTracingRelevantPrimitiveList));
+	}
+#endif // RHI_RAYTRACING
+
 	const bool bUseGBuffer = IsUsingGBuffers(ShaderPlatform);
 	
 	const bool bRenderDeferredLighting = ViewFamily.EngineShowFlags.Lighting
diff --git a/Engine/Source/Runtime/Renderer/Private/DeferredShadingRenderer.h b/Engine/Source/Runtime/Renderer/Private/DeferredShadingRenderer.h
--- a/Engine/Source/Runtime/Renderer/Private/DeferredShadingRenderer.h
+++ b/Engine/Source/Runtime/Renderer/Private/DeferredShadingRenderer.h
@@ -1011,7 +1011,7 @@ private:
 	void RenderRayTracingBarycentrics(FRDGBuilder& GraphBuilder, const FViewInfo& View, FRDGTextureRef SceneColorOutputTexture, bool bVisualizeProceduralPrimitives);
 
 	/** Fills RayTracingScene instance list for the given View and adds relevant ray tracing data to the view. Does not reset previous scene contents. */
-	bool GatherRayTracingWorldInstancesForView(FRDGBuilder& GraphBuilder, FViewInfo& View, FRayTracingScene& RayTracingScene);
+	bool GatherRayTracingWorldInstancesForView(FRDGBuilder& GraphBuilder, FViewInfo& View, FRayTracingScene& RayTracingScene, struct FRayTracingRelevantPrimitiveList&& RelevantPrimitiveList);
 
 	bool SetupRayTracingPipelineStates(FRDGBuilder& GraphBuilder);
 	void SetupRayTracingLightDataForViews(FRDGBuilder& GraphBuilder);
diff --git a/Engine/Source/Runtime/Renderer/Private/RayTracing/RayTracingInstanceCulling.h b/Engine/Source/Runtime/Renderer/Private/RayTracing/RayTracingInstanceCulling.h
--- a/Engine/Source/Runtime/Renderer/Private/RayTracing/RayTracingInstanceCulling.h
+++ b/Engine/Source/Runtime/Renderer/Private/RayTracing/RayTracingInstanceCulling.h
@@ -34,7 +34,7 @@ bool ShouldSkipPerInstanceCullingForPrimitive(const FRayTracingCullingParameters
 
 struct FRayTracingCullPrimitiveInstancesClosure
 {
-	FScene* Scene;
+	const FScene* Scene;
 	int32 PrimitiveIndex;
 	const FPrimitiveSceneInfo* SceneInfo;
 	bool bIsFarFieldPrimitive;
diff --git a/Engine/Source/Runtime/Renderer/Private/SceneRendering.h b/Engine/Source/Runtime/Renderer/Private/SceneRendering.h
--- a/Engine/Source/Runtime/Renderer/Private/SceneRendering.h
+++ b/Engine/Source/Runtime/Renderer/Private/SceneRendering.h
@@ -1284,7 +1284,6 @@ public:
 	FDynamicRayTracingMeshCommandStorage DynamicRayTracingMeshCommandStorage;
 
 	FRayTracingCullingParameters RayTracingCullingParameters;
-	FGraphEventArray RayTracingPerInstanceCullingTaskList;
 
 	// Data required for FRayTracingScene, depends on RT instance culling tasks
 	FRayTracingSceneWithGeometryInstances RayTracingSceneInitData;
