Fix variable shadowing warning

diff --git a/Engine/Source/Runtime/Renderer/Private/DeferredShadingRenderer.cpp b/Engine/Source/Runtime/Renderer/Private/DeferredShadingRenderer.cpp
--- a/Engine/Source/Runtime/Renderer/Private/DeferredShadingRenderer.cpp
+++ b/Engine/Source/Runtime/Renderer/Private/DeferredShadingRenderer.cpp
@@ -1183,7 +1183,7 @@ bool FDeferredShadingSceneRenderer::GatherRayTracingWorldInstancesForView(FRDGBu
 			int32 Index = INDEX_NONE;
 
 			// Copies the next InstanceSceneDataOffset and user data into the current batch, returns true if arrays were re-allocated.
-			bool Add(FRayTracingScene& RayTracingScene, uint32 InInstanceSceneDataOffset, uint32 InUserData)
+			bool Add(FRayTracingScene& InRayTracingScene, uint32 InInstanceSceneDataOffset, uint32 InUserData)
 			{
 				// Adhoc TArray-like resize behavior, in lieu of support for using a custom FMemStackBase in TArray.
 				// Idea for future: if batch becomes large enough, we could actually split it into multiple instances to avoid memory waste.
@@ -1195,14 +1195,14 @@ bool FDeferredShadingSceneRenderer::GatherRayTracingWorldInstancesForView(FRDGBu
 					int32 PrevCount = InstanceSceneDataOffsets.Num();
 					int32 NextCount = FMath::Max(PrevCount * 2, 1);
 
-					TArrayView<uint32> NewInstanceSceneDataOffsets = RayTracingScene.Allocate<uint32>(NextCount);
+					TArrayView<uint32> NewInstanceSceneDataOffsets = InRayTracingScene.Allocate<uint32>(NextCount);
 					if (PrevCount)
 					{
 						FMemory::Memcpy(NewInstanceSceneDataOffsets.GetData(), InstanceSceneDataOffsets.GetData(), InstanceSceneDataOffsets.GetTypeSize() * InstanceSceneDataOffsets.Num());
 					}
 					InstanceSceneDataOffsets = NewInstanceSceneDataOffsets;
 
-					TArrayView<uint32> NewUserData = RayTracingScene.Allocate<uint32>(NextCount);
+					TArrayView<uint32> NewUserData = InRayTracingScene.Allocate<uint32>(NextCount);
 					if (PrevCount)
 					{
 						FMemory::Memcpy(NewUserData.GetData(), UserData.GetData(), UserData.GetTypeSize() * UserData.Num());
