Compile clustered shading PS in wave32 mode
- Improves performance ~10% due to heavy register pressure

diff --git a/Engine/Source/Runtime/Renderer/Private/ClusteredDeferredShadingPass.cpp b/Engine/Source/Runtime/Renderer/Private/ClusteredDeferredShadingPass.cpp
--- a/Engine/Source/Runtime/Renderer/Private/ClusteredDeferredShadingPass.cpp
+++ b/Engine/Source/Runtime/Renderer/Private/ClusteredDeferredShadingPass.cpp
@@ -127,6 +127,9 @@ class FClusteredShadingPS : public FGlobalShader
 		FGlobalShader::ModifyCompilationEnvironment(Parameters, OutEnvironment);
 		FForwardLightingParameters::ModifyCompilationEnvironment(Parameters.Platform, OutEnvironment);
 		FVirtualShadowMapArray::SetShaderDefines(OutEnvironment);
+
+		// Occupancy is very poor on this shader - this helps a bit in the mean time
+		OutEnvironment.CompilerFlags.Add(CFLAG_Wave32);
 	}
 };
 
