Tightened frustum bounds for rendering deep opacity maps for hair.

#rnx

diff --git a/Engine/Shaders/Private/HairStrands/HairStrandsDebugPrint.usf b/Engine/Shaders/Private/HairStrands/HairStrandsDebugPrint.usf
--- a/Engine/Shaders/Private/HairStrands/HairStrandsDebugPrint.usf
+++ b/Engine/Shaders/Private/HairStrands/HairStrandsDebugPrint.usf
@@ -43,6 +43,7 @@ SamplerState LinearSampler;
 Buffer<uint> HairVisibilityIndirectArgsBuffer;
 
 Buffer<int> HairMacroGroupAABBBuffer;
+Buffer<int> HairMacroGroupVoxelAlignedAABBBuffer;
 
 #define MATERIAL_DATA 1
 #define DETAIL_DATA 1
@@ -87,11 +88,19 @@ void MainCS(uint GroupIndex : SV_GroupIndex, uint3 DispatchThreadId : SV_Dispatc
 		// Macro groups bounds
 		for (uint GroupIndex =0;GroupIndex< HairMacroGroupCount; GroupIndex++)
 		{
-			const float4 InstanceBoundColor = float4(1.f, 0.5f, 0.5f, 1.0f);
+			const float4 InstanceBoundColor = float4(0.25f, 1.f, 0.25f, 1.0f); 
 			FHairAABB Bound = ReadHairAABB(GroupIndex, HairMacroGroupAABBBuffer);
 			AddAABBTWS(Bound.Min, Bound.Max, InstanceBoundColor);
 		}
 
+		// Voxel aligned macro groups bounds
+		for (uint GroupIndex = 0; GroupIndex < HairMacroGroupCount; GroupIndex++)
+		{
+			const float4 InstanceBoundColor = float4(1.f, 0.5f, 0.5f, 1.0f);
+			FHairAABB Bound = ReadHairAABB(GroupIndex, HairMacroGroupVoxelAlignedAABBBuffer);
+			AddAABBTWS(Bound.Min, Bound.Max, InstanceBoundColor);
+		}
+
 		// Hair instance bounds
 		for (uint HairInstanceIndex = 0; HairInstanceIndex < HairInstanceCount; HairInstanceIndex++)
 		{
diff --git a/Engine/Shaders/Private/HairStrands/HairStrandsDeepShadowAllocation.usf b/Engine/Shaders/Private/HairStrands/HairStrandsDeepShadowAllocation.usf
--- a/Engine/Shaders/Private/HairStrands/HairStrandsDeepShadowAllocation.usf
+++ b/Engine/Shaders/Private/HairStrands/HairStrandsDeepShadowAllocation.usf
@@ -85,13 +85,10 @@ void ComputeTranslatedWorldToLightClip(
 	const FHairAABB Light_AABB = Transform(TranslatedWorld_AABB, Coarse_TranslatedWorldToLight);
 	const float3 Light_Extents = GetExtents(Light_AABB);
 
-	const float Radius = length(Light_Extents.xy) * AABBScale;
+	const float Radius = length(GetExtents(TranslatedWorld_AABB)) * AABBScale;
 	float MinZ = -Light_Extents.z * AABBScale;
 	float MaxZ = +Light_Extents.z * AABBScale;
 
-	// Update AABB center based on the tigher Z entents
-	Center = ((MaxZ * OutLightDirection) + (MinZ * OutLightDirection)) * 0.5f + Center;
-
 	const float StrandHairStableRasterizationScale = max(RasterizationScale, 1.0f);
 	OutMinStrandRadiusAtDepth1 = 1;
 	OutTranslatedWorldToClipTransform = 0;
diff --git a/Engine/Shaders/Private/HairStrands/HairStrandsVoxelPageAllocation.usf b/Engine/Shaders/Private/HairStrands/HairStrandsVoxelPageAllocation.usf
--- a/Engine/Shaders/Private/HairStrands/HairStrandsVoxelPageAllocation.usf
+++ b/Engine/Shaders/Private/HairStrands/HairStrandsVoxelPageAllocation.usf
@@ -27,7 +27,8 @@ uint				CPU_bUseCPUData;
 
 Buffer<float>								GPUVoxelWorldSize;
 Buffer<float>								MacroGroupVoxelSizeBuffer;
-RWBuffer<int>								MacroGroupAABBBuffer;
+Buffer<int>									MacroGroupAABBBuffer;
+RWBuffer<int>								MacroGroupVoxelAlignedAABBBuffer;
 RWBuffer<uint4>								OutPageIndexResolutionAndOffsetBuffer;
 RWBuffer<uint>								OutPageIndexAllocationIndirectBufferArgs;
 RWStructuredBuffer<FVoxelizationViewInfo>	OutVoxelizationViewInfoBuffer;
@@ -67,7 +68,11 @@ void AllocatePageIndex(uint2 DispatchThreadId : SV_DispatchThreadID)
 	// Page index allocation
 	int3 PageIndexResolution = 0;
 	{
-		// Snap the max AABB to the voxel size
+		// Snap the max AABB to the voxel size.
+		
+		// The contents of MacroGroupAABBBuffer (tight fitting AABBs) and MacroGroupVoxelAlignedAABBBuffer diverge here
+		// because the macro group AABBs for voxelization need to be snapped to the voxel page boundary.
+		
 		// Allocate enough pages to cover the AABB, where page (0,0,0) origin sit on MinAABB.
 		float3 MacroGroupSize = 0;
 		if (bIsValid)
@@ -102,7 +107,7 @@ void AllocatePageIndex(uint2 DispatchThreadId : SV_DispatchThreadID)
 		if (bIsValid)
 		{
 			OutPageIndexResolutionAndOffsetBuffer[MacroGroupId] = uint4(PageIndexResolution, PageIndexOffset);
-			WriteHairAABB(MacroGroupId, Bound, MacroGroupAABBBuffer);
+			WriteHairAABB(MacroGroupId, Bound, MacroGroupVoxelAlignedAABBBuffer);
 		}
 		else
 		{
@@ -135,7 +140,7 @@ void AllocatePageIndex(uint2 DispatchThreadId : SV_DispatchThreadID)
 		Bound.Min = CPU_TranslatedWorldMinAABB;
 		Bound.Max = CPU_TranslatedWorldMaxAABB;
 		PageIndexResolution = CPU_PageIndexResolution;
-		WriteHairAABB(MacroGroupId, Bound, MacroGroupAABBBuffer);
+		WriteHairAABB(MacroGroupId, Bound, MacroGroupVoxelAlignedAABBBuffer);
 	}
 
 	// Prepare indirect buffer for doing the actual page index allocation and filling the page index
@@ -280,7 +285,7 @@ uint			bForceDirectPageAllocation;
 
 Buffer<int>		GroupAABBsBuffer;
 Buffer<int>		ClusterAABBsBuffer;
-Buffer<int>		MacroGroupAABBBuffer;
+Buffer<int>		MacroGroupVoxelAlignedAABBBuffer;
 Buffer<uint4>	PageIndexResolutionAndOffsetBuffer;
 
 RWBuffer<uint>	OutValidPageIndexBuffer;
@@ -317,7 +322,7 @@ void MarkValid_PrepareCS(uint2 DispatchThreadId : SV_DispatchThreadID)
 		return;
 
 	const uint4 PageIndexResolutionAndOffset = PageIndexResolutionAndOffsetBuffer.Load(MacroGroupId);
-	FHairAABB MacroGroupBound				 = ReadHairAABB(MacroGroupId, MacroGroupAABBBuffer);
+	FHairAABB MacroGroupBound				 = ReadHairAABB(MacroGroupId, MacroGroupVoxelAlignedAABBBuffer);
 	const int3 PageIndexResolution			 = PageIndexResolutionAndOffset.xyz;
 	const uint PageIndexOffset				 = PageIndexResolutionAndOffset.w;
 
@@ -388,7 +393,7 @@ void MarkValid_PrepareCS(uint2 DispatchThreadId : SV_DispatchThreadID)
 void MarkValid_PrepareCS(uint3 DispatchThreadId : SV_DispatchThreadID)
 {
 	const uint4 PageIndexResolutionAndOffset = PageIndexResolutionAndOffsetBuffer.Load(MacroGroupId);
-	const FHairAABB MacroGroupBound = ReadHairAABB(MacroGroupId, MacroGroupAABBBuffer);
+	const FHairAABB MacroGroupBound = ReadHairAABB(MacroGroupId, MacroGroupVoxelAlignedAABBBuffer);
 	const int3 PageIndexResolution = PageIndexResolutionAndOffset.xyz;
 	const uint PageIndexOffset = PageIndexResolutionAndOffset.w;
 
@@ -489,7 +494,7 @@ float3			CPU_TranslatedWorldMaxAABB;
 uint			MacroGroupId;
 
 Buffer<int>		ClusterAABBsBuffer;
-Buffer<int>		MacroGroupAABBBuffer;
+Buffer<int>		MacroGroupVoxelAlignedAABBBuffer;
 Buffer<uint4>	PageIndexResolutionAndOffsetBuffer;
 RWBuffer<uint>	OutValidPageIndexBuffer;
 
@@ -523,7 +528,7 @@ void MarkValidCS(uint2 DispatchThreadId : SV_DispatchThreadID)
 #if PERMUTATION_GPU_DRIVEN == 1
 	const uint4 PageIndexResolutionAndOffset = PageIndexResolutionAndOffsetBuffer.Load(MacroGroupId);
 
-	MacroGroupBound			= ReadHairAABB(MacroGroupId, MacroGroupAABBBuffer);
+	MacroGroupBound			= ReadHairAABB(MacroGroupId, MacroGroupVoxelAlignedAABBBuffer);
 	PageIndexResolution		= PageIndexResolutionAndOffset.xyz;
 	PageIndexOffset			= PageIndexResolutionAndOffset.w;
 #else
@@ -656,7 +661,7 @@ float CPU_VoxelWorldSize;
 uint bUseCPUVoxelWorldSize; // When adaptive voxel size is disabled, we use CPU voxel size value
 
 Buffer<float> GPU_VoxelWorldSize;
-Buffer<int> MacroGroupAABBBuffer;
+Buffer<int> MacroGroupVoxelAlignedAABBBuffer;
 Buffer<float> MacroGroupVoxelSizeBuffer;
 Buffer<uint4> PageIndexResolutionAndOffsetBuffer;
 RWStructuredBuffer<FPackedVirtualVoxelNodeDesc> OutNodeDescBuffer;
@@ -668,7 +673,7 @@ void AddDescCS(uint GroupIndex : SV_GroupIndex, uint3 DispatchThreadId : SV_Disp
 
 #if PERMUTATION_GPU_DRIVEN == 1
 	const uint4 PageIndexResolutionAndOffset = PageIndexResolutionAndOffsetBuffer.Load(MacroGroupId);
-	const FHairAABB TranslatedWorldBound = ReadHairAABB(MacroGroupId, MacroGroupAABBBuffer);
+	const FHairAABB TranslatedWorldBound = ReadHairAABB(MacroGroupId, MacroGroupVoxelAlignedAABBBuffer);
 	const float VoxelWorldSize = MacroGroupVoxelSizeBuffer[MacroGroupId];
 
 	Node.TranslatedWorldMinAABB = TranslatedWorldBound.Min;
diff --git a/Engine/Source/Runtime/Renderer/Private/HairStrands/HairStrandsCluster.cpp b/Engine/Source/Runtime/Renderer/Private/HairStrands/HairStrandsCluster.cpp
--- a/Engine/Source/Runtime/Renderer/Private/HairStrands/HairStrandsCluster.cpp
+++ b/Engine/Source/Runtime/Renderer/Private/HairStrands/HairStrandsCluster.cpp
@@ -312,6 +312,7 @@ void CreateHairStrandsMacroGroups(
 		RDG_GPU_STAT_SCOPE(GraphBuilder, HairStrandsAABB);
 
 		MacroGroupResources.MacroGroupAABBsBuffer = GraphBuilder.CreateBuffer(FRDGBufferDesc::CreateBufferDesc(4, 6 * MacroGroupCount), TEXT("Hair.MacroGroupAABBBuffer"));
+		MacroGroupResources.MacroGroupVoxelAlignedAABBsBuffer = GraphBuilder.CreateBuffer(FRDGBufferDesc::CreateBufferDesc(4, 6 * MacroGroupCount), TEXT("Hair.MacroGroupVoxelAlignedAABBBuffer"));
 		MacroGroupResources.MacroGroupVoxelSizeBuffer = GraphBuilder.CreateBuffer(FRDGBufferDesc::CreateBufferDesc(2, MacroGroupCount), TEXT("Hair.MacroGroupVoxelSize"));
 		FRDGBufferUAVRef MacroGroupAABBBufferUAV = GraphBuilder.CreateUAV(MacroGroupResources.MacroGroupAABBsBuffer, PF_R32_SINT);
 		FRDGBufferUAVRef MacroGroupVoxelSizeBufferUAV = GraphBuilder.CreateUAV(MacroGroupResources.MacroGroupVoxelSizeBuffer, PF_R16F);
diff --git a/Engine/Source/Runtime/Renderer/Private/HairStrands/HairStrandsData.h b/Engine/Source/Runtime/Renderer/Private/HairStrands/HairStrandsData.h
--- a/Engine/Source/Runtime/Renderer/Private/HairStrands/HairStrandsData.h
+++ b/Engine/Source/Runtime/Renderer/Private/HairStrands/HairStrandsData.h
@@ -273,7 +273,8 @@ struct FHairStrandsDeepShadowResources
 struct FHairStrandsMacroGroupResources
 {
 	uint32 MacroGroupCount = 0;
-	FRDGBufferRef MacroGroupAABBsBuffer = nullptr;
+	FRDGBufferRef MacroGroupAABBsBuffer = nullptr; // Tight bounding boxes
+	FRDGBufferRef MacroGroupVoxelAlignedAABBsBuffer = nullptr; // Contents of MacroGroupAABBsBuffer, but snapped to the voxel page boundaries
 	FRDGBufferRef MacroGroupVoxelSizeBuffer = nullptr;
 };
 
diff --git a/Engine/Source/Runtime/Renderer/Private/HairStrands/HairStrandsDebug.cpp b/Engine/Source/Runtime/Renderer/Private/HairStrands/HairStrandsDebug.cpp
--- a/Engine/Source/Runtime/Renderer/Private/HairStrands/HairStrandsDebug.cpp
+++ b/Engine/Source/Runtime/Renderer/Private/HairStrands/HairStrandsDebug.cpp
@@ -317,6 +317,7 @@ class FHairDebugPrintCS : public FGlobalShader
 		SHADER_PARAMETER_RDG_TEXTURE(Texture2D, HairCountUintTexture)
 		SHADER_PARAMETER_RDG_BUFFER_SRV(Buffer<uint>, HairVisibilityIndirectArgsBuffer)
 		SHADER_PARAMETER_RDG_BUFFER_SRV(Buffer, HairMacroGroupAABBBuffer)
+		SHADER_PARAMETER_RDG_BUFFER_SRV(Buffer, HairMacroGroupVoxelAlignedAABBBuffer)
 		SHADER_PARAMETER_RDG_TEXTURE_SRV(Texture2D, StencilTexture)
 		SHADER_PARAMETER_SAMPLER(SamplerState, LinearSampler)
 		SHADER_PARAMETER_STRUCT_INCLUDE(FGPUSceneResourceParameters, GPUSceneResource)
@@ -402,6 +403,7 @@ static void AddDebugHairPrintPass(
 	Parameters->HairMacroGroupCount = MacroGroupResources.MacroGroupCount;
 	Parameters->HairStrands = View->HairStrandsViewData.UniformBuffer;
 	Parameters->HairMacroGroupAABBBuffer = GraphBuilder.CreateSRV(MacroGroupResources.MacroGroupAABBsBuffer, PF_R32_SINT);
+	Parameters->HairMacroGroupVoxelAlignedAABBBuffer = GraphBuilder.CreateSRV(MacroGroupResources.MacroGroupVoxelAlignedAABBsBuffer, PF_R32_SINT);
 	ShaderPrint::SetParameters(GraphBuilder, View->ShaderPrintData, Parameters->ShaderPrintUniformBuffer);
 	TShaderMapRef<FHairDebugPrintCS> ComputeShader(View->ShaderMap);
 
diff --git a/Engine/Source/Runtime/Renderer/Private/HairStrands/HairStrandsVoxelization.cpp b/Engine/Source/Runtime/Renderer/Private/HairStrands/HairStrandsVoxelization.cpp
--- a/Engine/Source/Runtime/Renderer/Private/HairStrands/HairStrandsVoxelization.cpp
+++ b/Engine/Source/Runtime/Renderer/Private/HairStrands/HairStrandsVoxelization.cpp
@@ -236,7 +236,8 @@ class FVoxelAllocatePageIndexCS : public FGlobalShader
 
 		SHADER_PARAMETER_RDG_BUFFER_SRV(Buffer<float>, GPUVoxelWorldSize)
 		SHADER_PARAMETER_RDG_BUFFER_SRV(Buffer<float>, MacroGroupVoxelSizeBuffer)
-		SHADER_PARAMETER_RDG_BUFFER_UAV(RWBuffer<int>, MacroGroupAABBBuffer)
+		SHADER_PARAMETER_RDG_BUFFER_SRV(Buffer<int>, MacroGroupAABBBuffer)
+		SHADER_PARAMETER_RDG_BUFFER_UAV(RWBuffer<int>, MacroGroupVoxelAlignedAABBBuffer)
 		SHADER_PARAMETER_RDG_BUFFER_UAV(RWBuffer<uint4>, OutPageIndexResolutionAndOffsetBuffer)
 		SHADER_PARAMETER_RDG_BUFFER_UAV(RWBuffer, OutVoxelizationViewInfoBuffer)
 		SHADER_PARAMETER_RDG_BUFFER_UAV(RWBuffer, OutPageIndexAllocationIndirectBufferArgs)
@@ -267,7 +268,7 @@ class FVoxelMarkValidPageIndex_PrepareCS : public FGlobalShader
 
 		SHADER_PARAMETER_RDG_BUFFER_SRV(Buffer, GroupAABBsBuffer)
 		SHADER_PARAMETER_RDG_BUFFER_SRV(Buffer, ClusterAABBsBuffer)
-		SHADER_PARAMETER_RDG_BUFFER_SRV(Buffer, MacroGroupAABBBuffer)
+		SHADER_PARAMETER_RDG_BUFFER_SRV(Buffer, MacroGroupVoxelAlignedAABBBuffer)
 		SHADER_PARAMETER_RDG_BUFFER_SRV(Buffer, PageIndexResolutionAndOffsetBuffer)
 		SHADER_PARAMETER_RDG_BUFFER_UAV(RWBuffer<uint>, OutValidPageIndexBuffer)
 		SHADER_PARAMETER_RDG_BUFFER_UAV(RWBuffer<uint>, OutDeferredScatterCounter)
@@ -340,7 +341,7 @@ class FVoxelMarkValidPageIndexCS : public FGlobalShader
 		SHADER_PARAMETER(uint32, CPU_PageIndexOffset)
 		SHADER_PARAMETER(uint32, MacroGroupId)
 		SHADER_PARAMETER_RDG_BUFFER_SRV(Buffer, ClusterAABBsBuffer)
-		SHADER_PARAMETER_RDG_BUFFER_SRV(Buffer, MacroGroupAABBBuffer)
+		SHADER_PARAMETER_RDG_BUFFER_SRV(Buffer, MacroGroupVoxelAlignedAABBBuffer)
 		SHADER_PARAMETER_RDG_BUFFER_SRV(Buffer, PageIndexResolutionAndOffsetBuffer)
 		SHADER_PARAMETER_RDG_BUFFER_UAV(RWBuffer<uint>, OutValidPageIndexBuffer)
 	END_SHADER_PARAMETER_STRUCT()
@@ -402,7 +403,7 @@ class FVoxelAddNodeDescCS : public FGlobalShader
 		SHADER_PARAMETER(float, CPU_VoxelWorldSize)
 		SHADER_PARAMETER(uint32, bUseCPUVoxelWorldSize)
 		SHADER_PARAMETER_RDG_BUFFER_SRV(Buffer, GPU_VoxelWorldSize)
-		SHADER_PARAMETER_RDG_BUFFER_SRV(Buffer, MacroGroupAABBBuffer)
+		SHADER_PARAMETER_RDG_BUFFER_SRV(Buffer, MacroGroupVoxelAlignedAABBBuffer)
 		SHADER_PARAMETER_RDG_BUFFER_SRV(Buffer, MacroGroupVoxelSizeBuffer)
 		SHADER_PARAMETER_RDG_BUFFER_SRV(Buffer, PageIndexResolutionAndOffsetBuffer)
 		SHADER_PARAMETER_RDG_BUFFER_UAV(RWStructuredBuffer, OutNodeDescBuffer)
@@ -659,7 +660,8 @@ static void AddAllocateVoxelPagesPass(
 		Parameters->MacroGroupCount = MacroGroupCount;
 		Parameters->bDoesMacroGroupSupportVoxelization = bDoesMacroGroupSupportVoxelization;
 		Parameters->MacroGroupVoxelSizeBuffer = GraphBuilder.CreateSRV(MacroGroupResources.MacroGroupVoxelSizeBuffer, PF_R16F);
-		Parameters->MacroGroupAABBBuffer = GraphBuilder.CreateUAV(MacroGroupResources.MacroGroupAABBsBuffer, PF_R32_SINT);
+		Parameters->MacroGroupAABBBuffer = GraphBuilder.CreateSRV(MacroGroupResources.MacroGroupAABBsBuffer, PF_R32_SINT);
+		Parameters->MacroGroupVoxelAlignedAABBBuffer = GraphBuilder.CreateUAV(MacroGroupResources.MacroGroupVoxelAlignedAABBsBuffer, PF_R32_SINT);
 		Parameters->IndirectDispatchGroupSize = GroupSize; // This is the GroupSize used for FVoxelAllocateVoxelPageCS
 		Parameters->OutPageIndexResolutionAndOffsetBuffer = GraphBuilder.CreateUAV(PageIndexResolutionBuffer, PF_R32G32B32A32_UINT);
 		Parameters->OutVoxelizationViewInfoBuffer = GraphBuilder.CreateUAV(VoxelizationViewInfoBuffer);
@@ -746,7 +748,7 @@ static void AddAllocateVoxelPagesPass(
 
 				Parameters->GroupAABBsBuffer			= RegisterAsSRV(GraphBuilder, HairGroupData->GetGroupAABBBuffer());
 				Parameters->ClusterAABBsBuffer			= RegisterAsSRV(GraphBuilder, HairGroupData->GetClusterAABBBuffer());
-				Parameters->MacroGroupAABBBuffer		= GraphBuilder.CreateSRV(MacroGroupResources.MacroGroupAABBsBuffer, PF_R32_SINT);
+				Parameters->MacroGroupVoxelAlignedAABBBuffer = GraphBuilder.CreateSRV(MacroGroupResources.MacroGroupVoxelAlignedAABBsBuffer, PF_R32_SINT);
 				Parameters->PageIndexResolutionAndOffsetBuffer = PageIndexResolutionAndOffsetBufferSRV;
 
 				Parameters->OutDeferredScatterCounter	= ScatterCounterUAV;
@@ -827,7 +829,7 @@ static void AddAllocateVoxelPagesPass(
 
 				if (bIsGPUDriven)
 				{
-					Parameters->MacroGroupAABBBuffer = GraphBuilder.CreateSRV(MacroGroupResources.MacroGroupAABBsBuffer, PF_R32_SINT);
+					Parameters->MacroGroupVoxelAlignedAABBBuffer = GraphBuilder.CreateSRV(MacroGroupResources.MacroGroupVoxelAlignedAABBsBuffer, PF_R32_SINT);
 					Parameters->PageIndexResolutionAndOffsetBuffer = PageIndexResolutionAndOffsetBufferSRV;
 				}
 
@@ -866,7 +868,7 @@ static void AddAllocateVoxelPagesPass(
 
 			if (bIsGPUDriven)
 			{
-				Parameters->MacroGroupAABBBuffer = GraphBuilder.CreateSRV(MacroGroupResources.MacroGroupAABBsBuffer, PF_R32_SINT);
+				Parameters->MacroGroupVoxelAlignedAABBBuffer = GraphBuilder.CreateSRV(MacroGroupResources.MacroGroupVoxelAlignedAABBsBuffer, PF_R32_SINT);
 				Parameters->PageIndexResolutionAndOffsetBuffer = PageIndexResolutionAndOffsetBufferSRV;
 			}
 
