This CL:
- Adds a separate temporal accumulation pass for strand based hair if Temporal Layering is enabled - this greatly reduces temporal artifacts (flickering and ghosting) when Temporal Layering is enabled.
- And fixes a bug in Continuous LOD screen size calculation.

diff --git a/Engine/Plugins/Runtime/HairStrands/Source/HairStrandsCore/Private/GroomManager.cpp b/Engine/Plugins/Runtime/HairStrands/Source/HairStrandsCore/Private/GroomManager.cpp
--- a/Engine/Plugins/Runtime/HairStrands/Source/HairStrandsCore/Private/GroomManager.cpp
+++ b/Engine/Plugins/Runtime/HairStrands/Source/HairStrandsCore/Private/GroomManager.cpp
@@ -745,6 +745,21 @@ static void RunHairLODSelection(
 			// Feedback game thread with LOD selection 
 			Instance->Debug.LODPredictedIndex = LODViewIndex;
 			Instance->HairGroupPublicData->DebugScreenSize = MaxScreenSize;
+
+			if (IsHairStrandContinuousDecimationReorderingEnabled())
+			{
+				if (IsHairVisibilityComputeRasterContinuousLODEnabled())
+				{
+					Instance->HairGroupPublicData->ContinuousLODBounds = SphereBound;
+					Instance->HairGroupPublicData->MaxScreenSize = MaxScreenSize;
+				}
+
+				Instance->HairGroupPublicData->UpdateTemporalIndex();
+			}
+			else
+			{
+				Instance->HairGroupPublicData->MaxScreenSize = 1.0;
+			}
 		}
 
 		// Function for selecting, loading, & initializing LOD resources
diff --git a/Engine/Shaders/Private/HairStrands/HairStrandsVisibilityComposeSubPixelPS.usf b/Engine/Shaders/Private/HairStrands/HairStrandsVisibilityComposeSubPixelPS.usf
--- a/Engine/Shaders/Private/HairStrands/HairStrandsVisibilityComposeSubPixelPS.usf
+++ b/Engine/Shaders/Private/HairStrands/HairStrandsVisibilityComposeSubPixelPS.usf
@@ -35,6 +35,9 @@ uint bComposeDofDepth;
 Texture2D<float>	HairDOFDepthTexture;
 Texture2D<float4>	HairLightingSampleBuffer;
 
+uint TemporalLayerCount;
+Texture2D HairTemporalAccumulationTexture;
+
 float4 EvaluateVolumetric(float3 WorldPosition, float SceneDepth)
 {
 	const float3 CameraRelative_WorldPosition = WorldPosition - LWCHackToFloat(PrimaryView.WorldCameraOrigin);
@@ -82,7 +85,18 @@ void ComposeSamplePS(
 {
 	OutColor = 0.0f;
 	const uint3 PixelCoord = uint3(floor(Input.Position.xy), 0);
+	
+	const float HairDeviceZ = HairStrands.HairOnlyDepthTexture.Load(PixelCoord).x;
+
+#if PERMUTATION_TEMPORAL == 1
+	OutColor.rgba = HairTemporalAccumulationTexture.Load(PixelCoord);
+
+	// Approximation: convert running temporal average layer coverage (alpha) to average haircount, multiply by number of layers to get haircount sum, then convert haircount sum to final coverage using equation: Coverage = saturate(min(1.0 - pow(2.71828, -HairCount/1000.0), 1.0) / 0.98);
+	float HairCountAverage = (log(1.0 - (OutColor.a * 0.98)) / log(2.71828))*-1000.0;
+	float HairCountSum = HairCountAverage*TemporalLayerCount;
+	OutColor.a = saturate(min(1.0 - pow(2.71828, -HairCountSum/1000.0), 1.0) / 0.98);
 
+#else // PERMUTATION_TEMPORAL
 	const FNodeDesc NodeDesc = DecodeNodeDesc(HairStrands.HairSampleOffset.Load(PixelCoord));
 	if (NodeDesc.Count == 0)
 	{
@@ -91,12 +105,13 @@ void ComposeSamplePS(
 
 	const float PixelCoverage = min(HairStrands.HairCoverageTexture.Load(PixelCoord), 1);
 	if (PixelCoverage == 0)
+	{
 		discard;
+	}
 
 	const uint TotalNodeCount = HairStrands.HairSampleCount.Load(uint3(0, 0, 0));
 	const uint Resolution = uint(ceil(sqrt(float(TotalNodeCount))));
 
-	const float HairDeviceZ = HairStrands.HairOnlyDepthTexture.Load(PixelCoord).x;
 	const float3 ClosetPointWorldPosition = LWCHackToFloat(SvPositionToWorld(float4(Input.Position.xy, HairDeviceZ, 1.0)));
 	const float4 Volumetric = EvaluateVolumetric(ClosetPointWorldPosition, ConvertFromDeviceZ(HairDeviceZ));
 
@@ -112,6 +127,8 @@ void ComposeSamplePS(
 
 	OutColor.rgb *= PixelCoverage;
 	OutColor.a = PixelCoverage;
+#endif // PERMUTATION_TEMPORAL
+
 	OutDepth = HairDeviceZ;
 	if (bComposeDofDepth)
 	{
diff --git a/Engine/Shaders/Private/TemporalAA.usf b/Engine/Shaders/Private/TemporalAA.usf
--- a/Engine/Shaders/Private/TemporalAA.usf
+++ b/Engine/Shaders/Private/TemporalAA.usf
@@ -102,7 +102,13 @@
 			#define AA_SAMPLE_CACHE_METHOD (AA_SAMPLE_CACHE_METHOD_LDS)
 		#endif
 	#endif
-
+#elif TAA_PASS_CONFIG == 7 // Used for Hair
+	#define AA_HISTORY_PAYLOAD (HISTORY_PAYLOAD_RGB_OPACITY)
+	#define AA_DYNAMIC 1
+	#define AA_FILTERED 0
+	#define AA_LERP 3
+	#define AA_MANUALLY_CLAMP_HISTORY_UV 1
+	#define AA_YCOCG 0
 #elif TAA_PASS_CONFIG == 3 // Used for SSR
 	#define AA_HISTORY_PAYLOAD (HISTORY_PAYLOAD_RGB_OPACITY)
 	#define AA_DYNAMIC 1
diff --git a/Engine/Source/Runtime/Renderer/Private/DeferredShadingRenderer.cpp b/Engine/Source/Runtime/Renderer/Private/DeferredShadingRenderer.cpp
--- a/Engine/Source/Runtime/Renderer/Private/DeferredShadingRenderer.cpp
+++ b/Engine/Source/Runtime/Renderer/Private/DeferredShadingRenderer.cpp
@@ -3190,7 +3190,7 @@ void FDeferredShadingSceneRenderer::Render(FRDGBuilder& GraphBuilder)
 	if (GetHairStrandsComposition() == EHairStrandsCompositionType::BeforeTranslucent)
 	{
 		RDG_GPU_STAT_SCOPE(GraphBuilder, HairRendering);
-		RenderHairComposition(GraphBuilder, Views, SceneTextures.Color.Target, SceneTextures.Depth.Target);
+		RenderHairComposition(GraphBuilder, Views, SceneTextures.Color.Target, SceneTextures.Depth.Target, SceneTextures.Velocity);
 	}
 
 	// Draw translucency.
@@ -3239,7 +3239,7 @@ void FDeferredShadingSceneRenderer::Render(FRDGBuilder& GraphBuilder)
 		if (GetHairStrandsComposition() == EHairStrandsCompositionType::AfterTranslucent)
 		{
 			RDG_GPU_STAT_SCOPE(GraphBuilder, HairRendering);
-			RenderHairComposition(GraphBuilder, Views, SceneTextures.Color.Target, SceneTextures.Depth.Target);
+			RenderHairComposition(GraphBuilder, Views, SceneTextures.Color.Target, SceneTextures.Depth.Target, SceneTextures.Velocity);
 		}
 
 		if (bShouldRenderDistortion)
@@ -3267,7 +3267,7 @@ void FDeferredShadingSceneRenderer::Render(FRDGBuilder& GraphBuilder)
 	else if (GetHairStrandsComposition() == EHairStrandsCompositionType::AfterTranslucent)
 	{
 		RDG_GPU_STAT_SCOPE(GraphBuilder, HairRendering);
-		RenderHairComposition(GraphBuilder, Views, SceneTextures.Color.Target, SceneTextures.Depth.Target);
+		RenderHairComposition(GraphBuilder, Views, SceneTextures.Color.Target, SceneTextures.Depth.Target, SceneTextures.Velocity);
 	}
 
 #if !UE_BUILD_SHIPPING
diff --git a/Engine/Source/Runtime/Renderer/Private/DeferredShadingRenderer.h b/Engine/Source/Runtime/Renderer/Private/DeferredShadingRenderer.h
--- a/Engine/Source/Runtime/Renderer/Private/DeferredShadingRenderer.h
+++ b/Engine/Source/Runtime/Renderer/Private/DeferredShadingRenderer.h
@@ -698,7 +698,7 @@ private:
 	/** Renders the scene's translucency passes. */
 	void RenderTranslucency(
 		FRDGBuilder& GraphBuilder,
-		const FMinimalSceneTextures& SceneTextures,
+		const FSceneTextures& SceneTextures,
 		const FTranslucencyLightingVolumeTextures& TranslucencyLightingVolumeTextures,
 		FTranslucencyPassResourcesMap* OutTranslucencyResourceMap,
 		ETranslucencyView ViewsToRender,
diff --git a/Engine/Source/Runtime/Renderer/Private/HairStrands/HairStrandsComposition.cpp b/Engine/Source/Runtime/Renderer/Private/HairStrands/HairStrandsComposition.cpp
--- a/Engine/Source/Runtime/Renderer/Private/HairStrands/HairStrandsComposition.cpp
+++ b/Engine/Source/Runtime/Renderer/Private/HairStrands/HairStrandsComposition.cpp
@@ -11,6 +11,7 @@
 #include "HairStrandsRendering.h"
 #include "HairStrandsTile.h"
 #include "FogRendering.h"
+#include "TemporalAA.h"
 
 /////////////////////////////////////////////////////////////////////////////////////////
 
@@ -37,7 +38,8 @@ enum EHairStrandsCommonPassType
 	Composition,
 	DOF,
 	TAAFastResolve,
-	GBuffer
+	GBuffer,
+	Blit
 };
 
 template<typename TPassParameter, typename TPixelShader>
@@ -88,6 +90,10 @@ void InternalCommonDrawPass(
 			// 1: invalid pixel (background)
 			GraphicsPSOInit.BlendState = TStaticBlendState<CW_RGBA, BO_Add, BF_One, BF_InverseSourceAlpha, BO_Add, BF_Zero, BF_InverseSourceAlpha>::GetRHI();
 		}
+		else if (Type == EHairStrandsCommonPassType::Blit) //used to write data into the temporal input buffer
+		{
+			GraphicsPSOInit.BlendState = TStaticBlendState<CW_RGBA, BO_Add, BF_One, BF_Zero, BO_Add, BF_One, BF_Zero>::GetRHI();
+		}
 		else
 		{
 			GraphicsPSOInit.BlendState = TStaticBlendState<CW_RGBA, BO_Add, BF_One, BF_Zero, BO_Max, BF_One, BF_One>::GetRHI();
@@ -111,7 +117,7 @@ void InternalCommonDrawPass(
 				false, CF_Always, SO_Keep, SO_Keep, SO_Keep,
 				STENCIL_TEMPORAL_RESPONSIVE_AA_MASK, STENCIL_TEMPORAL_RESPONSIVE_AA_MASK>::GetRHI();
 		}
-		else if (Type == EHairStrandsCommonPassType::GBuffer)
+		else if ((Type == EHairStrandsCommonPassType::GBuffer) || (Type == EHairStrandsCommonPassType::Blit))
 		{
 			if (bWriteDepth)
 			{
@@ -165,7 +171,8 @@ class FHairVisibilityComposeSamplePS : public FGlobalShader
 	SHADER_USE_PARAMETER_STRUCT(FHairVisibilityComposeSamplePS, FGlobalShader);
 
 	class FDebug : SHADER_PERMUTATION_BOOL("PERMUTATION_DEBUG");
-	using FPermutationDomain = TShaderPermutationDomain<FDebug>;
+	class FTemporal : SHADER_PERMUTATION_BOOL("PERMUTATION_TEMPORAL");
+	using FPermutationDomain = TShaderPermutationDomain<FDebug, FTemporal>;
 	
 	BEGIN_SHADER_PARAMETER_STRUCT(FParameters, )
 		SHADER_PARAMETER_STRUCT_REF(FViewUniformShaderParameters, ViewUniformBuffer)
@@ -175,6 +182,9 @@ class FHairVisibilityComposeSamplePS : public FGlobalShader
 		SHADER_PARAMETER_RDG_UNIFORM_BUFFER(FHairStrandsViewUniformParameters, HairStrands)
 		SHADER_PARAMETER_RDG_TEXTURE(Texture2D, HairLightingSampleBuffer)
 		SHADER_PARAMETER_RDG_TEXTURE(Texture2D, HairDOFDepthTexture)
+		SHADER_PARAMETER(uint32, TemporalLayerCount)
+		SHADER_PARAMETER_RDG_TEXTURE(Texture2D, HairTemporalAccumulationTexture)
+		SHADER_PARAMETER_RDG_TEXTURE_UAV(RWTexture2D<uint>, OutMetaTexture)
 		SHADER_PARAMETER_RDG_UNIFORM_BUFFER(FFogUniformParameters, FogStruct)
 		RENDER_TARGET_BINDING_SLOTS()
 	END_SHADER_PARAMETER_STRUCT()
@@ -190,18 +200,24 @@ public:
 
 IMPLEMENT_GLOBAL_SHADER(FHairVisibilityComposeSamplePS, "/Engine/Private/HairStrands/HairStrandsVisibilityComposeSubPixelPS.usf", "ComposeSamplePS", SF_Pixel);
 
+int32 GetHairVisibilityComputeRasterTemporalLayerCount();
+
 static void AddHairVisibilityComposeSamplePass(
 	FRDGBuilder& GraphBuilder,
 	const FViewInfo& View,
 	const FHairStrandsVisibilityData& VisibilityData,
 	const FRDGTextureRef& CategorizationTexture,
 	const FRDGTextureRef& HairDOFDepthTexture,
+	const FRDGTextureRef& HairTemporalAccumulationTexture,
 	FRDGTextureRef& OutColorTexture,
-	FRDGTextureRef& OutDepthTexture)
+	FRDGTextureRef& OutDepthTexture,
+	const bool bTemporalLayeringEnabled)
 {
 	check(VisibilityData.SampleLightingTexture);
 	const bool bDOFEnable = HairDOFDepthTexture != nullptr ? 1 : 0;
 
+	const bool bTemporal = HairTemporalAccumulationTexture != nullptr ? 1 : 0;
+
 	TRDGUniformBufferRef<FFogUniformParameters> FogBuffer = CreateFogUniformBuffer(GraphBuilder, View);
 
 	FHairVisibilityComposeSamplePS::FParameters* Parameters = GraphBuilder.AllocParameters<FHairVisibilityComposeSamplePS::FParameters>();
@@ -212,19 +228,22 @@ static void AddHairVisibilityComposeSamplePass(
 	Parameters->ViewUniformBuffer = View.ViewUniformBuffer;
 	Parameters->HairStrands = View.HairStrandsViewData.UniformBuffer;
 	Parameters->FogStruct = FogBuffer;
-	Parameters->RenderTargets[0] = FRenderTargetBinding(OutColorTexture, ERenderTargetLoadAction::ELoad);
-	Parameters->RenderTargets.DepthStencil = FDepthStencilBinding(OutDepthTexture, ERenderTargetLoadAction::ELoad, ERenderTargetLoadAction::ENoAction, FExclusiveDepthStencil::DepthWrite_StencilNop);
+	Parameters->TemporalLayerCount = bTemporal ? GetHairVisibilityComputeRasterTemporalLayerCount() : 1;
+	Parameters->HairTemporalAccumulationTexture = bTemporal ? HairTemporalAccumulationTexture : GSystemTextures.GetBlackDummy(GraphBuilder);
+	Parameters->RenderTargets[0] = FRenderTargetBinding(OutColorTexture, bTemporal ? ERenderTargetLoadAction::ELoad : (bTemporalLayeringEnabled ? ERenderTargetLoadAction::EClear : ERenderTargetLoadAction::ELoad) );
+	Parameters->RenderTargets.DepthStencil = FDepthStencilBinding(OutDepthTexture, ERenderTargetLoadAction::ELoad, ERenderTargetLoadAction::ENoAction, FExclusiveDepthStencil::DepthWrite_StencilNop);	
 
 	const bool bDebugComposition = View.Family->EngineShowFlags.LODColoration;
 	FHairVisibilityComposeSamplePS::FPermutationDomain PermutationVector;
 	PermutationVector.Set<FHairVisibilityComposeSamplePS::FDebug>(bDebugComposition);
+	PermutationVector.Set<FHairVisibilityComposeSamplePS::FTemporal>(bTemporal);
 	TShaderMapRef<FHairVisibilityComposeSamplePS> PixelShader(View.ShaderMap, PermutationVector);
 	InternalCommonDrawPass(
 		GraphBuilder,
 		RDG_EVENT_NAME("HairStrands::ComposeSample"),
 		View,
 		OutColorTexture->Desc.Extent,
-		EHairStrandsCommonPassType::Composition,
+		bTemporal ? EHairStrandsCommonPassType::Composition : (bTemporalLayeringEnabled ? EHairStrandsCommonPassType::Blit : EHairStrandsCommonPassType::Composition),
 		false,
 		VisibilityData.TileData,
 		PixelShader,
@@ -467,11 +486,14 @@ static void AddHairVisibilityGBufferWritePass(
 		Parameters);
 }
 
+bool IsHairVisibilityComputeRasterTemporalLayeringEnabled();
+
 static void InternalRenderHairComposition(
 	FRDGBuilder& GraphBuilder,
 	const FViewInfo& View,
 	FRDGTextureRef SceneColorTexture,
-	FRDGTextureRef SceneDepthTexture)
+	FRDGTextureRef SceneDepthTexture,
+	FRDGTextureRef SceneVelocityTexture)
 {
 	DECLARE_GPU_STAT(HairStrandsComposition);
 	RDG_EVENT_SCOPE(GraphBuilder, "HairStrandsComposition");
@@ -500,14 +522,57 @@ static void InternalRenderHairComposition(
 						SceneDepthTexture);
 				}
 
+				FRDGTextureRef AccumulatedColor = nullptr;
+
+				const bool bTemporalLayeringEnabled = IsHairVisibilityComputeRasterTemporalLayeringEnabled() && View.ViewState;
+
+				if (bTemporalLayeringEnabled)
+				{
+					FIntPoint OutputResolution = SceneColorTexture->Desc.Extent;
+
+					FRDGTextureDesc HairColorInputDesc = FRDGTextureDesc::Create2D(OutputResolution, PF_FloatRGBA, FClearValueBinding(FLinearColor(0, 0, 0, 0)), TexCreate_RenderTargetable | TexCreate_ShaderResource | TexCreate_UAV, 1);
+					FRDGTextureRef HairColorInputTexture = GraphBuilder.CreateTexture(HairColorInputDesc, TEXT("Hair.ColorInputTexture"));
+
+					AddHairVisibilityComposeSamplePass(
+						GraphBuilder,
+						View,
+						VisibilityData,
+						VisibilityData.CoverageTexture,
+						DOFDepth,
+						nullptr,
+						HairColorInputTexture,
+						SceneDepthTexture,
+						bTemporalLayeringEnabled);
+
+
+					check(View.ViewState);
+					FTAAPassParameters TAASettings(View);
+					TAASettings.SceneDepthTexture = SceneDepthTexture;
+					TAASettings.SceneVelocityTexture = SceneVelocityTexture;
+					TAASettings.Pass = ETAAPassConfig::Hair;
+					TAASettings.SceneColorInput = HairColorInputTexture;
+					TAASettings.bOutputRenderTargetable = true;
+
+					FTAAOutputs TAAOutputs = AddTemporalAAPass(
+						GraphBuilder,
+						View,
+						TAASettings,
+						View.PrevViewInfo.HairHistory,
+						&View.ViewState->PrevFrameViewInfo.HairHistory);
+
+					AccumulatedColor = TAAOutputs.SceneColor;			
+				}
+
 				AddHairVisibilityComposeSamplePass(
 					GraphBuilder,
 					View,
 					VisibilityData,
 					VisibilityData.CoverageTexture,
 					DOFDepth,
+					AccumulatedColor,
 					SceneColorTexture,
-					SceneDepthTexture);
+					SceneDepthTexture,
+					bTemporalLayeringEnabled);
 
 				if (VisibilityData.ResolveMaskTexture)
 				{
@@ -568,7 +633,8 @@ void RenderHairComposition(
 	FRDGBuilder& GraphBuilder,
 	const TArray<FViewInfo>& Views,
 	FRDGTextureRef SceneColorTexture,
-	FRDGTextureRef SceneDepthTexture)
+	FRDGTextureRef SceneDepthTexture,
+	FRDGTextureRef SceneVelocityTexture)
 {
 	for (const FViewInfo& View : Views)
 	{		
@@ -578,7 +644,8 @@ void RenderHairComposition(
 				GraphBuilder,
 				View,
 				SceneColorTexture,
-				SceneDepthTexture);
+				SceneDepthTexture,
+				SceneVelocityTexture);
 		}
 	}
 }
@@ -587,7 +654,8 @@ void RenderHairComposition(
 	FRDGBuilder& GraphBuilder,
 	const FViewInfo& View,
 	FRDGTextureRef SceneColorTexture,
-	FRDGTextureRef SceneDepthTexture)
+	FRDGTextureRef SceneDepthTexture,
+	FRDGTextureRef SceneVelocityTexture)
 {
 	if (View.Family && HairStrands::HasViewHairStrandsData(View))
 	{
@@ -595,6 +663,7 @@ void RenderHairComposition(
 			GraphBuilder,
 			View,
 			SceneColorTexture,
-			SceneDepthTexture);
+			SceneDepthTexture,
+			SceneVelocityTexture);
 	}
 }
\ No newline at end of file
diff --git a/Engine/Source/Runtime/Renderer/Private/HairStrands/HairStrandsComposition.h b/Engine/Source/Runtime/Renderer/Private/HairStrands/HairStrandsComposition.h
--- a/Engine/Source/Runtime/Renderer/Private/HairStrands/HairStrandsComposition.h
+++ b/Engine/Source/Runtime/Renderer/Private/HairStrands/HairStrandsComposition.h
@@ -13,10 +13,12 @@ void RenderHairComposition(
 	FRDGBuilder& GraphBuilder,
 	const FViewInfo& View,
 	FRDGTextureRef SceneColorTexture,
-	FRDGTextureRef SceneDepthTexture);
+	FRDGTextureRef SceneDepthTexture, 
+	FRDGTextureRef VelocityTexture);
 
 void RenderHairComposition(
 	FRDGBuilder& GraphBuilder, 
 	const TArray<FViewInfo>& Views,
 	FRDGTextureRef SceneColorTexture,
-	FRDGTextureRef SceneDepthTexture);
+	FRDGTextureRef SceneDepthTexture,
+	FRDGTextureRef VelocityTexture);
diff --git a/Engine/Source/Runtime/Renderer/Private/HairStrands/HairStrandsInterface.cpp b/Engine/Source/Runtime/Renderer/Private/HairStrands/HairStrandsInterface.cpp
--- a/Engine/Source/Runtime/Renderer/Private/HairStrands/HairStrandsInterface.cpp
+++ b/Engine/Source/Runtime/Renderer/Private/HairStrands/HairStrandsInterface.cpp
@@ -80,13 +80,13 @@ static TAutoConsoleVariable<int32> CVarHairStrandsVisibilityComputeRaster_Contin
 	ECVF_RenderThreadSafe);
 
 static TAutoConsoleVariable<int32> CVarHairStrandsVisibilityComputeRaster_TemporalLayering(
-	TEXT("r.HairStrands.Visibility.ComputeRaster.TemporalLayering"), 0,
+	TEXT("r.HairStrands.Visibility.ComputeRaster.TemporalLayering"), 1,
 	TEXT("Enable Experimental WIP Temporal Layering (requires TAA changes to work well)"),
 	ECVF_RenderThreadSafe);
 
 static TAutoConsoleVariable<int32> CVarHairStrandsVisibilityComputeRaster_TemporalLayering_LayerCount(
-	TEXT("r.HairStrands.Visibility.ComputeRaster.TemporalLayering.LayerCount"), 8,
-	TEXT("Enable Temporal Layering Override Index  8)"),
+	TEXT("r.HairStrands.Visibility.ComputeRaster.TemporalLayering.LayerCount"), 2,
+	TEXT("Temporal Layering Layer Count (default: 2)"),
 	ECVF_RenderThreadSafe);
 
 static TAutoConsoleVariable<int32> CVarHairStrandsVisibilityComputeRaster_TemporalLayering_OverrideIndex(
@@ -470,7 +470,8 @@ bool IsHairVisibilityComputeRasterContinuousLODEnabled()
 
 float GetHairVisibilityComputeRasterContinuousLODScale(float ScreenSize)
 {
-	return FMath::Pow(FMath::Clamp(ScreenSize, 1.0f / 8.0f, 1.0f), 2.0f);
+	//TODO: make values in this calculation customizable per groom 
+	return FMath::Pow(FMath::Clamp(ScreenSize, 1.0f / 16.0f, 1.0f), 1.0f);
 }
 
 bool IsHairVisibilityComputeRasterTemporalLayeringEnabled()
@@ -484,7 +485,7 @@ int32 GetHairVisibilityComputeRasterTemporalLayerCount()
 	return LayerCount;
 }
 
-uint32 GetHairVisibilityComputeRasterVertexStart(uint32 TemporalAASampleIndex, uint32 InVertexCount)
+uint32 GetHairVisibilityComputeRasterVertexStart(uint32 TemporalIndex, uint32 InVertexCount)
 {
 	uint32 VertexStart = 0;
 
@@ -495,7 +496,7 @@ uint32 GetHairVisibilityComputeRasterVertexStart(uint32 TemporalAASampleIndex, u
 
 		const uint32 VertexCount = InVertexCount / LayerCount;
 
-		VertexStart = (((OverrideIndex >= 0) ? OverrideIndex : TemporalAASampleIndex) % LayerCount) * VertexCount;
+		VertexStart = (((OverrideIndex >= 0) ? OverrideIndex : TemporalIndex) % LayerCount) * VertexCount;
 	}
 
 	return VertexStart;
@@ -518,11 +519,11 @@ uint32 GetHairVisibilityComputeRasterVertexCount(float ScreenSize, uint32 InVert
 	return VertexCount;
 }
 
-float GetHairVisibilityComputeRasterSampleWeight(float ScreenSize)
+float GetHairVisibilityComputeRasterSampleWeight(float ScreenSize, bool bUseTemporalWeight)
 {
 	float SampleWeight = 1.0;
 
-	if (IsHairVisibilityComputeRasterTemporalLayeringEnabled())
+	if (IsHairVisibilityComputeRasterTemporalLayeringEnabled() && bUseTemporalWeight)
 	{
 		// sample weight should increase if layer count increases
 		SampleWeight *= GetHairVisibilityComputeRasterTemporalLayerCount();
@@ -537,6 +538,10 @@ float GetHairVisibilityComputeRasterSampleWeight(float ScreenSize)
 	return SampleWeight;
 }
 
+void FHairGroupPublicData::UpdateTemporalIndex()
+{
+	TemporalIndex = IsHairVisibilityComputeRasterTemporalLayeringEnabled() ? ((TemporalIndex + 1) % GetHairVisibilityComputeRasterTemporalLayerCount()) : 0;
+}
 
 ///////////////////////////////////////////////////////////////////////////////////////////////////
 // Bookmark API
diff --git a/Engine/Source/Runtime/Renderer/Private/HairStrands/HairStrandsVisibility.cpp b/Engine/Source/Runtime/Renderer/Private/HairStrands/HairStrandsVisibility.cpp
--- a/Engine/Source/Runtime/Renderer/Private/HairStrands/HairStrandsVisibility.cpp
+++ b/Engine/Source/Runtime/Renderer/Private/HairStrands/HairStrandsVisibility.cpp
@@ -3015,7 +3015,7 @@ IMPLEMENT_GLOBAL_SHADER(FVisiblityRasterComputeRasterizeCS, "/Engine/Private/Hai
 
 uint32 GetHairVisibilityComputeRasterVertexStart(uint32 TemporalAASampleIndex, uint32 InVertexCount);
 uint32 GetHairVisibilityComputeRasterVertexCount(float ScreenSize, uint32 InVertexCount);
-float GetHairVisibilityComputeRasterSampleWeight(float ScreenSize);
+float GetHairVisibilityComputeRasterSampleWeight(float ScreenSize, bool bUseTemporalWeight);
 bool IsHairStrandContinuousDecimationReorderingEnabled();
 
 static FRasterComputeOutput AddVisibilityComputeRasterPass(
@@ -3105,8 +3105,6 @@ static FRasterComputeOutput AddVisibilityComputeRasterPass(
 
 	const uint32 GroupSize = GetVendorOptimalGroupSize1D();
 	
-	const uint32 TemporalAASampleIndex = ViewInfo.ViewState ? ViewInfo.ViewState->GetCurrentTemporalAASampleIndex() : 0;
-
 	FVisiblityRasterComputeBinningCS::FPermutationDomain BinningPermutationVector_CullingOff;
 	BinningPermutationVector_CullingOff.Set<FVisiblityRasterComputeBinningCS::FGroupSize>(GroupSize);
 	BinningPermutationVector_CullingOff.Set<FVisiblityRasterComputeBinningCS::FCulling>(false);
@@ -3138,8 +3136,6 @@ static FRasterComputeOutput AddVisibilityComputeRasterPass(
 	{
 		const FHairStrandsMacroGroupData::TPrimitiveInfos& PrimitiveSceneInfos = MacroGroup.PrimitivesInfos;
 
-		const FSphere MacroBoundsSphere = MacroGroup.Bounds.GetSphere();
-
 		for (const FHairStrandsMacroGroupData::PrimitiveInfo& PrimitiveInfo : PrimitiveSceneInfos)
 		{
 			check(PrimitiveInfo.Mesh && PrimitiveInfo.Mesh->Elements.Num() > 0);
@@ -3151,11 +3147,13 @@ static FRasterComputeOutput AddVisibilityComputeRasterPass(
 
 			const bool bCullingEnable = GHairVisibilityComputeRaster_Culling ? HairGroupPublicData->GetCullingResultAvailable() : false;
 			
-			const float ScreenSize = ComputeBoundsScreenSize(FVector4(MacroBoundsSphere.Center, 1), MacroBoundsSphere.W, ViewInfo);
+			// allow fewer strands to be rasterized if screen size for current view is smaller than max screen size
+			const FSphere BoundsSphere = HairGroupPublicData->ContinuousLODBounds.GetSphere();
+			const float ScreenSize = FMath::Min(ComputeBoundsScreenSize(FVector4(BoundsSphere.Center, 1), BoundsSphere.W, ViewInfo), HairGroupPublicData->MaxScreenSize);
 
-			const uint32 VertexStart = GetHairVisibilityComputeRasterVertexStart(TemporalAASampleIndex, VFInput.Strands.VertexCount);
+			const uint32 VertexStart = GetHairVisibilityComputeRasterVertexStart(HairGroupPublicData->TemporalIndex, VFInput.Strands.VertexCount);
 			const uint32 VertexCount = GetHairVisibilityComputeRasterVertexCount(ScreenSize, VFInput.Strands.VertexCount);
-			const float SampleWeight = GetHairVisibilityComputeRasterSampleWeight(ScreenSize);
+			const float SampleWeight = GetHairVisibilityComputeRasterSampleWeight(ScreenSize, false);
 
 			//reset buffers
 			uint32 IndexGridClearValues[4] = { 0x0,0x0,0x0,0x0 };
@@ -3215,7 +3213,7 @@ static FRasterComputeOutput AddVisibilityComputeRasterPass(
 			RasterParameters->OutputResolutionf = FVector2f(Out.SuperResolution.X, Out.SuperResolution.Y);
 			RasterParameters->ResolutionMultiplier = Out.ResolutionMultiplier;
 			RasterParameters->MacroGroupId = MacroGroup.MacroGroupId;
-			RasterParameters->MaxRasterCount = TileSize; //TODO: implement tile clipping in raster pass to exisiting avoid edge case
+			RasterParameters->MaxRasterCount = TileSize; //TODO: implement tile clipping in raster pass to avoid exisiting edge case
 			RasterParameters->HairMaterialId = PrimitiveInfo.MaterialId;
 			RasterParameters->ViewUniformBuffer = ViewUniformShaderParameters;
 			RasterParameters->SceneDepthTexture = SceneDepthTexture;
diff --git a/Engine/Source/Runtime/Renderer/Private/HairStrands/HairStrandsVoxelization.cpp b/Engine/Source/Runtime/Renderer/Private/HairStrands/HairStrandsVoxelization.cpp
--- a/Engine/Source/Runtime/Renderer/Private/HairStrands/HairStrandsVoxelization.cpp
+++ b/Engine/Source/Runtime/Renderer/Private/HairStrands/HairStrandsVoxelization.cpp
@@ -1265,9 +1265,9 @@ public:
 
 IMPLEMENT_GLOBAL_SHADER(FVoxelRasterComputeCS, "/Engine/Private/HairStrands/HairStrandsVoxelRasterCompute.usf", "MainCS", SF_Compute);
 
-uint32 GetHairVisibilityComputeRasterVertexStart(uint32 TemporalAASampleIndex, uint32 InVertexCount);
+uint32 GetHairVisibilityComputeRasterVertexStart(uint32 TemporalIndex, uint32 InVertexCount);
 uint32 GetHairVisibilityComputeRasterVertexCount(float ScreenSize, uint32 InVertexCount);
-float GetHairVisibilityComputeRasterSampleWeight(float ScreenSize);
+float GetHairVisibilityComputeRasterSampleWeight(float ScreenSize, bool bUseTemporalWeight);
 bool IsHairStrandContinuousDecimationReorderingEnabled();
 
 static void AddVirtualVoxelizationComputeRasterPass(
@@ -1284,8 +1284,6 @@ static void AddVirtualVoxelizationComputeRasterPass(
 
 	if (ViewInfo)
 	{
-		const uint32 TemporalAASampleIndex = ViewInfo->ViewState ? ViewInfo->ViewState->GetCurrentTemporalAASampleIndex() : 0;
-
 		const FHairStrandsMacroGroupData::TPrimitiveInfos& PrimitiveSceneInfos = MacroGroup.PrimitivesInfos;
 
 		FRDGBufferSRVRef VoxelizationViewInfoBufferSRV = GraphBuilder.CreateSRV(VoxelResources.VoxelizationViewInfoBuffer);
@@ -1317,9 +1315,9 @@ static void AddVirtualVoxelizationComputeRasterPass(
 					continue;
 				}
 
-				const uint32 VertexStart = GetHairVisibilityComputeRasterVertexStart(TemporalAASampleIndex, VFInput.Strands.VertexCount);
-				const uint32 VertexCount = GetHairVisibilityComputeRasterVertexCount(HairGroupPublicData->DebugScreenSize, VFInput.Strands.VertexCount);
-				const float SampleWeight = GetHairVisibilityComputeRasterSampleWeight(HairGroupPublicData->DebugScreenSize);
+				const uint32 VertexStart = GetHairVisibilityComputeRasterVertexStart(HairGroupPublicData->TemporalIndex, VFInput.Strands.VertexCount);
+				const uint32 VertexCount = GetHairVisibilityComputeRasterVertexCount(HairGroupPublicData->MaxScreenSize, VFInput.Strands.VertexCount);
+				const float SampleWeight = GetHairVisibilityComputeRasterSampleWeight(HairGroupPublicData->MaxScreenSize, true);
 
 				FTransform LocalToTranslatedWorldTransform = VFInput.LocalToWorldTransform;
 				LocalToTranslatedWorldTransform.AddToTranslation(TranslatedWorldOffset);
diff --git a/Engine/Source/Runtime/Renderer/Private/PostProcess/PostProcessing.cpp b/Engine/Source/Runtime/Renderer/Private/PostProcess/PostProcessing.cpp
--- a/Engine/Source/Runtime/Renderer/Private/PostProcess/PostProcessing.cpp
+++ b/Engine/Source/Runtime/Renderer/Private/PostProcess/PostProcessing.cpp
@@ -577,7 +577,7 @@ void AddPostProcessingPasses(
 
 			if (GetHairStrandsComposition() == EHairStrandsCompositionType::AfterSeparateTranslucent)
 			{
-				RenderHairComposition(GraphBuilder, View, SceneColor.Texture, SceneDepth.Texture);
+				RenderHairComposition(GraphBuilder, View, SceneColor.Texture, SceneDepth.Texture, Velocity.Texture);
 			}
 		}
 
diff --git a/Engine/Source/Runtime/Renderer/Private/PostProcess/TemporalAA.cpp b/Engine/Source/Runtime/Renderer/Private/PostProcess/TemporalAA.cpp
--- a/Engine/Source/Runtime/Renderer/Private/PostProcess/TemporalAA.cpp
+++ b/Engine/Source/Runtime/Renderer/Private/PostProcess/TemporalAA.cpp
@@ -342,6 +342,7 @@ const TCHAR* const kTAAOutputNames[] = {
 	TEXT("LightShaft.TemporalAA"),
 	TEXT("DOF.TemporalAA"),
 	TEXT("DOF.TemporalAA"),
+	TEXT("Hair.TemporalAA"),
 };
 
 const TCHAR* const kTAAPassNames[] = {
@@ -352,6 +353,7 @@ const TCHAR* const kTAAPassNames[] = {
 	TEXT("LightShaft"),
 	TEXT("DOF"),
 	TEXT("DOFUpsampling"),
+	TEXT("Hair"),
 };
 
 const TCHAR* const kTAAQualityNames[] = {
diff --git a/Engine/Source/Runtime/Renderer/Private/PostProcess/TemporalAA.h b/Engine/Source/Runtime/Renderer/Private/PostProcess/TemporalAA.h
--- a/Engine/Source/Runtime/Renderer/Private/PostProcess/TemporalAA.h
+++ b/Engine/Source/Runtime/Renderer/Private/PostProcess/TemporalAA.h
@@ -43,6 +43,9 @@ enum class ETAAPassConfig
 	// Permutation for DOF that handle Coc.
 	DiaphragmDOF,
 	DiaphragmDOFUpsampling,
+	
+	// Permutation for hair.
+	Hair,
 
 	MAX
 };
diff --git a/Engine/Source/Runtime/Renderer/Private/SceneRendering.h b/Engine/Source/Runtime/Renderer/Private/SceneRendering.h
--- a/Engine/Source/Runtime/Renderer/Private/SceneRendering.h
+++ b/Engine/Source/Runtime/Renderer/Private/SceneRendering.h
@@ -1049,6 +1049,9 @@ struct FPreviousViewInfo
 	FTemporalAAHistory SSRHistory;
 	FTemporalAAHistory WaterSSRHistory;
 
+	// Temporal AA history for Hair
+	FTemporalAAHistory HairHistory;
+
 	// Temporal AA history for the editor primitive depth upsampling
 	FTemporalAAHistory EditorPrimtiveDepthHistory;
 
diff --git a/Engine/Source/Runtime/Renderer/Private/TranslucentRendering.cpp b/Engine/Source/Runtime/Renderer/Private/TranslucentRendering.cpp
--- a/Engine/Source/Runtime/Renderer/Private/TranslucentRendering.cpp
+++ b/Engine/Source/Runtime/Renderer/Private/TranslucentRendering.cpp
@@ -1406,7 +1406,7 @@ void FDeferredShadingSceneRenderer::RenderTranslucencyInner(
 
 void FDeferredShadingSceneRenderer::RenderTranslucency(
 	FRDGBuilder& GraphBuilder,
-	const FMinimalSceneTextures& SceneTextures,
+	const FSceneTextures& SceneTextures,
 	const FTranslucencyLightingVolumeTextures& TranslucentLightingVolumeTextures,
 	FTranslucencyPassResourcesMap* OutTranslucencyResourceMap,
 	ETranslucencyView ViewsToRender,
@@ -1480,7 +1480,7 @@ void FDeferredShadingSceneRenderer::RenderTranslucency(
 		RenderTranslucencyInner(GraphBuilder, SceneTextures, TranslucentLightingVolumeTextures, OutTranslucencyResourceMap, SharedDepthTexture, ViewsToRender, SceneColorCopyTexture, ETranslucencyPass::TPT_StandardTranslucency, InstanceCullingManager);
 		if (GetHairStrandsComposition() == EHairStrandsCompositionType::AfterTranslucentBeforeTranslucentAfterDOF)
 		{
-			RenderHairComposition(GraphBuilder, Views, SceneTextures.Color.Target, SceneTextures.Depth.Target);
+			RenderHairComposition(GraphBuilder, Views, SceneTextures.Color.Target, SceneTextures.Depth.Target, SceneTextures.Velocity);
 		}
 		RenderTranslucencyInner(GraphBuilder, SceneTextures, TranslucentLightingVolumeTextures, OutTranslucencyResourceMap, SharedDepthTexture, ViewsToRender, SceneColorCopyTexture, ETranslucencyPass::TPT_TranslucencyAfterDOF, InstanceCullingManager);
 		RenderTranslucencyInner(GraphBuilder, SceneTextures, TranslucentLightingVolumeTextures, OutTranslucencyResourceMap, SharedDepthTexture, ViewsToRender, SceneColorCopyTexture, ETranslucencyPass::TPT_TranslucencyAfterDOFModulate, InstanceCullingManager);
diff --git a/Engine/Source/Runtime/Renderer/Public/HairStrandsInterface.h b/Engine/Source/Runtime/Renderer/Public/HairStrandsInterface.h
--- a/Engine/Source/Runtime/Renderer/Public/HairStrandsInterface.h
+++ b/Engine/Source/Runtime/Renderer/Public/HairStrandsInterface.h
@@ -241,6 +241,9 @@ public:
 	bool GetLODVisibility() const { return bLODVisibility; }
 
 	uint32 GetClusterCount() const { return ClusterCount;  }
+
+	void UpdateTemporalIndex();
+
 	struct FVertexFactoryInput 
 	{
 		struct FStrands
@@ -342,6 +345,10 @@ public:
 	float LODBias = 0;			// Current LOD bias
 	bool bLODVisibility = true; // Enable/disable hair rendering for this component
 
+	FBoxSphereBounds ContinuousLODBounds; 	//used by Continuous LOD
+	float MaxScreenSize = 0.f; 				//used by Continuous LOD
+	uint32 TemporalIndex = 0; 				//used by Temporal Layering
+
 	// Debug
 	bool  bDebugDrawLODInfo = false; // Enable/disable hair LOD info
 	float DebugScreenSize = 0.f;
