Remove the VisibleRayTracingMeshCommands reservation / reallocation validation check.

The reservation estimate of one mesh command per relevant primitive is not conservative and making it conservative is not practical.
However, since FRayTracingMeshCommandOneFrameArray is just a basic TArray, it is OK to resize it on a task (at a perf penalty due to realloc).

diff --git a/Engine/Source/Runtime/Renderer/Private/DeferredShadingRenderer.cpp b/Engine/Source/Runtime/Renderer/Private/DeferredShadingRenderer.cpp
--- a/Engine/Source/Runtime/Renderer/Private/DeferredShadingRenderer.cpp
+++ b/Engine/Source/Runtime/Renderer/Private/DeferredShadingRenderer.cpp
@@ -1159,15 +1159,14 @@ bool FDeferredShadingSceneRenderer::GatherRayTracingWorldInstancesForView(FRDGBu
 		const FScene& Scene;
 		TArray<FRayTracingRelevantPrimitive> RelevantPrimitives;
 		const FRayTracingCullingParameters& CullingParameters;
-		int32 ExpectedMaxVisibleRayTracingMeshCommands = 0;
 
 		// Outputs
 
 		FRayTracingScene& RayTracingScene; // New instances are added into FRayTracingScene::Instances and FRayTracingScene::Allocator is used for temporary data
-		FRayTracingMeshCommandOneFrameArray& VisibleRayTracingMeshCommands; // New elements are added here by this task
+		TArray<FVisibleRayTracingMeshCommand>& VisibleRayTracingMeshCommands; // New elements are added here by this task
 
 		FRayTracingSceneAddInstancesTask(const FScene& InScene, TArray<FRayTracingRelevantPrimitive> InRelevantPrimitives, const FRayTracingCullingParameters& InCullingParameters,
-											FRayTracingScene& InRayTracingScene, FRayTracingMeshCommandOneFrameArray& InVisibleRayTracingMeshCommands)
+											FRayTracingScene& InRayTracingScene, TArray<FVisibleRayTracingMeshCommand>& InVisibleRayTracingMeshCommands)
 			: Scene(InScene)
 			, RelevantPrimitives(MoveTemp(InRelevantPrimitives))
 			, CullingParameters(InCullingParameters)
@@ -1175,7 +1174,6 @@ bool FDeferredShadingSceneRenderer::GatherRayTracingWorldInstancesForView(FRDGBu
 			, VisibleRayTracingMeshCommands(InVisibleRayTracingMeshCommands)
 		{
 			VisibleRayTracingMeshCommands.Reserve(RelevantPrimitives.Num());
-			ExpectedMaxVisibleRayTracingMeshCommands = VisibleRayTracingMeshCommands.Max();
 		}
 
 		struct FAutoInstanceBatch
@@ -1413,9 +1411,6 @@ bool FDeferredShadingSceneRenderer::GatherRayTracingWorldInstancesForView(FRDGBu
 			{
 				MyCompletionGraphEvent->DontCompleteUntil(CullingTask);
 			}
-
-
-			checkf(VisibleRayTracingMeshCommands.Max() == ExpectedMaxVisibleRayTracingMeshCommands, TEXT("VisibleRayTracingMeshCommands should not be reallocated during FRayTracingSceneAddInstancesTask"));
 		}
 	};
 
