Multi-view-family:  Fix crash due to ray tracing scene being released mid-frame if one of the view families doesn't use ray tracing.  Track ray tracing instance count for the entire frame, and release the ray tracing scene on the last relevant view family if nothing was rendered.

#rnx

diff --git a/Engine/Source/Runtime/Renderer/Private/DeferredShadingRenderer.cpp b/Engine/Source/Runtime/Renderer/Private/DeferredShadingRenderer.cpp
--- a/Engine/Source/Runtime/Renderer/Private/DeferredShadingRenderer.cpp
+++ b/Engine/Source/Runtime/Renderer/Private/DeferredShadingRenderer.cpp
@@ -1775,10 +1775,11 @@ static void ReleaseRaytracingResources(FRDGBuilder& GraphBuilder, TArrayView<FVi
 		{
 			RHICmdList.ClearRayTracingBindings(RayTracingScene.GetRHIRayTracingScene());
 
-			// If we did not end up rendering anything this frame, then release all ray tracing scene resources.
-			if (RayTracingScene.Instances.Num() == 0)
+			// Track if we ended up rendering anything this frame.  After rendering all view families, we'll release the
+			// ray tracing scene resources if nothing used ray tracing.
+			if (RayTracingScene.Instances.Num() > 0)
 			{
-				RayTracingScene.ResetAndReleaseResources();
+				RayTracingScene.bUsedThisFrame = true;
 			}
 		}
 
diff --git a/Engine/Source/Runtime/Renderer/Private/RayTracing/RayTracingScene.h b/Engine/Source/Runtime/Renderer/Private/RayTracing/RayTracingScene.h
--- a/Engine/Source/Runtime/Renderer/Private/RayTracing/RayTracingScene.h
+++ b/Engine/Source/Runtime/Renderer/Private/RayTracing/RayTracingScene.h
@@ -90,6 +90,8 @@ public:
 	// This must be filled before calling CreateRayTracingSceneWithGeometryInstances() and Create().
 	TArray<FRayTracingGeometryInstance> Instances;
 
+	bool bUsedThisFrame = false;
+
 	uint32 NumCallableShaderSlots = 0;
 	TArray<FRayTracingShaderCommand> CallableCommands;
 
diff --git a/Engine/Source/Runtime/Renderer/Private/SceneRendering.cpp b/Engine/Source/Runtime/Renderer/Private/SceneRendering.cpp
--- a/Engine/Source/Runtime/Renderer/Private/SceneRendering.cpp
+++ b/Engine/Source/Runtime/Renderer/Private/SceneRendering.cpp
@@ -3865,6 +3865,9 @@ void FSceneRenderer::CreateSceneRenderers(TArrayView<const FSceneViewFamily*> In
 					bShouldUpdateRayTracingScene = false;
 				}
 			}
+
+			// Clear flag that tracks whether ray tracing was used this frame
+			SceneRenderers[0]->Scene->RayTracingScene.bUsedThisFrame = false;
 		});
 	}
 #endif  // RHI_RAYTRACING
@@ -4293,6 +4296,9 @@ static void RenderViewFamilies_RenderThread(FRHICommandListImmediate& RHICmdList
 {
 	LLM_SCOPE(ELLMTag::SceneRender);
 
+	// All renderers point to the same Scene (calling code asserts this)
+	FScene* const Scene = SceneRenderers[0]->Scene;
+
 	// We need to execute the pre-render view extensions before we do any view dependent work.
 	for (FSceneRenderer* SceneRenderer : SceneRenderers)
 	{
@@ -4388,9 +4394,9 @@ static void RenderViewFamilies_RenderThread(FRHICommandListImmediate& RHICmdList
 	{
 		CSV_SCOPED_TIMING_STAT_EXCLUSIVE(PostRenderCleanUp);
 
-		if (IsHairStrandsEnabled(EHairStrandsShaderType::All, SceneRenderers[0]->Scene->GetShaderPlatform()) && (SceneRenderers[0]->AllFamilyViews.Num() > 0) && !bAnyShowHitProxies)
+		if (IsHairStrandsEnabled(EHairStrandsShaderType::All, Scene->GetShaderPlatform()) && (SceneRenderers[0]->AllFamilyViews.Num() > 0) && !bAnyShowHitProxies)
 		{
-			FHairStrandsBookmarkParameters Parameters = CreateHairStrandsBookmarkParameters(SceneRenderers[0]->Scene, SceneRenderers[0]->Views, SceneRenderers[0]->AllFamilyViews);
+			FHairStrandsBookmarkParameters Parameters = CreateHairStrandsBookmarkParameters(Scene, SceneRenderers[0]->Views, SceneRenderers[0]->AllFamilyViews);
 			if (Parameters.HasInstances())
 			{
 				RunHairStrandsBookmark(EHairStrandsBookmark::ProcessEndOfFrame, Parameters);
@@ -4400,12 +4406,12 @@ static void RenderViewFamilies_RenderThread(FRHICommandListImmediate& RHICmdList
 		// Only reset per-frame scene state once all views have processed their frame, including those in planar reflections
 		for (int32 CacheType = 0; CacheType < UE_ARRAY_COUNT(SceneRenderers[0]->Scene->DistanceFieldSceneData.PrimitiveModifiedBounds); CacheType++)
 		{
-			ResetAndShrinkModifiedBounds(SceneRenderers[0]->Scene->DistanceFieldSceneData.PrimitiveModifiedBounds[CacheType]);
+			ResetAndShrinkModifiedBounds(Scene->DistanceFieldSceneData.PrimitiveModifiedBounds[CacheType]);
 		}
 
-		if (SceneRenderers[0]->Scene->LumenSceneData)
+		if (Scene->LumenSceneData)
 		{
-			ResetAndShrinkModifiedBounds(SceneRenderers[0]->Scene->LumenSceneData->PrimitiveModifiedBounds);
+			ResetAndShrinkModifiedBounds(Scene->LumenSceneData->PrimitiveModifiedBounds);
 		}
 
 		// Immediately issue EndFrame() for all extensions in case any of the outstanding tasks they issued getting out of this frame
@@ -4417,12 +4423,20 @@ static void RenderViewFamilies_RenderThread(FRHICommandListImmediate& RHICmdList
 		}
 	}
 
+#if RHI_RAYTRACING
+	// Release the ray tracing scene resources if ray tracing wasn't used
+	if (!Scene->RayTracingScene.bUsedThisFrame)
+	{
+		Scene->RayTracingScene.ResetAndReleaseResources();
+	}
+#endif  // RHI_RAYTRACING
+
 #if STATS
 	{
 		QUICK_SCOPE_CYCLE_COUNTER(STAT_RenderViewFamily_RenderThread_MemStats);
 
 		// Update scene memory stats that couldn't be tracked continuously
-		SET_MEMORY_STAT(STAT_RenderingSceneMemory, SceneRenderers[0]->Scene->GetSizeBytes());
+		SET_MEMORY_STAT(STAT_RenderingSceneMemory, Scene->GetSizeBytes());
 
 		SIZE_T ViewStateMemory = 0;
 		for (FSceneRenderer* SceneRenderer : SceneRenderers)
@@ -5517,4 +5531,4 @@ FRDGTextureRef CreateHalfResolutionDepthCheckerboardMinMax(FRDGBuilder& GraphBui
 	}
 
 	return SmallDepthTexture;
-}
\ No newline at end of file
+}
