Fix ExpHeight fog taking the vFog start distance when the volumetric fog is disable via device profile

diff --git a/Engine/Source/Runtime/Renderer/Private/FogRendering.cpp b/Engine/Source/Runtime/Renderer/Private/FogRendering.cpp
--- a/Engine/Source/Runtime/Renderer/Private/FogRendering.cpp
+++ b/Engine/Source/Runtime/Renderer/Private/FogRendering.cpp
@@ -310,7 +310,7 @@ static void RenderViewFog(
 	if (GraphicsPSOInit.bDepthBounds)
 	{
 		float ExpFogStartDistance = View.ExponentialFogParameters.W;
-		float VolFogStartDistance = View.bEnableVolumetricFog ? View.VolumetricFogStartDistance : ExpFogStartDistance;
+		float VolFogStartDistance = bShouldRenderVolumetricFog ? View.VolumetricFogStartDistance : ExpFogStartDistance;
 
 		// The fog can be set to start at a certain euclidean distance.
 		// clamp the value to be behind the near plane z, according to the smallest distance between volumetric fog and height fog (if they are enabled). 
