Added passing ViewFamily.DebugDPIScale scale to FCanvas in AddDrawCanvasPass to make debug text readable for non-editor targets.

diff --git a/Engine/Source/Runtime/Renderer/Private/ScreenPass.h b/Engine/Source/Runtime/Renderer/Private/ScreenPass.h
--- a/Engine/Source/Runtime/Renderer/Private/ScreenPass.h
+++ b/Engine/Source/Runtime/Renderer/Private/ScreenPass.h
@@ -657,7 +657,7 @@ FORCEINLINE void AddDrawCanvasPass(
 	check(Output.IsValid());
 
 	const FSceneViewFamily& ViewFamily = *View.Family;
-	FCanvas& Canvas = *FCanvas::Create(GraphBuilder, Output.Texture, nullptr, ViewFamily.Time, View.GetFeatureLevel());
+	FCanvas& Canvas = *FCanvas::Create(GraphBuilder, Output.Texture, nullptr, ViewFamily.Time, View.GetFeatureLevel(), ViewFamily.DebugDPIScale);
 	Canvas.SetRenderTargetRect(Output.ViewRect);
 
 	Function(Canvas);
