Adding missing condition to not run distance field code for water when it's turned off

#rnx

diff --git a/Engine/Source/Runtime/Renderer/Private/SingleLayerWaterRendering.cpp b/Engine/Source/Runtime/Renderer/Private/SingleLayerWaterRendering.cpp
--- a/Engine/Source/Runtime/Renderer/Private/SingleLayerWaterRendering.cpp
+++ b/Engine/Source/Runtime/Renderer/Private/SingleLayerWaterRendering.cpp
@@ -517,59 +517,63 @@ void FDeferredShadingSceneRenderer::RenderSingleLayerWaterReflections(
 
 		const FPerViewPipelineState& ViewPipelineState = GetViewPipelineState(View);
 
-		FProjectedShadowInfo* DistanceFieldShadowInfo = nullptr;
-
-		// Try to find the ProjectedShadowInfo corresponding to ray trace shadow info for the main directional light.
-		const FLightSceneProxy* SelectedForwardDirectionalLightProxy = View.ForwardLightingResources.SelectedForwardDirectionalLightProxy;
-		if (SelectedForwardDirectionalLightProxy)
+		if (IsWaterDistanceFieldShadowEnabled_Runtime(View.GetShaderPlatform()))
 		{
-			FLightSceneInfo* LightSceneInfo = SelectedForwardDirectionalLightProxy->GetLightSceneInfo();
-			FVisibleLightInfo& VisibleLightViewInfo = VisibleLightInfos[LightSceneInfo->Id];
-			
-			for (int32 ShadowIndex = 0; ShadowIndex < VisibleLightViewInfo.ShadowsToProject.Num(); ShadowIndex++)
+
+			FProjectedShadowInfo* DistanceFieldShadowInfo = nullptr;
+
+			// Try to find the ProjectedShadowInfo corresponding to ray trace shadow info for the main directional light.
+			const FLightSceneProxy* SelectedForwardDirectionalLightProxy = View.ForwardLightingResources.SelectedForwardDirectionalLightProxy;
+			if (SelectedForwardDirectionalLightProxy)
 			{
-				FProjectedShadowInfo* ProjectedShadowInfo = VisibleLightViewInfo.ShadowsToProject[ShadowIndex];
-				if (ProjectedShadowInfo->bRayTracedDistanceField)
+				FLightSceneInfo* LightSceneInfo = SelectedForwardDirectionalLightProxy->GetLightSceneInfo();
+				FVisibleLightInfo& VisibleLightViewInfo = VisibleLightInfos[LightSceneInfo->Id];
+
+				for (int32 ShadowIndex = 0; ShadowIndex < VisibleLightViewInfo.ShadowsToProject.Num(); ShadowIndex++)
 				{
-					DistanceFieldShadowInfo = ProjectedShadowInfo;
+					FProjectedShadowInfo* ProjectedShadowInfo = VisibleLightViewInfo.ShadowsToProject[ShadowIndex];
+					if (ProjectedShadowInfo->bRayTracedDistanceField)
+					{
+						DistanceFieldShadowInfo = ProjectedShadowInfo;
+					}
 				}
 			}
-		}
 
-		// If DFShadow data has been found, then combine it with the separate main directional light luminance texture.
-		FRDGTextureRef ScreenShadowMaskTexture = SystemTextures.White;
-		if (DistanceFieldShadowInfo)
-		{
-			RDG_EVENT_SCOPE(GraphBuilder, "SLW::DistanceFieldShadow");
-
-			FIntRect ScissorRect;
-			if (!SelectedForwardDirectionalLightProxy->GetScissorRect(ScissorRect, View, View.ViewRect))
+			// If DFShadow data has been found, then combine it with the separate main directional light luminance texture.
+			FRDGTextureRef ScreenShadowMaskTexture = SystemTextures.White;
+			if (DistanceFieldShadowInfo)
 			{
-				ScissorRect = View.ViewRect;
-			}
+				RDG_EVENT_SCOPE(GraphBuilder, "SLW::DistanceFieldShadow");
 
-			// Reset the cached texture to create a new one mapping to the water depth buffer
-			DistanceFieldShadowInfo->ResetRayTracedDistanceFieldShadow(&View);
+				FIntRect ScissorRect;
+				if (!SelectedForwardDirectionalLightProxy->GetScissorRect(ScissorRect, View, View.ViewRect))
+				{
+					ScissorRect = View.ViewRect;
+				}
 
-			FTiledShadowRendering TiledShadowRendering;
-			if (bRunTiled)
-			{
-				TiledShadowRendering.DrawIndirectParametersBuffer	= TiledScreenSpaceReflection.DrawIndirectParametersBuffer;
-				TiledShadowRendering.TileListDataBufferSRV			= TiledScreenSpaceReflection.TileListDataBufferSRV;
-				TiledShadowRendering.TileSize						= TiledScreenSpaceReflection.TileSize;
-			}
+				// Reset the cached texture to create a new one mapping to the water depth buffer
+				DistanceFieldShadowInfo->ResetRayTracedDistanceFieldShadow(&View);
 
-			const bool bProjectingForForwardShading = false;
-			const bool bForceRGBModulation = true;
-			DistanceFieldShadowInfo->RenderRayTracedDistanceFieldProjection( 
-				GraphBuilder,
-				SceneTextures,
-				SceneWithoutWaterTextures.SeparatedMainDirLightTexture,
-				View,
-				ScissorRect,
-				bProjectingForForwardShading,
-				bForceRGBModulation,
-				bRunTiled ? &TiledShadowRendering : nullptr);
+				FTiledShadowRendering TiledShadowRendering;
+				if (bRunTiled)
+				{
+					TiledShadowRendering.DrawIndirectParametersBuffer = TiledScreenSpaceReflection.DrawIndirectParametersBuffer;
+					TiledShadowRendering.TileListDataBufferSRV = TiledScreenSpaceReflection.TileListDataBufferSRV;
+					TiledShadowRendering.TileSize = TiledScreenSpaceReflection.TileSize;
+				}
+
+				const bool bProjectingForForwardShading = false;
+				const bool bForceRGBModulation = true;
+				DistanceFieldShadowInfo->RenderRayTracedDistanceFieldProjection(
+					GraphBuilder,
+					SceneTextures,
+					SceneWithoutWaterTextures.SeparatedMainDirLightTexture,
+					View,
+					ScissorRect,
+					bProjectingForForwardShading,
+					bForceRGBModulation,
+					bRunTiled ? &TiledShadowRendering : nullptr);
+			}
 		}
 
 		if (ViewPipelineState.ReflectionsMethod == EReflectionsMethod::Lumen && CVarWaterSingleLayerLumenReflections.GetValueOnRenderThread() != 0)
