Fix Niagara's early view UB initialization in ISR.

diff --git a/Engine/Source/Runtime/Renderer/Private/SceneVisibility.cpp b/Engine/Source/Runtime/Renderer/Private/SceneVisibility.cpp
--- a/Engine/Source/Runtime/Renderer/Private/SceneVisibility.cpp
+++ b/Engine/Source/Runtime/Renderer/Private/SceneVisibility.cpp
@@ -4919,6 +4919,11 @@ void FDeferredShadingSceneRenderer::InitViews(FRDGBuilder& GraphBuilder, const F
 		// This needs to run before ComputeViewVisibility() is called, but the views normally initialize the ViewUniformBuffer after that (at the end of this method).
 		if (FXSystem && FXSystem->RequiresEarlyViewUniformBuffer() && Views.IsValidIndex(0))
 		{
+			// during ISR, instanced view RHI resources need to be initialized first.
+			if (FViewInfo* InstancedView = const_cast<FViewInfo*>(Views[0].GetInstancedView()))
+			{
+				InstancedView->InitRHIResources();
+			}
 			Views[0].InitRHIResources();
 			FXSystem->PostInitViews(GraphBuilder, Views, !ViewFamily.EngineShowFlags.HitProxies);
 		}
