Fixed bug where an assert occurs when an FScene gets deleted before an FSceneViewState that's linked to the Scene.  We can't assume the order in which these get deleted, so the code is being changed to defer linking the two structures together until an API is used that requires linking the structures (such as enabling a per-view VSM cache).  And deletion logic can now happen in any order, with either structure unlinking itself from the other as needed.

Because the linking is deferred, we no longer need to allocate the FSceneViewState from the FScene, so that API can be removed (it was added in 5.1 and is being removed before 5.1 ships, so no need to deprecate).  We also don't need to worry about the code path that supports reloading the renderer, as the FScene and FSceneViewState can be relinked later -- it doesn't need to be done during reload.  So the need to pass a UWorld to the FSceneViewStateReference is removed.

#rnx

diff --git a/Engine/Plugins/Runtime/nDisplay/Source/DisplayCluster/Private/Game/EngineClasses/Basics/DisplayClusterViewportClient.cpp b/Engine/Plugins/Runtime/nDisplay/Source/DisplayCluster/Private/Game/EngineClasses/Basics/DisplayClusterViewportClient.cpp
--- a/Engine/Plugins/Runtime/nDisplay/Source/DisplayCluster/Private/Game/EngineClasses/Basics/DisplayClusterViewportClient.cpp
+++ b/Engine/Plugins/Runtime/nDisplay/Source/DisplayCluster/Private/Game/EngineClasses/Basics/DisplayClusterViewportClient.cpp
@@ -588,7 +588,7 @@ void UDisplayClusterViewportClient::Draw(FViewport* InViewport, FCanvas* SceneCa
 					View->CameraConstrainedViewRect = View->UnscaledViewRect;
 
 					// Enable per-view virtual shadow map caching
-					View->State->AddVirtualShadowMapCache();
+					View->State->AddVirtualShadowMapCache(MyWorld->Scene);
 
 					{
 						// Save the location of the view.
diff --git a/Engine/Source/Runtime/Engine/Private/LocalPlayer.cpp b/Engine/Source/Runtime/Engine/Private/LocalPlayer.cpp
--- a/Engine/Source/Runtime/Engine/Private/LocalPlayer.cpp
+++ b/Engine/Source/Runtime/Engine/Private/LocalPlayer.cpp
@@ -811,7 +811,7 @@ bool ULocalPlayer::CalcSceneViewInitOptions(
 		const UWorld* CurrentWorld = GetWorld();
 		const ERHIFeatureLevel::Type FeatureLevel = CurrentWorld ? CurrentWorld->FeatureLevel.GetValue() : GMaxRHIFeatureLevel;
 
-		ViewStates[ViewIndex].Allocate(FeatureLevel, CurrentWorld);
+		ViewStates[ViewIndex].Allocate(FeatureLevel);
 	}
 
 	ViewInitOptions.SceneViewStateInterface = ViewStates[ViewIndex].GetReference();
diff --git a/Engine/Source/Runtime/Engine/Private/SceneView.cpp b/Engine/Source/Runtime/Engine/Private/SceneView.cpp
--- a/Engine/Source/Runtime/Engine/Private/SceneView.cpp
+++ b/Engine/Source/Runtime/Engine/Private/SceneView.cpp
@@ -364,20 +364,12 @@ void FSceneViewStateReference::AllocateInternal(ERHIFeatureLevel::Type FeatureLe
 		ShareOriginInterface = ShareOriginTarget->Reference;
 	}
 
-	if (World && World->Scene)
-	{
-		Reference = World->Scene->AllocateViewState(ShareOriginInterface);
-	}
-	else
-	{
-		Reference = GetRendererModule().AllocateViewState(FeatureLevel, ShareOriginInterface);
-	}
+	Reference = GetRendererModule().AllocateViewState(FeatureLevel, ShareOriginInterface);
 }
 
-void FSceneViewStateReference::Allocate(ERHIFeatureLevel::Type FeatureLevel, const UWorld* InWorld)
+void FSceneViewStateReference::Allocate(ERHIFeatureLevel::Type FeatureLevel)
 {
 	check(!Reference);
-	World = InWorld;
 
 	AllocateInternal(FeatureLevel);
 
@@ -385,14 +377,9 @@ void FSceneViewStateReference::Allocate(ERHIFeatureLevel::Type FeatureLevel, con
 	GlobalListLink.LinkHead(GetSceneViewStateList());
 }
 
-void FSceneViewStateReference::Allocate(ERHIFeatureLevel::Type FeatureLevel)
-{
-	Allocate(FeatureLevel, nullptr);
-}
-
 void FSceneViewStateReference::Allocate()
 {
-	Allocate(GMaxRHIFeatureLevel, nullptr);
+	Allocate(GMaxRHIFeatureLevel);
 }
 
 ENGINE_API void FSceneViewStateReference::ShareOrigin(FSceneViewStateReference* Target)
diff --git a/Engine/Source/Runtime/Engine/Public/SceneInterface.h b/Engine/Source/Runtime/Engine/Public/SceneInterface.h
--- a/Engine/Source/Runtime/Engine/Public/SceneInterface.h
+++ b/Engine/Source/Runtime/Engine/Public/SceneInterface.h
@@ -571,8 +571,6 @@ public:
 	virtual bool RequestGPUSceneUpdate(FPrimitiveSceneInfo& PrimitiveSceneInfo, EPrimitiveDirtyState PrimitiveDirtyState) { return false; }
 
 protected:
-	virtual FSceneViewStateInterface* AllocateViewState(FSceneViewStateInterface* ShareOriginTarget = nullptr) = 0;
-
 	virtual ~FSceneInterface() {}
 
 	/** This scene's feature level */
diff --git a/Engine/Source/Runtime/Engine/Public/SceneManagement.h b/Engine/Source/Runtime/Engine/Public/SceneManagement.h
--- a/Engine/Source/Runtime/Engine/Public/SceneManagement.h
+++ b/Engine/Source/Runtime/Engine/Public/SceneManagement.h
@@ -256,9 +256,9 @@ public:
 
 	/**
 	* Adds a per-view virtual shadow map cache, which can help performance, at a cost in memory.  Does nothing if one is already present.
-	* Requires the View State to have been originally allocated with a UWorld with a valid Scene pointer.
+	* The cache only works for the first Scene this function is called for.
 	*/
-	virtual void AddVirtualShadowMapCache() {}
+	virtual void AddVirtualShadowMapCache(FSceneInterface* InScene) {}
 
 protected:
 	// Don't allow direct deletion of the view state, Destroy should be called instead.
@@ -273,7 +273,7 @@ private:
 	/** Reference counts the number of children parented to this state. */
 	int32							NumChildren;
 
-	virtual FVirtualShadowMapArrayCacheManager* GetVirtualShadowMapCache() const { return nullptr; }
+	virtual FVirtualShadowMapArrayCacheManager* GetVirtualShadowMapCache(const FScene* InScene) const { return nullptr; }
 	friend class FScene;
 };
 
diff --git a/Engine/Source/Runtime/Engine/Public/SceneTypes.h b/Engine/Source/Runtime/Engine/Public/SceneTypes.h
--- a/Engine/Source/Runtime/Engine/Public/SceneTypes.h
+++ b/Engine/Source/Runtime/Engine/Public/SceneTypes.h
@@ -92,19 +92,15 @@ class FSceneViewStateReference
 {
 public:
 	FSceneViewStateReference()
-	: World(nullptr), Reference(nullptr), ShareOriginTarget(nullptr), ShareOriginRefCount(0)
+	: Reference(nullptr), ShareOriginTarget(nullptr), ShareOriginRefCount(0)
 	{
 	}
 
 	ENGINE_API virtual ~FSceneViewStateReference();
 
 	/**
-	 * Allocates the Scene view state.  The first Allocate function accepts a UWorld, rather than an FScene, as the "RecompileRenderer" feature
-	 * involves destroying and reconstructing the FScene.  We need to pull the updated scene pointer from the UWorld in that case, but nothing
-	 * else in the UWorld is accessed.
+	 * Allocates the Scene view state.
 	 */
-	ENGINE_API void Allocate(ERHIFeatureLevel::Type FeatureLevel, const UWorld* World);
-
 	ENGINE_API void Allocate(ERHIFeatureLevel::Type FeatureLevel);
 
 	UE_DEPRECATED(5.0, "Allocate must be called with an appropriate RHI Feature Level")
@@ -136,7 +132,6 @@ public:
 	}
 
 private:
-	const UWorld* World;
 	FSceneViewStateInterface* Reference;
 	TLinkedList<FSceneViewStateReference*> GlobalListLink;
 
diff --git a/Engine/Source/Runtime/Renderer/Private/RendererScene.cpp b/Engine/Source/Runtime/Renderer/Private/RendererScene.cpp
--- a/Engine/Source/Runtime/Renderer/Private/RendererScene.cpp
+++ b/Engine/Source/Runtime/Renderer/Private/RendererScene.cpp
@@ -437,24 +437,6 @@ FSceneViewState::~FSceneViewState()
 	}
 }
 
-FSceneViewStateInterface* FScene::AllocateViewState(FSceneViewStateInterface* ShareOriginTarget)
-{
-	check(!ShareOriginTarget || (((FSceneViewState*)ShareOriginTarget)->Scene == this));
-
-	FSceneViewState* Result = new FSceneViewState(FeatureLevel, (FSceneViewState*)ShareOriginTarget);
-	Result->Scene = this;
-
-	// Modification of scene structure needs to happen on render thread
-	FScene* Scene = this;
-	ENQUEUE_RENDER_COMMAND(SceneViewStateAdd)(
-		[Scene, Result](FRHICommandList&)
-		{
-			Scene->ViewStates.Add(Result);
-		});
-
-	return Result;
-}
-
 void FScene::RemoveViewState_RenderThread(FSceneViewStateInterface* ViewState)
 {
 	for (int32 ViewStateIndex = 0; ViewStateIndex < ViewStates.Num(); ViewStateIndex++)
@@ -765,10 +747,22 @@ FFXSystemInterface* FScene::GetFXSystem()
 	return FXSystem;
 }
 
-void FSceneViewState::AddVirtualShadowMapCache()
+void FSceneViewState::AddVirtualShadowMapCache(FSceneInterface* InScene)
 {
-	// Can only allocate a virtual shadow map cache if we have a scene pointer available
-	if (Scene)
+	check(InScene);
+	if (!Scene)
+	{
+		Scene = (FScene*)InScene;
+
+		// Modification of scene structure needs to happen on render thread
+		ENQUEUE_RENDER_COMMAND(SceneViewStateAdd)(
+			[RenderScene = Scene, RenderViewState = this](FRHICommandList&)
+			{
+				RenderScene->ViewStates.Add(RenderViewState);
+			});
+	}
+
+	if (Scene == InScene)
 	{
 		// Don't add the cache if one has already been added.  We have a separate bool since the write to the cache
 		// pointer is deferred to the render thread.
@@ -788,9 +782,10 @@ void FSceneViewState::AddVirtualShadowMapCache()
 	}
 }
 
-FVirtualShadowMapArrayCacheManager* FSceneViewState::GetVirtualShadowMapCache() const
+FVirtualShadowMapArrayCacheManager* FSceneViewState::GetVirtualShadowMapCache(const FScene* InScene) const
 {
-	return ViewVirtualShadowMapCache;
+	// Per-view VSM cache can only be used for the Scene the view state was previously linked to
+	return Scene == InScene ? ViewVirtualShadowMapCache : nullptr;
 }
 
 FVirtualShadowMapArrayCacheManager* FScene::GetVirtualShadowMapCache(FSceneView& View) const
@@ -798,7 +793,7 @@ FVirtualShadowMapArrayCacheManager* FScene::GetVirtualShadowMapCache(FSceneView&
 	FVirtualShadowMapArrayCacheManager* Result = DefaultVirtualShadowMapCache;
 	if (View.State)
 	{
-		FVirtualShadowMapArrayCacheManager* ViewCache = View.State->GetVirtualShadowMapCache();
+		FVirtualShadowMapArrayCacheManager* ViewCache = View.State->GetVirtualShadowMapCache(this);
 		if (ViewCache)
 		{
 			Result = ViewCache;
@@ -816,7 +811,8 @@ void FScene::GetAllVirtualShadowMapCacheManagers(TArray<FVirtualShadowMapArrayCa
 	}
 	for (const FSceneViewState* ViewState : ViewStates)
 	{
-		if (ViewState->ViewVirtualShadowMapCache)
+		// Per-view VSM cache can only be used for the Scene the view state was previously linked to
+		if (ViewState->ViewVirtualShadowMapCache && ViewState->Scene == this)
 		{
 			OutCacheManagers.Add(ViewState->ViewVirtualShadowMapCache);
 		}
@@ -1382,7 +1378,20 @@ FScene::~FScene()
 	checkf(RemovedPrimitiveSceneInfos.Num() == 0, TEXT("All pending primitive removal operations are expected to be flushed when the scene is destroyed. Remaining operations are likely to cause a memory leak."));
 	checkf(AddedPrimitiveSceneInfos.Num() == 0, TEXT("All pending primitive addition operations are expected to be flushed when the scene is destroyed. Remaining operations are likely to cause a memory leak."));
 	checkf(Primitives.Num() == 0, TEXT("All primitives are expected to be removed before the scene is destroyed. Remaining primitives are likely to cause a memory leak."));
-	checkf(ViewStates.Num() == 0, TEXT("All scene view states are expected to be removed before the scene is destroyed."));
+
+	// Unlink any view states from the scene
+	for (FSceneViewState* ViewState : ViewStates)
+	{
+		check(ViewState->Scene == this);
+		ViewState->Scene = nullptr;
+
+		if (ViewState->ViewVirtualShadowMapCache)
+		{
+			delete ViewState->ViewVirtualShadowMapCache;
+			ViewState->ViewVirtualShadowMapCache = nullptr;
+		}
+	}
+	ViewStates.Empty();
 
 	// Delete default cache
 	if (DefaultVirtualShadowMapCache)
@@ -5537,12 +5546,6 @@ public:
 
 	virtual bool HasAnyLights() const override { return false; }
 
-protected:
-	virtual FSceneViewStateInterface* AllocateViewState(FSceneViewStateInterface* ShareOriginTarget) override
-	{
-		return new FSceneViewState(FeatureLevel, (FSceneViewState*)ShareOriginTarget);
-	}
-
 private:
 	UWorld* World;
 	class FFXSystemInterface* FXSystem;
diff --git a/Engine/Source/Runtime/Renderer/Private/ScenePrivate.h b/Engine/Source/Runtime/Renderer/Private/ScenePrivate.h
--- a/Engine/Source/Runtime/Renderer/Private/ScenePrivate.h
+++ b/Engine/Source/Runtime/Renderer/Private/ScenePrivate.h
@@ -682,12 +682,10 @@ public:
 	uint32 UniqueID;
 
 	/**
-	 * The scene pointer may be NULL -- it's only filled in if a UWorld with a valid Scene pointer was passed into the original call
-	 * to FSceneViewStateReference::Allocate.  Currently, the Scene pointer in the class is only used for constructing and cleaning up
-	 * the ViewVirtualShadowMapCache member, if the user chooses to enable per-view shadow caching via AddVirtualShadowMapCache().
-	 * If the Scene isn't present, AddVirtualShadowMapCache() has no effect.
-	 * 
-	 * TODO:  deprecate support for allocating an FSceneViewState without a valid Scene...
+	 * The scene pointer may be NULL -- it's filled in by certain API calls that require a FSceneViewState and FScene to know about each other,
+	 * such as FSceneViewState::AddVirtualShadowMapCache.  Whenever a ViewState and Scene get linked, this pointer is set, and a pointer
+	 * to the ViewState is added to an array in the Scene.  The linking is necessary in cases where incremental FScene updates need to be
+	 * reflected in cached data stored in FSceneViewState.
 	 */
 	FScene* Scene;
 
@@ -1486,8 +1484,8 @@ public:
 		return SequencerState;
 	}
 
-	virtual void AddVirtualShadowMapCache() override;
-	virtual FVirtualShadowMapArrayCacheManager* GetVirtualShadowMapCache() const override;
+	virtual void AddVirtualShadowMapCache(FSceneInterface* InScene) override;
+	virtual FVirtualShadowMapArrayCacheManager* GetVirtualShadowMapCache(const FScene* InScene) const override;
 
 	/** Information about visibility/occlusion states in past frames for individual primitives. */
 	TSet<FPrimitiveOcclusionHistory,FPrimitiveOcclusionHistoryKeyFuncs> PrimitiveOcclusionHistorySet;
@@ -3365,10 +3363,6 @@ public:
 protected:
 
 private:
-	/**
-	 * Internal view state allocation interface, used by FSceneViewState
-	 */
-	virtual FSceneViewStateInterface* AllocateViewState(FSceneViewStateInterface* ShareOriginTarget) override;
 	void RemoveViewState_RenderThread(FSceneViewStateInterface*);
 
 	/**
