Fix incorrect assert when r.AOGlobalDistanceFieldStaggeredUpdates is disabled.

diff --git a/Engine/Source/Runtime/Renderer/Private/GlobalDistanceField.cpp b/Engine/Source/Runtime/Renderer/Private/GlobalDistanceField.cpp
--- a/Engine/Source/Runtime/Renderer/Private/GlobalDistanceField.cpp
+++ b/Engine/Source/Runtime/Renderer/Private/GlobalDistanceField.cpp
@@ -1124,7 +1124,7 @@ static void ComputeUpdateRegionsAndUpdateViewState(
 			ClipmapViewState.CacheMostlyStaticSeparately = GAOGlobalDistanceFieldCacheMostlyStaticSeparately;
 		}
 
-		ensureMsgf(GAOGlobalDistanceFieldStaggeredUpdates || NumClipmapUpdateRequests <= GetNumClipmapUpdatesPerFrame(), TEXT("ShouldUpdateClipmapThisFrame needs to be adjusted for the NumClipmaps to even out the work distribution"));
+		ensureMsgf(!GAOGlobalDistanceFieldStaggeredUpdates || NumClipmapUpdateRequests <= GetNumClipmapUpdatesPerFrame(), TEXT("ShouldUpdateClipmapThisFrame needs to be adjusted for the NumClipmaps to even out the work distribution"));
 	}
 	else
 	{
