MGPU:  Fix for global distance field data not being initialized when swapping GPUs.  Force full update of global distance field data if the GPU index changes.  If AFR is enabled, force global distance fields to be updated across all GPUs, to avoid thrashing due to the GPU index changing every frame.

#rnx

diff --git a/Engine/Source/Runtime/Renderer/Private/DistanceFieldObjectManagement.cpp b/Engine/Source/Runtime/Renderer/Private/DistanceFieldObjectManagement.cpp
--- a/Engine/Source/Runtime/Renderer/Private/DistanceFieldObjectManagement.cpp
+++ b/Engine/Source/Runtime/Renderer/Private/DistanceFieldObjectManagement.cpp
@@ -1072,14 +1072,15 @@ void FSceneRenderer::PrepareDistanceFieldScene(FRDGBuilder& GraphBuilder, bool b
 		{
 			GraphBuilder.AddDispatchHint();
 		}
-		for (int32 ViewIndex = 0; ViewIndex < Views.Num(); ViewIndex++)
+
+		if (ShouldPrepareGlobalDistanceField())
 		{
-			FViewInfo& View = Views[ViewIndex];
+			for (int32 ViewIndex = 0; ViewIndex < Views.Num(); ViewIndex++)
+			{
+				FViewInfo& View = Views[ViewIndex];
 
-			RDG_GPU_MASK_SCOPE(GraphBuilder, View.GPUMask);
+				RDG_GPU_MASK_SCOPE(GraphBuilder, GNumAlternateFrameRenderingGroups > 1 ? FRHIGPUMask::All() : View.GPUMask);
 
-			if (ShouldPrepareGlobalDistanceField())
-			{
 				float OcclusionMaxDistance = Scene->DefaultMaxDistanceFieldOcclusionDistance;
 
 				// Use the skylight's max distance if there is one
@@ -1088,7 +1089,7 @@ void FSceneRenderer::PrepareDistanceFieldScene(FRDGBuilder& GraphBuilder, bool b
 					OcclusionMaxDistance = Scene->SkyLight->OcclusionMaxDistance;
 				}
 
-				UpdateGlobalDistanceFieldVolume(GraphBuilder, Views[ViewIndex], Scene, OcclusionMaxDistance, IsLumenEnabled(View), Views[ViewIndex].GlobalDistanceFieldInfo);
+				UpdateGlobalDistanceFieldVolume(GraphBuilder, View, Scene, OcclusionMaxDistance, IsLumenEnabled(View), View.GlobalDistanceFieldInfo);
 			}
 		}
 		if (!bSplitDispatch)
diff --git a/Engine/Source/Runtime/Renderer/Private/GlobalDistanceField.cpp b/Engine/Source/Runtime/Renderer/Private/GlobalDistanceField.cpp
--- a/Engine/Source/Runtime/Renderer/Private/GlobalDistanceField.cpp
+++ b/Engine/Source/Runtime/Renderer/Private/GlobalDistanceField.cpp
@@ -753,6 +753,16 @@ static void ComputeUpdateRegionsAndUpdateViewState(
 			}
 		}
 
+		// Check if GPU mask changed
+		bool bForceFullUpdateMGPU = false;
+#if WITH_MGPU
+		if (GlobalDistanceFieldData.LastGPUMask != RHICmdList.GetGPUMask())
+		{
+			bForceFullUpdateMGPU = true;
+			GlobalDistanceFieldData.LastGPUMask = RHICmdList.GetGPUMask();
+		}
+#endif
+
 		int32 NumClipmapUpdateRequests = 0;
 
 		FViewElementPDI ViewPDI(&View, nullptr, &View.DynamicPrimitiveCollector);
@@ -959,6 +969,7 @@ static void ComputeUpdateRegionsAndUpdateViewState(
 			}
 
 			const bool bForceFullUpdate = bSharedDataReallocated
+				|| bForceFullUpdateMGPU
 				|| !GlobalDistanceFieldData.bInitializedOrigins
 				// Detect when max occlusion distance has changed
 				|| ClipmapViewState.CachedClipmapExtent != ClipmapExtent
diff --git a/Engine/Source/Runtime/Renderer/Private/ScenePrivate.h b/Engine/Source/Runtime/Renderer/Private/ScenePrivate.h
--- a/Engine/Source/Runtime/Renderer/Private/ScenePrivate.h
+++ b/Engine/Source/Runtime/Renderer/Private/ScenePrivate.h
@@ -623,6 +623,9 @@ struct FPersistentGlobalDistanceFieldData : public FThreadSafeRefCountedObject
 	FVector	CameraVelocityOffset = FVector(0);
 	bool	bUpdateViewOrigin = true;
 	FVector	LastViewOrigin = FVector(0);
+#if WITH_MGPU
+	FRHIGPUMask LastGPUMask;
+#endif
 
 	TRefCountPtr<FRDGPooledBuffer> PageFreeListAllocatorBuffer;
 	TRefCountPtr<FRDGPooledBuffer> PageFreeListBuffer;
