PR #9166: Ray tracing light function implementation (Contributed by [REDACTED_PERSON])

Create raytracing miss shaders that implement light functions for raytraced reflections, refractions and the raytracing debug primary view. Each light with a valid light function material is assigned a miss shader in the main raytracing pipeline which is used to modulate the light intensity as part of the shadowing calculation. The light function implementation intentionally converts to grayscale for visual compatibility with the deferred renderer.

As part of this CL, the default behavior for miss shaders in RT pipeline is changed. We now require the caller to use RHISetRayTracingMissShader before the first dispatch to ensure the desired shader is assigned to the desired slot in the pipeline. This relaxes some prior assumptions about the shader slot index needing to matching the shader index in the pipeline (which would not hold after having added more miss shaders to the main raytracing pipeline).

Basic validation has been added to the D3D12 RHI to ensure we are correctly binding at least the 0'th miss shader slot prior to the first dispatch.

diff --git a/Engine/Plugins/Experimental/GPULightmass/Source/GPULightmass/Private/LightmapRenderer.cpp b/Engine/Plugins/Experimental/GPULightmass/Source/GPULightmass/Private/LightmapRenderer.cpp
--- a/Engine/Plugins/Experimental/GPULightmass/Source/GPULightmass/Private/LightmapRenderer.cpp
+++ b/Engine/Plugins/Experimental/GPULightmass/Source/GPULightmass/Private/LightmapRenderer.cpp
@@ -1201,7 +1201,8 @@ bool FSceneRenderState::SetupRayTracingScene(int32 LODIndex)
 			FShaderMapResource::GetRayTracingHitGroupLibrary(RayTracingHitGroupLibrary, DefaultClosestHitShader);
 
 			PSOInitializer.SetHitGroupTable(RayTracingHitGroupLibrary);
-
+			
+			// TODO(UE-157946): This pipeline does not bind any miss shader and relies on the pipeline to do this automatically. This should be made explicit.
 			RayTracingPipelineState = PipelineStateCache::GetAndOrCreateRayTracingPipelineState(RHICmdList, PSOInitializer);
 
 			TUniquePtr<FRayTracingLocalShaderBindingWriter> BindingWriter = MakeUnique<FRayTracingLocalShaderBindingWriter>();
@@ -1278,6 +1279,7 @@ bool FSceneRenderState::SetupRayTracingScene(int32 LODIndex)
 						NumTotalBindings, MergedBindings,
 						bCopyDataToInlineStorage);
 				}
+				RHICmdList.SetRayTracingMissShader(RayTracingScene, 0, RayTracingPipelineState, 0 /* ShaderIndexInPipeline */, 0, nullptr, 0);
 
 				// Move the ray tracing binding container ownership to the command list, so that memory will be
 				// released on the RHI thread timeline, after the commands that reference it are processed.
diff --git a/Engine/Plugins/FX/Niagara/Source/NiagaraShader/Private/NiagaraAsyncGpuTraceProviderHwrt.cpp b/Engine/Plugins/FX/Niagara/Source/NiagaraShader/Private/NiagaraAsyncGpuTraceProviderHwrt.cpp
--- a/Engine/Plugins/FX/Niagara/Source/NiagaraShader/Private/NiagaraAsyncGpuTraceProviderHwrt.cpp
+++ b/Engine/Plugins/FX/Niagara/Source/NiagaraShader/Private/NiagaraAsyncGpuTraceProviderHwrt.cpp
@@ -244,6 +244,7 @@ static void BindNiagaraRayTracingMeshCommands(
 		NumTotalBindings,
 		Bindings,
 		bCopyDataToInlineStorage);
+	RHICmdList.SetRayTracingMissShader(RayTracingScene, 0, Pipeline, 0 /* ShaderIndexInPipeline */, 0, nullptr, 0);
 }
 
 ////////////////////////////////////////////////////////////////////////////////
diff --git a/Engine/Shaders/Private/LightFunctionCommon.ush b/Engine/Shaders/Private/LightFunctionCommon.ush
--- a/Engine/Shaders/Private/LightFunctionCommon.ush
+++ b/Engine/Shaders/Private/LightFunctionCommon.ush
@@ -6,9 +6,12 @@
 
 #include "/Engine/Private/Strata/Strata.ush" 
 
+// This ifdef is to allow shaders(notably the ratracing miss shader implementing light functions) to override where these parameters come from.
+#ifndef LightFunctionParameters
 /** Tan of spotlight cone outer angle in x, ShadowFadeFraction in y, IsSpotLight in z, IsPointLight in w. */
 float4 LightFunctionParameters;
 float4x4 LightFunctionTranslatedWorldToLight;
+#endif
 
 /** 
  * Calculates the light function color with the given light vector. 
diff --git a/Engine/Shaders/Private/RayTracing/RayTracingLightFunctionCommon.ush b/Engine/Shaders/Private/RayTracing/RayTracingLightFunctionCommon.ush
new file mode 100644
--- /dev/null
+++ b/Engine/Shaders/Private/RayTracing/RayTracingLightFunctionCommon.ush
@@ -0,0 +1,31 @@
+// Copyright 1998-2019 Epic Games, Inc. All Rights Reserved.
+
+
+// base light function material shader, so it can be folded where apprpriate
+// see LightFunctionCommon and LightFunctionPixelShader for additional information
+float3 GetRayTracingLightFunction(float3 TranslatedWorldPosition)
+{
+	float3 LightVector;
+
+	{
+		// Not really translated, just dealing with the fact that common code uses this convention
+		float4 Hom = mul(float4(TranslatedWorldPosition, 1), LightFunctionTranslatedWorldToLight);
+		LightVector = Hom.xyz / Hom.w;
+	}
+
+	// Calculate radial view distance for stable fading
+	float ViewDistance = length(View.TranslatedWorldCameraOrigin - TranslatedWorldPosition);
+
+	float3 Color = GetLightFunctionColor(LightVector, TranslatedWorldPosition);
+
+	float DistanceFadeAlpha = saturate((RaytracingLightFunctionParameters.LightFunctionParameters2.x - ViewDistance) / (RaytracingLightFunctionParameters.LightFunctionParameters2.x * .2f));
+	// Fade to disabled based on LightFunctionFadeDistance
+	Color = lerp(RaytracingLightFunctionParameters.LightFunctionParameters2.y, Color, DistanceFadeAlpha);
+
+	// Fade to disabled based on ShadowFadeFraction
+	Color = lerp(RaytracingLightFunctionParameters.LightFunctionParameters2.y, Color, LightFunctionParameters.y);
+
+	return Color;
+}
+
+
diff --git a/Engine/Shaders/Private/RayTracing/RayTracingLightingCommon.ush b/Engine/Shaders/Private/RayTracing/RayTracingLightingCommon.ush
--- a/Engine/Shaders/Private/RayTracing/RayTracingLightingCommon.ush
+++ b/Engine/Shaders/Private/RayTracing/RayTracingLightingCommon.ush
@@ -32,7 +32,7 @@ struct FRTLightingData
 	uint Type;
 	int LightProfileIndex;
 	float RectLightAtlasMaxLevel;
-	uint Pad;
+	uint LightMissShaderIndex;
 	float3 TranslatedLightPosition;
 	float InvRadius;
 	float3 Direction;
@@ -70,7 +70,7 @@ FRTLightingData GetRayTracingLightData(int LightIndex)
 	Data.Type = Element.x;
 	Data.LightProfileIndex = Element.y;
 	Data.RectLightAtlasMaxLevel = asfloat(Element.z);
-	Data.Pad = Element.w;
+	Data.LightMissShaderIndex = Element.w;
 
 	Element = LightDataBuffer[LightIndex++];
 
@@ -113,7 +113,8 @@ FRTLightingData GetRayTracingLightData(int LightIndex)
 
 FDeferredLightData GetRayTracingDeferredLightData(int LightIndex,
 	inout int OutLightProfileIndex,
-	inout uint OutLightType)
+	inout uint OutLightType,
+	inout uint OutLightMissShaderIndex)
 {
 	FDeferredLightData LightData = (FDeferredLightData)0;
 	FRTLightingData RayTracingLightData = GetRayTracingLightData(LightIndex);
@@ -160,10 +161,19 @@ FDeferredLightData GetRayTracingDeferredLightData(int LightIndex,
 
 	OutLightProfileIndex = RayTracingLightData.LightProfileIndex;
 	OutLightType = LightType;
+	OutLightMissShaderIndex = RayTracingLightData.LightMissShaderIndex;
 
 	return LightData;
 }
 
+FDeferredLightData GetRayTracingDeferredLightData(int LightIndex,
+	inout int OutLightProfileIndex,
+	inout uint OutLightType)
+{
+	uint LightMissShaderIndex = 0;
+	return GetRayTracingDeferredLightData(LightIndex, OutLightProfileIndex, OutLightType, LightMissShaderIndex);
+}
+
 FRectTexture GetRayTracingRectTextureData(FDeferredLightData In)
 {
 	FRectTexture Out;
@@ -415,7 +425,8 @@ float3 ComputeDirectLighting(
 
 		int LightProfileIndex = -1;
 		uint LightType = 0;
-		FDeferredLightData LightData = GetRayTracingDeferredLightData(LightIndex, LightProfileIndex, LightType);
+		uint LightMissShaderIndex = 0;
+		FDeferredLightData LightData = GetRayTracingDeferredLightData(LightIndex, LightProfileIndex, LightType, LightMissShaderIndex);
 
 		float3 ShadowRayDirection;
 		// ToLight should not be normalized because its length is used to compute the shadow ray TMax 
@@ -492,7 +503,7 @@ float3 ComputeDirectLighting(
 		// use the bound miss shader lighting evaluation if lighting, else nothing
 		if (bMissFarField)
 		{
-			uint MissShaderIndex = Lit ? RAY_TRACING_MISS_SHADER_SLOT_LIGHTING : RAY_TRACING_MISS_SHADER_SLOT_DEFAULT;
+			uint MissShaderIndex = Lit ? LightMissShaderIndex : RAY_TRACING_MISS_SHADER_SLOT_DEFAULT;
 			TraceShadowRayMissShaderLighting(ShadowRay, RayFlags, InstanceInclusionMask, TLAS, MissShaderIndex, PixelCoord, Payload);
 		}
 	}
diff --git a/Engine/Shaders/Private/RayTracing/RayTracingLightingMS.usf b/Engine/Shaders/Private/RayTracing/RayTracingLightingMS.usf
--- a/Engine/Shaders/Private/RayTracing/RayTracingLightingMS.usf
+++ b/Engine/Shaders/Private/RayTracing/RayTracingLightingMS.usf
@@ -13,6 +13,10 @@
 #define STRATA_SSS_MATERIAL_OVERRIDE 0
 #define MATERIAL_FULLY_ROUGH 0
 
+#ifndef SUPPORT_LIGHT_FUNCTION
+#define SUPPORT_LIGHT_FUNCTION 0
+#endif
+
 #include "../SceneTextureParameters.ush"
 #include "RayTracingCommon.ush"
 
@@ -30,6 +34,17 @@
 #include "../Strata/StrataDeferredLighting.ush"
 #endif
 
+#if SUPPORT_LIGHT_FUNCTION
+#include "/Engine/Generated/Material.ush"
+
+// TODO: if we could just bind loose parameters, we could avoid this workaround for getting the shader parameters needed by the routines in LightFunctionCommon.ush
+#define LightFunctionParameters RaytracingLightFunctionParameters.LightFunctionParameters
+#define LightFunctionTranslatedWorldToLight RaytracingLightFunctionParameters.LightFunctionTranslatedWorldToLight
+
+#include "../LightFunctionCommon.ush"
+#include "RayTracingLightFunctionCommon.ush"
+#endif
+
 RAY_TRACING_ENTRY_MISS(RayTracingLightingMS, FPackedMaterialClosestHitPayload, PackedPayload)
 {
 	RayDesc DummyRay = (RayDesc)0;
@@ -104,6 +119,11 @@ RAY_TRACING_ENTRY_MISS(RayTracingLightingMS, FPackedMaterialClosestHitPayload, P
 	float3 LightContribution = GetDynamicLighting(WorldPosition, CameraVector, GBufferData, 1.0f, GBufferData.ShadingModelID, LightData, LightAttenuation, Dither, SVPos, SurfaceShadow).xyz;
 #endif // STRATA_ENABLED
 
+#if SUPPORT_LIGHT_FUNCTION
+	float3 LightFunctionMult = GetRayTracingLightFunction(WorldPosition);
+	LightProfileMultiplier *= dot(LightFunctionMult, .3333f); // convert to greyscale (for compatibility with the appearance in the primary view)
+#endif
+
 	float3 AccumulatedRadiance = Payload.Radiance + LightContribution * LightProfileMultiplier;
 
 	PackedPayload.SetRadiance(AccumulatedRadiance);
diff --git a/Engine/Source/Runtime/D3D12RHI/Private/D3D12RayTracing.cpp b/Engine/Source/Runtime/D3D12RHI/Private/D3D12RayTracing.cpp
--- a/Engine/Source/Runtime/D3D12RHI/Private/D3D12RayTracing.cpp
+++ b/Engine/Source/Runtime/D3D12RHI/Private/D3D12RayTracing.cpp
@@ -1779,7 +1779,7 @@ public:
 	#endif // USE_STATIC_ROOT_SIGNATURE
 	}
 
-	void Init(const FInitializer& Initializer, FD3D12Device* Device)
+	void Init(const FInitializer& Initializer, FD3D12Device* Device, const TArrayView<const FD3D12ShaderIdentifier>& RaygenIdentifiers, const FD3D12ShaderIdentifier& DefaultHitGroupIdentifier)
 	{
 		checkf(Initializer.LocalRootDataSize <= 4096, TEXT("The maximum size of a local root signature is 4KB.")); // as per section 4.22.1 of DXR spec v1.0
 		checkf(Initializer.NumRayGenShaders >= 1, TEXT("All shader tables must contain at least one raygen shader."));
@@ -1844,6 +1844,13 @@ public:
 		NumLocalRecords = (TotalDataSize - LocalShaderTableOffset) / LocalRecordStride;
 
 		Data.SetNumZeroed(TotalDataSize);
+#if DO_CHECK
+		bWasDefaultMissShaderSet = false;
+#endif
+		SetRayGenIdentifiers(RaygenIdentifiers);
+		SetDefaultHitGroupIdentifier(DefaultHitGroupIdentifier);
+		SetDefaultMissShaderIdentifier(FD3D12ShaderIdentifier::Null);
+		SetDefaultCallableShaderIdentifier(FD3D12ShaderIdentifier::Null);
 
 		// Keep CPU-side data after upload
 		Data.SetAllowCPUAccess(true);
@@ -1892,6 +1899,12 @@ public:
 	void SetMissIdentifier(uint32 RecordIndex, const FD3D12ShaderIdentifier& ShaderIdentifier)
 	{
 		const uint32 WriteOffset = MissShaderTableOffset + RecordIndex * LocalRecordStride;
+#if DO_CHECK
+		if (RecordIndex == 0)
+		{
+			bWasDefaultMissShaderSet = true;
+		}
+#endif
 		WriteData(WriteOffset, ShaderIdentifier.Data, ShaderIdentifierSize);
 	}
 
@@ -1924,22 +1937,20 @@ public:
 		}
 	}
 
-	void SetMissIdentifiers(const TArrayView<const FD3D12ShaderIdentifier>& Identifiers)
+	void SetDefaultMissShaderIdentifier(const FD3D12ShaderIdentifier& ShaderIdentifier)
 	{
-		// Set all to the default
-		check(Identifiers.Num());
+		// Set all slots to the same default
 		for (uint32 Index = 0; Index < NumMissRecords; ++Index)
 		{
-			SetMissIdentifier(Index, Identifiers[0]);
+			SetLocalShaderIdentifier(MissShaderRecordIndexOffset + Index, ShaderIdentifier);
 		}
 	}
 
-	void SetLocalShaderIdentifiers(uint32 RecordIndexOffset, const TArrayView<const FD3D12ShaderIdentifier>& Identifiers)
+	void SetDefaultCallableShaderIdentifier(const FD3D12ShaderIdentifier& ShaderIdentifier)
 	{
-		check(Identifiers.Num() == NumCallableRecords);
-		for (int32 Index = 0; Index < Identifiers.Num(); ++Index)
+		for (uint32 Index = 0; Index < NumCallableRecords; ++Index)
 		{
-			SetLocalShaderIdentifier(RecordIndexOffset + Index, Identifiers[Index]);
+			SetLocalShaderIdentifier(CallableShaderRecordIndexOffset + Index, ShaderIdentifier);
 		}
 	}
 
@@ -1951,6 +1962,8 @@ public:
 
 		checkf(Data.Num(), TEXT("Shader table is expected to be initialized before copying to GPU."));
 
+		checkf(bWasDefaultMissShaderSet, TEXT("At least the first miss shader must have been set before copying to GPU."));
+
 		FD3D12Device* Device = CommandContext.GetParentDevice();
 		FD3D12Adapter* Adapter = Device->GetParentAdapter();
 		D3D12_RESOURCE_DESC BufferDesc = CD3DX12_RESOURCE_DESC::Buffer(Data.GetResourceDataSize(), D3D12_RESOURCE_FLAG_NONE, D3D12_RAYTRACING_SHADER_TABLE_BYTE_ALIGNMENT);
@@ -2019,7 +2032,6 @@ public:
 	static constexpr uint32 ShaderIdentifierSize = D3D12_SHADER_IDENTIFIER_SIZE_IN_BYTES;
 	uint32 NumHitRecords = 0;
 	uint32 NumRayGenShaders = 0;
-	uint32 NumMissShaders = 0;
 	uint32 NumCallableRecords = 0;
 	uint32 NumMissRecords = 0;
 	uint32 NumLocalRecords = 0;
@@ -2045,6 +2057,9 @@ public:
 
 	bool bIsDirty = true;
 	TRefCountPtr<FD3D12Buffer> Buffer;
+#if DO_CHECK
+	bool bWasDefaultMissShaderSet = false;
+#endif
 
 	// SBTs have their own descriptor heaps
 	FD3D12RayTracingDescriptorCache* DescriptorCache = nullptr;
@@ -2525,16 +2540,6 @@ public:
 			MissShaders.Shaders.Add(Shader);
 		}
 
-		// Ensure miss shader 0 (default) requires no parameters
-		if (!Initializer.bPartial)
-		{
-			FD3D12RayTracingShader* Shader = FD3D12DynamicRHI::ResourceCast(InitializerMissShaders[0]);
-
-			const uint32 ShaderViewDescriptors = Shader->ResourceCounts.NumSRVs + Shader->ResourceCounts.NumUAVs;
-			
-			checkf(ShaderViewDescriptors == 0, TEXT("Default miss shader (miss shader 0) must take no parameters"));
-		}
-
 		// Add hit groups
 
 		TArray<FD3D12RayTracingPipelineCache::FEntry*> HitGroupEntries;
@@ -2735,10 +2740,11 @@ public:
 		for (uint32 GPUIndex : FRHIGPUMask::All())
 		{
 			FD3D12Device* NodeDevice = Adapter->GetDevice(GPUIndex);
-			DefaultShaderTables[GPUIndex].Init(SBTInitializer, NodeDevice);
-			DefaultShaderTables[GPUIndex].SetRayGenIdentifiers(RayGenShaders.Identifiers);
-			DefaultShaderTables[GPUIndex].SetMissIdentifiers(MissShaders.Identifiers);
-			DefaultShaderTables[GPUIndex].SetDefaultHitGroupIdentifier(HitGroupShaders.Identifiers[0]);
+			DefaultShaderTables[GPUIndex].Init(
+				SBTInitializer,
+				NodeDevice,
+				RayGenShaders.Identifiers,
+				HitGroupShaders.Identifiers[0]);
 		}
 
 		PipelineStackSize = PipelineProperties->GetPipelineStackSize();
@@ -4323,10 +4329,7 @@ FD3D12RayTracingShaderTable* FD3D12RayTracingScene::FindOrCreateShaderTable(cons
 	SBTInitializer.LocalRootDataSize = Pipeline->MaxLocalRootSignatureSize;
 	SBTInitializer.MaxViewDescriptorsPerRecord = Pipeline->MaxHitGroupViewDescriptors;
 
-	CreatedShaderTable->Init(SBTInitializer, Device);
-	CreatedShaderTable->SetRayGenIdentifiers(Pipeline->RayGenShaders.Identifiers);
-	CreatedShaderTable->SetMissIdentifiers(Pipeline->MissShaders.Identifiers);
-	CreatedShaderTable->SetDefaultHitGroupIdentifier(Pipeline->HitGroupShaders.Identifiers[0]);
+	CreatedShaderTable->Init(SBTInitializer, Device, Pipeline->RayGenShaders.Identifiers, Pipeline->HitGroupShaders.Identifiers[0]);
 
 	ShaderTables[GPUIndex].Add(Pipeline, CreatedShaderTable);
 
@@ -5597,7 +5600,7 @@ static void SetRayTracingMissShader(
 		LooseParameterDataSize, LooseParameterData, // Loose parameters
 		ResourceBinder);
 
-	ShaderTable->SetLocalShaderIdentifier(RecordIndex,
+	ShaderTable->SetMissIdentifier(ShaderSlotInScene,
 		bResourcesBound
 		? Pipeline->MissShaders.Identifiers[ShaderIndexInPipeline]
 		: FD3D12ShaderIdentifier::Null);
@@ -5733,9 +5736,6 @@ void FD3D12CommandContext::RHISetRayTracingCallableShader(
 {
 	FD3D12RayTracingScene* Scene = FD3D12DynamicRHI::ResourceCast(InScene);
 	FD3D12RayTracingPipelineState* Pipeline = FD3D12DynamicRHI::ResourceCast(InPipeline);
-
-	checkf(ShaderSlotInScene < Scene->Initializer.NumCallableShaderSlots, TEXT("Shader slot is invalid. Make sure that NumCallableShaderSlots is correct on FRayTracingSceneInitializer."));
-
 	FD3D12RayTracingShaderTable* ShaderTable = Scene->FindOrCreateShaderTable(Pipeline, GetParentDevice());
 	const uint32 WorkerIndex = 0;
 
@@ -5750,33 +5750,13 @@ void FD3D12CommandContext::RHISetRayTracingMissShader(
 {
 	FD3D12RayTracingScene* Scene = FD3D12DynamicRHI::ResourceCast(InScene);
 	FD3D12RayTracingPipelineState* Pipeline = FD3D12DynamicRHI::ResourceCast(InPipeline);
-
-	checkf(ShaderSlotInScene < Scene->Initializer.NumMissShaderSlots, TEXT("Shader slot is invalid. Make sure that NumMissShaderSlots is correct on FRayTracingSceneInitializer."));
-	checkf(ShaderSlotInScene > 0, TEXT("Shader slot 0 is reserved for the default MissShader."));
-
 	FD3D12RayTracingShaderTable* ShaderTable = Scene->FindOrCreateShaderTable(Pipeline, GetParentDevice());
-
-	const uint32 RecordIndex = ShaderTable->MissShaderRecordIndexOffset + ShaderSlotInScene;
-
-	const uint32 UserDataOffset = offsetof(FHitGroupSystemParameters, RootConstants) + offsetof(FHitGroupSystemRootConstants, UserData);
-	ShaderTable->SetLocalShaderParameters(RecordIndex, UserDataOffset, UserData);
-
-	const FD3D12RayTracingShader* Shader = Pipeline->MissShaders.Shaders[ShaderIndexInPipeline];
-
-	uint32 WorkerIndex = 0;
-	FD3D12RayTracingLocalResourceBinder ResourceBinder(*GetParentDevice(), *ShaderTable, *(Shader->pRootSignature), RecordIndex, WorkerIndex);
-	const bool bResourcesBound = SetRayTracingShaderResources(Shader,
-		0, nullptr, // Textures
-		0, nullptr, // SRVs
+	const uint32 WorkerIndex = 0;
+	SetRayTracingMissShader(GetParentDevice(), ShaderTable, Scene, Pipeline,
+		ShaderSlotInScene, ShaderIndexInPipeline,
 		NumUniformBuffers, UniformBuffers,
-		0, nullptr, // Samplers
-		0, nullptr, // UAVs
 		0, nullptr, // Loose parameters
-		ResourceBinder);
-
-	ShaderTable->SetLocalShaderIdentifier(RecordIndex,
-		bResourcesBound
-		? Pipeline->MissShaders.Identifiers[ShaderIndexInPipeline]
-		: FD3D12ShaderIdentifier::Null);
+		UserData,
+		WorkerIndex);
 }
 #endif // D3D12_RHI_RAYTRACING
diff --git a/Engine/Source/Runtime/Renderer/Private/DeferredShadingRenderer.cpp b/Engine/Source/Runtime/Renderer/Private/DeferredShadingRenderer.cpp
--- a/Engine/Source/Runtime/Renderer/Private/DeferredShadingRenderer.cpp
+++ b/Engine/Source/Runtime/Renderer/Private/DeferredShadingRenderer.cpp
@@ -1617,7 +1617,7 @@ bool FDeferredShadingSceneRenderer::SetupRayTracingPipelineStates(FRDGBuilder& G
 	return true;
 }
 
-void FDeferredShadingSceneRenderer::SetupRayTracingLightDataForViews(FRDGBuilder& GraphBuilder)
+void FDeferredShadingSceneRenderer::SetupRayTracingLightDataForViews(FRDGBuilder& GraphBuilder, const FRayTracingLightFunctionMap& RayTracingLightFunctionMap)
 {
 	if (!bAnyRayTracingPassEnabled)
 	{
@@ -1640,7 +1640,7 @@ void FDeferredShadingSceneRenderer::SetupRayTracingLightDataForViews(FRDGBuilder
 		else
 		{
 			// This light data is a function of the camera position, so must be computed per view.
-			View.RayTracingLightDataUniformBuffer = CreateRayTracingLightData(GraphBuilder, Scene->Lights, View, View.ShaderMap, NumOfSkippedRayTracingLights);
+			View.RayTracingLightDataUniformBuffer = CreateRayTracingLightData(GraphBuilder, Scene, View, RayTracingLightFunctionMap, View.ShaderMap, NumOfSkippedRayTracingLights);
 		}
 	}
 
@@ -1795,7 +1795,7 @@ static void ReleaseRaytracingResources(FRDGBuilder& GraphBuilder, TArrayView<FVi
 	});
 }
 
-void FDeferredShadingSceneRenderer::WaitForRayTracingScene(FRDGBuilder& GraphBuilder, FRDGBufferRef DynamicGeometryScratchBuffer)
+void FDeferredShadingSceneRenderer::WaitForRayTracingScene(FRDGBuilder& GraphBuilder, FRDGBufferRef DynamicGeometryScratchBuffer, FRayTracingLightFunctionMap& RayTracingLightFunctionMap)
 {
 	if (!bAnyRayTracingPassEnabled)
 	{
@@ -1831,7 +1831,7 @@ void FDeferredShadingSceneRenderer::WaitForRayTracingScene(FRDGBuilder& GraphBui
 	PassParams->LightDataPacked = bIsPathTracing ? nullptr : ReferenceView.RayTracingLightDataUniformBuffer; // accessed by FRayTracingLightingMS
 
 	GraphBuilder.AddPass(RDG_EVENT_NAME("WaitForRayTracingScene"), PassParams, ERDGPassFlags::Compute | ERDGPassFlags::NeverCull,
-		[this, PassParams, bIsPathTracing, &ReferenceView, bAnyLumenHardwareInlineRayTracingPassEnabled](FRHICommandListImmediate& RHICmdList)
+		[this, PassParams, bIsPathTracing, &ReferenceView, bAnyLumenHardwareInlineRayTracingPassEnabled, RayTracingLightFunctionMap = MoveTemp(RayTracingLightFunctionMap)](FRHICommandListImmediate& RHICmdList)
 	{
 		check(ReferenceView.RayTracingMaterialPipeline || ReferenceView.RayTracingMaterialBindings.Num() == 0);
 
@@ -1839,9 +1839,13 @@ void FDeferredShadingSceneRenderer::WaitForRayTracingScene(FRDGBuilder& GraphBui
 		{
 			BindRayTracingMaterialPipeline(RHICmdList, ReferenceView, ReferenceView.RayTracingMaterialPipeline);
 
+			SetupRayTracingDefaultMissShader(RHICmdList, ReferenceView);
+
 			if (!bIsPathTracing)
 			{
 				SetupRayTracingLightingMissShader(RHICmdList, ReferenceView);
+				
+				BindLightFunctionShaders(RHICmdList, Scene, RayTracingLightFunctionMap, ReferenceView);
 			}
 		}
 
@@ -1861,11 +1865,13 @@ void FDeferredShadingSceneRenderer::WaitForRayTracingScene(FRDGBuilder& GraphBui
 			{
 				if (ReferenceView.RayTracingMaterialGatherPipeline)
 				{
+					RHICmdList.SetRayTracingMissShader(ReferenceView.GetRayTracingSceneChecked(), RAY_TRACING_MISS_SHADER_SLOT_DEFAULT, ReferenceView.RayTracingMaterialGatherPipeline, 0 /* MissShaderPipelineIndex */, 0, nullptr, 0);
 					BindRayTracingDeferredMaterialGatherPipeline(RHICmdList, ReferenceView, ReferenceView.RayTracingMaterialGatherPipeline);
 				}
 
 				if (ReferenceView.LumenHardwareRayTracingMaterialPipeline)
 				{
+					RHICmdList.SetRayTracingMissShader(ReferenceView.GetRayTracingSceneChecked(), RAY_TRACING_MISS_SHADER_SLOT_DEFAULT, ReferenceView.LumenHardwareRayTracingMaterialPipeline, 0 /* MissShaderPipelineIndex */, 0, nullptr, 0);
 					BindLumenHardwareRayTracingMaterialPipeline(RHICmdList, LumenHardwareRayTracingMaterialBindings, ReferenceView, ReferenceView.LumenHardwareRayTracingMaterialPipeline, ReferenceView.LumenHardwareRayTracingHitDataBuffer);
 				}
 			}
@@ -2180,6 +2186,7 @@ void FDeferredShadingSceneRenderer::Render(FRDGBuilder& GraphBuilder)
 
 	// Prepare the scene for rendering this frame.
 	FRayTracingScene& RayTracingScene = Scene->RayTracingScene;
+	FRayTracingLightFunctionMap RayTracingLightFunctionMap;
 	RayTracingScene.Reset(); // Resets the internal arrays, but does not release any resources.
 
 	const int32 ReferenceViewIndex = 0;
@@ -2196,6 +2203,17 @@ void FDeferredShadingSceneRenderer::Render(FRDGBuilder& GraphBuilder)
 		ReferenceView.RayTracingDecalUniformBuffer = CreateNullRayTracingDecalsUniformBuffer(GraphBuilder);
 		ReferenceView.bHasRayTracingDecals = false;
 	}
+	
+	// TODO: currently Path tracing only uses the default miss shader, but future extensions should be handled here
+	if (!ViewFamily.EngineShowFlags.PathTracing)
+	{
+		// get the default lighting miss shader (to implicitly fill in the MissShader library before the RT pipeline is created)
+		GetRayTracingLightingMissShader(ReferenceView);
+		RayTracingScene.NumMissShaderSlots++;
+
+		// gather all the light functions that may be used (and also count how many miss shaders we will need)
+		RayTracingLightFunctionMap = GatherLightFunctionLights(Scene, ViewFamily.EngineShowFlags, ReferenceView.GetFeatureLevel());
+	}
 
 	// Asynchronously create a list of primitives relevant to ray tracing scene
 	FRayTracingRelevantPrimitiveList RayTracingRelevantPrimitiveList;
@@ -2888,7 +2906,7 @@ void FDeferredShadingSceneRenderer::Render(FRDGBuilder& GraphBuilder)
 
 #if RHI_RAYTRACING
 	/** Should be called somewhere before "WaitForRayTracingScene" */
-	SetupRayTracingLightDataForViews(GraphBuilder);
+	SetupRayTracingLightDataForViews(GraphBuilder, RayTracingLightFunctionMap);
 #endif // RHI_RAYTRACING
 
 	// Shadows, lumen and fog after base pass
@@ -2914,7 +2932,7 @@ void FDeferredShadingSceneRenderer::Render(FRDGBuilder& GraphBuilder)
 		// Lumen scene lighting requires ray tracing scene to be ready if HWRT shadows are desired
 		if (!bRayTracingSceneReady && Lumen::UseHardwareRayTracedSceneLighting(ViewFamily))
 		{
-			WaitForRayTracingScene(GraphBuilder, DynamicGeometryScratchBuffer);
+			WaitForRayTracingScene(GraphBuilder, DynamicGeometryScratchBuffer, RayTracingLightFunctionMap);
 			bRayTracingSceneReady = true;
 		}
 #endif // RHI_RAYTRACING
@@ -3004,7 +3022,7 @@ void FDeferredShadingSceneRenderer::Render(FRDGBuilder& GraphBuilder)
 	// If Lumen did not force an earlier ray tracing scene sync, we must wait for it here.
 	if (!bRayTracingSceneReady)
 	{
-		WaitForRayTracingScene(GraphBuilder, DynamicGeometryScratchBuffer);
+		WaitForRayTracingScene(GraphBuilder, DynamicGeometryScratchBuffer, RayTracingLightFunctionMap);
 		bRayTracingSceneReady = true;
 	}
 #endif // RHI_RAYTRACING
diff --git a/Engine/Source/Runtime/Renderer/Private/DeferredShadingRenderer.h b/Engine/Source/Runtime/Renderer/Private/DeferredShadingRenderer.h
--- a/Engine/Source/Runtime/Renderer/Private/DeferredShadingRenderer.h
+++ b/Engine/Source/Runtime/Renderer/Private/DeferredShadingRenderer.h
@@ -18,6 +18,7 @@
 #include "ScreenSpaceDenoise.h"
 #include "Lumen/LumenProbeHierarchy.h"
 #include "Lumen/LumenSceneRendering.h"
+#include "RayTracing/RayTracingLighting.h"
 #include "IndirectLightRendering.h"
 #include "ScreenSpaceRayTracing.h"
 #include "RenderGraphUtils.h"
@@ -991,6 +992,9 @@ private:
 		int32 HeightFog,
 		float ResolutionFraction);
 
+	/** Setup the default miss shader (required for any raytracing pipeline) */
+	void SetupRayTracingDefaultMissShader(FRHICommandListImmediate& RHICmdList, const FViewInfo& View);
+
 	/** Lighting Evaluation shader setup (used by ray traced reflections and translucency) */
 	void SetupRayTracingLightingMissShader(FRHICommandListImmediate& RHICmdList, const FViewInfo& View);
 
@@ -1004,7 +1008,7 @@ private:
 	void ComputePathCompaction(FRHICommandListImmediate& RHICmdList, const FViewInfo& View, FRHITexture* RadianceTexture, FRHITexture* SampleCountTexture, FRHITexture* PixelPositionTexture,
 		FRHIUnorderedAccessView* RadianceSortedRedUAV, FRHIUnorderedAccessView* RadianceSortedGreenUAV, FRHIUnorderedAccessView* RadianceSortedBlueUAV, FRHIUnorderedAccessView* RadianceSortedAlphaUAV, FRHIUnorderedAccessView* SampleCountSortedUAV);
 
-	void WaitForRayTracingScene(FRDGBuilder& GraphBuilder, FRDGBufferRef DynamicGeometryScratchBuffer);
+	void WaitForRayTracingScene(FRDGBuilder& GraphBuilder, FRDGBufferRef DynamicGeometryScratchBuffer, FRayTracingLightFunctionMap& RayTracingLightFunctionMap);
 
 	/** Debug ray tracing functions. */
 	void RenderRayTracingDebug(FRDGBuilder& GraphBuilder, const FViewInfo& View, FRDGTextureRef SceneColorOutputTexture);
@@ -1014,7 +1018,7 @@ private:
 	bool GatherRayTracingWorldInstancesForView(FRDGBuilder& GraphBuilder, FViewInfo& View, FRayTracingScene& RayTracingScene, struct FRayTracingRelevantPrimitiveList&& RelevantPrimitiveList);
 
 	bool SetupRayTracingPipelineStates(FRDGBuilder& GraphBuilder);
-	void SetupRayTracingLightDataForViews(FRDGBuilder& GraphBuilder);
+	void SetupRayTracingLightDataForViews(FRDGBuilder& GraphBuilder, const FRayTracingLightFunctionMap& RayTracingLightFunctionMap);
 	bool DispatchRayTracingWorldUpdates(FRDGBuilder& GraphBuilder, FRDGBufferRef& OutDynamicGeometryScratchBuffer);
 
 	/** Functions to create ray tracing pipeline state objects for various effects */
@@ -1066,7 +1070,8 @@ private:
 	static void PrepareLumenHardwareRayTracingDirectLightingLumenMaterial(const FViewInfo& View, TArray<FRHIRayTracingShader*>& OutRayGenShaders);
 
 	/** Lighting evaluation shader registration */
-	static FRHIRayTracingShader* GetRayTracingLightingMissShader(FViewInfo& View);
+	static FRHIRayTracingShader* GetRayTracingDefaultMissShader(const FViewInfo& View);
+	static FRHIRayTracingShader* GetRayTracingLightingMissShader(const FViewInfo& View);
 
 	const FRHITransition* RayTracingDynamicGeometryUpdateEndTransition = nullptr; // Signaled when all AS for this frame are built
 #endif // RHI_RAYTRACING
diff --git a/Engine/Source/Runtime/Renderer/Private/MeshPassProcessor.cpp b/Engine/Source/Runtime/Renderer/Private/MeshPassProcessor.cpp
--- a/Engine/Source/Runtime/Renderer/Private/MeshPassProcessor.cpp
+++ b/Engine/Source/Runtime/Renderer/Private/MeshPassProcessor.cpp
@@ -369,6 +369,53 @@ FRayTracingLocalShaderBindings* FMeshDrawShaderBindings::SetRayTracingShaderBind
 	return SetRayTracingShaderBindingsForHitGroup(BindingWriter, 0, 0, ShaderIndexInPipeline, ShaderSlot);
 }
 
+void FMeshDrawShaderBindings::SetRayTracingShaderBindingsForMissShader(
+	FRHICommandList& RHICmdList,
+	FRHIRayTracingScene* Scene,
+	FRayTracingPipelineState* PipelineState,
+	uint32 ShaderIndexInPipeline,
+	uint32 ShaderSlot) const
+{
+	check(ShaderLayouts.Num() == 1);
+
+	FReadOnlyMeshDrawSingleShaderBindings SingleShaderBindings(ShaderLayouts[0], GetData());
+
+	FShaderBindingState BindingState;
+
+	const int32 MaxUniformBuffers = UE_ARRAY_COUNT(BindingState.UniformBuffers);
+
+	FRHIUniformBuffer* const* RESTRICT UniformBufferBindings = SingleShaderBindings.GetUniformBufferStart();
+	const FShaderParameterInfo* RESTRICT UniformBufferParameters = SingleShaderBindings.ParameterMapInfo.UniformBuffers.GetData();
+	const int32 NumUniformBufferParameters = SingleShaderBindings.ParameterMapInfo.UniformBuffers.Num();
+
+	checkf(SingleShaderBindings.ParameterMapInfo.TextureSamplers.Num() == 0, TEXT("Texture sampler parameters are not supported for ray tracing. UniformBuffers must be used for all resource binding."));
+	checkf(SingleShaderBindings.ParameterMapInfo.SRVs.Num() == 0, TEXT("SRV parameters are not supported for ray tracing. UniformBuffers must be used for all resource binding."));
+
+	// Measure parameter memory requirements
+
+	int32 MaxUniformBufferUsed = -1;
+	for (int32 UniformBufferIndex = 0; UniformBufferIndex < NumUniformBufferParameters; UniformBufferIndex++)
+	{
+		FShaderParameterInfo Parameter = UniformBufferParameters[UniformBufferIndex];
+		checkSlow(Parameter.BaseIndex < MaxUniformBuffers);
+		FRHIUniformBuffer* UniformBuffer = UniformBufferBindings[UniformBufferIndex];
+		if (Parameter.BaseIndex < MaxUniformBuffers)
+		{
+			BindingState.UniformBuffers[Parameter.BaseIndex] = UniformBuffer;
+			MaxUniformBufferUsed = FMath::Max((int32)Parameter.BaseIndex, MaxUniformBufferUsed);
+		}
+	}
+
+	checkf(SingleShaderBindings.ParameterMapInfo.TextureSamplers.Num() == 0, TEXT("Texture sampler parameters are not supported for ray tracing. UniformBuffers must be used for all resource binding."));
+	checkf(SingleShaderBindings.ParameterMapInfo.SRVs.Num() == 0, TEXT("SRV parameters are not supported for ray tracing. UniformBuffers must be used for all resource binding."));
+	checkf(SingleShaderBindings.ParameterMapInfo.LooseParameterBuffers.Num() == 0, TEXT("Ray tracing miss shaders may not have loose parameters"));
+
+	uint32 NumUniformBuffersToSet = MaxUniformBufferUsed + 1;
+	const uint32 UserData = 0; // UserData could be used to store material ID or any other kind of per-material constant. This can be retrieved in hit shaders via GetHitGroupUserData().
+	RHICmdList.SetRayTracingMissShader(Scene, ShaderSlot, PipelineState, ShaderIndexInPipeline,
+		NumUniformBuffersToSet, BindingState.UniformBuffers,
+		UserData);
+}
 #endif // RHI_RAYTRACING
 
 FGraphicsMinimalPipelineStateId FGraphicsMinimalPipelineStateId::GetPersistentId(const FGraphicsMinimalPipelineStateInitializer& InPipelineState)
diff --git a/Engine/Source/Runtime/Renderer/Private/RayTracing/RayTracingAmbientOcclusion.cpp b/Engine/Source/Runtime/Renderer/Private/RayTracing/RayTracingAmbientOcclusion.cpp
--- a/Engine/Source/Runtime/Renderer/Private/RayTracing/RayTracingAmbientOcclusion.cpp
+++ b/Engine/Source/Runtime/Renderer/Private/RayTracing/RayTracingAmbientOcclusion.cpp
@@ -195,7 +195,10 @@ void FDeferredShadingSceneRenderer::RenderRayTracingAmbientOcclusion(
 			Initializer.SetHitGroupTable(HitGroupTable);
 			Initializer.bAllowHitGroupIndexing = false; // Use the same hit shader for all geometry in the scene by disabling SBT indexing.
 
+			// TODO(UE-157946): This pipeline does not bind any miss shader and relies on the pipeline to do this automatically. This should be made explicit.
 			Pipeline = PipelineStateCache::GetAndOrCreateRayTracingPipelineState(RHICmdList, Initializer);
+
+			RHICmdList.SetRayTracingMissShader(View.GetRayTracingSceneChecked(), 0, Pipeline, 0 /* ShaderIndexInPipeline */, 0, nullptr, 0);
 		}
 
 		FRHIRayTracingScene* RayTracingSceneRHI = View.GetRayTracingSceneChecked();
diff --git a/Engine/Source/Runtime/Renderer/Private/RayTracing/RayTracingBarycentrics.cpp b/Engine/Source/Runtime/Renderer/Private/RayTracing/RayTracingBarycentrics.cpp
--- a/Engine/Source/Runtime/Renderer/Private/RayTracing/RayTracingBarycentrics.cpp
+++ b/Engine/Source/Runtime/Renderer/Private/RayTracing/RayTracingBarycentrics.cpp
@@ -131,6 +131,7 @@ void RenderRayTracingBarycentricsRGS(FRDGBuilder& GraphBuilder, const FViewInfo&
 	Initializer.SetHitGroupTable(HitGroupTable);
 	Initializer.bAllowHitGroupIndexing = false; // Use the same hit shader for all geometry in the scene by disabling SBT indexing.
 
+	// TODO(UE-157946): This pipeline does not bind any miss shader and relies on the pipeline to do this automatically. This should be made explicit.
 	FRayTracingPipelineState* Pipeline = PipelineStateCache::GetAndOrCreateRayTracingPipelineState(GraphBuilder.RHICmdList, Initializer);
 
 	FRayTracingBarycentricsRGS::FParameters* RayGenParameters = GraphBuilder.AllocParameters<FRayTracingBarycentricsRGS::FParameters>();
@@ -151,6 +152,7 @@ void RenderRayTracingBarycentricsRGS(FRDGBuilder& GraphBuilder, const FViewInfo&
 		SetShaderParameters(GlobalResources, RayGenShader, *RayGenParameters);
 
 		// Dispatch rays using default shader binding table
+		RHICmdList.SetRayTracingMissShader(View.GetRayTracingSceneChecked(), 0, Pipeline, 0 /* ShaderIndexInPipeline */, 0, nullptr, 0);
 		RHICmdList.RayTraceDispatch(Pipeline, RayGenShader.GetRayTracingShader(), View.GetRayTracingSceneChecked(), GlobalResources, ViewRect.Size().X, ViewRect.Size().Y);
 	});
 }
diff --git a/Engine/Source/Runtime/Renderer/Private/RayTracing/RayTracingDebug.cpp b/Engine/Source/Runtime/Renderer/Private/RayTracing/RayTracingDebug.cpp
--- a/Engine/Source/Runtime/Renderer/Private/RayTracing/RayTracingDebug.cpp
+++ b/Engine/Source/Runtime/Renderer/Private/RayTracing/RayTracingDebug.cpp
@@ -383,7 +383,7 @@ void FDeferredShadingSceneRenderer::RenderRayTracingDebug(FRDGBuilder& GraphBuil
 		Initializer.SetHitGroupTable(HitGroupTable);
 		Initializer.bAllowHitGroupIndexing = true; // Required for stable output using GetBaseInstanceIndex().
 		Initializer.MaxPayloadSizeInBytes = RAY_TRACING_MAX_ALLOWED_PAYLOAD_SIZE;
-
+		// TODO(UE-157946): This pipeline does not bind any miss shader and relies on the pipeline to do this automatically. This should be made explicit.
 		Pipeline = PipelineStateCache::GetAndOrCreateRayTracingPipelineState(GraphBuilder.RHICmdList, Initializer);
 		bRequiresBindings = true;
 	}
@@ -428,6 +428,7 @@ void FDeferredShadingSceneRenderer::RenderRayTracingDebug(FRDGBuilder& GraphBuil
 		if (bRequiresBindings)
 		{
 			BindRayTracingDebugCHSMaterialBindings(RHICmdList, View, Pipeline);
+			RHICmdList.SetRayTracingMissShader(View.GetRayTracingSceneChecked(), 0, Pipeline, 0 /* ShaderIndexInPipeline */, 0, nullptr, 0);
 		}
 
 		RHICmdList.RayTraceDispatch(Pipeline, RayGenShader.GetRayTracingShader(), View.GetRayTracingSceneChecked(), GlobalResources, ViewRect.Size().X, ViewRect.Size().Y);
diff --git a/Engine/Source/Runtime/Renderer/Private/RayTracing/RayTracingLighting.cpp b/Engine/Source/Runtime/Renderer/Private/RayTracing/RayTracingLighting.cpp
--- a/Engine/Source/Runtime/Renderer/Private/RayTracing/RayTracingLighting.cpp
+++ b/Engine/Source/Runtime/Renderer/Private/RayTracing/RayTracingLighting.cpp
@@ -22,6 +22,11 @@ static TAutoConsoleVariable<float> CVarRayTracingLightingCellSize(
 	ECVF_RenderThreadSafe
 );
 
+static TAutoConsoleVariable<int32> CVarRayTracingLightFunction(
+	TEXT("r.RayTracing.LightFunction"),
+	1,
+	TEXT("Whether to support light material functions in ray tracing effects. (default = 1)"),
+	ECVF_RenderThreadSafe);
 IMPLEMENT_GLOBAL_SHADER_PARAMETER_STRUCT(FRaytracingLightDataPacked, "RaytracingLightsDataPacked");
 
 class FSetupRayTracingLightCullData : public FGlobalShader
@@ -62,7 +67,7 @@ DECLARE_GPU_STAT_NAMED(LightCullingVolumeCompute, TEXT("RT Light Culling Volume
 
 
 static void SelectRaytracingLights(
-	const TSparseArray<FLightSceneInfoCompact, TAlignedSparseArrayAllocator<alignof(FLightSceneInfoCompact)>>& Lights,
+	const FScene::FLightSceneInfoCompactSparseArray& Lights,
 	TArray<int32>& OutSelectedLights,
 	uint32& NumOfSkippedRayTracingLights
 )
@@ -95,7 +100,7 @@ static int32 GetCellsPerDim()
 
 static void CreateRaytracingLightCullingStructure(
 	FRDGBuilder& GraphBuilder,
-	const TSparseArray<FLightSceneInfoCompact, TAlignedSparseArrayAllocator<alignof(FLightSceneInfoCompact)>>& Lights,
+	const FScene::FLightSceneInfoCompactSparseArray& Lights,
 	const FSceneView& View,
 	FGlobalShaderMap* ShaderMap,
 	const TArray<int32>& LightIndices,
@@ -163,12 +168,14 @@ static void CreateRaytracingLightCullingStructure(
 
 static void SetupRaytracingLightDataPacked(
 	FRDGBuilder& GraphBuilder,
-	const TSparseArray<FLightSceneInfoCompact, TAlignedSparseArrayAllocator<alignof(FLightSceneInfoCompact)>>& Lights,
+	const FScene* Scene,
 	const TArray<int32>& LightIndices,
 	const FSceneView& View,
+	const FRayTracingLightFunctionMap& RayTracingLightFunctionMap,
 	FRaytracingLightDataPacked& OutLightData,
 	TArrayView<FRTLightingData>& OutLightDataArray)
 {
+	const FScene::FLightSceneInfoCompactSparseArray& Lights = Scene->Lights;
 	TMap<UTextureLightProfile*, int32> IESLightProfilesMap;
 	TMap<FRHITexture*, uint32> RectTextureMap;
 
@@ -261,7 +268,7 @@ static void SetupRaytracingLightDataPacked(
 		LightDataElement.RectLightAtlasUVScale[0] = LightParameters.RectLightAtlasUVScale.X;
 		LightDataElement.RectLightAtlasUVScale[1] = LightParameters.RectLightAtlasUVScale.Y;
 		LightDataElement.RectLightAtlasMaxLevel = LightParameters.RectLightAtlasMaxLevel;
-		LightDataElement.Pad = 0;
+		LightDataElement.LightMissShaderIndex = RAY_TRACING_MISS_SHADER_SLOT_LIGHTING;
 
 		// Stuff directional light's shadow angle factor into a RectLight parameter
 		if (Light.LightType == LightType_Directional)
@@ -269,6 +276,15 @@ static void SetupRaytracingLightDataPacked(
 			LightDataElement.RectLightBarnCosAngle = Light.LightSceneInfo->Proxy->GetShadowSourceAngleFactor();
 		}
 
+		// NOTE: This map will be empty if the light functions are disabled for some reason
+		const int32* LightFunctionIndex = RayTracingLightFunctionMap.Find(Light.LightSceneInfo);
+		if (LightFunctionIndex)
+		{
+			check(uint32(*LightFunctionIndex) > RAY_TRACING_MISS_SHADER_SLOT_LIGHTING);
+			check(uint32(*LightFunctionIndex) < Scene->RayTracingScene.NumMissShaderSlots);
+			LightDataElement.LightMissShaderIndex = *LightFunctionIndex;
+		}
+
 		OutLightData.Count++;
 	}
 
@@ -292,8 +308,9 @@ static void SetupRaytracingLightDataPacked(
 
 TRDGUniformBufferRef<FRaytracingLightDataPacked> CreateRayTracingLightData(
 	FRDGBuilder& GraphBuilder,
-	const TSparseArray<FLightSceneInfoCompact, TAlignedSparseArrayAllocator<alignof(FLightSceneInfoCompact)>>& Lights,
+	const FScene* Scene,
 	const FSceneView& View,
+	const FRayTracingLightFunctionMap& RayTracingLightFunctionMap,
 	FGlobalShaderMap* ShaderMap,
 	uint32& NumOfSkippedRayTracingLights)
 {
@@ -302,15 +319,15 @@ TRDGUniformBufferRef<FRaytracingLightDataPacked> CreateRayTracingLightData(
 	LightData->CellScale = CVarRayTracingLightingCellSize.GetValueOnRenderThread() / 2.0f;
 
 	TArray<int32> LightIndices;
-	SelectRaytracingLights(Lights, LightIndices, NumOfSkippedRayTracingLights);
+	SelectRaytracingLights(Scene->Lights, LightIndices, NumOfSkippedRayTracingLights);
 
 	// Create light culling volume
 	FRDGBufferRef LightCullVolume;
 	FRDGBufferRef LightIndicesBuffer;
-	CreateRaytracingLightCullingStructure(GraphBuilder, Lights, View, ShaderMap, LightIndices, LightCullVolume, LightIndicesBuffer);
+	CreateRaytracingLightCullingStructure(GraphBuilder, Scene->Lights, View, ShaderMap, LightIndices, LightCullVolume, LightIndicesBuffer);
 
 	FRDGUploadData<FRTLightingData> LightDataArray(GraphBuilder, FMath::Max(LightIndices.Num(), 1));
-	SetupRaytracingLightDataPacked(GraphBuilder, Lights, LightIndices, View, *LightData, LightDataArray);
+	SetupRaytracingLightDataPacked(GraphBuilder, Scene, LightIndices, View, RayTracingLightFunctionMap, *LightData, LightDataArray);
 
 	check(LightData->Count == LightIndices.Num());
 
@@ -343,11 +360,255 @@ class FRayTracingLightingMS : public FGlobalShader
 
 IMPLEMENT_GLOBAL_SHADER(FRayTracingLightingMS, "/Engine/Private/RayTracing/RayTracingLightingMS.usf", "RayTracingLightingMS", SF_RayMiss);
 
-FRHIRayTracingShader* FDeferredShadingSceneRenderer::GetRayTracingLightingMissShader(FViewInfo& View)
+/**
+ * FLightFunctionParametersRayTracing
+ * Global constant buffer derived from loose parameters of standard light function materials. Note that it lacks
+ * the screen to world transform, as the RT version never have screen as a reference frame
+ * This is nearly identical to the one found in Lumen and ultimately should be converted to a shared solution
+ *
+ * function to create the constant buffer is deirved from the LightFunctionMaterial SetParameters code
+ */
+BEGIN_GLOBAL_SHADER_PARAMETER_STRUCT(FLightFunctionParametersRayTracing, )
+	SHADER_PARAMETER(FMatrix44f, LightFunctionTranslatedWorldToLight)
+	SHADER_PARAMETER(FVector4f, LightFunctionParameters)
+	SHADER_PARAMETER(FVector3f, LightFunctionParameters2)
+END_GLOBAL_SHADER_PARAMETER_STRUCT()
+
+IMPLEMENT_GLOBAL_SHADER_PARAMETER_STRUCT(FLightFunctionParametersRayTracing, "RaytracingLightFunctionParameters");
+
+static TUniformBufferRef<FLightFunctionParametersRayTracing> CreateLightFunctionParametersBuffer(
+	const FLightSceneInfo* LightSceneInfo,
+	const FSceneView& View,
+	EUniformBufferUsage Usage)
+{
+	FLightFunctionParametersRayTracing LightFunctionParameters;
+
+	const FVector Scale = LightSceneInfo->Proxy->GetLightFunctionScale();
+	// Switch x and z so that z of the user specified scale affects the distance along the light direction
+	const FVector InverseScale = FVector(1.f / Scale.Z, 1.f / Scale.Y, 1.f / Scale.X);
+	const FMatrix WorldToLight = LightSceneInfo->Proxy->GetWorldToLight() * FScaleMatrix(FVector(InverseScale));
+
+	LightFunctionParameters.LightFunctionTranslatedWorldToLight = FMatrix44f(FTranslationMatrix(-View.ViewMatrices.GetPreViewTranslation()) * WorldToLight);
+
+	const bool bIsSpotLight = LightSceneInfo->Proxy->GetLightType() == LightType_Spot;
+	const bool bIsPointLight = LightSceneInfo->Proxy->GetLightType() == LightType_Point;
+	const float TanOuterAngle = bIsSpotLight ? FMath::Tan(LightSceneInfo->Proxy->GetOuterConeAngle()) : 1.0f;
+
+	// should this match raster?
+	const float ShadowFadeFraction = 1.0f;
+
+	LightFunctionParameters.LightFunctionParameters = FVector4f(TanOuterAngle, ShadowFadeFraction, bIsSpotLight ? 1.0f : 0.0f, bIsPointLight ? 1.0f : 0.0f);
+
+	// do we need this?
+	const bool bRenderingPreviewShadowIndicator = false;
+
+	LightFunctionParameters.LightFunctionParameters2 = FVector3f(
+		LightSceneInfo->Proxy->GetLightFunctionFadeDistance(),
+		LightSceneInfo->Proxy->GetLightFunctionDisabledBrightness(),
+		bRenderingPreviewShadowIndicator ? 1.0f : 0.0f);
+
+	return CreateUniformBufferImmediate(LightFunctionParameters, Usage);
+}
+
+/**
+ * Generic light function for ray tracing compilable as miss shader with lighting
+ */
+class FLightFunctionRayTracingShader : public FMaterialShader
+{
+	DECLARE_SHADER_TYPE(FLightFunctionRayTracingShader, Material);
+public:
+
+	/**
+	  * Makes sure only shaders for materials that are explicitly flagged
+	  * as 'UsedAsLightFunction' in the Material Editor gets compiled into
+	  * the shader cache.
+	  */
+	static bool ShouldCompilePermutation(const FMaterialShaderPermutationParameters& Parameters)
+	{
+		return Parameters.MaterialParameters.MaterialDomain == MD_LightFunction && ShouldCompileRayTracingShadersForProject(Parameters.Platform);
+	}
+
+	FLightFunctionRayTracingShader() {}
+	FLightFunctionRayTracingShader(const ShaderMetaType::CompiledShaderInitializerType& Initializer)
+		: FMaterialShader(Initializer)
+	{
+		LightMaterialsParameter.Bind(Initializer.ParameterMap, TEXT("RaytracingLightFunctionParameters"));
+		LightDataPacked.Bind(Initializer.ParameterMap, TEXT("RaytracingLightsDataPacked"));
+	}
+
+	void GetShaderBindings(
+		const FScene* Scene,
+		ERHIFeatureLevel::Type FeatureLevel,
+		const FMaterialRenderProxy& MaterialRenderProxy,
+		const FMaterial& Material,
+		const FViewInfo& View,
+		const TUniformBufferRef<FDeferredLightUniformStruct>& DeferredLightBuffer,
+		const TUniformBufferRef<FLightFunctionParametersRayTracing>& LightFunctionParameters,
+		const TUniformBufferRef<FRaytracingLightDataPacked>& LightDataPackedBuffer,
+		FMeshDrawSingleShaderBindings& ShaderBindings) const
+	{
+		FMaterialShader::GetShaderBindings(Scene, FeatureLevel, MaterialRenderProxy, Material, ShaderBindings);
+
+		// Bind view
+		ShaderBindings.Add(GetUniformBufferParameter<FViewUniformShaderParameters>(), View.ViewUniformBuffer);
+
+		// Bind light parameters
+		ShaderBindings.Add(GetUniformBufferParameter<FDeferredLightUniformStruct>(), DeferredLightBuffer);
+
+		//bind Lightfunction parameters
+		ShaderBindings.Add(LightMaterialsParameter, LightFunctionParameters);
+
+		//bind light data
+		ShaderBindings.Add(LightDataPacked, LightDataPackedBuffer);
+	}
+
+	static void ModifyCompilationEnvironment(const FMaterialShaderPermutationParameters& Parameters, FShaderCompilerEnvironment& OutEnvironment)
+	{
+		OutEnvironment.SetDefine(TEXT("SUPPORT_LIGHT_FUNCTION"), 1);
+		FMaterialShader::ModifyCompilationEnvironment(Parameters, OutEnvironment);
+	}
+
+private:
+
+	LAYOUT_FIELD(FShaderUniformBufferParameter, LightMaterialsParameter);
+	LAYOUT_FIELD(FShaderUniformBufferParameter, LightDataPacked);
+};
+
+IMPLEMENT_MATERIAL_SHADER_TYPE(, FLightFunctionRayTracingShader, TEXT("/Engine/Private/RayTracing/RayTracingLightingMS.usf"), TEXT("RayTracingLightingMS"), SF_RayMiss);
+
+FRayTracingLightFunctionMap GatherLightFunctionLights(FScene* Scene, const FEngineShowFlags EngineShowFlags, ERHIFeatureLevel::Type InFeatureLevel)
+{
+	// TODO: implement light functions in the path tracer
+	if (!EngineShowFlags.LightFunctions || EngineShowFlags.PathTracing)
+	{
+		return {};
+	}
+
+	// gives control over lighting functions in raytraced effects, independently of the show flag (for performance testing / debug)
+	if (CVarRayTracingLightFunction.GetValueOnRenderThread() == 0)
+	{
+		return {};
+	}
+
+	FRayTracingLightFunctionMap RayTracingLightFunctionMap;
+	for (const FLightSceneInfoCompact& Light : Scene->Lights)
+	{
+		FLightSceneInfo* LightSceneInfo = Light.LightSceneInfo;
+		auto MaterialProxy = LightSceneInfo->Proxy->GetLightFunctionMaterial();
+		if (MaterialProxy)
+		{
+			const FMaterialRenderProxy* FallbackMaterialRenderProxyPtr = nullptr;
+			const FMaterial& Material = MaterialProxy->GetMaterialWithFallback(InFeatureLevel, FallbackMaterialRenderProxyPtr);
+			if (Material.IsLightFunction())
+			{
+				const FMaterialShaderMap* MaterialShaderMap = Material.GetRenderingThreadShaderMap();
+				// Getting the shader here has the side-effect of populating the raytracing miss shader library which is used when building the raytracing pipeline
+				MaterialShaderMap->GetShader<FLightFunctionRayTracingShader>().GetRayTracingShader();
+
+				int32 Index = Scene->RayTracingScene.NumMissShaderSlots;
+				Scene->RayTracingScene.NumMissShaderSlots++;
+				RayTracingLightFunctionMap.Add(LightSceneInfo, Index);
+			}
+		}
+	}
+	return RayTracingLightFunctionMap;
+}
+
+static void BindLightFunction(
+	FRHICommandList& RHICmdList,
+	const FScene* Scene,
+	const FViewInfo& View,
+	const FMaterial& Material,
+	const FMaterialRenderProxy& MaterialRenderProxy,
+	const TUniformBufferRef<FDeferredLightUniformStruct>& DeferredLightBuffer,
+	const TUniformBufferRef<FLightFunctionParametersRayTracing>& LightFunctionParameters,
+	int32 Index
+	)
+{
+	FRHIRayTracingScene* RTScene = View.GetRayTracingSceneChecked();
+	FRayTracingPipelineState* Pipeline = View.RayTracingMaterialPipeline;
+	const FMaterialShaderMap* MaterialShaderMap = Material.GetRenderingThreadShaderMap();
+
+
+	auto Shader = MaterialShaderMap->GetShader<FLightFunctionRayTracingShader>();
+
+	TMeshProcessorShaders<
+		FMaterialShader,
+		FMaterialShader,
+		FMaterialShader,
+		FLightFunctionRayTracingShader,
+		FMaterialShader> RayTracingShaders;
+
+	RayTracingShaders.RayTracingShader = Shader;
+
+	FMeshDrawShaderBindings ShaderBindings;
+	ShaderBindings.Initialize(RayTracingShaders.GetUntypedShaders());
+
+	int32 DataOffset = 0;
+	FMeshDrawSingleShaderBindings SingleShaderBindings = ShaderBindings.GetSingleShaderBindings( SF_RayMiss, DataOffset);
+
+	RayTracingShaders.RayTracingShader->GetShaderBindings(Scene, Scene->GetFeatureLevel(), MaterialRenderProxy, Material, View, DeferredLightBuffer, LightFunctionParameters, View.RayTracingLightDataUniformBuffer->GetRHIRef(), SingleShaderBindings);
+
+	int32 MissShaderPipelineIndex = FindRayTracingMissShaderIndex(View.RayTracingMaterialPipeline, Shader.GetRayTracingShader(), true);
+
+	ShaderBindings.SetRayTracingShaderBindingsForMissShader(RHICmdList, RTScene, Pipeline, MissShaderPipelineIndex, Index);
+}
+
+FRHIRayTracingShader* FDeferredShadingSceneRenderer::GetRayTracingLightingMissShader(const FViewInfo& View)
 {
 	return View.ShaderMap->GetShader<FRayTracingLightingMS>().GetRayTracingShader();
 }
 
+void PrepareLightFunctionMissShaders(
+	const TArray<FLightSceneInfo*>& LightFunctionLights,
+	const class FViewInfo& View,
+	TArray<FRHIRayTracingShader*>& OutMissShaders)
+{
+	for (const FLightSceneInfo* LightSceneInfo : LightFunctionLights)
+	{
+		auto MaterialProxy = LightSceneInfo->Proxy->GetLightFunctionMaterial();
+		if (MaterialProxy)
+		{
+			const FMaterialRenderProxy* FallbackMaterialRenderProxyPtr = nullptr;
+			const FMaterial& Material = MaterialProxy->GetMaterialWithFallback(View.GetFeatureLevel(), FallbackMaterialRenderProxyPtr);
+			if (Material.IsLightFunction())
+			{
+				const FMaterialShaderMap* MaterialShaderMap = Material.GetRenderingThreadShaderMap();
+				OutMissShaders.Add(MaterialShaderMap->GetShader<FLightFunctionRayTracingShader>().GetRayTracingShader());
+			}
+		}
+	}
+}
+
+void BindLightFunctionShaders(
+	FRHICommandList& RHICmdList,
+	const FScene* Scene,
+	const FRayTracingLightFunctionMap& RayTracingLightFunctionMap,
+	const class FViewInfo& View)
+{
+	for (const FRayTracingLightFunctionMap::ElementType& LightAndIndex : RayTracingLightFunctionMap)
+	{
+		const FLightSceneInfo* LightSceneInfo = LightAndIndex.Key;
+
+		const FMaterialRenderProxy* MaterialProxy = LightSceneInfo->Proxy->GetLightFunctionMaterial();
+		check(MaterialProxy != nullptr);
+		// Catch the fallback material case
+		const FMaterialRenderProxy* FallbackMaterialRenderProxyPtr = nullptr;
+		const FMaterial& Material = MaterialProxy->GetMaterialWithFallback(Scene->GetFeatureLevel(), FallbackMaterialRenderProxyPtr);
+
+		check(Material.IsLightFunction());
+
+		const FMaterialRenderProxy& MaterialRenderProxy = FallbackMaterialRenderProxyPtr ? *FallbackMaterialRenderProxyPtr : *MaterialProxy;
+
+		//create the uniform buffers we need
+		TUniformBufferRef<FDeferredLightUniformStruct> DeferredLightBuffer = CreateUniformBufferImmediate(GetDeferredLightParameters(View, *LightSceneInfo), EUniformBufferUsage::UniformBuffer_SingleFrame);
+		TUniformBufferRef<FLightFunctionParametersRayTracing> LightFunctionParameters = CreateLightFunctionParametersBuffer(LightSceneInfo, View, EUniformBufferUsage::UniformBuffer_SingleFrame);
+
+		int32 MissIndex = LightAndIndex.Value;
+		BindLightFunction(RHICmdList, Scene, View, Material, MaterialRenderProxy, DeferredLightBuffer, LightFunctionParameters, MissIndex);
+	}
+}
+
 template< typename ShaderClass>
 static int32 BindParameters(const TShaderRef<ShaderClass>& Shader, typename ShaderClass::FParameters & Parameters, int32 MaxParams, const FRHIUniformBuffer **OutUniformBuffers)
 {
@@ -382,6 +643,17 @@ static int32 BindParameters(const TShaderRef<ShaderClass>& Shader, typename Shad
 	return MaxUniformBufferUsed + 1;
 }
 
+void FDeferredShadingSceneRenderer::SetupRayTracingDefaultMissShader(FRHICommandListImmediate& RHICmdList, const FViewInfo& View)
+{
+	int32 MissShaderPipelineIndex = FindRayTracingMissShaderIndex(View.RayTracingMaterialPipeline, GetRayTracingDefaultMissShader(View), true);
+
+	RHICmdList.SetRayTracingMissShader(View.GetRayTracingSceneChecked(),
+		RAY_TRACING_MISS_SHADER_SLOT_DEFAULT,
+		View.RayTracingMaterialPipeline,
+		MissShaderPipelineIndex,
+		0, nullptr, 0);
+}
+
 void FDeferredShadingSceneRenderer::SetupRayTracingLightingMissShader(FRHICommandListImmediate& RHICmdList, const FViewInfo& View)
 {
 	FRayTracingLightingMS::FParameters MissParameters;
@@ -394,10 +666,12 @@ void FDeferredShadingSceneRenderer::SetupRayTracingLightingMissShader(FRHIComman
 
 	int32 ParameterSlots = BindParameters(MissShader, MissParameters, MaxUniformBuffers, MissData);
 
+	int32 MissShaderPipelineIndex = FindRayTracingMissShaderIndex(View.RayTracingMaterialPipeline, MissShader.GetRayTracingShader(), true);
+
 	RHICmdList.SetRayTracingMissShader(View.GetRayTracingSceneChecked(),
-		RAY_TRACING_MISS_SHADER_SLOT_LIGHTING, // Shader slot in the scene
+		RAY_TRACING_MISS_SHADER_SLOT_LIGHTING,
 		View.RayTracingMaterialPipeline,
-		RAY_TRACING_MISS_SHADER_SLOT_LIGHTING, // Miss shader index in the pipeline
+		MissShaderPipelineIndex,
 		ParameterSlots, (FRHIUniformBuffer**)MissData, 0);
 }
 
diff --git a/Engine/Source/Runtime/Renderer/Private/RayTracing/RayTracingLighting.h b/Engine/Source/Runtime/Renderer/Private/RayTracing/RayTracingLighting.h
--- a/Engine/Source/Runtime/Renderer/Private/RayTracing/RayTracingLighting.h
+++ b/Engine/Source/Runtime/Renderer/Private/RayTracing/RayTracingLighting.h
@@ -30,9 +30,7 @@ struct FRTLightingData
 	int32 Type;
 	int32 LightProfileIndex;
 	float RectLightAtlasMaxLevel;
-
-	// Force alignment before next vector
-	int32 Pad;
+	uint32 LightMissShaderIndex;
 
 	FVector3f TranslatedLightPosition;
 	float InvRadius;
@@ -55,11 +53,26 @@ struct FRTLightingData
 
 static_assert(sizeof(FRTLightingData) == 128, "Unexpected FRTLightingData size.");
 
+using FRayTracingLightFunctionMap = TMap<const FLightSceneInfo*, int32>;
+FRayTracingLightFunctionMap GatherLightFunctionLights(FScene* Scene, const FEngineShowFlags EngineShowFlags, ERHIFeatureLevel::Type InFeatureLevel);
+
 TRDGUniformBufferRef<FRaytracingLightDataPacked> CreateRayTracingLightData(
 	FRDGBuilder& GraphBuilder,
-	const TSparseArray<FLightSceneInfoCompact, TAlignedSparseArrayAllocator<alignof(FLightSceneInfoCompact)>>& Lights,
+	const FScene* Scene,
 	const FSceneView& View,
+	const FRayTracingLightFunctionMap& RayTracingLightFunctionMap,
 	FGlobalShaderMap* ShaderMap,
 	uint32& NumOfSkippedRayTracingLights);
 
+void BindLightFunctionShaders(
+	FRHICommandList& RHICmdList,
+	const FScene* Scene,
+	const FRayTracingLightFunctionMap& RayTracingLightFunctionMap,
+	const class FViewInfo& View);
+
+void PrepareLightFunctionMissShaders(
+	const TArray<FLightSceneInfo*>& LightFunctionLights,
+	const class FViewInfo& View,
+	TArray<FRHIRayTracingShader*>& OutMissShaders);
+
 #endif
diff --git a/Engine/Source/Runtime/Renderer/Private/RayTracing/RayTracingMaterialHitShaders.cpp b/Engine/Source/Runtime/Renderer/Private/RayTracing/RayTracingMaterialHitShaders.cpp
--- a/Engine/Source/Runtime/Renderer/Private/RayTracing/RayTracingMaterialHitShaders.cpp
+++ b/Engine/Source/Runtime/Renderer/Private/RayTracing/RayTracingMaterialHitShaders.cpp
@@ -9,6 +9,7 @@
 #include "RayTracingInstance.h"
 #include "BuiltInRayTracingShaders.h"
 #include "RaytracingOptions.h"
+#include "RayTracingLighting.h"
 
 int32 GEnableRayTracingMaterials = 1;
 static FAutoConsoleVariableRef CVarEnableRayTracingMaterials(
@@ -579,6 +580,11 @@ static bool PipelineContainsHitShaders(FRayTracingPipelineState* Pipeline, const
 	return true;
 }
 
+FRHIRayTracingShader* FDeferredShadingSceneRenderer::GetRayTracingDefaultMissShader(const FViewInfo& View)
+{
+	return View.ShaderMap->GetShader<FPackedMaterialClosestHitPayloadMS>().GetRayTracingShader();
+}
+
 FRayTracingPipelineState* FDeferredShadingSceneRenderer::CreateRayTracingMaterialPipeline(
 	FRHICommandList& RHICmdList,
 	FViewInfo& View,
@@ -595,20 +601,15 @@ FRayTracingPipelineState* FDeferredShadingSceneRenderer::CreateRayTracingMateria
 	Initializer.MaxPayloadSizeInBytes = RAY_TRACING_MAX_ALLOWED_PAYLOAD_SIZE; // sizeof(FPackedMaterialClosestHitPayload)
 	Initializer.bAllowHitGroupIndexing = true;
 
-	FRHIRayTracingShader* DefaultMissShader = View.ShaderMap->GetShader<FPackedMaterialClosestHitPayloadMS>().GetRayTracingShader();
+	FRHIRayTracingShader* DefaultMissShader = GetRayTracingDefaultMissShader(View);
 
-	FRHIRayTracingShader* RayTracingMissShaderLibrary[RAY_TRACING_NUM_MISS_SHADER_SLOTS] = {};
-	RayTracingMissShaderLibrary[RAY_TRACING_MISS_SHADER_SLOT_DEFAULT] = DefaultMissShader;
-	if (bIsPathTracing)
-	{
-		// TODO: path tracer will want a different light miss shader
-		Initializer.SetMissShaderTable(TArrayView<FRHIRayTracingShader*>(RayTracingMissShaderLibrary, 1));
-	}
-	else
-	{
-		RayTracingMissShaderLibrary[RAY_TRACING_MISS_SHADER_SLOT_LIGHTING] = GetRayTracingLightingMissShader(View);
-		Initializer.SetMissShaderTable(RayTracingMissShaderLibrary);
-	}
+	TArray<FRHIRayTracingShader*> RayTracingMissShaderLibrary;
+	FShaderMapResource::GetRayTracingMissShaderLibrary(RayTracingMissShaderLibrary, DefaultMissShader);
+
+	// make sure we have at least one miss shader present
+	check(RayTracingMissShaderLibrary.Num() > 0);
+
+	Initializer.SetMissShaderTable(RayTracingMissShaderLibrary);
 
 	Initializer.SetRayGenShaderTable(RayGenShaderTable);
 
@@ -663,16 +664,16 @@ FRayTracingPipelineState* FDeferredShadingSceneRenderer::CreateRayTracingMateria
 		Initializer.SetCallableTable(RayTracingCallableShaderLibrary);
 	}
 
-	FRayTracingPipelineState* FallbackPipelineState = GRayTracingNonBlockingPipelineCreation && View.ViewState
+	const bool bAllowNonBlockingPipelineCreation = GRayTracingNonBlockingPipelineCreation && !View.bIsOfflineRender;
+	FRayTracingPipelineState* FallbackPipelineState = bAllowNonBlockingPipelineCreation && View.ViewState
 		? PipelineStateCache::GetRayTracingPipelineState(View.ViewState->LastRayTracingMaterialPipelineSignature)
 		: nullptr;
 
 	ERayTracingPipelineCacheFlags PipelineCacheFlags = ERayTracingPipelineCacheFlags::Default;
-	if (GRayTracingNonBlockingPipelineCreation
-		&& FallbackPipelineState
-		&& View.ViewState
+	if (FallbackPipelineState
 		&& IsCompatibleFallbackPipelineSignature(View.ViewState->LastRayTracingMaterialPipelineSignature, Initializer)
 		&& PipelineContainsHitShaders(FallbackPipelineState, RequiredHitShaders)
+		&& FindRayTracingMissShaderIndex(FallbackPipelineState, DefaultMissShader, false) != INDEX_NONE
 		&& (!bCallableShadersSupported || FindRayTracingCallableShaderIndex(FallbackPipelineState, DefaultCallableShader, false) != INDEX_NONE))
 	{
 		PipelineCacheFlags |= ERayTracingPipelineCacheFlags::NonBlocking;
diff --git a/Engine/Source/Runtime/Renderer/Private/RayTracing/RayTracingScene.cpp b/Engine/Source/Runtime/Renderer/Private/RayTracing/RayTracingScene.cpp
--- a/Engine/Source/Runtime/Renderer/Private/RayTracing/RayTracingScene.cpp
+++ b/Engine/Source/Runtime/Renderer/Private/RayTracing/RayTracingScene.cpp
@@ -31,20 +31,13 @@ FRayTracingSceneWithGeometryInstances FRayTracingScene::BuildInitializationData(
 		Instances,
 		uint8(ERayTracingSceneLayer::NUM),
 		RAY_TRACING_NUM_SHADER_SLOTS,
-		RAY_TRACING_NUM_MISS_SHADER_SLOTS,
+		NumMissShaderSlots,
 		NumCallableShaderSlots);
 }
 
 void FRayTracingScene::Create(FRDGBuilder& GraphBuilder, const FGPUScene& GPUScene, const FViewMatrices& ViewMatrices)
 {
-	FRayTracingSceneWithGeometryInstances InitializationData = CreateRayTracingSceneWithGeometryInstances(
-		Instances,
-		uint8(ERayTracingSceneLayer::NUM),
-		RAY_TRACING_NUM_SHADER_SLOTS,
-		RAY_TRACING_NUM_MISS_SHADER_SLOTS,
-		NumCallableShaderSlots);
-
-	CreateWithInitializationData(GraphBuilder, GPUScene, ViewMatrices, MoveTemp(InitializationData));
+	CreateWithInitializationData(GraphBuilder, GPUScene, ViewMatrices, BuildInitializationData());
 }
 
 void FRayTracingScene::CreateWithInitializationData(FRDGBuilder& GraphBuilder, const FGPUScene& GPUScene, const FViewMatrices& ViewMatrices, FRayTracingSceneWithGeometryInstances SceneWithGeometryInstances)
@@ -296,6 +289,7 @@ void FRayTracingScene::Reset()
 	WaitForTasks();
 
 	Instances.Reset();
+	NumMissShaderSlots = 1;
 	NumCallableShaderSlots = 0;
 	CallableCommands.Reset();
 	UniformBuffers.Reset();
diff --git a/Engine/Source/Runtime/Renderer/Private/RayTracing/RayTracingScene.h b/Engine/Source/Runtime/Renderer/Private/RayTracing/RayTracingScene.h
--- a/Engine/Source/Runtime/Renderer/Private/RayTracing/RayTracingScene.h
+++ b/Engine/Source/Runtime/Renderer/Private/RayTracing/RayTracingScene.h
@@ -91,7 +91,7 @@ public:
 	TArray<FRayTracingGeometryInstance> Instances;
 
 	bool bUsedThisFrame = false;
-
+	uint32 NumMissShaderSlots = 1; // we must have a default miss shader, so always include it from the start
 	uint32 NumCallableShaderSlots = 0;
 	TArray<FRayTracingShaderCommand> CallableCommands;
 
diff --git a/Engine/Source/Runtime/Renderer/Private/RayTracing/RayTracingShadows.cpp b/Engine/Source/Runtime/Renderer/Private/RayTracing/RayTracingShadows.cpp
--- a/Engine/Source/Runtime/Renderer/Private/RayTracing/RayTracingShadows.cpp
+++ b/Engine/Source/Runtime/Renderer/Private/RayTracing/RayTracingShadows.cpp
@@ -382,8 +382,9 @@ void FDeferredShadingSceneRenderer::RenderRayTracingShadows(
 				Initializer.SetHitGroupTable(HitGroupTable);
 				Initializer.bAllowHitGroupIndexing = false; // Use the same hit shader for all geometry in the scene by disabling SBT indexing.
 
+				// TODO(UE-157946): This pipeline does not bind any miss shader and relies on the pipeline to do this automatically. This should be made explicit.
 				FRayTracingPipelineState* Pipeline = PipelineStateCache::GetAndOrCreateRayTracingPipelineState(RHICmdList, Initializer);
-
+				RHICmdList.SetRayTracingMissShader(RayTracingSceneRHI, 0, Pipeline, 0 /* ShaderIndexInPipeline */, 0, nullptr, 0);
 				RHICmdList.RayTraceDispatch(Pipeline, RayGenerationShader.GetRayTracingShader(), RayTracingSceneRHI, GlobalResources, Resolution.X, Resolution.Y);
 			}
 		});
diff --git a/Engine/Source/Runtime/Renderer/Private/RayTracing/RaytracingSkylight.cpp b/Engine/Source/Runtime/Renderer/Private/RayTracing/RaytracingSkylight.cpp
--- a/Engine/Source/Runtime/Renderer/Private/RayTracing/RaytracingSkylight.cpp
+++ b/Engine/Source/Runtime/Renderer/Private/RayTracing/RaytracingSkylight.cpp
@@ -470,6 +470,7 @@ void FDeferredShadingSceneRenderer::RenderRayTracingSkyLight(
 			SetShaderParameters(GlobalResources, RayGenerationShader, *PassParameters);
 
 			FRayTracingPipelineState* Pipeline = View.RayTracingMaterialPipeline;
+			FRHIRayTracingScene* RayTracingSceneRHI = View.GetRayTracingSceneChecked();
 			if (CVarRayTracingSkyLightEnableMaterials.GetValueOnRenderThread() == 0)
 			{
 				// Declare default pipeline
@@ -482,10 +483,11 @@ void FDeferredShadingSceneRenderer::RenderRayTracingSkyLight(
 				Initializer.SetHitGroupTable(HitGroupTable);
 				Initializer.bAllowHitGroupIndexing = false; // Use the same hit shader for all geometry in the scene by disabling SBT indexing.
 
+				// TODO(UE-157946): This pipeline does not bind any miss shader and relies on the pipeline to do this automatically. This should be made explicit.
 				Pipeline = PipelineStateCache::GetAndOrCreateRayTracingPipelineState(RHICmdList, Initializer);
+				RHICmdList.SetRayTracingMissShader(RayTracingSceneRHI, 0, Pipeline, 0 /* ShaderIndexInPipeline */, 0, nullptr, 0);
 			}
 
-			FRHIRayTracingScene* RayTracingSceneRHI = View.GetRayTracingSceneChecked();
 			RHICmdList.RayTraceDispatch(Pipeline, RayGenerationShader.GetRayTracingShader(), RayTracingSceneRHI, GlobalResources, RayTracingResolution.X, RayTracingResolution.Y);
 		});
 
diff --git a/Engine/Source/Runtime/Renderer/Private/ScenePrivate.h b/Engine/Source/Runtime/Renderer/Private/ScenePrivate.h
--- a/Engine/Source/Runtime/Renderer/Private/ScenePrivate.h
+++ b/Engine/Source/Runtime/Renderer/Private/ScenePrivate.h
@@ -2739,13 +2739,14 @@ public:
 	TArray<FTypeOffsetTableEntry> TypeOffsetTable;
 
 	/** The lights in the scene. */
-	TSparseArray<FLightSceneInfoCompact, TAlignedSparseArrayAllocator<alignof(FLightSceneInfoCompact)>> Lights;
+	using FLightSceneInfoCompactSparseArray = TSparseArray<FLightSceneInfoCompact, TAlignedSparseArrayAllocator<alignof(FLightSceneInfoCompact)>>;
+	FLightSceneInfoCompactSparseArray Lights;
 
 	/** 
 	 * Lights in the scene which are invisible, but still needed by the editor for previewing. 
 	 * Lights in this array cannot be in the Lights array.  They also are not fully set up, as AddLightSceneInfo_RenderThread is not called for them.
 	 */
-	TSparseArray<FLightSceneInfoCompact, TAlignedSparseArrayAllocator<alignof(FLightSceneInfoCompact)>> InvisibleLights;
+	FLightSceneInfoCompactSparseArray InvisibleLights;
 
 	/** Shadow casting lights that couldn't get a shadowmap channel assigned and therefore won't have valid dynamic shadows, forward renderer only. */
 	TArray<FString> OverflowingDynamicShadowedLights;
diff --git a/Engine/Source/Runtime/Renderer/Public/MeshPassProcessor.h b/Engine/Source/Runtime/Renderer/Public/MeshPassProcessor.h
--- a/Engine/Source/Runtime/Renderer/Public/MeshPassProcessor.h
+++ b/Engine/Source/Runtime/Renderer/Public/MeshPassProcessor.h
@@ -826,6 +826,9 @@ public:
 #if RHI_RAYTRACING
 	RENDERER_API FRayTracingLocalShaderBindings* SetRayTracingShaderBindingsForHitGroup(FRayTracingLocalShaderBindingWriter* BindingWriter, uint32 InstanceIndex, uint32 SegmentIndex, uint32 HitGroupIndexInPipeline, uint32 ShaderSlot) const;
 	RENDERER_API FRayTracingLocalShaderBindings* SetRayTracingShaderBindings(FRayTracingLocalShaderBindingWriter* BindingWriter, uint32 ShaderIndexInPipeline, uint32 ShaderSlot) const;
+
+	// TODO: should these move to a binding writer too? should we introduce a different class to do these bindings since they aren't mesh related? rename this class entirely?
+	void SetRayTracingShaderBindingsForMissShader(FRHICommandList& RHICmdList, FRHIRayTracingScene* Scene, FRayTracingPipelineState* Pipeline, uint32 ShaderIndexInPipeline, uint32 ShaderSlot) const;
 #endif // RHI_RAYTRACING
 
 	/** Returns whether this set of shader bindings can be merged into an instanced draw call with another. */
