Fix for TSR to not compose translucency over render target when non overview buffer visualisation is enabled.

diff --git a/Engine/Source/Runtime/Renderer/Private/PostProcess/TemporalSuperResolution.cpp b/Engine/Source/Runtime/Renderer/Private/PostProcess/TemporalSuperResolution.cpp
--- a/Engine/Source/Runtime/Renderer/Private/PostProcess/TemporalSuperResolution.cpp
+++ b/Engine/Source/Runtime/Renderer/Private/PostProcess/TemporalSuperResolution.cpp
@@ -767,6 +767,14 @@ ITemporalUpscaler::FOutputs AddTemporalSuperResolutionPasses(
 	FIntRect SeparateTranslucencyRect = FIntRect(0, 0, 1, 1);
 	FRDGTextureRef SeparateTranslucencyTexture = BlackAlphaOneDummy;
 	bool bHasSeparateTranslucency = PassInputs.PostDOFTranslucencyResources.IsValid();
+#if WITH_EDITOR
+	// Do not composite translucency if we are visualizing a buffer, unless it is the overview mode.
+	static FName OverviewName = FName(TEXT("Overview"));
+	bHasSeparateTranslucency &= 
+		   (!View.Family->EngineShowFlags.VisualizeBuffer || (View.Family->EngineShowFlags.VisualizeBuffer && View.CurrentBufferVisualizationMode == OverviewName))
+		&& (!View.Family->EngineShowFlags.VisualizeNanite || (View.Family->EngineShowFlags.VisualizeNanite && View.CurrentNaniteVisualizationMode == OverviewName))
+		&& (!View.Family->EngineShowFlags.VisualizeLumen  || (View.Family->EngineShowFlags.VisualizeLumen  && View.CurrentLumenVisualizationMode  == OverviewName));
+#endif
 	if (bHasSeparateTranslucency)
 	{
 		SeparateTranslucencyTexture = PassInputs.PostDOFTranslucencyResources.ColorTexture.Resolve;
