Fixed crash when attempting to load ray tracing shaders when ray tracing is disabled

diff --git a/Engine/Source/Runtime/Renderer/Private/DeferredShadingRenderer.cpp b/Engine/Source/Runtime/Renderer/Private/DeferredShadingRenderer.cpp
--- a/Engine/Source/Runtime/Renderer/Private/DeferredShadingRenderer.cpp
+++ b/Engine/Source/Runtime/Renderer/Private/DeferredShadingRenderer.cpp
@@ -2205,7 +2205,7 @@ void FDeferredShadingSceneRenderer::Render(FRDGBuilder& GraphBuilder)
 	}
 	
 	// TODO: currently Path tracing only uses the default miss shader, but future extensions should be handled here
-	if (!ViewFamily.EngineShowFlags.PathTracing)
+	if (IsRayTracingEnabled() && !ViewFamily.EngineShowFlags.PathTracing)
 	{
 		// get the default lighting miss shader (to implicitly fill in the MissShader library before the RT pipeline is created)
 		GetRayTracingLightingMissShader(ReferenceView);
