OpenXR: protect access to AllocateDepthTexture

Previously, we'd call AllocateDepthTexture from every SceneRenderer. Now we gate access with a bit in FSceneTexturesConfig

diff --git a/Engine/Plugins/Runtime/OpenXR/Source/OpenXRHMD/Private/OpenXRHMD.cpp b/Engine/Plugins/Runtime/OpenXR/Source/OpenXRHMD/Private/OpenXRHMD.cpp
--- a/Engine/Plugins/Runtime/OpenXR/Source/OpenXRHMD/Private/OpenXRHMD.cpp
+++ b/Engine/Plugins/Runtime/OpenXR/Source/OpenXRHMD/Private/OpenXRHMD.cpp
@@ -1894,21 +1894,13 @@ bool FOpenXRHMD::AllocateDepthTexture(uint32 Index, uint32 SizeX, uint32 SizeY,
 {
 	check(IsInRenderingThread());
 
-	// FIXME: UE constantly calls this function even when there is no reason to reallocate the depth texture
+	// FIXME: UE constantly calls this function even when there is no reason to reallocate the depth texture (see NeedReAllocateDepthTexture)
 	FReadScopeLock Lock(SessionHandleMutex);
 	if (!Session || !bDepthExtensionSupported)
 	{
 		return false;
 	}
 
-	// Ensure the texture size matches the eye layer. We may get other depth allocations unrelated to the main scene render.
-	// FIXME: This introduces a magic number that will break this logic. A developer can select a scene capture target that
-	// happens to be the same size as the ideal render target and break everything for reasons that will be a mystery to them.
-	if (FIntPoint(SizeX, SizeY) != GetIdealRenderTargetSize())
-	{
-		return false;
-	}
-
 	// We're only creating a 1x target here, but we don't know whether it'll be the targeted texture
 	// or the resolve texture. Because of this, we unify the input flags.
 	ETextureCreateFlags UnifiedCreateFlags = Flags | TargetableTextureFlags;
diff --git a/Engine/Plugins/Runtime/OpenXR/Source/OpenXRHMD/Private/OpenXRHMD.h b/Engine/Plugins/Runtime/OpenXR/Source/OpenXRHMD/Private/OpenXRHMD.h
--- a/Engine/Plugins/Runtime/OpenXR/Source/OpenXRHMD/Private/OpenXRHMD.h
+++ b/Engine/Plugins/Runtime/OpenXR/Source/OpenXRHMD/Private/OpenXRHMD.h
@@ -291,9 +291,10 @@ public:
 	virtual bool NeedReAllocateDepthTexture(const TRefCountPtr<IPooledRenderTarget>& DepthTarget) override final { return bNeedReAllocatedDepth; }
 	virtual bool AllocateDepthTexture(uint32 Index, uint32 SizeX, uint32 SizeY, uint8 Format, uint32 NumMips, ETextureCreateFlags InTexFlags, ETextureCreateFlags TargetableTextureFlags, FTexture2DRHIRef& OutTargetableTexture, FTexture2DRHIRef& OutShaderResourceTexture, uint32 NumSamples = 1) override final;
 
+	/** FXRRenderTargetManager */
 	virtual FXRRenderBridge* GetActiveRenderBridge_GameThread(bool bUseSeparateRenderTarget) override;
 
-	/** IStereoRenderTargetManager */
+	/** IXRTrackingSystem */
 	virtual void OnBeginPlay(FWorldContext& InWorldContext) override;
 	virtual void OnEndPlay(FWorldContext& InWorldContext) override;
 
diff --git a/Engine/Source/Runtime/Engine/Public/SceneTexturesConfig.h b/Engine/Source/Runtime/Engine/Public/SceneTexturesConfig.h
--- a/Engine/Source/Runtime/Engine/Public/SceneTexturesConfig.h
+++ b/Engine/Source/Runtime/Engine/Public/SceneTexturesConfig.h
@@ -106,6 +106,7 @@ struct ENGINE_API FSceneTexturesConfig
 		, bKeepDepthContent{ 1 }
 		, bPreciseDepthAux{}
 		, bSamplesCustomDepthAndStencil{}
+		, bSupportsXRTargetManagerDepthAlloc{}
 	{}
 
 	FORCEINLINE bool IsValid() const
@@ -167,6 +168,9 @@ struct ENGINE_API FSceneTexturesConfig
 	// (Mobile) True if both CustomDepth and CustomStencil are sampled in a shader
 	uint32 bSamplesCustomDepthAndStencil : 1;
 
+	// (XR) True if we can request an XR depth swapchain
+	uint32 bSupportsXRTargetManagerDepthAlloc : 1;
+
 private:
 	static FSceneTexturesConfig GlobalInstance;
 };
diff --git a/Engine/Source/Runtime/Renderer/Private/SceneTextures.cpp b/Engine/Source/Runtime/Renderer/Private/SceneTextures.cpp
--- a/Engine/Source/Runtime/Renderer/Private/SceneTextures.cpp
+++ b/Engine/Source/Runtime/Renderer/Private/SceneTextures.cpp
@@ -265,13 +265,16 @@ static IStereoRenderTargetManager* FindStereoRenderTargetManager()
 	return GEngine->StereoRenderingDevice->GetRenderTargetManager();
 }
 
-static TRefCountPtr<FRHITexture2D> FindStereoDepthTexture(FIntPoint TextureExtent, uint32 NumSamples)
+static TRefCountPtr<FRHITexture2D> FindStereoDepthTexture(uint32 bSupportsXRDepth, FIntPoint TextureExtent, uint32 NumSamples)
 {
-	if (IStereoRenderTargetManager* StereoRenderTargetManager = FindStereoRenderTargetManager())
+	if (bSupportsXRDepth == 1)
 	{
-		TRefCountPtr<FRHITexture2D> DepthTex, SRTex;
-		StereoRenderTargetManager->AllocateDepthTexture(0, TextureExtent.X, TextureExtent.Y, PF_DepthStencil, 1, TexCreate_None, TexCreate_DepthStencilTargetable | TexCreate_ShaderResource | TexCreate_InputAttachmentRead, DepthTex, SRTex, NumSamples);
-		return MoveTemp(SRTex);
+		if (IStereoRenderTargetManager* StereoRenderTargetManager = FindStereoRenderTargetManager())
+		{
+			TRefCountPtr<FRHITexture2D> DepthTex, SRTex;
+			StereoRenderTargetManager->AllocateDepthTexture(0, TextureExtent.X, TextureExtent.Y, PF_DepthStencil, 1, TexCreate_None, TexCreate_DepthStencilTargetable | TexCreate_ShaderResource | TexCreate_InputAttachmentRead, DepthTex, SRTex, NumSamples);
+			return MoveTemp(SRTex);
+		}
 	}
 	return nullptr;
 }
@@ -541,6 +544,14 @@ void InitializeSceneTexturesConfig(FSceneTexturesConfig& Config, const FSceneVie
 	Config.bRequireMultiView = ViewFamily.bRequireMultiView;
 	Config.bIsUsingGBuffers = IsUsingGBuffers(Config.ShaderPlatform);
 
+	{
+		const bool bNeedsStereoAlloc = ViewFamily.Views.ContainsByPredicate([](const FSceneView* View)
+			{
+				return (IStereoRendering::IsStereoEyeView(*View) && (FindStereoRenderTargetManager() != nullptr));
+			});
+		Config.bSupportsXRTargetManagerDepthAlloc = bNeedsStereoAlloc ? 1 : 0;
+	}
+
 	switch (Config.ShadingPath)
 	{
 	case EShadingPath::Deferred:
@@ -616,7 +627,7 @@ void FMinimalSceneTextures::InitializeViewFamily(FRDGBuilder& GraphBuilder, FVie
 
 	// If not using MSAA, we need to make sure to grab the stereo depth texture if appropriate.
 	FTexture2DRHIRef StereoDepthRHI;
-	if (Config.NumSamples == 1 && (StereoDepthRHI = FindStereoDepthTexture(Config.Extent, Config.NumSamples)) != nullptr)
+	if (Config.NumSamples == 1 && (StereoDepthRHI = FindStereoDepthTexture(Config.bSupportsXRTargetManagerDepthAlloc, Config.Extent, Config.NumSamples)) != nullptr)
 	{
 		SceneTextures.Depth = RegisterExternalTexture(GraphBuilder, StereoDepthRHI, TEXT("SceneDepthZ"));
 		SceneTextures.Stencil = GraphBuilder.CreateSRV(FRDGTextureSRVDesc::CreateWithPixelFormat(SceneTextures.Depth.Target, PF_X24_G8));
@@ -646,7 +657,7 @@ void FMinimalSceneTextures::InitializeViewFamily(FRDGBuilder& GraphBuilder, FVie
 		{
 			Desc.NumSamples = 1;
 
-			if ((StereoDepthRHI = FindStereoDepthTexture(Config.Extent, Desc.NumSamples)) != nullptr)
+			if ((StereoDepthRHI = FindStereoDepthTexture(Config.bSupportsXRTargetManagerDepthAlloc, Config.Extent, Desc.NumSamples)) != nullptr)
 			{
 				SceneTextures.Depth.Resolve = RegisterExternalTexture(GraphBuilder, StereoDepthRHI, TEXT("SceneDepthZ"));
 			}
