1. Remove all flip-y completely and only flip y at the end on androidl GL with a native window api, no need to handle flip-y on high level.
2. Supports CustomDepth in translucency pass on mobile LDR.

diff --git a/Engine/Plugins/Experimental/Water/Source/Runtime/Private/WaterInfoRendering.cpp b/Engine/Plugins/Experimental/Water/Source/Runtime/Private/WaterInfoRendering.cpp
--- a/Engine/Plugins/Experimental/Water/Source/Runtime/Private/WaterInfoRendering.cpp
+++ b/Engine/Plugins/Experimental/Water/Source/Runtime/Private/WaterInfoRendering.cpp
@@ -439,15 +439,13 @@ static void UpdateWaterInfoRendering_RenderThread(
 		if (DepthRenderer->Scene->GetShadingPath() == EShadingPath::Mobile)
 		{
 			const FMinimalSceneTextures& SceneTextures = View.GetSceneTextures();
-			const bool bNeedsFlippedRenderTarget = false;
 			RDG_EVENT_SCOPE(GraphBuilder, "CaptureSceneColor");
 			CopySceneCaptureComponentToTarget(
 				GraphBuilder,
 				SceneTextures,
 				DepthTexture,
 				DepthRenderer->ViewFamily,
-				DepthRenderer->Views,
-				bNeedsFlippedRenderTarget);
+				DepthRenderer->Views);
 		}
 		else
 		{
@@ -500,15 +498,13 @@ static void UpdateWaterInfoRendering_RenderThread(
 
 		// This CopySceneCaptureComponentToTarget is required on all platforms as it extracts a higher pixel depth texture than the scene render target so we can't just copy out of it
 		{
-			const bool bNeedsFlippedRenderTarget = false;
 			RDG_EVENT_SCOPE(GraphBuilder, "CaptureSceneColor");
 			CopySceneCaptureComponentToTarget(
 				GraphBuilder,
 				SceneTextures,
 				ColorTexture,
 				ColorRenderer->ViewFamily,
-				ColorRenderer->Views,
-				bNeedsFlippedRenderTarget);
+				ColorRenderer->Views);
 		}
 
 		// We currently can't have multiple scene renderers run within the same RDGBuilder. Therefore, we must
@@ -556,15 +552,13 @@ static void UpdateWaterInfoRendering_RenderThread(
 
 		if (DilationRenderer->Scene->GetShadingPath() == EShadingPath::Mobile)
 		{
-			const bool bNeedsFlippedRenderTarget = false;
 			RDG_EVENT_SCOPE(GraphBuilder, "CaptureSceneColor");
 			CopySceneCaptureComponentToTarget(
 				GraphBuilder,
 				SceneTextures,
 				DilationTexture,
 				DilationRenderer->ViewFamily,
-				DilationRenderer->Views,
-				bNeedsFlippedRenderTarget);
+				DilationRenderer->Views);
 		}
 		else
 		{
diff --git a/Engine/Plugins/FX/Niagara/Source/Niagara/Private/NiagaraGpuComputeDebug.cpp b/Engine/Plugins/FX/Niagara/Source/Niagara/Private/NiagaraGpuComputeDebug.cpp
--- a/Engine/Plugins/FX/Niagara/Source/Niagara/Private/NiagaraGpuComputeDebug.cpp
+++ b/Engine/Plugins/FX/Niagara/Source/Niagara/Private/NiagaraGpuComputeDebug.cpp
@@ -274,7 +274,6 @@ void FNiagaraGpuComputeDebug::DrawDebug(FRDGBuilder& GraphBuilder, const FViewIn
 		AddDrawCanvasPass(GraphBuilder, {}, View, Output,
 			[Location, SourceName=VisualizeEntry.SourceName.ToString(), SystemName, Font](FCanvas& Canvas)
 			{
-				Canvas.SetAllowSwitchVerticalAxis(true);
 				Canvas.DrawShadowedString(Location.X, Location.Y, *FString::Printf(TEXT("DataInterface: %s, System: %s"), *SourceName, *SystemName), Font, FLinearColor(1, 1, 1));
 			}
 		);
diff --git a/Engine/Plugins/Runtime/Oculus/OculusVR/Source/OculusMR/Private/OculusMR_CastingCameraActor.cpp b/Engine/Plugins/Runtime/Oculus/OculusVR/Source/OculusMR/Private/OculusMR_CastingCameraActor.cpp
--- a/Engine/Plugins/Runtime/Oculus/OculusVR/Source/OculusMR/Private/OculusMR_CastingCameraActor.cpp
+++ b/Engine/Plugins/Runtime/Oculus/OculusVR/Source/OculusMR/Private/OculusMR_CastingCameraActor.cpp
@@ -939,7 +939,6 @@ void AOculusMR_CastingCameraActor::SetupMRCScreen()
 #endif
 		UpdateRenderTargetSize();
 
-		GetCaptureComponent2D()->bDisableFlipCopyGLES = true;
 		// LDR for gamma correction and post process
 		GetCaptureComponent2D()->CaptureSource = ESceneCaptureSource::SCS_FinalColorLDR;
 
@@ -952,7 +951,6 @@ void AOculusMR_CastingCameraActor::SetupMRCScreen()
 		{
 			ForegroundCaptureActor = GetWorld()->SpawnActor<ASceneCapture2D>();
 
-			ForegroundCaptureActor->GetCaptureComponent2D()->bDisableFlipCopyGLES = true;
 			// LDR for gamma correction and post process
 			ForegroundCaptureActor->GetCaptureComponent2D()->CaptureSource = ESceneCaptureSource::SCS_FinalColorLDR;
 #if PLATFORM_ANDROID
diff --git a/Engine/Shaders/Private/MobileBasePassPixelShader.usf b/Engine/Shaders/Private/MobileBasePassPixelShader.usf
--- a/Engine/Shaders/Private/MobileBasePassPixelShader.usf
+++ b/Engine/Shaders/Private/MobileBasePassPixelShader.usf
@@ -275,11 +275,6 @@ void Main(
 	VertexFactoryUnpackInterpolants(Interpolants, PackedInterpolants);
 #endif
 
-#if COMPILER_GLSL_ES3_1 && !OUTPUT_MOBILE_HDR && !MOBILE_EMULATION
-	// LDR Mobile needs screen vertical flipped
-	SvPosition.y = ResolvedView.BufferSizeAndInvSize.y - SvPosition.y - 1;
-#endif
-
 	FMaterialPixelParameters MaterialParameters = GetMaterialPixelParameters(Interpolants, SvPosition);
 	FPixelMaterialInputs PixelMaterialInputs;
 	{
diff --git a/Engine/Shaders/Private/MobileBasePassVertexShader.usf b/Engine/Shaders/Private/MobileBasePassVertexShader.usf
--- a/Engine/Shaders/Private/MobileBasePassVertexShader.usf
+++ b/Engine/Shaders/Private/MobileBasePassVertexShader.usf
@@ -125,8 +125,4 @@ void Main(
 #if PACK_INTERPOLANTS
 	VertexFactoryPackInterpolants(Output.FactoryInterpolants, PackedInterps);
 #endif // PACK_INTERPOLANTS
-
-#if !OUTPUT_MOBILE_HDR && COMPILER_GLSL_ES3_1
-	Output.Position.y *= -1;
-#endif
 }
diff --git a/Engine/Shaders/Private/OcclusionQueryVertexShader.usf b/Engine/Shaders/Private/OcclusionQueryVertexShader.usf
--- a/Engine/Shaders/Private/OcclusionQueryVertexShader.usf
+++ b/Engine/Shaders/Private/OcclusionQueryVertexShader.usf
@@ -37,8 +37,4 @@ void Main(
 
 	float3 TransformedPosition = InPosition.xyz * StencilingGeometryPosAndScale.w + StencilingGeometryPosAndScale.xyz;
 	OutPosition = mul(float4(TransformedPosition, 1), View.TranslatedWorldToClip);
-
-#if !OUTPUT_MOBILE_HDR && COMPILER_GLSL_ES3_1
-	OutPosition.y *= -1;
-#endif
 }
diff --git a/Engine/Shaders/Private/PostProcessMaterialShaders.usf b/Engine/Shaders/Private/PostProcessMaterialShaders.usf
--- a/Engine/Shaders/Private/PostProcessMaterialShaders.usf
+++ b/Engine/Shaders/Private/PostProcessMaterialShaders.usf
@@ -279,8 +279,6 @@ void MainPS_VideoOverlay(
 
 #else // FEATURE_LEVEL > FEATURE_LEVEL_ES3_1
 
-uint bFlipYAxis;
-
 //
 // Mobile version
 //
@@ -329,10 +327,6 @@ void MainPS(
 
 	// can be optimized
 	float4 SvPosition = Input.Position;
-	if (bFlipYAxis)
-	{
-		SvPosition.y = PostProcessOutput_ViewportSize.y - SvPosition.y;
-	}
 	float2 ViewportUV = (SvPosition - PostProcessOutput_ViewportMin) * PostProcessOutput_ViewportSizeInverse;
 	float2 BufferUV = ViewportUVToBufferUV(ViewportUV);
 	
diff --git a/Engine/Shaders/Private/PostProcessTonemap.usf b/Engine/Shaders/Private/PostProcessTonemap.usf
--- a/Engine/Shaders/Private/PostProcessTonemap.usf
+++ b/Engine/Shaders/Private/PostProcessTonemap.usf
@@ -181,12 +181,6 @@ half3 ColorLookupTable( half3 LinearColor )
 	return OutDeviceColor * 1.05;
 }
 
-float SwitchVerticalAxis;
-
-#ifndef NEEDTOSWITCHVERTICLEAXIS
-#define NEEDTOSWITCHVERTICLEAXIS 0
-#endif
-
 // can be optimized
 float2 ScreenPosToUV(float2 ScreenPos, float2 ExtentInverse)
 {
@@ -215,10 +209,6 @@ void TonemapCommonVS(
 	// TODO: Looks to be Ronin specific.. check this:
 	OutFullViewUV.xy = Position.xy * float2(0.5,-0.5) + 0.5;
 
-#if NEEDTOSWITCHVERTICLEAXIS
-	OutFullViewUV.y = 1.0 - OutFullViewUV.y;
-#endif
-
 #if EYEADAPTATION_EXPOSURE_FIX == 1
 	#if FEATURE_LEVEL >= FEATURE_LEVEL_SM5
 		// texture can be GWhiteTexture which is 1x1. It's important we don't read outside bounds.
@@ -262,11 +252,6 @@ void MainVS(
 	out 				float4 OutPosition 				: SV_POSITION
 	)
 {
-#if NEEDTOSWITCHVERTICLEAXIS
-	// This is currently the last pass, so flip the texture on V to match D3D
-	InTexCoord.y = 1.0 - InTexCoord.y;
-#endif
-
 	DrawRectangle(InPosition, InTexCoord, OutPosition, OutTexCoord);
     TonemapCommonVS(OutPosition, OutTexCoord, OutExposureScaleVignette, OutGrainUV, OutScreenPos, OutFullViewUV);
 }
@@ -751,9 +736,7 @@ void MainCS(uint2 DispatchThreadId : SV_DispatchThreadID)
 	float4 OutColor = TonemapCommonPS(UV, ExposureScaleVignette, GrainUV, ScreenPos, FullViewUV, SvPosition);
 
 	uint2 PixelPos = DispatchThreadId + Output_ViewportMin;
-#if NEEDTOSWITCHVERTICLEAXIS
-	PixelPos.y = (uint)(Output_Extent.y) - PixelPos.y;
-#endif
+
 	RWOutputTexture[PixelPos] = OutColor;
 }
 #endif
diff --git a/Engine/Shaders/Private/SimpleElementVertexShader.usf b/Engine/Shaders/Private/SimpleElementVertexShader.usf
--- a/Engine/Shaders/Private/SimpleElementVertexShader.usf
+++ b/Engine/Shaders/Private/SimpleElementVertexShader.usf
@@ -6,10 +6,6 @@
 
 #include "Common.ush"
 
-#ifndef ALLOW_SWITCH_VERTICALAXIS
-#define ALLOW_SWITCH_VERTICALAXIS 0
-#endif
-
 #ifndef ENABLE_LWC
 #define ENABLE_LWC 0
 #endif
@@ -19,8 +15,6 @@ float4x4 Transform;
 float3 TransformTilePosition;
 #endif
 
-float SwitchVerticalAxis;
-
 void Main(
 	in float4 InRelativePosition : ATTRIBUTE0,
 	in float4 InTilePosition : ATTRIBUTE1,
@@ -42,10 +36,6 @@ void Main(
 	OutPosition = mul(InRelativePosition, Transform);
 #endif
 
-#if ALLOW_SWITCH_VERTICALAXIS
-	OutPosition.y *= SwitchVerticalAxis;
-#endif
-
 	OutTextureCoordinate = InTextureCoordinate;
 	OutColor = InColor;
 	OutHitProxyId = InHitProxyId;
diff --git a/Engine/Shaders/Private/SlateMaskingShader.usf b/Engine/Shaders/Private/SlateMaskingShader.usf
--- a/Engine/Shaders/Private/SlateMaskingShader.usf
+++ b/Engine/Shaders/Private/SlateMaskingShader.usf
@@ -11,8 +11,6 @@ float4x4 ViewProjection;
 // MaskRect[1].zw = BottomRight
 float4 MaskRectPacked[2];
 
-float SwitchVerticalAxisMultiplier;
-
 struct VertexToPixelInterpolants
 {
     float4 Position : SV_POSITION;
@@ -28,11 +26,6 @@ VertexToPixelInterpolants MainVS(in uint InPosition : ATTRIBUTE0)
 
 	VOut.Position = mul(WorldPosition, ViewProjection);
 
-#if (ES3_1_PROFILE && COMPILER_GLSL_ES3_1)
-	// @todo-mobile: Fix this in the projection matrix
-	VOut.Position.y *= SwitchVerticalAxisMultiplier;
-#endif //(ES3_1_PROFILE && COMPILER_GLSL_ES3_1)
-
 	return VOut;
 }
 
diff --git a/Engine/Shaders/Private/SlateVertexShader.usf b/Engine/Shaders/Private/SlateVertexShader.usf
--- a/Engine/Shaders/Private/SlateVertexShader.usf
+++ b/Engine/Shaders/Private/SlateVertexShader.usf
@@ -6,8 +6,6 @@
 
 float4x4 ViewProjection;
 
-float SwitchVerticalAxisMultiplier;
-
 VertexToPixelInterpolants Main(
 	in float4 InTextureCoordinates : ATTRIBUTE0,
 	in float2 InMaterialTextureCoordinates : ATTRIBUTE1,
@@ -104,10 +102,5 @@ VertexToPixelInterpolants Main(
 
 	VOut.Position = mul(WorldPosition, ViewProjection);
 
-#if (ES3_1_PROFILE && COMPILER_GLSL_ES3_1)
-	// @todo-mobile: Fix this in the projection matrix
-	VOut.Position.y *= SwitchVerticalAxisMultiplier;
-#endif //(ES3_1_PROFILE && COMPILER_GLSL_ES3_1)
-
 	return VOut;
 }
diff --git a/Engine/Source/Developer/MaterialBaking/Private/MaterialRenderItem.cpp b/Engine/Source/Developer/MaterialBaking/Private/MaterialRenderItem.cpp
--- a/Engine/Source/Developer/MaterialBaking/Private/MaterialRenderItem.cpp
+++ b/Engine/Source/Developer/MaterialBaking/Private/MaterialRenderItem.cpp
@@ -47,9 +47,6 @@ bool FMeshMaterialRenderItem::Render_RenderThread(FCanvasRenderContext& RenderCo
 	View.FinalPostProcessSettings.bOverride_IndirectLightingIntensity = 1;
 	View.FinalPostProcessSettings.IndirectLightingIntensity = 0.0f;
 
-	const bool bNeedsToSwitchVerticalAxis = RHINeedsToSwitchVerticalAxis(Canvas->GetShaderPlatform()) && !Canvas->GetAllowSwitchVerticalAxis();
-	check(bNeedsToSwitchVerticalAxis == false);
-
 	if (Vertices.Num() && Indices.Num())
 	{
 		FMeshPassProcessorRenderState LocalDrawRenderState;
diff --git a/Engine/Source/Developer/MaterialUtilities/Private/MeshRendering.cpp b/Engine/Source/Developer/MaterialUtilities/Private/MeshRendering.cpp
--- a/Engine/Source/Developer/MaterialUtilities/Private/MeshRendering.cpp
+++ b/Engine/Source/Developer/MaterialUtilities/Private/MeshRendering.cpp
@@ -540,9 +540,6 @@ public:
 		ViewInitOptions.BackgroundColor = FLinearColor::Black;
 		ViewInitOptions.OverlayColor = FLinearColor::White;
 
-		bool bNeedsToSwitchVerticalAxis = RHINeedsToSwitchVerticalAxis(Canvas->GetShaderPlatform()) && !Canvas->GetAllowSwitchVerticalAxis();
-		check(bNeedsToSwitchVerticalAxis == false);
-
 		FSceneView* View = new FSceneView(ViewInitOptions);
 
 		RenderMaterial(RenderContext, DrawRenderState, *View, *Data);
@@ -576,8 +573,6 @@ public:
 
 		FSceneView* View = new FSceneView(ViewInitOptions);
 
-		bool bNeedsToSwitchVerticalAxis = RHINeedsToSwitchVerticalAxis(Canvas->GetShaderPlatform()) && !Canvas->GetAllowSwitchVerticalAxis();
-		check(bNeedsToSwitchVerticalAxis == false);
 		struct FDrawMaterialParameters
 		{
 			FSceneView* View;
diff --git a/Engine/Source/Editor/StaticMeshEditor/Private/StaticMeshEditorViewportClient.cpp b/Engine/Source/Editor/StaticMeshEditor/Private/StaticMeshEditorViewportClient.cpp
--- a/Engine/Source/Editor/StaticMeshEditor/Private/StaticMeshEditorViewportClient.cpp
+++ b/Engine/Source/Editor/StaticMeshEditor/Private/StaticMeshEditorViewportClient.cpp
@@ -929,8 +929,7 @@ void FStaticMeshEditorViewportClient::DrawCanvas( FViewport& InViewport, FSceneV
 
  	int32 X = Canvas.GetRenderTarget()->GetSizeXY().X - 300;
  	int32 Y = 30;
-	// Make sure draws to the canvas are not rendered upside down.
-	Canvas.SetAllowSwitchVerticalAxis(false);
+
 	if (StaticMesh->GetBodySetup() && (!(StaticMesh->GetBodySetup()->bHasCookedCollisionData || StaticMesh->GetBodySetup()->bNeverNeedsCookedCollisionData) || StaticMesh->GetBodySetup()->bFailedToCreatePhysicsMeshes))
 	{
 		static const FText Message = NSLOCTEXT("Renderer", "NoCookedCollisionObject", "NO COOKED COLLISION OBJECT: TOO SMALL?");
diff --git a/Engine/Source/Runtime/ApplicationCore/Private/Android/AndroidApplication.cpp b/Engine/Source/Runtime/ApplicationCore/Private/Android/AndroidApplication.cpp
--- a/Engine/Source/Runtime/ApplicationCore/Private/Android/AndroidApplication.cpp
+++ b/Engine/Source/Runtime/ApplicationCore/Private/Android/AndroidApplication.cpp
@@ -99,46 +99,17 @@ void FAndroidApplication::PollGameDeviceState( const float TimeDelta )
 
 void FAndroidApplication::HandleDeviceOrientation()
 {
-#if USE_ANDROID_JNI
-    JNIEnv* JEnv = AndroidJavaEnv::GetJavaEnv();
-    if (JEnv)
-    {
-        static jmethodID getOrientationMethod = 0;
-        const auto PreviousDeviceOrientation = DeviceOrientation;
-        
-        if (getOrientationMethod == 0)
-        {
-            jclass MainClass = AndroidJavaEnv::FindJavaClassGlobalRef("com/epicgames/unreal/GameActivity");
-            if (MainClass != nullptr)
-            {
-                getOrientationMethod = JEnv->GetMethodID(MainClass, "AndroidThunkJava_GetDeviceOrientation", "()I");
-                JEnv->DeleteGlobalRef(MainClass);
-            }
-        }
-        
-        if (getOrientationMethod != 0)
-        {
-            const int Orientation = JEnv->CallIntMethod(AndroidJavaEnv::GetGameActivityThis(), getOrientationMethod);
-            switch (Orientation)
-            {
-                case 0: DeviceOrientation = EDeviceScreenOrientation::Portrait;             break;
-                case 1: DeviceOrientation = EDeviceScreenOrientation::LandscapeLeft;        break;
-                case 2: DeviceOrientation = EDeviceScreenOrientation::PortraitUpsideDown;   break;
-                case 3: DeviceOrientation = EDeviceScreenOrientation::LandscapeRight;       break;
-            }
-        }
-        
-        FAndroidMisc::SaveDeviceOrientation(DeviceOrientation);
-        
-        if (PreviousDeviceOrientation != DeviceOrientation)
-        {
-            FCoreDelegates::ApplicationReceivedScreenOrientationChangedNotificationDelegate.Broadcast((int32)DeviceOrientation);
-               
-            //we also want to fire off the safe frame event
-            FCoreDelegates::OnSafeFrameChangedEvent.Broadcast();
-        }
-    }
-#endif
+	FAndroidMisc::UpdateDeviceOrientation();
+	const auto PreviousDeviceOrientation = DeviceOrientation;
+	DeviceOrientation = FAndroidMisc::GetDeviceOrientation();
+
+	if (PreviousDeviceOrientation != DeviceOrientation)
+	{
+		FCoreDelegates::ApplicationReceivedScreenOrientationChangedNotificationDelegate.Broadcast((int32)DeviceOrientation);
+
+		//we also want to fire off the safe frame event
+		FCoreDelegates::OnSafeFrameChangedEvent.Broadcast();
+	}
 }
 
 FPlatformRect FAndroidApplication::GetWorkArea( const FPlatformRect& CurrentWindow ) const
diff --git a/Engine/Source/Runtime/Core/Private/Android/AndroidPlatformMisc.cpp b/Engine/Source/Runtime/Core/Private/Android/AndroidPlatformMisc.cpp
--- a/Engine/Source/Runtime/Core/Private/Android/AndroidPlatformMisc.cpp
+++ b/Engine/Source/Runtime/Core/Private/Android/AndroidPlatformMisc.cpp
@@ -591,6 +591,40 @@ void FAndroidMisc::PlatformTearDown()
 	RemoveBinding(FCoreDelegates::ApplicationHasEnteredForegroundDelegate, AndroidOnForegroundBinding);
 }
 
+void FAndroidMisc::UpdateDeviceOrientation()
+{
+	QUICK_SCOPE_CYCLE_COUNTER(STAT_FAndroidMisc_UpdateDeviceOrientation);
+#if USE_ANDROID_JNI
+	JNIEnv* JEnv = AndroidJavaEnv::GetJavaEnv();
+	if (JEnv)
+	{
+		static jmethodID getOrientationMethod = 0;
+
+		if (getOrientationMethod == 0)
+		{
+			jclass MainClass = AndroidJavaEnv::FindJavaClassGlobalRef("com/epicgames/unreal/GameActivity");
+			if (MainClass != nullptr)
+			{
+				getOrientationMethod = JEnv->GetMethodID(MainClass, "AndroidThunkJava_GetDeviceOrientation", "()I");
+				JEnv->DeleteGlobalRef(MainClass);
+			}
+		}
+
+		if (getOrientationMethod != 0)
+		{
+			const int Orientation = JEnv->CallIntMethod(AndroidJavaEnv::GetGameActivityThis(), getOrientationMethod);
+			switch (Orientation)
+			{
+			case 0: DeviceOrientation = EDeviceScreenOrientation::Portrait;             break;
+			case 1: DeviceOrientation = EDeviceScreenOrientation::LandscapeLeft;        break;
+			case 2: DeviceOrientation = EDeviceScreenOrientation::PortraitUpsideDown;   break;
+			case 3: DeviceOrientation = EDeviceScreenOrientation::LandscapeRight;       break;
+			}
+		}
+	}
+#endif
+}
+
 void FAndroidMisc::PlatformHandleSplashScreen(bool ShowSplashScreen)
 {
 #if USE_ANDROID_JNI
@@ -598,6 +632,8 @@ void FAndroidMisc::PlatformHandleSplashScreen(bool ShowSplashScreen)
 	{
 		AndroidThunkCpp_DismissSplashScreen();
 	}
+	// Update the device orientation in case the game thread is blocked
+	FAndroidMisc::UpdateDeviceOrientation();
 #endif
 }
 
diff --git a/Engine/Source/Runtime/Core/Private/HAL/ConsoleManager.cpp b/Engine/Source/Runtime/Core/Private/HAL/ConsoleManager.cpp
--- a/Engine/Source/Runtime/Core/Private/HAL/ConsoleManager.cpp
+++ b/Engine/Source/Runtime/Core/Private/HAL/ConsoleManager.cpp
@@ -2923,14 +2923,6 @@ static TAutoConsoleVariable<int32> CVarSetVSyncEditorEnabled(
 	TEXT("0: VSync is disabled in editor.(default)\n"
 		 "1: VSync is enabled in editor."),
 	ECVF_RenderThreadSafe);
-
-static TAutoConsoleVariable<int32> CVarMobileForceRHISwitchVerticalAxis(
-	TEXT("r.Mobile.ForceRHISwitchVerticalAxis"),
-	0,
-	TEXT("Enable RHISwitchVerticalAxis when previewing mobile renderer. (Useful to test GLES y-axis flip codepaths)\n"
-		 "0: RHISwitchVerticalAxis disabled (default).\n"
-		 "1: RHISwitchVerticalAxis enabled.\n"),
-	ECVF_Scalability | ECVF_RenderThreadSafe);
 #endif
 
 static TAutoConsoleVariable<int32> CVarFinishCurrentFrame(
diff --git a/Engine/Source/Runtime/Core/Public/Android/AndroidPlatformMisc.h b/Engine/Source/Runtime/Core/Public/Android/AndroidPlatformMisc.h
--- a/Engine/Source/Runtime/Core/Public/Android/AndroidPlatformMisc.h
+++ b/Engine/Source/Runtime/Core/Public/Android/AndroidPlatformMisc.h
@@ -279,6 +279,8 @@ public:
 	// Doesn't guarantee to work on all devices. Some devices require root access rights to read sensors information, in that case 0.0 will be returned
 	static float GetCPUTemperature();
 
+	static void UpdateDeviceOrientation();
+
 	static void SaveDeviceOrientation(EDeviceScreenOrientation NewDeviceOrentation) { DeviceOrientation = NewDeviceOrentation; }
 
 	// Window access is locked by the game thread before preinit and unlocked here after RHIInit (PlatformCreateDynamicRHI). 
diff --git a/Engine/Source/Runtime/Engine/Classes/Components/SceneCaptureComponent2D.h b/Engine/Source/Runtime/Engine/Classes/Components/SceneCaptureComponent2D.h
--- a/Engine/Source/Runtime/Engine/Classes/Components/SceneCaptureComponent2D.h
+++ b/Engine/Source/Runtime/Engine/Classes/Components/SceneCaptureComponent2D.h
@@ -110,14 +110,6 @@ public:
 	UPROPERTY(EditAnywhere, BlueprintReadWrite, Category = SceneCapture)
 	uint32 bConsiderUnrenderedOpaquePixelAsFullyTranslucent : 1;
 
-	/** 
-	 * Scene captures render an extra flip pass for LDR color on GLES so that the final output is oriented correctly.
-	 * This check disabled the extra flip pass, improving performance, but causes the capture to be flipped vertically.
-	 * (Does not affect scene captures on other non-GLES renderers or with non-LDR output)
-	 */
-	UPROPERTY(EditAnywhere, BlueprintReadWrite, AdvancedDisplay, Category = SceneCapture, meta = (DisplayName = "Disable Flip Copy GLES"))
-	bool bDisableFlipCopyGLES;
-
 	/** Array of scene view extensions specifically to apply to this scene capture */
 	TArray< TWeakPtr<ISceneViewExtension, ESPMode::ThreadSafe> > SceneViewExtensions;
 
diff --git a/Engine/Source/Runtime/Engine/Private/BatchedElements.cpp b/Engine/Source/Runtime/Engine/Private/BatchedElements.cpp
--- a/Engine/Source/Runtime/Engine/Private/BatchedElements.cpp
+++ b/Engine/Source/Runtime/Engine/Private/BatchedElements.cpp
@@ -462,7 +462,6 @@ void FBatchedElements::PrepareShaders(
 	ERHIFeatureLevel::Type FeatureLevel,
 	ESimpleElementBlendMode BlendMode,
 	const FRelativeViewMatrices& ViewMatrices,
-	bool bSwitchVerticalAxis,
 	FBatchedElementParameters* BatchedElementParameters,
 	const FTexture* Texture,
 	bool bHitTesting,
@@ -715,7 +714,7 @@ void FBatchedElements::PrepareShaders(
 		}
 
 		// Set the simple element vertex shader parameters
-		VertexShader->SetParameters(RHICmdList, ViewMatrices, bSwitchVerticalAxis);
+		VertexShader->SetParameters(RHICmdList, ViewMatrices);
 	}
 }
 
@@ -774,7 +773,7 @@ FSceneView FBatchedElements::CreateProxySceneView(const FMatrix& ProjectionMatri
 	return FSceneView(ProxyViewInitOptions);
 }
 
-bool FBatchedElements::Draw(FRHICommandList& RHICmdList, const FMeshPassProcessorRenderState& DrawRenderState, ERHIFeatureLevel::Type FeatureLevel, bool bNeedToSwitchVerticalAxis, const FSceneView& View, bool bHitTesting, float Gamma /* = 1.0f */, EBlendModeFilter::Type Filter /* = EBlendModeFilter::All */) const
+bool FBatchedElements::Draw(FRHICommandList& RHICmdList, const FMeshPassProcessorRenderState& DrawRenderState, ERHIFeatureLevel::Type FeatureLevel, const FSceneView& View, bool bHitTesting, float Gamma /* = 1.0f */, EBlendModeFilter::Type Filter /* = EBlendModeFilter::All */) const
 {
 	const FRelativeViewMatrices RelativeMatrices = FRelativeViewMatrices::Create(View.ViewMatrices);
 	const FMatrix& WorldToClip = View.ViewMatrices.GetViewProjectionMatrix();
@@ -815,7 +814,7 @@ bool FBatchedElements::Draw(FRHICommandList& RHICmdList, const FMeshPassProcesso
 				GraphicsPSOInit.PrimitiveType = PT_LineList;
 
 				// Set the appropriate pixel shader parameters & shader state for the non-textured elements.
-				PrepareShaders(RHICmdList, GraphicsPSOInit, StencilRef, FeatureLevel, SE_BLEND_Opaque, RelativeMatrices, bNeedToSwitchVerticalAxis, BatchedElementParameters, GWhiteTexture, bHitTesting, Gamma, NULL, &View);
+				PrepareShaders(RHICmdList, GraphicsPSOInit, StencilRef, FeatureLevel, SE_BLEND_Opaque, RelativeMatrices, BatchedElementParameters, GWhiteTexture, bHitTesting, Gamma, NULL, &View);
 
 				FRHIResourceCreateInfo CreateInfo(TEXT("Lines"));
 				FBufferRHIRef VertexBufferRHI = RHICreateVertexBuffer(sizeof(FSimpleElementVertex) * LineVertices.Num(), BUF_Volatile, CreateInfo);
@@ -847,7 +846,7 @@ bool FBatchedElements::Draw(FRHICommandList& RHICmdList, const FMeshPassProcesso
 			GraphicsPSOInit.PrimitiveType = PT_TriangleList;
 
 			// Set the appropriate pixel shader parameters & shader state for the non-textured elements.
-			PrepareShaders(RHICmdList, GraphicsPSOInit, StencilRef, FeatureLevel, SE_BLEND_Opaque, RelativeMatrices, bNeedToSwitchVerticalAxis, BatchedElementParameters, GWhiteTexture, bHitTesting, Gamma, NULL, &View);
+			PrepareShaders(RHICmdList, GraphicsPSOInit, StencilRef, FeatureLevel, SE_BLEND_Opaque, RelativeMatrices, BatchedElementParameters, GWhiteTexture, bHitTesting, Gamma, NULL, &View);
 
 			// Draw points
 			DrawPointElements(RHICmdList, WorldToClip, ViewportSizeX, ViewportSizeY, CameraX, CameraY);
@@ -884,7 +883,7 @@ bool FBatchedElements::Draw(FRHICommandList& RHICmdList, const FMeshPassProcesso
 					FRasterizerStateInitializerRHI Initializer(FM_Solid, CM_None, 0, DepthBiasThisBatch, ERasterizerDepthClipMode::DepthClip, bEnableMSAA, bEnableLineAA);
 					auto RasterState = RHICreateRasterizerState(Initializer);
 					GraphicsPSOInit.RasterizerState = RasterState.GetReference();
-					PrepareShaders(RHICmdList, GraphicsPSOInit, StencilRef, FeatureLevel, SE_BLEND_AlphaBlend, RelativeMatrices, bNeedToSwitchVerticalAxis, BatchedElementParameters, GWhiteTexture, bHitTesting, Gamma, NULL, &View);
+					PrepareShaders(RHICmdList, GraphicsPSOInit, StencilRef, FeatureLevel, SE_BLEND_AlphaBlend, RelativeMatrices, BatchedElementParameters, GWhiteTexture, bHitTesting, Gamma, NULL, &View);
 
 					FRHIResourceCreateInfo CreateInfo(TEXT("ThickLines"));
 					FBufferRHIRef VertexBufferRHI = RHICreateVertexBuffer(sizeof(FSimpleElementVertex) * 8 * 3 * NumLinesThisBatch, BUF_Volatile, CreateInfo);
@@ -1007,7 +1006,7 @@ bool FBatchedElements::Draw(FRHICommandList& RHICmdList, const FMeshPassProcesso
 					Initializer.DepthBias = DepthBias;
 					auto RasterState = RHICreateRasterizerState(Initializer);
 					GraphicsPSOInit.RasterizerState = RasterState.GetReference();
-					PrepareShaders(RHICmdList, GraphicsPSOInit, StencilRef, FeatureLevel, SE_BLEND_Opaque, RelativeMatrices, bNeedToSwitchVerticalAxis, BatchedElementParameters, GWhiteTexture, bHitTesting, Gamma, NULL, &View);
+					PrepareShaders(RHICmdList, GraphicsPSOInit, StencilRef, FeatureLevel, SE_BLEND_Opaque, RelativeMatrices, BatchedElementParameters, GWhiteTexture, bHitTesting, Gamma, NULL, &View);
 
 					int32 NumTris = MaxTri - MinTri;
 					RHICmdList.DrawPrimitive(MinTri * 3, NumTris, 1);
@@ -1104,7 +1103,7 @@ bool FBatchedElements::Draw(FRHICommandList& RHICmdList, const FMeshPassProcesso
 						const int32 SpriteNum = SpriteIndex - BatchStartIndex;
 						const int32 BaseVertex = BatchStartIndex * 6;
 						const int32 PrimCount = SpriteNum * 2;
-						PrepareShaders(RHICmdList, GraphicsPSOInit, StencilRef, FeatureLevel, CurrentBlendMode, RelativeMatrices, bNeedToSwitchVerticalAxis, BatchedElementParameters, CurrentTexture, bHitTesting, Gamma, NULL, &View, CurrentOpacityMask);
+						PrepareShaders(RHICmdList, GraphicsPSOInit, StencilRef, FeatureLevel, CurrentBlendMode, RelativeMatrices, BatchedElementParameters, CurrentTexture, bHitTesting, Gamma, NULL, &View, CurrentOpacityMask);
 						RHICmdList.SetStreamSource(0, VertexBufferRHI, 0);
 						RHICmdList.DrawPrimitive(BaseVertex, PrimCount, 1);
 
@@ -1148,7 +1147,7 @@ bool FBatchedElements::Draw(FRHICommandList& RHICmdList, const FMeshPassProcesso
 						MeshElement.Texture ? *MeshElement.Texture->GetFriendlyName() : TEXT(""));
 
 					// Set the appropriate pixel shader for the mesh.
-					PrepareShaders(RHICmdList, GraphicsPSOInit, StencilRef, FeatureLevel, MeshElement.BlendMode, RelativeMatrices, bNeedToSwitchVerticalAxis, MeshElement.BatchedElementParameters, MeshElement.Texture, bHitTesting, Gamma, &MeshElement.GlowInfo, &View);
+					PrepareShaders(RHICmdList, GraphicsPSOInit, StencilRef, FeatureLevel, MeshElement.BlendMode, RelativeMatrices, MeshElement.BatchedElementParameters, MeshElement.Texture, bHitTesting, Gamma, &MeshElement.GlowInfo, &View);
 
 					FBufferRHIRef IndexBufferRHI = RHICreateIndexBuffer(sizeof(uint16), sizeof(uint16) * MeshElement.Indices.Num(), BUF_Volatile, CreateInfo);
 					void* VoidPtr2 = RHILockBuffer(IndexBufferRHI, 0, sizeof(uint16) * MeshElement.Indices.Num(), RLM_WriteOnly);
diff --git a/Engine/Source/Runtime/Engine/Private/Components/DirectionalLightComponent.cpp b/Engine/Source/Runtime/Engine/Private/Components/DirectionalLightComponent.cpp
--- a/Engine/Source/Runtime/Engine/Private/Components/DirectionalLightComponent.cpp
+++ b/Engine/Source/Runtime/Engine/Private/Components/DirectionalLightComponent.cpp
@@ -93,9 +93,6 @@ static TAutoConsoleVariable<float> CVarRTDFDistanceScale(
 // Computes a shadow culling volume (convex hull) based on a set of 8 vertices and a light direction
 void ComputeShadowCullingVolume(bool bReverseCulling, const FVector* CascadeFrustumVerts, const FVector& LightDirection, FConvexVolume& ConvexVolumeOut, FPlane& NearPlaneOut, FPlane& FarPlaneOut) 
 {
-	// For mobile platforms that switch vertical axis and MobileHDR == false the sense of bReverseCulling is inverted.
-	bReverseCulling = XOR(bReverseCulling, (RHINeedsToSwitchVerticalAxis(GShaderPlatformForFeatureLevel[GMaxRHIFeatureLevel]) && !IsMobileHDR()));
-
 	// Pairs of plane indices from SubFrustumPlanes whose intersections
 	// form the edges of the frustum.
 	static const int32 AdjacentPlanePairs[12][2] =
diff --git a/Engine/Source/Runtime/Engine/Private/Components/SceneCaptureComponent.cpp b/Engine/Source/Runtime/Engine/Private/Components/SceneCaptureComponent.cpp
--- a/Engine/Source/Runtime/Engine/Private/Components/SceneCaptureComponent.cpp
+++ b/Engine/Source/Runtime/Engine/Private/Components/SceneCaptureComponent.cpp
@@ -512,7 +512,6 @@ USceneCaptureComponent2D::USceneCaptureComponent2D(const FObjectInitializer& Obj
 	ClipPlaneNormal = FVector(0, 0, 1);
 	bCameraCutThisFrame = false;
 	bConsiderUnrenderedOpaquePixelAsFullyTranslucent = false;
-	bDisableFlipCopyGLES = false;
 	
 	TileID = 0;
 
@@ -751,11 +750,6 @@ bool USceneCaptureComponent2D::CanEditChange(const FProperty* InProperty) const
 		{
 			return bUseCustomProjectionMatrix;
 		}
-
-		if (PropertyName == GET_MEMBER_NAME_STRING_CHECKED(USceneCaptureComponent2D, bDisableFlipCopyGLES))
-		{
-			return CaptureSource == SCS_FinalColorLDR;
-		}
 	}
 
 	return Super::CanEditChange(InProperty);
diff --git a/Engine/Source/Runtime/Engine/Private/SceneManagement.cpp b/Engine/Source/Runtime/Engine/Private/SceneManagement.cpp
--- a/Engine/Source/Runtime/Engine/Private/SceneManagement.cpp
+++ b/Engine/Source/Runtime/Engine/Private/SceneManagement.cpp
@@ -195,9 +195,6 @@ void FSimpleElementCollector::RegisterDynamicResource(FDynamicPrimitiveResource*
 
 void FSimpleElementCollector::DrawBatchedElements(FRHICommandList& RHICmdList, const FMeshPassProcessorRenderState& DrawRenderState, const FSceneView& InView, EBlendModeFilter::Type Filter, ESceneDepthPriorityGroup DepthPriorityGroup) const
 {
-	// Mobile HDR does not execute post process, so does not need to render flipped
-	const bool bNeedToSwitchVerticalAxis = RHINeedsToSwitchVerticalAxis(InView.GetShaderPlatform()) && !bIsMobileHDR;
-
 	const FBatchedElements& Elements = DepthPriorityGroup == SDPG_World ? BatchedElements : TopBatchedElements;
 
 	// Draw the batched elements.
@@ -205,7 +202,6 @@ void FSimpleElementCollector::DrawBatchedElements(FRHICommandList& RHICmdList, c
 		RHICmdList,
 		DrawRenderState,
 		InView.GetFeatureLevel(),
-		bNeedToSwitchVerticalAxis,
 		InView,
 		InView.Family->EngineShowFlags.HitProxies,
 		1.0f,
diff --git a/Engine/Source/Runtime/Engine/Private/SceneView.cpp b/Engine/Source/Runtime/Engine/Private/SceneView.cpp
--- a/Engine/Source/Runtime/Engine/Private/SceneView.cpp
+++ b/Engine/Source/Runtime/Engine/Private/SceneView.cpp
@@ -814,17 +814,6 @@ FSceneView::FSceneView(const FSceneViewInitOptions& InitOptions)
 	auto ShaderPlatform = GShaderPlatformForFeatureLevel[FeatureLevel];
 	bool bUsingMobileRenderer = FSceneInterface::GetShadingPath(FeatureLevel) == EShadingPath::Mobile;
 
-	bool bPlatformRequiresReverseCulling = IsOpenGLPlatform(ShaderPlatform);
-	bPlatformRequiresReverseCulling = bPlatformRequiresReverseCulling || FDataDrivenShaderPlatformInfo::GetRequiresReverseCullingOnMobile(ShaderPlatform);
-	bPlatformRequiresReverseCulling = bPlatformRequiresReverseCulling && bUsingMobileRenderer;
-	bPlatformRequiresReverseCulling = bPlatformRequiresReverseCulling && !IsPCPlatform(ShaderPlatform);
-	bPlatformRequiresReverseCulling = bPlatformRequiresReverseCulling && !IsVulkanMobilePlatform(ShaderPlatform);
-
-	static auto* MobileHDRCvar = IConsoleManager::Get().FindTConsoleVariableDataInt(TEXT("r.MobileHDR"));
-	check(MobileHDRCvar);
-	const bool bSkipPostprocessing = MobileHDRCvar->GetValueOnAnyThread() == 0;
-	bReverseCulling = (bPlatformRequiresReverseCulling && bSkipPostprocessing) ? !bReverseCulling : bReverseCulling;
-
 	// Setup transformation constants to be used by the graphics hardware to transform device normalized depth samples
 	// into world oriented z.
 	InvDeviceZToWorldZTransform = CreateInvDeviceZToWorldZTransform(ProjectionMatrixUnadjustedForRHI);
diff --git a/Engine/Source/Runtime/Engine/Private/SimpleElementShaders.cpp b/Engine/Source/Runtime/Engine/Private/SimpleElementShaders.cpp
--- a/Engine/Source/Runtime/Engine/Private/SimpleElementShaders.cpp
+++ b/Engine/Source/Runtime/Engine/Private/SimpleElementShaders.cpp
@@ -19,27 +19,23 @@ FSimpleElementVS::FSimpleElementVS(const ShaderMetaType::CompiledShaderInitializ
 {
 	RelativeTransform.Bind(Initializer.ParameterMap,TEXT("Transform"), SPF_Mandatory);
 	TransformTilePosition.Bind(Initializer.ParameterMap, TEXT("TransformTilePosition"), SPF_Optional); // TransformTilePosition may be optimized out if LWC is disabled
-	SwitchVerticalAxis.Bind(Initializer.ParameterMap,TEXT("SwitchVerticalAxis"), SPF_Optional);
 }
 
-void FSimpleElementVS::SetParameters(FRHICommandList& RHICmdList, const FMatrix& WorldToClipMatrix, bool bSwitchVerticalAxis)
+void FSimpleElementVS::SetParameters(FRHICommandList& RHICmdList, const FMatrix& WorldToClipMatrix)
 {
 	SetShaderValue(RHICmdList, RHICmdList.GetBoundVertexShader(), RelativeTransform, FMatrix44f(WorldToClipMatrix));
 	SetShaderValue(RHICmdList, RHICmdList.GetBoundVertexShader(), TransformTilePosition, FVector3f::ZeroVector);
-	SetShaderValue(RHICmdList, RHICmdList.GetBoundVertexShader(), SwitchVerticalAxis, bSwitchVerticalAxis ? -1.0f : 1.0f);
 }
 
-void FSimpleElementVS::SetParameters(FRHICommandList& RHICmdList, const FRelativeViewMatrices& Matrices, bool bSwitchVerticalAxis)
+void FSimpleElementVS::SetParameters(FRHICommandList& RHICmdList, const FRelativeViewMatrices& Matrices)
 {
 	SetShaderValue(RHICmdList, RHICmdList.GetBoundVertexShader(), RelativeTransform, Matrices.RelativeWorldToClip);
 	SetShaderValue(RHICmdList, RHICmdList.GetBoundVertexShader(), TransformTilePosition, Matrices.TilePosition);
-	SetShaderValue(RHICmdList, RHICmdList.GetBoundVertexShader(), SwitchVerticalAxis, bSwitchVerticalAxis ? -1.0f : 1.0f);
 }
 
 void FSimpleElementVS::ModifyCompilationEnvironment(const FGlobalShaderPermutationParameters& Parameters, FShaderCompilerEnvironment& OutEnvironment)
 {
 	FGlobalShader::ModifyCompilationEnvironment(Parameters, OutEnvironment);
-	OutEnvironment.SetDefine(TEXT("ALLOW_SWITCH_VERTICALAXIS"), !IsMetalMobilePlatform(Parameters.Platform));
 	OutEnvironment.SetDefine(TEXT("ENABLE_LWC"), 1); // Currently SimpleElementVertexShader.usf is shared with FCubemapTexturePropertiesVS, so need separate paths for LWC vs non-LWC
 }
 
diff --git a/Engine/Source/Runtime/Engine/Private/Slate/DebugCanvas.cpp b/Engine/Source/Runtime/Engine/Private/Slate/DebugCanvas.cpp
--- a/Engine/Source/Runtime/Engine/Private/Slate/DebugCanvas.cpp
+++ b/Engine/Source/Runtime/Engine/Private/Slate/DebugCanvas.cpp
@@ -287,9 +287,6 @@ void FDebugCanvasDrawer::DrawRenderThread(FRHICommandListImmediate& RHICmdList,
 		RHICmdList.Transition(FRHITransitionInfo(RT, ERHIAccess::Unknown, ERHIAccess::RTV));
 		RenderTarget->SetRenderTargetTexture(RT);
 
-		bool bNeedToFlipVertical = RenderThreadCanvas->GetAllowSwitchVerticalAxis();
-		// Flip when rendering to the back buffer
-		RenderThreadCanvas->SetAllowSwitchVerticalAxis(true);
 		if (RenderThreadCanvas->IsScaledToRenderTarget() && RT) 
 		{
 			RenderThreadCanvas->SetRenderTargetRect( FIntRect(0, 0, (RT)->GetSizeX(), (RT)->GetSizeY()) );
@@ -300,7 +297,6 @@ void FDebugCanvasDrawer::DrawRenderThread(FRHICommandListImmediate& RHICmdList,
 		}
 
 		RenderThreadCanvas->Flush_RenderThread(RHICmdList, true);
-		RenderThreadCanvas->SetAllowSwitchVerticalAxis(bNeedToFlipVertical);
 		RenderTarget->ClearRenderTargetTexture();
 	}
 }
diff --git a/Engine/Source/Runtime/Engine/Private/TileRendering.cpp b/Engine/Source/Runtime/Engine/Private/TileRendering.cpp
--- a/Engine/Source/Runtime/Engine/Private/TileRendering.cpp
+++ b/Engine/Source/Runtime/Engine/Private/TileRendering.cpp
@@ -86,7 +86,7 @@ uint32 FCanvasTileRendererItem::FRenderData::GetNumIndices() const
 	return Tiles.Num() * CanvasTileIndexCount;
 }
 
-void FCanvasTileRendererItem::FRenderData::InitTileMesh(const FSceneView& View, bool bNeedsToSwitchVerticalAxis)
+void FCanvasTileRendererItem::FRenderData::InitTileMesh(const FSceneView& View)
 {
 	static_assert(CanvasTileVertexCount == 4, "Invalid tile tri-list size.");
 	static_assert(CanvasTileIndexCount == 6, "Invalid tile tri-list size.");
@@ -124,30 +124,15 @@ void FCanvasTileRendererItem::FRenderData::InitTileMesh(const FSceneView& View,
 		const float SizeU = Tile.SizeU;
 		const float SizeV = Tile.SizeV;
 
-		if (bNeedsToSwitchVerticalAxis)
-		{
-			StaticMeshVertexBuffers.PositionVertexBuffer.VertexPosition(FirstVertex + 0) = FVector3f(X + SizeX, View.UnscaledViewRect.Height() - (Y + SizeY), 0.0f);
-			StaticMeshVertexBuffers.PositionVertexBuffer.VertexPosition(FirstVertex + 1) = FVector3f(X, View.UnscaledViewRect.Height() - (Y + SizeY), 0.0f);
-			StaticMeshVertexBuffers.PositionVertexBuffer.VertexPosition(FirstVertex + 2) = FVector3f(X + SizeX, View.UnscaledViewRect.Height() - Y, 0.0f);
-			StaticMeshVertexBuffers.PositionVertexBuffer.VertexPosition(FirstVertex + 3) = FVector3f(X, View.UnscaledViewRect.Height() - Y, 0.0f);
-
-			StaticMeshVertexBuffers.StaticMeshVertexBuffer.SetVertexUV(FirstVertex + 0, 0, FVector2f(U + SizeU, V + SizeV));
-			StaticMeshVertexBuffers.StaticMeshVertexBuffer.SetVertexUV(FirstVertex + 1, 0, FVector2f(U, V + SizeV));
-			StaticMeshVertexBuffers.StaticMeshVertexBuffer.SetVertexUV(FirstVertex + 2, 0, FVector2f(U + SizeU, V));
-			StaticMeshVertexBuffers.StaticMeshVertexBuffer.SetVertexUV(FirstVertex + 3, 0, FVector2f(U, V));
-		}
-		else
-		{
-			StaticMeshVertexBuffers.PositionVertexBuffer.VertexPosition(FirstVertex + 0) = FVector3f(X + SizeX, Y, 0.0f);
-			StaticMeshVertexBuffers.PositionVertexBuffer.VertexPosition(FirstVertex + 1) = FVector3f(X, Y, 0.0f);
-			StaticMeshVertexBuffers.PositionVertexBuffer.VertexPosition(FirstVertex + 2) = FVector3f(X + SizeX, Y + SizeY, 0.0f);
-			StaticMeshVertexBuffers.PositionVertexBuffer.VertexPosition(FirstVertex + 3) = FVector3f(X, Y + SizeY, 0.0f);
-
-			StaticMeshVertexBuffers.StaticMeshVertexBuffer.SetVertexUV(FirstVertex + 0, 0, FVector2f(U + SizeU, V));
-			StaticMeshVertexBuffers.StaticMeshVertexBuffer.SetVertexUV(FirstVertex + 1, 0, FVector2f(U, V));
-			StaticMeshVertexBuffers.StaticMeshVertexBuffer.SetVertexUV(FirstVertex + 2, 0, FVector2f(U + SizeU, V + SizeV));
-			StaticMeshVertexBuffers.StaticMeshVertexBuffer.SetVertexUV(FirstVertex + 3, 0, FVector2f(U, V + SizeV));
-		}
+		StaticMeshVertexBuffers.PositionVertexBuffer.VertexPosition(FirstVertex + 0) = FVector3f(X + SizeX, Y, 0.0f);
+		StaticMeshVertexBuffers.PositionVertexBuffer.VertexPosition(FirstVertex + 1) = FVector3f(X, Y, 0.0f);
+		StaticMeshVertexBuffers.PositionVertexBuffer.VertexPosition(FirstVertex + 2) = FVector3f(X + SizeX, Y + SizeY, 0.0f);
+		StaticMeshVertexBuffers.PositionVertexBuffer.VertexPosition(FirstVertex + 3) = FVector3f(X, Y + SizeY, 0.0f);
+
+		StaticMeshVertexBuffers.StaticMeshVertexBuffer.SetVertexUV(FirstVertex + 0, 0, FVector2f(U + SizeU, V));
+		StaticMeshVertexBuffers.StaticMeshVertexBuffer.SetVertexUV(FirstVertex + 1, 0, FVector2f(U, V));
+		StaticMeshVertexBuffers.StaticMeshVertexBuffer.SetVertexUV(FirstVertex + 2, 0, FVector2f(U + SizeU, V + SizeV));
+		StaticMeshVertexBuffers.StaticMeshVertexBuffer.SetVertexUV(FirstVertex + 3, 0, FVector2f(U, V + SizeV));
 
 		for (int j = 0; j < CanvasTileVertexCount; j++)
 		{
@@ -177,7 +162,6 @@ void FCanvasTileRendererItem::FRenderData::RenderTiles(
 	FMeshPassProcessorRenderState& DrawRenderState,
 	const FSceneView& View,
 	bool bIsHitTesting,
-	bool bNeedsToSwitchVerticalAxis,
 	bool bUse128bitRT)
 {
 	check(IsInRenderingThread());
@@ -194,7 +178,7 @@ void FCanvasTileRendererItem::FRenderData::RenderTiles(
 
 	IRendererModule& RendererModule = GetRendererModule();
 
-	InitTileMesh(View, bNeedsToSwitchVerticalAxis);
+	InitTileMesh(View);
 
 	// We know we have at least 1 tile so prep up a new batch right away : 
 	FMeshBatch* CurrentMeshBatch = AllocTileMeshBatch(RenderContext, Tiles[0].HitProxyId);
@@ -261,9 +245,7 @@ bool FCanvasTileRendererItem::Render_RenderThread(FCanvasRenderContext& RenderCo
 
 	const FSceneView& View = *RenderContext.Alloc<const FSceneView>(ViewInitOptions);
 
-	const bool bNeedsToSwitchVerticalAxis = RHINeedsToSwitchVerticalAxis(Canvas->GetShaderPlatform()) && Canvas->GetAllowSwitchVerticalAxis(); 
-
-	Data->RenderTiles(RenderContext, DrawRenderState, View, Canvas->IsHitTesting(), bNeedsToSwitchVerticalAxis);
+	Data->RenderTiles(RenderContext, DrawRenderState, View, Canvas->IsHitTesting());
 
 	if (Canvas->GetAllowedModes() & FCanvas::Allow_DeleteOnRender)
 	{
@@ -308,7 +290,6 @@ bool FCanvasTileRendererItem::Render_GameThread(const FCanvas* Canvas, FCanvasRe
 
 		const FSceneView* View = new FSceneView(ViewInitOptions);
 
-		const bool bNeedsToSwitchVerticalAxis = RHINeedsToSwitchVerticalAxis(Canvas->GetShaderPlatform()) && Canvas->GetAllowSwitchVerticalAxis();
 		const bool bIsHitTesting = Canvas->IsHitTesting();
 		const bool bDeleteOnRender = Canvas->GetAllowedModes() & FCanvas::Allow_DeleteOnRender;
 
@@ -321,7 +302,7 @@ bool FCanvasTileRendererItem::Render_GameThread(const FCanvas* Canvas, FCanvasRe
 		}
 
 		RenderScope.EnqueueRenderCommand(
-			[LocalData = Data, View, bIsHitTesting, bNeedsToSwitchVerticalAxis, bRequiresExplicit128bitRT]
+			[LocalData = Data, View, bIsHitTesting, bRequiresExplicit128bitRT]
 			(FCanvasRenderContext& RenderContext) mutable
 		{
 			FMeshPassProcessorRenderState DrawRenderState;
@@ -329,7 +310,7 @@ bool FCanvasTileRendererItem::Render_GameThread(const FCanvas* Canvas, FCanvasRe
 			// disable depth test & writes
 			DrawRenderState.SetDepthStencilState(TStaticDepthStencilState<false, CF_Always>::GetRHI());
 
-			LocalData->RenderTiles(RenderContext, DrawRenderState, *View, bIsHitTesting, bNeedsToSwitchVerticalAxis, bRequiresExplicit128bitRT);
+			LocalData->RenderTiles(RenderContext, DrawRenderState, *View, bIsHitTesting, bRequiresExplicit128bitRT);
 
 			RenderContext.DeferredRelease(MoveTemp(LocalData));
 			RenderContext.DeferredDelete(View->Family);
diff --git a/Engine/Source/Runtime/Engine/Private/TriangleRendering.cpp b/Engine/Source/Runtime/Engine/Private/TriangleRendering.cpp
--- a/Engine/Source/Runtime/Engine/Private/TriangleRendering.cpp
+++ b/Engine/Source/Runtime/Engine/Private/TriangleRendering.cpp
@@ -75,7 +75,7 @@ uint32 FCanvasTriangleRendererItem::FRenderData::GetNumIndices() const
 	return Triangles.Num() * 3;
 }
 
-void FCanvasTriangleRendererItem::FRenderData::InitTriangleMesh(const FSceneView& View, bool bNeedsToSwitchVerticalAxis)
+void FCanvasTriangleRendererItem::FRenderData::InitTriangleMesh(const FSceneView& View)
 {
 	const uint32 NumIndices = GetNumIndices();
 	const uint32 NumVertices = GetNumIndices();
@@ -101,18 +101,9 @@ void FCanvasTriangleRendererItem::FRenderData::InitTriangleMesh(const FSceneView
 		const FCanvasUVTri& Tri = Triangles[i].Tri;
 
 		// create verts. Notice the order is (1, 0, 2)
-		if (bNeedsToSwitchVerticalAxis)
-		{
-			StaticMeshVertexBuffers.PositionVertexBuffer.VertexPosition(StartIndex + 0) = FVector3f(Tri.V1_Pos.X, View.UnscaledViewRect.Height() - Tri.V1_Pos.Y, 0.0f);
-			StaticMeshVertexBuffers.PositionVertexBuffer.VertexPosition(StartIndex + 1) = FVector3f(Tri.V0_Pos.X, View.UnscaledViewRect.Height() - Tri.V0_Pos.Y, 0.0f);
-			StaticMeshVertexBuffers.PositionVertexBuffer.VertexPosition(StartIndex + 2) = FVector3f(Tri.V2_Pos.X, View.UnscaledViewRect.Height() - Tri.V2_Pos.Y, 0.0f);
-		}
-		else
-		{
-			StaticMeshVertexBuffers.PositionVertexBuffer.VertexPosition(StartIndex + 0) = FVector3f(Tri.V1_Pos.X, Tri.V1_Pos.Y, 0.0f);
-			StaticMeshVertexBuffers.PositionVertexBuffer.VertexPosition(StartIndex + 1) = FVector3f(Tri.V0_Pos.X, Tri.V0_Pos.Y, 0.0f);
-			StaticMeshVertexBuffers.PositionVertexBuffer.VertexPosition(StartIndex + 2) = FVector3f(Tri.V2_Pos.X, Tri.V2_Pos.Y, 0.0f);
-		}
+		StaticMeshVertexBuffers.PositionVertexBuffer.VertexPosition(StartIndex + 0) = FVector3f(Tri.V1_Pos.X, Tri.V1_Pos.Y, 0.0f);
+		StaticMeshVertexBuffers.PositionVertexBuffer.VertexPosition(StartIndex + 1) = FVector3f(Tri.V0_Pos.X, Tri.V0_Pos.Y, 0.0f);
+		StaticMeshVertexBuffers.PositionVertexBuffer.VertexPosition(StartIndex + 2) = FVector3f(Tri.V2_Pos.X, Tri.V2_Pos.Y, 0.0f);
 
 		StaticMeshVertexBuffers.StaticMeshVertexBuffer.SetVertexTangents(StartIndex + 0, FVector3f(1.0f, 0.0f, 0.0f), FVector3f(0.0f, 1.0f, 0.0f), FVector3f(0.0f, 0.0f, 1.0f));
 		StaticMeshVertexBuffers.StaticMeshVertexBuffer.SetVertexTangents(StartIndex + 1, FVector3f(1.0f, 0.0f, 0.0f), FVector3f(0.0f, 1.0f, 0.0f), FVector3f(0.0f, 0.0f, 1.0f));
@@ -147,8 +138,7 @@ void FCanvasTriangleRendererItem::FRenderData::RenderTriangles(
 	FCanvasRenderContext& RenderContext,
 	FMeshPassProcessorRenderState& DrawRenderState,
 	const FSceneView& View,
-	bool bIsHitTesting,
-	bool bNeedsToSwitchVerticalAxis)
+	bool bIsHitTesting)
 {
 	check(IsInRenderingThread());
 
@@ -164,7 +154,7 @@ void FCanvasTriangleRendererItem::FRenderData::RenderTriangles(
 
 	IRendererModule& RendererModule = GetRendererModule();
 
-	InitTriangleMesh(View, bNeedsToSwitchVerticalAxis);
+	InitTriangleMesh(View);
 
 
 	// We know we have at least 1 triangle so prep up a new batch right away : 
@@ -232,9 +222,7 @@ bool FCanvasTriangleRendererItem::Render_RenderThread(FCanvasRenderContext& Rend
 
 	const FSceneView& View = *RenderContext.Alloc<const FSceneView>(ViewInitOptions);
 
-	const bool bNeedsToSwitchVerticalAxis = RHINeedsToSwitchVerticalAxis(Canvas->GetShaderPlatform()) && Canvas->GetAllowSwitchVerticalAxis();
-
-	Data->RenderTriangles(RenderContext, DrawRenderState, View, Canvas->IsHitTesting(), bNeedsToSwitchVerticalAxis);
+	Data->RenderTriangles(RenderContext, DrawRenderState, View, Canvas->IsHitTesting());
 
 	if (Canvas->GetAllowedModes() & FCanvas::Allow_DeleteOnRender)
 	{
@@ -278,19 +266,18 @@ bool FCanvasTriangleRendererItem::Render_GameThread(const FCanvas* Canvas, FCanv
 
 	const FSceneView* View = new FSceneView(ViewInitOptions);
 
-	const bool bNeedsToSwitchVerticalAxis = RHINeedsToSwitchVerticalAxis(Canvas->GetShaderPlatform()) && Canvas->GetAllowSwitchVerticalAxis();
 	const bool bIsHitTesting = Canvas->IsHitTesting();
 	const bool bDeleteOnRender = Canvas->GetAllowedModes() & FCanvas::Allow_DeleteOnRender;
 
 	RenderScope.EnqueueRenderCommand(
-		[LocalData = Data, View, bIsHitTesting, bNeedsToSwitchVerticalAxis](FCanvasRenderContext& RenderContext) mutable
+		[LocalData = Data, View, bIsHitTesting](FCanvasRenderContext& RenderContext) mutable
 	{
 		FMeshPassProcessorRenderState DrawRenderState;
 
 		// disable depth test & writes
 		DrawRenderState.SetDepthStencilState(TStaticDepthStencilState<false, CF_Always>::GetRHI());
 
-		LocalData->RenderTriangles(RenderContext, DrawRenderState, *View, bIsHitTesting, bNeedsToSwitchVerticalAxis);
+		LocalData->RenderTriangles(RenderContext, DrawRenderState, *View, bIsHitTesting);
 
 		RenderContext.DeferredRelease(MoveTemp(LocalData));
 		RenderContext.DeferredDelete(View->Family);
diff --git a/Engine/Source/Runtime/Engine/Private/UnrealClient.cpp b/Engine/Source/Runtime/Engine/Private/UnrealClient.cpp
--- a/Engine/Source/Runtime/Engine/Private/UnrealClient.cpp
+++ b/Engine/Source/Runtime/Engine/Private/UnrealClient.cpp
@@ -1651,8 +1651,6 @@ void FViewport::Draw( bool bShouldPresent /*= true */)
 				FCanvas Canvas(this, nullptr, ViewportWorld, ViewportWorld ? ViewportWorld->FeatureLevel.GetValue() : GMaxRHIFeatureLevel, FCanvas::CDM_DeferDrawing, ViewportClient->ShouldDPIScaleSceneCanvas() ? ViewportClient->GetDPIScale() : 1.0f);
 				Canvas.SetRenderTargetRect(FIntRect(0, 0, SizeX, SizeY));
 				{
-					// Make sure the Canvas is not rendered upside down
-					Canvas.SetAllowSwitchVerticalAxis(true);
 					ViewportClient->Draw(this, &Canvas);
 				}
 				Canvas.Flush_GameThread();
diff --git a/Engine/Source/Runtime/Engine/Private/UserInterface/Canvas.cpp b/Engine/Source/Runtime/Engine/Private/UserInterface/Canvas.cpp
--- a/Engine/Source/Runtime/Engine/Private/UserInterface/Canvas.cpp
+++ b/Engine/Source/Runtime/Engine/Private/UserInterface/Canvas.cpp
@@ -299,7 +299,6 @@ FCanvas::FCanvas(FRenderTarget* InRenderTarget,FHitProxyConsumer* InHitProxyCons
 ,	AllowedModes(0xFFFFFFFF)
 ,	bRenderTargetDirty(false)
 ,	Time(InTime)
-,	bAllowsToSwitchVerticalAxis(false)
 ,	FeatureLevel(InFeatureLevel)
 ,	bUseInternalTexture(false)
 ,	DrawMode(CDM_DeferDrawing)
@@ -312,7 +311,6 @@ void FCanvas::Construct()
 {
 	bStereoRendering = false;
 	bScaledToRenderTarget = false;
-	bAllowsToSwitchVerticalAxis = false;
 
 	const FIntPoint RenderTargetSizeXY = RenderTarget ? RenderTarget->GetSizeXY() : FIntPoint(1, 1);
 
@@ -441,14 +439,11 @@ bool FCanvasBatchedElementRenderItem::Render_RenderThread(FCanvasRenderContext&
 				Gamma = 1.0f;
 			}
 
-			bool bNeedsToSwitchVerticalAxis = RHINeedsToSwitchVerticalAxis(Canvas->GetShaderPlatform()) && Canvas->GetAllowSwitchVerticalAxis();
-
 			// draw batched items
 			LocalData->BatchedElements.Draw(
 				RHICmdList,
 				DrawRenderState,
 				Canvas->GetFeatureLevel(),
-				bNeedsToSwitchVerticalAxis,
 				FBatchedElements::CreateProxySceneView(LocalData->Transform.GetMatrix(), FIntRect(0, 0, CanvasRenderTarget->GetSizeXY().X, CanvasRenderTarget->GetSizeXY().Y)),
 				Canvas->IsHitTesting(),
 				Gamma
@@ -482,14 +477,11 @@ bool FCanvasBatchedElementRenderItem::Render_GameThread(const FCanvas* Canvas, F
 			Gamma = 1.0f;
 		}
 
-		bool bNeedsToSwitchVerticalAxis = RHINeedsToSwitchVerticalAxis(Canvas->GetShaderPlatform()) && Canvas->GetAllowSwitchVerticalAxis();
-
 		// Render the batched elements.
 		struct FBatchedDrawParameters
 		{
 			FRenderData* RenderData;
 			uint32 bHitTesting : 1;
-			uint32 bNeedsToSwitchVerticalAxis : 1;
 			uint32 ViewportSizeX;
 			uint32 ViewportSizeY;
 			float DisplayGamma;
@@ -502,7 +494,6 @@ bool FCanvasBatchedElementRenderItem::Render_GameThread(const FCanvas* Canvas, F
 		{
 			Data,
 			(uint32)(Canvas->IsHitTesting() ? 1 : 0),
-			(uint32)(bNeedsToSwitchVerticalAxis ? 1 : 0),
 			(uint32)CanvasRenderTarget->GetSizeXY().X,
 			(uint32)CanvasRenderTarget->GetSizeXY().Y,
 			Gamma,
@@ -527,7 +518,6 @@ bool FCanvasBatchedElementRenderItem::Render_GameThread(const FCanvas* Canvas, F
 				RHICmdList,
 				DrawRenderState,
 				DrawParameters.FeatureLevel,
-				DrawParameters.bNeedsToSwitchVerticalAxis,
 				SceneView,
 				DrawParameters.bHitTesting,
 				DrawParameters.DisplayGamma);
diff --git a/Engine/Source/Runtime/Engine/Public/BatchedElements.h b/Engine/Source/Runtime/Engine/Public/BatchedElements.h
--- a/Engine/Source/Runtime/Engine/Public/BatchedElements.h
+++ b/Engine/Source/Runtime/Engine/Public/BatchedElements.h
@@ -249,7 +249,7 @@ public:
 	 * @param Gamma			Optional gamma override
 	 * @param DepthTexture	DepthTexture for manual depth testing with editor compositing in the pixel shader
 	 */
-	bool Draw(FRHICommandList& RHICmdList, const FMeshPassProcessorRenderState& DrawRenderState, ERHIFeatureLevel::Type FeatureLevel, bool bNeedToSwitchVerticalAxis, const FSceneView& View, bool bHitTesting, float Gamma = 1.0f, EBlendModeFilter::Type Filter = EBlendModeFilter::All) const;
+	bool Draw(FRHICommandList& RHICmdList, const FMeshPassProcessorRenderState& DrawRenderState, ERHIFeatureLevel::Type FeatureLevel, const FSceneView& View, bool bHitTesting, float Gamma = 1.0f, EBlendModeFilter::Type Filter = EBlendModeFilter::All) const;
 
 	/**
 	 * Creates a proxy FSceneView for operations that are not tied directly to a scene but still require batched elements to be drawn.
@@ -402,7 +402,6 @@ private:
 		ERHIFeatureLevel::Type FeatureLevel,
 		ESimpleElementBlendMode BlendMode,
 		const FRelativeViewMatrices& ViewMatrices,
-		bool bSwitchVerticalAxis,
 		FBatchedElementParameters* BatchedElementParameters,
 		const FTexture* Texture,
 		bool bHitTesting,
diff --git a/Engine/Source/Runtime/Engine/Public/CanvasTypes.h b/Engine/Source/Runtime/Engine/Public/CanvasTypes.h
--- a/Engine/Source/Runtime/Engine/Public/CanvasTypes.h
+++ b/Engine/Source/Runtime/Engine/Public/CanvasTypes.h
@@ -531,11 +531,6 @@ public:
 	*/
 	EShaderPlatform GetShaderPlatform() const { return GShaderPlatformForFeatureLevel[FeatureLevel]; }
 
-	// Get/Set if this Canvas allows its batched elements to switch vertical axis (e.g., rendering to back buffer should never flip)
-	bool GetAllowSwitchVerticalAxis() const { return bAllowsToSwitchVerticalAxis; }
-
-	void SetAllowSwitchVerticalAxis(bool bInAllowsToSwitchVerticalAxis) { bAllowsToSwitchVerticalAxis = bInAllowsToSwitchVerticalAxis; }
-
 public:
 	float AlphaModulate;
 
@@ -633,8 +628,6 @@ private:
 	FGameTime Time;
 	/** true, if Canvas should be scaled to whole render target */
 	bool bScaledToRenderTarget;
-	// True if canvas allows switching vertical axis; false will ignore any flip
-	bool bAllowsToSwitchVerticalAxis;
 	/** Feature level that we are currently rendering with */
 	ERHIFeatureLevel::Type FeatureLevel;
 
@@ -1079,7 +1072,6 @@ private:
 			FMeshPassProcessorRenderState& DrawRenderState,
 			const FSceneView& View,
 			bool bIsHitTesting,
-			bool bNeedsToSwitchVerticalAxis,
 			bool bUse128bitRT = false);
 
 		const FMaterialRenderProxy* const MaterialRenderProxy;
@@ -1096,7 +1088,7 @@ private:
 
 	private:
 		FMeshBatch* AllocTileMeshBatch(FCanvasRenderContext& InRenderContext, FHitProxyId InHitProxyId);
-		void InitTileMesh(const FSceneView& View, bool bNeedsToSwitchVerticalAxis);
+		void InitTileMesh(const FSceneView& View);
 		void ReleaseTileMesh();
 
 		FRawIndexBuffer16or32 IndexBuffer;
@@ -1256,8 +1248,7 @@ private:
 			FCanvasRenderContext& RenderContext,
 			FMeshPassProcessorRenderState& DrawRenderState,
 			const FSceneView& View,
-			bool bIsHitTesting,
-			bool bNeedsToSwitchVerticalAxis);
+			bool bIsHitTesting);
 
 		const FMaterialRenderProxy* const MaterialRenderProxy;
 		const FCanvas::FTransformEntry Transform;
@@ -1267,7 +1258,7 @@ private:
 
 	private:
 		FMeshBatch* AllocTriangleMeshBatch(FCanvasRenderContext& InRenderContext, FHitProxyId InHitProxyId);
-		void InitTriangleMesh(const FSceneView& View, bool bNeedsToSwitchVerticalAxis);
+		void InitTriangleMesh(const FSceneView& View);
 		void ReleaseTriangleMesh();
 
 		FRawIndexBuffer16or32 IndexBuffer;
diff --git a/Engine/Source/Runtime/Engine/Public/SimpleElementShaders.h b/Engine/Source/Runtime/Engine/Public/SimpleElementShaders.h
--- a/Engine/Source/Runtime/Engine/Public/SimpleElementShaders.h
+++ b/Engine/Source/Runtime/Engine/Public/SimpleElementShaders.h
@@ -29,8 +29,8 @@ public:
 	FSimpleElementVS(const ShaderMetaType::CompiledShaderInitializerType& Initializer);
 	FSimpleElementVS() {}
 
-	void SetParameters(FRHICommandList& RHICmdList, const FMatrix& WorldToClipMatrix, bool bSwitchVerticalAxis = false);
-	void SetParameters(FRHICommandList& RHICmdList, const FRelativeViewMatrices& Matrices, bool bSwitchVerticalAxis = false);
+	void SetParameters(FRHICommandList& RHICmdList, const FMatrix& WorldToClipMatrix);
+	void SetParameters(FRHICommandList& RHICmdList, const FRelativeViewMatrices& Matrices);
 
 	//virtual bool Serialize(FArchive& Ar) override;
 
@@ -39,7 +39,6 @@ public:
 private:
 	LAYOUT_FIELD(FShaderParameter, RelativeTransform);
 	LAYOUT_FIELD(FShaderParameter, TransformTilePosition);
-	LAYOUT_FIELD(FShaderParameter, SwitchVerticalAxis);
 };
 
 /**
diff --git a/Engine/Source/Runtime/OpenGLDrv/Private/Android/AndroidEGL.cpp b/Engine/Source/Runtime/OpenGLDrv/Private/Android/AndroidEGL.cpp
--- a/Engine/Source/Runtime/OpenGLDrv/Private/Android/AndroidEGL.cpp
+++ b/Engine/Source/Runtime/OpenGLDrv/Private/Android/AndroidEGL.cpp
@@ -15,6 +15,7 @@
 #include "Misc/ScopeLock.h"
 #include "UObject/GarbageCollection.h"
 #include "Android/AndroidPlatformFramePacer.h"
+#include <dlfcn.h>
 
 
 AndroidEGL* AndroidEGL::Singleton = NULL;
@@ -27,6 +28,9 @@ DEFINE_LOG_CATEGORY(LogEGL);
 #endif // EGL_KHR_create_context_no_error
 #endif // USE_ANDROID_EGL_NO_ERROR_CONTEXT
 
+typedef int32(*PFN_ANativeWindow_setBuffersTransform)(struct ANativeWindow* window, int32 transform);
+static PFN_ANativeWindow_setBuffersTransform ANativeWindow_setBuffersTransform_API = nullptr;
+
 
 const  int EGLMinRedBits		= 5;
 const  int EGLMinGreenBits		= 6;
@@ -150,6 +154,13 @@ AndroidEGL::AndroidEGL()
 ,	ContextAttributes(nullptr)
 {
 	PImplData = new AndroidESPImpl();
+
+	void* const LibNativeWindow = dlopen("libnativewindow.so", RTLD_NOW | RTLD_LOCAL);
+	if (LibNativeWindow != nullptr)
+	{
+		ANativeWindow_setBuffersTransform_API = reinterpret_cast<PFN_ANativeWindow_setBuffersTransform>(dlsym(LibNativeWindow, "ANativeWindow_setBuffersTransform"));
+	}
+	UE_LOG(LogAndroid, Log, TEXT("ANativeWindow_setBuffersTransform is %s on this device"), ANativeWindow_setBuffersTransform_API == nullptr ? TEXT("not supported") : TEXT("supported"));
 }
 
 void AndroidEGL::ResetDisplay()
@@ -624,19 +635,6 @@ void AndroidEGL::DestroyBackBuffer()
 
 void AndroidEGL::InitBackBuffer()
 {
-	//add check to see if any context was made current. 
-	GLint OnScreenWidth, OnScreenHeight;
-	if (FPlatformMisc::SupportsBackbufferSampling())
-	{
-		glGenFramebuffers(1, &PImplData->ResolveFrameBuffer);
-	}
-	else
-	{
-		PImplData->ResolveFrameBuffer = 0;
-	}
-	OnScreenWidth = PImplData->eglWidth;
-	OnScreenHeight = PImplData->eglHeight;
-
 	PImplData->RenderingContext.ViewportFramebuffer = GetResolveFrameBuffer();
 	PImplData->SharedContext.ViewportFramebuffer = GetResolveFrameBuffer();
 	PImplData->SingleThreadedContext.ViewportFramebuffer = GetResolveFrameBuffer();
@@ -915,6 +913,20 @@ void AndroidEGL::AcquireCurrentRenderingContext()
 		PImplData->SharedContext.DummyFrameBuffer = PImplData->DummyFrameBuffer;
 		PImplData->SingleThreadedContext.DummyFrameBuffer = PImplData->DummyFrameBuffer;
 	}
+
+	if (FAndroidMisc::SupportsBackbufferSampling())
+	{
+		// Needs to be generated on rendering context
+		if (!PImplData->ResolveFrameBuffer)
+		{
+			glGenFramebuffers(1, &PImplData->ResolveFrameBuffer);
+		}
+	}
+	else
+	{
+		PImplData->ResolveFrameBuffer = 0;
+	}
+
 }
 
 void AndroidEGL::SetCurrentRenderingContext()
@@ -1122,6 +1134,39 @@ void AndroidEGL::LogConfigInfo(EGLConfig  EGLConfigInfo)
 	eglGetConfigAttrib(PImplData->eglDisplay,EGLConfigInfo, EGL_TRANSPARENT_BLUE_VALUE, &ResultValue);  FPlatformMisc::LowLevelOutputDebugStringf( TEXT("EGLConfigInfo :EGL_TRANSPARENT_BLUE_VALUE :	%u" ), ResultValue );
 }
 
+void AndroidEGL::UpdateBuffersTransform()
+{
+	if (ANativeWindow_setBuffersTransform_API != nullptr)
+	{
+		int32 BufferTransform = ANATIVEWINDOW_TRANSFORM_IDENTITY;
+
+		EDeviceScreenOrientation ScreenOrientation = FPlatformMisc::GetDeviceOrientation();
+		switch (ScreenOrientation)
+		{
+		case EDeviceScreenOrientation::Portrait:
+			BufferTransform = ANATIVEWINDOW_TRANSFORM_MIRROR_VERTICAL;
+			break;
+
+		case EDeviceScreenOrientation::PortraitUpsideDown:
+			BufferTransform = ANATIVEWINDOW_TRANSFORM_MIRROR_HORIZONTAL;
+			break;
+
+		case EDeviceScreenOrientation::LandscapeLeft:
+			BufferTransform = ANATIVEWINDOW_TRANSFORM_ROTATE_90 | ANATIVEWINDOW_TRANSFORM_MIRROR_VERTICAL;
+			break;
+
+		case EDeviceScreenOrientation::LandscapeRight:
+			BufferTransform = ANATIVEWINDOW_TRANSFORM_ROTATE_90 | ANATIVEWINDOW_TRANSFORM_MIRROR_HORIZONTAL;
+			break;
+
+		default:
+			ensureMsgf(false, TEXT("BufferTransform %d should be handled with no exception, otherwise wrong orientation could be displayed on device"), BufferTransform);
+			break;
+		}
+
+		ANativeWindow_setBuffersTransform_API(GetNativeWindow(), BufferTransform);
+	}
+}
 
 ///
 extern FCriticalSection GAndroidWindowLock;
diff --git a/Engine/Source/Runtime/OpenGLDrv/Private/Android/AndroidEGL.h b/Engine/Source/Runtime/OpenGLDrv/Private/Android/AndroidEGL.h
--- a/Engine/Source/Runtime/OpenGLDrv/Private/Android/AndroidEGL.h
+++ b/Engine/Source/Runtime/OpenGLDrv/Private/Android/AndroidEGL.h
@@ -81,6 +81,7 @@ public:
 	void InitSurface(bool bUseSmallSurface, bool bCreateWndSurface);
 	void InitRenderSurface(bool bUseSmallSurface, bool bCreateWndSurface);
 	void InitSharedSurface(bool bUseSmallSurface);
+	void UpdateBuffersTransform();
 
 	void GetDimensions(uint32& OutWidth, uint32& OutHeight);
 	
diff --git a/Engine/Source/Runtime/OpenGLDrv/Private/Android/AndroidOpenGL.cpp b/Engine/Source/Runtime/OpenGLDrv/Private/Android/AndroidOpenGL.cpp
--- a/Engine/Source/Runtime/OpenGLDrv/Private/Android/AndroidOpenGL.cpp
+++ b/Engine/Source/Runtime/OpenGLDrv/Private/Android/AndroidOpenGL.cpp
@@ -225,7 +225,8 @@ bool PlatformBlitToViewport(FPlatformOpenGLDevice* Device, const FOpenGLViewport
 	
 	FPlatformOpenGLContext* const Context = Viewport.GetGLContext();
 
-	if (FPlatformMisc::SupportsBackbufferSampling())
+	if (bPresent && FPlatformMisc::SupportsBackbufferSampling())
+
 	{
 		if (Device->TargetDirty)
 		{
@@ -263,6 +264,7 @@ bool PlatformBlitToViewport(FPlatformOpenGLDevice* Device, const FOpenGLViewport
 	}
 	if (bPresent)
 	{
+		AndroidEGL::GetInstance()->UpdateBuffersTransform();
 		FAndroidPlatformRHIFramePacer::SwapBuffers(bLockToVsync);
 	}
 	static IConsoleVariable* CVar = IConsoleManager::Get().FindConsoleVariable(TEXT("a.UseFrameTimeStampsForPacing"));
diff --git a/Engine/Source/Runtime/OpenGLDrv/Private/Android/AndroidOpenGLPrivate.h b/Engine/Source/Runtime/OpenGLDrv/Private/Android/AndroidOpenGLPrivate.h
--- a/Engine/Source/Runtime/OpenGLDrv/Private/Android/AndroidOpenGLPrivate.h
+++ b/Engine/Source/Runtime/OpenGLDrv/Private/Android/AndroidOpenGLPrivate.h
@@ -51,7 +51,7 @@ private:
 		// Do not create a window surface if the app is for Oculus Mobile (use small buffer)
 		bool bCreateSurface = !AndroidThunkCpp_IsOculusMobileApplication();
 		FPlatformMisc::LowLevelOutputDebugString(TEXT("FAndroidGPUInfo"));
-		EGL->InitSurface(bCreateSurface, bCreateSurface);
+		EGL->InitSurface(false, bCreateSurface);
 		EGL->SetCurrentSharedContext();
 
 		// get extensions
diff --git a/Engine/Source/Runtime/RHI/Private/RHI.cpp b/Engine/Source/Runtime/RHI/Private/RHI.cpp
--- a/Engine/Source/Runtime/RHI/Private/RHI.cpp
+++ b/Engine/Source/Runtime/RHI/Private/RHI.cpp
@@ -2431,8 +2431,6 @@ void FGenericDataDrivenShaderPlatformInfo::SetDefaultValues()
 {
 	MaxFeatureLevel = ERHIFeatureLevel::Num;
 	bSupportsMSAA = true;
-
-	bNeedsToSwitchVerticalAxisOnMobileOpenGL = true;
 	bSupportsDOFHybridScattering = true;
 	bSupportsHZBOcclusion = true;
 	bSupportsWaterIndirectDraw = true;
@@ -2529,7 +2527,6 @@ void FGenericDataDrivenShaderPlatformInfo::ParseDataDrivenShaderInfo(const FConf
 	GET_SECTION_INT_HELPER(NumberOfComputeThreads);
 
 	GET_SECTION_BOOL_HELPER(bWaterUsesSimpleForwardShading);
-	GET_SECTION_BOOL_HELPER(bNeedsToSwitchVerticalAxisOnMobileOpenGL);
 	GET_SECTION_BOOL_HELPER(bSupportsHairStrandGeometry);
 	GET_SECTION_BOOL_HELPER(bSupportsDOFHybridScattering);
 	GET_SECTION_BOOL_HELPER(bNeedsExtraMobileFrames);
diff --git a/Engine/Source/Runtime/RHI/Public/RHIDefinitions.h b/Engine/Source/Runtime/RHI/Public/RHIDefinitions.h
--- a/Engine/Source/Runtime/RHI/Public/RHIDefinitions.h
+++ b/Engine/Source/Runtime/RHI/Public/RHIDefinitions.h
@@ -398,7 +398,6 @@ class RHI_API FGenericDataDrivenShaderPlatformInfo
 	uint32 bSupportsVariableRateShading : 1;
 	uint32 NumberOfComputeThreads : 10;
 	uint32 bWaterUsesSimpleForwardShading : 1;
-	uint32 bNeedsToSwitchVerticalAxisOnMobileOpenGL : 1;
 	uint32 bSupportsHairStrandGeometry : 1;
 	uint32 bSupportsDOFHybridScattering : 1;
 	uint32 bNeedsExtraMobileFrames : 1;
@@ -879,12 +878,6 @@ public:
 		return Infos[Platform].bWaterUsesSimpleForwardShading;
 	}
 
-	static FORCEINLINE_DEBUGGABLE const bool GetNeedsToSwitchVerticalAxisOnMobileOpenGL(const FStaticShaderPlatform Platform)
-	{
-		check(IsValid(Platform));
-		return Infos[Platform].bNeedsToSwitchVerticalAxisOnMobileOpenGL;
-	}
-
 	static FORCEINLINE_DEBUGGABLE const bool GetSupportsHairStrandGeometry(const FStaticShaderPlatform Platform)
 	{
 		check(IsValid(Platform));
@@ -2205,22 +2198,6 @@ inline bool IsFeatureLevelSupported(const FStaticShaderPlatform InShaderPlatform
 	return InFeatureLevel <= GetMaxSupportedFeatureLevel(InShaderPlatform);
 }
 
-inline bool RHINeedsToSwitchVerticalAxis(const FStaticShaderPlatform Platform)
-{
-#if WITH_EDITOR
-	static const auto CVar = IConsoleManager::Get().FindTConsoleVariableDataInt(TEXT("r.Mobile.ForceRHISwitchVerticalAxis"));
-	if (CVar->GetValueOnAnyThread())
-	{
-		// only allow this for mobile preview.
-		return IsMobilePlatform(Platform);
-	}
-#endif
-
-	// ES3.1 need to flip when rendering to an RT that will be post processed
-	return IsOpenGLPlatform(Platform) && IsMobilePlatform(Platform) && !IsPCPlatform(Platform) && !IsMetalMobilePlatform(Platform) && !IsVulkanPlatform(Platform)
-			&& FDataDrivenShaderPlatformInfo::GetNeedsToSwitchVerticalAxisOnMobileOpenGL(Platform);
-}
-
 inline bool RHISupportsSeparateMSAAAndResolveTextures(const FStaticShaderPlatform Platform)
 {
 	// Metal mobile devices and Android ES3.1 need to handle MSAA and resolve textures internally (unless RHICreateTexture2D was changed to take an optional resolve target)
diff --git a/Engine/Source/Runtime/Renderer/Private/BasePassRendering.cpp b/Engine/Source/Runtime/Renderer/Private/BasePassRendering.cpp
--- a/Engine/Source/Runtime/Renderer/Private/BasePassRendering.cpp
+++ b/Engine/Source/Runtime/Renderer/Private/BasePassRendering.cpp
@@ -979,8 +979,6 @@ static void RenderEditorPrimitivesForDPG(
 
 		if (!View.Family->EngineShowFlags.CompositeEditorPrimitives)
 		{
-			const bool bNeedToSwitchVerticalAxis = RHINeedsToSwitchVerticalAxis(View.GetShaderPlatform());
-
 			DrawDynamicMeshPass(View, RHICmdList,
 				[&View, &DrawRenderState](FDynamicPassMeshDrawListContext* DynamicMeshPassContext)
 			{
@@ -1024,7 +1022,7 @@ static void RenderEditorPrimitivesForDPG(
 			});
 
 			// Draw the view's batched simple elements(lines, sprites, etc).
-			View.TopBatchedViewElements.Draw(RHICmdList, DrawRenderState, View.GetFeatureLevel(), bNeedToSwitchVerticalAxis, View, false);
+			View.TopBatchedViewElements.Draw(RHICmdList, DrawRenderState, View.GetFeatureLevel(), View, false);
 		}
 	});
 }
diff --git a/Engine/Source/Runtime/Renderer/Private/DepthRendering.cpp b/Engine/Source/Runtime/Renderer/Private/DepthRendering.cpp
--- a/Engine/Source/Runtime/Renderer/Private/DepthRendering.cpp
+++ b/Engine/Source/Runtime/Renderer/Private/DepthRendering.cpp
@@ -405,8 +405,6 @@ static void RenderPrePassEditorPrimitives(
 
 		if (!View.Family->EngineShowFlags.CompositeEditorPrimitives)
 		{
-			const bool bNeedToSwitchVerticalAxis = RHINeedsToSwitchVerticalAxis(View.GetShaderPlatform());
-
 			DrawDynamicMeshPass(View, RHICmdList, [&](FDynamicPassMeshDrawListContext* DynamicMeshPassContext)
 			{
 				FDepthPassMeshProcessor PassMeshProcessor(
@@ -429,7 +427,7 @@ static void RenderPrePassEditorPrimitives(
 			});
 
 			// Draw the view's batched simple elements(lines, sprites, etc).
-			View.BatchedViewElements.Draw(RHICmdList, DrawRenderState, View.FeatureLevel, bNeedToSwitchVerticalAxis, View, false);
+			View.BatchedViewElements.Draw(RHICmdList, DrawRenderState, View.FeatureLevel, View, false);
 
 			DrawDynamicMeshPass(View, RHICmdList, [&](FDynamicPassMeshDrawListContext* DynamicMeshPassContext)
 			{
@@ -453,7 +451,7 @@ static void RenderPrePassEditorPrimitives(
 			});
 
 			// Draw the view's batched simple elements(lines, sprites, etc).
-			View.TopBatchedViewElements.Draw(RHICmdList, DrawRenderState, View.FeatureLevel, bNeedToSwitchVerticalAxis, View, false);
+			View.TopBatchedViewElements.Draw(RHICmdList, DrawRenderState, View.FeatureLevel, View, false);
 		}
 	});
 }
diff --git a/Engine/Source/Runtime/Renderer/Private/MobileBasePassRendering.cpp b/Engine/Source/Runtime/Renderer/Private/MobileBasePassRendering.cpp
--- a/Engine/Source/Runtime/Renderer/Private/MobileBasePassRendering.cpp
+++ b/Engine/Source/Runtime/Renderer/Private/MobileBasePassRendering.cpp
@@ -348,8 +348,6 @@ void FMobileSceneRenderer::RenderMobileEditorPrimitives(FRHICommandList& RHICmdL
 
 	if (!View.Family->EngineShowFlags.CompositeEditorPrimitives)
 	{
-		const bool bNeedToSwitchVerticalAxis = RHINeedsToSwitchVerticalAxis(GShaderPlatformForFeatureLevel[FeatureLevel]) && !IsMobileHDR();
-
 		DrawDynamicMeshPass(View, RHICmdList,
 			[&View, &DrawRenderState](FDynamicPassMeshDrawListContext* DynamicMeshPassContext)
 			{
@@ -371,7 +369,7 @@ void FMobileSceneRenderer::RenderMobileEditorPrimitives(FRHICommandList& RHICmdL
 			});
 
 		// Draw the view's batched simple elements(lines, sprites, etc).
-		View.BatchedViewElements.Draw(RHICmdList, DrawRenderState, FeatureLevel, bNeedToSwitchVerticalAxis, View, false);
+		View.BatchedViewElements.Draw(RHICmdList, DrawRenderState, FeatureLevel, View, false);
 
 		DrawDynamicMeshPass(View, RHICmdList,
 			[&View, &DrawRenderState](FDynamicPassMeshDrawListContext* DynamicMeshPassContext)
@@ -394,6 +392,6 @@ void FMobileSceneRenderer::RenderMobileEditorPrimitives(FRHICommandList& RHICmdL
 		});
 
 		// Draw the view's batched simple elements(lines, sprites, etc).
-		View.TopBatchedViewElements.Draw(RHICmdList, DrawRenderState, FeatureLevel, bNeedToSwitchVerticalAxis, View, false);
+		View.TopBatchedViewElements.Draw(RHICmdList, DrawRenderState, FeatureLevel, View, false);
 	}
 }
\ No newline at end of file
diff --git a/Engine/Source/Runtime/Renderer/Private/MobileShadingRenderer.cpp b/Engine/Source/Runtime/Renderer/Private/MobileShadingRenderer.cpp
--- a/Engine/Source/Runtime/Renderer/Private/MobileShadingRenderer.cpp
+++ b/Engine/Source/Runtime/Renderer/Private/MobileShadingRenderer.cpp
@@ -443,7 +443,8 @@ void FMobileSceneRenderer::InitViews(FRDGBuilder& GraphBuilder, FSceneTexturesCo
 						|| Views[0].bIsSceneCapture 
 						|| Views[0].bIsReflectionCapture 
 						// If the resolve texture is not the same as the MSAA texture, we need to render to scene color and copy to back buffer.
-						|| (NumMSAASamples > 1 && !RHISupportsSeparateMSAAAndResolveTextures(ShaderPlatform));
+						|| (NumMSAASamples > 1 && !RHISupportsSeparateMSAAAndResolveTextures(ShaderPlatform))
+						|| bIsFullDepthPrepassEnabled;
 	const FPlanarReflectionSceneProxy* PlanarReflectionSceneProxy = Scene ? Scene->GetForwardPassGlobalPlanarReflection() : nullptr;
 
 	bRequiresPixelProjectedPlanarRelfectionPass = IsUsingMobilePixelProjectedReflection(ShaderPlatform)
@@ -540,7 +541,7 @@ void FMobileSceneRenderer::InitViews(FRDGBuilder& GraphBuilder, FSceneTexturesCo
 	
 	// Find out whether custom depth pass should be rendered.
 	{
-		bool bCouldUseCustomDepthStencil = !bGammaSpace && (!Scene->World || (Scene->World->WorldType != EWorldType::EditorPreview && Scene->World->WorldType != EWorldType::Inactive));
+		bool bCouldUseCustomDepthStencil = (!Scene->World || (Scene->World->WorldType != EWorldType::EditorPreview && Scene->World->WorldType != EWorldType::Inactive));
 		FMobileCustomDepthStencilUsage CustomDepthStencilUsage;
 		for (int32 ViewIndex = 0; ViewIndex < Views.Num(); ViewIndex++)
 		{
diff --git a/Engine/Source/Runtime/Renderer/Private/PlanarReflectionRendering.cpp b/Engine/Source/Runtime/Renderer/Private/PlanarReflectionRendering.cpp
--- a/Engine/Source/Runtime/Renderer/Private/PlanarReflectionRendering.cpp
+++ b/Engine/Source/Runtime/Renderer/Private/PlanarReflectionRendering.cpp
@@ -240,12 +240,6 @@ void PrefilterPlanarReflection(
 			FIntPoint UV = View.ViewRect.Min;
 			FIntPoint UVSize = View.ViewRect.Size();
 
-			if (RHINeedsToSwitchVerticalAxis(GShaderPlatformForFeatureLevel[View.FeatureLevel]) && !IsMobileHDR())
-			{
-				UV.Y = UV.Y + UVSize.Y;
-				UVSize.Y = -UVSize.Y;
-			}
-
 			FDeferredLightVS::FParameters ParametersVS = FDeferredLightVS::GetParameters(View, 
 				0, 0,
 				View.ViewRect.Width(), View.ViewRect.Height(),
diff --git a/Engine/Source/Runtime/Renderer/Private/PostProcess/PostProcessCompositeEditorPrimitives.cpp b/Engine/Source/Runtime/Renderer/Private/PostProcess/PostProcessCompositeEditorPrimitives.cpp
--- a/Engine/Source/Runtime/Renderer/Private/PostProcess/PostProcessCompositeEditorPrimitives.cpp
+++ b/Engine/Source/Runtime/Renderer/Private/PostProcess/PostProcessCompositeEditorPrimitives.cpp
@@ -163,17 +163,15 @@ void RenderEditorPrimitives(FRHICommandListImmediate& RHICmdList, const FViewInf
 
 	const auto FeatureLevel = View.GetFeatureLevel();
 	const auto ShaderPlatform = GShaderPlatformForFeatureLevel[FeatureLevel];
-	const bool bNeedToSwitchVerticalAxis = RHINeedsToSwitchVerticalAxis(ShaderPlatform);
 
 	// Draw the view's batched simple elements(lines, sprites, etc).
-	View.BatchedViewElements.Draw(RHICmdList, DrawRenderState, FeatureLevel, bNeedToSwitchVerticalAxis, View, false, 1.0f);
+	View.BatchedViewElements.Draw(RHICmdList, DrawRenderState, FeatureLevel, View, false, 1.0f);
 }
 
 void RenderForegroundEditorPrimitives(FRHICommandList& RHICmdList, const FViewInfo& View, FMeshPassProcessorRenderState& DrawRenderState, FInstanceCullingManager& InstanceCullingManager)
 {
 	const auto FeatureLevel = View.GetFeatureLevel();
 	const auto ShaderPlatform = GShaderPlatformForFeatureLevel[FeatureLevel];
-	const bool bNeedToSwitchVerticalAxis = RHINeedsToSwitchVerticalAxis(ShaderPlatform);
 
 	// Draw a first time the foreground primitive without depth test to over right depth from non-foreground editor primitives.
 	{
@@ -201,7 +199,7 @@ void RenderForegroundEditorPrimitives(FRHICommandList& RHICmdList, const FViewIn
 			}
 		});
 
-		View.TopBatchedViewElements.Draw(RHICmdList, DrawRenderState, FeatureLevel, bNeedToSwitchVerticalAxis, View, false);
+		View.TopBatchedViewElements.Draw(RHICmdList, DrawRenderState, FeatureLevel, View, false);
 	}
 
 	// Draw a second time the foreground primitive with depth test to have proper depth test between foreground primitives.
@@ -230,7 +228,7 @@ void RenderForegroundEditorPrimitives(FRHICommandList& RHICmdList, const FViewIn
 			}
 		});
 
-		View.TopBatchedViewElements.Draw(RHICmdList, DrawRenderState, FeatureLevel, bNeedToSwitchVerticalAxis, View, false);
+		View.TopBatchedViewElements.Draw(RHICmdList, DrawRenderState, FeatureLevel, View, false);
 	}
 }
 
diff --git a/Engine/Source/Runtime/Renderer/Private/PostProcess/PostProcessMaterial.cpp b/Engine/Source/Runtime/Renderer/Private/PostProcess/PostProcessMaterial.cpp
--- a/Engine/Source/Runtime/Renderer/Private/PostProcess/PostProcessMaterial.cpp
+++ b/Engine/Source/Runtime/Renderer/Private/PostProcess/PostProcessMaterial.cpp
@@ -343,24 +343,6 @@ TGlobalResource<FPostProcessMaterialVertexDeclaration> GPostProcessMaterialVerte
 
 } //! namespace
 
-static void AddCopyAndFlipTexturePass(FRDGBuilder& GraphBuilder, const FViewInfo& View, FRDGTextureRef SrcTexture, FRDGTextureRef DestTexture)
-{
-	// Other RHIs can't flip and copy at the same time, so we'll use a pixel shader to perform the copy, together with the FlipYAxis
-	// flag on the screen pass. This path is only taken when using the mobile preview feature in the editor with
-	// r.Mobile.ForceRHISwitchVerticalAxis set to 1, so we don't care about it being sub-optimal.
-	FIntPoint Size = SrcTexture->Desc.Extent;
-	const FScreenPassTextureViewport InputViewport(SrcTexture->Desc.Extent, FIntRect(FIntPoint::ZeroValue, Size));
-	const FScreenPassTextureViewport OutputViewport(DestTexture->Desc.Extent, FIntRect(FIntPoint::ZeroValue, Size));
-	TShaderMapRef<FCopyRectPS> PixelShader(View.ShaderMap);
-
-	FCopyRectPS::FParameters* Parameters = GraphBuilder.AllocParameters<FCopyRectPS::FParameters>();
-	Parameters->InputTexture = SrcTexture;
-	Parameters->InputSampler = TStaticSamplerState<>::GetRHI();
-	Parameters->RenderTargets[0] = FRenderTargetBinding(DestTexture, ERenderTargetLoadAction::ENoAction);
-
-	AddDrawScreenPass(GraphBuilder, RDG_EVENT_NAME("DrawTexture"), View, OutputViewport, InputViewport, PixelShader, Parameters, EScreenPassDrawFlags::FlipYAxis);
-}
-
 void AddMobileMSAADecodeAndDrawTexturePass(
 	FRDGBuilder& GraphBuilder,
 	const FViewInfo& View,
@@ -434,15 +416,10 @@ FScreenPassTexture AddPostProcessMaterialPass(
 	// is actually the backbuffer, but it's not worth doing all the plumbing and increasing the RHI surface area just for this hack.
 	const bool bBackbufferWithDepthStencil = (DepthStencilTexture != nullptr && !GRHISupportsBackBufferWithCustomDepthStencil && Inputs.OverrideOutput.IsValid());
 
-	// The other case when we must render to an intermediate target is when we have to flip the image vertically because we're the last postprocess pass on mobile OpenGL.
-	// We can't simply output a flipped image, because the parts of the input image which show through the stencil mask or are blended in must also be flipped. In that case,
-	// we render normally to the intermediate target and flip the image when we copy to the output target.
-	const bool bCompositeWithInputAndFlipY = bCompositeWithInput && Inputs.bFlipYAxis;
-
 	// We need to decode the target color for blending material, force it rendering to an intermediate render target and decode the color.
 	const bool bCompositeWithInputAndDecode = Inputs.bMetalMSAAHDRDecode && bCompositeWithInput;
 
-	const bool bForceIntermediateTarget = bBackbufferWithDepthStencil || bCompositeWithInputAndFlipY || bCompositeWithInputAndDecode;
+	const bool bForceIntermediateTarget = bBackbufferWithDepthStencil || bCompositeWithInputAndDecode;
 
 	FScreenPassRenderTarget Output = Inputs.OverrideOutput;
 
@@ -544,7 +521,6 @@ FScreenPassTexture AddPostProcessMaterialPass(
 		PostProcessMaterialParameters->PostProcessInput[InputIndex] = GetScreenPassTextureInput(Input, PointClampSampler);
 	}
 
-	PostProcessMaterialParameters->bFlipYAxis = Inputs.bFlipYAxis && !bForceIntermediateTarget;
 	PostProcessMaterialParameters->Strata = Strata::BindStrataGlobalUniformParameters(View);
 
 	// SceneDepthWithoutWater
@@ -581,11 +557,6 @@ FScreenPassTexture AddPostProcessMaterialPass(
 
 	EScreenPassDrawFlags ScreenPassFlags = EScreenPassDrawFlags::AllowHMDHiddenAreaMask;
 
-	if (PostProcessMaterialParameters->bFlipYAxis)
-	{
-		ScreenPassFlags |= EScreenPassDrawFlags::FlipYAxis;
-	}
-
 	// check if we can skip that draw call in case if all pixels will fail the stencil test of the material
 	bool bSkipPostProcess = false;
 
@@ -615,46 +586,29 @@ FScreenPassTexture AddPostProcessMaterialPass(
 
 	if (!bSkipPostProcess)
 	{
-	AddDrawScreenPass(
-		GraphBuilder,
-		RDG_EVENT_NAME("PostProcessMaterial"),
-		View,
-		OutputViewport,
-		SceneColorViewport,
-		// Uses default depth stencil on mobile since the stencil test is done in pixel shader.
-		FScreenPassPipelineState(VertexShader, PixelShader, BlendState, DepthStencilState, MaterialStencilRef),
-		PostProcessMaterialParameters,
-		ScreenPassFlags,
-		[&View, VertexShader, PixelShader, MaterialRenderProxy, Material, PostProcessMaterialParameters](FRHICommandList& RHICmdList)
-	{
-		FPostProcessMaterialVS::SetParameters(RHICmdList, VertexShader, View, MaterialRenderProxy, *Material, *PostProcessMaterialParameters);
-		FPostProcessMaterialPS::SetParameters(RHICmdList, PixelShader, View, MaterialRenderProxy, *Material, *PostProcessMaterialParameters);
-	});
-
-	if (bForceIntermediateTarget && !bCompositeWithInputAndDecode)
-	{
-		if (!Inputs.bFlipYAxis)
+		AddDrawScreenPass(
+			GraphBuilder,
+			RDG_EVENT_NAME("PostProcessMaterial"),
+			View,
+			OutputViewport,
+			SceneColorViewport,
+			// Uses default depth stencil on mobile since the stencil test is done in pixel shader.
+			FScreenPassPipelineState(VertexShader, PixelShader, BlendState, DepthStencilState, MaterialStencilRef),
+			PostProcessMaterialParameters,
+			ScreenPassFlags,
+			[&View, VertexShader, PixelShader, MaterialRenderProxy, Material, PostProcessMaterialParameters](FRHICommandList& RHICmdList)
+		{
+			FPostProcessMaterialVS::SetParameters(RHICmdList, VertexShader, View, MaterialRenderProxy, *Material, *PostProcessMaterialParameters);
+			FPostProcessMaterialPS::SetParameters(RHICmdList, PixelShader, View, MaterialRenderProxy, *Material, *PostProcessMaterialParameters);
+		});
+
+		if (bForceIntermediateTarget && !bCompositeWithInputAndDecode)
 		{
 			// We shouldn't get here unless we had an override target.
 			check(Inputs.OverrideOutput.IsValid());
 			AddDrawTexturePass(GraphBuilder, View, Output.Texture, Inputs.OverrideOutput.Texture);
 			Output = Inputs.OverrideOutput;
 		}
-		else
-		{
-			FScreenPassRenderTarget TempTarget = Output;
-			if (Inputs.OverrideOutput.IsValid())
-			{
-				Output = Inputs.OverrideOutput;
-			}
-			else
-			{
-				Output = FScreenPassRenderTarget(SceneColor, ERenderTargetLoadAction::ENoAction);
-			}
-
-			AddCopyAndFlipTexturePass(GraphBuilder, View, TempTarget.Texture, Output.Texture);
-		}
-	}
 	}
 	else
 	{
@@ -767,7 +721,6 @@ FScreenPassTexture AddPostProcessMaterialChain(
 		if (MaterialInterface != Materials.Last())
 		{
 			Inputs.OverrideOutput = FScreenPassRenderTarget();
-			Inputs.bFlipYAxis = false;
 		}
 
 		Outputs = AddPostProcessMaterialPass(GraphBuilder, View, Inputs, MaterialInterface);
diff --git a/Engine/Source/Runtime/Renderer/Private/PostProcess/PostProcessMaterial.h b/Engine/Source/Runtime/Renderer/Private/PostProcess/PostProcessMaterial.h
--- a/Engine/Source/Runtime/Renderer/Private/PostProcess/PostProcessMaterial.h
+++ b/Engine/Source/Runtime/Renderer/Private/PostProcess/PostProcessMaterial.h
@@ -46,7 +46,6 @@ BEGIN_SHADER_PARAMETER_STRUCT(FPostProcessMaterialParameters, )
 	SHADER_PARAMETER_SAMPLER(SamplerState, PostProcessInput_BilinearSampler)
 	SHADER_PARAMETER_RDG_TEXTURE(Texture2D, EyeAdaptationTexture)
 	SHADER_PARAMETER_RDG_BUFFER_SRV(Buffer<float4>, EyeAdaptationBuffer)
-	SHADER_PARAMETER(uint32, bFlipYAxis)
 	SHADER_PARAMETER(uint32, bMetalMSAAHDRDecode)
 	SHADER_PARAMETER(uint32, bSceneDepthWithoutWaterTextureAvailable)
 	SHADER_PARAMETER_RDG_TEXTURE(Texture2D, SceneDepthWithoutSingleLayerWaterTexture)
@@ -113,9 +112,6 @@ struct FPostProcessMaterialInputs
 	/** Depth and color textures of the scene without single layer water. May be nullptr if not available. */
 	const FSceneWithoutWaterTextures* SceneWithoutWaterTextures = nullptr;
 
-	/** Performs a vertical axis flip if the RHI allows it. */
-	bool bFlipYAxis = false;
-
 	/** Allows (but doesn't guarantee) an optimization where, if possible, the scene color input is reused as
 	 *  the output. This can elide a copy in certain circumstances; for example, when the scene color input isn't
 	 *  actually used by the post process material and no special depth-stencil / blend composition is required.
diff --git a/Engine/Source/Runtime/Renderer/Private/PostProcess/PostProcessTonemap.cpp b/Engine/Source/Runtime/Renderer/Private/PostProcess/PostProcessTonemap.cpp
--- a/Engine/Source/Runtime/Renderer/Private/PostProcess/PostProcessTonemap.cpp
+++ b/Engine/Source/Runtime/Renderer/Private/PostProcess/PostProcessTonemap.cpp
@@ -70,7 +70,6 @@ class FTonemapperLocalExposureDim : SHADER_PERMUTATION_BOOL("USE_LOCAL_EXPOSURE"
 class FTonemapperVignetteDim       : SHADER_PERMUTATION_BOOL("USE_VIGNETTE");
 class FTonemapperSharpenDim        : SHADER_PERMUTATION_BOOL("USE_SHARPEN");
 class FTonemapperFilmGrainDim      : SHADER_PERMUTATION_BOOL("USE_FILM_GRAIN");
-class FTonemapperSwitchAxis        : SHADER_PERMUTATION_BOOL("NEEDTOSWITCHVERTICLEAXIS");
 class FTonemapperMsaaDim           : SHADER_PERMUTATION_BOOL("METAL_MSAA_HDR_DECODE");
 class FTonemapperEyeAdaptationDim  : SHADER_PERMUTATION_BOOL("EYEADAPTATION_EXPOSURE_FIX");
 
@@ -81,17 +80,10 @@ using FCommonDomain = TShaderPermutationDomain<
 	FTonemapperVignetteDim,
 	FTonemapperSharpenDim,
 	FTonemapperFilmGrainDim,
-	FTonemapperSwitchAxis,
 	FTonemapperMsaaDim>;
 
 bool ShouldCompileCommonPermutation(const FGlobalShaderPermutationParameters& Parameters, const FCommonDomain& PermutationVector)
 {
-	// Prevent switch axis permutation on platforms that dont require it.
-	if (PermutationVector.Get<FTonemapperSwitchAxis>() && !RHINeedsToSwitchVerticalAxis(Parameters.Platform))
-	{
-		return false;
-	}
-
 	// MSAA pre-resolve step only used on iOS atm
 	if (PermutationVector.Get<FTonemapperMsaaDim>() && !IsMetalMobilePlatform(Parameters.Platform))
 	{
@@ -112,7 +104,7 @@ bool ShouldCompileCommonPermutation(const FGlobalShaderPermutationParameters& Pa
 }
 
 // Common conversion of engine settings into.
-FCommonDomain BuildCommonPermutationDomain(const FViewInfo& View, bool bGammaOnly, bool bLocalExposure, bool bSwitchVerticalAxis, bool bMetalMSAAHDRDecode)
+FCommonDomain BuildCommonPermutationDomain(const FViewInfo& View, bool bGammaOnly, bool bLocalExposure, bool bMetalMSAAHDRDecode)
 {
 	const FSceneViewFamily* Family = View.Family;
 
@@ -133,7 +125,6 @@ FCommonDomain BuildCommonPermutationDomain(const FViewInfo& View, bool bGammaOnl
 	PermutationVector.Set<FTonemapperLocalExposureDim>(bLocalExposure);
 	PermutationVector.Set<FTonemapperFilmGrainDim>(View.FilmGrainTexture != nullptr);
 	PermutationVector.Set<FTonemapperSharpenDim>(CVarTonemapperSharpen.GetValueOnRenderThread() > 0.0f);	
-	PermutationVector.Set<FTonemapperSwitchAxis>(bSwitchVerticalAxis);
 	PermutationVector.Set<FTonemapperMsaaDim>(bMetalMSAAHDRDecode);
 	return PermutationVector;
 }
@@ -322,7 +313,6 @@ BEGIN_SHADER_PARAMETER_STRUCT(FTonemapParameters, )
 	SHADER_PARAMETER(FVector4f, TonemapperParams)
 	SHADER_PARAMETER(FVector4f, LensPrincipalPointOffsetScale)
 	SHADER_PARAMETER(FVector4f, LensPrincipalPointOffsetScaleInverse)
-	SHADER_PARAMETER(float, SwitchVerticalAxis)
 	SHADER_PARAMETER(float, DefaultEyeExposure)
 	SHADER_PARAMETER(float, EditorNITLevel)
 	SHADER_PARAMETER(uint32, bOutputInHDR)
@@ -374,17 +364,11 @@ public:
 	// FDrawRectangleParameters is filled by DrawScreenPass.
 	SHADER_USE_PARAMETER_STRUCT_WITH_LEGACY_BASE(FTonemapVS, FGlobalShader);
 
-	using FPermutationDomain = TShaderPermutationDomain<TonemapperPermutation::FTonemapperSwitchAxis, TonemapperPermutation::FTonemapperEyeAdaptationDim>;
+	using FPermutationDomain = TShaderPermutationDomain<TonemapperPermutation::FTonemapperEyeAdaptationDim>;
 	using FParameters = FTonemapParameters;
 
 	static bool ShouldCompilePermutation(const FGlobalShaderPermutationParameters& Parameters)
 	{
-		FPermutationDomain PermutationVector(Parameters.PermutationId);
-		// Prevent switch axis permutation on platforms that dont require it.
-		if (PermutationVector.Get<TonemapperPermutation::FTonemapperSwitchAxis>() && !RHINeedsToSwitchVerticalAxis(Parameters.Platform))
-		{
-			return false;
-		}
 		return true;
 	}
 };
@@ -757,7 +741,6 @@ FScreenPassTexture AddTonemapPass(FRDGBuilder& GraphBuilder, const FViewInfo& Vi
 	CommonParameters.ColorScale0 = PostProcessSettings.SceneColorTint;
 	CommonParameters.ChromaticAberrationParams = ChromaticAberrationParams;
 	CommonParameters.TonemapperParams = FVector4f(PostProcessSettings.VignetteIntensity, SharpenDiv6, 0.0f, 0.0f);
-	CommonParameters.SwitchVerticalAxis = Inputs.bFlipYAxis;
 	CommonParameters.DefaultEyeExposure = DefaultEyeExposure;
 	CommonParameters.EditorNITLevel = EditorNITLevel;
 	CommonParameters.bOutputInHDR = ViewFamily.bIsHDR;
@@ -831,7 +814,7 @@ FScreenPassTexture AddTonemapPass(FRDGBuilder& GraphBuilder, const FViewInfo& Vi
 	TonemapperPermutation::FDesktopDomain DesktopPermutationVector;
 
 	{
-		TonemapperPermutation::FCommonDomain CommonDomain = TonemapperPermutation::BuildCommonPermutationDomain(View, Inputs.bGammaOnly, Inputs.LocalExposureTexture != nullptr, Inputs.bFlipYAxis, Inputs.bMetalMSAAHDRDecode);
+		TonemapperPermutation::FCommonDomain CommonDomain = TonemapperPermutation::BuildCommonPermutationDomain(View, Inputs.bGammaOnly, Inputs.LocalExposureTexture != nullptr, Inputs.bMetalMSAAHDRDecode);
 		DesktopPermutationVector.Set<TonemapperPermutation::FCommonDomain>(CommonDomain);
 
 		if (!CommonDomain.Get<TonemapperPermutation::FTonemapperGammaOnlyDim>())
@@ -880,7 +863,6 @@ FScreenPassTexture AddTonemapPass(FRDGBuilder& GraphBuilder, const FViewInfo& Vi
 		PassParameters->RenderTargets[0] = Output.GetRenderTargetBinding();
 
 		FTonemapVS::FPermutationDomain VertexPermutationVector;
-		VertexPermutationVector.Set<TonemapperPermutation::FTonemapperSwitchAxis>(Inputs.bFlipYAxis);
 		VertexPermutationVector.Set<TonemapperPermutation::FTonemapperEyeAdaptationDim>(bEyeAdaptation);
 
 		TShaderMapRef<FTonemapVS> VertexShader(View.ShaderMap, VertexPermutationVector);
diff --git a/Engine/Source/Runtime/Renderer/Private/PostProcess/PostProcessTonemap.h b/Engine/Source/Runtime/Renderer/Private/PostProcess/PostProcessTonemap.h
--- a/Engine/Source/Runtime/Renderer/Private/PostProcess/PostProcessTonemap.h
+++ b/Engine/Source/Runtime/Renderer/Private/PostProcess/PostProcessTonemap.h
@@ -57,9 +57,6 @@ struct FTonemapInputs
 	// [Optional, ES31] Eye adaptation buffer used to compute exposure. 
 	FRDGBufferRef EyeAdaptationBuffer = nullptr;
 
-	// [Raster Only, Mobile] Flips the image vertically on output.
-	bool bFlipYAxis = false;
-
 	// [Raster Only] Controls whether the alpha channel of the scene texture should be written to the output texture.
 	bool bWriteAlphaChannel = false;
 
diff --git a/Engine/Source/Runtime/Renderer/Private/PostProcess/PostProcessing.cpp b/Engine/Source/Runtime/Renderer/Private/PostProcess/PostProcessing.cpp
--- a/Engine/Source/Runtime/Renderer/Private/PostProcess/PostProcessing.cpp
+++ b/Engine/Source/Runtime/Renderer/Private/PostProcess/PostProcessing.cpp
@@ -1742,12 +1742,6 @@ static bool IsGaussianActive(const FViewInfo& View)
 	return true;
 }
 
-// Scene Captures are going to be used as textures and sampled in the main scene, so they must remain unflipped.
-bool IsMobilePostProcessingFlipEnabled(const FViewInfo& View)
-{
-	return RHINeedsToSwitchVerticalAxis(View.GetShaderPlatform()) && !View.bIsSceneCapture;
-}
-
 void AddMobilePostProcessingPasses(FRDGBuilder& GraphBuilder, FScene* Scene, const FViewInfo& View, const FMobilePostProcessingInputs& Inputs, FInstanceCullingManager& InstanceCullingManager)
 {
 	CSV_SCOPED_TIMING_STAT_EXCLUSIVE(RenderPostProcessing);
@@ -1889,8 +1883,6 @@ void AddMobilePostProcessingPasses(FRDGBuilder& GraphBuilder, FScene* Scene, con
 
 		PostProcessMaterialInputs.CustomDepthTexture = CustomDepth.Texture;
 
-		PostProcessMaterialInputs.bFlipYAxis = IsMobilePostProcessingFlipEnabled(View) && bLastPass;
-
 		PostProcessMaterialInputs.bMetalMSAAHDRDecode = bMetalMSAAHDRDecode;
 
 		PostProcessMaterialInputs.SceneTextures = GetSceneTextureShaderParameters(Inputs.SceneTextures);
@@ -2357,7 +2349,6 @@ void AddMobilePostProcessingPasses(FRDGBuilder& GraphBuilder, FScene* Scene, con
 		TonemapperInputs.EyeAdaptationTexture = nullptr;
 		TonemapperInputs.ColorGradingTexture = ColorGradingTexture;
 		TonemapperInputs.bWriteAlphaChannel = View.AntiAliasingMethod == AAM_FXAA || IsPostProcessingWithAlphaChannelSupported() || bUseMobileDof || IsMobilePropagateAlphaEnabled(View.GetShaderPlatform());
-		TonemapperInputs.bFlipYAxis = IsMobilePostProcessingFlipEnabled(View) && !PassSequence.IsEnabled(EPass::PostProcessMaterialAfterTonemapping);
 		TonemapperInputs.bOutputInHDR = bHDRTonemapperOutput;
 		TonemapperInputs.bGammaOnly = bDoGammaOnly;
 		TonemapperInputs.bMetalMSAAHDRDecode = bMetalMSAAHDRDecode;
diff --git a/Engine/Source/Runtime/Renderer/Private/SceneCaptureRendering.cpp b/Engine/Source/Runtime/Renderer/Private/SceneCaptureRendering.cpp
--- a/Engine/Source/Runtime/Renderer/Private/SceneCaptureRendering.cpp
+++ b/Engine/Source/Runtime/Renderer/Private/SceneCaptureRendering.cpp
@@ -175,8 +175,7 @@ void CopySceneCaptureComponentToTarget(
 	const FMinimalSceneTextures& SceneTextures,
 	FRDGTextureRef ViewFamilyTexture,
 	const FSceneViewFamily& ViewFamily,
-	const TArray<FViewInfo>& Views,
-	bool bNeedsFlippedRenderTarget)
+	const TArray<FViewInfo>& Views)
 {
 	ESceneCaptureSource SceneCaptureSource = ViewFamily.SceneCaptureSource;
 
@@ -232,7 +231,7 @@ void CopySceneCaptureComponentToTarget(
 				RDG_EVENT_NAME("View(%d)", ViewIndex),
 				PassParameters,
 				ERDGPassFlags::Raster,
-				[PassParameters, GraphicsPSOInit, VertexShader, PixelShader, &View, bNeedsFlippedRenderTarget] (FRHICommandList& RHICmdList)
+				[PassParameters, GraphicsPSOInit, VertexShader, PixelShader, &View] (FRHICommandList& RHICmdList)
 			{
 				FGraphicsPipelineStateInitializer LocalGraphicsPSOInit = GraphicsPSOInit;
 				RHICmdList.ApplyCachedRenderTargets(LocalGraphicsPSOInit);
@@ -241,32 +240,16 @@ void CopySceneCaptureComponentToTarget(
 				
 				CopyCaptureToTargetSetViewportFn(RHICmdList);
 
-				if (bNeedsFlippedRenderTarget)
-				{
-					DrawRectangle(
-						RHICmdList,
-						View.ViewRect.Min.X, View.ViewRect.Min.Y,
-						View.ViewRect.Width(), View.ViewRect.Height(),
-						View.ViewRect.Min.X, View.ViewRect.Height() - View.ViewRect.Min.Y,
-						View.ViewRect.Width(), -View.ViewRect.Height(),
-						View.UnconstrainedViewRect.Size(),
-						View.GetSceneTexturesConfig().Extent,
-						VertexShader,
-						EDRF_UseTriangleOptimization);
-				}
-				else
-				{
-					DrawRectangle(
-						RHICmdList,
-						View.ViewRect.Min.X, View.ViewRect.Min.Y,
-						View.ViewRect.Width(), View.ViewRect.Height(),
-						View.ViewRect.Min.X, View.ViewRect.Min.Y,
-						View.ViewRect.Width(), View.ViewRect.Height(),
-						View.UnconstrainedViewRect.Size(),
-						View.GetSceneTexturesConfig().Extent,
-						VertexShader,
-						EDRF_UseTriangleOptimization);
-				}
+				DrawRectangle(
+					RHICmdList,
+					View.ViewRect.Min.X, View.ViewRect.Min.Y,
+					View.ViewRect.Width(), View.ViewRect.Height(),
+					View.ViewRect.Min.X, View.ViewRect.Min.Y,
+					View.ViewRect.Width(), View.ViewRect.Height(),
+					View.UnconstrainedViewRect.Size(),
+					View.GetSceneTexturesConfig().Extent,
+					VertexShader,
+					EDRF_UseTriangleOptimization);
 			});
 		}
 	}
@@ -370,8 +353,7 @@ void UpdateSceneCaptureContentMobile_RenderThread(
 	const FString& EventName,
 	const FRHICopyTextureInfo& CopyInfo,
 	bool bGenerateMips,
-	const FGenerateMipsParams& GenerateMipsParams,
-	bool bDisableFlipCopyGLES)
+	const FGenerateMipsParams& GenerateMipsParams)
 {
 	// We need to execute the pre-render view extensions before we do any view dependent work.
 	FSceneRenderer::ViewExtensionPreRender_RenderThread(RHICmdList, SceneRenderer);
@@ -380,8 +362,6 @@ void UpdateSceneCaptureContentMobile_RenderThread(
 
 	// update any resources that needed a deferred update
 	FDeferredUpdateResource::UpdateResources(RHICmdList);
-	bool bUseSceneTextures = SceneRenderer->ViewFamily.SceneCaptureSource != SCS_FinalColorLDR &&
-		SceneRenderer->ViewFamily.SceneCaptureSource != SCS_FinalColorHDR;
 
 	{
 #if WANTS_DRAW_MESH_EVENTS
@@ -394,77 +374,32 @@ void UpdateSceneCaptureContentMobile_RenderThread(
 
 		FViewInfo& View = SceneRenderer->Views[0];
 
-		const bool bIsMobileHDR = IsMobileHDR();
-		const bool bRHINeedsFlip = RHINeedsToSwitchVerticalAxis(GMaxRHIShaderPlatform) && !bDisableFlipCopyGLES;
-		// note that GLES code will flip the image when:
-		//	bIsMobileHDR && SceneCaptureSource == SCS_FinalColorLDR (flip performed during post processing)
-		//	!bIsMobileHDR (rendering is flipped by vertex shader)
-		// they need flipping again so it is correct for texture addressing.
-		const bool bNeedsFlippedCopy = (!bIsMobileHDR || !bUseSceneTextures) && bRHINeedsFlip;
-		const bool bNeedsFlippedFinalColor = bNeedsFlippedCopy && !bUseSceneTextures;
-
 		// Intermediate render target that will need to be flipped (needed on !IsMobileHDR())
 		FRDGTextureRef FlippedOutputTexture{};
 
 		const FRenderTarget* Target = SceneRenderer->ViewFamily.RenderTarget;
-		if (bNeedsFlippedFinalColor)
-		{
-			// We need to use an intermediate render target since the result will be flipped
-			auto& RenderTargetRHI = Target->GetRenderTargetTexture();
-			FRDGTextureDesc Desc(FRDGTextureDesc::Create2D(
-				Target->GetSizeXY(),
-				RenderTargetRHI.GetReference()->GetFormat(),
-				RenderTargetRHI.GetReference()->GetClearBinding(),
-				TexCreate_RenderTargetable));
-			FlippedOutputTexture = GraphBuilder.CreateTexture(Desc, TEXT("SceneCaptureFlipped"));
-		}
 
 		// We don't support screen percentage in scene capture.
 		FIntRect ViewRect = View.UnscaledViewRect;
 		FIntRect UnconstrainedViewRect = View.UnconstrainedViewRect;
 
-		if (bNeedsFlippedFinalColor)
-		{
-			AddClearRenderTargetPass(GraphBuilder, FlippedOutputTexture, FLinearColor::Black, ViewRect);
-		}
-
 		const FIntRect CopyDestRect = CopyInfo.GetDestRect();
 
 		if (!CopyDestRect.IsEmpty())
 		{
-			CopyCaptureToTargetSetViewportFn = [CopyDestRect, bNeedsFlippedFinalColor, ViewRect, FlippedOutputTexture](FRHICommandList& RHICmdList)
+			CopyCaptureToTargetSetViewportFn = [CopyDestRect, ViewRect, FlippedOutputTexture](FRHICommandList& RHICmdList)
 			{
 				RHICmdList.SetScissorRect(false, 0, 0, 0, 0);
 
-				if (bNeedsFlippedFinalColor)
-				{
-					FIntRect DestRect = CopyDestRect;
-					int32 TileYID = DestRect.Min.Y / ViewRect.Height();
-					int32 TileYCount = (FlippedOutputTexture->Desc.GetSize().Y / ViewRect.Height()) - 1;
-					DestRect.Min.Y = (TileYCount - TileYID) * ViewRect.Height();
-					DestRect.Max.Y = DestRect.Min.Y + ViewRect.Height();
-					RHICmdList.SetViewport
-					(
-						float(DestRect.Min.X),
-						float(DestRect.Min.Y),
-						0.0f,
-						float(DestRect.Max.X),
-						float(DestRect.Max.Y),
-						1.0f
-					);
-				}
-				else
-				{
-					RHICmdList.SetViewport
-					(
-						float(CopyDestRect.Min.X),
-						float(CopyDestRect.Min.Y),
-						0.0f,
-						float(CopyDestRect.Max.X),
-						float(CopyDestRect.Max.Y),
-						1.0f
-					);
-				}
+				RHICmdList.SetViewport
+				(
+					float(CopyDestRect.Min.X),
+					float(CopyDestRect.Min.Y),
+					0.0f,
+					float(CopyDestRect.Max.X),
+					float(CopyDestRect.Max.Y),
+					1.0f
+				);
 			};
 		}
 		else
@@ -475,37 +410,7 @@ void UpdateSceneCaptureContentMobile_RenderThread(
 		// Render the scene normally
 		{
 			RDG_RHI_EVENT_SCOPE(GraphBuilder, RenderScene);
-
-			if (bNeedsFlippedFinalColor)
-			{
-				// Helper class to allow setting render target
-				struct FRenderTargetOverride : public FRenderTarget
-				{
-					FRenderTargetOverride(const FRenderTarget* TargetIn, FRHITexture2D* OverrideTexture)
-					{
-						RenderTargetTextureRHI = OverrideTexture;
-						OriginalTarget = TargetIn;
-					}
-
-					virtual FIntPoint GetSizeXY() const override { return FIntPoint(RenderTargetTextureRHI->GetSizeX(), RenderTargetTextureRHI->GetSizeY()); }
-					virtual float GetDisplayGamma() const override { return OriginalTarget->GetDisplayGamma(); }
-
-					FTexture2DRHIRef GetTextureParamRef() { return RenderTargetTextureRHI; }
-					const FRenderTarget* OriginalTarget;
-				};
-
-				// Hijack the render target
-				FRHITexture2D* FlippedOutputTextureRHI = GraphBuilder.ConvertToExternalTexture(FlippedOutputTexture)->GetRHI()->GetTexture2D();
-				SceneRenderer->ViewFamily.RenderTarget = GraphBuilder.AllocObject<FRenderTargetOverride>(Target, FlippedOutputTextureRHI); //-V506
-			}
-
 			SceneRenderer->Render(GraphBuilder);
-
-			if (bNeedsFlippedFinalColor)
-			{
-				// And restore it
-				SceneRenderer->ViewFamily.RenderTarget = Target;
-			}
 		}
 
 		FRDGTextureRef OutputTexture = RegisterExternalTexture(GraphBuilder, Target->GetRenderTargetTexture(), TEXT("OutputTexture"));
@@ -520,8 +425,7 @@ void UpdateSceneCaptureContentMobile_RenderThread(
 				SceneTextures,
 				OutputTexture,
 				SceneRenderer->ViewFamily,
-				SceneRenderer->Views,
-				bNeedsFlippedFinalColor);
+				SceneRenderer->Views);
 		}
 
 		if (bGenerateMips)
@@ -599,7 +503,6 @@ static void UpdateSceneCaptureContent_RenderThread(
 	const FRHICopyTextureInfo& CopyInfo,
 	bool bGenerateMips,
 	const FGenerateMipsParams& GenerateMipsParams,
-	const bool bDisableFlipCopyLDRGLES,
 	bool bClearRenderTarget,
 	bool bOrthographicCamera)
 {
@@ -617,8 +520,7 @@ static void UpdateSceneCaptureContent_RenderThread(
 				EventName,
 				CopyInfo,
 				bGenerateMips,
-				GenerateMipsParams,
-				bDisableFlipCopyLDRGLES);
+				GenerateMipsParams);
 			break;
 		}
 		case EShadingPath::Deferred:
@@ -1079,7 +981,6 @@ void FScene::UpdateSceneCaptureContents(USceneCaptureComponent2D* CaptureCompone
 			TextureRenderTarget->MipsAddressU == TA_Wrap ? AM_Wrap : (TextureRenderTarget->MipsAddressU == TA_Mirror ? AM_Mirror : AM_Clamp),
 			TextureRenderTarget->MipsAddressV == TA_Wrap ? AM_Wrap : (TextureRenderTarget->MipsAddressV == TA_Mirror ? AM_Mirror : AM_Clamp)};
 
-		const bool bDisableFlipCopyGLES = CaptureComponent->bDisableFlipCopyGLES;
 		const bool bOrthographicCamera = CaptureComponent->ProjectionType == ECameraProjectionMode::Orthographic;
 
 
@@ -1102,7 +1003,7 @@ void FScene::UpdateSceneCaptureContents(USceneCaptureComponent2D* CaptureCompone
 		bool bIsCompositing = (CaptureComponent->CompositeMode != SCCM_Overwrite) && (CaptureComponent->CaptureSource == SCS_SceneColorHDR);
 
 		ENQUEUE_RENDER_COMMAND(CaptureCommand)(
-			[SceneRenderer, TextureRenderTargetResource, TexturePtrNotDeferenced, EventName, bGenerateMips, GenerateMipsParams, bDisableFlipCopyGLES, GameViewportRT, bEnableOrthographicTiling, bIsCompositing, bOrthographicCamera, NumXTiles, NumYTiles, TileID](FRHICommandListImmediate& RHICmdList)
+			[SceneRenderer, TextureRenderTargetResource, TexturePtrNotDeferenced, EventName, bGenerateMips, GenerateMipsParams, GameViewportRT, bEnableOrthographicTiling, bIsCompositing, bOrthographicCamera, NumXTiles, NumYTiles, TileID](FRHICommandListImmediate& RHICmdList)
 			{
 				if (GameViewportRT != nullptr)
 				{
@@ -1133,7 +1034,7 @@ void FScene::UpdateSceneCaptureContents(USceneCaptureComponent2D* CaptureCompone
 				// Don't clear the render target when compositing, or in a tiling mode that fills in the render target in multiple passes.
 				bool bClearRenderTarget = !bIsCompositing && !bEnableOrthographicTiling;
 
-				UpdateSceneCaptureContent_RenderThread(RHICmdList, SceneRenderer, TextureRenderTargetResource, TextureRenderTargetResource, EventName, CopyInfo, bGenerateMips, GenerateMipsParams, bDisableFlipCopyGLES, bClearRenderTarget, bOrthographicCamera);
+				UpdateSceneCaptureContent_RenderThread(RHICmdList, SceneRenderer, TextureRenderTargetResource, TextureRenderTargetResource, EventName, CopyInfo, bGenerateMips, GenerateMipsParams, bClearRenderTarget, bOrthographicCamera);
 			}
 		);
 	}
@@ -1262,7 +1163,7 @@ void FScene::UpdateSceneCaptureContents(USceneCaptureComponentCube* CaptureCompo
 					{
 						FRHICopyTextureInfo CopyInfo;
 						CopyInfo.DestSliceIndex = TargetFace;
-						UpdateSceneCaptureContent_RenderThread(RHICmdList, SceneRenderer, TextureRenderTarget, TextureRenderTarget, EventName, CopyInfo, false, FGenerateMipsParams(), false, true, false);
+						UpdateSceneCaptureContent_RenderThread(RHICmdList, SceneRenderer, TextureRenderTarget, TextureRenderTarget, EventName, CopyInfo, false, FGenerateMipsParams(), true, false);
 					}
 				);
 			}
diff --git a/Engine/Source/Runtime/Renderer/Private/SceneCaptureRendering.h b/Engine/Source/Runtime/Renderer/Private/SceneCaptureRendering.h
--- a/Engine/Source/Runtime/Renderer/Private/SceneCaptureRendering.h
+++ b/Engine/Source/Runtime/Renderer/Private/SceneCaptureRendering.h
@@ -15,5 +15,4 @@ RENDERER_API void CopySceneCaptureComponentToTarget(
 	const FMinimalSceneTextures& SceneTextures,
 	FRDGTextureRef ViewFamilyTexture,
 	const FSceneViewFamily& ViewFamily,
-	const TArray<FViewInfo>& Views,
-	bool bNeedsFlippedRenderTarget = false);
\ No newline at end of file
+	const TArray<FViewInfo>& Views);
diff --git a/Engine/Source/Runtime/Renderer/Private/SceneHitProxyRendering.cpp b/Engine/Source/Runtime/Renderer/Private/SceneHitProxyRendering.cpp
--- a/Engine/Source/Runtime/Renderer/Private/SceneHitProxyRendering.cpp
+++ b/Engine/Source/Runtime/Renderer/Private/SceneHitProxyRendering.cpp
@@ -268,7 +268,6 @@ static void DoRenderHitProxies(
 	auto& ViewFamily = SceneRenderer->ViewFamily;
 	auto& Views = SceneRenderer->Views;
 	const auto FeatureLevel = SceneRenderer->FeatureLevel;
-	const bool bNeedToSwitchVerticalAxis = RHINeedsToSwitchVerticalAxis(GShaderPlatformForFeatureLevel[FeatureLevel]);
 	const FIntPoint HitProxyTextureExtent = HitProxyTexture->Desc.Extent;
 
 	{
@@ -339,7 +338,7 @@ static void DoRenderHitProxies(
 			RDG_EVENT_NAME("HitProxies::Render"),
 			PassParameters,
 			ERDGPassFlags::Raster,
-			[SceneRenderer, &View, LocalScene, FeatureLevel, bNeedToSwitchVerticalAxis, PassParameters](FRHICommandListImmediate& RHICmdList)
+			[SceneRenderer, &View, LocalScene, FeatureLevel, PassParameters](FRHICommandListImmediate& RHICmdList)
 		{
 			FMeshPassProcessorRenderState DrawRenderState;
 
@@ -427,7 +426,7 @@ static void DoRenderHitProxies(
 
 
 			// Draw the view's batched simple elements(lines, sprites, etc).
-			View.BatchedViewElements.Draw(RHICmdList, DrawRenderState, FeatureLevel, bNeedToSwitchVerticalAxis, View, true);
+			View.BatchedViewElements.Draw(RHICmdList, DrawRenderState, FeatureLevel, View, true);
 
 			// Some elements should never be occluded (e.g. gizmos).
 			// So we render those twice, first to overwrite potentially nearer objects,
@@ -453,7 +452,7 @@ static void DoRenderHitProxies(
 				}
 			});
 
-			View.TopBatchedViewElements.Draw(RHICmdList, DrawRenderState, FeatureLevel, bNeedToSwitchVerticalAxis, View, true);
+			View.TopBatchedViewElements.Draw(RHICmdList, DrawRenderState, FeatureLevel, View, true);
 
 			DrawRenderState.SetDepthStencilState(TStaticDepthStencilState<true, CF_DepthNearOrEqual>::GetRHI());
 
@@ -476,7 +475,7 @@ static void DoRenderHitProxies(
 				}
 			});
 
-			View.TopBatchedViewElements.Draw(RHICmdList, DrawRenderState, FeatureLevel, bNeedToSwitchVerticalAxis, View, true);
+			View.TopBatchedViewElements.Draw(RHICmdList, DrawRenderState, FeatureLevel, View, true);
 		});
 	}
 
@@ -496,7 +495,7 @@ static void DoRenderHitProxies(
 			RDG_EVENT_NAME("HitProxies::CopyOutput"),
 			PassParameters,
 			ERDGPassFlags::Raster,
-			[&Views, HitProxyTextureExtent, HitProxyTexture, ViewFamilyTexture, FeatureLevel, bNeedToSwitchVerticalAxis](FRHICommandListImmediate& RHICmdList)
+			[&Views, HitProxyTextureExtent, HitProxyTexture, ViewFamilyTexture, FeatureLevel](FRHICommandListImmediate& RHICmdList)
 		{
 			// Set up a FTexture that is used to draw the hit proxy buffer to the view family's render target.
 			FTexture HitProxyRenderTargetTexture;
@@ -550,7 +549,6 @@ static void DoRenderHitProxies(
 				RHICmdList,
 				DrawRenderState,
 				FeatureLevel,
-				bNeedToSwitchVerticalAxis,
 				SceneView,
 				false,
 				1.0f
diff --git a/Engine/Source/Runtime/Renderer/Private/SceneRendering.cpp b/Engine/Source/Runtime/Renderer/Private/SceneRendering.cpp
--- a/Engine/Source/Runtime/Renderer/Private/SceneRendering.cpp
+++ b/Engine/Source/Runtime/Renderer/Private/SceneRendering.cpp
@@ -3415,8 +3415,6 @@ void FSceneRenderer::RenderFinish(FRDGBuilder& GraphBuilder, FRDGTextureRef View
 						// so it can get the screen size
 						FScreenMessageWriter Writer(Canvas, 130);
 
-						// Make sure draws to the canvas are not rendered upside down.
-						Canvas.SetAllowSwitchVerticalAxis(true);
 						if (bViewParentOrFrozen)
 						{
 							const FText StateText =
diff --git a/Engine/Source/Runtime/Renderer/Private/SceneRenderingUtils.h b/Engine/Source/Runtime/Renderer/Private/SceneRenderingUtils.h
--- a/Engine/Source/Runtime/Renderer/Private/SceneRenderingUtils.h
+++ b/Engine/Source/Runtime/Renderer/Private/SceneRenderingUtils.h
@@ -40,18 +40,18 @@ inline void AddSimpleElementCollectorPass(const FSimpleElementCollector& SimpleE
 }
 
 template <typename PassParametersType>
-inline void AddBatchedElementsPass(const FBatchedElements& BatchedElements, FRDGBuilder& GraphBuilder, PassParametersType* PassParameters, const FViewInfo& View, const FMeshPassProcessorRenderState& DrawRenderState, bool bNeedToSwitchVerticalAxis)
+inline void AddBatchedElementsPass(const FBatchedElements& BatchedElements, FRDGBuilder& GraphBuilder, PassParametersType* PassParameters, const FViewInfo& View, const FMeshPassProcessorRenderState& DrawRenderState)
 {
 	GraphBuilder.AddPass(
 		RDG_EVENT_NAME("BatchedElements"),
 		PassParameters,
 		ERDGPassFlags::Raster,
-		[&View, &BatchedElements, DrawRenderState, bNeedToSwitchVerticalAxis](FRHICommandListImmediate& RHICmdList)
+		[&View, &BatchedElements, DrawRenderState](FRHICommandListImmediate& RHICmdList)
 		{
 			RHICmdList.SetViewport(View.ViewRect.Min.X, View.ViewRect.Min.Y, 0.0f, View.ViewRect.Max.X, View.ViewRect.Max.Y, 1.0f);
 
 			// Draw the view's batched simple elements(lines, sprites, etc).
-			BatchedElements.Draw(RHICmdList, DrawRenderState, View.FeatureLevel, bNeedToSwitchVerticalAxis, View, false);
+			BatchedElements.Draw(RHICmdList, DrawRenderState, View.FeatureLevel, View, false);
 		}
 	);
 }
diff --git a/Engine/Source/Runtime/Renderer/Private/ScreenPass.h b/Engine/Source/Runtime/Renderer/Private/ScreenPass.h
--- a/Engine/Source/Runtime/Renderer/Private/ScreenPass.h
+++ b/Engine/Source/Runtime/Renderer/Private/ScreenPass.h
@@ -387,9 +387,6 @@ enum class EScreenPassDrawFlags : uint8
 {
 	None,
 
-	// Flips the Y axis of the rendered quad. Used by mobile rendering.
-	FlipYAxis = 0x1,
-
 	// Allows the screen pass to use a HMD hidden area mask if one is available. Used for VR.
 	AllowHMDHiddenAreaMask = 0x2
 };
@@ -427,18 +424,6 @@ void DrawScreenPass(
 	FIntPoint LocalOutputSize(OutputSize);
 	EDrawRectangleFlags DrawRectangleFlags = EDRF_UseTriangleOptimization;
 
-	const bool bFlipYAxis = (Flags & EScreenPassDrawFlags::FlipYAxis) == EScreenPassDrawFlags::FlipYAxis;
-
-	if (bFlipYAxis)
-	{
-		// Draw the quad flipped. Requires that the cull mode be disabled.
-		LocalOutputPos.Y = OutputSize.Y;
-		LocalOutputSize.Y = -OutputSize.Y;
-
-		// Triangle optimization currently doesn't work when flipped.
-		DrawRectangleFlags = EDRF_Default;
-	}
-
 	const bool bUseHMDHiddenAreaMask = (Flags & EScreenPassDrawFlags::AllowHMDHiddenAreaMask) == EScreenPassDrawFlags::AllowHMDHiddenAreaMask
 		? View.bHMDHiddenAreaMaskActive
 		: false;
diff --git a/Engine/Source/Runtime/Renderer/Private/ShadowDepthRendering.cpp b/Engine/Source/Runtime/Renderer/Private/ShadowDepthRendering.cpp
--- a/Engine/Source/Runtime/Renderer/Private/ShadowDepthRendering.cpp
+++ b/Engine/Source/Runtime/Renderer/Private/ShadowDepthRendering.cpp
@@ -2188,15 +2188,10 @@ bool FShadowDepthPassMeshProcessor::TryAddMeshBatch(const FMeshBatch& RESTRICT M
 		const ERasterizerCullMode MeshCullMode = ComputeMeshCullMode(MeshBatch, Material, OverrideSettings);
 
 		const bool bTwoSided = Material.IsTwoSided() || PrimitiveSceneProxy->CastsShadowAsTwoSided();
-		// Invert culling order when mobile HDR == false.
-		auto ShaderPlatform = GShaderPlatformForFeatureLevel[FeatureLevel];
-		static auto* MobileHDRCvar = IConsoleManager::Get().FindTConsoleVariableDataInt(TEXT("r.MobileHDR"));
-		check(MobileHDRCvar);
-		const bool bPlatformReversesCulling = (RHINeedsToSwitchVerticalAxis(ShaderPlatform) && MobileHDRCvar->GetValueOnAnyThread() == 0);
 
 		const bool bRenderSceneTwoSided = bTwoSided;
 		const bool bShadowReversesCulling = MeshPassTargetType == EMeshPass::VSMShadowDepth ? false : ShadowDepthType.bOnePassPointLightShadow;
-		const bool bReverseCullMode = XOR(bPlatformReversesCulling, bShadowReversesCulling);
+		const bool bReverseCullMode = bShadowReversesCulling;
 
 		FinalCullMode = bRenderSceneTwoSided ? CM_None : bReverseCullMode ? InverseCullMode(MeshCullMode) : MeshCullMode;
 	}
diff --git a/Engine/Source/Runtime/SlateRHIRenderer/Private/Slate3DRenderer.cpp b/Engine/Source/Runtime/SlateRHIRenderer/Private/Slate3DRenderer.cpp
--- a/Engine/Source/Runtime/SlateRHIRenderer/Private/Slate3DRenderer.cpp
+++ b/Engine/Source/Runtime/SlateRHIRenderer/Private/Slate3DRenderer.cpp
@@ -183,7 +183,6 @@ void FSlate3DRenderer::DrawWindowToTarget_RenderThread(FRHICommandListImmediate&
 
 			FSlateRenderingParams DrawOptions(ProjectionMatrix, FGameTime::CreateDilated(Context.RealTimeSeconds, Context.DeltaRealTimeSeconds, Context.WorldTimeSeconds, Context.DeltaTimeSeconds));
 			// The scene renderer will handle it in this case
-			DrawOptions.bAllowSwitchVerticalAxis = false;
 			DrawOptions.ViewOffset = UE::Slate::CastToVector2f(DrawOffset);
 
 			FTexture2DRHIRef ColorTarget = Context.RenderTarget->GetRenderTargetTexture();
diff --git a/Engine/Source/Runtime/SlateRHIRenderer/Private/SlateMaterialShader.cpp b/Engine/Source/Runtime/SlateRHIRenderer/Private/SlateMaterialShader.cpp
--- a/Engine/Source/Runtime/SlateRHIRenderer/Private/SlateMaterialShader.cpp
+++ b/Engine/Source/Runtime/SlateRHIRenderer/Private/SlateMaterialShader.cpp
@@ -11,7 +11,6 @@ FSlateMaterialShaderVS::FSlateMaterialShaderVS(const FMaterialShaderType::Compil
 	: FMaterialShader(Initializer)
 {
 	ViewProjection.Bind(Initializer.ParameterMap, TEXT("ViewProjection"));
-	SwitchVerticalAxisMultiplier.Bind( Initializer.ParameterMap, TEXT("SwitchVerticalAxisMultiplier"));
 }
 
 
@@ -43,11 +42,6 @@ void FSlateMaterialShaderVS::SetMaterialShaderParameters(FRHICommandList& RHICmd
 	FMaterialShader::SetParameters<FRHIVertexShader>(RHICmdList, ShaderRHI, MaterialRenderProxy, *Material, View);
 }
 
-void FSlateMaterialShaderVS::SetVerticalAxisMultiplier(FRHICommandList& RHICmdList, float InMultiplier )
-{
-	SetShaderValue(RHICmdList, RHICmdList.GetBoundVertexShader(), SwitchVerticalAxisMultiplier, InMultiplier );
-}
-
 bool FSlateMaterialShaderPS::ShouldCompilePermutation(const FMaterialShaderPermutationParameters& Parameters)
 {
 	return Parameters.MaterialParameters.MaterialDomain == MD_UI;
diff --git a/Engine/Source/Runtime/SlateRHIRenderer/Private/SlateMaterialShader.h b/Engine/Source/Runtime/SlateRHIRenderer/Private/SlateMaterialShader.h
--- a/Engine/Source/Runtime/SlateRHIRenderer/Private/SlateMaterialShader.h
+++ b/Engine/Source/Runtime/SlateRHIRenderer/Private/SlateMaterialShader.h
@@ -29,20 +29,12 @@ public:
 
 	void SetMaterialShaderParameters(FRHICommandList& RHICmdList, const FSceneView& View, const FMaterialRenderProxy* MaterialRenderProxy, const FMaterial* Material);
 
-	/**
-	 * Sets the vertical axis multiplier to use depending on graphics api
-	 */
-	void SetVerticalAxisMultiplier(FRHICommandList& RHICmdList, float InMultiplier);
-
 	/** Serializes the shader data */
 	//virtual bool Serialize( FArchive& Ar ) override;
 private:
 	
 		/** ViewProjection parameter used by the shader */
 		LAYOUT_FIELD(FShaderParameter, ViewProjection)
-		/** Parameter used to determine if we need to swtich the vertical axis for opengl */
-		LAYOUT_FIELD(FShaderParameter, SwitchVerticalAxisMultiplier)
-	
 };
 
 class FSlateMaterialShaderPS : public FMaterialShader
diff --git a/Engine/Source/Runtime/SlateRHIRenderer/Private/SlateRHIRenderingPolicy.cpp b/Engine/Source/Runtime/SlateRHIRenderer/Private/SlateRHIRenderingPolicy.cpp
--- a/Engine/Source/Runtime/SlateRHIRenderer/Private/SlateRHIRenderingPolicy.cpp
+++ b/Engine/Source/Runtime/SlateRHIRenderer/Private/SlateRHIRenderingPolicy.cpp
@@ -268,7 +268,6 @@ static bool UpdateScissorRect(
 	FTexture2DRHIRef& DepthStencilTarget,
 	const FSlateClippingState*& LastClippingState,
 	const FVector2f ViewTranslation2D, 
-	bool bSwitchVerticalAxis,
 	FGraphicsPipelineStateInitializer& InGraphicsPSOInit,
 	FSlateStencilClipVertexBuffer& StencilVertexBuffer,
 	const FMatrix& ViewProjection, 
@@ -314,16 +313,7 @@ static bool UpdateScissorRect(
 				const FVector2f TopLeft     = FMath::Min(FMath::Max(ScissorRect.TopLeft     + ViewTranslation2D, FVector2f(0.0f, 0.0f)), ViewSize);
 				const FVector2f BottomRight = FMath::Min(FMath::Max(ScissorRect.BottomRight + ViewTranslation2D, FVector2f(0.0f, 0.0f)), ViewSize);
 				
-				if (bSwitchVerticalAxis)
-				{
-					const int32 MinY = (ViewSize.Y - BottomRight.Y);
-					const int32 MaxY = (ViewSize.Y - TopLeft.Y);
-					RHICmdList.SetScissorRect(true, TopLeft.X, MinY, BottomRight.X, MaxY);
-				}
-				else
-				{
-					RHICmdList.SetScissorRect(true, TopLeft.X, TopLeft.Y, BottomRight.X, BottomRight.Y);
-				}
+				RHICmdList.SetScissorRect(true, TopLeft.X, TopLeft.Y, BottomRight.X, BottomRight.Y);
 
 				// Disable depth/stencil testing by default
 				InGraphicsPSOInit.DepthStencilState = TStaticDepthStencilState<false, CF_Always>::GetRHI();
@@ -379,17 +369,7 @@ static bool UpdateScissorRect(
 					ScissorRect.Right = FMath::Clamp(ScissorRect.Right, ScissorRect.Left, static_cast<float>(BackBufferSize.X));
 					ScissorRect.Bottom = FMath::Clamp(ScissorRect.Bottom, ScissorRect.Top, static_cast<float>(BackBufferSize.Y));
 
-					if (bSwitchVerticalAxis)
-					{
-						const FIntPoint ViewSize = BackBuffer.GetSizeXY();
-						const int32 MinY = (ViewSize.Y - ScissorRect.Bottom);
-						const int32 MaxY = (ViewSize.Y - ScissorRect.Top);
-						RHICmdList.SetScissorRect(true, ScissorRect.Left, MinY, ScissorRect.Right, MaxY);
-					}
-					else
-					{
-						RHICmdList.SetScissorRect(true, ScissorRect.Left, ScissorRect.Top, ScissorRect.Right, ScissorRect.Bottom);
-					}
+					RHICmdList.SetScissorRect(true, ScissorRect.Left, ScissorRect.Top, ScissorRect.Right, ScissorRect.Bottom);
 				}
 
 				// Don't bother setting the render targets unless we actually need to clear them.
@@ -453,7 +433,6 @@ static bool UpdateScissorRect(
 					SetGraphicsPipelineState(RHICmdList, WriteMaskPSOInit, MaskingID + 1);
 
 					VertexShader->SetViewProjection(RHICmdList, FMatrix44f(ViewProjection));
-					VertexShader->SetVerticalAxisMultiplier(RHICmdList, bSwitchVerticalAxis ? -1.0f : 1.0f);
 
 					// Draw the first stencil using SO_Replace, so that we stomp any pixel with a MaskingID + 1.
 					{
@@ -491,7 +470,6 @@ static bool UpdateScissorRect(
 						SetGraphicsPipelineState(RHICmdList, WriteMaskPSOInit, 0);
 
 						VertexShader->SetViewProjection(RHICmdList, FMatrix44f(ViewProjection));
-						VertexShader->SetVerticalAxisMultiplier(RHICmdList, bSwitchVerticalAxis ? -1.0f : 1.0f);
 					}
 				}
 
@@ -664,7 +642,6 @@ void FSlateRHIRenderingPolicy::DrawElements(
 #endif
 
 	const bool bAbsoluteIndices = CVarSlateAbsoluteIndices.GetValueOnRenderThread() != 0;
-	const bool bSwitchVerticalAxis = Params.bAllowSwitchVerticalAxis && RHINeedsToSwitchVerticalAxis(GShaderPlatformForFeatureLevel[GMaxRHIFeatureLevel]);
 
 	// This variable tracks the last clipping state, so that if multiple batches have the same clipping state, we don't have to do any work.
 	const FSlateClippingState* LastClippingState;
@@ -771,7 +748,6 @@ void FSlateRHIRenderingPolicy::DrawElements(
 				DepthStencilTarget,
 				LastClippingState,
 				ViewTranslation2D,
-				bSwitchVerticalAxis,
 				GraphicsPSOInit,
 				StencilVertexBuffer,
 				ViewProjection,
@@ -992,7 +968,6 @@ void FSlateRHIRenderingPolicy::DrawElements(
 
 				{
 					GlobalVertexShader->SetViewProjection(RHICmdList, FMatrix44f(ViewProjection));
-					GlobalVertexShader->SetVerticalAxisMultiplier(RHICmdList, bSwitchVerticalAxis ? -1.0f : 1.0f);
 
 					if (bIsVirtualTexture && (TextureResource != nullptr))
 					{
@@ -1131,7 +1106,6 @@ void FSlateRHIRenderingPolicy::DrawElements(
 
 							{
 								VertexShader->SetViewProjection(RHICmdList, FMatrix44f(ViewProjection));
-								VertexShader->SetVerticalAxisMultiplier(RHICmdList, bSwitchVerticalAxis ? -1.0f : 1.0f);
 								VertexShader->SetMaterialShaderParameters(RHICmdList, ActiveSceneView, MaterialRenderProxy, EffectiveMaterial);
 
 								PixelShader->SetParameters(RHICmdList, ActiveSceneView, MaterialRenderProxy, EffectiveMaterial, ShaderParams);
@@ -1205,7 +1179,6 @@ void FSlateRHIRenderingPolicy::DrawElements(
 						DepthStencilTarget,
 						LastClippingState,
 						ViewTranslation2D,
-						bSwitchVerticalAxis,
 						InGraphicsPSOInit,
 						StencilVertexBuffer,
 						FMatrix(Params.ViewProjectionMatrix),
diff --git a/Engine/Source/Runtime/SlateRHIRenderer/Private/SlateRHIRenderingPolicy.h b/Engine/Source/Runtime/SlateRHIRenderer/Private/SlateRHIRenderingPolicy.h
--- a/Engine/Source/Runtime/SlateRHIRenderer/Private/SlateRHIRenderingPolicy.h
+++ b/Engine/Source/Runtime/SlateRHIRenderer/Private/SlateRHIRenderingPolicy.h
@@ -29,7 +29,6 @@ struct FSlateRenderingParams
 	FMatrix44f ViewProjectionMatrix;
 	FVector2f ViewOffset;
 	FGameTime Time;
-	bool bAllowSwitchVerticalAxis;
 	bool bWireFrame;
 	bool bIsHDR;
 
@@ -37,7 +36,6 @@ struct FSlateRenderingParams
 		: ViewProjectionMatrix(InViewProjectionMatrix)
 		, ViewOffset(0.f, 0.f)
 		, Time(InTime)
-		, bAllowSwitchVerticalAxis(true)
 		, bWireFrame(false)
 		, bIsHDR(false)
 	{
diff --git a/Engine/Source/Runtime/SlateRHIRenderer/Private/SlateShaders.cpp b/Engine/Source/Runtime/SlateRHIRenderer/Private/SlateShaders.cpp
--- a/Engine/Source/Runtime/SlateRHIRenderer/Private/SlateShaders.cpp
+++ b/Engine/Source/Runtime/SlateRHIRenderer/Private/SlateShaders.cpp
@@ -141,7 +141,6 @@ FSlateElementVS::FSlateElementVS( const ShaderMetaType::CompiledShaderInitialize
 {
 	ViewProjection.Bind(Initializer.ParameterMap, TEXT("ViewProjection"));
 	VertexShaderParams.Bind( Initializer.ParameterMap, TEXT("VertexShaderParams"));
-	SwitchVerticalAxisMultiplier.Bind( Initializer.ParameterMap, TEXT("SwitchVerticalAxisMultiplier"));
 }
 
 void FSlateElementVS::SetViewProjection(FRHICommandList& RHICmdList, const FMatrix44f& InViewProjection )
@@ -154,11 +153,6 @@ void FSlateElementVS::SetShaderParameters(FRHICommandList& RHICmdList, const FVe
 	SetShaderValue(RHICmdList, RHICmdList.GetBoundVertexShader(), VertexShaderParams, ShaderParams );
 }
 
-void FSlateElementVS::SetVerticalAxisMultiplier(FRHICommandList& RHICmdList, float InMultiplier )
-{
-	SetShaderValue(RHICmdList, RHICmdList.GetBoundVertexShader(), SwitchVerticalAxisMultiplier, InMultiplier );
-}
-
 /** Serializes the shader data */
 /*bool FSlateElementVS::Serialize( FArchive& Ar )
 {
@@ -166,7 +160,6 @@ void FSlateElementVS::SetVerticalAxisMultiplier(FRHICommandList& RHICmdList, flo
 
 	Ar << ViewProjection;
 	Ar << VertexShaderParams;
-	Ar << SwitchVerticalAxisMultiplier;
 
 	return bShaderHasOutdatedParameters;
 }*/
@@ -181,7 +174,6 @@ FSlateMaskingVS::FSlateMaskingVS(const ShaderMetaType::CompiledShaderInitializer
 {
 	ViewProjection.Bind(Initializer.ParameterMap, TEXT("ViewProjection"));
 	MaskRect.Bind(Initializer.ParameterMap, TEXT("MaskRectPacked"));
-	SwitchVerticalAxisMultiplier.Bind(Initializer.ParameterMap, TEXT("SwitchVerticalAxisMultiplier"));
 }
 
 void FSlateMaskingVS::SetViewProjection(FRHICommandList& RHICmdList, const FMatrix44f& InViewProjection)
@@ -189,11 +181,6 @@ void FSlateMaskingVS::SetViewProjection(FRHICommandList& RHICmdList, const FMatr
 	SetShaderValue(RHICmdList, RHICmdList.GetBoundVertexShader(), ViewProjection, InViewProjection);
 }
 
-void FSlateMaskingVS::SetVerticalAxisMultiplier(FRHICommandList& RHICmdList, float InMultiplier )
-{
-	SetShaderValue(RHICmdList, RHICmdList.GetBoundVertexShader(), SwitchVerticalAxisMultiplier, InMultiplier );
-}
-
 void FSlateMaskingVS::SetMaskRect(FRHICommandList& RHICmdList, const FVector2f TopLeft, const FVector2f TopRight, const FVector2f BotLeft, const FVector2f BotRight)
 {
 	FVector4f MaskRectVal[2] = { FVector4f(TopLeft, TopRight), FVector4f(BotLeft, BotRight) };
@@ -208,7 +195,6 @@ void FSlateMaskingVS::SetMaskRect(FRHICommandList& RHICmdList, const FVector2f T
 
 	Ar << ViewProjection;
 	Ar << MaskRect;
-	Ar << SwitchVerticalAxisMultiplier;
 
 	return bShaderHasOutdatedParameters;
 }*/
diff --git a/Engine/Source/Runtime/SlateRHIRenderer/Private/SlateShaders.h b/Engine/Source/Runtime/SlateRHIRenderer/Private/SlateShaders.h
--- a/Engine/Source/Runtime/SlateRHIRenderer/Private/SlateShaders.h
+++ b/Engine/Source/Runtime/SlateRHIRenderer/Private/SlateShaders.h
@@ -89,11 +89,6 @@ public:
 	 */
 	void SetShaderParameters(FRHICommandList& RHICmdList, const FVector4f& ShaderParams );
 
-	/**
-	 * Sets the vertical axis multiplier to use depending on graphics api
-	 */
-	void SetVerticalAxisMultiplier(FRHICommandList& RHICmdList, float InMultiplier);
-
 	/** Serializes the shader data */
 	//virtual bool Serialize( FArchive& Ar ) override;
 
@@ -102,8 +97,6 @@ private:
 	LAYOUT_FIELD(FShaderParameter, ViewProjection)
 	/** Shader parmeters used by the shader */
 	LAYOUT_FIELD(FShaderParameter, VertexShaderParams)
-	/** Parameter used to determine if we need to swtich the vertical axis for opengl */
-	LAYOUT_FIELD(FShaderParameter, SwitchVerticalAxisMultiplier)
 };
 
 /** 
@@ -494,11 +487,6 @@ public:
 	*/
 	void SetViewProjection(FRHICommandList& RHICmdList, const FMatrix44f& InViewProjection);
 
-	/**
-	 * Sets the vertical axis multiplier to use depending on graphics api
-	 */
-	void SetVerticalAxisMultiplier(FRHICommandList& RHICmdList, float InMultiplier);
-
 	/**
 	 * Sets the mask rect positions
 	 */
@@ -511,8 +499,6 @@ private:
 	LAYOUT_FIELD(FShaderParameter, MaskRect)
 	/** ViewProjection parameter used by the shader */
 	LAYOUT_FIELD(FShaderParameter, ViewProjection)
-	/** Parameter used to determine if we need to swtich the vertical axis for opengl */
-	LAYOUT_FIELD(FShaderParameter, SwitchVerticalAxisMultiplier)
 };
 
 class FSlateMaskingPS : public FGlobalShader
