Added r.Shadow.Virtual.Cache.ForceInvalidateClipmaps to make it possible to force cache invalidations of clipmaps to emulate a moving sun.

diff --git a/Engine/Source/Runtime/Renderer/Private/VirtualShadowMaps/VirtualShadowMapCacheManager.cpp b/Engine/Source/Runtime/Renderer/Private/VirtualShadowMaps/VirtualShadowMapCacheManager.cpp
--- a/Engine/Source/Runtime/Renderer/Private/VirtualShadowMaps/VirtualShadowMapCacheManager.cpp
+++ b/Engine/Source/Runtime/Renderer/Private/VirtualShadowMaps/VirtualShadowMapCacheManager.cpp
@@ -60,6 +60,11 @@ FAutoConsoleVariableRef CVarCacheInvalidateOftenMoving(
 	TEXT("If enabled, Primitive Proxies that are marked as having deformable meshes (HasDeformableMesh() == true) causes invalidations regardless of whether their transforms are updated."),
 	ECVF_RenderThreadSafe);
 
+static TAutoConsoleVariable<int32> CVarForceInvalidateClipmaps(
+	TEXT("r.Shadow.Virtual.Cache.ForceInvalidateClipmaps"),
+	0,
+	TEXT("Forces the clipmap to always invalidate, useful to emulate a moving sun to avoid misrepresenting cache performance."),
+	ECVF_RenderThreadSafe);
 
 void FVirtualShadowMapCacheEntry::UpdateClipmap(
 	int32 VirtualShadowMapId,
@@ -70,7 +75,7 @@ void FVirtualShadowMapCacheEntry::UpdateClipmap(
 	// NOTE: ViewRadiusZ must be constant for a given clipmap level
 	double ViewRadiusZ)
 {
-	bool bCacheValid = (CurrentVirtualShadowMapId != INDEX_NONE);
+	bool bCacheValid = (CurrentVirtualShadowMapId != INDEX_NONE) && CVarForceInvalidateClipmaps.GetValueOnRenderThread() == 0;
 	
 	if (bCacheValid && WorldToLight != Clipmap.WorldToLight)
 	{
