Fixed crash when setting SceneRenderTargetResizeMethod higher than expected

diff --git a/Engine/Source/Runtime/Renderer/Private/SceneTextures.cpp b/Engine/Source/Runtime/Renderer/Private/SceneTextures.cpp
--- a/Engine/Source/Runtime/Renderer/Private/SceneTextures.cpp
+++ b/Engine/Source/Runtime/Renderer/Private/SceneTextures.cpp
@@ -374,7 +374,9 @@ public:
 			if (bUseResizeMethodCVar)
 			{
 				// Otherwise use the setting specified by the console variable.
-				SceneTargetsSizingMethod = (ESizingMethods)FMath::Clamp(CVarSceneTargetsResizeMethod.GetValueOnRenderThread(), 0, (int32)VisibleSizingMethodsCount);
+				// #jira UE-156400: The 'clamp()' includes min and max values, so the range is [0 .. Count-1]
+				// The checkNoEntry() macro is called from 'default:' in the switch() below, when the SceneTargetsSizingMethod is out of the supported range.
+				SceneTargetsSizingMethod = (ESizingMethods)FMath::Clamp(CVarSceneTargetsResizeMethod.GetValueOnRenderThread(), 0, (int32)VisibleSizingMethodsCount - 1);
 			}
 		}
 
