Store depth target when capturing SceneDepth
Correctly select permutation without CSM for single layer water material on mobile

diff --git a/Engine/Source/Runtime/Renderer/Private/MobileBasePass.cpp b/Engine/Source/Runtime/Renderer/Private/MobileBasePass.cpp
--- a/Engine/Source/Runtime/Renderer/Private/MobileBasePass.cpp
+++ b/Engine/Source/Runtime/Renderer/Private/MobileBasePass.cpp
@@ -227,7 +227,8 @@ ELightMapPolicyType MobileBasePass::SelectMeshLightmapPolicy(
 
 		if (!ReadOnlyCVARCache.bAllowStaticLighting || (ReadOnlyCVARCache.bMobileEnableNoPrecomputedLightingCSMShader && Scene && Scene->GetForceNoPrecomputedLighting()))
 		{
-			if (!IsTranslucentBlendMode(BlendMode))
+			if (!IsTranslucentBlendMode(BlendMode) &&
+				!ShadingModels.HasShadingModel(MSM_SingleLayerWater))
 			{
 				// Whether to use a single CSM permutation with a branch in the shader
 				bPrimReceivesCSM |= MobileUseCSMShaderBranch();
diff --git a/Engine/Source/Runtime/Renderer/Private/MobileShadingRenderer.cpp b/Engine/Source/Runtime/Renderer/Private/MobileShadingRenderer.cpp
--- a/Engine/Source/Runtime/Renderer/Private/MobileShadingRenderer.cpp
+++ b/Engine/Source/Runtime/Renderer/Private/MobileShadingRenderer.cpp
@@ -445,6 +445,12 @@ void FMobileSceneRenderer::InitViews(FRDGBuilder& GraphBuilder, FSceneTexturesCo
 						// If the resolve texture is not the same as the MSAA texture, we need to render to scene color and copy to back buffer.
 						|| (NumMSAASamples > 1 && !RHISupportsSeparateMSAAAndResolveTextures(ShaderPlatform))
 						|| bIsFullDepthPrepassEnabled;
+	
+	const bool bSceneDepthCapture = (
+		ViewFamily.SceneCaptureSource == SCS_SceneColorSceneDepth ||
+		ViewFamily.SceneCaptureSource == SCS_SceneDepth ||
+		ViewFamily.SceneCaptureSource == SCS_DeviceDepth);
+
 	const FPlanarReflectionSceneProxy* PlanarReflectionSceneProxy = Scene ? Scene->GetForwardPassGlobalPlanarReflection() : nullptr;
 
 	bRequiresPixelProjectedPlanarRelfectionPass = IsUsingMobilePixelProjectedReflection(ShaderPlatform)
@@ -494,6 +500,7 @@ void FMobileSceneRenderer::InitViews(FRDGBuilder& GraphBuilder, FSceneTexturesCo
 		bSeparateTranslucencyActive ||
 		Views[0].bIsReflectionCapture ||
 		(bDeferredShading && bPostProcessUsesSceneDepth) ||
+		(bDeferredShading && bSceneDepthCapture) ||
 		bShouldRenderVelocities ||
 		bRequireSeparateViewPass ||
 		bIsFullDepthPrepassEnabled ||
@@ -532,9 +539,7 @@ void FMobileSceneRenderer::InitViews(FRDGBuilder& GraphBuilder, FSceneTexturesCo
 	}
 
 	// When we capturing scene depth, use a more precise format for SceneDepthAux as it will be used as a source DepthTexture
-	if (ViewFamily.SceneCaptureSource == SCS_SceneColorSceneDepth ||
-		ViewFamily.SceneCaptureSource == SCS_SceneDepth ||
-		ViewFamily.SceneCaptureSource == SCS_DeviceDepth)
+	if (bSceneDepthCapture)
 	{
 		SceneTexturesConfig.bPreciseDepthAux = true;
 	}
