Moved scene renderer deletion off the render thread.

diff --git a/Engine/Source/Runtime/Renderer/Private/SceneRendering.cpp b/Engine/Source/Runtime/Renderer/Private/SceneRendering.cpp
--- a/Engine/Source/Runtime/Renderer/Private/SceneRendering.cpp
+++ b/Engine/Source/Runtime/Renderer/Private/SceneRendering.cpp
@@ -3995,7 +3995,7 @@ static FAutoConsoleVariableRef CVarSceneRenderCleanUpMode(
 );
 
 enum class ESceneRenderCleanUpMode
-	{
+{
 	Immediate,
 	Deferred,
 	DeferredAndAsync
@@ -4018,7 +4018,7 @@ inline ESceneRenderCleanUpMode GetSceneRenderCleanUpMode()
 	return (ESceneRenderCleanUpMode)GSceneRenderCleanUpMode;
 }
 
-static void DeleteSceneRenderers(FRHICommandListImmediate& RHICmdList, const TArray<FSceneRenderer*>& SceneRenderers)
+static void FinishCleanUp(FRHICommandListImmediate& RHICmdList)
 {
 	static const IConsoleVariable* AsyncDispatch = IConsoleManager::Get().FindConsoleVariable(TEXT("r.RHICmdAsyncRHIThreadDispatch"));
 
@@ -4028,11 +4028,7 @@ static void DeleteSceneRenderers(FRHICommandListImmediate& RHICmdList, const TAr
 		RHICmdList.ImmediateFlush(EImmediateFlushType::WaitForDispatchToRHIThread); // we want to make sure this all gets to the rhi thread this frame and doesn't hang around
 	}
 
-	TRACE_CPUPROFILER_EVENT_SCOPE(DeleteSceneRenderer);
-	for (FSceneRenderer* SceneRenderer : SceneRenderers)
-	{
-		delete SceneRenderer;
-	}
+	TRACE_CPUPROFILER_EVENT_SCOPE(FinishCleanUp);
 
 	// Can release only after all mesh pass tasks are finished.
 	GPrimitiveIdVertexBufferPool.DiscardAll();
@@ -4040,14 +4036,40 @@ static void DeleteSceneRenderers(FRHICommandListImmediate& RHICmdList, const TAr
 	FRenderResource::CoalesceResourceList();
 }
 
-static void ReleaseSceneRenderers(FRHICommandListImmediate& RHICmdList, const TArray<FSceneRenderer*>& SceneRenderers)
+static void DeleteSceneRenderers(const TArray<FSceneRenderer*>& SceneRenderers, FParallelMeshDrawCommandPass::EWaitThread WaitThread)
+{
+	SCOPED_NAMED_EVENT_TEXT("DeleteSceneRenderer", FColor::Red);
+
+	for (FSceneRenderer* SceneRenderer : SceneRenderers)
+	{
+		// Wait for all dispatched shadow mesh draw tasks.
+		for (int32 PassIndex = 0; PassIndex < SceneRenderer->DispatchedShadowDepthPasses.Num(); ++PassIndex)
+		{
+			SceneRenderer->DispatchedShadowDepthPasses[PassIndex]->WaitForTasksAndEmpty(WaitThread);
+		}
+
+		for (FViewInfo& View : SceneRenderer->Views)
+		{
+			View.WaitForTasks(WaitThread);
+		}
+	}
+
+	FViewInfo::DestroyAllSnapshots(WaitThread);
+
+	for (FSceneRenderer* SceneRenderer : SceneRenderers)
+	{
+		delete SceneRenderer;
+	}
+}
+
+static void WaitForTasksAndDeleteSceneRenderers(FRHICommandListImmediate& RHICmdList, const TArray<FSceneRenderer*>& SceneRenderers)
 {
 	{
 		QUICK_SCOPE_CYCLE_COUNTER(STAT_DeleteSceneRenderer_WaitForTasks);
 		RHICmdList.ImmediateFlush(EImmediateFlushType::WaitForOutstandingTasksOnly);
 	}
 
-	FSceneRenderer::WaitForTasksAndClearSnapshots(SceneRenderers, FParallelMeshDrawCommandPass::EWaitThread::Render);
+	DeleteSceneRenderers(SceneRenderers, FParallelMeshDrawCommandPass::EWaitThread::Render);
 }
 
 void FSceneRenderer::RenderThreadBegin(FRHICommandListImmediate& RHICmdList)
@@ -4122,9 +4144,8 @@ void FSceneRenderer::RenderThreadEnd(FRHICommandListImmediate& RHICmdList, const
 
 	if (GSceneRenderCleanUpState.CompletionMode == ESceneRenderCleanUpMode::Immediate)
 	{
-		// Mem stack mark is stored on first scene renderer
-		ReleaseSceneRenderers(RHICmdList, SceneRenderers);
-		DeleteSceneRenderers(RHICmdList, SceneRenderers);
+		WaitForTasksAndDeleteSceneRenderers(RHICmdList, SceneRenderers);
+		FinishCleanUp(RHICmdList);
 	}
 	else
 	{
@@ -4175,7 +4196,7 @@ void FSceneRenderer::RenderThreadEnd(FRHICommandListImmediate& RHICmdList, const
 
 			GSceneRenderCleanUpState.Task = FFunctionGraphTask::CreateAndDispatchWhenReady([LocalSceneRenderers = CopyTemp(SceneRenderers)]
 			{
-				FSceneRenderer::WaitForTasksAndClearSnapshots(LocalSceneRenderers, FParallelMeshDrawCommandPass::EWaitThread::TaskAlreadyWaited);
+				DeleteSceneRenderers(LocalSceneRenderers, FParallelMeshDrawCommandPass::EWaitThread::TaskAlreadyWaited);
 			}, TStatId(), &CommandListTasks);
 		}
 	}
@@ -4193,7 +4214,7 @@ void FSceneRenderer::CleanUp(FRHICommandListImmediate& RHICmdList)
 		switch (GSceneRenderCleanUpState.CompletionMode)
 		{
 		case ESceneRenderCleanUpMode::Deferred:
-			ReleaseSceneRenderers(RHICmdList, GSceneRenderCleanUpState.Renderers);
+			WaitForTasksAndDeleteSceneRenderers(RHICmdList, GSceneRenderCleanUpState.Renderers);
 			break;
 		case ESceneRenderCleanUpMode::DeferredAndAsync:
 			GSceneRenderCleanUpState.Task->Wait(ENamedThreads::GetRenderThread_Local());
@@ -4201,7 +4222,7 @@ void FSceneRenderer::CleanUp(FRHICommandListImmediate& RHICmdList)
 		}
 	}
 
-	DeleteSceneRenderers(RHICmdList, GSceneRenderCleanUpState.Renderers);
+	FinishCleanUp(RHICmdList);
 	GSceneRenderCleanUpState = {};
 }
 
@@ -4215,7 +4236,7 @@ void FSceneRenderer::WaitForCleanUpTasks(FRHICommandListImmediate& RHICmdList)
 	switch (GSceneRenderCleanUpState.CompletionMode)
 	{
 	case ESceneRenderCleanUpMode::Deferred:
-		ReleaseSceneRenderers(RHICmdList, GSceneRenderCleanUpState.Renderers);
+		WaitForTasksAndDeleteSceneRenderers(RHICmdList, GSceneRenderCleanUpState.Renderers);
 		break;
 	case ESceneRenderCleanUpMode::DeferredAndAsync:
 		GSceneRenderCleanUpState.Task->Wait(ENamedThreads::GetRenderThread_Local());
@@ -4226,27 +4247,6 @@ void FSceneRenderer::WaitForCleanUpTasks(FRHICommandListImmediate& RHICmdList)
 	GSceneRenderCleanUpState.bWaitForTasksComplete = true;
 }
 
-void FSceneRenderer::WaitForTasksAndClearSnapshots(const TArray<FSceneRenderer*>& SceneRenderers, FParallelMeshDrawCommandPass::EWaitThread WaitThread)
-{
-	SCOPED_NAMED_EVENT_TEXT("FSceneRenderer::WaitForTasksAndClearSnapshots", FColor::Red);
-
-	for (FSceneRenderer* SceneRenderer : SceneRenderers)
-	{
-		// Wait for all dispatched shadow mesh draw tasks.
-		for (int32 PassIndex = 0; PassIndex < SceneRenderer->DispatchedShadowDepthPasses.Num(); ++PassIndex)
-		{
-			SceneRenderer->DispatchedShadowDepthPasses[PassIndex]->WaitForTasksAndEmpty(WaitThread);
-		}
-
-		for (FViewInfo& View : SceneRenderer->Views)
-		{
-			View.WaitForTasks(WaitThread);
-		}
-	}
-
-	FViewInfo::DestroyAllSnapshots(WaitThread);
-}
-
 void ResetAndShrinkModifiedBounds(TArray<FRenderBounds>& Bounds)
 {
 	const int32 MaxAllocatedSize = FMath::RoundUpToPowerOfTwo(FMath::Max<uint32>(DistanceField::MinPrimitiveModifiedBoundsAllocation, Bounds.Num()));
diff --git a/Engine/Source/Runtime/Renderer/Private/SceneRendering.h b/Engine/Source/Runtime/Renderer/Private/SceneRendering.h
--- a/Engine/Source/Runtime/Renderer/Private/SceneRendering.h
+++ b/Engine/Source/Runtime/Renderer/Private/SceneRendering.h
@@ -2038,9 +2038,6 @@ public:
 	*/
 	static void RENDERER_API ViewExtensionPreRender_RenderThread(FRHICommandListImmediate& RHICmdList, FSceneRenderer* SceneRenderer);
 
-	/** the last thing we do with scene renderers, lots of cleanup related to the threading. **/
-	static void WaitForTasksAndClearSnapshots(const TArray<FSceneRenderer*>& SceneRenderers, FParallelMeshDrawCommandPass::EWaitThread WaitThread);
-
 	/** Called to release any deallocations that were deferred until the next render. */
 	static void CleanUp(FRHICommandListImmediate& RHICmdList);
 
