fix for translucency with multi-view
NumProcessedViews was incremented in both the for loop statement and at the end of the for loop. If the first view skipped the rendering because it doesn't have the relevant flags, we would not clear the rendertarget. On some platforms, this would end up with a blackscreen during ComposeTranslucencyToNewSceneColor because the values would not be neutral ((0,0,0,0) instead of (0,0,0,1))

Now, we clear the Rendertarget for the first view that gets processed

diff --git a/Engine/Source/Runtime/Renderer/Private/TranslucentRendering.cpp b/Engine/Source/Runtime/Renderer/Private/TranslucentRendering.cpp
--- a/Engine/Source/Runtime/Renderer/Private/TranslucentRendering.cpp
+++ b/Engine/Source/Runtime/Renderer/Private/TranslucentRendering.cpp
@@ -1246,7 +1246,7 @@ void FDeferredShadingSceneRenderer::RenderTranslucencyInner(
 		// Create resources shared by each view (each view data is tiled into each of the render target resources)
 		FRDGTextureMSAA SharedColorTexture = CreatePostDOFTranslucentTexture(GraphBuilder, TranslucencyPass, SeparateTranslucencyDimensions, bIsModulate, ShaderPlatform);
 
-		for (int32 ViewIndex = 0, NumProcessedViews = 0; ViewIndex < Views.Num(); ++ViewIndex, ++NumProcessedViews)
+		for (int32 ViewIndex = 0, NumProcessedViews = 0; ViewIndex < Views.Num(); ++ViewIndex)
 		{
 			FViewInfo& View = Views[ViewIndex];
 			const ETranslucencyView TranslucencyView = GetTranslucencyView(View);
