Fix hair strands support for light function.

The fetched scene depth was incorrect. This CLs ensure that if the platform support hair strands, light function evaluated for hair strands, will fetch the correct depth value.

diff --git a/Engine/Shaders/Private/LightFunctionPixelShader.usf b/Engine/Shaders/Private/LightFunctionPixelShader.usf
--- a/Engine/Shaders/Private/LightFunctionPixelShader.usf
+++ b/Engine/Shaders/Private/LightFunctionPixelShader.usf
@@ -12,8 +12,8 @@
 // used like this: mul(float4(SvPosition.xyz, 1), SvPositionToLight);
 float4x4 SvPositionToLight;
 
-/** Fade distance in x, disabled brightness in y, output for preview shadows mask in z. */
-float3 LightFunctionParameters2;
+/** Fade distance in x, disabled brightness in y, output for preview shadows mask in z, hair strands in w */
+float4 LightFunctionParameters2;
 
 void Main(
 	// in float4 ScreenPosition : TEXCOORD0, // Not used and not output in VS
@@ -27,6 +27,13 @@ void Main(
 
 	// make SvPosition appear to be rasterized with the depth from the depth buffer
 	SvPosition.z = LookupDeviceZ(ScreenUV);
+#if HAIR_STRANDS_SUPPORTED
+	BRANCH
+	if (LightFunctionParameters2.w > 0)
+	{
+		SvPosition.z = Texture2DSampleLevel(HairStrands.HairOnlyDepthTexture, SceneTexturesStruct_SceneDepthTextureSampler, ScreenUV, 0).r;
+	}
+#endif
 
 	float3 LightVector;
 	{
diff --git a/Engine/Source/Runtime/Renderer/Private/LightFunctionRendering.cpp b/Engine/Source/Runtime/Renderer/Private/LightFunctionRendering.cpp
--- a/Engine/Source/Runtime/Renderer/Private/LightFunctionRendering.cpp
+++ b/Engine/Source/Runtime/Renderer/Private/LightFunctionRendering.cpp
@@ -23,6 +23,7 @@
 #include "LightRendering.h"
 #include "PipelineStateCache.h"
 #include "ClearQuad.h"
+#include "HairStrands/HairStrandsData.h"
 
 /**
  * A vertex shader for projecting a light function onto the scene.
@@ -93,6 +94,7 @@ public:
 	{
 		FMaterialShader::ModifyCompilationEnvironment(Parameters, OutEnvironment);
 		OutEnvironment.SetDefine(TEXT("STRATA_INLINE_SHADING"), 1);
+		OutEnvironment.SetDefine(TEXT("HAIR_STRANDS_SUPPORTED"), IsHairStrandsSupported(EHairStrandsShaderType::Strands, Parameters.Platform) ? 1 : 0);
 	}
 
 	FLightFunctionPS() {}
@@ -104,7 +106,7 @@ public:
 		LightFunctionParameters2.Bind(Initializer.ParameterMap,TEXT("LightFunctionParameters2"));
 	}
 
-	void SetParameters(FRHICommandList& RHICmdList, const FViewInfo& View, const FLightSceneInfo* LightSceneInfo, const FMaterialRenderProxy* MaterialProxy, const FMaterial& Material, bool bRenderingPreviewShadowIndicator, float ShadowFadeFraction )
+	void SetParameters(FRHICommandList& RHICmdList, const FViewInfo& View, const FLightSceneInfo* LightSceneInfo, const FMaterialRenderProxy* MaterialProxy, const FMaterial& Material, bool bRenderingPreviewShadowIndicator, float ShadowFadeFraction, bool bUseHairStrands, FRHIUniformBuffer* HairStrandsUniformBuffer)
 	{
 		FRHIPixelShader* ShaderRHI = RHICmdList.GetBoundPixelShader();
 
@@ -143,13 +145,20 @@ public:
 
 		LightFunctionParameters.Set(RHICmdList, ShaderRHI, LightSceneInfo, ShadowFadeFraction);
 
-		SetShaderValue(RHICmdList, ShaderRHI, LightFunctionParameters2, FVector3f(
+		SetShaderValue(RHICmdList, ShaderRHI, LightFunctionParameters2, FVector4f(
 			LightSceneInfo->Proxy->GetLightFunctionFadeDistance(), 
 			LightSceneInfo->Proxy->GetLightFunctionDisabledBrightness(),
-			bRenderingPreviewShadowIndicator ? 1.0f : 0.0f));
+			bRenderingPreviewShadowIndicator ? 1.0f : 0.0f,
+			bUseHairStrands ? 1.0f : 0.0f));
 
-		auto DeferredLightParameter = GetUniformBufferParameter<FDeferredLightUniformStruct>();
+		
+		auto HairStrandsParameter = GetUniformBufferParameter<FHairStrandsViewUniformParameters>();
+		if (HairStrandsParameter.IsBound() && HairStrandsUniformBuffer)
+		{
+			SetUniformBufferParameter(RHICmdList, ShaderRHI, HairStrandsParameter, HairStrandsUniformBuffer);
+		}
 
+		auto DeferredLightParameter = GetUniformBufferParameter<FDeferredLightUniformStruct>();
 		if (DeferredLightParameter.IsBound())
 		{
 			SetDeferredLightParameters(RHICmdList, ShaderRHI, DeferredLightParameter, LightSceneInfo, View);
@@ -266,6 +275,7 @@ bool FDeferredShadingSceneRenderer::RenderPreviewShadowsIndicator(
 
 BEGIN_SHADER_PARAMETER_STRUCT(FRenderLightFunctionForMaterialParameters, )
 	SHADER_PARAMETER_RDG_UNIFORM_BUFFER(FSceneTextureUniformParameters, SceneTextures)
+	SHADER_PARAMETER_RDG_UNIFORM_BUFFER(FHairStrandsViewUniformParameters, HairStrands)
 	RENDER_TARGET_BINDING_SLOTS()
 END_SHADER_PARAMETER_STRUCT()
 
@@ -325,14 +335,21 @@ bool FDeferredShadingSceneRenderer::RenderLightFunctionForMaterial(
 		{
 			FRenderLightFunctionForMaterialParameters* PassParameters = GraphBuilder.AllocParameters<FRenderLightFunctionForMaterialParameters>();
 			PassParameters->SceneTextures = SceneTextures.UniformBuffer;
+			PassParameters->HairStrands = HairStrands::BindHairStrandsViewUniformParameters(View);
 			PassParameters->RenderTargets[0] = FRenderTargetBinding(ScreenShadowMaskTexture, bLightAttenuationCleared ? ERenderTargetLoadAction::ELoad : ERenderTargetLoadAction::ENoAction);
 			PassParameters->RenderTargets.DepthStencil = FDepthStencilBinding(SceneTextures.Depth.Target, ERenderTargetLoadAction::ELoad, ERenderTargetLoadAction::ELoad, FExclusiveDepthStencil::DepthRead_StencilWrite);
 
+			// If render shadow mask for hair strands, then swap depth to hair only depth
+			if (bUseHairStrands)
+			{
+				PassParameters->RenderTargets.DepthStencil = FDepthStencilBinding(View.HairStrandsViewData.VisibilityData.HairOnlyDepthTexture, ERenderTargetLoadAction::ELoad, ERenderTargetLoadAction::ELoad, FExclusiveDepthStencil::DepthRead_StencilWrite);
+			}
+
 			GraphBuilder.AddPass(
 				RDG_EVENT_NAME("LightFunction Material=%s", *MaterialForRendering->GetFriendlyName()),
 				PassParameters,
 				ERDGPassFlags::Raster,
-				[&View, LightSceneInfo, LightSceneProxy, MaterialProxyForRendering, MaterialForRendering, MaterialShaders, bLightAttenuationCleared, bProjectingForForwardShading, bRenderingPreviewShadowsIndicator, bUseHairStrands](FRHICommandList& RHICmdList)
+				[&View, PassParameters, LightSceneInfo, LightSceneProxy, MaterialProxyForRendering, MaterialForRendering, MaterialShaders, bLightAttenuationCleared, bProjectingForForwardShading, bRenderingPreviewShadowsIndicator, bUseHairStrands](FRHICommandList& RHICmdList)
 			{
 				FSphere LightBounds = LightSceneProxy->GetBoundingSphere();
 
@@ -411,10 +428,12 @@ bool FDeferredShadingSceneRenderer::RenderLightFunctionForMaterial(
 						LightSceneProxy->SetScissorRect(RHICmdList, View, TotalRect);
 					}
 
+					FRHIUniformBuffer* HairStrandsUniformBuffer = PassParameters->HairStrands.GetUniformBuffer()->GetRHI();
+
 					// Render a bounding light sphere.
 					SetGraphicsPipelineState(RHICmdList, GraphicsPSOInit, 0);
 					VertexShader->SetParameters(RHICmdList, View, LightSceneInfo);
-					PixelShader->SetParameters(RHICmdList, View, LightSceneInfo, MaterialProxyForRendering, *MaterialForRendering, bRenderingPreviewShadowsIndicator, FadeAlpha);
+					PixelShader->SetParameters(RHICmdList, View, LightSceneInfo, MaterialProxyForRendering, *MaterialForRendering, bRenderingPreviewShadowsIndicator, FadeAlpha, bUseHairStrands, HairStrandsUniformBuffer);
 
 					// Project the light function using a sphere around the light
 					//@todo - could use a cone for spotlights
