Mobile: Issue occlusion queries at the end of base-pass after translucency. Make it consistent between deferred and forward path.

diff --git a/Engine/Source/Runtime/Renderer/Private/MobileShadingRenderer.cpp b/Engine/Source/Runtime/Renderer/Private/MobileShadingRenderer.cpp
--- a/Engine/Source/Runtime/Renderer/Private/MobileShadingRenderer.cpp
+++ b/Engine/Source/Runtime/Renderer/Private/MobileShadingRenderer.cpp
@@ -1146,13 +1146,6 @@ void FMobileSceneRenderer::RenderForwardSinglePass(FRDGBuilder& GraphBuilder, FM
 		RenderMobileBasePass(RHICmdList, View);
 		RenderMobileDebugView(RHICmdList, View);
 		RHICmdList.PollOcclusionQueries();
-		const bool bAdrenoOcclusionMode = (CVarMobileAdrenoOcclusionMode.GetValueOnRenderThread() != 0 && IsOpenGLPlatform(ShaderPlatform));
-		if (!bAdrenoOcclusionMode)
-		{
-			// Issue occlusion queries
-			RHICmdList.SetCurrentStat(GET_STATID(STAT_CLMM_Occlusion));
-			RenderOcclusion(RHICmdList, View);
-		}
 		PostRenderBasePass(RHICmdList, View);
 		// scene depth is read only and can be fetched
 		RHICmdList.NextSubpass();
@@ -1162,14 +1155,20 @@ void FMobileSceneRenderer::RenderForwardSinglePass(FRDGBuilder& GraphBuilder, FM
 		RenderFog(RHICmdList, View);
 		// Draw translucency.
 		RenderTranslucency(RHICmdList, View);
-		if (bAdrenoOcclusionMode)
+
+		if (!bIsFullDepthPrepassEnabled)
 		{
-			RHICmdList.SetCurrentStat(GET_STATID(STAT_CLMM_Occlusion));
-			// flush
-			RHICmdList.SubmitCommandsHint();
 			// Issue occlusion queries
+			RHICmdList.SetCurrentStat(GET_STATID(STAT_CLMM_Occlusion));
+			const bool bAdrenoOcclusionMode = (CVarMobileAdrenoOcclusionMode.GetValueOnRenderThread() != 0 && IsOpenGLPlatform(ShaderPlatform));
+			if (bAdrenoOcclusionMode)
+			{
+				// flush
+				RHICmdList.SubmitCommandsHint();
+			}
 			RenderOcclusion(RHICmdList, View);
 		}
+
 		// Pre-tonemap before MSAA resolve (iOS only)
 		PreTonemapMSAA(RHICmdList, SceneTextures);
 	});
@@ -1195,10 +1194,7 @@ void FMobileSceneRenderer::RenderForwardMultiPass(FRDGBuilder& GraphBuilder, FMo
 		RHICmdList.SetCurrentStat(GET_STATID(STAT_CLMM_Opaque));
 		RenderMobileBasePass(RHICmdList, View);
 		RenderMobileDebugView(RHICmdList, View);
-		// Issue occlusion queries
-		RHICmdList.SetCurrentStat(GET_STATID(STAT_CLMM_Occlusion));
 		RHICmdList.PollOcclusionQueries();
-		RenderOcclusion(RHICmdList, View);
 		PostRenderBasePass(RHICmdList, View);
 	});
 
@@ -1231,7 +1227,7 @@ void FMobileSceneRenderer::RenderForwardMultiPass(FRDGBuilder& GraphBuilder, FMo
 	SecondPassParameters->RenderTargets = BasePassRenderTargets;
 
 	GraphBuilder.AddPass(
-		{},
+		RDG_EVENT_NAME("DecalsAndTranslucency"),
 		SecondPassParameters,
 		ERDGPassFlags::Raster,
 		[this, SecondPassParameters, ViewIndex, &View, &SceneTextures](FRHICommandListImmediate& RHICmdList)
@@ -1244,6 +1240,14 @@ void FMobileSceneRenderer::RenderForwardMultiPass(FRDGBuilder& GraphBuilder, FMo
 		RenderFog(RHICmdList, View);
 		// Draw translucency.
 		RenderTranslucency(RHICmdList, View);
+
+		if (!bIsFullDepthPrepassEnabled)
+		{
+			// Issue occlusion queries
+			RHICmdList.SetCurrentStat(GET_STATID(STAT_CLMM_Occlusion));
+			RenderOcclusion(RHICmdList, View);
+		}
+
 		// Pre-tonemap before MSAA resolve (iOS only)
 		PreTonemapMSAA(RHICmdList, SceneTextures);
 	});
@@ -1432,35 +1436,29 @@ void FMobileSceneRenderer::RenderDeferredSinglePass(FRDGBuilder& GraphBuilder, c
 		// Opaque and masked
 		RHICmdList.SetCurrentStat(GET_STATID(STAT_CLMM_Opaque));
 		RenderMobileBasePass(RHICmdList, View);
-		
-		// Issue occlusion queries
-		if (!bUsingPixelLocalStorage)
-		{
-			RHICmdList.SetCurrentStat(GET_STATID(STAT_CLMM_Occlusion));
-			RenderOcclusion(RHICmdList, View);
-		}
 		RHICmdList.PollOcclusionQueries();
-
 		PostRenderBasePass(RHICmdList, View);
 		// SceneColor + GBuffer write, SceneDepth is read only
 		RHICmdList.NextSubpass();
+		RHICmdList.SetCurrentStat(GET_STATID(STAT_CLMM_Translucency));
 		RenderDecals(RHICmdList, View);
 		// SceneColor write, SceneDepth is read only
 		RHICmdList.NextSubpass();
 		MobileDeferredShadingPass(RHICmdList, ViewIndex, NumViews, View, *Scene, SortedLightSet, VisibleLightInfos);
-
 		if (bUsingPixelLocalStorage)
 		{
 			MobileDeferredCopyBuffer<FMobileDeferredCopyPLSPS>(RHICmdList, View);
-			// SceneColor write, SceneDepth is read only
-			RHICmdList.NextSubpass();
-
-			RHICmdList.SetCurrentStat(GET_STATID(STAT_CLMM_Occlusion));
-			RenderOcclusion(RHICmdList, View);
 		}
 		RenderFog(RHICmdList, View);
 		// Draw translucency.
 		RenderTranslucency(RHICmdList, View);
+
+		if (!bIsFullDepthPrepassEnabled)
+		{
+			// Issue occlusion queries
+			RHICmdList.SetCurrentStat(GET_STATID(STAT_CLMM_Occlusion));
+			RenderOcclusion(RHICmdList, View);
+		}
 	});
 }
 
@@ -1478,10 +1476,7 @@ void FMobileSceneRenderer::RenderDeferredMultiPass(FRDGBuilder& GraphBuilder, cl
 		// Opaque and masked
 		RHICmdList.SetCurrentStat(GET_STATID(STAT_CLMM_Opaque));
 		RenderMobileBasePass(RHICmdList, View);
-		// Issue occlusion queries
-		RHICmdList.SetCurrentStat(GET_STATID(STAT_CLMM_Occlusion));
 		RHICmdList.PollOcclusionQueries();
-		RenderOcclusion(RHICmdList, View);
 		PostRenderBasePass(RHICmdList, View);
 	});
 
@@ -1505,7 +1500,7 @@ void FMobileSceneRenderer::RenderDeferredMultiPass(FRDGBuilder& GraphBuilder, cl
 	SecondPassParameters->RenderTargets = BasePassRenderTargets;
 
 	GraphBuilder.AddPass(
-		{},
+		RDG_EVENT_NAME("Decals"),
 		SecondPassParameters,
 		ERDGPassFlags::Raster,
 		[this, SecondPassParameters, &View](FRHICommandListImmediate& RHICmdList)
@@ -1529,16 +1524,24 @@ void FMobileSceneRenderer::RenderDeferredMultiPass(FRDGBuilder& GraphBuilder, cl
 	ThirdPassParameters->RenderTargets = BasePassRenderTargets;
 
 	GraphBuilder.AddPass(
-		{},
+		RDG_EVENT_NAME("LightingAndTranslucency"),
 		ThirdPassParameters,
 		ERDGPassFlags::Raster,
 		[this, ThirdPassParameters, ViewIndex, NumViews, &View, &SceneTextures, &SortedLightSet](FRHICommandListImmediate& RHICmdList)
 	{
 		RHICmdList.NextSubpass();
+		RHICmdList.SetCurrentStat(GET_STATID(STAT_CLMM_Translucency));
 		MobileDeferredShadingPass(RHICmdList, ViewIndex, NumViews, View, *Scene, SortedLightSet, VisibleLightInfos);
 		RenderFog(RHICmdList, View);
 		// Draw translucency.
 		RenderTranslucency(RHICmdList, View);
+
+		if (!bIsFullDepthPrepassEnabled)
+		{
+			// Issue occlusion queries
+			RHICmdList.SetCurrentStat(GET_STATID(STAT_CLMM_Occlusion));
+			RenderOcclusion(RHICmdList, View);
+		}
 	});
 }
 
