* Add shader print depth texture for culling filled primitives.
* Add filled triangle primitive.

diff --git a/Engine/Shaders/Private/ShaderPrint.ush b/Engine/Shaders/Private/ShaderPrint.ush
--- a/Engine/Shaders/Private/ShaderPrint.ush
+++ b/Engine/Shaders/Private/ShaderPrint.ush
@@ -893,6 +893,28 @@ void AddTriangleElement(FTriangleElement In, uint Index)
 	ShaderPrint_RWEntryBuffer[Index12 + 11]= Out.Packed2.w;
 }
 
+///////////////////////////////////////////////////////////////////////////////////////////////////
+// Filled Triangle
+void AddFilledTriangleTWS(float3 Pos0, float3 Pos1, float3 Pos2, float4 Color)
+{
+	uint Offset = 0;
+	if (AllocateTriangleElement(1, Offset))
+	{
+		FTriangleElement Element;
+		Element.bIsScreenSpace = false;
+		Element.Color = Color;
+
+		Element.Pos0 = Pos0;
+		Element.Pos1 = Pos1;
+		Element.Pos2 = Pos2;
+		AddTriangleElement(Element, Offset + 0);
+	}
+}
+
+void AddFilledTriangleWS (float3 Pos0, float3 Pos1, float3 Pos2, float4 Color)							{ AddFilledTriangleTWS(Pos0 + ShaderPrintData.TranslatedWorldOffset, Pos1 + ShaderPrintData.TranslatedWorldOffset, Pos2 + ShaderPrintData.TranslatedWorldOffset, Color); }
+void AddFilledTriangleTWS(FShaderPrintContext Ctx, float3 Pos0, float3 Pos1, float3 Pos2, float4 Color) { if (Ctx.bIsActive) { AddFilledTriangleTWS(Pos0, Pos1, Pos2, Color); } }
+void AddFilledTriangleWS (FShaderPrintContext Ctx, float3 Pos0, float3 Pos1, float3 Pos2, float4 Color) { if (Ctx.bIsActive) { AddFilledTriangleTWS(Pos0 + ShaderPrintData.TranslatedWorldOffset, Pos1 + ShaderPrintData.TranslatedWorldOffset, Pos2 + ShaderPrintData.TranslatedWorldOffset, Color); } }
+
 ///////////////////////////////////////////////////////////////////////////////////////////////////
 // Filled Quad
 
diff --git a/Engine/Shaders/Private/ShaderPrintDrawPrimitive.usf b/Engine/Shaders/Private/ShaderPrintDrawPrimitive.usf
--- a/Engine/Shaders/Private/ShaderPrintDrawPrimitive.usf
+++ b/Engine/Shaders/Private/ShaderPrintDrawPrimitive.usf
@@ -6,6 +6,7 @@
 #include "ShaderPrint.ush"
 
 float3 TranslatedWorldOffsetConversion;
+uint bCheckerboardEnabled;
 
 // The shader used to render the debug element on view
 StructuredBuffer<uint> ShaderDrawDebugPrimitive;
@@ -102,7 +103,7 @@ void ShaderDrawDebugPS(
 
 	const float ColorScale = bIsHidden ? 0.4f : 1;
 	const float CheckerboardScale = ((PixelCoord.x%2) ^ (PixelCoord.y%2)) > 0 ? 1 : 0.0;
-	const float OcclusionScale = bIsHidden ? CheckerboardScale : 1;
+	const float OcclusionScale = (bCheckerboardEnabled && bIsHidden) ? CheckerboardScale : 1;
 	OutSceneColor = OcclusionScale * ColorAlpha * float4(ColorAlpha.aaa * ColorScale, 1.0); // Pre multiplied alpha with color scale
 }
 #endif
diff --git a/Engine/Source/Runtime/Renderer/Private/ShaderPrint.cpp b/Engine/Source/Runtime/Renderer/Private/ShaderPrint.cpp
--- a/Engine/Source/Runtime/Renderer/Private/ShaderPrint.cpp
+++ b/Engine/Source/Runtime/Renderer/Private/ShaderPrint.cpp
@@ -483,6 +483,7 @@ namespace ShaderPrint
 			SHADER_PARAMETER(FVector2f, OriginalViewRectMin)
 			SHADER_PARAMETER(FVector2f, OriginalViewSize)
 			SHADER_PARAMETER(FVector2f, OriginalBufferInvSize)
+			SHADER_PARAMETER(uint32, bCheckerboardEnabled)
 			SHADER_PARAMETER_RDG_TEXTURE(Texture2D, DepthTexture)
 			SHADER_PARAMETER_SAMPLER(SamplerState, DepthSampler)
 			RENDER_TARGET_BINDING_SLOTS()
@@ -869,16 +870,21 @@ namespace ShaderPrint
 		TShaderMapRef<FShaderDrawDebugVS> VertexShader(GlobalShaderMap, PermutationVector);
 		TShaderMapRef<FShaderDrawDebugPS> PixelShader(GlobalShaderMap);
 
+		// Create a transient depth texture which allows to depth test filled primitive between themselves. These primitives are not culled against the scene depth texture, but only 'checkerboarded'.
+		FRDGTextureRef TransientDepthTexture = GraphBuilder.CreateTexture(FRDGTextureDesc::Create2D(OutputTexture->Desc.Extent, PF_DepthStencil, FClearValueBinding::DepthFar, TexCreate_DepthStencilTargetable | TexCreate_ShaderResource), TEXT("ShaderPrint.DepthTexture"));
+
 		FShaderDrawVSPSParameters* PassParameters = GraphBuilder.AllocParameters<FShaderDrawVSPSParameters >();
 		PassParameters->VS.TranslatedWorldOffsetConversion = FVector3f(TranslatedWorldOffsetConversion);
 		PassParameters->VS.TranslatedWorldToClip = FMatrix44f(TranslatedWorldToClip);
 		PassParameters->PS.RenderTargets[0] = FRenderTargetBinding(OutputTexture, ERenderTargetLoadAction::ELoad);
+		PassParameters->PS.RenderTargets.DepthStencil = FDepthStencilBinding(TransientDepthTexture, ERenderTargetLoadAction::EClear, ERenderTargetLoadAction::ENoAction, FExclusiveDepthStencil::DepthWrite_StencilNop);
 		PassParameters->PS.OutputInvResolution = FVector2f(1.f / UnscaledViewRect.Width(), 1.f / UnscaledViewRect.Height());
 		PassParameters->PS.OriginalViewRectMin = FVector2f(ViewRect.Min);
 		PassParameters->PS.OriginalViewSize = FVector2f(ViewRect.Width(), ViewRect.Height());
 		PassParameters->PS.OriginalBufferInvSize = FVector2f(1.f / DepthTexture->Desc.Extent.X, 1.f / DepthTexture->Desc.Extent.Y);
 		PassParameters->PS.DepthTexture = DepthTexture;
 		PassParameters->PS.DepthSampler = TStaticSamplerState<SF_Point, AM_Clamp, AM_Clamp, AM_Clamp>::GetRHI();
+		PassParameters->PS.bCheckerboardEnabled = bLines ? 1u : 0u;
 		PassParameters->VS.ShaderDrawDebugPrimitive = GraphBuilder.CreateSRV(ShaderPrintPrimitiveBuffer);
 		PassParameters->VS.IndirectBuffer = IndirectBuffer;
 		PassParameters->VS.Common = ShaderPrintData.UniformBuffer;
@@ -900,7 +906,7 @@ namespace ShaderPrint
 
 				FGraphicsPipelineStateInitializer GraphicsPSOInit;
 				RHICmdList.ApplyCachedRenderTargets(GraphicsPSOInit);
-				GraphicsPSOInit.DepthStencilState = TStaticDepthStencilState<false, CF_Always>::GetRHI();
+				GraphicsPSOInit.DepthStencilState = TStaticDepthStencilState<true, CF_DepthNearOrEqual>::GetRHI();
 				GraphicsPSOInit.BlendState = TStaticBlendState<CW_RGBA, BO_Add, BF_One, BF_InverseSourceAlpha, BO_Add, BF_Zero, BF_One>::GetRHI(); // Premultiplied-alpha composition
 				GraphicsPSOInit.RasterizerState = TStaticRasterizerState<FM_Solid, CM_None, true>::GetRHI();
 				GraphicsPSOInit.PrimitiveType = bLines ? PT_LineList : PT_TriangleList;
