Multi-view-family:  fix crash where ray tracing resources are released mid-frame when rendering multiple views on the same GPU.  Need to call ClearRayTracingBindings for last renderer only.  We also need to use FRHIGPUMask::All() to clear bindings for all GPUs, to avoid a leak or stale state being present.

The bug was hidden because in my test case I was running two views with ray tracing on separate GPUs, and ReleaseRaytracingResources was using the view GPU mask, so it only released the resources for one of the GPUs at a time.  Other ray tracing resources used in the same context (RayTracingSubSurfaceProfileSRV, RayTracingSubSurfaceProfileTexture) were ref counted, and so they don't get released until the last view using them, and it worked fine.

#rnx

diff --git a/Engine/Source/Runtime/Renderer/Private/DeferredShadingRenderer.cpp b/Engine/Source/Runtime/Renderer/Private/DeferredShadingRenderer.cpp
--- a/Engine/Source/Runtime/Renderer/Private/DeferredShadingRenderer.cpp
+++ b/Engine/Source/Runtime/Renderer/Private/DeferredShadingRenderer.cpp
@@ -1767,13 +1767,20 @@ bool FDeferredShadingSceneRenderer::DispatchRayTracingWorldUpdates(FRDGBuilder&
 	return true;
 }
 
-static void ReleaseRaytracingResources(FRDGBuilder& GraphBuilder, TArrayView<FViewInfo> Views, FRayTracingScene &RayTracingScene)
+static void ReleaseRaytracingResources(FRDGBuilder& GraphBuilder, TArrayView<FViewInfo> Views, FRayTracingScene &RayTracingScene, bool bIsLastRenderer)
 {
-	AddPass(GraphBuilder, RDG_EVENT_NAME("ReleaseRayTracingResources"), [Views, &RayTracingScene](FRHICommandListImmediate& RHICmdList)
+	// Ray tracing is always set up for all GPUs (see WaitForRayTracingScene), so ClearRayTracingBindings must run on all GPUs as well
+	RDG_GPU_MASK_SCOPE(GraphBuilder, FRHIGPUMask::All());
+
+	AddPass(GraphBuilder, RDG_EVENT_NAME("ReleaseRayTracingResources"), [Views, &RayTracingScene, bIsLastRenderer](FRHICommandListImmediate& RHICmdList)
 	{
 		if (RayTracingScene.IsCreated())
 		{
-			RHICmdList.ClearRayTracingBindings(RayTracingScene.GetRHIRayTracingScene());
+			// Clear ray tracing bindings only on the last renderer, where multiple view families are rendered
+			if (bIsLastRenderer)
+			{
+				RHICmdList.ClearRayTracingBindings(RayTracingScene.GetRHIRayTracingScene());
+			}
 
 			// Track if we ended up rendering anything this frame.  After rendering all view families, we'll release the
 			// ray tracing scene resources if nothing used ray tracing.
@@ -1788,7 +1795,8 @@ static void ReleaseRaytracingResources(FRDGBuilder& GraphBuilder, TArrayView<FVi
 		{
 			FViewInfo& View = Views[ViewIndex];
 
-			// Release common lighting resources
+			// Release common lighting resources -- these are ref counted, so they won't be released until after the last view
+			// is finished using them (where multiple view families are rendered).
 			View.RayTracingSubSurfaceProfileSRV.SafeRelease();
 			View.RayTracingSubSurfaceProfileTexture.SafeRelease();
 		}
@@ -3493,7 +3501,7 @@ void FDeferredShadingSceneRenderer::Render(FRDGBuilder& GraphBuilder)
 	GEngine->GetPostRenderDelegateEx().Broadcast(GraphBuilder);
 
 #if RHI_RAYTRACING
-	ReleaseRaytracingResources(GraphBuilder, Views, Scene->RayTracingScene);
+	ReleaseRaytracingResources(GraphBuilder, Views, Scene->RayTracingScene, bIsLastSceneRenderer);
 #endif //  RHI_RAYTRACING
 
 #if WITH_MGPU
