Upgraded CSV stats to RDG variants to include RDG pass executions to reduce RDG_Execute CSV tag.

diff --git a/Engine/Source/Runtime/Renderer/Private/BasePassRendering.cpp b/Engine/Source/Runtime/Renderer/Private/BasePassRendering.cpp
--- a/Engine/Source/Runtime/Renderer/Private/BasePassRendering.cpp
+++ b/Engine/Source/Runtime/Renderer/Private/BasePassRendering.cpp
@@ -907,6 +907,7 @@ void FDeferredShadingSceneRenderer::RenderBasePass(
 	if (ViewFamily.ViewExtensions.Num() > 0)
 	{
 		SCOPE_CYCLE_COUNTER(STAT_FDeferredShadingSceneRenderer_ViewExtensionPostRenderBasePass);
+		RDG_CSV_STAT_EXCLUSIVE_SCOPE(GraphBuilder, ViewExtensions);
 		RDG_EVENT_SCOPE(GraphBuilder, "BasePass_ViewExtensions");
 		auto* PassParameters = GraphBuilder.AllocParameters<FPostBasePassViewExtensionParameters>();
 		PassParameters->RenderTargets = BasePassRenderTargets;
diff --git a/Engine/Source/Runtime/Renderer/Private/CustomDepthRendering.cpp b/Engine/Source/Runtime/Renderer/Private/CustomDepthRendering.cpp
--- a/Engine/Source/Runtime/Renderer/Private/CustomDepthRendering.cpp
+++ b/Engine/Source/Runtime/Renderer/Private/CustomDepthRendering.cpp
@@ -124,7 +124,7 @@ bool FSceneRenderer::RenderCustomDepthPass(FRDGBuilder& GraphBuilder, FCustomDep
 		return false;
 	}
 
-	CSV_SCOPED_TIMING_STAT_EXCLUSIVE(RenderCustomDepthPass);
+	RDG_CSV_STAT_EXCLUSIVE_SCOPE(GraphBuilder, RenderCustomDepthPass);
 	RDG_GPU_STAT_SCOPE(GraphBuilder, CustomDepth);
 
 	bool bCustomDepthRendered = false;
diff --git a/Engine/Source/Runtime/Renderer/Private/DeferredShadingRenderer.cpp b/Engine/Source/Runtime/Renderer/Private/DeferredShadingRenderer.cpp
--- a/Engine/Source/Runtime/Renderer/Private/DeferredShadingRenderer.cpp
+++ b/Engine/Source/Runtime/Renderer/Private/DeferredShadingRenderer.cpp
@@ -2125,6 +2125,7 @@ void FDeferredShadingSceneRenderer::Render(FRDGBuilder& GraphBuilder)
 	const bool bUseVirtualTexturing = UseVirtualTexturing(FeatureLevel);
 	if (bUseVirtualTexturing)
 	{
+		RDG_CSV_STAT_EXCLUSIVE_SCOPE(GraphBuilder, VirtualTextureUpdate);
 		RDG_GPU_STAT_SCOPE(GraphBuilder, VirtualTextureUpdate);
 		// AllocateResources needs to be called before RHIBeginScene
 		FVirtualTextureSystem::Get().AllocateResources(GraphBuilder, FeatureLevel);
@@ -2248,6 +2249,7 @@ void FDeferredShadingSceneRenderer::Render(FRDGBuilder& GraphBuilder)
 	}
 
 	{
+		RDG_CSV_STAT_EXCLUSIVE_SCOPE(GraphBuilder, UpdateGPUScene);
 		RDG_GPU_STAT_SCOPE(GraphBuilder, GPUSceneUpdate);
 
 		if (bIsFirstSceneRenderer)
@@ -2422,10 +2424,7 @@ void FDeferredShadingSceneRenderer::Render(FRDGBuilder& GraphBuilder)
 		}
 	}
 
-	{
-		RDG_RHI_GPU_STAT_SCOPE(GraphBuilder, GPUSceneUpdate);
-		PrepareDistanceFieldScene(GraphBuilder, false);
-	}
+	PrepareDistanceFieldScene(GraphBuilder, false);
 
 	const bool bShouldRenderVelocities = ShouldRenderVelocities();
 	const bool bBasePassCanOutputVelocity = FVelocityRendering::BasePassCanOutputVelocity(FeatureLevel);
@@ -2629,6 +2628,7 @@ void FDeferredShadingSceneRenderer::Render(FRDGBuilder& GraphBuilder)
 	// NOTE: The ordering of the lights is used to select sub-sets for different purposes, e.g., those that support clustered deferred.
 	FSortedLightSetSceneInfo& SortedLightSet = *GraphBuilder.AllocObject<FSortedLightSetSceneInfo>();
 	{
+		RDG_CSV_STAT_EXCLUSIVE_SCOPE(GraphBuilder, SortLights);
 		RDG_GPU_STAT_SCOPE(GraphBuilder, SortLights);
 		GatherLightsAndComputeLightGrid(GraphBuilder, bComputeLightGrid, SortedLightSet);
 	}
@@ -2746,7 +2746,7 @@ void FDeferredShadingSceneRenderer::Render(FRDGBuilder& GraphBuilder)
 	FDBufferTextures DBufferTextures = CreateDBufferTextures(GraphBuilder, SceneTextures.Config.Extent, ShaderPlatform);
 
 	{
-		CSV_SCOPED_TIMING_STAT_EXCLUSIVE(DeferredShadingSceneRenderer_DBuffer);
+		RDG_CSV_STAT_EXCLUSIVE_SCOPE(GraphBuilder, DeferredShadingSceneRenderer_DBuffer);
 		SCOPE_CYCLE_COUNTER(STAT_FDeferredShadingSceneRenderer_DBuffer);
 		CompositionLighting.ProcessBeforeBasePass(GraphBuilder, DBufferTextures);
 	}
@@ -3003,7 +3003,7 @@ void FDeferredShadingSceneRenderer::Render(FRDGBuilder& GraphBuilder)
 	// Pre-lighting composition lighting stage
 	// e.g. deferred decals, SSAO
 	{
-		CSV_SCOPED_TIMING_STAT_EXCLUSIVE(AfterBasePass);
+		RDG_CSV_STAT_EXCLUSIVE_SCOPE(GraphBuilder, AfterBasePass);
 		SCOPE_CYCLE_COUNTER(STAT_FDeferredShadingSceneRenderer_AfterBasePass);
 
 		if (!IsForwardShadingEnabled(ShaderPlatform))
diff --git a/Engine/Source/Runtime/Renderer/Private/DistanceFieldObjectManagement.cpp b/Engine/Source/Runtime/Renderer/Private/DistanceFieldObjectManagement.cpp
--- a/Engine/Source/Runtime/Renderer/Private/DistanceFieldObjectManagement.cpp
+++ b/Engine/Source/Runtime/Renderer/Private/DistanceFieldObjectManagement.cpp
@@ -24,6 +24,8 @@
 #include "HAL/LowLevelMemStats.h"
 #include "ProfilingDebugging/CpuProfilerTrace.h"
 
+DECLARE_GPU_STAT(DistanceFields);
+
 extern int32 GDistanceFieldOffsetDataStructure;
 
 float GMeshDistanceFieldsMaxObjectBoundingRadius = 100000;
@@ -1033,11 +1035,12 @@ void FSceneRenderer::UpdateGlobalHeightFieldObjectBuffers(FRDGBuilder& GraphBuil
 
 void FSceneRenderer::PrepareDistanceFieldScene(FRDGBuilder& GraphBuilder, bool bSplitDispatch)
 {
-	CSV_SCOPED_TIMING_STAT_EXCLUSIVE(PrepareDistanceFieldScene);
+	RDG_CSV_STAT_EXCLUSIVE_SCOPE(GraphBuilder, PrepareDistanceFieldScene);
 	TRACE_CPUPROFILER_EVENT_SCOPE(FSceneRenderer::PrepareDistanceFieldScene);
 	QUICK_SCOPE_CYCLE_COUNTER(STAT_PrepareDistanceFieldScene);
 	LLM_SCOPE_BYTAG(DistanceFields);
 	RDG_GPU_MASK_SCOPE(GraphBuilder, FRHIGPUMask::All());
+	RDG_RHI_GPU_STAT_SCOPE(GraphBuilder, DistanceFields);
 
 	const bool bShouldPrepareHeightFieldScene = ShouldPrepareHeightFieldScene();
 	const bool bShouldPrepareDistanceFieldScene = ShouldPrepareDistanceFieldScene();
diff --git a/Engine/Source/Runtime/Renderer/Private/GPUScene.cpp b/Engine/Source/Runtime/Renderer/Private/GPUScene.cpp
--- a/Engine/Source/Runtime/Renderer/Private/GPUScene.cpp
+++ b/Engine/Source/Runtime/Renderer/Private/GPUScene.cpp
@@ -675,7 +675,6 @@ void FGPUScene::UpdateInternal(FRDGBuilder& GraphBuilder, FScene& Scene, FRDGExt
 
 	{
 		SCOPED_NAMED_EVENT(STAT_UpdateGPUScene, FColor::Green);
-		CSV_SCOPED_TIMING_STAT_EXCLUSIVE(UpdateGPUScene);
 		QUICK_SCOPE_CYCLE_COUNTER(STAT_UpdateGPUScene);
 		SCOPE_CYCLE_COUNTER(STAT_UpdateGPUSceneTime);
 
@@ -1411,7 +1410,7 @@ void FGPUScene::UploadDynamicPrimitiveShaderDataForViewInternal(FRDGBuilder& Gra
 	const bool bNeedsUpload = Collector.UploadData != nullptr && NumPrimitiveDataUploads > 0 && !Collector.UploadData->bIsUploaded;
 	if (bNeedsUpload)
 	{
-		CSV_SCOPED_TIMING_STAT_EXCLUSIVE(UploadDynamicPrimitiveShaderData);
+		RDG_CSV_STAT_EXCLUSIVE_SCOPE(GraphBuilder, UploadDynamicPrimitiveShaderData);
 		QUICK_SCOPE_CYCLE_COUNTER(STAT_UploadDynamicPrimitiveShaderData);
 
 		Collector.UploadData->bIsUploaded = true;
diff --git a/Engine/Source/Runtime/Renderer/Private/LightGridInjection.cpp b/Engine/Source/Runtime/Renderer/Private/LightGridInjection.cpp
--- a/Engine/Source/Runtime/Renderer/Private/LightGridInjection.cpp
+++ b/Engine/Source/Runtime/Renderer/Private/LightGridInjection.cpp
@@ -858,7 +858,7 @@ void FDeferredShadingSceneRenderer::RenderForwardShadowProjections(
 
 	if (bScreenShadowMaskNeeded)
 	{
-		CSV_SCOPED_TIMING_STAT_EXCLUSIVE(RenderForwardShadingShadowProjections);
+		RDG_CSV_STAT_EXCLUSIVE_SCOPE(GraphBuilder, RenderForwardShadingShadowProjections);
 
 		FRDGTextureMSAA ForwardScreenSpaceShadowMask;
 		FRDGTextureMSAA ForwardScreenSpaceShadowMaskSubPixel;
diff --git a/Engine/Source/Runtime/Renderer/Private/MobileShadingRenderer.cpp b/Engine/Source/Runtime/Renderer/Private/MobileShadingRenderer.cpp
--- a/Engine/Source/Runtime/Renderer/Private/MobileShadingRenderer.cpp
+++ b/Engine/Source/Runtime/Renderer/Private/MobileShadingRenderer.cpp
@@ -394,7 +394,7 @@ void FMobileSceneRenderer::InitViews(FRDGBuilder& GraphBuilder, FSceneTexturesCo
 	SCOPED_DRAW_EVENT(RHICmdList, InitViews);
 
 	SCOPE_CYCLE_COUNTER(STAT_InitViewsTime);
-	CSV_SCOPED_TIMING_STAT_EXCLUSIVE(InitViews_Scene);
+	RDG_CSV_STAT_EXCLUSIVE_SCOPE(GraphBuilder, InitViews_Scene);
 
 	check(Scene);
 
@@ -612,6 +612,7 @@ void FMobileSceneRenderer::InitViews(FRDGBuilder& GraphBuilder, FSceneTexturesCo
 	}
 
 	{
+		RDG_CSV_STAT_EXCLUSIVE_SCOPE(GraphBuilder, UpdateGPUScene);
 		FRDGExternalAccessQueue ExternalAccessQueue;
 
 		Scene->GPUScene.Update(GraphBuilder, *Scene, ExternalAccessQueue);
@@ -831,6 +832,7 @@ void FMobileSceneRenderer::Render(FRDGBuilder& GraphBuilder)
 	const bool bEnableClusteredShading = bEnableClusteredLocalLights || bEnableClusteredReflections;
 	if (bEnableClusteredShading)
 	{
+		RDG_CSV_STAT_EXCLUSIVE_SCOPE(GraphBuilder, SortLights);
 		// Shadows are applied in clustered shading on mobile forward and separately on mobile deferred.
 		bool bShadowedLightsInClustered = bRequiresShadowProjections && !bDeferredShading;
 		GatherAndSortLights(SortedLightSet, bShadowedLightsInClustered);
@@ -893,7 +895,7 @@ void FMobileSceneRenderer::Render(FRDGBuilder& GraphBuilder)
 
 		if (bRequiresShadowProjections)
 		{
-			CSV_SCOPED_TIMING_STAT_EXCLUSIVE(RenderMobileShadowProjections);
+			RDG_CSV_STAT_EXCLUSIVE_SCOPE(GraphBuilder, RenderMobileShadowProjections);
 			RDG_GPU_STAT_SCOPE(GraphBuilder, ShadowProjection);
 			RenderMobileShadowProjections(GraphBuilder, SceneTextures.Depth.Resolve);
 		}
diff --git a/Engine/Source/Runtime/Renderer/Private/PostProcess/PostProcessing.cpp b/Engine/Source/Runtime/Renderer/Private/PostProcess/PostProcessing.cpp
--- a/Engine/Source/Runtime/Renderer/Private/PostProcess/PostProcessing.cpp
+++ b/Engine/Source/Runtime/Renderer/Private/PostProcess/PostProcessing.cpp
@@ -1744,7 +1744,7 @@ static bool IsGaussianActive(const FViewInfo& View)
 
 void AddMobilePostProcessingPasses(FRDGBuilder& GraphBuilder, FScene* Scene, const FViewInfo& View, const FMobilePostProcessingInputs& Inputs, FInstanceCullingManager& InstanceCullingManager)
 {
-	CSV_SCOPED_TIMING_STAT_EXCLUSIVE(RenderPostProcessing);
+	RDG_CSV_STAT_EXCLUSIVE_SCOPE(GraphBuilder, RenderPostProcessing);
 	QUICK_SCOPE_CYCLE_COUNTER(STAT_PostProcessing_Process);
 
 	check(IsInRenderingThread());
diff --git a/Engine/Source/Runtime/Renderer/Private/ReflectionEnvironmentRealTimeCapture.cpp b/Engine/Source/Runtime/Renderer/Private/ReflectionEnvironmentRealTimeCapture.cpp
--- a/Engine/Source/Runtime/Renderer/Private/ReflectionEnvironmentRealTimeCapture.cpp
+++ b/Engine/Source/Runtime/Renderer/Private/ReflectionEnvironmentRealTimeCapture.cpp
@@ -310,6 +310,7 @@ void FScene::AllocateAndCaptureFrameSkyEnvMap(
 
 	RDG_EVENT_SCOPE(GraphBuilder, "CaptureConvolveSkyEnvMap");
 	RDG_GPU_STAT_SCOPE(GraphBuilder, CaptureConvolveSkyEnvMap);
+	RDG_CSV_STAT_EXCLUSIVE_SCOPE(GraphBuilder, SkyAtmosphere);
 
 	const uint32 CubeWidth = SkyLight->CaptureCubeMapResolution;
 	const uint32 CubeMipCount = FMath::CeilLogTwo(CubeWidth) + 1;
diff --git a/Engine/Source/Runtime/Renderer/Private/SceneVisibility.cpp b/Engine/Source/Runtime/Renderer/Private/SceneVisibility.cpp
--- a/Engine/Source/Runtime/Renderer/Private/SceneVisibility.cpp
+++ b/Engine/Source/Runtime/Renderer/Private/SceneVisibility.cpp
@@ -4894,7 +4894,7 @@ void FDeferredShadingSceneRenderer::InitViews(FRDGBuilder& GraphBuilder, const F
 {
 	SCOPED_NAMED_EVENT(FDeferredShadingSceneRenderer_InitViews, FColor::Emerald);
 	SCOPE_CYCLE_COUNTER(STAT_InitViewsTime);
-	CSV_SCOPED_TIMING_STAT_EXCLUSIVE(InitViews_Scene);
+	RDG_CSV_STAT_EXCLUSIVE_SCOPE(GraphBuilder, InitViews_Scene);
 
 	PreVisibilityFrameSetup(GraphBuilder, SceneTexturesConfig);
 
diff --git a/Engine/Source/Runtime/Renderer/Private/ShadowDepthRendering.cpp b/Engine/Source/Runtime/Renderer/Private/ShadowDepthRendering.cpp
--- a/Engine/Source/Runtime/Renderer/Private/ShadowDepthRendering.cpp
+++ b/Engine/Source/Runtime/Renderer/Private/ShadowDepthRendering.cpp
@@ -1797,7 +1797,7 @@ void FSceneRenderer::RenderShadowDepthMaps(FRDGBuilder& GraphBuilder, FInstanceC
 {
 	ensureMsgf(!bShadowDepthRenderCompleted, TEXT("RenderShadowDepthMaps called twice in the same frame"));
 
-	CSV_SCOPED_TIMING_STAT_EXCLUSIVE(RenderShadows);
+	RDG_CSV_STAT_EXCLUSIVE_SCOPE(GraphBuilder, RenderShadows);
 
 	TRACE_CPUPROFILER_EVENT_SCOPE(FSceneRenderer::RenderShadowDepthMaps);
 	SCOPED_NAMED_EVENT(FSceneRenderer_RenderShadowDepthMaps, FColor::Emerald);
diff --git a/Engine/Source/Runtime/Renderer/Private/SkyAtmosphereRendering.cpp b/Engine/Source/Runtime/Renderer/Private/SkyAtmosphereRendering.cpp
--- a/Engine/Source/Runtime/Renderer/Private/SkyAtmosphereRendering.cpp
+++ b/Engine/Source/Runtime/Renderer/Private/SkyAtmosphereRendering.cpp
@@ -1315,6 +1315,7 @@ void FSceneRenderer::RenderSkyAtmosphereLookUpTables(FRDGBuilder& GraphBuilder)
 
 	RDG_EVENT_SCOPE(GraphBuilder, "SkyAtmosphereLUTs");
 	RDG_GPU_STAT_SCOPE(GraphBuilder, SkyAtmosphereLUTs);
+	RDG_CSV_STAT_EXCLUSIVE_SCOPE(GraphBuilder, SkyAtmosphere);
 
 	FSkyAtmosphereRenderSceneInfo& SkyInfo = *Scene->GetSkyAtmosphereSceneInfo();
 	const FSkyAtmosphereSceneProxy& SkyAtmosphereSceneProxy = SkyInfo.GetSkyAtmosphereSceneProxy();
diff --git a/Engine/Source/Runtime/Renderer/Private/VT/VirtualTextureSystem.cpp b/Engine/Source/Runtime/Renderer/Private/VT/VirtualTextureSystem.cpp
--- a/Engine/Source/Runtime/Renderer/Private/VT/VirtualTextureSystem.cpp
+++ b/Engine/Source/Runtime/Renderer/Private/VT/VirtualTextureSystem.cpp
@@ -1327,7 +1327,7 @@ void FVirtualTextureSystem::Update(FRDGBuilder& GraphBuilder, ERHIFeatureLevel::
 {
 	check(IsInRenderingThread());
 
-	CSV_SCOPED_TIMING_STAT_EXCLUSIVE(VirtualTextureSystem_Update);
+	RDG_CSV_STAT_EXCLUSIVE_SCOPE(GraphBuilder, VirtualTextureSystem_Update);
 	TRACE_CPUPROFILER_EVENT_SCOPE(FVirtualTextureSystem::Update);
 	SCOPE_CYCLE_COUNTER(STAT_VirtualTextureSystem_Update);
 	RDG_GPU_STAT_SCOPE(GraphBuilder, VirtualTexture);
