Path Tracer: Light function material support

Extend the recently added raytraced light function support to the path tracer. Unlike in Ray Tracing which must retain visual compatibility with the deferred renderer, the path tracing implementation supports colored light functions.

Also guard the logic for miss shaders in DeferredRendering to avoid issues on platforms that support ray tracing but not ray tracing shaders yet (Vulkan).

diff --git a/Engine/Shaders/Private/PathTracing/Light/PathTracingLightCommon.ush b/Engine/Shaders/Private/PathTracing/Light/PathTracingLightCommon.ush
--- a/Engine/Shaders/Private/PathTracing/Light/PathTracingLightCommon.ush
+++ b/Engine/Shaders/Private/PathTracing/Light/PathTracingLightCommon.ush
@@ -368,4 +368,9 @@ bool HasRectTexture(int LightId)
 float GetVolumetricScatteringIntensity(int LightId)
 {
 	return SceneLights[LightId].VolumetricScatteringIntensity;
-}
\ No newline at end of file
+}
+
+uint GetLightMissShaderIndex(int LightId)
+{
+	return SceneLights[LightId].MissShaderIndex;
+}
diff --git a/Engine/Shaders/Private/PathTracing/PathTracingCore.ush b/Engine/Shaders/Private/PathTracing/PathTracingCore.ush
--- a/Engine/Shaders/Private/PathTracing/PathTracingCore.ush
+++ b/Engine/Shaders/Private/PathTracing/PathTracingCore.ush
@@ -460,16 +460,28 @@ FPathTracingPayload TraceTransparentRay(inout FPathState PathState, bool bIsCame
 	return Payload;
 }
 
-float3 TraceTransparentVisibilityRay(RayDesc Ray, uint2 PixelCoord, float PathRoughness, inout RandomSequence RandSequence)
+float3 TraceTransparentVisibilityRay(RayDesc Ray, uint2 PixelCoord, float PathRoughness, uint MissShaderIndex, bool bCastShadows, inout RandomSequence RandSequence)
 {
+	if (!bCastShadows && MissShaderIndex == 0)
+	{
+		// no work to do
+		return 1.0;
+	}
+
 	const uint RayFlags = RAY_FLAG_ACCEPT_FIRST_HIT_AND_END_SEARCH | RAY_FLAG_SKIP_CLOSEST_HIT_SHADER;
 	const uint InstanceInclusionMask = RAY_TRACING_MASK_SHADOW | (length2(Ray.Origin) < Pow2(PATH_TRACING_SKIP_HAIR_DISTANCE) ? RAY_TRACING_MASK_THIN_SHADOW : 0);
 	const uint RayContributionToHitGroupIndex = RAY_TRACING_SHADER_SLOT_SHADOW;
 	const uint MultiplierForGeometryContributionToShaderIndex = RAY_TRACING_NUM_SHADER_SLOTS;
-	const uint MissShaderIndex = 0;
 
 	FPackedPathTracingPayload PackedPayload = InitPathTracingVisibilityPayload(PixelCoord, PathRoughness);
 
+	if (!bCastShadows)
+	{
+		// ray should not cast shadows, make it degenerate so we can still run the miss shader
+		Ray.TMin = POSITIVE_INFINITY;
+		Ray.TMax = POSITIVE_INFINITY;
+	}
+
 	// Trace the ray
 	TraceRay(
 		TLAS,
@@ -483,18 +495,20 @@ float3 TraceTransparentVisibilityRay(RayDesc Ray, uint2 PixelCoord, float PathRo
 
 	if (PackedPayload.IsMiss())
 	{
-		FVolumeIntersectionList VolumeIsectList = VolumeIntersect(Ray.Origin, Ray.Direction, Ray.TMin, Ray.TMax);
-		if (VolumeIsectList.HitBlocker())
-		{
-			// blockers are opaque -- done!
-			return 0.0;
-		}
 		float3 Throughput = PackedPayload.GetRayThroughput();
-		if (VolumeIsectList.HitVolume())
+		if (bCastShadows)
 		{
-			Throughput = VolumeGetTransmittance(Throughput, Ray.Origin, Ray.Direction, VolumeIsectList, RandSequence);
+			FVolumeIntersectionList VolumeIsectList = VolumeIntersect(Ray.Origin, Ray.Direction, Ray.TMin, Ray.TMax);
+			if (VolumeIsectList.HitBlocker())
+			{
+				// blockers are opaque -- done!
+				return 0.0;
+			}
+			if (VolumeIsectList.HitVolume())
+			{
+				Throughput = VolumeGetTransmittance(Throughput, Ray.Origin, Ray.Direction, VolumeIsectList, RandSequence);
+			}
 		}
-
 		// we didn't hit any geo - just return whatever throughput is left
 		float3 Tau = PackedPayload.GetTau() * AbsorptionScale;
 		return Throughput * exp(-max(Tau, 0.0));
@@ -914,21 +928,20 @@ bool PathTracingKernel(inout FPathState PathState, int Bounce)
 				FPathTracingPayload VolPayload = CreateMediumHitPayload(EquiAngularT, TranslatedWorldPos, Result);
 				float4 LightRandValue = RandomSequence_GenerateSample4D(PathState.RandSequence);
 
+				const bool bCastShadows = CastsVolumeShadow(LightSampleLightId);
+				const uint MissShaderIndex = GetLightMissShaderIndex(LightSampleLightId);
 				// compute direct light sampling?
 				if (MISMode != 0)
 				{
 					FLightSample LightSample = SampleLight(LightSampleLightId, LightRandValue.xy, TranslatedWorldPos, float3(0, 0, 0));
 					if (LightSample.Pdf > 0)
 					{
-						if (CastsVolumeShadow(LightSampleLightId))
-						{
-							RayDesc LightRay;
-							LightRay.Origin = TranslatedWorldPos;
-							LightRay.TMin = 0;
-							LightRay.Direction = LightSample.Direction;
-							LightRay.TMax = LightSample.Distance;
-							LightSample.RadianceOverPdf *= TraceTransparentVisibilityRay(LightRay, PathState.PixelIndex, 1.0, PathState.RandSequence);
-						}
+						RayDesc LightRay;
+						LightRay.Origin = TranslatedWorldPos;
+						LightRay.TMin = 0;
+						LightRay.Direction = LightSample.Direction;
+						LightRay.TMax = LightSample.Distance;
+						LightSample.RadianceOverPdf *= TraceTransparentVisibilityRay(LightRay, PathState.PixelIndex, 1.0, MissShaderIndex, bCastShadows, PathState.RandSequence);
 						// #dxr_todo: Is it cheaper to fire the ray first? Or eval the material first?
 						if (any(LightSample.RadianceOverPdf > 0))
 						{
@@ -965,10 +978,7 @@ bool PathTracingKernel(inout FPathState PathState, int Bounce)
 								LightContrib *= MISWeightRobust(MaterialSample.Pdf, LightResult.Pdf);
 							}
 							MaterialRay.TMax = LightResult.HitT;
-							if (CastsVolumeShadow(LightSampleLightId))
-							{
-								LightContrib *= TraceTransparentVisibilityRay(MaterialRay, PathState.PixelIndex, MaterialSample.Roughness, PathState.RandSequence);
-							}
+							LightContrib *= TraceTransparentVisibilityRay(MaterialRay, PathState.PixelIndex, MaterialSample.Roughness, MissShaderIndex, bCastShadows, PathState.RandSequence);
 							AccumulateRadiance(PathState.Radiance, LightContrib, bIsCameraRay);
 						}
 					}
@@ -1089,22 +1099,21 @@ bool PathTracingKernel(inout FPathState PathState, int Bounce)
 			LightSample.Pdf *= LightPickPdf;
 			if (LightSample.Pdf > 0)
 			{
-				if (CastsShadow(LightId))
-				{
-					// for transmissive materials, bias the position to the other side of the surface if the light is coming from behind
-					const float SignedPositionBias = bIsTransmissiveMaterial ? sign(dot(Payload.WorldGeoNormal, LightSample.Direction)) : 1.0;
-					const bool bEnableShadowTermFix = Payload.ShadingModelID != SHADINGMODELID_HAIR && Payload.ShadingModelID != SHADINGMODELID_MEDIUM;
-					RayDesc LightRay;
-					LightRay.Origin = TranslatedWorldPos;
-					LightRay.TMin = 0;
-					LightRay.Direction = LightSample.Direction;
-					LightRay.TMax = LightSample.Distance;
-					ApplyRayBias(LightRay, Payload.HitT, SignedPositionBias * Payload.WorldGeoNormal, bEnableShadowTermFix);
-
-					float AvgRoughness = ApproximateCaustics ? GetAverageRoughness(Payload) : 0.0;
-
-					LightSample.RadianceOverPdf *= TraceTransparentVisibilityRay(LightRay, PathState.PixelIndex, AvgRoughness, PathState.RandSequence);
-				}
+				// for transmissive materials, bias the position to the other side of the surface if the light is coming from behind
+				const float SignedPositionBias = bIsTransmissiveMaterial ? sign(dot(Payload.WorldGeoNormal, LightSample.Direction)) : 1.0;
+				const bool bEnableShadowTermFix = Payload.ShadingModelID != SHADINGMODELID_HAIR && Payload.ShadingModelID != SHADINGMODELID_MEDIUM;
+				RayDesc LightRay;
+				LightRay.Origin = TranslatedWorldPos;
+				LightRay.TMin = 0;
+				LightRay.Direction = LightSample.Direction;
+				LightRay.TMax = LightSample.Distance;
+				ApplyRayBias(LightRay, Payload.HitT, SignedPositionBias * Payload.WorldGeoNormal, bEnableShadowTermFix);
+
+				float AvgRoughness = ApproximateCaustics ? GetAverageRoughness(Payload) : 0.0;
+
+				const bool bCastShadows = CastsShadow(LightId);
+				const uint MissShaderIndex = GetLightMissShaderIndex(LightId);
+				LightSample.RadianceOverPdf *= TraceTransparentVisibilityRay(LightRay, PathState.PixelIndex, AvgRoughness, MissShaderIndex, bCastShadows, PathState.RandSequence);
 
 				if (bIsVolumeSample)
 				{
@@ -1267,12 +1276,11 @@ bool PathTracingKernel(inout FPathState PathState, int Bounce)
 
 			if (any(LightContrib > 0))
 			{
-				if (CastsShadow(LightId))
-				{
-					RayDesc LightRay = PathState.Ray;
-					LightRay.TMax = LightResult.HitT;
-					LightContrib *= TraceTransparentVisibilityRay(LightRay, PathState.PixelIndex, ShadowRayRoughness, PathState.RandSequence);
-				}
+				RayDesc LightRay = PathState.Ray;
+				LightRay.TMax = LightResult.HitT;
+				const bool bCastShadows = CastsShadow(LightId);
+				const uint MissShaderIndex = GetLightMissShaderIndex(LightId);
+				LightContrib *= TraceTransparentVisibilityRay(LightRay, PathState.PixelIndex, ShadowRayRoughness, MissShaderIndex, bCastShadows, PathState.RandSequence);
 				// the light made some contribution, and there was nothing along the shadow ray
 				AccumulateRadiance(PathState.Radiance, LightContrib, bIsCameraRay);
 			}
diff --git a/Engine/Shaders/Private/PathTracing/PathTracingLightingMissShader.usf b/Engine/Shaders/Private/PathTracing/PathTracingLightingMissShader.usf
new file mode 100644
--- /dev/null
+++ b/Engine/Shaders/Private/PathTracing/PathTracingLightingMissShader.usf
@@ -0,0 +1,21 @@
+// Copyright Epic Games, Inc. All Rights Reserved.
+
+#include "PathTracingCommon.ush"
+
+#include "/Engine/Generated/Material.ush"
+
+// TODO: if we could just bind loose parameters, we could avoid this workaround for getting the shader parameters needed by the routines in LightFunctionCommon.ush
+#define LightFunctionParameters  PathTracingLightFunctionParameters.LightFunctionParameters
+#define LightFunctionParameters2 PathTracingLightFunctionParameters.LightFunctionParameters2
+#define LightFunctionTranslatedWorldToLight PathTracingLightFunctionParameters.LightFunctionTranslatedWorldToLight
+#include "../LightFunctionCommon.ush"
+#include "../RayTracing/RayTracingLightFunctionCommon.ush"
+
+RAY_TRACING_ENTRY_MISS(PathTracingLightingMS, FPackedPathTracingPayload, Payload)
+{
+	Payload.SetMiss();
+	float3 TranslatedWorldPosition = WorldRayOrigin(); // Shadow ray is traced from the shaded point
+	float3 Throughput = Payload.GetRayThroughput();
+	Throughput *= GetRayTracingLightFunction(TranslatedWorldPosition);
+	Payload.SetRayThroughput(Throughput);
+}
diff --git a/Engine/Shaders/Private/PathTracing/PathTracingMissShader.usf b/Engine/Shaders/Private/PathTracing/PathTracingMissShader.usf
new file mode 100644
--- /dev/null
+++ b/Engine/Shaders/Private/PathTracing/PathTracingMissShader.usf
@@ -0,0 +1,8 @@
+// Copyright Epic Games, Inc. All Rights Reserved.
+
+#include "PathTracingCommon.ush"
+
+RAY_TRACING_ENTRY_MISS(PathTracingDefaultMS, FPackedPathTracingPayload, Payload)
+{
+	Payload.SetMiss();
+}
diff --git a/Engine/Shaders/Private/RayTracing/RayTracingLightFunctionCommon.ush b/Engine/Shaders/Private/RayTracing/RayTracingLightFunctionCommon.ush
--- a/Engine/Shaders/Private/RayTracing/RayTracingLightFunctionCommon.ush
+++ b/Engine/Shaders/Private/RayTracing/RayTracingLightFunctionCommon.ush
@@ -18,12 +18,12 @@ float3 GetRayTracingLightFunction(float3 TranslatedWorldPosition)
 
 	float3 Color = GetLightFunctionColor(LightVector, TranslatedWorldPosition);
 
-	float DistanceFadeAlpha = saturate((RaytracingLightFunctionParameters.LightFunctionParameters2.x - ViewDistance) / (RaytracingLightFunctionParameters.LightFunctionParameters2.x * .2f));
+	float DistanceFadeAlpha = saturate((LightFunctionParameters2.x - ViewDistance) / (LightFunctionParameters2.x * .2f));
 	// Fade to disabled based on LightFunctionFadeDistance
-	Color = lerp(RaytracingLightFunctionParameters.LightFunctionParameters2.y, Color, DistanceFadeAlpha);
+	Color = lerp(LightFunctionParameters2.y, Color, DistanceFadeAlpha);
 
 	// Fade to disabled based on ShadowFadeFraction
-	Color = lerp(RaytracingLightFunctionParameters.LightFunctionParameters2.y, Color, LightFunctionParameters.y);
+	Color = lerp(LightFunctionParameters2.y, Color, LightFunctionParameters.y);
 
 	return Color;
 }
diff --git a/Engine/Shaders/Private/RayTracing/RayTracingLightingMS.usf b/Engine/Shaders/Private/RayTracing/RayTracingLightingMS.usf
--- a/Engine/Shaders/Private/RayTracing/RayTracingLightingMS.usf
+++ b/Engine/Shaders/Private/RayTracing/RayTracingLightingMS.usf
@@ -38,7 +38,8 @@
 #include "/Engine/Generated/Material.ush"
 
 // TODO: if we could just bind loose parameters, we could avoid this workaround for getting the shader parameters needed by the routines in LightFunctionCommon.ush
-#define LightFunctionParameters RaytracingLightFunctionParameters.LightFunctionParameters
+#define LightFunctionParameters  RaytracingLightFunctionParameters.LightFunctionParameters
+#define LightFunctionParameters2 RaytracingLightFunctionParameters.LightFunctionParameters2
 #define LightFunctionTranslatedWorldToLight RaytracingLightFunctionParameters.LightFunctionTranslatedWorldToLight
 
 #include "../LightFunctionCommon.ush"
diff --git a/Engine/Shaders/Shared/RayTracingTypes.h b/Engine/Shaders/Shared/RayTracingTypes.h
--- a/Engine/Shaders/Shared/RayTracingTypes.h
+++ b/Engine/Shaders/Shared/RayTracingTypes.h
@@ -26,6 +26,7 @@
 	float   VolumetricScatteringIntensity;  // scale for volume contributions
 	int32   IESTextureSlice;
 	uint32  Flags; // see defines PATHTRACER_FLAG_*
+	uint32  MissShaderIndex;  // used to implement light functions
 	FVector3f TranslatedBoundMin;
 	FVector3f TranslatedBoundMax;
 	uint32 RectLightAtlasUVScale;  // Rect. light atlas UV transformation, encoded as f16x2
@@ -33,7 +34,7 @@
 	// keep structure aligned
 };
 
-static_assert(sizeof(FPathTracingLight) == 128, "Path tracing light structure should be aligned to 128 bytes for optimal access on the GPU");
+static_assert(sizeof(FPathTracingLight) == 132, "Path tracing light structure should be kept as small as possible");
 
 struct FPathTracingPackedPathState {
 	uint32    PixelIndex;
@@ -83,6 +84,7 @@
 	float   VolumetricScatteringIntensity;  // scale for volume contributions
 	int     IESTextureSlice;
 	uint    Flags;
+	uint    MissShaderIndex;
 	float3  TranslatedBoundMin;
 	float3  TranslatedBoundMax;
 	uint	RectLightAtlasUVScale;
diff --git a/Engine/Source/Runtime/Renderer/Private/DeferredShadingRenderer.cpp b/Engine/Source/Runtime/Renderer/Private/DeferredShadingRenderer.cpp
--- a/Engine/Source/Runtime/Renderer/Private/DeferredShadingRenderer.cpp
+++ b/Engine/Source/Runtime/Renderer/Private/DeferredShadingRenderer.cpp
@@ -1617,7 +1617,7 @@ bool FDeferredShadingSceneRenderer::SetupRayTracingPipelineStates(FRDGBuilder& G
 	return true;
 }
 
-void FDeferredShadingSceneRenderer::SetupRayTracingLightDataForViews(FRDGBuilder& GraphBuilder, const FRayTracingLightFunctionMap& RayTracingLightFunctionMap)
+void FDeferredShadingSceneRenderer::SetupRayTracingLightDataForViews(FRDGBuilder& GraphBuilder)
 {
 	if (!bAnyRayTracingPassEnabled)
 	{
@@ -1640,7 +1640,7 @@ void FDeferredShadingSceneRenderer::SetupRayTracingLightDataForViews(FRDGBuilder
 		else
 		{
 			// This light data is a function of the camera position, so must be computed per view.
-			View.RayTracingLightDataUniformBuffer = CreateRayTracingLightData(GraphBuilder, Scene, View, RayTracingLightFunctionMap, View.ShaderMap, NumOfSkippedRayTracingLights);
+			View.RayTracingLightDataUniformBuffer = CreateRayTracingLightData(GraphBuilder, Scene, View, View.ShaderMap, NumOfSkippedRayTracingLights);
 		}
 	}
 
@@ -1803,7 +1803,7 @@ static void ReleaseRaytracingResources(FRDGBuilder& GraphBuilder, TArrayView<FVi
 	});
 }
 
-void FDeferredShadingSceneRenderer::WaitForRayTracingScene(FRDGBuilder& GraphBuilder, FRDGBufferRef DynamicGeometryScratchBuffer, FRayTracingLightFunctionMap& RayTracingLightFunctionMap)
+void FDeferredShadingSceneRenderer::WaitForRayTracingScene(FRDGBuilder& GraphBuilder, FRDGBufferRef DynamicGeometryScratchBuffer)
 {
 	if (!bAnyRayTracingPassEnabled)
 	{
@@ -1838,8 +1838,9 @@ void FDeferredShadingSceneRenderer::WaitForRayTracingScene(FRDGBuilder& GraphBui
 	PassParams->DynamicGeometryScratchBuffer = DynamicGeometryScratchBuffer;
 	PassParams->LightDataPacked = bIsPathTracing ? nullptr : ReferenceView.RayTracingLightDataUniformBuffer; // accessed by FRayTracingLightingMS
 
+	const FRayTracingLightFunctionMap* RayTracingLightFunctionMap = GraphBuilder.Blackboard.Get<FRayTracingLightFunctionMap>();
 	GraphBuilder.AddPass(RDG_EVENT_NAME("WaitForRayTracingScene"), PassParams, ERDGPassFlags::Compute | ERDGPassFlags::NeverCull,
-		[this, PassParams, bIsPathTracing, &ReferenceView, bAnyLumenHardwareInlineRayTracingPassEnabled, RayTracingLightFunctionMap = MoveTemp(RayTracingLightFunctionMap)](FRHICommandListImmediate& RHICmdList)
+		[this, PassParams, bIsPathTracing, &ReferenceView, bAnyLumenHardwareInlineRayTracingPassEnabled, RayTracingLightFunctionMap](FRHICommandListImmediate& RHICmdList)
 	{
 		check(ReferenceView.RayTracingMaterialPipeline || ReferenceView.RayTracingMaterialBindings.Num() == 0);
 
@@ -1847,10 +1848,15 @@ void FDeferredShadingSceneRenderer::WaitForRayTracingScene(FRDGBuilder& GraphBui
 		{
 			BindRayTracingMaterialPipeline(RHICmdList, ReferenceView, ReferenceView.RayTracingMaterialPipeline);
 
-			SetupRayTracingDefaultMissShader(RHICmdList, ReferenceView);
+			if (bIsPathTracing)
+			{
+				SetupPathTracingDefaultMissShader(RHICmdList, ReferenceView);
 
-			if (!bIsPathTracing)
+				BindLightFunctionShadersPathTracing(RHICmdList, Scene, RayTracingLightFunctionMap, ReferenceView);
+			}
+			else
 			{
+				SetupRayTracingDefaultMissShader(RHICmdList, ReferenceView);
 				SetupRayTracingLightingMissShader(RHICmdList, ReferenceView);
 				
 				BindLightFunctionShaders(RHICmdList, Scene, RayTracingLightFunctionMap, ReferenceView);
@@ -2193,7 +2199,6 @@ void FDeferredShadingSceneRenderer::Render(FRDGBuilder& GraphBuilder)
 
 	// Prepare the scene for rendering this frame.
 	FRayTracingScene& RayTracingScene = Scene->RayTracingScene;
-	FRayTracingLightFunctionMap RayTracingLightFunctionMap;
 	RayTracingScene.Reset(); // Resets the internal arrays, but does not release any resources.
 
 	const int32 ReferenceViewIndex = 0;
@@ -2210,16 +2215,35 @@ void FDeferredShadingSceneRenderer::Render(FRDGBuilder& GraphBuilder)
 		ReferenceView.RayTracingDecalUniformBuffer = CreateNullRayTracingDecalsUniformBuffer(GraphBuilder);
 		ReferenceView.bHasRayTracingDecals = false;
 	}
-	
-	// TODO: currently Path tracing only uses the default miss shader, but future extensions should be handled here
-	if (IsRayTracingEnabled() && !ViewFamily.EngineShowFlags.PathTracing)
+
+	if (IsRayTracingEnabled() && RHISupportsRayTracingShaders(ViewFamily.GetShaderPlatform()))
 	{
-		// get the default lighting miss shader (to implicitly fill in the MissShader library before the RT pipeline is created)
-		GetRayTracingLightingMissShader(ReferenceView);
-		RayTracingScene.NumMissShaderSlots++;
+		if (!ViewFamily.EngineShowFlags.PathTracing)
+		{
+			// get the default lighting miss shader (to implicitly fill in the MissShader library before the RT pipeline is created)
+			GetRayTracingLightingMissShader(ReferenceView);
+			RayTracingScene.NumMissShaderSlots++;
+		}
 
-		// gather all the light functions that may be used (and also count how many miss shaders we will need)
-		RayTracingLightFunctionMap = GatherLightFunctionLights(Scene, ViewFamily.EngineShowFlags, ReferenceView.GetFeatureLevel());
+		if (ViewFamily.EngineShowFlags.LightFunctions)
+		{
+			// gather all the light functions that may be used (and also count how many miss shaders we will need)
+			FRayTracingLightFunctionMap RayTracingLightFunctionMap;
+			if (ViewFamily.EngineShowFlags.PathTracing)
+			{
+				RayTracingLightFunctionMap = GatherLightFunctionLightsPathTracing(Scene, ViewFamily.EngineShowFlags, ReferenceView.GetFeatureLevel());
+			}
+			else
+			{
+				RayTracingLightFunctionMap = GatherLightFunctionLights(Scene, ViewFamily.EngineShowFlags, ReferenceView.GetFeatureLevel());
+			}
+			if (!RayTracingLightFunctionMap.IsEmpty())
+			{
+				// If we got some light functions in our map, store them in the RDG blackboard so downstream functions can use them.
+				// The map itself will be strictly read-only from this point on.
+				GraphBuilder.Blackboard.Create<FRayTracingLightFunctionMap>(MoveTemp(RayTracingLightFunctionMap));
+			}
+		}
 	}
 
 	// Asynchronously create a list of primitives relevant to ray tracing scene
@@ -2912,7 +2936,7 @@ void FDeferredShadingSceneRenderer::Render(FRDGBuilder& GraphBuilder)
 
 #if RHI_RAYTRACING
 	/** Should be called somewhere before "WaitForRayTracingScene" */
-	SetupRayTracingLightDataForViews(GraphBuilder, RayTracingLightFunctionMap);
+	SetupRayTracingLightDataForViews(GraphBuilder);
 #endif // RHI_RAYTRACING
 
 	// Shadows, lumen and fog after base pass
@@ -2938,7 +2962,7 @@ void FDeferredShadingSceneRenderer::Render(FRDGBuilder& GraphBuilder)
 		// Lumen scene lighting requires ray tracing scene to be ready if HWRT shadows are desired
 		if (!bRayTracingSceneReady && Lumen::UseHardwareRayTracedSceneLighting(ViewFamily))
 		{
-			WaitForRayTracingScene(GraphBuilder, DynamicGeometryScratchBuffer, RayTracingLightFunctionMap);
+			WaitForRayTracingScene(GraphBuilder, DynamicGeometryScratchBuffer);
 			bRayTracingSceneReady = true;
 		}
 #endif // RHI_RAYTRACING
@@ -3028,7 +3052,7 @@ void FDeferredShadingSceneRenderer::Render(FRDGBuilder& GraphBuilder)
 	// If Lumen did not force an earlier ray tracing scene sync, we must wait for it here.
 	if (!bRayTracingSceneReady)
 	{
-		WaitForRayTracingScene(GraphBuilder, DynamicGeometryScratchBuffer, RayTracingLightFunctionMap);
+		WaitForRayTracingScene(GraphBuilder, DynamicGeometryScratchBuffer);
 		bRayTracingSceneReady = true;
 	}
 #endif // RHI_RAYTRACING
diff --git a/Engine/Source/Runtime/Renderer/Private/DeferredShadingRenderer.h b/Engine/Source/Runtime/Renderer/Private/DeferredShadingRenderer.h
--- a/Engine/Source/Runtime/Renderer/Private/DeferredShadingRenderer.h
+++ b/Engine/Source/Runtime/Renderer/Private/DeferredShadingRenderer.h
@@ -995,6 +995,7 @@ private:
 
 	/** Setup the default miss shader (required for any raytracing pipeline) */
 	void SetupRayTracingDefaultMissShader(FRHICommandListImmediate& RHICmdList, const FViewInfo& View);
+	void SetupPathTracingDefaultMissShader(FRHICommandListImmediate& RHICmdList, const FViewInfo& View);
 
 	/** Lighting Evaluation shader setup (used by ray traced reflections and translucency) */
 	void SetupRayTracingLightingMissShader(FRHICommandListImmediate& RHICmdList, const FViewInfo& View);
@@ -1009,7 +1010,7 @@ private:
 	void ComputePathCompaction(FRHICommandListImmediate& RHICmdList, const FViewInfo& View, FRHITexture* RadianceTexture, FRHITexture* SampleCountTexture, FRHITexture* PixelPositionTexture,
 		FRHIUnorderedAccessView* RadianceSortedRedUAV, FRHIUnorderedAccessView* RadianceSortedGreenUAV, FRHIUnorderedAccessView* RadianceSortedBlueUAV, FRHIUnorderedAccessView* RadianceSortedAlphaUAV, FRHIUnorderedAccessView* SampleCountSortedUAV);
 
-	void WaitForRayTracingScene(FRDGBuilder& GraphBuilder, FRDGBufferRef DynamicGeometryScratchBuffer, FRayTracingLightFunctionMap& RayTracingLightFunctionMap);
+	void WaitForRayTracingScene(FRDGBuilder& GraphBuilder, FRDGBufferRef DynamicGeometryScratchBuffer);
 
 	/** Debug ray tracing functions. */
 	void RenderRayTracingDebug(FRDGBuilder& GraphBuilder, const FViewInfo& View, FRDGTextureRef SceneColorOutputTexture);
@@ -1019,7 +1020,7 @@ private:
 	bool GatherRayTracingWorldInstancesForView(FRDGBuilder& GraphBuilder, FViewInfo& View, FRayTracingScene& RayTracingScene, struct FRayTracingRelevantPrimitiveList&& RelevantPrimitiveList);
 
 	bool SetupRayTracingPipelineStates(FRDGBuilder& GraphBuilder);
-	void SetupRayTracingLightDataForViews(FRDGBuilder& GraphBuilder, const FRayTracingLightFunctionMap& RayTracingLightFunctionMap);
+	void SetupRayTracingLightDataForViews(FRDGBuilder& GraphBuilder);
 	bool DispatchRayTracingWorldUpdates(FRDGBuilder& GraphBuilder, FRDGBufferRef& OutDynamicGeometryScratchBuffer);
 
 	/** Functions to create ray tracing pipeline state objects for various effects */
@@ -1073,6 +1074,7 @@ private:
 	/** Lighting evaluation shader registration */
 	static FRHIRayTracingShader* GetRayTracingDefaultMissShader(const FViewInfo& View);
 	static FRHIRayTracingShader* GetRayTracingLightingMissShader(const FViewInfo& View);
+	static FRHIRayTracingShader* GetPathTracingDefaultMissShader(const FViewInfo& View);
 
 	const FRHITransition* RayTracingDynamicGeometryUpdateEndTransition = nullptr; // Signaled when all AS for this frame are built
 #endif // RHI_RAYTRACING
diff --git a/Engine/Source/Runtime/Renderer/Private/PathTracing.cpp b/Engine/Source/Runtime/Renderer/Private/PathTracing.cpp
--- a/Engine/Source/Runtime/Renderer/Private/PathTracing.cpp
+++ b/Engine/Source/Runtime/Renderer/Private/PathTracing.cpp
@@ -824,7 +824,225 @@ class FPathTracingBuildAtmosphereOpticalDepthLUTCS : public FGlobalShader
 };
 IMPLEMENT_SHADER_TYPE(, FPathTracingBuildAtmosphereOpticalDepthLUTCS, TEXT("/Engine/Private/PathTracing/PathTracingBuildAtmosphereLUT.usf"), TEXT("PathTracingBuildAtmosphereOpticalDepthLUTCS"), SF_Compute);
 
+// Default miss shader (using the path tracing payload)
+class FPathTracingDefaultMS : public FGlobalShader
+{
+	DECLARE_SHADER_TYPE(FPathTracingDefaultMS, Global, );
+public:
+
+	FPathTracingDefaultMS() = default;
+	FPathTracingDefaultMS(const ShaderMetaType::CompiledShaderInitializerType& Initializer)
+		: FGlobalShader(Initializer)
+	{}
 
+	static bool ShouldCompilePermutation(const FGlobalShaderPermutationParameters& Parameters)
+	{
+		return ShouldCompilePathTracingShadersForProject(Parameters.Platform);
+	}
+
+	static void ModifyCompilationEnvironment(const FGlobalShaderPermutationParameters& Parameters, FShaderCompilerEnvironment& OutEnvironment)
+	{
+	}
+};
+IMPLEMENT_SHADER_TYPE(, FPathTracingDefaultMS, TEXT("/Engine/Private/PathTracing/PathTracingMissShader.usf"), TEXT("PathTracingDefaultMS"), SF_RayMiss);
+
+FRHIRayTracingShader* FDeferredShadingSceneRenderer::GetPathTracingDefaultMissShader(const FViewInfo& View)
+{
+	return View.ShaderMap->GetShader<FPathTracingDefaultMS>().GetRayTracingShader();
+}
+
+void FDeferredShadingSceneRenderer::SetupPathTracingDefaultMissShader(FRHICommandListImmediate& RHICmdList, const FViewInfo& View)
+{
+	int32 MissShaderPipelineIndex = FindRayTracingMissShaderIndex(View.RayTracingMaterialPipeline, GetPathTracingDefaultMissShader(View), true);
+
+	RHICmdList.SetRayTracingMissShader(View.GetRayTracingSceneChecked(),
+		RAY_TRACING_MISS_SHADER_SLOT_DEFAULT,
+		View.RayTracingMaterialPipeline,
+		MissShaderPipelineIndex,
+		0, nullptr, 0);
+}
+
+
+BEGIN_GLOBAL_SHADER_PARAMETER_STRUCT(FLightFunctionParametersPathTracing, )
+	SHADER_PARAMETER(FMatrix44f, LightFunctionTranslatedWorldToLight)
+	SHADER_PARAMETER(FVector4f, LightFunctionParameters)
+	SHADER_PARAMETER(FVector3f, LightFunctionParameters2)
+END_GLOBAL_SHADER_PARAMETER_STRUCT()
+
+IMPLEMENT_GLOBAL_SHADER_PARAMETER_STRUCT(FLightFunctionParametersPathTracing, "PathTracingLightFunctionParameters");
+
+static TUniformBufferRef<FLightFunctionParametersPathTracing> CreateLightFunctionParametersBuffer(
+	const FLightSceneInfo* LightSceneInfo,
+	const FSceneView& View,
+	EUniformBufferUsage Usage)
+{
+	FLightFunctionParametersPathTracing LightFunctionParameters;
+
+	const FVector Scale = LightSceneInfo->Proxy->GetLightFunctionScale();
+	// Switch x and z so that z of the user specified scale affects the distance along the light direction
+	const FVector InverseScale = FVector(1.f / Scale.Z, 1.f / Scale.Y, 1.f / Scale.X);
+	const FMatrix WorldToLight = LightSceneInfo->Proxy->GetWorldToLight() * FScaleMatrix(FVector(InverseScale));
+
+	LightFunctionParameters.LightFunctionTranslatedWorldToLight = FMatrix44f(FTranslationMatrix(-View.ViewMatrices.GetPreViewTranslation()) * WorldToLight);
+
+	const bool bIsSpotLight = LightSceneInfo->Proxy->GetLightType() == LightType_Spot;
+	const bool bIsPointLight = LightSceneInfo->Proxy->GetLightType() == LightType_Point;
+	const float TanOuterAngle = bIsSpotLight ? FMath::Tan(LightSceneInfo->Proxy->GetOuterConeAngle()) : 1.0f;
+
+	const float ShadowFadeFraction = 1.0f;
+
+	LightFunctionParameters.LightFunctionParameters = FVector4f(TanOuterAngle, ShadowFadeFraction, bIsSpotLight ? 1.0f : 0.0f, bIsPointLight ? 1.0f : 0.0f);
+
+	const bool bRenderingPreviewShadowIndicator = false;
+
+	LightFunctionParameters.LightFunctionParameters2 = FVector3f(
+		LightSceneInfo->Proxy->GetLightFunctionFadeDistance(),
+		LightSceneInfo->Proxy->GetLightFunctionDisabledBrightness(),
+		bRenderingPreviewShadowIndicator ? 1.0f : 0.0f);
+
+	return CreateUniformBufferImmediate(LightFunctionParameters, Usage);
+}
+
+// Miss Shader implementing light functions
+class FPathTracingLightingMS : public FMaterialShader
+{
+	DECLARE_SHADER_TYPE(FPathTracingLightingMS, Material);
+	LAYOUT_FIELD(FShaderUniformBufferParameter, LightMaterialsParameter);
+
+public:
+	static bool ShouldCompilePermutation(const FMaterialShaderPermutationParameters& Parameters)
+	{
+		return Parameters.MaterialParameters.MaterialDomain == MD_LightFunction && ShouldCompilePathTracingShadersForProject(Parameters.Platform);
+	}
+
+	FPathTracingLightingMS() {}
+	FPathTracingLightingMS(const ShaderMetaType::CompiledShaderInitializerType& Initializer)
+		: FMaterialShader(Initializer)
+	{
+		LightMaterialsParameter.Bind(Initializer.ParameterMap, TEXT("PathTracingLightFunctionParameters"));
+	}
+
+	void GetShaderBindings(
+		const FScene* Scene,
+		ERHIFeatureLevel::Type FeatureLevel,
+		const FMaterialRenderProxy& MaterialRenderProxy,
+		const FMaterial& Material,
+		const FViewInfo& View,
+		const TUniformBufferRef<FLightFunctionParametersPathTracing>& LightFunctionParameters,
+		FMeshDrawSingleShaderBindings& ShaderBindings) const
+	{
+		FMaterialShader::GetShaderBindings(Scene, FeatureLevel, MaterialRenderProxy, Material, ShaderBindings);
+		ShaderBindings.Add(GetUniformBufferParameter<FViewUniformShaderParameters>(), View.ViewUniformBuffer);
+		ShaderBindings.Add(LightMaterialsParameter, LightFunctionParameters);
+	}
+
+	static void ModifyCompilationEnvironment(const FMaterialShaderPermutationParameters& Parameters, FShaderCompilerEnvironment& OutEnvironment)
+	{
+		OutEnvironment.SetDefine(TEXT("SUPPORT_LIGHT_FUNCTION"), 1);
+		FMaterialShader::ModifyCompilationEnvironment(Parameters, OutEnvironment);
+	}
+};
+
+IMPLEMENT_MATERIAL_SHADER_TYPE(, FPathTracingLightingMS, TEXT("/Engine/Private/PathTracing/PathTracingLightingMissShader.usf"), TEXT("PathTracingLightingMS"), SF_RayMiss);
+
+
+static void BindLightFunction(
+	FRHICommandList& RHICmdList,
+	const FScene* Scene,
+	const FViewInfo& View,
+	const FMaterial& Material,
+	const FMaterialRenderProxy& MaterialRenderProxy,
+	const TUniformBufferRef<FLightFunctionParametersPathTracing>& LightFunctionParameters,
+	int32 Index
+)
+{
+	FRHIRayTracingScene* RTScene = View.GetRayTracingSceneChecked();
+	FRayTracingPipelineState* Pipeline = View.RayTracingMaterialPipeline;
+	const FMaterialShaderMap* MaterialShaderMap = Material.GetRenderingThreadShaderMap();
+
+
+	auto Shader = MaterialShaderMap->GetShader<FPathTracingLightingMS>();
+
+	TMeshProcessorShaders<
+		FMaterialShader,
+		FMaterialShader,
+		FMaterialShader,
+		FPathTracingLightingMS,
+		FMaterialShader> RayTracingShaders;
+
+	RayTracingShaders.RayTracingShader = Shader;
+
+	FMeshDrawShaderBindings ShaderBindings;
+	ShaderBindings.Initialize(RayTracingShaders.GetUntypedShaders());
+
+	int32 DataOffset = 0;
+	FMeshDrawSingleShaderBindings SingleShaderBindings = ShaderBindings.GetSingleShaderBindings(SF_RayMiss, DataOffset);
+
+	RayTracingShaders.RayTracingShader->GetShaderBindings(Scene, Scene->GetFeatureLevel(), MaterialRenderProxy, Material, View, LightFunctionParameters, SingleShaderBindings);
+
+	int32 MissShaderPipelineIndex = FindRayTracingMissShaderIndex(View.RayTracingMaterialPipeline, Shader.GetRayTracingShader(), true);
+
+	ShaderBindings.SetRayTracingShaderBindingsForMissShader(RHICmdList, RTScene, Pipeline, MissShaderPipelineIndex, Index);
+}
+
+void BindLightFunctionShadersPathTracing(
+	FRHICommandList& RHICmdList,
+	const FScene* Scene,
+	const FRayTracingLightFunctionMap* RayTracingLightFunctionMap,
+	const class FViewInfo& View)
+{
+	if (RayTracingLightFunctionMap == nullptr)
+	{
+		return;
+	}
+	for (const FRayTracingLightFunctionMap::ElementType& LightAndIndex : *RayTracingLightFunctionMap)
+	{
+		const FLightSceneInfo* LightSceneInfo = LightAndIndex.Key;
+
+		const FMaterialRenderProxy* MaterialProxy = LightSceneInfo->Proxy->GetLightFunctionMaterial();
+		check(MaterialProxy != nullptr);
+		// Catch the fallback material case
+		const FMaterialRenderProxy* FallbackMaterialRenderProxyPtr = nullptr;
+		const FMaterial& Material = MaterialProxy->GetMaterialWithFallback(Scene->GetFeatureLevel(), FallbackMaterialRenderProxyPtr);
+
+		check(Material.IsLightFunction());
+
+		const FMaterialRenderProxy& MaterialRenderProxy = FallbackMaterialRenderProxyPtr ? *FallbackMaterialRenderProxyPtr : *MaterialProxy;
+
+		TUniformBufferRef<FLightFunctionParametersPathTracing> LightFunctionParameters = CreateLightFunctionParametersBuffer(LightSceneInfo, View, EUniformBufferUsage::UniformBuffer_SingleFrame);
+
+		int32 MissIndex = LightAndIndex.Value;
+		BindLightFunction(RHICmdList, Scene, View, Material, MaterialRenderProxy, LightFunctionParameters, MissIndex);
+	}
+}
+
+
+FRayTracingLightFunctionMap GatherLightFunctionLightsPathTracing(FScene* Scene, const FEngineShowFlags EngineShowFlags, ERHIFeatureLevel::Type InFeatureLevel)
+{
+	checkf(EngineShowFlags.LightFunctions, TEXT("This function should not be called if light functions are disabled"));
+	FRayTracingLightFunctionMap RayTracingLightFunctionMap;
+	for (const FLightSceneInfoCompact& Light : Scene->Lights)
+	{
+		FLightSceneInfo* LightSceneInfo = Light.LightSceneInfo;
+		auto MaterialProxy = LightSceneInfo->Proxy->GetLightFunctionMaterial();
+		if (MaterialProxy)
+		{
+			const FMaterialRenderProxy* FallbackMaterialRenderProxyPtr = nullptr;
+			const FMaterial& Material = MaterialProxy->GetMaterialWithFallback(InFeatureLevel, FallbackMaterialRenderProxyPtr);
+			if (Material.IsLightFunction())
+			{
+				const FMaterialShaderMap* MaterialShaderMap = Material.GetRenderingThreadShaderMap();
+				// Getting the shader here has the side-effect of populating the raytracing miss shader library which is used when building the raytracing pipeline
+				MaterialShaderMap->GetShader<FPathTracingLightingMS>().GetRayTracingShader();
+
+				int32 Index = Scene->RayTracingScene.NumMissShaderSlots;
+				Scene->RayTracingScene.NumMissShaderSlots++;
+				RayTracingLightFunctionMap.Add(LightSceneInfo, Index);
+			}
+		}
+	}
+	return RayTracingLightFunctionMap;
+}
 
 template<bool UseAnyHitShader, bool UseIntersectionShader, bool UseSimplifiedShader>
 class TPathTracingMaterial : public FMeshMaterialShader
@@ -1365,6 +1583,7 @@ void SetLightParameters(FRDGBuilder& GraphBuilder, FPathTracingRG::FParameters*
 		DestLight.Flags |= Scene->SkyLight->bCastVolumetricShadow ? PATHTRACER_FLAG_CAST_VOL_SHADOW_MASK : 0;
 		DestLight.VolumetricScatteringIntensity = Scene->SkyLight->VolumetricScatteringIntensity;
 		DestLight.IESTextureSlice = -1;
+		DestLight.MissShaderIndex = 0;
 		DestLight.TranslatedBoundMin = FVector3f(-Inf, -Inf, -Inf);
 		DestLight.TranslatedBoundMax = FVector3f( Inf,  Inf,  Inf);
 		if (Scene->SkyLight->bRealTimeCaptureEnabled || CVarPathTracingVisibleLights.GetValueOnRenderThread() == 2)
@@ -1405,6 +1624,7 @@ void SetLightParameters(FRDGBuilder& GraphBuilder, FPathTracingRG::FParameters*
 			DestLight.Flags |= Light.LightSceneInfo->Proxy->CastsDynamicShadow() ? PATHTRACER_FLAG_CAST_SHADOW_MASK : 0;
 			DestLight.Flags |= Light.LightSceneInfo->Proxy->CastsVolumetricShadow() ? PATHTRACER_FLAG_CAST_VOL_SHADOW_MASK : 0;
 			DestLight.IESTextureSlice = -1;
+			DestLight.MissShaderIndex = 0;
 
 			// these mean roughly the same thing across all light types
 			DestLight.Color = FVector3f(LightParameters.Color);
@@ -1439,6 +1659,7 @@ void SetLightParameters(FRDGBuilder& GraphBuilder, FPathTracingRG::FParameters*
 	int32 NextRectTextureIndex = 0;
 
 	TMap<FTexture*, int> IESLightProfilesMap;
+	const FRayTracingLightFunctionMap* RayTracingLightFunctionMap = GraphBuilder.Blackboard.Get<FRayTracingLightFunctionMap>();
 	for (auto Light : Scene->Lights)
 	{
 		ELightComponentType LightComponentType = (ELightComponentType)Light.LightSceneInfo->Proxy->GetLightType();
@@ -1470,6 +1691,7 @@ void SetLightParameters(FRDGBuilder& GraphBuilder, FPathTracingRG::FParameters*
 		DestLight.Flags |= Light.LightSceneInfo->Proxy->CastsDynamicShadow() ? PATHTRACER_FLAG_CAST_SHADOW_MASK : 0;
 		DestLight.Flags |= Light.LightSceneInfo->Proxy->CastsVolumetricShadow() ? PATHTRACER_FLAG_CAST_VOL_SHADOW_MASK : 0;
 		DestLight.IESTextureSlice = -1;
+		DestLight.MissShaderIndex = 0;
 
 		if (View.Family->EngineShowFlags.TexturedLightProfiles)
 		{
@@ -1493,6 +1715,15 @@ void SetLightParameters(FRDGBuilder& GraphBuilder, FPathTracingRG::FParameters*
 		DestLight.RectLightAtlasUVOffset = 0;
 		DestLight.RectLightAtlasUVScale = 0;
 
+		if (RayTracingLightFunctionMap)
+		{
+			const int32* LightFunctionIndex = RayTracingLightFunctionMap->Find(Light.LightSceneInfo);
+			if (LightFunctionIndex)
+			{
+				DestLight.MissShaderIndex = *LightFunctionIndex;
+			}
+		}
+
 		switch (LightComponentType)
 		{
 			case LightType_Rect:
@@ -1734,6 +1965,7 @@ void FDeferredShadingSceneRenderer::RenderPathTracing(
 	Config.LightShowFlags |= View.Family->EngineShowFlags.SpotLights            ? 1 << 3 : 0;
 	Config.LightShowFlags |= View.Family->EngineShowFlags.PointLights           ? 1 << 4 : 0;
 	Config.LightShowFlags |= View.Family->EngineShowFlags.TexturedLightProfiles ? 1 << 5 : 0;
+	Config.LightShowFlags |= View.Family->EngineShowFlags.LightFunctions        ? 1 << 6 : 0;
 
 	PreparePathTracingData(Scene, View, Config.PathTracingData);
 
diff --git a/Engine/Source/Runtime/Renderer/Private/RayTracing/RayTracingLighting.cpp b/Engine/Source/Runtime/Renderer/Private/RayTracing/RayTracingLighting.cpp
--- a/Engine/Source/Runtime/Renderer/Private/RayTracing/RayTracingLighting.cpp
+++ b/Engine/Source/Runtime/Renderer/Private/RayTracing/RayTracingLighting.cpp
@@ -171,7 +171,6 @@ static void SetupRaytracingLightDataPacked(
 	const FScene* Scene,
 	const TArray<int32>& LightIndices,
 	const FSceneView& View,
-	const FRayTracingLightFunctionMap& RayTracingLightFunctionMap,
 	FRaytracingLightDataPacked& OutLightData,
 	TArrayView<FRTLightingData>& OutLightDataArray)
 {
@@ -201,6 +200,7 @@ static void SetupRaytracingLightDataPacked(
 		OutLightData.IESLightProfileTextureSampler = TStaticSamplerState<SF_Bilinear, AM_Clamp, AM_Clamp, AM_Clamp>::GetRHI();
 	}
 
+	const FRayTracingLightFunctionMap* RayTracingLightFunctionMap = GraphBuilder.Blackboard.Get<FRayTracingLightFunctionMap>();
 	for (auto LightIndex : LightIndices)
 	{
 		auto Light = Lights[LightIndex];
@@ -277,12 +277,15 @@ static void SetupRaytracingLightDataPacked(
 		}
 
 		// NOTE: This map will be empty if the light functions are disabled for some reason
-		const int32* LightFunctionIndex = RayTracingLightFunctionMap.Find(Light.LightSceneInfo);
-		if (LightFunctionIndex)
+		if (RayTracingLightFunctionMap)
 		{
-			check(uint32(*LightFunctionIndex) > RAY_TRACING_MISS_SHADER_SLOT_LIGHTING);
-			check(uint32(*LightFunctionIndex) < Scene->RayTracingScene.NumMissShaderSlots);
-			LightDataElement.LightMissShaderIndex = *LightFunctionIndex;
+			const int32* LightFunctionIndex = RayTracingLightFunctionMap->Find(Light.LightSceneInfo);
+			if (LightFunctionIndex)
+			{
+				check(uint32(*LightFunctionIndex) > RAY_TRACING_MISS_SHADER_SLOT_LIGHTING);
+				check(uint32(*LightFunctionIndex) < Scene->RayTracingScene.NumMissShaderSlots);
+				LightDataElement.LightMissShaderIndex = *LightFunctionIndex;
+			}
 		}
 
 		OutLightData.Count++;
@@ -310,7 +313,6 @@ TRDGUniformBufferRef<FRaytracingLightDataPacked> CreateRayTracingLightData(
 	FRDGBuilder& GraphBuilder,
 	const FScene* Scene,
 	const FSceneView& View,
-	const FRayTracingLightFunctionMap& RayTracingLightFunctionMap,
 	FGlobalShaderMap* ShaderMap,
 	uint32& NumOfSkippedRayTracingLights)
 {
@@ -327,7 +329,7 @@ TRDGUniformBufferRef<FRaytracingLightDataPacked> CreateRayTracingLightData(
 	CreateRaytracingLightCullingStructure(GraphBuilder, Scene->Lights, View, ShaderMap, LightIndices, LightCullVolume, LightIndicesBuffer);
 
 	FRDGUploadData<FRTLightingData> LightDataArray(GraphBuilder, FMath::Max(LightIndices.Num(), 1));
-	SetupRaytracingLightDataPacked(GraphBuilder, Scene, LightIndices, View, RayTracingLightFunctionMap, *LightData, LightDataArray);
+	SetupRaytracingLightDataPacked(GraphBuilder, Scene, LightIndices, View, *LightData, LightDataArray);
 
 	check(LightData->Count == LightIndices.Num());
 
@@ -478,11 +480,7 @@ IMPLEMENT_MATERIAL_SHADER_TYPE(, FLightFunctionRayTracingShader, TEXT("/Engine/P
 
 FRayTracingLightFunctionMap GatherLightFunctionLights(FScene* Scene, const FEngineShowFlags EngineShowFlags, ERHIFeatureLevel::Type InFeatureLevel)
 {
-	// TODO: implement light functions in the path tracer
-	if (!EngineShowFlags.LightFunctions || EngineShowFlags.PathTracing)
-	{
-		return {};
-	}
+	checkf(EngineShowFlags.LightFunctions, TEXT("This function should not be called if light functions are disabled"));
 
 	// gives control over lighting functions in raytraced effects, independently of the show flag (for performance testing / debug)
 	if (CVarRayTracingLightFunction.GetValueOnRenderThread() == 0)
@@ -559,34 +557,18 @@ FRHIRayTracingShader* FDeferredShadingSceneRenderer::GetRayTracingLightingMissSh
 	return View.ShaderMap->GetShader<FRayTracingLightingMS>().GetRayTracingShader();
 }
 
-void PrepareLightFunctionMissShaders(
-	const TArray<FLightSceneInfo*>& LightFunctionLights,
-	const class FViewInfo& View,
-	TArray<FRHIRayTracingShader*>& OutMissShaders)
-{
-	for (const FLightSceneInfo* LightSceneInfo : LightFunctionLights)
-	{
-		auto MaterialProxy = LightSceneInfo->Proxy->GetLightFunctionMaterial();
-		if (MaterialProxy)
-		{
-			const FMaterialRenderProxy* FallbackMaterialRenderProxyPtr = nullptr;
-			const FMaterial& Material = MaterialProxy->GetMaterialWithFallback(View.GetFeatureLevel(), FallbackMaterialRenderProxyPtr);
-			if (Material.IsLightFunction())
-			{
-				const FMaterialShaderMap* MaterialShaderMap = Material.GetRenderingThreadShaderMap();
-				OutMissShaders.Add(MaterialShaderMap->GetShader<FLightFunctionRayTracingShader>().GetRayTracingShader());
-			}
-		}
-	}
-}
-
 void BindLightFunctionShaders(
 	FRHICommandList& RHICmdList,
 	const FScene* Scene,
-	const FRayTracingLightFunctionMap& RayTracingLightFunctionMap,
+	const FRayTracingLightFunctionMap* RayTracingLightFunctionMap,
 	const class FViewInfo& View)
 {
-	for (const FRayTracingLightFunctionMap::ElementType& LightAndIndex : RayTracingLightFunctionMap)
+	if (RayTracingLightFunctionMap == nullptr)
+	{
+		return;
+	}
+
+	for (const FRayTracingLightFunctionMap::ElementType& LightAndIndex : *RayTracingLightFunctionMap)
 	{
 		const FLightSceneInfo* LightSceneInfo = LightAndIndex.Key;
 
diff --git a/Engine/Source/Runtime/Renderer/Private/RayTracing/RayTracingLighting.h b/Engine/Source/Runtime/Renderer/Private/RayTracing/RayTracingLighting.h
--- a/Engine/Source/Runtime/Renderer/Private/RayTracing/RayTracingLighting.h
+++ b/Engine/Source/Runtime/Renderer/Private/RayTracing/RayTracingLighting.h
@@ -54,25 +54,30 @@ struct FRTLightingData
 static_assert(sizeof(FRTLightingData) == 128, "Unexpected FRTLightingData size.");
 
 using FRayTracingLightFunctionMap = TMap<const FLightSceneInfo*, int32>;
+
+// Register this in the graph builder so we can easily move it around and access it from both the main rendering thread and RDG passes
+RDG_REGISTER_BLACKBOARD_STRUCT(FRayTracingLightFunctionMap)
+
 FRayTracingLightFunctionMap GatherLightFunctionLights(FScene* Scene, const FEngineShowFlags EngineShowFlags, ERHIFeatureLevel::Type InFeatureLevel);
+FRayTracingLightFunctionMap GatherLightFunctionLightsPathTracing(FScene* Scene, const FEngineShowFlags EngineShowFlags, ERHIFeatureLevel::Type InFeatureLevel);
 
 TRDGUniformBufferRef<FRaytracingLightDataPacked> CreateRayTracingLightData(
 	FRDGBuilder& GraphBuilder,
 	const FScene* Scene,
 	const FSceneView& View,
-	const FRayTracingLightFunctionMap& RayTracingLightFunctionMap,
 	FGlobalShaderMap* ShaderMap,
 	uint32& NumOfSkippedRayTracingLights);
 
 void BindLightFunctionShaders(
 	FRHICommandList& RHICmdList,
 	const FScene* Scene,
-	const FRayTracingLightFunctionMap& RayTracingLightFunctionMap,
+	const FRayTracingLightFunctionMap* RayTracingLightFunctionMap,
 	const class FViewInfo& View);
 
-void PrepareLightFunctionMissShaders(
-	const TArray<FLightSceneInfo*>& LightFunctionLights,
-	const class FViewInfo& View,
-	TArray<FRHIRayTracingShader*>& OutMissShaders);
+void BindLightFunctionShadersPathTracing(
+	FRHICommandList& RHICmdList,
+	const FScene* Scene,
+	const FRayTracingLightFunctionMap* RayTracingLightFunctionMap,
+	const class FViewInfo& View);
 
 #endif
diff --git a/Engine/Source/Runtime/Renderer/Private/RayTracing/RayTracingMaterialHitShaders.cpp b/Engine/Source/Runtime/Renderer/Private/RayTracing/RayTracingMaterialHitShaders.cpp
--- a/Engine/Source/Runtime/Renderer/Private/RayTracing/RayTracingMaterialHitShaders.cpp
+++ b/Engine/Source/Runtime/Renderer/Private/RayTracing/RayTracingMaterialHitShaders.cpp
@@ -601,7 +601,7 @@ FRayTracingPipelineState* FDeferredShadingSceneRenderer::CreateRayTracingMateria
 	Initializer.MaxPayloadSizeInBytes = RAY_TRACING_MAX_ALLOWED_PAYLOAD_SIZE; // sizeof(FPackedMaterialClosestHitPayload)
 	Initializer.bAllowHitGroupIndexing = true;
 
-	FRHIRayTracingShader* DefaultMissShader = GetRayTracingDefaultMissShader(View);
+	FRHIRayTracingShader* DefaultMissShader = bIsPathTracing ? GetPathTracingDefaultMissShader(View) : GetRayTracingDefaultMissShader(View);
 
 	TArray<FRHIRayTracingShader*> RayTracingMissShaderLibrary;
 	FShaderMapResource::GetRayTracingMissShaderLibrary(RayTracingMissShaderLibrary, DefaultMissShader);
