Remove draw canvas from hair debug code.

diff --git a/Engine/Shaders/Private/HairStrands/HairStrandsDebug.usf b/Engine/Shaders/Private/HairStrands/HairStrandsDebug.usf
--- a/Engine/Shaders/Private/HairStrands/HairStrandsDebug.usf
+++ b/Engine/Shaders/Private/HairStrands/HairStrandsDebug.usf
@@ -664,11 +664,28 @@ void VisualizeDomPS(
 
 #include "../ShaderPrint.ush"
 #include "../Matrices.ush"
+#include "../ShadingCommon.ush"
 #include "HairStrandsDeepShadowCommonStruct.ush"
 #include "HairStrandsAABBCommon.ush"
 
+void PrintBool(inout FShaderPrintContext Ctx, bool In)
+{
+	if (In)
+	{
+		Print(Ctx, TEXT("Yes  "), FontGreen);
+	}
+	else
+	{
+		Print(Ctx, TEXT("No   "), FontRed);
+	}
+}
+
 uint AllocatedSlotCount;
 uint MacroGroupCount;
+uint bViewRectOptimizeEnabled;
+uint bVoxelizationEnabled;
+int2 AtlasResolution;
+uint bIsGPUDriven;
 
 Buffer<int>	MacroGroupAABBBuffer;
 StructuredBuffer<FDeepShadowViewInfo> ShadowViewInfoBuffer;
@@ -728,6 +745,28 @@ void MainCS(uint3 DispatchThreadId : SV_DispatchThreadID)
 		AddLineTWS(P21.xyz, P31.xyz, FrustumColor, FrustumColor);
 		AddLineTWS(P31.xyz, P01.xyz, FrustumColor, FrustumColor);
 	}
+
+
+	FShaderPrintContext Context = InitShaderPrintContext(true, uint2(50, 40));
+
+	FFontColor FontColor = FontSilver;
+	Print(Context, TEXT("Hair Components "), FontYellow); Newline(Context);
+	Print(Context, TEXT("R                     : "), FontColor); PrintBool(Context, View.HairComponents & HAIR_COMPONENT_R); Newline(Context);
+	Print(Context, TEXT("TT                    : "), FontColor); PrintBool(Context, View.HairComponents & HAIR_COMPONENT_TT); Newline(Context);
+	Print(Context, TEXT("TRT                   : "), FontColor); PrintBool(Context, View.HairComponents & HAIR_COMPONENT_TRT); Newline(Context);
+	Print(Context, TEXT("Global Scatt.         : "), FontColor); PrintBool(Context, View.HairComponents & HAIR_COMPONENT_GS); Newline(Context);
+	Print(Context, TEXT("Local Scatt.          : "), FontColor); PrintBool(Context, View.HairComponents & HAIR_COMPONENT_LS); Newline(Context);
+	Newline(Context);
+
+	Print(Context, TEXT("Misc "), FontYellow); Newline(Context); 
+	Print(Context, TEXT("Voxelization          : "), FontColor); PrintBool(Context, bVoxelizationEnabled); Newline(Context);
+	Print(Context, TEXT("View rect optim.      : "), FontColor); PrintBool(Context, bViewRectOptimizeEnabled); Newline(Context);
+	Newline(Context);
+
+	Print(Context, TEXT("DOM "), FontYellow); Newline(Context);
+	Print(Context, TEXT("Atlas resolution      : "), FontColor); Print(Context, AtlasResolution.x, FontColor, 4, 0); Print(Context, TEXT(" x "), FontColor); Print(Context, AtlasResolution.y, FontColor, 4, 0); Newline(Context);
+	Print(Context, TEXT("Atlas slots           : "), FontColor); Print(Context, AllocatedSlotCount, FontColor); Newline(Context);
+	Print(Context, TEXT("GPU driven            : "), FontColor); PrintBool(Context, bIsGPUDriven); Newline(Context);
 }
 #endif // SHADER_DOMINFO
 
diff --git a/Engine/Source/Runtime/Renderer/Private/HairStrands/HairStrandsDebug.cpp b/Engine/Source/Runtime/Renderer/Private/HairStrands/HairStrandsDebug.cpp
--- a/Engine/Source/Runtime/Renderer/Private/HairStrands/HairStrandsDebug.cpp
+++ b/Engine/Source/Runtime/Renderer/Private/HairStrands/HairStrandsDebug.cpp
@@ -660,11 +660,15 @@ class FDeepShadowInfoCS : public FGlobalShader
 		SHADER_PARAMETER(FVector2f, OutputResolution)
 		SHADER_PARAMETER(uint32, AllocatedSlotCount)
 		SHADER_PARAMETER(uint32, MacroGroupCount)
+		SHADER_PARAMETER(uint32, bViewRectOptimizeEnabled)
+		SHADER_PARAMETER(uint32, bVoxelizationEnabled)
+		SHADER_PARAMETER(FIntPoint, AtlasResolution)
+		SHADER_PARAMETER(uint32, bIsGPUDriven)
 		SHADER_PARAMETER_STRUCT_REF(FViewUniformShaderParameters, ViewUniformBuffer)
 		SHADER_PARAMETER_RDG_BUFFER_SRV(Buffer, MacroGroupAABBBuffer)
 		SHADER_PARAMETER_RDG_BUFFER_SRV(StructuredBuffer, ShadowViewInfoBuffer)
 		SHADER_PARAMETER_RDG_TEXTURE_UAV(RWTexture2D, OutputTexture)
-		END_SHADER_PARAMETER_STRUCT()
+	END_SHADER_PARAMETER_STRUCT()
 
 public:
 	static bool ShouldCompilePermutation(const FGlobalShaderPermutationParameters& Parameters) { return IsHairStrandsSupported(EHairStrandsShaderType::Tool, Parameters.Platform); }
@@ -703,6 +707,10 @@ static void AddDeepShadowInfoPass(
 	Parameters->AllocatedSlotCount = DeepShadowResources.TotalAtlasSlotCount;
 	Parameters->MacroGroupCount = MacroGroupResources.MacroGroupCount;
 	Parameters->SceneTextures = SceneTextures;
+	Parameters->bViewRectOptimizeEnabled = IsHairStrandsViewRectOptimEnable() ? 1u : 0u;
+	Parameters->bVoxelizationEnabled = View.HairStrandsViewData.VirtualVoxelResources.IsValid() ? 1u : 0u;
+	Parameters->AtlasResolution = View.HairStrandsViewData.DeepShadowResources.DepthAtlasTexture->Desc.Extent;
+	Parameters->bIsGPUDriven = View.HairStrandsViewData.DeepShadowResources.bIsGPUDriven ? 1u : 0u;
 	Parameters->MacroGroupAABBBuffer = GraphBuilder.CreateSRV(MacroGroupResources.MacroGroupAABBsBuffer, PF_R32_SINT);
 	Parameters->ShadowViewInfoBuffer = GraphBuilder.CreateSRV(DeepShadowResources.DeepShadowViewInfoBuffer);
 	ShaderPrint::SetParameters(GraphBuilder, View.ShaderPrintData, Parameters->ShaderPrintParameters);
@@ -1367,29 +1375,27 @@ static void InternalRenderHairStrandsDebugInfo(
 		}
 	}
 
+	if (HairData.DebugData.IsPPLLDataValid()) // Check if PPLL rendering is used and its debug view is enabled.
 	{
-		if (HairData.DebugData.IsPPLLDataValid()) // Check if PPLL rendering is used and its debug view is enabled.
-		{
-			// Force ShaderPrint on.
-			ShaderPrint::SetEnabled(true);
-			ShaderPrint::RequestSpaceForCharacters(256);
-
-			const FIntPoint PPLLResolution = HairData.DebugData.PPLLNodeIndexTexture->Desc.Extent;
-			FHairVisibilityDebugPPLLCS::FParameters* PassParameters = GraphBuilder.AllocParameters<FHairVisibilityDebugPPLLCS::FParameters>();
-			PassParameters->PPLLMeanListElementCountPerPixel = GetHairStrandsMeanSamplePerPixel();
-			PassParameters->PPLLMaxTotalListElementCount = HairData.DebugData.PPLLNodeDataBuffer->Desc.NumElements;
-			PassParameters->PPLLCounter = HairData.DebugData.PPLLNodeCounterTexture;
-			PassParameters->PPLLNodeIndex = HairData.DebugData.PPLLNodeIndexTexture;
-			PassParameters->PPLLNodeData = GraphBuilder.CreateSRV(FRDGBufferSRVDesc(HairData.DebugData.PPLLNodeDataBuffer));
-			PassParameters->ViewUniformBuffer = View.ViewUniformBuffer;
-			PassParameters->SceneColorTextureUAV = GraphBuilder.CreateUAV(SceneColorTexture);
-			ShaderPrint::SetParameters(GraphBuilder, View.ShaderPrintData, PassParameters->ShaderPrintParameters);
-
-			FHairVisibilityDebugPPLLCS::FPermutationDomain PermutationVector;
-			TShaderMapRef<FHairVisibilityDebugPPLLCS> ComputeShader(View.ShaderMap, PermutationVector);
-			FIntVector TextureSize = SceneColorTexture->Desc.GetSize(); TextureSize.Z = 1;
-			FComputeShaderUtils::AddPass(GraphBuilder, RDG_EVENT_NAME("HairStrands::PPLLDebug"), ComputeShader, PassParameters, FIntVector::DivideAndRoundUp(TextureSize, FIntVector(8, 8, 1)));
-		}
+		// Force ShaderPrint on.
+		ShaderPrint::SetEnabled(true);
+		ShaderPrint::RequestSpaceForCharacters(256);
+
+		const FIntPoint PPLLResolution = HairData.DebugData.PPLLNodeIndexTexture->Desc.Extent;
+		FHairVisibilityDebugPPLLCS::FParameters* PassParameters = GraphBuilder.AllocParameters<FHairVisibilityDebugPPLLCS::FParameters>();
+		PassParameters->PPLLMeanListElementCountPerPixel = GetHairStrandsMeanSamplePerPixel();
+		PassParameters->PPLLMaxTotalListElementCount = HairData.DebugData.PPLLNodeDataBuffer->Desc.NumElements;
+		PassParameters->PPLLCounter = HairData.DebugData.PPLLNodeCounterTexture;
+		PassParameters->PPLLNodeIndex = HairData.DebugData.PPLLNodeIndexTexture;
+		PassParameters->PPLLNodeData = GraphBuilder.CreateSRV(FRDGBufferSRVDesc(HairData.DebugData.PPLLNodeDataBuffer));
+		PassParameters->ViewUniformBuffer = View.ViewUniformBuffer;
+		PassParameters->SceneColorTextureUAV = GraphBuilder.CreateUAV(SceneColorTexture);
+		ShaderPrint::SetParameters(GraphBuilder, View.ShaderPrintData, PassParameters->ShaderPrintParameters);
+
+		FHairVisibilityDebugPPLLCS::FPermutationDomain PermutationVector;
+		TShaderMapRef<FHairVisibilityDebugPPLLCS> ComputeShader(View.ShaderMap, PermutationVector);
+		FIntVector TextureSize = SceneColorTexture->Desc.GetSize(); TextureSize.Z = 1;
+		FComputeShaderUtils::AddPass(GraphBuilder, RDG_EVENT_NAME("HairStrands::PPLLDebug"), ComputeShader, PassParameters, FIntVector::DivideAndRoundUp(TextureSize, FIntVector(8, 8, 1)));
 	}
 
 	if (GHairStrandsClusterDebug > 0)
@@ -1397,57 +1403,6 @@ static void InternalRenderHairStrandsDebugInfo(
 		AddDrawDebugClusterPass(GraphBuilder, HairClusterData, View);
 	}
 
-	// Text
-	if (HairDebugMode == EHairDebugMode::LightBounds || HairDebugMode == EHairDebugMode::DeepOpacityMaps)
-	{
-		AddDrawCanvasPass(GraphBuilder, {}, View, SceneColor, [&HairData, YStep, StrandsDebugMode, HairDebugMode](FCanvas& Canvas)
-		{
-			const uint32 AtlasTotalSlotCount = FHairStrandsDeepShadowResources::MaxAtlasSlotCount;
-			uint32 AtlasAllocatedSlot = 0;
-			FIntPoint AtlasSlotResolution = FIntPoint(0, 0);
-			FIntPoint AtlasResolution = FIntPoint(0, 0);
-			bool bIsGPUDriven = false;
-			{
-				const FHairStrandsDeepShadowResources& Resources = HairData.DeepShadowResources;
-				AtlasResolution = Resources.DepthAtlasTexture ? Resources.DepthAtlasTexture->Desc.Extent : FIntPoint(0, 0);
-				AtlasAllocatedSlot = Resources.TotalAtlasSlotCount;
-				bIsGPUDriven = Resources.bIsGPUDriven;
-			}
-
-			const uint32 DomTextureIndex = GDeepShadowDebugIndex;
-
-			float X = 20;
-			float Y = 38;
-
-			FLinearColor DebugColor(1, 1, 0);
-			FString Line;
-
-			const FHairComponent HairComponent = GetHairComponents();
-			Line = FString::Printf(TEXT("Hair Components : (R=%d, TT=%d, TRT=%d, GS=%d, LS=%d)"), HairComponent.R, HairComponent.TT, HairComponent.TRT, HairComponent.GlobalScattering, HairComponent.LocalScattering);
-			Canvas.DrawShadowedString(X, Y += YStep, *Line, GetStatsFont(), DebugColor);
-			Line = FString::Printf(TEXT("----------------------------------------------------------------"));						Canvas.DrawShadowedString(X, Y += YStep, *Line, GetStatsFont(), DebugColor);
-			Line = FString::Printf(TEXT("Debug strands mode : %s"), ToString(StrandsDebugMode));									Canvas.DrawShadowedString(X, Y += YStep, *Line, GetStatsFont(), DebugColor);
-			Line = FString::Printf(TEXT("Voxelization : %s"), IsHairStrandsVoxelizationEnable() ? TEXT("On") : TEXT("Off"));		Canvas.DrawShadowedString(X, Y += YStep, *Line, GetStatsFont(), DebugColor);
-			Line = FString::Printf(TEXT("View rect optim.: %s"), IsHairStrandsViewRectOptimEnable() ? TEXT("On") : TEXT("Off"));	Canvas.DrawShadowedString(X, Y += YStep, *Line, GetStatsFont(), DebugColor);
-			Line = FString::Printf(TEXT("----------------------------------------------------------------"));						Canvas.DrawShadowedString(X, Y += YStep, *Line, GetStatsFont(), DebugColor);
-			Line = FString::Printf(TEXT("DOM Atlas resolution  : %dx%d"), AtlasResolution.X, AtlasResolution.Y);					Canvas.DrawShadowedString(X, Y += YStep, *Line, GetStatsFont(), DebugColor);
-			Line = FString::Printf(TEXT("DOM Atlas slot        : %d/%d"), AtlasAllocatedSlot, AtlasTotalSlotCount);					Canvas.DrawShadowedString(X, Y += YStep, *Line, GetStatsFont(), DebugColor);
-			Line = FString::Printf(TEXT("DOM Texture Index     : %d/%d"), DomTextureIndex, AtlasAllocatedSlot);						Canvas.DrawShadowedString(X, Y += YStep, *Line, GetStatsFont(), DebugColor);
-			Line = FString::Printf(TEXT("DOM GPU driven        : %s"), bIsGPUDriven ? TEXT("On") : TEXT("Off"));					Canvas.DrawShadowedString(X, Y += YStep, *Line, GetStatsFont(), DebugColor);
-
-			{
-				for (const FHairStrandsMacroGroupData& MacroGroupData : HairData.MacroGroupDatas)
-				{
-					for (const FHairStrandsDeepShadowData& DomData : MacroGroupData.DeepShadowDatas)
-					{
-						Line = FString::Printf(TEXT(" %d - Bound Radus: %f.2m (%dx%d)"), DomData.AtlasSlotIndex, DomData.Bounds.GetSphere().W / 10.f, DomData.ShadowResolution.X, DomData.ShadowResolution.Y);
-						Canvas.DrawShadowedString(X, Y += YStep, *Line, GetStatsFont(), DebugColor);
-					}
-				}
-			}
-		});
-	}
-
 	if (StrandsDebugMode != EHairStrandsDebugMode::NoneDebug || HairDebugMode != EHairDebugMode::None)
 	{
 		AddDrawCanvasPass(GraphBuilder, {}, View, SceneColor, [&View, &StrandsDebugMode, YStep, HairDebugMode](FCanvas& Canvas)
