Fixed assert due to RDG stat scope exceeding scope of RDG builder.

diff --git a/Engine/Source/Runtime/Renderer/Private/SceneRendering.cpp b/Engine/Source/Runtime/Renderer/Private/SceneRendering.cpp
--- a/Engine/Source/Runtime/Renderer/Private/SceneRendering.cpp
+++ b/Engine/Source/Runtime/Renderer/Private/SceneRendering.cpp
@@ -3964,18 +3964,19 @@ void FSceneRenderer::ViewExtensionPreRender_RenderThread(FRHICommandListImmediat
 
 	{
 		FRDGBuilder GraphBuilder(RHICmdList);
-		RDG_CSV_STAT_EXCLUSIVE_SCOPE(GraphBuilder, PreRender);
-		SCOPE_CYCLE_COUNTER(STAT_FDeferredShadingSceneRenderer_ViewExtensionPreRenderView);
-
-		for (int ViewExt = 0; ViewExt < SceneRenderer->ViewFamily.ViewExtensions.Num(); ViewExt++)
 		{
-			SceneRenderer->ViewFamily.ViewExtensions[ViewExt]->PreRenderViewFamily_RenderThread(GraphBuilder, SceneRenderer->ViewFamily);
-			for (int ViewIndex = 0; ViewIndex < SceneRenderer->ViewFamily.Views.Num(); ViewIndex++)
+			RDG_CSV_STAT_EXCLUSIVE_SCOPE(GraphBuilder, PreRender);
+			SCOPE_CYCLE_COUNTER(STAT_FDeferredShadingSceneRenderer_ViewExtensionPreRenderView);
+
+			for (int ViewExt = 0; ViewExt < SceneRenderer->ViewFamily.ViewExtensions.Num(); ViewExt++)
 			{
-				SceneRenderer->ViewFamily.ViewExtensions[ViewExt]->PreRenderView_RenderThread(GraphBuilder, SceneRenderer->Views[ViewIndex]);
+				SceneRenderer->ViewFamily.ViewExtensions[ViewExt]->PreRenderViewFamily_RenderThread(GraphBuilder, SceneRenderer->ViewFamily);
+				for (int ViewIndex = 0; ViewIndex < SceneRenderer->ViewFamily.Views.Num(); ViewIndex++)
+				{
+					SceneRenderer->ViewFamily.ViewExtensions[ViewExt]->PreRenderView_RenderThread(GraphBuilder, SceneRenderer->Views[ViewIndex]);
+				}
 			}
 		}
-
 		GraphBuilder.Execute();
 	}
 	
