Mobile: clean up render-pass setup around muti-pass rendering

diff --git a/Engine/Source/Runtime/Renderer/Private/MobileShadingRenderer.cpp b/Engine/Source/Runtime/Renderer/Private/MobileShadingRenderer.cpp
--- a/Engine/Source/Runtime/Renderer/Private/MobileShadingRenderer.cpp
+++ b/Engine/Source/Runtime/Renderer/Private/MobileShadingRenderer.cpp
@@ -1063,11 +1063,8 @@ void FMobileSceneRenderer::RenderForward(FRDGBuilder& GraphBuilder, FRDGTextureR
 		FDepthStencilBinding(SceneDepth, ERenderTargetLoadAction::ELoad, ERenderTargetLoadAction::ELoad, FExclusiveDepthStencil::DepthRead_StencilWrite) : 
 		FDepthStencilBinding(SceneDepth, ERenderTargetLoadAction::EClear, ERenderTargetLoadAction::EClear, FExclusiveDepthStencil::DepthWrite_StencilWrite);
 	BasePassRenderTargets.ShadingRateTexture = (!MainView.bIsSceneCapture && !MainView.bIsReflectionCapture && ShadingRateTarget.IsValid()) ? RegisterExternalTexture(GraphBuilder, ShadingRateTarget->GetRHI(), TEXT("ShadingRateTexture")) : nullptr;
-	BasePassRenderTargets.SubpassHint = ESubpassHint::DepthReadSubpass;
-	if (!bIsFullDepthPrepassEnabled)
-	{
-		BasePassRenderTargets.NumOcclusionQueries = ComputeNumOcclusionQueriesToBatch();
-	}
+	BasePassRenderTargets.SubpassHint = ESubpassHint::None;
+	BasePassRenderTargets.NumOcclusionQueries = 0u;
 
 	//if the scenecolor isn't multiview but the app is, need to render as a single-view multiview due to shaders
 	BasePassRenderTargets.MultiViewCount = MainView.bIsMobileMultiViewEnabled ? 2 : (bIsMultiViewApplication ? 1 : 0);
@@ -1128,6 +1125,12 @@ void FMobileSceneRenderer::RenderForward(FRDGBuilder& GraphBuilder, FRDGTextureR
 
 void FMobileSceneRenderer::RenderForwardSinglePass(FRDGBuilder& GraphBuilder, FMobileRenderPassParameters* PassParameters, int32 ViewIndex, FViewInfo& View, FSceneTextures& SceneTextures)
 {
+	PassParameters->RenderTargets.SubpassHint = ESubpassHint::DepthReadSubpass;
+	if (!bIsFullDepthPrepassEnabled)
+	{
+		PassParameters->RenderTargets.NumOcclusionQueries = ComputeNumOcclusionQueriesToBatch();
+	}
+	
 	GraphBuilder.AddPass(
 		RDG_EVENT_NAME("SceneColorRendering"),
 		PassParameters,
@@ -1218,8 +1221,6 @@ void FMobileSceneRenderer::RenderForwardMultiPass(FRDGBuilder& GraphBuilder, FMo
 	BasePassRenderTargets.DepthStencil.SetDepthLoadAction(ERenderTargetLoadAction::ELoad);
 	BasePassRenderTargets.DepthStencil.SetStencilLoadAction(ERenderTargetLoadAction::ELoad);
 	BasePassRenderTargets.DepthStencil.SetDepthStencilAccess(ExclusiveDepthStencil);
-	BasePassRenderTargets.NumOcclusionQueries = 0;
-	BasePassRenderTargets.SubpassHint = ESubpassHint::DepthReadSubpass;
 
 	EMobileSceneTextureSetupMode SetupMode = EMobileSceneTextureSetupMode::SceneDepth | EMobileSceneTextureSetupMode::SceneDepthAux | EMobileSceneTextureSetupMode::CustomDepth;
 	FMobileRenderPassParameters* SecondPassParameters = GraphBuilder.AllocParameters<FMobileRenderPassParameters>();
@@ -1227,6 +1228,10 @@ void FMobileSceneRenderer::RenderForwardMultiPass(FRDGBuilder& GraphBuilder, FMo
 	SecondPassParameters->MobileBasePass = CreateMobileBasePassUniformBuffer(GraphBuilder, View, EMobileBasePass::Translucent, SetupMode);
 	SecondPassParameters->ReflectionCapture = View.MobileReflectionCaptureUniformBuffer;
 	SecondPassParameters->RenderTargets = BasePassRenderTargets;
+	if (!bIsFullDepthPrepassEnabled)
+	{
+		SecondPassParameters->RenderTargets.NumOcclusionQueries = ComputeNumOcclusionQueriesToBatch();
+	}
 
 	GraphBuilder.AddPass(
 		RDG_EVENT_NAME("DecalsAndTranslucency"),
@@ -1235,7 +1240,6 @@ void FMobileSceneRenderer::RenderForwardMultiPass(FRDGBuilder& GraphBuilder, FMo
 		[this, SecondPassParameters, ViewIndex, &View, &SceneTextures](FRHICommandListImmediate& RHICmdList)
 	{
 		// scene depth is read only and can be fetched
-		RHICmdList.NextSubpass();
 		RHICmdList.SetCurrentStat(GET_STATID(STAT_CLMM_Translucency));
 		RenderDecals(RHICmdList, View);
 		RenderModulatedShadowProjections(RHICmdList, ViewIndex, View);
@@ -1373,11 +1377,8 @@ void FMobileSceneRenderer::RenderDeferred(FRDGBuilder& GraphBuilder, const FSort
 	BasePassRenderTargets.DepthStencil = bIsFullDepthPrepassEnabled ? 
 		FDepthStencilBinding(SceneTextures.Depth.Target, ERenderTargetLoadAction::ELoad, ERenderTargetLoadAction::ELoad, FExclusiveDepthStencil::DepthRead_StencilWrite) : 
 		FDepthStencilBinding(SceneTextures.Depth.Target, ERenderTargetLoadAction::EClear, ERenderTargetLoadAction::EClear, FExclusiveDepthStencil::DepthWrite_StencilWrite);
-	BasePassRenderTargets.SubpassHint = ESubpassHint::DeferredShadingSubpass;
-	if (!bIsFullDepthPrepassEnabled)
-	{
-		BasePassRenderTargets.NumOcclusionQueries = ComputeNumOcclusionQueriesToBatch();
-	}
+	BasePassRenderTargets.SubpassHint = ESubpassHint::None;
+	BasePassRenderTargets.NumOcclusionQueries = 0u;
 	BasePassRenderTargets.ShadingRateTexture = nullptr;
 	BasePassRenderTargets.MultiViewCount = 0;
 
@@ -1426,6 +1427,12 @@ void FMobileSceneRenderer::RenderDeferred(FRDGBuilder& GraphBuilder, const FSort
 
 void FMobileSceneRenderer::RenderDeferredSinglePass(FRDGBuilder& GraphBuilder, class FMobileRenderPassParameters* PassParameters, int32 ViewIndex, int32 NumViews, FViewInfo& View, FSceneTextures& SceneTextures, const FSortedLightSetSceneInfo& SortedLightSet, bool bUsingPixelLocalStorage)
 {
+	PassParameters->RenderTargets.SubpassHint = ESubpassHint::DeferredShadingSubpass;
+	if (!bIsFullDepthPrepassEnabled)
+	{
+		PassParameters->RenderTargets.NumOcclusionQueries = ComputeNumOcclusionQueriesToBatch();
+	}
+		
 	GraphBuilder.AddPass(
 		RDG_EVENT_NAME("SceneColorRendering"),
 		PassParameters,
@@ -1491,9 +1498,7 @@ void FMobileSceneRenderer::RenderDeferredMultiPass(FRDGBuilder& GraphBuilder, cl
 	BasePassRenderTargets.DepthStencil.SetDepthLoadAction(ERenderTargetLoadAction::ELoad);
 	BasePassRenderTargets.DepthStencil.SetStencilLoadAction(ERenderTargetLoadAction::ELoad);
 	BasePassRenderTargets.DepthStencil.SetDepthStencilAccess(FExclusiveDepthStencil::DepthRead_StencilRead);
-	BasePassRenderTargets.SubpassHint = ESubpassHint::None;
-	BasePassRenderTargets.NumOcclusionQueries = 0;
-
+	
 	EMobileSceneTextureSetupMode SetupMode = EMobileSceneTextureSetupMode::SceneDepth | EMobileSceneTextureSetupMode::SceneDepthAux | EMobileSceneTextureSetupMode::CustomDepth;
 	auto* SecondPassParameters = GraphBuilder.AllocParameters<FMobileRenderPassParameters>();
 	*SecondPassParameters = *PassParameters;
@@ -1507,7 +1512,6 @@ void FMobileSceneRenderer::RenderDeferredMultiPass(FRDGBuilder& GraphBuilder, cl
 		ERDGPassFlags::Raster,
 		[this, SecondPassParameters, &View](FRHICommandListImmediate& RHICmdList)
 	{
-		RHICmdList.NextSubpass();
 		RenderDecals(RHICmdList, View);
 	});
 
@@ -1524,6 +1528,10 @@ void FMobileSceneRenderer::RenderDeferredMultiPass(FRDGBuilder& GraphBuilder, cl
 	ThirdPassParameters->MobileBasePass = CreateMobileBasePassUniformBuffer(GraphBuilder, View, EMobileBasePass::Translucent, SetupMode);
 	ThirdPassParameters->ReflectionCapture = View.MobileReflectionCaptureUniformBuffer;
 	ThirdPassParameters->RenderTargets = BasePassRenderTargets;
+	if (!bIsFullDepthPrepassEnabled)
+	{
+		ThirdPassParameters->RenderTargets.NumOcclusionQueries = ComputeNumOcclusionQueriesToBatch();
+	}
 
 	GraphBuilder.AddPass(
 		RDG_EVENT_NAME("LightingAndTranslucency"),
@@ -1531,7 +1539,6 @@ void FMobileSceneRenderer::RenderDeferredMultiPass(FRDGBuilder& GraphBuilder, cl
 		ERDGPassFlags::Raster,
 		[this, ThirdPassParameters, ViewIndex, NumViews, &View, &SceneTextures, &SortedLightSet](FRHICommandListImmediate& RHICmdList)
 	{
-		RHICmdList.NextSubpass();
 		RHICmdList.SetCurrentStat(GET_STATID(STAT_CLMM_Translucency));
 		MobileDeferredShadingPass(RHICmdList, ViewIndex, NumViews, View, *Scene, SortedLightSet, VisibleLightInfos);
 		RenderFog(RHICmdList, View);
