* Move compute raster common parameter into a common struct
* Add debug pass for compute raster to visualize settings & binning.

diff --git a/Engine/Shaders/Private/HairStrands/HairStrandsVisibilityRasterCompute.usf b/Engine/Shaders/Private/HairStrands/HairStrandsVisibilityRasterCompute.usf
--- a/Engine/Shaders/Private/HairStrands/HairStrandsVisibilityRasterCompute.usf
+++ b/Engine/Shaders/Private/HairStrands/HairStrandsVisibilityRasterCompute.usf
@@ -5,6 +5,30 @@
 #include "HairStrandsVertexFactoryCommon.ush"
 #include "HairStrandsVisibilityCommon.ush"
 
+///////////////////////////////////////////////////////////////////////////
+// Common parameters
+
+uint	TileSizeAsShift;
+uint	TileSize;
+float	RcpTileSize;
+uint	SqrTileSize;
+int2	TileRes;
+
+uint	NumBinners;
+float	RcpNumBinners;
+uint	NumRasterizers;
+float	RcpNumRasterizers;
+
+uint	MaxRasterCount;
+uint	FrameIdMod8;
+uint	ResolutionMultiplier;
+int2	OutputResolution;
+float2	OutputResolutionf;
+
+///////////////////////////////////////////////////////////////////////////
+
+#if SHADER_RASTERCOMPUTE_BINNING || SHADER_RASTERCOMPUTE_RASTER || SHADER_RASTERCOMPUTE_DEPTH_GRID
+
 ///////////////////////////////////////////////////////////////////////////
 // 
 // Wave size
@@ -22,15 +46,6 @@
 
 ///////////////////////////////////////////////////////////////////////////
 
-uint TileSize;
-float RcpTileSize;
-uint SqrTileSize;
-
-uint NumBinners;
-float RcpNumBinners;
-uint NumRasterizers;
-float RcpNumRasterizers;
-
 /* 
 // use untyped buffer for segment tiles to reduce VGPR usage - 16 bytes
 struct FVisTile
@@ -60,12 +75,7 @@ uint2 UppackVisTileCoord(uint Packed)
 
 ///////////////////////////////////////////////////////////////////////////
 
-uint			MaxRasterCount;
-uint			FrameIdMod8;
 uint			MacroGroupId;
-uint			ResolutionMultiplier;
-int2			OutputResolution;
-float2			OutputResolutionf;
 uint			HairMaterialId;
 uint			HairStrandsVF_bIsCullingEnable;
 float			HairStrandsVF_Radius;
@@ -82,8 +92,6 @@ Buffer<float4>	HairStrandsVF_PositionOffsetBuffer;
 
 Texture2D<float> SceneDepthTexture;
 
-int2			TileRes;
-
 uint 			VertexCount;
 uint 			VertexStart;
 float			SampleWeight;
@@ -129,6 +137,8 @@ void CalcScreenPosAndRad(in uint4 PB, in float3 PBO, out float4 SP, out float Ra
 	Rad = CP.WorldRadius*2000.0; // OutputResolutionf.x; //TODO: figure this out correctly?
 }
 
+#endif // Common rasetrizer helper function & parameters
+
 ///////////////////////////////////////////////////////////////////////////
 
 #if SHADER_RASTERCOMPUTE_DEPTH_GRID
@@ -660,3 +670,130 @@ void RasterCS(uint DispatchThreadID : SV_DispatchThreadID, uint GroupThreadID :
 
 ///////////////////////////////////////////////////////////////////////////
 
+#if SHADER_RASTERCOMPUTE_DEBUG
+
+#include "../ShaderPrint.ush"
+void PrintBool(inout FShaderPrintContext Ctx, bool In)
+{
+	if (In)
+	{
+		Print(Ctx, TEXT("Yes"), FontGreen);
+	}
+	else
+	{
+		Print(Ctx, TEXT("No"), FontRed);
+	}
+}
+
+Texture2D<uint>						VisTileDepthGrid;
+Texture3D<uint> 					VisTileBinningGrid;
+//Buffer<uint>						OutVisTilePrims;
+Buffer<uint>						VisTileArgs;
+//ByteAddressBuffer					OutVisTileData;
+
+uint MacroGroupId;
+uint PrimitiveInfoIndex;
+uint TotalPrimitiveInfoCount;
+
+#define TilePrintOffset (TileSize >> 1)
+
+float4 Transparent(float4 Color) { return float4(Color.xyz, 0.5f); }
+
+uint GetTileTotalSegment(uint2 TileCoord, bool bPrintDetails)
+{
+	const float TileDisplayScale = 1.5f;
+	const uint DisplayTileSize = TileSize * TileDisplayScale;
+	uint2 InlinedTileCoord = uint2(0, 0);
+
+	uint TotalSegments = 0;
+	const uint BinCount = NumBinners;// * 2; // Each binner fill in 2 bins, see binning algo.
+	for (uint BinIt = 0; BinIt < BinCount; ++BinIt)
+	{
+		const uint CurrTileSegments = VisTileBinningGrid.Load(uint4(TileCoord, BinIt, 0));
+		TotalSegments += CurrTileSegments;
+
+		if (bPrintDetails)
+		{
+			AddFilledQuadSS(InlinedTileCoord * DisplayTileSize, (InlinedTileCoord + 1) * DisplayTileSize, CurrTileSegments > 0 ? Transparent(ColorLightGreen) : Transparent(ColorLightRed));
+			AddQuadSS(InlinedTileCoord * DisplayTileSize, (InlinedTileCoord + 1) * DisplayTileSize, ColorYellow);
+
+			FShaderPrintContext Context = InitShaderPrintContext(true, InlinedTileCoord * DisplayTileSize + TilePrintOffset);
+			Print(Context, CurrTileSegments, FontWhite);
+			++InlinedTileCoord.x;
+
+			// Span details onto 2 lines
+			if (BinIt == NumBinners-1)
+			{
+				InlinedTileCoord.x = 0;
+				++InlinedTileCoord.y;
+			}
+		}
+	}
+	return TotalSegments;
+}
+
+void PrintTile(uint2 TileCoord, uint TotalSegments, bool bPrintText)
+{
+	AddFilledQuadSS(TileCoord * TileSize, (TileCoord + 1) * TileSize, TotalSegments > 0 ? Transparent(ColorLightGreen) : Transparent(ColorLightRed));
+	if (bPrintText)
+	{
+		FShaderPrintContext Context = InitShaderPrintContext(true, TileCoord * TileSize + uint2(0, TileSize * 1.5f));
+		Print(Context, TotalSegments, FontWhite);
+
+		AddQuadSS(TileCoord * TileSize, (TileCoord + 1) * TileSize, ColorYellow);
+	}
+}
+
+[numthreads(8, 8, 1)]
+void MainCS(uint3 ThreadId : SV_DispatchThreadID)
+{
+	// Info/Stats
+	if (all(ThreadId == 0))
+	{
+		FShaderPrintContext Context = InitShaderPrintContext(true, uint2(50, 110));
+		Print(Context, TEXT("Raster compute         "), FontYellow); Newline(Context);
+		Print(Context, TEXT("Macro Group Id       : "), FontSilver); Print(Context, MacroGroupId, FontWhite); Newline(Context);
+		Print(Context, TEXT("Primitive Info       : "), FontSilver); Print(Context, PrimitiveInfoIndex, FontWhite, 2, 0); Print(Context, TEXT("/"), FontSilver); Print(Context, TotalPrimitiveInfoCount, FontWhite, 2, 0); Newline(Context);
+		Newline(Context);
+
+		Print(Context, TEXT("Configuration          "), FontYellow); Newline(Context);
+		Print(Context, TEXT("Output Resolution    : "), FontSilver); Print(Context, OutputResolution, FontWhite); Newline(Context);
+		Print(Context, TEXT("Resolution Multiplier: "), FontSilver); Print(Context, ResolutionMultiplier, FontWhite); Newline(Context);
+		Newline(Context);
+		
+		Print(Context, TEXT("Tile Size            : "), FontSilver); Print(Context, TileSize, FontWhite); Newline(Context);
+		Print(Context, TEXT("Tile Res             : "), FontSilver); Print(Context, TileRes.x, FontWhite, 2, 0); Print(Context, TEXT("x"), FontSilver); Print(Context, TileRes.y, FontWhite, 2, 0); Newline(Context);
+		Newline(Context);
+		
+		Print(Context, TEXT("Num Binners          : "), FontSilver); Print(Context, NumBinners, FontWhite); Newline(Context);
+		Print(Context, TEXT("Num Rasterizers      : "), FontSilver); Print(Context, NumRasterizers, FontWhite); Newline(Context);
+		Print(Context, TEXT("Max Raster Count     : "), FontSilver); Print(Context, MaxRasterCount, FontWhite); Newline(Context);
+		Newline(Context);
+
+		Print(Context, TEXT("Allocated Tile Count : "), FontSilver); Print(Context, VisTileArgs[0], FontWhite); Newline(Context);
+		
+
+	}
+
+	// Cursor info	
+	if (all(ThreadId == 0) && all(ShaderPrintData.CursorCoord >= 0))
+	{
+		const uint2 PixelCoord = ShaderPrintData.CursorCoord;
+		const uint2 TileCoord  = PixelCoord >> TileSizeAsShift;
+	
+		const uint TotalSegments = GetTileTotalSegment(TileCoord, true);
+		PrintTile(TileCoord, TotalSegments, true);
+	}
+
+	// All tile
+	{
+		const uint2 TileCoord = ThreadId.xy;
+		const uint TotalSegments = GetTileTotalSegment(TileCoord, false);
+		if (TotalSegments)
+		{
+			PrintTile(TileCoord, TotalSegments, false);
+		}
+	}
+}
+#endif //SHADER_RASTERCOMPUTE_RASTER
+
diff --git a/Engine/Shaders/Private/ShaderPrint.ush b/Engine/Shaders/Private/ShaderPrint.ush
--- a/Engine/Shaders/Private/ShaderPrint.ush
+++ b/Engine/Shaders/Private/ShaderPrint.ush
@@ -1168,6 +1168,30 @@ void AddQuadTWS(FShaderPrintContext Ctx, float3 Pos0, float3 Pos1, float3 Pos2,
 void AddQuadWS (FShaderPrintContext Ctx, float3 Pos0, float3 Pos1, float3 Pos2, float3 Pos3, float4 Color)	{ if (Ctx.bIsActive) { AddQuadTWS(Pos0+ShaderPrintData.TranslatedWorldOffset, Pos1+ShaderPrintData.TranslatedWorldOffset, Pos2+ShaderPrintData.TranslatedWorldOffset, Pos3+ShaderPrintData.TranslatedWorldOffset, Color); } }
 void AddQuad   (FShaderPrintContext Ctx, float3 Pos0, float3 Pos1, float3 Pos2, float3 Pos3, float4 Color)	{ if (Ctx.bIsActive) { AddQuadTWS(Pos0+ShaderPrintData.TranslatedWorldOffset, Pos1+ShaderPrintData.TranslatedWorldOffset, Pos2+ShaderPrintData.TranslatedWorldOffset, Pos3+ShaderPrintData.TranslatedWorldOffset, Color); } }
 
+// Draw 2D quad line
+void AddQuadSS(float2 MinPos, float2 MaxPos, float4 Color)
+{
+	uint Offset = 0;
+	if (AllocateLineElement(4, Offset))
+	{
+		MinPos = float2(MinPos) / float2(ShaderPrintData.Resolution);
+		MaxPos = float2(MaxPos) / float2(ShaderPrintData.Resolution);
+
+		float3 Pos0 = float3(MinPos.x, MinPos.y, 0);
+		float3 Pos1 = float3(MaxPos.x, MinPos.y, 0);
+		float3 Pos2 = float3(MaxPos.x, MaxPos.y, 0);
+		float3 Pos3 = float3(MinPos.x, MaxPos.y, 0);
+
+		FLineElement Element;
+		Element.bIsScreenSpace = true;
+		Element.Color0 = Element.Color1 = Color;
+		Element.Pos0 = Pos0; Element.Pos1 = Pos1; AddLineElement(Element, Offset + 0);
+		Element.Pos0 = Pos1; Element.Pos1 = Pos2; AddLineElement(Element, Offset + 1);
+		Element.Pos0 = Pos2; Element.Pos1 = Pos3; AddLineElement(Element, Offset + 2);
+		Element.Pos0 = Pos3; Element.Pos1 = Pos0; AddLineElement(Element, Offset + 3);
+	}
+}
+
 ///////////////////////////////////////////////////////////////////////////////////////////////////
 // AABB
 
diff --git a/Engine/Source/Runtime/Renderer/Private/HairStrands/HairStrandsVisibility.cpp b/Engine/Source/Runtime/Renderer/Private/HairStrands/HairStrandsVisibility.cpp
--- a/Engine/Source/Runtime/Renderer/Private/HairStrands/HairStrandsVisibility.cpp
+++ b/Engine/Source/Runtime/Renderer/Private/HairStrands/HairStrandsVisibility.cpp
@@ -59,12 +59,14 @@ int32 GHairVisibilityComputeRaster_MaxTiles = 8192;
 int32 GHairVisibilityComputeRaster_TileSize = 32;
 int32 GHairVisibilityComputeRaster_NumBinners = 32;
 int32 GHairVisibilityComputeRaster_NumRasterizers = 256;
+int32 GHairVisibilityComputeRaster_Debug = 0;
 
 static FAutoConsoleVariableRef CVarHairVisibilityComputeRaster_Culling(TEXT("r.HairStrands.Visibility.ComputeRaster.Culling"), GHairVisibilityComputeRaster_Culling, TEXT("Use culling buffers with compute rasterization."));
 static FAutoConsoleVariableRef CVarHairVisibilityComputeRaster_MaxTiles(TEXT("r.HairStrands.Visibility.ComputeRaster.MaxTiles"), GHairVisibilityComputeRaster_MaxTiles, TEXT("Maximum number of tiles used for compute rasterization. 8192 is default"));
 static FAutoConsoleVariableRef CVarHairVisibilityComputeRaster_TileSize(TEXT("r.HairStrands.Visibility.ComputeRaster.TileSize"), GHairVisibilityComputeRaster_TileSize, TEXT("Tile size used for compute rasterization. Experimental - only size of 32 currently supported"));
 static FAutoConsoleVariableRef CVarHairVisibilityComputeRaster_NumBinners(TEXT("r.HairStrands.Visibility.ComputeRaster.NumBinners"), GHairVisibilityComputeRaster_NumBinners, TEXT("Number of Binners used in Binning compute rasterization pass. 32 is default"));
 static FAutoConsoleVariableRef CVarHairVisibilityComputeRaster_NumRasterizers(TEXT("r.HairStrands.Visibility.ComputeRaster.NumRasterizers"), GHairVisibilityComputeRaster_NumRasterizers, TEXT("Number of Rasterizers used compute rasterization. 256 is default"));
+static FAutoConsoleVariableRef CVarHairVisibilityComputeRaster_Debug(TEXT("r.HairStrands.Visibility.ComputeRaster.Debug"), GHairVisibilityComputeRaster_Debug, TEXT("Debug compute raster output"));
 
 static float GHairStrandsFullCoverageThreshold = 0.98f;
 static FAutoConsoleVariableRef CVarHairStrandsFullCoverageThreshold(TEXT("r.HairStrands.Visibility.FullCoverageThreshold"), GHairStrandsFullCoverageThreshold, TEXT("Define the coverage threshold at which a pixel is considered fully covered."));
@@ -2804,7 +2806,7 @@ static FRDGTextureRef AddHairHairCountToTransmittancePass(
 	FHairCountToCoverageCS::FPermutationDomain PermutationVector;
 	PermutationVector.Set<FHairCountToCoverageCS::FInputType>(bUseOneChannel ? 1 : 0);
 	TShaderMapRef<FHairCountToCoverageCS> ComputeShader(ViewInfo.ShaderMap, PermutationVector);
-	FComputeShaderUtils::AddPass(GraphBuilder, RDG_EVENT_NAME("HairCountToTransmittancePass"), ComputeShader, PassParameters, FComputeShaderUtils::GetGroupCount(OutputResolution, FIntPoint(8,8)));
+	FComputeShaderUtils::AddPass(GraphBuilder, RDG_EVENT_NAME("HairStrands::HairCountToTransmittancePass"), ComputeShader, PassParameters, FComputeShaderUtils::GetGroupCount(OutputResolution, FIntPoint(8,8)));
 
 	return OutputTexture;
 }
@@ -2860,6 +2862,27 @@ void AddMeshDrawTransitionPass(
 
 	ExternalAccessQueue.Submit(GraphBuilder);
 }
+
+///////////////////////////////////////////////////////////////////////////////////////////////////
+// Raster compute common parameters
+
+BEGIN_SHADER_PARAMETER_STRUCT(FRasterComputeCommonParameters, )
+	SHADER_PARAMETER(uint32, MaxRasterCount)
+	SHADER_PARAMETER(uint32, FrameIdMod8)
+	SHADER_PARAMETER(uint32, ResolutionMultiplier)
+	SHADER_PARAMETER(FIntPoint, OutputResolution)
+	SHADER_PARAMETER(FVector2f, OutputResolutionf)
+	SHADER_PARAMETER(uint32, TileSize)
+	SHADER_PARAMETER(float, RcpTileSize)
+	SHADER_PARAMETER(uint32, TileSizeAsShift)
+	SHADER_PARAMETER(uint32, SqrTileSize)
+	SHADER_PARAMETER(uint32, NumRasterizers)
+	SHADER_PARAMETER(float, RcpNumRasterizers)
+	SHADER_PARAMETER(FIntPoint, TileRes)
+	SHADER_PARAMETER(uint32, NumBinners)
+	SHADER_PARAMETER(float, RcpNumBinners)
+END_SHADER_PARAMETER_STRUCT()
+
 ///////////////////////////////////////////////////////////////////////////////////////////////////
 
 class FVisiblityRasterComputePrepareDepthGridCS : public FGlobalShader
@@ -2868,15 +2891,11 @@ class FVisiblityRasterComputePrepareDepthGridCS : public FGlobalShader
 	SHADER_USE_PARAMETER_STRUCT(FVisiblityRasterComputePrepareDepthGridCS, FGlobalShader);
 
 	BEGIN_SHADER_PARAMETER_STRUCT(FParameters, )
+		SHADER_PARAMETER_STRUCT_INCLUDE(FRasterComputeCommonParameters, Common)
 		SHADER_PARAMETER_RDG_TEXTURE(Texture2D, SceneDepthTexture)
-		SHADER_PARAMETER(FIntPoint, OutputResolution)
 		SHADER_PARAMETER_RDG_TEXTURE_UAV(RWTexture2D, OutVisTileDepthGrid)
 		SHADER_PARAMETER_RDG_TEXTURE_UAV(RWTexture2D, OutDepthCovTexture)
-		SHADER_PARAMETER(FIntPoint, TileRes)
-		SHADER_PARAMETER(uint32, TileSize)
-		SHADER_PARAMETER(float, RcpTileSize)
-		SHADER_PARAMETER(uint32, SqrTileSize)
-		END_SHADER_PARAMETER_STRUCT()
+	END_SHADER_PARAMETER_STRUCT()
 
 public:
 	static bool ShouldCompilePermutation(const FGlobalShaderPermutationParameters& Parameters) { return IsHairStrandsSupported(EHairStrandsShaderType::Strands, Parameters.Platform); }
@@ -2889,6 +2908,8 @@ public:
 
 IMPLEMENT_GLOBAL_SHADER(FVisiblityRasterComputePrepareDepthGridCS, "/Engine/Private/HairStrands/HairStrandsVisibilityRasterCompute.usf", "PrepareDepthGridCS", SF_Compute);
 
+/////////////////////////////////////////////////////////////////////////////////////////
+
 class FVisiblityRasterComputeBinningCS : public FGlobalShader
 {
 	DECLARE_GLOBAL_SHADER(FVisiblityRasterComputeBinningCS);
@@ -2899,13 +2920,10 @@ class FVisiblityRasterComputeBinningCS : public FGlobalShader
 	using FPermutationDomain = TShaderPermutationDomain<FGroupSize, FCulling>;
 
 	BEGIN_SHADER_PARAMETER_STRUCT(FParameters, )
+		SHADER_PARAMETER_STRUCT_INCLUDE(FRasterComputeCommonParameters, Common)
 		SHADER_PARAMETER(uint32, MacroGroupId)
 		SHADER_PARAMETER(uint32, DispatchCountX)
-		SHADER_PARAMETER(uint32, MaxRasterCount)
-		SHADER_PARAMETER(uint32, FrameIdMod8)
 		SHADER_PARAMETER(uint32, HairMaterialId)
-		SHADER_PARAMETER(uint32, ResolutionMultiplier)
-		SHADER_PARAMETER(FIntPoint, OutputResolution)
 		SHADER_PARAMETER(uint32, HairStrandsVF_bIsCullingEnable)
 		SHADER_PARAMETER(float, HairStrandsVF_Density)
 		SHADER_PARAMETER(float, HairStrandsVF_Radius)
@@ -2932,12 +2950,6 @@ class FVisiblityRasterComputeBinningCS : public FGlobalShader
 		RDG_TEXTURE_ACCESS(VisTileBinningGridTex, ERHIAccess::UAVCompute)
 		SHADER_PARAMETER(uint32, VertexCount)
 		SHADER_PARAMETER(uint32, VertexStart)
-		SHADER_PARAMETER(FIntPoint, TileRes)
-		SHADER_PARAMETER(uint32, TileSize)
-		SHADER_PARAMETER(float, RcpTileSize)
-		SHADER_PARAMETER(uint32, SqrTileSize)
-		SHADER_PARAMETER(uint32, NumBinners)
-		SHADER_PARAMETER(float, RcpNumBinners)
 	END_SHADER_PARAMETER_STRUCT()
 
 public:
@@ -2951,6 +2963,8 @@ public:
 
 IMPLEMENT_GLOBAL_SHADER(FVisiblityRasterComputeBinningCS, "/Engine/Private/HairStrands/HairStrandsVisibilityRasterCompute.usf", "BinningCS", SF_Compute);
 
+/////////////////////////////////////////////////////////////////////////////////////////
+
 class FVisiblityRasterComputeRasterizeCS : public FGlobalShader
 {
 	DECLARE_GLOBAL_SHADER(FVisiblityRasterComputeRasterizeCS);
@@ -2960,14 +2974,10 @@ class FVisiblityRasterComputeRasterizeCS : public FGlobalShader
 	using FPermutationDomain = TShaderPermutationDomain<FGroupSize>;
 
 	BEGIN_SHADER_PARAMETER_STRUCT(FParameters, )
+		SHADER_PARAMETER_STRUCT_INCLUDE(FRasterComputeCommonParameters, Common)
 		SHADER_PARAMETER(uint32, MacroGroupId)
 		SHADER_PARAMETER(uint32, DispatchCountX)
-		SHADER_PARAMETER(uint32, MaxRasterCount)
-		SHADER_PARAMETER(uint32, FrameIdMod8)
 		SHADER_PARAMETER(uint32, HairMaterialId)
-		SHADER_PARAMETER(uint32, ResolutionMultiplier)
-		SHADER_PARAMETER(FIntPoint, OutputResolution)
-		SHADER_PARAMETER(FVector2f, OutputResolutionf)
 		SHADER_PARAMETER(uint32, HairStrandsVF_bIsCullingEnable)
 		SHADER_PARAMETER(float, HairStrandsVF_Density)
 		SHADER_PARAMETER(float, HairStrandsVF_Radius)
@@ -2985,11 +2995,6 @@ class FVisiblityRasterComputeRasterizeCS : public FGlobalShader
 		RDG_BUFFER_ACCESS(IndirectBufferArgs, ERHIAccess::IndirectArgs)
 		SHADER_PARAMETER_RDG_TEXTURE(Texture2D, SceneDepthTexture)
 		SHADER_PARAMETER(float, SampleWeight)
-		SHADER_PARAMETER(uint32, TileSize)
-		SHADER_PARAMETER(float, RcpTileSize)
-		SHADER_PARAMETER(uint32, SqrTileSize)
-		SHADER_PARAMETER(uint32, NumRasterizers)
-		SHADER_PARAMETER(float, RcpNumRasterizers)
 		SHADER_PARAMETER_RDG_TEXTURE_UAV(RWTexture2D, RWVisTileDepthGrid)
 		SHADER_PARAMETER_RDG_BUFFER_SRV(Buffer, VisTilePrims)
 		SHADER_PARAMETER_RDG_BUFFER_SRV(Buffer, VisTileArgs)
@@ -2999,7 +3004,6 @@ class FVisiblityRasterComputeRasterizeCS : public FGlobalShader
 		SHADER_PARAMETER_RDG_TEXTURE_UAV(RWTexture2D, OutDepthCovTexture)
 		SHADER_PARAMETER_RDG_TEXTURE_UAV(RWTexture2D, OutPrimMatTexture)
 		SHADER_PARAMETER_STRUCT_REF(FViewUniformShaderParameters, ViewUniformBuffer)
-		SHADER_PARAMETER(FIntPoint, TileRes)
 	END_SHADER_PARAMETER_STRUCT()
 
 public:
@@ -3013,6 +3017,40 @@ public:
 
 IMPLEMENT_GLOBAL_SHADER(FVisiblityRasterComputeRasterizeCS, "/Engine/Private/HairStrands/HairStrandsVisibilityRasterCompute.usf", "RasterCS", SF_Compute);
 
+/////////////////////////////////////////////////////////////////////////////////////////
+
+class FVisiblityRasterComputeDebugCS : public FGlobalShader
+{
+	DECLARE_GLOBAL_SHADER(FVisiblityRasterComputeDebugCS);
+	SHADER_USE_PARAMETER_STRUCT(FVisiblityRasterComputeDebugCS, FGlobalShader);
+
+	using FPermutationDomain = TShaderPermutationDomain<>;
+
+	BEGIN_SHADER_PARAMETER_STRUCT(FParameters, )
+		SHADER_PARAMETER_STRUCT_INCLUDE(FRasterComputeCommonParameters, Common)
+		SHADER_PARAMETER(uint32, MacroGroupId)
+		SHADER_PARAMETER(uint32, PrimitiveInfoIndex)
+		SHADER_PARAMETER(uint32, TotalPrimitiveInfoCount)
+		SHADER_PARAMETER_RDG_TEXTURE(Texture2D, VisTileDepthGrid)
+		SHADER_PARAMETER_RDG_TEXTURE(Texture3D, VisTileBinningGrid)
+		SHADER_PARAMETER_RDG_BUFFER_SRV(Buffer, VisTileArgs)
+		SHADER_PARAMETER_STRUCT_REF(FViewUniformShaderParameters, ViewUniformBuffer)
+		SHADER_PARAMETER_STRUCT_INCLUDE(ShaderPrint::FShaderParameters, ShaderPrintParameters)
+	END_SHADER_PARAMETER_STRUCT()
+
+public:
+	static bool ShouldCompilePermutation(const FGlobalShaderPermutationParameters& Parameters) { return IsHairStrandsSupported(EHairStrandsShaderType::Strands, Parameters.Platform); }
+	static void ModifyCompilationEnvironment(const FGlobalShaderPermutationParameters& Parameters, FShaderCompilerEnvironment& OutEnvironment)
+	{
+		FGlobalShader::ModifyCompilationEnvironment(Parameters, OutEnvironment);
+		OutEnvironment.SetDefine(TEXT("SHADER_RASTERCOMPUTE_DEBUG"), 1);
+	}
+};
+
+IMPLEMENT_GLOBAL_SHADER(FVisiblityRasterComputeDebugCS, "/Engine/Private/HairStrands/HairStrandsVisibilityRasterCompute.usf", "MainCS", SF_Compute);
+
+/////////////////////////////////////////////////////////////////////////////////////////
+
 bool IsHairStrandContinuousDecimationReorderingEnabled();
 
 static FRasterComputeOutput AddVisibilityComputeRasterPass(
@@ -3116,19 +3154,52 @@ static FRasterComputeOutput AddVisibilityComputeRasterPass(
 	RasterPermutationVector.Set<FVisiblityRasterComputeRasterizeCS::FGroupSize>(GroupSize);
 	TShaderMapRef<FVisiblityRasterComputeRasterizeCS> ComputeShaderRaster(ViewInfo.ShaderMap, RasterPermutationVector);
 
-	TShaderMapRef<FVisiblityRasterComputePrepareDepthGridCS> ComputeShaderPrepareDepthGrid(ViewInfo.ShaderMap);	
-	FVisiblityRasterComputePrepareDepthGridCS::FParameters* PrepDepthGridParameters = GraphBuilder.AllocParameters<FVisiblityRasterComputePrepareDepthGridCS::FParameters>();
-	PrepDepthGridParameters->OutputResolution = Out.SuperResolution;
-	PrepDepthGridParameters->SceneDepthTexture = SceneDepthTexture;
-	PrepDepthGridParameters->OutVisTileDepthGrid = OutVisTileDepthGridUAV; 
-	PrepDepthGridParameters->OutDepthCovTexture = DepthCovTextureUAV;
-	PrepDepthGridParameters->TileRes = TileGridRes;
-	PrepDepthGridParameters->TileSize = TileSize;
-	PrepDepthGridParameters->RcpTileSize = 1.0 / TileSize;
-	PrepDepthGridParameters->SqrTileSize = TileSize * TileSize;
+	// Common parameters
+	FRasterComputeCommonParameters Common;
+	{		
+		Common.MaxRasterCount = TileSize; //TODO: implement tile clipping in raster pass to avoid exisiting edge case
+		Common.FrameIdMod8 =
+		Common.ResolutionMultiplier = Out.ResolutionMultiplier;
+		Common.OutputResolution = Out.SuperResolution;
+		Common.OutputResolutionf = FVector2f(Out.SuperResolution.X, Out.SuperResolution.Y);
+		Common.TileSize = TileSize;
+		Common.TileSizeAsShift = uint32(FMath::Log2(float(TileSize)));
+		Common.RcpTileSize = 1.0 / TileSize;
+		Common.SqrTileSize = TileSize * TileSize;
+		Common.NumRasterizers = NumRasterizers;
+		Common.RcpNumRasterizers = 1.0 / NumRasterizers;
+		Common.TileRes = TileGridRes;
+		Common.NumBinners = NumBinners;
+		Common.RcpNumBinners = 1.0 / NumBinners;
+	}
 
-	FComputeShaderUtils::AddPass(GraphBuilder, RDG_EVENT_NAME("HairStrandsVisibilityComputeRasterPrepDepthGrid"), ComputeShaderPrepareDepthGrid, PrepDepthGridParameters, FIntVector(TileGridRes.X*TileGridRes.Y, 1, 1));
+	// Fill in tile depth
+	{
+		TShaderMapRef<FVisiblityRasterComputePrepareDepthGridCS> ComputeShaderPrepareDepthGrid(ViewInfo.ShaderMap);
+		FVisiblityRasterComputePrepareDepthGridCS::FParameters* PrepDepthGridParameters = GraphBuilder.AllocParameters<FVisiblityRasterComputePrepareDepthGridCS::FParameters>();
+		
+		PrepDepthGridParameters->Common = Common;
+		PrepDepthGridParameters->SceneDepthTexture = SceneDepthTexture;
+		PrepDepthGridParameters->OutVisTileDepthGrid = OutVisTileDepthGridUAV;
+		PrepDepthGridParameters->OutDepthCovTexture = DepthCovTextureUAV;
+		FComputeShaderUtils::AddPass(GraphBuilder, RDG_EVENT_NAME("HairStrands::VisibilityComputeRasterPrepDepthGrid"), ComputeShaderPrepareDepthGrid, PrepDepthGridParameters, FIntVector(TileGridRes.X* TileGridRes.Y, 1, 1));
+	}
 
+	const bool bDebugEnabled = GHairVisibilityComputeRaster_Debug > 0;
+	uint32 TotalPrimitiveInfoCount = 0;
+	if (bDebugEnabled)
+	{
+		ShaderPrint::SetEnabled(true);
+		ShaderPrint::RequestSpaceForTriangles(2 * NumBinners * Common.TileRes.X * Common.TileRes.Y + 2 * NumBinners * 10);
+		ShaderPrint::RequestSpaceForLines(4 * 2 * NumBinners * 10);
+		for (const FHairStrandsMacroGroupData& MacroGroup : MacroGroupDatas)
+		{
+			const FHairStrandsMacroGroupData::TPrimitiveInfos& PrimitiveSceneInfos = MacroGroup.PrimitivesInfos;
+			TotalPrimitiveInfoCount += PrimitiveSceneInfos.Num();
+		}
+	}
+
+	uint32 PrimitiveInfoIndex = 0;
 	for (const FHairStrandsMacroGroupData& MacroGroup : MacroGroupDatas)
 	{
 		const FHairStrandsMacroGroupData::TPrimitiveInfos& PrimitiveSceneInfos = MacroGroup.PrimitivesInfos;
@@ -3152,101 +3223,111 @@ static FRasterComputeOutput AddVisibilityComputeRasterPass(
 			const uint32 VertexCount = HairGroupPublicData->GetActiveStrandsVertexCount(VFInput.Strands.VertexCount, ScreenSize);
 			const float SampleWeight = HairGroupPublicData->GetActiveStrandsSampleWeight(false, ScreenSize);
 
-			//reset buffers
-			uint32 IndexGridClearValues[4] = { 0x0,0x0,0x0,0x0 };
-			AddClearUAVPass(GraphBuilder, OutVisTileBinningGridUAV, IndexGridClearValues);
-			AddClearUAVPass(GraphBuilder, OutVisTileArgsUAV, 0u);
-
-			//binning
-			FVisiblityRasterComputeBinningCS::FParameters* BinningParameters = GraphBuilder.AllocParameters<FVisiblityRasterComputeBinningCS::FParameters>();
-			BinningParameters->OutputResolution = Out.SuperResolution;
-			BinningParameters->ResolutionMultiplier = Out.ResolutionMultiplier;
-			BinningParameters->MacroGroupId = MacroGroup.MacroGroupId;
-			BinningParameters->HairMaterialId = PrimitiveInfo.MaterialId;
-			BinningParameters->ViewUniformBuffer = ViewUniformShaderParameters;
-			BinningParameters->SceneDepthTexture = SceneDepthTexture;
-			BinningParameters->OutVisTilePrims = OutVisTilePrimsUAV;
-			BinningParameters->OutVisTileBinningGrid = OutVisTileBinningGridUAV;  
-			BinningParameters->VisTileDepthGrid = OutVisTileDepthGrid;
-			BinningParameters->OutVisTileArgs = OutVisTileArgsUAV;
-			BinningParameters->OutVisTileData = OutVisTileDataUAV;
-
-			BinningParameters->HairStrandsVF_PositionBuffer = VFInput.Strands.PositionBuffer.SRV;
-			BinningParameters->HairStrandsVF_PositionOffsetBuffer = VFInput.Strands.PositionOffsetBuffer.SRV;
-			BinningParameters->HairStrandsVF_VertexCount = VFInput.Strands.VertexCount;
-			BinningParameters->HairStrandsVF_Radius = VFInput.Strands.HairRadius;
-			BinningParameters->HairStrandsVF_RootScale = VFInput.Strands.HairRootScale;
-			BinningParameters->HairStrandsVF_TipScale = VFInput.Strands.HairTipScale;
-			BinningParameters->HairStrandsVF_Length = VFInput.Strands.HairLength;
-			BinningParameters->HairStrandsVF_bUseStableRasterization = VFInput.Strands.bUseStableRasterization ? 1 : 0;
-			BinningParameters->HairStrandsVF_Density = VFInput.Strands.HairDensity;
-			BinningParameters->HairStrandsVF_LocalToWorldPrimitiveTransform = FMatrix44f(LocalToWorldPrimitiveTransform); // LWC_TODO: Precision loss
-
-			BinningParameters->VertexCount = VertexCount;
-			BinningParameters->VertexStart = VertexStart;
-			BinningParameters->TileRes = TileGridRes;
-			BinningParameters->VisTileBinningGridTex = OutVisTileBinningGrid;
-			BinningParameters->TileSize = TileSize;
-			BinningParameters->RcpTileSize = 1.0/TileSize;
-			BinningParameters->SqrTileSize = TileSize*TileSize;
-			BinningParameters->NumBinners = NumBinners;
-			BinningParameters->RcpNumBinners = 1.0/NumBinners;
-
-			if (bCullingEnable)
+			// Reset buffers
 			{
-				FRDGImportedBuffer CullingIndirectBuffer = Register(GraphBuilder, HairGroupPublicData->GetDrawIndirectRasterComputeBuffer(), ERDGImportedBufferFlags::CreateSRV);
-				BinningParameters->HairStrandsVF_CullingIndirectBuffer = CullingIndirectBuffer.SRV;
-				BinningParameters->HairStrandsVF_bIsCullingEnable = bCullingEnable ? 1 : 0;
-				BinningParameters->HairStrandsVF_CullingIndexBuffer = RegisterAsSRV(GraphBuilder, HairGroupPublicData->GetCulledVertexIdBuffer());
-				BinningParameters->HairStrandsVF_CullingRadiusScaleBuffer = RegisterAsSRV(GraphBuilder, HairGroupPublicData->GetCulledVertexRadiusScaleBuffer());
-				BinningParameters->IndirectBufferArgs = CullingIndirectBuffer.Buffer;
+				uint32 IndexGridClearValues[4] = { 0x0,0x0,0x0,0x0 };
+				AddClearUAVPass(GraphBuilder, OutVisTileBinningGridUAV, IndexGridClearValues);
+				AddClearUAVPass(GraphBuilder, OutVisTileArgsUAV, 0u);
 			}
 
-			FComputeShaderUtils::AddPass(GraphBuilder, RDG_EVENT_NAME("HairStrandsVisibilityComputeRasterBinning(culling=%s)", bCullingEnable ? TEXT("On") : TEXT("Off")), bCullingEnable ? ComputeShaderBinning_CullingOn : ComputeShaderBinning_CullingOff, BinningParameters, FIntVector(NumBinners, 1, 1));
+			// Binning pass
+			{
+				FVisiblityRasterComputeBinningCS::FParameters* BinningParameters = GraphBuilder.AllocParameters<FVisiblityRasterComputeBinningCS::FParameters>();
+				BinningParameters->Common = Common;
+				BinningParameters->MacroGroupId = MacroGroup.MacroGroupId;
+				BinningParameters->HairMaterialId = PrimitiveInfo.MaterialId;
+				BinningParameters->ViewUniformBuffer = ViewUniformShaderParameters;
+				BinningParameters->SceneDepthTexture = SceneDepthTexture;
+				BinningParameters->OutVisTilePrims = OutVisTilePrimsUAV;
+				BinningParameters->OutVisTileBinningGrid = OutVisTileBinningGridUAV;  
+				BinningParameters->VisTileDepthGrid = OutVisTileDepthGrid;
+				BinningParameters->OutVisTileArgs = OutVisTileArgsUAV;
+				BinningParameters->OutVisTileData = OutVisTileDataUAV;
+
+				BinningParameters->HairStrandsVF_PositionBuffer = VFInput.Strands.PositionBuffer.SRV;
+				BinningParameters->HairStrandsVF_PositionOffsetBuffer = VFInput.Strands.PositionOffsetBuffer.SRV;
+				BinningParameters->HairStrandsVF_VertexCount = VFInput.Strands.VertexCount;
+				BinningParameters->HairStrandsVF_Radius = VFInput.Strands.HairRadius;
+				BinningParameters->HairStrandsVF_RootScale = VFInput.Strands.HairRootScale;
+				BinningParameters->HairStrandsVF_TipScale = VFInput.Strands.HairTipScale;
+				BinningParameters->HairStrandsVF_Length = VFInput.Strands.HairLength;
+				BinningParameters->HairStrandsVF_bUseStableRasterization = VFInput.Strands.bUseStableRasterization ? 1 : 0;
+				BinningParameters->HairStrandsVF_Density = VFInput.Strands.HairDensity;
+				BinningParameters->HairStrandsVF_LocalToWorldPrimitiveTransform = FMatrix44f(LocalToWorldPrimitiveTransform); // LWC_TODO: Precision loss
+
+				BinningParameters->VertexCount = VertexCount;
+				BinningParameters->VertexStart = VertexStart;
+				
+				BinningParameters->VisTileBinningGridTex = OutVisTileBinningGrid;
+
+				if (bCullingEnable)
+				{
+					FRDGImportedBuffer CullingIndirectBuffer = Register(GraphBuilder, HairGroupPublicData->GetDrawIndirectRasterComputeBuffer(), ERDGImportedBufferFlags::CreateSRV);
+					BinningParameters->HairStrandsVF_CullingIndirectBuffer = CullingIndirectBuffer.SRV;
+					BinningParameters->HairStrandsVF_bIsCullingEnable = bCullingEnable ? 1 : 0;
+					BinningParameters->HairStrandsVF_CullingIndexBuffer = RegisterAsSRV(GraphBuilder, HairGroupPublicData->GetCulledVertexIdBuffer());
+					BinningParameters->HairStrandsVF_CullingRadiusScaleBuffer = RegisterAsSRV(GraphBuilder, HairGroupPublicData->GetCulledVertexRadiusScaleBuffer());
+					BinningParameters->IndirectBufferArgs = CullingIndirectBuffer.Buffer;
+				}
+
+				FComputeShaderUtils::AddPass(GraphBuilder, RDG_EVENT_NAME("HairStrands::VisibilityComputeRasterBinning(culling=%s)", bCullingEnable ? TEXT("On") : TEXT("Off")), bCullingEnable ? ComputeShaderBinning_CullingOn : ComputeShaderBinning_CullingOff, BinningParameters, FIntVector(NumBinners, 1, 1));
+			}
 			
-			//raster pass
-			FVisiblityRasterComputeRasterizeCS::FParameters* RasterParameters = GraphBuilder.AllocParameters<FVisiblityRasterComputeRasterizeCS::FParameters>();
-			RasterParameters->OutputResolution = Out.SuperResolution;
-			RasterParameters->OutputResolutionf = FVector2f(Out.SuperResolution.X, Out.SuperResolution.Y);
-			RasterParameters->ResolutionMultiplier = Out.ResolutionMultiplier;
-			RasterParameters->MacroGroupId = MacroGroup.MacroGroupId;
-			RasterParameters->MaxRasterCount = TileSize; //TODO: implement tile clipping in raster pass to avoid exisiting edge case
-			RasterParameters->HairMaterialId = PrimitiveInfo.MaterialId;
-			RasterParameters->ViewUniformBuffer = ViewUniformShaderParameters;
-			RasterParameters->SceneDepthTexture = SceneDepthTexture;
-			RasterParameters->OutHairCountTexture = HairCountTextureUAV;
-			RasterParameters->OutDepthCovTexture = DepthCovTextureUAV;
-			RasterParameters->OutPrimMatTexture = PrimMatTextureUAV;
-
-			RasterParameters->HairStrandsVF_PositionBuffer	= VFInput.Strands.PositionBuffer.SRV;
-			RasterParameters->HairStrandsVF_PositionOffsetBuffer = VFInput.Strands.PositionOffsetBuffer.SRV;
-			RasterParameters->HairStrandsVF_VertexCount		= VFInput.Strands.VertexCount;
-			RasterParameters->HairStrandsVF_Radius			= VFInput.Strands.HairRadius;
-			RasterParameters->HairStrandsVF_RootScale			= VFInput.Strands.HairRootScale;
-			RasterParameters->HairStrandsVF_TipScale			= VFInput.Strands.HairTipScale;
-			RasterParameters->HairStrandsVF_Length			= VFInput.Strands.HairLength;
-			RasterParameters->HairStrandsVF_bUseStableRasterization = VFInput.Strands.bUseStableRasterization ? 1 : 0;
-			RasterParameters->HairStrandsVF_Density			= VFInput.Strands.HairDensity;
-			RasterParameters->HairStrandsVF_LocalToWorldPrimitiveTransform = FMatrix44f(LocalToWorldPrimitiveTransform); // LWC_TODO: Precision loss
-
-			RasterParameters->VisTilePrims = GraphBuilder.CreateSRV(OutVisTilePrims);
-			RasterParameters->RWVisTileDepthGrid = OutVisTileDepthGridUAV;
-			RasterParameters->VisTileArgs = GraphBuilder.CreateSRV(OutVisTileArgs);
-			RasterParameters->RWVisTileData = OutVisTileDataUAV;
-			RasterParameters->VisTileIndirect = GraphBuilder.CreateSRV(OutVisTileIndirect);
-			RasterParameters->IndirectBufferArgs = OutVisTileIndirect;
-			RasterParameters->TileRes = TileGridRes;
-			RasterParameters->SampleWeight = SampleWeight;
-			RasterParameters->TileSize = TileSize;
-			RasterParameters->RcpTileSize = 1.0 / TileSize;
-			RasterParameters->SqrTileSize = TileSize * TileSize;
-			RasterParameters->NumRasterizers = NumRasterizers;
-			RasterParameters->RcpNumRasterizers = 1.0 / NumRasterizers;
-
-			FComputeShaderUtils::AddPass(GraphBuilder, RDG_EVENT_NAME("HairStrandsVisibilityComputeRasterRaster(tiled)"), ComputeShaderRaster, RasterParameters, FIntVector(NumRasterizers, 1, 1));
+			// Raster pass
+			{
+				FVisiblityRasterComputeRasterizeCS::FParameters* RasterParameters = GraphBuilder.AllocParameters<FVisiblityRasterComputeRasterizeCS::FParameters>();
+				RasterParameters->Common = Common;
+				RasterParameters->MacroGroupId = MacroGroup.MacroGroupId;
+				RasterParameters->HairMaterialId = PrimitiveInfo.MaterialId;
+				RasterParameters->ViewUniformBuffer = ViewUniformShaderParameters;
+				RasterParameters->SceneDepthTexture = SceneDepthTexture;
+				RasterParameters->OutHairCountTexture = HairCountTextureUAV;
+				RasterParameters->OutDepthCovTexture = DepthCovTextureUAV;
+				RasterParameters->OutPrimMatTexture = PrimMatTextureUAV;
+				RasterParameters->SampleWeight = SampleWeight;
+
+				RasterParameters->HairStrandsVF_PositionBuffer	= VFInput.Strands.PositionBuffer.SRV;
+				RasterParameters->HairStrandsVF_PositionOffsetBuffer = VFInput.Strands.PositionOffsetBuffer.SRV;
+				RasterParameters->HairStrandsVF_VertexCount		= VFInput.Strands.VertexCount;
+				RasterParameters->HairStrandsVF_Radius			= VFInput.Strands.HairRadius;
+				RasterParameters->HairStrandsVF_RootScale			= VFInput.Strands.HairRootScale;
+				RasterParameters->HairStrandsVF_TipScale			= VFInput.Strands.HairTipScale;
+				RasterParameters->HairStrandsVF_Length			= VFInput.Strands.HairLength;
+				RasterParameters->HairStrandsVF_bUseStableRasterization = VFInput.Strands.bUseStableRasterization ? 1 : 0;
+				RasterParameters->HairStrandsVF_Density			= VFInput.Strands.HairDensity;
+				RasterParameters->HairStrandsVF_LocalToWorldPrimitiveTransform = FMatrix44f(LocalToWorldPrimitiveTransform); // LWC_TODO: Precision loss
+
+				RasterParameters->VisTilePrims = GraphBuilder.CreateSRV(OutVisTilePrims);
+				RasterParameters->RWVisTileDepthGrid = OutVisTileDepthGridUAV;
+				RasterParameters->VisTileArgs = GraphBuilder.CreateSRV(OutVisTileArgs);
+				RasterParameters->RWVisTileData = OutVisTileDataUAV;
+				RasterParameters->VisTileIndirect = GraphBuilder.CreateSRV(OutVisTileIndirect);
+				RasterParameters->IndirectBufferArgs = OutVisTileIndirect;
+
+				FComputeShaderUtils::AddPass(GraphBuilder, RDG_EVENT_NAME("HairStrands::VisibilityComputeRasterRaster(tiled)"), ComputeShaderRaster, RasterParameters, FIntVector(NumRasterizers, 1, 1));
+			}
+
+			if (bDebugEnabled)
+			{
+				FVisiblityRasterComputeDebugCS::FParameters* Parameters = GraphBuilder.AllocParameters<FVisiblityRasterComputeDebugCS::FParameters>();
+				Parameters->MacroGroupId = 0;
+				Parameters->Common = Common;
+				Parameters->PrimitiveInfoIndex = PrimitiveInfoIndex;
+				Parameters->TotalPrimitiveInfoCount = TotalPrimitiveInfoCount;
+				Parameters->VisTileDepthGrid = OutVisTileDepthGrid;
+				Parameters->VisTileBinningGrid = OutVisTileBinningGrid;
+				Parameters->VisTileArgs = GraphBuilder.CreateSRV(OutVisTileArgs);
+				ShaderPrint::SetParameters(GraphBuilder, ViewInfo.ShaderPrintData, Parameters->ShaderPrintParameters);
+
+				TShaderMapRef<FVisiblityRasterComputeDebugCS> DebugComputeShader(ViewInfo.ShaderMap);
+				FComputeShaderUtils::AddPass(GraphBuilder, RDG_EVENT_NAME("HairStrands::VisibilityComputeRaster(Debug)"), DebugComputeShader, Parameters, FIntVector(FMath::DivideAndRoundUp(Common.TileRes.X, 8), FMath::DivideAndRoundUp(Common.TileRes.Y, 8), 1));
+			}
+
+			++PrimitiveInfoIndex;
 		}
 	}
 
+
 	return Out;
 }
 
