Mobile: Make clustered lights and reflections optional for translucency when deferred shading is enabled

diff --git a/Engine/Source/Runtime/Engine/Classes/Engine/RendererSettings.h b/Engine/Source/Runtime/Engine/Classes/Engine/RendererSettings.h
--- a/Engine/Source/Runtime/Engine/Classes/Engine/RendererSettings.h
+++ b/Engine/Source/Runtime/Engine/Classes/Engine/RendererSettings.h
@@ -1019,14 +1019,14 @@ class ENGINE_API URendererSettings : public UDeveloperSettings
 	UPROPERTY(config, EditAnywhere, Category = Mobile, meta = (
 		ConsoleVariable = "r.Mobile.Forward.EnableLocalLights",
 		DisplayName = "Enable local lights support on mobile forward",
-		ToolTip = "Enable local lights support with clustered lighting on mobile forward. 0 is disabled, 1 is enabled (default). Changing this setting requires restarting the editor.",
+		ToolTip = "Enable local lights support for mobile forward shading (including translucency in deferred). 0 is disabled, 1 is enabled (default). Changing this setting requires restarting the editor.",
 		ConfigRestartRequired = true))
 		uint32 bMobileForwardEnableLocalLights : 1;
 
 	UPROPERTY(config, EditAnywhere, Category = Mobile, meta = (
 		ConsoleVariable = "r.Mobile.Forward.EnableClusteredReflections",
 		DisplayName = "Enable clustered reflections on mobile forward",
-		ToolTip = "Whether to enable clustered reflections on mobile forward, it's always supported on mobile deferred. Changing this setting requires restarting the editor.",
+		ToolTip = "Whether to enable clustered reflections on mobile forward (including translucency in deferred). Always supported for opaque geometry on mobile deferred. Changing this setting requires restarting the editor.",
 		ConfigRestartRequired = true))
 		uint32 bMobileForwardEnableClusteredReflections : 1;
 
diff --git a/Engine/Source/Runtime/Engine/Private/MapBuildData.cpp b/Engine/Source/Runtime/Engine/Private/MapBuildData.cpp
--- a/Engine/Source/Runtime/Engine/Private/MapBuildData.cpp
+++ b/Engine/Source/Runtime/Engine/Private/MapBuildData.cpp
@@ -499,7 +499,7 @@ void UMapBuildDataRegistry::Serialize(FArchive& Ar)
 void UMapBuildDataRegistry::PostLoad()
 {
 	Super::PostLoad();
-	bool bMobileEnableClusteredReflections = MobileEnableClusteredReflections(GMaxRHIShaderPlatform);
+	bool bMobileEnableClusteredReflections = MobileForwardEnableClusteredReflections(GMaxRHIShaderPlatform) || IsMobileDeferredShadingEnabled(GMaxRHIShaderPlatform);
 	bool bFullDataRequired = GMaxRHIFeatureLevel >= ERHIFeatureLevel::SM5 || bMobileEnableClusteredReflections;
 	bool bEncodedDataRequired = (GIsEditor || (GMaxRHIFeatureLevel == ERHIFeatureLevel::ES3_1 && !bMobileEnableClusteredReflections));
 
diff --git a/Engine/Source/Runtime/RenderCore/Private/RenderUtils.cpp b/Engine/Source/Runtime/RenderCore/Private/RenderUtils.cpp
--- a/Engine/Source/Runtime/RenderCore/Private/RenderUtils.cpp
+++ b/Engine/Source/Runtime/RenderCore/Private/RenderUtils.cpp
@@ -1158,7 +1158,8 @@ RENDERCORE_API bool SupportsTextureCubeArray(ERHIFeatureLevel::Type FeatureLevel
 {
 	return FeatureLevel >= ERHIFeatureLevel::SM5 
 		// requries ES3.2 feature set
-		|| MobileEnableClusteredReflections(GetFeatureLevelShaderPlatform(FeatureLevel));
+		|| IsMobileDeferredShadingEnabled(GetFeatureLevelShaderPlatform(FeatureLevel))
+		|| MobileForwardEnableClusteredReflections(GetFeatureLevelShaderPlatform(FeatureLevel));
 }
 
 RENDERCORE_API bool MaskedInEarlyPass(const FStaticShaderPlatform Platform)
@@ -1231,10 +1232,10 @@ RENDERCORE_API bool MobileForwardEnableLocalLights(const FStaticShaderPlatform P
 	return MobileForwardEnableLocalLightsIniValue.Get(Platform);
 }
 
-RENDERCORE_API bool MobileEnableClusteredReflections(const FStaticShaderPlatform Platform)
+RENDERCORE_API bool MobileForwardEnableClusteredReflections(const FStaticShaderPlatform Platform)
 {
 	static FShaderPlatformCachedIniValue<bool> MobileForwardEnableClusteredReflectionsIniValue(TEXT("r.Mobile.Forward.EnableClusteredReflections"));
-	return IsMobileDeferredShadingEnabled(Platform) || MobileForwardEnableClusteredReflectionsIniValue.Get(Platform);
+	return MobileForwardEnableClusteredReflectionsIniValue.Get(Platform);
 }
 
 RENDERCORE_API bool MobileUsesShadowMaskTexture(const FStaticShaderPlatform Platform)
diff --git a/Engine/Source/Runtime/RenderCore/Public/RenderUtils.h b/Engine/Source/Runtime/RenderCore/Public/RenderUtils.h
--- a/Engine/Source/Runtime/RenderCore/Public/RenderUtils.h
+++ b/Engine/Source/Runtime/RenderCore/Public/RenderUtils.h
@@ -509,7 +509,7 @@ RENDERCORE_API bool IsMobileMovableSpotlightShadowsEnabled(const FStaticShaderPl
 
 RENDERCORE_API bool MobileForwardEnableLocalLights(const FStaticShaderPlatform Platform);
 
-RENDERCORE_API bool MobileEnableClusteredReflections(const FStaticShaderPlatform Platform);
+RENDERCORE_API bool MobileForwardEnableClusteredReflections(const FStaticShaderPlatform Platform);
 
 RENDERCORE_API bool MobileUsesShadowMaskTexture(const FStaticShaderPlatform Platform);
 
diff --git a/Engine/Source/Runtime/Renderer/Private/MobileBasePass.cpp b/Engine/Source/Runtime/Renderer/Private/MobileBasePass.cpp
--- a/Engine/Source/Runtime/Renderer/Private/MobileBasePass.cpp
+++ b/Engine/Source/Runtime/Renderer/Private/MobileBasePass.cpp
@@ -666,9 +666,8 @@ bool FMobileBasePassMeshProcessor::Process(
 	bool bEnableLocalLights = false;
 	if (Scene && PrimitiveSceneProxy && ShadingModels.IsLit())
 	{
-		const bool bIsTranslucentMaterialOnDeferred = Scene && bTranslucentBasePass && IsMobileDeferredShadingEnabled(Scene->GetShaderPlatform());
-		// Clustered lighting could only be enabled on mobile forward and always enabled on mobile deferred translucent material
-		if (!bUsesDeferredShading && (MobileForwardEnableLocalLights(Scene->GetShaderPlatform()) || bIsTranslucentMaterialOnDeferred))
+		// Whether material with a forward shading needs clustered lights
+		if (!bUsesDeferredShading && MobileForwardEnableLocalLights(Scene->GetShaderPlatform()))
 		{
 			bEnableLocalLights = PrimitiveSceneProxy->GetPrimitiveSceneInfo()->NumMobileDynamicLocalLights > 0;
 		}
diff --git a/Engine/Source/Runtime/Renderer/Private/MobileBasePassRendering.h b/Engine/Source/Runtime/Renderer/Private/MobileBasePassRendering.h
--- a/Engine/Source/Runtime/Renderer/Private/MobileBasePassRendering.h
+++ b/Engine/Source/Runtime/Renderer/Private/MobileBasePassRendering.h
@@ -399,7 +399,7 @@ public:
 		const bool bMobileUseHWsRGBEncoding = (MobileUseHWsRGBEncodingCVAR && MobileUseHWsRGBEncodingCVAR->GetValueOnAnyThread() == 1);
 
 		const bool bMobileUsesShadowMaskTexture = MobileUsesShadowMaskTexture(Parameters.Platform);
-		const bool bEnableClusteredReflections = MobileEnableClusteredReflections(Parameters.Platform);
+		const bool bEnableClusteredReflections = MobileForwardEnableClusteredReflections(Parameters.Platform);
 		const bool bTranslucentMaterial = IsTranslucentBlendMode(Parameters.MaterialParameters.BlendMode) || Parameters.MaterialParameters.ShadingModels.HasShadingModel(MSM_SingleLayerWater);
 
 		TMobileBasePassPSBaseType<LightMapPolicyType>::ModifyCompilationEnvironment(Parameters, OutEnvironment);
diff --git a/Engine/Source/Runtime/Renderer/Private/MobileDeferredShadingPass.cpp b/Engine/Source/Runtime/Renderer/Private/MobileDeferredShadingPass.cpp
--- a/Engine/Source/Runtime/Renderer/Private/MobileDeferredShadingPass.cpp
+++ b/Engine/Source/Runtime/Renderer/Private/MobileDeferredShadingPass.cpp
@@ -13,10 +13,16 @@ int32 GMobileUseClusteredDeferredShading = 0;
 static FAutoConsoleVariableRef CVarMobileUseClusteredDeferredShading(
 	TEXT("r.Mobile.UseClusteredDeferredShading"),
 	GMobileUseClusteredDeferredShading,
-	TEXT("Toggle use of clustered deferred shading for lights that support it. 0 is off (default), 1 is on"),
+	TEXT("Toggle use of clustered deferred shading for lights that support it. 0 is off (default), 1 is on. (requires LightGrid: r.Mobile.Forward.EnableLocalLights=1)"),
 	ECVF_RenderThreadSafe
 );
 
+static bool UseClusteredDeferredShading(const FStaticShaderPlatform Platform)
+{
+	// Needs LightGrid to function
+	return GMobileUseClusteredDeferredShading != 0 && MobileForwardEnableLocalLights(Platform);
+}
+
 int32 GMobileUseLightStencilCulling = 1;
 static FAutoConsoleVariableRef CVarMobileUseLightStencilCulling(
 	TEXT("r.Mobile.UseLightStencilCulling"),
@@ -113,7 +119,7 @@ class FMobileDirectionalLightFunctionPS : public FMaterialShader
 
 	static FPermutationDomain BuildPermutationVector(const FViewInfo& View, bool bInlineReflectionAndSky, bool bShadingModelSupport, bool bDynamicShadows, bool bSkyLight, bool bPlanarReflection)
 	{
-		bool bUseClusteredLights = GMobileUseClusteredDeferredShading != 0;
+		bool bUseClusteredLights = UseClusteredDeferredShading(View.GetShaderPlatform());
 		bool bClustredReflection = bInlineReflectionAndSky && (View.NumBoxReflectionCaptures + View.NumSphereReflectionCaptures) > 0;
 		bool bEnableSkyLight = bInlineReflectionAndSky && bSkyLight;
 		const bool bMobileUsesShadowMaskTexture = MobileUsesShadowMaskTexture(View.GetShaderPlatform());
@@ -909,8 +915,9 @@ void MobileDeferredShadingPass(
 	check(DefaultMaterial.Material);
 
 	RenderDirectionalLights(RHICmdList, Scene, View, DefaultMaterial);
-
-	if (GMobileUseClusteredDeferredShading == 0)
+	
+	const bool bMobileUseClusteredDeferredShading = UseClusteredDeferredShading(View.GetShaderPlatform());
+	if (!bMobileUseClusteredDeferredShading)
 	{
 		// Render non-clustered simple lights
 		RenderSimpleLights(RHICmdList, Scene, ViewIndex, NumViews, View, SortedLightSet, DefaultMaterial);
@@ -920,7 +927,7 @@ void MobileDeferredShadingPass(
 	int32 NumLights = SortedLightSet.SortedLights.Num();
 	const int32 UnbatchedLightStart = SortedLightSet.UnbatchedLightStart;
 	int32 StandardDeferredStart = SortedLightSet.SimpleLightsEnd;
-	if (GMobileUseClusteredDeferredShading != 0)
+	if (bMobileUseClusteredDeferredShading)
 	{
 		StandardDeferredStart = SortedLightSet.ClusteredSupportedEnd;
 	}
diff --git a/Engine/Source/Runtime/Renderer/Private/MobileShadingRenderer.cpp b/Engine/Source/Runtime/Renderer/Private/MobileShadingRenderer.cpp
--- a/Engine/Source/Runtime/Renderer/Private/MobileShadingRenderer.cpp
+++ b/Engine/Source/Runtime/Renderer/Private/MobileShadingRenderer.cpp
@@ -260,8 +260,8 @@ FMobileSceneRenderer::FMobileSceneRenderer(const FSceneViewFamily* InViewFamily,
 	bIsFullDepthPrepassEnabled = Scene->EarlyZPassMode == DDM_AllOpaque;
 	bIsMaskedOnlyDepthPrepassEnabled = Scene->EarlyZPassMode == DDM_MaskedOnly;
 	bRequiresSceneDepthAux = MobileRequiresSceneDepthAux(ShaderPlatform);
-	bEnableClusteredLocalLights = bDeferredShading || MobileForwardEnableLocalLights(ShaderPlatform);
-	bEnableClusteredReflections = MobileEnableClusteredReflections(ShaderPlatform);
+	bEnableClusteredLocalLights = MobileForwardEnableLocalLights(ShaderPlatform);
+	bEnableClusteredReflections = MobileForwardEnableClusteredReflections(ShaderPlatform);
 	
 	StandardTranslucencyPass = ViewFamily.AllowTranslucencyAfterDOF() ? ETranslucencyPass::TPT_StandardTranslucency : ETranslucencyPass::TPT_AllTranslucency;
 	StandardTranslucencyMeshPass = TranslucencyPassToMeshPass(StandardTranslucencyPass);
@@ -595,8 +595,7 @@ void FMobileSceneRenderer::InitViews(FRDGBuilder& GraphBuilder, FSceneTexturesCo
 	{
 		Scene->IndirectLightingCache.FinalizeCacheUpdates(Scene, *this, ILCTaskData);
 	}
-
-	const bool bEnableClusteredShading = bEnableClusteredLocalLights || bEnableClusteredReflections;
+		
 	// initialize per-view uniform buffer.  Pass in shadow info as necessary.
 	for (int32 ViewIndex = Views.Num() - 1; ViewIndex >= 0; --ViewIndex)
 	{
@@ -644,7 +643,9 @@ void FMobileSceneRenderer::InitViews(FRDGBuilder& GraphBuilder, FSceneTexturesCo
 		}
 	}
 
-	if (bEnableClusteredShading)
+	if (bDeferredShading ||
+		bEnableClusteredLocalLights || 
+		bEnableClusteredReflections)
 	{
 		SetupSceneReflectionCaptureBuffer(RHICmdList);
 	}
@@ -829,18 +830,19 @@ void FMobileSceneRenderer::Render(FRDGBuilder& GraphBuilder)
 
 	FSortedLightSetSceneInfo& SortedLightSet = *GraphBuilder.AllocObject<FSortedLightSetSceneInfo>();
 
-	const bool bEnableClusteredShading = bEnableClusteredLocalLights || bEnableClusteredReflections;
-	if (bEnableClusteredShading)
+	if (bDeferredShading ||
+		bEnableClusteredLocalLights || 
+		bEnableClusteredReflections)
 	{
 		RDG_CSV_STAT_EXCLUSIVE_SCOPE(GraphBuilder, SortLights);
 		// Shadows are applied in clustered shading on mobile forward and separately on mobile deferred.
 		bool bShadowedLightsInClustered = bRequiresShadowProjections && !bDeferredShading;
 		GatherAndSortLights(SortedLightSet, bShadowedLightsInClustered);
 		int32 NumReflectionCaptures = Views[0].NumBoxReflectionCaptures + Views[0].NumSphereReflectionCaptures;
-		bool bCullLightsToGrid = ((bEnableClusteredReflections && NumReflectionCaptures > 0) || bEnableClusteredLocalLights);
+		bool bCullLightsToGrid = (((bEnableClusteredReflections || bDeferredShading) && NumReflectionCaptures > 0) || bEnableClusteredLocalLights);
 		if (bCullLightsToGrid)
 		{
-			ComputeLightGrid(GraphBuilder, true, SortedLightSet);
+			ComputeLightGrid(GraphBuilder, bEnableClusteredLocalLights, SortedLightSet);
 		}
 	}
 
