Add hair hit proxy support for raster compute.

diff --git a/Engine/Plugins/Runtime/HairStrands/Source/HairStrandsCore/Private/HairStrandsRendering.cpp b/Engine/Plugins/Runtime/HairStrands/Source/HairStrandsCore/Private/HairStrandsRendering.cpp
--- a/Engine/Plugins/Runtime/HairStrands/Source/HairStrandsCore/Private/HairStrandsRendering.cpp
+++ b/Engine/Plugins/Runtime/HairStrands/Source/HairStrandsCore/Private/HairStrandsRendering.cpp
@@ -1581,17 +1581,19 @@ bool NeedsPatchAttributeBuffer(EHairStrandsDebugMode DebugMode)
 
 static void ConvertHairStrandsVFParameters(
 	FRDGBuilder* GraphBuilder,
-	FRDGImportedBuffer& Buffer,
-	FShaderResourceViewRHIRef& BufferRHISRV,
+	FRDGImportedBuffer& OutBuffer,
+	FRDGExternalBuffer& OutExternalBuffer,
+	FShaderResourceViewRHIRef& OutBufferRHISRV,
 	const FRDGExternalBuffer& ExternalBuffer)
 {
 	if (GraphBuilder)
 	{
-		Buffer = Register(*GraphBuilder, ExternalBuffer, ERDGImportedBufferFlags::CreateSRV);
+		OutBuffer = Register(*GraphBuilder, ExternalBuffer, ERDGImportedBufferFlags::CreateSRV);
 	}
-	BufferRHISRV = ExternalBuffer.SRV;
+	OutBufferRHISRV = ExternalBuffer.SRV;
+	OutExternalBuffer = ExternalBuffer;
 }
-#define CONVERT_HAIRSSTRANDS_VF_PARAMETERS(OutName, ExternalBuffer) ConvertHairStrandsVFParameters(GraphBuilder, OutName, OutName##RHISRV, ExternalBuffer);
+#define CONVERT_HAIRSSTRANDS_VF_PARAMETERS(OutName, ExternalBuffer) ConvertHairStrandsVFParameters(GraphBuilder, OutName, OutName##External, OutName##RHISRV, ExternalBuffer);
 
 // Compute/Update the hair strands description which will be used for rendering (VF) / voxelization & co.
 static FHairGroupPublicData::FVertexFactoryInput InternalComputeHairStrandsVertexInputData(FRDGBuilder* GraphBuilder, const FHairGroupInstance* Instance)
diff --git a/Engine/Shaders/Private/HairStrands/HairStrandsVisibilityCompactionComputeRaster.usf b/Engine/Shaders/Private/HairStrands/HairStrandsVisibilityCompactionComputeRaster.usf
--- a/Engine/Shaders/Private/HairStrands/HairStrandsVisibilityCompactionComputeRaster.usf
+++ b/Engine/Shaders/Private/HairStrands/HairStrandsVisibilityCompactionComputeRaster.usf
@@ -103,9 +103,6 @@ void MainCS(
 
 	FSampleSetDesc SampleDesc;
 
-	const float OpaqueDepth = SceneTexturesStruct.SceneDepthTexture.Load(uint3(PixelCoord, 0)).r;
-
-
 	const uint VisData0 = DepthCovTexture.Load(uint3(PixelCoord, 0));
 	const uint VisData1 = PrimMatTexture.Load(uint3(PixelCoord, 0));
 
diff --git a/Engine/Source/Runtime/Renderer/Private/HairStrands/HairStrandsVisibility.cpp b/Engine/Source/Runtime/Renderer/Private/HairStrands/HairStrandsVisibility.cpp
--- a/Engine/Source/Runtime/Renderer/Private/HairStrands/HairStrandsVisibility.cpp
+++ b/Engine/Source/Runtime/Renderer/Private/HairStrands/HairStrandsVisibility.cpp
@@ -1828,7 +1828,7 @@ class FHairVisibilityCompactionComputeRasterCS : public FGlobalShader
 		SHADER_PARAMETER_RDG_BUFFER_SRV(Buffer<uint2>, TileDataBuffer)						// Tile coords (RG16)
 		SHADER_PARAMETER_RDG_BUFFER_SRV(Buffer<uint>, TileCountBuffer)						// Tile total count (actual number of tiles)
 
-		SHADER_PARAMETER_RDG_UNIFORM_BUFFER(FSceneTextureUniformParameters, SceneTexturesStruct)
+//		SHADER_PARAMETER_RDG_UNIFORM_BUFFER(FSceneTextureUniformParameters, SceneTexturesStruct)
 		SHADER_PARAMETER_STRUCT_REF(FViewUniformShaderParameters, ViewUniformBuffer)
 		RDG_BUFFER_ACCESS(IndirectBufferArgs, ERHIAccess::IndirectArgs)
 	END_SHADER_PARAMETER_STRUCT()
@@ -1842,7 +1842,7 @@ public:
 
 	static bool ShouldCompilePermutation(const FGlobalShaderPermutationParameters& Parameters)
 	{
-		return IsHairStrandsSupported(EHairStrandsShaderType::Strands, Parameters.Platform);
+		return IsHairStrandsSupported(EHairStrandsShaderType::Strands, Parameters.Platform);	
 	}
 };
 
@@ -1906,7 +1906,7 @@ static void AddHairVisibilityCompactionComputeRasterPass(
 	PassParameters->MaxNodeCount			= MaxRenderNodeCount;
 	PassParameters->CoverageThreshold		= GetHairStrandsFullCoverageThreshold();
 	PassParameters->ViewUniformBuffer		= View.ViewUniformBuffer;
-	PassParameters->SceneTexturesStruct		= CreateSceneTextureUniformBuffer(GraphBuilder, View.GetSceneTexturesChecked(), View.FeatureLevel);
+//	PassParameters->SceneTexturesStruct		= CreateSceneTextureUniformBuffer(GraphBuilder, View.GetSceneTexturesChecked(), View.FeatureLevel);
 	PassParameters->OutCompactNodeCounter	= GraphBuilder.CreateUAV(OutCompactCounter);
 	PassParameters->OutCompactNodeIndex		= GraphBuilder.CreateUAV(OutCompactNodeIndex);
 	PassParameters->OutCompactNodeVis		= GraphBuilder.CreateUAV(OutCompactNodeVis);
@@ -3056,7 +3056,9 @@ static FRasterComputeOutput AddVisibilityComputeRasterPass(
 	const FHairStrandsMacroGroupDatas& MacroGroupDatas,
 	const FIntPoint& InResolution,
 	const uint32 SamplePerPixelCount,
-	const FRDGTextureRef SceneDepthTexture)
+	const FRDGTextureRef SceneDepthTexture,
+	bool bSupportCulling,
+	bool bForceRegister)
 {	
 	check(DoesSupportRasterCompute());
 
@@ -3210,7 +3212,7 @@ static FRasterComputeOutput AddVisibilityComputeRasterPass(
 			const FHairGroupPublicData::FVertexFactoryInput& VFInput = HairGroupPublicData->VFInput;
 			const FMatrix44d LocalToWorldPrimitiveTransform = VFInput.LocalToWorldTransform.ToMatrixWithScale(); 
 
-			const bool bCullingEnable = GHairVisibilityComputeRaster_Culling ? HairGroupPublicData->GetCullingResultAvailable() : false;
+			const bool bCullingEnable = bSupportCulling && GHairVisibilityComputeRaster_Culling ? HairGroupPublicData->GetCullingResultAvailable() : false;
 			
 			// calculate current view screen size - which can result in fewer strands rasterized in current view
 			const FSphere BoundsSphere = HairGroupPublicData->ContinuousLODBounds.GetSphere();
@@ -3241,8 +3243,18 @@ static FRasterComputeOutput AddVisibilityComputeRasterPass(
 				BinningParameters->OutVisTileArgs = OutVisTileArgsUAV;
 				BinningParameters->OutVisTileData = OutVisTileDataUAV;
 
-				BinningParameters->HairStrandsVF_PositionBuffer = VFInput.Strands.PositionBuffer.SRV;
-				BinningParameters->HairStrandsVF_PositionOffsetBuffer = VFInput.Strands.PositionOffsetBuffer.SRV;
+				if (bForceRegister)
+				{
+					BinningParameters->HairStrandsVF_PositionBuffer = Register(GraphBuilder, VFInput.Strands.PositionBufferExternal, ERDGImportedBufferFlags::CreateSRV).SRV;;
+					BinningParameters->HairStrandsVF_PositionOffsetBuffer = Register(GraphBuilder, VFInput.Strands.PositionOffsetBufferExternal, ERDGImportedBufferFlags::CreateSRV).SRV;;
+
+				}
+				else
+				{
+					BinningParameters->HairStrandsVF_PositionBuffer = VFInput.Strands.PositionBuffer.SRV;
+					BinningParameters->HairStrandsVF_PositionOffsetBuffer = VFInput.Strands.PositionOffsetBuffer.SRV;
+				}
+
 				BinningParameters->HairStrandsVF_VertexCount = VFInput.Strands.VertexCount;
 				BinningParameters->HairStrandsVF_Radius = VFInput.Strands.HairRadius;
 				BinningParameters->HairStrandsVF_RootScale = VFInput.Strands.HairRootScale;
@@ -3283,8 +3295,20 @@ static FRasterComputeOutput AddVisibilityComputeRasterPass(
 				RasterParameters->OutPrimMatTexture = PrimMatTextureUAV;
 				RasterParameters->SampleWeight = SampleWeight;
 
-				RasterParameters->HairStrandsVF_PositionBuffer	= VFInput.Strands.PositionBuffer.SRV;
-				RasterParameters->HairStrandsVF_PositionOffsetBuffer = VFInput.Strands.PositionOffsetBuffer.SRV;
+				if (bForceRegister)
+				{
+					RasterParameters->HairStrandsVF_PositionBuffer = Register(GraphBuilder, VFInput.Strands.PositionBufferExternal, ERDGImportedBufferFlags::CreateSRV).SRV;;
+					RasterParameters->HairStrandsVF_PositionOffsetBuffer = Register(GraphBuilder, VFInput.Strands.PositionOffsetBufferExternal, ERDGImportedBufferFlags::CreateSRV).SRV;;
+
+				}
+				else
+				{
+					RasterParameters->HairStrandsVF_PositionBuffer = VFInput.Strands.PositionBuffer.SRV;
+					RasterParameters->HairStrandsVF_PositionOffsetBuffer = VFInput.Strands.PositionOffsetBuffer.SRV;
+				}
+
+//				RasterParameters->HairStrandsVF_PositionBuffer	= VFInput.Strands.PositionBuffer.SRV;
+//				RasterParameters->HairStrandsVF_PositionOffsetBuffer = VFInput.Strands.PositionOffsetBuffer.SRV;
 				RasterParameters->HairStrandsVF_VertexCount		= VFInput.Strands.VertexCount;
 				RasterParameters->HairStrandsVF_Radius			= VFInput.Strands.HairRadius;
 				RasterParameters->HairStrandsVF_RootScale			= VFInput.Strands.HairRootScale;
@@ -3598,7 +3622,8 @@ void DrawHitProxies(
 {
 	// Proxy rendering is only supported/compatible with MSAA-visibility rendering. 
 	// PPLL is not supported, but it is supposed to be used only for final render.
-	if (GetHairVisibilityRenderMode() != HairVisibilityRenderMode_MSAA_Visibility)
+	const EHairVisibilityRenderMode RenderMode = GetHairVisibilityRenderMode();
+	if (RenderMode != HairVisibilityRenderMode_MSAA_Visibility && RenderMode != HairVisibilityRenderMode_ComputeRaster)
 	{
 		return;
 	}
@@ -3618,60 +3643,92 @@ void DrawHitProxies(
 	FHairPrimaryTransmittance ViewTransmittance;
 
 	FRDGTextureRef SceneDepthTexture = HitProxyDepthTexture;
-	FRDGTextureRef VisDepthTexture = AddHairVisibilityFillOpaqueDepth(
-		GraphBuilder,
-		View,
-		Resolution,
-		MacroGroupDatas,
-		TileData,
-		SceneDepthTexture);
-
-	FRDGTextureRef VisIdTexture = nullptr;
-	AddHairVisibilityMSAAPass(
-		GraphBuilder,
-		&Scene,
-		&View,
-		MacroGroupDatas,
-		Resolution,
-		TileData,
-		InstanceCullingManager,
-		VisIdTexture,
-		VisDepthTexture);
-
-
-	FHairVisibilityPrimitiveIdCompactionCS::FParameters* PassParameters = GraphBuilder.AllocParameters<FHairVisibilityPrimitiveIdCompactionCS::FParameters>();
-	PassParameters->MSAA_DepthTexture = VisDepthTexture;
-	PassParameters->MSAA_IDTexture = VisIdTexture;
-	PassParameters->ViewTransmittanceTexture = DummyTransmittanceTexture;
 
 	const uint32 NodeGroupSize = GetVendorOptimalGroupSize1D();
 	const uint32 MaxSampleCount = 4;
-
 	FRDGTextureRef NodeCounter = nullptr;
 	FRDGTextureRef VisNodeIndex = nullptr;
 	FRDGBufferRef  VisNodeData = nullptr;
 	FRDGBufferRef  VisNodeCoord = nullptr;
-	FRDGBufferRef  IndirectArgsBuffer = nullptr;
-	FRDGTextureRef ResolveMaskTexture = nullptr;
 	FRDGTextureRef CoverageTexture = nullptr;
-
+	FRDGBufferRef  IndirectArgsBuffer = nullptr;
 	uint32 OutMaxNodeCount = 0;
-	AddHairVisibilityPrimitiveIdCompactionPass(
-		false, // bUsePPLL
-		GraphBuilder,
-		View,
-		SceneDepthTexture,
-		MacroGroupDatas,
-		NodeGroupSize,
-		TileData,
-		PassParameters,
-		NodeCounter,
-		VisNodeIndex,
-		VisNodeData,
-		VisNodeCoord,
-		CoverageTexture,
-		IndirectArgsBuffer,
-		OutMaxNodeCount);
+	if (RenderMode == HairVisibilityRenderMode_ComputeRaster)
+	{
+		FRasterComputeOutput RasterOutput = AddVisibilityComputeRasterPass(
+			GraphBuilder,
+			View,
+			MacroGroupDatas,
+			Resolution,
+			MaxSampleCount,
+			SceneDepthTexture,
+			false,
+			true);
+
+		// Generate Tile data
+		check(RasterOutput.PrimMatTexture);
+		TileData = AddHairStrandsGenerateTilesPass(GraphBuilder, View, RasterOutput.HairCountTexture);
+
+		AddHairVisibilityCompactionComputeRasterPass(
+			GraphBuilder,
+			View,
+			NodeGroupSize,
+			MaxSampleCount,
+			RasterOutput,
+			TileData,
+			NodeCounter,
+			VisNodeIndex,
+			VisNodeData,
+			VisNodeCoord,
+			CoverageTexture,
+			IndirectArgsBuffer,
+			OutMaxNodeCount);
+
+	}
+	else if (RenderMode == HairVisibilityRenderMode_MSAA_Visibility)
+	{
+		FRDGTextureRef VisDepthTexture = AddHairVisibilityFillOpaqueDepth(
+			GraphBuilder,
+			View,
+			Resolution,
+			MacroGroupDatas,
+			TileData,
+			SceneDepthTexture);
+
+		FRDGTextureRef VisIdTexture = nullptr;
+		AddHairVisibilityMSAAPass(
+			GraphBuilder,
+			&Scene,
+			&View,
+			MacroGroupDatas,
+			Resolution,
+			TileData,
+			InstanceCullingManager,
+			VisIdTexture,
+			VisDepthTexture);
+
+		FHairVisibilityPrimitiveIdCompactionCS::FParameters* PassParameters = GraphBuilder.AllocParameters<FHairVisibilityPrimitiveIdCompactionCS::FParameters>();
+		PassParameters->MSAA_DepthTexture = VisDepthTexture;
+		PassParameters->MSAA_IDTexture = VisIdTexture;
+		PassParameters->ViewTransmittanceTexture = DummyTransmittanceTexture;
+
+		AddHairVisibilityPrimitiveIdCompactionPass(
+			false, // bUsePPLL
+			GraphBuilder,
+			View,
+			SceneDepthTexture,
+			MacroGroupDatas,
+			NodeGroupSize,
+			TileData,
+			PassParameters,
+			NodeCounter,
+			VisNodeIndex,
+			VisNodeData,
+			VisNodeCoord,
+			CoverageTexture,
+			IndirectArgsBuffer,
+			OutMaxNodeCount);
+	}
 
 	AddHairStrandsHitProxyIdPass(GraphBuilder, Scene, View, VisNodeIndex, VisNodeData, HitProxyTexture, HitProxyDepthTexture);
 }
@@ -3806,7 +3863,9 @@ void RenderHairStrandsVisibilityBuffer(
 					MacroGroupDatas,
 					Resolution,
 					VisibilityData.MaxSampleCount,
-					SceneDepthTexture);
+					SceneDepthTexture,
+					true,
+					false);
 
 				// Generate Tile data
 				check(RasterOutput.PrimMatTexture);
diff --git a/Engine/Source/Runtime/Renderer/Public/HairStrandsInterface.h b/Engine/Source/Runtime/Renderer/Public/HairStrandsInterface.h
--- a/Engine/Source/Runtime/Renderer/Public/HairStrandsInterface.h
+++ b/Engine/Source/Runtime/Renderer/Public/HairStrandsInterface.h
@@ -261,6 +261,15 @@ public:
 			FRDGImportedBuffer PositionOffsetBuffer;
 			FRDGImportedBuffer PrevPositionOffsetBuffer;
 
+			FRDGExternalBuffer PositionBufferExternal;
+			FRDGExternalBuffer PrevPositionBufferExternal;
+			FRDGExternalBuffer TangentBufferExternal;
+			FRDGExternalBuffer MaterialBufferExternal;
+			FRDGExternalBuffer Attribute0BufferExternal;
+			FRDGExternalBuffer Attribute1BufferExternal;
+			FRDGExternalBuffer PositionOffsetBufferExternal;
+			FRDGExternalBuffer PrevPositionOffsetBufferExternal;
+
 			FShaderResourceViewRHIRef PositionBufferRHISRV				= nullptr;
 			FShaderResourceViewRHIRef PrevPositionBufferRHISRV			= nullptr;
 			FShaderResourceViewRHIRef TangentBufferRHISRV				= nullptr;
