Path Tracer: Default light functions to be greyscale like the deferred renderer, with a cvar to toggle colored light functions if needed.

diff --git a/Engine/Shaders/Private/PathTracing/PathTracingLightingMissShader.usf b/Engine/Shaders/Private/PathTracing/PathTracingLightingMissShader.usf
--- a/Engine/Shaders/Private/PathTracing/PathTracingLightingMissShader.usf
+++ b/Engine/Shaders/Private/PathTracing/PathTracingLightingMissShader.usf
@@ -16,6 +16,13 @@ RAY_TRACING_ENTRY_MISS(PathTracingLightingMS, FPackedPathTracingPayload, Payload
 	Payload.SetMiss();
 	float3 TranslatedWorldPosition = WorldRayOrigin(); // Shadow ray is traced from the shaded point
 	float3 Throughput = Payload.GetRayThroughput();
-	Throughput *= GetRayTracingLightFunction(TranslatedWorldPosition);
+	float3 LightFunctionMult = GetRayTracingLightFunction(TranslatedWorldPosition);
+
+	if (PathTracingLightFunctionParameters.EnableColoredLightFunctions == 0)
+	{
+		// convert to greyscale for compatibility with the main engine (optionally)
+		LightFunctionMult = dot(LightFunctionMult, .3333f);
+	}
+	Throughput *= LightFunctionMult;
 	Payload.SetRayThroughput(Throughput);
 }
diff --git a/Engine/Source/Runtime/Renderer/Private/PathTracing.cpp b/Engine/Source/Runtime/Renderer/Private/PathTracing.cpp
--- a/Engine/Source/Runtime/Renderer/Private/PathTracing.cpp
+++ b/Engine/Source/Runtime/Renderer/Private/PathTracing.cpp
@@ -269,6 +269,15 @@ TAutoConsoleVariable<int32> CVarPathTracingDecalGridVisualize(
 	ECVF_RenderThreadSafe
 );
 
+TAutoConsoleVariable<int32> CVarPathTracingLightFunctionColor(
+	TEXT("r.PathTracing.LightFunctionColor"),
+	0,
+	TEXT("Enables light functions to be colored instead of greyscale (default = 0)\n")
+	TEXT("0: off (default)\n")
+	TEXT("1: on (light function material output is used directly instead of converting to greyscale)\n"),
+	ECVF_RenderThreadSafe
+);
+
 
 BEGIN_SHADER_PARAMETER_STRUCT(FPathTracingData, )
 	SHADER_PARAMETER(float, BlendFactor)
@@ -867,6 +876,7 @@ BEGIN_GLOBAL_SHADER_PARAMETER_STRUCT(FLightFunctionParametersPathTracing, )
 	SHADER_PARAMETER(FMatrix44f, LightFunctionTranslatedWorldToLight)
 	SHADER_PARAMETER(FVector4f, LightFunctionParameters)
 	SHADER_PARAMETER(FVector3f, LightFunctionParameters2)
+	SHADER_PARAMETER(int32    , EnableColoredLightFunctions)
 END_GLOBAL_SHADER_PARAMETER_STRUCT()
 
 IMPLEMENT_GLOBAL_SHADER_PARAMETER_STRUCT(FLightFunctionParametersPathTracing, "PathTracingLightFunctionParameters");
@@ -900,6 +910,8 @@ static TUniformBufferRef<FLightFunctionParametersPathTracing> CreateLightFunctio
 		LightSceneInfo->Proxy->GetLightFunctionDisabledBrightness(),
 		bRenderingPreviewShadowIndicator ? 1.0f : 0.0f);
 
+	LightFunctionParameters.EnableColoredLightFunctions = CVarPathTracingLightFunctionColor.GetValueOnRenderThread();
+
 	return CreateUniformBufferImmediate(LightFunctionParameters, Usage);
 }
 
@@ -1957,7 +1969,7 @@ void FDeferredShadingSceneRenderer::RenderPathTracing(
 	MaxSPP = FMath::Max(MaxSPP, 1u);
 	Config.LockedSamplingPattern = CVarPathTracingFrameIndependentTemporalSeed.GetValueOnRenderThread() == 0;
 
-	// compute an integer code of what show flags related to lights are currently enabled so we can detect changes
+	// compute an integer code of what show flags and booleans related to lights are currently enabled so we can detect changes
 	Config.LightShowFlags = 0;
 	Config.LightShowFlags |= View.Family->EngineShowFlags.SkyLighting           ? 1 << 0 : 0;
 	Config.LightShowFlags |= View.Family->EngineShowFlags.DirectionalLights     ? 1 << 1 : 0;
@@ -1966,6 +1978,7 @@ void FDeferredShadingSceneRenderer::RenderPathTracing(
 	Config.LightShowFlags |= View.Family->EngineShowFlags.PointLights           ? 1 << 4 : 0;
 	Config.LightShowFlags |= View.Family->EngineShowFlags.TexturedLightProfiles ? 1 << 5 : 0;
 	Config.LightShowFlags |= View.Family->EngineShowFlags.LightFunctions        ? 1 << 6 : 0;
+	Config.LightShowFlags |= CVarPathTracingLightFunctionColor.GetValueOnRenderThread() ? 1 << 7 : 0;
 
 	PreparePathTracingData(Scene, View, Config.PathTracingData);
 
