Fix function name clash in unity build

diff --git a/Engine/Source/Runtime/Renderer/Private/PathTracing.cpp b/Engine/Source/Runtime/Renderer/Private/PathTracing.cpp
--- a/Engine/Source/Runtime/Renderer/Private/PathTracing.cpp
+++ b/Engine/Source/Runtime/Renderer/Private/PathTracing.cpp
@@ -881,7 +881,7 @@ END_GLOBAL_SHADER_PARAMETER_STRUCT()
 
 IMPLEMENT_GLOBAL_SHADER_PARAMETER_STRUCT(FLightFunctionParametersPathTracing, "PathTracingLightFunctionParameters");
 
-static TUniformBufferRef<FLightFunctionParametersPathTracing> CreateLightFunctionParametersBuffer(
+static TUniformBufferRef<FLightFunctionParametersPathTracing> CreateLightFunctionParametersBufferPT(
 	const FLightSceneInfo* LightSceneInfo,
 	const FSceneView& View,
 	EUniformBufferUsage Usage)
@@ -1021,7 +1021,7 @@ void BindLightFunctionShadersPathTracing(
 
 		const FMaterialRenderProxy& MaterialRenderProxy = FallbackMaterialRenderProxyPtr ? *FallbackMaterialRenderProxyPtr : *MaterialProxy;
 
-		TUniformBufferRef<FLightFunctionParametersPathTracing> LightFunctionParameters = CreateLightFunctionParametersBuffer(LightSceneInfo, View, EUniformBufferUsage::UniformBuffer_SingleFrame);
+		TUniformBufferRef<FLightFunctionParametersPathTracing> LightFunctionParameters = CreateLightFunctionParametersBufferPT(LightSceneInfo, View, EUniformBufferUsage::UniformBuffer_SingleFrame);
 
 		int32 MissIndex = LightAndIndex.Value;
 		BindLightFunction(RHICmdList, Scene, View, Material, MaterialRenderProxy, LightFunctionParameters, MissIndex);
diff --git a/Engine/Source/Runtime/Renderer/Private/RayTracing/RayTracingLighting.cpp b/Engine/Source/Runtime/Renderer/Private/RayTracing/RayTracingLighting.cpp
--- a/Engine/Source/Runtime/Renderer/Private/RayTracing/RayTracingLighting.cpp
+++ b/Engine/Source/Runtime/Renderer/Private/RayTracing/RayTracingLighting.cpp
@@ -378,7 +378,7 @@ END_GLOBAL_SHADER_PARAMETER_STRUCT()
 
 IMPLEMENT_GLOBAL_SHADER_PARAMETER_STRUCT(FLightFunctionParametersRayTracing, "RaytracingLightFunctionParameters");
 
-static TUniformBufferRef<FLightFunctionParametersRayTracing> CreateLightFunctionParametersBuffer(
+static TUniformBufferRef<FLightFunctionParametersRayTracing> CreateLightFunctionParametersBufferRT(
 	const FLightSceneInfo* LightSceneInfo,
 	const FSceneView& View,
 	EUniformBufferUsage Usage)
@@ -584,7 +584,7 @@ void BindLightFunctionShaders(
 
 		//create the uniform buffers we need
 		TUniformBufferRef<FDeferredLightUniformStruct> DeferredLightBuffer = CreateUniformBufferImmediate(GetDeferredLightParameters(View, *LightSceneInfo), EUniformBufferUsage::UniformBuffer_SingleFrame);
-		TUniformBufferRef<FLightFunctionParametersRayTracing> LightFunctionParameters = CreateLightFunctionParametersBuffer(LightSceneInfo, View, EUniformBufferUsage::UniformBuffer_SingleFrame);
+		TUniformBufferRef<FLightFunctionParametersRayTracing> LightFunctionParameters = CreateLightFunctionParametersBufferRT(LightSceneInfo, View, EUniformBufferUsage::UniformBuffer_SingleFrame);
 
 		int32 MissIndex = LightAndIndex.Value;
 		BindLightFunction(RHICmdList, Scene, View, Material, MaterialRenderProxy, DeferredLightBuffer, LightFunctionParameters, MissIndex);
