Fix crash in light function rendering when hair pluging is disabled.

diff --git a/Engine/Source/Runtime/Renderer/Private/LightFunctionRendering.cpp b/Engine/Source/Runtime/Renderer/Private/LightFunctionRendering.cpp
--- a/Engine/Source/Runtime/Renderer/Private/LightFunctionRendering.cpp
+++ b/Engine/Source/Runtime/Renderer/Private/LightFunctionRendering.cpp
@@ -428,7 +428,7 @@ bool FDeferredShadingSceneRenderer::RenderLightFunctionForMaterial(
 						LightSceneProxy->SetScissorRect(RHICmdList, View, TotalRect);
 					}
 
-					FRHIUniformBuffer* HairStrandsUniformBuffer = PassParameters->HairStrands.GetUniformBuffer()->GetRHI();
+					FRHIUniformBuffer* HairStrandsUniformBuffer = PassParameters->HairStrands.GetUniformBuffer() ? PassParameters->HairStrands.GetUniformBuffer()->GetRHI() : nullptr;
 
 					// Render a bounding light sphere.
 					SetGraphicsPipelineState(RHICmdList, GraphicsPSOInit, 0);
