Illuminance meter no is measure over a patch of surface, keeping its normal but assuming a white surface.

diff --git a/Engine/Shaders/Private/DebugProbes.usf b/Engine/Shaders/Private/DebugProbes.usf
--- a/Engine/Shaders/Private/DebugProbes.usf
+++ b/Engine/Shaders/Private/DebugProbes.usf
@@ -234,7 +234,7 @@ void MainPS(
 	const float3 RayDir = normalize(TranslatedWorldPos.xyz - View.TranslatedWorldCameraOrigin);
 	const float3 RayOrig = LWCHackToFloat(ResolvedView.WorldCameraOrigin);
 
-#if PERMUTATION_VISUALIZEHDR /////////////////////////////////////////////////////////////////////////////////////////////////////////////////////////////////////
+#if PERMUTATION_ILLUMINANCEMETER /////////////////////////////////////////////////////////////////////////////////////////////////////////////////////////////////////
 
 	// In this case we want to show a small patch in front of the camera with pure white albedo.
 	// That is a way to measure the illuminance at the camera position, on the hemisphere pointing otwards the back of the camera.
@@ -307,7 +307,7 @@ void MainPS(
 
 
 
-#else // PERMUTATION_VISUALIZEHDR /////////////////////////////////////////////////////////////////////////////////////////////////////////////////////////////////////
+#else // PERMUTATION_ILLUMINANCEMETER /////////////////////////////////////////////////////////////////////////////////////////////////////////////////////////////////////
 
 
 
@@ -502,7 +502,7 @@ void MainPS(
 
 #endif // PERMUTATION_PASS
 
-#endif // PERMUTATION_VISUALIZEHDR /////////////////////////////////////////////////////////////////////////////////////////////////////////////////////////////////////
+#endif // PERMUTATION_ILLUMINANCEMETER /////////////////////////////////////////////////////////////////////////////////////////////////////////////////////////////////////
 
 #if DEPTH_OUTPUT_ENABLED
 	OutDepth = DeviceZ;
diff --git a/Engine/Shaders/Private/PostProcessVisualizeHDR.usf b/Engine/Shaders/Private/PostProcessVisualizeHDR.usf
--- a/Engine/Shaders/Private/PostProcessVisualizeHDR.usf
+++ b/Engine/Shaders/Private/PostProcessVisualizeHDR.usf
@@ -22,6 +22,7 @@ Texture2D HistogramTexture;
 SCREEN_PASS_TEXTURE_VIEWPORT(Input)
 SCREEN_PASS_TEXTURE_VIEWPORT(Output)
 
+float IlluminanceMeterEnabled;
 
 float ComputeRgbLuminance(float3 Color)
 {
@@ -292,6 +293,7 @@ void MainPS(noperspective float4 UVAndScreenPos : TEXCOORD0, float4 SvPosition :
 	}
 
 	// Illuminance meter for the hemisphere in front of the camera
+	if(IlluminanceMeterEnabled > 0.0f)
 	{
 		OutColor = float4(0.0f, 0.0f, 0.0f, 0.0f);
 
diff --git a/Engine/Source/Runtime/Renderer/Private/DebugProbeRendering.cpp b/Engine/Source/Runtime/Renderer/Private/DebugProbeRendering.cpp
--- a/Engine/Source/Runtime/Renderer/Private/DebugProbeRendering.cpp
+++ b/Engine/Source/Runtime/Renderer/Private/DebugProbeRendering.cpp
@@ -29,6 +29,7 @@ static TAutoConsoleVariable<int32> CVarVisualizeLightingOnProbes(
 
 DECLARE_GPU_STAT(StampDeferredDebugProbe);
 
+extern bool IsIlluminanceMeterSupportedByView(const FViewInfo& View);
 
 // Must match DebugProbes.usf
 #define RENDER_DEPTHPREPASS  0
@@ -42,8 +43,8 @@ class FStampDeferredDebugProbePS : public FGlobalShader
 	SHADER_USE_PARAMETER_STRUCT(FStampDeferredDebugProbePS, FGlobalShader);
 
 	class FRenderPass : SHADER_PERMUTATION_RANGE_INT("PERMUTATION_PASS", 0, 3);
-	class FVisualizeHDR : SHADER_PERMUTATION_BOOL("PERMUTATION_VISUALIZEHDR");
-	using FPermutationDomain = TShaderPermutationDomain<FRenderPass, FVisualizeHDR>;
+	class FIlluminanceMeter : SHADER_PERMUTATION_BOOL("PERMUTATION_ILLUMINANCEMETER");
+	using FPermutationDomain = TShaderPermutationDomain<FRenderPass, FIlluminanceMeter>;
 
 	BEGIN_SHADER_PARAMETER_STRUCT(FParameters, )
 		SHADER_PARAMETER_STRUCT_REF(FViewUniformShaderParameters, ViewUniformBuffer)
@@ -77,14 +78,12 @@ IMPLEMENT_GLOBAL_SHADER(FStampDeferredDebugProbePS, "/Engine/Private/DebugProbes
 
 #if DEBUG_PROBE_ENABLED
 
-static bool AnyViewRequireVisualizeHDR(TArrayView<const FViewInfo> Views)
+static bool ViewRequiresAndSupportsIlluminanceMeter(const FViewInfo& View)
 {
-	for (const FViewInfo& View : Views)
+	if (View.Family->EngineShowFlags.VisualizeHDR		// This is the debug view contain the illuminance meter
+		&& IsIlluminanceMeterSupportedByView(View))		// For instance: forward shading does not work with illuminance meter (we need a material buffer to work with)
 	{
-		if (View.Family->EngineShowFlags.VisualizeHDR)
-		{
-			return true;
-		}
+		return true;
 	}
 	return false;
 }
@@ -95,7 +94,7 @@ static void CommonStampDeferredDebugProbeDrawCall(
 	const FViewInfo& View, 
 	FStampDeferredDebugProbePS::FParameters* PassParameters,
 	int32 RenderPass,
-	bool bVisualizeHDR)
+	bool bIlluminanceMeter)
 {
 	PassParameters->ViewUniformBuffer = View.ViewUniformBuffer;
 	PassParameters->MaterialTextureArrayUAV = View.StrataViewData.SceneData->MaterialTextureArrayUAVWithoutRTs;
@@ -105,7 +104,7 @@ static void CommonStampDeferredDebugProbeDrawCall(
 		
 	FStampDeferredDebugProbePS::FPermutationDomain PermutationVector;
 	PermutationVector.Set<FStampDeferredDebugProbePS::FRenderPass>(RenderPass);
-	PermutationVector.Set<FStampDeferredDebugProbePS::FVisualizeHDR>(bVisualizeHDR);
+	PermutationVector.Set<FStampDeferredDebugProbePS::FIlluminanceMeter>(bIlluminanceMeter);
 	TShaderMapRef<FStampDeferredDebugProbePS> PixelShader(View.ShaderMap, PermutationVector);
 
 	FPixelShaderUtils::AddFullscreenPass<FStampDeferredDebugProbePS>(
@@ -127,19 +126,21 @@ void StampDeferredDebugProbeDepthPS(
 	RDG_EVENT_SCOPE(GraphBuilder, "StampDeferredDebugProbeDepth");
 	RDG_GPU_STAT_SCOPE(GraphBuilder, StampDeferredDebugProbe);
 
-	const bool bVisualizeHDR = AnyViewRequireVisualizeHDR(Views);
 	const bool bVisualizeLightingOnProbes = CVarVisualizeLightingOnProbes.GetValueOnRenderThread() > 0;
 	for (const FViewInfo& View : Views)
 	{
-		if (!(bVisualizeLightingOnProbes || View.Family->EngineShowFlags.VisualizeLightingOnProbes || bVisualizeHDR) || View.bIsReflectionCapture)
+		const bool bViewRequiresAndSupportsIlluminanceMeter = ViewRequiresAndSupportsIlluminanceMeter(View);
+		if (!(bVisualizeLightingOnProbes || View.Family->EngineShowFlags.VisualizeLightingOnProbes) || View.bIsReflectionCapture || bViewRequiresAndSupportsIlluminanceMeter)
 		{
+			// When "Visualizing HDR with Illuminance Meter", we want to evaluated the illuminance on the surface behind the target square patch.  
+			// So we do not want write depth or velocity in this case.
 			continue;
 		}
 
 		FStampDeferredDebugProbePS::FParameters* PassParameters = GraphBuilder.AllocParameters<FStampDeferredDebugProbePS::FParameters>();
 		PassParameters->RenderTargets.DepthStencil = FDepthStencilBinding(SceneDepthTexture, ERenderTargetLoadAction::ELoad, ERenderTargetLoadAction::ELoad, FExclusiveDepthStencil::DepthWrite_StencilWrite);
 
-		CommonStampDeferredDebugProbeDrawCall<true, CF_DepthNearOrEqual>(GraphBuilder, View, PassParameters, RENDER_DEPTHPREPASS, bVisualizeHDR);
+		CommonStampDeferredDebugProbeDrawCall<true, CF_DepthNearOrEqual>(GraphBuilder, View, PassParameters, RENDER_DEPTHPREPASS, bViewRequiresAndSupportsIlluminanceMeter);
 	}
 #endif // DEBUG_PROBE_ENABLED
 }
@@ -154,28 +155,52 @@ void StampDeferredDebugProbeMaterialPS(
 	RDG_EVENT_SCOPE(GraphBuilder, "StampDeferredDebugProbeMaterial");
 	RDG_GPU_STAT_SCOPE(GraphBuilder, StampDeferredDebugProbe);
 
-	const bool bVisualizeHDR = AnyViewRequireVisualizeHDR(Views);
 	const bool bVisualizeLightingOnProbes = CVarVisualizeLightingOnProbes.GetValueOnRenderThread() > 0;
 	for (const FViewInfo& View : Views)
 	{
-		if (!(bVisualizeLightingOnProbes || View.Family->EngineShowFlags.VisualizeLightingOnProbes || bVisualizeHDR) || View.bIsReflectionCapture)
+		const bool bViewRequiresAndSupportsIlluminanceMeter = ViewRequiresAndSupportsIlluminanceMeter(View);
+		if (!(bVisualizeLightingOnProbes || View.Family->EngineShowFlags.VisualizeLightingOnProbes || bViewRequiresAndSupportsIlluminanceMeter) || View.bIsReflectionCapture)
 		{
 			continue;
 		}
 
 		FStampDeferredDebugProbePS::FParameters* PassParameters = GraphBuilder.AllocParameters<FStampDeferredDebugProbePS::FParameters>();
 		PassParameters->RenderTargets = BasePassRenderTargets;
+
+		if (bViewRequiresAndSupportsIlluminanceMeter)
+		{
+			// We do not want to update the normal of the material in this case,
+			// because in this case we want to use the scene normal as the direction of the hemisphere over which we evaluate Illuminance.
+			// So for legacy, we unbind the normal buffer and for strata we unbind the toplayer data (that is enough since the material written is going to be of simple type)
+			if (Strata::IsStrataEnabled())
+			{
+				// Search for the TopLayerTexture render target and nullify it to not update it.
+				for (int32 i = 0; i < MaxSimultaneousRenderTargets; ++i)
+				{
+					if (PassParameters->RenderTargets.Output[i].GetTexture() == View.StrataViewData.SceneData->TopLayerTexture)
+					{
+						PassParameters->RenderTargets.Output[i] = FRenderTargetBinding();
+					}
+				}
+			}
+			else
+			{
+				// Do not write to the GBuffer normal render target, use a dummy render target to avoid validation issue with target textures that must be packed.
+				PassParameters->RenderTargets.Output[1] = FRenderTargetBinding(GraphBuilder.CreateTexture(PassParameters->RenderTargets.Output[1].GetTexture()->Desc, TEXT("DummyGBufferNormalTexture")), ERenderTargetLoadAction::ELoad);
+			}
+		}
+
 		if (Strata::IsStrataEnabled())
 		{
 			// Make sure we do not write depth so that we can safely read it from texture parameters
 			PassParameters->RenderTargets.DepthStencil = FDepthStencilBinding();
 			PassParameters->SceneTextures = CreateSceneTextureShaderParameters(GraphBuilder, View.GetSceneTexturesChecked(), View.GetFeatureLevel(), ESceneTextureSetupMode::SceneDepth);
 
-			CommonStampDeferredDebugProbeDrawCall<false, CF_Always>(GraphBuilder, View, PassParameters, RENDER_BASEPASS, bVisualizeHDR);
+			CommonStampDeferredDebugProbeDrawCall<false, CF_Always>(GraphBuilder, View, PassParameters, RENDER_BASEPASS, bViewRequiresAndSupportsIlluminanceMeter);
 		}
 		else
 		{
-			CommonStampDeferredDebugProbeDrawCall<false, CF_DepthNearOrEqual>(GraphBuilder, View, PassParameters, RENDER_BASEPASS, bVisualizeHDR);
+			CommonStampDeferredDebugProbeDrawCall<false, CF_DepthNearOrEqual>(GraphBuilder, View, PassParameters, RENDER_BASEPASS, bViewRequiresAndSupportsIlluminanceMeter);
 		}
 	}
 #endif // DEBUG_PROBE_ENABLED
@@ -190,12 +215,14 @@ void StampDeferredDebugProbeVelocityPS(
 	RDG_EVENT_SCOPE(GraphBuilder, "StampDeferredDebugProbeVelocity");
 	RDG_GPU_STAT_SCOPE(GraphBuilder, StampDeferredDebugProbe);
 
-	const bool bVisualizeHDR = AnyViewRequireVisualizeHDR(Views);
 	const bool bVisualizeLightingOnProbes = CVarVisualizeLightingOnProbes.GetValueOnRenderThread() > 0;
 	for (const FViewInfo& View : Views)
 	{
-		if (!(bVisualizeLightingOnProbes || View.Family->EngineShowFlags.VisualizeLightingOnProbes || bVisualizeHDR) || View.bIsReflectionCapture)
+		const bool bViewRequiresAndSupportsIlluminanceMeter = ViewRequiresAndSupportsIlluminanceMeter(View);
+		if (!(bVisualizeLightingOnProbes || View.Family->EngineShowFlags.VisualizeLightingOnProbes) || View.bIsReflectionCapture || bViewRequiresAndSupportsIlluminanceMeter)
 		{
+			// When "Visualizing HDR with Illuminance Meter", we want to evaluated the illuminance on the surface behind the target square patch. 
+			// So we do not want write depth or velocity in this case.
 			continue;
 		}
 
@@ -203,7 +230,7 @@ void StampDeferredDebugProbeVelocityPS(
 		PassParameters->RenderTargets = BasePassRenderTargets;
 
 		const bool bRenderVelocity = true;
-		CommonStampDeferredDebugProbeDrawCall<false, CF_DepthNearOrEqual>(GraphBuilder, View, PassParameters, RENDER_VELOCITYPASS, bVisualizeHDR);
+		CommonStampDeferredDebugProbeDrawCall<false, CF_DepthNearOrEqual>(GraphBuilder, View, PassParameters, RENDER_VELOCITYPASS, bViewRequiresAndSupportsIlluminanceMeter);
 	}
 #endif // DEBUG_PROBE_ENABLED
 }
diff --git a/Engine/Source/Runtime/Renderer/Private/PostProcess/PostProcessVisualizeHDR.cpp b/Engine/Source/Runtime/Renderer/Private/PostProcess/PostProcessVisualizeHDR.cpp
--- a/Engine/Source/Runtime/Renderer/Private/PostProcess/PostProcessVisualizeHDR.cpp
+++ b/Engine/Source/Runtime/Renderer/Private/PostProcess/PostProcessVisualizeHDR.cpp
@@ -31,6 +31,7 @@ public:
 		SHADER_PARAMETER(float, FilmShoulder)
 		SHADER_PARAMETER(float, FilmBlackClip)
 		SHADER_PARAMETER(float, FilmWhiteClip)
+		SHADER_PARAMETER(float, IlluminanceMeterEnabled)
 		RENDER_TARGET_BINDING_SLOTS()
 	END_SHADER_PARAMETER_STRUCT()
 
@@ -51,6 +52,11 @@ public:
 
 IMPLEMENT_GLOBAL_SHADER(FVisualizeHDRPS, "/Engine/Private/PostProcessVisualizeHDR.usf", "MainPS", SF_Pixel);
 
+bool IsIlluminanceMeterSupportedByView(const FViewInfo& View)
+{
+	return !IsForwardShadingEnabled(View.GetShaderPlatform()) && !View.Family->EngineShowFlags.PathTracing;
+}
+
 FScreenPassTexture AddVisualizeHDRPass(FRDGBuilder& GraphBuilder, const FViewInfo& View, const FVisualizeHDRInputs& Inputs)
 {
 	check(Inputs.SceneColor.IsValid());
@@ -92,6 +98,8 @@ FScreenPassTexture AddVisualizeHDRPass(FRDGBuilder& GraphBuilder, const FViewInf
 	PassParameters->FilmShoulder = Settings.FilmShoulder;
 	PassParameters->FilmBlackClip = Settings.FilmBlackClip;
 	PassParameters->FilmWhiteClip = Settings.FilmWhiteClip;
+	PassParameters->IlluminanceMeterEnabled = IsIlluminanceMeterSupportedByView(View) ? 1.0f : 0.0f;
+	
 
 	TShaderMapRef<FVisualizeHDRPS> PixelShader(View.ShaderMap);
 
@@ -207,7 +215,15 @@ FScreenPassTexture AddVisualizeHDRPass(FRDGBuilder& GraphBuilder, const FViewInf
 
 			const float IlluminanceMeterTextX = Output.ViewRect.Min.X + Output.ViewRect.Size().X * 0.5f - 180.0f;
 			const float IlluminanceMeterTextY = Output.ViewRect.Min.Y + Output.ViewRect.Size().Y * 0.5f - 130.0f;
-			Canvas.DrawShadowedString(IlluminanceMeterTextX, IlluminanceMeterTextY, TEXT("Illuminance meter measuring the back hemisphere of the camera"), GetStatsFont(), FLinearColor(1, 1, 1));
+			if(IsIlluminanceMeterSupportedByView(View))
+			{
+				Canvas.DrawShadowedString(IlluminanceMeterTextX, IlluminanceMeterTextY, TEXT("Illuminance meter - over the hemisphere of the surface patch"), GetStatsFont(), FLinearColor(1, 1, 1));
+			}
+			else
+			{
+				Canvas.DrawShadowedString(IlluminanceMeterTextX, IlluminanceMeterTextY, TEXT("Illuminance meter - not available with Forward Shading"), GetStatsFont(), FLinearColor(1, 1, 1));
+			}
+			
 
 			const float LuminanceMeterTextX = Output.ViewRect.Min.X + Output.ViewRect.Size().X * 0.5f - 50.0f;
 			const float LuminanceMeterTextY = Output.ViewRect.Min.Y + Output.ViewRect.Size().Y * 0.5f - 30.0f;
