Merged PreRender view extension RDG builder with scene render RDG builder.
 - Ported remaining ISceneViewExtension calls to RDG.
 - Removed several empty override passes being added to the graph.
 - Merged two graphs into one.

diff --git a/Engine/Plugins/Compositing/Composure/Source/Composure/Private/ComposureViewExtension.h b/Engine/Plugins/Compositing/Composure/Source/Composure/Private/ComposureViewExtension.h
--- a/Engine/Plugins/Compositing/Composure/Source/Composure/Private/ComposureViewExtension.h
+++ b/Engine/Plugins/Compositing/Composure/Source/Composure/Private/ComposureViewExtension.h
@@ -21,8 +21,8 @@ public:
 	virtual void SetupViewFamily(FSceneViewFamily& InViewFamily) override {}
 	virtual void SetupView(FSceneViewFamily& InViewFamily, FSceneView& InView) override {}
 	virtual void BeginRenderViewFamily(FSceneViewFamily& InViewFamily) override;
-	virtual void PreRenderView_RenderThread(FRHICommandListImmediate& RHICmdList, FSceneView& InView) override {}
-	virtual void PreRenderViewFamily_RenderThread(FRHICommandListImmediate& RHICmdList, FSceneViewFamily& InViewFamily) override {}
+	virtual void PreRenderView_RenderThread(FRDGBuilder& RHICmdList, FSceneView& InView) override {}
+	virtual void PreRenderViewFamily_RenderThread(FRDGBuilder& RHICmdList, FSceneViewFamily& InViewFamily) override {}
 	virtual int32 GetPriority() const override;
 	virtual bool IsActiveThisFrame_Internal(const FSceneViewExtensionContext& Context) const override;
 
diff --git a/Engine/Plugins/Compositing/Composure/Source/Composure/Private/MovieScene/MovieSceneComposureExportSectionTemplate.cpp b/Engine/Plugins/Compositing/Composure/Source/Composure/Private/MovieScene/MovieSceneComposureExportSectionTemplate.cpp
--- a/Engine/Plugins/Compositing/Composure/Source/Composure/Private/MovieScene/MovieSceneComposureExportSectionTemplate.cpp
+++ b/Engine/Plugins/Compositing/Composure/Source/Composure/Private/MovieScene/MovieSceneComposureExportSectionTemplate.cpp
@@ -61,8 +61,8 @@ private:
 	virtual bool IsActiveThisFrame_Internal(const FSceneViewExtensionContext& Context) const override;
 
 	virtual void BeginRenderViewFamily(FSceneViewFamily& InViewFamily) override {}
-	virtual void PreRenderViewFamily_RenderThread(FRHICommandListImmediate& RHICmdList, FSceneViewFamily& InViewFamily) override {}
-	virtual void PreRenderView_RenderThread(FRHICommandListImmediate& RHICmdList, FSceneView& InView) override {}
+	virtual void PreRenderViewFamily_RenderThread(FRDGBuilder& RHICmdList, FSceneViewFamily& InViewFamily) override {}
+	virtual void PreRenderView_RenderThread(FRDGBuilder& RHICmdList, FSceneView& InView) override {}
 
 	/** Makes a new endpoint for an image pixel pipe that just forwards the pixels onto the capture protocol */
 	static TFunction<void(TUniquePtr<FImagePixelData>&&)> MakeForwardingEndpoint(const FCapturedPixelsID& BufferID, const FFrameMetrics& CurrentFrameMetrics);
diff --git a/Engine/Plugins/Compositing/OpenColorIO/Source/OpenColorIO/Public/OpenColorIODisplayExtension.h b/Engine/Plugins/Compositing/OpenColorIO/Source/OpenColorIO/Public/OpenColorIODisplayExtension.h
--- a/Engine/Plugins/Compositing/OpenColorIO/Source/OpenColorIO/Public/OpenColorIODisplayExtension.h
+++ b/Engine/Plugins/Compositing/OpenColorIO/Source/OpenColorIO/Public/OpenColorIODisplayExtension.h
@@ -31,8 +31,8 @@ public:
 	virtual void SetupView(FSceneViewFamily& InViewFamily, FSceneView& InView) override;
 	virtual void BeginRenderViewFamily(FSceneViewFamily& InViewFamily) override {};
 
-	virtual void PreRenderViewFamily_RenderThread(FRHICommandListImmediate& RHICmdList, FSceneViewFamily& InViewFamily) {};
-	virtual void PreRenderView_RenderThread(FRHICommandListImmediate& RHICmdList, FSceneView& InView) {};
+	virtual void PreRenderViewFamily_RenderThread(FRDGBuilder& GraphBuilder, FSceneViewFamily& InViewFamily) {};
+	virtual void PreRenderView_RenderThread(FRDGBuilder& GraphBuilder, FSceneView& InView) {};
 
 	virtual void SubscribeToPostProcessingPass(EPostProcessingPass PassId, FAfterPassCallbackDelegateArray& InOutPassCallbacks, bool bIsPassEnabled) override;
 	virtual bool IsActiveThisFrame_Internal(const FSceneViewExtensionContext& Context) const override;
diff --git a/Engine/Plugins/Experimental/AutomationUtils/Source/AutomationUtils/Private/AutomationUtilsBlueprintLibrary.cpp b/Engine/Plugins/Experimental/AutomationUtils/Source/AutomationUtils/Private/AutomationUtilsBlueprintLibrary.cpp
--- a/Engine/Plugins/Experimental/AutomationUtils/Source/AutomationUtils/Private/AutomationUtilsBlueprintLibrary.cpp
+++ b/Engine/Plugins/Experimental/AutomationUtils/Source/AutomationUtils/Private/AutomationUtilsBlueprintLibrary.cpp
@@ -56,8 +56,8 @@ public:
 
 	void SetupView(FSceneViewFamily& InViewFamily, FSceneView& InView) override {}
 	void BeginRenderViewFamily(FSceneViewFamily& InViewFamily) override {}
-	void PreRenderViewFamily_RenderThread(FRHICommandListImmediate& RHICmdList, FSceneViewFamily& InViewFamily) override {}
-	void PreRenderView_RenderThread(FRHICommandListImmediate& RHICmdList, FSceneView& InView) override {}
+	void PreRenderViewFamily_RenderThread(FRDGBuilder& GraphBuilder, FSceneViewFamily& InViewFamily) override {}
+	void PreRenderView_RenderThread(FRDGBuilder& GraphBuilder, FSceneView& InView) override {}
 
 	/** We always want to go last. */
 	virtual int32 GetPriority() const override { return MIN_int32; }
diff --git a/Engine/Plugins/Experimental/ColorCorrectRegions/Source/ColorCorrectRegions/Public/ColorCorrectRegionsSceneViewExtension.h b/Engine/Plugins/Experimental/ColorCorrectRegions/Source/ColorCorrectRegions/Public/ColorCorrectRegionsSceneViewExtension.h
--- a/Engine/Plugins/Experimental/ColorCorrectRegions/Source/ColorCorrectRegions/Public/ColorCorrectRegionsSceneViewExtension.h
+++ b/Engine/Plugins/Experimental/ColorCorrectRegions/Source/ColorCorrectRegions/Public/ColorCorrectRegionsSceneViewExtension.h
@@ -17,9 +17,6 @@ public:
 	virtual void SetupViewFamily(FSceneViewFamily& InViewFamily) override {};
 	virtual void SetupView(FSceneViewFamily& InViewFamily, FSceneView& InView) override {};
 	virtual void BeginRenderViewFamily(FSceneViewFamily& InViewFamily) override {};
-	virtual void PreRenderViewFamily_RenderThread(FRHICommandListImmediate& RHICmdList, FSceneViewFamily& InViewFamily) override {};
-	virtual void PreRenderView_RenderThread(FRHICommandListImmediate& RHICmdList, FSceneView& InView) override {};
-	virtual void PostRenderBasePass_RenderThread(FRHICommandListImmediate& RHICmdList, FSceneView& InView) override {};
 	virtual void PrePostProcessPass_RenderThread(FRDGBuilder& GraphBuilder, const FSceneView& View, const FPostProcessingInputs& Inputs) override;
 	//~ End FSceneViewExtensionBase Interface
 
diff --git a/Engine/Plugins/Experimental/RemoteSession/Source/RemoteSession/Private/Channels/RemoteSessionARCameraChannel.cpp b/Engine/Plugins/Experimental/RemoteSession/Source/RemoteSession/Private/Channels/RemoteSessionARCameraChannel.cpp
--- a/Engine/Plugins/Experimental/RemoteSession/Source/RemoteSession/Private/Channels/RemoteSessionARCameraChannel.cpp
+++ b/Engine/Plugins/Experimental/RemoteSession/Source/RemoteSession/Private/Channels/RemoteSessionARCameraChannel.cpp
@@ -145,8 +145,8 @@ private:
 	virtual void SetupViewFamily(FSceneViewFamily& InViewFamily) override;
 	virtual void SetupView(FSceneViewFamily& InViewFamily, FSceneView& InView) override;
 	virtual void BeginRenderViewFamily(FSceneViewFamily& InViewFamily) override;
-	virtual void PreRenderView_RenderThread(FRHICommandListImmediate& RHICmdList, FSceneView& InView) override;
-	virtual void PreRenderViewFamily_RenderThread(FRHICommandListImmediate& RHICmdList, FSceneViewFamily& InViewFamily) override;
+	virtual void PreRenderView_RenderThread(FRDGBuilder& GraphBuilder, FSceneView& InView) override;
+	virtual void PreRenderViewFamily_RenderThread(FRDGBuilder& GraphBuilder, FSceneViewFamily& InViewFamily) override;
 	virtual void PostRenderViewFamily_RenderThread(FRDGBuilder& GraphBuilder, FSceneViewFamily& InViewFamily) override;
 	virtual bool IsActiveThisFrame_Internal(const FSceneViewExtensionContext& Context) const override;
 	//~ISceneViewExtension interface
@@ -185,7 +185,7 @@ void FARCameraSceneViewExtension::BeginRenderViewFamily(FSceneViewFamily& InView
 
 }
 
-void FARCameraSceneViewExtension::PreRenderView_RenderThread(FRHICommandListImmediate& RHICmdList, FSceneView& InView)
+void FARCameraSceneViewExtension::PreRenderView_RenderThread(FRDGBuilder& GraphBuilder, FSceneView& InView)
 {
 	if (VertexBufferRHI == nullptr || !VertexBufferRHI.IsValid())
 	{
@@ -226,7 +226,7 @@ void FARCameraSceneViewExtension::PreRenderView_RenderThread(FRHICommandListImme
 	PPMaterial = Channel.GetPostProcessMaterial();
 }
 
-void FARCameraSceneViewExtension::PreRenderViewFamily_RenderThread(FRHICommandListImmediate& RHICmdList, FSceneViewFamily& InViewFamily)
+void FARCameraSceneViewExtension::PreRenderViewFamily_RenderThread(FRDGBuilder& GraphBuilder, FSceneViewFamily& InViewFamily)
 {
 
 }
diff --git a/Engine/Plugins/Experimental/SimpleHMD/Source/SimpleHMD/Private/SimpleHMD.cpp b/Engine/Plugins/Experimental/SimpleHMD/Source/SimpleHMD/Private/SimpleHMD.cpp
--- a/Engine/Plugins/Experimental/SimpleHMD/Source/SimpleHMD/Private/SimpleHMD.cpp
+++ b/Engine/Plugins/Experimental/SimpleHMD/Source/SimpleHMD/Private/SimpleHMD.cpp
@@ -295,12 +295,12 @@ void FSimpleHMD::SetupView(FSceneViewFamily& InViewFamily, FSceneView& InView)
 	InView.BaseHmdLocation = FVector(0.f);
 }
 
-void FSimpleHMD::PreRenderView_RenderThread(FRHICommandListImmediate& RHICmdList, FSceneView& InView)
+void FSimpleHMD::PreRenderView_RenderThread(FRDGBuilder& GraphBuilder, FSceneView& InView)
 {
 	check(IsInRenderingThread());
 }
 
-void FSimpleHMD::PreRenderViewFamily_RenderThread(FRHICommandListImmediate& RHICmdList, FSceneViewFamily& ViewFamily)
+void FSimpleHMD::PreRenderViewFamily_RenderThread(FRDGBuilder& GraphBuilder, FSceneViewFamily& ViewFamily)
 {
 	check(IsInRenderingThread());
 }
diff --git a/Engine/Plugins/Experimental/SimpleHMD/Source/SimpleHMD/Private/SimpleHMD.h b/Engine/Plugins/Experimental/SimpleHMD/Source/SimpleHMD/Private/SimpleHMD.h
--- a/Engine/Plugins/Experimental/SimpleHMD/Source/SimpleHMD/Private/SimpleHMD.h
+++ b/Engine/Plugins/Experimental/SimpleHMD/Source/SimpleHMD/Private/SimpleHMD.h
@@ -83,8 +83,8 @@ public:
 	virtual void SetupViewFamily(FSceneViewFamily& InViewFamily) override;
 	virtual void SetupView(FSceneViewFamily& InViewFamily, FSceneView& InView) override;
 	virtual void BeginRenderViewFamily(FSceneViewFamily& InViewFamily) {}
-	virtual void PreRenderView_RenderThread(FRHICommandListImmediate& RHICmdList, FSceneView& InView) override;
-	virtual void PreRenderViewFamily_RenderThread(FRHICommandListImmediate& RHICmdList, FSceneViewFamily& InViewFamily) override;
+	virtual void PreRenderView_RenderThread(FRDGBuilder& GraphBuilder, FSceneView& InView) override;
+	virtual void PreRenderViewFamily_RenderThread(FRDGBuilder& GraphBuilder, FSceneViewFamily& InViewFamily) override;
 
 public:
 	/** Constructor */
diff --git a/Engine/Plugins/Experimental/Water/Source/Runtime/Private/GerstnerWaterWaveViewExtension.cpp b/Engine/Plugins/Experimental/Water/Source/Runtime/Private/GerstnerWaterWaveViewExtension.cpp
--- a/Engine/Plugins/Experimental/Water/Source/Runtime/Private/GerstnerWaterWaveViewExtension.cpp
+++ b/Engine/Plugins/Experimental/Water/Source/Runtime/Private/GerstnerWaterWaveViewExtension.cpp
@@ -124,7 +124,7 @@ void FGerstnerWaterWaveViewExtension::SetupViewFamily(FSceneViewFamily& InViewFa
 	}
 }
 
-void FGerstnerWaterWaveViewExtension::PreRenderView_RenderThread(FRHICommandListImmediate& RHICmdList, FSceneView& InView)
+void FGerstnerWaterWaveViewExtension::PreRenderView_RenderThread(FRDGBuilder& GraphBuilder, FSceneView& InView)
 {
 	if (WaveGPUData->DataSRV && WaveGPUData->IndirectionSRV)
 	{
diff --git a/Engine/Plugins/Experimental/Water/Source/Runtime/Private/WaterInfoRendering.cpp b/Engine/Plugins/Experimental/Water/Source/Runtime/Private/WaterInfoRendering.cpp
--- a/Engine/Plugins/Experimental/Water/Source/Runtime/Private/WaterInfoRendering.cpp
+++ b/Engine/Plugins/Experimental/Water/Source/Runtime/Private/WaterInfoRendering.cpp
@@ -397,6 +397,7 @@ static void UpdateWaterInfoRendering_RenderThread(
 	FScopedLandscapeLODOverride ScopedLandscapeLODOverride(Params);
 
 	FMaterialRenderProxy::UpdateDeferredCachedUniformExpressions();
+	FDeferredUpdateResource::UpdateResources(RHICmdList);
 
 	FRenderTarget* RenderTarget = Params.RenderTarget;
 	FTexture* OutputTexture = Params.OutputTexture;
@@ -408,16 +409,14 @@ static void UpdateWaterInfoRendering_RenderThread(
 	{
 		FSceneRenderer* DepthRenderer = Params.DepthRenderer;
 
-		// We need to execute the pre-render view extensions before we do any view dependent work.
-		FSceneRenderer::ViewExtensionPreRender_RenderThread(RHICmdList, DepthRenderer);
-
 		DepthRenderer->RenderThreadBegin(RHICmdList);
-		
-		FDeferredUpdateResource::UpdateResources(RHICmdList);
-		
+
 		SCOPED_DRAW_EVENT(RHICmdList, DepthRendering_RT);
-		
+
 		FRDGBuilder GraphBuilder(RHICmdList, RDG_EVENT_NAME("WaterInfoColorRendering"), ERDGBuilderFlags::AllowParallelExecute);
+
+		// We need to execute the pre-render view extensions before we do any view dependent work.
+		FSceneRenderer::ViewExtensionPreRender_RenderThread(GraphBuilder, DepthRenderer);
 		
 		FRDGTextureRef TargetTexture = RegisterExternalTexture(GraphBuilder, RenderTarget->GetRenderTargetTexture(), TEXT("WaterDepthTarget"));
 
@@ -466,16 +465,17 @@ static void UpdateWaterInfoRendering_RenderThread(
 	{
 		FSceneRenderer* ColorRenderer = Params.ColorRenderer;
 
-		// We need to execute the pre-render view extensions before we do any view dependent work.
-		FSceneRenderer::ViewExtensionPreRender_RenderThread(RHICmdList, ColorRenderer);
-
 		ColorRenderer->RenderThreadBegin(RHICmdList);
-		SetWaterBodiesWithinWaterInfoPass(ColorRenderer, EWaterInfoPass::Color);
-		
+
 		SCOPED_DRAW_EVENT(RHICmdList, ColorRendering_RT);
-		
+
 		FRDGBuilder GraphBuilder(RHICmdList, RDG_EVENT_NAME("WaterInfoColorRendering"), ERDGBuilderFlags::AllowParallelExecute);
 
+		// We need to execute the pre-render view extensions before we do any view dependent work.
+		FSceneRenderer::ViewExtensionPreRender_RenderThread(GraphBuilder, ColorRenderer);
+
+		SetWaterBodiesWithinWaterInfoPass(ColorRenderer, EWaterInfoPass::Color);
+		
 		FRDGTextureRef TargetTexture = RegisterExternalTexture(GraphBuilder, RenderTarget->GetRenderTargetTexture(), TEXT("WaterColorTarget"));
 
 		FRDGTextureDesc ColorTextureDesc(TargetTexture->Desc);
@@ -520,15 +520,17 @@ static void UpdateWaterInfoRendering_RenderThread(
 	{
 		FSceneRenderer* DilationRenderer = Params.DilationRenderer;
 
-		// We need to execute the pre-render view extensions before we do any view dependent work.
-		FSceneRenderer::ViewExtensionPreRender_RenderThread(RHICmdList, DilationRenderer);
-
 		DilationRenderer->RenderThreadBegin(RHICmdList);
-		SetWaterBodiesWithinWaterInfoPass(DilationRenderer, EWaterInfoPass::Dilation);
 
 		SCOPED_DRAW_EVENT(RHICmdList, DilationRendering_RT);
-		
+
 		FRDGBuilder GraphBuilder(RHICmdList, RDG_EVENT_NAME("WaterInfoColorRendering"), ERDGBuilderFlags::AllowParallelExecute);
+
+		// We need to execute the pre-render view extensions before we do any view dependent work.
+		FSceneRenderer::ViewExtensionPreRender_RenderThread(GraphBuilder, DilationRenderer);
+
+		SetWaterBodiesWithinWaterInfoPass(DilationRenderer, EWaterInfoPass::Dilation);
+
 		FRDGTextureRef TargetTexture = RegisterExternalTexture(GraphBuilder, RenderTarget->GetRenderTargetTexture(), TEXT("WaterDilationTarget"));
 
 		FRDGTextureRef DepthTexture = GraphBuilder.RegisterExternalTexture(ExtractedDepthTexture);
diff --git a/Engine/Plugins/Experimental/Water/Source/Runtime/Private/WaterViewExtension.cpp b/Engine/Plugins/Experimental/Water/Source/Runtime/Private/WaterViewExtension.cpp
--- a/Engine/Plugins/Experimental/Water/Source/Runtime/Private/WaterViewExtension.cpp
+++ b/Engine/Plugins/Experimental/Water/Source/Runtime/Private/WaterViewExtension.cpp
@@ -34,7 +34,7 @@ void FWaterViewExtension::SetupViewFamily(FSceneViewFamily& InViewFamily)
 	}
 }
 
-void FWaterViewExtension::PreRenderViewFamily_RenderThread(FRHICommandListImmediate& RHICmdList, FSceneViewFamily& InViewFamily)
+void FWaterViewExtension::PreRenderViewFamily_RenderThread(FRDGBuilder& GraphBuilder, FSceneViewFamily& InViewFamily)
 {
 }
 
diff --git a/Engine/Plugins/Experimental/Water/Source/Runtime/Public/GerstnerWaterWaveViewExtension.h b/Engine/Plugins/Experimental/Water/Source/Runtime/Public/GerstnerWaterWaveViewExtension.h
--- a/Engine/Plugins/Experimental/Water/Source/Runtime/Public/GerstnerWaterWaveViewExtension.h
+++ b/Engine/Plugins/Experimental/Water/Source/Runtime/Public/GerstnerWaterWaveViewExtension.h
@@ -31,8 +31,8 @@ public:
 	virtual void SetupViewFamily(FSceneViewFamily& InViewFamily) override;
 	virtual void SetupView(FSceneViewFamily& InViewFamily, FSceneView& InView) override {}
 	virtual void BeginRenderViewFamily(FSceneViewFamily& InViewFamily) override {}
-	virtual void PreRenderViewFamily_RenderThread(FRHICommandListImmediate& RHICmdList, FSceneViewFamily& InViewFamily) override {}
-	virtual void PreRenderView_RenderThread(FRHICommandListImmediate& RHICmdList, FSceneView& InView) override;
+	virtual void PreRenderViewFamily_RenderThread(FRDGBuilder& GraphBuilder, FSceneViewFamily& InViewFamily) override {}
+	virtual void PreRenderView_RenderThread(FRDGBuilder& GraphBuilder, FSceneView& InView) override;
 
 	bool bRebuildGPUData = false;
 
diff --git a/Engine/Plugins/Experimental/Water/Source/Runtime/Public/WaterViewExtension.h b/Engine/Plugins/Experimental/Water/Source/Runtime/Public/WaterViewExtension.h
--- a/Engine/Plugins/Experimental/Water/Source/Runtime/Public/WaterViewExtension.h
+++ b/Engine/Plugins/Experimental/Water/Source/Runtime/Public/WaterViewExtension.h
@@ -20,7 +20,7 @@ public:
 	virtual void BeginRenderViewFamily(FSceneViewFamily& InViewFamily) override {}
 
 	virtual void SetupViewFamily(FSceneViewFamily& InViewFamily) override;
-	virtual void PreRenderViewFamily_RenderThread(FRHICommandListImmediate& RHICmdList, FSceneViewFamily& InViewFamily) override;
+	virtual void PreRenderViewFamily_RenderThread(FRDGBuilder& GraphBuilder, FSceneViewFamily& InViewFamily) override;
 	// End FSceneViewExtensionBase implementation
 
 	void MarkWaterInfoTextureForRebuild(const UE::WaterInfo::FRenderingContext& RenderContext);
diff --git a/Engine/Plugins/Runtime/AR/AppleAR/AppleARKit/Source/AppleARKit/Private/AppleARKitSystem.cpp b/Engine/Plugins/Runtime/AR/AppleAR/AppleARKit/Source/AppleARKit/Private/AppleARKitSystem.cpp
--- a/Engine/Plugins/Runtime/AR/AppleAR/AppleARKit/Source/AppleARKit/Private/AppleARKitSystem.cpp
+++ b/Engine/Plugins/Runtime/AR/AppleAR/AppleARKit/Source/AppleARKit/Private/AppleARKitSystem.cpp
@@ -24,6 +24,7 @@
 #include "Misc/FileHelper.h"
 #include "ARBlueprintLibrary.h"
 #include "ARKitGeoTrackingSupport.h"
+#include "RenderGraphBuilder.h"
 
 // For mesh occlusion
 #include "MRMeshComponent.h"
@@ -154,7 +155,7 @@ private:
 		FDefaultXRCamera::BeginRenderViewFamily(InViewFamily);
 	}
 	
-	virtual void PreRenderView_RenderThread(FRHICommandListImmediate& RHICmdList, FSceneView& InView) override
+	virtual void PreRenderView_RenderThread(FRDGBuilder& GraphBuilder, FSceneView& InView) override
 	{
 		// Adjust our thread priority if requested
 		if (LastThreadPriority.GetValue() != ThreadPriority.GetValue())
@@ -162,10 +163,10 @@ private:
 			SetThreadPriority(ThreadPriority.GetValue());
 			LastThreadPriority.Set(ThreadPriority.GetValue());
 		}
-		FDefaultXRCamera::PreRenderView_RenderThread(RHICmdList, InView);
+		FDefaultXRCamera::PreRenderView_RenderThread(GraphBuilder, InView);
 	}
-	
-	virtual void PreRenderViewFamily_RenderThread(FRHICommandListImmediate& RHICmdList, FSceneViewFamily& InViewFamily) override
+
+	virtual void PreRenderViewFamily_RenderThread(FRDGBuilder& GraphBuilder, FSceneViewFamily& InViewFamily) override
 	{
 		// Grab the latest frame from ARKit
 		{
@@ -173,10 +174,30 @@ private:
 			ARKitSystem.RenderThreadFrame = ARKitSystem.LastReceivedFrame;
 		}
 
-		FDefaultXRCamera::PreRenderViewFamily_RenderThread(RHICmdList, InViewFamily);
+		FDefaultXRCamera::PreRenderViewFamily_RenderThread(GraphBuilder, InViewFamily);
 	}
+
+	BEGIN_SHADER_PARAMETER_STRUCT(FPostBasePassViewExtensionParameters, )
+		SHADER_PARAMETER_RDG_UNIFORM_BUFFER(FSceneTextureUniformParameters, SceneTextures)
+		RENDER_TARGET_BINDING_SLOTS()
+	END_SHADER_PARAMETER_STRUCT()
 	
-	virtual void PostRenderBasePass_RenderThread(FRHICommandListImmediate& RHICmdList, FSceneView& InView) override
+	virtual void PostRenderBasePassDeferred_RenderThread(FRDGBuilder& GraphBuilder, FSceneView& InView, const FRenderTargetBindingSlots& RenderTargets, TRDGUniformBufferRef<FSceneTextureUniformParameters> SceneTextures) override
+	{
+		if (ARKitSystem.RenderThreadFrame.IsValid())
+		{
+			auto* PassParameters = GraphBuilder.AllocParameters<FPostBasePassViewExtensionParameters>();
+			PassParameters->RenderTargets = RenderTargets;
+			PassParameters->SceneTextures = SceneTextures;
+
+			GraphBuilder.AddPass(RDG_EVENT_NAME("RenderVideoOverlay_RenderThread"), PassParameters, ERDGPassFlags::Raster, [this, &InView](FRHICommandListImmediate& RHICmdList)
+			{
+				VideoOverlay.RenderVideoOverlay_RenderThread(RHICmdList, InView, *ARKitSystem.RenderThreadFrame, ARKitSystem.GetWorldToMetersScale());
+			});
+		}
+	}
+
+	virtual void PostRenderBasePassMobile_RenderThread(FRHICommandListImmediate& RHICmdList, FSceneView& InView) override
 	{
 		if (ARKitSystem.RenderThreadFrame.IsValid())
 		{
diff --git a/Engine/Plugins/Runtime/AR/Google/GoogleARCore/Source/GoogleARCoreBase/Private/GoogleARCoreXRCamera.cpp b/Engine/Plugins/Runtime/AR/Google/GoogleARCore/Source/GoogleARCoreBase/Private/GoogleARCoreXRCamera.cpp
--- a/Engine/Plugins/Runtime/AR/Google/GoogleARCore/Source/GoogleARCoreBase/Private/GoogleARCoreXRCamera.cpp
+++ b/Engine/Plugins/Runtime/AR/Google/GoogleARCore/Source/GoogleARCoreBase/Private/GoogleARCoreXRCamera.cpp
@@ -5,6 +5,7 @@
 #include "SceneView.h"
 #include "GoogleARCorePassthroughCameraRenderer.h"
 #include "GoogleARCoreAndroidHelper.h"
+#include "RenderGraphUtils.h"
 
 #if PLATFORM_ANDROID
 #include <GLES2/gl2.h>
@@ -38,9 +39,9 @@ void FGoogleARCoreXRCamera::BeginRenderViewFamily(FSceneViewFamily& InViewFamily
 	FDefaultXRCamera::BeginRenderViewFamily(InViewFamily);
 }
 
-void FGoogleARCoreXRCamera::PreRenderViewFamily_RenderThread(FRHICommandListImmediate& RHICmdList, FSceneViewFamily& InViewFamily)
+void FGoogleARCoreXRCamera::PreRenderViewFamily_RenderThread(FRDGBuilder& GraphBuilder, FSceneViewFamily& InViewFamily)
 {
-	FDefaultXRCamera::PreRenderViewFamily_RenderThread(RHICmdList, InViewFamily);
+	FDefaultXRCamera::PreRenderViewFamily_RenderThread(GraphBuilder, InViewFamily);
 
 	FGoogleARCoreXRTrackingSystem& TS = GoogleARCoreTrackingSystem;
 
@@ -61,7 +62,24 @@ void FGoogleARCoreXRCamera::PreRenderViewFamily_RenderThread(FRHICommandListImme
 #endif
 }
 
-void FGoogleARCoreXRCamera::PostRenderBasePass_RenderThread(FRHICommandListImmediate& RHICmdList, FSceneView& InView)
+BEGIN_SHADER_PARAMETER_STRUCT(FPostBasePassViewExtensionParameters, )
+	SHADER_PARAMETER_RDG_UNIFORM_BUFFER(FSceneTextureUniformParameters, SceneTextures)
+	RENDER_TARGET_BINDING_SLOTS()
+END_SHADER_PARAMETER_STRUCT()
+
+void FGoogleARCoreXRCamera::PostRenderBasePassDeferred_RenderThread(FRDGBuilder& GraphBuilder, FSceneView& InView, const FRenderTargetBindingSlots& RenderTargets, TRDGUniformBufferRef<FSceneTextureUniformParameters> SceneTextures)
+{
+	auto* PassParameters = GraphBuilder.AllocParameters<FPostBasePassViewExtensionParameters>();
+	PassParameters->RenderTargets = RenderTargets;
+	PassParameters->SceneTextures = SceneTextures;
+
+	GraphBuilder.AddPass(RDG_EVENT_NAME("RenderVideoOverlay_RenderThread"), PassParameters, ERDGPassFlags::Raster, [this, &InView](FRHICommandListImmediate& RHICmdList)
+	{
+		PassthroughRenderer->RenderVideoOverlay_RenderThread(RHICmdList, InView);
+	});
+}
+
+void FGoogleARCoreXRCamera::PostRenderBasePassMobile_RenderThread(FRHICommandListImmediate& RHICmdList, FSceneView& InView)
 {
 	PassthroughRenderer->RenderVideoOverlay_RenderThread(RHICmdList, InView);
 }
diff --git a/Engine/Plugins/Runtime/AR/Google/GoogleARCore/Source/GoogleARCoreBase/Private/GoogleARCoreXRCamera.h b/Engine/Plugins/Runtime/AR/Google/GoogleARCore/Source/GoogleARCoreBase/Private/GoogleARCoreXRCamera.h
--- a/Engine/Plugins/Runtime/AR/Google/GoogleARCore/Source/GoogleARCoreBase/Private/GoogleARCoreXRCamera.h
+++ b/Engine/Plugins/Runtime/AR/Google/GoogleARCore/Source/GoogleARCoreBase/Private/GoogleARCoreXRCamera.h
@@ -17,8 +17,9 @@ public:
 	virtual void SetupView(FSceneViewFamily& InViewFamily, FSceneView& InView) override;
 	virtual void SetupViewProjectionMatrix(FSceneViewProjectionData& InOutProjectionData) override;
 	virtual void BeginRenderViewFamily(FSceneViewFamily& InViewFamily) override;
-	virtual void PreRenderViewFamily_RenderThread(FRHICommandListImmediate& RHICmdList, FSceneViewFamily& InViewFamily) override;
-	virtual void PostRenderBasePass_RenderThread(FRHICommandListImmediate& RHICmdList, FSceneView& InView) override;
+	virtual void PreRenderViewFamily_RenderThread(FRDGBuilder& GraphBuilder, FSceneViewFamily& InViewFamily) override;
+	virtual void PostRenderBasePassDeferred_RenderThread(FRDGBuilder& GraphBuilder, FSceneView& InView, const FRenderTargetBindingSlots& RenderTargets, TRDGUniformBufferRef<FSceneTextureUniformParameters> SceneTextures) override;
+	virtual void PostRenderBasePassMobile_RenderThread(FRHICommandListImmediate& RHICmdList, FSceneView& InView) override;
 	virtual bool IsActiveThisFrame_Internal(const FSceneViewExtensionContext& Context) const override;
 	virtual bool GetPassthroughCameraUVs_RenderThread(TArray<FVector2D>& OutUVs) override;
 	//~ FDefaultXRCamera
diff --git a/Engine/Plugins/Runtime/Oculus/OculusVR/Source/OculusHMD/Private/OculusHMD.cpp b/Engine/Plugins/Runtime/Oculus/OculusVR/Source/OculusHMD/Private/OculusHMD.cpp
--- a/Engine/Plugins/Runtime/Oculus/OculusVR/Source/OculusHMD/Private/OculusHMD.cpp
+++ b/Engine/Plugins/Runtime/Oculus/OculusVR/Source/OculusHMD/Private/OculusHMD.cpp
@@ -1915,7 +1915,7 @@ namespace OculusHMD
 	}
 
 
-	void FOculusHMD::PreRenderViewFamily_RenderThread(FRHICommandListImmediate& RHICmdList, FSceneViewFamily& ViewFamily)
+	void FOculusHMD::PreRenderViewFamily_RenderThread(FRDGBuilder& GraphBuilder, FSceneViewFamily& ViewFamily)
 	{
 		CheckInRenderThread();
 
@@ -1979,16 +1979,16 @@ namespace OculusHMD
 	}
 
 
-	void FOculusHMD::PreRenderView_RenderThread(FRHICommandListImmediate& RHICmdList, FSceneView& InView)
+	void FOculusHMD::PreRenderView_RenderThread(FRDGBuilder& GraphBuilder, FSceneView& InView)
 	{
 	}
 
 
-	void FOculusHMD::PostRenderViewFamily_RenderThread(FRHICommandListImmediate& RHICmdList, FSceneViewFamily& InViewFamily)
+	void FOculusHMD::PostRenderViewFamily_RenderThread(FRDGBuilder& GraphBuilder, FSceneViewFamily& InViewFamily)
 	{
 		CheckInRenderThread();
 
-		FinishRenderFrame_RenderThread(RHICmdList);
+		FinishRenderFrame_RenderThread(GraphBuilder);
 	}
 
 
@@ -3431,7 +3431,7 @@ namespace OculusHMD
 	}
 
 
-	void FOculusHMD::FinishRenderFrame_RenderThread(FRHICommandListImmediate& RHICmdList)
+	void FOculusHMD::FinishRenderFrame_RenderThread(FRDGBuilder& GraphBuilder)
 	{
 		CheckInRenderThread();
 
@@ -3439,30 +3439,32 @@ namespace OculusHMD
 		{
 			UE_LOG(LogHMD, VeryVerbose, TEXT("FinishRenderFrame %u"), Frame_RenderThread->FrameNumber);
 
-			if (Frame_RenderThread->ShowFlags.Rendering)
+			AddPass(GraphBuilder, RDG_EVENT_NAME("FinishRenderFrame"), [this](FRHICommandListImmediate& RHICmdList)
 			{
-				for (int32 LayerIndex = 0; LayerIndex < Layers_RenderThread.Num(); LayerIndex++)
+				if (Frame_RenderThread->ShowFlags.Rendering)
 				{
-					Layers_RenderThread[LayerIndex]->UpdateTexture_RenderThread(CustomPresent, RHICmdList);
+					for (int32 LayerIndex = 0; LayerIndex < Layers_RenderThread.Num(); LayerIndex++)
+					{
+						Layers_RenderThread[LayerIndex]->UpdateTexture_RenderThread(CustomPresent, RHICmdList);
+					}
 				}
-			}
 
-			if (GRHISupportsLateVariableRateShadingUpdate && CVarOculusEnableLowLatencyVRS.GetValueOnAnyThread() == 1)
-			{
-				// late-update foveation if supported and needed
-				ExecuteOnRHIThread_DoNotWait([this]()
+				if (GRHISupportsLateVariableRateShadingUpdate && CVarOculusEnableLowLatencyVRS.GetValueOnAnyThread() == 1)
 				{
-					ovrpResult Result;
-					if (Frame_RHIThread && OVRP_FAILURE(Result = FOculusHMDModule::GetPluginWrapper().UpdateFoveation(Frame_RHIThread->FrameNumber)))
+					// late-update foveation if supported and needed
+					RHICmdList.EnqueueLambda([this](FRHICommandListImmediate&)
 					{
-						UE_LOG(LogHMD, Error, TEXT("FOculusHMDModule::GetPluginWrapper().UpdateFoveation %u failed (%d)"), Frame_RHIThread->FrameNumber, Result);
-					}
-				});
-			}
+						ovrpResult Result;
+						if (Frame_RHIThread && OVRP_FAILURE(Result = FOculusHMDModule::GetPluginWrapper().UpdateFoveation(Frame_RHIThread->FrameNumber)))
+						{
+							UE_LOG(LogHMD, Error, TEXT("FOculusHMDModule::GetPluginWrapper().UpdateFoveation %u failed (%d)"), Frame_RHIThread->FrameNumber, Result);
+						}
+					});
+				}
 
+				Frame_RenderThread.Reset();
+			});
 		}
-
-		Frame_RenderThread.Reset();
 	}
 
 
diff --git a/Engine/Plugins/Runtime/Oculus/OculusVR/Source/OculusHMD/Private/OculusHMD.h b/Engine/Plugins/Runtime/Oculus/OculusVR/Source/OculusHMD/Private/OculusHMD.h
--- a/Engine/Plugins/Runtime/Oculus/OculusVR/Source/OculusHMD/Private/OculusHMD.h
+++ b/Engine/Plugins/Runtime/Oculus/OculusVR/Source/OculusHMD/Private/OculusHMD.h
@@ -210,9 +210,9 @@ public:
 	virtual void SetupViewFamily(FSceneViewFamily& InViewFamily) override;
 	virtual void SetupView(FSceneViewFamily& InViewFamily, FSceneView& InView) override;
 	virtual void BeginRenderViewFamily(FSceneViewFamily& InViewFamily) override;
-	virtual void PreRenderViewFamily_RenderThread(FRHICommandListImmediate& RHICmdList, FSceneViewFamily& InViewFamily) override;
-	virtual void PreRenderView_RenderThread(FRHICommandListImmediate& RHICmdList, FSceneView& InView) override;
-	virtual void PostRenderViewFamily_RenderThread(FRHICommandListImmediate& RHICmdList, FSceneViewFamily& InViewFamily) override;
+	virtual void PreRenderViewFamily_RenderThread(FRDGBuilder& GraphBuilder, FSceneViewFamily& InViewFamily) override;
+	virtual void PreRenderView_RenderThread(FRDGBuilder& GraphBuilder, FSceneView& InView) override;
+	virtual void PostRenderViewFamily_RenderThread(FRDGBuilder& GraphBuilder, FSceneViewFamily& InViewFamily) override;
 	virtual int32 GetPriority() const override;
 //	virtual bool LateLatchingEnabled() const override;
 //	virtual void PreLateLatchingViewFamily_RenderThread(FRHICommandListImmediate& RHICmdList, FSceneViewFamily& InViewFamily) override;
@@ -335,7 +335,7 @@ public:
 	OCULUSHMD_API void StartGameFrame_GameThread(); // Called from OnStartGameFrame or from FOculusInput::SendControllerEvents (first actual call of the frame)
 	void FinishGameFrame_GameThread(); // Called from OnEndGameFrame
 	void StartRenderFrame_GameThread(); // Called from BeginRenderViewFamily
-	void FinishRenderFrame_RenderThread(FRHICommandListImmediate& RHICmdList); // Called from PostRenderViewFamily_RenderThread
+	void FinishRenderFrame_RenderThread(FRDGBuilder& GraphBuilder); // Called from PostRenderViewFamily_RenderThread
 	void StartRHIFrame_RenderThread(); // Called from PreRenderViewFamily_RenderThread
 	void FinishRHIFrame_RHIThread(); // Called from FinishRendering_RHIThread
 
diff --git a/Engine/Plugins/Runtime/OpenXR/Source/OpenXRHMD/Private/OpenXRHMD.cpp b/Engine/Plugins/Runtime/OpenXR/Source/OpenXRHMD/Private/OpenXRHMD.cpp
--- a/Engine/Plugins/Runtime/OpenXR/Source/OpenXRHMD/Private/OpenXRHMD.cpp
+++ b/Engine/Plugins/Runtime/OpenXR/Source/OpenXRHMD/Private/OpenXRHMD.cpp
@@ -979,12 +979,12 @@ void FOpenXRHMD::BeginRenderViewFamily(FSceneViewFamily& InViewFamily)
 	}
 }
 
-void FOpenXRHMD::PreRenderView_RenderThread(FRHICommandListImmediate& RHICmdList, FSceneView& InView)
+void FOpenXRHMD::PreRenderView_RenderThread(FRDGBuilder& GraphBuilder, FSceneView& InView)
 {
 	check(IsInRenderingThread());
 }
 
-void FOpenXRHMD::PreRenderViewFamily_RenderThread(FRHICommandListImmediate& RHICmdList, FSceneViewFamily& ViewFamily)
+void FOpenXRHMD::PreRenderViewFamily_RenderThread(FRDGBuilder& GraphBuilder, FSceneViewFamily& ViewFamily)
 {
 	check(IsInRenderingThread());
 
diff --git a/Engine/Plugins/Runtime/OpenXR/Source/OpenXRHMD/Private/OpenXRHMD.h b/Engine/Plugins/Runtime/OpenXR/Source/OpenXRHMD/Private/OpenXRHMD.h
--- a/Engine/Plugins/Runtime/OpenXR/Source/OpenXRHMD/Private/OpenXRHMD.h
+++ b/Engine/Plugins/Runtime/OpenXR/Source/OpenXRHMD/Private/OpenXRHMD.h
@@ -282,8 +282,8 @@ public:
 	virtual void SetupViewFamily(FSceneViewFamily& InViewFamily) override;
 	virtual void SetupView(FSceneViewFamily& InViewFamily, FSceneView& InView) override;
 	virtual void BeginRenderViewFamily(FSceneViewFamily& InViewFamily) override;
-	virtual void PreRenderView_RenderThread(FRHICommandListImmediate& RHICmdList, FSceneView& InView) override;
-	virtual void PreRenderViewFamily_RenderThread(FRHICommandListImmediate& RHICmdList, FSceneViewFamily& InViewFamily) override;
+	virtual void PreRenderView_RenderThread(FRDGBuilder& GraphBuilder, FSceneView& InView) override;
+	virtual void PreRenderViewFamily_RenderThread(FRDGBuilder& GraphBuilder, FSceneViewFamily& InViewFamily) override;
 
 	/** FHMDSceneViewExtension interface */
 	virtual bool IsActiveThisFrame_Internal(const FSceneViewExtensionContext& Context) const;
diff --git a/Engine/Plugins/Runtime/Steam/SteamVR/Source/SteamVR/Private/SteamVRHMD.h b/Engine/Plugins/Runtime/Steam/SteamVR/Source/SteamVR/Private/SteamVRHMD.h
--- a/Engine/Plugins/Runtime/Steam/SteamVR/Source/SteamVR/Private/SteamVRHMD.h
+++ b/Engine/Plugins/Runtime/Steam/SteamVR/Source/SteamVR/Private/SteamVRHMD.h
@@ -192,9 +192,9 @@ public:
 	virtual void SetupViewFamily(FSceneViewFamily& InViewFamily) override {};
 	virtual void SetupView(FSceneViewFamily& InViewFamily, FSceneView& InView) override {}
 	virtual void BeginRenderViewFamily(FSceneViewFamily& InViewFamily) override {}
-	virtual void PreRenderView_RenderThread(FRHICommandListImmediate& RHICmdList, FSceneView& InView) override {}
-	virtual void PreRenderViewFamily_RenderThread(FRHICommandListImmediate& RHICmdList, FSceneViewFamily& InViewFamily) override {}
-	virtual void PostRenderView_RenderThread(FRHICommandListImmediate& RHICmdList, FSceneView& InView) override;
+	virtual void PreRenderView_RenderThread(FRDGBuilder& GraphBuilder, FSceneView& InView) override {}
+	virtual void PreRenderViewFamily_RenderThread(FRDGBuilder& GraphBuilder, FSceneViewFamily& InViewFamily) override {}
+	virtual void PostRenderView_RenderThread(FRDGBuilder& GraphBuilder, FSceneView& InView) override;
 
 	// FHMDSceneViewExtension interface
 	virtual bool IsActiveThisFrame_Internal(const FSceneViewExtensionContext& Context) const override;
diff --git a/Engine/Plugins/Runtime/Steam/SteamVR/Source/SteamVR/Private/SteamVRRender.cpp b/Engine/Plugins/Runtime/Steam/SteamVR/Source/SteamVR/Private/SteamVRRender.cpp
--- a/Engine/Plugins/Runtime/Steam/SteamVR/Source/SteamVR/Private/SteamVRRender.cpp
+++ b/Engine/Plugins/Runtime/Steam/SteamVR/Source/SteamVR/Private/SteamVRRender.cpp
@@ -56,7 +56,7 @@ void FSteamVRHMD::RenderTexture_RenderThread(FRHICommandListImmediate& RHICmdLis
 	SpectatorScreenController->RenderSpectatorScreen_RenderThread(RHICmdList, BackBuffer, SrcTexture, WindowSize);
 }
 
-void FSteamVRHMD::PostRenderView_RenderThread(FRHICommandListImmediate& RHICmdList, FSceneView& InView)
+void FSteamVRHMD::PostRenderView_RenderThread(FRDGBuilder& GraphBuilder, FSceneView& InView)
 {
 	check(IsInRenderingThread());
 	UpdateStereoLayers_RenderThread();
diff --git a/Engine/Plugins/Runtime/nDisplay/Source/DisplayCluster/Private/Render/Viewport/DisplayClusterViewportManagerViewExtension.cpp b/Engine/Plugins/Runtime/nDisplay/Source/DisplayCluster/Private/Render/Viewport/DisplayClusterViewportManagerViewExtension.cpp
--- a/Engine/Plugins/Runtime/nDisplay/Source/DisplayCluster/Private/Render/Viewport/DisplayClusterViewportManagerViewExtension.cpp
+++ b/Engine/Plugins/Runtime/nDisplay/Source/DisplayCluster/Private/Render/Viewport/DisplayClusterViewportManagerViewExtension.cpp
@@ -15,7 +15,7 @@ FDisplayClusterViewportManagerViewExtension::FDisplayClusterViewportManagerViewE
 	, ViewportManager(InViewportManager)
 { }
 
-void FDisplayClusterViewportManagerViewExtension::PreRenderView_RenderThread(FRHICommandListImmediate& RHICmdList, FSceneView& InView)
+void FDisplayClusterViewportManagerViewExtension::PreRenderView_RenderThread(FRDGBuilder& GraphBuilder, FSceneView& InView)
 {
 	if (InView.bIsViewInfo)
 	{
@@ -42,7 +42,7 @@ bool FDisplayClusterViewportManagerViewExtension::IsActiveThisFrame_Internal(con
 	return false;
 }
 
-void FDisplayClusterViewportManagerViewExtension::PostRenderViewFamily_RenderThread(FRHICommandListImmediate& RHICmdList, FSceneViewFamily& InViewFamily)
+void FDisplayClusterViewportManagerViewExtension::PostRenderViewFamily_RenderThread(FRDGBuilder& GraphBuilder, FSceneViewFamily& InViewFamily)
 {
 	// Since FSceneViewFamily has no direct references to the nD viewport it relates to,
 	// we find a proper viewport ID by comparing RenderTarget addresses.
@@ -55,7 +55,7 @@ void FDisplayClusterViewportManagerViewExtension::PostRenderViewFamily_RenderThr
 			const TArrayView<IDisplayClusterViewportProxy*> ViewportProxies = ViewportMgrPoxy->GetViewports_RenderThread();
 
 			// Filter those viewports that refer to the RenderTarget used in the ViewFamily
-			const TArray<IDisplayClusterViewportProxy*> FoundViewportProxies = ViewportProxies.FilterByPredicate([&InViewFamily](const IDisplayClusterViewportProxy* ViewportProxy)
+			TArray<IDisplayClusterViewportProxy*> FoundViewportProxies = ViewportProxies.FilterByPredicate([&InViewFamily](const IDisplayClusterViewportProxy* ViewportProxy)
 			{
 				TArray<FRHITexture*> RenderTargets;
 				if (ViewportProxy->GetResources_RenderThread(EDisplayClusterViewportResourceType::InternalRenderTargetResource, RenderTargets))
@@ -69,11 +69,14 @@ void FDisplayClusterViewportManagerViewExtension::PostRenderViewFamily_RenderThr
 				return false;
 			});
 
-			// Now we can perform per-viewport notification
-			for (const IDisplayClusterViewportProxy* ViewportProxy : FoundViewportProxies)
+			AddPass(GraphBuilder, RDG_EVENT_NAME("ViewportProxy_PostRenderViewFamily"), [&InViewFamily, FoundViewportProxies = MoveTemp(FoundViewportProxies)] (FRHICommandListImmediate& RHICmdList)
 			{
-				IDisplayCluster::Get().GetCallbacks().OnDisplayClusterPostRenderViewFamily_RenderThread().Broadcast(RHICmdList, InViewFamily, ViewportProxy);
-			}
+				// Now we can perform per-viewport notification
+				for (const IDisplayClusterViewportProxy* ViewportProxy : FoundViewportProxies)
+				{
+					IDisplayCluster::Get().GetCallbacks().OnDisplayClusterPostRenderViewFamily_RenderThread().Broadcast(RHICmdList, InViewFamily, ViewportProxy);
+				}
+			});
 		}
 	}
 }
diff --git a/Engine/Plugins/Runtime/nDisplay/Source/DisplayCluster/Private/Render/Viewport/DisplayClusterViewportManagerViewExtension.h b/Engine/Plugins/Runtime/nDisplay/Source/DisplayCluster/Private/Render/Viewport/DisplayClusterViewportManagerViewExtension.h
--- a/Engine/Plugins/Runtime/nDisplay/Source/DisplayCluster/Private/Render/Viewport/DisplayClusterViewportManagerViewExtension.h
+++ b/Engine/Plugins/Runtime/nDisplay/Source/DisplayCluster/Private/Render/Viewport/DisplayClusterViewportManagerViewExtension.h
@@ -20,10 +20,10 @@ public:
 
 	virtual void BeginRenderViewFamily(FSceneViewFamily& InViewFamily) override {};
 
-	virtual void PreRenderView_RenderThread(FRHICommandListImmediate& RHICmdList, FSceneView& InView) override;
+	virtual void PreRenderView_RenderThread(FRDGBuilder& GraphBuilder, FSceneView& InView) override;
 	virtual bool IsActiveThisFrame_Internal(const FSceneViewExtensionContext& Context) const override;
 
-	virtual void PostRenderViewFamily_RenderThread(FRHICommandListImmediate& RHICmdList, FSceneViewFamily& InViewFamily) override;
+	virtual void PostRenderViewFamily_RenderThread(FRDGBuilder& GraphBuilder, FSceneViewFamily& InViewFamily) override;
 	//~End ISceneVIewExtension interface
 
 private:
diff --git a/Engine/Source/Developer/FunctionalTesting/Private/AutomationBlueprintFunctionLibrary.cpp b/Engine/Source/Developer/FunctionalTesting/Private/AutomationBlueprintFunctionLibrary.cpp
--- a/Engine/Source/Developer/FunctionalTesting/Private/AutomationBlueprintFunctionLibrary.cpp
+++ b/Engine/Source/Developer/FunctionalTesting/Private/AutomationBlueprintFunctionLibrary.cpp
@@ -226,8 +226,8 @@ public:
 		}
 
 	virtual void BeginRenderViewFamily(FSceneViewFamily& InViewFamily) {}
-	virtual void PreRenderViewFamily_RenderThread(FRHICommandListImmediate& RHICmdList, FSceneViewFamily& InViewFamily) {}
-	virtual void PreRenderView_RenderThread(FRHICommandListImmediate& RHICmdList, FSceneView& InView) {}
+	virtual void PreRenderViewFamily_RenderThread(FRDGBuilder& GraphBuilder, FSceneViewFamily& InViewFamily) {}
+	virtual void PreRenderView_RenderThread(FRDGBuilder& GraphBuilder, FSceneView& InView) {}
 
 	/** We always want to go last. */
 	virtual int32 GetPriority() const override { return MIN_int32; }
diff --git a/Engine/Source/Editor/PixelInspector/Private/PixelInspector.h b/Engine/Source/Editor/PixelInspector/Private/PixelInspector.h
--- a/Engine/Source/Editor/PixelInspector/Private/PixelInspector.h
+++ b/Engine/Source/Editor/PixelInspector/Private/PixelInspector.h
@@ -33,8 +33,8 @@ namespace PixelInspector
 		virtual void SetupViewFamily(FSceneViewFamily& InViewFamily) override {};
 		virtual void SetupView(FSceneViewFamily& InViewFamily, FSceneView& InView) override {};
 		virtual void BeginRenderViewFamily(FSceneViewFamily& InViewFamily) override;
-		virtual void PreRenderViewFamily_RenderThread(FRHICommandListImmediate& RHICmdList, FSceneViewFamily& InViewFamily) override {};
-		virtual void PreRenderView_RenderThread(FRHICommandListImmediate& RHICmdList, FSceneView& InView) override {};
+		virtual void PreRenderViewFamily_RenderThread(FRDGBuilder& GraphBuilder, FSceneViewFamily& InViewFamily) override {};
+		virtual void PreRenderView_RenderThread(FRDGBuilder& GraphBuilder, FSceneView& InView) override {};
 		virtual void SubscribeToPostProcessingPass(EPostProcessingPass PassId, FAfterPassCallbackDelegateArray& InOutPassCallbacks, bool bIsPassEnabled) override;
 
 	public:
diff --git a/Engine/Source/Runtime/Engine/Private/SceneViewExtension.cpp b/Engine/Source/Runtime/Engine/Private/SceneViewExtension.cpp
--- a/Engine/Source/Runtime/Engine/Private/SceneViewExtension.cpp
+++ b/Engine/Source/Runtime/Engine/Private/SceneViewExtension.cpp
@@ -2,6 +2,7 @@
 
 #include "SceneViewExtension.h"
 #include "RenderGraphUtils.h"
+#include "SceneTexturesConfig.h"
 
 //
 // FSceneViewExtensionBase
@@ -13,14 +14,6 @@ FSceneViewExtensionBase::~FSceneViewExtensionBase()
 	// so they will be automatically unregistered when removed.
 }
 
-// Temporary override so that old behaviour still functions. Will be removed along with IsActiveThisFrame(FViewport*).
-bool FSceneViewExtensionBase::IsActiveThisFrame_Internal(const FSceneViewExtensionContext & Context) const
-{
-	PRAGMA_DISABLE_DEPRECATION_WARNINGS
-	return ISceneViewExtension::IsActiveThisFrame(Context.Viewport);
-	PRAGMA_ENABLE_DEPRECATION_WARNINGS
-}
-
 bool FSceneViewExtensionBase::IsActiveThisFrame(const FSceneViewExtensionContext& Context) const
 {
 	// Go over any existing activation lambdas
@@ -39,38 +32,6 @@ bool FSceneViewExtensionBase::IsActiveThisFrame(const FSceneViewExtensionContext
 	return IsActiveThisFrame_Internal(Context);
 }
 
-void ISceneViewExtension::PreRenderViewFamily_RenderThread(FRDGBuilder& GraphBuilder, FSceneViewFamily& InViewFamily)
-{
-	AddPass(GraphBuilder, RDG_EVENT_NAME("PreRenderViewFamily_RenderThread"), [this, &InViewFamily](FRHICommandListImmediate& RHICmdList)
-	{
-		PreRenderViewFamily_RenderThread(RHICmdList, InViewFamily);
-	});
-}
-
-void ISceneViewExtension::PreRenderView_RenderThread(FRDGBuilder& GraphBuilder, FSceneView& InView)
-{
-	AddPass(GraphBuilder, RDG_EVENT_NAME("PreRenderView_RenderThread"), [this, &InView](FRHICommandListImmediate& RHICmdList)
-	{
-		PreRenderView_RenderThread(RHICmdList, InView);
-	});
-}
-
-void ISceneViewExtension::PostRenderViewFamily_RenderThread(FRDGBuilder& GraphBuilder, FSceneViewFamily& InViewFamily)
-{
-	AddPass(GraphBuilder, RDG_EVENT_NAME("PostRenderViewFamily_RenderThread"), [this, &InViewFamily](FRHICommandListImmediate& RHICmdList)
-	{
-		PostRenderViewFamily_RenderThread(RHICmdList, InViewFamily);
-	});
-}
-
-void ISceneViewExtension::PostRenderView_RenderThread(FRDGBuilder& GraphBuilder, FSceneView& InView)
-{
-	AddPass(GraphBuilder, RDG_EVENT_NAME("PostRenderView_RenderThread"), [this, &InView](FRHICommandListImmediate& RHICmdList)
-	{
-		PostRenderView_RenderThread(RHICmdList, InView);
-	});
-}
-
 //
 // FWorldSceneViewExtension
 //
diff --git a/Engine/Source/Runtime/Engine/Public/SceneViewExtension.h b/Engine/Source/Runtime/Engine/Public/SceneViewExtension.h
--- a/Engine/Source/Runtime/Engine/Public/SceneViewExtension.h
+++ b/Engine/Source/Runtime/Engine/Public/SceneViewExtension.h
@@ -76,6 +76,8 @@
  *
  */
 
+#include "SceneTexturesConfig.h"
+
 class APlayerController;
 class FRHICommandListImmediate;
 class FSceneView;
@@ -141,19 +143,22 @@ public:
     /**
      * Called on render thread at the start of rendering.
      */
-	virtual void PreRenderViewFamily_RenderThread(FRHICommandListImmediate& RHICmdList, FSceneViewFamily& InViewFamily) {};
-	virtual ENGINE_API void PreRenderViewFamily_RenderThread(FRDGBuilder& GraphBuilder, FSceneViewFamily& InViewFamily);
+	virtual void PreRenderViewFamily_RenderThread(FRDGBuilder& GraphBuilder, FSceneViewFamily& InViewFamily) {}
 
 	/**
      * Called on render thread at the start of rendering, for each view, after PreRenderViewFamily_RenderThread call.
-     */
-	virtual void PreRenderView_RenderThread(FRHICommandListImmediate& RHICmdList, FSceneView& InView) {};
-	virtual ENGINE_API void PreRenderView_RenderThread(FRDGBuilder& GraphBuilder, FSceneView& InView);
+	 */
+	virtual void PreRenderView_RenderThread(FRDGBuilder& GraphBuilder, FSceneView& InView) {}
 
 	/**
-	 * Called right after Base Pass rendering finished
+	 * Called right after Base Pass rendering finished when using the deferred renderer.
 	 */
-	virtual void PostRenderBasePass_RenderThread(FRHICommandListImmediate& RHICmdList, FSceneView& InView) {};
+	virtual void PostRenderBasePassDeferred_RenderThread(FRDGBuilder& GraphBuilder, FSceneView& InView, const FRenderTargetBindingSlots& RenderTargets, TRDGUniformBufferRef<FSceneTextureUniformParameters> SceneTextures) {}
+
+	/**
+	 * Called right after Base Pass rendering finished when using the mobile renderer.
+	 */
+	virtual void PostRenderBasePassMobile_RenderThread(FRHICommandListImmediate& RHICmdList, FSceneView& InView) {}
 
 	/**
 	 * Called right before Post Processing rendering begins
@@ -168,36 +173,43 @@ public:
 	/**
 	 * Allows to render content after the 3D content scene, useful for debugging
 	 */
-	virtual void PostRenderViewFamily_RenderThread(FRHICommandListImmediate& RHICmdList, FSceneViewFamily& InViewFamily) {}
-	virtual ENGINE_API void PostRenderViewFamily_RenderThread(FRDGBuilder& GraphBuilder, FSceneViewFamily& InViewFamily);
+	virtual void PostRenderViewFamily_RenderThread(FRDGBuilder& GraphBuilder, FSceneViewFamily& InViewFamily) {}
 
 	/**
 	 * Allows to render content after the 3D content scene, useful for debugging
 	 */
-	virtual void PostRenderView_RenderThread(FRHICommandListImmediate& RHICmdList, FSceneView& InView) {}
-	virtual ENGINE_API void PostRenderView_RenderThread(FRDGBuilder& GraphBuilder, FSceneView& InView);
+	virtual void PostRenderView_RenderThread(FRDGBuilder& GraphBuilder, FSceneView& InView) {}
 
 	/**
      * Called to determine view extensions priority in relation to other view extensions, higher comes first
      */
 	virtual int32 GetPriority() const { return 0; }
 
-	/**
-	 * Returning false disables the extension for the current frame. This will be queried each frame to determine if the extension wants to run.
-	 */
-	UE_DEPRECATED(4.27, "Deprecated. Please use IsActiveThisFrame by passing an FSceneViewExtensionContext parameter")
-	virtual bool IsActiveThisFrame(class FViewport* InViewport) const { return true; }
-
 	/**
 	 * Returning false disables the extension for the current frame in the given context. This will be queried each frame to determine if the extension wants to run.
 	 */
 	virtual bool IsActiveThisFrame(const FSceneViewExtensionContext& Context) const { return IsActiveThisFrame_Internal(Context); }
 
-	/**
-     * Returning false disables the extension for the current frame in the given context. This will be queried each frame to determine if the extension wants to run.
-     */
-	UE_DEPRECATED(4.27, "Deprecated. Please use IsActiveThisFrame_Internal instead.")
-	virtual bool IsActiveThisFrameInContext(FSceneViewExtensionContext& Context) const { return IsActiveThisFrame(Context); }
+	///////////////////////////////////////////////////////////////////////////////////////
+	//! Deprecated APIs - These are no longer called and must be converted to restore functionality.
+
+	UE_DEPRECATED(5.1, "Use PreRenderViewFamily_RenderThread with an RDG builder instead. RHI commands can be enqueued directly using GraphBuilder.RHICmdList or queued onto the graph.")
+	virtual void PreRenderViewFamily_RenderThread(FRHICommandListImmediate& RHICmdList, FSceneViewFamily& InViewFamily) {}
+
+	UE_DEPRECATED(5.1, "Use PreRenderView_RenderThread with an RDG builder instead. RHI commands can be enqueued directly using GraphBuilder.RHICmdList or queued onto the graph.")
+	virtual void PreRenderView_RenderThread(FRHICommandListImmediate& RHICmdList, FSceneView& InView) {};
+
+	UE_DEPRECATED(5.1, "Use PostRePostRenderView_RenderThreadnderViewFamily_RenderThread with an RDG builder instead. RHI commands can be enqueued directly using GraphBuilder.RHICmdList or queued onto the graph.")
+	virtual void PostRenderView_RenderThread(FRHICommandListImmediate& RHICmdList, FSceneView& InView) {}
+
+	UE_DEPRECATED(5.1, "Use PostRenderViewFamily_RenderThread with an RDG builder instead. RHI commands can be enqueued directly using GraphBuilder.RHICmdList or queued onto the graph.")
+	virtual void PostRenderViewFamily_RenderThread(FRHICommandListImmediate& RHICmdList, FSceneViewFamily& InViewFamily) {}
+
+	UE_DEPRECATED(5.1, "PostRenderBasePass now has a mobile and deferred variant.")
+	virtual void PostRenderBasePass_RenderThread(FRHICommandListImmediate& RHICmdList, FSceneView& InView) {};
+
+	///////////////////////////////////////////////////////////////////////////////////////
+
 protected:
 	/**
 	 * Called if no IsActive functors returned a definitive answer to whether this extension should be active this frame.
@@ -230,9 +242,6 @@ public:
 
 	// Determines if the extension should be active for the current frame and given context.
 	virtual bool IsActiveThisFrame(const FSceneViewExtensionContext& Context) const override final;
-protected:
-	// Temporary override so that old behaviour still functions. Will be removed along with IsActiveThisFrame(FViewport*).
-	virtual bool IsActiveThisFrame_Internal(const FSceneViewExtensionContext& Context) const override;
 };
 
 /** Scene View Extension which is enabled for all Viewports/Scenes which have the same world. */
diff --git a/Engine/Source/Runtime/Engine/Public/TextureResource.h b/Engine/Source/Runtime/Engine/Public/TextureResource.h
--- a/Engine/Source/Runtime/Engine/Public/TextureResource.h
+++ b/Engine/Source/Runtime/Engine/Public/TextureResource.h
@@ -280,6 +280,11 @@ public:
 		bNeedsUpdate = true;
 	}
 
+	static bool IsUpdateNeeded()
+	{
+		return bNeedsUpdate;
+	}
+
 protected:
 
 	/**
diff --git a/Engine/Source/Runtime/HeadMountedDisplay/Private/DefaultStereoLayers.cpp b/Engine/Source/Runtime/HeadMountedDisplay/Private/DefaultStereoLayers.cpp
--- a/Engine/Source/Runtime/HeadMountedDisplay/Private/DefaultStereoLayers.cpp
+++ b/Engine/Source/Runtime/HeadMountedDisplay/Private/DefaultStereoLayers.cpp
@@ -15,6 +15,7 @@
 #include "SceneView.h"
 #include "CommonRenderResources.h"
 #include "IXRLoadingScreen.h"
+#include "RenderGraphUtils.h"
 
 namespace 
 {
@@ -166,7 +167,7 @@ void FDefaultStereoLayers::StereoLayerRender(FRHICommandListImmediate& RHICmdLis
 	}
 }
 
-void FDefaultStereoLayers::PreRenderViewFamily_RenderThread(FRHICommandListImmediate& RHICmdList, FSceneViewFamily& InViewFamily)
+void FDefaultStereoLayers::PreRenderViewFamily_RenderThread(FRDGBuilder& GraphBuilder, FSceneViewFamily& InViewFamily)
 {
 	check(IsInRenderingThread());
 
@@ -207,89 +208,92 @@ void FDefaultStereoLayers::PreRenderViewFamily_RenderThread(FRHICommandListImmed
 }
 
 
-void FDefaultStereoLayers::PostRenderView_RenderThread(FRHICommandListImmediate& RHICmdList, FSceneView& InView)
+void FDefaultStereoLayers::PostRenderView_RenderThread(FRDGBuilder& GraphBuilder, FSceneView& InView)
 {
 	if (!IStereoRendering::IsStereoEyeView(InView))
 	{
 		return;
 	}
 
-	FViewMatrices ModifiedViewMatrices = InView.ViewMatrices;
-	ModifiedViewMatrices.HackRemoveTemporalAAProjectionJitter();
-	const FMatrix& ProjectionMatrix = ModifiedViewMatrices.GetProjectionMatrix();
-	const FMatrix& ViewProjectionMatrix = ModifiedViewMatrices.GetViewProjectionMatrix();
+	AddPass(GraphBuilder, RDG_EVENT_NAME("StereoLayerRender"), [this, &InView](FRHICommandListImmediate& RHICmdList)
+	{
+		FViewMatrices ModifiedViewMatrices = InView.ViewMatrices;
+		ModifiedViewMatrices.HackRemoveTemporalAAProjectionJitter();
+		const FMatrix& ProjectionMatrix = ModifiedViewMatrices.GetProjectionMatrix();
+		const FMatrix& ViewProjectionMatrix = ModifiedViewMatrices.GetViewProjectionMatrix();
 
-	// Calculate a view matrix that only adjusts for eye position, ignoring head position, orientation and world position.
-	FVector EyeShift;
-	FQuat EyeOrientation;
-	HMDDevice->GetRelativeEyePose(IXRTrackingSystem::HMDDeviceId, InView.StereoViewIndex, EyeOrientation, EyeShift);
+		// Calculate a view matrix that only adjusts for eye position, ignoring head position, orientation and world position.
+		FVector EyeShift;
+		FQuat EyeOrientation;
+		HMDDevice->GetRelativeEyePose(IXRTrackingSystem::HMDDeviceId, InView.StereoViewIndex, EyeOrientation, EyeShift);
 
-	FMatrix EyeMatrix = FTranslationMatrix(-EyeShift) * FInverseRotationMatrix(EyeOrientation.Rotator()) * FMatrix(
-		FPlane(0, 0, 1, 0),
-		FPlane(1, 0, 0, 0),
-		FPlane(0, 1, 0, 0),
-		FPlane(0, 0, 0, 1));
+		FMatrix EyeMatrix = FTranslationMatrix(-EyeShift) * FInverseRotationMatrix(EyeOrientation.Rotator()) * FMatrix(
+			FPlane(0, 0, 1, 0),
+			FPlane(1, 0, 0, 0),
+			FPlane(0, 1, 0, 0),
+			FPlane(0, 0, 0, 1));
 
-	FQuat HmdOrientation = HmdTransform.GetRotation();
-	FVector HmdLocation = HmdTransform.GetTranslation();
-	FMatrix TrackerMatrix = FTranslationMatrix(-HmdLocation) * FInverseRotationMatrix(HmdOrientation.Rotator()) * EyeMatrix;
+		FQuat HmdOrientation = HmdTransform.GetRotation();
+		FVector HmdLocation = HmdTransform.GetTranslation();
+		FMatrix TrackerMatrix = FTranslationMatrix(-HmdLocation) * FInverseRotationMatrix(HmdOrientation.Rotator()) * EyeMatrix;
 
-	FLayerRenderParams RenderParams{
-		InView.UnscaledViewRect, // Viewport
+		FLayerRenderParams RenderParams{
+			InView.UnscaledViewRect, // Viewport
+			{
+				ViewProjectionMatrix,				// WorldLocked,
+				TrackerMatrix * ProjectionMatrix,	// TrackerLocked,
+				EyeMatrix * ProjectionMatrix		// FaceLocked
+			}
+		};
+
+		TArray<FRHITransitionInfo, TInlineAllocator<16>> Infos;
+		for (uint32 LayerIndex : SortedSceneLayers)
 		{
-			ViewProjectionMatrix,				// WorldLocked,
-			TrackerMatrix * ProjectionMatrix,	// TrackerLocked,
-			EyeMatrix * ProjectionMatrix		// FaceLocked
+			Infos.Add(FRHITransitionInfo(RenderThreadLayers[LayerIndex].Texture, ERHIAccess::Unknown, ERHIAccess::SRVGraphics));
+		}
+		for (uint32 LayerIndex : SortedOverlayLayers)
+		{
+			Infos.Add(FRHITransitionInfo(RenderThreadLayers[LayerIndex].Texture, ERHIAccess::Unknown, ERHIAccess::SRVGraphics));
+		}
+		if (Infos.Num())
+		{
+			RHICmdList.Transition(Infos);
 		}
-	};
-	
-	TArray<FRHITransitionInfo, TInlineAllocator<16>> Infos;
-	for (uint32 LayerIndex : SortedSceneLayers)
-	{
-		Infos.Add(FRHITransitionInfo(RenderThreadLayers[LayerIndex].Texture, ERHIAccess::Unknown, ERHIAccess::SRVGraphics));
-	}
-	for (uint32 LayerIndex : SortedOverlayLayers)
-	{
-		Infos.Add(FRHITransitionInfo(RenderThreadLayers[LayerIndex].Texture, ERHIAccess::Unknown, ERHIAccess::SRVGraphics));
-	}
-	if (Infos.Num())
-	{
-		RHICmdList.Transition(Infos);
-	}
 
-	FTexture2DRHIRef RenderTarget = HMDDevice->GetSceneLayerTarget_RenderThread(InView.StereoViewIndex, RenderParams.Viewport);
-	if (!RenderTarget.IsValid())
-	{
-		RenderTarget = InView.Family->RenderTarget->GetRenderTargetTexture();
-	}
+		FTexture2DRHIRef RenderTarget = HMDDevice->GetSceneLayerTarget_RenderThread(InView.StereoViewIndex, RenderParams.Viewport);
+		if (!RenderTarget.IsValid())
+		{
+			RenderTarget = InView.Family->RenderTarget->GetRenderTargetTexture();
+		}
 
-	FRHIRenderPassInfo RPInfo(RenderTarget, ERenderTargetActions::Load_Store);
-	RHICmdList.BeginRenderPass(RPInfo, TEXT("StereoLayerRender"));
-	RHICmdList.SetViewport((float)RenderParams.Viewport.Min.X, (float)RenderParams.Viewport.Min.Y, 0.0f, (float)RenderParams.Viewport.Max.X, (float)RenderParams.Viewport.Max.Y, 1.0f);
+		FRHIRenderPassInfo RPInfo(RenderTarget, ERenderTargetActions::Load_Store);
+		RHICmdList.BeginRenderPass(RPInfo, TEXT("StereoLayerRender"));
+		RHICmdList.SetViewport((float)RenderParams.Viewport.Min.X, (float)RenderParams.Viewport.Min.Y, 0.0f, (float)RenderParams.Viewport.Max.X, (float)RenderParams.Viewport.Max.Y, 1.0f);
 
-	if (bSplashIsShown || !IsBackgroundLayerVisible())
-	{
-		DrawClearQuad(RHICmdList, FLinearColor::Black);
-	}
+		if (bSplashIsShown || !IsBackgroundLayerVisible())
+		{
+			DrawClearQuad(RHICmdList, FLinearColor::Black);
+		}
 
-	StereoLayerRender(RHICmdList, SortedSceneLayers, RenderParams);
-	
-	// Optionally render face-locked layers into a non-reprojected target if supported by the HMD platform
-	FTexture2DRHIRef OverlayRenderTarget = HMDDevice->GetOverlayLayerTarget_RenderThread(InView.StereoViewIndex, RenderParams.Viewport);
-	if (OverlayRenderTarget.IsValid())
-	{
-		RHICmdList.EndRenderPass();
+		StereoLayerRender(RHICmdList, SortedSceneLayers, RenderParams);
 
-		FRHIRenderPassInfo RPInfoOverlayRenderTarget(OverlayRenderTarget, ERenderTargetActions::Load_Store);
-		RHICmdList.BeginRenderPass(RPInfoOverlayRenderTarget, TEXT("StereoLayerRenderIntoOverlay"));
+		// Optionally render face-locked layers into a non-reprojected target if supported by the HMD platform
+		FTexture2DRHIRef OverlayRenderTarget = HMDDevice->GetOverlayLayerTarget_RenderThread(InView.StereoViewIndex, RenderParams.Viewport);
+		if (OverlayRenderTarget.IsValid())
+		{
+			RHICmdList.EndRenderPass();
 
-		DrawClearQuad(RHICmdList, FLinearColor(0.0f, 0.0f, 0.0f, 0.0f));
-		RHICmdList.SetViewport((float)RenderParams.Viewport.Min.X, (float)RenderParams.Viewport.Min.Y, 0.0f, (float)RenderParams.Viewport.Max.X, (float)RenderParams.Viewport.Max.Y, 1.0f);
-	}
+			FRHIRenderPassInfo RPInfoOverlayRenderTarget(OverlayRenderTarget, ERenderTargetActions::Load_Store);
+			RHICmdList.BeginRenderPass(RPInfoOverlayRenderTarget, TEXT("StereoLayerRenderIntoOverlay"));
+
+			DrawClearQuad(RHICmdList, FLinearColor(0.0f, 0.0f, 0.0f, 0.0f));
+			RHICmdList.SetViewport((float)RenderParams.Viewport.Min.X, (float)RenderParams.Viewport.Min.Y, 0.0f, (float)RenderParams.Viewport.Max.X, (float)RenderParams.Viewport.Max.Y, 1.0f);
+		}
 
-	StereoLayerRender(RHICmdList, SortedOverlayLayers, RenderParams);
+		StereoLayerRender(RHICmdList, SortedOverlayLayers, RenderParams);
 
-	RHICmdList.EndRenderPass();
+		RHICmdList.EndRenderPass();
+	});
 }
 
 
diff --git a/Engine/Source/Runtime/HeadMountedDisplay/Private/DefaultXRCamera.cpp b/Engine/Source/Runtime/HeadMountedDisplay/Private/DefaultXRCamera.cpp
--- a/Engine/Source/Runtime/HeadMountedDisplay/Private/DefaultXRCamera.cpp
+++ b/Engine/Source/Runtime/HeadMountedDisplay/Private/DefaultXRCamera.cpp
@@ -103,7 +103,7 @@ void FDefaultXRCamera::CalculateStereoCameraOffset(const int32 ViewIndex, FRotat
 	}
 }
 
-void FDefaultXRCamera::PreRenderView_RenderThread(FRHICommandListImmediate& RHICmdList, FSceneView& View)
+void FDefaultXRCamera::PreRenderView_RenderThread(FRDGBuilder& GraphBuilder, FSceneView& View)
 {
 	check(IsInRenderingThread());
 
@@ -153,7 +153,7 @@ void FDefaultXRCamera::BeginRenderViewFamily(FSceneViewFamily& InViewFamily)
 	TrackingSystem->OnBeginRendering_GameThread();
 }
 
-void FDefaultXRCamera::PreRenderViewFamily_RenderThread(FRHICommandListImmediate& RHICmdList, FSceneViewFamily& ViewFamily)
+void FDefaultXRCamera::PreRenderViewFamily_RenderThread(FRDGBuilder& GraphBuilder, FSceneViewFamily& ViewFamily)
 {
 	check(IsInRenderingThread());
 
@@ -163,7 +163,7 @@ void FDefaultXRCamera::PreRenderViewFamily_RenderThread(FRHICommandListImmediate
 		TrackingSystem->RefreshPoses();
 		PRAGMA_ENABLE_DEPRECATION_WARNINGS
 	}
-	TrackingSystem->OnBeginRendering_RenderThread(RHICmdList, ViewFamily);
+	TrackingSystem->OnBeginRendering_RenderThread(GraphBuilder.RHICmdList, ViewFamily);
 
 	{
 		FQuat CurrentOrientation;
@@ -180,7 +180,7 @@ void FDefaultXRCamera::PreRenderViewFamily_RenderThread(FRHICommandListImmediate
 			const FTransform CurrentRelativeTransform(CurrentOrientation, CurrentPosition);
 
 			LateUpdate.Apply_RenderThread(ViewFamily.Scene, OldRelativeTransform, CurrentRelativeTransform);
-			TrackingSystem->OnLateUpdateApplied_RenderThread(RHICmdList, CurrentRelativeTransform);
+			TrackingSystem->OnLateUpdateApplied_RenderThread(GraphBuilder.RHICmdList, CurrentRelativeTransform);
 
 			{
 				// Backwards compatibility during deprecation phase. Remove once IHeadMountedDisplay::BeginRendering_RenderThread has been removed.
@@ -188,7 +188,7 @@ void FDefaultXRCamera::PreRenderViewFamily_RenderThread(FRHICommandListImmediate
 					auto HMD = TrackingSystem->GetHMDDevice();
 				if (HMD)
 				{
-					HMD->BeginRendering_RenderThread(CurrentRelativeTransform, RHICmdList, ViewFamily);
+					HMD->BeginRendering_RenderThread(CurrentRelativeTransform, GraphBuilder.RHICmdList, ViewFamily);
 				}
 				PRAGMA_ENABLE_DEPRECATION_WARNINGS
 			}
diff --git a/Engine/Source/Runtime/HeadMountedDisplay/Private/MotionControllerComponent.cpp b/Engine/Source/Runtime/HeadMountedDisplay/Private/MotionControllerComponent.cpp
--- a/Engine/Source/Runtime/HeadMountedDisplay/Private/MotionControllerComponent.cpp
+++ b/Engine/Source/Runtime/HeadMountedDisplay/Private/MotionControllerComponent.cpp
@@ -641,7 +641,7 @@ void UMotionControllerComponent::FViewExtension::BeginRenderViewFamily(FSceneVie
 }
 
 //=============================================================================
-void UMotionControllerComponent::FViewExtension::PreRenderViewFamily_RenderThread(FRHICommandListImmediate& RHICmdList, FSceneViewFamily& InViewFamily)
+void UMotionControllerComponent::FViewExtension::PreRenderViewFamily_RenderThread(FRDGBuilder& GraphBuilder, FSceneViewFamily& InViewFamily)
 {
 	if (!MotionControllerComponent)
 	{
diff --git a/Engine/Source/Runtime/HeadMountedDisplay/Private/MotionDelayBuffer.cpp b/Engine/Source/Runtime/HeadMountedDisplay/Private/MotionDelayBuffer.cpp
--- a/Engine/Source/Runtime/HeadMountedDisplay/Private/MotionDelayBuffer.cpp
+++ b/Engine/Source/Runtime/HeadMountedDisplay/Private/MotionDelayBuffer.cpp
@@ -492,7 +492,7 @@ void FMotionDelayClient::BeginRenderViewFamily(FSceneViewFamily& ViewFamily)
 }
 
 //------------------------------------------------------------------------------
-void FMotionDelayClient::PostRenderViewFamily_RenderThread(FRHICommandListImmediate& /*RHICmdList*/, FSceneViewFamily& ViewFamily)
+void FMotionDelayClient::PostRenderViewFamily_RenderThread(FRDGBuilder& GraphBuilder, FSceneViewFamily& ViewFamily)
 {
 	if (MotionDelayService_Impl::PostRenderCleanupId_RenderThread != ViewFamily.FrameNumber)
 	{
diff --git a/Engine/Source/Runtime/HeadMountedDisplay/Public/DefaultStereoLayers.h b/Engine/Source/Runtime/HeadMountedDisplay/Public/DefaultStereoLayers.h
--- a/Engine/Source/Runtime/HeadMountedDisplay/Public/DefaultStereoLayers.h
+++ b/Engine/Source/Runtime/HeadMountedDisplay/Public/DefaultStereoLayers.h
@@ -19,9 +19,9 @@ public:
 	virtual void SetupViewFamily(FSceneViewFamily& InViewFamily) override;
 	virtual void SetupView(FSceneViewFamily& InViewFamily, FSceneView& InView) override {}
 	virtual void BeginRenderViewFamily(FSceneViewFamily& InViewFamily) override {}
-	virtual void PreRenderView_RenderThread(FRHICommandListImmediate& RHICmdList, FSceneView& InView) override {}
-	virtual void PreRenderViewFamily_RenderThread(FRHICommandListImmediate& RHICmdList, FSceneViewFamily& InViewFamily) override;
-	virtual void PostRenderView_RenderThread(FRHICommandListImmediate& RHICmdList, FSceneView& InView) override;
+	virtual void PreRenderView_RenderThread(FRDGBuilder& GraphBuilder, FSceneView& InView) override {}
+	virtual void PreRenderViewFamily_RenderThread(FRDGBuilder& GraphBuilder, FSceneViewFamily& InViewFamily) override;
+	virtual void PostRenderView_RenderThread(FRDGBuilder& GraphBuilder, FSceneView& InView) override;
 
 	virtual bool ShouldCopyDebugLayersToSpectatorScreen() const override
 	{
diff --git a/Engine/Source/Runtime/HeadMountedDisplay/Public/DefaultXRCamera.h b/Engine/Source/Runtime/HeadMountedDisplay/Public/DefaultXRCamera.h
--- a/Engine/Source/Runtime/HeadMountedDisplay/Public/DefaultXRCamera.h
+++ b/Engine/Source/Runtime/HeadMountedDisplay/Public/DefaultXRCamera.h
@@ -76,8 +76,8 @@ public:
 	virtual void SetupViewFamily(FSceneViewFamily& InViewFamily) override;
 	virtual void SetupView(FSceneViewFamily& InViewFamily, FSceneView& InView) override;
 	virtual void BeginRenderViewFamily(FSceneViewFamily& InViewFamily) override;
-	virtual void PreRenderView_RenderThread(FRHICommandListImmediate& RHICmdList, FSceneView& InView) override;
-	virtual void PreRenderViewFamily_RenderThread(FRHICommandListImmediate& RHICmdList, FSceneViewFamily& InViewFamily) override;
+	virtual void PreRenderView_RenderThread(FRDGBuilder& GraphBuilder, FSceneView& InView) override;
+	virtual void PreRenderViewFamily_RenderThread(FRDGBuilder& GraphBuilder, FSceneViewFamily& InViewFamily) override;
 
 	// FWorldSceneViewExtension interface:
 	virtual bool IsActiveThisFrame_Internal(const FSceneViewExtensionContext& Context) const override;
diff --git a/Engine/Source/Runtime/HeadMountedDisplay/Public/MotionControllerComponent.h b/Engine/Source/Runtime/HeadMountedDisplay/Public/MotionControllerComponent.h
--- a/Engine/Source/Runtime/HeadMountedDisplay/Public/MotionControllerComponent.h
+++ b/Engine/Source/Runtime/HeadMountedDisplay/Public/MotionControllerComponent.h
@@ -157,8 +157,8 @@ private:
 		virtual void SetupViewFamily(FSceneViewFamily& InViewFamily) override {}
 		virtual void SetupView(FSceneViewFamily& InViewFamily, FSceneView& InView) override {}
 		virtual void BeginRenderViewFamily(FSceneViewFamily& InViewFamily) override;
-		virtual void PreRenderView_RenderThread(FRHICommandListImmediate& RHICmdList, FSceneView& InView) override {}
-		virtual void PreRenderViewFamily_RenderThread(FRHICommandListImmediate& RHICmdList, FSceneViewFamily& InViewFamily) override;
+		virtual void PreRenderView_RenderThread(FRDGBuilder& GraphBuilder, FSceneView& InView) override {}
+		virtual void PreRenderViewFamily_RenderThread(FRDGBuilder& GraphBuilder, FSceneViewFamily& InViewFamily) override;
 		virtual int32 GetPriority() const override { return -10; }
 		virtual bool IsActiveThisFrame_Internal(const FSceneViewExtensionContext& Context) const;
 
diff --git a/Engine/Source/Runtime/HeadMountedDisplay/Public/MotionDelayBuffer.h b/Engine/Source/Runtime/HeadMountedDisplay/Public/MotionDelayBuffer.h
--- a/Engine/Source/Runtime/HeadMountedDisplay/Public/MotionDelayBuffer.h
+++ b/Engine/Source/Runtime/HeadMountedDisplay/Public/MotionDelayBuffer.h
@@ -46,9 +46,9 @@ public:
 	virtual void SetupViewFamily(FSceneViewFamily& InViewFamily) override {}
 	virtual void SetupView(FSceneViewFamily& InViewFamily, FSceneView& InView) override {}
 	virtual void BeginRenderViewFamily(FSceneViewFamily& InViewFamily) override;
-	virtual void PreRenderView_RenderThread(FRHICommandListImmediate& RHICmdList, FSceneView& InView) override {}
-	virtual void PreRenderViewFamily_RenderThread(FRHICommandListImmediate& RHICmdList, FSceneViewFamily& InViewFamily) override {}
-	virtual void PostRenderViewFamily_RenderThread(FRHICommandListImmediate& RHICmdList, FSceneViewFamily& InViewFamily) override;
+	virtual void PreRenderView_RenderThread(FRDGBuilder& GraphBuilder, FSceneView& InView) override {}
+	virtual void PreRenderViewFamily_RenderThread(FRDGBuilder& GraphBuilder, FSceneViewFamily& InViewFamily) override {}
+	virtual void PostRenderViewFamily_RenderThread(FRDGBuilder& GraphBuilder, FSceneViewFamily& InViewFamily) override;
 	virtual int32 GetPriority() const override;
 	virtual bool IsActiveThisFrame_Internal(const FSceneViewExtensionContext& Context) const override;
 
diff --git a/Engine/Source/Runtime/Landscape/Private/LandscapeRender.cpp b/Engine/Source/Runtime/Landscape/Private/LandscapeRender.cpp
--- a/Engine/Source/Runtime/Landscape/Private/LandscapeRender.cpp
+++ b/Engine/Source/Runtime/Landscape/Private/LandscapeRender.cpp
@@ -965,7 +965,7 @@ FLandscapeSceneViewExtension::~FLandscapeSceneViewExtension()
 	FCoreDelegates::OnEndFrameRT.RemoveAll(this);
 }
 
-void FLandscapeSceneViewExtension::PreRenderViewFamily_RenderThread(FRHICommandListImmediate& RHICmdList, FSceneViewFamily& InViewFamily)
+void FLandscapeSceneViewExtension::PreRenderViewFamily_RenderThread(FRDGBuilder& GraphBuilder, FSceneViewFamily& InViewFamily)
 {
 	for (auto& Pair : LandscapeRenderSystems)
 	{
@@ -974,7 +974,7 @@ void FLandscapeSceneViewExtension::PreRenderViewFamily_RenderThread(FRHICommandL
 	}
 }
 
-void FLandscapeSceneViewExtension::PreRenderView_RenderThread(FRHICommandListImmediate& RHICmdList, FSceneView& InView)
+void FLandscapeSceneViewExtension::PreRenderView_RenderThread(FRDGBuilder& GraphBuilder, FSceneView& InView)
 {
 	FShaderResourceViewRHIRef LODDataSRV;
 	FShaderResourceViewRHIRef IndirectionSRV;
diff --git a/Engine/Source/Runtime/Landscape/Public/LandscapeRender.h b/Engine/Source/Runtime/Landscape/Public/LandscapeRender.h
--- a/Engine/Source/Runtime/Landscape/Public/LandscapeRender.h
+++ b/Engine/Source/Runtime/Landscape/Public/LandscapeRender.h
@@ -490,8 +490,8 @@ public:
 	virtual void SetupView(FSceneViewFamily& InViewFamily, FSceneView& InView) override {}
 	virtual void BeginRenderViewFamily(FSceneViewFamily& InViewFamily) override {}
 
-	virtual void PreRenderViewFamily_RenderThread(FRHICommandListImmediate& RHICmdList, FSceneViewFamily& InViewFamily) override;
-	virtual void PreRenderView_RenderThread(FRHICommandListImmediate& RHICmdList, FSceneView& InView) override;
+	virtual void PreRenderViewFamily_RenderThread(FRDGBuilder& GraphBuilder, FSceneViewFamily& InViewFamily) override;
+	virtual void PreRenderView_RenderThread(FRDGBuilder& GraphBuilder, FSceneView& InView) override;
 
 	LANDSCAPE_API const TMap<uint32, FLandscapeRenderSystem*>& GetLandscapeRenderSystems() const;
 private:
diff --git a/Engine/Source/Runtime/MovieSceneCapture/Private/CompositionGraphCaptureProtocol.cpp b/Engine/Source/Runtime/MovieSceneCapture/Private/CompositionGraphCaptureProtocol.cpp
--- a/Engine/Source/Runtime/MovieSceneCapture/Private/CompositionGraphCaptureProtocol.cpp
+++ b/Engine/Source/Runtime/MovieSceneCapture/Private/CompositionGraphCaptureProtocol.cpp
@@ -126,8 +126,8 @@ struct FFrameCaptureViewExtension : public FSceneViewExtensionBase
 	}
 
 	virtual void BeginRenderViewFamily(FSceneViewFamily& InViewFamily) {}
-	virtual void PreRenderViewFamily_RenderThread(FRHICommandListImmediate& RHICmdList, FSceneViewFamily& InViewFamily) {}
-	virtual void PreRenderView_RenderThread(FRHICommandListImmediate& RHICmdList, FSceneView& InView) {}
+	virtual void PreRenderViewFamily_RenderThread(FRDGBuilder& GraphBuilder, FSceneViewFamily& InViewFamily) {}
+	virtual void PreRenderView_RenderThread(FRDGBuilder& GraphBuilder, FSceneView& InView) {}
 
 	virtual bool IsActiveThisFrame_Internal(const FSceneViewExtensionContext&) const override { return IsEnabled(); }
 
diff --git a/Engine/Source/Runtime/Renderer/Private/BasePassRendering.cpp b/Engine/Source/Runtime/Renderer/Private/BasePassRendering.cpp
--- a/Engine/Source/Runtime/Renderer/Private/BasePassRendering.cpp
+++ b/Engine/Source/Runtime/Renderer/Private/BasePassRendering.cpp
@@ -709,11 +709,6 @@ static void ClearGBufferAtMaxZ(
 	}
 }
 
-BEGIN_SHADER_PARAMETER_STRUCT(FPostBasePassViewExtensionParameters, )
-	SHADER_PARAMETER_RDG_UNIFORM_BUFFER(FSceneTextureUniformParameters, SceneTextures)
-	RENDER_TARGET_BINDING_SLOTS()
-END_SHADER_PARAMETER_STRUCT()
-
 void FDeferredShadingSceneRenderer::RenderBasePass(
 	FRDGBuilder& GraphBuilder,
 	FSceneTextures& SceneTextures,
@@ -904,30 +899,12 @@ void FDeferredShadingSceneRenderer::RenderBasePass(
 	RenderBasePassInternal(GraphBuilder, SceneTextures, BasePassRenderTargets, BasePassDepthStencilAccess, ForwardBasePassTextures, DBufferTextures, bDoParallelBasePass, bRenderLightmapDensity, InstanceCullingManager, bNaniteEnabled, NaniteRasterResults);
 	GraphBuilder.SetCommandListStat(GET_STATID(STAT_CLM_AfterBasePass));
 
-	if (ViewFamily.ViewExtensions.Num() > 0)
+	for (const TSharedRef<ISceneViewExtension>& ViewExtension : ViewFamily.ViewExtensions)
 	{
-		SCOPE_CYCLE_COUNTER(STAT_FDeferredShadingSceneRenderer_ViewExtensionPostRenderBasePass);
-		RDG_CSV_STAT_EXCLUSIVE_SCOPE(GraphBuilder, ViewExtensions);
-		RDG_EVENT_SCOPE(GraphBuilder, "BasePass_ViewExtensions");
-		auto* PassParameters = GraphBuilder.AllocParameters<FPostBasePassViewExtensionParameters>();
-		PassParameters->RenderTargets = BasePassRenderTargets;
-		PassParameters->SceneTextures = CreateSceneTextureUniformBuffer(GraphBuilder, &GetActiveSceneTextures(), FeatureLevel, ESceneTextureSetupMode::None);
-
-		GraphBuilder.AddPass(
-			{},
-			PassParameters,
-			ERDGPassFlags::Raster,
-			[this](FRHICommandListImmediate& RHICmdList)
+		for (FViewInfo& View : Views)
 		{
-			for (auto& ViewExtension : ViewFamily.ViewExtensions)
-			{
-				for (FViewInfo& View : Views)
-				{
-					SCOPED_GPU_MASK(RHICmdList, View.GPUMask);
-					ViewExtension->PostRenderBasePass_RenderThread(RHICmdList, View);
-				}
-			}
-		});
+			ViewExtension->PostRenderBasePassDeferred_RenderThread(GraphBuilder, View, BasePassRenderTargets, SceneTextures.UniformBuffer);
+		}
 	}
 
 	if (bRequiresFarZQuadClear)
diff --git a/Engine/Source/Runtime/Renderer/Private/MobileShadingRenderer.cpp b/Engine/Source/Runtime/Renderer/Private/MobileShadingRenderer.cpp
--- a/Engine/Source/Runtime/Renderer/Private/MobileShadingRenderer.cpp
+++ b/Engine/Source/Runtime/Renderer/Private/MobileShadingRenderer.cpp
@@ -1564,7 +1564,7 @@ void FMobileSceneRenderer::PostRenderBasePass(FRHICommandListImmediate& RHICmdLi
 		QUICK_SCOPE_CYCLE_COUNTER(STAT_FMobileSceneRenderer_ViewExtensionPostRenderBasePass);
 		for (int32 ViewExt = 0; ViewExt < ViewFamily.ViewExtensions.Num(); ++ViewExt)
 		{
-			ViewFamily.ViewExtensions[ViewExt]->PostRenderBasePass_RenderThread(RHICmdList, View);
+			ViewFamily.ViewExtensions[ViewExt]->PostRenderBasePassMobile_RenderThread(RHICmdList, View);
 		}
 	}
 }
diff --git a/Engine/Source/Runtime/Renderer/Private/PlanarReflectionRendering.cpp b/Engine/Source/Runtime/Renderer/Private/PlanarReflectionRendering.cpp
--- a/Engine/Source/Runtime/Renderer/Private/PlanarReflectionRendering.cpp
+++ b/Engine/Source/Runtime/Renderer/Private/PlanarReflectionRendering.cpp
@@ -276,11 +276,17 @@ static void UpdatePlanarReflectionContents_RenderThread(
 {
 	QUICK_SCOPE_CYCLE_COUNTER(STAT_RenderPlanarReflection);
 
-	// We need to execute the pre-render view extensions before we do any view dependent work.
-	FSceneRenderer::ViewExtensionPreRender_RenderThread(RHICmdList, SceneRenderer);
-
 	SceneRenderer->RenderThreadBegin(RHICmdList);
 
+	// update any resources that needed a deferred update
+	FDeferredUpdateResource::UpdateResources(RHICmdList);
+
+	const ERHIFeatureLevel::Type FeatureLevel = SceneRenderer->FeatureLevel;
+	FRDGBuilder GraphBuilder(RHICmdList, RDG_EVENT_NAME("PlanarReflection"), FSceneRenderer::GetRDGParalelExecuteFlags(FeatureLevel));
+
+	// We need to execute the pre-render view extensions before we do any view dependent work.
+	FSceneRenderer::ViewExtensionPreRender_RenderThread(GraphBuilder, SceneRenderer);
+
 	// Make sure we render to the same set of GPUs as the main scene renderer.
 	if (MainSceneRenderer->ViewFamily.RenderTarget != nullptr)
 	{
@@ -344,20 +350,14 @@ static void UpdatePlanarReflectionContents_RenderThread(
 
 		if (bIsVisibleInAnyView)
 		{
-			// update any resources that needed a deferred update
-			FDeferredUpdateResource::UpdateResources(RHICmdList);
-
 			{
 #if WANTS_DRAW_MESH_EVENTS
 				FString EventName;
 				OwnerName.ToString(EventName);
-				SCOPED_DRAW_EVENTF(RHICmdList, SceneCapture, TEXT("PlanarReflection %s"), *EventName);
+				RDG_EVENT_SCOPE(GraphBuilder, "PlanarReflection %s", *EventName);
 #else
-				SCOPED_DRAW_EVENT(RHICmdList, UpdatePlanarReflectionContent_RenderThread);
+				RDG_EVENT_SCOPE(GraphBuilder, "UpdatePlanarReflectionContent_RenderThread");
 #endif
-				const ERHIFeatureLevel::Type FeatureLevel = SceneRenderer->FeatureLevel;
-				FRDGBuilder GraphBuilder(RHICmdList, RDG_EVENT_NAME("PlanarReflection"), FSceneRenderer::GetRDGParalelExecuteFlags(FeatureLevel));
-
 				// Reflection view late update
 				if (SceneRenderer->Views.Num() > 1)
 				{
@@ -406,12 +406,12 @@ static void UpdatePlanarReflectionContents_RenderThread(
 						PrefilterPlanarReflection<false>(GraphBuilder, View, SceneTextureParameters, SceneProxy, SceneTextures.Color.Resolve, ReflectionOutputTexture);
 					}
 				}
-
-				GraphBuilder.Execute();
 			}
 		}
 	}
 
+	GraphBuilder.Execute();
+
 	SceneRenderer->RenderThreadEnd(RHICmdList);
 }
 
@@ -427,9 +427,6 @@ static void UpdatePlanarReflectionContentsWithoutRendering_RenderThread(
 {
 	QUICK_SCOPE_CYCLE_COUNTER(STAT_RenderPlanarReflection);
 
-	// We need to execute the pre-render view extensions before we do any view dependent work.
-	FSceneRenderer::ViewExtensionPreRender_RenderThread(RHICmdList, SceneRenderer);
-
 	SceneRenderer->RenderThreadBegin(RHICmdList);
 
 	FBox PlanarReflectionBounds = SceneProxy->WorldBounds;
diff --git a/Engine/Source/Runtime/Renderer/Private/ReflectionEnvironmentCapture.cpp b/Engine/Source/Runtime/Renderer/Private/ReflectionEnvironmentCapture.cpp
--- a/Engine/Source/Runtime/Renderer/Private/ReflectionEnvironmentCapture.cpp
+++ b/Engine/Source/Runtime/Renderer/Private/ReflectionEnvironmentCapture.cpp
@@ -473,19 +473,20 @@
 	const FLinearColor& LowerHemisphereColor,
 	bool bCapturingForMobile)
 {
-	// We need to execute the pre-render view extensions before we do any view dependent work.
-	FSceneRenderer::ViewExtensionPreRender_RenderThread(RHICmdList, SceneRenderer);
-
 	SceneRenderer->RenderThreadBegin(RHICmdList);
 
+	const ERHIFeatureLevel::Type FeatureLevel = SceneRenderer->FeatureLevel;
+
 	// update any resources that needed a deferred update
 	FDeferredUpdateResource::UpdateResources(RHICmdList);
-	FMaterialRenderProxy::UpdateDeferredCachedUniformExpressions();
 
-	const ERHIFeatureLevel::Type FeatureLevel = SceneRenderer->FeatureLevel;
-
 	FRDGBuilder GraphBuilder(RHICmdList, RDG_EVENT_NAME("CubeMapCapture"), FSceneRenderer::GetRDGParalelExecuteFlags(FeatureLevel));
 
+	// We need to execute the pre-render view extensions before we do any view dependent work.
+	FSceneRenderer::ViewExtensionPreRender_RenderThread(GraphBuilder, SceneRenderer);
+
+	FMaterialRenderProxy::UpdateDeferredCachedUniformExpressions();
+
 	{
 		RDG_EVENT_SCOPE(GraphBuilder, "CubeMapCapture");
 
diff --git a/Engine/Source/Runtime/Renderer/Private/SceneCaptureRendering.cpp b/Engine/Source/Runtime/Renderer/Private/SceneCaptureRendering.cpp
--- a/Engine/Source/Runtime/Renderer/Private/SceneCaptureRendering.cpp
+++ b/Engine/Source/Runtime/Renderer/Private/SceneCaptureRendering.cpp
@@ -268,24 +268,25 @@ static void UpdateSceneCaptureContentDeferred_RenderThread(
 	bool bOrthographicCamera
 	)
 {
-	// We need to execute the pre-render view extensions before we do any view dependent work.
-	FSceneRenderer::ViewExtensionPreRender_RenderThread(RHICmdList, SceneRenderer);
-
 	SceneRenderer->RenderThreadBegin(RHICmdList);
 
 	// update any resources that needed a deferred update
 	FDeferredUpdateResource::UpdateResources(RHICmdList);
-	{
-		const ERHIFeatureLevel::Type FeatureLevel = SceneRenderer->FeatureLevel;
+
+	const ERHIFeatureLevel::Type FeatureLevel = SceneRenderer->FeatureLevel;
 
 #if WANTS_DRAW_MESH_EVENTS
-		SCOPED_DRAW_EVENTF(RHICmdList, SceneCapture, TEXT("SceneCapture %s"), *EventName);
-		FRDGBuilder GraphBuilder(RHICmdList, RDG_EVENT_NAME("SceneCapture %s", *EventName), FSceneRenderer::GetRDGParalelExecuteFlags(FeatureLevel));
+	SCOPED_DRAW_EVENTF(RHICmdList, SceneCapture, TEXT("SceneCapture %s"), *EventName);
+	FRDGBuilder GraphBuilder(RHICmdList, RDG_EVENT_NAME("SceneCapture %s", *EventName), FSceneRenderer::GetRDGParalelExecuteFlags(FeatureLevel));
 #else
-		SCOPED_DRAW_EVENT(RHICmdList, UpdateSceneCaptureContent_RenderThread);
-		FRDGBuilder GraphBuilder(RHICmdList, RDG_EVENT_NAME("SceneCapture"), FSceneRenderer::GetRDGParalelExecuteFlags(FeatureLevel));
+	SCOPED_DRAW_EVENT(RHICmdList, UpdateSceneCaptureContent_RenderThread);
+	FRDGBuilder GraphBuilder(RHICmdList, RDG_EVENT_NAME("SceneCapture"), FSceneRenderer::GetRDGParalelExecuteFlags(FeatureLevel));
 #endif
 
+	// We need to execute the pre-render view extensions before we do any view dependent work.
+	FSceneRenderer::ViewExtensionPreRender_RenderThread(GraphBuilder, SceneRenderer);
+
+	{
 		FRDGTextureRef TargetTexture = RegisterExternalTexture(GraphBuilder, RenderTarget->GetRenderTargetTexture(), TEXT("SceneCaptureTarget"));
 		FRDGTextureRef ShaderResourceTexture = RegisterExternalTexture(GraphBuilder, RenderTargetTexture->TextureRHI, TEXT("SceneCaptureTexture"));
 
@@ -355,23 +356,23 @@ void UpdateSceneCaptureContentMobile_RenderThread(
 	bool bGenerateMips,
 	const FGenerateMipsParams& GenerateMipsParams)
 {
-	// We need to execute the pre-render view extensions before we do any view dependent work.
-	FSceneRenderer::ViewExtensionPreRender_RenderThread(RHICmdList, SceneRenderer);
-
 	SceneRenderer->RenderThreadBegin(RHICmdList);
 
 	// update any resources that needed a deferred update
 	FDeferredUpdateResource::UpdateResources(RHICmdList);
 
-	{
 #if WANTS_DRAW_MESH_EVENTS
-		SCOPED_DRAW_EVENTF(RHICmdList, SceneCaptureMobile, TEXT("SceneCaptureMobile %s"), *EventName);
-		FRDGBuilder GraphBuilder(RHICmdList, RDG_EVENT_NAME("SceneCaptureMobile %s", *EventName));
+	SCOPED_DRAW_EVENTF(RHICmdList, SceneCaptureMobile, TEXT("SceneCaptureMobile %s"), *EventName);
+	FRDGBuilder GraphBuilder(RHICmdList, RDG_EVENT_NAME("SceneCaptureMobile %s", *EventName));
 #else
-		SCOPED_DRAW_EVENT(RHICmdList, UpdateSceneCaptureContentMobile_RenderThread);
-		FRDGBuilder GraphBuilder(RHICmdList, RDG_EVENT_NAME("SceneCaptureMobile"));
+	SCOPED_DRAW_EVENT(RHICmdList, UpdateSceneCaptureContentMobile_RenderThread);
+	FRDGBuilder GraphBuilder(RHICmdList, RDG_EVENT_NAME("SceneCaptureMobile"));
 #endif
 
+	// We need to execute the pre-render view extensions before we do any view dependent work.
+	FSceneRenderer::ViewExtensionPreRender_RenderThread(GraphBuilder, SceneRenderer);
+
+	{
 		FViewInfo& View = SceneRenderer->Views[0];
 
 		// Intermediate render target that will need to be flipped (needed on !IsMobileHDR())
diff --git a/Engine/Source/Runtime/Renderer/Private/SceneRendering.cpp b/Engine/Source/Runtime/Renderer/Private/SceneRendering.cpp
--- a/Engine/Source/Runtime/Renderer/Private/SceneRendering.cpp
+++ b/Engine/Source/Runtime/Renderer/Private/SceneRendering.cpp
@@ -3955,7 +3955,7 @@ void FSceneRenderer::UpdatePrimitiveIndirectLightingCacheBuffers()
 *
 * @param SceneRenderer	Scene renderer to use for rendering.
 */
-void FSceneRenderer::ViewExtensionPreRender_RenderThread(FRHICommandListImmediate& RHICmdList, FSceneRenderer* SceneRenderer)
+void FSceneRenderer::ViewExtensionPreRender_RenderThread(FRDGBuilder& GraphBuilder, FSceneRenderer* SceneRenderer)
 {
 	if (SceneRenderer->ViewFamily.ViewExtensions.IsEmpty())
 	{
@@ -3963,7 +3963,6 @@ void FSceneRenderer::ViewExtensionPreRender_RenderThread(FRHICommandListImmediat
 	}
 
 	{
-		FRDGBuilder GraphBuilder(RHICmdList);
 		{
 			RDG_CSV_STAT_EXCLUSIVE_SCOPE(GraphBuilder, PreRender);
 			SCOPE_CYCLE_COUNTER(STAT_FDeferredShadingSceneRenderer_ViewExtensionPreRenderView);
@@ -3977,11 +3976,8 @@ void FSceneRenderer::ViewExtensionPreRender_RenderThread(FRHICommandListImmediat
 				}
 			}
 		}
-		GraphBuilder.Execute();
 	}
-	
-	// update any resources that needed a deferred update
-	FDeferredUpdateResource::UpdateResources(RHICmdList);
+	check(!FDeferredUpdateResource::IsUpdateNeeded());
 }
 
 static int32 GSceneRenderCleanUpMode = 2;
@@ -4272,17 +4268,8 @@ static void RenderViewFamilies_RenderThread(FRHICommandListImmediate& RHICmdList
 	// All renderers point to the same Scene (calling code asserts this)
 	FScene* const Scene = SceneRenderers[0]->Scene;
 
-	// We need to execute the pre-render view extensions before we do any view dependent work.
-	for (FSceneRenderer* SceneRenderer : SceneRenderers)
-	{
-		FSceneRenderer::ViewExtensionPreRender_RenderThread(RHICmdList, SceneRenderer);
-	}
-
 	FSceneRenderer::RenderThreadBegin(RHICmdList, SceneRenderers);
 
-	// update any resources that needed a deferred update
-	FDeferredUpdateResource::UpdateResources(RHICmdList);
-
 	bool bAnyShowHitProxies = false;
 
 #if WITH_DEBUG_VIEW_MODES
@@ -4290,10 +4277,27 @@ static void RenderViewFamilies_RenderThread(FRHICommandListImmediate& RHICmdList
 	bool bUpdatedGPUSkinCacheVisualization = false;
 #endif
 
+	// update any resources that needed a deferred update
+	FDeferredUpdateResource::UpdateResources(RHICmdList);
+
 	for (FSceneRenderer* SceneRenderer : SceneRenderers)
 	{
+		const ERHIFeatureLevel::Type FeatureLevel = SceneRenderer->FeatureLevel;
+		
 		FSceneViewFamily& ViewFamily = SceneRenderer->ViewFamily;
 
+		FRDGBuilder GraphBuilder(
+			RHICmdList,
+			RDG_EVENT_NAME("SceneRenderer_%s(ViewFamily=%s)",
+				ViewFamily.EngineShowFlags.HitProxies ? TEXT("RenderHitProxies") : TEXT("Render"),
+				ViewFamily.bResolveScene ? TEXT("Primary") : TEXT("Auxiliary")
+			),
+			FSceneRenderer::GetRDGParalelExecuteFlags(FeatureLevel)
+		);
+
+		// We need to execute the pre-render view extensions before we do any view dependent work.
+		FSceneRenderer::ViewExtensionPreRender_RenderThread(GraphBuilder, SceneRenderer);
+
 		SCOPE_CYCLE_COUNTER_VERBOSE(STAT_TotalSceneRenderingTime, ViewFamily.ProfileDescription.IsEmpty() ? nullptr : *ViewFamily.ProfileDescription);
 		const uint64 FamilyRenderStart = FPlatformTime::Cycles64();
 
@@ -4322,41 +4326,28 @@ static void RenderViewFamilies_RenderThread(FRHICommandListImmediate& RHICmdList
 		}
 #endif  // WITH_DEBUG_VIEW_MODES
 
-		{
-			const ERHIFeatureLevel::Type FeatureLevel = SceneRenderer->FeatureLevel;
-
-			FRDGBuilder GraphBuilder(
-				RHICmdList,
-				RDG_EVENT_NAME("SceneRenderer_%s(ViewFamily=%s)",
-					ViewFamily.EngineShowFlags.HitProxies ? TEXT("RenderHitProxies") : TEXT("Render"),
-					ViewFamily.bResolveScene ? TEXT("Primary") : TEXT("Auxiliary")
-				),
-				FSceneRenderer::GetRDGParalelExecuteFlags(FeatureLevel)
-			);
-
 #if WITH_MGPU
-			if (ViewFamily.bForceCopyCrossGPU)
-			{
-				GraphBuilder.EnableForceCopyCrossGPU();
-			}
+		if (ViewFamily.bForceCopyCrossGPU)
+		{
+			GraphBuilder.EnableForceCopyCrossGPU();
+		}
 #endif
 
-			if (ViewFamily.EngineShowFlags.HitProxies)
-			{
-				// Render the scene's hit proxies.
-				SceneRenderer->RenderHitProxies(GraphBuilder);
-				bAnyShowHitProxies = true;
-			}
-			else
-			{
-				// Render the scene.
-				SceneRenderer->Render(GraphBuilder);
-			}
+		if (ViewFamily.EngineShowFlags.HitProxies)
+		{
+			// Render the scene's hit proxies.
+			SceneRenderer->RenderHitProxies(GraphBuilder);
+			bAnyShowHitProxies = true;
+		}
+		else
+		{
+			// Render the scene.
+			SceneRenderer->Render(GraphBuilder);
+		}
 
-			SceneRenderer->FlushCrossGPUFences(GraphBuilder);
+		SceneRenderer->FlushCrossGPUFences(GraphBuilder);
 
-			GraphBuilder.Execute();
-		}
+		GraphBuilder.Execute();
 
 		if (SceneRenderer->ViewFamily.ProfileSceneRenderTime)
 		{
@@ -5014,109 +5005,112 @@ static uint32 ComputeScalabilityCVarHash()
 	return Ret;
 }
 
-static void DisplayInternals(FRHICommandListImmediate& RHICmdList, FViewInfo& InView)
+static void DisplayInternals(FRDGBuilder& GraphBuilder, FViewInfo& InView)
 {
 #if !(UE_BUILD_SHIPPING || UE_BUILD_TEST)
 	auto Family = InView.Family;
 	// if r.DisplayInternals != 0
 	if(Family->EngineShowFlags.OnScreenDebug && Family->DisplayInternalsData.IsValid())
 	{
-		// could be 0
-		auto State = InView.ViewState;
+		AddPass(GraphBuilder, RDG_EVENT_NAME("DisplayInternals"), [Family, &InView] (FRHICommandListImmediate& RHICmdList)
+		{
+			// could be 0
+			auto State = InView.ViewState;
 
-		FCanvas Canvas((FRenderTarget*)Family->RenderTarget, NULL, Family->Time, InView.GetFeatureLevel());
-		Canvas.SetRenderTargetRect(FIntRect(0, 0, Family->RenderTarget->GetSizeXY().X, Family->RenderTarget->GetSizeXY().Y));
+			FCanvas Canvas((FRenderTarget*)Family->RenderTarget, NULL, Family->Time, InView.GetFeatureLevel());
+			Canvas.SetRenderTargetRect(FIntRect(0, 0, Family->RenderTarget->GetSizeXY().X, Family->RenderTarget->GetSizeXY().Y));
 
 
-		FRHIRenderPassInfo RenderPassInfo(Family->RenderTarget->GetRenderTargetTexture(), ERenderTargetActions::Load_Store);
-		RHICmdList.BeginRenderPass(RenderPassInfo, TEXT("DisplayInternalsRenderPass"));
+			FRHIRenderPassInfo RenderPassInfo(Family->RenderTarget->GetRenderTargetTexture(), ERenderTargetActions::Load_Store);
+			RHICmdList.BeginRenderPass(RenderPassInfo, TEXT("DisplayInternalsRenderPass"));
 
-		// further down to not intersect with "LIGHTING NEEDS TO BE REBUILT"
-		FVector2D Pos(30, 140);
-		const int32 FontSizeY = 14;
+			// further down to not intersect with "LIGHTING NEEDS TO BE REBUILT"
+			FVector2D Pos(30, 140);
+			const int32 FontSizeY = 14;
 
-		// dark background
-		const uint32 BackgroundHeight = 30;
-		Canvas.DrawTile(Pos.X - 4, Pos.Y - 4, 500 + 8, FontSizeY * BackgroundHeight + 8, 0, 0, 1, 1, FLinearColor(0,0,0,0.6f), 0, true);
+			// dark background
+			const uint32 BackgroundHeight = 30;
+			Canvas.DrawTile(Pos.X - 4, Pos.Y - 4, 500 + 8, FontSizeY * BackgroundHeight + 8, 0, 0, 1, 1, FLinearColor(0,0,0,0.6f), 0, true);
 
-		UFont* Font = GEngine->GetSmallFont();
-		FCanvasTextItem SmallTextItem( Pos, FText::GetEmpty(), GEngine->GetSmallFont(), FLinearColor::White );
+			UFont* Font = GEngine->GetSmallFont();
+			FCanvasTextItem SmallTextItem( Pos, FText::GetEmpty(), GEngine->GetSmallFont(), FLinearColor::White );
 
-		SmallTextItem.SetColor(FLinearColor::White);
-		SmallTextItem.Text = FText::FromString(FString::Printf(TEXT("r.DisplayInternals = %d"), Family->DisplayInternalsData.DisplayInternalsCVarValue));
-		Canvas.DrawItem(SmallTextItem, Pos);
-		SmallTextItem.SetColor(FLinearColor::Gray);
-		Pos.Y += 2 * FontSizeY;
+			SmallTextItem.SetColor(FLinearColor::White);
+			SmallTextItem.Text = FText::FromString(FString::Printf(TEXT("r.DisplayInternals = %d"), Family->DisplayInternalsData.DisplayInternalsCVarValue));
+			Canvas.DrawItem(SmallTextItem, Pos);
+			SmallTextItem.SetColor(FLinearColor::Gray);
+			Pos.Y += 2 * FontSizeY;
 
-		FViewInfo& ViewInfo = (FViewInfo&)InView;
-#define CANVAS_HEADER(txt) \
-		{ \
-			SmallTextItem.SetColor(FLinearColor::Gray); \
-			SmallTextItem.Text = FText::FromString(txt); \
-			Canvas.DrawItem(SmallTextItem, Pos); \
-			Pos.Y += FontSizeY; \
-		}
-#define CANVAS_LINE(bHighlight, txt, ... ) \
-		{ \
-			SmallTextItem.SetColor(bHighlight ? FLinearColor::Red : FLinearColor::Gray); \
-			SmallTextItem.Text = FText::FromString(FString::Printf(txt, __VA_ARGS__)); \
-			Canvas.DrawItem(SmallTextItem, Pos); \
-			Pos.Y += FontSizeY; \
-		}
+			FViewInfo& ViewInfo = (FViewInfo&)InView;
+	#define CANVAS_HEADER(txt) \
+			{ \
+				SmallTextItem.SetColor(FLinearColor::Gray); \
+				SmallTextItem.Text = FText::FromString(txt); \
+				Canvas.DrawItem(SmallTextItem, Pos); \
+				Pos.Y += FontSizeY; \
+			}
+	#define CANVAS_LINE(bHighlight, txt, ... ) \
+			{ \
+				SmallTextItem.SetColor(bHighlight ? FLinearColor::Red : FLinearColor::Gray); \
+				SmallTextItem.Text = FText::FromString(FString::Printf(txt, __VA_ARGS__)); \
+				Canvas.DrawItem(SmallTextItem, Pos); \
+				Pos.Y += FontSizeY; \
+			}
 
-		CANVAS_HEADER(TEXT("command line options:"))
-		{
-			bool bHighlight = !(FApp::UseFixedTimeStep() && FApp::bUseFixedSeed);
-			CANVAS_LINE(bHighlight, TEXT("  -UseFixedTimeStep: %u"), FApp::UseFixedTimeStep())
-			CANVAS_LINE(bHighlight, TEXT("  -FixedSeed: %u"), FApp::bUseFixedSeed)
-			CANVAS_LINE(false, TEXT("  -gABC= (changelist): %d"), GetChangeListNumberForPerfTesting())
-		}
+			CANVAS_HEADER(TEXT("command line options:"))
+			{
+				bool bHighlight = !(FApp::UseFixedTimeStep() && FApp::bUseFixedSeed);
+				CANVAS_LINE(bHighlight, TEXT("  -UseFixedTimeStep: %u"), FApp::UseFixedTimeStep())
+				CANVAS_LINE(bHighlight, TEXT("  -FixedSeed: %u"), FApp::bUseFixedSeed)
+				CANVAS_LINE(false, TEXT("  -gABC= (changelist): %d"), GetChangeListNumberForPerfTesting())
+			}
 
-		CANVAS_HEADER(TEXT("Global:"))
-		CANVAS_LINE(false, TEXT("  FrameNumberRT: %u"), GFrameNumberRenderThread)
-		CANVAS_LINE(false, TEXT("  Scalability CVar Hash: %x (use console command \"Scalability\")"), ComputeScalabilityCVarHash())
-		//not really useful as it is non deterministic and should not be used for rendering features:  CANVAS_LINE(false, TEXT("  FrameNumberRT: %u"), GFrameNumberRenderThread)
-		CANVAS_LINE(false, TEXT("  FrameCounter: %llu"), (uint64)GFrameCounter)
-		CANVAS_LINE(false, TEXT("  rand()/SRand: %x/%x"), FMath::Rand(), FMath::GetRandSeed())
-		{
-			bool bHighlight = Family->DisplayInternalsData.NumPendingStreamingRequests != 0;
-			CANVAS_LINE(bHighlight, TEXT("  FStreamAllResourcesLatentCommand: %d"), bHighlight)
-		}
-		{
-			static auto* Var = IConsoleManager::Get().FindTConsoleVariableDataInt(TEXT("r.Streaming.FramesForFullUpdate"));
-			int32 Value = Var->GetValueOnRenderThread();
-			bool bHighlight = Value != 0;
-			CANVAS_LINE(bHighlight, TEXT("  r.Streaming.FramesForFullUpdate: %u%s"), Value, bHighlight ? TEXT(" (should be 0)") : TEXT(""));
-		}
+			CANVAS_HEADER(TEXT("Global:"))
+			CANVAS_LINE(false, TEXT("  FrameNumberRT: %u"), GFrameNumberRenderThread)
+			CANVAS_LINE(false, TEXT("  Scalability CVar Hash: %x (use console command \"Scalability\")"), ComputeScalabilityCVarHash())
+			//not really useful as it is non deterministic and should not be used for rendering features:  CANVAS_LINE(false, TEXT("  FrameNumberRT: %u"), GFrameNumberRenderThread)
+			CANVAS_LINE(false, TEXT("  FrameCounter: %llu"), (uint64)GFrameCounter)
+			CANVAS_LINE(false, TEXT("  rand()/SRand: %x/%x"), FMath::Rand(), FMath::GetRandSeed())
+			{
+				bool bHighlight = Family->DisplayInternalsData.NumPendingStreamingRequests != 0;
+				CANVAS_LINE(bHighlight, TEXT("  FStreamAllResourcesLatentCommand: %d"), bHighlight)
+			}
+			{
+				static auto* Var = IConsoleManager::Get().FindTConsoleVariableDataInt(TEXT("r.Streaming.FramesForFullUpdate"));
+				int32 Value = Var->GetValueOnRenderThread();
+				bool bHighlight = Value != 0;
+				CANVAS_LINE(bHighlight, TEXT("  r.Streaming.FramesForFullUpdate: %u%s"), Value, bHighlight ? TEXT(" (should be 0)") : TEXT(""));
+			}
 
-		if(State)
-		{
-			CANVAS_HEADER(TEXT("State:"))
-			CANVAS_LINE(false, TEXT("  TemporalAASample: %u"), State->GetCurrentTemporalAASampleIndex())
-			CANVAS_LINE(false, TEXT("  FrameIndexMod8: %u"), State->GetFrameIndex(8))
-			CANVAS_LINE(false, TEXT("  LODTransition: %.2f"), State->GetTemporalLODTransition())
-		}
+			if(State)
+			{
+				CANVAS_HEADER(TEXT("State:"))
+				CANVAS_LINE(false, TEXT("  TemporalAASample: %u"), State->GetCurrentTemporalAASampleIndex())
+				CANVAS_LINE(false, TEXT("  FrameIndexMod8: %u"), State->GetFrameIndex(8))
+				CANVAS_LINE(false, TEXT("  LODTransition: %.2f"), State->GetTemporalLODTransition())
+			}
 
-		CANVAS_HEADER(TEXT("Family:"))
-		CANVAS_LINE(false, TEXT("  Time (Real/World/DeltaWorld): %.2f/%.2f/%.2f"), Family->Time.GetRealTimeSeconds(), Family->Time.GetWorldTimeSeconds(), Family->Time.GetDeltaWorldTimeSeconds())
-		CANVAS_LINE(false, TEXT("  FrameNumber: %u"), Family->FrameNumber)
-		CANVAS_LINE(false, TEXT("  ExposureSettings: %s"), *Family->ExposureSettings.ToString())
-		CANVAS_LINE(false, TEXT("  GammaCorrection: %.2f"), Family->GammaCorrection)
+			CANVAS_HEADER(TEXT("Family:"))
+			CANVAS_LINE(false, TEXT("  Time (Real/World/DeltaWorld): %.2f/%.2f/%.2f"), Family->Time.GetRealTimeSeconds(), Family->Time.GetWorldTimeSeconds(), Family->Time.GetDeltaWorldTimeSeconds())
+			CANVAS_LINE(false, TEXT("  FrameNumber: %u"), Family->FrameNumber)
+			CANVAS_LINE(false, TEXT("  ExposureSettings: %s"), *Family->ExposureSettings.ToString())
+			CANVAS_LINE(false, TEXT("  GammaCorrection: %.2f"), Family->GammaCorrection)
 
-		CANVAS_HEADER(TEXT("View:"))
-		CANVAS_LINE(false, TEXT("  TemporalJitter: %.2f/%.2f"), ViewInfo.TemporalJitterPixels.X, ViewInfo.TemporalJitterPixels.Y)
-		CANVAS_LINE(false, TEXT("  ViewProjectionMatrix Hash: %x"), InView.ViewMatrices.GetViewProjectionMatrix().ComputeHash())
-		CANVAS_LINE(false, TEXT("  ViewLocation: %s"), *InView.ViewLocation.ToString())
-		CANVAS_LINE(false, TEXT("  ViewRotation: %s"), *InView.ViewRotation.ToString())
-		CANVAS_LINE(false, TEXT("  ViewRect: %s"), *ViewInfo.ViewRect.ToString())
+			CANVAS_HEADER(TEXT("View:"))
+			CANVAS_LINE(false, TEXT("  TemporalJitter: %.2f/%.2f"), ViewInfo.TemporalJitterPixels.X, ViewInfo.TemporalJitterPixels.Y)
+			CANVAS_LINE(false, TEXT("  ViewProjectionMatrix Hash: %x"), InView.ViewMatrices.GetViewProjectionMatrix().ComputeHash())
+			CANVAS_LINE(false, TEXT("  ViewLocation: %s"), *InView.ViewLocation.ToString())
+			CANVAS_LINE(false, TEXT("  ViewRotation: %s"), *InView.ViewRotation.ToString())
+			CANVAS_LINE(false, TEXT("  ViewRect: %s"), *ViewInfo.ViewRect.ToString())
 
-		CANVAS_LINE(false, TEXT("  DynMeshElements/TranslPrim: %d/%d"), ViewInfo.DynamicMeshElements.Num(), ViewInfo.TranslucentPrimCount.NumPrims())
+			CANVAS_LINE(false, TEXT("  DynMeshElements/TranslPrim: %d/%d"), ViewInfo.DynamicMeshElements.Num(), ViewInfo.TranslucentPrimCount.NumPrims())
 
-#undef CANVAS_LINE
-#undef CANVAS_HEADER
+	#undef CANVAS_LINE
+	#undef CANVAS_HEADER
 
-		RHICmdList.EndRenderPass();
-		Canvas.Flush_RenderThread(RHICmdList);
+			RHICmdList.EndRenderPass();
+			Canvas.Flush_RenderThread(RHICmdList);
+		});
 	}
 #endif
 
@@ -5130,13 +5124,13 @@ TSharedRef<ISceneViewExtension, ESPMode::ThreadSafe> GetRendererViewExtension()
 		virtual void SetupViewFamily(FSceneViewFamily& InViewFamily) {}
 		virtual void SetupView(FSceneViewFamily& InViewFamily, FSceneView& InView) {}
 		virtual void BeginRenderViewFamily(FSceneViewFamily& InViewFamily) {}
-		virtual void PreRenderViewFamily_RenderThread(FRHICommandListImmediate& RHICmdList, FSceneViewFamily& InViewFamily) {}
-		virtual void PreRenderView_RenderThread(FRHICommandListImmediate& RHICmdList, FSceneView& InView) {}
+		virtual void PreRenderViewFamily_RenderThread(FRDGBuilder& GraphBuilder, FSceneViewFamily& InViewFamily) {}
+		virtual void PreRenderView_RenderThread(FRDGBuilder& GraphBuilder, FSceneView& InView) {}
 		virtual int32 GetPriority() const { return 0; }
-		virtual void PostRenderView_RenderThread(FRHICommandListImmediate& RHICmdList, FSceneView& InView)
+		virtual void PostRenderView_RenderThread(FRDGBuilder& GraphBuilder, FSceneView& InView)
 		{
 			FViewInfo& View = static_cast<FViewInfo&>(InView);
-			DisplayInternals(RHICmdList, View);
+			DisplayInternals(GraphBuilder, View);
 		}
 	};
 	TSharedRef<FRendererViewExtension, ESPMode::ThreadSafe> ref(new FRendererViewExtension);
diff --git a/Engine/Source/Runtime/Renderer/Private/SceneRendering.h b/Engine/Source/Runtime/Renderer/Private/SceneRendering.h
--- a/Engine/Source/Runtime/Renderer/Private/SceneRendering.h
+++ b/Engine/Source/Runtime/Renderer/Private/SceneRendering.h
@@ -2036,7 +2036,7 @@ public:
 	*
 	* @param SceneRenderer	Scene renderer to use for rendering.
 	*/
-	static void RENDERER_API ViewExtensionPreRender_RenderThread(FRHICommandListImmediate& RHICmdList, FSceneRenderer* SceneRenderer);
+	static void RENDERER_API ViewExtensionPreRender_RenderThread(FRDGBuilder& GraphBuilder, FSceneRenderer* SceneRenderer);
 
 	/** Called to release any deallocations that were deferred until the next render. */
 	static void CleanUp(FRHICommandListImmediate& RHICmdList);
