Implemented support for asynchronous creation and deletion of FRenderResource.
The resource list uses a thread local free list that is coalesced later.

diff --git a/Engine/Source/Runtime/Renderer/Private/SceneRendering.cpp b/Engine/Source/Runtime/Renderer/Private/SceneRendering.cpp
--- a/Engine/Source/Runtime/Renderer/Private/SceneRendering.cpp
+++ b/Engine/Source/Runtime/Renderer/Private/SceneRendering.cpp
@@ -4057,6 +4057,7 @@ static void DeleteSceneRenderers(const TArray<FSceneRenderer*>& SceneRenderers,
 	{
 		delete SceneRenderer;
 	}
+	FRenderResource::CoalesceResourceList();
 }
 
 static void WaitForTasksAndDeleteSceneRenderers(FRHICommandListImmediate& RHICmdList, const TArray<FSceneRenderer*>& SceneRenderers)
