Hierarchical LOD meshes are not frustum/occlusion culled

diff --git a/Engine/Source/Runtime/Renderer/Private/SceneVisibility.cpp b/Engine/Source/Runtime/Renderer/Private/SceneVisibility.cpp
--- a/Engine/Source/Runtime/Renderer/Private/SceneVisibility.cpp
+++ b/Engine/Source/Runtime/Renderer/Private/SceneVisibility.cpp
@@ -752,18 +752,16 @@ static void PrimitiveCullTask(FThreadSafeCounter& NumCulledPrimitives, const FSc
 				{
 					if (HLODState->IsNodeForcedVisible(Index))
 					{
-						bShouldFrustumCull = false;
 						bShouldDistanceCull = false;
 					}
 					else if (HLODState->IsNodeForcedHidden(Index))
 					{
-						bShouldDistanceCull = false;
-						bShouldFrustumCull = false;
 						bIsVisible = false;
 					}
 				}
 
 				// Frustum first
+				bShouldFrustumCull = bShouldFrustumCull && bIsVisible;
 				if (bShouldFrustumCull)
 				{
 					if (Flags.bUseVisibilityOctree)
@@ -790,9 +788,8 @@ static void PrimitiveCullTask(FThreadSafeCounter& NumCulledPrimitives, const FSc
 					}
 				}
 
-				bShouldDistanceCull = bShouldDistanceCull && bIsVisible;
-
 				// Distance cull if frustum cull passed
+				bShouldDistanceCull = bShouldDistanceCull && bIsVisible;
 				if (bShouldDistanceCull)
 				{
 					// If cull distance is disabled, always show the primitive (except foliage)
