Add common structure for passing hair strands instance data.

diff --git a/Engine/Shaders/Private/HairStrands/HairStrandsVertexFactoryCommon.ush b/Engine/Shaders/Private/HairStrands/HairStrandsVertexFactoryCommon.ush
--- a/Engine/Shaders/Private/HairStrands/HairStrandsVertexFactoryCommon.ush
+++ b/Engine/Shaders/Private/HairStrands/HairStrandsVertexFactoryCommon.ush
@@ -81,4 +81,32 @@ FHairAttributes DecodeHairAttribute(float2 Attributes0, uint Attributes1, float
 	Out.WorldLength = Attributes0.x * InVF_Length;
 	Out.Seed		= Attributes0.y;
 	return Out;
-}
\ No newline at end of file
+}
+
+///////////////////////////////////////////////////////////////////////////////////////////////////
+// Hair strands instance parameters
+
+#if defined(HAIR_STRANDS_PARAMETERS) && HAIR_STRANDS_PARAMETERS == 1
+uint			HairStrandsVF_bIsCullingEnable;
+uint			HairStrandsVF_bHasRaytracedGeometry;
+float			HairStrandsVF_Radius;
+float			HairStrandsVF_RootScale;
+float			HairStrandsVF_TipScale;
+float			HairStrandsVF_Length;
+float			HairStrandsVF_Density;
+uint			HairStrandsVF_bUseStableRasterization;
+float3			HairStrandsVF_PositionOffset;
+uint			HairStrandsVF_VertexCount;
+float4x4		HairStrandsVF_LocalToWorldPrimitiveTransform;
+float4x4		HairStrandsVF_LocalToTranslatedWorldPrimitiveTransform;
+Buffer<uint4>	HairStrandsVF_PositionBuffer;
+Buffer<float4>	HairStrandsVF_PositionOffsetBuffer;
+
+#if defined(PERMUTATION_CULLING) && PERMUTATION_CULLING == 1
+Buffer<uint>	HairStrandsVF_CullingIndirectBuffer;
+Buffer<uint>	HairStrandsVF_CullingIndexBuffer;
+Buffer<float>	HairStrandsVF_CullingRadiusScaleBuffer;
+Buffer<uint>	HairStrandsVF_CullingIndirectBufferArgs;
+#endif
+
+#endif // HAIR_STRANDS_PARAMETERS
\ No newline at end of file
diff --git a/Engine/Shaders/Private/HairStrands/HairStrandsVisibilityRasterCompute.usf b/Engine/Shaders/Private/HairStrands/HairStrandsVisibilityRasterCompute.usf
--- a/Engine/Shaders/Private/HairStrands/HairStrandsVisibilityRasterCompute.usf
+++ b/Engine/Shaders/Private/HairStrands/HairStrandsVisibilityRasterCompute.usf
@@ -1,5 +1,7 @@
 // Copyright Epic Games, Inc. All Rights Reserved.
 
+#define HAIR_STRANDS_PARAMETERS 1
+
 #include "../Common.ush"
 #include "HairStrandsClusterCommon.ush"
 #include "HairStrandsVertexFactoryCommon.ush"
@@ -77,18 +79,7 @@ uint2 UppackVisTileCoord(uint Packed)
 
 uint			MacroGroupId;
 uint			HairMaterialId;
-uint			HairStrandsVF_bIsCullingEnable;
-float			HairStrandsVF_Radius;
-float			HairStrandsVF_RootScale;
-float			HairStrandsVF_TipScale;
-float			HairStrandsVF_Length;
-float			HairStrandsVF_Density;
-uint			HairStrandsVF_bUseStableRasterization;
-float3			HairStrandsVF_PositionOffset;
-uint			HairStrandsVF_VertexCount;
-float4x4		HairStrandsVF_LocalToWorldPrimitiveTransform;
-Buffer<uint4>	HairStrandsVF_PositionBuffer;
-Buffer<float4>	HairStrandsVF_PositionOffsetBuffer;
+
 
 Texture2D<float> SceneDepthTexture;
 
@@ -195,13 +186,6 @@ void PrepareDepthGridCS(uint dId : SV_DispatchThreadID, uint GroupThreadID : SV_
 
 #if SHADER_RASTERCOMPUTE_BINNING
 
-#if PERMUTATION_CULLING
-Buffer<uint>	HairStrandsVF_CullingIndirectBuffer;
-Buffer<uint>	HairStrandsVF_CullingIndexBuffer;
-Buffer<float>	HairStrandsVF_CullingRadiusScaleBuffer;
-Buffer<uint>	IndirectBufferArgs;
-#endif
-
 RWTexture3D<uint> 					OutVisTileBinningGrid;
 Texture2D<uint>						VisTileDepthGrid;
 RWBuffer<uint>						OutVisTilePrims;
diff --git a/Engine/Shaders/Private/HairStrands/HairStrandsVoxelRasterCompute.usf b/Engine/Shaders/Private/HairStrands/HairStrandsVoxelRasterCompute.usf
--- a/Engine/Shaders/Private/HairStrands/HairStrandsVoxelRasterCompute.usf
+++ b/Engine/Shaders/Private/HairStrands/HairStrandsVoxelRasterCompute.usf
@@ -1,7 +1,8 @@
 // Copyright Epic Games, Inc. All Rights Reserved.
 
-#include "../Common.ush"
+#define HAIR_STRANDS_PARAMETERS 1
 
+#include "../Common.ush"
 #include "HairStrandsClusterCommon.ush"
 #include "HairStrandsVertexFactoryCommon.ush"
 #include "HairStrandsVoxelPageCommon.ush" 
@@ -12,24 +13,6 @@ uint			DispatchCountX;
 uint			MaxRasterCount;
 uint			FrameIdMod8;
 uint			MacroGroupId;
-uint			HairStrandsVF_bIsCullingEnable;
-uint			HairStrandsVF_bHasRaytracedGeometry;
-float			HairStrandsVF_Radius;
-float			HairStrandsVF_RootScale;
-float			HairStrandsVF_TipScale;
-float			HairStrandsVF_Density;
-float3			HairStrandsVF_PositionOffset;
-uint			HairStrandsVF_VertexCount;
-float4x4		HairStrandsVF_LocalToTranslatedWorldPrimitiveTransform;
-Buffer<uint4>	HairStrandsVF_PositionBuffer;
-Buffer<float4>	HairStrandsVF_PositionOffsetBuffer;
-
-#if PERMUTATION_CULLING == 1
-Buffer<uint>	HairStrandsVF_CullingIndirectBuffer;
-Buffer<uint>	HairStrandsVF_CullingIndexBuffer;
-Buffer<float>	HairStrandsVF_CullingRadiusScaleBuffer;
-Buffer<uint>	IndirectBufferArgs;
-#endif
 
 uint			VirtualVoxelParams_PageIndexCount;
 uint			VirtualVoxelParams_PageResolution;
diff --git a/Engine/Source/Runtime/Renderer/Private/HairStrands/HairStrandsData.h b/Engine/Source/Runtime/Renderer/Private/HairStrands/HairStrandsData.h
--- a/Engine/Source/Runtime/Renderer/Private/HairStrands/HairStrandsData.h
+++ b/Engine/Source/Runtime/Renderer/Private/HairStrands/HairStrandsData.h
@@ -16,6 +16,7 @@ class FPrimitiveSceneProxy;
 class FViewInfo;
 class FScene;
 class FInstanceCullingManager;
+class FHairGroupPublicData;
 struct FMeshBatch;
 struct FMeshBatchAndRelevance;
 
@@ -42,6 +43,32 @@ BEGIN_GLOBAL_SHADER_PARAMETER_STRUCT(FHairStrandsViewUniformParameters, )
 	SHADER_PARAMETER_RDG_BUFFER_SRV(Buffer<uint>,  HairTileCount)						// Tile total count (actual number of tiles)
 END_GLOBAL_SHADER_PARAMETER_STRUCT()
 
+////////////////////////////////////////////////////////////////////////////////////
+// Hair Instance data
+
+BEGIN_SHADER_PARAMETER_STRUCT(FHairStrandsInstanceParameters, )
+	SHADER_PARAMETER(uint32, HairStrandsVF_bIsCullingEnable)
+	SHADER_PARAMETER(uint32, HairStrandsVF_bHasRaytracedGeometry)
+	SHADER_PARAMETER(float, HairStrandsVF_Density)
+	SHADER_PARAMETER(float, HairStrandsVF_Radius)
+	SHADER_PARAMETER(float, HairStrandsVF_RootScale)
+	SHADER_PARAMETER(float, HairStrandsVF_TipScale)
+	SHADER_PARAMETER(float, HairStrandsVF_Length)
+	SHADER_PARAMETER(uint32, HairStrandsVF_bUseStableRasterization)
+	SHADER_PARAMETER(uint32, HairStrandsVF_VertexCount)
+	SHADER_PARAMETER(FVector3f, HairStrandsVF_PositionOffset)
+	SHADER_PARAMETER(FMatrix44f, HairStrandsVF_LocalToWorldPrimitiveTransform)
+	SHADER_PARAMETER(FMatrix44f, HairStrandsVF_LocalToTranslatedWorldPrimitiveTransform)
+	SHADER_PARAMETER_RDG_BUFFER_SRV(Buffer, HairStrandsVF_PositionBuffer)
+	SHADER_PARAMETER_RDG_BUFFER_SRV(Buffer, HairStrandsVF_PositionOffsetBuffer)
+	SHADER_PARAMETER_RDG_BUFFER_SRV(Buffer, HairStrandsVF_CullingIndirectBuffer)
+	SHADER_PARAMETER_RDG_BUFFER_SRV(Buffer, HairStrandsVF_CullingIndexBuffer)
+	SHADER_PARAMETER_RDG_BUFFER_SRV(Buffer, HairStrandsVF_CullingRadiusScaleBuffer)
+	RDG_BUFFER_ACCESS(HairStrandsVF_CullingIndirectBufferArgs, ERHIAccess::IndirectArgs)
+END_SHADER_PARAMETER_STRUCT()
+
+FHairStrandsInstanceParameters GetHairStrandsInstanceParameters(FRDGBuilder& GraphBuilder, const FViewInfo& ViewInfo, const FHairGroupPublicData* HairGroupPublicData, bool bCullingEnable, bool bForceRegister);
+
 ////////////////////////////////////////////////////////////////////////////////////
 // Tile data
 
diff --git a/Engine/Source/Runtime/Renderer/Private/HairStrands/HairStrandsVisibility.cpp b/Engine/Source/Runtime/Renderer/Private/HairStrands/HairStrandsVisibility.cpp
--- a/Engine/Source/Runtime/Renderer/Private/HairStrands/HairStrandsVisibility.cpp
+++ b/Engine/Source/Runtime/Renderer/Private/HairStrands/HairStrandsVisibility.cpp
@@ -2907,6 +2907,61 @@ IMPLEMENT_GLOBAL_SHADER(FVisiblityRasterComputePrepareDepthGridCS, "/Engine/Priv
 
 /////////////////////////////////////////////////////////////////////////////////////////
 
+FHairStrandsInstanceParameters GetHairStrandsInstanceParameters(FRDGBuilder& GraphBuilder, const FViewInfo& ViewInfo, const FHairGroupPublicData* HairGroupPublicData, bool bCullingEnable, bool bForceRegister)
+{
+	check(HairGroupPublicData);
+
+	const FHairGroupPublicData::FVertexFactoryInput& VFInput = HairGroupPublicData->VFInput;
+
+	FHairStrandsInstanceParameters Out;
+	if (bForceRegister)
+	{
+		Out.HairStrandsVF_PositionBuffer		= Register(GraphBuilder, VFInput.Strands.PositionBufferExternal, ERDGImportedBufferFlags::CreateSRV).SRV;;
+		Out.HairStrandsVF_PositionOffsetBuffer	= Register(GraphBuilder, VFInput.Strands.PositionOffsetBufferExternal, ERDGImportedBufferFlags::CreateSRV).SRV;;
+
+	}
+	else
+	{
+		Out.HairStrandsVF_PositionBuffer		= VFInput.Strands.PositionBuffer.SRV;
+		Out.HairStrandsVF_PositionOffsetBuffer	= VFInput.Strands.PositionOffsetBuffer.SRV;
+	}
+
+
+	Out.HairStrandsVF_VertexCount						= VFInput.Strands.VertexCount;
+	Out.HairStrandsVF_Radius							= VFInput.Strands.HairRadius;
+	Out.HairStrandsVF_RootScale							= VFInput.Strands.HairRootScale;
+	Out.HairStrandsVF_TipScale							= VFInput.Strands.HairTipScale;
+	Out.HairStrandsVF_Length							= VFInput.Strands.HairLength;
+	Out.HairStrandsVF_bUseStableRasterization			= VFInput.Strands.bUseStableRasterization ? 1 : 0;
+	Out.HairStrandsVF_Density							= VFInput.Strands.HairDensity;
+	Out.HairStrandsVF_bIsCullingEnable					= 0;
+	Out.HairStrandsVF_bHasRaytracedGeometry				= VFInput.Strands.bUseRaytracingGeometry ? 1u : 0u;
+	Out.HairStrandsVF_PositionOffset					= (FVector3f)VFInput.Strands.PositionOffset;
+
+	// Absolute local to world
+	Out.HairStrandsVF_LocalToWorldPrimitiveTransform	= FMatrix44f(VFInput.LocalToWorldTransform.ToMatrixWithScale()); // LWC_TODO: Precision loss
+
+	// Translated local to world
+	const FVector& TranslatedWorldOffset = ViewInfo.ViewMatrices.GetPreViewTranslation();
+	FTransform LocalToTranslatedWorldTransform = VFInput.LocalToWorldTransform;
+	LocalToTranslatedWorldTransform.AddToTranslation(TranslatedWorldOffset);
+	Out.HairStrandsVF_LocalToTranslatedWorldPrimitiveTransform = FMatrix44f(LocalToTranslatedWorldTransform.ToMatrixWithScale());
+
+	//Out.VertexCount = HairGroupPublicData->GetActiveStrandsVertexStart(VFInput.Strands.VertexCount);
+	//Out.VertexStart = HairGroupPublicData->GetActiveStrandsVertexCount(VFInput.Strands.VertexCount, ScreenSize);
+
+	if (bCullingEnable)
+	{
+		FRDGImportedBuffer CullingIndirectBuffer = Register(GraphBuilder, HairGroupPublicData->GetDrawIndirectRasterComputeBuffer(), ERDGImportedBufferFlags::CreateSRV);
+		Out.HairStrandsVF_CullingIndirectBuffer = CullingIndirectBuffer.SRV;
+		Out.HairStrandsVF_bIsCullingEnable = 1;
+		Out.HairStrandsVF_CullingIndexBuffer = RegisterAsSRV(GraphBuilder, HairGroupPublicData->GetCulledVertexIdBuffer());
+		Out.HairStrandsVF_CullingRadiusScaleBuffer = RegisterAsSRV(GraphBuilder, HairGroupPublicData->GetCulledVertexRadiusScaleBuffer());
+		Out.HairStrandsVF_CullingIndirectBufferArgs = CullingIndirectBuffer.Buffer;
+	}
+	return Out;
+}
+
 class FVisiblityRasterComputeBinningCS : public FGlobalShader
 {
 	DECLARE_GLOBAL_SHADER(FVisiblityRasterComputeBinningCS);
@@ -2921,21 +2976,7 @@ class FVisiblityRasterComputeBinningCS : public FGlobalShader
 		SHADER_PARAMETER(uint32, MacroGroupId)
 		SHADER_PARAMETER(uint32, DispatchCountX)
 		SHADER_PARAMETER(uint32, HairMaterialId)
-		SHADER_PARAMETER(uint32, HairStrandsVF_bIsCullingEnable)
-		SHADER_PARAMETER(float, HairStrandsVF_Density)
-		SHADER_PARAMETER(float, HairStrandsVF_Radius)
-		SHADER_PARAMETER(float, HairStrandsVF_RootScale)
-		SHADER_PARAMETER(float, HairStrandsVF_TipScale)
-		SHADER_PARAMETER(float, HairStrandsVF_Length)
-		SHADER_PARAMETER(uint32, HairStrandsVF_bUseStableRasterization)
-		SHADER_PARAMETER(uint32, HairStrandsVF_VertexCount)
-		SHADER_PARAMETER(FMatrix44f, HairStrandsVF_LocalToWorldPrimitiveTransform)
-		SHADER_PARAMETER_RDG_BUFFER_SRV(Buffer, HairStrandsVF_PositionBuffer)
-		SHADER_PARAMETER_RDG_BUFFER_SRV(Buffer, HairStrandsVF_PositionOffsetBuffer)
-		SHADER_PARAMETER_RDG_BUFFER_SRV(Buffer, HairStrandsVF_CullingIndirectBuffer)
-		SHADER_PARAMETER_RDG_BUFFER_SRV(Buffer, HairStrandsVF_CullingIndexBuffer)
-		SHADER_PARAMETER_RDG_BUFFER_SRV(Buffer, HairStrandsVF_CullingRadiusScaleBuffer)
-		RDG_BUFFER_ACCESS(IndirectBufferArgs, ERHIAccess::IndirectArgs)
+		SHADER_PARAMETER_STRUCT_INCLUDE(FHairStrandsInstanceParameters, HairInstance)
 		SHADER_PARAMETER_RDG_TEXTURE(Texture2D, SceneDepthTexture)
 		SHADER_PARAMETER_RDG_TEXTURE(Texture3D, VisTileBinningGrid)
 		SHADER_PARAMETER_RDG_BUFFER_UAV(RWBuffer, OutVisTilePrims)
@@ -2975,20 +3016,7 @@ class FVisiblityRasterComputeRasterizeCS : public FGlobalShader
 		SHADER_PARAMETER(uint32, MacroGroupId)
 		SHADER_PARAMETER(uint32, DispatchCountX)
 		SHADER_PARAMETER(uint32, HairMaterialId)
-		SHADER_PARAMETER(uint32, HairStrandsVF_bIsCullingEnable)
-		SHADER_PARAMETER(float, HairStrandsVF_Density)
-		SHADER_PARAMETER(float, HairStrandsVF_Radius)
-		SHADER_PARAMETER(float, HairStrandsVF_RootScale)
-		SHADER_PARAMETER(float, HairStrandsVF_TipScale)
-		SHADER_PARAMETER(float, HairStrandsVF_Length)
-		SHADER_PARAMETER(uint32, HairStrandsVF_bUseStableRasterization)
-		SHADER_PARAMETER(uint32, HairStrandsVF_VertexCount)
-		SHADER_PARAMETER(FMatrix44f, HairStrandsVF_LocalToWorldPrimitiveTransform)
-		SHADER_PARAMETER_RDG_BUFFER_SRV(Buffer, HairStrandsVF_PositionBuffer)
-		SHADER_PARAMETER_RDG_BUFFER_SRV(Buffer, HairStrandsVF_PositionOffsetBuffer)
-		SHADER_PARAMETER_RDG_BUFFER_SRV(Buffer, HairStrandsVF_CullingIndirectBuffer)
-		SHADER_PARAMETER_RDG_BUFFER_SRV(Buffer, HairStrandsVF_CullingIndexBuffer)
-		SHADER_PARAMETER_RDG_BUFFER_SRV(Buffer, HairStrandsVF_CullingRadiusScaleBuffer)
+		SHADER_PARAMETER_STRUCT_INCLUDE(FHairStrandsInstanceParameters, HairInstance)
 		RDG_BUFFER_ACCESS(IndirectBufferArgs, ERHIAccess::IndirectArgs)
 		SHADER_PARAMETER_RDG_TEXTURE(Texture2D, SceneDepthTexture)
 		SHADER_PARAMETER(float, SampleWeight)
@@ -3242,43 +3270,13 @@ static FRasterComputeOutput AddVisibilityComputeRasterPass(
 				BinningParameters->VisTileDepthGrid = OutVisTileDepthGrid;
 				BinningParameters->OutVisTileArgs = OutVisTileArgsUAV;
 				BinningParameters->OutVisTileData = OutVisTileDataUAV;
-
-				if (bForceRegister)
-				{
-					BinningParameters->HairStrandsVF_PositionBuffer = Register(GraphBuilder, VFInput.Strands.PositionBufferExternal, ERDGImportedBufferFlags::CreateSRV).SRV;;
-					BinningParameters->HairStrandsVF_PositionOffsetBuffer = Register(GraphBuilder, VFInput.Strands.PositionOffsetBufferExternal, ERDGImportedBufferFlags::CreateSRV).SRV;;
-
-				}
-				else
-				{
-					BinningParameters->HairStrandsVF_PositionBuffer = VFInput.Strands.PositionBuffer.SRV;
-					BinningParameters->HairStrandsVF_PositionOffsetBuffer = VFInput.Strands.PositionOffsetBuffer.SRV;
-				}
-
-				BinningParameters->HairStrandsVF_VertexCount = VFInput.Strands.VertexCount;
-				BinningParameters->HairStrandsVF_Radius = VFInput.Strands.HairRadius;
-				BinningParameters->HairStrandsVF_RootScale = VFInput.Strands.HairRootScale;
-				BinningParameters->HairStrandsVF_TipScale = VFInput.Strands.HairTipScale;
-				BinningParameters->HairStrandsVF_Length = VFInput.Strands.HairLength;
-				BinningParameters->HairStrandsVF_bUseStableRasterization = VFInput.Strands.bUseStableRasterization ? 1 : 0;
-				BinningParameters->HairStrandsVF_Density = VFInput.Strands.HairDensity;
-				BinningParameters->HairStrandsVF_LocalToWorldPrimitiveTransform = FMatrix44f(LocalToWorldPrimitiveTransform); // LWC_TODO: Precision loss
-
+				BinningParameters->HairInstance = GetHairStrandsInstanceParameters(GraphBuilder, ViewInfo, HairGroupPublicData, bCullingEnable, bForceRegister);
+				
 				BinningParameters->VertexCount = VertexCount;
 				BinningParameters->VertexStart = VertexStart;
 				
 				BinningParameters->VisTileBinningGridTex = OutVisTileBinningGrid;
 
-				if (bCullingEnable)
-				{
-					FRDGImportedBuffer CullingIndirectBuffer = Register(GraphBuilder, HairGroupPublicData->GetDrawIndirectRasterComputeBuffer(), ERDGImportedBufferFlags::CreateSRV);
-					BinningParameters->HairStrandsVF_CullingIndirectBuffer = CullingIndirectBuffer.SRV;
-					BinningParameters->HairStrandsVF_bIsCullingEnable = bCullingEnable ? 1 : 0;
-					BinningParameters->HairStrandsVF_CullingIndexBuffer = RegisterAsSRV(GraphBuilder, HairGroupPublicData->GetCulledVertexIdBuffer());
-					BinningParameters->HairStrandsVF_CullingRadiusScaleBuffer = RegisterAsSRV(GraphBuilder, HairGroupPublicData->GetCulledVertexRadiusScaleBuffer());
-					BinningParameters->IndirectBufferArgs = CullingIndirectBuffer.Buffer;
-				}
-
 				FComputeShaderUtils::AddPass(GraphBuilder, RDG_EVENT_NAME("HairStrands::VisibilityComputeRasterBinning(culling=%s)", bCullingEnable ? TEXT("On") : TEXT("Off")), bCullingEnable ? ComputeShaderBinning_CullingOn : ComputeShaderBinning_CullingOff, BinningParameters, FIntVector(NumBinners, 1, 1));
 			}
 			
@@ -3294,30 +3292,7 @@ static FRasterComputeOutput AddVisibilityComputeRasterPass(
 				RasterParameters->OutDepthCovTexture = DepthCovTextureUAV;
 				RasterParameters->OutPrimMatTexture = PrimMatTextureUAV;
 				RasterParameters->SampleWeight = SampleWeight;
-
-				if (bForceRegister)
-				{
-					RasterParameters->HairStrandsVF_PositionBuffer = Register(GraphBuilder, VFInput.Strands.PositionBufferExternal, ERDGImportedBufferFlags::CreateSRV).SRV;;
-					RasterParameters->HairStrandsVF_PositionOffsetBuffer = Register(GraphBuilder, VFInput.Strands.PositionOffsetBufferExternal, ERDGImportedBufferFlags::CreateSRV).SRV;;
-
-				}
-				else
-				{
-					RasterParameters->HairStrandsVF_PositionBuffer = VFInput.Strands.PositionBuffer.SRV;
-					RasterParameters->HairStrandsVF_PositionOffsetBuffer = VFInput.Strands.PositionOffsetBuffer.SRV;
-				}
-
-//				RasterParameters->HairStrandsVF_PositionBuffer	= VFInput.Strands.PositionBuffer.SRV;
-//				RasterParameters->HairStrandsVF_PositionOffsetBuffer = VFInput.Strands.PositionOffsetBuffer.SRV;
-				RasterParameters->HairStrandsVF_VertexCount		= VFInput.Strands.VertexCount;
-				RasterParameters->HairStrandsVF_Radius			= VFInput.Strands.HairRadius;
-				RasterParameters->HairStrandsVF_RootScale			= VFInput.Strands.HairRootScale;
-				RasterParameters->HairStrandsVF_TipScale			= VFInput.Strands.HairTipScale;
-				RasterParameters->HairStrandsVF_Length			= VFInput.Strands.HairLength;
-				RasterParameters->HairStrandsVF_bUseStableRasterization = VFInput.Strands.bUseStableRasterization ? 1 : 0;
-				RasterParameters->HairStrandsVF_Density			= VFInput.Strands.HairDensity;
-				RasterParameters->HairStrandsVF_LocalToWorldPrimitiveTransform = FMatrix44f(LocalToWorldPrimitiveTransform); // LWC_TODO: Precision loss
-
+				RasterParameters->HairInstance = GetHairStrandsInstanceParameters(GraphBuilder, ViewInfo, HairGroupPublicData, bCullingEnable, bForceRegister);
 				RasterParameters->VisTilePrims = GraphBuilder.CreateSRV(OutVisTilePrims);
 				RasterParameters->RWVisTileDepthGrid = OutVisTileDepthGridUAV;
 				RasterParameters->VisTileArgs = GraphBuilder.CreateSRV(OutVisTileArgs);
diff --git a/Engine/Source/Runtime/Renderer/Private/HairStrands/HairStrandsVoxelization.cpp b/Engine/Source/Runtime/Renderer/Private/HairStrands/HairStrandsVoxelization.cpp
--- a/Engine/Source/Runtime/Renderer/Private/HairStrands/HairStrandsVoxelization.cpp
+++ b/Engine/Source/Runtime/Renderer/Private/HairStrands/HairStrandsVoxelization.cpp
@@ -1232,22 +1232,8 @@ class FVoxelRasterComputeCS : public FGlobalShader
 		SHADER_PARAMETER(uint32, DispatchCountX)
 		SHADER_PARAMETER(uint32, MaxRasterCount)
 		SHADER_PARAMETER(uint32, FrameIdMod8)
-		SHADER_PARAMETER(uint32, HairStrandsVF_bIsCullingEnable)
-		SHADER_PARAMETER(uint32, HairStrandsVF_bHasRaytracedGeometry)
-		SHADER_PARAMETER(float,	  HairStrandsVF_Density)
-		SHADER_PARAMETER(float,   HairStrandsVF_Radius)
-		SHADER_PARAMETER(float,   HairStrandsVF_RootScale)
-		SHADER_PARAMETER(float,   HairStrandsVF_TipScale)
-		SHADER_PARAMETER(FVector3f, HairStrandsVF_PositionOffset)
-		SHADER_PARAMETER(uint32,  HairStrandsVF_VertexCount)
-		SHADER_PARAMETER(FMatrix44f, HairStrandsVF_LocalToTranslatedWorldPrimitiveTransform)
-		SHADER_PARAMETER_RDG_BUFFER_SRV(Buffer, HairStrandsVF_PositionBuffer)
-		SHADER_PARAMETER_RDG_BUFFER_SRV(Buffer, HairStrandsVF_PositionOffsetBuffer)
-		SHADER_PARAMETER_RDG_BUFFER_SRV(Buffer, HairStrandsVF_CullingIndirectBuffer)
-		SHADER_PARAMETER_RDG_BUFFER_SRV(Buffer, HairStrandsVF_CullingIndexBuffer)
-		SHADER_PARAMETER_RDG_BUFFER_SRV(Buffer, HairStrandsVF_CullingRadiusScaleBuffer)
+		SHADER_PARAMETER_STRUCT_INCLUDE(FHairStrandsInstanceParameters, HairInstance)
 		SHADER_PARAMETER_RDG_BUFFER_SRV(Buffer, VoxelizationViewInfoBuffer)
-		RDG_BUFFER_ACCESS(IndirectBufferArgs, ERHIAccess::IndirectArgs)
 		SHADER_PARAMETER_RDG_TEXTURE_UAV(RWTexture3D, OutPageTexture)
 		SHADER_PARAMETER(uint32, VertexCount)
 		SHADER_PARAMETER(uint32, VertexStart)
@@ -1328,34 +1314,17 @@ static void AddVirtualVoxelizationComputeRasterPass(
 				PassParameters->OutPageTexture = PageTextureUAV;
 				PassParameters->FrameIdMod8 = FrameIdMode8;
 
-				PassParameters->HairStrandsVF_PositionBuffer = VFInput.Strands.PositionBuffer.SRV;
-				PassParameters->HairStrandsVF_PositionOffset = (FVector3f)VFInput.Strands.PositionOffset;
-				PassParameters->HairStrandsVF_PositionOffsetBuffer = VFInput.Strands.PositionOffsetBuffer.SRV;
-				PassParameters->HairStrandsVF_VertexCount = VFInput.Strands.VertexCount;
-				PassParameters->HairStrandsVF_Radius = VFInput.Strands.HairRadius;
-				PassParameters->HairStrandsVF_RootScale = VFInput.Strands.HairRootScale;
-				PassParameters->HairStrandsVF_TipScale = VFInput.Strands.HairTipScale;
-				PassParameters->HairStrandsVF_Density = VFInput.Strands.HairDensity;
-				PassParameters->HairStrandsVF_LocalToTranslatedWorldPrimitiveTransform = FMatrix44f(LocalToTranslatedWorldTransform.ToMatrixWithScale());
-				PassParameters->HairStrandsVF_bHasRaytracedGeometry = VFInput.Strands.bUseRaytracingGeometry ? 1u : 0u;
-
 				PassParameters->VertexCount = VertexCount;
 				PassParameters->VertexStart = VertexStart;
 				PassParameters->SampleWeight = SampleWeight;
 			
 				const bool bCullingEnable = !IsHairStrandContinuousDecimationReorderingEnabled() && HairGroupPublicData->GetCullingResultAvailable();
 
-				PassParameters->HairStrandsVF_bIsCullingEnable = bCullingEnable ? 1 : 0;
+				PassParameters->HairInstance = GetHairStrandsInstanceParameters(GraphBuilder, *ViewInfo, HairGroupPublicData, bCullingEnable, false);
 
 				if (bCullingEnable)
 				{
-					FRDGImportedBuffer CullingIndirectBuffer = Register(GraphBuilder, HairGroupPublicData->GetDrawIndirectRasterComputeBuffer(), ERDGImportedBufferFlags::CreateSRV);
-					PassParameters->HairStrandsVF_CullingIndirectBuffer		= CullingIndirectBuffer.SRV;
-					PassParameters->HairStrandsVF_CullingIndexBuffer		= RegisterAsSRV(GraphBuilder, HairGroupPublicData->GetCulledVertexIdBuffer());
-					PassParameters->HairStrandsVF_CullingRadiusScaleBuffer	= RegisterAsSRV(GraphBuilder, HairGroupPublicData->GetCulledVertexRadiusScaleBuffer());
-					PassParameters->IndirectBufferArgs = CullingIndirectBuffer.Buffer;
-
-					FComputeShaderUtils::AddPass(GraphBuilder, RDG_EVENT_NAME("HairStrands::VoxelComputeRaster(culling=on)"), ComputeShader_CullingOn, PassParameters, CullingIndirectBuffer.Buffer, 0);
+					FComputeShaderUtils::AddPass(GraphBuilder, RDG_EVENT_NAME("HairStrands::VoxelComputeRaster(culling=on)"), ComputeShader_CullingOn, PassParameters, PassParameters->HairInstance.HairStrandsVF_CullingIndirectBufferArgs, 0);
 				}
 				else
 				{
