Fixed bad merge causing crash on startup.

diff --git a/Engine/Source/Runtime/Renderer/Private/SceneRendering.cpp b/Engine/Source/Runtime/Renderer/Private/SceneRendering.cpp
--- a/Engine/Source/Runtime/Renderer/Private/SceneRendering.cpp
+++ b/Engine/Source/Runtime/Renderer/Private/SceneRendering.cpp
@@ -4057,7 +4057,6 @@ static void DeleteSceneRenderers(const TArray<FSceneRenderer*>& SceneRenderers,
 	{
 		delete SceneRenderer;
 	}
-	FRenderResource::CoalesceResourceList();
 }
 
 static void WaitForTasksAndDeleteSceneRenderers(FRHICommandListImmediate& RHICmdList, const TArray<FSceneRenderer*>& SceneRenderers)
