Fixed banned word warning.

#rnx

diff --git a/Engine/Source/Runtime/Renderer/Private/SceneVisibility.cpp b/Engine/Source/Runtime/Renderer/Private/SceneVisibility.cpp
--- a/Engine/Source/Runtime/Renderer/Private/SceneVisibility.cpp
+++ b/Engine/Source/Runtime/Renderer/Private/SceneVisibility.cpp
@@ -1783,7 +1783,6 @@ static int32 FetchVisibilityForPrimitives(const FScene* Scene, FViewInfo& View,
 				QUICK_SCOPE_CYCLE_COUNTER(STAT_FetchVisibilityForPrimitivesSanity);
 				TSet<int32> ReleaseQuerySet;
 				TSet<int32> RunQuerySet;
-				TSet<int32> MasterPrimsProcessed;
 				for (int32 i = 0; i < NumTasks; ++i)
 				{
 					bool bAlreadyIn = false;
