Fixed trying to call a shader without editor.

diff --git a/Engine/Source/Runtime/Renderer/Private/DebugProbeRendering.cpp b/Engine/Source/Runtime/Renderer/Private/DebugProbeRendering.cpp
--- a/Engine/Source/Runtime/Renderer/Private/DebugProbeRendering.cpp
+++ b/Engine/Source/Runtime/Renderer/Private/DebugProbeRendering.cpp
@@ -8,7 +8,8 @@
 #include "SceneTextureParameters.h"
 #include "Strata/Strata.h"
 
-#define DEBUG_PROBE_ENABLED (!UE_BUILD_SHIPPING)
+// We do not want such debug in shipping build or when the editor is not available.
+#define DEBUG_PROBE_ENABLED (!UE_BUILD_SHIPPING && WITH_EDITORONLY_DATA)
 
 //
 //	Deferred probes are only stamped in deferred mode.
