Fixed: Water textures fail to render on mobile

diff --git a/Engine/Source/Runtime/Renderer/Private/MobileBasePass.cpp b/Engine/Source/Runtime/Renderer/Private/MobileBasePass.cpp
--- a/Engine/Source/Runtime/Renderer/Private/MobileBasePass.cpp
+++ b/Engine/Source/Runtime/Renderer/Private/MobileBasePass.cpp
@@ -36,7 +36,9 @@ bool MobileUsesNoLightMapPermutation(const FMeshMaterialShaderPermutationParamet
 	const bool bIsLitMaterial = Parameters.MaterialParameters.ShadingModels.IsLit();
 	const bool bDeferredShading = IsMobileDeferredShadingEnabled(Parameters.Platform);
 
-	if (!bDeferredShading && !bAllowStaticLighting && bIsLitMaterial && !IsTranslucentBlendMode(Parameters.MaterialParameters.BlendMode))
+	if (!bDeferredShading && !bAllowStaticLighting && bIsLitMaterial && 
+		!IsTranslucentBlendMode(Parameters.MaterialParameters.BlendMode) && 
+		!Parameters.MaterialParameters.ShadingModels.HasShadingModel(MSM_SingleLayerWater))
 	{
 		// We don't need NoLightMap permutation if CSM shader can handle no-CSM case with a branch inside shader
 		return !MobileUseCSMShaderBranch();
