Appease PVS Studio.

diff --git a/Engine/Source/Runtime/Renderer/Private/RayTracing/RayTracingMaterialHitShaders.cpp b/Engine/Source/Runtime/Renderer/Private/RayTracing/RayTracingMaterialHitShaders.cpp
--- a/Engine/Source/Runtime/Renderer/Private/RayTracing/RayTracingMaterialHitShaders.cpp
+++ b/Engine/Source/Runtime/Renderer/Private/RayTracing/RayTracingMaterialHitShaders.cpp
@@ -671,6 +671,7 @@ FRayTracingPipelineState* FDeferredShadingSceneRenderer::CreateRayTracingMateria
 
 	ERayTracingPipelineCacheFlags PipelineCacheFlags = ERayTracingPipelineCacheFlags::Default;
 	if (FallbackPipelineState
+		&& View.ViewState
 		&& IsCompatibleFallbackPipelineSignature(View.ViewState->LastRayTracingMaterialPipelineSignature, Initializer)
 		&& PipelineContainsHitShaders(FallbackPipelineState, RequiredHitShaders)
 		&& FindRayTracingMissShaderIndex(FallbackPipelineState, DefaultMissShader, false) != INDEX_NONE
