VSM multi-GPU:  Invalidate VSM cache when GPU mask changes, to prevent uninitialized data from being used.

#rnx

diff --git a/Engine/Source/Runtime/Renderer/Private/VirtualShadowMaps/VirtualShadowMapArray.cpp b/Engine/Source/Runtime/Renderer/Private/VirtualShadowMaps/VirtualShadowMapArray.cpp
--- a/Engine/Source/Runtime/Renderer/Private/VirtualShadowMaps/VirtualShadowMapArray.cpp
+++ b/Engine/Source/Runtime/Renderer/Private/VirtualShadowMaps/VirtualShadowMapArray.cpp
@@ -352,6 +352,10 @@ void FVirtualShadowMapArray::Initialize(FRDGBuilder& GraphBuilder, FVirtualShado
 	bEnabled = bInEnabled;
 	CacheManager = InCacheManager;
 
+#if WITH_MGPU
+	CacheManager->UpdateGPUMask(GraphBuilder.RHICmdList.GetGPUMask());
+#endif
+
 	bCullBackfacingPixels = CVarCullBackfacingPixels.GetValueOnRenderThread() != 0;
 	bUseHzbOcclusion = CVarShadowsVirtualUseHZB.GetValueOnRenderThread() != 0;
 	bUseTwoPassHzbOcclusion = CVarShadowsVirtualUseHZB.GetValueOnRenderThread() == 2;
diff --git a/Engine/Source/Runtime/Renderer/Private/VirtualShadowMaps/VirtualShadowMapCacheManager.cpp b/Engine/Source/Runtime/Renderer/Private/VirtualShadowMaps/VirtualShadowMapCacheManager.cpp
--- a/Engine/Source/Runtime/Renderer/Private/VirtualShadowMaps/VirtualShadowMapCacheManager.cpp
+++ b/Engine/Source/Runtime/Renderer/Private/VirtualShadowMaps/VirtualShadowMapCacheManager.cpp
@@ -892,6 +892,17 @@ void FVirtualShadowMapArrayCacheManager::SetHZBViewParams(int32 HZBKey, Nanite::
 	}
 }
 
+#if WITH_MGPU
+void FVirtualShadowMapArrayCacheManager::UpdateGPUMask(FRHIGPUMask GPUMask)
+{
+	if (LastGPUMask != GPUMask)
+	{
+		LastGPUMask = GPUMask;
+		Invalidate();
+	}
+}
+#endif  // WITH_MGPU
+
 static void SetupCommonParameters(FRDGBuilder& GraphBuilder, FVirtualShadowMapArrayCacheManager* CacheManager, int32 TotalInstanceCount, const FGPUScene& GPUScene, 
 	FVirtualSmInvalidateInstancePagesCS::FParameters& OutPassParameters,
 	FVirtualSmInvalidateInstancePagesCS::FPermutationDomain &OutPermutationVector)
diff --git a/Engine/Source/Runtime/Renderer/Private/VirtualShadowMaps/VirtualShadowMapCacheManager.h b/Engine/Source/Runtime/Renderer/Private/VirtualShadowMaps/VirtualShadowMapCacheManager.h
--- a/Engine/Source/Runtime/Renderer/Private/VirtualShadowMaps/VirtualShadowMapCacheManager.h
+++ b/Engine/Source/Runtime/Renderer/Private/VirtualShadowMaps/VirtualShadowMapCacheManager.h
@@ -208,6 +208,9 @@ public:
 		
 	void SetHZBViewParams(int32 HZBKey, Nanite::FPackedViewParams& OutParams);
 
+#if WITH_MGPU
+	void UpdateGPUMask(FRHIGPUMask GPUMask);
+#endif
 
 	GPUMessage::FSocket StatusFeedbackSocket;
 
@@ -236,6 +239,9 @@ private:
 	FDelegateHandle ScreenMessageDelegate;
 	int32 LastOverflowFrame = -1;
 	bool bLoggedPageOverflow = false;
+#endif
+#if WITH_MGPU
+	FRHIGPUMask LastGPUMask;
 #endif
 	FScene* Scene;
 };
