Split FRayTracingRelevantPrimitiveList RelevantPrimitives into separate static and dynamic primitive lists, use chunked arrays for them. Also use ParallelFor to run ray tracing mesh LOD calculations.

diff --git a/Engine/Source/Runtime/Renderer/Private/DeferredShadingRenderer.cpp b/Engine/Source/Runtime/Renderer/Private/DeferredShadingRenderer.cpp
--- a/Engine/Source/Runtime/Renderer/Private/DeferredShadingRenderer.cpp
+++ b/Engine/Source/Runtime/Renderer/Private/DeferredShadingRenderer.cpp
@@ -73,6 +73,8 @@
 #include "GPUMessaging.h"
 #include "RectLightTextureManager.h"
 #include "Lumen/LumenFrontLayerTranslucency.h"
+#include "Containers/ChunkedArray.h"
+#include "Async/ParallelFor.h"
 
 extern int32 GNaniteShowStats;
 extern int32 GNanitePickingDomain;
@@ -601,8 +603,9 @@ struct FRayTracingRelevantPrimitive
 
 struct FRayTracingRelevantPrimitiveList
 {
-	// Filtered list of relevant primitives
-	TArray<FRayTracingRelevantPrimitive> RelevantPrimitives;
+	// Filtered lists of relevant primitives
+	TChunkedArray<FRayTracingRelevantPrimitive> StaticPrimitives;
+	TChunkedArray<FRayTracingRelevantPrimitive> DynamicPrimitives;
 
 	// Array of primitives that should update their cached ray tracing instances via FPrimitiveSceneInfo::UpdateCachedRaytracingData()
 	TArray<FPrimitiveSceneInfo*> DirtyCachedRayTracingPrimitives;
@@ -620,9 +623,6 @@ static FRayTracingRelevantPrimitiveList GatherRayTracingRelevantPrimitives(const
 {
 	FRayTracingRelevantPrimitiveList Result;
 
-	// Unified array is used for static and dynamic primitives because we don't know ahead of time how many we'll have of each.
-	Result.RelevantPrimitives.Reserve(Scene.PrimitiveSceneProxies.Num());
-
 	Result.DirtyCachedRayTracingPrimitives.Reserve(Scene.PrimitiveSceneProxies.Num());
 
 	const bool bGameView = View.bIsGameView || View.Family->EngineShowFlags.Game;
@@ -707,12 +707,15 @@ static FRayTracingRelevantPrimitiveList GatherRayTracingRelevantPrimitives(const
 				&& RayTracingStaticMeshesCVar && RayTracingStaticMeshesCVar->GetValueOnRenderThread() > 0)
 			{
 				Item.bStatic = true;
-				Result.RelevantPrimitives.Add(Item);
+				Result.StaticPrimitives.AddElement(Item);
 			}
 			else if (View.Family->EngineShowFlags.SkeletalMeshes)
 			{
+				checkf(!EnumHasAllFlags(Scene.PrimitiveRayTracingFlags[PrimitiveIndex], ERayTracingPrimitiveFlags::CacheInstances), 
+					TEXT("Only static primitives are expected to use CacheInstances flag."));
+
 				Item.bStatic = false;
-				Result.RelevantPrimitives.Add(Item);
+				Result.DynamicPrimitives.AddElement(Item);
 			}
 		}
 	}
@@ -724,12 +727,10 @@ static FRayTracingRelevantPrimitiveList GatherRayTracingRelevantPrimitives(const
 		const float LODScaleCVarValue = ICVarStaticMeshLODDistanceScale->GetFloat();
 		const int32 ForcedLODLevel = GetCVarForceLOD();
 
-		for (FRayTracingRelevantPrimitive& RelevantPrimitive : Result.RelevantPrimitives)
+		ParallelFor(TEXT("GatherRayTracingRelevantPrimitives_ComputeLOD"), Result.StaticPrimitives.Num(), 128,
+			[&Result, &Scene, &View, LODScaleCVarValue, ForcedLODLevel](int32 ItemIndex) 
 		{
-			if (!RelevantPrimitive.bStatic)
-			{
-				continue; // skip dynamic primitives
-			}
+			FRayTracingRelevantPrimitive& RelevantPrimitive = Result.StaticPrimitives[ItemIndex];
 
 			const int32 PrimitiveIndex = RelevantPrimitive.PrimitiveIndex;
 			const FPrimitiveSceneInfo* SceneInfo = Scene.Primitives[PrimitiveIndex];
@@ -758,7 +759,7 @@ static FRayTracingRelevantPrimitiveList GatherRayTracingRelevantPrimitives(const
 				FRHIRayTracingGeometry* RayTracingGeometryInstance = SceneInfo->GetStaticRayTracingGeometryInstance(LODIndex);
 				if (RayTracingGeometryInstance == nullptr)
 				{
-					continue;
+					return;
 				}
 
 				// Sometimes LODIndex is out of range because it is clamped by ClampToFirstLOD, like the requested LOD is being streamed in and hasn't been available
@@ -800,7 +801,7 @@ static FRayTracingRelevantPrimitiveList GatherRayTracingRelevantPrimitives(const
 					}
 				}
 			}
-		}
+		});
 	}
 
 	Result.bValid = true;
@@ -862,7 +863,8 @@ bool FDeferredShadingSceneRenderer::GatherRayTracingWorldInstancesForView(FRDGBu
 
 	// Consume output of the relevant primitive gathering task
 	RayTracingScene.UsedCoarseMeshStreamingHandles = MoveTemp(RelevantPrimitiveList.UsedCoarseMeshStreamingHandles);
-	TArray<FRayTracingRelevantPrimitive> RelevantPrimitives = MoveTemp(RelevantPrimitiveList.RelevantPrimitives);
+	TChunkedArray<FRayTracingRelevantPrimitive> RelevantDynamicPrimitives = MoveTemp(RelevantPrimitiveList.DynamicPrimitives);
+	TChunkedArray<FRayTracingRelevantPrimitive> RelevantStaticPrimitives = MoveTemp(RelevantPrimitiveList.StaticPrimitives);
 	TArray<FPrimitiveSceneInfo*> DirtyCachedRayTracingPrimitives = MoveTemp(RelevantPrimitiveList.DirtyCachedRayTracingPrimitives);
 
 	// Inform the coarse mesh streaming manager about all the used streamable render assets in the scene
@@ -872,7 +874,7 @@ bool FDeferredShadingSceneRenderer::GatherRayTracingWorldInstancesForView(FRDGBu
 		CoarseMeshSM->AddUsedStreamingHandles(RayTracingScene.UsedCoarseMeshStreamingHandles);
 	}
 
-	INC_DWORD_STAT_BY(STAT_VisibleRayTracingPrimitives, RelevantPrimitives.Num());
+	INC_DWORD_STAT_BY(STAT_VisibleRayTracingPrimitives, RelevantDynamicPrimitives.Num() + RelevantStaticPrimitives.Num());
 
 	FPrimitiveSceneInfo::UpdateCachedRaytracingData(Scene, DirtyCachedRayTracingPrimitives);
 
@@ -968,13 +970,8 @@ bool FDeferredShadingSceneRenderer::GatherRayTracingWorldInstancesForView(FRDGBu
 		// Local temporary array of instances used for GetDynamicRayTracingInstances()
 		TArray<FRayTracingInstance> TempRayTracingInstances;
 
-		for (const FRayTracingRelevantPrimitive& RelevantPrimitive : RelevantPrimitives)
+		for (const FRayTracingRelevantPrimitive& RelevantPrimitive : RelevantDynamicPrimitives)
 		{
-			if (RelevantPrimitive.bStatic)
-			{
-				continue;
-			}
-
 			const int32 PrimitiveIndex = RelevantPrimitive.PrimitiveIndex;
 			FPrimitiveSceneInfo* SceneInfo = Scene->Primitives[PrimitiveIndex];
 
@@ -1157,7 +1154,7 @@ bool FDeferredShadingSceneRenderer::GatherRayTracingWorldInstancesForView(FRDGBu
 		// Inputs
 
 		const FScene& Scene;
-		TArray<FRayTracingRelevantPrimitive> RelevantPrimitives;
+		TChunkedArray<FRayTracingRelevantPrimitive> RelevantStaticPrimitives;
 		const FRayTracingCullingParameters& CullingParameters;
 
 		// Outputs
@@ -1165,15 +1162,17 @@ bool FDeferredShadingSceneRenderer::GatherRayTracingWorldInstancesForView(FRDGBu
 		FRayTracingScene& RayTracingScene; // New instances are added into FRayTracingScene::Instances and FRayTracingScene::Allocator is used for temporary data
 		TArray<FVisibleRayTracingMeshCommand>& VisibleRayTracingMeshCommands; // New elements are added here by this task
 
-		FRayTracingSceneAddInstancesTask(const FScene& InScene, TArray<FRayTracingRelevantPrimitive> InRelevantPrimitives, const FRayTracingCullingParameters& InCullingParameters,
+		FRayTracingSceneAddInstancesTask(const FScene& InScene,
+											TChunkedArray<FRayTracingRelevantPrimitive> InRelevantStaticPrimitives,
+											const FRayTracingCullingParameters& InCullingParameters,
 											FRayTracingScene& InRayTracingScene, TArray<FVisibleRayTracingMeshCommand>& InVisibleRayTracingMeshCommands)
 			: Scene(InScene)
-			, RelevantPrimitives(MoveTemp(InRelevantPrimitives))
+			, RelevantStaticPrimitives(MoveTemp(InRelevantStaticPrimitives))
 			, CullingParameters(InCullingParameters)
 			, RayTracingScene(InRayTracingScene)
 			, VisibleRayTracingMeshCommands(InVisibleRayTracingMeshCommands)
 		{
-			VisibleRayTracingMeshCommands.Reserve(RelevantPrimitives.Num());
+			VisibleRayTracingMeshCommands.Reserve(RelevantStaticPrimitives.Num());
 		}
 
 		struct FAutoInstanceBatch
@@ -1242,12 +1241,12 @@ bool FDeferredShadingSceneRenderer::GatherRayTracingWorldInstancesForView(FRDGBu
 			TArray<FRayTracingCullPrimitiveInstancesClosure> CullInstancesClosures;
 			if (CullingParameters.CullInRayTracing > 0 && GetRayTracingCullingPerInstance())
 			{
-				CullInstancesClosures.Reserve(RelevantPrimitives.Num());
-				CullingTasks.Reserve(RelevantPrimitives.Num() / 256 + 1);
+				CullInstancesClosures.Reserve(RelevantStaticPrimitives.Num());
+				CullingTasks.Reserve(RelevantStaticPrimitives.Num() / 256 + 1);
 			}
 
 			// scan relevant primitives computing hash data to look for duplicate instances
-			for (const FRayTracingRelevantPrimitive& RelevantPrimitive : RelevantPrimitives)
+			for (const FRayTracingRelevantPrimitive& RelevantPrimitive : RelevantStaticPrimitives)
 			{
 				const int32 PrimitiveIndex = RelevantPrimitive.PrimitiveIndex;
 				FPrimitiveSceneInfo* SceneInfo = Scene.Primitives[PrimitiveIndex];
@@ -1415,7 +1414,7 @@ bool FDeferredShadingSceneRenderer::GatherRayTracingWorldInstancesForView(FRDGBu
 	};
 
 	FGraphEventRef AddInstancesTask = TGraphTask<FRayTracingSceneAddInstancesTask>::CreateTask().ConstructAndDispatchWhenReady(
-		*Scene, MoveTemp(RelevantPrimitives), View.RayTracingCullingParameters, // inputs 
+		*Scene, MoveTemp(RelevantStaticPrimitives), View.RayTracingCullingParameters, // inputs 
 		RayTracingScene, View.VisibleRayTracingMeshCommands // outputs
 	);
 
