VulkanRHI: Fix validation errors where 2D default textures were used for Texture3D parameters.

diff --git a/Engine/Source/Runtime/Renderer/Private/SystemTextures.cpp b/Engine/Source/Runtime/Renderer/Private/SystemTextures.cpp
--- a/Engine/Source/Runtime/Renderer/Private/SystemTextures.cpp
+++ b/Engine/Source/Runtime/Renderer/Private/SystemTextures.cpp
@@ -1020,7 +1020,7 @@ bool operator !=(const FDefaultTextureKey& A, const FDefaultTextureKey& B)
 }
 
 template<typename T>
-static FDefaultTextureKey GetDefaultTextureKey(EPixelFormat Format, const T& In)
+static FDefaultTextureKey GetDefaultTextureKey(ETextureDimension Dimension, EPixelFormat Format, const T& In)
 {
 	FDefaultTextureKey Out;
 	const uint32 Size = sizeof(T);
@@ -1030,6 +1030,7 @@ static FDefaultTextureKey GetDefaultTextureKey(EPixelFormat Format, const T& In)
 	Out.ValueAsUInt[2] = Size > 8  ? InAsUInt[2] : 0u;
 	Out.ValueAsUInt[3] = Size > 12 ? InAsUInt[3] : 0u;
 	Out.Format = Format;
+	Out.Dimension = Dimension;
 	return Out;
 }
 
@@ -1363,7 +1364,7 @@ FRDGTextureRef GetInternalDefaultTexture(
 	if (Format == PF_D24 || Format == PF_DepthStencil)	{ Format = PF_R32_FLOAT; }
 	if (Format == PF_ShadowDepth)						{ Format = PF_R32_FLOAT; }
 
-	const FDefaultTextureKey Key = GetDefaultTextureKey(Format, Value);
+	const FDefaultTextureKey Key = GetDefaultTextureKey(Dimension, Format, Value);
 	const uint32 Hash = Murmur32({uint32(Key.Dimension), uint32(Key.Format), Key.ValueAsUInt[0], Key.ValueAsUInt[1], Key.ValueAsUInt[2], Key.ValueAsUInt[3]});
 
 	uint32 Index = HashDefaultTextures.First(Hash);
