Fix Strata material data clearing for pixel with fully hair strands coverage.

diff --git a/Engine/Shaders/Private/HairStrands/HairStrandsVisibilityDepthPS.usf b/Engine/Shaders/Private/HairStrands/HairStrandsVisibilityDepthPS.usf
--- a/Engine/Shaders/Private/HairStrands/HairStrandsVisibilityDepthPS.usf
+++ b/Engine/Shaders/Private/HairStrands/HairStrandsVisibilityDepthPS.usf
@@ -6,16 +6,16 @@
 #include "HairStrandsVisibilityCommon.ush"
 
 #define OUTPUT_DEPTH 0
-#define OUTPUT_GBUFFER 1 
+#define OUTPUT_MATERIALDATA 1 
 #define OUTPUT_LIGHTCHANNEL 2
-#define OUTPUT_GBUFFER_LIGHTCHANNEL 3
+#define OUTPUT_MATERIALDATA_LIGHTCHANNEL 3
 
 // Manual binding of the sample data as the uniform buffer is not created at this point
 Texture2D<float>					CoverageTexture;
 Texture2D<uint>						HairSampleOffset;
 StructuredBuffer<FPackedHairSample> HairSampleData;
 
-#if PERMUTATION_OUTPUT_TYPE == OUTPUT_GBUFFER_LIGHTCHANNEL || PERMUTATION_OUTPUT_TYPE == OUTPUT_LIGHTCHANNEL
+#if PERMUTATION_OUTPUT_TYPE == OUTPUT_MATERIALDATA_LIGHTCHANNEL || PERMUTATION_OUTPUT_TYPE == OUTPUT_LIGHTCHANNEL
 RWTexture2D<uint> OutLightChannelMaskTexture;
 #endif
 
@@ -45,19 +45,23 @@ FHairSampleData ExtractSampleData(uint3 PixelCoord)
 
 ////////////////////////////////////////////////////////////////////////////////////////////////////////////////////////////////////
 // Gbuffer patching (+ LightChannel mask)
-#if PERMUTATION_OUTPUT_TYPE == OUTPUT_GBUFFER || PERMUTATION_OUTPUT_TYPE == OUTPUT_GBUFFER_LIGHTCHANNEL
+#if STRATA_ENABLED
+#define OUTPUT_FORMAT uint
+#else
+#define OUTPUT_FORMAT float4
+#endif
+#if PERMUTATION_OUTPUT_TYPE == OUTPUT_MATERIALDATA || PERMUTATION_OUTPUT_TYPE == OUTPUT_MATERIALDATA_LIGHTCHANNEL
 void MainPS(
 	in FScreenVertexOutput Input,
-	out float SvDepth		: SV_DEPTH,
-	out float4 OutGBufferB	: SV_Target0,
-	out float4 OutGBufferC	: SV_Target1,
-	out float4 OutColor		: SV_Target2
-)
+	out float SvDepth			: SV_DEPTH,
+	out OUTPUT_FORMAT Output0	: SV_Target0,
+	out OUTPUT_FORMAT Output1	: SV_Target1,
+	out float4 OutColor			: SV_Target2)
 {
 	const uint3 PixelCoord = uint3(Input.Position.xy, 0);
 	const float Coverage = CoverageTexture.Load(PixelCoord);
 
-	#if PERMUTATION_OUTPUT_TYPE == OUTPUT_GBUFFER
+	#if PERMUTATION_OUTPUT_TYPE == OUTPUT_MATERIALDATA
 	if (Coverage < 1.0f)
 	{
 		discard;
@@ -69,11 +73,16 @@ void MainPS(
 	const float DepthBias = 0.001f; 
 	const FHairSampleData SampleData = ExtractSampleData(PixelCoord);
 	SvDepth = SampleData.DeviceZ - DepthBias;
-	OutGBufferB = float4(0, 0, 0, 0); // EncodeShadingModelIdAndSelectiveOutputMask(SHADINGMODELID_HAIR, 0)
-	OutGBufferC = 0;
+#if STRATA_ENABLED
+	Output0 = 0; // MaterialData
+	Output1 = 0; // TopLayerData
+#else
+	Output0 = float4(0, 0, 0, 0);
+	Output1 = 0;
+#endif
 	OutColor = 0;
 
-	#if PERMUTATION_OUTPUT_TYPE == OUTPUT_GBUFFER_LIGHTCHANNEL
+	#if PERMUTATION_OUTPUT_TYPE == OUTPUT_MATERIALDATA_LIGHTCHANNEL
 	OutLightChannelMaskTexture[PixelCoord.xy] = SampleData.LightChannelMask;
 	if (Coverage < 1.0f)
 	{
diff --git a/Engine/Source/Runtime/Renderer/Private/HairStrands/HairStrandsVisibility.cpp b/Engine/Source/Runtime/Renderer/Private/HairStrands/HairStrandsVisibility.cpp
--- a/Engine/Source/Runtime/Renderer/Private/HairStrands/HairStrandsVisibility.cpp
+++ b/Engine/Source/Runtime/Renderer/Private/HairStrands/HairStrandsVisibility.cpp
@@ -2465,8 +2465,8 @@ IMPLEMENT_GLOBAL_SHADER(FHairVisibilityDepthPS, "/Engine/Private/HairStrands/Hai
 
 enum class EHairAuxilaryPassType
 {
-	GBufferPatch,
-	GBufferPatch_LightChannelMask,
+	MaterialData,
+	MaterialData_LightChannelMask,
 	LightChannelMask,
 	DepthPatch,
 	DepthClear
@@ -2480,8 +2480,8 @@ static void AddHairAuxilaryPass(
 	const FRDGTextureRef& CoverageTexture,
 	const FRDGTextureRef& HairSampleOffset,
 	const FRDGBufferRef& HairSampleData,
-	FRDGTextureRef OutGBufferBTexture,
-	FRDGTextureRef OutGBufferCTexture,
+	FRDGTextureRef Output0,
+	FRDGTextureRef Output1,
 	FRDGTextureRef OutColorTexture,
 	FRDGTextureRef OutDepthTexture,
 	FRDGTextureRef OutLightChannelMaskTexture)
@@ -2496,6 +2496,8 @@ static void AddHairAuxilaryPass(
 	Parameters->HairSampleOffset = HairSampleOffset;
 	Parameters->HairSampleData = GraphBuilder.CreateSRV(HairSampleData);
 
+	const bool bStrataEnabled = Strata::IsStrataEnabled();
+
 	const bool bDepthTested = PassType != EHairAuxilaryPassType::LightChannelMask;
 	if (bDepthTested)
 	{
@@ -2506,15 +2508,15 @@ static void AddHairAuxilaryPass(
 			FExclusiveDepthStencil::DepthWrite_StencilNop);
 	}
 
-	if (PassType == EHairAuxilaryPassType::GBufferPatch || PassType == EHairAuxilaryPassType::GBufferPatch_LightChannelMask)
+	if (PassType == EHairAuxilaryPassType::MaterialData || PassType == EHairAuxilaryPassType::MaterialData_LightChannelMask)
 	{
-		check(OutGBufferBTexture && OutGBufferCTexture && OutColorTexture);
-		Parameters->RenderTargets[0] = FRenderTargetBinding(OutGBufferBTexture, ERenderTargetLoadAction::ELoad);
-		Parameters->RenderTargets[1] = FRenderTargetBinding(OutGBufferCTexture, ERenderTargetLoadAction::ELoad);
+		check(Output0 && Output1 && OutColorTexture);
+		Parameters->RenderTargets[0] = FRenderTargetBinding(Output0, ERenderTargetLoadAction::ELoad, 0 /*Mip-0*/, 0/*First slice, which contains state/counter*/);
+		Parameters->RenderTargets[1] = FRenderTargetBinding(Output1, ERenderTargetLoadAction::ELoad);
 		Parameters->RenderTargets[2] = FRenderTargetBinding(OutColorTexture, ERenderTargetLoadAction::ELoad);
 	}
 
-	if (PassType == EHairAuxilaryPassType::GBufferPatch_LightChannelMask || PassType == EHairAuxilaryPassType::LightChannelMask)
+	if (PassType == EHairAuxilaryPassType::MaterialData_LightChannelMask || PassType == EHairAuxilaryPassType::LightChannelMask)
 	{
 		check(OutLightChannelMaskTexture);
 		Parameters->OutLightChannelMaskTexture = GraphBuilder.CreateUAV(OutLightChannelMaskTexture);
@@ -2526,9 +2528,9 @@ static void AddHairAuxilaryPass(
 	{
 		case EHairAuxilaryPassType::DepthPatch:						OutputType = 0; Method = TEXT("HairOnlyDepth"); break;
 		case EHairAuxilaryPassType::DepthClear:						OutputType = 0; Method = TEXT("HairOnlyDepth:Clear"); break;
-		case EHairAuxilaryPassType::GBufferPatch:					OutputType = 1; Method = TEXT("GBuffer"); break;
+		case EHairAuxilaryPassType::MaterialData:					OutputType = 1; Method = TEXT("MaterialData"); break;
 		case EHairAuxilaryPassType::LightChannelMask:				OutputType = 2; Method = TEXT("LightChannel"); break;
-		case EHairAuxilaryPassType::GBufferPatch_LightChannelMask:	OutputType = 3; Method = TEXT("GBuffer, LightChannel"); break;
+		case EHairAuxilaryPassType::MaterialData_LightChannelMask:	OutputType = 3; Method = TEXT("MaterialData, LightChannel"); break;
 		default: 													OutputType = 0; Method = TEXT("Unknown"); break;
 	};
 
@@ -2625,20 +2627,20 @@ static FRDGTextureRef AddHairLightChannelMaskPass(
 #endif
 
 
-static void AddHairGbufferPatchPass(
+static void AddHairMaterialDataPatchPass(
 	FRDGBuilder& GraphBuilder,
 	const FViewInfo& View,
 	const FHairStrandsTiles& TileData,
 	const FRDGTextureRef& CoverageTexture,
 	const FRDGTextureRef& HairSampleOffset,
 	const FRDGBufferRef& HairSampleData,
-	FRDGTextureRef& OutGBufferBTexture,
-	FRDGTextureRef& OutGBufferCTexture,
+	FRDGTextureRef& OutMaterial0,
+	FRDGTextureRef& OutMaterial1,
 	FRDGTextureRef& OutColorTexture,
 	FRDGTextureRef& OutDepthTexture,
 	FRDGTextureRef& OutLightChannelMask)
 {
-	if (!OutGBufferBTexture || !OutGBufferCTexture || !OutColorTexture || !OutDepthTexture)
+	if (!OutMaterial0 || !OutMaterial1 || !OutColorTexture || !OutDepthTexture)
 	{
 		return;
 	}
@@ -2657,12 +2659,12 @@ static void AddHairGbufferPatchPass(
 		GraphBuilder,
 		View,
 		TileData,
-		bLightingChannel ? EHairAuxilaryPassType::GBufferPatch_LightChannelMask : EHairAuxilaryPassType::GBufferPatch,
+		bLightingChannel ? EHairAuxilaryPassType::MaterialData_LightChannelMask : EHairAuxilaryPassType::MaterialData,
 		CoverageTexture,
 		HairSampleOffset,
 		HairSampleData,
-		OutGBufferBTexture,
-		OutGBufferCTexture,
+		OutMaterial0,
+		OutMaterial1,
 		OutColorTexture,
 		OutDepthTexture,
 		OutLightChannelMask);
@@ -3805,6 +3807,19 @@ void RenderHairStrandsVisibilityBuffer(
 		return;
 	}
 
+	FRDGTextureRef SceneMaterial0 = nullptr;
+	FRDGTextureRef SceneMaterial1 = nullptr;
+	if (Strata::IsStrataEnabled())
+	{
+		SceneMaterial0 = View.StrataViewData.SceneData->MaterialTextureArray;
+		SceneMaterial1 = View.StrataViewData.SceneData->TopLayerTexture;
+	}
+	else
+	{
+		SceneMaterial0 = SceneGBufferBTexture;
+		SceneMaterial1 = SceneGBufferCTexture;
+	}
+
 	{
 		
 		{
@@ -3821,7 +3836,9 @@ void RenderHairStrandsVisibilityBuffer(
 			const EHairVisibilityRenderMode RenderMode = GetHairVisibilityRenderMode();
 			check(RenderMode == HairVisibilityRenderMode_MSAA_Visibility || RenderMode == HairVisibilityRenderMode_PPLL || RenderMode == HairVisibilityRenderMode_ComputeRaster);
 
-			const bool bRunColorAndDepthPatching = SceneGBufferBTexture && SceneColorTexture && RenderMode != HairVisibilityRenderMode_ComputeRaster;
+			const bool bRunColorAndDepthPatching = 
+				(SceneMaterial0 && SceneMaterial1) &&
+				SceneColorTexture && RenderMode != HairVisibilityRenderMode_ComputeRaster;
 
 			FRDGTextureRef HairOnlyDepthTexture = GraphBuilder.CreateTexture(SceneDepthTexture->Desc, TEXT("Hair.HairOnlyDepthTexture"));
 			FRDGTextureRef CoverageTexture = nullptr;
@@ -3917,15 +3934,15 @@ void RenderHairStrandsVisibilityBuffer(
 					// * unlit shading model ID 
 					if (bRunColorAndDepthPatching)
 					{
-						AddHairGbufferPatchPass(
+						AddHairMaterialDataPatchPass(
 							GraphBuilder,
 							View,
 							VisibilityData.TileData,
 							CoverageTexture,
 							CompactNodeIndex,
 							CompactNodeData,
-							SceneGBufferBTexture,
-							SceneGBufferCTexture,
+							SceneMaterial0,
+							SceneMaterial1,
 							SceneColorTexture,
 							SceneDepthTexture,
 							VisibilityData.LightChannelMaskTexture);
@@ -4113,15 +4130,15 @@ void RenderHairStrandsVisibilityBuffer(
 					// * unlit shading model ID 
 					if (bRunColorAndDepthPatching)
 					{
-						AddHairGbufferPatchPass(
+						AddHairMaterialDataPatchPass(
 							GraphBuilder,
 							View,
 							VisibilityData.TileData,
 							CoverageTexture,
 							CompactNodeIndex,
 							CompactNodeData,
-							SceneGBufferBTexture,
-							SceneGBufferCTexture,
+							SceneMaterial0,
+							SceneMaterial1,
 							SceneColorTexture,
 							SceneDepthTexture,
 							VisibilityData.LightChannelMaskTexture);
@@ -4251,15 +4268,15 @@ void RenderHairStrandsVisibilityBuffer(
 
 				if (bRunColorAndDepthPatching)
 				{
-					AddHairGbufferPatchPass(
+					AddHairMaterialDataPatchPass(
 						GraphBuilder,
 						View,
 						VisibilityData.TileData,
 						CoverageTexture,
 						CompactNodeIndex,
 						CompactNodeData,
-						SceneGBufferBTexture,
-						SceneGBufferCTexture,
+						SceneMaterial0,
+						SceneMaterial1,
 						SceneColorTexture,
 						SceneDepthTexture,
 						VisibilityData.LightChannelMaskTexture);
