Add tile support for hit proxy.

diff --git a/Engine/Shaders/Private/HairStrands/HairStrandsVisibilityTile.usf b/Engine/Shaders/Private/HairStrands/HairStrandsVisibilityTile.usf
--- a/Engine/Shaders/Private/HairStrands/HairStrandsVisibilityTile.usf
+++ b/Engine/Shaders/Private/HairStrands/HairStrandsVisibilityTile.usf
@@ -12,6 +12,7 @@
 #if SHADER_TILE_GENERATION
 int2 BufferResolution;
 uint bUintTexture;
+uint bForceOutputAllTiles;
 float TransmittanceThreshold;
 uint IntCoverageThreshold;
 Texture2D<float> InputFloatTexture;
@@ -36,7 +37,13 @@ void TileMainCS(uint2 DispatchThreadId : SV_DispatchThreadID, uint LinearIndex :
 	bool pHair = false;
 	if (all(PixelCoord < uint2(BufferResolution)))
 	{
-		if (bUintTexture)
+		if (bForceOutputAllTiles)
+		{
+			// When bForceOutputAllTiles, we generate tiles for all pixels.
+			fHair = true;
+			pHair = true;
+		}
+		else if (bUintTexture)
 		{
 			const uint Coverage = InputUintTexture.Load(uint3(PixelCoord, 0));
 			fHair = Coverage >= IntCoverageThreshold;
diff --git a/Engine/Source/Runtime/Renderer/Private/HairStrands/HairStrandsTile.cpp b/Engine/Source/Runtime/Renderer/Private/HairStrands/HairStrandsTile.cpp
--- a/Engine/Source/Runtime/Renderer/Private/HairStrands/HairStrandsTile.cpp
+++ b/Engine/Source/Runtime/Renderer/Private/HairStrands/HairStrandsTile.cpp
@@ -121,6 +121,7 @@ class FHairStrandsTileGenerationPassCS : public FGlobalShader
 		SHADER_PARAMETER_STRUCT_REF(FViewUniformShaderParameters, ViewUniformBuffer)
 		SHADER_PARAMETER(FIntPoint, BufferResolution)
 		SHADER_PARAMETER(uint32, bUintTexture)
+		SHADER_PARAMETER(uint32, bForceOutputAllTiles)
 		SHADER_PARAMETER(float, TransmittanceThreshold)
 		SHADER_PARAMETER(uint32, IntCoverageThreshold)
 		SHADER_PARAMETER_RDG_TEXTURE(Texture2D<float>, InputFloatTexture)
@@ -149,16 +150,19 @@ IMPLEMENT_GLOBAL_SHADER(FHairStrandsTileGenerationPassCS, "/Engine/Private/HairS
 float GetHairStrandsFullCoverageThreshold();
 uint32 GetHairStrandsIntCoverageThreshold();
 
-FHairStrandsTiles AddHairStrandsGenerateTilesPass(
+static FHairStrandsTiles AddHairStrandsGenerateTilesPass_Internal(
 	FRDGBuilder& GraphBuilder,
 	const FViewInfo& View,
+	const FIntPoint& InputResolution,
 	const FRDGTextureRef& InputTexture)
 {
+	const bool bHasValidInput = InputTexture != nullptr;
+	const bool bUintTexture   = InputTexture && InputTexture->Desc.Format == PF_R32_UINT;
+
 	FHairStrandsTiles Out;
 
 	check(FHairStrandsTiles::TilePerThread_GroupSize == 64); // If this value change, we need to update the shaders using 
 	check(FHairStrandsTiles::TileSize == 8); // only size supported for now
-	const FIntPoint InputResolution = InputTexture->Desc.Extent;
 	Out.TileCountXY = FIntPoint(FMath::CeilToInt(InputResolution.X / float(FHairStrandsTiles::TileSize)), FMath::CeilToInt(InputResolution.Y / float(FHairStrandsTiles::TileSize)));
 	Out.TileCount = Out.TileCountXY.X * Out.TileCountXY.Y;
 	Out.BufferResolution = InputResolution;
@@ -176,15 +180,23 @@ FHairStrandsTiles AddHairStrandsGenerateTilesPass(
 	FRDGBufferUAVRef TileCountUAV = GraphBuilder.CreateUAV(Out.TileCountBuffer, PF_R32_UINT);
 	AddClearUAVPass(GraphBuilder, TileCountUAV, 0u);
 
-	const bool bUintTexture = InputTexture->Desc.Format == PF_R32_UINT;
 	TShaderMapRef<FHairStrandsTileGenerationPassCS> ComputeShader(View.ShaderMap);
 	FHairStrandsTileGenerationPassCS::FParameters* PassParameters = GraphBuilder.AllocParameters<FHairStrandsTileGenerationPassCS::FParameters>();
 	PassParameters->BufferResolution		= InputResolution;//View.ViewRect.Size();
+	PassParameters->bForceOutputAllTiles	= bHasValidInput ? 0 : 1;
 	PassParameters->bUintTexture			= bUintTexture ? 1u : 0u;
 	PassParameters->TransmittanceThreshold	= 1.f - GetHairStrandsFullCoverageThreshold();	
 	PassParameters->IntCoverageThreshold 	= GetHairStrandsIntCoverageThreshold();
-	PassParameters->InputFloatTexture = bUintTexture ? GSystemTextures.GetBlackDummy(GraphBuilder) : InputTexture;
-	PassParameters->InputUintTexture  = bUintTexture ? InputTexture : GSystemTextures.GetZeroUIntDummy(GraphBuilder);
+	PassParameters->InputFloatTexture		= GSystemTextures.GetBlackDummy(GraphBuilder);
+	PassParameters->InputUintTexture		= GSystemTextures.GetZeroUIntDummy(GraphBuilder);
+	if (bHasValidInput && bUintTexture)
+	{
+		PassParameters->InputUintTexture	= InputTexture;
+	}
+	else if (bHasValidInput && !bUintTexture)
+	{
+		PassParameters->InputFloatTexture	= InputTexture;
+	}
 	PassParameters->TileHairAllBuffer		= GraphBuilder.CreateUAV(Out.TileDataBuffer[ToIndex(FHairStrandsTiles::ETileType::HairAll)], PF_R16G16_UINT);
 	PassParameters->TileHairFullBuffer		= GraphBuilder.CreateUAV(Out.TileDataBuffer[ToIndex(FHairStrandsTiles::ETileType::HairFull)], PF_R16G16_UINT);
 	PassParameters->TileHairPartialBuffer	= GraphBuilder.CreateUAV(Out.TileDataBuffer[ToIndex(FHairStrandsTiles::ETileType::HairPartial)], PF_R16G16_UINT);
@@ -213,6 +225,16 @@ FHairStrandsTiles AddHairStrandsGenerateTilesPass(
 	return Out;
 }
 
+FHairStrandsTiles AddHairStrandsGenerateTilesPass(FRDGBuilder& GraphBuilder, const FViewInfo& View, const FRDGTextureRef& InputTexture)
+{
+	return AddHairStrandsGenerateTilesPass_Internal(GraphBuilder, View, InputTexture->Desc.Extent, InputTexture);
+}
+
+FHairStrandsTiles AddHairStrandsGenerateTilesPass(FRDGBuilder& GraphBuilder, const FViewInfo& View, const FIntPoint& Resolution)
+{
+	return AddHairStrandsGenerateTilesPass_Internal(GraphBuilder, View, Resolution, nullptr);
+}
+
 ////////////////////////////////////////////////////////////////////////////////////////////////////////////////////////////////////////////////
 
 bool FHairStrandsTilePassVS::ShouldCompilePermutation(const FGlobalShaderPermutationParameters& Parameters)
diff --git a/Engine/Source/Runtime/Renderer/Private/HairStrands/HairStrandsTile.h b/Engine/Source/Runtime/Renderer/Private/HairStrands/HairStrandsTile.h
--- a/Engine/Source/Runtime/Renderer/Private/HairStrands/HairStrandsTile.h
+++ b/Engine/Source/Runtime/Renderer/Private/HairStrands/HairStrandsTile.h
@@ -40,6 +40,11 @@ FHairStrandsTiles AddHairStrandsGenerateTilesPass(
 	const FViewInfo& View,
 	const FRDGTextureRef& InputTexture);
 
+FHairStrandsTiles AddHairStrandsGenerateTilesPass(
+	FRDGBuilder& GraphBuilder,
+	const FViewInfo& View,
+	const FIntPoint& Extent);
+
 void AddHairStrandsTileClearPass(
 	FRDGBuilder& GraphBuilder,
 	const FViewInfo& View,
diff --git a/Engine/Source/Runtime/Renderer/Private/HairStrands/HairStrandsVisibility.cpp b/Engine/Source/Runtime/Renderer/Private/HairStrands/HairStrandsVisibility.cpp
--- a/Engine/Source/Runtime/Renderer/Private/HairStrands/HairStrandsVisibility.cpp
+++ b/Engine/Source/Runtime/Renderer/Private/HairStrands/HairStrandsVisibility.cpp
@@ -3672,6 +3672,10 @@ void DrawHitProxies(
 			TileData,
 			SceneDepthTexture);
 
+		// Generate Tile data
+		check(VisDepthTexture);
+		TileData = AddHairStrandsGenerateTilesPass(GraphBuilder, View, VisDepthTexture->Desc.Extent);
+
 		FRDGTextureRef VisIdTexture = nullptr;
 		AddHairVisibilityMSAAPass(
 			GraphBuilder,
