Remove all non-tile path from hair code.

This removes a few global shaders permutation.

diff --git a/Engine/Shaders/Private/HairStrands/HairStrandsVelocity.usf b/Engine/Shaders/Private/HairStrands/HairStrandsVelocity.usf
--- a/Engine/Shaders/Private/HairStrands/HairStrandsVelocity.usf
+++ b/Engine/Shaders/Private/HairStrands/HairStrandsVelocity.usf
@@ -36,12 +36,10 @@
   #endif
 #endif
 
-#if PERMUTATION_TILE
 uint TileSize;
 int2 TileCountXY;
 Buffer<uint> TileCountBuffer;
 Buffer<uint2> TileDataBuffer;
-#endif
 
 RWTexture2D<uint> OutResolveMaskTexture;
 
@@ -54,7 +52,6 @@
 [numthreads(TILE_PIXEL_SIZE_X, TILE_PIXEL_SIZE_Y, 1)]
 void MainCS(uint3 DispatchThreadId : SV_DispatchThreadID, uint2 InGroupId : SV_GroupID, uint2 InGroupThreadId : SV_GroupThreadID)
 {
-#if PERMUTATION_TILE
 	const uint TileCount = TileCountBuffer[HAIRTILE_HAIR_ALL];
 	const uint TileIndex1D = InGroupId.x + InGroupId.y * TileCountXY.x;
 	if (TileIndex1D >= TileCount)
@@ -63,9 +60,6 @@
 	}
 	const uint2 GroupId = TileDataBuffer[TileIndex1D];
 	const uint2 PixelCoord = GroupId * TileSize + InGroupThreadId;
-#else
-	const uint2 PixelCoord = DispatchThreadId.xy + ResolutionOffset;
-#endif
 
 	const FNodeDesc NodeDesc = DecodeNodeDesc(NodeIndex.Load(uint3(PixelCoord, 0)));
 	bool bNeedFastResolve = false;
diff --git a/Engine/Shaders/Private/HairStrands/HairStrandsVisibilityCompaction.usf b/Engine/Shaders/Private/HairStrands/HairStrandsVisibilityCompaction.usf
--- a/Engine/Shaders/Private/HairStrands/HairStrandsVisibilityCompaction.usf
+++ b/Engine/Shaders/Private/HairStrands/HairStrandsVisibilityCompaction.usf
@@ -345,12 +345,10 @@ float  GetDepthTexture(uint2 PixelCoord, uint SampleIt) { return MSAA_DepthTextu
 
 #endif // PERMUTATION_MSAACOUNT>1
 
-#if PERMUTATION_TILE
 int2 TileCountXY;
 uint TileSize;
 Buffer<uint> TileCountBuffer;
 Buffer<uint2> TileDataBuffer;
-#endif // PERMUTATION_TILE
 
 Texture2D<float> ViewTransmittanceTexture;
 
@@ -366,7 +364,6 @@ void MainCS(
 		AllocationNodeCount = 0;
 		AllocationNodeOffset = 0;
 	}
-#if PERMUTATION_TILE
 	const uint TileCount = TileCountBuffer[HAIRTILE_HAIR_ALL];
 	const uint LinearIndex = GroupId.x + GroupId.y * TileCountXY.x;
 	if (LinearIndex >= TileCount)
@@ -375,9 +372,6 @@ void MainCS(
 	}
 	const uint2 TileCoord = TileDataBuffer[LinearIndex];
 	uint2 PixelCoord = TileCoord * TileSize + GroupThreadId.xy;
-#else
-	uint2 PixelCoord = DispatchThreadId.xy + ResolutionOffset;
-#endif // PERMUTATION_TILE
 
 	if (PixelCoord.x >= uint(OutputResolution.x) || PixelCoord.y >= uint(OutputResolution.y))
 		PixelCoord = uint2(0, 0);
diff --git a/Engine/Shaders/Private/HairStrands/HairStrandsVisibilityCompactionComputeRaster.usf b/Engine/Shaders/Private/HairStrands/HairStrandsVisibilityCompactionComputeRaster.usf
--- a/Engine/Shaders/Private/HairStrands/HairStrandsVisibilityCompactionComputeRaster.usf
+++ b/Engine/Shaders/Private/HairStrands/HairStrandsVisibilityCompactionComputeRaster.usf
@@ -65,12 +65,10 @@ Texture2D<uint>  DepthCovTexture;
 Texture2D<uint>  PrimMatTexture;
 Texture2D<uint>  HairCountTexture;
 
-#if PERMUTATION_TILE
 int2 TileCountXY;
 uint TileSize;
 Buffer<uint> TileCountBuffer;
 Buffer<uint2> TileDataBuffer;
-#endif // PERMUTATION_TILE
 
 [numthreads(TILE_PIXEL_SIZE_X, TILE_PIXEL_SIZE_Y, 1)]
 void MainCS(
@@ -85,7 +83,6 @@ void MainCS(
 		AllocationNodeOffset = 0;
 	}
 
-#if PERMUTATION_TILE
 	const uint TileCount = TileCountBuffer[HAIRTILE_HAIR_ALL];
 	const uint LinearIndex = GroupId.x + GroupId.y * TileCountXY.x;
 	if (LinearIndex >= TileCount)
@@ -94,9 +91,6 @@ void MainCS(
 	}
 	const uint2 TileCoord = TileDataBuffer[LinearIndex];
 	uint2 PixelCoord = TileCoord * TileSize + GroupThreadId.xy;
-#else
-	uint2 PixelCoord = DispatchThreadId.xy;// +ResolutionOffset;
-#endif // PERMUTATION_TILE
 
 	if (PixelCoord.x >= uint(OutputResolution.x) || PixelCoord.y >= uint(OutputResolution.y))
 		PixelCoord = uint2(0, 0);
diff --git a/Engine/Source/Runtime/Renderer/Private/HairStrands/HairStrandsRendering.cpp b/Engine/Source/Runtime/Renderer/Private/HairStrands/HairStrandsRendering.cpp
--- a/Engine/Source/Runtime/Renderer/Private/HairStrands/HairStrandsRendering.cpp
+++ b/Engine/Source/Runtime/Renderer/Private/HairStrands/HairStrandsRendering.cpp
@@ -27,23 +27,11 @@ static TRDGUniformBufferRef<FHairStrandsViewUniformParameters> InternalCreateHai
 		Parameters->HairSampleViewportResolution = In->SampleLightingViewportResolution;
 		Parameters->MaxSamplePerPixelCount = In->MaxSampleCount;
 
-		if (In->TileData.IsValid())
-		{
-			Parameters->HairTileData = In->TileData.GetTileBufferSRV(FHairStrandsTiles::ETileType::HairAll);
-			Parameters->HairTileCount = In->TileData.TileCountSRV;
-			Parameters->HairTileCountXY = In->TileData.TileCountXY;
-			Parameters->bHairTileValid = true;
-		}
-		else
-		{
-			FRDGBufferRef DummyBuffer = GSystemTextures.GetDefaultBuffer(GraphBuilder, 4);
-			FRDGBufferSRVRef DummyBufferR32SRV = GraphBuilder.CreateSRV(DummyBuffer, PF_R32_UINT);
-			FRDGBufferSRVRef DummyBufferRG16SRV = GraphBuilder.CreateSRV(DummyBuffer, PF_R16G16_UINT);
-			Parameters->HairTileData = DummyBufferRG16SRV;
-			Parameters->HairTileCount = DummyBufferR32SRV;
-			Parameters->HairTileCountXY = FIntPoint(0,0);
-			Parameters->bHairTileValid = false;
-		}
+		check(In->TileData.IsValid())
+		Parameters->HairTileData = In->TileData.GetTileBufferSRV(FHairStrandsTiles::ETileType::HairAll);
+		Parameters->HairTileCount = In->TileData.TileCountSRV;
+		Parameters->HairTileCountXY = In->TileData.TileCountXY;
+		Parameters->bHairTileValid = true;
 
 		if (!Parameters->HairOnlyDepthFurthestHZBTexture)
 		{
diff --git a/Engine/Source/Runtime/Renderer/Private/HairStrands/HairStrandsVisibility.cpp b/Engine/Source/Runtime/Renderer/Private/HairStrands/HairStrandsVisibility.cpp
--- a/Engine/Source/Runtime/Renderer/Private/HairStrands/HairStrandsVisibility.cpp
+++ b/Engine/Source/Runtime/Renderer/Private/HairStrands/HairStrandsVisibility.cpp
@@ -944,8 +944,7 @@ class FHairVelocityCS: public FGlobalShader
 
 	class FVelocity : SHADER_PERMUTATION_INT("PERMUTATION_VELOCITY", 4);
 	class FOuputFormat : SHADER_PERMUTATION_INT("PERMUTATION_OUTPUT_FORMAT", 2);
-	class FTile : SHADER_PERMUTATION_BOOL("PERMUTATION_TILE");
-	using FPermutationDomain = TShaderPermutationDomain<FVelocity, FOuputFormat, FTile>;
+	using FPermutationDomain = TShaderPermutationDomain<FVelocity, FOuputFormat>;
 
 	BEGIN_SHADER_PARAMETER_STRUCT(FParameters, )
 
@@ -998,22 +997,15 @@ static void AddHairVelocityPass(
 	if (!bWriteOutVelocity)
 		return;
 
+	check(TileData.IsValid());
+
 	// If velocity texture has not been created by the base-pass, clear it here
 	const bool bNeedClear = !HasBeenProduced(OutVelocityTexture);
 	if (bNeedClear)
 	{
-		if (!TileData.IsValid())
-		{
-			AddClearUAVPass(GraphBuilder, GraphBuilder.CreateUAV(OutVelocityTexture), 0.f);
-		}
-		else
-		{
-			AddHairStrandsTileClearPass(GraphBuilder, View, TileData, FHairStrandsTiles::ETileType::Other, OutVelocityTexture);
-		}
+		AddHairStrandsTileClearPass(GraphBuilder, View, TileData, FHairStrandsTiles::ETileType::Other, OutVelocityTexture);
 	}
 
-	const bool bUseTile = TileData.IsValid();
-
 	const FIntPoint Resolution = OutVelocityTexture->Desc.Extent;
 	OutResolveMaskTexture = GraphBuilder.CreateTexture(FRDGTextureDesc::Create2D(Resolution, PF_R8_UINT, FClearValueBinding::None, TexCreate_UAV), TEXT("Hair.VelocityResolveMaskTexture"));
 
@@ -1024,7 +1016,6 @@ static void AddHairVelocityPass(
 	FHairVelocityCS::FPermutationDomain PermutationVector;
 	PermutationVector.Set<FHairVelocityCS::FVelocity>(bWriteOutVelocity ? FMath::Clamp(GHairVelocityType + 1, 0, 3) : 0);
 	PermutationVector.Set<FHairVelocityCS::FOuputFormat>(bTwoChannelsOutput ? 0 : 1);
-	PermutationVector.Set<FHairVelocityCS::FTile>(bUseTile);
 
 	FHairVelocityCS::FParameters* PassParameters = GraphBuilder.AllocParameters<FHairVelocityCS::FParameters>();
 	PassParameters->bNeedClear = bNeedClear ? 1u : 0u;
@@ -1038,44 +1029,23 @@ static void AddHairVelocityPass(
 	PassParameters->OutVelocityTexture = GraphBuilder.CreateUAV(OutVelocityTexture);
 	PassParameters->OutResolveMaskTexture = GraphBuilder.CreateUAV(OutResolveMaskTexture);
 
-	if (bUseTile)
-	{
-		const FHairStrandsTiles::ETileType TileType = FHairStrandsTiles::ETileType::HairAll;
-
-		PassParameters->ResolutionOffset	= FIntPoint(0,0);
-		PassParameters->Resolution			= Resolution;
-		PassParameters->TileCountXY			= TileData.TileCountXY;
-		PassParameters->TileSize			= TileData.TileSize;
-		PassParameters->TileCountBuffer		= TileData.TileCountSRV;
-		PassParameters->TileDataBuffer		= TileData.GetTileBufferSRV(TileType);
-		PassParameters->TileIndirectArgs	= TileData.TileIndirectDispatchBuffer;
-
-		TShaderMapRef<FHairVelocityCS> ComputeShader(View.ShaderMap, PermutationVector);
-		FComputeShaderUtils::AddPass(
-			GraphBuilder,
-			RDG_EVENT_NAME("HairStrands::Velocity(Tile)"),
-			ComputeShader,
-			PassParameters,
-			PassParameters->TileIndirectArgs, TileData.GetIndirectDispatchArgOffset(TileType));
-	}
-	else
-	{
-		// We don't use the CPU screen projection for running the velocity pass, as we need to clear the entire 
-		// velocity mask through the UAV write, otherwise the mask will be partially invalid.
-		const FIntRect TotalRect = View.ViewRect; 
-		const FIntPoint RectResolution(TotalRect.Width(), TotalRect.Height());
+	const FHairStrandsTiles::ETileType TileType = FHairStrandsTiles::ETileType::HairAll;
 
-		PassParameters->ResolutionOffset = FIntPoint(TotalRect.Min.X, TotalRect.Min.Y);
-		PassParameters->Resolution		 = RectResolution;
+	PassParameters->ResolutionOffset	= FIntPoint(0,0);
+	PassParameters->Resolution			= Resolution;
+	PassParameters->TileCountXY			= TileData.TileCountXY;
+	PassParameters->TileSize			= TileData.TileSize;
+	PassParameters->TileCountBuffer		= TileData.TileCountSRV;
+	PassParameters->TileDataBuffer		= TileData.GetTileBufferSRV(TileType);
+	PassParameters->TileIndirectArgs	= TileData.TileIndirectDispatchBuffer;
 
-		TShaderMapRef<FHairVelocityCS> ComputeShader(View.ShaderMap, PermutationVector);
-		FComputeShaderUtils::AddPass(
-			GraphBuilder,
-			RDG_EVENT_NAME("HairStrands::Velocity(Screen)"),
-			ComputeShader,
-			PassParameters,
-			FComputeShaderUtils::GetGroupCount(RectResolution, FIntPoint(8, 8)));
-	}
+	TShaderMapRef<FHairVelocityCS> ComputeShader(View.ShaderMap, PermutationVector);
+	FComputeShaderUtils::AddPass(
+		GraphBuilder,
+		RDG_EVENT_NAME("HairStrands::Velocity"),
+		ComputeShader,
+		PassParameters,
+		PassParameters->TileIndirectArgs, TileData.GetIndirectDispatchArgOffset(TileType));
 }
 
 /////////////////////////////////////////////////////////////////////////////////////////
@@ -1473,7 +1443,7 @@ static void AddClearGraphicPass(
 	FRDGTextureRef& OutTarget)
 {
 	check(OutTarget);
-	const bool bUseTile = TileData.IsValid();
+	check(TileData.IsValid());
 
 	const FHairStrandsTiles::ETileType TileType = FHairStrandsTiles::ETileType::HairAll;
 
@@ -1492,10 +1462,9 @@ static void AddClearGraphicPass(
 		PermutationVector.Set<FClearUIntGraphicPS::FOutputFormat>(1);
 	}
 
-	TShaderMapRef<FPostProcessVS> ScreenVertexShader(View->ShaderMap);
 	TShaderMapRef<FHairStrandsTilePassVS> TileVertexShader(View->ShaderMap);
 	TShaderMapRef<FClearUIntGraphicPS> PixelShader(View->ShaderMap, PermutationVector);
-	const FIntRect Viewport = bUseTile ? View->ViewRect : FIntRect(FIntPoint(0, 0), OutTarget->Desc.Extent);
+	const FIntRect Viewport = View->ViewRect;
 	const FIntPoint Resolution = OutTarget->Desc.Extent;
 
 	//ClearUnusedGraphResources(PixelShader, Parameters);
@@ -1504,7 +1473,7 @@ static void AddClearGraphicPass(
 		Forward<FRDGEventName>(PassName),
 		Parameters,
 		ERDGPassFlags::Raster,
-		[Parameters, ScreenVertexShader, TileVertexShader, PixelShader, Viewport, Resolution, bUseTile, TileType](FRHICommandList& RHICmdList)
+		[Parameters, TileVertexShader, PixelShader, Viewport, Resolution, TileType](FRHICommandList& RHICmdList)
 	{
 		FHairStrandsTilePassVS::FParameters ParametersVS = Parameters->TileData;
 
@@ -1514,32 +1483,16 @@ static void AddClearGraphicPass(
 		GraphicsPSOInit.RasterizerState = TStaticRasterizerState<>::GetRHI();
 		GraphicsPSOInit.DepthStencilState = TStaticDepthStencilState<false, CF_Always>::GetRHI();
 		GraphicsPSOInit.BoundShaderState.VertexDeclarationRHI = GFilterVertexDeclaration.VertexDeclarationRHI;
-		GraphicsPSOInit.BoundShaderState.VertexShaderRHI = bUseTile ? TileVertexShader.GetVertexShader() : ScreenVertexShader.GetVertexShader();
+		GraphicsPSOInit.BoundShaderState.VertexShaderRHI = TileVertexShader.GetVertexShader();
 		GraphicsPSOInit.BoundShaderState.PixelShaderRHI = PixelShader.GetPixelShader();
 		GraphicsPSOInit.PrimitiveType = Parameters->TileData.bRectPrimitive > 0 ? PT_RectList : PT_TriangleList;
 		SetGraphicsPipelineState(RHICmdList, GraphicsPSOInit, 0);
 
 		RHICmdList.SetViewport(Viewport.Min.X, Viewport.Min.Y, 0.0f, Viewport.Max.X, Viewport.Max.Y, 1.0f);
 		SetShaderParameters(RHICmdList, PixelShader, PixelShader.GetPixelShader(), *Parameters);
-		if (bUseTile)
-		{
-			SetShaderParameters(RHICmdList, TileVertexShader, TileVertexShader.GetVertexShader(), ParametersVS);
-			RHICmdList.SetStreamSource(0, nullptr, 0);
-			RHICmdList.DrawPrimitiveIndirect(Parameters->TileData.TileIndirectBuffer->GetRHI(), FHairStrandsTiles::GetIndirectDrawArgOffset(TileType));
-		}
-		else
-		{
-			DrawRectangle(
-				RHICmdList,
-				0, 0,
-				Viewport.Width(), Viewport.Height(),
-				Viewport.Min.X, Viewport.Min.Y,
-				Viewport.Width(), Viewport.Height(),
-				Viewport.Size(),
-				Resolution,
-				ScreenVertexShader,
-				EDRF_UseTriangleOptimization);
-		}
+		SetShaderParameters(RHICmdList, TileVertexShader, TileVertexShader.GetVertexShader(), ParametersVS);
+		RHICmdList.SetStreamSource(0, nullptr, 0);
+		RHICmdList.DrawPrimitiveIndirect(Parameters->TileData.TileIndirectBuffer->GetRHI(), FHairStrandsTiles::GetIndirectDrawArgOffset(TileType));
 	});
 }
 
@@ -1599,10 +1552,9 @@ class FHairVisibilityPrimitiveIdCompactionCS : public FGlobalShader
 	SHADER_USE_PARAMETER_STRUCT(FHairVisibilityPrimitiveIdCompactionCS, FGlobalShader);
 
 	class FGroupSize	: SHADER_PERMUTATION_SPARSE_INT("PERMUTATION_GROUPSIZE", 32, 64);
-	class FTile			: SHADER_PERMUTATION_BOOL("PERMUTATION_TILE");
 	class FPPLL 		: SHADER_PERMUTATION_SPARSE_INT("PERMUTATION_PPLL", 0, 8, 16, 32); // See GetPPLLMaxRenderNodePerPixel
 	class FMSAACount 	: SHADER_PERMUTATION_SPARSE_INT("PERMUTATION_MSAACOUNT", 1, 2, 4, 8);
-	using FPermutationDomain = TShaderPermutationDomain<FGroupSize, FTile, FPPLL, FMSAACount>;
+	using FPermutationDomain = TShaderPermutationDomain<FGroupSize, FPPLL, FMSAACount>;
 
 	BEGIN_SHADER_PARAMETER_STRUCT(FParameters, )
 		SHADER_PARAMETER(FIntPoint, OutputResolution)
@@ -1730,15 +1682,14 @@ static void AddHairVisibilityPrimitiveIdCompactionPass(
 	const uint32 MSAASampleCount = GetHairVisibilityRenderMode() == HairVisibilityRenderMode_MSAA_Visibility ? GetMaxSamplePerPixel() : 1;
 	const uint32 PPLLMaxRenderNodePerPixel = GetMaxSamplePerPixel();
 	const uint32 MaxRenderNodeCount = GetTotalSampleCountForAllocation(EffectiveResolution);
-	const bool bUseTile = TileData.IsValid();
+	check(TileData.IsValid());
 
 	OutCompactNodeVis = GraphBuilder.CreateBuffer(FRDGBufferDesc::CreateStructuredDesc(sizeof(HairStrandsVisibilityInternal::NodeVis), MaxRenderNodeCount), TEXT("Hair.VisibilityNodeVis"));
 	// Pixel coord of the node. Stored as 2*R16_UINT
 	OutCompactNodeCoord = GraphBuilder.CreateBuffer(FRDGBufferDesc::CreateBufferDesc(sizeof(uint32), MaxRenderNodeCount), TEXT("Hair.VisibilityNodeCoord"));
 
 	FHairVisibilityPrimitiveIdCompactionCS::FPermutationDomain PermutationVector;
-	PermutationVector.Set<FHairVisibilityPrimitiveIdCompactionCS::FGroupSize>(bUseTile ? FHairStrandsTiles::GroupSize : GetVendorOptimalGroupSize1D());
-	PermutationVector.Set<FHairVisibilityPrimitiveIdCompactionCS::FTile>(bUseTile);
+	PermutationVector.Set<FHairVisibilityPrimitiveIdCompactionCS::FGroupSize>(FHairStrandsTiles::GroupSize);
 	PermutationVector.Set<FHairVisibilityPrimitiveIdCompactionCS::FPPLL>(bUsePPLL ? PPLLMaxRenderNodePerPixel : 0);
 	PermutationVector.Set<FHairVisibilityPrimitiveIdCompactionCS::FMSAACount>(MSAASampleCount);
 	PermutationVector = FHairVisibilityPrimitiveIdCompactionCS::RemapPermutation(PermutationVector);
@@ -1759,38 +1710,20 @@ static void AddHairVisibilityPrimitiveIdCompactionPass(
 	PassParameters->OutCoverageTexture = GraphBuilder.CreateUAV(OutCoverageTexture);
 
 	const FHairStrandsTiles::ETileType TileType = FHairStrandsTiles::ETileType::HairAll;
-	if (bUseTile)
-	{
-		PassParameters->TileCountXY			= TileData.TileCountXY;
-		PassParameters->TileSize			= FHairStrandsTiles::TileSize;
-		PassParameters->TileCountBuffer		= GraphBuilder.CreateSRV(TileData.TileCountBuffer, PF_R32_UINT);
-		PassParameters->TileDataBuffer		= TileData.GetTileBufferSRV(TileType);
-		PassParameters->IndirectBufferArgs	= TileData.TileIndirectDispatchBuffer;
-	}
+	PassParameters->TileCountXY			= TileData.TileCountXY;
+	PassParameters->TileSize			= FHairStrandsTiles::TileSize;
+	PassParameters->TileCountBuffer		= GraphBuilder.CreateSRV(TileData.TileCountBuffer, PF_R32_UINT);
+	PassParameters->TileDataBuffer		= TileData.GetTileBufferSRV(TileType);
+	PassParameters->IndirectBufferArgs	= TileData.TileIndirectDispatchBuffer;
  
 	TShaderMapRef<FHairVisibilityPrimitiveIdCompactionCS> ComputeShader(View.ShaderMap, PermutationVector);
-	if (bUseTile)
-	{
-		FComputeShaderUtils::AddPass(
-			GraphBuilder,
-			RDG_EVENT_NAME("HairStrands::VisibilityCompaction(Tile)"),
-			ComputeShader,
-			PassParameters,
-			TileData.TileIndirectDispatchBuffer,
-			FHairStrandsTiles::GetIndirectDispatchArgOffset(TileType));
-	}
-	else
-	{
-		const FIntPoint GroupSize = GetVendorOptimalGroupSize2D();
-		const FIntPoint RectResolution(View.ViewRect.Width(), View.ViewRect.Height());
-
-		FComputeShaderUtils::AddPass(
-			GraphBuilder,
-			RDG_EVENT_NAME("HairStrands::VisibilityCompaction(Screen)"),
-			ComputeShader,
-			PassParameters,
-			FComputeShaderUtils::GetGroupCount(RectResolution, GroupSize));
-	}
+	FComputeShaderUtils::AddPass(
+		GraphBuilder,
+		RDG_EVENT_NAME("HairStrands::VisibilityCompaction"),
+		ComputeShader,
+		PassParameters,
+		TileData.TileIndirectDispatchBuffer,
+		FHairStrandsTiles::GetIndirectDispatchArgOffset(TileType));
 
 	OutIndirectArgsBuffer = AddCopyIndirectArgPass(GraphBuilder, &View, NodeGroupSize, 1, OutCompactCounter);
 	OutMaxRenderNodeCount = MaxRenderNodeCount;
@@ -1803,8 +1736,7 @@ class FHairVisibilityCompactionComputeRasterCS : public FGlobalShader
 	SHADER_USE_PARAMETER_STRUCT(FHairVisibilityCompactionComputeRasterCS, FGlobalShader);
 
 	class FGroupSize : SHADER_PERMUTATION_SPARSE_INT("PERMUTATION_GROUPSIZE", 32, 64);
-	class FTile : SHADER_PERMUTATION_BOOL("PERMUTATION_TILE");
-	using FPermutationDomain = TShaderPermutationDomain<FGroupSize, FTile>;
+	using FPermutationDomain = TShaderPermutationDomain<FGroupSize>;
 
 	BEGIN_SHADER_PARAMETER_STRUCT(FParameters, )
 		SHADER_PARAMETER(FIntPoint, OutputResolution)
@@ -1886,7 +1818,7 @@ static void AddHairVisibilityCompactionComputeRasterPass(
 	AddClearUAVPass(GraphBuilder, GraphBuilder.CreateUAV(OutCoverageTexture), 0.f);
 
 	// Select render node count according to current mode
-	const bool bUseTile = TileData.IsValid();
+	check(TileData.IsValid());
 	const FHairStrandsTiles::ETileType TileType = FHairStrandsTiles::ETileType::HairAll;
 	const uint32 MSAASampleCount = GetHairVisibilityRenderMode() == HairVisibilityRenderMode_MSAA_Visibility ? GetMaxSamplePerPixel() : 1;
 	const uint32 PPLLMaxRenderNodePerPixel = GetMaxSamplePerPixel();
@@ -1912,41 +1844,23 @@ static void AddHairVisibilityCompactionComputeRasterPass(
 	PassParameters->OutCompactNodeVis		= GraphBuilder.CreateUAV(OutCompactNodeVis);
 	PassParameters->OutCompactNodeCoord		= GraphBuilder.CreateUAV(OutCompactNodeCoord, FHairStrandsVisibilityData::NodeCoordFormat);
 	PassParameters->OutCoverageTexture		= GraphBuilder.CreateUAV(OutCoverageTexture);
-
-	if (bUseTile)
-	{
-		PassParameters->TileCountXY = TileData.TileCountXY;
-		PassParameters->TileSize = FHairStrandsTiles::TileSize;
-		PassParameters->TileCountBuffer = GraphBuilder.CreateSRV(TileData.TileCountBuffer, PF_R32_UINT);
-		PassParameters->TileDataBuffer = TileData.GetTileBufferSRV(TileType);
-		PassParameters->IndirectBufferArgs = TileData.TileIndirectDispatchBuffer;
-	}
+	PassParameters->TileCountXY = TileData.TileCountXY;
+	PassParameters->TileSize = FHairStrandsTiles::TileSize;
+	PassParameters->TileCountBuffer = GraphBuilder.CreateSRV(TileData.TileCountBuffer, PF_R32_UINT);
+	PassParameters->TileDataBuffer = TileData.GetTileBufferSRV(TileType);
+	PassParameters->IndirectBufferArgs = TileData.TileIndirectDispatchBuffer;
 
 	const FIntPoint GroupSize = GetVendorOptimalGroupSize2D();
 	FHairVisibilityCompactionComputeRasterCS::FPermutationDomain PermutationVector;
-	PermutationVector.Set<FHairVisibilityCompactionComputeRasterCS::FGroupSize>(bUseTile ? FHairStrandsTiles::GroupSize : GetVendorOptimalGroupSize1D());
-	PermutationVector.Set<FHairVisibilityCompactionComputeRasterCS::FTile>(bUseTile);
+	PermutationVector.Set<FHairVisibilityCompactionComputeRasterCS::FGroupSize>(FHairStrandsTiles::GroupSize);
 	TShaderMapRef<FHairVisibilityCompactionComputeRasterCS> ComputeShader(View.ShaderMap, PermutationVector);
-
-	if (bUseTile)
-	{
-		FComputeShaderUtils::AddPass(
-			GraphBuilder,
-			RDG_EVENT_NAME("HairStrands::VisibilityCompaction(Tile)"),
-			ComputeShader,
-			PassParameters,
-			TileData.TileIndirectDispatchBuffer,
-			FHairStrandsTiles::GetIndirectDispatchArgOffset(TileType));
-	}
-	else
-	{
-		FComputeShaderUtils::AddPass(
-			GraphBuilder,
-			RDG_EVENT_NAME("HairStrands::VisibilityCompaction(Screen)"),
-			ComputeShader,
-			PassParameters,
-			FComputeShaderUtils::GetGroupCount(Resolution, GroupSize));
-	}
+	FComputeShaderUtils::AddPass(
+		GraphBuilder,
+		RDG_EVENT_NAME("HairStrands::VisibilityCompaction"),
+		ComputeShader,
+		PassParameters,
+		TileData.TileIndirectDispatchBuffer,
+		FHairStrandsTiles::GetIndirectDispatchArgOffset(TileType));
 
 	OutIndirectArgsBuffer = AddCopyIndirectArgPass(GraphBuilder, &View, NodeGroupSize, 1, OutCompactCounter);
 	OutMaxRenderNodeCount = MaxRenderNodeCount;
@@ -1958,8 +1872,7 @@ class FHairVisibilityFillOpaqueDepthPS : public FGlobalShader
 	DECLARE_GLOBAL_SHADER(FHairVisibilityFillOpaqueDepthPS);
 	SHADER_USE_PARAMETER_STRUCT(FHairVisibilityFillOpaqueDepthPS, FGlobalShader);
 
-	class FTile : SHADER_PERMUTATION_BOOL("PERMUTATION_TILE");
-	using FPermutationDomain = TShaderPermutationDomain<FTile>;
+	using FPermutationDomain = TShaderPermutationDomain<>;
 
 	BEGIN_SHADER_PARAMETER_STRUCT(FParameters, )
 		SHADER_PARAMETER_RDG_TEXTURE(Texture2D, SceneDepthTexture)
@@ -1985,8 +1898,8 @@ static FRDGTextureRef AddHairVisibilityFillOpaqueDepth(
 	const FRDGTextureRef& SceneDepthTexture)
 {
 	check(GetHairVisibilityRenderMode() == HairVisibilityRenderMode_MSAA_Visibility || GetHairVisibilityRenderMode() == HairVisibilityRenderMode_ComputeRaster);
+	check(TileData.IsValid());
 
-	const bool bUseTile = TileData.IsValid();
 	const FHairStrandsTiles::ETileType TileType = FHairStrandsTiles::ETileType::HairAll;
 	FRDGTextureRef OutVisibilityDepthTexture = GraphBuilder.CreateTexture(FRDGTextureDesc::Create2D(Resolution, PF_D24, FClearValueBinding::DepthFar, TexCreate_DepthStencilTargetable | TexCreate_ShaderResource, 1, GetMaxSamplePerPixel()), TEXT("Hair.VisibilityDepthTexture"));
 
@@ -2000,69 +1913,20 @@ static FRDGTextureRef AddHairVisibilityFillOpaqueDepth(
 		ERenderTargetLoadAction::ENoAction,
 		FExclusiveDepthStencil::DepthWrite_StencilNop);
 
-	FHairVisibilityFillOpaqueDepthPS::FPermutationDomain PermutationVector;
-	PermutationVector.Set<FHairVisibilityFillOpaqueDepthPS::FTile>(bUseTile);
-	TShaderMapRef<FHairVisibilityFillOpaqueDepthPS> PixelShader(View.ShaderMap, PermutationVector);
+	TShaderMapRef<FHairVisibilityFillOpaqueDepthPS> PixelShader(View.ShaderMap);
 	
 	const FIntRect Viewport = View.ViewRect;
-	if (bUseTile)
-	{
-		TShaderMapRef<FHairStrandsTilePassVS> TileVertexShader(View.ShaderMap);
-		//ClearUnusedGraphResources(PixelShader, Parameters);
-
-		GraphBuilder.AddPass(
-			RDG_EVENT_NAME("HairStrands::FillVisibilityDepth(Tile)"),
-			Parameters,
-			ERDGPassFlags::Raster,
-			[Parameters, TileVertexShader, PixelShader, Viewport, TileType](FRHICommandList& RHICmdList)
-			{
-				FHairStrandsTilePassVS::FParameters ParametersVS = Parameters->TileData;
-
-				FGraphicsPipelineStateInitializer GraphicsPSOInit;
-				RHICmdList.ApplyCachedRenderTargets(GraphicsPSOInit);
-				GraphicsPSOInit.BlendState = TStaticBlendState<CW_RGBA, BO_Add, BF_One, BF_Zero, BO_Add, BF_One, BF_Zero>::GetRHI();
-				GraphicsPSOInit.RasterizerState = TStaticRasterizerState<>::GetRHI();
-				GraphicsPSOInit.DepthStencilState = TStaticDepthStencilState<true, CF_DepthNearOrEqual>::GetRHI();
-
-				GraphicsPSOInit.BoundShaderState.VertexDeclarationRHI = GFilterVertexDeclaration.VertexDeclarationRHI;
-				GraphicsPSOInit.BoundShaderState.VertexShaderRHI = TileVertexShader.GetVertexShader();
-				GraphicsPSOInit.BoundShaderState.PixelShaderRHI = PixelShader.GetPixelShader();
-				GraphicsPSOInit.PrimitiveType = Parameters->TileData.bRectPrimitive > 0 ? PT_RectList : PT_TriangleList;
-				SetGraphicsPipelineState(RHICmdList, GraphicsPSOInit, 0);
-
-				SetShaderParameters(RHICmdList, PixelShader, PixelShader.GetPixelShader(), *Parameters);
-
-				RHICmdList.SetViewport(Viewport.Min.X, Viewport.Min.Y, 0.0f, Viewport.Max.X, Viewport.Max.Y, 1.0f);
-				SetShaderParameters(RHICmdList, TileVertexShader, TileVertexShader.GetVertexShader(), ParametersVS);
-				RHICmdList.SetStreamSource(0, nullptr, 0);
-				RHICmdList.DrawPrimitiveIndirect(Parameters->TileData.TileIndirectBuffer->GetRHI(), FHairStrandsTiles::GetIndirectDrawArgOffset(TileType));
-			});
+	TShaderMapRef<FHairStrandsTilePassVS> TileVertexShader(View.ShaderMap);
+	//ClearUnusedGraphResources(PixelShader, Parameters);
 
-	}
-	else
-	{
-		TShaderMapRef<FPostProcessVS> VertexShader(View.ShaderMap);
-		TArray<FIntRect> MacroGroupRects;
-		if (IsHairStrandsViewRectOptimEnable())
-		{
-			for (const FHairStrandsMacroGroupData& MacroGroupData : MacroGroupDatas)
-			{
-				MacroGroupRects.Add(MacroGroupData.ScreenRect);
-			}
-		}
-		else
+	GraphBuilder.AddPass(
+		RDG_EVENT_NAME("HairStrands::FillVisibilityDepth(Tile)"),
+		Parameters,
+		ERDGPassFlags::Raster,
+		[Parameters, TileVertexShader, PixelShader, Viewport, TileType](FRHICommandList& RHICmdList)
 		{
-			MacroGroupRects.Add(Viewport);
-		}
-
-		ClearUnusedGraphResources(PixelShader, Parameters);
+			FHairStrandsTilePassVS::FParameters ParametersVS = Parameters->TileData;
 
-		GraphBuilder.AddPass(
-			RDG_EVENT_NAME("HairStrands::FillVisibilityDepth(View)"),
-			Parameters,
-			ERDGPassFlags::Raster,
-			[Parameters, VertexShader, PixelShader, MacroGroupRects, Viewport, Resolution](FRHICommandList& RHICmdList)
-		{
 			FGraphicsPipelineStateInitializer GraphicsPSOInit;
 			RHICmdList.ApplyCachedRenderTargets(GraphicsPSOInit);
 			GraphicsPSOInit.BlendState = TStaticBlendState<CW_RGBA, BO_Add, BF_One, BF_Zero, BO_Add, BF_One, BF_Zero>::GetRHI();
@@ -2070,29 +1934,18 @@ static FRDGTextureRef AddHairVisibilityFillOpaqueDepth(
 			GraphicsPSOInit.DepthStencilState = TStaticDepthStencilState<true, CF_DepthNearOrEqual>::GetRHI();
 
 			GraphicsPSOInit.BoundShaderState.VertexDeclarationRHI = GFilterVertexDeclaration.VertexDeclarationRHI;
-			GraphicsPSOInit.BoundShaderState.VertexShaderRHI = VertexShader.GetVertexShader();
+			GraphicsPSOInit.BoundShaderState.VertexShaderRHI = TileVertexShader.GetVertexShader();
 			GraphicsPSOInit.BoundShaderState.PixelShaderRHI = PixelShader.GetPixelShader();
-			GraphicsPSOInit.PrimitiveType = PT_TriangleList;
+			GraphicsPSOInit.PrimitiveType = Parameters->TileData.bRectPrimitive > 0 ? PT_RectList : PT_TriangleList;
 			SetGraphicsPipelineState(RHICmdList, GraphicsPSOInit, 0);
 
 			SetShaderParameters(RHICmdList, PixelShader, PixelShader.GetPixelShader(), *Parameters);
 
-			for (const FIntRect& ViewRect : MacroGroupRects)
-			{
-				RHICmdList.SetViewport(ViewRect.Min.X, ViewRect.Min.Y, 0.0f, ViewRect.Max.X, ViewRect.Max.Y, 1.0f);
-				DrawRectangle(
-					RHICmdList,
-					0, 0,
-					Viewport.Width(), Viewport.Height(),
-					Viewport.Min.X,   Viewport.Min.Y,
-					Viewport.Width(), Viewport.Height(),
-					Viewport.Size(),
-					Resolution,
-					VertexShader,
-					EDRF_UseTriangleOptimization);
-			}
+			RHICmdList.SetViewport(Viewport.Min.X, Viewport.Min.Y, 0.0f, Viewport.Max.X, Viewport.Max.Y, 1.0f);
+			SetShaderParameters(RHICmdList, TileVertexShader, TileVertexShader.GetVertexShader(), ParametersVS);
+			RHICmdList.SetStreamSource(0, nullptr, 0);
+			RHICmdList.DrawPrimitiveIndirect(Parameters->TileData.TileIndirectBuffer->GetRHI(), FHairStrandsTiles::GetIndirectDrawArgOffset(TileType));
 		});
-	}
 
 	// Ensure HTile is valid after manually feeding the scene depth value
 	if (GRHISupportsResummarizeHTile)
@@ -2534,7 +2387,6 @@ static void AddHairAuxilaryPass(
 		default: 													OutputType = 0; Method = TEXT("Unknown"); break;
 	};
 
-	TShaderMapRef<FPostProcessVS> ScreenVertexShader(View.ShaderMap);
 	TShaderMapRef<FHairStrandsTilePassVS> TileVertexShader(View.ShaderMap);
 
 	FHairVisibilityDepthPS::FPermutationDomain PermutationVector;
@@ -2542,13 +2394,12 @@ static void AddHairAuxilaryPass(
 	TShaderMapRef<FHairVisibilityDepthPS> PixelShader(View.ShaderMap, PermutationVector);
 	const FIntRect Viewport = View.ViewRect;
 	const FIntPoint Resolution = OutDepthTexture->Desc.Extent; //-V522
-	const bool bUseTile = TileData.IsValid();
 
 	GraphBuilder.AddPass(
-		RDG_EVENT_NAME("HairStrands::AuxilaryPass(%s)(%s)", Method, bUseTile ? TEXT("Tile") : TEXT("Screen")),
+		RDG_EVENT_NAME("HairStrands::AuxilaryPass(%s)", Method),
 		Parameters,
 		ERDGPassFlags::Raster,
-		[Parameters, ScreenVertexShader, TileVertexShader, PixelShader, Viewport, Resolution, bUseTile, bDepthTested, TileType](FRHICommandList& RHICmdList)
+		[Parameters, TileVertexShader, PixelShader, Viewport, Resolution, bDepthTested, TileType](FRHICommandList& RHICmdList)
 	{
 		FHairStrandsTilePassVS::FParameters ParametersVS = Parameters->TileData;
 
@@ -2561,7 +2412,7 @@ static void AddHairAuxilaryPass(
 			TStaticDepthStencilState<false, CF_Always>::GetRHI();
 
 		GraphicsPSOInit.BoundShaderState.VertexDeclarationRHI = GFilterVertexDeclaration.VertexDeclarationRHI;
-		GraphicsPSOInit.BoundShaderState.VertexShaderRHI = bUseTile ? TileVertexShader.GetVertexShader() : ScreenVertexShader.GetVertexShader();
+		GraphicsPSOInit.BoundShaderState.VertexShaderRHI = TileVertexShader.GetVertexShader();
 		GraphicsPSOInit.BoundShaderState.PixelShaderRHI = PixelShader.GetPixelShader();
 		GraphicsPSOInit.PrimitiveType = Parameters->TileData.bRectPrimitive > 0 ? PT_RectList : PT_TriangleList;
 		SetGraphicsPipelineState(RHICmdList, GraphicsPSOInit, 0);
@@ -2569,25 +2420,9 @@ static void AddHairAuxilaryPass(
 		SetShaderParameters(RHICmdList, PixelShader, PixelShader.GetPixelShader(), *Parameters);
 
 		RHICmdList.SetViewport(Viewport.Min.X, Viewport.Min.Y, 0.0f, Viewport.Max.X, Viewport.Max.Y, 1.0f);
-		if (bUseTile)
-		{
-			SetShaderParameters(RHICmdList, TileVertexShader, TileVertexShader.GetVertexShader(), ParametersVS);
-			RHICmdList.SetStreamSource(0, nullptr, 0);
-			RHICmdList.DrawPrimitiveIndirect(Parameters->TileData.TileIndirectBuffer->GetRHI(), FHairStrandsTiles::GetIndirectDrawArgOffset(TileType));
-		}
-		else
-		{
-			DrawRectangle(
-				RHICmdList,
-				0, 0,
-				Viewport.Width(), Viewport.Height(),
-				Viewport.Min.X, Viewport.Min.Y,
-				Viewport.Width(), Viewport.Height(),
-				Viewport.Size(),
-				Resolution,
-				ScreenVertexShader,
-				EDRF_UseTriangleOptimization);
-		}
+		SetShaderParameters(RHICmdList, TileVertexShader, TileVertexShader.GetVertexShader(), ParametersVS);
+		RHICmdList.SetStreamSource(0, nullptr, 0);
+		RHICmdList.DrawPrimitiveIndirect(Parameters->TileData.TileIndirectBuffer->GetRHI(), FHairStrandsTiles::GetIndirectDrawArgOffset(TileType));
 	});
 }
 
@@ -3664,6 +3499,9 @@ void DrawHitProxies(
 	}
 	else if (RenderMode == HairVisibilityRenderMode_MSAA_Visibility)
 	{
+		// Generate Tile data
+		TileData = AddHairStrandsGenerateTilesPass(GraphBuilder, View, Resolution);
+
 		FRDGTextureRef VisDepthTexture = AddHairVisibilityFillOpaqueDepth(
 			GraphBuilder,
 			View,
@@ -3672,10 +3510,6 @@ void DrawHitProxies(
 			TileData,
 			SceneDepthTexture);
 
-		// Generate Tile data
-		check(VisDepthTexture);
-		TileData = AddHairStrandsGenerateTilesPass(GraphBuilder, View, VisDepthTexture->Desc.Extent);
-
 		FRDGTextureRef VisIdTexture = nullptr;
 		AddHairVisibilityMSAAPass(
 			GraphBuilder,
