The change proably does not contribute to performance because it's maked as FORCEINLINE, but it could be good to do regardless for runs where FORCEINLINE is disabled.

PR #9073: Optimization: mesh draw command sorting wasn't using references (Contributed by [REDACTED_PERSON])

diff --git a/Engine/Source/Runtime/Renderer/Public/MeshPassProcessor.h b/Engine/Source/Runtime/Renderer/Public/MeshPassProcessor.h
--- a/Engine/Source/Runtime/Renderer/Public/MeshPassProcessor.h
+++ b/Engine/Source/Runtime/Renderer/Public/MeshPassProcessor.h
@@ -1416,7 +1416,7 @@ public:
 
 struct FCompareFMeshDrawCommands
 {
-	FORCEINLINE bool operator() (FVisibleMeshDrawCommand A, FVisibleMeshDrawCommand B) const
+	FORCEINLINE bool operator() (const FVisibleMeshDrawCommand& A, const FVisibleMeshDrawCommand& B) const
 	{
 		// First order by a sort key.
 		if (A.SortKey != B.SortKey)
