Change light function pixel shader to bind hair strands depth texture rather than hair strands uniform buffer.

This solves the hair strands uniform buffer is not created when hair strands plugin is not enabled.

diff --git a/Engine/Shaders/Private/LightFunctionPixelShader.usf b/Engine/Shaders/Private/LightFunctionPixelShader.usf
--- a/Engine/Shaders/Private/LightFunctionPixelShader.usf
+++ b/Engine/Shaders/Private/LightFunctionPixelShader.usf
@@ -15,6 +15,8 @@ float4x4 SvPositionToLight;
 /** Fade distance in x, disabled brightness in y, output for preview shadows mask in z, hair strands in w */
 float4 LightFunctionParameters2;
 
+Texture2D HairOnlyDepthTexture;
+
 void Main(
 	// in float4 ScreenPosition : TEXCOORD0, // Not used and not output in VS
 	in float4 SvPosition : SV_Position,		// after all interpolators
@@ -31,7 +33,7 @@ void Main(
 	BRANCH
 	if (LightFunctionParameters2.w > 0)
 	{
-		SvPosition.z = Texture2DSampleLevel(HairStrands.HairOnlyDepthTexture, SceneTexturesStruct_SceneDepthTextureSampler, ScreenUV, 0).r;
+		SvPosition.z = Texture2DSampleLevel(HairOnlyDepthTexture, SceneTexturesStruct_SceneDepthTextureSampler, ScreenUV, 0).r; 
 	}
 #endif
 
diff --git a/Engine/Source/Runtime/Renderer/Private/LightFunctionRendering.cpp b/Engine/Source/Runtime/Renderer/Private/LightFunctionRendering.cpp
--- a/Engine/Source/Runtime/Renderer/Private/LightFunctionRendering.cpp
+++ b/Engine/Source/Runtime/Renderer/Private/LightFunctionRendering.cpp
@@ -104,9 +104,10 @@ public:
 		SvPositionToLight.Bind(Initializer.ParameterMap,TEXT("SvPositionToLight"));
 		LightFunctionParameters.Bind(Initializer.ParameterMap);
 		LightFunctionParameters2.Bind(Initializer.ParameterMap,TEXT("LightFunctionParameters2"));
+		HairOnlyDepthTexture.Bind(Initializer.ParameterMap, TEXT("HairOnlyDepthTexture"));
 	}
 
-	void SetParameters(FRHICommandList& RHICmdList, const FViewInfo& View, const FLightSceneInfo* LightSceneInfo, const FMaterialRenderProxy* MaterialProxy, const FMaterial& Material, bool bRenderingPreviewShadowIndicator, float ShadowFadeFraction, bool bUseHairStrands, FRHIUniformBuffer* HairStrandsUniformBuffer)
+	void SetParameters(FRHICommandList& RHICmdList, const FViewInfo& View, const FLightSceneInfo* LightSceneInfo, const FMaterialRenderProxy* MaterialProxy, const FMaterial& Material, bool bRenderingPreviewShadowIndicator, float ShadowFadeFraction, bool bUseHairStrands, FRHITexture* InHairOnlyDepthTexture)
 	{
 		FRHIPixelShader* ShaderRHI = RHICmdList.GetBoundPixelShader();
 
@@ -151,11 +152,9 @@ public:
 			bRenderingPreviewShadowIndicator ? 1.0f : 0.0f,
 			bUseHairStrands ? 1.0f : 0.0f));
 
-		
-		auto HairStrandsParameter = GetUniformBufferParameter<FHairStrandsViewUniformParameters>();
-		if (HairStrandsParameter.IsBound() && HairStrandsUniformBuffer)
+		if (HairOnlyDepthTexture.IsBound() && InHairOnlyDepthTexture)
 		{
-			SetUniformBufferParameter(RHICmdList, ShaderRHI, HairStrandsParameter, HairStrandsUniformBuffer);
+			SetTextureParameter(RHICmdList, ShaderRHI, HairOnlyDepthTexture, InHairOnlyDepthTexture);
 		}
 
 		auto DeferredLightParameter = GetUniformBufferParameter<FDeferredLightUniformStruct>();
@@ -169,6 +168,7 @@ private:
 	LAYOUT_FIELD(FShaderParameter, SvPositionToLight);
 	LAYOUT_FIELD(FLightFunctionSharedParameters, LightFunctionParameters);
 	LAYOUT_FIELD(FShaderParameter, LightFunctionParameters2);
+	LAYOUT_FIELD(FShaderResourceParameter, HairOnlyDepthTexture);
 };
 
 IMPLEMENT_MATERIAL_SHADER_TYPE(,FLightFunctionPS,TEXT("/Engine/Private/LightFunctionPixelShader.usf"),TEXT("Main"),SF_Pixel);
@@ -275,7 +275,7 @@ bool FDeferredShadingSceneRenderer::RenderPreviewShadowsIndicator(
 
 BEGIN_SHADER_PARAMETER_STRUCT(FRenderLightFunctionForMaterialParameters, )
 	SHADER_PARAMETER_RDG_UNIFORM_BUFFER(FSceneTextureUniformParameters, SceneTextures)
-	SHADER_PARAMETER_RDG_UNIFORM_BUFFER(FHairStrandsViewUniformParameters, HairStrands)
+	SHADER_PARAMETER_RDG_TEXTURE(Texture2D, HairOnlyDepthTexture)
 	RENDER_TARGET_BINDING_SLOTS()
 END_SHADER_PARAMETER_STRUCT()
 
@@ -335,7 +335,7 @@ bool FDeferredShadingSceneRenderer::RenderLightFunctionForMaterial(
 		{
 			FRenderLightFunctionForMaterialParameters* PassParameters = GraphBuilder.AllocParameters<FRenderLightFunctionForMaterialParameters>();
 			PassParameters->SceneTextures = SceneTextures.UniformBuffer;
-			PassParameters->HairStrands = HairStrands::BindHairStrandsViewUniformParameters(View);
+			PassParameters->HairOnlyDepthTexture = (View.HairStrandsViewData.bIsValid && View.HairStrandsViewData.VisibilityData.HairOnlyDepthTexture) ? View.HairStrandsViewData.VisibilityData.HairOnlyDepthTexture : GSystemTextures.GetDepthDummy(GraphBuilder);
 			PassParameters->RenderTargets[0] = FRenderTargetBinding(ScreenShadowMaskTexture, bLightAttenuationCleared ? ERenderTargetLoadAction::ELoad : ERenderTargetLoadAction::ENoAction);
 			PassParameters->RenderTargets.DepthStencil = FDepthStencilBinding(SceneTextures.Depth.Target, ERenderTargetLoadAction::ELoad, ERenderTargetLoadAction::ELoad, FExclusiveDepthStencil::DepthRead_StencilWrite);
 
@@ -428,12 +428,10 @@ bool FDeferredShadingSceneRenderer::RenderLightFunctionForMaterial(
 						LightSceneProxy->SetScissorRect(RHICmdList, View, TotalRect);
 					}
 
-					FRHIUniformBuffer* HairStrandsUniformBuffer = PassParameters->HairStrands.GetUniformBuffer() ? PassParameters->HairStrands.GetUniformBuffer()->GetRHI() : nullptr;
-
 					// Render a bounding light sphere.
 					SetGraphicsPipelineState(RHICmdList, GraphicsPSOInit, 0);
 					VertexShader->SetParameters(RHICmdList, View, LightSceneInfo);
-					PixelShader->SetParameters(RHICmdList, View, LightSceneInfo, MaterialProxyForRendering, *MaterialForRendering, bRenderingPreviewShadowsIndicator, FadeAlpha, bUseHairStrands, HairStrandsUniformBuffer);
+					PixelShader->SetParameters(RHICmdList, View, LightSceneInfo, MaterialProxyForRendering, *MaterialForRendering, bRenderingPreviewShadowsIndicator, FadeAlpha, bUseHairStrands, PassParameters->HairOnlyDepthTexture->GetRHI());
 
 					// Project the light function using a sphere around the light
 					//@todo - could use a cone for spotlights
