- Add missing stats to VT tasks

#rnx

diff --git a/Engine/Source/Runtime/Renderer/Private/VT/VirtualTextureSystem.cpp b/Engine/Source/Runtime/Renderer/Private/VT/VirtualTextureSystem.cpp
--- a/Engine/Source/Runtime/Renderer/Private/VT/VirtualTextureSystem.cpp
+++ b/Engine/Source/Runtime/Renderer/Private/VT/VirtualTextureSystem.cpp
@@ -210,7 +210,7 @@ public:
 
 	static ESubsequentsMode::Type GetSubsequentsMode() { return ESubsequentsMode::TrackSubsequents; }
 	ENamedThreads::Type GetDesiredThread() { return ENamedThreads::AnyNormalThreadNormalTask; }
-	FORCEINLINE TStatId GetStatId() const { return TStatId(); }
+	FORCEINLINE TStatId GetStatId() const { RETURN_QUICK_DECLARE_CYCLE_STAT(FFeedbackAnalysisTask, STATGROUP_VirtualTexturing); }
 };
 
 struct FAddRequestedTilesParameters
@@ -243,7 +243,7 @@ public:
 
 	static ESubsequentsMode::Type GetSubsequentsMode() { return ESubsequentsMode::TrackSubsequents; }
 	ENamedThreads::Type GetDesiredThread() { return ENamedThreads::AnyNormalThreadNormalTask; }
-	FORCEINLINE TStatId GetStatId() const { return TStatId(); }
+	FORCEINLINE TStatId GetStatId() const { RETURN_QUICK_DECLARE_CYCLE_STAT(FAddRequestedTilesTask, STATGROUP_VirtualTexturing); }
 };
 
 struct FGatherRequestsParameters
@@ -274,7 +274,7 @@ public:
 
 	static ESubsequentsMode::Type GetSubsequentsMode() { return ESubsequentsMode::TrackSubsequents; }
 	ENamedThreads::Type GetDesiredThread() { return ENamedThreads::AnyNormalThreadNormalTask; }
-	FORCEINLINE TStatId GetStatId() const { return TStatId(); }
+	FORCEINLINE TStatId GetStatId() const { RETURN_QUICK_DECLARE_CYCLE_STAT(FGatherRequestsTask, STATGROUP_VirtualTexturing); }
 };
 
 static FVirtualTextureSystem* GVirtualTextureSystem = nullptr;
