Sphere reflection capture:  Added logic to automatically adjust cube map resolution based on what will fit in GPU memory without producing an OOM, in WITH_EDITOR builds.  In non-editor builds, the same logic estimates if an OOM will occur, to give the user more detailed information about what may have caused the OOM.

The specific case we are trying to handle is that D3D12 fails to allocate a texture resource if its size equals or exceeds 4 GB, returning an OOM error, even if plenty of GPU memory is available.  The goal is to degrade quality gracefully in the editor rather than crashing, and produce an actionable message otherwise.

#rnx

diff --git a/Engine/Source/Runtime/Renderer/Private/ReflectionEnvironmentCapture.cpp b/Engine/Source/Runtime/Renderer/Private/ReflectionEnvironmentCapture.cpp
--- a/Engine/Source/Runtime/Renderer/Private/ReflectionEnvironmentCapture.cpp
+++ b/Engine/Source/Runtime/Renderer/Private/ReflectionEnvironmentCapture.cpp
@@ -696,36 +696,37 @@ void EndReflectionCaptureSlowTask(int32 NumCaptures)
 	}
 }
 
-int32 GetReflectionCaptureSizeForArrayCount(int32 InRequestedCaptureSize, int32 InRequestedMaxCubeMaps)
+static int32 GetMaxReflectionCapturesAtSize(int32 CaptureSize)
 {
-	int32 OutCaptureSize = InRequestedCaptureSize;
-#if WITH_EDITOR
-	if(GIsEditor)
-	{
-		FTextureMemoryStats TextureMemStats;
-		RHIGetTextureMemoryStats(TextureMemStats);
-		
-		SIZE_T TexMemRequired = CalcTextureSize(OutCaptureSize, OutCaptureSize, PF_FloatRGBA, FMath::CeilLogTwo(OutCaptureSize) + 1) * CubeFace_MAX * InRequestedMaxCubeMaps;
-		// Assumption: Texture arrays prefer to be contiguous in memory but not always
-		// Single large cube array allocations can fail on low end systems even if we try to fit it in total avail video and/or avail system memory
-		
-		// Attempt to limit the resource size to within percentage (3/4) of total video memory to give consistant stable results
-		const SIZE_T MaxResourceVideoMemoryFootprint = ((SIZE_T)TextureMemStats.DedicatedVideoMemory * (SIZE_T)3) / (SIZE_T)4;
+	FRHITextureDesc RHITexDesc(
+		ETextureDimension::TextureCube,
+		ETextureCreateFlags::None,
+		PF_FloatRGBA,
+		FClearValueBinding::Black,
+		FIntPoint(CaptureSize, CaptureSize),
+		1,										// Depth
+		1,										// ArraySize
+		FMath::CeilLogTwo(CaptureSize) + 1,		// InNumMips
+		1,										// Samples
+		0);
+
+	SIZE_T TexMemRequiredPerCapture = RHICalcTexturePlatformSize(RHITexDesc).Size;
+
+	// Attempt to limit the resource size to within percentage (3/4) of total video memory to give consistent stable results.
+	// Also limit max size to 4 GB (technically 4 GB minus one), as D3D12 allocation fails for individual resources 4 GB or over,
+	// and we'll assume the same is true for other platforms.
+	FTextureMemoryStats TextureMemStats;
+	RHIGetTextureMemoryStats(TextureMemStats);
 		
-		// Bottom out at 128 as that is the default for CVarReflectionCaptureSize
-        while(TexMemRequired > MaxResourceVideoMemoryFootprint && OutCaptureSize > 128)
-        {
-			OutCaptureSize = FMath::RoundUpToPowerOfTwo(OutCaptureSize) >> 1;
-			TexMemRequired = CalcTextureSize(OutCaptureSize, OutCaptureSize, PF_FloatRGBA, FMath::CeilLogTwo(OutCaptureSize) + 1) * CubeFace_MAX * InRequestedMaxCubeMaps;
-        }
-        
-        if(OutCaptureSize != InRequestedCaptureSize)
-        {
-			UE_LOG(LogEngine, Error, TEXT("Requested reflection capture cube size of %d with %d elements results in too large a resource for host machine, limiting reflection capture cube size to %d"), InRequestedCaptureSize, InRequestedMaxCubeMaps, OutCaptureSize);
-        }
+	SIZE_T DedicatedVideoMemoryLimit = ((SIZE_T)TextureMemStats.DedicatedVideoMemory * (SIZE_T)3) / (SIZE_T)4;
+	if (!DedicatedVideoMemoryLimit)
+	{
+		DedicatedVideoMemoryLimit = SIZE_T_MAX;
 	}
-#endif // WITH_EDITOR
-	return OutCaptureSize;
+
+	const SIZE_T MaxResourceVideoMemoryFootprint = FMath::Min(DedicatedVideoMemoryLimit, (SIZE_T)UINT_MAX);
+
+	return MaxResourceVideoMemoryFootprint / TexMemRequiredPerCapture;
 }
 
 int32 NumUniqueReflectionCaptures(const TSparseArray<UReflectionCaptureComponent*>& CaptureComponents)
@@ -773,11 +774,41 @@ void FScene::AllocateReflectionCaptures(const TArray<UReflectionCaptureComponent
 				}
 			}
 
-			// Request the exact amount needed by default
+			// Request the exact amount and size needed by default
 			int32 DesiredMaxCubemaps = NumUniqueReflectionCaptures(ReflectionSceneData.AllocatedReflectionCapturesGameThread);
-			const float MaxCubemapsRoundUpBase = 1.5f;
+			int32 DesiredCaptureSize = UReflectionCaptureComponent::GetReflectionCaptureSize();
+			int32 ReflectionCaptureSize = DesiredCaptureSize;
+
+			DesiredMaxCubemaps = FMath::Min(DesiredMaxCubemaps, PlatformMaxNumReflectionCaptures);
+
+#if WITH_EDITOR
+			// Reduce the capture size (resolution) until we can fit all the captures we need in a single texture cube array resource
+			int32 MaxCapturesAtSize;
+			for (MaxCapturesAtSize = GetMaxReflectionCapturesAtSize(ReflectionCaptureSize);
+				 MaxCapturesAtSize < DesiredMaxCubemaps;
+				 MaxCapturesAtSize = GetMaxReflectionCapturesAtSize(ReflectionCaptureSize))
+			{
+				ReflectionCaptureSize >>= 1;
+			}
+
+			if (ReflectionCaptureSize != DesiredCaptureSize)
+			{
+				UE_LOG(LogEngine, Error, TEXT("Requested reflection capture cube size of %d with %d elements results in too large a resource for host machine, limiting reflection capture cube size to %d"), DesiredCaptureSize, DesiredMaxCubemaps, ReflectionCaptureSize);
+			}
+#else
+			// When not in editor, let the code proceed with the desired size and number, but warn the user an OOM failure is likely and why
+			int32 MaxCapturesAtSize = GetMaxReflectionCapturesAtSize(ReflectionCaptureSize);
+			if (MaxCapturesAtSize < DesiredMaxCubemaps)
+			{
+				UE_LOG(LogEngine, Error, TEXT("Reflection capture of size %d with %d elements exceeds estimated GPU memory limit of %d elements, OOM likely"), DesiredCaptureSize, DesiredMaxCubemaps, MaxCapturesAtSize);
+				
+				// We expect an OOM, but set the limit to the exact number desired to give the best chance of it succeeding anyway
+				MaxCapturesAtSize = DesiredMaxCubemaps;
+			}
+#endif
 
 			// If this is not the first time the scene has allocated the cubemap array, include slack to reduce reallocations
+			const float MaxCubemapsRoundUpBase = 1.5f;
 			if (ReflectionSceneData.MaxAllocatedReflectionCubemapsGameThread > 0)
 			{
 				float Exponent = FMath::LogX(MaxCubemapsRoundUpBase, DesiredMaxCubemaps);
@@ -786,9 +817,9 @@ void FScene::AllocateReflectionCaptures(const TArray<UReflectionCaptureComponent
 				DesiredMaxCubemaps = FMath::Pow(MaxCubemapsRoundUpBase, FMath::TruncToInt(Exponent) + 1);
 			}
 
-			DesiredMaxCubemaps = FMath::Min(DesiredMaxCubemaps, PlatformMaxNumReflectionCaptures);
+			// After slack calculation, need to clamp again at our hard limits.
+			DesiredMaxCubemaps = FMath::Min3(DesiredMaxCubemaps, PlatformMaxNumReflectionCaptures, MaxCapturesAtSize);
 
-			const int32 ReflectionCaptureSize = GetReflectionCaptureSizeForArrayCount(UReflectionCaptureComponent::GetReflectionCaptureSize(), DesiredMaxCubemaps);
 			bool bNeedsUpdateAllCaptures = ReflectionSceneData.DoesAllocatedDataNeedUpdate(DesiredMaxCubemaps, ReflectionCaptureSize);
 
 			if (bNeedsUpdateAllCaptures)
