Keep depth content if TAA is enabled.
It's because the bKeepDepthContent is false if CameraCut is used and the SceneDepthTexture could be a dummy black texture(RGBA8 format) and a RDG validation failed at creating StencilTexture SRV which is expecting PF_DepthStencil format.

diff --git a/Engine/Source/Runtime/Renderer/Private/MobileShadingRenderer.cpp b/Engine/Source/Runtime/Renderer/Private/MobileShadingRenderer.cpp
--- a/Engine/Source/Runtime/Renderer/Private/MobileShadingRenderer.cpp
+++ b/Engine/Source/Runtime/Renderer/Private/MobileShadingRenderer.cpp
@@ -501,7 +501,7 @@ void FMobileSceneRenderer::InitViews(FRDGBuilder& GraphBuilder, FSceneTexturesCo
 		Views[0].bIsReflectionCapture ||
 		(bDeferredShading && bPostProcessUsesSceneDepth) ||
 		(bDeferredShading && bSceneDepthCapture) ||
-		bShouldRenderVelocities ||
+		Views[0].AntiAliasingMethod == AAM_TemporalAA ||
 		bRequireSeparateViewPass ||
 		bIsFullDepthPrepassEnabled ||
 		GraphBuilder.IsDumpingFrame();
