Build fix for missing SIZE_T_MAX

#rnx

diff --git a/Engine/Source/Runtime/Renderer/Private/ReflectionEnvironmentCapture.cpp b/Engine/Source/Runtime/Renderer/Private/ReflectionEnvironmentCapture.cpp
--- a/Engine/Source/Runtime/Renderer/Private/ReflectionEnvironmentCapture.cpp
+++ b/Engine/Source/Runtime/Renderer/Private/ReflectionEnvironmentCapture.cpp
@@ -721,7 +721,7 @@ static int32 GetMaxReflectionCapturesAtSize(int32 CaptureSize)
 	SIZE_T DedicatedVideoMemoryLimit = ((SIZE_T)TextureMemStats.DedicatedVideoMemory * (SIZE_T)3) / (SIZE_T)4;
 	if (!DedicatedVideoMemoryLimit)
 	{
-		DedicatedVideoMemoryLimit = SIZE_T_MAX;
+		DedicatedVideoMemoryLimit = (SIZE_T)UINT_MAX;
 	}
 
 	const SIZE_T MaxResourceVideoMemoryFootprint = FMath::Min(DedicatedVideoMemoryLimit, (SIZE_T)UINT_MAX);
