Use DDPI option to use the compressed transforms with scenedata. For platforms not supporting nanite, using non compressed transforms saves 0.4ms on the GPU.

diff --git a/Engine/Config/VulkanPC/DataDrivenPlatformInfo.ini b/Engine/Config/VulkanPC/DataDrivenPlatformInfo.ini
--- a/Engine/Config/VulkanPC/DataDrivenPlatformInfo.ini
+++ b/Engine/Config/VulkanPC/DataDrivenPlatformInfo.ini
@@ -67,6 +67,7 @@ bSupportsUInt64ImageAtomics = true
 bRequiresVendorExtensionsForAtomics = false
 
 bSupportsNanite = true
+bSupportsSceneDataCompressedTransforms = true
 bSupportsLumenGI = true
 
 ; //#todo-vulkanddpi
diff --git a/Engine/Config/Windows/DataDrivenPlatformInfo.ini b/Engine/Config/Windows/DataDrivenPlatformInfo.ini
--- a/Engine/Config/Windows/DataDrivenPlatformInfo.ini
+++ b/Engine/Config/Windows/DataDrivenPlatformInfo.ini
@@ -71,6 +71,7 @@ bSupportsHighEndRayTracingReflections=true
 bSupportsWaveOperations=RuntimeDependent
 bSupportsGPUScene=true
 bSupportsNanite=false
+bSupportsSceneDataCompressedTransforms=false
 bSupportsComputeFramework=true
 bSupportsByteBufferComputeShaders=true
 bSupportsLumenGI=true
@@ -117,6 +118,7 @@ bSupportsPathTracing=true
 bSupportsHighEndRayTracingReflections=true
 bSupportsGPUScene=true
 bSupportsNanite=true
+bSupportsSceneDataCompressedTransforms=true
 bSupportsComputeFramework=true
 bSupportsByteBufferComputeShaders=true
 bSupportsLumenGI=true
diff --git a/Engine/Plugins/Experimental/GPULightmass/Source/GPULightmass/Private/LightmapRenderer.cpp b/Engine/Plugins/Experimental/GPULightmass/Source/GPULightmass/Private/LightmapRenderer.cpp
--- a/Engine/Plugins/Experimental/GPULightmass/Source/GPULightmass/Private/LightmapRenderer.cpp
+++ b/Engine/Plugins/Experimental/GPULightmass/Source/GPULightmass/Private/LightmapRenderer.cpp
@@ -561,8 +561,8 @@ void FCachedRayTracingSceneData::SetupViewUniformBufferFromSceneRenderState(FSce
 			TRACE_CPUPROFILER_EVENT_SCOPE(InstanceSceneData);
 
 			TResourceArray<FVector4f> InstanceSceneDataSOA;
-			InstanceSceneDataSOA.AddZeroed(FInstanceSceneShaderData::DataStrideInFloat4s * InstanceSceneData.Num());
-			for (int32 ArrayIndex = 0; ArrayIndex < FInstanceSceneShaderData::DataStrideInFloat4s; ArrayIndex++)
+			InstanceSceneDataSOA.AddZeroed(FInstanceSceneShaderData::GetDataStrideInFloat4s() * InstanceSceneData.Num());
+			for (int32 ArrayIndex = 0; ArrayIndex < FInstanceSceneShaderData::GetDataStrideInFloat4s(); ArrayIndex++)
 			{
 				for (int32 DataIndex = 0; DataIndex < InstanceSceneData.Num(); DataIndex++)
 				{
@@ -573,7 +573,7 @@ void FCachedRayTracingSceneData::SetupViewUniformBufferFromSceneRenderState(FSce
 			FRHIResourceCreateInfo CreateInfo(TEXT("InstanceSceneDataBuffer"), &InstanceSceneDataSOA);
 			if (InstanceSceneDataSOA.GetResourceDataSize() == 0)
 			{
-				InstanceSceneDataSOA.AddZeroed(FInstanceSceneShaderData::DataStrideInFloat4s);
+				InstanceSceneDataSOA.AddZeroed(FInstanceSceneShaderData::GetDataStrideInFloat4s());
 			}
 
 			InstanceSceneDataBufferRHI = RHICreateStructuredBuffer(sizeof(FVector4f), InstanceSceneDataSOA.GetResourceDataSize(), BUF_Static | BUF_ShaderResource, CreateInfo);
diff --git a/Engine/Shaders/Private/GPUScene/GPUSceneWriter.ush b/Engine/Shaders/Private/GPUScene/GPUSceneWriter.ush
--- a/Engine/Shaders/Private/GPUScene/GPUSceneWriter.ush
+++ b/Engine/Shaders/Private/GPUScene/GPUSceneWriter.ush
@@ -76,7 +76,7 @@ void WriteInstancePrimitiveId(uint InstanceId, uint NewPrimitiveId)
  **/
 void WriteInstanceLocalToRelativeWorld(uint InstanceId, float4x4 LocalToRelativeWorld, bool StoreDeterminantSign)
 {
-#if INSTANCE_SCENE_DATA_COMPRESSED_TRANSFORMS
+#if PLATFORM_ALLOW_SCENE_DATA_COMPRESSED_TRANSFORMS
 	uint4 RotationScale;
 	float3 Translation;
 	EncodeTransform(LocalToRelativeWorld, RotationScale, Translation);
@@ -210,7 +210,7 @@ void WriteInstanceDynamicData(FInstancePayloadDataOffsets PayloadOffsets, float4
 		return;
 	}
 
-#if INSTANCE_SCENE_DATA_COMPRESSED_TRANSFORMS
+#if PLATFORM_ALLOW_SCENE_DATA_COMPRESSED_TRANSFORMS
 	uint4 RotationScale;
 	float3 Translation;
 	EncodeTransform(PrevLocalToRelativeWorld, RotationScale, Translation);
diff --git a/Engine/Shaders/Private/SceneData.ush b/Engine/Shaders/Private/SceneData.ush
--- a/Engine/Shaders/Private/SceneData.ush
+++ b/Engine/Shaders/Private/SceneData.ush
@@ -31,10 +31,6 @@
 
 #define ENABLE_PER_INSTANCE_CUSTOM_DATA (USES_PER_INSTANCE_CUSTOM_DATA || VF_REQUIRES_PER_INSTANCE_CUSTOM_DATA)
 
-// When transforms come from the InstanceSceneData buffer, indicates whether or not the transforms are compressed (ignored on mobile)
-// TODO: Global setting/define
-#define INSTANCE_SCENE_DATA_COMPRESSED_TRANSFORMS 1
-
 // Whether to fetch primitive values (eg LocalToWorld) by dynamically indexing a scene-wide buffer, or to reference a single Primitive uniform buffer
 #if VF_SUPPORTS_PRIMITIVE_SCENE_DATA
 	#if FEATURE_LEVEL == FEATURE_LEVEL_ES3_1
@@ -99,7 +95,7 @@
 #define INSTANCE_SCENE_DATA_FLAG_HAS_EDITOR_DATA			0x80
 #define INSTANCE_SCENE_DATA_FLAG_IS_RAYTRACING_FAR_FIELD	0x100
 
-#if INSTANCE_SCENE_DATA_COMPRESSED_TRANSFORMS
+#if PLATFORM_ALLOW_SCENE_DATA_COMPRESSED_TRANSFORMS
 	static const uint InstanceTransformSizeFloat4Count = 2u; // compressed transform
 #else
 	static const uint InstanceTransformSizeFloat4Count = 3u; // encoded scale/rotation (uint4) and translation (float3)
@@ -823,7 +819,7 @@ void ComputeInstanceDerivedData(inout FInstanceSceneData InstanceData, float3 Ti
 	// Do not put any load operations here!
 	//
 
-#if (VF_USE_PRIMITIVE_SCENE_DATA == 2) || !INSTANCE_SCENE_DATA_COMPRESSED_TRANSFORMS
+#if (VF_USE_PRIMITIVE_SCENE_DATA == 2) || !PLATFORM_ALLOW_SCENE_DATA_COMPRESSED_TRANSFORMS
 	// Non-uniform scale must be computed from the transform because it was not already computed when decoding it (see below in GetInstanceSceneData)
 	float3 Scale2;
 	Scale2.x = length2(LocalToRelativeWorld[0].xyz);
@@ -869,7 +865,7 @@ FInstanceSceneData GetInstanceSceneData(uint InstanceId, uint SOAStride, bool bC
 	// LocalBounds Center		float0.yzw
 	// LocalBounds Extent		float1.xyz
 	// __UNUSED					float1.w
-#if INSTANCE_SCENE_DATA_COMPRESSED_TRANSFORMS
+#if PLATFORM_ALLOW_SCENE_DATA_COMPRESSED_TRANSFORMS
 	// Previous Transform[0]	float2.xyzw
 	// Previous Transform[1]	float3.xyzw
 	// LM/SM Scale Bias			float4.xyzw
@@ -901,7 +897,7 @@ FInstanceSceneData GetInstanceSceneData(uint InstanceId, uint SOAStride, bool bC
 
 		float3 TilePosition 									= GetPrimitiveData(InstanceData.PrimitiveId).TilePosition;
 
-#if INSTANCE_SCENE_DATA_COMPRESSED_TRANSFORMS
+#if PLATFORM_ALLOW_SCENE_DATA_COMPRESSED_TRANSFORMS
 		uint4  RotationScale									= asuint(LoadInstanceSceneDataElement(1 * SOAStride + InstanceId));
 		float3 Translation										= LoadInstanceSceneDataElement(2 * SOAStride + InstanceId).xyz;
 		float3 Scale = 0;
@@ -971,7 +967,7 @@ FInstanceSceneData GetInstanceSceneData(uint InstanceId, uint SOAStride, bool bC
 		BRANCH
 		if (Offsets.DynamicData != INVALID_INSTANCE_PAYLOAD_OFFSET)
 		{
-		#if INSTANCE_SCENE_DATA_COMPRESSED_TRANSFORMS
+		#if PLATFORM_ALLOW_SCENE_DATA_COMPRESSED_TRANSFORMS
 			uint4  PrevRotationScale			= asuint(LoadInstancePayloadDataElement(Offsets.DynamicData + 0));
 			float3 PrevTranslation				= LoadInstancePayloadDataElement(Offsets.DynamicData + 1).xyz;
 			float3 PrevScale					= 0;
@@ -986,21 +982,6 @@ FInstanceSceneData GetInstanceSceneData(uint InstanceId, uint SOAStride, bool bC
 			float3 TilePosition 				= GetPrimitiveData(InstanceData.PrimitiveId).TilePosition;
 			InstanceData.PrevLocalToWorld		= MakeLWCMatrix(TilePosition, PrevLocalToRelativeWorld);
 		}
-		else
-		{
-		#if INSTANCE_SCENE_DATA_COMPRESSED_TRANSFORMS
-			// TODO: Temporary PrevVelocityHack
-			uint4  PrevRotationScale			= asuint(LoadInstanceSceneDataElement(3 * SOAStride + InstanceId));
-			float3 PrevTranslation				= LoadInstanceSceneDataElement(4 * SOAStride + InstanceId).xyz;
-			float3 PrevScale = 0;
-			float4x4 PrevLocalToRelativeWorld	= DecodeTransform(PrevRotationScale, PrevTranslation, PrevScale);
-		#else
-			float4x4 PrevLocalToRelativeWorld	= InstanceData.LocalToWorld;
-		#endif
-
-			float3 TilePosition 				= GetPrimitiveData(InstanceData.PrimitiveId).TilePosition;
-			InstanceData.PrevLocalToWorld		= MakeLWCMatrix(TilePosition, PrevLocalToRelativeWorld);
-		}
 
 	#if 1 //NEEDS_LIGHTMAP_COORDINATE
 		BRANCH
diff --git a/Engine/Source/Runtime/Engine/Private/InstanceUniformShaderParameters.cpp b/Engine/Source/Runtime/Engine/Private/InstanceUniformShaderParameters.cpp
--- a/Engine/Source/Runtime/Engine/Private/InstanceUniformShaderParameters.cpp
+++ b/Engine/Source/Runtime/Engine/Private/InstanceUniformShaderParameters.cpp
@@ -77,16 +77,19 @@ void FInstanceSceneShaderData::BuildInternal
 	Data[0].W  = *(const float*)&RandomID;
 
 	// TODO: Temporary PrevVelocityHack
-#if INSTANCE_SCENE_DATA_COMPRESSED_TRANSFORMS
-	FCompressedTransform CompressedLocalToWorld(LocalToWorld);
-	Data[1] = *(const FVector4f*)&CompressedLocalToWorld.Rotation[0];
-	Data[2] = *(const FVector3f*)&CompressedLocalToWorld.Translation;
+	if (FDataDrivenShaderPlatformInfo::GetSupportSceneDataCompressedTransforms(GMaxRHIShaderPlatform))
+	{
+		FCompressedTransform CompressedLocalToWorld(LocalToWorld);
+		Data[1] = *(const FVector4f*)&CompressedLocalToWorld.Rotation[0];
+		Data[2] = *(const FVector3f*)&CompressedLocalToWorld.Translation;
 
-	FCompressedTransform CompressedPrevLocalToWorld(PrevLocalToWorld);
-	Data[3] = *(const FVector4f*)&CompressedPrevLocalToWorld.Rotation[0];
-	Data[4] = *(const FVector3f*)&CompressedPrevLocalToWorld.Translation;
-#else
-	// Note: writes 3x float4s
-	LocalToWorld.To3x4MatrixTranspose((float*)&Data[1]);
-#endif
+		FCompressedTransform CompressedPrevLocalToWorld(PrevLocalToWorld);
+		Data[3] = *(const FVector4f*)&CompressedPrevLocalToWorld.Rotation[0];
+		Data[4] = *(const FVector3f*)&CompressedPrevLocalToWorld.Translation;
+	}
+	else
+	{
+		// Note: writes 3x float4s
+		LocalToWorld.To3x4MatrixTranspose((float*)&Data[1]);
+	}
 }
diff --git a/Engine/Source/Runtime/Engine/Private/PrimitiveSceneProxy.cpp b/Engine/Source/Runtime/Engine/Private/PrimitiveSceneProxy.cpp
--- a/Engine/Source/Runtime/Engine/Private/PrimitiveSceneProxy.cpp
+++ b/Engine/Source/Runtime/Engine/Private/PrimitiveSceneProxy.cpp
@@ -546,11 +546,14 @@ uint32 FPrimitiveSceneProxy::GetPayloadDataStride() const
 	uint32 PayloadDataCount = 0;
 
 	// Random ID is packed into scene data currently
-#if INSTANCE_SCENE_DATA_COMPRESSED_TRANSFORMS
-	PayloadDataCount += HasPerInstanceDynamicData() ? 2 : 0;	// Compressed transform
-#else
-	PayloadDataCount += HasPerInstanceDynamicData() ? 3 : 0;	// FRenderTransform
-#endif
+	if (FDataDrivenShaderPlatformInfo::GetSupportSceneDataCompressedTransforms(GMaxRHIShaderPlatform))
+	{
+		PayloadDataCount += HasPerInstanceDynamicData() ? 2 : 0;	// Compressed transform
+	}
+	else
+	{
+		PayloadDataCount += HasPerInstanceDynamicData() ? 3 : 0;	// FRenderTransform
+	}
 		
 	// Hierarchy is packed in with local bounds if they are both present (almost always the case)
 	if (HasPerInstanceLocalBounds())
diff --git a/Engine/Source/Runtime/Engine/Private/PrimitiveUniformShaderParameters.cpp b/Engine/Source/Runtime/Engine/Private/PrimitiveUniformShaderParameters.cpp
--- a/Engine/Source/Runtime/Engine/Private/PrimitiveUniformShaderParameters.cpp
+++ b/Engine/Source/Runtime/Engine/Private/PrimitiveUniformShaderParameters.cpp
@@ -35,7 +35,7 @@ void FSinglePrimitiveStructured::InitRHI()
 		LightmapSceneDataBufferSRV = RHICreateShaderResourceView(LightmapSceneDataBufferRHI);
 
 		CreateInfo.DebugName = TEXT("InstanceSceneDataBuffer");
-		InstanceSceneDataBufferRHI = RHICreateStructuredBuffer(sizeof(FVector4f), FInstanceSceneShaderData::DataStrideInFloat4s * sizeof(FVector4f), BUF_Static | BUF_ShaderResource, CreateInfo);
+		InstanceSceneDataBufferRHI = RHICreateStructuredBuffer(sizeof(FVector4f), FInstanceSceneShaderData::GetDataStrideInFloat4s() * sizeof(FVector4f), BUF_Static | BUF_ShaderResource, CreateInfo);
 		InstanceSceneDataBufferSRV = RHICreateShaderResourceView(InstanceSceneDataBufferRHI);
 
 		CreateInfo.DebugName = TEXT("InstancePayloadDataBuffer");
@@ -64,8 +64,8 @@ void FSinglePrimitiveStructured::UploadToGPU()
 		FPlatformMemory::Memcpy(LockedData, LightmapSceneData.Data.GetData(), FLightmapSceneShaderData::DataStrideInFloat4s * sizeof(FVector4f));
 		RHIUnlockBuffer(LightmapSceneDataBufferRHI);
 
-		LockedData = RHILockBuffer(InstanceSceneDataBufferRHI, 0, FInstanceSceneShaderData::DataStrideInFloat4s * sizeof(FVector4f), RLM_WriteOnly);
-		FPlatformMemory::Memcpy(LockedData, InstanceSceneData.Data.GetData(), FInstanceSceneShaderData::DataStrideInFloat4s * sizeof(FVector4f));
+		LockedData = RHILockBuffer(InstanceSceneDataBufferRHI, 0, FInstanceSceneShaderData::GetDataStrideInFloat4s() * sizeof(FVector4f), RLM_WriteOnly);
+		FPlatformMemory::Memcpy(LockedData, InstanceSceneData.Data.GetData(), FInstanceSceneShaderData::GetDataStrideInFloat4s() * sizeof(FVector4f));
 		RHIUnlockBuffer(InstanceSceneDataBufferRHI);
 
 		LockedData = RHILockBuffer(InstancePayloadDataBufferRHI, 0, 1 /* unused dummy */ * sizeof(FVector4f), RLM_WriteOnly);
diff --git a/Engine/Source/Runtime/Engine/Private/ShaderCompiler/ShaderCompiler.cpp b/Engine/Source/Runtime/Engine/Private/ShaderCompiler/ShaderCompiler.cpp
--- a/Engine/Source/Runtime/Engine/Private/ShaderCompiler/ShaderCompiler.cpp
+++ b/Engine/Source/Runtime/Engine/Private/ShaderCompiler/ShaderCompiler.cpp
@@ -6052,6 +6052,7 @@ void GlobalBeginCompileShader(
 	Input.Environment.SetDefine(TEXT("PLATFORM_SUPPORTS_MESH_SHADERS_TIER0"), RHISupportsMeshShadersTier0(EShaderPlatform(Target.Platform)) ? 1 : 0);
 	Input.Environment.SetDefine(TEXT("PLATFORM_SUPPORTS_MESH_SHADERS_TIER1"), RHISupportsMeshShadersTier1(EShaderPlatform(Target.Platform)) ? 1 : 0);
 	Input.Environment.SetDefine(TEXT("PLATFORM_SUPPORTS_BINDLESS"), RHISupportsBindless(EShaderPlatform(Target.Platform)) ? 1 : 0);
+	Input.Environment.SetDefine(TEXT("PLATFORM_ALLOW_SCENE_DATA_COMPRESSED_TRANSFORMS"), FDataDrivenShaderPlatformInfo::GetSupportSceneDataCompressedTransforms(EShaderPlatform(Target.Platform)) ? 1 : 0);
 
 	if (RHISupportsBindless(EShaderPlatform(Target.Platform)))
 	{
diff --git a/Engine/Source/Runtime/Engine/Private/ShaderCompiler/ShaderGenerationUtil.cpp b/Engine/Source/Runtime/Engine/Private/ShaderCompiler/ShaderGenerationUtil.cpp
--- a/Engine/Source/Runtime/Engine/Private/ShaderCompiler/ShaderGenerationUtil.cpp
+++ b/Engine/Source/Runtime/Engine/Private/ShaderCompiler/ShaderGenerationUtil.cpp
@@ -98,6 +98,7 @@ void FShaderCompileUtilities::ApplyFetchEnvironment(FShaderGlobalDefines& SrcDef
 	FETCH_COMPILE_BOOL(PLATFORM_SUPPORTS_PER_PIXEL_DBUFFER_MASK);
 	FETCH_COMPILE_BOOL(PLATFORM_SUPPORTS_DISTANCE_FIELDS);
 	FETCH_COMPILE_BOOL(COMPILE_SHADERS_FOR_DEVELOPMENT_ALLOWED);
+	FETCH_COMPILE_BOOL(PLATFORM_ALLOW_SCENE_DATA_COMPRESSED_TRANSFORMS);
 
 	// note that we are doing an if so that if we call ApplyFetchEnvironment() twice, we get the logical OR of bSupportsDualBlending support
 	if (RHISupportsDualSourceBlending(Platform))
@@ -388,6 +389,7 @@ static FShaderGlobalDefines FetchShaderGlobalDefines(EShaderPlatform TargetPlatf
 	Ret.PLATFORM_SUPPORTS_RENDERTARGET_WRITE_MASK = RHISupportsRenderTargetWriteMask(EShaderPlatform(TargetPlatform)) ? 1 : 0;
 	Ret.PLATFORM_SUPPORTS_PER_PIXEL_DBUFFER_MASK = FDataDrivenShaderPlatformInfo::GetSupportsPerPixelDBufferMask(EShaderPlatform(TargetPlatform)) ? 1 : 0;
 	Ret.PLATFORM_SUPPORTS_DISTANCE_FIELDS = DoesPlatformSupportDistanceFields(EShaderPlatform(TargetPlatform)) ? 1 : 0;
+	Ret.PLATFORM_ALLOW_SCENE_DATA_COMPRESSED_TRANSFORMS = FDataDrivenShaderPlatformInfo::GetSupportSceneDataCompressedTransforms(EShaderPlatform(TargetPlatform)) ? 1 : 0;
 
 	Ret.bSupportsDualBlending = RHISupportsDualSourceBlending(TargetPlatform);
 
diff --git a/Engine/Source/Runtime/Engine/Public/InstanceUniformShaderParameters.h b/Engine/Source/Runtime/Engine/Public/InstanceUniformShaderParameters.h
--- a/Engine/Source/Runtime/Engine/Public/InstanceUniformShaderParameters.h
+++ b/Engine/Source/Runtime/Engine/Public/InstanceUniformShaderParameters.h
@@ -94,15 +94,31 @@ struct FPrimitiveInstanceDynamicData
 
 struct FInstanceSceneShaderData
 {
+private:
 	// Must match GetInstanceSceneData() in SceneData.ush
-#if INSTANCE_SCENE_DATA_COMPRESSED_TRANSFORMS
-	// Compressed transform
-	enum { DataStrideInFloat4s = 5 }; // TODO: Temporary PrevVelocityHack
-#else
-	enum { DataStrideInFloat4s = 4 };
-#endif
+	// Allocate the max Float4s usage when compressed transform is used.
+	// TODO: Temporary PrevVelocityHack (last float4s when compressed)
+	static constexpr uint32 CompressedTransformDataStrideInFloat4s = 5;
+	static constexpr uint32 UnCompressedTransformDataStrideInFloat4s = 4;
 
-	TStaticArray<FVector4f, DataStrideInFloat4s> Data;
+public:
+
+	static uint32 GetDataStrideInFloat4s()
+	{
+		if (FDataDrivenShaderPlatformInfo::GetSupportSceneDataCompressedTransforms(GMaxRHIShaderPlatform))
+		{
+			return CompressedTransformDataStrideInFloat4s;
+		}
+		else
+		{
+			return UnCompressedTransformDataStrideInFloat4s;
+		}
+	}
+
+	static uint32 GetEffectiveNumBytes()
+	{
+		return (GetDataStrideInFloat4s() * sizeof(FVector4f));
+	}
 
 	FInstanceSceneShaderData() : Data(InPlace, NoInit)
 	{
@@ -142,4 +158,6 @@ struct FInstanceSceneShaderData
 		const FRenderTransform& LocalToWorld,
 		const FRenderTransform& PrevLocalToWorld // Assumes shear has been removed already // TODO: Temporary PrevVelocityHack
 	);
+
+	TStaticArray<FVector4f, CompressedTransformDataStrideInFloat4s> Data;
 };
diff --git a/Engine/Source/Runtime/RHI/Private/RHI.cpp b/Engine/Source/Runtime/RHI/Private/RHI.cpp
--- a/Engine/Source/Runtime/RHI/Private/RHI.cpp
+++ b/Engine/Source/Runtime/RHI/Private/RHI.cpp
@@ -2547,6 +2547,7 @@ void FGenericDataDrivenShaderPlatformInfo::ParseDataDrivenShaderInfo(const FConf
 	GET_SECTION_BOOL_HELPER(bSupportsOIT);
 	GET_SECTION_SUPPORT_HELPER(bSupportsRealTypes);
 	GET_SECTION_INT_HELPER(EnablesHLSL2021ByDefault);
+	GET_SECTION_BOOL_HELPER(bSupportsSceneDataCompressedTransforms);
 #undef GET_SECTION_BOOL_HELPER
 #undef GET_SECTION_INT_HELPER
 #undef GET_SECTION_SUPPORT_HELPER
diff --git a/Engine/Source/Runtime/RHI/Public/RHIDefinitions.h b/Engine/Source/Runtime/RHI/Public/RHIDefinitions.h
--- a/Engine/Source/Runtime/RHI/Public/RHIDefinitions.h
+++ b/Engine/Source/Runtime/RHI/Public/RHIDefinitions.h
@@ -418,6 +418,7 @@ class RHI_API FGenericDataDrivenShaderPlatformInfo
 	uint32 bSupportsOIT : 1;
 	uint32 bSupportsRealTypes : int32(ERHIFeatureSupport::NumBits);
 	uint32 EnablesHLSL2021ByDefault : 2; // 0: disabled, 1: global shaders only, 2: all shaders
+	uint32 bSupportsSceneDataCompressedTransforms : 1;
 		
 #if WITH_EDITOR
 	FText FriendlyName;
@@ -619,6 +620,12 @@ public:
 		return Infos[Platform].bSupportsRenderTargetWriteMask;
 	}
 
+	static FORCEINLINE_DEBUGGABLE const bool GetSupportSceneDataCompressedTransforms(const FStaticShaderPlatform Platform)
+	{
+		check(IsValid(Platform));
+		return Infos[Platform].bSupportsSceneDataCompressedTransforms;
+	}
+	
 	static FORCEINLINE_DEBUGGABLE const bool GetSupportsRayTracing(const FStaticShaderPlatform Platform)
 	{
 		check(IsValid(Platform));
diff --git a/Engine/Source/Runtime/RenderCore/Private/Shader.cpp b/Engine/Source/Runtime/RenderCore/Private/Shader.cpp
--- a/Engine/Source/Runtime/RenderCore/Private/Shader.cpp
+++ b/Engine/Source/Runtime/RenderCore/Private/Shader.cpp
@@ -1888,6 +1888,11 @@ void ShaderMapAppendKeyString(EShaderPlatform Platform, FString& KeyString)
 		}
 	}
 
+	if (FDataDrivenShaderPlatformInfo::GetSupportSceneDataCompressedTransforms(Platform))
+	{
+		KeyString += TEXT("_sdct");
+	}
+
 	{
 		static const auto CVar = IConsoleManager::Get().FindTConsoleVariableDataInt(TEXT("r.GBufferDiffuseSampleOcclusion"));
 		if (CVar && CVar->GetValueOnAnyThread() != 0)
diff --git a/Engine/Source/Runtime/RenderCore/Public/ShaderMaterial.h b/Engine/Source/Runtime/RenderCore/Public/ShaderMaterial.h
--- a/Engine/Source/Runtime/RenderCore/Public/ShaderMaterial.h
+++ b/Engine/Source/Runtime/RenderCore/Public/ShaderMaterial.h
@@ -38,6 +38,7 @@ struct FShaderGlobalDefines
 	bool PLATFORM_SUPPORTS_RENDERTARGET_WRITE_MASK;
 	bool PLATFORM_SUPPORTS_PER_PIXEL_DBUFFER_MASK;
 	bool PLATFORM_SUPPORTS_DISTANCE_FIELDS;
+	bool PLATFORM_ALLOW_SCENE_DATA_COMPRESSED_TRANSFORMS;
 
 	// This is a sepcial one. The flag COMPILE_SHADERS_FOR_DEVELOPMENT is set by the cvar r.CompileShadersForDevelopment, but only
 	// if bAllowDevelopmentShaderCompile is true in the call to GlobalBeginCompileShader(). So the flag COMPILE_SHADERS_FOR_DEVELOPMENT_ALLOWED
diff --git a/Engine/Source/Runtime/Renderer/Private/GPUScene.cpp b/Engine/Source/Runtime/Renderer/Private/GPUScene.cpp
--- a/Engine/Source/Runtime/Renderer/Private/GPUScene.cpp
+++ b/Engine/Source/Runtime/Renderer/Private/GPUScene.cpp
@@ -109,8 +109,6 @@ FAutoConsoleVariableRef CVarGPUSceneInstanceUploadViaCreate(
 );
 
 
-constexpr uint32 InstanceSceneDataNumArrays = FInstanceSceneShaderData::DataStrideInFloat4s;
-
 LLM_DECLARE_TAG_API(GPUScene, RENDERER_API);
 DECLARE_LLM_MEMORY_STAT(TEXT("GPUScene"), STAT_GPUSceneLLM, STATGROUP_LLMFULL);
 DECLARE_LLM_MEMORY_STAT(TEXT("GPUScene"), STAT_GPUSceneSummaryLLM, STATGROUP_LLM);
@@ -707,8 +705,8 @@ void FGPUScene::UpdateBufferState(FRDGBuilder& GraphBuilder, FScene* Scene, cons
 
 	const uint32 InstanceSceneDataSizeReserve = FMath::RoundUpToPowerOfTwo(FMath::Max(InstanceSceneDataAllocator.GetMaxSize(), InitialBufferSize));
 	FResizeResourceSOAParams ResizeParams;
-	ResizeParams.NumBytes = InstanceSceneDataSizeReserve * sizeof(FInstanceSceneShaderData::Data);
-	ResizeParams.NumArrays = FInstanceSceneShaderData::DataStrideInFloat4s;
+	ResizeParams.NumBytes = InstanceSceneDataSizeReserve * FInstanceSceneShaderData::GetEffectiveNumBytes();
+	ResizeParams.NumArrays = FInstanceSceneShaderData::GetDataStrideInFloat4s();
 
 	BufferState.InstanceSceneDataBuffer = ResizeStructuredBufferSOAIfNeeded(GraphBuilder, InstanceSceneDataBuffer, ResizeParams, TEXT("GPUScene.InstanceSceneData"));
 	InstanceSceneDataSOAStride = InstanceSceneDataSizeReserve;
@@ -1021,7 +1019,7 @@ void FGPUScene::UploadGeneral(FRDGBuilder& GraphBuilder, FScene *Scene, FRDGExte
 				SCOPED_NAMED_EVENT(STAT_UpdateGPUSceneInstances, FColor::Green);
 				QUICK_SCOPE_CYCLE_COUNTER(STAT_UpdateGPUSceneInstances);
 
-				InstanceSceneUploadBuffer.InitPreSized(GraphBuilder, NumInstanceSceneDataUploads * InstanceSceneDataNumArrays, sizeof(FVector4f), true, TEXT("InstanceSceneUploadBuffer"));
+				InstanceSceneUploadBuffer.InitPreSized(GraphBuilder, NumInstanceSceneDataUploads * FInstanceSceneShaderData::GetDataStrideInFloat4s(), sizeof(FVector4f), true, TEXT("InstanceSceneUploadBuffer"));
 
 				if (!InstanceUpdates.UpdateBatches.IsEmpty())
 				{
@@ -1062,9 +1060,9 @@ void FGPUScene::UploadGeneral(FRDGBuilder& GraphBuilder, FScene *Scene, FRDGExte
 								);
 
 								// RefIndex* BufferState.InstanceSceneDataSOAStride + UploadInfo.InstanceSceneDataOffset + InstanceIndex
-								const uint32 UploadInstanceItemOffset = (PrimitiveItemInfo.InstanceSceneDataUploadOffset + InstanceIndex) * InstanceSceneDataNumArrays;
+								const uint32 UploadInstanceItemOffset = (PrimitiveItemInfo.InstanceSceneDataUploadOffset + InstanceIndex) * FInstanceSceneShaderData::GetDataStrideInFloat4s();
 
-								for (uint32 RefIndex = 0; RefIndex < InstanceSceneDataNumArrays; ++RefIndex)
+								for (uint32 RefIndex = 0; RefIndex < FInstanceSceneShaderData::GetDataStrideInFloat4s(); ++RefIndex)
 								{
 									FVector4f* DstVector = static_cast<FVector4f*>(InstanceSceneUploadBuffer.Set_GetRef(UploadInstanceItemOffset + RefIndex, RefIndex * BufferState.InstanceSceneDataSOAStride + UploadInfo.InstanceSceneDataOffset + InstanceIndex));
 									*DstVector = InstanceSceneData.Data[RefIndex];
@@ -1113,18 +1111,21 @@ void FGPUScene::UploadGeneral(FRDGBuilder& GraphBuilder, FScene *Scene, FRDGExte
 									{
 										check(UploadInfo.InstanceDynamicData.Num() == UploadInfo.PrimitiveInstances.Num());
 										const FRenderTransform PrevLocalToWorld = UploadInfo.InstanceDynamicData[InstanceIndex].ComputePrevLocalToWorld(UploadInfo.PrevPrimitiveToWorld);
-									#if INSTANCE_SCENE_DATA_COMPRESSED_TRANSFORMS
-										check(PayloadPosition + 1 < InstancePayloadData.Num()); // Sanity check
-										FCompressedTransform CompressedPrevLocalToWorld(PrevLocalToWorld);
-										InstancePayloadData[PayloadPosition + 0] = *(const FVector4f*)&CompressedPrevLocalToWorld.Rotation[0];
-										InstancePayloadData[PayloadPosition + 1] = *(const FVector3f*)&CompressedPrevLocalToWorld.Translation;
-										PayloadPosition += 2;
-									#else
-										// Note: writes 3x float4s
-										check(PayloadPosition + 2 < InstancePayloadData.Num()); // Sanity check
-										PrevLocalToWorld.To3x4MatrixTranspose((float*)&InstancePayloadData[PayloadPosition]);
-										PayloadPosition += 3;
-									#endif
+										if (FDataDrivenShaderPlatformInfo::GetSupportSceneDataCompressedTransforms(GMaxRHIShaderPlatform))
+										{
+											check(PayloadPosition + 1 < InstancePayloadData.Num()); // Sanity check
+											FCompressedTransform CompressedPrevLocalToWorld(PrevLocalToWorld);
+											InstancePayloadData[PayloadPosition + 0] = *(const FVector4f*)&CompressedPrevLocalToWorld.Rotation[0];
+											InstancePayloadData[PayloadPosition + 1] = *(const FVector3f*)&CompressedPrevLocalToWorld.Translation;
+											PayloadPosition += 2;
+										}
+										else
+										{
+											// Note: writes 3x float4s
+											check(PayloadPosition + 2 < InstancePayloadData.Num()); // Sanity check
+											PrevLocalToWorld.To3x4MatrixTranspose((float*)&InstancePayloadData[PayloadPosition]);
+											PayloadPosition += 3;
+										}
 									}
 
 									if (UploadInfo.InstanceFlags & INSTANCE_SCENE_DATA_FLAG_HAS_LIGHTSHADOW_UV_BIAS)
