UE5_MAIN:  compile error fix.  Cast INDEX_NONE to (uint32) when creating FLumenSceneDataKey.

#rnx

diff --git a/Engine/Source/Runtime/Renderer/Private/RendererScene.cpp b/Engine/Source/Runtime/Renderer/Private/RendererScene.cpp
--- a/Engine/Source/Runtime/Renderer/Private/RendererScene.cpp
+++ b/Engine/Source/Runtime/Renderer/Private/RendererScene.cpp
@@ -450,7 +450,7 @@ void FScene::RemoveViewState_RenderThread(FSceneViewStateInterface* ViewState)
 		}
 	}
 
-	FLumenSceneDataKey ByViewKey = { ViewState->GetViewKey(), INDEX_NONE };
+	FLumenSceneDataKey ByViewKey = { ViewState->GetViewKey(), (uint32)INDEX_NONE };
 	FLumenSceneData* const* Found = PerViewOrGPULumenSceneData.Find(ByViewKey);
 	if (Found)
 	{
@@ -844,7 +844,7 @@ void FSceneViewState::AddLumenSceneData(FSceneInterface* InScene)
 					SceneData->CopyInitialData(*Scene->DefaultLumenSceneData);
 
 					// Key shouldn't already exist in Scene, because the bLumenSceneDataAdded flag should only allow it to be added once.
-					FLumenSceneDataKey ByViewKey = { GetViewKey(), INDEX_NONE };
+					FLumenSceneDataKey ByViewKey = { GetViewKey(), (uint32)INDEX_NONE };
 					check(Scene->PerViewOrGPULumenSceneData.Find(ByViewKey) == nullptr);
 
 					Scene->PerViewOrGPULumenSceneData.Emplace(ByViewKey, SceneData);
@@ -856,7 +856,7 @@ void FSceneViewState::AddLumenSceneData(FSceneInterface* InScene)
 FLumenSceneData* FScene::FindLumenSceneData(uint32 ViewKey, uint32 GPUIndex) const
 {
 	// First search by ViewKey
-	FLumenSceneDataKey ByViewKey = { ViewKey, INDEX_NONE };
+	FLumenSceneDataKey ByViewKey = { ViewKey, (uint32)INDEX_NONE };
 	FLumenSceneData* const* Found = PerViewOrGPULumenSceneData.Find(ByViewKey);
 	if (Found)
 	{
