Enable r.RayTracing.Culling.UseMinDrawDistance by default

* Avoids HLOD meshes getting added to RT scene simultaneously with full detail meshes

diff --git a/Engine/Source/Runtime/Renderer/Private/RayTracing/RayTracingInstanceCulling.cpp b/Engine/Source/Runtime/Renderer/Private/RayTracing/RayTracingInstanceCulling.cpp
--- a/Engine/Source/Runtime/Renderer/Private/RayTracing/RayTracingInstanceCulling.cpp
+++ b/Engine/Source/Runtime/Renderer/Private/RayTracing/RayTracingInstanceCulling.cpp
@@ -35,7 +35,7 @@ static TAutoConsoleVariable<float> CVarRayTracingCullingAngle(
 
 static TAutoConsoleVariable<int32> CVarRayTracingCullingUseMinDrawDistance(
 	TEXT("r.RayTracing.Culling.UseMinDrawDistance"),
-	0,
+	1,
 	TEXT("Use min draw distance for culling"),
 	ECVF_RenderThreadSafe);
 
