New RayTracingDebug modes:
* "Dynamic Instances" - shows instances that go through dynamic path in green and all others in red
* "Proxy Type" - shows instances colored by proxy type hash
* "Picker" - allows picking instances/triangles with a mouse to show BLAS and instance info.

diff --git a/Engine/Shaders/Private/RayTracing/RayTracingCommon.ush b/Engine/Shaders/Private/RayTracing/RayTracingCommon.ush
--- a/Engine/Shaders/Private/RayTracing/RayTracingCommon.ush
+++ b/Engine/Shaders/Private/RayTracing/RayTracingCommon.ush
@@ -189,6 +189,26 @@ void SetRayTracingInstanceTransform(inout FRayTracingInstanceDescriptor Desc, fl
 	Desc.Transform = LocalToWorld;
 }
 
+uint GetRayTracingInstanceDescriptorInstanceId(FRayTracingInstanceDescriptor Desc)
+{
+	return Desc.GetInstanceId();
+}
+
+uint GetRayTracingInstanceDescriptorMask(FRayTracingInstanceDescriptor Desc)
+{
+	return Desc.GetMask();
+}
+
+uint GetRayTracingInstanceDescriptorInstanceContributionToHitGroupIndex(FRayTracingInstanceDescriptor Desc)
+{
+	return Desc.GetInstanceContributionToHitGroupIndex();
+}
+
+uint GetRayTracingInstanceDescriptorFlags(FRayTracingInstanceDescriptor Desc)
+{
+	return Desc.GetFlags();
+}
+
 FRayTracingInstanceDescriptor BuildPlatformRayTracingInstanceDesc(uint InstanceMask, uint InstanceId, uint InstanceFlags, uint InstanceContributionToHitGroupIndex, float3x4 LocalToWorld, uint2 AccelerationStructureAddress)
 {
 	FRayTracingInstanceDescriptor Desc;
diff --git a/Engine/Shaders/Private/RayTracing/RayTracingDebug.usf b/Engine/Shaders/Private/RayTracing/RayTracingDebug.usf
--- a/Engine/Shaders/Private/RayTracing/RayTracingDebug.usf
+++ b/Engine/Shaders/Private/RayTracing/RayTracingDebug.usf
@@ -10,11 +10,14 @@
 #include "RayTracingHitGroupCommon.ush"
 
 #include "/Engine/Shared/RayTracingDebugDefinitions.h"
+#include "/Engine/Shared/RayTracingDebugTypes.h"
 
 RWTexture2D<float4> Output;
 RaytracingAccelerationStructure TLAS;
 
 uint VisualizationMode;
+uint UseDebugCHS;
+uint PickerDomain;
 uint ShouldUsePreExposure;
 uint OpaqueOnly;
 float TimingScale;
@@ -23,6 +26,9 @@ float MaxTraceDistance;
 float FarFieldMaxTraceDistance;
 float3 FarFieldReferencePos;
 
+StructuredBuffer<FRayTracingInstanceDebugData> InstancesDebugData;
+StructuredBuffer<FRayTracingPickingFeedback> PickingBuffer;
+
 float4 DebugRayTracingMaterial(RayDesc Ray, bool bIsFarField)
 {
 	uint2 PixelCoord = DispatchRaysIndex().xy + View.ViewRectMin.xy;
@@ -157,8 +163,12 @@ float4 DebugRayTracingMaterial(RayDesc Ray, bool bIsFarField)
 struct FRayTracingDebugPayload : FMinimalPayload
 {
 	uint InstanceHash;
-	uint TriangleHash;
+	uint TriangleIndex;
+	uint InstanceIndex;
+	uint GeometryInstanceIndex;
+
 	float3 WorldNormal;
+	uint GeometryIndex;
 };
 
 float4 DebugRayTracingInstance(RayDesc Ray)
@@ -180,17 +190,54 @@ float4 DebugRayTracingInstance(RayDesc Ray)
 
 	if (Payload.IsHit())
 	{
+		const float3 InstanceColor = IntToColor(Payload.InstanceHash);
+		const float3 TriangleColor = IntToColor(MurmurMix(Payload.InstanceHash + Payload.TriangleIndex));
+
 		switch (VisualizationMode)
 		{
 		default:
 		case RAY_TRACING_DEBUG_VIZ_INSTANCES:
-			Result.rgb = IntToColor(Payload.InstanceHash);
+			Result.rgb = InstanceColor;
 			break;
 		case RAY_TRACING_DEBUG_VIZ_TRIANGLES:
-			Result.rgb = IntToColor(Payload.TriangleHash);
+			Result.rgb = TriangleColor;
+			break;
+#if ENABLE_INSTANCE_DEBUG_DATA
+		case RAY_TRACING_DEBUG_VIZ_DYNAMIC_INSTANCES:
+			Result.rgb = InstancesDebugData[Payload.GeometryInstanceIndex].Flags ? float3(0.25f, 0.50f, 0.25f) : float3(0.50f, 0.25f, 0.25f);
+			break;
+		case RAY_TRACING_DEBUG_VIZ_PROXY_TYPE:
+			Result.rgb = IntToColor(InstancesDebugData[Payload.GeometryInstanceIndex].ProxyHash);
 			break;
+#endif // ENABLE_INSTANCE_DEBUG_DATA
 		}
-		
+	
+#if ENABLE_PICKING
+		if (VisualizationMode == RAY_TRACING_DEBUG_VIZ_PICKER)
+		{
+			const float ColorScale = 0.1f;
+			const float BrighterColorScale = 0.2f;
+			const float3 PickedColor = float3(1.0f, 0.0f, 1.0f);
+
+			const bool bSameTriangle = (PickingBuffer[0].TriangleIndex == Payload.TriangleIndex);
+			const bool bSameSegment = (PickingBuffer[0].GeometryIndex == Payload.GeometryIndex);
+			const bool bSameInstance = (PickingBuffer[0].InstanceIndex == Payload.InstanceIndex);
+
+			switch (PickerDomain)
+			{
+			case RAY_TRACING_DEBUG_PICKER_DOMAIN_TRIANGLE:
+				Result.rgb = bSameInstance ? (bSameSegment && bSameTriangle ? PickedColor : TriangleColor * BrighterColorScale) : TriangleColor * ColorScale;
+				break;				
+			case RAY_TRACING_DEBUG_PICKER_DOMAIN_SEGMENT:
+				Result.rgb = bSameInstance ? (bSameSegment ? PickedColor : InstanceColor * BrighterColorScale) : InstanceColor * ColorScale; // TODO
+				break;
+			case RAY_TRACING_DEBUG_PICKER_DOMAIN_INSTANCE:
+				Result.rgb = bSameInstance ? PickedColor : InstanceColor * ColorScale;
+				break;
+			}
+		}
+#endif // ENABLE_PICKING
+
 		// apply some shading based on the primitive normal (map dot product result from [0,1] to [0.75,1]).
 		Result.rgb *= dot(Payload.WorldNormal, normalize(float3(1.0f, 1.0f, 1.0f))) * 0.125f + 0.875;
 	}
@@ -208,7 +255,7 @@ RAY_TRACING_ENTRY_RAYGEN(RayTracingDebugMainRGS)
 
 	float4 Result = (float4)0;
 
-	if (VisualizationMode == RAY_TRACING_DEBUG_VIZ_INSTANCES || VisualizationMode == RAY_TRACING_DEBUG_VIZ_TRIANGLES)
+	if (UseDebugCHS)
 	{
 		Result = DebugRayTracingInstance(Ray);
 	}
@@ -261,7 +308,56 @@ RAY_TRACING_ENTRY_CLOSEST_HIT(RayTracingDebugMainCHS,
 	const int ScenePrimitiveInstanceIndex = InstanceIndex() - GetBaseInstanceIndex();
 	Payload.InstanceHash = MurmurAdd(ScenePrimitiveIndex, ScenePrimitiveInstanceIndex);
 
-	Payload.TriangleHash = MurmurMix(Payload.InstanceHash + PrimitiveIndex());
+	Payload.TriangleIndex = PrimitiveIndex();
 
 	Payload.WorldNormal = ComputeWorldNormal();
+	Payload.GeometryInstanceIndex = GetHitGroupUserData();
+	Payload.InstanceIndex = InstanceIndex();
+	Payload.GeometryIndex = 0; // GeometryIndex();
 }
+
+RWStructuredBuffer<FRayTracingPickingFeedback> PickingOutput;
+StructuredBuffer<FRayTracingInstanceDescriptor> InstanceBuffer;
+
+RAY_TRACING_ENTRY_RAYGEN(RayTracingDebugPickingRGS)
+{
+	const uint2 PixelPos = (float2(View.CursorPosition) * View.ViewResolutionFraction);
+	float2 RenderTargetUV = (float2(PixelPos) + .5f) * View.BufferSizeAndInvSize.zw;
+
+	RayDesc Ray = CreatePrimaryRay(RenderTargetUV);
+
+	FRayTracingDebugPayload Payload = (FRayTracingDebugPayload)0;
+	Payload.SetMiss();
+
+	TraceRay(
+		TLAS,
+		RAY_FLAG_CULL_BACK_FACING_TRIANGLES  /*RayFlags*/,
+		0xff /*InstanceInclusionMask*/,
+		0 /*RayContributionToHitGroupIndex*/,
+		RAY_TRACING_NUM_SHADER_SLOTS /*MultiplierForGeometryContributionToShaderIndex*/,
+		0 /*MissShaderIndex*/,
+		Ray,
+		Payload);
+
+	FRayTracingPickingFeedback PickingFeedback = (FRayTracingPickingFeedback)0xFFFFFFFF;
+
+	if (Payload.IsHit())
+	{
+		FRayTracingInstanceDebugData InstanceDebugData = InstancesDebugData[Payload.GeometryInstanceIndex];
+		FRayTracingInstanceDescriptor InstanceDescriptor = InstanceBuffer[Payload.InstanceIndex];
+
+		PickingFeedback.GeometryAddress = InstanceDebugData.GeometryAddress;
+
+		PickingFeedback.InstanceIndex = Payload.InstanceIndex;
+		PickingFeedback.GeometryInstanceIndex = Payload.GeometryInstanceIndex;
+		PickingFeedback.TriangleIndex = Payload.TriangleIndex;
+		PickingFeedback.GeometryIndex = Payload.GeometryIndex;
+
+		PickingFeedback.InstanceId = GetRayTracingInstanceDescriptorInstanceId(InstanceDescriptor);
+		PickingFeedback.Flags = GetRayTracingInstanceDescriptorFlags(InstanceDescriptor);
+		PickingFeedback.Mask = GetRayTracingInstanceDescriptorMask(InstanceDescriptor);
+		PickingFeedback.InstanceContributionToHitGroupIndex = GetRayTracingInstanceDescriptorInstanceContributionToHitGroupIndex(InstanceDescriptor);
+	}
+
+	PickingOutput[0] = PickingFeedback;
+}
\ No newline at end of file
diff --git a/Engine/Shaders/Shared/RayTracingDebugTypes.h b/Engine/Shaders/Shared/RayTracingDebugTypes.h
new file mode 100644
--- /dev/null
+++ b/Engine/Shaders/Shared/RayTracingDebugTypes.h
@@ -0,0 +1,60 @@
+// Copyright Epic Games, Inc. All Rights Reserved.
+
+#pragma once
+
+/*================================================================================================
+	RayTracingDebugTypes.h: used in ray tracing shaders and C++ code to define common types
+	!!! Changing this file requires recompilation of the engine !!!
+=================================================================================================*/
+
+#ifdef __cplusplus
+	#define UINT32_TYPE uint32
+	#define UINT64_TYPE uint64
+#else
+	#define UINT32_TYPE uint
+	#define UINT64_TYPE uint64_t
+#endif
+
+struct FRayTracingInstanceDebugData
+{
+	UINT32_TYPE Flags;
+	UINT32_TYPE ProxyHash;
+	UINT64_TYPE GeometryAddress;
+
+#ifdef __cplusplus
+	FRayTracingInstanceDebugData()
+		: Flags(0)
+		, ProxyHash(0)
+		, GeometryAddress(0)
+	{}
+#endif
+};
+
+struct FRayTracingPickingFeedback
+{
+	// Hit data
+	UINT32_TYPE GeometryInstanceIndex;
+	UINT32_TYPE InstanceIndex;
+	UINT32_TYPE GeometryIndex;
+	UINT32_TYPE TriangleIndex;
+
+	// Geometry data
+	UINT64_TYPE GeometryAddress;
+
+	// Instance data
+	UINT32_TYPE InstanceId;
+	UINT32_TYPE Mask;
+	UINT32_TYPE InstanceContributionToHitGroupIndex;
+	UINT32_TYPE Flags;
+
+#ifdef __cplusplus
+	FRayTracingPickingFeedback()
+		: GeometryInstanceIndex(0xFFFFFFFF)
+		, InstanceIndex(0xFFFFFFFF)
+		, GeometryAddress(0xFFFFFFFFFFFFFFFF)
+	{}
+#endif
+};
+
+#undef UINT32_TYPE
+#undef UINT64_TYPE
\ No newline at end of file
diff --git a/Engine/Shaders/Shared/RaytracingDebugDefinitions.h b/Engine/Shaders/Shared/RaytracingDebugDefinitions.h
--- a/Engine/Shaders/Shared/RaytracingDebugDefinitions.h
+++ b/Engine/Shaders/Shared/RaytracingDebugDefinitions.h
@@ -38,3 +38,11 @@
 #define RAY_TRACING_DEBUG_VIZ_TRAVERSAL_TRIANGLE		28
 #define RAY_TRACING_DEBUG_VIZ_TRAVERSAL_ALL				29
 #define RAY_TRACING_DEBUG_VIZ_TRAVERSAL_STATISTICS		30
+#define RAY_TRACING_DEBUG_VIZ_DYNAMIC_INSTANCES			31
+#define RAY_TRACING_DEBUG_VIZ_PROXY_TYPE				32
+#define RAY_TRACING_DEBUG_VIZ_PICKER					33
+
+#define RAY_TRACING_DEBUG_PICKER_DOMAIN_TRIANGLE		0
+#define RAY_TRACING_DEBUG_PICKER_DOMAIN_INSTANCE		1
+#define RAY_TRACING_DEBUG_PICKER_DOMAIN_SEGMENT			2
+
diff --git a/Engine/Source/Editor/UnrealEd/Private/RayTracingDebugVisualizationMenuCommands.cpp b/Engine/Source/Editor/UnrealEd/Private/RayTracingDebugVisualizationMenuCommands.cpp
--- a/Engine/Source/Editor/UnrealEd/Private/RayTracingDebugVisualizationMenuCommands.cpp
+++ b/Engine/Source/Editor/UnrealEd/Private/RayTracingDebugVisualizationMenuCommands.cpp
@@ -62,6 +62,9 @@ void FRayTracingDebugVisualizationMenuCommands::CreateRayTracingDebugVisualizati
 
 	RayTracingDebugModeNames.Add(LOCTEXT("Triangles", "Triangles"));
 	RayTracingDebugModeNames.Add(LOCTEXT("FarField", "FarField"));
+	RayTracingDebugModeNames.Add(LOCTEXT("Dynamic Instances", "Dynamic Instances"));
+	RayTracingDebugModeNames.Add(LOCTEXT("Proxy Type", "Proxy Type"));
+	RayTracingDebugModeNames.Add(LOCTEXT("Picker", "Picker"));
 
 	for ( int32 RayTracingDebugIndex = 0; RayTracingDebugIndex < RayTracingDebugModeNames.Num(); ++RayTracingDebugIndex)
 	{
diff --git a/Engine/Source/Runtime/Renderer/Private/DeferredShadingRenderer.cpp b/Engine/Source/Runtime/Renderer/Private/DeferredShadingRenderer.cpp
--- a/Engine/Source/Runtime/Renderer/Private/DeferredShadingRenderer.cpp
+++ b/Engine/Source/Runtime/Renderer/Private/DeferredShadingRenderer.cpp
@@ -860,6 +860,7 @@ bool FDeferredShadingSceneRenderer::GatherRayTracingWorldInstancesForView(FRDGBu
 	};
 
 	const float CurrentWorldTime = View.Family->Time.GetWorldTimeSeconds();
+	const bool bEnableInstanceDebugData = IsRayTracingInstanceDebugDataEnabled(View);
 
 	// Consume output of the relevant primitive gathering task
 	RayTracingScene.UsedCoarseMeshStreamingHandles = MoveTemp(RelevantPrimitiveList.UsedCoarseMeshStreamingHandles);
@@ -1031,6 +1032,11 @@ bool FDeferredShadingSceneRenderer::GatherRayTracingWorldInstancesForView(FRDGBu
 					FRayTracingGeometryInstance& RayTracingInstance = RayTracingScene.Instances.AddDefaulted_GetRef();
 					RayTracingInstance.GeometryRHI = Geometry->RayTracingGeometryRHI;
 					checkf(RayTracingInstance.GeometryRHI, TEXT("Ray tracing instance must have a valid geometry."));
+					
+					if (bEnableInstanceDebugData)
+					{
+						RayTracingScene.AddInstanceDebugData(RayTracingInstance.GeometryRHI, SceneProxy, true);
+					}					
 
 					RayTracingInstance.DefaultUserData = PrimitiveIndex;
 					RayTracingInstance.Mask = Instance.Mask;
@@ -1156,6 +1162,7 @@ bool FDeferredShadingSceneRenderer::GatherRayTracingWorldInstancesForView(FRDGBu
 		const FScene& Scene;
 		TChunkedArray<FRayTracingRelevantPrimitive> RelevantStaticPrimitives;
 		const FRayTracingCullingParameters& CullingParameters;
+		const bool bEnableInstanceDebugData;
 
 		// Outputs
 
@@ -1164,11 +1171,12 @@ bool FDeferredShadingSceneRenderer::GatherRayTracingWorldInstancesForView(FRDGBu
 
 		FRayTracingSceneAddInstancesTask(const FScene& InScene,
 											TChunkedArray<FRayTracingRelevantPrimitive> InRelevantStaticPrimitives,
-											const FRayTracingCullingParameters& InCullingParameters,
+											const FRayTracingCullingParameters& InCullingParameters, bool InEnableInstanceDebugData,
 											FRayTracingScene& InRayTracingScene, TArray<FVisibleRayTracingMeshCommand>& InVisibleRayTracingMeshCommands)
 			: Scene(InScene)
 			, RelevantStaticPrimitives(MoveTemp(InRelevantStaticPrimitives))
 			, CullingParameters(InCullingParameters)
+			, bEnableInstanceDebugData(InEnableInstanceDebugData)
 			, RayTracingScene(InRayTracingScene)
 			, VisibleRayTracingMeshCommands(InVisibleRayTracingMeshCommands)
 		{
@@ -1276,6 +1284,11 @@ bool FDeferredShadingSceneRenderer::GatherRayTracingWorldInstancesForView(FRDGBu
 					checkf(SceneInfo->CachedRayTracingInstance.GeometryRHI, TEXT("Ray tracing instance must have a valid geometry."));
 					RayTracingScene.Instances.Add(SceneInfo->CachedRayTracingInstance);
 
+					if (bEnableInstanceDebugData)
+					{
+						RayTracingScene.AddInstanceDebugData(SceneInfo->CachedRayTracingInstance.GeometryRHI, SceneInfo->Proxy, false);
+					}					
+
 					const Experimental::FHashElementId GroupId = Scene.PrimitiveRayTracingGroupIds[PrimitiveIndex];
 					const bool bUseGroupBounds = CullingParameters.bCullUsingGroupIds && GroupId.IsValid();
 
@@ -1374,10 +1387,14 @@ bool FDeferredShadingSceneRenderer::GatherRayTracingWorldInstancesForView(FRDGBu
 						}
 
 						FRayTracingGeometryInstance& RayTracingInstance = RayTracingScene.Instances.AddDefaulted_GetRef();
-
 						RayTracingInstance.GeometryRHI = RelevantPrimitive.RayTracingGeometryRHI;
 						checkf(RayTracingInstance.GeometryRHI, TEXT("Ray tracing instance must have a valid geometry."));
 
+						if (bEnableInstanceDebugData)
+						{
+							RayTracingScene.AddInstanceDebugData(RayTracingInstance.GeometryRHI, SceneInfo->Proxy, false);
+						}						
+
 						InstanceBatch.Add(RayTracingScene, SceneInfo->GetInstanceSceneDataOffset(), (uint32)PrimitiveIndex);
 						RayTracingInstance.InstanceSceneDataOffsets = InstanceBatch.InstanceSceneDataOffsets;
 						RayTracingInstance.UserData = InstanceBatch.UserData;
@@ -1414,7 +1431,7 @@ bool FDeferredShadingSceneRenderer::GatherRayTracingWorldInstancesForView(FRDGBu
 	};
 
 	FGraphEventRef AddInstancesTask = TGraphTask<FRayTracingSceneAddInstancesTask>::CreateTask().ConstructAndDispatchWhenReady(
-		*Scene, MoveTemp(RelevantStaticPrimitives), View.RayTracingCullingParameters, // inputs 
+		*Scene, MoveTemp(RelevantStaticPrimitives), View.RayTracingCullingParameters, bEnableInstanceDebugData, // inputs 
 		RayTracingScene, View.VisibleRayTracingMeshCommands // outputs
 	);
 
@@ -3379,7 +3396,13 @@ void FDeferredShadingSceneRenderer::Render(FRDGBuilder& GraphBuilder)
 		{
 			for (const FViewInfo& View : Views)
 			{
-				RenderRayTracingDebug(GraphBuilder, View, SceneTextures.Color.Target);
+				FRayTracingPickingFeedback PickingFeedback = {};
+				RenderRayTracingDebug(GraphBuilder, View, SceneTextures.Color.Target, PickingFeedback);
+
+				OnGetOnScreenMessages.AddLambda([this, PickingFeedback](FScreenMessageWriter& ScreenMessageWriter)->void
+					{
+						RayTracingDisplayPicking(PickingFeedback, ScreenMessageWriter);
+					});
 			}
 		}
 	}
diff --git a/Engine/Source/Runtime/Renderer/Private/DeferredShadingRenderer.h b/Engine/Source/Runtime/Renderer/Private/DeferredShadingRenderer.h
--- a/Engine/Source/Runtime/Renderer/Private/DeferredShadingRenderer.h
+++ b/Engine/Source/Runtime/Renderer/Private/DeferredShadingRenderer.h
@@ -1013,8 +1013,9 @@ private:
 	void WaitForRayTracingScene(FRDGBuilder& GraphBuilder, FRDGBufferRef DynamicGeometryScratchBuffer);
 
 	/** Debug ray tracing functions. */
-	void RenderRayTracingDebug(FRDGBuilder& GraphBuilder, const FViewInfo& View, FRDGTextureRef SceneColorOutputTexture);
+	void RenderRayTracingDebug(FRDGBuilder& GraphBuilder, const FViewInfo& View, FRDGTextureRef SceneColorOutputTexture, FRayTracingPickingFeedback& PickingFeedback);
 	void RenderRayTracingBarycentrics(FRDGBuilder& GraphBuilder, const FViewInfo& View, FRDGTextureRef SceneColorOutputTexture, bool bVisualizeProceduralPrimitives);
+	void RayTracingDisplayPicking(const FRayTracingPickingFeedback& PickingFeedback, FScreenMessageWriter& Writer);
 
 	/** Fills RayTracingScene instance list for the given View and adds relevant ray tracing data to the view. Does not reset previous scene contents. */
 	bool GatherRayTracingWorldInstancesForView(FRDGBuilder& GraphBuilder, FViewInfo& View, FRayTracingScene& RayTracingScene, struct FRayTracingRelevantPrimitiveList&& RelevantPrimitiveList);
diff --git a/Engine/Source/Runtime/Renderer/Private/RayTracing/RayTracingDebug.cpp b/Engine/Source/Runtime/Renderer/Private/RayTracing/RayTracingDebug.cpp
--- a/Engine/Source/Runtime/Renderer/Private/RayTracing/RayTracingDebug.cpp
+++ b/Engine/Source/Runtime/Renderer/Private/RayTracing/RayTracingDebug.cpp
@@ -12,6 +12,7 @@
 #include "RenderGraphBuilder.h"
 #include "SceneUtils.h"
 #include "RaytracingDebugDefinitions.h"
+#include "RaytracingDebugTypes.h"
 #include "RayTracing/RayTracingLighting.h"
 #include "RayTracing/RaytracingOptions.h"
 #include "RayTracing/RayTracingTraversalStatistics.h"
@@ -26,6 +27,15 @@ static TAutoConsoleVariable<FString> CVarRayTracingDebugMode(
 	TEXT("Sets the ray tracing debug visualization mode (default = None - Driven by viewport menu) .\n")
 	);
 
+TAutoConsoleVariable<int32> CVarRayTracingDebugPickerDomain(
+	TEXT("r.RayTracing.Debug.PickerDomain"),
+	0,
+	TEXT("Changes the picker domain to highlight:\n")
+	TEXT("0 - Triangles (default)\n")
+	TEXT("1 - Instances\n"),
+	ECVF_RenderThreadSafe
+);
+
 TAutoConsoleVariable<int32> CVarRayTracingDebugModeOpaqueOnly(
 	TEXT("r.RayTracing.DebugVisualizationMode.OpaqueOnly"),
 	1,
@@ -71,8 +81,14 @@ class FRayTracingDebugRGS : public FGlobalShader
 	DECLARE_GLOBAL_SHADER(FRayTracingDebugRGS)
 	SHADER_USE_ROOT_PARAMETER_STRUCT(FRayTracingDebugRGS, FGlobalShader)
 
+	class FEnablePicking : SHADER_PERMUTATION_BOOL("ENABLE_PICKING");
+	class FEnableInstanceDebugData : SHADER_PERMUTATION_BOOL("ENABLE_INSTANCE_DEBUG_DATA");
+	using FPermutationDomain = TShaderPermutationDomain<FEnablePicking, FEnableInstanceDebugData>;
+
 	BEGIN_SHADER_PARAMETER_STRUCT(FParameters, )
 		SHADER_PARAMETER(uint32, VisualizationMode)
+		SHADER_PARAMETER(uint32, UseDebugCHS)
+		SHADER_PARAMETER(uint32, PickerDomain)
 		SHADER_PARAMETER(int32, ShouldUsePreExposure)
 		SHADER_PARAMETER(float, TimingScale)
 		SHADER_PARAMETER(float, MaxTraceDistance)
@@ -81,6 +97,8 @@ class FRayTracingDebugRGS : public FGlobalShader
 		SHADER_PARAMETER(int32, OpaqueOnly)
 		SHADER_PARAMETER_SRV(RaytracingAccelerationStructure, TLAS)
 		SHADER_PARAMETER_RDG_TEXTURE_UAV(RWTexture2D<float4>, Output)
+		SHADER_PARAMETER_RDG_BUFFER_SRV(StructuredBuffer, InstancesDebugData)
+		SHADER_PARAMETER_RDG_BUFFER_SRV(StructuredBuffer<FRayTracingPickingFeedback>, PickingBuffer)
 		SHADER_PARAMETER_STRUCT_REF(FViewUniformShaderParameters, ViewUniformBuffer)
 	END_SHADER_PARAMETER_STRUCT()
 
@@ -170,6 +188,55 @@ class FRayTracingDebugTraversalCS : public FGlobalShader
 };
 IMPLEMENT_GLOBAL_SHADER(FRayTracingDebugTraversalCS, "/Engine/Private/RayTracing/RayTracingDebugTraversal.usf", "RayTracingDebugTraversalCS", SF_Compute);
 
+class FRayTracingPickingRGS : public FGlobalShader
+{
+	DECLARE_GLOBAL_SHADER(FRayTracingPickingRGS)
+	SHADER_USE_ROOT_PARAMETER_STRUCT(FRayTracingPickingRGS, FGlobalShader)
+
+	BEGIN_SHADER_PARAMETER_STRUCT(FParameters, )
+		SHADER_PARAMETER_SRV(RaytracingAccelerationStructure, TLAS)
+		SHADER_PARAMETER_RDG_BUFFER_UAV(StructuredBuffer, PickingOutput)
+		SHADER_PARAMETER_RDG_BUFFER_SRV(StructuredBuffer, InstancesDebugData)
+		SHADER_PARAMETER_RDG_BUFFER_SRV(StructuredBuffer, InstanceBuffer)
+		SHADER_PARAMETER_STRUCT_REF(FViewUniformShaderParameters, ViewUniformBuffer)
+	END_SHADER_PARAMETER_STRUCT()
+
+	static bool ShouldCompilePermutation(const FGlobalShaderPermutationParameters& Parameters)
+	{
+		return ShouldCompileRayTracingShadersForProject(Parameters.Platform);
+	}
+};
+IMPLEMENT_GLOBAL_SHADER(FRayTracingPickingRGS, "/Engine/Private/RayTracing/RayTracingDebug.usf", "RayTracingDebugPickingRGS", SF_RayGen);
+
+struct FRayTracingDebugResources : public FRenderResource
+{
+	const int32 MaxPickingBuffers = 4;
+	int32 PickingBufferWriteIndex = 0;
+	int32 PickingBufferNumPending = 0;
+	TArray<FRHIGPUBufferReadback*> PickingBuffers;
+
+	virtual void InitRHI() override
+	{
+		PickingBuffers.AddZeroed(MaxPickingBuffers);
+	}
+
+	virtual void ReleaseRHI() override
+	{
+		for (int32 BufferIndex = 0; BufferIndex < PickingBuffers.Num(); ++BufferIndex)
+		{
+			if (PickingBuffers[BufferIndex])
+			{
+				delete PickingBuffers[BufferIndex];
+				PickingBuffers[BufferIndex] = nullptr;
+			}
+		}
+
+		PickingBuffers.Reset();
+	}
+};
+
+TGlobalResource<FRayTracingDebugResources> GRayTracingDebugResources;
+
 void BindRayTracingDebugCHSMaterialBindings(FRHICommandList& RHICmdList, const FViewInfo& View, FRayTracingPipelineState* PipelineState)
 {
 	const int32 NumTotalBindings = View.VisibleRayTracingMeshCommands.Num();
@@ -200,6 +267,7 @@ void BindRayTracingDebugCHSMaterialBindings(FRHICommandList& RHICmdList, const F
 		Binding.SegmentIndex = MeshCommand.GeometrySegmentIndex;
 		Binding.UniformBuffers = UniformBufferArray;
 		Binding.NumUniformBuffers = NumUniformBuffers;
+		Binding.UserData = VisibleMeshCommand.InstanceIndex;
 
 		Bindings[BindingIndex] = Binding;
 		BindingIndex++;
@@ -215,7 +283,11 @@ void BindRayTracingDebugCHSMaterialBindings(FRHICommandList& RHICmdList, const F
 
 static bool RequiresRayTracingDebugCHS(uint32 DebugVisualizationMode)
 {
-	return DebugVisualizationMode == RAY_TRACING_DEBUG_VIZ_INSTANCES || DebugVisualizationMode == RAY_TRACING_DEBUG_VIZ_TRIANGLES;
+	return DebugVisualizationMode == RAY_TRACING_DEBUG_VIZ_INSTANCES || 
+		DebugVisualizationMode == RAY_TRACING_DEBUG_VIZ_TRIANGLES ||
+		DebugVisualizationMode == RAY_TRACING_DEBUG_VIZ_DYNAMIC_INSTANCES ||
+		DebugVisualizationMode == RAY_TRACING_DEBUG_VIZ_PROXY_TYPE ||
+		DebugVisualizationMode == RAY_TRACING_DEBUG_VIZ_PICKER;
 }
 
 static bool IsRayTracingDebugTraversalMode(uint32 DebugVisualizationMode)
@@ -227,18 +299,133 @@ static bool IsRayTracingDebugTraversalMode(uint32 DebugVisualizationMode)
 		DebugVisualizationMode == RAY_TRACING_DEBUG_VIZ_TRAVERSAL_STATISTICS;
 }
 
+static bool IsRayTracingPickingEnabled(uint32 DebugVisualizationMode)
+{
+	return DebugVisualizationMode == RAY_TRACING_DEBUG_VIZ_PICKER;
+}
+
 void FDeferredShadingSceneRenderer::PrepareRayTracingDebug(const FSceneViewFamily& ViewFamily, TArray<FRHIRayTracingShader*>& OutRayGenShaders)
 {
 	// Declare all RayGen shaders that require material closest hit shaders to be bound
 	bool bEnabled = ViewFamily.EngineShowFlags.RayTracingDebug && ShouldRenderRayTracingEffect(ERayTracingPipelineCompatibilityFlags::FullPipeline);
 	if (bEnabled)
 	{
-		auto RayGenShader = GetGlobalShaderMap(ViewFamily.GetShaderPlatform())->GetShader<FRayTracingDebugRGS>();
-		OutRayGenShaders.Add(RayGenShader.GetRayTracingShader());
+		for (int32 EnablePicking = 0; EnablePicking < 2; ++EnablePicking)
+		{
+			for (int32 EnableInstanceDebugData = 0; EnableInstanceDebugData < 2; ++EnableInstanceDebugData)
+			{
+				FRayTracingDebugRGS::FPermutationDomain PermutationVector;
+				PermutationVector.Set<FRayTracingDebugRGS::FEnablePicking>((bool)EnablePicking);
+				PermutationVector.Set<FRayTracingDebugRGS::FEnableInstanceDebugData>((bool)EnableInstanceDebugData);
+
+				auto RayGenShader = GetGlobalShaderMap(ViewFamily.GetShaderPlatform())->GetShader<FRayTracingDebugRGS>(PermutationVector);
+				OutRayGenShaders.Add(RayGenShader.GetRayTracingShader());
+			}			
+		}		
+
+		auto PickingRayGenShader = GetGlobalShaderMap(ViewFamily.GetShaderPlatform())->GetShader<FRayTracingPickingRGS>();
+		OutRayGenShaders.Add(PickingRayGenShader.GetRayTracingShader());
+	}
+}
+
+static FRDGBufferRef RayTracingPerformPicking(FRDGBuilder& GraphBuilder, const FScene* Scene, const FViewInfo& View, FRayTracingPickingFeedback& PickingFeedback)
+{
+	FGlobalShaderMap* ShaderMap = GetGlobalShaderMap(GMaxRHIFeatureLevel);
+	auto RayGenShader = ShaderMap->GetShader<FRayTracingPickingRGS>();
+
+	FRayTracingPipelineStateInitializer Initializer;
+
+	FRHIRayTracingShader* RayGenShaderTable[] = { RayGenShader.GetRayTracingShader() };
+	Initializer.SetRayGenShaderTable(RayGenShaderTable);
+
+	auto ClosestHitShader = ShaderMap->GetShader<FRayTracingDebugCHS>();
+	FRHIRayTracingShader* HitGroupTable[] = { ClosestHitShader.GetRayTracingShader() };
+	Initializer.SetHitGroupTable(HitGroupTable);
+	Initializer.bAllowHitGroupIndexing = true; // Required for stable output using GetBaseInstanceIndex().
+	Initializer.MaxPayloadSizeInBytes = RAY_TRACING_MAX_ALLOWED_PAYLOAD_SIZE;
+	// TODO(UE-157946): This pipeline does not bind any miss shader and relies on the pipeline to do this automatically. This should be made explicit.
+	FRayTracingPipelineState* PickingPipeline = PipelineStateCache::GetAndOrCreateRayTracingPipelineState(GraphBuilder.RHICmdList, Initializer);
+
+	FRDGBufferDesc PickingBufferDesc = FRDGBufferDesc::CreateStructuredDesc(sizeof(FRayTracingPickingFeedback), 1);
+	PickingBufferDesc.Usage = EBufferUsageFlags(PickingBufferDesc.Usage | BUF_SourceCopy);
+	FRDGBufferRef PickingBuffer = GraphBuilder.CreateBuffer(PickingBufferDesc, TEXT("RayTracingDebug.PickingBuffer"));
+
+	FRayTracingPickingRGS::FParameters* RayGenParameters = GraphBuilder.AllocParameters<FRayTracingPickingRGS::FParameters>();
+	RayGenParameters->InstancesDebugData = GraphBuilder.CreateSRV(Scene->RayTracingScene.InstanceDebugBuffer); // TODO
+	RayGenParameters->TLAS = Scene->RayTracingScene.GetLayerSRVChecked(ERayTracingSceneLayer::Base);
+	RayGenParameters->InstanceBuffer = GraphBuilder.CreateSRV(Scene->RayTracingScene.InstanceBuffer);
+	RayGenParameters->ViewUniformBuffer = View.ViewUniformBuffer;
+	RayGenParameters->PickingOutput = GraphBuilder.CreateUAV(PickingBuffer);
+
+	GraphBuilder.AddPass(
+		RDG_EVENT_NAME("RayTracingPicking"),
+		RayGenParameters,
+		ERDGPassFlags::Compute,
+		[RayGenParameters, RayGenShader, &View, PickingPipeline](FRHIRayTracingCommandList& RHICmdList)
+		{
+			FRayTracingShaderBindingsWriter GlobalResources;
+			SetShaderParameters(GlobalResources, RayGenShader, *RayGenParameters);
+
+			BindRayTracingDebugCHSMaterialBindings(RHICmdList, View, PickingPipeline);
+			RHICmdList.SetRayTracingMissShader(View.GetRayTracingSceneChecked(), 0, PickingPipeline, 0 /* ShaderIndexInPipeline */, 0, nullptr, 0);
+
+			RHICmdList.RayTraceDispatch(PickingPipeline, RayGenShader.GetRayTracingShader(), View.GetRayTracingSceneChecked(), GlobalResources, 1, 1);
+		});
+
+	const int32 MaxPickingBuffers = GRayTracingDebugResources.MaxPickingBuffers;
+
+	int32& PickingBufferWriteIndex = GRayTracingDebugResources.PickingBufferWriteIndex;
+	int32& PickingBufferNumPending = GRayTracingDebugResources.PickingBufferNumPending;
+
+	TArray<FRHIGPUBufferReadback*>& PickingBuffers = GRayTracingDebugResources.PickingBuffers;
+
+	{
+		FRHIGPUBufferReadback* LatestPickingBuffer = nullptr;
+
+		// Find latest buffer that is ready
+		while (PickingBufferNumPending > 0)
+		{
+			uint32 Index = (PickingBufferWriteIndex + MaxPickingBuffers - PickingBufferNumPending) % MaxPickingBuffers;
+			if (PickingBuffers[Index]->IsReady())
+			{
+				--PickingBufferNumPending;
+				LatestPickingBuffer = PickingBuffers[Index];
+			}
+			else
+			{
+				break;
+			}
+		}
+
+		if (LatestPickingBuffer != nullptr)
+		{
+			TRACE_CPUPROFILER_EVENT_SCOPE(LockBuffer);
+			PickingFeedback = *((const FRayTracingPickingFeedback*)LatestPickingBuffer->Lock(sizeof(FRayTracingPickingFeedback)));
+			LatestPickingBuffer->Unlock();
+		}
 	}
+
+	// Skip when queue is full. It is NOT safe to EnqueueCopy on a buffer that already has a pending copy
+	if (PickingBufferNumPending != MaxPickingBuffers)
+	{
+		if (PickingBuffers[PickingBufferWriteIndex] == nullptr)
+		{
+			FRHIGPUBufferReadback* GPUBufferReadback = new FRHIGPUBufferReadback(TEXT("RayTracingDebug.PickingFeedback"));
+			PickingBuffers[PickingBufferWriteIndex] = GPUBufferReadback;
+		}
+
+		FRHIGPUBufferReadback* PickingReadback = PickingBuffers[PickingBufferWriteIndex];
+
+		AddEnqueueCopyPass(GraphBuilder, PickingReadback, PickingBuffer, 0u);
+
+		PickingBufferWriteIndex = (PickingBufferWriteIndex + 1) % MaxPickingBuffers;
+		PickingBufferNumPending = FMath::Min(PickingBufferNumPending + 1, MaxPickingBuffers);
+	}	
+
+	return PickingBuffer;
 }
 
-void FDeferredShadingSceneRenderer::RenderRayTracingDebug(FRDGBuilder& GraphBuilder, const FViewInfo& View, FRDGTextureRef SceneColorTexture)
+void FDeferredShadingSceneRenderer::RenderRayTracingDebug(FRDGBuilder& GraphBuilder, const FViewInfo& View, FRDGTextureRef SceneColorTexture, FRayTracingPickingFeedback& PickingFeedback)
 {
 	static TMap<FName, uint32> RayTracingDebugVisualizationModes;
 	if (RayTracingDebugVisualizationModes.Num() == 0)
@@ -274,6 +461,9 @@ void FDeferredShadingSceneRenderer::RenderRayTracingDebug(FRDGBuilder& GraphBuil
 		RayTracingDebugVisualizationModes.Emplace(FName(*LOCTEXT("Traversal Triangle", "Traversal Triangle").ToString()),						RAY_TRACING_DEBUG_VIZ_TRAVERSAL_TRIANGLE);
 		RayTracingDebugVisualizationModes.Emplace(FName(*LOCTEXT("Traversal All", "Traversal All").ToString()),									RAY_TRACING_DEBUG_VIZ_TRAVERSAL_ALL);
 		RayTracingDebugVisualizationModes.Emplace(FName(*LOCTEXT("Traversal Statistics", "Traversal Statistics").ToString()),					RAY_TRACING_DEBUG_VIZ_TRAVERSAL_STATISTICS);
+		RayTracingDebugVisualizationModes.Emplace(FName(*LOCTEXT("Dynamic Instances", "Dynamic Instances").ToString()),							RAY_TRACING_DEBUG_VIZ_DYNAMIC_INSTANCES);
+		RayTracingDebugVisualizationModes.Emplace(FName(*LOCTEXT("Proxy Type", "Proxy Type").ToString()),										RAY_TRACING_DEBUG_VIZ_PROXY_TYPE);
+		RayTracingDebugVisualizationModes.Emplace(FName(*LOCTEXT("Picker", "Picker").ToString()),												RAY_TRACING_DEBUG_VIZ_PICKER);
 	}
 
 	uint32 DebugVisualizationMode;
@@ -367,13 +557,24 @@ void FDeferredShadingSceneRenderer::RenderRayTracingDebug(FRDGBuilder& GraphBuil
 		return;
 	}
 
+	FRDGBufferRef PickingBuffer = nullptr;
+	if (IsRayTracingPickingEnabled(DebugVisualizationMode) && Scene->RayTracingScene.InstanceDebugBuffer != nullptr)
+	{
+		PickingBuffer = RayTracingPerformPicking(GraphBuilder, Scene, View, PickingFeedback);
+	}
+
+	FRayTracingDebugRGS::FPermutationDomain PermutationVector;
+	PermutationVector.Set<FRayTracingDebugRGS::FEnablePicking>(PickingBuffer != nullptr);
+	PermutationVector.Set<FRayTracingDebugRGS::FEnableInstanceDebugData>(Scene->RayTracingScene.InstanceDebugBuffer != nullptr);
+
 	FGlobalShaderMap* ShaderMap = GetGlobalShaderMap(FeatureLevel);
-	auto RayGenShader = ShaderMap->GetShader<FRayTracingDebugRGS>();
+	auto RayGenShader = ShaderMap->GetShader<FRayTracingDebugRGS>(PermutationVector);
 
 	FRayTracingPipelineState* Pipeline = View.RayTracingMaterialPipeline;
 	bool bRequiresBindings = false;
 
-	if (RequiresRayTracingDebugCHS(DebugVisualizationMode))
+	const bool bRequiresDebugCHS = RequiresRayTracingDebugCHS(DebugVisualizationMode);
+	if (bRequiresDebugCHS)
 	{
 		FRayTracingPipelineStateInitializer Initializer;
 
@@ -393,10 +594,22 @@ void FDeferredShadingSceneRenderer::RenderRayTracingDebug(FRDGBuilder& GraphBuil
 	FRayTracingDebugRGS::FParameters* RayGenParameters = GraphBuilder.AllocParameters<FRayTracingDebugRGS::FParameters>();
 
 	RayGenParameters->VisualizationMode = DebugVisualizationMode;
+	RayGenParameters->UseDebugCHS = bRequiresDebugCHS;
+	RayGenParameters->PickerDomain = CVarRayTracingDebugPickerDomain.GetValueOnRenderThread();
 	RayGenParameters->ShouldUsePreExposure = View.Family->EngineShowFlags.Tonemapper;
 	RayGenParameters->TimingScale = CVarRayTracingDebugTimingScale.GetValueOnAnyThread() / 25000.0f;
 	RayGenParameters->OpaqueOnly = CVarRayTracingDebugModeOpaqueOnly.GetValueOnRenderThread();
 	
+	if (Scene->RayTracingScene.InstanceDebugBuffer)
+	{
+		RayGenParameters->InstancesDebugData = GraphBuilder.CreateSRV(Scene->RayTracingScene.InstanceDebugBuffer);
+	}
+
+	if (PickingBuffer)
+	{		
+		RayGenParameters->PickingBuffer = GraphBuilder.CreateSRV(PickingBuffer);
+	}	
+
 	if (Lumen::UseFarField(ViewFamily))
 	{
 		RayGenParameters->MaxTraceDistance = Lumen::GetMaxTraceDistance(View);
@@ -437,5 +650,152 @@ void FDeferredShadingSceneRenderer::RenderRayTracingDebug(FRDGBuilder& GraphBuil
 	});
 }
 
+void FDeferredShadingSceneRenderer::RayTracingDisplayPicking(const FRayTracingPickingFeedback& PickingFeedback, FScreenMessageWriter& Writer)
+{
+	if (PickingFeedback.InstanceIndex == ~uint32(0))
+	{
+		return;
+	}
+	
+	int32 PickerDomain = CVarRayTracingDebugPickerDomain.GetValueOnRenderThread();
+	switch (PickerDomain)
+	{
+	case RAY_TRACING_DEBUG_PICKER_DOMAIN_TRIANGLE:
+		Writer.DrawLine(FText::FromString(TEXT("Domain [Triangle]")), 10, FColor::Yellow);
+		break;
+
+	case RAY_TRACING_DEBUG_PICKER_DOMAIN_SEGMENT:
+		Writer.DrawLine(FText::FromString(TEXT("Domain [Segment]")), 10, FColor::Yellow);
+		break;
+
+	case RAY_TRACING_DEBUG_PICKER_DOMAIN_INSTANCE:
+		Writer.DrawLine(FText::FromString(TEXT("Domain [Instance]")), 10, FColor::Yellow);
+		break;
+
+	default:
+		break; // Invalid picking domain
+	}
+
+	Writer.EmptyLine();
+
+	Writer.DrawLine(FText::FromString(TEXT("(Use r.RayTracing.Debug.PickerDomain to change domain)")), 10, FColor::Yellow);
+
+	Writer.EmptyLine();
+
+	Writer.DrawLine(FText::FromString(TEXT("[Hit]")), 10, FColor::Yellow);	
+
+	Writer.EmptyLine();
+
+	Writer.DrawLine(FText::FromString(FString::Printf(TEXT("Instance Index: %u"), PickingFeedback.InstanceIndex)), 10, FColor::Yellow);
+	Writer.DrawLine(FText::FromString(FString::Printf(TEXT("Geometry Instance Index: %u"), PickingFeedback.GeometryInstanceIndex)), 10, FColor::Yellow);
+	Writer.DrawLine(FText::FromString(FString::Printf(TEXT("Segment Index: %u"), PickingFeedback.GeometryIndex)), 10, FColor::Yellow);
+	Writer.DrawLine(FText::FromString(FString::Printf(TEXT("Triangle Index: %u"), PickingFeedback.TriangleIndex)), 10, FColor::Yellow);
+
+	Writer.EmptyLine();
+
+	FRHIRayTracingGeometry* Geometry = nullptr;
+	for (const FRayTracingGeometryInstance& Instance : Scene->RayTracingScene.Instances)
+	{
+		if (Instance.GeometryRHI)
+		{
+			const uint64 GeometryAddress = uint64(Instance.GeometryRHI.GetReference());
+			if (PickingFeedback.GeometryAddress == GeometryAddress)
+			{
+				Geometry = Instance.GeometryRHI;
+				break;
+			}
+		}
+	}
+	
+	Writer.DrawLine(FText::FromString(TEXT("[BLAS]")), 10, FColor::Yellow);
+	Writer.EmptyLine();
+
+	if (Geometry)
+	{
+		const FRayTracingGeometryInitializer& Initializer = Geometry->GetInitializer();
+		Writer.DrawLine(FText::FromString(FString::Printf(TEXT("Name: %s"), *Initializer.DebugName.ToString())), 10, FColor::Yellow);
+		Writer.DrawLine(FText::FromString(FString::Printf(TEXT("Num Segments: %d"), Initializer.Segments.Num())), 10, FColor::Yellow);
+		Writer.DrawLine(FText::FromString(FString::Printf(TEXT("Total Primitive Count: %u"), Initializer.TotalPrimitiveCount)), 10, FColor::Yellow);
+		Writer.DrawLine(FText::FromString(FString::Printf(TEXT("Fast Build: %d"), Initializer.bFastBuild)), 10, FColor::Yellow);
+		Writer.DrawLine(FText::FromString(FString::Printf(TEXT("Allow Update: %d"), Initializer.bAllowUpdate)), 10, FColor::Yellow);
+		Writer.DrawLine(FText::FromString(FString::Printf(TEXT("Allow Compaction: %d"), Initializer.bAllowCompaction)), 10, FColor::Yellow);
+
+		Writer.EmptyLine();
+
+		FRayTracingAccelerationStructureSize SizeInfo = Geometry->GetSizeInfo();
+		Writer.DrawLine(FText::FromString(FString::Printf(TEXT("Result Size: %u"), SizeInfo.ResultSize)), 10, FColor::Yellow);
+		Writer.DrawLine(FText::FromString(FString::Printf(TEXT("Build Scratch Size: %u"), SizeInfo.BuildScratchSize)), 10, FColor::Yellow);
+		Writer.DrawLine(FText::FromString(FString::Printf(TEXT("Update Scratch Size: %u"), SizeInfo.UpdateScratchSize)), 10, FColor::Yellow);
+	}
+	else
+	{
+		Writer.DrawLine(FText::FromString(TEXT("UNKNOWN")), 10, FColor::Yellow);
+	}	
+
+	Writer.EmptyLine();
+
+	Writer.DrawLine(FText::FromString(TEXT("[TLAS]")), 10, FColor::Yellow);
+
+	Writer.EmptyLine();
+
+	Writer.DrawLine(FText::FromString(FString::Printf(TEXT("InstanceId: %u"), PickingFeedback.InstanceId)), 10, FColor::Yellow);
+	Writer.DrawLine(FText::FromString(FString::Printf(TEXT("Mask: %u"), PickingFeedback.Mask)), 10, FColor::Yellow);
+	Writer.DrawLine(FText::FromString(FString::Printf(TEXT("ContributionToHitGroup: %u"), PickingFeedback.InstanceContributionToHitGroupIndex)), 10, FColor::Yellow);
+	{
+		const ERayTracingInstanceFlags Flags = (ERayTracingInstanceFlags)PickingFeedback.Flags;
+		FString FlagNames(TEXT(""));
+		if (EnumHasAnyFlags(Flags, ERayTracingInstanceFlags::TriangleCullDisable))
+		{
+			FlagNames += FString(TEXT("CullDisable "));
+		}
+
+		if (EnumHasAnyFlags(Flags, ERayTracingInstanceFlags::TriangleCullReverse))
+		{
+			FlagNames += FString(TEXT("CullReverse "));
+		}
+
+		if (EnumHasAnyFlags(Flags, ERayTracingInstanceFlags::ForceOpaque))
+		{
+			FlagNames += FString(TEXT("ForceOpaque "));
+		}
+
+		if (EnumHasAnyFlags(Flags, ERayTracingInstanceFlags::ForceNonOpaque))
+		{
+			FlagNames += FString(TEXT("ForceNonOpaque "));
+		}
+
+		Writer.DrawLine(FText::FromString(FString::Printf(TEXT("Flags: %u - %s"), PickingFeedback.Flags, *FlagNames)), 10, FColor::Yellow);
+	}	
+
+	Writer.EmptyLine();
+}
+
+bool IsRayTracingInstanceDebugDataEnabled(const FViewInfo& View)
+{
+#if !UE_BUILD_SHIPPING	
+	static TArray<FName> RayTracingDebugVisualizationModes;
+	if (RayTracingDebugVisualizationModes.Num() == 0)
+	{
+		RayTracingDebugVisualizationModes.Add(*LOCTEXT("Dynamic Instances", "Dynamic Instances").ToString());
+		RayTracingDebugVisualizationModes.Add(*LOCTEXT("Proxy Type", "Proxy Type").ToString());
+		RayTracingDebugVisualizationModes.Add(*LOCTEXT("Picker", "Picker").ToString());
+	}
+
+	FString ConsoleViewMode = CVarRayTracingDebugMode.GetValueOnRenderThread();
+	if (!ConsoleViewMode.IsEmpty())
+	{
+		return RayTracingDebugVisualizationModes.Contains(FName(*ConsoleViewMode));
+	}
+	else if (View.CurrentRayTracingDebugVisualizationMode != NAME_None)
+	{
+		return RayTracingDebugVisualizationModes.Contains(View.CurrentRayTracingDebugVisualizationMode);
+	}
+#endif
+	{
+		return false;
+	}
+}
+
 #undef LOCTEXT_NAMESPACE
+
 #endif
diff --git a/Engine/Source/Runtime/Renderer/Private/RayTracing/RayTracingScene.cpp b/Engine/Source/Runtime/Renderer/Private/RayTracing/RayTracingScene.cpp
--- a/Engine/Source/Runtime/Renderer/Private/RayTracing/RayTracingScene.cpp
+++ b/Engine/Source/Runtime/Renderer/Private/RayTracing/RayTracingScene.cpp
@@ -9,6 +9,7 @@
 #include "RayTracingDefinitions.h"
 #include "RenderGraphBuilder.h"
 #include "RenderGraphUtils.h"
+#include "RayTracingOptions.h"
 
 BEGIN_SHADER_PARAMETER_STRUCT(FBuildInstanceBufferPassParams, )
 	SHADER_PARAMETER_RDG_BUFFER_UAV(RWStructuredBuffer, InstanceBuffer)
@@ -176,6 +177,12 @@ void FRayTracingScene::CreateWithInitializationData(FRDGBuilder& GraphBuilder, c
 				InstanceUploadData,
 				TransformUploadData);
 		}, TStatId(), nullptr, ENamedThreads::AnyThread);
+		
+		if (InstancesDebugData.Num() > 0)
+		{
+			check(InstancesDebugData.Num() == Instances.Num());
+			InstanceDebugBuffer = CreateStructuredBuffer(GraphBuilder, TEXT("FRayTracingScene::InstanceDebugData"), InstancesDebugData);
+		}
 
 		FBuildInstanceBufferPassParams* PassParams = GraphBuilder.AllocParameters<FBuildInstanceBufferPassParams>();
 		PassParams->InstanceBuffer = GraphBuilder.CreateUAV(InstanceBuffer);
@@ -284,11 +291,24 @@ FRHIShaderResourceView* FRayTracingScene::GetLayerSRVChecked(ERayTracingSceneLay
 	return LayerSRVs[uint8(Layer)].GetReference();
 }
 
+void FRayTracingScene::AddInstanceDebugData(const FRHIRayTracingGeometry* GeometryRHI, const FPrimitiveSceneProxy* Proxy, bool bDynamic)
+{
+	FRayTracingInstanceDebugData& InstanceDebugData = InstancesDebugData.AddDefaulted_GetRef();
+	InstanceDebugData.Flags = bDynamic ? 1 : 0;
+	InstanceDebugData.GeometryAddress = uint64(GeometryRHI);
+
+	if (Proxy)
+	{
+		InstanceDebugData.ProxyHash = Proxy->GetTypeHash();
+	}
+}
+
 void FRayTracingScene::Reset()
 {
 	WaitForTasks();
 
 	Instances.Reset();
+	InstancesDebugData.Reset();
 	NumMissShaderSlots = 1;
 	NumCallableShaderSlots = 0;
 	CallableCommands.Reset();
@@ -299,6 +319,7 @@ void FRayTracingScene::Reset()
 	Allocator.Flush();
 
 	BuildScratchBuffer = nullptr;
+	InstanceDebugBuffer = nullptr;
 }
 
 void FRayTracingScene::ResetAndReleaseResources()
@@ -306,6 +327,7 @@ void FRayTracingScene::ResetAndReleaseResources()
 	Reset();
 
 	Instances.Empty();
+	InstancesDebugData.Empty();
 	CallableCommands.Empty();
 	UniformBuffers.Empty();
 	GeometriesToBuild.Empty();
diff --git a/Engine/Source/Runtime/Renderer/Private/RayTracing/RayTracingScene.h b/Engine/Source/Runtime/Renderer/Private/RayTracing/RayTracingScene.h
--- a/Engine/Source/Runtime/Renderer/Private/RayTracing/RayTracingScene.h
+++ b/Engine/Source/Runtime/Renderer/Private/RayTracing/RayTracingScene.h
@@ -15,6 +15,7 @@
 #include "PrimitiveSceneProxy.h"
 #include "MeshPassProcessor.h"
 #include "RayTracingInstanceBufferUtil.h"
+#include "RayTracingDebugTypes.h"
 
 class FGPUScene;
 class FRHIRayTracingScene;
@@ -80,7 +81,8 @@ public:
 	RENDERER_API FRHIBuffer* GetBufferChecked() const;
 	
 	RENDERER_API FRHIShaderResourceView* GetLayerSRVChecked(ERayTracingSceneLayer Layer) const;
-
+	
+	void AddInstanceDebugData(const FRHIRayTracingGeometry* GeometryRHI, const FPrimitiveSceneProxy* Proxy, bool bDynamic);
 public:
 
 	// Public members for initial refactoring step (previously were public members of FViewInfo).
@@ -107,6 +109,10 @@ public:
 	FRDGBufferRef InstanceBuffer;
 	FRDGBufferRef BuildScratchBuffer;
 
+	// Special data for debugging purposes
+	FRDGBufferRef InstanceDebugBuffer = nullptr;
+	TArray<FRayTracingInstanceDebugData> InstancesDebugData;
+
 private:
 	void WaitForTasks() const;
 
diff --git a/Engine/Source/Runtime/Renderer/Private/RayTracing/RaytracingOptions.h b/Engine/Source/Runtime/Renderer/Private/RayTracing/RaytracingOptions.h
--- a/Engine/Source/Runtime/Renderer/Private/RayTracing/RaytracingOptions.h
+++ b/Engine/Source/Runtime/Renderer/Private/RayTracing/RaytracingOptions.h
@@ -75,6 +75,7 @@ extern bool EnableRayTracingShadowTwoSidedGeometry();
 extern float GetRaytracingMaxNormalBias();
 extern int32 GetRayTracingCulling();
 extern float GetRayTracingCullingRadius();
+extern bool IsRayTracingInstanceDebugDataEnabled(const FViewInfo& View);
 
 extern bool CanUseRayTracingAMDHitToken();
 
@@ -155,6 +156,11 @@ FORCEINLINE bool CanUseRayTracingAMDHitToken()
 	return false;
 }
 
+FORCEINLINE bool IsRayTracingInstanceDebugDataEnabled()
+{
+	return false;
+}
+
 #endif // RHI_RAYTRACING
 
 FORCEINLINE bool ShouldRenderRayTracingEffect(ERayTracingPipelineCompatibilityFlags CompatibilityFlags)
