Fix cooking error - use RWStructuredBuffer for RDG_BUFFER_UAV parameter

diff --git a/Engine/Source/Runtime/Renderer/Private/RayTracing/RayTracingDebug.cpp b/Engine/Source/Runtime/Renderer/Private/RayTracing/RayTracingDebug.cpp
--- a/Engine/Source/Runtime/Renderer/Private/RayTracing/RayTracingDebug.cpp
+++ b/Engine/Source/Runtime/Renderer/Private/RayTracing/RayTracingDebug.cpp
@@ -195,7 +195,7 @@ class FRayTracingPickingRGS : public FGlobalShader
 
 	BEGIN_SHADER_PARAMETER_STRUCT(FParameters, )
 		SHADER_PARAMETER_SRV(RaytracingAccelerationStructure, TLAS)
-		SHADER_PARAMETER_RDG_BUFFER_UAV(StructuredBuffer, PickingOutput)
+		SHADER_PARAMETER_RDG_BUFFER_UAV(RWStructuredBuffer, PickingOutput)
 		SHADER_PARAMETER_RDG_BUFFER_SRV(StructuredBuffer, InstancesDebugData)
 		SHADER_PARAMETER_RDG_BUFFER_SRV(StructuredBuffer, InstanceBuffer)
 		SHADER_PARAMETER_STRUCT_REF(FViewUniformShaderParameters, ViewUniformBuffer)
