Add planar reflection support for hair strands.

diff --git a/Engine/Plugins/Runtime/HairStrands/Source/HairStrandsCore/Private/GroomComponent.cpp b/Engine/Plugins/Runtime/HairStrands/Source/HairStrandsCore/Private/GroomComponent.cpp
--- a/Engine/Plugins/Runtime/HairStrands/Source/HairStrandsCore/Private/GroomComponent.cpp
+++ b/Engine/Plugins/Runtime/HairStrands/Source/HairStrandsCore/Private/GroomComponent.cpp
@@ -632,7 +632,7 @@ public:
 		for (int32 ViewIndex = 0; ViewIndex < Views.Num(); ViewIndex++)
 		{
 			const FSceneView* View = Views[ViewIndex];
-			if (View->bIsReflectionCapture || View->bIsPlanarReflection)
+			if (View->bIsReflectionCapture)
 			{
 				continue;
 			}
diff --git a/Engine/Source/Runtime/Renderer/Private/HairStrands/HairStrandsCluster.cpp b/Engine/Source/Runtime/Renderer/Private/HairStrands/HairStrandsCluster.cpp
--- a/Engine/Source/Runtime/Renderer/Private/HairStrands/HairStrandsCluster.cpp
+++ b/Engine/Source/Runtime/Renderer/Private/HairStrands/HairStrandsCluster.cpp
@@ -190,7 +190,7 @@ void CreateHairStrandsMacroGroups(
 	FHairStrandsViewData& OutHairStrandsViewData)
 {
 	const bool bHasHairStrandsElements = View.HairStrandsMeshElements.Num() != 0 || Scene->HairStrandsSceneData.RegisteredProxies.Num() != 0;
-	if (!View.Family || !bHasHairStrandsElements || View.bIsPlanarReflection || View.bIsReflectionCapture)
+	if (!View.Family || !bHasHairStrandsElements || View.bIsReflectionCapture)
 	{
 		return;
 	}
diff --git a/Engine/Source/Runtime/Renderer/Private/HairStrands/HairStrandsVisibility.cpp b/Engine/Source/Runtime/Renderer/Private/HairStrands/HairStrandsVisibility.cpp
--- a/Engine/Source/Runtime/Renderer/Private/HairStrands/HairStrandsVisibility.cpp
+++ b/Engine/Source/Runtime/Renderer/Private/HairStrands/HairStrandsVisibility.cpp
@@ -664,6 +664,7 @@ bool FHairMaterialProcessor::Process(
 	FHairMaterialShaderElementData ShaderElementData(MacroGroupId, HairMaterialId, HairPrimitiveId, HairPrimitiveLightChannelMask);
 	ShaderElementData.InitializeMeshMaterialData(ViewIfDynamicMeshCommand, PrimitiveSceneProxy, MeshBatch, StaticMeshId, false);
 
+	const bool bReverseCulling = ViewIfDynamicMeshCommand ? ViewIfDynamicMeshCommand->bIsPlanarReflection : false;
 	BuildMeshDrawCommands(
 		MeshBatch,
 		BatchElementMask,
@@ -673,7 +674,7 @@ bool FHairMaterialProcessor::Process(
 		DrawRenderState,
 		PassShaders,
 		ERasterizerFillMode::FM_Solid,
-		ERasterizerCullMode::CM_CCW,
+		bReverseCulling ? ERasterizerCullMode::CM_CW : ERasterizerCullMode::CM_CCW,
 		FMeshDrawCommandSortKey::Default,
 		EMeshPassFeatures::Default,
 		ShaderElementData);
