Removed DummyRectLightTextureForCapsuleCompilerWarning dummy resource.

#rnx

diff --git a/Engine/Shaders/Private/CapsuleLightIntegrate.ush b/Engine/Shaders/Private/CapsuleLightIntegrate.ush
--- a/Engine/Shaders/Private/CapsuleLightIntegrate.ush
+++ b/Engine/Shaders/Private/CapsuleLightIntegrate.ush
@@ -10,8 +10,6 @@
 #include "CapsuleLightSampling.ush"
 #include "SobolRandom.ush"
 
-Texture2D DummyRectLightTextureForCapsuleCompilerWarning;
-
 // Should this be SH instead?
 float IntegrateLight( FCapsuleLight Capsule, bool bInverseSquared )
 {
diff --git a/Engine/Source/Runtime/Renderer/Private/LightRendering.cpp b/Engine/Source/Runtime/Renderer/Private/LightRendering.cpp
--- a/Engine/Source/Runtime/Renderer/Private/LightRendering.cpp
+++ b/Engine/Source/Runtime/Renderer/Private/LightRendering.cpp
@@ -553,8 +553,6 @@ class FDeferredLightPS : public FGlobalShader
 		SHADER_PARAMETER_SAMPLER(SamplerState, LightingChannelsSampler)
 		SHADER_PARAMETER_RDG_BUFFER_SRV(Buffer, HairTransmittanceBuffer)
 		SHADER_PARAMETER_RDG_TEXTURE(Texture2D, ScreenShadowMaskSubPixelTexture)
-		SHADER_PARAMETER_RDG_TEXTURE(Texture2D, DummyRectLightTextureForCapsuleCompilerWarning)
-		SHADER_PARAMETER_SAMPLER(SamplerState, DummyRectLightSamplerForCapsuleCompilerWarning)
 		SHADER_PARAMETER_STRUCT_REF(FViewUniformShaderParameters, View)
 		SHADER_PARAMETER_RDG_UNIFORM_BUFFER(FDeferredLightUniformStruct, DeferredLight)
 		RENDER_TARGET_BINDING_SLOTS()
@@ -1880,7 +1878,6 @@ static FDeferredLightPS::FParameters GetDeferredLightPSParameters(
 	const bool bIsDirectional = LightType == LightType_Directional;
 
 	FRDGTextureRef WhiteDummy = GSystemTextures.GetWhiteDummy(GraphBuilder);
-	FRDGTextureRef DepthDummy = GSystemTextures.GetDepthDummy(GraphBuilder);
 	FRDGBufferRef BufferDummy = GSystemTextures.GetDefaultBuffer(GraphBuilder, 4, 0u);
 	FRDGBufferSRVRef BufferDummySRV = GraphBuilder.CreateSRV(BufferDummy, PF_R32_UINT);
 
@@ -1895,8 +1892,6 @@ static FDeferredLightPS::FParameters GetDeferredLightPSParameters(
 	Out.CloudShadowEnabled = SetupLightCloudTransmittanceParameters(GraphBuilder, Scene, View, LightSceneInfo, Out.CloudShadow) ? 1 : 0;
 	Out.LightAttenuationTexture = ShadowMaskTexture ? ShadowMaskTexture : WhiteDummy;
 	Out.LightAttenuationTextureSampler = TStaticSamplerState<SF_Point, AM_Wrap, AM_Wrap, AM_Wrap>::GetRHI();
-	Out.DummyRectLightTextureForCapsuleCompilerWarning = DepthDummy;
-	Out.DummyRectLightSamplerForCapsuleCompilerWarning = TStaticSamplerState<SF_Bilinear, AM_Clamp, AM_Clamp, AM_Clamp>::GetRHI();
 	Out.IESTextureSampler = TStaticSamplerState<SF_Bilinear, AM_Clamp, AM_Clamp, AM_Clamp>::GetRHI();
 	Out.IESTexture = GSystemTextures.WhiteDummy->GetRHI();
 	if (LightSceneInfo->Proxy->GetIESTextureResource())
@@ -2439,7 +2434,6 @@ static FSimpleLightsStandardDeferredParameters GetRenderLightSimpleParameters(
 	FSimpleLightsStandardDeferredParameters Out;
 
 	FRDGTextureRef WhiteDummy = GSystemTextures.GetWhiteDummy(GraphBuilder);
-	FRDGTextureRef DepthDummy = GSystemTextures.GetDepthDummy(GraphBuilder);
 	FRDGBufferRef BufferDummy = GSystemTextures.GetDefaultBuffer(GraphBuilder, 4, 0u);
 	FRDGBufferSRVRef BufferDummySRV = GraphBuilder.CreateSRV(BufferDummy, PF_R32_UINT);
 	
@@ -2454,8 +2448,6 @@ static FSimpleLightsStandardDeferredParameters GetRenderLightSimpleParameters(
 	SetupLightCloudTransmittanceParameters(GraphBuilder, nullptr, View, nullptr, Out.PS.CloudShadow);
 	Out.PS.LightAttenuationTexture = WhiteDummy;
 	Out.PS.LightAttenuationTextureSampler = TStaticSamplerState<SF_Point, AM_Wrap, AM_Wrap, AM_Wrap>::GetRHI();
-	Out.PS.DummyRectLightTextureForCapsuleCompilerWarning = DepthDummy;
-	Out.PS.DummyRectLightSamplerForCapsuleCompilerWarning = TStaticSamplerState<SF_Bilinear, AM_Clamp, AM_Clamp, AM_Clamp>::GetRHI();
 	Out.PS.IESTextureSampler = TStaticSamplerState<SF_Bilinear, AM_Clamp, AM_Clamp, AM_Clamp>::GetRHI();
 	Out.PS.IESTexture = GSystemTextures.WhiteDummy->GetRHI();
 	Out.PS.View = View.ViewUniformBuffer;
