Multi-viewport on PC (D3D11/12, Vulkan).

- RHISupportsMultiviewport made to work on relevant shader platforms.
- RHISetStereoViewport implemented
- Unused outputs removal fixed to always consider viewport index used.
- Renamed confusing terminology (multiview -> multiviewport).
- Added a CVar (vr.MultiViewport) to force disable
- Added a canvas output to know the method in non-Shipping

diff --git a/Engine/Config/Mac/DataDrivenPlatformInfo.ini b/Engine/Config/Mac/DataDrivenPlatformInfo.ini
--- a/Engine/Config/Mac/DataDrivenPlatformInfo.ini
+++ b/Engine/Config/Mac/DataDrivenPlatformInfo.ini
@@ -42,7 +42,7 @@ bSupportsDiaphragmDOF=true
 bSupportsVertexShaderLayer=true
 bSupportsDxc=true
 bSupportsInstancedStereo=true
-bSupportsMultiView=true
+SupportsMultiViewport=RuntimeGuaranteed
 bSupportsVolumetricFog=true
 bSupportsPercentageCloserShadows=true
 bSupportsCapsuleShadows=true
diff --git a/Engine/Config/VulkanPC/DataDrivenPlatformInfo.ini b/Engine/Config/VulkanPC/DataDrivenPlatformInfo.ini
--- a/Engine/Config/VulkanPC/DataDrivenPlatformInfo.ini
+++ b/Engine/Config/VulkanPC/DataDrivenPlatformInfo.ini
@@ -36,11 +36,8 @@ bSupportsPercentageCloserShadows=true
 bSupportsVolumetricFog = true
 bSupportsIndexBufferUAVs = true
 
-; //#todo-vulkanddpi
-bSupportsInstancedStereo = false
-
-; //#todo-vulkanddpi
-bSupportsMultiView = false
+bSupportsInstancedStereo = true
+SupportsMultiViewport = RuntimeDependent
 
 bSupportsMSAA = true
 
diff --git a/Engine/Config/Windows/DataDrivenPlatformInfo.ini b/Engine/Config/Windows/DataDrivenPlatformInfo.ini
--- a/Engine/Config/Windows/DataDrivenPlatformInfo.ini
+++ b/Engine/Config/Windows/DataDrivenPlatformInfo.ini
@@ -92,6 +92,7 @@ bSupportsVolumetricFog=true
 bSupportsIndexBufferUAVs=true
 bSupportsROV=true
 bSupportsOIT=true
+SupportsMultiViewport=RuntimeDependent
 
 [PreviewPlatform PCD3D_SM5]
 PlatformName=PCD3D_SM5
@@ -146,6 +147,7 @@ bSupportsROV=true
 bSupportsOIT=true
 bSupportsRealTypes=RuntimeDependent
 EnablesHLSL2021ByDefault=1
+SupportsMultiViewport=RuntimeGuaranteed
 
 [PreviewPlatform PCD3D_SM6]
 PlatformName=PCD3D_SM6
diff --git a/Engine/Shaders/Public/ShaderVersion.ush b/Engine/Shaders/Public/ShaderVersion.ush
--- a/Engine/Shaders/Public/ShaderVersion.ush
+++ b/Engine/Shaders/Public/ShaderVersion.ush
@@ -2,4 +2,4 @@
 // Each time the console command is executed it generates a new GUID. As this file is included
 // in Platform.ush (which should be included in any shader) it allows to invalidate the shader DDC.
 // 
-// GUID = C255D40A-6FE5-43D2-B5D8-F6963F9F92D1
\ No newline at end of file
+// GUID = 0131629C-FCE1-4B72-9872-DB06D2BA8DE8
diff --git a/Engine/Source/Developer/Windows/ShaderFormatD3D/Private/D3DShaderCompiler.cpp b/Engine/Source/Developer/Windows/ShaderFormatD3D/Private/D3DShaderCompiler.cpp
--- a/Engine/Source/Developer/Windows/ShaderFormatD3D/Private/D3DShaderCompiler.cpp
+++ b/Engine/Source/Developer/Windows/ShaderFormatD3D/Private/D3DShaderCompiler.cpp
@@ -1253,6 +1253,7 @@ void CompileD3DShader(const FShaderCompilerInput& Input, FShaderCompilerOutput&
 		// Always add SV_Position
 		TArray<FString> UsedOutputs = Input.UsedOutputs;
 		UsedOutputs.AddUnique(TEXT("SV_POSITION"));
+		UsedOutputs.AddUnique(TEXT("SV_ViewPortArrayIndex"));
 
 		// We can't remove any of the output-only system semantics
 		//@todo - there are a bunch of tessellation ones as well
diff --git a/Engine/Source/Runtime/D3D12RHI/Private/D3D12CommandContext.h b/Engine/Source/Runtime/D3D12RHI/Private/D3D12CommandContext.h
--- a/Engine/Source/Runtime/D3D12RHI/Private/D3D12CommandContext.h
+++ b/Engine/Source/Runtime/D3D12RHI/Private/D3D12CommandContext.h
@@ -304,6 +304,7 @@ public:
 	virtual void RHIEndScene() final override;
 	virtual void RHISetStreamSource(uint32 StreamIndex, FRHIBuffer* VertexBuffer, uint32 Offset) final override;
 	virtual void RHISetViewport(float MinX, float MinY, float MinZ, float MaxX, float MaxY, float MaxZ) final override;
+	virtual void RHISetStereoViewport(float LeftMinX, float RightMinX, float LeftMinY, float RightMinY, float MinZ, float LeftMaxX, float RightMaxX, float LeftMaxY, float RightMaxY, float MaxZ) override;
 	virtual void RHISetScissorRect(bool bEnable, uint32 MinX, uint32 MinY, uint32 MaxX, uint32 MaxY) final override;
 	virtual void RHISetGraphicsPipelineState(FRHIGraphicsPipelineState* GraphicsPipelineState, uint32 StencilRef, bool bApplyAdditionalState) final override;
 	virtual void RHISetShaderTexture(FRHIGraphicsShader* Shader, uint32 TextureIndex, FRHITexture* NewTexture) final override;
@@ -622,6 +623,11 @@ public:
 	{
 		ContextRedirect(RHISetViewport(MinX, MinY, MinZ, MaxX, MaxY, MaxZ));
 	}
+	FORCEINLINE virtual void RHISetStereoViewport(float LeftMinX, float RightMinX, float LeftMinY, float RightMinY, float MinZ, float LeftMaxX, float RightMaxX, float LeftMaxY, float RightMaxY, float MaxZ) override
+	{
+		ContextRedirect(RHISetStereoViewport(LeftMinX, RightMinX, LeftMinY, RightMinY, MinZ, LeftMaxX, RightMaxX, LeftMaxY, RightMaxY, MaxZ));
+	}
+
 	FORCEINLINE virtual void RHISetScissorRect(bool bEnable, uint32 MinX, uint32 MinY, uint32 MaxX, uint32 MaxY) final override
 	{
 		ContextRedirect(RHISetScissorRect(bEnable, MinX, MinY, MaxX, MaxY));
diff --git a/Engine/Source/Runtime/D3D12RHI/Private/D3D12Commands.cpp b/Engine/Source/Runtime/D3D12RHI/Private/D3D12Commands.cpp
--- a/Engine/Source/Runtime/D3D12RHI/Private/D3D12Commands.cpp
+++ b/Engine/Source/Runtime/D3D12RHI/Private/D3D12Commands.cpp
@@ -741,6 +741,36 @@ void FD3D12CommandContext::RHISetViewport(float MinX, float MinY, float MinZ, fl
 	}
 }
 
+void FD3D12CommandContext::RHISetStereoViewport(float LeftMinX, float RightMinX, float LeftMinY, float RightMinY, float MinZ, float LeftMaxX, float RightMaxX, float LeftMaxY, float RightMaxY, float MaxZ)
+{
+	// Set up both viewports
+	D3D12_VIEWPORT Viewports[2] = {0};
+
+	Viewports[0].TopLeftX = FMath::FloorToInt(LeftMinX);
+	Viewports[0].TopLeftY = FMath::FloorToInt(LeftMinY);
+	Viewports[0].Width = FMath::CeilToInt(LeftMaxX - LeftMinX);
+	Viewports[0].Height = FMath::CeilToInt(LeftMaxY - LeftMinY);
+	Viewports[0].MinDepth = MinZ;
+	Viewports[0].MaxDepth = MaxZ;
+
+	Viewports[1].TopLeftX = FMath::FloorToInt(RightMinX);
+	Viewports[1].TopLeftY = FMath::FloorToInt(RightMinY);
+	Viewports[1].Width = FMath::CeilToInt(RightMaxX - RightMinX);
+	Viewports[1].Height = FMath::CeilToInt(RightMaxY - RightMinY);
+	Viewports[1].MinDepth = MinZ;
+	Viewports[1].MaxDepth = MaxZ;
+
+	D3D12_RECT ScissorRects[2] =
+	{
+		{ Viewports[0].TopLeftX, Viewports[0].TopLeftY, Viewports[0].TopLeftX + Viewports[0].Width, Viewports[0].TopLeftY + Viewports[0].Height },
+		{ Viewports[1].TopLeftX, Viewports[1].TopLeftY, Viewports[1].TopLeftX + Viewports[1].Width, Viewports[1].TopLeftY + Viewports[1].Height }
+	};
+
+	StateCache.SetViewports(2, Viewports);
+	// Set the scissor rects appropriately.
+	StateCache.SetScissorRects(2, ScissorRects);
+}
+
 void FD3D12CommandContext::RHISetScissorRect(bool bEnable, uint32 MinX, uint32 MinY, uint32 MaxX, uint32 MaxY)
 {
 	if (bEnable)
diff --git a/Engine/Source/Runtime/Engine/Private/SceneView.cpp b/Engine/Source/Runtime/Engine/Private/SceneView.cpp
--- a/Engine/Source/Runtime/Engine/Private/SceneView.cpp
+++ b/Engine/Source/Runtime/Engine/Private/SceneView.cpp
@@ -5,6 +5,7 @@
 =============================================================================*/
 
 #include "SceneView.h"
+#include "Engine/Scene.h"
 #include "Misc/CommandLine.h"
 #include "Misc/Paths.h"
 #include "EngineGlobals.h"
@@ -758,7 +759,10 @@ FSceneView::FSceneView(const FSceneViewInitOptions& InitOptions)
 	, bIsLocked(false)
 	, bStaticSceneOnly(false)
 	, bIsInstancedStereoEnabled(false)
+PRAGMA_DISABLE_DEPRECATION_WARNINGS
 	, bIsMultiViewEnabled(false)
+PRAGMA_ENABLE_DEPRECATION_WARNINGS
+	, bIsMultiViewportEnabled(false)
 	, bIsMobileMultiViewEnabled(false)
 	, bShouldBindInstancedViewUB(false)
 	, UnderwaterDepth(-1.0f)
@@ -849,8 +853,11 @@ FSceneView::FSceneView(const FSceneViewInitOptions& InitOptions)
 	static const auto CVar = IConsoleManager::Get().FindTConsoleVariableDataInt(TEXT("vr.InstancedStereo"));
 	bIsInstancedStereoEnabled = !bUsingMobileRenderer && RHISupportsInstancedStereo(ShaderPlatform) && (CVar && CVar->GetValueOnAnyThread() != 0);
 
-	// TODO: Should be renamed to multi-viewport
-	bIsMultiViewEnabled = RHISupportsMultiView(ShaderPlatform) && bIsInstancedStereoEnabled;
+	static const auto CVarMultiViewport = IConsoleManager::Get().FindTConsoleVariableDataInt(TEXT("vr.MultiViewport"));
+	bIsMultiViewportEnabled = RHISupportsMultiViewport(ShaderPlatform) && bIsInstancedStereoEnabled && (CVarMultiViewport && CVarMultiViewport->GetValueOnAnyThread() != 0);
+	PRAGMA_DISABLE_DEPRECATION_WARNINGS
+	bIsMultiViewEnabled = bIsMultiViewportEnabled;	// temporary, as a graceful way to support plugins/licensee mods
+	PRAGMA_ENABLE_DEPRECATION_WARNINGS
 
 	bIsMobileMultiViewEnabled = Family && Family->bRequireMultiView;
 	if (bIsMobileMultiViewEnabled && !RHISupportsMobileMultiView(ShaderPlatform))
@@ -866,6 +873,12 @@ FSceneView::FSceneView(const FSceneViewInitOptions& InitOptions)
 	{
 		const UE::StereoRenderUtils::FStereoShaderAspects Aspects(ShaderPlatform);
 		bShouldBindInstancedViewUB = Aspects.IsInstancedStereoEnabled() || Aspects.IsMobileMultiViewEnabled();
+
+		// make sure that FStereoShaderAspects logic matches the above logic (TODO: above should be replaced by setting them through aspects after more testing)
+		ensure(bIsInstancedStereoEnabled == Aspects.IsInstancedStereoEnabled());
+		ensure(bIsMultiViewportEnabled == Aspects.IsInstancedMultiViewportEnabled());
+		ensure(bIsMobileMultiViewEnabled == (Family && Family->bRequireMultiView && Aspects.IsMobileMultiViewEnabled()));
+		ensure(bShouldBindInstancedViewUB == bIsInstancedStereoEnabled || bIsMobileMultiViewEnabled);
 	}
 
 	SetupAntiAliasingMethod();
diff --git a/Engine/Source/Runtime/Engine/Public/SceneView.h b/Engine/Source/Runtime/Engine/Public/SceneView.h
--- a/Engine/Source/Runtime/Engine/Public/SceneView.h
+++ b/Engine/Source/Runtime/Engine/Public/SceneView.h
@@ -1168,9 +1168,13 @@ public:
 	/** True if instanced stereo is enabled. */
 	bool bIsInstancedStereoEnabled;
 
-	/** True if multi-view is enabled. */
+	/** OLD variable governing multi-viewport rendering. */
+	UE_DEPRECATED(5.1, "bIsMultiViewEnabled has been deprecated. Use bIsMultiViewportEnabled")
 	bool bIsMultiViewEnabled;
 
+	/** True if multi-viewport instanced stereo rendering is enabled. */
+	bool bIsMultiViewportEnabled;
+
 	/** True if mobile multi-view is enabled. */
 	bool bIsMobileMultiViewEnabled;
 
@@ -1289,6 +1293,12 @@ public:
 	/** Initialization constructor. */
 	FSceneView(const FSceneViewInitOptions& InitOptions);
 
+	/** These are only needed because of deprecated members being accessed in them */
+	PRAGMA_DISABLE_DEPRECATION_WARNINGS
+	FSceneView(const FSceneView& Other) = default;
+	FSceneView(FSceneView&& Other) = default;
+	PRAGMA_ENABLE_DEPRECATION_WARNINGS
+
 #if DO_CHECK || USING_CODE_ANALYSIS
 	/** Verifies all the assertions made on members. */
 	bool VerifyMembersChecks() const;
diff --git a/Engine/Source/Runtime/RHI/Private/RHI.cpp b/Engine/Source/Runtime/RHI/Private/RHI.cpp
--- a/Engine/Source/Runtime/RHI/Private/RHI.cpp
+++ b/Engine/Source/Runtime/RHI/Private/RHI.cpp
@@ -2483,7 +2483,7 @@ void FGenericDataDrivenShaderPlatformInfo::ParseDataDrivenShaderInfo(const FConf
 	GET_SECTION_BOOL_HELPER(bSupportsVolumetricFog);
 	GET_SECTION_BOOL_HELPER(bSupportsIndexBufferUAVs);
 	GET_SECTION_BOOL_HELPER(bSupportsInstancedStereo);
-	GET_SECTION_BOOL_HELPER(bSupportsMultiView);
+	GET_SECTION_SUPPORT_HELPER(SupportsMultiViewport);
 	GET_SECTION_BOOL_HELPER(bSupportsMSAA);
 	GET_SECTION_BOOL_HELPER(bSupports4ComponentUAVReadWrite);
 	GET_SECTION_BOOL_HELPER(bSupportsRenderTargetWriteMask);
diff --git a/Engine/Source/Runtime/RHI/Public/RHI.h b/Engine/Source/Runtime/RHI/Public/RHI.h
--- a/Engine/Source/Runtime/RHI/Public/RHI.h
+++ b/Engine/Source/Runtime/RHI/Public/RHI.h
@@ -143,9 +143,21 @@ inline bool RHISupportsInstancedStereo(const FStaticShaderPlatform Platform)
 	return FDataDrivenShaderPlatformInfo::GetSupportsInstancedStereo(Platform);
 }
 
+UE_DEPRECATED(5.1, "RHISupportsMultiView has been deprecated. Use RHISupportsMultiViewport to avoid confusion.")
 inline bool RHISupportsMultiView(const FStaticShaderPlatform Platform)
 {
+	PRAGMA_DISABLE_DEPRECATION_WARNINGS
 	return FDataDrivenShaderPlatformInfo::GetSupportsMultiView(Platform);
+	PRAGMA_ENABLE_DEPRECATION_WARNINGS
+}
+
+/** 
+ * Can this platform implement instanced stereo rendering by rendering to multiple viewports.
+ * Note: run-time users should always check GRHISupportsArrayIndexFromAnyShader as well, since for some SPs (particularly PCD3D_SM5) minspec does not guarantee that feature.
+ **/
+inline bool RHISupportsMultiViewport(const FStaticShaderPlatform Platform)
+{
+	return FDataDrivenShaderPlatformInfo::GetSupportsMultiViewport(Platform) != ERHIFeatureSupport::Unsupported;
 }
 
 inline bool RHISupportsMSAA(const FStaticShaderPlatform Platform)
diff --git a/Engine/Source/Runtime/RHI/Public/RHIDefinitions.h b/Engine/Source/Runtime/RHI/Public/RHIDefinitions.h
--- a/Engine/Source/Runtime/RHI/Public/RHIDefinitions.h
+++ b/Engine/Source/Runtime/RHI/Public/RHIDefinitions.h
@@ -354,7 +354,7 @@ class RHI_API FGenericDataDrivenShaderPlatformInfo
 	uint32 bSupportsVolumetricFog: 1; // also used for FVVoxelization
 	uint32 bSupportsIndexBufferUAVs: 1;
 	uint32 bSupportsInstancedStereo: 1;
-	uint32 bSupportsMultiView: 1;
+	uint32 SupportsMultiViewport: int32(ERHIFeatureSupport::NumBits);
 	uint32 bSupportsMSAA: 1;
 	uint32 bSupports4ComponentUAVReadWrite: 1;
 	uint32 bSupportsRenderTargetWriteMask: 1;
@@ -596,12 +596,19 @@ public:
 		return Infos[Platform].bSupportsInstancedStereo;
 	}
 
+	UE_DEPRECATED(5.1, "bSupportsMultiView shader platform property has been deprecated. Use SupportsMultiViewport")
 	static FORCEINLINE_DEBUGGABLE const bool GetSupportsMultiView(const FStaticShaderPlatform Platform)
 	{
 		check(IsValid(Platform));
-		return Infos[Platform].bSupportsMultiView;
+		return ERHIFeatureSupport(Infos[Platform].SupportsMultiViewport) != ERHIFeatureSupport::Unsupported;
 	}
 
+	static FORCEINLINE_DEBUGGABLE const ERHIFeatureSupport GetSupportsMultiViewport(const FStaticShaderPlatform Platform)
+	{
+		check(IsValid(Platform));
+		return ERHIFeatureSupport(Infos[Platform].SupportsMultiViewport);
+	}	
+
 	static FORCEINLINE_DEBUGGABLE const bool GetSupportsMSAA(const FStaticShaderPlatform Platform)
 	{
 		check(IsValid(Platform));
diff --git a/Engine/Source/Runtime/RenderCore/Private/StereoRenderUtils.cpp b/Engine/Source/Runtime/RenderCore/Private/StereoRenderUtils.cpp
--- a/Engine/Source/Runtime/RenderCore/Private/StereoRenderUtils.cpp
+++ b/Engine/Source/Runtime/RenderCore/Private/StereoRenderUtils.cpp
@@ -19,6 +19,7 @@ RENDERCORE_API FStereoShaderAspects::FStereoShaderAspects(EShaderPlatform Platfo
 	// Would be nice to use URendererSettings, but not accessible in RenderCore
 	static FShaderPlatformCachedIniValue<bool> CVarInstancedStereo(TEXT("vr.InstancedStereo"));
 	static FShaderPlatformCachedIniValue<bool> CVarMobileMultiView(TEXT("vr.MobileMultiView"));
+	static FShaderPlatformCachedIniValue<bool> CVarMultiViewport(TEXT("vr.MultiViewport"));
 	static FShaderPlatformCachedIniValue<bool> CVarMobileHDR(TEXT("r.MobileHDR"));
 	
 	const bool bInstancedStereo = CVarInstancedStereo.Get(Platform);
@@ -26,6 +27,7 @@ RENDERCORE_API FStereoShaderAspects::FStereoShaderAspects(EShaderPlatform Platfo
 	const bool bMobilePlatform = IsMobilePlatform(Platform);
 	const bool bMobilePostprocessing = CVarMobileHDR.Get(Platform);
 	const bool bMobileMultiView = CVarMobileMultiView.Get(Platform);
+	const bool bMultiViewport = CVarMultiViewport.Get(Platform);
 	
 	bInstancedStereoNative = !bMobilePlatform && bInstancedStereo && RHISupportsInstancedStereo(Platform);
 	
@@ -43,7 +45,7 @@ RENDERCORE_API FStereoShaderAspects::FStereoShaderAspects(EShaderPlatform Platfo
 	}
 	
 	bInstancedStereoEnabled = bInstancedStereoNative || bMobileMultiViewFallback;
-	bInstancedMultiViewportEnabled = bInstancedStereoNative && RHISupportsMultiView(Platform);
+	bInstancedMultiViewportEnabled = bInstancedStereoNative && bMultiViewport && GRHISupportsArrayIndexFromAnyShader && RHISupportsMultiViewport(Platform);
 	bMobileMultiViewEnabled = bMobileMultiViewNative || bMobileMultiViewFallback;
 }
 	
diff --git a/Engine/Source/Runtime/Renderer/Private/PostProcess/PostProcessMaterial.cpp b/Engine/Source/Runtime/Renderer/Private/PostProcess/PostProcessMaterial.cpp
--- a/Engine/Source/Runtime/Renderer/Private/PostProcess/PostProcessMaterial.cpp
+++ b/Engine/Source/Runtime/Renderer/Private/PostProcess/PostProcessMaterial.cpp
@@ -534,7 +534,7 @@ FScreenPassTexture AddPostProcessMaterialPass(
 	PostProcessMaterialParameters->SceneWithoutSingleLayerWaterInvTextureSize = FVector2f(0.0f, 0.0f);
 	if (bHasValidSceneDepthWithoutWater)
 	{
-		const bool bIsInstancedStereoSideBySide = View.bIsInstancedStereoEnabled && !View.bIsMultiViewEnabled && IStereoRendering::IsStereoEyeView(View);
+		const bool bIsInstancedStereoSideBySide = View.bIsInstancedStereoEnabled && !View.bIsMobileMultiViewEnabled && IStereoRendering::IsStereoEyeView(View);
 		int32 WaterViewIndex = INDEX_NONE;
 		if (bIsInstancedStereoSideBySide)
 		{
diff --git a/Engine/Source/Runtime/Renderer/Private/SceneRendering.cpp b/Engine/Source/Runtime/Renderer/Private/SceneRendering.cpp
--- a/Engine/Source/Runtime/Renderer/Private/SceneRendering.cpp
+++ b/Engine/Source/Runtime/Renderer/Private/SceneRendering.cpp
@@ -178,6 +178,12 @@ static TAutoConsoleVariable<int32> CVarMobileMultiView(
 	TEXT("0 to disable mobile multi-view, 1 to enable.\n"),
 	ECVF_ReadOnly | ECVF_RenderThreadSafe);
 
+static TAutoConsoleVariable<int32> CVarMultiViewport(
+	TEXT("vr.MultiViewport"),
+	1,
+	TEXT("0 to disable multi-viewport instanced stereo rendering, 1 to enable its use (if supported by the RHI/shader platform).\n"),
+	ECVF_ReadOnly | ECVF_RenderThreadSafe);
+
 static TAutoConsoleVariable<int32> CVarRoundRobinOcclusion(
 	TEXT("vr.RoundRobinOcclusion"),
 	0,
@@ -3367,6 +3373,7 @@ void FSceneRenderer::RenderFinish(FRDGBuilder& GraphBuilder, FRDGTextureRef View
 				FSceneViewState* ViewState = (FSceneViewState*)View.State;
 				bool bViewParentOrFrozen = ViewState && (ViewState->HasViewParent() || ViewState->bIsFrozen);
 				bool bLocked = View.bIsLocked;
+				const bool bStereoView = IStereoRendering::IsStereoEyeView(View);
 
 				// display a warning if an ambient cubemap uses non-angular mipmap filtering
 				bool bShowAmbientCubemapMipGenSettingsWarning = false;
@@ -3384,17 +3391,23 @@ void FSceneRenderer::RenderFinish(FRDGBuilder& GraphBuilder, FRDGTextureRef View
 					}
 				}
 #endif
-				if ((GAreScreenMessagesEnabled && !GEngine->bSuppressMapWarnings) && (bViewParentOrFrozen || bLocked || bShowAmbientCubemapMipGenSettingsWarning || bAnyWarning))
+				if ((GAreScreenMessagesEnabled && !GEngine->bSuppressMapWarnings) && (bViewParentOrFrozen || bLocked || bStereoView || bShowAmbientCubemapMipGenSettingsWarning || bAnyWarning))
 				{
 					RDG_EVENT_SCOPE_CONDITIONAL(GraphBuilder, Views.Num() > 1, "View%d", ViewIndex);
 
+					const bool bPrimaryStereoView = IStereoRendering::IsAPrimaryView(View);
+					const bool bIsInstancedStereoEnabled = View.bIsInstancedStereoEnabled;
+					const bool bIsMultiViewportEnabled = View.bIsMultiViewportEnabled;
+					const bool bIsMobileMultiViewEnabled = View.bIsMobileMultiViewEnabled;
+
 					AddDrawCanvasPass(GraphBuilder, {}, View, Output,
 						[this, &ReadOnlyCVARCache, ViewState, GPUSkinCacheExtraRequiredMemory,
 						bLocked, bShowPrecomputedVisibilityWarning, bShowDemotedLocalMemoryWarning, bShowGlobalClipPlaneWarning, bShowDFAODisabledWarning, bShowDFDisabledWarning,
 						bViewParentOrFrozen, bShowSkylightWarning, bShowPointLightWarning, bShowShadowedLightOverflowWarning,
 						bShowMobileLowQualityLightmapWarning, bShowMobileMovableDirectionalLightWarning, bShowMobileDynamicCSMWarning, bMobileMissingSkyMaterial, 
 						bShowSkinCacheOOM, bSingleLayerWaterWarning, bShowNoSkyAtmosphereComponentWarning, bFxDebugDraw, FXInterface, bShowLocalExposureDisabledWarning,
-						bLumenEnabledButHasNoDataForTracing, bLumenEnabledButDisabledForTheProject, bNaniteEnabledButNoAtomics, bNaniteEnabledButDisabledInProject, bRealTimeSkyCaptureButNothingToCapture, bShowWaitingSkylight, bShowAmbientCubemapMipGenSettingsWarning]
+						bLumenEnabledButHasNoDataForTracing, bLumenEnabledButDisabledForTheProject, bNaniteEnabledButNoAtomics, bNaniteEnabledButDisabledInProject, bRealTimeSkyCaptureButNothingToCapture, bShowWaitingSkylight, bShowAmbientCubemapMipGenSettingsWarning,
+						bStereoView, bPrimaryStereoView, bIsInstancedStereoEnabled, bIsMultiViewportEnabled, bIsMobileMultiViewEnabled]
 						(FCanvas& Canvas)
 					{
 						// so it can get the screen size
@@ -3547,6 +3560,39 @@ void FSceneRenderer::RenderFinish(FRDGBuilder& GraphBuilder, FRDGTextureRef View
 							Writer.DrawLine(Message);
 						}
 
+#if !UE_BUILD_SHIPPING
+						if (bStereoView)
+						{
+							const TCHAR* SecondaryOrInstanced = bIsInstancedStereoEnabled ? TEXT("Instanced") : TEXT("Secondary");
+							FString ViewIdString = FString::Printf(TEXT("StereoView: %s"), bPrimaryStereoView ? TEXT("Primary") : SecondaryOrInstanced);
+							Writer.DrawLine(FText::FromString(ViewIdString));
+
+							// display information (in the primary view only) about the particular method used
+							if (bPrimaryStereoView)
+							{
+								const TCHAR* Technique = TEXT("Splitscreen-like");
+								if (bIsInstancedStereoEnabled)
+								{
+									if (bIsMultiViewportEnabled)
+									{
+										Technique = TEXT("Multi-viewport");
+									}
+									else if (bIsMobileMultiViewEnabled)
+									{
+										Technique = TEXT("Multi-view (mobile)");
+									}
+									else
+									{
+										Technique = TEXT("Instanced, clip planes");
+									}
+								}
+
+								FString TechniqueString = FString::Printf(TEXT("Stereo rendering method: %s"), Technique);
+								Writer.DrawLine(FText::FromString(TechniqueString));
+							}
+						}
+#endif
+
 #if WITH_EDITOR
 						FSkyLightSceneProxy* SkyLight = Scene->SkyLight;
 						if (bShowWaitingSkylight && SkyLight)
@@ -5089,7 +5135,7 @@ void FSceneRenderer::SetStereoViewport(FRHICommandList& RHICmdList, const FViewI
 {
 	if (View.IsInstancedStereoPass())
 	{
-		if (View.bIsMultiViewEnabled)
+		if (View.bIsMultiViewportEnabled)
 		{
 			const FViewInfo& LeftView = View;
 			const uint32 LeftMinX = LeftView.ViewRect.Min.X * ViewportScale;
diff --git a/Engine/Source/Runtime/Renderer/Private/ShadowRendering.cpp b/Engine/Source/Runtime/Renderer/Private/ShadowRendering.cpp
--- a/Engine/Source/Runtime/Renderer/Private/ShadowRendering.cpp
+++ b/Engine/Source/Runtime/Renderer/Private/ShadowRendering.cpp
@@ -905,7 +905,7 @@ void FProjectedShadowInfo::SetupProjectionStencilMask(
 		//       This means doing 2x renders, but letting the scissor rect kill the undersired half. Drawing the full mask once is easy, but since the outer
 		//       loop is over each view, the stencil mask is not retained when the right view comes around.
 		// TODO: Support instanced stereo properly in the projection stenciling pass.
-		const bool bIsInstancedStereoEmulated = View->bIsInstancedStereoEnabled && !View->bIsMultiViewEnabled && IStereoRendering::IsStereoEyeView(*View);
+		const bool bIsInstancedStereoEmulated = View->bIsInstancedStereoEnabled && !View->bIsMultiViewportEnabled && IStereoRendering::IsStereoEyeView(*View);
 		if (bIsInstancedStereoEmulated && ProjectionStencilingPasses.IsValidIndex(View->PrimaryViewIndex))
 		{
 			ensure(ProjectionStencilingPasses[View->PrimaryViewIndex]->GetInstanceCullingMode() == EInstanceCullingMode::Stereo);
@@ -1184,7 +1184,7 @@ void FProjectedShadowInfo::RenderProjection(
 		PassParameters->ShadowTexture1 = GraphBuilder.RegisterExternalTexture(RenderTargets.ColorTargets[1]);
 	}
 
-	const bool bIsInstancedStereoEmulated = View->bIsInstancedStereoEnabled && !View->bIsMultiViewEnabled && IStereoRendering::IsStereoEyeView(*View);
+	const bool bIsInstancedStereoEmulated = View->bIsInstancedStereoEnabled && !View->bIsMultiViewportEnabled && IStereoRendering::IsStereoEyeView(*View);
 	if (ViewIndex < ProjectionStencilingPasses.Num() && ProjectionStencilingPasses[ViewIndex] != nullptr)
 	{
 		// GPUCULL_TODO: get rid of const cast
diff --git a/Engine/Source/Runtime/Renderer/Private/SimpleMeshDrawCommandPass.cpp b/Engine/Source/Runtime/Renderer/Private/SimpleMeshDrawCommandPass.cpp
--- a/Engine/Source/Runtime/Renderer/Private/SimpleMeshDrawCommandPass.cpp
+++ b/Engine/Source/Runtime/Renderer/Private/SimpleMeshDrawCommandPass.cpp
@@ -15,7 +15,7 @@ FSimpleMeshDrawCommandPass::FSimpleMeshDrawCommandPass(const FSceneView& View, F
 	
 	TArray<int32, TFixedAllocator<2> > ViewIds;
 	ViewIds.Add(ViewInfo->GPUSceneViewId);
-	bUsingStereo = bEnableStereo && ViewInfo->bIsInstancedStereoEnabled && !View.bIsMultiViewEnabled && IStereoRendering::IsStereoEyeView(View);
+	bUsingStereo = bEnableStereo && ViewInfo->bIsInstancedStereoEnabled && !View.bIsMultiViewportEnabled && IStereoRendering::IsStereoEyeView(View);
 	if (bUsingStereo)
 	{
 		check(ViewInfo->GetInstancedView() != nullptr);
diff --git a/Engine/Source/Runtime/Renderer/Private/SingleLayerWaterRendering.cpp b/Engine/Source/Runtime/Renderer/Private/SingleLayerWaterRendering.cpp
--- a/Engine/Source/Runtime/Renderer/Private/SingleLayerWaterRendering.cpp
+++ b/Engine/Source/Runtime/Renderer/Private/SingleLayerWaterRendering.cpp
@@ -376,7 +376,7 @@ static FSceneWithoutWaterTextures AddCopySceneWithoutWaterPass(
 		auto PixelShader = View.ShaderMap->GetShader<FWaterRefractionCopyPS>(PermutationVector);
 
 		// if we have a particular case of ISR where two views are laid out in side by side, we should copy both views at once
-		const bool bIsInstancedStereoSideBySide = View.bIsInstancedStereoEnabled && !View.bIsMultiViewEnabled && IStereoRendering::IsStereoEyeView(View);
+		const bool bIsInstancedStereoSideBySide = View.bIsInstancedStereoEnabled && !View.bIsMobileMultiViewEnabled && IStereoRendering::IsStereoEyeView(View);
 		FIntRect RectToCopy = View.ViewRect;
 		if (bIsInstancedStereoSideBySide)
 		{
@@ -461,7 +461,7 @@ void FDeferredShadingSceneRenderer::RenderSingleLayerWaterReflections(
 			FIntVector DepthTextureSize = SceneWithoutWaterTextures.DepthTexture ? SceneWithoutWaterTextures.DepthTexture->Desc.GetSize() : FIntVector::ZeroValue;
 			const bool bShouldUseBilinearSamplerForDepth = SceneWithoutWaterTextures.DepthTexture && ShouldUseBilinearSamplerForDepthWithoutSingleLayerWater(SceneWithoutWaterTextures.DepthTexture->Desc.Format);
 
-			const bool bIsInstancedStereoSideBySide = View.bIsInstancedStereoEnabled && !View.bIsMultiViewEnabled && IStereoRendering::IsStereoEyeView(View);
+			const bool bIsInstancedStereoSideBySide = View.bIsInstancedStereoEnabled && !View.bIsMobileMultiViewEnabled && IStereoRendering::IsStereoEyeView(View);
 
 			Parameters.ScreenSpaceReflectionsTexture = ReflectionsColor ? ReflectionsColor : BlackDummyTexture;
 			Parameters.ScreenSpaceReflectionsSampler = TStaticSamplerState<SF_Point>::GetRHI();
diff --git a/Engine/Source/Runtime/Renderer/Private/TranslucentRendering.cpp b/Engine/Source/Runtime/Renderer/Private/TranslucentRendering.cpp
--- a/Engine/Source/Runtime/Renderer/Private/TranslucentRendering.cpp
+++ b/Engine/Source/Runtime/Renderer/Private/TranslucentRendering.cpp
@@ -690,7 +690,7 @@ bool FSceneRenderer::ShouldRenderTranslucency(ETranslucencyPass::Type Translucen
 FScreenPassTextureViewport FSeparateTranslucencyDimensions::GetInstancedStereoViewport(const FViewInfo& View) const
 {
 	FIntRect ViewRect = View.ViewRect;
-	if (View.IsInstancedStereoPass() && !View.bIsMultiViewEnabled)
+	if (View.IsInstancedStereoPass() && !View.bIsMultiViewportEnabled)
 	{
 		ViewRect.Max.X = ViewRect.Min.X + View.InstancedStereoWidth;
 	}
diff --git a/Engine/Source/Runtime/VulkanRHI/Private/VulkanCommandBuffer.cpp b/Engine/Source/Runtime/VulkanRHI/Private/VulkanCommandBuffer.cpp
--- a/Engine/Source/Runtime/VulkanRHI/Private/VulkanCommandBuffer.cpp
+++ b/Engine/Source/Runtime/VulkanRHI/Private/VulkanCommandBuffer.cpp
@@ -80,8 +80,8 @@ void FVulkanCmdBuffer::AllocMemory()
 {
 	// Assumes we are inside a lock for the pool
 	check(State == EState::NotAllocated);
-	FMemory::Memzero(CurrentViewport);
-	FMemory::Memzero(CurrentScissor);
+	CurrentViewports.Empty();
+	CurrentScissors.Empty();
 
 	VkCommandBufferAllocateInfo CreateCmdBufInfo;
 	ZeroVulkanStruct(CreateCmdBufInfo, VK_STRUCTURE_TYPE_COMMAND_BUFFER_ALLOCATE_INFO);
@@ -377,8 +377,8 @@ void FVulkanCmdBuffer::RefreshFenceStatus()
 			}
 			SubmittedWaitSemaphores.Reset();
 
-			FMemory::Memzero(CurrentViewport);
-			FMemory::Memzero(CurrentScissor);
+			CurrentViewports.Empty();
+			CurrentScissors.Empty();
 			CurrentStencilRef = 0;
 #if VULKAN_REUSE_FENCES
 			Fence->GetOwner()->ResetFence(Fence);
diff --git a/Engine/Source/Runtime/VulkanRHI/Private/VulkanCommandBuffer.h b/Engine/Source/Runtime/VulkanRHI/Private/VulkanCommandBuffer.h
--- a/Engine/Source/Runtime/VulkanRHI/Private/VulkanCommandBuffer.h
+++ b/Engine/Source/Runtime/VulkanRHI/Private/VulkanCommandBuffer.h
@@ -159,8 +159,8 @@ public:
 		NeedReset,
 	};
 
-	VkViewport CurrentViewport;
-	VkRect2D CurrentScissor;
+	TArray<VkViewport, TInlineAllocator<2>> CurrentViewports;
+	TArray<VkRect2D, TInlineAllocator<2>> CurrentScissors;
 	uint32 CurrentStencilRef;
 	EState State;
 	uint8 bNeedsDynamicStateSet			: 1;
diff --git a/Engine/Source/Runtime/VulkanRHI/Private/VulkanContext.h b/Engine/Source/Runtime/VulkanRHI/Private/VulkanContext.h
--- a/Engine/Source/Runtime/VulkanRHI/Private/VulkanContext.h
+++ b/Engine/Source/Runtime/VulkanRHI/Private/VulkanContext.h
@@ -38,6 +38,7 @@ public:
 
 	virtual void RHISetStreamSource(uint32 StreamIndex, FRHIBuffer* VertexBuffer, uint32 Offset) final override;
 	virtual void RHISetViewport(float MinX, float MinY, float MinZ, float MaxX, float MaxY, float MaxZ) final override;
+	virtual void RHISetStereoViewport(float LeftMinX, float RightMinX, float LeftMinY, float RightMinY, float MinZ, float LeftMaxX, float RightMaxX, float LeftMaxY, float RightMaxY, float MaxZ) override;
 	virtual void RHISetScissorRect(bool bEnable, uint32 MinX, uint32 MinY, uint32 MaxX, uint32 MaxY) final override;
 	virtual void RHISetGraphicsPipelineState(FRHIGraphicsPipelineState* GraphicsState, uint32 StencilRef, bool bApplyAdditionalState) final override;
 	virtual void RHISetShaderTexture(FRHIGraphicsShader* Shader, uint32 TextureIndex, FRHITexture* NewTexture) final override;
diff --git a/Engine/Source/Runtime/VulkanRHI/Private/VulkanPendingState.cpp b/Engine/Source/Runtime/VulkanRHI/Private/VulkanPendingState.cpp
--- a/Engine/Source/Runtime/VulkanRHI/Private/VulkanPendingState.cpp
+++ b/Engine/Source/Runtime/VulkanRHI/Private/VulkanPendingState.cpp
@@ -506,33 +506,34 @@ void FVulkanPendingGfxState::PrepareForDraw(FVulkanCmdBuffer* CmdBuffer)
 
 void FVulkanPendingGfxState::InternalUpdateDynamicStates(FVulkanCmdBuffer* Cmd)
 {
-	bool bInCmdNeedsDynamicState = Cmd->bNeedsDynamicStateSet;
-
-	bool bNeedsUpdateViewport = !Cmd->bHasViewport || (FMemory::Memcmp((const void*)&Cmd->CurrentViewport, (const void*)&Viewport, sizeof(VkViewport)) != 0);
+	const bool bNeedsUpdateViewport = !Cmd->bHasViewport || Viewports.Num() != Cmd->CurrentViewports.Num() || (FMemory::Memcmp((const void*)Cmd->CurrentViewports.GetData(), (const void*)Viewports.GetData(), Viewports.Num() * sizeof(VkViewport)) != 0);
 	// Validate and update Viewport
 	if (bNeedsUpdateViewport)
 	{
-		ensure(Viewport.width > 0 || Viewport.height > 0);
+		ensure(Viewports[0].width > 0 || Viewports[0].height > 0);
 
 		// Flip viewport on Y-axis to be uniform between DXC generated SPIR-V shaders (requires VK_KHR_maintenance1 extension)
-		VkViewport FlippedViewport = Viewport;
-		FlippedViewport.y += FlippedViewport.height;
-		FlippedViewport.height = -FlippedViewport.height;
-		VulkanRHI::vkCmdSetViewport(Cmd->GetHandle(), 0, 1, &FlippedViewport);
+		TArray<VkViewport, TInlineAllocator<2>> FlippedViewports = Viewports;
+		for (VkViewport& FlippedViewport : FlippedViewports)
+		{
+			FlippedViewport.y += FlippedViewport.height;
+			FlippedViewport.height = -FlippedViewport.height;
+		}
+		VulkanRHI::vkCmdSetViewport(Cmd->GetHandle(), 0, FlippedViewports.Num(), FlippedViewports.GetData());
 
-		FMemory::Memcpy(Cmd->CurrentViewport, Viewport);
+		Cmd->CurrentViewports = Viewports;
 		Cmd->bHasViewport = true;
 	}
 
-	bool bNeedsUpdateScissor = !Cmd->bHasScissor || (FMemory::Memcmp((const void*)&Cmd->CurrentScissor, (const void*)&Scissor, sizeof(VkRect2D)) != 0);
+	const bool bNeedsUpdateScissor = !Cmd->bHasScissor || Scissors.Num() != Cmd->CurrentScissors.Num() || (FMemory::Memcmp((const void*)Cmd->CurrentScissors.GetData(), (const void*)Scissors.GetData(), Scissors.Num() * sizeof(VkRect2D)) != 0);
 	if (bNeedsUpdateScissor)
 	{
-		VulkanRHI::vkCmdSetScissor(Cmd->GetHandle(), 0, 1, &Scissor);
-		FMemory::Memcpy(Cmd->CurrentScissor, Scissor);
+		VulkanRHI::vkCmdSetScissor(Cmd->GetHandle(), 0, Scissors.Num(), Scissors.GetData());
+		Cmd->CurrentScissors = Scissors;
 		Cmd->bHasScissor = true;
 	}
 
-	bool bNeedsUpdateStencil = !Cmd->bHasStencilRef || (Cmd->CurrentStencilRef != StencilRef);
+	const bool bNeedsUpdateStencil = !Cmd->bHasStencilRef || (Cmd->CurrentStencilRef != StencilRef);
 	if (bNeedsUpdateStencil)
 	{
 		VulkanRHI::vkCmdSetStencilReference(Cmd->GetHandle(), VK_STENCIL_FRONT_AND_BACK, StencilRef);
diff --git a/Engine/Source/Runtime/VulkanRHI/Private/VulkanPendingState.h b/Engine/Source/Runtime/VulkanRHI/Private/VulkanPendingState.h
--- a/Engine/Source/Runtime/VulkanRHI/Private/VulkanPendingState.h
+++ b/Engine/Source/Runtime/VulkanRHI/Private/VulkanPendingState.h
@@ -169,8 +169,8 @@ public:
 
 	void Reset()
 	{
-		FMemory::Memzero(Scissor);
-		FMemory::Memzero(Viewport);
+		Viewports.SetNumZeroed(1);
+		Scissors.SetNumZeroed(1);
 		StencilRef = 0;
 		bScissorEnable = false;
 
@@ -196,27 +196,43 @@ public:
 
 	void SetViewport(float MinX, float MinY, float MinZ, float MaxX, float MaxY, float MaxZ)
 	{
-		FMemory::Memzero(Viewport);
+		Viewports.SetNumZeroed(1);
 
-		Viewport.x = MinX;
-		Viewport.y = MinY;
-		Viewport.width = MaxX - MinX;
-		Viewport.height = MaxY - MinY;
-		Viewport.minDepth = MinZ;
+		Viewports[0].x = MinX;
+		Viewports[0].y = MinY;
+		Viewports[0].width = MaxX - MinX;
+		Viewports[0].height = MaxY - MinY;
+		Viewports[0].minDepth = MinZ;
 		if (MinZ == MaxZ)
 		{
 			// Engine pases in some cases MaxZ as 0.0
-			Viewport.maxDepth = MinZ + 1.0f;
+			Viewports[0].maxDepth = MinZ + 1.0f;
 		}
 		else
 		{
-			Viewport.maxDepth = MaxZ;
+			Viewports[0].maxDepth = MaxZ;
 		}
 
 		SetScissorRect((uint32)MinX, (uint32)MinY, (uint32)(MaxX - MinX), (uint32)(MaxY - MinY));
 		bScissorEnable = false;
 	}
 
+	void SetMultiViewport(const TArrayView<VkViewport>& InViewports)
+	{
+		Viewports = InViewports;
+
+		// Set the scissor rects appropriately.
+		Scissors.SetNumZeroed(Viewports.Num());
+		for (int32 Idx = 0; Idx < Scissors.Num(); ++Idx)
+		{
+			Scissors[Idx].offset.x = Viewports[Idx].x;
+			Scissors[Idx].offset.y = Viewports[Idx].y;
+			Scissors[Idx].extent.width = Viewports[Idx].width;
+			Scissors[Idx].extent.height = Viewports[Idx].height;
+		}
+		bScissorEnable = true;
+	}
+
 	inline void SetScissor(bool bInEnable, uint32 MinX, uint32 MinY, uint32 MaxX, uint32 MaxY)
 	{
 		if (bInEnable)
@@ -225,7 +241,8 @@ public:
 		}
 		else
 		{
-			SetScissorRect(Viewport.x, Viewport.y, Viewport.width, Viewport.height);
+			checkf(Viewports.Num() > 0, TEXT("At least one Viewport is expected to be configured."));
+			SetScissorRect(Viewports[0].x, Viewports[0].y, Viewports[0].width, Viewports[0].height);
 		}
 
 		bScissorEnable = bInEnable;
@@ -233,12 +250,12 @@ public:
 
 	inline void SetScissorRect(uint32 MinX, uint32 MinY, uint32 Width, uint32 Height)
 	{
-		FMemory::Memzero(Scissor);
+		Scissors.SetNumZeroed(1);
 
-		Scissor.offset.x = MinX;
-		Scissor.offset.y = MinY;
-		Scissor.extent.width = Width;
-		Scissor.extent.height = Height;
+		Scissors[0].offset.x = MinX;
+		Scissors[0].offset.y = MinY;
+		Scissors[0].extent.width = Width;
+		Scissors[0].extent.height = Height;
 	}
 
 	inline void SetStreamSource(uint32 StreamIndex, VkBuffer VertexBuffer, uint32 Offset)
@@ -363,8 +380,8 @@ public:
 	}
 
 protected:
-	VkViewport Viewport;
-	VkRect2D Scissor;
+	TArray<VkViewport, TInlineAllocator<2>> Viewports;
+	TArray<VkRect2D, TInlineAllocator<2>> Scissors;
 
 	EPrimitiveType PrimitiveType = PT_Num;
 	uint32 StencilRef;
diff --git a/Engine/Source/Runtime/VulkanRHI/Private/VulkanViewport.cpp b/Engine/Source/Runtime/VulkanRHI/Private/VulkanViewport.cpp
--- a/Engine/Source/Runtime/VulkanRHI/Private/VulkanViewport.cpp
+++ b/Engine/Source/Runtime/VulkanRHI/Private/VulkanViewport.cpp
@@ -1123,6 +1123,27 @@ void FVulkanCommandListContext::RHISetViewport(float MinX, float MinY, float Min
 	PendingGfxState->SetViewport(MinX, MinY, MinZ, MaxX, MaxY, MaxZ);
 }
 
+void FVulkanCommandListContext::RHISetStereoViewport(float LeftMinX, float RightMinX, float LeftMinY, float RightMinY, float MinZ, float LeftMaxX, float RightMaxX, float LeftMaxY, float RightMaxY, float MaxZ)
+{
+	TStaticArray<VkViewport, 2> Viewports;
+
+	Viewports[0].x = FMath::FloorToInt(LeftMinX);
+	Viewports[0].y = FMath::FloorToInt(LeftMinY);
+	Viewports[0].width = FMath::CeilToInt(LeftMaxX - LeftMinX);
+	Viewports[0].height = FMath::CeilToInt(LeftMaxY - LeftMinY);
+	Viewports[0].minDepth = MinZ;
+	Viewports[0].maxDepth = MaxZ;
+
+	Viewports[1].x = FMath::FloorToInt(RightMinX);
+	Viewports[1].y = FMath::FloorToInt(RightMinY);
+	Viewports[1].width = FMath::CeilToInt(RightMaxX - RightMinX);
+	Viewports[1].height = FMath::CeilToInt(RightMaxY - RightMinY);
+	Viewports[1].minDepth = MinZ;
+	Viewports[1].maxDepth = MaxZ;
+
+	PendingGfxState->SetMultiViewport(Viewports);
+}
+
 void FVulkanCommandListContext::RHISetMultipleViewports(uint32 Count, const FViewportBounds* Data)
 {
 	VULKAN_SIGNAL_UNIMPLEMENTED();
diff --git a/Engine/Source/Runtime/Windows/D3D11RHI/Private/D3D11Commands.cpp b/Engine/Source/Runtime/Windows/D3D11RHI/Private/D3D11Commands.cpp
--- a/Engine/Source/Runtime/Windows/D3D11RHI/Private/D3D11Commands.cpp
+++ b/Engine/Source/Runtime/Windows/D3D11RHI/Private/D3D11Commands.cpp
@@ -272,13 +272,46 @@ void FD3D11DynamicRHI::RHISetViewport(float MinX, float MinY, float MinZ, float
 
 static void ValidateScissorRect(const D3D11_VIEWPORT& Viewport, const D3D11_RECT& ScissorRect)
 {
-	ensure(ScissorRect.left   >= (LONG)Viewport.TopLeftX);
-	ensure(ScissorRect.top    >= (LONG)Viewport.TopLeftY);
-	ensure(ScissorRect.right  <= (LONG)Viewport.TopLeftX + (LONG)Viewport.Width);
+	ensure(ScissorRect.left >= (LONG)Viewport.TopLeftX);
+	ensure(ScissorRect.top >= (LONG)Viewport.TopLeftY);
+	ensure(ScissorRect.right <= (LONG)Viewport.TopLeftX + (LONG)Viewport.Width);
 	ensure(ScissorRect.bottom <= (LONG)Viewport.TopLeftY + (LONG)Viewport.Height);
 	ensure(ScissorRect.left <= ScissorRect.right && ScissorRect.top <= ScissorRect.bottom);
 }
 
+void FD3D11DynamicRHI::RHISetStereoViewport(float LeftMinX, float RightMinX, float LeftMinY, float RightMinY, float MinZ, float LeftMaxX, float RightMaxX, float LeftMaxY, float RightMaxY, float MaxZ)
+{
+	// Set up both viewports
+	D3D11_VIEWPORT StereoViewports[2] = { 0 };
+
+	StereoViewports[0].TopLeftX = FMath::FloorToInt(LeftMinX);
+	StereoViewports[0].TopLeftY = FMath::FloorToInt(LeftMinY);
+	StereoViewports[0].Width = FMath::CeilToInt(LeftMaxX - LeftMinX);
+	StereoViewports[0].Height = FMath::CeilToInt(LeftMaxY - LeftMinY);
+	StereoViewports[0].MinDepth = MinZ;
+	StereoViewports[0].MaxDepth = MaxZ;
+
+	StereoViewports[1].TopLeftX = FMath::FloorToInt(RightMinX);
+	StereoViewports[1].TopLeftY = FMath::FloorToInt(RightMinY);
+	StereoViewports[1].Width = FMath::CeilToInt(RightMaxX - RightMinX);
+	StereoViewports[1].Height = FMath::CeilToInt(RightMaxY - RightMinY);
+	StereoViewports[1].MinDepth = MinZ;
+	StereoViewports[1].MaxDepth = MaxZ;
+
+	D3D11_RECT ScissorRects[2] =
+	{
+		{ StereoViewports[0].TopLeftX, StereoViewports[0].TopLeftY, StereoViewports[0].TopLeftX + StereoViewports[0].Width, StereoViewports[0].TopLeftY + StereoViewports[0].Height },
+		{ StereoViewports[1].TopLeftX, StereoViewports[1].TopLeftY, StereoViewports[1].TopLeftX + StereoViewports[1].Width, StereoViewports[1].TopLeftY + StereoViewports[1].Height }
+	};
+
+	ValidateScissorRect(StereoViewports[0], ScissorRects[0]);
+	ValidateScissorRect(StereoViewports[1], ScissorRects[1]);
+
+	StateCache.SetViewports(2, StereoViewports);
+	// Set the scissor rect appropriately.
+	Direct3DDeviceIMContext->RSSetScissorRects(2, ScissorRects);
+}
+
 void FD3D11DynamicRHI::RHISetScissorRect(bool bEnable,uint32 MinX,uint32 MinY,uint32 MaxX,uint32 MaxY)
 {
 	D3D11_VIEWPORT Viewport;
diff --git a/Engine/Source/Runtime/Windows/D3D11RHI/Private/D3D11RHIPrivate.h b/Engine/Source/Runtime/Windows/D3D11RHI/Private/D3D11RHIPrivate.h
--- a/Engine/Source/Runtime/Windows/D3D11RHI/Private/D3D11RHIPrivate.h
+++ b/Engine/Source/Runtime/Windows/D3D11RHI/Private/D3D11RHIPrivate.h
@@ -577,6 +577,7 @@ public:
 	virtual void RHISetStreamSource(uint32 StreamIndex, FRHIBuffer* VertexBuffer, uint32 Offset) final override;
 	virtual void RHISetRasterizerState(FRHIRasterizerState* NewState) final override;
 	virtual void RHISetViewport(float MinX, float MinY, float MinZ, float MaxX, float MaxY, float MaxZ) final override;
+	virtual void RHISetStereoViewport(float LeftMinX, float RightMinX, float LeftMinY, float RightMinY, float MinZ, float LeftMaxX, float RightMaxX, float LeftMaxY, float RightMaxY, float MaxZ) override;
 	virtual void RHISetScissorRect(bool bEnable, uint32 MinX, uint32 MinY, uint32 MaxX, uint32 MaxY) final override;
 	virtual void RHISetBoundShaderState(FRHIBoundShaderState* BoundShaderState) final override;
 	virtual void RHISetGraphicsPipelineState(FRHIGraphicsPipelineState* GraphicsState, uint32 StencilRef, bool bApplyAdditionalState) final override;
