Prevent crash when using forcing number of distance field mips
- If r.DistanceFields.Debug.ForceNumMips is set to a value higher than DistanceField::NumMips it caused a crash.
- Limit DebugForceNumMips to never exceed DistanceField::NumMips

diff --git a/Engine/Source/Runtime/Renderer/Private/DistanceFieldStreaming.cpp b/Engine/Source/Runtime/Renderer/Private/DistanceFieldStreaming.cpp
--- a/Engine/Source/Runtime/Renderer/Private/DistanceFieldStreaming.cpp
+++ b/Engine/Source/Runtime/Renderer/Private/DistanceFieldStreaming.cpp
@@ -1101,7 +1101,7 @@ void FDistanceFieldSceneData::GenerateStreamingRequests(
 			PassParameters->RWDistanceFieldAssetWantedNumMips = GraphBuilder.CreateUAV(WantedNumMips);
 			PassParameters->RWDistanceFieldAssetStreamingRequests = GraphBuilder.CreateUAV(StreamingRequestsBuffer);
 			PassParameters->DistanceFieldObjectBuffers = DistanceField::SetupObjectBufferParameters(GraphBuilder, *this);
-			PassParameters->DebugForceNumMips = CVarDebugForceNumMips.GetValueOnRenderThread();
+			PassParameters->DebugForceNumMips = FMath::Clamp(CVarDebugForceNumMips.GetValueOnRenderThread(), 0, DistanceField::NumMips);
 			extern int32 GAOGlobalDistanceFieldNumClipmaps;
 			// Request Mesh SDF mips based off of the Global SDF clipmaps
 			PassParameters->Mip1WorldTranslatedCenter = FVector3f(View.ViewMatrices.GetViewOrigin() + View.ViewMatrices.GetPreViewTranslation());
