Fix hair strands shadow culling when hair instance is not visible in primary view.

diff --git a/Engine/Shaders/Private/HairStrands/HairStrandsDebug.usf b/Engine/Shaders/Private/HairStrands/HairStrandsDebug.usf
--- a/Engine/Shaders/Private/HairStrands/HairStrandsDebug.usf
+++ b/Engine/Shaders/Private/HairStrands/HairStrandsDebug.usf
@@ -2337,4 +2337,96 @@ void MainPS(
 
 #endif // SHADER_PLOTSAMPLE
 
-//////////////////////////////////////////////////////////////////////////////////////////////////////////////////////////////
\ No newline at end of file
+//////////////////////////////////////////////////////////////////////////////////////////////////////////////////////////////
+// Display groom culling from shadow perspective 
+
+#if SHADER_SHADOW_CULLING
+
+#include "../ShaderPrint.ush"
+#include "../Matrices.ush"
+
+void AddFrustumWS(float4x4 In, float4 InColor, bool bInv)
+{
+	const float Near = 0.0001f;
+	const float Far = 0.9999f;
+
+	const float4x4 T = bInv ? Inverse(In) : In;
+	float4 P00 = mul(float4(-1, -1, Near, 1), T); P00 /= P00.w;
+	float4 P10 = mul(float4( 1, -1, Near, 1), T); P10 /= P10.w;
+	float4 P20 = mul(float4( 1,  1, Near, 1), T); P20 /= P20.w;
+	float4 P30 = mul(float4(-1,  1, Near, 1), T); P30 /= P30.w;
+
+	float4 P01 = mul(float4(-1, -1, Far, 1),  T); P01 /= P01.w;
+	float4 P11 = mul(float4( 1, -1, Far, 1),  T); P11 /= P11.w;
+	float4 P21 = mul(float4( 1,  1, Far, 1),  T); P21 /= P21.w;
+	float4 P31 = mul(float4(-1,  1, Far, 1),  T); P31 /= P31.w;
+
+	AddLineWS(P00.xyz, P10.xyz, InColor);
+	AddLineWS(P10.xyz, P20.xyz, InColor);
+	AddLineWS(P20.xyz, P30.xyz, InColor);
+	AddLineWS(P30.xyz, P00.xyz, InColor);
+
+	AddLineWS(P01.xyz, P11.xyz, InColor);
+	AddLineWS(P11.xyz, P21.xyz, InColor);
+	AddLineWS(P21.xyz, P31.xyz, InColor);
+	AddLineWS(P31.xyz, P01.xyz, InColor);
+
+	AddLineWS(P00.xyz, P01.xyz, InColor);
+	AddLineWS(P10.xyz, P11.xyz, InColor);
+	AddLineWS(P20.xyz, P21.xyz, InColor);
+	AddLineWS(P30.xyz, P31.xyz, InColor);
+}
+
+Buffer<float> InstanceBoundInLightSpace;
+Buffer<float> InstanceBoundInWorldSpace;
+Buffer<uint>  InstanceIntersection;
+
+uint InstanceCount;
+
+float3 LightCenter;
+float3 LightExtent;
+
+float4x4 LightToWorld;
+
+float4x4 ViewWorldToProj;
+float4x4 ViewProjToWorld;
+
+[numthreads(1, 1, 1)]
+void MainCS(uint GroupIndex : SV_GroupIndex, uint3 DispatchThreadId : SV_DispatchThreadID)
+{
+	// Draw AABB of the instance groups
+	if (any(DispatchThreadId != 0))
+	{
+		return;
+	}
+
+	FShaderPrintContext Context = InitShaderPrintContext(true, uint2(50, 50));
+
+	// View Frustum
+	//AddFrustumWS(ViewProjToWorld, ColorBlue, false);
+	AddFrustumWS(ViewWorldToProj, ColorBlue, true);
+
+	// Light bound
+	AddAABBWS(LightCenter - LightExtent, LightCenter + LightExtent, ColorPurple);
+
+	for (uint InstanceIndex = 0; InstanceIndex < InstanceCount; ++InstanceIndex)
+	{	
+		const uint Index6 = InstanceIndex * 6;
+		// In light space
+		{
+			const bool bIntersect = InstanceIntersection[InstanceIndex] > 0;
+			const float3 MinP = float3(InstanceBoundInLightSpace[Index6 + 0], InstanceBoundInLightSpace[Index6 + 1], InstanceBoundInLightSpace[Index6 + 2]);
+			const float3 MaxP = float3(InstanceBoundInLightSpace[Index6 + 3], InstanceBoundInLightSpace[Index6 + 4], InstanceBoundInLightSpace[Index6 + 5]);
+			// Transform back from light space to world space
+			AddOBBWS(MinP, MaxP, bIntersect ? ColorGreen : ColorYellow, LightToWorld);
+		}
+	
+		// In world space
+		{
+			const float3 MinP = float3(InstanceBoundInWorldSpace[Index6 + 0], InstanceBoundInWorldSpace[Index6 + 1], InstanceBoundInWorldSpace[Index6 + 2]);
+			const float3 MaxP = float3(InstanceBoundInWorldSpace[Index6 + 3], InstanceBoundInWorldSpace[Index6 + 4], InstanceBoundInWorldSpace[Index6 + 5]);
+			AddAABBWS(MinP, MaxP, ColorRed);
+		}
+	}
+}
+#endif // SHADER_SHADOW_CULLING
diff --git a/Engine/Source/Runtime/Renderer/Private/HairStrands/HairStrandsData.h b/Engine/Source/Runtime/Renderer/Private/HairStrands/HairStrandsData.h
--- a/Engine/Source/Runtime/Renderer/Private/HairStrands/HairStrandsData.h
+++ b/Engine/Source/Runtime/Renderer/Private/HairStrands/HairStrandsData.h
@@ -10,6 +10,7 @@
 #include "RendererInterface.h"
 #include "RHIGPUReadback.h"
 #include "Shader.h"
+#include "ConvexVolume.h"
 
 class FLightSceneInfo;
 class FPrimitiveSceneProxy;
@@ -358,6 +359,7 @@ struct FHairStrandsDebugData
 		SHADER_PARAMETER_RDG_BUFFER_SRV(Buffer, Debug_SampleCounter)
 	END_SHADER_PARAMETER_STRUCT()
 
+	// Plot Data
 	struct ShadingInfo
 	{
 		FVector3f	BaseColor;
@@ -378,32 +380,62 @@ struct FHairStrandsDebugData
 	static const uint32 MaxShadingPointCount = 32;
 	static const uint32 MaxSampleCount = 1024 * 32;
 
-	struct Data
+	struct FPlotData
 	{
 		FRDGBufferRef ShadingPointBuffer = nullptr;
 		FRDGBufferRef ShadingPointCounter = nullptr;
 		FRDGBufferRef SampleBuffer = nullptr;
 		FRDGBufferRef SampleCounter = nullptr;
-	} Resources;
+	} PlotData;
 
 	bool IsPlotDataValid() const
 	{
-		return Resources.ShadingPointBuffer && Resources.ShadingPointCounter && Resources.SampleBuffer && Resources.SampleCounter;
+		return PlotData.ShadingPointBuffer && PlotData.ShadingPointCounter && PlotData.SampleBuffer && PlotData.SampleCounter;
 	}
 
-	static Data CreateData(FRDGBuilder& GraphBuilder);
-	static void SetParameters(FRDGBuilder& GraphBuilder, const Data& In, FWriteParameters& Out);
-	static void SetParameters(FRDGBuilder& GraphBuilder, const Data& In, FReadParameters& Out);
+	static FPlotData CreatePlotData(FRDGBuilder& GraphBuilder);
+	static void SetParameters(FRDGBuilder& GraphBuilder, const FPlotData& In, FWriteParameters& Out);
+	static void SetParameters(FRDGBuilder& GraphBuilder, const FPlotData& In, FReadParameters& Out);
 
 	// PPLL debug data
+	struct FPPLLData
+	{
+		FRDGTextureRef	NodeCounterTexture = nullptr;
+		FRDGTextureRef	NodeIndexTexture = nullptr;
+		FRDGBufferRef	NodeDataBuffer = nullptr;
+	} PPLLData;
+
 	bool IsPPLLDataValid() const 
 	{ 
-		return PPLLNodeCounterTexture && PPLLNodeIndexTexture && PPLLNodeDataBuffer; 
+		return PPLLData.NodeCounterTexture && PPLLData.NodeIndexTexture && PPLLData.NodeDataBuffer;
 	}
 
-	FRDGTextureRef	PPLLNodeCounterTexture = nullptr;
-	FRDGTextureRef	PPLLNodeIndexTexture = nullptr;
-	FRDGBufferRef	PPLLNodeDataBuffer = nullptr;
+	// Instance cull data
+	struct FCullData
+	{
+		struct FBound
+		{
+			FVector3f Min;
+			FVector3f Max;
+		};
+
+		struct FLight
+		{
+			FMatrix WorldToLight;
+			FMatrix LightToWorld;
+			FVector3f Center;
+			FVector3f Extent;
+			FConvexVolume ViewFrustumInLightSpace;
+			TArray<FBound> InstanceBoundInLightSpace;
+			TArray<FBound> InstanceBoundInWorldSpace;
+			TArray<uint32> InstanceIntersection;
+		};
+
+		bool bIsValid = false;
+		TArray<FLight> DirectionalLights;
+		FConvexVolume ViewFrustum;
+	} CullData;
+
 };
 
 typedef TArray<FHairStrandsMacroGroupData, SceneRenderingAllocator> FHairStrandsMacroGroupDatas;
@@ -424,8 +456,9 @@ struct FHairStrandsViewData
 	FHairStrandsDebugData DebugData;
 
 	// Transient: store all light visible in primary view(s)
+	struct FDirectionalLightCullData { const FLightSceneInfo* LightInfo = nullptr; FConvexVolume ViewFrustumInLightSpace; };
 	TArray<const FLightSceneInfo*> VisibleShadowCastingLights;
-	TArray<FSphere> VisibleShadowCastingBounds;
+	TArray<FDirectionalLightCullData> VisibleShadowCastingDirectionalLights;
 };
 
 // View State data (i.e., persistent across frame)
@@ -476,5 +509,4 @@ namespace HairStrands
 	bool HasHairStrandsVisible(const TArray<FViewInfo>& Views);
 
 	void AddVisibleShadowCastingLight(const FScene& Scene, TArray<FViewInfo>& Views, const FLightSceneInfo* LightSceneInfo);
-	void AddVisibleShadowCastingLight(const FScene& Scene, TArray<FViewInfo>& Views, const FSphere& Bounds);
 }
\ No newline at end of file
diff --git a/Engine/Source/Runtime/Renderer/Private/HairStrands/HairStrandsDebug.cpp b/Engine/Source/Runtime/Renderer/Private/HairStrands/HairStrandsDebug.cpp
--- a/Engine/Source/Runtime/Renderer/Private/HairStrands/HairStrandsDebug.cpp
+++ b/Engine/Source/Runtime/Renderer/Private/HairStrands/HairStrandsDebug.cpp
@@ -91,9 +91,9 @@ bool IsHairStrandsClusterDebugAABBEnable()
 	return GHairStrandsClusterDebug > 1;
 }
 
-FHairStrandsDebugData::Data FHairStrandsDebugData::CreateData(FRDGBuilder& GraphBuilder)
+FHairStrandsDebugData::FPlotData FHairStrandsDebugData::CreatePlotData(FRDGBuilder& GraphBuilder)
 {
-	FHairStrandsDebugData::Data Out;
+	FHairStrandsDebugData::FPlotData Out;
 	Out.ShadingPointBuffer = GraphBuilder.CreateBuffer(FRDGBufferDesc::CreateStructuredDesc(sizeof(ShadingInfo), MaxShadingPointCount), TEXT("Hair.DebugShadingPoint"));
 	Out.ShadingPointCounter = GraphBuilder.CreateBuffer(FRDGBufferDesc::CreateBufferDesc(sizeof(uint32), 1), TEXT("Hair.DebugShadingPointCounter"));
 	Out.SampleBuffer = GraphBuilder.CreateBuffer(FRDGBufferDesc::CreateStructuredDesc(sizeof(Sample), MaxSampleCount), TEXT("Hair.DebugSample"));
@@ -103,7 +103,7 @@ FHairStrandsDebugData::Data FHairStrandsDebugData::CreateData(FRDGBuilder& Graph
 	return Out;
 }
 
-void FHairStrandsDebugData::SetParameters(FRDGBuilder& GraphBuilder, const FHairStrandsDebugData::Data& In, FHairStrandsDebugData::FWriteParameters& Out)
+void FHairStrandsDebugData::SetParameters(FRDGBuilder& GraphBuilder, const FHairStrandsDebugData::FPlotData& In, FHairStrandsDebugData::FWriteParameters& Out)
 {
 	Out.Debug_MaxSampleCount = FHairStrandsDebugData::MaxSampleCount;
 	Out.Debug_MaxShadingPointCount = FHairStrandsDebugData::MaxShadingPointCount;
@@ -113,7 +113,7 @@ void FHairStrandsDebugData::SetParameters(FRDGBuilder& GraphBuilder, const FHair
 	Out.Debug_SampleCounter = GraphBuilder.CreateUAV(In.SampleCounter, PF_R32_UINT);
 }
 
-void FHairStrandsDebugData::SetParameters(FRDGBuilder& GraphBuilder, const FHairStrandsDebugData::Data& In, FHairStrandsDebugData::FReadParameters& Out)
+void FHairStrandsDebugData::SetParameters(FRDGBuilder& GraphBuilder, const FHairStrandsDebugData::FPlotData& In, FHairStrandsDebugData::FReadParameters& Out)
 {
 	Out.Debug_MaxSampleCount = FHairStrandsDebugData::MaxSampleCount;
 	Out.Debug_MaxShadingPointCount = FHairStrandsDebugData::MaxShadingPointCount;
@@ -367,12 +367,19 @@ static void AddDebugHairPrintPass(
 		const FHairStrandsMacroGroupData& MacroGroup = MacroGroupDatas[MacroGroupIndex];
 		for (const FHairStrandsMacroGroupData::PrimitiveInfo& PrimitiveInfo : MacroGroup.PrimitivesInfos)
 		{
-			uint32 PrimitiveID = ~0u;
-			if (const FPrimitiveSceneInfo* SceneInfo = PrimitiveInfo.PrimitiveSceneProxy->GetPrimitiveSceneInfo())
+			if (PrimitiveInfo.PrimitiveSceneProxy)
 			{
-				PrimitiveID = SceneInfo->GetIndex();
+				uint32 PrimitiveID = ~0u;
+				if (const FPrimitiveSceneInfo* SceneInfo = PrimitiveInfo.PrimitiveSceneProxy->GetPrimitiveSceneInfo())
+				{
+					PrimitiveID = SceneInfo->GetIndex();
+				}
+				InstanceIDs.Add(PrimitiveID);
+			}
+			else
+			{
+				InstanceIDs.Add(0u);
 			}
-			InstanceIDs.Add(PrimitiveID);
 		}
 	}
 	FRDGBufferRef InstancesIDBuffer = CreateVertexBuffer(GraphBuilder, TEXT("Hair.Debug.InstanceIDs"), FRDGBufferDesc::CreateBufferDesc(4, InstanceIDs.Num()), InstanceIDs.GetData(), 4u * InstanceIDs.Num());
@@ -417,6 +424,112 @@ static void AddDebugHairPrintPass(
 		FIntVector(1, 1, 1));
 }
 
+///////////////////////////////////////////////////////////////////////////////////////////////////
+class FHairDebugShadowCullingCS : public FGlobalShader
+{
+	DECLARE_GLOBAL_SHADER(FHairDebugShadowCullingCS);
+	SHADER_USE_PARAMETER_STRUCT(FHairDebugShadowCullingCS, FGlobalShader);
+
+	BEGIN_SHADER_PARAMETER_STRUCT(FParameters, )
+		SHADER_PARAMETER(uint32, InstanceCount)
+		SHADER_PARAMETER(FVector3f, LightCenter)
+		SHADER_PARAMETER(FVector3f, LightExtent)
+		SHADER_PARAMETER(FMatrix44f, LightToWorld)
+		SHADER_PARAMETER(FMatrix44f, ViewWorldToProj)
+		SHADER_PARAMETER(FMatrix44f, ViewProjToWorld)
+		SHADER_PARAMETER_RDG_BUFFER_SRV(Buffer, InstanceBoundInWorldSpace)
+		SHADER_PARAMETER_RDG_BUFFER_SRV(Buffer, InstanceBoundInLightSpace)
+		SHADER_PARAMETER_RDG_BUFFER_SRV(Buffer, InstanceIntersection)
+		SHADER_PARAMETER_STRUCT_REF(FViewUniformShaderParameters, ViewUniformBuffer)
+		SHADER_PARAMETER_STRUCT_INCLUDE(ShaderPrint::FShaderParameters, ShaderPrintUniformBuffer)
+		SHADER_PARAMETER_RDG_UNIFORM_BUFFER(FHairStrandsViewUniformParameters, HairStrands)
+	END_SHADER_PARAMETER_STRUCT()
+
+public:
+	static bool ShouldCompilePermutation(const FGlobalShaderPermutationParameters& Parameters) { return IsHairStrandsSupported(EHairStrandsShaderType::Tool, Parameters.Platform); }
+	static void ModifyCompilationEnvironment(const FGlobalShaderPermutationParameters& Parameters, FShaderCompilerEnvironment& OutEnvironment)
+	{
+		// Skip optimization for avoiding long compilation time due to large UAV writes
+		FGlobalShader::ModifyCompilationEnvironment(Parameters, OutEnvironment);
+		OutEnvironment.CompilerFlags.Add(CFLAG_Debug);
+		OutEnvironment.SetDefine(TEXT("SHADER_SHADOW_CULLING"), 1);
+	}
+};
+
+IMPLEMENT_GLOBAL_SHADER(FHairDebugShadowCullingCS, "/Engine/Private/HairStrands/HairStrandsDebug.usf", "MainCS", SF_Compute);
+
+static void AddDebugHairShadowCullingPass(
+	FRDGBuilder& GraphBuilder,
+	const FScene* Scene,
+	const FViewInfo* View)
+{
+	if (!Scene || !View || !View->HairStrandsViewData.UniformBuffer) return;
+
+	const uint32 LightCount = View->HairStrandsViewData.DebugData.CullData.DirectionalLights.Num();
+	uint32 InstanceCount = 0;
+	for (auto It : View->HairStrandsViewData.DebugData.CullData.DirectionalLights)
+	{
+		InstanceCount += It.InstanceBoundInLightSpace.Num();
+	}
+
+	if (LightCount == 0 || InstanceCount == 0)
+	{
+		return;
+	}
+
+	if (!TryEnableShaderDrawAndShaderPrint(*View, 2000u, 2000u))
+	{
+		return;
+	}
+
+	auto CreateBoundBuffer = [&](const TArray<FHairStrandsDebugData::FCullData::FBound>& In)
+	{
+		TArray<float> Bound; Bound.Reserve(In.Num() * 6);
+		for (auto& B : In)
+		{
+			Bound.Add(B.Min.X);
+			Bound.Add(B.Min.Y);
+			Bound.Add(B.Min.Z);
+			Bound.Add(B.Max.X);
+			Bound.Add(B.Max.Y);
+			Bound.Add(B.Max.Z);
+		}
+		return CreateVertexBuffer(GraphBuilder, TEXT("Hair.Debug.InstanceBounds"), FRDGBufferDesc::CreateBufferDesc(4, Bound.Num()), Bound.GetData(), 4u * Bound.Num());
+	};
+
+	uint32 LightIndex = 0;
+	for (auto It : View->HairStrandsViewData.DebugData.CullData.DirectionalLights)
+	{
+		FRDGBufferRef InstanceBoundInLightSpaceBuffer = CreateBoundBuffer(It.InstanceBoundInLightSpace);
+		FRDGBufferRef InstanceBoundInWorldSpaceBuffer = CreateBoundBuffer(It.InstanceBoundInWorldSpace);
+		FRDGBufferRef InstanceIntersectionBuffer = CreateVertexBuffer(GraphBuilder, TEXT("Hair.Debug.InstanceIntersection"), FRDGBufferDesc::CreateBufferDesc(4, It.InstanceIntersection.Num()), It.InstanceIntersection.GetData(), 4u * It.InstanceIntersection.Num());
+
+		FHairDebugShadowCullingCS::FParameters* Parameters = GraphBuilder.AllocParameters<FHairDebugShadowCullingCS::FParameters>();
+		Parameters->InstanceCount = InstanceCount;
+		Parameters->LightCenter = It.Center;
+		Parameters->LightExtent = It.Extent;
+		Parameters->LightToWorld = FMatrix44f(It.LightToWorld);
+		Parameters->ViewWorldToProj = FMatrix44f(View->ViewMatrices.GetViewProjectionMatrix());
+		Parameters->ViewProjToWorld = FMatrix44f(View->ViewMatrices.GetInvViewProjectionMatrix());
+		Parameters->InstanceBoundInLightSpace = GraphBuilder.CreateSRV(InstanceBoundInLightSpaceBuffer, PF_R32_FLOAT);
+		Parameters->InstanceBoundInWorldSpace = GraphBuilder.CreateSRV(InstanceBoundInWorldSpaceBuffer, PF_R32_FLOAT);
+		Parameters->InstanceIntersection = GraphBuilder.CreateSRV(InstanceIntersectionBuffer, PF_R32_UINT);
+		Parameters->ViewUniformBuffer = View->ViewUniformBuffer;
+		ShaderPrint::SetParameters(GraphBuilder, View->ShaderPrintData, Parameters->ShaderPrintUniformBuffer);
+		Parameters->HairStrands = View->HairStrandsViewData.UniformBuffer;
+		TShaderMapRef<FHairDebugShadowCullingCS> ComputeShader(View->ShaderMap);
+
+		ClearUnusedGraphResources(ComputeShader, Parameters);
+
+		FComputeShaderUtils::AddPass(
+			GraphBuilder,
+			RDG_EVENT_NAME("HairStrands::DebugShadowCulling(%d/%d)", LightIndex++, LightCount),
+			ComputeShader,
+			Parameters,
+			FIntVector(1, 1, 1));
+	}
+}
+
 ///////////////////////////////////////////////////////////////////////////////////////////////////
 class FHairDebugPS : public FGlobalShader
 {
@@ -974,7 +1087,7 @@ IMPLEMENT_GLOBAL_SHADER(FHairStrandsPlotSamplePS, "/Engine/Private/HairStrands/H
 static void AddPlotSamplePass(
 	FRDGBuilder& GraphBuilder,
 	const FViewInfo& View,
-	const FHairStrandsDebugData::Data& DebugData,
+	const FHairStrandsDebugData::FPlotData& DebugData,
 	FRDGTextureRef& OutputTexture)
 {
 	FSceneTextureParameters SceneTextures = GetSceneTextureParameters(GraphBuilder, View);
@@ -1291,12 +1404,17 @@ static void InternalRenderHairStrandsDebugInfo(
 		}
 		if (HairData.DebugData.IsPlotDataValid())
 		{
-			AddPlotSamplePass(GraphBuilder, View, HairData.DebugData.Resources, SceneColorTexture);
+			AddPlotSamplePass(GraphBuilder, View, HairData.DebugData.PlotData, SceneColorTexture);
 		}	
 	}
 
 	float ClusterY = 38;
 
+	if (View.HairStrandsViewData.DebugData.CullData.bIsValid)
+	{
+		AddDebugHairShadowCullingPass(GraphBuilder, Scene, &View);
+	}
+
 	if (HairDebugMode == EHairDebugMode::MacroGroups)
 	{
 		AddDebugHairPrintPass(GraphBuilder, Scene, &View, HairDebugMode, HairData.VisibilityData, HairData.MacroGroupDatas, HairData.MacroGroupResources, SceneTextures.Stencil);
@@ -1382,13 +1500,13 @@ static void InternalRenderHairStrandsDebugInfo(
 		ShaderPrint::SetEnabled(true);
 		ShaderPrint::RequestSpaceForCharacters(256);
 
-		const FIntPoint PPLLResolution = HairData.DebugData.PPLLNodeIndexTexture->Desc.Extent;
+		const FIntPoint PPLLResolution = HairData.DebugData.PPLLData.NodeIndexTexture->Desc.Extent;
 		FHairVisibilityDebugPPLLCS::FParameters* PassParameters = GraphBuilder.AllocParameters<FHairVisibilityDebugPPLLCS::FParameters>();
 		PassParameters->PPLLMeanListElementCountPerPixel = GetHairStrandsMeanSamplePerPixel();
-		PassParameters->PPLLMaxTotalListElementCount = HairData.DebugData.PPLLNodeDataBuffer->Desc.NumElements;
-		PassParameters->PPLLCounter = HairData.DebugData.PPLLNodeCounterTexture;
-		PassParameters->PPLLNodeIndex = HairData.DebugData.PPLLNodeIndexTexture;
-		PassParameters->PPLLNodeData = GraphBuilder.CreateSRV(FRDGBufferSRVDesc(HairData.DebugData.PPLLNodeDataBuffer));
+		PassParameters->PPLLMaxTotalListElementCount = HairData.DebugData.PPLLData.NodeDataBuffer->Desc.NumElements;
+		PassParameters->PPLLCounter = HairData.DebugData.PPLLData.NodeCounterTexture;
+		PassParameters->PPLLNodeIndex = HairData.DebugData.PPLLData.NodeIndexTexture;
+		PassParameters->PPLLNodeData = GraphBuilder.CreateSRV(FRDGBufferSRVDesc(HairData.DebugData.PPLLData.NodeDataBuffer));
 		PassParameters->ViewUniformBuffer = View.ViewUniformBuffer;
 		PassParameters->SceneColorTextureUAV = GraphBuilder.CreateUAV(SceneColorTexture);
 		ShaderPrint::SetParameters(GraphBuilder, View.ShaderPrintData, PassParameters->ShaderPrintParameters);
diff --git a/Engine/Source/Runtime/Renderer/Private/HairStrands/HairStrandsEnvironment.cpp b/Engine/Source/Runtime/Renderer/Private/HairStrands/HairStrandsEnvironment.cpp
--- a/Engine/Source/Runtime/Renderer/Private/HairStrands/HairStrandsEnvironment.cpp
+++ b/Engine/Source/Runtime/Renderer/Private/HairStrands/HairStrandsEnvironment.cpp
@@ -317,7 +317,7 @@ static void AddHairStrandsEnvironmentLightingPassPS(
 	const FHairStrandsVoxelResources& VirtualVoxelResources,
 	const FRDGTextureRef SceneColorTexture,
 	const EHairLightingSourceType LightingType,
-	const FHairStrandsDebugData::Data* DebugData)
+	const FHairStrandsDebugData::FPlotData* DebugData)
 {
 	if (!Scene)
 	{
@@ -509,7 +509,7 @@ static void InternalRenderHairStrandsEnvironmentLighting(
 		return;
 	}
 	
-	AddHairStrandsEnvironmentLightingPassPS(GraphBuilder, Scene, View, VisibilityData, VoxelResources, nullptr, LightingType, View.HairStrandsViewData.DebugData.IsPlotDataValid() ? &View.HairStrandsViewData.DebugData.Resources : nullptr);
+	AddHairStrandsEnvironmentLightingPassPS(GraphBuilder, Scene, View, VisibilityData, VoxelResources, nullptr, LightingType, View.HairStrandsViewData.DebugData.IsPlotDataValid() ? &View.HairStrandsViewData.DebugData.PlotData : nullptr);
 }
 
 void RenderHairStrandsLumenLighting(
diff --git a/Engine/Source/Runtime/Renderer/Private/HairStrands/HairStrandsInterface.cpp b/Engine/Source/Runtime/Renderer/Private/HairStrands/HairStrandsInterface.cpp
--- a/Engine/Source/Runtime/Renderer/Private/HairStrands/HairStrandsInterface.cpp
+++ b/Engine/Source/Runtime/Renderer/Private/HairStrands/HairStrandsInterface.cpp
@@ -59,10 +59,20 @@ static TAutoConsoleVariable<int32> CVarHairStrandsSimulation(
 	ECVF_RenderThreadSafe | ECVF_Scalability);
 
 static TAutoConsoleVariable<int32> CVarHairStrandsNonVisibleShadowCasting(
-	TEXT("r.HairStrands.Shadow.CastShadowWhenNonVisible"), 0,
+	TEXT("r.HairStrands.Shadow.CastShadowWhenNonVisible"), 1,
 	TEXT("Enable shadow casting for hair strands even when culled out from the primary view"),
 	ECVF_RenderThreadSafe);
 
+static TAutoConsoleVariable<float> CVarHairStrandsNonVisibleShadowCasting_CullDistance(
+	TEXT("r.HairStrands.Visibility.NonVisibleShadowCasting.CullDistance"), 2000,
+	TEXT("Cull distance at which shadow casting starts to be disabled for non-visible hair strands instances."),
+	ECVF_RenderThreadSafe);
+
+static TAutoConsoleVariable<int32> CVarHairStrandsNonVisibleShadowCasting_Debug(
+	TEXT("r.HairStrands.Visibility.NonVisibleShadowCasting.Debug"), 0,
+	TEXT("Enable debug rendering for non-visible hair strands instance, casting shadow."),
+	ECVF_RenderThreadSafe);
+
 static TAutoConsoleVariable<int32> CVarHairStrandsContinuousDecimationReordering(
 	TEXT("r.HairStrands.ContinuousDecimationReordering"), 0,
 	TEXT("Enable strand reordering to allow Continuous LOD. Experimental"),
@@ -413,35 +423,70 @@ bool IsHairStrandsNonVisibleShadowCastingEnable()
 
 bool IsHairStrandsVisibleInShadows(const FViewInfo& View, const FHairStrandsInstance& Instance)
 {
+	const bool bDebugEnable = CVarHairStrandsNonVisibleShadowCasting_Debug.GetValueOnRenderThread() > 0;
+	FHairStrandsDebugData::FCullData* CullingData = nullptr;
+	if (bDebugEnable)
+	{
+		CullingData = const_cast<FHairStrandsDebugData::FCullData*>(&View.HairStrandsViewData.DebugData.CullData);
+		CullingData->bIsValid = true;
+		CullingData->ViewFrustum = View.ViewFrustum;
+	}
+
 	bool bIsVisibleInShadow = false;
 	if (const FHairGroupPublicData* HairData = Instance.GetHairData())
 	{
 		const int32 LODIndex = FMath::CeilToInt(HairData->LODIndex);
 		const bool bIsStrands = LODIndex >= 0 && HairData->IsVisible(LODIndex) && HairData->GetGeometryType(LODIndex) == EHairGeometryType::Strands;
-		if (bIsStrands)
+		if (!bIsStrands)
+		{
+			return false;
+		}
+
+		if (const FBoxSphereBounds* Bounds = Instance.GetBounds())
 		{
-			if (const FBoxSphereBounds* Bounds = Instance.GetBounds())
+			// Local lights
+			for (const FLightSceneInfo* LightInfo : View.HairStrandsViewData.VisibleShadowCastingLights)
 			{
+				if (LightInfo->Proxy->AffectsBounds(*Bounds))
 				{
-					for (const FLightSceneInfo* LightInfo : View.HairStrandsViewData.VisibleShadowCastingLights)
-					{
-						// Influence radius check
-						if (LightInfo->Proxy->AffectsBounds(*Bounds))
-						{
-							bIsVisibleInShadow = true;
-							break;
-						}
-					}
+					bIsVisibleInShadow = true;
+					break;
 				}
+			}
 
-				if (!bIsVisibleInShadow)
+			// Directional lights
+			if (!bIsVisibleInShadow)
+			{
+				const float CullDistance = FMath::Max(0.f, CVarHairStrandsNonVisibleShadowCasting_CullDistance.GetValueOnAnyThread());
+				for (const FHairStrandsViewData::FDirectionalLightCullData& CullData : View.HairStrandsViewData.VisibleShadowCastingDirectionalLights)
 				{
-					for (const FSphere& LightBound : View.HairStrandsViewData.VisibleShadowCastingBounds)
+					// Transform frustum into light space
+					const FMatrix& WorldToLight = CullData.LightInfo->Proxy->GetWorldToLight();
+					
+					// Transform groom bound into light space, and extend it along the light direction
+					FBoxSphereBounds BoundsInLightSpace = Bounds->TransformBy(WorldToLight);
+					BoundsInLightSpace.BoxExtent.X += CullDistance;
+
+					const bool bIntersect = CullData.ViewFrustumInLightSpace.IntersectBox(BoundsInLightSpace.Origin, BoundsInLightSpace.BoxExtent);
+					if (bDebugEnable)
+					{
+						FHairStrandsDebugData::FCullData::FLight& LightData = CullingData->DirectionalLights.AddDefaulted_GetRef();
+						LightData.WorldToLight = WorldToLight;
+						LightData.LightToWorld = CullData.LightInfo->Proxy->GetLightToWorld();
+						LightData.Center = FVector3f(CullData.LightInfo->Proxy->GetBoundingSphere().Center);
+						LightData.Extent = FVector3f(CullData.LightInfo->Proxy->GetBoundingSphere().W);
+						LightData.ViewFrustumInLightSpace = CullData.ViewFrustumInLightSpace;
+						LightData.InstanceBoundInLightSpace.Add({FVector3f(BoundsInLightSpace.GetBox().Min), FVector3f(BoundsInLightSpace.GetBox().Max)});
+						LightData.InstanceBoundInWorldSpace.Add({ FVector3f(Bounds->GetBox().Min), FVector3f(Bounds->GetBox().Max)});
+						LightData.InstanceIntersection.Add(bIntersect ? 1u : 0u);
+					}
+
+					// Ensure the extended groom bound interest the frustum
+					if (bIntersect)
 					{
-						// Influence radius check
-						if (Bounds->GetSphere().Intersects(LightBound))
+						bIsVisibleInShadow = true;
+						if (!bDebugEnable)
 						{
-							bIsVisibleInShadow = true;
 							break;
 						}
 					}
@@ -715,26 +760,37 @@ FHairGroupPublicData* GetHairData(const FMeshBatch* Mesh)
 
 void AddVisibleShadowCastingLight(const FScene& Scene, TArray<FViewInfo>& Views, const FLightSceneInfo* LightSceneInfo)
 {
+	const uint8 LightType = LightSceneInfo->Proxy->GetLightType();
 	for (FViewInfo& View : Views)
 	{
-		// If any hair data are registered, track which lights are visible so that hair strands can cast shadow even if not visibible in primary view
+		// If any hair data are registered, track which lights are visible so that hair strands can cast shadow even if not visible in primary view
+		// 
+		// Actual intersection test is done in IsHairStrandsVisibleInShadows():
+		// * For local lights, shadow casting is determined based on light influence radius
+		// * For directional light, view frustum and hair instance bounds are transformed in light space to determined potential intersections
 		if (Scene.HairStrandsSceneData.RegisteredProxies.Num() > 0)
 		{
-			View.HairStrandsViewData.VisibleShadowCastingLights.Add(LightSceneInfo);
-			break;
-		}
-	}
-}
+			if (LightType == ELightComponentType::LightType_Directional)
+			{
+				FHairStrandsViewData::FDirectionalLightCullData& Data = View.HairStrandsViewData.VisibleShadowCastingDirectionalLights.AddDefaulted_GetRef();
 
-void AddVisibleShadowCastingLight(const FScene& Scene, TArray<FViewInfo>& Views, const FSphere& Bounds)
-{
-	for (FViewInfo& View : Views)
-	{
-		// If any hair data are registered, track which lights are visible so that hair strands can cast shadow even if not visibible in primary view
-		if (Scene.HairStrandsSceneData.RegisteredProxies.Num() > 0)
-		{
-			View.HairStrandsViewData.VisibleShadowCastingBounds.Add(Bounds);
-			break;
+				// Transform view frustum into light space
+				const FMatrix& WorldToLight = LightSceneInfo->Proxy->GetWorldToLight();
+				const uint32 PlaneCount = View.ViewFrustum.Planes.Num();
+				Data.ViewFrustumInLightSpace.Planes.SetNum(PlaneCount);
+				for (uint32 PlaneIt = 0; PlaneIt < PlaneCount; ++PlaneIt)
+				{
+					Data.ViewFrustumInLightSpace.Planes[PlaneIt] = View.ViewFrustum.Planes[PlaneIt].TransformBy(WorldToLight);
+				}
+				Data.ViewFrustumInLightSpace.Init();
+				Data.LightInfo = LightSceneInfo;
+				break;
+			}
+			else
+			{
+				View.HairStrandsViewData.VisibleShadowCastingLights.Add(LightSceneInfo);
+				break;
+			}
 		}
 	}
 }
diff --git a/Engine/Source/Runtime/Renderer/Private/HairStrands/HairStrandsRendering.cpp b/Engine/Source/Runtime/Renderer/Private/HairStrands/HairStrandsRendering.cpp
--- a/Engine/Source/Runtime/Renderer/Private/HairStrands/HairStrandsRendering.cpp
+++ b/Engine/Source/Runtime/Renderer/Private/HairStrands/HairStrandsRendering.cpp
@@ -160,7 +160,7 @@ void RenderHairBasePass(
 			const bool bDebugSamplingEnable = GetHairStrandsSkyLightingDebugEnable();
 			if (bDebugSamplingEnable)
 			{
-				View.HairStrandsViewData.DebugData.Resources = FHairStrandsDebugData::CreateData(GraphBuilder);
+				View.HairStrandsViewData.DebugData.PlotData = FHairStrandsDebugData::CreatePlotData(GraphBuilder);
 			}
 		}
 		
diff --git a/Engine/Source/Runtime/Renderer/Private/HairStrands/HairStrandsVisibility.cpp b/Engine/Source/Runtime/Renderer/Private/HairStrands/HairStrandsVisibility.cpp
--- a/Engine/Source/Runtime/Renderer/Private/HairStrands/HairStrandsVisibility.cpp
+++ b/Engine/Source/Runtime/Renderer/Private/HairStrands/HairStrandsVisibility.cpp
@@ -4142,9 +4142,9 @@ void RenderHairStrandsVisibilityBuffer(
 				// Extract texture for debug visualization
 				if (GHairStrandsDebugPPLL > 0)
 				{
-					View.HairStrandsViewData.DebugData.PPLLNodeCounterTexture = PPLLNodeCounterTexture;
-					View.HairStrandsViewData.DebugData.PPLLNodeIndexTexture = PPLLNodeIndexTexture;
-					View.HairStrandsViewData.DebugData.PPLLNodeDataBuffer = PPLLNodeDataBuffer;
+					View.HairStrandsViewData.DebugData.PPLLData.NodeCounterTexture = PPLLNodeCounterTexture;
+					View.HairStrandsViewData.DebugData.PPLLData.NodeIndexTexture = PPLLNodeIndexTexture;
+					View.HairStrandsViewData.DebugData.PPLLData.NodeDataBuffer = PPLLNodeDataBuffer;
 				}
 			#endif
 			}
diff --git a/Engine/Source/Runtime/Renderer/Private/ShadowSetup.cpp b/Engine/Source/Runtime/Renderer/Private/ShadowSetup.cpp
--- a/Engine/Source/Runtime/Renderer/Private/ShadowSetup.cpp
+++ b/Engine/Source/Runtime/Renderer/Private/ShadowSetup.cpp
@@ -3985,7 +3985,7 @@ void FSceneRenderer::InitProjectedShadowVisibility()
 			// Register visible lights for allowing hair strands to cast shadow (directional light)
 			if (bHairStrands && LightIt->LightType == ELightComponentType::LightType_Directional)
 			{
-				HairStrands::AddVisibleShadowCastingLight(*Scene, Views, ProjectedShadowInfo.ShadowBounds);
+				HairStrands::AddVisibleShadowCastingLight(*Scene, Views, LightIt->LightSceneInfo);
 			}
 		}
 	}
@@ -5871,4 +5871,4 @@ void FSceneRenderer::AllocateMobileCSMAndSpotLightShadowDepthTargets(FRHICommand
 			}
 		}
 	}
-}
\ No newline at end of file
+}
