Fix crash due to temporary material proxies being used by parallel setup tasks after the render thread deletes them.

We need to wait for the setup tasks to complete before deleting the proxies.

#rnx

diff --git a/Engine/Source/Runtime/Renderer/Private/SceneRendering.cpp b/Engine/Source/Runtime/Renderer/Private/SceneRendering.cpp
--- a/Engine/Source/Runtime/Renderer/Private/SceneRendering.cpp
+++ b/Engine/Source/Runtime/Renderer/Private/SceneRendering.cpp
@@ -4180,8 +4180,6 @@ void FSceneRenderer::RenderThreadEnd(FRHICommandListImmediate& RHICmdList, const
 		for (FSceneRenderer* SceneRenderer : SceneRenderers)
 		{
 			SceneRenderer->GPUSceneDynamicContext.Release();
-			SceneRenderer->MeshCollector.DeleteTemporaryProxies();
-			SceneRenderer->RayTracingCollector.DeleteTemporaryProxies();
 		}
 
 		// Mem stack mark is stored on first scene renderer
@@ -4222,6 +4220,15 @@ void FSceneRenderer::RenderThreadEnd(FRHICommandListImmediate& RHICmdList, const
 				}
 			}
 
+			// The temporary proxies would normally be deleted by the scene renderer destructors, but this needs to happen on the render thread, so
+			// we will do it now, before destroying the renderers on a background thread. It's important to do this after the above wait, since the
+			// setup tasks use these temporary proxies.
+			for (FSceneRenderer* SceneRenderer : SceneRenderers)
+			{
+				SceneRenderer->MeshCollector.DeleteTemporaryProxies();
+				SceneRenderer->RayTracingCollector.DeleteTemporaryProxies();
+			}
+
 			FGraphEventArray CommandListTasks = MoveTemp(RHICmdList.GetRenderThreadTaskArray());
 
 			GSceneRenderCleanUpState.Task = FFunctionGraphTask::CreateAndDispatchWhenReady([LocalSceneRenderers = CopyTemp(SceneRenderers)]
