Fix for another crash while opening MotionVectors_Hair level. A hair simulation shader was passed dummy resources with wrong formats. Another shader (hair visibility) was passed StructuredBuffers where (RW)Buffer<uint> was expected.

#rnx

diff --git a/Engine/Plugins/Runtime/HairStrands/Source/HairStrandsCore/Private/Niagara/NiagaraDataInterfaceHairStrands.cpp b/Engine/Plugins/Runtime/HairStrands/Source/HairStrandsCore/Private/Niagara/NiagaraDataInterfaceHairStrands.cpp
--- a/Engine/Plugins/Runtime/HairStrands/Source/HairStrandsCore/Private/Niagara/NiagaraDataInterfaceHairStrands.cpp
+++ b/Engine/Plugins/Runtime/HairStrands/Source/HairStrandsCore/Private/Niagara/NiagaraDataInterfaceHairStrands.cpp
@@ -3529,8 +3529,8 @@ void UNiagaraDataInterfaceHairStrands::SetShaderParameters(const FNiagaraDataInt
 		ShaderParameters->DeformedTrianglePositionCBuffer = FNiagaraRenderer::GetDummyFloatBuffer();
 		ShaderParameters->RestSamplePositionsBuffer = FNiagaraRenderer::GetDummyFloatBuffer();
 		ShaderParameters->MeshSampleWeightsBuffer = FNiagaraRenderer::GetDummyFloatBuffer();
-		ShaderParameters->RootBarycentricCoordinatesBuffer = FNiagaraRenderer::GetDummyFloatBuffer();
-		ShaderParameters->RootToUniqueTriangleIndexBuffer = FNiagaraRenderer::GetDummyFloatBuffer();
+		ShaderParameters->RootBarycentricCoordinatesBuffer = FNiagaraRenderer::GetDummyUIntBuffer();
+		ShaderParameters->RootToUniqueTriangleIndexBuffer = FNiagaraRenderer::GetDummyUIntBuffer();
 		ShaderParameters->ParamsScaleBuffer = FNiagaraRenderer::GetDummyFloatBuffer();
 	}
 }
diff --git a/Engine/Source/Runtime/Renderer/Private/HairStrands/HairStrandsVisibility.cpp b/Engine/Source/Runtime/Renderer/Private/HairStrands/HairStrandsVisibility.cpp
--- a/Engine/Source/Runtime/Renderer/Private/HairStrands/HairStrandsVisibility.cpp
+++ b/Engine/Source/Runtime/Renderer/Private/HairStrands/HairStrandsVisibility.cpp
@@ -2974,7 +2974,7 @@ static FRasterComputeOutput AddVisibilityComputeRasterPass(
 	FRDGBufferUAVRef OutVisTileDataUAV = nullptr;
 	FRDGBufferUAVRef OutVisTileIndirectUAV = nullptr;
 	
-	OutVisTilePrims = GraphBuilder.CreateBuffer(FRDGBufferDesc::CreateStructuredDesc(sizeof(uint32), MaxTiles*1024), TEXT("Hair.OutVisTilePrims"));
+	OutVisTilePrims = GraphBuilder.CreateBuffer(FRDGBufferDesc::CreateBufferDesc(sizeof(uint32), MaxTiles*1024), TEXT("Hair.OutVisTilePrims"));
 
 	FRDGTextureDesc DescBinningGrid = FRDGTextureDesc::Create3D(TileGridRes3D, PF_R32_UINT, FClearValueBinding::None, TexCreate_UAV | TexCreate_ShaderResource);
 	OutVisTileBinningGrid = GraphBuilder.CreateTexture(DescBinningGrid, TEXT("Hair.VisTileBinningGrid"));
@@ -2982,7 +2982,7 @@ static FRasterComputeOutput AddVisibilityComputeRasterPass(
 	FRDGTextureDesc DescDepthGrid = FRDGTextureDesc::Create2D(TileGridRes, PF_R32_UINT, FClearValueBinding::None, TexCreate_UAV | TexCreate_ShaderResource);
 	OutVisTileDepthGrid = GraphBuilder.CreateTexture(DescDepthGrid, TEXT("Hair.VisTileDepthGrid"));
 
-	OutVisTileArgs = GraphBuilder.CreateBuffer(FRDGBufferDesc::CreateStructuredDesc(sizeof(uint32), 4+16+(((TileGridRes.X*TileGridRes.Y)+31)/32) ), TEXT("Hair.OutVisTileArgs"));
+	OutVisTileArgs = GraphBuilder.CreateBuffer(FRDGBufferDesc::CreateBufferDesc(sizeof(uint32), 4+16+(((TileGridRes.X*TileGridRes.Y)+31)/32) ), TEXT("Hair.OutVisTileArgs"));
 
 	FRDGBufferDesc DescTileData = FRDGBufferDesc::CreateStructuredDesc(sizeof(uint32), MaxTiles * 4 * 2);
 	DescTileData.Usage = EBufferUsageFlags(DescTileData.Usage | BUF_ByteAddressBuffer);
@@ -2990,10 +2990,10 @@ static FRasterComputeOutput AddVisibilityComputeRasterPass(
 
 	OutVisTileIndirect = GraphBuilder.CreateBuffer(FRDGBufferDesc::CreateIndirectDesc<FRHIDispatchIndirectParameters>(1), TEXT("Hair.OutVisTileIndirect"));
 
-	OutVisTilePrimsUAV = GraphBuilder.CreateUAV(OutVisTilePrims);
+	OutVisTilePrimsUAV = GraphBuilder.CreateUAV(OutVisTilePrims, PF_R32_UINT);
 	OutVisTileBinningGridUAV = GraphBuilder.CreateUAV(OutVisTileBinningGrid);
 	OutVisTileDepthGridUAV = GraphBuilder.CreateUAV(OutVisTileDepthGrid);
-	OutVisTileArgsUAV = GraphBuilder.CreateUAV(OutVisTileArgs);
+	OutVisTileArgsUAV = GraphBuilder.CreateUAV(OutVisTileArgs, PF_R32_UINT);
 	OutVisTileDataUAV = GraphBuilder.CreateUAV(OutVisTileData);
 	OutVisTileIndirectUAV = GraphBuilder.CreateUAV(OutVisTileIndirect);
 
@@ -3136,9 +3136,9 @@ static FRasterComputeOutput AddVisibilityComputeRasterPass(
 				RasterParameters->OutPrimMatTexture = PrimMatTextureUAV;
 				RasterParameters->SampleWeight = SampleWeight;
 				RasterParameters->HairInstance = GetHairStrandsInstanceParameters(GraphBuilder, ViewInfo, HairGroupPublicData, bCullingEnable, bForceRegister);
-				RasterParameters->VisTilePrims = GraphBuilder.CreateSRV(OutVisTilePrims);
+				RasterParameters->VisTilePrims = GraphBuilder.CreateSRV(OutVisTilePrims, PF_R32_UINT);
 				RasterParameters->RWVisTileDepthGrid = OutVisTileDepthGridUAV;
-				RasterParameters->VisTileArgs = GraphBuilder.CreateSRV(OutVisTileArgs);
+				RasterParameters->VisTileArgs = GraphBuilder.CreateSRV(OutVisTileArgs, PF_R32_UINT);
 				RasterParameters->RWVisTileData = OutVisTileDataUAV;
 				RasterParameters->VisTileIndirect = GraphBuilder.CreateSRV(OutVisTileIndirect);
 				RasterParameters->IndirectBufferArgs = OutVisTileIndirect;
@@ -3155,7 +3155,7 @@ static FRasterComputeOutput AddVisibilityComputeRasterPass(
 				Parameters->TotalPrimitiveInfoCount = TotalPrimitiveInfoCount;
 				Parameters->VisTileDepthGrid = OutVisTileDepthGrid;
 				Parameters->VisTileBinningGrid = OutVisTileBinningGrid;
-				Parameters->VisTileArgs = GraphBuilder.CreateSRV(OutVisTileArgs);
+				Parameters->VisTileArgs = GraphBuilder.CreateSRV(OutVisTileArgs, PF_R32_UINT);
 				ShaderPrint::SetParameters(GraphBuilder, ViewInfo.ShaderPrintData, Parameters->ShaderPrintParameters);
 
 				TShaderMapRef<FVisiblityRasterComputeDebugCS> DebugComputeShader(ViewInfo.ShaderMap);
