Removed Blue Noise asserts that fail with async texture compilation

diff --git a/Engine/Source/Runtime/Renderer/Private/BlueNoise.cpp b/Engine/Source/Runtime/Renderer/Private/BlueNoise.cpp
--- a/Engine/Source/Runtime/Renderer/Private/BlueNoise.cpp
+++ b/Engine/Source/Runtime/Renderer/Private/BlueNoise.cpp
@@ -14,9 +14,6 @@ FBlueNoise GetBlueNoiseParameters()
 	check(GEngine);
 	check(GEngine->BlueNoiseScalarTexture && GEngine->BlueNoiseVec2Texture);
 
-	check(GEngine->BlueNoiseScalarTexture->GetSizeX() == GEngine->BlueNoiseVec2Texture->GetSizeX()
-		&& GEngine->BlueNoiseScalarTexture->GetSizeY() == GEngine->BlueNoiseVec2Texture->GetSizeY());
-
 	BlueNoise.Dimensions = FIntVector(
 		GEngine->BlueNoiseScalarTexture->GetSizeX(),
 		GEngine->BlueNoiseScalarTexture->GetSizeX(),
