Fixed invalid RESTRICT use due to calling of helper functions which modified the memory.

diff --git a/Engine/Source/Runtime/Renderer/Private/DynamicBVH.h b/Engine/Source/Runtime/Renderer/Private/DynamicBVH.h
--- a/Engine/Source/Runtime/Renderer/Private/DynamicBVH.h
+++ b/Engine/Source/Runtime/Renderer/Private/DynamicBVH.h
@@ -821,12 +821,12 @@ bool FDynamicBVH< MaxChildren, FRootPolicy, FDirtyPolicy, FCostMetric >::Recursi
 template< uint32 MaxChildren, typename FRootPolicy, typename FDirtyPolicy, typename FCostMetric >
 uint32 FDynamicBVH< MaxChildren, FRootPolicy, FDirtyPolicy, FCostMetric >::PromoteChild( uint32 NodeIndex )
 {
-	FNode& RESTRICT Node = GetNode( NodeIndex );
+	FNode& Node = GetNode( NodeIndex );
 	check( !Node.IsLeaf( NodeIndex ) );
 	check( Node.NumChildren < MaxChildren );
 
 	uint32 FirstChild = Node.GetFirstChild( NodeIndex );
-	FNode& RESTRICT Children = GetNode( FirstChild );
+	FNode& Children = GetNode( FirstChild );
 	
 	FBounds3f Excluded[ MaxChildren ];
 
@@ -898,7 +898,7 @@ uint32 FDynamicBVH< MaxChildren, FRootPolicy, FDirtyPolicy, FCostMetric >::Promo
 template< uint32 MaxChildren, typename FRootPolicy, typename FDirtyPolicy, typename FCostMetric >
 void FDynamicBVH< MaxChildren, FRootPolicy, FDirtyPolicy, FCostMetric >::Rotate( uint32 NodeIndex )
 {
-	FNode& RESTRICT Node = GetNode( NodeIndex );
+	FNode& Node = GetNode( NodeIndex );
 
 	if( Node.IsRoot() )
 		return;
@@ -910,7 +910,7 @@ void FDynamicBVH< MaxChildren, FRootPolicy, FDirtyPolicy, FCostMetric >::Rotate(
 			ExcludedBounds += Node.ChildBounds[i];
 	}
 	
-	FNode& RESTRICT ParentNode = GetNode( Node.ParentIndex );
+	FNode& ParentNode = GetNode( Node.ParentIndex );
 
 	float	BestCost = CostMetric( ParentNode.GetBounds( Node.ParentIndex ) );
 	uint32	BestIndex = ~0u;
@@ -1374,4 +1374,4 @@ void FDynamicBVH< MaxChildren, FRootPolicy, FDirtyPolicy, FCostMetric >::Build(
 			Range = Children[ Last ];
 		}
 	}
-}
\ No newline at end of file
+}
