Respect OpaqueOnly instance mask for the ray tracing debug picker.

diff --git a/Engine/Shaders/Private/RayTracing/RayTracingDebug.usf b/Engine/Shaders/Private/RayTracing/RayTracingDebug.usf
--- a/Engine/Shaders/Private/RayTracing/RayTracingDebug.usf
+++ b/Engine/Shaders/Private/RayTracing/RayTracingDebug.usf
@@ -326,13 +326,15 @@ RAY_TRACING_ENTRY_RAYGEN(RayTracingDebugPickingRGS)
 
 	RayDesc Ray = CreatePrimaryRay(RenderTargetUV);
 
+	const uint InstanceInclusionMask = (OpaqueOnly ? RAY_TRACING_MASK_OPAQUE : RAY_TRACING_MASK_ALL);
+
 	FRayTracingDebugPayload Payload = (FRayTracingDebugPayload)0;
 	Payload.SetMiss();
 
 	TraceRay(
 		TLAS,
 		RAY_FLAG_CULL_BACK_FACING_TRIANGLES  /*RayFlags*/,
-		0xff /*InstanceInclusionMask*/,
+		InstanceInclusionMask,
 		0 /*RayContributionToHitGroupIndex*/,
 		RAY_TRACING_NUM_SHADER_SLOTS /*MultiplierForGeometryContributionToShaderIndex*/,
 		0 /*MissShaderIndex*/,
diff --git a/Engine/Source/Runtime/Renderer/Private/RayTracing/RayTracingDebug.cpp b/Engine/Source/Runtime/Renderer/Private/RayTracing/RayTracingDebug.cpp
--- a/Engine/Source/Runtime/Renderer/Private/RayTracing/RayTracingDebug.cpp
+++ b/Engine/Source/Runtime/Renderer/Private/RayTracing/RayTracingDebug.cpp
@@ -195,6 +195,7 @@ class FRayTracingPickingRGS : public FGlobalShader
 
 	BEGIN_SHADER_PARAMETER_STRUCT(FParameters, )
 		SHADER_PARAMETER_SRV(RaytracingAccelerationStructure, TLAS)
+		SHADER_PARAMETER(int32, OpaqueOnly)
 		SHADER_PARAMETER_RDG_BUFFER_UAV(RWStructuredBuffer, PickingOutput)
 		SHADER_PARAMETER_RDG_BUFFER_SRV(StructuredBuffer, InstancesDebugData)
 		SHADER_PARAMETER_RDG_BUFFER_SRV(StructuredBuffer, InstanceBuffer)
@@ -351,8 +352,9 @@ static FRDGBufferRef RayTracingPerformPicking(FRDGBuilder& GraphBuilder, const F
 	FRDGBufferRef PickingBuffer = GraphBuilder.CreateBuffer(PickingBufferDesc, TEXT("RayTracingDebug.PickingBuffer"));
 
 	FRayTracingPickingRGS::FParameters* RayGenParameters = GraphBuilder.AllocParameters<FRayTracingPickingRGS::FParameters>();
-	RayGenParameters->InstancesDebugData = GraphBuilder.CreateSRV(Scene->RayTracingScene.InstanceDebugBuffer); // TODO
+	RayGenParameters->InstancesDebugData = GraphBuilder.CreateSRV(Scene->RayTracingScene.InstanceDebugBuffer);
 	RayGenParameters->TLAS = Scene->RayTracingScene.GetLayerSRVChecked(ERayTracingSceneLayer::Base);
+	RayGenParameters->OpaqueOnly = CVarRayTracingDebugModeOpaqueOnly.GetValueOnRenderThread();
 	RayGenParameters->InstanceBuffer = GraphBuilder.CreateSRV(Scene->RayTracingScene.InstanceBuffer);
 	RayGenParameters->ViewUniformBuffer = View.ViewUniformBuffer;
 	RayGenParameters->PickingOutput = GraphBuilder.CreateUAV(PickingBuffer);
