Added support for Distant Lights (VSM)
- Added FShadowSceneRenderer to handle shadow rendering (transient) work - intended to take over more responsibility from here on.
- Added cache invalidation when lights are removed from the FScene (required as ID / cache key may change)
- Modified caching scheme to drop cache entries when caching is disabled.
- Changed local light footprint estimation to depend on the scaled resolution (screen percentage)
- Moved light location data and test to per-light cache entry (rather than per-shadow map)

diff --git a/Engine/Shaders/Private/VirtualShadowMaps/VirtualShadowMapCacheManagement.usf b/Engine/Shaders/Private/VirtualShadowMaps/VirtualShadowMapCacheManagement.usf
--- a/Engine/Shaders/Private/VirtualShadowMaps/VirtualShadowMapCacheManagement.usf
+++ b/Engine/Shaders/Private/VirtualShadowMaps/VirtualShadowMapCacheManagement.usf
@@ -141,7 +141,12 @@ void VirtualSmInvalidateInstancePagesCS(
 				{
 					++VirtualShadowMapId;
 				}
-				
+
+				// Stop invalidations if the light is 'distant'
+				if (ProjectionData.bCurrentDistantLight)
+				{
+					continue;
+				}
 				// NOTE: This is the *shadow view*'s translated world, not primary view
 				float4x4 LocalToTranslatedWorld = LWCMultiplyTranslation(InstanceSceneData.LocalToWorld, ProjectionData.PreViewTranslation);
 				
diff --git a/Engine/Shaders/Private/VirtualShadowMaps/VirtualShadowMapPageManagement.usf b/Engine/Shaders/Private/VirtualShadowMaps/VirtualShadowMapPageManagement.usf
--- a/Engine/Shaders/Private/VirtualShadowMaps/VirtualShadowMapPageManagement.usf
+++ b/Engine/Shaders/Private/VirtualShadowMaps/VirtualShadowMapPageManagement.usf
@@ -108,6 +108,48 @@ void MarkPage(uint VirtualShadowMapId, uint MipLevel, float3 TranslatedWorldPosi
 	}
 }
 
+// 
+RWStructuredBuffer<uint> OutPrunedLightGridData;
+RWStructuredBuffer<uint> OutPrunedNumCulledLightsGrid;
+
+[numthreads(VSM_DEFAULT_CS_GROUP_X, 1, 1)]
+void PruneLightGridCS(uint GridLinearIndex : SV_DispatchThreadID)
+{
+	uint EyeIndex = 0; // ??
+	const FLightGridData GridData = GetLightGridData(EyeIndex);
+
+	if (GridLinearIndex >= GridData.CulledGridSize.x * GridData.CulledGridSize.y * GridData.CulledGridSize.z)
+	{
+		return;
+	}
+	const FCulledLightsGridData CulledLightGridData = GetCulledLightsGrid(GridLinearIndex, EyeIndex);
+
+	uint NumRetainedLights = 0U;
+	LOOP
+	for (uint Index = 0; Index < CulledLightGridData.NumLocalLights; ++Index)
+	{
+		const FLocalLightData LightData = GetLocalLightData(CulledLightGridData.DataStartIndex + Index, EyeIndex);
+		int VirtualShadowMapId = LightData.VirtualShadowMapId;
+		// Discard any light without a VSM
+		if (VirtualShadowMapId != INDEX_NONE)
+		{
+			FVirtualShadowMapProjectionShaderData ProjectionData = GetVirtualShadowMapProjectionData(VirtualShadowMapId);
+
+			// Discard any distant lights (they get marked through the coarse page marking)
+			if (!ProjectionData.bCurrentDistantLight)
+			{
+				// Copy light grid data index
+				OutPrunedLightGridData[CulledLightGridData.DataStartIndex + NumRetainedLights++] = ForwardLightData.CulledLightDataGrid[CulledLightGridData.DataStartIndex + Index];
+			}
+		}
+	}
+	OutPrunedNumCulledLightsGrid[GridLinearIndex] = NumRetainedLights;
+}
+
+StructuredBuffer<uint> PrunedLightGridData;
+StructuredBuffer<uint> PrunedNumCulledLightsGrid;
+
+
 [numthreads(VSM_DEFAULT_CS_GROUP_XY, VSM_DEFAULT_CS_GROUP_XY, 1)]
 void GeneratePageFlagsFromPixels(
 	uint3 InGroupId : SV_GroupID,
@@ -223,11 +265,13 @@ void GeneratePageFlagsFromPixels(
 		uint GridLinearIndex = ComputeLightGridCellIndex(GridCoordinate, EyeIndex);
 		const FCulledLightsGridData CulledLightGridData = GetCulledLightsGrid(GridLinearIndex, EyeIndex);
 
+		const uint PrunedNumLocalLights = PrunedNumCulledLightsGrid[GridLinearIndex];
+
 		LOOP
-		for (uint Index = 0; Index < CulledLightGridData.NumLocalLights; ++Index)
+		for (uint Index = 0; Index < PrunedNumLocalLights; ++Index)
 		{
 			const FLocalLightData LightData = GetLocalLightData(CulledLightGridData.DataStartIndex + Index, EyeIndex);
-			const uint LightGridLightIndex = ForwardLightData.CulledLightDataGrid[CulledLightGridData.DataStartIndex + Index];
+			const uint LightGridLightIndex = PrunedLightGridData[CulledLightGridData.DataStartIndex + Index];
 
 			// Relative to PrimaryView
 			const float3 LightTranslatedWorldPosition = LightData.LightPositionAndInvRadius.xyz;
@@ -257,17 +301,14 @@ void GeneratePageFlagsFromPixels(
 			}
 #endif
 			int VirtualShadowMapId = LightData.VirtualShadowMapId;
-			if (VirtualShadowMapId != INDEX_NONE)
+			bool bSpotLight = LightData.SpotAnglesAndSourceRadiusPacked.x > -2.0f;
+			if( !bSpotLight )
 			{
-				bool bSpotLight = LightData.SpotAnglesAndSourceRadiusPacked.x > -2.0f;
-				if( !bSpotLight )
-				{
-					VirtualShadowMapId += VirtualShadowMapGetCubeFace(-ToLight);
-				}
-
-				uint MipLevel = GetMipLevelLocal(VirtualShadowMapId, TranslatedWorldPosition, SceneDepth);
-				MarkPage(VirtualShadowMapId, MipLevel, TranslatedWorldPosition, bUsePageDilation, PageDilationOffset);
+				VirtualShadowMapId += VirtualShadowMapGetCubeFace(-ToLight);
 			}
+
+			uint MipLevel = GetMipLevelLocal(VirtualShadowMapId, TranslatedWorldPosition, SceneDepth);
+			MarkPage(VirtualShadowMapId, MipLevel, TranslatedWorldPosition, bUsePageDilation, PageDilationOffset);
 		}
 	}
 }
@@ -315,11 +356,13 @@ void MarkCoarsePages(uint3 Index : SV_DispatchThreadID)
 			MarkPageAddress(CalcPageOffset(VirtualShadowMapId, 0, PageAddressLowHigh.zw), Flags);
 		}
 	}
-	else if (bMarkCoarsePagesLocal != 0)
+	// Note: always mark last mip for "distant" light
+	else if (bMarkCoarsePagesLocal != 0U || ProjectionData.bCurrentDistantLight)
 	{
 		// Mark last mip
 		uint PageOffset = CalcPageOffset(VirtualShadowMapId, VSM_MAX_MIP_LEVELS - 1, uint2(0, 0));
-		MarkPageAddress(PageOffset, Flags);
+		// For distant lights we want to inlcude the non-nanite geometry.
+		MarkPageAddress(PageOffset, ProjectionData.bCurrentDistantLight ? (VSM_ALLOCATED_FLAG | VSM_NON_NANITE_FLAG) : Flags);
 	}
 }
 
diff --git a/Engine/Shaders/Private/VirtualShadowMaps/VirtualShadowMapPrintStats.usf b/Engine/Shaders/Private/VirtualShadowMaps/VirtualShadowMapPrintStats.usf
--- a/Engine/Shaders/Private/VirtualShadowMaps/VirtualShadowMapPrintStats.usf
+++ b/Engine/Shaders/Private/VirtualShadowMaps/VirtualShadowMapPrintStats.usf
@@ -37,6 +37,27 @@ float2 ShaderPrintLeftAlign(float2 ScreenPos, uint Value)
 	return ShaderPrint(ScreenPos, Value);
 }
 
+
+float2 ShaderPrintUnits(in float2 InScreenPos, uint Value)
+{
+	float2 ScreenPos = InScreenPos;
+	if (Value >= 1024U * 1024U)
+	{
+		ScreenPos = ShaderPrint(ScreenPos, float((Value * 100U) / (1024U * 1024U)) / 100.0f);
+		ScreenPos = ShaderPrintSymbol(ScreenPos, _M_);
+	}
+	else if (Value >= 1024U)
+	{
+		ScreenPos = ShaderPrint(ScreenPos, float((Value * 100U) / (1024U)) / 100.0f);
+		ScreenPos = ShaderPrintText(ScreenPos, TEXT("k"), FontWhite);
+	}
+	else
+	{
+		ScreenPos = ShaderPrint(ScreenPos, Value);
+	}
+	return ScreenPos;
+}
+
 [numthreads(1, 1, 1)]
 void PrintStats()
 {
@@ -59,9 +80,12 @@ void PrintStats()
 		ScreenPos = ShaderPrintNewline(ScreenPos);
 
 		ScreenPos.x = ItemX;
-		ScreenPos = ShaderPrintText(ScreenPos, TEXT("Allocated"), FontWhite);
+		ScreenPos = ShaderPrintText(ScreenPos, TEXT("Allocated [texels]"), FontWhite);
 		ScreenPos.x = ValueX;
 		ScreenPos = ShaderPrintLeftAlign(ScreenPos, InStatsBuffer[VSM_STAT_ALLOCATED_PAGES]);
+		ScreenPos = ShaderPrintText(ScreenPos, TEXT("["), FontWhite);
+		ScreenPos = ShaderPrintUnits(ScreenPos, InStatsBuffer[VSM_STAT_ALLOCATED_PAGES] * VSM_PAGE_SIZE * VSM_PAGE_SIZE);
+		ScreenPos = ShaderPrintText(ScreenPos, TEXT("]"), FontWhite);
 		ScreenPos = ShaderPrintNewline(ScreenPos);
 
 		ScreenPos.x = ItemX;
diff --git a/Engine/Source/Runtime/Engine/Private/Components/SpotLightComponent.cpp b/Engine/Source/Runtime/Engine/Private/Components/SpotLightComponent.cpp
--- a/Engine/Source/Runtime/Engine/Private/Components/SpotLightComponent.cpp
+++ b/Engine/Source/Runtime/Engine/Private/Components/SpotLightComponent.cpp
@@ -129,6 +129,26 @@ public:
 	{
 		return FMath::ComputeBoundingSphereForCone(GetOrigin(), GetDirection(), (FSphere::FReal)Radius, (FSphere::FReal)CosOuterCone, (FSphere::FReal)SinOuterCone);
 	}
+	
+	virtual float GetEffectiveScreenRadius(const FViewMatrices& ShadowViewMatrices, const FIntPoint& CameraViewRectSize) const 
+	{
+		// Heuristic: use the radius of the inscribed sphere at the cone's end as the light's effective screen radius
+		// We do so because we do not want to use the light's radius directly, which will make us overestimate the shadow map resolution greatly for a spot light
+
+		// In the correct form,
+		//   InscribedSpherePosition = GetOrigin() + GetDirection() * GetRadius() / CosOuterCone
+		//   InscribedSphereRadius = GetRadius() / SinOuterCone
+		// Do it incorrectly to avoid division which is more expensive and risks division by zero
+		const FVector InscribedSpherePosition = GetOrigin() + GetDirection() * GetRadius() * CosOuterCone;
+		const float InscribedSphereRadius = GetRadius() * SinOuterCone;
+
+		const float SphereDistanceFromViewOrigin = (InscribedSpherePosition - ShadowViewMatrices.GetViewOrigin()).Size();
+
+		const FVector2D &ProjectionScale = ShadowViewMatrices.GetProjectionScale();
+		const float ScreenScale = FMath::Max(CameraViewRectSize.X * 0.5f * ProjectionScale.X, CameraViewRectSize.Y * 0.5f * ProjectionScale.Y);
+
+		return ScreenScale * InscribedSphereRadius / FMath::Max(SphereDistanceFromViewOrigin, 1.0f);
+	}
 
 	virtual float GetEffectiveScreenRadius(const FViewMatrices& ShadowViewMatrices) const override
 	{
diff --git a/Engine/Source/Runtime/Engine/Private/LocalLightSceneProxy.cpp b/Engine/Source/Runtime/Engine/Private/LocalLightSceneProxy.cpp
--- a/Engine/Source/Runtime/Engine/Private/LocalLightSceneProxy.cpp
+++ b/Engine/Source/Runtime/Engine/Private/LocalLightSceneProxy.cpp
@@ -85,6 +85,15 @@ FSphere FLocalLightSceneProxy::GetBoundingSphere() const
 	return FSphere(GetPosition(), GetRadius());
 }
 
+float FLocalLightSceneProxy::GetEffectiveScreenRadius(const FViewMatrices& ShadowViewMatrices, const FIntPoint& CameraViewRectSize) const
+{
+	const FVector2D& ProjectionScale = ShadowViewMatrices.GetProjectionScale();
+	const float ScreenScale = FMath::Max(CameraViewRectSize.X * 0.5f * ProjectionScale.X, CameraViewRectSize.Y * 0.5f * ProjectionScale.Y);
+
+	const float LightDistance = (GetOrigin() - ShadowViewMatrices.GetViewOrigin()).Size();
+	return ScreenScale * GetRadius() / FMath::Max(LightDistance, 1.0f);
+}
+
 float FLocalLightSceneProxy::GetEffectiveScreenRadius(const FViewMatrices& ShadowViewMatrices) const
 {
 	// Use the distance from the view origin to the light to approximate perspective projection
diff --git a/Engine/Source/Runtime/Engine/Public/LocalLightSceneProxy.h b/Engine/Source/Runtime/Engine/Public/LocalLightSceneProxy.h
--- a/Engine/Source/Runtime/Engine/Public/LocalLightSceneProxy.h
+++ b/Engine/Source/Runtime/Engine/Public/LocalLightSceneProxy.h
@@ -47,6 +47,7 @@ public:
 
 	virtual FSphere GetBoundingSphere() const;
 
+	virtual float GetEffectiveScreenRadius(const FViewMatrices& ShadowViewMatrices, const FIntPoint& CameraViewRectSize) const override;
 	virtual float GetEffectiveScreenRadius(const FViewMatrices& ShadowViewMatrices) const override;
 
 	virtual FVector GetPerObjectProjectedShadowProjectionPoint(const FBoxSphereBounds& SubjectBounds) const;
diff --git a/Engine/Source/Runtime/Engine/Public/SceneManagement.h b/Engine/Source/Runtime/Engine/Public/SceneManagement.h
--- a/Engine/Source/Runtime/Engine/Public/SceneManagement.h
+++ b/Engine/Source/Runtime/Engine/Public/SceneManagement.h
@@ -1423,6 +1423,9 @@ public:
 	virtual float GetLightSourceAngle() const { return 0.0f; }
 	virtual float GetShadowSourceAngleFactor() const { return 1.0f; }
 	virtual float GetTraceDistance() const { return 0.0f; }
+	// TODO: refactor this to move into the shadow scene renderer
+	virtual float GetEffectiveScreenRadius(const FViewMatrices& ShadowViewMatrices, const FIntPoint& CameraViewRectSize) const { return 0.0f; }
+	UE_DEPRECATED(5.1, "The GetEffectiveScreenRadius() that uses the screen-percentage scaled view rect (above) is used now.")
 	virtual float GetEffectiveScreenRadius(const FViewMatrices& ShadowViewMatrices) const { return 0.0f; }
 
 	/** Accesses parameters needed for rendering the light. */
diff --git a/Engine/Source/Runtime/Engine/Public/ShowFlags.h b/Engine/Source/Runtime/Engine/Public/ShowFlags.h
--- a/Engine/Source/Runtime/Engine/Public/ShowFlags.h
+++ b/Engine/Source/Runtime/Engine/Public/ShowFlags.h
@@ -476,6 +476,7 @@ private:
 		SetVisualizeVolumetricCloudConservativeDensity(false);
 		SetVisualizeStrataMaterial(false);
 		SetDrawOnlyVSMInvalidatingGeo(false);
+		SetDebugDrawDistantVirtualSMLights(false);
 
 		SetLumenScreenTraces(true);
 		SetLumenDetailTraces(true);
diff --git a/Engine/Source/Runtime/Engine/Public/ShowFlagsValues.inl b/Engine/Source/Runtime/Engine/Public/ShowFlagsValues.inl
--- a/Engine/Source/Runtime/Engine/Public/ShowFlagsValues.inl
+++ b/Engine/Source/Runtime/Engine/Public/ShowFlagsValues.inl
@@ -398,8 +398,8 @@ SHOWFLAG_FIXED_IN_SHIPPING(0, VisualizeStrataMaterial, SFG_Visualize, NSLOCTEXT(
 
 SHOWFLAG_FIXED_IN_SHIPPING(1, VirtualShadowMapCaching, SFG_Developer, NSLOCTEXT("UnrealEd", "VirtualShadowMapCaching", "Cache Virtual Shadow Maps"))
 
-/** Debug the Strata material buffer content */
 SHOWFLAG_FIXED_IN_SHIPPING(0, DrawOnlyVSMInvalidatingGeo, SFG_Visualize, NSLOCTEXT("UnrealEd", "DrawOnlyVSMInvalidatingGeo", "Draw Only Geometry Causing VSM Invalidation"))
+SHOWFLAG_FIXED_IN_SHIPPING(0, DebugDrawDistantVirtualSMLights, SFG_Advanced, NSLOCTEXT("UnrealEd", "DebugDrawDistantVirtualSMLightsSF", "Debug Draw Distant VSM Lights"))
 
 /** Debug the Virtual Texture System */
 SHOWFLAG_FIXED_IN_SHIPPING(0, VirtualTextureResidency, SFG_Hidden, NSLOCTEXT("UnrealEd", "VirtualTextureResidency", "Virtual Texture Residency"))
diff --git a/Engine/Source/Runtime/Renderer/Private/DeferredShadingRenderer.cpp b/Engine/Source/Runtime/Renderer/Private/DeferredShadingRenderer.cpp
--- a/Engine/Source/Runtime/Renderer/Private/DeferredShadingRenderer.cpp
+++ b/Engine/Source/Runtime/Renderer/Private/DeferredShadingRenderer.cpp
@@ -75,6 +75,7 @@
 #include "Lumen/LumenFrontLayerTranslucency.h"
 #include "Containers/ChunkedArray.h"
 #include "Async/ParallelFor.h"
+#include "Shadows/ShadowSceneRenderer.h"
 
 extern int32 GNaniteShowStats;
 extern int32 GNanitePickingDomain;
@@ -390,6 +391,8 @@ FDeferredShadingSceneRenderer::FDeferredShadingSceneRenderer(const FSceneViewFam
 	bAnyRayTracingPassEnabled = false;
 	bShouldUpdateRayTracingScene = false;
 #endif
+
+	ShadowSceneRenderer = MakeUnique<FShadowSceneRenderer>(*this);
 }
 
 /** 
@@ -2966,7 +2969,7 @@ void FDeferredShadingSceneRenderer::Render(FRDGBuilder& GraphBuilder)
 				RDG_GPU_STAT_SCOPE(GraphBuilder, ShadowDepths);
 
 				ensureMsgf(AreLightsInLightGrid(), TEXT("Virtual shadow map setup requires local lights to be injected into the light grid (this may be caused by 'r.LightCulling.Quality=0')."));
-				VirtualShadowMapArray.BuildPageAllocations(GraphBuilder, SceneTextures, Views, ViewFamily.EngineShowFlags, SortedLightSet, VisibleLightInfos, NaniteRasterResults, *Scene);
+				VirtualShadowMapArray.BuildPageAllocations(GraphBuilder, SceneTextures, Views, ViewFamily.EngineShowFlags, SortedLightSet, VisibleLightInfos, NaniteRasterResults);
 			}
 
 			RenderShadowDepthMaps(GraphBuilder, InstanceCullingManager);
diff --git a/Engine/Source/Runtime/Renderer/Private/RendererScene.cpp b/Engine/Source/Runtime/Renderer/Private/RendererScene.cpp
--- a/Engine/Source/Runtime/Renderer/Private/RendererScene.cpp
+++ b/Engine/Source/Runtime/Renderer/Private/RendererScene.cpp
@@ -3387,6 +3387,17 @@ void FScene::RemoveLightSceneInfo_RenderThread(FLightSceneInfo* LightSceneInfo)
 		// Remove the light from the lights list.
 		Lights.RemoveAt(LightSceneInfo->Id);
 
+		// TODO: move this work to FShadowScene & batch the light removals
+		{
+			TArray<FVirtualShadowMapArrayCacheManager*, SceneRenderingAllocator> VirtualShadowCacheManagers;
+			GetAllVirtualShadowMapCacheManagers(VirtualShadowCacheManagers);
+
+			for (FVirtualShadowMapArrayCacheManager* CacheManager : VirtualShadowCacheManagers)
+			{
+				CacheManager->OnLightRemoved(LightSceneInfo->Id);
+			}
+		}
+
 		if (!LightSceneInfo->Proxy->HasStaticShadowing()
 			&& LightSceneInfo->Proxy->CastsDynamicShadow()
 			&& LightSceneInfo->GetDynamicShadowMapChannel() == -1)
diff --git a/Engine/Source/Runtime/Renderer/Private/SceneRendering.cpp b/Engine/Source/Runtime/Renderer/Private/SceneRendering.cpp
--- a/Engine/Source/Runtime/Renderer/Private/SceneRendering.cpp
+++ b/Engine/Source/Runtime/Renderer/Private/SceneRendering.cpp
@@ -78,6 +78,7 @@
 #include "RectLightTextureManager.h"
 #include "DynamicResolutionState.h"
 #include "NaniteVisualizationData.h"
+#include "Shadows/ShadowSceneRenderer.h"
 
 /*-----------------------------------------------------------------------------
 	Globals
@@ -2435,6 +2436,7 @@ FSceneRenderer::FSceneRenderer(const FSceneViewFamily* InViewFamily, FHitProxyCo
 ,	ViewFamily(*CheckPointer(InViewFamily))
 ,	MeshCollector(InViewFamily->GetFeatureLevel(), Allocator)
 ,	RayTracingCollector(InViewFamily->GetFeatureLevel(), Allocator)
+,	VirtualShadowMapArray(*CheckPointer(Scene))
 ,	bHasRequestedToggleFreeze(false)
 ,	bUsedPrecomputedVisibility(false)
 ,	bGPUMasksComputed(false)
diff --git a/Engine/Source/Runtime/Renderer/Private/SceneRendering.h b/Engine/Source/Runtime/Renderer/Private/SceneRendering.h
--- a/Engine/Source/Runtime/Renderer/Private/SceneRendering.h
+++ b/Engine/Source/Runtime/Renderer/Private/SceneRendering.h
@@ -59,6 +59,7 @@ class FVirtualShadowMapClipmap;
 class FShadowProjectionPassParameters;
 class FSceneTextureShaderParameters;
 class FLumenSceneData;
+class FShadowSceneRenderer;
 
 struct FCloudRenderContext;
 struct FSceneWithoutWaterTextures;
@@ -1948,6 +1949,9 @@ public:
 
 	FVirtualShadowMapArray VirtualShadowMapArray;
 
+	// TODO: Move to deferred scene renderer
+	TUniquePtr<FShadowSceneRenderer> ShadowSceneRenderer;
+
 	/** If a freeze request has been made */
 	bool bHasRequestedToggleFreeze;
 
@@ -2117,6 +2121,9 @@ public:
 	FORCEINLINE const FSceneTextures& GetActiveSceneTextures() const { return ViewFamily.GetSceneTextures(); }
 	FORCEINLINE const FSceneTexturesConfig& GetActiveSceneTexturesConfig() const { return ViewFamily.SceneTexturesConfig; }
 
+	DECLARE_MULTICAST_DELEGATE_OneParam(FSceneOnScreenMessagesDelegate, FScreenMessageWriter&);
+	FSceneOnScreenMessagesDelegate OnGetOnScreenMessages;
+
 protected:
 
 	/** Size of the family. */
@@ -2355,9 +2362,6 @@ protected:
 		checkf(bShadowDepthRenderCompleted, TEXT("Shadow depth rendering was not done before shadow projections, this will cause severe shadow artifacts and indicates an engine bug (pass ordering)"));
 	}
 
-	DECLARE_MULTICAST_DELEGATE_OneParam(FSceneOnScreenMessagesDelegate, FScreenMessageWriter&);
-	FSceneOnScreenMessagesDelegate OnGetOnScreenMessages;
-
 private:
 	void ComputeFamilySize();
 
diff --git a/Engine/Source/Runtime/Renderer/Private/ShadowDepthRendering.cpp b/Engine/Source/Runtime/Renderer/Private/ShadowDepthRendering.cpp
--- a/Engine/Source/Runtime/Renderer/Private/ShadowDepthRendering.cpp
+++ b/Engine/Source/Runtime/Renderer/Private/ShadowDepthRendering.cpp
@@ -1709,13 +1709,16 @@ void FSceneRenderer::RenderVirtualShadowMaps(FRDGBuilder& GraphBuilder, bool bNa
 				{
 					if (ProjectedShadowInfo->ShouldClampToNearPlane() == bShouldClampToNearPlane && ProjectedShadowInfo->HasVirtualShadowMap())
 					{
-						VirtualShadowMapArray.AddRenderViews(
-							ProjectedShadowInfo, 
-							ComputeNaniteShadowsLODScaleFactor(),
-							PrevHZBPhysical.IsValid(),
-							bVSMUseHZB,
-							bShouldClampToNearPlane,
-							VirtualShadowViews);
+						if (ProjectedShadowInfo->bShouldRenderVSM)
+						{
+							VirtualShadowMapArray.AddRenderViews(
+								ProjectedShadowInfo,
+								ComputeNaniteShadowsLODScaleFactor(),
+								PrevHZBPhysical.IsValid(),
+								bVSMUseHZB,
+								bShouldClampToNearPlane,
+								VirtualShadowViews);
+						}
 					}
 				}
 
@@ -1786,7 +1789,7 @@ void FSceneRenderer::RenderVirtualShadowMaps(FRDGBuilder& GraphBuilder, bool bNa
 
 	if (UseNonNaniteVirtualShadowMaps(ShaderPlatform, FeatureLevel))
 	{
-		VirtualShadowMapArray.RenderVirtualShadowMapsNonNanite(GraphBuilder, SortedShadowsForShadowDepthPass.VirtualShadowMapShadows, *Scene, Views);
+		VirtualShadowMapArray.RenderVirtualShadowMapsNonNanite(GraphBuilder, SortedShadowsForShadowDepthPass.VirtualShadowMapShadows, Views);
 	}
 
 	// If separate static/dynamic caching is enabled, we may need to merge some pages after rendering
diff --git a/Engine/Source/Runtime/Renderer/Private/ShadowRendering.h b/Engine/Source/Runtime/Renderer/Private/ShadowRendering.h
--- a/Engine/Source/Runtime/Renderer/Private/ShadowRendering.h
+++ b/Engine/Source/Runtime/Renderer/Private/ShadowRendering.h
@@ -39,7 +39,7 @@ class FScene;
 class FSceneRenderer;
 class FViewInfo;
 class FVirtualShadowMapArrayCacheManager;
-class FVirtualShadowMapCacheEntry;
+class FVirtualShadowMapPerLightCacheEntry;
 class FLightTileIntersectionParameters;
 class FDistanceFieldCulledObjectBufferParameters;
 
@@ -396,13 +396,15 @@ public:
 
 	/** Whether the the shadow overlaps any nanite primitives */
 	uint32 bContainsNaniteSubjects : 1;
-
+	uint32 bShouldRenderVSM : 1;
+	
 	/** Used to fetch the correct cached static mesh draw commands */
 	EMeshPass::Type MeshPassTargetType = EMeshPass::CSMShadowDepth;
 
 	EShadowMeshSelection MeshSelectionMask = EShadowMeshSelection::All;
 	/** */
 	TArray< FVirtualShadowMap*, TInlineAllocator<6> > VirtualShadowMaps;
+	TSharedPtr<FVirtualShadowMapPerLightCacheEntry> VirtualShadowMapPerLightCacheEntry;
 
 	/** View projection matrices for each cubemap face, used by one pass point light shadows. */
 	TArray<FMatrix> OnePassShadowViewProjectionMatrices;
diff --git a/Engine/Source/Runtime/Renderer/Private/ShadowSetup.cpp b/Engine/Source/Runtime/Renderer/Private/ShadowSetup.cpp
--- a/Engine/Source/Runtime/Renderer/Private/ShadowSetup.cpp
+++ b/Engine/Source/Runtime/Renderer/Private/ShadowSetup.cpp
@@ -32,6 +32,7 @@
 #include "VirtualShadowMaps/VirtualShadowMapCacheManager.h"
 #include "VirtualShadowMaps/VirtualShadowMapClipmap.h"
 #include "InstanceCulling/InstanceCullingManager.h"
+#include "Shadows/ShadowSceneRenderer.h"
 
 /** Number of cube map shadow depth surfaces that will be created and used for rendering one pass point light shadows. */
 static const int32 NumCubeShadowDepthSurfaces = 5;
@@ -707,6 +708,7 @@ FProjectedShadowInfo::FProjectedShadowInfo()
 	, bNaniteGeometry(true)
 	, bIncludeInScreenSpaceShadowMask(true)
 	, bContainsNaniteSubjects(false)
+	, bShouldRenderVSM(true)
 	, PerObjectShadowFadeStart(UE_OLD_WORLD_MAX)		// LWC_TODO: Upgrade to double? Should have been HALF_WORLD_MAX to match precision.
 	, InvPerObjectShadowFadeLength(0.0f)
 	, OverlappedUVOnCachedShadowMap(-1.0f, -1.0f, -1.0f, -1.0f)
@@ -1695,10 +1697,10 @@ bool FProjectedShadowInfo::AddSubjectPrimitive(FPrimitiveSceneInfo* PrimitiveSce
 			{
 				for (FViewInfo& CurrentView : Views)
 				{
-					const float DistanceSquared = ( Bounds.Origin - CurrentView.ShadowViewMatrices.GetViewOrigin() ).SizeSquared();
-
-					if (bWholeSceneShadow && CacheMode != SDCM_StaticPrimitivesOnly)
+					// Note: skip small-mesh culling for VSM since it needs it drawn for GPU-side caching.
+					if (bWholeSceneShadow && CacheMode != SDCM_StaticPrimitivesOnly && MeshPassTargetType != EMeshPass::VSMShadowDepth)
 					{
+						const float DistanceSquared = ( Bounds.Origin - CurrentView.ShadowViewMatrices.GetViewOrigin() ).SizeSquared();
 						const float LODScaleSquared = FMath::Square(CurrentView.LODDistanceFactor);
 						const bool bDrawShadowDepth = FMath::Square(Bounds.SphereRadius) > FMath::Square(GMinScreenRadiusForShadowCaster) * DistanceSquared * LODScaleSquared;
 						if( !bDrawShadowDepth )
@@ -3533,7 +3535,9 @@ void FSceneRenderer::CreateWholeSceneProjectedShadow(
 		const uint32 ShadowFadeResolution = FMath::Max<int32>(0, CVarShadowFadeResolution.GetValueOnRenderThread());
 				
 		// Compute the maximum resolution required for the shadow by any view. Also keep track of the unclamped resolution for fading.
-		float MaxDesiredResolution = 0;
+		float MaxDesiredResolution = 0.0f;
+		float MaxScreenRadius = 0.0f;
+
 		TArray<float, TInlineAllocator<2> > FadeAlphas;
 		float MaxFadeAlpha = 0;
 		bool bStaticSceneOnly = false;
@@ -3543,7 +3547,8 @@ void FSceneRenderer::CreateWholeSceneProjectedShadow(
 		{
 			const FViewInfo& View = Views[ViewIndex];
 
-			const float ScreenRadius = LightSceneInfo->Proxy->GetEffectiveScreenRadius(View.ShadowViewMatrices);
+			const float ScreenRadius = LightSceneInfo->Proxy->GetEffectiveScreenRadius(View.ShadowViewMatrices, View.ViewRect.Size());
+			MaxScreenRadius = FMath::Max(MaxScreenRadius, ScreenRadius);
 
 			// Determine the amount of shadow buffer resolution needed for this view.
 			float UnclampedResolution = 1.0f;
@@ -3731,42 +3736,32 @@ void FSceneRenderer::CreateWholeSceneProjectedShadow(
 						VirtualSizeY,
 						0 // no border
 					);
-					
-					FVirtualShadowMapArrayCacheManager* CacheManager = VirtualShadowMapArray.CacheManager;
-					TSharedPtr<FVirtualShadowMapPerLightCacheEntry> VirtualSmPerLightCacheEntry;
-					if (CacheManager->IsValid())
-					{
-						VirtualSmPerLightCacheEntry = CacheManager->FindCreateLightCacheEntry(LightSceneInfo->Id);
-					}
 
-					int32 NumMaps = ProjectedShadowInitializer.bOnePassPointLightShadow ? 6 : 1;
-					for( int32 i = 0; i < NumMaps; i++ )
-					{
-						FVirtualShadowMap* VirtualShadowMap = VirtualShadowMapArray.Allocate();
-						ProjectedShadowInfo->VirtualShadowMaps.Add( VirtualShadowMap );
+					TSharedPtr<FVirtualShadowMapPerLightCacheEntry> VirtualSmPerLightCacheEntry = ShadowSceneRenderer->AddLocalLightShadow(ProjectedShadowInitializer, ProjectedShadowInfo, LightSceneInfo, MaxScreenRadius);
 
-						if (VirtualSmPerLightCacheEntry.IsValid())
-						{
-							TSharedPtr<FVirtualShadowMapCacheEntry> VirtualSmCacheEntry = VirtualSmPerLightCacheEntry->FindCreateShadowMapEntry(i);
-							// Note: the caching for spot lights only handles static lights and does not need the border, this also means pages can cache at all levels.
-							if(VirtualSmCacheEntry.IsValid())
-							{
-								VirtualSmCacheEntry->UpdateLocal(VirtualShadowMap->ID, ProjectedShadowInitializer);
-								VirtualShadowMap->VirtualShadowMapCacheEntry = VirtualSmCacheEntry;
-							}
-						}
+					if (VirtualSmPerLightCacheEntry)
+					{
+						ProjectedShadowInfo->VirtualShadowMapPerLightCacheEntry = VirtualSmPerLightCacheEntry;
+						ProjectedShadowInfo->bShouldRenderVSM = !VirtualSmPerLightCacheEntry->IsFullyCached();
 					}
 
 					ProjectedShadowInfo->MeshPassTargetType = EMeshPass::VSMShadowDepth;
 					ProjectedShadowInfo->MeshSelectionMask = EShadowMeshSelection::VSM;
 
 					bool bContainsNaniteSubjects = false;
-					// Skip convex hull tests for VSM since this causes artifacts due to caching (potentially check if caching is enabled but may lead to race condition).
-					// The interaction setup has already tested the light bounds (by calling FLightSceneProxy::AffectsBounds)
-					AddInteractingPrimitives(LightSceneInfo->GetDynamicInteractionOftenMovingPrimitiveList(), ProjectedShadowInfo, bContainsNaniteSubjects, false, VirtualSmPerLightCacheEntry);
-					AddInteractingPrimitives(LightSceneInfo->GetDynamicInteractionStaticPrimitiveList(), ProjectedShadowInfo, bContainsNaniteSubjects, false, VirtualSmPerLightCacheEntry);
-					ProjectedShadowInfo->bContainsNaniteSubjects = bContainsNaniteSubjects;
 
+					// TODO: Perform no rendering work (command setup etc) if if fully cached, this optimiztion would be nice to get back, however, right now this is tricky
+					//       as the logic to select those that get to re-render must happen when the full set of rendered local lights is known.
+					//       At that point we would need to re-do this logic, which is currently rather tied to this part of the code. 
+					//       After new/persistent shadow setup is done, this should be far easier.
+					//if (!bHasFullyCachedVSM)
+					{
+						// Skip convex hull tests for VSM since this causes artifacts due to caching (potentially check if caching is enabled but may lead to race condition).
+						// The interaction setup has already tested the light bounds (by calling FLightSceneProxy::AffectsBounds)
+						AddInteractingPrimitives(LightSceneInfo->GetDynamicInteractionOftenMovingPrimitiveList(), ProjectedShadowInfo, bContainsNaniteSubjects, false, VirtualSmPerLightCacheEntry);
+						AddInteractingPrimitives(LightSceneInfo->GetDynamicInteractionStaticPrimitiveList(), ProjectedShadowInfo, bContainsNaniteSubjects, false, VirtualSmPerLightCacheEntry);
+						ProjectedShadowInfo->bContainsNaniteSubjects = bContainsNaniteSubjects;
+					}
 					VisibleLightInfo.VirtualShadowMapId = ProjectedShadowInfo->VirtualShadowMaps[0]->ID;
 					VisibleLightInfo.AllProjectedShadows.Add(ProjectedShadowInfo);
 				}
@@ -4096,8 +4091,11 @@ void FSceneRenderer::GatherShadowDynamicMeshElements(FGlobalDynamicIndexBuffer&
 		// GPUCULL_TODO: Replace with new shadow culling processor thing
 		for (FProjectedShadowInfo* ProjectedShadowInfo : SortedShadowsForShadowDepthPass.VirtualShadowMapShadows)
 		{
-			FVisibleLightInfo& VisibleLightInfo = VisibleLightInfos[ProjectedShadowInfo->GetLightSceneInfo().Id];
-			ProjectedShadowInfo->GatherDynamicMeshElements(*this, VisibleLightInfo, ReusedViewsArray, DynamicIndexBuffer, DynamicVertexBuffer, DynamicReadBuffer, InstanceCullingManager);
+			if (ProjectedShadowInfo->bShouldRenderVSM)
+			{
+				FVisibleLightInfo& VisibleLightInfo = VisibleLightInfos[ProjectedShadowInfo->GetLightSceneInfo().Id];
+				ProjectedShadowInfo->GatherDynamicMeshElements(*this, VisibleLightInfo, ReusedViewsArray, DynamicIndexBuffer, DynamicVertexBuffer, DynamicReadBuffer, InstanceCullingManager);
+			}
 		}
 	}
 }
@@ -4987,8 +4985,12 @@ void FSceneRenderer::AllocateShadowDepthTargets(FRHICommandListImmediate& RHICmd
 				bShadowIsVisible |= bHasViewRelevance && VisibleLightViewInfo.ProjectedShadowVisibilityMap[ShadowIndex];
 			}
 
+			// Fully cached VSM should not be skipped, since they need to be queued for projection.
+			if (!ProjectedShadowInfo->bShouldRenderVSM)
+			{
+			}
 			// Skip the shadow rendering if there is no primitive in the cached shadow map and the current shadow map either
-			if ((ProjectedShadowInfo->CacheMode == SDCM_StaticPrimitivesOnly || (ProjectedShadowInfo->CacheMode == SDCM_MovablePrimitivesOnly && !ProjectedShadowInfo->IsWholeSceneDirectionalShadow())) && !ProjectedShadowInfo->HasSubjectPrims())
+			else if ((ProjectedShadowInfo->CacheMode == SDCM_StaticPrimitivesOnly || (ProjectedShadowInfo->CacheMode == SDCM_MovablePrimitivesOnly && !ProjectedShadowInfo->IsWholeSceneDirectionalShadow())) && !ProjectedShadowInfo->HasSubjectPrims())
 			{
 				const FCachedShadowMapData& CachedShadowMapData = Scene->GetCachedShadowMapDataRef(ProjectedShadowInfo->GetLightSceneInfo().Id, FMath::Max(ProjectedShadowInfo->CascadeSettings.ShadowSplitIndex, 0));
 
@@ -5802,6 +5804,11 @@ void FSceneRenderer::FinishInitDynamicShadows(FRHICommandListImmediate& RHICmdLi
 
 	FinishGatherShadowPrimitives(TaskData);
 
+	// TODO: Move this to the deferred scene renderer
+	if (ShadowSceneRenderer)
+	{
+		ShadowSceneRenderer->PostInitDynamicShadowsSetup();
+	}
 	AllocateShadowDepthTargets(RHICmdList);
 
 	// Generate mesh element arrays from shadow primitive arrays
diff --git a/Engine/Source/Runtime/Renderer/Private/Shadows/ShadowSceneRenderer.cpp b/Engine/Source/Runtime/Renderer/Private/Shadows/ShadowSceneRenderer.cpp
new file mode 100644
--- /dev/null
+++ b/Engine/Source/Runtime/Renderer/Private/Shadows/ShadowSceneRenderer.cpp
@@ -0,0 +1,161 @@
+// Copyright Epic Games, Inc. All Rights Reserved.
+/*=============================================================================
+	ShadowSceneRenderer.cpp:
+=============================================================================*/
+#include "ShadowSceneRenderer.h"
+#include "../ScenePrivate.h"
+#include "../DeferredShadingRenderer.h"
+#include "../VirtualShadowMaps/VirtualShadowMapCacheManager.h"
+
+#if !(UE_BUILD_SHIPPING || UE_BUILD_TEST)
+#include "DynamicPrimitiveDrawing.h"
+#endif
+
+TAutoConsoleVariable<int32> CVarMaxDistantLightsPerFrame(
+	TEXT("r.Shadow.Virtual.MaxDistantUpdatePerFrame"),
+	1,
+	TEXT("Maximum number of distant lights to update each frame."),
+	ECVF_Scalability | ECVF_RenderThreadSafe
+);
+
+
+static TAutoConsoleVariable<int32> CVarDistantLightMode(
+	TEXT("r.Shadow.Virtual.DistantLightMode"),
+	0,
+	TEXT("Control whether distant light mode is enabled for local lights.\n0 == Off (default), \n1 == On, \n2 == Force All."),
+	ECVF_RenderThreadSafe
+);
+
+
+FShadowSceneRenderer::FShadowSceneRenderer(FDeferredShadingSceneRenderer& InSceneRenderer)
+	: SceneRenderer(InSceneRenderer)
+	, Scene(*InSceneRenderer.Scene)
+	, VirtualShadowMapArray(InSceneRenderer.VirtualShadowMapArray)
+{
+}
+
+TSharedPtr<FVirtualShadowMapPerLightCacheEntry> FShadowSceneRenderer::AddLocalLightShadow(const FWholeSceneProjectedShadowInitializer& ProjectedShadowInitializer, FProjectedShadowInfo* ProjectedShadowInfo, FLightSceneInfo* LightSceneInfo, float MaxScreenRadius)
+{
+	FVirtualShadowMapArrayCacheManager* CacheManager = VirtualShadowMapArray.CacheManager;
+
+	const int32 LocalLightShadowIndex = LocalLights.Num();
+	FLocalLightShadowFrameSetup& LocalLightShadowFrameSetup = LocalLights.AddDefaulted_GetRef();
+	LocalLightShadowFrameSetup.ProjectedShadowInfo = ProjectedShadowInfo;
+	LocalLightShadowFrameSetup.LightSceneInfo = LightSceneInfo;
+
+	const int32 NumMaps = ProjectedShadowInitializer.bOnePassPointLightShadow ? 6 : 1;
+	for (int32 Index = 0; Index < NumMaps; ++Index)
+	{
+		FVirtualShadowMap* VirtualShadowMap = VirtualShadowMapArray.Allocate();
+		// TODO: redundant
+		ProjectedShadowInfo->VirtualShadowMaps.Add(VirtualShadowMap);
+		LocalLightShadowFrameSetup.VirtualShadowMaps.Add(VirtualShadowMap);
+	}
+
+	TSharedPtr<FVirtualShadowMapPerLightCacheEntry> PerLightCacheEntry = CacheManager->FindCreateLightCacheEntry(LightSceneInfo->Id);
+	if (PerLightCacheEntry.IsValid())
+	{
+		LocalLightShadowFrameSetup.PerLightCacheEntry = PerLightCacheEntry;
+		// Single page res, at this point we force the VSM to be single page
+		// TODO: this computation does not match up with page marking logic super-well, particularly for long spot lights,
+		//       we can absolutely mirror the page marking calc better, just unclear how much it helps. 
+		//       Also possible to feed back from gpu - which would be more accurate wrt partially visible lights (e.g., a spot going through the ground).
+		//       Of course this creates jumps if visibility changes, which may or may not create unsolvable artifacts.
+		const bool bIsDistantLight = CVarDistantLightMode.GetValueOnRenderThread() != 0
+			&& (MaxScreenRadius <= FVirtualShadowMap::PageSize || CVarDistantLightMode.GetValueOnRenderThread() == 2);
+
+		PerLightCacheEntry->UpdateLocal(ProjectedShadowInitializer, bIsDistantLight);
+
+		for (int32 Index = 0; Index < NumMaps; ++Index)
+		{
+			FVirtualShadowMap* VirtualShadowMap = LocalLightShadowFrameSetup.VirtualShadowMaps[Index];
+
+			TSharedPtr<FVirtualShadowMapCacheEntry> VirtualSmCacheEntry = PerLightCacheEntry->FindCreateShadowMapEntry(Index);
+			VirtualSmCacheEntry->UpdateLocal(VirtualShadowMap->ID, *PerLightCacheEntry);
+			VirtualShadowMap->VirtualShadowMapCacheEntry = VirtualSmCacheEntry;
+		}
+
+		if (bIsDistantLight)
+		{
+			// This priority could be calculated based also on whether the light has actually been invalidated or not (currently not tracked on host).
+			// E.g., all things being equal update those with an animated mesh in, for example. Plus don't update those the don't need it at all.
+			int32 FramesSinceLastRender = int32(Scene.GetFrameNumber()) - int32(PerLightCacheEntry->GetLastScheduledFrameNumber());
+			DistantLightUpdateQueue.Add(-FramesSinceLastRender, LocalLightShadowIndex);
+		}
+	}
+	return PerLightCacheEntry;
+}
+
+void FShadowSceneRenderer::PostInitDynamicShadowsSetup()
+{
+	UpdateDistantLightPriorityRender();
+
+	PostSetupDebugRender();
+}
+
+void FShadowSceneRenderer::UpdateDistantLightPriorityRender()
+{
+	int32 UpdateBudgetNumLights = CVarMaxDistantLightsPerFrame.GetValueOnRenderThread() < 0 ? int32(DistantLightUpdateQueue.Num()) : FMath::Min(int32(DistantLightUpdateQueue.Num()), CVarMaxDistantLightsPerFrame.GetValueOnRenderThread());
+	for (int32 Index = 0; Index < UpdateBudgetNumLights; ++Index)
+	{
+		const int32 LocalLightShadowIndex = DistantLightUpdateQueue.Top();
+		const int32 Age = DistantLightUpdateQueue.GetKey(LocalLightShadowIndex);
+		// UE_LOG(LogTemp, Log, TEXT("Index: %d Age: %d"), LocalLightShadowIndex, Age);
+		DistantLightUpdateQueue.Pop();
+
+		FLocalLightShadowFrameSetup& LocalLightShadowFrameSetup = LocalLights[LocalLightShadowIndex];
+		
+		// Force fully cached to be off.
+		LocalLightShadowFrameSetup.ProjectedShadowInfo->bShouldRenderVSM = true;
+		LocalLightShadowFrameSetup.PerLightCacheEntry->CurrenScheduledFrameNumber = Scene.GetFrameNumber();
+		// Should trigger invalidations also.
+		LocalLightShadowFrameSetup.PerLightCacheEntry->Invalidate();
+	}
+}
+
+
+void FShadowSceneRenderer::PostSetupDebugRender()
+{
+#if !(UE_BUILD_SHIPPING || UE_BUILD_TEST)
+	// TODO: Move to debug rendering function in FShadowSceneRenderer
+	if ((SceneRenderer.ViewFamily.EngineShowFlags.DebugDrawDistantVirtualSMLights))
+	{
+		int32 NumDistant = 0;
+		for (FViewInfo& View : SceneRenderer.Views)
+		{
+			FViewElementPDI DebugPDI(&View, nullptr, &View.DynamicPrimitiveCollector);
+
+			for (const FLocalLightShadowFrameSetup& LightSetup : LocalLights)
+			{			
+				FLinearColor Color = FLinearColor(FColor::Blue);
+				if (LightSetup.PerLightCacheEntry && LightSetup.PerLightCacheEntry->bCurrentIsDistantLight)
+				{
+					++NumDistant;
+					int32 FramesSinceLastRender = int32(Scene.GetFrameNumber()) - int32(LightSetup.PerLightCacheEntry->GetLastScheduledFrameNumber());
+					float Fade = FMath::Min(0.8f, float(FramesSinceLastRender) / float(LocalLights.Num()));
+					Color = LightSetup.PerLightCacheEntry->IsFullyCached() ? FMath::Lerp(FLinearColor(FColor::Green), FLinearColor(FColor::Red), Fade) : FLinearColor(FColor::Red);
+				}
+
+				Color.A = 1.0f;
+				if (LightSetup.LightSceneInfo->Proxy->GetLightType() == LightType_Spot)
+				{
+					FTransform TransformNoScale = FTransform(LightSetup.LightSceneInfo->Proxy->GetLightToWorld());
+					TransformNoScale.RemoveScaling();
+
+					DrawWireSphereCappedCone(&DebugPDI, TransformNoScale, LightSetup.LightSceneInfo->Proxy->GetRadius(), FMath::RadiansToDegrees(LightSetup.LightSceneInfo->Proxy->GetOuterConeAngle()), 16, 4, 8, Color, SDPG_World);
+				}
+				else
+				{
+					DrawWireSphereAutoSides(&DebugPDI, -LightSetup.ProjectedShadowInfo->PreShadowTranslation, Color, LightSetup.LightSceneInfo->Proxy->GetRadius(), SDPG_World);
+				}
+			}
+		}
+		SceneRenderer.OnGetOnScreenMessages.AddLambda([this, NumDistant](FScreenMessageWriter& ScreenMessageWriter)->void
+		{
+			ScreenMessageWriter.DrawLine(FText::FromString(FString::Printf(TEXT("Distant Light Count: %d"), NumDistant)), 10, FColor::Yellow);
+			ScreenMessageWriter.DrawLine(FText::FromString(FString::Printf(TEXT("Active Local Light Count: %d"), LocalLights.Num())), 10, FColor::Yellow);
+			ScreenMessageWriter.DrawLine(FText::FromString(FString::Printf(TEXT("Scene Light Count: %d"), Scene.Lights.Num())), 10, FColor::Yellow);
+		});
+	}
+#endif
+}
\ No newline at end of file
diff --git a/Engine/Source/Runtime/Renderer/Private/Shadows/ShadowSceneRenderer.h b/Engine/Source/Runtime/Renderer/Private/Shadows/ShadowSceneRenderer.h
new file mode 100644
--- /dev/null
+++ b/Engine/Source/Runtime/Renderer/Private/Shadows/ShadowSceneRenderer.h
@@ -0,0 +1,71 @@
+// Copyright Epic Games, Inc. All Rights Reserved.
+/*=============================================================================
+	ShadowSceneRenderer.h:
+=============================================================================*/
+#pragma once
+
+#include "CoreMinimal.h"
+#include "Containers/SparseArray.h"
+#include "Containers/ArrayView.h"
+#include "Containers/BinaryHeap.h"
+#include "Containers/Array.h"
+
+#include "../VirtualShadowMaps/VirtualShadowMapArray.h"
+
+class FProjectedShadowInfo;
+class FDeferredShadingSceneRenderer;
+class FWholeSceneProjectedShadowInitializer;
+class FRDGBuilder;
+class FVirtualShadowMapPerLightCacheEntry;
+
+/**
+ * Transient scope for per-frame rendering resources for the shadow rendering.
+ */
+class FShadowSceneRenderer
+{
+public:
+	FShadowSceneRenderer(FDeferredShadingSceneRenderer& InSceneRenderer);
+
+	/**
+	 * Add a cube/spot light for processing this frame.
+	 * TODO: Don't use legacy FProjectedShadowInfo or other params, instead info should flow from persistent setup & update.
+	 * TODO: Return reference to FLocalLightShadowFrameSetup ?
+	 */
+	TSharedPtr<FVirtualShadowMapPerLightCacheEntry> AddLocalLightShadow(const FWholeSceneProjectedShadowInitializer& Initializer, FProjectedShadowInfo* ProjectedShadowInfo, FLightSceneInfo* LightSceneInfo, float MaxScreenRadius);
+
+	/**
+	 * Call after view-dependent setup has been processed (InitView etc) but before any rendering activity has been kicked off.
+	 */
+	void PostInitDynamicShadowsSetup();
+
+private:
+	/**
+	 * Select the budgeted set of distant lights to update this frame.
+	 */
+	void UpdateDistantLightPriorityRender();
+
+	void PostSetupDebugRender();
+
+	// TODO: void RenderVirtualShadowMaps(FRDGBuilder& GraphBuilder, bool bNaniteEnabled);
+
+	struct FLocalLightShadowFrameSetup
+	{
+		TArray<FVirtualShadowMap*, TInlineAllocator<6>> VirtualShadowMaps;
+		// link to legacy system stuff, to be removed in due time
+		FProjectedShadowInfo* ProjectedShadowInfo = nullptr;
+		FLightSceneInfo* LightSceneInfo = nullptr;
+		TSharedPtr<FVirtualShadowMapPerLightCacheEntry> PerLightCacheEntry;
+	};
+
+	// TODO: maybe we want to keep these in a 1:1 sparse array wrt the light scene infos, for easy crossreference & GPU access (maybe)?
+	//       tradeoff is easy to look up (given light ID) but not compact, but OTOH can keep compact lists of indices for various purposes
+	TArray<FLocalLightShadowFrameSetup, SceneRenderingAllocator> LocalLights;
+
+	// Priority queue of distant lights to update.
+	FBinaryHeap<int32, uint32> DistantLightUpdateQueue;
+
+	// Links to other systems etc.
+	FDeferredShadingSceneRenderer& SceneRenderer;
+	FScene& Scene;
+	FVirtualShadowMapArray& VirtualShadowMapArray;
+};
diff --git a/Engine/Source/Runtime/Renderer/Private/VirtualShadowMaps/VirtualShadowMapArray.cpp b/Engine/Source/Runtime/Renderer/Private/VirtualShadowMaps/VirtualShadowMapArray.cpp
--- a/Engine/Source/Runtime/Renderer/Private/VirtualShadowMaps/VirtualShadowMapArray.cpp
+++ b/Engine/Source/Runtime/Renderer/Private/VirtualShadowMaps/VirtualShadowMapArray.cpp
@@ -292,7 +292,8 @@ static bool SetStatsArgsAndPermutation(FRDGBuilder& GraphBuilder, FRDGBufferRef
 	return bGenerateStats;
 }
 
-FVirtualShadowMapArray::FVirtualShadowMapArray()
+FVirtualShadowMapArray::FVirtualShadowMapArray(FScene& InScene) 
+	: Scene(InScene)
 {
 }
 
@@ -567,6 +568,20 @@ public:
 	}
 };
 
+class FPruneLightGridCS : public FVirtualPageManagementShader
+{
+	DECLARE_GLOBAL_SHADER(FPruneLightGridCS);
+	SHADER_USE_PARAMETER_STRUCT(FPruneLightGridCS, FVirtualPageManagementShader)
+	
+	BEGIN_SHADER_PARAMETER_STRUCT(FParameters, )
+		SHADER_PARAMETER_RDG_UNIFORM_BUFFER(FVirtualShadowMapUniformParameters, VirtualShadowMap)
+		SHADER_PARAMETER_STRUCT_REF(FViewUniformShaderParameters, View)
+		SHADER_PARAMETER_RDG_UNIFORM_BUFFER(FForwardLightData, ForwardLightData)
+		SHADER_PARAMETER_RDG_BUFFER_UAV(RWStructuredBuffer<uint>, OutPrunedLightGridData)
+		SHADER_PARAMETER_RDG_BUFFER_UAV(RWStructuredBuffer<uint>, OutPrunedNumCulledLightsGrid)
+	END_SHADER_PARAMETER_STRUCT()
+};
+IMPLEMENT_GLOBAL_SHADER(FPruneLightGridCS, "/Engine/Private/VirtualShadowMaps/VirtualShadowMapPageManagement.usf", "PruneLightGridCS", SF_Compute);
 
 class FGeneratePageFlagsFromPixelsCS : public FVirtualPageManagementShader
 {
@@ -585,6 +600,8 @@ class FGeneratePageFlagsFromPixelsCS : public FVirtualPageManagementShader
 		SHADER_PARAMETER_RDG_TEXTURE(Texture2D<uint2>, VisBuffer64)
 		SHADER_PARAMETER_RDG_BUFFER_UAV(RWStructuredBuffer<uint>, OutPageRequestFlags)
 		SHADER_PARAMETER_RDG_BUFFER_SRV(StructuredBuffer< uint >, DirectionalLightIds)
+		SHADER_PARAMETER_RDG_BUFFER_SRV(StructuredBuffer< uint >, PrunedLightGridData)
+		SHADER_PARAMETER_RDG_BUFFER_SRV(StructuredBuffer< uint >, PrunedNumCulledLightsGrid)
 		SHADER_PARAMETER_RDG_UNIFORM_BUFFER(FStrataGlobalUniformParameters, Strata)
 		RDG_BUFFER_ACCESS(IndirectBufferArgs, ERHIAccess::IndirectArgs)
 		SHADER_PARAMETER(uint32, InputType)
@@ -1032,8 +1049,7 @@ void FVirtualShadowMapArray::BuildPageAllocations(
 	const FEngineShowFlags& EngineShowFlags,
 	const FSortedLightSetSceneInfo& SortedLightsInfo,
 	const TArray<FVisibleLightInfo, SceneRenderingAllocator>& VisibleLightInfos,
-	const TArray<Nanite::FRasterResults, TInlineAllocator<2>>& NaniteRasterResults,
-	FScene& Scene)
+	const TArray<Nanite::FRasterResults, TInlineAllocator<2>>& NaniteRasterResults)
 {
 	check(IsEnabled());
 
@@ -1123,8 +1139,13 @@ void FVirtualShadowMapArray::BuildPageAllocations(
 					check(ShadowViewRect.Max.X == FVirtualShadowMap::VirtualMaxResolutionXY);
 					check(ShadowViewRect.Max.Y == FVirtualShadowMap::VirtualMaxResolutionXY);
 				}
+				TSharedPtr<FVirtualShadowMapPerLightCacheEntry> CacheEntry = CacheManager->FindCreateLightCacheEntry(ProjectedShadowInfo->GetLightSceneInfo().Id);
 
 				uint32 Flags = 0U;
+				if (CacheEntry.IsValid())
+				{
+					Flags |= CacheEntry->bCurrentIsDistantLight ? VSM_PROJ_FLAG_CURRENT_DISTANT_LIGHT : 0U;
+				}
 
 				int32 NumMaps = ProjectedShadowInfo->bOnePassPointLightShadow ? 6 : 1;
 				for( int32 i = 0; i < NumMaps; i++ )
@@ -1236,7 +1257,7 @@ void FVirtualShadowMapArray::BuildPageAllocations(
 			// Because of this we also cannot overlap this pass with the following ones.
 			bool bMarkCoarsePagesDirectional = CVarMarkCoarsePagesDirectional.GetValueOnRenderThread() != 0;
 			bool bMarkCoarsePagesLocal = CVarMarkCoarsePagesLocal.GetValueOnRenderThread() != 0;
-			if (bMarkCoarsePagesDirectional || bMarkCoarsePagesLocal)
+			// Note: always run this pass such that the distant lights may be marked if need be
 			{
 				FMarkCoarsePagesCS::FParameters* PassParameters = GraphBuilder.AllocParameters< FMarkCoarsePagesCS::FParameters >();
 				PassParameters->VirtualShadowMap = GetUniformBuffer(GraphBuilder);
@@ -1262,6 +1283,25 @@ void FVirtualShadowMapArray::BuildPageAllocations(
 				// It's currently safe to overlap these passes that all write to same page request flags
 				FRDGBufferUAVRef PageRequestFlagsUAV = GraphBuilder.CreateUAV(PageRequestFlagsRDG, ERDGUnorderedAccessViewFlags::SkipBarrier);
 
+				const FIntVector& CulledGridSize = View.ForwardLightingResources.ForwardLightData->CulledGridSize;
+				const uint32 NumLightGridEntries = CulledGridSize.X * CulledGridSize.Y * CulledGridSize.Z;
+				FRDGBufferRef PrunedLightGridDataRDG = GraphBuilder.CreateBuffer(FRDGBufferDesc::CreateStructuredDesc(sizeof(uint32), NumLightGridEntries), TEXT("Shadow.Virtual.PrunedLightGridData"));
+				uint32 LightGridIndexBufferSize = View.ForwardLightingResources.ForwardLightData->CulledLightDataGrid->Desc.Buffer->Desc.GetSize();
+				FRDGBufferRef PrunedNumCulledLightsGridRDG = GraphBuilder.CreateBuffer(FRDGBufferDesc::CreateStructuredDesc(sizeof(uint32), LightGridIndexBufferSize), TEXT("Shadow.Virtual.PrunedNumCulledLightsGrid"));
+
+				{
+					FPruneLightGridCS::FParameters* PassParameters = GraphBuilder.AllocParameters< FPruneLightGridCS::FParameters >();
+					PassParameters->VirtualShadowMap = GetUniformBuffer(GraphBuilder);
+					PassParameters->View = View.ViewUniformBuffer;
+					PassParameters->ForwardLightData = View.ForwardLightingResources.ForwardLightUniformBuffer;
+					PassParameters->OutPrunedLightGridData = GraphBuilder.CreateUAV(PrunedLightGridDataRDG);
+					PassParameters->OutPrunedNumCulledLightsGrid = GraphBuilder.CreateUAV(PrunedNumCulledLightsGridRDG);
+					auto ComputeShader = View.ShaderMap->GetShader<FPruneLightGridCS>();
+
+					FComputeShaderUtils::AddPass(GraphBuilder, RDG_EVENT_NAME("PruneLightGrid"), ComputeShader, PassParameters, FComputeShaderUtils::GetGroupCount(NumLightGridEntries, FPruneLightGridCS::DefaultCSGroupX));
+				};
+
+
 				auto GeneratePageFlags = [&](const EVirtualShadowMapProjectionInputType InputType)
 				{
 					FGeneratePageFlagsFromPixelsCS::FPermutationDomain PermutationVector;
@@ -1276,6 +1316,8 @@ void FVirtualShadowMapArray::BuildPageAllocations(
 					PassParameters->OutPageRequestFlags = PageRequestFlagsUAV;
 					PassParameters->ForwardLightData = View.ForwardLightingResources.ForwardLightUniformBuffer;
 					PassParameters->DirectionalLightIds = GraphBuilder.CreateSRV(DirectionalLightIdsRDG);
+					PassParameters->PrunedLightGridData = GraphBuilder.CreateSRV(PrunedLightGridDataRDG);
+					PassParameters->PrunedNumCulledLightsGrid = GraphBuilder.CreateSRV(PrunedNumCulledLightsGridRDG);
 					PassParameters->PageDilationBorderSizeLocal = CVarPageDilationBorderSizeLocal.GetValueOnRenderThread();
 					PassParameters->PageDilationBorderSizeDirectional = CVarPageDilationBorderSizeDirectional.GetValueOnRenderThread();
 					PassParameters->bCullBackfacingPixels = ShouldCullBackfacingPixels() ? 1 : 0;
@@ -2210,7 +2252,7 @@ static void AddRasterPass(
 		});
 }
 
-void FVirtualShadowMapArray::RenderVirtualShadowMapsNonNanite(FRDGBuilder& GraphBuilder, const TArray<FProjectedShadowInfo *, SceneRenderingAllocator>& VirtualSmMeshCommandPasses, FScene& Scene, TArrayView<FViewInfo> Views)
+void FVirtualShadowMapArray::RenderVirtualShadowMapsNonNanite(FRDGBuilder& GraphBuilder, const TArray<FProjectedShadowInfo *, SceneRenderingAllocator>& VirtualSmMeshCommandPasses, TArrayView<FViewInfo> Views)
 {
 	if (VirtualSmMeshCommandPasses.Num() == 0)
 	{
@@ -2314,34 +2356,37 @@ void FVirtualShadowMapArray::RenderVirtualShadowMapsNonNanite(FRDGBuilder& Graph
 		}
 		else if (ProjectedShadowInfo->HasVirtualShadowMap())
 		{
-			FParallelMeshDrawCommandPass& MeshCommandPass = ProjectedShadowInfo->GetShadowDepthPass();
-			MeshCommandPass.WaitForSetupTask();
-
-			FInstanceCullingContext* InstanceCullingContext = MeshCommandPass.GetInstanceCullingContext();
-
-			if (InstanceCullingContext->HasCullingCommands())
+			if (ProjectedShadowInfo->bShouldRenderVSM)
 			{
-				VSMCullingBatchInfo.NumPrimaryViews = AddRenderViews(ProjectedShadowInfo, 1.0f, HZBTexture != nullptr, false, ProjectedShadowInfo->ShouldClampToNearPlane(), VirtualShadowViews);
+				FParallelMeshDrawCommandPass& MeshCommandPass = ProjectedShadowInfo->GetShadowDepthPass();
+				MeshCommandPass.WaitForSetupTask();
 
-				if (CVarDoNonNaniteBatching.GetValueOnRenderThread())
-				{
-					FViewInfo* ShadowDepthView = ProjectedShadowInfo->ShadowDepthView;
-					uint32 DynamicInstanceIdOffset = ShadowDepthView->DynamicPrimitiveCollector.GetInstanceSceneDataOffset();
-					uint32 DynamicInstanceIdMax = DynamicInstanceIdOffset + ShadowDepthView->DynamicPrimitiveCollector.NumInstances();
-
-					VSMCullingBatchInfos.Add(VSMCullingBatchInfo);
-
-					// Note: we have to allocate these up front as the context merging machinery writes the offsets directly to the &PassParameters->InstanceCullingDrawParams, 
-					// this is a side-effect from sharing the code with the deferred culling. Should probably be refactored.
-					FVirtualShadowDepthPassParameters* PassParameters = GraphBuilder.AllocParameters<FVirtualShadowDepthPassParameters>();
-					InstanceCullingMergedContext.AddBatch(GraphBuilder, InstanceCullingContext, DynamicInstanceIdOffset, ShadowDepthView->DynamicPrimitiveCollector.NumInstances(), &PassParameters->InstanceCullingDrawParams);
-					BatchedVirtualSmMeshCommandPasses.Add(ProjectedShadowInfo);
-					BatchedPassParameters.Add(PassParameters);
-				}
-				else
+				FInstanceCullingContext* InstanceCullingContext = MeshCommandPass.GetInstanceCullingContext();
+
+				if (InstanceCullingContext->HasCullingCommands())
 				{
-					UnBatchedVSMCullingBatchInfo.Add(VSMCullingBatchInfo);
-					UnBatchedVirtualSmMeshCommandPasses.Add(ProjectedShadowInfo);
+					VSMCullingBatchInfo.NumPrimaryViews = AddRenderViews(ProjectedShadowInfo, 1.0f, HZBTexture != nullptr, false, ProjectedShadowInfo->ShouldClampToNearPlane(), VirtualShadowViews);
+
+					if (CVarDoNonNaniteBatching.GetValueOnRenderThread())
+					{
+						FViewInfo* ShadowDepthView = ProjectedShadowInfo->ShadowDepthView;
+						uint32 DynamicInstanceIdOffset = ShadowDepthView->DynamicPrimitiveCollector.GetInstanceSceneDataOffset();
+						uint32 DynamicInstanceIdMax = DynamicInstanceIdOffset + ShadowDepthView->DynamicPrimitiveCollector.NumInstances();
+
+						VSMCullingBatchInfos.Add(VSMCullingBatchInfo);
+
+						// Note: we have to allocate these up front as the context merging machinery writes the offsets directly to the &PassParameters->InstanceCullingDrawParams, 
+						// this is a side-effect from sharing the code with the deferred culling. Should probably be refactored.
+						FVirtualShadowDepthPassParameters* PassParameters = GraphBuilder.AllocParameters<FVirtualShadowDepthPassParameters>();
+						InstanceCullingMergedContext.AddBatch(GraphBuilder, InstanceCullingContext, DynamicInstanceIdOffset, ShadowDepthView->DynamicPrimitiveCollector.NumInstances(), &PassParameters->InstanceCullingDrawParams);
+						BatchedVirtualSmMeshCommandPasses.Add(ProjectedShadowInfo);
+						BatchedPassParameters.Add(PassParameters);
+					}
+					else
+					{
+						UnBatchedVSMCullingBatchInfo.Add(VSMCullingBatchInfo);
+						UnBatchedVirtualSmMeshCommandPasses.Add(ProjectedShadowInfo);
+					}
 				}
 			}
 		}
@@ -2719,6 +2764,11 @@ uint32 FVirtualShadowMapArray::AddRenderViews(const TSharedPtr<FVirtualShadowMap
 	BaseParams.TargetMipCount = 1;	// No mips for clipmaps
 	BaseParams.Flags = bClampToNearPlane ? 0u : NANITE_VIEW_FLAG_NEAR_CLIP;
 
+	if (Clipmap->GetCacheEntry())
+	{
+		Clipmap->GetCacheEntry()->MarkRendered(Scene.GetFrameNumber());
+	}
+
 	for (int32 ClipmapLevelIndex = 0; ClipmapLevelIndex < Clipmap->GetLevelCount(); ++ClipmapLevelIndex)
 	{
 		FVirtualShadowMap* VirtualShadowMap = Clipmap->GetVirtualShadowMap(ClipmapLevelIndex);
@@ -2748,12 +2798,6 @@ uint32 FVirtualShadowMapArray::AddRenderViews(const TSharedPtr<FVirtualShadowMap
 
 		Nanite::FPackedView View = Nanite::CreatePackedView(Params);
 		OutVirtualShadowViews.Add(View);
-
-		// Mark that we rendered to this VSM for caching purposes
-		if (VirtualShadowMap->VirtualShadowMapCacheEntry)
-		{
-			VirtualShadowMap->VirtualShadowMapCacheEntry->MarkRendered();
-		}
 	}
 
 	return uint32(Clipmap->GetLevelCount());
@@ -2772,6 +2816,11 @@ uint32 FVirtualShadowMapArray::AddRenderViews(const FProjectedShadowInfo* Projec
 	// local lights enable distance cull by default
 	BaseParams.Flags = NANITE_VIEW_FLAG_DISTANCE_CULL | (bClampToNearPlane ? 0u : NANITE_VIEW_FLAG_NEAR_CLIP);
 
+	if (ProjectedShadowInfo->VirtualShadowMapPerLightCacheEntry)
+	{
+		ProjectedShadowInfo->VirtualShadowMapPerLightCacheEntry->MarkRendered(Scene.GetFrameNumber());
+	}
+
 	int32 NumMaps = ProjectedShadowInfo->bOnePassPointLightShadow ? 6 : 1;
 	for (int32 Index = 0; Index < NumMaps; ++Index)
 	{
@@ -2799,11 +2848,6 @@ uint32 FVirtualShadowMapArray::AddRenderViews(const FProjectedShadowInfo* Projec
 		}
 
 		OutVirtualShadowViews.Add(Nanite::CreatePackedView(Params));
-
-		if (VirtualShadowMap->VirtualShadowMapCacheEntry)
-		{
-			VirtualShadowMap->VirtualShadowMapCacheEntry->MarkRendered();
-		}
 	}
 
 	return uint32(NumMaps);
diff --git a/Engine/Source/Runtime/Renderer/Private/VirtualShadowMaps/VirtualShadowMapArray.h b/Engine/Source/Runtime/Renderer/Private/VirtualShadowMaps/VirtualShadowMapArray.h
--- a/Engine/Source/Runtime/Renderer/Private/VirtualShadowMaps/VirtualShadowMapArray.h
+++ b/Engine/Source/Runtime/Renderer/Private/VirtualShadowMaps/VirtualShadowMapArray.h
@@ -191,7 +191,7 @@ private:
 class FVirtualShadowMapArray
 {
 public:	
-	FVirtualShadowMapArray();
+	FVirtualShadowMapArray(FScene& InScene);
 	~FVirtualShadowMapArray();
 
 	void Initialize(FRDGBuilder& GraphBuilder, FVirtualShadowMapArrayCacheManager* InCacheManager, bool bInEnabled);
@@ -234,9 +234,7 @@ public:
 		const FEngineShowFlags& EngineShowFlags,
 		const FSortedLightSetSceneInfo& SortedLights, 
 		const TArray<FVisibleLightInfo, SceneRenderingAllocator> &VisibleLightInfos, 
-		const TArray<Nanite::FRasterResults, TInlineAllocator<2>> &NaniteRasterResults, 
-
-		FScene& Scene);
+		const TArray<Nanite::FRasterResults, TInlineAllocator<2>> &NaniteRasterResults);
 
 	bool IsAllocated() const
 	{
@@ -253,7 +251,7 @@ public:
 	/**
 	 * Draw Non-Nanite geometry into the VSMs.
 	 */
-	void RenderVirtualShadowMapsNonNanite(FRDGBuilder& GraphBuilder, const TArray<FProjectedShadowInfo*, SceneRenderingAllocator>& VirtualSmMeshCommandPasses, FScene& Scene, TArrayView<FViewInfo> Views);
+	void RenderVirtualShadowMapsNonNanite(FRDGBuilder& GraphBuilder, const TArray<FProjectedShadowInfo*, SceneRenderingAllocator>& VirtualSmMeshCommandPasses, TArrayView<FViewInfo> Views);
 
 	void RenderDebugInfo(FRDGBuilder& GraphBuilder, TArrayView<FViewInfo> Views);
 	
@@ -338,6 +336,7 @@ public:
 	FVirtualShadowMapVisualizeLightSearch VisualizeLight;
 
 private:
+	FScene &Scene;
 	//
 	bool bUseHzbOcclusion = true;
 	bool bUseTwoPassHzbOcclusion = true;
diff --git a/Engine/Source/Runtime/Renderer/Private/VirtualShadowMaps/VirtualShadowMapCacheManager.cpp b/Engine/Source/Runtime/Renderer/Private/VirtualShadowMaps/VirtualShadowMapCacheManager.cpp
--- a/Engine/Source/Runtime/Renderer/Private/VirtualShadowMaps/VirtualShadowMapCacheManager.cpp
+++ b/Engine/Source/Runtime/Renderer/Private/VirtualShadowMaps/VirtualShadowMapCacheManager.cpp
@@ -73,7 +73,8 @@ void FVirtualShadowMapCacheEntry::UpdateClipmap(
 	double LevelRadius,
 	double ViewCenterZ,
 	// NOTE: ViewRadiusZ must be constant for a given clipmap level
-	double ViewRadiusZ)
+	double ViewRadiusZ,
+	const FVirtualShadowMapPerLightCacheEntry& PerLightEntry)
 {
 	bool bCacheValid = (CurrentVirtualShadowMapId != INDEX_NONE) && CVarForceInvalidateClipmaps.GetValueOnRenderThread() == 0;
 	
@@ -105,6 +106,12 @@ void FVirtualShadowMapCacheEntry::UpdateClipmap(
 		}
 	}
 
+	// Not valid if it was never rendered
+	if (PerLightEntry.PrevRenderedFrameNumber < 0)
+	{
+		bCacheValid = false;
+	}
+
 	bool bRadiusMatches = (ViewRadiusZ == Clipmap.ViewRadiusZ);
 
 	if (bCacheValid && bRadiusMatches)
@@ -127,31 +134,70 @@ void FVirtualShadowMapCacheEntry::UpdateClipmap(
 	}
 		
 	PrevPageSpaceLocation = CurrentPageSpaceLocation;
-	bPrevRendered = bCurrentRendered;
 	
 	CurrentVirtualShadowMapId = VirtualShadowMapId;
 	CurrentPageSpaceLocation = PageSpaceLocation;
-	bCurrentRendered = false;
 }
 
-void FVirtualShadowMapCacheEntry::UpdateLocal(int32 VirtualShadowMapId, const FWholeSceneProjectedShadowInitializer &InCacheValidKey)
+void FVirtualShadowMapCacheEntry::UpdateLocal(int32 VirtualShadowMapId, const FVirtualShadowMapPerLightCacheEntry& PerLightEntry)
 {
 	// Swap previous frame data over.
 	PrevPageSpaceLocation = FIntPoint(0, 0);		// Not used for local lights
 	PrevVirtualShadowMapId = CurrentVirtualShadowMapId;
-	bPrevRendered = bCurrentRendered;
 
-	// Check cache validity based of shadow setup
-	if (!LocalCacheValidKey.IsCachedShadowValid(InCacheValidKey))
+	// Not valid if it was never rendered
+	if (PerLightEntry.PrevRenderedFrameNumber < 0)
 	{
 		PrevVirtualShadowMapId = INDEX_NONE;
-		//UE_LOG(LogRenderer, Display, TEXT("Invalidated!"));
 	}
-	LocalCacheValidKey = InCacheValidKey;
 
 	CurrentVirtualShadowMapId = VirtualShadowMapId;
 	CurrentPageSpaceLocation = FIntPoint(0, 0);		// Not used for local lights
-	bCurrentRendered = false;
+}
+
+void FVirtualShadowMapCacheEntry::Invalidate()
+{
+	PrevVirtualShadowMapId = INDEX_NONE;
+}
+
+void FVirtualShadowMapPerLightCacheEntry::UpdateClipmap()
+{
+	PrevRenderedFrameNumber = FMath::Max(PrevRenderedFrameNumber, CurrentRenderedFrameNumber);
+	CurrentRenderedFrameNumber = -1;
+}
+
+void FVirtualShadowMapPerLightCacheEntry::UpdateLocal(const FProjectedShadowInitializer& InCacheKey, bool bIsDistantLight)
+{
+	bPrevIsDistantLight = bCurrentIsDistantLight;
+	PrevRenderedFrameNumber = FMath::Max(PrevRenderedFrameNumber, CurrentRenderedFrameNumber);
+	PrevScheduledFrameNumber = FMath::Max(PrevScheduledFrameNumber, CurrenScheduledFrameNumber);
+
+	// Check cache validity based of shadow setup
+	if (!LocalCacheKey.IsCachedShadowValid(InCacheKey))
+	{
+		// If it is a distant light, we want to let the time-share perform the invalidation.
+		// TODO: track invalidation state somehow for later.
+		if (!bIsDistantLight)
+		{
+			PrevRenderedFrameNumber = -1;
+		}
+		//UE_LOG(LogRenderer, Display, TEXT("Invalidated!"));
+	}
+	LocalCacheKey = InCacheKey;
+
+	bCurrentIsDistantLight = bIsDistantLight;
+	CurrentRenderedFrameNumber = -1;
+	CurrenScheduledFrameNumber = -1;
+}
+
+void FVirtualShadowMapPerLightCacheEntry::Invalidate()
+{
+	PrevRenderedFrameNumber = -1;
+
+	for (TSharedPtr<FVirtualShadowMapCacheEntry>& Entry : ShadowMapEntries)
+	{
+		Entry->Invalidate();
+	}
 }
 
 static inline uint32 EncodeInstanceInvalidationPayload(bool bInvalidateStaticPage, int32 ClipmapVirtualShadowMapId = INDEX_NONE)
@@ -552,9 +598,6 @@ void FVirtualShadowMapArrayCacheManager::ExtractFrameData(
 
 			GraphBuilder.QueueBufferExtraction(VirtualShadowMapArray.InvalidatingInstancesRDG, &PrevBuffers.InvalidatingInstancesBuffer);
 			PrevBuffers.NumInvalidatingInstanceSlots = VirtualShadowMapArray.NumInvalidatingInstanceSlots;
-
-			// Move cache entries to previous frame, this implicitly removes any that were not used
-			PrevCacheEntries = CacheEntries;
 			
 			// Store but drop any temp references embedded in the uniform parameters this frame.
 			// We'll reestablish them when we reimport the extracted resources next frame
@@ -563,6 +606,8 @@ void FVirtualShadowMapArrayCacheManager::ExtractFrameData(
 			PrevUniformParameters.PageTable = nullptr;
 			PrevUniformParameters.PhysicalPagePool = nullptr;
 		}
+		// Move cache entries to previous frame, this implicitly removes any that were not used
+		PrevCacheEntries = CacheEntries;
 
 		if (bExtractHzbData)
 		{
@@ -587,6 +632,14 @@ void FVirtualShadowMapArrayCacheManager::ExtractFrameData(
 	}
 	else
 	{
+		// Caching is disabled
+		if (CVarCacheVirtualSMs.GetValueOnRenderThread() == 0)
+		{
+			// Make sure we empty out any resources associated with caching
+			PrevCacheEntries.Reset();
+			CacheEntries.Reset();
+		}
+
 		// Do nothing; maintain the data that we had
 		// This allows us to work around some cases where the renderer gets called multiple times in a given frame
 		// - such as scene captures - but does no shadow-related work in all but one of them. We do not want to drop
@@ -798,6 +851,11 @@ void FVirtualShadowMapArrayCacheManager::OnSceneChange()
 	}
 }
 
+void FVirtualShadowMapArrayCacheManager::OnLightRemoved(int32 LightId)
+{
+	CacheEntries.Remove(LightId);
+	PrevCacheEntries.Remove(LightId);
+}
 
 /**
  * Compute shader to project and invalidate the rectangles of given instances.
diff --git a/Engine/Source/Runtime/Renderer/Private/VirtualShadowMaps/VirtualShadowMapCacheManager.h b/Engine/Source/Runtime/Renderer/Private/VirtualShadowMaps/VirtualShadowMapCacheManager.h
--- a/Engine/Source/Runtime/Renderer/Private/VirtualShadowMaps/VirtualShadowMapCacheManager.h
+++ b/Engine/Source/Runtime/Renderer/Private/VirtualShadowMaps/VirtualShadowMapCacheManager.h
@@ -14,6 +14,7 @@
 
 class FRHIGPUBufferReadback;
 class FGPUScene;
+class FVirtualShadowMapPerLightCacheEntry;
 
 #define VSM_LOG_INVALIDATIONS 0
 
@@ -22,32 +23,27 @@ class FVirtualShadowMapCacheEntry
 public:
 	// True if the cache has been (re)populated, set to false on init and set to true once the cache update process has happened.
 	// Also set to false whenever key data was not valid and all cached data is invalidated.
-	bool IsValid() { return PrevVirtualShadowMapId != INDEX_NONE && bPrevRendered; }
+	bool IsValid() { return PrevVirtualShadowMapId != INDEX_NONE; }
 
-	void UpdateLocal(int32 VirtualShadowMapId,
-		const FWholeSceneProjectedShadowInitializer &InCacheValidKey);
+	void UpdateLocal(int32 VirtualShadowMapId, const FVirtualShadowMapPerLightCacheEntry &PerLightEntry);
 
 	void UpdateClipmap(int32 VirtualShadowMapId,
 		const FMatrix &WorldToLight,
 		FIntPoint PageSpaceLocation,
 		double LevelRadius,
 		double ViewCenterZ,
-		double ViewRadiusZ);
+		double ViewRadiusZ, 
+		const FVirtualShadowMapPerLightCacheEntry& PerLightEntry);
 
-	void MarkRendered() { bCurrentRendered = true; }
+	void Invalidate();
 
 	// Previous frame data
 	FIntPoint PrevPageSpaceLocation = FIntPoint(0, 0);
 	int32 PrevVirtualShadowMapId = INDEX_NONE;
-	bool bPrevRendered = false;
 
 	// Current frame data
 	FIntPoint CurrentPageSpaceLocation = FIntPoint(0, 0);
 	int32 CurrentVirtualShadowMapId = INDEX_NONE;
-	bool bCurrentRendered = false;
-
-	// TODO: Potentially refactor this to decouple the cache key details
-	FWholeSceneProjectedShadowInitializer LocalCacheValidKey;
 
 	struct FClipmapInfo
 	{
@@ -70,7 +66,27 @@ public:
 	TSharedPtr<FVirtualShadowMapCacheEntry> FindCreateShadowMapEntry(int32 Index);
 
 	void OnPrimitiveRendered(const FPrimitiveSceneInfo* PrimitiveSceneInfo);
+	/**
+	 * The (local) VSM is fully cached if the previous frame if was distant and is distant this frame also.
+	 */
+	inline bool IsFullyCached() const { return bCurrentIsDistantLight && bPrevIsDistantLight; }
+	void MarkRendered(int32 FrameIndex) { CurrentRenderedFrameNumber = FrameIndex; }
+	int32 GetLastScheduledFrameNumber() const { return PrevScheduledFrameNumber; }
+	void UpdateClipmap();
+	void UpdateLocal(const FProjectedShadowInitializer &InCacheKey, bool bIsDistantLight = false);
+
+	/**
+	 * Mark as invalid, i.e., needing rendering.
+	 */
+	void Invalidate();
+
+	bool bPrevIsDistantLight = false;
+	int32 PrevRenderedFrameNumber = -1;
+	int32 PrevScheduledFrameNumber = -1;
 
+	bool bCurrentIsDistantLight = false;
+	int32 CurrentRenderedFrameNumber = -1;
+	int32 CurrenScheduledFrameNumber = -1;
 	// Primitives that have been rendered (not culled) the previous frame, when a primitive transitions from being culled to not it must be rendered into the VSM
 	// Key culling reasons are small size or distance cutoff.
 	TBitArray<> RenderedPrimitives;
@@ -89,6 +105,9 @@ public:
 	};
 
 	TArray<FInstanceRange> PrimitiveInstancesToInvalidate;
+
+private:
+	FProjectedShadowInitializer LocalCacheKey;
 };
 
 // Persistent buffers that we ping pong frame by frame
@@ -201,6 +220,11 @@ public:
 	 */
 	void OnSceneChange();
 
+	/**
+	 * Handle light removal, need to clear out cache entries as the ID may be reused after this.
+	 */
+	void OnLightRemoved(int32 LightId);
+
 	TRDGUniformBufferRef<FVirtualShadowMapUniformParameters> GetPreviousUniformBuffer(FRDGBuilder& GraphBuilder) const;
 
 	FVirtualShadowMapArrayFrameData PrevBuffers;
diff --git a/Engine/Source/Runtime/Renderer/Private/VirtualShadowMaps/VirtualShadowMapClipmap.cpp b/Engine/Source/Runtime/Renderer/Private/VirtualShadowMaps/VirtualShadowMapClipmap.cpp
--- a/Engine/Source/Runtime/Renderer/Private/VirtualShadowMaps/VirtualShadowMapClipmap.cpp
+++ b/Engine/Source/Runtime/Renderer/Private/VirtualShadowMaps/VirtualShadowMapClipmap.cpp
@@ -110,9 +110,10 @@ FVirtualShadowMapClipmap::FVirtualShadowMapClipmap(
 
 	WorldOrigin = CameraViewMatrices.GetViewOrigin();
 
-	if (bCacheValid)
+	PerLightCacheEntry = VirtualShadowMapArrayCacheManager->FindCreateLightCacheEntry(LightSceneInfo.Id);
+	if (PerLightCacheEntry.IsValid())
 	{
-		PerLightCacheEntry = VirtualShadowMapArrayCacheManager->FindCreateLightCacheEntry(LightSceneInfo.Id);
+		PerLightCacheEntry->UpdateClipmap();
 	}
 
 	for (int32 Index = 0; Index < LevelCount; ++Index)
@@ -175,7 +176,8 @@ FVirtualShadowMapClipmap::FVirtualShadowMapClipmap(
 				PageOffset,
 				RawLevelRadius,
 				ViewCenter.Z,
-				ViewRadiusZ);
+				ViewRadiusZ,
+				*PerLightCacheEntry);
 
 			Level.VirtualShadowMap->VirtualShadowMapCacheEntry = CacheEntry;
 
diff --git a/Engine/Source/Runtime/Renderer/Private/VirtualShadowMaps/VirtualShadowMapClipmap.h b/Engine/Source/Runtime/Renderer/Private/VirtualShadowMaps/VirtualShadowMapClipmap.h
--- a/Engine/Source/Runtime/Renderer/Private/VirtualShadowMaps/VirtualShadowMapClipmap.h
+++ b/Engine/Source/Runtime/Renderer/Private/VirtualShadowMaps/VirtualShadowMapClipmap.h
@@ -97,6 +97,8 @@ public:
 	TConstArrayView<uint32> GetRevealedPrimitivesMask() const { return RevealedPrimitivesMask.IsEmpty() ? MakeArrayView<uint32>(nullptr, 0) : MakeArrayView(RevealedPrimitivesMask.GetData(), FBitSet::CalculateNumWords(RevealedPrimitivesMask.Num())); }
 	int32 GetNumRevealedPrimitives() const { return RevealedPrimitivesMask.Num(); }
 
+	TSharedPtr<FVirtualShadowMapPerLightCacheEntry>& GetCacheEntry() { return PerLightCacheEntry; }
+
 private:
 	void ComputeBoundingVolumes(const FViewMatrices& CameraViewMatrices);
 
