Allow light type showflags in Test configurations.

#rnx

diff --git a/Engine/Source/Runtime/Renderer/Private/LightSceneInfo.cpp b/Engine/Source/Runtime/Renderer/Private/LightSceneInfo.cpp
--- a/Engine/Source/Runtime/Renderer/Private/LightSceneInfo.cpp
+++ b/Engine/Source/Runtime/Renderer/Private/LightSceneInfo.cpp
@@ -169,7 +169,7 @@ bool FLightSceneInfo::ShouldRenderLight(const FViewInfo& View, bool bOffscreen)
 	// Only render the light if it is in the view frustum
 	bool bLocalVisible = bVisible && (bOffscreen ? View.VisibleLightInfos[Id].bInDrawRange : View.VisibleLightInfos[Id].bInViewFrustum);
 
-#if !(UE_BUILD_SHIPPING || UE_BUILD_TEST)
+#if !UE_BUILD_SHIPPING
 	ELightComponentType Type = (ELightComponentType)Proxy->GetLightType();
 
 	switch(Type)
