Implement per-view clipmap cache entry for VSM & make page marking only mark relevant views.

diff --git a/Engine/Shaders/Private/VirtualShadowMaps/VirtualShadowMapPageManagement.usf b/Engine/Shaders/Private/VirtualShadowMaps/VirtualShadowMapPageManagement.usf
--- a/Engine/Shaders/Private/VirtualShadowMaps/VirtualShadowMapPageManagement.usf
+++ b/Engine/Shaders/Private/VirtualShadowMaps/VirtualShadowMapPageManagement.usf
@@ -42,6 +42,7 @@ float PageDilationBorderSizeDirectional;
 float PageDilationBorderSizeLocal;
 uint InputType;
 uint bCullBackfacingPixels;
+uint NumDirectionalLightSmInds;
 
 
 uint GetMipLevelLocal(int VirtualShadowMapId, float3 TranslatedWorldPosition, float SceneDepth)
@@ -232,7 +233,7 @@ void GeneratePageFlagsFromPixels(
 		const bool bUsePageDilation = PageDilationBorderSizeDirectional > 0.0f;
 		const float2 PageDilationOffset = PageDilationBorderSizeDirectional * PageDilationDither;
 
-		for (uint Index = 0; Index < VirtualShadowMap.NumDirectionalLights; ++Index)
+		for (uint Index = 0; Index < NumDirectionalLightSmInds; ++Index)
 		{
 			int ClipmapStartId = DirectionalLightIds[Index];
 			FVirtualShadowMapProjectionShaderData BaseProjectionData = GetVirtualShadowMapProjectionData(ClipmapStartId);
diff --git a/Engine/Source/Runtime/Renderer/Private/VirtualShadowMaps/VirtualShadowMapArray.cpp b/Engine/Source/Runtime/Renderer/Private/VirtualShadowMaps/VirtualShadowMapArray.cpp
--- a/Engine/Source/Runtime/Renderer/Private/VirtualShadowMaps/VirtualShadowMapArray.cpp
+++ b/Engine/Source/Runtime/Renderer/Private/VirtualShadowMaps/VirtualShadowMapArray.cpp
@@ -1099,15 +1099,12 @@ void FVirtualShadowMapArray::BuildPageAllocations(
 	TArray<FVirtualShadowMapProjectionShaderData, SceneRenderingAllocator> ProjectionData;
 	ProjectionData.AddDefaulted(ShadowMaps.Num());
 
-	// Gather directional light virtual shadow maps
-	TArray<int32, SceneRenderingAllocator> DirectionalLightIds;
 	for (const FVisibleLightInfo& VisibleLightInfo : VisibleLightInfos)
 	{
 		for (const TSharedPtr<FVirtualShadowMapClipmap>& Clipmap : VisibleLightInfo.VirtualShadowMapClipmaps)
 		{
 			// NOTE: Shader assumes all levels from a given clipmap are contiguous
 			int32 ClipmapID = Clipmap->GetVirtualShadowMap()->ID;
-			DirectionalLightIds.Add(ClipmapID);
 			for (int32 ClipmapLevel = 0; ClipmapLevel < Clipmap->GetLevelCount(); ++ClipmapLevel)
 			{
 				ProjectionData[ClipmapID + ClipmapLevel] = Clipmap->GetProjectionShaderData(ClipmapLevel);
@@ -1178,7 +1175,6 @@ void FVirtualShadowMapArray::BuildPageAllocations(
 	}
 
 	UniformParameters.NumShadowMaps = ShadowMaps.Num();
-	UniformParameters.NumDirectionalLights = DirectionalLightIds.Num();
 
 	ProjectionDataRDG = CreateProjectionDataBuffer(GraphBuilder, TEXT("Shadow.Virtual.ProjectionData"), ProjectionData);
 
@@ -1237,6 +1233,22 @@ void FVirtualShadowMapArray::BuildPageAllocations(
 
 		const FViewInfo &View = Views[ViewIndex];
 
+		// Gather directional light virtual shadow maps
+		TArray<int32, SceneRenderingAllocator> DirectionalLightIds;
+		for (const FVisibleLightInfo& VisibleLightInfo : VisibleLightInfos)
+		{
+			for (const TSharedPtr<FVirtualShadowMapClipmap>& Clipmap : VisibleLightInfo.VirtualShadowMapClipmaps)
+			{
+				// mark pages from both views in stereo for the left (primary) view clipmap - assumes GetPrimaryView returns self for non-stereo views.
+				if (Clipmap->GetDependentView() == Views[ViewIndex].GetPrimaryView())
+				{
+					// NOTE: Shader assumes all levels from a given clipmap are contiguous
+					int32 ClipmapID = Clipmap->GetVirtualShadowMap()->ID;
+					DirectionalLightIds.Add(ClipmapID);
+				}
+			}
+		}	
+		
 		// This view contained no local lights (that were stored in the light grid), and no directional lights, so nothing to do.
 		if (View.ForwardLightingResources.LocalLightVisibleLightInfosIndex.Num() + DirectionalLightIds.Num() == 0)
 		{
@@ -1318,6 +1330,7 @@ void FVirtualShadowMapArray::BuildPageAllocations(
 					PassParameters->DirectionalLightIds = GraphBuilder.CreateSRV(DirectionalLightIdsRDG);
 					PassParameters->PrunedLightGridData = GraphBuilder.CreateSRV(PrunedLightGridDataRDG);
 					PassParameters->PrunedNumCulledLightsGrid = GraphBuilder.CreateSRV(PrunedNumCulledLightsGridRDG);
+					PassParameters->NumDirectionalLightSmInds = uint32(DirectionalLightIds.Num());
 					PassParameters->PageDilationBorderSizeLocal = CVarPageDilationBorderSizeLocal.GetValueOnRenderThread();
 					PassParameters->PageDilationBorderSizeDirectional = CVarPageDilationBorderSizeDirectional.GetValueOnRenderThread();
 					PassParameters->bCullBackfacingPixels = ShouldCullBackfacingPixels() ? 1 : 0;
diff --git a/Engine/Source/Runtime/Renderer/Private/VirtualShadowMaps/VirtualShadowMapArray.h b/Engine/Source/Runtime/Renderer/Private/VirtualShadowMaps/VirtualShadowMapArray.h
--- a/Engine/Source/Runtime/Renderer/Private/VirtualShadowMaps/VirtualShadowMapArray.h
+++ b/Engine/Source/Runtime/Renderer/Private/VirtualShadowMaps/VirtualShadowMapArray.h
@@ -116,7 +116,7 @@ struct FVirtualShadowMapHZBMetadata
 
 BEGIN_GLOBAL_SHADER_PARAMETER_STRUCT(FVirtualShadowMapUniformParameters, )
 	SHADER_PARAMETER(uint32, NumShadowMaps)
-	SHADER_PARAMETER(uint32, NumDirectionalLights)
+	SHADER_PARAMETER(uint32, NumDirectionalLights) // TODO: Remove
 	SHADER_PARAMETER(uint32, MaxPhysicalPages)
 	// Set to 0 if separate static caching is disabled
 	SHADER_PARAMETER(uint32, StaticCachedArrayIndex)
diff --git a/Engine/Source/Runtime/Renderer/Private/VirtualShadowMaps/VirtualShadowMapCacheManager.cpp b/Engine/Source/Runtime/Renderer/Private/VirtualShadowMaps/VirtualShadowMapCacheManager.cpp
--- a/Engine/Source/Runtime/Renderer/Private/VirtualShadowMaps/VirtualShadowMapCacheManager.cpp
+++ b/Engine/Source/Runtime/Renderer/Private/VirtualShadowMaps/VirtualShadowMapCacheManager.cpp
@@ -470,23 +470,25 @@ TSharedPtr<FVirtualShadowMapCacheEntry> FVirtualShadowMapPerLightCacheEntry::Fin
 	return EntryRef;
 }
 
-TSharedPtr<FVirtualShadowMapPerLightCacheEntry> FVirtualShadowMapArrayCacheManager::FindCreateLightCacheEntry(int32 LightSceneId)
+TSharedPtr<FVirtualShadowMapPerLightCacheEntry> FVirtualShadowMapArrayCacheManager::FindCreateLightCacheEntry(int32 LightSceneId, uint32 ViewUniqueID)
 {
 	if (CVarCacheVirtualSMs.GetValueOnRenderThread() == 0)
 	{
 		return nullptr;
 	}
 
-	if (TSharedPtr<FVirtualShadowMapPerLightCacheEntry> *LightEntry = CacheEntries.Find(LightSceneId))
+	const uint64 CacheKey = (uint64(ViewUniqueID) << 32U) | uint64(LightSceneId);
+
+	if (TSharedPtr<FVirtualShadowMapPerLightCacheEntry> *LightEntry = CacheEntries.Find(CacheKey))
 	{
 		return *LightEntry;
 	}
 
 	// Add to current frame / active set.
-	TSharedPtr<FVirtualShadowMapPerLightCacheEntry>& NewLightEntry = CacheEntries.Add(LightSceneId);
+	TSharedPtr<FVirtualShadowMapPerLightCacheEntry>& NewLightEntry = CacheEntries.Add(CacheKey);
 
 	// Copy data if available
-	if (TSharedPtr<FVirtualShadowMapPerLightCacheEntry>* PrevNewLightEntry = PrevCacheEntries.Find(LightSceneId))
+	if (TSharedPtr<FVirtualShadowMapPerLightCacheEntry>* PrevNewLightEntry = PrevCacheEntries.Find(CacheKey))
 	{
 		NewLightEntry = *PrevNewLightEntry;
 	}
@@ -501,17 +503,6 @@ TSharedPtr<FVirtualShadowMapPerLightCacheEntry> FVirtualShadowMapArrayCacheManag
 
 
 
-TSharedPtr<FVirtualShadowMapCacheEntry> FVirtualShadowMapArrayCacheManager::FindCreateCacheEntry(int32 LightSceneId, int32 Index)
-{
-	if (CVarCacheVirtualSMs.GetValueOnRenderThread() == 0)
-	{
-		return nullptr;
-	}
-
-	return FindCreateLightCacheEntry(LightSceneId)->FindCreateShadowMapEntry(Index);
-}
-
-
 void FVirtualShadowMapPerLightCacheEntry::OnPrimitiveRendered(const FPrimitiveSceneInfo* PrimitiveSceneInfo)
 {
 	// Mark as (potentially present in a cached page somehwere, so we'd need to invalidate if it is removed/moved)
diff --git a/Engine/Source/Runtime/Renderer/Private/VirtualShadowMaps/VirtualShadowMapCacheManager.h b/Engine/Source/Runtime/Renderer/Private/VirtualShadowMaps/VirtualShadowMapCacheManager.h
--- a/Engine/Source/Runtime/Renderer/Private/VirtualShadowMaps/VirtualShadowMapCacheManager.h
+++ b/Engine/Source/Runtime/Renderer/Private/VirtualShadowMaps/VirtualShadowMapCacheManager.h
@@ -165,12 +165,7 @@ public:
 	/**
 	 * Finds an existing cache entry and moves to the active set or creates a fresh one.
 	 */
-	TSharedPtr<FVirtualShadowMapPerLightCacheEntry> FindCreateLightCacheEntry(int32 LightSceneId);
-
-	/**
-	 * Finds an existing cache entry and moves to the active set or creates a fresh one.
-	 */
-	TSharedPtr<FVirtualShadowMapCacheEntry> FindCreateCacheEntry(int32 LightSceneId, int32 Index = 0);
+	TSharedPtr<FVirtualShadowMapPerLightCacheEntry> FindCreateLightCacheEntry(int32 LightSceneId, uint32 ViewUniqueID = 0U);
 
 	/*
 	 * Returns true if cached data is available.
@@ -252,8 +247,8 @@ private:
 	TRefCountPtr<IPooledRenderTarget> HZBPhysicalPagePool;
 
 	// Index the Cache entries by the light ID
-	TMap< int32, TSharedPtr<FVirtualShadowMapPerLightCacheEntry> > CacheEntries;
-	TMap< int32, TSharedPtr<FVirtualShadowMapPerLightCacheEntry> > PrevCacheEntries;
+	TMap< uint64, TSharedPtr<FVirtualShadowMapPerLightCacheEntry> > CacheEntries;
+	TMap< uint64, TSharedPtr<FVirtualShadowMapPerLightCacheEntry> > PrevCacheEntries;
 
 	// Stores stats over frames when activated.
 	TRefCountPtr<FRDGPooledBuffer> AccumulatedStatsBuffer;
diff --git a/Engine/Source/Runtime/Renderer/Private/VirtualShadowMaps/VirtualShadowMapClipmap.cpp b/Engine/Source/Runtime/Renderer/Private/VirtualShadowMaps/VirtualShadowMapClipmap.cpp
--- a/Engine/Source/Runtime/Renderer/Private/VirtualShadowMaps/VirtualShadowMapClipmap.cpp
+++ b/Engine/Source/Runtime/Renderer/Private/VirtualShadowMaps/VirtualShadowMapClipmap.cpp
@@ -110,7 +110,7 @@ FVirtualShadowMapClipmap::FVirtualShadowMapClipmap(
 
 	WorldOrigin = CameraViewMatrices.GetViewOrigin();
 
-	PerLightCacheEntry = VirtualShadowMapArrayCacheManager->FindCreateLightCacheEntry(LightSceneInfo.Id);
+	PerLightCacheEntry = VirtualShadowMapArrayCacheManager->FindCreateLightCacheEntry(LightSceneInfo.Id, InDependentView->ViewState->GetViewKey());
 	if (PerLightCacheEntry.IsValid())
 	{
 		PerLightCacheEntry->UpdateClipmap();
