Move SubmitCommandsHint() out of RayTracingScene GPU scope and add more precise GPU stats for geometry and scene update

* Fixes inaccurate/inflated RayTracingScene timings in `stat gpu` and UnrealInsights in CPU-limited situations

diff --git a/Engine/Source/Runtime/Renderer/Private/DeferredShadingRenderer.cpp b/Engine/Source/Runtime/Renderer/Private/DeferredShadingRenderer.cpp
--- a/Engine/Source/Runtime/Renderer/Private/DeferredShadingRenderer.cpp
+++ b/Engine/Source/Runtime/Renderer/Private/DeferredShadingRenderer.cpp
@@ -293,6 +293,7 @@ DECLARE_CYCLE_STAT(TEXT("DeferredShadingSceneRenderer RenderFog"), STAT_FDeferre
 DECLARE_CYCLE_STAT(TEXT("DeferredShadingSceneRenderer RenderLightShaftBloom"), STAT_FDeferredShadingSceneRenderer_RenderLightShaftBloom, STATGROUP_SceneRendering);
 DECLARE_CYCLE_STAT(TEXT("DeferredShadingSceneRenderer RenderFinish"), STAT_FDeferredShadingSceneRenderer_RenderFinish, STATGROUP_SceneRendering);
 
+DECLARE_GPU_STAT(RayTracingUpdate);
 DECLARE_GPU_STAT(RayTracingScene);
 DECLARE_GPU_STAT(RayTracingGeometry);
 
@@ -1737,7 +1738,7 @@ bool FDeferredShadingSceneRenderer::DispatchRayTracingWorldUpdates(FRDGBuilder&
 	const ERDGPassFlags ComputePassFlags = bRayTracingAsyncBuild ? ERDGPassFlags::AsyncCompute : ERDGPassFlags::Compute;
 
 	{
-		RDG_GPU_STAT_SCOPE(GraphBuilder, RayTracingScene);
+		RDG_GPU_STAT_SCOPE(GraphBuilder, RayTracingUpdate);
 
 		FBuildAccelerationStructurePassParams* PassParams = GraphBuilder.AllocParameters<FBuildAccelerationStructurePassParams>();
 		PassParams->RayTracingSceneScratchBuffer = Scene->RayTracingScene.BuildScratchBuffer;
@@ -1747,12 +1748,17 @@ bool FDeferredShadingSceneRenderer::DispatchRayTracingWorldUpdates(FRDGBuilder&
 		PassParams->LightDataPacked =  nullptr;
 
 		// Use ERDGPassFlags::NeverParallel so the pass never runs off the render thread and we always get the following order of execution on the CPU:
-		// BuildTLASInstanceBuffer, RayTracingScene, EndUpdate, ..., ReleaseRayTracingResources		
-		GraphBuilder.AddPass(RDG_EVENT_NAME("RayTracingScene"), PassParams, ComputePassFlags | ERDGPassFlags::NeverCull | ERDGPassFlags::NeverParallel,
+		// BuildTLASInstanceBuffer, RayTracingUpdate, RayTracingEndUpdate, ..., ReleaseRayTracingResources
+		GraphBuilder.AddPass(RDG_EVENT_NAME("RayTracingUpdate"), PassParams, ComputePassFlags | ERDGPassFlags::NeverCull | ERDGPassFlags::NeverParallel,
 			[this, PassParams, bRayTracingAsyncBuild](FRHIComputeCommandList& RHICmdList)
 		{
-			FRHIBuffer* DynamicGeometryScratchBuffer = PassParams->DynamicGeometryScratchBuffer ? PassParams->DynamicGeometryScratchBuffer->GetRHI() : nullptr;
-			Scene->GetRayTracingDynamicGeometryCollection()->DispatchUpdates(RHICmdList, DynamicGeometryScratchBuffer);
+			{
+				SCOPED_GPU_STAT(RHICmdList, RayTracingGeometry);
+				FRHIBuffer* DynamicGeometryScratchBuffer = PassParams->DynamicGeometryScratchBuffer ? PassParams->DynamicGeometryScratchBuffer->GetRHI() : nullptr;
+				Scene->GetRayTracingDynamicGeometryCollection()->DispatchUpdates(RHICmdList, DynamicGeometryScratchBuffer);
+			}
+
+			SCOPED_GPU_STAT(RHICmdList, RayTracingScene);
 
 			FRHIRayTracingScene* RayTracingSceneRHI = Scene->RayTracingScene.GetRHIRayTracingSceneChecked();
 			FRHIBuffer* AccelerationStructureBuffer = Scene->RayTracingScene.GetBufferChecked();
@@ -1768,18 +1774,18 @@ bool FDeferredShadingSceneRenderer::DispatchRayTracingWorldUpdates(FRDGBuilder&
 
 			RHICmdList.BindAccelerationStructureMemory(RayTracingSceneRHI, AccelerationStructureBuffer, 0);
 			RHICmdList.BuildAccelerationStructure(BuildParams);
-
-			if (!bRayTracingAsyncBuild)
-			{
-				// Submit potentially expensive BVH build commands to the GPU as soon as possible.
-				// Avoids a GPU bubble in some CPU-limited cases.
-				RHICmdList.SubmitCommandsHint();
-			}
 		});
 	}
 
-	AddPass(GraphBuilder, RDG_EVENT_NAME("EndUpdate"), [this](FRHICommandListImmediate& RHICmdList)
+	AddPass(GraphBuilder, RDG_EVENT_NAME("RayTracingEndUpdate"), [this, bRayTracingAsyncBuild](FRHICommandListImmediate& RHICmdList)
 	{
+		if (!bRayTracingAsyncBuild)
+		{
+			// Submit potentially expensive BVH build commands to the GPU as soon as possible.
+			// Avoids a GPU bubble in some CPU-limited cases.
+			RHICmdList.SubmitCommandsHint();
+		}
+
 		Scene->GetRayTracingDynamicGeometryCollection()->EndUpdate(RHICmdList);
 	});
 
