Fix VSM crash updating scene captures (UE-160140).

diff --git a/Engine/Source/Runtime/Renderer/Private/VirtualShadowMaps/VirtualShadowMapClipmap.cpp b/Engine/Source/Runtime/Renderer/Private/VirtualShadowMaps/VirtualShadowMapClipmap.cpp
--- a/Engine/Source/Runtime/Renderer/Private/VirtualShadowMaps/VirtualShadowMapClipmap.cpp
+++ b/Engine/Source/Runtime/Renderer/Private/VirtualShadowMaps/VirtualShadowMapClipmap.cpp
@@ -110,10 +110,13 @@ FVirtualShadowMapClipmap::FVirtualShadowMapClipmap(
 
 	WorldOrigin = CameraViewMatrices.GetViewOrigin();
 
-	PerLightCacheEntry = VirtualShadowMapArrayCacheManager->FindCreateLightCacheEntry(LightSceneInfo.Id, InDependentView->ViewState->GetViewKey());
-	if (PerLightCacheEntry.IsValid())
+	if (InDependentView && InDependentView->ViewState)	// scene capture views don't have a persistent view state
 	{
-		PerLightCacheEntry->UpdateClipmap();
+		PerLightCacheEntry = VirtualShadowMapArrayCacheManager->FindCreateLightCacheEntry(LightSceneInfo.Id, InDependentView->ViewState->GetViewKey());
+		if (PerLightCacheEntry.IsValid())
+		{
+			PerLightCacheEntry->UpdateClipmap();
+		}
 	}
 
 	for (int32 Index = 0; Index < LevelCount; ++Index)
