Fix shadow rendering with multiviewport (UE-157530)

diff --git a/Engine/Source/Runtime/Renderer/Private/ShadowRendering.cpp b/Engine/Source/Runtime/Renderer/Private/ShadowRendering.cpp
--- a/Engine/Source/Runtime/Renderer/Private/ShadowRendering.cpp
+++ b/Engine/Source/Runtime/Renderer/Private/ShadowRendering.cpp
@@ -905,15 +905,46 @@ void FProjectedShadowInfo::SetupProjectionStencilMask(
 		//       This means doing 2x renders, but letting the scissor rect kill the undersired half. Drawing the full mask once is easy, but since the outer
 		//       loop is over each view, the stencil mask is not retained when the right view comes around.
 		// TODO: Support instanced stereo properly in the projection stenciling pass.
-		const bool bIsInstancedStereoEmulated = View->bIsInstancedStereoEnabled && !View->bIsMultiViewportEnabled && IStereoRendering::IsStereoEyeView(*View);
+		const bool bIsInstancedStereoEmulated = View->bIsInstancedStereoEnabled && !View->bIsMobileMultiViewEnabled && IStereoRendering::IsStereoEyeView(*View);
 		if (bIsInstancedStereoEmulated && ProjectionStencilingPasses.IsValidIndex(View->PrimaryViewIndex))
 		{
 			ensure(ProjectionStencilingPasses[View->PrimaryViewIndex]->GetInstanceCullingMode() == EInstanceCullingMode::Stereo);
 
 			const FViewInfo* PrimaryView = View->GetPrimaryView();
 			const FViewInfo* InstancedView = View->GetInstancedView();
-			RHICmdList.SetViewport(PrimaryView->ViewRect.Min.X, PrimaryView->ViewRect.Min.Y, 0.0f, InstancedView->ViewRect.Max.X, InstancedView->ViewRect.Max.Y, 1.0f);
-			RHICmdList.SetScissorRect(true, View->ViewRect.Min.X, View->ViewRect.Min.Y, View->ViewRect.Max.X, View->ViewRect.Max.Y);
+			if (View->bIsMultiViewportEnabled)
+			{
+				float LeftMinX = PrimaryView->ViewRect.Min.X;
+				float LeftMaxX = PrimaryView->ViewRect.Max.X;
+				float LeftMinY = PrimaryView->ViewRect.Min.Y;
+				float LeftMaxY = PrimaryView->ViewRect.Max.Y;
+
+				float RightMinX = InstancedView->ViewRect.Min.X;
+				float RightMaxX = InstancedView->ViewRect.Max.X;
+				float RightMinY = InstancedView->ViewRect.Min.Y;
+				float RightMaxY = InstancedView->ViewRect.Max.Y;
+
+				// multi-viewport - collapse the other view in pair to be 0-width 0-height, effectively disabling it
+				if (IStereoRendering::IsAPrimaryView(*View))
+				{
+					RightMaxX = RightMinX;
+					RightMaxY = RightMinY;
+				}
+				else
+				{
+					LeftMinX = LeftMaxX;	// not a typo, just to have viewports adjacent to each other still
+					LeftMaxY = LeftMinY;
+				}
+
+				RHICmdList.SetStereoViewport(LeftMinX, RightMinX, LeftMinY, RightMinY, /*MinZ*/ 0.0f,
+					LeftMaxX, RightMaxX, LeftMaxY, RightMaxY, /*MaxZ*/ 1.0f);
+			}
+			else
+			{
+				// clip planes
+				RHICmdList.SetViewport(PrimaryView->ViewRect.Min.X, PrimaryView->ViewRect.Min.Y, 0.0f, InstancedView->ViewRect.Max.X, InstancedView->ViewRect.Max.Y, 1.0f);
+				RHICmdList.SetScissorRect(true, View->ViewRect.Min.X, View->ViewRect.Min.Y, View->ViewRect.Max.X, View->ViewRect.Max.Y);
+			}
 			// Submit the first (and only pass - we share that at least) as the pass is set up for stereo.
 			ProjectionStencilingPasses[View->PrimaryViewIndex]->SubmitDraw(RHICmdList, InstanceCullingDrawParams);
 
@@ -922,6 +953,7 @@ void FProjectedShadowInfo::SetupProjectionStencilMask(
 		}
 		else if (ViewIndex < ProjectionStencilingPasses.Num())
 		{
+			check(ProjectionStencilingPasses[ViewIndex] != nullptr);
 			ProjectionStencilingPasses[ViewIndex]->SubmitDraw(RHICmdList, InstanceCullingDrawParams);
 		}
 
@@ -1184,7 +1216,7 @@ void FProjectedShadowInfo::RenderProjection(
 		PassParameters->ShadowTexture1 = GraphBuilder.RegisterExternalTexture(RenderTargets.ColorTargets[1]);
 	}
 
-	const bool bIsInstancedStereoEmulated = View->bIsInstancedStereoEnabled && !View->bIsMultiViewportEnabled && IStereoRendering::IsStereoEyeView(*View);
+	const bool bIsInstancedStereoEmulated = View->bIsInstancedStereoEnabled && !View->bIsMobileMultiViewEnabled && IStereoRendering::IsStereoEyeView(*View);
 	if (ViewIndex < ProjectionStencilingPasses.Num() && ProjectionStencilingPasses[ViewIndex] != nullptr)
 	{
 		// GPUCULL_TODO: get rid of const cast
diff --git a/Engine/Source/Runtime/Renderer/Private/SimpleMeshDrawCommandPass.cpp b/Engine/Source/Runtime/Renderer/Private/SimpleMeshDrawCommandPass.cpp
--- a/Engine/Source/Runtime/Renderer/Private/SimpleMeshDrawCommandPass.cpp
+++ b/Engine/Source/Runtime/Renderer/Private/SimpleMeshDrawCommandPass.cpp
@@ -15,7 +15,7 @@ FSimpleMeshDrawCommandPass::FSimpleMeshDrawCommandPass(const FSceneView& View, F
 	
 	TArray<int32, TFixedAllocator<2> > ViewIds;
 	ViewIds.Add(ViewInfo->GPUSceneViewId);
-	bUsingStereo = bEnableStereo && ViewInfo->bIsInstancedStereoEnabled && !View.bIsMultiViewportEnabled && IStereoRendering::IsStereoEyeView(View);
+	bUsingStereo = bEnableStereo && ViewInfo->bIsInstancedStereoEnabled && !ViewInfo->bIsMobileMultiViewEnabled && IStereoRendering::IsStereoEyeView(View);
 	if (bUsingStereo)
 	{
 		check(ViewInfo->GetInstancedView() != nullptr);
