Run ray tracing static primitive LOD calculation asynchronously, as results are only needed much later for FRayTracingSceneAddInstancesTask

* This allows GatherRayTracingWorldInstancesForView() to start earlier on the rendering thread

diff --git a/Engine/Source/Runtime/Renderer/Private/DeferredShadingRenderer.cpp b/Engine/Source/Runtime/Renderer/Private/DeferredShadingRenderer.cpp
--- a/Engine/Source/Runtime/Renderer/Private/DeferredShadingRenderer.cpp
+++ b/Engine/Source/Runtime/Renderer/Private/DeferredShadingRenderer.cpp
@@ -611,6 +611,10 @@ struct FRayTracingRelevantPrimitiveList
 	TChunkedArray<FRayTracingRelevantPrimitive> StaticPrimitives;
 	TChunkedArray<FRayTracingRelevantPrimitive> DynamicPrimitives;
 
+	// Relevant static primitive LODs are computed asynchronously.
+	// This task must complete before accessing StaticPrimitives in FRayTracingSceneAddInstancesTask.
+	FGraphEventRef StaticPrimitiveLODTask;
+
 	// Array of primitives that should update their cached ray tracing instances via FPrimitiveSceneInfo::UpdateCachedRaytracingData()
 	TArray<FPrimitiveSceneInfo*> DirtyCachedRayTracingPrimitives;
 
@@ -621,12 +625,10 @@ struct FRayTracingRelevantPrimitiveList
 	bool bValid = false;
 };
 
-// Iterates over Scene's PrimitiveSceneProxies and extracts ones that are relevant for ray tracing. 
+// Iterates over Scene's PrimitiveSceneProxies and extracts ones that are relevant for ray tracing.
 // This function can run on any thread.
-static FRayTracingRelevantPrimitiveList GatherRayTracingRelevantPrimitives(const FScene& Scene, const FViewInfo& View)
+static void GatherRayTracingRelevantPrimitives(const FScene& Scene, const FViewInfo& View, FRayTracingRelevantPrimitiveList& Result)
 {
-	FRayTracingRelevantPrimitiveList Result;
-
 	Result.DirtyCachedRayTracingPrimitives.Reserve(Scene.PrimitiveSceneProxies.Num());
 
 	const bool bGameView = View.bIsGameView || View.Family->EngineShowFlags.Game;
@@ -724,13 +726,15 @@ static FRayTracingRelevantPrimitiveList GatherRayTracingRelevantPrimitives(const
 		}
 	}
 
+	static const auto ICVarStaticMeshLODDistanceScale = IConsoleManager::Get().FindConsoleVariable(TEXT("r.StaticMeshLODDistanceScale"));
+	const float LODScaleCVarValue = ICVarStaticMeshLODDistanceScale->GetFloat();
+	const int32 ForcedLODLevel = GetCVarForceLOD();
+
+	Result.StaticPrimitiveLODTask = FFunctionGraphTask::CreateAndDispatchWhenReady(
+		[&Result, &Scene, &View, LODScaleCVarValue, ForcedLODLevel]()
 	{
 		TRACE_CPUPROFILER_EVENT_SCOPE(GatherRayTracingWorldInstances_ComputeLOD);
 
-		static const auto ICVarStaticMeshLODDistanceScale = IConsoleManager::Get().FindConsoleVariable(TEXT("r.StaticMeshLODDistanceScale"));
-		const float LODScaleCVarValue = ICVarStaticMeshLODDistanceScale->GetFloat();
-		const int32 ForcedLODLevel = GetCVarForceLOD();
-
 		ParallelFor(TEXT("GatherRayTracingRelevantPrimitives_ComputeLOD"), Result.StaticPrimitives.Num(), 128,
 			[&Result, &Scene, &View, LODScaleCVarValue, ForcedLODLevel](int32 ItemIndex) 
 		{
@@ -806,14 +810,12 @@ static FRayTracingRelevantPrimitiveList GatherRayTracingRelevantPrimitives(const
 				}
 			}
 		});
-	}
+	}, TStatId(), nullptr, ENamedThreads::AnyThread);
 
 	Result.bValid = true;
-
-	return Result;
 }
 
-bool FDeferredShadingSceneRenderer::GatherRayTracingWorldInstancesForView(FRDGBuilder& GraphBuilder, FViewInfo& View, FRayTracingScene& RayTracingScene, struct FRayTracingRelevantPrimitiveList&& RelevantPrimitiveList)
+bool FDeferredShadingSceneRenderer::GatherRayTracingWorldInstancesForView(FRDGBuilder& GraphBuilder, FViewInfo& View, FRayTracingScene& RayTracingScene, FRayTracingRelevantPrimitiveList& RelevantPrimitiveList)
 {
 	checkf(IsRayTracingEnabled() && bAnyRayTracingPassEnabled, TEXT("GatherRayTracingWorldInstancesForView should only be called if ray tracing is used"))
 
@@ -868,9 +870,6 @@ bool FDeferredShadingSceneRenderer::GatherRayTracingWorldInstancesForView(FRDGBu
 
 	// Consume output of the relevant primitive gathering task
 	RayTracingScene.UsedCoarseMeshStreamingHandles = MoveTemp(RelevantPrimitiveList.UsedCoarseMeshStreamingHandles);
-	TChunkedArray<FRayTracingRelevantPrimitive> RelevantDynamicPrimitives = MoveTemp(RelevantPrimitiveList.DynamicPrimitives);
-	TChunkedArray<FRayTracingRelevantPrimitive> RelevantStaticPrimitives = MoveTemp(RelevantPrimitiveList.StaticPrimitives);
-	TArray<FPrimitiveSceneInfo*> DirtyCachedRayTracingPrimitives = MoveTemp(RelevantPrimitiveList.DirtyCachedRayTracingPrimitives);
 
 	// Inform the coarse mesh streaming manager about all the used streamable render assets in the scene
 	Nanite::FCoarseMeshStreamingManager* CoarseMeshSM = IStreamingManager::Get().GetNaniteCoarseMeshStreamingManager();
@@ -879,9 +878,9 @@ bool FDeferredShadingSceneRenderer::GatherRayTracingWorldInstancesForView(FRDGBu
 		CoarseMeshSM->AddUsedStreamingHandles(RayTracingScene.UsedCoarseMeshStreamingHandles);
 	}
 
-	INC_DWORD_STAT_BY(STAT_VisibleRayTracingPrimitives, RelevantDynamicPrimitives.Num() + RelevantStaticPrimitives.Num());
+	INC_DWORD_STAT_BY(STAT_VisibleRayTracingPrimitives, RelevantPrimitiveList.DynamicPrimitives.Num() + RelevantPrimitiveList.StaticPrimitives.Num());
 
-	FPrimitiveSceneInfo::UpdateCachedRaytracingData(Scene, DirtyCachedRayTracingPrimitives);
+	FPrimitiveSceneInfo::UpdateCachedRaytracingData(Scene, RelevantPrimitiveList.DirtyCachedRayTracingPrimitives);
 
 	{
 		TRACE_CPUPROFILER_EVENT_SCOPE(GatherRayTracingWorldInstances_DynamicElements);
@@ -975,7 +974,7 @@ bool FDeferredShadingSceneRenderer::GatherRayTracingWorldInstancesForView(FRDGBu
 		// Local temporary array of instances used for GetDynamicRayTracingInstances()
 		TArray<FRayTracingInstance> TempRayTracingInstances;
 
-		for (const FRayTracingRelevantPrimitive& RelevantPrimitive : RelevantDynamicPrimitives)
+		for (const FRayTracingRelevantPrimitive& RelevantPrimitive : RelevantPrimitiveList.DynamicPrimitives)
 		{
 			const int32 PrimitiveIndex = RelevantPrimitive.PrimitiveIndex;
 			FPrimitiveSceneInfo* SceneInfo = Scene->Primitives[PrimitiveIndex];
@@ -1164,7 +1163,7 @@ bool FDeferredShadingSceneRenderer::GatherRayTracingWorldInstancesForView(FRDGBu
 		// Inputs
 
 		const FScene& Scene;
-		TChunkedArray<FRayTracingRelevantPrimitive> RelevantStaticPrimitives;
+		TChunkedArray<FRayTracingRelevantPrimitive>& RelevantStaticPrimitives;
 		const FRayTracingCullingParameters& CullingParameters;
 		const bool bEnableInstanceDebugData;
 
@@ -1174,11 +1173,11 @@ bool FDeferredShadingSceneRenderer::GatherRayTracingWorldInstancesForView(FRDGBu
 		TArray<FVisibleRayTracingMeshCommand>& VisibleRayTracingMeshCommands; // New elements are added here by this task
 
 		FRayTracingSceneAddInstancesTask(const FScene& InScene,
-											TChunkedArray<FRayTracingRelevantPrimitive> InRelevantStaticPrimitives,
+											TChunkedArray<FRayTracingRelevantPrimitive>& InRelevantStaticPrimitives,
 											const FRayTracingCullingParameters& InCullingParameters, bool InEnableInstanceDebugData,
 											FRayTracingScene& InRayTracingScene, TArray<FVisibleRayTracingMeshCommand>& InVisibleRayTracingMeshCommands)
 			: Scene(InScene)
-			, RelevantStaticPrimitives(MoveTemp(InRelevantStaticPrimitives))
+			, RelevantStaticPrimitives(InRelevantStaticPrimitives)
 			, CullingParameters(InCullingParameters)
 			, bEnableInstanceDebugData(InEnableInstanceDebugData)
 			, RayTracingScene(InRayTracingScene)
@@ -1434,8 +1433,11 @@ bool FDeferredShadingSceneRenderer::GatherRayTracingWorldInstancesForView(FRDGBu
 		}
 	};
 
-	FGraphEventRef AddInstancesTask = TGraphTask<FRayTracingSceneAddInstancesTask>::CreateTask().ConstructAndDispatchWhenReady(
-		*Scene, MoveTemp(RelevantStaticPrimitives), View.RayTracingCullingParameters, bEnableInstanceDebugData, // inputs 
+	FGraphEventArray AddInstancesTaskPrerequisites;
+	AddInstancesTaskPrerequisites.Add(RelevantPrimitiveList.StaticPrimitiveLODTask);
+
+	FGraphEventRef AddInstancesTask = TGraphTask<FRayTracingSceneAddInstancesTask>::CreateTask(&AddInstancesTaskPrerequisites).ConstructAndDispatchWhenReady(
+		*Scene, RelevantPrimitiveList.StaticPrimitives, View.RayTracingCullingParameters, bEnableInstanceDebugData, // inputs 
 		RayTracingScene, View.VisibleRayTracingMeshCommands // outputs
 	);
 
@@ -2280,7 +2282,7 @@ void FDeferredShadingSceneRenderer::Render(FRDGBuilder& GraphBuilder)
 			[Scene = this->Scene, &ReferenceView, &RayTracingRelevantPrimitiveList]()
 		{
 			FTaskTagScope TaskTagScope(ETaskTag::EParallelRenderingThread);
-			RayTracingRelevantPrimitiveList = GatherRayTracingRelevantPrimitives(*Scene, ReferenceView);
+			GatherRayTracingRelevantPrimitives(*Scene, ReferenceView, RayTracingRelevantPrimitiveList);
 		}, TStatId(), nullptr, ENamedThreads::AnyNormalThreadHiPriTask);
 	}
 #endif // RHI_RAYTRACING
@@ -2354,7 +2356,7 @@ void FDeferredShadingSceneRenderer::Render(FRDGBuilder& GraphBuilder)
 		// Prepare ray tracing scene instance list
 		checkf(RayTracingRelevantPrimitiveList.bValid, TEXT("Ray tracing relevant primitive list is expected to have been created before GatherRayTracingWorldInstancesForView() is called."));
 
-		GatherRayTracingWorldInstancesForView(GraphBuilder, ReferenceView, RayTracingScene, MoveTemp(RayTracingRelevantPrimitiveList));
+		GatherRayTracingWorldInstancesForView(GraphBuilder, ReferenceView, RayTracingScene, RayTracingRelevantPrimitiveList);
 	}
 #endif // RHI_RAYTRACING
 
diff --git a/Engine/Source/Runtime/Renderer/Private/DeferredShadingRenderer.h b/Engine/Source/Runtime/Renderer/Private/DeferredShadingRenderer.h
--- a/Engine/Source/Runtime/Renderer/Private/DeferredShadingRenderer.h
+++ b/Engine/Source/Runtime/Renderer/Private/DeferredShadingRenderer.h
@@ -1021,7 +1021,7 @@ private:
 	void RayTracingDisplayPicking(const FRayTracingPickingFeedback& PickingFeedback, FScreenMessageWriter& Writer);
 
 	/** Fills RayTracingScene instance list for the given View and adds relevant ray tracing data to the view. Does not reset previous scene contents. */
-	bool GatherRayTracingWorldInstancesForView(FRDGBuilder& GraphBuilder, FViewInfo& View, FRayTracingScene& RayTracingScene, struct FRayTracingRelevantPrimitiveList&& RelevantPrimitiveList);
+	bool GatherRayTracingWorldInstancesForView(FRDGBuilder& GraphBuilder, FViewInfo& View, FRayTracingScene& RayTracingScene, struct FRayTracingRelevantPrimitiveList& RelevantPrimitiveList);
 
 	bool SetupRayTracingPipelineStates(FRDGBuilder& GraphBuilder);
 	void SetupRayTracingLightDataForViews(FRDGBuilder& GraphBuilder);
