r.DynamicGlobalIlluminationMethod and r.ReflectionMethod are no longer scalability cvars.  These are the project default methods, which get overridden by [REDACTED_PERSON], and should not be driven by scalability (use r.Lumen.DiffuseIndirect.Allow and r.Lumen.Reflections.Allow for scalability)

diff --git a/Engine/Source/Runtime/Renderer/Private/IndirectLightRendering.cpp b/Engine/Source/Runtime/Renderer/Private/IndirectLightRendering.cpp
--- a/Engine/Source/Runtime/Renderer/Private/IndirectLightRendering.cpp
+++ b/Engine/Source/Runtime/Renderer/Private/IndirectLightRendering.cpp
@@ -25,6 +25,7 @@
 #include "ProfilingDebugging/CpuProfilerTrace.h"
 #include "Lumen/LumenTracingUtils.h"
 
+// This is the project default dynamic global illumination, NOT the scalability setting (see r.Lumen.DiffuseIndirect.Allow for scalability)
 // Must match EDynamicGlobalIlluminationMethod
 // Note: Default for new projects set by GameProjectUtils
 static TAutoConsoleVariable<int32> CVarDynamicGlobalIlluminationMethod(
@@ -34,8 +35,9 @@ static TAutoConsoleVariable<int32> CVarDynamicGlobalIlluminationMethod(
 	TEXT("2 - SSGI.  Standalone Screen Space Global Illumination.  Low cost, but limited by screen space information.\n")
 	TEXT("3 - RTGI.  Ray Traced Global Illumination technique.  Deprecated, use Lumen Global Illumination instead.\n")
 	TEXT("4 - Plugin.  Use a plugin for Global Illumination."),
-	ECVF_RenderThreadSafe | ECVF_Scalability);
+	ECVF_RenderThreadSafe);
 
+// This is the project default reflection method, NOT the scalability setting (see r.Lumen.Reflections.Allow for scalability)
 // Must match EReflectionMethod
 // Note: Default for new projects set by GameProjectUtils
 static TAutoConsoleVariable<int32> CVarReflectionMethod(
@@ -44,7 +46,7 @@ static TAutoConsoleVariable<int32> CVarReflectionMethod(
 	TEXT("1 - Lumen.  Use Lumen Reflections, which supports Screen / Software / Hardware Ray Tracing together and integrates with Lumen Global Illumination for rough reflections and Global Illumination seen in reflections.\n")
 	TEXT("2 - SSR.  Standalone Screen Space Reflections.  Low cost, but limited by screen space information.\n")
 	TEXT("3 - RT Reflections.  Ray Traced Reflections technique.  Deprecated, use Lumen Reflections instead."),
-	ECVF_RenderThreadSafe | ECVF_Scalability);
+	ECVF_RenderThreadSafe);
 
 static TAutoConsoleVariable<int32> CVarDiffuseIndirectHalfRes(
 	TEXT("r.DiffuseIndirect.HalfRes"), 1,
