Fixed: Planar reflections rendering artifacts with mobile deferred shading enabled

diff --git a/Engine/Source/Runtime/Renderer/Private/MobileDeferredShadingPass.cpp b/Engine/Source/Runtime/Renderer/Private/MobileDeferredShadingPass.cpp
--- a/Engine/Source/Runtime/Renderer/Private/MobileDeferredShadingPass.cpp
+++ b/Engine/Source/Runtime/Renderer/Private/MobileDeferredShadingPass.cpp
@@ -691,7 +691,7 @@ static void RenderLocalLight(
 
 	FGraphicsPipelineStateInitializer GraphicsPSOInit;
 	RHICmdList.ApplyCachedRenderTargets(GraphicsPSOInit);
-	GraphicsPSOInit.BlendState = TStaticBlendState<CW_RGBA, BO_Add, BF_One, BF_One, BO_Add, BF_One, BF_One>::GetRHI();
+	GraphicsPSOInit.BlendState = TStaticBlendState<CW_RGB, BO_Add, BF_One, BF_One, BO_Add, BF_One, BF_One>::GetRHI();
 	GraphicsPSOInit.PrimitiveType = PT_TriangleList;
 	const FSphere LightBounds = LightSceneInfo.Proxy->GetBoundingSphere();
 
@@ -841,7 +841,7 @@ static void RenderSimpleLights(
 		
 		RHICmdList.ApplyCachedRenderTargets(GraphicsPSOLight[PassIndex]);
 		// Use additive blending for color
-		GraphicsPSOLight[PassIndex].BlendState = TStaticBlendState<CW_RGBA, BO_Add, BF_One, BF_One, BO_Add, BF_One, BF_One>::GetRHI();
+		GraphicsPSOLight[PassIndex].BlendState = TStaticBlendState<CW_RGB, BO_Add, BF_One, BF_One, BO_Add, BF_One, BF_One>::GetRHI();
 		GraphicsPSOLight[PassIndex].PrimitiveType = PT_TriangleList;
 		GraphicsPSOLight[PassIndex].BoundShaderState.VertexDeclarationRHI = GetVertexDeclarationFVector4();
 		GraphicsPSOLight[PassIndex].BoundShaderState.VertexShaderRHI = VertexShader.GetVertexShader();
