Fixed Sky debug post process rendering.

diff --git a/Engine/Source/Runtime/Renderer/Private/PostProcess/PostProcessing.cpp b/Engine/Source/Runtime/Renderer/Private/PostProcess/PostProcessing.cpp
--- a/Engine/Source/Runtime/Renderer/Private/PostProcess/PostProcessing.cpp
+++ b/Engine/Source/Runtime/Renderer/Private/PostProcess/PostProcessing.cpp
@@ -355,8 +355,8 @@ void AddPostProcessingPasses(
 #else
 	PassSequence.SetEnabled(EPass::VisualizePostProcessStack, false);
 #endif
-	PassSequence.SetEnabled(EPass::VisualizeStrata, Strata::ShouldRenderStrataDebugPasses(View));
 	PassSequence.SetEnabled(EPass::VisualizeLumenScene, LumenVisualizeMode >= 0 && LumenVisualizeMode != VISUALIZE_MODE_OVERVIEW && bPostProcessingEnabled);
+	PassSequence.SetEnabled(EPass::VisualizeStrata, Strata::ShouldRenderStrataDebugPasses(View));
 	PassSequence.SetEnabled(EPass::VisualizeLightGrid, ShouldVisualizeLightGrid());
 
 #if WITH_EDITOR
@@ -1133,18 +1133,18 @@ void AddPostProcessingPasses(
 	}
 #endif
 
-	if (PassSequence.IsEnabled(EPass::VisualizeLightGrid))
+	if (PassSequence.IsEnabled(EPass::VisualizeStrata))
 	{
 		FScreenPassRenderTarget OverrideOutput;
-		PassSequence.AcceptOverrideIfLastPass(EPass::VisualizeLightGrid, OverrideOutput);
-		SceneColor = AddVisualizeLightGridPass(GraphBuilder, View, SceneColor, SceneDepth.Texture);
+		PassSequence.AcceptOverrideIfLastPass(EPass::VisualizeStrata, OverrideOutput);
+		SceneColor = Strata::AddStrataDebugPasses(GraphBuilder, View, SceneColor);
 	}
 
-	if (PassSequence.IsEnabled(EPass::VisualizeStrata))
+	if (PassSequence.IsEnabled(EPass::VisualizeLightGrid))
 	{
 		FScreenPassRenderTarget OverrideOutput;
-		PassSequence.AcceptOverrideIfLastPass(EPass::VisualizeStrata, OverrideOutput);
-		SceneColor = Strata::AddStrataDebugPasses(GraphBuilder, View, SceneColor);
+		PassSequence.AcceptOverrideIfLastPass(EPass::VisualizeLightGrid, OverrideOutput);
+		SceneColor = AddVisualizeLightGridPass(GraphBuilder, View, SceneColor, SceneDepth.Texture);
 	}
 
 #if WITH_EDITOR
@@ -1152,8 +1152,7 @@ void AddPostProcessingPasses(
 	{
 		FScreenPassRenderTarget OverrideOutput;
 		PassSequence.AcceptOverrideIfLastPass(EPass::VisualizeSkyAtmosphere, OverrideOutput);
-		OverrideOutput = OverrideOutput.IsValid() ? OverrideOutput : FScreenPassRenderTarget::CreateFromInput(GraphBuilder, SceneColor, View.GetOverwriteLoadAction(), TEXT("VisualizeSkyAtmosphere"));
-		SceneColor = AddSkyAtmosphereDebugPasses(GraphBuilder, Scene, *View.Family, View, OverrideOutput);
+		SceneColor = AddSkyAtmosphereDebugPasses(GraphBuilder, Scene, *View.Family, View, SceneColor);
 	}
 
 	if (PassSequence.IsEnabled(EPass::VisualizeLevelInstance))
