Exposure is now applied on some BufferVisualisation mode (those that are HDR only).
Also, removed a random x20 multiplier when visualising SeparateTranslucentRGB.

diff --git a/Engine/Source/Runtime/Renderer/Private/PostProcess/PostProcessEyeAdaptation.cpp b/Engine/Source/Runtime/Renderer/Private/PostProcess/PostProcessEyeAdaptation.cpp
--- a/Engine/Source/Runtime/Renderer/Private/PostProcess/PostProcessEyeAdaptation.cpp
+++ b/Engine/Source/Runtime/Renderer/Private/PostProcess/PostProcessEyeAdaptation.cpp
@@ -233,7 +233,10 @@ bool IsAutoExposureDebugMode(const FViewInfo& View)
 
 	return View.Family->UseDebugViewPS() ||
 		!EngineShowFlags.Lighting ||
-		(EngineShowFlags.VisualizeBuffer && View.CurrentBufferVisualizationMode != NAME_None) ||
+		(EngineShowFlags.VisualizeBuffer && View.CurrentBufferVisualizationMode != NAME_None
+			// Some HDR buffer visualization do want to have exposure applied.
+			&& View.CurrentBufferVisualizationMode != TEXT("SeparateTranslucencyRGB") && View.CurrentBufferVisualizationMode != TEXT("SceneColor") && View.CurrentBufferVisualizationMode != TEXT("FinalImage")
+			&& View.CurrentBufferVisualizationMode != TEXT("PreTonemapHDRColor") && View.CurrentBufferVisualizationMode != TEXT("PostTonemapHDRColor")) ||
 		EngineShowFlags.RayTracingDebug ||
 		EngineShowFlags.VisualizeDistanceFieldAO ||
 		EngineShowFlags.VisualizeVolumetricCloudConservativeDensity ||
