Fixed: TM-ShaderModels flickers with mobile deferred shading enabled on Adreno 630, 640 OS 10 devices

diff --git a/Engine/Source/Runtime/Renderer/Private/MobileShadingRenderer.cpp b/Engine/Source/Runtime/Renderer/Private/MobileShadingRenderer.cpp
--- a/Engine/Source/Runtime/Renderer/Private/MobileShadingRenderer.cpp
+++ b/Engine/Source/Runtime/Renderer/Private/MobileShadingRenderer.cpp
@@ -218,21 +218,29 @@ static void SetupGBufferFlags(FSceneTexturesConfig& SceneTexturesConfig, bool bR
 	if (!bRequiresMultiPass)
 	{
 		// use memoryless GBuffer when possible
-		AddFlags|= TexCreate_Memoryless;
+		AddFlags |= TexCreate_Memoryless;
 	}
-	
-	SceneTexturesConfig.GBufferA.Flags|= AddFlags;
-	SceneTexturesConfig.GBufferB.Flags|= AddFlags;
-	SceneTexturesConfig.GBufferC.Flags|= AddFlags;
-	SceneTexturesConfig.GBufferD.Flags|= AddFlags;
-	SceneTexturesConfig.GBufferE.Flags|= AddFlags;
-	
+
+	SceneTexturesConfig.GBufferA.Flags |= AddFlags;
+	SceneTexturesConfig.GBufferB.Flags |= AddFlags;
+	SceneTexturesConfig.GBufferC.Flags |= AddFlags;
+	SceneTexturesConfig.GBufferD.Flags |= AddFlags;
+	SceneTexturesConfig.GBufferE.Flags |= AddFlags;
+
 	// Mobile uses FBF/subpassLoad to fetch data from GBuffer, and FBF does not always work with sRGB targets 
 	SceneTexturesConfig.GBufferA.Flags &= (~TexCreate_SRGB);
 	SceneTexturesConfig.GBufferB.Flags &= (~TexCreate_SRGB);
 	SceneTexturesConfig.GBufferC.Flags &= (~TexCreate_SRGB);
 	SceneTexturesConfig.GBufferD.Flags &= (~TexCreate_SRGB);
 	SceneTexturesConfig.GBufferE.Flags &= (~TexCreate_SRGB);
+
+	// Input attachments with PF_R8G8B8A8 has better support on mobile than PF_B8G8R8A8
+	auto OverrideB8G8R8A8 = [](FGBufferBinding& Binding) { if (Binding.Format == PF_B8G8R8A8) Binding.Format = PF_R8G8B8A8; };
+	OverrideB8G8R8A8(SceneTexturesConfig.GBufferA);
+	OverrideB8G8R8A8(SceneTexturesConfig.GBufferB);
+	OverrideB8G8R8A8(SceneTexturesConfig.GBufferC);
+	OverrideB8G8R8A8(SceneTexturesConfig.GBufferD);
+	OverrideB8G8R8A8(SceneTexturesConfig.GBufferE);
 }
 
 static void PollOcclusionQueriesPass(FRDGBuilder& GraphBuilder)
diff --git a/Engine/Source/Runtime/VulkanRHI/Private/VulkanRenderpass.cpp b/Engine/Source/Runtime/VulkanRHI/Private/VulkanRenderpass.cpp
--- a/Engine/Source/Runtime/VulkanRHI/Private/VulkanRenderpass.cpp
+++ b/Engine/Source/Runtime/VulkanRHI/Private/VulkanRenderpass.cpp
@@ -152,9 +152,9 @@ public:
 				TSubpassDependencyClass& SubpassDep = SubpassDependencies[NumDependencies++];
 				SubpassDep.srcSubpass = 0;
 				SubpassDep.dstSubpass = 1;
-				SubpassDep.srcStageMask = VK_PIPELINE_STAGE_COLOR_ATTACHMENT_OUTPUT_BIT;
+				SubpassDep.srcStageMask = VK_PIPELINE_STAGE_COLOR_ATTACHMENT_OUTPUT_BIT | VK_PIPELINE_STAGE_EARLY_FRAGMENT_TESTS_BIT | VK_PIPELINE_STAGE_LATE_FRAGMENT_TESTS_BIT;
 				SubpassDep.dstStageMask = VK_PIPELINE_STAGE_FRAGMENT_SHADER_BIT;
-				SubpassDep.srcAccessMask = VK_ACCESS_COLOR_ATTACHMENT_WRITE_BIT;
+				SubpassDep.srcAccessMask = VK_ACCESS_COLOR_ATTACHMENT_WRITE_BIT | VK_ACCESS_DEPTH_STENCIL_ATTACHMENT_WRITE_BIT;
 				SubpassDep.dstAccessMask = VK_ACCESS_INPUT_ATTACHMENT_READ_BIT;
 				SubpassDep.dependencyFlags = VK_DEPENDENCY_BY_REGION_BIT;
 			}
