* Fix deep shadow not being rendered.
* Fix crash due to dangling RDG resource for non-visible grooms casting shadow
* Fix invalid groom bounds for non-visible grooms casting shadow

diff --git a/Engine/Plugins/Runtime/HairStrands/Source/HairStrandsCore/Private/GroomComponent.cpp b/Engine/Plugins/Runtime/HairStrands/Source/HairStrandsCore/Private/GroomComponent.cpp
--- a/Engine/Plugins/Runtime/HairStrands/Source/HairStrandsCore/Private/GroomComponent.cpp
+++ b/Engine/Plugins/Runtime/HairStrands/Source/HairStrandsCore/Private/GroomComponent.cpp
@@ -333,8 +333,6 @@ public:
 			const FHairGroupData& InGroupData = Component->GroomAsset->HairGroupsData[GroupIt];
 			FHairGroupInstance* HairInstance = Component->HairGroupInstances[GroupIt];
 			check(HairInstance->HairGroupPublicData);
-			HairInstance->ProxyBounds = &GetBounds();
-			HairInstance->ProxyLocalBounds = &GetLocalBounds();
 			HairInstance->bForceCards = Component->bUseCards;
 			HairInstance->bUpdatePositionOffset = Component->RegisteredMeshComponent != nullptr;
 			HairInstance->bCastShadow = Component->CastShadow;
@@ -460,6 +458,7 @@ public:
 			{
 				check(Instance->HairGroupPublicData != nullptr);
 				Instance->AddRef();
+				Instance->Debug.Proxy = this;
 				LocalScene.AddHairStrands(Instance);
 			}
 		}
@@ -483,6 +482,7 @@ public:
 			{
 				check(Instance->GetRefCount() > 0);
 				LocalScene.RemoveHairStrands(Instance);
+				Instance->Debug.Proxy = nullptr;
 				Instance->Release();
 			}
 		}
diff --git a/Engine/Plugins/Runtime/HairStrands/Source/HairStrandsCore/Private/GroomManager.cpp b/Engine/Plugins/Runtime/HairStrands/Source/HairStrandsCore/Private/GroomManager.cpp
--- a/Engine/Plugins/Runtime/HairStrands/Source/HairStrandsCore/Private/GroomManager.cpp
+++ b/Engine/Plugins/Runtime/HairStrands/Source/HairStrandsCore/Private/GroomManager.cpp
@@ -743,7 +743,7 @@ static void RunHairLODSelection(
 		float LODViewIndex = -1;
 		{
 			float MaxScreenSize = 0.f;
-			const FSphere SphereBound = Instance->ProxyBounds ? Instance->ProxyBounds->GetSphere() : FSphere(0);
+			const FSphere SphereBound = Instance->GetBounds().GetSphere();
 			for (const FSceneView* View : Views)
 			{
 				const float ScreenSize = ComputeBoundsScreenSize(FVector4(SphereBound.Center, 1), SphereBound.W, *View);
diff --git a/Engine/Plugins/Runtime/HairStrands/Source/HairStrandsCore/Public/GroomInstance.h b/Engine/Plugins/Runtime/HairStrands/Source/HairStrandsCore/Public/GroomInstance.h
--- a/Engine/Plugins/Runtime/HairStrands/Source/HairStrandsCore/Public/GroomInstance.h
+++ b/Engine/Plugins/Runtime/HairStrands/Source/HairStrandsCore/Public/GroomInstance.h
@@ -250,6 +250,7 @@ struct HAIRSTRANDSCORE_API FHairGroupInstance : public FHairStrandsInstance
 		int32					MeshLODIndex = ~0;
 		EGroomBindingMeshType	GroomBindingType;
 		EGroomCacheType			GroomCacheType;
+		FPrimitiveSceneProxy*	Proxy = nullptr;
 		UMeshComponent*			MeshComponent = nullptr;
 		FString					MeshComponentName;
 		const UGroomComponent*	GroomComponentForDebug = nullptr; // For debug only, shouldn't be deferred on the rendering thread
@@ -274,8 +275,6 @@ struct HAIRSTRANDSCORE_API FHairGroupInstance : public FHairStrandsInstance
 	FHairGroupPublicData*	HairGroupPublicData = nullptr;
 	EHairGeometryType		GeometryType = EHairGeometryType::NoneGeometry;
 	EHairBindingType		BindingType = EHairBindingType::NoneBinding;
-	const FBoxSphereBounds*	ProxyBounds = nullptr;
-	const FBoxSphereBounds* ProxyLocalBounds = nullptr;
 	bool					bForceCards = false;
 	bool					bUpdatePositionOffset = false;
 	bool					bCastShadow = true;
@@ -291,7 +290,8 @@ struct HAIRSTRANDSCORE_API FHairGroupInstance : public FHairStrandsInstance
 		return Meshes.IsValid() || Cards.IsValid() || Strands.IsValid();
 	}
 
-	virtual const FBoxSphereBounds* GetBounds() const override { return ProxyBounds;  }
+	virtual const FBoxSphereBounds& GetBounds() const override { return Debug.Proxy->GetBounds(); }
+	virtual const FBoxSphereBounds& GetLocalBounds() const { return Debug.Proxy->GetLocalBounds(); }
 	virtual const FHairGroupPublicData* GetHairData() const { return HairGroupPublicData; }
 
 	/** Get the current local to world transform according to the internal binding type */
diff --git a/Engine/Source/Runtime/Renderer/Private/HairStrands/HairStrandsCluster.cpp b/Engine/Source/Runtime/Renderer/Private/HairStrands/HairStrandsCluster.cpp
--- a/Engine/Source/Runtime/Renderer/Private/HairStrands/HairStrandsCluster.cpp
+++ b/Engine/Source/Runtime/Renderer/Private/HairStrands/HairStrandsCluster.cpp
@@ -201,7 +201,7 @@ void CreateHairStrandsMacroGroups(
 
 	// Aggregate all hair primitives within the same area into macro groups, for allocating/rendering DOM/voxel
 	uint32 MacroGroupId = 0;
-	auto UpdateMacroGroup = [&MacroGroups, &MacroGroupId, &MaterialId](FHairGroupPublicData* HairData, const FMeshBatch* Mesh,  const FPrimitiveSceneProxy* Proxy, const FBoxSphereBounds* Bounds)
+	auto UpdateMacroGroup = [&MacroGroups, &MacroGroupId, &MaterialId](FHairGroupPublicData* HairData, const FMeshBatch* Mesh,  const FPrimitiveSceneProxy* Proxy, const FBoxSphereBounds& Bounds)
 	{
 		check(HairData);
 
@@ -210,7 +210,7 @@ void CreateHairStrandsMacroGroups(
 		if (!bIsValid)
 			return;
 
-		const FBoxSphereBounds& PrimitiveBounds = Proxy ? Proxy->GetBounds() : *Bounds;
+		const FBoxSphereBounds& PrimitiveBounds = Proxy ? Proxy->GetBounds() : Bounds;
 
 		bool bFound = false;
 		float MinDistance = FLT_MAX;
@@ -264,6 +264,7 @@ void CreateHairStrandsMacroGroups(
 	};
 
 	// 1. Add all visible hair-strands instances
+	static FBoxSphereBounds EmptyBound(ForceInit);
 	const int32 ActiveInstanceCount = Scene->HairStrandsSceneData.RegisteredProxies.Num();
 	TBitArray InstancesVisibility(false, ActiveInstanceCount);
 	for (const FMeshBatchAndRelevance& MeshBatchAndRelevance : View.HairStrandsMeshElements)
@@ -272,7 +273,7 @@ void CreateHairStrandsMacroGroups(
 		{
 			if (FHairGroupPublicData* HairData = HairStrands::GetHairData(MeshBatchAndRelevance.Mesh))
 			{
-				UpdateMacroGroup(HairData, MeshBatchAndRelevance.Mesh, MeshBatchAndRelevance.PrimitiveSceneProxy, nullptr);
+				UpdateMacroGroup(HairData, MeshBatchAndRelevance.Mesh, MeshBatchAndRelevance.PrimitiveSceneProxy, EmptyBound);
 				InstancesVisibility[HairData->Instance->RegisteredIndex] = true;
 			}
 		}
@@ -284,7 +285,7 @@ void CreateHairStrandsMacroGroups(
 	{
 		for (FHairStrandsInstance* Instance : Scene->HairStrandsSceneData.RegisteredProxies)
 		{
-			if (Instance->RegisteredIndex >= 0 && Instance->RegisteredIndex < ActiveInstanceCount && !InstancesVisibility[Instance->RegisteredIndex])
+			if (Instance && Instance->RegisteredIndex >= 0 && Instance->RegisteredIndex < ActiveInstanceCount && !InstancesVisibility[Instance->RegisteredIndex])
 			{
 				if (IsHairStrandsVisibleInShadows(View, *Instance))
 				{
diff --git a/Engine/Source/Runtime/Renderer/Private/HairStrands/HairStrandsDeepShadow.cpp b/Engine/Source/Runtime/Renderer/Private/HairStrands/HairStrandsDeepShadow.cpp
--- a/Engine/Source/Runtime/Renderer/Private/HairStrands/HairStrandsDeepShadow.cpp
+++ b/Engine/Source/Runtime/Renderer/Private/HairStrands/HairStrandsDeepShadow.cpp
@@ -264,13 +264,13 @@ void RenderHairStrandsDeepShadows(
 				const FLightSceneProxy* LightProxy = LightInfo->Proxy;
 				if (!LightProxy->AffectsBounds(MacroGroupBounds))
 				{
-					return;
+					continue;
 				}
 
 				// Run out of atlas slot
 				if (DOMSlotCount >= FHairStrandsDeepShadowResources::MaxAtlasSlotCount)
 				{
-					return;
+					continue;
 				}
 
 				DOMSlotCount++;
diff --git a/Engine/Source/Runtime/Renderer/Private/HairStrands/HairStrandsInterface.cpp b/Engine/Source/Runtime/Renderer/Private/HairStrands/HairStrandsInterface.cpp
--- a/Engine/Source/Runtime/Renderer/Private/HairStrands/HairStrandsInterface.cpp
+++ b/Engine/Source/Runtime/Renderer/Private/HairStrands/HairStrandsInterface.cpp
@@ -442,12 +442,12 @@ bool IsHairStrandsVisibleInShadows(const FViewInfo& View, const FHairStrandsInst
 			return false;
 		}
 
-		if (const FBoxSphereBounds* Bounds = Instance.GetBounds())
+		const FBoxSphereBounds& Bounds = Instance.GetBounds();
 		{
 			// Local lights
 			for (const FLightSceneInfo* LightInfo : View.HairStrandsViewData.VisibleShadowCastingLights)
 			{
-				if (LightInfo->Proxy->AffectsBounds(*Bounds))
+				if (LightInfo->Proxy->AffectsBounds(Bounds))
 				{
 					bIsVisibleInShadow = true;
 					break;
@@ -464,7 +464,7 @@ bool IsHairStrandsVisibleInShadows(const FViewInfo& View, const FHairStrandsInst
 					const FMatrix& WorldToLight = CullData.LightInfo->Proxy->GetWorldToLight();
 					
 					// Transform groom bound into light space, and extend it along the light direction
-					FBoxSphereBounds BoundsInLightSpace = Bounds->TransformBy(WorldToLight);
+					FBoxSphereBounds BoundsInLightSpace = Bounds.TransformBy(WorldToLight);
 					BoundsInLightSpace.BoxExtent.X += CullDistance;
 
 					const bool bIntersect = CullData.ViewFrustumInLightSpace.IntersectBox(BoundsInLightSpace.Origin, BoundsInLightSpace.BoxExtent);
@@ -477,7 +477,7 @@ bool IsHairStrandsVisibleInShadows(const FViewInfo& View, const FHairStrandsInst
 						LightData.Extent = FVector3f(CullData.LightInfo->Proxy->GetBoundingSphere().W);
 						LightData.ViewFrustumInLightSpace = CullData.ViewFrustumInLightSpace;
 						LightData.InstanceBoundInLightSpace.Add({FVector3f(BoundsInLightSpace.GetBox().Min), FVector3f(BoundsInLightSpace.GetBox().Max)});
-						LightData.InstanceBoundInWorldSpace.Add({ FVector3f(Bounds->GetBox().Min), FVector3f(Bounds->GetBox().Max)});
+						LightData.InstanceBoundInWorldSpace.Add({ FVector3f(Bounds.GetBox().Min), FVector3f(Bounds.GetBox().Max)});
 						LightData.InstanceIntersection.Add(bIntersect ? 1u : 0u);
 					}
 
diff --git a/Engine/Source/Runtime/Renderer/Private/HairStrands/HairStrandsRasterCommon.cpp b/Engine/Source/Runtime/Renderer/Private/HairStrands/HairStrandsRasterCommon.cpp
--- a/Engine/Source/Runtime/Renderer/Private/HairStrands/HairStrandsRasterCommon.cpp
+++ b/Engine/Source/Runtime/Renderer/Private/HairStrands/HairStrandsRasterCommon.cpp
@@ -355,10 +355,16 @@ void AddHairStrandsRasterPass(
 
 		for (const FHairStrandsMacroGroupData::PrimitiveInfo& PrimitiveInfo : PrimitiveSceneInfos)
 		{
-			const bool bCullingEnable = PrimitiveInfo.IsCullingEnable();
-			const FMeshBatch& MeshBatch = *PrimitiveInfo.Mesh;
-			const uint64 BatchElementMask = ~0ull;
-			HairRasterMeshProcessor.AddMeshBatch(MeshBatch, BatchElementMask, PrimitiveInfo.PrimitiveSceneProxy, -1 , bCullingEnable);
+			// Ensure that we submit only primitive with valid MeshBatch data. Non-visible groom casting 
+			// deep shadow map, will be skipped. This will be fixed once we remove mesh processor to replace 
+			// it with a global shader
+			if (PrimitiveInfo.Mesh != nullptr)
+			{
+				const bool bCullingEnable = PrimitiveInfo.IsCullingEnable();
+				const FMeshBatch& MeshBatch = *PrimitiveInfo.Mesh;
+				const uint64 BatchElementMask = ~0ull;
+				HairRasterMeshProcessor.AddMeshBatch(MeshBatch, BatchElementMask, PrimitiveInfo.PrimitiveSceneProxy, -1 , bCullingEnable);
+			}
 		}
 	});
 }
diff --git a/Engine/Source/Runtime/Renderer/Private/HairStrands/HairStrandsVisibility.cpp b/Engine/Source/Runtime/Renderer/Private/HairStrands/HairStrandsVisibility.cpp
--- a/Engine/Source/Runtime/Renderer/Private/HairStrands/HairStrandsVisibility.cpp
+++ b/Engine/Source/Runtime/Renderer/Private/HairStrands/HairStrandsVisibility.cpp
@@ -2662,6 +2662,11 @@ void AddMeshDrawTransitionPass(
 	{
 		for (const FHairStrandsMacroGroupData::PrimitiveInfo& PrimitiveInfo : MacroGroup.PrimitivesInfos)
 		{
+			if (PrimitiveInfo.Mesh == nullptr)
+			{
+				continue;
+			}
+
 			FHairGroupPublicData* HairGroupPublicData = PrimitiveInfo.PublicDataPtr;
 			check(HairGroupPublicData);
 
diff --git a/Engine/Source/Runtime/Renderer/Private/HairStrands/HairStrandsVoxelization.cpp b/Engine/Source/Runtime/Renderer/Private/HairStrands/HairStrandsVoxelization.cpp
--- a/Engine/Source/Runtime/Renderer/Private/HairStrands/HairStrandsVoxelization.cpp
+++ b/Engine/Source/Runtime/Renderer/Private/HairStrands/HairStrandsVoxelization.cpp
@@ -1320,7 +1320,7 @@ static void AddVirtualVoxelizationComputeRasterPass(
 			
 				const bool bCullingEnable = !IsHairStrandContinuousDecimationReorderingEnabled() && HairGroupPublicData->GetCullingResultAvailable();
 
-				PassParameters->HairInstance = GetHairStrandsInstanceParameters(GraphBuilder, *ViewInfo, HairGroupPublicData, bCullingEnable, false);
+				PassParameters->HairInstance = GetHairStrandsInstanceParameters(GraphBuilder, *ViewInfo, HairGroupPublicData, bCullingEnable, true);
 
 				if (bCullingEnable)
 				{
diff --git a/Engine/Source/Runtime/Renderer/Public/HairStrandsInterface.h b/Engine/Source/Runtime/Renderer/Public/HairStrandsInterface.h
--- a/Engine/Source/Runtime/Renderer/Public/HairStrandsInterface.h
+++ b/Engine/Source/Runtime/Renderer/Public/HairStrandsInterface.h
@@ -143,7 +143,8 @@ struct RENDERER_API FHairStrandsInstance
 	uint32 AddRef() const;
 	uint32 Release() const;
 	int32 RegisteredIndex = -1;
-	virtual const FBoxSphereBounds* GetBounds() const { return nullptr; }
+	virtual const FBoxSphereBounds& GetBounds() const = 0;
+	virtual const FBoxSphereBounds& GetLocalBounds() const = 0;
 	virtual const FHairGroupPublicData* GetHairData() const { return nullptr; }
 protected:
 	mutable uint32 RefCount = 0;
