Multi-view Lumen lighting artifacts:  CVar added to disable feature that only updates ray tracing scene once.  Costs around 1.6 ms GPU for a Virtual Production test case (no change in perf outside nDisplay).  The optimization is new to 5.1, so this is just reverting to 5.0 perf levels.  Planning to investigate more, but wanted artifacts solved in the meantime for Virtual Production clients.  The CVar would be useful regardless for A/B testing.

It's possible the artifacts are caused by a bug, or the ray tracing scene needs to be upgraded in some way to be mul0ti-view aware.  If neither of those options work, we could consider supporting multiple ray tracing scenes.  If views are already on separate GPUs, that would involve no additional cost.

#rnx

diff --git a/Engine/Source/Runtime/Renderer/Private/SceneRendering.cpp b/Engine/Source/Runtime/Renderer/Private/SceneRendering.cpp
--- a/Engine/Source/Runtime/Renderer/Private/SceneRendering.cpp
+++ b/Engine/Source/Runtime/Renderer/Private/SceneRendering.cpp
@@ -206,6 +206,13 @@ static TAutoConsoleVariable<int32> CVarViewRectUseScreenBottom(
 	ECVF_RenderThreadSafe
 );
 
+static TAutoConsoleVariable<int32> CVarRayTracingSceneUpdateOnce(
+	TEXT("r.RayTracing.SceneUpdateOnce"),
+	0,
+	TEXT("Experimental:  Improves GPU perf by updating ray tracing scene once, but may cause artifacts (mainly for nDisplay)\n"),
+	ECVF_RenderThreadSafe
+);
+
 #if !(UE_BUILD_SHIPPING || UE_BUILD_TEST)
 static TAutoConsoleVariable<float> CVarGeneralPurposeTweak(
 	TEXT("r.GeneralPurposeTweak"),
@@ -3862,8 +3869,10 @@ void FSceneRenderer::CreateSceneRenderers(TArrayView<const FSceneViewFamily*> In
 		{
 			// For multi-view-family scene rendering, we need to determine which scene renderer will update the ray tracing
 			// scene.  This will be the first view family that uses ray tracing, and subsequent families that use ray
-			// tracing can skip that step.  
+			// tracing can skip that step.  If the optimization to update the ray tracing scene once is disabled, we'll
+			// update it for all scene renders.
 			bool bShouldUpdateRayTracingScene = true;
+			const bool bRayTracingSceneUpdateOnce = CVarRayTracingSceneUpdateOnce.GetValueOnRenderThread() != 0;
 
 			for (int32 RendererIndex = 0; RendererIndex < SceneRenderers.Num(); RendererIndex++)
 			{
@@ -3874,7 +3883,10 @@ void FSceneRenderer::CreateSceneRenderers(TArrayView<const FSceneViewFamily*> In
 				if (SceneRenderer->bAnyRayTracingPassEnabled)
 				{
 					SceneRenderer->bShouldUpdateRayTracingScene = bShouldUpdateRayTracingScene;
-					bShouldUpdateRayTracingScene = false;
+					if (bRayTracingSceneUpdateOnce)
+					{
+						bShouldUpdateRayTracingScene = false;
+					}
 				}
 			}
 
