Multi-view-family render:  More robust handling of ray tracing scene update.  There is an optimization where we only update the ray tracing scene for the first view family when rendering multiple view families together...

* There was a bug where if the first view family didn't use ray tracing, the ray tracing scene would never be updated.  Now the ray tracing scene is updated on the first view family using ray tracing.  Also cleaned up multiple places we compute "bAnyRayTracingPassEnabled" by moving that flag to FDeferredSceneRenderer.
* There was another bug where view specific light ray tracing RDG buffers were created during the scene renderer that updates the ray tracing scene, potentially leading to those RDG buffers being used outside the graph where they were created.  Now the view specific light buffers are created in each scene renderer.  The code was cleaned up so the buffer creation logic only has access to the static "FSceneView", not any of the transient state in "FViewInfo", to enforce that there are no dependencies on transient state.

#rnx

diff --git a/Engine/Source/Runtime/Renderer/Private/DeferredShadingRenderer.cpp b/Engine/Source/Runtime/Renderer/Private/DeferredShadingRenderer.cpp
--- a/Engine/Source/Runtime/Renderer/Private/DeferredShadingRenderer.cpp
+++ b/Engine/Source/Runtime/Renderer/Private/DeferredShadingRenderer.cpp
@@ -394,6 +394,19 @@ FDeferredShadingSceneRenderer::FDeferredShadingSceneRenderer(const FSceneViewFam
 #endif
 
 	ShadowSceneRenderer = MakeUnique<FShadowSceneRenderer>(*this);
+
+#if RHI_RAYTRACING
+	bAnyRayTracingPassEnabled = false;
+	for (int32 ViewIndex = 0; ViewIndex < Views.Num(); ViewIndex++)
+	{
+		bool bHasRayTracing = AnyRayTracingPassEnabled(Scene, Views[ViewIndex]);
+
+		Views[ViewIndex].bHasAnyRayTracingPass = bHasRayTracing;
+
+		bAnyRayTracingPassEnabled |= bHasRayTracing;
+	}
+	bShouldUpdateRayTracingScene = bAnyRayTracingPassEnabled;
+#endif  // RHI_RAYTRACING
 }
 
 /** 
