RoboMerge integration fix:  Part of a CL was integrated that shouldn't have been.  Integrating from //UE5/Release-Engine-Staging, where bad integration was introduced.

#rnx

diff --git a/Engine/Source/Runtime/Renderer/Private/DeferredShadingRenderer.cpp b/Engine/Source/Runtime/Renderer/Private/DeferredShadingRenderer.cpp
--- a/Engine/Source/Runtime/Renderer/Private/DeferredShadingRenderer.cpp
+++ b/Engine/Source/Runtime/Renderer/Private/DeferredShadingRenderer.cpp
@@ -394,19 +394,6 @@ FDeferredShadingSceneRenderer::FDeferredShadingSceneRenderer(const FSceneViewFam
 #endif
 
 	ShadowSceneRenderer = MakeUnique<FShadowSceneRenderer>(*this);
-
-#if RHI_RAYTRACING
-	bAnyRayTracingPassEnabled = false;
-	for (int32 ViewIndex = 0; ViewIndex < Views.Num(); ViewIndex++)
-	{
-		bool bHasRayTracing = AnyRayTracingPassEnabled(Scene, Views[ViewIndex]);
-
-		Views[ViewIndex].bHasAnyRayTracingPass = bHasRayTracing;
-
-		bAnyRayTracingPassEnabled |= bHasRayTracing;
-	}
-	bShouldUpdateRayTracingScene = bAnyRayTracingPassEnabled;
-#endif  // RHI_RAYTRACING
 }
 
 /** 
