Compile the VisualizeHDR shader using DXC to be faster.

diff --git a/Engine/Source/Runtime/Renderer/Private/PostProcess/PostProcessVisualizeHDR.cpp b/Engine/Source/Runtime/Renderer/Private/PostProcess/PostProcessVisualizeHDR.cpp
--- a/Engine/Source/Runtime/Renderer/Private/PostProcess/PostProcessVisualizeHDR.cpp
+++ b/Engine/Source/Runtime/Renderer/Private/PostProcess/PostProcessVisualizeHDR.cpp
@@ -47,6 +47,9 @@ public:
 		OutEnvironment.SetDefine(TEXT("USE_SHADOW_TINT"), 1);
 		OutEnvironment.SetDefine(TEXT("USE_CONTRAST"), 1);
 		OutEnvironment.SetDefine(TEXT("USE_APPROXIMATE_SRGB"), (uint32)0);
+
+		// This shader takes a very long time to compile with FXC, so we pre-compile it with DXC first and then forward the optimized HLSL to FXC.
+		OutEnvironment.CompilerFlags.Add(CFLAG_PrecompileWithDXC);
 	}
 };
 
