Fixed G-buffer visualization label positioning.
Canvas base transform now uses viewrect instead of texture size. Label coordinates are now unscaled before drawing, as DrawShadowedString also performs DPI scaling.

diff --git a/Engine/Source/Runtime/Renderer/Private/PostProcess/PostProcessVisualizeBuffer.cpp b/Engine/Source/Runtime/Renderer/Private/PostProcess/PostProcessVisualizeBuffer.cpp
--- a/Engine/Source/Runtime/Renderer/Private/PostProcess/PostProcessVisualizeBuffer.cpp
+++ b/Engine/Source/Runtime/Renderer/Private/PostProcess/PostProcessVisualizeBuffer.cpp
@@ -138,10 +138,13 @@ FScreenPassTexture AddVisualizeBufferPass(FRDGBuilder& GraphBuilder, const FView
 
 	AddDrawCanvasPass(GraphBuilder, RDG_EVENT_NAME("Labels"), View, Output, [LocalTileLabels = MoveTemp(TileLabels)](FCanvas& Canvas)
 	{
+		Canvas.SetBaseTransform(FMatrix(FScaleMatrix(Canvas.GetDPIScale()) * Canvas.CalcBaseTransform2D(Canvas.GetViewRect().Width(), Canvas.GetViewRect().Height())));
+
 		const FLinearColor LabelColor(1, 1, 0);
 		for (const FTileLabel& TileLabel : LocalTileLabels)
 		{
-			Canvas.DrawShadowedString(TileLabel.Location.X, TileLabel.Location.Y, *TileLabel.Label, GetStatsFont(), LabelColor);
+			const float DPIScale = Canvas.GetDPIScale();
+			Canvas.DrawShadowedString(TileLabel.Location.X / DPIScale, TileLabel.Location.Y / DPIScale, *TileLabel.Label, GetStatsFont(), LabelColor);
 		}
 	});
 
