Wrap TAA downsample optimization around cvar, and disable on XboxOneGDK (small gain: 0.04ms)

diff --git a/Engine/Source/Runtime/Renderer/Private/PostProcess/PostProcessing.cpp b/Engine/Source/Runtime/Renderer/Private/PostProcess/PostProcessing.cpp
--- a/Engine/Source/Runtime/Renderer/Private/PostProcess/PostProcessing.cpp
+++ b/Engine/Source/Runtime/Renderer/Private/PostProcess/PostProcessing.cpp
@@ -131,6 +131,13 @@ TAutoConsoleVariable<int32> CVarPostProcessingForceAsyncDispatch(
 	ECVF_RenderThreadSafe);
 #endif
 
+TAutoConsoleVariable<bool> CVarAllowDownsampleSceneColorInTAA(
+	TEXT("r.PostProcessing.AllowDownsampleSceneColorInTAA"),
+	true,
+	TEXT("Allow TAA to output scene color in half resolution.\n")
+	TEXT(" (default: true)"),
+	ECVF_RenderThreadSafe);
+
 } //! namespace
 
 EDownsampleQuality GetDownsampleQuality(const TAutoConsoleVariable<int32>& CVar)
@@ -628,7 +635,8 @@ void AddPostProcessingPasses(
 			{
 				UpscalerPassInputs.bAllowDownsampleSceneColor =
 					(bNeedPostMotionBlurHalfRes || bNeedPostMotionBlurQuarterRes) &&
-					DownsampleQuality == EDownsampleQuality::Low;;
+					DownsampleQuality == EDownsampleQuality::Low &&
+					CVarAllowDownsampleSceneColorInTAA.GetValueOnRenderThread() == true;
 			}
 			UpscalerPassInputs.DownsampleOverrideFormat = DownsampleOverrideFormat;
 			UpscalerPassInputs.SceneColorTexture = SceneColor.Texture;
