Change material buffer clear to use Cmask data when available.

* This allows to no clear material buffer prior to the base pass
* Material buffer clearing is done during the classification pass using Cmask for coarse clearing and TopLayerData for fine clearing

Net gain is -0.08ms.

diff --git a/Engine/Shaders/Private/Strata/StrataMaterialClassification.usf b/Engine/Shaders/Private/Strata/StrataMaterialClassification.usf
--- a/Engine/Shaders/Private/Strata/StrataMaterialClassification.usf
+++ b/Engine/Shaders/Private/Strata/StrataMaterialClassification.usf
@@ -16,7 +16,12 @@ int bRectPrimitive;
 int2 ViewResolution;
 uint MaxBytesPerPixel;
 Texture2D<uint> TopLayerTexture;
+#if PERMUTATION_CMASK
+Texture2D<uint> TopLayerCmaskTexture;
+RWTexture2D<uint> MaterialHeaderTextureUAV;
+#else
 Texture2DArray<uint> MaterialTextureArray;
+#endif
 
 // Indirect draw data buffer for all tile types
 RWBuffer<uint> TileDrawIndirectDataBuffer;
@@ -35,9 +40,7 @@ Texture2D<float3> OpaqueRoughRefractionTexture;
 groupshared uint s_TileFlags[GROUP_THREAD_COUNT];
 #endif
 
-#if PERMUTATION_STRATA_CLEAR_DURING_CATEGORIZATION
 RWTexture2DArray<uint> SSSTextureUAV;
-#endif
 
 [numthreads(STRATA_TILE_SIZE, STRATA_TILE_SIZE, 1)]
 void TileMainCS(uint2 DispatchThreadId : SV_DispatchThreadID, uint LinearIndex : SV_GroupIndex, uint3 GroupId : SV_GroupID)
@@ -54,7 +57,25 @@ void TileMainCS(uint2 DispatchThreadId : SV_DispatchThreadID, uint LinearIndex :
 
 	const uint2 PixelCoord = DispatchThreadId.xy + View.ViewRectMin.xy;
 	const bool bIsValid = all(DispatchThreadId.xy < View.ViewSizeAndInvSize.xy);
-	
+
+	// If CMask data are available, we use it as a coarse evaluation to know if a tile contains any data. 
+	// * If the tile is entirely empty: we clear the header & SSS data
+	// * If the data contains any data: we do fine grain checking, and clear header & SSS data only for needed pixels. The top layer data texture is used 
+	//   to know if a pixel is valid or not (since the material header is not cleared when the Cmask permutation is used).
+
+#if PERMUTATION_CMASK
+	// Coarse test for clearing header (& SSS data) based on CMask data
+	const uint CMask = TopLayerCmaskTexture.Load(uint3(GroupId.xy, 0));
+	BRANCH
+	if (CMask == 0x0)
+	{
+		MaterialHeaderTextureUAV[PixelCoord] = 0u;
+		SSSTextureUAV[uint3(PixelCoord, 0)] = 0; // This is a good clear for FStrataSubsurfaceHeader, and we only need to clear the header.
+	}
+	else
+	{
+#endif
+
 	// Pixels outside of the view area are considered simple to enable screen borders to receive the simple permutation when not aligned to shader group size.
 	bool bContainsComplexMaterial = false;
 	bool bContainsSimpleMaterial  = false;
@@ -66,7 +87,14 @@ void TileMainCS(uint2 DispatchThreadId : SV_DispatchThreadID, uint LinearIndex :
 	if (bIsValid)
 	{
 		FStrataAddressing StrataAddressing = GetStrataPixelDataByteOffset(PixelCoord, uint2(View.BufferSizeAndInvSize.xy), MaxBytesPerPixel);
+	#if PERMUTATION_CMASK
+		FStrataPixelHeader StrataPixelHeader = InitialiseStrataPixelHeader();
+		const uint PackedHeader		= MaterialHeaderTextureUAV[PixelCoord];
+		StrataPixelHeader.State		= HEADER_GETCOMMONSTATES(PackedHeader);
+		StrataPixelHeader.BSDFCount = HEADER_GETBSDFCOUNT(PackedHeader);
+	#else
 		FStrataPixelHeader StrataPixelHeader = UnpackStrataHeaderIn(MaterialTextureArray, StrataAddressing, TopLayerTexture);
+	#endif
 
 		const bool bIsSimple =  IsSimpleMaterial(StrataPixelHeader) || StrataPixelHeader.BSDFCount == 0; // BSDFCount == 0 ensures that non-strata pixel, like sky pixels, won't make a simple tile flagged as complex
 		const bool bIsSingle = !IsSimpleMaterial(StrataPixelHeader) && IsSingleMaterial(StrataPixelHeader);
@@ -76,20 +104,36 @@ void TileMainCS(uint2 DispatchThreadId : SV_DispatchThreadID, uint LinearIndex :
 		bContainsComplexMaterial = !bIsSingle & !bIsSimple;
 		bContainsScreenSpaceSubsurfaceScattering = HasSubsurface(StrataPixelHeader);
 
-#if STRATA_OPAQUE_ROUGH_REFRACTION_ENABLED
+	#if STRATA_OPAQUE_ROUGH_REFRACTION_ENABLED
 		OpaqueRoughRefractionData = StrataUnpackOpaqueRoughRefractionData(OpaqueRoughRefractionTexture[PixelCoord]);
 		bContainsOpaqueRoughRefraction = OpaqueRoughRefractionData.OpaqueRoughRefractionEnabled > 0.0f;
-#endif
-	}
+	#endif
 
-#if PERMUTATION_STRATA_CLEAR_DURING_CATEGORIZATION
+		// Fine grain test for clearing based on CMask data
+	#if PERMUTATION_CMASK
+		// Fine grain check if clear is needed
+		bool bClearHeader = false;
+		BRANCH
+		if (CMask > 0u && CMask < 0xF)
+		{
+			bClearHeader = TopLayerTexture.Load(uint3(PixelCoord,0)) == 0;
+		}
+
+		// Header clear
+		BRANCH
+		if (bClearHeader)
+		{
+			MaterialHeaderTextureUAV[PixelCoord] = 0u;
+		}
+	#endif
+	}
+	
 	BRANCH
 	if (!bContainsScreenSpaceSubsurfaceScattering)
 	{
 		// We must fill all the pixel which does not have subsurface scattering by default so that the SSS code is not executed where it should not.
 		SSSTextureUAV[uint3(PixelCoord, 0)] = 0; // This is a good clear for FStrataSubsurfaceHeader, and we only need to clear the header.
 	}
-#endif
 
 #if PERMUTATION_WAVE_OPS
 	const bool bTileContainsStrata							= WaveActiveAnyTrue(bContainsStrataMaterial);
@@ -184,6 +228,10 @@ void TileMainCS(uint2 DispatchThreadId : SV_DispatchThreadID, uint LinearIndex :
 		}
 	#endif
 	}
+
+#if PERMUTATION_CMASK
+	}
+#endif
 }
 #endif // SHADER_TILE_CATEGORIZATION
 
@@ -260,38 +308,6 @@ void ArgsMainCS(uint2 DispatchThreadId : SV_DispatchThreadID)
 
 ////////////////////////////////////////////////////////////////////////////////////////////////////////////
 
-#if SHADER_CLEAR_MATERIAL_BUFFER
-
-RWTexture2DArray<uint> MaterialTextureArrayUAV;
-RWTexture2DArray<uint> SSSTextureUAV;
-uint MaxBytesPerPixel;
-int2 TiledViewBufferResolution;
-
-[numthreads(8, 8, 1)]
-void ClearMaterialBufferMainCS(uint2 DispatchThreadId : SV_DispatchThreadID)
-{
-	if (any(int2(DispatchThreadId.xy) >= TiledViewBufferResolution))
-	{
-		return;
-	}
-
-	// Custom clear of the Strata material buffer.
-	// The first layer of tiled uints contains the header that we need to clear so we only write a single uint per pixel instead of clearing the entire buffer which would be too slow.
-
-	FStrataPixelHeader StrataHeader = InitialiseStrataPixelHeader();
-	uint BSDFCount = 0;
-	uint PackedStrataHeader = PackStrataHeader(BSDFCount, StrataHeader);
-
-	const uint FirstSlice = 0;
-	MaterialTextureArrayUAV[uint3(DispatchThreadId.xy, FirstSlice)] = PackedStrataHeader;
-
-	SSSTextureUAV[uint3(DispatchThreadId.xy, 0)] = 0; // This is a good clear for FStrataSubsurfaceHeader, and we only need to clear the header.
-}
-
-#endif // SHADER_CLEAR_MATERIAL_BUFFER
-
-////////////////////////////////////////////////////////////////////////////////////////////////////////////
-
 #if SHADER_BSDF_TILE
 
 int2 ViewResolution;
diff --git a/Engine/Source/Runtime/Renderer/Private/Strata/Strata.cpp b/Engine/Source/Runtime/Renderer/Private/Strata/Strata.cpp
--- a/Engine/Source/Runtime/Renderer/Private/Strata/Strata.cpp
+++ b/Engine/Source/Runtime/Renderer/Private/Strata/Strata.cpp
@@ -11,6 +11,7 @@
 #include "SceneTextureParameters.h"
 #include "ShaderCompiler.h"
 #include "Lumen/Lumen.h"
+#include "RendererUtils.h"
 
 //PRAGMA_DISABLE_OPTIMIZATION
 
@@ -48,8 +49,8 @@ static TAutoConsoleVariable<int32> CVarMaterialRoughDiffuse(
 	ECVF_ReadOnly | ECVF_RenderThreadSafe);
 
 // STRATA_TODO we keep this for now and can remove it once battletested.
-static TAutoConsoleVariable<int32> CVarClearDuringCategorization(
-	TEXT("r.Strata.ClearDuringCategorization"),
+static TAutoConsoleVariable<int32> CVarUseCmaskClear(
+	TEXT("r.Strata.UseCmaskClear"),
 	1,
 	TEXT("TEST."),
 	ECVF_RenderThreadSafe);
@@ -98,22 +99,14 @@ const TCHAR* ToString(EStrataTileType Type)
 	return TEXT("Unknown");
 }
 
-FORCEINLINE bool ClearDuringCategorization()
+FORCEINLINE bool SupportsCMask(const FStaticShaderPlatform InPlatform)
 {
-	return CVarClearDuringCategorization.GetValueOnRenderThread() > 0;
+	return CVarUseCmaskClear.GetValueOnRenderThread() > 0 && FDataDrivenShaderPlatformInfo::GetSupportsRenderTargetWriteMask(InPlatform);
 }
 
 namespace Strata
 {
 
-// Forward declaration
-static void AddStrataClearMaterialBufferPass(
-	FRDGBuilder& GraphBuilder, 
-	FRDGTextureUAVRef MaterialTextureArrayUAV,
-	FRDGTextureUAVRef SSSTextureUAV,
-	uint32 MaxBytesPerPixel, 
-	FIntPoint TiledViewBufferResolution);
-
 bool IsStrataEnabled()
 {
 	return CVarStrata.GetValueOnAnyThread() > 0;
@@ -360,16 +353,6 @@ void InitialiseStrataFrameSceneData(FRDGBuilder& GraphBuilder, FSceneRenderer& S
 	Out.PeelLayersAboveDepth = FMath::Max(CVarStrataDebugPeelLayersAboveDepth.GetValueOnRenderThread(), 0);
 	Out.SliceStoringDebugStrataTree = SliceCount - 1 - STRATA_BASE_PASS_MRT_OUTPUT_COUNT;	// The UAV skips the first slices set as render target
 
-	if (IsStrataEnabled())
-	{
-		AddStrataClearMaterialBufferPass(
-			GraphBuilder, 
-			GraphBuilder.CreateUAV(FRDGTextureUAVDesc(Out.MaterialTextureArray, 0)),
-			Out.SSSTextureUAV,
-			Out.MaxBytesPerPixel, 
-			MaterialBufferSizeXY);
-	}
-
 	// Initialized view data
 	for (int32 ViewIndex = 0; ViewIndex < SceneRenderer.Views.Num(); ViewIndex++)
 	{
@@ -500,35 +483,6 @@ TRDGUniformBufferRef<FStrataGlobalUniformParameters> BindStrataGlobalUniformPara
 
 ////////////////////////////////////////////////////////////////////////////////////////////////////////////////////////////////////////////////
 
-class FStrataClearMaterialBufferCS : public FGlobalShader
-{
-	DECLARE_GLOBAL_SHADER(FStrataClearMaterialBufferCS);
-	SHADER_USE_PARAMETER_STRUCT(FStrataClearMaterialBufferCS, FGlobalShader);
-
-	using FPermutationDomain = TShaderPermutationDomain<>;
-
-	BEGIN_SHADER_PARAMETER_STRUCT(FParameters, )
-		SHADER_PARAMETER_RDG_TEXTURE_UAV(RWTexture2DArray<uint>, MaterialTextureArrayUAV)
-		SHADER_PARAMETER_RDG_TEXTURE_UAV(RWTexture2DArray<uint>, SSSTextureUAV)
-		SHADER_PARAMETER(uint32, MaxBytesPerPixel)
-		SHADER_PARAMETER(FIntPoint, TiledViewBufferResolution)
-	END_SHADER_PARAMETER_STRUCT()
-
-	static bool ShouldCompilePermutation(const FGlobalShaderPermutationParameters& Parameters)
-	{
-		return GetMaxSupportedFeatureLevel(Parameters.Platform) >= ERHIFeatureLevel::SM5 && Strata::IsStrataEnabled();
-	}
-
-	static void ModifyCompilationEnvironment(const FGlobalShaderPermutationParameters& Parameters, FShaderCompilerEnvironment& OutEnvironment)
-	{
-		FGlobalShader::ModifyCompilationEnvironment(Parameters, OutEnvironment);
-		OutEnvironment.SetDefine(TEXT("SHADER_CLEAR_MATERIAL_BUFFER"), 1);
-	}
-};
-IMPLEMENT_GLOBAL_SHADER(FStrataClearMaterialBufferCS, "/Engine/Private/Strata/StrataMaterialClassification.usf", "ClearMaterialBufferMainCS", SF_Compute);
-
-////////////////////////////////////////////////////////////////////////////////////////////////////////////////////////////////////////////////
-
 class FStrataBSDFTilePassCS : public FGlobalShader
 {
 	DECLARE_GLOBAL_SHADER(FStrataBSDFTilePassCS);
@@ -590,9 +544,9 @@ class FStrataMaterialTileClassificationPassCS : public FGlobalShader
 	DECLARE_GLOBAL_SHADER(FStrataMaterialTileClassificationPassCS);
 	SHADER_USE_PARAMETER_STRUCT(FStrataMaterialTileClassificationPassCS, FGlobalShader);
 
-	class FStrataClearDuringCategorization : SHADER_PERMUTATION_BOOL("PERMUTATION_STRATA_CLEAR_DURING_CATEGORIZATION");
+	class FCmask : SHADER_PERMUTATION_BOOL("PERMUTATION_CMASK");
 	class FWaveOps : SHADER_PERMUTATION_BOOL("PERMUTATION_WAVE_OPS");
-	using FPermutationDomain = TShaderPermutationDomain<FWaveOps, FStrataClearDuringCategorization>;
+	using FPermutationDomain = TShaderPermutationDomain<FWaveOps, FCmask>;
 
 	BEGIN_SHADER_PARAMETER_STRUCT(FParameters, )
 		SHADER_PARAMETER_STRUCT_REF(FViewUniformShaderParameters, ViewUniformBuffer)
@@ -600,6 +554,7 @@ class FStrataMaterialTileClassificationPassCS : public FGlobalShader
 		SHADER_PARAMETER(FIntPoint, ViewResolution)
 		SHADER_PARAMETER(uint32, MaxBytesPerPixel)
 		SHADER_PARAMETER_RDG_TEXTURE(Texture2D, TopLayerTexture)
+		SHADER_PARAMETER_RDG_TEXTURE(Texture2D, TopLayerCmaskTexture)
 		SHADER_PARAMETER_RDG_TEXTURE_SRV(Texture2DArray<uint>, MaterialTextureArray)
 		SHADER_PARAMETER_RDG_BUFFER_UAV(RWBuffer, TileDrawIndirectDataBuffer)
 		SHADER_PARAMETER_RDG_BUFFER_UAV(RWBuffer, SimpleTileListDataBuffer)
@@ -607,6 +562,7 @@ class FStrataMaterialTileClassificationPassCS : public FGlobalShader
 		SHADER_PARAMETER_RDG_BUFFER_UAV(RWBuffer, ComplexTileListDataBuffer)
 		SHADER_PARAMETER_RDG_BUFFER_UAV(RWBuffer, OpaqueRoughRefractionTileListDataBuffer)
 		SHADER_PARAMETER_RDG_BUFFER_UAV(RWBuffer, SSSWithoutOpaqueRoughRefractionTileListDataBuffer)
+		SHADER_PARAMETER_RDG_TEXTURE_UAV(RWTexture2D<uint>, MaterialHeaderTextureUAV)
 		SHADER_PARAMETER_RDG_TEXTURE_UAV(RWTexture2DArray<uint>, SSSTextureUAV)
 		SHADER_PARAMETER_RDG_TEXTURE(Texture2D<float3>, OpaqueRoughRefractionTexture)
 	END_SHADER_PARAMETER_STRUCT()
@@ -961,9 +917,10 @@ void AppendStrataMRTs(const FSceneRenderer& SceneRenderer, uint32& RenderTargetC
 			RenderTargets[RenderTargetCount] = FTextureRenderTargetBinding(SceneRenderer.Scene->StrataSceneData.MaterialTextureArray, StrataMaterialArraySlice, bNeverClear);
 			RenderTargetCount++;
 		};
+		const bool bSupportCMask = SupportsCMask(GMaxRHIShaderPlatform);
 		for (int i = 0; i < STRATA_BASE_PASS_MRT_OUTPUT_COUNT; ++i)
 		{
-			const bool bNeverClear = i != 0; // Only allow clearing the first slice containing the header
+			const bool bNeverClear = bSupportCMask || i != 0; // Only allow clearing the first slice containing the header
 			AddStrataOutputTarget(i, bNeverClear);
 		}
 
@@ -993,6 +950,8 @@ void SetBasePassRenderTargetOutputFormat(const EShaderPlatform Platform, FShader
 
 ////////////////////////////////////////////////////////////////////////////////////////////////////////////////////////////////////////////////
 
+//#include "RendererUtils.h"
+
 void AddStrataMaterialClassificationPass(FRDGBuilder& GraphBuilder, const FMinimalSceneTextures& SceneTextures, const TArray<FViewInfo>& Views)
 {
 	RDG_EVENT_SCOPE_CONDITIONAL(GraphBuilder, IsStrataEnabled() && Views.Num() > 0, "Strata::MaterialClassification");
@@ -1016,9 +975,19 @@ void AddStrataMaterialClassificationPass(FRDGBuilder& GraphBuilder, const FMinim
 
 		// Tile reduction
 		{
-			const bool bClear = ClearDuringCategorization();
+			// When the platform support explicit CMask texture, we disable material data bufferclear. Material buffer buffer clear (the header part) is done during the classification pass.  
+			// To reduce the reading bandwidth, we rely on TopLayerData CMask to 'drive' the clearing process. This allows to clear quickly empty tiles.
+			const bool bSupportCMask = SupportsCMask(View.GetShaderPlatform());
+			FRDGTextureRef TopLayerCmaskTexture = StrataSceneData->TopLayerTexture;			
+			if (bSupportCMask)
+			{
+				// Combine DBuffer RTWriteMasks; will end up in one texture we can load from in the base pass PS and decide whether to do the actual work or not.
+				FRDGTextureRef SourceCMaskTextures[] = { StrataSceneData->TopLayerTexture };
+				FRenderTargetWriteMask::Decode(GraphBuilder, View.ShaderMap, MakeArrayView(SourceCMaskTextures), TopLayerCmaskTexture, GFastVRamConfig.DBufferMask, TEXT("Strata::TopLayerCmask"));
+			}
+
 			FStrataMaterialTileClassificationPassCS::FPermutationDomain PermutationVector;
-			PermutationVector.Set< FStrataMaterialTileClassificationPassCS::FStrataClearDuringCategorization >(bClear);
+			PermutationVector.Set< FStrataMaterialTileClassificationPassCS::FCmask >(bSupportCMask);
 			PermutationVector.Set< FStrataMaterialTileClassificationPassCS::FWaveOps >(bWaveOps);
 			TShaderMapRef<FStrataMaterialTileClassificationPassCS> ComputeShader(View.ShaderMap, PermutationVector);
 			FStrataMaterialTileClassificationPassCS::FParameters* PassParameters = GraphBuilder.AllocParameters<FStrataMaterialTileClassificationPassCS::FParameters>();
@@ -1027,7 +996,9 @@ void AddStrataMaterialClassificationPass(FRDGBuilder& GraphBuilder, const FMinim
 			PassParameters->ViewResolution = View.ViewRect.Size();
 			PassParameters->MaxBytesPerPixel = StrataSceneData->MaxBytesPerPixel;
 			PassParameters->TopLayerTexture = StrataSceneData->TopLayerTexture;
+			PassParameters->TopLayerCmaskTexture = TopLayerCmaskTexture;
 			PassParameters->MaterialTextureArray = StrataSceneData->MaterialTextureArraySRV;
+			PassParameters->MaterialHeaderTextureUAV = GraphBuilder.CreateUAV(FRDGTextureUAVDesc(StrataSceneData->MaterialTextureArray, 0, PF_R32_UINT, 0, 1));
 			PassParameters->SSSTextureUAV = StrataSceneData->SSSTextureUAV;
 			PassParameters->OpaqueRoughRefractionTexture = StrataSceneData->OpaqueRoughRefractionTexture;
 			PassParameters->TileDrawIndirectDataBuffer = StrataViewData->ClassificationTileDrawIndirectBufferUAV;
@@ -1040,7 +1011,7 @@ void AddStrataMaterialClassificationPass(FRDGBuilder& GraphBuilder, const FMinim
 			const uint32 GroupSize = 8;
 			FComputeShaderUtils::AddPass(
 				GraphBuilder,
-				RDG_EVENT_NAME("Strata::MaterialTileClassification(%s%s)", bWaveOps ? TEXT("Wave") : TEXT("SharedMemory"), bClear ? TEXT(", Clear") : TEXT("")),
+				RDG_EVENT_NAME("Strata::MaterialTileClassification(%s%s)", bWaveOps ? TEXT("Wave") : TEXT("SharedMemory"), bSupportCMask ? TEXT(", CMask") : TEXT("")),
 				PassFlags,
 				ComputeShader,
 				PassParameters,
@@ -1124,32 +1095,4 @@ void AddStrataMaterialClassificationPass(FRDGBuilder& GraphBuilder, const FMinim
 	}
 }
 
-static void AddStrataClearMaterialBufferPass(
-	FRDGBuilder& GraphBuilder,
-	FRDGTextureUAVRef MaterialTextureArrayUAV,
-	FRDGTextureUAVRef SSSTextureUAV,
-	uint32 MaxBytesPerPixel,
-	FIntPoint TiledViewBufferResolution)
-{
-	if (ClearDuringCategorization())
-	{
-		return;
-	}
-
-	TShaderMapRef<FStrataClearMaterialBufferCS> ComputeShader(GetGlobalShaderMap(GMaxRHIFeatureLevel));
-	FStrataClearMaterialBufferCS::FParameters* PassParameters = GraphBuilder.AllocParameters<FStrataClearMaterialBufferCS::FParameters>();
-	PassParameters->MaterialTextureArrayUAV = MaterialTextureArrayUAV;
-	PassParameters->SSSTextureUAV = SSSTextureUAV;
-	PassParameters->MaxBytesPerPixel = MaxBytesPerPixel;
-	PassParameters->TiledViewBufferResolution = TiledViewBufferResolution;
-
-	const uint32 GroupSize = 8;
-	FComputeShaderUtils::AddPass(
-		GraphBuilder,
-		RDG_EVENT_NAME("Strata::ClearMaterialBuffer"),
-		ComputeShader,
-		PassParameters,
-		FComputeShaderUtils::GetGroupCount(TiledViewBufferResolution, GroupSize));
-}
-
 } // namespace Strata
