Fix Strata on platform which relies on explicit output format declaration in shader.

diff --git a/Engine/Source/Runtime/Renderer/Private/BasePassRendering.h b/Engine/Source/Runtime/Renderer/Private/BasePassRendering.h
--- a/Engine/Source/Runtime/Renderer/Private/BasePassRendering.h
+++ b/Engine/Source/Runtime/Renderer/Private/BasePassRendering.h
@@ -319,8 +319,8 @@ public:
 			const int32 VelocityIndex = IsForwardShadingEnabled(Parameters.Platform) ? 1 : 4; // As defined in BasePassPixelShader.usf
 			OutEnvironment.SetRenderTargetOutputFormat(VelocityIndex, PF_G16R16);
 		}
-
-		Strata::SetBasePassRenderTargetOutputFormat(Parameters.Platform, OutEnvironment);
+		
+		Strata::SetBasePassRenderTargetOutputFormat(Parameters.Platform, Parameters.MaterialParameters, OutEnvironment);
 
 		FForwardLightingParameters::ModifyCompilationEnvironment(Parameters.Platform, OutEnvironment);
 	}
diff --git a/Engine/Source/Runtime/Renderer/Private/Strata/Strata.cpp b/Engine/Source/Runtime/Renderer/Private/Strata/Strata.cpp
--- a/Engine/Source/Runtime/Renderer/Private/Strata/Strata.cpp
+++ b/Engine/Source/Runtime/Renderer/Private/Strata/Strata.cpp
@@ -932,11 +932,19 @@ void AppendStrataMRTs(const FSceneRenderer& SceneRenderer, uint32& RenderTargetC
 	}
 }
 
-void SetBasePassRenderTargetOutputFormat(const EShaderPlatform Platform, FShaderCompilerEnvironment& OutEnvironment)
+void SetBasePassRenderTargetOutputFormat(const EShaderPlatform Platform, const FMaterialShaderParameters& MaterialParameters, FShaderCompilerEnvironment& OutEnvironment)
 {
 	if (Strata::IsStrataEnabled())
 	{
-		const FGBufferParams GBufferParams = FShaderCompileUtilities::FetchGBufferParamsRuntime(Platform);
+		FGBufferParams GBufferParams = FShaderCompileUtilities::FetchGBufferParamsRuntime(Platform);
+
+		// If it is not a water material, we force bHasSingleLayerWaterSeparatedMainLight to false, in order to 
+		// ensure non-used MRTs are not inserted in BufferInfo. Otherwise this would offset Strata MRTs, causing 
+		// MRTs' format to be incorrect
+		if (!MaterialParameters.bIsUsedWithWater)
+		{
+			GBufferParams.bHasSingleLayerWaterSeparatedMainLight = false;
+		}
 		const FGBufferInfo BufferInfo = FetchFullGBufferInfo(GBufferParams);
 
 		// Add 2 uint for Strata fast path
diff --git a/Engine/Source/Runtime/Renderer/Private/Strata/Strata.h b/Engine/Source/Runtime/Renderer/Private/Strata/Strata.h
--- a/Engine/Source/Runtime/Renderer/Private/Strata/Strata.h
+++ b/Engine/Source/Runtime/Renderer/Private/Strata/Strata.h
@@ -152,7 +152,7 @@ void BindStrataForwardPasslUniformParameters(FRDGBuilder& GraphBuilder, const FV
 TRDGUniformBufferRef<FStrataGlobalUniformParameters> BindStrataGlobalUniformParameters(const FViewInfo& View);
 
 void AppendStrataMRTs(const FSceneRenderer& SceneRenderer, uint32& BasePassTextureCount, TStaticArray<FTextureRenderTargetBinding, MaxSimultaneousRenderTargets>& BasePassTextures);
-void SetBasePassRenderTargetOutputFormat(const EShaderPlatform Platform, FShaderCompilerEnvironment& OutEnvironment);
+void SetBasePassRenderTargetOutputFormat(const EShaderPlatform Platform, const FMaterialShaderParameters& MaterialParameters, FShaderCompilerEnvironment& OutEnvironment);
 
 
 void AddStrataMaterialClassificationPass(FRDGBuilder& GraphBuilder, const FMinimalSceneTextures& SceneTextures, const TArray<FViewInfo>& Views);
