Remove unnecessary warning. This warning can fire for empty foliage actors.

diff --git a/Engine/Source/Runtime/Renderer/Private/PrimitiveSceneInfo.cpp b/Engine/Source/Runtime/Renderer/Private/PrimitiveSceneInfo.cpp
--- a/Engine/Source/Runtime/Renderer/Private/PrimitiveSceneInfo.cpp
+++ b/Engine/Source/Runtime/Renderer/Private/PrimitiveSceneInfo.cpp
@@ -1621,11 +1621,7 @@ void FPrimitiveSceneInfo::UpdateRuntimeVirtualTextureFlags()
 		{
 			UE_LOG(LogRenderer, Warning, TEXT("Rendering a nanite mesh to a runtime virtual texture isn't yet supported. Please disable this option on primitive component : %s"), *Proxy->GetOwnerName().ToString());
 		}
-		else if (StaticMeshes.Num() == 0)
-		{
-			UE_LOG(LogRenderer, Warning, TEXT("Rendering a primitive in a runtime virtual texture implies that there is a mesh to render. Please disable this option on primitive component : %s"), *Proxy->GetOwnerName().ToString());
-		}
-		else
+		else if (StaticMeshes.Num() > 0)
 		{
 			RuntimeVirtualTextureFlags.bRenderToVirtualTexture = true;
 
