Fix for ensures firing on LWC decals.

diff --git a/Engine/Source/Runtime/Renderer/Private/DecalRenderingShared.cpp b/Engine/Source/Runtime/Renderer/Private/DecalRenderingShared.cpp
--- a/Engine/Source/Runtime/Renderer/Private/DecalRenderingShared.cpp
+++ b/Engine/Source/Runtime/Renderer/Private/DecalRenderingShared.cpp
@@ -87,10 +87,10 @@ public:
 
 		FMaterialShader::SetParameters(RHICmdList, ShaderRHI, MaterialProxy, *MaterialResource, View);
 
-		const FLargeWorldRenderPosition AbsoluteOrigin(View.ViewMatrices.GetInvViewMatrix().GetOrigin());
-		const FVector3f TilePosition = AbsoluteOrigin.GetTile();
-		const FMatrix WorldToDecalMatrix = DecalProxy.ComponentTrans.ToInverseMatrixWithScale();
 		const FMatrix DecalToWorldMatrix = DecalProxy.ComponentTrans.ToMatrixWithScale();
+		const FMatrix WorldToDecalMatrix = DecalProxy.ComponentTrans.ToInverseMatrixWithScale();
+		const FLargeWorldRenderPosition AbsoluteOrigin(DecalToWorldMatrix.GetOrigin());
+		const FVector3f TilePosition = AbsoluteOrigin.GetTile();
 		const FMatrix44f RelativeDecalToWorldMatrix = FLargeWorldRenderScalar::MakeToRelativeWorldMatrix(AbsoluteOrigin.GetTileOffset(), DecalToWorldMatrix);
 		const FVector3f OrientationVector = (FVector3f)DecalProxy.ComponentTrans.GetUnitAxis(EAxis::X);
 
