Removed a couple of overzealous asserts.

The code handles the case where the scene is null just fine, no need to assert that it isn't (and this situation sometimes arises in practice when Slate requests hit proxy rendering between level changes).

#rnx

diff --git a/Engine/Source/Runtime/Renderer/Private/SceneRendering.cpp b/Engine/Source/Runtime/Renderer/Private/SceneRendering.cpp
--- a/Engine/Source/Runtime/Renderer/Private/SceneRendering.cpp
+++ b/Engine/Source/Runtime/Renderer/Private/SceneRendering.cpp
@@ -4531,9 +4531,7 @@ void FRendererModule::BeginRenderingViewFamilies(FCanvas* Canvas, TArrayView<FSc
 	for (FSceneViewFamily* ViewFamily : ViewFamilies)
 	{
 		check(ViewFamily);
-		check(ViewFamily->Scene);
 		check(ViewFamily->Scene == ViewFamilies[0]->Scene);
-		check(ViewFamily->GetScreenPercentageInterface());
 	}
 
 	UWorld* World = nullptr;
