[Backout] - CL21397411
Original CL Desc
-----------------------------------------------------------------
Wrap TAA downsample optimization around cvar, and disable on XboxOneGDK (small gain: 0.04ms)

Did profiling on a Xbox One X with emulation mode (XboxOneS) setting. Not representative of the XboxOneS perf..., doing the first downsample in the TAA shader is faster.

diff --git a/Engine/Source/Runtime/Renderer/Private/PostProcess/PostProcessing.cpp b/Engine/Source/Runtime/Renderer/Private/PostProcess/PostProcessing.cpp
--- a/Engine/Source/Runtime/Renderer/Private/PostProcess/PostProcessing.cpp
+++ b/Engine/Source/Runtime/Renderer/Private/PostProcess/PostProcessing.cpp
@@ -131,13 +131,6 @@ TAutoConsoleVariable<int32> CVarPostProcessingForceAsyncDispatch(
 	ECVF_RenderThreadSafe);
 #endif
 
-TAutoConsoleVariable<bool> CVarAllowDownsampleSceneColorInTAA(
-	TEXT("r.PostProcessing.AllowDownsampleSceneColorInTAA"),
-	true,
-	TEXT("Allow TAA to output scene color in half resolution.\n")
-	TEXT(" (default: true)"),
-	ECVF_RenderThreadSafe);
-
 } //! namespace
 
 EDownsampleQuality GetDownsampleQuality(const TAutoConsoleVariable<int32>& CVar)
@@ -635,8 +628,7 @@ void AddPostProcessingPasses(
 			{
 				UpscalerPassInputs.bAllowDownsampleSceneColor =
 					(bNeedPostMotionBlurHalfRes || bNeedPostMotionBlurQuarterRes) &&
-					DownsampleQuality == EDownsampleQuality::Low &&
-					CVarAllowDownsampleSceneColorInTAA.GetValueOnRenderThread() == true;
+					DownsampleQuality == EDownsampleQuality::Low;;
 			}
 			UpscalerPassInputs.DownsampleOverrideFormat = DownsampleOverrideFormat;
 			UpscalerPassInputs.SceneColorTexture = SceneColor.Texture;
