Fix crash when swapping groom geometry type from strands to cards.

diff --git a/Engine/Source/Runtime/Renderer/Private/HairStrands/HairStrandsInterface.cpp b/Engine/Source/Runtime/Renderer/Private/HairStrands/HairStrandsInterface.cpp
--- a/Engine/Source/Runtime/Renderer/Private/HairStrands/HairStrandsInterface.cpp
+++ b/Engine/Source/Runtime/Renderer/Private/HairStrands/HairStrandsInterface.cpp
@@ -435,7 +435,7 @@ bool IsHairStrandsVisibleInShadows(const FViewInfo& View, const FHairStrandsInst
 	bool bIsVisibleInShadow = false;
 	if (const FHairGroupPublicData* HairData = Instance.GetHairData())
 	{
-		const int32 LODIndex = FMath::CeilToInt(HairData->LODIndex);
+		const int32 LODIndex = HairData->LODIndex;
 		const bool bIsStrands = LODIndex >= 0 && HairData->IsVisible(LODIndex) && HairData->GetGeometryType(LODIndex) == EHairGeometryType::Strands;
 		if (!bIsStrands)
 		{
