Fixes DOF passes not being setup when only setting the depth blur

diff --git a/Engine/Source/Runtime/Renderer/Private/PostProcess/DiaphragmDOF.cpp b/Engine/Source/Runtime/Renderer/Private/PostProcess/DiaphragmDOF.cpp
--- a/Engine/Source/Runtime/Renderer/Private/PostProcess/DiaphragmDOF.cpp
+++ b/Engine/Source/Runtime/Renderer/Private/PostProcess/DiaphragmDOF.cpp
@@ -1320,8 +1320,7 @@ bool DiaphragmDOF::IsEnabled(const FViewInfo& View)
 		View.Family->EngineShowFlags.DepthOfField &&
 		bDepthOfFieldRequestedByCVar &&
 		!(View.Family->EngineShowFlags.PathTracing && View.FinalPostProcessSettings.PathTracingEnableReferenceDOF) &&
-		View.FinalPostProcessSettings.DepthOfFieldFstop > 0 &&
-		View.FinalPostProcessSettings.DepthOfFieldFocalDistance > 0;
+		((View.FinalPostProcessSettings.DepthOfFieldFstop > 0.f && View.FinalPostProcessSettings.DepthOfFieldFocalDistance > 0.f) || View.FinalPostProcessSettings.DepthOfFieldDepthBlurRadius > 0.f);
 }
 
 FRDGTextureRef DiaphragmDOF::AddPasses(
