Build ray tracing decal grid for each view.
- Decal fade is view dependent.
- Also fixes crash when path tracing is used in splitscreen, since RayTracingDecalUniformBuffer was only assigned to the first view.

diff --git a/Engine/Source/Runtime/Renderer/Private/DeferredShadingRenderer.cpp b/Engine/Source/Runtime/Renderer/Private/DeferredShadingRenderer.cpp
--- a/Engine/Source/Runtime/Renderer/Private/DeferredShadingRenderer.cpp
+++ b/Engine/Source/Runtime/Renderer/Private/DeferredShadingRenderer.cpp
@@ -2228,21 +2228,34 @@ void FDeferredShadingSceneRenderer::Render(FRDGBuilder& GraphBuilder)
 	FRayTracingScene& RayTracingScene = Scene->RayTracingScene;
 	RayTracingScene.Reset(); // Resets the internal arrays, but does not release any resources.
 
-	const int32 ReferenceViewIndex = 0;
-	FViewInfo& ReferenceView = Views[ReferenceViewIndex];
-
-	if (ShouldPrepareRayTracingDecals(*Scene, *ReferenceView.Family))
+	if (ShouldPrepareRayTracingDecals(*Scene, ViewFamily))
 	{
-		ReferenceView.RayTracingDecalUniformBuffer = CreateRayTracingDecalData(GraphBuilder, *Scene, ReferenceView, RayTracingScene.NumCallableShaderSlots);
-		ReferenceView.bHasRayTracingDecals = true;
-		RayTracingScene.NumCallableShaderSlots += Scene->Decals.Num();
+		// Calculate decal grid for ray tracing per view since decal fade is view dependent
+		// TODO: investigate reusing the same grid for all views (ie: different callable shader SBT entries for each view so fade alpha is still correct for each view)
+
+		for (int32 ViewIndex = 0; ViewIndex < Views.Num(); ViewIndex++)
+		{
+			FViewInfo& View = Views[ViewIndex];
+			View.RayTracingDecalUniformBuffer = CreateRayTracingDecalData(GraphBuilder, *Scene, View, RayTracingScene.NumCallableShaderSlots);
+			View.bHasRayTracingDecals = true;
+			RayTracingScene.NumCallableShaderSlots += Scene->Decals.Num();
+		}
 	}
 	else
 	{
-		ReferenceView.RayTracingDecalUniformBuffer = CreateNullRayTracingDecalsUniformBuffer(GraphBuilder);
-		ReferenceView.bHasRayTracingDecals = false;
+		TRDGUniformBufferRef<FRayTracingDecals> NullRayTracingDecalUniformBuffer = CreateNullRayTracingDecalsUniformBuffer(GraphBuilder);
+
+		for (int32 ViewIndex = 0; ViewIndex < Views.Num(); ViewIndex++)
+		{
+			FViewInfo& View = Views[ViewIndex];
+			View.RayTracingDecalUniformBuffer = NullRayTracingDecalUniformBuffer;
+			View.bHasRayTracingDecals = false;
+		}
 	}
 
+	const int32 ReferenceViewIndex = 0;
+	FViewInfo& ReferenceView = Views[ReferenceViewIndex];
+
 	if (IsRayTracingEnabled() && RHISupportsRayTracingShaders(ViewFamily.GetShaderPlatform()))
 	{
 		if (!ViewFamily.EngineShowFlags.PathTracing)
