Don't call view extensions when scene renderer is not being rendered. Fixes inconsistency with calls to PreRenderView_RenderThread and PreInitView_RenderThread / PostRenderView_RenderThread.

diff --git a/Engine/Source/Runtime/Renderer/Private/PlanarReflectionRendering.cpp b/Engine/Source/Runtime/Renderer/Private/PlanarReflectionRendering.cpp
--- a/Engine/Source/Runtime/Renderer/Private/PlanarReflectionRendering.cpp
+++ b/Engine/Source/Runtime/Renderer/Private/PlanarReflectionRendering.cpp
@@ -276,46 +276,30 @@ static void UpdatePlanarReflectionContents_RenderThread(
 {
 	QUICK_SCOPE_CYCLE_COUNTER(STAT_RenderPlanarReflection);
 
-	SceneRenderer->RenderThreadBegin(RHICmdList);
-
-	// update any resources that needed a deferred update
-	FDeferredUpdateResource::UpdateResources(RHICmdList);
-
-	const ERHIFeatureLevel::Type FeatureLevel = SceneRenderer->FeatureLevel;
-	FRDGBuilder GraphBuilder(RHICmdList, RDG_EVENT_NAME("PlanarReflection"), FSceneRenderer::GetRDGParalelExecuteFlags(FeatureLevel));
-
-	// We need to execute the pre-render view extensions before we do any view dependent work.
-	FSceneRenderer::ViewExtensionPreRender_RenderThread(GraphBuilder, SceneRenderer);
-
-	// Make sure we render to the same set of GPUs as the main scene renderer.
-	if (MainSceneRenderer->ViewFamily.RenderTarget != nullptr)
-	{
-		RenderTarget->SetActiveGPUMask(MainSceneRenderer->ViewFamily.RenderTarget->GetGPUMask(RHICmdList));
-	}
-	else
 	{
-		RenderTarget->SetActiveGPUMask(FRHIGPUMask::GPU0());
-	}
-
-	FBox PlanarReflectionBounds = SceneProxy->WorldBounds;
+		FBox PlanarReflectionBounds = SceneProxy->WorldBounds;
+		bool bIsInAnyFrustum = false;
 
-	bool bIsInAnyFrustum = false;
-	for (int32 ViewIndex = 0; ViewIndex < MainSceneRenderer->Views.Num(); ++ViewIndex)
-	{
-		FViewInfo& View = MainSceneRenderer->Views[ViewIndex];
-		if (MirrorPlane.PlaneDot(View.ViewMatrices.GetViewOrigin()) > 0)
+		for (int32 ViewIndex = 0; ViewIndex < MainSceneRenderer->Views.Num(); ++ViewIndex)
 		{
-			if (View.ViewFrustum.IntersectBox(PlanarReflectionBounds.GetCenter(), PlanarReflectionBounds.GetExtent()))
+			FViewInfo& View = MainSceneRenderer->Views[ViewIndex];
+			if (MirrorPlane.PlaneDot(View.ViewMatrices.GetViewOrigin()) > 0)
 			{
-				bIsInAnyFrustum = true;
-				break;
+				if (View.ViewFrustum.IntersectBox(PlanarReflectionBounds.GetCenter(), PlanarReflectionBounds.GetExtent()))
+				{
+					bIsInAnyFrustum = true;
+					break;
+				}
 			}
 		}
-	}
 
-	if (bIsInAnyFrustum)
-	{
+		if (!bIsInAnyFrustum)
+		{
+			return;
+		}
+
 		bool bIsVisibleInAnyView = true;
+
 		for (int32 ViewIndex = 0; ViewIndex < MainSceneRenderer->Views.Num(); ++ViewIndex)
 		{
 			FViewInfo& View = MainSceneRenderer->Views[ViewIndex];
@@ -348,64 +332,89 @@ static void UpdatePlanarReflectionContents_RenderThread(
 			}
 		}
 
-		if (bIsVisibleInAnyView)
+		if (!bIsVisibleInAnyView)
 		{
-			{
+			return;
+		}
+	}
+
+	SceneRenderer->RenderThreadBegin(RHICmdList);
+
+	FUniformExpressionCacheAsyncUpdateScope AsyncUpdateScope;
+
+	// update any resources that needed a deferred update
+	FDeferredUpdateResource::UpdateResources(RHICmdList);
+
+	const ERHIFeatureLevel::Type FeatureLevel = SceneRenderer->FeatureLevel;
+	FRDGBuilder GraphBuilder(RHICmdList, RDG_EVENT_NAME("PlanarReflection"), FSceneRenderer::GetRDGParalelExecuteFlags(FeatureLevel));
+
+	// We need to execute the pre-render view extensions before we do any view dependent work.
+	FSceneRenderer::ViewExtensionPreRender_RenderThread(GraphBuilder, SceneRenderer);
+
+	// Make sure we render to the same set of GPUs as the main scene renderer.
+	if (MainSceneRenderer->ViewFamily.RenderTarget != nullptr)
+	{
+		RenderTarget->SetActiveGPUMask(MainSceneRenderer->ViewFamily.RenderTarget->GetGPUMask(RHICmdList));
+	}
+	else
+	{
+		RenderTarget->SetActiveGPUMask(FRHIGPUMask::GPU0());
+	}
+
+	{
 #if WANTS_DRAW_MESH_EVENTS
-				FString EventName;
-				OwnerName.ToString(EventName);
-				RDG_EVENT_SCOPE(GraphBuilder, "PlanarReflection %s", *EventName);
+		FString EventName;
+		OwnerName.ToString(EventName);
+		RDG_EVENT_SCOPE(GraphBuilder, "PlanarReflection %s", *EventName);
 #else
-				RDG_EVENT_SCOPE(GraphBuilder, "UpdatePlanarReflectionContent_RenderThread");
+		RDG_EVENT_SCOPE(GraphBuilder, "UpdatePlanarReflectionContent_RenderThread");
 #endif
-				// Reflection view late update
-				if (SceneRenderer->Views.Num() > 1)
-				{
-					const FMirrorMatrix MirrorMatrix(MirrorPlane);
-					for (int32 ViewIndex = 0; ViewIndex < SceneRenderer->Views.Num(); ++ViewIndex)
-					{
-						FViewInfo& ReflectionViewToUpdate = SceneRenderer->Views[ViewIndex];
-						const FViewInfo& UpdatedParentView = MainSceneRenderer->Views[ViewIndex];
+		// Reflection view late update
+		if (SceneRenderer->Views.Num() > 1)
+		{
+			const FMirrorMatrix MirrorMatrix(MirrorPlane);
+			for (int32 ViewIndex = 0; ViewIndex < SceneRenderer->Views.Num(); ++ViewIndex)
+			{
+				FViewInfo& ReflectionViewToUpdate = SceneRenderer->Views[ViewIndex];
+				const FViewInfo& UpdatedParentView = MainSceneRenderer->Views[ViewIndex];
 
-						ReflectionViewToUpdate.UpdatePlanarReflectionViewMatrix(UpdatedParentView, MirrorMatrix);
-					}
-				}
+				ReflectionViewToUpdate.UpdatePlanarReflectionViewMatrix(UpdatedParentView, MirrorMatrix);
+			}
+		}
 
-				// Render the scene normally
-				{
-					RDG_RHI_EVENT_SCOPE(GraphBuilder, RenderScene);
-					SceneRenderer->Render(GraphBuilder);
-				}
+		// Render the scene normally
+		{
+			RDG_RHI_EVENT_SCOPE(GraphBuilder, RenderScene);
+			SceneRenderer->Render(GraphBuilder);
+		}
 
-				SceneProxy->RenderTarget = RenderTarget;
+		SceneProxy->RenderTarget = RenderTarget;
 
-				// Update the view rects into the planar reflection proxy.
-				for (int32 ViewIndex = 0; ViewIndex < SceneRenderer->Views.Num(); ++ViewIndex)
-				{
-					// Make sure screen percentage has correctly been set on render thread.
-					check(SceneRenderer->Views[ViewIndex].ViewRect.Area() > 0);
-					SceneProxy->ViewRect[ViewIndex] = SceneRenderer->Views[ViewIndex].ViewRect;
-				}
+		// Update the view rects into the planar reflection proxy.
+		for (int32 ViewIndex = 0; ViewIndex < SceneRenderer->Views.Num(); ++ViewIndex)
+		{
+			// Make sure screen percentage has correctly been set on render thread.
+			check(SceneRenderer->Views[ViewIndex].ViewRect.Area() > 0);
+			SceneProxy->ViewRect[ViewIndex] = SceneRenderer->Views[ViewIndex].ViewRect;
+		}
 
-				FRDGTextureRef ReflectionOutputTexture = GraphBuilder.RegisterExternalTexture(CreateRenderTarget(RenderTarget->TextureRHI, TEXT("ReflectionOutputTexture")));
-				GraphBuilder.SetTextureAccessFinal(ReflectionOutputTexture, ERHIAccess::SRVGraphics);
+		FRDGTextureRef ReflectionOutputTexture = GraphBuilder.RegisterExternalTexture(CreateRenderTarget(RenderTarget->TextureRHI, TEXT("ReflectionOutputTexture")));
+		GraphBuilder.SetTextureAccessFinal(ReflectionOutputTexture, ERHIAccess::SRVGraphics);
 
-				FSceneTextureShaderParameters SceneTextureParameters = CreateSceneTextureShaderParameters(GraphBuilder, &SceneRenderer->GetActiveSceneTextures(), SceneRenderer->FeatureLevel, ESceneTextureSetupMode::SceneDepth);
-				const FMinimalSceneTextures& SceneTextures = SceneRenderer->GetActiveSceneTextures();
+		FSceneTextureShaderParameters SceneTextureParameters = CreateSceneTextureShaderParameters(GraphBuilder, &SceneRenderer->GetActiveSceneTextures(), SceneRenderer->FeatureLevel, ESceneTextureSetupMode::SceneDepth);
+		const FMinimalSceneTextures& SceneTextures = SceneRenderer->GetActiveSceneTextures();
 
-				for (int32 ViewIndex = 0; ViewIndex < SceneRenderer->Views.Num(); ++ViewIndex)
-				{
-					FViewInfo& View = SceneRenderer->Views[ViewIndex];
-					RDG_GPU_MASK_SCOPE(GraphBuilder, View.GPUMask);
-					if (MainSceneRenderer->Scene->GetShadingPath() == EShadingPath::Deferred)
-					{
-						PrefilterPlanarReflection<true>(GraphBuilder, View, SceneTextureParameters, SceneProxy, SceneTextures.Color.Resolve, ReflectionOutputTexture);
-					}
-					else
-					{
-						PrefilterPlanarReflection<false>(GraphBuilder, View, SceneTextureParameters, SceneProxy, SceneTextures.Color.Resolve, ReflectionOutputTexture);
-					}
-				}
+		for (int32 ViewIndex = 0; ViewIndex < SceneRenderer->Views.Num(); ++ViewIndex)
+		{
+			FViewInfo& View = SceneRenderer->Views[ViewIndex];
+			RDG_GPU_MASK_SCOPE(GraphBuilder, View.GPUMask);
+			if (MainSceneRenderer->Scene->GetShadingPath() == EShadingPath::Deferred)
+			{
+				PrefilterPlanarReflection<true>(GraphBuilder, View, SceneTextureParameters, SceneProxy, SceneTextures.Color.Resolve, ReflectionOutputTexture);
+			}
+			else
+			{
+				PrefilterPlanarReflection<false>(GraphBuilder, View, SceneTextureParameters, SceneProxy, SceneTextures.Color.Resolve, ReflectionOutputTexture);
 			}
 		}
 	}
diff --git a/Engine/Source/Runtime/Renderer/Private/SceneRendering.cpp b/Engine/Source/Runtime/Renderer/Private/SceneRendering.cpp
--- a/Engine/Source/Runtime/Renderer/Private/SceneRendering.cpp
+++ b/Engine/Source/Runtime/Renderer/Private/SceneRendering.cpp
@@ -4002,7 +4002,7 @@ void FSceneRenderer::UpdatePrimitiveIndirectLightingCacheBuffers()
 */
 void FSceneRenderer::ViewExtensionPreRender_RenderThread(FRDGBuilder& GraphBuilder, FSceneRenderer* SceneRenderer)
 {
-	if (SceneRenderer->ViewFamily.ViewExtensions.IsEmpty())
+	if (SceneRenderer->ViewFamily.ViewExtensions.IsEmpty() || !SceneRenderer->ViewFamily.EngineShowFlags.Rendering)
 	{
 		return;
 	}
