Removed distance field async that was being immediately waited on. Removes a large render thread stall.

diff --git a/Engine/Source/Runtime/Renderer/Private/DistanceFieldStreaming.cpp b/Engine/Source/Runtime/Renderer/Private/DistanceFieldStreaming.cpp
--- a/Engine/Source/Runtime/Renderer/Private/DistanceFieldStreaming.cpp
+++ b/Engine/Source/Runtime/Renderer/Private/DistanceFieldStreaming.cpp
@@ -1457,30 +1457,21 @@ void FDistanceFieldSceneData::UpdateDistanceFieldAtlas(
 			UpdateParameters.BrickUploadCoordinatesPtr = AtlasUpload.BrickUploadCoordinatesPtr;
 		}
 
-		check(AsyncTaskEvents.IsEmpty());
-
 		if (NewReadRequests.Num() || ReadRequestsToUpload.Num() || ReadRequestsToCleanUp.Num())
 		{
 			UpdateParameters.NewReadRequests = MoveTemp(NewReadRequests);
 			UpdateParameters.ReadRequestsToUpload = MoveTemp(ReadRequestsToUpload);
 			UpdateParameters.ReadRequestsToCleanUp = MoveTemp(ReadRequestsToCleanUp);
 
-			// Kick off an async task to copy completed read requests into upload staging buffers, and issue new read requests
-			AsyncTaskEvents.Add(TGraphTask<FDistanceFieldStreamingUpdateTask>::CreateTask().ConstructAndDispatchWhenReady(UpdateParameters));
+			// TODO: We actually run this synchronously now after the RDG conversion, as it would otherwise immediately sync.
+			AsyncUpdate(UpdateParameters);
 		}
 
-		if (AsyncTaskEvents.Num() || NumBrickUploads > 0 || NumIndirectionTableAdds > 0)
+		if (NumBrickUploads > 0 || NumIndirectionTableAdds > 0)
 		{
 			QUICK_SCOPE_CYCLE_COUNTER(STAT_WaitOnDistanceFieldStreamingUpdate);
 			TRACE_CPUPROFILER_EVENT_SCOPE(WaitOnDistanceFieldStreamingUpdate);
 
-			if (!AsyncTaskEvents.IsEmpty())
-			{
-				// Block on the async task before RDG execution of compute scatter uploads
-				FTaskGraphInterface::Get().WaitUntilTasksComplete(AsyncTaskEvents, ENamedThreads::GetRenderThread_Local());
-				AsyncTaskEvents.Empty();
-			}
-
 			if (NumBrickUploads > 0)
 			{
 				AtlasUpload.Unlock();
diff --git a/Engine/Source/Runtime/Renderer/Private/ScenePrivate.h b/Engine/Source/Runtime/Renderer/Private/ScenePrivate.h
--- a/Engine/Source/Runtime/Renderer/Private/ScenePrivate.h
+++ b/Engine/Source/Runtime/Renderer/Private/ScenePrivate.h
@@ -2058,7 +2058,6 @@ public:
 	FReadBuffer BrickUploadDataBuffer;
 
 	TArray<FDistanceFieldReadRequest> ReadRequests;
-	FGraphEventArray AsyncTaskEvents;
 
 	/** Stores the primitive and instance index of every entry in the object buffer. */
 	TArray<FPrimitiveAndInstance> PrimitiveInstanceMapping;
