Skip caching MobileBasePassCSM commands when deferred shading is enabled

diff --git a/Engine/Source/Runtime/Renderer/Private/MeshDrawCommands.cpp b/Engine/Source/Runtime/Renderer/Private/MeshDrawCommands.cpp
--- a/Engine/Source/Runtime/Renderer/Private/MeshDrawCommands.cpp
+++ b/Engine/Source/Runtime/Renderer/Private/MeshDrawCommands.cpp
@@ -667,6 +667,7 @@ void GenerateMobileBasePassDynamicMeshDrawCommands(
 	MobilePassCSMPassMeshProcessor->SetDrawListContext(&DynamicPassMeshDrawListContext);
 
 	const FMobileCSMVisibilityInfo& MobileCSMVisibilityInfo = View.MobileCSMVisibilityInfo;
+	const bool bSkipCSMShaderCulling = MobileBasePassAlwaysUsesCSM(View.GetShaderPlatform());
 	
 	{
 		const int32 NumCommandsBefore = VisibleCommands.Num();
@@ -680,7 +681,8 @@ void GenerateMobileBasePassDynamicMeshDrawCommands(
 				const uint64 BatchElementMask = ~0ull;
 
 				const int32 PrimitiveIndex = MeshAndRelevance.PrimitiveSceneProxy->GetPrimitiveSceneInfo()->GetIndex();
-				if (MobileCSMVisibilityInfo.bMobileDynamicCSMInUse
+				if (!bSkipCSMShaderCulling
+					&& MobileCSMVisibilityInfo.bMobileDynamicCSMInUse
 					&& (MobileCSMVisibilityInfo.bAlwaysUseCSM || MobileCSMVisibilityInfo.MobilePrimitiveCSMReceiverVisibilityMap[PrimitiveIndex]))
 				{
 					MobilePassCSMPassMeshProcessor->AddMeshBatch(*MeshAndRelevance.Mesh, BatchElementMask, MeshAndRelevance.PrimitiveSceneProxy);
@@ -706,7 +708,8 @@ void GenerateMobileBasePassDynamicMeshDrawCommands(
 			const FStaticMeshBatch* StaticMeshBatch = DynamicMeshCommandBuildRequests[MeshIndex];
 
 			const int32 PrimitiveIndex = StaticMeshBatch->PrimitiveSceneInfo->Proxy->GetPrimitiveSceneInfo()->GetIndex();
-			if (MobileCSMVisibilityInfo.bMobileDynamicCSMInUse
+			if (!bSkipCSMShaderCulling
+				&& MobileCSMVisibilityInfo.bMobileDynamicCSMInUse
 				&& (MobileCSMVisibilityInfo.bAlwaysUseCSM || MobileCSMVisibilityInfo.MobilePrimitiveCSMReceiverVisibilityMap[PrimitiveIndex]))
 			{
 				const uint64 DefaultBatchElementMask = ~0ul;
diff --git a/Engine/Source/Runtime/Renderer/Private/MobileBasePass.cpp b/Engine/Source/Runtime/Renderer/Private/MobileBasePass.cpp
--- a/Engine/Source/Runtime/Renderer/Private/MobileBasePass.cpp
+++ b/Engine/Source/Runtime/Renderer/Private/MobileBasePass.cpp
@@ -620,7 +620,9 @@ bool FMobileBasePassMeshProcessor::TryAddMeshBatch(const FMeshBatch& RESTRICT Me
 
 void FMobileBasePassMeshProcessor::AddMeshBatch(const FMeshBatch& RESTRICT MeshBatch, uint64 BatchElementMask, const FPrimitiveSceneProxy* RESTRICT PrimitiveSceneProxy, int32 StaticMeshId)
 {
-	if (!MeshBatch.bUseForMaterial || (PrimitiveSceneProxy && !PrimitiveSceneProxy->ShouldRenderInMainPass()))
+	if (!MeshBatch.bUseForMaterial || 
+		(Flags & FMobileBasePassMeshProcessor::EFlags::DoNotCache) == FMobileBasePassMeshProcessor::EFlags::DoNotCache ||
+		(PrimitiveSceneProxy && !PrimitiveSceneProxy->ShouldRenderInMainPass()))
 	{
 		return;
 	}
@@ -772,8 +774,13 @@ FMeshPassProcessor* CreateMobileBasePassCSMProcessor(const FScene* Scene, const
 	PassDrawRenderState.SetDepthStencilAccess(Scene->DefaultBasePassDepthStencilAccess);
 	PassDrawRenderState.SetDepthStencilState(TStaticDepthStencilState<true, CF_DepthNearOrEqual>::GetRHI());
 
-	const FMobileBasePassMeshProcessor::EFlags Flags = FMobileBasePassMeshProcessor::EFlags::CanReceiveCSM | FMobileBasePassMeshProcessor::EFlags::CanUseDepthStencil;
-
+	// By default this processor will not cache anything. Only enable it when CSM culling is active
+	FMobileBasePassMeshProcessor::EFlags Flags = FMobileBasePassMeshProcessor::EFlags::DoNotCache;
+	if (!MobileBasePassAlwaysUsesCSM(Scene->GetShaderPlatform()))
+	{
+		Flags = FMobileBasePassMeshProcessor::EFlags::CanReceiveCSM | FMobileBasePassMeshProcessor::EFlags::CanUseDepthStencil;
+	}
+	
 	return new FMobileBasePassMeshProcessor(Scene, Scene->GetFeatureLevel(), InViewIfDynamicMeshCommand, PassDrawRenderState, InDrawListContext, Flags);
 }
 
diff --git a/Engine/Source/Runtime/Renderer/Private/MobileBasePassRendering.h b/Engine/Source/Runtime/Renderer/Private/MobileBasePassRendering.h
--- a/Engine/Source/Runtime/Renderer/Private/MobileBasePassRendering.h
+++ b/Engine/Source/Runtime/Renderer/Private/MobileBasePassRendering.h
@@ -438,7 +438,10 @@ public:
 		CanReceiveCSM = (1 << 1),
 
 		// Informs the processor to use PassDrawRenderState for all mesh commands
-		ForcePassDrawRenderState = (1 << 2)
+		ForcePassDrawRenderState = (1 << 2),
+
+		// Informs the processor to not cache any mesh commands
+		DoNotCache = (1 << 3)
 	};
 
 	FMobileBasePassMeshProcessor(
