Fix for depth bounds setting outside of valid range

- issue
  - D3D validation errors encountered on depth bound range
    - fog in the Lyra front end exposes the issue

- fix
  - clamp the depth bounds range set by the fog pass
  - add checks to RHI validation layer to catch future issues

diff --git a/Engine/Source/Runtime/RHI/Public/RHIValidationContext.h b/Engine/Source/Runtime/RHI/Public/RHIValidationContext.h
--- a/Engine/Source/Runtime/RHI/Public/RHIValidationContext.h
+++ b/Engine/Source/Runtime/RHI/Public/RHIValidationContext.h
@@ -814,6 +814,8 @@ public:
 	*/
 	virtual void RHISetDepthBounds(float MinDepth, float MaxDepth) override final
 	{
+		checkf(MinDepth >= 0.f && MinDepth <= 1.f, TEXT("Depth bounds min of %f is outside allowed range of [0, 1]"), MinDepth);
+		checkf(MaxDepth >= 0.f && MaxDepth <= 1.f, TEXT("Depth bounds max of %f is outside allowed range of [0, 1]"), MaxDepth);
 		RHIContext->RHISetDepthBounds(MinDepth, MaxDepth);
 	}
 
diff --git a/Engine/Source/Runtime/Renderer/Private/FogRendering.cpp b/Engine/Source/Runtime/Renderer/Private/FogRendering.cpp
--- a/Engine/Source/Runtime/Renderer/Private/FogRendering.cpp
+++ b/Engine/Source/Runtime/Renderer/Private/FogRendering.cpp
@@ -331,6 +331,7 @@ static void RenderViewFog(
 		FVector ViewSpaceStartFogPoint(0.0f, 0.0f, FogStartDistance * Ratio);
 		FVector4f ClipSpaceMaxDistance = (FVector4f)View.ViewMatrices.GetProjectionMatrix().TransformPosition(ViewSpaceStartFogPoint); // LWC_TODO: precision loss
 		float FogClipSpaceZ = ClipSpaceMaxDistance.Z / ClipSpaceMaxDistance.W;
+		FogClipSpaceZ = FMath::Clamp(FogClipSpaceZ, 0.f, 1.f);
 
 		if (bool(ERHIZBuffer::IsInverted))
 		{
