Make VSM debug visualizations work in splitscreen

diff --git a/Engine/Shaders/Private/VirtualShadowMaps/VirtualShadowMapProjection.usf b/Engine/Shaders/Private/VirtualShadowMaps/VirtualShadowMapProjection.usf
--- a/Engine/Shaders/Private/VirtualShadowMaps/VirtualShadowMapProjection.usf
+++ b/Engine/Shaders/Private/VirtualShadowMaps/VirtualShadowMapProjection.usf
@@ -669,7 +669,7 @@ void VirtualShadowMapProjection(
 #if VISUALIZE_OUTPUT
 	if (InputType == INPUT_TYPE_GBUFFER)
 	{
-		OutputVisualize( PixelPos , VisualizeResult );
+		OutputVisualize( PixelPos - View.ViewRectMin.xy, VisualizeResult );
 	}
 #endif // VISUALIZE_OUTPUT
 }
diff --git a/Engine/Source/Runtime/Renderer/Private/VirtualShadowMaps/VirtualShadowMapArray.cpp b/Engine/Source/Runtime/Renderer/Private/VirtualShadowMaps/VirtualShadowMapArray.cpp
--- a/Engine/Source/Runtime/Renderer/Private/VirtualShadowMaps/VirtualShadowMapArray.cpp
+++ b/Engine/Source/Runtime/Renderer/Private/VirtualShadowMaps/VirtualShadowMapArray.cpp
@@ -1063,6 +1063,7 @@ void FVirtualShadowMapArray::BuildPageAllocations(
 
 	bool bDebugOutputEnabled = false;
 	VisualizeLight.Reset();
+	VisualizeLight.AddDefaulted(Views.Num());
 
 #if !UE_BUILD_SHIPPING
 	if (GDumpVSMLightNames)
@@ -1112,7 +1113,13 @@ void FVirtualShadowMapArray::BuildPageAllocations(
 
 			if (bDebugOutputEnabled)
 			{
-				VisualizeLight.CheckLight(Clipmap->GetLightSceneInfo().Proxy, ClipmapID);
+				for (int32 ViewIndex = 0; ViewIndex < Views.Num(); ViewIndex++)
+				{
+					if (Clipmap->GetDependentView() == Views[ViewIndex].GetPrimaryView())
+					{
+						VisualizeLight[ViewIndex].CheckLight(Clipmap->GetLightSceneInfo().Proxy, ClipmapID);
+					}
+				}
 			}
 		}
 
@@ -1168,7 +1175,10 @@ void FVirtualShadowMapArray::BuildPageAllocations(
 
 				if (bDebugOutputEnabled)
 				{
-					VisualizeLight.CheckLight(ProjectedShadowInfo->GetLightSceneInfo().Proxy, ProjectedShadowInfo->VirtualShadowMaps[0]->ID);
+					for (int32 ViewIndex = 0; ViewIndex < Views.Num(); ViewIndex++)
+					{
+						VisualizeLight[ViewIndex].CheckLight(ProjectedShadowInfo->GetLightSceneInfo().Proxy, ProjectedShadowInfo->VirtualShadowMaps[0]->ID);
+					}
 				}
 			}
 		}
@@ -1647,7 +1657,7 @@ void FVirtualShadowMapArray::RenderDebugInfo(FRDGBuilder& GraphBuilder, TArrayVi
 {
 	check(IsEnabled());
 			
-	if (DebugVisualizationOutput.IsEmpty() || !VisualizeLight.IsValid())
+	if (DebugVisualizationOutput.IsEmpty() || VisualizeLight.IsEmpty())
 	{
 		return;
 	}
@@ -1662,6 +1672,11 @@ void FVirtualShadowMapArray::RenderDebugInfo(FRDGBuilder& GraphBuilder, TArrayVi
 
 	for (int32 ViewIndex = 0; ViewIndex < Views.Num(); ViewIndex++)
 	{
+		if (!VisualizeLight[ViewIndex].IsValid())
+		{
+			continue;
+		}
+
 		FViewInfo& View = Views[ViewIndex];
 
 		FIntPoint DebugTargetExtent = DebugVisualizationOutput[ViewIndex]->Desc.Extent;
@@ -1673,7 +1688,7 @@ void FVirtualShadowMapArray::RenderDebugInfo(FRDGBuilder& GraphBuilder, TArrayVi
 		PassParameters->DebugTargetHeight = DebugTargetExtent.Y;
 		PassParameters->BorderWidth = BorderWidth;
 		PassParameters->VisualizeModeId = VisualizationData.GetActiveModeID();
-		PassParameters->VirtualShadowMapId = VisualizeLight.GetVirtualShadowMapId();
+		PassParameters->VirtualShadowMapId = VisualizeLight[ViewIndex].GetVirtualShadowMapId();
 
 		PassParameters->OutVisualize = GraphBuilder.CreateUAV(DebugVisualizationOutput[ViewIndex]);
 
@@ -2875,7 +2890,7 @@ void FVirtualShadowMapArray::AddVisualizePass(FRDGBuilder& GraphBuilder, const F
 	}
 
 	const FVirtualShadowMapVisualizationData& VisualizationData = GetVirtualShadowMapVisualizationData();
-	if (VisualizationData.IsActive() && VisualizeLight.IsValid())
+	if (VisualizationData.IsActive() && VisualizeLight[ViewIndex].IsValid())
 	{	
 		FCopyRectPS::FParameters* Parameters = GraphBuilder.AllocParameters<FCopyRectPS::FParameters>();
 		Parameters->InputTexture = DebugVisualizationOutput[ViewIndex];
@@ -2910,7 +2925,7 @@ void FVirtualShadowMapArray::AddVisualizePass(FRDGBuilder& GraphBuilder, const F
 			FScreenPassRenderTarget OutputTarget(Output.Texture, View.UnscaledViewRect, ERenderTargetLoadAction::ELoad);
 
 			AddDrawCanvasPass(GraphBuilder, RDG_EVENT_NAME("Labels"), View, OutputTarget,
-				[&VisualizeLight=VisualizeLight, &OutputViewport=OutputViewport](FCanvas& Canvas)
+				[&VisualizeLight=VisualizeLight[ViewIndex], &OutputViewport=OutputViewport](FCanvas& Canvas)
 			{
 				const FLinearColor LabelColor(1, 1, 0);
 				Canvas.DrawShadowedString(
diff --git a/Engine/Source/Runtime/Renderer/Private/VirtualShadowMaps/VirtualShadowMapArray.h b/Engine/Source/Runtime/Renderer/Private/VirtualShadowMaps/VirtualShadowMapArray.h
--- a/Engine/Source/Runtime/Renderer/Private/VirtualShadowMaps/VirtualShadowMapArray.h
+++ b/Engine/Source/Runtime/Renderer/Private/VirtualShadowMaps/VirtualShadowMapArray.h
@@ -333,7 +333,7 @@ public:
 
 	// Debug visualization
 	TArray<FRDGTextureRef> DebugVisualizationOutput;
-	FVirtualShadowMapVisualizeLightSearch VisualizeLight;
+	TArray<FVirtualShadowMapVisualizeLightSearch> VisualizeLight;
 
 private:
 	FScene &Scene;
diff --git a/Engine/Source/Runtime/Renderer/Private/VirtualShadowMaps/VirtualShadowMapProjection.cpp b/Engine/Source/Runtime/Renderer/Private/VirtualShadowMaps/VirtualShadowMapProjection.cpp
--- a/Engine/Source/Runtime/Renderer/Private/VirtualShadowMaps/VirtualShadowMapProjection.cpp
+++ b/Engine/Source/Runtime/Renderer/Private/VirtualShadowMaps/VirtualShadowMapProjection.cpp
@@ -415,13 +415,13 @@ static void RenderVirtualShadowMapProjectionCommon(
 	
 	bool bDebugOutput = false;
 #if !UE_BUILD_SHIPPING
-	if ( !VirtualShadowMapArray.DebugVisualizationOutput.IsEmpty() && InputType == EVirtualShadowMapProjectionInputType::GBuffer && VirtualShadowMapArray.VisualizeLight.IsValid() )
+	if ( !VirtualShadowMapArray.DebugVisualizationOutput.IsEmpty() && InputType == EVirtualShadowMapProjectionInputType::GBuffer && VirtualShadowMapArray.VisualizeLight[ViewIndex].IsValid())
 	{
 		const FVirtualShadowMapVisualizationData& VisualizationData = GetVirtualShadowMapVisualizationData();
 
 		bDebugOutput = true;
 		PassParameters->VisualizeModeId = VisualizationData.GetActiveModeID();
-		PassParameters->VisualizeVirtualShadowMapId = VirtualShadowMapArray.VisualizeLight.GetVirtualShadowMapId();
+		PassParameters->VisualizeVirtualShadowMapId = VirtualShadowMapArray.VisualizeLight[ViewIndex].GetVirtualShadowMapId();
 		PassParameters->PhysicalPageMetaData = GraphBuilder.CreateSRV( VirtualShadowMapArray.PhysicalPageMetaDataRDG );
 		PassParameters->OutVisualize = GraphBuilder.CreateUAV( VirtualShadowMapArray.DebugVisualizationOutput[ViewIndex] );
 	}
