Disabling Strata CMask clearing path until we fix all platforms.

diff --git a/Engine/Source/Runtime/Renderer/Private/Strata/Strata.cpp b/Engine/Source/Runtime/Renderer/Private/Strata/Strata.cpp
--- a/Engine/Source/Runtime/Renderer/Private/Strata/Strata.cpp
+++ b/Engine/Source/Runtime/Renderer/Private/Strata/Strata.cpp
@@ -51,7 +51,7 @@ static TAutoConsoleVariable<int32> CVarMaterialRoughDiffuse(
 // STRATA_TODO we keep this for now and can remove it once battletested.
 static TAutoConsoleVariable<int32> CVarUseCmaskClear(
 	TEXT("r.Strata.UseCmaskClear"),
-	1,
+	0,
 	TEXT("TEST."),
 	ECVF_RenderThreadSafe);
 
