Fixes DOF CoC computation when the focus distance is wrong, but depth blur is enabled.

diff --git a/Engine/Source/Runtime/Renderer/Private/PostProcess/DiaphragmDOFUtils.cpp b/Engine/Source/Runtime/Renderer/Private/PostProcess/DiaphragmDOFUtils.cpp
--- a/Engine/Source/Runtime/Renderer/Private/PostProcess/DiaphragmDOFUtils.cpp
+++ b/Engine/Source/Runtime/Renderer/Private/PostProcess/DiaphragmDOFUtils.cpp
@@ -117,6 +117,7 @@ void DiaphragmDOF::FPhysicalCocModel::Compile(const FViewInfo& View)
 	}
 
 	// Compile coc model equation.
+	if (View.FinalPostProcessSettings.DepthOfFieldFstop > 0.f && View.FinalPostProcessSettings.DepthOfFieldFocalDistance > 0.f)
 	{
 
 		float FocalLengthInMM = DiaphragmDOF::ComputeFocalLengthFromFov(View);
@@ -140,6 +141,11 @@ void DiaphragmDOF::FPhysicalCocModel::Compile(const FViewInfo& View)
 		// Convert diameter in mm to resolution less radius on the filmback.
 		InfinityBackgroundCocRadius = DiameterInMM * 0.5f / SensorWidthInMM;
 	}
+	else
+	{
+		InfinityBackgroundCocRadius = 0.0f;
+		MinForegroundCocRadius = 0.0;
+	}
 }
 
 float DiaphragmDOF::FPhysicalCocModel::DepthToResCocRadius(float SceneDepth, float HorizontalResolution) const
