HeterogeneousVolumes Rendering Library

This is the initial submission for the HeterogeneousVolumes rendering subsystem. The initial goal of the rendering subsystem is to formalize Niagara Fluids volume rendering as a core rendering feature. Doing so gives us the opportunity to fully integrate with the rest of the UE5 rendering feature set.

To render a Niagara Fluids volume with the HeterogeneousVolumes subsystem, the artist must modify the "Override Materials" parameter of the Mesh Renderer settings on an existing Grid3D emitter. Rather than relying on custom HLSL ray-marching nodes bound to a "Surface" Material Domain, HeterogeneousVolumes incorporates the existing Volumetric Fog paradigm. This requires the artist to bind a "Volume" Material Domain with an "Additive" Blend mode.

The HeterogeneousVolumes algorithm uses ray-marching to evaluate the direct lighting component of the Volume Rendering Equation. This is similar to the existing Niagara Fluids volume rendering algorithm, though some of its limitations have been removed. Most importantly, HeterogeneousVolumes does not impose a limit on the number of lights. HeterogeneousVolumes integrates with the Deferred Lighting model and currently supports all UE5 light types, except for the SkyLight model.

An artist can enable the HeterogeneousVolumes system with the following CVar:
r.HeterogeneousVolumes=1 (Default = 0)

The ray-marching step size can be controlled with the following CVar:
r.HeterogeneousVolumes.StepSize=X (Default = 1.0)

Shadowing fidelity can be separately controlled with the following CVar:
r.HeterogeneousVolumes.ShadowStepFactor=X (Default = 8.0)

diff --git a/Engine/Shaders/Private/HeterogeneousVolumes/HeterogeneousVolumesComposite.usf b/Engine/Shaders/Private/HeterogeneousVolumes/HeterogeneousVolumesComposite.usf
new file mode 100644
--- /dev/null
+++ b/Engine/Shaders/Private/HeterogeneousVolumes/HeterogeneousVolumesComposite.usf
@@ -0,0 +1,26 @@
+#include "../Common.ush"
+
+Texture2D<float4> HeterogeneousVolumeRadiance;
+RWTexture2D<float4> RWColorTexture;
+
+[numthreads(THREADGROUP_SIZE_2D, THREADGROUP_SIZE_2D, 1)]
+void HeterogeneousVolumesCompositeCS(
+	uint2 GroupThreadId : SV_GroupThreadID,
+	uint2 DispatchThreadId : SV_DispatchThreadID
+)
+{
+	uint2 PixelCoord = DispatchThreadId.xy;
+	if (any(PixelCoord >= View.ViewSizeAndInvSize.xy))
+	{
+		return;
+	}
+
+	float4 Result;
+	float HeterogeneousVolumeTransmittance = saturate(1.0 - HeterogeneousVolumeRadiance[PixelCoord].a);
+	Result.rgb = HeterogeneousVolumeRadiance[PixelCoord].rgb + RWColorTexture[PixelCoord].rgb * HeterogeneousVolumeTransmittance;
+
+	float AccumulatedTransmittance = (1.0 - RWColorTexture[PixelCoord].a) * HeterogeneousVolumeTransmittance;
+	Result.a = 1.0 - AccumulatedTransmittance;
+
+	RWColorTexture[PixelCoord] = Result;
+}
diff --git a/Engine/Shaders/Private/HeterogeneousVolumes/HeterogeneousVolumesHardwareRayTracing.usf b/Engine/Shaders/Private/HeterogeneousVolumes/HeterogeneousVolumesHardwareRayTracing.usf
new file mode 100644
--- /dev/null
+++ b/Engine/Shaders/Private/HeterogeneousVolumes/HeterogeneousVolumesHardwareRayTracing.usf
@@ -0,0 +1,502 @@
+// Copyright Epic Games, Inc. All Rights Reserved.
+
+#include "HeterogeneousVolumesSparseVoxelUtils.ush"
+#include "../Common.ush"
+#include "../RayTracing/RayTracingCommon.ush"
+
+#define SUPPORT_CONTACT_SHADOWS 0
+
+#include "HeterogeneousVolumesSparseVoxelUniformBufferUtils.ush"
+#include "HeterogeneousVolumesTracingUtils.ush"
+#include "HeterogeneousVolumesLightingUtils.ush"
+
+#define VOXEL_COUNT_MAX 6
+
+#define SHADOW_BIT 0x80000000
+
+struct FSparseVoxelPayload
+{
+	float HitT[2];
+	uint VoxelID;
+	uint Data;
+	float3 Transmittance;
+
+	void SetMiss() { HitT[0] = -1; }
+	bool IsMiss() { return HitT[0] < 0; }
+	bool IsHit() { return !IsMiss(); }
+
+	void SetShadowRay()
+	{
+		Data = (Data & ~SHADOW_BIT) | SHADOW_BIT;
+	}
+	bool IsShadowRay()
+	{
+		return (Data & SHADOW_BIT) != 0;
+	}
+
+	void SetTransmittance(float3 Transm)
+	{
+		Transmittance = Transm;
+	}
+
+	float3 GetTransmittance()
+	{
+		return Transmittance;
+	}
+};
+
+FSparseVoxelPayload CreateSparseVoxelPayload()
+{
+	FSparseVoxelPayload Payload = (FSparseVoxelPayload)0;
+	Payload.SetMiss();
+	Payload.SetTransmittance(1.0);
+
+	return Payload;
+}
+
+uint3 GetVolumeResolutionAtMip(uint MipLevel)
+{
+	uint3 VolumeResolution = GetVolumeResolution();
+	return uint3(
+		VolumeResolution.x >> MipLevel,
+		VolumeResolution.y >> MipLevel,
+		VolumeResolution.z >> MipLevel
+	);
+}
+
+FVoxelData GetVoxelData(uint PrimitiveIndex, uint InstanceID)
+{
+	FVoxelDataPacked VoxelDataPacked = GetVoxelDataPacked(PrimitiveIndex);
+	return DecodeVoxelData(VoxelDataPacked, GetVolumeResolution());
+}
+
+float3 GetObjectBoundsMin()
+{
+	float3 ObjectBoundsMin = GetLocalBoundsOrigin() - GetLocalBoundsExtent();
+	return ObjectBoundsMin;
+}
+
+float3 GetObjectBoundsMax()
+{
+	float3 ObjectBoundsMax = GetLocalBoundsOrigin() + GetLocalBoundsExtent();
+	return ObjectBoundsMax;
+}
+
+struct FSparseVoxelIntersectionAttributes
+{
+	float2 HitT;
+};
+
+FSparseVoxelIntersectionAttributes CreateVolumeIntersectionAttributes(float2 HitT)
+{
+	FSparseVoxelIntersectionAttributes Attr = (FSparseVoxelIntersectionAttributes)0;
+	Attr.HitT = HitT;
+	return Attr;
+}
+
+RAY_TRACING_ENTRY_CLOSEST_HIT(SparseVoxelsClosestHitShader,
+	FSparseVoxelPayload, Payload,
+	FSparseVoxelIntersectionAttributes, Attributes)
+{
+	FVoxelData VoxelData = GetVoxelData(PrimitiveIndex(), InstanceID());
+
+	Payload.HitT[0] = Attributes.HitT.x;
+	Payload.HitT[1] = Attributes.HitT.y;
+	Payload.VoxelID = PrimitiveIndex();
+}
+
+RAY_TRACING_ENTRY_ANY_HIT(SparseVoxelsAnyHitShader,
+	FSparseVoxelPayload, Payload,
+	FSparseVoxelIntersectionAttributes, Attributes)
+{
+	FVoxelData VoxelData = GetVoxelData(PrimitiveIndex(), InstanceID());
+
+	// Compute transmittance instead..
+	Payload.HitT[0] = Attributes.HitT.x;
+	Payload.HitT[1] = Attributes.HitT.y;
+	Payload.VoxelID = PrimitiveIndex();
+
+	if (Payload.IsShadowRay())
+	{
+		float3 Transmittance = Payload.GetTransmittance();
+
+		float StepSize = GetStepSize() * GetShadowStepFactor();
+		float HitSpan = Payload.HitT[1] - Payload.HitT[0];
+		int StepCount = max(HitSpan / StepSize, 1);
+		StepSize = HitSpan / StepCount;
+		float MipLevel = 0.0;
+
+		for (int StepIndex = 0; StepIndex < StepCount; ++StepIndex)
+		{
+			float MarchT = StepIndex * StepSize;
+			float3 ObjectPosition = ObjectRayOrigin() + ObjectRayDirection() * (Payload.HitT[0] + MarchT);
+#if 1
+			float3 UVW = (ObjectPosition - GetObjectBoundsMin()) / (GetObjectBoundsMax() - GetObjectBoundsMin());
+			float3 Extinction = Texture3DSampleLevel(SparseVoxelUniformBuffer.ExtinctionTexture, SparseVoxelUniformBuffer.TextureSampler, UVW, MipLevel).rgb;
+#else
+			float3 WorldPosition = mul(float4(ObjectPosition, 1.0), GetLocalToWorld()).xyz;
+
+			FVolumeSampleContext SampleContext = CreateVolumeSampleContext(ObjectPosition, WorldPosition);
+			float3 Extinction = SampleExtinction(SampleContext);
+#endif
+			Transmittance *= exp(-Extinction * StepSize);
+		}
+
+		Payload.SetTransmittance(Transmittance);
+
+		// Early ray termination
+		float Epsilon = 1.0e-3;
+		if (all(Transmittance < Epsilon))
+		{
+			AcceptHitAndEndSearch();
+		}
+		else
+		{
+			IgnoreHit();
+		}
+	}
+}
+
+RAY_TRACING_ENTRY_INTERSECTION(SparseVoxelsIntersectionShader)
+{
+	FVoxelData VoxelData = GetVoxelData(PrimitiveIndex(), InstanceID());
+	FBox Box = GetVoxelBoundingBox(VoxelData, GetVolumeResolution(), GetObjectBoundsMin(), GetObjectBoundsMax());
+
+	// Intersect ray against box
+	float2 HitT = IntersectAABB(ObjectRayOrigin(), ObjectRayDirection(), RayTMin(), RayTCurrent(), Box.Min, Box.Max);
+	if (HitT.y > HitT.x)
+	{
+		uint HitKind = 0;
+		ReportHit(HitT.x, HitKind, CreateVolumeIntersectionAttributes(HitT));
+	}
+}
+
+RAY_TRACING_ENTRY_MISS(SparseVoxelsMissShader,
+	FSparseVoxelPayload, Payload)
+{
+	Payload.SetMiss();
+}
+
+// Scene
+RaytracingAccelerationStructure TLAS;
+
+// Ray data
+float MaxTraceDistance;
+int MaxStepCount;
+int bJitter;
+
+// Volume sampling data
+int MipLevel;
+
+// Output
+RWTexture2D<float4> RWLightingTexture;
+
+RayDesc CreateVolumeRay(
+	float3 Origin,
+	float3 Direction,
+	float TMin,
+	float TMax
+)
+{
+	RayDesc Ray;
+	Ray.Origin = Origin;
+	Ray.Direction = Direction;
+	Ray.TMin = TMin;
+	Ray.TMax = TMax;
+
+	return Ray;
+}
+
+void TraceVolumeRay(
+	RayDesc Ray,
+	inout FSparseVoxelPayload Payload
+)
+{
+	uint RayFlags = 0;
+	uint InstanceInclusionMask = 0xFF;
+	uint RayContributionToHitGroupIndex = 0;
+	uint MultiplierForGeometryContributionToHitGroupIndex = 0;
+	uint MissShaderIndex = 0;
+
+	TraceRay(
+		TLAS,
+		RayFlags,
+		InstanceInclusionMask,
+		RayContributionToHitGroupIndex,
+		MultiplierForGeometryContributionToHitGroupIndex,
+		MissShaderIndex,
+		Ray,
+		Payload
+	);
+}
+
+void TraceVolumeShadowRay(
+	RayDesc Ray,
+	inout FSparseVoxelPayload Payload
+)
+{
+	Payload.SetShadowRay();
+	TraceVolumeRay(Ray, Payload);
+}
+
+// h = [0, 1], s = [0 ,1], v = [0, 1]
+// h: hue, 
+// s: saturate
+// v: value
+float3 hsv2rgb(float h, float s, float v)
+{
+	h = 360.0 * h;
+
+	float C = v * s;
+	float X = C * (1 - abs((h / 60) % 2 - 1));
+	float m = v - C;
+
+	float3 RGB = 0;
+	if (h >= 0 && h < 60)
+	{
+		RGB.r = C;
+		RGB.g = X;
+	}
+	else if (h >= 60 && h < 120)
+	{
+		RGB.r = X;
+		RGB.g = C;
+	}
+	else if (h >= 120 && h < 180)
+	{
+		RGB.g = C;
+		RGB.b = X;
+	}
+	else if (h >= 180 && h < 240)
+	{
+		RGB.g = X;
+		RGB.b = C;
+	}
+	else if (h >= 240 && h < 300)
+	{
+		RGB.r = X;
+		RGB.b = C;
+	}
+	else if (h >= 300 && h < 360)
+	{
+		RGB.r = C;
+		RGB.b = X;
+	}
+
+	RGB += m;
+	return RGB;
+}
+
+float3 GetVoxelDebugColor(uint VoxelID)
+{
+	float Hue = (VoxelID % GetVolumeResolutionAtMip(MipLevel).x) / float(GetVolumeResolutionAtMip(MipLevel).x);
+	return hsv2rgb(Hue, 1.0, 1.0);
+}
+
+float3 GetVoxelDebugColor(uint3 VoxelIndex)
+{
+	//float Hue = RandFast(VoxelIndex.yz);
+	float Hue = float(Rand3DPCG32(VoxelIndex).x) / 0xFFFFFFFF;
+	return hsv2rgb(Hue, 1.0, 1.0);
+}
+
+RAY_TRACING_ENTRY_RAYGEN(SparseVoxelsRayGenerationShader1)
+{
+	uint2 PixelCoord = DispatchRaysIndex().xy;
+	if (any(PixelCoord.xy >= View.ViewSizeAndInvSize.xy))
+	{
+		return;
+	}
+	PixelCoord += View.ViewRectMin.xy;
+
+	// Extract depth
+	float DeviceZ = SceneDepthTexture.Load(int3(PixelCoord, 0)).r;
+#if HAS_INVERTED_Z_BUFFER
+	DeviceZ = max(0.000000000001, DeviceZ);
+#endif // HAS_INVERTED_Z_BUFFER
+
+	// Clip trace distance
+	float SceneDepth = min(ConvertFromDeviceZ(DeviceZ), MaxTraceDistance);
+	DeviceZ = ConvertToDeviceZ(SceneDepth);
+
+	float Jitter = bJitter ? InterleavedGradientNoise(PixelCoord, View.StateFrameIndexMod8) : 0.0;
+
+	// Intersect ray with bounding volume
+	// TODO: LWC integration..
+	float3 WorldRayOrigin = LWCHackToFloat(LWCSubtract(View.TranslatedWorldCameraOrigin, PrimaryView.PreViewTranslation));
+	float3 WorldRayEnd = LWCHackToFloat(SvPositionToWorld(float4(PixelCoord + 0.5, DeviceZ, 1)));
+	float3 WorldRayDirection = WorldRayEnd - WorldRayOrigin;
+	float WorldRayLength = length(WorldRayDirection); // SceneDepth
+
+	float3 LocalRayOrigin = mul(float4(WorldRayOrigin, 1.0), GetWorldToLocal()).xyz;
+	float3 LocalRayEnd = mul(float4(WorldRayEnd, 1.0), GetWorldToLocal()).xyz;
+	float3 LocalRayDirection = LocalRayEnd - LocalRayOrigin;
+	float LocalRayLength = length(LocalRayDirection);
+
+	float3 LocalBoundsMin = GetLocalBoundsOrigin() - GetLocalBoundsExtent();
+	float3 LocalBoundsMax = GetLocalBoundsOrigin() + GetLocalBoundsExtent();
+	float2 HitT = LineBoxIntersect(LocalRayOrigin, LocalRayEnd, LocalBoundsMin, LocalBoundsMax);
+
+	float HitSpan = HitT.y - HitT.x;
+	if (HitSpan > 0.0)
+	{
+		float3 WorldRayStart = WorldRayOrigin + WorldRayDirection * HitT.x;
+		float3 LocalRayStart = LocalRayOrigin + LocalRayDirection * HitT.x;
+
+		// Convert normalized t-value intersection parameters to world-space units
+		WorldRayDirection *= rcp(WorldRayLength);
+		LocalRayDirection *= rcp(LocalRayLength);
+		LocalRayLength *= HitSpan;
+
+		// Jitter starting position to mitigate moire
+		LocalRayStart += LocalRayDirection * Jitter;
+
+		// Intersect against TLAS - HitT.x is real close to first intersection shader return value...
+		RayDesc Ray = CreateVolumeRay(WorldRayOrigin + WorldRayDirection * Jitter, WorldRayDirection, WorldRayLength * HitT.x, WorldRayLength * HitT.y);
+		FSparseVoxelPayload Payload = CreateSparseVoxelPayload();
+		TraceVolumeShadowRay(Ray, Payload);
+
+		// Output..
+		float3 Transmittance = Payload.GetTransmittance();
+		float3 Opacity = 1.0 - Transmittance;
+		RWLightingTexture[PixelCoord] = float4(Transmittance, Luminance(Opacity));
+	}
+}
+
+int bApplyEmission;
+int bApplyDirectLighting;
+int LightType;
+
+float3 ComputeTransmittanceHardwareRayTracing(
+	float3 Origin,
+	float3 Direction,
+	float TMin,
+	float TMax
+)
+{
+	// TODO: Apply jitter?
+	RayDesc ShadowRay = CreateVolumeRay(Origin, Direction, TMin, TMax);
+	FSparseVoxelPayload Payload = CreateSparseVoxelPayload();
+	TraceVolumeShadowRay(ShadowRay, Payload);
+	return Payload.GetTransmittance();
+}
+
+#define HARDWARE_RAY_TRACING 1
+#include "HeterogeneousVolumesRayMarchingUtils.ush"
+
+RAY_TRACING_ENTRY_RAYGEN(RenderSingleScatteringWithPreshadingRGS)
+{
+	float3 Radiance = 0.0;
+	float3 Transmittance = 1.0;
+
+	// TODO: Render tiles?
+
+	uint2 PixelCoord = DispatchRaysIndex().xy;
+	if (any(PixelCoord.xy >= View.ViewSizeAndInvSize.xy))
+	{
+		return;
+	}
+	PixelCoord += View.ViewRectMin.xy;
+
+	// Extract depth
+	float DeviceZ = SceneDepthTexture.Load(int3(PixelCoord, 0)).r;
+#if HAS_INVERTED_Z_BUFFER
+	DeviceZ = max(0.000000000001, DeviceZ);
+#endif // HAS_INVERTED_Z_BUFFER
+
+	// Clip trace distance
+	float SceneDepth = min(ConvertFromDeviceZ(DeviceZ), MaxTraceDistance);
+	DeviceZ = ConvertToDeviceZ(SceneDepth);
+
+	float Jitter = bJitter ? InterleavedGradientNoise(PixelCoord, View.StateFrameIndexMod8) : 0.0;
+
+	// Intersect ray with bounding volume
+	// TODO: LWC integration..
+	float3 WorldRayOrigin = LWCHackToFloat(LWCSubtract(View.TranslatedWorldCameraOrigin, PrimaryView.PreViewTranslation));
+	float3 WorldRayEnd = LWCHackToFloat(SvPositionToWorld(float4(PixelCoord + 0.5, DeviceZ, 1)));
+	float3 WorldRayDirection = WorldRayEnd - WorldRayOrigin;
+	float WorldRayLength = length(WorldRayDirection);
+	WorldRayDirection /= WorldRayLength;
+
+	float3 LocalRayOrigin = mul(float4(WorldRayOrigin, 1.0), GetWorldToLocal()).xyz;
+	float3 LocalRayEnd = mul(float4(WorldRayEnd, 1.0), GetWorldToLocal()).xyz;
+	float3 LocalRayDirection = LocalRayEnd - LocalRayOrigin;
+	float LocalRayLength = length(LocalRayDirection);
+	LocalRayDirection /= LocalRayLength;
+
+	float3 LocalBoundsMin = GetLocalBoundsOrigin() - GetLocalBoundsExtent();
+	float3 LocalBoundsMax = GetLocalBoundsOrigin() + GetLocalBoundsExtent();
+
+	// Test bounding volume
+	float2 BoundingHitT = IntersectAABB(LocalRayOrigin, LocalRayDirection, 0.0, SceneDepth, LocalBoundsMin, LocalBoundsMax);
+	bool bIntersectsVolume = (BoundingHitT.y - BoundingHitT.x > 0.0);
+	if (!bIntersectsVolume)
+	{
+		return;
+	}
+
+	// Intersect against TLAS - HitT.x is real close to first intersection shader return value...
+	RayDesc Ray = CreateVolumeRay(WorldRayOrigin, WorldRayDirection, BoundingHitT.x, BoundingHitT.y);
+
+	FSparseVoxelPayload Payload = CreateSparseVoxelPayload();
+	TraceVolumeRay(Ray, Payload);
+
+	float HitSpan = Payload.HitT[1] - Payload.HitT[0];
+	while (Payload.IsHit() && Luminance(Transmittance) > 0.001)
+	{
+		FRayMarchingContext RayMarchingContext = CreateRayMarchingContext(
+			// Local-space
+			LocalRayOrigin,
+			LocalRayDirection,
+			Payload.HitT[0],
+			Payload.HitT[1],
+			// World-space
+			WorldRayOrigin,
+			WorldRayDirection,
+			// Ray-step attributes
+			Jitter,
+			GetStepSize(),
+			MaxStepCount,
+			bApplyEmission,
+			bApplyDirectLighting
+		);
+
+		FDeferredLightData LightData = LoadLightData(LightType);
+		uint StepCount = CalcStepCount(RayMarchingContext);
+
+		RayMarchSingleScattering(
+			RayMarchingContext,
+			LightData,
+			LightType,
+			StepCount,
+			Radiance,
+			Transmittance
+		);
+
+//#define DEBUG_VISUALIZATION 1
+#if DEBUG_VISUALIZATION
+		FVoxelData VoxelData = GetVoxelData(Payload.VoxelID, 0);
+		if (View.GeneralPurposeTweak == 0)
+		{
+			// TODO: Remap Extinction to Transmittance
+			float3 Extinction = 0.02 * (1.0 - GetVoxelDebugColor(VoxelData.VoxelIndex));
+			Transmittance *= exp(-Extinction * HitSpan);
+			if (!bDirectLighting)
+			{
+				float3 Emission = 0.002 * GetVoxelDebugColor(VoxelData.VoxelIndex);
+				Radiance += Emission * HitSpan * Transmittance;
+			}
+		}
+#endif // DEBUG_VISUALIZATION
+
+		//Ray = CreateVolumeRay(WorldRayOrigin, WorldRayDirection, Payload.HitT[1] + 0.001, WorldRayLength * HitT.y);
+		Ray = CreateVolumeRay(WorldRayOrigin, WorldRayDirection, Payload.HitT[1] + 0.001, BoundingHitT.y);
+		Payload = CreateSparseVoxelPayload();
+		TraceVolumeRay(Ray, Payload);
+		HitSpan = Payload.HitT[1] - Payload.HitT[0];
+	}
+
+	// Output..
+	float3 Opacity = 1.0 - Transmittance;
+	RWLightingTexture[PixelCoord] += float4(Radiance, Luminance(Opacity));
+}
diff --git a/Engine/Shaders/Private/HeterogeneousVolumes/HeterogeneousVolumesLightingUtils.ush b/Engine/Shaders/Private/HeterogeneousVolumes/HeterogeneousVolumesLightingUtils.ush
new file mode 100644
--- /dev/null
+++ b/Engine/Shaders/Private/HeterogeneousVolumes/HeterogeneousVolumesLightingUtils.ush
@@ -0,0 +1,42 @@
+// Copyright Epic Games, Inc. All Rights Reserved.
+
+#include "../Common.ush"
+#include "../DeferredLightingCommon.ush"
+#include "../LightGridCommon.ush"
+
+FDeferredLightData LoadLightData(uint LightType)
+{
+	FDeferredLightData LightData = (FDeferredLightData)0;
+	LightData.TranslatedWorldPosition = GetDeferredLightTranslatedWorldPosition();
+	LightData.InvRadius = DeferredLightUniforms.InvRadius;
+	LightData.Color = DeferredLightUniforms.Color;
+	LightData.FalloffExponent = DeferredLightUniforms.FalloffExponent;
+	LightData.Direction = DeferredLightUniforms.Direction;
+	LightData.Tangent = DeferredLightUniforms.Tangent;
+	LightData.SpotAngles = DeferredLightUniforms.SpotAngles;
+	LightData.SourceRadius = DeferredLightUniforms.SourceRadius;
+	LightData.SourceLength = DeferredLightUniforms.SourceLength;
+	LightData.SoftSourceRadius = DeferredLightUniforms.SoftSourceRadius;
+	LightData.SpecularScale = DeferredLightUniforms.SpecularScale;
+	LightData.ContactShadowLength = abs(DeferredLightUniforms.ContactShadowLength);
+	LightData.ContactShadowLengthInWS = DeferredLightUniforms.ContactShadowLength < 0.0f;
+	LightData.ContactShadowNonShadowCastingIntensity = DeferredLightUniforms.ContactShadowNonShadowCastingIntensity;
+	LightData.DistanceFadeMAD = DeferredLightUniforms.DistanceFadeMAD;
+	LightData.ShadowMapChannelMask = DeferredLightUniforms.ShadowMapChannelMask;
+	LightData.ShadowedBits = DeferredLightUniforms.ShadowedBits;
+	LightData.bInverseSquared = LightType != LIGHT_TYPE_DIRECTIONAL && DeferredLightUniforms.FalloffExponent == 0; // Directional lights don't use 'inverse squared attenuation'
+
+	LightData.RectLightBarnCosAngle = DeferredLightUniforms.RectLightBarnCosAngle;
+	LightData.RectLightBarnLength = DeferredLightUniforms.RectLightBarnLength;
+	LightData.RectLightAtlasMaxLevel = DeferredLightUniforms.RectLightAtlasMaxLevel;
+	LightData.RectLightAtlasUVOffset = DeferredLightUniforms.RectLightAtlasUVOffset;
+	LightData.RectLightAtlasUVScale = DeferredLightUniforms.RectLightAtlasUVScale;
+
+	// Lumen interpretation
+	LightData.bRadialLight = LightType != LIGHT_TYPE_DIRECTIONAL;
+	LightData.bSpotLight = LightType == LIGHT_TYPE_SPOT;
+	LightData.bRectLight = LightType == LIGHT_TYPE_RECT;
+
+	LightData.HairTransmittance = InitHairTransmittanceData();
+	return LightData;
+}
\ No newline at end of file
diff --git a/Engine/Shaders/Private/HeterogeneousVolumes/HeterogeneousVolumesLiveShadingPipeline.usf b/Engine/Shaders/Private/HeterogeneousVolumes/HeterogeneousVolumesLiveShadingPipeline.usf
new file mode 100644
--- /dev/null
+++ b/Engine/Shaders/Private/HeterogeneousVolumes/HeterogeneousVolumesLiveShadingPipeline.usf
@@ -0,0 +1,228 @@
+// Copyright Epic Games, Inc. All Rights Reserved.
+
+#include "HeterogeneousVolumesSparseVoxelUtils.ush"
+#include "../Common.ush"
+
+#define SUPPORT_CONTACT_SHADOWS 0
+
+#include "/Engine/Generated/Material.ush"
+#include "HeterogeneousVolumesLiveShadingUtils.ush"
+#include "HeterogeneousVolumesTracingUtils.ush"
+#include "HeterogeneousVolumesLightingUtils.ush"
+
+// Object data
+float4x4 LocalToWorld;
+float4x4 WorldToLocal;
+float3 LocalBoundsOrigin;
+float3 LocalBoundsExtent;
+int PrimitiveId;
+
+// Lighting
+int bApplyEmission;
+int bApplyDirectLighting;
+int LightType;
+float ShadowStepFactor;
+
+// Ray data
+float MaxTraceDistance;
+float StepSize;
+int MaxStepCount;
+int bJitter;
+
+// Dispatch data
+int3 GroupCount;
+
+// Output
+RWTexture2D<float4> RWLightingTexture;
+
+struct FDebugOutput
+{
+	int PrimitiveId;
+	float3 LocalObjectBoundsMin;
+	float3 LocalObjectBoundsMax;
+	int Paddington;
+	float4x4 LocalToWorld;
+	float4x4 WorldToLocal;
+};
+
+RWStructuredBuffer<FDebugOutput> RWDebugOutputBuffer;
+
+FDebugOutput CreateDebugOutput(
+	inout FMaterialPixelParameters MaterialParameters,
+	inout FPixelMaterialInputs PixelMaterialInputs)
+{
+	FDebugOutput DebugOutput = (FDebugOutput)0;
+
+	DebugOutput.PrimitiveId = MaterialParameters.PrimitiveId;
+	DebugOutput.LocalObjectBoundsMin = GetPrimitiveData(MaterialParameters).LocalObjectBoundsMin.xyz;
+	DebugOutput.LocalObjectBoundsMax = GetPrimitiveData(MaterialParameters).LocalObjectBoundsMax.xyz;
+	DebugOutput.LocalToWorld = LWCToFloat(GetPrimitiveData(MaterialParameters).LocalToWorld);
+	DebugOutput.WorldToLocal = LWCToFloat(GetPrimitiveData(MaterialParameters).WorldToLocal);
+
+	return DebugOutput;
+}
+
+FPrimitiveSceneData GetPrimitiveData(FMaterialPixelParameters Parameters)
+{
+	return GetPrimitiveData(Parameters, LocalToWorld, WorldToLocal, LocalBoundsOrigin, LocalBoundsExtent);
+}
+
+float4x4 GetLocalToWorld()
+{
+	return LocalToWorld;
+}
+
+float4x4 GetWorldToLocal()
+{
+	return WorldToLocal;
+}
+
+float3 GetLocalBoundsOrigin()
+{
+	return LocalBoundsOrigin;
+}
+
+float3 GetLocalBoundsExtent()
+{
+	return LocalBoundsExtent;
+}
+
+float GetStepSize()
+{
+	return StepSize;
+}
+
+float GetShadowStepFactor()
+{
+	return ShadowStepFactor;
+}
+
+struct FVolumeSampleContext
+{
+	FMaterialPixelParameters MaterialParameters;
+	FPixelMaterialInputs PixelMaterialInputs;
+};
+
+FVolumeSampleContext CreateVolumeSampleContext(float3 LocalPosition, float3 WorldPosition)
+{
+	FVolumeSampleContext VolumeSampleContext;
+	VolumeSampleContext.MaterialParameters = MakeInitializedMaterialPixelParameters();
+	//VolumeSampleContext.MaterialParameters.PrimitiveId = PrimitiveId;
+	VolumeSampleContext.MaterialParameters.AbsoluteWorldPosition = LWCPromote(WorldPosition);
+
+	// Evaluate Material graph
+	CalcPixelMaterialInputs(VolumeSampleContext.MaterialParameters, VolumeSampleContext.PixelMaterialInputs);
+
+	return VolumeSampleContext;
+}
+
+float3 SampleExtinction(inout FVolumeSampleContext Context)
+{
+	float3 Extinction = SampleExtinctionCoefficients(Context.PixelMaterialInputs);
+	return Extinction;
+}
+
+float3 SampleEmission(inout FVolumeSampleContext Context)
+{
+	float3 Emission = SampleEmissive(Context.PixelMaterialInputs);
+	return Emission;
+}
+
+float3 SampleAlbedo(inout FVolumeSampleContext Context)
+{
+	float3 Albedo = SampleAlbedo(Context.PixelMaterialInputs);
+	return Albedo;
+}
+
+#include "HeterogeneousVolumesRayMarchingUtils.ush"
+
+[numthreads(THREADGROUP_SIZE_2D, THREADGROUP_SIZE_2D, 1)]
+void RenderSingleScatteringWithLiveShadingCS(
+	uint2 GroupThreadId : SV_GroupThreadID,
+	uint2 DispatchThreadId : SV_DispatchThreadID
+)
+{
+	float3 Radiance = 0.0;
+	float3 Transmittance = 1.0;
+	FDebugOutput DebugOutput = (FDebugOutput)0;
+
+	// Create screen ray
+	if (any(DispatchThreadId.xy >= View.ViewSizeAndInvSize.xy))
+	{
+		return;
+	}
+	uint LinearIndex = DispatchThreadId.y * GroupCount.x * THREADGROUP_SIZE_2D + DispatchThreadId.x;
+	uint2 PixelCoord = DispatchThreadId.xy + View.ViewRectMin.xy;
+
+	// Extract depth
+	float DeviceZ = SceneDepthTexture.Load(int3(PixelCoord, 0)).r;
+#if HAS_INVERTED_Z_BUFFER
+	DeviceZ = max(0.000000000001, DeviceZ);
+#endif // HAS_INVERTED_Z_BUFFER
+
+	// Clip trace distance
+	float SceneDepth = min(ConvertFromDeviceZ(DeviceZ), MaxTraceDistance);
+	DeviceZ = ConvertToDeviceZ(SceneDepth);
+
+	float Jitter = bJitter ? InterleavedGradientNoise(PixelCoord, View.StateFrameIndexMod8) : 0.0;
+
+	// Intersect ray with bounding volume
+	// TODO: LWC integration..
+	float3 WorldRayOrigin = LWCHackToFloat(LWCSubtract(View.TranslatedWorldCameraOrigin, PrimaryView.PreViewTranslation));
+	float3 WorldRayEnd = LWCHackToFloat(SvPositionToWorld(float4(PixelCoord + 0.5, DeviceZ, 1)));
+	float3 WorldRayDirection = WorldRayEnd - WorldRayOrigin;
+	float WorldRayLength = length(WorldRayDirection);
+	WorldRayDirection /= WorldRayLength;
+
+	float3 LocalRayOrigin = mul(float4(WorldRayOrigin, 1.0), WorldToLocal).xyz;
+	float3 LocalRayEnd = mul(float4(WorldRayEnd, 1.0), WorldToLocal).xyz;
+	float3 LocalRayDirection = LocalRayEnd - LocalRayOrigin;
+	float LocalRayLength = length(LocalRayDirection);
+	LocalRayDirection /= LocalRayLength;
+
+	float3 LocalBoundsMin = LocalBoundsOrigin - LocalBoundsExtent;
+	float3 LocalBoundsMax = LocalBoundsOrigin + LocalBoundsExtent;
+
+	// Test bounding volume
+	float2 HitT = IntersectAABB(LocalRayOrigin, LocalRayDirection, 0.0, SceneDepth, LocalBoundsMin, LocalBoundsMax);
+	float HitSpan = HitT.y - HitT.x;
+	if (HitSpan > 0.0)
+	{
+		// March
+		FRayMarchingContext RayMarchingContext = CreateRayMarchingContext(
+			// Local-space
+			LocalRayOrigin,
+			LocalRayDirection,
+			HitT.x,
+			HitT.y,
+			// World-space
+			WorldRayOrigin,
+			WorldRayDirection,
+			// Ray-step attributes
+			Jitter,
+			StepSize,
+			MaxStepCount,
+			bApplyEmission,
+			bApplyDirectLighting
+		);
+
+		FDeferredLightData LightData = LoadLightData(LightType);
+		uint StepCount = CalcStepCount(RayMarchingContext);
+
+		RayMarchSingleScattering(
+			RayMarchingContext,
+			LightData,
+			LightType,
+			StepCount,
+			Radiance,
+			Transmittance
+		);
+
+		// Output..
+		float3 Opacity = 1.0 - Transmittance;
+		RWLightingTexture[PixelCoord] += float4(Radiance, Luminance(Opacity));
+
+		// Debug output
+		//RWDebugOutputBuffer[LinearIndex] = DebugOutput;
+	}
+}
diff --git a/Engine/Shaders/Private/HeterogeneousVolumes/HeterogeneousVolumesLiveShadingUtils.ush b/Engine/Shaders/Private/HeterogeneousVolumes/HeterogeneousVolumesLiveShadingUtils.ush
new file mode 100644
--- /dev/null
+++ b/Engine/Shaders/Private/HeterogeneousVolumes/HeterogeneousVolumesLiveShadingUtils.ush
@@ -0,0 +1,61 @@
+// Copyright Epic Games, Inc. All Rights Reserved.
+
+FPrimitiveSceneData GetPrimitiveData(
+	FMaterialPixelParameters Parameters,
+	float4x4 LocalToWorld,
+	float4x4 WorldToLocal,
+	float3 LocalBoundsOrigin,
+	float3 LocalBoundsExtent
+)
+{
+	FPrimitiveSceneData PrimitiveData = GetPrimitiveData(Parameters.PrimitiveId);
+
+	// Replace local-space transforms..
+	PrimitiveData.LocalToWorld = MakeLWCMatrix(PrimitiveData.TilePosition, LocalToWorld);
+	PrimitiveData.WorldToLocal = MakeLWCInverseMatrix(PrimitiveData.TilePosition, WorldToLocal);
+	PrimitiveData.LocalObjectBoundsMin = LocalBoundsOrigin - LocalBoundsExtent;
+	PrimitiveData.LocalObjectBoundsMax = LocalBoundsOrigin + LocalBoundsExtent;
+
+	return PrimitiveData;
+}
+
+// Stolen from VolumetricFog
+float3 SampleExtinctionCoefficients(in FPixelMaterialInputs PixelMaterialInputs)
+{
+	float3 Extinction = 0.0f;
+#if STRATA_ENABLED
+	FStrataBSDF BSDF = PixelMaterialInputs.FrontMaterial.InlinedBSDF;
+	Extinction = VOLUMETRICFOGCLOUD_EXTINCTION(BSDF).r;
+#else
+#if !MATERIAL_SHADINGMODEL_UNLIT
+	Extinction = GetMaterialSubsurfaceDataRaw(PixelMaterialInputs).rgb;
+#endif
+#endif
+	return clamp(Extinction, 0.0f, 65000.0f);
+}
+
+float3 SampleEmissive(in FPixelMaterialInputs PixelMaterialInputs)
+{
+	float3 EmissiveColor = 0.0f;
+#if STRATA_ENABLED
+	FStrataBSDF BSDF = PixelMaterialInputs.FrontMaterial.InlinedBSDF;
+	EmissiveColor = BSDF_GETEMISSIVE(BSDF);
+#else
+	EmissiveColor = GetMaterialEmissiveRaw(PixelMaterialInputs);
+#endif
+	return clamp(EmissiveColor, 0.0f, 65000.0f);
+}
+
+float3 SampleAlbedo(in FPixelMaterialInputs PixelMaterialInputs)
+{
+	float3 Albedo = 0.0f;
+#if STRATA_ENABLED
+	FStrataBSDF BSDF = PixelMaterialInputs.FrontMaterial.InlinedBSDF;
+	Albedo = VOLUMETRICFOGCLOUD_ALBEDO(BSDF);
+#else
+#if !MATERIAL_SHADINGMODEL_UNLIT
+	Albedo = GetMaterialBaseColor(PixelMaterialInputs);
+#endif
+#endif
+	return saturate(Albedo) * View.DiffuseOverrideParameter.w + View.DiffuseOverrideParameter.xyz;
+}
diff --git a/Engine/Shaders/Private/HeterogeneousVolumes/HeterogeneousVolumesMaterialBakingPipeline.usf b/Engine/Shaders/Private/HeterogeneousVolumes/HeterogeneousVolumesMaterialBakingPipeline.usf
new file mode 100644
--- /dev/null
+++ b/Engine/Shaders/Private/HeterogeneousVolumes/HeterogeneousVolumesMaterialBakingPipeline.usf
@@ -0,0 +1,63 @@
+// Copyright Epic Games, Inc. All Rights Reserved.
+
+#include "../Common.ush"
+
+#define SUPPORT_CONTACT_SHADOWS 0
+
+#include "/Engine/Generated/Material.ush"
+#include "../PositionReconstructionCommon.ush"
+#include "HeterogeneousVolumesLiveShadingUtils.ush"
+
+#ifndef THREADGROUP_SIZE_3D
+#define THREADGROUP_SIZE_3D 1
+#endif // THREADGROUP_SIZE_3D
+
+// Object data
+float4x4 LocalToWorld;
+float4x4 WorldToLocal;
+float3 LocalBoundsOrigin;
+float3 LocalBoundsExtent;
+int PrimitiveId;
+
+// Volume data
+int3 VolumeResolution;
+
+// Output
+RWTexture3D<float3> RWExtinctionTexture;
+RWTexture3D<float3> RWEmissionTexture;
+RWTexture3D<float3> RWAlbedoTexture;
+
+FPrimitiveSceneData GetPrimitiveData(FMaterialPixelParameters Parameters)
+{
+	return GetPrimitiveData(Parameters, LocalToWorld, WorldToLocal, LocalBoundsOrigin, LocalBoundsExtent);
+}
+
+[numthreads(THREADGROUP_SIZE_3D, THREADGROUP_SIZE_3D, THREADGROUP_SIZE_3D)]
+void HeterogeneousVolumesBakeMaterialCS(
+	uint3 GroupThreadId : SV_GroupThreadID,
+	uint3 DispatchThreadId : SV_DispatchThreadID
+)
+{
+	uint3 VoxelIndex = DispatchThreadId;
+	if (all(VoxelIndex < VolumeResolution))
+	{
+		float3 UVW = (VoxelIndex + 0.5) / float3(VolumeResolution);
+		float3 LocalBoundsMin = LocalBoundsOrigin - LocalBoundsExtent;
+		float3 LocalPosition = UVW * 2.0 * LocalBoundsExtent + LocalBoundsMin;
+		float3 WorldPosition = mul(float4(LocalPosition, 1.0), LocalToWorld).xyz;
+
+		// Setup evaluation context
+		FMaterialPixelParameters MaterialParameters = MakeInitializedMaterialPixelParameters();
+		MaterialParameters.PrimitiveId = PrimitiveId;
+		MaterialParameters.AbsoluteWorldPosition = LWCPromote(WorldPosition);
+
+		// Evaluate material graph
+		FPixelMaterialInputs PixelMaterialInputs;
+		CalcPixelMaterialInputs(MaterialParameters, PixelMaterialInputs);
+
+		// Extract volume rendering coefficients
+		RWExtinctionTexture[VoxelIndex] = SampleExtinctionCoefficients(PixelMaterialInputs);
+		RWEmissionTexture[VoxelIndex] = SampleEmissive(PixelMaterialInputs);
+		RWAlbedoTexture[VoxelIndex] = SampleAlbedo(PixelMaterialInputs);
+	}
+}
\ No newline at end of file
diff --git a/Engine/Shaders/Private/HeterogeneousVolumes/HeterogeneousVolumesPreshadingPipeline.usf b/Engine/Shaders/Private/HeterogeneousVolumes/HeterogeneousVolumesPreshadingPipeline.usf
new file mode 100644
--- /dev/null
+++ b/Engine/Shaders/Private/HeterogeneousVolumes/HeterogeneousVolumesPreshadingPipeline.usf
@@ -0,0 +1,541 @@
+// Copyright Epic Games, Inc. All Rights Reserved.
+
+#include "HeterogeneousVolumesSparseVoxelUtils.ush"
+#include "../Common.ush"
+
+#define SUPPORT_CONTACT_SHADOWS 0
+
+#include "HeterogeneousVolumesSparseVoxelUniformBufferUtils.ush"
+#include "HeterogeneousVolumesTracingUtils.ush"
+#include "HeterogeneousVolumesLightingUtils.ush"
+#include "HeterogeneousVolumesRayMarchingUtils.ush"
+
+#ifndef THREADGROUP_SIZE_1D
+#define THREADGROUP_SIZE_1D 1
+#endif // THREADGROUP_SIZE_1D
+
+#ifndef THREADGROUP_SIZE_2D
+#define THREADGROUP_SIZE_2D 1
+#endif // THREADGROUP_SIZE_2D
+
+#ifndef THREADGROUP_SIZE_3D
+#define THREADGROUP_SIZE_3D 1
+#endif // THREADGROUP_SIZE_3D
+
+#ifndef DIM_SPARSE_VOXEL_TRACING
+#define DIM_SPARSE_VOXEL_TRACING 1
+#endif // DIM_SPARSE_VOXEL_TRACING
+
+#ifndef DIM_VOXEL_CULLING
+#define DIM_VOXEL_CULLING 1
+#endif // DIM_VOXEL_CULLING
+
+// Lighting
+int bApplyEmission;
+int bApplyDirectLighting;
+int LightType;
+
+// Ray data
+float MaxTraceDistance;
+float StepSize;
+int MaxStepCount;
+int bJitter;
+
+// Dispatch data
+int3 GroupCount;
+
+// Volume data
+int MipLevel;
+
+// Output
+RWTexture2D<float4> RWLightingTexture;
+
+int3 TextureResolution;
+Texture3D InputTexture;
+RWTexture3D<float3> RWOutputTexture;
+
+[numthreads(THREADGROUP_SIZE_3D, THREADGROUP_SIZE_3D, THREADGROUP_SIZE_3D)]
+void CopyTexture3D(
+	uint3 DispatchThreadId : SV_DispatchThreadID
+)
+{
+	uint3 Index = DispatchThreadId;
+	if (all(Index < TextureResolution))
+	{
+		RWOutputTexture[Index] = InputTexture[Index];
+	}
+}
+
+SamplerState TextureSampler;
+
+[numthreads(THREADGROUP_SIZE_3D, THREADGROUP_SIZE_3D, THREADGROUP_SIZE_3D)]
+void GenerateMips3D(
+	uint3 GroupThreadId : SV_GroupThreadID,
+	uint3 DispatchThreadId : SV_DispatchThreadID
+)
+{
+	uint3 Index = DispatchThreadId;
+	if (all(Index < TextureResolution))
+	{
+#if 1
+		float3 UVW = (Index + 0.5) / float3(TextureResolution);
+		RWOutputTexture[Index] = InputTexture.SampleLevel(TextureSampler, UVW, 0).rgb;
+#else
+		float3 Value = 0;
+		UNROLL for (uint i = 0; i < 8; ++i)
+		{
+			uint3 BaseIndex = Index * 2 + float3((i & 4u) >> 2, (i & 2u) >> 1, (i & 1u));
+			BaseIndex = min(BaseIndex, TextureResolution * 2);
+			Value += InputTexture[BaseIndex];
+		}
+		Value *= 0.125;
+
+		RWOutputTexture[Index] = Value;
+#endif
+	}
+}
+
+[numthreads(THREADGROUP_SIZE_3D, THREADGROUP_SIZE_3D, THREADGROUP_SIZE_3D)]
+void GenerateMin3D(
+	uint3 GroupThreadId : SV_GroupThreadID,
+	uint3 DispatchThreadId : SV_DispatchThreadID
+)
+{
+	uint3 Index = DispatchThreadId;
+	if (all(Index < TextureResolution))
+	{
+		float3 Value = 65535.0f;
+		UNROLL for (uint i = 0; i < 8; ++i)
+		{
+			int3 BaseIndex = Index * 2 + float3((i & 4u) >> 2, (i & 2u) >> 1, (i & 1u));
+			BaseIndex = min(BaseIndex, TextureResolution * 2);
+			Value = min(Value, Luminance(InputTexture[BaseIndex]));
+		}
+
+		RWOutputTexture[Index] = Value;
+	}
+}
+
+StructuredBuffer<FVoxelDataPacked> VoxelBuffer;
+Buffer<uint> NumVoxelsBuffer;
+
+// A screen-tile is an 8x8 block of computation for ray marching
+struct FRayMarchingTile
+{
+	uint2 PixelOffset;
+	uint Voxels[2];
+
+	// TODO: For tracking purposes..
+	uint Id;
+	uint Padding[3];
+};
+
+FRayMarchingTile CreateRayMarchingTile(uint TileId, uint2 PixelOffset, uint VoxelMin, uint VoxelMax)
+{
+	FRayMarchingTile Tile = (FRayMarchingTile) 0;
+	Tile.PixelOffset = PixelOffset;
+
+	Tile.Voxels[0] = VoxelMin;
+	Tile.Voxels[1] = VoxelMax;
+
+	Tile.Id = TileId;
+	return Tile;
+}
+
+RWStructuredBuffer<FRayMarchingTile> RWRayMarchingTilesBuffer;
+RWBuffer<uint> RWNumRayMarchingTilesBuffer;
+
+RWStructuredBuffer<FVoxelDataPacked> RWVoxelsPerTileBuffer;
+RWBuffer<uint> RWNumVoxelsPerTileBuffer;
+
+struct FRayMarchingDebug
+{
+	float4 Planes[5];
+	float4 BBox[2];
+
+	float Padding[4];
+};
+
+RWStructuredBuffer<FRayMarchingDebug> RWRayMarchingDebugBuffer;
+
+float4 CreatePlane(
+	float3 P_000,
+	float3 P_100,
+	float3 P_010
+)
+{
+	float3 X = P_100 - P_000;
+	float3 Y = P_010 - P_000;
+
+	float3 N = normalize(cross(X, Y));
+	float d = -dot(N, P_000);
+	return float4(N, d);
+}
+
+float CalcSignedDistance(float3 P, float4 Plane)
+{
+	return dot(P, Plane.xyz) + Plane.w;
+}
+
+bool PointIntersectsFrustum(float3 P, float4 Planes[5])
+{
+	for (int i = 0; i < 5; ++i)
+	{
+		if (CalcSignedDistance(P, Planes[i]) > 0.0)
+		{
+			return false;
+		}
+	}
+
+	return true;
+}
+
+bool BoxIntersectsFrustum(float3 BoxOrigin, float3 BoxExtent, float4 Planes[5])
+{
+	for (int i = 0; i < 5; ++i)
+	{
+		float Distance = CalcSignedDistance(BoxOrigin, Planes[i]);
+		float ProjectedExtent = dot(BoxExtent, abs(Planes[i].xyz));
+
+		if (Distance > ProjectedExtent)
+		{
+			return false;
+		}
+	}
+
+	return true;
+}
+
+[numthreads(THREADGROUP_SIZE_2D, THREADGROUP_SIZE_2D, 1)]
+void GenerateRayMarchingTiles(
+	uint2 DispatchThreadId : SV_DispatchThreadID
+)
+{
+	if (!any(DispatchThreadId))
+	{
+		RWNumRayMarchingTilesBuffer[1] = 1;
+		RWNumRayMarchingTilesBuffer[2] = 1;
+	}
+
+	uint4 Viewport = uint4(DispatchThreadId.xy * THREADGROUP_SIZE_2D, (DispatchThreadId.xy + 1) * THREADGROUP_SIZE_2D);
+	uint2 PixelOffset = Viewport.xy;
+	if (any(PixelOffset >= View.ViewSizeAndInvSize.xy))
+	{
+		return;
+	}
+
+	// Build tile-frustum
+	float Epsilon = 1.0e-6;
+#if HAS_INVERTED_Z_BUFFER
+	float NearDepth = 1.0;
+	float FarDepth = Epsilon;
+#else
+	float NearDepth = 0.0;
+	float FarDepth = 1.0 - Epsilon;
+#endif // HAS_INVERTED_Z_BUFFER
+
+	float4x4 WorldToLocal = GetWorldToLocal();
+	float3 W_000 = LWCHackToFloat(SvPositionToWorld(float4(Viewport.xw, NearDepth, 1)));
+	float3 W_100 = LWCHackToFloat(SvPositionToWorld(float4(Viewport.zw, NearDepth, 1)));
+	float3 W_010 = LWCHackToFloat(SvPositionToWorld(float4(Viewport.xy, NearDepth, 1)));
+	float3 W_110 = LWCHackToFloat(SvPositionToWorld(float4(Viewport.zy, NearDepth, 1)));
+	float3 W_001 = LWCHackToFloat(SvPositionToWorld(float4(Viewport.xw, FarDepth, 1)));
+	float3 W_101 = LWCHackToFloat(SvPositionToWorld(float4(Viewport.zw, FarDepth, 1)));
+	float3 W_011 = LWCHackToFloat(SvPositionToWorld(float4(Viewport.xy, FarDepth, 1)));
+	float3 W_111 = LWCHackToFloat(SvPositionToWorld(float4(Viewport.zy, FarDepth, 1)));
+
+	float3 L_000 = mul(float4(W_000, 1), WorldToLocal).xyz;
+	float3 L_100 = mul(float4(W_100, 1), WorldToLocal).xyz;
+	float3 L_010 = mul(float4(W_010, 1), WorldToLocal).xyz;
+	float3 L_110 = mul(float4(W_110, 1), WorldToLocal).xyz;
+	float3 L_001 = mul(float4(W_001, 1), WorldToLocal).xyz;
+	float3 L_101 = mul(float4(W_101, 1), WorldToLocal).xyz;
+	float3 L_011 = mul(float4(W_011, 1), WorldToLocal).xyz;
+	float3 L_111 = mul(float4(W_111, 1), WorldToLocal).xyz;
+
+	float4 Planes[5] = {
+		// Front
+		CreatePlane(L_000, L_010, L_100),
+		// Back
+		//CreatePlane(L_101, L_001, L_111),
+		// +X
+		CreatePlane(L_100, L_110, L_101),
+		// +Y
+		CreatePlane(L_010, L_011, L_110),
+		// -X
+		CreatePlane(L_000, L_001, L_010),
+		// -Y
+		CreatePlane(L_000, L_100, L_001)
+	};
+
+	uint DebugIndex = DispatchThreadId.y * THREADGROUP_SIZE_2D + DispatchThreadId.x;
+#if DIM_DEBUG
+	for (uint i = 0; i < 5; ++i)
+	{
+		RWRayMarchingDebugBuffer[DebugIndex].Planes[i] = Planes[i];
+	}
+
+	RWRayMarchingDebugBuffer[DebugIndex].BBox[0] = float4(GetLocalBoundsOrigin() - GetLocalBoundsExtent(), 0.0);
+	RWRayMarchingDebugBuffer[DebugIndex].BBox[1] = float4(GetLocalBoundsOrigin() + GetLocalBoundsExtent(), 0.0);
+#endif // DIM_DEBUG
+
+	// Intersect frustum against bounding volume
+	if (!BoxIntersectsFrustum(GetLocalBoundsOrigin(), GetLocalBoundsExtent(), Planes))
+	{
+		return;
+	}
+
+#if DIM_VOXEL_CULLING
+	float3 LocalBoundsMin = GetLocalBoundsOrigin() - GetLocalBoundsExtent();
+	float3 LocalBoundsMax = GetLocalBoundsOrigin() + GetLocalBoundsExtent();
+
+	// Intersect against voxels to determine allocation
+	uint VoxelCount = 0;
+	for (uint VoxelIndex = 0; VoxelIndex < GetNumVoxels(); ++VoxelIndex)
+	{
+		FVoxelData VoxelData = DecodeVoxelData(GetVoxelDataPacked(VoxelIndex), GetVolumeResolution());
+		FBox VoxelBoundingBox = GetVoxelBoundingBox(VoxelData, GetVolumeResolution(), LocalBoundsMin, LocalBoundsMax);
+
+		float3 BoxOrigin = (VoxelBoundingBox.Max + VoxelBoundingBox.Min) / 2.0;
+		float3 BoxExtent = VoxelBoundingBox.Max - VoxelBoundingBox.Min;
+		if (BoxIntersectsFrustum(BoxOrigin, BoxExtent, Planes))
+		{
+			VoxelCount++;
+		}
+	}
+
+	// Allocate contiguous span of voxel indices
+	uint VoxelEntryIndex;
+	if (VoxelCount > 0)
+	{
+		InterlockedAdd(RWNumVoxelsPerTileBuffer[0], VoxelCount, VoxelEntryIndex);
+
+		// Intersect and store voxels
+		VoxelCount = 0;
+		for (uint VoxelIndex = 0; VoxelIndex < GetNumVoxels(); ++VoxelIndex)
+		{
+			FVoxelData VoxelData = DecodeVoxelData(GetVoxelDataPacked(VoxelIndex), GetVolumeResolution());
+			FBox VoxelBoundingBox = GetVoxelBoundingBox(VoxelData, GetVolumeResolution(), LocalBoundsMin, LocalBoundsMax);
+
+			float3 BoxOrigin = (VoxelBoundingBox.Max + VoxelBoundingBox.Min) / 2.0;
+			float3 BoxExtent = VoxelBoundingBox.Max - VoxelBoundingBox.Min;
+			if (BoxIntersectsFrustum(BoxOrigin, BoxExtent, Planes))
+			{
+				RWVoxelsPerTileBuffer[VoxelEntryIndex + VoxelCount] = GetVoxelDataPacked(VoxelIndex);
+				VoxelCount++;
+			}
+		}
+	}
+#else
+	uint VoxelEntryIndex = 0;
+	uint VoxelCount = GetNumVoxels();
+#endif
+
+	// Allocate tile
+	if (VoxelCount > 0)
+	{
+		uint TileId;
+		InterlockedAdd(RWNumRayMarchingTilesBuffer[0], 1, TileId);
+
+		RWRayMarchingTilesBuffer[TileId] = CreateRayMarchingTile(TileId, PixelOffset, VoxelEntryIndex, VoxelEntryIndex + VoxelCount);
+	}
+}
+
+struct FTraceVoxelResult
+{
+	FVoxelData VoxelData;
+	float2 HitT;
+};
+
+FTraceVoxelResult TraceVoxels(float3 LocalRayOrigin, float3 LocalRayDirection, float LocalRayTMin, float LocalRayTMax,
+	float3 LocalBoundsMin, float3 LocalBoundsMax)
+{
+	FTraceVoxelResult TraceVoxelResult;
+	TraceVoxelResult.HitT.x = LocalRayTMax;
+	TraceVoxelResult.HitT.y = LocalRayTMax;
+
+	for (uint VoxelIndex = 0; VoxelIndex < GetNumVoxels(); ++VoxelIndex)
+	{
+		FVoxelData VoxelData = DecodeVoxelData(GetVoxelDataPacked(VoxelIndex), GetVolumeResolution());
+
+		FBox VoxelBoundingBox = GetVoxelBoundingBox(VoxelData, GetVolumeResolution(), LocalBoundsMin, LocalBoundsMax);
+		float2 HitT = IntersectAABB(LocalRayOrigin, LocalRayDirection, LocalRayTMin, LocalRayTMax, VoxelBoundingBox.Min, VoxelBoundingBox.Max);
+
+		bool bIsHit = (HitT.y - HitT.x) > 0.0;
+		if (bIsHit && (HitT.x < TraceVoxelResult.HitT.x))
+		{
+			TraceVoxelResult.HitT = HitT;
+			TraceVoxelResult.VoxelData = VoxelData;
+		}
+	}
+
+	return TraceVoxelResult;
+}
+
+StructuredBuffer<FVoxelDataPacked> VoxelsPerTileBuffer;
+
+FTraceVoxelResult TraceVoxels(float3 LocalRayOrigin, float3 LocalRayDirection, float LocalRayTMin, float LocalRayTMax,
+	float3 LocalBoundsMin, float3 LocalBoundsMax,
+	uint VoxelStartIndex, uint VoxelEndIndex)
+{
+	FTraceVoxelResult TraceVoxelResult;
+	TraceVoxelResult.HitT.x = LocalRayTMax;
+	TraceVoxelResult.HitT.y = LocalRayTMax;
+
+	for (uint VoxelIndex = VoxelStartIndex; VoxelIndex < VoxelEndIndex; ++VoxelIndex)
+	{
+		FVoxelData VoxelData = DecodeVoxelData(VoxelsPerTileBuffer[VoxelIndex], GetVolumeResolution());
+
+		FBox VoxelBoundingBox = GetVoxelBoundingBox(VoxelData, GetVolumeResolution(), LocalBoundsMin, LocalBoundsMax);
+		float2 HitT = IntersectAABB(LocalRayOrigin, LocalRayDirection, LocalRayTMin, LocalRayTMax, VoxelBoundingBox.Min, VoxelBoundingBox.Max);
+
+		bool bIsHit = (HitT.y - HitT.x) > 0.0;
+		if (bIsHit && (HitT.x < TraceVoxelResult.HitT.x))
+		{
+			TraceVoxelResult.HitT = HitT;
+			TraceVoxelResult.VoxelData = VoxelData;
+		}
+	}
+
+	return TraceVoxelResult;
+}
+
+RWStructuredBuffer<FVoxelDataPacked> RWVoxelOutputBuffer;
+
+StructuredBuffer<FRayMarchingTile> RayMarchingTilesBuffer;
+Buffer<uint> NumRayMarchingTilesBuffer;
+
+[numthreads(THREADGROUP_SIZE_1D, 1, 1)]
+void RenderSingleScatteringWithPreshadingCS(
+	uint GroupId : SV_GroupID,
+	uint GroupThreadId : SV_GroupThreadID
+)
+{
+	float3 Radiance = 0.0;
+	float3 Transmittance = 1.0;
+
+	// Convert 1D index to pixel coordinate
+	FRayMarchingTile RayMarchingTile = RayMarchingTilesBuffer[GroupId];
+	// TODO: Morton..
+	uint2 TileCoord = uint2(GroupThreadId % THREADGROUP_SIZE_2D, GroupThreadId / THREADGROUP_SIZE_2D);
+	//float2 PixelCoord = View.ViewRectMin.xy + RayMarchingTile.PixelOffset + TileCoord;
+	float2 PixelCoord = RayMarchingTile.PixelOffset + TileCoord;
+
+	// Create screen ray
+	if (any(PixelCoord.xy >= View.ViewSizeAndInvSize.xy))
+	{
+		return;
+	}
+	PixelCoord += View.ViewRectMin.xy;
+
+	// Extract depth
+	float DeviceZ = SceneDepthTexture.Load(int3(PixelCoord, 0)).r;
+#if HAS_INVERTED_Z_BUFFER
+	DeviceZ = max(0.000000000001, DeviceZ);
+#endif // HAS_INVERTED_Z_BUFFER
+
+	// Clip trace distance
+	float SceneDepth = min(ConvertFromDeviceZ(DeviceZ), MaxTraceDistance);
+	DeviceZ = ConvertToDeviceZ(SceneDepth);
+
+	float Jitter = bJitter ? InterleavedGradientNoise(PixelCoord, View.StateFrameIndexMod8) : 0.0;
+
+	// Intersect ray with bounding volume
+	// TODO: LWC integration..
+	float3 WorldRayOrigin = LWCHackToFloat(LWCSubtract(View.TranslatedWorldCameraOrigin, PrimaryView.PreViewTranslation));
+	float3 WorldRayEnd = LWCHackToFloat(SvPositionToWorld(float4(PixelCoord + 0.5, DeviceZ, 1)));
+	float3 WorldRayDirection = WorldRayEnd - WorldRayOrigin;
+	float WorldRayLength = length(WorldRayDirection);
+	WorldRayDirection /= WorldRayLength;
+
+	float3 LocalRayOrigin = mul(float4(WorldRayOrigin, 1.0), GetWorldToLocal()).xyz;
+	float3 LocalRayEnd = mul(float4(WorldRayEnd, 1.0), GetWorldToLocal()).xyz;
+	float3 LocalRayDirection = LocalRayEnd - LocalRayOrigin;
+	float LocalRayLength = length(LocalRayDirection);
+	LocalRayDirection /= LocalRayLength;
+
+	float3 LocalBoundsMin = GetLocalBoundsOrigin() - GetLocalBoundsExtent();
+	float3 LocalBoundsMax = GetLocalBoundsOrigin() + GetLocalBoundsExtent();
+
+	// Test bounding volume
+	float2 HitT = IntersectAABB(LocalRayOrigin, LocalRayDirection, 0.0, SceneDepth, LocalBoundsMin, LocalBoundsMax);
+	float HitSpan = HitT.y - HitT.x;
+	if (HitSpan <= 0.0)
+	{
+		return;
+	}
+
+#if DIM_SPARSE_VOXEL_TRACING
+	float RayTMax = HitT.y;
+
+	#if DIM_VOXEL_CULLING
+		HitT = TraceVoxels(LocalRayOrigin, LocalRayDirection, 0.0, RayTMax,
+			LocalBoundsMin, LocalBoundsMax,
+			RayMarchingTile.Voxels[0], RayMarchingTile.Voxels[1]).HitT;
+	#else
+		HitT = TraceVoxels(LocalRayOrigin, LocalRayDirection, 0.0, RayTMax, LocalBoundsMin, LocalBoundsMax).HitT;
+	#endif // DIM_VOXEL_CULLING
+
+	HitSpan = HitT.y - HitT.x;
+
+	while (HitSpan > 0.0)
+#endif // DIM_SPARSE_VOXEL_TRACING
+	{
+		FRayMarchingContext RayMarchingContext = CreateRayMarchingContext(
+			// Local-space
+			LocalRayOrigin,
+			LocalRayDirection,
+			HitT.x,
+			HitT.y,
+			// World-space
+			WorldRayOrigin,
+			WorldRayDirection,
+			// Ray-step attributes
+			Jitter,
+			StepSize,
+			MaxStepCount,
+			bApplyEmission,
+			bApplyDirectLighting
+		);
+
+		FDeferredLightData LightData = LoadLightData(LightType);
+		uint StepCount = CalcStepCount(RayMarchingContext);
+
+		RayMarchSingleScattering(
+			RayMarchingContext,
+			LightData,
+			LightType,
+			StepCount,
+			Radiance,
+			Transmittance
+		);
+
+#if DIM_SPARSE_VOXEL_TRACING
+	#if DIM_VOXEL_CULLING
+		HitT = TraceVoxels(LocalRayOrigin, LocalRayDirection, HitT.y + 0.01, RayTMax,
+			LocalBoundsMin, LocalBoundsMax,
+			RayMarchingTile.Voxels[0], RayMarchingTile.Voxels[1]).HitT;
+	#else
+		HitT = TraceVoxels(LocalRayOrigin, LocalRayDirection, HitT.y + 0.0001, RayTMax, LocalBoundsMin, LocalBoundsMax).HitT;
+	#endif // DIM_VOXEL_CULLING
+
+		HitSpan = HitT.y - HitT.x;
+#endif // DIM_SPARSE_VOXEL_TRACING
+	}
+
+	// Output..
+	float3 Opacity = 1.0 - Transmittance;
+	RWLightingTexture[PixelCoord] += float4(Radiance, Luminance(Opacity));
+
+#if DIM_DEBUG
+	//if (!all(DispatchThreadId))
+	if (GroupId == 0 && GroupThreadId == 0)
+	{
+		RWVoxelOutputBuffer[0].MipLevel = GetNumVoxels();
+		for (uint VoxelIndex = 0; VoxelIndex < GetNumVoxels(); ++VoxelIndex)
+		{
+			RWVoxelOutputBuffer[VoxelIndex + 1] = GetVoxelDataPacked(VoxelIndex);
+		}
+	}
+#endif
+}
diff --git a/Engine/Shaders/Private/HeterogeneousVolumes/HeterogeneousVolumesRayMarchingUtils.ush b/Engine/Shaders/Private/HeterogeneousVolumes/HeterogeneousVolumesRayMarchingUtils.ush
new file mode 100644
--- /dev/null
+++ b/Engine/Shaders/Private/HeterogeneousVolumes/HeterogeneousVolumesRayMarchingUtils.ush
@@ -0,0 +1,243 @@
+// Copyright Epic Games, Inc. All Rights Reserved.
+
+#ifndef HARDWARE_RAY_TRACING
+#define HARDWARE_RAY_TRACING 0
+#endif // HARDWARE_RAY_TRACING
+
+struct FRayMarchingContext
+{
+	// Local-space
+	float3 LocalRayOrigin;
+	float3 LocalRayDirection;
+	float LocalRayTMin;
+	float LocalRayTMax;
+	// World-space
+	float3 WorldRayOrigin;
+	float3 WorldRayDirection;
+	// Ray-step attributes
+	float Jitter;
+	float StepSize;
+	uint MaxStepCount;
+	// Lighting
+	int bApplyEmission;
+	int bApplyDirectLighting;
+};
+
+FRayMarchingContext CreateRayMarchingContext(
+	// Local-space
+	float3 LocalRayOrigin,
+	float3 LocalRayDirection,
+	float LocalRayTMin,
+	float LocalRayTMax,
+	// World-space
+	float3 WorldRayOrigin,
+	float3 WorldRayDirection,
+	// Ray-step attributes
+	float Jitter,
+	float StepSize,
+	uint MaxStepCount,
+	// Lighting
+	int bApplyEmission,
+	int bApplyDirectLighting
+)
+{
+	FRayMarchingContext Context;
+	// Local-space
+	Context.LocalRayOrigin = LocalRayOrigin;
+	Context.LocalRayDirection = LocalRayDirection;
+	Context.LocalRayTMin = LocalRayTMin;
+	Context.LocalRayTMax = LocalRayTMax;
+	// World-space
+	Context.WorldRayOrigin = WorldRayOrigin;
+	Context.WorldRayDirection = WorldRayDirection;
+	// Ray-step attributes
+	Context.Jitter = Jitter;
+	Context.StepSize = StepSize;
+	Context.MaxStepCount = MaxStepCount;
+	// Lighting
+	Context.bApplyEmission = bApplyEmission;
+	Context.bApplyDirectLighting = bApplyDirectLighting;
+
+	return Context;
+}
+
+uint CalcStepCount(inout FRayMarchingContext Context)
+{
+	float HitSpan = Context.LocalRayTMax - Context.LocalRayTMin;
+	float StepSize = Context.StepSize;
+	uint StepCount = max(ceil(HitSpan / StepSize), 1);
+	if (StepCount > Context.MaxStepCount)
+	{
+		StepCount = Context.MaxStepCount;
+		StepSize = HitSpan / StepCount;
+	}
+
+	Context.StepSize = StepSize;
+	return StepCount;
+}
+
+void RayMarchTransmittance(
+	inout FRayMarchingContext RayMarchingContext,
+	uint StepCount,
+	inout float3 Transmittance
+)
+{
+	for (uint StepIndex = 0; StepIndex < StepCount; ++StepIndex)
+	{
+		float3 LocalPosition = RayMarchingContext.LocalRayOrigin + RayMarchingContext.LocalRayDirection * (RayMarchingContext.LocalRayTMin + RayMarchingContext.StepSize * StepIndex);
+		float3 WorldPosition = RayMarchingContext.WorldRayOrigin + RayMarchingContext.WorldRayDirection * (RayMarchingContext.LocalRayTMin + RayMarchingContext.StepSize * StepIndex);
+
+		FVolumeSampleContext SampleContext = CreateVolumeSampleContext(LocalPosition, WorldPosition);
+		float3 Extinction = SampleExtinction(SampleContext);
+		Transmittance *= exp(-Extinction * RayMarchingContext.StepSize);
+	}
+}
+
+float3 ComputeTransmittance(
+	inout FRayMarchingContext RayMarchingContext
+)
+{
+#if HARDWARE_RAY_TRACING
+	return ComputeTransmittanceHardwareRayTracing(RayMarchingContext.WorldRayOrigin, RayMarchingContext.WorldRayDirection, RayMarchingContext.LocalRayTMin, RayMarchingContext.LocalRayTMax);
+#else // HARDWARE_RAY_TRACING
+
+	float3 Transmittance = 1.0;
+
+	float3 LocalBoundsMin = GetLocalBoundsOrigin() - GetLocalBoundsExtent();
+	float3 LocalBoundsMax = GetLocalBoundsOrigin() + GetLocalBoundsExtent();
+	float2 HitT = IntersectAABB(RayMarchingContext.LocalRayOrigin, RayMarchingContext.LocalRayDirection, RayMarchingContext.LocalRayTMin, RayMarchingContext.LocalRayTMax,
+		LocalBoundsMin, LocalBoundsMax);
+
+	float HitSpan = HitT.y - HitT.x;
+	if (HitSpan > 0.0)
+	{
+		RayMarchingContext.LocalRayTMin = HitT.x;
+		RayMarchingContext.LocalRayTMax = HitT.y;
+		uint StepCount = CalcStepCount(RayMarchingContext);
+
+		RayMarchTransmittance(RayMarchingContext, StepCount, Transmittance);
+	}
+
+	return Transmittance;
+#endif // HARDWARE_RAY_TRACING
+}
+
+void RayMarchEmissionAbsorption(
+	inout FRayMarchingContext RayMarchingContext,
+	uint StepCount,
+	inout float3 Radiance,
+	inout float3 Transmittance
+)
+{
+	for (uint StepIndex = 0; StepIndex < StepCount; ++StepIndex)
+	{
+		float3 LocalPosition = RayMarchingContext.LocalRayOrigin + RayMarchingContext.LocalRayDirection * (RayMarchingContext.LocalRayTMin + RayMarchingContext.StepSize * StepIndex);
+		float3 WorldPosition = RayMarchingContext.WorldRayOrigin + RayMarchingContext.WorldRayDirection * (RayMarchingContext.LocalRayTMin + RayMarchingContext.StepSize * StepIndex);
+
+		FVolumeSampleContext SampleContext = CreateVolumeSampleContext(LocalPosition, WorldPosition);
+		float3 Extinction = SampleExtinction(SampleContext);
+
+		float3 Emission = SampleEmission(SampleContext);
+		Radiance += Emission * RayMarchingContext.StepSize * Transmittance;
+
+		// Accumulate transmittance for the next evaluation
+		Transmittance *= exp(-Extinction * RayMarchingContext.StepSize);
+	}
+}
+
+void RayMarchSingleScattering(
+	inout FRayMarchingContext RayMarchingContext,
+	inout FDeferredLightData LightData,
+	uint LightType,
+	uint StepCount,
+	inout float3 Radiance,
+	inout float3 Transmittance
+)
+{
+	for (uint StepIndex = 0; StepIndex < StepCount; ++StepIndex)
+	{
+		float3 LocalPosition = RayMarchingContext.LocalRayOrigin + RayMarchingContext.LocalRayDirection * (RayMarchingContext.LocalRayTMin + RayMarchingContext.StepSize * StepIndex);
+		float3 WorldPosition = RayMarchingContext.WorldRayOrigin + RayMarchingContext.WorldRayDirection * (RayMarchingContext.LocalRayTMin + RayMarchingContext.StepSize * StepIndex);
+		FVolumeSampleContext SampleContext = CreateVolumeSampleContext(LocalPosition, WorldPosition);
+
+		if (RayMarchingContext.bApplyEmission)
+		{
+			float3 Emission = SampleEmission(SampleContext);
+			Radiance += Emission * RayMarchingContext.StepSize * Transmittance;
+		}
+
+		float3 Extinction = SampleExtinction(SampleContext);
+		if (RayMarchingContext.bApplyDirectLighting && (any(Extinction > 0.0)))
+		{
+			float3 Albedo = SampleAlbedo(SampleContext);
+			if (any(Albedo > 0.0))
+			{
+				float3 TranslatedWorldPosition = LWCHackToFloat(LWCAdd(WorldPosition, PrimaryView.PreViewTranslation));
+
+				float3 L = LightData.Direction;
+				//float3 ToLight = L * RayMarchingContext.MaxShadowTraceDistance;
+				float3 ToLight = L * 10000;
+
+				float LightAttenuation = 1.0;
+				if (LightType != LIGHT_TYPE_DIRECTIONAL)
+				{
+					LightAttenuation = GetLocalLightAttenuation(TranslatedWorldPosition, LightData, ToLight, L);
+					if (LightData.bRectLight)
+					{
+						FRect Rect = GetRect(ToLight, LightData);
+						LightAttenuation *= IntegrateLight(Rect);
+					}
+					else
+					{
+						FCapsuleLight Capsule = GetCapsule(ToLight, LightData);
+						Capsule.DistBiasSqr = 0;
+						LightAttenuation *= IntegrateLight(Capsule, LightData.bInverseSquared);
+					}
+				}
+
+				float3 Inscattering = LightData.Color * LightAttenuation;
+
+				if (any(Inscattering > 0.0))
+				{
+					// Build shadow ray and find intersection
+					float3 LocalShadowRayOrigin = mul(float4(WorldPosition, 1.0), GetWorldToLocal()).xyz;
+					float3 LocalShadowRayEnd = mul(float4(WorldPosition + ToLight, 1.0), GetWorldToLocal()).xyz;
+					float3 LocalShadowRayDirection = LocalShadowRayEnd - LocalShadowRayOrigin;
+					float LocalShadowRayTMin = 0.0;
+					float LocalShadowRayTMax = length(LocalShadowRayDirection);
+					LocalShadowRayDirection /= LocalShadowRayTMax;
+
+					//float ShadowJitter = InterleavedGradientNoise(PixelCoord, StepIndex);
+					float ShadowJitter = 0.0;
+					float ShadowStepSize = GetStepSize() * GetShadowStepFactor();
+					int bApplyEmission = 0;
+					int bApplyDirectLighting = 0;
+
+					FRayMarchingContext ShadowRayMarchingContext = CreateRayMarchingContext(
+						LocalShadowRayOrigin,
+						LocalShadowRayDirection,
+						LocalShadowRayTMin,
+						LocalShadowRayTMax,
+						WorldPosition, // WorldShadowRayOrigin
+						normalize(ToLight), // WorldShadowRayDirection
+						ShadowJitter,
+						ShadowStepSize,
+						RayMarchingContext.MaxStepCount,
+						bApplyEmission,
+						bApplyDirectLighting
+						//RayMarchingContext.MaxShadowTraceDistance
+					);
+					float3 ShadowTransmittance = ComputeTransmittance(ShadowRayMarchingContext);
+
+					// Convert in-scattering to exitant radiance
+					float3 ScatteringCoefficient = Albedo * Extinction;
+					float IsotropicPhase = 1.0 / (4.0 * PI);
+					Radiance += Inscattering * ShadowTransmittance * ScatteringCoefficient * IsotropicPhase * Transmittance * RayMarchingContext.StepSize;
+				}
+			}
+		}
+
+		// Accumulate transmittance for the next evaluation
+		Transmittance *= exp(-Extinction * RayMarchingContext.StepSize);
+	}
+}
diff --git a/Engine/Shaders/Private/HeterogeneousVolumes/HeterogeneousVolumesSparseVoxelPipeline.usf b/Engine/Shaders/Private/HeterogeneousVolumes/HeterogeneousVolumesSparseVoxelPipeline.usf
new file mode 100644
--- /dev/null
+++ b/Engine/Shaders/Private/HeterogeneousVolumes/HeterogeneousVolumesSparseVoxelPipeline.usf
@@ -0,0 +1,214 @@
+// Copyright Epic Games, Inc. All Rights Reserved.
+
+#include "HeterogeneousVolumesSparseVoxelUtils.ush"
+#include "../Common.ush"
+
+#include "HeterogeneousVolumesSparseVoxelUniformBufferUtils.ush"
+
+#ifndef THREADGROUP_SIZE_1D
+#define THREADGROUP_SIZE_1D 1
+#endif // THREADGROUP_SIZE_1D
+
+#ifndef THREADGROUP_SIZE_2D
+#define THREADGROUP_SIZE_2D 1
+#endif // THREADGROUP_SIZE_2D
+
+#ifndef THREADGROUP_SIZE_3D
+#define THREADGROUP_SIZE_3D 1
+#endif // THREADGROUP_SIZE_3D
+
+RWBuffer<uint> RWAllocatorBuffer;
+
+[numthreads(1, 1, 1)]
+void ClearAllocator(
+	uint DispatchThreadId : SV_DispatchThreadID
+)
+{
+	RWAllocatorBuffer[0] = 0;
+	RWAllocatorBuffer[1] = 1;
+	RWAllocatorBuffer[2] = 1;
+}
+
+Texture3D VoxelMinTexture;
+uint3 VoxelMinResolution;
+
+// Volume data
+int3 VolumeResolution;
+int MipLevel;
+
+RWStructuredBuffer<FVoxelDataPacked> RWVoxelBuffer;
+RWBuffer<uint> RWNumVoxelsBuffer;
+
+[numthreads(THREADGROUP_SIZE_3D, THREADGROUP_SIZE_3D, THREADGROUP_SIZE_3D)]
+void AllocateSparseVoxels(
+	uint3 GroupThreadId : SV_GroupThreadID,
+	uint3 DispatchThreadId : SV_DispatchThreadID
+)
+{
+	uint3 VoxelIndex = DispatchThreadId;
+	if (all(VoxelIndex < VoxelMinResolution))
+	{
+		float3 Value = VoxelMinTexture.Load(int4(VoxelIndex, 0));
+		if (dot(Value, Value) > 0.0)
+		{
+			// Allocate
+			uint EntryIndex;
+			InterlockedAdd(RWNumVoxelsBuffer[0], 1, EntryIndex);
+
+			// Assign
+			RWVoxelBuffer[EntryIndex] = EncodeVoxelData(VoxelIndex, MipLevel, VoxelMinResolution);
+		}
+	}
+}
+
+RWStructuredBuffer<FVoxelDataPacked> RWVoxelsToRefineBuffer;
+RWBuffer<uint> RWNumVoxelsToRefineBuffer;
+
+// TODO: Operate at lower MIP, using wave intrinsics to broadcast non-zero 
+[numthreads(THREADGROUP_SIZE_3D, THREADGROUP_SIZE_3D, THREADGROUP_SIZE_3D)]
+void AllocateSparseVoxelsToRefine(
+	uint3 GroupThreadId : SV_GroupThreadID,
+	uint3 DispatchThreadId : SV_DispatchThreadID
+)
+{
+	RWNumVoxelsToRefineBuffer[1] = 1;
+	RWNumVoxelsToRefineBuffer[2] = 1;
+
+	uint3 VoxelIndex = DispatchThreadId;
+	if (all(VoxelIndex < VoxelMinResolution) && !any(VoxelIndex % 2))
+	{
+		// Load siblings
+		uint ValidCount = 0;
+		for (uint i = 0; i < 8; ++i)
+		{
+			uint3 Offset = uint3((i & 0x04) >> 2, (i & 0x02) >> 1, i & 0x01);
+			float3 Value = VoxelMinTexture.Load(int4(VoxelIndex + Offset, 0));
+			if (dot(Value, Value) > 0.0)
+			{
+				ValidCount++;
+			}
+		}
+
+		if (ValidCount == 8)
+		{
+			// Emit parent voxel for further refinement
+			uint EntryIndex;
+			InterlockedAdd(RWNumVoxelsToRefineBuffer[0], 1, EntryIndex);
+			RWVoxelsToRefineBuffer[EntryIndex] = EncodeVoxelData(VoxelIndex / 2, MipLevel + 1, VoxelMinResolution / 2);
+		}
+		else
+		{
+			// Emit siblings as final voxels
+			uint EntryIndex;
+			uint SiblingIndex = 0;
+			InterlockedAdd(RWNumVoxelsBuffer[0], ValidCount, EntryIndex);
+			for (uint i = 0; i < 8; ++i)
+			{
+				uint3 Offset = uint3((i & 0x04) >> 2, (i & 0x02) >> 1, i & 0x01);
+				float3 Value = VoxelMinTexture.Load(int4(VoxelIndex + Offset, 0));
+				if (dot(Value, Value) > 0.0)
+				{
+					RWVoxelBuffer[EntryIndex + SiblingIndex++] = EncodeVoxelData(VoxelIndex + Offset, MipLevel, VoxelMinResolution);
+				}
+			}
+		}
+	}
+}
+
+StructuredBuffer<FVoxelDataPacked> VoxelsToRefineBuffer;
+Buffer<uint> NumVoxelsToRefineBuffer;
+uint3 MipVolumeResolution;
+
+[numthreads(THREADGROUP_SIZE_1D, 1, 1)]
+void RefineSparseVoxels(
+	uint GroupThreadId : SV_GroupThreadID,
+	uint DispatchThreadId : SV_DispatchThreadID
+)
+{
+	RWNumVoxelsToRefineBuffer[1] = 1;
+	RWNumVoxelsToRefineBuffer[2] = 1;
+
+	uint Index = DispatchThreadId;
+	if (Index < NumVoxelsToRefineBuffer[0])
+	{
+		FVoxelData VoxelData = DecodeVoxelData(VoxelsToRefineBuffer[Index], VolumeResolution);
+
+		uint3 ParentVoxelIndex = VoxelData.VoxelIndex / 2;
+		uint3 AlignedVoxelIndex = ParentVoxelIndex * 2;
+
+		// Load siblings
+		bool bAllValid = true;
+		for (uint i = 0; i < 8; ++i)
+		{
+			uint3 Offset = uint3((i & 0x04) >> 2, (i & 0x02) >> 1, i & 0x01);
+			float3 Value = VoxelMinTexture.Load(int4(AlignedVoxelIndex + Offset, 0));
+			if (dot(Value, Value) <= 0.0)
+			{
+				bAllValid = false;
+			}
+		}
+
+		if (bAllValid)
+		{
+			// Skip any odd voxels, the aligned sibiling will add the parent
+			if (any(VoxelData.VoxelIndex % 2))
+			{
+				return;
+			}
+
+			// Emit parent voxel for further refinement
+			uint EntryIndex;
+			InterlockedAdd(RWNumVoxelsToRefineBuffer[0], 1, EntryIndex);
+			RWVoxelsToRefineBuffer[EntryIndex] = EncodeVoxelData(ParentVoxelIndex, MipLevel + 1, MipVolumeResolution / 2);
+		}
+		else
+		{
+			// Emit voxel
+			uint EntryIndex;
+			InterlockedAdd(RWNumVoxelsBuffer[0], 1, EntryIndex);
+			RWVoxelBuffer[EntryIndex] = VoxelsToRefineBuffer[Index];
+		}
+	}
+}
+
+// Using RWStructuredBuffer<float> instead of RWStructuredBuffer<float3> to overcome Vulkan alignment error:
+// error: cannot instantiate RWStructuredBuffer with given packed alignment; 'VK_EXT_scalar_block_layout' not supported
+#define VULKAN_ALIGNMENT_WORKAROUND 1
+#if VULKAN_ALIGNMENT_WORKAROUND
+	RWStructuredBuffer<float> RWPositionBuffer;
+#else
+	RWStructuredBuffer<float3> RWPositionBuffer;
+#endif
+
+[numthreads(64, 1, 1)]
+void CreateSparseVoxelBLAS(
+	uint GroupThreadId : SV_GroupThreadID,
+	uint DispatchThreadId : SV_DispatchThreadID
+)
+{
+	FBox GlobalBounds;
+	GlobalBounds.Min = GetLocalBoundsOrigin() - GetLocalBoundsExtent();
+	GlobalBounds.Max = GetLocalBoundsOrigin() + GetLocalBoundsExtent();
+
+	uint3 VolumeResolution = GetVolumeResolution();
+
+	if (DispatchThreadId < GetNumVoxels())
+	{
+		FVoxelDataPacked VoxelDataPacked = GetVoxelDataPacked(DispatchThreadId);
+		FVoxelData VoxelData = DecodeVoxelData(VoxelDataPacked, VolumeResolution);
+		FBox VoxelBounds = GetVoxelBoundingBox(VoxelData, VolumeResolution, GlobalBounds.Min, GlobalBounds.Max);
+
+#if VULKAN_ALIGNMENT_WORKAROUND
+		RWPositionBuffer[DispatchThreadId * 2 * 3 + 0] = VoxelBounds.Min.x;
+		RWPositionBuffer[DispatchThreadId * 2 * 3 + 1] = VoxelBounds.Min.y;
+		RWPositionBuffer[DispatchThreadId * 2 * 3 + 2] = VoxelBounds.Min.z;
+
+		RWPositionBuffer[DispatchThreadId * 2 * 3 + 3] = VoxelBounds.Max.x;
+		RWPositionBuffer[DispatchThreadId * 2 * 3 + 4] = VoxelBounds.Max.y;
+		RWPositionBuffer[DispatchThreadId * 2 * 3 + 5] = VoxelBounds.Max.z;
+#else
+		RWPositionBuffer[DispatchThreadId * 2] = VoxelBounds.Min;
+		RWPositionBuffer[DispatchThreadId * 2 + 1] = VoxelBounds.Max;
+#endif
+	}
+}
\ No newline at end of file
diff --git a/Engine/Shaders/Private/HeterogeneousVolumes/HeterogeneousVolumesSparseVoxelUniformBufferUtils.ush b/Engine/Shaders/Private/HeterogeneousVolumes/HeterogeneousVolumesSparseVoxelUniformBufferUtils.ush
new file mode 100644
--- /dev/null
+++ b/Engine/Shaders/Private/HeterogeneousVolumes/HeterogeneousVolumesSparseVoxelUniformBufferUtils.ush
@@ -0,0 +1,97 @@
+// Copyright Epic Games, Inc. All Rights Reserved.
+
+float4x4 GetLocalToWorld()
+{
+	return SparseVoxelUniformBuffer.LocalToWorld;
+}
+
+float4x4 GetWorldToLocal()
+{
+	return SparseVoxelUniformBuffer.WorldToLocal;
+}
+
+float3 GetLocalBoundsOrigin()
+{
+	return SparseVoxelUniformBuffer.LocalBoundsOrigin;
+}
+
+float3 GetLocalBoundsExtent()
+{
+	return SparseVoxelUniformBuffer.LocalBoundsExtent;
+}
+
+float GetStepSize()
+{
+	return SparseVoxelUniformBuffer.StepSize;
+}
+
+float GetShadowStepFactor()
+{
+	return SparseVoxelUniformBuffer.ShadowStepFactor;
+}
+
+uint3 GetVolumeResolution()
+{
+	return SparseVoxelUniformBuffer.VolumeResolution;
+}
+
+uint GetNumVoxels()
+{
+	return SparseVoxelUniformBuffer.NumVoxelsBuffer[0];
+}
+
+FVoxelDataPacked GetVoxelDataPacked(uint VoxelIndex)
+{
+	return SparseVoxelUniformBuffer.VoxelBuffer[VoxelIndex];
+}
+
+struct FVolumeSampleContext
+{
+	float3 UVW;
+	uint MipLevel;
+};
+
+FVolumeSampleContext CreateVolumeSampleContext(float3 LocalPosition, float3 WorldPosition)
+{
+	float3 LocalBoundsMin = GetLocalBoundsOrigin() - GetLocalBoundsExtent();
+	float3 LocalBoundsMax = GetLocalBoundsOrigin() + GetLocalBoundsExtent();
+
+	FVolumeSampleContext VolumeSampleContext;
+	VolumeSampleContext.UVW = (LocalPosition - LocalBoundsMin) / (LocalBoundsMax - LocalBoundsMin);
+	VolumeSampleContext.MipLevel = 0;// RayMarchingContext.MipLevel;
+
+	return VolumeSampleContext;
+}
+
+float3 SampleExtinction(inout FVolumeSampleContext Context)
+{
+	float3 Extinction = Texture3DSampleLevel(
+		SparseVoxelUniformBuffer.ExtinctionTexture,
+		SparseVoxelUniformBuffer.TextureSampler,
+		Context.UVW,
+		Context.MipLevel).rgb;
+
+	return Extinction;
+}
+
+float3 SampleEmission(inout FVolumeSampleContext Context)
+{
+	float3 Emission = Texture3DSampleLevel(
+		SparseVoxelUniformBuffer.EmissionTexture,
+		SparseVoxelUniformBuffer.TextureSampler,
+		Context.UVW,
+		Context.MipLevel).xyz;
+
+	return Emission;
+}
+
+float3 SampleAlbedo(inout FVolumeSampleContext Context)
+{
+	float3 Albedo = Texture3DSampleLevel(
+		SparseVoxelUniformBuffer.AlbedoTexture,
+		SparseVoxelUniformBuffer.TextureSampler,
+		Context.UVW,
+		Context.MipLevel).xyz;
+
+	return Albedo;
+}
\ No newline at end of file
diff --git a/Engine/Shaders/Private/HeterogeneousVolumes/HeterogeneousVolumesSparseVoxelUtils.ush b/Engine/Shaders/Private/HeterogeneousVolumes/HeterogeneousVolumesSparseVoxelUtils.ush
new file mode 100644
--- /dev/null
+++ b/Engine/Shaders/Private/HeterogeneousVolumes/HeterogeneousVolumesSparseVoxelUtils.ush
@@ -0,0 +1,89 @@
+// Copyright Epic Games, Inc. All Rights Reserved.
+
+#define INVALID_VOXEL 0xFFFFFFFF;
+
+struct FVoxelData
+{
+	uint3 VoxelIndex;
+	uint MipLevel;
+};
+
+struct FVoxelDataPacked
+{
+	uint LinearIndex;
+	uint MipLevel;
+};
+
+FVoxelDataPacked EncodeVoxelData(uint3 VoxelIndex, uint MipLevel, uint3 VolumeResolutionAtMipLevel)
+{
+	FVoxelDataPacked VoxelDataPacked = (FVoxelDataPacked)0;
+
+	// TODO: Morton Order
+	VoxelDataPacked.LinearIndex =
+		VoxelIndex.z * VolumeResolutionAtMipLevel.x * VolumeResolutionAtMipLevel.y +
+		VoxelIndex.y * VolumeResolutionAtMipLevel.x +
+		VoxelIndex.x;
+	VoxelDataPacked.MipLevel = MipLevel;
+
+	return VoxelDataPacked;
+}
+
+FVoxelData DecodeVoxelData(FVoxelDataPacked VoxelDataPacked, uint3 VolumeResolutionAtMip0)
+{
+	FVoxelData VoxelData;
+	VoxelData.MipLevel = VoxelDataPacked.MipLevel;
+
+	// TODO: Morton Order
+	uint3 VolumeResolution = VolumeResolutionAtMip0 >> VoxelData.MipLevel;
+	VoxelData.VoxelIndex.x = (VoxelDataPacked.LinearIndex % (VolumeResolution.x * VolumeResolution.y)) % VolumeResolution.x;
+	VoxelData.VoxelIndex.y = (VoxelDataPacked.LinearIndex % (VolumeResolution.x * VolumeResolution.y)) / VolumeResolution.x;
+	VoxelData.VoxelIndex.z = VoxelDataPacked.LinearIndex / (VolumeResolution.x * VolumeResolution.y);
+
+	return VoxelData;
+}
+
+FVoxelData CreateEmptyVoxelData()
+{
+	FVoxelData VoxelData;
+	VoxelData.VoxelIndex = 0;
+	VoxelData.MipLevel = INVALID_VOXEL;
+	return VoxelData;
+}
+
+bool IsVoxelDataValid(FVoxelData VoxelData)
+{
+	return VoxelData.MipLevel != INVALID_VOXEL;
+}
+
+bool IsVoxelDataEqual(FVoxelData LHS, FVoxelData RHS)
+{
+	return LHS.VoxelIndex.x == RHS.VoxelIndex.x &&
+		LHS.VoxelIndex.y == RHS.VoxelIndex.y &&
+		LHS.VoxelIndex.z == RHS.VoxelIndex.z &&
+		LHS.MipLevel == RHS.MipLevel;
+}
+
+struct FBox
+{
+	float3 Min;
+	float3 Max;
+};
+
+FBox GetVoxelBoundingBox(FVoxelData VoxelData, int3 VolumeResolution, float3 ObjectBoundsMin, float3 ObjectBoundsMax)
+{
+	float3 MipResolution = float3(
+		VolumeResolution.x >> VoxelData.MipLevel,
+		VolumeResolution.y >> VoxelData.MipLevel,
+		VolumeResolution.z >> VoxelData.MipLevel
+	);
+
+	float3 UVW0 = VoxelData.VoxelIndex / MipResolution;
+	float3 UVW1 = (VoxelData.VoxelIndex + 1) / MipResolution;
+	float3 ObjectBoundsExtent = ObjectBoundsMax - ObjectBoundsMin;
+
+	FBox Box;
+	const float Epsilon = 0.0001;
+	Box.Min = ObjectBoundsMin + ObjectBoundsExtent * UVW0 - Epsilon;
+	Box.Max = ObjectBoundsMin + ObjectBoundsExtent * UVW1 + Epsilon;
+	return Box;
+}
diff --git a/Engine/Shaders/Private/HeterogeneousVolumes/HeterogeneousVolumesTracingUtils.ush b/Engine/Shaders/Private/HeterogeneousVolumes/HeterogeneousVolumesTracingUtils.ush
new file mode 100644
--- /dev/null
+++ b/Engine/Shaders/Private/HeterogeneousVolumes/HeterogeneousVolumesTracingUtils.ush
@@ -0,0 +1,23 @@
+// Copyright Epic Games, Inc. All Rights Reserved.
+
+// Modified LineBoxIntersect() to operate at consistent RayDirection scale, within some [TMin, TMax] span
+float2 IntersectAABB(float3 RayOrigin, float3 RayDirection, float RayTMin, float RayTMax, float3 BoxMin, float3 BoxMax)
+{
+	float3 InvRayDir = 1.0 / RayDirection;
+	//find the ray intersection with each of the 3 planes defined by the minimum extrema.
+	float3 FirstPlaneIntersections = (BoxMin - RayOrigin) * InvRayDir;
+	//find the ray intersection with each of the 3 planes defined by the maximum extrema.
+	float3 SecondPlaneIntersections = (BoxMax - RayOrigin) * InvRayDir;
+	//get the closest of these intersections along the ray
+	float3 ClosestPlaneIntersections = min(FirstPlaneIntersections, SecondPlaneIntersections);
+	//get the furthest of these intersections along the ray
+	float3 FurthestPlaneIntersections = max(FirstPlaneIntersections, SecondPlaneIntersections);
+
+	float2 BoxIntersections;
+	//find the furthest near intersection
+	BoxIntersections.x = max(ClosestPlaneIntersections.x, max(ClosestPlaneIntersections.y, ClosestPlaneIntersections.z));
+	//find the closest far intersection
+	BoxIntersections.y = min(FurthestPlaneIntersections.x, min(FurthestPlaneIntersections.y, FurthestPlaneIntersections.z));
+	//clamp the intersections to be between RayTMin and RayTMax
+	return float2(max(BoxIntersections.x, RayTMin), min(BoxIntersections.y, RayTMax));
+}
diff --git a/Engine/Shaders/Private/MaterialTemplate.ush b/Engine/Shaders/Private/MaterialTemplate.ush
--- a/Engine/Shaders/Private/MaterialTemplate.ush
+++ b/Engine/Shaders/Private/MaterialTemplate.ush
@@ -591,6 +591,10 @@ void MaterialPackInterpolatorComponent(in out FMaterialVertexParameters Paramete
 #endif
 }
 
+#if GET_PRIMITIVE_DATA_OVERRIDE
+	FPrimitiveSceneData GetPrimitiveData(FMaterialVertexParameters Parameters);
+	FPrimitiveSceneData GetPrimitiveData(FMaterialPixelParameters Parameters);
+#else
 FPrimitiveSceneData GetPrimitiveData(FMaterialVertexParameters Parameters)
 {
 	return Parameters.SceneData.Primitive;
@@ -600,6 +604,7 @@ FPrimitiveSceneData GetPrimitiveData(FMaterialPixelParameters Parameters)
 {
 	return GetPrimitiveData(Parameters.PrimitiveId);
 }
+#endif 
 
 bool UnpackUniform_bool(uint Packed, uint BitOffset)
 {
diff --git a/Engine/Source/Runtime/Renderer/Private/DeferredShadingRenderer.cpp b/Engine/Source/Runtime/Renderer/Private/DeferredShadingRenderer.cpp
--- a/Engine/Source/Runtime/Renderer/Private/DeferredShadingRenderer.cpp
+++ b/Engine/Source/Runtime/Renderer/Private/DeferredShadingRenderer.cpp
@@ -75,6 +75,7 @@
 #include "Containers/ChunkedArray.h"
 #include "Async/ParallelFor.h"
 #include "Shadows/ShadowSceneRenderer.h"
+#include "HeterogeneousVolumes/HeterogeneousVolumes.h"
 
 extern int32 GNaniteShowStats;
 extern int32 GNanitePickingDomain;
@@ -1721,7 +1722,7 @@ bool FDeferredShadingSceneRenderer::DispatchRayTracingWorldUpdates(FRDGBuilder&
 
 	RDG_GPU_MASK_SCOPE(GraphBuilder, FRHIGPUMask::All());
 
-	RayTracingScene.CreateWithInitializationData(GraphBuilder, Scene->GPUScene, ReferenceView.ViewMatrices, MoveTemp(ReferenceView.RayTracingSceneInitData));
+	RayTracingScene.CreateWithInitializationData(GraphBuilder, &Scene->GPUScene, ReferenceView.ViewMatrices, MoveTemp(ReferenceView.RayTracingSceneInitData));
 
 	const uint32 BLASScratchSize = Scene->GetRayTracingDynamicGeometryCollection()->ComputeScratchBufferSize();
 	if (BLASScratchSize > 0)
@@ -3177,6 +3178,11 @@ void FDeferredShadingSceneRenderer::Render(FRDGBuilder& GraphBuilder)
 		ComputeVolumetricFog(GraphBuilder, SceneTextures);
 	}
 
+	if (ShouldRenderHeterogeneousVolumes(Scene) && !bHasRayTracedOverlay)
+	{
+		RenderHeterogeneousVolumes(GraphBuilder, SceneTextures);
+	}
+
 	if (bShouldRenderVolumetricCloud && IsVolumetricRenderTargetEnabled() && !bHasHalfResCheckerboardMinMaxDepth && !bHasRayTracedOverlay)
 	{
 		HalfResolutionDepthCheckerboardMinMaxTexture = CreateHalfResolutionDepthCheckerboardMinMax(GraphBuilder, Views, SceneTextures.Depth.Resolve);
@@ -3466,6 +3472,11 @@ void FDeferredShadingSceneRenderer::Render(FRDGBuilder& GraphBuilder)
 		RenderStationaryLightOverlap(GraphBuilder, SceneTextures, LightingChannelsTexture);
 	}
 
+	if (ShouldRenderHeterogeneousVolumes(Scene) && !bHasRayTracedOverlay)
+	{
+		CompositeHeterogeneousVolumes(GraphBuilder, SceneTextures);
+	}
+
 	if (bShouldVisualizeVolumetricCloud && !bHasRayTracedOverlay)
 	{
 		RenderVolumetricCloud(GraphBuilder, SceneTextures, false, true, HalfResolutionDepthCheckerboardMinMaxTexture, false, InstanceCullingManager);
diff --git a/Engine/Source/Runtime/Renderer/Private/DeferredShadingRenderer.h b/Engine/Source/Runtime/Renderer/Private/DeferredShadingRenderer.h
--- a/Engine/Source/Runtime/Renderer/Private/DeferredShadingRenderer.h
+++ b/Engine/Source/Runtime/Renderer/Private/DeferredShadingRenderer.h
@@ -863,6 +863,9 @@ private:
 	void ComputeVolumetricFog(FRDGBuilder& GraphBuilder,
 		const FSceneTextures& SceneTextures);
 
+	void RenderHeterogeneousVolumes(FRDGBuilder& GraphBuilder, const FSceneTextures& SceneTextures);
+	void CompositeHeterogeneousVolumes(FRDGBuilder& GraphBuilder, const FSceneTextures& SceneTextures);
+
 	void VisualizeVolumetricLightmap(FRDGBuilder& GraphBuilder, const FSceneTextures& SceneTextures);
 
 	/** Render image based reflections (SSR, Env, SkyLight) without compute shaders */
diff --git a/Engine/Source/Runtime/Renderer/Private/HeterogeneousVolumes/HeterogeneousVolumes.cpp b/Engine/Source/Runtime/Renderer/Private/HeterogeneousVolumes/HeterogeneousVolumes.cpp
new file mode 100644
--- /dev/null
+++ b/Engine/Source/Runtime/Renderer/Private/HeterogeneousVolumes/HeterogeneousVolumes.cpp
@@ -0,0 +1,432 @@
+// Copyright Epic Games, Inc. All Rights Reserved.
+
+#include "HeterogeneousVolumes.h"
+
+#include "PixelShaderUtils.h"
+#include "RayTracingDefinitions.h"
+#include "RayTracingInstance.h"
+#include "RayTracingInstanceBufferUtil.h"
+#include "RendererPrivate.h"
+#include "ScenePrivate.h"
+#include "SceneManagement.h"
+
+static TAutoConsoleVariable<int32> CVarHeterogeneousVolumes(
+	TEXT("r.HeterogeneousVolumes"),
+	0,
+	TEXT("Enables the Heterogeneous volume integrator (Default = 0)"),
+	ECVF_RenderThreadSafe
+);
+
+static TAutoConsoleVariable<int32> CVarHeterogeneousVolumesDebug(
+	TEXT("r.HeterogeneousVolumes.Debug"),
+	0,
+	TEXT("Creates auxillary output buffers for debugging (Default = 0)"),
+	ECVF_RenderThreadSafe
+);
+
+static TAutoConsoleVariable<int32> CVarHeterogeneousVolumesHardwareRayTracing(
+	TEXT("r.HeterogeneousVolumes.HardwareRayTracing"),
+	0,
+	TEXT("Enables hardware ray tracing acceleration (Default = 0)"),
+	ECVF_RenderThreadSafe
+);
+
+static TAutoConsoleVariable<int32> CVarHeterogeneousVolumesJitter(
+	TEXT("r.HeterogeneousVolumes.Jitter"),
+	1,
+	TEXT("Enables jitter when ray marching (Default = 1)"),
+	ECVF_RenderThreadSafe
+);
+
+static TAutoConsoleVariable<int32> CVarHeterogeneousVolumesMaxStepCount(
+	TEXT("r.HeterogeneousVolumes.MaxStepCount"),
+	128,
+	TEXT("The maximum ray-marching step count (Default = 128)"),
+	ECVF_RenderThreadSafe
+);
+
+static TAutoConsoleVariable<float> CVarHeterogeneousVolumesMaxTraceDistance(
+	TEXT("r.HeterogeneousVolumes.MaxTraceDistance"),
+	10000.0,
+	TEXT("The maximum trace view-distance for direct volume rendering (Default = 10000)"),
+	ECVF_RenderThreadSafe
+);
+
+static TAutoConsoleVariable<int32> CVarHeterogeneousVolumesPreshading(
+	TEXT("r.HeterogeneousVolumes.Preshading"),
+	0,
+	TEXT("Evaluates the material into a canonical preshaded volume before rendering the result (Default = 0)"),
+	ECVF_RenderThreadSafe
+);
+
+static TAutoConsoleVariable<int32> CVarHeterogeneousVolumesPreshadingMipLevel(
+	TEXT("r.HeterogeneousVolumes.Preshading.MipLevel"),
+	0,
+	TEXT("Statically determines the MIP-level when evaluating preshaded volume data (Default = 0)"),
+	ECVF_RenderThreadSafe
+);
+
+static TAutoConsoleVariable<int32> CVarHeterogeneousVolumesPreshadingVolumeResolutionX(
+	TEXT("r.HeterogeneousVolumes.Preshading.VolumeResolution.X"),
+	256,
+	TEXT("Determines the preshading volume resolution in X (Default = 256)"),
+	ECVF_RenderThreadSafe
+);
+
+static TAutoConsoleVariable<int32> CVarHeterogeneousVolumesPreshadingVolumeResolutionY(
+	TEXT("r.HeterogeneousVolumes.Preshading.VolumeResolution.Y"),
+	256,
+	TEXT("Determines the preshading volume resolution in Y (Default = 256)"),
+	ECVF_RenderThreadSafe
+);
+
+static TAutoConsoleVariable<int32> CVarHeterogeneousVolumesPreshadingVolumeResolutionZ(
+	TEXT("r.HeterogeneousVolumes.Preshading.VolumeResolution.Z"),
+	256,
+	TEXT("Determines the preshading volume resolution in Z (Default = 256)"),
+	ECVF_RenderThreadSafe
+);
+
+static TAutoConsoleVariable<float> CVarHeterogeneousVolumesShadowStepFactor(
+	TEXT("r.HeterogeneousVolumes.ShadowStepFactor"),
+	10.0,
+	TEXT("The ray-marching step-size factor for shadow rays (Default = 8.0)"),
+	ECVF_RenderThreadSafe
+);
+
+static TAutoConsoleVariable<int32> CVarHeterogeneousVolumesSparseVoxel(
+	TEXT("r.HeterogeneousVolumes.SparseVoxel"),
+	0,
+	TEXT("Uses sparse-voxel rendering algorithms (Default = 0)"),
+	ECVF_RenderThreadSafe
+);
+
+static TAutoConsoleVariable<int32> CVarHeterogeneousVolumesSparseVoxelGenerationMipBias(
+	TEXT("r.HeterogeneousVolumes.SparseVoxel.GenerationMipBias"),
+	3,
+	TEXT("Determines MIP bias for sparse voxel generation (Default = 3)"),
+	ECVF_RenderThreadSafe
+);
+
+static TAutoConsoleVariable<int32> CVarHeterogeneousVolumesSparseVoxelPerTileCulling(
+	TEXT("r.HeterogeneousVolumes.SparseVoxel.PerTileCulling"),
+	1,
+	TEXT("Enables sparse-voxel culling when using tiled rendering (Default = 1)"),
+	ECVF_RenderThreadSafe
+);
+
+static TAutoConsoleVariable<int32> CVarHeterogeneousVolumesSparseVoxelRefinement(
+	TEXT("r.HeterogeneousVolumes.SparseVoxel.Refinement"),
+	1,
+	TEXT("Uses hierarchical refinement to coalesce neighboring sparse-voxels (Default = 1)"),
+	ECVF_RenderThreadSafe
+);
+
+static TAutoConsoleVariable<float> CVarHeterogeneousVolumesStepSize(
+	TEXT("r.HeterogeneousVolumes.StepSize"),
+	1.0,
+	TEXT("The ray-marching step-size (Default = 1.0)"),
+	ECVF_RenderThreadSafe
+);
+
+DECLARE_GPU_STAT_NAMED(HeterogeneousVolumesStat, TEXT("HeterogeneousVolumes"));
+
+static bool IsHeterogeneousVolumesEnabled()
+{
+	return CVarHeterogeneousVolumes.GetValueOnRenderThread() != 0;
+}
+
+bool ShouldRenderHeterogeneousVolumes(
+	const FScene* Scene
+)
+{
+	return IsHeterogeneousVolumesEnabled()
+		&& Scene != nullptr;
+}
+
+bool ShouldRenderHeterogeneousVolumesForView(
+	const FSceneView& View
+)
+{
+	return IsHeterogeneousVolumesEnabled()
+		&& View.Family
+		&& !View.bIsPlanarReflection
+		&& !View.bIsSceneCapture
+		&& !View.bIsReflectionCapture
+		&& View.State;
+}
+
+bool DoesPlatformSupportHeterogeneousVolumes(EShaderPlatform Platform)
+{
+	return IsFeatureLevelSupported(Platform, ERHIFeatureLevel::SM5)
+		// TODO:
+		// && FDataDrivenShaderPlatformInfo::GetSupportsHeterogeneousVolumes(Platform)
+		&& !IsForwardShadingEnabled(Platform);
+}
+
+namespace HeterogeneousVolumes
+{
+	// CVars
+
+	FIntVector GetVolumeResolution()
+	{
+		FIntVector VolumeResolution;
+		VolumeResolution.X = FMath::Clamp(CVarHeterogeneousVolumesPreshadingVolumeResolutionX.GetValueOnRenderThread(), 64, 1024);
+		VolumeResolution.Y = FMath::Clamp(CVarHeterogeneousVolumesPreshadingVolumeResolutionY.GetValueOnRenderThread(), 64, 1024);
+		VolumeResolution.Z = FMath::Clamp(CVarHeterogeneousVolumesPreshadingVolumeResolutionZ.GetValueOnRenderThread(), 64, 1024);
+		return VolumeResolution;
+	}
+
+	float GetShadowStepFactor()
+	{
+		return CVarHeterogeneousVolumesShadowStepFactor.GetValueOnRenderThread();
+	}
+
+	float GetMaxTraceDistance()
+	{
+		return CVarHeterogeneousVolumesMaxTraceDistance.GetValueOnRenderThread();
+	}
+
+	float GetStepSize()
+	{
+		return CVarHeterogeneousVolumesStepSize.GetValueOnRenderThread();
+	}
+
+	float GetMaxStepCount()
+	{
+		return CVarHeterogeneousVolumesMaxStepCount.GetValueOnRenderThread();
+	}
+
+	int32 GetMipLevel()
+	{
+		return CVarHeterogeneousVolumesPreshadingMipLevel.GetValueOnRenderThread();
+	}
+
+	uint32 GetSparseVoxelMipBias()
+	{
+		// TODO: Clamp based on texture dimension..
+		return FMath::Clamp(CVarHeterogeneousVolumesSparseVoxelGenerationMipBias.GetValueOnRenderThread(), 0, 10);
+	}
+
+	int32 GetDebugMode()
+	{
+		return CVarHeterogeneousVolumesDebug.GetValueOnRenderThread();
+	}
+
+	bool UseSparseVoxelPipeline()
+	{
+		return CVarHeterogeneousVolumesSparseVoxel.GetValueOnRenderThread() != 0;
+	}
+
+	bool ShouldRefineSparseVoxels()
+	{
+		return CVarHeterogeneousVolumesSparseVoxelRefinement.GetValueOnRenderThread() != 0;
+	}
+
+	bool UseSparseVoxelPerTileCulling()
+	{
+		return CVarHeterogeneousVolumesSparseVoxelPerTileCulling.GetValueOnRenderThread() != 0;
+	}
+
+	bool ShouldJitter()
+	{
+		return CVarHeterogeneousVolumesJitter.GetValueOnRenderThread() != 0;
+	}
+
+	bool UseHardwareRayTracing()
+	{
+		return IsRayTracingEnabled()
+			&& (CVarHeterogeneousVolumesHardwareRayTracing.GetValueOnRenderThread() != 0);
+	}
+
+	// Convenience Utils
+
+	int GetVoxelCount(FIntVector VolumeResolution)
+	{
+		return VolumeResolution.X * VolumeResolution.Y * VolumeResolution.Z;
+	}
+
+	int GetVoxelCount(FRDGTextureDesc TextureDesc)
+	{
+		return TextureDesc.Extent.X * TextureDesc.Extent.Y * TextureDesc.Depth;
+	}
+
+	FIntVector GetMipVolumeResolution(FIntVector VolumeResolution, uint32 MipLevel)
+	{
+		return FIntVector(
+			FMath::Max(VolumeResolution.X >> MipLevel, 1),
+			FMath::Max(VolumeResolution.Y >> MipLevel, 1),
+			FMath::Max(VolumeResolution.Z >> MipLevel, 1)
+		);
+	}
+}
+
+void FDeferredShadingSceneRenderer::RenderHeterogeneousVolumes(
+	FRDGBuilder& GraphBuilder,
+	const FSceneTextures& SceneTextures
+)
+{
+	RDG_EVENT_SCOPE(GraphBuilder, "HeterogeneousVolumes");
+	RDG_GPU_STAT_SCOPE(GraphBuilder, HeterogeneousVolumesStat);
+
+	for (int32 ViewIndex = 0; ViewIndex < Views.Num(); ViewIndex++)
+	{
+		FViewInfo& View = Views[ViewIndex];
+
+		// Per-view??
+		FRDGTextureDesc Desc = SceneTextures.Color.Target->Desc;
+		Desc.Format = PF_FloatRGBA;
+		Desc.Flags &= ~(TexCreate_FastVRAM);
+		FRDGTextureRef HeterogeneousVolumeRadiance = GraphBuilder.CreateTexture(Desc, TEXT("HeterogeneousVolumes"));
+		AddClearUAVPass(GraphBuilder, GraphBuilder.CreateUAV(HeterogeneousVolumeRadiance), FLinearColor::Transparent);
+
+		if (ShouldRenderHeterogeneousVolumesForView(View))
+		{
+			for (int32 MeshBatchIndex = 0; MeshBatchIndex < View.VolumetricMeshBatches.Num(); ++MeshBatchIndex)
+			{
+				const FMeshBatch* Mesh = View.VolumetricMeshBatches[MeshBatchIndex].Mesh;
+				const FMaterialRenderProxy* MaterialRenderProxy = Mesh->MaterialRenderProxy;
+				const FMaterial& Material = MaterialRenderProxy->GetMaterialWithFallback(View.GetFeatureLevel(), MaterialRenderProxy);
+				// Only Niagara mesh particles bound to volume materials
+				if (Material.GetMaterialDomain() != MD_Volume || !Material.IsUsedWithNiagaraMeshParticles())
+				{
+					continue;
+				}
+
+				const FPrimitiveSceneProxy* PrimitiveSceneProxy = View.VolumetricMeshBatches[MeshBatchIndex].Proxy;
+				const FPrimitiveSceneInfo* PrimitiveSceneInfo = PrimitiveSceneProxy->GetPrimitiveSceneInfo();
+				const int32 PrimitiveId = PrimitiveSceneInfo->GetIndex();
+				const FBoxSphereBounds LocalBoxSphereBounds = PrimitiveSceneProxy->GetLocalBounds();
+
+				RDG_EVENT_SCOPE(GraphBuilder, "%s", *PrimitiveSceneProxy->GetResourceName().ToString());
+
+				// Material baking executes a pre-shading pipeline
+				if (CVarHeterogeneousVolumesPreshading.GetValueOnRenderThread())
+				{
+					RenderWithPreshading(
+						GraphBuilder,
+						SceneTextures,
+						Scene,
+						ViewFamily,
+						View,
+						// Object Data
+						PrimitiveSceneProxy,
+						MaterialRenderProxy,
+						PrimitiveId,
+						LocalBoxSphereBounds,
+						// Output
+						HeterogeneousVolumeRadiance
+					);
+				}
+				// Otherwise execute a live-shading pipeline
+				else
+				{
+					RenderWithLiveShading(
+						GraphBuilder,
+						SceneTextures,
+						Scene,
+						View,
+						// Object Data
+						PrimitiveSceneProxy,
+						MaterialRenderProxy,
+						PrimitiveId,
+						LocalBoxSphereBounds,
+						// Output
+						HeterogeneousVolumeRadiance
+					);
+				}
+			}
+		}
+
+		View.HeterogeneousVolumeRadiance = HeterogeneousVolumeRadiance;
+	}
+}
+
+class FHeterogeneousVolumesCompositeCS : public FGlobalShader
+{
+	DECLARE_GLOBAL_SHADER(FHeterogeneousVolumesCompositeCS);
+	SHADER_USE_PARAMETER_STRUCT(FHeterogeneousVolumesCompositeCS, FGlobalShader);
+
+	BEGIN_SHADER_PARAMETER_STRUCT(FParameters, )
+		// Scene data
+		SHADER_PARAMETER_STRUCT_REF(FViewUniformShaderParameters, View)
+
+		// Volume data
+		SHADER_PARAMETER_RDG_TEXTURE(Texture2D<float4>, HeterogeneousVolumeRadiance)
+
+		// Dispatch data
+		SHADER_PARAMETER(FIntVector, GroupCount)
+
+		// Output
+		SHADER_PARAMETER_RDG_TEXTURE_UAV(RWTexture2D<float4>, RWColorTexture)
+		END_SHADER_PARAMETER_STRUCT()
+
+public:
+	static bool ShouldCompilePermutation(
+		const FGlobalShaderPermutationParameters& Parameters
+	)
+	{
+		// Apply conditional project settings for Heterogeneous volumes?
+		return DoesPlatformSupportHeterogeneousVolumes(Parameters.Platform);
+	}
+
+	static void ModifyCompilationEnvironment(
+		const FGlobalShaderPermutationParameters& Parameters,
+		FShaderCompilerEnvironment& OutEnvironment
+	)
+	{
+		FGlobalShader::ModifyCompilationEnvironment(Parameters, OutEnvironment);
+		OutEnvironment.SetDefine(TEXT("THREADGROUP_SIZE_1D"), GetThreadGroupSize1D());
+		OutEnvironment.SetDefine(TEXT("THREADGROUP_SIZE_2D"), GetThreadGroupSize2D());
+
+		// This shader takes a very long time to compile with FXC, so we pre-compile it with DXC first and then forward the optimized HLSL to FXC.
+		//OutEnvironment.CompilerFlags.Add(CFLAG_PrecompileWithDXC);
+		OutEnvironment.CompilerFlags.Add(CFLAG_AllowTypedUAVLoads);
+
+		OutEnvironment.SetDefine(TEXT("GET_PRIMITIVE_DATA_OVERRIDE"), 1);
+	}
+
+	static int32 GetThreadGroupSize1D() { return GetThreadGroupSize2D() * GetThreadGroupSize2D(); }
+	static int32 GetThreadGroupSize2D() { return 8; }
+};
+
+IMPLEMENT_GLOBAL_SHADER(FHeterogeneousVolumesCompositeCS, "/Engine/Private/HeterogeneousVolumes/HeterogeneousVolumesComposite.usf", "HeterogeneousVolumesCompositeCS", SF_Compute);
+
+
+void FDeferredShadingSceneRenderer::CompositeHeterogeneousVolumes(
+	FRDGBuilder& GraphBuilder,
+	const FSceneTextures& SceneTextures
+)
+{
+	for (int32 ViewIndex = 0; ViewIndex < Views.Num(); ViewIndex++)
+	{
+		FViewInfo& View = Views[ViewIndex];
+
+		if (ShouldRenderHeterogeneousVolumesForView(View))
+		{
+			uint32 GroupCountX = FMath::DivideAndRoundUp(View.ViewRect.Size().X, FHeterogeneousVolumesCompositeCS::GetThreadGroupSize2D());
+			uint32 GroupCountY = FMath::DivideAndRoundUp(View.ViewRect.Size().Y, FHeterogeneousVolumesCompositeCS::GetThreadGroupSize2D());
+			FIntVector GroupCount = FIntVector(GroupCountX, GroupCountY, 1);
+
+			FHeterogeneousVolumesCompositeCS::FParameters* PassParameters = GraphBuilder.AllocParameters<FHeterogeneousVolumesCompositeCS::FParameters>();
+			{
+				// Scene data
+				PassParameters->View = View.ViewUniformBuffer;
+				// Volume data
+				PassParameters->HeterogeneousVolumeRadiance = View.HeterogeneousVolumeRadiance;
+				// Dispatch data
+				PassParameters->GroupCount = GroupCount;
+				// Output
+				PassParameters->RWColorTexture = GraphBuilder.CreateUAV(SceneTextures.Color.Target);
+			}
+
+			TShaderRef<FHeterogeneousVolumesCompositeCS> ComputeShader = View.ShaderMap->GetShader<FHeterogeneousVolumesCompositeCS>();
+			FComputeShaderUtils::AddPass(
+				GraphBuilder,
+				RDG_EVENT_NAME("FHeterogeneousVolumesCompositeCS"),
+				ComputeShader,
+				PassParameters,
+				GroupCount);
+		}
+	}
+}
\ No newline at end of file
diff --git a/Engine/Source/Runtime/Renderer/Private/HeterogeneousVolumes/HeterogeneousVolumes.h b/Engine/Source/Runtime/Renderer/Private/HeterogeneousVolumes/HeterogeneousVolumes.h
new file mode 100644
--- /dev/null
+++ b/Engine/Source/Runtime/Renderer/Private/HeterogeneousVolumes/HeterogeneousVolumes.h
@@ -0,0 +1,205 @@
+// Copyright Epic Games, Inc. All Rights Reserved.
+
+#pragma once
+
+#include "PixelShaderUtils.h"
+#include "RayTracingDefinitions.h"
+#include "RayTracingInstance.h"
+#include "RayTracingInstanceBufferUtil.h"
+#include "RendererPrivate.h"
+#include "ScenePrivate.h"
+#include "SceneManagement.h"
+
+//
+// External API
+//
+
+bool ShouldRenderHeterogeneousVolumes(const FScene* Scene);
+bool ShouldRenderHeterogeneousVolumesForView(const FSceneView& View);
+bool DoesPlatformSupportHeterogeneousVolumes(EShaderPlatform Platform);
+
+//
+// Internal API
+//
+
+namespace HeterogeneousVolumes
+{
+	// CVars
+	FIntVector GetVolumeResolution();
+
+	float GetShadowStepFactor();
+	float GetMaxTraceDistance();
+	float GetStepSize();
+	float GetMaxStepCount();
+
+	int32 GetMipLevel();
+	int32 GetDebugMode();
+	uint32 GetSparseVoxelMipBias();
+
+	bool ShouldJitter();
+	bool ShouldRefineSparseVoxels();
+	bool UseHardwareRayTracing();
+	bool UseSparseVoxelPipeline();
+	bool UseSparseVoxelPerTileCulling();
+
+	// Convenience Utils
+	int GetVoxelCount(FIntVector VolumeResolution);
+	int GetVoxelCount(FRDGTextureDesc TextureDesc);
+	FIntVector GetMipVolumeResolution(FIntVector VolumeResolution, uint32 MipLevel);
+}
+
+struct FVoxelDataPacked
+{
+	uint32 LinearIndex;
+	uint32 MipLevel;
+};
+
+BEGIN_UNIFORM_BUFFER_STRUCT(FSparseVoxelUniformBufferParameters, )
+	// Object data
+	SHADER_PARAMETER(FMatrix44f, LocalToWorld)
+	SHADER_PARAMETER(FMatrix44f, WorldToLocal)
+	SHADER_PARAMETER(FVector3f, LocalBoundsOrigin)
+	SHADER_PARAMETER(FVector3f, LocalBoundsExtent)
+
+	// Volume data
+	SHADER_PARAMETER(FIntVector, VolumeResolution)
+	SHADER_PARAMETER_RDG_TEXTURE(Texture3D, ExtinctionTexture)
+	SHADER_PARAMETER_RDG_TEXTURE(Texture3D, EmissionTexture)
+	SHADER_PARAMETER_RDG_TEXTURE(Texture3D, AlbedoTexture)
+	SHADER_PARAMETER_SAMPLER(SamplerState, TextureSampler)
+
+	// Sparse voxel data
+	SHADER_PARAMETER_RDG_BUFFER_SRV(Buffer<uint>, NumVoxelsBuffer)
+	SHADER_PARAMETER_RDG_BUFFER_SRV(StructuredBuffer<FVoxelDataPacked>, VoxelBuffer)
+	SHADER_PARAMETER(int, MipLevel)
+
+	// Traversal hints
+	SHADER_PARAMETER(float, MaxTraceDistance)
+	SHADER_PARAMETER(float, StepSize)
+	SHADER_PARAMETER(float, ShadowStepFactor)
+END_UNIFORM_BUFFER_STRUCT()
+
+// Render specializations
+
+void RenderWithLiveShading(
+	FRDGBuilder& GraphBuilder,
+	const FSceneTextures& SceneTextures,
+	const FScene* Scene,
+	const FViewInfo& View,
+	// Object data
+	const FPrimitiveSceneProxy* PrimitiveSceneProxy,
+	const FMaterialRenderProxy* MaterialRenderProxy,
+	const int32 PrimitiveId,
+	const FBoxSphereBounds LocalBoxSphereBounds,
+	// Output
+	FRDGTextureRef& HeterogeneousVolumeRadiance
+);
+
+void RenderWithPreshading(
+	FRDGBuilder& GraphBuilder,
+	// Scene data
+	const FSceneTextures& SceneTextures,
+	FScene* Scene,
+	const FSceneViewFamily& ViewFamily,
+	FViewInfo& View,
+	// Object data
+	const FPrimitiveSceneProxy* PrimitiveSceneProxy,
+	const FMaterialRenderProxy* MaterialRenderProxy,
+	const int32 PrimitiveId,
+	const FBoxSphereBounds LocalBoxSphereBounds,
+	// Output
+	FRDGTextureRef& HeterogeneousVolumeRadiance
+);
+
+// Preshading Pipeline
+void ComputeHeterogeneousVolumeBakeMaterial(
+	FRDGBuilder& GraphBuilder,
+	// Scene data
+	const FScene* Scene,
+	const FViewInfo& View,
+	// Object data
+	const FPrimitiveSceneProxy* PrimitiveSceneProxy,
+	const FMaterialRenderProxy* MaterialRenderProxy,
+	const int32 PrimitiveId,
+	const FBoxSphereBounds LocalBoxSphereBounds,
+	// Volume data
+	FIntVector VolumeResolution,
+	// Output
+	FRDGTextureRef& HeterogeneousVolumeExtinctionTexture,
+	FRDGTextureRef& HeterogeneousVolumeEmissionTexture,
+	FRDGTextureRef& HeterogeneousVolumeAlbedoTexture
+);
+
+// Sparse Voxel Pipeline
+
+void CopyTexture3D(
+	FRDGBuilder& GraphBuilder,
+	const FViewInfo& View,
+	FRDGTextureRef Texture,
+	uint32 InputMipLevel,
+	FRDGTextureRef& OutputTexture
+);
+
+void GenerateSparseVoxels(
+	FRDGBuilder& GraphBuilder,
+	const FViewInfo& View,
+	FRDGTextureRef VoxelMinTexture,
+	FIntVector VolumeResolution,
+	uint32 MipLevel,
+	FRDGBufferRef& NumVoxelsBuffer,
+	FRDGBufferRef& VoxelBuffer
+);
+
+#if RHI_RAYTRACING
+
+void GenerateRayTracingGeometryInstance(
+	FRDGBuilder& GraphBuilder,
+	// Scene data
+	const FScene* Scene,
+	const FViewInfo& View,
+	// Object data
+	const FPrimitiveSceneProxy* PrimitiveSceneProxy,
+	// Sparse voxel data
+	FRDGBufferRef NumVoxelsBuffer,
+	TRDGUniformBufferRef<FSparseVoxelUniformBufferParameters> SparseVoxelUniformBuffer,
+	// Output
+	TArray<FRayTracingGeometryRHIRef>& RayTracingGeometries,
+	TArray<FMatrix>& RayTracingTransforms
+);
+
+void GenerateRayTracingScene(
+	FRDGBuilder& GraphBuilder,
+	// Scene data
+	const FScene* Scene,
+	const FViewInfo& View,
+	// Object data
+	const FPrimitiveSceneProxy* PrimitiveSceneProxy,
+	// Ray tracing data
+	TArray<FRayTracingGeometryRHIRef>& RayTracingGeometries,
+	TArray<FMatrix>& RayTracingTransforms,
+	// Output
+	FRayTracingScene& RayTracingScene
+);
+
+void RenderSingleScatteringWithPreshadingHardwareRayTracing(
+	FRDGBuilder& GraphBuilder,
+	// Scene data
+	const FScene* Scene,
+	const FViewInfo& View,
+	const FSceneTextures& SceneTextures,
+	// Light data
+	bool bApplyEmission,
+	bool bApplyDirectLighting,
+	uint32 LightType,
+	const FLightSceneInfo* LightSceneInfo,
+	// Object data
+	const FPrimitiveSceneProxy* PrimitiveSceneProxy,
+	// Sparse voxel data
+	TRDGUniformBufferRef<FSparseVoxelUniformBufferParameters> SparseVoxelUniformBuffer,
+	// Ray tracing data
+	FRayTracingScene& RayTracingScene,
+	// Output
+	FRDGTextureRef& HeterogeneousVolumeTexture
+);
+
+#endif // RHI_RAYTRACING
\ No newline at end of file
diff --git a/Engine/Source/Runtime/Renderer/Private/HeterogeneousVolumes/HeterogeneousVolumesHardwareRayTracing.cpp b/Engine/Source/Runtime/Renderer/Private/HeterogeneousVolumes/HeterogeneousVolumesHardwareRayTracing.cpp
new file mode 100644
--- /dev/null
+++ b/Engine/Source/Runtime/Renderer/Private/HeterogeneousVolumes/HeterogeneousVolumesHardwareRayTracing.cpp
@@ -0,0 +1,573 @@
+// Copyright Epic Games, Inc. All Rights Reserved.
+
+#include "HeterogeneousVolumes.h"
+
+#include "PixelShaderUtils.h"
+#include "RayTracingDefinitions.h"
+#include "RayTracingInstance.h"
+#include "RayTracingInstanceBufferUtil.h"
+#include "RendererPrivate.h"
+#include "ScenePrivate.h"
+#include "SceneManagement.h"
+
+#if RHI_RAYTRACING
+FIntVector GetVoxelCoord(uint32 VoxelIndex, FIntVector VolumeResolution)
+{
+	FIntVector VoxelCoord;
+
+	uint32 SliceSize = VolumeResolution.X * VolumeResolution.Y;
+	uint32 SliceIndex = VoxelIndex / SliceSize;
+	uint32 SliceCoord = VoxelIndex - SliceIndex * SliceSize;
+
+	VoxelCoord.X = SliceCoord % VolumeResolution.X;
+	VoxelCoord.Y = SliceCoord / VolumeResolution.X;
+	VoxelCoord.Z = SliceIndex;
+
+	return VoxelCoord;
+}
+
+FBox GetVoxelBounds(uint32 VoxelIndex, FIntVector VolumeResolution, FVector LocalBoundsOrigin, FVector LocalBoundsExtent)
+{
+	FBox VoxelBounds;
+
+	FIntVector VoxelCoord = GetVoxelCoord(VoxelIndex, VolumeResolution);
+	FVector VoxelSize = (LocalBoundsExtent * 2.0) / FVector(VolumeResolution);
+
+	VoxelBounds.Min = LocalBoundsOrigin - LocalBoundsExtent + FVector(VoxelCoord) * VoxelSize;
+	VoxelBounds.Max = VoxelBounds.Min + VoxelSize;
+	return VoxelBounds;
+}
+
+class FCreateSparseVoxelBLAS : public FGlobalShader
+{
+	DECLARE_GLOBAL_SHADER(FCreateSparseVoxelBLAS);
+	SHADER_USE_PARAMETER_STRUCT(FCreateSparseVoxelBLAS, FGlobalShader);
+
+	static bool ShouldCompilePermutation(const FGlobalShaderPermutationParameters& Parameters)
+	{
+		return DoesPlatformSupportHeterogeneousVolumes(Parameters.Platform);
+	}
+
+	static void ModifyCompilationEnvironment(
+		const FGlobalShaderPermutationParameters& Parameters,
+		FShaderCompilerEnvironment& OutEnvironment
+	)
+	{
+		OutEnvironment.CompilerFlags.Add(CFLAG_AllowTypedUAVLoads);
+	}
+
+	BEGIN_SHADER_PARAMETER_STRUCT(FParameters, )
+		// Input
+		SHADER_PARAMETER_RDG_UNIFORM_BUFFER(FSparseVoxelUniformBufferParameters, SparseVoxelUniformBuffer)
+
+		// Output
+		// Using RWStructuredBuffer<float> instead of RWStructuredBuffer<float3> to overcome Vulkan alignment error:
+		// error: cannot instantiate RWStructuredBuffer with given packed alignment; 'VK_EXT_scalar_block_layout' not supported
+		// SHADER_PARAMETER_RDG_BUFFER_UAV(RWBuffer<FVector>, RWPositionBuffer)
+		SHADER_PARAMETER_RDG_BUFFER_UAV(RWBuffer<float>, RWPositionBuffer)
+
+		// Indirect args
+		RDG_BUFFER_ACCESS(IndirectArgs, ERHIAccess::IndirectArgs)
+	END_SHADER_PARAMETER_STRUCT()
+};
+
+IMPLEMENT_GLOBAL_SHADER(FCreateSparseVoxelBLAS, "/Engine/Private/HeterogeneousVolumes/HeterogeneousVolumesSparseVoxelPipeline.usf", "CreateSparseVoxelBLAS", SF_Compute);
+
+void CreateSparseVoxelBLAS(
+	FRDGBuilder& GraphBuilder,
+	// Scene data
+	const FViewInfo& View,
+	// Sparse voxel data
+	TRDGUniformBufferRef<FSparseVoxelUniformBufferParameters> SparseVoxelUniformBuffer,
+	FRDGBufferRef NumVoxelsBuffer,
+	// Output
+	FRDGBufferRef PositionBuffer
+)
+{
+	FCreateSparseVoxelBLAS::FParameters* PassParameters = GraphBuilder.AllocParameters<FCreateSparseVoxelBLAS::FParameters>();
+	{
+		PassParameters->SparseVoxelUniformBuffer = SparseVoxelUniformBuffer;
+		PassParameters->RWPositionBuffer = GraphBuilder.CreateUAV(PositionBuffer);
+		PassParameters->IndirectArgs = NumVoxelsBuffer;
+	}
+
+	TShaderRef<FCreateSparseVoxelBLAS> ComputeShader = View.ShaderMap->GetShader<FCreateSparseVoxelBLAS>();
+
+	FComputeShaderUtils::AddPass(
+		GraphBuilder,
+		RDG_EVENT_NAME("CreateSparseVoxelBLAS"),
+		ERDGPassFlags::Compute,
+		ComputeShader,
+		PassParameters,
+		PassParameters->IndirectArgs,
+		0);
+}
+
+void GenerateRayTracingGeometryInstance(
+	FRDGBuilder& GraphBuilder,
+	// Scene data
+	const FScene* Scene,
+	const FViewInfo& View,
+	// Object data
+	const FPrimitiveSceneProxy* PrimitiveSceneProxy,
+	// Volume data
+	// Sparse voxel data
+	FRDGBufferRef NumVoxelsBuffer,
+	TRDGUniformBufferRef<FSparseVoxelUniformBufferParameters> SparseVoxelUniformBuffer,
+	// Output
+	TArray<FRayTracingGeometryRHIRef>& RayTracingGeometries,
+	TArray<FMatrix>& RayTracingTransforms
+)
+{
+	FIntVector MipVolumeResolution = HeterogeneousVolumes::GetMipVolumeResolution(SparseVoxelUniformBuffer->GetParameters()->VolumeResolution, SparseVoxelUniformBuffer->GetParameters()->MipLevel);
+	uint32 MipVoxelCount = HeterogeneousVolumes::GetVoxelCount(MipVolumeResolution);
+
+	TRefCountPtr<FRDGPooledBuffer> PooledVertexBuffer;
+	{
+#if 0
+		TResourceArray<FVector3f> PositionData;
+
+		PositionData.SetNumUninitialized(MipVoxelCount * 2);
+		for (uint32 MipVoxelIndex = 0; MipVoxelIndex < MipVoxelCount; ++MipVoxelIndex)
+		{
+			FBox VoxelBounds = GetVoxelBounds(MipVoxelIndex, MipVolumeResolution, FVector(SparseVoxelUniformBuffer->GetParameters()->LocalBoundsOrigin), FVector(SparseVoxelUniformBuffer->GetParameters()->LocalBoundsExtent));
+
+			uint32 BoundsIndex = MipVoxelIndex * 2;
+			PositionData[BoundsIndex] = FVector3f(VoxelBounds.Min);
+			PositionData[BoundsIndex + 1] = FVector3f(VoxelBounds.Max);
+		}
+
+		FRHIResourceCreateInfo CreateInfo(TEXT("HeterogeneousVolumesVB"));
+		CreateInfo.ResourceArray = &PositionData;
+#endif
+		FRDGBufferRef VertexBuffer = GraphBuilder.CreateBuffer(
+		FRDGBufferDesc::CreateStructuredDesc(sizeof(FVector3f), MipVoxelCount * 2),
+		TEXT("CreateSparseVoxelBLAS.VertexBuffer"));
+
+		AddClearUAVPass(GraphBuilder, GraphBuilder.CreateUAV(VertexBuffer), 0.0);
+		PooledVertexBuffer = GraphBuilder.ConvertToExternalBuffer(VertexBuffer);
+	}
+
+	// Morphs the dense-voxel topology into the sparse-voxel topology.
+	CreateSparseVoxelBLAS(GraphBuilder, View, SparseVoxelUniformBuffer, NumVoxelsBuffer, GraphBuilder.RegisterExternalBuffer(PooledVertexBuffer));
+
+	FRayTracingGeometryInitializer GeometryInitializer;
+	// TODO: REMOVE STRING ALLOCATION
+	GeometryInitializer.DebugName = *PrimitiveSceneProxy->GetResourceName().ToString();// +TEXT(" (HeterogeneousVolume)");
+	//GeometryInitializer.IndexBuffer = EmptyIndexBuffer;
+	GeometryInitializer.GeometryType = RTGT_Procedural;
+	GeometryInitializer.bFastBuild = false;
+
+	FRayTracingGeometrySegment Segment;
+	Segment.NumPrimitives = MipVoxelCount;
+	Segment.MaxVertices = MipVoxelCount * 2;
+	Segment.VertexBufferStride = 2u * sizeof(FVector3f);
+	Segment.VertexBuffer = PooledVertexBuffer->GetRHI();
+
+	GeometryInitializer.Segments.Add(Segment);
+	GeometryInitializer.TotalPrimitiveCount = Segment.NumPrimitives;
+	RayTracingGeometries.Add(RHICreateRayTracingGeometry(GeometryInitializer));
+	RayTracingTransforms.Add(PrimitiveSceneProxy->GetLocalToWorld());
+}
+
+BEGIN_SHADER_PARAMETER_STRUCT(FBuildBLASPassParams, )
+	SHADER_PARAMETER_RDG_BUFFER_UAV(RWStructuredBuffer, InstanceBuffer)
+END_SHADER_PARAMETER_STRUCT()
+
+BEGIN_SHADER_PARAMETER_STRUCT(FBuildTLASPassParams, )
+	RDG_BUFFER_ACCESS(RayTracingSceneScratchBuffer, ERHIAccess::UAVCompute)
+	RDG_BUFFER_ACCESS(RayTracingSceneInstanceBuffer, ERHIAccess::SRVCompute)
+END_SHADER_PARAMETER_STRUCT()
+
+void GenerateRayTracingScene(
+	FRDGBuilder& GraphBuilder,
+	// Scene data
+	const FScene* Scene,
+	const FViewInfo& View,
+	// Object data
+	const FPrimitiveSceneProxy* PrimitiveSceneProxy,
+	// Ray tracing data
+	TArray<FRayTracingGeometryRHIRef>& RayTracingGeometries,
+	TArray<FMatrix>& RayTracingTransforms,
+	// Output
+	FRayTracingScene& RayTracingScene
+)
+{
+	RayTracingScene.Reset();
+
+	// Collect instances
+	TArray<FRayTracingGeometryInstance> RayTracingInstances;
+	for (int32 GeometryIndex = 0; GeometryIndex < RayTracingGeometries.Num(); ++GeometryIndex)
+	{
+		checkf(RayTracingGeometries[GeometryIndex], TEXT("RayTracingGeometryInstance not created."))
+			FRayTracingGeometryInstance RayTracingGeometryInstance = {};
+		RayTracingGeometryInstance.GeometryRHI = RayTracingGeometries[GeometryIndex];
+		RayTracingGeometryInstance.NumTransforms = 1;
+		RayTracingGeometryInstance.Transforms = MakeArrayView(&RayTracingTransforms[GeometryIndex], 1);
+
+		RayTracingInstances.Add(RayTracingGeometryInstance);
+	}
+
+	// Build instance BLAS
+	FBuildBLASPassParams* PassParamsBLAS = GraphBuilder.AllocParameters<FBuildBLASPassParams>();
+
+	GraphBuilder.AddPass(
+		RDG_EVENT_NAME("BuildTLASInstanceBuffer"),
+		PassParamsBLAS,
+		ERDGPassFlags::Compute | ERDGPassFlags::NeverCull | ERDGPassFlags::NeverParallel,
+		[
+			PassParamsBLAS,
+			&RayTracingGeometries
+		](FRHICommandListImmediate& RHICmdList)
+		{
+			for (FRayTracingGeometryRHIRef RayTracingGeometry : RayTracingGeometries)
+			{
+				checkf(RayTracingGeometry, TEXT("RayTracingGeometry not created."));
+				RHICmdList.BuildAccelerationStructure(RayTracingGeometry);
+
+				// #yuriy_todo: explicit transitions and state validation for BLAS
+				// RHICmdList.Transition(FRHITransitionInfo(RayTracingGeometry.GetReference(), ERHIAccess::BVHWrite, ERHIAccess::BVHRead));
+			}
+		}
+	);
+
+	// Create RayTracingScene
+	const FGPUScene* EmptyGPUScene = nullptr;
+	FViewMatrices EmptyViewMatrices;
+	RayTracingScene.Instances = RayTracingInstances;
+	RayTracingScene.Create(GraphBuilder, EmptyGPUScene, EmptyViewMatrices);
+
+	// Build TLAS
+	FBuildTLASPassParams* PassParamsTLAS = GraphBuilder.AllocParameters<FBuildTLASPassParams>();
+	PassParamsTLAS->RayTracingSceneScratchBuffer = RayTracingScene.BuildScratchBuffer;
+	PassParamsTLAS->RayTracingSceneInstanceBuffer = RayTracingScene.InstanceBuffer;
+
+	const bool bRayTracingAsyncBuild = false;//CVarRayTracingAsyncBuild.GetValueOnRenderThread() != 0 && GRHISupportsRayTracingAsyncBuildAccelerationStructure;
+	const ERDGPassFlags ComputePassFlags = bRayTracingAsyncBuild ? ERDGPassFlags::AsyncCompute : ERDGPassFlags::Compute;
+	GraphBuilder.AddPass(
+		RDG_EVENT_NAME("RayTracingScene"),
+		PassParamsTLAS,
+		ComputePassFlags | ERDGPassFlags::NeverCull | ERDGPassFlags::NeverParallel,
+		[
+			PassParamsTLAS,
+			&RayTracingScene,
+			bRayTracingAsyncBuild
+		](FRHIComputeCommandList& RHICmdList)
+		{
+			FRHIRayTracingScene* RayTracingSceneRHI = RayTracingScene.GetRHIRayTracingSceneChecked();
+			FRHIBuffer* AccelerationStructureBuffer = RayTracingScene.GetBufferChecked();
+
+			FRayTracingSceneBuildParams SceneBuildParams;
+			SceneBuildParams.Scene = RayTracingSceneRHI;
+			SceneBuildParams.ScratchBuffer = PassParamsTLAS->RayTracingSceneScratchBuffer->GetRHI();
+			SceneBuildParams.ScratchBufferOffset = 0;
+			SceneBuildParams.InstanceBuffer = PassParamsTLAS->RayTracingSceneInstanceBuffer->GetRHI();
+			SceneBuildParams.InstanceBufferOffset = 0;
+
+			RHICmdList.BindAccelerationStructureMemory(RayTracingSceneRHI, AccelerationStructureBuffer, 0);
+			RHICmdList.BuildAccelerationStructure(SceneBuildParams);
+			// Submit potentially expensive BVH build commands to the GPU as soon as possible.
+			// Avoids a GPU bubble in some CPU-limited cases.
+			RHICmdList.SubmitCommandsHint();
+
+			RHICmdList.Transition(FRHITransitionInfo(RayTracingSceneRHI, ERHIAccess::BVHWrite, ERHIAccess::BVHRead));
+		}
+	);
+}
+
+class FHeterogeneousVolumesSparseVoxelsHitGroup : public FGlobalShader
+{
+	DECLARE_GLOBAL_SHADER(FHeterogeneousVolumesSparseVoxelsHitGroup)
+	SHADER_USE_ROOT_PARAMETER_STRUCT(FHeterogeneousVolumesSparseVoxelsHitGroup, FGlobalShader)
+
+	BEGIN_SHADER_PARAMETER_STRUCT(FParameters, )
+		SHADER_PARAMETER_RDG_UNIFORM_BUFFER(FSparseVoxelUniformBufferParameters, SparseVoxelUniformBuffer)
+	END_SHADER_PARAMETER_STRUCT()
+
+	static bool ShouldCompilePermutation(const FGlobalShaderPermutationParameters& Parameters)
+	{
+		return ShouldCompileRayTracingShadersForProject(Parameters.Platform) && DoesPlatformSupportHeterogeneousVolumes(Parameters.Platform) &&
+			FDataDrivenShaderPlatformInfo::GetSupportsRayTracingProceduralPrimitive(Parameters.Platform);
+	}
+
+	static void ModifyCompilationEnvironment(const FGlobalShaderPermutationParameters& Parameters, FShaderCompilerEnvironment& OutEnvironment)
+	{
+		FGlobalShader::ModifyCompilationEnvironment(Parameters, OutEnvironment);
+	}
+};
+
+IMPLEMENT_GLOBAL_SHADER(FHeterogeneousVolumesSparseVoxelsHitGroup, "/Engine/Private/HeterogeneousVolumes/HeterogeneousVolumesHardwareRayTracing.usf", "closesthit=SparseVoxelsClosestHitShader anyhit=SparseVoxelsAnyHitShader intersection=SparseVoxelsIntersectionShader", SF_RayHitGroup);
+
+class FHeterogeneousVolumesSparseVoxelMS : public FGlobalShader
+{
+	DECLARE_GLOBAL_SHADER(FHeterogeneousVolumesSparseVoxelMS)
+	SHADER_USE_ROOT_PARAMETER_STRUCT(FHeterogeneousVolumesSparseVoxelMS, FGlobalShader)
+
+	static bool ShouldCompilePermutation(const FGlobalShaderPermutationParameters& Parameters)
+	{
+		return ShouldCompileRayTracingShadersForProject(Parameters.Platform) && DoesPlatformSupportHeterogeneousVolumes(Parameters.Platform) &&
+			FDataDrivenShaderPlatformInfo::GetSupportsRayTracingProceduralPrimitive(Parameters.Platform);
+	}
+
+	static void ModifyCompilationEnvironment(const FGlobalShaderPermutationParameters& Parameters, FShaderCompilerEnvironment& OutEnvironment)
+	{
+		FGlobalShader::ModifyCompilationEnvironment(Parameters, OutEnvironment);
+	}
+
+	using FParameters = FEmptyShaderParameters;
+};
+
+IMPLEMENT_GLOBAL_SHADER(FHeterogeneousVolumesSparseVoxelMS, "/Engine/Private/HeterogeneousVolumes/HeterogeneousVolumesHardwareRayTracing.usf", "SparseVoxelsMissShader", SF_RayMiss);
+
+class FRenderSingleScatteringWithPreshadingRGS : public FGlobalShader
+{
+	DECLARE_GLOBAL_SHADER(FRenderSingleScatteringWithPreshadingRGS)
+	SHADER_USE_ROOT_PARAMETER_STRUCT(FRenderSingleScatteringWithPreshadingRGS, FGlobalShader)
+
+	using FPermutationDomain = TShaderPermutationDomain<>;
+
+	BEGIN_SHADER_PARAMETER_STRUCT(FParameters, )
+		// Scene 
+		SHADER_PARAMETER_SRV(RaytracingAccelerationStructure, TLAS)
+		SHADER_PARAMETER_STRUCT_REF(FViewUniformShaderParameters, ViewUniformBuffer)
+		SHADER_PARAMETER_STRUCT_INCLUDE(FSceneTextureParameters, SceneTextures)
+
+		// Lighting data
+		SHADER_PARAMETER(int, bApplyEmission)
+		SHADER_PARAMETER(int, bApplyDirectLighting)
+		SHADER_PARAMETER(int, LightType)
+		SHADER_PARAMETER_STRUCT_REF(FDeferredLightUniformStruct, DeferredLight)
+
+		// Sparse Volume
+		SHADER_PARAMETER_RDG_UNIFORM_BUFFER(FSparseVoxelUniformBufferParameters, SparseVoxelUniformBuffer)
+
+		// Volume
+		SHADER_PARAMETER(int, MipLevel)
+
+		// Ray
+		SHADER_PARAMETER(float, MaxTraceDistance)
+		SHADER_PARAMETER(int, MaxStepCount)
+		SHADER_PARAMETER(int, bJitter)
+
+		// Output
+		SHADER_PARAMETER_RDG_TEXTURE_UAV(RWTexture2D<float4>, RWLightingTexture)
+	END_SHADER_PARAMETER_STRUCT()
+
+	static bool ShouldCompilePermutation(const FGlobalShaderPermutationParameters& Parameters)
+	{
+		return ShouldCompileRayTracingShadersForProject(Parameters.Platform) && DoesPlatformSupportHeterogeneousVolumes(Parameters.Platform) &&
+			FDataDrivenShaderPlatformInfo::GetSupportsRayTracingProceduralPrimitive(Parameters.Platform);
+	}
+};
+
+IMPLEMENT_GLOBAL_SHADER(FRenderSingleScatteringWithPreshadingRGS, "/Engine/Private/HeterogeneousVolumes/HeterogeneousVolumesHardwareRayTracing.usf", "RenderSingleScatteringWithPreshadingRGS", SF_RayGen);
+
+FRayTracingLocalShaderBindings* BuildRayTracingMaterialBindings(
+	FRHICommandList& RHICmdList,
+	const FViewInfo& View,
+	FRHIUniformBuffer* UniformBuffer
+)
+{
+	auto Alloc = [&](uint32 Size, uint32 Align)
+	{
+		return RHICmdList.Bypass()
+			? FMemStack::Get().Alloc(Size, Align)
+			: RHICmdList.Alloc(Size, Align);
+	};
+
+	// Allocate bindings
+	const uint32 NumBindings = 1;
+	FRayTracingLocalShaderBindings* Bindings = (FRayTracingLocalShaderBindings*)Alloc(sizeof(FRayTracingLocalShaderBindings) * NumBindings, alignof(FRayTracingLocalShaderBindings));
+
+	// Allocate and assign uniform buffers
+	const uint32 NumUniformBuffers = 1;
+	FRHIUniformBuffer** UniformBufferArray = (FRHIUniformBuffer**)Alloc(sizeof(FRHIUniformBuffer*) * NumUniformBuffers, alignof(FRHIUniformBuffer*));
+	UniformBufferArray[0] = UniformBuffer;
+
+	// Fill bindings
+	for (uint32 BindingIndex = 0; BindingIndex < NumBindings; ++BindingIndex)
+	{
+		// TODO: Declare useful user-data??
+		uint32 UserData = 0;
+
+		FRayTracingLocalShaderBindings Binding = {};
+		Binding.InstanceIndex = 0;
+		Binding.SegmentIndex = 0;
+		Binding.UserData = UserData;
+		Binding.UniformBuffers = UniformBufferArray;
+		Binding.NumUniformBuffers = NumUniformBuffers;
+
+		Bindings[BindingIndex] = Binding;
+	}
+
+	return Bindings;
+}
+
+FRayTracingPipelineState* BuildRayTracingPipelineState(
+	FRHICommandList& RHICmdList,
+	const FViewInfo& View,
+	FRHIRayTracingShader* RayGenerationShader
+)
+{
+	FRayTracingPipelineStateInitializer Initializer;
+	Initializer.MaxPayloadSizeInBytes = 32; // sizeof FSparseVoxelPayload
+
+	// Get the ray tracing materials
+	auto HitGroupShaders = View.ShaderMap->GetShader<FHeterogeneousVolumesSparseVoxelsHitGroup>();
+	FRHIRayTracingShader* HitShaderTable[] = {
+		HitGroupShaders.GetRayTracingShader()
+	};
+	Initializer.SetHitGroupTable(HitShaderTable);
+	// WARNING: Currently hit-group indexing is required to bind uniform buffers to hit-group shaders.
+	Initializer.bAllowHitGroupIndexing = true;
+
+	auto MissShader = View.ShaderMap->GetShader<FHeterogeneousVolumesSparseVoxelMS>();
+	FRHIRayTracingShader* MissShaderTable[] = {
+		MissShader.GetRayTracingShader()
+	};
+	Initializer.SetMissShaderTable(MissShaderTable);
+
+	FRHIRayTracingShader* RayGenShaderTable[] = {
+		RayGenerationShader
+	};
+	Initializer.SetRayGenShaderTable(RayGenShaderTable);
+
+	FRayTracingPipelineState* RayTracingPipelineState = PipelineStateCache::GetAndOrCreateRayTracingPipelineState(RHICmdList, Initializer);
+
+	return RayTracingPipelineState;
+}
+
+void RenderSingleScatteringWithPreshadingHardwareRayTracing(
+	FRDGBuilder& GraphBuilder,
+	// Scene data
+	const FScene* Scene,
+	const FViewInfo& View,
+	const FSceneTextures& SceneTextures,
+	// Light data
+	bool bApplyEmission,
+	bool bApplyDirectLighting,
+	uint32 LightType,
+	const FLightSceneInfo* LightSceneInfo,
+	// Object data
+	const FPrimitiveSceneProxy* PrimitiveSceneProxy,
+	// Sparse voxel data
+	TRDGUniformBufferRef<FSparseVoxelUniformBufferParameters> SparseVoxelUniformBuffer,
+	// Ray tracing data
+	FRayTracingScene& RayTracingScene,
+	// Output
+	FRDGTextureRef& HeterogeneousVolumeTexture
+)
+{
+	FRenderSingleScatteringWithPreshadingRGS::FParameters* PassParameters = GraphBuilder.AllocParameters<FRenderSingleScatteringWithPreshadingRGS::FParameters>();
+	{
+		// Scene
+		PassParameters->TLAS = RayTracingScene.GetLayerSRVChecked(ERayTracingSceneLayer::Base);
+		PassParameters->ViewUniformBuffer = View.ViewUniformBuffer;
+		PassParameters->SceneTextures = GetSceneTextureParameters(GraphBuilder, SceneTextures);
+
+		// Light data
+		FDeferredLightUniformStruct DeferredLightUniform;
+		PassParameters->bApplyEmission = bApplyEmission;
+		PassParameters->bApplyDirectLighting = bApplyDirectLighting;
+		if (PassParameters->bApplyDirectLighting)
+		{
+			DeferredLightUniform = GetDeferredLightParameters(View, *LightSceneInfo);
+		}
+		PassParameters->DeferredLight = CreateUniformBufferImmediate(DeferredLightUniform, UniformBuffer_SingleDraw);
+		PassParameters->LightType = LightType;
+
+		// Sparse Voxel data
+		PassParameters->SparseVoxelUniformBuffer = SparseVoxelUniformBuffer;
+
+		// Volume data
+		PassParameters->MipLevel = HeterogeneousVolumes::GetMipLevel();
+
+		// Ray data
+		PassParameters->MaxTraceDistance = HeterogeneousVolumes::GetMaxTraceDistance();
+		PassParameters->MaxStepCount = HeterogeneousVolumes::GetMaxStepCount();
+		PassParameters->bJitter = HeterogeneousVolumes::ShouldJitter();
+
+		// Output
+		PassParameters->RWLightingTexture = GraphBuilder.CreateUAV(HeterogeneousVolumeTexture);
+	}
+
+	FRenderSingleScatteringWithPreshadingRGS::FPermutationDomain PermutationVector;
+	TShaderRef<FRenderSingleScatteringWithPreshadingRGS> RayGenerationShader = View.ShaderMap->GetShader<FRenderSingleScatteringWithPreshadingRGS>(PermutationVector);
+	FIntPoint DispatchResolution = View.ViewRect.Size();
+
+	FString LightName = TEXT("none");
+	if (LightSceneInfo != nullptr)
+	{
+		FSceneRenderer::GetLightNameForDrawEvent(LightSceneInfo->Proxy, LightName);
+	}
+
+	GraphBuilder.AddPass(
+		RDG_EVENT_NAME("RenderSingleScatteringWithPreshadingRGS (Light = %s) %ux%u", *LightName, DispatchResolution.X, DispatchResolution.Y),
+		PassParameters,
+		ERDGPassFlags::Compute | ERDGPassFlags::NeverCull,
+		[
+			PassParameters,
+			&View,
+			&RayTracingScene,
+			RayGenerationShader,
+			DispatchResolution
+		](FRHIRayTracingCommandList& RHICmdList)
+		{
+			// Set ray-gen bindings
+			FRayTracingShaderBindingsWriter GlobalResources;
+			SetShaderParameters(GlobalResources, RayGenerationShader, *PassParameters);
+
+			// Create pipeline
+			FRayTracingPipelineState* RayTracingPipelineState = BuildRayTracingPipelineState(RHICmdList, View, RayGenerationShader.GetRayTracingShader());
+
+			// Set hit-group bindings
+#if 0
+			const uint32 NumBindings = 1;
+			FRayTracingLocalShaderBindings* Bindings = BuildRayTracingMaterialBindings(RHICmdList, View, PassParameters->SparseVoxelUniformBuffer->GetRHI());
+			const bool bCopyDataToInlineStorage = true;
+			RHICmdList.SetRayTracingHitGroups(RayTracingScene, RayTracingPipelineState, NumBindings, Bindings, bCopyDataToInlineStorage);
+#else
+			uint32 InstanceIndex = 0;
+			uint32 SegmentIndex = 0;
+			uint32 ShaderSlot = 0;
+			uint32 HitGroupIndex = 0;
+			uint32 LooseParameterDataSize = 0;
+			void* LooseParameterData = nullptr;
+			uint32 UserData = 0;
+			uint32 NumUniformBuffers = 1;
+			FRHIUniformBuffer* UniformBuffers[] = {
+				PassParameters->SparseVoxelUniformBuffer->GetRHI()
+			};
+			RHICmdList.SetRayTracingHitGroup(
+				RayTracingScene.GetRHIRayTracingSceneChecked(),
+				InstanceIndex, SegmentIndex, ShaderSlot,
+				RayTracingPipelineState, HitGroupIndex,
+				NumUniformBuffers, UniformBuffers,
+				LooseParameterDataSize, LooseParameterData,
+				UserData
+			);
+
+			uint32 ShaderIndexInPipeline = 0;
+			RHICmdList.SetRayTracingMissShader(
+				RayTracingScene.GetRHIRayTracingSceneChecked(),
+				ShaderSlot,
+				RayTracingPipelineState,
+				ShaderIndexInPipeline,
+				NumUniformBuffers,
+				UniformBuffers,
+				UserData
+			);
+#endif
+
+			// Dispatch
+			RHICmdList.RayTraceDispatch(
+				RayTracingPipelineState,
+				RayGenerationShader.GetRayTracingShader(),
+				RayTracingScene.GetRHIRayTracingSceneChecked(),
+				GlobalResources,
+				DispatchResolution.X, DispatchResolution.Y);
+		}
+	);
+}
+
+#endif // RHI_RAYTRACING
diff --git a/Engine/Source/Runtime/Renderer/Private/HeterogeneousVolumes/HeterogeneousVolumesLiveShadingPipeline.cpp b/Engine/Source/Runtime/Renderer/Private/HeterogeneousVolumes/HeterogeneousVolumesLiveShadingPipeline.cpp
new file mode 100644
--- /dev/null
+++ b/Engine/Source/Runtime/Renderer/Private/HeterogeneousVolumes/HeterogeneousVolumesLiveShadingPipeline.cpp
@@ -0,0 +1,299 @@
+// Copyright Epic Games, Inc. All Rights Reserved.
+
+#include "HeterogeneousVolumes.h"
+
+#include "PixelShaderUtils.h"
+#include "RayTracingDefinitions.h"
+#include "RayTracingInstance.h"
+#include "RayTracingInstanceBufferUtil.h"
+#include "RendererPrivate.h"
+#include "ScenePrivate.h"
+#include "SceneManagement.h"
+
+class FRenderSingleScatteringWithLiveShadingCS : public FMeshMaterialShader
+{
+	DECLARE_SHADER_TYPE(FRenderSingleScatteringWithLiveShadingCS, MeshMaterial);
+
+	BEGIN_SHADER_PARAMETER_STRUCT(FParameters, )
+		// Scene data
+		SHADER_PARAMETER_STRUCT_REF(FViewUniformShaderParameters, View)
+		SHADER_PARAMETER_STRUCT_INCLUDE(FSceneTextureParameters, SceneTextures)
+
+		// Light data
+		SHADER_PARAMETER(int, bApplyEmission)
+		SHADER_PARAMETER(int, bApplyDirectLighting)
+		SHADER_PARAMETER(int, LightType)
+		SHADER_PARAMETER_STRUCT_REF(FDeferredLightUniformStruct, DeferredLight)
+
+		// Shadow data
+		SHADER_PARAMETER(float, ShadowStepFactor)
+
+		// Object data
+		SHADER_PARAMETER(FMatrix44f, LocalToWorld)
+		SHADER_PARAMETER(FMatrix44f, WorldToLocal)
+		SHADER_PARAMETER(FVector3f, LocalBoundsOrigin)
+		SHADER_PARAMETER(FVector3f, LocalBoundsExtent)
+		SHADER_PARAMETER(int32, PrimitiveId)
+
+		// Ray data
+		SHADER_PARAMETER(float, MaxTraceDistance)
+		SHADER_PARAMETER(float, StepSize)
+		SHADER_PARAMETER(int, MaxStepCount)
+		SHADER_PARAMETER(int, bJitter)
+
+		// Dispatch data
+		SHADER_PARAMETER(FIntVector, GroupCount)
+
+		// Output
+		SHADER_PARAMETER_RDG_TEXTURE_UAV(RWTexture2D<float4>, RWLightingTexture)
+		//SHADER_PARAMETER_RDG_BUFFER_UAV(RWStructuredBuffer<Volumes::FDebugOutput>, RWDebugOutputBuffer)
+		END_SHADER_PARAMETER_STRUCT()
+
+		FRenderSingleScatteringWithLiveShadingCS() = default;
+
+	FRenderSingleScatteringWithLiveShadingCS(
+		const FMeshMaterialShaderType::CompiledShaderInitializerType& Initializer
+	)
+		: FMeshMaterialShader(Initializer)
+	{
+		Bindings.BindForLegacyShaderParameters(
+			this,
+			Initializer.PermutationId,
+			Initializer.ParameterMap,
+			*FParameters::FTypeInfo::GetStructMetadata(),
+			// Don't require full bindings, we use FMaterialShader::SetParameters
+			false
+		);
+	}
+
+	static bool ShouldCompilePermutation(
+		const FMaterialShaderPermutationParameters& Parameters
+	)
+	{
+		return DoesPlatformSupportHeterogeneousVolumes(Parameters.Platform)
+			&& (Parameters.MaterialParameters.MaterialDomain == MD_Volume)
+			&& Parameters.MaterialParameters.bIsUsedWithNiagaraMeshParticles;
+	}
+
+	static FPermutationDomain RemapPermutation(FPermutationDomain PermutationVector)
+	{
+		return PermutationVector;
+	}
+
+	static void ModifyCompilationEnvironment(
+		const FMaterialShaderPermutationParameters& Parameters,
+		FShaderCompilerEnvironment& OutEnvironment
+	)
+	{
+		FMaterialShader::ModifyCompilationEnvironment(Parameters, OutEnvironment);
+		OutEnvironment.SetDefine(TEXT("THREADGROUP_SIZE_1D"), GetThreadGroupSize1D());
+		OutEnvironment.SetDefine(TEXT("THREADGROUP_SIZE_2D"), GetThreadGroupSize2D());
+
+		// This shader takes a very long time to compile with FXC, so we pre-compile it with DXC first and then forward the optimized HLSL to FXC.
+		OutEnvironment.CompilerFlags.Add(CFLAG_PrecompileWithDXC);
+		OutEnvironment.CompilerFlags.Add(CFLAG_AllowTypedUAVLoads);
+
+		OutEnvironment.SetDefine(TEXT("GET_PRIMITIVE_DATA_OVERRIDE"), 1);
+	}
+
+	void SetParameters(
+		FRHIComputeCommandList& RHICmdList,
+		FRHIComputeShader* ShaderRHI,
+		const FViewInfo& View,
+		const FMaterialRenderProxy* MaterialProxy,
+		const FMaterial& Material
+	)
+	{
+		FMaterialShader::SetViewParameters(RHICmdList, ShaderRHI, View, View.ViewUniformBuffer);
+		FMaterialShader::SetParameters(RHICmdList, ShaderRHI, MaterialProxy, Material, View);
+	}
+
+	static int32 GetThreadGroupSize1D() { return GetThreadGroupSize2D() * GetThreadGroupSize2D(); }
+	static int32 GetThreadGroupSize2D() { return 8; }
+};
+
+IMPLEMENT_MATERIAL_SHADER_TYPE(, FRenderSingleScatteringWithLiveShadingCS, TEXT("/Engine/Private/HeterogeneousVolumes/HeterogeneousVolumesLiveShadingPipeline.usf"), TEXT("RenderSingleScatteringWithLiveShadingCS"), SF_Compute);
+
+void RenderSingleScatteringWithLiveShading(
+	FRDGBuilder& GraphBuilder,
+	// Scene data
+	const FScene* Scene,
+	const FViewInfo& View,
+	const FSceneTextures& SceneTextures,
+	// Light data
+	bool bApplyEmission,
+	bool bApplyDirectLighting,
+	uint32 LightType,
+	const FLightSceneInfo* LightSceneInfo,
+	// Object data
+	const FPrimitiveSceneProxy* PrimitiveSceneProxy,
+	const FMaterialRenderProxy* MaterialRenderProxy,
+	const int32 PrimitiveId,
+	const FBoxSphereBounds LocalBoxSphereBounds,
+	// Output
+	FRDGTextureRef& HeterogeneousVolumeTexture
+)
+{
+	const FMaterial& Material = MaterialRenderProxy->GetMaterialWithFallback(View.GetFeatureLevel(), MaterialRenderProxy);
+	if (Material.GetMaterialDomain() != MD_Volume)
+	{
+		return;
+	}
+
+	uint32 GroupCountX = FMath::DivideAndRoundUp(View.ViewRect.Size().X, FRenderSingleScatteringWithLiveShadingCS::GetThreadGroupSize2D());
+	uint32 GroupCountY = FMath::DivideAndRoundUp(View.ViewRect.Size().Y, FRenderSingleScatteringWithLiveShadingCS::GetThreadGroupSize2D());
+	FIntVector GroupCount = FIntVector(GroupCountX, GroupCountY, 1);
+
+#if 0
+	FRDGBufferRef DebugOutputBuffer = GraphBuilder.CreateBuffer(
+		FRDGBufferDesc::CreateStructuredDesc(sizeof(Volumes::FDebugOutput), GroupCountX * GroupCountY * FRenderSingleScatteringWithLiveShadingCS::GetThreadGroupSize2D()),
+		TEXT("Lumen.Reflections.TraceDataPacked"));
+	AddClearUAVPass(GraphBuilder, GraphBuilder.CreateUAV(DebugOutputBuffer), 0);
+#endif
+
+	FRenderSingleScatteringWithLiveShadingCS::FParameters* PassParameters = GraphBuilder.AllocParameters<FRenderSingleScatteringWithLiveShadingCS::FParameters>();
+	{
+		// Scene data
+		PassParameters->View = View.ViewUniformBuffer;
+		PassParameters->SceneTextures = GetSceneTextureParameters(GraphBuilder, SceneTextures);
+
+		// Light data
+		FDeferredLightUniformStruct DeferredLightUniform;
+		PassParameters->bApplyEmission = bApplyEmission;
+		PassParameters->bApplyDirectLighting = bApplyDirectLighting;
+		if (PassParameters->bApplyDirectLighting)
+		{
+			DeferredLightUniform = GetDeferredLightParameters(View, *LightSceneInfo);
+		}
+		PassParameters->DeferredLight = CreateUniformBufferImmediate(DeferredLightUniform, UniformBuffer_SingleDraw);
+		PassParameters->LightType = LightType;
+		PassParameters->ShadowStepFactor = HeterogeneousVolumes::GetShadowStepFactor();
+
+		// Object data
+		FMatrix44f LocalToWorld = FMatrix44f(PrimitiveSceneProxy->GetLocalToWorld());
+		PassParameters->LocalToWorld = LocalToWorld;
+		PassParameters->WorldToLocal = LocalToWorld.Inverse();
+		PassParameters->LocalBoundsOrigin = FVector3f(LocalBoxSphereBounds.Origin);
+		PassParameters->LocalBoundsExtent = FVector3f(LocalBoxSphereBounds.BoxExtent);
+		PassParameters->PrimitiveId = PrimitiveId;
+
+		// Ray data
+		PassParameters->MaxTraceDistance = HeterogeneousVolumes::GetMaxTraceDistance();
+		PassParameters->StepSize = HeterogeneousVolumes::GetStepSize();
+		PassParameters->MaxStepCount = HeterogeneousVolumes::GetMaxStepCount();
+		PassParameters->bJitter = HeterogeneousVolumes::ShouldJitter();
+
+		// Dispatch data
+		PassParameters->GroupCount = GroupCount;
+
+		// Output
+		PassParameters->RWLightingTexture = GraphBuilder.CreateUAV(HeterogeneousVolumeTexture);
+		//PassParameters->RWDebugOutputBuffer = GraphBuilder.CreateUAV(DebugOutputBuffer);
+	}
+
+	FString LightName = TEXT("none");
+	if (LightSceneInfo != nullptr)
+	{
+		FSceneRenderer::GetLightNameForDrawEvent(LightSceneInfo->Proxy, LightName);
+	}
+
+	GraphBuilder.AddPass(
+		RDG_EVENT_NAME("RenderSingleScatteringWithLiveShadingCS (Light = %s)", *LightName),
+		PassParameters,
+		ERDGPassFlags::Compute,
+		// Why is scene explicitly copied??
+		[PassParameters, LocalScene = Scene, &View, MaterialRenderProxy, &Material, GroupCount](FRHIComputeCommandList& RHICmdList)
+		{
+			FRenderSingleScatteringWithLiveShadingCS::FPermutationDomain PermutationVector;
+			TShaderRef<FRenderSingleScatteringWithLiveShadingCS> ComputeShader = Material.GetShader<FRenderSingleScatteringWithLiveShadingCS>(&FLocalVertexFactory::StaticType, PermutationVector, false);
+
+			if (!ComputeShader.IsNull())
+			{
+				FMeshPassProcessorRenderState DrawRenderState;
+
+				FMeshMaterialShaderElementData ShaderElementData;
+				ShaderElementData.FadeUniformBuffer = GDistanceCullFadedInUniformBuffer.GetUniformBufferRHI();
+				ShaderElementData.DitherUniformBuffer = GDitherFadedInUniformBuffer.GetUniformBufferRHI();
+
+				FMeshProcessorShaders PassShaders;
+				PassShaders.ComputeShader = ComputeShader;
+
+				FMeshDrawShaderBindings ShaderBindings;
+				{
+					ShaderBindings.Initialize(PassShaders);
+
+					int32 DataOffset = 0;
+					FMeshDrawSingleShaderBindings SingleShaderBindings = ShaderBindings.GetSingleShaderBindings(SF_Compute, DataOffset);
+					ComputeShader->GetShaderBindings(LocalScene, LocalScene->GetFeatureLevel(), nullptr, *MaterialRenderProxy, Material, DrawRenderState, ShaderElementData, SingleShaderBindings);
+
+					ShaderBindings.Finalize(&PassShaders);
+				}
+				ShaderBindings.SetOnCommandList(RHICmdList, ComputeShader.GetComputeShader());
+
+				FComputeShaderUtils::Dispatch(RHICmdList, ComputeShader, *PassParameters, GroupCount);
+			}
+		}
+	);
+}
+
+void RenderWithLiveShading(
+	FRDGBuilder& GraphBuilder,
+	const FSceneTextures& SceneTextures,
+	const FScene* Scene,
+	const FViewInfo& View,
+	// Object data
+	const FPrimitiveSceneProxy* PrimitiveSceneProxy,
+	const FMaterialRenderProxy* MaterialRenderProxy,
+	const int32 PrimitiveId,
+	const FBoxSphereBounds LocalBoxSphereBounds,
+	// Output
+	FRDGTextureRef& HeterogeneousVolumeRadiance
+)
+{
+	// Light culling
+	TArray<FLightSceneInfoCompact, TInlineAllocator<64>> LightSceneInfoCompact;
+	for (auto LightIt = Scene->Lights.CreateConstIterator(); LightIt; ++LightIt)
+	{
+		if (LightIt->AffectsPrimitive(PrimitiveSceneProxy->GetBounds(), PrimitiveSceneProxy))
+		{
+			LightSceneInfoCompact.Add(*LightIt);
+		}
+	}
+
+	// Light loop:
+	int32 NumPasses = FMath::Max(LightSceneInfoCompact.Num(), 1);
+	for (int32 PassIndex = 0; PassIndex < NumPasses; ++PassIndex)
+	{
+		bool bApplyEmission = PassIndex == 0;
+		bool bApplyDirectLighting = !LightSceneInfoCompact.IsEmpty();
+
+		uint32 LightType = 0;
+		FLightSceneInfo* LightSceneInfo = nullptr;
+		if (bApplyDirectLighting)
+		{
+			LightType = LightSceneInfoCompact[PassIndex].LightType;
+			LightSceneInfo = LightSceneInfoCompact[PassIndex].LightSceneInfo;
+			check(LightSceneInfo != nullptr);
+		}
+
+		RenderSingleScatteringWithLiveShading(
+			GraphBuilder,
+			// Scene data
+			Scene,
+			View,
+			SceneTextures,
+			// Light data
+			bApplyEmission,
+			bApplyDirectLighting,
+			LightType,
+			LightSceneInfo,
+			// Object data
+			PrimitiveSceneProxy,
+			MaterialRenderProxy,
+			PrimitiveId,
+			LocalBoxSphereBounds,
+			// Output
+			HeterogeneousVolumeRadiance
+		);
+	}
+}
\ No newline at end of file
diff --git a/Engine/Source/Runtime/Renderer/Private/HeterogeneousVolumes/HeterogeneousVolumesMaterialBakingPipeline.cpp b/Engine/Source/Runtime/Renderer/Private/HeterogeneousVolumes/HeterogeneousVolumesMaterialBakingPipeline.cpp
new file mode 100644
--- /dev/null
+++ b/Engine/Source/Runtime/Renderer/Private/HeterogeneousVolumes/HeterogeneousVolumesMaterialBakingPipeline.cpp
@@ -0,0 +1,209 @@
+// Copyright Epic Games, Inc. All Rights Reserved.
+
+#include "HeterogeneousVolumes.h"
+
+#include "PixelShaderUtils.h"
+#include "RayTracingDefinitions.h"
+#include "RayTracingInstance.h"
+#include "RayTracingInstanceBufferUtil.h"
+#include "RendererPrivate.h"
+#include "ScenePrivate.h"
+#include "SceneManagement.h"
+
+class FHeterogeneousVolumesBakeMaterialCS : public FMeshMaterialShader
+{
+	DECLARE_SHADER_TYPE(FHeterogeneousVolumesBakeMaterialCS, MeshMaterial);
+
+	BEGIN_SHADER_PARAMETER_STRUCT(FParameters, )
+		// Scene data
+		SHADER_PARAMETER_STRUCT_REF(FViewUniformShaderParameters, View)
+
+		// Object data
+		SHADER_PARAMETER(FMatrix44f, LocalToWorld)
+		SHADER_PARAMETER(FMatrix44f, WorldToLocal)
+		SHADER_PARAMETER(FVector3f, LocalBoundsOrigin)
+		SHADER_PARAMETER(FVector3f, LocalBoundsExtent)
+		SHADER_PARAMETER(int32, PrimitiveId)
+
+		// Volume data
+		SHADER_PARAMETER(FIntVector, VolumeResolution)
+
+		// Sample data
+		SHADER_PARAMETER(int, bJitter)
+
+		// Dispatch data
+		SHADER_PARAMETER(FIntVector, GroupCount)
+
+		// Output
+		SHADER_PARAMETER_RDG_TEXTURE_UAV(RWTexture2D<float3>, RWExtinctionTexture)
+		SHADER_PARAMETER_RDG_TEXTURE_UAV(RWTexture2D<float3>, RWEmissionTexture)
+		SHADER_PARAMETER_RDG_TEXTURE_UAV(RWTexture2D<float3>, RWAlbedoTexture)
+		//SHADER_PARAMETER_RDG_BUFFER_UAV(RWStructuredBuffer<Volumes::FDebugOutput>, RWDebugOutputBuffer)
+	END_SHADER_PARAMETER_STRUCT()
+
+	FHeterogeneousVolumesBakeMaterialCS() = default;
+
+	FHeterogeneousVolumesBakeMaterialCS(
+		const FMeshMaterialShaderType::CompiledShaderInitializerType& Initializer
+	)
+		: FMeshMaterialShader(Initializer)
+	{
+		Bindings.BindForLegacyShaderParameters(
+			this,
+			Initializer.PermutationId,
+			Initializer.ParameterMap,
+			*FParameters::FTypeInfo::GetStructMetadata(),
+			// Don't require full bindings, we use FMaterialShader::SetParameters
+			false
+		);
+	}
+
+	static bool ShouldCompilePermutation(
+		const FMaterialShaderPermutationParameters& Parameters
+	)
+	{
+		return DoesPlatformSupportHeterogeneousVolumes(Parameters.Platform)
+			&& (Parameters.MaterialParameters.MaterialDomain == MD_Volume)
+			&& Parameters.MaterialParameters.bIsUsedWithNiagaraMeshParticles;
+	}
+
+	static FPermutationDomain RemapPermutation(FPermutationDomain PermutationVector)
+	{
+		return PermutationVector;
+	}
+
+	static void ModifyCompilationEnvironment(
+		const FMaterialShaderPermutationParameters& Parameters,
+		FShaderCompilerEnvironment& OutEnvironment
+	)
+	{
+		FMaterialShader::ModifyCompilationEnvironment(Parameters, OutEnvironment);
+		OutEnvironment.SetDefine(TEXT("THREADGROUP_SIZE_1D"), GetThreadGroupSize1D());
+		OutEnvironment.SetDefine(TEXT("THREADGROUP_SIZE_2D"), GetThreadGroupSize2D());
+		OutEnvironment.SetDefine(TEXT("THREADGROUP_SIZE_3D"), GetThreadGroupSize3D());
+
+		// This shader takes a very long time to compile with FXC, so we pre-compile it with DXC first and then forward the optimized HLSL to FXC.
+		OutEnvironment.CompilerFlags.Add(CFLAG_PrecompileWithDXC);
+		OutEnvironment.CompilerFlags.Add(CFLAG_AllowTypedUAVLoads);
+
+		OutEnvironment.SetDefine(TEXT("GET_PRIMITIVE_DATA_OVERRIDE"), 1);
+	}
+
+	void SetParameters(
+		FRHIComputeCommandList& RHICmdList,
+		FRHIComputeShader* ShaderRHI,
+		const FViewInfo& View,
+		const FMaterialRenderProxy* MaterialProxy,
+		const FMaterial& Material
+	)
+	{
+		FMaterialShader::SetViewParameters(RHICmdList, ShaderRHI, View, View.ViewUniformBuffer);
+		FMaterialShader::SetParameters(RHICmdList, ShaderRHI, MaterialProxy, Material, View);
+	}
+
+	static int32 GetThreadGroupSize1D() { return GetThreadGroupSize2D() * GetThreadGroupSize2D(); }
+	static int32 GetThreadGroupSize2D() { return 8; }
+	static int32 GetThreadGroupSize3D() { return 4; }
+};
+
+IMPLEMENT_MATERIAL_SHADER_TYPE(, FHeterogeneousVolumesBakeMaterialCS, TEXT("/Engine/Private/HeterogeneousVolumes/HeterogeneousVolumesMaterialBakingPipeline.usf"), TEXT("HeterogeneousVolumesBakeMaterialCS"), SF_Compute);
+
+void ComputeHeterogeneousVolumeBakeMaterial(
+	FRDGBuilder& GraphBuilder,
+	// Scene data
+	const FScene* Scene,
+	const FViewInfo& View,
+	// Object data
+	const FPrimitiveSceneProxy* PrimitiveSceneProxy,
+	const FMaterialRenderProxy* MaterialRenderProxy,
+	const int32 PrimitiveId,
+	const FBoxSphereBounds LocalBoxSphereBounds,
+	// Volume data
+	FIntVector VolumeResolution,
+	// Output
+	FRDGTextureRef& HeterogeneousVolumeExtinctionTexture,
+	FRDGTextureRef& HeterogeneousVolumeEmissionTexture,
+	FRDGTextureRef& HeterogeneousVolumeAlbedoTexture
+)
+{
+	const FMaterial& Material = MaterialRenderProxy->GetMaterialWithFallback(View.GetFeatureLevel(), MaterialRenderProxy);
+	check(Material.GetMaterialDomain() == MD_Volume);
+
+	// FMath::DivideAndRoundUp() is not compatible with FIntVector?
+	FIntVector GroupSize = FIntVector(FHeterogeneousVolumesBakeMaterialCS::GetThreadGroupSize3D());
+	FIntVector GroupCount = (VolumeResolution + GroupSize - FIntVector(1)) / FHeterogeneousVolumesBakeMaterialCS::GetThreadGroupSize3D();
+
+#if 0
+	FRDGBufferRef DebugOutputBuffer = GraphBuilder.CreateBuffer(
+		FRDGBufferDesc::CreateStructuredDesc(sizeof(Volumes::FDebugOutput), GetVoxelCount(VolumeResolution)),
+		TEXT("HeterogeneousVolumes.DebugOutputBuffer"));
+	AddClearUAVPass(GraphBuilder, GraphBuilder.CreateUAV(DebugOutputBuffer), 0);
+#endif
+
+	FHeterogeneousVolumesBakeMaterialCS::FParameters* PassParameters = GraphBuilder.AllocParameters<FHeterogeneousVolumesBakeMaterialCS::FParameters>();
+	{
+		// Scene data
+		PassParameters->View = View.ViewUniformBuffer;
+
+		// Object data
+		FMatrix44f LocalToWorld = FMatrix44f(PrimitiveSceneProxy->GetLocalToWorld());
+		PassParameters->LocalToWorld = LocalToWorld;
+		PassParameters->WorldToLocal = LocalToWorld.Inverse();
+		PassParameters->LocalBoundsOrigin = FVector3f(LocalBoxSphereBounds.Origin);
+		PassParameters->LocalBoundsExtent = FVector3f(LocalBoxSphereBounds.BoxExtent);
+		PassParameters->PrimitiveId = PrimitiveId;
+
+		// Volume data
+		PassParameters->VolumeResolution = VolumeResolution;
+
+		// Sample data
+		PassParameters->bJitter = HeterogeneousVolumes::ShouldJitter();
+
+		// Dispatch data
+		PassParameters->GroupCount = GroupCount;
+
+		PassParameters->RWExtinctionTexture = GraphBuilder.CreateUAV(HeterogeneousVolumeExtinctionTexture);
+		PassParameters->RWEmissionTexture = GraphBuilder.CreateUAV(HeterogeneousVolumeEmissionTexture);
+		PassParameters->RWAlbedoTexture = GraphBuilder.CreateUAV(HeterogeneousVolumeAlbedoTexture);
+		//PassParameters->RWDebugOutputBuffer = GraphBuilder.CreateUAV(DebugOutputBuffer);
+	}
+
+	GraphBuilder.AddPass(
+		RDG_EVENT_NAME("HeterogeneousVolumesBakedMaterialCS"),
+		PassParameters,
+		ERDGPassFlags::Compute,
+		// Why is scene explicitly copied??
+		[PassParameters, LocalScene = Scene, &View, MaterialRenderProxy, &Material, GroupCount](FRHIComputeCommandList& RHICmdList)
+		{
+			FHeterogeneousVolumesBakeMaterialCS::FPermutationDomain PermutationVector;
+			TShaderRef<FHeterogeneousVolumesBakeMaterialCS> ComputeShader = Material.GetShader<FHeterogeneousVolumesBakeMaterialCS>(&FLocalVertexFactory::StaticType, PermutationVector, false);
+
+			if (!ComputeShader.IsNull())
+			{
+				FMeshPassProcessorRenderState DrawRenderState;
+
+				FMeshMaterialShaderElementData ShaderElementData;
+				ShaderElementData.FadeUniformBuffer = GDistanceCullFadedInUniformBuffer.GetUniformBufferRHI();
+				ShaderElementData.DitherUniformBuffer = GDitherFadedInUniformBuffer.GetUniformBufferRHI();
+
+				FMeshProcessorShaders PassShaders;
+				PassShaders.ComputeShader = ComputeShader;
+
+				FMeshDrawShaderBindings ShaderBindings;
+				{
+					ShaderBindings.Initialize(PassShaders);
+
+					int32 DataOffset = 0;
+					FMeshDrawSingleShaderBindings SingleShaderBindings = ShaderBindings.GetSingleShaderBindings(SF_Compute, DataOffset);
+					ComputeShader->GetShaderBindings(LocalScene, LocalScene->GetFeatureLevel(), nullptr, *MaterialRenderProxy, Material, DrawRenderState, ShaderElementData, SingleShaderBindings);
+
+					ShaderBindings.Finalize(&PassShaders);
+				}
+				ShaderBindings.SetOnCommandList(RHICmdList, ComputeShader.GetComputeShader());
+
+				FComputeShaderUtils::Dispatch(RHICmdList, ComputeShader, *PassParameters, GroupCount);
+				//FComputeShaderUtils::DispatchIndirect(RHICmdList, ComputeShader, *PassParameters, IndirectArgs->GetIndirectRHICallBuffer(), IndirectOffset);
+			}
+		}
+	);
+}
diff --git a/Engine/Source/Runtime/Renderer/Private/HeterogeneousVolumes/HeterogeneousVolumesPreshadingPipeline.cpp b/Engine/Source/Runtime/Renderer/Private/HeterogeneousVolumes/HeterogeneousVolumesPreshadingPipeline.cpp
new file mode 100644
--- /dev/null
+++ b/Engine/Source/Runtime/Renderer/Private/HeterogeneousVolumes/HeterogeneousVolumesPreshadingPipeline.cpp
@@ -0,0 +1,872 @@
+// Copyright Epic Games, Inc. All Rights Reserved.
+
+#include "HeterogeneousVolumes.h"
+
+#include "PixelShaderUtils.h"
+#include "RayTracingDefinitions.h"
+#include "RayTracingInstance.h"
+#include "RayTracingInstanceBufferUtil.h"
+#include "RendererPrivate.h"
+#include "ScenePrivate.h"
+#include "SceneManagement.h"
+
+class FGenerateRayMarchingTiles : public FGlobalShader
+{
+	DECLARE_GLOBAL_SHADER(FGenerateRayMarchingTiles);
+	SHADER_USE_PARAMETER_STRUCT(FGenerateRayMarchingTiles, FGlobalShader);
+
+	class FDebugDim : SHADER_PERMUTATION_BOOL("DIM_DEBUG");
+	class FVoxelCullingDim : SHADER_PERMUTATION_BOOL("DIM_VOXEL_CULLING");
+	using FPermutationDomain = TShaderPermutationDomain<FDebugDim, FVoxelCullingDim>;
+
+	BEGIN_SHADER_PARAMETER_STRUCT(FParameters, )
+		// Scene data
+		SHADER_PARAMETER_STRUCT_REF(FViewUniformShaderParameters, View)
+		SHADER_PARAMETER_STRUCT_INCLUDE(FSceneTextureParameters, SceneTextures)
+
+		// Object data
+		SHADER_PARAMETER(FMatrix44f, LocalToWorld)
+		SHADER_PARAMETER(FMatrix44f, WorldToLocal)
+
+		// Sparse voxel data
+		SHADER_PARAMETER_RDG_UNIFORM_BUFFER(FSparseVoxelUniformBufferParameters, SparseVoxelUniformBuffer)
+
+		// Ray data
+		SHADER_PARAMETER(float, MaxTraceDistance)
+		SHADER_PARAMETER(float, StepSize)
+		SHADER_PARAMETER(int, MaxStepCount)
+		SHADER_PARAMETER(int, bJitter)
+
+		// Dispatch data
+		SHADER_PARAMETER(FIntVector, GroupCount)
+
+		// Debug Output
+		SHADER_PARAMETER_RDG_BUFFER_UAV(RWStructuredBuffer<Volumes::FRayMarchingDebug>, RWRayMarchingDebugBuffer)
+
+		// Output
+		SHADER_PARAMETER_RDG_BUFFER_UAV(RWBuffer<uint>, RWNumRayMarchingTilesBuffer)
+		SHADER_PARAMETER_RDG_BUFFER_UAV(RWStructuredBuffer<Volumes::FRayMarchingTile>, RWRayMarchingTilesBuffer)
+
+		SHADER_PARAMETER_RDG_BUFFER_UAV(RWBuffer<uint>, RWNumVoxelsPerTileBuffer)
+		SHADER_PARAMETER_RDG_BUFFER_UAV(RWStructuredBuffer<FVoxelDataPacked>, RWVoxelsPerTileBuffer)
+	END_SHADER_PARAMETER_STRUCT()
+
+	static bool ShouldCompilePermutation(
+		const FGlobalShaderPermutationParameters& Parameters
+	)
+	{
+		return DoesPlatformSupportHeterogeneousVolumes(Parameters.Platform);
+	}
+
+	static void ModifyCompilationEnvironment(
+		const FGlobalShaderPermutationParameters& Parameters,
+		FShaderCompilerEnvironment& OutEnvironment
+	)
+	{
+		FMaterialShader::ModifyCompilationEnvironment(Parameters, OutEnvironment);
+		OutEnvironment.SetDefine(TEXT("THREADGROUP_SIZE_1D"), GetThreadGroupSize1D());
+		OutEnvironment.SetDefine(TEXT("THREADGROUP_SIZE_2D"), GetThreadGroupSize2D());
+
+		// This shader takes a very long time to compile with FXC, so we pre-compile it with DXC first and then forward the optimized HLSL to FXC.
+		OutEnvironment.CompilerFlags.Add(CFLAG_AllowTypedUAVLoads);
+	}
+
+	static int32 GetThreadGroupSize1D() { return GetThreadGroupSize2D() * GetThreadGroupSize2D(); }
+	static int32 GetThreadGroupSize2D() { return 8; }
+};
+
+IMPLEMENT_MATERIAL_SHADER_TYPE(, FGenerateRayMarchingTiles, TEXT("/Engine/Private/HeterogeneousVolumes/HeterogeneousVolumesPreshadingPipeline.usf"), TEXT("GenerateRayMarchingTiles"), SF_Compute);
+
+namespace HeterogeneousVolumes
+{
+	struct FRayMarchingTile
+	{
+		FIntPoint PixelOffset;
+		uint32 Voxels[2];
+
+		uint32 Id;
+		uint32 Padding[3];
+	};
+
+	struct FRayMarchingDebug
+	{
+		FVector4f Planes[5];
+		FVector4f BBox[2];
+
+		float Padding[4];
+	};
+}
+
+void GenerateRayMarchingTiles(
+	FRDGBuilder& GraphBuilder,
+	// Scene data
+	const FScene* Scene,
+	const FViewInfo& View,
+	const FSceneTextures& SceneTextures,
+	// Object data
+	const FPrimitiveSceneProxy* PrimitiveSceneProxy,
+	// Sparse voxel data
+	FRDGBufferRef NumVoxelsBuffer,
+	const TRDGUniformBufferRef<FSparseVoxelUniformBufferParameters>& SparseVoxelUniformBuffer,
+	// Output
+	FRDGBufferRef& NumRayMarchingTilesBuffer,
+	FRDGBufferRef& RayMarchingTilesBuffer,
+	FRDGBufferRef& VoxelsPerTileBuffer
+)
+{
+	uint32 GroupCountX = FMath::DivideAndRoundUp(View.ViewRect.Size().X, FGenerateRayMarchingTiles::GetThreadGroupSize2D());
+	uint32 GroupCountY = FMath::DivideAndRoundUp(View.ViewRect.Size().Y, FGenerateRayMarchingTiles::GetThreadGroupSize2D());
+	FIntVector GroupCount = FIntVector(GroupCountX, GroupCountY, 1);
+	uint32 NumTiles = GroupCountX * GroupCountY;
+
+	NumRayMarchingTilesBuffer = GraphBuilder.CreateBuffer(
+		FRDGBufferDesc::CreateIndirectDesc<FRHIDispatchIndirectParameters>(1),
+		TEXT("HeterogeneousVolume.NumRayMarchingTilesBuffer")
+	);
+	AddClearUAVPass(GraphBuilder, GraphBuilder.CreateUAV(NumRayMarchingTilesBuffer, PF_R32_UINT), 0);
+
+	RayMarchingTilesBuffer = GraphBuilder.CreateBuffer(
+		FRDGBufferDesc::CreateStructuredDesc(sizeof(HeterogeneousVolumes::FRayMarchingTile), NumTiles),
+		TEXT("HeterogeneousVolumes.RayMarchingTileBuffer")
+	);
+
+	FRDGBufferRef RayMarchingDebugBuffer = GraphBuilder.CreateBuffer(
+		FRDGBufferDesc::CreateStructuredDesc(sizeof(HeterogeneousVolumes::FRayMarchingDebug), NumTiles),
+		TEXT("HeterogeneousVolume.RayMarchingDebugBuffer")
+	);
+
+	FRDGBufferRef NumVoxelsPerTileBuffer = GraphBuilder.CreateBuffer(
+		FRDGBufferDesc::CreateBufferDesc(sizeof(uint32), NumTiles),
+		TEXT("HeterogeneousVolume.NumVoxelsPerTileBuffer")
+	);
+	AddClearUAVPass(GraphBuilder, GraphBuilder.CreateUAV(NumVoxelsPerTileBuffer, PF_R32_UINT), 0);
+
+	FIntVector VolumeResolution = SparseVoxelUniformBuffer->GetParameters()->VolumeResolution;
+	uint32 SparseMipLevel = SparseVoxelUniformBuffer->GetParameters()->MipLevel;
+	FIntVector SparseVolumeResolution = FIntVector(VolumeResolution.X >> SparseMipLevel,
+		VolumeResolution.Y >> SparseMipLevel,
+		VolumeResolution.Z >> SparseMipLevel
+	);
+
+	// TODO: Tight frustum culling guarantees no more than Length(SparseVolumeResolution) but approximate intersection cannot guarantee even L1 distance..
+	//uint32 DiagonalLength = FMath::CeilToInt(FMath::Sqrt(SparseVolumeResolution.X * SparseVolumeResolution.Y * SparseVolumeResolution.Z));
+	uint32 DiagonalLength = SparseVolumeResolution.X * SparseVolumeResolution.Y * SparseVolumeResolution.Z;
+
+	VoxelsPerTileBuffer = GraphBuilder.CreateBuffer(
+		FRDGBufferDesc::CreateStructuredDesc(sizeof(FVoxelDataPacked), NumTiles * DiagonalLength),
+		TEXT("HeterogeneousVolumes.VoxelsPerTileBuffer")
+	);
+
+	FGenerateRayMarchingTiles::FParameters* PassParameters = GraphBuilder.AllocParameters<FGenerateRayMarchingTiles::FParameters>();
+	{
+		// Scene data
+		PassParameters->View = View.ViewUniformBuffer;
+
+		// Sparse voxel data
+		PassParameters->SparseVoxelUniformBuffer = SparseVoxelUniformBuffer;
+
+		// Ray data
+		PassParameters->MaxTraceDistance = HeterogeneousVolumes::GetMaxTraceDistance();
+		PassParameters->StepSize = HeterogeneousVolumes::GetStepSize();
+		PassParameters->MaxStepCount = HeterogeneousVolumes::GetMaxStepCount();
+		PassParameters->bJitter = HeterogeneousVolumes::ShouldJitter();
+
+		// Dispatch data
+		PassParameters->GroupCount = GroupCount;
+
+		// Debug
+		PassParameters->RWRayMarchingDebugBuffer = GraphBuilder.CreateUAV(RayMarchingDebugBuffer);
+
+		// Output
+		PassParameters->RWNumRayMarchingTilesBuffer = GraphBuilder.CreateUAV(NumRayMarchingTilesBuffer, PF_R32_UINT);
+		PassParameters->RWRayMarchingTilesBuffer = GraphBuilder.CreateUAV(RayMarchingTilesBuffer);
+
+		PassParameters->RWNumVoxelsPerTileBuffer = GraphBuilder.CreateUAV(NumVoxelsPerTileBuffer, PF_R32_UINT);
+		PassParameters->RWVoxelsPerTileBuffer = GraphBuilder.CreateUAV(VoxelsPerTileBuffer);
+	}
+
+	FGenerateRayMarchingTiles::FPermutationDomain PermutationVector;
+	PermutationVector.Set<FGenerateRayMarchingTiles::FDebugDim>(HeterogeneousVolumes::GetDebugMode() != 0);
+	PermutationVector.Set<FGenerateRayMarchingTiles::FVoxelCullingDim>(HeterogeneousVolumes::UseSparseVoxelPerTileCulling());
+
+	TShaderRef<FGenerateRayMarchingTiles> ComputeShader = View.ShaderMap->GetShader<FGenerateRayMarchingTiles>(PermutationVector);
+	FComputeShaderUtils::AddPass(
+		GraphBuilder,
+		RDG_EVENT_NAME("FGenerateRayMarchingTiles"),
+		ComputeShader,
+		PassParameters,
+		GroupCount
+	);
+}
+
+class FRenderSingleScatteringWithPreshadingCS : public FGlobalShader
+{
+	DECLARE_GLOBAL_SHADER(FRenderSingleScatteringWithPreshadingCS);
+	SHADER_USE_PARAMETER_STRUCT(FRenderSingleScatteringWithPreshadingCS, FGlobalShader);
+
+	class FDebugDim : SHADER_PERMUTATION_BOOL("DIM_DEBUG");
+	class FVoxelCullingDim : SHADER_PERMUTATION_BOOL("DIM_VOXEL_CULLING");
+	class FSparseVoxelTracingDim : SHADER_PERMUTATION_BOOL("DIM_SPARSE_VOXEL_TRACING");
+	using FPermutationDomain = TShaderPermutationDomain<FDebugDim, FVoxelCullingDim, FSparseVoxelTracingDim>;
+
+	BEGIN_SHADER_PARAMETER_STRUCT(FParameters, )
+		// Scene data
+		SHADER_PARAMETER_STRUCT_REF(FViewUniformShaderParameters, View)
+		SHADER_PARAMETER_STRUCT_INCLUDE(FSceneTextureParameters, SceneTextures)
+
+		// Light data
+		SHADER_PARAMETER(int, bApplyEmission)
+		SHADER_PARAMETER(int, bApplyDirectLighting)
+		SHADER_PARAMETER(int, LightType)
+		SHADER_PARAMETER_STRUCT_REF(FDeferredLightUniformStruct, DeferredLight)
+
+		// Volume data
+		SHADER_PARAMETER(int, MipLevel)
+
+		// Sparse Volume
+		SHADER_PARAMETER_RDG_UNIFORM_BUFFER(FSparseVoxelUniformBufferParameters, SparseVoxelUniformBuffer)
+
+		// Ray data
+		SHADER_PARAMETER(float, MaxTraceDistance)
+		SHADER_PARAMETER(float, StepSize)
+		SHADER_PARAMETER(int, MaxStepCount)
+		SHADER_PARAMETER(int, bJitter)
+
+		// Ray marching data
+		SHADER_PARAMETER_RDG_BUFFER_SRV(StructuredBuffer<Volumes::FRayMarchingTile>, RayMarchingTilesBuffer)
+		SHADER_PARAMETER_RDG_BUFFER_SRV(StructuredBuffer<FVoxelDataPacked>, VoxelsPerTileBuffer)
+
+		// Indirect args
+		RDG_BUFFER_ACCESS(IndirectArgs, ERHIAccess::IndirectArgs)
+
+		// Output
+		SHADER_PARAMETER_RDG_TEXTURE_UAV(RWTexture2D<float4>, RWLightingTexture)
+		SHADER_PARAMETER_RDG_BUFFER_UAV(RWStructuredBuffer<FVoxelDataPacked>, RWVoxelOutputBuffer)
+	END_SHADER_PARAMETER_STRUCT()
+
+	static bool ShouldCompilePermutation(
+		const FGlobalShaderPermutationParameters& Parameters
+	)
+	{
+		return DoesPlatformSupportHeterogeneousVolumes(Parameters.Platform);
+	}
+
+	static void ModifyCompilationEnvironment(
+		const FGlobalShaderPermutationParameters& Parameters,
+		FShaderCompilerEnvironment& OutEnvironment
+	)
+	{
+		FMaterialShader::ModifyCompilationEnvironment(Parameters, OutEnvironment);
+		OutEnvironment.SetDefine(TEXT("THREADGROUP_SIZE_1D"), GetThreadGroupSize1D());
+		OutEnvironment.SetDefine(TEXT("THREADGROUP_SIZE_2D"), GetThreadGroupSize2D());
+
+		// This shader takes a very long time to compile with FXC, so we pre-compile it with DXC first and then forward the optimized HLSL to FXC.
+		//OutEnvironment.CompilerFlags.Add(CFLAG_PrecompileWithDXC);
+		OutEnvironment.CompilerFlags.Add(CFLAG_AllowTypedUAVLoads);
+
+		OutEnvironment.SetDefine(TEXT("GET_PRIMITIVE_DATA_OVERRIDE"), 1);
+	}
+
+	static int32 GetThreadGroupSize1D() { return GetThreadGroupSize2D() * GetThreadGroupSize2D(); }
+	static int32 GetThreadGroupSize2D() { return 8; }
+};
+
+IMPLEMENT_MATERIAL_SHADER_TYPE(, FRenderSingleScatteringWithPreshadingCS, TEXT("/Engine/Private/HeterogeneousVolumes/HeterogeneousVolumesPreshadingPipeline.usf"), TEXT("RenderSingleScatteringWithPreshadingCS"), SF_Compute);
+
+void RenderSingleScatteringWithPreshadingCompute(
+	FRDGBuilder& GraphBuilder,
+	// Scene data
+	const FScene* Scene,
+	const FViewInfo& View,
+	const FSceneTextures& SceneTextures,
+	// Light data
+	bool bApplyEmission,
+	bool bApplyDirectLighting,
+	uint32 LightType,
+	const FLightSceneInfo* LightSceneInfo,
+	// Object data
+	const FPrimitiveSceneProxy* PrimitiveSceneProxy,
+	// Sparse voxel data
+	FRDGBufferRef NumVoxelsBuffer,
+	const TRDGUniformBufferRef<FSparseVoxelUniformBufferParameters>& SparseVoxelUniformBuffer,
+	// Ray marching tiles
+	FRDGBufferRef NumRayMarchingTilesBuffer,
+	FRDGBufferRef RayMarchingTilesBuffer,
+	FRDGBufferRef VoxelsPerTileBuffer,
+	// Output
+	FRDGTextureRef& HeterogeneousVolumeTexture
+)
+{
+	FRDGBufferRef VoxelOutputBuffer = GraphBuilder.CreateBuffer(
+		FRDGBufferDesc::CreateStructuredDesc(sizeof(FVoxelDataPacked), HeterogeneousVolumes::GetVoxelCount(SparseVoxelUniformBuffer->GetParameters()->VolumeResolution)),
+		TEXT("HeterogeneousVolumes.VoxelOutputBuffer")
+	);
+
+	FRenderSingleScatteringWithPreshadingCS::FParameters* PassParameters = GraphBuilder.AllocParameters<FRenderSingleScatteringWithPreshadingCS::FParameters>();
+	{
+		// Scene data
+		PassParameters->View = View.ViewUniformBuffer;
+		PassParameters->SceneTextures = GetSceneTextureParameters(GraphBuilder, SceneTextures);
+
+		// Light data
+		FDeferredLightUniformStruct DeferredLightUniform;
+		PassParameters->bApplyEmission = bApplyEmission;
+		PassParameters->bApplyDirectLighting = bApplyDirectLighting;
+		if (PassParameters->bApplyDirectLighting)
+		{
+			DeferredLightUniform = GetDeferredLightParameters(View, *LightSceneInfo);
+		}
+		PassParameters->DeferredLight = CreateUniformBufferImmediate(DeferredLightUniform, UniformBuffer_SingleDraw);
+		PassParameters->LightType = LightType;
+
+		// Volume data
+		PassParameters->MipLevel = HeterogeneousVolumes::GetMipLevel();
+
+		// Sparse voxel data
+		PassParameters->SparseVoxelUniformBuffer = SparseVoxelUniformBuffer;
+
+		// Ray data
+		PassParameters->MaxTraceDistance = HeterogeneousVolumes::GetMaxTraceDistance();
+		PassParameters->StepSize = HeterogeneousVolumes::GetStepSize();
+		PassParameters->MaxStepCount = HeterogeneousVolumes::GetMaxStepCount();
+		PassParameters->bJitter = HeterogeneousVolumes::ShouldJitter();
+
+		// Ray marching data
+		PassParameters->RayMarchingTilesBuffer = GraphBuilder.CreateSRV(RayMarchingTilesBuffer);
+		PassParameters->VoxelsPerTileBuffer = GraphBuilder.CreateSRV(VoxelsPerTileBuffer);
+
+		// Dispatch data
+		PassParameters->IndirectArgs = NumRayMarchingTilesBuffer;
+
+		// Output
+		PassParameters->RWLightingTexture = GraphBuilder.CreateUAV(HeterogeneousVolumeTexture);
+		PassParameters->RWVoxelOutputBuffer = GraphBuilder.CreateUAV(VoxelOutputBuffer);
+	}
+
+	FString LightName = TEXT("none");
+	if (LightSceneInfo != nullptr)
+	{
+		FSceneRenderer::GetLightNameForDrawEvent(LightSceneInfo->Proxy, LightName);
+	}
+
+	FRenderSingleScatteringWithPreshadingCS::FPermutationDomain PermutationVector;
+	PermutationVector.Set<FRenderSingleScatteringWithPreshadingCS::FDebugDim>(HeterogeneousVolumes::GetDebugMode() != 0);
+	PermutationVector.Set<FRenderSingleScatteringWithPreshadingCS::FVoxelCullingDim>(HeterogeneousVolumes::UseSparseVoxelPerTileCulling());
+	PermutationVector.Set<FRenderSingleScatteringWithPreshadingCS::FSparseVoxelTracingDim>(HeterogeneousVolumes::UseSparseVoxelPipeline());
+	TShaderRef<FRenderSingleScatteringWithPreshadingCS> ComputeShader = View.ShaderMap->GetShader<FRenderSingleScatteringWithPreshadingCS>(PermutationVector);
+	FComputeShaderUtils::AddPass(
+		GraphBuilder,
+		RDG_EVENT_NAME("RenderSingleScatteringWithPreshadingCS (Light = %s)", *LightName),
+		ComputeShader,
+		PassParameters,
+		PassParameters->IndirectArgs,
+		0);
+}
+
+void RenderWithPreshadingCompute(
+	FRDGBuilder& GraphBuilder,
+	// Scene data
+	const FSceneTextures& SceneTextures,
+	const FScene* Scene,
+	const FSceneViewFamily& ViewFamily,
+	FViewInfo& View,
+	// Object data
+	const FPrimitiveSceneProxy* PrimitiveSceneProxy,
+	const FMaterialRenderProxy* MaterialRenderProxy,
+	// Sparse voxel data
+	FRDGBufferRef NumVoxelsBuffer,
+	const TRDGUniformBufferRef<FSparseVoxelUniformBufferParameters>& SparseVoxelUniformBuffer,
+	// Output
+	FRDGTextureRef& HeterogeneousVolumeRadiance
+)
+{
+	RDG_EVENT_SCOPE(GraphBuilder, "Software Ray Tracing");
+
+	FRDGBufferRef NumRayMarchingTilesBuffer;
+	FRDGBufferRef RayMarchingTilesBuffer;
+	FRDGBufferRef VoxelsPerTileBuffer;
+	{
+		RDG_EVENT_SCOPE(GraphBuilder, "Ray Tile Generation");
+		GenerateRayMarchingTiles(
+			GraphBuilder,
+			// Scene
+			Scene,
+			View,
+			SceneTextures,
+			// Object
+			PrimitiveSceneProxy,
+			// Volume data
+			NumVoxelsBuffer,
+			// Sparse voxel data
+			SparseVoxelUniformBuffer,
+			// Output
+			NumRayMarchingTilesBuffer,
+			RayMarchingTilesBuffer,
+			VoxelsPerTileBuffer
+		);
+	}
+
+	// Direct Volume Rendering
+	{
+		RDG_EVENT_SCOPE(GraphBuilder, "Direct Volume Rendering");
+
+		// Light culling
+		TArray<FLightSceneInfoCompact, TInlineAllocator<64>> LightSceneInfoCompact;
+		for (auto LightIt = Scene->Lights.CreateConstIterator(); LightIt; ++LightIt)
+		{
+			if (LightIt->AffectsPrimitive(PrimitiveSceneProxy->GetBounds(), PrimitiveSceneProxy))
+			{
+				LightSceneInfoCompact.Add(*LightIt);
+			}
+		}
+
+		// Single-scattering
+		int32 NumPasses = FMath::Max(LightSceneInfoCompact.Num(), 1);
+		for (int32 PassIndex = 0; PassIndex < NumPasses; ++PassIndex)
+		{
+			bool bApplyEmission = PassIndex == 0;
+			bool bApplyDirectLighting = !LightSceneInfoCompact.IsEmpty();
+
+			uint32 LightType = 0;
+			FLightSceneInfo* LightSceneInfo = nullptr;
+			if (bApplyDirectLighting)
+			{
+				LightType = LightSceneInfoCompact[PassIndex].LightType;
+				LightSceneInfo = LightSceneInfoCompact[PassIndex].LightSceneInfo;
+				check(LightSceneInfo != nullptr);
+			}
+
+			RenderSingleScatteringWithPreshadingCompute(
+				GraphBuilder,
+				// Scene
+				Scene,
+				View,
+				SceneTextures,
+				// Light
+				bApplyEmission,
+				bApplyDirectLighting,
+				LightType,
+				LightSceneInfo,
+				// Object
+				PrimitiveSceneProxy,
+				// Volume data
+				NumVoxelsBuffer,
+				// Sparse voxel data
+				SparseVoxelUniformBuffer,
+				// Ray marching tile
+				NumRayMarchingTilesBuffer,
+				RayMarchingTilesBuffer,
+				VoxelsPerTileBuffer,
+				// Output
+				HeterogeneousVolumeRadiance
+			);
+		}
+	}
+}
+
+void RenderWithPreshadingHardwareRayTracing(
+	FRDGBuilder& GraphBuilder,
+	// Scene data
+	const FSceneTextures& SceneTextures,
+	FScene* Scene,
+	const FSceneViewFamily& ViewFamily,
+	FViewInfo& View,
+	// Object data
+	const FPrimitiveSceneProxy* PrimitiveSceneProxy,
+	const FMaterialRenderProxy* MaterialRenderProxy,
+	// Sparse voxel data
+	FRDGBufferRef NumVoxelsBuffer,
+	const TRDGUniformBufferRef<FSparseVoxelUniformBufferParameters>& SparseVoxelUniformBuffer,
+	// Output
+	FRDGTextureRef& HeterogeneousVolumeRadiance
+)
+{
+#if RHI_RAYTRACING
+	RDG_EVENT_SCOPE(GraphBuilder, "Hardware Ray Tracing");
+
+	// WARNING: Currently works, but I'm skeptical if all RHI resources have the correct lifetime management
+	TArray<FRayTracingGeometryRHIRef> RayTracingGeometries;
+	TArray<FMatrix> RayTracingTransforms;
+	{
+		RDG_EVENT_SCOPE(GraphBuilder, "Acceleration Structure Build");
+
+		GenerateRayTracingGeometryInstance(
+			GraphBuilder,
+			// Scene
+			Scene,
+			View,
+			// Object
+			PrimitiveSceneProxy,
+			// Sparse voxel
+			NumVoxelsBuffer,
+			SparseVoxelUniformBuffer,
+			// Output
+			RayTracingGeometries,
+			RayTracingTransforms
+		);
+
+		GenerateRayTracingScene(
+			GraphBuilder,
+			// Scene
+			Scene,
+			View,
+			// Object
+			PrimitiveSceneProxy,
+			// Ray tracing data
+			RayTracingGeometries,
+			RayTracingTransforms,
+			// Output
+			Scene->HeterogeneousVolumesRayTracingScene
+		);
+	}
+
+	// Direct Volume Rendering
+	{
+		RDG_EVENT_SCOPE(GraphBuilder, "Direct Volume Rendering");
+
+		// Light culling
+		TArray<FLightSceneInfoCompact, TInlineAllocator<64>> LightSceneInfoCompact;
+		for (auto LightIt = Scene->Lights.CreateConstIterator(); LightIt; ++LightIt)
+		{
+			if (LightIt->AffectsPrimitive(PrimitiveSceneProxy->GetBounds(), PrimitiveSceneProxy))
+			{
+				LightSceneInfoCompact.Add(*LightIt);
+			}
+		}
+
+		// Single-scattering
+		int32 NumPasses = FMath::Max(LightSceneInfoCompact.Num(), 1);
+		for (int32 PassIndex = 0; PassIndex < NumPasses; ++PassIndex)
+		{
+			bool bApplyEmission = PassIndex == 0;
+			bool bApplyDirectLighting = !LightSceneInfoCompact.IsEmpty();
+
+			uint32 LightType = 0;
+			FLightSceneInfo* LightSceneInfo = nullptr;
+			if (bApplyDirectLighting)
+			{
+				LightType = LightSceneInfoCompact[PassIndex].LightType;
+				LightSceneInfo = LightSceneInfoCompact[PassIndex].LightSceneInfo;
+				check(LightSceneInfo != nullptr);
+			}
+
+			RenderSingleScatteringWithPreshadingHardwareRayTracing(
+				GraphBuilder,
+				// Scene data
+				Scene,
+				View,
+				SceneTextures,
+				// Light data
+				bApplyEmission,
+				bApplyDirectLighting,
+				LightType,
+				LightSceneInfo,
+				// Object data
+				PrimitiveSceneProxy,
+				// Sparse voxel
+				SparseVoxelUniformBuffer,
+				// Ray tracing data
+				Scene->HeterogeneousVolumesRayTracingScene,
+				// Output
+				HeterogeneousVolumeRadiance
+			);
+		}
+	}
+
+	// Tear-down ray tracing scene
+	AddPass(GraphBuilder,
+		RDG_EVENT_NAME("ReleaseRayTracingResources"),
+		[&RayTracingScene = Scene->HeterogeneousVolumesRayTracingScene](FRHICommandListImmediate& RHICmdList)
+		{
+		if (RayTracingScene.IsCreated())
+			{
+				RHICmdList.ClearRayTracingBindings(RayTracingScene.GetRHIRayTracingScene());
+			}
+		}
+	);
+
+#endif // RHI_RAYTRACING
+}
+
+class FGenerateMips3D : public FGlobalShader
+{
+	DECLARE_GLOBAL_SHADER(FGenerateMips3D);
+	SHADER_USE_PARAMETER_STRUCT(FGenerateMips3D, FGlobalShader);
+
+	static bool ShouldCompilePermutation(const FGlobalShaderPermutationParameters& Parameters)
+	{
+		return DoesPlatformSupportHeterogeneousVolumes(Parameters.Platform);
+	}
+
+	static void ModifyCompilationEnvironment(
+		const FGlobalShaderPermutationParameters& Parameters,
+		FShaderCompilerEnvironment& OutEnvironment
+	)
+	{
+		OutEnvironment.SetDefine(TEXT("THREADGROUP_SIZE_3D"), GetThreadGroupSize3D());
+
+		OutEnvironment.CompilerFlags.Add(CFLAG_AllowTypedUAVLoads);
+	}
+
+	BEGIN_SHADER_PARAMETER_STRUCT(FParameters, )
+		// Input
+		SHADER_PARAMETER_RDG_TEXTURE_SRV(Texture3D, InputTexture)
+		SHADER_PARAMETER_SAMPLER(SamplerState, TextureSampler)
+		//SHADER_PARAMETER_RDG_TEXTURE_UAV(RWTexture3D<float3>, InputTexture)
+		SHADER_PARAMETER(FIntVector, TextureResolution)
+		// Output
+		SHADER_PARAMETER_RDG_TEXTURE_UAV(RWTexture3D<float3>, RWOutputTexture)
+	END_SHADER_PARAMETER_STRUCT()
+
+	static int32 GetThreadGroupSize3D() { return 4; }
+};
+
+IMPLEMENT_GLOBAL_SHADER(FGenerateMips3D, "/Engine/Private/HeterogeneousVolumes/HeterogeneousVolumesPreshadingPipeline.usf", "GenerateMips3D", SF_Compute);
+
+class FGenerateMin3D : public FGlobalShader
+{
+	DECLARE_GLOBAL_SHADER(FGenerateMin3D);
+	SHADER_USE_PARAMETER_STRUCT(FGenerateMin3D, FGlobalShader);
+
+	static bool ShouldCompilePermutation(const FGlobalShaderPermutationParameters& Parameters)
+	{
+		return DoesPlatformSupportHeterogeneousVolumes(Parameters.Platform);
+	}
+
+	static void ModifyCompilationEnvironment(
+		const FGlobalShaderPermutationParameters& Parameters,
+		FShaderCompilerEnvironment& OutEnvironment
+	)
+	{
+		OutEnvironment.SetDefine(TEXT("THREADGROUP_SIZE_3D"), GetThreadGroupSize3D());
+
+		OutEnvironment.CompilerFlags.Add(CFLAG_AllowTypedUAVLoads);
+	}
+
+	BEGIN_SHADER_PARAMETER_STRUCT(FParameters, )
+		// Input
+		SHADER_PARAMETER_RDG_TEXTURE_SRV(Texture3D, InputTexture)
+		SHADER_PARAMETER_SAMPLER(SamplerState, TextureSampler)
+		//SHADER_PARAMETER_RDG_TEXTURE_UAV(RWTexture3D<float3>, InputTexture)
+		SHADER_PARAMETER(FIntVector, TextureResolution)
+		// Output
+		SHADER_PARAMETER_RDG_TEXTURE_UAV(RWTexture3D<float3>, RWOutputTexture)
+	END_SHADER_PARAMETER_STRUCT()
+
+	static int32 GetThreadGroupSize3D() { return 4; }
+};
+
+IMPLEMENT_GLOBAL_SHADER(FGenerateMin3D, "/Engine/Private/HeterogeneousVolumes/HeterogeneousVolumesPreshadingPipeline.usf", "GenerateMin3D", SF_Compute);
+
+template<typename ShaderType>
+void GenerateMips3D(
+	FRDGBuilder& GraphBuilder,
+	const FViewInfo& View,
+	FRDGTextureRef Texture,
+	uint32 MipLevel
+)
+{
+	FRDGTextureDesc TextureDesc = Texture->Desc;
+	const FIntVector TextureResolution(
+		FMath::Max(TextureDesc.Extent.X >> MipLevel, 1),
+		FMath::Max(TextureDesc.Extent.Y >> MipLevel, 1),
+		FMath::Max(TextureDesc.Depth >> MipLevel, 1));
+
+	typename ShaderType::FParameters* PassParameters = GraphBuilder.AllocParameters<typename ShaderType::FParameters>();
+	{
+		PassParameters->TextureResolution = TextureResolution;
+		PassParameters->InputTexture = GraphBuilder.CreateSRV(FRDGTextureSRVDesc::CreateForMipLevel(Texture, MipLevel - 1));
+		PassParameters->TextureSampler = TStaticSamplerState<SF_Trilinear, AM_Clamp, AM_Clamp, AM_Clamp>::GetRHI();
+		//PassParameters->InputTexture = GraphBuilder.CreateUAV(FRDGTextureUAVDesc(Texture, MipLevel - 1));
+		PassParameters->RWOutputTexture = GraphBuilder.CreateUAV(FRDGTextureUAVDesc(Texture, MipLevel));
+	}
+
+	uint32 GroupCountX = FMath::DivideAndRoundUp(TextureResolution.X, ShaderType::GetThreadGroupSize3D());
+	uint32 GroupCountY = FMath::DivideAndRoundUp(TextureResolution.Y, ShaderType::GetThreadGroupSize3D());
+	uint32 GroupCountZ = FMath::DivideAndRoundUp(TextureResolution.Z, ShaderType::GetThreadGroupSize3D());
+	FIntVector GroupCount(GroupCountX, GroupCountY, GroupCountZ);
+
+	TShaderRef<ShaderType> ComputeShader = View.ShaderMap->GetShader<ShaderType>();
+	FComputeShaderUtils::AddPass(
+		GraphBuilder,
+		RDG_EVENT_NAME("FGenerateMips3D"),
+		ComputeShader,
+		PassParameters,
+		GroupCount);
+}
+
+void RenderWithPreshading(
+	FRDGBuilder& GraphBuilder,
+	// Scene data
+	const FSceneTextures& SceneTextures,
+	FScene* Scene,
+	const FSceneViewFamily& ViewFamily,
+	FViewInfo& View,
+	// Object data
+	const FPrimitiveSceneProxy* PrimitiveSceneProxy,
+	const FMaterialRenderProxy* MaterialRenderProxy,
+	const int32 PrimitiveId,
+	const FBoxSphereBounds LocalBoxSphereBounds,
+	// Output
+	FRDGTextureRef& HeterogeneousVolumeRadiance
+)
+{
+	// Determine baking voxel resolution
+	FIntVector VolumeResolution = HeterogeneousVolumes::GetVolumeResolution();
+
+	// Create baked material grids
+	uint32 NumMips = FMath::Log2(float(FMath::Min(FMath::Min(VolumeResolution.X, VolumeResolution.Y), VolumeResolution.Z))) + 1;
+	FRDGTextureDesc BakedMaterialDesc = FRDGTextureDesc::Create3D(
+		VolumeResolution,
+		PF_FloatR11G11B10,
+		FClearValueBinding::Black,
+		TexCreate_ShaderResource | TexCreate_UAV | TexCreate_3DTiling,
+		NumMips
+	);
+
+	FRDGTextureRef HeterogeneousVolumesExtinctionTexture = GraphBuilder.CreateTexture(BakedMaterialDesc, TEXT("HeterogeneousVolumes.ExtinctionTexture"));
+	FRDGTextureRef HeterogeneousVolumesEmissionTexture = GraphBuilder.CreateTexture(BakedMaterialDesc, TEXT("HeterogeneousVolumes.EmissionTexture"));
+	FRDGTextureRef HeterogeneousVolumesAlbedoTexture = GraphBuilder.CreateTexture(BakedMaterialDesc, TEXT("HeterogeneousVolumes.AlbedoTexture"));
+
+	// Preshading pipeline
+	{
+		RDG_EVENT_SCOPE(GraphBuilder, "Preshading Pipeline");
+
+		{
+			RDG_EVENT_SCOPE(GraphBuilder, "Material Baking");
+
+			ComputeHeterogeneousVolumeBakeMaterial(
+				GraphBuilder,
+				// Scene data
+				Scene,
+				View,
+				// Object data
+				PrimitiveSceneProxy,
+				MaterialRenderProxy,
+				PrimitiveId,
+				LocalBoxSphereBounds,
+				// Volume data
+				VolumeResolution,
+				// Output
+				HeterogeneousVolumesExtinctionTexture,
+				HeterogeneousVolumesEmissionTexture,
+				HeterogeneousVolumesAlbedoTexture
+			);
+		}
+
+		// MIP Generation
+		{
+			RDG_EVENT_SCOPE(GraphBuilder, "MIP Generation");
+			for (uint32 MipLevel = 1; MipLevel < NumMips; ++MipLevel)
+			{
+				GenerateMips3D<FGenerateMips3D>(GraphBuilder, View, HeterogeneousVolumesExtinctionTexture, MipLevel);
+				// TODO: Reinstate once ray-marching determines appropriate MIP level to sample
+				//GenerateMips3D<FGenerateMips3D>(GraphBuilder, View, HeterogeneousVolumesEmissionTexture, MipLevel);
+				//GenerateMips3D<FGenerateMips3D>(GraphBuilder, View, HeterogeneousVolumesAlbedoTexture, MipLevel);
+			}
+		}
+	}
+
+	// Sparse Voxel Pipeline
+	int32 MipBias = HeterogeneousVolumes::GetSparseVoxelMipBias();
+	uint32 SparseMipLevel = FMath::Clamp(int32(NumMips) - MipBias, 0, int32(NumMips) - 1);
+
+	FRDGTextureRef MinTexture;
+	FRDGBufferRef NumVoxelsBuffer;
+	FRDGBufferRef VoxelBuffer;
+	{
+		RDG_EVENT_SCOPE(GraphBuilder, "Sparse Voxel Pipeline");
+
+		{
+			RDG_EVENT_SCOPE(GraphBuilder, "Min Generation");
+
+			FRDGTextureDesc MinTextureDesc = HeterogeneousVolumesExtinctionTexture->Desc;
+			MinTextureDesc.Extent.X = FMath::Max(MinTextureDesc.Extent.X >> SparseMipLevel, 1);
+			MinTextureDesc.Extent.Y = FMath::Max(MinTextureDesc.Extent.Y >> SparseMipLevel, 1);
+			MinTextureDesc.Depth = FMath::Max(MinTextureDesc.Depth >> SparseMipLevel, 1);
+			MinTextureDesc.NumMips = FMath::Log2(float(FMath::Min(FMath::Min(MinTextureDesc.Extent.X, MinTextureDesc.Extent.Y), (int32)MinTextureDesc.Depth))) + 1;
+
+			MinTexture = GraphBuilder.CreateTexture(MinTextureDesc, TEXT("HeterogeneousVolumes.MinTexture"));
+			CopyTexture3D(GraphBuilder, View, HeterogeneousVolumesExtinctionTexture, SparseMipLevel, MinTexture);
+			for (uint32 MipLevel = 1; MipLevel < MinTextureDesc.NumMips; ++MipLevel)
+			{
+				GenerateMips3D<FGenerateMin3D>(GraphBuilder, View, MinTexture, MipLevel);
+			}
+		}
+
+		{
+			RDG_EVENT_SCOPE(GraphBuilder, "Sparse Voxel Generation");
+			GenerateSparseVoxels(GraphBuilder, View, MinTexture, VolumeResolution, SparseMipLevel, NumVoxelsBuffer, VoxelBuffer);
+		}
+	}
+
+	// Create Sparse Voxel UniformBuffer
+	FSparseVoxelUniformBufferParameters* SparseVoxelUniformBufferParameters = GraphBuilder.AllocParameters<FSparseVoxelUniformBufferParameters>();
+	{
+		// Object data
+		FMatrix44f LocalToWorld = FMatrix44f(PrimitiveSceneProxy->GetLocalToWorld());
+		SparseVoxelUniformBufferParameters->LocalToWorld = LocalToWorld;
+		SparseVoxelUniformBufferParameters->WorldToLocal = LocalToWorld.Inverse();
+		SparseVoxelUniformBufferParameters->LocalBoundsOrigin = FVector3f(LocalBoxSphereBounds.Origin);
+		SparseVoxelUniformBufferParameters->LocalBoundsExtent = FVector3f(LocalBoxSphereBounds.BoxExtent);
+
+		// Volume data
+		SparseVoxelUniformBufferParameters->VolumeResolution = VolumeResolution;
+		SparseVoxelUniformBufferParameters->ExtinctionTexture = HeterogeneousVolumesExtinctionTexture;
+		SparseVoxelUniformBufferParameters->EmissionTexture = HeterogeneousVolumesEmissionTexture;
+		SparseVoxelUniformBufferParameters->AlbedoTexture = HeterogeneousVolumesAlbedoTexture;
+		SparseVoxelUniformBufferParameters->TextureSampler = TStaticSamplerState<SF_Trilinear, AM_Clamp, AM_Clamp, AM_Clamp>::GetRHI();;
+
+		// Sparse voxel data
+		SparseVoxelUniformBufferParameters->NumVoxelsBuffer = GraphBuilder.CreateSRV(NumVoxelsBuffer, PF_R32_UINT);
+		SparseVoxelUniformBufferParameters->VoxelBuffer = GraphBuilder.CreateSRV(VoxelBuffer);
+		SparseVoxelUniformBufferParameters->MipLevel = SparseMipLevel;
+
+		// Traversal hints
+		SparseVoxelUniformBufferParameters->MaxTraceDistance = HeterogeneousVolumes::GetMaxTraceDistance();
+		SparseVoxelUniformBufferParameters->StepSize = HeterogeneousVolumes::GetStepSize();
+		SparseVoxelUniformBufferParameters->ShadowStepFactor = HeterogeneousVolumes::GetShadowStepFactor();
+	}
+	TRDGUniformBufferRef<FSparseVoxelUniformBufferParameters> SparseVoxelUniformBuffer = GraphBuilder.CreateUniformBuffer(SparseVoxelUniformBufferParameters);
+
+	// Hardware ray tracing
+	if (HeterogeneousVolumes::UseHardwareRayTracing())
+	{
+		RenderWithPreshadingHardwareRayTracing(
+			GraphBuilder,
+			// Scene data
+			SceneTextures,
+			Scene,
+			ViewFamily,
+			View,
+			// Object data
+			PrimitiveSceneProxy,
+			MaterialRenderProxy,
+			// Sparse voxel data
+			NumVoxelsBuffer,
+			SparseVoxelUniformBuffer,
+			// Output
+			HeterogeneousVolumeRadiance
+		);
+	}
+	// Software ray tracing
+	else
+	{
+		RenderWithPreshadingCompute(
+			GraphBuilder,
+			// Scene data
+			SceneTextures,
+			Scene,
+			ViewFamily,
+			View,
+			// Object data
+			PrimitiveSceneProxy,
+			MaterialRenderProxy,
+			// Sparse voxel data
+			NumVoxelsBuffer,
+			SparseVoxelUniformBuffer,
+			// Output
+			HeterogeneousVolumeRadiance
+		);
+	}
+}
\ No newline at end of file
diff --git a/Engine/Source/Runtime/Renderer/Private/HeterogeneousVolumes/HeterogeneousVolumesSparseVoxelPipeline.cpp b/Engine/Source/Runtime/Renderer/Private/HeterogeneousVolumes/HeterogeneousVolumesSparseVoxelPipeline.cpp
new file mode 100644
--- /dev/null
+++ b/Engine/Source/Runtime/Renderer/Private/HeterogeneousVolumes/HeterogeneousVolumesSparseVoxelPipeline.cpp
@@ -0,0 +1,384 @@
+// Copyright Epic Games, Inc. All Rights Reserved.
+
+#include "HeterogeneousVolumes.h"
+
+#include "PixelShaderUtils.h"
+#include "RayTracingDefinitions.h"
+#include "RayTracingInstance.h"
+#include "RayTracingInstanceBufferUtil.h"
+#include "RendererPrivate.h"
+#include "ScenePrivate.h"
+#include "SceneManagement.h"
+
+IMPLEMENT_UNIFORM_BUFFER_STRUCT(FSparseVoxelUniformBufferParameters, "SparseVoxelUniformBuffer");
+
+class FCopyTexture3D : public FGlobalShader
+{
+	DECLARE_GLOBAL_SHADER(FCopyTexture3D);
+	SHADER_USE_PARAMETER_STRUCT(FCopyTexture3D, FGlobalShader);
+
+	static bool ShouldCompilePermutation(const FGlobalShaderPermutationParameters& Parameters)
+	{
+		return DoesPlatformSupportHeterogeneousVolumes(Parameters.Platform);
+	}
+
+	static void ModifyCompilationEnvironment(
+		const FGlobalShaderPermutationParameters& Parameters,
+		FShaderCompilerEnvironment& OutEnvironment
+	)
+	{
+		OutEnvironment.SetDefine(TEXT("THREADGROUP_SIZE_3D"), GetThreadGroupSize3D());
+
+		OutEnvironment.CompilerFlags.Add(CFLAG_AllowTypedUAVLoads);
+	}
+
+	BEGIN_SHADER_PARAMETER_STRUCT(FParameters, )
+		// Input
+		SHADER_PARAMETER_RDG_TEXTURE_SRV(Texture3D, InputTexture)
+		SHADER_PARAMETER_SAMPLER(SamplerState, TextureSampler)
+		SHADER_PARAMETER(FIntVector, TextureResolution)
+		// Output
+		SHADER_PARAMETER_RDG_TEXTURE_UAV(RWTexture3D<float3>, RWOutputTexture)
+		END_SHADER_PARAMETER_STRUCT()
+
+		static int32 GetThreadGroupSize3D() { return 4; }
+};
+
+IMPLEMENT_GLOBAL_SHADER(FCopyTexture3D, "/Engine/Private/HeterogeneousVolumes/HeterogeneousVolumesPreshadingPipeline.usf", "CopyTexture3D", SF_Compute);
+
+void CopyTexture3D(
+	FRDGBuilder& GraphBuilder,
+	const FViewInfo& View,
+	FRDGTextureRef Texture,
+	uint32 InputMipLevel,
+	FRDGTextureRef& OutputTexture
+)
+{
+	FRDGTextureDesc OutputTextureDesc = OutputTexture->Desc;
+	const FIntVector TextureResolution(OutputTextureDesc.Extent.X, OutputTextureDesc.Extent.Y, OutputTextureDesc.Depth);
+	FCopyTexture3D::FParameters* PassParameters = GraphBuilder.AllocParameters<FCopyTexture3D::FParameters>();
+	{
+		PassParameters->TextureResolution = TextureResolution;
+		PassParameters->InputTexture = GraphBuilder.CreateSRV(FRDGTextureSRVDesc::CreateForMipLevel(Texture, InputMipLevel));
+		PassParameters->RWOutputTexture = GraphBuilder.CreateUAV(FRDGTextureUAVDesc(OutputTexture));
+	}
+
+	uint32 GroupCountX = FMath::DivideAndRoundUp(TextureResolution.X, FCopyTexture3D::GetThreadGroupSize3D());
+	uint32 GroupCountY = FMath::DivideAndRoundUp(TextureResolution.Y, FCopyTexture3D::GetThreadGroupSize3D());
+	uint32 GroupCountZ = FMath::DivideAndRoundUp(TextureResolution.Z, FCopyTexture3D::GetThreadGroupSize3D());
+	FIntVector GroupCount(GroupCountX, GroupCountY, GroupCountZ);
+
+	TShaderRef<FCopyTexture3D> ComputeShader = View.ShaderMap->GetShader<FCopyTexture3D>();
+	FComputeShaderUtils::AddPass(
+		GraphBuilder,
+		RDG_EVENT_NAME("CopyTexture3D"),
+		ComputeShader,
+		PassParameters,
+		GroupCount);
+}
+
+class FClearAllocator : public FGlobalShader
+{
+	DECLARE_GLOBAL_SHADER(FClearAllocator);
+	SHADER_USE_PARAMETER_STRUCT(FClearAllocator, FGlobalShader);
+
+	static bool ShouldCompilePermutation(const FGlobalShaderPermutationParameters& Parameters)
+	{
+		return DoesPlatformSupportHeterogeneousVolumes(Parameters.Platform);
+	}
+
+	BEGIN_SHADER_PARAMETER_STRUCT(FParameters, )
+		SHADER_PARAMETER_RDG_BUFFER_UAV(RWBuffer<uint32>, RWAllocatorBuffer)
+		END_SHADER_PARAMETER_STRUCT()
+};
+
+IMPLEMENT_GLOBAL_SHADER(FClearAllocator, "/Engine/Private/HeterogeneousVolumes/HeterogeneousVolumesSparseVoxelPipeline.usf", "ClearAllocator", SF_Compute);
+
+class FAllocateSparseVoxels : public FGlobalShader
+{
+	DECLARE_GLOBAL_SHADER(FAllocateSparseVoxels);
+	SHADER_USE_PARAMETER_STRUCT(FAllocateSparseVoxels, FGlobalShader);
+
+	static bool ShouldCompilePermutation(const FGlobalShaderPermutationParameters& Parameters)
+	{
+		return DoesPlatformSupportHeterogeneousVolumes(Parameters.Platform);
+	}
+
+	static void ModifyCompilationEnvironment(
+		const FGlobalShaderPermutationParameters& Parameters,
+		FShaderCompilerEnvironment& OutEnvironment
+	)
+	{
+		OutEnvironment.SetDefine(TEXT("THREADGROUP_SIZE_3D"), GetThreadGroupSize3D());
+		OutEnvironment.CompilerFlags.Add(CFLAG_AllowTypedUAVLoads);
+	}
+
+	BEGIN_SHADER_PARAMETER_STRUCT(FParameters, )
+		// Input
+		SHADER_PARAMETER_RDG_TEXTURE_SRV(Texture3D, VoxelMinTexture)
+		SHADER_PARAMETER_SAMPLER(SamplerState, TextureSampler)
+		SHADER_PARAMETER(FIntVector, VoxelMinResolution)
+		SHADER_PARAMETER(uint32, MipLevel)
+		// Output
+		SHADER_PARAMETER_RDG_BUFFER_UAV(RWBuffer<uint>, RWNumVoxelsBuffer)
+		SHADER_PARAMETER_RDG_BUFFER_UAV(RWStructuredBuffer<FVoxelDataPacked>, RWVoxelBuffer)
+		END_SHADER_PARAMETER_STRUCT()
+
+		static int32 GetThreadGroupSize3D() { return 4; }
+};
+
+IMPLEMENT_GLOBAL_SHADER(FAllocateSparseVoxels, "/Engine/Private/HeterogeneousVolumes/HeterogeneousVolumesSparseVoxelPipeline.usf", "AllocateSparseVoxels", SF_Compute);
+
+class FAllocateSparseVoxelsToRefine : public FGlobalShader
+{
+	DECLARE_GLOBAL_SHADER(FAllocateSparseVoxelsToRefine);
+	SHADER_USE_PARAMETER_STRUCT(FAllocateSparseVoxelsToRefine, FGlobalShader);
+
+	static bool ShouldCompilePermutation(const FGlobalShaderPermutationParameters& Parameters)
+	{
+		return DoesPlatformSupportHeterogeneousVolumes(Parameters.Platform);
+	}
+
+	static void ModifyCompilationEnvironment(
+		const FGlobalShaderPermutationParameters& Parameters,
+		FShaderCompilerEnvironment& OutEnvironment
+	)
+	{
+		OutEnvironment.SetDefine(TEXT("THREADGROUP_SIZE_3D"), GetThreadGroupSize3D());
+		OutEnvironment.CompilerFlags.Add(CFLAG_AllowTypedUAVLoads);
+	}
+
+	BEGIN_SHADER_PARAMETER_STRUCT(FParameters, )
+		// Input
+		SHADER_PARAMETER_RDG_TEXTURE_SRV(Texture3D, VoxelMinTexture)
+		SHADER_PARAMETER_SAMPLER(SamplerState, TextureSampler)
+		SHADER_PARAMETER(FIntVector, VoxelMinResolution)
+		SHADER_PARAMETER(uint32, MipLevel)
+		// Output
+		SHADER_PARAMETER_RDG_BUFFER_UAV(RWBuffer<uint>, RWNumVoxelsBuffer)
+		SHADER_PARAMETER_RDG_BUFFER_UAV(RWStructuredBuffer<FVoxelDataPacked>, RWVoxelBuffer)
+
+		SHADER_PARAMETER_RDG_BUFFER_UAV(RWBuffer<uint>, RWNumVoxelsToRefineBuffer)
+		SHADER_PARAMETER_RDG_BUFFER_UAV(RWStructuredBuffer<FVoxelDataPacked>, RWVoxelsToRefineBuffer)
+		END_SHADER_PARAMETER_STRUCT()
+
+		static int32 GetThreadGroupSize3D() { return 4; }
+};
+
+IMPLEMENT_GLOBAL_SHADER(FAllocateSparseVoxelsToRefine, "/Engine/Private/HeterogeneousVolumes/HeterogeneousVolumesSparseVoxelPipeline.usf", "AllocateSparseVoxelsToRefine", SF_Compute);
+
+class FRefineSparseVoxels : public FGlobalShader
+{
+	DECLARE_GLOBAL_SHADER(FRefineSparseVoxels);
+	SHADER_USE_PARAMETER_STRUCT(FRefineSparseVoxels, FGlobalShader);
+
+	static bool ShouldCompilePermutation(const FGlobalShaderPermutationParameters& Parameters)
+	{
+		return DoesPlatformSupportHeterogeneousVolumes(Parameters.Platform);
+	}
+
+	static void ModifyCompilationEnvironment(
+		const FGlobalShaderPermutationParameters& Parameters,
+		FShaderCompilerEnvironment& OutEnvironment
+	)
+	{
+		OutEnvironment.SetDefine(TEXT("THREADGROUP_SIZE_1D"), GetThreadGroupSize1D());
+		OutEnvironment.CompilerFlags.Add(CFLAG_AllowTypedUAVLoads);
+	}
+
+	BEGIN_SHADER_PARAMETER_STRUCT(FParameters, )
+		// Input
+		SHADER_PARAMETER_RDG_TEXTURE_SRV(Texture3D, VoxelMinTexture)
+		SHADER_PARAMETER_SAMPLER(SamplerState, TextureSampler)
+		SHADER_PARAMETER(FIntVector, VolumeResolution)
+		SHADER_PARAMETER(FIntVector, MipVolumeResolution)
+		SHADER_PARAMETER(uint32, MipLevel)
+
+		// Refinement data
+		SHADER_PARAMETER_RDG_BUFFER_SRV(Buffer<uint>, NumVoxelsToRefineBuffer)
+		SHADER_PARAMETER_RDG_BUFFER_SRV(StructuredBuffer<FVoxelDataPacked>, VoxelsToRefineBuffer)
+		SHADER_PARAMETER_RDG_BUFFER_UAV(RWBuffer<uint>, RWNumVoxelsToRefineBuffer)
+		SHADER_PARAMETER_RDG_BUFFER_UAV(RWStructuredBuffer<FVoxelDataPacked>, RWVoxelsToRefineBuffer)
+
+		// Output
+		SHADER_PARAMETER_RDG_BUFFER_UAV(RWBuffer<uint>, RWNumVoxelsBuffer)
+		SHADER_PARAMETER_RDG_BUFFER_UAV(RWStructuredBuffer<FVoxelDataPacked>, RWVoxelBuffer)
+
+		// Indirect args
+		RDG_BUFFER_ACCESS(IndirectArgs, ERHIAccess::IndirectArgs)
+		END_SHADER_PARAMETER_STRUCT()
+
+		static int32 GetThreadGroupSize1D() { return 64; }
+};
+
+IMPLEMENT_GLOBAL_SHADER(FRefineSparseVoxels, "/Engine/Private/HeterogeneousVolumes/HeterogeneousVolumesSparseVoxelPipeline.usf", "RefineSparseVoxels", SF_Compute);
+
+void GenerateSparseVoxels(
+	FRDGBuilder& GraphBuilder,
+	const FViewInfo& View,
+	FRDGTextureRef VoxelMinTexture,
+	FIntVector VolumeResolution,
+	uint32 MipLevel,
+	FRDGBufferRef& NumVoxelsBuffer,
+	FRDGBufferRef& VoxelBuffer
+)
+{
+	FIntVector VoxelMinResolution = FIntVector(
+		VoxelMinTexture->Desc.Extent.X,
+		VoxelMinTexture->Desc.Extent.Y,
+		VoxelMinTexture->Desc.Depth
+	);
+
+	NumVoxelsBuffer = GraphBuilder.CreateBuffer(
+		FRDGBufferDesc::CreateIndirectDesc<FRHIDispatchIndirectParameters>(1),
+		TEXT("HeterogeneousVolumes.NumVoxelsBuffer")
+	);
+
+	// Clear allocator
+	{
+		FClearAllocator::FParameters* PassParameters = GraphBuilder.AllocParameters<FClearAllocator::FParameters>();
+		{
+			PassParameters->RWAllocatorBuffer = GraphBuilder.CreateUAV(NumVoxelsBuffer, PF_R32_UINT);
+		}
+
+		TShaderRef<FClearAllocator> ComputeShader = View.ShaderMap->GetShader<FClearAllocator>();
+		FComputeShaderUtils::AddPass(
+			GraphBuilder,
+			RDG_EVENT_NAME("ClearAllocator"),
+			ERDGPassFlags::Compute | ERDGPassFlags::NeverCull,
+			ComputeShader,
+			PassParameters,
+			FIntVector(1)
+		);
+	}
+
+	// Generate sets of OOBBs at a candidate MIP level for a given VRE coefficient
+	VoxelBuffer = GraphBuilder.CreateBuffer(
+		FRDGBufferDesc::CreateStructuredDesc(sizeof(FVoxelDataPacked), HeterogeneousVolumes::GetVoxelCount(VoxelMinResolution)),
+		TEXT("HeterogeneousVolumes.VoxelBuffer")
+	);
+
+	if (HeterogeneousVolumes::ShouldRefineSparseVoxels())
+	{
+		// Refinement buffers
+		FRDGBufferRef NumVoxelsToRefineBuffer[2] = {
+			GraphBuilder.CreateBuffer(FRDGBufferDesc::CreateIndirectDesc<FRHIDispatchIndirectParameters>(1), TEXT("HeterogeneousVolume.NumVoxelsToRefineBuffer1")),
+			GraphBuilder.CreateBuffer(FRDGBufferDesc::CreateIndirectDesc<FRHIDispatchIndirectParameters>(1), TEXT("HeterogeneousVolume.NumVoxelsToRefineBuffer2"))
+		};
+		FRDGBufferRef VoxelsToRefineBuffer[2] = {
+			GraphBuilder.CreateBuffer(VoxelBuffer->Desc, TEXT("HeterogeneousVolume.VoxelsToRefineBuffer1")),
+			GraphBuilder.CreateBuffer(VoxelBuffer->Desc, TEXT("HeterogeneousVolume.VoxelsToRefineBuffer2"))
+		};
+		uint32 InputIndex = 0;
+		uint32 OutputIndex = 1;
+
+		// Initial dense dispatch
+		{
+			AddClearUAVPass(GraphBuilder, GraphBuilder.CreateUAV(NumVoxelsToRefineBuffer[OutputIndex], PF_R32_UINT), 0);
+
+			FAllocateSparseVoxelsToRefine::FParameters* PassParameters = GraphBuilder.AllocParameters<FAllocateSparseVoxelsToRefine::FParameters>();
+			{
+				PassParameters->VoxelMinTexture = GraphBuilder.CreateSRV(FRDGTextureSRVDesc::CreateForMipLevel(VoxelMinTexture, 0));
+				PassParameters->VoxelMinResolution = VoxelMinResolution;
+				PassParameters->MipLevel = MipLevel;
+
+				PassParameters->RWNumVoxelsBuffer = GraphBuilder.CreateUAV(NumVoxelsBuffer, PF_R32_UINT);
+				PassParameters->RWVoxelBuffer = GraphBuilder.CreateUAV(VoxelBuffer);
+
+				PassParameters->RWNumVoxelsToRefineBuffer = GraphBuilder.CreateUAV(NumVoxelsToRefineBuffer[OutputIndex], PF_R32_UINT);
+				PassParameters->RWVoxelsToRefineBuffer = GraphBuilder.CreateUAV(VoxelsToRefineBuffer[OutputIndex]);
+			}
+
+			uint32 GroupCountX = FMath::DivideAndRoundUp(VoxelMinResolution.X, FAllocateSparseVoxelsToRefine::GetThreadGroupSize3D());
+			uint32 GroupCountY = FMath::DivideAndRoundUp(VoxelMinResolution.Y, FAllocateSparseVoxelsToRefine::GetThreadGroupSize3D());
+			uint32 GroupCountZ = FMath::DivideAndRoundUp(VoxelMinResolution.Z, FAllocateSparseVoxelsToRefine::GetThreadGroupSize3D());
+			FIntVector GroupCount(GroupCountX, GroupCountY, GroupCountZ);
+			TShaderRef<FAllocateSparseVoxelsToRefine> ComputeShader = View.ShaderMap->GetShader<FAllocateSparseVoxelsToRefine>();
+			FComputeShaderUtils::AddPass(
+				GraphBuilder,
+				RDG_EVENT_NAME("AllocateSparseVoxelsToRefine"),
+				ERDGPassFlags::Compute | ERDGPassFlags::NeverCull,
+				ComputeShader,
+				PassParameters,
+				GroupCount
+			);
+		}
+
+		// Flip input/output buffers for next iteration
+		Swap(InputIndex, OutputIndex);
+
+		for (uint32 MipLevelOffset = 1; MipLevelOffset < VoxelMinTexture->Desc.NumMips; ++MipLevelOffset)
+		{
+			// Clear refinement output buffers
+			AddClearUAVPass(GraphBuilder, GraphBuilder.CreateUAV(NumVoxelsToRefineBuffer[OutputIndex], PF_R32_UINT), 0);
+
+			FRefineSparseVoxels::FParameters* PassParameters = GraphBuilder.AllocParameters<FRefineSparseVoxels::FParameters>();
+			{
+				PassParameters->VoxelMinTexture = GraphBuilder.CreateSRV(FRDGTextureSRVDesc::CreateForMipLevel(VoxelMinTexture, MipLevelOffset));
+				FIntVector MipVolumeResolution = FIntVector(
+					VoxelMinResolution.X >> MipLevelOffset,
+					VoxelMinResolution.Y >> MipLevelOffset,
+					VoxelMinResolution.Z >> MipLevelOffset
+				);
+				PassParameters->VolumeResolution = VolumeResolution;
+				PassParameters->MipVolumeResolution = MipVolumeResolution;
+				PassParameters->MipLevel = MipLevel + MipLevelOffset;
+
+				// Intermediate refinement buffers
+				PassParameters->NumVoxelsToRefineBuffer = GraphBuilder.CreateSRV(NumVoxelsToRefineBuffer[InputIndex], PF_R32_UINT);
+				PassParameters->VoxelsToRefineBuffer = GraphBuilder.CreateSRV(VoxelsToRefineBuffer[InputIndex]);
+				PassParameters->RWNumVoxelsToRefineBuffer = GraphBuilder.CreateUAV(NumVoxelsToRefineBuffer[OutputIndex], PF_R32_UINT);
+				PassParameters->RWVoxelsToRefineBuffer = GraphBuilder.CreateUAV(VoxelsToRefineBuffer[OutputIndex]);
+
+				// Sparse voxel output
+				PassParameters->RWNumVoxelsBuffer = GraphBuilder.CreateUAV(NumVoxelsBuffer, PF_R32_UINT);
+				PassParameters->RWVoxelBuffer = GraphBuilder.CreateUAV(VoxelBuffer);
+
+				PassParameters->IndirectArgs = NumVoxelsToRefineBuffer[InputIndex];
+			}
+
+			TShaderRef<FRefineSparseVoxels> ComputeShader = View.ShaderMap->GetShader<FRefineSparseVoxels>();
+			FComputeShaderUtils::AddPass(
+				GraphBuilder,
+				RDG_EVENT_NAME("RefineSparseVoxels"),
+				ERDGPassFlags::Compute,
+				ComputeShader,
+				PassParameters,
+				PassParameters->IndirectArgs,
+				0
+			);
+
+			// Flip input/output buffers for next iteration
+			Swap(InputIndex, OutputIndex);
+		}
+	}
+	else
+	{
+		FAllocateSparseVoxels::FParameters* PassParameters = GraphBuilder.AllocParameters<FAllocateSparseVoxels::FParameters>();
+		{
+			//PassParameters->VoxelMinTexture = GraphBuilder.CreateSRV(FRDGTextureSRVDesc::CreateForMipLevel(VoxelMinTexture, MipLevel));
+			PassParameters->VoxelMinTexture = GraphBuilder.CreateSRV(VoxelMinTexture);
+			// TODO: Make sure point-sampling isn't causing clipping of the isosurface
+			PassParameters->TextureSampler = TStaticSamplerState<SF_Point, AM_Clamp, AM_Clamp, AM_Clamp>::GetRHI();
+			PassParameters->VoxelMinResolution = VoxelMinResolution;
+			PassParameters->MipLevel = MipLevel;
+
+			PassParameters->RWNumVoxelsBuffer = GraphBuilder.CreateUAV(NumVoxelsBuffer, PF_R32_UINT);
+			PassParameters->RWVoxelBuffer = GraphBuilder.CreateUAV(VoxelBuffer);
+		}
+
+		uint32 GroupCountX = FMath::DivideAndRoundUp(VoxelMinResolution.X, FAllocateSparseVoxels::GetThreadGroupSize3D());
+		uint32 GroupCountY = FMath::DivideAndRoundUp(VoxelMinResolution.Y, FAllocateSparseVoxels::GetThreadGroupSize3D());
+		uint32 GroupCountZ = FMath::DivideAndRoundUp(VoxelMinResolution.Z, FAllocateSparseVoxels::GetThreadGroupSize3D());
+		FIntVector GroupCount(GroupCountX, GroupCountY, GroupCountZ);
+		TShaderRef<FAllocateSparseVoxels> ComputeShader = View.ShaderMap->GetShader<FAllocateSparseVoxels>();
+		FComputeShaderUtils::AddPass(
+			GraphBuilder,
+			RDG_EVENT_NAME("AllocateSparseVoxels"),
+			ERDGPassFlags::Compute | ERDGPassFlags::NeverCull,
+			ComputeShader,
+			PassParameters,
+			GroupCount
+		);
+	}
+}
\ No newline at end of file
diff --git a/Engine/Source/Runtime/Renderer/Private/RayTracing/RayTracingScene.cpp b/Engine/Source/Runtime/Renderer/Private/RayTracing/RayTracingScene.cpp
--- a/Engine/Source/Runtime/Renderer/Private/RayTracing/RayTracingScene.cpp
+++ b/Engine/Source/Runtime/Renderer/Private/RayTracing/RayTracingScene.cpp
@@ -36,12 +36,12 @@ FRayTracingSceneWithGeometryInstances FRayTracingScene::BuildInitializationData(
 		NumCallableShaderSlots);
 }
 
-void FRayTracingScene::Create(FRDGBuilder& GraphBuilder, const FGPUScene& GPUScene, const FViewMatrices& ViewMatrices)
+void FRayTracingScene::Create(FRDGBuilder& GraphBuilder, const FGPUScene* GPUScene, const FViewMatrices& ViewMatrices)
 {
 	CreateWithInitializationData(GraphBuilder, GPUScene, ViewMatrices, BuildInitializationData());
 }
 
-void FRayTracingScene::CreateWithInitializationData(FRDGBuilder& GraphBuilder, const FGPUScene& GPUScene, const FViewMatrices& ViewMatrices, FRayTracingSceneWithGeometryInstances SceneWithGeometryInstances)
+void FRayTracingScene::CreateWithInitializationData(FRDGBuilder& GraphBuilder, const FGPUScene* GPUScene, const FViewMatrices& ViewMatrices, FRayTracingSceneWithGeometryInstances SceneWithGeometryInstances)
 {
 	QUICK_SCOPE_CYCLE_COUNTER(FRayTracingScene_BeginCreate);
 
@@ -196,7 +196,7 @@ void FRayTracingScene::CreateWithInitializationData(FRDGBuilder& GraphBuilder, c
 			ERDGPassFlags::Compute,
 			[PassParams,
 			this,
-			GPUScene = &GPUScene,
+			GPUScene,
 			ViewTilePosition = AbsoluteViewOrigin.GetTile(),
 			RelativePreViewTranslation = ViewMatrices.GetPreViewTranslation() + ViewTileOffset,
 			&SceneInitializer,
diff --git a/Engine/Source/Runtime/Renderer/Private/RayTracing/RayTracingScene.h b/Engine/Source/Runtime/Renderer/Private/RayTracing/RayTracingScene.h
--- a/Engine/Source/Runtime/Renderer/Private/RayTracing/RayTracingScene.h
+++ b/Engine/Source/Runtime/Renderer/Private/RayTracing/RayTracingScene.h
@@ -47,10 +47,10 @@ public:
 	// Allocates GPU memory to fit at least the current number of instances.
 	// Kicks off instance buffer build to parallel thread along with RDG pass.
 	// NOTE: SceneWithGeometryInstances is passed in by value because ownership of the internal data is taken over. Use MoveTemp at call site, if possible.
-	void CreateWithInitializationData(FRDGBuilder& GraphBuilder, const FGPUScene& GPUScene, const FViewMatrices& ViewMatrices, FRayTracingSceneWithGeometryInstances SceneWithGeometryInstances);
+	void CreateWithInitializationData(FRDGBuilder& GraphBuilder, const FGPUScene* GPUScene, const FViewMatrices& ViewMatrices, FRayTracingSceneWithGeometryInstances SceneWithGeometryInstances);
 
 	// Backwards-compatible version of Create() which internally calls CreateRayTracingSceneWithGeometryInstances().
-	void Create(FRDGBuilder& GraphBuilder, const FGPUScene& GPUScene, const FViewMatrices& ViewMatrices);
+	void Create(FRDGBuilder& GraphBuilder, const FGPUScene* GPUScene, const FViewMatrices& ViewMatrices);
 
 	// Resets the instance list and reserves memory for this frame.
 	void Reset();
diff --git a/Engine/Source/Runtime/Renderer/Private/ScenePrivate.h b/Engine/Source/Runtime/Renderer/Private/ScenePrivate.h
--- a/Engine/Source/Runtime/Renderer/Private/ScenePrivate.h
+++ b/Engine/Source/Runtime/Renderer/Private/ScenePrivate.h
@@ -2916,6 +2916,7 @@ public:
 #if RHI_RAYTRACING
 	/** Persistently-allocated ray tracing scene data. */
 	FRayTracingScene RayTracingScene;
+	FRayTracingScene HeterogeneousVolumesRayTracingScene;
 
 	bool bHasRayTracedLights = false;
 	void UpdateRayTracedLights();
diff --git a/Engine/Source/Runtime/Renderer/Private/SceneRendering.h b/Engine/Source/Runtime/Renderer/Private/SceneRendering.h
--- a/Engine/Source/Runtime/Renderer/Private/SceneRendering.h
+++ b/Engine/Source/Runtime/Renderer/Private/SceneRendering.h
@@ -1435,6 +1435,8 @@ public:
 	FForwardLightingViewResources ForwardLightingResources;
 	FVolumetricFogViewResources VolumetricFogResources;
 
+	FRDGTextureRef HeterogeneousVolumeRadiance = nullptr;
+
 	// Size of the HZB's mipmap 0
 	// NOTE: the mipmap 0 is downsampled version of the depth buffer
 	FIntPoint HZBMipmap0Size;
diff --git a/Engine/Source/Runtime/Renderer/Private/VolumetricFogVoxelization.cpp b/Engine/Source/Runtime/Renderer/Private/VolumetricFogVoxelization.cpp
--- a/Engine/Source/Runtime/Renderer/Private/VolumetricFogVoxelization.cpp
+++ b/Engine/Source/Runtime/Renderer/Private/VolumetricFogVoxelization.cpp
@@ -15,6 +15,7 @@
 #include "StaticMeshResources.h"
 #include "MeshPassProcessor.inl"
 #include "VolumetricCloudRendering.h"
+#include "HeterogeneousVolumes/HeterogeneousVolumes.h"
 
 int32 GVolumetricFogVoxelizationSlicesPerGSPass = 8;
 FAutoConsoleVariableRef CVarVolumetricFogVoxelizationSlicesPerPass(
@@ -661,9 +662,20 @@ void FDeferredShadingSceneRenderer::VoxelizeFogVolumePrimitives(
 						&View,
 						DynamicMeshPassContext);
 
+					const bool bShouldRenderHeterogeneousVolumes = ShouldRenderHeterogeneousVolumesForView(View);
+
 					for (int32 MeshBatchIndex = 0; MeshBatchIndex < View.VolumetricMeshBatches.Num(); ++MeshBatchIndex)
 					{
 						const FMeshBatch* Mesh = View.VolumetricMeshBatches[MeshBatchIndex].Mesh;
+						const FMaterialRenderProxy* MaterialRenderProxy = Mesh->MaterialRenderProxy;
+						const FMaterial& Material = MaterialRenderProxy->GetMaterialWithFallback(View.GetFeatureLevel(), MaterialRenderProxy);
+
+						// Skip volumes flagged as rendered with HeterogenousVolumes
+						if (bShouldRenderHeterogeneousVolumes && Material.IsUsedWithNiagaraMeshParticles())
+						{
+							continue;
+						}
+
 						const FPrimitiveSceneProxy* PrimitiveSceneProxy = View.VolumetricMeshBatches[MeshBatchIndex].Proxy;
 						const FPrimitiveSceneInfo* PrimitiveSceneInfo = PrimitiveSceneProxy->GetPrimitiveSceneInfo();
 						const FBoxSphereBounds Bounds = PrimitiveSceneProxy->GetBounds();
