Fixing static anlysis warnings.

diff --git a/Engine/Source/Runtime/Renderer/Private/HeterogeneousVolumes/HeterogeneousVolumesLiveShadingPipeline.cpp b/Engine/Source/Runtime/Renderer/Private/HeterogeneousVolumes/HeterogeneousVolumesLiveShadingPipeline.cpp
--- a/Engine/Source/Runtime/Renderer/Private/HeterogeneousVolumes/HeterogeneousVolumesLiveShadingPipeline.cpp
+++ b/Engine/Source/Runtime/Renderer/Private/HeterogeneousVolumes/HeterogeneousVolumesLiveShadingPipeline.cpp
@@ -274,6 +274,8 @@ void RenderWithLiveShading(
 			LightType = LightSceneInfoCompact[PassIndex].LightType;
 			LightSceneInfo = LightSceneInfoCompact[PassIndex].LightSceneInfo;
 			check(LightSceneInfo != nullptr);
+
+			bApplyDirectLighting = (LightSceneInfo != nullptr);
 		}
 
 		RenderSingleScatteringWithLiveShading(
diff --git a/Engine/Source/Runtime/Renderer/Private/HeterogeneousVolumes/HeterogeneousVolumesPreshadingPipeline.cpp b/Engine/Source/Runtime/Renderer/Private/HeterogeneousVolumes/HeterogeneousVolumesPreshadingPipeline.cpp
--- a/Engine/Source/Runtime/Renderer/Private/HeterogeneousVolumes/HeterogeneousVolumesPreshadingPipeline.cpp
+++ b/Engine/Source/Runtime/Renderer/Private/HeterogeneousVolumes/HeterogeneousVolumesPreshadingPipeline.cpp
@@ -434,6 +434,8 @@ void RenderWithPreshadingCompute(
 				LightType = LightSceneInfoCompact[PassIndex].LightType;
 				LightSceneInfo = LightSceneInfoCompact[PassIndex].LightSceneInfo;
 				check(LightSceneInfo != nullptr);
+
+				bApplyDirectLighting = (LightSceneInfo != nullptr);
 			}
 
 			RenderSingleScatteringWithPreshadingCompute(
@@ -548,6 +550,8 @@ void RenderWithPreshadingHardwareRayTracing(
 				LightType = LightSceneInfoCompact[PassIndex].LightType;
 				LightSceneInfo = LightSceneInfoCompact[PassIndex].LightSceneInfo;
 				check(LightSceneInfo != nullptr);
+
+				bApplyDirectLighting = (LightSceneInfo != nullptr);
 			}
 
 			RenderSingleScatteringWithPreshadingHardwareRayTracing(
