HeterogeneousVolumes: Only apply volumetric self-shadowing for lights that are marked as casting volumetric shadows.

diff --git a/Engine/Shaders/Private/HeterogeneousVolumes/HeterogeneousVolumesHardwareRayTracing.usf b/Engine/Shaders/Private/HeterogeneousVolumes/HeterogeneousVolumesHardwareRayTracing.usf
--- a/Engine/Shaders/Private/HeterogeneousVolumes/HeterogeneousVolumesHardwareRayTracing.usf
+++ b/Engine/Shaders/Private/HeterogeneousVolumes/HeterogeneousVolumesHardwareRayTracing.usf
@@ -10,6 +10,10 @@
 #include "HeterogeneousVolumesTracingUtils.ush"
 #include "HeterogeneousVolumesLightingUtils.ush"
 
+#ifndef DIM_APPLY_SHADOW_TRANSMITTANCE
+#define DIM_APPLY_SHADOW_TRANSMITTANCE 0
+#endif // DIM_APPLY_SHADOW_TRANSMITTANCE
+
 #define VOXEL_COUNT_MAX 6
 
 #define SHADOW_BIT 0x80000000
@@ -458,7 +462,8 @@ RAY_TRACING_ENTRY_RAYGEN(RenderSingleScatteringWithPreshadingRGS)
 			GetStepSize(),
 			MaxStepCount,
 			bApplyEmission,
-			bApplyDirectLighting
+			bApplyDirectLighting,
+			DIM_APPLY_SHADOW_TRANSMITTANCE
 		);
 
 		FDeferredLightData LightData = LoadLightData(LightType);
diff --git a/Engine/Shaders/Private/HeterogeneousVolumes/HeterogeneousVolumesLiveShadingPipeline.usf b/Engine/Shaders/Private/HeterogeneousVolumes/HeterogeneousVolumesLiveShadingPipeline.usf
--- a/Engine/Shaders/Private/HeterogeneousVolumes/HeterogeneousVolumesLiveShadingPipeline.usf
+++ b/Engine/Shaders/Private/HeterogeneousVolumes/HeterogeneousVolumesLiveShadingPipeline.usf
@@ -20,6 +20,7 @@ int PrimitiveId;
 // Lighting
 int bApplyEmission;
 int bApplyDirectLighting;
+int bApplyShadowTransmittance;
 int LightType;
 float ShadowStepFactor;
 
@@ -203,7 +204,8 @@ void RenderSingleScatteringWithLiveShadingCS(
 			StepSize,
 			MaxStepCount,
 			bApplyEmission,
-			bApplyDirectLighting
+			bApplyDirectLighting,
+			bApplyShadowTransmittance
 		);
 
 		FDeferredLightData LightData = LoadLightData(LightType);
diff --git a/Engine/Shaders/Private/HeterogeneousVolumes/HeterogeneousVolumesPreshadingPipeline.usf b/Engine/Shaders/Private/HeterogeneousVolumes/HeterogeneousVolumesPreshadingPipeline.usf
--- a/Engine/Shaders/Private/HeterogeneousVolumes/HeterogeneousVolumesPreshadingPipeline.usf
+++ b/Engine/Shaders/Private/HeterogeneousVolumes/HeterogeneousVolumesPreshadingPipeline.usf
@@ -30,6 +30,10 @@
 #define DIM_VOXEL_CULLING 1
 #endif // DIM_VOXEL_CULLING
 
+#ifndef DIM_APPLY_SHADOW_TRANSMITTANCE
+#define DIM_APPLY_SHADOW_TRANSMITTANCE 0
+#endif // DIM_APPLY_SHADOW_TRANSMITTANCE
+
 // Lighting
 int bApplyEmission;
 int bApplyDirectLighting;
@@ -495,7 +499,8 @@ void RenderSingleScatteringWithPreshadingCS(
 			StepSize,
 			MaxStepCount,
 			bApplyEmission,
-			bApplyDirectLighting
+			bApplyDirectLighting,
+			DIM_APPLY_SHADOW_TRANSMITTANCE
 		);
 
 		FDeferredLightData LightData = LoadLightData(LightType);
diff --git a/Engine/Shaders/Private/HeterogeneousVolumes/HeterogeneousVolumesRayMarchingUtils.ush b/Engine/Shaders/Private/HeterogeneousVolumes/HeterogeneousVolumesRayMarchingUtils.ush
--- a/Engine/Shaders/Private/HeterogeneousVolumes/HeterogeneousVolumesRayMarchingUtils.ush
+++ b/Engine/Shaders/Private/HeterogeneousVolumes/HeterogeneousVolumesRayMarchingUtils.ush
@@ -21,6 +21,7 @@ struct FRayMarchingContext
 	// Lighting
 	int bApplyEmission;
 	int bApplyDirectLighting;
+	int bApplyShadowTransmittance;
 };
 
 FRayMarchingContext CreateRayMarchingContext(
@@ -38,7 +39,8 @@ FRayMarchingContext CreateRayMarchingContext(
 	uint MaxStepCount,
 	// Lighting
 	int bApplyEmission,
-	int bApplyDirectLighting
+	int bApplyDirectLighting,
+	int bApplyShadowTransmittance
 )
 {
 	FRayMarchingContext Context;
@@ -57,6 +59,7 @@ FRayMarchingContext CreateRayMarchingContext(
 	// Lighting
 	Context.bApplyEmission = bApplyEmission;
 	Context.bApplyDirectLighting = bApplyDirectLighting;
+	Context.bApplyShadowTransmittance = bApplyShadowTransmittance;
 
 	return Context;
 }
@@ -199,35 +202,43 @@ void RayMarchSingleScattering(
 
 				if (any(Inscattering > 0.0))
 				{
-					// Build shadow ray and find intersection
-					float3 LocalShadowRayOrigin = mul(float4(WorldPosition, 1.0), GetWorldToLocal()).xyz;
-					float3 LocalShadowRayEnd = mul(float4(WorldPosition + ToLight, 1.0), GetWorldToLocal()).xyz;
-					float3 LocalShadowRayDirection = LocalShadowRayEnd - LocalShadowRayOrigin;
-					float LocalShadowRayTMin = 0.0;
-					float LocalShadowRayTMax = length(LocalShadowRayDirection);
-					LocalShadowRayDirection /= LocalShadowRayTMax;
-
-					//float ShadowJitter = InterleavedGradientNoise(PixelCoord, StepIndex);
-					float ShadowJitter = 0.0;
-					float ShadowStepSize = GetStepSize() * GetShadowStepFactor();
-					int bApplyEmission = 0;
-					int bApplyDirectLighting = 0;
-
-					FRayMarchingContext ShadowRayMarchingContext = CreateRayMarchingContext(
-						LocalShadowRayOrigin,
-						LocalShadowRayDirection,
-						LocalShadowRayTMin,
-						LocalShadowRayTMax,
-						WorldPosition, // WorldShadowRayOrigin
-						normalize(ToLight), // WorldShadowRayDirection
-						ShadowJitter,
-						ShadowStepSize,
-						RayMarchingContext.MaxStepCount,
-						bApplyEmission,
-						bApplyDirectLighting
-						//RayMarchingContext.MaxShadowTraceDistance
-					);
-					float3 ShadowTransmittance = ComputeTransmittance(ShadowRayMarchingContext);
+					float3 ShadowTransmittance = 1.0;
+
+					if (RayMarchingContext.bApplyShadowTransmittance)
+					{
+						// Build shadow ray and find intersection
+						float3 LocalShadowRayOrigin = mul(float4(WorldPosition, 1.0), GetWorldToLocal()).xyz;
+						float3 LocalShadowRayEnd = mul(float4(WorldPosition + ToLight, 1.0), GetWorldToLocal()).xyz;
+						float3 LocalShadowRayDirection = LocalShadowRayEnd - LocalShadowRayOrigin;
+						float LocalShadowRayTMin = 0.0;
+						float LocalShadowRayTMax = length(LocalShadowRayDirection);
+						LocalShadowRayDirection /= LocalShadowRayTMax;
+
+						//float ShadowJitter = InterleavedGradientNoise(PixelCoord, StepIndex);
+						float ShadowJitter = 0.0;
+						float ShadowStepSize = GetStepSize() * GetShadowStepFactor();
+						int bApplyEmission = 0;
+						int bApplyDirectLighting = 0;
+						int bApplyShadowTransmittance = 0;
+
+						FRayMarchingContext ShadowRayMarchingContext = CreateRayMarchingContext(
+							LocalShadowRayOrigin,
+							LocalShadowRayDirection,
+							LocalShadowRayTMin,
+							LocalShadowRayTMax,
+							WorldPosition, // WorldShadowRayOrigin
+							normalize(ToLight), // WorldShadowRayDirection
+							ShadowJitter,
+							ShadowStepSize,
+							RayMarchingContext.MaxStepCount,
+							bApplyEmission,
+							bApplyDirectLighting,
+							bApplyShadowTransmittance
+							//RayMarchingContext.MaxShadowTraceDistance
+						);
+
+						ShadowTransmittance = ComputeTransmittance(ShadowRayMarchingContext);
+					}
 
 					// Convert in-scattering to exitant radiance
 					float3 ScatteringCoefficient = Albedo * Extinction;
diff --git a/Engine/Source/Runtime/Renderer/Private/HeterogeneousVolumes/HeterogeneousVolumes.h b/Engine/Source/Runtime/Renderer/Private/HeterogeneousVolumes/HeterogeneousVolumes.h
--- a/Engine/Source/Runtime/Renderer/Private/HeterogeneousVolumes/HeterogeneousVolumes.h
+++ b/Engine/Source/Runtime/Renderer/Private/HeterogeneousVolumes/HeterogeneousVolumes.h
@@ -190,6 +190,7 @@ void RenderSingleScatteringWithPreshadingHardwareRayTracing(
 	// Light data
 	bool bApplyEmission,
 	bool bApplyDirectLighting,
+	bool bApplyShadowTransmittance,
 	uint32 LightType,
 	const FLightSceneInfo* LightSceneInfo,
 	// Object data
diff --git a/Engine/Source/Runtime/Renderer/Private/HeterogeneousVolumes/HeterogeneousVolumesHardwareRayTracing.cpp b/Engine/Source/Runtime/Renderer/Private/HeterogeneousVolumes/HeterogeneousVolumesHardwareRayTracing.cpp
--- a/Engine/Source/Runtime/Renderer/Private/HeterogeneousVolumes/HeterogeneousVolumesHardwareRayTracing.cpp
+++ b/Engine/Source/Runtime/Renderer/Private/HeterogeneousVolumes/HeterogeneousVolumesHardwareRayTracing.cpp
@@ -324,7 +324,8 @@ class FRenderSingleScatteringWithPreshadingRGS : public FGlobalShader
 	DECLARE_GLOBAL_SHADER(FRenderSingleScatteringWithPreshadingRGS)
 	SHADER_USE_ROOT_PARAMETER_STRUCT(FRenderSingleScatteringWithPreshadingRGS, FGlobalShader)
 
-	using FPermutationDomain = TShaderPermutationDomain<>;
+	class FApplyShadowTransmittanceDim : SHADER_PERMUTATION_BOOL("DIM_APPLY_SHADOW_TRANSMITTANCE");
+	using FPermutationDomain = TShaderPermutationDomain<FApplyShadowTransmittanceDim>;
 
 	BEGIN_SHADER_PARAMETER_STRUCT(FParameters, )
 		// Scene 
@@ -446,6 +447,7 @@ void RenderSingleScatteringWithPreshadingHardwareRayTracing(
 	// Light data
 	bool bApplyEmission,
 	bool bApplyDirectLighting,
+	bool bApplyShadowTransmittance,
 	uint32 LightType,
 	const FLightSceneInfo* LightSceneInfo,
 	// Object data
@@ -492,6 +494,7 @@ void RenderSingleScatteringWithPreshadingHardwareRayTracing(
 	}
 
 	FRenderSingleScatteringWithPreshadingRGS::FPermutationDomain PermutationVector;
+	PermutationVector.Set<FRenderSingleScatteringWithPreshadingRGS::FApplyShadowTransmittanceDim>(bApplyShadowTransmittance);
 	TShaderRef<FRenderSingleScatteringWithPreshadingRGS> RayGenerationShader = View.ShaderMap->GetShader<FRenderSingleScatteringWithPreshadingRGS>(PermutationVector);
 	FIntPoint DispatchResolution = View.ViewRect.Size();
 
diff --git a/Engine/Source/Runtime/Renderer/Private/HeterogeneousVolumes/HeterogeneousVolumesLiveShadingPipeline.cpp b/Engine/Source/Runtime/Renderer/Private/HeterogeneousVolumes/HeterogeneousVolumesLiveShadingPipeline.cpp
--- a/Engine/Source/Runtime/Renderer/Private/HeterogeneousVolumes/HeterogeneousVolumesLiveShadingPipeline.cpp
+++ b/Engine/Source/Runtime/Renderer/Private/HeterogeneousVolumes/HeterogeneousVolumesLiveShadingPipeline.cpp
@@ -22,6 +22,7 @@ class FRenderSingleScatteringWithLiveShadingCS : public FMeshMaterialShader
 		// Light data
 		SHADER_PARAMETER(int, bApplyEmission)
 		SHADER_PARAMETER(int, bApplyDirectLighting)
+		SHADER_PARAMETER(int, bApplyShadowTransmittance)
 		SHADER_PARAMETER(int, LightType)
 		SHADER_PARAMETER_STRUCT_REF(FDeferredLightUniformStruct, DeferredLight)
 
@@ -123,6 +124,7 @@ void RenderSingleScatteringWithLiveShading(
 	// Light data
 	bool bApplyEmission,
 	bool bApplyDirectLighting,
+	bool bApplyShadowTransmittance,
 	uint32 LightType,
 	const FLightSceneInfo* LightSceneInfo,
 	// Object data
@@ -161,6 +163,7 @@ void RenderSingleScatteringWithLiveShading(
 		FDeferredLightUniformStruct DeferredLightUniform;
 		PassParameters->bApplyEmission = bApplyEmission;
 		PassParameters->bApplyDirectLighting = bApplyDirectLighting;
+		PassParameters->bApplyShadowTransmittance = bApplyShadowTransmittance;
 		if (PassParameters->bApplyDirectLighting)
 		{
 			DeferredLightUniform = GetDeferredLightParameters(View, *LightSceneInfo);
@@ -266,6 +269,7 @@ void RenderWithLiveShading(
 	{
 		bool bApplyEmission = PassIndex == 0;
 		bool bApplyDirectLighting = !LightSceneInfoCompact.IsEmpty();
+		bool bApplyShadowTransmittance = false;
 
 		uint32 LightType = 0;
 		FLightSceneInfo* LightSceneInfo = nullptr;
@@ -276,6 +280,7 @@ void RenderWithLiveShading(
 			check(LightSceneInfo != nullptr);
 
 			bApplyDirectLighting = (LightSceneInfo != nullptr);
+			bApplyShadowTransmittance = LightSceneInfo->Proxy->CastsVolumetricShadow();
 		}
 
 		RenderSingleScatteringWithLiveShading(
@@ -287,6 +292,7 @@ void RenderWithLiveShading(
 			// Light data
 			bApplyEmission,
 			bApplyDirectLighting,
+			bApplyShadowTransmittance,
 			LightType,
 			LightSceneInfo,
 			// Object data
diff --git a/Engine/Source/Runtime/Renderer/Private/HeterogeneousVolumes/HeterogeneousVolumesPreshadingPipeline.cpp b/Engine/Source/Runtime/Renderer/Private/HeterogeneousVolumes/HeterogeneousVolumesPreshadingPipeline.cpp
--- a/Engine/Source/Runtime/Renderer/Private/HeterogeneousVolumes/HeterogeneousVolumesPreshadingPipeline.cpp
+++ b/Engine/Source/Runtime/Renderer/Private/HeterogeneousVolumes/HeterogeneousVolumesPreshadingPipeline.cpp
@@ -204,10 +204,11 @@ class FRenderSingleScatteringWithPreshadingCS : public FGlobalShader
 	DECLARE_GLOBAL_SHADER(FRenderSingleScatteringWithPreshadingCS);
 	SHADER_USE_PARAMETER_STRUCT(FRenderSingleScatteringWithPreshadingCS, FGlobalShader);
 
-	class FDebugDim : SHADER_PERMUTATION_BOOL("DIM_DEBUG");
+	class FApplyShadowTransmittanceDim : SHADER_PERMUTATION_BOOL("DIM_APPLY_SHADOW_TRANSMITTANCE");
 	class FVoxelCullingDim : SHADER_PERMUTATION_BOOL("DIM_VOXEL_CULLING");
 	class FSparseVoxelTracingDim : SHADER_PERMUTATION_BOOL("DIM_SPARSE_VOXEL_TRACING");
-	using FPermutationDomain = TShaderPermutationDomain<FDebugDim, FVoxelCullingDim, FSparseVoxelTracingDim>;
+	class FDebugDim : SHADER_PERMUTATION_BOOL("DIM_DEBUG");
+	using FPermutationDomain = TShaderPermutationDomain<FApplyShadowTransmittanceDim, FVoxelCullingDim, FSparseVoxelTracingDim, FDebugDim>;
 
 	BEGIN_SHADER_PARAMETER_STRUCT(FParameters, )
 		// Scene data
@@ -282,6 +283,7 @@ void RenderSingleScatteringWithPreshadingCompute(
 	// Light data
 	bool bApplyEmission,
 	bool bApplyDirectLighting,
+	bool bApplyShadowTransmittance,
 	uint32 LightType,
 	const FLightSceneInfo* LightSceneInfo,
 	// Object data
@@ -350,9 +352,10 @@ void RenderSingleScatteringWithPreshadingCompute(
 	}
 
 	FRenderSingleScatteringWithPreshadingCS::FPermutationDomain PermutationVector;
-	PermutationVector.Set<FRenderSingleScatteringWithPreshadingCS::FDebugDim>(HeterogeneousVolumes::GetDebugMode() != 0);
+	PermutationVector.Set<FRenderSingleScatteringWithPreshadingCS::FApplyShadowTransmittanceDim>(bApplyShadowTransmittance);
 	PermutationVector.Set<FRenderSingleScatteringWithPreshadingCS::FVoxelCullingDim>(HeterogeneousVolumes::UseSparseVoxelPerTileCulling());
 	PermutationVector.Set<FRenderSingleScatteringWithPreshadingCS::FSparseVoxelTracingDim>(HeterogeneousVolumes::UseSparseVoxelPipeline());
+	PermutationVector.Set<FRenderSingleScatteringWithPreshadingCS::FDebugDim>(HeterogeneousVolumes::GetDebugMode() != 0);
 	TShaderRef<FRenderSingleScatteringWithPreshadingCS> ComputeShader = View.ShaderMap->GetShader<FRenderSingleScatteringWithPreshadingCS>(PermutationVector);
 	FComputeShaderUtils::AddPass(
 		GraphBuilder,
@@ -426,6 +429,7 @@ void RenderWithPreshadingCompute(
 		{
 			bool bApplyEmission = PassIndex == 0;
 			bool bApplyDirectLighting = !LightSceneInfoCompact.IsEmpty();
+			bool bApplyShadowTransmittance = false;
 
 			uint32 LightType = 0;
 			FLightSceneInfo* LightSceneInfo = nullptr;
@@ -436,6 +440,7 @@ void RenderWithPreshadingCompute(
 				check(LightSceneInfo != nullptr);
 
 				bApplyDirectLighting = (LightSceneInfo != nullptr);
+				bApplyShadowTransmittance = LightSceneInfo->Proxy->CastsVolumetricShadow();
 			}
 
 			RenderSingleScatteringWithPreshadingCompute(
@@ -447,6 +452,7 @@ void RenderWithPreshadingCompute(
 				// Light
 				bApplyEmission,
 				bApplyDirectLighting,
+				bApplyShadowTransmittance,
 				LightType,
 				LightSceneInfo,
 				// Object
@@ -542,6 +548,7 @@ void RenderWithPreshadingHardwareRayTracing(
 		{
 			bool bApplyEmission = PassIndex == 0;
 			bool bApplyDirectLighting = !LightSceneInfoCompact.IsEmpty();
+			bool bApplyShadowTransmittance = false;
 
 			uint32 LightType = 0;
 			FLightSceneInfo* LightSceneInfo = nullptr;
@@ -552,6 +559,7 @@ void RenderWithPreshadingHardwareRayTracing(
 				check(LightSceneInfo != nullptr);
 
 				bApplyDirectLighting = (LightSceneInfo != nullptr);
+				bApplyShadowTransmittance = LightSceneInfo->Proxy->CastsVolumetricShadow();
 			}
 
 			RenderSingleScatteringWithPreshadingHardwareRayTracing(
@@ -563,6 +571,7 @@ void RenderWithPreshadingHardwareRayTracing(
 				// Light data
 				bApplyEmission,
 				bApplyDirectLighting,
+				bApplyShadowTransmittance,
 				LightType,
 				LightSceneInfo,
 				// Object data
