Shader Complexity fixed for Preview ES31
Quad Complexity support code added for Preview ES31
Quad Complexity temporary disabled for Preview ES31

diff --git a/Engine/Config/Windows/DataDrivenPlatformInfo.ini b/Engine/Config/Windows/DataDrivenPlatformInfo.ini
--- a/Engine/Config/Windows/DataDrivenPlatformInfo.ini
+++ b/Engine/Config/Windows/DataDrivenPlatformInfo.ini
@@ -145,6 +145,7 @@ bSupportsByteBufferComputeShaders = true
 bSupportsInstancedStereo=true
 bSupportsManualVertexFetch=false
 bSupportsSceneDataCompressedTransforms=true
+bSupportsDebugViewShaders=true
 
 [ShaderPlatform OPENGL_PCES3_1]
 Language=OpenGL
diff --git a/Engine/Shaders/Private/DebugViewModePixelShader.usf b/Engine/Shaders/Private/DebugViewModePixelShader.usf
--- a/Engine/Shaders/Private/DebugViewModePixelShader.usf
+++ b/Engine/Shaders/Private/DebugViewModePixelShader.usf
@@ -17,7 +17,9 @@
 
 #include "Common.ush"
 #include "SHCommon.ush"
+#if OUTPUT_QUAD_OVERDRAW
 #include "QuadOverdraw.ush"
+#endif
 #include "VirtualTextureCommon.ush"
 
 #define UNDEF_ACCURACY 1024
@@ -462,7 +464,7 @@ float4 VisualizeShaderComplexityAccumulate(in FDebugPSIn DebugInputs, in bool bI
 	{
 		float3 FinalComplexity = NormalizedComplexity.xyz;
 
-#if !(COMPILER_METAL)
+#if !(COMPILER_METAL) && (OUTPUT_QUAD_OVERDRAW)
 		[branch]
 		if (bShowQuadOverdraw && NormalizedComplexity.x > 0)
 		{
@@ -470,7 +472,7 @@ float4 VisualizeShaderComplexityAccumulate(in FDebugPSIn DebugInputs, in bool bI
 			// The extra cost from quad overdraw is assumed to be linear.
 			FinalComplexity.x *= 4.f / (float)(Coverage);
 		}
-#endif // !(COMPILER_METAL)
+#endif // !(COMPILER_METAL) && (OUTPUT_QUAD_OVERDRAW)
 
 		// use the maximum range allowed for scene color
 		Result = FinalComplexity;
diff --git a/Engine/Shaders/Private/QuadOverdraw.ush b/Engine/Shaders/Private/QuadOverdraw.ush
--- a/Engine/Shaders/Private/QuadOverdraw.ush
+++ b/Engine/Shaders/Private/QuadOverdraw.ush
@@ -45,7 +45,11 @@ float ComplexityToFloat(uint Complexity)
 	return (float)Complexity / (float)QO_FRACTIONNAL_COMPLEXITY;
 }
 
+#if SHADING_PATH_MOBILE
+#define RWQuadBuffer MobileBasePass.QuadOverdraw 
+#else
 #define RWQuadBuffer DebugViewModePass.QuadOverdraw
+#endif
 
 #if OUTPUT_QUAD_OVERDRAW
 
diff --git a/Engine/Source/Runtime/Renderer/Private/MobileBasePassRendering.cpp b/Engine/Source/Runtime/Renderer/Private/MobileBasePassRendering.cpp
--- a/Engine/Source/Runtime/Renderer/Private/MobileBasePassRendering.cpp
+++ b/Engine/Source/Runtime/Renderer/Private/MobileBasePassRendering.cpp
@@ -217,6 +217,18 @@ TRDGUniformBufferRef<FMobileBasePassUniformParameters> CreateMobileBasePassUnifo
 		SetupDebugViewModePassUniformBufferConstants(View, BasePassParameters->DebugViewMode);
 	}
 #endif
+
+	// Cannot be null even when DebugVM is off because it is in the same pass as RenderMobileBasePass(ValidateShaderParameterResourcesRHI UAV).
+	// If this is problematic, move DebugView to run in another pass and use FDebugViewModePassUniformParameters
+	const FSceneTextures& SceneTextures = View.GetSceneTextures();
+	FRDGTextureRef QuadOverdrawTexture = SceneTextures.QuadOverdraw;
+	if (!QuadOverdrawTexture)
+	{
+		QuadOverdrawTexture = GraphBuilder.CreateTexture(FRDGTextureDesc::Create2D(FIntPoint(1, 1), PF_R32_UINT, FClearValueBinding::None, TexCreate_UAV), TEXT("DummyOverdrawUAV"));
+	}
+
+	BasePassParameters->QuadOverdraw = GraphBuilder.CreateUAV(QuadOverdrawTexture);
+
 	return GraphBuilder.CreateUniformBuffer(BasePassParameters);
 }
 
diff --git a/Engine/Source/Runtime/Renderer/Private/MobileBasePassRendering.h b/Engine/Source/Runtime/Renderer/Private/MobileBasePassRendering.h
--- a/Engine/Source/Runtime/Renderer/Private/MobileBasePassRendering.h
+++ b/Engine/Source/Runtime/Renderer/Private/MobileBasePassRendering.h
@@ -38,6 +38,7 @@ BEGIN_GLOBAL_SHADER_PARAMETER_STRUCT(FMobileBasePassUniformParameters, )
 	SHADER_PARAMETER_STRUCT(FPlanarReflectionUniformParameters, PlanarReflection) // Single global planar reflection for the forward pass.
 	SHADER_PARAMETER_STRUCT(FMobileSceneTextureUniformParameters, SceneTextures)
 	SHADER_PARAMETER_STRUCT(FDebugViewModeUniformParameters, DebugViewMode)
+	SHADER_PARAMETER_RDG_TEXTURE_UAV(RWTexture2D<uint>, QuadOverdraw)
 	SHADER_PARAMETER_STRUCT(FReflectionUniformParameters, ReflectionsParameters)
 	SHADER_PARAMETER_TEXTURE(Texture2D, PreIntegratedGFTexture)
 	SHADER_PARAMETER_SAMPLER(SamplerState, PreIntegratedGFSampler)
