Adds DOF squeeze factor support to path tracer.

diff --git a/Engine/Shaders/Private/PathTracing/PathTracingCore.ush b/Engine/Shaders/Private/PathTracing/PathTracingCore.ush
--- a/Engine/Shaders/Private/PathTracing/PathTracingCore.ush
+++ b/Engine/Shaders/Private/PathTracing/PathTracingCore.ush
@@ -59,7 +59,7 @@ float MaxNormalBias;
 float FilterWidth;
 float AbsorptionScale; // inverse distance at which we hit BaseColor - exposed mainly so we can disable absorption handling and experiment with the default scale factor
 float CameraFocusDistance;
-float CameraLensRadius;
+float2 CameraLensRadius;
 
 RWTexture2D<float4> RadianceTexture;
 RWTexture2D<float4> AlbedoTexture;
@@ -794,7 +794,7 @@ FPathState CreatePathState(int2 PixelIndex)
 		Output.Ray = CreatePrimaryRay(ViewportUV);
 	}
 
-	if (CameraLensRadius > 0)
+	if (CameraLensRadius.y > 0)
 	{
 	    // DOF enabled - apply simple thin-lens model
 		float2 LensSample = RandomSequence_GenerateSample2D(Output.RandSequence);
diff --git a/Engine/Source/Runtime/Renderer/Private/PathTracing.cpp b/Engine/Source/Runtime/Renderer/Private/PathTracing.cpp
--- a/Engine/Source/Runtime/Renderer/Private/PathTracing.cpp
+++ b/Engine/Source/Runtime/Renderer/Private/PathTracing.cpp
@@ -304,7 +304,7 @@ BEGIN_SHADER_PARAMETER_STRUCT(FPathTracingData, )
 	SHADER_PARAMETER(float, FilterWidth)
 	SHADER_PARAMETER(float, AbsorptionScale)
 	SHADER_PARAMETER(float, CameraFocusDistance)
-	SHADER_PARAMETER(float, CameraLensRadius)
+	SHADER_PARAMETER(FVector2f, CameraLensRadius)
 END_SHADER_PARAMETER_STRUCT()
 
 
@@ -493,7 +493,7 @@ static void PreparePathTracingData(const FScene* Scene, const FViewInfo& View, F
 	PathTracingData.FilterWidth = FilterWidth;
 	PathTracingData.AbsorptionScale = CVarPathTracingAbsorptionScale.GetValueOnRenderThread();
 	PathTracingData.CameraFocusDistance = 0;
-	PathTracingData.CameraLensRadius = 0;
+	PathTracingData.CameraLensRadius = FVector2f::ZeroVector;
 	if (View.Family->EngineShowFlags.DepthOfField &&
 		View.FinalPostProcessSettings.PathTracingEnableReferenceDOF &&
 		View.FinalPostProcessSettings.DepthOfFieldFocalDistance > 0 &&
@@ -501,7 +501,8 @@ static void PreparePathTracingData(const FScene* Scene, const FViewInfo& View, F
 	{
 		const float FocalLengthInCM = 0.05f * View.FinalPostProcessSettings.DepthOfFieldSensorWidth * View.ViewMatrices.GetProjectionMatrix().M[0][0];
 		PathTracingData.CameraFocusDistance = View.FinalPostProcessSettings.DepthOfFieldFocalDistance;
-		PathTracingData.CameraLensRadius = 0.5f * FocalLengthInCM / View.FinalPostProcessSettings.DepthOfFieldFstop;
+		PathTracingData.CameraLensRadius.Y = 0.5f * FocalLengthInCM / View.FinalPostProcessSettings.DepthOfFieldFstop;
+		PathTracingData.CameraLensRadius.X = PathTracingData.CameraLensRadius.Y / FMath::Clamp(View.FinalPostProcessSettings.DepthOfFieldSqueezeFactor, 1.0f, 2.0f);
 	}
 	PathTracingData.EnableAtmosphere =
 		ShouldRenderSkyAtmosphere(Scene, View.Family->EngineShowFlags) && 
