Fixes a couples more squeeze factor issues in DOF

diff --git a/Engine/Shaders/Private/DiaphragmDOF/DOFBokehLUT.usf b/Engine/Shaders/Private/DiaphragmDOF/DOFBokehLUT.usf
--- a/Engine/Shaders/Private/DiaphragmDOF/DOFBokehLUT.usf
+++ b/Engine/Shaders/Private/DiaphragmDOF/DOFBokehLUT.usf
@@ -202,7 +202,7 @@ void BuildBokehLUTMainCS(uint2 DispatchThreadId : SV_DispatchThreadID)
 			CocRadiusToBokehEdgeFactor = 1.0;
 		}
 
-		float PixelCocDistanceFromBokehCenter = FullResPixelDistanceToCocDistance(length(RelativePosFromCenter) / CocRadiusToBokehEdgeFactor);
+		float PixelCocDistanceFromBokehCenter = FullResPixelDistanceToCocDistance(length(RelativePosFromCenter * float2(CocSqueeze, 1.0)) / CocRadiusToBokehEdgeFactor);
 		BokehLUTOutput[DispatchThreadId] = float4(PixelCocDistanceFromBokehCenter, 0, 0, 0);
 	}
 	#elif DIM_LUT_FORMAT == LUT_FORMAT_GATHER_SAMPLE_POS
@@ -258,6 +258,9 @@ void BuildBokehLUTMainCS(uint2 DispatchThreadId : SV_DispatchThreadID)
 
 		// Final position of the sample in kernel.
 		float2 SamplePos = KernelUnitPos * RingRadius;
+		
+		// Apply anamorphic squeeze factor to save VALU in DOFGatherKernel.ush
+		SamplePos.x *= CocInvSqueeze;
 
 		BokehLUTOutput[DispatchThreadId] = float4(SamplePos, 0, 0);
 	}
diff --git a/Engine/Shaders/Private/DiaphragmDOF/DOFGatherKernel.ush b/Engine/Shaders/Private/DiaphragmDOF/DOFGatherKernel.ush
--- a/Engine/Shaders/Private/DiaphragmDOF/DOFGatherKernel.ush
+++ b/Engine/Shaders/Private/DiaphragmDOF/DOFGatherKernel.ush
@@ -462,7 +462,6 @@ void GatherRingSamples(
 				};
 				
 				SamplePosOffset[0] = KernelScope.SampleCountToSamplePosition * BokehLUT.SampleLevel(GlobalBilinearWrappedSampler, LookUpUV[0], 0).xy;
-				SamplePosOffset[0].x *= CocInvSqueeze;
 
 				#if DIM_BOKEH_SIMULATION == BOKEH_SIMULATION_SIMMETRIC && 0
 					SamplePosOffset[1] = -SamplePosOffset[0];
diff --git a/Engine/Shaders/Private/DiaphragmDOF/DOFRecombine.usf b/Engine/Shaders/Private/DiaphragmDOF/DOFRecombine.usf
--- a/Engine/Shaders/Private/DiaphragmDOF/DOFRecombine.usf
+++ b/Engine/Shaders/Private/DiaphragmDOF/DOFRecombine.usf
@@ -562,7 +562,7 @@ void FetchAndAccumulateSamplePair(
 	float SampleDistance = FullResPixelDistanceToCocDistance(length(PixelOffset));
 
 	// Scene buffer offset.
-	float2 SceneBufferUVOffset = PixelOffset * View.BufferSizeAndInvSize.zw;
+	float2 SceneBufferUVOffset = PixelOffset * View.BufferSizeAndInvSize.zw * float2(CocInvSqueeze, 1.0);
 	
 	// Two sample to gather at same time.
 	FGatherSample Sample[2];
@@ -578,7 +578,7 @@ void FetchAndAccumulateSamplePair(
 		{
 			const float InvLutSize = rcp(float(BOKEH_LUT_SIZE));
 
-			float2 LookupUV = (0.5 + 0.5 * InvLutSize) + PixelOffset * (SampleSign * InvLutSize);
+			float2 LookupUV = (0.5 + 0.5 * InvLutSize) + PixelOffset * (SampleSign * InvLutSize) * float2(CocInvSqueeze, 1.0);
 			float4 LookupSample = BokehLUT.SampleLevel(GlobalPointClampedSampler, LookupUV, 0);
 			SampleDistance = LookupSample.x;
 		}
diff --git a/Engine/Source/Runtime/Renderer/Private/PostProcess/DiaphragmDOF.cpp b/Engine/Source/Runtime/Renderer/Private/PostProcess/DiaphragmDOF.cpp
--- a/Engine/Source/Runtime/Renderer/Private/PostProcess/DiaphragmDOF.cpp
+++ b/Engine/Source/Runtime/Renderer/Private/PostProcess/DiaphragmDOF.cpp
@@ -940,6 +940,8 @@ class FDiaphragmDOFBuildBokehLUTCS : public FDiaphragmDOFShader
 		SHADER_PARAMETER(float, CocRadiusToIncircleRadius)
 		SHADER_PARAMETER(float, DiaphragmBladeRadius)
 		SHADER_PARAMETER(float, DiaphragmBladeCenterOffset)
+		SHADER_PARAMETER(float, CocSqueeze)
+		SHADER_PARAMETER(float, CocInvSqueeze)
 
 		SHADER_PARAMETER_RDG_TEXTURE_UAV(RWTexture2D<float4>, BokehLUTOutput)
 	END_SHADER_PARAMETER_STRUCT()
@@ -2110,6 +2112,8 @@ FRDGTextureRef DiaphragmDOF::AddPasses(
 		PassParameters->CocRadiusToIncircleRadius = BokehModel.CocRadiusToIncircleRadius;
 		PassParameters->DiaphragmBladeRadius = BokehModel.RoundedBlades.DiaphragmBladeRadius;
 		PassParameters->DiaphragmBladeCenterOffset = BokehModel.RoundedBlades.DiaphragmBladeCenterOffset;
+		PassParameters->CocSqueeze = CocModel.Squeeze;
+		PassParameters->CocInvSqueeze = 1.0f / CocModel.Squeeze;
 		PassParameters->BokehLUTOutput = GraphBuilder.CreateUAV(BokehLUT);
 			
 		TShaderMapRef<FDiaphragmDOFBuildBokehLUTCS> ComputeShader(View.ShaderMap, PermutationVector);
