Adding bindless support to more MDC parameter types. Removing unneeded data from MDC UB binding info.

diff --git a/Engine/Shaders/Public/ShaderVersion.ush b/Engine/Shaders/Public/ShaderVersion.ush
--- a/Engine/Shaders/Public/ShaderVersion.ush
+++ b/Engine/Shaders/Public/ShaderVersion.ush
@@ -2,4 +2,4 @@
 // Each time the console command is executed it generates a new GUID. As this file is included
 // in Platform.ush (which should be included in any shader) it allows to invalidate the shader DDC.
 // 
-// GUID = 0131629C-FCE1-4B72-9872-DB06D2BA8DE8
+// GUID = E4136F6D-EC93-4BC0-9E76-AB60D1E5BABF
diff --git a/Engine/Source/Runtime/RenderCore/Private/Shader.cpp b/Engine/Source/Runtime/RenderCore/Private/Shader.cpp
--- a/Engine/Source/Runtime/RenderCore/Private/Shader.cpp
+++ b/Engine/Source/Runtime/RenderCore/Private/Shader.cpp
@@ -38,7 +38,9 @@ IMPLEMENT_TYPE_LAYOUT(FShaderParameterBindings);
 IMPLEMENT_TYPE_LAYOUT(FShaderMapContent);
 IMPLEMENT_TYPE_LAYOUT(FShaderTypeDependency);
 IMPLEMENT_TYPE_LAYOUT(FShaderPipeline);
-IMPLEMENT_TYPE_LAYOUT(FShaderParameterInfo);
+IMPLEMENT_TYPE_LAYOUT(FShaderUniformBufferParameterInfo);
+IMPLEMENT_TYPE_LAYOUT(FShaderResourceParameterInfo);
+IMPLEMENT_TYPE_LAYOUT(FShaderLooseParameterInfo);
 IMPLEMENT_TYPE_LAYOUT(FShaderLooseParameterBufferInfo);
 IMPLEMENT_TYPE_LAYOUT(FShaderParameterMapInfo);
 
@@ -599,6 +601,14 @@ void FShader::Finalize(const FShaderMapResourceCode* Code)
 	ResourceIndex = NewResourceIndex;
 }
 
+template<class TType>
+static void CityHashArray(uint64& Hash, const TMemoryImageArray<TType>& Array)
+{
+	const int32 ArrayNum = Array.Num();
+	CityHash64WithSeed((const char*)&ArrayNum, sizeof(ArrayNum), Hash);
+	CityHash64WithSeed((const char*)Array.GetData(), Array.Num() * sizeof(TType), Hash);
+}
+
 void FShader::BuildParameterMapInfo(const TMap<FString, FParameterAllocation>& ParameterMap)
 {
 	uint32 UniformCount = 0;
@@ -629,12 +639,10 @@ void FShader::BuildParameterMapInfo(const TMap<FString, FParameterAllocation>& P
 	ParameterMapInfo.TextureSamplers.Empty(SamplerCount);
 	ParameterMapInfo.SRVs.Empty(SRVCount);
 
-	auto GetParameterMap = [this](EShaderParameterType ParameterType) -> TMemoryImageArray<FShaderParameterInfo>*
+	auto GetResourceParameterMap = [this](EShaderParameterType ParameterType) -> TMemoryImageArray<FShaderResourceParameterInfo>*
 	{
 		switch (ParameterType)
 		{
-		case EShaderParameterType::UniformBuffer:
-			return &ParameterMapInfo.UniformBuffers;
 		case EShaderParameterType::Sampler:
 			return &ParameterMapInfo.TextureSamplers;
 		case EShaderParameterType::SRV:
@@ -662,8 +670,7 @@ void FShader::BuildParameterMapInfo(const TMap<FString, FParameterAllocation>& P
 
 				if (LooseParameterBufferInfo.BaseIndex == ParamValue.BufferIndex)
 				{
-					FShaderParameterInfo ParameterInfo(ParamValue.BaseIndex, ParamValue.Size);
-					LooseParameterBufferInfo.Parameters.Add(ParameterInfo);
+					LooseParameterBufferInfo.Parameters.Emplace(ParamValue.BaseIndex, ParamValue.Size);
 					LooseParameterBufferInfo.Size += ParamValue.Size;
 					bAddedToExistingBuffer = true;
 				}
@@ -673,17 +680,18 @@ void FShader::BuildParameterMapInfo(const TMap<FString, FParameterAllocation>& P
 			{
 				FShaderLooseParameterBufferInfo NewParameterBufferInfo(ParamValue.BufferIndex, ParamValue.Size);
 
-				FShaderParameterInfo ParameterInfo(ParamValue.BaseIndex, ParamValue.Size);
-				NewParameterBufferInfo.Parameters.Add(ParameterInfo);
+				NewParameterBufferInfo.Parameters.Emplace(ParamValue.BaseIndex, ParamValue.Size);
 
 				ParameterMapInfo.LooseParameterBuffers.Add(NewParameterBufferInfo);
 			}
 		}
-		else if (TMemoryImageArray<FShaderParameterInfo>* ParameterInfoArray = GetParameterMap(ParamValue.Type))
+		else if (ParamValue.Type == EShaderParameterType::UniformBuffer)
 		{
-			const uint16 BaseIndex = ParamValue.Type == EShaderParameterType::UniformBuffer ? ParamValue.BufferIndex : ParamValue.BaseIndex;
-			FShaderParameterInfo ParameterInfo(BaseIndex, ParamValue.Size);
-			ParameterInfoArray->Add(ParameterInfo);
+			ParameterMapInfo.UniformBuffers.Emplace(ParamValue.BufferIndex);
+		}
+		else if (TMemoryImageArray<FShaderResourceParameterInfo>* ParameterInfoArray = GetResourceParameterMap(ParamValue.Type))
+		{
+			ParameterInfoArray->Emplace(ParamValue.BaseIndex, ParamValue.BufferIndex, ParamValue.Type);
 		}
 	}
 
@@ -704,21 +712,15 @@ void FShader::BuildParameterMapInfo(const TMap<FString, FParameterAllocation>& P
 			CityHash64WithSeed((const char*)&Value, sizeof(Value), Hash);
 		};
 
-		const auto CityHashArray = [&](const TMemoryImageArray<FShaderParameterInfo>& Array)
-		{
-			CityHashValue(Array.Num());
-			CityHash64WithSeed((const char*)Array.GetData(), Array.Num() * sizeof(FShaderParameterInfo), Hash);
-		};
-
 		for (FShaderLooseParameterBufferInfo& Info : ParameterMapInfo.LooseParameterBuffers)
 		{
 			CityHashValue(Info.BaseIndex);
 			CityHashValue(Info.Size);
-			CityHashArray(Info.Parameters);
+			CityHashArray(Hash, Info.Parameters);
 		}
-		CityHashArray(ParameterMapInfo.UniformBuffers);
-		CityHashArray(ParameterMapInfo.TextureSamplers);
-		CityHashArray(ParameterMapInfo.SRVs);
+		CityHashArray(Hash, ParameterMapInfo.UniformBuffers);
+		CityHashArray(Hash, ParameterMapInfo.TextureSamplers);
+		CityHashArray(Hash, ParameterMapInfo.SRVs);
 	}
 
 	ParameterMapInfo.Hash = Hash;
diff --git a/Engine/Source/Runtime/RenderCore/Public/Shader.h b/Engine/Source/Runtime/RenderCore/Public/Shader.h
--- a/Engine/Source/Runtime/RenderCore/Public/Shader.h
+++ b/Engine/Source/Runtime/RenderCore/Public/Shader.h
@@ -137,36 +137,106 @@ public:
 	}
 };
 
-class FShaderParameterInfo
+class FShaderUniformBufferParameterInfo
 {
-	DECLARE_TYPE_LAYOUT(FShaderParameterInfo, NonVirtual);
+	DECLARE_TYPE_LAYOUT(FShaderUniformBufferParameterInfo, NonVirtual);
+public:
+	LAYOUT_FIELD(uint16, BaseIndex);
+
+	FShaderUniformBufferParameterInfo() = default;
+
+	FShaderUniformBufferParameterInfo(uint16 InBaseIndex)
+	{
+		BaseIndex = InBaseIndex;
+		checkf(BaseIndex == InBaseIndex, TEXT("Tweak FShaderUniformBufferParameterInfo type sizes"));
+	}
+
+	friend FArchive& operator<<(FArchive& Ar, FShaderUniformBufferParameterInfo& Info)
+	{
+		Ar << Info.BaseIndex;
+		return Ar;
+	}
+
+	inline bool operator==(const FShaderUniformBufferParameterInfo& Rhs) const
+	{
+		return BaseIndex == Rhs.BaseIndex;
+	}
+
+	inline bool operator<(const FShaderUniformBufferParameterInfo& Rhs) const
+	{
+		return BaseIndex < Rhs.BaseIndex;
+	}
+};
+
+class FShaderResourceParameterInfo
+{
+	DECLARE_TYPE_LAYOUT(FShaderResourceParameterInfo, NonVirtual);
+public:
+	LAYOUT_FIELD(uint16, BaseIndex);
+	LAYOUT_FIELD(uint8, BufferIndex);
+	LAYOUT_FIELD(EShaderParameterType, Type);
+
+	FShaderResourceParameterInfo() = default;
+
+	FShaderResourceParameterInfo(uint16 InBaseIndex, uint8 InBufferIndex, EShaderParameterType InType)
+	{
+		BaseIndex = InBaseIndex;
+		BufferIndex = InBufferIndex;
+		Type = InType;
+		checkf(BaseIndex == InBaseIndex && BufferIndex == InBufferIndex && Type == InType, TEXT("Tweak FShaderResourceParameterInfo type sizes"));
+	}
+
+	friend FArchive& operator<<(FArchive& Ar, FShaderResourceParameterInfo& Info)
+	{
+		Ar << Info.BaseIndex;
+		Ar << Info.BufferIndex;
+		Ar << Info.Type;
+		return Ar;
+	}
+
+	inline bool operator==(const FShaderResourceParameterInfo& Rhs) const
+	{
+		return BaseIndex == Rhs.BaseIndex
+			&& BufferIndex == Rhs.BufferIndex
+			&& Type == Rhs.Type;
+	}
+
+	inline bool operator<(const FShaderResourceParameterInfo& Rhs) const
+	{
+		return BaseIndex < Rhs.BaseIndex;
+	}
+};
+
+class FShaderLooseParameterInfo
+{
+	DECLARE_TYPE_LAYOUT(FShaderLooseParameterInfo, NonVirtual);
 public:
 	LAYOUT_FIELD(uint16, BaseIndex);
 	LAYOUT_FIELD(uint16, Size);
 
-	FShaderParameterInfo() {}
+	FShaderLooseParameterInfo() = default;
 
-	FShaderParameterInfo(uint16 InBaseIndex, uint16 InSize)
+	FShaderLooseParameterInfo(uint16 InBaseIndex, uint16 InSize)
 	{
 		BaseIndex = InBaseIndex;
 		Size = InSize;
-		checkf(BaseIndex == InBaseIndex && Size == InSize, TEXT("Tweak FShaderParameterInfo type sizes"));
+		checkf(BaseIndex == InBaseIndex && Size == InSize, TEXT("Tweak FShaderLooseParameterInfo type sizes"));
 	}
 
-	friend FArchive& operator<<(FArchive& Ar, FShaderParameterInfo& Info)
+	friend FArchive& operator<<(FArchive& Ar, FShaderLooseParameterInfo& Info)
 	{
 		Ar << Info.BaseIndex;
 		Ar << Info.Size;
 		return Ar;
 	}
 
-	inline bool operator==(const FShaderParameterInfo& Rhs) const
+	inline bool operator==(const FShaderLooseParameterInfo& Rhs) const
 	{
 		return BaseIndex == Rhs.BaseIndex
 			&& Size == Rhs.Size;
 	}
 
-	inline bool operator<(const FShaderParameterInfo& Rhs) const
+	inline bool operator<(const FShaderLooseParameterInfo& Rhs) const
 	{
 		return BaseIndex < Rhs.BaseIndex;
 	}
@@ -179,7 +249,7 @@ public:
 	LAYOUT_FIELD(uint16, BaseIndex);
 	LAYOUT_FIELD(uint16, Size);
 
-	LAYOUT_FIELD(TMemoryImageArray<FShaderParameterInfo>, Parameters);
+	LAYOUT_FIELD(TMemoryImageArray<FShaderLooseParameterInfo>, Parameters);
 
 	FShaderLooseParameterBufferInfo() {}
 
@@ -215,9 +285,9 @@ class FShaderParameterMapInfo
 {
 	DECLARE_TYPE_LAYOUT(FShaderParameterMapInfo, NonVirtual);
 public:
-	LAYOUT_FIELD(TMemoryImageArray<FShaderParameterInfo>, UniformBuffers);
-	LAYOUT_FIELD(TMemoryImageArray<FShaderParameterInfo>, TextureSamplers);
-	LAYOUT_FIELD(TMemoryImageArray<FShaderParameterInfo>, SRVs);
+	LAYOUT_FIELD(TMemoryImageArray<FShaderUniformBufferParameterInfo>, UniformBuffers);
+	LAYOUT_FIELD(TMemoryImageArray<FShaderResourceParameterInfo>, TextureSamplers);
+	LAYOUT_FIELD(TMemoryImageArray<FShaderResourceParameterInfo>, SRVs);
 	LAYOUT_FIELD(TMemoryImageArray<FShaderLooseParameterBufferInfo>, LooseParameterBuffers);
 	LAYOUT_FIELD(uint64, Hash);
 
diff --git a/Engine/Source/Runtime/Renderer/Private/MeshPassProcessor.cpp b/Engine/Source/Runtime/Renderer/Private/MeshPassProcessor.cpp
--- a/Engine/Source/Runtime/Renderer/Private/MeshPassProcessor.cpp
+++ b/Engine/Source/Runtime/Renderer/Private/MeshPassProcessor.cpp
@@ -103,6 +103,133 @@ private:
 	const uint8* Data;
 };
 
+template<typename TRHICmdList, typename TRHIShader>
+FORCEINLINE void SetBindlessParameter(TRHICmdList& RHICmdList, TRHIShader* Shader, FShaderResourceParameterInfo Parameter, FRHIDescriptorHandle Handle)
+{
+	if (Handle.IsValid())
+	{
+		const uint32 BindlessIndex = Handle.GetIndex();
+		RHICmdList.SetShaderParameter(Shader, Parameter.BufferIndex, Parameter.BaseIndex, sizeof(BindlessIndex), &BindlessIndex);
+	}
+}
+
+template<class TRHICmdList, class TRHIShader>
+inline void SetTextureParameter(TRHICmdList& RHICmdList, TRHIShader* Shader, FShaderBindingState& RESTRICT ShaderBindingState, FShaderResourceParameterInfo Parameter, FRHITexture* TextureRHI)
+{
+#if PLATFORM_SUPPORTS_BINDLESS_RENDERING
+	if (Parameter.Type == EShaderParameterType::BindlessResourceIndex)
+	{
+		const FRHIDescriptorHandle Handle = TextureRHI ? TextureRHI->GetDefaultBindlessHandle() : FRHIDescriptorHandle();
+		SetBindlessParameter(RHICmdList, Shader, Parameter, Handle);
+	}
+	else
+#endif
+	{
+		checkSlow(Parameter.BaseIndex < UE_ARRAY_COUNT(ShaderBindingState.Textures));
+		if (TextureRHI != ShaderBindingState.Textures[Parameter.BaseIndex])
+		{
+			RHICmdList.SetShaderTexture(Shader, Parameter.BaseIndex, TextureRHI);
+
+			ShaderBindingState.Textures[Parameter.BaseIndex] = TextureRHI;
+			ShaderBindingState.MaxTextureUsed = FMath::Max((int32)Parameter.BaseIndex, ShaderBindingState.MaxTextureUsed);
+		}
+	}
+}
+
+template<class TRHICmdList, class TRHIShader>
+inline void SetTextureParameter(TRHICmdList& RHICmdList, TRHIShader* Shader, FShaderResourceParameterInfo Parameter, FRHITexture* TextureRHI)
+{
+#if PLATFORM_SUPPORTS_BINDLESS_RENDERING
+	if (Parameter.Type == EShaderParameterType::BindlessResourceIndex)
+	{
+		const FRHIDescriptorHandle Handle = TextureRHI ? TextureRHI->GetDefaultBindlessHandle() : FRHIDescriptorHandle();
+		SetBindlessParameter(RHICmdList, Shader, Parameter, Handle);
+	}
+	else
+#endif
+	{
+		RHICmdList.SetShaderTexture(Shader, Parameter.BaseIndex, TextureRHI);
+	}
+}
+
+template<class TRHICmdList, class TRHIShader>
+inline void SetSrvParameter(TRHICmdList& RHICmdList, TRHIShader* Shader, FShaderBindingState& RESTRICT ShaderBindingState, FShaderResourceParameterInfo Parameter, FRHIShaderResourceView* SrvRHI)
+{
+#if PLATFORM_SUPPORTS_BINDLESS_RENDERING
+	if (Parameter.Type == EShaderParameterType::BindlessResourceIndex)
+	{
+		const FRHIDescriptorHandle Handle = SrvRHI ? SrvRHI->GetBindlessHandle() : FRHIDescriptorHandle();
+		SetBindlessParameter(RHICmdList, Shader, Parameter, Handle);
+	}
+	else
+#endif
+	{
+		checkSlow(Parameter.BaseIndex < UE_ARRAY_COUNT(ShaderBindingState.SRVs));
+		if (SrvRHI != ShaderBindingState.SRVs[Parameter.BaseIndex])
+		{
+			RHICmdList.SetShaderResourceViewParameter(Shader, Parameter.BaseIndex, SrvRHI);
+
+			ShaderBindingState.SRVs[Parameter.BaseIndex] = SrvRHI;
+			ShaderBindingState.MaxSRVUsed = FMath::Max((int32)Parameter.BaseIndex, ShaderBindingState.MaxSRVUsed);
+		}
+	}
+}
+
+template<class TRHICmdList, class TRHIShader>
+inline void SetSrvParameter(TRHICmdList& RHICmdList, TRHIShader* Shader, FShaderResourceParameterInfo Parameter, FRHIShaderResourceView* SrvRHI)
+{
+#if PLATFORM_SUPPORTS_BINDLESS_RENDERING
+	if (Parameter.Type == EShaderParameterType::BindlessResourceIndex)
+	{
+		const FRHIDescriptorHandle Handle = SrvRHI ? SrvRHI->GetBindlessHandle() : FRHIDescriptorHandle();
+		SetBindlessParameter(RHICmdList, Shader, Parameter, Handle);
+	}
+	else
+#endif
+	{
+		RHICmdList.SetShaderResourceViewParameter(Shader, Parameter.BaseIndex, SrvRHI);
+	}
+}
+
+template<class TRHICmdList, class TRHIShader>
+inline void SetSamplerParameter(TRHICmdList& RHICmdList, TRHIShader* Shader, FShaderBindingState& RESTRICT ShaderBindingState, FShaderResourceParameterInfo Parameter, FRHISamplerState* SamplerStateRHI)
+{
+#if PLATFORM_SUPPORTS_BINDLESS_RENDERING
+	if (Parameter.Type == EShaderParameterType::BindlessSamplerIndex)
+	{
+		const FRHIDescriptorHandle Handle = SamplerStateRHI ? SamplerStateRHI->GetBindlessHandle() : FRHIDescriptorHandle();
+		SetBindlessParameter(RHICmdList, Shader, Parameter, Handle);
+	}
+	else
+#endif
+	{
+		checkSlow(Parameter.BaseIndex < UE_ARRAY_COUNT(ShaderBindingState.Samplers));
+		if (SamplerStateRHI != ShaderBindingState.Samplers[Parameter.BaseIndex])
+		{
+			RHICmdList.SetShaderSampler(Shader, Parameter.BaseIndex, SamplerStateRHI);
+
+			ShaderBindingState.Samplers[Parameter.BaseIndex] = SamplerStateRHI;
+			ShaderBindingState.MaxSamplerUsed = FMath::Max((int32)Parameter.BaseIndex, ShaderBindingState.MaxSamplerUsed);
+		}
+	}
+}
+
+template<class TRHICmdList, class TRHIShader>
+inline void SetSamplerParameter(TRHICmdList& RHICmdList, TRHIShader* Shader, FShaderResourceParameterInfo Parameter, FRHISamplerState* SamplerStateRHI)
+{
+#if PLATFORM_SUPPORTS_BINDLESS_RENDERING
+	if (Parameter.Type == EShaderParameterType::BindlessSamplerIndex)
+	{
+		const FRHIDescriptorHandle Handle = SamplerStateRHI ? SamplerStateRHI->GetBindlessHandle() : FRHIDescriptorHandle();
+		SetBindlessParameter(RHICmdList, Shader, Parameter, Handle);
+	}
+	else
+#endif
+	{
+		RHICmdList.SetShaderSampler(Shader, Parameter.BaseIndex, SamplerStateRHI);
+	}
+}
+
 template<class RHICmdListType, class RHIShaderType>
 void FMeshDrawShaderBindings::SetShaderBindings(
 	RHICmdListType& RHICmdList,
@@ -111,12 +238,12 @@ void FMeshDrawShaderBindings::SetShaderBindings(
 	FShaderBindingState& RESTRICT ShaderBindingState)
 {
 	FRHIUniformBuffer* const* RESTRICT UniformBufferBindings = SingleShaderBindings.GetUniformBufferStart();
-	const FShaderParameterInfo* RESTRICT UniformBufferParameters = SingleShaderBindings.ParameterMapInfo.UniformBuffers.GetData();
+	const FShaderUniformBufferParameterInfo* RESTRICT UniformBufferParameters = SingleShaderBindings.ParameterMapInfo.UniformBuffers.GetData();
 	const int32 NumUniformBuffers = SingleShaderBindings.ParameterMapInfo.UniformBuffers.Num();
 
 	for (int32 UniformBufferIndex = 0; UniformBufferIndex < NumUniformBuffers; UniformBufferIndex++)
 	{
-		FShaderParameterInfo Parameter = UniformBufferParameters[UniformBufferIndex];
+		FShaderUniformBufferParameterInfo Parameter = UniformBufferParameters[UniformBufferIndex];
 		checkSlow(Parameter.BaseIndex < UE_ARRAY_COUNT(ShaderBindingState.UniformBuffers));
 		FRHIUniformBuffer* UniformBuffer = UniformBufferBindings[UniformBufferIndex];
 
@@ -129,32 +256,25 @@ void FMeshDrawShaderBindings::SetShaderBindings(
 	}
 
 	FRHISamplerState* const* RESTRICT SamplerBindings = SingleShaderBindings.GetSamplerStart();
-	const FShaderParameterInfo* RESTRICT TextureSamplerParameters = SingleShaderBindings.ParameterMapInfo.TextureSamplers.GetData();
+	const FShaderResourceParameterInfo* RESTRICT TextureSamplerParameters = SingleShaderBindings.ParameterMapInfo.TextureSamplers.GetData();
 	const int32 NumTextureSamplers = SingleShaderBindings.ParameterMapInfo.TextureSamplers.Num();
 
 	for (int32 SamplerIndex = 0; SamplerIndex < NumTextureSamplers; SamplerIndex++)
 	{
-		FShaderParameterInfo Parameter = TextureSamplerParameters[SamplerIndex];
-		checkSlow(Parameter.BaseIndex < UE_ARRAY_COUNT(ShaderBindingState.Samplers));
+		FShaderResourceParameterInfo Parameter = TextureSamplerParameters[SamplerIndex];
 		FRHISamplerState* Sampler = (FRHISamplerState*)SamplerBindings[SamplerIndex];
 
-		if (Sampler != ShaderBindingState.Samplers[Parameter.BaseIndex])
-		{
-			RHICmdList.SetShaderSampler(Shader, Parameter.BaseIndex, Sampler);
-			ShaderBindingState.Samplers[Parameter.BaseIndex] = Sampler;
-			ShaderBindingState.MaxSamplerUsed = FMath::Max((int32)Parameter.BaseIndex, ShaderBindingState.MaxSamplerUsed);
-		}
+		SetSamplerParameter(RHICmdList, Shader, ShaderBindingState, Parameter, Sampler);
 	}
 
 	const uint8* RESTRICT SRVType = SingleShaderBindings.GetSRVTypeStart();
 	FRHIResource* const* RESTRICT SRVBindings = SingleShaderBindings.GetSRVStart();
-	const FShaderParameterInfo* RESTRICT SRVParameters = SingleShaderBindings.ParameterMapInfo.SRVs.GetData();
+	const FShaderResourceParameterInfo* RESTRICT SRVParameters = SingleShaderBindings.ParameterMapInfo.SRVs.GetData();
 	const uint32 NumSRVs = SingleShaderBindings.ParameterMapInfo.SRVs.Num();
 
 	for (uint32 SRVIndex = 0; SRVIndex < NumSRVs; SRVIndex++)
 	{
-		FShaderParameterInfo Parameter = SRVParameters[SRVIndex];
-		checkSlow(Parameter.BaseIndex < UE_ARRAY_COUNT(ShaderBindingState.SRVs));
+		FShaderResourceParameterInfo Parameter = SRVParameters[SRVIndex];
 
 		uint32 TypeByteIndex = SRVIndex / 8;
 		uint32 TypeBitIndex = SRVIndex % 8;
@@ -162,24 +282,12 @@ void FMeshDrawShaderBindings::SetShaderBindings(
 		if (SRVType[TypeByteIndex] & (1 << TypeBitIndex))
 		{
 			FRHIShaderResourceView* SRV = (FRHIShaderResourceView*)SRVBindings[SRVIndex];
-
-			if (SRV != ShaderBindingState.SRVs[Parameter.BaseIndex])
-			{
-				RHICmdList.SetShaderResourceViewParameter(Shader, Parameter.BaseIndex, SRV);
-				ShaderBindingState.SRVs[Parameter.BaseIndex] = SRV;
-				ShaderBindingState.MaxSRVUsed = FMath::Max((int32)Parameter.BaseIndex, ShaderBindingState.MaxSRVUsed);
-			}
+			SetSrvParameter(RHICmdList, Shader, ShaderBindingState, Parameter, SRV);
 		}
 		else
 		{
 			FRHITexture* Texture = (FRHITexture*)SRVBindings[SRVIndex];
-
-			if (Texture != ShaderBindingState.Textures[Parameter.BaseIndex])
-			{
-				RHICmdList.SetShaderTexture(Shader, Parameter.BaseIndex, Texture);
-				ShaderBindingState.Textures[Parameter.BaseIndex] = Texture;
-				ShaderBindingState.MaxTextureUsed = FMath::Max((int32)Parameter.BaseIndex, ShaderBindingState.MaxTextureUsed);
-			}
+			SetTextureParameter(RHICmdList, Shader, ShaderBindingState, Parameter, Texture);
 		}
 	}
 
@@ -187,7 +295,7 @@ void FMeshDrawShaderBindings::SetShaderBindings(
 
 	for (const FShaderLooseParameterBufferInfo& LooseParameterBuffer : SingleShaderBindings.ParameterMapInfo.LooseParameterBuffers)
 	{
-		for (FShaderParameterInfo Parameter : LooseParameterBuffer.Parameters)
+		for (FShaderLooseParameterInfo Parameter : LooseParameterBuffer.Parameters)
 		{
 			RHICmdList.SetShaderParameter(
 				Shader,
@@ -209,12 +317,12 @@ void FMeshDrawShaderBindings::SetShaderBindings(
 	const FReadOnlyMeshDrawSingleShaderBindings& RESTRICT SingleShaderBindings)
 {
 	FRHIUniformBuffer* const* RESTRICT UniformBufferBindings = SingleShaderBindings.GetUniformBufferStart();
-	const FShaderParameterInfo* RESTRICT UniformBufferParameters = SingleShaderBindings.ParameterMapInfo.UniformBuffers.GetData();
+	const FShaderUniformBufferParameterInfo* RESTRICT UniformBufferParameters = SingleShaderBindings.ParameterMapInfo.UniformBuffers.GetData();
 	const int32 NumUniformBuffers = SingleShaderBindings.ParameterMapInfo.UniformBuffers.Num();
 
 	for (int32 UniformBufferIndex = 0; UniformBufferIndex < NumUniformBuffers; UniformBufferIndex++)
 	{
-		FShaderParameterInfo Parameter = UniformBufferParameters[UniformBufferIndex];
+		FShaderUniformBufferParameterInfo Parameter = UniformBufferParameters[UniformBufferIndex];
 		FRHIUniformBuffer* UniformBuffer = UniformBufferBindings[UniformBufferIndex];
 
 		if (UniformBuffer)
@@ -224,25 +332,25 @@ void FMeshDrawShaderBindings::SetShaderBindings(
 	}
 
 	FRHISamplerState* const* RESTRICT SamplerBindings = SingleShaderBindings.GetSamplerStart();
-	const FShaderParameterInfo* RESTRICT TextureSamplerParameters = SingleShaderBindings.ParameterMapInfo.TextureSamplers.GetData();
+	const FShaderResourceParameterInfo* RESTRICT TextureSamplerParameters = SingleShaderBindings.ParameterMapInfo.TextureSamplers.GetData();
 	const int32 NumTextureSamplers = SingleShaderBindings.ParameterMapInfo.TextureSamplers.Num();
 
 	for (int32 SamplerIndex = 0; SamplerIndex < NumTextureSamplers; SamplerIndex++)
 	{
-		FShaderParameterInfo Parameter = TextureSamplerParameters[SamplerIndex];
+		FShaderResourceParameterInfo Parameter = TextureSamplerParameters[SamplerIndex];
 		FRHISamplerState* Sampler = (FRHISamplerState*)SamplerBindings[SamplerIndex];
 
-		RHICmdList.SetShaderSampler(Shader, Parameter.BaseIndex, Sampler);
+		SetSamplerParameter(RHICmdList, Shader, Parameter, Sampler);
 	}
 
 	const uint8* RESTRICT SRVType = SingleShaderBindings.GetSRVTypeStart();
 	FRHIResource* const* RESTRICT SRVBindings = SingleShaderBindings.GetSRVStart();
-	const FShaderParameterInfo* RESTRICT SRVParameters = SingleShaderBindings.ParameterMapInfo.SRVs.GetData();
+	const FShaderResourceParameterInfo* RESTRICT SRVParameters = SingleShaderBindings.ParameterMapInfo.SRVs.GetData();
 	const uint32 NumSRVs = SingleShaderBindings.ParameterMapInfo.SRVs.Num();
 
 	for (uint32 SRVIndex = 0; SRVIndex < NumSRVs; SRVIndex++)
 	{
-		FShaderParameterInfo Parameter = SRVParameters[SRVIndex];
+		FShaderResourceParameterInfo Parameter = SRVParameters[SRVIndex];
 
 		uint32 TypeByteIndex = SRVIndex / 8;
 		uint32 TypeBitIndex = SRVIndex % 8;
@@ -250,12 +358,12 @@ void FMeshDrawShaderBindings::SetShaderBindings(
 		if (SRVType[TypeByteIndex] & (1 << TypeBitIndex))
 		{
 			FRHIShaderResourceView* SRV = (FRHIShaderResourceView*)SRVBindings[SRVIndex];
-			RHICmdList.SetShaderResourceViewParameter(Shader, Parameter.BaseIndex, SRV);
+			SetSrvParameter(RHICmdList, Shader, Parameter, SRV);
 		}
 		else
 		{
 			FRHITexture* Texture = (FRHITexture*)SRVBindings[SRVIndex];
-			RHICmdList.SetShaderTexture(Shader, Parameter.BaseIndex, Texture);
+			SetTextureParameter(RHICmdList, Shader, Parameter, Texture);
 		}
 	}
 	
@@ -263,7 +371,7 @@ void FMeshDrawShaderBindings::SetShaderBindings(
 
 	for (const FShaderLooseParameterBufferInfo& LooseParameterBuffer : SingleShaderBindings.ParameterMapInfo.LooseParameterBuffers)
 	{
-		for (FShaderParameterInfo Parameter : LooseParameterBuffer.Parameters)
+		for (FShaderLooseParameterInfo Parameter : LooseParameterBuffer.Parameters)
 		{
 			RHICmdList.SetShaderParameter(
 				Shader,
@@ -292,7 +400,7 @@ FRayTracingLocalShaderBindings* FMeshDrawShaderBindings::SetRayTracingShaderBind
 	FReadOnlyMeshDrawSingleShaderBindings SingleShaderBindings(ShaderLayouts[0], GetData());
 
 	FRHIUniformBuffer* const* RESTRICT UniformBufferBindings = SingleShaderBindings.GetUniformBufferStart();
-	const FShaderParameterInfo* RESTRICT UniformBufferParameters = SingleShaderBindings.ParameterMapInfo.UniformBuffers.GetData();
+	const FShaderUniformBufferParameterInfo* RESTRICT UniformBufferParameters = SingleShaderBindings.ParameterMapInfo.UniformBuffers.GetData();
 	const int32 NumUniformBufferParameters = SingleShaderBindings.ParameterMapInfo.UniformBuffers.Num();
 
 	checkf(SingleShaderBindings.ParameterMapInfo.TextureSamplers.Num() == 0, TEXT("Texture sampler parameters are not supported for ray tracing. UniformBuffers must be used for all resource binding."));
@@ -303,7 +411,7 @@ FRayTracingLocalShaderBindings* FMeshDrawShaderBindings::SetRayTracingShaderBind
 	int32 MaxUniformBufferUsed = -1;
 	for (int32 UniformBufferIndex = 0; UniformBufferIndex < NumUniformBufferParameters; UniformBufferIndex++)
 	{
-		FShaderParameterInfo Parameter = UniformBufferParameters[UniformBufferIndex];
+		FShaderUniformBufferParameterInfo Parameter = UniformBufferParameters[UniformBufferIndex];
 		const FRHIUniformBuffer* UniformBuffer = UniformBufferBindings[UniformBufferIndex];
 		MaxUniformBufferUsed = FMath::Max((int32)Parameter.BaseIndex, MaxUniformBufferUsed);
 	}
@@ -324,7 +432,7 @@ FRayTracingLocalShaderBindings* FMeshDrawShaderBindings::SetRayTracingShaderBind
 
 		for (int32 LooseParameterIndex = 0; LooseParameterIndex < LooseParameterBuffer.Parameters.Num(); LooseParameterIndex++)
 		{
-			FShaderParameterInfo LooseParameter = LooseParameterBuffer.Parameters[LooseParameterIndex];
+			FShaderLooseParameterInfo LooseParameter = LooseParameterBuffer.Parameters[LooseParameterIndex];
 			LooseParameterDataSize = FMath::Max<uint32>(LooseParameterDataSize, LooseParameter.BaseIndex + LooseParameter.Size);
 		}
 	}
@@ -343,7 +451,7 @@ FRayTracingLocalShaderBindings* FMeshDrawShaderBindings::SetRayTracingShaderBind
 
 	for (int32 UniformBufferIndex = 0; UniformBufferIndex < NumUniformBufferParameters; UniformBufferIndex++)
 	{
-		FShaderParameterInfo Parameter = UniformBufferParameters[UniformBufferIndex];
+		FShaderUniformBufferParameterInfo Parameter = UniformBufferParameters[UniformBufferIndex];
 		const FRHIUniformBuffer* UniformBuffer = UniformBufferBindings[UniformBufferIndex];
 		Bindings.UniformBuffers[Parameter.BaseIndex] = const_cast<FRHIUniformBuffer*>(UniformBuffer);
 	}
@@ -354,7 +462,7 @@ FRayTracingLocalShaderBindings* FMeshDrawShaderBindings::SetRayTracingShaderBind
 		const uint8* LooseDataOffset = SingleShaderBindings.GetLooseDataStart();
 		for (int32 LooseParameterIndex = 0; LooseParameterIndex < LooseParameterBuffer.Parameters.Num(); LooseParameterIndex++)
 		{
-			FShaderParameterInfo LooseParameter = LooseParameterBuffer.Parameters[LooseParameterIndex];
+			FShaderLooseParameterInfo LooseParameter = LooseParameterBuffer.Parameters[LooseParameterIndex];
 			FMemory::Memcpy(Bindings.LooseParameterData + LooseParameter.BaseIndex, LooseDataOffset, LooseParameter.Size);
 			LooseDataOffset += LooseParameter.Size;
 		}
@@ -385,7 +493,7 @@ void FMeshDrawShaderBindings::SetRayTracingShaderBindingsForMissShader(
 	const int32 MaxUniformBuffers = UE_ARRAY_COUNT(BindingState.UniformBuffers);
 
 	FRHIUniformBuffer* const* RESTRICT UniformBufferBindings = SingleShaderBindings.GetUniformBufferStart();
-	const FShaderParameterInfo* RESTRICT UniformBufferParameters = SingleShaderBindings.ParameterMapInfo.UniformBuffers.GetData();
+	const FShaderUniformBufferParameterInfo* RESTRICT UniformBufferParameters = SingleShaderBindings.ParameterMapInfo.UniformBuffers.GetData();
 	const int32 NumUniformBufferParameters = SingleShaderBindings.ParameterMapInfo.UniformBuffers.Num();
 
 	checkf(SingleShaderBindings.ParameterMapInfo.TextureSamplers.Num() == 0, TEXT("Texture sampler parameters are not supported for ray tracing. UniformBuffers must be used for all resource binding."));
@@ -396,7 +504,7 @@ void FMeshDrawShaderBindings::SetRayTracingShaderBindingsForMissShader(
 	int32 MaxUniformBufferUsed = -1;
 	for (int32 UniformBufferIndex = 0; UniformBufferIndex < NumUniformBufferParameters; UniformBufferIndex++)
 	{
-		FShaderParameterInfo Parameter = UniformBufferParameters[UniformBufferIndex];
+		FShaderUniformBufferParameterInfo Parameter = UniformBufferParameters[UniformBufferIndex];
 		checkSlow(Parameter.BaseIndex < MaxUniformBuffers);
 		FRHIUniformBuffer* UniformBuffer = UniformBufferBindings[UniformBufferIndex];
 		if (Parameter.BaseIndex < MaxUniformBuffers)
@@ -642,7 +750,7 @@ void FMeshDrawShaderBindings::Finalize(const FMeshProcessorShaders* ShadersForDe
 
 		for (int32 BindingIndex = 0; BindingIndex < ShaderLayout.ParameterMapInfo.UniformBuffers.Num(); BindingIndex++)
 		{
-			FShaderParameterInfo ParameterInfo = ShaderLayout.ParameterMapInfo.UniformBuffers[BindingIndex];
+			FShaderUniformBufferParameterInfo ParameterInfo = ShaderLayout.ParameterMapInfo.UniformBuffers[BindingIndex];
 
 			FRHIUniformBuffer* UniformBufferValue = UniformBufferBindings[BindingIndex];
 
@@ -675,7 +783,7 @@ void FMeshDrawShaderBindings::Finalize(const FMeshProcessorShaders* ShadersForDe
 
 		for (int32 BindingIndex = 0; BindingIndex < ShaderLayout.ParameterMapInfo.TextureSamplers.Num(); BindingIndex++)
 		{
-			FShaderParameterInfo ParameterInfo = ShaderLayout.ParameterMapInfo.TextureSamplers[BindingIndex];
+			FShaderResourceParameterInfo ParameterInfo = ShaderLayout.ParameterMapInfo.TextureSamplers[BindingIndex];
 			const FRHISamplerState* SamplerValue = SamplerBindings[BindingIndex];
 			ensureMsgf(SamplerValue, TEXT("Shader %s with vertex factory %s never set sampler at BaseIndex %u.  This can cause GPU hangs, depending on how the shader uses it."), 
 				Shader.GetType()->GetName(), 
@@ -685,12 +793,12 @@ void FMeshDrawShaderBindings::Finalize(const FMeshProcessorShaders* ShadersForDe
 
 		const uint8* RESTRICT SRVType = SingleShaderBindings.GetSRVTypeStart();
 		FRHIResource* const* RESTRICT SRVBindings = SingleShaderBindings.GetSRVStart();
-		const FShaderParameterInfo* RESTRICT SRVParameters = SingleShaderBindings.ParameterMapInfo.SRVs.GetData();
+		const FShaderResourceParameterInfo* RESTRICT SRVParameters = SingleShaderBindings.ParameterMapInfo.SRVs.GetData();
 		const uint32 NumSRVs = SingleShaderBindings.ParameterMapInfo.SRVs.Num();
 
 		for (uint32 SRVIndex = 0; SRVIndex < NumSRVs; SRVIndex++)
 		{
-			FShaderParameterInfo Parameter = SRVParameters[SRVIndex];
+			FShaderResourceParameterInfo Parameter = SRVParameters[SRVIndex];
 
 			uint32 TypeByteIndex = SRVIndex / 8;
 			uint32 TypeBitIndex = SRVIndex % 8;
diff --git a/Engine/Source/Runtime/Renderer/Private/RayTracing/RayTracingLighting.cpp b/Engine/Source/Runtime/Renderer/Private/RayTracing/RayTracingLighting.cpp
--- a/Engine/Source/Runtime/Renderer/Private/RayTracing/RayTracingLighting.cpp
+++ b/Engine/Source/Runtime/Renderer/Private/RayTracing/RayTracingLighting.cpp
@@ -612,7 +612,7 @@ static int32 BindParameters(const TShaderRef<ShaderClass>& Shader, typename Shad
 	int32 MaxUniformBufferUsed = -1;
 	for (int32 UniformBufferIndex = 0; UniformBufferIndex < NumUniformBuffers; UniformBufferIndex++)
 	{
-		const FShaderParameterInfo &Parameter = ParameterMap.UniformBuffers[UniformBufferIndex];
+		const FShaderUniformBufferParameterInfo Parameter = ParameterMap.UniformBuffers[UniformBufferIndex];
 		checkSlow(Parameter.BaseIndex < MaxParams);
 		const FRHIUniformBuffer* UniformBuffer = ResourceBinder.UniformBuffers[UniformBufferIndex];
 		if (Parameter.BaseIndex < MaxParams)
diff --git a/Engine/Source/Runtime/Renderer/Public/MeshDrawShaderBindings.h b/Engine/Source/Runtime/Renderer/Public/MeshDrawShaderBindings.h
--- a/Engine/Source/Runtime/Renderer/Public/MeshDrawShaderBindings.h
+++ b/Engine/Source/Runtime/Renderer/Public/MeshDrawShaderBindings.h
@@ -199,10 +199,10 @@ public:
 
 				if (LooseParameterBuffer.BaseIndex == Parameter.GetBufferIndex())
 				{
-					TArrayView<const FShaderParameterInfo> Parameters(LooseParameterBuffer.Parameters);
+					TArrayView<const FShaderLooseParameterInfo> Parameters(LooseParameterBuffer.Parameters);
 					for (int32 LooseParameterIndex = 0; LooseParameterIndex < Parameters.Num(); LooseParameterIndex++)
 					{
-						FShaderParameterInfo LooseParameter = Parameters[LooseParameterIndex];
+						const FShaderLooseParameterInfo LooseParameter = Parameters[LooseParameterIndex];
 
 						if (Parameter.GetBaseIndex() == LooseParameter.BaseIndex)
 						{
