Adding missing headers in SceneVisibility.cpp for static mesh, skeletal mesh component, and other dependencies. Also removed a commented out code block that was no longer needed.

diff --git a/Engine/Source/Developer/DrawPrimitiveDebugger/Private/DrawPrimitiveDebuggerConfig.cpp b/Engine/Source/Developer/DrawPrimitiveDebugger/Private/DrawPrimitiveDebuggerConfig.cpp
--- a/Engine/Source/Developer/DrawPrimitiveDebugger/Private/DrawPrimitiveDebuggerConfig.cpp
+++ b/Engine/Source/Developer/DrawPrimitiveDebugger/Private/DrawPrimitiveDebuggerConfig.cpp
@@ -1,6 +1,7 @@
 ﻿// Copyright Epic Games, Inc. All Rights Reserved.
 
 #include "DrawPrimitiveDebuggerConfig.h"
+#include "Misc/ConfigCacheIni.h"
 
 void UDrawPrimitiveDebuggerUserSettings::SetFontSize(const int32 InFontSize)
 {
diff --git a/Engine/Source/Runtime/Renderer/Private/SceneVisibility.cpp b/Engine/Source/Runtime/Renderer/Private/SceneVisibility.cpp
--- a/Engine/Source/Runtime/Renderer/Private/SceneVisibility.cpp
+++ b/Engine/Source/Runtime/Renderer/Private/SceneVisibility.cpp
@@ -47,7 +47,17 @@
 #include "RayTracing/RayTracingInstanceCulling.h"
 #include "RendererModule.h"
 #include "SceneViewExtension.h"
+
+#if !UE_BUILD_SHIPPING
 #include "ViewDebug.h"
+#include "Engine/StaticMesh.h"
+#include "Engine/SkeletalMesh.h"
+#include "Materials/MaterialInterface.h"
+#include "Components/PrimitiveComponent.h"
+#include "Components/StaticMeshComponent.h"
+#include "Components/SkeletalMeshComponent.h"
+#include "Rendering/SkeletalMeshRenderData.h"
+#endif
 
 /*------------------------------------------------------------------------------
 	Globals
@@ -4223,18 +4233,6 @@ FViewDebugInfo::FViewDebugInfo()
 	bShouldCaptureSingleFrame = false;
 }
 
-/*bool FViewDebugInfo::FPrimitiveInfo::operator<(const FPrimitiveInfo& Other) const
-{
-	// Sort by name to group similar assets together, then by exact primitives so we can ignore duplicates
-	const int32 NameCompare = Name.Compare(Other.Name);
-	if (NameCompare != 0)
-	{
-		return NameCompare < 0;
-	}
-
-	return PrimitiveSceneInfo < Other.PrimitiveSceneInfo;
-}*/
-
 void FViewDebugInfo::ProcessPrimitive(FPrimitiveSceneInfo* PrimitiveSceneInfo, const FViewInfo& View, FScene* Scene, const UPrimitiveComponent* DebugComponent)
 {
 	if (!DebugComponent->IsRegistered())
