Move light visibility calculations, fog setup, and capsule shadow setup into an async task during ComputeViewVisibility.

diff --git a/Engine/Source/Runtime/RenderCore/Public/ShaderParameterMacros.h b/Engine/Source/Runtime/RenderCore/Public/ShaderParameterMacros.h
--- a/Engine/Source/Runtime/RenderCore/Public/ShaderParameterMacros.h
+++ b/Engine/Source/Runtime/RenderCore/Public/ShaderParameterMacros.h
@@ -115,7 +115,6 @@ public:
 	/** Creates a uniform buffer with the given value, and returns a structured reference to it. */
 	static TUniformBufferRef<TBufferStruct> CreateUniformBufferImmediate(const TBufferStruct& Value, EUniformBufferUsage Usage, EUniformBufferValidation Validation = EUniformBufferValidation::ValidateResources)
 	{
-		check(IsInRenderingThread() || IsInRHIThread());
 		return TUniformBufferRef<TBufferStruct>(RHICreateUniformBuffer(&Value, TBufferStruct::FTypeInfo::GetStructMetadata()->GetLayoutPtr(), Usage, Validation));
 	}
 
diff --git a/Engine/Source/Runtime/Renderer/Private/DeferredShadingRenderer.h b/Engine/Source/Runtime/Renderer/Private/DeferredShadingRenderer.h
--- a/Engine/Source/Runtime/Renderer/Private/DeferredShadingRenderer.h
+++ b/Engine/Source/Runtime/Renderer/Private/DeferredShadingRenderer.h
@@ -446,6 +446,8 @@ private:
 
 	void PreVisibilityFrameSetup(FRDGBuilder& GraphBuilder, const FSceneTexturesConfig& SceneTexturesConfig);
 
+	void ComputeLightVisibility();
+
 	/** Determines which primitives are visible for each view. */
 	void InitViews(FRDGBuilder& GraphBuilder, const FSceneTexturesConfig& SceneTexturesConfig, FExclusiveDepthStencil::Type BasePassDepthStencilAccess, struct FILCUpdatePrimTaskData& ILCTaskData, FInstanceCullingManager& InstanceCullingManager);
 
diff --git a/Engine/Source/Runtime/Renderer/Private/SceneRendering.h b/Engine/Source/Runtime/Renderer/Private/SceneRendering.h
--- a/Engine/Source/Runtime/Renderer/Private/SceneRendering.h
+++ b/Engine/Source/Runtime/Renderer/Private/SceneRendering.h
@@ -1241,7 +1241,10 @@ public:
 	TArray<FSortedTrianglesMeshBatch, SceneRenderingAllocator> SortedTrianglesMeshBatches;
 
 	/** A map from light ID to a boolean visibility value. */
-	TArray<FVisibleLightViewInfo,SceneRenderingAllocator> VisibleLightInfos;
+	TArray<FVisibleLightViewInfo, SceneRenderingAllocator> VisibleLightInfos;
+
+	/** Tracks the list of visible reflection capture lights that need to add meshes to the view. */
+	TArray<const FLightSceneProxy*, SceneRenderingAllocator> VisibleReflectionCaptureLights;
 
 	/** The view's batched elements. */
 	FBatchedElements BatchedViewElements;
@@ -2269,6 +2272,10 @@ protected:
 	void ComputeViewVisibility(FRHICommandListImmediate& RHICmdList, FExclusiveDepthStencil::Type BasePassDepthStencilAccess, FViewVisibleCommandsPerView& ViewCommandsPerView, 
 		FGlobalDynamicIndexBuffer& DynamicIndexBuffer, FGlobalDynamicVertexBuffer& DynamicVertexBuffer, FGlobalDynamicReadBuffer& DynamicReadBuffer, FInstanceCullingManager& InstanceCullingManager);
 
+	virtual void ComputeLightVisibility();
+
+	void GatherReflectionCaptureLightMeshElements();
+
 	/** Performs once per frame setup after to visibility determination. */
 	void PostVisibilityFrameSetup(FILCUpdatePrimTaskData& OutILCTaskData);
 
@@ -2382,6 +2389,19 @@ private:
 	friend class FRendererModule;
 };
 
+template <typename LambdaType>
+UE::Tasks::FTask LaunchSceneRenderTask(const TCHAR* DebugName, LambdaType&& Lambda, bool bExecuteInParallelCondition = true, LowLevelTasks::ETaskPriority TaskPriority = LowLevelTasks::ETaskPriority::Normal)
+{
+	const bool bExecuteInParallel = bExecuteInParallelCondition && FApp::ShouldUseThreadingForPerformance() && GIsThreadedRendering;
+
+	return UE::Tasks::Launch(DebugName, [Lambda = MoveTemp(Lambda)]
+	{
+		FTaskTagScope Scope(ETaskTag::EParallelRenderingThread);
+		Lambda();
+	},
+	TaskPriority, bExecuteInParallel ? UE::Tasks::EExtendedTaskPriority::None : UE::Tasks::EExtendedTaskPriority::Inline);
+}
+
 struct FForwardScreenSpaceShadowMaskTextureMobileOutputs
 {
 	TRefCountPtr<IPooledRenderTarget> ScreenSpaceShadowMaskTextureMobile;
diff --git a/Engine/Source/Runtime/Renderer/Private/SceneVisibility.cpp b/Engine/Source/Runtime/Renderer/Private/SceneVisibility.cpp
--- a/Engine/Source/Runtime/Renderer/Private/SceneVisibility.cpp
+++ b/Engine/Source/Runtime/Renderer/Private/SceneVisibility.cpp
@@ -4551,11 +4551,10 @@ void FSceneRenderer::ComputeViewVisibility(
 	STAT(int32 NumCulledPrimitives = 0);
 	STAT(int32 NumOccludedPrimitives = 0);
 
-	// Allocate the visible light info.
-	if (Scene->Lights.GetMaxIndex() > 0)
+	UE::Tasks::FTask ComputeLightVisibilityTask = LaunchSceneRenderTask(UE_SOURCE_LOCATION, [this]
 	{
-		VisibleLightInfos.AddDefaulted(Scene->Lights.GetMaxIndex());
-	}
+		ComputeLightVisibility();
+	});
 
 	int32 NumPrimitives = Scene->Primitives.Num();
 	float CurrentRealTime = ViewFamily.Time.GetRealTimeSeconds();
@@ -4623,27 +4622,10 @@ void FSceneRenderer::ComputeViewVisibility(
 			View.StaticMeshFadeInDitheredLODMap.Init(false, Scene->StaticMeshes.GetMaxIndex());
 			View.PrimitivesLODMask.Init(FLODMask(), Scene->Primitives.Num());
 
-			View.VisibleLightInfos.Empty(Scene->Lights.GetMaxIndex());
-
 			// The dirty list allocation must take into account the max possible size because when GILCUpdatePrimTaskEnabled is true,
 			// the indirect lighting cache will be update on by threaded job, which can not do reallocs on the buffer (since it uses the SceneRenderingAllocator).
 			View.DirtyIndirectLightingCacheBufferPrimitives.Reserve(Scene->Primitives.Num());
 
-			for (int32 LightIndex = 0;LightIndex < Scene->Lights.GetMaxIndex();LightIndex++)
-			{
-				if (LightIndex + 2 < Scene->Lights.GetMaxIndex())
-				{
-					if (LightIndex > 2)
-					{
-						FLUSH_CACHE_LINE(&View.VisibleLightInfos(LightIndex - 2));
-					}
-					// @todo optimization These prefetches cause asserts since LightIndex > View.VisibleLightInfos.Num() - 1
-					//FPlatformMisc::Prefetch(&View.VisibleLightInfos[LightIndex+2]);
-					//FPlatformMisc::Prefetch(&View.VisibleLightInfos[LightIndex+1]);
-				}
-				new(View.VisibleLightInfos) FVisibleLightViewInfo();
-			}
-
 			View.PrimitiveViewRelevanceMap.Reset(Scene->Primitives.Num());
 			View.PrimitiveViewRelevanceMap.AddZeroed(Scene->Primitives.Num());
 
@@ -4891,37 +4873,40 @@ void FSceneRenderer::ComputeViewVisibility(
 	INC_DWORD_STAT_BY(STAT_ProcessedPrimitives,NumProcessedPrimitives);
 	INC_DWORD_STAT_BY(STAT_CulledPrimitives,NumCulledPrimitives);
 	INC_DWORD_STAT_BY(STAT_OccludedPrimitives,NumOccludedPrimitives);
+	
+	ComputeLightVisibilityTask.Wait();
 }
 
-void FSceneRenderer::PostVisibilityFrameSetup(FILCUpdatePrimTaskData& OutILCTaskData)
+void FDeferredShadingSceneRenderer::ComputeLightVisibility()
 {
-	QUICK_SCOPE_CYCLE_COUNTER(STAT_PostVisibilityFrameSetup);
+	FSceneRenderer::ComputeLightVisibility();
 
-	{
-		QUICK_SCOPE_CYCLE_COUNTER(STAT_PostVisibilityFrameSetup_Sort);
-		for (int32 ViewIndex = 0; ViewIndex < Views.Num(); ViewIndex++)
-		{		
-			FViewInfo& View = Views[ViewIndex];
+	CreateIndirectCapsuleShadows();
 
-			View.MeshDecalBatches.Sort();
+	InitFogConstants();
 
-			if (View.State)
-			{
-				((FSceneViewState*)View.State)->TrimHistoryRenderTargets(Scene);
-			}
-		}
-	}
+	SetupVolumetricFog();
+}
 
-	const bool bSetupMobileLightShafts = FeatureLevel <= ERHIFeatureLevel::ES3_1 && ShouldRenderLightShafts(ViewFamily);
+void FSceneRenderer::ComputeLightVisibility()
+{
+	QUICK_SCOPE_CYCLE_COUNTER(STAT_PostVisibilityFrameSetup_Light_Visibility);
 
-	if (ViewFamily.EngineShowFlags.HitProxies == 0 && Scene->PrecomputedLightVolumes.Num() > 0
-		&& GILCUpdatePrimTaskEnabled && FPlatformProcess::SupportsMultithreading())
+	VisibleLightInfos.AddDefaulted(Scene->Lights.GetMaxIndex());
+
+	for (int32 ViewIndex = 0; ViewIndex < Views.Num(); ViewIndex++)
 	{
-		Scene->IndirectLightingCache.StartUpdateCachePrimitivesTask(Scene, *this, true, OutILCTaskData);
+		FViewInfo& View = Views[ViewIndex];
+		View.VisibleLightInfos.Empty(Scene->Lights.GetMaxIndex());
+
+		for (int32 LightIndex = 0; LightIndex < Scene->Lights.GetMaxIndex(); LightIndex++)
+		{
+			new (View.VisibleLightInfos) FVisibleLightViewInfo();
+		}
 	}
 
-	{
-		QUICK_SCOPE_CYCLE_COUNTER(STAT_PostVisibilityFrameSetup_Light_Visibility);
+	const bool bSetupMobileLightShafts = FeatureLevel <= ERHIFeatureLevel::ES3_1 && ShouldRenderLightShafts(ViewFamily);
+
 	// determine visibility of each light
 	for(auto LightIt = Scene->Lights.CreateConstIterator();LightIt;++LightIt)
 	{
@@ -4982,102 +4967,146 @@ void FSceneRenderer::PostVisibilityFrameSetup(FILCUpdatePrimTaskData& OutILCTask
 
 				if( DistanceSqr < Radius * Radius )
 				{
-					FLightRenderParameters LightParameters;
+					View.VisibleReflectionCaptureLights.Emplace(Proxy);
+				}
+			}
+		}
+	}
+}
 
-					Proxy->GetLightShaderParameters(LightParameters);
+void FSceneRenderer::GatherReflectionCaptureLightMeshElements()
+{
+	// view frustum cull lights in each view
+	for (FViewInfo& View : Views)
+	{
+		for (const FLightSceneProxy* Proxy : View.VisibleReflectionCaptureLights)
+		{
+			FVector Origin = Proxy->GetOrigin();
+			FVector ToLight = Origin - View.ViewMatrices.GetViewOrigin();
+			float DistanceSqr = ToLight | ToLight;
+			float Radius = Proxy->GetRadius();
 
-					// Force to be at least 0.75 pixels
-					float CubemapSize = (float)IConsoleManager::Get().FindTConsoleVariableDataInt( TEXT("r.ReflectionCaptureResolution") )->GetValueOnAnyThread();
-					float Distance = FMath::Sqrt( DistanceSqr );
-					float MinRadius = Distance * 0.75f / CubemapSize;
-					LightParameters.SourceRadius = FMath::Max( MinRadius, LightParameters.SourceRadius );
+			FLightRenderParameters LightParameters;
+			Proxy->GetLightShaderParameters(LightParameters);
 
-					// Snap to cubemap pixel center to reduce aliasing
-					FVector Scale = ToLight.GetAbs();
-					int32 MaxComponent = Scale.X > Scale.Y ? ( Scale.X > Scale.Z ? 0 : 2 ) : ( Scale.Y > Scale.Z ? 1 : 2 );
-					for( int32 k = 1; k < 3; k++ )
-					{
-						float Projected = ToLight[ (MaxComponent + k) % 3 ] / Scale[ MaxComponent ];
-						float Quantized = ( FMath::RoundToFloat( Projected * (0.5f * CubemapSize) - 0.5f ) + 0.5f ) / (0.5f * CubemapSize);
-						ToLight[ (MaxComponent + k) % 3 ] = Quantized * Scale[ MaxComponent ];
-					}
-					Origin = ToLight + View.ViewMatrices.GetViewOrigin();
-				
-					FLinearColor Color( LightParameters.Color.R, LightParameters.Color.G, LightParameters.Color.B, LightParameters.FalloffExponent );
-					const bool bIsRectLight = Proxy->IsRectLight();
-					if( !bIsRectLight )
-					{
-						const float SphereArea = (4.0f * PI) * FMath::Square( LightParameters.SourceRadius );
-						const float CylinderArea = (2.0f * PI) * LightParameters.SourceRadius * LightParameters.SourceLength;
-						const float SurfaceArea = SphereArea + CylinderArea;
-						Color *= 4.0f / SurfaceArea;
-					}
+			// Force to be at least 0.75 pixels
+			float CubemapSize = (float)IConsoleManager::Get().FindTConsoleVariableDataInt(TEXT("r.ReflectionCaptureResolution"))->GetValueOnAnyThread();
+			float Distance = FMath::Sqrt(DistanceSqr);
+			float MinRadius = Distance * 0.75f / CubemapSize;
+			LightParameters.SourceRadius = FMath::Max(MinRadius, LightParameters.SourceRadius);
 
-					if( Proxy->IsInverseSquared() )
-					{
-						float LightRadiusMask = FMath::Square( 1.0f - FMath::Square( DistanceSqr * FMath::Square( LightParameters.InvRadius ) ) );
-						Color.A = LightRadiusMask;
-					}
-					else
-					{
-						// Remove inverse square falloff
-						Color *= DistanceSqr + 1.0f;
+			// Snap to cubemap pixel center to reduce aliasing
+			FVector Scale = ToLight.GetAbs();
+			int32 MaxComponent = Scale.X > Scale.Y ? (Scale.X > Scale.Z ? 0 : 2) : (Scale.Y > Scale.Z ? 1 : 2);
+			for (int32 k = 1; k < 3; k++)
+			{
+				float Projected = ToLight[(MaxComponent + k) % 3] / Scale[MaxComponent];
+				float Quantized = (FMath::RoundToFloat(Projected * (0.5f * CubemapSize) - 0.5f) + 0.5f) / (0.5f * CubemapSize);
+				ToLight[(MaxComponent + k) % 3] = Quantized * Scale[MaxComponent];
+			}
+			Origin = ToLight + View.ViewMatrices.GetViewOrigin();
 
-						// Apply falloff
-						Color.A = FMath::Pow( 1.0f - DistanceSqr * FMath::Square(LightParameters.InvRadius ), LightParameters.FalloffExponent );
-					}
-					
-					// Spot falloff
-					FVector L = ToLight.GetSafeNormal();
-					Color.A *= FMath::Square( FMath::Clamp( ( (L | (FVector)LightParameters.Direction) - LightParameters.SpotAngles.X ) * LightParameters.SpotAngles.Y, 0.0f, 1.0f ) );
+			FLinearColor Color(LightParameters.Color.R, LightParameters.Color.G, LightParameters.Color.B, LightParameters.FalloffExponent);
+			const bool bIsRectLight = Proxy->IsRectLight();
+			if (!bIsRectLight)
+			{
+				const float SphereArea = (4.0f * PI) * FMath::Square(LightParameters.SourceRadius);
+				const float CylinderArea = (2.0f * PI) * LightParameters.SourceRadius * LightParameters.SourceLength;
+				const float SurfaceArea = SphereArea + CylinderArea;
+				Color *= 4.0f / SurfaceArea;
+			}
 
-					Color.A *= LightParameters.SpecularScale;
+			if (Proxy->IsInverseSquared())
+			{
+				float LightRadiusMask = FMath::Square(1.0f - FMath::Square(DistanceSqr * FMath::Square(LightParameters.InvRadius)));
+				Color.A = LightRadiusMask;
+			}
+			else
+			{
+				// Remove inverse square falloff
+				Color *= DistanceSqr + 1.0f;
 
-					// Rect is one sided
-					if( bIsRectLight && (L | (FVector)LightParameters.Direction) < 0.0f )
-						continue;
+				// Apply falloff
+				Color.A = FMath::Pow(1.0f - DistanceSqr * FMath::Square(LightParameters.InvRadius), LightParameters.FalloffExponent);
+			}
 
-					UTexture* SurfaceTexture = nullptr;
-					if (bIsRectLight)
-					{
-						const FRectLightSceneProxy* RectLightProxy = (const FRectLightSceneProxy*)Proxy;
-						SurfaceTexture = RectLightProxy->SourceTexture;
-					}
-					
-					FMaterialRenderProxy* ColoredMeshInstance = nullptr;
-					if (SurfaceTexture)
-						ColoredMeshInstance = Allocator.Create<FColoredTexturedMaterialRenderProxy>(GEngine->EmissiveMeshMaterial->GetRenderProxy(), Color, NAME_Color, SurfaceTexture, NAME_LinearColor);
-					else
-						ColoredMeshInstance = Allocator.Create<FColoredMaterialRenderProxy>(GEngine->EmissiveMeshMaterial->GetRenderProxy(), Color, NAME_Color);
+			// Spot falloff
+			FVector L = ToLight.GetSafeNormal();
+			Color.A *= FMath::Square(FMath::Clamp(((L | (FVector)LightParameters.Direction) - LightParameters.SpotAngles.X) * LightParameters.SpotAngles.Y, 0.0f, 1.0f));
 
-					FMatrix LightToWorld = Proxy->GetLightToWorld();
-					LightToWorld.RemoveScaling();
+			Color.A *= LightParameters.SpecularScale;
 
-					FViewElementPDI LightPDI( &View, NULL, &View.DynamicPrimitiveCollector);
+			// Rect is one sided
+			if (bIsRectLight && (L | (FVector)LightParameters.Direction) < 0.0f)
+				continue;
 
-					if( bIsRectLight )
-					{
-						DrawBox( &LightPDI, LightToWorld, FVector( 0.0f, LightParameters.SourceRadius, LightParameters.SourceLength ), ColoredMeshInstance, SDPG_World );
-					}
-					else if( LightParameters.SourceLength > 0.0f )
-					{
-						DrawSphere( &LightPDI, Origin + 0.5f * LightParameters.SourceLength * LightToWorld.GetUnitAxis( EAxis::Z ), FRotator::ZeroRotator, LightParameters.SourceRadius * FVector::OneVector, 36, 24, ColoredMeshInstance, SDPG_World );
-						DrawSphere( &LightPDI, Origin - 0.5f * LightParameters.SourceLength * LightToWorld.GetUnitAxis( EAxis::Z ), FRotator::ZeroRotator, LightParameters.SourceRadius * FVector::OneVector, 36, 24, ColoredMeshInstance, SDPG_World );
-						DrawCylinder( &LightPDI, Origin, LightToWorld.GetUnitAxis( EAxis::X ), LightToWorld.GetUnitAxis( EAxis::Y ), LightToWorld.GetUnitAxis( EAxis::Z ), LightParameters.SourceRadius, 0.5f * LightParameters.SourceLength, 36, ColoredMeshInstance, SDPG_World );
-					}
-					else
-					{
-						DrawSphere( &LightPDI, Origin, FRotator::ZeroRotator, LightParameters.SourceRadius * FVector::OneVector, 36, 24, ColoredMeshInstance, SDPG_World );
-					}
-				}
+			UTexture* SurfaceTexture = nullptr;
+			if (bIsRectLight)
+			{
+				const FRectLightSceneProxy* RectLightProxy = (const FRectLightSceneProxy*)Proxy;
+				SurfaceTexture = RectLightProxy->SourceTexture;
+			}
+
+			FMaterialRenderProxy* ColoredMeshInstance = nullptr;
+			if (SurfaceTexture)
+			{
+				ColoredMeshInstance = Allocator.Create<FColoredTexturedMaterialRenderProxy>(GEngine->EmissiveMeshMaterial->GetRenderProxy(), Color, NAME_Color, SurfaceTexture, NAME_LinearColor);
+			}
+			else
+			{
+				ColoredMeshInstance = Allocator.Create<FColoredMaterialRenderProxy>(GEngine->EmissiveMeshMaterial->GetRenderProxy(), Color, NAME_Color);
+			}
+
+			FMatrix LightToWorld = Proxy->GetLightToWorld();
+			LightToWorld.RemoveScaling();
+
+			FViewElementPDI LightPDI(&View, NULL, &View.DynamicPrimitiveCollector);
+
+			if (bIsRectLight)
+			{
+				DrawBox(&LightPDI, LightToWorld, FVector(0.0f, LightParameters.SourceRadius, LightParameters.SourceLength), ColoredMeshInstance, SDPG_World);
+			}
+			else if (LightParameters.SourceLength > 0.0f)
+			{
+				DrawSphere(&LightPDI, Origin + 0.5f * LightParameters.SourceLength * LightToWorld.GetUnitAxis(EAxis::Z), FRotator::ZeroRotator, LightParameters.SourceRadius * FVector::OneVector, 36, 24, ColoredMeshInstance, SDPG_World);
+				DrawSphere(&LightPDI, Origin - 0.5f * LightParameters.SourceLength * LightToWorld.GetUnitAxis(EAxis::Z), FRotator::ZeroRotator, LightParameters.SourceRadius * FVector::OneVector, 36, 24, ColoredMeshInstance, SDPG_World);
+				DrawCylinder(&LightPDI, Origin, LightToWorld.GetUnitAxis(EAxis::X), LightToWorld.GetUnitAxis(EAxis::Y), LightToWorld.GetUnitAxis(EAxis::Z), LightParameters.SourceRadius, 0.5f * LightParameters.SourceLength, 36, ColoredMeshInstance, SDPG_World);
+			}
+			else
+			{
+				DrawSphere(&LightPDI, Origin, FRotator::ZeroRotator, LightParameters.SourceRadius * FVector::OneVector, 36, 24, ColoredMeshInstance, SDPG_World);
 			}
 		}
+
+		View.VisibleReflectionCaptureLights.Empty();
 	}
-	}
+}
+
+void FSceneRenderer::PostVisibilityFrameSetup(FILCUpdatePrimTaskData& OutILCTaskData)
+{
+	QUICK_SCOPE_CYCLE_COUNTER(STAT_PostVisibilityFrameSetup);
+
 	{
+		QUICK_SCOPE_CYCLE_COUNTER(STAT_PostVisibilityFrameSetup_Sort);
+		for (int32 ViewIndex = 0; ViewIndex < Views.Num(); ViewIndex++)
+		{		
+			FViewInfo& View = Views[ViewIndex];
+
+			View.MeshDecalBatches.Sort();
+
+			if (View.State)
+			{
+				((FSceneViewState*)View.State)->TrimHistoryRenderTargets(Scene);
+			}
+		}
+	}
 
-		QUICK_SCOPE_CYCLE_COUNTER(STAT_PostVisibilityFrameSetup_InitFogConstants);
-		InitFogConstants();
+	GatherReflectionCaptureLightMeshElements();
+
+	if (ViewFamily.EngineShowFlags.HitProxies == 0 && Scene->PrecomputedLightVolumes.Num() > 0
+		&& GILCUpdatePrimTaskEnabled && FPlatformProcess::SupportsMultithreading())
+	{
+		Scene->IndirectLightingCache.StartUpdateCachePrimitivesTask(Scene, *this, true, OutILCTaskData);
 	}
 }
 
@@ -5113,13 +5142,6 @@ void FDeferredShadingSceneRenderer::InitViews(FRDGBuilder& GraphBuilder, const F
 
 	FRHICommandListImmediate& RHICmdList = GraphBuilder.RHICmdList;
 
-	const auto DispatchToRHIThread = [&]()
-	{
-		RHICmdList.ImmediateFlush(EImmediateFlushType::DispatchToRHIThread);
-	};
-
-	DispatchToRHIThread();
-
 	// Create GPU-side representation of the view for instance culling.
 	for (int32 ViewIndex = 0; ViewIndex < Views.Num(); ++ViewIndex)
 	{
@@ -5148,19 +5170,12 @@ void FDeferredShadingSceneRenderer::InitViews(FRDGBuilder& GraphBuilder, const F
 
 	ComputeViewVisibility(RHICmdList, BasePassDepthStencilAccess, ViewCommandsPerView, DynamicIndexBufferForInitViews, DynamicVertexBufferForInitViews, DynamicReadBufferForInitViews, InstanceCullingManager);
 
-	DispatchToRHIThread();
-
-	// This has to happen before Scene->IndirectLightingCache.UpdateCache, since primitives in View.IndirectShadowPrimitives need ILC updates
-	CreateIndirectCapsuleShadows();
-
 	// This must happen before we start initialising and using views.
 	if (Scene)
 	{
 		UpdateSkyIrradianceGpuBuffer(RHICmdList, ViewFamily.EngineShowFlags, Scene->SkyLight, Scene->SkyIrradianceEnvironmentMap);
 	}
 
-	DispatchToRHIThread();
-
 	// Initialise Sky/View resources before the view global uniform buffer is built.
 	if (ShouldRenderSkyAtmosphere(Scene, ViewFamily.EngineShowFlags))
 	{
@@ -5168,16 +5183,6 @@ void FDeferredShadingSceneRenderer::InitViews(FRDGBuilder& GraphBuilder, const F
 	}
 
 	PostVisibilityFrameSetup(ILCTaskData);
-	DispatchToRHIThread();
-
-	FVector AverageViewPosition(0);
-
-	for (int32 ViewIndex = 0; ViewIndex < Views.Num(); ViewIndex++)
-	{
-		FViewInfo& View = Views[ViewIndex];
-		AverageViewPosition += View.ViewMatrices.GetViewOrigin() / Views.Num();
-	}
-
 	InitViewsBeforePrepass(GraphBuilder, InstanceCullingManager);
 
 	{
@@ -5202,12 +5207,12 @@ void FDeferredShadingSceneRenderer::InitViews(FRDGBuilder& GraphBuilder, const F
 		}
 	}
 
-	SetupVolumetricFog();
-
 	{
 		QUICK_SCOPE_CYCLE_COUNTER(STAT_InitViews_OnStartRender);
 		OnStartRender(RHICmdList);
 	}
+
+	RHICmdList.ImmediateFlush(EImmediateFlushType::DispatchToRHIThread);
 }
 
 template<class T>
