Prevent crash when BeginRenderingViewFamily is called without a scene.

#rnx

diff --git a/Engine/Source/Runtime/Renderer/Private/SceneRendering.cpp b/Engine/Source/Runtime/Renderer/Private/SceneRendering.cpp
--- a/Engine/Source/Runtime/Renderer/Private/SceneRendering.cpp
+++ b/Engine/Source/Runtime/Renderer/Private/SceneRendering.cpp
@@ -4539,7 +4539,7 @@ void FRendererModule::BeginRenderingViewFamilies(FCanvas* Canvas, TArrayView<FSc
 
 	UWorld* World = nullptr;
 
-	FScene* const Scene = ViewFamilies[0]->Scene->GetRenderScene();
+	FScene* const Scene = ViewFamilies[0]->Scene ? ViewFamilies[0]->Scene->GetRenderScene() : nullptr;
 	if (Scene)
 	{
 		World = Scene->GetWorld();
