Disable indirect light rendering when the Lighting Complexity viewmode is enabled

diff --git a/Engine/Source/Runtime/Renderer/Private/IndirectLightRendering.cpp b/Engine/Source/Runtime/Renderer/Private/IndirectLightRendering.cpp
--- a/Engine/Source/Runtime/Renderer/Private/IndirectLightRendering.cpp
+++ b/Engine/Source/Runtime/Renderer/Private/IndirectLightRendering.cpp
@@ -789,7 +789,8 @@ void FDeferredShadingSceneRenderer::RenderDiffuseIndirectAndAmbientOcclusion(
 	using namespace HybridIndirectLighting;
 
 	extern int32 GLumenVisualizeIndirectDiffuse;
-	if ((GLumenVisualizeIndirectDiffuse != 0) != bIsVisualizePass)
+	if ((GLumenVisualizeIndirectDiffuse != 0) != bIsVisualizePass
+		|| ViewFamily.EngineShowFlags.VisualizeLightCulling)
 	{
 		return;
 	}
