Fixed: Several issues with a mobile preview in the editor for projets with MSAA and LDR renderng enabled

diff --git a/Engine/Shaders/Private/PostProcessCompositeEditorPrimitives.usf b/Engine/Shaders/Private/PostProcessCompositeEditorPrimitives.usf
--- a/Engine/Shaders/Private/PostProcessCompositeEditorPrimitives.usf
+++ b/Engine/Shaders/Private/PostProcessCompositeEditorPrimitives.usf
@@ -292,10 +292,6 @@ void MainCompositeEditorPrimitivesPS(
 	
 	float4 SceneColor = Texture2DSample(ColorTexture, ColorSampler, ColorUV);
 
-#if OUTPUT_SRGB_BUFFER
-	SceneColor.rgb = pow(SceneColor.rgb, 1.0f / 2.2f);
-#endif
-
 	// Resolve editor primitive scene color and depth.
 	float4 EditorPrimitiveColor;
 	float DepthMask;
diff --git a/Engine/Source/Runtime/D3D12RHI/Private/D3D12Viewport.cpp b/Engine/Source/Runtime/D3D12RHI/Private/D3D12Viewport.cpp
--- a/Engine/Source/Runtime/D3D12RHI/Private/D3D12Viewport.cpp
+++ b/Engine/Source/Runtime/D3D12RHI/Private/D3D12Viewport.cpp
@@ -362,7 +362,7 @@ FD3D12Texture* FD3D12Viewport::CreateDummyBackBufferTextures(FD3D12Adapter* InAd
 		FRHITextureCreateDesc::Create2D(TEXT("BackBufferReference"))
 		.SetExtent(FIntPoint(InSizeX, InSizeY))
 		.SetFormat(InPixelFormat)
-		.SetFlags(ETextureCreateFlags::RenderTargetable | ETextureCreateFlags::Presentable)
+		.SetFlags(ETextureCreateFlags::RenderTargetable | ETextureCreateFlags::Presentable | ETextureCreateFlags::ResolveTargetable)
 		.SetInitialState(ERHIAccess::Present);
 
 	FD3D12Texture* Result = InAdapter->CreateLinkedObject<FD3D12Texture>(FRHIGPUMask::All(), [&](FD3D12Device* Device)
diff --git a/Engine/Source/Runtime/Renderer/Private/MobileShadingRenderer.cpp b/Engine/Source/Runtime/Renderer/Private/MobileShadingRenderer.cpp
--- a/Engine/Source/Runtime/Renderer/Private/MobileShadingRenderer.cpp
+++ b/Engine/Source/Runtime/Renderer/Private/MobileShadingRenderer.cpp
@@ -438,7 +438,13 @@ void FMobileSceneRenderer::InitViews(FRDGBuilder& GraphBuilder, FSceneTexturesCo
 	ComputeViewVisibility(RHICmdList, BasePassDepthStencilAccess, ViewCommandsPerView, DynamicIndexBuffer, DynamicVertexBuffer, DynamicReadBuffer, InstanceCullingManager);
 	PostVisibilityFrameSetup(ILCTaskData);
 
-	const FIntPoint RenderTargetSize = (ViewFamily.RenderTarget->GetRenderTargetTexture().IsValid()) ? ViewFamily.RenderTarget->GetRenderTargetTexture()->GetSizeXY() : ViewFamily.RenderTarget->GetSizeXY();
+	FIntPoint RenderTargetSize = ViewFamily.RenderTarget->GetSizeXY();
+	EPixelFormat RenderTargetPixelFormat = PF_Unknown;
+	if (ViewFamily.RenderTarget->GetRenderTargetTexture().IsValid())
+	{
+		RenderTargetSize = ViewFamily.RenderTarget->GetRenderTargetTexture()->GetSizeXY();
+		RenderTargetPixelFormat = ViewFamily.RenderTarget->GetRenderTargetTexture()->GetFormat();
+	}
 	const bool bRequiresUpscale = ((int32)RenderTargetSize.X > FamilySize.X || (int32)RenderTargetSize.Y > FamilySize.Y);
 	// ES requires that the back buffer and depth match dimensions.
 	// For the most part this is not the case when using scene captures. Thus scene captures always render to scene color target.
@@ -452,6 +458,7 @@ void FMobileSceneRenderer::InitViews(FRDGBuilder& GraphBuilder, FSceneTexturesCo
 						|| Views[0].bIsReflectionCapture 
 						// If the resolve texture is not the same as the MSAA texture, we need to render to scene color and copy to back buffer.
 						|| (NumMSAASamples > 1 && !RHISupportsSeparateMSAAAndResolveTextures(ShaderPlatform))
+						|| (NumMSAASamples > 1 && (RenderTargetPixelFormat != PF_Unknown && RenderTargetPixelFormat != SceneTexturesConfig.ColorFormat))
 						|| bIsFullDepthPrepassEnabled;
 
 	const bool bSceneDepthCapture = (
@@ -531,7 +538,7 @@ void FMobileSceneRenderer::InitViews(FRDGBuilder& GraphBuilder, FSceneTexturesCo
 	// Update the bKeepDepthContent based on the mobile renderer status.
 	SceneTexturesConfig.bKeepDepthContent = bKeepDepthContent;
 	// If we render in a single pass MSAA targets can be memoryless
-	SceneTexturesConfig.bMemorylessMSAA = !bRequiresMultiPass;
+	SceneTexturesConfig.bMemorylessMSAA = !(bRequiresMultiPass || bShouldCompositeEditorPrimitives);
 	
 	if (bDeferredShading) 
 	{
@@ -1189,6 +1196,12 @@ void FMobileSceneRenderer::RenderForwardSinglePass(FRDGBuilder& GraphBuilder, FM
 		// Pre-tonemap before MSAA resolve (iOS only)
 		PreTonemapMSAA(RHICmdList, SceneTextures);
 	});
+	
+	// resolve MSAA depth
+	if (!bIsFullDepthPrepassEnabled)
+	{
+		AddResolveSceneDepthPass(GraphBuilder, View, SceneTextures.Depth);
+	}
 }
 
 void FMobileSceneRenderer::RenderForwardMultiPass(FRDGBuilder& GraphBuilder, FMobileRenderPassParameters* PassParameters, FRenderTargetBindingSlots& BasePassRenderTargets, int32 ViewIndex, FViewInfo& View, FSceneTextures& SceneTextures)
@@ -1215,9 +1228,12 @@ void FMobileSceneRenderer::RenderForwardMultiPass(FRDGBuilder& GraphBuilder, FMo
 		PostRenderBasePass(RHICmdList, View);
 	});
 
-	// resolve MSAA depth for translucency
-	AddResolveSceneDepthPass(GraphBuilder, View, SceneTextures.Depth);
+	// resolve MSAA depth
 	AddResolveSceneColorPass(GraphBuilder, View, SceneTextures.DepthAux);
+	if (!bIsFullDepthPrepassEnabled)
+	{
+		AddResolveSceneDepthPass(GraphBuilder, View, SceneTextures.Depth);
+	}
 
 	FExclusiveDepthStencil::Type ExclusiveDepthStencil = FExclusiveDepthStencil::DepthRead_StencilRead;
 	if (bModulatedShadowsInUse)
diff --git a/Engine/Source/Runtime/Renderer/Private/PostProcess/PostProcessCompositeEditorPrimitives.cpp b/Engine/Source/Runtime/Renderer/Private/PostProcess/PostProcessCompositeEditorPrimitives.cpp
--- a/Engine/Source/Runtime/Renderer/Private/PostProcess/PostProcessCompositeEditorPrimitives.cpp
+++ b/Engine/Source/Runtime/Renderer/Private/PostProcess/PostProcessCompositeEditorPrimitives.cpp
@@ -111,7 +111,6 @@ public:
 	static void ModifyCompilationEnvironment(const FGlobalShaderPermutationParameters& Parameters, FShaderCompilerEnvironment& OutEnvironment)
 	{
 		FGlobalShader::ModifyCompilationEnvironment(Parameters, OutEnvironment);
-		OutEnvironment.SetDefine(TEXT("OUTPUT_SRGB_BUFFER"), IsMobileColorsRGB() && IsMobilePlatform(Parameters.Platform));
 	}
 };
 
diff --git a/Engine/Source/Runtime/Renderer/Private/PostProcess/PostProcessing.cpp b/Engine/Source/Runtime/Renderer/Private/PostProcess/PostProcessing.cpp
--- a/Engine/Source/Runtime/Renderer/Private/PostProcess/PostProcessing.cpp
+++ b/Engine/Source/Runtime/Renderer/Private/PostProcess/PostProcessing.cpp
@@ -1790,7 +1790,7 @@ void AddMobilePostProcessingPasses(FRDGBuilder& GraphBuilder, FScene* Scene, con
 		MAX
 	};
 
-	const TCHAR* PassNames[] =
+	static const TCHAR* PassNames[] =
 	{
 		TEXT("Distortion"),
 		TEXT("SunMask"),
@@ -1832,7 +1832,7 @@ void AddMobilePostProcessingPasses(FRDGBuilder& GraphBuilder, FScene* Scene, con
 	
 	const EAutoExposureMethod AutoExposureMethod = GetAutoExposureMethod(View);
 	const bool bUseEyeAdaptation = IsMobileEyeAdaptationEnabled(View);
-
+	const bool bIsPostProcessingEnabled = IsPostProcessingEnabled(View);
 	
 	//The input scene color has been encoded to non-linear space and needs to decode somewhere if MSAA enabled on Metal platform
 	bool bMetalMSAAHDRDecode = GSupportsShaderFramebufferFetch && IsMetalMobilePlatform(View.GetShaderPlatform()) && GetDefaultMSAACount(ERHIFeatureLevel::ES3_1) > 1;
@@ -1843,7 +1843,7 @@ void AddMobilePostProcessingPasses(FRDGBuilder& GraphBuilder, FScene* Scene, con
 	bool bUseDof = GetMobileDepthOfFieldScale(View) > 0.0f && View.Family->EngineShowFlags.DepthOfField && !View.Family->EngineShowFlags.VisualizeDOF;
 	bool bUseMobileDof = bUseDof && !View.FinalPostProcessSettings.bMobileHQGaussian;
 
-	bool bUseToneMapper = !View.Family->EngineShowFlags.ShaderComplexity && IsMobileHDR();
+	bool bUseToneMapper = !View.Family->EngineShowFlags.ShaderComplexity && (IsMobileHDR() || IsMobileColorsRGB());
 
 	bool bUseHighResolutionScreenshotMask = IsHighResolutionScreenshotMaskEnabled(View);
 
@@ -1897,8 +1897,8 @@ void AddMobilePostProcessingPasses(FRDGBuilder& GraphBuilder, FScene* Scene, con
 			bMetalMSAAHDRDecode = false;
 		}
 	};
-
-	if (IsPostProcessingEnabled(View))
+	
+	if (bIsPostProcessingEnabled)
 	{
 		bool bUseSun = View.MobileLightShaft.IsSet();
 			
@@ -2302,11 +2302,11 @@ void AddMobilePostProcessingPasses(FRDGBuilder& GraphBuilder, FScene* Scene, con
 			BloomOutput = BlackAlphaOneDummy;
 		}
 
-		bool bDoGammaOnly = false;
+		bool bDoGammaOnly = !IsMobileHDR();
 
 		FRDGTextureRef ColorGradingTexture = nullptr;
 
-		if (IStereoRendering::IsAPrimaryView(View))
+		if (IStereoRendering::IsAPrimaryView(View) && !bDoGammaOnly)
 		{
 			ColorGradingTexture = AddCombineLUTPass(GraphBuilder, View);
 		}
diff --git a/Engine/Source/Runtime/Renderer/Private/SceneRendering.cpp b/Engine/Source/Runtime/Renderer/Private/SceneRendering.cpp
--- a/Engine/Source/Runtime/Renderer/Private/SceneRendering.cpp
+++ b/Engine/Source/Runtime/Renderer/Private/SceneRendering.cpp
@@ -5222,7 +5222,8 @@ void AddResolveSceneColorPass(FRDGBuilder& GraphBuilder, const FViewInfo& View,
 	const uint32 NumSamples = SceneColor.Target->Desc.NumSamples;
 	const EShaderPlatform CurrentShaderPlatform = GetFeatureLevelShaderPlatform(View.FeatureLevel);
 
-	if (NumSamples == 1 || !SceneColor.IsSeparate())
+	if (NumSamples == 1 || !SceneColor.IsSeparate()
+		|| EnumHasAnyFlags(SceneColor.Target->Desc.Flags, TexCreate_Memoryless))
 	{
 		return;
 	}
@@ -5386,7 +5387,8 @@ void AddResolveSceneDepthPass(FRDGBuilder& GraphBuilder, const FViewInfo& View,
 	const uint32 NumSamples = SceneDepth.Target->Desc.NumSamples;
 	const EShaderPlatform CurrentShaderPlatform = GetFeatureLevelShaderPlatform(View.FeatureLevel);
 
-	if (NumSamples == 1 || !SceneDepth.IsSeparate())
+	if (NumSamples == 1 || !SceneDepth.IsSeparate() 
+		|| EnumHasAnyFlags(SceneDepth.Target->Desc.Flags, TexCreate_Memoryless))
 	{
 		return;
 	}
