Fix a bug of missing reflection because of feature level SM6.

diff --git a/Engine/Source/Runtime/Renderer/Private/ReflectionEnvironmentCapture.cpp b/Engine/Source/Runtime/Renderer/Private/ReflectionEnvironmentCapture.cpp
--- a/Engine/Source/Runtime/Renderer/Private/ReflectionEnvironmentCapture.cpp
+++ b/Engine/Source/Runtime/Renderer/Private/ReflectionEnvironmentCapture.cpp
@@ -1418,7 +1418,7 @@ void FScene::CaptureOrUploadReflectionCapture(UReflectionCaptureComponent* Captu
 
 				FRDGTexture* FilteredSceneCubemapTexture = FilterReflectionEnvironment(GraphBuilder, ShaderMap, SceneCubemapTexture, nullptr);
 
-				if (FeatureLevel == ERHIFeatureLevel::SM5)
+				if (FeatureLevel >= ERHIFeatureLevel::SM5)
 				{
 					CopyToSceneArray(GraphBuilder, Scene, FilteredSceneCubemapTexture, ReflectionProxy, CubemapIndex);
 				}
