Fixed: missing modulated shadows in the Editor mobile preview when project uses LDR

diff --git a/Engine/Source/Runtime/Renderer/Private/MobileShadingRenderer.cpp b/Engine/Source/Runtime/Renderer/Private/MobileShadingRenderer.cpp
--- a/Engine/Source/Runtime/Renderer/Private/MobileShadingRenderer.cpp
+++ b/Engine/Source/Runtime/Renderer/Private/MobileShadingRenderer.cpp
@@ -1229,11 +1229,14 @@ void FMobileSceneRenderer::RenderForwardMultiPass(FRDGBuilder& GraphBuilder, FMo
 	});
 
 	// resolve MSAA depth
-	AddResolveSceneColorPass(GraphBuilder, View, SceneTextures.DepthAux);
 	if (!bIsFullDepthPrepassEnabled)
 	{
 		AddResolveSceneDepthPass(GraphBuilder, View, SceneTextures.Depth);
 	}
+	if (bRequiresSceneDepthAux)
+	{
+		AddResolveSceneColorPass(GraphBuilder, View, SceneTextures.DepthAux);
+	}
 
 	FExclusiveDepthStencil::Type ExclusiveDepthStencil = FExclusiveDepthStencil::DepthRead_StencilRead;
 	if (bModulatedShadowsInUse)
@@ -1664,7 +1667,7 @@ bool FMobileSceneRenderer::RequiresMultiPass(FRHICommandListImmediate& RHICmdLis
 	}
 		
 	// Always render LDR in single pass
-	if (!IsMobileHDR())
+	if (!IsMobileHDR() && !IsSimulatedPlatform(ShaderPlatform))
 	{
 		return false;
 	}
