Remove QuadComplexity from ES31 temporary until DebugView is in a separate pass

diff --git a/Engine/Shaders/Private/QuadOverdraw.ush b/Engine/Shaders/Private/QuadOverdraw.ush
--- a/Engine/Shaders/Private/QuadOverdraw.ush
+++ b/Engine/Shaders/Private/QuadOverdraw.ush
@@ -45,11 +45,7 @@ float ComplexityToFloat(uint Complexity)
 	return (float)Complexity / (float)QO_FRACTIONNAL_COMPLEXITY;
 }
 
-#if SHADING_PATH_MOBILE
-#define RWQuadBuffer MobileBasePass.QuadOverdraw 
-#else
 #define RWQuadBuffer DebugViewModePass.QuadOverdraw
-#endif
 
 #if OUTPUT_QUAD_OVERDRAW
 
diff --git a/Engine/Source/Runtime/Renderer/Private/DebugViewModeRendering.h b/Engine/Source/Runtime/Renderer/Private/DebugViewModeRendering.h
--- a/Engine/Source/Runtime/Renderer/Private/DebugViewModeRendering.h
+++ b/Engine/Source/Runtime/Renderer/Private/DebugViewModeRendering.h
@@ -175,7 +175,9 @@ public:
 		TCHAR BufferRegister[] = { 'u', '0', 0 };
 		BufferRegister[1] += GetQuadOverdrawUAVIndex(Parameters.Platform, Parameters.MaterialParameters.FeatureLevel);
 		OutEnvironment.SetDefine(TEXT("QUAD_BUFFER_REGISTER"), BufferRegister);
-		OutEnvironment.SetDefine(TEXT("OUTPUT_QUAD_OVERDRAW"), 1);
+		const bool bUsingMobileRenderer = FSceneInterface::GetShadingPath(GetMaxSupportedFeatureLevel(Parameters.Platform)) == EShadingPath::Mobile;
+		OutEnvironment.SetDefine(TEXT("OUTPUT_QUAD_OVERDRAW"), !bUsingMobileRenderer);
+
 
 		for (int i = 0; i < DVSM_MAX; ++i)
 		{
diff --git a/Engine/Source/Runtime/Renderer/Private/MobileBasePassRendering.cpp b/Engine/Source/Runtime/Renderer/Private/MobileBasePassRendering.cpp
--- a/Engine/Source/Runtime/Renderer/Private/MobileBasePassRendering.cpp
+++ b/Engine/Source/Runtime/Renderer/Private/MobileBasePassRendering.cpp
@@ -218,17 +218,6 @@ TRDGUniformBufferRef<FMobileBasePassUniformParameters> CreateMobileBasePassUnifo
 	}
 #endif
 
-	// Cannot be null even when DebugVM is off because it is in the same pass as RenderMobileBasePass(ValidateShaderParameterResourcesRHI UAV).
-	// If this is problematic, move DebugView to run in another pass and use FDebugViewModePassUniformParameters
-	const FSceneTextures& SceneTextures = View.GetSceneTextures();
-	FRDGTextureRef QuadOverdrawTexture = SceneTextures.QuadOverdraw;
-	if (!QuadOverdrawTexture)
-	{
-		QuadOverdrawTexture = GraphBuilder.CreateTexture(FRDGTextureDesc::Create2D(FIntPoint(1, 1), PF_R32_UINT, FClearValueBinding::None, TexCreate_UAV), TEXT("DummyOverdrawUAV"));
-	}
-
-	BasePassParameters->QuadOverdraw = GraphBuilder.CreateUAV(QuadOverdrawTexture);
-
 	return GraphBuilder.CreateUniformBuffer(BasePassParameters);
 }
 
diff --git a/Engine/Source/Runtime/Renderer/Private/MobileBasePassRendering.h b/Engine/Source/Runtime/Renderer/Private/MobileBasePassRendering.h
--- a/Engine/Source/Runtime/Renderer/Private/MobileBasePassRendering.h
+++ b/Engine/Source/Runtime/Renderer/Private/MobileBasePassRendering.h
@@ -38,7 +38,6 @@ BEGIN_GLOBAL_SHADER_PARAMETER_STRUCT(FMobileBasePassUniformParameters, )
 	SHADER_PARAMETER_STRUCT(FPlanarReflectionUniformParameters, PlanarReflection) // Single global planar reflection for the forward pass.
 	SHADER_PARAMETER_STRUCT(FMobileSceneTextureUniformParameters, SceneTextures)
 	SHADER_PARAMETER_STRUCT(FDebugViewModeUniformParameters, DebugViewMode)
-	SHADER_PARAMETER_RDG_TEXTURE_UAV(RWTexture2D<uint>, QuadOverdraw)
 	SHADER_PARAMETER_STRUCT(FReflectionUniformParameters, ReflectionsParameters)
 	SHADER_PARAMETER_TEXTURE(Texture2D, PreIntegratedGFTexture)
 	SHADER_PARAMETER_SAMPLER(SamplerState, PreIntegratedGFSampler)
