Implemented new VS/PS based rasterizer for hair. Implemented a naive software rasterizer for hair. Added an optional compaction pass to the hair software rasterizer.

#rnx

diff --git a/Engine/Shaders/Private/HairStrands/HairStrandsVisibilityRasterCompute.usf b/Engine/Shaders/Private/HairStrands/HairStrandsVisibilityRasterCompute.usf
--- a/Engine/Shaders/Private/HairStrands/HairStrandsVisibilityRasterCompute.usf
+++ b/Engine/Shaders/Private/HairStrands/HairStrandsVisibilityRasterCompute.usf
@@ -29,7 +29,7 @@ float2	OutputResolutionf;
 
 ///////////////////////////////////////////////////////////////////////////
 
-#if SHADER_RASTERCOMPUTE_BINNING || SHADER_RASTERCOMPUTE_RASTER || SHADER_RASTERCOMPUTE_DEPTH_GRID
+#if SHADER_RASTERCOMPUTE_BINNING || SHADER_RASTERCOMPUTE_COMPACTION || SHADER_RASTERCOMPUTE_RASTER || SHADER_RASTERCOMPUTE_DEPTH_GRID
 
 ///////////////////////////////////////////////////////////////////////////
 // 
@@ -585,6 +585,151 @@ void BinningCS(uint DispatchThreadID : SV_DispatchThreadID, uint GroupThreadID :
 
 ///////////////////////////////////////////////////////////////////////////
 
+#if SHADER_RASTERCOMPUTE_COMPACTION
+
+ByteAddressBuffer					VisTileData;
+Buffer<uint> 						VisTilePrims;
+Buffer<uint> 						VisTileArgs;
+RWByteAddressBuffer					OutCompactedVisTileData;
+RWBuffer<uint> 						OutCompactedVisTilePrims;
+RWBuffer<uint> 						OutCompactedVisTileArgs;
+
+uint LoadVisTileData(uint index, uint offset)
+{
+	return VisTileData.Load((((index)) * VT_SIZE * 4) + ((offset) * 4));
+}
+
+void StoreCompactedVisTileData(uint index, uint offset, uint value)
+{
+	OutCompactedVisTileData.Store((((index)) * VT_SIZE * 4) + ((offset) * 4), (value));
+}
+
+groupshared uint group_TotalPrimCount;
+groupshared uint group_PrimWriteOffset;
+groupshared uint group_NumTiles;
+groupshared uint group_TilesToCompact[1024];
+
+[numthreads(1024, 1, 1)]
+void CompactionCS(uint DispatchThreadID : SV_DispatchThreadID, uint GroupThreadID : SV_GroupThreadID, uint2 GroupID : SV_GroupID)
+{
+	if (GroupThreadID == 0)
+	{
+		group_TotalPrimCount = 0;
+		group_NumTiles = 0;
+	}
+
+	GroupMemoryBarrierWithGroupSync();
+
+	const uint NumTiles = VisTileArgs[0];
+	const uint PackedCoord = PackVisTileCoord(GroupID);
+
+	// Compute total number of primitives at this tile coordinate
+	uint LocalPrimCount = 0;
+	{
+		for (uint TileIdx = GroupThreadID; TileIdx < NumTiles; TileIdx += 1024)
+		{
+			const uint TilePackedCoord = LoadVisTileData(TileIdx, VT_Coord);
+			if (PackedCoord == TilePackedCoord)
+			{
+				LocalPrimCount += LoadVisTileData(TileIdx, VT_PrimCount);
+				
+				uint WritePos;
+				InterlockedAdd(group_NumTiles, 1, WritePos);
+				if (WritePos < 1024)
+				{
+					group_TilesToCompact[WritePos] = TileIdx;
+				}
+			}
+		}
+	}
+
+	GroupMemoryBarrierWithGroupSync();
+
+	if (LocalPrimCount > 0)
+	{
+		InterlockedAdd(group_TotalPrimCount, LocalPrimCount);
+	}
+
+	GroupMemoryBarrierWithGroupSync();
+
+	const uint TotalPrimCount = group_TotalPrimCount;
+
+	if (TotalPrimCount == 0)
+	{
+		return;
+	}
+
+	// Allocate space
+	if (GroupThreadID == 0)
+	{
+		const uint NumTilesToAllocate = (TotalPrimCount + 1023) / 1024;
+
+		uint FirstCompactedTile;
+		InterlockedAdd(OutCompactedVisTileArgs[0], NumTilesToAllocate, FirstCompactedTile);
+
+		group_PrimWriteOffset = FirstCompactedTile * 1024;
+
+		// Initialize new tiles
+		for (uint TileIdx = 0; TileIdx < NumTilesToAllocate; ++TileIdx)
+		{
+			const uint CompactedTile = FirstCompactedTile + TileIdx;
+
+			const uint PrimCount = ((TileIdx * 1024) < TotalPrimCount) ? 1024 : (TotalPrimCount - (TileIdx * 1024));
+			StoreCompactedVisTileData(CompactedTile, VT_PrimCount, PrimCount);
+			StoreCompactedVisTileData(CompactedTile, VT_Coord, PackedCoord);
+		}
+	}
+
+	GroupMemoryBarrierWithGroupSync();
+
+	// Copy PrimIDs to compacted memory
+	{
+		uint CurrentWriteOffset = group_PrimWriteOffset;
+
+		// First process the LDS list of tiles
+		const uint NumInputTiles = min(group_NumTiles, 1024);
+		for (uint LDSIdx = 0; LDSIdx < NumInputTiles; ++LDSIdx)
+		{
+			const uint TileIdx = group_TilesToCompact[LDSIdx];
+
+			const uint TilePrimOffset = TileIdx * 1024;
+			const uint TilePrimCount = LoadVisTileData(TileIdx, VT_PrimCount);
+
+			if (GroupThreadID < TilePrimCount)
+			{
+				OutCompactedVisTilePrims[CurrentWriteOffset + GroupThreadID] = VisTilePrims[TilePrimOffset + GroupThreadID];
+			}
+
+			CurrentWriteOffset += TilePrimCount;
+		}
+
+		// Check any remaning tiles
+		if (group_NumTiles > 1024)
+		{
+			for (uint TileIdx = group_TilesToCompact[1023] + 1; TileIdx < NumTiles; ++TileIdx)
+			{
+				const uint TilePackedCoord = LoadVisTileData(TileIdx, VT_Coord);
+				if (PackedCoord == TilePackedCoord)
+				{
+					const uint TilePrimOffset = TileIdx * 1024;
+					const uint TilePrimCount = LoadVisTileData(TileIdx, VT_PrimCount);
+
+					if (GroupThreadID < TilePrimCount)
+					{
+						OutCompactedVisTilePrims[CurrentWriteOffset + GroupThreadID] = VisTilePrims[TilePrimOffset + GroupThreadID];
+					}
+
+					CurrentWriteOffset += TilePrimCount;
+				}
+			}
+		}
+	}
+}
+
+#endif // SHADER_RASTERCOMPUTE_COMPACTION
+
+///////////////////////////////////////////////////////////////////////////
+
 #if SHADER_RASTERCOMPUTE_RASTER
 
 // Simple rasterization algorithm that lerps between line endpoints. Is currently more robust than the Wu algorithm
diff --git a/Engine/Shaders/Private/HairStrands/HairStrandsVisibilityRasterComputeNaive.usf b/Engine/Shaders/Private/HairStrands/HairStrandsVisibilityRasterComputeNaive.usf
new file mode 100644
--- /dev/null
+++ b/Engine/Shaders/Private/HairStrands/HairStrandsVisibilityRasterComputeNaive.usf
@@ -0,0 +1,173 @@
+// Copyright Epic Games, Inc. All Rights Reserved.
+
+#define HAIR_STRANDS_PARAMETERS 1
+
+#include "../Common.ush"
+#include "HairStrandsVisibilityCommon.ush"
+#include "HairStrandsVertexFactoryCommon.ush"
+
+RWTexture2D<uint>	OutHairCountTexture;
+RWTexture2D<uint>	OutDepthCovTexture;
+RWTexture2D<uint>	OutPrimMatTexture;
+Texture2D<float>	SceneDepthTexture;
+float2				OutputResolution;
+uint				HairMaterialId;
+uint				ControlPointStart;
+uint				ControlPointCount;
+float				SampleWeight;
+
+// Line clipping based on "CLIPPING USING HOMOGENEOUS COORDINATES" by Blinn et al.
+bool BlinnLineClipping(inout float4 P0, inout float4 P1)
+{
+	float2 T = float2(0.0f, 1.0f);
+	bool bIsRemoved = P0.w < 0.0f && P1.w < 0.0f; // Both points behind near plane
+
+	bool bSign = false;
+
+	UNROLL
+	for (uint PlaneIdx = 0; PlaneIdx < 6; ++PlaneIdx)
+	{
+		// Compute boundary coordinates of both points (w+x, w-x, w+y, w-y, z, w-z)
+		bSign = !bSign;
+		const uint CompIdx = PlaneIdx / 2;
+		const float Sign = bSign ? 1.0f : -1.0f;
+		const float WFactor = PlaneIdx != 4 ? 1.0f : 0.0f;
+		const float2 BC = WFactor * float2(P0.w, P1.w) + Sign * float2(P0[CompIdx], P1[CompIdx]);
+
+		float Num = BC.x;
+		float Denom = BC.x - BC.y;
+		bIsRemoved = bIsRemoved || (BC.x < 0.0f && BC.y < 0.0f); // Both points outside the plane
+		float Alpha = Num / Denom;
+
+		// If the denominator is negative, P0 has a smaller boundary coordinate than P1, so we can assume
+		// that P1 is inside the plane (or bIsRemoved is true), so we need to update the alpha for P0.
+		// The reverse is true if the denominator is positive.
+		if (Denom < 0.0f)
+		{
+			T.x = max(T.x, Alpha);
+		}
+		else
+		{
+			T.y = min(T.y, Alpha);
+		}
+	}
+
+	if (!bIsRemoved)
+	{
+		const float4 P0Clipped = lerp(P0, P1, T.x);
+		const float4 P1Clipped = lerp(P0, P1, T.y);
+		P0 = P0Clipped;
+		P1 = P1Clipped;
+	}
+
+	return !bIsRemoved;
+}
+
+float3 NDCToPixelCoord(float4 InDC)
+{
+	const float3 NDC = InDC.xyz / InDC.w;
+	float2 UV = NDC.xy * ResolvedView.ScreenPositionScaleBias.xy + ResolvedView.ScreenPositionScaleBias.wz;
+	return float3(UV * OutputResolution, NDC.z);
+}
+
+void PlotInternal(int2 Coord, float AntiAliasingFactor, float4 P0, float4 P1, float Rad0, float Rad1, float SegmentLenSqRcp, uint PrimMatID)
+{
+	// Project P onto line segment and compute the lerp alpha between P0 and P1
+	// Simplification of:
+	// A = P - P0
+	// B = P1 - P0
+	// Alpha = dot(A, B) / dot(B, B)
+	const float2 P = Coord + 0.5f;
+	const float Alpha = saturate(dot(P - P0.xy, P1.xy - P0.xy) * SegmentLenSqRcp);
+	const float Depth = lerp(P0.z, P1.z, Alpha);
+	const uint PackedDepthCov = PackHairVisDepthCoverage(Depth, 1.0f);
+
+	// Write Depth + PrimMatID if depth test against hair depths is passed
+	uint OldValue;
+	InterlockedMax(OutDepthCovTexture[Coord], PackedDepthCov, OldValue);
+	if (PackedDepthCov > OldValue)
+	{
+		OutPrimMatTexture[Coord] = PrimMatID;
+	}
+
+	// Add hair count if depth test against scene depth is passed
+	if (PackedDepthCov > PackHairVisDepthCoverage(SceneDepthTexture.Load(uint3(Coord, 0)), 1.0f))
+	{
+		// Alpha value for perspective correct interpolation. We store the reciprocal of w in the w compoennt,
+		// so this is a simplification of:
+		// (Alpha / w1) / ((1 - Alpha) / w0 + Alpha / w1)
+		const float LerpedRcpW = lerp(P0.w, P1.w, Alpha);
+		const float PerspectiveAlpha = (Alpha * P1.w) / LerpedRcpW;
+		// Divide by W to make thickness dependent on screen space depth? This division was kept from the previous line rasterization algorithm.
+		const float Rad = lerp(Rad0, Rad1, PerspectiveAlpha) * LerpedRcpW;
+
+		InterlockedAdd(OutHairCountTexture[Coord], min(Rad, 0.5f) * 2.0f * 1000.0f * SampleWeight * AntiAliasingFactor);
+	}
+}
+
+[numthreads(64, 1, 1)]
+void CSMain(uint ThreadId : SV_DispatchThreadID)
+{
+	if (ThreadId >= ControlPointCount)
+	{
+		return;
+	}
+
+	ResolvedView = ResolveView();
+
+	const uint PrimID = ControlPointStart + ThreadId;
+
+	const float3 PositionOffset = HairStrandsVF_PositionOffsetBuffer[0].xyz;
+
+	// Fetch both control points
+	const FHairControlPoint CP0 = UnpackHairControlPoint(
+		HairStrandsVF_PositionBuffer[PrimID],
+		PositionOffset,
+		HairStrandsVF_Radius,
+		HairStrandsVF_RootScale,
+		HairStrandsVF_TipScale);
+
+	if (CP0.Type == HAIR_CONTROLPOINT_END)
+	{
+		return;
+	}
+
+	const FHairControlPoint CP1 = UnpackHairControlPoint(
+		HairStrandsVF_PositionBuffer[PrimID + 1],
+		PositionOffset,
+		HairStrandsVF_Radius,
+		HairStrandsVF_RootScale,
+		HairStrandsVF_TipScale);
+
+	// Transform to world space
+	const float3 CP0WP = mul(float4(CP0.Position, 1.0f), HairStrandsVF_LocalToWorldPrimitiveTransform).xyz;
+	const float3 CP1WP = mul(float4(CP1.Position, 1.0f), HairStrandsVF_LocalToWorldPrimitiveTransform).xyz;
+	float4 SP0 = mul(float4(CP0WP, 1.0f), LWCHackToFloat(PrimaryView.WorldToClip));
+	float4 SP1 = mul(float4(CP1WP, 1.0f), LWCHackToFloat(PrimaryView.WorldToClip));
+
+	if (!BlinnLineClipping(SP0, SP1))
+	{
+		return;
+	}
+
+	SP0 = float4(NDCToPixelCoord(SP0), 1.0f / SP0.w);
+	SP1 = float4(NDCToPixelCoord(SP1), 1.0f / SP1.w);
+
+	const float SegmentLenSqRcp = 1.0f / dot(SP1.xy - SP0.xy, SP1.xy - SP0.xy);
+
+	const float2 DeltaAbs = abs(SP1.xy - SP0.xy);
+	const uint NumSteps = (uint)ceil(max(DeltaAbs.x, DeltaAbs.y)) + 1;
+	const float RcpNumSteps = NumSteps > 1 ? (1.0f / (NumSteps - 1)) : 1.0f;
+	const bool bIsSteep = DeltaAbs.y > DeltaAbs.x;
+	const uint PrimMatID = PackHairVisPrimitiveMaterialId(PrimID, HairMaterialId);
+
+	LOOP
+	for (int J = 0; J < NumSteps; ++J)
+	{
+		const float Alpha = J * RcpNumSteps;
+		const float4 SP = lerp(SP0, SP1, Alpha);
+
+		const float AntiAliasingFactor = 1.0f;
+		PlotInternal(SP.xy, AntiAliasingFactor, SP0, SP1, CP0.WorldRadius * 2000.0f, CP1.WorldRadius * 2000.0f, SegmentLenSqRcp, PrimMatID);
+	}
+}
\ No newline at end of file
diff --git a/Engine/Shaders/Private/HairStrands/HairStrandsVisibilityRasterHW.usf b/Engine/Shaders/Private/HairStrands/HairStrandsVisibilityRasterHW.usf
new file mode 100644
--- /dev/null
+++ b/Engine/Shaders/Private/HairStrands/HairStrandsVisibilityRasterHW.usf
@@ -0,0 +1,119 @@
+// Copyright Epic Games, Inc. All Rights Reserved.
+
+#define HAIR_STRANDS_PARAMETERS 1
+
+#include "../Common.ush"
+#include "HairStrandsVisibilityCommon.ush"
+#include "HairStrandsVertexFactoryCommon.ush"
+
+#if SHADER_RASTER_HW_VS
+
+Buffer<uint> 	HWRasterPrimIDs;
+float3			ViewDir;
+uint			HairMaterialId;
+float3			CameraOrigin;
+float			RadiusAtDepth1;
+uint			ControlPointStart;
+uint			bIsOrthoView;
+
+
+struct VSInput
+{
+	uint VertexID : SV_VertexID;
+	uint InstanceID : SV_InstanceID;
+};
+
+struct VSOutput
+{
+	float4 Position : SV_Position;
+	float Radius : RADIUS;
+	nointerpolation uint PrimMatID : PRIM_MAT_ID;
+};
+
+VSOutput VSMain(VSInput Input)
+{
+	const uint PrimID = ControlPointStart + Input.InstanceID;// HWRasterPrimIDs[Input.InstanceID];
+
+	const float3 PositionOffset = HairStrandsVF_PositionOffsetBuffer[0].xyz;
+
+	// Fetch both control points
+	const FHairControlPoint CP0 = UnpackHairControlPoint(
+		HairStrandsVF_PositionBuffer[PrimID],
+		PositionOffset,
+		HairStrandsVF_Radius,
+		HairStrandsVF_RootScale,
+		HairStrandsVF_TipScale);
+
+	// Kill quad
+	if (CP0.Type == HAIR_CONTROLPOINT_END)
+	{
+		VSOutput Output = (VSOutput)0;
+		Output.Position = 0.0f / 0.0f;
+		return Output;
+	}
+
+	const FHairControlPoint CP1 = UnpackHairControlPoint(
+		HairStrandsVF_PositionBuffer[PrimID + 1],
+		PositionOffset,
+		HairStrandsVF_Radius,
+		HairStrandsVF_RootScale,
+		HairStrandsVF_TipScale);
+
+	// Transform to world space
+	const float3 CP0WP = mul(float4(CP0.Position, 1.0f), HairStrandsVF_LocalToWorldPrimitiveTransform).xyz;
+	const float3 CP1WP = mul(float4(CP1.Position, 1.0f), HairStrandsVF_LocalToWorldPrimitiveTransform).xyz;
+	
+	const float3 CPWP = Input.VertexID < 2 ? CP0WP : CP1WP;
+	const float CPWorldRadius = Input.VertexID < 2 ? CP0.WorldRadius : CP1.WorldRadius;
+
+	// Minimal radius to snap the strand to a sample/pixel center (to avoid aliasing)
+	const float DistanceToCamera = length(CameraOrigin - CPWP);
+	const float MinStrandHairRadius = bIsOrthoView ? RadiusAtDepth1 : DistanceToCamera * RadiusAtDepth1;
+
+	const float3 Tangent = normalize(CP1WP - CP0WP);
+	const float3 Right = normalize(cross(Tangent, ViewDir));
+	
+	const float3 WP = CPWP
+		+ (Input.VertexID == 1 || Input.VertexID == 3 ? Right : -Right)
+		* max(CPWorldRadius, MinStrandHairRadius);
+
+	VSOutput Output = (VSOutput)0;
+	
+	Output.Position = mul(float4(WP, 1.0f), LWCHackToFloat(PrimaryView.WorldToClip));
+	Output.Radius = CPWorldRadius * 2000.0f;
+	Output.PrimMatID = PackHairVisPrimitiveMaterialId(PrimID, HairMaterialId);
+
+	return Output;
+}
+
+#endif // SHADER_RASTER_HW_VS
+
+#if SHADER_RASTER_HW_PS
+
+RWTexture2D<uint>	OutHairCountTexture;
+RWTexture2D<uint>	OutDepthCovTexture;
+RWTexture2D<uint>	OutPrimMatTexture;
+float				SampleWeight;
+
+[earlydepthstencil]
+void PSMain(
+	float4 Position : SV_Position,
+	float Radius : RADIUS,
+	nointerpolation uint PrimMatID : PRIM_MAT_ID
+)
+{
+	const uint2 Coord = Position.xy;
+	const uint PackedDepthCov = PackHairVisDepthCoverage(Position.z, 1.0f);
+
+	InterlockedAdd(OutHairCountTexture[Coord], min(Radius / Position.w, 0.5f) * 2.0f * 1000.0f * SampleWeight);
+
+	// Write Depth + PrimMatID if depth test against hair depths is passed
+	uint OldValue;
+	InterlockedMax(OutDepthCovTexture[Coord], PackedDepthCov, OldValue);
+	if (PackedDepthCov > OldValue)
+	{
+		OutPrimMatTexture[Coord] = PrimMatID;
+	}
+}
+
+#endif // SHADER_RASTER_HW_PS
\ No newline at end of file
diff --git a/Engine/Source/Runtime/Renderer/Private/HairStrands/HairStrandsVisibility.cpp b/Engine/Source/Runtime/Renderer/Private/HairStrands/HairStrandsVisibility.cpp
--- a/Engine/Source/Runtime/Renderer/Private/HairStrands/HairStrandsVisibility.cpp
+++ b/Engine/Source/Runtime/Renderer/Private/HairStrands/HairStrandsVisibility.cpp
@@ -99,6 +99,15 @@ static FAutoConsoleVariableRef CVarHairStrands_InvalidationPosition_Debug(TEXT("
 static float GHairStrands_Selection_CoverageThreshold = 0.0f;
 static FAutoConsoleVariableRef CVarHairStrands_Selection_CoverageThreshold(TEXT("r.HairStrands.Selection.CoverageThreshold"), GHairStrands_Selection_CoverageThreshold, TEXT("Coverage threshold for making hair strands outline selection finer"));
 
+static int32 GHairStrandsUseHWRaster = 0;
+static FAutoConsoleVariableRef CVarHairStrandsUseHWRaster(TEXT("r.HairStrands.Visibility.UseHWRaster"), GHairStrandsUseHWRaster, TEXT("Toggles the hardware rasterizer for hair strands visibility rendering."));
+
+static int32 GHairStrandsUseNaiveSWRaster = 0;
+static FAutoConsoleVariableRef CVarHairStrandsUseNaiveSWRaster(TEXT("r.HairStrands.Visibility.UseNaiveSWRaster"), GHairStrandsUseNaiveSWRaster, TEXT("Toggles a naive version of the software rasterizer for hair strands visibility rendering."));
+
+static int32 GHairStrandsTileCompaction = 0;
+static FAutoConsoleVariableRef CVarHairStrandsTileCompaction(TEXT("r.HairStrands.Visibility.TileCompaction"), GHairStrandsTileCompaction, TEXT("Enables a compaction pass to run on the output of the binning pass of the hair software rasterizer."));
+
 /////////////////////////////////////////////////////////////////////////////////////////
 
 namespace HairStrandsVisibilityInternal
@@ -2850,6 +2859,39 @@ IMPLEMENT_GLOBAL_SHADER(FVisiblityRasterComputeBinningCS, "/Engine/Private/HairS
 
 /////////////////////////////////////////////////////////////////////////////////////////
 
+class FVisiblityRasterComputeCompactionCS : public FGlobalShader
+{
+	DECLARE_GLOBAL_SHADER(FVisiblityRasterComputeCompactionCS);
+	SHADER_USE_PARAMETER_STRUCT(FVisiblityRasterComputeCompactionCS, FGlobalShader);
+
+	class FGroupSize : SHADER_PERMUTATION_SPARSE_INT("PERMUTATION_GROUP_SIZE", 32, 64);
+	using FPermutationDomain = TShaderPermutationDomain<FGroupSize>;
+
+	BEGIN_SHADER_PARAMETER_STRUCT(FParameters, )
+		SHADER_PARAMETER_STRUCT_INCLUDE(FRasterComputeCommonParameters, Common)
+		SHADER_PARAMETER_STRUCT_INCLUDE(FHairStrandsInstanceParameters, HairInstance)
+		SHADER_PARAMETER_RDG_BUFFER_SRV(ByteAddressBuffer, VisTileData)
+		SHADER_PARAMETER_RDG_BUFFER_SRV(Buffer, VisTilePrims)
+		SHADER_PARAMETER_RDG_BUFFER_SRV(Buffer, VisTileArgs)
+		SHADER_PARAMETER_RDG_BUFFER_UAV(RWByteAddressBuffer, OutCompactedVisTileData)
+		SHADER_PARAMETER_RDG_BUFFER_UAV(RWBuffer, OutCompactedVisTilePrims)
+		SHADER_PARAMETER_RDG_BUFFER_UAV(RWBuffer, OutCompactedVisTileArgs)
+		SHADER_PARAMETER_STRUCT_REF(FViewUniformShaderParameters, ViewUniformBuffer)
+		END_SHADER_PARAMETER_STRUCT()
+
+public:
+	static bool ShouldCompilePermutation(const FGlobalShaderPermutationParameters &Parameters) { return IsHairStrandsSupported(EHairStrandsShaderType::Strands, Parameters.Platform); }
+	static void ModifyCompilationEnvironment(const FGlobalShaderPermutationParameters &Parameters, FShaderCompilerEnvironment &OutEnvironment)
+	{
+		FGlobalShader::ModifyCompilationEnvironment(Parameters, OutEnvironment);
+		OutEnvironment.SetDefine(TEXT("SHADER_RASTERCOMPUTE_COMPACTION"), 1);
+	}
+};
+
+IMPLEMENT_GLOBAL_SHADER(FVisiblityRasterComputeCompactionCS, "/Engine/Private/HairStrands/HairStrandsVisibilityRasterCompute.usf", "CompactionCS", SF_Compute);
+
+/////////////////////////////////////////////////////////////////////////////////////////
+
 class FVisiblityRasterComputeRasterizeCS : public FGlobalShader
 {
 	DECLARE_GLOBAL_SHADER(FVisiblityRasterComputeRasterizeCS);
@@ -2923,6 +2965,102 @@ IMPLEMENT_GLOBAL_SHADER(FVisiblityRasterComputeDebugCS, "/Engine/Private/HairStr
 
 /////////////////////////////////////////////////////////////////////////////////////////
 
+class FVisiblityRasterHWVS : public FGlobalShader
+{
+	DECLARE_GLOBAL_SHADER(FVisiblityRasterHWVS);
+	SHADER_USE_PARAMETER_STRUCT(FVisiblityRasterHWVS, FGlobalShader);
+
+	using FPermutationDomain = TShaderPermutationDomain<>;
+
+	BEGIN_SHADER_PARAMETER_STRUCT(FParameters, )
+		SHADER_PARAMETER_RDG_BUFFER_SRV(Buffer, HWRasterPrimIDs)
+		SHADER_PARAMETER(FVector3f, ViewDir)
+		SHADER_PARAMETER(uint32, HairMaterialId)
+		SHADER_PARAMETER(FVector3f, CameraOrigin)
+		SHADER_PARAMETER(float, RadiusAtDepth1)
+		SHADER_PARAMETER(uint32, ControlPointStart)
+		SHADER_PARAMETER(uint32, bIsOrthoView)
+		SHADER_PARAMETER_STRUCT_INCLUDE(FHairStrandsInstanceParameters, HairInstance)
+		SHADER_PARAMETER_STRUCT_REF(FViewUniformShaderParameters, ViewUniformBuffer)
+	END_SHADER_PARAMETER_STRUCT()
+
+public:
+	static bool ShouldCompilePermutation(const FGlobalShaderPermutationParameters &Parameters) { return IsHairStrandsSupported(EHairStrandsShaderType::Strands, Parameters.Platform); }
+	static void ModifyCompilationEnvironment(const FGlobalShaderPermutationParameters &Parameters, FShaderCompilerEnvironment &OutEnvironment)
+	{
+		FGlobalShader::ModifyCompilationEnvironment(Parameters, OutEnvironment);
+		OutEnvironment.SetDefine(TEXT("SHADER_RASTER_HW_VS"), 1);
+	}
+};
+
+IMPLEMENT_GLOBAL_SHADER(FVisiblityRasterHWVS, "/Engine/Private/HairStrands/HairStrandsVisibilityRasterHW.usf", "VSMain", SF_Vertex);
+
+/////////////////////////////////////////////////////////////////////////////////////////
+
+class FVisiblityRasterHWPS : public FGlobalShader
+{
+	DECLARE_GLOBAL_SHADER(FVisiblityRasterHWPS);
+	SHADER_USE_PARAMETER_STRUCT(FVisiblityRasterHWPS, FGlobalShader);
+
+	using FPermutationDomain = TShaderPermutationDomain<>;
+
+	BEGIN_SHADER_PARAMETER_STRUCT(FParameters, )
+		SHADER_PARAMETER(float, SampleWeight)
+		SHADER_PARAMETER_STRUCT_INCLUDE(FHairStrandsInstanceParameters, HairInstance)
+		SHADER_PARAMETER_STRUCT_REF(FViewUniformShaderParameters, ViewUniformBuffer)
+		SHADER_PARAMETER_RDG_TEXTURE_UAV(RWTexture2D, OutHairCountTexture)
+		SHADER_PARAMETER_RDG_TEXTURE_UAV(RWTexture2D, OutDepthCovTexture)
+		SHADER_PARAMETER_RDG_TEXTURE_UAV(RWTexture2D, OutPrimMatTexture)
+		RENDER_TARGET_BINDING_SLOTS()
+	END_SHADER_PARAMETER_STRUCT()
+
+public:
+	static bool ShouldCompilePermutation(const FGlobalShaderPermutationParameters &Parameters) { return IsHairStrandsSupported(EHairStrandsShaderType::Strands, Parameters.Platform); }
+	static void ModifyCompilationEnvironment(const FGlobalShaderPermutationParameters &Parameters, FShaderCompilerEnvironment &OutEnvironment)
+	{
+		FGlobalShader::ModifyCompilationEnvironment(Parameters, OutEnvironment);
+		OutEnvironment.SetDefine(TEXT("SHADER_RASTER_HW_PS"), 1);
+	}
+};
+
+IMPLEMENT_GLOBAL_SHADER(FVisiblityRasterHWPS, "/Engine/Private/HairStrands/HairStrandsVisibilityRasterHW.usf", "PSMain", SF_Pixel);
+
+/////////////////////////////////////////////////////////////////////////////////////////
+
+class FVisiblityRasterComputeNaiveCS : public FGlobalShader
+{
+	DECLARE_GLOBAL_SHADER(FVisiblityRasterComputeNaiveCS);
+	SHADER_USE_PARAMETER_STRUCT(FVisiblityRasterComputeNaiveCS, FGlobalShader);
+
+	using FPermutationDomain = TShaderPermutationDomain<>;
+
+	BEGIN_SHADER_PARAMETER_STRUCT(FParameters, )
+		SHADER_PARAMETER(FVector2f, OutputResolution)
+		SHADER_PARAMETER(uint32, HairMaterialId)
+		SHADER_PARAMETER(uint32, ControlPointStart)
+		SHADER_PARAMETER(uint32, ControlPointCount)
+		SHADER_PARAMETER(float, SampleWeight)
+		SHADER_PARAMETER_STRUCT_INCLUDE(FHairStrandsInstanceParameters, HairInstance)
+		SHADER_PARAMETER_STRUCT_REF(FViewUniformShaderParameters, ViewUniformBuffer)
+		SHADER_PARAMETER_RDG_TEXTURE_UAV(RWTexture2D, OutHairCountTexture)
+		SHADER_PARAMETER_RDG_TEXTURE_UAV(RWTexture2D, OutDepthCovTexture)
+		SHADER_PARAMETER_RDG_TEXTURE_UAV(RWTexture2D, OutPrimMatTexture)
+		SHADER_PARAMETER_RDG_TEXTURE(Texture2D, SceneDepthTexture)
+	END_SHADER_PARAMETER_STRUCT()
+
+public:
+	static bool ShouldCompilePermutation(const FGlobalShaderPermutationParameters &Parameters) { return IsHairStrandsSupported(EHairStrandsShaderType::Strands, Parameters.Platform); }
+	static void ModifyCompilationEnvironment(const FGlobalShaderPermutationParameters &Parameters, FShaderCompilerEnvironment &OutEnvironment)
+	{
+		FGlobalShader::ModifyCompilationEnvironment(Parameters, OutEnvironment);
+		OutEnvironment.SetDefine(TEXT("SHADER_RASTERCOMPUTE_RASTER"), 1);
+	}
+};
+
+IMPLEMENT_GLOBAL_SHADER(FVisiblityRasterComputeNaiveCS, "/Engine/Private/HairStrands/HairStrandsVisibilityRasterComputeNaive.usf", "CSMain", SF_Compute);
+
+/////////////////////////////////////////////////////////////////////////////////////////
+
 bool IsHairStrandContinuousDecimationReorderingEnabled();
 
 static FRasterComputeOutput AddVisibilityComputeRasterPass(
@@ -2974,6 +3112,9 @@ static FRasterComputeOutput AddVisibilityComputeRasterPass(
 	FRDGTextureRef OutVisTileDepthGrid = nullptr;
 	FRDGBufferRef OutVisTileArgs = nullptr;
 	FRDGBufferRef OutVisTileData = nullptr;
+	FRDGBufferRef OutCompactedVisTilePrims = nullptr;
+	FRDGBufferRef OutCompactedVisTileArgs = nullptr;
+	FRDGBufferRef OutCompactedVisTileData = nullptr;
 	FRDGBufferRef OutVisTileIndirect = nullptr;
 
 	FRDGBufferUAVRef OutVisTilePrimsUAV = nullptr;
@@ -2981,9 +3122,13 @@ static FRasterComputeOutput AddVisibilityComputeRasterPass(
 	FRDGTextureUAVRef OutVisTileDepthGridUAV = nullptr;
 	FRDGBufferUAVRef OutVisTileArgsUAV = nullptr;
 	FRDGBufferUAVRef OutVisTileDataUAV = nullptr;
+	FRDGBufferUAVRef OutCompactedVisTilePrimsUAV = nullptr;
+	FRDGBufferUAVRef OutCompactedVisTileArgsUAV = nullptr;
+	FRDGBufferUAVRef OutCompactedVisTileDataUAV = nullptr;
 	FRDGBufferUAVRef OutVisTileIndirectUAV = nullptr;
 	
 	OutVisTilePrims = GraphBuilder.CreateBuffer(FRDGBufferDesc::CreateBufferDesc(sizeof(uint32), MaxTiles*1024), TEXT("Hair.OutVisTilePrims"));
+	OutCompactedVisTilePrims = GraphBuilder.CreateBuffer(FRDGBufferDesc::CreateBufferDesc(sizeof(uint32), MaxTiles*1024), TEXT("Hair.OutCompactedVisTilePrims"));
 
 	FRDGTextureDesc DescBinningGrid = FRDGTextureDesc::Create2DArray(TileGridRes, PF_R32_UINT, FClearValueBinding::None, TexCreate_UAV | TexCreate_ShaderResource, NumBinners * 3);
 	OutVisTileBinningGrid = GraphBuilder.CreateTexture(DescBinningGrid, TEXT("Hair.VisTileBinningGrid"));
@@ -2992,18 +3137,23 @@ static FRasterComputeOutput AddVisibilityComputeRasterPass(
 	OutVisTileDepthGrid = GraphBuilder.CreateTexture(DescDepthGrid, TEXT("Hair.VisTileDepthGrid"));
 
 	OutVisTileArgs = GraphBuilder.CreateBuffer(FRDGBufferDesc::CreateBufferDesc(sizeof(uint32), 4+16+(((TileGridRes.X*TileGridRes.Y)+31)/32) ), TEXT("Hair.OutVisTileArgs"));
+	OutCompactedVisTileArgs = GraphBuilder.CreateBuffer(FRDGBufferDesc::CreateBufferDesc(sizeof(uint32), 4+16+(((TileGridRes.X*TileGridRes.Y)+31)/32) ), TEXT("Hair.OutCompactedVisTileArgs"));
 
 	FRDGBufferDesc DescTileData = FRDGBufferDesc::CreateStructuredDesc(sizeof(uint32), MaxTiles * 4 * 2);
 	DescTileData.Usage = EBufferUsageFlags(DescTileData.Usage | BUF_ByteAddressBuffer);
 	OutVisTileData = GraphBuilder.CreateBuffer(DescTileData, TEXT("Hair.OutVisTileData"));
+	OutCompactedVisTileData = GraphBuilder.CreateBuffer(DescTileData, TEXT("Hair.OutCompactedVisTileData"));
 
 	OutVisTileIndirect = GraphBuilder.CreateBuffer(FRDGBufferDesc::CreateIndirectDesc<FRHIDispatchIndirectParameters>(1), TEXT("Hair.OutVisTileIndirect"));
 
 	OutVisTilePrimsUAV = GraphBuilder.CreateUAV(OutVisTilePrims, PF_R32_UINT);
+	OutCompactedVisTilePrimsUAV = GraphBuilder.CreateUAV(OutCompactedVisTilePrims, PF_R32_UINT);
 	OutVisTileBinningGridUAV = GraphBuilder.CreateUAV(OutVisTileBinningGrid);
 	OutVisTileDepthGridUAV = GraphBuilder.CreateUAV(OutVisTileDepthGrid);
 	OutVisTileArgsUAV = GraphBuilder.CreateUAV(OutVisTileArgs, PF_R32_UINT);
+	OutCompactedVisTileArgsUAV = GraphBuilder.CreateUAV(OutCompactedVisTileArgs, PF_R32_UINT);
 	OutVisTileDataUAV = GraphBuilder.CreateUAV(OutVisTileData);
+	OutCompactedVisTileDataUAV = GraphBuilder.CreateUAV(OutCompactedVisTileData);
 	OutVisTileIndirectUAV = GraphBuilder.CreateUAV(OutVisTileIndirect);
 
 	// Create and set the uniform buffer
@@ -3024,10 +3174,18 @@ static FRasterComputeOutput AddVisibilityComputeRasterPass(
 	BinningPermutationVector_CullingOn.Set<FVisiblityRasterComputeBinningCS::FCulling>(true);
 	TShaderMapRef<FVisiblityRasterComputeBinningCS> ComputeShaderBinning_CullingOn(ViewInfo.ShaderMap, BinningPermutationVector_CullingOn);
 
+	FVisiblityRasterComputeCompactionCS::FPermutationDomain CompactionPermutationVector;
+	CompactionPermutationVector.Set<FVisiblityRasterComputeCompactionCS::FGroupSize>(GroupSize);
+	TShaderMapRef<FVisiblityRasterComputeCompactionCS> ComputeShaderCompaction(ViewInfo.ShaderMap, CompactionPermutationVector);
+
 	FVisiblityRasterComputeRasterizeCS::FPermutationDomain RasterPermutationVector;
 	RasterPermutationVector.Set<FVisiblityRasterComputeRasterizeCS::FGroupSize>(GroupSize);
 	TShaderMapRef<FVisiblityRasterComputeRasterizeCS> ComputeShaderRaster(ViewInfo.ShaderMap, RasterPermutationVector);
 
+	TShaderMapRef<FVisiblityRasterHWVS> VertexShaderRaster(ViewInfo.ShaderMap);
+	TShaderMapRef<FVisiblityRasterHWPS> PixelShaderRaster(ViewInfo.ShaderMap);
+	TShaderMapRef<FVisiblityRasterComputeNaiveCS> ComputeShaderRasterNaive(ViewInfo.ShaderMap);
+
 	// Common parameters
 	FRasterComputeCommonParameters Common;
 	{		
@@ -3047,7 +3205,12 @@ static FRasterComputeOutput AddVisibilityComputeRasterPass(
 		Common.RcpNumBinners = 1.0 / NumBinners;
 	}
 
+	const bool bUseHWRaster = GHairStrandsUseHWRaster != 0;
+	const bool bUseNaiveSWRaster = GHairStrandsUseNaiveSWRaster != 0;
+	const bool bTileCompactionPass = GHairStrandsTileCompaction != 0;
+
 	// Fill in tile depth
+	if (!bUseHWRaster && !bUseNaiveSWRaster)
 	{
 		TShaderMapRef<FVisiblityRasterComputePrepareDepthGridCS> ComputeShaderPrepareDepthGrid(ViewInfo.ShaderMap);
 		FVisiblityRasterComputePrepareDepthGridCS::FParameters* PrepDepthGridParameters = GraphBuilder.AllocParameters<FVisiblityRasterComputePrepareDepthGridCS::FParameters>();
@@ -3061,7 +3224,7 @@ static FRasterComputeOutput AddVisibilityComputeRasterPass(
 
 	const bool bDebugEnabled = GHairVisibilityComputeRaster_Debug > 0;
 	uint32 TotalPrimitiveInfoCount = 0;
-	if (bDebugEnabled)
+	if (bDebugEnabled && !bUseHWRaster && !bUseNaiveSWRaster)
 	{
 		ShaderPrint::SetEnabled(true);
 		ShaderPrint::RequestSpaceForTriangles(2 * NumBinners * Common.TileRes.X * Common.TileRes.Y + 2 * NumBinners * 10);
@@ -3102,73 +3265,169 @@ static FRasterComputeOutput AddVisibilityComputeRasterPass(
 			const uint32 VertexCount = HairGroupPublicData->GetActiveStrandsVertexCount(VFInput.Strands.VertexCount, ScreenSize);
 			const float SampleWeight = HairGroupPublicData->GetActiveStrandsSampleWeight(false, ScreenSize);
 
-			// Reset buffers
+			if (bUseHWRaster)
 			{
-				uint32 IndexGridClearValues[4] = { 0x0,0x0,0x0,0x0 };
-				AddClearUAVPass(GraphBuilder, OutVisTileBinningGridUAV, IndexGridClearValues);
-				AddClearUAVPass(GraphBuilder, OutVisTileArgsUAV, 0u);
+				const FVector ViewDirection = ViewInfo.GetViewDirection();
+				const FVector CameraOrigin = ViewInfo.ViewMatrices.GetViewOrigin();
+				const FIntRect Viewport = ViewInfo.ViewRect;
+				const FIntPoint Resolution = SceneDepthTexture->Desc.Extent;
+
+				FVisiblityRasterHWVS::FParameters *ParametersVS = GraphBuilder.AllocParameters<FVisiblityRasterHWVS::FParameters>();
+				ParametersVS->HWRasterPrimIDs = nullptr;
+				ParametersVS->ViewDir = FVector3f(ViewDirection.X, ViewDirection.Y, ViewDirection.Z);
+				ParametersVS->HairMaterialId = PrimitiveInfo.MaterialId;
+				ParametersVS->CameraOrigin = FVector3f(CameraOrigin.X, CameraOrigin.Y, CameraOrigin.Z);
+				ParametersVS->RadiusAtDepth1 = ComputeMinStrandRadiusAtDepth1(Resolution, ViewInfo.FOV, 1, -1.0f).Primary;
+				ParametersVS->ControlPointStart = VertexStart;
+				ParametersVS->bIsOrthoView = !ViewInfo.IsPerspectiveProjection();
+				ParametersVS->HairInstance = GetHairStrandsInstanceParameters(GraphBuilder, ViewInfo, HairGroupPublicData, bCullingEnable, bForceRegister);
+				ParametersVS->ViewUniformBuffer = ViewUniformShaderParameters;
+
+				FVisiblityRasterHWPS::FParameters *ParametersPS = GraphBuilder.AllocParameters<FVisiblityRasterHWPS::FParameters>();
+				ParametersPS->OutHairCountTexture = HairCountTextureUAV;
+				ParametersPS->OutDepthCovTexture = DepthCovTextureUAV;
+				ParametersPS->OutPrimMatTexture = PrimMatTextureUAV;
+				ParametersPS->SampleWeight = SampleWeight;
+				ParametersPS->HairInstance = GetHairStrandsInstanceParameters(GraphBuilder, ViewInfo, HairGroupPublicData, bCullingEnable, bForceRegister);
+				ParametersPS->ViewUniformBuffer = ViewUniformShaderParameters;
+				ParametersPS->RenderTargets.DepthStencil = FDepthStencilBinding(
+					SceneDepthTexture,
+					ERenderTargetLoadAction::ELoad,
+					ERenderTargetLoadAction::ENoAction,
+					FExclusiveDepthStencil::DepthRead_StencilNop);
+
+				GraphBuilder.AddPass(
+					RDG_EVENT_NAME("HairStrands::VisibilityRasterHW"),
+					ParametersPS,
+					ERDGPassFlags::Raster,
+					[ParametersVS, ParametersPS, VertexShaderRaster, PixelShaderRaster, Viewport, Resolution, VertexCount](FRHICommandList &RHICmdList)
+					{
+						FGraphicsPipelineStateInitializer GraphicsPSOInit;
+						RHICmdList.ApplyCachedRenderTargets(GraphicsPSOInit);
+						GraphicsPSOInit.BlendState = TStaticBlendState<CW_RGBA, BO_Add, BF_One, BF_Zero, BO_Add, BF_One, BF_Zero>::GetRHI();
+						GraphicsPSOInit.RasterizerState = TStaticRasterizerState<>::GetRHI();
+						GraphicsPSOInit.DepthStencilState = TStaticDepthStencilState<false, CF_DepthNear>::GetRHI();
+						GraphicsPSOInit.BoundShaderState.VertexDeclarationRHI = GEmptyVertexDeclaration.VertexDeclarationRHI;
+						GraphicsPSOInit.BoundShaderState.VertexShaderRHI = VertexShaderRaster.GetVertexShader();
+						GraphicsPSOInit.BoundShaderState.PixelShaderRHI = PixelShaderRaster.GetPixelShader();
+						GraphicsPSOInit.PrimitiveType = PT_TriangleList;
+						SetGraphicsPipelineState(RHICmdList, GraphicsPSOInit, 0);
+
+						SetShaderParameters(RHICmdList, VertexShaderRaster, VertexShaderRaster.GetPixelShader(), *ParametersVS);
+						SetShaderParameters(RHICmdList, PixelShaderRaster, PixelShaderRaster.GetPixelShader(), *ParametersPS);
+
+						RHICmdList.SetViewport(Viewport.Min.X, Viewport.Min.Y, 0.0f, Viewport.Max.X, Viewport.Max.Y, 1.0f);
+						RHICmdList.SetStreamSource(0, nullptr, 0);
+						RHICmdList.DrawIndexedPrimitive(GTwoTrianglesIndexBuffer.IndexBufferRHI, 0, 0, 6, 0, 2, VertexCount);
+					});
 			}
-
-			// Binning pass
+			else
 			{
-				FVisiblityRasterComputeBinningCS::FParameters* BinningParameters = GraphBuilder.AllocParameters<FVisiblityRasterComputeBinningCS::FParameters>();
-				BinningParameters->Common = Common;
-				BinningParameters->MacroGroupId = MacroGroup.MacroGroupId;
-				BinningParameters->HairMaterialId = PrimitiveInfo.MaterialId;
-				BinningParameters->ViewUniformBuffer = ViewUniformShaderParameters;
-				BinningParameters->SceneDepthTexture = SceneDepthTexture;
-				BinningParameters->OutVisTilePrims = OutVisTilePrimsUAV;
-				BinningParameters->OutVisTileBinningGrid = OutVisTileBinningGridUAV;  
-				BinningParameters->VisTileDepthGrid = OutVisTileDepthGrid;
-				BinningParameters->OutVisTileArgs = OutVisTileArgsUAV;
-				BinningParameters->OutVisTileData = OutVisTileDataUAV;
-				BinningParameters->HairInstance = GetHairStrandsInstanceParameters(GraphBuilder, ViewInfo, HairGroupPublicData, bCullingEnable, bForceRegister);
-				
-				BinningParameters->VertexCount = VertexCount;
-				BinningParameters->VertexStart = VertexStart;
-				
-				BinningParameters->VisTileBinningGridTex = OutVisTileBinningGrid;
+				if (bUseNaiveSWRaster)
+				{
+					FVisiblityRasterComputeNaiveCS::FParameters *RasterParameters = GraphBuilder.AllocParameters<FVisiblityRasterComputeNaiveCS::FParameters>();
+					RasterParameters->OutputResolution = Common.OutputResolutionf;
+					RasterParameters->HairMaterialId = PrimitiveInfo.MaterialId;
+					RasterParameters->ControlPointStart = VertexStart;
+					RasterParameters->ControlPointCount = VertexCount;
+					RasterParameters->SampleWeight = SampleWeight;
+					RasterParameters->HairInstance = GetHairStrandsInstanceParameters(GraphBuilder, ViewInfo, HairGroupPublicData, bCullingEnable, bForceRegister);
+					RasterParameters->ViewUniformBuffer = ViewUniformShaderParameters;
+					RasterParameters->OutHairCountTexture = HairCountTextureUAV;
+					RasterParameters->OutDepthCovTexture = DepthCovTextureUAV;
+					RasterParameters->OutPrimMatTexture = PrimMatTextureUAV;
+					RasterParameters->SceneDepthTexture = SceneDepthTexture;
+
+					const uint32 NumWorkgroups = (VertexCount + 63) / 64;
+					FComputeShaderUtils::AddPass(GraphBuilder, RDG_EVENT_NAME("HairStrands::VisibilityComputeRasterNaive"), ComputeShaderRasterNaive, RasterParameters, FIntVector(NumWorkgroups, 1, 1));
+				}
+				else
+				{
+					// Reset buffers
+					{
+						uint32 IndexGridClearValues[4] = { 0x0,0x0,0x0,0x0 };
+						AddClearUAVPass(GraphBuilder, OutVisTileBinningGridUAV, IndexGridClearValues);
+						AddClearUAVPass(GraphBuilder, OutVisTileArgsUAV, 0u);
+					}
 
-				FComputeShaderUtils::AddPass(GraphBuilder, RDG_EVENT_NAME("HairStrands::VisibilityComputeRasterBinning(culling=%s)", bCullingEnable ? TEXT("On") : TEXT("Off")), bCullingEnable ? ComputeShaderBinning_CullingOn : ComputeShaderBinning_CullingOff, BinningParameters, FIntVector(NumBinners, 1, 1));
-			}
-			
-			// Raster pass
-			{
-				FVisiblityRasterComputeRasterizeCS::FParameters* RasterParameters = GraphBuilder.AllocParameters<FVisiblityRasterComputeRasterizeCS::FParameters>();
-				RasterParameters->Common = Common;
-				RasterParameters->MacroGroupId = MacroGroup.MacroGroupId;
-				RasterParameters->HairMaterialId = PrimitiveInfo.MaterialId;
-				RasterParameters->ViewUniformBuffer = ViewUniformShaderParameters;
-				RasterParameters->SceneDepthTexture = SceneDepthTexture;
-				RasterParameters->OutHairCountTexture = HairCountTextureUAV;
-				RasterParameters->OutDepthCovTexture = DepthCovTextureUAV;
-				RasterParameters->OutPrimMatTexture = PrimMatTextureUAV;
-				RasterParameters->SampleWeight = SampleWeight;
-				RasterParameters->HairInstance = GetHairStrandsInstanceParameters(GraphBuilder, ViewInfo, HairGroupPublicData, bCullingEnable, bForceRegister);
-				RasterParameters->VisTilePrims = GraphBuilder.CreateSRV(OutVisTilePrims, PF_R32_UINT);
-				RasterParameters->RWVisTileDepthGrid = OutVisTileDepthGridUAV;
-				RasterParameters->VisTileArgs = GraphBuilder.CreateSRV(OutVisTileArgs, PF_R32_UINT);
-				RasterParameters->RWVisTileData = OutVisTileDataUAV;
-				RasterParameters->VisTileIndirect = GraphBuilder.CreateSRV(OutVisTileIndirect);
-				RasterParameters->IndirectBufferArgs = OutVisTileIndirect;
-
-				FComputeShaderUtils::AddPass(GraphBuilder, RDG_EVENT_NAME("HairStrands::VisibilityComputeRasterRaster(tiled)"), ComputeShaderRaster, RasterParameters, FIntVector(NumRasterizers, 1, 1));
-			}
+					// Binning pass
+					{
+						FVisiblityRasterComputeBinningCS::FParameters *BinningParameters = GraphBuilder.AllocParameters<FVisiblityRasterComputeBinningCS::FParameters>();
+						BinningParameters->Common = Common;
+						BinningParameters->MacroGroupId = MacroGroup.MacroGroupId;
+						BinningParameters->HairMaterialId = PrimitiveInfo.MaterialId;
+						BinningParameters->ViewUniformBuffer = ViewUniformShaderParameters;
+						BinningParameters->SceneDepthTexture = SceneDepthTexture;
+						BinningParameters->OutVisTilePrims = OutVisTilePrimsUAV;
+						BinningParameters->OutVisTileBinningGrid = OutVisTileBinningGridUAV;
+						BinningParameters->VisTileDepthGrid = OutVisTileDepthGrid;
+						BinningParameters->OutVisTileArgs = OutVisTileArgsUAV;
+						BinningParameters->OutVisTileData = OutVisTileDataUAV;
+						BinningParameters->HairInstance = GetHairStrandsInstanceParameters(GraphBuilder, ViewInfo, HairGroupPublicData, bCullingEnable, bForceRegister);
+
+						BinningParameters->VertexCount = VertexCount;
+						BinningParameters->VertexStart = VertexStart;
+
+						BinningParameters->VisTileBinningGridTex = OutVisTileBinningGrid;
+
+						FComputeShaderUtils::AddPass(GraphBuilder, RDG_EVENT_NAME("HairStrands::VisibilityComputeRasterBinning(culling=%s)", bCullingEnable ? TEXT("On") : TEXT("Off")), bCullingEnable ? ComputeShaderBinning_CullingOn : ComputeShaderBinning_CullingOff, BinningParameters, FIntVector(NumBinners, 1, 1));
+					}
 
-			if (bDebugEnabled)
-			{
-				FVisiblityRasterComputeDebugCS::FParameters* Parameters = GraphBuilder.AllocParameters<FVisiblityRasterComputeDebugCS::FParameters>();
-				Parameters->MacroGroupId = 0;
-				Parameters->Common = Common;
-				Parameters->PrimitiveInfoIndex = PrimitiveInfoIndex;
-				Parameters->TotalPrimitiveInfoCount = TotalPrimitiveInfoCount;
-				Parameters->VisTileDepthGrid = OutVisTileDepthGrid;
-				Parameters->VisTileBinningGrid = OutVisTileBinningGrid;
-				Parameters->VisTileArgs = GraphBuilder.CreateSRV(OutVisTileArgs, PF_R32_UINT);
-				ShaderPrint::SetParameters(GraphBuilder, ViewInfo.ShaderPrintData, Parameters->ShaderPrintParameters);
-
-				TShaderMapRef<FVisiblityRasterComputeDebugCS> DebugComputeShader(ViewInfo.ShaderMap);
-				FComputeShaderUtils::AddPass(GraphBuilder, RDG_EVENT_NAME("HairStrands::VisibilityComputeRaster(Debug)"), DebugComputeShader, Parameters, FIntVector(FMath::DivideAndRoundUp(Common.TileRes.X, 8), FMath::DivideAndRoundUp(Common.TileRes.Y, 8), 1));
+					if (bTileCompactionPass)
+					{
+						FVisiblityRasterComputeCompactionCS::FParameters *CompactionParameters = GraphBuilder.AllocParameters<FVisiblityRasterComputeCompactionCS::FParameters>();
+						CompactionParameters->Common = Common;
+						CompactionParameters->ViewUniformBuffer = ViewUniformShaderParameters;
+						CompactionParameters->HairInstance = GetHairStrandsInstanceParameters(GraphBuilder, ViewInfo, HairGroupPublicData, bCullingEnable, bForceRegister);
+						CompactionParameters->VisTileData = GraphBuilder.CreateSRV(OutVisTileData);
+						CompactionParameters->VisTilePrims = GraphBuilder.CreateSRV(OutVisTilePrims, PF_R32_UINT);
+						CompactionParameters->VisTileArgs = GraphBuilder.CreateSRV(OutVisTileArgs, PF_R32_UINT);
+						CompactionParameters->OutCompactedVisTileData = OutCompactedVisTileDataUAV;
+						CompactionParameters->OutCompactedVisTilePrims = OutCompactedVisTilePrimsUAV;
+						CompactionParameters->OutCompactedVisTileArgs = OutCompactedVisTileArgsUAV;
+
+						FComputeShaderUtils::AddPass(GraphBuilder, RDG_EVENT_NAME("HairStrands::VisibilityComputeRasterCompaction"), ComputeShaderCompaction, CompactionParameters, FIntVector(TileGridRes.X, TileGridRes.Y, 1));
+					}
+
+					// Raster pass
+					{
+						FVisiblityRasterComputeRasterizeCS::FParameters *RasterParameters = GraphBuilder.AllocParameters<FVisiblityRasterComputeRasterizeCS::FParameters>();
+						RasterParameters->Common = Common;
+						RasterParameters->MacroGroupId = MacroGroup.MacroGroupId;
+						RasterParameters->HairMaterialId = PrimitiveInfo.MaterialId;
+						RasterParameters->ViewUniformBuffer = ViewUniformShaderParameters;
+						RasterParameters->SceneDepthTexture = SceneDepthTexture;
+						RasterParameters->OutHairCountTexture = HairCountTextureUAV;
+						RasterParameters->OutDepthCovTexture = DepthCovTextureUAV;
+						RasterParameters->OutPrimMatTexture = PrimMatTextureUAV;
+						RasterParameters->SampleWeight = SampleWeight;
+						RasterParameters->HairInstance = GetHairStrandsInstanceParameters(GraphBuilder, ViewInfo, HairGroupPublicData, bCullingEnable, bForceRegister);
+						RasterParameters->VisTilePrims = GraphBuilder.CreateSRV(bTileCompactionPass ? OutCompactedVisTilePrims : OutVisTilePrims, PF_R32_UINT);
+						RasterParameters->RWVisTileDepthGrid = OutVisTileDepthGridUAV;
+						RasterParameters->VisTileArgs = GraphBuilder.CreateSRV(bTileCompactionPass ? OutCompactedVisTileArgs : OutVisTileArgs, PF_R32_UINT);
+						RasterParameters->RWVisTileData = bTileCompactionPass ? OutCompactedVisTileDataUAV : OutVisTileDataUAV;
+						RasterParameters->VisTileIndirect = GraphBuilder.CreateSRV(OutVisTileIndirect);
+						RasterParameters->IndirectBufferArgs = OutVisTileIndirect;
+
+						FComputeShaderUtils::AddPass(GraphBuilder, RDG_EVENT_NAME("HairStrands::VisibilityComputeRasterRaster(tiled)"), ComputeShaderRaster, RasterParameters, FIntVector(NumRasterizers, 1, 1));
+					}
+
+					if (bDebugEnabled)
+					{
+						FVisiblityRasterComputeDebugCS::FParameters *Parameters = GraphBuilder.AllocParameters<FVisiblityRasterComputeDebugCS::FParameters>();
+						Parameters->MacroGroupId = 0;
+						Parameters->Common = Common;
+						Parameters->PrimitiveInfoIndex = PrimitiveInfoIndex;
+						Parameters->TotalPrimitiveInfoCount = TotalPrimitiveInfoCount;
+						Parameters->VisTileDepthGrid = OutVisTileDepthGrid;
+						Parameters->VisTileBinningGrid = OutVisTileBinningGrid;
+						Parameters->VisTileArgs = GraphBuilder.CreateSRV(OutVisTileArgs, PF_R32_UINT);
+						ShaderPrint::SetParameters(GraphBuilder, ViewInfo.ShaderPrintData, Parameters->ShaderPrintParameters);
+
+						TShaderMapRef<FVisiblityRasterComputeDebugCS> DebugComputeShader(ViewInfo.ShaderMap);
+						FComputeShaderUtils::AddPass(GraphBuilder, RDG_EVENT_NAME("HairStrands::VisibilityComputeRaster(Debug)"), DebugComputeShader, Parameters, FIntVector(FMath::DivideAndRoundUp(Common.TileRes.X, 8), FMath::DivideAndRoundUp(Common.TileRes.Y, 8), 1));
+					}
+				}
 			}
 
 			++PrimitiveInfoIndex;
