Fixed the bug that Gizmos are affecting GI.

Add translucent editor primitives support. Top view mesh elements with translucent materials will go through editor primitive rendering instead of the standard translucency pass.

diff --git a/Engine/Source/Runtime/Renderer/Private/PostProcess/PostProcessCompositeEditorPrimitives.cpp b/Engine/Source/Runtime/Renderer/Private/PostProcess/PostProcessCompositeEditorPrimitives.cpp
--- a/Engine/Source/Runtime/Renderer/Private/PostProcess/PostProcessCompositeEditorPrimitives.cpp
+++ b/Engine/Source/Runtime/Renderer/Private/PostProcess/PostProcessCompositeEditorPrimitives.cpp
@@ -167,6 +167,44 @@ void RenderEditorPrimitives(FRHICommandList& RHICmdList, const FViewInfo& View,
 	View.BatchedViewElements.Draw(RHICmdList, DrawRenderState, FeatureLevel, View, false, 1.0f);
 }
 
+void RenderForegroundTranslucentEditorPrimitives(FRHICommandList& RHICmdList, const FViewInfo& View, FMeshPassProcessorRenderState& DrawRenderState, FInstanceCullingManager& InstanceCullingManager)
+{
+	const auto FeatureLevel = View.GetFeatureLevel();
+	const auto ShaderPlatform = GShaderPlatformForFeatureLevel[FeatureLevel];
+
+	// Force all translucent editor primitives to standard translucent rendering
+	const ETranslucencyPass::Type TranslucencyPass = ETranslucencyPass::TPT_StandardTranslucency;
+
+	if (TranslucencyPass == ETranslucencyPass::TPT_StandardTranslucency)
+	{
+		DrawRenderState.SetDepthStencilState(TStaticDepthStencilState<false, CF_DepthNearOrEqual>::GetRHI());
+	}
+
+	View.EditorSimpleElementCollector.DrawBatchedElements(RHICmdList, DrawRenderState, View, EBlendModeFilter::Translucent, SDPG_Foreground);
+
+	DrawDynamicMeshPass(View, RHICmdList,
+		[&View, &DrawRenderState](FDynamicPassMeshDrawListContext* DynamicMeshPassContext)
+		{
+			FEditorPrimitivesBasePassMeshProcessor PassMeshProcessor(
+				View.Family->Scene->GetRenderScene(),
+				View.GetFeatureLevel(),
+				&View,
+				DrawRenderState,
+				true,
+				DynamicMeshPassContext);
+
+			const uint64 DefaultBatchElementMask = ~0ull;
+
+			for (int32 MeshIndex = 0; MeshIndex < View.TopViewMeshElements.Num(); MeshIndex++)
+			{
+				const FMeshBatch& MeshBatch = View.TopViewMeshElements[MeshIndex];
+				PassMeshProcessor.AddMeshBatch(MeshBatch, DefaultBatchElementMask, nullptr);
+			}
+		});
+	
+	View.TopBatchedViewElements.Draw(RHICmdList, DrawRenderState, FeatureLevel, View, false, 1.0f, EBlendModeFilter::Translucent);
+}
+
 void RenderForegroundEditorPrimitives(FRHICommandList& RHICmdList, const FViewInfo& View, FMeshPassProcessorRenderState& DrawRenderState, FInstanceCullingManager& InstanceCullingManager)
 {
 	const auto FeatureLevel = View.GetFeatureLevel();
@@ -263,6 +301,7 @@ BEGIN_SHADER_PARAMETER_STRUCT(FEditorPrimitivesPassParameters, )
 	SHADER_PARAMETER_STRUCT_INCLUDE(FViewShaderParameters, View)
 	SHADER_PARAMETER_STRUCT_REF(FReflectionCaptureShaderData, ReflectionCapture)
 	SHADER_PARAMETER_RDG_UNIFORM_BUFFER(FOpaqueBasePassUniformParameters, BasePass)
+	SHADER_PARAMETER_RDG_UNIFORM_BUFFER(FTranslucentBasePassUniformParameters, TranslucentBasePass)
 	SHADER_PARAMETER_RDG_UNIFORM_BUFFER(FMobileBasePassUniformParameters, MobileBasePass)
 	SHADER_PARAMETER_RDG_UNIFORM_BUFFER(FInstanceCullingGlobalUniforms, InstanceCulling)
 	RENDER_TARGET_BINDING_SLOTS()
@@ -456,7 +495,7 @@ FScreenPassTexture AddEditorPrimitivePass(
 			}
 		}
 
-		// Draws the editors primitives
+		// Draws the editors opaque primitives
 		{
 			FEditorPrimitivesPassParameters* PassParameters = GraphBuilder.AllocParameters<FEditorPrimitivesPassParameters>();
 			PassParameters->View = EditorView->GetShaderParameters();
@@ -559,6 +598,52 @@ FScreenPassTexture AddEditorPrimitivePass(
 			Output.ViewRect);
 	}
 
+	// Draws the editor translucent primitives on top of the opaque scene primitives
+	if (View.Family->EngineShowFlags.CompositeEditorPrimitives && View.bHasTranslucentViewMeshElements)
+	{
+		FEditorPrimitivesPassParameters* PassParameters = GraphBuilder.AllocParameters<FEditorPrimitivesPassParameters>();
+		PassParameters->View = EditorView->GetShaderParameters();
+		PassParameters->ReflectionCapture = View.ReflectionCaptureUniformBuffer;
+		PassParameters->InstanceCulling = InstanceCullingManager.GetDummyInstanceCullingUniformBuffer();
+		PassParameters->RenderTargets[0] = Output.GetRenderTargetBinding();
+
+		const FEditorPrimitiveInputs::EBasePassType BasePassType = Inputs.BasePassType;
+
+		if (BasePassType == FEditorPrimitiveInputs::EBasePassType::Deferred)
+		{
+			PassParameters->TranslucentBasePass = CreateTranslucentBasePassUniformBuffer(GraphBuilder, nullptr, *EditorView, 0);
+		}
+		else
+		{
+			PassParameters->MobileBasePass = CreateMobileBasePassUniformBuffer(GraphBuilder, *EditorView, EMobileBasePass::Translucent, EMobileSceneTextureSetupMode::None);
+		}
+
+		GraphBuilder.AddPass(
+			RDG_EVENT_NAME("EditorPrimitives Translucent %dx%d MSAA=%d",
+				EditorPrimitivesViewport.Rect.Width(),
+				EditorPrimitivesViewport.Rect.Height(),
+				NumSamples),
+			PassParameters,
+			ERDGPassFlags::Raster,
+			[&View, &InstanceCullingManager, PassParameters, EditorView, EditorPrimitivesViewport, BasePassType, NumSamples](FRHICommandListImmediate& RHICmdList)
+			{
+				RHICmdList.SetViewport(EditorPrimitivesViewport.Rect.Min.X, EditorPrimitivesViewport.Rect.Min.Y, 0.0f, EditorPrimitivesViewport.Rect.Max.X, EditorPrimitivesViewport.Rect.Max.Y, 1.0f);
+
+				FMeshPassProcessorRenderState DrawRenderState;
+				DrawRenderState.SetDepthStencilAccess(FExclusiveDepthStencil::DepthRead_StencilNop);
+				DrawRenderState.SetBlendState(TStaticBlendStateWriteMask<CW_RGBA>::GetRHI());
+
+				// Draw foreground editor primitives.
+				{
+					SCOPED_DRAW_EVENTF(RHICmdList, EditorPrimitives,
+						TEXT("RenderViewEditorForegroundTranslucentPrimitives %dx%d msaa=%d"),
+						EditorPrimitivesViewport.Rect.Width(), EditorPrimitivesViewport.Rect.Height(), NumSamples);
+
+					RenderForegroundTranslucentEditorPrimitives(RHICmdList, *EditorView, DrawRenderState, InstanceCullingManager);
+				}
+			});
+	}
+
 	return MoveTemp(Output);
 }
 
diff --git a/Engine/Source/Runtime/Renderer/Private/TranslucentRendering.cpp b/Engine/Source/Runtime/Renderer/Private/TranslucentRendering.cpp
--- a/Engine/Source/Runtime/Renderer/Private/TranslucentRendering.cpp
+++ b/Engine/Source/Runtime/Renderer/Private/TranslucentRendering.cpp
@@ -1092,6 +1092,7 @@ static void RenderViewTranslucencyInner(
 				});
 			}
 
+			if (!View.Family->EngineShowFlags.CompositeEditorPrimitives)
 			{
 				QUICK_SCOPE_CYCLE_COUNTER(RenderTranslucencyParallel_SDPG_Foreground);
 
