Fixing static analysis warning

diff --git a/Engine/Source/Runtime/Renderer/Private/HeterogeneousVolumes/HeterogeneousVolumesHardwareRayTracing.cpp b/Engine/Source/Runtime/Renderer/Private/HeterogeneousVolumes/HeterogeneousVolumesHardwareRayTracing.cpp
--- a/Engine/Source/Runtime/Renderer/Private/HeterogeneousVolumes/HeterogeneousVolumesHardwareRayTracing.cpp
+++ b/Engine/Source/Runtime/Renderer/Private/HeterogeneousVolumes/HeterogeneousVolumesHardwareRayTracing.cpp
@@ -471,7 +471,7 @@ void RenderSingleScatteringWithPreshadingHardwareRayTracing(
 		FDeferredLightUniformStruct DeferredLightUniform;
 		PassParameters->bApplyEmission = bApplyEmission;
 		PassParameters->bApplyDirectLighting = bApplyDirectLighting;
-		if (PassParameters->bApplyDirectLighting)
+		if (PassParameters->bApplyDirectLighting && (LightSceneInfo != nullptr))
 		{
 			DeferredLightUniform = GetDeferredLightParameters(View, *LightSceneInfo);
 		}
diff --git a/Engine/Source/Runtime/Renderer/Private/HeterogeneousVolumes/HeterogeneousVolumesLiveShadingPipeline.cpp b/Engine/Source/Runtime/Renderer/Private/HeterogeneousVolumes/HeterogeneousVolumesLiveShadingPipeline.cpp
--- a/Engine/Source/Runtime/Renderer/Private/HeterogeneousVolumes/HeterogeneousVolumesLiveShadingPipeline.cpp
+++ b/Engine/Source/Runtime/Renderer/Private/HeterogeneousVolumes/HeterogeneousVolumesLiveShadingPipeline.cpp
@@ -164,7 +164,7 @@ void RenderSingleScatteringWithLiveShading(
 		PassParameters->bApplyEmission = bApplyEmission;
 		PassParameters->bApplyDirectLighting = bApplyDirectLighting;
 		PassParameters->bApplyShadowTransmittance = bApplyShadowTransmittance;
-		if (PassParameters->bApplyDirectLighting)
+		if (PassParameters->bApplyDirectLighting && (LightSceneInfo != nullptr))
 		{
 			DeferredLightUniform = GetDeferredLightParameters(View, *LightSceneInfo);
 		}
diff --git a/Engine/Source/Runtime/Renderer/Private/HeterogeneousVolumes/HeterogeneousVolumesPreshadingPipeline.cpp b/Engine/Source/Runtime/Renderer/Private/HeterogeneousVolumes/HeterogeneousVolumesPreshadingPipeline.cpp
--- a/Engine/Source/Runtime/Renderer/Private/HeterogeneousVolumes/HeterogeneousVolumesPreshadingPipeline.cpp
+++ b/Engine/Source/Runtime/Renderer/Private/HeterogeneousVolumes/HeterogeneousVolumesPreshadingPipeline.cpp
@@ -274,6 +274,13 @@ class FRenderSingleScatteringWithPreshadingCS : public FGlobalShader
 
 IMPLEMENT_MATERIAL_SHADER_TYPE(, FRenderSingleScatteringWithPreshadingCS, TEXT("/Engine/Private/HeterogeneousVolumes/HeterogeneousVolumesPreshadingPipeline.usf"), TEXT("RenderSingleScatteringWithPreshadingCS"), SF_Compute);
 
+void RenderTransmissionWithPreshadingCompute(
+
+)
+{
+
+}
+
 void RenderSingleScatteringWithPreshadingCompute(
 	FRDGBuilder& GraphBuilder,
 	// Scene data
@@ -314,7 +321,7 @@ void RenderSingleScatteringWithPreshadingCompute(
 		FDeferredLightUniformStruct DeferredLightUniform;
 		PassParameters->bApplyEmission = bApplyEmission;
 		PassParameters->bApplyDirectLighting = bApplyDirectLighting;
-		if (PassParameters->bApplyDirectLighting)
+		if (PassParameters->bApplyDirectLighting && (LightSceneInfo != nullptr))
 		{
 			DeferredLightUniform = GetDeferredLightParameters(View, *LightSceneInfo);
 		}
