Use light view culling only on non-static, non-vsm shadow maps

diff --git a/Engine/Source/Runtime/Renderer/Private/ShadowSetup.cpp b/Engine/Source/Runtime/Renderer/Private/ShadowSetup.cpp
--- a/Engine/Source/Runtime/Renderer/Private/ShadowSetup.cpp
+++ b/Engine/Source/Runtime/Renderer/Private/ShadowSetup.cpp
@@ -3661,28 +3661,12 @@ void FSceneRenderer::CreateWholeSceneProjectedShadow(
 					SizeY = ShadowMapSize.Y - ShadowBorder * 2;
 				}
 
-				// Build light-view convex hulls for shadow caster culling if someone is going to need it
-				FLightViewFrustumConvexHulls LightViewFrustumConvexHulls;
-				// Actually we probably want to have fallback even for these? Another performance regresion however (since we'd be adding culling & draw calls for DF lights).
-				if (!ProjectedShadowInitializer.bRayTracedDistanceField)
-				{
-					for (int32 CacheModeIndex = 0; CacheModeIndex < NumShadowMaps; CacheModeIndex++)
-					{
-						if (CacheMode[CacheModeIndex] != SDCM_StaticPrimitivesOnly || bNeedsVirtualShadowMap)
-						{
-							FVector const& LightOrigin = LightSceneInfo->Proxy->GetOrigin();
-							BuildLightViewFrustumConvexHulls(LightOrigin, Views, LightViewFrustumConvexHulls);
-							break;
-						}
-					}
-				}
-
 				// Consolidate repeated operation into a single site.
-				auto AddInteractingPrimitives = [bStaticSceneOnly, &Views = this->Views, FeatureLevel = this->FeatureLevel, &LightViewFrustumConvexHulls]
+				auto AddInteractingPrimitives = [bStaticSceneOnly, &Views = this->Views, FeatureLevel = this->FeatureLevel]
 				(FLightPrimitiveInteraction* InteractionList, 
 					FProjectedShadowInfo* ProjectedShadowInfo, 
 					bool& bContainsNaniteSubjectsOut, 
-					bool bTestConvexHull,
+					const FLightViewFrustumConvexHulls& LightViewFrustumConvexHulls,
 					const TSharedPtr<FVirtualShadowMapPerLightCacheEntry> &VirtualSmCacheEntry
 				)
 				{
@@ -3704,7 +3688,7 @@ void FSceneRenderer::CreateWholeSceneProjectedShadow(
 							{
 								const FPrimitiveSceneInfo* PrimitiveSceneInfo = Interaction->GetPrimitiveSceneInfo();
 								FBoxSphereBounds const& Bounds = PrimitiveSceneInfo->Proxy->GetBounds();
-								if (!bTestConvexHull || IntersectsConvexHulls(LightViewFrustumConvexHulls, Bounds))
+								if (IntersectsConvexHulls(LightViewFrustumConvexHulls, Bounds))
 								{
 									if (ProjectedShadowInfo->AddSubjectPrimitive(Interaction->GetPrimitiveSceneInfo(), Views, false) 
 										&& VirtualSmCacheEntry.IsValid())
@@ -3750,7 +3734,7 @@ void FSceneRenderer::CreateWholeSceneProjectedShadow(
 					ProjectedShadowInfo->MeshSelectionMask = EShadowMeshSelection::VSM;
 
 					bool bContainsNaniteSubjects = false;
-
+					
 					// TODO: Perform no rendering work (command setup etc) if if fully cached, this optimiztion would be nice to get back, however, right now this is tricky
 					//       as the logic to select those that get to re-render must happen when the full set of rendered local lights is known.
 					//       At that point we would need to re-do this logic, which is currently rather tied to this part of the code. 
@@ -3759,8 +3743,9 @@ void FSceneRenderer::CreateWholeSceneProjectedShadow(
 					{
 						// Skip convex hull tests for VSM since this causes artifacts due to caching (potentially check if caching is enabled but may lead to race condition).
 						// The interaction setup has already tested the light bounds (by calling FLightSceneProxy::AffectsBounds)
-						AddInteractingPrimitives(LightSceneInfo->GetDynamicInteractionOftenMovingPrimitiveList(), ProjectedShadowInfo, bContainsNaniteSubjects, false, VirtualSmPerLightCacheEntry);
-						AddInteractingPrimitives(LightSceneInfo->GetDynamicInteractionStaticPrimitiveList(), ProjectedShadowInfo, bContainsNaniteSubjects, false, VirtualSmPerLightCacheEntry);
+						FLightViewFrustumConvexHulls EmptyHull{};
+						AddInteractingPrimitives(LightSceneInfo->GetDynamicInteractionOftenMovingPrimitiveList(), ProjectedShadowInfo, bContainsNaniteSubjects, EmptyHull, VirtualSmPerLightCacheEntry);
+						AddInteractingPrimitives(LightSceneInfo->GetDynamicInteractionStaticPrimitiveList(), ProjectedShadowInfo, bContainsNaniteSubjects, EmptyHull, VirtualSmPerLightCacheEntry);
 						ProjectedShadowInfo->bContainsNaniteSubjects = bContainsNaniteSubjects;
 					}
 					VisibleLightInfo.VirtualShadowMapId = ProjectedShadowInfo->VirtualShadowMaps[0]->ID;
@@ -3801,19 +3786,28 @@ void FSceneRenderer::CreateWholeSceneProjectedShadow(
 						bool bContainsNaniteSubjects = false;
 
 						// Ray traced shadows use the GPU managed distance field object buffers, no CPU culling should be used
+						// Actually we probably want to have fallback even for these? Another performance regresion however (since we'd be adding culling & draw calls for DF lights).
 						if (!ProjectedShadowInfo->bRayTracedDistanceField)
 						{
+							// Build light-view convex hulls for shadow caster culling
+							FLightViewFrustumConvexHulls LightViewFrustumConvexHulls{};
+							if (CacheMode[CacheModeIndex] != SDCM_StaticPrimitivesOnly)
+							{
+								FVector const& LightOrigin = LightSceneInfo->Proxy->GetOrigin();
+								BuildLightViewFrustumConvexHulls(LightOrigin, Views, LightViewFrustumConvexHulls);
+							}
+
 							bool bCastCachedShadowFromMovablePrimitives = GCachedShadowsCastFromMovablePrimitives || LightSceneInfo->Proxy->GetForceCachedShadowsForMovablePrimitives();
 							if (CacheMode[CacheModeIndex] != SDCM_StaticPrimitivesOnly 
 								&& (CacheMode[CacheModeIndex] != SDCM_MovablePrimitivesOnly || bCastCachedShadowFromMovablePrimitives))
 							{
 								// Add all the shadow casting primitives affected by the light to the shadow's subject primitive list.
-								AddInteractingPrimitives(LightSceneInfo->GetDynamicInteractionOftenMovingPrimitiveList(), ProjectedShadowInfo, bContainsNaniteSubjects, true, nullptr);
+								AddInteractingPrimitives(LightSceneInfo->GetDynamicInteractionOftenMovingPrimitiveList(), ProjectedShadowInfo, bContainsNaniteSubjects, LightViewFrustumConvexHulls, nullptr);
 							}
 						
 							if (CacheMode[CacheModeIndex] != SDCM_MovablePrimitivesOnly)
 							{
-								AddInteractingPrimitives(LightSceneInfo->GetDynamicInteractionStaticPrimitiveList(), ProjectedShadowInfo, bContainsNaniteSubjects, true, nullptr);
+								AddInteractingPrimitives(LightSceneInfo->GetDynamicInteractionStaticPrimitiveList(), ProjectedShadowInfo, bContainsNaniteSubjects, LightViewFrustumConvexHulls, nullptr);
 							}
 						}
 						ProjectedShadowInfo->bContainsNaniteSubjects = bContainsNaniteSubjects;
