Use DistanceFieldSampler in the shaders used to update global SDF.
- Remove depedency on ViewUniformBuffer.
- Changed DistanceFieldSampler to use wrap addressing so that it matches the global sampler addressing mode.

diff --git a/Engine/Plugins/FX/Niagara/Source/NiagaraShader/Private/NiagaraDistanceFieldHelper.cpp b/Engine/Plugins/FX/Niagara/Source/NiagaraShader/Private/NiagaraDistanceFieldHelper.cpp
--- a/Engine/Plugins/FX/Niagara/Source/NiagaraShader/Private/NiagaraDistanceFieldHelper.cpp
+++ b/Engine/Plugins/FX/Niagara/Source/NiagaraShader/Private/NiagaraDistanceFieldHelper.cpp
@@ -88,7 +88,7 @@ void FNiagaraDistanceFieldHelper::SetMeshDistanceFieldParameters(FRDGBuilder& Gr
 		AtlasShaderParameters.DistanceFieldIndirection2Table = DefaultVector4;
 		AtlasShaderParameters.DistanceFieldIndirectionAtlas = FRDGSystemTextures::Get(GraphBuilder).VolumetricBlack;
 		AtlasShaderParameters.DistanceFieldBrickTexture = GBlackVolumeTexture->GetTextureRHI();
-		AtlasShaderParameters.DistanceFieldSampler = TStaticSamplerState<SF_Bilinear, AM_Clamp, AM_Clamp, AM_Clamp>::GetRHI();
+		AtlasShaderParameters.DistanceFieldSampler = TStaticSamplerState<SF_Bilinear, AM_Wrap, AM_Wrap, AM_Wrap>::GetRHI();
 		AtlasShaderParameters.DistanceFieldBrickSize = FVector3f::ZeroVector;
 		AtlasShaderParameters.DistanceFieldUniqueDataBrickSize = FVector3f::ZeroVector;
 		AtlasShaderParameters.DistanceFieldBrickAtlasSizeInBricks = FIntVector::ZeroValue;
diff --git a/Engine/Shaders/Private/DistanceField/GlobalDistanceField.usf b/Engine/Shaders/Private/DistanceField/GlobalDistanceField.usf
--- a/Engine/Shaders/Private/DistanceField/GlobalDistanceField.usf
+++ b/Engine/Shaders/Private/DistanceField/GlobalDistanceField.usf
@@ -1,5 +1,7 @@
 // Copyright Epic Games, Inc. All Rights Reserved.
 
+#define USE_DISTANCE_FIELD_SAMPLER 1
+
 #include "../Common.ush"
 #include "../DeferredShadingCommon.ush"
 #include "../DistanceFieldLightingShared.ush"
diff --git a/Engine/Shaders/Private/DistanceField/GlobalDistanceFieldHeightfields.usf b/Engine/Shaders/Private/DistanceField/GlobalDistanceFieldHeightfields.usf
--- a/Engine/Shaders/Private/DistanceField/GlobalDistanceFieldHeightfields.usf
+++ b/Engine/Shaders/Private/DistanceField/GlobalDistanceFieldHeightfields.usf
@@ -1,5 +1,7 @@
 // Copyright Epic Games, Inc. All Rights Reserved.
 
+#define USE_DISTANCE_FIELD_SAMPLER 1
+
 #include "../Common.ush"
 #include "../HeightfieldLightingShared.ush"
 #include "GlobalDistanceFieldShared.ush"
diff --git a/Engine/Shaders/Private/DistanceField/GlobalDistanceFieldMip.usf b/Engine/Shaders/Private/DistanceField/GlobalDistanceFieldMip.usf
--- a/Engine/Shaders/Private/DistanceField/GlobalDistanceFieldMip.usf
+++ b/Engine/Shaders/Private/DistanceField/GlobalDistanceFieldMip.usf
@@ -11,6 +11,7 @@ float MipInfluenceRadius;
 Texture3D<uint> PageTableTexture;
 Texture3D PageAtlasTexture;
 Texture3D PrevMipTexture;
+SamplerState DistanceFieldSampler;
 
 // https://en.wikipedia.org/wiki/Eikonal_equation
 float Eikonal1(float X, float Step)
@@ -58,7 +59,7 @@ float LoadPrevDistanceFieldValue(int3 MipCoord, int3 Offset)
 	if (PageIndex < GLOBAL_DISTANCE_FIELD_INVALID_PAGE_ID)
 	{
 		float3 PageUV = ComputeGlobalDistanceFieldPageUV(ScrolledClipmapUV, PageIndex);
-		DistanceFieldValue = Texture3DSampleLevel(PageAtlasTexture, GlobalTrilinearWrappedSampler, PageUV, 0).x;
+		DistanceFieldValue = Texture3DSampleLevel(PageAtlasTexture, DistanceFieldSampler, PageUV, 0).x;
 		if (DistanceFieldValue < 1.0f)
 		{
 			DistanceFieldValue = DistanceFieldValue * CoarseDistanceFieldValueScale + CoarseDistanceFieldValueBias;
diff --git a/Engine/Shaders/Private/DistanceFieldLightingShared.ush b/Engine/Shaders/Private/DistanceFieldLightingShared.ush
--- a/Engine/Shaders/Private/DistanceFieldLightingShared.ush
+++ b/Engine/Shaders/Private/DistanceFieldLightingShared.ush
@@ -36,6 +36,10 @@
 #define OFFSET_DATA_STRUCT (0)
 #endif
 
+#ifndef USE_DISTANCE_FIELD_SAMPLER
+#define USE_DISTANCE_FIELD_SAMPLER (0)
+#endif
+
 // Size of SDF brick in voxels, including padding
 float3 DistanceFieldBrickSize;
 
@@ -367,7 +371,7 @@ bool ViewFrustumIntersectAABB(float3 BoxCenter, float3 BoxExtent)
 
 float SampleDistanceFieldBrickTexture(float3 UV)
 {
-#if SUPPORTS_INDEPENDENT_SAMPLERS
+#if SUPPORTS_INDEPENDENT_SAMPLERS && !USE_DISTANCE_FIELD_SAMPLER
 	return Texture3DSampleLevel(DistanceFieldBrickTexture, GlobalBilinearWrappedSampler, UV, 0).x;
 #else
 	return Texture3DSampleLevel(DistanceFieldBrickTexture, DistanceFieldSampler, UV, 0).x;
@@ -483,7 +487,7 @@ Texture2D HFVisibilityTexture;
 
 float SampleHeightFieldAtlas(float2 UV)
 {
-#if SUPPORTS_INDEPENDENT_SAMPLERS
+#if SUPPORTS_INDEPENDENT_SAMPLERS && !USE_DISTANCE_FIELD_SAMPLER
 	float4 SampleValue = Texture2DSampleLevel(HeightFieldTexture, GlobalBilinearWrappedSampler, UV, 0);
 #else
 	float4 SampleValue = Texture2DSampleLevel(HeightFieldTexture, DistanceFieldSampler, UV, 0);
@@ -493,7 +497,7 @@ float SampleHeightFieldAtlas(float2 UV)
 
 float SampleHFVisibilityTexture(float2 UV)
 {
-#if SUPPORTS_INDEPENDENT_SAMPLERS
+#if SUPPORTS_INDEPENDENT_SAMPLERS && !USE_DISTANCE_FIELD_SAMPLER
 	return Texture2DSampleLevel(HFVisibilityTexture, GlobalBilinearWrappedSampler, UV, 0).r;
 #else
 	return Texture2DSampleLevel(HFVisibilityTexture, DistanceFieldSampler, UV, 0).r;
diff --git a/Engine/Source/Runtime/Renderer/Private/DistanceFieldObjectManagement.cpp b/Engine/Source/Runtime/Renderer/Private/DistanceFieldObjectManagement.cpp
--- a/Engine/Source/Runtime/Renderer/Private/DistanceFieldObjectManagement.cpp
+++ b/Engine/Source/Runtime/Renderer/Private/DistanceFieldObjectManagement.cpp
@@ -151,7 +151,7 @@ FDistanceFieldAtlasParameters DistanceField::SetupAtlasParameters(FRDGBuilder& G
 		}
 
 		SceneParameters.DistanceFieldBrickTexture = DistanceFieldSceneData.DistanceFieldBrickVolumeTexture->GetRHI();
-		SceneParameters.DistanceFieldSampler = TStaticSamplerState<SF_Bilinear,AM_Clamp,AM_Clamp,AM_Clamp>::GetRHI();
+		SceneParameters.DistanceFieldSampler = TStaticSamplerState<SF_Bilinear,AM_Wrap, AM_Wrap, AM_Wrap>::GetRHI();
 
 		SceneParameters.DistanceFieldBrickSize = FVector3f(DistanceField::BrickSize);
 		SceneParameters.DistanceFieldUniqueDataBrickSize = FVector3f(DistanceField::UniqueDataBrickSize);
diff --git a/Engine/Source/Runtime/Renderer/Private/GlobalDistanceField.cpp b/Engine/Source/Runtime/Renderer/Private/GlobalDistanceField.cpp
--- a/Engine/Source/Runtime/Renderer/Private/GlobalDistanceField.cpp
+++ b/Engine/Source/Runtime/Renderer/Private/GlobalDistanceField.cpp
@@ -1352,7 +1352,6 @@ class FBuildGridTilesCS : public FGlobalShader
 	SHADER_USE_PARAMETER_STRUCT(FBuildGridTilesCS, FGlobalShader);
 
 	BEGIN_SHADER_PARAMETER_STRUCT(FParameters, )
-		SHADER_PARAMETER_STRUCT_REF(FViewUniformShaderParameters, View)
 		SHADER_PARAMETER_RDG_BUFFER_UAV(RWStructuredBuffer<uint>, RWPageTileBuffer)
 		SHADER_PARAMETER_RDG_BUFFER_UAV(RWStructuredBuffer<uint>, RWPageIndirectArgBuffer)
 		SHADER_PARAMETER_RDG_BUFFER_UAV(RWStructuredBuffer<uint>, RWCullGridTileBuffer)
@@ -1430,7 +1429,6 @@ class FComposeObjectsIntoPagesCS : public FGlobalShader
 	SHADER_USE_PARAMETER_STRUCT(FComposeObjectsIntoPagesCS, FGlobalShader);
 
 	BEGIN_SHADER_PARAMETER_STRUCT(FParameters, )
-		SHADER_PARAMETER_STRUCT_REF(FViewUniformShaderParameters, View)
 		SHADER_PARAMETER_RDG_TEXTURE_UAV(RWTexture3D<UNORM float>, RWPageAtlasTexture)
 		SHADER_PARAMETER_RDG_TEXTURE_UAV(RWTexture3D<UNORM float>, RWCoverageAtlasTexture)
 		RDG_BUFFER_ACCESS(ComposeIndirectArgBuffer, ERHIAccess::IndirectArgs)
@@ -1527,7 +1525,6 @@ class FAllocatePagesCS : public FGlobalShader
 	SHADER_USE_PARAMETER_STRUCT(FAllocatePagesCS, FGlobalShader);
 
 	BEGIN_SHADER_PARAMETER_STRUCT(FParameters, )
-		SHADER_PARAMETER_STRUCT_REF(FViewUniformShaderParameters, View)
 		RDG_BUFFER_ACCESS(PageUpdateIndirectArgBuffer, ERHIAccess::IndirectArgs)
 		SHADER_PARAMETER_RDG_BUFFER_SRV(StructuredBuffer<uint>, PageUpdateTileBuffer)
 		SHADER_PARAMETER_RDG_BUFFER_SRV(StructuredBuffer<uint>, MarkedHeightfieldPageBuffer)
@@ -1668,11 +1665,11 @@ class FPropagateMipDistanceCS : public FGlobalShader
 	SHADER_USE_PARAMETER_STRUCT(FPropagateMipDistanceCS, FGlobalShader);
 
 	BEGIN_SHADER_PARAMETER_STRUCT(FParameters, )
-		SHADER_PARAMETER_STRUCT_REF(FViewUniformShaderParameters, View)
 		SHADER_PARAMETER_RDG_TEXTURE_UAV(RWTexture3D<float>, RWMipTexture)
 		SHADER_PARAMETER_RDG_TEXTURE(Texture3D<float>, PrevMipTexture)
 		SHADER_PARAMETER_RDG_TEXTURE(Texture3D<uint>, PageTableTexture)
 		SHADER_PARAMETER_RDG_TEXTURE(Texture3D<float>, PageAtlasTexture)
+		SHADER_PARAMETER_SAMPLER(SamplerState, DistanceFieldSampler)
 		SHADER_PARAMETER(FVector3f, GlobalDistanceFieldInvPageAtlasSize)
 		SHADER_PARAMETER(uint32, GlobalDistanceFieldClipmapSizeInPages)
 		SHADER_PARAMETER(uint32, ClipmapMipResolution)
@@ -2082,7 +2079,6 @@ void UpdateGlobalDistanceFieldVolume(
 						// Prepare page tiles which need to be updated for update regions
 						{
 							FBuildGridTilesCS::FParameters* PassParameters = GraphBuilder.AllocParameters<FBuildGridTilesCS::FParameters>();
-							PassParameters->View = View.ViewUniformBuffer;
 							PassParameters->RWPageTileBuffer = GraphBuilder.CreateUAV(PageUpdateTileBuffer, PF_R32_UINT);
 							PassParameters->RWPageIndirectArgBuffer = GraphBuilder.CreateUAV(PageUpdateIndirectArgBuffer, PF_R32_UINT);
 							PassParameters->RWCullGridTileBuffer = GraphBuilder.CreateUAV(CullGridUpdateTileBuffer, PF_R32_UINT);
@@ -2138,7 +2134,6 @@ void UpdateGlobalDistanceFieldVolume(
 									UTexture2D* VisibilityTexture = It.Key().Visibility;
 
 									FMarkHeightfieldPagesCS::FParameters* PassParameters = GraphBuilder.AllocParameters<FMarkHeightfieldPagesCS::FParameters>();
-									PassParameters->View = View.ViewUniformBuffer;
 									PassParameters->RWMarkedHeightfieldPageBuffer = GraphBuilder.CreateUAV(MarkedHeightfieldPageBuffer, PF_R32_UINT);
 									PassParameters->PageUpdateIndirectArgBuffer = PageUpdateIndirectArgBuffer;
 									PassParameters->PageUpdateTileBuffer = GraphBuilder.CreateSRV(PageUpdateTileBuffer, PF_R32_UINT);
@@ -2191,7 +2186,6 @@ void UpdateGlobalDistanceFieldVolume(
 
 								{
 									FBuildHeightfieldComposeTilesCS::FParameters* PassParameters = GraphBuilder.AllocParameters<FBuildHeightfieldComposeTilesCS::FParameters>();
-									PassParameters->View = View.ViewUniformBuffer;
 									PassParameters->RWPageComposeHeightfieldIndirectArgBuffer = GraphBuilder.CreateUAV(PageComposeHeightfieldIndirectArgBuffer, PF_R32_UINT);
 									PassParameters->RWPageComposeHeightfieldTileBuffer = GraphBuilder.CreateUAV(PageComposeHeightfieldTileBuffer, PF_R32_UINT);;
 									PassParameters->PageUpdateTileBuffer = GraphBuilder.CreateSRV(PageUpdateTileBuffer, PF_R32_UINT);
@@ -2262,7 +2256,6 @@ void UpdateGlobalDistanceFieldVolume(
 							// Allocate pages for objects
 							{
 								FAllocatePagesCS::FParameters* PassParameters = GraphBuilder.AllocParameters<FAllocatePagesCS::FParameters>();
-								PassParameters->View = View.ViewUniformBuffer;
 								PassParameters->PageUpdateIndirectArgBuffer = PageUpdateIndirectArgBuffer;
 								PassParameters->PageUpdateTileBuffer = GraphBuilder.CreateSRV(PageUpdateTileBuffer, PF_R32_UINT);
 								PassParameters->MarkedHeightfieldPageBuffer = MarkedHeightfieldPageBuffer ? GraphBuilder.CreateSRV(MarkedHeightfieldPageBuffer, PF_R32_UINT) : nullptr;
@@ -2363,7 +2356,6 @@ void UpdateGlobalDistanceFieldVolume(
 							const FVector PageComposeTileWorldExtent = ClipmapVoxelExtent * PageComposeTileSize;
 
 							FComposeObjectsIntoPagesCS::FParameters* PassParameters = GraphBuilder.AllocParameters<FComposeObjectsIntoPagesCS::FParameters>();
-							PassParameters->View = View.ViewUniformBuffer;
 							PassParameters->RWPageAtlasTexture = GraphBuilder.CreateUAV(PageAtlasTexture);
 							PassParameters->RWCoverageAtlasTexture = CoverageAtlasTexture ? GraphBuilder.CreateUAV(CoverageAtlasTexture) : nullptr;
 							PassParameters->ComposeIndirectArgBuffer = PageComposeIndirectArgBuffer;
@@ -2430,7 +2422,6 @@ void UpdateGlobalDistanceFieldVolume(
 									UTexture2D* VisibilityTexture = It.Key().Visibility;
 
 									FComposeHeightfieldsIntoPagesCS::FParameters* PassParameters = GraphBuilder.AllocParameters<FComposeHeightfieldsIntoPagesCS::FParameters>();
-									PassParameters->View = View.ViewUniformBuffer;
 									PassParameters->RWPageAtlasTexture = GraphBuilder.CreateUAV(PageAtlasTexture);
 									PassParameters->RWCoverageAtlasTexture = CoverageAtlasTexture ? GraphBuilder.CreateUAV(CoverageAtlasTexture) : nullptr;
 									PassParameters->ComposeIndirectArgBuffer = PageComposeHeightfieldIndirectArgBuffer;
@@ -2492,10 +2483,10 @@ void UpdateGlobalDistanceFieldVolume(
 								}
 
 								FPropagateMipDistanceCS::FParameters* PassParameters = GraphBuilder.AllocParameters<FPropagateMipDistanceCS::FParameters>();
-								PassParameters->View = View.ViewUniformBuffer;
 								PassParameters->RWMipTexture = GraphBuilder.CreateUAV(NextTexture);
 								PassParameters->PageTableTexture = GAOGlobalDistanceFieldCacheMostlyStaticSeparately ? PageTableCombinedTexture : PageTableLayerTexture;
 								PassParameters->PageAtlasTexture = PageAtlasTexture;
+								PassParameters->DistanceFieldSampler = TStaticSamplerState<SF_Bilinear, AM_Wrap, AM_Wrap, AM_Wrap>::GetRHI();
 								PassParameters->GlobalDistanceFieldInvPageAtlasSize = FVector3f::OneVector / FVector3f(GlobalDistanceField::GetPageAtlasSize(bLumenEnabled, View.FinalPostProcessSettings.LumenSceneViewDistance));
 								PassParameters->GlobalDistanceFieldClipmapSizeInPages = GlobalDistanceField::GetPageTableTextureResolution(bLumenEnabled, View.FinalPostProcessSettings.LumenSceneViewDistance).X;
 								PassParameters->PrevMipTexture = PrevTexture;
diff --git a/Engine/Source/Runtime/Renderer/Private/GlobalDistanceFieldHeightfields.h b/Engine/Source/Runtime/Renderer/Private/GlobalDistanceFieldHeightfields.h
--- a/Engine/Source/Runtime/Renderer/Private/GlobalDistanceFieldHeightfields.h
+++ b/Engine/Source/Runtime/Renderer/Private/GlobalDistanceFieldHeightfields.h
@@ -12,7 +12,6 @@ class FMarkHeightfieldPagesCS : public FGlobalShader
 	SHADER_USE_PARAMETER_STRUCT(FMarkHeightfieldPagesCS, FGlobalShader);
 
 	BEGIN_SHADER_PARAMETER_STRUCT(FParameters, )
-		SHADER_PARAMETER_STRUCT_REF(FViewUniformShaderParameters, View)
 		SHADER_PARAMETER_RDG_BUFFER_UAV(RWStructuredBuffer<uint>, RWMarkedHeightfieldPageBuffer)
 		RDG_BUFFER_ACCESS(PageUpdateIndirectArgBuffer, ERHIAccess::IndirectArgs)
 		SHADER_PARAMETER_RDG_BUFFER_SRV(StructuredBuffer<uint>, PageUpdateTileBuffer)
@@ -86,7 +85,6 @@ class FBuildHeightfieldComposeTilesCS : public FGlobalShader
 	SHADER_USE_PARAMETER_STRUCT(FBuildHeightfieldComposeTilesCS, FGlobalShader);
 
 	BEGIN_SHADER_PARAMETER_STRUCT(FParameters, )
-		SHADER_PARAMETER_STRUCT_REF(FViewUniformShaderParameters, View)
 		SHADER_PARAMETER_RDG_BUFFER_UAV(RWBuffer<uint>, RWPageComposeHeightfieldIndirectArgBuffer)
 		SHADER_PARAMETER_RDG_BUFFER_UAV(RWStructuredBuffer<uint>, RWPageComposeHeightfieldTileBuffer)
 		SHADER_PARAMETER_RDG_BUFFER_SRV(StructuredBuffer<uint>, PageUpdateTileBuffer)
@@ -122,7 +120,6 @@ class FComposeHeightfieldsIntoPagesCS : public FGlobalShader
 	SHADER_USE_PARAMETER_STRUCT(FComposeHeightfieldsIntoPagesCS, FGlobalShader);
 
 	BEGIN_SHADER_PARAMETER_STRUCT(FParameters, )
-		SHADER_PARAMETER_STRUCT_REF(FViewUniformShaderParameters, View)
 		SHADER_PARAMETER_RDG_TEXTURE_UAV(RWTexture3D<float>, RWPageAtlasTexture)
 		SHADER_PARAMETER_RDG_TEXTURE_UAV(RWTexture3D<UNORM float>, RWCoverageAtlasTexture)
 		RDG_BUFFER_ACCESS(ComposeIndirectArgBuffer, ERHIAccess::IndirectArgs)
