Fix GPU tags on RenderVelocities pass to have a valid drawCall count.

diff --git a/Engine/Source/Runtime/Renderer/Private/VelocityRendering.cpp b/Engine/Source/Runtime/Renderer/Private/VelocityRendering.cpp
--- a/Engine/Source/Runtime/Renderer/Private/VelocityRendering.cpp
+++ b/Engine/Source/Runtime/Renderer/Private/VelocityRendering.cpp
@@ -58,7 +58,7 @@ static TAutoConsoleVariable<int32> CVarRHICmdFlushRenderThreadTasksVelocityPass(
 	0,
 	TEXT("Wait for completion of parallel render thread tasks at the end of the velocity pass.  A more granular version of r.RHICmdFlushRenderThreadTasks. If either r.RHICmdFlushRenderThreadTasks or r.RHICmdFlushRenderThreadTasksVelocityPass is > 0 we will flush."));
 
-DECLARE_GPU_STAT_NAMED(RenderVelocities, TEXT("Render Velocities"));
+DECLARE_GPU_DRAWCALL_STAT_NAMED(RenderVelocities, TEXT("Render Velocities"));
 
 /** Validate that deprecated CVars are no longer set. */
 inline void ValidateVelocityCVars()
