Update Global SDF in translated world space.
- Clipmap center is used as origin of translated world space during update.
- Modified HeightfieldDescriptions to store LWC matrices.

diff --git a/Engine/Shaders/Private/DistanceField/GlobalDistanceField.usf b/Engine/Shaders/Private/DistanceField/GlobalDistanceField.usf
--- a/Engine/Shaders/Private/DistanceField/GlobalDistanceField.usf
+++ b/Engine/Shaders/Private/DistanceField/GlobalDistanceField.usf
@@ -33,11 +33,6 @@ float ComputeSquaredDistanceBetweenAABBs(float3 CenterA, float3 ExtentA, float3
 	return dot(AxisDistances, AxisDistances);
 }
 
-float InfluenceRadius;
-uint ClipmapIndex;
-float3 VolumeTexelSize;
-uint3 CullGridDimension;
-
 RWStructuredBuffer<uint> RWCullGridAllocator;
 RWStructuredBuffer<uint> RWCullGridObjectHeader;
 RWStructuredBuffer<uint> RWCullGridObjectArray;
@@ -45,7 +40,6 @@ RWStructuredBuffer<uint> RWCullGridObjectArray;
 #define CULL_GRID_FACTOR 4
 #define GRID_CULL_THREADGROUP_TOTALSIZE 256
 
-uint3 UpdateRegionSize;
 StructuredBuffer<uint> CullGridObjectHeader;
 StructuredBuffer<uint> CullGridObjectArray;
 
@@ -65,31 +59,39 @@ StructuredBuffer<uint> CullGridObjectArray;
 
 RWStructuredBuffer<uint> RWObjectIndexBuffer;
 RWStructuredBuffer<uint> RWObjectIndexNumBuffer;
-float3 ClipmapWorldCenter;
+
+float3 ClipmapTranslatedWorldCenter;
 float3 ClipmapWorldExtent;
+float InfluenceRadius;
 float InfluenceRadiusSq;
 // 0: accept only mostly static objects, 1: accept only often moving objects, 2: accept all objects
 uint AcceptOftenMovingObjectsOnly;
 float MeshSDFRadiusThreshold;
 
+float3 ViewTilePosition;
+float3 RelativePreViewTranslation;
+
 [numthreads(CULLOBJECTS_THREADGROUP_SIZE, 1, 1)]
 void CullObjectsToClipmapCS(
 	uint3 GroupId : SV_GroupID,
 	uint3 DispatchThreadId : SV_DispatchThreadID,
 	uint3 GroupThreadId : SV_GroupThreadID)
 {
-	uint ObjectIndex = DispatchThreadId.x;
+	const FLWCVector3 PreViewTranslation = MakeLWCVector3(-ViewTilePosition, RelativePreViewTranslation);
+
+	const uint ObjectIndex = DispatchThreadId.x;
 
 	if (ObjectIndex < NumSceneObjects)
 	{
-		FDFObjectBounds DFObjectBounds = LoadDFObjectBounds(ObjectIndex);
+		const FDFObjectBounds DFObjectBounds = LoadDFObjectBounds(ObjectIndex);
+		const float3 TranslatedBoundsCenter = LWCToFloat(LWCAdd(DFObjectBounds.Center, PreViewTranslation));
 
 		BRANCH
 		if (DFObjectBounds.bVisible
 			&& (DFObjectBounds.SphereRadius > MeshSDFRadiusThreshold || DFObjectBounds.bEmissiveLightSource)
 			&& (AcceptOftenMovingObjectsOnly == 2 || DFObjectBounds.OftenMoving == AcceptOftenMovingObjectsOnly))
 		{
-			float DistanceSq = ComputeSquaredDistanceBetweenAABBs(ClipmapWorldCenter, ClipmapWorldExtent, LWCHackToFloat(DFObjectBounds.Center), DFObjectBounds.BoxExtent);
+			float DistanceSq = ComputeSquaredDistanceBetweenAABBs(ClipmapTranslatedWorldCenter, ClipmapWorldExtent, TranslatedBoundsCenter, DFObjectBounds.BoxExtent);
 			if (DistanceSq < InfluenceRadiusSq)
 			{
 				uint DestIndex;
@@ -136,13 +138,13 @@ Buffer<float4> UpdateBoundsBuffer;
 uint NumUpdateBounds;
 
 uint3 CullGridResolution;
-float3 CullGridCoordToWorldCenterScale;
-float3 CullGridCoordToWorldCenterBias;
+float3 CullGridCoordToTranslatedWorldCenterScale;
+float3 CullGridCoordToTranslatedWorldCenterBias;
 float3 CullGridTileWorldExtent;
 
 uint3 PageGridResolution;
-float3 PageGridCoordToWorldCenterScale;
-float3 PageGridCoordToWorldCenterBias;
+float3 PageGridCoordToTranslatedWorldCenterScale;
+float3 PageGridCoordToTranslatedWorldCenterBias;
 float3 PageGridTileWorldExtent;
 
 #define MAX_CULLED_TILES 511u
@@ -211,7 +213,7 @@ void BuildGridTilesCS(
 	// First cull per group
 	if (all(CullGridCoord < CullGridResolution))
 	{
-		float3 CullGridWorldCenter = CullGridCoord * CullGridCoordToWorldCenterScale + CullGridCoordToWorldCenterBias;
+		float3 CullGridTranslatedWorldCenter = CullGridCoord * CullGridCoordToTranslatedWorldCenterScale + CullGridCoordToTranslatedWorldCenterBias;
 
 		for (uint UpdateBoundsIndex = ThreadIndex; UpdateBoundsIndex < NumUpdateBounds; UpdateBoundsIndex += THREADGROUP_SIZE)
 		{
@@ -221,7 +223,7 @@ void BuildGridTilesCS(
 			bool bExtendByInflueceRadius = UpdateBoundsBuffer[UpdateBoundsIndex * 2 + 0].w > 0.0f;
 			float UpdateBoundsInfluenceRadiusSq = bExtendByInflueceRadius ? InfluenceRadiusSq : 0.0f;
 
-			float DistanceSq = ComputeSquaredDistanceBetweenAABBs(UpdateBoundsCenter, UpdateBoundsExtent, CullGridWorldCenter, CullGridTileWorldExtent);
+			float DistanceSq = ComputeSquaredDistanceBetweenAABBs(UpdateBoundsCenter, UpdateBoundsExtent, CullGridTranslatedWorldCenter, CullGridTileWorldExtent);
 			if (DistanceSq <= UpdateBoundsInfluenceRadiusSq)
 			{
 				uint DestIndex;
@@ -242,7 +244,7 @@ void BuildGridTilesCS(
 	// Then cull per page
 	if (all(PageGridCoord < PageGridResolution))
 	{
-		float3 PageGridWorldCenter = PageGridCoord * PageGridCoordToWorldCenterScale + PageGridCoordToWorldCenterBias;
+		float3 PageGridTranslatedWorldCenter = PageGridCoord * PageGridCoordToTranslatedWorldCenterScale + PageGridCoordToTranslatedWorldCenterBias;
 
 		for (uint IndexInCulledBoundsList = 0; IndexInCulledBoundsList < SharedNumCulledBounds; ++IndexInCulledBoundsList)
 		{
@@ -254,7 +256,7 @@ void BuildGridTilesCS(
 			bool bExtendByInflueceRadius = UpdateBoundsBuffer[UpdateBoundsIndex * 2 + 0].w > 0.0f;
 			float UpdateBoundsInfluenceRadiusSq = bExtendByInflueceRadius ? InfluenceRadiusSq : 0.0f;
 
-			float DistanceSq = ComputeSquaredDistanceBetweenAABBs(UpdateBoundsCenter, UpdateBoundsExtent, PageGridWorldCenter, PageGridTileWorldExtent);
+			float DistanceSq = ComputeSquaredDistanceBetweenAABBs(UpdateBoundsCenter, UpdateBoundsExtent, PageGridTranslatedWorldCenter, PageGridTileWorldExtent);
 			if (DistanceSq <= UpdateBoundsInfluenceRadiusSq)
 			{
 				SetSharedCullGridMaskBit(LocalPageGridCoord);
@@ -316,9 +318,9 @@ uint PageTableClipmapOffsetZ;
 float3 PageWorldExtent;
 float PageWorldRadius;
 float ClipmapInfluenceRadius;
-float3 PageCoordToPageWorldCenterScale;
-float3 PageCoordToPageWorldCenterBias;
-float4 ClipmapVolumeWorldToUVAddAndMul;
+float3 PageCoordToPageTranslatedWorldCenterScale;
+float3 PageCoordToPageTranslatedWorldCenterBias;
+float4 ClipmapVolumeTranslatedWorldToUVAddAndMul;
 
 RWStructuredBuffer<uint> RWPageComposeTileBuffer;
 
@@ -333,10 +335,9 @@ groupshared float GroupMinPageDistance[64];
 
 uint3 PageGridCoordToPageTableTextureCoord(uint3 PageGridCoord)
 {
-	float3 PageWorldCenter = PageGridCoord * PageCoordToPageWorldCenterScale + PageCoordToPageWorldCenterBias;
+	float3 PageTranslatedWorldCenter = PageGridCoord * PageCoordToPageTranslatedWorldCenterScale + PageCoordToPageTranslatedWorldCenterBias;
 
-	float4 WorldToUVAddAndMul = ClipmapVolumeWorldToUVAddAndMul;
-	float3 ClipmapUV = frac(PageWorldCenter * WorldToUVAddAndMul.www + WorldToUVAddAndMul.xyz);
+	float3 ClipmapUV = frac(PageTranslatedWorldCenter * ClipmapVolumeTranslatedWorldToUVAddAndMul.www + ClipmapVolumeTranslatedWorldToUVAddAndMul.xyz);
 
 	int3 PageTableTextureCoord = clamp(saturate(ClipmapUV) * PageGridResolution, 0, PageGridResolution - 1) + int3(0, 0, PageTableClipmapOffsetZ);
 	return PageTableTextureCoord;
@@ -352,6 +353,8 @@ void AllocatePagesCS(
 	uint3 DispatchThreadId : SV_DispatchThreadID,
 	uint3 GroupThreadId : SV_GroupThreadID)
 {
+	const FLWCVector3 PreViewTranslation = MakeLWCVector3(-ViewTilePosition, RelativePreViewTranslation);
+
 	uint ThreadIndex = GroupThreadId.x;
 	uint IndexInPageTileBuffer = GroupId.x;
 	uint PackedPageGridCoord = PageUpdateTileBuffer[IndexInPageTileBuffer];
@@ -367,7 +370,8 @@ void AllocatePagesCS(
 	PageGridCoord.z = (PackedPageGridCoord >> 16) & 0xFF;
 
 	uint3 PageTableTextureCoord = PageGridCoordToPageTableTextureCoord(PageGridCoord);
-	float3 PageWorldCenter = PageGridCoord * PageCoordToPageWorldCenterScale + PageCoordToPageWorldCenterBias;
+	float3 PageTranslatedWorldCenter = PageGridCoord * PageCoordToPageTranslatedWorldCenterScale + PageCoordToPageTranslatedWorldCenterBias;
+	FLWCVector3 PageWorldCenter = LWCSubtract(PageTranslatedWorldCenter, PreViewTranslation);
 
 #if PROCESS_DISTANCE_FIELDS
 	// Loop through objects to find out if the current page contains any objects and should be allocated
@@ -385,13 +389,14 @@ void AllocatePagesCS(
 		{
 			uint ObjectIndex = CullGridObjectArray[CullGridObjectOffset + IndexInCullCell];
 
-			FDFObjectBounds DFObjectBounds = LoadDFObjectBounds(ObjectIndex);
+			const FDFObjectBounds DFObjectBounds = LoadDFObjectBounds(ObjectIndex);
+			const float3 TranslatedBoundsCenter = LWCToFloat(LWCAdd(DFObjectBounds.Center, PreViewTranslation));
 
-			float SquaredDistanceBetweenAABBs = ComputeSquaredDistanceBetweenAABBs(PageWorldCenter, PageWorldExtent, LWCHackToFloat(DFObjectBounds.Center), DFObjectBounds.BoxExtent);
+			float SquaredDistanceBetweenAABBs = ComputeSquaredDistanceBetweenAABBs(PageTranslatedWorldCenter, PageWorldExtent, TranslatedBoundsCenter, DFObjectBounds.BoxExtent);
 
 			if (SquaredDistanceBetweenAABBs < ClipmapInfluenceRadius * ClipmapInfluenceRadius)
 			{
-				float DistanceToPageCenter = DistanceToNearestSurfaceForObject(ObjectIndex, LWCPromote(PageWorldCenter), InfluenceMax); // LWC_TODO
+				float DistanceToPageCenter = DistanceToNearestSurfaceForObject(ObjectIndex, PageWorldCenter, InfluenceMax);
 				MinDistanceToPageCenter = min(MinDistanceToPageCenter, DistanceToPageCenter);
 			}
 		}
@@ -562,6 +567,8 @@ void CullObjectsToGridCS(
 	uint3 DispatchThreadId : SV_DispatchThreadID,
     uint3 GroupThreadId : SV_GroupThreadID) 
 {
+	const FLWCVector3 PreViewTranslation = MakeLWCVector3(-ViewTilePosition, RelativePreViewTranslation);
+
 	uint ThreadIndex = GroupThreadId.x;
 
 	uint PackedCullGridTile = CullGridTileBuffer[GroupId.x];
@@ -580,7 +587,7 @@ void CullObjectsToGridCS(
 
 	GroupMemoryBarrierWithGroupSync();
 
-	float3 CullTileWorldCenter = CullGridCoord * CullGridCoordToWorldCenterScale + CullGridCoordToWorldCenterBias;
+	float3 CullTileTranslatedWorldCenter = CullGridCoord * CullGridCoordToTranslatedWorldCenterScale + CullGridCoordToTranslatedWorldCenterBias;
 
 	if (all(CullGridCoord < CullGridResolution))
 	{
@@ -588,9 +595,10 @@ void CullObjectsToGridCS(
 		{
 			uint ObjectIndex = ObjectIndexBuffer[IndexInObjectIndexBuffer];
 
-			FDFObjectBounds DFObjectBounds = LoadDFObjectBounds(ObjectIndex);
+			const FDFObjectBounds DFObjectBounds = LoadDFObjectBounds(ObjectIndex);
+			const float3 TranslatedBoundsCenter = LWCToFloat(LWCAdd(DFObjectBounds.Center, PreViewTranslation));
 
-			float DistanceSq = ComputeSquaredDistanceBetweenAABBs(CullTileWorldCenter, CullTileWorldExtent, LWCHackToFloat(DFObjectBounds.Center), DFObjectBounds.BoxExtent);
+			float DistanceSq = ComputeSquaredDistanceBetweenAABBs(CullTileTranslatedWorldCenter, CullTileWorldExtent, TranslatedBoundsCenter, DFObjectBounds.BoxExtent);
 
 			if (DistanceSq <= InfluenceRadiusSq)
 			{
@@ -618,9 +626,10 @@ void CullObjectsToGridCS(
 		{
 			uint ObjectIndex = ObjectIndexBuffer[IndexInObjectIndexBuffer];
 
-			FDFObjectBounds DFObjectBounds = LoadDFObjectBounds(ObjectIndex);
+			const FDFObjectBounds DFObjectBounds = LoadDFObjectBounds(ObjectIndex);
+			const float3 TranslatedBoundsCenter = LWCToFloat(LWCAdd(DFObjectBounds.Center, PreViewTranslation));
 
-			float DistanceSq = ComputeSquaredDistanceBetweenAABBs(CullTileWorldCenter, CullTileWorldExtent, LWCHackToFloat(DFObjectBounds.Center), DFObjectBounds.BoxExtent);
+			float DistanceSq = ComputeSquaredDistanceBetweenAABBs(CullTileTranslatedWorldCenter, CullTileWorldExtent, TranslatedBoundsCenter, DFObjectBounds.BoxExtent);
 
 			if (DistanceSq <= InfluenceRadiusSq)
 			{
@@ -648,13 +657,12 @@ StructuredBuffer<uint> PageTileBuffer;
 StructuredBuffer<uint> ComposeTileBuffer;
 uint3 ClipmapResolution;
 float3 ComposeTileWorldExtent;
-float3 ClipmapMinBounds;
 float ClipmapVoxelExtent;
 
-float3 PageCoordToVoxelCenterScale;
-float3 PageCoordToVoxelCenterBias;
+float3 PageCoordToVoxelTranslatedCenterScale;
+float3 PageCoordToVoxelTranslatedCenterBias;
 
-void CompositeMeshSDF(uint ObjectIndex, float3 VoxelWorldCenter, inout float MinDistance, inout bool bVoxelHasOneSidedMeshes)
+void CompositeMeshSDF(uint ObjectIndex, FLWCVector3 VoxelWorldCenter, inout float MinDistance, inout bool bVoxelHasOneSidedMeshes)
 {
 	float MaxEarlyOutDistance = 0.0f;
 
@@ -666,7 +674,7 @@ void CompositeMeshSDF(uint ObjectIndex, float3 VoxelWorldCenter, inout float Min
 	float MaxQueryDistance = min(InfluenceRadius, max(MinDistance, MaxEarlyOutDistance));
 
 	FDFObjectData DFObjectData = LoadDFObjectData(ObjectIndex);
-	float DistanceToOccluder = DistanceToNearestSurfaceForObject(DFObjectData, LWCPromote(VoxelWorldCenter), MaxQueryDistance); // LWC_TODO
+	float DistanceToOccluder = DistanceToNearestSurfaceForObject(DFObjectData, VoxelWorldCenter, MaxQueryDistance);
 	MinDistance = min(MinDistance, DistanceToOccluder);
 
 #if COMPOSITE_COVERAGE_ATLAS
@@ -683,6 +691,8 @@ void ComposeObjectsIntoPagesCS(
 	uint3 DispatchThreadId : SV_DispatchThreadID,
     uint3 GroupThreadId : SV_GroupThreadID) 
 {
+	const FLWCVector3 PreViewTranslation = MakeLWCVector3(-ViewTilePosition, RelativePreViewTranslation);
+
 	uint ThreadIndex = (GroupThreadId.z * COMPOSITE_THREADGROUP_SIZEY + GroupThreadId.y) * COMPOSITE_THREADGROUP_SIZEX + GroupThreadId.x;
 
 	uint PackedPageTile = ComposeTileBuffer[GroupId.x / GLOBAL_DISTANCE_FIELD_PAGE_GROUP_NUM];
@@ -705,8 +715,10 @@ void ComposeObjectsIntoPagesCS(
 	float3 ComposeTileCenterPageCoord = PageGridCoord * GLOBAL_DISTANCE_FIELD_PAGE_RESOLUTION + PageGroupOffset * 4 + 1.5f - GLOBAL_DISTANCE_FIELD_PAGE_BORDER;
 	uint3 CullGridCoord = PageGridCoord / CULL_GRID_FACTOR;
 
-	float3 VoxelWorldCenter = PageCoord * PageCoordToVoxelCenterScale + PageCoordToVoxelCenterBias;
-	float3 ComposeTileWorldCenter = ComposeTileCenterPageCoord * PageCoordToVoxelCenterScale + PageCoordToVoxelCenterBias;
+	float3 VoxelTranslatedWorldCenter = PageCoord * PageCoordToVoxelTranslatedCenterScale + PageCoordToVoxelTranslatedCenterBias;
+	float3 ComposeTileTranslatedWorldCenter = ComposeTileCenterPageCoord * PageCoordToVoxelTranslatedCenterScale + PageCoordToVoxelTranslatedCenterBias;
+
+	const FLWCVector3 VoxelWorldCenter = LWCSubtract(VoxelTranslatedWorldCenter, PreViewTranslation);
 
 	uint3 PageTableTextureCoord = PageGridCoordToPageTableTextureCoord(PageGridCoord);
 	uint PageId = PageTableLayerTexture.Load(int4(PageTableTextureCoord, 0));
@@ -742,9 +754,10 @@ void ComposeObjectsIntoPagesCS(
 		{
 			uint ObjectIndex = CullGridObjectArray[CullGridObjectOffset + IndexInCullCell];
 
-			FDFObjectBounds DFObjectBounds = LoadDFObjectBounds(ObjectIndex);
+			const FDFObjectBounds DFObjectBounds = LoadDFObjectBounds(ObjectIndex);
+			const float3 TranslatedBoundsCenter = LWCToFloat(LWCAdd(DFObjectBounds.Center, PreViewTranslation));
 
-			float DistanceSq = ComputeSquaredDistanceBetweenAABBs(ComposeTileWorldCenter, ComposeTileWorldExtent, LWCHackToFloat(DFObjectBounds.Center), DFObjectBounds.BoxExtent);
+			float DistanceSq = ComputeSquaredDistanceBetweenAABBs(ComposeTileTranslatedWorldCenter, ComposeTileWorldExtent, TranslatedBoundsCenter, DFObjectBounds.BoxExtent);
 
 			if (DistanceSq < InfluenceRadiusSq)
 			{
diff --git a/Engine/Shaders/Private/DistanceField/GlobalDistanceFieldHeightfields.usf b/Engine/Shaders/Private/DistanceField/GlobalDistanceFieldHeightfields.usf
--- a/Engine/Shaders/Private/DistanceField/GlobalDistanceFieldHeightfields.usf
+++ b/Engine/Shaders/Private/DistanceField/GlobalDistanceFieldHeightfields.usf
@@ -16,12 +16,12 @@ StructuredBuffer<uint> PageUpdateTileBuffer;
 StructuredBuffer<uint> ComposeTileBuffer;
 Texture3D<uint> PageTableLayerTexture;
 
-float3 PageCoordToPageWorldCenterScale;
-float3 PageCoordToPageWorldCenterBias;
-float3 PageCoordToVoxelCenterScale;
-float3 PageCoordToVoxelCenterBias;
+float3 PageCoordToPageTranslatedWorldCenterScale;
+float3 PageCoordToPageTranslatedWorldCenterBias;
+float3 PageCoordToVoxelTranslatedCenterScale;
+float3 PageCoordToVoxelTranslatedCenterBias;
 float3 PageWorldExtent;
-float4 ClipmapVolumeWorldToUVAddAndMul;
+float4 ClipmapVolumeTranslatedWorldToUVAddAndMul;
 float3 InvPageGridResolution;
 uint3 PageGridResolution;
 uint PageTableClipmapOffsetZ;
@@ -29,12 +29,14 @@ float ClipmapVoxelExtent;
 float InfluenceRadius;
 float HeightfieldThickness;
 
+float3 ViewTilePosition;
+float3 RelativePreViewTranslation;
+
 uint3 PageGridCoordToPageTableTextureCoord(uint3 PageGridCoord)
 {
-	float3 PageWorldCenter = PageGridCoord * PageCoordToPageWorldCenterScale + PageCoordToPageWorldCenterBias;
+	float3 PageTranslatedWorldCenter = PageGridCoord * PageCoordToPageTranslatedWorldCenterScale + PageCoordToPageTranslatedWorldCenterBias;
 
-	float4 WorldToUVAddAndMul = ClipmapVolumeWorldToUVAddAndMul;
-	float3 ClipmapUV = frac(PageWorldCenter * WorldToUVAddAndMul.www + WorldToUVAddAndMul.xyz);
+	float3 ClipmapUV = frac(PageTranslatedWorldCenter * ClipmapVolumeTranslatedWorldToUVAddAndMul.www + ClipmapVolumeTranslatedWorldToUVAddAndMul.xyz);
 
 	int3 PageTableTextureCoord = clamp(saturate(ClipmapUV) * PageGridResolution, 0, PageGridResolution - 1) + int3(0, 0, PageTableClipmapOffsetZ);
 	return PageTableTextureCoord;
@@ -48,6 +50,8 @@ void MarkHeightfieldPagesCS(
 	uint3 DispatchThreadId : SV_DispatchThreadID,
 	uint3 GroupThreadId : SV_GroupThreadID)
 {
+	const FLWCVector3 PreViewTranslation = MakeLWCVector3(-ViewTilePosition, RelativePreViewTranslation);
+
 	uint3 VoxelCoord = uint3(GroupThreadId.xy, 0);
 	uint ThreadIndex = GroupThreadId.x + GroupThreadId.y * THREADGROUP_SIZE_X;
 	uint IndexInPageBuffer = GroupId.x;
@@ -58,8 +62,9 @@ void MarkHeightfieldPagesCS(
 	PageGridCoord.y = (PackedPageGridCoord >> 8) & 0xFF;
 	PageGridCoord.z = (PackedPageGridCoord >> 16) & 0xFF;
 
-	float3 PageWorldCenter = PageGridCoord * PageCoordToPageWorldCenterScale + PageCoordToPageWorldCenterBias;
-	float3 VoxelWorldCenter = PageWorldCenter - PageWorldExtent + VoxelCoord * 2.0f * PageWorldExtent;
+	float3 PageTranslatedWorldCenter = PageGridCoord * PageCoordToPageTranslatedWorldCenterScale + PageCoordToPageTranslatedWorldCenterBias;
+	float3 VoxelTranslatedWorldCenter = PageTranslatedWorldCenter - PageWorldExtent + VoxelCoord * 2.0f * PageWorldExtent;
+	FLWCVector3 VoxelWorldCenter = LWCSubtract(VoxelTranslatedWorldCenter, PreViewTranslation);
 
 	GroupMarkedPage[ThreadIndex] = 0;
 
@@ -68,7 +73,7 @@ void MarkHeightfieldPagesCS(
 	LOOP
 	for (uint HeightfieldIndex = 0; HeightfieldIndex < NumHeightfields; HeightfieldIndex++)
 	{
-		float3 LocalPosition = mul(float4(VoxelWorldCenter, 1), GetWorldToLocal(HeightfieldIndex)).xyz;
+		float3 LocalPosition = LWCMultiply(VoxelWorldCenter, GetWorldToLocal(HeightfieldIndex));
 		float4 MinMaxHeightfieldUV;
 		float2 HeightfieldUV = GetHeightfieldUV(HeightfieldIndex, LocalPosition.xy, MinMaxHeightfieldUV);
 
@@ -79,7 +84,7 @@ void MarkHeightfieldPagesCS(
 		}
 	}
 
-	float3 LocalPosition = mul(float4(VoxelWorldCenter, 1), GetWorldToLocal(ValidHeightfieldIndex)).xyz;
+	float3 LocalPosition = LWCMultiply(VoxelWorldCenter, GetWorldToLocal(ValidHeightfieldIndex));
 	float4 MinMaxHeightfieldUV;
 	float2 HeightfieldUV = GetHeightfieldUV(ValidHeightfieldIndex, LocalPosition.xy, MinMaxHeightfieldUV);
 
@@ -87,15 +92,16 @@ void MarkHeightfieldPagesCS(
 	{
 		float3 WorldHeightfieldNormal;
 		float HeightfieldVisibility;
-		float3 WorldHeightfieldShadingPosition = GetHeightfieldWorldPositionAndNormal(ValidHeightfieldIndex, LocalPosition.xy, HeightfieldUV, WorldHeightfieldNormal, HeightfieldVisibility);
+		FLWCVector3 WorldHeightfieldPosition = GetHeightfieldWorldPositionAndNormal(ValidHeightfieldIndex, LocalPosition.xy, HeightfieldUV, WorldHeightfieldNormal, HeightfieldVisibility);
+		float3 TranslatedWorldHeightfieldPosition = LWCToFloat(LWCAdd(WorldHeightfieldPosition, PreViewTranslation));
 
 		if (HeightfieldVisibility > 0.5f) // Skip holes in the heightfield
 		{
-			float3 WorldPositionNearestZ = VoxelWorldCenter;
-			WorldPositionNearestZ.z = clamp(WorldHeightfieldShadingPosition.z, VoxelWorldCenter.z, VoxelWorldCenter.z + GLOBAL_DISTANCE_FIELD_PAGE_RESOLUTION_IN_ATLAS * 2.0f * ClipmapVoxelExtent);				
+			float3 TranslatedWorldPositionNearestZ = VoxelTranslatedWorldCenter;
+			TranslatedWorldPositionNearestZ.z = clamp(TranslatedWorldHeightfieldPosition.z, VoxelTranslatedWorldCenter.z, VoxelTranslatedWorldCenter.z + GLOBAL_DISTANCE_FIELD_PAGE_RESOLUTION_IN_ATLAS * 2.0f * ClipmapVoxelExtent);
 
 			// Project the vertical height vector onto the normal of the heightfield directly below the point we are computing the distance field for, use the perpendicular distance
-			float DistanceToHeightfieldPlane = dot(WorldHeightfieldNormal, WorldPositionNearestZ - WorldHeightfieldShadingPosition);
+			float DistanceToHeightfieldPlane = dot(WorldHeightfieldNormal, TranslatedWorldPositionNearestZ - TranslatedWorldHeightfieldPosition);
 
 			// Limit negative region of a heightfield to a user defined thickness  
 			const float MinInteriorDistance = -HeightfieldThickness;
@@ -211,6 +217,8 @@ void ComposeHeightfieldsIntoPagesCS(
 	uint3 GroupId : SV_GroupID,
 	uint3 GroupThreadId : SV_GroupThreadID)
 {
+	const FLWCVector3 PreViewTranslation = MakeLWCVector3(-ViewTilePosition, RelativePreViewTranslation);
+
 	uint PackedPageTile = ComposeTileBuffer[GroupId.x];
 
 	uint3 PageGridCoord;
@@ -221,7 +229,8 @@ void ComposeHeightfieldsIntoPagesCS(
 	uint3 TexelCoordInPage = uint3(GroupThreadId.xy, 0);
 
 	float3 PageCoord = PageGridCoord * GLOBAL_DISTANCE_FIELD_PAGE_RESOLUTION + TexelCoordInPage - GLOBAL_DISTANCE_FIELD_PAGE_BORDER;
-	float3 VoxelWorldCenter = PageCoord * PageCoordToVoxelCenterScale + PageCoordToVoxelCenterBias;
+	float3 VoxelTranslatedWorldCenter = PageCoord * PageCoordToVoxelTranslatedCenterScale + PageCoordToVoxelTranslatedCenterBias;
+	FLWCVector3 VoxelWorldCenter = LWCSubtract(VoxelTranslatedWorldCenter, PreViewTranslation);
 
 	uint3 PageTableTextureCoord = PageGridCoordToPageTableTextureCoord(PageGridCoord);
 	uint PageId = PageTableLayerTexture.Load(int4(PageTableTextureCoord, 0));
@@ -233,7 +242,7 @@ void ComposeHeightfieldsIntoPagesCS(
 	LOOP
 	for (uint HeightfieldIndex = 0; HeightfieldIndex < NumHeightfields; HeightfieldIndex++)
 	{
-		float3 LocalPosition = mul(float4(VoxelWorldCenter, 1), GetWorldToLocal(HeightfieldIndex)).xyz;
+		float3 LocalPosition = LWCMultiply(VoxelWorldCenter, GetWorldToLocal(HeightfieldIndex));
 		float4 MinMaxHeightfieldUV;
 		float2 HeightfieldUV = GetHeightfieldUV(HeightfieldIndex, LocalPosition.xy, MinMaxHeightfieldUV);
 
@@ -246,7 +255,7 @@ void ComposeHeightfieldsIntoPagesCS(
 
 	if (ValidHeightfieldIndex < NumHeightfields)
 	{
-		float3 LocalPosition = mul(float4(VoxelWorldCenter, 1), GetWorldToLocal(ValidHeightfieldIndex)).xyz;
+		float3 LocalPosition = LWCMultiply(VoxelWorldCenter, GetWorldToLocal(ValidHeightfieldIndex));
 		float4 MinMaxHeightfieldUV;
 		float2 HeightfieldUV = GetHeightfieldUV(ValidHeightfieldIndex, LocalPosition.xy, MinMaxHeightfieldUV);
 
@@ -254,7 +263,8 @@ void ComposeHeightfieldsIntoPagesCS(
 		{
 			float3 WorldHeightfieldNormal;
 			float HeightfieldVisibility;
-			float3 WorldHeightfieldShadingPosition = GetHeightfieldWorldPositionAndNormal(ValidHeightfieldIndex, LocalPosition.xy, HeightfieldUV, WorldHeightfieldNormal, HeightfieldVisibility);
+			FLWCVector3 WorldHeightfieldShadingPosition = GetHeightfieldWorldPositionAndNormal(ValidHeightfieldIndex, LocalPosition.xy, HeightfieldUV, WorldHeightfieldNormal, HeightfieldVisibility);
+			float3 TranslatedWorldHeightfieldShadingPosition = LWCToFloat(LWCAdd(WorldHeightfieldShadingPosition, PreViewTranslation));
 
 			if (HeightfieldVisibility > 0.5f) // Skip holes in the heightfield
 			{
@@ -262,10 +272,10 @@ void ComposeHeightfieldsIntoPagesCS(
 				LOOP
 				for (uint ZIndex = 0; ZIndex < GLOBAL_DISTANCE_FIELD_PAGE_RESOLUTION_IN_ATLAS; ++ZIndex)
 				{
-					float3 WorldPosition = VoxelWorldCenter.xyz + float3(0.0f, 0.0f, ZIndex * PageCoordToVoxelCenterScale.z);
+					float3 TranslatedWorldPosition = VoxelTranslatedWorldCenter.xyz + float3(0.0f, 0.0f, ZIndex * PageCoordToVoxelTranslatedCenterScale.z);
 
 					// Project the vertical height vector onto the normal of the heightfield directly below the point we are computing the distance field for, use the perpendicular distance
-					float DistanceToHeightfieldPlane = dot(WorldHeightfieldNormal, WorldPosition - WorldHeightfieldShadingPosition);
+					float DistanceToHeightfieldPlane = dot(WorldHeightfieldNormal, TranslatedWorldPosition - TranslatedWorldHeightfieldShadingPosition);
 
 					// Limit negative region of a heightfield to a user defined thickness  
 					const float MinInteriorDistance = -HeightfieldThickness;
diff --git a/Engine/Shaders/Private/HeightfieldLightingShared.ush b/Engine/Shaders/Private/HeightfieldLightingShared.ush
--- a/Engine/Shaders/Private/HeightfieldLightingShared.ush
+++ b/Engine/Shaders/Private/HeightfieldLightingShared.ush
@@ -14,7 +14,7 @@ uint NumHeightfields;
 Buffer<float4> HeightfieldDescriptions;
 
 // In float4's, must match C++
-#define HEIGHTFIELD_DATA_STRIDE 10
+#define HEIGHTFIELD_DATA_STRIDE 11
 
 float4 GetHeightfieldScaleBias(uint Index, out bool bHasSubsections)
 {
@@ -40,27 +40,35 @@ float2 GetHeightfieldSizeInv(uint Index)
 	return HeightfieldDescriptions[Index * HEIGHTFIELD_DATA_STRIDE + 2].zw;
 }
 
-float4x4 GetWorldToLocal(uint Index)
+FLWCInverseMatrix GetWorldToLocal(uint Index)
 {
-	float4 M0 = HeightfieldDescriptions[Index * HEIGHTFIELD_DATA_STRIDE + 3];
-	float4 M1 = HeightfieldDescriptions[Index * HEIGHTFIELD_DATA_STRIDE + 4];
-	float4 M2 = HeightfieldDescriptions[Index * HEIGHTFIELD_DATA_STRIDE + 5];
+	float3 TilePosition = HeightfieldDescriptions[Index * HEIGHTFIELD_DATA_STRIDE + 3].xyz;
 
-	return transpose(float4x4(M0, M1, M2, float4(0.0f, 0.0f, 0.0f, 1.0f)));
+	float4 M0 = HeightfieldDescriptions[Index * HEIGHTFIELD_DATA_STRIDE + 4];
+	float4 M1 = HeightfieldDescriptions[Index * HEIGHTFIELD_DATA_STRIDE + 5];
+	float4 M2 = HeightfieldDescriptions[Index * HEIGHTFIELD_DATA_STRIDE + 6];
+
+	float4x4 RelativeWorldToLocal = transpose(float4x4(M0, M1, M2, float4(0.0f, 0.0f, 0.0f, 1.0f)));
+
+	return MakeLWCInverseMatrix4x3(TilePosition, RelativeWorldToLocal);
 }
 
-float4x4 GetLocalToWorld(uint Index)
+FLWCMatrix GetLocalToWorld(uint Index)
 {
-	float4 M0 = HeightfieldDescriptions[Index * HEIGHTFIELD_DATA_STRIDE + 6];
-	float4 M1 = HeightfieldDescriptions[Index * HEIGHTFIELD_DATA_STRIDE + 7];
-	float4 M2 = HeightfieldDescriptions[Index * HEIGHTFIELD_DATA_STRIDE + 8];
+	float3 TilePosition = HeightfieldDescriptions[Index * HEIGHTFIELD_DATA_STRIDE + 3].xyz;
+
+	float4 M0 = HeightfieldDescriptions[Index * HEIGHTFIELD_DATA_STRIDE + 7];
+	float4 M1 = HeightfieldDescriptions[Index * HEIGHTFIELD_DATA_STRIDE + 8];
+	float4 M2 = HeightfieldDescriptions[Index * HEIGHTFIELD_DATA_STRIDE + 9];
+
+	float4x4 LocalToRelativeWorld = transpose(float4x4(M0, M1, M2, float4(0.0f, 0.0f, 0.0f, 1.0f)));
 
-	return transpose(float4x4(M0, M1, M2, float4(0.0f, 0.0f, 0.0f, 1.0f)));
+	return MakeLWCMatrix(TilePosition, LocalToRelativeWorld);
 }
 
 float4 GetVisibilityMask(uint Index)
 {
-	return HeightfieldDescriptions[Index * HEIGHTFIELD_DATA_STRIDE + 9];
+	return HeightfieldDescriptions[Index * HEIGHTFIELD_DATA_STRIDE + 10];
 }
 
 float2 GetHeightfieldUV(uint HeightfieldIndex, float2 LocalPosition, out float4 MinMaxHeightfieldUV)
@@ -93,7 +101,7 @@ float2 GetVisibilityUV(uint HeightfieldIndex, float2 LocalPosition)
 	return LocalPosition * GetHeightfieldSizeInv(HeightfieldIndex);
 }
 
-float3 GetHeightfieldWorldPositionAndNormal(uint HeightfieldIndex, float2 LocalPosition, float2 HeightfieldUV, out float3 WorldNormal, out float Visibility)
+FLWCVector3 GetHeightfieldWorldPositionAndNormal(uint HeightfieldIndex, float2 LocalPosition, float2 HeightfieldUV, out float3 WorldNormal, out float Visibility)
 {
 	//@todo - don't call GetHeightfieldScaleBias twice
 	bool bHasSubsections;
@@ -109,13 +117,13 @@ float3 GetHeightfieldWorldPositionAndNormal(uint HeightfieldIndex, float2 LocalP
 	float4 SampleVisibility = Texture2DSampleLevel(VisibilityTexture, VisibilitySampler, VisibilityUV, 0);
 	Visibility = 1.f - dot(SampleVisibility, GetVisibilityMask(HeightfieldIndex));
 
-	float3 HeightfieldShadingPosition = float3(LocalPosition, HeightfieldHeight);
-	float3 WorldHeightfieldShadingPosition = mul(float4(HeightfieldShadingPosition, 1), GetLocalToWorld(HeightfieldIndex)).xyz;
+	float3 HeightfieldLocalPosition = float3(LocalPosition, HeightfieldHeight);
+	FLWCVector3 HeightfieldWorldPosition = LWCMultiply(HeightfieldLocalPosition, GetLocalToWorld(HeightfieldIndex));
 
-	return WorldHeightfieldShadingPosition;
+	return HeightfieldWorldPosition;
 }
 
-float3 GetHeightfieldWorldPosition(uint HeightfieldIndex, float2 LocalPosition, float2 HeightfieldUV)
+FLWCVector3 GetHeightfieldWorldPosition(uint HeightfieldIndex, float2 LocalPosition, float2 HeightfieldUV)
 {
 	float3 WorldNormal;
 	float Visibility;
diff --git a/Engine/Source/Runtime/Renderer/Private/GlobalDistanceField.cpp b/Engine/Source/Runtime/Renderer/Private/GlobalDistanceField.cpp
--- a/Engine/Source/Runtime/Renderer/Private/GlobalDistanceField.cpp
+++ b/Engine/Source/Runtime/Renderer/Private/GlobalDistanceField.cpp
@@ -1290,11 +1290,14 @@ class FCullObjectsToClipmapCS : public FGlobalShader
 		SHADER_PARAMETER_RDG_BUFFER_UAV(RWStructuredBuffer<uint>, RWObjectIndexBuffer)
 		SHADER_PARAMETER_RDG_BUFFER_UAV(RWStructuredBuffer<uint>, RWObjectIndexNumBuffer)
 		SHADER_PARAMETER_STRUCT_INCLUDE(FDistanceFieldObjectBufferParameters, DistanceFieldObjectBuffers)
-		SHADER_PARAMETER(FVector3f, ClipmapWorldCenter)
+		SHADER_PARAMETER(FVector3f, ClipmapTranslatedWorldCenter)
 		SHADER_PARAMETER(FVector3f, ClipmapWorldExtent)
 		SHADER_PARAMETER(uint32, AcceptOftenMovingObjectsOnly)
 		SHADER_PARAMETER(float, MeshSDFRadiusThreshold)
 		SHADER_PARAMETER(float, InfluenceRadiusSq)
+
+		SHADER_PARAMETER(FVector3f, ViewTilePosition)
+		SHADER_PARAMETER(FVector3f, RelativePreViewTranslation)
 	END_SHADER_PARAMETER_STRUCT()
 
 	static bool ShouldCompilePermutation(const FGlobalShaderPermutationParameters& Parameters)
@@ -1360,12 +1363,12 @@ class FBuildGridTilesCS : public FGlobalShader
 		SHADER_PARAMETER(uint32, NumUpdateBounds)
 		SHADER_PARAMETER(float, InfluenceRadiusSq)
 		SHADER_PARAMETER(FIntVector, PageGridResolution)
-		SHADER_PARAMETER(FVector3f, PageGridCoordToWorldCenterScale)
-		SHADER_PARAMETER(FVector3f, PageGridCoordToWorldCenterBias)
+		SHADER_PARAMETER(FVector3f, PageGridCoordToTranslatedWorldCenterScale)
+		SHADER_PARAMETER(FVector3f, PageGridCoordToTranslatedWorldCenterBias)
 		SHADER_PARAMETER(FVector3f, PageGridTileWorldExtent)
 		SHADER_PARAMETER(FIntVector, CullGridResolution)
-		SHADER_PARAMETER(FVector3f, CullGridCoordToWorldCenterScale)
-		SHADER_PARAMETER(FVector3f, CullGridCoordToWorldCenterBias)
+		SHADER_PARAMETER(FVector3f, CullGridCoordToTranslatedWorldCenterScale)
+		SHADER_PARAMETER(FVector3f, CullGridCoordToTranslatedWorldCenterBias)
 		SHADER_PARAMETER(FVector3f, CullGridTileWorldExtent)
 	END_SHADER_PARAMETER_STRUCT()
 
@@ -1404,10 +1407,13 @@ class FCullObjectsToGridCS : public FGlobalShader
 		SHADER_PARAMETER_RDG_BUFFER_SRV(StructuredBuffer<uint>, ObjectIndexNumBuffer)
 		SHADER_PARAMETER_STRUCT_INCLUDE(FDistanceFieldObjectBufferParameters, DistanceFieldObjectBuffers)
 		SHADER_PARAMETER(FIntVector, CullGridResolution)
-		SHADER_PARAMETER(FVector3f, CullGridCoordToWorldCenterScale)
-		SHADER_PARAMETER(FVector3f, CullGridCoordToWorldCenterBias)
+		SHADER_PARAMETER(FVector3f, CullGridCoordToTranslatedWorldCenterScale)
+		SHADER_PARAMETER(FVector3f, CullGridCoordToTranslatedWorldCenterBias)
 		SHADER_PARAMETER(FVector3f, CullTileWorldExtent)
 		SHADER_PARAMETER(float, InfluenceRadiusSq)
+
+		SHADER_PARAMETER(FVector3f, ViewTilePosition)
+		SHADER_PARAMETER(FVector3f, RelativePreViewTranslation)
 	END_SHADER_PARAMETER_STRUCT()
 
 	static bool ShouldCompilePermutation(const FGlobalShaderPermutationParameters& Parameters)
@@ -1446,19 +1452,20 @@ class FComposeObjectsIntoPagesCS : public FGlobalShader
 		SHADER_PARAMETER(float, InfluenceRadius)
 		SHADER_PARAMETER(float, InfluenceRadiusSq)
 		SHADER_PARAMETER(FIntVector, CullGridResolution)
-		SHADER_PARAMETER(FIntVector, GlobalDistanceFieldScrollOffset)
 		SHADER_PARAMETER(FVector3f, GlobalDistanceFieldInvPageAtlasSize)
 		SHADER_PARAMETER(FVector3f, InvPageGridResolution)
 		SHADER_PARAMETER(FIntVector, PageGridResolution)
 		SHADER_PARAMETER(FIntVector, ClipmapResolution)
-		SHADER_PARAMETER(FVector3f, PageCoordToVoxelCenterScale)
-		SHADER_PARAMETER(FVector3f, PageCoordToVoxelCenterBias)
-		SHADER_PARAMETER(FVector3f, PageCoordToPageWorldCenterScale)
-		SHADER_PARAMETER(FVector3f, PageCoordToPageWorldCenterBias)
-		SHADER_PARAMETER(FVector4f, ClipmapVolumeWorldToUVAddAndMul)
+		SHADER_PARAMETER(FVector3f, PageCoordToVoxelTranslatedCenterScale)
+		SHADER_PARAMETER(FVector3f, PageCoordToVoxelTranslatedCenterBias)
+		SHADER_PARAMETER(FVector3f, PageCoordToPageTranslatedWorldCenterScale)
+		SHADER_PARAMETER(FVector3f, PageCoordToPageTranslatedWorldCenterBias)
+		SHADER_PARAMETER(FVector4f, ClipmapVolumeTranslatedWorldToUVAddAndMul)
 		SHADER_PARAMETER(FVector3f, ComposeTileWorldExtent)
-		SHADER_PARAMETER(FVector3f, ClipmapMinBounds)
 		SHADER_PARAMETER(uint32, PageTableClipmapOffsetZ)
+
+		SHADER_PARAMETER(FVector3f, ViewTilePosition)
+		SHADER_PARAMETER(FVector3f, RelativePreViewTranslation)
 	END_SHADER_PARAMETER_STRUCT()
 
 	class FComposeParentDistanceField : SHADER_PERMUTATION_BOOL("COMPOSE_PARENT_DISTANCE_FIELD");
@@ -1544,14 +1551,17 @@ class FAllocatePagesCS : public FGlobalShader
 		SHADER_PARAMETER(FVector3f, PageWorldExtent)
 		SHADER_PARAMETER(float, PageWorldRadius)
 		SHADER_PARAMETER(float, ClipmapInfluenceRadius)
-		SHADER_PARAMETER(FVector3f, PageCoordToPageWorldCenterScale)
-		SHADER_PARAMETER(FVector3f, PageCoordToPageWorldCenterBias)
-		SHADER_PARAMETER(FVector4f, ClipmapVolumeWorldToUVAddAndMul)
+		SHADER_PARAMETER(FVector3f, PageCoordToPageTranslatedWorldCenterScale)
+		SHADER_PARAMETER(FVector3f, PageCoordToPageTranslatedWorldCenterBias)
+		SHADER_PARAMETER(FVector4f, ClipmapVolumeTranslatedWorldToUVAddAndMul)
 		SHADER_PARAMETER_RDG_BUFFER_SRV(StructuredBuffer<uint>, CullGridObjectHeader)
 		SHADER_PARAMETER_RDG_BUFFER_SRV(StructuredBuffer<uint>, CullGridObjectArray)
 		SHADER_PARAMETER(FIntVector, CullGridResolution)
 		SHADER_PARAMETER_STRUCT_INCLUDE(FDistanceFieldObjectBufferParameters, DistanceFieldObjectBuffers)
 		SHADER_PARAMETER_STRUCT_INCLUDE(FDistanceFieldAtlasParameters, DistanceFieldAtlas)
+
+		SHADER_PARAMETER(FVector3f, ViewTilePosition)
+		SHADER_PARAMETER(FVector3f, RelativePreViewTranslation)
 	END_SHADER_PARAMETER_STRUCT()
 
 	class FProcessDistanceFields : SHADER_PERMUTATION_BOOL("PROCESS_DISTANCE_FIELDS");
@@ -1900,20 +1910,27 @@ void UpdateGlobalDistanceFieldVolume(
 				{
 					RDG_EVENT_SCOPE(GraphBuilder, "Clipmap:%d CacheType:%s", ClipmapIndex, CacheType == GDF_MostlyStatic ? TEXT("MostlyStatic") : TEXT("Movable"));
 
-					FGlobalDistanceFieldClipmap& Clipmap = Clipmaps[ClipmapIndex];
+					const FGlobalDistanceFieldClipmap& Clipmap = Clipmaps[ClipmapIndex];
 
-					const int32 ClipmapResolution = GlobalDistanceField::GetClipmapResolution(bLumenEnabled);
 					const FVector ClipmapWorldCenter = Clipmap.Bounds.GetCenter();
-					const FVector ClipmapWorldExtent = Clipmap.Bounds.GetExtent();
-					const FVector ClipmapSize = Clipmap.Bounds.GetSize();
+
+					const FLargeWorldRenderPosition AbsoluteViewOrigin(ClipmapWorldCenter);
+					const FVector ViewTileOffset = AbsoluteViewOrigin.GetTileOffset();
+					const FVector PreViewTranslation = -ClipmapWorldCenter;
+
+					const FVector3f ViewTilePosition = AbsoluteViewOrigin.GetTile();
+					const FVector3f RelativePreViewTranslation = FVector3f(PreViewTranslation + ViewTileOffset);
+
+					const FBox TranslatedBounds = Clipmap.Bounds.ShiftBy(PreViewTranslation);
+
+					const int32 ClipmapResolution = GlobalDistanceField::GetClipmapResolution(bLumenEnabled);
+					const FVector ClipmapSize = TranslatedBounds.GetSize();
 					const FVector ClipmapVoxelSize = ClipmapSize / FVector(ClipmapResolution);
 					const FVector ClipmapVoxelExtent = 0.5f * ClipmapVoxelSize;
-					const float ClipmapVoxelRadius = ClipmapVoxelExtent.Size();
 					const float ClipmapInfluenceRadius = (GGlobalDistanceFieldInfluenceRangeInVoxels * ClipmapSize.X) / ClipmapResolution;
 
-					FVector4f ClipmapVolumeWorldToUVAddAndMul;
-					const FVector WorldToUVAdd = (Clipmap.ScrollOffset - Clipmap.Bounds.GetCenter()) / (Clipmap.Bounds.GetExtent().X * 2.0f) + FVector(0.5f);
-					ClipmapVolumeWorldToUVAddAndMul = FVector4f((FVector3f)WorldToUVAdd, 1.0f / (Clipmap.Bounds.GetExtent().X * 2.0f));
+					const FVector TranslatedWorldToUVAdd = (Clipmap.ScrollOffset - TranslatedBounds.GetCenter()) / (TranslatedBounds.GetExtent().X * 2.0f) + FVector(0.5f);
+					const FVector4f ClipmapVolumeTranslatedWorldToUVAddAndMul = FVector4f((FVector3f)TranslatedWorldToUVAdd, 1.0f / (TranslatedBounds.GetExtent().X * 2.0f));
 
 					int32 MaxSDFMeshObjects = FMath::RoundUpToPowerOfTwo(DistanceFieldSceneData.NumObjectsInBuffer);
 					FRDGBufferRef ObjectIndexBuffer = GraphBuilder.CreateBuffer(FRDGBufferDesc::CreateStructuredDesc(sizeof(uint32), MaxSDFMeshObjects), TEXT("ObjectIndices"));
@@ -1932,7 +1949,7 @@ void UpdateGlobalDistanceFieldVolume(
 						{
 							const FClipmapUpdateBounds& UpdateBounds = Clipmap.UpdateBounds[UpdateBoundsIndex];
 
-							UpdateBoundsData[NumUpdateBounds * BufferStrideInFloat4 + 0] = FVector4f((FVector3f)UpdateBounds.Center, UpdateBounds.bExpandByInfluenceRadius ? 1.0f : 0.0f);
+							UpdateBoundsData[NumUpdateBounds * BufferStrideInFloat4 + 0] = FVector4f((FVector3f)(UpdateBounds.Center + PreViewTranslation), UpdateBounds.bExpandByInfluenceRadius ? 1.0f : 0.0f);
 							UpdateBoundsData[NumUpdateBounds * BufferStrideInFloat4 + 1] = FVector4f((FVector3f)UpdateBounds.Extent, 0.0f);
 							++NumUpdateBounds;
 						}
@@ -2014,13 +2031,16 @@ void UpdateGlobalDistanceFieldVolume(
 							PassParameters->RWObjectIndexBuffer = GraphBuilder.CreateUAV(ObjectIndexBuffer, PF_R32_UINT);
 							PassParameters->RWObjectIndexNumBuffer = GraphBuilder.CreateUAV(ObjectIndexNumBuffer, PF_R32_UINT);
 							PassParameters->DistanceFieldObjectBuffers = DistanceField::SetupObjectBufferParameters(GraphBuilder, DistanceFieldSceneData);
-							PassParameters->ClipmapWorldCenter = (FVector3f)Clipmap.Bounds.GetCenter();
-							PassParameters->ClipmapWorldExtent = (FVector3f)Clipmap.Bounds.GetExtent();
+							PassParameters->ClipmapTranslatedWorldCenter = (FVector3f)TranslatedBounds.GetCenter();
+							PassParameters->ClipmapWorldExtent = (FVector3f)TranslatedBounds.GetExtent();
 							PassParameters->AcceptOftenMovingObjectsOnly = AcceptOftenMovingObjectsOnlyValue;
 							const float RadiusThresholdScale = bLumenEnabled ? 1.0f / FMath::Clamp(View.FinalPostProcessSettings.LumenSceneDetail, .01f, 100.0f) : 1.0f;
 							PassParameters->MeshSDFRadiusThreshold = GetMinMeshSDFRadius(ClipmapVoxelSize.X) * RadiusThresholdScale;
 							PassParameters->InfluenceRadiusSq = ClipmapInfluenceRadius * ClipmapInfluenceRadius;
 
+							PassParameters->ViewTilePosition = ViewTilePosition;
+							PassParameters->RelativePreViewTranslation = RelativePreViewTranslation;
+
 							auto ComputeShader = View.ShaderMap->GetShader<FCullObjectsToClipmapCS>();
 							const FIntVector GroupSize = FComputeShaderUtils::GetGroupCount(DistanceFieldSceneData.NumObjectsInBuffer, FCullObjectsToClipmapCS::GetGroupSize());
 
@@ -2040,14 +2060,18 @@ void UpdateGlobalDistanceFieldVolume(
 
 						const FVector PageTileWorldExtent = ClipmapVoxelExtent * GGlobalDistanceFieldPageResolutionInAtlas;
 						const FVector PageTileWorldExtentWithoutBorders = ClipmapVoxelExtent * GGlobalDistanceFieldPageResolution;
-						const FVector PageGridCoordToWorldCenterScale = ClipmapSize / FVector(PageGridResolution);
-						const FVector PageGridCoordToWorldCenterBias = Clipmap.Bounds.Min + 0.5f * PageGridCoordToWorldCenterScale;
+						const FVector PageGridCoordToTranslatedWorldCenterScale = ClipmapSize / FVector(PageGridResolution);
+						const FVector PageGridCoordToTranslatedWorldCenterBias = TranslatedBounds.Min + 0.5f * PageGridCoordToTranslatedWorldCenterScale;
 
 						const FIntVector CullGridResolution = FComputeShaderUtils::GetGroupCount(PageGridResolution, GlobalDistanceField::CullGridFactor);
 						const int32 CullGridSize = CullGridResolution.X * CullGridResolution.Y * CullGridResolution.Z;
 						const FVector CullTileWorldExtent = ClipmapVoxelExtent * GGlobalDistanceFieldPageResolutionInAtlas * GlobalDistanceField::CullGridFactor;
-						const FVector CullGridCoordToWorldCenterScale = ClipmapSize / FVector(CullGridResolution);
-						const FVector CullGridCoordToWorldCenterBias = Clipmap.Bounds.Min + 0.5 * CullGridCoordToWorldCenterScale;
+						const FVector CullGridCoordToTranslatedWorldCenterScale = ClipmapSize / FVector(CullGridResolution);
+						const FVector CullGridCoordToTranslatedWorldCenterBias = TranslatedBounds.Min + 0.5 * CullGridCoordToTranslatedWorldCenterScale;
+
+						const FVector PageVoxelExtent = 0.5f * ClipmapSize / FVector(ClipmapResolution);
+						const FVector PageCoordToVoxelTranslatedCenterScale = ClipmapSize / FVector(ClipmapResolution);
+						const FVector PageCoordToVoxelTranslatedCenterBias = TranslatedBounds.Min + PageVoxelExtent;
 
 						FRDGBufferRef PageUpdateTileBuffer = GraphBuilder.CreateBuffer(FRDGBufferDesc::CreateStructuredDesc(sizeof(uint32), PageGridSize), TEXT("PageUpdateTiles"));
 						FRDGBufferRef PageComposeTileBuffer = GraphBuilder.CreateBuffer(FRDGBufferDesc::CreateStructuredDesc(sizeof(uint32), PageGridSize), TEXT("PageComposeTiles"));
@@ -2088,13 +2112,13 @@ void UpdateGlobalDistanceFieldVolume(
 							PassParameters->InfluenceRadiusSq = ClipmapInfluenceRadius * ClipmapInfluenceRadius;
 							// Page grid
 							PassParameters->PageGridResolution = PageGridResolution;
-							PassParameters->PageGridCoordToWorldCenterScale = (FVector3f)PageGridCoordToWorldCenterScale;
-							PassParameters->PageGridCoordToWorldCenterBias = (FVector3f)PageGridCoordToWorldCenterBias;
+							PassParameters->PageGridCoordToTranslatedWorldCenterScale = (FVector3f)PageGridCoordToTranslatedWorldCenterScale;
+							PassParameters->PageGridCoordToTranslatedWorldCenterBias = (FVector3f)PageGridCoordToTranslatedWorldCenterBias;
 							PassParameters->PageGridTileWorldExtent = (FVector3f)PageTileWorldExtent;
 							// Cull grid
 							PassParameters->CullGridResolution = CullGridResolution;
-							PassParameters->CullGridCoordToWorldCenterScale = (FVector3f)CullGridCoordToWorldCenterScale;
-							PassParameters->CullGridCoordToWorldCenterBias = (FVector3f)CullGridCoordToWorldCenterBias;
+							PassParameters->CullGridCoordToTranslatedWorldCenterScale = (FVector3f)CullGridCoordToTranslatedWorldCenterScale;
+							PassParameters->CullGridCoordToTranslatedWorldCenterBias = (FVector3f)CullGridCoordToTranslatedWorldCenterBias;
 							PassParameters->CullGridTileWorldExtent = (FVector3f)CullTileWorldExtent;
 
 							auto ComputeShader = View.ShaderMap->GetShader<FBuildGridTilesCS>();
@@ -2118,10 +2142,6 @@ void UpdateGlobalDistanceFieldVolume(
 							MarkedHeightfieldPageBuffer = GraphBuilder.CreateBuffer(FRDGBufferDesc::CreateStructuredDesc(sizeof(uint32), PageGridSize), TEXT("MarkedHeightfieldPages"));
 							AddClearUAVPass(GraphBuilder, GraphBuilder.CreateUAV(MarkedHeightfieldPageBuffer, PF_R32_UINT), 0);
 
-							const FVector PageVoxelExtent = 0.5f * ClipmapSize / FVector(ClipmapResolution);
-							const FVector PageCoordToVoxelCenterScale = ClipmapSize / FVector(ClipmapResolution);
-							const FVector PageCoordToVoxelCenterBias = Clipmap.Bounds.Min + PageVoxelExtent;
-
 							for (TMap<FHeightfieldComponentTextures, TArray<FHeightfieldComponentDescription>>::TConstIterator It(UpdateRegionHeightfield.ComponentDescriptions); It; ++It)
 							{
 								const TArray<FHeightfieldComponentDescription>& HeightfieldDescriptions = It.Value();
@@ -2137,9 +2157,8 @@ void UpdateGlobalDistanceFieldVolume(
 									PassParameters->RWMarkedHeightfieldPageBuffer = GraphBuilder.CreateUAV(MarkedHeightfieldPageBuffer, PF_R32_UINT);
 									PassParameters->PageUpdateIndirectArgBuffer = PageUpdateIndirectArgBuffer;
 									PassParameters->PageUpdateTileBuffer = GraphBuilder.CreateSRV(PageUpdateTileBuffer, PF_R32_UINT);
-									PassParameters->InfluenceRadius = ClipmapInfluenceRadius;
-									PassParameters->PageCoordToPageWorldCenterScale = (FVector3f)PageGridCoordToWorldCenterScale;
-									PassParameters->PageCoordToPageWorldCenterBias = (FVector3f)PageGridCoordToWorldCenterBias;
+									PassParameters->PageCoordToPageTranslatedWorldCenterScale = (FVector3f)PageGridCoordToTranslatedWorldCenterScale;
+									PassParameters->PageCoordToPageTranslatedWorldCenterBias = (FVector3f)PageGridCoordToTranslatedWorldCenterBias;
 									PassParameters->PageWorldExtent = (FVector3f)PageTileWorldExtentWithoutBorders;
 									PassParameters->ClipmapVoxelExtent = ClipmapVoxelExtent.X;
 									PassParameters->PageGridResolution = PageGridResolution;
@@ -2152,6 +2171,9 @@ void UpdateGlobalDistanceFieldVolume(
 									PassParameters->VisibilitySampler = TStaticSamplerState<SF_Bilinear>::GetRHI();
 									PassParameters->HeightfieldDescriptions = GraphBuilder.CreateSRV(HeightfieldDescriptionBuffer, EPixelFormat::PF_A32B32G32R32F);
 
+									PassParameters->ViewTilePosition = ViewTilePosition;
+									PassParameters->RelativePreViewTranslation = RelativePreViewTranslation;
+
 									auto ComputeShader = View.ShaderMap->GetShader<FMarkHeightfieldPagesCS>();
 
 									FComputeShaderUtils::AddPass(
@@ -2230,11 +2252,14 @@ void UpdateGlobalDistanceFieldVolume(
 							PassParameters->ObjectIndexNumBuffer = GraphBuilder.CreateSRV(ObjectIndexNumBuffer, PF_R32_UINT);
 							PassParameters->DistanceFieldObjectBuffers = DistanceFieldObjectBuffers;
 							PassParameters->CullGridResolution = CullGridResolution;
-							PassParameters->CullGridCoordToWorldCenterScale = (FVector3f)CullGridCoordToWorldCenterScale;
-							PassParameters->CullGridCoordToWorldCenterBias = (FVector3f)CullGridCoordToWorldCenterBias;
+							PassParameters->CullGridCoordToTranslatedWorldCenterScale = (FVector3f)CullGridCoordToTranslatedWorldCenterScale;
+							PassParameters->CullGridCoordToTranslatedWorldCenterBias = (FVector3f)CullGridCoordToTranslatedWorldCenterBias;
 							PassParameters->CullTileWorldExtent = (FVector3f)CullTileWorldExtent;
 							PassParameters->InfluenceRadiusSq = ClipmapInfluenceRadius * ClipmapInfluenceRadius;
 
+							PassParameters->ViewTilePosition = ViewTilePosition;
+							PassParameters->RelativePreViewTranslation = RelativePreViewTranslation;
+
 							auto ComputeShader = View.ShaderMap->GetShader<FCullObjectsToGridCS>();
 
 							FComputeShaderUtils::AddPass(
@@ -2276,9 +2301,9 @@ void UpdateGlobalDistanceFieldVolume(
 								PassParameters->PageGridResolution = PageGridResolution;
 								PassParameters->InvPageGridResolution = FVector3f::OneVector / (FVector3f)PageGridResolution;
 								PassParameters->GlobalDistanceFieldMaxPageNum = GGlobalDistanceFieldMaxPageNum;
-								PassParameters->PageCoordToPageWorldCenterScale = (FVector3f)PageGridCoordToWorldCenterScale;
-								PassParameters->PageCoordToPageWorldCenterBias = (FVector3f)PageGridCoordToWorldCenterBias;
-								PassParameters->ClipmapVolumeWorldToUVAddAndMul = ClipmapVolumeWorldToUVAddAndMul;
+								PassParameters->PageCoordToPageTranslatedWorldCenterScale = (FVector3f)PageGridCoordToTranslatedWorldCenterScale;
+								PassParameters->PageCoordToPageTranslatedWorldCenterBias = (FVector3f)PageGridCoordToTranslatedWorldCenterBias;
+								PassParameters->ClipmapVolumeTranslatedWorldToUVAddAndMul = ClipmapVolumeTranslatedWorldToUVAddAndMul;
 								PassParameters->PageTableClipmapOffsetZ = ClipmapIndex * PageGridResolution.Z;
 
 								PassParameters->CullGridObjectHeader = GraphBuilder.CreateSRV(CullGridObjectHeader, PF_R32_UINT);
@@ -2288,6 +2313,9 @@ void UpdateGlobalDistanceFieldVolume(
 								PassParameters->DistanceFieldObjectBuffers = DistanceFieldObjectBuffers;
 								PassParameters->DistanceFieldAtlas = DistanceFieldAtlas;
 
+								PassParameters->ViewTilePosition = ViewTilePosition;
+								PassParameters->RelativePreViewTranslation = RelativePreViewTranslation;
+
 								FAllocatePagesCS::FPermutationDomain PermutationVector;
 								PermutationVector.Set<FAllocatePagesCS::FProcessDistanceFields>(Scene->DistanceFieldSceneData.NumObjectsInBuffer > 0);
 								PermutationVector.Set<FAllocatePagesCS::FMarkedHeightfieldPageBuffer>(MarkedHeightfieldPageBuffer != nullptr);
@@ -2348,10 +2376,6 @@ void UpdateGlobalDistanceFieldVolume(
 						// Initialize pages and compose the mesh SDFs into allocated pages
 						if(Scene->DistanceFieldSceneData.NumObjectsInBuffer > 0 || UpdateRegionHeightfield.ComponentDescriptions.Num() > 0)
 						{
-							const FVector PageVoxelExtent = 0.5f * ClipmapSize / FVector(ClipmapResolution);
-							const FVector PageCoordToVoxelCenterScale = ClipmapSize / FVector(ClipmapResolution);
-							const FVector PageCoordToVoxelCenterBias = Clipmap.Bounds.Min + PageVoxelExtent;
-
 							const uint32 PageComposeTileSize = 4;
 							const FVector PageComposeTileWorldExtent = ClipmapVoxelExtent * PageComposeTileSize;
 
@@ -2375,15 +2399,17 @@ void UpdateGlobalDistanceFieldVolume(
 							PassParameters->PageGridResolution = PageGridResolution;
 							PassParameters->InvPageGridResolution = FVector3f::OneVector / (FVector3f)PageGridResolution;
 							PassParameters->ClipmapResolution = FIntVector(ClipmapResolution);
-							PassParameters->PageCoordToVoxelCenterScale = (FVector3f)PageCoordToVoxelCenterScale;
-							PassParameters->PageCoordToVoxelCenterBias = (FVector3f)PageCoordToVoxelCenterBias;
+							PassParameters->PageCoordToVoxelTranslatedCenterScale = (FVector3f)PageCoordToVoxelTranslatedCenterScale;
+							PassParameters->PageCoordToVoxelTranslatedCenterBias = (FVector3f)PageCoordToVoxelTranslatedCenterBias;
 							PassParameters->ComposeTileWorldExtent = (FVector3f)PageComposeTileWorldExtent;
-							PassParameters->ClipmapMinBounds = (FVector3f)Clipmap.Bounds.Min;
-							PassParameters->PageCoordToPageWorldCenterScale = (FVector3f)PageGridCoordToWorldCenterScale;
-							PassParameters->PageCoordToPageWorldCenterBias = (FVector3f)PageGridCoordToWorldCenterBias;
-							PassParameters->ClipmapVolumeWorldToUVAddAndMul = ClipmapVolumeWorldToUVAddAndMul;
+							PassParameters->PageCoordToPageTranslatedWorldCenterScale = (FVector3f)PageGridCoordToTranslatedWorldCenterScale;
+							PassParameters->PageCoordToPageTranslatedWorldCenterBias = (FVector3f)PageGridCoordToTranslatedWorldCenterBias;
+							PassParameters->ClipmapVolumeTranslatedWorldToUVAddAndMul = ClipmapVolumeTranslatedWorldToUVAddAndMul;
 							PassParameters->PageTableClipmapOffsetZ = ClipmapIndex * PageGridResolution.Z;
 
+							PassParameters->ViewTilePosition = ViewTilePosition;
+							PassParameters->RelativePreViewTranslation = RelativePreViewTranslation;
+
 							FComposeObjectsIntoPagesCS::FPermutationDomain PermutationVector;
 							PermutationVector.Set<FComposeObjectsIntoPagesCS::FComposeParentDistanceField>(ParentPageTableLayerTexture != nullptr);
 							PermutationVector.Set<FComposeObjectsIntoPagesCS::FProcessDistanceFields>(Scene->DistanceFieldSceneData.NumObjectsInBuffer > 0);
@@ -2406,10 +2432,6 @@ void UpdateGlobalDistanceFieldVolume(
 						{
 							RDG_EVENT_SCOPE(GraphBuilder, "ComposeHeightfieldsIntoPages");
 
-							const FVector PageVoxelExtent = 0.5f * ClipmapSize / FVector(ClipmapResolution);
-							const FVector PageCoordToVoxelCenterScale = ClipmapSize / FVector(ClipmapResolution);
-							const FVector PageCoordToVoxelCenterBias = Clipmap.Bounds.Min + PageVoxelExtent;
-
 							for (TMap<FHeightfieldComponentTextures, TArray<FHeightfieldComponentDescription>>::TConstIterator It(UpdateRegionHeightfield.ComponentDescriptions); It; ++It)
 							{
 								const TArray<FHeightfieldComponentDescription>& HeightfieldDescriptions = It.Value();
@@ -2427,15 +2449,14 @@ void UpdateGlobalDistanceFieldVolume(
 									PassParameters->ComposeIndirectArgBuffer = PageComposeHeightfieldIndirectArgBuffer;
 									PassParameters->ComposeTileBuffer = GraphBuilder.CreateSRV(PageComposeHeightfieldTileBuffer, PF_R32_UINT);
 									PassParameters->PageTableLayerTexture = PageTableLayerTexture;
-									PassParameters->InfluenceRadius = ClipmapInfluenceRadius;
-									PassParameters->PageCoordToVoxelCenterScale = (FVector3f)PageCoordToVoxelCenterScale;
-									PassParameters->PageCoordToVoxelCenterBias = (FVector3f)PageCoordToVoxelCenterBias;
+									PassParameters->PageCoordToVoxelTranslatedCenterScale = (FVector3f)PageCoordToVoxelTranslatedCenterScale;
+									PassParameters->PageCoordToVoxelTranslatedCenterBias = (FVector3f)PageCoordToVoxelTranslatedCenterBias;
 									PassParameters->ClipmapVoxelExtent = ClipmapVoxelExtent.X;
 									PassParameters->PageGridResolution = PageGridResolution;
 									PassParameters->InvPageGridResolution = FVector3f::OneVector / (FVector3f)PageGridResolution;
-									PassParameters->PageCoordToPageWorldCenterScale = (FVector3f)PageGridCoordToWorldCenterScale;
-									PassParameters->PageCoordToPageWorldCenterBias = (FVector3f)PageGridCoordToWorldCenterBias;
-									PassParameters->ClipmapVolumeWorldToUVAddAndMul = ClipmapVolumeWorldToUVAddAndMul;
+									PassParameters->PageCoordToPageTranslatedWorldCenterScale = (FVector3f)PageGridCoordToTranslatedWorldCenterScale;
+									PassParameters->PageCoordToPageTranslatedWorldCenterBias = (FVector3f)PageGridCoordToTranslatedWorldCenterBias;
+									PassParameters->ClipmapVolumeTranslatedWorldToUVAddAndMul = ClipmapVolumeTranslatedWorldToUVAddAndMul;
 									PassParameters->PageTableClipmapOffsetZ = ClipmapIndex * PageGridResolution.Z;
 									PassParameters->NumHeightfields = HeightfieldDescriptions.Num();
 									PassParameters->InfluenceRadius = ClipmapInfluenceRadius;
@@ -2446,6 +2467,9 @@ void UpdateGlobalDistanceFieldVolume(
 									PassParameters->VisibilitySampler = TStaticSamplerState<SF_Bilinear>::GetRHI();
 									PassParameters->HeightfieldDescriptions = GraphBuilder.CreateSRV(HeightfieldDescriptionBuffer, EPixelFormat::PF_A32B32G32R32F);
 
+									PassParameters->ViewTilePosition = ViewTilePosition;
+									PassParameters->RelativePreViewTranslation = RelativePreViewTranslation;
+
 									FComposeHeightfieldsIntoPagesCS::FPermutationDomain PermutationVector;
 									PermutationVector.Set<FComposeHeightfieldsIntoPagesCS::FCompositeCoverageAtlas>(CoverageAtlasTexture != nullptr);
 									auto ComputeShader = View.ShaderMap->GetShader<FComposeHeightfieldsIntoPagesCS>(PermutationVector);
diff --git a/Engine/Source/Runtime/Renderer/Private/GlobalDistanceFieldHeightfields.h b/Engine/Source/Runtime/Renderer/Private/GlobalDistanceFieldHeightfields.h
--- a/Engine/Source/Runtime/Renderer/Private/GlobalDistanceFieldHeightfields.h
+++ b/Engine/Source/Runtime/Renderer/Private/GlobalDistanceFieldHeightfields.h
@@ -15,8 +15,8 @@ class FMarkHeightfieldPagesCS : public FGlobalShader
 		SHADER_PARAMETER_RDG_BUFFER_UAV(RWStructuredBuffer<uint>, RWMarkedHeightfieldPageBuffer)
 		RDG_BUFFER_ACCESS(PageUpdateIndirectArgBuffer, ERHIAccess::IndirectArgs)
 		SHADER_PARAMETER_RDG_BUFFER_SRV(StructuredBuffer<uint>, PageUpdateTileBuffer)
-		SHADER_PARAMETER(FVector3f, PageCoordToPageWorldCenterScale)
-		SHADER_PARAMETER(FVector3f, PageCoordToPageWorldCenterBias)
+		SHADER_PARAMETER(FVector3f, PageCoordToPageTranslatedWorldCenterScale)
+		SHADER_PARAMETER(FVector3f, PageCoordToPageTranslatedWorldCenterBias)
 		SHADER_PARAMETER(FVector3f, PageWorldExtent)
 		SHADER_PARAMETER(FVector3f, InvPageGridResolution)
 		SHADER_PARAMETER(FIntVector, PageGridResolution)
@@ -29,6 +29,9 @@ class FMarkHeightfieldPagesCS : public FGlobalShader
 		SHADER_PARAMETER_RDG_BUFFER_SRV(Buffer<float4>, HeightfieldDescriptions)
 		SHADER_PARAMETER(uint32, NumHeightfields)
 		SHADER_PARAMETER(float, HeightfieldThickness)
+
+		SHADER_PARAMETER(FVector3f, ViewTilePosition)
+		SHADER_PARAMETER(FVector3f, RelativePreViewTranslation)
 	END_SHADER_PARAMETER_STRUCT()
 
 	static bool ShouldCompilePermutation(const FGlobalShaderPermutationParameters& Parameters)
@@ -127,11 +130,11 @@ class FComposeHeightfieldsIntoPagesCS : public FGlobalShader
 		SHADER_PARAMETER_RDG_TEXTURE(Texture3D<uint>, PageTableLayerTexture)
 		SHADER_PARAMETER(FVector3f, InvPageGridResolution)
 		SHADER_PARAMETER(FIntVector, PageGridResolution)
-		SHADER_PARAMETER(FVector3f, PageCoordToVoxelCenterScale)
-		SHADER_PARAMETER(FVector3f, PageCoordToVoxelCenterBias)
-		SHADER_PARAMETER(FVector3f, PageCoordToPageWorldCenterScale)
-		SHADER_PARAMETER(FVector3f, PageCoordToPageWorldCenterBias)
-		SHADER_PARAMETER(FVector4f, ClipmapVolumeWorldToUVAddAndMul)
+		SHADER_PARAMETER(FVector3f, PageCoordToVoxelTranslatedCenterScale)
+		SHADER_PARAMETER(FVector3f, PageCoordToVoxelTranslatedCenterBias)
+		SHADER_PARAMETER(FVector3f, PageCoordToPageTranslatedWorldCenterScale)
+		SHADER_PARAMETER(FVector3f, PageCoordToPageTranslatedWorldCenterBias)
+		SHADER_PARAMETER(FVector4f, ClipmapVolumeTranslatedWorldToUVAddAndMul)
 		SHADER_PARAMETER(float, ClipmapVoxelExtent)
 		SHADER_PARAMETER(float, InfluenceRadius)
 		SHADER_PARAMETER(uint32, PageTableClipmapOffsetZ)
@@ -142,6 +145,9 @@ class FComposeHeightfieldsIntoPagesCS : public FGlobalShader
 		SHADER_PARAMETER_RDG_BUFFER_SRV(Buffer<float4>, HeightfieldDescriptions)
 		SHADER_PARAMETER(uint32, NumHeightfields)
 		SHADER_PARAMETER(float, HeightfieldThickness)
+
+		SHADER_PARAMETER(FVector3f, ViewTilePosition)
+		SHADER_PARAMETER(FVector3f, RelativePreViewTranslation)
 	END_SHADER_PARAMETER_STRUCT()
 
 	class FCompositeCoverageAtlas : SHADER_PERMUTATION_BOOL("COMPOSITE_COVERAGE_ATLAS");
diff --git a/Engine/Source/Runtime/Renderer/Private/HeightfieldLighting.cpp b/Engine/Source/Runtime/Renderer/Private/HeightfieldLighting.cpp
--- a/Engine/Source/Runtime/Renderer/Private/HeightfieldLighting.cpp
+++ b/Engine/Source/Runtime/Renderer/Private/HeightfieldLighting.cpp
@@ -18,7 +18,7 @@ HeightfieldLighting.cpp
 #include "PipelineStateCache.h"
 
 // In float4's, must match usf
-static const int32 HEIGHTFIELD_DATA_STRIDE = 10;
+static const int32 HEIGHTFIELD_DATA_STRIDE = 11;
 
 void FillHeightfieldDescriptionData(const TArray<FHeightfieldComponentDescription>& HeightfieldDescriptions, TArray<FVector4f, SceneRenderingAllocator>& HeightfieldDescriptionData)
 {
@@ -42,13 +42,22 @@ void FillHeightfieldDescriptionData(const TArray<FHeightfieldComponentDescriptio
 
 		HeightfieldDescriptionData.Add(FVector4f(Description.HeightfieldRect.Size().X, Description.HeightfieldRect.Size().Y, 1.f / Description.HeightfieldRect.Size().X, 1.f / Description.HeightfieldRect.Size().Y));
 
-		const FMatrix44f LocalToWorldT = FMatrix44f(Description.LocalToWorld.GetTransposed());
-		const FMatrix44f WorldToLocalT = FMatrix44f(Description.LocalToWorld.Inverse().GetTransposed());
+		const FLargeWorldRenderPosition WorldPosition(Description.LocalToWorld.GetOrigin());
+		const FVector TilePositionOffset = WorldPosition.GetTileOffset();
+
+		// Inverse on FMatrix44f can generate NaNs if the source matrix contains large scaling, so do it in double precision.
+		const FMatrix LocalToRelativeWorld = FLargeWorldRenderScalar::MakeToRelativeWorldMatrixDouble(TilePositionOffset, Description.LocalToWorld);
+
+		HeightfieldDescriptionData.Add(WorldPosition.GetTile());
+
+		const FMatrix44f WorldToLocalT = FMatrix44f(LocalToRelativeWorld.Inverse().GetTransposed());
 
 		HeightfieldDescriptionData.Add(*(FVector4f*)&WorldToLocalT.M[0]);
 		HeightfieldDescriptionData.Add(*(FVector4f*)&WorldToLocalT.M[1]);
 		HeightfieldDescriptionData.Add(*(FVector4f*)&WorldToLocalT.M[2]);
 
+		const FMatrix44f LocalToWorldT = FMatrix44f(LocalToRelativeWorld.GetTransposed());
+
 		HeightfieldDescriptionData.Add(*(FVector4f*)&LocalToWorldT.M[0]);
 		HeightfieldDescriptionData.Add(*(FVector4f*)&LocalToWorldT.M[1]);
 		HeightfieldDescriptionData.Add(*(FVector4f*)&LocalToWorldT.M[2]);
