Fix a bug that the screen is black on mobile preview.

diff --git a/Engine/Source/Runtime/Renderer/Private/PostProcess/PostProcessCompositeEditorPrimitives.cpp b/Engine/Source/Runtime/Renderer/Private/PostProcess/PostProcessCompositeEditorPrimitives.cpp
--- a/Engine/Source/Runtime/Renderer/Private/PostProcess/PostProcessCompositeEditorPrimitives.cpp
+++ b/Engine/Source/Runtime/Renderer/Private/PostProcess/PostProcessCompositeEditorPrimitives.cpp
@@ -606,6 +606,7 @@ FScreenPassTexture AddEditorPrimitivePass(
 		PassParameters->ReflectionCapture = View.ReflectionCaptureUniformBuffer;
 		PassParameters->InstanceCulling = InstanceCullingManager.GetDummyInstanceCullingUniformBuffer();
 		PassParameters->RenderTargets[0] = Output.GetRenderTargetBinding();
+		PassParameters->RenderTargets[0].SetLoadAction(ERenderTargetLoadAction::ELoad);
 
 		const FEditorPrimitiveInputs::EBasePassType BasePassType = Inputs.BasePassType;
 
