Add basic analytic support for Strata.

diff --git a/Engine/Source/Runtime/Renderer/Private/Strata/Strata.cpp b/Engine/Source/Runtime/Renderer/Private/Strata/Strata.cpp
--- a/Engine/Source/Runtime/Renderer/Private/Strata/Strata.cpp
+++ b/Engine/Source/Runtime/Renderer/Private/Strata/Strata.cpp
@@ -12,6 +12,7 @@
 #include "ShaderCompiler.h"
 #include "Lumen/Lumen.h"
 #include "RendererUtils.h"
+#include "EngineAnalytics.h"
 
 //PRAGMA_DISABLE_OPTIMIZATION
 
@@ -259,6 +260,23 @@ static bool NeedBSDFOffsets(const FScene* Scene, const FViewInfo& View)
 	return  ShouldRenderLumenDiffuseGI(Scene, View) || ShouldRenderLumenReflections(View) || Strata::ShouldRenderStrataDebugPasses(View);
 }
 
+
+static void RecordStrataAnalytics()
+{
+	if (FEngineAnalytics::IsAvailable())
+	{
+		const uint32 MaterialConservativeByteCountPerPixel = CVarStrataBytePerPixel.GetValueOnAnyThread();
+
+		TArray<FAnalyticsEventAttribute> EventAttributes;
+		EventAttributes.Add(FAnalyticsEventAttribute(TEXT("Enabled"), 1));
+		EventAttributes.Add(FAnalyticsEventAttribute(TEXT("BytesPerPixel"), MaterialConservativeByteCountPerPixel));
+
+		FString OutStr(TEXT("Strata"));
+		FEngineAnalytics::GetProvider().RecordEvent(OutStr, EventAttributes);
+	}
+}
+
+
 void InitialiseStrataFrameSceneData(FRDGBuilder& GraphBuilder, FSceneRenderer& SceneRenderer)
 {
 	FStrataSceneData& Out = SceneRenderer.Scene->StrataSceneData;
@@ -281,6 +299,14 @@ void InitialiseStrataFrameSceneData(FRDGBuilder& GraphBuilder, FSceneRenderer& S
 	UpdateMaterialBufferToTiledResolution(FIntPoint(1, 1), MaterialBufferSizeXY);
 	if (IsStrataEnabled())
 	{
+		// Analytics for tracking Strata usage
+		static bool bAnalyticsInitialized = false;
+		if (!bAnalyticsInitialized)
+		{
+			RecordStrataAnalytics();
+			bAnalyticsInitialized = true;
+		}
+
 		FIntPoint SceneTextureExtent = SceneRenderer.GetActiveSceneTexturesConfig().Extent;
 		
 		// We need to allocate enough for the tiled memory addressing of material data to always work
