Implemented classification pass to classify hair segments to be rasterized either with hardware or software. Also cleaned up hair visibility rasterization shader resources and implemented support for hair culling in the HW rasterizer and the naive SW rasterizer.

#rnx

diff --git a/Engine/Shaders/Private/HairStrands/HairStrandsVisibilityRasterClassification.usf b/Engine/Shaders/Private/HairStrands/HairStrandsVisibilityRasterClassification.usf
new file mode 100644
--- /dev/null
+++ b/Engine/Shaders/Private/HairStrands/HairStrandsVisibilityRasterClassification.usf
@@ -0,0 +1,149 @@
+// Copyright Epic Games, Inc. All Rights Reserved.
+
+#define HAIR_STRANDS_PARAMETERS 1
+
+#ifndef PERMUTATION_CULLING
+#define PERMUTATION_CULLING 0
+#endif
+
+#include "../Common.ush"
+#include "HairStrandsVisibilityCommon.ush"
+#include "HairStrandsVertexFactoryCommon.ush"
+
+RWBuffer<uint>		OutPrimIDs;
+RWByteAddressBuffer OutDrawIndexedArgs;
+RWByteAddressBuffer OutSWRasterPrimitiveCount;
+uint				ControlPointStart;
+uint				ControlPointCount;
+uint				PrimIDsBufferSize;
+uint				NumWorkGroups;
+
+bool ShouldSoftwareRasterize(float3 PositionOffset, uint PrimID, out bool bSegValid)
+{
+	FHairControlPoint CP0 = UnpackHairControlPoint(
+		HairStrandsVF_PositionBuffer[PrimID],
+		PositionOffset,
+		HairStrandsVF_Radius,
+		HairStrandsVF_RootScale,
+		HairStrandsVF_TipScale);
+
+	bSegValid = CP0.Type != HAIR_CONTROLPOINT_END;
+
+	if (!bSegValid)
+	{
+		return false;
+	}
+
+	FHairControlPoint CP1 = UnpackHairControlPoint(
+		HairStrandsVF_PositionBuffer[PrimID + 1],
+		PositionOffset,
+		HairStrandsVF_Radius,
+		HairStrandsVF_RootScale,
+		HairStrandsVF_TipScale);
+
+	// Transform to world space
+	const float3 CP0WP = mul(float4(CP0.Position, 1.0f), HairStrandsVF_LocalToWorldPrimitiveTransform).xyz;
+	const float3 CP1WP = mul(float4(CP1.Position, 1.0f), HairStrandsVF_LocalToWorldPrimitiveTransform).xyz;
+	float4 SP0 = mul(float4(CP0WP, 1.0f), LWCHackToFloat(PrimaryView.WorldToClip));
+	float4 SP1 = mul(float4(CP1WP, 1.0f), LWCHackToFloat(PrimaryView.WorldToClip));
+
+	return (PrimID & 1) == 0; // TODO
+}
+
+[numthreads(1024, 1, 1)]
+void CSMain(uint DispatchThreadID : SV_DispatchThreadID, uint GroupThreadID : SV_GroupThreadID, uint GroupID : SV_GroupID)
+{
+	ResolvedView = ResolveView();
+
+	// Write IndexCountPerInstance
+	if (DispatchThreadID == 0)
+	{
+		OutDrawIndexedArgs.Store(0, 6);
+	}
+
+#if PERMUTATION_CULLING	
+	uint NumVerticesAfterCulling = ControlPointCount;
+	if (HairStrandsVF_bIsCullingEnable)
+	{
+		NumVerticesAfterCulling = HairStrandsVF_CullingIndirectBuffer[3];
+	}
+	const uint NumBatches = (NumVerticesAfterCulling + 1023) / 1024;
+#else
+	const uint NumBatches = (ControlPointCount + 1023) / 1024;
+#endif // PERMUTATION_CULLING
+	const uint NumLoops = (NumBatches + (NumWorkGroups - 1)) / NumWorkGroups;
+	const float3 PositionOffset = HairStrandsVF_PositionOffsetBuffer[0].xyz;
+
+	LOOP for (uint LoopIndex = 0; LoopIndex < NumLoops; ++LoopIndex)
+	{
+		const uint BatchIndex = LoopIndex + (GroupID * NumLoops);
+
+#if PERMUTATION_CULLING	
+		uint PrimID = BatchIndex * 1024 + GroupThreadID;
+		bool bSegValid = (BatchIndex < NumBatches) && (PrimID < NumVerticesAfterCulling);
+
+		if (bSegValid && HairStrandsVF_bIsCullingEnable)
+		{
+			const uint FetchIndex0 = PrimID;
+			const uint FetchIndex1 = min(FetchIndex0 + 1, NumVerticesAfterCulling - 1);
+
+			const uint VertexIndex0 = HairStrandsVF_CullingIndexBuffer[FetchIndex0];
+			const uint VertexIndex1 = HairStrandsVF_CullingIndexBuffer[FetchIndex1];
+			const float LodRadiusScale0 = HairStrandsVF_CullingRadiusScaleBuffer[FetchIndex0];
+			const float LodRadiusScale1 = HairStrandsVF_CullingRadiusScaleBuffer[FetchIndex1];
+
+			if ((LodRadiusScale0 <= 0.0f) || (VertexIndex1 != (VertexIndex0 + 1)))
+			{
+				bSegValid = false;
+			}
+			else
+			{
+				PrimID = VertexIndex0;
+			}
+		}
+		else if (bSegValid)
+		{
+			PrimID += ControlPointStart;
+		}
+#else
+		const uint PrimID = BatchIndex * 1024 + GroupThreadID + ControlPointStart;
+		bool bSegValid = (BatchIndex < NumBatches) && (PrimID < (ControlPointStart + ControlPointCount));
+#endif // PERMUTATION_CULLING
+		
+		bool bShouldSoftwareRasterize = false;
+		if (bSegValid)
+		{
+			bShouldSoftwareRasterize = ShouldSoftwareRasterize(PositionOffset, PrimID, bSegValid);
+		}
+
+		const uint NumSWSegs = WaveActiveCountBits(bSegValid && bShouldSoftwareRasterize);
+		const uint NumHWSegs = WaveActiveCountBits(bSegValid && !bShouldSoftwareRasterize);
+		
+		uint WriteOffsetSW = 0;
+		uint WriteOffsetHW = 0;
+		if (WaveIsFirstLane())
+		{
+			if (NumSWSegs > 0)
+			{
+				OutSWRasterPrimitiveCount.InterlockedAdd(0, NumSWSegs, WriteOffsetSW);
+			}
+			if (NumHWSegs > 0)
+			{
+				// Increment InstanceCount
+				OutDrawIndexedArgs.InterlockedAdd(1 << 2, NumHWSegs, WriteOffsetHW);
+			}
+		}
+
+		WriteOffsetSW = WaveReadLaneFirst(WriteOffsetSW);
+		WriteOffsetHW = WaveReadLaneFirst(WriteOffsetHW);
+		
+		const uint LocalWritePosSW = WavePrefixCountBits(bSegValid && bShouldSoftwareRasterize);
+		const uint LocalWritePosHW = WavePrefixCountBits(bSegValid && !bShouldSoftwareRasterize);
+
+		if (bSegValid)
+		{
+			const uint WritePos = bShouldSoftwareRasterize ? (WriteOffsetSW + LocalWritePosSW) : (PrimIDsBufferSize - 1 - (WriteOffsetHW + NumHWSegs) + LocalWritePosHW);
+			OutPrimIDs[WritePos] = PrimID;
+		}
+	}
+}
\ No newline at end of file
diff --git a/Engine/Shaders/Private/HairStrands/HairStrandsVisibilityRasterCompute.usf b/Engine/Shaders/Private/HairStrands/HairStrandsVisibilityRasterCompute.usf
--- a/Engine/Shaders/Private/HairStrands/HairStrandsVisibilityRasterCompute.usf
+++ b/Engine/Shaders/Private/HairStrands/HairStrandsVisibilityRasterCompute.usf
@@ -31,21 +31,6 @@ float2	OutputResolutionf;
 
 #if SHADER_RASTERCOMPUTE_BINNING || SHADER_RASTERCOMPUTE_COMPACTION || SHADER_RASTERCOMPUTE_RASTER || SHADER_RASTERCOMPUTE_DEPTH_GRID
 
-///////////////////////////////////////////////////////////////////////////
-// 
-// Wave size
-#if SHADER_RASTERCOMPUTE_DEPTH_GRID == 0
-
-#if PERMUTATION_GROUP_SIZE == 64
-#define WAVE_SIZE 32
-#elif PERMUTATION_GROUP_SIZE == 32
-#define WAVE_SIZE 32
-#else
-#error Unknown group size	
-#endif
-
-#endif //SHADER_RASTERCOMPUTE_DEPTH_GRID
-
 ///////////////////////////////////////////////////////////////////////////
 
 /* 
@@ -70,7 +55,7 @@ uint PackVisTileCoord(uint2 Coord)
 	return uint(((Coord.x & 0xff) << 0) | (((Coord.y) & 0xff) << 8));
 }
 
-uint2 UppackVisTileCoord(uint Packed)
+uint2 UnpackVisTileCoord(uint Packed)
 {
 	return uint2(((Packed >> 0) & 0xff), ((Packed >> 8) & 0xff));
 }
@@ -274,11 +259,17 @@ void PrepareDepthGridCS(uint dId : SV_DispatchThreadID, uint GroupThreadID : SV_
 
 #if SHADER_RASTERCOMPUTE_BINNING
 
+#ifndef PERMUTATION_INDIRECT_PRIM_IDS
+#define PERMUTATION_INDIRECT_PRIM_IDS 0
+#endif
+
 RWTexture2DArray<uint> 				OutVisTileBinningGrid;
-Texture2D<uint>						VisTileDepthGrid;
 RWBuffer<uint>						OutVisTilePrims;
 RWBuffer<uint>						OutVisTileArgs;
 RWByteAddressBuffer					OutVisTileData;
+Texture2D<uint>						VisTileDepthGrid;
+ByteAddressBuffer					IndirectPrimIDCount;
+Buffer<uint>						IndirectPrimIDs;
 
 // TODO: Without setting a limit, we are sometimes caught in an infinite loop even though this should not happen.
 #define DDA_MAX_ITERATIONS 256
@@ -329,6 +320,7 @@ void StoreOutVisTileData(uint index, uint offset, uint value)
 }
 
 groupshared uint group_LoopNum;
+groupshared uint group_VerticesNum;
 groupshared uint group_BatchNum;
 
 #define TILES_TO_ALLOCATE_MAX 1024
@@ -344,75 +336,68 @@ void BinningCS(uint DispatchThreadID : SV_DispatchThreadID, uint GroupThreadID :
 	ResolvedView = ResolveView();
 	if (GroupThreadID == 0)
 	{
-#if PERMUTATION_CULLING	
-		if (HairStrandsVF_bIsCullingEnable)
-		{
-			group_BatchNum = (HairStrandsVF_CullingIndirectBuffer[3] + 1023) / 1024;
-		}
-		else
-		{
-			group_BatchNum = (VertexCount + 1023) / 1024;
-		}
-#else
-		group_BatchNum = (VertexCount + 1023) / 1024;
-#endif
-		
+#if PERMUTATION_INDIRECT_PRIM_IDS
+		group_VerticesNum = IndirectPrimIDCount.Load(0);
+#else // PERMUTATION_INDIRECT_PRIM_IDS
+	#if PERMUTATION_CULLING	
+		group_VerticesNum = HairStrandsVF_bIsCullingEnable ? HairStrandsVF_CullingIndirectBuffer[3] : VertexCount;
+	#else // PERMUTATION_CULLING
+		group_VerticesNum = VertexCount;
+	#endif //PERMUTATION_CULLING
+#endif // PERMUTATION_INDIRECT_PRIM_IDS
+
+		group_BatchNum = (group_VerticesNum + 1023) / 1024;
 		group_LoopNum = (group_BatchNum + (NumBinners - 1)) * RcpNumBinners;
 	}
 
 	GroupMemoryBarrierWithGroupSync();
 
-	LOOP
-	for (uint LoopIndex = 0; LoopIndex < group_LoopNum; LoopIndex++)
+	LOOP for (uint LoopIndex = 0; LoopIndex < group_LoopNum; LoopIndex++)
 	{
-		uint BatchIndex = LoopIndex + (GroupID * group_LoopNum);
-
+		const uint BatchIndex = LoopIndex + (GroupID * group_LoopNum);
 		bool bSegValid = (BatchIndex < group_BatchNum);
 
-#if PERMUTATION_CULLING
-		uint PrimID = BatchIndex * 1024 + GroupThreadID;
-		if (HairStrandsVF_bIsCullingEnable)
-		{
-			group_BatchNum = (HairStrandsVF_CullingIndirectBuffer[3] + 1023) / 1024;
-			bSegValid = bSegValid && (PrimID < (HairStrandsVF_CullingIndirectBuffer[3]));
-		}
-		else
+#if PERMUTATION_INDIRECT_PRIM_IDS
+		uint PrimID = 0;
+		const uint PrimIDIndex = BatchIndex * 1024 + GroupThreadID;
+		bSegValid = bSegValid && (PrimIDIndex < group_VerticesNum);
+		if (bSegValid)
 		{
-			bSegValid = bSegValid && (PrimID < (VertexCount));
+			PrimID = IndirectPrimIDs[PrimIDIndex];
 		}
+#else // PERMUTATION_INDIRECT_PRIM_IDS
+	#if PERMUTATION_CULLING
+		uint PrimID = BatchIndex * 1024 + GroupThreadID;
+		bSegValid = bSegValid && (PrimID < group_VerticesNum);
 
 		if (bSegValid && HairStrandsVF_bIsCullingEnable)
 		{
-			const uint VertexCountAfterCulling = HairStrandsVF_CullingIndirectBuffer[3];
-			uint FetchIndex0 = PrimID;
-			if (FetchIndex0 >= (VertexCountAfterCulling-1))
+			const uint FetchIndex0 = PrimID;
+			const uint FetchIndex1 = min(FetchIndex0 + 1, group_VerticesNum - 1);
+
+			const uint VertexIndex0 = HairStrandsVF_CullingIndexBuffer[FetchIndex0];
+			const uint VertexIndex1 = HairStrandsVF_CullingIndexBuffer[FetchIndex1];
+			const float LodRadiusScale0 = HairStrandsVF_CullingRadiusScaleBuffer[FetchIndex0];
+			const float LodRadiusScale1 = HairStrandsVF_CullingRadiusScaleBuffer[FetchIndex1];
+
+			if ((LodRadiusScale0 <= 0.0f) || (VertexIndex1 != VertexIndex0 + 1))
 			{
 				bSegValid = false;
 			}
 			else
 			{
-				uint FetchIndex1 = min(FetchIndex0 + 1, VertexCountAfterCulling - 1);
-
-				uint VertexIndex0 = HairStrandsVF_CullingIndexBuffer[FetchIndex0];
-				float LodRadiusScale0 = HairStrandsVF_CullingRadiusScaleBuffer[FetchIndex0];
-
-				uint VertexIndex1 = HairStrandsVF_CullingIndexBuffer[FetchIndex1];
-				float LodRadiusScale1 = HairStrandsVF_CullingRadiusScaleBuffer[FetchIndex1];
-
-				if ((LodRadiusScale0 <= 0.0f) || (VertexIndex1 != VertexIndex0 + 1))
-				{
-					bSegValid = false;
-				}
-				else
-				{
-					PrimID = VertexIndex0;
-				}
+				PrimID = VertexIndex0;
 			}
 		}
-#else
-		uint PrimID = BatchIndex * 1024 + GroupThreadID + VertexStart;
+		else if (bSegValid)
+		{
+			PrimID += VertexStart;
+		}
+	#else // PERMUTATION_CULLING
+		const uint PrimID = BatchIndex * 1024 + GroupThreadID + VertexStart;
 		bSegValid = bSegValid && (PrimID < (VertexStart + VertexCount));
-#endif
+	#endif // PERMUTATION_CULLING
+#endif // PERMUTATION_INDIRECT_PRIM_IDS
 
 		const uint SegmentCountLayerIdx = GroupID; // Stores number of segments per tile per workgroup.
 		const uint TmpSegmentCountLayerIdx = SegmentCountLayerIdx + NumBinners; // Also stores number of segments per tile per workgroup. Used as second counter for this two pass algorithm.
@@ -510,7 +495,7 @@ void BinningCS(uint DispatchThreadID : SV_DispatchThreadID, uint GroupThreadID :
 		for (uint TileIdx = GroupThreadID; TileIdx < TilesToAllocateCount; TileIdx += 1024)
 		{
 			const uint PackedTileCoord = group_TilesToAllocate[TileIdx];
-			const uint2 TileCoord = UppackVisTileCoord(PackedTileCoord);
+			const uint2 TileCoord = UnpackVisTileCoord(PackedTileCoord);
 
 			const uint TotalNewWriteCount = OutVisTileBinningGrid[uint3(TileCoord, SegmentCountLayerIdx)];
 			const uint TotalOldWriteCount = OutVisTileBinningGrid[uint3(TileCoord, TmpSegmentCountLayerIdx)];
@@ -732,6 +717,15 @@ void CompactionCS(uint DispatchThreadID : SV_DispatchThreadID, uint GroupThreadI
 
 #if SHADER_RASTERCOMPUTE_RASTER
 
+// Wave size
+#if PERMUTATION_GROUP_SIZE == 64
+#define WAVE_SIZE 32
+#elif PERMUTATION_GROUP_SIZE == 32
+#define WAVE_SIZE 32
+#else
+#error Unknown group size	
+#endif
+
 // Simple rasterization algorithm that lerps between line endpoints. Is currently more robust than the Wu algorithm
 // and optionally supports anti-aliasing similar to the Wu algorithm.
 #define RASTER_LINEAR 0 
@@ -743,27 +737,18 @@ void CompactionCS(uint DispatchThreadID : SV_DispatchThreadID, uint GroupThreadI
 
 #define RASTER_ALGO RASTER_LINEAR
 
-RWTexture2D<uint> 					RWVisTileDepthGrid;
-Buffer<uint> 						VisTilePrims;
-Buffer<uint> 						VisTileArgs;
-RWByteAddressBuffer					RWVisTileData;
-Buffer<uint>						VisTileIndirect;
-
-uint LoadRWVisTileData(uint index, uint offset)
-{
-	return RWVisTileData.Load((((index)) * VT_SIZE * 4) + ((offset) * 4));
-}
+Buffer<uint> 				VisTilePrims;
+Buffer<uint> 				VisTileArgs;
+ByteAddressBuffer			VisTileData;
+RWTexture2D<uint>			OutHairCountTexture;
+RWTexture2D<uint>			OutDepthCovTexture;
+RWTexture2D<uint>			OutPrimMatTexture;
 
-void StoreRWVisTileData(uint index, uint offset, uint value)
+uint LoadVisTileData(uint index, uint offset)
 {
-	RWVisTileData.Store((((index)) * VT_SIZE * 4) + ((offset) * 4), (value));
+	return VisTileData.Load((((index)) * VT_SIZE * 4) + ((offset) * 4));
 }
 
-RWTexture2D<uint> OutHairCountTexture;
-RWTexture2D<uint> OutDepthCovTexture;
-RWTexture2D<uint> OutPrimMatTexture;
-
-
 groupshared uint4 group_SubTile[1024]; //(32 x 32 x 4 x 4 bytes = 16k bytes)
 
 groupshared float3 group_PositionOffset;
@@ -876,10 +861,10 @@ void RasterCS(uint DispatchThreadID : SV_DispatchThreadID, uint GroupThreadID :
 		bool bTileValid = (group_TileIndex < group_TileNum);
 
 		uint PrimOffset = group_TileIndex * 1024;
-		uint PrimCount = LoadRWVisTileData(group_TileIndex, VT_PrimCount);
+		uint PrimCount = LoadVisTileData(group_TileIndex, VT_PrimCount);
 
-		uint PackedCoord = LoadRWVisTileData(group_TileIndex, VT_Coord);
-		uint2 SubTileMin = UppackVisTileCoord(PackedCoord) * TileSize;
+		uint PackedCoord = LoadVisTileData(group_TileIndex, VT_Coord);
+		uint2 SubTileMin = UnpackVisTileCoord(PackedCoord) * TileSize;
 
 		uint PackedTileMin = ((SubTileMin.x & 0xffff) << 0) | ((SubTileMin.y & 0xffff) << 16);
 
@@ -1143,13 +1128,10 @@ void PrintBool(inout FShaderPrintContext Ctx, bool In)
 
 Texture2D<uint>						VisTileDepthGrid;
 Texture2DArray<uint> 				VisTileBinningGrid;
-//Buffer<uint>						OutVisTilePrims;
 Buffer<uint>						VisTileArgs;
-//ByteAddressBuffer					OutVisTileData;
-
-uint MacroGroupId;
-uint PrimitiveInfoIndex;
-uint TotalPrimitiveInfoCount;
+uint								MacroGroupId;
+uint								PrimitiveInfoIndex;
+uint								TotalPrimitiveInfoCount;
 
 #define TilePrintOffset (TileSize >> 1)
 
diff --git a/Engine/Shaders/Private/HairStrands/HairStrandsVisibilityRasterComputeNaive.usf b/Engine/Shaders/Private/HairStrands/HairStrandsVisibilityRasterComputeNaive.usf
--- a/Engine/Shaders/Private/HairStrands/HairStrandsVisibilityRasterComputeNaive.usf
+++ b/Engine/Shaders/Private/HairStrands/HairStrandsVisibilityRasterComputeNaive.usf
@@ -2,6 +2,14 @@
 
 #define HAIR_STRANDS_PARAMETERS 1
 
+#ifndef PERMUTATION_INDIRECT_PRIM_IDS
+#define PERMUTATION_INDIRECT_PRIM_IDS 0
+#endif
+
+#ifndef PERMUTATION_CULLING
+#define PERMUTATION_CULLING 0
+#endif
+
 #include "../Common.ush"
 #include "HairStrandsVisibilityCommon.ush"
 #include "HairStrandsVertexFactoryCommon.ush"
@@ -10,11 +18,14 @@ RWTexture2D<uint>	OutHairCountTexture;
 RWTexture2D<uint>	OutDepthCovTexture;
 RWTexture2D<uint>	OutPrimMatTexture;
 Texture2D<float>	SceneDepthTexture;
+ByteAddressBuffer	IndirectPrimIDCount;
+Buffer<uint>		IndirectPrimIDs;
 float2				OutputResolution;
 uint				HairMaterialId;
 uint				ControlPointStart;
 uint				ControlPointCount;
 float				SampleWeight;
+uint				NumWorkGroups;
 
 // Line clipping based on "CLIPPING USING HOMOGENEOUS COORDINATES" by Blinn et al.
 bool BlinnLineClipping(inout float4 P0, inout float4 P1)
@@ -24,8 +35,7 @@ bool BlinnLineClipping(inout float4 P0, inout float4 P1)
 
 	bool bSign = false;
 
-	UNROLL
-	for (uint PlaneIdx = 0; PlaneIdx < 6; ++PlaneIdx)
+	UNROLL for (uint PlaneIdx = 0; PlaneIdx < 6; ++PlaneIdx)
 	{
 		// Compute boundary coordinates of both points (w+x, w-x, w+y, w-y, z, w-z)
 		bSign = !bSign;
@@ -70,7 +80,7 @@ float3 NDCToPixelCoord(float4 InDC)
 	return float3(UV * OutputResolution, NDC.z);
 }
 
-void PlotInternal(int2 Coord, float AntiAliasingFactor, float4 P0, float4 P1, float Rad0, float Rad1, float SegmentLenSqRcp, uint PrimMatID)
+void Plot(int2 Coord, float AntiAliasingFactor, float4 P0, float4 P1, float Rad0, float Rad1, float SegmentLenSqRcp, uint PrimMatID)
 {
 	// Project P onto line segment and compute the lerp alpha between P0 and P1
 	// Simplification of:
@@ -106,68 +116,125 @@ void PlotInternal(int2 Coord, float AntiAliasingFactor, float4 P0, float4 P1, fl
 }
 
 [numthreads(64, 1, 1)]
-void CSMain(uint ThreadId : SV_DispatchThreadID)
+void CSMain(uint DispatchThreadID : SV_DispatchThreadID, uint GroupThreadID : SV_GroupThreadID, uint GroupID : SV_GroupID)
 {
-	if (ThreadId >= ControlPointCount)
-	{
-		return;
-	}
-
 	ResolvedView = ResolveView();
 
-	const uint PrimID = ControlPointStart + ThreadId;
-
+#if PERMUTATION_INDIRECT_PRIM_IDS
+	const uint NumControlPoints = IndirectPrimIDCount.Load(0);
+#else // PERMUTATION_INDIRECT_PRIM_IDS
+	#if PERMUTATION_CULLING	
+	const uint NumControlPoints = HairStrandsVF_bIsCullingEnable ? HairStrandsVF_CullingIndirectBuffer[3] : ControlPointCount;
+	#else // PERMUTATION_CULLING
+	const uint NumControlPoints = ControlPointCount;
+	#endif //PERMUTATION_CULLING
+#endif // PERMUTATION_INDIRECT_PRIM_IDS
+
+	const uint NumBatches = (NumControlPoints + 63) / 64;
+	const uint NumLoops = (NumBatches + (NumWorkGroups - 1)) / NumWorkGroups;
 	const float3 PositionOffset = HairStrandsVF_PositionOffsetBuffer[0].xyz;
 
-	// Fetch both control points
-	const FHairControlPoint CP0 = UnpackHairControlPoint(
-		HairStrandsVF_PositionBuffer[PrimID],
-		PositionOffset,
-		HairStrandsVF_Radius,
-		HairStrandsVF_RootScale,
-		HairStrandsVF_TipScale);
-
-	if (CP0.Type == HAIR_CONTROLPOINT_END)
+	LOOP for (uint LoopIndex = 0; LoopIndex < NumLoops; ++LoopIndex)
 	{
-		return;
-	}
+		const uint BatchIndex = LoopIndex + (GroupID * NumLoops);
+		bool bSegValid = (BatchIndex < NumBatches);
+
+#if PERMUTATION_INDIRECT_PRIM_IDS
+		uint PrimID = 0;
+		const uint PrimIDIndex = BatchIndex * 64 + GroupThreadID;
+		bSegValid = bSegValid && (PrimIDIndex < NumControlPoints);
+		if (bSegValid)
+		{
+			PrimID = IndirectPrimIDs[PrimIDIndex];
+		}
+#else // PERMUTATION_INDIRECT_PRIM_IDS
+	#if PERMUTATION_CULLING
+		uint PrimID = BatchIndex * 64 + GroupThreadID;
+		bSegValid = bSegValid && (PrimID < NumControlPoints);
 
-	const FHairControlPoint CP1 = UnpackHairControlPoint(
-		HairStrandsVF_PositionBuffer[PrimID + 1],
-		PositionOffset,
-		HairStrandsVF_Radius,
-		HairStrandsVF_RootScale,
-		HairStrandsVF_TipScale);
+		if (bSegValid && HairStrandsVF_bIsCullingEnable)
+		{
+			const uint FetchIndex0 = PrimID;
+			const uint FetchIndex1 = min(FetchIndex0 + 1, NumControlPoints - 1);
+
+			const uint VertexIndex0 = HairStrandsVF_CullingIndexBuffer[FetchIndex0];
+			const uint VertexIndex1 = HairStrandsVF_CullingIndexBuffer[FetchIndex1];
+			const float LodRadiusScale0 = HairStrandsVF_CullingRadiusScaleBuffer[FetchIndex0];
+			const float LodRadiusScale1 = HairStrandsVF_CullingRadiusScaleBuffer[FetchIndex1];
+
+			if ((LodRadiusScale0 <= 0.0f) || (VertexIndex1 != VertexIndex0 + 1))
+			{
+				bSegValid = false;
+			}
+			else
+			{
+				PrimID = VertexIndex0;
+			}
+		}
+		else if (bSegValid)
+		{
+			PrimID += ControlPointStart;
+		}
+	#else // PERMUTATION_CULLING
+		const uint PrimID = BatchIndex * 64 + GroupThreadID + ControlPointStart;
+		bSegValid = bSegValid && (PrimID < (ControlPointStart + ControlPointCount));
+	#endif // PERMUTATION_CULLING
+#endif // PERMUTATION_INDIRECT_PRIM_IDS
 
-	// Transform to world space
-	const float3 CP0WP = mul(float4(CP0.Position, 1.0f), HairStrandsVF_LocalToWorldPrimitiveTransform).xyz;
-	const float3 CP1WP = mul(float4(CP1.Position, 1.0f), HairStrandsVF_LocalToWorldPrimitiveTransform).xyz;
-	float4 SP0 = mul(float4(CP0WP, 1.0f), LWCHackToFloat(PrimaryView.WorldToClip));
-	float4 SP1 = mul(float4(CP1WP, 1.0f), LWCHackToFloat(PrimaryView.WorldToClip));
+		if (!bSegValid)
+		{
+			continue;
+		}
 
-	if (!BlinnLineClipping(SP0, SP1))
-	{
-		return;
-	}
+		// Fetch both control points
+		FHairControlPoint CP0 = UnpackHairControlPoint(
+			HairStrandsVF_PositionBuffer[PrimID],
+			PositionOffset,
+			HairStrandsVF_Radius,
+			HairStrandsVF_RootScale,
+			HairStrandsVF_TipScale);
 
-	SP0 = float4(NDCToPixelCoord(SP0), 1.0f / SP0.w);
-	SP1 = float4(NDCToPixelCoord(SP1), 1.0f / SP1.w);
+		if (CP0.Type == HAIR_CONTROLPOINT_END)
+		{
+			continue;
+		}
 
-	const float SegmentLenSqRcp = 1.0f / dot(SP1.xy - SP0.xy, SP1.xy - SP0.xy);
+		const FHairControlPoint CP1 = UnpackHairControlPoint(
+			HairStrandsVF_PositionBuffer[PrimID + 1],
+			PositionOffset,
+			HairStrandsVF_Radius,
+			HairStrandsVF_RootScale,
+			HairStrandsVF_TipScale);
 
-	const float2 DeltaAbs = abs(SP1.xy - SP0.xy);
-	const uint NumSteps = (uint)ceil(max(DeltaAbs.x, DeltaAbs.y)) + 1;
-	const float RcpNumSteps = NumSteps > 1 ? (1.0f / (NumSteps - 1)) : 1.0f;
-	const bool bIsSteep = DeltaAbs.y > DeltaAbs.x;
-	const uint PrimMatID = PackHairVisPrimitiveMaterialId(PrimID, HairMaterialId);
+		// Transform to world space
+		const float3 CP0WP = mul(float4(CP0.Position, 1.0f), HairStrandsVF_LocalToWorldPrimitiveTransform).xyz;
+		const float3 CP1WP = mul(float4(CP1.Position, 1.0f), HairStrandsVF_LocalToWorldPrimitiveTransform).xyz;
+		float4 SP0 = mul(float4(CP0WP, 1.0f), LWCHackToFloat(PrimaryView.WorldToClip));
+		float4 SP1 = mul(float4(CP1WP, 1.0f), LWCHackToFloat(PrimaryView.WorldToClip));
 
-	LOOP
-	for (int J = 0; J < NumSteps; ++J)
-	{
-		const float Alpha = J * RcpNumSteps;
-		const float4 SP = lerp(SP0, SP1, Alpha);
+		if (!BlinnLineClipping(SP0, SP1))
+		{
+			continue;
+		}
+
+		SP0 = float4(NDCToPixelCoord(SP0), 1.0f / SP0.w);
+		SP1 = float4(NDCToPixelCoord(SP1), 1.0f / SP1.w);
+
+		const float SegmentLenSqRcp = 1.0f / dot(SP1.xy - SP0.xy, SP1.xy - SP0.xy);
+
+		const float2 DeltaAbs = abs(SP1.xy - SP0.xy);
+		const uint NumSteps = (uint)ceil(max(DeltaAbs.x, DeltaAbs.y)) + 1;
+		const float RcpNumSteps = NumSteps > 1 ? (1.0f / (NumSteps - 1)) : 1.0f;
+		const bool bIsSteep = DeltaAbs.y > DeltaAbs.x;
+		const uint PrimMatID = PackHairVisPrimitiveMaterialId(PrimID, HairMaterialId);
 
-		const float AntiAliasingFactor = 1.0f;
-		PlotInternal(SP.xy, AntiAliasingFactor, SP0, SP1, CP0.WorldRadius * 2000.0f, CP1.WorldRadius * 2000.0f, SegmentLenSqRcp, PrimMatID);
+		LOOP for (int J = 0; J < NumSteps; ++J)
+		{
+			const float Alpha = J * RcpNumSteps;
+			const float4 SP = lerp(SP0, SP1, Alpha);
+
+			const float AntiAliasingFactor = 1.0f;
+			Plot(SP.xy, AntiAliasingFactor, SP0, SP1, CP0.WorldRadius * 2000.0f, CP1.WorldRadius * 2000.0f, SegmentLenSqRcp, PrimMatID);
+		}
 	}
 }
\ No newline at end of file
diff --git a/Engine/Shaders/Private/HairStrands/HairStrandsVisibilityRasterHW.usf b/Engine/Shaders/Private/HairStrands/HairStrandsVisibilityRasterHW.usf
--- a/Engine/Shaders/Private/HairStrands/HairStrandsVisibilityRasterHW.usf
+++ b/Engine/Shaders/Private/HairStrands/HairStrandsVisibilityRasterHW.usf
@@ -8,13 +8,22 @@
 
 #if SHADER_RASTER_HW_VS
 
-Buffer<uint> 	HWRasterPrimIDs;
-float3			ViewDir;
-uint			HairMaterialId;
-float3			CameraOrigin;
-float			RadiusAtDepth1;
-uint			ControlPointStart;
-uint			bIsOrthoView;
+#ifndef PERMUTATION_DRAW_INDIRECT
+#define PERMUTATION_DRAW_INDIRECT 0
+#endif
+
+#ifndef PERMUTATION_CULLING
+#define PERMUTATION_CULLING 0
+#endif
+
+Buffer<uint> 		RasterPrimIDs;
+float3				ViewDir;
+uint				HairMaterialId;
+float3				CameraOrigin;
+float				RadiusAtDepth1;
+uint				ControlPointStart;
+uint				PrimIDsBufferSize;
+uint				bIsOrthoView;
 
 
 struct VSInput
@@ -30,9 +39,53 @@ struct VSOutput
 	nointerpolation uint PrimMatID : PRIM_MAT_ID;
 };
 
+VSOutput GetNullOutput()
+{
+	VSOutput Output = (VSOutput)0;
+	Output.Position = INFINITE_FLOAT;
+	return Output;
+}
+
 VSOutput VSMain(VSInput Input)
 {
-	const uint PrimID = ControlPointStart + Input.InstanceID;// HWRasterPrimIDs[Input.InstanceID];
+#if PERMUTATION_DRAW_INDIRECT
+	const uint PrimID = RasterPrimIDs[PrimIDsBufferSize - 1 - Input.InstanceID];
+#else // PERMUTATION_DRAW_INDIRECT
+	#if PERMUTATION_CULLING
+	uint PrimID = Input.InstanceID;
+
+	if (HairStrandsVF_bIsCullingEnable)
+	{
+		const uint NumVerticesAfterCulling = HairStrandsVF_CullingIndirectBuffer[3];
+		if (PrimID >= NumVerticesAfterCulling)
+		{
+			return GetNullOutput();
+		}
+		const uint FetchIndex0 = PrimID;
+		const uint FetchIndex1 = min(FetchIndex0 + 1, NumVerticesAfterCulling - 1);
+
+		const uint VertexIndex0 = HairStrandsVF_CullingIndexBuffer[FetchIndex0];
+		const uint VertexIndex1 = HairStrandsVF_CullingIndexBuffer[FetchIndex1];
+		const float LodRadiusScale0 = HairStrandsVF_CullingRadiusScaleBuffer[FetchIndex0];
+		const float LodRadiusScale1 = HairStrandsVF_CullingRadiusScaleBuffer[FetchIndex1];
+
+		if ((LodRadiusScale0 <= 0.0f) || (VertexIndex1 != VertexIndex0 + 1))
+		{
+			return GetNullOutput();
+		}
+		else
+		{
+			PrimID = VertexIndex0;
+		}
+	}
+	else
+	{
+		PrimID += ControlPointStart;
+	}
+	#else
+	const uint PrimID = ControlPointStart + Input.InstanceID;
+	#endif // PERMUTATION_CULLING
+#endif // PERMUTATION_DRAW_INDIRECT
 
 	const float3 PositionOffset = HairStrandsVF_PositionOffsetBuffer[0].xyz;
 
@@ -47,9 +100,7 @@ VSOutput VSMain(VSInput Input)
 	// Kill quad
 	if (CP0.Type == HAIR_CONTROLPOINT_END)
 	{
-		VSOutput Output = (VSOutput)0;
-		Output.Position = 0.0f / 0.0f;
-		return Output;
+		return GetNullOutput();
 	}
 
 	const FHairControlPoint CP1 = UnpackHairControlPoint(
diff --git a/Engine/Source/Runtime/Renderer/Private/HairStrands/HairStrandsVisibility.cpp b/Engine/Source/Runtime/Renderer/Private/HairStrands/HairStrandsVisibility.cpp
--- a/Engine/Source/Runtime/Renderer/Private/HairStrands/HairStrandsVisibility.cpp
+++ b/Engine/Source/Runtime/Renderer/Private/HairStrands/HairStrandsVisibility.cpp
@@ -61,15 +61,19 @@ int32 GHairVisibilityComputeRaster_Culling = 0;
 // TODO: Find an efficient alternative algorithm where the binners output into shared tiles.
 int32 GHairVisibilityComputeRaster_MaxTiles = 65536;
 int32 GHairVisibilityComputeRaster_TileSize = 32;
+int32 GHairVisibility_NumClassifiers = 32;
 int32 GHairVisibilityComputeRaster_NumBinners = 32;
 int32 GHairVisibilityComputeRaster_NumRasterizers = 256;
+int32 GHairVisibilityComputeRaster_NumRasterizersNaive = 256;
 int32 GHairVisibilityComputeRaster_Debug = 0;
 
 static FAutoConsoleVariableRef CVarHairVisibilityComputeRaster_Culling(TEXT("r.HairStrands.Visibility.ComputeRaster.Culling"), GHairVisibilityComputeRaster_Culling, TEXT("Use culling buffers with compute rasterization."));
 static FAutoConsoleVariableRef CVarHairVisibilityComputeRaster_MaxTiles(TEXT("r.HairStrands.Visibility.ComputeRaster.MaxTiles"), GHairVisibilityComputeRaster_MaxTiles, TEXT("Maximum number of tiles used for compute rasterization. 8192 is default"));
 static FAutoConsoleVariableRef CVarHairVisibilityComputeRaster_TileSize(TEXT("r.HairStrands.Visibility.ComputeRaster.TileSize"), GHairVisibilityComputeRaster_TileSize, TEXT("Tile size used for compute rasterization. Experimental - only size of 32 currently supported"));
+static FAutoConsoleVariableRef CVarHairVisibility_NumClassifiers(TEXT("r.HairStrands.Visibility.NumClassifiers"), GHairVisibility_NumClassifiers, TEXT("Number of workgroups used in hair segment classification pass. 32 is default"));
 static FAutoConsoleVariableRef CVarHairVisibilityComputeRaster_NumBinners(TEXT("r.HairStrands.Visibility.ComputeRaster.NumBinners"), GHairVisibilityComputeRaster_NumBinners, TEXT("Number of Binners used in Binning compute rasterization pass. 32 is default"));
 static FAutoConsoleVariableRef CVarHairVisibilityComputeRaster_NumRasterizers(TEXT("r.HairStrands.Visibility.ComputeRaster.NumRasterizers"), GHairVisibilityComputeRaster_NumRasterizers, TEXT("Number of Rasterizers used compute rasterization. 256 is default"));
+static FAutoConsoleVariableRef CVarHairVisibilityComputeRaster_NumRasterizersNaive(TEXT("r.HairStrands.Visibility.ComputeRaster.NumRasterizersNaive"), GHairVisibilityComputeRaster_NumRasterizersNaive, TEXT("Number of Rasterizers used in naive compute rasterization. 256 is default"));
 static FAutoConsoleVariableRef CVarHairVisibilityComputeRaster_Debug(TEXT("r.HairStrands.Visibility.ComputeRaster.Debug"), GHairVisibilityComputeRaster_Debug, TEXT("Debug compute raster output"));
 
 static float GHairStrandsFullCoverageThreshold = 0.98f;
@@ -108,6 +112,9 @@ static FAutoConsoleVariableRef CVarHairStrandsUseNaiveSWRaster(TEXT("r.HairStran
 static int32 GHairStrandsTileCompaction = 0;
 static FAutoConsoleVariableRef CVarHairStrandsTileCompaction(TEXT("r.HairStrands.Visibility.TileCompaction"), GHairStrandsTileCompaction, TEXT("Enables a compaction pass to run on the output of the binning pass of the hair software rasterizer."));
 
+static int32 GHairStrandsHWSWClassifaction = 0;
+static FAutoConsoleVariableRef CVarHairStrandsHWSWClassifaction(TEXT("r.HairStrands.Visibility.HWSWClassifaction"), GHairStrandsHWSWClassifaction, TEXT("Enables classifying hair segments to be rasterized with a hardware rasterizer or a software rasterizer."));
+
 /////////////////////////////////////////////////////////////////////////////////////////
 
 namespace HairStrandsVisibilityInternal
@@ -2818,14 +2825,47 @@ FHairStrandsInstanceParameters GetHairStrandsInstanceParameters(FRDGBuilder& Gra
 	return Out;
 }
 
+class FVisiblityRasterClassificationCS : public FGlobalShader
+{
+	DECLARE_GLOBAL_SHADER(FVisiblityRasterClassificationCS);
+	SHADER_USE_PARAMETER_STRUCT(FVisiblityRasterClassificationCS, FGlobalShader);
+
+	class FCulling : SHADER_PERMUTATION_BOOL("PERMUTATION_CULLING");
+	using FPermutationDomain = TShaderPermutationDomain<FCulling>;
+
+	BEGIN_SHADER_PARAMETER_STRUCT(FParameters, )
+		SHADER_PARAMETER_STRUCT_INCLUDE(FRasterComputeCommonParameters, Common)
+		SHADER_PARAMETER(uint32, ControlPointStart)
+		SHADER_PARAMETER(uint32, ControlPointCount)
+		SHADER_PARAMETER(uint32, PrimIDsBufferSize)
+		SHADER_PARAMETER(uint32, NumWorkGroups)
+		SHADER_PARAMETER_STRUCT_INCLUDE(FHairStrandsInstanceParameters, HairInstance)
+		SHADER_PARAMETER_RDG_BUFFER_UAV(RWBuffer, OutPrimIDs)
+		SHADER_PARAMETER_RDG_BUFFER_UAV(RWByteAddressBuffer, OutDrawIndexedArgs)
+		SHADER_PARAMETER_RDG_BUFFER_UAV(RWByteAddressBuffer, OutSWRasterPrimitiveCount)
+		SHADER_PARAMETER_STRUCT_REF(FViewUniformShaderParameters, ViewUniformBuffer)
+	END_SHADER_PARAMETER_STRUCT()
+
+public:
+	static bool ShouldCompilePermutation(const FGlobalShaderPermutationParameters &Parameters) { return IsHairStrandsSupported(EHairStrandsShaderType::Strands, Parameters.Platform); }
+	static void ModifyCompilationEnvironment(const FGlobalShaderPermutationParameters &Parameters, FShaderCompilerEnvironment &OutEnvironment)
+	{
+		FGlobalShader::ModifyCompilationEnvironment(Parameters, OutEnvironment);
+	}
+};
+
+IMPLEMENT_GLOBAL_SHADER(FVisiblityRasterClassificationCS, "/Engine/Private/HairStrands/HairStrandsVisibilityRasterClassification.usf", "CSMain", SF_Compute);
+
+/////////////////////////////////////////////////////////////////////////////////////////
+
 class FVisiblityRasterComputeBinningCS : public FGlobalShader
 {
 	DECLARE_GLOBAL_SHADER(FVisiblityRasterComputeBinningCS);
 	SHADER_USE_PARAMETER_STRUCT(FVisiblityRasterComputeBinningCS, FGlobalShader);
 
-	class FGroupSize : SHADER_PERMUTATION_SPARSE_INT("PERMUTATION_GROUP_SIZE", 32, 64);
 	class FCulling : SHADER_PERMUTATION_BOOL("PERMUTATION_CULLING");
-	using FPermutationDomain = TShaderPermutationDomain<FGroupSize, FCulling>;
+	class FIndirectPrimIDs : SHADER_PERMUTATION_BOOL("PERMUTATION_INDIRECT_PRIM_IDS");
+	using FPermutationDomain = TShaderPermutationDomain<FCulling, FIndirectPrimIDs>;
 
 	BEGIN_SHADER_PARAMETER_STRUCT(FParameters, )
 		SHADER_PARAMETER_STRUCT_INCLUDE(FRasterComputeCommonParameters, Common)
@@ -2844,6 +2884,8 @@ class FVisiblityRasterComputeBinningCS : public FGlobalShader
 		RDG_TEXTURE_ACCESS(VisTileBinningGridTex, ERHIAccess::UAVCompute)
 		SHADER_PARAMETER(uint32, VertexCount)
 		SHADER_PARAMETER(uint32, VertexStart)
+		SHADER_PARAMETER_RDG_BUFFER_SRV(ByteAddressBuffer, IndirectPrimIDCount)
+		SHADER_PARAMETER_RDG_BUFFER_SRV(Buffer, IndirectPrimIDs)
 	END_SHADER_PARAMETER_STRUCT()
 
 public:
@@ -2864,8 +2906,7 @@ class FVisiblityRasterComputeCompactionCS : public FGlobalShader
 	DECLARE_GLOBAL_SHADER(FVisiblityRasterComputeCompactionCS);
 	SHADER_USE_PARAMETER_STRUCT(FVisiblityRasterComputeCompactionCS, FGlobalShader);
 
-	class FGroupSize : SHADER_PERMUTATION_SPARSE_INT("PERMUTATION_GROUP_SIZE", 32, 64);
-	using FPermutationDomain = TShaderPermutationDomain<FGroupSize>;
+	using FPermutationDomain = TShaderPermutationDomain<>;
 
 	BEGIN_SHADER_PARAMETER_STRUCT(FParameters, )
 		SHADER_PARAMETER_STRUCT_INCLUDE(FRasterComputeCommonParameters, Common)
@@ -2906,14 +2947,11 @@ class FVisiblityRasterComputeRasterizeCS : public FGlobalShader
 		SHADER_PARAMETER(uint32, DispatchCountX)
 		SHADER_PARAMETER(uint32, HairMaterialId)
 		SHADER_PARAMETER_STRUCT_INCLUDE(FHairStrandsInstanceParameters, HairInstance)
-		RDG_BUFFER_ACCESS(IndirectBufferArgs, ERHIAccess::IndirectArgs)
 		SHADER_PARAMETER_RDG_TEXTURE(Texture2D, SceneDepthTexture)
 		SHADER_PARAMETER(float, SampleWeight)
-		SHADER_PARAMETER_RDG_TEXTURE_UAV(RWTexture2D, RWVisTileDepthGrid)
 		SHADER_PARAMETER_RDG_BUFFER_SRV(Buffer, VisTilePrims)
 		SHADER_PARAMETER_RDG_BUFFER_SRV(Buffer, VisTileArgs)
-		SHADER_PARAMETER_RDG_BUFFER_UAV(RWByteAddressBuffer, RWVisTileData)
-		SHADER_PARAMETER_RDG_BUFFER_SRV(Buffer, VisTileIndirect)
+		SHADER_PARAMETER_RDG_BUFFER_SRV(ByteAddressBuffer, VisTileData)
 		SHADER_PARAMETER_RDG_TEXTURE_UAV(RWTexture2D, OutHairCountTexture)
 		SHADER_PARAMETER_RDG_TEXTURE_UAV(RWTexture2D, OutDepthCovTexture)
 		SHADER_PARAMETER_RDG_TEXTURE_UAV(RWTexture2D, OutPrimMatTexture)
@@ -2970,23 +3008,27 @@ class FVisiblityRasterHWVS : public FGlobalShader
 	DECLARE_GLOBAL_SHADER(FVisiblityRasterHWVS);
 	SHADER_USE_PARAMETER_STRUCT(FVisiblityRasterHWVS, FGlobalShader);
 
-	using FPermutationDomain = TShaderPermutationDomain<>;
+	class FDrawIndirect : SHADER_PERMUTATION_BOOL("PERMUTATION_DRAW_INDIRECT");
+	class FCulling : SHADER_PERMUTATION_BOOL("PERMUTATION_CULLING");
+	using FPermutationDomain = TShaderPermutationDomain<FDrawIndirect, FCulling>;
 
 	BEGIN_SHADER_PARAMETER_STRUCT(FParameters, )
-		SHADER_PARAMETER_RDG_BUFFER_SRV(Buffer, HWRasterPrimIDs)
+		SHADER_PARAMETER_RDG_BUFFER_SRV(Buffer, RasterPrimIDs)
 		SHADER_PARAMETER(FVector3f, ViewDir)
 		SHADER_PARAMETER(uint32, HairMaterialId)
 		SHADER_PARAMETER(FVector3f, CameraOrigin)
 		SHADER_PARAMETER(float, RadiusAtDepth1)
 		SHADER_PARAMETER(uint32, ControlPointStart)
+		SHADER_PARAMETER(uint32, PrimIDsBufferSize)
 		SHADER_PARAMETER(uint32, bIsOrthoView)
 		SHADER_PARAMETER_STRUCT_INCLUDE(FHairStrandsInstanceParameters, HairInstance)
 		SHADER_PARAMETER_STRUCT_REF(FViewUniformShaderParameters, ViewUniformBuffer)
+		RDG_BUFFER_ACCESS(IndirectBufferArgs, ERHIAccess::IndirectArgs)
 	END_SHADER_PARAMETER_STRUCT()
 
 public:
-	static bool ShouldCompilePermutation(const FGlobalShaderPermutationParameters &Parameters) { return IsHairStrandsSupported(EHairStrandsShaderType::Strands, Parameters.Platform); }
-	static void ModifyCompilationEnvironment(const FGlobalShaderPermutationParameters &Parameters, FShaderCompilerEnvironment &OutEnvironment)
+	static bool ShouldCompilePermutation(const FGlobalShaderPermutationParameters& Parameters) { return IsHairStrandsSupported(EHairStrandsShaderType::Strands, Parameters.Platform); }
+	static void ModifyCompilationEnvironment(const FGlobalShaderPermutationParameters& Parameters, FShaderCompilerEnvironment& OutEnvironment)
 	{
 		FGlobalShader::ModifyCompilationEnvironment(Parameters, OutEnvironment);
 		OutEnvironment.SetDefine(TEXT("SHADER_RASTER_HW_VS"), 1);
@@ -3005,18 +3047,16 @@ class FVisiblityRasterHWPS : public FGlobalShader
 	using FPermutationDomain = TShaderPermutationDomain<>;
 
 	BEGIN_SHADER_PARAMETER_STRUCT(FParameters, )
-		SHADER_PARAMETER(float, SampleWeight)
-		SHADER_PARAMETER_STRUCT_INCLUDE(FHairStrandsInstanceParameters, HairInstance)
-		SHADER_PARAMETER_STRUCT_REF(FViewUniformShaderParameters, ViewUniformBuffer)
 		SHADER_PARAMETER_RDG_TEXTURE_UAV(RWTexture2D, OutHairCountTexture)
 		SHADER_PARAMETER_RDG_TEXTURE_UAV(RWTexture2D, OutDepthCovTexture)
 		SHADER_PARAMETER_RDG_TEXTURE_UAV(RWTexture2D, OutPrimMatTexture)
+		SHADER_PARAMETER(float, SampleWeight)
 		RENDER_TARGET_BINDING_SLOTS()
 	END_SHADER_PARAMETER_STRUCT()
 
 public:
-	static bool ShouldCompilePermutation(const FGlobalShaderPermutationParameters &Parameters) { return IsHairStrandsSupported(EHairStrandsShaderType::Strands, Parameters.Platform); }
-	static void ModifyCompilationEnvironment(const FGlobalShaderPermutationParameters &Parameters, FShaderCompilerEnvironment &OutEnvironment)
+	static bool ShouldCompilePermutation(const FGlobalShaderPermutationParameters& Parameters) { return IsHairStrandsSupported(EHairStrandsShaderType::Strands, Parameters.Platform); }
+	static void ModifyCompilationEnvironment(const FGlobalShaderPermutationParameters& Parameters, FShaderCompilerEnvironment& OutEnvironment)
 	{
 		FGlobalShader::ModifyCompilationEnvironment(Parameters, OutEnvironment);
 		OutEnvironment.SetDefine(TEXT("SHADER_RASTER_HW_PS"), 1);
@@ -3025,6 +3065,11 @@ public:
 
 IMPLEMENT_GLOBAL_SHADER(FVisiblityRasterHWPS, "/Engine/Private/HairStrands/HairStrandsVisibilityRasterHW.usf", "PSMain", SF_Pixel);
 
+BEGIN_SHADER_PARAMETER_STRUCT(FVisiblityRasterHWParameters, )
+	SHADER_PARAMETER_STRUCT_INCLUDE(FVisiblityRasterHWVS::FParameters, VS)
+	SHADER_PARAMETER_STRUCT_INCLUDE(FVisiblityRasterHWPS::FParameters, PS)
+END_SHADER_PARAMETER_STRUCT()
+
 /////////////////////////////////////////////////////////////////////////////////////////
 
 class FVisiblityRasterComputeNaiveCS : public FGlobalShader
@@ -3032,7 +3077,9 @@ class FVisiblityRasterComputeNaiveCS : public FGlobalShader
 	DECLARE_GLOBAL_SHADER(FVisiblityRasterComputeNaiveCS);
 	SHADER_USE_PARAMETER_STRUCT(FVisiblityRasterComputeNaiveCS, FGlobalShader);
 
-	using FPermutationDomain = TShaderPermutationDomain<>;
+	class FIndirectPrimIDs : SHADER_PERMUTATION_BOOL("PERMUTATION_INDIRECT_PRIM_IDS");
+	class FCulling : SHADER_PERMUTATION_BOOL("PERMUTATION_CULLING");
+	using FPermutationDomain = TShaderPermutationDomain<FIndirectPrimIDs, FCulling>;
 
 	BEGIN_SHADER_PARAMETER_STRUCT(FParameters, )
 		SHADER_PARAMETER(FVector2f, OutputResolution)
@@ -3040,17 +3087,20 @@ class FVisiblityRasterComputeNaiveCS : public FGlobalShader
 		SHADER_PARAMETER(uint32, ControlPointStart)
 		SHADER_PARAMETER(uint32, ControlPointCount)
 		SHADER_PARAMETER(float, SampleWeight)
+		SHADER_PARAMETER(uint32, NumWorkGroups)
 		SHADER_PARAMETER_STRUCT_INCLUDE(FHairStrandsInstanceParameters, HairInstance)
 		SHADER_PARAMETER_STRUCT_REF(FViewUniformShaderParameters, ViewUniformBuffer)
 		SHADER_PARAMETER_RDG_TEXTURE_UAV(RWTexture2D, OutHairCountTexture)
 		SHADER_PARAMETER_RDG_TEXTURE_UAV(RWTexture2D, OutDepthCovTexture)
 		SHADER_PARAMETER_RDG_TEXTURE_UAV(RWTexture2D, OutPrimMatTexture)
 		SHADER_PARAMETER_RDG_TEXTURE(Texture2D, SceneDepthTexture)
+		SHADER_PARAMETER_RDG_BUFFER_SRV(ByteAddressBuffer, IndirectPrimIDCount)
+		SHADER_PARAMETER_RDG_BUFFER_SRV(Buffer, IndirectPrimIDs)
 	END_SHADER_PARAMETER_STRUCT()
 
 public:
-	static bool ShouldCompilePermutation(const FGlobalShaderPermutationParameters &Parameters) { return IsHairStrandsSupported(EHairStrandsShaderType::Strands, Parameters.Platform); }
-	static void ModifyCompilationEnvironment(const FGlobalShaderPermutationParameters &Parameters, FShaderCompilerEnvironment &OutEnvironment)
+	static bool ShouldCompilePermutation(const FGlobalShaderPermutationParameters& Parameters) { return IsHairStrandsSupported(EHairStrandsShaderType::Strands, Parameters.Platform); }
+	static void ModifyCompilationEnvironment(const FGlobalShaderPermutationParameters& Parameters, FShaderCompilerEnvironment& OutEnvironment)
 	{
 		FGlobalShader::ModifyCompilationEnvironment(Parameters, OutEnvironment);
 		OutEnvironment.SetDefine(TEXT("SHADER_RASTERCOMPUTE_RASTER"), 1);
@@ -3102,59 +3152,79 @@ static FRasterComputeOutput AddVisibilityComputeRasterPass(
 	// See the comment on the GHairVisibilityComputeRaster_MaxTiles declaration for an explanation for the large upper bound.
 	const uint32 MaxTiles = FMath::Min(FMath::Max(GHairVisibilityComputeRaster_MaxTiles, 1024), 262144);
 	const uint32 TileSize = FMath::Min(FMath::Max(GHairVisibilityComputeRaster_TileSize, 8),32);
+	const uint32 NumClassifiers = FMath::Min(FMath::Max(GHairVisibility_NumClassifiers, 1), 256);
 	const uint32 NumBinners = FMath::Min(FMath::Max(GHairVisibilityComputeRaster_NumBinners, 1), 256);
 	const uint32 NumRasterizers = FMath::Min(FMath::Max(GHairVisibilityComputeRaster_NumRasterizers, 1), 1024);
+	const uint32 NumRasterizersNaive = FMath::Min(FMath::Max(GHairVisibilityComputeRaster_NumRasterizersNaive, 1), 1024);
 
 	const FIntPoint TileGridRes = FIntPoint((InResolution.X + (TileSize-1)) / TileSize, ((InResolution.Y + (TileSize - 1)) / TileSize)*1);
 
-	FRDGBufferRef OutVisTilePrims = nullptr;
-	FRDGTextureRef OutVisTileBinningGrid = nullptr;
-	FRDGTextureRef OutVisTileDepthGrid = nullptr;
-	FRDGBufferRef OutVisTileArgs = nullptr;
-	FRDGBufferRef OutVisTileData = nullptr;
-	FRDGBufferRef OutCompactedVisTilePrims = nullptr;
-	FRDGBufferRef OutCompactedVisTileArgs = nullptr;
-	FRDGBufferRef OutCompactedVisTileData = nullptr;
-	FRDGBufferRef OutVisTileIndirect = nullptr;
-
-	FRDGBufferUAVRef OutVisTilePrimsUAV = nullptr;
-	FRDGTextureUAVRef OutVisTileBinningGridUAV = nullptr;
-	FRDGTextureUAVRef OutVisTileDepthGridUAV = nullptr;
-	FRDGBufferUAVRef OutVisTileArgsUAV = nullptr;
-	FRDGBufferUAVRef OutVisTileDataUAV = nullptr;
-	FRDGBufferUAVRef OutCompactedVisTilePrimsUAV = nullptr;
-	FRDGBufferUAVRef OutCompactedVisTileArgsUAV = nullptr;
-	FRDGBufferUAVRef OutCompactedVisTileDataUAV = nullptr;
-	FRDGBufferUAVRef OutVisTileIndirectUAV = nullptr;
+	// Compute maximum number of PrimIDs
+	uint32 MaxNumPrimIDs = 0;
+	for (const FHairStrandsMacroGroupData &MacroGroup : MacroGroupDatas)
+	{
+		for (const FHairStrandsMacroGroupData::PrimitiveInfo &PrimitiveInfo : MacroGroup.PrimitivesInfos)
+		{
+			// If a groom is not visible in primary view, but visible in shadow view, its PrimitiveInfo.Mesh will be null.
+			if (PrimitiveInfo.Mesh == nullptr || PrimitiveInfo.Mesh->Elements.Num() == 0)
+			{
+				continue;
+			}
+
+			const FHairGroupPublicData *HairGroupPublicData = reinterpret_cast<const FHairGroupPublicData *>(PrimitiveInfo.Mesh->Elements[0].VertexFactoryUserData);
+			check(HairGroupPublicData);
+
+			const FHairGroupPublicData::FVertexFactoryInput &VFInput = HairGroupPublicData->VFInput;
+
+			uint32 VertexCount = VFInput.Strands.VertexCount;
+			const bool bCullingPossible = bSupportCulling && GHairVisibilityComputeRaster_Culling;
+			if (!bCullingPossible)
+			{
+				// calculate current view screen size - which can result in fewer strands rasterized in current view
+				const FSphere BoundsSphere = HairGroupPublicData->ContinuousLODBounds.GetSphere();
+				const float ScreenSize = ComputeBoundsScreenSize(FVector4(BoundsSphere.Center, 1), BoundsSphere.W, ViewInfo);
+
+				VertexCount = HairGroupPublicData->GetActiveStrandsVertexCount(VFInput.Strands.VertexCount, ScreenSize);
+			}
+
+			MaxNumPrimIDs = FMath::Max(MaxNumPrimIDs, VertexCount);
+		}
+	}
+
+	FRDGBufferRef VisTilePrims = GraphBuilder.CreateBuffer(FRDGBufferDesc::CreateBufferDesc(sizeof(uint32), MaxTiles*1024), TEXT("Hair.VisTilePrims"));
+	FRDGBufferSRVRef VisTilePrimsSRV = GraphBuilder.CreateSRV(VisTilePrims, PF_R32_UINT);
 	
-	OutVisTilePrims = GraphBuilder.CreateBuffer(FRDGBufferDesc::CreateBufferDesc(sizeof(uint32), MaxTiles*1024), TEXT("Hair.OutVisTilePrims"));
-	OutCompactedVisTilePrims = GraphBuilder.CreateBuffer(FRDGBufferDesc::CreateBufferDesc(sizeof(uint32), MaxTiles*1024), TEXT("Hair.OutCompactedVisTilePrims"));
+	FRDGBufferRef CompactedVisTilePrims = GraphBuilder.CreateBuffer(FRDGBufferDesc::CreateBufferDesc(sizeof(uint32), MaxTiles*1024), TEXT("Hair.CompactedVisTilePrims"));
 
 	FRDGTextureDesc DescBinningGrid = FRDGTextureDesc::Create2DArray(TileGridRes, PF_R32_UINT, FClearValueBinding::None, TexCreate_UAV | TexCreate_ShaderResource, NumBinners * 3);
-	OutVisTileBinningGrid = GraphBuilder.CreateTexture(DescBinningGrid, TEXT("Hair.VisTileBinningGrid"));
+	FRDGTextureRef VisTileBinningGrid = GraphBuilder.CreateTexture(DescBinningGrid, TEXT("Hair.VisTileBinningGrid"));
+	FRDGTextureUAVRef VisTileBinningGridUAV = GraphBuilder.CreateUAV(VisTileBinningGrid);
 
 	FRDGTextureDesc DescDepthGrid = FRDGTextureDesc::Create2D(TileGridRes, PF_R32_UINT, FClearValueBinding::None, TexCreate_UAV | TexCreate_ShaderResource);
-	OutVisTileDepthGrid = GraphBuilder.CreateTexture(DescDepthGrid, TEXT("Hair.VisTileDepthGrid"));
+	FRDGTextureRef VisTileDepthGrid = GraphBuilder.CreateTexture(DescDepthGrid, TEXT("Hair.VisTileDepthGrid"));
 
-	OutVisTileArgs = GraphBuilder.CreateBuffer(FRDGBufferDesc::CreateBufferDesc(sizeof(uint32), 4+16+(((TileGridRes.X*TileGridRes.Y)+31)/32) ), TEXT("Hair.OutVisTileArgs"));
-	OutCompactedVisTileArgs = GraphBuilder.CreateBuffer(FRDGBufferDesc::CreateBufferDesc(sizeof(uint32), 4+16+(((TileGridRes.X*TileGridRes.Y)+31)/32) ), TEXT("Hair.OutCompactedVisTileArgs"));
+	FRDGBufferRef VisTileArgs = GraphBuilder.CreateBuffer(FRDGBufferDesc::CreateBufferDesc(sizeof(uint32), 4+16+(((TileGridRes.X*TileGridRes.Y)+31)/32) ), TEXT("Hair.VisTileArgs"));
+	FRDGBufferUAVRef VisTileArgsUAV = GraphBuilder.CreateUAV(VisTileArgs, PF_R32_UINT);
+	FRDGBufferSRVRef VisTileArgsSRV = GraphBuilder.CreateSRV(VisTileArgs, PF_R32_UINT);
+	
+	FRDGBufferRef CompactedVisTileArgs = GraphBuilder.CreateBuffer(FRDGBufferDesc::CreateBufferDesc(sizeof(uint32), 4+16+(((TileGridRes.X*TileGridRes.Y)+31)/32) ), TEXT("Hair.CompactedVisTileArgs"));
 
 	FRDGBufferDesc DescTileData = FRDGBufferDesc::CreateStructuredDesc(sizeof(uint32), MaxTiles * 4 * 2);
 	DescTileData.Usage = EBufferUsageFlags(DescTileData.Usage | BUF_ByteAddressBuffer);
-	OutVisTileData = GraphBuilder.CreateBuffer(DescTileData, TEXT("Hair.OutVisTileData"));
-	OutCompactedVisTileData = GraphBuilder.CreateBuffer(DescTileData, TEXT("Hair.OutCompactedVisTileData"));
+	FRDGBufferRef VisTileData = GraphBuilder.CreateBuffer(DescTileData, TEXT("Hair.VisTileData"));
+	FRDGBufferSRVRef VisTileDataSRV = GraphBuilder.CreateSRV(VisTileData);
+	
+	FRDGBufferRef CompactedVisTileData = GraphBuilder.CreateBuffer(DescTileData, TEXT("Hair.CompactedVisTileData"));
+
+	FRDGBufferRef RasterizerPrimIDs = GraphBuilder.CreateBuffer(FRDGBufferDesc::CreateBufferDesc(sizeof(uint32), FMath::Max(MaxNumPrimIDs, 1u)), TEXT("Hair.RasterizerPrimIDs"));
+	FRDGBufferSRVRef RasterizerPrimIDsSRV = GraphBuilder.CreateSRV(RasterizerPrimIDs, PF_R32_UINT);
 
-	OutVisTileIndirect = GraphBuilder.CreateBuffer(FRDGBufferDesc::CreateIndirectDesc<FRHIDispatchIndirectParameters>(1), TEXT("Hair.OutVisTileIndirect"));
+	FRDGBufferRef DrawIndexedIndirectArgs = GraphBuilder.CreateBuffer(FRDGBufferDesc::CreateIndirectDesc<FRHIDrawIndexedIndirectParameters>(1), TEXT("Hair.DrawIndexedIndirectArgs"));
+	FRDGBufferUAVRef DrawIndexedIndirectArgsUAV = GraphBuilder.CreateUAV(DrawIndexedIndirectArgs, PF_R32_UINT);
 
-	OutVisTilePrimsUAV = GraphBuilder.CreateUAV(OutVisTilePrims, PF_R32_UINT);
-	OutCompactedVisTilePrimsUAV = GraphBuilder.CreateUAV(OutCompactedVisTilePrims, PF_R32_UINT);
-	OutVisTileBinningGridUAV = GraphBuilder.CreateUAV(OutVisTileBinningGrid);
-	OutVisTileDepthGridUAV = GraphBuilder.CreateUAV(OutVisTileDepthGrid);
-	OutVisTileArgsUAV = GraphBuilder.CreateUAV(OutVisTileArgs, PF_R32_UINT);
-	OutCompactedVisTileArgsUAV = GraphBuilder.CreateUAV(OutCompactedVisTileArgs, PF_R32_UINT);
-	OutVisTileDataUAV = GraphBuilder.CreateUAV(OutVisTileData);
-	OutCompactedVisTileDataUAV = GraphBuilder.CreateUAV(OutCompactedVisTileData);
-	OutVisTileIndirectUAV = GraphBuilder.CreateUAV(OutVisTileIndirect);
+	FRDGBufferRef SWRasterPrimIDCount = GraphBuilder.CreateBuffer(FRDGBufferDesc::CreateBufferDesc(sizeof(uint32), 1), TEXT("Hair.SWRasterPrimIDCount"));
+	FRDGBufferUAVRef SWRasterPrimIDCountUAV = GraphBuilder.CreateUAV(SWRasterPrimIDCount, PF_R32_UINT);
+	FRDGBufferSRVRef SWRasterPrimIDCountSRV = GraphBuilder.CreateSRV(SWRasterPrimIDCount, PF_R32_UINT);
 
 	// Create and set the uniform buffer
 	const bool bEnableMSAA = false;
@@ -3162,29 +3232,59 @@ static FRasterComputeOutput AddVisibilityComputeRasterPass(
 	SetUpViewHairRenderInfo(ViewInfo, bEnableMSAA, ViewInfo.CachedViewUniformShaderParameters->HairRenderInfo, ViewInfo.CachedViewUniformShaderParameters->HairRenderInfoBits, ViewInfo.CachedViewUniformShaderParameters->HairComponents);
 	ViewUniformShaderParameters = TUniformBufferRef<FViewUniformShaderParameters>::CreateUniformBufferImmediate(*ViewInfo.CachedViewUniformShaderParameters, UniformBuffer_SingleFrame);
 
+	const bool bClassification = GHairStrandsHWSWClassifaction != 0;
+	const bool bUseHWRaster = GHairStrandsUseHWRaster != 0;
+	const bool bUseNaiveSWRaster = GHairStrandsUseNaiveSWRaster != 0;
+	const bool bTileCompactionPass = GHairStrandsTileCompaction != 0;
+
 	const uint32 GroupSize = GetVendorOptimalGroupSize1D();
+
+	FVisiblityRasterClassificationCS::FPermutationDomain ClassificationPermutationVector_CullingOff;
+	ClassificationPermutationVector_CullingOff.Set<FVisiblityRasterClassificationCS::FCulling>(false);
+	TShaderMapRef<FVisiblityRasterClassificationCS> ComputeShaderClassification_CullingOff(ViewInfo.ShaderMap, ClassificationPermutationVector_CullingOff);
+	
+	FVisiblityRasterClassificationCS::FPermutationDomain ClassificationPermutationVector_CullingOn;
+	ClassificationPermutationVector_CullingOn.Set<FVisiblityRasterClassificationCS::FCulling>(true);
+	TShaderMapRef<FVisiblityRasterClassificationCS> ComputeShaderClassification_CullingOn(ViewInfo.ShaderMap, ClassificationPermutationVector_CullingOn);
 	
 	FVisiblityRasterComputeBinningCS::FPermutationDomain BinningPermutationVector_CullingOff;
-	BinningPermutationVector_CullingOff.Set<FVisiblityRasterComputeBinningCS::FGroupSize>(GroupSize);
 	BinningPermutationVector_CullingOff.Set<FVisiblityRasterComputeBinningCS::FCulling>(false);
+	BinningPermutationVector_CullingOff.Set<FVisiblityRasterComputeBinningCS::FIndirectPrimIDs>(bClassification);
 	TShaderMapRef<FVisiblityRasterComputeBinningCS> ComputeShaderBinning_CullingOff(ViewInfo.ShaderMap, BinningPermutationVector_CullingOff);
 
 	FVisiblityRasterComputeBinningCS::FPermutationDomain BinningPermutationVector_CullingOn;
-	BinningPermutationVector_CullingOn.Set<FVisiblityRasterComputeBinningCS::FGroupSize>(GroupSize);
 	BinningPermutationVector_CullingOn.Set<FVisiblityRasterComputeBinningCS::FCulling>(true);
+	BinningPermutationVector_CullingOn.Set<FVisiblityRasterComputeBinningCS::FIndirectPrimIDs>(bClassification);
 	TShaderMapRef<FVisiblityRasterComputeBinningCS> ComputeShaderBinning_CullingOn(ViewInfo.ShaderMap, BinningPermutationVector_CullingOn);
 
 	FVisiblityRasterComputeCompactionCS::FPermutationDomain CompactionPermutationVector;
-	CompactionPermutationVector.Set<FVisiblityRasterComputeCompactionCS::FGroupSize>(GroupSize);
 	TShaderMapRef<FVisiblityRasterComputeCompactionCS> ComputeShaderCompaction(ViewInfo.ShaderMap, CompactionPermutationVector);
 
 	FVisiblityRasterComputeRasterizeCS::FPermutationDomain RasterPermutationVector;
 	RasterPermutationVector.Set<FVisiblityRasterComputeRasterizeCS::FGroupSize>(GroupSize);
 	TShaderMapRef<FVisiblityRasterComputeRasterizeCS> ComputeShaderRaster(ViewInfo.ShaderMap, RasterPermutationVector);
 
-	TShaderMapRef<FVisiblityRasterHWVS> VertexShaderRaster(ViewInfo.ShaderMap);
+	FVisiblityRasterHWVS::FPermutationDomain VSRasterPermutationVector_CullingOff;
+	VSRasterPermutationVector_CullingOff.Set<FVisiblityRasterHWVS::FDrawIndirect>(bClassification);
+	VSRasterPermutationVector_CullingOff.Set<FVisiblityRasterHWVS::FCulling>(false);
+	TShaderMapRef<FVisiblityRasterHWVS> VertexShaderRaster_CullingOff(ViewInfo.ShaderMap, VSRasterPermutationVector_CullingOff);
+
+	FVisiblityRasterHWVS::FPermutationDomain VSRasterPermutationVector_CullingOn;
+	VSRasterPermutationVector_CullingOn.Set<FVisiblityRasterHWVS::FDrawIndirect>(bClassification);
+	VSRasterPermutationVector_CullingOn.Set<FVisiblityRasterHWVS::FCulling>(true);
+	TShaderMapRef<FVisiblityRasterHWVS> VertexShaderRaster_CullingOn(ViewInfo.ShaderMap, VSRasterPermutationVector_CullingOn);
+	
 	TShaderMapRef<FVisiblityRasterHWPS> PixelShaderRaster(ViewInfo.ShaderMap);
-	TShaderMapRef<FVisiblityRasterComputeNaiveCS> ComputeShaderRasterNaive(ViewInfo.ShaderMap);
+
+	FVisiblityRasterComputeNaiveCS::FPermutationDomain RasterNaivePermutationVector_CullingOff;
+	RasterNaivePermutationVector_CullingOff.Set<FVisiblityRasterComputeNaiveCS::FIndirectPrimIDs>(bClassification);
+	RasterNaivePermutationVector_CullingOff.Set<FVisiblityRasterComputeNaiveCS::FCulling>(false);
+	TShaderMapRef<FVisiblityRasterComputeNaiveCS> ComputeShaderRasterNaive_CullingOff(ViewInfo.ShaderMap, RasterNaivePermutationVector_CullingOff);
+
+	FVisiblityRasterComputeNaiveCS::FPermutationDomain RasterNaivePermutationVector_CullingOn;
+	RasterNaivePermutationVector_CullingOn.Set<FVisiblityRasterComputeNaiveCS::FIndirectPrimIDs>(bClassification);
+	RasterNaivePermutationVector_CullingOn.Set<FVisiblityRasterComputeNaiveCS::FCulling>(true);
+	TShaderMapRef<FVisiblityRasterComputeNaiveCS> ComputeShaderRasterNaive_CullingOn(ViewInfo.ShaderMap, RasterNaivePermutationVector_CullingOn);
 
 	// Common parameters
 	FRasterComputeCommonParameters Common;
@@ -3205,19 +3305,15 @@ static FRasterComputeOutput AddVisibilityComputeRasterPass(
 		Common.RcpNumBinners = 1.0 / NumBinners;
 	}
 
-	const bool bUseHWRaster = GHairStrandsUseHWRaster != 0;
-	const bool bUseNaiveSWRaster = GHairStrandsUseNaiveSWRaster != 0;
-	const bool bTileCompactionPass = GHairStrandsTileCompaction != 0;
-
 	// Fill in tile depth
-	if (!bUseHWRaster && !bUseNaiveSWRaster)
+	if (!bUseNaiveSWRaster && (bClassification || !bUseHWRaster))
 	{
 		TShaderMapRef<FVisiblityRasterComputePrepareDepthGridCS> ComputeShaderPrepareDepthGrid(ViewInfo.ShaderMap);
 		FVisiblityRasterComputePrepareDepthGridCS::FParameters* PrepDepthGridParameters = GraphBuilder.AllocParameters<FVisiblityRasterComputePrepareDepthGridCS::FParameters>();
 		
 		PrepDepthGridParameters->Common = Common;
 		PrepDepthGridParameters->SceneDepthTexture = SceneDepthTexture;
-		PrepDepthGridParameters->OutVisTileDepthGrid = OutVisTileDepthGridUAV;
+		PrepDepthGridParameters->OutVisTileDepthGrid = GraphBuilder.CreateUAV(VisTileDepthGrid);
 		PrepDepthGridParameters->OutDepthCovTexture = DepthCovTextureUAV;
 		FComputeShaderUtils::AddPass(GraphBuilder, RDG_EVENT_NAME("HairStrands::VisibilityComputeRasterPrepDepthGrid"), ComputeShaderPrepareDepthGrid, PrepDepthGridParameters, FIntVector(TileGridRes.X* TileGridRes.Y, 1, 1));
 	}
@@ -3256,6 +3352,7 @@ static FRasterComputeOutput AddVisibilityComputeRasterPass(
 			const FMatrix44d LocalToWorldPrimitiveTransform = VFInput.LocalToWorldTransform.ToMatrixWithScale(); 
 
 			const bool bCullingEnable = bSupportCulling && GHairVisibilityComputeRaster_Culling ? HairGroupPublicData->GetCullingResultAvailable() : false;
+			const bool bCullingEnableInRasterizers = bCullingEnable && !bClassification;
 			
 			// calculate current view screen size - which can result in fewer strands rasterized in current view
 			const FSphere BoundsSphere = HairGroupPublicData->ContinuousLODBounds.GetSphere();
@@ -3265,42 +3362,74 @@ static FRasterComputeOutput AddVisibilityComputeRasterPass(
 			const uint32 VertexCount = HairGroupPublicData->GetActiveStrandsVertexCount(VFInput.Strands.VertexCount, ScreenSize);
 			const float SampleWeight = HairGroupPublicData->GetActiveStrandsSampleWeight(false, ScreenSize);
 
-			if (bUseHWRaster)
+			// HW/SW classification
+			if (bClassification)
+			{
+				// Reset buffers
+				{
+					AddClearUAVPass(GraphBuilder, DrawIndexedIndirectArgsUAV, 0u);
+					AddClearUAVPass(GraphBuilder, SWRasterPrimIDCountUAV, 0u);
+				}
+
+				const uint32 NumWorkGroups = NumClassifiers;
+
+				FVisiblityRasterClassificationCS::FParameters* ClassificationParameters = GraphBuilder.AllocParameters<FVisiblityRasterClassificationCS::FParameters>();
+				ClassificationParameters->Common = Common;
+				ClassificationParameters->ControlPointStart = VertexStart;
+				ClassificationParameters->ControlPointCount = VertexCount;
+				ClassificationParameters->PrimIDsBufferSize = MaxNumPrimIDs;
+				ClassificationParameters->NumWorkGroups = NumWorkGroups;
+				ClassificationParameters->HairInstance = GetHairStrandsInstanceParameters(GraphBuilder, ViewInfo, HairGroupPublicData, bCullingEnable, bForceRegister);
+				ClassificationParameters->ViewUniformBuffer = ViewUniformShaderParameters;
+				ClassificationParameters->OutPrimIDs = GraphBuilder.CreateUAV(RasterizerPrimIDs, PF_R32_UINT);
+				ClassificationParameters->OutDrawIndexedArgs = DrawIndexedIndirectArgsUAV;
+				ClassificationParameters->OutSWRasterPrimitiveCount = SWRasterPrimIDCountUAV;
+
+				FComputeShaderUtils::AddPass(
+					GraphBuilder, 
+					RDG_EVENT_NAME("HairStrands::VisibilityRasterClassification"), 
+					bCullingEnable ? ComputeShaderClassification_CullingOn : ComputeShaderClassification_CullingOff, 
+					ClassificationParameters, 
+					FIntVector(NumWorkGroups, 1, 1));
+			}
+
+			if (bClassification || bUseHWRaster)
 			{
 				const FVector ViewDirection = ViewInfo.GetViewDirection();
 				const FVector CameraOrigin = ViewInfo.ViewMatrices.GetViewOrigin();
 				const FIntRect Viewport = ViewInfo.ViewRect;
 				const FIntPoint Resolution = SceneDepthTexture->Desc.Extent;
 
-				FVisiblityRasterHWVS::FParameters *ParametersVS = GraphBuilder.AllocParameters<FVisiblityRasterHWVS::FParameters>();
-				ParametersVS->HWRasterPrimIDs = nullptr;
-				ParametersVS->ViewDir = FVector3f(ViewDirection.X, ViewDirection.Y, ViewDirection.Z);
-				ParametersVS->HairMaterialId = PrimitiveInfo.MaterialId;
-				ParametersVS->CameraOrigin = FVector3f(CameraOrigin.X, CameraOrigin.Y, CameraOrigin.Z);
-				ParametersVS->RadiusAtDepth1 = ComputeMinStrandRadiusAtDepth1(Resolution, ViewInfo.FOV, 1, -1.0f).Primary;
-				ParametersVS->ControlPointStart = VertexStart;
-				ParametersVS->bIsOrthoView = !ViewInfo.IsPerspectiveProjection();
-				ParametersVS->HairInstance = GetHairStrandsInstanceParameters(GraphBuilder, ViewInfo, HairGroupPublicData, bCullingEnable, bForceRegister);
-				ParametersVS->ViewUniformBuffer = ViewUniformShaderParameters;
-
-				FVisiblityRasterHWPS::FParameters *ParametersPS = GraphBuilder.AllocParameters<FVisiblityRasterHWPS::FParameters>();
-				ParametersPS->OutHairCountTexture = HairCountTextureUAV;
-				ParametersPS->OutDepthCovTexture = DepthCovTextureUAV;
-				ParametersPS->OutPrimMatTexture = PrimMatTextureUAV;
-				ParametersPS->SampleWeight = SampleWeight;
-				ParametersPS->HairInstance = GetHairStrandsInstanceParameters(GraphBuilder, ViewInfo, HairGroupPublicData, bCullingEnable, bForceRegister);
-				ParametersPS->ViewUniformBuffer = ViewUniformShaderParameters;
-				ParametersPS->RenderTargets.DepthStencil = FDepthStencilBinding(
+				FVisiblityRasterHWParameters* Parameters = GraphBuilder.AllocParameters<FVisiblityRasterHWParameters>();
+				Parameters->VS.RasterPrimIDs = bClassification ? RasterizerPrimIDsSRV : nullptr;
+				Parameters->VS.ViewDir = FVector3f(ViewDirection.X, ViewDirection.Y, ViewDirection.Z);
+				Parameters->VS.HairMaterialId = PrimitiveInfo.MaterialId;
+				Parameters->VS.CameraOrigin = FVector3f(CameraOrigin.X, CameraOrigin.Y, CameraOrigin.Z);
+				Parameters->VS.RadiusAtDepth1 = ComputeMinStrandRadiusAtDepth1(Resolution, ViewInfo.FOV, 1, -1.0f).Primary;
+				Parameters->VS.ControlPointStart = VertexStart;
+				Parameters->VS.PrimIDsBufferSize = MaxNumPrimIDs;
+				Parameters->VS.bIsOrthoView = !ViewInfo.IsPerspectiveProjection();
+				Parameters->VS.HairInstance = GetHairStrandsInstanceParameters(GraphBuilder, ViewInfo, HairGroupPublicData, bCullingEnable, bForceRegister);
+				Parameters->VS.ViewUniformBuffer = ViewUniformShaderParameters;
+				Parameters->VS.IndirectBufferArgs = bClassification ? DrawIndexedIndirectArgs : nullptr;
+				Parameters->PS.OutHairCountTexture = HairCountTextureUAV;
+				Parameters->PS.OutDepthCovTexture = DepthCovTextureUAV;
+				Parameters->PS.OutPrimMatTexture = PrimMatTextureUAV;
+				Parameters->PS.SampleWeight = SampleWeight;
+				Parameters->PS.RenderTargets.DepthStencil = FDepthStencilBinding(
 					SceneDepthTexture,
 					ERenderTargetLoadAction::ELoad,
 					ERenderTargetLoadAction::ENoAction,
 					FExclusiveDepthStencil::DepthRead_StencilNop);
 
+				TShaderMapRef<FVisiblityRasterHWVS> VertexShaderRaster = bCullingEnableInRasterizers ? VertexShaderRaster_CullingOn : VertexShaderRaster_CullingOff;
+				const uint32 NumInstances = bCullingEnableInRasterizers ? VFInput.Strands.VertexCount : VertexCount;
+
 				GraphBuilder.AddPass(
 					RDG_EVENT_NAME("HairStrands::VisibilityRasterHW"),
-					ParametersPS,
+					Parameters,
 					ERDGPassFlags::Raster,
-					[ParametersVS, ParametersPS, VertexShaderRaster, PixelShaderRaster, Viewport, Resolution, VertexCount](FRHICommandList &RHICmdList)
+					[Parameters, VertexShaderRaster, bClassification, PixelShaderRaster, Viewport, Resolution, NumInstances, DrawIndexedIndirectArgs](FRHICommandList &RHICmdList)
 					{
 						FGraphicsPipelineStateInitializer GraphicsPSOInit;
 						RHICmdList.ApplyCachedRenderTargets(GraphicsPSOInit);
@@ -3313,85 +3442,107 @@ static FRasterComputeOutput AddVisibilityComputeRasterPass(
 						GraphicsPSOInit.PrimitiveType = PT_TriangleList;
 						SetGraphicsPipelineState(RHICmdList, GraphicsPSOInit, 0);
 
-						SetShaderParameters(RHICmdList, VertexShaderRaster, VertexShaderRaster.GetPixelShader(), *ParametersVS);
-						SetShaderParameters(RHICmdList, PixelShaderRaster, PixelShaderRaster.GetPixelShader(), *ParametersPS);
+						SetShaderParameters(RHICmdList, VertexShaderRaster, VertexShaderRaster.GetVertexShader(), Parameters->VS);
+						SetShaderParameters(RHICmdList, PixelShaderRaster, PixelShaderRaster.GetPixelShader(), Parameters->PS);
 
 						RHICmdList.SetViewport(Viewport.Min.X, Viewport.Min.Y, 0.0f, Viewport.Max.X, Viewport.Max.Y, 1.0f);
 						RHICmdList.SetStreamSource(0, nullptr, 0);
-						RHICmdList.DrawIndexedPrimitive(GTwoTrianglesIndexBuffer.IndexBufferRHI, 0, 0, 6, 0, 2, VertexCount);
+						if (bClassification)
+						{
+							RHICmdList.DrawIndexedPrimitiveIndirect(GTwoTrianglesIndexBuffer.IndexBufferRHI, DrawIndexedIndirectArgs->GetIndirectRHICallBuffer(), 0);
+						}
+						else
+						{
+							RHICmdList.DrawIndexedPrimitive(GTwoTrianglesIndexBuffer.IndexBufferRHI, 0, 0, 6, 0, 2, NumInstances);
+						}
 					});
 			}
-			else
+
+			if (bClassification || !bUseHWRaster)
 			{
 				if (bUseNaiveSWRaster)
 				{
-					FVisiblityRasterComputeNaiveCS::FParameters *RasterParameters = GraphBuilder.AllocParameters<FVisiblityRasterComputeNaiveCS::FParameters>();
+					const uint32 NumWorkGroups = NumRasterizersNaive;
+
+					FVisiblityRasterComputeNaiveCS::FParameters* RasterParameters = GraphBuilder.AllocParameters<FVisiblityRasterComputeNaiveCS::FParameters>();
 					RasterParameters->OutputResolution = Common.OutputResolutionf;
 					RasterParameters->HairMaterialId = PrimitiveInfo.MaterialId;
 					RasterParameters->ControlPointStart = VertexStart;
 					RasterParameters->ControlPointCount = VertexCount;
 					RasterParameters->SampleWeight = SampleWeight;
-					RasterParameters->HairInstance = GetHairStrandsInstanceParameters(GraphBuilder, ViewInfo, HairGroupPublicData, bCullingEnable, bForceRegister);
+					RasterParameters->NumWorkGroups = NumWorkGroups;
+					RasterParameters->HairInstance = GetHairStrandsInstanceParameters(GraphBuilder, ViewInfo, HairGroupPublicData, bCullingEnableInRasterizers, bForceRegister);
 					RasterParameters->ViewUniformBuffer = ViewUniformShaderParameters;
 					RasterParameters->OutHairCountTexture = HairCountTextureUAV;
 					RasterParameters->OutDepthCovTexture = DepthCovTextureUAV;
 					RasterParameters->OutPrimMatTexture = PrimMatTextureUAV;
 					RasterParameters->SceneDepthTexture = SceneDepthTexture;
-
-					const uint32 NumWorkgroups = (VertexCount + 63) / 64;
-					FComputeShaderUtils::AddPass(GraphBuilder, RDG_EVENT_NAME("HairStrands::VisibilityComputeRasterNaive"), ComputeShaderRasterNaive, RasterParameters, FIntVector(NumWorkgroups, 1, 1));
+					RasterParameters->IndirectPrimIDCount = bClassification ? SWRasterPrimIDCountSRV : nullptr;
+					RasterParameters->IndirectPrimIDs = bClassification ? RasterizerPrimIDsSRV : nullptr;
+
+					FComputeShaderUtils::AddPass(
+						GraphBuilder, 
+						RDG_EVENT_NAME("HairStrands::VisibilityComputeRasterNaive"), 
+						bCullingEnableInRasterizers ? ComputeShaderRasterNaive_CullingOn : ComputeShaderRasterNaive_CullingOff,
+						RasterParameters, 
+						FIntVector(NumWorkGroups, 1, 1));
 				}
 				else
 				{
 					// Reset buffers
 					{
 						uint32 IndexGridClearValues[4] = { 0x0,0x0,0x0,0x0 };
-						AddClearUAVPass(GraphBuilder, OutVisTileBinningGridUAV, IndexGridClearValues);
-						AddClearUAVPass(GraphBuilder, OutVisTileArgsUAV, 0u);
+						AddClearUAVPass(GraphBuilder, VisTileBinningGridUAV, IndexGridClearValues);
+						AddClearUAVPass(GraphBuilder, VisTileArgsUAV, 0u);
 					}
 
 					// Binning pass
 					{
-						FVisiblityRasterComputeBinningCS::FParameters *BinningParameters = GraphBuilder.AllocParameters<FVisiblityRasterComputeBinningCS::FParameters>();
+						FVisiblityRasterComputeBinningCS::FParameters* BinningParameters = GraphBuilder.AllocParameters<FVisiblityRasterComputeBinningCS::FParameters>();
 						BinningParameters->Common = Common;
 						BinningParameters->MacroGroupId = MacroGroup.MacroGroupId;
 						BinningParameters->HairMaterialId = PrimitiveInfo.MaterialId;
 						BinningParameters->ViewUniformBuffer = ViewUniformShaderParameters;
 						BinningParameters->SceneDepthTexture = SceneDepthTexture;
-						BinningParameters->OutVisTilePrims = OutVisTilePrimsUAV;
-						BinningParameters->OutVisTileBinningGrid = OutVisTileBinningGridUAV;
-						BinningParameters->VisTileDepthGrid = OutVisTileDepthGrid;
-						BinningParameters->OutVisTileArgs = OutVisTileArgsUAV;
-						BinningParameters->OutVisTileData = OutVisTileDataUAV;
-						BinningParameters->HairInstance = GetHairStrandsInstanceParameters(GraphBuilder, ViewInfo, HairGroupPublicData, bCullingEnable, bForceRegister);
-
+						BinningParameters->OutVisTilePrims = GraphBuilder.CreateUAV(VisTilePrims, PF_R32_UINT);
+						BinningParameters->OutVisTileBinningGrid = VisTileBinningGridUAV;
+						BinningParameters->VisTileDepthGrid = VisTileDepthGrid;
+						BinningParameters->OutVisTileArgs = VisTileArgsUAV;
+						BinningParameters->OutVisTileData = GraphBuilder.CreateUAV(VisTileData);
+						BinningParameters->HairInstance = GetHairStrandsInstanceParameters(GraphBuilder, ViewInfo, HairGroupPublicData, bCullingEnableInRasterizers, bForceRegister);
 						BinningParameters->VertexCount = VertexCount;
 						BinningParameters->VertexStart = VertexStart;
-
-						BinningParameters->VisTileBinningGridTex = OutVisTileBinningGrid;
-
-						FComputeShaderUtils::AddPass(GraphBuilder, RDG_EVENT_NAME("HairStrands::VisibilityComputeRasterBinning(culling=%s)", bCullingEnable ? TEXT("On") : TEXT("Off")), bCullingEnable ? ComputeShaderBinning_CullingOn : ComputeShaderBinning_CullingOff, BinningParameters, FIntVector(NumBinners, 1, 1));
+						BinningParameters->VisTileBinningGridTex = VisTileBinningGrid;
+						BinningParameters->IndirectPrimIDCount = bClassification ? SWRasterPrimIDCountSRV : nullptr;
+						BinningParameters->IndirectPrimIDs = bClassification ? RasterizerPrimIDsSRV : nullptr;
+
+						FComputeShaderUtils::AddPass(
+							GraphBuilder, 
+							RDG_EVENT_NAME("HairStrands::VisibilityComputeRasterBinning(culling=%s)", bCullingEnableInRasterizers ? TEXT("On") : TEXT("Off")),
+							bCullingEnableInRasterizers ? ComputeShaderBinning_CullingOn : ComputeShaderBinning_CullingOff,
+							BinningParameters, 
+							FIntVector(NumBinners, 1, 1));
 					}
 
 					if (bTileCompactionPass)
 					{
-						FVisiblityRasterComputeCompactionCS::FParameters *CompactionParameters = GraphBuilder.AllocParameters<FVisiblityRasterComputeCompactionCS::FParameters>();
+						FVisiblityRasterComputeCompactionCS::FParameters* CompactionParameters = GraphBuilder.AllocParameters<FVisiblityRasterComputeCompactionCS::FParameters>();
 						CompactionParameters->Common = Common;
 						CompactionParameters->ViewUniformBuffer = ViewUniformShaderParameters;
-						CompactionParameters->HairInstance = GetHairStrandsInstanceParameters(GraphBuilder, ViewInfo, HairGroupPublicData, bCullingEnable, bForceRegister);
-						CompactionParameters->VisTileData = GraphBuilder.CreateSRV(OutVisTileData);
-						CompactionParameters->VisTilePrims = GraphBuilder.CreateSRV(OutVisTilePrims, PF_R32_UINT);
-						CompactionParameters->VisTileArgs = GraphBuilder.CreateSRV(OutVisTileArgs, PF_R32_UINT);
-						CompactionParameters->OutCompactedVisTileData = OutCompactedVisTileDataUAV;
-						CompactionParameters->OutCompactedVisTilePrims = OutCompactedVisTilePrimsUAV;
-						CompactionParameters->OutCompactedVisTileArgs = OutCompactedVisTileArgsUAV;
+						CompactionParameters->HairInstance = GetHairStrandsInstanceParameters(GraphBuilder, ViewInfo, HairGroupPublicData, bCullingEnableInRasterizers, bForceRegister);
+						CompactionParameters->VisTileData = VisTileDataSRV;
+						CompactionParameters->VisTilePrims = VisTilePrimsSRV;
+						CompactionParameters->VisTileArgs = VisTileArgsSRV;
+						CompactionParameters->OutCompactedVisTileData = GraphBuilder.CreateUAV(CompactedVisTileData);
+						CompactionParameters->OutCompactedVisTilePrims = GraphBuilder.CreateUAV(CompactedVisTilePrims, PF_R32_UINT);
+						CompactionParameters->OutCompactedVisTileArgs = GraphBuilder.CreateUAV(CompactedVisTileArgs, PF_R32_UINT);
 
 						FComputeShaderUtils::AddPass(GraphBuilder, RDG_EVENT_NAME("HairStrands::VisibilityComputeRasterCompaction"), ComputeShaderCompaction, CompactionParameters, FIntVector(TileGridRes.X, TileGridRes.Y, 1));
 					}
 
 					// Raster pass
 					{
-						FVisiblityRasterComputeRasterizeCS::FParameters *RasterParameters = GraphBuilder.AllocParameters<FVisiblityRasterComputeRasterizeCS::FParameters>();
+						FVisiblityRasterComputeRasterizeCS::FParameters* RasterParameters = GraphBuilder.AllocParameters<FVisiblityRasterComputeRasterizeCS::FParameters>();
 						RasterParameters->Common = Common;
 						RasterParameters->MacroGroupId = MacroGroup.MacroGroupId;
 						RasterParameters->HairMaterialId = PrimitiveInfo.MaterialId;
@@ -3401,27 +3552,24 @@ static FRasterComputeOutput AddVisibilityComputeRasterPass(
 						RasterParameters->OutDepthCovTexture = DepthCovTextureUAV;
 						RasterParameters->OutPrimMatTexture = PrimMatTextureUAV;
 						RasterParameters->SampleWeight = SampleWeight;
-						RasterParameters->HairInstance = GetHairStrandsInstanceParameters(GraphBuilder, ViewInfo, HairGroupPublicData, bCullingEnable, bForceRegister);
-						RasterParameters->VisTilePrims = GraphBuilder.CreateSRV(bTileCompactionPass ? OutCompactedVisTilePrims : OutVisTilePrims, PF_R32_UINT);
-						RasterParameters->RWVisTileDepthGrid = OutVisTileDepthGridUAV;
-						RasterParameters->VisTileArgs = GraphBuilder.CreateSRV(bTileCompactionPass ? OutCompactedVisTileArgs : OutVisTileArgs, PF_R32_UINT);
-						RasterParameters->RWVisTileData = bTileCompactionPass ? OutCompactedVisTileDataUAV : OutVisTileDataUAV;
-						RasterParameters->VisTileIndirect = GraphBuilder.CreateSRV(OutVisTileIndirect);
-						RasterParameters->IndirectBufferArgs = OutVisTileIndirect;
+						RasterParameters->HairInstance = GetHairStrandsInstanceParameters(GraphBuilder, ViewInfo, HairGroupPublicData, bCullingEnableInRasterizers, bForceRegister);
+						RasterParameters->VisTilePrims = bTileCompactionPass ? GraphBuilder.CreateSRV(CompactedVisTilePrims, PF_R32_UINT) : VisTilePrimsSRV;
+						RasterParameters->VisTileArgs = bTileCompactionPass ? GraphBuilder.CreateSRV(CompactedVisTileArgs, PF_R32_UINT) : VisTileArgsSRV;
+						RasterParameters->VisTileData = bTileCompactionPass ? GraphBuilder.CreateSRV(CompactedVisTileData) : VisTileDataSRV;
 
 						FComputeShaderUtils::AddPass(GraphBuilder, RDG_EVENT_NAME("HairStrands::VisibilityComputeRasterRaster(tiled)"), ComputeShaderRaster, RasterParameters, FIntVector(NumRasterizers, 1, 1));
 					}
 
 					if (bDebugEnabled)
 					{
-						FVisiblityRasterComputeDebugCS::FParameters *Parameters = GraphBuilder.AllocParameters<FVisiblityRasterComputeDebugCS::FParameters>();
+						FVisiblityRasterComputeDebugCS::FParameters* Parameters = GraphBuilder.AllocParameters<FVisiblityRasterComputeDebugCS::FParameters>();
 						Parameters->MacroGroupId = 0;
 						Parameters->Common = Common;
 						Parameters->PrimitiveInfoIndex = PrimitiveInfoIndex;
 						Parameters->TotalPrimitiveInfoCount = TotalPrimitiveInfoCount;
-						Parameters->VisTileDepthGrid = OutVisTileDepthGrid;
-						Parameters->VisTileBinningGrid = OutVisTileBinningGrid;
-						Parameters->VisTileArgs = GraphBuilder.CreateSRV(OutVisTileArgs, PF_R32_UINT);
+						Parameters->VisTileDepthGrid = VisTileDepthGrid;
+						Parameters->VisTileBinningGrid = VisTileBinningGrid;
+						Parameters->VisTileArgs = VisTileArgsSRV;
 						ShaderPrint::SetParameters(GraphBuilder, ViewInfo.ShaderPrintData, Parameters->ShaderPrintParameters);
 
 						TShaderMapRef<FVisiblityRasterComputeDebugCS> DebugComputeShader(ViewInfo.ShaderMap);
