Remove Omni-Directional Stereo Capture (experimental)

diff --git a/Engine/Plugins/Runtime/nDisplay/Source/DisplayCluster/Private/Render/Viewport/DisplayClusterViewport_EditorPreview.cpp b/Engine/Plugins/Runtime/nDisplay/Source/DisplayCluster/Private/Render/Viewport/DisplayClusterViewport_EditorPreview.cpp
--- a/Engine/Plugins/Runtime/nDisplay/Source/DisplayCluster/Private/Render/Viewport/DisplayClusterViewport_EditorPreview.cpp
+++ b/Engine/Plugins/Runtime/nDisplay/Source/DisplayCluster/Private/Render/Viewport/DisplayClusterViewport_EditorPreview.cpp
@@ -86,7 +86,6 @@ FSceneView* FDisplayClusterViewport::ImplCalcScenePreview(FSceneViewFamilyContex
 			FPlane(0, 1, 0, 0),
 			FPlane(0, 0, 0, 1));
 
-		float StereoIPD = 0.f;
 		FIntRect ViewRect = Contexts[InContextNum].RenderTargetRect;
 
 		FSceneViewInitOptions ViewInitOptions;
@@ -112,8 +111,6 @@ FSceneView* FDisplayClusterViewport::ImplCalcScenePreview(FSceneViewFamilyContex
 			ViewInitOptions.WorldToMetersScale = InOutViewFamily.Scene->GetWorld()->GetWorldSettings()->WorldToMeters;
 		}
 
-		ViewInitOptions.StereoIPD = StereoIPD * (ViewInitOptions.WorldToMetersScale / 100.0f);
-
 		ViewInitOptions.BackgroundColor = FLinearColor::Black;
 
 		if (Owner.GetRenderFrameSettings().bPreviewEnablePostProcess == false)
diff --git a/Engine/Shaders/Private/AnisotropyPassShader.usf b/Engine/Shaders/Private/AnisotropyPassShader.usf
--- a/Engine/Shaders/Private/AnisotropyPassShader.usf
+++ b/Engine/Shaders/Private/AnisotropyPassShader.usf
@@ -67,12 +67,7 @@ void MainVertexShader(
 
 	{
 		float4 RasterizedWorldPosition = VertexFactoryGetRasterizedWorldPosition(Input, VFIntermediates, WorldPos);
-	#if ODS_CAPTURE
-		float3 ODS = OffsetODS(RasterizedWorldPosition.xyz, ResolvedView.TranslatedWorldCameraOrigin.xyz, ResolvedView.StereoIPD);
-		Output.Position = INVARIANT(mul(float4(RasterizedWorldPosition.xyz + ODS, 1.0), ResolvedView.TranslatedWorldToClip));
-	#else
 		Output.Position = INVARIANT(mul(RasterizedWorldPosition, ResolvedView.TranslatedWorldToClip));
-	#endif
 	}
 
 	#if INSTANCED_STEREO && !MULTI_VIEW
diff --git a/Engine/Shaders/Private/BasePassVertexShader.usf b/Engine/Shaders/Private/BasePassVertexShader.usf
--- a/Engine/Shaders/Private/BasePassVertexShader.usf
+++ b/Engine/Shaders/Private/BasePassVertexShader.usf
@@ -69,12 +69,7 @@ void Main(
 
 	{
 		float4 RasterizedWorldPosition = VertexFactoryGetRasterizedWorldPosition(Input, VFIntermediates, WorldPosition);
-	#if ODS_CAPTURE
-		float3 ODS = OffsetODS(RasterizedWorldPosition.xyz, ResolvedView.TranslatedWorldCameraOrigin.xyz, ResolvedView.StereoIPD);
-		ClipSpacePosition = mul(float4(RasterizedWorldPosition.xyz + ODS, 1.0), ResolvedView.TranslatedWorldToClip);
-	#else
 		ClipSpacePosition = mul(RasterizedWorldPosition, ResolvedView.TranslatedWorldToClip);
-	#endif
 		Output.Position = INVARIANT(ClipSpacePosition);
 	}
 
diff --git a/Engine/Shaders/Private/DepthOnlyVertexShader.usf b/Engine/Shaders/Private/DepthOnlyVertexShader.usf
--- a/Engine/Shaders/Private/DepthOnlyVertexShader.usf
+++ b/Engine/Shaders/Private/DepthOnlyVertexShader.usf
@@ -70,12 +70,7 @@ void Main(
 
 	{
 		float4 RasterizedWorldPosition = VertexFactoryGetRasterizedWorldPosition(Input, VFIntermediates, WorldPos);
-	#if ODS_CAPTURE
-		float3 ODS = OffsetODS(RasterizedWorldPosition.xyz, ResolvedView.TranslatedWorldCameraOrigin.xyz, ResolvedView.StereoIPD);
-		Output.Position = INVARIANT(mul(float4(RasterizedWorldPosition.xyz + ODS, 1.0), ResolvedView.TranslatedWorldToClip));
-	#else
 		Output.Position = INVARIANT(mul(RasterizedWorldPosition, ResolvedView.TranslatedWorldToClip));
-	#endif
 	}
 
 	#if INSTANCED_STEREO && !MULTI_VIEW
diff --git a/Engine/Shaders/Private/ODSCapture.usf b/Engine/Shaders/Private/ODSCapture.usf
deleted file mode 100644
diff --git a/Engine/Shaders/Private/PositionOnlyDepthVertexShader.usf b/Engine/Shaders/Private/PositionOnlyDepthVertexShader.usf
--- a/Engine/Shaders/Private/PositionOnlyDepthVertexShader.usf
+++ b/Engine/Shaders/Private/PositionOnlyDepthVertexShader.usf
@@ -37,12 +37,7 @@ void Main(
 	float4 WorldPos = VertexFactoryGetWorldPosition(Input);
 
 	{
-	#if ODS_CAPTURE
-		float3 ODS = OffsetODS(WorldPos.xyz, ResolvedView.TranslatedWorldCameraOrigin.xyz, ResolvedView.StereoIPD);
-		OutPosition = INVARIANT(mul(float4(WorldPos.xyz + ODS, 1.0), ResolvedView.TranslatedWorldToClip));
-	#else
 		OutPosition = INVARIANT(mul(WorldPos, ResolvedView.TranslatedWorldToClip));
-	#endif
 	}
 
 	#if INSTANCED_STEREO && !MULTI_VIEW
diff --git a/Engine/Shaders/Private/VelocityShader.usf b/Engine/Shaders/Private/VelocityShader.usf
--- a/Engine/Shaders/Private/VelocityShader.usf
+++ b/Engine/Shaders/Private/VelocityShader.usf
@@ -84,22 +84,10 @@ void MainVertexShader(
 	{
 		float4 ScreenPos;
 		// compute the old screen pos with the old world position and the old camera matrix
-	#if ODS_CAPTURE
-		float3 PrevODS = OffsetODS(PrevTranslatedWorldPosition.xyz, ResolvedView.PrevTranslatedWorldCameraOrigin, ResolvedView.StereoIPD);
-		float4 PrevScreenPosObj = mul(float4(PrevTranslatedWorldPosition.xyz + PrevODS, 1.0), ResolvedView.PrevTranslatedWorldToClip);
-	#else
 		float4 PrevScreenPosObj = mul(float4(PrevTranslatedWorldPosition.xyz, 1), ResolvedView.PrevTranslatedWorldToClip);
-	#endif
 
 		{
-
-		#if ODS_CAPTURE
-			float3 ODS = OffsetODS(TranslatedWorldPosition.xyz, ResolvedView.TranslatedWorldCameraOrigin.xyz, ResolvedView.StereoIPD);
-			ScreenPos = mul(float4(TranslatedWorldPosition.xyz + ODS, 1.0), ResolvedView.TranslatedWorldToClip);
-		#else
 			ScreenPos = mul(float4(TranslatedWorldPosition.xyz, 1), ResolvedView.TranslatedWorldToClip);
-		#endif
-			
 			Output.Position = ScreenPos;
 		}
 
diff --git a/Engine/Shaders/Private/VertexFactoryCommon.ush b/Engine/Shaders/Private/VertexFactoryCommon.ush
--- a/Engine/Shaders/Private/VertexFactoryCommon.ush
+++ b/Engine/Shaders/Private/VertexFactoryCommon.ush
@@ -209,26 +209,3 @@ float3 OctToUnit( float2 Oct )
 	}
 	return normalize(N);
 }
-
-#if ODS_CAPTURE
-/*
-	Apply omni-directional stereo displacement to a stereo cube map render.
-	This can later be reconstructed into a stereo lat-long by ODSCapture.usf
-	See: https://developers.google.com/vr/jump/rendering-ods-content.pdf
-*/
-float3 OffsetODS(float3 WorldPosition, float3 WorldCameraPosition, float IPD)
-{
-	float3 CamOffset = WorldPosition.xyz - WorldCameraPosition.xyz;
-	
-	float4 D = float4(CamOffset.xyz, dot(CamOffset.xyz, CamOffset.xyz));
-	D *= rsqrt(D.w);
-	
-	float3 WorldUp = float3(0.0, 0.0, 1.0);
-	float3 T = normalize(cross(D.xyz, WorldUp.xyz));
-
-	float A = IPD * IPD / D.w;
-	float B = IPD / D.w * sqrt(D.w * D.w - IPD * IPD);
-	
-	return -A * D + B * T;
-}
-#endif
diff --git a/Engine/Source/Runtime/Engine/Classes/Components/SceneCaptureComponent.h b/Engine/Source/Runtime/Engine/Classes/Components/SceneCaptureComponent.h
--- a/Engine/Source/Runtime/Engine/Classes/Components/SceneCaptureComponent.h
+++ b/Engine/Source/Runtime/Engine/Classes/Components/SceneCaptureComponent.h
@@ -24,7 +24,6 @@ struct FSceneCaptureViewInfo
 	FIntRect ViewRect;
 	EStereoscopicPass StereoPass;
 	int32 StereoViewIndex;
-	float StereoIPD;
 };
 
 USTRUCT(BlueprintType)
diff --git a/Engine/Source/Runtime/Engine/Classes/Components/SceneCaptureComponentCube.h b/Engine/Source/Runtime/Engine/Classes/Components/SceneCaptureComponentCube.h
--- a/Engine/Source/Runtime/Engine/Classes/Components/SceneCaptureComponentCube.h
+++ b/Engine/Source/Runtime/Engine/Classes/Components/SceneCaptureComponentCube.h
@@ -26,29 +26,6 @@ class ENGINE_API USceneCaptureComponentCube : public USceneCaptureComponent
 	UPROPERTY(EditAnywhere, BlueprintReadWrite, Category = SceneCapture)
 	bool bCaptureRotation;
 
-	/**
-		Omni-directional Stereo Capture
-
-		If vr.ODSCapture is enabled and both left, right and ODS render targets are set,
-		we'll ignore the texture target and instead do an omni-directional stereo capture.
-		Warped cube maps will be rendered into the left and right texture targets using the 
-		IPD property for stereo offset. We will then reconstruct a stereo lat-long with the left
-		eye stacked on top of the right eye in the ODS target.
-		See: https://developers.google.com/vr/jump/rendering-ods-content.pdf
-	*/
-
-	UPROPERTY(EditAnywhere, BlueprintReadWrite, Category = "SceneCapture|Omni-directional Stereo Capture")
-	TObjectPtr<class UTextureRenderTargetCube> TextureTargetLeft;
-
-	UPROPERTY(EditAnywhere, BlueprintReadWrite, Category = "SceneCapture|Omni-directional Stereo Capture")
-	TObjectPtr<class UTextureRenderTargetCube> TextureTargetRight;
-
-	UPROPERTY(EditAnywhere, BlueprintReadWrite, Category = "SceneCapture|Omni-directional Stereo Capture")
-	TObjectPtr<class UTextureRenderTarget2D> TextureTargetODS;
-
-	UPROPERTY(EditAnywhere, BlueprintReadWrite, Category = "SceneCapture|Omni-directional Stereo Capture", meta = (DisplayName = "Interpupillary Distance (cm)", UIMin = "0.0", UIMax = "16.0", ClampMin = "0.0", ClampMax = "16.0"))
-	float IPD;
-
 public:
 	//~ Begin UActorComponent Interface
 	virtual void OnComponentDestroyed(bool bDestroyingHierarchy) override;
diff --git a/Engine/Source/Runtime/Engine/Classes/Engine/RendererSettings.h b/Engine/Source/Runtime/Engine/Classes/Engine/RendererSettings.h
--- a/Engine/Source/Runtime/Engine/Classes/Engine/RendererSettings.h
+++ b/Engine/Source/Runtime/Engine/Classes/Engine/RendererSettings.h
@@ -829,12 +829,6 @@ class ENGINE_API URendererSettings : public UDeveloperSettings
 		ConfigRestartRequired = false))
 	uint32 bRoundRobinOcclusion : 1;
 
-	UPROPERTY(config, EditAnywhere, Category = Experimental, meta = (
-		ConsoleVariable = "vr.ODSCapture", DisplayName = "Omni-directional Stereo Capture",
-		ToolTip = "Enable Omni-directional Stereo Capture.",
-		ConfigRestartRequired = true))
-		uint32 bODSCapture : 1;
-
 	UPROPERTY(config, EditAnywhere, Category = Experimental, meta = (
 		ConsoleVariable="r.MeshStreaming",DisplayName="Mesh Streaming",
 		ToolTip="When enabled mesh will stream in based on what is visible on screen.",
diff --git a/Engine/Source/Runtime/Engine/Private/Components/SceneCaptureComponent.cpp b/Engine/Source/Runtime/Engine/Private/Components/SceneCaptureComponent.cpp
--- a/Engine/Source/Runtime/Engine/Private/Components/SceneCaptureComponent.cpp
+++ b/Engine/Source/Runtime/Engine/Private/Components/SceneCaptureComponent.cpp
@@ -1101,7 +1101,6 @@ USceneCaptureComponentCube::USceneCaptureComponentCube(const FObjectInitializer&
 	PrimaryComponentTick.TickGroup = TG_DuringPhysics;
 	PrimaryComponentTick.bAllowTickOnDedicatedServer = false;
 	bTickInEditor = true;
-	IPD = 6.2f;
 	bCaptureRotation = false;
 
 #if WITH_EDITORONLY_DATA
diff --git a/Engine/Source/Runtime/Engine/Private/SceneView.cpp b/Engine/Source/Runtime/Engine/Private/SceneView.cpp
--- a/Engine/Source/Runtime/Engine/Private/SceneView.cpp
+++ b/Engine/Source/Runtime/Engine/Private/SceneView.cpp
@@ -730,7 +730,6 @@ FSceneView::FSceneView(const FSceneViewInitOptions& InitOptions)
 	, StereoPass(InitOptions.StereoPass)
 	, StereoViewIndex(InitOptions.StereoViewIndex)
 	, PrimaryViewIndex(INDEX_NONE)
-	, StereoIPD(InitOptions.StereoIPD)
 	, bAllowCrossGPUTransfer(true)
 	, bOverrideGPUMask(false)
 	, GPUMask(FRHIGPUMask::GPU0())
diff --git a/Engine/Source/Runtime/Engine/Private/ShaderCompiler/ShaderCompiler.cpp b/Engine/Source/Runtime/Engine/Private/ShaderCompiler/ShaderCompiler.cpp
--- a/Engine/Source/Runtime/Engine/Private/ShaderCompiler/ShaderCompiler.cpp
+++ b/Engine/Source/Runtime/Engine/Private/ShaderCompiler/ShaderCompiler.cpp
@@ -5615,10 +5615,6 @@ void GlobalBeginCompileShader(
 			UE_LOG(LogShaderCompilers, Log, TEXT("Instanced stereo rendering is not supported for the %s shader platform."), *ShaderFormatName.ToString());
 			GShaderCompilingManager->SuppressWarnings(ShaderPlatform);
 		}
-
-		static const auto CVarODSCapture = IConsoleManager::Get().FindTConsoleVariableDataInt(TEXT("vr.ODSCapture"));
-		const bool bIsODSCapture = CVarODSCapture && (CVarODSCapture->GetValueOnAnyThread() != 0);
-		Input.Environment.SetDefine(TEXT("ODS_CAPTURE"), bIsODSCapture);
 	}
 
 	ShaderType->AddReferencedUniformBufferIncludes(Input.Environment, Input.SourceFilePrefix, ShaderPlatform);
diff --git a/Engine/Source/Runtime/Engine/Private/TextureRenderTarget2DArray.cpp b/Engine/Source/Runtime/Engine/Private/TextureRenderTarget2DArray.cpp
--- a/Engine/Source/Runtime/Engine/Private/TextureRenderTarget2DArray.cpp
+++ b/Engine/Source/Runtime/Engine/Private/TextureRenderTarget2DArray.cpp
@@ -101,10 +101,7 @@ EMaterialValueType UTextureRenderTarget2DArray::GetMaterialType() const
 #if WITH_EDITOR
 void UTextureRenderTarget2DArray::PostEditChangeProperty(FPropertyChangedEvent& PropertyChangedEvent)
 {
-	// Allow for high resolution captures when ODS is enabled
-	static const auto CVarODSCapture = IConsoleManager::Get().FindTConsoleVariableDataInt(TEXT("vr.ODSCapture"));
-	const bool bIsODSCapture = CVarODSCapture && (CVarODSCapture->GetValueOnGameThread() != 0);
-	const int32 MaxSize = (bIsODSCapture) ? 4096 : 2048;
+	constexpr int32 MaxSize = 2048;
 
 	EPixelFormat Format = GetFormat();
 	SizeX = FMath::Clamp<int32>(SizeX - (SizeX % GPixelFormats[Format].BlockSizeX), 1, MaxSize);
diff --git a/Engine/Source/Runtime/Engine/Private/TextureRenderTargetCube.cpp b/Engine/Source/Runtime/Engine/Private/TextureRenderTargetCube.cpp
--- a/Engine/Source/Runtime/Engine/Private/TextureRenderTargetCube.cpp
+++ b/Engine/Source/Runtime/Engine/Private/TextureRenderTargetCube.cpp
@@ -103,10 +103,7 @@ EMaterialValueType UTextureRenderTargetCube::GetMaterialType() const
 #if WITH_EDITOR
 void UTextureRenderTargetCube::PostEditChangeProperty(FPropertyChangedEvent& PropertyChangedEvent)
 {
-	// Allow for high resolution captures when ODS is enabled
-	static const auto CVarODSCapture = IConsoleManager::Get().FindTConsoleVariableDataInt(TEXT("vr.ODSCapture"));
-	const bool bIsODSCapture = CVarODSCapture && (CVarODSCapture->GetValueOnGameThread() != 0);
-	const int32 MaxSize = (bIsODSCapture) ? 4096 : 2048;
+	constexpr int32 MaxSize = 2048;
 
 	EPixelFormat Format = GetFormat();
 	SizeX = FMath::Clamp<int32>(SizeX - (SizeX % GPixelFormats[Format].BlockSizeX),1,MaxSize);
diff --git a/Engine/Source/Runtime/Engine/Private/TextureRenderTargetVolume.cpp b/Engine/Source/Runtime/Engine/Private/TextureRenderTargetVolume.cpp
--- a/Engine/Source/Runtime/Engine/Private/TextureRenderTargetVolume.cpp
+++ b/Engine/Source/Runtime/Engine/Private/TextureRenderTargetVolume.cpp
@@ -105,10 +105,7 @@ EMaterialValueType UTextureRenderTargetVolume::GetMaterialType() const
 #if WITH_EDITOR
 void UTextureRenderTargetVolume::PostEditChangeProperty(FPropertyChangedEvent& PropertyChangedEvent)
 {
-	// Allow for high resolution captures when ODS is enabled
-	static const auto CVarODSCapture = IConsoleManager::Get().FindTConsoleVariableDataInt(TEXT("vr.ODSCapture"));
-	const bool bIsODSCapture = CVarODSCapture && (CVarODSCapture->GetValueOnGameThread() != 0);
-	const int32 MaxSize = (bIsODSCapture) ? 4096 : 2048;
+	constexpr int32 MaxSize = 2048;
 
 	EPixelFormat Format = GetFormat();
 	SizeX = FMath::Clamp<int32>(SizeX - (SizeX % GPixelFormats[Format].BlockSizeX), 1, MaxSize);
diff --git a/Engine/Source/Runtime/Engine/Public/SceneView.h b/Engine/Source/Runtime/Engine/Public/SceneView.h
--- a/Engine/Source/Runtime/Engine/Public/SceneView.h
+++ b/Engine/Source/Runtime/Engine/Public/SceneView.h
@@ -132,9 +132,6 @@ struct FSceneViewInitOptions : public FSceneViewProjectionData
 	/** For stereoscopic rendering, a unique index to identify the view across view families */
 	int32 StereoViewIndex;
 
-	/** For stereoscopic scene capture rendering. Half of the view's stereo IPD (- for lhs, + for rhs) */
-	float StereoIPD;
-
 	/** Conversion from world units (uu) to meters, so we can scale motion to the world appropriately */
 	float WorldToMetersScale;
 
@@ -193,7 +190,6 @@ struct FSceneViewInitOptions : public FSceneViewProjectionData
 		, ColorScale(FLinearColor::White)
 		, StereoPass(EStereoscopicPass::eSSP_FULL)
 		, StereoViewIndex(INDEX_NONE)
-		, StereoIPD(0.0f)
 		, WorldToMetersScale(100.f)
 		, ViewLocation(ForceInitToZero)
 		, ViewRotation(ForceInitToZero)
@@ -770,7 +766,6 @@ enum ETranslucencyVolumeCascade
 	VIEW_UNIFORM_BUFFER_MEMBER(FVector3f, VolumetricLightmapIndirectionTextureSize) \
 	VIEW_UNIFORM_BUFFER_MEMBER(float, VolumetricLightmapBrickSize) \
 	VIEW_UNIFORM_BUFFER_MEMBER(FVector3f, VolumetricLightmapBrickTexelSize) \
-	VIEW_UNIFORM_BUFFER_MEMBER(float, StereoIPD) \
 	VIEW_UNIFORM_BUFFER_MEMBER(float, IndirectLightingCacheShowFlag) \
 	VIEW_UNIFORM_BUFFER_MEMBER(float, EyeToPixelSpreadAngle) \
 	VIEW_UNIFORM_BUFFER_MEMBER_ARRAY(FVector4f, XRPassthroughCameraUVs, [2]) \
@@ -1019,9 +1014,6 @@ public:
 	/** For stereoscopic rendering, view family index of the primary view associated with this view */
 	int32 PrimaryViewIndex;
 
-	/** Half of the view's stereo IPD (- for lhs, + for rhs) */
-	float StereoIPD;
-
 	/** Allow cross GPU transfer for this view */
 	bool bAllowCrossGPUTransfer;
 
diff --git a/Engine/Source/Runtime/RenderCore/Private/Shader.cpp b/Engine/Source/Runtime/RenderCore/Private/Shader.cpp
--- a/Engine/Source/Runtime/RenderCore/Private/Shader.cpp
+++ b/Engine/Source/Runtime/RenderCore/Private/Shader.cpp
@@ -1444,15 +1444,6 @@ void ShaderMapAppendKeyString(EShaderPlatform Platform, FString& KeyString)
 		}
 	}
 
-	{
-		static const auto CVarODSCapture = IConsoleManager::Get().FindTConsoleVariableDataInt(TEXT("vr.ODSCapture"));
-		const bool bIsODSCapture = CVarODSCapture && (CVarODSCapture->GetValueOnGameThread() != 0);
-		if (bIsODSCapture)
-		{
-			KeyString += TEXT("_ODSC");
-		}
-	}
-
 	{
 		KeyString += IsUsingSelectiveBasePassOutputs(Platform) ? TEXT("_SO") : TEXT("");
 	}
diff --git a/Engine/Source/Runtime/Renderer/Private/SceneCaptureRendering.cpp b/Engine/Source/Runtime/Renderer/Private/SceneCaptureRendering.cpp
--- a/Engine/Source/Runtime/Renderer/Private/SceneCaptureRendering.cpp
+++ b/Engine/Source/Runtime/Renderer/Private/SceneCaptureRendering.cpp
@@ -146,23 +146,6 @@ public:
 
 IMPLEMENT_GLOBAL_SHADER(FSceneCapturePS, "/Engine/Private/SceneCapturePixelShader.usf", "Main", SF_Pixel);
 
-class FODSCapturePS : public FGlobalShader
-{
-public:
-	DECLARE_GLOBAL_SHADER(FODSCapturePS);
-	SHADER_USE_PARAMETER_STRUCT(FODSCapturePS, FGlobalShader);
-
-	BEGIN_SHADER_PARAMETER_STRUCT(FParameters, )
-		SHADER_PARAMETER_RDG_TEXTURE(TextureCube, LeftEyeTexture)
-		SHADER_PARAMETER_RDG_TEXTURE(TextureCube, RightEyeTexture)
-		SHADER_PARAMETER_SAMPLER(SamplerState, LeftEyeTextureSampler)
-		SHADER_PARAMETER_SAMPLER(SamplerState, RightEyeTextureSampler)
-		RENDER_TARGET_BINDING_SLOTS()
-	END_SHADER_PARAMETER_STRUCT()
-};
-
-IMPLEMENT_GLOBAL_SHADER(FODSCapturePS, "/Engine/Private/ODSCapture.usf", "MainPS", SF_Pixel);
-
 static bool CaptureNeedsSceneColor(ESceneCaptureSource CaptureSource)
 {
 	return CaptureSource != SCS_FinalColorLDR && CaptureSource != SCS_FinalColorHDR && CaptureSource != SCS_FinalToneCurveHDR;
@@ -440,61 +423,6 @@ void UpdateSceneCaptureContentMobile_RenderThread(
 	SceneRenderer->RenderThreadEnd(RHICmdList);
 }
 
-static void ODSCapture_RenderThread(
-	FRDGBuilder& GraphBuilder,
-	FRDGTextureRef LeftEyeTexture,
-	FRDGTextureRef RightEyeTexture,
-	FRDGTextureRef OutputTexture,
-	const ERHIFeatureLevel::Type FeatureLevel)
-{
-	FODSCapturePS::FParameters* PassParameters = GraphBuilder.AllocParameters<FODSCapturePS::FParameters>();
-	PassParameters->LeftEyeTexture = LeftEyeTexture;
-	PassParameters->RightEyeTexture = RightEyeTexture;
-	PassParameters->LeftEyeTextureSampler = TStaticSamplerState<SF_Bilinear>::GetRHI();
-	PassParameters->RightEyeTextureSampler = TStaticSamplerState<SF_Bilinear>::GetRHI();
-	PassParameters->RenderTargets[0] = FRenderTargetBinding(OutputTexture, ERenderTargetLoadAction::ELoad);
-
-	const auto ShaderMap = GetGlobalShaderMap(FeatureLevel);
-	TShaderMapRef<FScreenVS> VertexShader(ShaderMap);
-	TShaderMapRef<FODSCapturePS> PixelShader(ShaderMap);
-
-	GraphBuilder.AddPass(
-		RDG_EVENT_NAME("ODSCapture"),
-		PassParameters,
-		ERDGPassFlags::Raster,
-		[VertexShader, PixelShader, OutputTexture](FRHICommandList& RHICmdList)
-	{
-		FGraphicsPipelineStateInitializer GraphicsPSOInit;
-		RHICmdList.ApplyCachedRenderTargets(GraphicsPSOInit);
-		GraphicsPSOInit.BlendState = TStaticBlendState<>::GetRHI();
-		GraphicsPSOInit.RasterizerState = TStaticRasterizerState<>::GetRHI();
-		GraphicsPSOInit.DepthStencilState = TStaticDepthStencilState<false, CF_Always>::GetRHI();
-
-		extern TGlobalResource<FFilterVertexDeclaration> GFilterVertexDeclaration;
-
-		GraphicsPSOInit.BoundShaderState.VertexDeclarationRHI = GFilterVertexDeclaration.VertexDeclarationRHI;
-		GraphicsPSOInit.BoundShaderState.VertexShaderRHI = VertexShader.GetVertexShader();
-		GraphicsPSOInit.BoundShaderState.PixelShaderRHI = PixelShader.GetPixelShader();
-		GraphicsPSOInit.PrimitiveType = PT_TriangleList;
-
-		SetGraphicsPipelineState(RHICmdList, GraphicsPSOInit, 0);
-
-		const FIntPoint& TargetSize = OutputTexture->Desc.Extent;
-		RHICmdList.SetViewport(0, 0, 0.0f, TargetSize.X, TargetSize.Y, 1.0f);
-
-		DrawRectangle(
-			RHICmdList,
-			0, 0,
-			static_cast<float>(TargetSize.X), static_cast<float>(TargetSize.Y),
-			0, 0,
-			TargetSize.X, TargetSize.Y,
-			TargetSize,
-			TargetSize,
-			VertexShader,
-			EDRF_UseTriangleOptimization);
-	});
-}
-
 static void UpdateSceneCaptureContent_RenderThread(
 	FRHICommandListImmediate& RHICmdList,
 	FSceneRenderer* SceneRenderer,
@@ -671,7 +599,6 @@ void SetupViewFamilyForSceneCapture(
 		{
 			ViewInitOptions.WorldToMetersScale = ViewFamily.Scene->GetWorld()->GetWorldSettings()->WorldToMeters;
 		}
-		ViewInitOptions.StereoIPD = SceneCaptureViewInfo.StereoIPD * (ViewInitOptions.WorldToMetersScale / 100.0f);
 
 		if (bCaptureSceneColor)
 		{
@@ -779,8 +706,7 @@ static FSceneRenderer* CreateSceneRendererForSceneCapture(
 	FPostProcessSettings* PostProcessSettings,
 	float PostProcessBlendWeight,
 	const AActor* ViewActor, 
-	int32 CubemapFaceIndex = INDEX_NONE,
-	const float StereoIPD = 0.0f)
+	int32 CubemapFaceIndex = INDEX_NONE)
 {
 	FSceneCaptureViewInfo SceneCaptureViewInfo;
 	SceneCaptureViewInfo.ViewRotationMatrix = ViewRotationMatrix;
@@ -788,7 +714,6 @@ static FSceneRenderer* CreateSceneRendererForSceneCapture(
 	SceneCaptureViewInfo.ProjectionMatrix = ProjectionMatrix;
 	SceneCaptureViewInfo.StereoPass = EStereoscopicPass::eSSP_FULL;
 	SceneCaptureViewInfo.StereoViewIndex = INDEX_NONE;
-	SceneCaptureViewInfo.StereoIPD = StereoIPD;
 	SceneCaptureViewInfo.ViewRect = FIntRect(0, 0, RenderTargetSize.X, RenderTargetSize.Y);
 
 	FSceneViewFamilyContext ViewFamily(FSceneViewFamily::ConstructionValues(
@@ -1087,16 +1012,6 @@ void FScene::UpdateSceneCaptureContents(USceneCaptureComponentCube* CaptureCompo
 
 	check(CaptureComponent);
 
-	const bool bIsODS = CaptureComponent->TextureTargetLeft && CaptureComponent->TextureTargetRight && CaptureComponent->TextureTargetODS;
-	const uint32 StartIndex = (bIsODS) ? 1 : 0;
-	const uint32 EndIndex = (bIsODS) ? 3 : 1;
-	
-	UTextureRenderTargetCube* const TextureTargets[] = {
-		CaptureComponent->TextureTarget, 
-		CaptureComponent->TextureTargetLeft, 
-		CaptureComponent->TextureTargetRight
-	};
-
 	FTransform Transform = CaptureComponent->GetComponentToWorld();
 	const FVector ViewLocation = Transform.GetTranslation();
 
@@ -1107,91 +1022,58 @@ void FScene::UpdateSceneCaptureContents(USceneCaptureComponentCube* CaptureCompo
 		Transform.SetScale3D(FVector::OneVector);
 	}
 
-	for (uint32 CaptureIter = StartIndex; CaptureIter < EndIndex; ++CaptureIter)
-	{
-		UTextureRenderTargetCube* const TextureTarget = TextureTargets[CaptureIter];
+	UTextureRenderTargetCube* const TextureTarget = CaptureComponent->TextureTarget;
 
-		if (TextureTarget)
+	if (TextureTarget)
+	{
+		const float FOV = 90 * (float)PI / 360.0f;
+		for (int32 faceidx = 0; faceidx < (int32)ECubeFace::CubeFace_MAX; faceidx++)
 		{
-			const float FOV = 90 * (float)PI / 360.0f;
-			for (int32 faceidx = 0; faceidx < (int32)ECubeFace::CubeFace_MAX; faceidx++)
-			{
-				const ECubeFace TargetFace = (ECubeFace)faceidx;
-				const FVector Location = CaptureComponent->GetComponentToWorld().GetTranslation();
+			const ECubeFace TargetFace = (ECubeFace)faceidx;
+			const FVector Location = CaptureComponent->GetComponentToWorld().GetTranslation();
 
-				FMatrix ViewRotationMatrix;
-
-				if (CaptureComponent->bCaptureRotation)
-				{
-					ViewRotationMatrix = Transform.ToInverseMatrixWithScale() * FLocal::CalcCubeFaceTransform(TargetFace);
-				}
-				else
-				{
-					ViewRotationMatrix = FLocal::CalcCubeFaceTransform(TargetFace);
-				}
-				FIntPoint CaptureSize(TextureTarget->GetSurfaceWidth(), TextureTarget->GetSurfaceHeight());
-				FMatrix ProjectionMatrix;
-				BuildProjectionMatrix(CaptureSize, FOV, GNearClippingPlane, ProjectionMatrix);
-				FPostProcessSettings PostProcessSettings;
+			FMatrix ViewRotationMatrix;
 
-				float StereoIPD = 0.0f;
-				if (bIsODS)
-				{
-					StereoIPD = (CaptureIter == 1) ? CaptureComponent->IPD * -0.5f : CaptureComponent->IPD * 0.5f;
-				}
+			if (CaptureComponent->bCaptureRotation)
+			{
+				ViewRotationMatrix = Transform.ToInverseMatrixWithScale() * FLocal::CalcCubeFaceTransform(TargetFace);
+			}
+			else
+			{
+				ViewRotationMatrix = FLocal::CalcCubeFaceTransform(TargetFace);
+			}
+			FIntPoint CaptureSize(TextureTarget->GetSurfaceWidth(), TextureTarget->GetSurfaceHeight());
+			FMatrix ProjectionMatrix;
+			BuildProjectionMatrix(CaptureSize, FOV, GNearClippingPlane, ProjectionMatrix);
+			FPostProcessSettings PostProcessSettings;
 
-				bool bCaptureSceneColor = CaptureNeedsSceneColor(CaptureComponent->CaptureSource);
+			bool bCaptureSceneColor = CaptureNeedsSceneColor(CaptureComponent->CaptureSource);
 
-				FSceneRenderer* SceneRenderer = CreateSceneRendererForSceneCapture(this, CaptureComponent,
-					TextureTarget->GameThread_GetRenderTargetResource(), CaptureSize, ViewRotationMatrix,
-					Location, ProjectionMatrix, false, CaptureComponent->MaxViewDistanceOverride,
-					bCaptureSceneColor, &PostProcessSettings, 0, CaptureComponent->GetViewOwner(), faceidx, StereoIPD);
+			FSceneRenderer* SceneRenderer = CreateSceneRendererForSceneCapture(this, CaptureComponent,
+				TextureTarget->GameThread_GetRenderTargetResource(), CaptureSize, ViewRotationMatrix,
+				Location, ProjectionMatrix, false, CaptureComponent->MaxViewDistanceOverride,
+				bCaptureSceneColor, &PostProcessSettings, 0, CaptureComponent->GetViewOwner(), faceidx);
 
-				SceneRenderer->ViewFamily.SceneCaptureSource = CaptureComponent->CaptureSource;
+			SceneRenderer->ViewFamily.SceneCaptureSource = CaptureComponent->CaptureSource;
 
-				FTextureRenderTargetCubeResource* TextureRenderTarget = static_cast<FTextureRenderTargetCubeResource*>(TextureTarget->GameThread_GetRenderTargetResource());
-				FString EventName;
-				if (!CaptureComponent->ProfilingEventName.IsEmpty())
-				{
-					EventName = CaptureComponent->ProfilingEventName;
-				}
-				else if (CaptureComponent->GetOwner())
+			FTextureRenderTargetCubeResource* TextureRenderTarget = static_cast<FTextureRenderTargetCubeResource*>(TextureTarget->GameThread_GetRenderTargetResource());
+			FString EventName;
+			if (!CaptureComponent->ProfilingEventName.IsEmpty())
+			{
+				EventName = CaptureComponent->ProfilingEventName;
+			}
+			else if (CaptureComponent->GetOwner())
+			{
+				CaptureComponent->GetOwner()->GetFName().ToString(EventName);
+			}
+			ENQUEUE_RENDER_COMMAND(CaptureCommand)(
+				[SceneRenderer, TextureRenderTarget, EventName, TargetFace](FRHICommandListImmediate& RHICmdList)
 				{
-					CaptureComponent->GetOwner()->GetFName().ToString(EventName);
+					FRHICopyTextureInfo CopyInfo;
+					CopyInfo.DestSliceIndex = TargetFace;
+					UpdateSceneCaptureContent_RenderThread(RHICmdList, SceneRenderer, TextureRenderTarget, TextureRenderTarget, EventName, CopyInfo, false, FGenerateMipsParams(), true, false);
 				}
-				ENQUEUE_RENDER_COMMAND(CaptureCommand)(
-					[SceneRenderer, TextureRenderTarget, EventName, TargetFace](FRHICommandListImmediate& RHICmdList)
-					{
-						FRHICopyTextureInfo CopyInfo;
-						CopyInfo.DestSliceIndex = TargetFace;
-						UpdateSceneCaptureContent_RenderThread(RHICmdList, SceneRenderer, TextureRenderTarget, TextureRenderTarget, EventName, CopyInfo, false, FGenerateMipsParams(), true, false);
-					}
-				);
-			}
+			);
 		}
 	}
-
-	if (bIsODS)
-	{
-		const FTextureRenderTargetCubeResource* const LeftEye = static_cast<FTextureRenderTargetCubeResource*>(CaptureComponent->TextureTargetLeft->GameThread_GetRenderTargetResource());
-		const FTextureRenderTargetCubeResource* const RightEye = static_cast<FTextureRenderTargetCubeResource*>(CaptureComponent->TextureTargetRight->GameThread_GetRenderTargetResource());
-		FTextureRenderTargetResource* const RenderTarget = CaptureComponent->TextureTargetODS->GameThread_GetRenderTargetResource();
-		const ERHIFeatureLevel::Type InFeatureLevel = FeatureLevel;
-
-		ENQUEUE_RENDER_COMMAND(ODSCaptureCommand)(
-			[LeftEye, RightEye, RenderTarget, InFeatureLevel](FRHICommandListImmediate& RHICmdList)
-		{
-			const ERHIAccess FinalAccess = ERHIAccess::RTV;
-
-			FRDGBuilder GraphBuilder(RHICmdList);
-			FRDGTextureRef OutputTexture = GraphBuilder.RegisterExternalTexture(CreateRenderTarget(RenderTarget->GetRenderTargetTexture(), TEXT("Output")));
-			FRDGTextureRef LeftEyeTexture = GraphBuilder.RegisterExternalTexture(CreateRenderTarget(LeftEye->TextureRHI, TEXT("LeftEye")));
-			FRDGTextureRef RightEyeTexture = GraphBuilder.RegisterExternalTexture(CreateRenderTarget(RightEye->TextureRHI, TEXT("RightEye")));
-			ODSCapture_RenderThread(GraphBuilder, LeftEyeTexture, RightEyeTexture, OutputTexture, InFeatureLevel);
-
-			GraphBuilder.SetTextureAccessFinal(LeftEyeTexture, FinalAccess);
-			GraphBuilder.SetTextureAccessFinal(RightEyeTexture, FinalAccess);
-			GraphBuilder.Execute();
-		});
-	}
 }
diff --git a/Engine/Source/Runtime/Renderer/Private/SceneRendering.cpp b/Engine/Source/Runtime/Renderer/Private/SceneRendering.cpp
--- a/Engine/Source/Runtime/Renderer/Private/SceneRendering.cpp
+++ b/Engine/Source/Runtime/Renderer/Private/SceneRendering.cpp
@@ -191,13 +191,6 @@ static TAutoConsoleVariable<int32> CVarRoundRobinOcclusion(
 	TEXT("0 to disable round-robin occlusion queries for stereo rendering (default), 1 to enable."),
 	ECVF_Scalability | ECVF_RenderThreadSafe);
 
-static TAutoConsoleVariable<int32> CVarODSCapture(
-	TEXT("vr.ODSCapture"),
-	0,
-	TEXT("Experimental")
-	TEXT("0 to disable Omni-directional stereo capture (default), 1 to enable."),
-	ECVF_ReadOnly | ECVF_RenderThreadSafe);
-
 static TAutoConsoleVariable<int32> CVarViewRectUseScreenBottom(
 	TEXT("r.ViewRectUseScreenBottom"),
 	0,
@@ -1785,7 +1778,6 @@ void FViewInfo::SetupUniformBufferParameters(
 	ViewUniformShaderParameters.ReflectionEnvironmentRoughnessMixingScaleBiasAndLargestWeight = (FVector3f)GetReflectionEnvironmentRoughnessMixingScaleBiasAndLargestWeight();
 
 	ViewUniformShaderParameters.StereoPassIndex = StereoViewIndex != INDEX_NONE ? StereoViewIndex : 0;
-	ViewUniformShaderParameters.StereoIPD = StereoIPD;
 
 	{
 		auto XRCamera = GEngine->XRSystem ? GEngine->XRSystem->GetXRCamera() : nullptr;
