Fixes crash in TSR when the GPU doesn't support wavesize=32

diff --git a/Engine/Config/VulkanPC/DataDrivenPlatformInfo.ini b/Engine/Config/VulkanPC/DataDrivenPlatformInfo.ini
--- a/Engine/Config/VulkanPC/DataDrivenPlatformInfo.ini
+++ b/Engine/Config/VulkanPC/DataDrivenPlatformInfo.ini
@@ -83,6 +83,8 @@ bSupportsIntrinsicWaveOnce = false
 bSupportsConservativeRasterization = false
 
 bSupportsWaveOperations=RuntimeDependent
+MinimumWaveSize=4
+MaximumWaveSize=128
 
 ; //#todo-vulkanddpi
 bRequiresExplicit128bitRT = false
diff --git a/Engine/Config/Windows/DataDrivenPlatformInfo.ini b/Engine/Config/Windows/DataDrivenPlatformInfo.ini
--- a/Engine/Config/Windows/DataDrivenPlatformInfo.ini
+++ b/Engine/Config/Windows/DataDrivenPlatformInfo.ini
@@ -47,6 +47,8 @@ bSupportsRayTracingIndirectInstanceData=true
 bSupportsPathTracing=true
 bSupportsHighEndRayTracingReflections=true
 bSupportsWaveOperations=RuntimeDependent
+MinimumWaveSize=4
+MaximumWaveSize=128
 bSupportsGPUScene=true
 bSupportsNanite=false
 bSupportsSceneDataCompressedTransforms=true
@@ -104,6 +106,8 @@ bSupportsLumenGI=true
 bSupportsDistanceFields=true
 bSupportsSSDIndirect=true
 bSupportsWaveOperations=RuntimeGuaranteed
+MinimumWaveSize=4
+MaximumWaveSize=128
 bSupportsDxc=true
 bSupportsDiaphragmDOF=true
 bRequiresVendorExtensionsForAtomics=false
diff --git a/Engine/Shaders/Private/TemporalSuperResolution/TSRRejectShading.usf b/Engine/Shaders/Private/TemporalSuperResolution/TSRRejectShading.usf
--- a/Engine/Shaders/Private/TemporalSuperResolution/TSRRejectShading.usf
+++ b/Engine/Shaders/Private/TemporalSuperResolution/TSRRejectShading.usf
@@ -2,9 +2,8 @@
 
 //------------------------------------------------------- CONFIG WAVE SIZE
 
-#define WAVE_SIZE 32
-#if !DIM_USE_WAVE_OPS
-	#define WAVE_BROADCAST_GROUPSIZE (WAVE_SIZE)
+#if DIM_WAVE_SIZE == 0 // fallback to LDS
+	#define WAVE_BROADCAST_GROUPSIZE 32
 #endif
 
 //------------------------------------------------------- INCLUDES
@@ -36,8 +35,24 @@
 //------------------------------------------------------- DERIVED
 
 /** Pixel coordinate stride between each lane in the GROUP_TILE_SIZE. */
-#define LANE_STRIDE_X 4
-#define LANE_STRIDE_Y 2
+#if DIM_WAVE_SIZE == 64
+	#define LANE_COUNT 64
+	#define LANE_STRIDE_X 2
+	#define LANE_STRIDE_Y 2
+
+#elif DIM_WAVE_SIZE == 32
+	#define LANE_COUNT 32
+	#define LANE_STRIDE_X 4
+	#define LANE_STRIDE_Y 2
+
+#elif DIM_WAVE_SIZE == 0
+	#define LANE_COUNT 32
+	#define LANE_STRIDE_X 4
+	#define LANE_STRIDE_Y 2
+
+#else
+	#error Unknown wave size
+#endif
 
 /** Number of lane the GROUP_TILE_SIZE. */
 #define LANE_COUNT_X (GROUP_TILE_SIZE / LANE_STRIDE_X)
@@ -114,13 +129,13 @@ TLaneVector<ScalarType, VectorSize, SimdSize> AccessNeighborTexel(TLaneVector<Sc
 		if (LaneRotation != 0)
 		{
 			// Only uses wave intrinsics on platforms that provides RDNA's ds_swizzle because we know the wave size of this platforms. 
-			#if DIM_USE_WAVE_OPS
+			#if DIM_WAVE_SIZE != 0
 				#if !PLATFORM_SUPPORTS_WAVE_ROTATE
 					//#error InitWaveRotateLaneGroup() isn't supported.
 				#endif
-				const FWaveBroadcastSettings BroadcastSettings = InitWaveRotateLaneGroup(/* LaneGroupSize = */ WAVE_SIZE, LaneRotation);
+				const FWaveBroadcastSettings BroadcastSettings = InitWaveRotateLaneGroup(/* LaneGroupSize = */ LANE_COUNT, LaneRotation);
 			#else
-				const FWaveBroadcastSettings BroadcastSettings = ConvertWaveBroadcastToLDS(InitWaveRotateLaneGroup(/* LaneGroupSize = */ WAVE_SIZE, LaneRotation), GroupThreadIndex);
+				const FWaveBroadcastSettings BroadcastSettings = ConvertWaveBroadcastToLDS(InitWaveRotateLaneGroup(/* LaneGroupSize = */ LANE_COUNT, LaneRotation), GroupThreadIndex);
 			#endif
 			ReturnSimdElement = WaveBroadcast(BroadcastSettings, SimdElement);
 		}
@@ -307,7 +322,7 @@ tsr_simd_halfC ClampFireFliersWithGuide(tsr_simd_halfC ToClamp, tsr_simd_halfC G
 
 //------------------------------------------------------- ENTRY POINT
 
-[numthreads(WAVE_SIZE, 1, 1)]
+[numthreads(LANE_COUNT, 1, 1)]
 void MainCS(
 	uint2 GroupId : SV_GroupID,
 	uint GroupThreadIndex : SV_GroupIndex)
diff --git a/Engine/Source/Runtime/RHI/Private/RHI.cpp b/Engine/Source/Runtime/RHI/Private/RHI.cpp
--- a/Engine/Source/Runtime/RHI/Private/RHI.cpp
+++ b/Engine/Source/Runtime/RHI/Private/RHI.cpp
@@ -2508,6 +2508,8 @@ void FGenericDataDrivenShaderPlatformInfo::ParseDataDrivenShaderInfo(const FConf
 	GET_SECTION_BOOL_HELPER(bSupportsIntrinsicWaveOnce);
 	GET_SECTION_BOOL_HELPER(bSupportsConservativeRasterization);
 	GET_SECTION_SUPPORT_HELPER(bSupportsWaveOperations);
+	GET_SECTION_INT_HELPER(MinimumWaveSize);
+	GET_SECTION_INT_HELPER(MaximumWaveSize);
 	GET_SECTION_BOOL_HELPER(bRequiresExplicit128bitRT);
 	GET_SECTION_BOOL_HELPER(bSupportsGen5TemporalAA);
 	GET_SECTION_BOOL_HELPER(bTargetsTiledGPU);
diff --git a/Engine/Source/Runtime/RHI/Public/RHIDefinitions.h b/Engine/Source/Runtime/RHI/Public/RHIDefinitions.h
--- a/Engine/Source/Runtime/RHI/Public/RHIDefinitions.h
+++ b/Engine/Source/Runtime/RHI/Public/RHIDefinitions.h
@@ -376,6 +376,8 @@ class RHI_API FGenericDataDrivenShaderPlatformInfo
 	uint32 bSupportsTemporalHistoryUpscale : 1;
 	uint32 bSupportsRTIndexFromVS : 1;
 	uint32 bSupportsWaveOperations : int32(ERHIFeatureSupport::NumBits);
+	uint32 MinimumWaveSize : 8;
+	uint32 MaximumWaveSize : 8;
 	uint32 bSupportsIntrinsicWaveOnce : 1;
 	uint32 bSupportsConservativeRasterization : 1;
 	uint32 bRequiresExplicit128bitRT : 1;
@@ -730,6 +732,18 @@ public:
 		return ERHIFeatureSupport(Infos[Platform].bSupportsWaveOperations);
 	}
 
+	static FORCEINLINE_DEBUGGABLE const uint32 GetMinimumWaveSize(const FStaticShaderPlatform Platform)
+	{
+		check(IsValid(Platform));
+		return Infos[Platform].MinimumWaveSize;
+	}
+
+	static FORCEINLINE_DEBUGGABLE const uint32 GetMaximumWaveSize(const FStaticShaderPlatform Platform)
+	{
+		check(IsValid(Platform));
+		return Infos[Platform].MaximumWaveSize;
+	}
+
 	static FORCEINLINE_DEBUGGABLE const bool GetSupportsTemporalHistoryUpscale(const FStaticShaderPlatform Platform)
 	{
 		check(IsValid(Platform));
diff --git a/Engine/Source/Runtime/Renderer/Private/PostProcess/TemporalSuperResolution.cpp b/Engine/Source/Runtime/Renderer/Private/PostProcess/TemporalSuperResolution.cpp
--- a/Engine/Source/Runtime/Renderer/Private/PostProcess/TemporalSuperResolution.cpp
+++ b/Engine/Source/Runtime/Renderer/Private/PostProcess/TemporalSuperResolution.cpp
@@ -348,10 +348,10 @@ class FTSRRejectShadingCS : public FTSRShader
 	DECLARE_GLOBAL_SHADER(FTSRRejectShadingCS);
 	SHADER_USE_PARAMETER_STRUCT(FTSRRejectShadingCS, FTSRShader);
 
-	class FUseWaveOps : SHADER_PERMUTATION_BOOL("DIM_USE_WAVE_OPS");
+	class FWaveSizeOps : SHADER_PERMUTATION_SPARSE_INT("DIM_WAVE_SIZE", 0, 32, 64);
 	class FSeparateTranslucencyDim : SHADER_PERMUTATION_BOOL("DIM_SEPARATE_TRANSLUCENCY");
 
-	using FPermutationDomain = TShaderPermutationDomain<FUseWaveOps, FSeparateTranslucencyDim>;
+	using FPermutationDomain = TShaderPermutationDomain<FWaveSizeOps, FSeparateTranslucencyDim>;
 
 	BEGIN_SHADER_PARAMETER_STRUCT(FParameters, )
 		SHADER_PARAMETER_STRUCT_INCLUDE(FTSRCommonParameters, CommonParameters)
@@ -374,18 +374,27 @@ class FTSRRejectShadingCS : public FTSRShader
 	static bool ShouldCompilePermutation(const FGlobalShaderPermutationParameters& Parameters)
 	{
 		FPermutationDomain PermutationVector(Parameters.PermutationId);
+		int32 WaveSize = PermutationVector.Get<FWaveSizeOps>();
 
 		ERHIFeatureSupport WaveOpsSupport = FDataDrivenShaderPlatformInfo::GetSupportsWaveOperations(Parameters.Platform);
-		if (PermutationVector.Get<FUseWaveOps>())
+		if (WaveSize != 0)
 		{
 			if (WaveOpsSupport == ERHIFeatureSupport::Unsupported)
 			{
 				return false;
 			}
+
+			if (WaveSize < int32(FDataDrivenShaderPlatformInfo::GetMinimumWaveSize(Parameters.Platform)) ||
+				WaveSize > int32(FDataDrivenShaderPlatformInfo::GetMaximumWaveSize(Parameters.Platform)))
+			{
+				return false;
+			}
 		}
-		else
+		else // if (WaveSize == LDS fallback)
 		{
-			if (WaveOpsSupport == ERHIFeatureSupport::RuntimeGuaranteed)
+			if (WaveOpsSupport == ERHIFeatureSupport::RuntimeGuaranteed &&
+				FDataDrivenShaderPlatformInfo::GetMinimumWaveSize(Parameters.Platform) >= 32 &&
+				FDataDrivenShaderPlatformInfo::GetMaximumWaveSize(Parameters.Platform) <= 64)
 			{
 				return false;
 			}
@@ -399,10 +408,13 @@ class FTSRRejectShadingCS : public FTSRShader
 		FPermutationDomain PermutationVector(Parameters.PermutationId);
 
 		FTSRShader::ModifyCompilationEnvironment(Parameters, OutEnvironment);
-		OutEnvironment.CompilerFlags.Add(CFLAG_Wave32);
 
-		if (PermutationVector.Get<FUseWaveOps>())
+		if (PermutationVector.Get<FWaveSizeOps>() != 0)
 		{
+			if (PermutationVector.Get<FWaveSizeOps>() == 32)
+			{
+				OutEnvironment.CompilerFlags.Add(CFLAG_Wave32);
+			}
 			OutEnvironment.CompilerFlags.Add(CFLAG_WaveOperations);
 		}
 	}
@@ -1256,14 +1268,14 @@ ITemporalUpscaler::FOutputs AddTemporalSuperResolutionPasses(
 		PassParameters->DebugOutput = CreateDebugUAV(InputExtent, TEXT("Debug.TSR.RejectShading"));
 
 		FTSRRejectShadingCS::FPermutationDomain PermutationVector;
-		PermutationVector.Set<FTSRRejectShadingCS::FUseWaveOps>(bUseWaveOps&& GRHIMinimumWaveSize <= 32 && GRHIMaximumWaveSize >= 32);
+		PermutationVector.Set<FTSRRejectShadingCS::FWaveSizeOps>(bUseWaveOps && GRHIMinimumWaveSize >= 32 && GRHIMinimumWaveSize <= 64 ? GRHIMinimumWaveSize : 0);
 		PermutationVector.Set<FTSRRejectShadingCS::FSeparateTranslucencyDim>(bAccumulateTranslucencySeparately);
 
 		TShaderMapRef<FTSRRejectShadingCS> ComputeShader(View.ShaderMap, PermutationVector);
 		FComputeShaderUtils::AddPass(
 			GraphBuilder,
-			RDG_EVENT_NAME("TSR RejectShading(%s%s) %dx%d",
-				PermutationVector.Get<FTSRRejectShadingCS::FUseWaveOps>() ? TEXT("WaveOps") : TEXT(""),
+			RDG_EVENT_NAME("TSR RejectShading(WaveSize=%d%s) %dx%d",
+				PermutationVector.Get<FTSRRejectShadingCS::FWaveSizeOps>(),
 				PermutationVector.Get<FTSRRejectShadingCS::FSeparateTranslucencyDim>() ? TEXT(" SeparateTranslucency") : TEXT(""),
 				InputRect.Width(), InputRect.Height()),
 			AsyncComputePasses >= 3 ? ERDGPassFlags::AsyncCompute : ERDGPassFlags::Compute,
