rely on shader platform instead of rhi feature level when generating shaders: this can be an issue of one starts to cook a shader platform with a feature level not supported by the host

diff --git a/Engine/Source/Runtime/Renderer/Private/DebugViewModeRendering.cpp b/Engine/Source/Runtime/Renderer/Private/DebugViewModeRendering.cpp
--- a/Engine/Source/Runtime/Renderer/Private/DebugViewModeRendering.cpp
+++ b/Engine/Source/Runtime/Renderer/Private/DebugViewModeRendering.cpp
@@ -82,15 +82,15 @@ static FAutoConsoleVariableRef CVarShaderComplexityBaselineDeferredUnlitPS(
 
 IMPLEMENT_MATERIAL_SHADER_TYPE(, FDebugViewModePS, TEXT("/Engine/Private/DebugViewModePixelShader.usf"), TEXT("Main"), SF_Pixel);
 
-int32 GetQuadOverdrawUAVIndex(EShaderPlatform Platform, ERHIFeatureLevel::Type FeatureLevel)
+int32 GetQuadOverdrawUAVIndex(EShaderPlatform Platform)
 {
 	if (IsForwardShadingEnabled(Platform))
 	{
-		return FVelocityRendering::BasePassCanOutputVelocity(FeatureLevel) ? 2 : 1;
+		return FVelocityRendering::BasePassCanOutputVelocity(Platform) ? 2 : 1;
 	}
 	else // GBuffer
 	{
-		return FVelocityRendering::BasePassCanOutputVelocity(FeatureLevel) ? 7 : 6;
+		return FVelocityRendering::BasePassCanOutputVelocity(Platform) ? 7 : 6;
 	}
 }
 
diff --git a/Engine/Source/Runtime/Renderer/Private/DebugViewModeRendering.h b/Engine/Source/Runtime/Renderer/Private/DebugViewModeRendering.h
--- a/Engine/Source/Runtime/Renderer/Private/DebugViewModeRendering.h
+++ b/Engine/Source/Runtime/Renderer/Private/DebugViewModeRendering.h
@@ -44,7 +44,7 @@ void SetupDebugViewModePassUniformBufferConstants(const FViewInfo& ViewInfo, FDe
 TRDGUniformBufferRef<FDebugViewModePassUniformParameters> CreateDebugViewModePassUniformBuffer(FRDGBuilder& GraphBuilder, const FViewInfo& View, FRDGTextureRef QuadOverdrawTexture);
 
 /** Returns the RT index where the QuadOverdrawUAV will be bound. */
-extern int32 GetQuadOverdrawUAVIndex(EShaderPlatform Platform, ERHIFeatureLevel::Type FeatureLevel);
+extern int32 GetQuadOverdrawUAVIndex(EShaderPlatform Platform);
 
 class FDebugViewModeShaderElementData : public FMeshMaterialShaderElementData
 {
@@ -173,7 +173,7 @@ public:
 		OutEnvironment.SetDefine(TEXT("MAX_NUM_TEXTURE_REGISTER"), (uint32)TEXSTREAM_MAX_NUM_TEXTURES_PER_MATERIAL);
 		OutEnvironment.SetDefine(TEXT("SCENE_TEXTURES_DISABLED"), 1u);
 		TCHAR BufferRegister[] = { 'u', '0', 0 };
-		BufferRegister[1] += GetQuadOverdrawUAVIndex(Parameters.Platform, Parameters.MaterialParameters.FeatureLevel);
+		BufferRegister[1] += GetQuadOverdrawUAVIndex(Parameters.Platform);
 		OutEnvironment.SetDefine(TEXT("QUAD_BUFFER_REGISTER"), BufferRegister);
 		const bool bUsingMobileRenderer = FSceneInterface::GetShadingPath(GetMaxSupportedFeatureLevel(Parameters.Platform)) == EShadingPath::Mobile;
 		OutEnvironment.SetDefine(TEXT("OUTPUT_QUAD_OVERDRAW"), !bUsingMobileRenderer);
