MGPU:  Bloom FFT kernel procedurally generated resource wasn't being initialized across all GPUs, causing a black screen when switching GPU assignments for a view.

#rnx

diff --git a/Engine/Source/Runtime/Renderer/Private/PostProcess/PostProcessFFTBloom.cpp b/Engine/Source/Runtime/Renderer/Private/PostProcess/PostProcessFFTBloom.cpp
--- a/Engine/Source/Runtime/Renderer/Private/PostProcess/PostProcessFFTBloom.cpp
+++ b/Engine/Source/Runtime/Renderer/Private/PostProcess/PostProcessFFTBloom.cpp
@@ -334,9 +334,9 @@ void InitDomainAndGetKernel(
 		}
 	}
 
-	// Re-transform the kernel if needed.
+	// Re-transform the kernel if needed -- needs to run on all GPUs in case view's GPU assignment changes later.
 	RDG_EVENT_SCOPE(GraphBuilder, "InitBloomKernel");
-
+	RDG_GPU_MASK_SCOPE(GraphBuilder, FRHIGPUMask::All());
 
 	FRDGTextureRef SpatialKernelTexture = RegisterExternalTexture(GraphBuilder, PhysicalSpaceKernelTextureRef, TEXT("Bloom.FFT.OriginalKernel"));
 
