Fixes PVS warning in TSR

diff --git a/Engine/Source/Runtime/Renderer/Private/PostProcess/TemporalSuperResolution.cpp b/Engine/Source/Runtime/Renderer/Private/PostProcess/TemporalSuperResolution.cpp
--- a/Engine/Source/Runtime/Renderer/Private/PostProcess/TemporalSuperResolution.cpp
+++ b/Engine/Source/Runtime/Renderer/Private/PostProcess/TemporalSuperResolution.cpp
@@ -1275,7 +1275,7 @@ ITemporalUpscaler::FOutputs AddTemporalSuperResolutionPasses(
 		FComputeShaderUtils::AddPass(
 			GraphBuilder,
 			RDG_EVENT_NAME("TSR RejectShading(WaveSize=%d%s) %dx%d",
-				PermutationVector.Get<FTSRRejectShadingCS::FWaveSizeOps>(),
+				int32(PermutationVector.Get<FTSRRejectShadingCS::FWaveSizeOps>()),
 				PermutationVector.Get<FTSRRejectShadingCS::FSeparateTranslucencyDim>() ? TEXT(" SeparateTranslucency") : TEXT(""),
 				InputRect.Width(), InputRect.Height()),
 			AsyncComputePasses >= 3 ? ERDGPassFlags::AsyncCompute : ERDGPassFlags::Compute,
