Fix crash when OIT triangle sorting debug is enabled, and the platform does not support it.

diff --git a/Engine/Source/Runtime/Renderer/Private/OIT/OIT.cpp b/Engine/Source/Runtime/Renderer/Private/OIT/OIT.cpp
--- a/Engine/Source/Runtime/Renderer/Private/OIT/OIT.cpp
+++ b/Engine/Source/Runtime/Renderer/Private/OIT/OIT.cpp
@@ -132,7 +132,8 @@ static void AddOITPixelDebugPass(
 	if (!OITData.SampleColorTexture ||
 		!OITData.SampleTransTexture ||
 		!OITData.SampleDepthTexture ||
-		!OITData.SampleCountTexture)
+		!OITData.SampleCountTexture || 
+		!ShaderPrint::IsSupported(View.GetShaderPlatform()))
 		return;
 
 	// Force ShaderPrint on.
@@ -643,7 +644,8 @@ static void AddOITSortTriangleIndexPass(
 	FRDGBufferUAVRef SliceCounterUAV = GraphBuilder.CreateUAV(SliceCounterBuffer, PF_R32_UINT);
 	AddClearUAVPass(GraphBuilder, SliceCounterUAV, 0u);
 
-	const bool bDebugEnable = DebugData.IsValid();
+	const EShaderPlatform Platform = View.Family->GetShaderPlatform();
+	const bool bDebugEnable = DebugData.IsValid() && ShaderPrint::IsSupported(Platform) && ShaderPrint::IsValid(View.ShaderPrintData);
 	FRHIBuffer* SortedIndexBufferRHI = Allocation.SortedIndexBuffer->IndexBufferRHI;
 
 	// 1. Scan the primitive and assign each primitive to a slice
@@ -674,15 +676,16 @@ static void AddOITSortTriangleIndexPass(
 		// Debug
 		if (bDebugEnable)
 		{
+			ShaderPrint::SetEnabled(true);
+			ShaderPrint::RequestSpaceForCharacters(DebugData.VisibleInstances * 256 + 512);
+			ShaderPrint::RequestSpaceForLines(DebugData.VisiblePrimitives * 8 + DebugData.VisibleInstances * 32);
+
 			Parameters->ViewToWorld		= FMatrix44f(View.ViewMatrices.GetViewMatrix().Inverse());	 // LWC_TODO: Precision loss?
 			Parameters->WorldBound_Min	= (FVector3f)Bounds.GetBox().Min;
 			Parameters->WorldBound_Max	= (FVector3f)Bounds.GetBox().Max;
 			Parameters->ViewBound_Min	= (FVector3f)ViewBounds.GetBox().Min;
 			Parameters->ViewBound_Max	= (FVector3f)ViewBounds.GetBox().Max;
-			if (ShaderPrint::IsValid(View.ShaderPrintData))
-			{
-				ShaderPrint::SetParameters(GraphBuilder, View.ShaderPrintData, Parameters->ShaderPrintParameters);
-			}
+			ShaderPrint::SetParameters(GraphBuilder, View.ShaderPrintData, Parameters->ShaderPrintParameters);
 
 			check(DebugData.Buffer);
 
@@ -769,14 +772,16 @@ static void AddOITTriangleDebugPass(
 	const FViewInfo& View,
 	const FOITDebugData& DebugData)
 {
-	if (!DebugData.IsValid())
+	const EShaderPlatform Platform = View.Family->GetShaderPlatform();
+	if (!DebugData.IsValid() || !ShaderPrint::IsSupported(Platform))
 	{
 		return;
 	}
 
 	// Force ShaderPrint on.
 	ShaderPrint::SetEnabled(true);
-	ShaderPrint::RequestSpaceForCharacters(512);
+	ShaderPrint::RequestSpaceForCharacters(DebugData.VisibleInstances * 256 + 512);
+	ShaderPrint::RequestSpaceForLines(DebugData.VisiblePrimitives * 8 + DebugData.VisibleInstances * 32);
 
 	FOITSortTriangleIndex_Debug::FParameters* Parameters = GraphBuilder.AllocParameters<FOITSortTriangleIndex_Debug::FParameters>();
 	Parameters->VisibleInstances = DebugData.VisibleInstances;
