Made some volumetric fog cvar safe to crash.

diff --git a/Engine/Source/Runtime/Renderer/Private/VolumetricFog.cpp b/Engine/Source/Runtime/Renderer/Private/VolumetricFog.cpp
--- a/Engine/Source/Runtime/Renderer/Private/VolumetricFog.cpp
+++ b/Engine/Source/Runtime/Renderer/Private/VolumetricFog.cpp
@@ -138,6 +138,16 @@ FAutoConsoleVariableRef CVarLightScatteringSampleJitterMultiplier(
 	ECVF_RenderThreadSafe | ECVF_Scalability
 );
 
+static int32 GetVolumetricFogGridPixelSize()
+{
+	return FMath::Max(1, GVolumetricFogGridPixelSize);
+}
+
+static int32 GetVolumetricFogGridSizeZ()
+{
+	return FMath::Max(1, GVolumetricFogGridSizeZ);
+}
+
 IMPLEMENT_GLOBAL_SHADER_PARAMETER_STRUCT(FVolumetricFogGlobalData, "VolumetricFog");
 
 DECLARE_GPU_STAT(VolumetricFog);
@@ -815,7 +825,7 @@ FIntVector GetVolumetricFogGridSize(FIntPoint ViewRectSize, int32& OutVolumetric
 {
 	extern int32 GLightGridSizeZ;
 	FIntPoint VolumetricFogGridSizeXY;
-	int32 VolumetricFogGridPixelSize = GVolumetricFogGridPixelSize;
+	int32 VolumetricFogGridPixelSize = GetVolumetricFogGridPixelSize();
 	VolumetricFogGridSizeXY = FIntPoint::DivideAndRoundUp(ViewRectSize, VolumetricFogGridPixelSize);
 	if(VolumetricFogGridSizeXY.X > GMaxVolumeTextureDimensions || VolumetricFogGridSizeXY.Y > GMaxVolumeTextureDimensions) //clamp to max volume texture dimensions. only happens for extreme resolutions (~8x2k)
 	{
@@ -825,7 +835,7 @@ FIntVector GetVolumetricFogGridSize(FIntPoint ViewRectSize, int32& OutVolumetric
 		VolumetricFogGridSizeXY = FIntPoint::DivideAndRoundUp(ViewRectSize, VolumetricFogGridPixelSize);
 	}
 	OutVolumetricFogGridPixelSize = VolumetricFogGridPixelSize;
-	return FIntVector(VolumetricFogGridSizeXY.X, VolumetricFogGridSizeXY.Y, GVolumetricFogGridSizeZ);
+	return FIntVector(VolumetricFogGridSizeXY.X, VolumetricFogGridSizeXY.Y, GetVolumetricFogGridSizeZ());
 }
 
 void SetupVolumetricFogGlobalData(const FViewInfo& View, FVolumetricFogGlobalData& Parameters)
@@ -982,7 +992,7 @@ void FDeferredShadingSceneRenderer::ComputeVolumetricFog(FRDGBuilder& GraphBuild
 			FIntPoint ConservativeDepthTextureSize = FIntPoint(VolumetricFogGridSize.X, VolumetricFogGridSize.Y);
 			ConservativeDepthTexture = GraphBuilder.CreateTexture(FRDGTextureDesc::Create2D(ConservativeDepthTextureSize, PF_R16F,
 				FClearValueBinding::Black, TexCreate_RenderTargetable | TexCreate_ShaderResource | TexCreate_UAV), TEXT("VolumetricFog.ConservativeDepthTexture"));
-			AddGenerateConservativeDepthBufferPass(View, GraphBuilder, ConservativeDepthTexture, GVolumetricFogGridPixelSize);
+			AddGenerateConservativeDepthBufferPass(View, GraphBuilder, ConservativeDepthTexture, GetVolumetricFogGridPixelSize());
 		}
 		else
 		{
