Remove BasePassCanOutputVelocity() that takes a feature level.
Use only the shader platform version.
#rnx

diff --git a/Engine/Source/Runtime/Renderer/Private/DeferredShadingRenderer.cpp b/Engine/Source/Runtime/Renderer/Private/DeferredShadingRenderer.cpp
--- a/Engine/Source/Runtime/Renderer/Private/DeferredShadingRenderer.cpp
+++ b/Engine/Source/Runtime/Renderer/Private/DeferredShadingRenderer.cpp
@@ -2490,8 +2490,9 @@ void FDeferredShadingSceneRenderer::Render(FRDGBuilder& GraphBuilder)
 	PrepareDistanceFieldScene(GraphBuilder, false);
 
 	const bool bShouldRenderVelocities = ShouldRenderVelocities();
-	const bool bBasePassCanOutputVelocity = FVelocityRendering::BasePassCanOutputVelocity(FeatureLevel);
-	const bool bHairStrandsEnable = HairStrandsBookmarkParameters.HasInstances() && Views.Num() > 0 && IsHairStrandsEnabled(EHairStrandsShaderType::Strands, GetViewFamilyInfo(Views).GetShaderPlatform());
+	const EShaderPlatform Platform = GetViewFamilyInfo(Views).GetShaderPlatform();
+	const bool bBasePassCanOutputVelocity = FVelocityRendering::BasePassCanOutputVelocity(Platform);
+	const bool bHairStrandsEnable = HairStrandsBookmarkParameters.HasInstances() && Views.Num() > 0 && IsHairStrandsEnabled(EHairStrandsShaderType::Strands, Platform);
 
 	{
 		GraphBuilder.SetCommandListStat(GET_STATID(STAT_CLM_PrePass));
diff --git a/Engine/Source/Runtime/Renderer/Private/VelocityRendering.cpp b/Engine/Source/Runtime/Renderer/Private/VelocityRendering.cpp
--- a/Engine/Source/Runtime/Renderer/Private/VelocityRendering.cpp
+++ b/Engine/Source/Runtime/Renderer/Private/VelocityRendering.cpp
@@ -407,12 +407,6 @@ bool FVelocityRendering::BasePassCanOutputVelocity(EShaderPlatform ShaderPlatfor
 	return IsUsingBasePassVelocity(ShaderPlatform);
 }
 
-bool FVelocityRendering::BasePassCanOutputVelocity(ERHIFeatureLevel::Type FeatureLevel)
-{
-	EShaderPlatform ShaderPlatform = GetFeatureLevelShaderPlatform(FeatureLevel);
-	return BasePassCanOutputVelocity(ShaderPlatform);
-}
-
 bool FVelocityRendering::IsParallelVelocity(EShaderPlatform ShaderPlatform)
 {
 	return GRHICommandList.UseParallelAlgorithms() && CVarParallelVelocity.GetValueOnRenderThread()
diff --git a/Engine/Source/Runtime/Renderer/Private/VelocityRendering.h b/Engine/Source/Runtime/Renderer/Private/VelocityRendering.h
--- a/Engine/Source/Runtime/Renderer/Private/VelocityRendering.h
+++ b/Engine/Source/Runtime/Renderer/Private/VelocityRendering.h
@@ -47,9 +47,6 @@ struct FVelocityRendering
 	/** Returns true if the velocity can be output in the BasePass. */
 	static bool BasePassCanOutputVelocity(EShaderPlatform ShaderPlatform);
 
-	/** Returns true if the velocity can be output in the BasePass. Only valid for the current platform. */
-	static bool BasePassCanOutputVelocity(ERHIFeatureLevel::Type FeatureLevel);
-
 	/** Returns true if the velocity pass is using parallel dispatch. */
 	static bool IsParallelVelocity(EShaderPlatform ShaderPlatform);
 
