HeterogeneousVolumes: Transmittance Volume single-scattering optimization.

The Transmittance Volume optimizes single-scattering integration by first computing light transmission as a pre-pass. A transmittance volume is created with dimensions proportional to that of the pre-shaded volume resolution. For each transmittance volume voxel, a ray is traced toward the light source. Transmittance along the ray is recorded and stored in the originating voxel. The single-scattering pass fetches the value in the Transmittance Volume rather than invoking a second ray-marching loop within the original integrator. The Transmittance Volume can be built at a lower resolution, lowering quality but significantly improving run-time performance.

Transmittance Volume CVars:
r.HeterogeneousVolume.TransmittanceVolume (default = 1)
r.HeterogeneousVolume.TransmittanceVolume.DownsampleFactor (default = 1)

The relative speedup from the Transmittance Volume is dependant on the sparsity of the volume within its bounding domain, but is typically a 1.25-1.5X speedup for live-shading, with speedups approaching 3X or more when invoking downsampling of the Transmittance Volume. Sample times associated with rendering the "Grid 3D Gas Explosion Cine" Niagara template with a Directional Light at 1080p:

r.HeterogeneousVolume.TransmittanceVolume=0
    2.38ms   GasExplosionCine_HVolumes 1 dispatch

r.HeterogeneousVolume.TransmittanceVolume=1
    1.86ms   GasExplosionCine_HVolumes 2 dispatches

r.HeterogeneousVolume.TransmittanceVolume=1
r.HeterogeneousVolume.TransmittanceVolume.DownsampleFactor=2
    0.80ms   GasExplosionCine_HVolumes 2 dispatches

diff --git a/Engine/Shaders/Private/HeterogeneousVolumes/HeterogeneousVolumesHardwareRayTracing.usf b/Engine/Shaders/Private/HeterogeneousVolumes/HeterogeneousVolumesHardwareRayTracing.usf
--- a/Engine/Shaders/Private/HeterogeneousVolumes/HeterogeneousVolumesHardwareRayTracing.usf
+++ b/Engine/Shaders/Private/HeterogeneousVolumes/HeterogeneousVolumesHardwareRayTracing.usf
@@ -140,7 +140,7 @@ RAY_TRACING_ENTRY_ANY_HIT(SparseVoxelsAnyHitShader,
 #else
 			float3 WorldPosition = mul(float4(ObjectPosition, 1.0), GetLocalToWorld()).xyz;
 
-			FVolumeSampleContext SampleContext = CreateVolumeSampleContext(ObjectPosition, WorldPosition);
+			FVolumeSampleContext SampleContext = CreateVolumeSampleContext(ObjectPosition, WorldPosition, MipLevel);
 			float3 Extinction = SampleExtinction(SampleContext);
 #endif
 			Transmittance *= exp(-Extinction * StepSize);
@@ -185,7 +185,6 @@ RAY_TRACING_ENTRY_MISS(SparseVoxelsMissShader,
 RaytracingAccelerationStructure TLAS;
 
 // Ray data
-float MaxTraceDistance;
 int MaxStepCount;
 int bJitter;
 
@@ -320,7 +319,7 @@ RAY_TRACING_ENTRY_RAYGEN(SparseVoxelsRayGenerationShader1)
 #endif // HAS_INVERTED_Z_BUFFER
 
 	// Clip trace distance
-	float SceneDepth = min(ConvertFromDeviceZ(DeviceZ), MaxTraceDistance);
+	float SceneDepth = min(ConvertFromDeviceZ(DeviceZ), GetMaxTraceDistance());
 	DeviceZ = ConvertToDeviceZ(SceneDepth);
 
 	float Jitter = bJitter ? InterleavedGradientNoise(PixelCoord, View.StateFrameIndexMod8) : 0.0;
@@ -388,6 +387,97 @@ float3 ComputeTransmittanceHardwareRayTracing(
 #define HARDWARE_RAY_TRACING 1
 #include "HeterogeneousVolumesRayMarchingUtils.ush"
 
+uint3 DispatchRaysToVoxelIndex(uint2 DispatchRaysIndex)
+{
+	uint3 VolumeResolution = GetTransmittanceVolumeResolution();
+	uint3 VoxelIndex = uint3(DispatchRaysIndex.x,
+		DispatchRaysIndex.y % TransmittanceVolumeResolution.y,
+		DispatchRaysIndex.y / TransmittanceVolumeResolution.y);
+	return VoxelIndex;
+}
+
+RWTexture3D<float3> RWTransmittanceVolumeTexture;
+
+RAY_TRACING_ENTRY_RAYGEN(RenderTransmittanceVolumeWithPreshadingRGS)
+{
+	uint3 VoxelIndex = DispatchRaysToVoxelIndex(DispatchRaysIndex().xy);
+	if (any(VoxelIndex >= TransmittanceVolumeResolution))
+	{
+		return;
+	}
+
+	float3 UVW = (VoxelIndex + 0.5) / float3(TransmittanceVolumeResolution);
+
+	float3 LocalBoundsMin = GetLocalBoundsOrigin() - GetLocalBoundsExtent();
+	float3 LocalBoundsMax = GetLocalBoundsOrigin() + GetLocalBoundsExtent();
+	float3 LocalRayOrigin = UVW * GetLocalBoundsExtent() * 2.0 + LocalBoundsMin;
+	float3 WorldRayOrigin = mul(float4(LocalRayOrigin, 1.0), GetLocalToWorld()).xyz;
+
+	// Existence culling
+	float MipLevel = 0.0f;
+	FVolumeSampleContext SampleContext = CreateVolumeSampleContext(LocalRayOrigin, WorldRayOrigin, MipLevel);
+	float3 Extinction = SampleExtinction(SampleContext);
+	float Epsilon = 1.0e-7;
+	if (all(Extinction < Epsilon))
+	{
+		RWTransmittanceVolumeTexture[VoxelIndex] = 0.0f;
+		return;
+	}
+
+	FDeferredLightData LightData = LoadLightData(LightType);
+	float3 L = LightData.Direction;
+	float3 ToLight = L * GetMaxShadowTraceDistance();
+	if (LightType != LIGHT_TYPE_DIRECTIONAL)
+	{
+		if (LightData.bRectLight)
+		{
+			FRect Rect = GetRect(ToLight, LightData);
+		}
+		else
+		{
+			FCapsuleLight Capsule = GetCapsule(ToLight, LightData);
+		}
+	}
+	float3 WorldRayEnd = WorldRayOrigin + ToLight;
+	float3 LocalRayEnd = mul(float4(WorldRayEnd, 1.0), GetWorldToLocal()).xyz;
+
+	float3 WorldRayDirection = normalize(WorldRayEnd - WorldRayOrigin);
+	float3 LocalRayDirection = LocalRayEnd - LocalRayOrigin;
+
+	float LocalRayTMin = 0.0f;
+	float LocalRayTMax = length(LocalRayDirection);
+	LocalRayDirection /= LocalRayTMax;
+
+	float Jitter = bJitter ? InterleavedGradientNoise(VoxelIndex.xy, VoxelIndex.z) : 0.5;
+	float StepSize = GetStepSize() * GetShadowStepFactor();
+	int bApplyEmission = 0;
+	int bApplyDirectLighting = 0;
+	int bApplyShadowTransmittance = 0;
+	FRayMarchingContext RayMarchingContext = (FRayMarchingContext)0;
+	RayMarchingContext = CreateRayMarchingContext(
+		// Local-space
+		LocalRayOrigin,
+		LocalRayDirection,
+		LocalRayTMin,
+		LocalRayTMax,
+		// World-space
+		WorldRayOrigin,
+		WorldRayDirection,
+		// Ray-step attributes
+		Jitter,
+		StepSize,
+		MaxStepCount,
+		// Lighting
+		bApplyEmission,
+		bApplyDirectLighting,
+		bApplyShadowTransmittance
+	);
+	RayMarchingContext.MipLevel = 0;
+
+	float3 Transmittance = ComputeTransmittance(RayMarchingContext);
+	RWTransmittanceVolumeTexture[VoxelIndex] = Transmittance;
+}
+
 RAY_TRACING_ENTRY_RAYGEN(RenderSingleScatteringWithPreshadingRGS)
 {
 	float3 Radiance = 0.0;
@@ -409,7 +499,7 @@ RAY_TRACING_ENTRY_RAYGEN(RenderSingleScatteringWithPreshadingRGS)
 #endif // HAS_INVERTED_Z_BUFFER
 
 	// Clip trace distance
-	float SceneDepth = min(ConvertFromDeviceZ(DeviceZ), MaxTraceDistance);
+	float SceneDepth = min(ConvertFromDeviceZ(DeviceZ), GetMaxTraceDistance());
 	DeviceZ = ConvertToDeviceZ(SceneDepth);
 
 	float Jitter = bJitter ? InterleavedGradientNoise(PixelCoord, View.StateFrameIndexMod8) : 0.0;
diff --git a/Engine/Shaders/Private/HeterogeneousVolumes/HeterogeneousVolumesLiveShadingPipeline.usf b/Engine/Shaders/Private/HeterogeneousVolumes/HeterogeneousVolumesLiveShadingPipeline.usf
--- a/Engine/Shaders/Private/HeterogeneousVolumes/HeterogeneousVolumesLiveShadingPipeline.usf
+++ b/Engine/Shaders/Private/HeterogeneousVolumes/HeterogeneousVolumesLiveShadingPipeline.usf
@@ -10,6 +10,18 @@
 #include "HeterogeneousVolumesTracingUtils.ush"
 #include "HeterogeneousVolumesLightingUtils.ush"
 
+#ifndef THREADGROUP_SIZE_1D
+#define THREADGROUP_SIZE_1D 1
+#endif // THREADGROUP_SIZE_1D
+
+#ifndef THREADGROUP_SIZE_2D
+#define THREADGROUP_SIZE_2D 1
+#endif // THREADGROUP_SIZE_2D
+
+#ifndef THREADGROUP_SIZE_3D
+#define THREADGROUP_SIZE_3D 1
+#endif // THREADGROUP_SIZE_3D
+
 // Object data
 float4x4 LocalToWorld;
 float4x4 WorldToLocal;
@@ -26,6 +38,7 @@ float ShadowStepFactor;
 
 // Ray data
 float MaxTraceDistance;
+float MaxShadowTraceDistance;
 float StepSize;
 int MaxStepCount;
 int bJitter;
@@ -98,13 +111,23 @@ float GetShadowStepFactor()
 	return ShadowStepFactor;
 }
 
+float GetMaxTraceDistance()
+{
+	return MaxTraceDistance;
+}
+
+float GetMaxShadowTraceDistance()
+{
+	return MaxShadowTraceDistance;
+}
+
 struct FVolumeSampleContext
 {
 	FMaterialPixelParameters MaterialParameters;
 	FPixelMaterialInputs PixelMaterialInputs;
 };
 
-FVolumeSampleContext CreateVolumeSampleContext(float3 LocalPosition, float3 WorldPosition)
+FVolumeSampleContext CreateVolumeSampleContext(float3 LocalPosition, float3 WorldPosition, float MipLevel)
 {
 	FVolumeSampleContext VolumeSampleContext;
 	VolumeSampleContext.MaterialParameters = MakeInitializedMaterialPixelParameters();
@@ -137,6 +160,99 @@ float3 SampleAlbedo(inout FVolumeSampleContext Context)
 
 #include "HeterogeneousVolumesRayMarchingUtils.ush"
 
+RWTexture3D<float3> RWTransmittanceVolumeTexture;
+
+[numthreads(THREADGROUP_SIZE_3D, THREADGROUP_SIZE_3D, THREADGROUP_SIZE_3D)]
+void RenderTransmittanceVolumeWithLiveShadingCS(
+	uint3 GroupId : SV_GroupID,
+	uint3 GroupThreadId : SV_GroupThreadID,
+	uint3 DispatchThreadId : SV_DispatchThreadID
+)
+{
+	float3 VoxelIndex = DispatchThreadId;
+	if (any(DispatchThreadId > TransmittanceVolumeResolution))
+	{
+		return;
+	}
+
+	float3 VoxelJitter = 0.5;
+	if (bJitter)
+	{
+		uint3 Rand32Bits = Rand4DPCG32(uint4(VoxelIndex, View.StateFrameIndexMod8)).xyz;
+		VoxelJitter = float3(Rand32Bits) / float(uint(0xffffffff));
+	}
+	//float3 UVW = (VoxelIndex + VectorJitter) / float3(GetTransmittanceVolumeResolution());
+	float3 UVW = (VoxelIndex + 0.5) / float3(GetTransmittanceVolumeResolution());
+
+	float3 LocalBoundsMin = GetLocalBoundsOrigin() - GetLocalBoundsExtent();
+	float3 LocalBoundsMax = GetLocalBoundsOrigin() + GetLocalBoundsExtent();
+	float3 LocalRayOrigin = UVW * GetLocalBoundsExtent() * 2.0 + LocalBoundsMin;
+	float3 WorldRayOrigin = mul(float4(LocalRayOrigin, 1.0), GetLocalToWorld()).xyz;
+
+	// Existence culling
+	float MipLevel = 0.0f;
+	FVolumeSampleContext SampleContext = CreateVolumeSampleContext(LocalRayOrigin, WorldRayOrigin, MipLevel);
+	float3 Extinction = SampleExtinction(SampleContext);
+	float Epsilon = 1.0e-7;
+	if (all(Extinction < Epsilon))
+	{
+		RWTransmittanceVolumeTexture[VoxelIndex] = 0.0f;
+		return;
+	}
+
+	FDeferredLightData LightData = LoadLightData(LightType);
+	float3 L = LightData.Direction;
+	float3 ToLight = L * GetMaxShadowTraceDistance();
+	if (LightType != LIGHT_TYPE_DIRECTIONAL)
+	{
+		if (LightData.bRectLight)
+		{
+			FRect Rect = GetRect(ToLight, LightData);
+		}
+		else
+		{
+			FCapsuleLight Capsule = GetCapsule(ToLight, LightData);
+		}
+	}
+	float3 WorldRayEnd = WorldRayOrigin + ToLight;
+	float3 LocalRayEnd = mul(float4(WorldRayEnd, 1.0), GetWorldToLocal()).xyz;
+
+	float3 WorldRayDirection = normalize(WorldRayEnd - WorldRayOrigin);
+	float3 LocalRayDirection = LocalRayEnd - LocalRayOrigin;
+
+	float LocalRayTMin = 0.0f;
+	float LocalRayTMax = length(LocalRayDirection);
+	LocalRayDirection /= LocalRayTMax;
+
+	float RayJitter = length(VoxelJitter);
+	float StepSize = GetStepSize() * GetShadowStepFactor();
+	int bApplyEmission = 0;
+	int bApplyDirectLighting = 0;
+	int bApplyShadowTransmittance = 0;
+	FRayMarchingContext RayMarchingContext = (FRayMarchingContext)0;
+	RayMarchingContext = CreateRayMarchingContext(
+		// Local-space
+		LocalRayOrigin,
+		LocalRayDirection,
+		LocalRayTMin,
+		LocalRayTMax,
+		// World-space
+		WorldRayOrigin,
+		WorldRayDirection,
+		// Ray-step attributes
+		RayJitter,
+		StepSize,
+		MaxStepCount,
+		// Lighting
+		bApplyEmission,
+		bApplyDirectLighting,
+		bApplyShadowTransmittance
+	);
+
+	float3 Transmittance = ComputeTransmittance(RayMarchingContext);
+	RWTransmittanceVolumeTexture[VoxelIndex] = Transmittance;
+}
+
 [numthreads(THREADGROUP_SIZE_2D, THREADGROUP_SIZE_2D, 1)]
 void RenderSingleScatteringWithLiveShadingCS(
 	uint2 GroupThreadId : SV_GroupThreadID,
@@ -162,7 +278,7 @@ void RenderSingleScatteringWithLiveShadingCS(
 #endif // HAS_INVERTED_Z_BUFFER
 
 	// Clip trace distance
-	float SceneDepth = min(ConvertFromDeviceZ(DeviceZ), MaxTraceDistance);
+	float SceneDepth = min(ConvertFromDeviceZ(DeviceZ), GetMaxTraceDistance());
 	DeviceZ = ConvertToDeviceZ(SceneDepth);
 
 	float Jitter = bJitter ? InterleavedGradientNoise(PixelCoord, View.StateFrameIndexMod8) : 0.0;
diff --git a/Engine/Shaders/Private/HeterogeneousVolumes/HeterogeneousVolumesPreshadingPipeline.usf b/Engine/Shaders/Private/HeterogeneousVolumes/HeterogeneousVolumesPreshadingPipeline.usf
--- a/Engine/Shaders/Private/HeterogeneousVolumes/HeterogeneousVolumesPreshadingPipeline.usf
+++ b/Engine/Shaders/Private/HeterogeneousVolumes/HeterogeneousVolumesPreshadingPipeline.usf
@@ -41,6 +41,7 @@ int LightType;
 
 // Ray data
 float MaxTraceDistance;
+float MaxShadowTraceDistance;
 float StepSize;
 int MaxStepCount;
 int bJitter;
@@ -404,6 +405,101 @@ FTraceVoxelResult TraceVoxels(float3 LocalRayOrigin, float3 LocalRayDirection, f
 	return TraceVoxelResult;
 }
 
+RWTexture3D<float3> RWTransmittanceVolumeTexture;
+
+[numthreads(THREADGROUP_SIZE_3D, THREADGROUP_SIZE_3D, THREADGROUP_SIZE_3D)]
+void RenderTransmittanceVolumeWithPreshadingCS(
+	uint3 GroupId : SV_GroupID,
+	uint3 GroupThreadId : SV_GroupThreadID,
+	uint3 DispatchThreadId : SV_DispatchThreadID
+)
+{
+	float3 VoxelIndex = DispatchThreadId;
+	if (any(VoxelIndex >= TransmittanceVolumeResolution))
+	{
+		return;
+	}
+
+	float3 VoxelJitter = 0.5;
+	if (bJitter)
+	{
+		uint3 Rand32Bits = Rand4DPCG32(uint4(VoxelIndex, View.StateFrameIndexMod8)).xyz;
+		VoxelJitter = float3(Rand32Bits) / float(uint(0xffffffff));
+	}
+	//float3 UVW = (VoxelIndex + VectorJitter) / float3(GetTransmittanceVolumeResolution());
+	float3 UVW = (VoxelIndex + 0.5) / float3(GetTransmittanceVolumeResolution());
+
+	float3 LocalBoundsMin = GetLocalBoundsOrigin() - GetLocalBoundsExtent();
+	float3 LocalBoundsMax = GetLocalBoundsOrigin() + GetLocalBoundsExtent();
+	float3 LocalRayOrigin = UVW * GetLocalBoundsExtent() * 2.0 + LocalBoundsMin;
+	float3 WorldRayOrigin = mul(float4(LocalRayOrigin, 1.0), GetLocalToWorld()).xyz;
+
+	// Existence culling
+	float MipLevel = 0.0f;
+	FVolumeSampleContext SampleContext = CreateVolumeSampleContext(LocalRayOrigin, WorldRayOrigin, MipLevel);
+	float3 Extinction = SampleExtinction(SampleContext);
+	float Epsilon = 1.0e-7;
+	if (all(Extinction < Epsilon))
+	{
+		RWTransmittanceVolumeTexture[VoxelIndex] = 0.0f;
+		return;
+	}
+
+	// Build shadow ray
+	FDeferredLightData LightData = LoadLightData(LightType);
+	float3 L = LightData.Direction;
+	float3 ToLight = L * GetMaxShadowTraceDistance();
+	if (LightType != LIGHT_TYPE_DIRECTIONAL)
+	{
+		if (LightData.bRectLight)
+		{
+			FRect Rect = GetRect(ToLight, LightData);
+		}
+		else
+		{
+			FCapsuleLight Capsule = GetCapsule(ToLight, LightData);
+		}
+	}
+	float3 WorldRayEnd = WorldRayOrigin + ToLight;
+	float3 LocalRayEnd = mul(float4(WorldRayEnd, 1.0), GetWorldToLocal()).xyz;
+
+	float3 WorldRayDirection = normalize(WorldRayEnd - WorldRayOrigin);
+	float3 LocalRayDirection = LocalRayEnd - LocalRayOrigin;
+
+	float LocalRayTMin = 0.0f;
+	float LocalRayTMax = length(LocalRayDirection);
+	LocalRayDirection /= LocalRayTMax;
+
+	// Build ray-marching context
+	float RayJitter = length(VoxelJitter);
+	float StepSize = GetStepSize() * GetShadowStepFactor();
+	int bApplyEmission = 0;
+	int bApplyDirectLighting = 0;
+	int bApplyShadowTransmittance = 0;
+	FRayMarchingContext RayMarchingContext = (FRayMarchingContext) 0;
+	RayMarchingContext = CreateRayMarchingContext(
+		// Local-space
+		LocalRayOrigin,
+		LocalRayDirection,
+		LocalRayTMin,
+		LocalRayTMax,
+		// World-space
+		WorldRayOrigin,
+		WorldRayDirection,
+		// Ray-step attributes
+		RayJitter,
+		StepSize,
+		MaxStepCount,
+		// Lighting
+		bApplyEmission,
+		bApplyDirectLighting,
+		bApplyShadowTransmittance
+	);
+
+	float3 Transmittance = ComputeTransmittance(RayMarchingContext);
+	RWTransmittanceVolumeTexture[VoxelIndex] = Transmittance;
+}
+
 RWStructuredBuffer<FVoxelDataPacked> RWVoxelOutputBuffer;
 
 StructuredBuffer<FRayMarchingTile> RayMarchingTilesBuffer;
diff --git a/Engine/Shaders/Private/HeterogeneousVolumes/HeterogeneousVolumesRayMarchingUtils.ush b/Engine/Shaders/Private/HeterogeneousVolumes/HeterogeneousVolumesRayMarchingUtils.ush
--- a/Engine/Shaders/Private/HeterogeneousVolumes/HeterogeneousVolumesRayMarchingUtils.ush
+++ b/Engine/Shaders/Private/HeterogeneousVolumes/HeterogeneousVolumesRayMarchingUtils.ush
@@ -1,5 +1,7 @@
 // Copyright Epic Games, Inc. All Rights Reserved.
 
+#include "HeterogeneousVolumesTransmittanceVolumeUtils.ush"
+
 #ifndef HARDWARE_RAY_TRACING
 #define HARDWARE_RAY_TRACING 0
 #endif // HARDWARE_RAY_TRACING
@@ -11,13 +13,18 @@ struct FRayMarchingContext
 	float3 LocalRayDirection;
 	float LocalRayTMin;
 	float LocalRayTMax;
+
 	// World-space
 	float3 WorldRayOrigin;
 	float3 WorldRayDirection;
+
 	// Ray-step attributes
 	float Jitter;
 	float StepSize;
 	uint MaxStepCount;
+
+	float MipLevel;
+
 	// Lighting
 	int bApplyEmission;
 	int bApplyDirectLighting;
@@ -43,7 +50,8 @@ FRayMarchingContext CreateRayMarchingContext(
 	int bApplyShadowTransmittance
 )
 {
-	FRayMarchingContext Context;
+	FRayMarchingContext Context = (FRayMarchingContext)0;
+
 	// Local-space
 	Context.LocalRayOrigin = LocalRayOrigin;
 	Context.LocalRayDirection = LocalRayDirection;
@@ -61,6 +69,8 @@ FRayMarchingContext CreateRayMarchingContext(
 	Context.bApplyDirectLighting = bApplyDirectLighting;
 	Context.bApplyShadowTransmittance = bApplyShadowTransmittance;
 
+	Context.MipLevel = 0.0f;
+
 	return Context;
 }
 
@@ -90,7 +100,7 @@ void RayMarchTransmittance(
 		float3 LocalPosition = RayMarchingContext.LocalRayOrigin + RayMarchingContext.LocalRayDirection * (RayMarchingContext.LocalRayTMin + RayMarchingContext.StepSize * StepIndex);
 		float3 WorldPosition = RayMarchingContext.WorldRayOrigin + RayMarchingContext.WorldRayDirection * (RayMarchingContext.LocalRayTMin + RayMarchingContext.StepSize * StepIndex);
 
-		FVolumeSampleContext SampleContext = CreateVolumeSampleContext(LocalPosition, WorldPosition);
+		FVolumeSampleContext SampleContext = CreateVolumeSampleContext(LocalPosition, WorldPosition, RayMarchingContext.MipLevel);
 		float3 Extinction = SampleExtinction(SampleContext);
 		Transmittance *= exp(-Extinction * RayMarchingContext.StepSize);
 	}
@@ -137,7 +147,7 @@ void RayMarchEmissionAbsorption(
 		float3 LocalPosition = RayMarchingContext.LocalRayOrigin + RayMarchingContext.LocalRayDirection * (RayMarchingContext.LocalRayTMin + RayMarchingContext.StepSize * StepIndex);
 		float3 WorldPosition = RayMarchingContext.WorldRayOrigin + RayMarchingContext.WorldRayDirection * (RayMarchingContext.LocalRayTMin + RayMarchingContext.StepSize * StepIndex);
 
-		FVolumeSampleContext SampleContext = CreateVolumeSampleContext(LocalPosition, WorldPosition);
+		FVolumeSampleContext SampleContext = CreateVolumeSampleContext(LocalPosition, WorldPosition, RayMarchingContext.MipLevel);
 		float3 Extinction = SampleExtinction(SampleContext);
 
 		float3 Emission = SampleEmission(SampleContext);
@@ -161,7 +171,7 @@ void RayMarchSingleScattering(
 	{
 		float3 LocalPosition = RayMarchingContext.LocalRayOrigin + RayMarchingContext.LocalRayDirection * (RayMarchingContext.LocalRayTMin + RayMarchingContext.StepSize * StepIndex);
 		float3 WorldPosition = RayMarchingContext.WorldRayOrigin + RayMarchingContext.WorldRayDirection * (RayMarchingContext.LocalRayTMin + RayMarchingContext.StepSize * StepIndex);
-		FVolumeSampleContext SampleContext = CreateVolumeSampleContext(LocalPosition, WorldPosition);
+		FVolumeSampleContext SampleContext = CreateVolumeSampleContext(LocalPosition, WorldPosition, RayMarchingContext.MipLevel);
 
 		if (RayMarchingContext.bApplyEmission)
 		{
@@ -206,8 +216,15 @@ void RayMarchSingleScattering(
 
 					if (RayMarchingContext.bApplyShadowTransmittance)
 					{
-						// Build shadow ray and find intersection
 						float3 LocalShadowRayOrigin = mul(float4(WorldPosition, 1.0), GetWorldToLocal()).xyz;
+
+#if DIM_USE_TRANSMITTANCE_VOLUME
+						float3 LocalBoundsMin = GetLocalBoundsOrigin() - GetLocalBoundsExtent();
+						float3 TransmittanceUVW = saturate((LocalShadowRayOrigin - LocalBoundsMin) / (2.0 * GetLocalBoundsExtent()));
+						float MipLevel = 0;
+						ShadowTransmittance = SampleTransmittanceVolume(TransmittanceUVW, MipLevel);
+#else
+						// Build shadow ray and find intersection
 						float3 LocalShadowRayEnd = mul(float4(WorldPosition + ToLight, 1.0), GetWorldToLocal()).xyz;
 						float3 LocalShadowRayDirection = LocalShadowRayEnd - LocalShadowRayOrigin;
 						float LocalShadowRayTMin = 0.0;
@@ -238,6 +255,7 @@ void RayMarchSingleScattering(
 						);
 
 						ShadowTransmittance = ComputeTransmittance(ShadowRayMarchingContext);
+#endif // DIM_USE_TRANSMITTANCE_VOLUME
 					}
 
 					// Convert in-scattering to exitant radiance
diff --git a/Engine/Shaders/Private/HeterogeneousVolumes/HeterogeneousVolumesSparseVoxelUniformBufferUtils.ush b/Engine/Shaders/Private/HeterogeneousVolumes/HeterogeneousVolumesSparseVoxelUniformBufferUtils.ush
--- a/Engine/Shaders/Private/HeterogeneousVolumes/HeterogeneousVolumesSparseVoxelUniformBufferUtils.ush
+++ b/Engine/Shaders/Private/HeterogeneousVolumes/HeterogeneousVolumesSparseVoxelUniformBufferUtils.ush
@@ -20,6 +20,16 @@ float3 GetLocalBoundsExtent()
 	return SparseVoxelUniformBuffer.LocalBoundsExtent;
 }
 
+float GetMaxTraceDistance()
+{
+	return SparseVoxelUniformBuffer.MaxTraceDistance;
+}
+
+float GetMaxShadowTraceDistance()
+{
+	return SparseVoxelUniformBuffer.MaxShadowTraceDistance;
+}
+
 float GetStepSize()
 {
 	return SparseVoxelUniformBuffer.StepSize;
@@ -48,17 +58,17 @@ FVoxelDataPacked GetVoxelDataPacked(uint VoxelIndex)
 struct FVolumeSampleContext
 {
 	float3 UVW;
-	uint MipLevel;
+	float MipLevel;
 };
 
-FVolumeSampleContext CreateVolumeSampleContext(float3 LocalPosition, float3 WorldPosition)
+FVolumeSampleContext CreateVolumeSampleContext(float3 LocalPosition, float3 WorldPosition, float MipLevel)
 {
 	float3 LocalBoundsMin = GetLocalBoundsOrigin() - GetLocalBoundsExtent();
 	float3 LocalBoundsMax = GetLocalBoundsOrigin() + GetLocalBoundsExtent();
 
 	FVolumeSampleContext VolumeSampleContext;
 	VolumeSampleContext.UVW = (LocalPosition - LocalBoundsMin) / (LocalBoundsMax - LocalBoundsMin);
-	VolumeSampleContext.MipLevel = 0;// RayMarchingContext.MipLevel;
+	VolumeSampleContext.MipLevel = MipLevel;
 
 	return VolumeSampleContext;
 }
diff --git a/Engine/Shaders/Private/HeterogeneousVolumes/HeterogeneousVolumesTransmittanceVolumeUtils.ush b/Engine/Shaders/Private/HeterogeneousVolumes/HeterogeneousVolumesTransmittanceVolumeUtils.ush
new file mode 100644
--- /dev/null
+++ b/Engine/Shaders/Private/HeterogeneousVolumes/HeterogeneousVolumesTransmittanceVolumeUtils.ush
@@ -0,0 +1,24 @@
+// Copyright Epic Games, Inc. All Rights Reserved.
+
+#ifndef DIM_USE_TRANSMITTANCE_VOLUME
+#define DIM_USE_TRANSMITTANCE_VOLUME 0
+#endif // DIM_USE_TRANSMITTANCE_VOLUME
+
+uint3 TransmittanceVolumeResolution;
+Texture3D TransmittanceVolumeTexture;
+
+uint3 GetTransmittanceVolumeResolution()
+{
+	return TransmittanceVolumeResolution;
+}
+
+float3 SampleTransmittanceVolume(float3 UVW, float MipLevel)
+{
+	float3 Transmittance = Texture3DSampleLevel(
+		TransmittanceVolumeTexture,
+		View.SharedTrilinearClampedSampler,
+		UVW,
+		MipLevel).r;
+
+	return Transmittance;
+}
diff --git a/Engine/Source/Runtime/Renderer/Private/HeterogeneousVolumes/HeterogeneousVolumes.cpp b/Engine/Source/Runtime/Renderer/Private/HeterogeneousVolumes/HeterogeneousVolumes.cpp
--- a/Engine/Source/Runtime/Renderer/Private/HeterogeneousVolumes/HeterogeneousVolumes.cpp
+++ b/Engine/Source/Runtime/Renderer/Private/HeterogeneousVolumes/HeterogeneousVolumes.cpp
@@ -52,6 +52,13 @@ static TAutoConsoleVariable<float> CVarHeterogeneousVolumesMaxTraceDistance(
 	ECVF_RenderThreadSafe
 );
 
+static TAutoConsoleVariable<float> CVarHeterogeneousVolumesMaxShadowTraceDistance(
+	TEXT("r.HeterogeneousVolumes.MaxShadowTraceDistance"),
+	10000.0,
+	TEXT("The maximum shadow-trace distance (Default = 10000)"),
+	ECVF_RenderThreadSafe
+);
+
 static TAutoConsoleVariable<int32> CVarHeterogeneousVolumesPreshading(
 	TEXT("r.HeterogeneousVolumes.Preshading"),
 	0,
@@ -89,7 +96,7 @@ static TAutoConsoleVariable<int32> CVarHeterogeneousVolumesPreshadingVolumeResol
 
 static TAutoConsoleVariable<float> CVarHeterogeneousVolumesShadowStepFactor(
 	TEXT("r.HeterogeneousVolumes.ShadowStepFactor"),
-	10.0,
+	8.0,
 	TEXT("The ray-marching step-size factor for shadow rays (Default = 8.0)"),
 	ECVF_RenderThreadSafe
 );
@@ -129,6 +136,20 @@ static TAutoConsoleVariable<float> CVarHeterogeneousVolumesStepSize(
 	ECVF_RenderThreadSafe
 );
 
+static TAutoConsoleVariable<int32> CVarHeterogeneousVolumesTransmittanceVolume(
+	TEXT("r.HeterogeneousVolumes.TransmittanceVolume"),
+	1,
+	TEXT("Enalbes transmittance volume optimized pre-pass (Default = 1)"),
+	ECVF_RenderThreadSafe
+);
+
+static TAutoConsoleVariable<int32> CVarHeterogeneousVolumesTransmittanceVolumeDownsampleFactor(
+	TEXT("r.HeterogeneousVolumes.TransmittanceVolume.DownsampleFactor"),
+	1,
+	TEXT("Determines the downsample factor, relative to the preshading volume resolution (Default = 1)"),
+	ECVF_RenderThreadSafe
+);
+
 DECLARE_GPU_STAT_NAMED(HeterogeneousVolumesStat, TEXT("HeterogeneousVolumes"));
 
 static bool IsHeterogeneousVolumesEnabled()
@@ -187,6 +208,11 @@ namespace HeterogeneousVolumes
 		return CVarHeterogeneousVolumesMaxTraceDistance.GetValueOnRenderThread();
 	}
 
+	float GetMaxShadowTraceDistance()
+	{
+		return CVarHeterogeneousVolumesMaxShadowTraceDistance.GetValueOnRenderThread();
+	}
+
 	float GetStepSize()
 	{
 		return CVarHeterogeneousVolumesStepSize.GetValueOnRenderThread();
@@ -228,6 +254,11 @@ namespace HeterogeneousVolumes
 		return CVarHeterogeneousVolumesSparseVoxelPerTileCulling.GetValueOnRenderThread() != 0;
 	}
 
+	bool UseTransmittanceVolume()
+	{
+		return CVarHeterogeneousVolumesTransmittanceVolume.GetValueOnRenderThread() != 0;
+	}
+
 	bool ShouldJitter()
 	{
 		return CVarHeterogeneousVolumesJitter.GetValueOnRenderThread() != 0;
@@ -259,6 +290,14 @@ namespace HeterogeneousVolumes
 			FMath::Max(VolumeResolution.Z >> MipLevel, 1)
 		);
 	}
+
+	FIntVector GetTransmittanceVolumeResolution()
+	{
+		float DownsampleFactor = CVarHeterogeneousVolumesTransmittanceVolumeDownsampleFactor.GetValueOnRenderThread();
+		FIntVector TransmittanceVolumeResolution = GetVolumeResolution() / DownsampleFactor;
+
+		return TransmittanceVolumeResolution;
+	}
 }
 
 void FDeferredShadingSceneRenderer::RenderHeterogeneousVolumes(
@@ -300,6 +339,19 @@ void FDeferredShadingSceneRenderer::RenderHeterogeneousVolumes(
 
 				RDG_EVENT_SCOPE(GraphBuilder, "%s", *PrimitiveSceneProxy->GetResourceName().ToString());
 
+				// Allocate transmittance volume
+				// TODO: Allow option for scalar transmittance to conserve bandwidth
+				FIntVector TransmittanceVolumeResolution = HeterogeneousVolumes::GetTransmittanceVolumeResolution();
+				uint32 NumMips = FMath::Log2(float(FMath::Min(FMath::Min(TransmittanceVolumeResolution.X, TransmittanceVolumeResolution.Y), TransmittanceVolumeResolution.Z))) + 1;
+				FRDGTextureDesc TransmittanceVolumeDesc = FRDGTextureDesc::Create3D(
+					TransmittanceVolumeResolution,
+					PF_FloatR11G11B10,
+					FClearValueBinding::Black,
+					TexCreate_ShaderResource | TexCreate_UAV | TexCreate_3DTiling,
+					NumMips
+				);
+				FRDGTextureRef TransmittanceVolumeTexture = GraphBuilder.CreateTexture(TransmittanceVolumeDesc, TEXT("HeterogeneousVolumes.TransmittanceVolumeTexture"));
+
 				// Material baking executes a pre-shading pipeline
 				if (CVarHeterogeneousVolumesPreshading.GetValueOnRenderThread())
 				{
@@ -314,6 +366,8 @@ void FDeferredShadingSceneRenderer::RenderHeterogeneousVolumes(
 						MaterialRenderProxy,
 						PrimitiveId,
 						LocalBoxSphereBounds,
+						// Transmittance accleration
+						TransmittanceVolumeTexture,
 						// Output
 						HeterogeneousVolumeRadiance
 					);
@@ -331,6 +385,8 @@ void FDeferredShadingSceneRenderer::RenderHeterogeneousVolumes(
 						MaterialRenderProxy,
 						PrimitiveId,
 						LocalBoxSphereBounds,
+						// Transmittance accleration
+						TransmittanceVolumeTexture,
 						// Output
 						HeterogeneousVolumeRadiance
 					);
diff --git a/Engine/Source/Runtime/Renderer/Private/HeterogeneousVolumes/HeterogeneousVolumes.h b/Engine/Source/Runtime/Renderer/Private/HeterogeneousVolumes/HeterogeneousVolumes.h
--- a/Engine/Source/Runtime/Renderer/Private/HeterogeneousVolumes/HeterogeneousVolumes.h
+++ b/Engine/Source/Runtime/Renderer/Private/HeterogeneousVolumes/HeterogeneousVolumes.h
@@ -26,9 +26,11 @@ namespace HeterogeneousVolumes
 {
 	// CVars
 	FIntVector GetVolumeResolution();
+	FIntVector GetTransmittanceVolumeResolution();
 
 	float GetShadowStepFactor();
 	float GetMaxTraceDistance();
+	float GetMaxShadowTraceDistance();
 	float GetStepSize();
 	float GetMaxStepCount();
 
@@ -41,6 +43,7 @@ namespace HeterogeneousVolumes
 	bool UseHardwareRayTracing();
 	bool UseSparseVoxelPipeline();
 	bool UseSparseVoxelPerTileCulling();
+	bool UseTransmittanceVolume();
 
 	// Convenience Utils
 	int GetVoxelCount(FIntVector VolumeResolution);
@@ -68,6 +71,9 @@ BEGIN_UNIFORM_BUFFER_STRUCT(FSparseVoxelUniformBufferParameters, )
 	SHADER_PARAMETER_RDG_TEXTURE(Texture3D, AlbedoTexture)
 	SHADER_PARAMETER_SAMPLER(SamplerState, TextureSampler)
 
+	// Resolution
+	SHADER_PARAMETER(FIntVector, TransmittanceVolumeResolution)
+
 	// Sparse voxel data
 	SHADER_PARAMETER_RDG_BUFFER_SRV(Buffer<uint>, NumVoxelsBuffer)
 	SHADER_PARAMETER_RDG_BUFFER_SRV(StructuredBuffer<FVoxelDataPacked>, VoxelBuffer)
@@ -75,10 +81,16 @@ BEGIN_UNIFORM_BUFFER_STRUCT(FSparseVoxelUniformBufferParameters, )
 
 	// Traversal hints
 	SHADER_PARAMETER(float, MaxTraceDistance)
+	SHADER_PARAMETER(float, MaxShadowTraceDistance)
 	SHADER_PARAMETER(float, StepSize)
 	SHADER_PARAMETER(float, ShadowStepFactor)
 END_UNIFORM_BUFFER_STRUCT()
 
+BEGIN_SHADER_PARAMETER_STRUCT(FTransmittanceVolumeParameters, )
+	SHADER_PARAMETER(FIntVector, TransmittanceVolumeResolution)
+	SHADER_PARAMETER_RDG_TEXTURE_SRV(Texture3D, TransmittanceVolumeTexture)
+END_SHADER_PARAMETER_STRUCT()
+
 // Render specializations
 
 void RenderWithLiveShading(
@@ -91,6 +103,8 @@ void RenderWithLiveShading(
 	const FMaterialRenderProxy* MaterialRenderProxy,
 	const int32 PrimitiveId,
 	const FBoxSphereBounds LocalBoxSphereBounds,
+	// Transmittance acceleration
+	FRDGTextureRef TransmittanceVolumeTexture,
 	// Output
 	FRDGTextureRef& HeterogeneousVolumeRadiance
 );
@@ -107,6 +121,8 @@ void RenderWithPreshading(
 	const FMaterialRenderProxy* MaterialRenderProxy,
 	const int32 PrimitiveId,
 	const FBoxSphereBounds LocalBoxSphereBounds,
+	// Transmittance acceleration
+	FRDGTextureRef TransmittanceVolumeTexture,
 	// Output
 	FRDGTextureRef& HeterogeneousVolumeRadiance
 );
@@ -181,6 +197,28 @@ void GenerateRayTracingScene(
 	FRayTracingScene& RayTracingScene
 );
 
+void RenderTransmittanceVolumeWithPreshadingHardwareRayTracing(
+	FRDGBuilder& GraphBuilder,
+	// Scene data
+	const FScene* Scene,
+	const FViewInfo& View,
+	const FSceneTextures& SceneTextures,
+	// Light data
+	bool bApplyEmission,
+	bool bApplyDirectLighting,
+	bool bApplyShadowTransmittance,
+	uint32 LightType,
+	const FLightSceneInfo* LightSceneInfo,
+	// Object data
+	const FPrimitiveSceneProxy* PrimitiveSceneProxy,
+	// Sparse voxel data
+	TRDGUniformBufferRef<FSparseVoxelUniformBufferParameters> SparseVoxelUniformBuffer,
+	// Ray tracing data
+	FRayTracingScene& RayTracingScene,
+	// Output
+	FRDGTextureRef& TransmittanceVolumeTexture
+);
+
 void RenderSingleScatteringWithPreshadingHardwareRayTracing(
 	FRDGBuilder& GraphBuilder,
 	// Scene data
@@ -199,6 +237,8 @@ void RenderSingleScatteringWithPreshadingHardwareRayTracing(
 	TRDGUniformBufferRef<FSparseVoxelUniformBufferParameters> SparseVoxelUniformBuffer,
 	// Ray tracing data
 	FRayTracingScene& RayTracingScene,
+	// Transmittance volume
+	FRDGTextureRef TransmittanceVolumeTexture,
 	// Output
 	FRDGTextureRef& HeterogeneousVolumeTexture
 );
diff --git a/Engine/Source/Runtime/Renderer/Private/HeterogeneousVolumes/HeterogeneousVolumesHardwareRayTracing.cpp b/Engine/Source/Runtime/Renderer/Private/HeterogeneousVolumes/HeterogeneousVolumesHardwareRayTracing.cpp
--- a/Engine/Source/Runtime/Renderer/Private/HeterogeneousVolumes/HeterogeneousVolumesHardwareRayTracing.cpp
+++ b/Engine/Source/Runtime/Renderer/Private/HeterogeneousVolumes/HeterogeneousVolumesHardwareRayTracing.cpp
@@ -319,13 +319,59 @@ class FHeterogeneousVolumesSparseVoxelMS : public FGlobalShader
 
 IMPLEMENT_GLOBAL_SHADER(FHeterogeneousVolumesSparseVoxelMS, "/Engine/Private/HeterogeneousVolumes/HeterogeneousVolumesHardwareRayTracing.usf", "SparseVoxelsMissShader", SF_RayMiss);
 
+class FRenderTransmittanceVolumeWithPreshadingRGS : public FGlobalShader
+{
+	DECLARE_GLOBAL_SHADER(FRenderTransmittanceVolumeWithPreshadingRGS)
+	SHADER_USE_ROOT_PARAMETER_STRUCT(FRenderTransmittanceVolumeWithPreshadingRGS, FGlobalShader)
+
+	class FApplyShadowTransmittanceDim : SHADER_PERMUTATION_BOOL("DIM_APPLY_SHADOW_TRANSMITTANCE");
+	using FPermutationDomain = TShaderPermutationDomain<FApplyShadowTransmittanceDim>;
+
+	BEGIN_SHADER_PARAMETER_STRUCT(FParameters, )
+		// Scene 
+		SHADER_PARAMETER_SRV(RaytracingAccelerationStructure, TLAS)
+		SHADER_PARAMETER_STRUCT_REF(FViewUniformShaderParameters, ViewUniformBuffer)
+		SHADER_PARAMETER_STRUCT_INCLUDE(FSceneTextureParameters, SceneTextures)
+
+		// Lighting data
+		SHADER_PARAMETER(int, LightType)
+		SHADER_PARAMETER_STRUCT_REF(FDeferredLightUniformStruct, DeferredLight)
+
+		// Sparse Volume
+		SHADER_PARAMETER_RDG_UNIFORM_BUFFER(FSparseVoxelUniformBufferParameters, SparseVoxelUniformBuffer)
+
+		// Volume
+		SHADER_PARAMETER(int, MipLevel)
+
+		// Ray
+		SHADER_PARAMETER(float, MaxTraceDistance)
+		SHADER_PARAMETER(int, MaxStepCount)
+		SHADER_PARAMETER(int, bJitter)
+
+		// Transmittance volume data
+		SHADER_PARAMETER_STRUCT_INCLUDE(FTransmittanceVolumeParameters, TransmittanceVolume)
+
+		// Output
+		SHADER_PARAMETER_RDG_TEXTURE_UAV(RWTexture2D<float4>, RWTransmittanceVolumeTexture)
+	END_SHADER_PARAMETER_STRUCT()
+
+	static bool ShouldCompilePermutation(const FGlobalShaderPermutationParameters& Parameters)
+	{
+		return ShouldCompileRayTracingShadersForProject(Parameters.Platform) && DoesPlatformSupportHeterogeneousVolumes(Parameters.Platform) &&
+			FDataDrivenShaderPlatformInfo::GetSupportsRayTracingProceduralPrimitive(Parameters.Platform);
+	}
+};
+
+IMPLEMENT_GLOBAL_SHADER(FRenderTransmittanceVolumeWithPreshadingRGS, "/Engine/Private/HeterogeneousVolumes/HeterogeneousVolumesHardwareRayTracing.usf", "RenderTransmittanceVolumeWithPreshadingRGS", SF_RayGen);
+
 class FRenderSingleScatteringWithPreshadingRGS : public FGlobalShader
 {
 	DECLARE_GLOBAL_SHADER(FRenderSingleScatteringWithPreshadingRGS)
 	SHADER_USE_ROOT_PARAMETER_STRUCT(FRenderSingleScatteringWithPreshadingRGS, FGlobalShader)
 
 	class FApplyShadowTransmittanceDim : SHADER_PERMUTATION_BOOL("DIM_APPLY_SHADOW_TRANSMITTANCE");
-	using FPermutationDomain = TShaderPermutationDomain<FApplyShadowTransmittanceDim>;
+	class FUseTransmittanceVolume : SHADER_PERMUTATION_BOOL("DIM_USE_TRANSMITTANCE_VOLUME");
+	using FPermutationDomain = TShaderPermutationDomain<FApplyShadowTransmittanceDim, FUseTransmittanceVolume>;
 
 	BEGIN_SHADER_PARAMETER_STRUCT(FParameters, )
 		// Scene 
@@ -345,6 +391,9 @@ class FRenderSingleScatteringWithPreshadingRGS : public FGlobalShader
 		// Volume
 		SHADER_PARAMETER(int, MipLevel)
 
+		// Transmittance volume
+		SHADER_PARAMETER_STRUCT_INCLUDE(FTransmittanceVolumeParameters, TransmittanceVolume)
+
 		// Ray
 		SHADER_PARAMETER(float, MaxTraceDistance)
 		SHADER_PARAMETER(int, MaxStepCount)
@@ -438,6 +487,105 @@ FRayTracingPipelineState* BuildRayTracingPipelineState(
 	return RayTracingPipelineState;
 }
 
+void RenderTransmittanceVolumeWithPreshadingHardwareRayTracing(
+	FRDGBuilder& GraphBuilder,
+	// Scene data
+	const FScene* Scene,
+	const FViewInfo& View,
+	const FSceneTextures& SceneTextures,
+	// Light data
+	bool bApplyEmission,
+	bool bApplyDirectLighting,
+	bool bApplyShadowTransmittance,
+	uint32 LightType,
+	const FLightSceneInfo* LightSceneInfo,
+	// Object data
+	const FPrimitiveSceneProxy* PrimitiveSceneProxy,
+	// Sparse voxel data
+	TRDGUniformBufferRef<FSparseVoxelUniformBufferParameters> SparseVoxelUniformBuffer,
+	// Ray tracing data
+	FRayTracingScene& RayTracingScene,
+	// Output
+	FRDGTextureRef& TransmittanceVolumeTexture
+)
+{
+	FRenderTransmittanceVolumeWithPreshadingRGS::FParameters* PassParameters = GraphBuilder.AllocParameters<FRenderTransmittanceVolumeWithPreshadingRGS::FParameters>();
+	{
+		// Scene
+		PassParameters->TLAS = RayTracingScene.GetLayerSRVChecked(ERayTracingSceneLayer::Base);
+		PassParameters->ViewUniformBuffer = View.ViewUniformBuffer;
+		PassParameters->SceneTextures = GetSceneTextureParameters(GraphBuilder, SceneTextures);
+
+		// Light data
+		check(LightSceneInfo != nullptr);
+		FDeferredLightUniformStruct DeferredLightUniform = GetDeferredLightParameters(View, *LightSceneInfo);
+		PassParameters->DeferredLight = CreateUniformBufferImmediate(DeferredLightUniform, UniformBuffer_SingleDraw);
+		PassParameters->LightType = LightType;
+
+		// Sparse Voxel data
+		PassParameters->SparseVoxelUniformBuffer = SparseVoxelUniformBuffer;
+
+		// Transmittance volume
+		PassParameters->TransmittanceVolume.TransmittanceVolumeResolution = HeterogeneousVolumes::GetTransmittanceVolumeResolution();
+		//PassParameters->TransmittanceVolume.TransmittanceVolumeTexture = GraphBuilder.CreateSRV(TransmittanceVolumeTexture);
+
+		// Ray data
+		PassParameters->MaxTraceDistance = HeterogeneousVolumes::GetMaxTraceDistance();
+		PassParameters->MaxStepCount = HeterogeneousVolumes::GetMaxStepCount();
+		PassParameters->bJitter = HeterogeneousVolumes::ShouldJitter();
+		PassParameters->MipLevel = HeterogeneousVolumes::GetMipLevel();
+
+		// Output
+		PassParameters->RWTransmittanceVolumeTexture = GraphBuilder.CreateUAV(TransmittanceVolumeTexture);
+	}
+
+	FRenderTransmittanceVolumeWithPreshadingRGS::FPermutationDomain PermutationVector;
+	TShaderRef<FRenderTransmittanceVolumeWithPreshadingRGS> RayGenerationShader = View.ShaderMap->GetShader<FRenderTransmittanceVolumeWithPreshadingRGS>(PermutationVector);
+	FIntVector VolumeResolution = HeterogeneousVolumes::GetTransmittanceVolumeResolution();
+	FIntPoint DispatchResolution = FIntPoint(VolumeResolution.X, VolumeResolution.Y * VolumeResolution.Z);
+
+	FString LightName = TEXT("none");
+	if (LightSceneInfo != nullptr)
+	{
+		FSceneRenderer::GetLightNameForDrawEvent(LightSceneInfo->Proxy, LightName);
+	}
+
+	GraphBuilder.AddPass(
+		RDG_EVENT_NAME("RenderTransmittanceVolumeWithPreshadingRGS (Light = %s) %ux%u", *LightName, DispatchResolution.X, DispatchResolution.Y),
+		PassParameters,
+		ERDGPassFlags::Compute | ERDGPassFlags::NeverCull,
+		[
+			PassParameters,
+			&View,
+			&RayTracingScene,
+			RayGenerationShader,
+			DispatchResolution
+		](FRHIRayTracingCommandList& RHICmdList)
+		{
+			// Set ray-gen bindings
+			FRayTracingShaderBindingsWriter GlobalResources;
+			SetShaderParameters(GlobalResources, RayGenerationShader, *PassParameters);
+
+			// Create pipeline
+			FRayTracingPipelineState* RayTracingPipelineState = BuildRayTracingPipelineState(RHICmdList, View, RayGenerationShader.GetRayTracingShader());
+
+			// Set hit-group bindings
+			const uint32 NumBindings = 1;
+			FRayTracingLocalShaderBindings* Bindings = BuildRayTracingMaterialBindings(RHICmdList, View, PassParameters->SparseVoxelUniformBuffer->GetRHI());
+			RHICmdList.SetRayTracingHitGroups(RayTracingScene.GetRHIRayTracingSceneChecked(), RayTracingPipelineState, NumBindings, Bindings);
+			RHICmdList.SetRayTracingMissShaders(RayTracingScene.GetRHIRayTracingSceneChecked(), RayTracingPipelineState, NumBindings, Bindings);
+
+			// Dispatch
+			RHICmdList.RayTraceDispatch(
+				RayTracingPipelineState,
+				RayGenerationShader.GetRayTracingShader(),
+				RayTracingScene.GetRHIRayTracingSceneChecked(),
+				GlobalResources,
+				DispatchResolution.X, DispatchResolution.Y);
+		}
+	);
+}
+
 void RenderSingleScatteringWithPreshadingHardwareRayTracing(
 	FRDGBuilder& GraphBuilder,
 	// Scene data
@@ -456,6 +604,8 @@ void RenderSingleScatteringWithPreshadingHardwareRayTracing(
 	TRDGUniformBufferRef<FSparseVoxelUniformBufferParameters> SparseVoxelUniformBuffer,
 	// Ray tracing data
 	FRayTracingScene& RayTracingScene,
+	// Transmittance volume
+	FRDGTextureRef TransmittanceVolumeTexture,
 	// Output
 	FRDGTextureRef& HeterogeneousVolumeTexture
 )
@@ -484,6 +634,13 @@ void RenderSingleScatteringWithPreshadingHardwareRayTracing(
 		// Volume data
 		PassParameters->MipLevel = HeterogeneousVolumes::GetMipLevel();
 
+		// Transmittance volume
+		if (HeterogeneousVolumes::UseTransmittanceVolume() && bApplyShadowTransmittance)
+		{
+			PassParameters->TransmittanceVolume.TransmittanceVolumeResolution = HeterogeneousVolumes::GetTransmittanceVolumeResolution();
+			PassParameters->TransmittanceVolume.TransmittanceVolumeTexture = GraphBuilder.CreateSRV(TransmittanceVolumeTexture);
+		}
+
 		// Ray data
 		PassParameters->MaxTraceDistance = HeterogeneousVolumes::GetMaxTraceDistance();
 		PassParameters->MaxStepCount = HeterogeneousVolumes::GetMaxStepCount();
@@ -495,6 +652,7 @@ void RenderSingleScatteringWithPreshadingHardwareRayTracing(
 
 	FRenderSingleScatteringWithPreshadingRGS::FPermutationDomain PermutationVector;
 	PermutationVector.Set<FRenderSingleScatteringWithPreshadingRGS::FApplyShadowTransmittanceDim>(bApplyShadowTransmittance);
+	PermutationVector.Set<FRenderSingleScatteringWithPreshadingRGS::FUseTransmittanceVolume>(HeterogeneousVolumes::UseTransmittanceVolume());
 	TShaderRef<FRenderSingleScatteringWithPreshadingRGS> RayGenerationShader = View.ShaderMap->GetShader<FRenderSingleScatteringWithPreshadingRGS>(PermutationVector);
 	FIntPoint DispatchResolution = View.ViewRect.Size();
 
@@ -524,43 +682,10 @@ void RenderSingleScatteringWithPreshadingHardwareRayTracing(
 			FRayTracingPipelineState* RayTracingPipelineState = BuildRayTracingPipelineState(RHICmdList, View, RayGenerationShader.GetRayTracingShader());
 
 			// Set hit-group bindings
-#if 0
 			const uint32 NumBindings = 1;
 			FRayTracingLocalShaderBindings* Bindings = BuildRayTracingMaterialBindings(RHICmdList, View, PassParameters->SparseVoxelUniformBuffer->GetRHI());
-			const bool bCopyDataToInlineStorage = true;
-			RHICmdList.SetRayTracingHitGroups(RayTracingScene, RayTracingPipelineState, NumBindings, Bindings, bCopyDataToInlineStorage);
-#else
-			uint32 InstanceIndex = 0;
-			uint32 SegmentIndex = 0;
-			uint32 ShaderSlot = 0;
-			uint32 HitGroupIndex = 0;
-			uint32 LooseParameterDataSize = 0;
-			void* LooseParameterData = nullptr;
-			uint32 UserData = 0;
-			uint32 NumUniformBuffers = 1;
-			FRHIUniformBuffer* UniformBuffers[] = {
-				PassParameters->SparseVoxelUniformBuffer->GetRHI()
-			};
-			RHICmdList.SetRayTracingHitGroup(
-				RayTracingScene.GetRHIRayTracingSceneChecked(),
-				InstanceIndex, SegmentIndex, ShaderSlot,
-				RayTracingPipelineState, HitGroupIndex,
-				NumUniformBuffers, UniformBuffers,
-				LooseParameterDataSize, LooseParameterData,
-				UserData
-			);
-
-			uint32 ShaderIndexInPipeline = 0;
-			RHICmdList.SetRayTracingMissShader(
-				RayTracingScene.GetRHIRayTracingSceneChecked(),
-				ShaderSlot,
-				RayTracingPipelineState,
-				ShaderIndexInPipeline,
-				NumUniformBuffers,
-				UniformBuffers,
-				UserData
-			);
-#endif
+			RHICmdList.SetRayTracingHitGroups(RayTracingScene.GetRHIRayTracingSceneChecked(), RayTracingPipelineState, NumBindings, Bindings);
+			RHICmdList.SetRayTracingMissShaders(RayTracingScene.GetRHIRayTracingSceneChecked(), RayTracingPipelineState, NumBindings, Bindings);
 
 			// Dispatch
 			RHICmdList.RayTraceDispatch(
diff --git a/Engine/Source/Runtime/Renderer/Private/HeterogeneousVolumes/HeterogeneousVolumesLiveShadingPipeline.cpp b/Engine/Source/Runtime/Renderer/Private/HeterogeneousVolumes/HeterogeneousVolumesLiveShadingPipeline.cpp
--- a/Engine/Source/Runtime/Renderer/Private/HeterogeneousVolumes/HeterogeneousVolumesLiveShadingPipeline.cpp
+++ b/Engine/Source/Runtime/Renderer/Private/HeterogeneousVolumes/HeterogeneousVolumesLiveShadingPipeline.cpp
@@ -10,10 +10,120 @@
 #include "ScenePrivate.h"
 #include "SceneManagement.h"
 
+class FRenderTransmittanceVolumeWithLiveShadingCS : public FMeshMaterialShader
+{
+	DECLARE_SHADER_TYPE(FRenderTransmittanceVolumeWithLiveShadingCS, MeshMaterial);
+
+	BEGIN_SHADER_PARAMETER_STRUCT(FParameters, )
+		// Scene data
+		SHADER_PARAMETER_STRUCT_REF(FViewUniformShaderParameters, View)
+		SHADER_PARAMETER_STRUCT_INCLUDE(FSceneTextureParameters, SceneTextures)
+
+		// Light data
+		SHADER_PARAMETER(int, bApplyEmission)
+		SHADER_PARAMETER(int, bApplyDirectLighting)
+		SHADER_PARAMETER(int, bApplyShadowTransmittance)
+		SHADER_PARAMETER(int, LightType)
+		SHADER_PARAMETER_STRUCT_REF(FDeferredLightUniformStruct, DeferredLight)
+
+		// Shadow data
+		SHADER_PARAMETER(float, ShadowStepFactor)
+
+		// Object data
+		SHADER_PARAMETER(FMatrix44f, LocalToWorld)
+		SHADER_PARAMETER(FMatrix44f, WorldToLocal)
+		SHADER_PARAMETER(FVector3f, LocalBoundsOrigin)
+		SHADER_PARAMETER(FVector3f, LocalBoundsExtent)
+		SHADER_PARAMETER(int32, PrimitiveId)
+
+		// Ray data
+		SHADER_PARAMETER(float, MaxTraceDistance)
+		SHADER_PARAMETER(float, MaxShadowTraceDistance)
+		SHADER_PARAMETER(float, StepSize)
+		SHADER_PARAMETER(int, MaxStepCount)
+		SHADER_PARAMETER(int, bJitter)
+
+		// Volume data
+		SHADER_PARAMETER_STRUCT_INCLUDE(FTransmittanceVolumeParameters, TransmittanceVolume)
+
+		// Output
+		SHADER_PARAMETER_RDG_TEXTURE_UAV(RWTexture3D<float>, RWTransmittanceVolumeTexture)
+	END_SHADER_PARAMETER_STRUCT()
+
+	FRenderTransmittanceVolumeWithLiveShadingCS() = default;
+
+	FRenderTransmittanceVolumeWithLiveShadingCS(
+		const FMeshMaterialShaderType::CompiledShaderInitializerType & Initializer
+	)
+		: FMeshMaterialShader(Initializer)
+	{
+		Bindings.BindForLegacyShaderParameters(
+			this,
+			Initializer.PermutationId,
+			Initializer.ParameterMap,
+			*FParameters::FTypeInfo::GetStructMetadata(),
+			// Don't require full bindings, we use FMaterialShader::SetParameters
+			false
+		);
+	}
+
+	static bool ShouldCompilePermutation(
+		const FMaterialShaderPermutationParameters & Parameters
+	)
+	{
+		return DoesPlatformSupportHeterogeneousVolumes(Parameters.Platform)
+			&& (Parameters.MaterialParameters.MaterialDomain == MD_Volume)
+			&& Parameters.MaterialParameters.bIsUsedWithNiagaraMeshParticles;
+	}
+
+	static FPermutationDomain RemapPermutation(FPermutationDomain PermutationVector)
+	{
+		return PermutationVector;
+	}
+
+	static void ModifyCompilationEnvironment(
+		const FMaterialShaderPermutationParameters & Parameters,
+		FShaderCompilerEnvironment & OutEnvironment
+	)
+	{
+		FMaterialShader::ModifyCompilationEnvironment(Parameters, OutEnvironment);
+		OutEnvironment.SetDefine(TEXT("THREADGROUP_SIZE_1D"), GetThreadGroupSize1D());
+		OutEnvironment.SetDefine(TEXT("THREADGROUP_SIZE_2D"), GetThreadGroupSize2D());
+		OutEnvironment.SetDefine(TEXT("THREADGROUP_SIZE_3D"), GetThreadGroupSize3D());
+
+		// This shader takes a very long time to compile with FXC, so we pre-compile it with DXC first and then forward the optimized HLSL to FXC.
+		OutEnvironment.CompilerFlags.Add(CFLAG_PrecompileWithDXC);
+		OutEnvironment.CompilerFlags.Add(CFLAG_AllowTypedUAVLoads);
+
+		OutEnvironment.SetDefine(TEXT("GET_PRIMITIVE_DATA_OVERRIDE"), 1);
+	}
+
+	void SetParameters(
+		FRHIComputeCommandList & RHICmdList,
+		FRHIComputeShader * ShaderRHI,
+		const FViewInfo & View,
+		const FMaterialRenderProxy * MaterialProxy,
+		const FMaterial & Material
+	)
+	{
+		FMaterialShader::SetViewParameters(RHICmdList, ShaderRHI, View, View.ViewUniformBuffer);
+		FMaterialShader::SetParameters(RHICmdList, ShaderRHI, MaterialProxy, Material, View);
+	}
+
+	static int32 GetThreadGroupSize1D() { return GetThreadGroupSize2D() * GetThreadGroupSize2D(); }
+	static int32 GetThreadGroupSize2D() { return 8; }
+	static int32 GetThreadGroupSize3D() { return 4; }
+};
+
+IMPLEMENT_MATERIAL_SHADER_TYPE(, FRenderTransmittanceVolumeWithLiveShadingCS, TEXT("/Engine/Private/HeterogeneousVolumes/HeterogeneousVolumesLiveShadingPipeline.usf"), TEXT("RenderTransmittanceVolumeWithLiveShadingCS"), SF_Compute);
+
 class FRenderSingleScatteringWithLiveShadingCS : public FMeshMaterialShader
 {
 	DECLARE_SHADER_TYPE(FRenderSingleScatteringWithLiveShadingCS, MeshMaterial);
 
+	class FUseTransmittanceVolume : SHADER_PERMUTATION_BOOL("DIM_USE_TRANSMITTANCE_VOLUME");
+	using FPermutationDomain = TShaderPermutationDomain<FUseTransmittanceVolume>;
+
 	BEGIN_SHADER_PARAMETER_STRUCT(FParameters, )
 		// Scene data
 		SHADER_PARAMETER_STRUCT_REF(FViewUniformShaderParameters, View)
@@ -36,6 +146,9 @@ class FRenderSingleScatteringWithLiveShadingCS : public FMeshMaterialShader
 		SHADER_PARAMETER(FVector3f, LocalBoundsExtent)
 		SHADER_PARAMETER(int32, PrimitiveId)
 
+		// Volume data
+		SHADER_PARAMETER_STRUCT_INCLUDE(FTransmittanceVolumeParameters, TransmittanceVolume)
+
 		// Ray data
 		SHADER_PARAMETER(float, MaxTraceDistance)
 		SHADER_PARAMETER(float, StepSize)
@@ -48,9 +161,9 @@ class FRenderSingleScatteringWithLiveShadingCS : public FMeshMaterialShader
 		// Output
 		SHADER_PARAMETER_RDG_TEXTURE_UAV(RWTexture2D<float4>, RWLightingTexture)
 		//SHADER_PARAMETER_RDG_BUFFER_UAV(RWStructuredBuffer<Volumes::FDebugOutput>, RWDebugOutputBuffer)
-		END_SHADER_PARAMETER_STRUCT()
+	END_SHADER_PARAMETER_STRUCT()
 
-		FRenderSingleScatteringWithLiveShadingCS() = default;
+	FRenderSingleScatteringWithLiveShadingCS() = default;
 
 	FRenderSingleScatteringWithLiveShadingCS(
 		const FMeshMaterialShaderType::CompiledShaderInitializerType& Initializer
@@ -115,7 +228,52 @@ class FRenderSingleScatteringWithLiveShadingCS : public FMeshMaterialShader
 
 IMPLEMENT_MATERIAL_SHADER_TYPE(, FRenderSingleScatteringWithLiveShadingCS, TEXT("/Engine/Private/HeterogeneousVolumes/HeterogeneousVolumesLiveShadingPipeline.usf"), TEXT("RenderSingleScatteringWithLiveShadingCS"), SF_Compute);
 
-void RenderSingleScatteringWithLiveShading(
+template <typename ComputeShaderType>
+void AddComputePass(
+	FRDGBuilder& GraphBuilder,
+	TShaderRef<ComputeShaderType>& ComputeShader,
+	typename ComputeShaderType::FParameters* PassParameters,
+	const FScene* Scene,
+	const FViewInfo& View,
+	const FMaterialRenderProxy* MaterialRenderProxy,
+	const FMaterial& Material,
+	const FString& PassName,
+	FIntVector GroupCount
+)
+{
+	GraphBuilder.AddPass(
+		RDG_EVENT_NAME("%s", *PassName),
+		PassParameters,
+		ERDGPassFlags::Compute,
+		[ComputeShader, PassParameters, Scene, &View, MaterialRenderProxy, &Material, GroupCount](FRHIComputeCommandList& RHICmdList)
+		{
+			FMeshPassProcessorRenderState DrawRenderState;
+
+			FMeshMaterialShaderElementData ShaderElementData;
+			ShaderElementData.FadeUniformBuffer = GDistanceCullFadedInUniformBuffer.GetUniformBufferRHI();
+			ShaderElementData.DitherUniformBuffer = GDitherFadedInUniformBuffer.GetUniformBufferRHI();
+
+			FMeshProcessorShaders PassShaders;
+			PassShaders.ComputeShader = ComputeShader;
+
+			FMeshDrawShaderBindings ShaderBindings;
+			{
+				ShaderBindings.Initialize(PassShaders);
+
+				int32 DataOffset = 0;
+				FMeshDrawSingleShaderBindings SingleShaderBindings = ShaderBindings.GetSingleShaderBindings(SF_Compute, DataOffset);
+				ComputeShader->GetShaderBindings(Scene, Scene->GetFeatureLevel(), nullptr, *MaterialRenderProxy, Material, DrawRenderState, ShaderElementData, SingleShaderBindings);
+
+				ShaderBindings.Finalize(&PassShaders);
+			}
+			ShaderBindings.SetOnCommandList(RHICmdList, ComputeShader.GetComputeShader());
+
+			FComputeShaderUtils::Dispatch(RHICmdList, ComputeShader, *PassParameters, GroupCount);
+		}
+	);
+}
+
+void RenderTransmittanceVolumeWithLiveShading(
 	FRDGBuilder& GraphBuilder,
 	// Scene data
 	const FScene* Scene,
@@ -133,14 +291,94 @@ void RenderSingleScatteringWithLiveShading(
 	const int32 PrimitiveId,
 	const FBoxSphereBounds LocalBoxSphereBounds,
 	// Output
-	FRDGTextureRef& HeterogeneousVolumeTexture
+	FRDGTextureRef TransmittanceVolumeTexture
 )
 {
 	const FMaterial& Material = MaterialRenderProxy->GetMaterialWithFallback(View.GetFeatureLevel(), MaterialRenderProxy);
-	if (Material.GetMaterialDomain() != MD_Volume)
+	check(Material.GetMaterialDomain() == MD_Volume);
+
+	FRenderTransmittanceVolumeWithLiveShadingCS::FParameters* PassParameters = GraphBuilder.AllocParameters<FRenderTransmittanceVolumeWithLiveShadingCS::FParameters>();
+	{
+		// Scene data
+		PassParameters->View = View.ViewUniformBuffer;
+		PassParameters->SceneTextures = GetSceneTextureParameters(GraphBuilder, SceneTextures);
+
+		// Light data
+		FDeferredLightUniformStruct DeferredLightUniform;
+		check(LightSceneInfo != nullptr)
+		DeferredLightUniform = GetDeferredLightParameters(View, *LightSceneInfo);
+		PassParameters->DeferredLight = CreateUniformBufferImmediate(DeferredLightUniform, UniformBuffer_SingleDraw);
+		PassParameters->LightType = LightType;
+
+		// Object data
+		FMatrix44f LocalToWorld = FMatrix44f(PrimitiveSceneProxy->GetLocalToWorld());
+		PassParameters->LocalToWorld = LocalToWorld;
+		PassParameters->WorldToLocal = LocalToWorld.Inverse();
+		PassParameters->LocalBoundsOrigin = FVector3f(LocalBoxSphereBounds.Origin);
+		PassParameters->LocalBoundsExtent = FVector3f(LocalBoxSphereBounds.BoxExtent);
+		PassParameters->PrimitiveId = PrimitiveId;
+
+		// Transmittance volume
+		PassParameters->TransmittanceVolume.TransmittanceVolumeResolution = HeterogeneousVolumes::GetTransmittanceVolumeResolution();
+		//PassParameters->TransmittanceVolume.TransmittanceVolumeTexture = GraphBuilder.CreateSRV(TransmittanceVolumeTexture);
+
+		// Ray data
+		PassParameters->MaxTraceDistance = HeterogeneousVolumes::GetMaxTraceDistance();
+		PassParameters->MaxShadowTraceDistance = HeterogeneousVolumes::GetMaxShadowTraceDistance();
+		PassParameters->StepSize = HeterogeneousVolumes::GetStepSize();
+		PassParameters->ShadowStepFactor = HeterogeneousVolumes::GetShadowStepFactor();
+		PassParameters->MaxStepCount = HeterogeneousVolumes::GetMaxStepCount();
+		PassParameters->bJitter = HeterogeneousVolumes::ShouldJitter();
+
+		// Output
+		PassParameters->RWTransmittanceVolumeTexture = GraphBuilder.CreateUAV(TransmittanceVolumeTexture);
+	}
+
+	FString LightName = TEXT("none");
+	if (LightSceneInfo != nullptr)
+	{
+		FSceneRenderer::GetLightNameForDrawEvent(LightSceneInfo->Proxy, LightName);
+	}
+	FString PassName = FString::Printf(TEXT("RenderTransmittanceVolumeWithLiveShadingCS (Light = %s)"), *LightName);
+
+	FIntVector GroupCount = HeterogeneousVolumes::GetTransmittanceVolumeResolution();
+	GroupCount.X = FMath::DivideAndRoundUp(GroupCount.X, FRenderTransmittanceVolumeWithLiveShadingCS::GetThreadGroupSize3D());
+	GroupCount.Y = FMath::DivideAndRoundUp(GroupCount.Y, FRenderTransmittanceVolumeWithLiveShadingCS::GetThreadGroupSize3D());
+	GroupCount.Z = FMath::DivideAndRoundUp(GroupCount.Z, FRenderTransmittanceVolumeWithLiveShadingCS::GetThreadGroupSize3D());
+
+	FRenderTransmittanceVolumeWithLiveShadingCS::FPermutationDomain PermutationVector;
+	TShaderRef<FRenderTransmittanceVolumeWithLiveShadingCS> ComputeShader = Material.GetShader<FRenderTransmittanceVolumeWithLiveShadingCS>(&FLocalVertexFactory::StaticType, PermutationVector, false);
+	if (!ComputeShader.IsNull())
 	{
-		return;
+		AddComputePass(GraphBuilder, ComputeShader, PassParameters, Scene, View, MaterialRenderProxy, Material, PassName, GroupCount);
 	}
+}
+
+void RenderSingleScatteringWithLiveShading(
+	FRDGBuilder& GraphBuilder,
+	// Scene data
+	const FScene* Scene,
+	const FViewInfo& View,
+	const FSceneTextures& SceneTextures,
+	// Light data
+	bool bApplyEmission,
+	bool bApplyDirectLighting,
+	bool bApplyShadowTransmittance,
+	uint32 LightType,
+	const FLightSceneInfo* LightSceneInfo,
+	// Object data
+	const FPrimitiveSceneProxy* PrimitiveSceneProxy,
+	const FMaterialRenderProxy* MaterialRenderProxy,
+	const int32 PrimitiveId,
+	const FBoxSphereBounds LocalBoxSphereBounds,
+	// Transmittance acceleration
+	FRDGTextureRef TransmittanceVolumeTexture,
+	// Output
+	FRDGTextureRef& HeterogeneousVolumeTexture
+)
+{
+	const FMaterial& Material = MaterialRenderProxy->GetMaterialWithFallback(View.GetFeatureLevel(), MaterialRenderProxy);
+	check(Material.GetMaterialDomain() == MD_Volume);
 
 	uint32 GroupCountX = FMath::DivideAndRoundUp(View.ViewRect.Size().X, FRenderSingleScatteringWithLiveShadingCS::GetThreadGroupSize2D());
 	uint32 GroupCountY = FMath::DivideAndRoundUp(View.ViewRect.Size().Y, FRenderSingleScatteringWithLiveShadingCS::GetThreadGroupSize2D());
@@ -186,6 +424,13 @@ void RenderSingleScatteringWithLiveShading(
 		PassParameters->MaxStepCount = HeterogeneousVolumes::GetMaxStepCount();
 		PassParameters->bJitter = HeterogeneousVolumes::ShouldJitter();
 
+		// Volume data
+		if (HeterogeneousVolumes::UseTransmittanceVolume() && bApplyShadowTransmittance)
+		{
+			PassParameters->TransmittanceVolume.TransmittanceVolumeResolution = HeterogeneousVolumes::GetTransmittanceVolumeResolution();
+			PassParameters->TransmittanceVolume.TransmittanceVolumeTexture = GraphBuilder.CreateSRV(TransmittanceVolumeTexture);
+		}
+
 		// Dispatch data
 		PassParameters->GroupCount = GroupCount;
 
@@ -199,44 +444,15 @@ void RenderSingleScatteringWithLiveShading(
 	{
 		FSceneRenderer::GetLightNameForDrawEvent(LightSceneInfo->Proxy, LightName);
 	}
+	FString PassName = FString::Printf(TEXT("RenderSingleScatteringWithLiveShadingCS (Light = %s)"), *LightName);
 
-	GraphBuilder.AddPass(
-		RDG_EVENT_NAME("RenderSingleScatteringWithLiveShadingCS (Light = %s)", *LightName),
-		PassParameters,
-		ERDGPassFlags::Compute,
-		// Why is scene explicitly copied??
-		[PassParameters, LocalScene = Scene, &View, MaterialRenderProxy, &Material, GroupCount](FRHIComputeCommandList& RHICmdList)
-		{
-			FRenderSingleScatteringWithLiveShadingCS::FPermutationDomain PermutationVector;
-			TShaderRef<FRenderSingleScatteringWithLiveShadingCS> ComputeShader = Material.GetShader<FRenderSingleScatteringWithLiveShadingCS>(&FLocalVertexFactory::StaticType, PermutationVector, false);
-
-			if (!ComputeShader.IsNull())
-			{
-				FMeshPassProcessorRenderState DrawRenderState;
-
-				FMeshMaterialShaderElementData ShaderElementData;
-				ShaderElementData.FadeUniformBuffer = GDistanceCullFadedInUniformBuffer.GetUniformBufferRHI();
-				ShaderElementData.DitherUniformBuffer = GDitherFadedInUniformBuffer.GetUniformBufferRHI();
-
-				FMeshProcessorShaders PassShaders;
-				PassShaders.ComputeShader = ComputeShader;
-
-				FMeshDrawShaderBindings ShaderBindings;
-				{
-					ShaderBindings.Initialize(PassShaders);
-
-					int32 DataOffset = 0;
-					FMeshDrawSingleShaderBindings SingleShaderBindings = ShaderBindings.GetSingleShaderBindings(SF_Compute, DataOffset);
-					ComputeShader->GetShaderBindings(LocalScene, LocalScene->GetFeatureLevel(), nullptr, *MaterialRenderProxy, Material, DrawRenderState, ShaderElementData, SingleShaderBindings);
-
-					ShaderBindings.Finalize(&PassShaders);
-				}
-				ShaderBindings.SetOnCommandList(RHICmdList, ComputeShader.GetComputeShader());
-
-				FComputeShaderUtils::Dispatch(RHICmdList, ComputeShader, *PassParameters, GroupCount);
-			}
-		}
-	);
+	FRenderSingleScatteringWithLiveShadingCS::FPermutationDomain PermutationVector;
+	PermutationVector.Set<FRenderSingleScatteringWithLiveShadingCS::FUseTransmittanceVolume>(HeterogeneousVolumes::UseTransmittanceVolume() && PassParameters->bApplyShadowTransmittance);
+	TShaderRef<FRenderSingleScatteringWithLiveShadingCS> ComputeShader = Material.GetShader<FRenderSingleScatteringWithLiveShadingCS>(&FLocalVertexFactory::StaticType, PermutationVector, false);
+	if (!ComputeShader.IsNull())
+	{
+		AddComputePass(GraphBuilder, ComputeShader, PassParameters, Scene, View, MaterialRenderProxy, Material, PassName, GroupCount);
+	}
 }
 
 void RenderWithLiveShading(
@@ -249,6 +465,8 @@ void RenderWithLiveShading(
 	const FMaterialRenderProxy* MaterialRenderProxy,
 	const int32 PrimitiveId,
 	const FBoxSphereBounds LocalBoxSphereBounds,
+	// Transmittance acceleration
+	FRDGTextureRef TransmittanceVolumeTexture,
 	// Output
 	FRDGTextureRef& HeterogeneousVolumeRadiance
 )
@@ -283,6 +501,30 @@ void RenderWithLiveShading(
 			bApplyShadowTransmittance = LightSceneInfo->Proxy->CastsVolumetricShadow();
 		}
 
+		if (HeterogeneousVolumes::UseTransmittanceVolume() && bApplyShadowTransmittance)
+		{
+			RenderTransmittanceVolumeWithLiveShading(
+				GraphBuilder,
+				// Scene data
+				Scene,
+				View,
+				SceneTextures,
+				// Light data
+				bApplyEmission,
+				bApplyDirectLighting,
+				bApplyShadowTransmittance,
+				LightType,
+				LightSceneInfo,
+				// Object data
+				PrimitiveSceneProxy,
+				MaterialRenderProxy,
+				PrimitiveId,
+				LocalBoxSphereBounds,
+				// Output
+				TransmittanceVolumeTexture
+			);
+		}
+
 		RenderSingleScatteringWithLiveShading(
 			GraphBuilder,
 			// Scene data
@@ -300,6 +542,8 @@ void RenderWithLiveShading(
 			MaterialRenderProxy,
 			PrimitiveId,
 			LocalBoxSphereBounds,
+			// Transmittance acceleration
+			TransmittanceVolumeTexture,
 			// Output
 			HeterogeneousVolumeRadiance
 		);
diff --git a/Engine/Source/Runtime/Renderer/Private/HeterogeneousVolumes/HeterogeneousVolumesPreshadingPipeline.cpp b/Engine/Source/Runtime/Renderer/Private/HeterogeneousVolumes/HeterogeneousVolumesPreshadingPipeline.cpp
--- a/Engine/Source/Runtime/Renderer/Private/HeterogeneousVolumes/HeterogeneousVolumesPreshadingPipeline.cpp
+++ b/Engine/Source/Runtime/Renderer/Private/HeterogeneousVolumes/HeterogeneousVolumesPreshadingPipeline.cpp
@@ -199,6 +199,68 @@ void GenerateRayMarchingTiles(
 	);
 }
 
+class FRenderTransmittanceVolumeWithPreshadingCS : public FGlobalShader
+{
+	DECLARE_GLOBAL_SHADER(FRenderTransmittanceVolumeWithPreshadingCS);
+	SHADER_USE_PARAMETER_STRUCT(FRenderTransmittanceVolumeWithPreshadingCS, FGlobalShader);
+
+	using FPermutationDomain = TShaderPermutationDomain<>;
+
+	BEGIN_SHADER_PARAMETER_STRUCT(FParameters, )
+		// Scene data
+		SHADER_PARAMETER_STRUCT_REF(FViewUniformShaderParameters, View)
+		SHADER_PARAMETER_STRUCT_INCLUDE(FSceneTextureParameters, SceneTextures)
+
+		// Light data
+		SHADER_PARAMETER(int, LightType)
+		SHADER_PARAMETER_STRUCT_REF(FDeferredLightUniformStruct, DeferredLight)
+
+		// Volume structures
+		SHADER_PARAMETER_RDG_UNIFORM_BUFFER(FSparseVoxelUniformBufferParameters, SparseVoxelUniformBuffer)
+		SHADER_PARAMETER_STRUCT_INCLUDE(FTransmittanceVolumeParameters, TransmittanceVolume)
+
+		// Ray data
+		SHADER_PARAMETER(float, MaxShadowTraceDistance)
+		SHADER_PARAMETER(float, StepSize)
+		SHADER_PARAMETER(int, MipLevel)
+		SHADER_PARAMETER(int, MaxStepCount)
+		SHADER_PARAMETER(int, bJitter)
+
+		// Output
+		SHADER_PARAMETER_RDG_TEXTURE_UAV(RWTexture3D<float>, RWTransmittanceVolumeTexture)
+	END_SHADER_PARAMETER_STRUCT()
+
+	static bool ShouldCompilePermutation(
+		const FGlobalShaderPermutationParameters& Parameters
+	)
+	{
+		return DoesPlatformSupportHeterogeneousVolumes(Parameters.Platform);
+	}
+
+	static void ModifyCompilationEnvironment(
+		const FGlobalShaderPermutationParameters& Parameters,
+		FShaderCompilerEnvironment& OutEnvironment
+	)
+	{
+		FMaterialShader::ModifyCompilationEnvironment(Parameters, OutEnvironment);
+		OutEnvironment.SetDefine(TEXT("THREADGROUP_SIZE_1D"), GetThreadGroupSize1D());
+		OutEnvironment.SetDefine(TEXT("THREADGROUP_SIZE_2D"), GetThreadGroupSize2D());
+		OutEnvironment.SetDefine(TEXT("THREADGROUP_SIZE_3D"), GetThreadGroupSize3D());
+
+		// This shader takes a very long time to compile with FXC, so we pre-compile it with DXC first and then forward the optimized HLSL to FXC.
+		//OutEnvironment.CompilerFlags.Add(CFLAG_PrecompileWithDXC);
+		OutEnvironment.CompilerFlags.Add(CFLAG_AllowTypedUAVLoads);
+
+		OutEnvironment.SetDefine(TEXT("GET_PRIMITIVE_DATA_OVERRIDE"), 1);
+	}
+
+	static int32 GetThreadGroupSize1D() { return GetThreadGroupSize2D() * GetThreadGroupSize2D(); }
+	static int32 GetThreadGroupSize2D() { return 8; }
+	static int32 GetThreadGroupSize3D() { return 4; }
+};
+
+IMPLEMENT_MATERIAL_SHADER_TYPE(, FRenderTransmittanceVolumeWithPreshadingCS, TEXT("/Engine/Private/HeterogeneousVolumes/HeterogeneousVolumesPreshadingPipeline.usf"), TEXT("RenderTransmittanceVolumeWithPreshadingCS"), SF_Compute);
+
 class FRenderSingleScatteringWithPreshadingCS : public FGlobalShader
 {
 	DECLARE_GLOBAL_SHADER(FRenderSingleScatteringWithPreshadingCS);
@@ -207,8 +269,9 @@ class FRenderSingleScatteringWithPreshadingCS : public FGlobalShader
 	class FApplyShadowTransmittanceDim : SHADER_PERMUTATION_BOOL("DIM_APPLY_SHADOW_TRANSMITTANCE");
 	class FVoxelCullingDim : SHADER_PERMUTATION_BOOL("DIM_VOXEL_CULLING");
 	class FSparseVoxelTracingDim : SHADER_PERMUTATION_BOOL("DIM_SPARSE_VOXEL_TRACING");
+	class FUseTransmittanceVolume : SHADER_PERMUTATION_BOOL("DIM_USE_TRANSMITTANCE_VOLUME");
 	class FDebugDim : SHADER_PERMUTATION_BOOL("DIM_DEBUG");
-	using FPermutationDomain = TShaderPermutationDomain<FApplyShadowTransmittanceDim, FVoxelCullingDim, FSparseVoxelTracingDim, FDebugDim>;
+	using FPermutationDomain = TShaderPermutationDomain<FApplyShadowTransmittanceDim, FVoxelCullingDim, FSparseVoxelTracingDim, FUseTransmittanceVolume, FDebugDim>;
 
 	BEGIN_SHADER_PARAMETER_STRUCT(FParameters, )
 		// Scene data
@@ -224,8 +287,9 @@ class FRenderSingleScatteringWithPreshadingCS : public FGlobalShader
 		// Volume data
 		SHADER_PARAMETER(int, MipLevel)
 
-		// Sparse Volume
+		// Volume structures
 		SHADER_PARAMETER_RDG_UNIFORM_BUFFER(FSparseVoxelUniformBufferParameters, SparseVoxelUniformBuffer)
+		SHADER_PARAMETER_STRUCT_INCLUDE(FTransmittanceVolumeParameters, TransmittanceVolume)
 
 		// Ray data
 		SHADER_PARAMETER(float, MaxTraceDistance)
@@ -274,11 +338,80 @@ class FRenderSingleScatteringWithPreshadingCS : public FGlobalShader
 
 IMPLEMENT_MATERIAL_SHADER_TYPE(, FRenderSingleScatteringWithPreshadingCS, TEXT("/Engine/Private/HeterogeneousVolumes/HeterogeneousVolumesPreshadingPipeline.usf"), TEXT("RenderSingleScatteringWithPreshadingCS"), SF_Compute);
 
-void RenderTransmissionWithPreshadingCompute(
-
+void RenderTransmitanceVolumeWithPreshadingCompute(
+	FRDGBuilder& GraphBuilder,
+	// Scene data
+	const FScene* Scene,
+	const FViewInfo& View,
+	const FSceneTextures& SceneTextures,
+	// Light data
+	bool bApplyEmission,
+	bool bApplyDirectLighting,
+	bool bApplyShadowTransmittance,
+	uint32 LightType,
+	const FLightSceneInfo* LightSceneInfo,
+	// Object data
+	const FPrimitiveSceneProxy* PrimitiveSceneProxy,
+	// Sparse voxel data
+	FRDGBufferRef NumVoxelsBuffer,
+	const TRDGUniformBufferRef<FSparseVoxelUniformBufferParameters>& SparseVoxelUniformBuffer,
+	// Ray marching tiles
+	FRDGBufferRef NumRayMarchingTilesBuffer,
+	FRDGBufferRef RayMarchingTilesBuffer,
+	FRDGBufferRef VoxelsPerTileBuffer,
+	// Output
+	FRDGTextureRef& TransmittanceVolumeTexture
 )
 {
+	FRenderTransmittanceVolumeWithPreshadingCS::FParameters* PassParameters = GraphBuilder.AllocParameters<FRenderTransmittanceVolumeWithPreshadingCS::FParameters>();
+	{
+		// Scene data
+		PassParameters->View = View.ViewUniformBuffer;
+		PassParameters->SceneTextures = GetSceneTextureParameters(GraphBuilder, SceneTextures);
 
+		// Light data
+		check(LightSceneInfo != nullptr);
+		FDeferredLightUniformStruct DeferredLightUniform = GetDeferredLightParameters(View, *LightSceneInfo);
+		PassParameters->DeferredLight = CreateUniformBufferImmediate(DeferredLightUniform, UniformBuffer_SingleDraw);
+		PassParameters->LightType = LightType;
+
+		// Sparse voxel data
+		PassParameters->SparseVoxelUniformBuffer = SparseVoxelUniformBuffer;
+
+		// Transmittance volume
+		PassParameters->TransmittanceVolume.TransmittanceVolumeResolution = HeterogeneousVolumes::GetTransmittanceVolumeResolution();
+		PassParameters->TransmittanceVolume.TransmittanceVolumeTexture = GraphBuilder.CreateSRV(TransmittanceVolumeTexture);
+
+		// Ray data
+		//PassParameters->StepSize = HeterogeneousVolumes::GetStepSize();
+		PassParameters->MaxStepCount = HeterogeneousVolumes::GetMaxStepCount();
+		PassParameters->bJitter = HeterogeneousVolumes::ShouldJitter();
+		PassParameters->MipLevel = HeterogeneousVolumes::GetMipLevel();
+
+		// Output
+		PassParameters->RWTransmittanceVolumeTexture = GraphBuilder.CreateUAV(TransmittanceVolumeTexture);
+	}
+
+	FString LightName = TEXT("none");
+	if (LightSceneInfo != nullptr)
+	{
+		FSceneRenderer::GetLightNameForDrawEvent(LightSceneInfo->Proxy, LightName);
+	}
+
+	FRenderTransmittanceVolumeWithPreshadingCS::FPermutationDomain PermutationVector;
+	TShaderRef<FRenderTransmittanceVolumeWithPreshadingCS> ComputeShader = View.ShaderMap->GetShader<FRenderTransmittanceVolumeWithPreshadingCS>(PermutationVector);
+
+	FIntVector GroupCount = HeterogeneousVolumes::GetTransmittanceVolumeResolution();
+	GroupCount.X = FMath::DivideAndRoundUp(GroupCount.X, FRenderTransmittanceVolumeWithPreshadingCS::GetThreadGroupSize3D());
+	GroupCount.Y = FMath::DivideAndRoundUp(GroupCount.Y, FRenderTransmittanceVolumeWithPreshadingCS::GetThreadGroupSize3D());
+	GroupCount.Z = FMath::DivideAndRoundUp(GroupCount.Z, FRenderTransmittanceVolumeWithPreshadingCS::GetThreadGroupSize3D());
+
+	FComputeShaderUtils::AddPass(
+		GraphBuilder,
+		RDG_EVENT_NAME("RenderTransmittanceVolumeWithPreshadingCS (Light = %s)", *LightName),
+		ComputeShader,
+		PassParameters,
+		GroupCount);
 }
 
 void RenderSingleScatteringWithPreshadingCompute(
@@ -298,6 +431,7 @@ void RenderSingleScatteringWithPreshadingCompute(
 	// Sparse voxel data
 	FRDGBufferRef NumVoxelsBuffer,
 	const TRDGUniformBufferRef<FSparseVoxelUniformBufferParameters>& SparseVoxelUniformBuffer,
+	FRDGTextureRef TransmittanceVolumeTexture,
 	// Ray marching tiles
 	FRDGBufferRef NumRayMarchingTilesBuffer,
 	FRDGBufferRef RayMarchingTilesBuffer,
@@ -334,6 +468,13 @@ void RenderSingleScatteringWithPreshadingCompute(
 		// Sparse voxel data
 		PassParameters->SparseVoxelUniformBuffer = SparseVoxelUniformBuffer;
 
+		// Transmittance volume
+		if (HeterogeneousVolumes::UseTransmittanceVolume() && bApplyShadowTransmittance)
+		{
+			PassParameters->TransmittanceVolume.TransmittanceVolumeResolution = HeterogeneousVolumes::GetTransmittanceVolumeResolution();
+			PassParameters->TransmittanceVolume.TransmittanceVolumeTexture = GraphBuilder.CreateSRV(TransmittanceVolumeTexture);
+		}
+
 		// Ray data
 		PassParameters->MaxTraceDistance = HeterogeneousVolumes::GetMaxTraceDistance();
 		PassParameters->StepSize = HeterogeneousVolumes::GetStepSize();
@@ -362,6 +503,7 @@ void RenderSingleScatteringWithPreshadingCompute(
 	PermutationVector.Set<FRenderSingleScatteringWithPreshadingCS::FApplyShadowTransmittanceDim>(bApplyShadowTransmittance);
 	PermutationVector.Set<FRenderSingleScatteringWithPreshadingCS::FVoxelCullingDim>(HeterogeneousVolumes::UseSparseVoxelPerTileCulling());
 	PermutationVector.Set<FRenderSingleScatteringWithPreshadingCS::FSparseVoxelTracingDim>(HeterogeneousVolumes::UseSparseVoxelPipeline());
+	PermutationVector.Set<FRenderSingleScatteringWithPreshadingCS::FUseTransmittanceVolume>(HeterogeneousVolumes::UseTransmittanceVolume());
 	PermutationVector.Set<FRenderSingleScatteringWithPreshadingCS::FDebugDim>(HeterogeneousVolumes::GetDebugMode() != 0);
 	TShaderRef<FRenderSingleScatteringWithPreshadingCS> ComputeShader = View.ShaderMap->GetShader<FRenderSingleScatteringWithPreshadingCS>(PermutationVector);
 	FComputeShaderUtils::AddPass(
@@ -385,8 +527,9 @@ void RenderWithPreshadingCompute(
 	const FMaterialRenderProxy* MaterialRenderProxy,
 	// Sparse voxel data
 	FRDGBufferRef NumVoxelsBuffer,
-	const TRDGUniformBufferRef<FSparseVoxelUniformBufferParameters>& SparseVoxelUniformBuffer,
+	TRDGUniformBufferRef<FSparseVoxelUniformBufferParameters>& SparseVoxelUniformBuffer,
 	// Output
+	FRDGTextureRef& TransmittanceVolumeTexture,
 	FRDGTextureRef& HeterogeneousVolumeRadiance
 )
 {
@@ -450,6 +593,35 @@ void RenderWithPreshadingCompute(
 				bApplyShadowTransmittance = LightSceneInfo->Proxy->CastsVolumetricShadow();
 			}
 
+			if (HeterogeneousVolumes::UseTransmittanceVolume() && bApplyShadowTransmittance)
+			{
+				RenderTransmitanceVolumeWithPreshadingCompute(
+					GraphBuilder,
+					// Scene
+					Scene,
+					View,
+					SceneTextures,
+					// Light
+					bApplyEmission,
+					bApplyDirectLighting,
+					bApplyShadowTransmittance,
+					LightType,
+					LightSceneInfo,
+					// Object
+					PrimitiveSceneProxy,
+					// Volume data
+					NumVoxelsBuffer,
+					// Sparse voxel data
+					SparseVoxelUniformBuffer,
+					// Ray marching tile
+					NumRayMarchingTilesBuffer,
+					RayMarchingTilesBuffer,
+					VoxelsPerTileBuffer,
+					// Output
+					TransmittanceVolumeTexture
+				);
+			}
+
 			RenderSingleScatteringWithPreshadingCompute(
 				GraphBuilder,
 				// Scene
@@ -468,6 +640,7 @@ void RenderWithPreshadingCompute(
 				NumVoxelsBuffer,
 				// Sparse voxel data
 				SparseVoxelUniformBuffer,
+				TransmittanceVolumeTexture,
 				// Ray marching tile
 				NumRayMarchingTilesBuffer,
 				RayMarchingTilesBuffer,
@@ -492,6 +665,8 @@ void RenderWithPreshadingHardwareRayTracing(
 	// Sparse voxel data
 	FRDGBufferRef NumVoxelsBuffer,
 	const TRDGUniformBufferRef<FSparseVoxelUniformBufferParameters>& SparseVoxelUniformBuffer,
+	// Transmittance acceleration
+	FRDGTextureRef TransmittanceVolumeTexture,
 	// Output
 	FRDGTextureRef& HeterogeneousVolumeRadiance
 )
@@ -569,6 +744,31 @@ void RenderWithPreshadingHardwareRayTracing(
 				bApplyShadowTransmittance = LightSceneInfo->Proxy->CastsVolumetricShadow();
 			}
 
+			if (HeterogeneousVolumes::UseTransmittanceVolume() && bApplyShadowTransmittance)
+			{
+				RenderTransmittanceVolumeWithPreshadingHardwareRayTracing(
+					GraphBuilder,
+					// Scene data
+					Scene,
+					View,
+					SceneTextures,
+					// Light data
+					bApplyEmission,
+					bApplyDirectLighting,
+					bApplyShadowTransmittance,
+					LightType,
+					LightSceneInfo,
+					// Object data
+					PrimitiveSceneProxy,
+					// Sparse voxel
+					SparseVoxelUniformBuffer,
+					// Ray tracing data
+					Scene->HeterogeneousVolumesRayTracingScene,
+					// Transmittance volume
+					TransmittanceVolumeTexture
+				);
+			}
+
 			RenderSingleScatteringWithPreshadingHardwareRayTracing(
 				GraphBuilder,
 				// Scene data
@@ -587,6 +787,8 @@ void RenderWithPreshadingHardwareRayTracing(
 				SparseVoxelUniformBuffer,
 				// Ray tracing data
 				Scene->HeterogeneousVolumesRayTracingScene,
+				// Transmittance volume
+				TransmittanceVolumeTexture,
 				// Output
 				HeterogeneousVolumeRadiance
 			);
@@ -727,6 +929,8 @@ void RenderWithPreshading(
 	const FMaterialRenderProxy* MaterialRenderProxy,
 	const int32 PrimitiveId,
 	const FBoxSphereBounds LocalBoxSphereBounds,
+	// Transmittance acceleration
+	FRDGTextureRef TransmittanceVolumeTexture,
 	// Output
 	FRDGTextureRef& HeterogeneousVolumeRadiance
 )
@@ -744,9 +948,9 @@ void RenderWithPreshading(
 		NumMips
 	);
 
-	FRDGTextureRef HeterogeneousVolumesExtinctionTexture = GraphBuilder.CreateTexture(BakedMaterialDesc, TEXT("HeterogeneousVolumes.ExtinctionTexture"));
-	FRDGTextureRef HeterogeneousVolumesEmissionTexture = GraphBuilder.CreateTexture(BakedMaterialDesc, TEXT("HeterogeneousVolumes.EmissionTexture"));
-	FRDGTextureRef HeterogeneousVolumesAlbedoTexture = GraphBuilder.CreateTexture(BakedMaterialDesc, TEXT("HeterogeneousVolumes.AlbedoTexture"));
+	FRDGTextureRef ExtinctionTexture = GraphBuilder.CreateTexture(BakedMaterialDesc, TEXT("HeterogeneousVolumes.ExtinctionTexture"));
+	FRDGTextureRef EmissionTexture = GraphBuilder.CreateTexture(BakedMaterialDesc, TEXT("HeterogeneousVolumes.EmissionTexture"));
+	FRDGTextureRef AlbedoTexture = GraphBuilder.CreateTexture(BakedMaterialDesc, TEXT("HeterogeneousVolumes.AlbedoTexture"));
 
 	// Preshading pipeline
 	{
@@ -768,9 +972,9 @@ void RenderWithPreshading(
 				// Volume data
 				VolumeResolution,
 				// Output
-				HeterogeneousVolumesExtinctionTexture,
-				HeterogeneousVolumesEmissionTexture,
-				HeterogeneousVolumesAlbedoTexture
+				ExtinctionTexture,
+				EmissionTexture,
+				AlbedoTexture
 			);
 		}
 
@@ -779,10 +983,10 @@ void RenderWithPreshading(
 			RDG_EVENT_SCOPE(GraphBuilder, "MIP Generation");
 			for (uint32 MipLevel = 1; MipLevel < NumMips; ++MipLevel)
 			{
-				GenerateMips3D<FGenerateMips3D>(GraphBuilder, View, HeterogeneousVolumesExtinctionTexture, MipLevel);
+				GenerateMips3D<FGenerateMips3D>(GraphBuilder, View, ExtinctionTexture, MipLevel);
 				// TODO: Reinstate once ray-marching determines appropriate MIP level to sample
-				//GenerateMips3D<FGenerateMips3D>(GraphBuilder, View, HeterogeneousVolumesEmissionTexture, MipLevel);
-				//GenerateMips3D<FGenerateMips3D>(GraphBuilder, View, HeterogeneousVolumesAlbedoTexture, MipLevel);
+				//GenerateMips3D<FGenerateMips3D>(GraphBuilder, View, EmissionTexture, MipLevel);
+				//GenerateMips3D<FGenerateMips3D>(GraphBuilder, View, AlbedoTexture, MipLevel);
 			}
 		}
 	}
@@ -800,14 +1004,14 @@ void RenderWithPreshading(
 		{
 			RDG_EVENT_SCOPE(GraphBuilder, "Min Generation");
 
-			FRDGTextureDesc MinTextureDesc = HeterogeneousVolumesExtinctionTexture->Desc;
+			FRDGTextureDesc MinTextureDesc = ExtinctionTexture->Desc;
 			MinTextureDesc.Extent.X = FMath::Max(MinTextureDesc.Extent.X >> SparseMipLevel, 1);
 			MinTextureDesc.Extent.Y = FMath::Max(MinTextureDesc.Extent.Y >> SparseMipLevel, 1);
 			MinTextureDesc.Depth = FMath::Max(MinTextureDesc.Depth >> SparseMipLevel, 1);
 			MinTextureDesc.NumMips = FMath::Log2(float(FMath::Min(FMath::Min(MinTextureDesc.Extent.X, MinTextureDesc.Extent.Y), (int32)MinTextureDesc.Depth))) + 1;
 
 			MinTexture = GraphBuilder.CreateTexture(MinTextureDesc, TEXT("HeterogeneousVolumes.MinTexture"));
-			CopyTexture3D(GraphBuilder, View, HeterogeneousVolumesExtinctionTexture, SparseMipLevel, MinTexture);
+			CopyTexture3D(GraphBuilder, View, ExtinctionTexture, SparseMipLevel, MinTexture);
 			for (uint32 MipLevel = 1; MipLevel < MinTextureDesc.NumMips; ++MipLevel)
 			{
 				GenerateMips3D<FGenerateMin3D>(GraphBuilder, View, MinTexture, MipLevel);
@@ -832,10 +1036,10 @@ void RenderWithPreshading(
 
 		// Volume data
 		SparseVoxelUniformBufferParameters->VolumeResolution = VolumeResolution;
-		SparseVoxelUniformBufferParameters->ExtinctionTexture = HeterogeneousVolumesExtinctionTexture;
-		SparseVoxelUniformBufferParameters->EmissionTexture = HeterogeneousVolumesEmissionTexture;
-		SparseVoxelUniformBufferParameters->AlbedoTexture = HeterogeneousVolumesAlbedoTexture;
-		SparseVoxelUniformBufferParameters->TextureSampler = TStaticSamplerState<SF_Trilinear, AM_Clamp, AM_Clamp, AM_Clamp>::GetRHI();;
+		SparseVoxelUniformBufferParameters->ExtinctionTexture = ExtinctionTexture;
+		SparseVoxelUniformBufferParameters->EmissionTexture = EmissionTexture;
+		SparseVoxelUniformBufferParameters->AlbedoTexture = AlbedoTexture;
+		SparseVoxelUniformBufferParameters->TextureSampler = TStaticSamplerState<SF_Trilinear, AM_Clamp, AM_Clamp, AM_Clamp>::GetRHI();
 
 		// Sparse voxel data
 		SparseVoxelUniformBufferParameters->NumVoxelsBuffer = GraphBuilder.CreateSRV(NumVoxelsBuffer, PF_R32_UINT);
@@ -844,6 +1048,7 @@ void RenderWithPreshading(
 
 		// Traversal hints
 		SparseVoxelUniformBufferParameters->MaxTraceDistance = HeterogeneousVolumes::GetMaxTraceDistance();
+		SparseVoxelUniformBufferParameters->MaxShadowTraceDistance = HeterogeneousVolumes::GetMaxShadowTraceDistance();
 		SparseVoxelUniformBufferParameters->StepSize = HeterogeneousVolumes::GetStepSize();
 		SparseVoxelUniformBufferParameters->ShadowStepFactor = HeterogeneousVolumes::GetShadowStepFactor();
 	}
@@ -865,6 +1070,8 @@ void RenderWithPreshading(
 			// Sparse voxel data
 			NumVoxelsBuffer,
 			SparseVoxelUniformBuffer,
+			// Transmittance acceleration
+			TransmittanceVolumeTexture,
 			// Output
 			HeterogeneousVolumeRadiance
 		);
@@ -885,6 +1092,8 @@ void RenderWithPreshading(
 			// Sparse voxel data
 			NumVoxelsBuffer,
 			SparseVoxelUniformBuffer,
+			// Transmittance acceleration
+			TransmittanceVolumeTexture,
 			// Output
 			HeterogeneousVolumeRadiance
 		);
