Path Tracing:  fix for profiling with multiple GPUs.

Scoped the root timing block of the scene renderer to all GPUs when the path tracer is active, so child profiling blocks on any GPU can show up (the GPU profiler discards timing intervals where any event in the hierarchy isn't active on a given GPU).  Also subdivided the path tracer into more timing blocks, to solve the same issue within the path tracer, and allow multi-GPU work, cross GPU copy, and single-GPU post processing to be differentiated.

#rnx

diff --git a/Engine/Source/Runtime/Renderer/Private/DeferredShadingRenderer.cpp b/Engine/Source/Runtime/Renderer/Private/DeferredShadingRenderer.cpp
--- a/Engine/Source/Runtime/Renderer/Private/DeferredShadingRenderer.cpp
+++ b/Engine/Source/Runtime/Renderer/Private/DeferredShadingRenderer.cpp
@@ -2120,7 +2120,7 @@ void FDeferredShadingSceneRenderer::Render(FRDGBuilder& GraphBuilder)
 #endif // WITH_MGPU
 
 	// By default, limit our GPU usage to only GPUs specified in the view masks.
-	RDG_GPU_MASK_SCOPE(GraphBuilder, AllViewsGPUMask);
+	RDG_GPU_MASK_SCOPE(GraphBuilder, ViewFamily.EngineShowFlags.PathTracing ? FRHIGPUMask::All() : AllViewsGPUMask);
 
 	WaitOcclusionTests(GraphBuilder.RHICmdList);
 
@@ -2150,6 +2150,10 @@ void FDeferredShadingSceneRenderer::Render(FRDGBuilder& GraphBuilder)
 	FSceneTexturesConfig& SceneTexturesConfig = GetActiveSceneTexturesConfig();
 	FSceneTexturesConfig::Set(SceneTexturesConfig);
 
+	FPlatformMisc::LowLevelOutputDebugStringf(TEXT("SceneRender\t%d\t%dx%d\n"), GFrameNumberRenderThread,
+		Views[0].UnscaledViewRect.Width(),
+		Views[0].UnscaledViewRect.Height());
+
 	const FRDGSystemTextures& SystemTextures = FRDGSystemTextures::Create(GraphBuilder);
 
 	const bool bHasRayTracedOverlay = HasRayTracedOverlay(ViewFamily);
diff --git a/Engine/Source/Runtime/Renderer/Private/PathTracing.cpp b/Engine/Source/Runtime/Renderer/Private/PathTracing.cpp
--- a/Engine/Source/Runtime/Renderer/Private/PathTracing.cpp
+++ b/Engine/Source/Runtime/Renderer/Private/PathTracing.cpp
@@ -1941,15 +1941,26 @@ END_SHADER_PARAMETER_STRUCT()
 #endif
 
 DECLARE_GPU_STAT_NAMED(Stat_GPU_PathTracing, TEXT("Path Tracing"));
+DECLARE_GPU_STAT_NAMED(Stat_GPU_PathTracingPost, TEXT("Path Tracing Post"));
+#if WITH_MGPU
+DECLARE_GPU_STAT_NAMED(Stat_GPU_PathTracingCopy, TEXT("Path Tracing Copy"));
+#endif
+
 void FDeferredShadingSceneRenderer::RenderPathTracing(
 	FRDGBuilder& GraphBuilder,
 	const FViewInfo& View,
 	TRDGUniformBufferRef<FSceneTextureUniformParameters> SceneTexturesUniformBuffer,
 	FRDGTextureRef SceneColorOutputTexture)
 {
-	RDG_GPU_STAT_SCOPE(GraphBuilder, Stat_GPU_PathTracing);
 	RDG_EVENT_SCOPE(GraphBuilder, "Path Tracing");
 
+	// To make the GPU profiler work for path tracing with multi-GPU, we need the root GPU profiling scope (marked as "Unaccounted") to be on all GPUs,
+	// as the profiler discards events where any event in the hierarchy wasn't on a given GPU.  So in the parent scene render code, we set the GPU mask
+	// to "All" when path tracing is enabled, instead of "AllViewsGPUMask".  Then we'll enable that scope inside the path tracer instead.  We also
+	// subdivide the profiling scopes inside the path tracer, so the multi-GPU rendering and single-GPU post processing are separate scopes, instead of
+	// a scope for the whole path tracer (which would create the same problem).
+	RDG_GPU_MASK_SCOPE(GraphBuilder, AllViewsGPUMask);
+
 	if (!ensureMsgf(FDataDrivenShaderPlatformInfo::GetSupportsPathTracing(View.GetShaderPlatform()),
 		TEXT("Attempting to use path tracing on unsupported platform.")))
 	{
@@ -2178,6 +2189,7 @@ void FDeferredShadingSceneRenderer::RenderPathTracing(
 			RDG_GPU_MASK_SCOPE(GraphBuilder, FRHIGPUMask::FromIndex(GPUIndex));
 #if WITH_MGPU
 			RDG_EVENT_SCOPE(GraphBuilder, "Path Tracing GPU%d", GPUIndex);
+			RDG_GPU_STAT_SCOPE(GraphBuilder, Stat_GPU_PathTracing);
 #endif
 			for (int32 TileY = 0; TileY < DispatchResY; TileY += DispatchSize)
 			{
@@ -2339,6 +2351,12 @@ void FDeferredShadingSceneRenderer::RenderPathTracing(
 #if WITH_MGPU
 		if (NumGPUs > 1)
 		{
+			// Treat the cross GPU copy as occurring on all GPUs, for profiling purposes.  Internally, the cross GPU transfer doesn't
+			// pay attention to the mask, so it has no effect on behavior.  Technically the work of the copy is done on the second GPU,
+			// and the first GPU stalls waiting on that, so it's useful to show this interval on both GPUs.
+			RDG_GPU_MASK_SCOPE(GraphBuilder, GPUMask);
+			RDG_GPU_STAT_SCOPE(GraphBuilder, Stat_GPU_PathTracingCopy);
+
 			FMGPUTransferParameters* Parameters = GraphBuilder.AllocParameters<FMGPUTransferParameters>();
 			Parameters->InputTexture = RadianceTexture;
 			Parameters->InputAlbedo = AlbedoTexture;
@@ -2397,6 +2415,9 @@ void FDeferredShadingSceneRenderer::RenderPathTracing(
 		++PathTracingState->FrameIndex;
 	}
 
+	RDG_GPU_MASK_SCOPE(GraphBuilder, View.GPUMask);
+	RDG_GPU_STAT_SCOPE(GraphBuilder, Stat_GPU_PathTracingPost);
+
 	// Figure out if the denoiser is enabled and needs to run
 	FRDGTexture* DenoisedRadianceTexture = nullptr;
 	bool IsDenoiserEnabled = IsPathTracingDenoiserEnabled(View);
