Fix potential crash when RT shadows are disabled for the project but there is still a light in the scene that enables RT shadows

diff --git a/Engine/Source/Runtime/Renderer/Private/DeferredShadingRenderer.cpp b/Engine/Source/Runtime/Renderer/Private/DeferredShadingRenderer.cpp
--- a/Engine/Source/Runtime/Renderer/Private/DeferredShadingRenderer.cpp
+++ b/Engine/Source/Runtime/Renderer/Private/DeferredShadingRenderer.cpp
@@ -2058,8 +2058,6 @@ void FDeferredShadingSceneRenderer::Render(FRDGBuilder& GraphBuilder)
 			Scene->CachedRayTracingMeshCommandsMode = CurrentMode;
 			Scene->RefreshRayTracingMeshCommandCache();
 		}
-
-		Scene->UpdateRayTracedLights();
 	}
 #endif
 
@@ -3676,6 +3674,9 @@ bool HasRayTracedOverlay(const FSceneViewFamily& ViewFamily)
 
 void FDeferredShadingSceneRenderer::InitializeRayTracingFlags_RenderThread()
 {
+	// The result of this call is used by AnyRayTracingPassEnabled to decide if we have any RT shadows enabled
+	Scene->UpdateRayTracedLights();
+
 	// This function may be called twice -- once in CreateSceneRenderers and again in Render.  We deliberately skip the logic
 	// if the flag is already set, because CreateSceneRenderers fills in the correct value for "bShouldUpdateRayTracingScene"
 	// in that case, and we don't want to overwrite it.
diff --git a/Engine/Source/Runtime/Renderer/Private/PrimitiveSceneInfo.cpp b/Engine/Source/Runtime/Renderer/Private/PrimitiveSceneInfo.cpp
--- a/Engine/Source/Runtime/Renderer/Private/PrimitiveSceneInfo.cpp
+++ b/Engine/Source/Runtime/Renderer/Private/PrimitiveSceneInfo.cpp
@@ -704,9 +704,8 @@ void FScene::UpdateRayTracedLights()
 	}
 
 	// We currently don't need a full list of RT lights, only whether there are any RT lights at all.
-	for (auto LightIt = Lights.CreateConstIterator(); LightIt; ++LightIt)
+	for (const FLightSceneInfoCompact& LightSceneInfoCompact : Lights)
 	{
-		const FLightSceneInfoCompact& LightSceneInfoCompact = *LightIt;
 		if (ShouldRenderRayTracingShadowsForLight(LightSceneInfoCompact))
 		{
 			bHasRayTracedLights = true;
