Remove debug printf spam...

#rnx

diff --git a/Engine/Source/Runtime/Renderer/Private/DeferredShadingRenderer.cpp b/Engine/Source/Runtime/Renderer/Private/DeferredShadingRenderer.cpp
--- a/Engine/Source/Runtime/Renderer/Private/DeferredShadingRenderer.cpp
+++ b/Engine/Source/Runtime/Renderer/Private/DeferredShadingRenderer.cpp
@@ -2148,10 +2148,6 @@ void FDeferredShadingSceneRenderer::Render(FRDGBuilder& GraphBuilder)
 	FSceneTexturesConfig& SceneTexturesConfig = GetActiveSceneTexturesConfig();
 	FSceneTexturesConfig::Set(SceneTexturesConfig);
 
-	FPlatformMisc::LowLevelOutputDebugStringf(TEXT("SceneRender\t%d\t%dx%d\n"), GFrameNumberRenderThread,
-		Views[0].UnscaledViewRect.Width(),
-		Views[0].UnscaledViewRect.Height());
-
 	const FRDGSystemTextures& SystemTextures = FRDGSystemTextures::Create(GraphBuilder);
 
 	const bool bHasRayTracedOverlay = HasRayTracedOverlay(ViewFamily);
