Automatic categorization of primitives into static/dynamic VSM caching based on invalidation requests:
- Store cache flag in primitive scene proxy, default based on mobility initially
- When CPU invalidations of primitives that are cached as static happen, invalidate the entire primitive (all instances) and flip the flag to dynamic
- When GPU invalidations of primitives that are cached as static happen, mark a bit in a buffer that is read back to the CPU. In future frames, use the bits to find the associated primitive (if it still exists), swap it to dynamic and invalidate all of its instances
- Dynamic primitives are always cached as dynamic/movable
- Moved instance invalidation buffer ownership to cache manager (we don't need two copies at once)

diff --git a/Engine/Shaders/Private/VirtualShadowMaps/VirtualShadowMapCacheManagement.usf b/Engine/Shaders/Private/VirtualShadowMaps/VirtualShadowMapCacheManagement.usf
--- a/Engine/Shaders/Private/VirtualShadowMaps/VirtualShadowMapCacheManagement.usf
+++ b/Engine/Shaders/Private/VirtualShadowMaps/VirtualShadowMapCacheManagement.usf
@@ -31,7 +31,23 @@ uint bDrawBounds;
 #if INPUT_KIND == INPUT_KIND_GPU_INSTANCES
 StructuredBuffer<uint> InvalidatingInstances;
 uint NumInvalidatingInstanceSlots;
-#endif
+
+// Flags buffer (one bit per PersistentPrimitiveIndex) to track statically cached primitives that are invalidating
+// Currently we use GPU messaging to send these primitives back for the CPU to handle as desired
+// in future frames, but this flags buffer ensures we only send one message per prim.
+RWStructuredBuffer<uint> OutStaticInvalidatingPrimitives;
+
+void RecordStaticPrimitiveInvalidation(uint PersistentPrimitiveIndex)
+{
+	{
+		uint PrimitiveIdWordOffset = PersistentPrimitiveIndex / 32U;
+		uint PrimitiveIdWordMask = 1U << (PersistentPrimitiveIndex % 32U);
+		
+		uint PreviousWordMask = 0U;
+		InterlockedOr(OutStaticInvalidatingPrimitives[PrimitiveIdWordOffset], PrimitiveIdWordMask, PreviousWordMask);
+	}
+}
+#endif		// INPUT_KIND == INPUT_KIND_GPU_INSTANCES
 
 struct FInstanceInvalidationPayload
 {
@@ -111,9 +127,8 @@ void VirtualSmInvalidateInstancePagesCS(
 		}
 
 
-		// TODO: Clean up hardcoded flag field test.
 		bool bCastShadows = InstanceSceneData.ValidInstance 
-			&& (GetPrimitiveData(InstanceSceneData.PrimitiveId).Flags & 1U) != 0U;
+			&& (GetPrimitiveData(InstanceSceneData.PrimitiveId).Flags & PRIMITIVE_SCENE_DATA_FLAG_CAST_SHADOWS) != 0U;
 #if ENABLE_DEBUG_MODE
 		uint PageInvalidationCount = 0U;
 #endif
@@ -121,6 +136,8 @@ void VirtualSmInvalidateInstancePagesCS(
 		// TODO: test the flag on the instance instead once it is updated correctly InstanceSceneData.CastShadows
 		if (bCastShadows)
 		{
+			bool bAnyStaticPagesInvalidated = false;
+
 			for (uint VirtualShadowMapId = FirstVirtualShadowMapId; VirtualShadowMapId < EndVirtualShadowMapId; ++VirtualShadowMapId)
 			{
 				// 1. Load cached projection data
@@ -146,9 +163,9 @@ void VirtualSmInvalidateInstancePagesCS(
 					}
 				}
 
-				// TODO: This is slightly messy, but should be ok for the moment
-				bool bInvalidateStaticPage = bForceInvalidateStaticPage ||  ShouldCacheInstanceAsStatic(InstanceSceneData, ProjectionData.bUnCached);
+				bool bInvalidateStaticPage = bForceInvalidateStaticPage || ShouldCacheInstanceAsStatic(InstanceSceneData, ProjectionData.bUnCached);
 				uint InvalidationFlags = bInvalidateStaticPage ? VSM_STATIC_UNCACHED_FLAG : VSM_DYNAMIC_UNCACHED_FLAG;
+				bAnyStaticPagesInvalidated = bAnyStaticPagesInvalidated || bInvalidateStaticPage;
 
 				// Go back to clip space
 				float4x4 UVToClip;
@@ -222,6 +239,15 @@ void VirtualSmInvalidateInstancePagesCS(
 					}
 				}
 			}
+
+#if INPUT_KIND == INPUT_KIND_GPU_INSTANCES
+			// For GPU instance invalidations of static pages, we signal this back to the CPU in case it wants to
+			// transition this primitive to dynamic (i.e. due to WPO, etc).
+			if (bAnyStaticPagesInvalidated)
+			{
+				RecordStaticPrimitiveInvalidation(GetPrimitiveData(InstanceSceneData.PrimitiveId).PersistentPrimitiveIndex);
+			}
+#endif
 		}
 #if ENABLE_DEBUG_MODE
 		if (bDrawBounds && PageInvalidationCount > 0U)
diff --git a/Engine/Shaders/Private/VirtualShadowMaps/VirtualShadowMapPrintStats.usf b/Engine/Shaders/Private/VirtualShadowMaps/VirtualShadowMapPrintStats.usf
--- a/Engine/Shaders/Private/VirtualShadowMaps/VirtualShadowMapPrintStats.usf
+++ b/Engine/Shaders/Private/VirtualShadowMaps/VirtualShadowMapPrintStats.usf
@@ -59,8 +59,9 @@ void PrintUnits(inout FShaderPrintContext Context, uint Value)
 void PrintStats()
 {
 	const bool bShowAll                = ShowStatsValue == 1;
-	const bool bShowPhysicalPages      = bShowAll || ShowStatsValue == 2;
-	const bool bShowNonNaniteInstances = bShowAll || ShowStatsValue == 3;
+	// Work around potential compiler bug...
+	const bool bShowPhysicalPages      = (ShowStatsValue == 1) || ShowStatsValue == 2;
+	const bool bShowNonNaniteInstances = (ShowStatsValue == 1) || ShowStatsValue == 3;
 
 	float TopMargin = 0.05f;
 	float HeadlineX = 0.47f;
diff --git a/Engine/Source/Runtime/Engine/Private/PrimitiveSceneProxy.cpp b/Engine/Source/Runtime/Engine/Private/PrimitiveSceneProxy.cpp
--- a/Engine/Source/Runtime/Engine/Private/PrimitiveSceneProxy.cpp
+++ b/Engine/Source/Runtime/Engine/Private/PrimitiveSceneProxy.cpp
@@ -494,7 +494,7 @@ void FPrimitiveSceneProxy::UpdateUniformBuffer()
 				.PreSkinnedLocalBounds(PreSkinnedLocalBounds)
 				//.ObjectDrawDistance(FVector2f(GetMinDrawDistance(), GetMaxDrawDistance())) // TODO
 				.ReceivesDecals(bReceivesDecals)
-				.ShouldCacheShadow(ShouldCacheShadow())
+				.CacheShadowAsStatic(PrimitiveSceneInfo->ShouldCacheShadowAsStatic())
 				.OutputVelocity(bOutputVelocity)
 				.EvaluateWorldPositionOffset(EvaluateWorldPositionOffset())
 				.LightingChannelMask(GetLightingChannelMask())
diff --git a/Engine/Source/Runtime/Engine/Private/PrimitiveUniformShaderParameters.cpp b/Engine/Source/Runtime/Engine/Private/PrimitiveUniformShaderParameters.cpp
--- a/Engine/Source/Runtime/Engine/Private/PrimitiveUniformShaderParameters.cpp
+++ b/Engine/Source/Runtime/Engine/Private/PrimitiveUniformShaderParameters.cpp
@@ -162,7 +162,7 @@ FPrimitiveSceneShaderData::FPrimitiveSceneShaderData(const FPrimitiveSceneProxy*
 		.HasCapsuleRepresentation(Proxy->HasDynamicIndirectShadowCasterRepresentation())
 		.UseSingleSampleShadowFromStationaryLights(Proxy->UseSingleSampleShadowFromStationaryLights())
 		.ReceivesDecals(Proxy->ReceivesDecals())
-		.ShouldCacheShadow(Proxy->ShouldCacheShadow())
+		.CacheShadowAsStatic(PrimitiveSceneInfo->ShouldCacheShadowAsStatic())
 		.OutputVelocity(bOutputVelocity)
 		.EvaluateWorldPositionOffset(Proxy->EvaluateWorldPositionOffset())
 		.CastContactShadow(Proxy->CastsContactShadow())
diff --git a/Engine/Source/Runtime/Engine/Public/PrimitiveSceneProxy.h b/Engine/Source/Runtime/Engine/Public/PrimitiveSceneProxy.h
--- a/Engine/Source/Runtime/Engine/Public/PrimitiveSceneProxy.h
+++ b/Engine/Source/Runtime/Engine/Public/PrimitiveSceneProxy.h
@@ -582,11 +582,6 @@ public:
 		return Mobility == EComponentMobility::Movable || !bGoodCandidateForCachedShadowmap; 
 	}
 
-	inline bool ShouldCacheShadow() const
-	{
-		return !IsMovable() && !bIsBeingMovedByEditor;
-	}
-
 	inline ELightmapType GetLightmapType() const { return LightmapType; }
 	inline bool IsStatic() const { return Mobility == EComponentMobility::Static; }
 	inline bool IsSelectable() const { return bSelectable; }
diff --git a/Engine/Source/Runtime/Engine/Public/PrimitiveUniformShaderParameters.h b/Engine/Source/Runtime/Engine/Public/PrimitiveUniformShaderParameters.h
--- a/Engine/Source/Runtime/Engine/Public/PrimitiveUniformShaderParameters.h
+++ b/Engine/Source/Runtime/Engine/Public/PrimitiveUniformShaderParameters.h
@@ -110,7 +110,7 @@ public:
 		bReceivesDecals								= false;
 		bUseSingleSampleShadowFromStationaryLights	= false;
 		bUseVolumetricLightmap						= false;
-		bShouldCacheShadow							= false;
+		bCacheShadowAsStatic						= false;
 		bOutputVelocity								= false;
 		bHasCapsuleRepresentation					= false;
 		bHasPreSkinnedLocalBounds					= false;
@@ -167,7 +167,7 @@ public:
 	PRIMITIVE_UNIFORM_BUILDER_FLAG_METHOD(bool,			CastShadow);
 	PRIMITIVE_UNIFORM_BUILDER_FLAG_METHOD(bool,			UseSingleSampleShadowFromStationaryLights);
 	PRIMITIVE_UNIFORM_BUILDER_FLAG_METHOD(bool,			UseVolumetricLightmap);
-	PRIMITIVE_UNIFORM_BUILDER_FLAG_METHOD(bool,			ShouldCacheShadow);
+	PRIMITIVE_UNIFORM_BUILDER_FLAG_METHOD(bool,			CacheShadowAsStatic);
 	PRIMITIVE_UNIFORM_BUILDER_FLAG_METHOD(bool,			OutputVelocity);
 	PRIMITIVE_UNIFORM_BUILDER_FLAG_METHOD(bool,			EvaluateWorldPositionOffset);
 	PRIMITIVE_UNIFORM_BUILDER_FLAG_METHOD(bool,			VisibleInGame);
@@ -397,7 +397,7 @@ public:
 		Parameters.Flags |= bHasCapsuleRepresentation ? PRIMITIVE_SCENE_DATA_FLAG_HAS_CAPSULE_REPRESENTATION : 0u;
 		Parameters.Flags |= bUseSingleSampleShadowFromStationaryLights ? PRIMITIVE_SCENE_DATA_FLAG_USE_SINGLE_SAMPLE_SHADOW_SL : 0u;
 		Parameters.Flags |= (bUseVolumetricLightmap && bUseSingleSampleShadowFromStationaryLights) ? PRIMITIVE_SCENE_DATA_FLAG_USE_VOLUMETRIC_LM_SHADOW_SL : 0u;
-		Parameters.Flags |= bShouldCacheShadow ? PRIMITIVE_SCENE_DATA_FLAG_SHOULD_CACHE_SHADOW : 0u;
+		Parameters.Flags |= bCacheShadowAsStatic ? PRIMITIVE_SCENE_DATA_FLAG_SHOULD_CACHE_SHADOW : 0u;
 		Parameters.Flags |= bOutputVelocity ? PRIMITIVE_SCENE_DATA_FLAG_OUTPUT_VELOCITY : 0u;
 		Parameters.Flags |= bEvaluateWorldPositionOffset ? PRIMITIVE_SCENE_DATA_FLAG_EVALUATE_WORLD_POSITION_OFFSET : 0u;
 		Parameters.Flags |= (Parameters.LocalToRelativeWorld.RotDeterminant() < 0.0f) ? PRIMITIVE_SCENE_DATA_FLAG_DETERMINANT_SIGN : 0u;
@@ -438,7 +438,7 @@ private:
 	uint32 bReceivesDecals : 1;
 	uint32 bUseSingleSampleShadowFromStationaryLights : 1;
 	uint32 bUseVolumetricLightmap : 1;
-	uint32 bShouldCacheShadow : 1;
+	uint32 bCacheShadowAsStatic : 1;
 	uint32 bOutputVelocity : 1;
 	uint32 bEvaluateWorldPositionOffset : 1;
 	uint32 bCastShadow : 1;
diff --git a/Engine/Source/Runtime/Renderer/Private/GPUScene.cpp b/Engine/Source/Runtime/Renderer/Private/GPUScene.cpp
--- a/Engine/Source/Runtime/Renderer/Private/GPUScene.cpp
+++ b/Engine/Source/Runtime/Renderer/Private/GPUScene.cpp
@@ -1316,7 +1316,8 @@ struct FUploadDataSourceAdapterDynamicPrimitives
 		const FGPUScenePrimitiveCollector::FPrimitiveData& PrimData = PrimitiveData[ItemIndex];		
 		FPrimitiveUniformShaderParameters Tmp = *PrimData.ShaderParams;
 		Tmp.InstanceSceneDataOffset = InstanceIDStartOffset + PrimData.LocalInstanceSceneDataOffset;
-		Tmp.NumInstanceSceneDataEntries = PrimData.NumInstances;			
+		Tmp.NumInstanceSceneDataEntries = PrimData.NumInstances;
+		Tmp.Flags &= ~(PRIMITIVE_SCENE_DATA_FLAG_SHOULD_CACHE_SHADOW);	// Dyanmic primtives are never cached as static for shadows
 		if (PrimitiveData[ItemIndex].LocalPayloadDataOffset != INDEX_NONE)
 		{
 			Tmp.InstancePayloadDataOffset = PayloadStartOffset + PrimData.LocalPayloadDataOffset;
diff --git a/Engine/Source/Runtime/Renderer/Private/PrimitiveSceneInfo.cpp b/Engine/Source/Runtime/Renderer/Private/PrimitiveSceneInfo.cpp
--- a/Engine/Source/Runtime/Renderer/Private/PrimitiveSceneInfo.cpp
+++ b/Engine/Source/Runtime/Renderer/Private/PrimitiveSceneInfo.cpp
@@ -189,6 +189,7 @@ FPrimitiveSceneInfo::FPrimitiveSceneInfo(UPrimitiveComponent* InComponent,FScene
 	bIndirectLightingCacheBufferDirty(false),
 	bRegisteredVirtualTextureProducerCallback(false),
 	bRegisteredWithVelocityData(false),
+	bCacheShadowAsStatic(InComponent->Mobility == EComponentMobility::Movable),
 	LevelUpdateNotificationIndex(INDEX_NONE),
 	InstanceSceneDataOffset(INDEX_NONE),
 	NumInstanceSceneDataEntries(0),
@@ -2115,3 +2116,12 @@ FString FPrimitiveSceneInfo::GetFullnameForDebuggingOnly() const
 	}
 	return FString(TEXT("Unknown Object"));
 }
+
+void FPrimitiveSceneInfo::SetCacheShadowAsStatic(bool bStatic)
+{
+	if (bCacheShadowAsStatic != bStatic)
+	{
+		bCacheShadowAsStatic = bStatic;
+		RequestGPUSceneUpdate(EPrimitiveDirtyState::ChangedOther);
+	}
+}
diff --git a/Engine/Source/Runtime/Renderer/Private/RendererScene.cpp b/Engine/Source/Runtime/Renderer/Private/RendererScene.cpp
--- a/Engine/Source/Runtime/Renderer/Private/RendererScene.cpp
+++ b/Engine/Source/Runtime/Renderer/Private/RendererScene.cpp
@@ -5087,19 +5087,19 @@ void FScene::UpdateAllPrimitiveSceneInfos(FRDGBuilder& GraphBuilder, bool bAsync
 		FVirtualShadowMapArrayCacheManager::FInvalidatingPrimitiveCollector InvalidatingPrimitiveCollector(CacheManager);
 
 		// All removed primitives must invalidate their footprints in the VSM before leaving
-		for (const FPrimitiveSceneInfo* PrimitiveSceneInfo : RemovedLocalPrimitiveSceneInfos)
+		for (FPrimitiveSceneInfo* PrimitiveSceneInfo : RemovedLocalPrimitiveSceneInfos)
 		{
-			InvalidatingPrimitiveCollector.Add(PrimitiveSceneInfo);
+			InvalidatingPrimitiveCollector.Removed(PrimitiveSceneInfo);
 		}
 		// All updated instances must also before moving or re-allocating (TODO: filter out only those actually updated)
 		for (const auto& Instance : UpdatedInstances)
 		{
-			InvalidatingPrimitiveCollector.Add(Instance.Key->GetPrimitiveSceneInfo());
+			InvalidatingPrimitiveCollector.UpdatedInstances(Instance.Key->GetPrimitiveSceneInfo());
 		}
 		// As must all primitive updates, 
 		for (const auto& Transform : UpdatedTransforms)
 		{
-			InvalidatingPrimitiveCollector.Add(Transform.Key->GetPrimitiveSceneInfo());
+			InvalidatingPrimitiveCollector.UpdatedTransform(Transform.Key->GetPrimitiveSceneInfo());
 		}
 
 		CacheManager->ProcessRemovedOrUpdatedPrimitives(GraphBuilder, GPUScene, InvalidatingPrimitiveCollector);
@@ -5189,6 +5189,10 @@ void FScene::UpdateAllPrimitiveSceneInfos(FRDGBuilder& GraphBuilder, bool bAsync
 						{
 							checkfSlow(DestIndex > SourceIndex, TEXT("Corrupted Prefix Sum [%d, %d]"), DestIndex, SourceIndex);
 							Primitives[DestIndex]->PackedIndex = SourceIndex;
+							// Update (the dynamic/compacted) primitive ID for the swapped primitive (not moved), no need to do the other one since it is being removed.
+							FPersistentPrimitiveIndex MovedPersisitentIndex = Primitives[DestIndex]->PersistentIndex;
+							PersistentPrimitiveIdToIndexMap[MovedPersisitentIndex.Index] = SourceIndex;
+
 							Primitives[SourceIndex]->PackedIndex = DestIndex;
 
 							TArraySwapElements(Primitives, DestIndex, SourceIndex);
@@ -5346,7 +5350,9 @@ void FScene::UpdateAllPrimitiveSceneInfos(FRDGBuilder& GraphBuilder, bool bAsync
 					}
 				}
 
-				PersistentPrimitiveIdAllocator.Free(PrimitiveSceneInfo->PersistentIndex.Index);
+				const int32 PersistentIndex = PrimitiveSceneInfo->PersistentIndex.Index;
+				PersistentPrimitiveIdAllocator.Free(PersistentIndex);
+				PersistentPrimitiveIdToIndexMap[PersistentIndex] = INDEX_NONE;
 			}
 
 			for (TMap<int32, TArray<FCachedShadowMapData>>::TIterator CachedShadowMapIt(CachedShadowMaps); CachedShadowMapIt; ++CachedShadowMapIt)
@@ -5465,7 +5471,14 @@ void FScene::UpdateAllPrimitiveSceneInfos(FRDGBuilder& GraphBuilder, bool bAsync
 				const int32 SourceIndex = PrimitiveSceneProxies.Num() - 1;
 				PrimitiveSceneInfo->PackedIndex = SourceIndex;
 				checkSlow(PrimitiveSceneInfo->PersistentIndex.Index == INDEX_NONE);
-				PrimitiveSceneInfo->PersistentIndex = FPersistentPrimitiveIndex{ PersistentPrimitiveIdAllocator.Allocate() };
+				FPersistentPrimitiveIndex PersistentPrimitiveIndex{ PersistentPrimitiveIdAllocator.Allocate() };
+				PrimitiveSceneInfo->PersistentIndex = PersistentPrimitiveIndex;
+				// Ensure map is large enough
+				if (PersistentPrimitiveIndex.Index >= PersistentPrimitiveIdToIndexMap.Num())
+				{
+					PersistentPrimitiveIdToIndexMap.SetNumUninitialized(PersistentPrimitiveIndex.Index + 1);
+				}
+				PersistentPrimitiveIdToIndexMap[PersistentPrimitiveIndex.Index] = SourceIndex;
 
 				GPUScene.AddPrimitiveToUpdate(SourceIndex, EPrimitiveDirtyState::AddedMask);
 			}
@@ -5529,6 +5542,15 @@ void FScene::UpdateAllPrimitiveSceneInfos(FRDGBuilder& GraphBuilder, bool bAsync
 							Primitives[DestIndex]->PackedIndex = SourceIndex;
 							Primitives[SourceIndex]->PackedIndex = DestIndex;
 
+							// Update (the dynamic/compacted) primitive ID for the swapped primitives
+							{
+								FPersistentPrimitiveIndex PersisitentIndex = Primitives[DestIndex]->PersistentIndex;
+								PersistentPrimitiveIdToIndexMap[PersisitentIndex.Index] = SourceIndex;
+							}
+							{
+								FPersistentPrimitiveIndex PersisitentIndex = Primitives[SourceIndex]->PersistentIndex;
+								PersistentPrimitiveIdToIndexMap[PersisitentIndex.Index] = DestIndex;
+							}
 							TArraySwapElements(Primitives, DestIndex, SourceIndex);
 							TArraySwapElements(PrimitiveTransforms, DestIndex, SourceIndex);
 							TArraySwapElements(PrimitiveSceneProxies, DestIndex, SourceIndex);
@@ -5954,6 +5976,14 @@ void FScene::UpdateAllPrimitiveSceneInfos(FRDGBuilder& GraphBuilder, bool bAsync
 	UpdatedOcclusionBoundsSlacks.Empty();
 	DistanceFieldSceneDataUpdates.Empty();
 	AddedPrimitiveSceneInfos.Empty();
+
+	for (const FPrimitiveSceneInfo* PrimitiveSceneInfo : Primitives)
+	{
+		check(PrimitiveSceneInfo->PackedIndex != INDEX_NONE);
+		check(PrimitiveSceneInfo->PackedIndex < Primitives.Num());
+		check(PrimitiveSceneInfo->PersistentIndex.Index != INDEX_NONE);
+		check(PersistentPrimitiveIdToIndexMap[PrimitiveSceneInfo->PersistentIndex.Index] == PrimitiveSceneInfo->PackedIndex);
+	}
 }
 
 void FScene::CreateLightPrimitiveInteractionsForPrimitive(FPrimitiveSceneInfo* PrimitiveInfo, bool bAsyncCreateLPIs)
diff --git a/Engine/Source/Runtime/Renderer/Private/ScenePrivate.h b/Engine/Source/Runtime/Renderer/Private/ScenePrivate.h
--- a/Engine/Source/Runtime/Renderer/Private/ScenePrivate.h
+++ b/Engine/Source/Runtime/Renderer/Private/ScenePrivate.h
@@ -2785,6 +2785,8 @@ public:
 	TBitArray<> PrimitivesNeedingStaticMeshUpdate;
 	TSet<FPrimitiveSceneInfo*> PrimitivesNeedingStaticMeshUpdateWithoutVisibilityCheck;
 
+	TArray<int32> PersistentPrimitiveIdToIndexMap;
+
 	struct FTypeOffsetTableEntry
 	{
 		FTypeOffsetTableEntry(SIZE_T InPrimitiveSceneProxyType, uint32 InOffset) : PrimitiveSceneProxyType(InPrimitiveSceneProxyType), Offset(InOffset) {}
@@ -3443,6 +3445,15 @@ public:
 	 */
 	inline int32 GetMaxPersistentPrimitiveIndex() const { return PersistentPrimitiveIdAllocator.GetMaxSize(); }
 
+	FORCEINLINE int32 GetPrimitiveIndex(const FPersistentPrimitiveIndex& PersistentPrimitiveIndex) const
+	{ 
+		if (PersistentPrimitiveIndex.Index < PersistentPrimitiveIdToIndexMap.Num())
+		{
+			return PersistentPrimitiveIdToIndexMap[PersistentPrimitiveIndex.Index];
+		}
+		return INDEX_NONE;
+	}
+
 	bool GetForceNoPrecomputedLighting() const
 	{
 		return bForceNoPrecomputedLighting;
diff --git a/Engine/Source/Runtime/Renderer/Private/VirtualShadowMaps/VirtualShadowMapArray.cpp b/Engine/Source/Runtime/Renderer/Private/VirtualShadowMaps/VirtualShadowMapArray.cpp
--- a/Engine/Source/Runtime/Renderer/Private/VirtualShadowMaps/VirtualShadowMapArray.cpp
+++ b/Engine/Source/Runtime/Renderer/Private/VirtualShadowMaps/VirtualShadowMapArray.cpp
@@ -416,6 +416,9 @@ void FVirtualShadowMapArray::Initialize(FRDGBuilder& GraphBuilder, FVirtualShado
 		TRefCountPtr<IPooledRenderTarget> PhysicalPagePool = CacheManager->SetPhysicalPoolSize(GraphBuilder, GetPhysicalPoolSize(), PoolArraySize);
 		PhysicalPagePoolRDG = GraphBuilder.RegisterExternalTexture(PhysicalPagePool);
 		UniformParameters.PhysicalPagePool = PhysicalPagePoolRDG;
+
+		// Initialize the instance invalidation buffer and primitive cache flags
+		SceneData = CacheManager->InitializeSceneDataForFrame(GraphBuilder);
 	}
 	else
 	{
@@ -1243,14 +1246,6 @@ void FVirtualShadowMapArray::BuildPageAllocations(
 	AddClearUAVPass(GraphBuilder, GraphBuilder.CreateUAV(DirtyPageFlagsRDG), 0);
 	CumulativeDirtyPageFlagsRDG = GraphBuilder.CreateBuffer(FRDGBufferDesc::CreateStructuredDesc(sizeof(uint32), GetMaxPhysicalPages()), TEXT("Shadow.Virtual.CumulativeDirtyPageFlags"));
 	AddClearUAVPass(GraphBuilder, GraphBuilder.CreateUAV(CumulativeDirtyPageFlagsRDG), 0);
-		
-	// Record the number of instances the buffer has capactiy for, should anything change (it shouldn't!)
-	NumInvalidatingInstanceSlots = Scene.GPUScene.GetNumInstances();
-	// Allocate space for counter, worst case ID storage, and flags.
-	int32 InstanceInvalidationBufferSize = 1 + NumInvalidatingInstanceSlots + FMath::DivideAndRoundUp(NumInvalidatingInstanceSlots, 32);
-	InvalidatingInstancesRDG = GraphBuilder.CreateBuffer(FRDGBufferDesc::CreateStructuredDesc(sizeof(uint32), InstanceInvalidationBufferSize), TEXT("Shadow.Virtual.InvalidatingInstances"));
-	// Clear to zero, technically only need to clear first Scene.GPUScene.GetNumInstances()  + 1 uints
-	AddClearUAVPass(GraphBuilder, GraphBuilder.CreateUAV(InvalidatingInstancesRDG), 0);
 
 	const uint32 NumPageRects = ShadowMaps.Num() * FVirtualShadowMap::MaxMipLevels;
 	const uint32 NumPageRectsToAllocate = FMath::RoundUpToPowerOfTwo(NumPageRects);
@@ -2197,8 +2192,8 @@ static FCullingResult AddCullingPasses(FRDGBuilder& GraphBuilder,
 
 		PassParameters->VisibleInstancesOut = GraphBuilder.CreateUAV(VisibleInstancesRdg);
 		PassParameters->VisibleInstanceCountBufferOut = GraphBuilder.CreateUAV(VisibleInstanceWriteOffsetRDG);
-		PassParameters->OutInvalidatingInstances = GraphBuilder.CreateUAV(VirtualShadowMapArray->InvalidatingInstancesRDG);
-		PassParameters->NumInvalidatingInstanceSlots = VirtualShadowMapArray->NumInvalidatingInstanceSlots;
+		PassParameters->OutInvalidatingInstances = GraphBuilder.CreateUAV(VirtualShadowMapArray->SceneData.InvalidatingInstancesBuffer);
+		PassParameters->NumInvalidatingInstanceSlots = VirtualShadowMapArray->SceneData.NumInvalidatingInstanceSlots;
 		PassParameters->HZBShaderParameters = HZBShaderParameters;
 
 		bool bGenerateStats = VirtualShadowMapArray->StatsBufferRDG != nullptr;
diff --git a/Engine/Source/Runtime/Renderer/Private/VirtualShadowMaps/VirtualShadowMapArray.h b/Engine/Source/Runtime/Renderer/Private/VirtualShadowMaps/VirtualShadowMapArray.h
--- a/Engine/Source/Runtime/Renderer/Private/VirtualShadowMaps/VirtualShadowMapArray.h
+++ b/Engine/Source/Runtime/Renderer/Private/VirtualShadowMaps/VirtualShadowMapArray.h
@@ -117,6 +117,14 @@ struct FVirtualShadowMapHZBMetadata
 	uint32		  TargetLayerIndex = INDEX_NONE;
 };
 
+// Scene data buffers sized and imported into the graph builder
+// Persistent, not ping-ponged
+struct FVirtualShadowMapSceneData
+{
+	FRDGBufferRef InvalidatingInstancesBuffer;
+	int32 NumInvalidatingInstanceSlots = -1;
+};
+
 BEGIN_GLOBAL_SHADER_PARAMETER_STRUCT(FVirtualShadowMapUniformParameters, )
 	SHADER_PARAMETER(uint32, NumFullShadowMaps)
 	SHADER_PARAMETER(uint32, NumSinglePageShadowMaps)
@@ -322,12 +330,9 @@ public:
 	// TODO: make transient - Buffer that stores flags marking each page that received dynamic geo.
 	FRDGBufferRef DynamicCasterPageFlagsRDG = nullptr;
 
-	// Buffer that stores flags marking each instance that needs to be invalidated the subsequent frame (handled by the cache manager).
-	// This covers things like WPO or GPU-side updates, and any other case where we determine an instance needs to invalidate its footprint. 
-	// Buffer of uints, organized as as follows: InvalidatingInstancesRDG[0] == count, InvalidatingInstancesRDG[1+MaxInstanceCount:1+MaxInstanceCount+MaxInstanceCount/32] == flags, 
-	// InvalidatingInstancesRDG[1:MaxInstanceCount] == growing compact array of instances that need invaldation
-	FRDGBufferRef InvalidatingInstancesRDG = nullptr;
-	int32 NumInvalidatingInstanceSlots = 0;
+	// References to scene data imported into the graph builder for this frame
+	// NOTE: The underlying data is owned by FVirtualShadowMapCacheManager.
+	FVirtualShadowMapSceneData SceneData;
 	
 	// uint4 buffer with one rect for each mip level in all SMs, calculated to bound committed pages
 	// Used to clip the rect size of clusters during culling.
diff --git a/Engine/Source/Runtime/Renderer/Private/VirtualShadowMaps/VirtualShadowMapCacheManager.cpp b/Engine/Source/Runtime/Renderer/Private/VirtualShadowMaps/VirtualShadowMapCacheManager.cpp
--- a/Engine/Source/Runtime/Renderer/Private/VirtualShadowMaps/VirtualShadowMapCacheManager.cpp
+++ b/Engine/Source/Runtime/Renderer/Private/VirtualShadowMaps/VirtualShadowMapCacheManager.cpp
@@ -15,6 +15,8 @@
 #include "ShaderPrint.h"
 #include "RendererOnScreenNotification.h"
 
+#define VSM_LOG_STATIC_CACHING 0
+
 static TAutoConsoleVariable<int32> CVarAccumulateStats(
 	TEXT("r.Shadow.Virtual.AccumulateStats"),
 	0,
@@ -67,6 +69,7 @@ static FAutoConsoleVariableRef  CVarForceInvalidateDirectionalVSM(
 	TEXT("Forces the clipmap to always invalidate, useful to emulate a moving sun to avoid misrepresenting cache performance."),
 	ECVF_RenderThreadSafe);
 
+
 void FVirtualShadowMapCacheEntry::UpdateClipmap(
 	int32 VirtualShadowMapId,
 	const FMatrix &WorldToLight,
@@ -227,17 +230,16 @@ static inline uint32 EncodeInstanceInvalidationPayload(bool bInvalidateStaticPag
 }
 
 FVirtualShadowMapArrayCacheManager::FInvalidatingPrimitiveCollector::FInvalidatingPrimitiveCollector(FVirtualShadowMapArrayCacheManager* InVirtualShadowMapArrayCacheManager)
-	: AlreadyAddedPrimitives(false, InVirtualShadowMapArrayCacheManager->Scene->Primitives.Num())
-	, Scene(*InVirtualShadowMapArrayCacheManager->Scene)
+	: Scene(*InVirtualShadowMapArrayCacheManager->Scene)
 	, GPUScene(InVirtualShadowMapArrayCacheManager->Scene->GPUScene)
 	, Manager(*InVirtualShadowMapArrayCacheManager)
+	, AlreadyAddedPrimitives(false, InVirtualShadowMapArrayCacheManager->Scene->Primitives.Num())
 {
-	bool bPossiblyCachedAsStatic = false;	// TODO
-
 	// Add and clear pending invalidations enqueued on the GPU Scene from dynamic primitives added since last invalidation
 	for (const FGPUScene::FInstanceRange& Range : GPUScene.DynamicPrimitiveInstancesToInvalidate)
 	{
-		LoadBalancer.Add(Range.InstanceSceneDataOffset, Range.NumInstanceSceneDataEntries, EncodeInstanceInvalidationPayload(bPossiblyCachedAsStatic));
+		// Dynamic primitives are never cached as static; see  FUploadDataSourceAdapterDynamicPrimitives::GetPrimitiveInfo
+		LoadBalancer.Add(Range.InstanceSceneDataOffset, Range.NumInstanceSceneDataEntries, EncodeInstanceInvalidationPayload(false));
 #if VSM_LOG_INVALIDATIONS
 		RangesStr.Appendf(TEXT("[%6d, %6d), "), Range.InstanceSceneDataOffset, Range.InstanceSceneDataOffset + Range.NumInstanceSceneDataEntries);
 #endif
@@ -256,9 +258,8 @@ FVirtualShadowMapArrayCacheManager::FInvalidatingPrimitiveCollector::FInvalidati
 			{
 				if (SmCacheEntry.IsValid())
 				{
-					// Lowest bit indicates whether to run the clipmap loop, add 1 to ID so != 0 <==> single level processing
 					LoadBalancer.Add(Range.InstanceSceneDataOffset, Range.NumInstanceSceneDataEntries,
-						EncodeInstanceInvalidationPayload(bPossiblyCachedAsStatic, SmCacheEntry->CurrentVirtualShadowMapId));
+						EncodeInstanceInvalidationPayload(Range.bInvalidateStaticPage, SmCacheEntry->CurrentVirtualShadowMapId));
 				}
 			}
 
@@ -269,9 +270,52 @@ FVirtualShadowMapArrayCacheManager::FInvalidatingPrimitiveCollector::FInvalidati
 		}
 		CacheEntry.Value->PrimitiveInstancesToInvalidate.Reset();
 	}
+
+	// Process any GPU readback static invalidations
+	{
+		FVirtualShadowMapFeedback::FReadbackInfo Readback = Manager.StaticGPUInvalidationsFeedback.GetLatestReadbackBuffer();
+		if (Readback.Size > 0)
+		{
+			TBitArray<SceneRenderingAllocator> Primitives;
+			{
+				const uint32* Data = (const uint32*)Readback.Buffer->Lock(Readback.Size);			
+				Primitives.AddRange(Data, Readback.Size * 8);
+				Readback.Buffer->Unlock();
+				// TODO: Mark that we've done this buffer? Not really harmful to redo it
+			}
+
+			for (TConstSetBitIterator<> PrimitivesIt(Primitives); PrimitivesIt; ++PrimitivesIt)
+			{
+				const FPersistentPrimitiveIndex PersistentPrimitiveIndex = FPersistentPrimitiveIndex{PrimitivesIt.GetIndex()};
+				
+				// NOTE: Have to be a bit careful as this primitive index came from a few frames ago... thus it may no longer
+				// be valid, or possibly even replaced by a new primitive.
+				// TODO: Dual iterator and avoid any primitives that have been removed recently (might resuse the slot)
+				const int32 PrimitiveIndex = Scene.GetPrimitiveIndex(PersistentPrimitiveIndex);
+				FPrimitiveSceneInfo* PrimitiveSceneInfo = Scene.GetPrimitiveSceneInfo(PrimitiveIndex);
+				if (PrimitiveSceneInfo)
+				{
+					if (Manager.WasRecentlyRemoved(PersistentPrimitiveIndex))
+					{
+						// Do nothing for now, as this slot may have been reused after a previous removal
+#if VSM_LOG_STATIC_CACHING
+						UE_LOG(LogRenderer, Warning, TEXT("Ignoring GPU invalidation due to recent primitive removal: %u"), PersistentPrimitiveIndex.Index);
+#endif
+					}
+					else
+					{
+#if VSM_LOG_STATIC_CACHING
+						UE_LOG(LogRenderer, Warning, TEXT("Transitioning GPU invalidation to dynamic caching: %u"), PersistentPrimitiveIndex.Index);
+#endif
+						AddInvalidation(PrimitiveSceneInfo, false);
+					}
+				}
+			}
+		}
+	}
 }
 
-void FVirtualShadowMapArrayCacheManager::FInvalidatingPrimitiveCollector::Add(const FPrimitiveSceneInfo * PrimitiveSceneInfo)
+void FVirtualShadowMapArrayCacheManager::FInvalidatingPrimitiveCollector::AddInvalidation(FPrimitiveSceneInfo * PrimitiveSceneInfo, bool bRemovedPrimitive)
 {
 	int32 PrimitiveID = PrimitiveSceneInfo->GetIndex();
 	if (PrimitiveID >= 0
@@ -282,11 +326,28 @@ void FVirtualShadowMapArrayCacheManager::FInvalidatingPrimitiveCollector::Add(co
 		&& !EnumHasAnyFlags(GPUScene.GetPrimitiveDirtyState(PrimitiveID), EPrimitiveDirtyState::Added))
 	{
 		AlreadyAddedPrimitives[PrimitiveID] = true;
-		int32 PersistentPrimitiveIndex = PrimitiveSceneInfo->GetPersistentIndex().Index;
+		const FPersistentPrimitiveIndex PersistentPrimitiveIndex = PrimitiveSceneInfo->GetPersistentIndex();
 
 		// Nanite meshes need special handling because they don't get culled on CPU, thus always process invalidations for those
 		const bool bIsNaniteMesh = Scene.PrimitiveFlagsCompact[PrimitiveID].bIsNaniteMesh;
-		const bool bPossiblyCachedAsStatic = !PrimitiveSceneInfo->Proxy->IsMovable();
+		const bool bPreviouslyCachedAsStatic = PrimitiveSceneInfo->ShouldCacheShadowAsStatic();
+
+		if (bRemovedPrimitive)
+		{
+			RemovedPrimitives.PadToNum(PersistentPrimitiveIndex.Index + 1, false);
+			RemovedPrimitives[PersistentPrimitiveIndex.Index] = true;
+		}
+		else
+		{
+			// Swap to dynamic caching if it was already static
+			if (bPreviouslyCachedAsStatic)
+			{
+				PrimitiveSceneInfo->SetCacheShadowAsStatic(false);
+#if VSM_LOG_STATIC_CACHING
+				UE_LOG(LogRenderer, Warning, TEXT("FVirtualShadowMapArrayCacheManager: '%s' switched to dynamic caching"), *PrimitiveSceneInfo->GetFullnameForDebuggingOnly());
+#endif
+			}
+		}
 
 		const int32 NumInstanceSceneDataEntries = PrimitiveSceneInfo->GetNumInstanceSceneDataEntries();
 
@@ -294,12 +355,12 @@ void FVirtualShadowMapArrayCacheManager::FInvalidatingPrimitiveCollector::Add(co
 		for (auto& CacheEntry : Manager.PrevCacheEntries)
 		{
 			TBitArray<>& CachedPrimitives = CacheEntry.Value->CachedPrimitives;
-			if (bIsNaniteMesh || (PersistentPrimitiveIndex < CachedPrimitives.Num() && CachedPrimitives[PersistentPrimitiveIndex]))
+			if (bIsNaniteMesh || (PersistentPrimitiveIndex.Index < CachedPrimitives.Num() && CachedPrimitives[PersistentPrimitiveIndex.Index]))
 			{
 				if (!bIsNaniteMesh)
 				{
 					// Clear the record as we're wiping it out.
-					CachedPrimitives[PersistentPrimitiveIndex] = false;
+					CachedPrimitives[PersistentPrimitiveIndex.Index] = false;
 				}
 
 				// Add item for each shadow map explicitly, inflates host data but improves load balancing,
@@ -310,7 +371,7 @@ void FVirtualShadowMapArrayCacheManager::FInvalidatingPrimitiveCollector::Add(co
 					{
 						checkSlow(SmCacheEntry->CurrentVirtualShadowMapId != INDEX_NONE);
 						LoadBalancer.Add(PrimitiveSceneInfo->GetInstanceSceneDataOffset(), NumInstanceSceneDataEntries,
-							EncodeInstanceInvalidationPayload(bPossiblyCachedAsStatic, SmCacheEntry->CurrentVirtualShadowMapId));
+							EncodeInstanceInvalidationPayload(bPreviouslyCachedAsStatic, SmCacheEntry->CurrentVirtualShadowMapId));
 					}
 				}
 			}
@@ -323,6 +384,75 @@ void FVirtualShadowMapArrayCacheManager::FInvalidatingPrimitiveCollector::Add(co
 }
 
 
+FVirtualShadowMapFeedback::FVirtualShadowMapFeedback()
+{
+	for (int32 i = 0; i < MaxBuffers; ++i)
+	{
+		Buffers[i].Buffer = new FRHIGPUBufferReadback(TEXT("Shadow.Virtual.Readback"));
+		Buffers[i].Size = 0;
+	}
+}
+
+FVirtualShadowMapFeedback::~FVirtualShadowMapFeedback()
+{
+	for (int32 i = 0; i < MaxBuffers; ++i)
+	{
+		delete Buffers[i].Buffer;
+		Buffers[i].Buffer = nullptr;
+		Buffers[i].Size = 0;
+	}
+}
+
+void FVirtualShadowMapFeedback::SubmitFeedbackBuffer(
+	FRDGBuilder& GraphBuilder,
+	FRDGBufferRef FeedbackBuffer)
+{
+	// Source copy usage is required for readback
+	check((FeedbackBuffer->Desc.Usage & EBufferUsageFlags::SourceCopy) == EBufferUsageFlags::SourceCopy);
+
+	if (NumPending == MaxBuffers)
+	{
+#if VSM_LOG_STATIC_CACHING
+		UE_LOG(LogRenderer, Warning, TEXT("FVirtualShadowMapFeedback ran out of feedback buffers!"));
+#endif
+		return;
+	}
+
+	FRHIGPUBufferReadback* ReadbackBuffer = Buffers[WriteIndex].Buffer;
+	Buffers[WriteIndex].Size = FeedbackBuffer->Desc.GetSize();
+
+	AddReadbackBufferPass(GraphBuilder, RDG_EVENT_NAME("Readback"), FeedbackBuffer,
+		[ReadbackBuffer, FeedbackBuffer](FRHICommandList& RHICmdList)
+		{
+			ReadbackBuffer->EnqueueCopy(RHICmdList, FeedbackBuffer->GetRHI(), 0u);
+		});
+
+	WriteIndex = (WriteIndex + 1) % MaxBuffers;
+	NumPending = FMath::Min(NumPending + 1, MaxBuffers);
+}
+
+FVirtualShadowMapFeedback::FReadbackInfo FVirtualShadowMapFeedback::GetLatestReadbackBuffer()
+{
+	int32 LatestBufferIndex = -1;
+
+	// Find latest buffer that is ready
+	while (NumPending > 0)
+	{
+		uint32 Index = (WriteIndex + MaxBuffers - NumPending) % MaxBuffers;
+		if (Buffers[Index].Buffer->IsReady())
+		{
+			--NumPending;
+			LatestBufferIndex = Index;
+		}
+		else
+		{
+			break;
+		}
+	}
+
+	return LatestBufferIndex >= 0 ? Buffers[LatestBufferIndex] : FReadbackInfo();
+}
+
 
 FVirtualShadowMapArrayCacheManager::FVirtualShadowMapArrayCacheManager(FScene* InScene) 
 	: Scene(InScene)
@@ -453,6 +583,37 @@ void FVirtualShadowMapArrayCacheManager::FreeHZBPhysicalPool()
 	}
 }
 
+FVirtualShadowMapSceneData FVirtualShadowMapArrayCacheManager::InitializeSceneDataForFrame(FRDGBuilder& GraphBuilder)
+{
+	FVirtualShadowMapSceneData Result;
+
+	const int32 TargetInstanceSlots = FMath::RoundUpToPowerOfTwo(Scene->GPUScene.GetNumInstances());
+	
+	// Resize instance invalidation buffer if needed
+	// TODO: Should this stuff be in OnSceneChange as well? Unsure where GetNumInstances() can change
+	if (NumInvalidatingInstanceSlots != TargetInstanceSlots)
+	{
+		// Record the number of instances the buffer has capactiy for, should anything change (it shouldn't!)
+		NumInvalidatingInstanceSlots = TargetInstanceSlots;
+
+		// Allocate space for counter, worst case ID storage, and flags.
+		int32 InstanceInvalidationBufferSize = 1 + NumInvalidatingInstanceSlots + FMath::DivideAndRoundUp(NumInvalidatingInstanceSlots, 32);
+
+		AllocatePooledBuffer(FRDGBufferDesc::CreateStructuredDesc(sizeof(uint32), InstanceInvalidationBufferSize),
+			InvalidatingInstancesBuffer, TEXT("Shadow.Virtual.InvalidatingInstances"));
+	}
+
+	Result.InvalidatingInstancesBuffer = GraphBuilder.RegisterExternalBuffer(
+		InvalidatingInstancesBuffer, TEXT("Shadow.Virtual.InvalidatingInstancesBuffer"));
+	Result.NumInvalidatingInstanceSlots = NumInvalidatingInstanceSlots;
+
+	// Clear to zero
+	AddClearUAVPass(GraphBuilder, GraphBuilder.CreateUAV(Result.InvalidatingInstancesBuffer), 0);
+
+	bGPUInstanceInvalidationsProcessedThisFrame = false;
+	return Result;
+}
+
 void FVirtualShadowMapArrayCacheManager::Invalidate()
 {
 	// Clear the cache
@@ -518,7 +679,11 @@ void FVirtualShadowMapPerLightCacheEntry::OnPrimitiveRendered(const FPrimitiveSc
 		// Deformable mesh primitives need to trigger invalidation (even if they did not move) or we get artifacts, for example skinned meshes that are animating but not currently moving.
 		if (PrimitiveSceneInfo->Proxy->HasDeformableMesh())
 		{
-			PrimitiveInstancesToInvalidate.Add(FInstanceRange{ PrimitiveSceneInfo->GetInstanceSceneDataOffset(), PrimitiveSceneInfo->GetNumInstanceSceneDataEntries() });
+			PrimitiveInstancesToInvalidate.Add(FInstanceRange{ 
+				PrimitiveSceneInfo->GetInstanceSceneDataOffset(),
+				PrimitiveSceneInfo->GetNumInstanceSceneDataEntries(),
+				PrimitiveSceneInfo->ShouldCacheShadowAsStatic()
+			});
 		}
 	}
 }
@@ -594,8 +759,6 @@ void FVirtualShadowMapArrayCacheManager::ExtractFrameData(
 			GraphBuilder.QueueBufferExtraction(VirtualShadowMapArray.DynamicCasterPageFlagsRDG, &PrevBuffers.DynamicCasterPageFlags);
 			GraphBuilder.QueueBufferExtraction(VirtualShadowMapArray.ProjectionDataRDG, &PrevBuffers.ProjectionData);
 
-			GraphBuilder.QueueBufferExtraction(VirtualShadowMapArray.InvalidatingInstancesRDG, &PrevBuffers.InvalidatingInstancesBuffer);
-			PrevBuffers.NumInvalidatingInstanceSlots = VirtualShadowMapArray.NumInvalidatingInstanceSlots;
 			
 			// Store but drop any temp references embedded in the uniform parameters this frame.
 			// We'll reestablish them when we reimport the extracted resources next frame
@@ -645,6 +808,16 @@ void FVirtualShadowMapArrayCacheManager::ExtractFrameData(
 		// Ideally in the long run we want the cache itself to be more robust against rendering multiple views. but
 		// for now this at least provides a work-around for some common cases where only one view is rendering VSMs.
 	}
+
+	// Every once in a while zero out our recently removed primitive flags. This lets us ignore slots that
+	// may have been reused since they were flagged from GPU invalidations.
+	++RecentlyRemovedFrameCounter;
+	if (RecentlyRemovedFrameCounter >= 3)
+	{
+		RecentlyRemovedPrimitives[RecentlyRemovedReadIndex].Reset();
+		RecentlyRemovedReadIndex = 1 - RecentlyRemovedReadIndex;
+		RecentlyRemovedFrameCounter = 0;
+	}
 }
 
 void FVirtualShadowMapArrayCacheManager::ExtractStats(FRDGBuilder& GraphBuilder, FVirtualShadowMapArray &VirtualShadowMapArray)
@@ -801,8 +974,16 @@ bool FVirtualShadowMapArrayCacheManager::IsAccumulatingStats()
 	return CVarAccumulateStats.GetValueOnRenderThread() != 0;
 }
 
+static uint32 GetPrimFlagsBufferSizeInDwords(int32 MaxPersistentPrimitiveIndex)
+{
+	return FMath::RoundUpToPowerOfTwo(FMath::DivideAndRoundUp(MaxPersistentPrimitiveIndex, 32));
+}
+
 void FVirtualShadowMapArrayCacheManager::ProcessRemovedOrUpdatedPrimitives(FRDGBuilder& GraphBuilder, const FGPUScene& GPUScene, FInvalidatingPrimitiveCollector& InvalidatingPrimitiveCollector)
 {
+	// Always incorporate any scene removals into the "recently removed" list
+	UpdateRecentlyRemoved(InvalidatingPrimitiveCollector.GetRemovedPrimitives());
+
 	if (CVarCacheVirtualSMs.GetValueOnRenderThread() != 0 && PrevBuffers.DynamicCasterPageFlags.IsValid())
 	{
 		RDG_EVENT_SCOPE(GraphBuilder, "Shadow.Virtual.ProcessRemovedOrUpdatedPrimitives");
@@ -818,7 +999,7 @@ void FVirtualShadowMapArrayCacheManager::ProcessRemovedOrUpdatedPrimitives(FRDGB
 	}
 }
 
-static void ResizeFlagArray(TBitArray<>& BitArray, int32 NewMax)
+static void ResizeFlagArray(TBitArray<>& BitArray, int32 NewMax, bool DefaultValue = false)
 {
 	if (BitArray.Num() > NewMax)
 	{
@@ -827,8 +1008,7 @@ static void ResizeFlagArray(TBitArray<>& BitArray, int32 NewMax)
 	}
 	else if (BitArray.Num() < NewMax)
 	{
-		// Add false
-		BitArray.Add(false, NewMax - BitArray.Num());
+		BitArray.Add(DefaultValue, NewMax - BitArray.Num());
 	}
 }
 
@@ -836,15 +1016,17 @@ void FVirtualShadowMapArrayCacheManager::OnSceneChange()
 {
 	if (CVarCacheVirtualSMs.GetValueOnRenderThread() != 0)
 	{
+		const int32 MaxPersistentPrimitiveIndex = FMath::Max(1, Scene->GetMaxPersistentPrimitiveIndex());
+
 		for (auto& CacheEntry : PrevCacheEntries)
 		{
-			ResizeFlagArray(CacheEntry.Value->CachedPrimitives, Scene->GetMaxPersistentPrimitiveIndex());
-			ResizeFlagArray(CacheEntry.Value->RenderedPrimitives, Scene->GetMaxPersistentPrimitiveIndex());
+			ResizeFlagArray(CacheEntry.Value->CachedPrimitives, MaxPersistentPrimitiveIndex);
+			ResizeFlagArray(CacheEntry.Value->RenderedPrimitives, MaxPersistentPrimitiveIndex);
 		}
 		for (auto& CacheEntry : CacheEntries)
 		{
-			ResizeFlagArray(CacheEntry.Value->CachedPrimitives, Scene->GetMaxPersistentPrimitiveIndex());
-			ResizeFlagArray(CacheEntry.Value->RenderedPrimitives, Scene->GetMaxPersistentPrimitiveIndex());
+			ResizeFlagArray(CacheEntry.Value->CachedPrimitives, MaxPersistentPrimitiveIndex);
+			ResizeFlagArray(CacheEntry.Value->RenderedPrimitives, MaxPersistentPrimitiveIndex);
 		}
 	}
 }
@@ -898,9 +1080,11 @@ public:
 		SHADER_PARAMETER_SAMPLER(SamplerState, HZBSampler)
 		SHADER_PARAMETER( FVector2f,	HZBSize )
 
+		// GPU instances parameters
 		SHADER_PARAMETER_RDG_BUFFER_SRV(StructuredBuffer<uint>, InvalidatingInstances)
 		SHADER_PARAMETER(uint32, NumInvalidatingInstanceSlots)
 		RDG_BUFFER_ACCESS(IndirectArgs, ERHIAccess::IndirectArgs)
+		SHADER_PARAMETER_RDG_BUFFER_UAV(RWStructuredBuffer<uint>, OutStaticInvalidatingPrimitives)
 
 		SHADER_PARAMETER_STRUCT_INCLUDE(FGPUScene::FInstanceGPULoadBalancer::FShaderParameters, LoadBalancerParameters)
 	END_SHADER_PARAMETER_STRUCT()
@@ -959,11 +1143,14 @@ void FVirtualShadowMapArrayCacheManager::UpdateGPUMask(FRHIGPUMask GPUMask)
 }
 #endif  // WITH_MGPU
 
-static void SetupCommonParameters(FRDGBuilder& GraphBuilder, FVirtualShadowMapArrayCacheManager* CacheManager, int32 TotalInstanceCount, const FGPUScene& GPUScene, 
+static void SetupCommonParameters(
+	FRDGBuilder& GraphBuilder,
+	FVirtualShadowMapArrayCacheManager* CacheManager,
+	int32 TotalInstanceCount,
+	const FGPUScene& GPUScene,
 	FVirtualSmInvalidateInstancePagesCS::FParameters& OutPassParameters,
 	FVirtualSmInvalidateInstancePagesCS::FPermutationDomain &OutPermutationVector)
 {
-
 	auto RegExtCreateSrv = [&GraphBuilder](const TRefCountPtr<FRDGPooledBuffer>& Buffer, const TCHAR* Name) -> FRDGBufferSRVRef
 	{
 		return GraphBuilder.CreateSRV(GraphBuilder.RegisterExternalBuffer(Buffer, Name));
@@ -1028,7 +1215,11 @@ static void SetupCommonParameters(FRDGBuilder& GraphBuilder, FVirtualShadowMapAr
 }
 
 
-void FVirtualShadowMapArrayCacheManager::ProcessInvalidations(FRDGBuilder& GraphBuilder, FInstanceGPULoadBalancer& Instances, int32 TotalInstanceCount, const FGPUScene& GPUScene)
+void FVirtualShadowMapArrayCacheManager::ProcessInvalidations(
+	FRDGBuilder& GraphBuilder,
+	FInstanceGPULoadBalancer& Instances,
+	int32 TotalInstanceCount,
+	const FGPUScene& GPUScene)
 {
 	if (Instances.IsEmpty())
 	{
@@ -1060,23 +1251,34 @@ void FVirtualShadowMapArrayCacheManager::ProcessInvalidations(FRDGBuilder& Graph
 }
 
 
-void FVirtualShadowMapArrayCacheManager::ProcessGPUInstanceInvalidations(FRDGBuilder& GraphBuilder, const FGPUScene& GPUScene)
+void FVirtualShadowMapArrayCacheManager::ProcessGPUInstanceInvalidations(
+	FRDGBuilder& GraphBuilder,
+	const FGPUScene& GPUScene)
 {
 	// Dispatch CS indirectly to process instances that are marked to update from the GPU side.
-	if (PrevBuffers.InvalidatingInstancesBuffer.IsValid())
+	if (InvalidatingInstancesBuffer.IsValid() && !bGPUInstanceInvalidationsProcessedThisFrame)
 	{
 		RDG_EVENT_SCOPE(GraphBuilder, "ProcessGPUInstanceInvalidations [GPU-Instances]");
 
-		FRDGBufferRef InvalidatingInstancesBufferRDG = GraphBuilder.RegisterExternalBuffer(PrevBuffers.InvalidatingInstancesBuffer, TEXT("Shadow.Virtual.PrevInvalidatingInstancesBuffer"));
+		FRDGBufferRef InvalidatingInstancesBufferRDG = GraphBuilder.RegisterExternalBuffer(InvalidatingInstancesBuffer, TEXT("Shadow.Virtual.PrevInvalidatingInstancesBuffer"));
 		FRDGBufferRef IndirectArgs = FComputeShaderUtils::AddIndirectArgsSetupCsPass1D(GraphBuilder, GPUScene.GetFeatureLevel(), InvalidatingInstancesBufferRDG, TEXT("Shadow.Virtual.ProcessGPUInstanceInvalidationsIndirectArgs"), FVirtualSmInvalidateInstancePagesCS::Cs1dGroupSizeX);
 
 		FVirtualSmInvalidateInstancePagesCS::FPermutationDomain PermutationVector;
 		FVirtualSmInvalidateInstancePagesCS::FParameters* PassParameters = GraphBuilder.AllocParameters<FVirtualSmInvalidateInstancePagesCS::FParameters>();
 		SetupCommonParameters(GraphBuilder, this, 16*1024, GPUScene, *PassParameters, PermutationVector);
 
+		// Buffer for marking static invalidating instances for readback
+		const uint32 StaticInvalidatingPrimitivesSize = FMath::Max(1, FMath::DivideAndRoundUp(Scene->GetMaxPersistentPrimitiveIndex(), 32));
+		FRDGBufferDesc StaticInvalidatingPrimitivesDesc = FRDGBufferDesc::CreateStructuredDesc(sizeof(uint32), StaticInvalidatingPrimitivesSize);
+		StaticInvalidatingPrimitivesDesc.Usage |= EBufferUsageFlags::SourceCopy;	// For copy to readback
+		FRDGBufferRef StaticInvalidatingPrimitives = GraphBuilder.CreateBuffer(StaticInvalidatingPrimitivesDesc, TEXT("Shadow.Virtual.StaticInvalidatingPrimitives"));
+
+		AddClearUAVPass(GraphBuilder, GraphBuilder.CreateUAV(StaticInvalidatingPrimitives), 0);
+
 		PassParameters->IndirectArgs = IndirectArgs;
 		PassParameters->InvalidatingInstances = GraphBuilder.CreateSRV(InvalidatingInstancesBufferRDG);
-		PassParameters->NumInvalidatingInstanceSlots = PrevBuffers.NumInvalidatingInstanceSlots;
+		PassParameters->NumInvalidatingInstanceSlots = NumInvalidatingInstanceSlots;
+		PassParameters->OutStaticInvalidatingPrimitives = GraphBuilder.CreateUAV(StaticInvalidatingPrimitives);
 
 		PermutationVector.Set<FVirtualSmInvalidateInstancePagesCS::FInputKindDim>(FVirtualSmInvalidateInstancePagesCS::EInputKind_GPUInstances);
 		auto ComputeShader = GetGlobalShaderMap(Scene->GetFeatureLevel())->GetShader<FVirtualSmInvalidateInstancePagesCS>(PermutationVector);
@@ -1090,7 +1292,10 @@ void FVirtualShadowMapArrayCacheManager::ProcessGPUInstanceInvalidations(FRDGBui
 			0
 		);
 
-		// Drop the InvalidatingInstancesBuffer to make sure we don't redundantly process the associated invalidations if ProcessRemovedOrUpdatedPrimitives is called multiple times.
-		PrevBuffers.InvalidatingInstancesBuffer.SafeRelease();
+		// Enqueue readback
+		StaticGPUInvalidationsFeedback.SubmitFeedbackBuffer(GraphBuilder, StaticInvalidatingPrimitives);
+
+		// Make sure we don't redundantly process the associated invalidations if ProcessRemovedOrUpdatedPrimitives is called multiple times.
+		bGPUInstanceInvalidationsProcessedThisFrame = true;
 	}
 }
diff --git a/Engine/Source/Runtime/Renderer/Private/VirtualShadowMaps/VirtualShadowMapCacheManager.h b/Engine/Source/Runtime/Renderer/Private/VirtualShadowMaps/VirtualShadowMapCacheManager.h
--- a/Engine/Source/Runtime/Renderer/Private/VirtualShadowMaps/VirtualShadowMapCacheManager.h
+++ b/Engine/Source/Runtime/Renderer/Private/VirtualShadowMaps/VirtualShadowMapCacheManager.h
@@ -102,6 +102,7 @@ public:
 	{
 		int32 InstanceSceneDataOffset;
 		int32 NumInstanceSceneDataEntries;
+		bool bInvalidateStaticPage;
 	};
 
 	TArray<FInstanceRange> PrimitiveInstancesToInvalidate;
@@ -110,6 +111,28 @@ private:
 	FProjectedShadowInitializer LocalCacheKey;
 };
 
+class FVirtualShadowMapFeedback
+{
+public:
+	FVirtualShadowMapFeedback();
+	~FVirtualShadowMapFeedback();
+
+	struct FReadbackInfo
+	{
+		FRHIGPUBufferReadback* Buffer = nullptr;
+		uint32 Size = 0;
+	};
+
+	void SubmitFeedbackBuffer(FRDGBuilder& GraphBuilder, FRDGBufferRef FeedbackBuffer);
+	FReadbackInfo GetLatestReadbackBuffer();
+
+private:
+	static const int32 MaxBuffers = 3;
+	int32 WriteIndex = 0;
+	int32 NumPending = 0;
+	FReadbackInfo Buffers[MaxBuffers];
+};
+
 // Persistent buffers that we ping pong frame by frame
 struct FVirtualShadowMapArrayFrameData
 {
@@ -126,9 +149,6 @@ struct FVirtualShadowMapArrayFrameData
 	TRefCountPtr<IPooledRenderTarget>			HZBPhysical;
 	TMap<int32, FVirtualShadowMapHZBMetadata>	HZBMetadata;
 
-	TRefCountPtr<FRDGPooledBuffer>				InvalidatingInstancesBuffer;
-	int32										NumInvalidatingInstanceSlots = 0;
-
 	uint64 GetGPUSizeBytes(bool bLogSizes) const;
 };
 
@@ -144,14 +164,15 @@ public:
 	// Called by VirtualShadowMapArray to potentially resize the physical pool
 	// If the requested size is not already the size, all cache data is dropped and the pool is resized.
 	TRefCountPtr<IPooledRenderTarget> SetPhysicalPoolSize(FRDGBuilder& GraphBuilder, FIntPoint RequestedSize, int RequestedArraySize);
-
 	void FreePhysicalPool();
 
 	// Called by VirtualShadowMapArray to potentially resize the HZB physical pool
 	TRefCountPtr<IPooledRenderTarget> SetHZBPhysicalPoolSize(FRDGBuilder& GraphBuilder, FIntPoint RequestedSize, const EPixelFormat Format);
-
 	void FreeHZBPhysicalPool();
 
+	// Called by VirtualShadowMapArray to resize and clear the invalidating instances buffer
+	FVirtualShadowMapSceneData InitializeSceneDataForFrame(FRDGBuilder& GraphBuilder);
+
 	// Invalidate the cache for all shadows, causing any pages to be rerendered
 	void Invalidate();
 
@@ -187,22 +208,50 @@ public:
 		FInvalidatingPrimitiveCollector(FVirtualShadowMapArrayCacheManager* InVirtualShadowMapArrayCacheManager);
 
 		/**
-		 * Add a primitive to invalidate the instances for, the function filters redundant primitive adds, and thus expects valid IDs (so can't be called for primitives that have not yet been added)
+		 * All of these functions filters redundant primitive adds, and thus expects valid IDs (so can't be called for primitives that have not yet been added)
 		 * and unchanging IDs (so can't be used over a span that include any scene mutation).
 		 */
-		void Add(const FPrimitiveSceneInfo* PrimitiveSceneInfo);
+
+		// Primitive was removed from the scene
+		void Removed(FPrimitiveSceneInfo* PrimitiveSceneInfo)
+		{
+			AddInvalidation(PrimitiveSceneInfo, true);
+		}
+
+		// Primitive instances updated
+		void UpdatedInstances(FPrimitiveSceneInfo* PrimitiveSceneInfo)
+		{
+			AddInvalidation(PrimitiveSceneInfo, false);
+		}
+
+		// Primitive moved/transform was updated
+		void UpdatedTransform(FPrimitiveSceneInfo* PrimitiveSceneInfo)
+		{
+			AddInvalidation(PrimitiveSceneInfo, false);
+		}
 
 		bool IsEmpty() const { return LoadBalancer.IsEmpty(); }
 
-		TBitArray<SceneRenderingAllocator> AlreadyAddedPrimitives;
+		const TBitArray<SceneRenderingAllocator>& GetRemovedPrimitives() const
+		{
+			return RemovedPrimitives;
+		}
+
 		FInstanceGPULoadBalancer LoadBalancer;
 		int32 TotalInstanceCount = 0;
 #if VSM_LOG_INVALIDATIONS
 		FString RangesStr;
 #endif
-		const FScene& Scene;
+
+	private:
+		void AddInvalidation(FPrimitiveSceneInfo* PrimitiveSceneInfo, bool bRemovedPrimitive);
+
+		FScene& Scene;
 		FGPUScene& GPUScene;
 		FVirtualShadowMapArrayCacheManager& Manager;
+
+		TBitArray<SceneRenderingAllocator> AlreadyAddedPrimitives;
+		TBitArray<SceneRenderingAllocator> RemovedPrimitives;
 	};
 
 	/**
@@ -226,7 +275,7 @@ public:
 
 	FVirtualShadowMapArrayFrameData PrevBuffers;
 	FVirtualShadowMapUniformParameters PrevUniformParameters;
-		
+			
 	void SetHZBViewParams(int32 HZBKey, Nanite::FPackedViewParams& OutParams);
 
 #if WITH_MGPU
@@ -237,6 +286,8 @@ public:
 
 	GPUMessage::FSocket StatusFeedbackSocket;
 
+	FVirtualShadowMapFeedback StaticGPUInvalidationsFeedback;
+
 private:
 	void ProcessInvalidations(FRDGBuilder& GraphBuilder, FInstanceGPULoadBalancer& Instances, int32 TotalInstanceCount, const FGPUScene& GPUScene);
 
@@ -244,16 +295,47 @@ private:
 
 	void ExtractStats(FRDGBuilder& GraphBuilder, FVirtualShadowMapArray &VirtualShadowMapArray);
 
+	template<typename Allocator>
+	void UpdateRecentlyRemoved(const TBitArray<Allocator>& Removed)
+	{
+		// Combine into *both* flag arrays
+		RecentlyRemovedPrimitives[0].CombineWithBitwiseOR(Removed, EBitwiseOperatorFlags::MaxSize);
+		RecentlyRemovedPrimitives[1].CombineWithBitwiseOR(Removed, EBitwiseOperatorFlags::MaxSize);
+	}
+
+	bool WasRecentlyRemoved(FPersistentPrimitiveIndex PersistentPrimitiveIndex) const
+	{
+		return RecentlyRemovedPrimitives[RecentlyRemovedReadIndex].Num() > PersistentPrimitiveIndex.Index &&
+			RecentlyRemovedPrimitives[RecentlyRemovedReadIndex][PersistentPrimitiveIndex.Index];
+	}
+
 	// The actual physical texture data is stored here rather than in VirtualShadowMapArray (which is recreated each frame)
 	// This allows us to (optionally) persist cached pages between frames. Regardless of whether caching is enabled,
 	// we store the physical pool here.
 	TRefCountPtr<IPooledRenderTarget> PhysicalPagePool;
 	TRefCountPtr<IPooledRenderTarget> HZBPhysicalPagePool;
 
+	// Buffer that stores flags marking each instance that needs to be invalidated the subsequent frame (handled by the cache manager).
+	// This covers things like WPO or GPU-side updates, and any other case where we determine an instance needs to invalidate its footprint. 
+	// Buffer of uints, organized as as follows: InvalidatingInstancesRDG[0] == count, InvalidatingInstancesRDG[1+MaxInstanceCount:1+MaxInstanceCount+MaxInstanceCount/32] == flags, 
+	// InvalidatingInstancesRDG[1:MaxInstanceCount] == growing compact array of instances that need invaldation
+	TRefCountPtr<FRDGPooledBuffer> InvalidatingInstancesBuffer;
+	int32 NumInvalidatingInstanceSlots = -1;
+	bool bGPUInstanceInvalidationsProcessedThisFrame = false;
+
 	// Index the Cache entries by the light ID
 	TMap< uint64, TSharedPtr<FVirtualShadowMapPerLightCacheEntry> > CacheEntries;
 	TMap< uint64, TSharedPtr<FVirtualShadowMapPerLightCacheEntry> > PrevCacheEntries;
 
+	// Tracks primitives (by persistent primitive index) that have been removed recently
+	// This allows us to ignore feedback from previous frames in the case of persistent primitive indices being
+	// reused after being removed. We mark bits in two bitfields, then zero out one of them and switch
+	// to the other each time we loop around the feedback buffers. This is somewhat overly conservative but
+	// relatively lightweight.
+	TBitArray<> RecentlyRemovedPrimitives[2];
+	int32 RecentlyRemovedReadIndex = 0;
+	int32 RecentlyRemovedFrameCounter = 0;
+
 	// Stores stats over frames when activated.
 	TRefCountPtr<FRDGPooledBuffer> AccumulatedStatsBuffer;
 	bool bAccumulatingStats = false;
diff --git a/Engine/Source/Runtime/Renderer/Public/PrimitiveSceneInfo.h b/Engine/Source/Runtime/Renderer/Public/PrimitiveSceneInfo.h
--- a/Engine/Source/Runtime/Renderer/Public/PrimitiveSceneInfo.h
+++ b/Engine/Source/Runtime/Renderer/Public/PrimitiveSceneInfo.h
@@ -598,6 +598,13 @@ public:
 	/** Return primitive fullname (for debugging only). */
 	FString GetFullnameForDebuggingOnly() const;
 
+	FORCEINLINE bool ShouldCacheShadowAsStatic() const
+	{
+		return bCacheShadowAsStatic;
+	}
+
+	void SetCacheShadowAsStatic(bool bStatic);
+
 private:
 
 	/** Let FScene have direct access to the Id. */
@@ -636,6 +643,9 @@ private:
 	/** True if the primitive registered with velocity data and needs to remove itself when being removed from the scene. */
 	bool bRegisteredWithVelocityData : 1;
 
+	/** True if the primitive should be treated as static for the purpose of caching shadows */
+	bool bCacheShadowAsStatic : 1;
+
 	/** Index into the scene's PrimitivesNeedingLevelUpdateNotification array for this primitive scene info level. */
 	int32 LevelUpdateNotificationIndex;
 
