Fix compilation issue with crossing merges/streams.

diff --git a/Engine/Source/Runtime/Renderer/Private/RendererScene.cpp b/Engine/Source/Runtime/Renderer/Private/RendererScene.cpp
--- a/Engine/Source/Runtime/Renderer/Private/RendererScene.cpp
+++ b/Engine/Source/Runtime/Renderer/Private/RendererScene.cpp
@@ -1100,8 +1100,7 @@ uint64 FVirtualShadowMapArrayFrameData::GetGPUSizeBytes(bool bLogSizes) const
 		GetBufferGPUSizeBytes(PageRectBounds, bLogSizes) +
 		GetBufferGPUSizeBytes(DynamicCasterPageFlags, bLogSizes) +
 		GetBufferGPUSizeBytes(PhysicalPageMetaData, bLogSizes) +
-		GetRenderTargetGPUSizeBytes(HZBPhysical, bLogSizes) +
-		GetBufferGPUSizeBytes(InvalidatingInstancesBuffer, bLogSizes);
+		GetRenderTargetGPUSizeBytes(HZBPhysical, bLogSizes);
 };
 
 uint64 FVirtualShadowMapArrayCacheManager::GetGPUSizeBytes(bool bLogSizes) const
