Add settings to FSceneViewFamily.
bOverrideVirtualTextureThrottle will disable the virtual texture page throttling.
VirtualTextureFeedbackFactor will set the feedback factor and defaults to r.vt.FeedbackFactor.
MRQ disables the virtual texture throttling and could override feedback factor based on some user setting.
#rnx

diff --git a/Engine/Plugins/MovieScene/MovieRenderPipeline/Source/MovieRenderPipelineRenderPasses/Private/MoviePipelineImagePassBase.cpp b/Engine/Plugins/MovieScene/MovieRenderPipeline/Source/MovieRenderPipelineRenderPasses/Private/MoviePipelineImagePassBase.cpp
--- a/Engine/Plugins/MovieScene/MovieRenderPipeline/Source/MovieRenderPipelineRenderPasses/Private/MoviePipelineImagePassBase.cpp
+++ b/Engine/Plugins/MovieScene/MovieRenderPipeline/Source/MovieRenderPipelineRenderPasses/Private/MoviePipelineImagePassBase.cpp
@@ -85,6 +85,8 @@ TSharedPtr<FSceneViewFamilyContext> UMoviePipelineImagePassBase::CalculateViewFa
 	OutViewFamily->SceneCaptureSource = InOutSampleState.SceneCaptureSource;
 	OutViewFamily->bWorldIsPaused = InOutSampleState.bWorldIsPaused;
 	OutViewFamily->ViewMode = ViewModeIndex;
+	OutViewFamily->bOverrideVirtualTextureThrottle = true;
+	// OutViewFamily->VirtualTextureFeedbackFactor = 1;
 
 	const bool bIsPerspective = true;
 	ApplyViewMode(OutViewFamily->ViewMode, bIsPerspective, OutViewFamily->EngineShowFlags);
diff --git a/Engine/Source/Runtime/Engine/Private/SceneView.cpp b/Engine/Source/Runtime/Engine/Private/SceneView.cpp
--- a/Engine/Source/Runtime/Engine/Private/SceneView.cpp
+++ b/Engine/Source/Runtime/Engine/Private/SceneView.cpp
@@ -322,6 +322,14 @@ static TAutoConsoleVariable<int32> CVarEnableTemporalUpsample(
 	TEXT(" 1: TemporalAA performs spatial and temporal upscale as screen percentage method (default)."),
 	ECVF_Default);
 
+int32 GVirtualTextureFeedbackFactor = 16;
+static FAutoConsoleVariableRef CVarVirtualTextureFeedbackFactor(
+	TEXT("r.vt.FeedbackFactor"),
+	GVirtualTextureFeedbackFactor,
+	TEXT("The size of the VT feedback buffer is calculated by dividing the render resolution by this factor.")
+	TEXT("The value set here is rounded up to the nearest power of two before use."),
+	ECVF_RenderThreadSafe);
+
 #if !(UE_BUILD_SHIPPING || UE_BUILD_TEST)
 
 static TAutoConsoleVariable<float> CVarOverrideTimeMaterialExpressions(
@@ -2833,6 +2841,9 @@ FSceneViewFamily::FSceneViewFamily(const ConstructionValues& CVS)
 
 	bCurrentlyBeingEdited = false;
 
+	bOverrideVirtualTextureThrottle = false;
+	VirtualTextureFeedbackFactor = GVirtualTextureFeedbackFactor;
+
 #if !WITH_EDITOR
 	check(!EngineShowFlags.StationaryLightOverlap);
 #else
diff --git a/Engine/Source/Runtime/Engine/Public/SceneView.h b/Engine/Source/Runtime/Engine/Public/SceneView.h
--- a/Engine/Source/Runtime/Engine/Public/SceneView.h
+++ b/Engine/Source/Runtime/Engine/Public/SceneView.h
@@ -1827,6 +1827,12 @@ public:
 	/** Whether the scene is currently being edited, which can be used to speed up lighting propagation. */
 	bool bCurrentlyBeingEdited;
 
+	/** Whether to disable virtual texture update throttling. */
+	bool bOverrideVirtualTextureThrottle;
+
+	/** Size in pixels of a virtual texture feedback tile. */
+	int32 VirtualTextureFeedbackFactor;
+
 #if WITH_EDITOR
 	/** Indicates whether, or not, the base attachment volume should be drawn. */
 	bool bDrawBaseInfo;
diff --git a/Engine/Source/Runtime/Renderer/Private/DeferredShadingRenderer.cpp b/Engine/Source/Runtime/Renderer/Private/DeferredShadingRenderer.cpp
--- a/Engine/Source/Runtime/Renderer/Private/DeferredShadingRenderer.cpp
+++ b/Engine/Source/Runtime/Renderer/Private/DeferredShadingRenderer.cpp
@@ -2351,7 +2351,9 @@ void FDeferredShadingSceneRenderer::Render(FRDGBuilder& GraphBuilder)
 	if (bUseVirtualTexturing)
 	{
 		RDG_GPU_STAT_SCOPE(GraphBuilder, VirtualTextureUpdate);
-		FVirtualTextureSystem::Get().Update(GraphBuilder, FeatureLevel, Scene);
+		FVirtualTextureUpdateSettings Settings;
+		Settings.DisableThrottling(ViewFamily.bOverrideVirtualTextureThrottle);
+		FVirtualTextureSystem::Get().Update(GraphBuilder, FeatureLevel, Scene, Settings);
 	}
 
 #if RHI_RAYTRACING
diff --git a/Engine/Source/Runtime/Renderer/Private/MobileShadingRenderer.cpp b/Engine/Source/Runtime/Renderer/Private/MobileShadingRenderer.cpp
--- a/Engine/Source/Runtime/Renderer/Private/MobileShadingRenderer.cpp
+++ b/Engine/Source/Runtime/Renderer/Private/MobileShadingRenderer.cpp
@@ -842,7 +842,9 @@ void FMobileSceneRenderer::Render(FRDGBuilder& GraphBuilder)
 	if (bUseVirtualTexturing)
 	{
 		RDG_GPU_STAT_SCOPE(GraphBuilder, VirtualTextureUpdate);
-		FVirtualTextureSystem::Get().Update(GraphBuilder, FeatureLevel, Scene);
+		FVirtualTextureUpdateSettings Settings;
+		Settings.DisableThrottling(ViewFamily.bOverrideVirtualTextureThrottle);
+		FVirtualTextureSystem::Get().Update(GraphBuilder, FeatureLevel, Scene, Settings);
 	}
 
 	FSortedLightSetSceneInfo& SortedLightSet = *GraphBuilder.AllocObject<FSortedLightSetSceneInfo>();
diff --git a/Engine/Source/Runtime/Renderer/Private/Renderer.cpp b/Engine/Source/Runtime/Renderer/Private/Renderer.cpp
--- a/Engine/Source/Runtime/Renderer/Private/Renderer.cpp
+++ b/Engine/Source/Runtime/Renderer/Private/Renderer.cpp
@@ -226,7 +226,10 @@ void FRendererModule::DrawTileMesh(FCanvasRenderContext& RenderContext, FMeshPas
 			RDG_GPU_STAT_SCOPE(GraphBuilder, VirtualTextureUpdate);
 			FVirtualTextureSystem::Get().AllocateResources(GraphBuilder, FeatureLevel);
 			FVirtualTextureSystem::Get().CallPendingCallbacks();
-			FVirtualTextureSystem::Get().Update(GraphBuilder, FeatureLevel, Scene);
+
+			FVirtualTextureUpdateSettings Settings;
+			Settings.DisableThrottling(true);
+			FVirtualTextureSystem::Get().Update(GraphBuilder, FeatureLevel, Scene, Settings);
 
 			VirtualTextureFeedbackBegin(GraphBuilder, TArrayView<const FViewInfo>(&View, 1), RenderContext.GetViewportRect().Size());
 		}
diff --git a/Engine/Source/Runtime/Renderer/Private/SceneRendering.cpp b/Engine/Source/Runtime/Renderer/Private/SceneRendering.cpp
--- a/Engine/Source/Runtime/Renderer/Private/SceneRendering.cpp
+++ b/Engine/Source/Runtime/Renderer/Private/SceneRendering.cpp
@@ -1809,13 +1809,13 @@ void FViewInfo::SetupUniformBufferParameters(
 
 	ViewUniformShaderParameters.GlobalVirtualTextureMipBias = FVirtualTextureSystem::Get().GetGlobalMipBias();
 
-	const uint32 VirtualTextureFeedbackScale = GetVirtualTextureFeedbackScale();
+	const uint32 VirtualTextureFeedbackScale = GetVirtualTextureFeedbackScale(Family);
 	check(VirtualTextureFeedbackScale == 1 << FMath::FloorLog2(VirtualTextureFeedbackScale));
 	ViewUniformShaderParameters.VirtualTextureFeedbackShift = FMath::FloorLog2(VirtualTextureFeedbackScale);
 	ViewUniformShaderParameters.VirtualTextureFeedbackMask = VirtualTextureFeedbackScale - 1;
-	ViewUniformShaderParameters.VirtualTextureFeedbackStride = GetVirtualTextureFeedbackBufferSize(EffectiveBufferSize).X;
+	ViewUniformShaderParameters.VirtualTextureFeedbackStride = GetVirtualTextureFeedbackBufferSize(Family, EffectiveBufferSize).X;
 	// Use some low(ish) discrepancy sequence to run over every pixel in the virtual texture feedback tile.
-	ViewUniformShaderParameters.VirtualTextureFeedbackJitterOffset = SampleVirtualTextureFeedbackSequence(FrameIndex);
+	ViewUniformShaderParameters.VirtualTextureFeedbackJitterOffset = SampleVirtualTextureFeedbackSequence(Family, FrameIndex);
 	// Offset the selected sample index for each frame and add an additional offset each time we iterate over a full virtual texture feedback tile to ensure we get full coverage of sample indices over time.
 	const uint32 NumPixelsInTile = FMath::Square(VirtualTextureFeedbackScale);
 	ViewUniformShaderParameters.VirtualTextureFeedbackSampleOffset = (FrameIndex % NumPixelsInTile) + (FrameIndex / NumPixelsInTile);
@@ -5484,7 +5484,7 @@ void VirtualTextureFeedbackBegin(FRDGBuilder& GraphBuilder, TArrayView<const FVi
 	}
 
 	FVirtualTextureFeedbackBufferDesc Desc;
-	Desc.Init2D(SceneTextureExtent, ViewRects, GetVirtualTextureFeedbackScale());
+	Desc.Init2D(SceneTextureExtent, ViewRects, GetVirtualTextureFeedbackScale(Views.Num() > 0 ? Views[0].Family : nullptr));
 	GVirtualTextureFeedbackBuffer.Begin(GraphBuilder, Desc);
 }
 
diff --git a/Engine/Source/Runtime/Renderer/Private/VT/VirtualTextureFeedbackBuffer.cpp b/Engine/Source/Runtime/Renderer/Private/VT/VirtualTextureFeedbackBuffer.cpp
--- a/Engine/Source/Runtime/Renderer/Private/VT/VirtualTextureFeedbackBuffer.cpp
+++ b/Engine/Source/Runtime/Renderer/Private/VT/VirtualTextureFeedbackBuffer.cpp
@@ -3,31 +3,24 @@
 #include "VT/VirtualTextureFeedbackBuffer.h"
 
 #include "RenderGraphUtils.h"
+#include "SceneView.h"
 #include "VT/VirtualTextureFeedback.h"
 
-int32 GVirtualTextureFeedbackFactor = 16;
-static FAutoConsoleVariableRef CVarVirtualTextureFeedbackFactor(
-	TEXT("r.vt.FeedbackFactor"),
-	GVirtualTextureFeedbackFactor,
-	TEXT("The size of the VT feedback buffer is calculated by dividing the render resolution by this factor.")
-	TEXT("The value set here is rounded up to the nearest power of two before use."),
-	ECVF_RenderThreadSafe
-);
-
-int32 GetVirtualTextureFeedbackScale()
+int32 GetVirtualTextureFeedbackScale(FSceneViewFamily const* InViewFamily)
 {
 	// Round to nearest power of two to ensure that shader maths is efficient and sampling sequence logic is simple.
-	return FMath::RoundUpToPowerOfTwo(FMath::Max(GVirtualTextureFeedbackFactor, 1));
+	const int32 VirtualTextureFeedbackFactor = InViewFamily ? InViewFamily->VirtualTextureFeedbackFactor : 0;
+	return FMath::RoundUpToPowerOfTwo(FMath::Max(VirtualTextureFeedbackFactor, 1));
 }
 
-FIntPoint GetVirtualTextureFeedbackBufferSize(FIntPoint InSceneTextureExtent)
+FIntPoint GetVirtualTextureFeedbackBufferSize(FSceneViewFamily const* InViewFamily, FIntPoint InSceneTextureExtent)
 {
-	return FIntPoint::DivideAndRoundUp(InSceneTextureExtent, GetVirtualTextureFeedbackScale());
+	return FIntPoint::DivideAndRoundUp(InSceneTextureExtent, GetVirtualTextureFeedbackScale(InViewFamily));
 }
 
-uint32 SampleVirtualTextureFeedbackSequence(uint32 InFrameIndex)
+uint32 SampleVirtualTextureFeedbackSequence(FSceneViewFamily const* InViewFamily, uint32 InFrameIndex)
 {
-	const uint32 TileSize = GetVirtualTextureFeedbackScale();
+	const uint32 TileSize = GetVirtualTextureFeedbackScale(InViewFamily);
 	const uint32 TileSizeLog2 = FMath::CeilLogTwo(TileSize);
 	const uint32 SequenceSize = FMath::Square(TileSize);
 	const uint32 PixelIndex = InFrameIndex % SequenceSize;
diff --git a/Engine/Source/Runtime/Renderer/Private/VT/VirtualTextureSystem.cpp b/Engine/Source/Runtime/Renderer/Private/VT/VirtualTextureSystem.cpp
--- a/Engine/Source/Runtime/Renderer/Private/VT/VirtualTextureSystem.cpp
+++ b/Engine/Source/Runtime/Renderer/Private/VT/VirtualTextureSystem.cpp
@@ -159,6 +159,21 @@ static TAutoConsoleVariable<int32> CVarVTCsvStats(
 );
 
 
+FVirtualTextureUpdateSettings::FVirtualTextureUpdateSettings()
+{
+	bEnableFeedback = CVarVTEnableFeedback.GetValueOnRenderThread() != 0;
+	bEnablePlayback = CVarVTEnablePlayback.GetValueOnRenderThread() != 0;
+	bForceContinuousUpdate = CVarVTForceContinuousUpdate.GetValueOnRenderThread() != 0;
+	bParallelFeedbackTasks = CVarVTParallelFeedbackTasks.GetValueOnRenderThread() != 0;
+	NumFeedbackTasks = CVarVTNumFeedbackTasks.GetValueOnRenderThread();
+	NumGatherTasks = CVarVTNumGatherTasks.GetValueOnRenderThread();
+	MaxGatherPagesBeforeFlush = CVarVTPageUpdateFlushCount.GetValueOnRenderThread();
+	MaxPageUploads = VirtualTextureScalability::GetMaxUploadsPerFrame();
+	MaxPagesProduced = VirtualTextureScalability::GetMaxPagesProducedPerFrame();
+	MaxContinuousUpdates = VirtualTextureScalability::GetMaxContinuousUpdatesPerFrame();
+}
+
+
 static FORCEINLINE uint32 EncodePage(uint32 ID, uint32 vLevel, uint32 vTileX, uint32 vTileY)
 {
 	const uint32 vLevelPlus1 = vLevel + 1u;
@@ -256,6 +271,7 @@ struct FGatherRequestsParameters
 	uint32 PageStartIndex = 0u;
 	uint32 NumPages = 0u;
 	uint32 FrameRequested;
+	bool bForceContinuousUpdate = false;
 };
 
 class FGatherRequestsTask
@@ -1202,6 +1218,7 @@ void FVirtualTextureSystem::LoadPendingTiles(FRDGBuilder& GraphBuilder, ERHIFeat
 
 	if (PackedTiles.Num() > 0)
 	{
+		FVirtualTextureUpdateSettings Settings;
 		FConcurrentLinearBulkObjectAllocator Allocator;
 
 		FUniquePageList* UniquePageList = Allocator.Create<FUniquePageList>();
@@ -1213,10 +1230,10 @@ void FVirtualTextureSystem::LoadPendingTiles(FRDGBuilder& GraphBuilder, ERHIFeat
 
 		FUniqueRequestList* RequestList = Allocator.Create<FUniqueRequestList>(Allocator);
 		RequestList->Initialize();
-		GatherRequests(RequestList, UniquePageList, Frame, Allocator);
+		GatherRequests(RequestList, UniquePageList, Frame, Allocator, Settings);
 		// No need to sort requests, since we're submitting all of them here (no throttling)
 		AllocateResources(GraphBuilder, FeatureLevel);
-		SubmitRequests(GraphBuilder, FeatureLevel, Allocator, RequestList, false);
+		SubmitRequests(GraphBuilder, FeatureLevel, Allocator, Settings, RequestList, false);
 	}
 }
 
@@ -1323,7 +1340,7 @@ void FVirtualTextureSystem::FeedbackAnalysisTask(const FFeedbackAnalysisParamete
 	}
 }
 
-void FVirtualTextureSystem::Update(FRDGBuilder& GraphBuilder, ERHIFeatureLevel::Type FeatureLevel, FScene* Scene)
+void FVirtualTextureSystem::Update(FRDGBuilder& GraphBuilder, ERHIFeatureLevel::Type FeatureLevel, FScene* Scene, FVirtualTextureUpdateSettings const& Settings)
 {
 	check(IsInRenderingThread());
 
@@ -1393,7 +1410,7 @@ void FVirtualTextureSystem::Update(FRDGBuilder& GraphBuilder, ERHIFeatureLevel::
 	FUniquePageList* MergedUniquePageList = Allocator.Create<FUniquePageList>();
 	MergedUniquePageList->Initialize();
 	
-	if (CVarVTEnableFeedback.GetValueOnRenderThread())
+	if (Settings.bEnableFeedback)
 	{
 		// Fetch feedback for analysis
 		FVirtualTextureFeedback::FMapResult FeedbackResult;
@@ -1407,7 +1424,7 @@ void FVirtualTextureSystem::Update(FRDGBuilder& GraphBuilder, ERHIFeatureLevel::
 		// Give each task a section of the feedback buffer to analyze
 		FFeedbackAnalysisParameters FeedbackAnalysisParameters[MaxNumTasks];
 
-		const uint32 MaxNumFeedbackTasks = FMath::Clamp((uint32)CVarVTNumFeedbackTasks.GetValueOnRenderThread(), 1u, MaxNumTasks);
+		const uint32 MaxNumFeedbackTasks = FMath::Clamp((uint32)Settings.NumFeedbackTasks, 1u, MaxNumTasks);
 		const uint32 FeedbackSizePerTask = FMath::DivideAndRoundUp(FeedbackResult.Size, MaxNumFeedbackTasks);
 
 		uint32 NumFeedbackTasks = 0;
@@ -1433,8 +1450,7 @@ void FVirtualTextureSystem::Update(FRDGBuilder& GraphBuilder, ERHIFeatureLevel::
 		}
 
 		// Kick the tasks
-		const bool bParallelTasks = CVarVTParallelFeedbackTasks.GetValueOnRenderThread() != 0;
-		const int32 LocalFeedbackTaskCount = bParallelTasks ? 1 : NumFeedbackTasks;
+		const int32 LocalFeedbackTaskCount = Settings.bParallelFeedbackTasks ? 1 : NumFeedbackTasks;
 		const int32 WorkerFeedbackTaskCount = NumFeedbackTasks - LocalFeedbackTaskCount;
 
 		FGraphEventArray Tasks;
@@ -1487,7 +1503,7 @@ void FVirtualTextureSystem::Update(FRDGBuilder& GraphBuilder, ERHIFeatureLevel::
 	// Add any page requests from recording playback.
 	if (PageRequestPlaybackBuffer.Num() > 0)
 	{
-		if (CVarVTEnablePlayback.GetValueOnRenderThread())
+		if (Settings.bEnablePlayback)
 		{
 			// todo: We can split this into concurrent tasks. 
 			FAddRequestedTilesParameters Parameters;
@@ -1567,14 +1583,14 @@ void FVirtualTextureSystem::Update(FRDGBuilder& GraphBuilder, ERHIFeatureLevel::
 		{
 			RequestedPageList->Add(Tile, 0xffff);
 		}
-		GatherRequests(MergedRequestList, RequestedPageList, Frame, Allocator);
+		GatherRequests(MergedRequestList, RequestedPageList, Frame, Allocator, Settings);
 	}
 
 	// Pages from feedback buffer were generated several frames ago, so they may no longer be valid for newly allocated VTs
 	static uint32 PendingFrameDelay = 3u;
 	if (Frame >= PendingFrameDelay)
 	{
-		GatherRequests(MergedRequestList, MergedUniquePageList, Frame - PendingFrameDelay, Allocator);
+		GatherRequests(MergedRequestList, MergedUniquePageList, Frame - PendingFrameDelay, Allocator, Settings);
 	}
 
 	if (MergedRequestList->GetNumAdaptiveAllocationRequests() > 0)
@@ -1588,7 +1604,7 @@ void FVirtualTextureSystem::Update(FRDGBuilder& GraphBuilder, ERHIFeatureLevel::
 
 		// Limit the number of uploads (account for MappedTilesToProduce this frame)
 		// Are all pages equal? Should there be different limits on different types of pages?
-		const int32 MaxNumUploads = VirtualTextureScalability::GetMaxUploadsPerFrame();
+		const int32 MaxNumUploads = Settings.MaxPageUploads;
 		const int32 MaxRequestUploads = FMath::Max(MaxNumUploads - MappedTilesToProduce.Num(), 1);
 
 		if (MaxRequestUploads < (int32)MergedRequestList->GetNumLoadRequests())
@@ -1605,16 +1621,17 @@ void FVirtualTextureSystem::Update(FRDGBuilder& GraphBuilder, ERHIFeatureLevel::
 
 	{
 		// After sorting and clamping the load requests, if we still have unused upload bandwidth then use it to add some continous updates
-		const int32 MaxNumUploads = VirtualTextureScalability::GetMaxUploadsPerFrame();
+		const int32 MaxNumUploads = Settings.MaxPageUploads;
 		const int32 MaxTilesToProduce = FMath::Max(MaxNumUploads - MappedTilesToProduce.Num() - (int32)MergedRequestList->GetNumLoadRequests(), 0);
+		const int32 MaxContinuousUpdates = Settings.MaxContinuousUpdates;
 
-		GetContinuousUpdatesToProduce(MergedRequestList, MaxTilesToProduce);
+		GetContinuousUpdatesToProduce(MergedRequestList, MaxTilesToProduce, MaxContinuousUpdates);
 	}
 
 	// Submit the requests to produce pages that are already mapped
 	SubmitPreMappedRequests(GraphBuilder, FeatureLevel);
 	// Submit the merged requests
-	SubmitRequests(GraphBuilder, FeatureLevel, Allocator, MergedRequestList, true);
+	SubmitRequests(GraphBuilder, FeatureLevel, Allocator, Settings, MergedRequestList, true);
 
 	Producers.NotifyRequestsCompleted();
 
@@ -1627,10 +1644,10 @@ void FVirtualTextureSystem::Update(FRDGBuilder& GraphBuilder, ERHIFeatureLevel::
 	ReleasePendingSpaces();
 }
 
-void FVirtualTextureSystem::GatherRequests(FUniqueRequestList* MergedRequestList, const FUniquePageList* UniquePageList, uint32 FrameRequested, FConcurrentLinearBulkObjectAllocator& Allocator)
+void FVirtualTextureSystem::GatherRequests(FUniqueRequestList* MergedRequestList, const FUniquePageList* UniquePageList, uint32 FrameRequested, FConcurrentLinearBulkObjectAllocator& Allocator, FVirtualTextureUpdateSettings const& Settings)
 {
-	const uint32 MaxNumGatherTasks = FMath::Clamp((uint32)CVarVTNumGatherTasks.GetValueOnRenderThread(), 1u, MaxNumTasks);
-	const uint32 PageUpdateFlushCount = FMath::Min<uint32>(CVarVTPageUpdateFlushCount.GetValueOnRenderThread(), FPageUpdateBuffer::PageCapacity);
+	const uint32 MaxNumGatherTasks = FMath::Clamp((uint32)Settings.NumGatherTasks, 1u, MaxNumTasks);
+	const uint32 PageUpdateFlushCount = FMath::Min<uint32>(Settings.MaxGatherPagesBeforeFlush, FPageUpdateBuffer::PageCapacity);
 
 	FGatherRequestsParameters GatherRequestsParameters[MaxNumTasks];
 	uint32 NumGatherTasks = 0u;
@@ -1648,6 +1665,7 @@ void FVirtualTextureSystem::GatherRequests(FUniqueRequestList* MergedRequestList
 				FGatherRequestsParameters& Params = GatherRequestsParameters[TaskIndex];
 				Params.System = this;
 				Params.FrameRequested = FrameRequested;
+				Params.bForceContinuousUpdate = Settings.bForceContinuousUpdate;
 				Params.UniquePageList = UniquePageList;
 				Params.PageUpdateFlushCount = PageUpdateFlushCount;
 				Params.PageUpdateBuffers = Allocator.CreateArray<FPageUpdateBuffer>(PhysicalSpaces.Num());
@@ -1763,14 +1781,13 @@ void FVirtualTextureSystem::GatherRequestsTask(const FGatherRequestsParameters&
 	FUniqueRequestList* RESTRICT RequestList = Parameters.RequestList;
 	const uint32 PageUpdateFlushCount = Parameters.PageUpdateFlushCount;
 	const uint32 PageEndIndex = Parameters.PageStartIndex + Parameters.NumPages;
+	const bool bForceContinuousUpdate = Parameters.bForceContinuousUpdate;
 
 	uint32 NumRequestsPages = 0u;
 	uint32 NumResidentPages = 0u;
 	uint32 NumNonResidentPages = 0u;
 	uint32 NumPrefetchPages = 0u;
 
-	const bool bForceContinuousUpdate = CVarVTForceContinuousUpdate.GetValueOnRenderThread() != 0;
-
 	for (uint32 i = Parameters.PageStartIndex; i < PageEndIndex; ++i)
 	{
 		const uint32 PageEncoded = UniquePageList->GetPage(i);
@@ -2145,13 +2162,12 @@ void FVirtualTextureSystem::GatherRequestsTask(const FGatherRequestsParameters&
 	INC_DWORD_STAT_BY(STAT_NumPagePrefetch, NumPrefetchPages);
 }
 
-void FVirtualTextureSystem::GetContinuousUpdatesToProduce(FUniqueRequestList const* RequestList, int32 MaxTilesToProduce)
+void FVirtualTextureSystem::GetContinuousUpdatesToProduce(FUniqueRequestList const* RequestList, int32 MaxTilesToProduce, int32 MaxContinuousUpdates)
 {
 	const int32 NumContinuousUpdateRequests = (int32)RequestList->GetNumContinuousUpdateRequests();
-	const int32 MaxContinuousUpdatesPerFrame = VirtualTextureScalability::GetMaxContinuousUpdatesPerFrame();
 	
 	// Negative maximum continous updates allows for uncapped requests
-	if (MaxContinuousUpdatesPerFrame < 0)
+	if (MaxContinuousUpdates < 0)
 	{
 		for (int32 i = 0; i < NumContinuousUpdateRequests; ++i)
 		{
@@ -2160,7 +2176,7 @@ void FVirtualTextureSystem::GetContinuousUpdatesToProduce(FUniqueRequestList con
 	}
 	else
 	{
-		const int32 MaxContinousUpdates = FMath::Min(MaxContinuousUpdatesPerFrame, NumContinuousUpdateRequests);
+		const int32 MaxContinousUpdates = FMath::Min(MaxContinuousUpdates, NumContinuousUpdateRequests);
 
 		int32 NumContinuousUpdates = 0;
 		while (NumContinuousUpdates < MaxContinousUpdates && ContinuousUpdateTilesToProduce.Num() < MaxTilesToProduce)
@@ -2281,7 +2297,7 @@ void FVirtualTextureSystem::SubmitPreMappedRequests(FRDGBuilder& GraphBuilder, E
 	}
 }
 
-void FVirtualTextureSystem::SubmitRequests(FRDGBuilder& GraphBuilder, ERHIFeatureLevel::Type FeatureLevel, FConcurrentLinearBulkObjectAllocator& Allocator, FUniqueRequestList* RequestList, bool bAsync)
+void FVirtualTextureSystem::SubmitRequests(FRDGBuilder& GraphBuilder, ERHIFeatureLevel::Type FeatureLevel, FConcurrentLinearBulkObjectAllocator& Allocator, FVirtualTextureUpdateSettings const& Settings, FUniqueRequestList* RequestList, bool bAsync)
 {
 	TRACE_CPUPROFILER_EVENT_SCOPE(FVirtualTextureSystem::SubmitRequests);
 	LLM_SCOPE(ELLMTag::VirtualTextureSystem);
@@ -2320,7 +2336,7 @@ void FVirtualTextureSystem::SubmitRequests(FRDGBuilder& GraphBuilder, ERHIFeatur
 
 		bool bWaitForProducers = false;
 
-		const uint32 MaxPagesProduced = VirtualTextureScalability::GetMaxPagesProducedPerFrame();
+		const uint32 MaxPagesProduced = Settings.MaxPagesProduced;
 		const uint32 PageFreeThreshold = VirtualTextureScalability::GetPageFreeThreshold();
 		uint32 NumStacksProduced = 0u;
 		uint32 NumPagesProduced = 0u;
diff --git a/Engine/Source/Runtime/Renderer/Private/VT/VirtualTextureSystem.h b/Engine/Source/Runtime/Renderer/Private/VT/VirtualTextureSystem.h
--- a/Engine/Source/Runtime/Renderer/Private/VT/VirtualTextureSystem.h
+++ b/Engine/Source/Runtime/Renderer/Private/VT/VirtualTextureSystem.h
@@ -31,6 +31,33 @@ struct FPageUpdateBuffer;
 
 extern uint32 GetTypeHash(const FAllocatedVTDescription& Description);
 
+struct FVirtualTextureUpdateSettings
+{
+	FVirtualTextureUpdateSettings();
+
+	/** Force settings so that throttling page uploads is effectively disabled. */
+	FVirtualTextureUpdateSettings& DisableThrottling(bool bDisable)
+	{
+		if (bDisable)
+		{
+			MaxPageUploads = 99999;
+			MaxPagesProduced = 99999;
+		}
+		return *this;
+	}
+
+	bool bEnableFeedback;
+	bool bEnablePlayback;
+	bool bForceContinuousUpdate;
+	bool bParallelFeedbackTasks;
+	int32 NumFeedbackTasks;
+	int32 NumGatherTasks;
+	int32 MaxGatherPagesBeforeFlush;
+	int32 MaxPageUploads;
+	int32 MaxPagesProduced;
+	int32 MaxContinuousUpdates;
+};
+
 class FVirtualTextureSystem
 {
 public:
@@ -43,7 +70,7 @@ public:
 	void AllocateResources(FRDGBuilder& GraphBuilder, ERHIFeatureLevel::Type FeatureLevel);
 	void FinalizeResources(FRDGBuilder& GraphBuilder, ERHIFeatureLevel::Type FeatureLevel);
 	void CallPendingCallbacks();
-	void Update( FRDGBuilder& GraphBuilder, ERHIFeatureLevel::Type FeatureLevel, FScene* Scene);
+	void Update(FRDGBuilder& GraphBuilder, ERHIFeatureLevel::Type FeatureLevel, FScene* Scene, FVirtualTextureUpdateSettings const& Settings);
 	void ReleasePendingResources();
 
 	IAllocatedVirtualTexture* AllocateVirtualTexture(const FAllocatedVTDescription& Desc);
@@ -107,9 +134,9 @@ private:
 
 	void SubmitPreMappedRequests(FRDGBuilder& GraphBuilder, ERHIFeatureLevel::Type FeatureLevel);
 
-	void SubmitRequests(FRDGBuilder& GraphBuilder, ERHIFeatureLevel::Type FeatureLevel, FConcurrentLinearBulkObjectAllocator& Allocator, FUniqueRequestList* RequestList, bool bAsync);
+	void SubmitRequests(FRDGBuilder& GraphBuilder, ERHIFeatureLevel::Type FeatureLevel, FConcurrentLinearBulkObjectAllocator& Allocator, FVirtualTextureUpdateSettings const& Settings, FUniqueRequestList* RequestList, bool bAsync);
 
-	void GatherRequests(FUniqueRequestList* MergedRequestList, const FUniquePageList* UniquePageList, uint32 FrameRequested, FConcurrentLinearBulkObjectAllocator& Allocator);
+	void GatherRequests(FUniqueRequestList* MergedRequestList, const FUniquePageList* UniquePageList, uint32 FrameRequested, FConcurrentLinearBulkObjectAllocator& Allocator, FVirtualTextureUpdateSettings const& Settings);
 
 	void AddPageUpdate(FPageUpdateBuffer* Buffers, uint32 FlushCount, uint32 PhysicalSpaceID, uint16 pAddress);
 
@@ -117,7 +144,7 @@ private:
 	void AddRequestedTilesTask(const FAddRequestedTilesParameters& Parameters);
 	void GatherRequestsTask(const FGatherRequestsParameters& Parameters);
 
-	void GetContinuousUpdatesToProduce(FUniqueRequestList const* RequestList, int32 MaxTilesToProduce);
+	void GetContinuousUpdatesToProduce(FUniqueRequestList const* RequestList, int32 MaxTilesToProduce, int32 MaxContinuousUpdates);
 
 	void UpdateResidencyTracking() const;
 
diff --git a/Engine/Source/Runtime/Renderer/Public/VT/VirtualTextureFeedbackBuffer.h b/Engine/Source/Runtime/Renderer/Public/VT/VirtualTextureFeedbackBuffer.h
--- a/Engine/Source/Runtime/Renderer/Public/VT/VirtualTextureFeedbackBuffer.h
+++ b/Engine/Source/Runtime/Renderer/Public/VT/VirtualTextureFeedbackBuffer.h
@@ -13,15 +13,16 @@
 #include "Templates/RefCounting.h"
 
 class FRDGBuilder;
+class FSceneViewFamily;
 
 /** Returns the scale used to divide the scene texture resolution to get the virtual texture feedback resolution. */
-extern int32 GetVirtualTextureFeedbackScale();
+extern int32 GetVirtualTextureFeedbackScale(FSceneViewFamily const* InViewFamily);
 
 /** Get size for a virtual texture feedback buffer. This is calculated internally by applying GetVirtualTextureFeedbackScale() to the SceneTextureExtent. */
-extern FIntPoint GetVirtualTextureFeedbackBufferSize(FIntPoint InSceneTextureExtent);
+extern FIntPoint GetVirtualTextureFeedbackBufferSize(FSceneViewFamily const* InViewFamily, FIntPoint InSceneTextureExtent);
 
 /** Returns a value from the sampling sequence used to traverse each feedback tile. */
-extern uint32 SampleVirtualTextureFeedbackSequence(uint32 InFrameIndex);
+extern uint32 SampleVirtualTextureFeedbackSequence(FSceneViewFamily const* InViewFamily, uint32 InFrameIndex);
 
 /** 
  * Description of how to interpret an RHIBuffer that is being fed to the virtual texture feedback system.
