Strata - advanced strata shader debug is now put into a slice of the strata material buffer.

diff --git a/Engine/Shaders/Private/BasePassPixelShader.usf b/Engine/Shaders/Private/BasePassPixelShader.usf
--- a/Engine/Shaders/Private/BasePassPixelShader.usf
+++ b/Engine/Shaders/Private/BasePassPixelShader.usf
@@ -1631,7 +1631,7 @@ void FPixelShaderInOut_MainPS(
 		FStrataIntegrationSettings Settings = InitStrataIntegrationSettings(false /*bForceFullyRough*/, StrataStruct.bRoughDiffuse, StrataStruct.PeelLayersAboveDepth);
 
 	#if STRATA_ADVANCED_DEBUG_ENABLED
-		Settings.SliceStoringDebugStrataTree = all(uint2(float2(View.CursorPosition) * View.ViewResolutionFraction) == PixelPos) ? StrataStruct.SliceStoringDebugStrataTree : Settings.SliceStoringDebugStrataTree;
+		Settings.SliceStoringDebugStrataTreeData = all(uint2(float2(View.CursorPosition) * View.ViewResolutionFraction) == PixelPos) ? StrataStruct.SliceStoringDebugStrataTreeDataWithoutMRT : Settings.SliceStoringDebugStrataTreeData;
 	#endif
 
 		// Generate the strata material data to write out
diff --git a/Engine/Shaders/Private/Strata/Strata.ush b/Engine/Shaders/Private/Strata/Strata.ush
--- a/Engine/Shaders/Private/Strata/Strata.ush
+++ b/Engine/Shaders/Private/Strata/Strata.ush
@@ -1055,11 +1055,11 @@ void StrataPackOutStrataTree(
 	RWTexture2DArray<uint> ExtraMaterialDataUAV,
 	in uint RootOperatorIndex,
 	in FStrataTree StrataTree,
-	in int SliceStoringDebugStrataTree)
+	in int SliceStoringDebugStrataTreeData)
 {
-#define WRITE_INT(Data) ExtraMaterialDataUAV[uint3(OutCoord++, OutCoordY, SliceStoringDebugStrataTree)] = asuint(Data);
-#define WRITE_UINT(Data) ExtraMaterialDataUAV[uint3(OutCoord++, OutCoordY, SliceStoringDebugStrataTree)] = Data;
-#define WRITE_FLOAT(Data) ExtraMaterialDataUAV[uint3(OutCoord++, OutCoordY, SliceStoringDebugStrataTree)] = asuint(Data);
+#define WRITE_INT(Data) ExtraMaterialDataUAV[uint3(OutCoord++, OutCoordY, SliceStoringDebugStrataTreeData)] = asuint(Data);
+#define WRITE_UINT(Data) ExtraMaterialDataUAV[uint3(OutCoord++, OutCoordY, SliceStoringDebugStrataTreeData)] = Data;
+#define WRITE_FLOAT(Data) ExtraMaterialDataUAV[uint3(OutCoord++, OutCoordY, SliceStoringDebugStrataTreeData)] = asuint(Data);
 #define WRITE_FLOAT3(Data) WRITE_FLOAT(Data.x);WRITE_FLOAT(Data.y);WRITE_FLOAT(Data.z);
 #define WRITE_FLOAT4(Data) WRITE_FLOAT(Data.x);WRITE_FLOAT(Data.y);WRITE_FLOAT(Data.z);;WRITE_FLOAT(Data.w);
 
@@ -1148,16 +1148,16 @@ struct FStrataTreeHeader
 	uint RootOperatorIndex;
 };
 
-#define READ_INT(Data) Data = asint(StrataMaterialData[uint3(OutCoord++, OutCoordY, SliceStoringDebugStrataTree)]);
-#define READ_UINT(Data) Data = StrataMaterialData[uint3(OutCoord++, OutCoordY, SliceStoringDebugStrataTree)];
-#define READ_FLOAT(Data) Data = asfloat(StrataMaterialData[uint3(OutCoord++, OutCoordY, SliceStoringDebugStrataTree)]);
+#define READ_INT(Data) Data = asint(StrataMaterialData[uint3(OutCoord++, OutCoordY, SliceStoringDebugStrataTreeData)]);
+#define READ_UINT(Data) Data = StrataMaterialData[uint3(OutCoord++, OutCoordY, SliceStoringDebugStrataTreeData)];
+#define READ_FLOAT(Data) Data = asfloat(StrataMaterialData[uint3(OutCoord++, OutCoordY, SliceStoringDebugStrataTreeData)]);
 #define READ_FLOAT3(Data) READ_FLOAT(Data.x);READ_FLOAT(Data.y);READ_FLOAT(Data.z);
 #define READ_FLOAT4(Data) READ_FLOAT(Data.x);READ_FLOAT(Data.y);READ_FLOAT(Data.z);READ_FLOAT(Data.w);
 
 void StrataUnpackInStrataTreeHeader(
 	Texture2DArray<uint> StrataMaterialData,
 	inout FStrataTreeHeader StrataTreeHeader,
-	in int SliceStoringDebugStrataTree)
+	in int SliceStoringDebugStrataTreeData)
 {
 	StrataTreeHeader = (FStrataTreeHeader)0;
 
@@ -1174,7 +1174,7 @@ FStrataBSDF StrataUnpackInStrataTreeBSDF(
 	uint BSDFIndex,
 	Texture2DArray<uint> StrataMaterialData,
 	inout FStrataTreeHeader StrataTreeHeader,
-	in int SliceStoringDebugStrataTree)
+	in int SliceStoringDebugStrataTreeData)
 {
 	int OutCoord = BSDFIndex * (STRATA_INLINE_SHADING ? 42 : 33);		// Skip BSDFs
 	int OutCoordY= 1;
@@ -1214,7 +1214,7 @@ FStrataOperator StrataUnpackInStrataTreeOperator(
 	uint OperatorIndex,
 	Texture2DArray<uint> StrataMaterialData,
 	inout FStrataTreeHeader StrataTreeHeader,
-	in int SliceStoringDebugStrataTree)
+	in int SliceStoringDebugStrataTreeData)
 {
 	int OutCoord = OperatorIndex * 8;	// Skip Operators
 	int OutCoordY = 2;
@@ -3392,7 +3392,7 @@ struct FStrataIntegrationSettings
 	bool bRoughDiffuseEnabled;
 
 	// Embbeded debug settings
-	int  SliceStoringDebugStrataTree;
+	int  SliceStoringDebugStrataTreeData;
 	int  PeelLayersAboveDepth;
 };
 
@@ -3403,7 +3403,7 @@ FStrataIntegrationSettings InitStrataIntegrationSettings(bool bForceFullyRough,
 	Out.bRoughDiffuseEnabled = bRoughDiffuseEnabled;
 
 	Out.PeelLayersAboveDepth = PeelLayersAboveDepth;
-	Out.SliceStoringDebugStrataTree = -1;
+	Out.SliceStoringDebugStrataTreeData = -1;
 	return Out;
 }
 
diff --git a/Engine/Shaders/Private/Strata/StrataExport.ush b/Engine/Shaders/Private/Strata/StrataExport.ush
--- a/Engine/Shaders/Private/Strata/StrataExport.ush
+++ b/Engine/Shaders/Private/Strata/StrataExport.ush
@@ -147,9 +147,9 @@ void PackStrataOut(
 		const uint RootOperatorIndex = Strata.OperatorIndex;
 
 	#if STRATA_ADVANCED_DEBUG_ENABLED && STRATA_MATERIALCONTAINER_IS_VIEWRESOURCE
-		if (Settings.SliceStoringDebugStrataTree > -1)
+		if (Settings.SliceStoringDebugStrataTreeData > -1)
 		{
-			StrataPackOutStrataTree(ExtraMaterialDataUAV, RootOperatorIndex, StrataTree, Settings.SliceStoringDebugStrataTree);
+			StrataPackOutStrataTree(ExtraMaterialDataUAV, RootOperatorIndex, StrataTree, Settings.SliceStoringDebugStrataTreeData);
 		}
 	#endif
 
diff --git a/Engine/Shaders/Private/Strata/StrataVisualize.usf b/Engine/Shaders/Private/Strata/StrataVisualize.usf
--- a/Engine/Shaders/Private/Strata/StrataVisualize.usf
+++ b/Engine/Shaders/Private/Strata/StrataVisualize.usf
@@ -639,13 +639,13 @@ void MaterialDebugStrataTreeCS(uint3 DispatchThreadId : SV_DispatchThreadID)
 	FStrataPixelHeader Header = UnpackStrataHeaderIn(Strata.MaterialTextureArray, StrataAddressing, Strata.TopLayerTexture);
 	if (Header.BSDFCount > 0)
 	{
-		const int SliceStoringDebugStrataTree = Strata.SliceStoringDebugStrataTree + STRATA_BASE_PASS_MRT_OUTPUT_COUNT;
+		const int SliceStoringDebugStrataTreeData = Strata.SliceStoringDebugStrataTreeData;
 
 		FStrataTreeHeader StrataTreeHeader = (FStrataTreeHeader)0;
 		StrataUnpackInStrataTreeHeader(
 			Strata.MaterialTextureArray,
 			StrataTreeHeader,
-			SliceStoringDebugStrataTree);
+			SliceStoringDebugStrataTreeData);
 
 		FShaderPrintContext Context;
 		Context = InitShaderPrintContext(true, uint2(50, 50));
@@ -685,7 +685,7 @@ void MaterialDebugStrataTreeCS(uint3 DispatchThreadId : SV_DispatchThreadID)
 					i,
 					Strata.MaterialTextureArray,
 					StrataTreeHeader,
-					SliceStoringDebugStrataTree);
+					SliceStoringDebugStrataTreeData);
 
 				Print(Context, TEXT("OperatorIndex = "), FontCyan);
 				Print(Context, i, FontCyan, 2, 0);
@@ -770,7 +770,7 @@ void MaterialDebugStrataTreeCS(uint3 DispatchThreadId : SV_DispatchThreadID)
 					i,
 					Strata.MaterialTextureArray,
 					StrataTreeHeader,
-					SliceStoringDebugStrataTree);
+					SliceStoringDebugStrataTreeData);
 
 				const bool bActive = any(BSDF.LuminanceWeightV > 0.0);
 				FFontColor BSDFFont = FontDarkRed;
@@ -841,7 +841,7 @@ void MaterialDebugStrataTreeCS(uint3 DispatchThreadId : SV_DispatchThreadID)
 				OperatorIndex,
 				Strata.MaterialTextureArray,
 				StrataTreeHeader,
-				SliceStoringDebugStrataTree);
+				SliceStoringDebugStrataTreeData);
 
 			Context = InitShaderPrintContext(true, uint2(PrintPosStack[StackPtr]));
 			if (Op.Type == STRATA_OPERATOR_WEIGHT)
@@ -1269,13 +1269,13 @@ void MaterialDebugStrataTreePS(
 			OutColor = float4(1.0f, 1.0f, 1.0f, 0.5f);							// Operator/BSDF separator
 		}
 
-		const int SliceStoringDebugStrataTree = Strata.SliceStoringDebugStrataTree + STRATA_BASE_PASS_MRT_OUTPUT_COUNT;
+		const int SliceStoringDebugStrataTreeData = Strata.SliceStoringDebugStrataTreeData;
 
 		FStrataTreeHeader StrataTreeHeader = (FStrataTreeHeader)0;
 		StrataUnpackInStrataTreeHeader(
 			Strata.MaterialTextureArray,
 			StrataTreeHeader,
-			SliceStoringDebugStrataTree);
+			SliceStoringDebugStrataTreeData);
 
 
 		float2 Pos = float2(BSDF_ELEMENT_STARTX, BSDF_ELEMENT_STARTY);
@@ -1295,7 +1295,7 @@ void MaterialDebugStrataTreePS(
 					i,
 					Strata.MaterialTextureArray,
 					StrataTreeHeader,
-					SliceStoringDebugStrataTree);
+					SliceStoringDebugStrataTreeData);
 
 				//if (BSDF.Coverage > 0.0f)
 				{
@@ -1344,13 +1344,13 @@ void MaterialDebugStrataTreePS(
 					i,
 					Strata.MaterialTextureArray,
 					StrataTreeHeader,
-					SliceStoringDebugStrataTree);
+					SliceStoringDebugStrataTreeData);
 
 				FStrataOperator Op = StrataUnpackInStrataTreeOperator(
 					BSDF.OperatorIndex,
 					Strata.MaterialTextureArray,
 					StrataTreeHeader,
-					SliceStoringDebugStrataTree);
+					SliceStoringDebugStrataTreeData);
 
 				const uint LayerDepth = Op.LayerDepth;
 
diff --git a/Engine/Source/Runtime/Renderer/Private/Strata/Strata.cpp b/Engine/Source/Runtime/Renderer/Private/Strata/Strata.cpp
--- a/Engine/Source/Runtime/Renderer/Private/Strata/Strata.cpp
+++ b/Engine/Source/Runtime/Renderer/Private/Strata/Strata.cpp
@@ -82,6 +82,8 @@ static TAutoConsoleVariable<int32> CVarStrataAsyncClassification(
 
 IMPLEMENT_GLOBAL_SHADER_PARAMETER_STRUCT(FStrataGlobalUniformParameters, "Strata");
 
+extern bool IsStrataAdvancedVisualizationShadersEnabled();
+
 void FStrataViewData::Reset()
 {
 	*this = FStrataViewData();
@@ -360,7 +362,8 @@ void InitialiseStrataFrameSceneData(FRDGBuilder& GraphBuilder, FSceneRenderer& S
 	FIntPoint SceneTextureExtent = IsStrataEnabled() ? SceneRenderer.GetActiveSceneTexturesConfig().Extent : FIntPoint(2, 2);
 
 	const uint32 SliceCountSSS = STRATA_SSS_DATA_UINT_COUNT;
-	const uint32 SliceCount = FMath::DivideAndRoundUp(Out.MaxBytesPerPixel, 4u) + SliceCountSSS;
+	const uint32 SliceCountAdvDebug = IsStrataAdvancedVisualizationShadersEnabled() ? 1 : 0;
+	const uint32 SliceCount = FMath::DivideAndRoundUp(Out.MaxBytesPerPixel, 4u) + SliceCountSSS + SliceCountAdvDebug;
 	const FRDGTextureDesc MaterialTextureDesc = FRDGTextureDesc::Create2DArray(SceneTextureExtent, PF_R32_UINT, FClearValueBinding::Transparent,
 		TexCreate_TargetArraySlicesIndependently | TexCreate_DisableDCC | TexCreate_NoFastClear | TexCreate_RenderTargetable | TexCreate_ShaderResource | TexCreate_UAV, SliceCount, 1, 1);
 	Out.MaterialTextureArray = GraphBuilder.CreateTexture(MaterialTextureDesc, TEXT("Strata.Material"));
@@ -368,7 +371,7 @@ void InitialiseStrataFrameSceneData(FRDGBuilder& GraphBuilder, FSceneRenderer& S
 	Out.MaterialTextureArrayUAV = GraphBuilder.CreateUAV(FRDGTextureUAVDesc(Out.MaterialTextureArray, 0));
 
 	// See AppendStrataMRTs
-	check(STRATA_BASE_PASS_MRT_OUTPUT_COUNT <= (SliceCount - SliceCountSSS)); // We want enough slice for MRTs but also do not want the SSSData to be a MRT.
+	check(STRATA_BASE_PASS_MRT_OUTPUT_COUNT <= (SliceCount - SliceCountSSS - SliceCountAdvDebug)); // We want enough slice for MRTs but also do not want the SSSData to be a MRT.
 	Out.MaterialTextureArrayUAVWithoutRTs = GraphBuilder.CreateUAV(FRDGTextureUAVDesc(Out.MaterialTextureArray, 0, PF_Unknown, STRATA_BASE_PASS_MRT_OUTPUT_COUNT, SliceCount - STRATA_BASE_PASS_MRT_OUTPUT_COUNT));
 
 	// Rough diffuse model
@@ -376,11 +379,12 @@ void InitialiseStrataFrameSceneData(FRDGBuilder& GraphBuilder, FSceneRenderer& S
 
 	Out.PeelLayersAboveDepth = FMath::Max(CVarStrataDebugPeelLayersAboveDepth.GetValueOnRenderThread(), 0);
 
-	// STRATA_TODO allocate a slice for StoringDebugStrata only if STRATA_ADVANCED_DEBUG_ENABLED is enabled
-	Out.SliceStoringDebugStrataTree = SliceCount - 1 - STRATA_BASE_PASS_MRT_OUTPUT_COUNT - SliceCountSSS;		// The UAV skips the first slices set as render target and other utility re
+	// STRATA_TODO allocate a slice for StoringDebugStrata only if STRATA_ADVANCED_DEBUG_ENABLED is enabled 
+	Out.SliceStoringDebugStrataTreeData				= SliceCount - SliceCountAdvDebug;										// When we read, there is no slices excluded
+	Out.SliceStoringDebugStrataTreeDataWithoutMRT	= SliceCount - SliceCountAdvDebug - STRATA_BASE_PASS_MRT_OUTPUT_COUNT;	// The UAV skips the first slices set as render target
 
-	Out.FirstSliceStoringStrataSSSData = SliceCount - SliceCountSSS;						// When we read, there is no slices excluded
-	Out.FirstSliceStoringStrataSSSDataWithoutMRT = SliceCount - SliceCountSSS - STRATA_BASE_PASS_MRT_OUTPUT_COUNT;	// The UAV skips the first slices set as render target
+	Out.FirstSliceStoringStrataSSSData				= SliceCount - SliceCountSSS - SliceCountAdvDebug;										// When we read, there is no slices excluded
+	Out.FirstSliceStoringStrataSSSDataWithoutMRT	= SliceCount - SliceCountSSS - SliceCountAdvDebug - STRATA_BASE_PASS_MRT_OUTPUT_COUNT;	// The UAV skips the first slices set as render target
 
 	// Initialized view data
 	for (int32 ViewIndex = 0; ViewIndex < SceneRenderer.Views.Num(); ViewIndex++)
@@ -397,7 +401,7 @@ void BindStrataBasePassUniformParameters(FRDGBuilder& GraphBuilder, const FViewI
 		OutStrataUniformParameters.bRoughDiffuse = StrataSceneData->bRoughDiffuse ? 1u : 0u;
 		OutStrataUniformParameters.MaxBytesPerPixel = StrataSceneData->MaxBytesPerPixel;
 		OutStrataUniformParameters.PeelLayersAboveDepth = StrataSceneData->PeelLayersAboveDepth;
-		OutStrataUniformParameters.SliceStoringDebugStrataTree = StrataSceneData->SliceStoringDebugStrataTree;
+		OutStrataUniformParameters.SliceStoringDebugStrataTreeDataWithoutMRT = StrataSceneData->SliceStoringDebugStrataTreeDataWithoutMRT;
 		OutStrataUniformParameters.FirstSliceStoringStrataSSSDataWithoutMRT = StrataSceneData->FirstSliceStoringStrataSSSDataWithoutMRT;
 		OutStrataUniformParameters.MaterialTextureArrayUAVWithoutRTs = StrataSceneData->MaterialTextureArrayUAVWithoutRTs;
 		OutStrataUniformParameters.OpaqueRoughRefractionTextureUAV = StrataSceneData->OpaqueRoughRefractionTextureUAV;
@@ -414,7 +418,7 @@ void BindStrataBasePassUniformParameters(FRDGBuilder& GraphBuilder, const FViewI
 		OutStrataUniformParameters.bRoughDiffuse = 0u;
 		OutStrataUniformParameters.MaxBytesPerPixel = 0;
 		OutStrataUniformParameters.PeelLayersAboveDepth = 0;
-		OutStrataUniformParameters.SliceStoringDebugStrataTree = -1;
+		OutStrataUniformParameters.SliceStoringDebugStrataTreeDataWithoutMRT = -1;
 		OutStrataUniformParameters.FirstSliceStoringStrataSSSDataWithoutMRT = -1;
 		OutStrataUniformParameters.MaterialTextureArrayUAVWithoutRTs = DummyWritableTextureArrayUAV;
 		OutStrataUniformParameters.OpaqueRoughRefractionTextureUAV = DummyWritableRefracTextureUAV;
@@ -429,7 +433,7 @@ static void BindStrataGlobalUniformParameters(FRDGBuilder& GraphBuilder, FStrata
 		OutStrataUniformParameters.bRoughDiffuse = StrataSceneData->bRoughDiffuse ? 1u : 0u;
 		OutStrataUniformParameters.MaxBytesPerPixel = StrataSceneData->MaxBytesPerPixel;
 		OutStrataUniformParameters.PeelLayersAboveDepth = StrataSceneData->PeelLayersAboveDepth;
-		OutStrataUniformParameters.SliceStoringDebugStrataTree = StrataSceneData->SliceStoringDebugStrataTree;
+		OutStrataUniformParameters.SliceStoringDebugStrataTreeData = StrataSceneData->SliceStoringDebugStrataTreeData;
 		OutStrataUniformParameters.FirstSliceStoringStrataSSSData = StrataSceneData->FirstSliceStoringStrataSSSData;
 		OutStrataUniformParameters.TileSize = STRATA_TILE_SIZE;
 		OutStrataUniformParameters.TileSizeLog2 = STRATA_TILE_SIZE_DIV_AS_SHIFT;
@@ -461,7 +465,7 @@ static void BindStrataGlobalUniformParameters(FRDGBuilder& GraphBuilder, FStrata
 		OutStrataUniformParameters.bRoughDiffuse = 0;
 		OutStrataUniformParameters.MaxBytesPerPixel = 0;
 		OutStrataUniformParameters.PeelLayersAboveDepth = 0;
-		OutStrataUniformParameters.SliceStoringDebugStrataTree = -1;
+		OutStrataUniformParameters.SliceStoringDebugStrataTreeData = -1;
 		OutStrataUniformParameters.FirstSliceStoringStrataSSSData = -1;
 		OutStrataUniformParameters.TileSize = 0;
 		OutStrataUniformParameters.TileSizeLog2 = 0;
diff --git a/Engine/Source/Runtime/Renderer/Private/Strata/Strata.h b/Engine/Source/Runtime/Renderer/Private/Strata/Strata.h
--- a/Engine/Source/Runtime/Renderer/Private/Strata/Strata.h
+++ b/Engine/Source/Runtime/Renderer/Private/Strata/Strata.h
@@ -21,7 +21,7 @@ BEGIN_SHADER_PARAMETER_STRUCT(FStrataBasePassUniformParameters, )
 	SHADER_PARAMETER(uint32, MaxBytesPerPixel)
 	SHADER_PARAMETER(uint32, bRoughDiffuse)
 	SHADER_PARAMETER(uint32, PeelLayersAboveDepth)
-	SHADER_PARAMETER(int32,  SliceStoringDebugStrataTree)
+	SHADER_PARAMETER(int32, SliceStoringDebugStrataTreeDataWithoutMRT)
 	SHADER_PARAMETER(int32, FirstSliceStoringStrataSSSDataWithoutMRT)
 	SHADER_PARAMETER_RDG_TEXTURE_UAV(RWTexture2DArray<uint>, MaterialTextureArrayUAVWithoutRTs)
 	SHADER_PARAMETER_RDG_TEXTURE_UAV(RWTexture2D<float3>, OpaqueRoughRefractionTextureUAV)
@@ -46,7 +46,7 @@ BEGIN_GLOBAL_SHADER_PARAMETER_STRUCT(FStrataGlobalUniformParameters, RENDERER_AP
 	SHADER_PARAMETER(uint32, MaxBytesPerPixel)
 	SHADER_PARAMETER(uint32, bRoughDiffuse)
 	SHADER_PARAMETER(uint32, PeelLayersAboveDepth)
-	SHADER_PARAMETER(int32,  SliceStoringDebugStrataTree)
+	SHADER_PARAMETER(int32,  SliceStoringDebugStrataTreeData)
 	SHADER_PARAMETER(int32,  FirstSliceStoringStrataSSSData)
 	SHADER_PARAMETER(uint32, TileSize)
 	SHADER_PARAMETER(uint32, TileSizeLog2)
@@ -81,7 +81,8 @@ struct FStrataSceneData
 	bool bRoughDiffuse;
 	int32 PeelLayersAboveDepth;
 
-	int32 SliceStoringDebugStrataTree;
+	int32 SliceStoringDebugStrataTreeDataWithoutMRT;
+	int32 SliceStoringDebugStrataTreeData;
 	int32 FirstSliceStoringStrataSSSDataWithoutMRT;
 	int32 FirstSliceStoringStrataSSSData;
 
