Volumetric cloud - clean up one permutation of cloud shader.
Only kept debug for conservative density through the always valid pixel shader path.

diff --git a/Engine/Shaders/Private/VolumetricCloud.usf b/Engine/Shaders/Private/VolumetricCloud.usf
--- a/Engine/Shaders/Private/VolumetricCloud.usf
+++ b/Engine/Shaders/Private/VolumetricCloud.usf
@@ -35,6 +35,10 @@
 #define CLOUD_SAMPLE_LOCAL_LIGHTS 0
 #endif
 
+#ifndef CLOUD_CONSERVATIVE_DENSITY_DEBUG
+#define CLOUD_CONSERVATIVE_DENSITY_DEBUG 0
+#endif
+
 #if defined(SHADER_MAINVS) || defined(SHADER_SHADOW_PS) || SHADER_RENDERVIEW_PS || SHADER_RENDERVIEW_CS
 
 #include "/Engine/Generated/Material.ush"
@@ -569,12 +573,9 @@ void MainCommon(in FMaterialPixelParameters MaterialParameters, in float4 SvPosi
 		const float MsPhaseFactor = 1.0f;
 #endif
 
-		uint DebugPrimaryMaterialSampleCount = 0;
-		uint DebugShadowSampleCount = 0;
-		uint DebugShadowMaterialSampleCount = 0;
-		uint DebugGroundShadowSampleCount = 0;
-		uint DebugGroundShadowMaterialSampleCount = 0;
+#if CLOUD_CONSERVATIVE_DENSITY_DEBUG
 		uint DebugZeroConservativeDensitySampleCount = 0;
+#endif
 
 		const float3 wi0 = Light0Direction;
 		const float3 wi1 = Light1Direction;
@@ -635,7 +636,7 @@ void MainCommon(in FMaterialPixelParameters MaterialParameters, in float4 SvPosi
 #if MATERIAL_VOLUMETRIC_ADVANCED_CONSERVATIVE_DENSITY
 			if (MaterialParameters.VolumeSampleConservativeDensity.x <= 0.0f)
 			{
-#if CLOUD_SAMPLE_COUNT_DEBUG_MODE
+#if CLOUD_CONSERVATIVE_DENSITY_DEBUG
 				++DebugZeroConservativeDensitySampleCount;
 				t += StepT;
 #else
@@ -650,7 +651,6 @@ void MainCommon(in FMaterialPixelParameters MaterialParameters, in float4 SvPosi
 			//////////////////////////////
 			// Sample the participating media material at the sample position.
 			//////////////////////////////
-			DebugPrimaryMaterialSampleCount++;
 			CalcPixelMaterialInputs(MaterialParameters, PixelMaterialInputs);
 			const MATVEC ExtinctionCoefficients = SampleExtinctionCoefficients(PixelMaterialInputs);
 			const MATVEC EmissiveLuminance = USES_EMISSIVE_COLOR ? SampleEmissive(PixelMaterialInputs) : 0.0f;
@@ -721,7 +721,6 @@ void MainCommon(in FMaterialPixelParameters MaterialParameters, in float4 SvPosi
 #if MATERIAL_VOLUMETRIC_ADVANCED_GROUND_CONTRIBUTION
 				if (MaxTransmittanceToView > 0.01f)
 				{
-					DebugGroundShadowSampleCount++;
 					// Cheap approximation assuming a single transmittance using only current sample extinction value. If we are withint the AO texture bounds, we lerp towards its more accurate result.
 					//DistantLightLuminance += TransmittedScatteredLightLuminance;
 
@@ -735,7 +734,6 @@ void MainCommon(in FMaterialPixelParameters MaterialParameters, in float4 SvPosi
 					float PreviousNormT = 0.0f;
 					for (float ShadowT = InvShadowStepCount; ShadowT <= 1.00001f; ShadowT += InvShadowStepCount)
 					{
-						DebugGroundShadowMaterialSampleCount++;
 						float CurrentNormT = ShadowT * ShadowT; // CurrentNormT is the end of the considered segment to integrate, PreviousNormT is its beginning.
 						const float DetlaNormT = CurrentNormT - PreviousNormT;
 						const float ShadowSampleDistance = ShadowLengthTest * (CurrentNormT - 0.5 * DetlaNormT);
@@ -817,8 +815,7 @@ void MainCommon(in FMaterialPixelParameters MaterialParameters, in float4 SvPosi
 					PMC.TransmittanceToLight0[ms] *= exp(-OutOpticalDepth * (ms == 0 ? 1.0f : pow(MsExtinFactor, ms)));
 				}
 #else 
-	// Use raymarched shadows
-				DebugShadowSampleCount++;
+				// Use raymarched shadows
 				UNROLL
 				for (ms = 0; ms < MSCOUNT; ++ms)
 				{
@@ -860,7 +857,6 @@ void MainCommon(in FMaterialPixelParameters MaterialParameters, in float4 SvPosi
 						break; // Ignore remaining samples since we have just traveld out of the cloud layer.
 					}
 
-					DebugShadowMaterialSampleCount++;
 					CalcPixelMaterialInputs(MaterialParameters, PixelMaterialInputs);
 					MATVEC ShadowExtinctionCoefficients = SampleExtinctionCoefficients(PixelMaterialInputs);
 
@@ -895,7 +891,6 @@ void MainCommon(in FMaterialPixelParameters MaterialParameters, in float4 SvPosi
 						PMC.TransmittanceToLight1[ms] *= exp(-OutOpticalDepth * (ms == 0 ? 1.0f : pow(MsExtinFactor, ms)));
 					}
 #else // Use cloud shadowmap
-					DebugShadowSampleCount++;
 					UNROLL
 					for (ms = 0; ms < MSCOUNT; ++ms)
 					{
@@ -926,7 +921,6 @@ void MainCommon(in FMaterialPixelParameters MaterialParameters, in float4 SvPosi
 							break; // Ignore remaining samples since we have just traveld out of the cloud layer.
 						}
 
-						DebugShadowMaterialSampleCount++;
 						CalcPixelMaterialInputs(MaterialParameters, PixelMaterialInputs);
 						MATVEC ShadowExtinctionCoefficients = SampleExtinctionCoefficients(PixelMaterialInputs);
 
@@ -1165,8 +1159,8 @@ void MainCommon(in FMaterialPixelParameters MaterialParameters, in float4 SvPosi
 			}
 
 			// This is helpful for performance. Can result in less light pucnhing through clouds. Should be a setting.
-			if (all(TransmittanceToView < RenderVolumetricCloudParameters.StopTracingTransmittanceThreshold)
-				&& (!CLOUD_SAMPLE_COUNT_DEBUG_MODE || RenderVolumetricCloudParameters.SampleCountDebugMode < 6))
+			if (all(TransmittanceToView < RenderVolumetricCloudParameters.StopTracingTransmittanceThreshold) 
+				&& !CLOUD_CONSERVATIVE_DENSITY_DEBUG)
 			{
 				break;
 			}
@@ -1174,38 +1168,8 @@ void MainCommon(in FMaterialPixelParameters MaterialParameters, in float4 SvPosi
 			t += StepT;
 			}
 
-#if CLOUD_SAMPLE_COUNT_DEBUG_MODE
-		if (RenderVolumetricCloudParameters.SampleCountDebugMode == 6)
-		{
-			OutColor0.r = 1.0 - saturate(DebugZeroConservativeDensitySampleCount / float(IStepCount));
-		}
-		else if (RenderVolumetricCloudParameters.SampleCountDebugMode > 0)
-		{
-			uint CountToDebug = 0;
-			switch (RenderVolumetricCloudParameters.SampleCountDebugMode)
-			{
-			case 1:
-				CountToDebug = DebugPrimaryMaterialSampleCount;
-				break;
-			case 2:
-				CountToDebug = DebugShadowSampleCount;
-				break;
-			case 3:
-				CountToDebug = DebugShadowMaterialSampleCount;
-				break;
-			case 4:
-				CountToDebug = DebugGroundShadowSampleCount;
-				break;
-			case 5:
-				CountToDebug = DebugGroundShadowMaterialSampleCount;
-				break;
-			}
-
-			const bool bMediumHit = CountToDebug > 0;
-			OutColor0 = float4(bMediumHit ? GetHSVDebugColor(saturate(float(CountToDebug) / 64.0f)) : float3(0.0f, 0.0f, 0.0), bMediumHit ? 0.0f : 1.0f);
-			OutColor1.x = bMediumHit ? OutColor1.x : MaxHalfFloat;
-			return;
-		}
+#if CLOUD_CONSERVATIVE_DENSITY_DEBUG
+		OutColor0.r = 1.0 - saturate(DebugZeroConservativeDensitySampleCount / float(IStepCount));
 #endif
 	}
 	////////////////////////////////////////////////////////////////////////////////////////////////////
@@ -1281,7 +1245,7 @@ void MainCommon(in FMaterialPixelParameters MaterialParameters, in float4 SvPosi
 	//
 
 	float GrayScaleTransmittance = MeanTransmittance < RenderVolumetricCloudParameters.StopTracingTransmittanceThreshold ? 0.0f : MeanTransmittance;
-#if CLOUD_SAMPLE_COUNT_DEBUG_MODE
+#if CLOUD_CONSERVATIVE_DENSITY_DEBUG
 	OutColor0 = float4(OutColor0.r * OutputPreExposure, Luminance.gb * OutputPreExposure, GrayScaleTransmittance);
 #else
 	OutColor0 = float4(Luminance * OutputPreExposure, GrayScaleTransmittance);
diff --git a/Engine/Source/Runtime/Renderer/Private/VolumetricCloudRendering.cpp b/Engine/Source/Runtime/Renderer/Private/VolumetricCloudRendering.cpp
--- a/Engine/Source/Runtime/Renderer/Private/VolumetricCloudRendering.cpp
+++ b/Engine/Source/Runtime/Renderer/Private/VolumetricCloudRendering.cpp
@@ -539,7 +539,6 @@ BEGIN_GLOBAL_SHADER_PARAMETER_STRUCT(FRenderVolumetricCloudGlobalParameters, )
 	SHADER_PARAMETER(uint32, NoiseFrameIndexModPattern)
 	SHADER_PARAMETER(int32, OpaqueIntersectionMode)
 	SHADER_PARAMETER(uint32, VolumetricRenderTargetMode)
-	SHADER_PARAMETER(uint32, SampleCountDebugMode)
 	SHADER_PARAMETER(uint32, IsReflectionRendering)
 	SHADER_PARAMETER(uint32, HasValidHZB)
 	SHADER_PARAMETER(uint32, ClampRayTToDepthBufferPostHZB)
@@ -563,16 +562,6 @@ IMPLEMENT_STATIC_UNIFORM_BUFFER_STRUCT(FRenderVolumetricCloudGlobalParameters, "
 
 IMPLEMENT_GLOBAL_SHADER_PARAMETER_STRUCT(FVolumetricCloudCommonGlobalShaderParameters, "VolumetricCloudCommonParameters");
 
-uint32 GetVolumetricCloudDebugSampleCountMode(const FEngineShowFlags& ShowFlags)
-{
-	if (ShowFlags.VisualizeVolumetricCloudConservativeDensity)
-	{
-		return 6;
-	}
-	// Add view modes for visualize other kinds of sample count
-	return FMath::Clamp(CVarVolumetricCloudDebugSampleCountMode.GetValueOnAnyThread(), 0, 5);
-}
-
 // When calling this, you still need to setup Light0Shadow yourself.
 void SetupDefaultRenderVolumetricCloudGlobalParameters(FRDGBuilder& GraphBuilder, FRenderVolumetricCloudGlobalParameters& VolumetricCloudParams, FVolumetricCloudRenderSceneInfo& CloudInfo, FViewInfo& ViewInfo)
 {
@@ -595,7 +584,6 @@ void SetupDefaultRenderVolumetricCloudGlobalParameters(FRDGBuilder& GraphBuilder
 	VolumetricCloudParams.TracingCoordToFullResPixelCoordScaleBias = FUintVector4(1, 1, 0, 0);
 	VolumetricCloudParams.NoiseFrameIndexModPattern = 0;
 	VolumetricCloudParams.VolumetricRenderTargetMode = ViewInfo.ViewState ? ViewInfo.ViewState->VolumetricCloudRenderTarget.GetMode() : 0;
-	VolumetricCloudParams.SampleCountDebugMode = GetVolumetricCloudDebugSampleCountMode(ViewInfo.Family->EngineShowFlags);
 
 	VolumetricCloudParams.HasValidHZB = 0;
 	VolumetricCloudParams.ClampRayTToDepthBufferPostHZB = 0;
@@ -698,7 +686,7 @@ BEGIN_SHADER_PARAMETER_STRUCT(FRenderVolumetricCloudRenderViewParametersPS, )
 	RENDER_TARGET_BINDING_SLOTS()
 END_SHADER_PARAMETER_STRUCT()
 
-template<EVolumetricCloudRenderViewPsPermutations Permutation, bool bSampleCountDebugMode>
+template<EVolumetricCloudRenderViewPsPermutations Permutation, bool bConvervativeDensityDebug>
 class FRenderVolumetricCloudRenderViewPs : public FMeshMaterialShader
 {
 	DECLARE_SHADER_TYPE(FRenderVolumetricCloudRenderViewPs, MeshMaterial);
@@ -717,7 +705,7 @@ public:
 	{
 		const bool bIsCompatible = IsMaterialCompatibleWithVolumetricCloud(Parameters.MaterialParameters, Parameters.Platform);
 		
-		if (bSampleCountDebugMode)
+		if (bConvervativeDensityDebug)
 		{
 			return bIsCompatible
 				&& EnumHasAllFlags(Parameters.Flags, EShaderPermutationFlags::HasEditorOnlyData)
@@ -754,7 +742,7 @@ public:
 										Permutation == CloudPs_AtmoTrans1_SampleShad1_2ndLight1 || Permutation == CloudPs_AtmoTrans1_SampleShad0_2ndLight1;
 		OutEnvironment.SetDefine(TEXT("CLOUD_SAMPLE_SECOND_LIGHT"), bSampleSecondLight ? TEXT("1") : TEXT("0"));
 
-		OutEnvironment.SetDefine(TEXT("CLOUD_SAMPLE_COUNT_DEBUG_MODE"), bSampleCountDebugMode);
+		OutEnvironment.SetDefine(TEXT("CLOUD_CONSERVATIVE_DENSITY_DEBUG"), bConvervativeDensityDebug);
 
 		// This shader takes a very long time to compile with FXC, so we pre-compile it with DXC first and then forward the optimized HLSL to FXC.
 		OutEnvironment.CompilerFlags.Add(CFLAG_PrecompileWithDXC);
@@ -790,14 +778,12 @@ class FRenderVolumetricCloudRenderViewCS : public FMeshMaterialShader
 	class FCloudSampleAtmosphericLightShadowmap : SHADER_PERMUTATION_BOOL("CLOUD_SAMPLE_ATMOSPHERIC_LIGHT_SHADOWMAP");
 	class FCloudSampleSecondLight : SHADER_PERMUTATION_BOOL("CLOUD_SAMPLE_SECOND_LIGHT");
 	class FCloudSampleLocalLights : SHADER_PERMUTATION_BOOL("CLOUD_SAMPLE_LOCAL_LIGHTS");
-	class FCloudSampleCountDebugMode : SHADER_PERMUTATION_BOOL("CLOUD_SAMPLE_COUNT_DEBUG_MODE");
 
 	using FPermutationDomain = TShaderPermutationDomain<
 		FCloudPerSampleAtmosphereTransmittance,
 		FCloudSampleAtmosphericLightShadowmap,
 		FCloudSampleSecondLight,
-		FCloudSampleLocalLights,
-		FCloudSampleCountDebugMode>;
+		FCloudSampleLocalLights>;
 
 	BEGIN_SHADER_PARAMETER_STRUCT(FParameters, )
 		SHADER_PARAMETER(FVector4f, OutputViewRect)
@@ -820,17 +806,6 @@ class FRenderVolumetricCloudRenderViewCS : public FMeshMaterialShader
 		const FPermutationDomain PermutationVector(Parameters.PermutationId);
 		const bool bIsCompatible = IsMaterialCompatibleWithVolumetricCloud(Parameters.MaterialParameters, Parameters.Platform);
 
-		if (PermutationVector.Get<FCloudSampleCountDebugMode>())
-		{
-			// Only compile visualization shaders for editor builds
-			return bIsCompatible
-				&& EnumHasAllFlags(Parameters.Flags, EShaderPermutationFlags::HasEditorOnlyData)
-				&& !PermutationVector.Get<FCloudPerSampleAtmosphereTransmittance>()
-				&& !PermutationVector.Get<FCloudSampleAtmosphericLightShadowmap>()
-				&& !PermutationVector.Get<FCloudSampleSecondLight>()
-				&& !PermutationVector.Get<FCloudSampleLocalLights>();
-		}
-
 		return bIsCompatible;
 	}
 
@@ -978,11 +953,11 @@ public:
 		, bVolumetricCloudPerSampleAtmosphereTransmittance(ShouldUsePerSampleAtmosphereTransmittance(Scene, InViewIfDynamicMeshCommand))
 		, bVolumetricCloudSampleLightShadowmap(!bSkipAtmosphericLightShadowmap && CVarVolumetricCloudShadowSampleAtmosphericLightShadowmap.GetValueOnAnyThread() > 0)
 		, bVolumetricCloudSecondLight(bSecondAtmosphereLightEnabled)
-		, bVisualizeConservativeDensityOrDebugSampleCount(bInVisualizeConservativeDensity)
+		, bVisualizeConservativeDensity(bInVisualizeConservativeDensity)
 	{
 		PassDrawRenderState.SetDepthStencilState(TStaticDepthStencilState<false, CF_Always>::GetRHI());
 
-		if (bShouldViewRenderVolumetricRenderTarget || bVisualizeConservativeDensityOrDebugSampleCount)
+		if (bShouldViewRenderVolumetricRenderTarget || bVisualizeConservativeDensity)
 		{
 			// No blending as we only render clouds in that render target today. Avoids clearing for now.
 			PassDrawRenderState.SetBlendState(TStaticBlendState<>::GetRHI());
@@ -1030,7 +1005,7 @@ public:
 		const ERasterizerFillMode MeshFillMode = FM_Solid;
 		const ERasterizerCullMode MeshCullMode = CM_None;
 
-		if (bVisualizeConservativeDensityOrDebugSampleCount)
+		if (bVisualizeConservativeDensity)
 		{
 			return TemplatedProcess<FRenderVolumetricCloudRenderViewPs<CloudPs_AtmoTrans0_SampleShad0_2ndLight0, true>>(
 				MeshBatch, BatchElementMask, PrimitiveSceneProxy, MaterialRenderProxy, Material, StaticMeshId, MeshFillMode, MeshCullMode);
@@ -1151,7 +1126,7 @@ private:
 	bool bVolumetricCloudPerSampleAtmosphereTransmittance;
 	bool bVolumetricCloudSampleLightShadowmap;
 	bool bVolumetricCloudSecondLight;
-	bool bVisualizeConservativeDensityOrDebugSampleCount;
+	bool bVisualizeConservativeDensity;
 };
 
 
@@ -1952,7 +1927,7 @@ FCloudRenderContext::FCloudRenderContext()
 	bSkipAerialPerspective = false;
 
 	bAsyncCompute = false;
-	bVisualizeConservativeDensityOrDebugSampleCount = false;
+	bVisualizeConservativeDensity = false;
 }
 
 static TRDGUniformBufferRef<FRenderVolumetricCloudGlobalParameters> CreateCloudPassUniformBuffer(FRDGBuilder& GraphBuilder, FCloudRenderContext& CloudRC)
@@ -2076,14 +2051,16 @@ void FSceneRenderer::RenderVolumetricCloudsInternal(FRDGBuilder& GraphBuilder, F
 	const bool bShouldViewRenderVolumetricRenderTarget = CloudRC.bShouldViewRenderVolumetricRenderTarget;
 	const bool bSkipAtmosphericLightShadowmap = CloudRC.bSkipAtmosphericLightShadowmap;
 	const bool bSecondAtmosphereLightEnabled = CloudRC.bSecondAtmosphereLightEnabled;
-	const bool bVisualizeConservativeDensityOrDebugSampleCount = CloudRC.bVisualizeConservativeDensityOrDebugSampleCount;
+	const bool bVisualizeConservativeDensity = CloudRC.bVisualizeConservativeDensity;
 	const FRDGTextureDesc& Desc = CloudRC.RenderTargets.Output[0].GetTexture()->Desc;
 	FViewInfo& MainView = *CloudRC.MainView;
 	TUniformBufferRef<FViewUniformShaderParameters> ViewUniformBuffer = CloudRC.ViewUniformBuffer;
 	TRDGUniformBufferRef<FRenderVolumetricCloudGlobalParameters> CloudPassUniformBuffer = CreateCloudPassUniformBuffer(GraphBuilder, CloudRC);
 
 	const bool bCloudEnableLocalLightSampling = CVarVolumetricCloudEnableLocalLightsSampling.GetValueOnRenderThread() > 0;
-	if (ShouldUseComputeForCloudTracing(Scene->GetFeatureLevel()))
+	// Debug permutations are only generated for pixel shader
+	const bool bShouldUseComputeForCloudTracing = !bVisualizeConservativeDensity && ShouldUseComputeForCloudTracing(Scene->GetFeatureLevel());
+	if (bShouldUseComputeForCloudTracing)
 	{
 		FRDGTextureRef CloudColorCubeTexture;
 		FRDGTextureRef CloudColorTexture;
@@ -2092,7 +2069,7 @@ void FSceneRenderer::RenderVolumetricCloudsInternal(FRDGBuilder& GraphBuilder, F
 
 		FRenderVolumetricCloudRenderViewCS::FParameters* PassParameters = GraphBuilder.AllocParameters<FRenderVolumetricCloudRenderViewCS::FParameters>();
 		PassParameters->OutputViewRect = FVector4f(0.f, 0.f, Desc.Extent.X, Desc.Extent.Y);
-		PassParameters->bBlendCloudColor = !bShouldViewRenderVolumetricRenderTarget && !bVisualizeConservativeDensityOrDebugSampleCount;
+		PassParameters->bBlendCloudColor = !bShouldViewRenderVolumetricRenderTarget;
 		PassParameters->TargetCubeFace = CloudRC.RenderTargets.Output[0].GetArraySlice();
 		PassParameters->VolumetricCloudRenderViewParamsUB = CloudPassUniformBuffer;
 		PassParameters->View = ViewUniformBuffer;
@@ -2106,11 +2083,10 @@ void FSceneRenderer::RenderVolumetricCloudsInternal(FRDGBuilder& GraphBuilder, F
 		MaterialRenderProxy = MaterialRenderProxy ? MaterialRenderProxy : CloudVolumeMaterialProxy;
 
 		typename FRenderVolumetricCloudRenderViewCS::FPermutationDomain PermutationVector;
-		PermutationVector.Set<typename FRenderVolumetricCloudRenderViewCS::FCloudPerSampleAtmosphereTransmittance>(!bVisualizeConservativeDensityOrDebugSampleCount && ShouldUsePerSampleAtmosphereTransmittance(Scene, &MainView));
-		PermutationVector.Set<typename FRenderVolumetricCloudRenderViewCS::FCloudSampleAtmosphericLightShadowmap>(!bVisualizeConservativeDensityOrDebugSampleCount && !bSkipAtmosphericLightShadowmap && CVarVolumetricCloudShadowSampleAtmosphericLightShadowmap.GetValueOnRenderThread() > 0);
-		PermutationVector.Set<typename FRenderVolumetricCloudRenderViewCS::FCloudSampleSecondLight>(!bVisualizeConservativeDensityOrDebugSampleCount && bSecondAtmosphereLightEnabled);
-		PermutationVector.Set<typename FRenderVolumetricCloudRenderViewCS::FCloudSampleLocalLights>(!bVisualizeConservativeDensityOrDebugSampleCount && bCloudEnableLocalLightSampling && !CloudRC.bIsSkyRealTimeReflectionRendering);
-		PermutationVector.Set<typename FRenderVolumetricCloudRenderViewCS::FCloudSampleCountDebugMode>(bVisualizeConservativeDensityOrDebugSampleCount);
+		PermutationVector.Set<typename FRenderVolumetricCloudRenderViewCS::FCloudPerSampleAtmosphereTransmittance>(ShouldUsePerSampleAtmosphereTransmittance(Scene, &MainView));
+		PermutationVector.Set<typename FRenderVolumetricCloudRenderViewCS::FCloudSampleAtmosphericLightShadowmap>(!bSkipAtmosphericLightShadowmap && CVarVolumetricCloudShadowSampleAtmosphericLightShadowmap.GetValueOnRenderThread() > 0);
+		PermutationVector.Set<typename FRenderVolumetricCloudRenderViewCS::FCloudSampleSecondLight>(bSecondAtmosphereLightEnabled);
+		PermutationVector.Set<typename FRenderVolumetricCloudRenderViewCS::FCloudSampleLocalLights>(bCloudEnableLocalLightSampling && !CloudRC.bIsSkyRealTimeReflectionRendering);
 
 		TShaderRef<FRenderVolumetricCloudRenderViewCS> ComputeShader = MaterialResource->GetShader<FRenderVolumetricCloudRenderViewCS>(&FLocalVertexFactory::StaticType, PermutationVector, false);
 
@@ -2175,16 +2151,16 @@ void FSceneRenderer::RenderVolumetricCloudsInternal(FRDGBuilder& GraphBuilder, F
 			RDG_EVENT_NAME("CloudView (PS) %dx%d", Desc.Extent.X, Desc.Extent.Y),
 			RenderViewPassParameters,
 			ERDGPassFlags::Raster,
-			[RenderViewPassParameters, Scene = Scene, &MainView, bVisualizeConservativeDensityOrDebugSampleCount,
+			[RenderViewPassParameters, Scene = Scene, &MainView, bVisualizeConservativeDensity,
 			bShouldViewRenderVolumetricRenderTarget, CloudVolumeMaterialProxy, bSkipAtmosphericLightShadowmap, bSecondAtmosphereLightEnabled](FRHICommandListImmediate& RHICmdList)
 			{
 				DrawDynamicMeshPass(MainView, RHICmdList,
 					[&MainView, bShouldViewRenderVolumetricRenderTarget, bSkipAtmosphericLightShadowmap, bSecondAtmosphereLightEnabled,
-					bVisualizeConservativeDensityOrDebugSampleCount, CloudVolumeMaterialProxy](FDynamicPassMeshDrawListContext* DynamicMeshPassContext)
+					bVisualizeConservativeDensity, CloudVolumeMaterialProxy](FDynamicPassMeshDrawListContext* DynamicMeshPassContext)
 					{
 						FVolumetricCloudRenderViewMeshProcessor PassMeshProcessor(
 							MainView.Family->Scene->GetRenderScene(), &MainView, bShouldViewRenderVolumetricRenderTarget,
-							bSkipAtmosphericLightShadowmap, bSecondAtmosphereLightEnabled, bVisualizeConservativeDensityOrDebugSampleCount, DynamicMeshPassContext);
+							bSkipAtmosphericLightShadowmap, bSecondAtmosphereLightEnabled, bVisualizeConservativeDensity, DynamicMeshPassContext);
 
 						FMeshBatch LocalSingleTriangleMesh;
 						GetSingleTriangleMeshBatch(LocalSingleTriangleMesh, CloudVolumeMaterialProxy, MainView.GetFeatureLevel());
@@ -2241,7 +2217,7 @@ bool FSceneRenderer::RenderVolumetricCloud(
 				bool bShouldViewRenderVolumetricCloudRenderTarget;
 				bool bEnableAerialPerspectiveSampling;
 				bool bShouldUseHighQualityAerialPerspective;
-				bool bVisualizeConservativeDensityOrDebugSampleCount;
+				bool bVisualizeConservativeDensity;
 			};
 			TArray<FLocalCloudView, TInlineAllocator<2>> ViewsToProcess;
 			
@@ -2255,14 +2231,14 @@ bool FSceneRenderer::RenderVolumetricCloud(
 					continue;
 				}
 
-				const bool bVisualizeConservativeDensityOrDebugSampleCount = ShouldViewVisualizeVolumetricCloudConservativeDensity(ViewInfo, ViewFamily.EngineShowFlags) || GetVolumetricCloudDebugSampleCountMode(ViewFamily.EngineShowFlags)>0;
+				const bool bVisualizeConservativeDensity = ShouldViewVisualizeVolumetricCloudConservativeDensity(ViewInfo, ViewFamily.EngineShowFlags);
 				const bool bEnableAerialPerspectiveSampling = CVarVolumetricCloudEnableAerialPerspectiveSampling.GetValueOnAnyThread() > 0;
 				const bool bShouldUseHighQualityAerialPerspective =
 					bEnableAerialPerspectiveSampling
 					&& Scene->HasSkyAtmosphere()
 					&& CVarVolumetricCloudHighQualityAerialPerspective.GetValueOnAnyThread() > 0
 					&& !ViewInfo.bIsReflectionCapture
-					&& !bVisualizeConservativeDensityOrDebugSampleCount;
+					&& !bVisualizeConservativeDensity;
 
 				if (bAsyncCompute
 					&& bShouldViewRenderVolumetricCloudRenderTarget
@@ -2284,7 +2260,7 @@ bool FSceneRenderer::RenderVolumetricCloud(
 				ViewToProcess.bShouldViewRenderVolumetricCloudRenderTarget = bShouldViewRenderVolumetricCloudRenderTarget;
 				ViewToProcess.bEnableAerialPerspectiveSampling = bEnableAerialPerspectiveSampling;
 				ViewToProcess.bShouldUseHighQualityAerialPerspective = bShouldUseHighQualityAerialPerspective;
-				ViewToProcess.bVisualizeConservativeDensityOrDebugSampleCount = bVisualizeConservativeDensityOrDebugSampleCount;
+				ViewToProcess.bVisualizeConservativeDensity = bVisualizeConservativeDensity;
 			}
 
 			for (int32 ViewIndex = 0; ViewIndex < ViewsToProcess.Num(); ViewIndex++)
@@ -2294,7 +2270,7 @@ bool FSceneRenderer::RenderVolumetricCloud(
 
 				CloudRC.MainView = ViewToProcess.ViewInfo;
 				CloudRC.bAsyncCompute = bAsyncCompute;
-				CloudRC.bVisualizeConservativeDensityOrDebugSampleCount = ViewToProcess.bVisualizeConservativeDensityOrDebugSampleCount;
+				CloudRC.bVisualizeConservativeDensity = ViewToProcess.bVisualizeConservativeDensity;
 
 				const bool bShouldViewRenderVolumetricCloudRenderTarget = ViewToProcess.bShouldViewRenderVolumetricCloudRenderTarget;
 				const bool bEnableAerialPerspectiveSampling = ViewToProcess.bEnableAerialPerspectiveSampling;
@@ -2339,7 +2315,7 @@ bool FSceneRenderer::RenderVolumetricCloud(
 
 					CloudRC.NoiseFrameIndexModPattern = VRT.GetNoiseFrameIndexModPattern();
 
-					check(VRT.GetMode() != 0 || CloudRC.bVisualizeConservativeDensityOrDebugSampleCount || HalfResolutionDepthCheckerboardMinMaxTexture);
+					check(VRT.GetMode() != 0 || CloudRC.bVisualizeConservativeDensity || HalfResolutionDepthCheckerboardMinMaxTexture);
 					CloudRC.SceneDepthZ = (VRT.GetMode() == 0 || VRT.GetMode() == 3) ? HalfResolutionDepthCheckerboardMinMaxTexture : SceneDepthZ;
 				}
 				else
@@ -2364,7 +2340,8 @@ bool FSceneRenderer::RenderVolumetricCloud(
 					DestinationRTDepth = GraphBuilder.CreateTexture(FRDGTextureDesc::Create2D(FIntPoint(RtSize.X, RtSize.Y), PF_G16R16F, FClearValueBinding::Black,
 						DepthTexCreateFlags, DepthNumMips, DepthNumSamples), TEXT("Cloud.DummyDepth"));
 
-					if (bShouldUseHighQualityAerialPerspective && ShouldUseComputeForCloudTracing(Scene->GetFeatureLevel()))
+					const bool bShouldUseComputeForCloudTracing = !CloudRC.bVisualizeConservativeDensity && ShouldUseComputeForCloudTracing(Scene->GetFeatureLevel());
+					if (bShouldUseHighQualityAerialPerspective && bShouldUseComputeForCloudTracing)
 					{
 						// If using the compute path, we then need to clear the intermediate render target manually as RDG won't do it for us in this case.
 						AddClearRenderTargetPass(GraphBuilder, IntermediateRT);
@@ -2377,7 +2354,7 @@ bool FSceneRenderer::RenderVolumetricCloud(
 					CloudRC.SceneDepthZ = SceneDepthZ;
 				}
 
-				if (CloudRC.bVisualizeConservativeDensityOrDebugSampleCount)
+				if (CloudRC.bVisualizeConservativeDensity)
 				{
 					CloudRC.SceneDepthZ = (bool)ERHIZBuffer::IsInverted ? GSystemTextures.GetBlackDummy(GraphBuilder) : GSystemTextures.GetWhiteDummy(GraphBuilder);
 				}
diff --git a/Engine/Source/Runtime/Renderer/Private/VolumetricCloudRendering.h b/Engine/Source/Runtime/Renderer/Private/VolumetricCloudRendering.h
--- a/Engine/Source/Runtime/Renderer/Private/VolumetricCloudRendering.h
+++ b/Engine/Source/Runtime/Renderer/Private/VolumetricCloudRendering.h
@@ -104,7 +104,6 @@ private:
 bool ShouldRenderVolumetricCloud(const FScene* Scene, const FEngineShowFlags& EngineShowFlags);
 bool ShouldViewVisualizeVolumetricCloudConservativeDensity(const FViewInfo& ViewInfo, const FEngineShowFlags& EngineShowFlags);
 bool VolumetricCloudWantsToSampleLocalLights(const FScene* Scene, const FEngineShowFlags& EngineShowFlags);
-uint32 GetVolumetricCloudDebugSampleCountMode(const FEngineShowFlags& ShowFlags);
 
 
 // Structure with data necessary to specify a cloud render.
@@ -133,7 +132,7 @@ struct FCloudRenderContext
 	bool bSecondAtmosphereLightEnabled;
 
 	bool bAsyncCompute;
-	bool bVisualizeConservativeDensityOrDebugSampleCount;
+	bool bVisualizeConservativeDensity;
 
 	FUintVector4 TracingCoordToZbufferCoordScaleBias;
 	FUintVector4 TracingCoordToFullResPixelCoordScaleBias;
diff --git a/Engine/Source/Runtime/Renderer/Private/VolumetricRenderTarget.cpp b/Engine/Source/Runtime/Renderer/Private/VolumetricRenderTarget.cpp
--- a/Engine/Source/Runtime/Renderer/Private/VolumetricRenderTarget.cpp
+++ b/Engine/Source/Runtime/Renderer/Private/VolumetricRenderTarget.cpp
@@ -529,8 +529,8 @@ void ReconstructVolumetricRenderTarget(
 		PassParameters->VolumetricRenderTargetMode = VolumetricCloudRT.GetMode();
 		PassParameters->HalfResDepthTexture = (VolumetricCloudRT.GetMode() == 0 || VolumetricCloudRT.GetMode() == 3) ? HalfResolutionDepthCheckerboardMinMaxTexture : SceneDepthTexture;
 
-		const bool bVisualizeConservativeDensityOrDebugSampleCount = ShouldViewVisualizeVolumetricCloudConservativeDensity(ViewInfo, ViewInfo.Family->EngineShowFlags) || GetVolumetricCloudDebugSampleCountMode(ViewInfo.Family->EngineShowFlags)>0;;
-		PassParameters->HalfResDepthTexture = bVisualizeConservativeDensityOrDebugSampleCount ?
+		const bool bVisualizeConservativeDensity = ShouldViewVisualizeVolumetricCloudConservativeDensity(ViewInfo, ViewInfo.Family->EngineShowFlags);
+		PassParameters->HalfResDepthTexture = bVisualizeConservativeDensity ?
 			((bool)ERHIZBuffer::IsInverted ? SystemTextures.Black : SystemTextures.White) :
 			((VolumetricCloudRT.GetMode() == 0 || VolumetricCloudRT.GetMode() == 3) ?
 				HalfResolutionDepthCheckerboardMinMaxTexture :
