Add possibility to run the SkyAtmosphereLUT render pass on async compute pipe (Disable by default). Can be enable/disable with a cvar (r.SkyAtmosphereASyncCompute).

-When running on async compute, we kick the SkyAtmosphereLUT render pass just before the Render Occlusion pass. This gives better waive occupency.

diff --git a/Engine/Source/Runtime/Renderer/Private/DeferredShadingRenderer.cpp b/Engine/Source/Runtime/Renderer/Private/DeferredShadingRenderer.cpp
--- a/Engine/Source/Runtime/Renderer/Private/DeferredShadingRenderer.cpp
+++ b/Engine/Source/Runtime/Renderer/Private/DeferredShadingRenderer.cpp
@@ -2706,6 +2706,21 @@ void FDeferredShadingSceneRenderer::Render(FRDGBuilder& GraphBuilder)
 		return GetViewPipelineState(Views[ViewIndex]).AmbientOcclusionMethod == EAmbientOcclusionMethod::SSAO;
 	});
 
+	const bool bShouldRenderSkyAtmosphere = ShouldRenderSkyAtmosphere(Scene, ViewFamily.EngineShowFlags);
+	const ESkyAtmospherePassLocation SkyAtmospherePassLocation = GetSkyAtmospherePassLocation();
+	const bool bShouldRenderVolumetricCloudBase = ShouldRenderVolumetricCloud(Scene, ViewFamily.EngineShowFlags);
+	const bool bShouldRenderVolumetricCloud = bShouldRenderVolumetricCloudBase && !ViewFamily.EngineShowFlags.VisualizeVolumetricCloudConservativeDensity;
+	const bool bShouldVisualizeVolumetricCloud = bShouldRenderVolumetricCloudBase && !!ViewFamily.EngineShowFlags.VisualizeVolumetricCloudConservativeDensity;
+	bool bAsyncComputeVolumetricCloud = IsVolumetricRenderTargetEnabled() && IsVolumetricRenderTargetAsyncCompute();
+	bool bHasHalfResCheckerboardMinMaxDepth = false;
+	bool bVolumetricRenderTargetRequired = bShouldRenderVolumetricCloud && !bHasRayTracedOverlay;
+
+	if (SkyAtmospherePassLocation == ESkyAtmospherePassLocation::BeforeOcclusion && bShouldRenderSkyAtmosphere)
+	{
+		// Generate the Sky/Atmosphere look up tables
+		RenderSkyAtmosphereLookUpTables(GraphBuilder);
+	}
+
 	const auto RenderOcclusionLambda = [&]()
 	{
 		RenderOcclusion(GraphBuilder, SceneTextures, bIsOcclusionTesting);
@@ -2729,13 +2744,8 @@ void FDeferredShadingSceneRenderer::Render(FRDGBuilder& GraphBuilder)
 
 	BeginAsyncDistanceFieldShadowProjections(GraphBuilder, SceneTextures);
 
-	const bool bShouldRenderSkyAtmosphere = ShouldRenderSkyAtmosphere(Scene, ViewFamily.EngineShowFlags);
-	const bool bShouldRenderVolumetricCloudBase = ShouldRenderVolumetricCloud(Scene, ViewFamily.EngineShowFlags);
-	const bool bShouldRenderVolumetricCloud = bShouldRenderVolumetricCloudBase && !ViewFamily.EngineShowFlags.VisualizeVolumetricCloudConservativeDensity;
-	const bool bShouldVisualizeVolumetricCloud = bShouldRenderVolumetricCloudBase && !!ViewFamily.EngineShowFlags.VisualizeVolumetricCloudConservativeDensity;
-	bool bAsyncComputeVolumetricCloud = IsVolumetricRenderTargetEnabled() && IsVolumetricRenderTargetAsyncCompute();
-	bool bHasHalfResCheckerboardMinMaxDepth = false;
-	bool bVolumetricRenderTargetRequired = bShouldRenderVolumetricCloud && !bHasRayTracedOverlay;
+
+	
 
 	if (bShouldRenderVolumetricCloudBase)
 	{
@@ -2747,7 +2757,7 @@ void FDeferredShadingSceneRenderer::Render(FRDGBuilder& GraphBuilder)
 	// Generate sky LUTs
 	// TODO: Valid shadow maps (for volumetric light shafts) have not yet been generated at this point in the frame. Need to resolve dependency ordering!
 	// This also must happen before the BasePass for Sky material to be able to sample valid LUTs.
-	if (bShouldRenderSkyAtmosphere)
+	if (SkyAtmospherePassLocation == ESkyAtmospherePassLocation::BeforeBasePass && bShouldRenderSkyAtmosphere)
 	{
 		// Generate the Sky/Atmosphere look up tables
 		RenderSkyAtmosphereLookUpTables(GraphBuilder);
diff --git a/Engine/Source/Runtime/Renderer/Private/SkyAtmosphereRendering.cpp b/Engine/Source/Runtime/Renderer/Private/SkyAtmosphereRendering.cpp
--- a/Engine/Source/Runtime/Renderer/Private/SkyAtmosphereRendering.cpp
+++ b/Engine/Source/Runtime/Renderer/Private/SkyAtmosphereRendering.cpp
@@ -216,6 +216,12 @@ static TAutoConsoleVariable<int32> CVarSkyAtmosphereEditorNotifications(
 	TEXT("Enable the rendering of in editor notification to warn the user about missing sky dome pixels on screen. It is better to keep it enabled and will be removed when shipping.\n"),
 	ECVF_RenderThreadSafe | ECVF_Scalability);
 
+static TAutoConsoleVariable<bool> CVarSkyAtmosphereASyncCompute(
+	TEXT("r.SkyAtmosphereASyncCompute"), false,
+	TEXT("SkyAtmosphere on async compute (default: false).\n"),
+	ECVF_RenderThreadSafe);
+
+
 DECLARE_GPU_STAT(SkyAtmosphereLUTs);
 DECLARE_GPU_STAT(SkyAtmosphere);
 DECLARE_GPU_STAT(SkyAtmosphereEditor);
@@ -275,6 +281,17 @@ IMPLEMENT_GLOBAL_SHADER_PARAMETER_STRUCT(FSkyAtmosphereInternalCommonParameters,
 #define KM_TO_CM  100000.0f
 #define CM_TO_KM  (1.0f / KM_TO_CM)
 
+ESkyAtmospherePassLocation GetSkyAtmospherePassLocation()
+{
+	if (CVarSkyAtmosphereASyncCompute.GetValueOnAnyThread())
+	{
+		// When SkyAtmLUT is running on async compute, try to kick it before occlusion pass to have better wave occupency
+		return ESkyAtmospherePassLocation::BeforeOcclusion;
+	}
+	// When SkyAtmLUT is running on graphics pipe, kickbefore BasePass to have a better overlap when DistanceFieldShadows is running async
+	return ESkyAtmospherePassLocation::BeforeBasePass;
+}
+
 float GetValidAerialPerspectiveStartDepthInCm(const FViewInfo& View, const FSkyAtmosphereSceneProxy& SkyAtmosphereProxy)
 {
 	float AerialPerspectiveStartDepthKm = SkyAtmosphereProxy.GetAerialPerspectiveStartDepthKm();
@@ -1336,6 +1353,8 @@ void FSceneRenderer::RenderSkyAtmosphereLookUpTables(FRDGBuilder& GraphBuilder)
 
 	FRDGExternalAccessQueue ExternalAccessQueue;
 
+	ERDGPassFlags PassFlag = CVarSkyAtmosphereASyncCompute.GetValueOnAnyThread() ? ERDGPassFlags::AsyncCompute : ERDGPassFlags::Compute;
+
 	// Transmittance LUT
 	FGlobalShaderMap* GlobalShaderMap = GetGlobalShaderMap(FeatureLevel);
 	if (CVarSkyAtmosphereTransmittanceLUT.GetValueOnRenderThread() > 0)
@@ -1350,7 +1369,7 @@ void FSceneRenderer::RenderSkyAtmosphereLookUpTables(FRDGBuilder& GraphBuilder)
 		FIntVector TextureSize = TransmittanceLut->Desc.GetSize();
 		TextureSize.Z = 1;
 		const FIntVector NumGroups = FIntVector::DivideAndRoundUp(TextureSize, FRenderTransmittanceLutCS::GroupSize);
-		FComputeShaderUtils::AddPass(GraphBuilder, RDG_EVENT_NAME("TransmittanceLut"), ComputeShader, PassParameters, NumGroups);
+		FComputeShaderUtils::AddPass(GraphBuilder, RDG_EVENT_NAME("TransmittanceLut"), PassFlag, ComputeShader, PassParameters, NumGroups);
 	}
 
 	// Multi-Scattering LUT
@@ -1371,7 +1390,7 @@ void FSceneRenderer::RenderSkyAtmosphereLookUpTables(FRDGBuilder& GraphBuilder)
 		FIntVector TextureSize = MultiScatteredLuminanceLut->Desc.GetSize();
 		TextureSize.Z = 1;
 		const FIntVector NumGroups = FIntVector::DivideAndRoundUp(TextureSize, FRenderMultiScatteredLuminanceLutCS::GroupSize);
-		FComputeShaderUtils::AddPass(GraphBuilder, RDG_EVENT_NAME("MultiScatteringLut"), ComputeShader, PassParameters, NumGroups);
+		FComputeShaderUtils::AddPass(GraphBuilder, RDG_EVENT_NAME("MultiScatteringLut"), PassFlag, ComputeShader, PassParameters, NumGroups);
 	}
 
 	// Distant Sky Light LUT
@@ -1420,7 +1439,7 @@ void FSceneRenderer::RenderSkyAtmosphereLookUpTables(FRDGBuilder& GraphBuilder)
 
 		FIntVector TextureSize = FIntVector(1, 1, 1);
 		const FIntVector NumGroups = FIntVector::DivideAndRoundUp(TextureSize, FRenderDistantSkyLightLutCS::GroupSize);
-		FComputeShaderUtils::AddPass(GraphBuilder, RDG_EVENT_NAME("DistantSkyLightLut"), ComputeShader, PassParameters, NumGroups);
+		FComputeShaderUtils::AddPass(GraphBuilder, RDG_EVENT_NAME("DistantSkyLightLut"), PassFlag, ComputeShader, PassParameters, NumGroups);
 
 		SkyInfo.GetDistantSkyLightLutTexture() = ConvertToExternalAccessTexture(GraphBuilder, ExternalAccessQueue, DistantSkyLightLut);
 	}
@@ -1520,7 +1539,7 @@ void FSceneRenderer::RenderSkyAtmosphereLookUpTables(FRDGBuilder& GraphBuilder)
 			FIntVector TextureSize = RealTimeReflectionCaptureSkyAtmosphereViewLutTexture->Desc.GetSize();
 			TextureSize.Z = 1;
 			const FIntVector NumGroups = FIntVector::DivideAndRoundUp(TextureSize, FRenderSkyViewLutCS::GroupSize);
-			FComputeShaderUtils::AddPass(GraphBuilder, RDG_EVENT_NAME("RealTimeCaptureSkyViewLut"), ComputeShader, PassParameters, NumGroups);
+			FComputeShaderUtils::AddPass(GraphBuilder, RDG_EVENT_NAME("RealTimeCaptureSkyViewLut"), PassFlag, ComputeShader, PassParameters, NumGroups);
 
 			Scene->RealTimeReflectionCaptureSkyAtmosphereViewLutTexture = ConvertToExternalAccessTexture(GraphBuilder, ExternalAccessQueue, RealTimeReflectionCaptureSkyAtmosphereViewLutTexture);
 		}
@@ -1573,7 +1592,7 @@ void FSceneRenderer::RenderSkyAtmosphereLookUpTables(FRDGBuilder& GraphBuilder)
 
 			FIntVector TextureSize = RealTimeReflectionCaptureCamera360APLutTexture->Desc.GetSize();
 			const FIntVector NumGroups = FIntVector::DivideAndRoundUp(TextureSize, FRenderCameraAerialPerspectiveVolumeCS::GroupSize);
-			FComputeShaderUtils::AddPass(GraphBuilder, RDG_EVENT_NAME("RealTimeCaptureCamera360VolumeLut"), ComputeShader, PassParameters, NumGroups);
+			FComputeShaderUtils::AddPass(GraphBuilder, RDG_EVENT_NAME("RealTimeCaptureCamera360VolumeLut"), PassFlag, ComputeShader, PassParameters, NumGroups);
 
 			Scene->RealTimeReflectionCaptureCamera360APLutTexture = ConvertToExternalAccessTexture(GraphBuilder, ExternalAccessQueue, RealTimeReflectionCaptureCamera360APLutTexture);
 		}
@@ -1649,7 +1668,7 @@ void FSceneRenderer::RenderSkyAtmosphereLookUpTables(FRDGBuilder& GraphBuilder)
 			FIntVector TextureSize = SkyAtmosphereViewLutTexture->Desc.GetSize();
 			TextureSize.Z = 1;
 			const FIntVector NumGroups = FIntVector::DivideAndRoundUp(TextureSize, FRenderSkyViewLutCS::GroupSize);
-			FComputeShaderUtils::AddPass(GraphBuilder, RDG_EVENT_NAME("SkyViewLut"), ComputeShader, PassParameters, NumGroups);
+			FComputeShaderUtils::AddPass(GraphBuilder, RDG_EVENT_NAME("SkyViewLut"), PassFlag, ComputeShader, PassParameters, NumGroups);
 		}
 
 		// Camera Atmosphere Volume
@@ -1692,7 +1711,7 @@ void FSceneRenderer::RenderSkyAtmosphereLookUpTables(FRDGBuilder& GraphBuilder)
 
 			FIntVector TextureSize = SkyAtmosphereCameraAerialPerspectiveVolume->Desc.GetSize();
 			const FIntVector NumGroups = FIntVector::DivideAndRoundUp(TextureSize, FRenderCameraAerialPerspectiveVolumeCS::GroupSize);
-			FComputeShaderUtils::AddPass(GraphBuilder, RDG_EVENT_NAME("CameraVolumeLut"), ComputeShader, PassParameters, NumGroups);
+			FComputeShaderUtils::AddPass(GraphBuilder, RDG_EVENT_NAME("CameraVolumeLut"), PassFlag, ComputeShader, PassParameters, NumGroups);
 		}
 
 		View.SkyAtmosphereViewLutTexture = ConvertToExternalAccessTexture(GraphBuilder, ExternalAccessQueue, SkyAtmosphereViewLutTexture);
diff --git a/Engine/Source/Runtime/Renderer/Private/SkyAtmosphereRendering.h b/Engine/Source/Runtime/Renderer/Private/SkyAtmosphereRendering.h
--- a/Engine/Source/Runtime/Renderer/Private/SkyAtmosphereRendering.h
+++ b/Engine/Source/Runtime/Renderer/Private/SkyAtmosphereRendering.h
@@ -163,6 +163,17 @@ private:
 	TRefCountPtr<IPooledRenderTarget> DistantSkyLightLutTexture;
 };
 
+enum class ESkyAtmospherePassLocation : uint32
+{
+	// Renders just before the Occlusion queries. Good for wave occuppency when SkyAtmosphere is async
+	BeforeOcclusion,
+
+	// Renders just before the BasePass.
+	BeforeBasePass
+};
+
+// Returns the location in the frame where SkyAtmosphere is rendered.
+extern ESkyAtmospherePassLocation GetSkyAtmospherePassLocation();
 
 
 bool ShouldRenderSkyAtmosphere(const FScene* Scene, const FEngineShowFlags& EngineShowFlags);
