Use HTile/Fast clear in CustomDepth render pass. Drived by cvar (r.CustomDepthEnableFastClear true/false)

diff --git a/Engine/Source/Runtime/Renderer/Private/CustomDepthRendering.cpp b/Engine/Source/Runtime/Renderer/Private/CustomDepthRendering.cpp
--- a/Engine/Source/Runtime/Renderer/Private/CustomDepthRendering.cpp
+++ b/Engine/Source/Runtime/Renderer/Private/CustomDepthRendering.cpp
@@ -32,6 +32,11 @@ static TAutoConsoleVariable<int32> CVarCustomDepthTemporalAAJitter(
 	TEXT("If disabled the Engine will remove the TemporalAA Jitter from the Custom Depth Pass. Only has effect when TemporalAA is used."),
 	ECVF_RenderThreadSafe);
 
+static TAutoConsoleVariable<bool> CVarCustomDepthEnableFastClear(
+	TEXT("r.CustomDepthEnableFastClear"), false,
+	TEXT("Enable HTile on the custom depth buffer (default:false).\n"),
+	ECVF_RenderThreadSafe);
+	
 DECLARE_GPU_DRAWCALL_STAT_NAMED(CustomDepth, TEXT("Custom Depth"));
 
 ECustomDepthPassLocation GetCustomDepthPassLocation(EShaderPlatform Platform)
@@ -70,7 +75,13 @@ FCustomDepthTextures FCustomDepthTextures::Create(FRDGBuilder& GraphBuilder, FIn
 
 	FCustomDepthTextures CustomDepthTextures;
 
-	const FRDGTextureDesc CustomDepthDesc = FRDGTextureDesc::Create2D(CustomDepthExtent, PF_DepthStencil, FClearValueBinding::DepthFar, GFastVRamConfig.CustomDepth | TexCreate_NoFastClear | TexCreate_DepthStencilTargetable | TexCreate_ShaderResource);
+	ETextureCreateFlags CreateFlags = GFastVRamConfig.CustomDepth | TexCreate_DepthStencilTargetable | TexCreate_ShaderResource;
+	if (!CVarCustomDepthEnableFastClear.GetValueOnRenderThread())
+	{
+		CreateFlags |= TexCreate_NoFastClear;
+	}
+
+	const FRDGTextureDesc CustomDepthDesc = FRDGTextureDesc::Create2D(CustomDepthExtent, PF_DepthStencil, FClearValueBinding::DepthFar, CreateFlags);
 
 	CustomDepthTextures.Depth = GraphBuilder.CreateTexture(CustomDepthDesc, TEXT("CustomDepth"));
 
