Fix for assert in distance field streaming async compute access.

diff --git a/Engine/Source/Runtime/Renderer/Private/DistanceFieldStreaming.cpp b/Engine/Source/Runtime/Renderer/Private/DistanceFieldStreaming.cpp
--- a/Engine/Source/Runtime/Renderer/Private/DistanceFieldStreaming.cpp
+++ b/Engine/Source/Runtime/Renderer/Private/DistanceFieldStreaming.cpp
@@ -1482,7 +1482,7 @@ void FDistanceFieldSceneData::UpdateDistanceFieldAtlas(
 				if (GDistanceFieldOffsetDataStructure == 0 || GDistanceFieldOffsetDataStructure == 1)
 				{
 					IndirectionTableUploadBuffer.ResourceUploadTo(GraphBuilder, IndirectionTableRDG);
-					ExternalAccessQueue.Add(IndirectionTableRDG, ERHIAccess::SRVMask);
+					ExternalAccessQueue.Add(IndirectionTableRDG, ERHIAccess::SRVMask, ERHIPipeline::All);
 				}
 				else
 				{
@@ -1560,7 +1560,7 @@ void FDistanceFieldSceneData::UpdateDistanceFieldAtlas(
 		GDistanceFieldAtlasLogStats = 0;
 	}
 
-	ExternalAccessQueue.Add(AssetDataBufferRDG, ERHIAccess::SRVMask);
+	ExternalAccessQueue.Add(AssetDataBufferRDG, ERHIAccess::SRVMask, ERHIPipeline::All);
 	ExternalAccessQueue.Submit(GraphBuilder);
 }
 
