Disable Nanite in planar reflections to prevent issues from not supporting the global clip plane.

diff --git a/Engine/Source/Runtime/Renderer/Private/SceneCaptureRendering.cpp b/Engine/Source/Runtime/Renderer/Private/SceneCaptureRendering.cpp
--- a/Engine/Source/Runtime/Renderer/Private/SceneCaptureRendering.cpp
+++ b/Engine/Source/Runtime/Renderer/Private/SceneCaptureRendering.cpp
@@ -596,6 +596,12 @@ void SetupViewFamilyForSceneCapture(
 	// For cube map capture, CubeMapFaceIndex takes precedence over view index, so we must have only one view for that case
 	check(CubemapFaceIndex == INDEX_NONE || Views.Num() == 1);
 
+	// Disable Nanite in planar reflections for now because Nanite doesn't support the global clip plane
+	if (bIsPlanarReflection)
+	{
+		ViewFamily.EngineShowFlags.NaniteMeshes = 0;
+	}
+
 	for (int32 ViewIndex = 0; ViewIndex < Views.Num(); ++ViewIndex)
 	{
 		const FSceneCaptureViewInfo& SceneCaptureViewInfo = Views[ViewIndex];
