Path Tracer MGPU:  Perf and synchronization improvements.
* Added "TransferResourceSignal" API to allow signaling that the destination GPU of a cross GPU transfer is done using a given resource, and it's safe for a source GPU to start copying to it.  Wanted to get this API level change into 5.1, as it could be used for various MGPU bugs and features going forward.
* Used new API in Path Tracer to fix potential bugs caused by cross GPU copy overlapping with rendering to the destination resource.
* Do multiple passes per frame based on the number of GPUs active, to reduce amortized overhead for cross GPU copies and other frame related work, and avoid situations where performance doesn't increase when rendering is CPU bound.  In a purely CPU bound case, performance doubles, while in a GPU bound case, performance increases by 6%.

#rnx

diff --git a/Engine/Source/Runtime/D3D12RHI/Private/D3D12CommandContext.cpp b/Engine/Source/Runtime/D3D12RHI/Private/D3D12CommandContext.cpp
--- a/Engine/Source/Runtime/D3D12RHI/Private/D3D12CommandContext.cpp
+++ b/Engine/Source/Runtime/D3D12RHI/Private/D3D12CommandContext.cpp
@@ -1159,8 +1159,24 @@ void FD3D12CommandContextRedirector::RHITransferResources(const TArrayView<const
 	// Note that this redirect to all GPUs from the context redirector.
 	RHISubmitCommandsHint();
 
-	// Submit the fence on all GPUs
+	// Wait on any pre-transfer fences first
 	FD3D12Fence* GPUFence = GetParentAdapter()->GetStagingFence();
+	for (const FTransferResourceParams& Param : Params)
+	{
+		if (Param.PreTransferFence)
+		{
+			FTransferResourceFenceData* FenceData = Param.PreTransferFence;
+
+			// It's possible to allocate fence data, but not actually signal the fence, in which case we treat this fence data as a nop.
+			if (FenceData->DeviceGPUIndex != INDEX_NONE)
+			{
+				GPUFence->GpuWait(FenceData->DeviceGPUIndex, ED3D12CommandQueueType::Direct, FenceData->LastSignal, FenceData->FenceGPUIndex);
+			}
+			delete FenceData;
+		}
+	}
+
+	// Submit the fence on all GPUs
 	if (SrcAndDestMask != FRHIGPUMask())
 	{
 		RHIMultiGPULockstep(SrcAndDestMask);
@@ -1265,6 +1281,29 @@ void FD3D12CommandContextRedirector::RHITransferResources(const TArrayView<const
 #endif // WITH_MGPU
 }
 
+void FD3D12CommandContextRedirector::RHITransferResourceSignal(const TArrayView<FTransferResourceFenceData* const> FenceDatas, FRHIGPUMask SrcGPUMask)
+{
+	check(FenceDatas.Num() == SrcGPUMask.GetNumActive());
+
+	// Make sure to Submit any pending work before signaling the fence.
+	RHISubmitCommandsHint();
+
+	uint64 LastSignal = GetParentAdapter()->GetStagingFence()->Signal(ED3D12CommandQueueType::Direct);
+
+	uint32 FenceIndex = 0;
+	for (uint32 SrcGPUIndex : SrcGPUMask)
+	{
+		FTransferResourceFenceData* FenceData = FenceDatas[FenceIndex++];
+
+		check(!GPUMask.Contains(SrcGPUIndex));
+		check(FenceData && FenceData->DeviceGPUIndex == INDEX_NONE);
+
+		FenceData->LastSignal = LastSignal;
+		FenceData->DeviceGPUIndex = SrcGPUIndex;
+		FenceData->FenceGPUIndex = GPUMask.GetFirstIndex();
+	}
+}
+
 void FD3D12CommandContextRedirector::RHITransferResourceWait(const TArrayView<FTransferResourceFenceData* const> FenceDatas)
 {
 #if WITH_MGPU
diff --git a/Engine/Source/Runtime/D3D12RHI/Private/D3D12CommandContext.h b/Engine/Source/Runtime/D3D12RHI/Private/D3D12CommandContext.h
--- a/Engine/Source/Runtime/D3D12RHI/Private/D3D12CommandContext.h
+++ b/Engine/Source/Runtime/D3D12RHI/Private/D3D12CommandContext.h
@@ -515,6 +515,7 @@ public:
 	virtual void RHIEndTransitions(TArrayView<const FRHITransition*> Transitions) final override;
 
 	virtual void RHITransferResources(const TArrayView<const FTransferResourceParams> Params) final override;
+	virtual void RHITransferResourceSignal(const TArrayView<FTransferResourceFenceData* const> FenceDatas, FRHIGPUMask SrcGPUMask) final override;
 	virtual void RHITransferResourceWait(const TArrayView<FTransferResourceFenceData* const> FenceDatas) final override;
 
 	FORCEINLINE virtual void RHICopyToStagingBuffer(FRHIBuffer* SourceBuffer, FRHIStagingBuffer* DestinationStagingBuffer, uint32 Offset, uint32 NumBytes) final override
diff --git a/Engine/Source/Runtime/RHI/Public/RHICommandList.h b/Engine/Source/Runtime/RHI/Public/RHICommandList.h
--- a/Engine/Source/Runtime/RHI/Public/RHICommandList.h
+++ b/Engine/Source/Runtime/RHI/Public/RHICommandList.h
@@ -963,6 +963,19 @@ FRHICOMMAND_MACRO(FRHICommandTransferResources)
 	RHI_API void Execute(FRHICommandListBase& CmdList);
 };
 
+FRHICOMMAND_MACRO(FRHICommandTransferResourceSignal)
+{
+	TArray<FTransferResourceFenceData*, TInlineAllocator<1>> FenceDatas;
+	FRHIGPUMask SrcGPUMask;
+
+	FORCEINLINE_DEBUGGABLE FRHICommandTransferResourceSignal(TArrayView<FTransferResourceFenceData* const> InFenceDatas, FRHIGPUMask InSrcGPUMask)
+		: FenceDatas(InFenceDatas), SrcGPUMask(InSrcGPUMask)
+	{
+	}
+
+	RHI_API void Execute(FRHICommandListBase & CmdList);
+};
+
 FRHICOMMAND_MACRO(FRHICommandTransferResourceWait)
 {
 	TArray<FTransferResourceFenceData*, TInlineAllocator<4>> FenceDatas;
@@ -3009,6 +3022,20 @@ PRAGMA_ENABLE_DEPRECATION_WARNINGS
 #endif // WITH_MGPU
 	}
 
+	FORCEINLINE_DEBUGGABLE void TransferResourceSignal(const TArrayView<FTransferResourceFenceData* const> FenceDatas, FRHIGPUMask SrcGPUMask)
+	{
+#if WITH_MGPU
+		if (Bypass())
+		{
+			GetComputeContext().RHITransferResourceSignal(FenceDatas, SrcGPUMask);
+		}
+		else
+		{
+			ALLOC_COMMAND(FRHICommandTransferResourceSignal)(FenceDatas, SrcGPUMask);
+		}
+#endif // WITH_MGPU
+	}
+
 	FORCEINLINE_DEBUGGABLE void TransferResourceWait(const TArrayView<FTransferResourceFenceData* const> FenceDatas)
 	{
 #if WITH_MGPU
diff --git a/Engine/Source/Runtime/RHI/Public/RHICommandListCommandExecutes.inl b/Engine/Source/Runtime/RHI/Public/RHICommandListCommandExecutes.inl
--- a/Engine/Source/Runtime/RHI/Public/RHICommandListCommandExecutes.inl
+++ b/Engine/Source/Runtime/RHI/Public/RHICommandListCommandExecutes.inl
@@ -115,6 +115,12 @@ void FRHICommandTransferResources::Execute(FRHICommandListBase& CmdList)
 	INTERNAL_DECORATOR_COMPUTE(RHITransferResources)(Params);
 }
 
+void FRHICommandTransferResourceSignal::Execute(FRHICommandListBase& CmdList)
+{
+	RHISTAT(TransferResourceSignal);
+	INTERNAL_DECORATOR_COMPUTE(RHITransferResourceSignal)(FenceDatas, SrcGPUMask);
+}
+
 void FRHICommandTransferResourceWait::Execute(FRHICommandListBase& CmdList)
 {
 	RHISTAT(TransferResourceWait);
diff --git a/Engine/Source/Runtime/RHI/Public/RHIContext.h b/Engine/Source/Runtime/RHI/Public/RHIContext.h
--- a/Engine/Source/Runtime/RHI/Public/RHIContext.h
+++ b/Engine/Source/Runtime/RHI/Public/RHIContext.h
@@ -180,6 +180,13 @@ struct FTransferResourceParams
 	  * about whether it occurred or not, but need to reserve the possibility.
 	  */
 	FTransferResourceFenceData* DelayedFence = nullptr;
+	/**
+	 * Optional pointer to a fence to wait on before starting the transfer.  Useful if a resource may be in use on the destination
+	 * GPU, and you need to wait until it's no longer in use before copying to it from the current GPU.  Fences are created via
+	 * "RHICreateTransferResourceFenceData", then signaled via "TransferResourceSignal" command, before being added to one of the
+	 * transfers in a batch that's dependent on the signal.
+	 */
+	FTransferResourceFenceData* PreTransferFence = nullptr;
 };
 
 /** Context that is capable of doing Compute work.  Can be async or compute on the gfx pipe. */
@@ -351,6 +358,23 @@ public:
 		/* empty default implementation */
 	}
 
+	/*
+	 * Signal where a cross GPU resource transfer can start.  Useful when the destination resource of a copy may still be in use, and
+	 * the copy from the source GPUs needs to wait until the destination is finished with it.  SrcGPUMask must not overlap the current
+	 * GPU mask of the context (which specifies the destination GPUs), and the number of items in the "FenceDatas" array MUST match the
+	 * number of bits set in SrcGPUMask.
+	 */
+	virtual void RHITransferResourceSignal(const TArrayView<FTransferResourceFenceData* const> FenceDatas, FRHIGPUMask SrcGPUMask)
+	{
+		/* default noop implementation */
+#if WITH_MGPU
+		for (FTransferResourceFenceData* FenceData : FenceDatas)
+		{
+			delete FenceData;
+		}
+#endif
+	}
+
 	virtual void RHITransferResourceWait(const TArrayView<FTransferResourceFenceData* const> FenceDatas)
 	{
 		/* default noop implementation */
diff --git a/Engine/Source/Runtime/Renderer/Private/PathTracing.cpp b/Engine/Source/Runtime/Renderer/Private/PathTracing.cpp
--- a/Engine/Source/Runtime/Renderer/Private/PathTracing.cpp
+++ b/Engine/Source/Runtime/Renderer/Private/PathTracing.cpp
@@ -216,7 +216,14 @@ TAutoConsoleVariable<int32> CVarPathTracingMultiGPU(
 	TEXT("Using this functionality in the editor requires r.AllowMultiGPUInEditor=1"),
 	ECVF_RenderThreadSafe
 );
-#endif
+
+TAutoConsoleVariable<bool> CVarPathTracingAdjustMultiGPUPasses(
+	TEXT("r.PathTracing.AdjustMultiGPUPasses"),
+	true,
+	TEXT("Run extra passes per frame when multiple GPUs are active, to improve perf scaling as GPUs are added (default = true)\n"),
+	ECVF_RenderThreadSafe
+);
+#endif  // WITH_MGPU
 
 TAutoConsoleVariable<int32> CVarPathTracingWiperMode(
 	TEXT("r.PathTracing.WiperMode"),
@@ -2101,20 +2108,6 @@ void FDeferredShadingSceneRenderer::RenderPathTracing(
 		View.ViewState->PathTracingInvalidate();
 	}
 
-	// Setup temporal seed _after_ invalidation in case we got reset
-	if (Config.LockedSamplingPattern)
-	{
-		// Count samples from 0 for deterministic results
-		Config.PathTracingData.TemporalSeed = PathTracingState->SampleIndex;
-	}
-	else
-	{
-		// Count samples from an ever-increasing counter to avoid screen-door effect
-		Config.PathTracingData.TemporalSeed = PathTracingState->FrameIndex;
-	}
-	Config.PathTracingData.Iteration = PathTracingState->SampleIndex;
-	Config.PathTracingData.BlendFactor = 1.0f / (Config.PathTracingData.Iteration + 1);
-
 	// Prepare radiance buffer (will be shared with display pass)
 	FRDGTexture* RadianceTexture = nullptr;
 	FRDGTexture* AlbedoTexture = nullptr;
@@ -2138,7 +2131,6 @@ void FDeferredShadingSceneRenderer::RenderPathTracing(
 		AlbedoTexture   = GraphBuilder.CreateTexture(Desc, TEXT("PathTracer.Albedo")  , ERDGTextureFlags::MultiFrame);
 		NormalTexture   = GraphBuilder.CreateTexture(Desc, TEXT("PathTracer.Normal")  , ERDGTextureFlags::MultiFrame);
 	}
-	const bool bNeedsMoreRays = Config.PathTracingData.Iteration < MaxSPP;
 
 	// should we use multiple GPUs to render the image?
 	const FRHIGPUMask GPUMask = Config.UseMultiGPU ? FRHIGPUMask::All() : View.GPUMask;
@@ -2147,210 +2139,281 @@ void FDeferredShadingSceneRenderer::RenderPathTracing(
 	const int32 DispatchResY = View.ViewRect.Size().Y;
 	const int32 DispatchSize = FMath::Max(CVarPathTracingDispatchSize.GetValueOnRenderThread(), 64);
 
-	if (bNeedsMoreRays)
+	// When running with multiple GPUs, do that number of passes per frame, to keep the GPU work done per frame consistent
+	// (given that each GPU processes a fraction of the pixels), but get the job done in fewer frames.
+#if WITH_MGPU
+	const int32 FramePassCount = !View.bIsOfflineRender && CVarPathTracingAdjustMultiGPUPasses.GetValueOnRenderThread() ? NumGPUs : 1;
+#else
+	const int32 FramePassCount = 1;
+#endif
+
+	bool bNeedsMoreRays = false;
+	bool bNeedsTextureExtract = false;
+
+	for (int32 FramePassIndex = 0; FramePassIndex < FramePassCount; FramePassIndex++)
 	{
-		// should we use path compaction?
-		const int CompactionType = CVarPathTracingCompaction.GetValueOnRenderThread();
-		const bool bUseIndirectDispatch = GRHISupportsRayTracingDispatchIndirect && CVarPathTracingIndirectDispatch.GetValueOnRenderThread() != 0;
-		const int FlushRenderingCommands = CVarPathTracingFlushDispatch.GetValueOnRenderThread();
-		
-		FRDGBuffer* ActivePaths[2] = {};
-		FRDGBuffer* NumActivePaths[2] = {};
-		FRDGBuffer* PathStateData = nullptr;
-		if (CompactionType == 1)
+		// Setup temporal seed _after_ invalidation in case we got reset
+		if (Config.LockedSamplingPattern)
 		{
-			const int32 NumPaths = FMath::Min(
-				DispatchSize * FMath::DivideAndRoundUp(DispatchSize, NumGPUs),
-				DispatchResX * FMath::DivideAndRoundUp(DispatchResY, NumGPUs)
-			);
-			ActivePaths[0] = GraphBuilder.CreateBuffer(FRDGBufferDesc::CreateBufferDesc(sizeof(int32), NumPaths), TEXT("PathTracer.ActivePaths0"));
-			ActivePaths[1] = GraphBuilder.CreateBuffer(FRDGBufferDesc::CreateBufferDesc(sizeof(int32), NumPaths), TEXT("PathTracer.ActivePaths1"));
-			if (bUseIndirectDispatch)
+			// Count samples from 0 for deterministic results
+			Config.PathTracingData.TemporalSeed = PathTracingState->SampleIndex;
+		}
+		else
+		{
+			// Count samples from an ever-increasing counter to avoid screen-door effect
+			Config.PathTracingData.TemporalSeed = PathTracingState->FrameIndex;
+		}
+		Config.PathTracingData.Iteration = PathTracingState->SampleIndex;
+		Config.PathTracingData.BlendFactor = 1.0f / (Config.PathTracingData.Iteration + 1);
+
+		bNeedsMoreRays = Config.PathTracingData.Iteration < MaxSPP;
+
+		if (bNeedsMoreRays)
+		{
+			// We are writing to the texture, we'll need to extract it...
+			bNeedsTextureExtract = true;
+
+			// should we use path compaction?
+			const int CompactionType = CVarPathTracingCompaction.GetValueOnRenderThread();
+			const bool bUseIndirectDispatch = GRHISupportsRayTracingDispatchIndirect && CVarPathTracingIndirectDispatch.GetValueOnRenderThread() != 0;
+			const int FlushRenderingCommands = CVarPathTracingFlushDispatch.GetValueOnRenderThread();
+
+			FRDGBuffer* ActivePaths[2] = {};
+			FRDGBuffer* NumActivePaths[2] = {};
+			FRDGBuffer* PathStateData = nullptr;
+			if (CompactionType == 1)
 			{
-				NumActivePaths[0] = GraphBuilder.CreateBuffer(FRDGBufferDesc::CreateIndirectDesc<int32>(3), TEXT("PathTracer.NumActivePaths0"));
-				NumActivePaths[1] = GraphBuilder.CreateBuffer(FRDGBufferDesc::CreateIndirectDesc<int32>(3), TEXT("PathTracer.NumActivePaths1"));
+				const int32 NumPaths = FMath::Min(
+					DispatchSize * FMath::DivideAndRoundUp(DispatchSize, NumGPUs),
+					DispatchResX * FMath::DivideAndRoundUp(DispatchResY, NumGPUs)
+				);
+				ActivePaths[0] = GraphBuilder.CreateBuffer(FRDGBufferDesc::CreateBufferDesc(sizeof(int32), NumPaths), TEXT("PathTracer.ActivePaths0"));
+				ActivePaths[1] = GraphBuilder.CreateBuffer(FRDGBufferDesc::CreateBufferDesc(sizeof(int32), NumPaths), TEXT("PathTracer.ActivePaths1"));
+				if (bUseIndirectDispatch)
+				{
+					NumActivePaths[0] = GraphBuilder.CreateBuffer(FRDGBufferDesc::CreateIndirectDesc<int32>(3), TEXT("PathTracer.NumActivePaths0"));
+					NumActivePaths[1] = GraphBuilder.CreateBuffer(FRDGBufferDesc::CreateIndirectDesc<int32>(3), TEXT("PathTracer.NumActivePaths1"));
+				}
+				else
+				{
+					NumActivePaths[0] = GraphBuilder.CreateBuffer(FRDGBufferDesc::CreateBufferDesc(sizeof(int32), 3), TEXT("PathTracer.NumActivePaths"));
+				}
+				PathStateData = GraphBuilder.CreateBuffer(FRDGBufferDesc::CreateStructuredDesc(sizeof(FPathTracingPackedPathState), NumPaths), TEXT("PathTracer.PathStateData"));
 			}
-			else
+			FPathTracingRG::FPermutationDomain PermutationVector;
+			PermutationVector.Set<FPathTracingRG::FCompactionType>(CompactionType);
+			TShaderMapRef<FPathTracingRG> RayGenShader(View.ShaderMap, PermutationVector);
+			FPathTracingRG::FParameters* PreviousPassParameters = nullptr;
+			// Divide each tile among all the active GPUs (interleaving scanlines)
+			// The assumption is that the tiles are as big as possible, hopefully covering the entire screen
+			// so rather than dividing tiles among GPUs, we divide each tile among all GPUs
+			int32 CurrentGPU = 0; // keep our own counter so that we don't assume the assigned GPUs in the view mask are sequential
+			for (int32 GPUIndex : GPUMask)
 			{
-				NumActivePaths[0] = GraphBuilder.CreateBuffer(FRDGBufferDesc::CreateBufferDesc(sizeof(int32), 3), TEXT("PathTracer.NumActivePaths"));
-			}
-			PathStateData = GraphBuilder.CreateBuffer(FRDGBufferDesc::CreateStructuredDesc(sizeof(FPathTracingPackedPathState), NumPaths), TEXT("PathTracer.PathStateData"));
-		}
-		FPathTracingRG::FPermutationDomain PermutationVector;
-		PermutationVector.Set<FPathTracingRG::FCompactionType>(CompactionType);
-		TShaderMapRef<FPathTracingRG> RayGenShader(View.ShaderMap, PermutationVector);
-		FPathTracingRG::FParameters* PreviousPassParameters = nullptr;
-		// Divide each tile among all the active GPUs (interleaving scanlines)
-		// The assumption is that the tiles are as big as possible, hopefully covering the entire screen
-		// so rather than dividing tiles among GPUs, we divide each tile among all GPUs
-		int32 CurrentGPU = 0; // keep our own counter so that we don't assume the assigned GPUs in the view mask are sequential
-		for (int32 GPUIndex : GPUMask)
-		{
-			RDG_GPU_MASK_SCOPE(GraphBuilder, FRHIGPUMask::FromIndex(GPUIndex));
+				RDG_GPU_MASK_SCOPE(GraphBuilder, FRHIGPUMask::FromIndex(GPUIndex));
 #if WITH_MGPU
-			RDG_EVENT_SCOPE(GraphBuilder, "Path Tracing GPU%d", GPUIndex);
-			RDG_GPU_STAT_SCOPE(GraphBuilder, Stat_GPU_PathTracing);
+				RDG_EVENT_SCOPE(GraphBuilder, "Path Tracing GPU%d", GPUIndex);
+				RDG_GPU_STAT_SCOPE(GraphBuilder, Stat_GPU_PathTracing);
 #endif
-			for (int32 TileY = 0; TileY < DispatchResY; TileY += DispatchSize)
-			{
-				for (int32 TileX = 0; TileX < DispatchResX; TileX += DispatchSize)
+				for (int32 TileY = 0; TileY < DispatchResY; TileY += DispatchSize)
 				{
-					const int32 DispatchSizeX = FMath::Min(DispatchSize, DispatchResX - TileX);
-					const int32 DispatchSizeY = FMath::Min(DispatchSize, DispatchResY - TileY);
-
-					const int32 DispatchSizeYSplit = FMath::DivideAndRoundUp(DispatchSizeY, NumGPUs);
-
-					// Compute the dispatch size for just this set of scanlines
-					const int32 DispatchSizeYLocal = FMath::Min(DispatchSizeYSplit, DispatchSizeY - CurrentGPU * DispatchSizeYSplit);
-					if (CompactionType == 1)
-					{
-						AddClearUAVPass(GraphBuilder, GraphBuilder.CreateUAV(ActivePaths[0], PF_R32_UINT), 0);
-					}
-					// When using path compaction, we need to run the path tracer once per bounce
-					// otherwise, the path tracer is the one doing the bounces
-					for (int Bounce = 0, MaxBounces = CompactionType == 1 ? Config.PathTracingData.MaxBounces : 0; Bounce <= MaxBounces; Bounce++)
+					for (int32 TileX = 0; TileX < DispatchResX; TileX += DispatchSize)
 					{
-						FPathTracingRG::FParameters* PassParameters = GraphBuilder.AllocParameters<FPathTracingRG::FParameters>();
-						PassParameters->TLAS = Scene->RayTracingScene.GetLayerSRVChecked(ERayTracingSceneLayer::Base);
-						PassParameters->ViewUniformBuffer = View.ViewUniformBuffer;
-						PassParameters->PathTracingData = Config.PathTracingData;
-						if (PreviousPassParameters == nullptr)
-						{
-							// upload sky/lights data
-							RDG_GPU_MASK_SCOPE(GraphBuilder, GPUMask); // make sure this happens on all GPUs we will be rendering on
-							SetLightParameters(GraphBuilder, PassParameters, Scene, View, Config.UseMISCompensation);
-						}
-						else
-						{
-							// re-use from last iteration
-							PassParameters->IESTexture = PreviousPassParameters->IESTexture;
-							PassParameters->IESTextureSampler = PreviousPassParameters->IESTextureSampler;
-							PassParameters->LightGridParameters = PreviousPassParameters->LightGridParameters;
-							PassParameters->SceneLightCount = PreviousPassParameters->SceneLightCount;
-							PassParameters->SceneVisibleLightCount = PreviousPassParameters->SceneVisibleLightCount;
-							PassParameters->SceneLights = PreviousPassParameters->SceneLights;
-							PassParameters->SkylightParameters = PreviousPassParameters->SkylightParameters;
-						}
-						if (Config.PathTracingData.EnableDirectLighting == 0)
-						{
-							PassParameters->SceneVisibleLightCount = 0;
-						}
+						const int32 DispatchSizeX = FMath::Min(DispatchSize, DispatchResX - TileX);
+						const int32 DispatchSizeY = FMath::Min(DispatchSize, DispatchResY - TileY);
 
-						PassParameters->DecalParameters = View.RayTracingDecalUniformBuffer;
+						const int32 DispatchSizeYSplit = FMath::DivideAndRoundUp(DispatchSizeY, NumGPUs);
 
-						PassParameters->RadianceTexture = GraphBuilder.CreateUAV(RadianceTexture);
-						PassParameters->AlbedoTexture = GraphBuilder.CreateUAV(AlbedoTexture);
-						PassParameters->NormalTexture = GraphBuilder.CreateUAV(NormalTexture);
-
-						if (PreviousPassParameters != nullptr)
-						{
-							PassParameters->Atmosphere = PreviousPassParameters->Atmosphere;
-							PassParameters->PlanetCenterTranslatedWorldHi = PreviousPassParameters->PlanetCenterTranslatedWorldHi;
-							PassParameters->PlanetCenterTranslatedWorldLo = PreviousPassParameters->PlanetCenterTranslatedWorldLo;
-						}
-						else if (Config.PathTracingData.EnableAtmosphere)
+						// Compute the dispatch size for just this set of scanlines
+						const int32 DispatchSizeYLocal = FMath::Min(DispatchSizeYSplit, DispatchSizeY - CurrentGPU * DispatchSizeYSplit);
+						if (CompactionType == 1)
 						{
-							PassParameters->Atmosphere = Scene->GetSkyAtmosphereSceneInfo()->GetAtmosphereUniformBuffer();
-							FVector PlanetCenterTranslatedWorld = Scene->GetSkyAtmosphereSceneInfo()->GetSkyAtmosphereSceneProxy().GetAtmosphereSetup().PlanetCenterKm * double(FAtmosphereSetup::SkyUnitToCm) + View.ViewMatrices.GetPreViewTranslation();
-							SplitDouble(PlanetCenterTranslatedWorld.X, &PassParameters->PlanetCenterTranslatedWorldHi.X, &PassParameters->PlanetCenterTranslatedWorldLo.X);
-							SplitDouble(PlanetCenterTranslatedWorld.Y, &PassParameters->PlanetCenterTranslatedWorldHi.Y, &PassParameters->PlanetCenterTranslatedWorldLo.Y);
-							SplitDouble(PlanetCenterTranslatedWorld.Z, &PassParameters->PlanetCenterTranslatedWorldHi.Z, &PassParameters->PlanetCenterTranslatedWorldLo.Z);
+							AddClearUAVPass(GraphBuilder, GraphBuilder.CreateUAV(ActivePaths[0], PF_R32_UINT), 0);
 						}
-						else
+						// When using path compaction, we need to run the path tracer once per bounce
+						// otherwise, the path tracer is the one doing the bounces
+						for (int Bounce = 0, MaxBounces = CompactionType == 1 ? Config.PathTracingData.MaxBounces : 0; Bounce <= MaxBounces; Bounce++)
 						{
-							FAtmosphereUniformShaderParameters AtmosphereParams = {};
-							PassParameters->Atmosphere = CreateUniformBufferImmediate(AtmosphereParams, EUniformBufferUsage::UniformBuffer_SingleFrame);
-							PassParameters->PlanetCenterTranslatedWorldHi = FVector3f(0);
-							PassParameters->PlanetCenterTranslatedWorldLo = FVector3f(0);
-						}
-						PassParameters->AtmosphereOpticalDepthLUT = AtmosphereOpticalDepthLUT;
-						PassParameters->AtmosphereOpticalDepthLUTSampler = TStaticSamplerState<SF_Bilinear, AM_Clamp, AM_Clamp, AM_Clamp>::GetRHI();
+							FPathTracingRG::FParameters* PassParameters = GraphBuilder.AllocParameters<FPathTracingRG::FParameters>();
+							PassParameters->TLAS = Scene->RayTracingScene.GetLayerSRVChecked(ERayTracingSceneLayer::Base);
+							PassParameters->ViewUniformBuffer = View.ViewUniformBuffer;
+							PassParameters->PathTracingData = Config.PathTracingData;
+							if (PreviousPassParameters == nullptr)
+							{
+								// upload sky/lights data
+								RDG_GPU_MASK_SCOPE(GraphBuilder, GPUMask); // make sure this happens on all GPUs we will be rendering on
+								SetLightParameters(GraphBuilder, PassParameters, Scene, View, Config.UseMISCompensation);
+							}
+							else
+							{
+								// re-use from last iteration
+								PassParameters->IESTexture = PreviousPassParameters->IESTexture;
+								PassParameters->IESTextureSampler = PreviousPassParameters->IESTextureSampler;
+								PassParameters->LightGridParameters = PreviousPassParameters->LightGridParameters;
+								PassParameters->SceneLightCount = PreviousPassParameters->SceneLightCount;
+								PassParameters->SceneVisibleLightCount = PreviousPassParameters->SceneVisibleLightCount;
+								PassParameters->SceneLights = PreviousPassParameters->SceneLights;
+								PassParameters->SkylightParameters = PreviousPassParameters->SkylightParameters;
+							}
+							if (Config.PathTracingData.EnableDirectLighting == 0)
+							{
+								PassParameters->SceneVisibleLightCount = 0;
+							}
 
-						if (Config.PathTracingData.EnableFog)
-						{
-							PassParameters->FogParameters = PrepareFogParameters(View, Scene->ExponentialFogs[0]);
-						}
-						else
-						{
-							PassParameters->FogParameters = {};
-						}
+							PassParameters->DecalParameters = View.RayTracingDecalUniformBuffer;
 
-						PassParameters->TilePixelOffset.X = TileX;
-						PassParameters->TilePixelOffset.Y = TileY + CurrentGPU;
-						PassParameters->TileTextureOffset.X = TileX;
-						PassParameters->TileTextureOffset.Y = TileY + CurrentGPU * DispatchSizeYSplit;
-						PassParameters->ScanlineStride = NumGPUs;
-						PassParameters->ScanlineWidth = DispatchSizeX;
+							PassParameters->RadianceTexture = GraphBuilder.CreateUAV(RadianceTexture);
+							PassParameters->AlbedoTexture = GraphBuilder.CreateUAV(AlbedoTexture);
+							PassParameters->NormalTexture = GraphBuilder.CreateUAV(NormalTexture);
 
-						PassParameters->Bounce = Bounce;
-						if (CompactionType == 1)
-						{
-							PassParameters->ActivePaths = GraphBuilder.CreateSRV(ActivePaths[Bounce & 1], PF_R32_SINT);
-							PassParameters->NextActivePaths = GraphBuilder.CreateUAV(ActivePaths[(Bounce & 1) ^ 1], PF_R32_SINT);
-							PassParameters->PathStateData = GraphBuilder.CreateUAV(PathStateData);
-							if (bUseIndirectDispatch)
+							if (PreviousPassParameters != nullptr)
+							{
+								PassParameters->Atmosphere = PreviousPassParameters->Atmosphere;
+								PassParameters->PlanetCenterTranslatedWorldHi = PreviousPassParameters->PlanetCenterTranslatedWorldHi;
+								PassParameters->PlanetCenterTranslatedWorldLo = PreviousPassParameters->PlanetCenterTranslatedWorldLo;
+							}
+							else if (Config.PathTracingData.EnableAtmosphere)
 							{
-								PassParameters->NumPathStates = GraphBuilder.CreateUAV(NumActivePaths[Bounce & 1], PF_R32_UINT);
-								PassParameters->PathTracingIndirectArgs = NumActivePaths[(Bounce & 1) ^ 1];
+								PassParameters->Atmosphere = Scene->GetSkyAtmosphereSceneInfo()->GetAtmosphereUniformBuffer();
+								FVector PlanetCenterTranslatedWorld = Scene->GetSkyAtmosphereSceneInfo()->GetSkyAtmosphereSceneProxy().GetAtmosphereSetup().PlanetCenterKm * double(FAtmosphereSetup::SkyUnitToCm) + View.ViewMatrices.GetPreViewTranslation();
+								SplitDouble(PlanetCenterTranslatedWorld.X, &PassParameters->PlanetCenterTranslatedWorldHi.X, &PassParameters->PlanetCenterTranslatedWorldLo.X);
+								SplitDouble(PlanetCenterTranslatedWorld.Y, &PassParameters->PlanetCenterTranslatedWorldHi.Y, &PassParameters->PlanetCenterTranslatedWorldLo.Y);
+								SplitDouble(PlanetCenterTranslatedWorld.Z, &PassParameters->PlanetCenterTranslatedWorldHi.Z, &PassParameters->PlanetCenterTranslatedWorldLo.Z);
 							}
 							else
 							{
-								PassParameters->NumPathStates = GraphBuilder.CreateUAV(NumActivePaths[0], PF_R32_UINT);
-								AddClearUAVPass(GraphBuilder, PassParameters->NextActivePaths, -1); // make sure everything is initialized to -1 since paths that go inactive don't write anything
+								FAtmosphereUniformShaderParameters AtmosphereParams = {};
+								PassParameters->Atmosphere = CreateUniformBufferImmediate(AtmosphereParams, EUniformBufferUsage::UniformBuffer_SingleFrame);
+								PassParameters->PlanetCenterTranslatedWorldHi = FVector3f(0);
+								PassParameters->PlanetCenterTranslatedWorldLo = FVector3f(0);
 							}
-							AddClearUAVPass(GraphBuilder, PassParameters->NumPathStates, 0);
-						}
-						ClearUnusedGraphResources(RayGenShader, PassParameters);
-						const bool bFlushRenderingCommands = FlushRenderingCommands == 1 || (FlushRenderingCommands == 2 && Bounce == MaxBounces);
-						GraphBuilder.AddPass(
-							CompactionType == 1
-							? RDG_EVENT_NAME("Path Tracer Compute (%d x %d) Tile=(%d,%d - %dx%d) Sample=%d/%d NumLights=%d (Bounce=%d%s)", DispatchResX, DispatchResY, TileX, TileY, DispatchSizeX, DispatchSizeYLocal, PathTracingState->SampleIndex, MaxSPP, PassParameters->SceneLightCount, Bounce, bUseIndirectDispatch && Bounce > 0 ? TEXT(" indirect") : TEXT(""))
-							: RDG_EVENT_NAME("Path Tracer Compute (%d x %d) Tile=(%d,%d - %dx%d) Sample=%d/%d NumLights=%d", DispatchResX, DispatchResY, TileX, TileY, DispatchSizeX, DispatchSizeYLocal, PathTracingState->SampleIndex, MaxSPP, PassParameters->SceneLightCount),
-							PassParameters,
-							ERDGPassFlags::Compute,
-							[PassParameters, RayGenShader, DispatchSizeX, DispatchSizeYLocal, bUseIndirectDispatch, bFlushRenderingCommands, GPUIndex, &View](FRHIRayTracingCommandList& RHICmdList)
+							PassParameters->AtmosphereOpticalDepthLUT = AtmosphereOpticalDepthLUT;
+							PassParameters->AtmosphereOpticalDepthLUTSampler = TStaticSamplerState<SF_Bilinear, AM_Clamp, AM_Clamp, AM_Clamp>::GetRHI();
+
+							if (Config.PathTracingData.EnableFog)
 							{
-								FRHIRayTracingScene* RayTracingSceneRHI = View.GetRayTracingSceneChecked();
+								PassParameters->FogParameters = PrepareFogParameters(View, Scene->ExponentialFogs[0]);
+							}
+							else
+							{
+								PassParameters->FogParameters = {};
+							}
+
+							PassParameters->TilePixelOffset.X = TileX;
+							PassParameters->TilePixelOffset.Y = TileY + CurrentGPU;
+							PassParameters->TileTextureOffset.X = TileX;
+							PassParameters->TileTextureOffset.Y = TileY + CurrentGPU * DispatchSizeYSplit;
+							PassParameters->ScanlineStride = NumGPUs;
+							PassParameters->ScanlineWidth = DispatchSizeX;
 
-								FRayTracingShaderBindingsWriter GlobalResources;
-								SetShaderParameters(GlobalResources, RayGenShader, *PassParameters);
-								if (bUseIndirectDispatch && PassParameters->Bounce > 0)
+							PassParameters->Bounce = Bounce;
+							if (CompactionType == 1)
+							{
+								PassParameters->ActivePaths = GraphBuilder.CreateSRV(ActivePaths[Bounce & 1], PF_R32_SINT);
+								PassParameters->NextActivePaths = GraphBuilder.CreateUAV(ActivePaths[(Bounce & 1) ^ 1], PF_R32_SINT);
+								PassParameters->PathStateData = GraphBuilder.CreateUAV(PathStateData);
+								if (bUseIndirectDispatch)
 								{
-									PassParameters->PathTracingIndirectArgs->MarkResourceAsUsed();
-									RHICmdList.RayTraceDispatchIndirect(
-										View.RayTracingMaterialPipeline,
-										RayGenShader.GetRayTracingShader(),
-										RayTracingSceneRHI, GlobalResources,
-										PassParameters->PathTracingIndirectArgs->GetIndirectRHICallBuffer(), 0
-									);
+									PassParameters->NumPathStates = GraphBuilder.CreateUAV(NumActivePaths[Bounce & 1], PF_R32_UINT);
+									PassParameters->PathTracingIndirectArgs = NumActivePaths[(Bounce & 1) ^ 1];
 								}
 								else
 								{
-									RHICmdList.RayTraceDispatch(
-										View.RayTracingMaterialPipeline,
-										RayGenShader.GetRayTracingShader(),
-										RayTracingSceneRHI, GlobalResources,
-										DispatchSizeX, DispatchSizeYLocal
-									);
+									PassParameters->NumPathStates = GraphBuilder.CreateUAV(NumActivePaths[0], PF_R32_UINT);
+									AddClearUAVPass(GraphBuilder, PassParameters->NextActivePaths, -1); // make sure everything is initialized to -1 since paths that go inactive don't write anything
 								}
-								if (bFlushRenderingCommands)
+								AddClearUAVPass(GraphBuilder, PassParameters->NumPathStates, 0);
+							}
+							ClearUnusedGraphResources(RayGenShader, PassParameters);
+							const bool bFlushRenderingCommands = FlushRenderingCommands == 1 || (FlushRenderingCommands == 2 && Bounce == MaxBounces);
+							GraphBuilder.AddPass(
+								CompactionType == 1
+								? RDG_EVENT_NAME("Path Tracer Compute (%d x %d) Tile=(%d,%d - %dx%d) Sample=%d/%d NumLights=%d (Bounce=%d%s)", DispatchResX, DispatchResY, TileX, TileY, DispatchSizeX, DispatchSizeYLocal, PathTracingState->SampleIndex, MaxSPP, PassParameters->SceneLightCount, Bounce, bUseIndirectDispatch && Bounce > 0 ? TEXT(" indirect") : TEXT(""))
+								: RDG_EVENT_NAME("Path Tracer Compute (%d x %d) Tile=(%d,%d - %dx%d) Sample=%d/%d NumLights=%d", DispatchResX, DispatchResY, TileX, TileY, DispatchSizeX, DispatchSizeYLocal, PathTracingState->SampleIndex, MaxSPP, PassParameters->SceneLightCount),
+								PassParameters,
+								ERDGPassFlags::Compute,
+								[PassParameters, RayGenShader, DispatchSizeX, DispatchSizeYLocal, bUseIndirectDispatch, bFlushRenderingCommands, GPUIndex, &View](FRHIRayTracingCommandList& RHICmdList)
 								{
-									RHICmdList.SubmitCommandsHint();
-								}
-							});
-						if (PreviousPassParameters == nullptr)
-						{
-							PreviousPassParameters = PassParameters;
+									FRHIRayTracingScene* RayTracingSceneRHI = View.GetRayTracingSceneChecked();
+
+									FRayTracingShaderBindingsWriter GlobalResources;
+									SetShaderParameters(GlobalResources, RayGenShader, *PassParameters);
+									if (bUseIndirectDispatch && PassParameters->Bounce > 0)
+									{
+										PassParameters->PathTracingIndirectArgs->MarkResourceAsUsed();
+										RHICmdList.RayTraceDispatchIndirect(
+											View.RayTracingMaterialPipeline,
+											RayGenShader.GetRayTracingShader(),
+											RayTracingSceneRHI, GlobalResources,
+											PassParameters->PathTracingIndirectArgs->GetIndirectRHICallBuffer(), 0
+										);
+									}
+									else
+									{
+										RHICmdList.RayTraceDispatch(
+											View.RayTracingMaterialPipeline,
+											RayGenShader.GetRayTracingShader(),
+											RayTracingSceneRHI, GlobalResources,
+											DispatchSizeX, DispatchSizeYLocal
+										);
+									}
+									if (bFlushRenderingCommands)
+									{
+										RHICmdList.SubmitCommandsHint();
+									}
+								});
+							if (PreviousPassParameters == nullptr)
+							{
+								PreviousPassParameters = PassParameters;
+							}
 						}
 					}
 				}
+				++CurrentGPU;
 			}
-			++CurrentGPU;
+
+			// Bump counters for next frame pass
+			++PathTracingState->SampleIndex;
+			++PathTracingState->FrameIndex;
 		}
+	}
 
+	if (bNeedsTextureExtract)
+	{
 #if WITH_MGPU
 		if (NumGPUs > 1)
 		{
+			// Need fences to prevent cross GPU copies from overlapping with rendering to the same buffers
+			TArray<FTransferResourceFenceData*> CopyFenceDatas;
+			CopyFenceDatas.AddUninitialized(NumGPUs - 1);
+			for (int32 FenceIndex = 0; FenceIndex < NumGPUs - 1; FenceIndex++)
+			{
+				CopyFenceDatas[FenceIndex] = RHICreateTransferResourceFenceData();
+			}
+
+			{
+				// Signal that the first GPU is done rendering, and other GPUs can copy to the buffer now.  Get all the GPUs
+				// besides the first GPU into a mask.  These are the source GPUs for copies to the first GPU.
+				FRHIGPUMask SrcGPUMask = FRHIGPUMask::FromIndex(GPUMask.GetLastIndex());
+				for (uint32 SrcGPUIndex : GPUMask)
+				{
+					if (SrcGPUIndex != GPUMask.GetFirstIndex())
+					{
+						SrcGPUMask |= FRHIGPUMask::FromIndex(SrcGPUIndex);
+					}
+				}
+
+				// Signal goes from first GPU (destination of copy), to remaining GPUs (sources of copy).
+				RDG_GPU_MASK_SCOPE(GraphBuilder, FRHIGPUMask::FromIndex(GPUMask.GetFirstIndex()));
+				GraphBuilder.AddPass(
+					RDG_EVENT_NAME("Path Tracer Cross-GPU Signal (%d GPUs)", NumGPUs),
+					ERDGPassFlags::None,
+					[this, LocalCopyFenceDatas = CopyTemp(CopyFenceDatas), SrcGPUMask](FRHICommandListImmediate& RHICmdList)
+				{
+					RHICmdList.TransferResourceSignal(LocalCopyFenceDatas, SrcGPUMask);
+				});
+			}
+
 			// Treat the cross GPU copy as occurring on all GPUs, for profiling purposes.  Internally, the cross GPU transfer doesn't
 			// pay attention to the mask, so it has no effect on behavior.  Technically the work of the copy is done on the second GPU,
 			// and the first GPU stalls waiting on that, so it's useful to show this interval on both GPUs.
@@ -2362,7 +2425,7 @@ void FDeferredShadingSceneRenderer::RenderPathTracing(
 			Parameters->InputAlbedo = AlbedoTexture;
 			Parameters->InputNormal = NormalTexture;
 			GraphBuilder.AddPass(RDG_EVENT_NAME("Path Tracer Cross-GPU Transfer (%d GPUs)", NumGPUs), Parameters, ERDGPassFlags::Readback,
-				[Parameters, DispatchResX, DispatchResY, DispatchSize, GPUMask, MainGPUMask = View.GPUMask](FRHICommandListImmediate& RHICmdList)
+				[Parameters, DispatchResX, DispatchResY, DispatchSize, GPUMask, MainGPUMask = View.GPUMask, LocalCopyFenceDatas = MoveTemp(CopyFenceDatas)](FRHICommandListImmediate& RHICmdList)
 				{
 					const int32 FirstGPUIndex = MainGPUMask.GetFirstIndex();
 					const int32 NumGPUs = GPUMask.GetNumActive();
@@ -2400,6 +2463,14 @@ void FDeferredShadingSceneRenderer::RenderPathTracing(
 							}
 						}
 					}
+
+					// Include the fences we need to wait on in our list of transfers
+					check(TransferParams.Num() >= LocalCopyFenceDatas.Num());
+					for (int32 FenceIndex = 0; FenceIndex < LocalCopyFenceDatas.Num(); FenceIndex++)
+					{
+						TransferParams[FenceIndex].PreTransferFence = LocalCopyFenceDatas[FenceIndex];
+					}
+
 					RHICmdList.TransferResources(TransferParams);
 				}
 			);
@@ -2409,10 +2480,6 @@ void FDeferredShadingSceneRenderer::RenderPathTracing(
 		GraphBuilder.QueueTextureExtraction(RadianceTexture, &PathTracingState->RadianceRT);
 		GraphBuilder.QueueTextureExtraction(AlbedoTexture, &PathTracingState->AlbedoRT);
 		GraphBuilder.QueueTextureExtraction(NormalTexture, &PathTracingState->NormalRT);
-
-		// Bump counters for next frame
-		++PathTracingState->SampleIndex;
-		++PathTracingState->FrameIndex;
 	}
 
 	RDG_GPU_MASK_SCOPE(GraphBuilder, View.GPUMask);
