On mobile renderer, PrepareDistanceFieldScene does not get called when r.DFShadowQuality is 0. This to prevent DF memory allocation that happens even if DF shadows is not displayed on screen

#rnx

diff --git a/Engine/Source/Runtime/Renderer/Private/MobileShadingRenderer.cpp b/Engine/Source/Runtime/Renderer/Private/MobileShadingRenderer.cpp
--- a/Engine/Source/Runtime/Renderer/Private/MobileShadingRenderer.cpp
+++ b/Engine/Source/Runtime/Renderer/Private/MobileShadingRenderer.cpp
@@ -489,13 +489,16 @@ void FMobileSceneRenderer::InitViews(FRDGBuilder& GraphBuilder, FSceneTexturesCo
 
 	bShouldRenderVelocities = ShouldRenderVelocities();
 
+	static auto CVarDistanceFieldShadowQuality = IConsoleManager::Get().FindConsoleVariable(TEXT("r.DFShadowQuality"));
+
 	bRequiresDistanceField = IsMobileDistanceFieldEnabled(ShaderPlatform)
 		&& ViewFamily.EngineShowFlags.Lighting
 		&& !Views[0].bIsReflectionCapture
 		&& !Views[0].bIsPlanarReflection
 		&& !ViewFamily.EngineShowFlags.HitProxies
 		&& !ViewFamily.EngineShowFlags.VisualizeLightCulling
-		&& !ViewFamily.UseDebugViewPS();
+		&& !ViewFamily.UseDebugViewPS()
+		&& (CVarDistanceFieldShadowQuality != nullptr && CVarDistanceFieldShadowQuality->GetInt() > 0);
 
 	bRequiresShadowProjections = MobileUsesShadowMaskTexture(ShaderPlatform);
 
@@ -644,7 +647,7 @@ void FMobileSceneRenderer::InitViews(FRDGBuilder& GraphBuilder, FSceneTexturesCo
 	{
 		PrepareDistanceFieldScene(GraphBuilder, false);
 	}
-	
+
 	InstanceCullingManager.BeginDeferredCulling(GraphBuilder, Scene->GPUScene);
 
 	extern TSet<IPersistentViewUniformBufferExtension*> PersistentViewUniformBufferExtensions;
