OpenXR: screen percentage enabled, allows TSR/TAAU upscaling

In order to use with OpenXRHMDs

* Standalone games - r.ScreenPercentage
* Editor PIE - r.Editor.Viewport.ScreenPercentageMode.RealTime=0 + r.Editor.Viewport.ScreenPercentage OR r.Editor.Viewport.OverridePIEScreenPercentage + r.ScreenPercentage

VREM still needs changes to the EditorViewportClient code

diff --git a/Engine/Plugins/Runtime/OpenXR/Source/OpenXRHMD/Private/OpenXRHMD.cpp b/Engine/Plugins/Runtime/OpenXR/Source/OpenXRHMD/Private/OpenXRHMD.cpp
--- a/Engine/Plugins/Runtime/OpenXR/Source/OpenXRHMD/Private/OpenXRHMD.cpp
+++ b/Engine/Plugins/Runtime/OpenXR/Source/OpenXRHMD/Private/OpenXRHMD.cpp
@@ -1053,10 +1053,6 @@ void FOpenXRHMD::SetupViewFamily(FSceneViewFamily& InViewFamily)
 	InViewFamily.EngineShowFlags.HMDDistortion = false;
 	InViewFamily.EngineShowFlags.StereoRendering = IsStereoEnabled();
 
-	// TODO: Handle dynamic resolution in the driver, so the runtime
-	// can take advantage of the extra resolution in the distortion process.
-	InViewFamily.EngineShowFlags.ScreenPercentage = 0;
-
 	const FPipelinedFrameState& FrameState = GetPipelinedFrameStateForThread();
 	if (FrameState.Views.Num() > 2)
 	{
diff --git a/Engine/Source/Runtime/Renderer/Private/SceneRendering.cpp b/Engine/Source/Runtime/Renderer/Private/SceneRendering.cpp
--- a/Engine/Source/Runtime/Renderer/Private/SceneRendering.cpp
+++ b/Engine/Source/Runtime/Renderer/Private/SceneRendering.cpp
@@ -2910,10 +2910,13 @@ void FSceneRenderer::PrepareViewRectsForRendering(FRHICommandListImmediate& RHIC
 	// Notify StereoRenderingDevice about new ViewRects
 	if (GEngine->StereoRenderingDevice.IsValid())
 	{
-		for (int32 i = 0; i < Views.Num(); i++)
+		for (const FViewInfo& View : Views)
 		{
-			FViewInfo& View = Views[i];
-			GEngine->StereoRenderingDevice->SetFinalViewRect(RHICmdList, View.StereoViewIndex, View.ViewRect);
+			// if we have an upscale pass, the final rect is _unscaled_ for the compositor
+			const FIntRect OutputViewRect =
+				(View.PrimaryScreenPercentageMethod == EPrimaryScreenPercentageMethod::RawOutput) ? View.ViewRect : View.UnscaledViewRect;
+
+			GEngine->StereoRenderingDevice->SetFinalViewRect(RHICmdList, View.StereoViewIndex, OutputViewRect);
 		}
 	}
 }
