Add Strata support for composite cubemap.

diff --git a/Engine/Shaders/Private/AmbientCubemapComposite.usf b/Engine/Shaders/Private/AmbientCubemapComposite.usf
--- a/Engine/Shaders/Private/AmbientCubemapComposite.usf
+++ b/Engine/Shaders/Private/AmbientCubemapComposite.usf
@@ -19,6 +19,119 @@
 Texture2D AmbientOcclusionTexture;
 SamplerState AmbientOcclusionSampler;
 
+#if STRATA_ENABLED
+
+// Override the default lighting input used for the env. lighting to sample the 'custom composite env. lighting'
+float3 GetEnvDiffuseLighting(float3 InBentNormal)
+{ 
+	const float AbsoluteDiffuseMip = AmbientCubemapMipAdjust.z;
+	return TextureCubeSampleLevel(AmbientCubemap, AmbientCubemapSampler, InBentNormal, AbsoluteDiffuseMip).rgb * AmbientCubemapColor.rgb;
+}
+
+float3 GetEnvSpecularLighting(
+	float CompositeAlpha,
+	float3 TranslatedWorldPosition,
+	float3 SpecularDirection,
+	float SpecularSafeRoughness,
+	float IndirectIrradiance,
+	float IndirectSpecularOcclusion,
+	float3 ExtraIndirectSpecular,
+	uint NumCulledReflectionCaptures,
+	uint CaptureDataStartIndex)
+{
+	const float SpecularMip = ComputeCubemapMipFromRoughness(SpecularSafeRoughness, AmbientCubemapMipAdjust.w);
+	return TextureCubeSampleLevel(AmbientCubemap, AmbientCubemapSampler, SpecularDirection, SpecularMip).rgb * AmbientCubemapColor.rgb;
+}
+
+#define APPLY_SKY_SHADOWING 0
+#define USE_DEFAULT_ENV_LIGHTING_INPUT 0
+#define USE_STRATA_ENV_LIGHTING_COMMON 1
+
+#include "ReflectionEnvironmentShared.ush"
+#include "Strata/Strata.ush"
+#include "Strata/StrataEvaluation.ush"
+#include "Strata/StrataEnvironmentLighting.ush"
+
+void MainPS(
+in noperspective float4 UVAndScreenPos : TEXCOORD0, 
+	float4 SvPosition : SV_POSITION, 
+	out float4 OutColor : SV_Target0
+	#if STRATA_OPAQUE_ROUGH_REFRACTION_ENABLED
+	, out float3 OutOpaqueRoughRefractionSceneColor : SV_Target1
+	, out float3 OutSubSurfaceSceneColor : SV_Target2
+	#endif
+	)
+{
+	const int2 PixelPos = int2(SvPosition.xy);
+	const float2 BufferUV = UVAndScreenPos.xy;
+	const float2 ScreenPosition = UVAndScreenPos.zw;
+
+	OutColor = 0;
+#if STRATA_OPAQUE_ROUGH_REFRACTION_ENABLED
+	OutOpaqueRoughRefractionSceneColor = 0.0f;
+	OutSubSurfaceSceneColor = 0.0f;
+#endif
+
+	FStrataAddressing StrataAddressing = GetStrataPixelDataByteOffset(PixelPos, uint2(View.BufferSizeAndInvSize.xy), Strata.MaxBytesPerPixel);
+	FStrataPixelHeader StrataPixelHeader = UnpackStrataHeaderIn(Strata.MaterialTextureArray, StrataAddressing, Strata.TopLayerTexture);
+	BRANCH
+	if (StrataPixelHeader.BSDFCount > 0) // This test is also enough to exclude sky pixels
+	{
+		float DeviceZ = SampleDeviceZFromSceneTextures(BufferUV);
+		float SceneDepth = ConvertFromDeviceZ(DeviceZ);
+
+		float3 WorldPosition = mul(float4(ScreenPosition * SceneDepth, SceneDepth, 1), LWCHackToFloat(PrimaryView.ScreenToWorld)).xyz;
+		float3 TranslatedWorldPosition = mul(float4(ScreenPosition * SceneDepth, SceneDepth, 1), View.ScreenToTranslatedWorld).xyz;
+		float3 CameraToPixel = normalize(WorldPosition - LWCHackToFloat(PrimaryView.WorldCameraOrigin));
+		float3 V = -CameraToPixel;
+
+		// Sample the ambient occlusion that is dynamically generated every frame.
+		const float AmbientOcclusion = AmbientOcclusionTexture.SampleLevel(AmbientOcclusionSampler, BufferUV, 0).r;
+
+		const float4 SSR = 0.f; //Texture2DSample(ScreenSpaceReflectionsTexture, ScreenSpaceReflectionsSampler, BufferUV);
+
+		// SSR is computed for the top level by averaging BSDF components.
+		// This is the top level specular contribution factor according to areas.
+		const float TopLayerSpecularContributionFactor = 1.0f - SSR.a;
+
+		// A reduction of SSR contribution is needed when there is haziness because we only consider SSR for the sharpest lobe.
+		float SSRReductionFactor = 1.0f;
+		float3 SSRTopLayerEnvBRDF = 0.0f;
+
+		const FStrataIntegrationSettings Settings = InitStrataIntegrationSettings(false /*bForceFullyRough*/, Strata.bRoughDiffuse, Strata.PeelLayersAboveDepth);
+
+		FStrataDeferredLighting StrataLighting = StrataEnvironmentLighting(
+			SvPosition,
+			BufferUV,
+			SceneDepth,
+			WorldPosition,
+			TranslatedWorldPosition,
+			CameraToPixel,
+			V,
+			AmbientOcclusion,
+			TopLayerSpecularContributionFactor,
+			true /*bEnableSpecular*/,
+			Settings,
+			Strata.MaterialTextureArray,
+			StrataAddressing,
+			StrataPixelHeader,
+			1.0f /*CloudVolumetricAOShadow*/,
+			SSRReductionFactor,
+			SSRTopLayerEnvBRDF);
+
+		// Screen space reflections added on top of environment lighting.
+		// Alpha is 0 to not affect the subsurface diffuse luma information.
+		const float4 SSRColor = float4(SSR.rgb * SSRTopLayerEnvBRDF * SSRReductionFactor, 0.0);
+
+		OutColor += StrataLighting.SceneColor + SSRColor;
+#if STRATA_OPAQUE_ROUGH_REFRACTION_ENABLED
+		OutOpaqueRoughRefractionSceneColor += StrataLighting.OpaqueRoughRefractionSceneColor;
+		OutSubSurfaceSceneColor += StrataLighting.SubSurfaceSceneColor;
+#endif
+	}
+}
+
+#else // STRATA_ENABLED
 
 float3 DiffuseIBL( uint2 Random, float3 DiffuseColor, float Roughness, float3 N, float3 V )
 {
@@ -478,3 +591,4 @@ void MainPS(in noperspective float4 UVAndScreenPos : TEXCOORD0, float4 SvPositio
 
 	OutColor *= AmbientOcclusion;
 }
+#endif // STRATA_ENABLED
\ No newline at end of file
diff --git a/Engine/Shaders/Private/ReflectionEnvironmentPixelShader.usf b/Engine/Shaders/Private/ReflectionEnvironmentPixelShader.usf
--- a/Engine/Shaders/Private/ReflectionEnvironmentPixelShader.usf
+++ b/Engine/Shaders/Private/ReflectionEnvironmentPixelShader.usf
@@ -293,6 +293,8 @@ void ReflectionEnvironmentSkyLighting(
 		const float CloudVolumetricAOShadow = GetCloudVolumetricAOShadow(TranslatedWorldPosition);
 		const FStrataIntegrationSettings Settings = InitStrataIntegrationSettings(false /*bForceFullyRough*/, Strata.bRoughDiffuse, Strata.PeelLayersAboveDepth);
 
+		const bool bEnableSpecular = ReflectionStruct.SkyLightParameters.y > 0.0f;
+
 		FStrataDeferredLighting StrataLighting = StrataEnvironmentLighting(
 			SvPosition,
 			BufferUV,
@@ -303,6 +305,7 @@ void ReflectionEnvironmentSkyLighting(
 			V,
 			AmbientOcclusion,
 			TopLayerSpecularContributionFactor,
+			bEnableSpecular,
 			Settings,
 			Strata.MaterialTextureArray,
 			StrataAddressing,
diff --git a/Engine/Shaders/Private/Strata/StrataEnvironmentLighting.ush b/Engine/Shaders/Private/Strata/StrataEnvironmentLighting.ush
--- a/Engine/Shaders/Private/Strata/StrataEnvironmentLighting.ush
+++ b/Engine/Shaders/Private/Strata/StrataEnvironmentLighting.ush
@@ -25,6 +25,7 @@ FStrataDeferredLighting StrataEnvironmentLighting(
 	float3 V,
 	float ScreenAmbientOcclusion,
 	float TopLayerSpecularContributionFactor,
+	bool bEnableSpecular,
 	FStrataIntegrationSettings Settings,
 	FStrataMaterialContainer MaterialBuffer,
 	FStrataAddressing StrataAddressing,
@@ -45,7 +46,7 @@ FStrataDeferredLighting StrataEnvironmentLighting(
 	float CombinedScreenAndMaterialAO = StrataGetAO(StrataPixelHeader) * ScreenAmbientOcclusion;
 
 	
-#if FEATURE_LEVEL >= FEATURE_LEVEL_SM5
+#if USE_DEFAULT_ENV_LIGHTING_INPUT && FEATURE_LEVEL >= FEATURE_LEVEL_SM5
 	float2 LocalPosition = SvPosition.xy - View.ViewRectMin.xy;
 
 	uint GridIndex = ComputeLightGridCellIndex(uint2(LocalPosition.x, LocalPosition.y), SceneDepth);
@@ -67,7 +68,6 @@ FStrataDeferredLighting StrataEnvironmentLighting(
 		const float3 BSDFThroughput = LuminanceWeight(StrataBSDFContext, BSDF); // Use the reflected direction
 
 		// Evaluate environment lighting
-		const bool bEnableSpecular = ReflectionStruct.SkyLightParameters.y > 0.0f;
 		FStrataEnvLightResult StrataEnvLight = StrataEvaluateForEnvLight(StrataBSDFContext, bEnableSpecular, Settings);
 
 		float3 BentNormal = StrataEnvLight.DiffuseNormal;
diff --git a/Engine/Shaders/Private/Strata/StrataLightingCommon.ush b/Engine/Shaders/Private/Strata/StrataLightingCommon.ush
--- a/Engine/Shaders/Private/Strata/StrataLightingCommon.ush
+++ b/Engine/Shaders/Private/Strata/StrataLightingCommon.ush
@@ -129,6 +129,46 @@ float3 StrataForwardLightingCommon(
 
 #ifdef USE_STRATA_ENV_LIGHTING_COMMON
 
+// USE_DEFAULT_ENV_LIGHTING_INPUT is use to allow lighting input override
+#ifndef USE_DEFAULT_ENV_LIGHTING_INPUT
+#define USE_DEFAULT_ENV_LIGHTING_INPUT 1
+#endif 
+
+#if USE_DEFAULT_ENV_LIGHTING_INPUT
+// Diffuse input lighting (sky)
+float3 GetEnvDiffuseLighting(float3 InBentNormal)
+{
+	return GetSkySHDiffuse(InBentNormal) * View.SkyLightColor.rgb;
+}
+
+// Specular input lighting (reflection capture + sky)
+float3 GetEnvSpecularLighting(
+	float CompositeAlpha,
+	float3 TranslatedWorldPosition,
+	float3 SpecularDirection,
+	float SpecularSafeRoughness,
+	float IndirectIrradiance,
+	float IndirectSpecularOcclusion,
+	float3 ExtraIndirectSpecular,
+	uint NumCulledReflectionCaptures,
+	uint CaptureDataStartIndex)
+{
+	const bool bCompositeSkylight = true;
+	return CompositeReflectionCapturesAndSkylightTWS(
+		CompositeAlpha,
+		TranslatedWorldPosition,
+		SpecularDirection,
+		SpecularSafeRoughness,
+		IndirectIrradiance,
+		IndirectSpecularOcclusion,
+		ExtraIndirectSpecular,
+		NumCulledReflectionCaptures,
+		CaptureDataStartIndex,
+		0,
+		bCompositeSkylight);
+}
+#endif // USE_DEFAULT_ENV_LIGHTING_INPUT
+
 void StrataEnvLightingCommon(
 	in FStrataEnvLightResult StrataEnvLight,
 	in FStrataPixelHeader StrataPixelHeader,
@@ -162,7 +202,7 @@ void StrataEnvLightingCommon(
 		if (bProcessFrontFaceDiffuse)
 		{
 			// Finally sample the sky diffuse contribution (spherical harmonic, Lambert BRDF)
-			float3 DiffuseLookup = GetSkySHDiffuse(BentNormal) * View.SkyLightColor.rgb;
+			float3 DiffuseLookup = GetEnvDiffuseLighting(BentNormal);
 			// And accumulate
 			// Note: Use diffuse directional albedo (i.e., DiffuseWeight) as first order approximation for env. integration (STRATA_TODO instead compute SH coefficients for Chan)
 			DiffuseLighting = CloudVolumetricAOShadow * BSDFThroughput * (SkyVisData.SkyDiffuseLookUpMul * DiffuseLookup + SkyVisData.SkyDiffuseLookUpAdd) * StrataEnvLight.DiffuseWeight;
@@ -173,7 +213,7 @@ void StrataEnvLightingCommon(
 			// However, we do apply SkyVisData.SkyDiffuseLookUpMul for scaling consistency.
 
 			// Finally sample the sky diffuse contribution (spherical harmonic, Lambert BRDF) along the opposite normal direction
-			float3 DiffuseLookup = GetSkySHDiffuse(-BentNormal) * View.SkyLightColor.rgb;
+			float3 DiffuseLookup = GetEnvDiffuseLighting(-BentNormal);
 			// And accumulate
 			// Note: Use diffuse directional albedo (i.e., DiffuseWeight) as first order approximation for env. integration (STRATA_TODO instead compute SH coefficients for Chan)
 			DiffuseLighting += CloudVolumetricAOShadow * BSDFThroughput * (SkyVisData.SkyDiffuseLookUpMul * DiffuseLookup) * StrataEnvLight.DiffuseBackFaceWeight;
@@ -215,9 +255,8 @@ void StrataEnvLightingCommon(
 		float RoughnessSquared = StrataEnvLight.SpecularSafeRoughness * StrataEnvLight.SpecularSafeRoughness;
 		float SpecularOcclusion = IndirectSpecularOcclusion * GetSpecularOcclusion(StrataBSDFContext.SatNoV, RoughnessSquared, CombinedScreenAndMaterialAO);
 
-		const bool bCompositeSkylight = true;
 		SpecularLighting += BSDFThroughput * StrataEnvLight.SpecularWeight *
-			CompositeReflectionCapturesAndSkylightTWS(
+			GetEnvSpecularLighting(
 				(bIsTopLayer ? TopLayerSpecularContributionFactor : 1.0f) * SpecularOcclusion,
 				TranslatedWorldPosition,
 				StrataEnvLight.SpecularDirection,
@@ -226,15 +265,13 @@ void StrataEnvLightingCommon(
 				IndirectSpecularOcclusion,
 				ExtraIndirectSpecular,
 				NumCulledReflectionCaptures,
-				CaptureDataStartIndex,
-				0,
-				bCompositeSkylight);
+				CaptureDataStartIndex);
 
 	#if STRATA_FASTPATH==0
 		if (BSDF_GETHASHAZINESS(BSDF))
 		{
 			SpecularLighting += BSDFThroughput * StrataEnvLight.SpecularHazeWeight *
-				CompositeReflectionCapturesAndSkylightTWS(
+				GetEnvSpecularLighting(
 					(bIsTopLayer ? TopLayerSpecularContributionFactor : 1.0f) * SpecularOcclusion,
 					TranslatedWorldPosition,
 					StrataEnvLight.SpecularHazeDirection,
@@ -243,9 +280,7 @@ void StrataEnvLightingCommon(
 					IndirectSpecularOcclusion,
 					ExtraIndirectSpecular,
 					NumCulledReflectionCaptures,
-					CaptureDataStartIndex,
-					0,
-					bCompositeSkylight);
+					CaptureDataStartIndex);
 
 			// SSR is traced for the sharpest lob. The smoothest one does not rely on SSR so we need to lower energy coming from SSR according to the lobe blend weight.
 			// And we also try to make the transition smooth using Haziness
diff --git a/Engine/Source/Runtime/Renderer/Private/IndirectLightRendering.cpp b/Engine/Source/Runtime/Renderer/Private/IndirectLightRendering.cpp
--- a/Engine/Source/Runtime/Renderer/Private/IndirectLightRendering.cpp
+++ b/Engine/Source/Runtime/Renderer/Private/IndirectLightRendering.cpp
@@ -197,7 +197,9 @@ class FAmbientCubemapCompositePS : public FGlobalShader
 
 		SHADER_PARAMETER_RDG_TEXTURE(Texture2D, AmbientOcclusionTexture)
 		SHADER_PARAMETER_SAMPLER(SamplerState,  AmbientOcclusionSampler)
-		
+
+		SHADER_PARAMETER_RDG_UNIFORM_BUFFER(FStrataGlobalUniformParameters, Strata)
+		SHADER_PARAMETER_STRUCT_INCLUDE(Strata::FStrataTilePassVS::FParameters, StrataTile)
 		SHADER_PARAMETER_STRUCT_INCLUDE(FAmbientCubemapParameters, AmbientCubemap)
 		SHADER_PARAMETER_STRUCT_INCLUDE(FSceneTextureParameters, SceneTextures)
 		SHADER_PARAMETER_STRUCT_REF(FViewUniformShaderParameters, ViewUniformBuffer)
@@ -1239,72 +1241,110 @@ void FDeferredShadingSceneRenderer::RenderDiffuseIndirectAndAmbientOcclusion(
 		// Apply the ambient cubemaps
 		if (IsAmbientCubemapPassRequired(View) && !bIsVisualizePass && !ViewPipelineState.bUseLumenProbeHierarchy)
 		{
-			FAmbientCubemapCompositePS::FParameters* PassParameters = GraphBuilder.AllocParameters<FAmbientCubemapCompositePS::FParameters>();
-			
-			PassParameters->PreIntegratedGF = GSystemTextures.PreintegratedGF->GetRHI();
-			PassParameters->PreIntegratedGFSampler = TStaticSamplerState<SF_Bilinear, AM_Clamp, AM_Clamp, AM_Clamp>::GetRHI();
+			auto ApplyAmbientCubemapComposite = [&](EStrataTileType TileType)
+			{			
+				FAmbientCubemapCompositePS::FParameters* PassParameters = GraphBuilder.AllocParameters<FAmbientCubemapCompositePS::FParameters>();
+				PassParameters->PreIntegratedGF = GSystemTextures.PreintegratedGF->GetRHI();
+				PassParameters->PreIntegratedGFSampler = TStaticSamplerState<SF_Bilinear, AM_Clamp, AM_Clamp, AM_Clamp>::GetRHI();
 			
-			PassParameters->AmbientOcclusionTexture = AmbientOcclusionMask;
-			PassParameters->AmbientOcclusionSampler = TStaticSamplerState<SF_Point>::GetRHI();
+				PassParameters->AmbientOcclusionTexture = AmbientOcclusionMask;
+				PassParameters->AmbientOcclusionSampler = TStaticSamplerState<SF_Point>::GetRHI();
 			
-			if (!PassParameters->AmbientOcclusionTexture)
-			{
-				PassParameters->AmbientOcclusionTexture = SystemTextures.White;
-			}
+				if (!PassParameters->AmbientOcclusionTexture)
+				{
+					PassParameters->AmbientOcclusionTexture = SystemTextures.White;
+				}
 
-			PassParameters->SceneTextures = SceneTextureParameters;
-			PassParameters->ViewUniformBuffer = View.ViewUniformBuffer;
+				PassParameters->Strata = Strata::BindStrataGlobalUniformParameters(View);
+				PassParameters->SceneTextures = SceneTextureParameters;
+				PassParameters->ViewUniformBuffer = View.ViewUniformBuffer;
 
-			PassParameters->RenderTargets[0] = FRenderTargetBinding(
-				SceneColorTexture, ERenderTargetLoadAction::ELoad);
-		
-			TShaderMapRef<FAmbientCubemapCompositePS> PixelShader(View.ShaderMap);
-			GraphBuilder.AddPass(
-				RDG_EVENT_NAME("AmbientCubemapComposite %dx%d", View.ViewRect.Width(), View.ViewRect.Height()),
-				PassParameters,
-				ERDGPassFlags::Raster,
-				[PassParameters, &View, PixelShader](FRHICommandList& RHICmdList)
-			{
-				TShaderMapRef<FPostProcessVS> VertexShader(View.ShaderMap);
-				
-				RHICmdList.SetViewport(View.ViewRect.Min.X, View.ViewRect.Min.Y, 0.0f, View.ViewRect.Max.X, View.ViewRect.Max.Y, 0.0);
+				PassParameters->RenderTargets[0] = FRenderTargetBinding(SceneColorTexture, ERenderTargetLoadAction::ELoad);
+				if (Strata::IsStrataOpaqueMaterialRoughRefractionEnabled())
+				{
+					PassParameters->RenderTargets[1] = FRenderTargetBinding(Scene->StrataSceneData.SeparatedOpaqueRoughRefractionSceneColor, ERenderTargetLoadAction::ELoad);
+					PassParameters->RenderTargets[2] = FRenderTargetBinding(Scene->StrataSceneData.SeparatedSubSurfaceSceneColor, ERenderTargetLoadAction::ELoad);
+				}
 
-				FGraphicsPipelineStateInitializer GraphicsPSOInit;
-				RHICmdList.ApplyCachedRenderTargets(GraphicsPSOInit);
+				TShaderMapRef<FPostProcessVS> FullScreenVertexShader(View.ShaderMap);
 
-				// set the state
-				GraphicsPSOInit.BlendState = TStaticBlendState<CW_RGB, BO_Add, BF_One, BF_One, BO_Add, BF_One, BF_One>::GetRHI();
-				GraphicsPSOInit.RasterizerState = TStaticRasterizerState<>::GetRHI();
-				GraphicsPSOInit.DepthStencilState = TStaticDepthStencilState<false, CF_Always>::GetRHI();
+				Strata::FStrataTilePassVS::FPermutationDomain VSPermutationVector;
+				VSPermutationVector.Set< Strata::FStrataTilePassVS::FEnableDebug >(false);
+				VSPermutationVector.Set< Strata::FStrataTilePassVS::FEnableTexCoordScreenVector >(false);
+				TShaderMapRef<Strata::FStrataTilePassVS> TileVertexShader(View.ShaderMap, VSPermutationVector);
 
-				GraphicsPSOInit.BoundShaderState.VertexDeclarationRHI = GFilterVertexDeclaration.VertexDeclarationRHI;
-				GraphicsPSOInit.BoundShaderState.VertexShaderRHI = VertexShader.GetVertexShader();
-				GraphicsPSOInit.BoundShaderState.PixelShaderRHI = PixelShader.GetPixelShader();
-				GraphicsPSOInit.PrimitiveType = PT_TriangleList;
+				const bool bStrataEnabled = Strata::IsStrataEnabled();
+				EPrimitiveType PrimitiveType = PT_TriangleList;
+				if (Strata::IsStrataEnabled())
+				{
+					PassParameters->StrataTile = Strata::SetTileParameters(GraphBuilder, View, TileType, PrimitiveType);
+				}
 
-				SetGraphicsPipelineState(RHICmdList, GraphicsPSOInit, 0);
+				TShaderMapRef<FAmbientCubemapCompositePS> PixelShader(View.ShaderMap);
+				GraphBuilder.AddPass(
+					RDG_EVENT_NAME("AmbientCubemapComposite(%dx%d%s%s)", View.ViewRect.Width(), View.ViewRect.Height(), bStrataEnabled ? TEXT(",") : TEXT(""), bStrataEnabled ? ToString(TileType) : TEXT("")),
+					PassParameters,
+					ERDGPassFlags::Raster,
+					[PassParameters, &View, FullScreenVertexShader, TileVertexShader, PixelShader, PrimitiveType, TileType, bStrataEnabled](FRHICommandList& RHICmdList)
+				{				
+					RHICmdList.SetViewport(View.ViewRect.Min.X, View.ViewRect.Min.Y, 0.0f, View.ViewRect.Max.X, View.ViewRect.Max.Y, 0.0);
 
-				uint32 Count = View.FinalPostProcessSettings.ContributingCubemaps.Num();
-				for (const FFinalPostProcessSettings::FCubemapEntry& CubemapEntry : View.FinalPostProcessSettings.ContributingCubemaps)
-				{
-					FAmbientCubemapCompositePS::FParameters ShaderParameters = *PassParameters;
-					SetupAmbientCubemapParameters(CubemapEntry, &ShaderParameters.AmbientCubemap);
-					SetShaderParameters(RHICmdList, PixelShader, PixelShader.GetPixelShader(), ShaderParameters);
+					FGraphicsPipelineStateInitializer GraphicsPSOInit;
+					RHICmdList.ApplyCachedRenderTargets(GraphicsPSOInit);
+
+					// set the state
+					GraphicsPSOInit.BlendState = TStaticBlendState<CW_RGB, BO_Add, BF_One, BF_One, BO_Add, BF_One, BF_One>::GetRHI();
+					GraphicsPSOInit.RasterizerState = TStaticRasterizerState<>::GetRHI();
+					GraphicsPSOInit.DepthStencilState = TStaticDepthStencilState<false, CF_Always>::GetRHI();
+
+					GraphicsPSOInit.BoundShaderState.VertexDeclarationRHI = GFilterVertexDeclaration.VertexDeclarationRHI;
+					GraphicsPSOInit.BoundShaderState.VertexShaderRHI = bStrataEnabled ? TileVertexShader.GetVertexShader() : FullScreenVertexShader.GetVertexShader();
+					GraphicsPSOInit.BoundShaderState.PixelShaderRHI = PixelShader.GetPixelShader();
+					GraphicsPSOInit.PrimitiveType = PrimitiveType;
+
+					SetGraphicsPipelineState(RHICmdList, GraphicsPSOInit, 0);
+
+					uint32 Count = View.FinalPostProcessSettings.ContributingCubemaps.Num();
+					for (const FFinalPostProcessSettings::FCubemapEntry& CubemapEntry : View.FinalPostProcessSettings.ContributingCubemaps)
+					{
+						FAmbientCubemapCompositePS::FParameters ShaderParameters = *PassParameters;
+						SetupAmbientCubemapParameters(CubemapEntry, &ShaderParameters.AmbientCubemap);
+						SetShaderParameters(RHICmdList, PixelShader, PixelShader.GetPixelShader(), ShaderParameters);
 					
-					DrawPostProcessPass(
-						RHICmdList,
-						0, 0,
-						View.ViewRect.Width(), View.ViewRect.Height(),
-						View.ViewRect.Min.X, View.ViewRect.Min.Y,
-						View.ViewRect.Width(), View.ViewRect.Height(),
-						View.ViewRect.Size(),
-						View.GetSceneTexturesConfig().Extent,
-						VertexShader,
-						View.StereoViewIndex,
-						false, // TODO.
-						EDRF_UseTriangleOptimization);
-				}
-			});
+						if (bStrataEnabled)
+						{
+							SetShaderParameters(RHICmdList, TileVertexShader, TileVertexShader.GetVertexShader(), PassParameters->StrataTile);
+							RHICmdList.DrawPrimitiveIndirect(PassParameters->StrataTile.TileIndirectBuffer->GetIndirectRHICallBuffer(), Strata::TileTypeDrawIndirectArgOffset(TileType));
+						}
+						else
+						{
+							DrawPostProcessPass(
+								RHICmdList,
+								0, 0,
+								View.ViewRect.Width(), View.ViewRect.Height(),
+								View.ViewRect.Min.X,   View.ViewRect.Min.Y,
+								View.ViewRect.Width(), View.ViewRect.Height(),
+								View.ViewRect.Size(),
+								View.GetSceneTexturesConfig().Extent,
+								FullScreenVertexShader,
+								View.StereoViewIndex,
+								false, // TODO.
+								EDRF_UseTriangleOptimization);
+						}
+					}
+				});				
+			};
+
+			if (Strata::IsStrataEnabled())
+			{
+				ApplyAmbientCubemapComposite(EStrataTileType::EComplex);
+				ApplyAmbientCubemapComposite(EStrataTileType::ESingle);
+				ApplyAmbientCubemapComposite(EStrataTileType::ESimple);
+			}
+			else
+			{
+				ApplyAmbientCubemapComposite(EStrataTileType::ECount);
+			}
 		} // if (IsAmbientCubemapPassRequired(View))
 	} // for (FViewInfo& View : Views)
 }
