Scene Capture:  Use after free bug for SceneRenderer structure, related to stat gathering for new Texture Editor Scene Capture memory display feature.  Need to grab view state pointer before calling renderer, which may delete the structure.

#rnx

diff --git a/Engine/Source/Runtime/Renderer/Private/SceneCaptureRendering.cpp b/Engine/Source/Runtime/Renderer/Private/SceneCaptureRendering.cpp
--- a/Engine/Source/Runtime/Renderer/Private/SceneCaptureRendering.cpp
+++ b/Engine/Source/Runtime/Renderer/Private/SceneCaptureRendering.cpp
@@ -1016,13 +1016,17 @@ void FScene::UpdateSceneCaptureContents(USceneCaptureComponent2D* CaptureCompone
 
 				RectLightAtlas::FAtlasTextureInvalidationScope Invalidation(TexturePtrNotDeferenced);
 
+#if WITH_EDITOR
+				// Scene renderer may be deleted in UpdateSceneCaptureContent_RenderThread, grab view state pointer first
+				const FSceneViewState* ViewState = SceneRenderer->Views[0].ViewState;
+#endif  // WITH_EDITOR
+
 				// Don't clear the render target when compositing, or in a tiling mode that fills in the render target in multiple passes.
 				bool bClearRenderTarget = !bIsCompositing && !bEnableOrthographicTiling;
 
 				UpdateSceneCaptureContent_RenderThread(RHICmdList, SceneRenderer, TextureRenderTargetResource, TextureRenderTargetResource, EventName, CopyInfo, bGenerateMips, GenerateMipsParams, bClearRenderTarget, bOrthographicCamera);
 
 #if WITH_EDITOR
-				const FSceneViewState* ViewState = SceneRenderer->Views[0].ViewState;
 				if (ViewState)
 				{
 					const bool bLogSizes = GDumpSceneCaptureMemoryFrame == GFrameNumberRenderThread;
