Add Strata tile coord encoding on 8 bit on platform with resolutin which do not exceed 1080p.

This saves memory bandwidth during classification

diff --git a/Engine/Shaders/Private/Strata/StrataTile.ush b/Engine/Shaders/Private/Strata/StrataTile.ush
--- a/Engine/Shaders/Private/Strata/StrataTile.ush
+++ b/Engine/Shaders/Private/Strata/StrataTile.ush
@@ -8,20 +8,32 @@
 
 #if STRATA_ENABLED
 
+#ifndef USE_8BIT_TILE_COORD
+#error USE_8BIT_TILE_COORD needs to be defined
+#define USE_8BIT_TILE_COORD 0
+#endif
 ////////////////////////////////////////////////////////////////////////////////////////////////////////////
 // Strata material type tile
 
 uint2 StrataUnpackTile(uint In)
 {
+#if USE_8BIT_TILE_COORD
+	return uint2(In & 0xFF, In >> 8);
+#else
 	return uint2(In & 0xFFFF, In >> 16);
+#endif
 }
 
 uint StrataPackTile(uint2 TileCoord)
 {
-#if COMPILER_SUPPORTS_PACK_INTRINSICS
-	return PackUInt2ToUInt(TileCoord.x, TileCoord.y);
+#if USE_8BIT_TILE_COORD
+	return TileCoord.x | (TileCoord.y << 8);
 #else
+  #if COMPILER_SUPPORTS_PACK_INTRINSICS
+	return PackUInt2ToUInt(TileCoord.x, TileCoord.y);
+  #else
 	return TileCoord.x | (TileCoord.y << 16); // assumes 16bit is enough to represent a tiled resolution up to 65,535 :)
+  #endif
 #endif
 }
 
diff --git a/Engine/Source/Runtime/Engine/Private/ShaderCompiler/ShaderCompiler.cpp b/Engine/Source/Runtime/Engine/Private/ShaderCompiler/ShaderCompiler.cpp
--- a/Engine/Source/Runtime/Engine/Private/ShaderCompiler/ShaderCompiler.cpp
+++ b/Engine/Source/Runtime/Engine/Private/ShaderCompiler/ShaderCompiler.cpp
@@ -6024,10 +6024,14 @@ void GlobalBeginCompileShader(
 			static FShaderPlatformCachedIniValue<int32> CVarStrataShadingQuality(TEXT("r.Strata.ShadingQuality"));
 			const uint32 StrataShadingQuality = CVarStrataShadingQuality.Get(Target.GetPlatform());
 			Input.Environment.SetDefine(TEXT("STRATA_SHADING_QUALITY"), FMath::Max(StrataShadingQuality, 1u));
-			if (StrataShadingQuality > 0)
+			if (StrataShadingQuality > 1)
 			{
 				Input.Environment.SetDefine(TEXT("USE_ACHROMATIC_BXDF_ENERGY"), 1u);
 			}
+
+			static FShaderPlatformCachedIniValue<int32> CVarStrataTileCoord8Bits(TEXT("r.Strata.TileCoord8bits"));
+			const uint32 StrataTileCoord8Bits = CVarStrataTileCoord8Bits.Get(Target.GetPlatform());
+			Input.Environment.SetDefine(TEXT("USE_8BIT_TILE_COORD"), StrataTileCoord8Bits > 0 ? 1 : 0);
 		}
 
 		static IConsoleVariable* CVarBackCompatibility = IConsoleManager::Get().FindConsoleVariable(TEXT("r.StrataBackCompatibility"));
diff --git a/Engine/Source/Runtime/Renderer/Private/Strata/Strata.cpp b/Engine/Source/Runtime/Renderer/Private/Strata/Strata.cpp
--- a/Engine/Source/Runtime/Renderer/Private/Strata/Strata.cpp
+++ b/Engine/Source/Runtime/Renderer/Private/Strata/Strata.cpp
@@ -38,7 +38,13 @@ static TAutoConsoleVariable<int32> CVarStrataBytePerPixel(
 static TAutoConsoleVariable<int32> CVarStrataShadingQuality(
 	TEXT("r.Strata.ShadingQuality"),
 	1,
-	TEXT("Defube Strata shading quality (1: accurate lighting, 2: approximate lighting). This variable is read-only."),
+	TEXT("Define Strata shading quality (1: accurate lighting, 2: approximate lighting). This variable is read-only."),
+	ECVF_ReadOnly | ECVF_RenderThreadSafe);
+
+static TAutoConsoleVariable<int32> CVarStrataTileCoord8Bits(
+	TEXT("r.Strata.TileCoord8bits"),
+	0,
+	TEXT("Format of tile coord. This variable is read-only."),
 	ECVF_ReadOnly | ECVF_RenderThreadSafe);
 
 static TAutoConsoleVariable<int32> CVarStrataRoughDiffuse(
@@ -157,6 +163,13 @@ FIntPoint GetStrataTextureResolution(const FIntPoint& InResolution)
 
 static void BindStrataGlobalUniformParameters(FRDGBuilder& GraphBuilder, FStrataViewData* StrataViewData, FStrataGlobalUniformParameters& OutStrataUniformParameters);
 
+static EPixelFormat GetClassificationTileFormat(const FIntPoint& InResolution)
+{
+	// For platform which whose resolution is never above 1080p, use 8bit tile format for performance.
+	check(InResolution.X <= 2048 && InResolution.Y <= 2048);
+	return CVarStrataTileCoord8Bits.GetValueOnRenderThread() == 1 ? PF_R16_UINT : PF_R32_UINT;
+}
+
 static void InitialiseStrataViewData(FRDGBuilder& GraphBuilder, FViewInfo& View, const FSceneTexturesConfig& SceneTexturesConfig, bool bNeedBSDFOffets, FStrataSceneData& SceneData)
 {
 	// Sanity check: the scene data should already exist 
@@ -192,11 +205,12 @@ static void InitialiseStrataViewData(FRDGBuilder& GraphBuilder, FViewInfo& View,
 			Out.ClassificationTileDispatchIndirectBufferUAV = GraphBuilder.CreateUAV(Out.ClassificationTileDispatchIndirectBuffer, PF_R32_UINT);
 			AddClearUAVPass(GraphBuilder, Out.ClassificationTileDispatchIndirectBufferUAV, 0);
 
+			const EPixelFormat ClassificationTileFormat = GetClassificationTileFormat(ViewResolution);
 			for (uint32 i = 0; i <= EStrataTileType::EComplex; ++i)
 			{
-				Out.ClassificationTileListBuffer[i] = GraphBuilder.CreateBuffer(FRDGBufferDesc::CreateBufferDesc(sizeof(uint32), TileResolution.X * TileResolution.Y), StrataTileListBufferNames[i]);
-				Out.ClassificationTileListBufferSRV[i] = GraphBuilder.CreateSRV(Out.ClassificationTileListBuffer[i], PF_R32_UINT);
-				Out.ClassificationTileListBufferUAV[i] = GraphBuilder.CreateUAV(Out.ClassificationTileListBuffer[i], PF_R32_UINT);
+				Out.ClassificationTileListBuffer[i] = GraphBuilder.CreateBuffer(FRDGBufferDesc::CreateBufferDesc(ClassificationTileFormat == PF_R16_UINT ? sizeof(uint16) : sizeof(uint32), TileResolution.X * TileResolution.Y), StrataTileListBufferNames[i]);
+				Out.ClassificationTileListBufferSRV[i] = GraphBuilder.CreateSRV(Out.ClassificationTileListBuffer[i], ClassificationTileFormat);
+				Out.ClassificationTileListBufferUAV[i] = GraphBuilder.CreateUAV(Out.ClassificationTileListBuffer[i], ClassificationTileFormat);
 			}
 		}
 
