Fixed non-unity compilation

diff --git a/Engine/Source/Runtime/Renderer/Public/GlobalDistanceFieldParameters.h b/Engine/Source/Runtime/Renderer/Public/GlobalDistanceFieldParameters.h
--- a/Engine/Source/Runtime/Renderer/Public/GlobalDistanceFieldParameters.h
+++ b/Engine/Source/Runtime/Renderer/Public/GlobalDistanceFieldParameters.h
@@ -7,6 +7,7 @@
 #include "ShaderParameters.h"
 #include "RenderUtils.h"
 #include "RHIStaticStates.h"
+#include "RenderGraphResources.h"
 
 class FShaderParameterMap;
 
