Fixed halfp compiler error on platforms not supporting DXC

diff --git a/Engine/Source/Runtime/Renderer/Private/PostProcess/PostProcessMobile.cpp b/Engine/Source/Runtime/Renderer/Private/PostProcess/PostProcessMobile.cpp
--- a/Engine/Source/Runtime/Renderer/Private/PostProcess/PostProcessMobile.cpp
+++ b/Engine/Source/Runtime/Renderer/Private/PostProcess/PostProcessMobile.cpp
@@ -2083,7 +2083,7 @@ public:
 
 		if (PermutationVector.Get<FFSR_UseFP16Dim>())
 		{
-			if (IsSimulatedPlatform(Parameters.Platform))
+			if (IsSimulatedPlatform(Parameters.Platform) || !IsDxcEnabledForPlatform(Parameters.Platform))
 			{
 				return false;
 			}
@@ -2125,7 +2125,7 @@ public:
 
 		if (PermutationVector.Get<FFSR_UseFP16Dim>())
 		{
-			if (IsSimulatedPlatform(Parameters.Platform))
+			if (IsSimulatedPlatform(Parameters.Platform) || !IsDxcEnabledForPlatform(Parameters.Platform))
 			{
 				return false; 
 			}
@@ -2144,7 +2144,7 @@ IMPLEMENT_GLOBAL_SHADER(FFSRPS, "/Engine/Private/Mobile/FSR/PostProcessFFX_FSR.u
 FScreenPassTexture AddEASUPass(FRDGBuilder& GraphBuilder, const FViewInfo& View, const ISpatialUpscaler::FInputs& PassInputs)
 {
 	FScreenPassTexture OutputTexture;
-	bool bUseFP16 = !IsSimulatedPlatform(View.GetShaderPlatform());
+	bool bUseFP16 = !IsSimulatedPlatform(View.GetShaderPlatform()) && IsDxcEnabledForPlatform(View.GetShaderPlatform());
 
 	FRDGTextureDesc FSROutputTextureDesc = PassInputs.SceneColor.Texture->Desc;
 	FSROutputTextureDesc.Reset();
@@ -2272,7 +2272,7 @@ public:
 
 		if (PermutationVector.Get<FFSR_UseFP16Dim>())
 		{
-			if (IsSimulatedPlatform(Parameters.Platform))
+			if (IsSimulatedPlatform(Parameters.Platform) || !IsDxcEnabledForPlatform(Parameters.Platform))
 			{
 				return false;
 			}
@@ -2317,7 +2317,7 @@ public:
 		FPermutationDomain PermutationVector(Parameters.PermutationId);
 		if (PermutationVector.Get<FFSR_UseFP16Dim>())
 		{
-			if (IsSimulatedPlatform(Parameters.Platform))
+			if (IsSimulatedPlatform(Parameters.Platform) || !IsDxcEnabledForPlatform(Parameters.Platform))
 			{
 				return false;
 			}
@@ -2336,7 +2336,7 @@ IMPLEMENT_GLOBAL_SHADER(FRCASPS, "/Engine/Private/Mobile/FSR/PostProcessFFX_RCAS
 FScreenPassTexture AddCASPass(FRDGBuilder& GraphBuilder, const FViewInfo& View, const ISpatialUpscaler::FInputs& PassInputs)
 {
 	FScreenPassTexture OutputTexture;
-	bool bUseFP16 = !IsSimulatedPlatform(View.GetShaderPlatform());
+	bool bUseFP16 = !IsSimulatedPlatform(View.GetShaderPlatform()) && IsDxcEnabledForPlatform(View.GetShaderPlatform());
 
 	if (PassInputs.OverrideOutput.IsValid())
 	{
