Make sure that GetBufferSizeForConeTracing() never returns 0.
- Prevents crashes when creating buffers/textures.

diff --git a/Engine/Source/Runtime/Renderer/Private/DistanceFieldScreenGridLighting.cpp b/Engine/Source/Runtime/Renderer/Private/DistanceFieldScreenGridLighting.cpp
--- a/Engine/Source/Runtime/Renderer/Private/DistanceFieldScreenGridLighting.cpp
+++ b/Engine/Source/Runtime/Renderer/Private/DistanceFieldScreenGridLighting.cpp
@@ -51,7 +51,8 @@ int32 GConeTraceDownsampleFactor = 4;
 
 FIntPoint GetBufferSizeForConeTracing(const FViewInfo& View)
 {
-	return FIntPoint::DivideAndRoundDown(GetBufferSizeForAO(View), GConeTraceDownsampleFactor);
+	const FIntPoint ConeTracingBufferSize = FIntPoint::DivideAndRoundDown(GetBufferSizeForAO(View), GConeTraceDownsampleFactor);
+	return FIntPoint(FMath::Max(ConeTracingBufferSize.X, 1), FMath::Max(ConeTracingBufferSize.Y, 1));
 }
 
 FVector2f JitterOffsets[4] = 
