Added support for FP32 sky light specified cubemap high dynamic range.

diff --git a/Engine/Shaders/Private/ReflectionEnvironmentShaders.usf b/Engine/Shaders/Private/ReflectionEnvironmentShaders.usf
--- a/Engine/Shaders/Private/ReflectionEnvironmentShaders.usf
+++ b/Engine/Shaders/Private/ReflectionEnvironmentShaders.usf
@@ -46,6 +46,8 @@ SamplerState SceneDepthSampler;
  */
 float4 SkyLightCaptureParameters;
 float4 LowerHemisphereColor;
+float4 ColorAlphaMultiplier;
+float ClampToFP16;
 
 void CopySceneColorToCubeFaceColorPS(
 	FCopyToCubeFaceVSOutput Input,
@@ -121,6 +123,40 @@ float3 GetCubemapVector(float2 ScaledUVs, int InCubeFace)
 	return CubeCoordinates;
 }
 
+void GetCubemapTangent(int InCubeFace, inout float3 OutX, inout float3 OutY)
+{
+	if (InCubeFace == 0)
+	{
+		OutX = float3(0,-1, 0);
+		OutY = float3(0, 0,-1);
+	}
+	else if (InCubeFace == 1)
+	{
+		OutX = float3(0,-1, 0);
+		OutY = float3(0, 0, 1);
+	}
+	else if (InCubeFace == 2)
+	{
+		OutX = float3(1, 0, 0);
+		OutY = float3(0, 0, 1);
+	}
+	else if (InCubeFace == 3)
+	{
+		OutX = float3(1, 0, 0);
+		OutY = float3(0, 0,-1);
+	}
+	else if (InCubeFace == 4)
+	{
+		OutX = float3(1, 0, 0);
+		OutY = float3(0,-1, 0);
+	}
+	else
+	{
+		OutX = float3(-1, 0, 0);
+		OutY = float3( 0,-1, 0);
+	}
+}
+
 float CubeTexelWeight( float3 N )
 {
 	uint Axis = 2;
@@ -171,6 +207,10 @@ void CopyCubemapToCubeFaceColorPS(
 	}
 
 	OutColor.a = 1;
+
+	OutColor *= ColorAlphaMultiplier;
+
+	OutColor = ClampToFP16 == 0.0f ? OutColor : clamp(OutColor, float4(0.0f, 0.0f, 0.0f, 0.0f), float4(MaxHalfFloat, MaxHalfFloat, MaxHalfFloat, 1.0f));
 }
 
 uint MipIndex;
@@ -247,6 +287,46 @@ void DownsamplePS(
 #endif
 }
 
+void DownsampleMaxPS(
+	float4 SvPosition : SV_POSITION,
+	out float4 OutColor : SV_Target0
+)
+{
+	float2 UV = SvPosition.xy * SvPositionToUVScale;
+	float2 ScaledUVs = UV * 2 - 1;
+	const int SelectedCubeFace = CubeFace;
+
+	float3 CubeCoordinates = GetCubemapVector(ScaledUVs, SelectedCubeFace);
+
+	uint MipSize = 1u << (NumMips - MipIndex - 1);
+
+	float3 TangentX = 0.0f;
+	float3 TangentY = 0.0f;
+	GetCubemapTangent(SelectedCubeFace, TangentX, TangentY);
+
+	const float SampleOffset = 1.0f / MipSize;
+
+	// Offset from the connect of 4 texels towards their center.
+	float2 Offsets[] =
+	{
+		float2(-0.5,-0.5),
+		float2( 0.5,-0.5),
+		float2( 0.5, 0.5),
+		float2(-0.5, 0.5),
+	};
+
+	OutColor = 0;
+
+	UNROLL
+	for (uint i = 0; i < 4; i++)
+	{
+		float3 SampleDir = CubeCoordinates;
+		SampleDir += TangentX * (Offsets[i].x * SampleOffset);
+		SampleDir += TangentY * (Offsets[i].y * SampleOffset);
+		OutColor = max(OutColor, SourceCubemapTexture.SampleLevel(SourceCubemapSampler, SampleDir, 0));
+	}
+}
+
 #ifdef USE_COMPUTE
 float4 LowerHemisphereSolidColor;
 [numthreads(THREADGROUP_SIZE, THREADGROUP_SIZE, 1)]
@@ -293,6 +373,24 @@ void ComputeBrightnessMain(out float4 OutColor : SV_Target0)
 
 #endif
 
+#if COMPUTEMAXLUMINANCE_PIXELSHADER
+
+void ComputeCubeMaxLuminancePS(out float4 OutColor : SV_Target0)
+{ 
+	// Sample the 6 1x1 cube faces and average
+	float3 MaxLuminanceRGB = TextureCubeSampleLevel(ReflectionEnvironmentColorTexture, ReflectionEnvironmentColorSampler, float3(1, 0, 0), NumCaptureArrayMips - 1).rgb;
+	MaxLuminanceRGB = max(MaxLuminanceRGB, TextureCubeSampleLevel(ReflectionEnvironmentColorTexture, ReflectionEnvironmentColorSampler, float3(-1, 0, 0), NumCaptureArrayMips - 1).rgb);
+	MaxLuminanceRGB = max(MaxLuminanceRGB, TextureCubeSampleLevel(ReflectionEnvironmentColorTexture, ReflectionEnvironmentColorSampler, float3( 0, 1, 0), NumCaptureArrayMips - 1).rgb);
+	MaxLuminanceRGB = max(MaxLuminanceRGB, TextureCubeSampleLevel(ReflectionEnvironmentColorTexture, ReflectionEnvironmentColorSampler, float3( 0,-1, 0), NumCaptureArrayMips - 1).rgb);
+	MaxLuminanceRGB = max(MaxLuminanceRGB, TextureCubeSampleLevel(ReflectionEnvironmentColorTexture, ReflectionEnvironmentColorSampler, float3( 0, 0, 1), NumCaptureArrayMips - 1).rgb);
+	MaxLuminanceRGB = max(MaxLuminanceRGB, TextureCubeSampleLevel(ReflectionEnvironmentColorTexture, ReflectionEnvironmentColorSampler, float3( 0, 0,-1), NumCaptureArrayMips - 1).rgb);
+
+	float MaxLuminance = max(MaxLuminanceRGB.r, max(MaxLuminanceRGB.g, MaxLuminanceRGB.b));
+	OutColor = MaxLuminance;
+}
+
+#endif
+
 float4 SampleCubemap(float3 Coordinates, uint InMipIndex)
 {
 	return TextureCubeSampleLevel(SourceCubemapTexture, SourceCubemapSampler, Coordinates, InMipIndex);
diff --git a/Engine/Source/Runtime/Engine/Classes/Components/SkyLightComponent.h b/Engine/Source/Runtime/Engine/Classes/Components/SkyLightComponent.h
--- a/Engine/Source/Runtime/Engine/Classes/Components/SkyLightComponent.h
+++ b/Engine/Source/Runtime/Engine/Classes/Components/SkyLightComponent.h
@@ -333,6 +333,8 @@ protected:
 	FSHVectorRGB3 BlendDestinationIrradianceEnvironmentMap;
 	float BlendDestinationAverageBrightness;
 
+	FLinearColor SpecifiedCubemapColorScale;
+
 	/** Tracks when the rendering thread has completed its writes to IrradianceEnvironmentMap. */
 	FRenderCommandFence IrradianceMapFence;
 
diff --git a/Engine/Source/Runtime/Engine/Private/Components/SkyLightComponent.cpp b/Engine/Source/Runtime/Engine/Private/Components/SkyLightComponent.cpp
--- a/Engine/Source/Runtime/Engine/Private/Components/SkyLightComponent.cpp
+++ b/Engine/Source/Runtime/Engine/Private/Components/SkyLightComponent.cpp
@@ -176,10 +176,11 @@ void FSkyLightSceneProxy::Initialize(
 	const FSHVectorRGB3* InIrradianceEnvironmentMap, 
 	const FSHVectorRGB3* BlendDestinationIrradianceEnvironmentMap,
 	const float* InAverageBrightness,
-	const float* BlendDestinationAverageBrightness)
+	const float* BlendDestinationAverageBrightness,
+	const FLinearColor* InSpecifiedCubemapColorScale)
 {
+	SpecifiedCubemapColorScale = *InSpecifiedCubemapColorScale;
 	BlendFraction = FMath::Clamp(InBlendFraction, 0.0f, 1.0f);
-
 	if (BlendFraction > 0 && BlendDestinationProcessedTexture != NULL)
 	{
 		if (BlendFraction < 1)
@@ -205,7 +206,7 @@ void FSkyLightSceneProxy::Initialize(
 
 FLinearColor FSkyLightSceneProxy::GetEffectiveLightColor() const
 {
-	return LightColor * GSkylightIntensityMultiplier;
+	return LightColor * GSkylightIntensityMultiplier * SpecifiedCubemapColorScale;
 }
 
 FSkyLightSceneProxy::FSkyLightSceneProxy(const USkyLightComponent* InLightComponent)
@@ -254,13 +255,15 @@ FSkyLightSceneProxy::FSkyLightSceneProxy(const USkyLightComponent* InLightCompon
 	const float* InAverageBrightness = &InLightComponent->AverageBrightness;
 	const float* BlendDestinationAverageBrightness = &InLightComponent->BlendDestinationAverageBrightness;
 	float InBlendFraction = InLightComponent->BlendFraction;
+	const FLinearColor* InSpecifiedCubemapColorScale = &InLightComponent->SpecifiedCubemapColorScale;
 	FSkyLightSceneProxy* LightSceneProxy = this;
 	ENQUEUE_RENDER_COMMAND(FInitSkyProxy)(
-		[InIrradianceEnvironmentMap, BlendDestinationIrradianceEnvironmentMap, InAverageBrightness, BlendDestinationAverageBrightness, InBlendFraction, LightSceneProxy](FRHICommandList& RHICmdList)
+		[InIrradianceEnvironmentMap, BlendDestinationIrradianceEnvironmentMap, InAverageBrightness, 
+		BlendDestinationAverageBrightness, InBlendFraction, LightSceneProxy, InSpecifiedCubemapColorScale](FRHICommandList& RHICmdList)
 		{
 			// Only access the irradiance maps on the RT, even though they belong to the USkyLightComponent, 
 			// Because FScene::UpdateSkyCaptureContents does not block the RT so the writes could still be in flight
-			LightSceneProxy->Initialize(InBlendFraction, InIrradianceEnvironmentMap, BlendDestinationIrradianceEnvironmentMap, InAverageBrightness, BlendDestinationAverageBrightness);
+			LightSceneProxy->Initialize(InBlendFraction, InIrradianceEnvironmentMap, BlendDestinationIrradianceEnvironmentMap, InAverageBrightness, BlendDestinationAverageBrightness, InSpecifiedCubemapColorScale);
 		});
 }
 
@@ -294,6 +297,7 @@ USkyLightComponent::USkyLightComponent(const FObjectInitializer& ObjectInitializ
 	LowerHemisphereColor = FLinearColor::Black;
 	AverageBrightness = 1.0f;
 	BlendDestinationAverageBrightness = 1.0f;
+	SpecifiedCubemapColorScale = FLinearColor::White;
 	bCastVolumetricShadow = true;
 	CastRaytracedShadow = ECastRayTracedShadow::UseProjectSetting;
 	bCastRaytracedShadow_DEPRECATED = false;
@@ -738,6 +742,9 @@ void USkyLightComponent::UpdateSkyCaptureContentsArray(UWorld* WorldToUpdate, TA
 		USkyLightComponent* CaptureComponent = ComponentArray[CaptureIndex];
 		AActor* Owner = CaptureComponent->GetOwner();
 
+		// Reset the luminance scale in case the texture has been switched
+		CaptureComponent->SpecifiedCubemapColorScale = FLinearColor::White;
+
 		// For specific cubemaps, we must wait until the texture is compiled before capturing the skylight
 		bool bIsCubemapCompiling = false;
 #if WITH_EDITOR
@@ -792,7 +799,7 @@ void USkyLightComponent::UpdateSkyCaptureContentsArray(UWorld* WorldToUpdate, TA
 						CaptureComponent->MarkRenderStateDirty();
 					}
 
-					WorldToUpdate->Scene->UpdateSkyCaptureContents(CaptureComponent, CaptureComponent->bCaptureEmissiveOnly, CaptureComponent->Cubemap, CaptureComponent->ProcessedSkyTexture, CaptureComponent->AverageBrightness, CaptureComponent->IrradianceEnvironmentMap, NULL);
+					WorldToUpdate->Scene->UpdateSkyCaptureContents(CaptureComponent, CaptureComponent->bCaptureEmissiveOnly, CaptureComponent->Cubemap, CaptureComponent->ProcessedSkyTexture, CaptureComponent->AverageBrightness, CaptureComponent->IrradianceEnvironmentMap, NULL, &CaptureComponent->SpecifiedCubemapColorScale);
 				}
 				else
 				{
@@ -807,7 +814,7 @@ void USkyLightComponent::UpdateSkyCaptureContentsArray(UWorld* WorldToUpdate, TA
 						CaptureComponent->MarkRenderStateDirty(); 
 					}
 
-					WorldToUpdate->Scene->UpdateSkyCaptureContents(CaptureComponent, CaptureComponent->bCaptureEmissiveOnly, CaptureComponent->BlendDestinationCubemap, CaptureComponent->BlendDestinationProcessedSkyTexture, CaptureComponent->BlendDestinationAverageBrightness, CaptureComponent->BlendDestinationIrradianceEnvironmentMap, NULL);
+					WorldToUpdate->Scene->UpdateSkyCaptureContents(CaptureComponent, CaptureComponent->bCaptureEmissiveOnly, CaptureComponent->BlendDestinationCubemap, CaptureComponent->BlendDestinationProcessedSkyTexture, CaptureComponent->BlendDestinationAverageBrightness, CaptureComponent->BlendDestinationIrradianceEnvironmentMap, NULL, &CaptureComponent->SpecifiedCubemapColorScale);
 				}
 
 				CaptureComponent->IrradianceMapFence.BeginFence();
@@ -930,9 +937,10 @@ void USkyLightComponent::CaptureEmissiveRadianceEnvironmentCubeMap(FSHVectorRGB3
 	if (GetScene() && (SourceType != SLS_SpecifiedCubemap || Cubemap))
 	{
 		float UnusedAverageBrightness = 1.0f;
+		FLinearColor* UnusedSpecifiedCubemapColorScale = nullptr;	// Disable
 		// Capture emissive scene lighting only for the lighting build
 		// This is necessary to avoid a feedback loop with the last lighting build results
-		GetScene()->UpdateSkyCaptureContents(this, true, Cubemap, NULL, UnusedAverageBrightness, OutIrradianceMap, &OutRadianceMap);
+		GetScene()->UpdateSkyCaptureContents(this, true, Cubemap, NULL, UnusedAverageBrightness, OutIrradianceMap, &OutRadianceMap, UnusedSpecifiedCubemapColorScale);
 		// Wait until writes to OutIrradianceMap have completed
 		FlushRenderingCommands();
 	}
@@ -1041,12 +1049,13 @@ void USkyLightComponent::SetCubemapBlend(UTextureCube* SourceCubemap, UTextureCu
 				const float* InAverageBrightness = &AverageBrightness;
 				const float* InBlendDestinationAverageBrightness = &BlendDestinationAverageBrightness;
 				FSkyLightSceneProxy* LightSceneProxy = SceneProxy;
+				const FLinearColor* InSpecifiedCubemapColorScale = &SpecifiedCubemapColorScale;
 				ENQUEUE_RENDER_COMMAND(FUpdateSkyProxy)(
-					[InIrradianceEnvironmentMap, InBlendDestinationIrradianceEnvironmentMap, InAverageBrightness, InBlendDestinationAverageBrightness, InBlendFraction, LightSceneProxy](FRHICommandList& RHICmdList)
+					[InIrradianceEnvironmentMap, InBlendDestinationIrradianceEnvironmentMap, InAverageBrightness, InBlendDestinationAverageBrightness, InBlendFraction, LightSceneProxy, InSpecifiedCubemapColorScale](FRHICommandList& RHICmdList)
 					{
 						// Only access the irradiance maps on the RT, even though they belong to the USkyLightComponent, 
 						// Because FScene::UpdateSkyCaptureContents does not block the RT so the writes could still be in flight
-						LightSceneProxy->Initialize(InBlendFraction, InIrradianceEnvironmentMap, InBlendDestinationIrradianceEnvironmentMap, InAverageBrightness, InBlendDestinationAverageBrightness);
+						LightSceneProxy->Initialize(InBlendFraction, InIrradianceEnvironmentMap, InBlendDestinationIrradianceEnvironmentMap, InAverageBrightness, InBlendDestinationAverageBrightness, InSpecifiedCubemapColorScale);
 					});
 			}
 		}
diff --git a/Engine/Source/Runtime/Engine/Public/SceneInterface.h b/Engine/Source/Runtime/Engine/Public/SceneInterface.h
--- a/Engine/Source/Runtime/Engine/Public/SceneInterface.h
+++ b/Engine/Source/Runtime/Engine/Public/SceneInterface.h
@@ -193,7 +193,7 @@ public:
 	 * Updates the contents of the given sky capture by rendering the scene. 
 	 * This must be called on the game thread.
 	 */
-	virtual void UpdateSkyCaptureContents(const USkyLightComponent* CaptureComponent, bool bCaptureEmissiveOnly, UTextureCube* SourceCubemap, FTexture* OutProcessedTexture, float& OutAverageBrightness, FSHVectorRGB3& OutIrradianceEnvironmentMap, TArray<FFloat16Color>* OutRadianceMap) {}
+	virtual void UpdateSkyCaptureContents(const USkyLightComponent* CaptureComponent, bool bCaptureEmissiveOnly, UTextureCube* SourceCubemap, FTexture* OutProcessedTexture, float& OutAverageBrightness, FSHVectorRGB3& OutIrradianceEnvironmentMap, TArray<FFloat16Color>* OutRadianceMap, FLinearColor* SpecifiedCubemapColorScale) {}
 
 	virtual void AllocateAndCaptureFrameSkyEnvMap(FRDGBuilder& GraphBuilder, FSceneRenderer& SceneRenderer, FViewInfo& MainView, bool bShouldRenderSkyAtmosphere, bool bShouldRenderVolumetricCloud, FInstanceCullingManager& InstanceCullingManager) {}
 	virtual void ValidateSkyLightRealTimeCapture(FRDGBuilder& GraphBuilder, const FViewInfo& View, FRDGTextureRef SceneColorTexture) {}
diff --git a/Engine/Source/Runtime/Engine/Public/SceneManagement.h b/Engine/Source/Runtime/Engine/Public/SceneManagement.h
--- a/Engine/Source/Runtime/Engine/Public/SceneManagement.h
+++ b/Engine/Source/Runtime/Engine/Public/SceneManagement.h
@@ -1171,7 +1171,8 @@ public:
 		const FSHVectorRGB3* InIrradianceEnvironmentMap, 
 		const FSHVectorRGB3* BlendDestinationIrradianceEnvironmentMap,
 		const float* InAverageBrightness,
-		const float* BlendDestinationAverageBrightness);
+		const float* BlendDestinationAverageBrightness,
+		const FLinearColor* InSpecifiedCubemapColorScale);
 
 	const USkyLightComponent* LightComponent;
 	FTexture* ProcessedTexture;
@@ -1207,6 +1208,7 @@ public:
 	uint32 CaptureCubeMapResolution;
 	FLinearColor LowerHemisphereColor;
 	bool bLowerHemisphereIsSolidColor;
+	FLinearColor SpecifiedCubemapColorScale;
 
 	bool IsMovable() { return bMovable; }
 
diff --git a/Engine/Source/Runtime/Renderer/Private/ReflectionEnvironmentCapture.cpp b/Engine/Source/Runtime/Renderer/Private/ReflectionEnvironmentCapture.cpp
--- a/Engine/Source/Runtime/Renderer/Private/ReflectionEnvironmentCapture.cpp
+++ b/Engine/Source/Runtime/Renderer/Private/ReflectionEnvironmentCapture.cpp
@@ -122,6 +122,15 @@ class FCubeDownsamplePS : public FGlobalShader
 
 IMPLEMENT_GLOBAL_SHADER(FCubeDownsamplePS, "/Engine/Private/ReflectionEnvironmentShaders.usf", "DownsamplePS", SF_Pixel);
 
+class FCubeDownsampleMaxPS : public FGlobalShader
+{
+	DECLARE_GLOBAL_SHADER(FCubeDownsampleMaxPS);
+	SHADER_USE_PARAMETER_STRUCT(FCubeDownsampleMaxPS, FGlobalShader);
+	using FParameters = FCubeShaderParameters;
+};
+
+IMPLEMENT_GLOBAL_SHADER(FCubeDownsampleMaxPS, "/Engine/Private/ReflectionEnvironmentShaders.usf", "DownsampleMaxPS", SF_Pixel);
+
 class FCubeFilterPS : public FGlobalShader
 {
 	DECLARE_GLOBAL_SHADER(FCubeFilterPS);
@@ -131,18 +140,19 @@ class FCubeFilterPS : public FGlobalShader
 
 IMPLEMENT_GLOBAL_SHADER(FCubeFilterPS, "/Engine/Private/ReflectionEnvironmentShaders.usf", "FilterPS", SF_Pixel);
 
+BEGIN_SHADER_PARAMETER_STRUCT(FCubeFinalGatherPS, )
+	SHADER_PARAMETER_RDG_TEXTURE(TextureCube, ReflectionEnvironmentColorTexture)
+	SHADER_PARAMETER_SAMPLER(SamplerState, ReflectionEnvironmentColorSampler)
+	SHADER_PARAMETER(int32, NumCaptureArrayMips)
+	RENDER_TARGET_BINDING_SLOTS()
+END_SHADER_PARAMETER_STRUCT()
+
 /** Computes the average brightness of a 1x1 mip of a cubemap. */
 class FComputeBrightnessPS : public FGlobalShader
 {
 	DECLARE_GLOBAL_SHADER(FComputeBrightnessPS);
 	SHADER_USE_PARAMETER_STRUCT(FComputeBrightnessPS, FGlobalShader);
-
-	BEGIN_SHADER_PARAMETER_STRUCT(FParameters, )
-		SHADER_PARAMETER_RDG_TEXTURE(TextureCube, ReflectionEnvironmentColorTexture)
-		SHADER_PARAMETER_SAMPLER(SamplerState, ReflectionEnvironmentColorSampler)
-		SHADER_PARAMETER(int32, NumCaptureArrayMips)
-		RENDER_TARGET_BINDING_SLOTS()
-	END_SHADER_PARAMETER_STRUCT()
+	using FParameters = FCubeFinalGatherPS;
 
 	static void ModifyCompilationEnvironment(const FGlobalShaderPermutationParameters& Parameters, FShaderCompilerEnvironment& OutEnvironment)
 	{
@@ -153,6 +163,21 @@ class FComputeBrightnessPS : public FGlobalShader
 
 IMPLEMENT_GLOBAL_SHADER(FComputeBrightnessPS, "/Engine/Private/ReflectionEnvironmentShaders.usf", "ComputeBrightnessMain", SF_Pixel);
 
+class FComputeCubeMaxLuminancePS : public FGlobalShader
+{
+	DECLARE_GLOBAL_SHADER(FComputeCubeMaxLuminancePS);
+	SHADER_USE_PARAMETER_STRUCT(FComputeCubeMaxLuminancePS, FGlobalShader);
+	using FParameters = FCubeFinalGatherPS;
+
+	static void ModifyCompilationEnvironment(const FGlobalShaderPermutationParameters& Parameters, FShaderCompilerEnvironment& OutEnvironment)
+	{
+		FGlobalShader::ModifyCompilationEnvironment(Parameters, OutEnvironment);
+		OutEnvironment.SetDefine(TEXT("COMPUTEMAXLUMINANCE_PIXELSHADER"), 1);
+	}
+};
+
+IMPLEMENT_GLOBAL_SHADER(FComputeCubeMaxLuminancePS, "/Engine/Private/ReflectionEnvironmentShaders.usf", "ComputeCubeMaxLuminancePS", SF_Pixel);
+
 /** Vertex shader used when writing to a cubemap. */
 class FCopyToCubeFaceVS : public FGlobalShader
 {
@@ -209,9 +234,11 @@ public:
 		SHADER_PARAMETER_SAMPLER(SamplerState, SourceCubemapSampler)
 		SHADER_PARAMETER(FVector4f, SkyLightCaptureParameters)
 		SHADER_PARAMETER(FVector4f, LowerHemisphereColor)
+		SHADER_PARAMETER(FVector4f, ColorAlphaMultiplier)
 		SHADER_PARAMETER(FVector2f, SinCosSourceCubemapRotation)
 		SHADER_PARAMETER(FVector2f, SvPositionToUVScale)
 		SHADER_PARAMETER(int32, CubeFace)
+		SHADER_PARAMETER(float, ClampToFP16)
 		RENDER_TARGET_BINDING_SLOTS()
 	END_SHADER_PARAMETER_STRUCT()
 };
@@ -597,6 +624,50 @@ void CaptureSceneToScratchCubemap(
 	SceneRenderer->RenderThreadEnd(RHICmdList);
 }
 
+
+static void CopyCubemapToScratchCubemapInner(
+	FRDGBuilder& GraphBuilder,
+	FGlobalShaderMap* ShaderMap,
+	const FTexture* SourceCubemapResource,
+	FRDGTextureRef OutputTexture,
+	int32 CubemapSize,
+	bool bIsSkyLight,
+	bool bLowerHemisphereIsBlack,
+	float SourceCubemapRotation,
+	const FLinearColor& LowerHemisphereColorValue,
+	const FLinearColor& ColorAlphaMultiplier,
+	const bool bClampToFP16)
+{
+	RDG_EVENT_SCOPE(GraphBuilder, "CopyCubemapToScratchCubemap");
+
+	TShaderMapRef<FCopyCubemapToCubeFacePS> PixelShader(ShaderMap);
+
+	for (uint32 CubeFace = 0; CubeFace < CubeFace_MAX; CubeFace++)
+	{
+		auto* PassParameters = GraphBuilder.AllocParameters<FCopyCubemapToCubeFacePS::FParameters>();
+		PassParameters->SourceCubemapSampler = TStaticSamplerState<SF_Bilinear, AM_Clamp, AM_Clamp, AM_Clamp>::GetRHI();
+		PassParameters->SourceCubemapTexture = SourceCubemapResource->TextureRHI;
+		PassParameters->LowerHemisphereColor = LowerHemisphereColorValue;
+		PassParameters->SkyLightCaptureParameters = FVector3f(bIsSkyLight ? 1.0f : 0.0f, 0.0f, bLowerHemisphereIsBlack ? 1.0f : 0.0f);
+		PassParameters->SinCosSourceCubemapRotation = FVector2f(FMath::Sin(SourceCubemapRotation), FMath::Cos(SourceCubemapRotation));
+		PassParameters->SvPositionToUVScale = FVector2f(1.0f / CubemapSize, 1.0f / CubemapSize);
+		PassParameters->CubeFace = CubeFace;
+		PassParameters->ColorAlphaMultiplier = ColorAlphaMultiplier;
+		PassParameters->ClampToFP16 = bClampToFP16 ? 1.0f : 0.0f;
+		PassParameters->RenderTargets[0] = FRenderTargetBinding(OutputTexture, ERenderTargetLoadAction::ENoAction, 0, CubeFace);
+
+		const FIntRect ViewRect(0, 0, CubemapSize, CubemapSize);
+
+		FPixelShaderUtils::AddFullscreenPass(
+			GraphBuilder,
+			ShaderMap,
+			RDG_EVENT_NAME("CopyCubemapToCubeFace"),
+			PixelShader,
+			PassParameters,
+			ViewRect);
+	}
+}
+
 void CopyCubemapToScratchCubemap(
 	FRHICommandListImmediate& RHICmdList,
 	ERHIFeatureLevel::Type FeatureLevel,
@@ -606,7 +677,8 @@ void CopyCubemapToScratchCubemap(
 	bool bIsSkyLight,
 	bool bLowerHemisphereIsBlack,
 	float SourceCubemapRotation,
-	const FLinearColor& LowerHemisphereColorValue)
+	const FLinearColor& LowerHemisphereColorValue,
+	const FLinearColor& ColorAlphaMultiplier)
 {
 	check(SourceCubemap);
 
@@ -621,37 +693,20 @@ void CopyCubemapToScratchCubemap(
 	FRDGBuilder GraphBuilder(RHICmdList);
 
 	auto ShaderMap = GetGlobalShaderMap(FeatureLevel);
-
-	{
-		RDG_EVENT_SCOPE(GraphBuilder, "CopyCubemapToScratchCubemap");
-
-		FRDGTextureRef OutputTexture = ReflectionCubemapTexture.GetRDG(GraphBuilder);
-
-		TShaderMapRef<FCopyCubemapToCubeFacePS> PixelShader(ShaderMap);
-
-		for (uint32 CubeFace = 0; CubeFace < CubeFace_MAX; CubeFace++)
-		{
-			auto* PassParameters = GraphBuilder.AllocParameters<FCopyCubemapToCubeFacePS::FParameters>();
-			PassParameters->SourceCubemapSampler = TStaticSamplerState<SF_Bilinear, AM_Clamp, AM_Clamp, AM_Clamp>::GetRHI();
-			PassParameters->SourceCubemapTexture = SourceCubemapResource->TextureRHI;
-			PassParameters->LowerHemisphereColor = LowerHemisphereColorValue;
-			PassParameters->SkyLightCaptureParameters = FVector3f(bIsSkyLight ? 1.0f : 0.0f, 0.0f, bLowerHemisphereIsBlack ? 1.0f : 0.0f);
-			PassParameters->SinCosSourceCubemapRotation = FVector2f(FMath::Sin(SourceCubemapRotation), FMath::Cos(SourceCubemapRotation));
-			PassParameters->SvPositionToUVScale = FVector2f(1.0f / CubemapSize, 1.0f / CubemapSize);
-			PassParameters->CubeFace = CubeFace;
-			PassParameters->RenderTargets[0] = FRenderTargetBinding(OutputTexture, ERenderTargetLoadAction::ENoAction, 0, CubeFace);
-
-			const FIntRect ViewRect(0, 0, CubemapSize, CubemapSize);
-
-			FPixelShaderUtils::AddFullscreenPass(
-				GraphBuilder,
-				ShaderMap,
-				RDG_EVENT_NAME("CopyCubemapToCubeFace"),
-				PixelShader,
-				PassParameters,
-				ViewRect);
-		}
-	}
+	FRDGTextureRef OutputTexture = ReflectionCubemapTexture.GetRDG(GraphBuilder);
+	const bool bClampToFP16 = true; // Rendering into an FP16 texture.
+	CopyCubemapToScratchCubemapInner(
+		GraphBuilder,
+		ShaderMap,
+		SourceCubemapResource,
+		OutputTexture,
+		CubemapSize,
+		bIsSkyLight,
+		bLowerHemisphereIsBlack,
+		SourceCubemapRotation,
+		LowerHemisphereColorValue,
+		ColorAlphaMultiplier,
+		bClampToFP16);
 
 	GraphBuilder.Execute();
 }
@@ -1389,7 +1444,7 @@ void FScene::CaptureOrUploadReflectionCapture(UReflectionCaptureComponent* Captu
 				ENQUEUE_RENDER_COMMAND(CopyCubemapCommand)(
 					[FeatureLevel = FeatureLevel, SourceCubemap, ReflectionCubemapTexture = ReflectionCubemapTexture.Get(), ReflectionCaptureSize, SourceCubemapRotation](FRHICommandListImmediate& RHICmdList)
 				{
-					CopyCubemapToScratchCubemap(RHICmdList, FeatureLevel, SourceCubemap, *ReflectionCubemapTexture, ReflectionCaptureSize, false, false, SourceCubemapRotation, FLinearColor());
+					CopyCubemapToScratchCubemap(RHICmdList, FeatureLevel, SourceCubemap, *ReflectionCubemapTexture, ReflectionCaptureSize, false, false, SourceCubemapRotation, FLinearColor(), FLinearColor::White);
 				});
 			}
 			else
@@ -1479,6 +1534,123 @@ void CopyToSkyTexture(FRDGBuilder& GraphBuilder, FScene* Scene, FRDGTexture* Inp
 	}
 }
 
+void ComputeSpecifiedCubemapColorScale(
+	FRDGBuilder& GraphBuilder, 
+	FGlobalShaderMap* ShaderMap, 
+	UTextureCube* SourceCubemap,
+	int32 CubemapSize,
+	bool bIsSkyLight,
+	bool bLowerHemisphereIsBlack,
+	float SourceCubemapRotation,
+	const FLinearColor& LowerHemisphereColorValue,
+	float& MaxFloatLuminance)
+{
+	const FTexture* SourceCubemapResource = SourceCubemap->GetResource();
+	if (SourceCubemapResource == nullptr)
+	{
+		UE_LOG(LogEngine, Warning, TEXT("Unable to copy from cubemap %s, it's RHI resource is null"), *SourceCubemap->GetPathName());
+		return;
+	}
+;
+	const int32 NumReflectionCaptureMips = GetNumMips(CubemapSize);
+	const FRDGTextureDesc TextureDesc = FRDGTextureDesc::CreateCube(
+		CubemapSize, PF_A32B32G32R32F, FClearValueBinding(FLinearColor(0, 10000, 0, 0)), TexCreate_RenderTargetable | TexCreate_ShaderResource | TexCreate_TargetArraySlicesIndependently | TexCreate_DisableDCC, NumReflectionCaptureMips);
+	FRDGTextureRef ReflectionMaxLuminanceCubeTexture = GraphBuilder.CreateTexture(TextureDesc, TEXT("ReflectionMaxLuminanceCubeTexture"));
+
+	// Copy the source texture into the first mip level of each faces of the cubemap.
+	{
+		RDG_EVENT_SCOPE(GraphBuilder, "CopyCubeTextureToFP32CubemapMip0");
+		const bool bClampToFP16 = false; // Rendering into an FP32 texture.
+		CopyCubemapToScratchCubemapInner(
+			GraphBuilder,
+			ShaderMap,
+			SourceCubemapResource,
+			ReflectionMaxLuminanceCubeTexture,
+			CubemapSize,
+			bIsSkyLight,
+			bLowerHemisphereIsBlack,
+			SourceCubemapRotation,
+			LowerHemisphereColorValue,
+			FLinearColor::White,
+			bClampToFP16);
+	}
+
+	// Generate all the mips of the texture, each time storing the maximum luminance value.
+	{
+		RDG_EVENT_SCOPE(GraphBuilder, "CreateFP32CubeMipsAsMaxLuminance");
+		TShaderMapRef<FCubeDownsampleMaxPS> PixelShader(ShaderMap);
+		const int32 NumMips = ReflectionMaxLuminanceCubeTexture->Desc.NumMips;
+
+		// Downsample all the mips, each one reads from the mip above it
+		for (int32 MipIndex = 1; MipIndex < NumMips; MipIndex++)
+		{
+			const int32 MipSize = 1 << (NumMips - MipIndex - 1);
+			const FIntRect ViewRect(0, 0, MipSize, MipSize);
+
+			for (int32 CubeFace = 0; CubeFace < CubeFace_MAX; CubeFace++)
+			{
+				auto* PassParameters = GraphBuilder.AllocParameters<FCubeDownsampleMaxPS::FParameters>();
+				PassParameters->SourceCubemapTexture = GraphBuilder.CreateSRV(FRDGTextureSRVDesc::CreateForMipLevel(ReflectionMaxLuminanceCubeTexture, MipIndex - 1));
+				PassParameters->SourceCubemapSampler = TStaticSamplerState<SF_Point, AM_Clamp, AM_Clamp, AM_Clamp>::GetRHI();
+				PassParameters->CubeFace = CubeFace;
+				PassParameters->MipIndex = MipIndex;
+				PassParameters->NumMips = NumMips;
+				PassParameters->SvPositionToUVScale = FVector2f(1.0f / MipSize, 1.0f / MipSize);
+				PassParameters->RenderTargets[0] = FRenderTargetBinding(ReflectionMaxLuminanceCubeTexture, ERenderTargetLoadAction::ENoAction, MipIndex, CubeFace);
+
+				FPixelShaderUtils::AddFullscreenPass(
+					GraphBuilder,
+					ShaderMap,
+					RDG_EVENT_NAME("CreateCubeMips (Mip: %d, Face: %d)", MipIndex, CubeFace),
+					PixelShader,
+					PassParameters,
+					ViewRect);
+			}
+		}
+	}
+
+	// Fetch the maximum luminance (acounting for each face) from GPU to CPU.
+	{
+		RDG_EVENT_SCOPE(GraphBuilder, "FetchFP32CubeMaxLuminance");
+
+		FRDGTexture* MaxLuminanceTexture = GraphBuilder.CreateTexture(FRDGTextureDesc::Create2D(FIntPoint(1, 1), PF_A32B32G32R32F, FClearValueBinding::None, TexCreate_RenderTargetable), TEXT("ReflectionBrightness"));
+
+		auto* PassParameters = GraphBuilder.AllocParameters<FComputeCubeMaxLuminancePS::FParameters>();
+
+		PassParameters->ReflectionEnvironmentColorTexture = ReflectionMaxLuminanceCubeTexture;
+		PassParameters->ReflectionEnvironmentColorSampler = TStaticSamplerState<SF_Point, AM_Clamp, AM_Clamp, AM_Clamp>::GetRHI();
+		PassParameters->NumCaptureArrayMips = ReflectionMaxLuminanceCubeTexture->Desc.NumMips;
+		PassParameters->RenderTargets[0] = FRenderTargetBinding(MaxLuminanceTexture, ERenderTargetLoadAction::ENoAction);
+
+		TShaderMapRef<FComputeCubeMaxLuminancePS> PixelShader(ShaderMap);
+
+		FPixelShaderUtils::AddFullscreenPass(
+			GraphBuilder,
+			ShaderMap,
+			RDG_EVENT_NAME("ReflectionBrightness"),
+			PixelShader,
+			PassParameters,
+			FIntRect(FIntPoint::ZeroValue, FIntPoint(1, 1)));
+
+		float* MaxFloatLuminancePtr = &MaxFloatLuminance;
+		AddReadbackTexturePass(GraphBuilder, RDG_EVENT_NAME("ReadbackTexture"), MaxLuminanceTexture, [MaxLuminanceTexture, MaxFloatLuminancePtr](FRHICommandListImmediate& RHICmdList)
+		{
+			FReadSurfaceDataFlags ReadDataFlags(RCM_MinMax);
+			ReadDataFlags.SetLinearToGamma(false);
+
+			FIntRect SourceRect = FIntRect(0, 0, 1, 1);
+
+			TArray<FLinearColor> RawPixels;
+			RawPixels.SetNum(SourceRect.Width() * SourceRect.Height());
+
+			RHICmdList.ReadSurfaceData(MaxLuminanceTexture->GetRHI(), SourceRect, RawPixels, ReadDataFlags);
+
+			// Shader outputs luminance to R
+			*MaxFloatLuminancePtr = RawPixels[0].R;
+		});
+	}
+}
+
 // Warning: returns before writes to OutIrradianceEnvironmentMap have completed, as they are queued on the rendering thread
 void FScene::UpdateSkyCaptureContents(
 	const USkyLightComponent* CaptureComponent, 
@@ -1487,7 +1659,8 @@ void FScene::UpdateSkyCaptureContents(
 	FTexture* OutProcessedTexture, 
 	float& OutAverageBrightness, 
 	FSHVectorRGB3& OutIrradianceEnvironmentMap,
-	TArray<FFloat16Color>* OutRadianceMap)
+	TArray<FFloat16Color>* OutRadianceMap,
+	FLinearColor* SpecifiedCubemapColorScale)
 {
 	if (GSupportsRenderTargetFormat_PF_FloatRGBA || FeatureLevel >= ERHIFeatureLevel::SM5)
 	{
@@ -1505,6 +1678,10 @@ void FScene::UpdateSkyCaptureContents(
 		const bool bLowerHemisphereIsBlack = CaptureComponent->bLowerHemisphereIsBlack;
 		const FLinearColor LowerHemisphereColor = CaptureComponent->LowerHemisphereColor;
 
+		// For FP32 textures, we are going to compute the maximum luminance accross each pixels and components, to rescale the luminance to fit in the FP16 range.
+		// Later when fetched, SpecifiedCubemapColorScale will be used to rescale the sky light to the correct original luminance.
+		const bool bSkySpecifiedCubemapUses32bitFloat = SpecifiedCubemapColorScale != nullptr && CaptureComponent->SourceType == SLS_SpecifiedCubemap && CaptureComponent->Cubemap->GetPixelFormat() == PF_A32B32G32R32F;
+
 		TRenderThreadStruct<FReflectionCubemapTexture> ReflectionCubemapTexture(CubemapResolution);
 
 		if (CaptureComponent->SourceType == SLS_CapturedScene)
@@ -1517,9 +1694,41 @@ void FScene::UpdateSkyCaptureContents(
 		{
 			const float SourceCubemapRotation = CaptureComponent->SourceCubemapAngle * (PI / 180.f);
 			ENQUEUE_RENDER_COMMAND(CopyCubemapCommand)(
-				[FeatureLevel = FeatureLevel, SourceCubemap, ReflectionCubemapTexture = ReflectionCubemapTexture.Get(), CubemapResolution, bLowerHemisphereIsBlack, SourceCubemapRotation, LowerHemisphereColor](FRHICommandListImmediate& RHICmdList)
+				[FeatureLevel = FeatureLevel, SourceCubemap, ReflectionCubemapTexture = ReflectionCubemapTexture.Get(), CubemapResolution, bLowerHemisphereIsBlack, 
+				SourceCubemapRotation, LowerHemisphereColor, bSkySpecifiedCubemapUses32bitFloat, SpecifiedCubemapColorScale](FRHICommandListImmediate& RHICmdList)
 			{
-				CopyCubemapToScratchCubemap(RHICmdList, FeatureLevel, SourceCubemap, *ReflectionCubemapTexture, CubemapResolution, true, bLowerHemisphereIsBlack, SourceCubemapRotation, LowerHemisphereColor);
+				float MaxFloatLuminance = 1.0f;
+				if (bSkySpecifiedCubemapUses32bitFloat)
+				{
+					FRDGBuilder GraphBuilder(RHICmdList);
+
+					FRDGTexture* SceneCubemapTexture = ReflectionCubemapTexture->GetRDG(GraphBuilder);
+					auto ShaderMap = GetGlobalShaderMap(FeatureLevel); 
+					ComputeSpecifiedCubemapColorScale(GraphBuilder, ShaderMap, SourceCubemap, CubemapResolution, true, bLowerHemisphereIsBlack, SourceCubemapRotation, LowerHemisphereColor, MaxFloatLuminance);
+
+					GraphBuilder.Execute(); // Execute the graph to actually fetch MaxFloatLuminance from the GPU processed texture.
+				}
+
+				// Now copy into the scratch cubemap and rescale so that the maximum values is MaxFP16, 
+				// because our runtime cubemap ReflectionCubemapTexture only uses FP16 for performance sake.
+				FLinearColor CubeLuminanceScale = FLinearColor(1.0f, 1.0f, 1.0f, 1.0f);
+				if (bSkySpecifiedCubemapUses32bitFloat && MaxFloatLuminance > FFloat16::MaxF16Float)
+				{
+					const float LuminanceScale = FFloat16::MaxF16Float / MaxFloatLuminance;
+					CubeLuminanceScale = FLinearColor(LuminanceScale, LuminanceScale, LuminanceScale, 1.0f);
+				}
+
+				CopyCubemapToScratchCubemap(RHICmdList, FeatureLevel, SourceCubemap, *ReflectionCubemapTexture, CubemapResolution, true, bLowerHemisphereIsBlack, SourceCubemapRotation, LowerHemisphereColor, CubeLuminanceScale);
+
+				// Now update SpecifiedCubemapColorScale so that it can expand a rescaled FP16 texture into its original FP32 range.
+				if (bSkySpecifiedCubemapUses32bitFloat && MaxFloatLuminance > FFloat16::MaxF16Float)
+				{
+					check(SpecifiedCubemapColorScale);
+					const float LuminanceScale = MaxFloatLuminance / FFloat16::MaxF16Float;
+					*SpecifiedCubemapColorScale = FLinearColor(LuminanceScale, LuminanceScale, LuminanceScale, 1.0);
+				}
+
+				// Now that we have rescale an FP32 input texture to fit in an FP16 range, the rest of the process is unchanged.
 			});
 		}
 		else if (CaptureComponent->IsRealTimeCaptureEnabled())
@@ -1533,7 +1742,8 @@ void FScene::UpdateSkyCaptureContents(
 		}
 
 		ENQUEUE_RENDER_COMMAND(UpdateCaptureContents)(
-			[Scene = this, FeatureLevel = FeatureLevel, ReflectionCubemapTexture = ReflectionCubemapTexture.Get(), OutAverageBrightness = &OutAverageBrightness, OutIrradianceEnvironmentMap = &OutIrradianceEnvironmentMap, OutProcessedTexture, OutRadianceMap](FRHICommandListImmediate& RHICmdList)
+			[Scene = this, FeatureLevel = FeatureLevel, ReflectionCubemapTexture = ReflectionCubemapTexture.Get(), OutAverageBrightness = &OutAverageBrightness, OutIrradianceEnvironmentMap = &OutIrradianceEnvironmentMap, 
+			OutProcessedTexture, OutRadianceMap](FRHICommandListImmediate& RHICmdList)
 		{
 			auto ShaderMap = GetGlobalShaderMap(FeatureLevel);
 
diff --git a/Engine/Source/Runtime/Renderer/Private/ScenePrivate.h b/Engine/Source/Runtime/Renderer/Private/ScenePrivate.h
--- a/Engine/Source/Runtime/Renderer/Private/ScenePrivate.h
+++ b/Engine/Source/Runtime/Renderer/Private/ScenePrivate.h
@@ -3116,7 +3116,7 @@ public:
 	virtual void UpdatePlanarReflectionContents(UPlanarReflectionComponent* CaptureComponent, FSceneRenderer& MainSceneRenderer) override;
 	virtual void AllocateReflectionCaptures(const TArray<UReflectionCaptureComponent*>& NewCaptures, const TCHAR* CaptureReason, bool bVerifyOnlyCapturing, bool bCapturingForMobile, bool bInsideTick) override;
 	virtual void ResetReflectionCaptures(bool bOnlyIfOOM) override;
-	virtual void UpdateSkyCaptureContents(const USkyLightComponent* CaptureComponent, bool bCaptureEmissiveOnly, UTextureCube* SourceCubemap, FTexture* OutProcessedTexture, float& OutAverageBrightness, FSHVectorRGB3& OutIrradianceEnvironmentMap, TArray<FFloat16Color>* OutRadianceMap) override; 
+	virtual void UpdateSkyCaptureContents(const USkyLightComponent* CaptureComponent, bool bCaptureEmissiveOnly, UTextureCube* SourceCubemap, FTexture* OutProcessedTexture, float& OutAverageBrightness, FSHVectorRGB3& OutIrradianceEnvironmentMap, TArray<FFloat16Color>* OutRadianceMap, FLinearColor* SpecifiedCubemapColorScale) override;
 	virtual void AllocateAndCaptureFrameSkyEnvMap(FRDGBuilder& GraphBuilder, FSceneRenderer& SceneRenderer, FViewInfo& MainView, bool bShouldRenderSkyAtmosphere, bool bShouldRenderVolumetricCloud, FInstanceCullingManager& InstanceCullingManager) override;
 	virtual void ValidateSkyLightRealTimeCapture(FRDGBuilder& GraphBuilder, const FViewInfo& View, FRDGTextureRef SceneColorTexture) override;
 	virtual void AddPrecomputedLightVolume(const class FPrecomputedLightVolume* Volume) override;
