Removed reference to buffer to pre-empt  breakage when removal merges from another stream.

#rnx

diff --git a/Engine/Source/Runtime/Renderer/Private/RendererScene.cpp b/Engine/Source/Runtime/Renderer/Private/RendererScene.cpp
--- a/Engine/Source/Runtime/Renderer/Private/RendererScene.cpp
+++ b/Engine/Source/Runtime/Renderer/Private/RendererScene.cpp
@@ -1096,7 +1096,6 @@ uint64 FVirtualShadowMapArrayFrameData::GetGPUSizeBytes(bool bLogSizes) const
 		GetBufferGPUSizeBytes(PageFlags, bLogSizes) +
 		GetBufferGPUSizeBytes(ProjectionData, bLogSizes) +
 		GetBufferGPUSizeBytes(PageRectBounds, bLogSizes) +
-		GetBufferGPUSizeBytes(DynamicCasterPageFlags, bLogSizes) +
 		GetBufferGPUSizeBytes(PhysicalPageMetaData, bLogSizes) +
 		GetRenderTargetGPUSizeBytes(HZBPhysical, bLogSizes);
 };
