Fix invalid check for platform with resolution higher than 2048.

diff --git a/Engine/Source/Runtime/Renderer/Private/Strata/Strata.cpp b/Engine/Source/Runtime/Renderer/Private/Strata/Strata.cpp
--- a/Engine/Source/Runtime/Renderer/Private/Strata/Strata.cpp
+++ b/Engine/Source/Runtime/Renderer/Private/Strata/Strata.cpp
@@ -166,8 +166,12 @@ static void BindStrataGlobalUniformParameters(FRDGBuilder& GraphBuilder, FStrata
 static EPixelFormat GetClassificationTileFormat(const FIntPoint& InResolution)
 {
 	// For platform which whose resolution is never above 1080p, use 8bit tile format for performance.
-	check(InResolution.X <= 2048 && InResolution.Y <= 2048);
-	return CVarStrataTileCoord8Bits.GetValueOnRenderThread() == 1 ? PF_R16_UINT : PF_R32_UINT;
+	const bool bRequest8bit = CVarStrataTileCoord8Bits.GetValueOnRenderThread() == 1;
+	if (bRequest8bit)
+	{
+		check(InResolution.X <= 2048 && InResolution.Y <= 2048);
+	}
+	return bRequest8bit ? PF_R16_UINT : PF_R32_UINT;
 }
 
 static void InitialiseStrataViewData(FRDGBuilder& GraphBuilder, FViewInfo& View, const FSceneTexturesConfig& SceneTexturesConfig, bool bNeedBSDFOffets, FStrataSceneData& SceneData)
