Fixes a crash when camera is under water with light shafts

diff --git a/Engine/Source/Runtime/Renderer/Private/TranslucentRendering.cpp b/Engine/Source/Runtime/Renderer/Private/TranslucentRendering.cpp
--- a/Engine/Source/Runtime/Renderer/Private/TranslucentRendering.cpp
+++ b/Engine/Source/Runtime/Renderer/Private/TranslucentRendering.cpp
@@ -443,6 +443,8 @@ FScreenPassTexture FTranslucencyComposition::AddPass(
 		return SceneColor;
 	}
 
+	ensure(TranslucencyTextures.Pass != ETranslucencyPass::TPT_MAX);
+
 	FRDGTextureRef SeparateModulationTexture = TranslucencyTextures.GetColorModulateForRead(GraphBuilder);
 	FRDGTextureRef SeparateTranslucencyTexture = TranslucencyTextures.GetColorForRead(GraphBuilder);
 
@@ -603,7 +605,7 @@ FScreenPassTexture FTranslucencyComposition::AddPass(
 			OpName,
 			kTranslucencyPassName[int32(TranslucencyTextures.Pass)],
 			bApplyModulateOnly ? TEXT(" ModulateOnly") : TEXT(""),
-			DepthUpscampling ? TEXT(" DepthUpscampling") : TEXT(""),
+			DepthUpscampling ? TEXT(" DepthUpsampling") : TEXT(""),
 			TranslucencyTextures.ViewRect.Width(), TranslucencyTextures.ViewRect.Height(),
 			OutputViewport.Rect.Width(), OutputViewport.Rect.Height()),
 		PixelShader,
@@ -1349,6 +1351,7 @@ void FDeferredShadingSceneRenderer::RenderTranslucencyInner(
 
 				//Invalidate.
 				TranslucencyPassResources = FTranslucencyPassResources();
+				TranslucencyPassResources.Pass = TranslucencyPass;
 			}
 			else
 			{
