HetergenousVolumes: Adding support for shadow-casting from hard surfaces.

diff --git a/Engine/Shaders/Private/HeterogeneousVolumes/HeterogeneousVolumesHardwareRayTracing.usf b/Engine/Shaders/Private/HeterogeneousVolumes/HeterogeneousVolumesHardwareRayTracing.usf
--- a/Engine/Shaders/Private/HeterogeneousVolumes/HeterogeneousVolumesHardwareRayTracing.usf
+++ b/Engine/Shaders/Private/HeterogeneousVolumes/HeterogeneousVolumesHardwareRayTracing.usf
@@ -6,6 +6,12 @@
 
 #define SUPPORT_CONTACT_SHADOWS 0
 
+#define DYNAMICALLY_SHADOWED 1
+#define TREAT_MAXDEPTH_UNSHADOWED 1
+
+#define SHADOW_QUALITY 2
+#define NO_TRANSLUCENCY_AVAILABLE
+
 #include "HeterogeneousVolumesSparseVoxelUniformBufferUtils.ush"
 #include "HeterogeneousVolumesTracingUtils.ush"
 #include "HeterogeneousVolumesLightingUtils.ush"
@@ -366,7 +372,7 @@ RAY_TRACING_ENTRY_RAYGEN(SparseVoxelsRayGenerationShader1)
 	}
 }
 
-int bApplyEmission;
+int bApplyEmissionAndTransmittance;
 int bApplyDirectLighting;
 int LightType;
 
@@ -429,6 +435,10 @@ RAY_TRACING_ENTRY_RAYGEN(RenderTransmittanceVolumeWithPreshadingRGS)
 	float3 ToLight = L * GetMaxShadowTraceDistance();
 	if (LightType != LIGHT_TYPE_DIRECTIONAL)
 	{
+		// Note: this function also permutes ToLight and L
+		float3 TranslatedWorldPosition = LWCHackToFloat(LWCAdd(WorldRayOrigin, PrimaryView.PreViewTranslation));
+		GetLocalLightAttenuation(TranslatedWorldPosition, LightData, ToLight, L);
+
 		if (LightData.bRectLight)
 		{
 			FRect Rect = GetRect(ToLight, LightData);
@@ -450,7 +460,7 @@ RAY_TRACING_ENTRY_RAYGEN(RenderTransmittanceVolumeWithPreshadingRGS)
 
 	float Jitter = bJitter ? InterleavedGradientNoise(VoxelIndex.xy, VoxelIndex.z) : 0.5;
 	float StepSize = GetStepSize() * GetShadowStepFactor();
-	int bApplyEmission = 0;
+	int bApplyEmissionAndTransmittance = 0;
 	int bApplyDirectLighting = 0;
 	int bApplyShadowTransmittance = 0;
 	FRayMarchingContext RayMarchingContext = (FRayMarchingContext)0;
@@ -468,7 +478,7 @@ RAY_TRACING_ENTRY_RAYGEN(RenderTransmittanceVolumeWithPreshadingRGS)
 		StepSize,
 		MaxStepCount,
 		// Lighting
-		bApplyEmission,
+		bApplyEmissionAndTransmittance,
 		bApplyDirectLighting,
 		bApplyShadowTransmittance
 	);
@@ -551,7 +561,7 @@ RAY_TRACING_ENTRY_RAYGEN(RenderSingleScatteringWithPreshadingRGS)
 			Jitter,
 			GetStepSize(),
 			MaxStepCount,
-			bApplyEmission,
+			bApplyEmissionAndTransmittance,
 			bApplyDirectLighting,
 			DIM_APPLY_SHADOW_TRANSMITTANCE
 		);
@@ -576,7 +586,7 @@ RAY_TRACING_ENTRY_RAYGEN(RenderSingleScatteringWithPreshadingRGS)
 			// TODO: Remap Extinction to Transmittance
 			float3 Extinction = 0.02 * (1.0 - GetVoxelDebugColor(VoxelData.VoxelIndex));
 			Transmittance *= exp(-Extinction * HitSpan);
-			if (!bDirectLighting)
+			if (!bApplyDirectLighting)
 			{
 				float3 Emission = 0.002 * GetVoxelDebugColor(VoxelData.VoxelIndex);
 				Radiance += Emission * HitSpan * Transmittance;
@@ -592,6 +602,6 @@ RAY_TRACING_ENTRY_RAYGEN(RenderSingleScatteringWithPreshadingRGS)
 	}
 
 	// Output..
-	float3 Opacity = 1.0 - Transmittance;
+	float3 Opacity = bApplyEmissionAndTransmittance ? 1.0 - Transmittance : 0.0;
 	RWLightingTexture[PixelCoord] += float4(Radiance, Luminance(Opacity));
 }
diff --git a/Engine/Shaders/Private/HeterogeneousVolumes/HeterogeneousVolumesLiveShadingPipeline.usf b/Engine/Shaders/Private/HeterogeneousVolumes/HeterogeneousVolumesLiveShadingPipeline.usf
--- a/Engine/Shaders/Private/HeterogeneousVolumes/HeterogeneousVolumesLiveShadingPipeline.usf
+++ b/Engine/Shaders/Private/HeterogeneousVolumes/HeterogeneousVolumesLiveShadingPipeline.usf
@@ -5,6 +5,12 @@
 
 #define SUPPORT_CONTACT_SHADOWS 0
 
+#define DYNAMICALLY_SHADOWED 1
+#define TREAT_MAXDEPTH_UNSHADOWED 1
+
+#define SHADOW_QUALITY 2
+#define NO_TRANSLUCENCY_AVAILABLE
+
 #include "/Engine/Generated/Material.ush"
 #include "HeterogeneousVolumesLiveShadingUtils.ush"
 #include "HeterogeneousVolumesTracingUtils.ush"
@@ -30,7 +36,7 @@ float3 LocalBoundsExtent;
 int PrimitiveId;
 
 // Lighting
-int bApplyEmission;
+int bApplyEmissionAndTransmittance;
 int bApplyDirectLighting;
 int bApplyShadowTransmittance;
 int LightType;
@@ -207,6 +213,10 @@ void RenderTransmittanceVolumeWithLiveShadingCS(
 	float3 ToLight = L * GetMaxShadowTraceDistance();
 	if (LightType != LIGHT_TYPE_DIRECTIONAL)
 	{
+		// Note: this function also permutes ToLight and L
+		float3 TranslatedWorldPosition = LWCHackToFloat(LWCAdd(WorldRayOrigin, PrimaryView.PreViewTranslation));
+		GetLocalLightAttenuation(TranslatedWorldPosition, LightData, ToLight, L);
+
 		if (LightData.bRectLight)
 		{
 			FRect Rect = GetRect(ToLight, LightData);
@@ -228,7 +238,7 @@ void RenderTransmittanceVolumeWithLiveShadingCS(
 
 	float RayJitter = length(VoxelJitter);
 	float StepSize = GetStepSize() * GetShadowStepFactor();
-	int bApplyEmission = 0;
+	int bApplyEmissionAndTransmittance = 0;
 	int bApplyDirectLighting = 0;
 	int bApplyShadowTransmittance = 0;
 	FRayMarchingContext RayMarchingContext = (FRayMarchingContext)0;
@@ -246,7 +256,7 @@ void RenderTransmittanceVolumeWithLiveShadingCS(
 		StepSize,
 		MaxStepCount,
 		// Lighting
-		bApplyEmission,
+		bApplyEmissionAndTransmittance,
 		bApplyDirectLighting,
 		bApplyShadowTransmittance
 	);
@@ -321,7 +331,7 @@ void RenderSingleScatteringWithLiveShadingCS(
 			Jitter,
 			StepSize,
 			MaxStepCount,
-			bApplyEmission,
+			bApplyEmissionAndTransmittance,
 			bApplyDirectLighting,
 			bApplyShadowTransmittance
 		);
@@ -339,7 +349,7 @@ void RenderSingleScatteringWithLiveShadingCS(
 		);
 
 		// Output..
-		float3 Opacity = 1.0 - Transmittance;
+		float3 Opacity = bApplyEmissionAndTransmittance ? 1.0 - Transmittance : 0.0;
 		RWLightingTexture[PixelCoord] += float4(Radiance, Luminance(Opacity));
 
 		// Debug output
diff --git a/Engine/Shaders/Private/HeterogeneousVolumes/HeterogeneousVolumesPreshadingPipeline.usf b/Engine/Shaders/Private/HeterogeneousVolumes/HeterogeneousVolumesPreshadingPipeline.usf
--- a/Engine/Shaders/Private/HeterogeneousVolumes/HeterogeneousVolumesPreshadingPipeline.usf
+++ b/Engine/Shaders/Private/HeterogeneousVolumes/HeterogeneousVolumesPreshadingPipeline.usf
@@ -5,6 +5,12 @@
 
 #define SUPPORT_CONTACT_SHADOWS 0
 
+#define DYNAMICALLY_SHADOWED 1
+#define TREAT_MAXDEPTH_UNSHADOWED 1
+
+#define SHADOW_QUALITY 2
+#define NO_TRANSLUCENCY_AVAILABLE
+
 #include "HeterogeneousVolumesSparseVoxelUniformBufferUtils.ush"
 #include "HeterogeneousVolumesTracingUtils.ush"
 #include "HeterogeneousVolumesLightingUtils.ush"
@@ -35,7 +41,7 @@
 #endif // DIM_APPLY_SHADOW_TRANSMITTANCE
 
 // Lighting
-int bApplyEmission;
+int bApplyEmissionAndTransmittance;
 int bApplyDirectLighting;
 int LightType;
 
@@ -451,6 +457,10 @@ void RenderTransmittanceVolumeWithPreshadingCS(
 	float3 ToLight = L * GetMaxShadowTraceDistance();
 	if (LightType != LIGHT_TYPE_DIRECTIONAL)
 	{
+		// Note: this function also permutes ToLight and L
+		float3 TranslatedWorldPosition = LWCHackToFloat(LWCAdd(WorldRayOrigin, PrimaryView.PreViewTranslation));
+		GetLocalLightAttenuation(TranslatedWorldPosition, LightData, ToLight, L);
+	
 		if (LightData.bRectLight)
 		{
 			FRect Rect = GetRect(ToLight, LightData);
@@ -473,7 +483,7 @@ void RenderTransmittanceVolumeWithPreshadingCS(
 	// Build ray-marching context
 	float RayJitter = length(VoxelJitter);
 	float StepSize = GetStepSize() * GetShadowStepFactor();
-	int bApplyEmission = 0;
+	int bApplyEmissionAndTransmittance = 0;
 	int bApplyDirectLighting = 0;
 	int bApplyShadowTransmittance = 0;
 	FRayMarchingContext RayMarchingContext = (FRayMarchingContext) 0;
@@ -491,7 +501,7 @@ void RenderTransmittanceVolumeWithPreshadingCS(
 		StepSize,
 		MaxStepCount,
 		// Lighting
-		bApplyEmission,
+		bApplyEmissionAndTransmittance,
 		bApplyDirectLighting,
 		bApplyShadowTransmittance
 	);
@@ -594,7 +604,7 @@ void RenderSingleScatteringWithPreshadingCS(
 			Jitter,
 			StepSize,
 			MaxStepCount,
-			bApplyEmission,
+			bApplyEmissionAndTransmittance,
 			bApplyDirectLighting,
 			DIM_APPLY_SHADOW_TRANSMITTANCE
 		);
@@ -625,7 +635,7 @@ void RenderSingleScatteringWithPreshadingCS(
 	}
 
 	// Output..
-	float3 Opacity = 1.0 - Transmittance;
+	float3 Opacity = bApplyEmissionAndTransmittance ? 1.0 - Transmittance : 0.0;
 	RWLightingTexture[PixelCoord] += float4(Radiance, Luminance(Opacity));
 
 #if DIM_DEBUG
diff --git a/Engine/Shaders/Private/HeterogeneousVolumes/HeterogeneousVolumesRayMarchingUtils.ush b/Engine/Shaders/Private/HeterogeneousVolumes/HeterogeneousVolumesRayMarchingUtils.ush
--- a/Engine/Shaders/Private/HeterogeneousVolumes/HeterogeneousVolumesRayMarchingUtils.ush
+++ b/Engine/Shaders/Private/HeterogeneousVolumes/HeterogeneousVolumesRayMarchingUtils.ush
@@ -2,10 +2,17 @@
 
 #include "HeterogeneousVolumesTransmittanceVolumeUtils.ush"
 
+#include "../ForwardShadowingCommon.ush"
+#include "../VolumeLightingCommon.ush"
+
 #ifndef HARDWARE_RAY_TRACING
 #define HARDWARE_RAY_TRACING 0
 #endif // HARDWARE_RAY_TRACING
 
+#ifndef HARD_SURFACE_SHADOWING
+#define HARD_SURFACE_SHADOWING 1
+#endif // HARD_SURFACE_SHADOWING
+
 struct FRayMarchingContext
 {
 	// Local-space
@@ -26,7 +33,7 @@ struct FRayMarchingContext
 	float MipLevel;
 
 	// Lighting
-	int bApplyEmission;
+	int bApplyEmissionAndTransmittance;
 	int bApplyDirectLighting;
 	int bApplyShadowTransmittance;
 };
@@ -45,7 +52,7 @@ FRayMarchingContext CreateRayMarchingContext(
 	float StepSize,
 	uint MaxStepCount,
 	// Lighting
-	int bApplyEmission,
+	int bApplyEmissionAndTransmittance,
 	int bApplyDirectLighting,
 	int bApplyShadowTransmittance
 )
@@ -65,7 +72,7 @@ FRayMarchingContext CreateRayMarchingContext(
 	Context.StepSize = StepSize;
 	Context.MaxStepCount = MaxStepCount;
 	// Lighting
-	Context.bApplyEmission = bApplyEmission;
+	Context.bApplyEmissionAndTransmittance = bApplyEmissionAndTransmittance;
 	Context.bApplyDirectLighting = bApplyDirectLighting;
 	Context.bApplyShadowTransmittance = bApplyShadowTransmittance;
 
@@ -173,7 +180,7 @@ void RayMarchSingleScattering(
 		float3 WorldPosition = RayMarchingContext.WorldRayOrigin + RayMarchingContext.WorldRayDirection * (RayMarchingContext.LocalRayTMin + RayMarchingContext.StepSize * StepIndex);
 		FVolumeSampleContext SampleContext = CreateVolumeSampleContext(LocalPosition, WorldPosition, RayMarchingContext.MipLevel);
 
-		if (RayMarchingContext.bApplyEmission)
+		if (RayMarchingContext.bApplyEmissionAndTransmittance)
 		{
 			float3 Emission = SampleEmission(SampleContext);
 			Radiance += Emission * RayMarchingContext.StepSize * Transmittance;
@@ -208,8 +215,33 @@ void RayMarchSingleScattering(
 					}
 				}
 
-				float3 Inscattering = LightData.Color * LightAttenuation;
+				float HardSurfaceShadowFactor = 1.0;
+#if HARD_SURFACE_SHADOWING
+				// Evaluate hard-surface shadow term
+				if (LightType == LIGHT_TYPE_DIRECTIONAL)
+				{
+					float SceneDepth = dot(TranslatedWorldPosition - PrimaryView.TranslatedWorldCameraOrigin, View.ViewForward);
+					bool bShadowingFromValidUVArea = false;
+					float ShadowFactor = ComputeDirectionalLightDynamicShadowing(TranslatedWorldPosition, SceneDepth, bShadowingFromValidUVArea);
+
+					if (bShadowingFromValidUVArea)
+					{
+						HardSurfaceShadowFactor *= ShadowFactor;
+					}
+				}
+				else // Local lights
+				{
+					bool bShadowingFromValidUVArea = false;
+					float ShadowFactor = ComputeVolumeShadowing(TranslatedWorldPosition, LightData.bRadialLight && !LightData.bSpotLight, LightData.bSpotLight, bShadowingFromValidUVArea);
+
+					if (bShadowingFromValidUVArea)
+					{
+						HardSurfaceShadowFactor *= ShadowFactor;
+					}
+				}
+#endif // HARD_SURFACE_SHADOWING
 
+				float3 Inscattering = LightData.Color * LightAttenuation * HardSurfaceShadowFactor;
 				if (any(Inscattering > 0.0))
 				{
 					float3 ShadowTransmittance = 1.0;
@@ -234,7 +266,7 @@ void RayMarchSingleScattering(
 						//float ShadowJitter = InterleavedGradientNoise(PixelCoord, StepIndex);
 						float ShadowJitter = 0.0;
 						float ShadowStepSize = GetStepSize() * GetShadowStepFactor();
-						int bApplyEmission = 0;
+						int bApplyEmissionAndTransmittance = 0;
 						int bApplyDirectLighting = 0;
 						int bApplyShadowTransmittance = 0;
 
@@ -248,7 +280,7 @@ void RayMarchSingleScattering(
 							ShadowJitter,
 							ShadowStepSize,
 							RayMarchingContext.MaxStepCount,
-							bApplyEmission,
+							bApplyEmissionAndTransmittance,
 							bApplyDirectLighting,
 							bApplyShadowTransmittance
 							//RayMarchingContext.MaxShadowTraceDistance
diff --git a/Engine/Source/Runtime/Renderer/Private/HeterogeneousVolumes/HeterogeneousVolumes.cpp b/Engine/Source/Runtime/Renderer/Private/HeterogeneousVolumes/HeterogeneousVolumes.cpp
--- a/Engine/Source/Runtime/Renderer/Private/HeterogeneousVolumes/HeterogeneousVolumes.cpp
+++ b/Engine/Source/Runtime/Renderer/Private/HeterogeneousVolumes/HeterogeneousVolumes.cpp
@@ -361,7 +361,9 @@ void FDeferredShadingSceneRenderer::RenderHeterogeneousVolumes(
 						Scene,
 						ViewFamily,
 						View,
-						// Object Data
+						// Shadow Data
+						VisibleLightInfos,
+						// Object data
 						PrimitiveSceneProxy,
 						MaterialRenderProxy,
 						PrimitiveId,
@@ -380,6 +382,8 @@ void FDeferredShadingSceneRenderer::RenderHeterogeneousVolumes(
 						SceneTextures,
 						Scene,
 						View,
+						// Shadow Data
+						VisibleLightInfos,
 						// Object Data
 						PrimitiveSceneProxy,
 						MaterialRenderProxy,
diff --git a/Engine/Source/Runtime/Renderer/Private/HeterogeneousVolumes/HeterogeneousVolumes.h b/Engine/Source/Runtime/Renderer/Private/HeterogeneousVolumes/HeterogeneousVolumes.h
--- a/Engine/Source/Runtime/Renderer/Private/HeterogeneousVolumes/HeterogeneousVolumes.h
+++ b/Engine/Source/Runtime/Renderer/Private/HeterogeneousVolumes/HeterogeneousVolumes.h
@@ -88,7 +88,7 @@ END_UNIFORM_BUFFER_STRUCT()
 
 BEGIN_SHADER_PARAMETER_STRUCT(FTransmittanceVolumeParameters, )
 	SHADER_PARAMETER(FIntVector, TransmittanceVolumeResolution)
-	SHADER_PARAMETER_RDG_TEXTURE_SRV(Texture3D, TransmittanceVolumeTexture)
+	SHADER_PARAMETER_RDG_TEXTURE(Texture3D, TransmittanceVolumeTexture)
 END_SHADER_PARAMETER_STRUCT()
 
 // Render specializations
@@ -98,6 +98,8 @@ void RenderWithLiveShading(
 	const FSceneTextures& SceneTextures,
 	const FScene* Scene,
 	const FViewInfo& View,
+	// Shadow data
+	TArray<FVisibleLightInfo, SceneRenderingAllocator>& VisibleLightInfos,
 	// Object data
 	const FPrimitiveSceneProxy* PrimitiveSceneProxy,
 	const FMaterialRenderProxy* MaterialRenderProxy,
@@ -116,6 +118,8 @@ void RenderWithPreshading(
 	FScene* Scene,
 	const FSceneViewFamily& ViewFamily,
 	FViewInfo& View,
+	// Shadow data
+	TArray<FVisibleLightInfo, SceneRenderingAllocator>& VisibleLightInfos,
 	// Object data
 	const FPrimitiveSceneProxy* PrimitiveSceneProxy,
 	const FMaterialRenderProxy* MaterialRenderProxy,
@@ -204,7 +208,7 @@ void RenderTransmittanceVolumeWithPreshadingHardwareRayTracing(
 	const FViewInfo& View,
 	const FSceneTextures& SceneTextures,
 	// Light data
-	bool bApplyEmission,
+	bool bApplyEmissionAndTransmittance,
 	bool bApplyDirectLighting,
 	bool bApplyShadowTransmittance,
 	uint32 LightType,
@@ -226,11 +230,13 @@ void RenderSingleScatteringWithPreshadingHardwareRayTracing(
 	const FViewInfo& View,
 	const FSceneTextures& SceneTextures,
 	// Light data
-	bool bApplyEmission,
+	bool bApplyEmissionAndTransmittance,
 	bool bApplyDirectLighting,
 	bool bApplyShadowTransmittance,
 	uint32 LightType,
 	const FLightSceneInfo* LightSceneInfo,
+	// Shadow data
+	const FVisibleLightInfo* VisibleLightInfo,
 	// Object data
 	const FPrimitiveSceneProxy* PrimitiveSceneProxy,
 	// Sparse voxel data
diff --git a/Engine/Source/Runtime/Renderer/Private/HeterogeneousVolumes/HeterogeneousVolumesHardwareRayTracing.cpp b/Engine/Source/Runtime/Renderer/Private/HeterogeneousVolumes/HeterogeneousVolumesHardwareRayTracing.cpp
--- a/Engine/Source/Runtime/Renderer/Private/HeterogeneousVolumes/HeterogeneousVolumesHardwareRayTracing.cpp
+++ b/Engine/Source/Runtime/Renderer/Private/HeterogeneousVolumes/HeterogeneousVolumesHardwareRayTracing.cpp
@@ -9,6 +9,7 @@
 #include "RendererPrivate.h"
 #include "ScenePrivate.h"
 #include "SceneManagement.h"
+#include "VolumetricFog.h"
 
 #if RHI_RAYTRACING
 FIntVector GetVoxelCoord(uint32 VoxelIndex, FIntVector VolumeResolution)
@@ -380,11 +381,15 @@ class FRenderSingleScatteringWithPreshadingRGS : public FGlobalShader
 		SHADER_PARAMETER_STRUCT_INCLUDE(FSceneTextureParameters, SceneTextures)
 
 		// Lighting data
-		SHADER_PARAMETER(int, bApplyEmission)
+		SHADER_PARAMETER(int, bApplyEmissionAndTransmittance)
 		SHADER_PARAMETER(int, bApplyDirectLighting)
 		SHADER_PARAMETER(int, LightType)
 		SHADER_PARAMETER_STRUCT_REF(FDeferredLightUniformStruct, DeferredLight)
 
+		// Shadow data
+		SHADER_PARAMETER_RDG_UNIFORM_BUFFER(FForwardLightData, ForwardLightData)
+		SHADER_PARAMETER_STRUCT_INCLUDE(FVolumeShadowingShaderParameters, VolumeShadowingShaderParameters)
+
 		// Sparse Volume
 		SHADER_PARAMETER_RDG_UNIFORM_BUFFER(FSparseVoxelUniformBufferParameters, SparseVoxelUniformBuffer)
 
@@ -408,6 +413,14 @@ class FRenderSingleScatteringWithPreshadingRGS : public FGlobalShader
 		return ShouldCompileRayTracingShadersForProject(Parameters.Platform) && DoesPlatformSupportHeterogeneousVolumes(Parameters.Platform) &&
 			FDataDrivenShaderPlatformInfo::GetSupportsRayTracingProceduralPrimitive(Parameters.Platform);
 	}
+
+	static void ModifyCompilationEnvironment(
+		const FGlobalShaderPermutationParameters& Parameters,
+		FShaderCompilerEnvironment& OutEnvironment
+	)
+	{
+		FMaterialShader::ModifyCompilationEnvironment(Parameters, OutEnvironment);
+	}
 };
 
 IMPLEMENT_GLOBAL_SHADER(FRenderSingleScatteringWithPreshadingRGS, "/Engine/Private/HeterogeneousVolumes/HeterogeneousVolumesHardwareRayTracing.usf", "RenderSingleScatteringWithPreshadingRGS", SF_RayGen);
@@ -494,7 +507,7 @@ void RenderTransmittanceVolumeWithPreshadingHardwareRayTracing(
 	const FViewInfo& View,
 	const FSceneTextures& SceneTextures,
 	// Light data
-	bool bApplyEmission,
+	bool bApplyEmissionAndTransmittance,
 	bool bApplyDirectLighting,
 	bool bApplyShadowTransmittance,
 	uint32 LightType,
@@ -593,11 +606,13 @@ void RenderSingleScatteringWithPreshadingHardwareRayTracing(
 	const FViewInfo& View,
 	const FSceneTextures& SceneTextures,
 	// Light data
-	bool bApplyEmission,
+	bool bApplyEmissionAndTransmittance,
 	bool bApplyDirectLighting,
 	bool bApplyShadowTransmittance,
 	uint32 LightType,
 	const FLightSceneInfo* LightSceneInfo,
+	// Shadow data
+	const FVisibleLightInfo* VisibleLightInfo,
 	// Object data
 	const FPrimitiveSceneProxy* PrimitiveSceneProxy,
 	// Sparse voxel data
@@ -619,7 +634,7 @@ void RenderSingleScatteringWithPreshadingHardwareRayTracing(
 
 		// Light data
 		FDeferredLightUniformStruct DeferredLightUniform;
-		PassParameters->bApplyEmission = bApplyEmission;
+		PassParameters->bApplyEmissionAndTransmittance = bApplyEmissionAndTransmittance;
 		PassParameters->bApplyDirectLighting = bApplyDirectLighting;
 		if (PassParameters->bApplyDirectLighting && (LightSceneInfo != nullptr))
 		{
@@ -628,6 +643,19 @@ void RenderSingleScatteringWithPreshadingHardwareRayTracing(
 		PassParameters->DeferredLight = CreateUniformBufferImmediate(DeferredLightUniform, UniformBuffer_SingleDraw);
 		PassParameters->LightType = LightType;
 
+		// Shadow data
+		PassParameters->ForwardLightData = View.ForwardLightingResources.ForwardLightUniformBuffer;
+		if (VisibleLightInfo != nullptr)
+		{
+			const FProjectedShadowInfo* ProjectedShadowInfo = GetShadowForInjectionIntoVolumetricFog(*VisibleLightInfo);
+			bool bDynamicallyShadowed = ProjectedShadowInfo != NULL;
+			GetVolumeShadowingShaderParameters(GraphBuilder, View, LightSceneInfo, ProjectedShadowInfo, PassParameters->VolumeShadowingShaderParameters);
+		}
+		else
+		{
+			SetVolumeShadowingDefaultShaderParametersGlobal(GraphBuilder, PassParameters->VolumeShadowingShaderParameters);
+		}
+
 		// Sparse Voxel data
 		PassParameters->SparseVoxelUniformBuffer = SparseVoxelUniformBuffer;
 
@@ -638,7 +666,7 @@ void RenderSingleScatteringWithPreshadingHardwareRayTracing(
 		if (HeterogeneousVolumes::UseTransmittanceVolume() && bApplyShadowTransmittance)
 		{
 			PassParameters->TransmittanceVolume.TransmittanceVolumeResolution = HeterogeneousVolumes::GetTransmittanceVolumeResolution();
-			PassParameters->TransmittanceVolume.TransmittanceVolumeTexture = GraphBuilder.CreateSRV(TransmittanceVolumeTexture);
+			PassParameters->TransmittanceVolume.TransmittanceVolumeTexture = TransmittanceVolumeTexture;
 		}
 
 		// Ray data
diff --git a/Engine/Source/Runtime/Renderer/Private/HeterogeneousVolumes/HeterogeneousVolumesLiveShadingPipeline.cpp b/Engine/Source/Runtime/Renderer/Private/HeterogeneousVolumes/HeterogeneousVolumesLiveShadingPipeline.cpp
--- a/Engine/Source/Runtime/Renderer/Private/HeterogeneousVolumes/HeterogeneousVolumesLiveShadingPipeline.cpp
+++ b/Engine/Source/Runtime/Renderer/Private/HeterogeneousVolumes/HeterogeneousVolumesLiveShadingPipeline.cpp
@@ -9,6 +9,7 @@
 #include "RendererPrivate.h"
 #include "ScenePrivate.h"
 #include "SceneManagement.h"
+#include "VolumetricFog.h"
 
 class FRenderTransmittanceVolumeWithLiveShadingCS : public FMeshMaterialShader
 {
@@ -20,7 +21,7 @@ class FRenderTransmittanceVolumeWithLiveShadingCS : public FMeshMaterialShader
 		SHADER_PARAMETER_STRUCT_INCLUDE(FSceneTextureParameters, SceneTextures)
 
 		// Light data
-		SHADER_PARAMETER(int, bApplyEmission)
+		SHADER_PARAMETER(int, bApplyEmissionAndTransmittance)
 		SHADER_PARAMETER(int, bApplyDirectLighting)
 		SHADER_PARAMETER(int, bApplyShadowTransmittance)
 		SHADER_PARAMETER(int, LightType)
@@ -130,7 +131,7 @@ class FRenderSingleScatteringWithLiveShadingCS : public FMeshMaterialShader
 		SHADER_PARAMETER_STRUCT_INCLUDE(FSceneTextureParameters, SceneTextures)
 
 		// Light data
-		SHADER_PARAMETER(int, bApplyEmission)
+		SHADER_PARAMETER(int, bApplyEmissionAndTransmittance)
 		SHADER_PARAMETER(int, bApplyDirectLighting)
 		SHADER_PARAMETER(int, bApplyShadowTransmittance)
 		SHADER_PARAMETER(int, LightType)
@@ -138,6 +139,8 @@ class FRenderSingleScatteringWithLiveShadingCS : public FMeshMaterialShader
 
 		// Shadow data
 		SHADER_PARAMETER(float, ShadowStepFactor)
+		SHADER_PARAMETER_RDG_UNIFORM_BUFFER(FForwardLightData, ForwardLightData)
+		SHADER_PARAMETER_STRUCT_INCLUDE(FVolumeShadowingShaderParameters, VolumeShadowingShaderParameters)
 
 		// Object data
 		SHADER_PARAMETER(FMatrix44f, LocalToWorld)
@@ -241,11 +244,13 @@ void AddComputePass(
 	FIntVector GroupCount
 )
 {
+	ClearUnusedGraphResources(ComputeShader, PassParameters);
+
 	GraphBuilder.AddPass(
 		RDG_EVENT_NAME("%s", *PassName),
 		PassParameters,
 		ERDGPassFlags::Compute,
-		[ComputeShader, PassParameters, Scene, &View, MaterialRenderProxy, &Material, GroupCount](FRHIComputeCommandList& RHICmdList)
+		[ComputeShader, PassParameters, Scene, MaterialRenderProxy, &Material, GroupCount](FRHIComputeCommandList& RHICmdList)
 		{
 			FMeshPassProcessorRenderState DrawRenderState;
 
@@ -280,21 +285,24 @@ void RenderTransmittanceVolumeWithLiveShading(
 	const FViewInfo& View,
 	const FSceneTextures& SceneTextures,
 	// Light data
-	bool bApplyEmission,
+	bool bApplyEmissionAndTransmittance,
 	bool bApplyDirectLighting,
 	bool bApplyShadowTransmittance,
 	uint32 LightType,
 	const FLightSceneInfo* LightSceneInfo,
 	// Object data
 	const FPrimitiveSceneProxy* PrimitiveSceneProxy,
-	const FMaterialRenderProxy* MaterialRenderProxy,
+	const FMaterialRenderProxy* DefaultMaterialRenderProxy,
 	const int32 PrimitiveId,
 	const FBoxSphereBounds LocalBoxSphereBounds,
 	// Output
 	FRDGTextureRef TransmittanceVolumeTexture
 )
 {
-	const FMaterial& Material = MaterialRenderProxy->GetMaterialWithFallback(View.GetFeatureLevel(), MaterialRenderProxy);
+	const FMaterialRenderProxy* MaterialRenderProxy = nullptr;
+	const FMaterial& Material = DefaultMaterialRenderProxy->GetMaterialWithFallback(View.GetFeatureLevel(), MaterialRenderProxy);
+	MaterialRenderProxy = MaterialRenderProxy ? MaterialRenderProxy : DefaultMaterialRenderProxy;
+
 	check(Material.GetMaterialDomain() == MD_Volume);
 
 	FRenderTransmittanceVolumeWithLiveShadingCS::FParameters* PassParameters = GraphBuilder.AllocParameters<FRenderTransmittanceVolumeWithLiveShadingCS::FParameters>();
@@ -304,9 +312,8 @@ void RenderTransmittanceVolumeWithLiveShading(
 		PassParameters->SceneTextures = GetSceneTextureParameters(GraphBuilder, SceneTextures);
 
 		// Light data
-		FDeferredLightUniformStruct DeferredLightUniform;
 		check(LightSceneInfo != nullptr)
-		DeferredLightUniform = GetDeferredLightParameters(View, *LightSceneInfo);
+		FDeferredLightUniformStruct DeferredLightUniform = GetDeferredLightParameters(View, *LightSceneInfo);
 		PassParameters->DeferredLight = CreateUniformBufferImmediate(DeferredLightUniform, UniformBuffer_SingleDraw);
 		PassParameters->LightType = LightType;
 
@@ -361,14 +368,16 @@ void RenderSingleScatteringWithLiveShading(
 	const FViewInfo& View,
 	const FSceneTextures& SceneTextures,
 	// Light data
-	bool bApplyEmission,
+	bool bApplyEmissionAndTransmittance,
 	bool bApplyDirectLighting,
 	bool bApplyShadowTransmittance,
 	uint32 LightType,
 	const FLightSceneInfo* LightSceneInfo,
+	// Shadow data
+	const FVisibleLightInfo* VisibleLightInfo,
 	// Object data
 	const FPrimitiveSceneProxy* PrimitiveSceneProxy,
-	const FMaterialRenderProxy* MaterialRenderProxy,
+	const FMaterialRenderProxy* DefaultMaterialRenderProxy,
 	const int32 PrimitiveId,
 	const FBoxSphereBounds LocalBoxSphereBounds,
 	// Transmittance acceleration
@@ -377,7 +386,9 @@ void RenderSingleScatteringWithLiveShading(
 	FRDGTextureRef& HeterogeneousVolumeTexture
 )
 {
-	const FMaterial& Material = MaterialRenderProxy->GetMaterialWithFallback(View.GetFeatureLevel(), MaterialRenderProxy);
+	const FMaterialRenderProxy* MaterialRenderProxy = nullptr;
+	const FMaterial& Material = DefaultMaterialRenderProxy->GetMaterialWithFallback(View.GetFeatureLevel(), MaterialRenderProxy);
+	MaterialRenderProxy = MaterialRenderProxy ? MaterialRenderProxy : DefaultMaterialRenderProxy;
 	check(Material.GetMaterialDomain() == MD_Volume);
 
 	uint32 GroupCountX = FMath::DivideAndRoundUp(View.ViewRect.Size().X, FRenderSingleScatteringWithLiveShadingCS::GetThreadGroupSize2D());
@@ -399,7 +410,7 @@ void RenderSingleScatteringWithLiveShading(
 
 		// Light data
 		FDeferredLightUniformStruct DeferredLightUniform;
-		PassParameters->bApplyEmission = bApplyEmission;
+		PassParameters->bApplyEmissionAndTransmittance = bApplyEmissionAndTransmittance;
 		PassParameters->bApplyDirectLighting = bApplyDirectLighting;
 		PassParameters->bApplyShadowTransmittance = bApplyShadowTransmittance;
 		if (PassParameters->bApplyDirectLighting && (LightSceneInfo != nullptr))
@@ -424,11 +435,32 @@ void RenderSingleScatteringWithLiveShading(
 		PassParameters->MaxStepCount = HeterogeneousVolumes::GetMaxStepCount();
 		PassParameters->bJitter = HeterogeneousVolumes::ShouldJitter();
 
+		// Shadow data
+		PassParameters->ForwardLightData = View.ForwardLightingResources.ForwardLightUniformBuffer;
+		if (VisibleLightInfo != nullptr)
+		{
+			const FProjectedShadowInfo* ProjectedShadowInfo = GetShadowForInjectionIntoVolumetricFog(*VisibleLightInfo);
+			bool bDynamicallyShadowed = ProjectedShadowInfo != NULL;
+			if (bDynamicallyShadowed)
+			{
+				GetVolumeShadowingShaderParameters(GraphBuilder, View, LightSceneInfo, ProjectedShadowInfo, PassParameters->VolumeShadowingShaderParameters);
+			}
+			else
+			{
+				SetVolumeShadowingDefaultShaderParametersGlobal(GraphBuilder, PassParameters->VolumeShadowingShaderParameters);
+			}
+		}
+		else
+		{
+			SetVolumeShadowingDefaultShaderParametersGlobal(GraphBuilder, PassParameters->VolumeShadowingShaderParameters);
+		}
+
 		// Volume data
 		if (HeterogeneousVolumes::UseTransmittanceVolume() && bApplyShadowTransmittance)
 		{
 			PassParameters->TransmittanceVolume.TransmittanceVolumeResolution = HeterogeneousVolumes::GetTransmittanceVolumeResolution();
-			PassParameters->TransmittanceVolume.TransmittanceVolumeTexture = GraphBuilder.CreateSRV(TransmittanceVolumeTexture);
+			//PassParameters->TransmittanceVolume.TransmittanceVolumeTexture = GraphBuilder.CreateSRV(TransmittanceVolumeTexture);
+			PassParameters->TransmittanceVolume.TransmittanceVolumeTexture = TransmittanceVolumeTexture;
 		}
 
 		// Dispatch data
@@ -460,6 +492,8 @@ void RenderWithLiveShading(
 	const FSceneTextures& SceneTextures,
 	const FScene* Scene,
 	const FViewInfo& View,
+	// Shadow data
+	TArray<FVisibleLightInfo, SceneRenderingAllocator>& VisibleLightInfos,
 	// Object data
 	const FPrimitiveSceneProxy* PrimitiveSceneProxy,
 	const FMaterialRenderProxy* MaterialRenderProxy,
@@ -485,16 +519,18 @@ void RenderWithLiveShading(
 	int32 NumPasses = FMath::Max(LightSceneInfoCompact.Num(), 1);
 	for (int32 PassIndex = 0; PassIndex < NumPasses; ++PassIndex)
 	{
-		bool bApplyEmission = PassIndex == 0;
+		bool bApplyEmissionAndTransmittance = PassIndex == 0;
 		bool bApplyDirectLighting = !LightSceneInfoCompact.IsEmpty();
 		bool bApplyShadowTransmittance = false;
 
 		uint32 LightType = 0;
 		FLightSceneInfo* LightSceneInfo = nullptr;
+		const FVisibleLightInfo* VisibleLightInfo = nullptr;
 		if (bApplyDirectLighting)
 		{
 			LightType = LightSceneInfoCompact[PassIndex].LightType;
 			LightSceneInfo = LightSceneInfoCompact[PassIndex].LightSceneInfo;
+			VisibleLightInfo = &VisibleLightInfos[LightSceneInfo->Id];
 			check(LightSceneInfo != nullptr);
 
 			bApplyDirectLighting = (LightSceneInfo != nullptr);
@@ -510,7 +546,7 @@ void RenderWithLiveShading(
 				View,
 				SceneTextures,
 				// Light data
-				bApplyEmission,
+				bApplyEmissionAndTransmittance,
 				bApplyDirectLighting,
 				bApplyShadowTransmittance,
 				LightType,
@@ -532,11 +568,13 @@ void RenderWithLiveShading(
 			View,
 			SceneTextures,
 			// Light data
-			bApplyEmission,
+			bApplyEmissionAndTransmittance,
 			bApplyDirectLighting,
 			bApplyShadowTransmittance,
 			LightType,
 			LightSceneInfo,
+			// Shadow data
+			VisibleLightInfo,
 			// Object data
 			PrimitiveSceneProxy,
 			MaterialRenderProxy,
diff --git a/Engine/Source/Runtime/Renderer/Private/HeterogeneousVolumes/HeterogeneousVolumesPreshadingPipeline.cpp b/Engine/Source/Runtime/Renderer/Private/HeterogeneousVolumes/HeterogeneousVolumesPreshadingPipeline.cpp
--- a/Engine/Source/Runtime/Renderer/Private/HeterogeneousVolumes/HeterogeneousVolumesPreshadingPipeline.cpp
+++ b/Engine/Source/Runtime/Renderer/Private/HeterogeneousVolumes/HeterogeneousVolumesPreshadingPipeline.cpp
@@ -9,6 +9,7 @@
 #include "RendererPrivate.h"
 #include "ScenePrivate.h"
 #include "SceneManagement.h"
+#include "VolumetricFog.h"
 
 class FGenerateRayMarchingTiles : public FGlobalShader
 {
@@ -279,11 +280,15 @@ class FRenderSingleScatteringWithPreshadingCS : public FGlobalShader
 		SHADER_PARAMETER_STRUCT_INCLUDE(FSceneTextureParameters, SceneTextures)
 
 		// Light data
-		SHADER_PARAMETER(int, bApplyEmission)
+		SHADER_PARAMETER(int, bApplyEmissionAndTransmittance)
 		SHADER_PARAMETER(int, bApplyDirectLighting)
 		SHADER_PARAMETER(int, LightType)
 		SHADER_PARAMETER_STRUCT_REF(FDeferredLightUniformStruct, DeferredLight)
 
+		// Shadow data
+		SHADER_PARAMETER_RDG_UNIFORM_BUFFER(FForwardLightData, ForwardLightData)
+		SHADER_PARAMETER_STRUCT_INCLUDE(FVolumeShadowingShaderParameters, VolumeShadowingShaderParameters)
+
 		// Volume data
 		SHADER_PARAMETER(int, MipLevel)
 
@@ -345,7 +350,7 @@ void RenderTransmitanceVolumeWithPreshadingCompute(
 	const FViewInfo& View,
 	const FSceneTextures& SceneTextures,
 	// Light data
-	bool bApplyEmission,
+	bool bApplyEmissionAndTransmittance,
 	bool bApplyDirectLighting,
 	bool bApplyShadowTransmittance,
 	uint32 LightType,
@@ -380,7 +385,7 @@ void RenderTransmitanceVolumeWithPreshadingCompute(
 
 		// Transmittance volume
 		PassParameters->TransmittanceVolume.TransmittanceVolumeResolution = HeterogeneousVolumes::GetTransmittanceVolumeResolution();
-		PassParameters->TransmittanceVolume.TransmittanceVolumeTexture = GraphBuilder.CreateSRV(TransmittanceVolumeTexture);
+		PassParameters->TransmittanceVolume.TransmittanceVolumeTexture = TransmittanceVolumeTexture;
 
 		// Ray data
 		//PassParameters->StepSize = HeterogeneousVolumes::GetStepSize();
@@ -421,11 +426,13 @@ void RenderSingleScatteringWithPreshadingCompute(
 	const FViewInfo& View,
 	const FSceneTextures& SceneTextures,
 	// Light data
-	bool bApplyEmission,
+	bool bApplyEmissionAndTransmittance,
 	bool bApplyDirectLighting,
 	bool bApplyShadowTransmittance,
 	uint32 LightType,
 	const FLightSceneInfo* LightSceneInfo,
+	// Shadow data
+	const FVisibleLightInfo* VisibleLightInfo,
 	// Object data
 	const FPrimitiveSceneProxy* PrimitiveSceneProxy,
 	// Sparse voxel data
@@ -453,7 +460,7 @@ void RenderSingleScatteringWithPreshadingCompute(
 
 		// Light data
 		FDeferredLightUniformStruct DeferredLightUniform;
-		PassParameters->bApplyEmission = bApplyEmission;
+		PassParameters->bApplyEmissionAndTransmittance = bApplyEmissionAndTransmittance;
 		PassParameters->bApplyDirectLighting = bApplyDirectLighting;
 		if (PassParameters->bApplyDirectLighting && (LightSceneInfo != nullptr))
 		{
@@ -462,6 +469,26 @@ void RenderSingleScatteringWithPreshadingCompute(
 		PassParameters->DeferredLight = CreateUniformBufferImmediate(DeferredLightUniform, UniformBuffer_SingleDraw);
 		PassParameters->LightType = LightType;
 
+		// Shadow data
+		PassParameters->ForwardLightData = View.ForwardLightingResources.ForwardLightUniformBuffer;
+		if (VisibleLightInfo != nullptr)
+		{
+			const FProjectedShadowInfo* ProjectedShadowInfo = GetShadowForInjectionIntoVolumetricFog(*VisibleLightInfo);
+			bool bDynamicallyShadowed = ProjectedShadowInfo != NULL;
+			if (bDynamicallyShadowed)
+			{
+				GetVolumeShadowingShaderParameters(GraphBuilder, View, LightSceneInfo, ProjectedShadowInfo, PassParameters->VolumeShadowingShaderParameters);
+			}
+			else
+			{
+				SetVolumeShadowingDefaultShaderParametersGlobal(GraphBuilder, PassParameters->VolumeShadowingShaderParameters);
+			}
+		}
+		else
+		{
+			SetVolumeShadowingDefaultShaderParametersGlobal(GraphBuilder, PassParameters->VolumeShadowingShaderParameters);
+		}
+
 		// Volume data
 		PassParameters->MipLevel = HeterogeneousVolumes::GetMipLevel();
 
@@ -472,7 +499,7 @@ void RenderSingleScatteringWithPreshadingCompute(
 		if (HeterogeneousVolumes::UseTransmittanceVolume() && bApplyShadowTransmittance)
 		{
 			PassParameters->TransmittanceVolume.TransmittanceVolumeResolution = HeterogeneousVolumes::GetTransmittanceVolumeResolution();
-			PassParameters->TransmittanceVolume.TransmittanceVolumeTexture = GraphBuilder.CreateSRV(TransmittanceVolumeTexture);
+			PassParameters->TransmittanceVolume.TransmittanceVolumeTexture = TransmittanceVolumeTexture;
 		}
 
 		// Ray data
@@ -522,6 +549,8 @@ void RenderWithPreshadingCompute(
 	const FScene* Scene,
 	const FSceneViewFamily& ViewFamily,
 	FViewInfo& View,
+	// Shadow data
+	TArray<FVisibleLightInfo, SceneRenderingAllocator>& VisibleLightInfos,
 	// Object data
 	const FPrimitiveSceneProxy* PrimitiveSceneProxy,
 	const FMaterialRenderProxy* MaterialRenderProxy,
@@ -577,16 +606,18 @@ void RenderWithPreshadingCompute(
 		int32 NumPasses = FMath::Max(LightSceneInfoCompact.Num(), 1);
 		for (int32 PassIndex = 0; PassIndex < NumPasses; ++PassIndex)
 		{
-			bool bApplyEmission = PassIndex == 0;
+			bool bApplyEmissionAndTransmittance = PassIndex == 0;
 			bool bApplyDirectLighting = !LightSceneInfoCompact.IsEmpty();
 			bool bApplyShadowTransmittance = false;
 
 			uint32 LightType = 0;
 			FLightSceneInfo* LightSceneInfo = nullptr;
+			const FVisibleLightInfo* VisibleLightInfo = nullptr;
 			if (bApplyDirectLighting)
 			{
 				LightType = LightSceneInfoCompact[PassIndex].LightType;
 				LightSceneInfo = LightSceneInfoCompact[PassIndex].LightSceneInfo;
+				VisibleLightInfo = &VisibleLightInfos[LightSceneInfo->Id];
 				check(LightSceneInfo != nullptr);
 
 				bApplyDirectLighting = (LightSceneInfo != nullptr);
@@ -602,7 +633,7 @@ void RenderWithPreshadingCompute(
 					View,
 					SceneTextures,
 					// Light
-					bApplyEmission,
+					bApplyEmissionAndTransmittance,
 					bApplyDirectLighting,
 					bApplyShadowTransmittance,
 					LightType,
@@ -629,11 +660,13 @@ void RenderWithPreshadingCompute(
 				View,
 				SceneTextures,
 				// Light
-				bApplyEmission,
+				bApplyEmissionAndTransmittance,
 				bApplyDirectLighting,
 				bApplyShadowTransmittance,
 				LightType,
 				LightSceneInfo,
+				// Shadow
+				VisibleLightInfo,
 				// Object
 				PrimitiveSceneProxy,
 				// Volume data
@@ -659,6 +692,8 @@ void RenderWithPreshadingHardwareRayTracing(
 	FScene* Scene,
 	const FSceneViewFamily& ViewFamily,
 	FViewInfo& View,
+	// Shadow data
+	TArray<FVisibleLightInfo, SceneRenderingAllocator>& VisibleLightInfos,
 	// Object data
 	const FPrimitiveSceneProxy* PrimitiveSceneProxy,
 	const FMaterialRenderProxy* MaterialRenderProxy,
@@ -728,16 +763,18 @@ void RenderWithPreshadingHardwareRayTracing(
 		int32 NumPasses = FMath::Max(LightSceneInfoCompact.Num(), 1);
 		for (int32 PassIndex = 0; PassIndex < NumPasses; ++PassIndex)
 		{
-			bool bApplyEmission = PassIndex == 0;
+			bool bApplyEmissionAndTransmittance = PassIndex == 0;
 			bool bApplyDirectLighting = !LightSceneInfoCompact.IsEmpty();
 			bool bApplyShadowTransmittance = false;
 
 			uint32 LightType = 0;
 			FLightSceneInfo* LightSceneInfo = nullptr;
+			const FVisibleLightInfo* VisibleLightInfo = nullptr;
 			if (bApplyDirectLighting)
 			{
 				LightType = LightSceneInfoCompact[PassIndex].LightType;
 				LightSceneInfo = LightSceneInfoCompact[PassIndex].LightSceneInfo;
+				VisibleLightInfo = &VisibleLightInfos[LightSceneInfo->Id];
 				check(LightSceneInfo != nullptr);
 
 				bApplyDirectLighting = (LightSceneInfo != nullptr);
@@ -753,7 +790,7 @@ void RenderWithPreshadingHardwareRayTracing(
 					View,
 					SceneTextures,
 					// Light data
-					bApplyEmission,
+					bApplyEmissionAndTransmittance,
 					bApplyDirectLighting,
 					bApplyShadowTransmittance,
 					LightType,
@@ -776,11 +813,13 @@ void RenderWithPreshadingHardwareRayTracing(
 				View,
 				SceneTextures,
 				// Light data
-				bApplyEmission,
+				bApplyEmissionAndTransmittance,
 				bApplyDirectLighting,
 				bApplyShadowTransmittance,
 				LightType,
 				LightSceneInfo,
+				// Shadow
+				VisibleLightInfo,
 				// Object data
 				PrimitiveSceneProxy,
 				// Sparse voxel
@@ -924,6 +963,8 @@ void RenderWithPreshading(
 	FScene* Scene,
 	const FSceneViewFamily& ViewFamily,
 	FViewInfo& View,
+	// Shadow data
+	TArray<FVisibleLightInfo, SceneRenderingAllocator>& VisibleLightInfos,
 	// Object data
 	const FPrimitiveSceneProxy* PrimitiveSceneProxy,
 	const FMaterialRenderProxy* MaterialRenderProxy,
@@ -1064,6 +1105,8 @@ void RenderWithPreshading(
 			Scene,
 			ViewFamily,
 			View,
+			// Shadow data
+			VisibleLightInfos,
 			// Object data
 			PrimitiveSceneProxy,
 			MaterialRenderProxy,
@@ -1086,6 +1129,8 @@ void RenderWithPreshading(
 			Scene,
 			ViewFamily,
 			View,
+			// Shadow data
+			VisibleLightInfos,
 			// Object data
 			PrimitiveSceneProxy,
 			MaterialRenderProxy,
diff --git a/Engine/Source/Runtime/Renderer/Private/MeshPassProcessor.cpp b/Engine/Source/Runtime/Renderer/Private/MeshPassProcessor.cpp
--- a/Engine/Source/Runtime/Renderer/Private/MeshPassProcessor.cpp
+++ b/Engine/Source/Runtime/Renderer/Private/MeshPassProcessor.cpp
@@ -340,7 +340,10 @@ void FMeshDrawShaderBindings::SetShaderBindings(
 		FShaderResourceParameterInfo Parameter = TextureSamplerParameters[SamplerIndex];
 		FRHISamplerState* Sampler = (FRHISamplerState*)SamplerBindings[SamplerIndex];
 
-		SetSamplerParameter(RHICmdList, Shader, Parameter, Sampler);
+		if (Sampler)
+		{
+			SetSamplerParameter(RHICmdList, Shader, Parameter, Sampler);
+		}
 	}
 
 	const uint8* RESTRICT SRVType = SingleShaderBindings.GetSRVTypeStart();
diff --git a/Engine/Source/Runtime/Renderer/Private/VolumeLighting.cpp b/Engine/Source/Runtime/Renderer/Private/VolumeLighting.cpp
--- a/Engine/Source/Runtime/Renderer/Private/VolumeLighting.cpp
+++ b/Engine/Source/Runtime/Renderer/Private/VolumeLighting.cpp
@@ -31,7 +31,7 @@ const FProjectedShadowInfo* GetFirstWholeSceneShadowMap(const FVisibleLightInfo&
 	return nullptr;
 }
 
-static void SetVolumeShadowingDefaultShaderParametersGlobal(FRDGBuilder& GraphBuilder, FVolumeShadowingShaderParameters& ShaderParams)
+void SetVolumeShadowingDefaultShaderParametersGlobal(FRDGBuilder& GraphBuilder, FVolumeShadowingShaderParameters& ShaderParams)
 {
 	const FRDGSystemTextures& SystemTextures = FRDGSystemTextures::Get(GraphBuilder);
 	FRDGTextureRef BlackDepthCubeTexture = SystemTextures.BlackDepthCube;
diff --git a/Engine/Source/Runtime/Renderer/Private/VolumeLighting.h b/Engine/Source/Runtime/Renderer/Private/VolumeLighting.h
--- a/Engine/Source/Runtime/Renderer/Private/VolumeLighting.h
+++ b/Engine/Source/Runtime/Renderer/Private/VolumeLighting.h
@@ -37,7 +37,9 @@ void GetVolumeShadowingShaderParameters(
 	const FProjectedShadowInfo* ShadowInfo,
 	FVolumeShadowingShaderParameters& OutParameters);
 
-
+void SetVolumeShadowingDefaultShaderParametersGlobal(
+	FRDGBuilder& GraphBuilder,
+	FVolumeShadowingShaderParameters& ShaderParams);
 
 ///
 ///
diff --git a/Engine/Source/Runtime/Renderer/Private/VolumetricFog.h b/Engine/Source/Runtime/Renderer/Private/VolumetricFog.h
--- a/Engine/Source/Runtime/Renderer/Private/VolumetricFog.h
+++ b/Engine/Source/Runtime/Renderer/Private/VolumetricFog.h
@@ -58,6 +58,7 @@ inline bool DoesPlatformSupportVolumetricFogVoxelization(const FStaticShaderPlat
 }
 
 extern bool ShouldRenderVolumetricFog(const FScene* Scene, const FSceneViewFamily& ViewFamily);
+extern const FProjectedShadowInfo* GetShadowForInjectionIntoVolumetricFog(const FVisibleLightInfo& VisibleLightInfo);
 
 extern bool LightNeedsSeparateInjectionIntoVolumetricFogForOpaqueShadow(const FViewInfo& View, const FLightSceneInfo* LightSceneInfo, const FVisibleLightInfo& VisibleLightInfo);
 extern bool LightNeedsSeparateInjectionIntoVolumetricFogForLightFunction(const FLightSceneInfo* LightSceneInfo);
