PSO Precache Base Changes (disabled by default)

Add PSOPrecache file which wraps engine level PSO structs and functions
- FPSOPrecacheParams used to wrap certain parameters which drive the PSO collection (mostly set from components)
- IPSOCollector interface used as engine entry to the mesh pass processors in renderer project
- FPSOCollectorCreateManager to wrap all statically registered IPSOCollector CreateFuncions
- PSOCollectorStats structs and helper functions

Add PSOPrecache support to PipelineStateCache:
- functions to precache compute and graphics PSOs
- functions to check current precache state
- FPrecacheGraphicsPipelineCache which wraps all currently precached PSOs

IMaterial entry function to PrecachePSOs with given vertex factories and params (implemented by UMaterial & UMaterialInstance) which forwards call to FMaterial and finally the FMaterialShaderMap
- FMaterialShaderMap will iterate all registered IPSOCollector to collect all possible PSOs used and forward them to the PipelineStateCache for actual (async) precaching

RHI Changes:
- Add RHI specific functions to retrieve & compare the FRHIVertexDeclaration and FGraphicsPipelineStateInitializer hash for each RHI
- Add helper functions to match RHI*State data by retrieving initializer

Minimal vertex factory changes for PSO precaching as prep - actual changes in upcoming CLs

MeshPassProcessor base functions to collect PSOs which can be used by all mesh pass processor implemention (actual changes in upcoming CLs)
- Cache complete precache data hash of stored minimal PSO so it can be used for fast actual PSO initializer hash computation (full PSO initializer is build at runtime)
- Cache if minimal PSO was successfully precache and option to skip draw when PSO is still precaching (disabled by default and doesn't work yet via GPU scene based rendering)
- Add helper functions to setup the per mesh pass render target data used during PSO collection

Refactor SceneTextureConfig so it can be initialized from an init struct and move all function to setup the internal members from SceneTexture to SceneTextureConfig using the data from the init struct (so it's available in Engine project)

Fix GBufferInfo format for velocity target (was always floating point while it should be unorm when not android

Make certain helper functions definitions to h file so they can be used everywhere (SupportsNaniteRendering, FTextureRenderTarget2DResource::GetFormat & CreateFlags)

diff --git a/Engine/Source/Runtime/D3D12RHI/Private/D3D12RHIPrivate.h b/Engine/Source/Runtime/D3D12RHI/Private/D3D12RHIPrivate.h
--- a/Engine/Source/Runtime/D3D12RHI/Private/D3D12RHIPrivate.h
+++ b/Engine/Source/Runtime/D3D12RHI/Private/D3D12RHIPrivate.h
@@ -395,6 +395,8 @@ public:
 	virtual void RHIAliasTextureResources(FTextureRHIRef& DestTexture, FTextureRHIRef& SrcTexture) final override;
 	virtual FTextureRHIRef RHICreateAliasedTexture(FTextureRHIRef& SourceTexture) final override;
 	virtual void RHIGetDisplaysInformation(FDisplayInformationArray& OutDisplayInformation) final override;
+	virtual uint64 RHIComputePrecachePSOHash(const FGraphicsPipelineStateInitializer& Initializer) final override;
+	virtual bool RHIMatchPrecachePSOInitializers(const FGraphicsPipelineStateInitializer& LHS, const FGraphicsPipelineStateInitializer& RHS) final override;
 	virtual void RHIAdvanceFrameFence() final override;
 	virtual void RHIAdvanceFrameForGetViewportBackBuffer(FRHIViewport* Viewport) final override;
 	virtual void RHIAcquireThreadOwnership() final override;
diff --git a/Engine/Source/Runtime/D3D12RHI/Private/D3D12Shader.h b/Engine/Source/Runtime/D3D12RHI/Private/D3D12Shader.h
--- a/Engine/Source/Runtime/D3D12RHI/Private/D3D12Shader.h
+++ b/Engine/Source/Runtime/D3D12RHI/Private/D3D12Shader.h
@@ -36,6 +36,7 @@ public:
 	}
 
 	virtual bool GetInitializer(FVertexDeclarationElementList& Init) final override;
+	virtual uint32 GetPrecachePSOHash() const final override { return HashNoStrides;  }
 };
 
 struct FD3D12ShaderData
diff --git a/Engine/Source/Runtime/D3D12RHI/Private/D3D12State.cpp b/Engine/Source/Runtime/D3D12RHI/Private/D3D12State.cpp
--- a/Engine/Source/Runtime/D3D12RHI/Private/D3D12State.cpp
+++ b/Engine/Source/Runtime/D3D12RHI/Private/D3D12State.cpp
@@ -458,6 +458,101 @@ bool FD3D12BlendState::GetInitializer(class FBlendStateInitializerRHI& Init)
 	return true;
 }
 
+uint64 FD3D12DynamicRHI::RHIComputePrecachePSOHash(const FGraphicsPipelineStateInitializer& Initializer)
+{
+	// When compute precache PSO hash we assume a valid state precache PSO hash is already provided
+	checkf(Initializer.StatePrecachePSOHash != 0, TEXT("Initializer should have a valid state precache PSO hash set when computing the full initializer PSO hash"));
+
+	// All members which are not part of the state objects and influence the PSO on D3D12
+	struct FNonStateHashKey
+	{
+		uint64							StatePrecachePSOHash;
+
+		EPrimitiveType					PrimitiveType;
+		uint32							RenderTargetsEnabled;
+		FGraphicsPipelineStateInitializer::TRenderTargetFormats RenderTargetFormats;
+		EPixelFormat					DepthStencilTargetFormat;
+		uint16							NumSamples;
+		EConservativeRasterization		ConservativeRasterization;
+		bool							bDepthBounds;
+		uint8							MultiViewCount;
+		bool							bHasFragmentDensityAttachment;
+		EVRSShadingRate					ShadingRate;
+	} HashKey;
+
+	FMemory::Memzero(&HashKey, sizeof(FNonStateHashKey));
+
+	HashKey.StatePrecachePSOHash			= Initializer.StatePrecachePSOHash;
+
+	HashKey.PrimitiveType					= Initializer.PrimitiveType;
+	HashKey.RenderTargetsEnabled			= Initializer.RenderTargetsEnabled;
+	HashKey.RenderTargetFormats				= Initializer.RenderTargetFormats;
+	HashKey.DepthStencilTargetFormat		= Initializer.DepthStencilTargetFormat;
+	HashKey.NumSamples						= Initializer.NumSamples;
+	HashKey.ConservativeRasterization		= Initializer.ConservativeRasterization;
+	HashKey.bDepthBounds					= Initializer.bDepthBounds;
+	HashKey.MultiViewCount					= Initializer.MultiViewCount;
+	HashKey.bHasFragmentDensityAttachment	= Initializer.bHasFragmentDensityAttachment;
+	HashKey.ShadingRate						= Initializer.ShadingRate;
+
+	return CityHash64((const char*)&HashKey, sizeof(FNonStateHashKey));
+}
+
+bool FD3D12DynamicRHI::RHIMatchPrecachePSOInitializers(const FGraphicsPipelineStateInitializer& LHS, const FGraphicsPipelineStateInitializer& RHS)
+{
+	// first check non pointer objects
+	if (LHS.ImmutableSamplerState != LHS.ImmutableSamplerState ||
+		LHS.PrimitiveType != RHS.PrimitiveType ||
+		LHS.bDepthBounds != RHS.bDepthBounds ||
+		LHS.MultiViewCount != RHS.MultiViewCount ||
+		LHS.ShadingRate != RHS.ShadingRate ||
+		LHS.bHasFragmentDensityAttachment != RHS.bHasFragmentDensityAttachment ||
+		LHS.RenderTargetsEnabled != RHS.RenderTargetsEnabled ||
+		LHS.RenderTargetFormats != RHS.RenderTargetFormats ||
+		LHS.DepthStencilTargetFormat != RHS.DepthStencilTargetFormat ||
+		LHS.NumSamples != RHS.NumSamples ||
+		LHS.ConservativeRasterization != RHS.ConservativeRasterization)
+	{
+		return false;
+	}
+
+	// check the RHI shaders (pointer check for shaders should be fine)
+	if (LHS.BoundShaderState.GetVertexShader() != RHS.BoundShaderState.GetVertexShader() ||
+		LHS.BoundShaderState.GetPixelShader() != RHS.BoundShaderState.GetPixelShader() ||
+		LHS.BoundShaderState.GetMeshShader() != RHS.BoundShaderState.GetMeshShader() ||
+		LHS.BoundShaderState.GetAmplificationShader() != RHS.BoundShaderState.GetAmplificationShader() ||
+		LHS.BoundShaderState.GetGeometryShader() != RHS.BoundShaderState.GetGeometryShader())
+	{
+		return false;
+	}
+
+	// Compare the d3d12 vertex elements without the stride
+	FD3D12VertexElements LHSVertexElements;
+	if (LHS.BoundShaderState.VertexDeclarationRHI)
+	{
+		LHSVertexElements = ((FD3D12VertexDeclaration*)LHS.BoundShaderState.VertexDeclarationRHI)->VertexElements;
+	}
+	FD3D12VertexElements RHSVertexElements;
+	if (RHS.BoundShaderState.VertexDeclarationRHI)
+	{
+		RHSVertexElements = ((FD3D12VertexDeclaration*)RHS.BoundShaderState.VertexDeclarationRHI)->VertexElements;
+	}
+	if (LHSVertexElements != RHSVertexElements)
+	{
+		return false;
+	}
+
+	// Check actual state content (each initializer can have it's own state and not going through a factory)
+	if (!MatchRHIState<FRHIBlendState, FBlendStateInitializerRHI>(LHS.BlendState, RHS.BlendState) ||
+		!MatchRHIState<FRHIRasterizerState, FRasterizerStateInitializerRHI>(LHS.RasterizerState, RHS.RasterizerState) ||
+		!MatchRHIState<FRHIDepthStencilState, FDepthStencilStateInitializerRHI>(LHS.DepthStencilState, RHS.DepthStencilState))
+	{
+		return false;
+	}
+
+	return true;
+}
+
 FGraphicsPipelineStateRHIRef FD3D12DynamicRHI::RHICreateGraphicsPipelineState(const FGraphicsPipelineStateInitializer& Initializer)
 {
 	SCOPE_CYCLE_COUNTER(STAT_PSOGraphicsFindOrCreateTime);
diff --git a/Engine/Source/Runtime/Engine/Classes/Materials/Material.h b/Engine/Source/Runtime/Engine/Classes/Materials/Material.h
--- a/Engine/Source/Runtime/Engine/Classes/Materials/Material.h
+++ b/Engine/Source/Runtime/Engine/Classes/Materials/Material.h
@@ -1056,6 +1056,8 @@ public:
 	ENGINE_API virtual USubsurfaceProfile* GetSubsurfaceProfile_Internal() const override;
 	ENGINE_API virtual bool CastsRayTracedShadows() const override;
 
+	ENGINE_API virtual FGraphEventArray PrecachePSOs(const TConstArrayView<const FVertexFactoryType*>& VertexFactoryTypes, const FPSOPrecacheParams& PreCacheParams) override;
+
 	ENGINE_API void SetShadingModel(EMaterialShadingModel NewModel) { ensure(ShadingModel < MSM_NUM); ShadingModel = NewModel; ShadingModels = FMaterialShadingModelField(ShadingModel); }
 
 	/** Checks to see if an input property should be active, based on the state of the material */
diff --git a/Engine/Source/Runtime/Engine/Classes/Materials/MaterialInstance.h b/Engine/Source/Runtime/Engine/Classes/Materials/MaterialInstance.h
--- a/Engine/Source/Runtime/Engine/Classes/Materials/MaterialInstance.h
+++ b/Engine/Source/Runtime/Engine/Classes/Materials/MaterialInstance.h
@@ -725,6 +725,8 @@ public:
 #endif
 	virtual ENGINE_API void RecacheUniformExpressions(bool bRecreateUniformBuffer) const override;
 	virtual ENGINE_API bool GetRefractionSettings(float& OutBiasValue) const override;
+	
+	virtual ENGINE_API FGraphEventArray PrecachePSOs(const TConstArrayView<const FVertexFactoryType*>& VertexFactoryTypes, const FPSOPrecacheParams& PreCacheParams) override;
 
 #if WITH_EDITOR
 	ENGINE_API virtual void ForceRecompileForRendering() override;
diff --git a/Engine/Source/Runtime/Engine/Classes/Materials/MaterialInterface.h b/Engine/Source/Runtime/Engine/Classes/Materials/MaterialInterface.h
--- a/Engine/Source/Runtime/Engine/Classes/Materials/MaterialInterface.h
+++ b/Engine/Source/Runtime/Engine/Classes/Materials/MaterialInterface.h
@@ -499,6 +499,15 @@ public:
 	 */
 	virtual bool GetMaterialLayers(FMaterialLayersFunctions& OutLayers, TMicRecursionGuard RecursionGuard = TMicRecursionGuard()) const PURE_VIRTUAL(UMaterialInterface::GetMaterialLayers, return false;);
 
+	/**
+	 * Precache PSOs which can be used for this material for the given vertex factory type and material paramaters
+	 */
+	virtual FGraphEventArray PrecachePSOs(const FVertexFactoryType* VertexFactoryType, const struct FPSOPrecacheParams& PreCacheParams)
+	{
+		return PrecachePSOs(MakeArrayView(&VertexFactoryType, 1), PreCacheParams);
+	}
+	virtual FGraphEventArray PrecachePSOs(const TConstArrayView<const FVertexFactoryType*>& VertexFactoryTypes, const struct FPSOPrecacheParams& PreCacheParams) PURE_VIRTUAL(UMaterialInterface::PrecachePSOs, return FGraphEventArray(););
+
 #if WITH_EDITORONLY_DATA
 
 	/**
diff --git a/Engine/Source/Runtime/Engine/Private/Materials/Material.cpp b/Engine/Source/Runtime/Engine/Private/Materials/Material.cpp
--- a/Engine/Source/Runtime/Engine/Private/Materials/Material.cpp
+++ b/Engine/Source/Runtime/Engine/Private/Materials/Material.cpp
@@ -542,7 +542,29 @@ void UMaterialInterface::InitDefaultMaterials()
 				}
 			}
 		}
-		
+
+		// Now precache PSOs for all the default materials
+		if (PipelineStateCache::IsPSOPrecachingEnabled())
+		{
+			FPSOPrecacheParams PrecachePSOParams;
+			TArray<const FVertexFactoryType*> AllVertexFactoryTypes;
+			for (TLinkedList<FVertexFactoryType*>::TIterator It(FVertexFactoryType::GetTypeList()); It; It.Next())
+			{
+				AllVertexFactoryTypes.Add(*It);
+			}
+			for (int32 Domain = 0; Domain < MD_MAX; ++Domain)
+			{
+				if (GDefaultMaterials[Domain])
+				{
+					PrecachePSOParams.Mobility = EComponentMobility::Static;
+					GDefaultMaterials[Domain]->PrecachePSOs(AllVertexFactoryTypes, PrecachePSOParams);
+
+					PrecachePSOParams.Mobility = EComponentMobility::Movable;
+					GDefaultMaterials[Domain]->PrecachePSOs(AllVertexFactoryTypes, PrecachePSOParams);
+				}
+			}
+		}
+
 		RecursionLevel--;
 #if USE_EVENT_DRIVEN_ASYNC_LOAD_AT_BOOT_TIME
 		bInitialized = !GEventDrivenLoaderEnabled || RecursionLevel == 0;
@@ -2478,6 +2500,29 @@ bool UMaterial::IsComplete() const
 	return bComplete;
 }
 
+FGraphEventArray UMaterial::PrecachePSOs(const TConstArrayView<const FVertexFactoryType*>& VertexFactoryTypes, const FPSOPrecacheParams& InPreCacheParams)
+{
+	FGraphEventArray GraphEvents;
+	if (FApp::CanEverRender() && PipelineStateCache::IsPSOPrecachingEnabled())
+	{
+		// Assume post load has been called on the UMaterial already and material resources are created
+		check(MaterialResources.Num() > 0);
+
+		EMaterialQualityLevel::Type ActiveQualityLevel = GetCachedScalabilityCVars().MaterialQualityLevel;
+		uint32 FeatureLevelsToCompile = GetFeatureLevelsToCompileForRendering();
+		while (FeatureLevelsToCompile != 0)
+		{
+			const ERHIFeatureLevel::Type FeatureLevel = (ERHIFeatureLevel::Type)FBitSet::GetAndClearNextBit(FeatureLevelsToCompile);
+			FMaterialResource* MaterialResource = FindMaterialResource(MaterialResources, FeatureLevel, ActiveQualityLevel, true/*bAllowDefaultMaterial*/);
+			if (MaterialResource)
+			{
+				GraphEvents.Append(MaterialResource->CollectPSOs(FeatureLevel, VertexFactoryTypes, InPreCacheParams));
+			}
+		}
+	}
+	return GraphEvents;
+}
+
 void UMaterial::ReleaseResourcesAndMutateDDCKey(const FGuid& TransformationId)
 {
 	if (TransformationId.IsValid())
diff --git a/Engine/Source/Runtime/Engine/Private/Materials/MaterialInstance.cpp b/Engine/Source/Runtime/Engine/Private/Materials/MaterialInstance.cpp
--- a/Engine/Source/Runtime/Engine/Private/Materials/MaterialInstance.cpp
+++ b/Engine/Source/Runtime/Engine/Private/Materials/MaterialInstance.cpp
@@ -2098,7 +2098,7 @@ void UMaterialInstance::CacheResourceShadersForRendering(EMaterialShaderPrecompi
 	OutResourcesToFree = MoveTemp(StaticPermutationMaterialResources);
 	StaticPermutationMaterialResources.Reset();
 #endif // STORE_ONLY_ACTIVE_SHADERMAPS
-
+	
 	if (bHasStaticPermutationResource && FApp::CanEverRender())
 	{
 		check(IsA(UMaterialInstanceConstant::StaticClass()));
@@ -2287,6 +2287,33 @@ void UMaterialInstance::CacheShaders(EMaterialShaderPrecompileMode CompileMode)
 	InitStaticPermutation(CompileMode);
 }
 
+FGraphEventArray UMaterialInstance::PrecachePSOs(const TConstArrayView<const FVertexFactoryType*>& VertexFactoryTypes, const FPSOPrecacheParams& InPreCacheParams)
+{
+	FGraphEventArray GraphEvents;
+	if (FApp::CanEverRender() && PipelineStateCache::IsPSOPrecachingEnabled() && Parent)
+	{
+		if (bHasStaticPermutationResource)
+		{			
+			EMaterialQualityLevel::Type ActiveQualityLevel = GetCachedScalabilityCVars().MaterialQualityLevel;
+			uint32 FeatureLevelsToCompile = GetFeatureLevelsToCompileForRendering();
+			while (FeatureLevelsToCompile != 0)
+			{
+				const ERHIFeatureLevel::Type FeatureLevel = (ERHIFeatureLevel::Type)FBitSet::GetAndClearNextBit(FeatureLevelsToCompile);
+				FMaterialResource* StaticPermutationResource = FindMaterialResource(StaticPermutationMaterialResources, FeatureLevel, ActiveQualityLevel, true/*bAllowDefaultMaterial*/);
+				if (StaticPermutationResource)
+				{
+					GraphEvents.Append(StaticPermutationResource->CollectPSOs(FeatureLevel, VertexFactoryTypes, InPreCacheParams));
+				}
+			}
+		}
+		else
+		{
+			GraphEvents = Parent->PrecachePSOs(VertexFactoryTypes, InPreCacheParams);
+		}
+	}
+	return GraphEvents;
+}
+
 #if WITH_EDITOR
 void UMaterialInstance::CacheGivenTypesForCooking(EShaderPlatform ShaderPlatform, ERHIFeatureLevel::Type FeatureLevel, EMaterialQualityLevel::Type QualityLevel, const TArray<const FVertexFactoryType*>& VFTypes, const TArray<const FShaderPipelineType*> PipelineTypes, const TArray<const FShaderType*>& ShaderTypes)
 {
diff --git a/Engine/Source/Runtime/Engine/Private/Materials/MaterialShader.cpp b/Engine/Source/Runtime/Engine/Private/Materials/MaterialShader.cpp
--- a/Engine/Source/Runtime/Engine/Private/Materials/MaterialShader.cpp
+++ b/Engine/Source/Runtime/Engine/Private/Materials/MaterialShader.cpp
@@ -21,6 +21,8 @@
 #include "ProfilingDebugging/CountersTrace.h"
 #include "ProfilingDebugging/LoadTimeTracker.h"
 #include "Misc/ScopeRWLock.h"
+#include "SceneTexturesConfig.h"
+#include "PSOPrecache.h"
 
 int32 GMaterialExcludeNonPipelinedShaders = 1;
 static FAutoConsoleVariableRef CVarMaterialExcludeNonPipelinedShaders(
@@ -2728,6 +2730,47 @@ bool FMaterialShaderMap::IsComplete(const FMaterial* Material, bool bSilent)
 	return true;
 }
 
+FGraphEventArray FMaterialShaderMap::CollectPSOs(ERHIFeatureLevel::Type InFeatureLevel, const FMaterial* Material, const TConstArrayView<const FVertexFactoryType*>& VertexFactoryTypes, const FPSOPrecacheParams& PreCacheParams)
+{
+	// Only feature level is currently set as init settings - rest is default
+	// (multiview & alpha channel not taken into account here)
+	FSceneTexturesConfigInitSettings SceneTexturesConfigInitSettings;
+	SceneTexturesConfigInitSettings.FeatureLevel = InFeatureLevel;
+
+	FSceneTexturesConfig SceneTexturesConfig;
+	SceneTexturesConfig.Init(SceneTexturesConfigInitSettings);
+
+	const FMaterialShaderMapContent* LocalContent = GetContent();
+	const EShaderPlatform Platform = LocalContent->GetShaderPlatform();
+	const FMaterialShaderParameters MaterialParameters(Material);
+
+	const EShadingPath ShadingPath = FSceneInterface::GetShadingPath(InFeatureLevel);
+
+	TArray<FPSOPrecacheData> PSOInitializers;
+	PSOInitializers.Reserve(32);
+	
+	for (uint32 Index = 0; Index < FPSOCollectorCreateManager::MaxPSOCollectorCount; ++Index)
+	{
+		PSOCollectorCreateFunction CreateFunction = FPSOCollectorCreateManager::GetCreateFunction(ShadingPath, Index);
+		if (CreateFunction)
+		{
+			IPSOCollector* PSOCollector = CreateFunction(InFeatureLevel);
+			if (PSOCollector != nullptr)
+			{
+				for (const FVertexFactoryType* VFType : VertexFactoryTypes)
+				{
+					if (VFType->SupportsPSOPrecaching() && LocalContent->GetMeshShaderMap(VFType->GetHashedName()))
+					{
+						PSOCollector->CollectPSOInitializers(SceneTexturesConfig, *Material, VFType, PreCacheParams, PSOInitializers);
+					}
+				}
+			}
+		}
+	}
+
+	return PrecachePSOs(PSOInitializers);
+}
+
 #if WITH_EDITOR
 void FMaterialShaderMap::LoadMissingShadersFromMemory(const FMaterial* Material)
 {
diff --git a/Engine/Source/Runtime/Engine/Private/Materials/MaterialShared.cpp b/Engine/Source/Runtime/Engine/Private/Materials/MaterialShared.cpp
--- a/Engine/Source/Runtime/Engine/Private/Materials/MaterialShared.cpp
+++ b/Engine/Source/Runtime/Engine/Private/Materials/MaterialShared.cpp
@@ -1076,6 +1076,11 @@ uint8 FMaterial::GetCustomDepthStencilUsageMask_GameThread() const
 	return CustomDepthStencilUsageMask;
 }
 
+uint8 FMaterial::GetRuntimeVirtualTextureOutputAttibuteMask_GameThread() const
+{
+	return GameThreadShaderMap.GetReference() ? GameThreadShaderMap->GetRuntimeVirtualTextureOutputAttributeMask() : 0;
+}
+
 uint8 FMaterial::GetRuntimeVirtualTextureOutputAttibuteMask_RenderThread() const
 {
 	check(IsInParallelRenderingThread());
@@ -2624,6 +2629,61 @@ bool FMaterial::CacheShaders(const FMaterialShaderMapId& ShaderMapId, EShaderPla
 #endif
 }
 
+FGraphEventArray FMaterial::CollectPSOs(ERHIFeatureLevel::Type InFeatureLevel, const TConstArrayView<const FVertexFactoryType*>& VertexFactoryTypes, const FPSOPrecacheParams& PreCacheParams)
+{
+	FGraphEventArray GraphEvents;
+
+	// Only care about inline shaders for now
+	if (!bContainsInlineShaders || GameThreadShaderMap == nullptr)
+	{
+		return GraphEvents;
+	}
+
+	TArray<const FVertexFactoryType*, TInlineAllocator<4>> MissingVFs;
+
+	// Try and find missing entries which still need to be precached
+	{
+		FRWScopeLock ReadLock(PrecachePSOVFLock, SLT_ReadOnly);
+		for (const FVertexFactoryType* VFType : VertexFactoryTypes)
+		{
+			FPrecacheVertexTypeWithParams PrecacheVertexTypeWithParams;
+			PrecacheVertexTypeWithParams.VertexFactoryType = VFType;
+			PrecacheVertexTypeWithParams.PrecachePSOParams = PreCacheParams;
+			if (!PrecachedPSOVertexFactories.Contains(PrecacheVertexTypeWithParams))
+			{
+				MissingVFs.Add(VFType);
+			}
+		}
+	}
+
+	if (MissingVFs.Num() > 0)
+	{
+		// Build array again with writing lock because we probably have missing entries
+		TArray<const FVertexFactoryType*, TInlineAllocator<4>> ActualMissingVFs;
+
+		{
+			// We possibly need to add
+			FRWScopeLock WriteLock(PrecachePSOVFLock, SLT_Write);
+
+			for (const FVertexFactoryType* VFType : MissingVFs)
+			{
+				FPrecacheVertexTypeWithParams PrecacheVertexTypeWithParams;
+				PrecacheVertexTypeWithParams.VertexFactoryType = VFType;
+				PrecacheVertexTypeWithParams.PrecachePSOParams = PreCacheParams;
+				if (PrecachedPSOVertexFactories.Contains(PrecacheVertexTypeWithParams))
+					continue;
+
+				ActualMissingVFs.Add(VFType);
+				PrecachedPSOVertexFactories.Add(PrecacheVertexTypeWithParams);
+			}
+		}
+
+		GraphEvents.Append(GameThreadShaderMap->CollectPSOs(InFeatureLevel, this, ActualMissingVFs, PreCacheParams));
+	}
+
+	return GraphEvents;
+}
+
 #if WITH_EDITOR
 
 void FMaterial::BeginCacheShaders(EShaderPlatform Platform, EMaterialShaderPrecompileMode PrecompileMode, const ITargetPlatform* TargetPlatform, TUniqueFunction<void(bool bSuccess)>&& CompletionCallback)
diff --git a/Engine/Source/Runtime/Engine/Private/PSOPrecache.cpp b/Engine/Source/Runtime/Engine/Private/PSOPrecache.cpp
new file mode 100644
--- /dev/null
+++ b/Engine/Source/Runtime/Engine/Private/PSOPrecache.cpp
@@ -0,0 +1,141 @@
+// Copyright Epic Games, Inc. All Rights Reserved.
+
+/*=============================================================================
+	MeshPassProcessorManager.cpp: 
+=============================================================================*/
+
+#include "PSOPrecache.h"
+
+PSOCollectorCreateFunction FPSOCollectorCreateManager::JumpTable[(int32)EShadingPath::Num][FPSOCollectorCreateManager::MaxPSOCollectorCount] = {};
+
+FGraphEventArray PrecachePSOs(const TArray<FPSOPrecacheData>& PSOInitializers)
+{
+	FGraphEventArray GraphEvents;
+	for (const FPSOPrecacheData& PrecacheData : PSOInitializers)
+	{
+#if PSO_PRECACHING_VALIDATE
+		PSOCollectorStats::AddPipelineStateToCache(PrecacheData.PSOInitializer, PrecacheData.MeshPassType, PrecacheData.VertexFactoryType);
+#endif // PSO_PRECACHING_VALIDATE
+
+		FGraphEventRef GraphEvent = PipelineStateCache::PrecacheGraphicsPipelineState(PrecacheData.PSOInitializer);
+		if (GraphEvent)
+		{
+			GraphEvents.Add(GraphEvent);
+		}
+	}
+
+	return GraphEvents;
+}
+
+#if PSO_PRECACHING_VALIDATE
+
+int32 GValidatePrecaching = 0;
+static FAutoConsoleVariableRef CVarValidatePSOPrecaching(
+	TEXT("r.PSOPrecacheValidation"),
+	GValidatePrecaching,
+	TEXT("Check if runtime used PSOs are correctly precached and track information per pass type, vertex factory and cache hit state (default 0)."),
+	ECVF_ReadOnly
+);
+
+int32 PSOCollectorStats::IsPrecachingValidationEnabled()
+{
+	return GValidatePrecaching;
+}
+
+void PSOCollectorStats::AddPipelineStateToCache(const FGraphicsPipelineStateInitializer& PSOInitializer, uint32 MeshPassType, const FVertexFactoryType* VertexFactoryType)
+{
+	if (!GValidatePrecaching)
+	{
+		return;
+	}
+
+	uint64 PrecachePSOHash = RHIComputePrecachePSOHash(PSOInitializer);
+
+	FScopeLock Lock(&FullPSOPrecacheLock);
+
+	static bool bFirstTime = true;
+	if (bFirstTime)
+	{
+		PSOCollectorStats::FullPSOPrecacheStats.Empty();
+		bFirstTime = false;
+	}
+
+	Experimental::FHashElementId TableId = FullPSOPrecacheMap.FindId(PrecachePSOHash);
+	if (!TableId.IsValid())
+	{
+		TableId = FullPSOPrecacheMap.FindOrAddId(PrecachePSOHash, FShaderStateUsage());
+	}
+
+	FShaderStateUsage& Value = FullPSOPrecacheMap.GetByElementId(TableId).Value;
+	if (!Value.bPrecached)
+	{
+		check(!Value.bUsed);
+		FullPSOPrecacheStats.PrecacheData.UpdateStats(MeshPassType, VertexFactoryType);
+		Value.bPrecached = true;
+	}
+}
+
+EPSOPrecacheResult PSOCollectorStats::CheckPipelineStateInCache(const FGraphicsPipelineStateInitializer& PSOInitializer, uint32 MeshPassType, const FVertexFactoryType* VertexFactoryType)
+{
+	if (!GValidatePrecaching)
+	{
+		return EPSOPrecacheResult::Unknown;
+	}
+
+	bool bUpdateStats = false;
+	bool bTracked = true;
+	if (MeshPassType < FPSOCollectorCreateManager::MaxPSOCollectorCount && VertexFactoryType)
+	{
+		const EShadingPath ShadingPath = FSceneInterface::GetShadingPath(GMaxRHIFeatureLevel);
+		bool bCollectPSOs = FPSOCollectorCreateManager::GetCreateFunction(ShadingPath, MeshPassType) != nullptr;
+		bTracked = bCollectPSOs && VertexFactoryType->SupportsPSOPrecaching();
+
+		bUpdateStats = true;
+	}
+	EPSOPrecacheResult PrecacheResult = EPSOPrecacheResult::NotSupported;
+
+	// only search the cache if it's tracked
+	if (bTracked)
+	{
+		PrecacheResult = PipelineStateCache::CheckPipelineStateInCache(PSOInitializer);
+	}
+
+	if (bUpdateStats)
+	{
+		uint64 PrecachePSOHash = RHIComputePrecachePSOHash(PSOInitializer);
+
+		FScopeLock Lock(&FullPSOPrecacheLock);
+
+		Experimental::FHashElementId TableId = FullPSOPrecacheMap.FindId(PrecachePSOHash);
+		if (!TableId.IsValid())
+		{
+			TableId = FullPSOPrecacheMap.FindOrAddId(PrecachePSOHash, FShaderStateUsage());
+		}
+
+		FShaderStateUsage& Value = FullPSOPrecacheMap.GetByElementId(TableId).Value;
+		if (!Value.bUsed)
+		{
+			FullPSOPrecacheStats.UsageData.UpdateStats(MeshPassType, VertexFactoryType);
+			if (!bTracked)
+			{
+				FullPSOPrecacheStats.UntrackedData.UpdateStats(MeshPassType, VertexFactoryType);
+			}
+			else if (!Value.bPrecached)
+			{
+				check(PrecacheResult == EPSOPrecacheResult::Missed);
+				FullPSOPrecacheStats.MissData.UpdateStats(MeshPassType, VertexFactoryType);
+			}
+			else
+			{
+				check(PrecacheResult == EPSOPrecacheResult::Active || PrecacheResult == EPSOPrecacheResult::Complete);
+				FullPSOPrecacheStats.HitData.UpdateStats(MeshPassType, VertexFactoryType);
+			}
+
+			Value.bUsed = true;
+		}
+	}
+
+	return PrecacheResult;
+}
+
+#endif // PSO_PRECACHING_VALIDATE
\ No newline at end of file
diff --git a/Engine/Source/Runtime/Engine/Private/PrimitiveSceneProxy.cpp b/Engine/Source/Runtime/Engine/Private/PrimitiveSceneProxy.cpp
--- a/Engine/Source/Runtime/Engine/Private/PrimitiveSceneProxy.cpp
+++ b/Engine/Source/Runtime/Engine/Private/PrimitiveSceneProxy.cpp
@@ -147,6 +147,22 @@ bool SupportsNaniteRendering(const FVertexFactory* RESTRICT VertexFactory, const
 	return false;
 }
 
+
+bool SupportsNaniteRendering(const FVertexFactoryType* RESTRICT VertexFactoryType, const class FMaterial& Material, ERHIFeatureLevel::Type FeatureLevel)
+{
+	if (FeatureLevel >= ERHIFeatureLevel::SM5 && VertexFactoryType->SupportsNaniteRendering())
+	{
+		const FMaterialShaderMap* ShaderMap = Material.GetGameThreadShaderMap();
+		return (Material.IsUsedWithNanite() || Material.IsSpecialEngineMaterial()) &&
+			Nanite::IsSupportedBlendMode(Material.GetBlendMode()) &&
+			Nanite::IsSupportedMaterialDomain(Material.GetMaterialDomain()) &&
+			(Nanite::IsWorldPositionOffsetSupported() || !ShaderMap->UsesWorldPositionOffset());
+	}
+
+	return false;
+}
+
+
 static bool VertexDeformationOutputsVelocity()
 {
 	static const auto CVarVelocityOutputPass = IConsoleManager::Get().FindConsoleVariable(TEXT("r.VelocityOutputPass"));
diff --git a/Engine/Source/Runtime/Engine/Private/SceneTexturesConfig.cpp b/Engine/Source/Runtime/Engine/Private/SceneTexturesConfig.cpp
--- a/Engine/Source/Runtime/Engine/Private/SceneTexturesConfig.cpp
+++ b/Engine/Source/Runtime/Engine/Private/SceneTexturesConfig.cpp
@@ -2,6 +2,7 @@
 
 #include "SceneTexturesConfig.h"
 #include "RenderGraphResources.h"
+#include "ShaderCompiler.h"
 
 FSceneTexturesConfig FSceneTexturesConfig::GlobalInstance;
 
@@ -21,4 +22,276 @@ FSceneTextureShaderParameters GetSceneTextureShaderParameters(TRDGUniformBufferR
 	FSceneTextureShaderParameters Parameters;
 	Parameters.MobileSceneTextures = UniformBuffer;
 	return Parameters;
+}
+
+static EPixelFormat GetDefaultMobileSceneColorLowPrecisionFormat()
+{
+	return (GEngine && GEngine->XRSystem.IsValid() && GEngine->StereoRenderingDevice.IsValid() && GEngine->StereoRenderingDevice->IsStandaloneStereoOnlyDevice()) ? PF_R8G8B8A8 : PF_B8G8R8A8;
+}
+
+static EPixelFormat GetMobileSceneColorFormat(bool bRequiresAlphaChannel)
+{
+	EPixelFormat DefaultColorFormat;
+	const bool bUseLowPrecisionFormat = !IsMobileHDR() || !GSupportsRenderTargetFormat_PF_FloatRGBA;
+	if (bUseLowPrecisionFormat)
+	{
+		DefaultColorFormat = GetDefaultMobileSceneColorLowPrecisionFormat();
+	}
+	else
+	{
+		DefaultColorFormat = bRequiresAlphaChannel ? PF_FloatRGBA : PF_FloatR11G11B10;
+	}
+
+	check(GPixelFormats[DefaultColorFormat].Supported);
+
+	EPixelFormat Format = DefaultColorFormat;
+	static const auto CVar = IConsoleManager::Get().FindTConsoleVariableDataInt(TEXT("r.Mobile.SceneColorFormat"));
+	int32 MobileSceneColor = CVar->GetValueOnRenderThread();
+	switch (MobileSceneColor)
+	{
+	case 1:
+		Format = PF_FloatRGBA; break;
+	case 2:
+		Format = PF_FloatR11G11B10; break;
+	case 3:
+		if (bUseLowPrecisionFormat)
+		{
+			// if bUseLowPrecisionFormat, DefaultColorFormat already contains the value of GetDefaultMobileSceneColorLowPrecisionFormat
+			checkSlow(DefaultColorFormat == GetDefaultMobileSceneColorLowPrecisionFormat());
+			Format = DefaultColorFormat;
+		}
+		else
+		{
+			Format = GetDefaultMobileSceneColorLowPrecisionFormat();
+		}
+		break;
+	default:
+		break;
+	}
+
+	return GPixelFormats[Format].Supported ? Format : DefaultColorFormat;
+}
+
+static EPixelFormat GetSceneColorFormat(bool bRequiresAlphaChannel)
+{
+	EPixelFormat Format = PF_FloatRGBA;
+
+	static const auto CVar = IConsoleManager::Get().FindTConsoleVariableDataInt(TEXT("r.SceneColorFormat"));
+
+	switch (CVar->GetValueOnAnyThread())
+	{
+	case 0:
+		Format = PF_R8G8B8A8; break;
+	case 1:
+		Format = PF_A2B10G10R10; break;
+	case 2:
+		Format = PF_FloatR11G11B10; break;
+	case 3:
+		Format = PF_FloatRGB; break;
+	case 4:
+		// default
+		break;
+	case 5:
+		Format = PF_A32B32G32R32F; break;
+	}
+
+	// Fallback in case the scene color selected isn't supported.
+	if (!GPixelFormats[Format].Supported)
+	{
+		Format = PF_FloatRGBA;
+	}
+
+	if (bRequiresAlphaChannel)
+	{
+		Format = PF_FloatRGBA;
+	}
+
+	return Format;
+}
+
+static void GetSceneColorFormatAndCreateFlags(ERHIFeatureLevel::Type FeatureLevel, bool bRequiresAlphaChannel, ETextureCreateFlags ExtraSceneColorCreateFlags, uint32 NumSamples, bool bMemorylessMSAA, EPixelFormat& SceneColorFormat, ETextureCreateFlags& SceneColorCreateFlags)
+{
+	EShadingPath ShadingPath = FSceneInterface::GetShadingPath(FeatureLevel);
+	switch (ShadingPath)
+	{
+	case EShadingPath::Deferred:
+	{
+		SceneColorFormat = GetSceneColorFormat(bRequiresAlphaChannel);
+		break;
+	}
+
+	case EShadingPath::Mobile:
+	{
+		SceneColorFormat = GetMobileSceneColorFormat(bRequiresAlphaChannel);
+		break;
+	}
+
+	default:
+		checkNoEntry();
+	}
+
+	const bool bIsMobilePlatform = ShadingPath == EShadingPath::Mobile;
+
+	SceneColorCreateFlags = TexCreate_RenderTargetable | TexCreate_ShaderResource | ExtraSceneColorCreateFlags;
+	const ETextureCreateFlags sRGBFlag = (bIsMobilePlatform && IsMobileColorsRGB()) ? TexCreate_SRGB : TexCreate_None;
+	if (FeatureLevel >= ERHIFeatureLevel::SM5 && NumSamples == 1)
+	{
+		SceneColorCreateFlags |= TexCreate_UAV;
+	}
+	if (NumSamples > 1 && bMemorylessMSAA)
+	{
+		SceneColorCreateFlags |= TexCreate_Memoryless;
+	}
+	SceneColorCreateFlags |= sRGBFlag;
+}
+
+static ETextureCreateFlags GetSceneDepthStencilCreateFlags(uint32 NumSamples, bool bKeepDepthContent, bool bMemorylessMSAA, ETextureCreateFlags ExtraSceneDepthCreateFlags)
+{
+	ETextureCreateFlags DepthCreateFlags = TexCreate_DepthStencilTargetable | TexCreate_ShaderResource | TexCreate_InputAttachmentRead | ExtraSceneDepthCreateFlags;
+	if (!bKeepDepthContent || (NumSamples > 1 && bMemorylessMSAA))
+	{
+		DepthCreateFlags = TexCreate_Memoryless;
+	}
+	if (NumSamples == 1 && GRHISupportsDepthUAV)
+	{
+		DepthCreateFlags |= TexCreate_UAV;
+	}
+	return DepthCreateFlags;
+}
+
+static uint32 GetEditorPrimitiveNumSamples(ERHIFeatureLevel::Type FeatureLevel)
+{
+	uint32 SampleCount = 1;
+
+	if (FeatureLevel >= ERHIFeatureLevel::SM5 && GRHISupportsMSAADepthSampleAccess)
+	{
+		static const auto CVar = IConsoleManager::Get().FindTConsoleVariableDataInt(TEXT("r.MSAA.CompositingSampleCount"));
+
+		SampleCount = CVar->GetValueOnAnyThread();
+
+		if (SampleCount <= 1)
+		{
+			SampleCount = 1;
+		}
+		else if (SampleCount <= 2)
+		{
+			SampleCount = 2;
+		}
+		else if (SampleCount <= 4)
+		{
+			SampleCount = 4;
+		}
+		else
+		{
+			SampleCount = 8;
+		}
+	}
+
+	return SampleCount;
+}
+
+void FSceneTexturesConfig::Init(const FSceneTexturesConfigInitSettings& InitSettings)
+{
+	FeatureLevel			= InitSettings.FeatureLevel;
+	ShadingPath				= FSceneInterface::GetShadingPath(FeatureLevel);
+	ShaderPlatform			= GetFeatureLevelShaderPlatform(FeatureLevel);
+	Extent					= InitSettings.Extent;
+	NumSamples				= GetDefaultMSAACount(FeatureLevel, GDynamicRHI->RHIGetPlatformTextureMaxSampleCount());
+	EditorPrimitiveNumSamples = GetEditorPrimitiveNumSamples(FeatureLevel);
+	ColorFormat				= PF_Unknown;
+	ColorClearValue			= FClearValueBinding::Black;
+	DepthClearValue			= FClearValueBinding::DepthFar;
+	bRequireMultiView		= InitSettings.bRequireMultiView;
+	bIsUsingGBuffers		= IsUsingGBuffers(ShaderPlatform);
+	bSupportsXRTargetManagerDepthAlloc = InitSettings.bSupportsXRTargetManagerDepthAlloc;
+
+	GetSceneColorFormatAndCreateFlags(FeatureLevel, InitSettings.bRequiresAlphaChannel, InitSettings.ExtraSceneColorCreateFlags, NumSamples, bMemorylessMSAA, ColorFormat, ColorCreateFlags);
+	DepthCreateFlags = GetSceneDepthStencilCreateFlags(NumSamples, bKeepDepthContent, bMemorylessMSAA, InitSettings.ExtraSceneDepthCreateFlags);
+
+	if (bIsUsingGBuffers)
+	{
+		GBufferParams = FShaderCompileUtilities::FetchGBufferParamsRuntime(ShaderPlatform);
+
+		// GBuffer configuration information is expensive to compute, the results are cached between runs.
+		struct FGBufferBindingCache
+		{
+			FGBufferParams GBufferParams;
+			FGBufferBinding GBufferA;
+			FGBufferBinding GBufferB;
+			FGBufferBinding GBufferC;
+			FGBufferBinding GBufferD;
+			FGBufferBinding GBufferE;
+			FGBufferBinding GBufferVelocity;
+			bool bInitialized = false;
+		};
+		static FGBufferBindingCache BindingCache;
+
+		if (!BindingCache.bInitialized || BindingCache.GBufferParams != GBufferParams)
+		{
+			const FGBufferInfo GBufferInfo = FetchFullGBufferInfo(GBufferParams);
+
+			BindingCache.GBufferParams = GBufferParams;
+			BindingCache.GBufferA = FindGBufferBindingByName(GBufferInfo, TEXT("GBufferA"));
+			BindingCache.GBufferB = FindGBufferBindingByName(GBufferInfo, TEXT("GBufferB"));
+			BindingCache.GBufferC = FindGBufferBindingByName(GBufferInfo, TEXT("GBufferC"));
+			BindingCache.GBufferD = FindGBufferBindingByName(GBufferInfo, TEXT("GBufferD"));
+			BindingCache.GBufferE = FindGBufferBindingByName(GBufferInfo, TEXT("GBufferE"));
+			BindingCache.GBufferVelocity = FindGBufferBindingByName(GBufferInfo, TEXT("Velocity"));
+			BindingCache.bInitialized = true;
+		}
+
+		GBufferA = BindingCache.GBufferA;
+		GBufferB = BindingCache.GBufferB;
+		GBufferC = BindingCache.GBufferC;
+		GBufferD = BindingCache.GBufferD;
+		GBufferE = BindingCache.GBufferE;
+		GBufferVelocity = BindingCache.GBufferVelocity;
+	}
+}
+
+uint32 FSceneTexturesConfig::GetGBufferRenderTargetsInfo(FGraphicsPipelineRenderTargetsInfo& RenderTargetsInfo) const 
+{
+	// Assume 1 sample for now
+	RenderTargetsInfo.NumSamples = 1;
+
+	uint32 RenderTargetCount = 0;
+
+	// All configurations use scene color in the first slot.
+	RenderTargetsInfo.RenderTargetFormats[RenderTargetCount] = ColorFormat;
+	RenderTargetsInfo.RenderTargetFlags[RenderTargetCount++] = ColorCreateFlags;
+
+	// Setup the other render targets
+	if (bIsUsingGBuffers)
+	{
+		const auto CheckGBufferBinding = [&RenderTargetsInfo, &RenderTargetCount](const FGBufferBinding& GBufferBinding)
+		{
+			if (GBufferBinding.Index > 0)
+			{
+				RenderTargetsInfo.RenderTargetFormats[GBufferBinding.Index] = GBufferBinding.Format;
+				RenderTargetsInfo.RenderTargetFlags[GBufferBinding.Index] = GBufferBinding.Flags;
+				RenderTargetCount = FMath::Max(RenderTargetCount, (uint32)GBufferBinding.Index + 1);
+			}
+		};
+		CheckGBufferBinding(GBufferA);
+		CheckGBufferBinding(GBufferB);
+		CheckGBufferBinding(GBufferC);
+		CheckGBufferBinding(GBufferD);
+		CheckGBufferBinding(GBufferE);
+		CheckGBufferBinding(GBufferVelocity);
+	}
+	// Forward shading path. Simple forward shading does not use velocity.
+	else if (IsUsingBasePassVelocity(ShaderPlatform))
+	{
+		RenderTargetsInfo.RenderTargetFormats[RenderTargetCount] = GBufferVelocity.Format;
+		RenderTargetsInfo.RenderTargetFlags[RenderTargetCount++] = GBufferVelocity.Flags;
+	}
+
+	// Store final number of render targets
+	RenderTargetsInfo.RenderTargetsEnabled = RenderTargetCount;
+
+	// Precache TODO: other flags
+	RenderTargetsInfo.MultiViewCount = 0;// RenderTargets.MultiViewCount;
+	RenderTargetsInfo.bHasFragmentDensityAttachment = false; // RenderTargets.ShadingRateTexture != nullptr;
+
+	return RenderTargetCount;
 }
\ No newline at end of file
diff --git a/Engine/Source/Runtime/Engine/Private/ShaderCompiler/ShaderCompiler.cpp b/Engine/Source/Runtime/Engine/Private/ShaderCompiler/ShaderCompiler.cpp
--- a/Engine/Source/Runtime/Engine/Private/ShaderCompiler/ShaderCompiler.cpp
+++ b/Engine/Source/Runtime/Engine/Private/ShaderCompiler/ShaderCompiler.cpp
@@ -6741,27 +6741,40 @@ void VerifyGlobalShaders(EShaderPlatform Platform, const ITargetPlatform* Target
 		int32 PermutationCountToCompile = 0;
 		for (int32 PermutationId = 0; PermutationId < GlobalShaderType->GetPermutationCount(); PermutationId++)
 		{
-			if (GlobalShaderType->ShouldCompilePermutation(Platform, PermutationId, PermutationFlags)
-				&& (!GlobalShaderMap->HasShader(GlobalShaderType, PermutationId) || (OutdatedShaderTypes && OutdatedShaderTypes->Contains(GlobalShaderType))))
+			if (GlobalShaderType->ShouldCompilePermutation(Platform, PermutationId, PermutationFlags))
 			{
-				if (bErrorOnMissing)
+				bool bOutdated = OutdatedShaderTypes && OutdatedShaderTypes->Contains(GlobalShaderType);
+				TShaderRef<FShader> GlobalShader = GlobalShaderMap->GetShader(GlobalShaderType, PermutationId);
+				if (bOutdated || !GlobalShader.IsValid())
 				{
-					UE_LOG(LogShaders, Fatal, TEXT("Missing global shader %s's permutation %i, Please make sure cooking was successful."),
-						GlobalShaderType->GetName(), PermutationId);
-				}
+					if (bErrorOnMissing)
+					{
+						UE_LOG(LogShaders, Fatal, TEXT("Missing global shader %s's permutation %i, Please make sure cooking was successful."),
+							GlobalShaderType->GetName(), PermutationId);
+					}
 
-				if (OutdatedShaderTypes)
+					if (OutdatedShaderTypes)
+					{
+						// Remove old shader, if it exists
+						GlobalShaderMap->RemoveShaderTypePermutaion(GlobalShaderType, PermutationId);
+					}
+
+					// Compile this global shader type.
+					FGlobalShaderTypeCompiler::BeginCompileShader(GlobalShaderType, PermutationId, Platform, PermutationFlags, GlobalShaderJobs);
+					//TShaderTypePermutation<const FShaderType> ShaderTypePermutation(GlobalShaderType, PermutationId);
+					//check(!SharedShaderJobs.Find(ShaderTypePermutation));
+					//SharedShaderJobs.Add(ShaderTypePermutation, Job);
+					PermutationCountToCompile++;
+				}
+				else if (GlobalShader.IsValid())
 				{
-					// Remove old shader, if it exists
-					GlobalShaderMap->RemoveShaderTypePermutaion(GlobalShaderType, PermutationId);
+					// try and precache if compute
+					if (GlobalShader->GetFrequency() == SF_Compute)
+					{
+						FRHIComputeShader* RHIComputeShader = GlobalShader.GetComputeShader();
+						PipelineStateCache::PrecacheComputePipelineState(RHIComputeShader);
+					}
 				}
-
-				// Compile this global shader type.
-				FGlobalShaderTypeCompiler::BeginCompileShader(GlobalShaderType, PermutationId, Platform, PermutationFlags, GlobalShaderJobs);
-				//TShaderTypePermutation<const FShaderType> ShaderTypePermutation(GlobalShaderType, PermutationId);
-				//check(!SharedShaderJobs.Find(ShaderTypePermutation));
-				//SharedShaderJobs.Add(ShaderTypePermutation, Job);
-				PermutationCountToCompile++;
 			}
 		}
 
diff --git a/Engine/Source/Runtime/Engine/Private/ShaderCompiler/ShaderGenerationUtil.cpp b/Engine/Source/Runtime/Engine/Private/ShaderCompiler/ShaderGenerationUtil.cpp
--- a/Engine/Source/Runtime/Engine/Private/ShaderCompiler/ShaderGenerationUtil.cpp
+++ b/Engine/Source/Runtime/Engine/Private/ShaderCompiler/ShaderGenerationUtil.cpp
@@ -637,6 +637,9 @@ static int32 GetBufferNumBits(EGBufferType Format, int32 Channel)
 	case GBT_Invalid:
 		check(0);
 		break;
+	case GBT_Unorm_16_16:
+		Ret = 16;
+		break;
 	case GBT_Unorm_8_8_8_8:
 		Ret = 8;
 		break;
@@ -652,6 +655,9 @@ static int32 GetBufferNumBits(EGBufferType Format, int32 Channel)
 		Ret = Sizes[Channel];
 	}
 	break;
+	case GBT_Unorm_16_16_16_16:
+		Ret = 16;
+		break;
 	case GBT_Float_16_16:
 		Ret = 16;
 		break;
@@ -675,6 +681,9 @@ static int32 GetTargetNumChannels(EGBufferType Type)
 	case GBT_Invalid:
 		Ret = 0;
 		break;
+	case GBT_Unorm_16_16:
+		Ret = 2;
+		break;
 	case GBT_Unorm_8_8_8_8:
 		Ret = 4;
 		break;
@@ -684,6 +693,9 @@ static int32 GetTargetNumChannels(EGBufferType Type)
 	case GBT_Unorm_10_10_10_2:
 		Ret = 4;
 		break;
+	case GBT_Unorm_16_16_16_16:
+		Ret = 4;
+		break;
 	case GBT_Float_16_16:
 		Ret = 2;
 		break;
@@ -2130,6 +2142,8 @@ FGBufferParams FShaderCompileUtilities::FetchGBufferParamsRuntime(EShaderPlatfor
 	// This code should match TBasePassPS
 
 	FGBufferParams Ret = {};
+	Ret.ShaderPlatform = Platform;
+
 	static const auto CVar = IConsoleManager::Get().FindTConsoleVariableDataInt(TEXT("r.AllowStaticLighting"));
 	Ret.bHasPrecShadowFactor = (CVar ? (CVar->GetValueOnAnyThread() != 0) : 1);
 
diff --git a/Engine/Source/Runtime/Engine/Private/TextureRenderTarget2D.cpp b/Engine/Source/Runtime/Engine/Private/TextureRenderTarget2D.cpp
--- a/Engine/Source/Runtime/Engine/Private/TextureRenderTarget2D.cpp
+++ b/Engine/Source/Runtime/Engine/Private/TextureRenderTarget2D.cpp
@@ -572,6 +572,29 @@ void FTextureRenderTarget2DResource::ClampSize(int32 MaxSizeX,int32 MaxSizeY)
 	}	
 }
 
+ETextureCreateFlags FTextureRenderTarget2DResource::GetCreateFlags()
+{
+	// Create the RHI texture. Only one mip is used and the texture is targetable for resolve.
+	ETextureCreateFlags TexCreateFlags = Owner->IsSRGB() ? ETextureCreateFlags::SRGB : ETextureCreateFlags::None;
+	TexCreateFlags |= Owner->bGPUSharedFlag ? ETextureCreateFlags::Shared : ETextureCreateFlags::None;
+	
+	if (Owner->bAutoGenerateMips)
+	{
+		TexCreateFlags |= ETextureCreateFlags::GenerateMipCapable;
+		if (FGenerateMips::WillFormatSupportCompute(Format))
+		{
+			TexCreateFlags |= ETextureCreateFlags::UAV;
+		}
+	}
+
+	if (Owner->bCanCreateUAV)
+	{
+		TexCreateFlags |= ETextureCreateFlags::UAV;
+	}
+
+	return TexCreateFlags | ETextureCreateFlags::RenderTargetable | ETextureCreateFlags::ShaderResource;
+}
+
 /**
  * Initializes the RHI render target resources used by this resource.
  * Called when the resource is initialized, or when reseting all RHI resources.
@@ -581,31 +604,15 @@ void FTextureRenderTarget2DResource::InitDynamicRHI()
 {
 	if( TargetSizeX > 0 && TargetSizeY > 0 )
 	{
-		// Create the RHI texture. Only one mip is used and the texture is targetable for resolve.
-		ETextureCreateFlags TexCreateFlags = Owner->IsSRGB() ? ETextureCreateFlags::SRGB : ETextureCreateFlags::None;
-		TexCreateFlags |= Owner->bGPUSharedFlag ? ETextureCreateFlags::Shared : ETextureCreateFlags::None;
 		FString ResourceName = Owner->GetName();
-
-		if (Owner->bAutoGenerateMips)
-		{
-			TexCreateFlags |= ETextureCreateFlags::GenerateMipCapable;
-			if (FGenerateMips::WillFormatSupportCompute(Format))
-			{
-				TexCreateFlags |= ETextureCreateFlags::UAV;
-			}
-		}
-
-		if (Owner->bCanCreateUAV)
-		{
-			TexCreateFlags |= ETextureCreateFlags::UAV;
-		}
+		ETextureCreateFlags TexCreateFlags = GetCreateFlags();
 
 		FRHITextureCreateDesc Desc =
 			FRHITextureCreateDesc::Create2D(*ResourceName)
 			.SetExtent(Owner->SizeX, Owner->SizeY)
 			.SetFormat(Format)
 			.SetNumMips(Owner->GetNumMips())
-			.SetFlags(TexCreateFlags | ETextureCreateFlags::RenderTargetable | ETextureCreateFlags::ShaderResource)
+			.SetFlags(TexCreateFlags)
 			.SetInitialState(ERHIAccess::SRVMask)
 			.SetClearValue(FClearValueBinding(ClearColor));
 
diff --git a/Engine/Source/Runtime/Engine/Public/MaterialShared.h b/Engine/Source/Runtime/Engine/Public/MaterialShared.h
--- a/Engine/Source/Runtime/Engine/Public/MaterialShared.h
+++ b/Engine/Source/Runtime/Engine/Public/MaterialShared.h
@@ -29,6 +29,7 @@
 #include "VirtualTexturing.h"
 #include "Templates/UnrealTemplate.h"
 #include "ShaderCompilerCore.h"
+#include "PSOPrecache.h"
 
 #include "Shader/Preshader.h"
 
@@ -1200,7 +1201,7 @@ public:
 	
 	FMaterialShaderMap* AcquireFinalizedClone();
 	FMaterialShaderMap* GetFinalizedClone() const;
-	
+
 
 	// ShaderMap interface
 	TShaderRef<FShader> GetShader(FShaderType* ShaderType, int32 PermutationId = 0) const
@@ -1251,6 +1252,11 @@ public:
 	 */
 	bool IsComplete(const FMaterial* Material, bool bSilent);
 
+	/**
+	 * Collect all possible PSO's  which can be used with this material shader map for given parameters - PSOs will be async precached
+	 */
+	FGraphEventArray CollectPSOs(ERHIFeatureLevel::Type InFeatureLevel, const FMaterial* Material, const TConstArrayView<const FVertexFactoryType*>& VertexFactoryTypes, const FPSOPrecacheParams& PreCacheParams);
+
 #if WITH_EDITOR
 	/** Attempts to load missing shaders from memory. */
 	void LoadMissingShadersFromMemory(const FMaterial* Material);
@@ -1626,7 +1632,7 @@ extern ENGINE_API bool CanConnectMaterialValueTypes(const uint32 InputType, cons
  */
 class FMaterial
 {
-public:
+public:	
 #if UE_CHECK_FMATERIAL_LIFETIME
 	ENGINE_API uint32 AddRef() const;
 	ENGINE_API uint32 Release() const;
@@ -1785,6 +1791,11 @@ public:
 	ENGINE_API void CacheGivenTypes(EShaderPlatform ShaderPlatform, const TArray<const FVertexFactoryType*>& VFTypes, const TArray<const FShaderPipelineType*>& PipelineTypes, const TArray<const FShaderType*>& ShaderTypes);
 #endif
 
+	/**
+	 * Collect all possible PSO's  which can be used with this material for given parameters - PSOs will be async precached
+	 */
+	ENGINE_API FGraphEventArray CollectPSOs(ERHIFeatureLevel::Type InFeatureLevel, const TConstArrayView<const FVertexFactoryType*>& VertexFactoryTypes, const FPSOPrecacheParams& PreCacheParams);
+
 	/**
 	 * Should the shader for this material with the given platform, shader type and vertex 
 	 * factory type combination be compiled
@@ -2076,6 +2087,7 @@ public:
 	ENGINE_API bool MaterialMayModifyMeshPosition() const;
 
 	/** Get the runtime virtual texture output attribute mask for the material. */
+	ENGINE_API uint8 GetRuntimeVirtualTextureOutputAttibuteMask_GameThread() const;
 	ENGINE_API uint8 GetRuntimeVirtualTextureOutputAttibuteMask_RenderThread() const;
 
 	ENGINE_API bool MaterialUsesAnisotropy_GameThread() const;
@@ -2375,6 +2387,33 @@ private:
 	uint32 bOwnerBeginDestroyed : 1;
 #endif // UE_CHECK_FMATERIAL_LIFETIME
 
+	/**
+	 * Keep track which VertexFactory/Precache params have already been requested for this material
+	 */
+	struct FPrecacheVertexTypeWithParams
+	{
+		const FVertexFactoryType* VertexFactoryType;
+		FPSOPrecacheParams PrecachePSOParams;
+
+		bool operator==(const FPrecacheVertexTypeWithParams& Other) const
+		{
+			return VertexFactoryType == Other.VertexFactoryType &&
+				PrecachePSOParams == Other.PrecachePSOParams;
+		}
+
+		bool operator!=(const FPrecacheVertexTypeWithParams& rhs) const
+		{
+			return !(*this == rhs);
+		}
+
+		friend uint32 GetTypeHash(const FPrecacheVertexTypeWithParams& Params)
+		{
+			return HashCombine(GetTypeHash(Params.PrecachePSOParams), GetTypeHash(Params.VertexFactoryType));
+		}
+	};
+	FRWLock PrecachePSOVFLock;
+	TArray<FPrecacheVertexTypeWithParams> PrecachedPSOVertexFactories;
+
 	/**
 	* Compiles this material for Platform, storing the result in OutShaderMap if the compile was synchronous
 	*/
diff --git a/Engine/Source/Runtime/Engine/Public/PSOPrecache.h b/Engine/Source/Runtime/Engine/Public/PSOPrecache.h
new file mode 100644
--- /dev/null
+++ b/Engine/Source/Runtime/Engine/Public/PSOPrecache.h
@@ -0,0 +1,276 @@
+// Copyright Epic Games, Inc. All Rights Reserved.
+
+/*=============================================================================
+	MeshPassProcessorManager.h
+=============================================================================*/
+
+#pragma once
+
+#include "RHIDefinitions.h"
+#include "Engine/EngineTypes.h"
+#include "SceneUtils.h"
+#include "PipelineStateCache.h"
+#include "Experimental/Containers/RobinHoodHashTable.h"
+
+struct FSceneTexturesConfig;
+class FMaterial;
+class FVertexFactoryType;
+class FGraphicsPipelineStateInitializer;
+
+#define PSO_PRECACHING_VALIDATE (!UE_BUILD_SHIPPING && !UE_BUILD_TEST)
+
+/**
+ * Parameters which are needed to collect all possible PSOs used by the PSO collectors
+ */
+struct FPSOPrecacheParams
+{
+	FPSOPrecacheParams()
+	{
+		PrimitiveType = (uint8)PT_TriangleList;
+		bRenderInMainPass = true;
+		bRenderInDepthPass = true;
+		bStaticLighting = true;
+		bCastShadow = true;
+		bAffectDynamicIndirectLighting = true;
+		bReverseCulling = false;
+		bCastShadowAsTwoSided = false;
+		bForceLODModel = false;
+		Mobility = (uint8)EComponentMobility::Static;
+		StencilWriteMask = (uint8)EStencilMask::SM_Default;
+		Unused = 0;
+	}
+
+	bool operator==(const FPSOPrecacheParams& Other) const
+	{
+		return Data == Other.Data;
+	}
+
+	bool operator!=(const FPSOPrecacheParams& rhs) const
+	{
+		return !(*this == rhs);
+	}
+
+	friend uint32 GetTypeHash(const FPSOPrecacheParams& Params)
+	{
+		return uint32(Params.Data);
+	}
+
+	void SetMobility(EComponentMobility::Type InMobility)
+	{
+		Mobility = (uint8)InMobility;
+	}
+
+	TEnumAsByte<EComponentMobility::Type> GetMobility() const
+	{
+		return TEnumAsByte<EComponentMobility::Type>((EComponentMobility::Type)Mobility);
+	}
+
+	bool IsMoveable() const
+	{
+		return Mobility == EComponentMobility::Movable || Mobility == EComponentMobility::Stationary;
+	}
+
+	void SetStencilWriteMask(EStencilMask InStencilMask)
+	{
+		StencilWriteMask = (uint8)InStencilMask;
+	}
+
+	TEnumAsByte<EStencilMask> GetStencilWriteMask() const
+	{
+		return TEnumAsByte<EStencilMask>((EStencilMask)StencilWriteMask);
+	}
+
+	union
+	{
+		struct
+		{
+			uint32 PrimitiveType : 6;
+
+			uint32 bRenderInMainPass : 1;
+			uint32 bRenderInDepthPass : 1;
+			uint32 bStaticLighting : 1;
+			uint32 bCastShadow : 1;
+
+			uint32 bAffectDynamicIndirectLighting : 1;
+			uint32 bReverseCulling : 1;
+			uint32 bCastShadowAsTwoSided : 1;
+			uint32 bForceLODModel : 1;
+
+			uint32 Mobility : 4;
+			uint32 StencilWriteMask : 4;
+
+			uint32 Unused : 10;
+		};
+		uint32 Data;
+	};
+};
+
+/**
+ * Wrapper class around the initializer to collect some extra validation data during PSO collection on the different collectors
+ */
+struct FPSOPrecacheData
+{
+	FGraphicsPipelineStateInitializer PSOInitializer;
+
+#if PSO_PRECACHING_VALIDATE
+	uint32 MeshPassType = 0;
+	const FVertexFactoryType* VertexFactoryType = nullptr;
+#endif // PSO_PRECACHING_VALIDATE
+};
+
+/** 
+ * Interface class implemented by the mesh pass processor to collect all possible PSOs
+ */
+class IPSOCollector
+{
+public:
+	
+	// Collect all PSO for given material, vertex factory & params
+	virtual void CollectPSOInitializers(const FSceneTexturesConfig& SceneTexturesConfig, const FMaterial& Material, const FVertexFactoryType* VertexFactoryType, const FPSOPrecacheParams& PreCacheParams, TArray<FPSOPrecacheData>& PSOInitializers) = 0;
+};
+
+/**
+ * Try and create PSOs for all the given initializers and return an optional array of graph events of async compiling PSOs
+ */
+extern FGraphEventArray PrecachePSOs(const TArray<FPSOPrecacheData>& PSOInitializers);
+
+/**
+ * Predeclared IPSOCollector create function
+ */
+typedef IPSOCollector* (*PSOCollectorCreateFunction)(ERHIFeatureLevel::Type InFeatureLevel);
+
+/**
+ * Manages all create functions of the globally defined IPSOCollectors
+ */
+class ENGINE_API FPSOCollectorCreateManager
+{
+public:
+
+	constexpr static uint32 MaxPSOCollectorCount = 32;
+
+	static PSOCollectorCreateFunction GetCreateFunction(EShadingPath ShadingPath, uint32 Index)
+	{
+		check(Index < MaxPSOCollectorCount);
+		uint32 ShadingPathIdx = (uint32)ShadingPath;
+		return JumpTable[ShadingPathIdx][Index];
+	}
+
+private:
+
+	// Have to used fixed size array instead of TArray because of order of initialization of static member variables
+	static PSOCollectorCreateFunction JumpTable[(uint32)EShadingPath::Num][MaxPSOCollectorCount];
+	friend class FRegisterPSOCollectorCreateFunction;
+};
+
+/**
+ * Helper class used to register/unregister the IPSOCollector to the manager at static startup time
+ */
+class ENGINE_API FRegisterPSOCollectorCreateFunction
+{
+public:
+	FRegisterPSOCollectorCreateFunction(PSOCollectorCreateFunction InCreateFunction, EShadingPath InShadingPath, uint32 InIndex)
+		: ShadingPath(InShadingPath)
+		, Index(InIndex)
+	{
+		check(InIndex < FPSOCollectorCreateManager::MaxPSOCollectorCount);
+		uint32 ShadingPathIdx = (uint32)ShadingPath;
+		FPSOCollectorCreateManager::JumpTable[ShadingPathIdx][Index] = InCreateFunction;
+	}
+
+	~FRegisterPSOCollectorCreateFunction()
+	{
+		uint32 ShadingPathIdx = (uint32)ShadingPath;
+		FPSOCollectorCreateManager::JumpTable[ShadingPathIdx][Index] = nullptr;
+	}
+
+private:
+	EShadingPath ShadingPath;
+	uint32 Index;
+};
+
+#if PSO_PRECACHING_VALIDATE
+
+/**
+ * Wraps all PSO precache state data
+ */
+namespace PSOCollectorStats
+{
+	extern ENGINE_API int32 IsPrecachingValidationEnabled();
+
+
+	using VertexFactoryCountTableType = Experimental::TRobinHoodHashMap<const FVertexFactoryType*, uint32>;
+	struct FShaderStateUsage
+	{
+		bool bPrecached = false;
+		bool bUsed = false;
+	};
+
+	/**
+	 * Track precache stats in total, per mesh pass type and per vertex factory type
+	 */
+	struct FPrecacheUsageData
+	{
+		void Empty()
+		{
+			FMemory::Memzero(PerMeshPassCount, FPSOCollectorCreateManager::MaxPSOCollectorCount * sizeof(uint32));
+		}
+
+		void UpdateStats(uint32 MeshPassType, const FVertexFactoryType* VFType)
+		{
+			Count++;
+			PerMeshPassCount[MeshPassType]++;
+
+			Experimental::FHashElementId TableId = PerVertexFactoryCount.FindId(VFType);
+			if (!TableId.IsValid())
+			{
+				TableId = PerVertexFactoryCount.FindOrAddId(VFType, uint32());
+			}
+			uint32& Value = PerVertexFactoryCount.GetByElementId(TableId).Value;
+			Value++;
+		}
+
+		uint32 Count;
+		uint32 PerMeshPassCount[FPSOCollectorCreateManager::MaxPSOCollectorCount];
+		VertexFactoryCountTableType PerVertexFactoryCount;
+	};
+
+	/**
+	 * Collect stats for different cache hit states
+	 */
+	struct FPrecacheStats
+	{
+		void Empty()
+		{
+			PrecacheData.Empty();
+			UsageData.Empty();
+			HitData.Empty();
+			MissData.Empty();
+			UntrackedData.Empty();
+		}
+
+		FPrecacheUsageData PrecacheData;		//< PSOs which have been precached
+		FPrecacheUsageData UsageData;			//< PSOs which are used during rendering 
+		FPrecacheUsageData HitData;				//< PSOs which are used during rendering and have been successfully precached
+		FPrecacheUsageData MissData;			//< PSOs which are used during rendering and have not been precached (but should have been)
+		FPrecacheUsageData UntrackedData;		//< PSOs which are used during rendering but are currently not precached because for example the MeshPassProcessor or VertexFactory type don't support PSO precaching yet
+	};
+
+	// Stats for the complete PSO
+	static FCriticalSection FullPSOPrecacheLock;
+	static FPrecacheStats FullPSOPrecacheStats;
+
+	// Cache stats on hash - TODO: should use FGraphicsPipelineStateInitializer as key and use FGraphicsPSOInitializerKeyFuncs for search
+	static Experimental::TRobinHoodHashMap<uint64, FShaderStateUsage> FullPSOPrecacheMap;
+
+	/**
+	 * Add PSO initializer to cache for validation & state tracking
+	 */
+	extern ENGINE_API void AddPipelineStateToCache(const FGraphicsPipelineStateInitializer& PSOInitializer, uint32 MeshPassType, const FVertexFactoryType* VertexFactoryType);
+
+	/**
+	 * Is the requested PSO initializer precached (only PSO relevant data is checked)
+	 */
+	extern ENGINE_API EPSOPrecacheResult CheckPipelineStateInCache(const FGraphicsPipelineStateInitializer& PSOInitializer, uint32 MeshPassType, const FVertexFactoryType* VertexFactoryType);
+}
+
+#endif // PSO_PRECACHING_VALIDATE
diff --git a/Engine/Source/Runtime/Engine/Public/PrimitiveSceneProxy.h b/Engine/Source/Runtime/Engine/Public/PrimitiveSceneProxy.h
--- a/Engine/Source/Runtime/Engine/Public/PrimitiveSceneProxy.h
+++ b/Engine/Source/Runtime/Engine/Public/PrimitiveSceneProxy.h
@@ -1480,3 +1480,5 @@ ENGINE_API extern bool SupportsCachingMeshDrawCommands(const FMeshBatch& MeshBat
 ENGINE_API extern bool SupportsNaniteRendering(const FVertexFactory* RESTRICT VertexFactory, const FPrimitiveSceneProxy* RESTRICT PrimitiveSceneProxy);
 
 ENGINE_API extern bool SupportsNaniteRendering(const FVertexFactory* RESTRICT VertexFactory, const FPrimitiveSceneProxy* RESTRICT PrimitiveSceneProxy, const class FMaterialRenderProxy* MaterialRenderProxy, ERHIFeatureLevel::Type FeatureLevel);
+
+ENGINE_API extern bool SupportsNaniteRendering(const FVertexFactoryType* RESTRICT VertexFactoryType, const class FMaterial& Material, ERHIFeatureLevel::Type FeatureLevel);
diff --git a/Engine/Source/Runtime/Engine/Public/SceneTexturesConfig.h b/Engine/Source/Runtime/Engine/Public/SceneTexturesConfig.h
--- a/Engine/Source/Runtime/Engine/Public/SceneTexturesConfig.h
+++ b/Engine/Source/Runtime/Engine/Public/SceneTexturesConfig.h
@@ -82,6 +82,17 @@ enum class ESceneTextureExtracts : uint32
 	All = Depth | CustomDepth
 };
 
+struct FSceneTexturesConfigInitSettings
+{
+	ERHIFeatureLevel::Type FeatureLevel = ERHIFeatureLevel::Num;
+	FIntPoint Extent = FIntPoint::ZeroValue;
+	bool bRequireMultiView = false;
+	bool bRequiresAlphaChannel = false;
+	bool bSupportsXRTargetManagerDepthAlloc = false;
+	ETextureCreateFlags ExtraSceneColorCreateFlags = ETextureCreateFlags::None;
+	ETextureCreateFlags ExtraSceneDepthCreateFlags = ETextureCreateFlags::None;
+};
+
 /** Struct containing the scene texture configuration used to create scene textures.  Use InitializeViewFamily to initialize the
  *  SceneTexturesConfig structure in the FViewFamilyInfo.  A global singleton instance is maintained manually with static Set / Get
  *  functions, but will soon be deprecated, in preference of using the structure from the FViewFamilyInfo.
@@ -110,6 +121,9 @@ struct ENGINE_API FSceneTexturesConfig
 		, bSupportsXRTargetManagerDepthAlloc{}
 	{}
 
+	void Init(const FSceneTexturesConfigInitSettings& InitSettings);
+	uint32 GetGBufferRenderTargetsInfo(FGraphicsPipelineRenderTargetsInfo& RenderTargetsInfo) const;
+
 	FORCEINLINE bool IsValid() const
 	{
 		return ShadingPath != EShadingPath::Num;
@@ -141,6 +155,10 @@ struct ENGINE_API FSceneTexturesConfig
 	// Pixel format to use when creating scene color.
 	EPixelFormat ColorFormat = PF_Unknown;
 
+	// Create flags when creating scene color / depth textures
+	ETextureCreateFlags ColorCreateFlags = ETextureCreateFlags::None;
+	ETextureCreateFlags DepthCreateFlags = ETextureCreateFlags::None;
+
 	// Optimized clear values to use for color / depth textures.
 	FClearValueBinding ColorClearValue = FClearValueBinding::Black;
 	FClearValueBinding DepthClearValue = FClearValueBinding::DepthFar;
diff --git a/Engine/Source/Runtime/Engine/Public/TextureResource.h b/Engine/Source/Runtime/Engine/Public/TextureResource.h
--- a/Engine/Source/Runtime/Engine/Public/TextureResource.h
+++ b/Engine/Source/Runtime/Engine/Public/TextureResource.h
@@ -402,6 +402,13 @@ public:
 		return ClearColor;
 	}
 
+	FORCEINLINE EPixelFormat GetFormat() const
+	{
+		return Format;
+	}
+
+	ETextureCreateFlags GetCreateFlags();
+
 	// FTextureRenderTargetResource interface
 
 	/** 
diff --git a/Engine/Source/Runtime/RHI/Private/DynamicRHI.cpp b/Engine/Source/Runtime/RHI/Private/DynamicRHI.cpp
--- a/Engine/Source/Runtime/RHI/Private/DynamicRHI.cpp
+++ b/Engine/Source/Runtime/RHI/Private/DynamicRHI.cpp
@@ -544,6 +544,118 @@ uint64 FDynamicRHI::RHIGetMinimumAlignmentForBufferBackedSRV(EPixelFormat Format
 	return 1;
 }
 
+uint64 FDynamicRHI::RHIComputePrecachePSOHash(const FGraphicsPipelineStateInitializer& Initializer)
+{
+	// When compute precache PSO hash we assume a valid state precache PSO hash is already provided
+	checkf(Initializer.StatePrecachePSOHash != 0, TEXT("Initializer should have a valid state precache PSO hash set when computing the full initializer PSO hash"));
+	
+	// All members which are not part of the state objects
+	struct FNonStateHashKey
+	{
+		uint64							StatePrecachePSOHash;
+
+		EPrimitiveType					PrimitiveType;
+		uint32							RenderTargetsEnabled;
+		FGraphicsPipelineStateInitializer::TRenderTargetFormats	RenderTargetFormats;
+		FGraphicsPipelineStateInitializer::TRenderTargetFlags RenderTargetFlags;
+		EPixelFormat					DepthStencilTargetFormat;
+		ETextureCreateFlags				DepthStencilTargetFlag;
+		ERenderTargetLoadAction			DepthTargetLoadAction;
+		ERenderTargetStoreAction		DepthTargetStoreAction;
+		ERenderTargetLoadAction			StencilTargetLoadAction;
+		ERenderTargetStoreAction		StencilTargetStoreAction;
+		FExclusiveDepthStencil			DepthStencilAccess;
+		uint16							NumSamples;
+		ESubpassHint					SubpassHint;
+		uint8							SubpassIndex;
+		EConservativeRasterization		ConservativeRasterization;
+		bool							bDepthBounds;
+		uint8							MultiViewCount;
+		bool							bHasFragmentDensityAttachment;
+		EVRSShadingRate					ShadingRate;
+	} HashKey;
+
+	FMemory::Memzero(&HashKey, sizeof(FNonStateHashKey));
+
+	HashKey.StatePrecachePSOHash		= Initializer.StatePrecachePSOHash;
+
+	HashKey.PrimitiveType				= Initializer.PrimitiveType;
+	HashKey.RenderTargetsEnabled		= Initializer.RenderTargetsEnabled;
+	HashKey.RenderTargetFormats			= Initializer.RenderTargetFormats;
+	HashKey.RenderTargetFlags			= Initializer.RenderTargetFlags;
+	HashKey.DepthStencilTargetFormat	= Initializer.DepthStencilTargetFormat;
+	HashKey.DepthStencilTargetFlag		= Initializer.DepthStencilTargetFlag;
+	HashKey.DepthTargetLoadAction		= Initializer.DepthTargetLoadAction;
+	HashKey.DepthTargetStoreAction		= Initializer.DepthTargetStoreAction;	
+	HashKey.StencilTargetLoadAction		= Initializer.StencilTargetLoadAction;
+	HashKey.StencilTargetStoreAction	= Initializer.StencilTargetStoreAction;
+	HashKey.DepthStencilAccess			= Initializer.DepthStencilAccess;
+	HashKey.NumSamples					= Initializer.NumSamples;
+	HashKey.SubpassHint					= Initializer.SubpassHint;
+	HashKey.SubpassIndex				= Initializer.SubpassIndex;
+	HashKey.ConservativeRasterization	= Initializer.ConservativeRasterization;
+	HashKey.bDepthBounds				= Initializer.bDepthBounds;
+	HashKey.MultiViewCount				= Initializer.MultiViewCount;
+	HashKey.bHasFragmentDensityAttachment = Initializer.bHasFragmentDensityAttachment;
+	HashKey.ShadingRate					= Initializer.ShadingRate;
+
+	return CityHash64((const char*)&HashKey, sizeof(FNonStateHashKey));
+}
+
+bool FDynamicRHI::RHIMatchPrecachePSOInitializers(const FGraphicsPipelineStateInitializer& LHS, const FGraphicsPipelineStateInitializer& RHS)
+{
+	// first check non pointer objects
+	if (LHS.ImmutableSamplerState != LHS.ImmutableSamplerState ||
+		LHS.PrimitiveType != RHS.PrimitiveType ||
+		LHS.bDepthBounds != RHS.bDepthBounds ||
+		LHS.MultiViewCount != RHS.MultiViewCount ||
+		LHS.ShadingRate != RHS.ShadingRate ||
+		LHS.bHasFragmentDensityAttachment != RHS.bHasFragmentDensityAttachment ||
+		LHS.RenderTargetsEnabled != RHS.RenderTargetsEnabled ||
+		LHS.RenderTargetFormats != RHS.RenderTargetFormats ||
+		LHS.RenderTargetFlags != RHS.RenderTargetFlags ||
+		LHS.DepthStencilTargetFormat != RHS.DepthStencilTargetFormat ||
+		LHS.DepthStencilTargetFlag != RHS.DepthStencilTargetFlag ||
+		LHS.DepthTargetLoadAction != RHS.DepthTargetLoadAction ||
+		LHS.DepthTargetStoreAction != RHS.DepthTargetStoreAction ||
+		LHS.StencilTargetLoadAction != RHS.StencilTargetLoadAction ||
+		LHS.StencilTargetStoreAction != RHS.StencilTargetStoreAction ||
+		LHS.DepthStencilAccess != RHS.DepthStencilAccess ||
+		LHS.NumSamples != RHS.NumSamples ||
+		LHS.SubpassHint != RHS.SubpassHint ||
+		LHS.SubpassIndex != RHS.SubpassIndex ||
+		LHS.ConservativeRasterization != RHS.ConservativeRasterization)
+	{
+		return false;
+	}
+
+	// check the RHI shaders (pointer check for shaders should be fine)
+	if (LHS.BoundShaderState.VertexShaderRHI != RHS.BoundShaderState.VertexShaderRHI ||
+		LHS.BoundShaderState.PixelShaderRHI != RHS.BoundShaderState.PixelShaderRHI ||
+		LHS.BoundShaderState.GetMeshShader() != RHS.BoundShaderState.GetMeshShader() ||
+		LHS.BoundShaderState.GetAmplificationShader() != RHS.BoundShaderState.GetAmplificationShader() ||
+		LHS.BoundShaderState.GetGeometryShader() != RHS.BoundShaderState.GetGeometryShader())
+	{
+		return false;
+	}
+
+	// Full compare the of the vertex declaration
+	if (!MatchRHIState<FRHIVertexDeclaration, FVertexDeclarationElementList>(LHS.BoundShaderState.VertexDeclarationRHI, RHS.BoundShaderState.VertexDeclarationRHI))
+	{
+		return false;
+	}
+
+	// Check actual state content (each initializer can have it's own state and not going through a factory)
+	if (!MatchRHIState<FRHIBlendState, FBlendStateInitializerRHI>(LHS.BlendState, RHS.BlendState) ||
+		!MatchRHIState<FRHIRasterizerState, FRasterizerStateInitializerRHI>(LHS.RasterizerState, RHS.RasterizerState) ||
+		!MatchRHIState<FRHIDepthStencilState, FDepthStencilStateInitializerRHI>(LHS.DepthStencilState, RHS.DepthStencilState))
+	{
+		return false;
+	}
+
+	return true;
+}
+
 FDefaultRHIRenderQueryPool::FDefaultRHIRenderQueryPool(ERenderQueryType InQueryType, FDynamicRHI* InDynamicRHI, uint32 InNumQueries)
 	: DynamicRHI(InDynamicRHI)
 	, QueryType(InQueryType)
diff --git a/Engine/Source/Runtime/RHI/Private/PipelineStateCache.cpp b/Engine/Source/Runtime/RHI/Private/PipelineStateCache.cpp
--- a/Engine/Source/Runtime/RHI/Private/PipelineStateCache.cpp
+++ b/Engine/Source/Runtime/RHI/Private/PipelineStateCache.cpp
@@ -91,6 +91,15 @@ static TAutoConsoleVariable<int32> CVarRTPSOCacheSize(
 );
 #endif // RHI_RAYTRACING
 
+int32 GPSOPrecaching = 0;
+static TAutoConsoleVariable<int32> CVarPSOPrecaching(
+	TEXT("r.PSOPrecaching"),
+	GPSOPrecaching,
+	TEXT("0 to Disable PSOs precaching\n")
+	TEXT("1 to Enable PSO precaching\n"),
+	ECVF_ReadOnly
+);
+
 extern void DumpPipelineCacheStats();
 
 static FAutoConsoleCommand DumpPipelineCmd(
@@ -279,6 +288,7 @@ public:
 	virtual ~FPipelineState() 
 	{
 		check(IsComplete());
+		check(!WaitCompletion());
 		check(!PrecompileTask.IsValid());
 	}
 
@@ -527,12 +537,12 @@ int32 FindRayTracingMissShaderIndex(FRayTracingPipelineState* Pipeline, FRHIRayT
 	return INDEX_NONE;
 }
 
-void SetGraphicsPipelineState(FRHICommandList& RHICmdList, const FGraphicsPipelineStateInitializer& Initializer, uint32 StencilRef, EApplyRendertargetOption ApplyFlags, bool bApplyAdditionalState)
+void SetGraphicsPipelineState(FRHICommandList& RHICmdList, const FGraphicsPipelineStateInitializer& Initializer, uint32 StencilRef, EApplyRendertargetOption ApplyFlags, bool bApplyAdditionalState, EPSOPrecacheResult PSOPrecacheResult)
 {
 #if PLATFORM_USE_FALLBACK_PSO
 	RHICmdList.SetGraphicsPipelineState(Initializer, StencilRef, bApplyAdditionalState);
 #else
-	FGraphicsPipelineState* PipelineState = PipelineStateCache::GetAndOrCreateGraphicsPipelineState(RHICmdList, Initializer, ApplyFlags);
+	FGraphicsPipelineState* PipelineState = PipelineStateCache::GetAndOrCreateGraphicsPipelineState(RHICmdList, Initializer, ApplyFlags, PSOPrecacheResult);
 	if (PipelineState && (PipelineState->RHIPipeline || !Initializer.bFromPSOFileCache))
 	{
 #if PIPELINESTATECACHE_VERIFYTHREADSAFE
@@ -859,6 +869,144 @@ private:
 
 };
 
+class FPrecacheGraphicsPipelineCache
+{
+public:
+
+	~FPrecacheGraphicsPipelineCache()
+	{
+		// Wait for all precache tasks to finished
+		WaitTasksComplete();
+	}
+
+	bool Contains(const FGraphicsPipelineStateInitializer& Initializer)
+	{
+		FRWScopeLock ReadLock(PrecachePSOsRWLock, SLT_ReadOnly);
+		return PrecachedPSOInitializers.Contains(Initializer);
+	}
+
+	void Add(const FGraphicsPipelineStateInitializer& Initializer, FGraphicsPipelineState* PipelineState, bool bAsync)
+	{		
+		// First add to set of precached PSOs initializers
+		{
+			FRWScopeLock WriteLock(PrecachePSOsRWLock, SLT_Write);
+			check(!PrecachedPSOInitializers.Contains(Initializer));
+			PrecachedPSOInitializers.Add(Initializer);
+		}
+
+		// Add to active set of precaching PSOs if it's created async
+		if (bAsync)
+		{
+			// Only need the hash of the PSO here
+			uint64 PrecachePSOHash = RHIComputePrecachePSOHash(Initializer);
+
+			FPrecacheTask PrecacheTask;
+			PrecacheTask.Initializer = Initializer;
+			PrecacheTask.PSO = PipelineState;
+
+			FRWScopeLock WriteLock(ActivePrecachingTasksRWLock, SLT_Write);
+			check(!ActivePrecachingTasks.Contains(PrecachePSOHash));
+			ActivePrecachingTasks.Add(PrecachePSOHash, PrecacheTask);
+		}
+	}
+
+	void WaitTasksComplete()
+	{
+		FRWScopeLock WriteLock(ActivePrecachingTasksRWLock, SLT_Write);
+
+		for (auto Iterator = ActivePrecachingTasks.CreateIterator(); Iterator; ++Iterator)
+		{
+			FPrecacheTask& PrecacheTask = Iterator->Value;
+			check(PrecacheTask.PSO);
+			PrecacheTask.PSO->WaitCompletion();
+			delete PrecacheTask.PSO;
+		}
+		ActivePrecachingTasks.Empty();
+	}
+
+	bool IsPrecaching(const FGraphicsPipelineStateInitializer& Initializer)
+	{
+		uint64 PrecachePSOHash = RHIComputePrecachePSOHash(Initializer);
+
+		FRWScopeLock ReadLock(ActivePrecachingTasksRWLock, SLT_ReadOnly);
+		FPrecacheTask* FindResult = ActivePrecachingTasks.Find(PrecachePSOHash);
+		if (FindResult)
+		{
+			// Check if not complete yet
+			return !(FindResult->PSO->IsComplete());
+		}
+
+		// not found then assume it's not precaching
+		return false;
+	}
+
+	bool IsPrecaching()
+	{
+		FRWScopeLock ReadLock(ActivePrecachingTasksRWLock, SLT_ReadOnly);
+		return !ActivePrecachingTasks.IsEmpty();
+	}
+
+	void PrecacheFinished(const FGraphicsPipelineStateInitializer& Initializer)
+	{
+		uint64 PrecachePSOHash = RHIComputePrecachePSOHash(Initializer);
+
+		FRWScopeLock WriteLock(ActivePrecachingTasksRWLock, SLT_Write);
+		FPrecacheTask* FindResult = ActivePrecachingTasks.Find(PrecachePSOHash);
+		check(FindResult);
+		verify(ActivePrecachingTasks.Remove(PrecachePSOHash) == 1);
+		PrecachedPSOs.Add(FindResult->PSO);
+	}
+
+	void ProcessDelayedCleanup()
+	{
+		FRWScopeLock WriteLock(ActivePrecachingTasksRWLock, SLT_Write);
+		for (int32 Index = 0; Index < PrecachedPSOs.Num(); ++Index)
+		{
+			FGraphicsPipelineState* GraphicsPSO = PrecachedPSOs[Index];
+			if (GraphicsPSO->IsComplete())
+			{
+				// This is needed to cleanup the members - bit strange because it's complete already
+				verify(!GraphicsPSO->WaitCompletion());
+
+				delete GraphicsPSO;
+				PrecachedPSOs.RemoveAtSwap(Index);
+				Index--;
+			}			
+		}
+
+	}
+
+private:
+
+	struct FGraphicsPSOInitializerKeyFuncs : DefaultKeyFuncs<FGraphicsPipelineStateInitializer>
+	{
+		static FORCEINLINE bool Matches(KeyInitType A, KeyInitType B)
+		{
+			return RHIMatchPrecachePSOInitializers(A, B);
+		}
+
+		static FORCEINLINE uint32 GetKeyHash(KeyInitType Key)
+		{
+			return RHIComputePrecachePSOHash(Key);
+		}
+	};
+
+	// Cache of all precached PSOs - need correct compare to make sure we precache all of them
+	FRWLock PrecachePSOsRWLock;
+	TSet<FGraphicsPipelineStateInitializer, FGraphicsPSOInitializerKeyFuncs> PrecachedPSOInitializers;
+
+	// Current active background precache operations - full PSO hash can be used for fast retrieval
+	// (on hash clash it might report that a PSO is still compiling)
+	FRWLock ActivePrecachingTasksRWLock;
+	struct FPrecacheTask
+	{
+		FGraphicsPipelineStateInitializer Initializer;
+		FGraphicsPipelineState* PSO = nullptr;
+	};
+	TMap<uint64, FPrecacheTask> ActivePrecachingTasks;
+	TArray<FGraphicsPipelineState*> PrecachedPSOs;
+};
+
 // Typed caches for compute and graphics
 typedef TDiscardableKeyValueCache< FRHIComputeShader*, FComputePipelineState*> FComputePipelineCache;
 typedef TSharedPipelineStateCache<FGraphicsPipelineStateInitializer, FGraphicsPipelineState*> FGraphicsPipelineCache;
@@ -866,6 +1014,7 @@ typedef TSharedPipelineStateCache<FGraphicsPipelineStateInitializer, FGraphicsPi
 // These are the actual caches for both pipelines
 FComputePipelineCache GComputePipelineCache;
 FGraphicsPipelineCache GGraphicsPipelineCache;
+FPrecacheGraphicsPipelineCache GPrecacheGraphicsPipelineCache;
 
 FAutoConsoleTaskPriority CPrio_FCompilePipelineStateTask(
 	TEXT("TaskGraph.TaskPriorities.CompilePipelineStateTask"),
@@ -1088,11 +1237,13 @@ class FCompilePipelineStateTask
 public:
 	FPipelineState* Pipeline;
 	FGraphicsPipelineStateInitializer Initializer;
+	EPSOPrecacheResult PSOPreCacheResult;
 
 	// InInitializer is only used for non-compute tasks, a default can just be used otherwise
-	FCompilePipelineStateTask(FPipelineState* InPipeline, const FGraphicsPipelineStateInitializer& InInitializer)
+	FCompilePipelineStateTask(FPipelineState* InPipeline, const FGraphicsPipelineStateInitializer& InInitializer, EPSOPrecacheResult InPSOPreCacheResult)
 		: Pipeline(InPipeline)
 		, Initializer(InInitializer)
+		, PSOPreCacheResult(InPSOPreCacheResult)
 	{
 		if(Initializer.bFromPSOFileCache)
 		{
@@ -1180,6 +1331,18 @@ public:
 		}
 		else
 		{
+			const TCHAR* PSOPrecacheResultString = nullptr;
+			switch (PSOPreCacheResult)
+			{
+			case EPSOPrecacheResult::Unknown:			PSOPrecacheResultString = TEXT("PSOPrecache: Unknown"); break;
+			case EPSOPrecacheResult::Active:			PSOPrecacheResultString = TEXT("PSOPrecache: Precaching"); break;
+			case EPSOPrecacheResult::Complete:			PSOPrecacheResultString = TEXT("PSOPrecache: Precached"); break;
+			case EPSOPrecacheResult::Missed:			PSOPrecacheResultString = TEXT("PSOPrecache: Missed"); break;
+			case EPSOPrecacheResult::NotSupported:		PSOPrecacheResultString = TEXT("PSOPrecache: Precache Untracked"); break;
+			case EPSOPrecacheResult::Untracked:			PSOPrecacheResultString = TEXT("PSOPrecache: Untracked"); break;
+			}
+			TRACE_CPUPROFILER_EVENT_SCOPE_TEXT(PSOPrecacheResultString);
+
 			if (GRHISupportsMeshShadersTier0)
 			{
 				if (!Initializer.BoundShaderState.VertexShaderRHI && !Initializer.BoundShaderState.GetMeshShader())
@@ -1207,6 +1370,12 @@ public:
 				HandlePipelineCreationFailure(Initializer);
 			}
 
+			// Mark as finished when it's a precaching job
+			if (PSOPreCacheResult == EPSOPrecacheResult::Active)
+			{
+				GPrecacheGraphicsPipelineCache.PrecacheFinished(Initializer);
+			}
+
 			if (Initializer.BoundShaderState.GetMeshShader())
 			{
 				Initializer.BoundShaderState.GetMeshShader()->Release();
@@ -1287,6 +1456,8 @@ void PipelineStateCache::FlushResources()
 	GGraphicsPipelineCache.ConsolidateThreadedCaches();
 	GGraphicsPipelineCache.ProcessDelayedCleanup();
 
+	GPrecacheGraphicsPipelineCache.ProcessDelayedCleanup();
+
 	{
 		int32 NumMissesThisFrame = GraphicsPipelineCacheMisses.Load(EMemoryOrder::Relaxed);
 		int32 NumMissesLastFrame = GraphicsPipelineCacheMissesHistory.Num() >= 2 ? GraphicsPipelineCacheMissesHistory[1] : 0;
@@ -1412,7 +1583,7 @@ FComputePipelineState* PipelineStateCache::GetAndOrCreateComputePipelineState(FR
 		// create a compilation task, or just do it now...
 		if (DoAsyncCompile)
 		{
-			OutCachedState->CompletionEvent = TGraphTask<FCompilePipelineStateTask>::CreateTask().ConstructAndDispatchWhenReady(OutCachedState, FGraphicsPipelineStateInitializer());
+			OutCachedState->CompletionEvent = TGraphTask<FCompilePipelineStateTask>::CreateTask().ConstructAndDispatchWhenReady(OutCachedState, FGraphicsPipelineStateInitializer(), EPSOPrecacheResult::Untracked);
 			RHICmdList.AddDispatchPrerequisite(OutCachedState->CompletionEvent);
 		}
 		else
@@ -1454,7 +1625,7 @@ FComputePipelineState* PipelineStateCache::GetAndOrCreateComputePipelineState(FR
 			// create a compilation task, or just do it now...
 			if (DoAsyncCompile)
 			{
-				PipelineState->CompletionEvent = TGraphTask<FCompilePipelineStateTask>::CreateTask().ConstructAndDispatchWhenReady(PipelineState.Get(), FGraphicsPipelineStateInitializer());
+				PipelineState->CompletionEvent = TGraphTask<FCompilePipelineStateTask>::CreateTask().ConstructAndDispatchWhenReady(PipelineState.Get(), FGraphicsPipelineStateInitializer(), EPSOPrecacheResult::Untracked);
 				RHICmdList.QueueAsyncPipelineStateCompile(PipelineState->CompletionEvent);
 			}
 			else
@@ -1750,7 +1921,64 @@ inline void ValidateGraphicsPipelineStateInitializer(const FGraphicsPipelineStat
 	check(Initializer.DepthStencilState && Initializer.BlendState && Initializer.RasterizerState);
 }
 
-FGraphicsPipelineState* PipelineStateCache::GetAndOrCreateGraphicsPipelineState(FRHICommandList& RHICmdList, const FGraphicsPipelineStateInitializer& Initializer, EApplyRendertargetOption ApplyFlags)
+static FGraphEventRef CreateGraphicsPipelineState(const FGraphicsPipelineStateInitializer& Initializer, EPSOPrecacheResult PSOPrecacheResult, bool bDoAsyncCompile, FGraphicsPipelineState* CachedState)
+{
+	FGraphEventRef GraphEvent;
+
+	// create a compilation task, or just do it now...
+	if (bDoAsyncCompile)
+	{
+		// Use normal task graph for non-precompile jobs.
+		if (GPSOPrecompileThreadPoolSize <= 0)
+		{
+			CachedState->CompletionEvent = TGraphTask<FCompilePipelineStateTask>::CreateTask().ConstructAndDispatchWhenReady(CachedState, Initializer, PSOPrecacheResult);
+			GraphEvent = CachedState->CompletionEvent;
+		}
+		else
+		{
+			// Here, PSO precompiles use a separate thread pool.
+			// Note that we do not add precompile tasks as cmdlist prerequisites.
+			GraphEvent = FGraphEvent::CreateGraphEvent();
+			TUniquePtr<FCompilePipelineStateTask> ThreadPoolTask = MakeUnique<FCompilePipelineStateTask>(CachedState, Initializer, PSOPrecacheResult);
+			CachedState->CompletionEvent = GraphEvent;
+			uint64 StartTime = FPlatformTime::Cycles64();
+#if	PSO_TRACK_CACHE_STATS
+			StatsStartPrecompile();
+#endif
+			CachedState->PrecompileTask = MakeUnique<FAsyncTaskLimitedReschedule<FPSOPrecompileTask>>(
+				[ThreadPoolTask = MoveTemp(ThreadPoolTask)]()
+			{
+				ThreadPoolTask->CompilePSO();
+			}
+			,
+				[GraphEvent, StartTime]()
+			{
+				GraphEvent->DispatchSubsequents();
+#if PSO_TRACK_CACHE_STATS
+				StatsEndPrecompile(FPlatformTime::Cycles64() - StartTime);
+#endif
+			}
+			);
+			CachedState->PrecompileTask->StartBackgroundTask(&GPSOThreadPool.Get(), EQueuedWorkPriority::Normal);
+		}
+	}
+	else
+	{
+		CachedState->RHIPipeline = RHICreateGraphicsPipelineState(Initializer);
+		if (CachedState->RHIPipeline)
+		{
+			CachedState->SortKey = CachedState->RHIPipeline->GetSortKey();
+		}
+		else
+		{
+			HandlePipelineCreationFailure(Initializer);
+		}
+	}
+
+	return GraphEvent;
+}
+
+FGraphicsPipelineState* PipelineStateCache::GetAndOrCreateGraphicsPipelineState(FRHICommandList& RHICmdList, const FGraphicsPipelineStateInitializer& Initializer, EApplyRendertargetOption ApplyFlags, EPSOPrecacheResult PSOPrecacheResult)
 {
 	LLM_SCOPE(ELLMTag::PSO);
 	ValidateGraphicsPipelineStateInitializer(Initializer);
@@ -1804,60 +2032,19 @@ FGraphicsPipelineState* PipelineStateCache::GetAndOrCreateGraphicsPipelineState(
 		// create new graphics state
 		OutCachedState = new FGraphicsPipelineState();
 		OutCachedState->Stats = FPipelineFileCacheManager::RegisterPSOStats(GetTypeHash(Initializer));
-
+		
 		if (!Initializer.bFromPSOFileCache)
 		{
 			GraphicsPipelineCacheMisses++;
 		}
 
-		// create a compilation task, or just do it now...
-		if (DoAsyncCompile)
-		{
-			// Use normal task graph for non-precompile jobs.
-			if (!Initializer.bFromPSOFileCache || GPSOPrecompileThreadPoolSize <=0)
-			{
-				OutCachedState->CompletionEvent = TGraphTask<FCompilePipelineStateTask>::CreateTask().ConstructAndDispatchWhenReady(OutCachedState, Initializer);
-				RHICmdList.AddDispatchPrerequisite(OutCachedState->CompletionEvent);
-			}
-			else
-			{
-				// Here, PSO precompiles use a separate thread pool.
-				// Note that we do not add precompile tasks as cmdlist prerequisites.
-				FGraphEventRef CompletedEventGraph = FGraphEvent::CreateGraphEvent();
-				TUniquePtr<FCompilePipelineStateTask> ThreadPoolTask = MakeUnique<FCompilePipelineStateTask>(OutCachedState, Initializer);
-				OutCachedState->CompletionEvent = CompletedEventGraph;
-				uint64 StartTime = FPlatformTime::Cycles64();
-#if	PSO_TRACK_CACHE_STATS
-				StatsStartPrecompile();
-#endif
-				OutCachedState->PrecompileTask = MakeUnique<FAsyncTaskLimitedReschedule<FPSOPrecompileTask>>(
-					[ThreadPoolTask = MoveTemp(ThreadPoolTask)]()
-					{
-						ThreadPoolTask->CompilePSO();
-				}
-				,
-					[CompletedEventGraph, StartTime]()
-					{
-						CompletedEventGraph->DispatchSubsequents();
-	#if PSO_TRACK_CACHE_STATS
-						StatsEndPrecompile(FPlatformTime::Cycles64() - StartTime);
-	#endif
-					}
-				 );
-				OutCachedState->PrecompileTask->StartBackgroundTask(&GPSOThreadPool.Get(), EQueuedWorkPriority::Normal);
-			}
-		}
-		else
+		FGraphEventRef GraphEvent = CreateGraphicsPipelineState(Initializer, PSOPrecacheResult, DoAsyncCompile, OutCachedState);
+
+		// Use normal task graph for non-precompile jobs.
+		if (DoAsyncCompile && (!Initializer.bFromPSOFileCache || GPSOPrecompileThreadPoolSize <= 0))
 		{
-			OutCachedState->RHIPipeline = RHICreateGraphicsPipelineState(Initializer);
-			if (OutCachedState->RHIPipeline)
-			{
-				OutCachedState->SortKey = OutCachedState->RHIPipeline->GetSortKey();
-			}
-			else
-			{
-				HandlePipelineCreationFailure(Initializer);
-			}
+			check(GraphEvent);
+			RHICmdList.AddDispatchPrerequisite(GraphEvent);
 		}
 
 		GGraphicsPipelineCache.Add(Initializer, OutCachedState);
@@ -1897,6 +2084,147 @@ FGraphicsPipelineState* PipelineStateCache::FindGraphicsPipelineState(const FGra
 	return (PipelineState && PipelineState->IsComplete()) ? PipelineState : nullptr;
 }
 
+bool PipelineStateCache::IsPSOPrecachingEnabled()
+{
+	return GPSOPrecaching != 0;
+}
+
+FGraphEventRef PipelineStateCache::PrecacheComputePipelineState(FRHIComputeShader* ComputeShader)
+{
+	FGraphEventRef GraphEvent = nullptr;
+	if (!IsPSOPrecachingEnabled())
+	{
+		return GraphEvent;
+	}
+
+	if (ComputeShader == nullptr)
+	{
+		return GraphEvent;
+	}
+
+	FComputePipelineState* CachedState = nullptr;
+	if (GComputePipelineCache.Find(ComputeShader, CachedState))
+	{
+		return GraphEvent;
+	}
+
+	// create new graphics state
+	CachedState = new FComputePipelineState(ComputeShader);
+
+	// create a compilation task, or just do it now...
+	static bool bDoAsyncCompile = FApp::ShouldUseThreadingForPerformance();
+	if (bDoAsyncCompile)
+	{
+		// Use normal task graph for non-precompile jobs.
+		if (GPSOPrecompileThreadPoolSize <= 0)
+		{
+			CachedState->CompletionEvent = TGraphTask<FCompilePipelineStateTask>::CreateTask().ConstructAndDispatchWhenReady(CachedState, FGraphicsPipelineStateInitializer(), EPSOPrecacheResult::Active);
+			GraphEvent = CachedState->CompletionEvent;
+		}
+		else
+		{
+			// Here, PSO precompiles use a separate thread pool.
+			// Note that we do not add precompile tasks as cmdlist prerequisites.
+			GraphEvent = FGraphEvent::CreateGraphEvent();
+			TUniquePtr<FCompilePipelineStateTask> ThreadPoolTask = MakeUnique<FCompilePipelineStateTask>(CachedState, FGraphicsPipelineStateInitializer(), EPSOPrecacheResult::Active);
+			CachedState->CompletionEvent = GraphEvent;
+			uint64 StartTime = FPlatformTime::Cycles64();
+#if	PSO_TRACK_CACHE_STATS
+			StatsStartPrecompile();
+#endif
+			CachedState->PrecompileTask = MakeUnique<FAsyncTaskLimitedReschedule<FPSOPrecompileTask>>(
+				[ThreadPoolTask = MoveTemp(ThreadPoolTask)]()
+			{
+				ThreadPoolTask->CompilePSO();
+			}
+			,
+				[GraphEvent, StartTime]()
+			{
+				GraphEvent->DispatchSubsequents();
+#if PSO_TRACK_CACHE_STATS
+				StatsEndPrecompile(FPlatformTime::Cycles64() - StartTime);
+#endif
+			}
+			);
+			CachedState->PrecompileTask->StartBackgroundTask(&GPSOThreadPool.Get(), EQueuedWorkPriority::Normal);
+		}
+	}
+	else
+	{
+		CachedState->RHIPipeline = RHICreateComputePipelineState(CachedState->ComputeShader);
+	}
+
+	GComputePipelineCache.Add(ComputeShader, CachedState);
+	return GraphEvent;
+}
+
+FGraphEventRef PipelineStateCache::PrecacheGraphicsPipelineState(const FGraphicsPipelineStateInitializer& Initializer)
+{
+	FGraphEventRef GraphEvent = nullptr;
+	if (!IsPSOPrecachingEnabled())
+	{
+		return GraphEvent;
+	}
+
+	LLM_SCOPE(ELLMTag::PSO);
+	ValidateGraphicsPipelineStateInitializer(Initializer);
+		
+	// Already in the cache?
+	if (GPrecacheGraphicsPipelineCache.Contains(Initializer))
+	{
+		return GraphEvent;
+	}
+
+	// create new graphics state
+	FGraphicsPipelineState* CachedState = new FGraphicsPipelineState();
+	
+	// Use async compilation
+	static bool bDoAsyncCompile = FApp::ShouldUseThreadingForPerformance();
+
+	// Add to cache before starting async operation, because the async op will remove the task from active ops when done
+	GPrecacheGraphicsPipelineCache.Add(Initializer, CachedState, bDoAsyncCompile);
+
+	// Start the async 
+	GraphEvent = CreateGraphicsPipelineState(Initializer, EPSOPrecacheResult::Active, bDoAsyncCompile, CachedState);
+
+	return GraphEvent;
+}
+
+EPSOPrecacheResult PipelineStateCache::CheckPipelineStateInCache(const FGraphicsPipelineStateInitializer& PipelineStateInitializer)
+{
+	if (!IsPSOPrecachingEnabled())
+	{
+		return EPSOPrecacheResult::Unknown;
+	}
+
+	if (!GPrecacheGraphicsPipelineCache.Contains(PipelineStateInitializer))
+	{
+		return EPSOPrecacheResult::Missed;
+	}
+
+	return GPrecacheGraphicsPipelineCache.IsPrecaching(PipelineStateInitializer) ? EPSOPrecacheResult::Active : EPSOPrecacheResult::Complete;
+}
+
+bool PipelineStateCache::IsPrecaching(const FGraphicsPipelineStateInitializer& PipelineStateInitializer)
+{
+	if (!IsPSOPrecachingEnabled())
+	{
+		return false;
+	}
+
+	return GPrecacheGraphicsPipelineCache.IsPrecaching(PipelineStateInitializer);
+}
+
+bool PipelineStateCache::IsPrecaching()
+{
+	if (!IsPSOPrecachingEnabled())
+	{
+		return false;
+	}
+
+	return GPrecacheGraphicsPipelineCache.IsPrecaching();
+}
+
 FRHIGraphicsPipelineState* ExecuteSetGraphicsPipelineState(FGraphicsPipelineState* GraphicsPipelineState)
 {
 	FRHIGraphicsPipelineState* RHIPipeline = GraphicsPipelineState->RHIPipeline;
diff --git a/Engine/Source/Runtime/RHI/Public/DynamicRHI.h b/Engine/Source/Runtime/RHI/Public/DynamicRHI.h
--- a/Engine/Source/Runtime/RHI/Public/DynamicRHI.h
+++ b/Engine/Source/Runtime/RHI/Public/DynamicRHI.h
@@ -855,6 +855,12 @@ public:
 	virtual void RHIGetDisplaysInformation(FDisplayInformationArray& OutDisplayInformation)
 	{
 	}
+	
+	// Compute the hash of the PSO initializer for PSO Precaching (only hash data relevant for the RHI specific PSO)
+	virtual uint64 RHIComputePrecachePSOHash(const FGraphicsPipelineStateInitializer& Initializer);
+
+	// Check if PSO Initializers are the same used during PSO Precaching (only compare data relevant for the RHI specific PSO)
+	virtual bool RHIMatchPrecachePSOInitializers(const FGraphicsPipelineStateInitializer& LHS, const FGraphicsPipelineStateInitializer& RHS);
 
 	virtual void RHIAdvanceFrameFence() {};
 
@@ -1640,6 +1646,16 @@ FORCEINLINE void RHIGetDisplaysInformation(FDisplayInformationArray& OutDisplayI
 	GDynamicRHI->RHIGetDisplaysInformation(OutDisplayInformation);
 }
 
+FORCEINLINE uint64 RHIComputePrecachePSOHash(const FGraphicsPipelineStateInitializer& Initializer)
+{
+	return GDynamicRHI->RHIComputePrecachePSOHash(Initializer);
+}
+
+FORCEINLINE bool RHIMatchPrecachePSOInitializers(const FGraphicsPipelineStateInitializer& LHS, const FGraphicsPipelineStateInitializer& RHS)
+{
+	return GDynamicRHI->RHIMatchPrecachePSOInitializers(LHS, RHS);
+}
+
 #if RHI_RAYTRACING
 
 FORCEINLINE FRayTracingAccelerationStructureSize RHICalcRayTracingSceneSize(uint32 MaxInstances, ERayTracingAccelerationStructureFlags Flags)
diff --git a/Engine/Source/Runtime/RHI/Public/PipelineStateCache.h b/Engine/Source/Runtime/RHI/Public/PipelineStateCache.h
--- a/Engine/Source/Runtime/RHI/Public/PipelineStateCache.h
+++ b/Engine/Source/Runtime/RHI/Public/PipelineStateCache.h
@@ -39,8 +39,18 @@ enum class ERayTracingPipelineCacheFlags : int
 };
 ENUM_CLASS_FLAGS(ERayTracingPipelineCacheFlags);
 
+enum class EPSOPrecacheResult
+{
+	Unknown,			//< No known pre cache state
+	Active,				//< PSO is currently precaching
+	Complete,			//< PSO has been precached successfully
+	Missed,				//< PSO precache miss, needs to be compiled at draw time
+	NotSupported,		//< PSO precache not supported (VertexFactory or MeshPassProcessor doesn't support/implement precaching)
+	Untracked,			//< PSO is not tracked at all (Compute or not coming from MeshDrawCommands)
+};
+
 extern RHI_API void SetComputePipelineState(FRHIComputeCommandList& RHICmdList, FRHIComputeShader* ComputeShader);
-extern RHI_API void SetGraphicsPipelineState(FRHICommandList& RHICmdList, const FGraphicsPipelineStateInitializer& Initializer, uint32 StencilRef, EApplyRendertargetOption ApplyFlags = EApplyRendertargetOption::CheckApply, bool bApplyAdditionalState = true);
+extern RHI_API void SetGraphicsPipelineState(FRHICommandList& RHICmdList, const FGraphicsPipelineStateInitializer& Initializer, uint32 StencilRef, EApplyRendertargetOption ApplyFlags = EApplyRendertargetOption::CheckApply, bool bApplyAdditionalState = true, EPSOPrecacheResult PSOPrecacheResult = EPSOPrecacheResult::Untracked);
 
 UE_DEPRECATED(5.0, "SetGraphicsPipelineState now requires a StencilRef argument and EApplyRendertargetOption::ForceApply will soon be removed")
 inline void SetGraphicsPipelineState(FRHICommandList& RHICmdList, const FGraphicsPipelineStateInitializer& Initializer, EApplyRendertargetOption ApplyFlags = EApplyRendertargetOption::CheckApply, bool bApplyAdditionalState = true)
@@ -54,7 +64,7 @@ namespace PipelineStateCache
 
 	extern RHI_API FComputePipelineState*	GetAndOrCreateComputePipelineState(FRHIComputeCommandList& RHICmdList, FRHIComputeShader* ComputeShader, bool bFromFileCache);
 
-	extern RHI_API FGraphicsPipelineState*	GetAndOrCreateGraphicsPipelineState(FRHICommandList& RHICmdList, const FGraphicsPipelineStateInitializer& OriginalInitializer, EApplyRendertargetOption ApplyFlags);
+	extern RHI_API FGraphicsPipelineState*	GetAndOrCreateGraphicsPipelineState(FRHICommandList& RHICmdList, const FGraphicsPipelineStateInitializer& OriginalInitializer, EApplyRendertargetOption ApplyFlags, EPSOPrecacheResult PSOPrecacheResult);
 
 	extern RHI_API FGraphicsPipelineState*	FindGraphicsPipelineState(const FGraphicsPipelineStateInitializer& Initializer);
 
@@ -83,6 +93,26 @@ namespace PipelineStateCache
 
 	/* Returns the number of PSO precompiles currently in progress */
 	extern RHI_API int32 GetNumActivePipelinePrecompileTasks();
+	
+	/* Is precaching currently enabled - can help to skip certain time critical code when precaching is disabled */
+	extern RHI_API bool						IsPSOPrecachingEnabled();
+
+	/* Precache the compute shader and return an optional graph event if precached async */
+	extern RHI_API FGraphEventRef			PrecacheComputePipelineState(FRHIComputeShader* ComputeShader);
+
+	/* Precache the graphic PSO and return an optional graph event if precached async */
+	extern RHI_API FGraphEventRef			PrecacheGraphicsPipelineState(const FGraphicsPipelineStateInitializer& PipelineStateInitializer);
+
+	/* Retrieve the current PSO precache result state (slightly slower than IsPrecaching) */
+	extern RHI_API EPSOPrecacheResult		CheckPipelineStateInCache(const FGraphicsPipelineStateInitializer& PipelineStateInitializer);
+
+	/* Is the given PSO initializer still precaching? */
+	extern RHI_API bool						IsPrecaching(const FGraphicsPipelineStateInitializer& PipelineStateInitializer);
+
+	/* Any async precaching operations still busy */
+	extern RHI_API bool						IsPrecaching();
+
+	
 }
 
 // Returns the shader index within the ray tracing pipeline or INDEX_NONE if given shader does not exist.
diff --git a/Engine/Source/Runtime/RHI/Public/RHI.h b/Engine/Source/Runtime/RHI/Public/RHI.h
--- a/Engine/Source/Runtime/RHI/Public/RHI.h
+++ b/Engine/Source/Runtime/RHI/Public/RHI.h
@@ -1256,6 +1256,16 @@ struct FVertexElement
 		bUseInstanceIndex = Other.bUseInstanceIndex;
 	}
 
+	bool operator==(const FVertexElement& Other) const
+	{
+		return (StreamIndex		== Other.StreamIndex &&
+				Offset			== Other.Offset &&
+				Type			== Other.Type &&
+				AttributeIndex	== Other.AttributeIndex &&
+				Stride			== Other.Stride &&
+				bUseInstanceIndex == Other.bUseInstanceIndex);
+	}
+
 	friend FArchive& operator<<(FArchive& Ar,FVertexElement& Element)
 	{
 		Ar << Element.StreamIndex;
diff --git a/Engine/Source/Runtime/RHI/Public/RHIResources.h b/Engine/Source/Runtime/RHI/Public/RHIResources.h
--- a/Engine/Source/Runtime/RHI/Public/RHIResources.h
+++ b/Engine/Source/Runtime/RHI/Public/RHIResources.h
@@ -552,6 +552,22 @@ public:
 	virtual bool GetInitializer(class FBlendStateInitializerRHI& Init) { return false; }
 };
 
+template <typename RHIState, typename RHIStateInitializer>
+static bool MatchRHIState(RHIState* LHSState, RHIState* RHSState)
+{
+	RHIStateInitializer LHSStateInitializerRHI;
+	RHIStateInitializer RHSStateInitializerRHI;
+	if (LHSState)
+	{
+		LHSState->GetInitializer(LHSStateInitializerRHI);
+	}
+	if (RHSState)
+	{
+		RHSState->GetInitializer(RHSStateInitializerRHI);
+	}
+	return LHSStateInitializerRHI == RHSStateInitializerRHI;
+}
+
 //
 // Shader bindings
 //
@@ -562,6 +578,7 @@ class FRHIVertexDeclaration : public FRHIResource
 public:
 	FRHIVertexDeclaration() : FRHIResource(RRT_VertexDeclaration) {}
 	virtual bool GetInitializer(FVertexDeclarationElementList& Init) { return false; }
+	virtual uint32 GetPrecachePSOHash() const { return 0; }
 };
 
 class FRHIBoundShaderState : public FRHIResource
@@ -3275,6 +3292,7 @@ public:
 		, bHasFragmentDensityAttachment(false)
 		, ShadingRate(EVRSShadingRate::VRSSR_1x1)
 		, Flags(0)
+		, StatePrecachePSOHash(0)
 	{
 #if PLATFORM_WINDOWS
 		static_assert(sizeof(TRenderTargetFormats::ElementType) == sizeof(uint8/*EPixelFormat*/), "Change TRenderTargetFormats's uint8 to EPixelFormat's size!");
@@ -3307,7 +3325,8 @@ public:
 		bool						bInDepthBounds,
 		uint8						InMultiViewCount,
 		bool						bInHasFragmentDensityAttachment,
-		EVRSShadingRate				InShadingRate)
+		EVRSShadingRate				InShadingRate,
+		uint64						InStatePrecachePSOHash)
 		: BoundShaderState(InBoundShaderState)
 		, BlendState(InBlendState)
 		, RasterizerState(InRasterizerState)
@@ -3333,6 +3352,7 @@ public:
 		, bHasFragmentDensityAttachment(bInHasFragmentDensityAttachment)
 		, ShadingRate(InShadingRate)
 		, Flags(InFlags)
+		, StatePrecachePSOHash(InStatePrecachePSOHash)
 	{
 	}
 
@@ -3432,6 +3452,10 @@ public:
 		};
 		uint16						Flags;
 	};
+
+	// Cached hash off all state data provided at creation time (Only contains hash of data which influences the PSO precaching for the current platform)
+	// Created from hashing the state data instead of the pointers which are used during fast runtime cache checking and compares
+	uint64							StatePrecachePSOHash;
 };
 
 class FRayTracingPipelineStateSignature
diff --git a/Engine/Source/Runtime/RenderCore/Private/GBufferInfo.cpp b/Engine/Source/Runtime/RenderCore/Private/GBufferInfo.cpp
--- a/Engine/Source/Runtime/RenderCore/Private/GBufferInfo.cpp
+++ b/Engine/Source/Runtime/RenderCore/Private/GBufferInfo.cpp
@@ -148,6 +148,9 @@ FGBufferBinding FindGBufferBindingByName(const FGBufferInfo& GBufferInfo, const
 
 		switch (Target.TargetType)
 		{
+		case GBT_Unorm_16_16:
+			PixelFormat = PF_G16R16;
+			break;
 		case GBT_Unorm_8_8_8_8:
 			PixelFormat = PF_B8G8R8A8;
 			break;
@@ -157,8 +160,11 @@ FGBufferBinding FindGBufferBindingByName(const FGBufferInfo& GBufferInfo, const
 		case GBT_Unorm_10_10_10_2:
 			PixelFormat = PF_A2B10G10R10;
 			break;
+		case GBT_Unorm_16_16_16_16:
+			PixelFormat = PF_A16B16G16R16;
+			break;
 		case GBT_Float_16_16:
-			PixelFormat = PF_G16R16;
+			PixelFormat = PF_G16R16F;
 			break;
 		case GBT_Float_16_16_16_16:
 			PixelFormat = PF_FloatRGBA;
@@ -270,7 +276,7 @@ FGBufferInfo RENDERCORE_API FetchLegacyGBufferInfo(const FGBufferParams& Params)
 
 		if (Params.bHasVelocity)
 		{
-			Info.Targets[TargetVelocity].Init(Params.bUsesVelocityDepth ? GBT_Float_16_16_16_16 : GBT_Float_16_16, TEXT("Velocity"), false, true, true, false); // Velocity
+			Info.Targets[TargetVelocity].Init(Params.bUsesVelocityDepth ? GBT_Unorm_16_16_16_16 : (IsAndroidOpenGLESPlatform(Params.ShaderPlatform) ? GBT_Float_16_16 : GBT_Unorm_16_16), TEXT("Velocity"), false, true, true, false); // Velocity
 			Info.Slots[GBS_Velocity] = FGBufferItem(GBS_Velocity, Params.bUsesVelocityDepth ? GBC_Raw_Float_16_16_16_16 : GBC_Raw_Float_16_16, GBCH_Both);
 			Info.Slots[GBS_Velocity].Packing[0] = FGBufferPacking(TargetVelocity, 0, 0);
 			Info.Slots[GBS_Velocity].Packing[1] = FGBufferPacking(TargetVelocity, 1, 1);
@@ -363,7 +369,7 @@ FGBufferInfo RENDERCORE_API FetchLegacyGBufferInfo(const FGBufferParams& Params)
 		TargetGBufferD = 5;
 
 		// note the false for use extra flags for velocity, not quite sure of all the ramifications, but this keeps it consistent with previous usage
-		Info.Targets[4].Init(Params.bUsesVelocityDepth ? GBT_Float_16_16_16_16 : GBT_Float_16_16,   TEXT("Velocity"), false,  true,  true, false);
+		Info.Targets[4].Init(Params.bUsesVelocityDepth ? GBT_Unorm_16_16_16_16 : (IsAndroidOpenGLESPlatform(Params.ShaderPlatform) ? GBT_Float_16_16 : GBT_Unorm_16_16), TEXT("Velocity"), false, true, true, false);
 		Info.Targets[5].Init(GBT_Unorm_8_8_8_8, TEXT("GBufferD"), false, true, true, true);
 		TargetSeparatedMainDirLight = 6;
 
diff --git a/Engine/Source/Runtime/RenderCore/Private/ShaderResource.cpp b/Engine/Source/Runtime/RenderCore/Private/ShaderResource.cpp
--- a/Engine/Source/Runtime/RenderCore/Private/ShaderResource.cpp
+++ b/Engine/Source/Runtime/RenderCore/Private/ShaderResource.cpp
@@ -596,8 +596,7 @@ void FShaderMapResource::BeginCreateAllShaders()
 }
 
 FRHIShader* FShaderMapResource::CreateShader(int32 ShaderIndex)
-{
-	check(IsInParallelRenderingThread());
+{	
 	check(!RHIShaders[ShaderIndex].load(std::memory_order_acquire));
 
 	TRefCountPtr<FRHIShader> RHIShader = CreateRHIShader(ShaderIndex);
diff --git a/Engine/Source/Runtime/RenderCore/Public/GBufferInfo.h b/Engine/Source/Runtime/RenderCore/Public/GBufferInfo.h
--- a/Engine/Source/Runtime/RenderCore/Public/GBufferInfo.h
+++ b/Engine/Source/Runtime/RenderCore/Public/GBufferInfo.h
@@ -82,11 +82,13 @@ enum EGBufferCompression
 enum EGBufferType
 {
 	GBT_Invalid,
+	GBT_Unorm_16_16,
 	GBT_Unorm_8_8_8_8,
 	GBT_Unorm_11_11_10,
 	GBT_Unorm_10_10_10_2,
+	GBT_Unorm_16_16_16_16,
 	GBT_Float_16_16,
-	GBT_Float_16_16_16_16,
+	GBT_Float_16_16_16_16,	
 	GBT_Num
 };
 
@@ -269,6 +271,7 @@ struct FGBufferInfo
 
 struct FGBufferParams
 {
+	EShaderPlatform ShaderPlatform = SP_NumPlatforms;
 	int32 LegacyFormatIndex = 0;
 	bool bHasVelocity = false;
 	bool bHasTangent = false;
@@ -279,6 +282,7 @@ struct FGBufferParams
 	bool operator == (const FGBufferParams& RHS) const
 	{
 		return
+			ShaderPlatform == RHS.ShaderPlatform &&
 			LegacyFormatIndex == RHS.LegacyFormatIndex &&
 			bHasVelocity == RHS.bHasVelocity &&
 			bHasTangent == RHS.bHasTangent &&
@@ -298,9 +302,6 @@ int32 RENDERCORE_API FindGBufferTargetByName(const FGBufferInfo& GBufferInfo, co
 
 FGBufferBinding RENDERCORE_API FindGBufferBindingByName(const FGBufferInfo& GBufferInfo, const FString& Name);
 
-//FGBufferInfo RENDERCORE_API FetchLegacyGBufferInfo(bool bHasVelocity, bool bHasTangent, bool bHasPrecShadowFactor, int LegacyFormatIndex, bool bUsesVelocityDepth);
-//FGBufferInfo RENDERCORE_API FetchFullGBufferInfo(bool bHasVelocity, bool bHasTangent, bool bHasPrecShadowFactor, int LegacyFormatIndex, bool bUsesVelocityDepth);
-FGBufferInfo RENDERCORE_API FetchLegacyGBufferInfo(const FGBufferParams& Params);
 FGBufferInfo RENDERCORE_API FetchFullGBufferInfo(const FGBufferParams& Params);
 
 bool RENDERCORE_API IsGBufferInfoEqual(const FGBufferInfo& Lhs, const FGBufferInfo& Rhs);
diff --git a/Engine/Source/Runtime/RenderCore/Public/VertexFactory.h b/Engine/Source/Runtime/RenderCore/Public/VertexFactory.h
--- a/Engine/Source/Runtime/RenderCore/Public/VertexFactory.h
+++ b/Engine/Source/Runtime/RenderCore/Public/VertexFactory.h
@@ -121,7 +121,8 @@ enum class EVertexFactoryFlags : uint32
 	SupportsRayTracing                    = 1u << 9,
 	SupportsRayTracingDynamicGeometry     = 1u << 10,
 	SupportsRayTracingProceduralPrimitive = 1u << 11,
-	SupportsLightmapBaking                = 1u << 12
+	SupportsLightmapBaking                = 1u << 12,
+	SupportsPSOPrecaching				  = 1u << 13,
 };
 ENUM_CLASS_FLAGS(EVertexFactoryFlags);
 
@@ -401,6 +402,7 @@ public:
 		const struct FMeshBatchElement& BatchElement,
 		class FMeshDrawSingleShaderBindings& ShaderBindings,
 		FVertexInputStreamArray& VertexStreams) const { (*GetParameterTypeElementShaderBindings)(ShaderFrequency, Parameters, Scene, View, Shader, InputStreamType, FeatureLevel, VertexFactory, BatchElement, ShaderBindings, VertexStreams); }
+	void GetPSOPrecacheVertexFetchElements(EVertexInputStreamType VertexInputStreamType, FVertexDeclarationElementList& Elements) const { }
 
 	EVertexFactoryFlags GetFlags() const
 	{
@@ -424,7 +426,8 @@ public:
 	bool SupportsRayTracing() const                    { return HasFlags(EVertexFactoryFlags::SupportsRayTracing); }
 	bool SupportsRayTracingDynamicGeometry() const     { return HasFlags(EVertexFactoryFlags::SupportsRayTracingDynamicGeometry); }
 	bool SupportsRayTracingProceduralPrimitive() const { return HasFlags(EVertexFactoryFlags::SupportsRayTracingProceduralPrimitive); }
-	bool SupportsLightmapBaking() const                { return HasFlags(EVertexFactoryFlags::SupportsLightmapBaking); }
+	bool SupportsLightmapBaking() const                { return HasFlags(EVertexFactoryFlags::SupportsLightmapBaking); }	
+	bool SupportsPSOPrecaching() const				   { return HasFlags(EVertexFactoryFlags::SupportsPSOPrecaching); }
 
 	// Hash function.
 	friend uint32 GetTypeHash(const FVertexFactoryType* Type)
diff --git a/Engine/Source/Runtime/Renderer/Private/MeshPassProcessor.cpp b/Engine/Source/Runtime/Renderer/Private/MeshPassProcessor.cpp
--- a/Engine/Source/Runtime/Renderer/Private/MeshPassProcessor.cpp
+++ b/Engine/Source/Runtime/Renderer/Private/MeshPassProcessor.cpp
@@ -45,6 +45,14 @@ static TAutoConsoleVariable<int32> CVarSafeStateLookup(
 	TEXT("Forces new-style safe state lookup for easy runtime perf comparison\n"),
 	ECVF_Scalability | ECVF_RenderThreadSafe);
 
+int32 GSkipDrawOnPSOPrecaching = 0;
+static FAutoConsoleVariableRef CVarSkipDrawOnPSOPrecaching(
+	TEXT("r.SkipDrawOnPSOPrecaching"),
+	GSkipDrawOnPSOPrecaching,
+	TEXT("Skips mesh draw call when the PSO is still compiling (default 0)."),
+	ECVF_RenderThreadSafe
+);
+
 #if WITH_EDITORONLY_DATA
 
 int32 GNaniteIsolateInvalidCoarseMesh = 0;
@@ -906,37 +914,96 @@ void FMeshDrawShaderBindings::Release()
 	Data.SetHeapData(nullptr);
 }
 
-void FMeshDrawCommand::SetShaders(FRHIVertexDeclaration* VertexDeclaration, const FMeshProcessorShaders& Shaders, FGraphicsMinimalPipelineStateInitializer& PipelineState)
+void FGraphicsMinimalPipelineStateInitializer::SetupBoundShaderState(FRHIVertexDeclaration* VertexDeclaration, const FMeshProcessorShaders& Shaders)
 {
-	PipelineState.BoundShaderState = FMinimalBoundShaderStateInput();
-	PipelineState.BoundShaderState.VertexDeclarationRHI = VertexDeclaration;
+	BoundShaderState = FMinimalBoundShaderStateInput();
+	BoundShaderState.VertexDeclarationRHI = VertexDeclaration;
 
 	checkf(Shaders.VertexShader.IsValid(), TEXT("Can't render without a vertex shader"));
 
 	if(Shaders.VertexShader.IsValid())
 	{
 		checkSlow(Shaders.VertexShader->GetFrequency() == SF_Vertex);
-		PipelineState.BoundShaderState.VertexShaderResource = Shaders.VertexShader.GetResource();
-		PipelineState.BoundShaderState.VertexShaderIndex = Shaders.VertexShader->GetResourceIndex();
-		check(PipelineState.BoundShaderState.VertexShaderResource->IsValidShaderIndex(PipelineState.BoundShaderState.VertexShaderIndex));
+		BoundShaderState.VertexShaderResource = Shaders.VertexShader.GetResource();
+		BoundShaderState.VertexShaderIndex = Shaders.VertexShader->GetResourceIndex();
+		check(BoundShaderState.VertexShaderResource->IsValidShaderIndex(BoundShaderState.VertexShaderIndex));
 	}
 	if (Shaders.PixelShader.IsValid())
 	{
 		checkSlow(Shaders.PixelShader->GetFrequency() == SF_Pixel);
-		PipelineState.BoundShaderState.PixelShaderResource = Shaders.PixelShader.GetResource();
-		PipelineState.BoundShaderState.PixelShaderIndex = Shaders.PixelShader->GetResourceIndex();
-		check(PipelineState.BoundShaderState.PixelShaderResource->IsValidShaderIndex(PipelineState.BoundShaderState.PixelShaderIndex));
+		BoundShaderState.PixelShaderResource = Shaders.PixelShader.GetResource();
+		BoundShaderState.PixelShaderIndex = Shaders.PixelShader->GetResourceIndex();
+		check(BoundShaderState.PixelShaderResource->IsValidShaderIndex(BoundShaderState.PixelShaderIndex));
 	}
 #if PLATFORM_SUPPORTS_GEOMETRY_SHADERS
 	if (Shaders.GeometryShader.IsValid())
 	{
 		checkSlow(Shaders.GeometryShader->GetFrequency() == SF_Geometry);
-		PipelineState.BoundShaderState.GeometryShaderResource = Shaders.GeometryShader.GetResource();
-		PipelineState.BoundShaderState.GeometryShaderIndex = Shaders.GeometryShader->GetResourceIndex();
-		check(PipelineState.BoundShaderState.GeometryShaderResource->IsValidShaderIndex(PipelineState.BoundShaderState.GeometryShaderIndex));
+		BoundShaderState.GeometryShaderResource = Shaders.GeometryShader.GetResource();
+		BoundShaderState.GeometryShaderIndex = Shaders.GeometryShader->GetResourceIndex();
+		check(BoundShaderState.GeometryShaderResource->IsValidShaderIndex(BoundShaderState.GeometryShaderIndex));
 	}
 #endif // PLATFORM_SUPPORTS_GEOMETRY_SHADERS
-	ShaderBindings.Initialize(Shaders);
+}
+
+void FGraphicsMinimalPipelineStateInitializer::ComputePrecachePSOHash()
+{
+	struct FHashKey
+	{
+		uint32 VertexDeclaration;
+		uint32 VertexShader;
+		uint32 PixelShader;
+#if PLATFORM_SUPPORTS_GEOMETRY_SHADERS
+		uint32 GeometryShader;
+#endif // PLATFORM_SUPPORTS_GEOMETRY_SHADERS
+		uint32 BlendState;
+		uint32 RasterizerState;
+		uint32 DepthStencilState;
+		uint32 ImmutableSamplerState;
+
+		uint32 MultiViewCount : 8;
+		uint32 DrawShadingRate : 8;
+		uint32 PrimitiveType : 8;
+		uint32 bDepthBounds : 1;
+		uint32 bHasFragmentDensityAttachment : 1;
+		uint32 Unused : 6;
+	} HashKey;
+
+	FMemory::Memzero(&HashKey, sizeof(FHashKey));
+
+	HashKey.VertexDeclaration = BoundShaderState.VertexDeclarationRHI ? BoundShaderState.VertexDeclarationRHI->GetPrecachePSOHash() : 0;
+	HashKey.VertexShader = HashCombine(GetTypeHash(BoundShaderState.VertexShaderIndex), GetTypeHash(BoundShaderState.VertexShaderResource));
+	HashKey.PixelShader = HashCombine(GetTypeHash(BoundShaderState.PixelShaderIndex), GetTypeHash(BoundShaderState.PixelShaderResource));
+#if PLATFORM_SUPPORTS_GEOMETRY_SHADERS
+	HashKey.GeometryShader = HashCombine(GetTypeHash(BoundShaderState.GeometryShaderIndex), GetTypeHash(BoundShaderState.GeometryShaderResource));
+#endif
+
+	FBlendStateInitializerRHI BlendStateInitializerRHI;
+	if (BlendState && BlendState->GetInitializer(BlendStateInitializerRHI))
+	{
+		HashKey.BlendState = GetTypeHash(BlendStateInitializerRHI);
+	}
+	FRasterizerStateInitializerRHI RasterizerStateInitializerRHI;
+	if (RasterizerState && RasterizerState->GetInitializer(RasterizerStateInitializerRHI))
+	{
+		HashKey.RasterizerState = GetTypeHash(RasterizerStateInitializerRHI);
+	}
+	FDepthStencilStateInitializerRHI DepthStencilStateInitializerRHI;
+	if (DepthStencilState && DepthStencilState->GetInitializer(DepthStencilStateInitializerRHI))
+	{
+		HashKey.DepthStencilState = GetTypeHash(DepthStencilStateInitializerRHI);
+	}
+
+	// Ignore immutable samplers for now
+	//HashKey.ImmutableSamplerState = GetTypeHash(ImmutableSamplerState);
+	
+	HashKey.MultiViewCount = MultiViewCount;
+	HashKey.DrawShadingRate = DrawShadingRate;
+	HashKey.PrimitiveType = PrimitiveType;
+	HashKey.bDepthBounds = bDepthBounds;
+	HashKey.bHasFragmentDensityAttachment = bHasFragmentDensityAttachment;
+
+	PrecachePSOHash = CityHash64((const char*)&HashKey, sizeof(FHashKey));
 }
 
 #if RHI_RAYTRACING
@@ -1253,7 +1320,7 @@ uint32 FMeshDrawShaderBindings::GetDynamicInstancingHash() const
 	return uint32(CityHash64((char*)&HashKey, sizeof(FHashKey)));
 }
 
-void FMeshDrawCommand::SubmitDrawBegin(
+bool FMeshDrawCommand::SubmitDrawBegin(
 	const FMeshDrawCommand& RESTRICT MeshDrawCommand, 
 	const FGraphicsMinimalPipelineStateSet& GraphicsMinimalPipelineStateSet,
 	FRHIBuffer* ScenePrimitiveIdsBuffer,
@@ -1273,8 +1340,32 @@ void FMeshDrawCommand::SubmitDrawBegin(
 	{
 		FGraphicsPipelineStateInitializer GraphicsPSOInit = MeshPipelineState.AsGraphicsPipelineStateInitializer();
 		RHICmdList.ApplyCachedRenderTargets(GraphicsPSOInit);
+
+		EPSOPrecacheResult PSOPrecacheResult = EPSOPrecacheResult::Unknown;
+#if PSO_PRECACHING_VALIDATE
+		// Retrieve state from cache - can be be cached on the MeshPipelineState to not retrieve it anymore after the final state is known
+#if MESH_DRAW_COMMAND_DEBUG_DATA
+		PSOPrecacheResult = PSOCollectorStats::CheckPipelineStateInCache(GraphicsPSOInit, MeshDrawCommand.DebugData.MeshPassType, MeshDrawCommand.DebugData.VertexFactoryType);
+#else
+		PSOPrecacheResult = PSOCollectorStats::CheckPipelineStateInCache(GraphicsPSOInit, EMeshPass::Num, nullptr);
+#endif // MESH_DRAW_COMMAND_DEBUG_DATA
+#endif // PSO_PRECACHING_VALIDATE
+
+		// Check if skip draw is needed when the PSO is still precaching
+		if (GSkipDrawOnPSOPrecaching && !MeshPipelineState.bPSOPrecached)
+		{
+			MeshPipelineState.bPSOPrecached = !PipelineStateCache::IsPrecaching(GraphicsPSOInit);
+
+			// Try and skip draw if not precached yet
+			if (!MeshPipelineState.bPSOPrecached)
+			{
+				return false;
+			}
+		}
+
 		// We can set the new StencilRef here to avoid the set below
-		SetGraphicsPipelineState(RHICmdList, GraphicsPSOInit, MeshDrawCommand.StencilRef);
+		bool bApplyAdditionalState = true;
+		SetGraphicsPipelineState(RHICmdList, GraphicsPSOInit, MeshDrawCommand.StencilRef, EApplyRendertargetOption::CheckApply, bApplyAdditionalState, PSOPrecacheResult);
 		StateCache.SetPipelineState(MeshDrawCommand.CachedPipelineId.GetId());
 		StateCache.StencilRef = MeshDrawCommand.StencilRef;
 	}
@@ -1302,6 +1393,8 @@ void FMeshDrawCommand::SubmitDrawBegin(
 	}
 
 	MeshDrawCommand.ShaderBindings.SetOnCommandList(RHICmdList, MeshPipelineState.BoundShaderState.AsBoundShaderState(), StateCache.ShaderBindings);
+
+	return true;
 }
 
 void FMeshDrawCommand::SubmitDrawEnd(const FMeshDrawCommand& MeshDrawCommand, uint32 InstanceFactor, FRHICommandList& RHICmdList,
@@ -1352,7 +1445,7 @@ void FMeshDrawCommand::SubmitDrawEnd(const FMeshDrawCommand& MeshDrawCommand, ui
 	}
 }
 
-void FMeshDrawCommand::SubmitDrawIndirectBegin(
+bool FMeshDrawCommand::SubmitDrawIndirectBegin(
 	const FMeshDrawCommand& RESTRICT MeshDrawCommand,
 	const FGraphicsMinimalPipelineStateSet& GraphicsMinimalPipelineStateSet,
 	FRHIBuffer* ScenePrimitiveIdsBuffer,
@@ -1361,7 +1454,7 @@ void FMeshDrawCommand::SubmitDrawIndirectBegin(
 	FRHICommandList& RHICmdList,
 	FMeshDrawCommandStateCache& RESTRICT StateCache)
 {
-	SubmitDrawBegin(
+	return SubmitDrawBegin(
 		MeshDrawCommand,
 		GraphicsMinimalPipelineStateSet,
 		ScenePrimitiveIdsBuffer,
@@ -1463,11 +1556,13 @@ void FMeshDrawCommand::SubmitDraw(
 #if WANTS_DRAW_MESH_EVENTS
 	FMeshDrawEvent MeshEvent(MeshDrawCommand, InstanceFactor, RHICmdList);
 #endif
-	SubmitDrawBegin(MeshDrawCommand, GraphicsMinimalPipelineStateSet, ScenePrimitiveIdsBuffer, PrimitiveIdOffset, InstanceFactor, RHICmdList, StateCache);
-	SubmitDrawEnd(MeshDrawCommand, InstanceFactor, RHICmdList, IndirectArgsOverrideBuffer, IndirectArgsOverrideByteOffset);
+	if (SubmitDrawBegin(MeshDrawCommand, GraphicsMinimalPipelineStateSet, ScenePrimitiveIdsBuffer, PrimitiveIdOffset, InstanceFactor, RHICmdList, StateCache))
+	{
+		SubmitDrawEnd(MeshDrawCommand, InstanceFactor, RHICmdList, IndirectArgsOverrideBuffer, IndirectArgsOverrideByteOffset);
+	}
 }
 
-static void ApplyTargetsInfo(FGraphicsPipelineStateInitializer& GraphicsPSOInit, const FGraphicsPipelineRenderTargetsInfo& RenderTargetsInfo)
+void ApplyTargetsInfo(FGraphicsPipelineStateInitializer& GraphicsPSOInit, const FGraphicsPipelineRenderTargetsInfo& RenderTargetsInfo)
 {
 	GraphicsPSOInit.RenderTargetsEnabled = RenderTargetsInfo.RenderTargetsEnabled;
 	GraphicsPSOInit.RenderTargetFormats = RenderTargetsInfo.RenderTargetFormats;
@@ -1498,7 +1593,8 @@ uint64 FMeshDrawCommand::GetPipelineStateSortingKey(FRHICommandList& RHICmdList,
 		FGraphicsPipelineStateInitializer GraphicsPSOInit = CachedPipelineId.GetPipelineState(PipelineStateSet).AsGraphicsPipelineStateInitializer();
 		ApplyTargetsInfo(GraphicsPSOInit, RenderTargetsInfo);
 
-		const FGraphicsPipelineState* PipelineState = PipelineStateCache::GetAndOrCreateGraphicsPipelineState(RHICmdList, GraphicsPSOInit, EApplyRendertargetOption::DoNothing);
+		// PSO is retrieved here already. This is currently only used by Nanite::DrawLumenMeshCards - can this also used the version without command list?
+		const FGraphicsPipelineState* PipelineState = PipelineStateCache::GetAndOrCreateGraphicsPipelineState(RHICmdList, GraphicsPSOInit, EApplyRendertargetOption::DoNothing, EPSOPrecacheResult::Unknown);
 		if (PipelineState)
 		{
 			const uint64 StateSortKey = PipelineStateCache::RetrieveGraphicsPipelineStateSortKey(PipelineState);
@@ -1538,13 +1634,15 @@ uint64 FMeshDrawCommand::GetPipelineStateSortingKey(const FGraphicsPipelineRende
 }
 
 #if MESH_DRAW_COMMAND_DEBUG_DATA
-void FMeshDrawCommand::SetDebugData(const FPrimitiveSceneProxy* PrimitiveSceneProxy, const FMaterial* Material, const FMaterialRenderProxy* MaterialRenderProxy, const FMeshProcessorShaders& UntypedShaders, const FVertexFactory* VertexFactory)
+void FMeshDrawCommand::SetDebugData(const FPrimitiveSceneProxy* PrimitiveSceneProxy, const FMaterial* Material, const FMaterialRenderProxy* MaterialRenderProxy, const FMeshProcessorShaders& UntypedShaders, const FVertexFactory* VertexFactory, uint32 MeshPassType)
 {
 	DebugData.PrimitiveSceneProxyIfNotUsingStateBuckets = PrimitiveSceneProxy;
 	DebugData.MaterialRenderProxy = MaterialRenderProxy;
 	DebugData.VertexShader = UntypedShaders.VertexShader;
 	DebugData.PixelShader = UntypedShaders.PixelShader;
 	DebugData.VertexFactory = VertexFactory;
+	DebugData.VertexFactoryType = VertexFactory->GetType();
+	DebugData.MeshPassType = MeshPassType;
 	DebugData.ResourceName =  PrimitiveSceneProxy ? PrimitiveSceneProxy->GetResourceName() : FName();
 	DebugData.MaterialName = Material->GetAssetName();
 }
@@ -1694,8 +1792,9 @@ FMeshDrawCommandSortKey CalculateMeshStaticSortKey(const FMeshMaterialShader* Ve
 	return SortKey;
 }
 
-FMeshPassProcessor::FMeshPassProcessor(const FScene* InScene, ERHIFeatureLevel::Type InFeatureLevel, const FSceneView* InViewIfDynamicMeshCommand, FMeshPassDrawListContext* InDrawListContext) 
-	: Scene(InScene)
+FMeshPassProcessor::FMeshPassProcessor(/*EMeshPass::Type InMeshPassType,*/ const FScene* InScene, ERHIFeatureLevel::Type InFeatureLevel, const FSceneView* InViewIfDynamicMeshCommand, FMeshPassDrawListContext* InDrawListContext)
+	: MeshPassType(EMeshPass::Num) // InMeshPassType)
+	, Scene(InScene)
 	, FeatureLevel(InFeatureLevel)
 	, ViewIfDynamicMeshCommand(InViewIfDynamicMeshCommand)
 	, DrawListContext(InDrawListContext)
@@ -2047,6 +2146,8 @@ void FCachedPassMeshDrawListContextDeferred::DeferredFinalizeMeshDrawCommands(co
 PassProcessorCreateFunction FPassProcessorManager::JumpTable[(int32)EShadingPath::Num][EMeshPass::Num] = {};
 EMeshPassFlags FPassProcessorManager::Flags[(int32)EShadingPath::Num][EMeshPass::Num] = {};
 
+static_assert(EMeshPass::Num <= FPSOCollectorCreateManager::MaxPSOCollectorCount);
+
 void FPassProcessorManager::SetPassFlags(EShadingPath ShadingPath, EMeshPass::Type PassType, EMeshPassFlags NewFlags)
 {
 	check(IsInGameThread());
@@ -2092,3 +2193,214 @@ FMeshDrawCommand::FMeshDrawEvent::FMeshDrawEvent(const FMeshDrawCommand& MeshDra
 }
 #endif
 
+void AddRenderTargetInfo(
+	EPixelFormat PixelFormat,
+	ETextureCreateFlags CreateFlags,
+	FGraphicsPipelineRenderTargetsInfo& RenderTargetsInfo)
+{
+	RenderTargetsInfo.RenderTargetFormats[RenderTargetsInfo.RenderTargetsEnabled] = PixelFormat;
+	RenderTargetsInfo.RenderTargetFlags[RenderTargetsInfo.RenderTargetsEnabled] = CreateFlags;
+	RenderTargetsInfo.RenderTargetsEnabled++;
+}
+
+void SetupDepthStencilInfo(
+	EPixelFormat DepthStencilFormat,
+	ETextureCreateFlags DepthStencilCreateFlags,
+	ERenderTargetLoadAction DepthTargetLoadAction,
+	ERenderTargetLoadAction StencilTargetLoadAction,
+	FExclusiveDepthStencil DepthStencilAccess,
+	FGraphicsPipelineRenderTargetsInfo& RenderTargetsInfo)
+{
+	// Setup depth stencil state
+	RenderTargetsInfo.DepthStencilTargetFormat = DepthStencilFormat;
+	RenderTargetsInfo.DepthStencilTargetFlag = DepthStencilCreateFlags;
+
+	RenderTargetsInfo.DepthTargetLoadAction = DepthTargetLoadAction;
+	RenderTargetsInfo.StencilTargetLoadAction = StencilTargetLoadAction;
+	RenderTargetsInfo.DepthStencilAccess = DepthStencilAccess;
+
+	const ERenderTargetStoreAction StoreAction = EnumHasAnyFlags(RenderTargetsInfo.DepthStencilTargetFlag, TexCreate_Memoryless) ? ERenderTargetStoreAction::ENoAction : ERenderTargetStoreAction::EStore;
+	RenderTargetsInfo.DepthTargetStoreAction = RenderTargetsInfo.DepthStencilAccess.IsUsingDepth() ? StoreAction : ERenderTargetStoreAction::ENoAction;
+	RenderTargetsInfo.StencilTargetStoreAction = RenderTargetsInfo.DepthStencilAccess.IsUsingStencil() ? StoreAction : ERenderTargetStoreAction::ENoAction;
+}
+
+void SetupGBufferRenderTargetInfo(const FSceneTexturesConfig& SceneTexturesConfig, FGraphicsPipelineRenderTargetsInfo& RenderTargetsInfo, bool bSetupDepthStencil)
+{
+	SceneTexturesConfig.GetGBufferRenderTargetsInfo(RenderTargetsInfo);
+	if (bSetupDepthStencil)
+	{
+		SetupDepthStencilInfo(PF_DepthStencil, SceneTexturesConfig.DepthCreateFlags, ERenderTargetLoadAction::ELoad,
+			ERenderTargetLoadAction::ELoad, FExclusiveDepthStencil::DepthWrite_StencilWrite, RenderTargetsInfo);
+	}
+}
+
+#if PSO_PRECACHING_VALIDATE
+
+// Helper function to remove certain members of the PSO initializer to help track down issues with PSO precache misses
+// eg Remove BlendState and check if the misses are gone for example
+static FGraphicsMinimalPipelineStateInitializer PatchMinimalPipelineStateToCheck(const FGraphicsMinimalPipelineStateInitializer& MinimalPSO)
+{
+	FGraphicsMinimalPipelineStateInitializer PSOInitializeToCheck = MinimalPSO;
+	//PSOInitializeToCheck.DepthStencilState = nullptr;
+	//PSOInitializeToCheck.RasterizerState = nullptr;
+	//PSOInitializeToCheck.BlendState = nullptr;
+	//PSOInitializeToCheck.PrimitiveType = PT_TriangleList;
+	//PSOInitializeToCheck.ImmutableSamplerState = FImmutableSamplerState();
+	//PSOInitializeToCheck.BoundShaderState.VertexDeclarationRHI = nullptr;
+	//PSOInitializeToCheck.UniqueEntry = false;
+	
+	// Recompute the hash when disabling certain states for checks
+	//PSOInitializeToCheck.ComputePrecachePSOHash();
+
+	return PSOInitializeToCheck;
+}
+
+void PSOCollectorStats::AddMinimalPipelineStateToCache(const FGraphicsMinimalPipelineStateInitializer& PSOInitialize, uint32 MeshPassType, const FVertexFactoryType* VertexFactoryType)
+{
+	if (!IsPrecachingValidationEnabled())
+	{
+		return;
+	}
+
+	check(VertexFactoryType->SupportsPSOPrecaching());
+
+	static bool bFirstTime = true;
+	if (bFirstTime)
+	{
+		PSOShadersPrecacheStats.Empty();
+		MinimalPSOPrecacheStats.Empty();
+		bFirstTime = false;
+	}
+
+	// Update the Shader only stats ignoring all state 
+	{
+		FScopeLock Lock(&PSOShadersPrecacheLock);
+		FGraphicsMinimalPipelineStateInitializer PSOInitializeToCheck;
+		PSOInitializeToCheck.BoundShaderState = PSOInitialize.BoundShaderState;
+		PSOInitializeToCheck.BoundShaderState.VertexDeclarationRHI = nullptr;
+		PSOInitializeToCheck.ComputePrecachePSOHash();
+
+		Experimental::FHashElementId TableId = PSOShadersPrecacheMap.FindId(PSOInitializeToCheck);
+		if (!TableId.IsValid())
+		{
+			TableId = PSOShadersPrecacheMap.FindOrAddId(PSOInitializeToCheck, FShaderStateUsage());
+
+		}
+
+		FShaderStateUsage& Value = PSOShadersPrecacheMap.GetByElementId(TableId).Value;
+		if (!Value.bPrecached)
+		{
+			check(!Value.bUsed);
+			PSOShadersPrecacheStats.PrecacheData.UpdateStats(MeshPassType, VertexFactoryType);
+			Value.bPrecached = true;
+		}
+	}
+
+	// Update the minimal PSO stats - optionally removing certain fields via PatchMinimalPipelineStateToCheck
+	{
+		FScopeLock Lock(&MinimalPSOPrecacheLock);
+		FGraphicsMinimalPipelineStateInitializer PSOInitializeToCheck  = PatchMinimalPipelineStateToCheck(PSOInitialize);
+		check(PSOInitializeToCheck.PrecachePSOHash != 0);
+
+		Experimental::FHashElementId TableId = MinimalPSOPrecacheMap.FindId(PSOInitializeToCheck.PrecachePSOHash);
+		if (!TableId.IsValid())
+		{
+			TableId = MinimalPSOPrecacheMap.FindOrAddId(PSOInitializeToCheck.PrecachePSOHash, FShaderStateUsage());
+
+		}
+
+		FShaderStateUsage& Value = MinimalPSOPrecacheMap.GetByElementId(TableId).Value;
+		if (!Value.bPrecached)
+		{
+			check(!Value.bUsed);
+			MinimalPSOPrecacheStats.PrecacheData.UpdateStats(MeshPassType, VertexFactoryType);
+			Value.bPrecached = true;
+		}
+	}
+}
+
+void PSOCollectorStats::CheckMinimalPipelineStateInCache(const FGraphicsMinimalPipelineStateInitializer& PSOInitialize, uint32 MeshPassType, const FVertexFactoryType* VertexFactoryType)
+{
+	if (!IsPrecachingValidationEnabled())
+	{
+		return;
+	}
+
+	const EShadingPath ShadingPath = FSceneInterface::GetShadingPath(GMaxRHIFeatureLevel);
+	bool bCollectPSOs = FPSOCollectorCreateManager::GetCreateFunction(ShadingPath, MeshPassType) != nullptr;
+	bool bTracked = bCollectPSOs && VertexFactoryType->SupportsPSOPrecaching();
+
+	// Enable this when we want to check the actual RHI shader resources during this function - slightly slower and don't need to be retrieved here yet
+	//PSOInitialize.BoundShaderState.AsBoundShaderState();
+	
+	// Update the Shader only stats ignoring all state
+	{
+		FScopeLock Lock(&PSOShadersPrecacheLock);
+		FGraphicsMinimalPipelineStateInitializer PSOInitializeToCheck;
+		PSOInitializeToCheck.BoundShaderState = PSOInitialize.BoundShaderState;
+		PSOInitializeToCheck.BoundShaderState.VertexDeclarationRHI = nullptr;
+		PSOInitializeToCheck.ComputePrecachePSOHash();
+
+		Experimental::FHashElementId TableId = PSOShadersPrecacheMap.FindId(PSOInitializeToCheck);
+		if (!TableId.IsValid())
+		{
+			TableId = PSOShadersPrecacheMap.FindOrAddId(PSOInitializeToCheck, FShaderStateUsage());
+		}
+
+		FShaderStateUsage& Value = PSOShadersPrecacheMap.GetByElementId(TableId).Value;
+		if (!Value.bUsed)
+		{
+			PSOShadersPrecacheStats.UsageData.UpdateStats(MeshPassType, VertexFactoryType);
+			if (!bTracked)
+			{
+				PSOShadersPrecacheStats.UntrackedData.UpdateStats(MeshPassType, VertexFactoryType);
+			}
+			else if (!Value.bPrecached)
+			{
+				PSOShadersPrecacheStats.MissData.UpdateStats(MeshPassType, VertexFactoryType);
+			}
+			else
+			{
+				PSOShadersPrecacheStats.HitData.UpdateStats(MeshPassType, VertexFactoryType);
+			}
+
+			Value.bUsed = true;
+		}
+	}
+
+	// Update the minimal PSO stats - optionally removing certain fields via PatchMinimalPipelineStateToCheck
+	{
+		FScopeLock Lock(&MinimalPSOPrecacheLock);
+		FGraphicsMinimalPipelineStateInitializer PSOInitializeToCheck = PatchMinimalPipelineStateToCheck(PSOInitialize);
+		check(PSOInitializeToCheck.PrecachePSOHash != 0);
+
+		Experimental::FHashElementId TableId = MinimalPSOPrecacheMap.FindId(PSOInitializeToCheck.PrecachePSOHash);
+		if (!TableId.IsValid())
+		{
+			TableId = MinimalPSOPrecacheMap.FindOrAddId(PSOInitializeToCheck.PrecachePSOHash, FShaderStateUsage());
+		}
+
+		FShaderStateUsage& Value = MinimalPSOPrecacheMap.GetByElementId(TableId).Value;
+		if (!Value.bUsed)
+		{
+			MinimalPSOPrecacheStats.UsageData.UpdateStats(MeshPassType, VertexFactoryType);
+			if (!bTracked)
+			{
+				MinimalPSOPrecacheStats.UntrackedData.UpdateStats(MeshPassType, VertexFactoryType);
+			}
+			else if (!Value.bPrecached)
+			{
+				MinimalPSOPrecacheStats.MissData.UpdateStats(MeshPassType, VertexFactoryType);
+			}
+			else
+			{
+				MinimalPSOPrecacheStats.HitData.UpdateStats(MeshPassType, VertexFactoryType);
+			}
+
+			Value.bUsed = true;
+		}
+	}
+}
+
+#endif // PSO_PRECACHING_VALIDATE
+
diff --git a/Engine/Source/Runtime/Renderer/Private/SceneTextures.cpp b/Engine/Source/Runtime/Renderer/Private/SceneTextures.cpp
--- a/Engine/Source/Runtime/Renderer/Private/SceneTextures.cpp
+++ b/Engine/Source/Runtime/Renderer/Private/SceneTextures.cpp
@@ -96,113 +96,6 @@ static EPixelFormat GetGBufferFFormat()
 	return NormalGBufferFormat;
 }
 
-static EPixelFormat GetDefaultMobileSceneColorLowPrecisionFormat()
-{
-	return (GEngine && GEngine->XRSystem.IsValid() && GEngine->StereoRenderingDevice.IsValid() && GEngine->StereoRenderingDevice->IsStandaloneStereoOnlyDevice()) ? PF_R8G8B8A8 : PF_B8G8R8A8;
-}
-
-static EPixelFormat GetMobileSceneColorFormat(const FSceneViewFamily& ViewFamily)
-{
-	bool bRequiresAlphaChannel = IsMobilePropagateAlphaEnabled(ViewFamily.GetShaderPlatform());
-
-	for (int32 ViewIndex = 0; ViewIndex < ViewFamily.Views.Num(); ViewIndex++)
-	{
-		// Planar reflections and scene captures use scene color alpha to keep track of where content has been rendered, for compositing into a different scene later
-		if (ViewFamily.Views[ViewIndex]->bIsPlanarReflection || ViewFamily.Views[ViewIndex]->bIsSceneCapture)
-		{
-			bRequiresAlphaChannel = true;
-		}
-	}
-	
-	EPixelFormat DefaultColorFormat;
-	const bool bUseLowPrecisionFormat = !IsMobileHDR() || !GSupportsRenderTargetFormat_PF_FloatRGBA;
-	if (bUseLowPrecisionFormat)
-	{
-		DefaultColorFormat = GetDefaultMobileSceneColorLowPrecisionFormat();
-	}
-	else
-	{
-		DefaultColorFormat = bRequiresAlphaChannel ? PF_FloatRGBA : PF_FloatR11G11B10;
-	}
-
-	check(GPixelFormats[DefaultColorFormat].Supported);
-
-	EPixelFormat Format = DefaultColorFormat;
-	static const auto CVar = IConsoleManager::Get().FindTConsoleVariableDataInt(TEXT("r.Mobile.SceneColorFormat"));
-	int32 MobileSceneColor = CVar->GetValueOnRenderThread();
-	switch (MobileSceneColor)
-	{
-	case 1:
-		Format = PF_FloatRGBA; break;
-	case 2:
-		Format = PF_FloatR11G11B10; break;
-	case 3:
-		if (bUseLowPrecisionFormat)
-		{
-			// if bUseLowPrecisionFormat, DefaultColorFormat already contains the value of GetDefaultMobileSceneColorLowPrecisionFormat
-			checkSlow(DefaultColorFormat == GetDefaultMobileSceneColorLowPrecisionFormat());
-			Format = DefaultColorFormat;
-		}
-		else
-		{
-			Format = GetDefaultMobileSceneColorLowPrecisionFormat();
-		}
-		break;
-	default:
-		break;
-	}
-
-	return GPixelFormats[Format].Supported ? Format : DefaultColorFormat;
-}
-
-static EPixelFormat GetSceneColorFormat(const FSceneViewFamily& ViewFamily)
-{
-	bool bRequiresAlphaChannel = false;
-
-	for (int32 ViewIndex = 0; ViewIndex < ViewFamily.Views.Num(); ViewIndex++)
-	{
-		// Planar reflections and scene captures use scene color alpha to keep track of where content has been rendered, for compositing into a different scene later
-		if (ViewFamily.Views[ViewIndex]->bIsPlanarReflection || ViewFamily.Views[ViewIndex]->bIsSceneCapture)
-		{
-			bRequiresAlphaChannel = true;
-		}
-	}
-
-	EPixelFormat Format = PF_FloatRGBA;
-
-	static const auto CVar = IConsoleManager::Get().FindTConsoleVariableDataInt(TEXT("r.SceneColorFormat"));
-
-	switch (CVar->GetValueOnRenderThread())
-	{
-	case 0:
-		Format = PF_R8G8B8A8; break;
-	case 1:
-		Format = PF_A2B10G10R10; break;
-	case 2:
-		Format = PF_FloatR11G11B10; break;
-	case 3:
-		Format = PF_FloatRGB; break;
-	case 4:
-		// default
-		break;
-	case 5:
-		Format = PF_A32B32G32R32F; break;
-	}
-
-	// Fallback in case the scene color selected isn't supported.
-	if (!GPixelFormats[Format].Supported)
-	{
-		Format = PF_FloatRGBA;
-	}
-
-	if (bRequiresAlphaChannel)
-	{
-		Format = PF_FloatRGBA;
-	}
-
-	return Format;
-}
-
 inline EPixelFormat GetMobileSceneDepthAuxPixelFormat(EShaderPlatform ShaderPlatform, bool bPreciseFormat)
 {
 	if (IsMobileDeferredShadingEnabled(ShaderPlatform) || bPreciseFormat)
@@ -224,37 +117,6 @@ inline EPixelFormat GetMobileSceneDepthAuxPixelFormat(EShaderPlatform ShaderPlat
 	return Format;
 }
 
-static uint32 GetEditorPrimitiveNumSamples(ERHIFeatureLevel::Type FeatureLevel)
-{
-	uint32 SampleCount = 1;
-
-	if (FeatureLevel >= ERHIFeatureLevel::SM5 && GRHISupportsMSAADepthSampleAccess)
-	{
-		static const auto CVar = IConsoleManager::Get().FindTConsoleVariableDataInt(TEXT("r.MSAA.CompositingSampleCount"));
-
-		SampleCount = CVar->GetValueOnRenderThread();
-
-		if (SampleCount <= 1)
-		{
-			SampleCount = 1;
-		}
-		else if (SampleCount <= 2)
-		{
-			SampleCount = 2;
-		}
-		else if (SampleCount <= 4)
-		{
-			SampleCount = 4;
-		}
-		else
-		{
-			SampleCount = 8;
-		}
-	}
-
-	return SampleCount;
-}
-
 static IStereoRenderTargetManager* FindStereoRenderTargetManager()
 {
 	if (!GEngine->StereoRenderingDevice.IsValid() || !GEngine->StereoRenderingDevice->IsStereoEnabled())
@@ -534,86 +396,33 @@ ENUM_CLASS_FLAGS(FSceneTextureExtentState::ERenderTargetHistory);
 
 void InitializeSceneTexturesConfig(FSceneTexturesConfig& Config, const FSceneViewFamily& ViewFamily)
 {
-	Config.FeatureLevel = ViewFamily.GetFeatureLevel();
-	Config.ShadingPath = FSceneInterface::GetShadingPath(Config.FeatureLevel);
-	Config.ShaderPlatform = GetFeatureLevelShaderPlatform(Config.FeatureLevel);
-	Config.Extent = FSceneTextureExtentState::Get().Compute(ViewFamily);
-	Config.NumSamples = GetDefaultMSAACount(Config.FeatureLevel, GDynamicRHI->RHIGetPlatformTextureMaxSampleCount());
-	Config.EditorPrimitiveNumSamples = GetEditorPrimitiveNumSamples(Config.FeatureLevel);
-	Config.ColorFormat = PF_Unknown;
-	Config.ColorClearValue = FClearValueBinding::Black;
-	Config.DepthClearValue = FClearValueBinding::DepthFar;
-	Config.bRequireMultiView = ViewFamily.bRequireMultiView;
-	Config.bIsUsingGBuffers = IsUsingGBuffers(Config.ShaderPlatform);
-
-	{
-		const bool bNeedsStereoAlloc = ViewFamily.Views.ContainsByPredicate([](const FSceneView* View)
-			{
-				return (IStereoRendering::IsStereoEyeView(*View) && (FindStereoRenderTargetManager() != nullptr));
-			});
-		Config.bSupportsXRTargetManagerDepthAlloc = bNeedsStereoAlloc ? 1 : 0;
-	}
-
-	switch (Config.ShadingPath)
-	{
-	case EShadingPath::Deferred:
-	{
-		Config.ColorFormat = GetSceneColorFormat(ViewFamily);
-		break;
-	}
-
-	case EShadingPath::Mobile:
-	{
-		Config.ColorFormat = GetMobileSceneColorFormat(ViewFamily);
-		break;
-	}
-
-	default:
-		checkNoEntry();
-	}
+	FIntPoint Extent = FSceneTextureExtentState::Get().Compute(ViewFamily);
+	EShadingPath ShadingPath = FSceneInterface::GetShadingPath(ViewFamily.GetFeatureLevel());
 
-	if (Config.bIsUsingGBuffers)
+	bool bRequiresAlphaChannel = ShadingPath == EShadingPath::Mobile ? IsMobilePropagateAlphaEnabled(ViewFamily.GetShaderPlatform()) : false;
+	for (int32 ViewIndex = 0; ViewIndex < ViewFamily.Views.Num(); ViewIndex++)
 	{
-		const FGBufferParams GBufferParams = FShaderCompileUtilities::FetchGBufferParamsRuntime(Config.ShaderPlatform);
-
-		// GBuffer configuration information is expensive to compute, the results are cached between runs.
-		struct FGBufferBindingCache
-		{
-			FGBufferParams GBufferParams;
-			FGBufferBinding GBufferA;
-			FGBufferBinding GBufferB;
-			FGBufferBinding GBufferC;
-			FGBufferBinding GBufferD;
-			FGBufferBinding GBufferE;
-			FGBufferBinding GBufferVelocity;
-			bool bInitialized = false;
-		};
-		static FGBufferBindingCache BindingCache;
-
-		if (!BindingCache.bInitialized || BindingCache.GBufferParams != GBufferParams)
+		// Planar reflections and scene captures use scene color alpha to keep track of where content has been rendered, for compositing into a different scene later
+		if (ViewFamily.Views[ViewIndex]->bIsPlanarReflection || ViewFamily.Views[ViewIndex]->bIsSceneCapture)
 		{
-			const FGBufferInfo GBufferInfo = FetchFullGBufferInfo(GBufferParams);
-
-			BindingCache.GBufferParams = GBufferParams;
-			BindingCache.GBufferA = FindGBufferBindingByName(GBufferInfo, TEXT("GBufferA"));
-			BindingCache.GBufferB = FindGBufferBindingByName(GBufferInfo, TEXT("GBufferB"));
-			BindingCache.GBufferC = FindGBufferBindingByName(GBufferInfo, TEXT("GBufferC"));
-			BindingCache.GBufferD = FindGBufferBindingByName(GBufferInfo, TEXT("GBufferD"));
-			BindingCache.GBufferE = FindGBufferBindingByName(GBufferInfo, TEXT("GBufferE"));
-			BindingCache.GBufferVelocity = FindGBufferBindingByName(GBufferInfo, TEXT("Velocity"));
-
-			BindingCache.bInitialized = true;
+			bRequiresAlphaChannel = true;
 		}
+	}
 
-		Config.GBufferA = BindingCache.GBufferA;
-		Config.GBufferB = BindingCache.GBufferB;
-		Config.GBufferC = BindingCache.GBufferC;
-		Config.GBufferD = BindingCache.GBufferD;
-		Config.GBufferE = BindingCache.GBufferE;
-		Config.GBufferVelocity = BindingCache.GBufferVelocity;
+	const bool bNeedsStereoAlloc = ViewFamily.Views.ContainsByPredicate([](const FSceneView* View)
+		{
+			return (IStereoRendering::IsStereoEyeView(*View) && (FindStereoRenderTargetManager() != nullptr));
+		});
 
-		Config.GBufferParams = GBufferParams;
-	}
+	FSceneTexturesConfigInitSettings SceneTexturesConfigInitSettings;
+	SceneTexturesConfigInitSettings.FeatureLevel = ViewFamily.GetFeatureLevel();
+	SceneTexturesConfigInitSettings.Extent = Extent;
+	SceneTexturesConfigInitSettings.bRequireMultiView = ViewFamily.bRequireMultiView;
+	SceneTexturesConfigInitSettings.bRequiresAlphaChannel = bRequiresAlphaChannel;
+	SceneTexturesConfigInitSettings.bSupportsXRTargetManagerDepthAlloc = bNeedsStereoAlloc ? 1 : 0;
+	SceneTexturesConfigInitSettings.ExtraSceneColorCreateFlags = GFastVRamConfig.SceneColor;
+	SceneTexturesConfigInitSettings.ExtraSceneDepthCreateFlags = GFastVRamConfig.SceneDepth;
+	Config.Init(SceneTexturesConfigInitSettings);
 }
 
 void FMinimalSceneTextures::InitializeViewFamily(FRDGBuilder& GraphBuilder, FViewFamilyInfo& ViewFamily)
@@ -636,23 +445,10 @@ void FMinimalSceneTextures::InitializeViewFamily(FRDGBuilder& GraphBuilder, FVie
 	}
 	else
 	{
-		ETextureCreateFlags Flags = TexCreate_DepthStencilTargetable | TexCreate_ShaderResource | TexCreate_InputAttachmentRead | GFastVRamConfig.SceneDepth;
-
-		ETextureCreateFlags MemorylessFlag = TexCreate_None;
-		if (!Config.bKeepDepthContent || (Config.NumSamples > 1 && Config.bMemorylessMSAA))
-		{
-			MemorylessFlag = TexCreate_Memoryless;
-		}
-
-		if (Config.NumSamples == 1 && GRHISupportsDepthUAV)
-		{
-			Flags |= TexCreate_UAV;
-		}
-
 		// TODO: Array-size could be values > 2, on upcoming XR devices. This should be part of the config.
 		FRDGTextureDesc Desc(Config.bRequireMultiView ?
-							 FRDGTextureDesc::Create2DArray(SceneTextures.Config.Extent, PF_DepthStencil, Config.DepthClearValue, Flags | MemorylessFlag, 2) :
-							 FRDGTextureDesc::Create2D(SceneTextures.Config.Extent, PF_DepthStencil, Config.DepthClearValue, Flags | MemorylessFlag));
+							 FRDGTextureDesc::Create2DArray(SceneTextures.Config.Extent, PF_DepthStencil, Config.DepthClearValue, Config.DepthCreateFlags, 2) :
+							 FRDGTextureDesc::Create2D(SceneTextures.Config.Extent, PF_DepthStencil, Config.DepthClearValue, Config.DepthCreateFlags));
 		Desc.NumSamples = Config.NumSamples;
 		SceneTextures.Depth = GraphBuilder.CreateTexture(Desc, TEXT("SceneDepthZ"));
 
@@ -676,29 +472,15 @@ void FMinimalSceneTextures::InitializeViewFamily(FRDGBuilder& GraphBuilder, FVie
 	// Scene Color
 	{
 		const bool bIsMobilePlatform = Config.ShadingPath == EShadingPath::Mobile;
-
-		ETextureCreateFlags Flags = TexCreate_RenderTargetable | TexCreate_ShaderResource | GFastVRamConfig.SceneColor;
 		const ETextureCreateFlags sRGBFlag = (bIsMobilePlatform && IsMobileColorsRGB()) ? TexCreate_SRGB : TexCreate_None;
 
-		if (Config.FeatureLevel >= ERHIFeatureLevel::SM5 && Config.NumSamples == 1)
-		{
-			Flags |= TexCreate_UAV;
-		}
-		
-		if (Config.NumSamples > 1 && Config.bMemorylessMSAA)
-		{
-			Flags |= TexCreate_Memoryless;
-		}
-
-		Flags |= sRGBFlag;
-
 		const TCHAR* SceneColorName = TEXT("SceneColor");
 
 		// Create the scene color.
 		// TODO: Array-size could be values > 2, on upcoming XR devices. This should be part of the config.
 		FRDGTextureDesc Desc(Config.bRequireMultiView ?
-							 FRDGTextureDesc::Create2DArray(Config.Extent, Config.ColorFormat, Config.ColorClearValue, Flags, 2) :
-							 FRDGTextureDesc::Create2D(Config.Extent, Config.ColorFormat, Config.ColorClearValue, Flags));
+							 FRDGTextureDesc::Create2DArray(Config.Extent, Config.ColorFormat, Config.ColorClearValue, Config.ColorCreateFlags, 2) :
+							 FRDGTextureDesc::Create2D(Config.Extent, Config.ColorFormat, Config.ColorClearValue, Config.ColorCreateFlags));
 		Desc.NumSamples = Config.NumSamples;
 		SceneTextures.Color = CreateTextureMSAA(GraphBuilder, Desc, SceneColorName, GFastVRamConfig.SceneColor | sRGBFlag);
 	}
diff --git a/Engine/Source/Runtime/Renderer/Public/MeshPassProcessor.h b/Engine/Source/Runtime/Renderer/Public/MeshPassProcessor.h
--- a/Engine/Source/Runtime/Renderer/Public/MeshPassProcessor.h
+++ b/Engine/Source/Runtime/Renderer/Public/MeshPassProcessor.h
@@ -9,6 +9,7 @@
 #include "MeshMaterialShader.h"
 #include "SceneUtils.h"
 #include "MeshBatch.h"
+#include "PSOPrecache.h"
 #include "Hash/CityHash.h"
 #include "Experimental/Containers/RobinHoodHashTable.h"
 #include <atomic>
@@ -20,6 +21,9 @@ class FInstanceCullingDrawParams;
 
 class FRayTracingLocalShaderBindingWriter;
 
+struct FMeshProcessorShaders;
+struct FSceneTexturesConfig;
+
 /** Mesh pass types supported. */
 namespace EMeshPass
 {
@@ -269,9 +273,14 @@ public:
 		, bDepthBounds(InMinimalState.bDepthBounds)
 		, DrawShadingRate(InMinimalState.DrawShadingRate)
 		, PrimitiveType(InMinimalState.PrimitiveType)
+		, PrecachePSOHash(InMinimalState.PrecachePSOHash)
+		, bPSOPrecached(InMinimalState.bPSOPrecached)
 	{
 	}
 
+	RENDERER_API void SetupBoundShaderState(FRHIVertexDeclaration* VertexDeclaration, const FMeshProcessorShaders& Shaders);
+	RENDERER_API void ComputePrecachePSOHash();
+
 	FGraphicsPipelineStateInitializer AsGraphicsPipelineStateInitializer() const
 	{	
 		return FGraphicsPipelineStateInitializer
@@ -300,6 +309,7 @@ public:
 			, MultiViewCount
 			, bHasFragmentDensityAttachment
 			, DrawShadingRate
+			, PrecachePSOHash
 		);
 	}
 
@@ -436,7 +446,12 @@ public:
 	bool							bHasFragmentDensityAttachment = false;
 	EVRSShadingRate					DrawShadingRate  = EVRSShadingRate::VRSSR_1x1;
 
-	EPrimitiveType			PrimitiveType;
+	EPrimitiveType					PrimitiveType;
+		
+	// Data hash of the minimal PSO which is used to optimize the computation of the full PSO
+	uint64							PrecachePSOHash = 0;
+	// Is the PSO is precached - updated at draw time and can be used to skip draw when still precaching
+	mutable bool					bPSOPrecached = false;
 };
 
 static_assert(sizeof(FMeshPassMask::Data) * 8 >= EMeshPass::Num, "FMeshPassMask::Data is too small to fit all mesh passes.");
@@ -640,6 +655,8 @@ struct FMeshDrawCommandDebugData
 	TShaderRef<FShader> VertexShader;
 	TShaderRef<FShader> PixelShader;
 	const FVertexFactory* VertexFactory;
+	const FVertexFactoryType* VertexFactoryType;
+	uint32 MeshPassType;
 	FName ResourceName;
 	FString MaterialName;
 #endif
@@ -981,7 +998,7 @@ public:
 	 * PSO
 	 */
 	FGraphicsMinimalPipelineStateId CachedPipelineId;
-
+	
 	/**
 	 * Draw command parameters
 	 */
@@ -1101,8 +1118,11 @@ public:
 		return uint32(CityHash64((char*)&HashKey, sizeof(FHashKey)));
 	}
 
-	/** Sets shaders on the mesh draw command and allocates room for the shader bindings. */
-	RENDERER_API void SetShaders(FRHIVertexDeclaration* VertexDeclaration, const FMeshProcessorShaders& Shaders, FGraphicsMinimalPipelineStateInitializer& PipelineState);
+	/** Allocates room for the shader bindings. */
+	inline void InitializeShaderBindings(const FMeshProcessorShaders& Shaders)
+	{
+		ShaderBindings.Initialize(Shaders);
+	}
 
 	inline void SetStencilRef(uint32 InStencilRef)
 	{
@@ -1124,8 +1144,11 @@ public:
 		ShaderBindings.Finalize(ShadersForDebugging);	
 	}
 
-	/** Submits the state and shader bindings to the RHI command list, but does not invoke the draw. */
-	static void SubmitDrawBegin(
+	/** 
+	 * Submits the state and shader bindings to the RHI command list, but does not invoke the draw. 	 
+	 * If function returns false, then draw can be skipped because the PSO used by the draw command is not compiled yet (still precaching)
+	 */
+	static bool SubmitDrawBegin(
 		const FMeshDrawCommand& RESTRICT MeshDrawCommand,
 		const FGraphicsMinimalPipelineStateSet& GraphicsMinimalPipelineStateSet,
 		// GPUCULL_TODO: Rename, and probably wrap in struct that links to GPU-Scene, maybe generalize (probably not)?
@@ -1141,8 +1164,11 @@ public:
 		FRHIBuffer* IndirectArgsOverrideBuffer = nullptr,
 		uint32 IndirectArgsOverrideByteOffset = 0U);
 
-	/** Submits the state and shader bindings to the RHI command list, but does not invoke the draw indirect. */
-	static void SubmitDrawIndirectBegin(
+	/** 	 
+	 * Submits the state and shader bindings to the RHI command list, but does not invoke the draw indirect. 
+	 * If function returns false, then draw will be skipped because PSO used by the draw command is not compiled yet (still precaching)
+	 */
+	static bool SubmitDrawIndirectBegin(
 		const FMeshDrawCommand& RESTRICT MeshDrawCommand,
 		const FGraphicsMinimalPipelineStateSet& GraphicsMinimalPipelineStateSet,
 		// GPUCULL_TODO: Rename, and probably wrap in struct that links to GPU-Scene, maybe generalize (probably not)?
@@ -1182,9 +1208,9 @@ public:
 		return Other.CachedPipelineId.GetId();
 	}
 #if MESH_DRAW_COMMAND_DEBUG_DATA
-	RENDERER_API void SetDebugData(const FPrimitiveSceneProxy* PrimitiveSceneProxy, const FMaterial* Material, const FMaterialRenderProxy* MaterialRenderProxy, const FMeshProcessorShaders& UntypedShaders, const FVertexFactory* VertexFactory);
+	RENDERER_API void SetDebugData(const FPrimitiveSceneProxy* PrimitiveSceneProxy, const FMaterial* Material, const FMaterialRenderProxy* MaterialRenderProxy, const FMeshProcessorShaders& UntypedShaders, const FVertexFactory* VertexFactory, uint32 MeshPassType);
 #else
-	void SetDebugData(const FPrimitiveSceneProxy* PrimitiveSceneProxy, const FMaterial* Material, const FMaterialRenderProxy* MaterialRenderProxy, const FMeshProcessorShaders& UntypedShaders, const FVertexFactory* VertexFactory){}
+	void SetDebugData(const FPrimitiveSceneProxy* PrimitiveSceneProxy, const FMaterial* Material, const FMaterialRenderProxy* MaterialRenderProxy, const FMeshProcessorShaders& UntypedShaders, const FVertexFactory* VertexFactory, uint32 MeshPassType){}
 #endif
 
 	SIZE_T GetAllocatedSize() const
@@ -1894,16 +1920,17 @@ ENUM_CLASS_FLAGS(EDrawingPolicyOverrideFlags);
 /** 
  * Base class of mesh processors, whose job is to transform FMeshBatch draw descriptions received from scene proxy implementations into FMeshDrawCommands ready for the RHI command list
  */
-class FMeshPassProcessor
+class FMeshPassProcessor : public IPSOCollector
 {
 public:
 
+	EMeshPass::Type MeshPassType;
 	const FScene* RESTRICT Scene;
 	ERHIFeatureLevel::Type FeatureLevel;
 	const FSceneView* ViewIfDynamicMeshCommand;
 	FMeshPassDrawListContext* DrawListContext;
-
-	RENDERER_API FMeshPassProcessor(const FScene* InScene, ERHIFeatureLevel::Type InFeatureLevel, const FSceneView* InViewIfDynamicMeshCommand, FMeshPassDrawListContext* InDrawListContext);
+		
+	RENDERER_API FMeshPassProcessor(/*EMeshPass::Type InMeshPassType,*/ const FScene* InScene, ERHIFeatureLevel::Type InFeatureLevel, const FSceneView* InViewIfDynamicMeshCommand, FMeshPassDrawListContext* InDrawListContext);
 
 	virtual ~FMeshPassProcessor() {}
 
@@ -1915,6 +1942,9 @@ public:
 	// FMeshPassProcessor interface
 	// Add a FMeshBatch to the pass
 	virtual void AddMeshBatch(const FMeshBatch& RESTRICT MeshBatch, uint64 BatchElementMask, const FPrimitiveSceneProxy* RESTRICT PrimitiveSceneProxy, int32 StaticMeshId = -1) = 0;
+		
+	// By default no PSOs collected 
+	virtual void CollectPSOInitializers(const FSceneTexturesConfig& SceneTexturesConfig, const FMaterial& Material, const FVertexFactoryType* VertexFactoryType, const FPSOPrecacheParams& PreCacheParams, TArray<FPSOPrecacheData>& PSOInitializers) override {}
 
 	static FORCEINLINE_DEBUGGABLE ERasterizerCullMode InverseCullMode(ERasterizerCullMode CullMode)
 	{
@@ -1946,6 +1976,19 @@ public:
 		EMeshPassFeatures MeshPassFeatures,
 		const ShaderElementDataType& ShaderElementData);
 
+	template<typename PassShadersType>
+	void AddGraphicsPipelineStateInitializer(
+		const FVertexFactoryType* RESTRICT VertexFactoryType,
+		const FMaterial& RESTRICT MaterialResource,
+		const FMeshPassProcessorRenderState& RESTRICT DrawRenderState,
+		const FGraphicsPipelineRenderTargetsInfo& RESTRICT RenderTargetsInfo,
+		PassShadersType PassShaders,
+		ERasterizerFillMode MeshFillMode,
+		ERasterizerCullMode MeshCullMode,
+		EPrimitiveType PrimitiveType,
+		EMeshPassFeatures MeshPassFeatures, 
+		TArray<FPSOPrecacheData>& PSOInitializers);
+
 protected:
 	UE_DEPRECATED(5.0, "This version of GetDrawCommandPrimitiveId is deprecated. Use the below function instead.")
 	RENDERER_API void GetDrawCommandPrimitiveId(
@@ -1964,6 +2007,27 @@ protected:
 	) const;
 };
 
+#if PSO_PRECACHING_VALIDATE
+
+namespace PSOCollectorStats
+{
+	// Stats for PSO shaders only (no state, only VS/PS/GS usage)
+	static FCriticalSection PSOShadersPrecacheLock;
+	static FPrecacheStats PSOShadersPrecacheStats;
+	static Experimental::TRobinHoodHashMap<FGraphicsMinimalPipelineStateInitializer, FShaderStateUsage> PSOShadersPrecacheMap;
+
+	// Stats for minimal PSO initializer data during MeshDrawCommand building (doesn't contain render target info)
+	static FCriticalSection MinimalPSOPrecacheLock;
+	static FPrecacheStats MinimalPSOPrecacheStats;		
+	static Experimental::TRobinHoodHashMap<uint64, FShaderStateUsage> MinimalPSOPrecacheMap;
+
+	// Add, check & track the minimal PSO initializer during precaching and MeshDrawCommand building
+	RENDERER_API extern void AddMinimalPipelineStateToCache(const FGraphicsMinimalPipelineStateInitializer& PipelineStateInitializer, uint32 MeshPassType, const FVertexFactoryType* VertexFactoryType);
+	RENDERER_API extern void CheckMinimalPipelineStateInCache(const FGraphicsMinimalPipelineStateInitializer& PSOInitialize, uint32 MeshPassType, const FVertexFactoryType* VertexFactoryType);
+}
+
+#endif // PSO_PRECACHING_VALIDATE
+
 typedef FMeshPassProcessor* (*PassProcessorCreateFunction)(const FScene* Scene, const FSceneView* InViewIfDynamicMeshCommand, FMeshPassDrawListContext* InDrawListContext);
 
 enum class EMeshPassFlags
@@ -2025,6 +2089,15 @@ private:
 	EMeshPass::Type PassType;
 };
 
+// Helper marco to register the mesh pass processor to both the FPassProcessorManager & the FPSOCollectorCreateManager
+#define REGISTER_MESHPASSPROCESSOR_AND_PSOCOLLECTOR(Name, MeshPassProcessorCreateFunction, ShadingPath, MeshPass, MeshPassFlags) \
+	FRegisterPassProcessorCreateFunction RegisterMeshPassProcesser##Name(&MeshPassProcessorCreateFunction, ShadingPath, MeshPass, MeshPassFlags); \
+	IPSOCollector* CreatePSOCollector##Name(ERHIFeatureLevel::Type FeatureLevel) \
+	{ \
+		return MeshPassProcessorCreateFunction(FeatureLevel, nullptr, nullptr, nullptr); \
+	} \
+	FRegisterPSOCollectorCreateFunction RegisterPSOCollector##Name(&CreatePSOCollector##Name, ShadingPath, (uint32) MeshPass);
+
 RENDERER_API extern void SubmitMeshDrawCommands(
 	const FMeshCommandOneFrameArray& VisibleMeshDrawCommands,
 	const FGraphicsMinimalPipelineStateSet& GraphicsMinimalPipelineStateSet, 
@@ -2065,6 +2138,11 @@ RENDERER_API extern void DrawDynamicMeshPassPrivate(
 
 RENDERER_API extern FMeshDrawCommandSortKey CalculateMeshStaticSortKey(const FMeshMaterialShader* VertexShader, const FMeshMaterialShader* PixelShader);
 
+RENDERER_API extern void AddRenderTargetInfo(EPixelFormat PixelFormat, ETextureCreateFlags CreateFlags, FGraphicsPipelineRenderTargetsInfo& RenderTargetsInfo);
+RENDERER_API extern void SetupDepthStencilInfo(EPixelFormat DepthStencilFormat, ETextureCreateFlags DepthStencilCreateFlags, ERenderTargetLoadAction DepthTargetLoadAction, ERenderTargetLoadAction StencilTargetLoadAction, FExclusiveDepthStencil DepthStencilAccess, FGraphicsPipelineRenderTargetsInfo& RenderTargetsInfo);
+RENDERER_API extern void SetupGBufferRenderTargetInfo(const FSceneTexturesConfig& SceneTexturesConfig, FGraphicsPipelineRenderTargetsInfo& RenderTargetsInfo, bool bSetupDepthStencil);
+RENDERER_API extern void ApplyTargetsInfo(FGraphicsPipelineStateInitializer& GraphicsPSOInit, const FGraphicsPipelineRenderTargetsInfo& RenderTargetsInfo);
+
 inline FMeshDrawCommandSortKey CalculateMeshStaticSortKey(const TShaderRef<FMeshMaterialShader>& VertexShader, const TShaderRef<FMeshMaterialShader>& PixelShader)
 {
 	return CalculateMeshStaticSortKey(VertexShader.GetShader(), PixelShader.GetShader());
diff --git a/Engine/Source/Runtime/Renderer/Public/MeshPassProcessor.inl b/Engine/Source/Runtime/Renderer/Public/MeshPassProcessor.inl
--- a/Engine/Source/Runtime/Renderer/Public/MeshPassProcessor.inl
+++ b/Engine/Source/Runtime/Renderer/Public/MeshPassProcessor.inl
@@ -80,7 +80,10 @@ void FMeshPassProcessor::BuildMeshDrawCommands(
 
 	check(!VertexFactory->NeedsDeclaration() || VertexDeclaration);
 
-	SharedMeshDrawCommand.SetShaders(VertexDeclaration, PassShaders.GetUntypedShaders(), PipelineState);
+	FMeshProcessorShaders MeshProcessorShaders = PassShaders.GetUntypedShaders();
+	PipelineState.SetupBoundShaderState(VertexDeclaration, MeshProcessorShaders);
+
+	SharedMeshDrawCommand.InitializeShaderBindings(MeshProcessorShaders);
 
 	PipelineState.RasterizerState = GetStaticRasterizerState<true>(MeshFillMode, MeshCullMode);
 
@@ -91,9 +94,19 @@ void FMeshPassProcessor::BuildMeshDrawCommands(
 	PipelineState.DepthStencilState = DrawRenderState.GetDepthStencilState();
 	PipelineState.DrawShadingRate = GetShadingRateFromMaterial(MaterialResource.GetShadingRate());
 
+	// PSO Precache hash only needed when PSO precaching is enabled
+	if (PipelineStateCache::IsPSOPrecachingEnabled())
+	{
+		PipelineState.ComputePrecachePSOHash();
+	}
+
 	check(VertexFactory && VertexFactory->IsInitialized());
 	VertexFactory->GetStreams(FeatureLevel, InputStreamType, SharedMeshDrawCommand.VertexStreams);
 
+#if PSO_PRECACHING_VALIDATE
+	PSOCollectorStats::CheckMinimalPipelineStateInCache(PipelineState, (uint32)MeshPassType, VertexFactory->GetType());
+#endif // PSO_PRECACHING_VALIDATE
+
 	SharedMeshDrawCommand.PrimitiveIdStreamIndex = VertexFactory->GetPrimitiveIdStreamIndex(FeatureLevel, InputStreamType);
 
 	if (SharedMeshDrawCommand.PrimitiveIdStreamIndex != INDEX_NONE)
@@ -120,7 +133,7 @@ void FMeshPassProcessor::BuildMeshDrawCommands(
 		PassShaders.GeometryShader->GetShaderBindings(Scene, FeatureLevel, PrimitiveSceneProxy, MaterialRenderProxy, MaterialResource, DrawRenderState, ShaderElementData, ShaderBindings);
 	}
 
-	SharedMeshDrawCommand.SetDebugData(PrimitiveSceneProxy, &MaterialResource, &MaterialRenderProxy, PassShaders.GetUntypedShaders(), VertexFactory);
+	SharedMeshDrawCommand.SetDebugData(PrimitiveSceneProxy, &MaterialResource, &MaterialRenderProxy, PassShaders.GetUntypedShaders(), VertexFactory, (uint32)MeshPassType);
 
 	const int32 NumElements = ShouldSkipMeshDrawCommand(MeshBatch, PrimitiveSceneProxy) ? 0 : MeshBatch.Elements.Num();
 
@@ -173,6 +186,64 @@ void FMeshPassProcessor::BuildMeshDrawCommands(
 	}
 }
 
+template<typename PassShadersType>
+void FMeshPassProcessor::AddGraphicsPipelineStateInitializer(
+	const FVertexFactoryType* RESTRICT VertexFactoryType,
+	const FMaterial& RESTRICT MaterialResource,
+	const FMeshPassProcessorRenderState& RESTRICT DrawRenderState,
+	const FGraphicsPipelineRenderTargetsInfo& RESTRICT RenderTargetsInfo,
+	PassShadersType PassShaders,
+	ERasterizerFillMode MeshFillMode,
+	ERasterizerCullMode MeshCullMode,
+	EPrimitiveType PrimitiveType,
+	EMeshPassFeatures MeshPassFeatures, 
+	TArray<FPSOPrecacheData>& PSOInitializers)
+{
+	FGraphicsMinimalPipelineStateInitializer MinimalPipelineStateInitializer;
+	MinimalPipelineStateInitializer.PrimitiveType = PrimitiveType;
+	
+	// Ignore immutable samplers for now - should be passed in?
+	//PipelineState.ImmutableSamplerState = MaterialRenderProxy.ImmutableSamplerState;
+	
+	EVertexInputStreamType InputStreamType = EVertexInputStreamType::Default;
+	if ((MeshPassFeatures & EMeshPassFeatures::PositionOnly) != EMeshPassFeatures::Default)				InputStreamType = EVertexInputStreamType::PositionOnly;
+	if ((MeshPassFeatures & EMeshPassFeatures::PositionAndNormalOnly) != EMeshPassFeatures::Default)	InputStreamType = EVertexInputStreamType::PositionAndNormalOnly;
+
+	FVertexDeclarationElementList Elements;
+	VertexFactoryType->GetPSOPrecacheVertexFetchElements(InputStreamType, Elements);
+	FRHIVertexDeclaration* VertexDeclaration = PipelineStateCache::GetOrCreateVertexDeclaration(Elements);
+	
+	FMeshProcessorShaders MeshProcessorShaders = PassShaders.GetUntypedShaders();
+	MinimalPipelineStateInitializer.SetupBoundShaderState(VertexDeclaration, MeshProcessorShaders);
+
+	MinimalPipelineStateInitializer.RasterizerState = GetStaticRasterizerState<true>(MeshFillMode, MeshCullMode);
+
+	check(DrawRenderState.GetDepthStencilState());
+	check(DrawRenderState.GetBlendState());
+
+	MinimalPipelineStateInitializer.BlendState = DrawRenderState.GetBlendState();
+	MinimalPipelineStateInitializer.DepthStencilState = DrawRenderState.GetDepthStencilState();
+	MinimalPipelineStateInitializer.DrawShadingRate = GetShadingRateFromMaterial(MaterialResource.GetShadingRate());
+
+	MinimalPipelineStateInitializer.ComputePrecachePSOHash();
+#if PSO_PRECACHING_VALIDATE
+	PSOCollectorStats::AddMinimalPipelineStateToCache(MinimalPipelineStateInitializer, (uint32)MeshPassType, VertexFactoryType);
+#endif // PSO_PRECACHING_VALIDATE
+
+	// NOTE: AsGraphicsPipelineStateInitializer will create the RHIShaders internally if they are not cached yet
+	FGraphicsPipelineStateInitializer PipelineStateInitializer = MinimalPipelineStateInitializer.AsGraphicsPipelineStateInitializer();
+	ApplyTargetsInfo(PipelineStateInitializer, RenderTargetsInfo);
+
+	FPSOPrecacheData PSOPrecacheData;
+	PSOPrecacheData.PSOInitializer = PipelineStateInitializer;
+#if PSO_PRECACHING_VALIDATE
+	PSOPrecacheData.MeshPassType = (uint32)MeshPassType;
+	PSOPrecacheData.VertexFactoryType = VertexFactoryType;
+#endif // PSO_PRECACHING_VALIDATE
+	PSOInitializers.Add(PSOPrecacheData);
+}
+
+
 /**
 * Provides a callback to build FMeshDrawCommands and then submits them immediately.  Useful for legacy / editor code paths.
 * Does many dynamic allocations - do not use for game rendering.
