Fix channel write mask renderstate not properly set when resolving mobile shadow projection in forward with distance field enabled

#rnx

diff --git a/Engine/Source/Runtime/Renderer/Private/ShadowRendering.cpp b/Engine/Source/Runtime/Renderer/Private/ShadowRendering.cpp
--- a/Engine/Source/Runtime/Renderer/Private/ShadowRendering.cpp
+++ b/Engine/Source/Runtime/Renderer/Private/ShadowRendering.cpp
@@ -2484,11 +2484,12 @@ void FMobileSceneRenderer::RenderMobileShadowProjections(
 		// Local light shadows don't render to shadow mask texture on mobile deferred
 		if (LightSceneProxy->GetLightType() == LightType_Directional || !IsMobileDeferredShadingEnabled(ShaderPlatform))
 		{
+			bool bProjectingForForwardShading = !IsMobileDistanceFieldEnabled(ShaderPlatform);
 			RenderShadowProjections(GraphBuilder, SceneTextures,
 				ScreenShadowMaskTexture,
 				nullptr,
 				LightSceneInfo,
-				true);
+				bProjectingForForwardShading);
 		}
 	}
 }
