Fix for translucency before dof to look correct when distortion is enabled (no self refraction or scene color leaking without transmittance).
The separated standard translucent is only ran when distortion is present and enabled with r.Translucency.StandardSeparated=1.
Only executed when distortion is present to avoid self distortion issues.
Modifed bSeparateTranslucencyModulate to be bTranslucencyModulate that then needs to be combined with Normal or Separate translucency tag.

diff --git a/Engine/Shaders/Private/DistortApplyScreenPS.usf b/Engine/Shaders/Private/DistortApplyScreenPS.usf
--- a/Engine/Shaders/Private/DistortApplyScreenPS.usf
+++ b/Engine/Shaders/Private/DistortApplyScreenPS.usf
@@ -29,11 +29,21 @@ SamplerState	SceneColorTextureSampler;
 Texture2DMS<float4> DistortionMSAATexture;
 Texture2DMS<float4> SceneColorMSAATexture;
 Texture2DMS<float4> RoughnessScatterMSAATexture;
+#if COMPOSITE_STANDARD_TRANSLUCENT
+Texture2DMS<float4> StandardTranslucentColorAlphaMSAATexture;
+Texture2DMS<float4> StandardTranslucentTransmittanceMSAATexture;
+#endif
 #else
 Texture2D		DistortionTexture;
 SamplerState	DistortionTextureSampler;
 Texture2D<float4> RoughnessScatterTexture;
 SamplerState      RoughnessScatterSampler;
+#if COMPOSITE_STANDARD_TRANSLUCENT
+Texture2D		StandardTranslucentColorAlphaTexture;
+SamplerState	StandardTranslucentColorAlphaTextureSampler;
+Texture2D		StandardTranslucentTransmittanceTexture;
+SamplerState	StandardTranslucentTransmittanceTextureSampler;
+#endif
 #endif
 
 float RefractionRoughnessToMipLevelFactor;
@@ -124,9 +134,30 @@ void Merge(
 	)
 {
 #if USE_MSAA
-	OutColor = SceneColorMSAATexture.Load(int2(View.BufferSizeAndInvSize.xy * TexCoord.xy), SampleIndex);
-#else
+
+	int2 PixCoord = int2(View.BufferSizeAndInvSize.xy * TexCoord.xy);
+	OutColor = SceneColorMSAATexture.Load(PixCoord, SampleIndex);
+#if COMPOSITE_STANDARD_TRANSLUCENT
+	const float4 TranslucentColorAlpha		= StandardTranslucentColorAlphaMSAATexture.Load(PixCoord, SampleIndex);
+	const float3 TranslucentTransmittance	= StandardTranslucentTransmittanceMSAATexture.Load(PixCoord, SampleIndex).rgb;
+#endif
+
+#else // USE_MSAA
+
 	OutColor = Texture2DSample(SceneColorTexture, SceneColorTextureSampler, TexCoord.xy);
+#if COMPOSITE_STANDARD_TRANSLUCENT
+	const float4 TranslucentColorAlpha		= Texture2DSample(StandardTranslucentColorAlphaTexture, StandardTranslucentColorAlphaTextureSampler, TexCoord.xy);
+	const float3 TranslucentTransmittance	= Texture2DSample(StandardTranslucentTransmittanceTexture, StandardTranslucentTransmittanceTextureSampler, TexCoord.xy).rgb;
+#endif
+
+#endif // USE_MSAA
+
+#if COMPOSITE_STANDARD_TRANSLUCENT
+	float GreyScaleModulateColorBackgroundVisibility = dot(TranslucentTransmittance, float3(1.0f / 3.0f, 1.0f / 3.0f, 1.0f / 3.0f));
+
+	OutColor = float4(
+		OutColor.rgb * TranslucentColorAlpha.a * TranslucentTransmittance + TranslucentColorAlpha.rgb,
+		OutColor.a   * TranslucentColorAlpha.a * GreyScaleModulateColorBackgroundVisibility);
 #endif
 }
 
diff --git a/Engine/Source/Runtime/Engine/Classes/Materials/MaterialRelevance.h b/Engine/Source/Runtime/Engine/Classes/Materials/MaterialRelevance.h
--- a/Engine/Source/Runtime/Engine/Classes/Materials/MaterialRelevance.h
+++ b/Engine/Source/Runtime/Engine/Classes/Materials/MaterialRelevance.h
@@ -23,7 +23,7 @@ struct ENGINE_API FMaterialRelevance
 			uint8 bHairStrands : 1;
 			uint8 bTwoSided : 1;
 			uint8 bSeparateTranslucency : 1; // Translucency After DOF
-			uint8 bSeparateTranslucencyModulate : 1;
+			uint8 bTranslucencyModulate : 1;
 			uint8 bPostMotionBlurTranslucency : 1;
 			uint8 bNormalTranslucency : 1;
 			uint8 bUsesSceneColorCopy : 1;
diff --git a/Engine/Source/Runtime/Engine/Private/Materials/MaterialInterface.cpp b/Engine/Source/Runtime/Engine/Private/Materials/MaterialInterface.cpp
--- a/Engine/Source/Runtime/Engine/Private/Materials/MaterialInterface.cpp
+++ b/Engine/Source/Runtime/Engine/Private/Materials/MaterialInterface.cpp
@@ -410,17 +410,15 @@ FMaterialRelevance UMaterialInterface::GetRelevance_Internal(const UMaterial* Ma
 
 			// If dual blending is supported, and we are rendering post-DOF translucency, then we also need to render a second pass to the modulation buffer.
 			// The modulation buffer can also be used for regular modulation shaders after DoF.
-			const bool bMaterialSeparateModulation =
-				(MaterialResource->IsDualBlendingEnabled(GShaderPlatformForFeatureLevel[InFeatureLevel]) || BlendMode == BLEND_Modulate)
-				&& TranslucencyPass == MTP_AfterDOF;
+			const bool bMaterialSeparateModulation = MaterialResource->IsDualBlendingEnabled(GShaderPlatformForFeatureLevel[InFeatureLevel]) || BlendMode == BLEND_Modulate || StrataBlendMode == SBM_ColoredTransmittanceOnly;
 
 			MaterialRelevance.bOpaque = !bIsTranslucent;
 			MaterialRelevance.bMasked = IsMasked();
 			MaterialRelevance.bDistortion = MaterialResource->IsDistorted();
 			MaterialRelevance.bHairStrands = IsCompatibleWithHairStrands(MaterialResource, InFeatureLevel);
 			MaterialRelevance.bTwoSided = MaterialResource->IsTwoSided();
-			MaterialRelevance.bSeparateTranslucency = (TranslucencyPass == MTP_AfterDOF) && StrataBlendMode != SBM_ColoredTransmittanceOnly;
-			MaterialRelevance.bSeparateTranslucencyModulate = bMaterialSeparateModulation || StrataBlendMode == SBM_ColoredTransmittanceOnly;
+			MaterialRelevance.bSeparateTranslucency = bIsTranslucent && (TranslucencyPass == MTP_AfterDOF);
+			MaterialRelevance.bTranslucencyModulate = bMaterialSeparateModulation;
 			MaterialRelevance.bPostMotionBlurTranslucency = (TranslucencyPass == MTP_AfterMotionBlur);
 			MaterialRelevance.bNormalTranslucency = bIsTranslucent && (TranslucencyPass == MTP_BeforeDOF);
 			MaterialRelevance.bDisableDepthTest = bIsTranslucent && Material->bDisableDepthTest;		
diff --git a/Engine/Source/Runtime/Engine/Private/SceneView.cpp b/Engine/Source/Runtime/Engine/Private/SceneView.cpp
--- a/Engine/Source/Runtime/Engine/Private/SceneView.cpp
+++ b/Engine/Source/Runtime/Engine/Private/SceneView.cpp
@@ -314,6 +314,14 @@ static TAutoConsoleVariable<int32> CVarAllowTranslucencyAfterDOF(
 	TEXT(" 1: on costs GPU performance and memory but keeps translucency unaffected by Depth of Field. (default)"),
 	ECVF_Scalability | ECVF_RenderThreadSafe);
 
+static TAutoConsoleVariable<int32> CVarTranslucencyStandardSeparated(
+	TEXT("r.Translucency.StandardSeparated"),
+	0,
+	TEXT("Render translucent meshes in separate buffer from the scene color.\n")
+	TEXT("This prevent those meshes from self refracting and leaking scnee color behind over edges when it should be affect by colored transmittance.\n")
+	TEXT("Forced disabled when r.SeparateTranslucency is 0.\n"),
+	ECVF_RenderThreadSafe | ECVF_Default);
+
 static TAutoConsoleVariable<int32> CVarEnableTemporalUpsample(
 	TEXT("r.TemporalAA.Upsampling"),
 	1,
@@ -2892,16 +2900,23 @@ FSceneViewFamily::FSceneViewFamily(const ConstructionValues& CVS)
 
 	// Check if the translucency are allowed to be rendered after DOF, if not, translucency after DOF will be rendered in standard translucency.
 	{
-		bAllowTranslucencyAfterDOF = CVarAllowTranslucencyAfterDOF.GetValueOnAnyThread() != 0
-			&& EngineShowFlags.PostProcessing // Used for reflection captures.
-			&& !UseDebugViewPS()
-			&& EngineShowFlags.SeparateTranslucency;
+		bool SeparateTranslucencyEnabled = EngineShowFlags.PostProcessing // Used for reflection captures.
+										&& !UseDebugViewPS()
+										&& EngineShowFlags.SeparateTranslucency;
 
-		if (GetFeatureLevel() == ERHIFeatureLevel::ES3_1)
+		const bool bIsMobile = GetFeatureLevel() == ERHIFeatureLevel::ES3_1;
+		if (bIsMobile)
 		{
 			const bool bMobileMSAA = GetDefaultMSAACount(ERHIFeatureLevel::ES3_1) > 1;
-			bAllowTranslucencyAfterDOF &= (IsMobileHDR() && !bMobileMSAA); // on <= ES3_1 separate translucency requires HDR on and MSAA off
+			SeparateTranslucencyEnabled &= (IsMobileHDR() && !bMobileMSAA); // on <= ES3_1 separate translucency requires HDR on and MSAA off
 		}
+
+		bAllowTranslucencyAfterDOF = SeparateTranslucencyEnabled && CVarAllowTranslucencyAfterDOF.GetValueOnAnyThread() != 0;
+
+		// We do not allow separated translucency on mobile
+		// When MSAA sample count is >1 it works, but hair has not been properly tested so far due to other issues, so MSAA cannot use separted standard translucent for now.
+		uint32 MSAASampleCount = GetDefaultMSAACount(GetFeatureLevel());
+		bAllowStandardTranslucencySeparated = SeparateTranslucencyEnabled && MSAASampleCount == 1 && !bIsMobile && CVarTranslucencyStandardSeparated.GetValueOnAnyThread() != 0;
 	}
 }
 
diff --git a/Engine/Source/Runtime/Engine/Public/PrimitiveViewRelevance.h b/Engine/Source/Runtime/Engine/Public/PrimitiveViewRelevance.h
--- a/Engine/Source/Runtime/Engine/Public/PrimitiveViewRelevance.h
+++ b/Engine/Source/Runtime/Engine/Public/PrimitiveViewRelevance.h
@@ -76,7 +76,7 @@ struct FPrimitiveViewRelevance : public FMaterialRelevance
 
 	bool HasTranslucency() const 
 	{
-		return bNormalTranslucency || bSeparateTranslucency || bSeparateTranslucencyModulate || bPostMotionBlurTranslucency;
+		return bNormalTranslucency || bSeparateTranslucency || bTranslucencyModulate || bPostMotionBlurTranslucency;
 	}
 
 	bool HasVelocity() const
diff --git a/Engine/Source/Runtime/Engine/Public/SceneView.h b/Engine/Source/Runtime/Engine/Public/SceneView.h
--- a/Engine/Source/Runtime/Engine/Public/SceneView.h
+++ b/Engine/Source/Runtime/Engine/Public/SceneView.h
@@ -1893,6 +1893,8 @@ public:
 
 	FORCEINLINE bool AllowTranslucencyAfterDOF() const { return bAllowTranslucencyAfterDOF; }
 
+	FORCEINLINE bool AllowStandardTranslucencySeparated() const { return bAllowStandardTranslucencySeparated; }
+
 	FORCEINLINE const ISceneViewFamilyScreenPercentage* GetScreenPercentageInterface() const
 	{
 		return ScreenPercentageInterface;
@@ -1975,6 +1977,9 @@ private:
 	/** whether the translucency are allowed to render after DOF, if not they will be rendered in standard translucency. */
 	bool bAllowTranslucencyAfterDOF;
 
+	/** whether the pre DOF translucency are allowed to be rendered in separated target from scene to allow for better composition with distortion.*/
+	bool bAllowStandardTranslucencySeparated;
+
 	PRAGMA_DISABLE_DEPRECATION_WARNINGS
 	// Only FSceneRenderer can copy a view family.
 	FSceneViewFamily(const FSceneViewFamily&) = default;
diff --git a/Engine/Source/Runtime/Renderer/Private/BasePassRendering.cpp b/Engine/Source/Runtime/Renderer/Private/BasePassRendering.cpp
--- a/Engine/Source/Runtime/Renderer/Private/BasePassRendering.cpp
+++ b/Engine/Source/Runtime/Renderer/Private/BasePassRendering.cpp
@@ -154,7 +154,7 @@ void SetTranslucentRenderState(FMeshPassProcessorRenderState& DrawRenderState, c
 				// Alpha is BF_Zero for source and BF_One for dest, which leaves alpha unchanged
 				DrawRenderState.SetBlendState(TStaticBlendState<CW_RGBA, BO_Add, BF_One, BF_Source1Color, BO_Add, BF_Zero, BF_One>::GetRHI());
 			}
-			else if (InTranslucencyPassType == ETranslucencyPass::TPT_TranslucencyAfterDOFModulate)
+			else if (InTranslucencyPassType == ETranslucencyPass::TPT_TranslucencyAfterDOFModulate || InTranslucencyPassType == ETranslucencyPass::TPT_StandardTranslucencyModulate)
 			{
 				// In the separate pass (after DoF) modulate, we want to only darken the target by our multiplication term, and ignore the addition term.
 				// For regular dual blending, our function is:
@@ -207,7 +207,7 @@ void SetTranslucentRenderState(FMeshPassProcessorRenderState& DrawRenderState, c
 				// Alpha is BF_Zero for source and BF_One for dest, which leaves alpha unchanged
 				DrawRenderState.SetBlendState(TStaticBlendState<CW_RGBA, BO_Add, BF_One, BF_Source1Color, BO_Add, BF_Zero, BF_One>::GetRHI());
 			}
-			else if (InTranslucencyPassType == ETranslucencyPass::TPT_TranslucencyAfterDOFModulate)
+			else if (InTranslucencyPassType == ETranslucencyPass::TPT_TranslucencyAfterDOFModulate || InTranslucencyPassType == ETranslucencyPass::TPT_StandardTranslucencyModulate)
 			{
 				// In the separate pass (after DoF) modulate, we want to only darken the target by our multiplication term, and ignore the addition term.
 				// For regular dual blending, our function is:
@@ -1467,6 +1467,11 @@ bool FBasePassMeshProcessor::TryAddMeshBatch(const FMeshBatch& RESTRICT MeshBatc
 				bShouldDraw = !Material.IsTranslucencyAfterDOFEnabled() && !Material.IsTranslucencyAfterMotionBlurEnabled();
 				break;
 
+			case ETranslucencyPass::TPT_StandardTranslucencyModulate:
+				bShouldDraw = !Material.IsTranslucencyAfterDOFEnabled() && !Material.IsTranslucencyAfterMotionBlurEnabled() 
+					&& (Material.IsDualBlendingEnabled(GetFeatureLevelShaderPlatform(FeatureLevel)) || BlendMode == BLEND_Modulate || StrataBlendMode == SBM_ColoredTransmittanceOnly);
+				break;
+
 			case ETranslucencyPass::TPT_TranslucencyAfterDOF:
 				bShouldDraw = Material.IsTranslucencyAfterDOFEnabled() && StrataBlendMode != SBM_ColoredTransmittanceOnly;
 				break;
@@ -1828,6 +1833,16 @@ FMeshPassProcessor* CreateTranslucencyStandardPassProcessor(const FScene* Scene,
 	return new FBasePassMeshProcessor(Scene, Scene->GetFeatureLevel(), InViewIfDynamicMeshCommand, PassDrawRenderState, InDrawListContext, Flags, ETranslucencyPass::TPT_StandardTranslucency);
 }
 
+FMeshPassProcessor* CreateTranslucencyStandardModulatePassProcessor(const FScene* Scene, const FSceneView* InViewIfDynamicMeshCommand, FMeshPassDrawListContext* InDrawListContext)
+{
+	FMeshPassProcessorRenderState PassDrawRenderState;
+	PassDrawRenderState.SetDepthStencilState(TStaticDepthStencilState<false, CF_DepthNearOrEqual>::GetRHI());
+
+	const FBasePassMeshProcessor::EFlags Flags = FBasePassMeshProcessor::EFlags::CanUseDepthStencil;
+
+	return new FBasePassMeshProcessor(Scene, Scene->GetFeatureLevel(), InViewIfDynamicMeshCommand, PassDrawRenderState, InDrawListContext, Flags, ETranslucencyPass::TPT_StandardTranslucencyModulate);
+}
+
 FMeshPassProcessor* CreateTranslucencyAfterDOFProcessor(const FScene* Scene, const FSceneView* InViewIfDynamicMeshCommand, FMeshPassDrawListContext* InDrawListContext)
 {
 	FMeshPassProcessorRenderState PassDrawRenderState;
@@ -1871,6 +1886,7 @@ FMeshPassProcessor* CreateTranslucencyAllPassProcessor(const FScene* Scene, cons
 
 FRegisterPassProcessorCreateFunction RegisterBasePass(&CreateBasePassProcessor, EShadingPath::Deferred, EMeshPass::BasePass, EMeshPassFlags::CachedMeshCommands | EMeshPassFlags::MainView);
 FRegisterPassProcessorCreateFunction RegisterTranslucencyStandardPass(&CreateTranslucencyStandardPassProcessor, EShadingPath::Deferred, EMeshPass::TranslucencyStandard, EMeshPassFlags::MainView);
+FRegisterPassProcessorCreateFunction RegisterTranslucencyStandardModulatePass(&CreateTranslucencyStandardModulatePassProcessor, EShadingPath::Deferred, EMeshPass::TranslucencyStandardModulate, EMeshPassFlags::MainView);
 FRegisterPassProcessorCreateFunction RegisterTranslucencyAfterDOFPass(&CreateTranslucencyAfterDOFProcessor, EShadingPath::Deferred, EMeshPass::TranslucencyAfterDOF, EMeshPassFlags::MainView);
 FRegisterPassProcessorCreateFunction RegisterTranslucencyAfterDOFModulatePass(&CreateTranslucencyAfterDOFModulateProcessor, EShadingPath::Deferred, EMeshPass::TranslucencyAfterDOFModulate, EMeshPassFlags::MainView);
 FRegisterPassProcessorCreateFunction RegisterTranslucencyAfterMotionBlurPass(&CreateTranslucencyAfterMotionBlurProcessor, EShadingPath::Deferred, EMeshPass::TranslucencyAfterMotionBlur, EMeshPassFlags::MainView);
diff --git a/Engine/Source/Runtime/Renderer/Private/DeferredShadingRenderer.cpp b/Engine/Source/Runtime/Renderer/Private/DeferredShadingRenderer.cpp
--- a/Engine/Source/Runtime/Renderer/Private/DeferredShadingRenderer.cpp
+++ b/Engine/Source/Runtime/Renderer/Private/DeferredShadingRenderer.cpp
@@ -3232,7 +3232,8 @@ void FDeferredShadingSceneRenderer::Render(FRDGBuilder& GraphBuilder)
 			RDG_CSV_STAT_EXCLUSIVE_SCOPE(GraphBuilder, RenderTranslucency);
 			SCOPE_CYCLE_COUNTER(STAT_TranslucencyDrawTime);
 			GraphBuilder.SetCommandListStat(GET_STATID(STAT_CLM_Translucency));
-			RenderTranslucency(GraphBuilder, SceneTextures, TranslucencyLightingVolumeTextures, &TranslucencyResourceMap, ETranslucencyView::UnderWater, InstanceCullingManager);
+			const bool bStandardTranslucentCanRenderSeparate = false;
+			RenderTranslucency(GraphBuilder, SceneTextures, TranslucencyLightingVolumeTextures, &TranslucencyResourceMap, ETranslucencyView::UnderWater, InstanceCullingManager, bStandardTranslucentCanRenderSeparate);
 			EnumRemoveFlags(TranslucencyViewsToRender, ETranslucencyView::UnderWater);
 		}
 
@@ -3294,7 +3295,7 @@ void FDeferredShadingSceneRenderer::Render(FRDGBuilder& GraphBuilder)
 	if (GetHairStrandsComposition() == EHairStrandsCompositionType::BeforeTranslucent)
 	{
 		RDG_GPU_STAT_SCOPE(GraphBuilder, HairRendering);
-		RenderHairComposition(GraphBuilder, Views, SceneTextures.Color.Target, SceneTextures.Depth.Target, SceneTextures.Velocity);
+		RenderHairComposition(GraphBuilder, Views, SceneTextures.Color.Target, SceneTextures.Depth.Target, SceneTextures.Velocity, TranslucencyResourceMap);
 	}
 
 	// Draw translucency.
@@ -3306,7 +3307,7 @@ void FDeferredShadingSceneRenderer::Render(FRDGBuilder& GraphBuilder)
 		RDG_EVENT_SCOPE(GraphBuilder, "Translucency");
 
 		// Raytracing doesn't need the distortion effect.
-		const bool bShouldRenderDistortion = TranslucencyViewsToRender != ETranslucencyView::RayTracing;
+		const bool bShouldRenderDistortion = TranslucencyViewsToRender != ETranslucencyView::RayTracing && ShouldRenderDistortion();
 
 #if RHI_RAYTRACING
 		if (EnumHasAnyFlags(TranslucencyViewsToRender, ETranslucencyView::RayTracing))
@@ -3333,23 +3334,26 @@ void FDeferredShadingSceneRenderer::Render(FRDGBuilder& GraphBuilder)
 			}
 		}
 
-		// Render all remaining translucency views.
-		GraphBuilder.SetCommandListStat(GET_STATID(STAT_CLM_Translucency));
-		RenderTranslucency(GraphBuilder, SceneTextures, TranslucencyLightingVolumeTextures, &TranslucencyResourceMap, TranslucencyViewsToRender, InstanceCullingManager);
-		TranslucencyViewsToRender = ETranslucencyView::None;
+		{
+			// Render all remaining translucency views.
+			GraphBuilder.SetCommandListStat(GET_STATID(STAT_CLM_Translucency));
+			const bool bStandardTranslucentCanRenderSeparate = bShouldRenderDistortion; // It is only needed to render standard translucent as separate when there is distortion (non self distortion of transmittance/specular/etc.)
+			RenderTranslucency(GraphBuilder, SceneTextures, TranslucencyLightingVolumeTextures, &TranslucencyResourceMap, TranslucencyViewsToRender, InstanceCullingManager, bStandardTranslucentCanRenderSeparate);
+			TranslucencyViewsToRender = ETranslucencyView::None;
+		}
 
 		// Compose hair before velocity/distortion pass since these pass write depth value, 
 		// and this would make the hair composition fails in this cases.
 		if (GetHairStrandsComposition() == EHairStrandsCompositionType::AfterTranslucent)
 		{
 			RDG_GPU_STAT_SCOPE(GraphBuilder, HairRendering);
-			RenderHairComposition(GraphBuilder, Views, SceneTextures.Color.Target, SceneTextures.Depth.Target, SceneTextures.Velocity);
+			RenderHairComposition(GraphBuilder, Views, SceneTextures.Color.Target, SceneTextures.Depth.Target, SceneTextures.Velocity, TranslucencyResourceMap);
 		}
 
 		if (bShouldRenderDistortion)
 		{
 			GraphBuilder.SetCommandListStat(GET_STATID(STAT_CLM_Distortion));
-			RenderDistortion(GraphBuilder, SceneTextures.Color.Target, SceneTextures.Depth.Target);
+			RenderDistortion(GraphBuilder, SceneTextures.Color.Target, SceneTextures.Depth.Target, TranslucencyResourceMap);
 		}
 
 		if (bShouldRenderVelocities)
@@ -3371,7 +3375,7 @@ void FDeferredShadingSceneRenderer::Render(FRDGBuilder& GraphBuilder)
 	else if (GetHairStrandsComposition() == EHairStrandsCompositionType::AfterTranslucent)
 	{
 		RDG_GPU_STAT_SCOPE(GraphBuilder, HairRendering);
-		RenderHairComposition(GraphBuilder, Views, SceneTextures.Color.Target, SceneTextures.Depth.Target, SceneTextures.Velocity);
+		RenderHairComposition(GraphBuilder, Views, SceneTextures.Color.Target, SceneTextures.Depth.Target, SceneTextures.Velocity, TranslucencyResourceMap);
 	}
 
 #if !UE_BUILD_SHIPPING
diff --git a/Engine/Source/Runtime/Renderer/Private/DeferredShadingRenderer.h b/Engine/Source/Runtime/Renderer/Private/DeferredShadingRenderer.h
--- a/Engine/Source/Runtime/Renderer/Private/DeferredShadingRenderer.h
+++ b/Engine/Source/Runtime/Renderer/Private/DeferredShadingRenderer.h
@@ -702,7 +702,8 @@ private:
 		const FTranslucencyLightingVolumeTextures& TranslucencyLightingVolumeTextures,
 		FTranslucencyPassResourcesMap* OutTranslucencyResourceMap,
 		ETranslucencyView ViewsToRender,
-		FInstanceCullingManager& InstanceCullingManager);
+		FInstanceCullingManager& InstanceCullingManager,
+		bool bStandardTranslucentCanRenderSeparate);
 
 	/** Renders the scene's translucency given a specific pass. */
 	void RenderTranslucencyInner(
@@ -714,7 +715,8 @@ private:
 		ETranslucencyView ViewsToRender,
 		FRDGTextureRef SceneColorCopyTexture,
 		ETranslucencyPass::Type TranslucencyPass,
-		FInstanceCullingManager& InstanceCullingManager);
+		FInstanceCullingManager& InstanceCullingManager,
+		bool bStandardTranslucentCanRenderSeparate);
 
 	/** Renders the scene's light shafts */
 	FRDGTextureRef RenderLightShaftOcclusion(
@@ -727,7 +729,10 @@ private:
 		FTranslucencyPassResourcesMap& OutTranslucencyResourceMap);
 
 	bool ShouldRenderDistortion() const;
-	void RenderDistortion(FRDGBuilder& GraphBuilder, FRDGTextureRef SceneColorTexture, FRDGTextureRef SceneDepthTexture);
+	void RenderDistortion(FRDGBuilder& GraphBuilder, 
+		FRDGTextureRef SceneColorTexture, 
+		FRDGTextureRef SceneDepthTexture,
+		FTranslucencyPassResourcesMap& TranslucencyResourceMap);
 
 	/** Renders capsule shadows for all per-object shadows using it for the given light. */
 	bool RenderCapsuleDirectShadows(
diff --git a/Engine/Source/Runtime/Renderer/Private/DistortionRendering.cpp b/Engine/Source/Runtime/Renderer/Private/DistortionRendering.cpp
--- a/Engine/Source/Runtime/Renderer/Private/DistortionRendering.cpp
+++ b/Engine/Source/Runtime/Renderer/Private/DistortionRendering.cpp
@@ -116,19 +116,26 @@ class FDistortionScreenPS : public FGlobalShader
 public:
 	class FUseMSAADim : SHADER_PERMUTATION_BOOL("USE_MSAA");
 	class FUseRoughRefractionDim : SHADER_PERMUTATION_BOOL("USE_ROUGH_REFRACTION");
-	using FPermutationDomain = TShaderPermutationDomain<FUseMSAADim, FUseRoughRefractionDim>;
+	class FCompositedStandardTranslucentDim : SHADER_PERMUTATION_BOOL("COMPOSITE_STANDARD_TRANSLUCENT");
+	using FPermutationDomain = TShaderPermutationDomain<FUseMSAADim, FUseRoughRefractionDim, FCompositedStandardTranslucentDim>;
 
 	BEGIN_SHADER_PARAMETER_STRUCT(FParameters, )
 		SHADER_PARAMETER_STRUCT_REF(FViewUniformShaderParameters, View)
 		SHADER_PARAMETER_RDG_TEXTURE(Texture2DMS<float4>, RoughnessScatterMSAATexture)
 		SHADER_PARAMETER_RDG_TEXTURE(Texture2DMS<float4>, DistortionMSAATexture)
 		SHADER_PARAMETER_RDG_TEXTURE(Texture2DMS<float4>, SceneColorMSAATexture)
+		SHADER_PARAMETER_RDG_TEXTURE(Texture2DMS<float4>, StandardTranslucentColorAlphaMSAATexture)
+		SHADER_PARAMETER_RDG_TEXTURE(Texture2DMS<float4>, StandardTranslucentTransmittanceMSAATexture)
 		SHADER_PARAMETER_RDG_TEXTURE(Texture2D, RoughnessScatterTexture)
 		SHADER_PARAMETER_RDG_TEXTURE(Texture2D, DistortionTexture)
 		SHADER_PARAMETER_RDG_TEXTURE(Texture2D, SceneColorTexture)
+		SHADER_PARAMETER_RDG_TEXTURE(Texture2D, StandardTranslucentColorAlphaTexture)
+		SHADER_PARAMETER_RDG_TEXTURE(Texture2D, StandardTranslucentTransmittanceTexture)
 		SHADER_PARAMETER_SAMPLER(SamplerState, RoughnessScatterSampler)
 		SHADER_PARAMETER_SAMPLER(SamplerState, DistortionTextureSampler)
 		SHADER_PARAMETER_SAMPLER(SamplerState, SceneColorTextureSampler)
+		SHADER_PARAMETER_SAMPLER(SamplerState, StandardTranslucentColorAlphaTextureSampler)
+		SHADER_PARAMETER_SAMPLER(SamplerState, StandardTranslucentTransmittanceTextureSampler)
 		SHADER_PARAMETER(float, RefractionRoughnessToMipLevelFactor)
 		RENDER_TARGET_BINDING_SLOTS()
 	END_SHADER_PARAMETER_STRUCT()
@@ -475,7 +482,11 @@ static void AddFilterSceneColorPass(FRDGBuilder& GraphBuilder, FFilterSceneColor
 
 //////////////////////////////////////////////////////////////////////////
 
-void FDeferredShadingSceneRenderer::RenderDistortion(FRDGBuilder& GraphBuilder, FRDGTextureRef SceneColorTexture, FRDGTextureRef SceneDepthTexture)
+void FDeferredShadingSceneRenderer::RenderDistortion(
+	FRDGBuilder& GraphBuilder, 
+	FRDGTextureRef SceneColorTexture, 
+	FRDGTextureRef SceneDepthTexture,
+	FTranslucencyPassResourcesMap& TranslucencyResourceMap)
 {
 	check(SceneDepthTexture);
 	check(SceneColorTexture);
@@ -731,11 +742,15 @@ void FDeferredShadingSceneRenderer::RenderDistortion(FRDGBuilder& GraphBuilder,
 	CommonParameters.SceneColorTextureSampler = bUseRoughRefraction ? TStaticSamplerState<SF_Trilinear>::GetRHI() : TStaticSamplerState<>::GetRHI();
 	CommonParameters.DistortionTextureSampler = TStaticSamplerState<>::GetRHI();
 	CommonParameters.RoughnessScatterSampler = TStaticSamplerState<>::GetRHI();
+	CommonParameters.StandardTranslucentColorAlphaTextureSampler = TStaticSamplerState<>::GetRHI();
+	CommonParameters.StandardTranslucentTransmittanceTextureSampler = TStaticSamplerState<>::GetRHI();
 	CommonParameters.RefractionRoughnessToMipLevelFactor = FMath::Max(0.0f, CVarRefractionRoughnessToMipLevelFactor.GetValueOnRenderThread());
 
+	const bool bAllowStandardTranslucencySeparated = ViewFamily.AllowStandardTranslucencySeparated();
 	FDistortionScreenPS::FPermutationDomain PermutationVector;
 	PermutationVector.Set<FDistortionScreenPS::FUseMSAADim>(SceneColorTexture->Desc.NumSamples > 1);
 	PermutationVector.Set<FDistortionScreenPS::FUseRoughRefractionDim>(bUseRoughRefraction);
+	PermutationVector.Set<FDistortionScreenPS::FCompositedStandardTranslucentDim>(bAllowStandardTranslucencySeparated);
 
 	TShaderMapRef<FScreenPassVS> VertexShader(ShaderMap);
 	TShaderMapRef<FDistortionApplyScreenPS> ApplyPixelShader(ShaderMap, PermutationVector);
@@ -753,13 +768,21 @@ void FDeferredShadingSceneRenderer::RenderDistortion(FRDGBuilder& GraphBuilder,
 		CommonParameters.RenderTargets.DepthStencil = StencilReadBinding;
 		PipelineState.PixelShader = ApplyPixelShader;
 
-		// Test against stencil mask but don't clear.
-		PipelineState.DepthStencilState = TStaticDepthStencilState<
-			false, CF_Always,
-			true, CF_Equal, SO_Keep, SO_Keep, SO_Keep,
-			false, CF_Always, SO_Keep, SO_Keep, SO_Keep,
-			kStencilMaskBit, kStencilMaskBit>::GetRHI();
-		PipelineState.StencilRef = kStencilMaskBit;
+		if (bAllowStandardTranslucencySeparated)
+		{
+			// We might have to compose translucent meshes that are not using any distortion so we cannot cull using stencil in this case.
+			PipelineState.DepthStencilState = TStaticDepthStencilState< false, CF_Always >::GetRHI();
+		}
+		else
+		{
+			// Test against stencil mask but don't clear.
+			PipelineState.DepthStencilState = TStaticDepthStencilState<
+				false, CF_Always,
+				true, CF_Equal, SO_Keep, SO_Keep, SO_Keep,
+				false, CF_Always, SO_Keep, SO_Keep, SO_Keep,
+				kStencilMaskBit, kStencilMaskBit>::GetRHI();
+			PipelineState.StencilRef = kStencilMaskBit;
+		}
 
 		ERenderTargetLoadAction LoadAction = ERenderTargetLoadAction::ENoAction;
 
@@ -801,13 +824,21 @@ void FDeferredShadingSceneRenderer::RenderDistortion(FRDGBuilder& GraphBuilder,
 		CommonParameters.RenderTargets.DepthStencil = StencilWriteBinding;
 		PipelineState.PixelShader = MergePixelShader;
 
-		// Test against stencil mask and clear it.
-		PipelineState.DepthStencilState = TStaticDepthStencilState<
-			false, CF_Always,
-			true, CF_Equal, SO_Keep, SO_Keep, SO_Zero,
-			false, CF_Always, SO_Keep, SO_Keep, SO_Keep,
-			kStencilMaskBit, kStencilMaskBit>::GetRHI();
-		PipelineState.StencilRef = kStencilMaskBit;
+		if (bAllowStandardTranslucencySeparated)
+		{
+			// We might have to compose translucent meshes that are not using any distortion so we cannot cull using stencil in this case.
+			PipelineState.DepthStencilState = TStaticDepthStencilState< false, CF_Always >::GetRHI();
+		}
+		else
+		{
+			// Test against stencil mask and clear it.
+			PipelineState.DepthStencilState = TStaticDepthStencilState<
+				false, CF_Always,
+				true, CF_Equal, SO_Keep, SO_Keep, SO_Zero,
+				false, CF_Always, SO_Keep, SO_Keep, SO_Keep,
+				kStencilMaskBit, kStencilMaskBit>::GetRHI();
+			PipelineState.StencilRef = kStencilMaskBit;
+		}
 
 		for (int32 ViewIndex = 0, Num = Views.Num(); ViewIndex < Num; ++ViewIndex)
 		{
@@ -820,6 +851,12 @@ void FDeferredShadingSceneRenderer::RenderDistortion(FRDGBuilder& GraphBuilder,
 			PassParameters->View = View.ViewUniformBuffer;
 			PassParameters->RenderTargets[0] = FRenderTargetBinding(SceneColorTexture, ERenderTargetLoadAction::ELoad);
 
+			// Fetch the separated before dof translucent data if available for composition
+			FTranslucencyPassResources& TranslucencyPassResources = TranslucencyResourceMap.Get(ViewIndex, ETranslucencyPass::TPT_StandardTranslucency);
+			PassParameters->StandardTranslucentColorAlphaTexture = TranslucencyPassResources.GetColorForRead(GraphBuilder);
+			PassParameters->StandardTranslucentColorAlphaMSAATexture = TranslucencyPassResources.ColorTexture.Target;
+			PassParameters->StandardTranslucentTransmittanceTexture = TranslucencyPassResources.GetColorModulateForRead(GraphBuilder);
+			PassParameters->StandardTranslucentTransmittanceMSAATexture = TranslucencyPassResources.ColorModulateTexture.Target;
 			Viewport.Rect = View.ViewRect;
 
 			ClearUnusedGraphResources(MergePixelShader, PassParameters);
diff --git a/Engine/Source/Runtime/Renderer/Private/HairStrands/HairStrandsComposition.cpp b/Engine/Source/Runtime/Renderer/Private/HairStrands/HairStrandsComposition.cpp
--- a/Engine/Source/Runtime/Renderer/Private/HairStrands/HairStrandsComposition.cpp
+++ b/Engine/Source/Runtime/Renderer/Private/HairStrands/HairStrandsComposition.cpp
@@ -613,19 +613,25 @@ void RenderHairComposition(
 	const TArray<FViewInfo>& Views,
 	FRDGTextureRef SceneColorTexture,
 	FRDGTextureRef SceneDepthTexture,
-	FRDGTextureRef SceneVelocityTexture)
+	FRDGTextureRef SceneVelocityTexture,
+	FTranslucencyPassResourcesMap& TranslucencyResourceMap)
 {
+	uint32 ViewIndex = 0;
 	for (const FViewInfo& View : Views)
 	{		
 		if (View.Family && HairStrands::HasViewHairStrandsData(View))
-		{			
+		{
+			FTranslucencyPassResources& TranslucencyPassResources = TranslucencyResourceMap.Get(ViewIndex, ETranslucencyPass::TPT_StandardTranslucency);
+			FRDGTextureRef SceneColorTarget = TranslucencyPassResources.ColorTexture.IsValid() ? TranslucencyPassResources.ColorTexture.Target : SceneColorTexture;
+
 			InternalRenderHairComposition(
 				GraphBuilder,
 				View,
-				SceneColorTexture,
+				SceneColorTarget,
 				SceneDepthTexture,
 				SceneVelocityTexture);
 		}
+		ViewIndex++;
 	}
 }
 
diff --git a/Engine/Source/Runtime/Renderer/Private/HairStrands/HairStrandsComposition.h b/Engine/Source/Runtime/Renderer/Private/HairStrands/HairStrandsComposition.h
--- a/Engine/Source/Runtime/Renderer/Private/HairStrands/HairStrandsComposition.h
+++ b/Engine/Source/Runtime/Renderer/Private/HairStrands/HairStrandsComposition.h
@@ -21,4 +21,5 @@ void RenderHairComposition(
 	const TArray<FViewInfo>& Views,
 	FRDGTextureRef SceneColorTexture,
 	FRDGTextureRef SceneDepthTexture,
-	FRDGTextureRef VelocityTexture);
+	FRDGTextureRef VelocityTexture,
+	struct FTranslucencyPassResourcesMap& TranslucencyResourceMap);
diff --git a/Engine/Source/Runtime/Renderer/Private/MeshDrawCommands.cpp b/Engine/Source/Runtime/Renderer/Private/MeshDrawCommands.cpp
--- a/Engine/Source/Runtime/Renderer/Private/MeshDrawCommands.cpp
+++ b/Engine/Source/Runtime/Renderer/Private/MeshDrawCommands.cpp
@@ -1184,6 +1184,7 @@ void FParallelMeshDrawCommandPass::DispatchPassSetup(
 	switch (PassType)
 	{
 		case EMeshPass::TranslucencyStandard: TaskContext.TranslucencyPass = ETranslucencyPass::TPT_StandardTranslucency; break;
+		case EMeshPass::TranslucencyStandardModulate: TaskContext.TranslucencyPass = ETranslucencyPass::TPT_StandardTranslucencyModulate; break;
 		case EMeshPass::TranslucencyAfterDOF: TaskContext.TranslucencyPass = ETranslucencyPass::TPT_TranslucencyAfterDOF; break;
 		case EMeshPass::TranslucencyAfterDOFModulate: TaskContext.TranslucencyPass = ETranslucencyPass::TPT_TranslucencyAfterDOFModulate; break;
 		case EMeshPass::TranslucencyAfterMotionBlur: TaskContext.TranslucencyPass = ETranslucencyPass::TPT_TranslucencyAfterMotionBlur; break;
diff --git a/Engine/Source/Runtime/Renderer/Private/RayTracing/RayTracingTranslucency.cpp b/Engine/Source/Runtime/Renderer/Private/RayTracing/RayTracingTranslucency.cpp
--- a/Engine/Source/Runtime/Renderer/Private/RayTracing/RayTracingTranslucency.cpp
+++ b/Engine/Source/Runtime/Renderer/Private/RayTracing/RayTracingTranslucency.cpp
@@ -144,7 +144,8 @@ bool ShouldRenderRayTracingTranslucency(const FViewInfo& View)
 
 void FDeferredShadingSceneRenderer::RenderRayTracingTranslucency(FRDGBuilder& GraphBuilder, FRDGTextureMSAA SceneColorTexture)
 {
-	if (!ShouldRenderTranslucency(ETranslucencyPass::TPT_StandardTranslucency)
+	if (   !ShouldRenderTranslucency(ETranslucencyPass::TPT_StandardTranslucency)
+		&& !ShouldRenderTranslucency(ETranslucencyPass::TPT_StandardTranslucencyModulate)
 		&& !ShouldRenderTranslucency(ETranslucencyPass::TPT_TranslucencyAfterDOF)
 		&& !ShouldRenderTranslucency(ETranslucencyPass::TPT_TranslucencyAfterDOFModulate)
 		&& !ShouldRenderTranslucency(ETranslucencyPass::TPT_TranslucencyAfterMotionBlur)
diff --git a/Engine/Source/Runtime/Renderer/Private/SceneRendering.cpp b/Engine/Source/Runtime/Renderer/Private/SceneRendering.cpp
--- a/Engine/Source/Runtime/Renderer/Private/SceneRendering.cpp
+++ b/Engine/Source/Runtime/Renderer/Private/SceneRendering.cpp
@@ -902,6 +902,7 @@ void FViewInfo::Init()
 	bSceneHasSkyMaterial = 0;
 	bHasSingleLayerWaterMaterial = 0;
 	bHasTranslucencySeparateModulation = 0;
+	bHasStandardTranslucencyModulation = 0;
 
 	NumVisibleStaticMeshElements = 0;
 	PrecomputedVisibilityData = 0;
diff --git a/Engine/Source/Runtime/Renderer/Private/SceneRendering.h b/Engine/Source/Runtime/Renderer/Private/SceneRendering.h
--- a/Engine/Source/Runtime/Renderer/Private/SceneRendering.h
+++ b/Engine/Source/Runtime/Renderer/Private/SceneRendering.h
@@ -1401,11 +1401,17 @@ public:
 	 */
 	uint32 bHasSingleLayerWaterMaterial : 1;
 	/**
-	 * true if the scene has at least one mesh with a material that needs dual blending AND is applied post DOF. If true,
-	 * that means we need to run the separate modulation render pass.
+	 * true if the scene has at least one mesh with a material that needs dual blending AND is applied post DOF. 
+	 * If true, that means we need to run the post-dof separate modulation render pass.
 	 */
 	uint32 bHasTranslucencySeparateModulation : 1;
 
+	/**
+	 * true if the scene has at least one mesh with a material that needs dual blending AND is applied before DOF. 
+	 * If true, that means we need to run the before-dof separate modulation render pass.
+	 */
+	uint32 bHasStandardTranslucencyModulation : 1;
+
 	/** Bitmask of all shading models used by primitives in this view */
 	uint16 ShadingModelMaskInView;
 
diff --git a/Engine/Source/Runtime/Renderer/Private/SceneVisibility.cpp b/Engine/Source/Runtime/Renderer/Private/SceneVisibility.cpp
--- a/Engine/Source/Runtime/Renderer/Private/SceneVisibility.cpp
+++ b/Engine/Source/Runtime/Renderer/Private/SceneVisibility.cpp
@@ -2290,6 +2290,7 @@ struct FRelevancePacket : public FSceneRenderingAllocatorObject<FRelevancePacket
 	bool bSceneHasSkyMaterial;
 	bool bHasSingleLayerWaterMaterial;
 	bool bHasTranslucencySeparateModulation;
+	bool bHasStandardTranslucencyModulation;
 
 	FRelevancePacket(
 		FRHICommandListImmediate& InRHICmdList,
@@ -2327,6 +2328,7 @@ struct FRelevancePacket : public FSceneRenderingAllocatorObject<FRelevancePacket
 		, bSceneHasSkyMaterial(false)
 		, bHasSingleLayerWaterMaterial(false)
 		, bHasTranslucencySeparateModulation(false)
+		, bHasStandardTranslucencyModulation(false)
 	{
 	}
 
@@ -2343,6 +2345,7 @@ struct FRelevancePacket : public FSceneRenderingAllocatorObject<FRelevancePacket
 		bSceneHasSkyMaterial = 0;
 		bHasSingleLayerWaterMaterial = 0;
 		bHasTranslucencySeparateModulation = 0;
+		bHasStandardTranslucencyModulation = 0;
 		bUsesGlobalDistanceField = false;
 		bUsesLightingChannels = false;
 		bTranslucentSurfaceLighting = false;
@@ -2434,12 +2437,17 @@ struct FRelevancePacket : public FSceneRenderingAllocatorObject<FRelevancePacket
 						TranslucentPrimCount.Add(ETranslucencyPass::TPT_StandardTranslucency, ViewRelevance.bUsesSceneColorCopy);
 					}
 
+					if (ViewRelevance.bNormalTranslucency && ViewRelevance.bTranslucencyModulate && View.Family->AllowStandardTranslucencySeparated())
+					{
+						TranslucentPrimCount.Add(ETranslucencyPass::TPT_StandardTranslucencyModulate, ViewRelevance.bUsesSceneColorCopy);
+					}
+
 					if (ViewRelevance.bSeparateTranslucency)
 					{
 						TranslucentPrimCount.Add(ETranslucencyPass::TPT_TranslucencyAfterDOF, ViewRelevance.bUsesSceneColorCopy);
 					}
 
-					if (ViewRelevance.bSeparateTranslucencyModulate)
+					if (ViewRelevance.bSeparateTranslucency && ViewRelevance.bTranslucencyModulate)
 					{
 						TranslucentPrimCount.Add(ETranslucencyPass::TPT_TranslucencyAfterDOFModulate, ViewRelevance.bUsesSceneColorCopy);
 					}
@@ -2470,7 +2478,8 @@ struct FRelevancePacket : public FSceneRenderingAllocatorObject<FRelevancePacket
 			bUsesCustomStencil |= (ViewRelevance.CustomDepthStencilUsageMask & (1 << 1)) > 0;
 			bSceneHasSkyMaterial |= ViewRelevance.bUsesSkyMaterial;
 			bHasSingleLayerWaterMaterial |= ViewRelevance.bUsesSingleLayerWaterMaterial;
-			bHasTranslucencySeparateModulation |= ViewRelevance.bSeparateTranslucencyModulate;
+			bHasStandardTranslucencyModulation |= ViewRelevance.bNormalTranslucency && ViewRelevance.bTranslucencyModulate && View.Family->AllowStandardTranslucencySeparated();
+			bHasTranslucencySeparateModulation |= ViewRelevance.bSeparateTranslucency && ViewRelevance.bTranslucencyModulate;
 
 			if (ViewRelevance.bRenderCustomDepth)
 			{
@@ -2775,12 +2784,17 @@ struct FRelevancePacket : public FSceneRenderingAllocatorObject<FRelevancePacket
 									DrawCommandPacket.AddCommandsForMesh(PrimitiveIndex, PrimitiveSceneInfo, StaticMeshRelevance, StaticMesh, Scene, bCanCache, EMeshPass::TranslucencyStandard);
 								}
 
+								if (ViewRelevance.bNormalTranslucency && ViewRelevance.bTranslucencyModulate && View.Family->AllowStandardTranslucencySeparated())
+								{
+									DrawCommandPacket.AddCommandsForMesh(PrimitiveIndex, PrimitiveSceneInfo, StaticMeshRelevance, StaticMesh, Scene, bCanCache, EMeshPass::TranslucencyStandardModulate);
+								}
+
 								if (ViewRelevance.bSeparateTranslucency)
 								{
 									DrawCommandPacket.AddCommandsForMesh(PrimitiveIndex, PrimitiveSceneInfo, StaticMeshRelevance, StaticMesh, Scene, bCanCache, EMeshPass::TranslucencyAfterDOF);
 								}
 
-								if (ViewRelevance.bSeparateTranslucencyModulate)
+								if (ViewRelevance.bSeparateTranslucency && ViewRelevance.bTranslucencyModulate)
 								{
 									DrawCommandPacket.AddCommandsForMesh(PrimitiveIndex, PrimitiveSceneInfo, StaticMeshRelevance, StaticMesh, Scene, bCanCache, EMeshPass::TranslucencyAfterDOFModulate);
 								}
@@ -2921,6 +2935,7 @@ struct FRelevancePacket : public FSceneRenderingAllocatorObject<FRelevancePacket
 		WriteView.bSceneHasSkyMaterial |= bSceneHasSkyMaterial;
 		WriteView.bHasSingleLayerWaterMaterial |= bHasSingleLayerWaterMaterial;
 		WriteView.bHasTranslucencySeparateModulation |= bHasTranslucencySeparateModulation;
+		WriteView.bHasStandardTranslucencyModulation |= bHasStandardTranslucencyModulation;
 		VisibleDynamicPrimitivesWithSimpleLights.AppendTo(WriteView.VisibleDynamicPrimitivesWithSimpleLights);
 		WriteView.NumVisibleDynamicPrimitives += NumVisibleDynamicPrimitives;
 		WriteView.NumVisibleDynamicEditorPrimitives += NumVisibleDynamicEditorPrimitives;
@@ -3254,13 +3269,19 @@ void ComputeDynamicMeshRelevance(EShadingPath ShadingPath, bool bAddLightmapDens
 				View.NumVisibleDynamicMeshElements[EMeshPass::TranslucencyStandard] += NumElements;
 			}
 
+			if (ViewRelevance.bNormalTranslucency && ViewRelevance.bTranslucencyModulate && View.Family->AllowStandardTranslucencySeparated())
+			{
+				PassMask.Set(EMeshPass::TranslucencyStandardModulate);
+				View.NumVisibleDynamicMeshElements[EMeshPass::TranslucencyStandardModulate] += NumElements;
+			}
+
 			if (ViewRelevance.bSeparateTranslucency)
 			{
 				PassMask.Set(EMeshPass::TranslucencyAfterDOF);
 				View.NumVisibleDynamicMeshElements[EMeshPass::TranslucencyAfterDOF] += NumElements;
 			}
 
-			if (ViewRelevance.bSeparateTranslucencyModulate)
+			if (ViewRelevance.bSeparateTranslucency && ViewRelevance.bTranslucencyModulate)
 			{
 				PassMask.Set(EMeshPass::TranslucencyAfterDOFModulate);
 				View.NumVisibleDynamicMeshElements[EMeshPass::TranslucencyAfterDOFModulate] += NumElements;
diff --git a/Engine/Source/Runtime/Renderer/Private/TranslucencyPass.h b/Engine/Source/Runtime/Renderer/Private/TranslucencyPass.h
--- a/Engine/Source/Runtime/Renderer/Private/TranslucencyPass.h
+++ b/Engine/Source/Runtime/Renderer/Private/TranslucencyPass.h
@@ -8,6 +8,7 @@ namespace ETranslucencyPass
 	enum Type
 	{
 		TPT_StandardTranslucency,
+		TPT_StandardTranslucencyModulate,
 		TPT_TranslucencyAfterDOF,
 		TPT_TranslucencyAfterDOFModulate,
 		TPT_TranslucencyAfterMotionBlur,
diff --git a/Engine/Source/Runtime/Renderer/Private/TranslucentRendering.cpp b/Engine/Source/Runtime/Renderer/Private/TranslucentRendering.cpp
--- a/Engine/Source/Runtime/Renderer/Private/TranslucentRendering.cpp
+++ b/Engine/Source/Runtime/Renderer/Private/TranslucentRendering.cpp
@@ -120,6 +120,7 @@ DynamicRenderScaling::FBudget GDynamicTranslucencyResolution(TEXT("DynamicTransl
 
 static const TCHAR* kTranslucencyPassName[] = {
 	TEXT("BeforeDistortion"),
+	TEXT("BeforeDistortionModulate"),
 	TEXT("AfterDOF"),
 	TEXT("AfterDOFModulate"),
 	TEXT("AfterMotionBlur"),
@@ -129,6 +130,7 @@ static_assert(UE_ARRAY_COUNT(kTranslucencyPassName) == int32(ETranslucencyPass::
 
 static const TCHAR* kTranslucencyColorTextureName[] = {
 	TEXT("Translucency.BeforeDistortion.Color"),
+	TEXT("Translucency.BeforeDistortion.Modulate"),
 	TEXT("Translucency.AfterDOF.Color"),
 	TEXT("Translucency.AfterDOF.Modulate"),
 	TEXT("Translucency.AfterMotionBlur.Color"),
@@ -148,6 +150,7 @@ EMeshPass::Type TranslucencyPassToMeshPass(ETranslucencyPass::Type TranslucencyP
 	switch (TranslucencyPass)
 	{
 	case ETranslucencyPass::TPT_StandardTranslucency: TranslucencyMeshPass = EMeshPass::TranslucencyStandard; break;
+	case ETranslucencyPass::TPT_StandardTranslucencyModulate: TranslucencyMeshPass = EMeshPass::TranslucencyStandardModulate; break;
 	case ETranslucencyPass::TPT_TranslucencyAfterDOF: TranslucencyMeshPass = EMeshPass::TranslucencyAfterDOF; break;
 	case ETranslucencyPass::TPT_TranslucencyAfterDOFModulate: TranslucencyMeshPass = EMeshPass::TranslucencyAfterDOFModulate; break;
 	case ETranslucencyPass::TPT_TranslucencyAfterMotionBlur: TranslucencyMeshPass = EMeshPass::TranslucencyAfterMotionBlur; break;
@@ -201,6 +204,7 @@ static bool IsSeparateTranslucencyEnabled(ETranslucencyPass::Type TranslucencyPa
 	// Currently AfterDOF is rendered earlier in the frame and must be rendered in a separate texture.
 	if (TranslucencyPass == ETranslucencyPass::TPT_TranslucencyAfterDOF
 		|| TranslucencyPass == ETranslucencyPass::TPT_TranslucencyAfterDOFModulate
+		|| TranslucencyPass == ETranslucencyPass::TPT_StandardTranslucencyModulate
 		|| TranslucencyPass == ETranslucencyPass::TPT_TranslucencyAfterMotionBlur
 		)
 	{
@@ -1234,7 +1238,8 @@ void FDeferredShadingSceneRenderer::RenderTranslucencyInner(
 	ETranslucencyView ViewsToRender,
 	FRDGTextureRef SceneColorCopyTexture,
 	ETranslucencyPass::Type TranslucencyPass,
-	FInstanceCullingManager& InstanceCullingManager)
+	FInstanceCullingManager& InstanceCullingManager,
+	bool bStandardTranslucentCanRenderSeparate)
 {
 	if (!ShouldRenderTranslucency(TranslucencyPass))
 	{
@@ -1243,11 +1248,12 @@ void FDeferredShadingSceneRenderer::RenderTranslucencyInner(
 
 	RDG_WAIT_FOR_TASKS_CONDITIONAL(GraphBuilder, IsTranslucencyWaitForTasksEnabled());
 
-	const bool bIsModulate = TranslucencyPass == ETranslucencyPass::TPT_TranslucencyAfterDOFModulate;
+	const bool bIsModulate = TranslucencyPass == ETranslucencyPass::TPT_TranslucencyAfterDOFModulate || TranslucencyPass == ETranslucencyPass::TPT_StandardTranslucencyModulate;
 	const bool bDepthTest = TranslucencyPass != ETranslucencyPass::TPT_TranslucencyAfterMotionBlur;
 	const bool bRenderInParallel = IsParallelTranslucencyEnabled();
 	const bool bIsScalingTranslucency = SeparateTranslucencyDimensions.Scale < 1.0f;
-	const bool bRenderInSeparateTranslucency = IsSeparateTranslucencyEnabled(TranslucencyPass, SeparateTranslucencyDimensions.Scale);
+	const bool bIsStandardSeparatedTranslucency = bStandardTranslucentCanRenderSeparate && TranslucencyPass == ETranslucencyPass::TPT_StandardTranslucency && ViewFamily.AllowStandardTranslucencySeparated();
+	const bool bRenderInSeparateTranslucency = IsSeparateTranslucencyEnabled(TranslucencyPass, SeparateTranslucencyDimensions.Scale) || bIsStandardSeparatedTranslucency;
 
 	// Can't reference scene color in scene textures. Scene color copy is used instead.
 	ESceneTextureSetupMode SceneTextureSetupMode = ESceneTextureSetupMode::All;
@@ -1274,7 +1280,7 @@ void FDeferredShadingSceneRenderer::RenderTranslucencyInner(
 			FIntRect ScaledViewRect = GetScaledRect(View.ViewRect, SeparateTranslucencyDimensions.Scale);
 
 			const FScreenPassTextureViewport SeparateTranslucencyViewport = SeparateTranslucencyDimensions.GetInstancedStereoViewport(View);
-			const bool bCompositeBackToSceneColor = IsMainTranslucencyPass(TranslucencyPass) || EnumHasAnyFlags(TranslucencyView, ETranslucencyView::UnderWater);
+			const bool bCompositeBackToSceneColor = (IsMainTranslucencyPass(TranslucencyPass) && !bIsStandardSeparatedTranslucency) || EnumHasAnyFlags(TranslucencyView, ETranslucencyView::UnderWater);
 			const bool bLumenGIEnabled = GetViewPipelineState(View).DiffuseIndirectMethod == EDiffuseIndirectMethod::Lumen;
 
 			/** Separate translucency color is either composited immediately or later during post processing. If done immediately, it's because the view doesn't support
@@ -1362,6 +1368,13 @@ void FDeferredShadingSceneRenderer::RenderTranslucencyInner(
 					ensure(TranslucencyPassResources.DepthTexture == SharedDepthTexture);
 					TranslucencyPassResources.ColorModulateTexture = SharedColorTexture;
 				}
+				else if (TranslucencyPass == ETranslucencyPass::TPT_StandardTranslucencyModulate)
+				{
+					FTranslucencyPassResources& TranslucencyPassResources = OutTranslucencyResourceMap->Get(ViewIndex, ETranslucencyPass::TPT_StandardTranslucency);
+					ensure(TranslucencyPassResources.ViewRect == ScaledViewRect);
+					ensure(TranslucencyPassResources.DepthTexture == SharedDepthTexture);
+					TranslucencyPassResources.ColorModulateTexture = SharedColorTexture;
+				}
 				else
 				{
 					check(!bIsModulate);
@@ -1423,7 +1436,8 @@ void FDeferredShadingSceneRenderer::RenderTranslucency(
 	const FTranslucencyLightingVolumeTextures& TranslucentLightingVolumeTextures,
 	FTranslucencyPassResourcesMap* OutTranslucencyResourceMap,
 	ETranslucencyView ViewsToRender,
-	FInstanceCullingManager& InstanceCullingManager)
+	FInstanceCullingManager& InstanceCullingManager,
+	bool bStandardTranslucentCanRenderSeparate)
 {
 	if (!EnumHasAnyFlags(ViewsToRender, ETranslucencyView::UnderWater | ETranslucencyView::AboveWater))
 	{
@@ -1490,18 +1504,23 @@ void FDeferredShadingSceneRenderer::RenderTranslucency(
 
 	if (ViewFamily.AllowTranslucencyAfterDOF())
 	{
-		RenderTranslucencyInner(GraphBuilder, SceneTextures, TranslucentLightingVolumeTextures, OutTranslucencyResourceMap, SharedDepthTexture, ViewsToRender, SceneColorCopyTexture, ETranslucencyPass::TPT_StandardTranslucency, InstanceCullingManager);
+		RenderTranslucencyInner(GraphBuilder, SceneTextures, TranslucentLightingVolumeTextures, OutTranslucencyResourceMap, SharedDepthTexture, ViewsToRender, SceneColorCopyTexture, ETranslucencyPass::TPT_StandardTranslucency, InstanceCullingManager, bStandardTranslucentCanRenderSeparate);
+		if (ViewFamily.AllowStandardTranslucencySeparated() && bStandardTranslucentCanRenderSeparate)
+		{
+			RenderTranslucencyInner(GraphBuilder, SceneTextures, TranslucentLightingVolumeTextures, OutTranslucencyResourceMap, SharedDepthTexture, ViewsToRender, SceneColorCopyTexture, ETranslucencyPass::TPT_StandardTranslucencyModulate, InstanceCullingManager, bStandardTranslucentCanRenderSeparate);
+		}
+
 		if (GetHairStrandsComposition() == EHairStrandsCompositionType::AfterTranslucentBeforeTranslucentAfterDOF)
 		{
-			RenderHairComposition(GraphBuilder, Views, SceneTextures.Color.Target, SceneTextures.Depth.Target, SceneTextures.Velocity);
+			RenderHairComposition(GraphBuilder, Views, SceneTextures.Color.Target, SceneTextures.Depth.Target, SceneTextures.Velocity, *OutTranslucencyResourceMap);
 		}
-		RenderTranslucencyInner(GraphBuilder, SceneTextures, TranslucentLightingVolumeTextures, OutTranslucencyResourceMap, SharedDepthTexture, ViewsToRender, SceneColorCopyTexture, ETranslucencyPass::TPT_TranslucencyAfterDOF, InstanceCullingManager);
-		RenderTranslucencyInner(GraphBuilder, SceneTextures, TranslucentLightingVolumeTextures, OutTranslucencyResourceMap, SharedDepthTexture, ViewsToRender, SceneColorCopyTexture, ETranslucencyPass::TPT_TranslucencyAfterDOFModulate, InstanceCullingManager);
-		RenderTranslucencyInner(GraphBuilder, SceneTextures, TranslucentLightingVolumeTextures, OutTranslucencyResourceMap, SharedDepthTexture, ViewsToRender, SceneColorCopyTexture, ETranslucencyPass::TPT_TranslucencyAfterMotionBlur, InstanceCullingManager);
+		RenderTranslucencyInner(GraphBuilder, SceneTextures, TranslucentLightingVolumeTextures, OutTranslucencyResourceMap, SharedDepthTexture, ViewsToRender, SceneColorCopyTexture, ETranslucencyPass::TPT_TranslucencyAfterDOF, InstanceCullingManager, bStandardTranslucentCanRenderSeparate);
+		RenderTranslucencyInner(GraphBuilder, SceneTextures, TranslucentLightingVolumeTextures, OutTranslucencyResourceMap, SharedDepthTexture, ViewsToRender, SceneColorCopyTexture, ETranslucencyPass::TPT_TranslucencyAfterDOFModulate, InstanceCullingManager, bStandardTranslucentCanRenderSeparate);
+		RenderTranslucencyInner(GraphBuilder, SceneTextures, TranslucentLightingVolumeTextures, OutTranslucencyResourceMap, SharedDepthTexture, ViewsToRender, SceneColorCopyTexture, ETranslucencyPass::TPT_TranslucencyAfterMotionBlur, InstanceCullingManager, bStandardTranslucentCanRenderSeparate);
 	}
 	else // Otherwise render translucent primitives in a single bucket.
 	{
-		RenderTranslucencyInner(GraphBuilder, SceneTextures, TranslucentLightingVolumeTextures, OutTranslucencyResourceMap, SharedDepthTexture, ViewsToRender, SceneColorCopyTexture, ETranslucencyPass::TPT_AllTranslucency, InstanceCullingManager);
+		RenderTranslucencyInner(GraphBuilder, SceneTextures, TranslucentLightingVolumeTextures, OutTranslucencyResourceMap, SharedDepthTexture, ViewsToRender, SceneColorCopyTexture, ETranslucencyPass::TPT_AllTranslucency, InstanceCullingManager, bStandardTranslucentCanRenderSeparate);
 	}
 
 	bool bUpscalePostDOFTranslucency = true;
diff --git a/Engine/Source/Runtime/Renderer/Public/MeshPassProcessor.h b/Engine/Source/Runtime/Renderer/Public/MeshPassProcessor.h
--- a/Engine/Source/Runtime/Renderer/Public/MeshPassProcessor.h
+++ b/Engine/Source/Runtime/Renderer/Public/MeshPassProcessor.h
@@ -40,6 +40,7 @@ namespace EMeshPass
 		Velocity,
 		TranslucentVelocity,
 		TranslucencyStandard,
+		TranslucencyStandardModulate,
 		TranslucencyAfterDOF,
 		TranslucencyAfterDOFModulate,
 		TranslucencyAfterMotionBlur,
@@ -82,6 +83,7 @@ inline const TCHAR* GetMeshPassName(EMeshPass::Type MeshPass)
 	case EMeshPass::Velocity: return TEXT("Velocity");
 	case EMeshPass::TranslucentVelocity: return TEXT("TranslucentVelocity");
 	case EMeshPass::TranslucencyStandard: return TEXT("TranslucencyStandard");
+	case EMeshPass::TranslucencyStandardModulate: return TEXT("TranslucencyStandardModulate");
 	case EMeshPass::TranslucencyAfterDOF: return TEXT("TranslucencyAfterDOF");
 	case EMeshPass::TranslucencyAfterDOFModulate: return TEXT("TranslucencyAfterDOFModulate");
 	case EMeshPass::TranslucencyAfterMotionBlur: return TEXT("TranslucencyAfterMotionBlur");
@@ -104,9 +106,9 @@ inline const TCHAR* GetMeshPassName(EMeshPass::Type MeshPass)
 	}
 
 #if WITH_EDITOR
-	static_assert(EMeshPass::Num == 24 + 4, "Need to update switch(MeshPass) after changing EMeshPass");
+	static_assert(EMeshPass::Num == 25 + 4, "Need to update switch(MeshPass) after changing EMeshPass");
 #else
-	static_assert(EMeshPass::Num == 24, "Need to update switch(MeshPass) after changing EMeshPass");
+	static_assert(EMeshPass::Num == 25, "Need to update switch(MeshPass) after changing EMeshPass");
 #endif
 
 	checkf(0, TEXT("Missing case for EMeshPass %u"), (uint32)MeshPass);
