Optimised separated translucent rendering by still having the distortion pass rely on stencil bit to know where it should act for distortion or to affect the scene color with other regular translucent element.

diff --git a/Engine/Source/Runtime/Renderer/Private/BasePassRendering.cpp b/Engine/Source/Runtime/Renderer/Private/BasePassRendering.cpp
--- a/Engine/Source/Runtime/Renderer/Private/BasePassRendering.cpp
+++ b/Engine/Source/Runtime/Renderer/Private/BasePassRendering.cpp
@@ -136,6 +136,34 @@ static bool IsBasePassWaitForTasksEnabled()
 	return CVarRHICmdFlushRenderThreadTasksBasePass.GetValueOnRenderThread() > 0 || CVarRHICmdFlushRenderThreadTasks.GetValueOnRenderThread() > 0;
 }
 
+static bool IsStandardTranslucenyPassSeparated()
+{
+	static const auto TranslucencyStandardSeparatedCVar = IConsoleManager::Get().FindTConsoleVariableDataInt(TEXT("r.Translucency.StandardSeparated"));
+	return TranslucencyStandardSeparatedCVar && TranslucencyStandardSeparatedCVar->GetValueOnRenderThread() != 0;
+}
+
+template<uint32 StencilRef> void SetTranslucentPassDepthStencilState(FMeshPassProcessorRenderState& DrawRenderState, bool bDisableDepthTest)
+{
+	if (bDisableDepthTest)
+	{
+		DrawRenderState.SetDepthStencilState(TStaticDepthStencilState<
+			false, CF_Always,
+			true , CF_Always, SO_Keep, SO_Keep, SO_Replace,
+			false, CF_Always, SO_Keep, SO_Keep, SO_Keep,
+			StencilRef, StencilRef>::GetRHI());
+		DrawRenderState.SetStencilRef(StencilRef);
+	}
+	else
+	{
+		DrawRenderState.SetDepthStencilState(TStaticDepthStencilState<
+			false, CF_DepthNearOrEqual,
+			true , CF_Always, SO_Keep, SO_Keep, SO_Replace,
+			false, CF_Always, SO_Keep, SO_Keep, SO_Keep,
+			StencilRef, StencilRef>::GetRHI());
+		DrawRenderState.SetStencilRef(StencilRef);
+	}
+}
+
 void SetTranslucentRenderState(FMeshPassProcessorRenderState& DrawRenderState, const FMaterial& Material, const EShaderPlatform Platform, ETranslucencyPass::Type InTranslucencyPassType)
 {
 	if (Material.IsStrataMaterial())
@@ -269,29 +297,24 @@ void SetTranslucentRenderState(FMeshPassProcessorRenderState& DrawRenderState, c
 	const bool bEnableResponsiveAA = Material.ShouldEnableResponsiveAA();
 	const bool bIsPostMotionBlur = Material.IsTranslucencyAfterMotionBlurEnabled();
 
+	// When separate standard translucent are used, we must mark the distoprtion bit for the composition to happen correctly for any BeforeDoF translucent.
+	const bool bSeparatedStandardTranslucent = (InTranslucencyPassType == ETranslucencyPass::TPT_StandardTranslucency || InTranslucencyPassType == ETranslucencyPass::TPT_AllTranslucency) && IsStandardTranslucenyPassSeparated();
+
 	if (bEnableResponsiveAA && !bIsPostMotionBlur)
 	{
-		if (bDisableDepthTest)
+		if (bSeparatedStandardTranslucent)
 		{
-			DrawRenderState.SetDepthStencilState(TStaticDepthStencilState<
-				false, CF_Always,
-				true, CF_Always, SO_Keep, SO_Keep, SO_Replace,
-				false, CF_Always, SO_Keep, SO_Keep, SO_Keep,
-				STENCIL_TEMPORAL_RESPONSIVE_AA_MASK, STENCIL_TEMPORAL_RESPONSIVE_AA_MASK
-			>::GetRHI());
-			DrawRenderState.SetStencilRef(STENCIL_TEMPORAL_RESPONSIVE_AA_MASK);
+			SetTranslucentPassDepthStencilState<STENCIL_TEMPORAL_RESPONSIVE_AA_MASK | DISTORTION_STENCIL_MASK_BIT>(DrawRenderState, bDisableDepthTest);
 		}
 		else
 		{
-			DrawRenderState.SetDepthStencilState(TStaticDepthStencilState<
-				false, CF_DepthNearOrEqual,
-				true, CF_Always, SO_Keep, SO_Keep, SO_Replace,
-				false, CF_Always, SO_Keep, SO_Keep, SO_Keep,
-				STENCIL_TEMPORAL_RESPONSIVE_AA_MASK, STENCIL_TEMPORAL_RESPONSIVE_AA_MASK
-			>::GetRHI());
-			DrawRenderState.SetStencilRef(STENCIL_TEMPORAL_RESPONSIVE_AA_MASK);
+			SetTranslucentPassDepthStencilState<STENCIL_TEMPORAL_RESPONSIVE_AA_MASK>(DrawRenderState, bDisableDepthTest);
 		}
 	}
+	else if (bSeparatedStandardTranslucent)
+	{
+		SetTranslucentPassDepthStencilState<DISTORTION_STENCIL_MASK_BIT>(DrawRenderState, bDisableDepthTest);
+	}
 	else if (bDisableDepthTest)
 	{
 		DrawRenderState.SetDepthStencilState(TStaticDepthStencilState<false, CF_Always>::GetRHI());
diff --git a/Engine/Source/Runtime/Renderer/Private/DistortionRendering.cpp b/Engine/Source/Runtime/Renderer/Private/DistortionRendering.cpp
--- a/Engine/Source/Runtime/Renderer/Private/DistortionRendering.cpp
+++ b/Engine/Source/Runtime/Renderer/Private/DistortionRendering.cpp
@@ -18,8 +18,6 @@
 
 DECLARE_GPU_DRAWCALL_STAT(Distortion);
 
-const uint8 kStencilMaskBit = STENCIL_SANDBOX_MASK;
-
 static TAutoConsoleVariable<int32> CVarDisableDistortion(
 	TEXT("r.DisableDistortion"),
 	0,
@@ -768,21 +766,13 @@ void FDeferredShadingSceneRenderer::RenderDistortion(
 		CommonParameters.RenderTargets.DepthStencil = StencilReadBinding;
 		PipelineState.PixelShader = ApplyPixelShader;
 
-		if (bAllowStandardTranslucencySeparated)
-		{
-			// We might have to compose translucent meshes that are not using any distortion so we cannot cull using stencil in this case.
-			PipelineState.DepthStencilState = TStaticDepthStencilState< false, CF_Always >::GetRHI();
-		}
-		else
-		{
-			// Test against stencil mask but don't clear.
-			PipelineState.DepthStencilState = TStaticDepthStencilState<
-				false, CF_Always,
-				true, CF_Equal, SO_Keep, SO_Keep, SO_Keep,
-				false, CF_Always, SO_Keep, SO_Keep, SO_Keep,
-				kStencilMaskBit, kStencilMaskBit>::GetRHI();
-			PipelineState.StencilRef = kStencilMaskBit;
-		}
+		// Test against stencil mask but don't clear.
+		PipelineState.DepthStencilState = TStaticDepthStencilState<
+			false, CF_Always,
+			true, CF_Equal, SO_Keep, SO_Keep, SO_Keep,
+			false, CF_Always, SO_Keep, SO_Keep, SO_Keep,
+			DISTORTION_STENCIL_MASK_BIT, DISTORTION_STENCIL_MASK_BIT>::GetRHI();
+		PipelineState.StencilRef = DISTORTION_STENCIL_MASK_BIT;
 
 		ERenderTargetLoadAction LoadAction = ERenderTargetLoadAction::ENoAction;
 
@@ -824,21 +814,13 @@ void FDeferredShadingSceneRenderer::RenderDistortion(
 		CommonParameters.RenderTargets.DepthStencil = StencilWriteBinding;
 		PipelineState.PixelShader = MergePixelShader;
 
-		if (bAllowStandardTranslucencySeparated)
-		{
-			// We might have to compose translucent meshes that are not using any distortion so we cannot cull using stencil in this case.
-			PipelineState.DepthStencilState = TStaticDepthStencilState< false, CF_Always >::GetRHI();
-		}
-		else
-		{
-			// Test against stencil mask and clear it.
-			PipelineState.DepthStencilState = TStaticDepthStencilState<
-				false, CF_Always,
-				true, CF_Equal, SO_Keep, SO_Keep, SO_Zero,
-				false, CF_Always, SO_Keep, SO_Keep, SO_Keep,
-				kStencilMaskBit, kStencilMaskBit>::GetRHI();
-			PipelineState.StencilRef = kStencilMaskBit;
-		}
+		// Test against stencil mask and clear it.
+		PipelineState.DepthStencilState = TStaticDepthStencilState<
+			false, CF_Always,
+			true, CF_Equal, SO_Keep, SO_Keep, SO_Zero,
+			false, CF_Always, SO_Keep, SO_Keep, SO_Keep,
+			DISTORTION_STENCIL_MASK_BIT, DISTORTION_STENCIL_MASK_BIT>::GetRHI();
+		PipelineState.StencilRef = DISTORTION_STENCIL_MASK_BIT;
 
 		for (int32 ViewIndex = 0, Num = Views.Num(); ViewIndex < Num; ++ViewIndex)
 		{
@@ -1009,9 +991,9 @@ FMeshPassProcessor* CreateDistortionPassProcessor(const FScene* Scene, const FSc
 		false, CF_DepthNearOrEqual,
 		true, CF_Always, SO_Keep, SO_Keep, SO_Replace,
 		false, CF_Always, SO_Keep, SO_Keep, SO_Keep,
-		kStencilMaskBit, kStencilMaskBit>::GetRHI());
+		DISTORTION_STENCIL_MASK_BIT, DISTORTION_STENCIL_MASK_BIT>::GetRHI());
 
-	DistortionPassState.SetStencilRef(kStencilMaskBit);
+	DistortionPassState.SetStencilRef(DISTORTION_STENCIL_MASK_BIT);
 
 	if (GetUseRoughRefraction())
 	{
@@ -1029,8 +1011,8 @@ FMeshPassProcessor* CreateDistortionPassProcessor(const FScene* Scene, const FSc
 		false, CF_Always,
 		true, CF_Always, SO_Keep, SO_Keep, SO_Replace,
 		false, CF_Always, SO_Keep, SO_Keep, SO_Keep,
-		kStencilMaskBit, kStencilMaskBit>::GetRHI());
-	DistortionPassStateNoDepthTest.SetStencilRef(kStencilMaskBit);
+		DISTORTION_STENCIL_MASK_BIT, DISTORTION_STENCIL_MASK_BIT>::GetRHI());
+	DistortionPassStateNoDepthTest.SetStencilRef(DISTORTION_STENCIL_MASK_BIT);
 
 	return new FDistortionMeshProcessor(Scene, InViewIfDynamicMeshCommand, DistortionPassState, DistortionPassStateNoDepthTest, InDrawListContext);
 }
diff --git a/Engine/Source/Runtime/Renderer/Private/TranslucentRendering.h b/Engine/Source/Runtime/Renderer/Private/TranslucentRendering.h
--- a/Engine/Source/Runtime/Renderer/Private/TranslucentRendering.h
+++ b/Engine/Source/Runtime/Renderer/Private/TranslucentRendering.h
@@ -9,6 +9,9 @@
 #include "VolumeRendering.h"
 #include "ScreenPass.h"
 #include "SystemTextures.h"
+#include "PostProcess/SceneRenderTargets.h"
+
+#define DISTORTION_STENCIL_MASK_BIT STENCIL_SANDBOX_MASK
 
 struct FSeparateTranslucencyDimensions
 {
