Works arround exessive TSR shader compilation time on SPIRV backends

diff --git a/Engine/Config/Mac/DataDrivenPlatformInfo.ini b/Engine/Config/Mac/DataDrivenPlatformInfo.ini
--- a/Engine/Config/Mac/DataDrivenPlatformInfo.ini
+++ b/Engine/Config/Mac/DataDrivenPlatformInfo.ini
@@ -32,7 +32,7 @@ bSupportsDebugViewShaders=true
 bSupportsAnisotropicMaterials=true
 bSupportsByteBufferComputeShaders=true
 bSupportsDistanceFields=true
-bSupportsGen5TemporalAA=false
+bSupportsGen5TemporalAA=true
 bSupportsLumenGI=true
 bSupportsSSDIndirect=true
 bSupportsWaveOperations=RuntimeDependent
@@ -41,6 +41,7 @@ bSupportsFFTBloom=true
 bSupportsDiaphragmDOF=true
 bSupportsVertexShaderLayer=true
 bSupportsDxc=true
+bIsSPIRV=true
 bSupportsInstancedStereo=true
 SupportsMultiViewport=RuntimeGuaranteed
 bSupportsVolumetricFog=true
diff --git a/Engine/Config/VulkanPC/DataDrivenPlatformInfo.ini b/Engine/Config/VulkanPC/DataDrivenPlatformInfo.ini
--- a/Engine/Config/VulkanPC/DataDrivenPlatformInfo.ini
+++ b/Engine/Config/VulkanPC/DataDrivenPlatformInfo.ini
@@ -107,12 +107,12 @@ bRequiresRenderTargetDuringRaster = true
 bRequiresDisableForwardLocalLights = true
 
 bCompileSignalProcessingPipeline = true
-bSupportsGen5TemporalAA=false
+bSupportsGen5TemporalAA=true
 
 ; //#todo-vulkanddpi
 bSupportsPerPixelDBufferMask = false
 
-bIsHlslcc = false
+bIsSPIRV=true
 
 bSupportsVariableRateShading = false
 
diff --git a/Engine/Config/Windows/DataDrivenPlatformInfo.ini b/Engine/Config/Windows/DataDrivenPlatformInfo.ini
--- a/Engine/Config/Windows/DataDrivenPlatformInfo.ini
+++ b/Engine/Config/Windows/DataDrivenPlatformInfo.ini
@@ -60,7 +60,7 @@ bSupportsSSDIndirect=true
 bRequiresVendorExtensionsForAtomics=true
 bSupportsAnisotropicMaterials=true
 bSupportsVariableRateShading=true
-bSupportsGen5TemporalAA=false
+bSupportsGen5TemporalAA=true
 bSupportsHairStrandGeometry=true
 bSupportsFFTBloom=true
 bSupportsDiaphragmDOF=true
diff --git a/Engine/Shaders/Private/TemporalSuperResolution/TSRRejectShading.usf b/Engine/Shaders/Private/TemporalSuperResolution/TSRRejectShading.usf
--- a/Engine/Shaders/Private/TemporalSuperResolution/TSRRejectShading.usf
+++ b/Engine/Shaders/Private/TemporalSuperResolution/TSRRejectShading.usf
@@ -3,7 +3,7 @@
 //------------------------------------------------------- CONFIG WAVE SIZE
 
 #if DIM_WAVE_SIZE == 0 // fallback to LDS
-	#define WAVE_BROADCAST_GROUPSIZE 32
+	#define WAVE_BROADCAST_GROUPSIZE (16 * 16)
 #endif
 
 //------------------------------------------------------- INCLUDES
@@ -45,11 +45,16 @@
 	#define LANE_STRIDE_X 4
 	#define LANE_STRIDE_Y 2
 
-#elif DIM_WAVE_SIZE == 0
+#elif DIM_WAVE_SIZE == 0 && 0
 	#define LANE_COUNT 32
 	#define LANE_STRIDE_X 4
 	#define LANE_STRIDE_Y 2
 
+#elif DIM_WAVE_SIZE == 0 && 1
+	#define LANE_COUNT (GROUP_TILE_SIZE * GROUP_TILE_SIZE)
+	#define LANE_STRIDE_X 1
+	#define LANE_STRIDE_Y 1
+
 #else
 	#error Unknown wave size
 #endif
@@ -422,10 +427,6 @@ void MainCS(
 	tsr_simd_halfC FilteredInput   = Convolve3x3HV<TAccumulatorAvg<tsr_simd_halfC>, tsr_simd_halfC>(ClampedInput  , bIsValid, GroupThreadIndex);
 	tsr_simd_halfC FilteredHistory = Convolve3x3HV<TAccumulatorAvg<tsr_simd_halfC>, tsr_simd_halfC>(ClampedHistory, bIsValid, GroupThreadIndex);
 	
-	Debug[0] = OriginalInput;
-	Debug[1] = AccessNeighborTexel(OriginalInput, tsr_short2(1, 0), GroupThreadIndex);
-	Debug[2] = FilteredInput;
-
 	// Measure rejection
 	tsr_simd_half Rejection;
 	Rejection.SetAllElements(tsr_half(1.0));
diff --git a/Engine/Shaders/Public/WaveBroadcastIntrinsics.ush b/Engine/Shaders/Public/WaveBroadcastIntrinsics.ush
--- a/Engine/Shaders/Public/WaveBroadcastIntrinsics.ush
+++ b/Engine/Shaders/Public/WaveBroadcastIntrinsics.ush
@@ -40,7 +40,7 @@
 	#define WAVE_BROADCAST_READ_LANE_AT 1
 	#define WAVE_BROADCAST_READ_LDS_AT 2
 	#define WAVE_BROADCAST_GCN_SWIZZLE 3
-	#define WAVE_BROADCAST_RDNA_ROTATE 4
+	#define WAVE_BROADCAST_ROTATE 4
 
 /** Compile time structure to choose which broadcasting should be done. */
 struct FWaveBroadcastSettings
@@ -89,9 +89,9 @@ uint GetWaveBroadcastSourceLaneIndex(const FWaveBroadcastSettings Settings, uint
 	{
 		return (((DestLaneIndex & Settings.SwizzleAnd) | Settings.SwizzleOr) ^ Settings.SwizzleXor) % 32;
 	}
-	else if (Settings.Operation == WAVE_BROADCAST_RDNA_ROTATE)
+	else if (Settings.Operation == WAVE_BROADCAST_ROTATE)
 	{
-		return ((Settings.RotateFixMask & DestLaneIndex) | ((~Settings.RotateFixMask) & uint(DestLaneIndex + Settings.Rotate))) % 32;
+		return (Settings.RotateFixMask & DestLaneIndex) | ((~Settings.RotateFixMask) & uint(DestLaneIndex + Settings.Rotate));
 	}
 	return DestLaneIndex;
 }
@@ -153,11 +153,11 @@ uint WaveBroadcast(const FWaveBroadcastSettings Settings, uint x)
 	#endif
 
 	#if PLATFORM_SUPPORTS_WAVE_ROTATE
-	else if (Settings.Operation == WAVE_BROADCAST_RDNA_ROTATE)
+	else if (Settings.Operation == WAVE_BROADCAST_ROTATE)
 	{
 		#if defined(COMPILER_SUPPORTS_WAVE_SWIZZLE_RDNA) && COMPILER_SUPPORTS_WAVE_SWIZZLE_RDNA
 		{
-			return WaveLaneRotateSwizzleRDNA(x, /* rotate_amount = */ Settings.Rotate, Settings.RotateFixMask);
+			return WaveLaneRotateSwizzleRDNA(x, /* rotate_amount = */ Settings.Rotate, Settings.RotateFixMask & 0x1F);
 		}
 		#elif FEATURE_LEVEL >= FEATURE_LEVEL_SM6 || PLATFORM_SUPPORTS_SM6_0_WAVE_OPERATIONS
 		{
@@ -498,7 +498,9 @@ FWaveBroadcastSettings InitWaveBroadcastLaneGroup(const uint LaneGroupSize, cons
 	return Settings;
 }
 
-/**  Rotate lane (-31 to +31) over a lane group (size is power of two in [2; 32]).
+/**  Rotate lane (-127 to +127) over a lane group (size is power of two in [2; 128]).
+ *
+ *  Notes: RDNA can only do this with LaneGroupSize <= 32
  *
  *	Requires PLATFORM_SUPPORTS_WAVE_ROTATE.
  *
@@ -515,12 +517,9 @@ FWaveBroadcastSettings InitWaveBroadcastLaneGroup(const uint LaneGroupSize, cons
 CALL_SITE_DEBUGLOC
 FWaveBroadcastSettings InitWaveRotateLaneGroup(const uint LaneGroupSize, const int LaneRotation)
 {
-	const uint RotateMask    = LaneGroupSize - 1;
-	const uint RotateFixMask = (~RotateMask) & 0x1F;
-
 	FWaveBroadcastSettings Settings;
-	Settings.Operation     = WAVE_BROADCAST_RDNA_ROTATE;
+	Settings.Operation     = WAVE_BROADCAST_ROTATE;
 	Settings.Rotate        = LaneRotation;
-	Settings.RotateFixMask = RotateFixMask;
+	Settings.RotateFixMask = ~(LaneGroupSize - 1);
 	return Settings;
 }
diff --git a/Engine/Source/Runtime/RHI/Private/RHI.cpp b/Engine/Source/Runtime/RHI/Private/RHI.cpp
--- a/Engine/Source/Runtime/RHI/Private/RHI.cpp
+++ b/Engine/Source/Runtime/RHI/Private/RHI.cpp
@@ -2529,6 +2529,7 @@ void FGenericDataDrivenShaderPlatformInfo::ParseDataDrivenShaderInfo(const FConf
 	GET_SECTION_BOOL_HELPER(bIsHlslcc);
 	GET_SECTION_BOOL_HELPER(bSupportsDxc);
 	GET_SECTION_BOOL_HELPER(bSupportsVariableRateShading);
+	GET_SECTION_BOOL_HELPER(bIsSPIRV);
 	GET_SECTION_INT_HELPER(NumberOfComputeThreads);
 
 	GET_SECTION_BOOL_HELPER(bWaterUsesSimpleForwardShading);
diff --git a/Engine/Source/Runtime/RHI/Public/RHIDefinitions.h b/Engine/Source/Runtime/RHI/Public/RHIDefinitions.h
--- a/Engine/Source/Runtime/RHI/Public/RHIDefinitions.h
+++ b/Engine/Source/Runtime/RHI/Public/RHIDefinitions.h
@@ -401,6 +401,7 @@ class RHI_API FGenericDataDrivenShaderPlatformInfo
 	uint32 bSupportsPerPixelDBufferMask : 1;
 	uint32 bIsHlslcc : 1;
 	uint32 bSupportsDxc : 1; // Whether DirectXShaderCompiler (DXC) is supported
+	uint32 bIsSPIRV : 1;
 	uint32 bSupportsVariableRateShading : 1;
 	uint32 NumberOfComputeThreads : 10;
 	uint32 bWaterUsesSimpleForwardShading : 1;
@@ -907,6 +908,12 @@ public:
 		return Infos[Platform].bSupportsDxc;
 	}
 
+	static FORCEINLINE_DEBUGGABLE const bool GetIsSPIRV(const FStaticShaderPlatform Platform)
+	{
+		check(IsValid(Platform));
+		return Infos[Platform].bIsSPIRV;
+	}
+
 	static FORCEINLINE_DEBUGGABLE const bool GetSupportsVariableRateShading(const FStaticShaderPlatform Platform)
 	{
 		check(IsValid(Platform));
diff --git a/Engine/Source/Runtime/Renderer/Private/PostProcess/TemporalSuperResolution.cpp b/Engine/Source/Runtime/Renderer/Private/PostProcess/TemporalSuperResolution.cpp
--- a/Engine/Source/Runtime/Renderer/Private/PostProcess/TemporalSuperResolution.cpp
+++ b/Engine/Source/Runtime/Renderer/Private/PostProcess/TemporalSuperResolution.cpp
@@ -213,6 +213,9 @@ FTSRHistoryUAVs CreateUAVs(FRDGBuilder& GraphBuilder, const FTSRHistoryTextures&
 class FTSRShader : public FGlobalShader
 {
 public:
+	static constexpr int32 kSupportMinWaveSize = 32;
+	static constexpr int32 kSupportMaxWaveSize = 64;
+	
 	FTSRShader(const ShaderMetaType::CompiledShaderInitializerType& Initializer)
 		: FGlobalShader(Initializer)
 	{ }
@@ -220,6 +223,35 @@ public:
 	FTSRShader()
 	{ }
 
+	static ERHIFeatureSupport SupportsWaveOps(EShaderPlatform Platform)
+	{
+		// UE-161125: SPIRV ends up with 5min+ compilation times, and D3D11 needs SPIRV backend for HLSL2021 features on FXC.
+		if (FDataDrivenShaderPlatformInfo::GetIsSPIRV(Platform) || Platform == SP_PCD3D_SM5)
+		{
+			return ERHIFeatureSupport::Unsupported;
+		}
+
+		return FDataDrivenShaderPlatformInfo::GetSupportsWaveOperations(Platform);
+	}
+
+	static bool SupportsLDS(EShaderPlatform Platform)
+	{
+		// Always support LDS if wave ops are not guarenteed
+		if (SupportsWaveOps(Platform) != ERHIFeatureSupport::RuntimeGuaranteed)
+		{
+			return true;
+		}
+
+		// Do not support LDS if shader supported wave size are guarenteed to support the platform.
+		if (FDataDrivenShaderPlatformInfo::GetMinimumWaveSize(Platform) >= kSupportMinWaveSize &&
+			FDataDrivenShaderPlatformInfo::GetMaximumWaveSize(Platform) <= kSupportMaxWaveSize)
+		{
+			return false;
+		}
+
+		return true;
+	}
+
 	static bool ShouldCompilePermutation(const FGlobalShaderPermutationParameters& Parameters)
 	{
 		return SupportsTSR(Parameters.Platform);
@@ -376,7 +408,7 @@ class FTSRRejectShadingCS : public FTSRShader
 		FPermutationDomain PermutationVector(Parameters.PermutationId);
 		int32 WaveSize = PermutationVector.Get<FWaveSizeOps>();
 
-		ERHIFeatureSupport WaveOpsSupport = FDataDrivenShaderPlatformInfo::GetSupportsWaveOperations(Parameters.Platform);
+		ERHIFeatureSupport WaveOpsSupport = FTSRShader::SupportsWaveOps(Parameters.Platform);
 		if (WaveSize != 0)
 		{
 			if (WaveOpsSupport == ERHIFeatureSupport::Unsupported)
@@ -392,12 +424,7 @@ class FTSRRejectShadingCS : public FTSRShader
 		}
 		else // if (WaveSize == LDS fallback)
 		{
-			if (WaveOpsSupport == ERHIFeatureSupport::RuntimeGuaranteed &&
-				FDataDrivenShaderPlatformInfo::GetMinimumWaveSize(Parameters.Platform) >= 32 &&
-				FDataDrivenShaderPlatformInfo::GetMaximumWaveSize(Parameters.Platform) <= 64)
-			{
-				return false;
-			}
+			return FTSRShader::SupportsLDS(Parameters.Platform);
 		}
 
 		return FTSRShader::ShouldCompilePermutation(Parameters);
@@ -666,8 +693,9 @@ ITemporalUpscaler::FOutputs AddTemporalSuperResolutionPasses(
 #endif
 
 	// Whether to use wave ops optimizations.
-	const ERHIFeatureSupport WaveOpsSupport = FDataDrivenShaderPlatformInfo::GetSupportsWaveOperations(View.GetShaderPlatform());
-	const bool bUseWaveOps = (CVarTSRWaveOps.GetValueOnRenderThread() != 0 && GRHISupportsWaveOperations && WaveOpsSupport == ERHIFeatureSupport::RuntimeDependent) || WaveOpsSupport == ERHIFeatureSupport::RuntimeGuaranteed;
+	const ERHIFeatureSupport WaveOpsSupport = FTSRShader::SupportsWaveOps(View.GetShaderPlatform());
+	const bool bSupportsLDS = FTSRShader::SupportsLDS(View.GetShaderPlatform());
+	const bool bUseWaveOps = (CVarTSRWaveOps.GetValueOnRenderThread() != 0 && GRHISupportsWaveOperations && bSupportsLDS && (WaveOpsSupport == ERHIFeatureSupport::RuntimeDependent || WaveOpsSupport == ERHIFeatureSupport::RuntimeGuaranteed)) || !bSupportsLDS;
 
 	// Whether alpha channel is supported.
 	const bool bSupportsAlpha = IsPostProcessingWithAlphaChannelSupported();
