Fixing static analysis warnings.

diff --git a/Engine/Source/Runtime/Renderer/Private/HeterogeneousVolumes/HeterogeneousVolumesLiveShadingPipeline.cpp b/Engine/Source/Runtime/Renderer/Private/HeterogeneousVolumes/HeterogeneousVolumesLiveShadingPipeline.cpp
--- a/Engine/Source/Runtime/Renderer/Private/HeterogeneousVolumes/HeterogeneousVolumesLiveShadingPipeline.cpp
+++ b/Engine/Source/Runtime/Renderer/Private/HeterogeneousVolumes/HeterogeneousVolumesLiveShadingPipeline.cpp
@@ -535,7 +535,10 @@ void RenderWithLiveShading(
 			check(LightSceneInfo != nullptr);
 
 			bApplyDirectLighting = (LightSceneInfo != nullptr);
-			bApplyShadowTransmittance = LightSceneInfo->Proxy->CastsVolumetricShadow();
+			if (LightSceneInfo)
+			{
+				bApplyShadowTransmittance = LightSceneInfo->Proxy->CastsVolumetricShadow();
+			}
 		}
 
 		if (HeterogeneousVolumes::UseTransmittanceVolume() && bApplyShadowTransmittance)
diff --git a/Engine/Source/Runtime/Renderer/Private/HeterogeneousVolumes/HeterogeneousVolumesPreshadingPipeline.cpp b/Engine/Source/Runtime/Renderer/Private/HeterogeneousVolumes/HeterogeneousVolumesPreshadingPipeline.cpp
--- a/Engine/Source/Runtime/Renderer/Private/HeterogeneousVolumes/HeterogeneousVolumesPreshadingPipeline.cpp
+++ b/Engine/Source/Runtime/Renderer/Private/HeterogeneousVolumes/HeterogeneousVolumesPreshadingPipeline.cpp
@@ -621,7 +621,10 @@ void RenderWithPreshadingCompute(
 				check(LightSceneInfo != nullptr);
 
 				bApplyDirectLighting = (LightSceneInfo != nullptr);
-				bApplyShadowTransmittance = LightSceneInfo->Proxy->CastsVolumetricShadow();
+				if (LightSceneInfo)
+				{
+					bApplyShadowTransmittance = LightSceneInfo->Proxy->CastsVolumetricShadow();
+				}
 			}
 
 			if (HeterogeneousVolumes::UseTransmittanceVolume() && bApplyShadowTransmittance)
@@ -778,7 +781,10 @@ void RenderWithPreshadingHardwareRayTracing(
 				check(LightSceneInfo != nullptr);
 
 				bApplyDirectLighting = (LightSceneInfo != nullptr);
-				bApplyShadowTransmittance = LightSceneInfo->Proxy->CastsVolumetricShadow();
+				if (LightSceneInfo)
+				{
+					bApplyShadowTransmittance = LightSceneInfo->Proxy->CastsVolumetricShadow();
+				}
 			}
 
 			if (HeterogeneousVolumes::UseTransmittanceVolume() && bApplyShadowTransmittance)
