Fixes TSR's shading rejection heuristic trying to compile on ES 3.1

diff --git a/Engine/Source/Runtime/Renderer/Private/PostProcess/TemporalSuperResolution.cpp b/Engine/Source/Runtime/Renderer/Private/PostProcess/TemporalSuperResolution.cpp
--- a/Engine/Source/Runtime/Renderer/Private/PostProcess/TemporalSuperResolution.cpp
+++ b/Engine/Source/Runtime/Renderer/Private/PostProcess/TemporalSuperResolution.cpp
@@ -405,6 +405,11 @@ class FTSRRejectShadingCS : public FTSRShader
 
 	static bool ShouldCompilePermutation(const FGlobalShaderPermutationParameters& Parameters)
 	{
+		if (!FTSRShader::ShouldCompilePermutation(Parameters))
+		{
+			return false;
+		}
+
 		FPermutationDomain PermutationVector(Parameters.PermutationId);
 		int32 WaveSize = PermutationVector.Get<FWaveSizeOps>();
 
@@ -427,7 +432,7 @@ class FTSRRejectShadingCS : public FTSRShader
 			return FTSRShader::SupportsLDS(Parameters.Platform);
 		}
 
-		return FTSRShader::ShouldCompilePermutation(Parameters);
+		return true;
 	}
 
 	static void ModifyCompilationEnvironment(const FGlobalShaderPermutationParameters& Parameters, FShaderCompilerEnvironment& OutEnvironment)
