Fix path tracer materials with AlphaHoldout blending mode

The logic for ShouldCompilePermutation did not exactly match the logic used at runtime, which could lead to the incorrect material being used.

diff --git a/Engine/Source/Runtime/Renderer/Private/PathTracing.cpp b/Engine/Source/Runtime/Renderer/Private/PathTracing.cpp
--- a/Engine/Source/Runtime/Renderer/Private/PathTracing.cpp
+++ b/Engine/Source/Runtime/Renderer/Private/PathTracing.cpp
@@ -1064,6 +1064,28 @@ FRayTracingLightFunctionMap GatherLightFunctionLightsPathTracing(FScene* Scene,
 	return RayTracingLightFunctionMap;
 }
 
+static bool NeedsAnyHitShader(bool bIsMasked, EBlendMode BlendMode)
+{
+	if (bIsMasked)
+	{
+		return true;
+	}
+	switch (BlendMode)
+	{
+	case BLEND_Opaque:
+	case BLEND_AlphaHoldout:
+		return false;
+	default:
+		return true;
+	}
+}
+
+static bool NeedsAnyHitShader(const FMaterial& RESTRICT MaterialResource)
+{
+	return NeedsAnyHitShader(MaterialResource.IsMasked(), MaterialResource.GetBlendMode());
+}
+
+
 template<bool UseAnyHitShader, bool UseIntersectionShader, bool UseSimplifiedShader>
 class TPathTracingMaterial : public FMeshMaterialShader
 {
@@ -1087,7 +1109,7 @@ public:
 			// does the VF support ray tracing at all?
 			return false;
 		}
-		if ((Parameters.MaterialParameters.bIsMasked || Parameters.MaterialParameters.BlendMode != BLEND_Opaque) != UseAnyHitShader)
+		if (NeedsAnyHitShader(Parameters.MaterialParameters.bIsMasked, Parameters.MaterialParameters.BlendMode) != UseAnyHitShader)
 		{
 			// the anyhit permutation is only required if the material is masked or has a non-opaque blend mode
 			return false;
@@ -1169,22 +1191,6 @@ IMPLEMENT_MATERIAL_SHADER_TYPE(template <>, FPathTracingMaterialCHS_AHS_IS, TEXT
 IMPLEMENT_MATERIAL_SHADER_TYPE(template <>, FGPULightmassCHS              , TEXT("/Engine/Private/PathTracing/PathTracingMaterialHitShader.usf"), TEXT("closesthit=PathTracingMaterialCHS"), SF_RayHitGroup);
 IMPLEMENT_MATERIAL_SHADER_TYPE(template <>, FGPULightmassCHS_AHS          , TEXT("/Engine/Private/PathTracing/PathTracingMaterialHitShader.usf"), TEXT("closesthit=PathTracingMaterialCHS anyhit=PathTracingMaterialAHS"), SF_RayHitGroup);
 
-static bool NeedsAnyHitShader(const FMaterial& RESTRICT MaterialResource)
-{
-	if (MaterialResource.IsMasked())
-	{
-		return true;
-	}
-	switch (MaterialResource.GetBlendMode())
-	{
-		case BLEND_Opaque:
-		case BLEND_AlphaHoldout:
-			return false;
-		default:
-			return true;
-	}
-}
-
 bool FRayTracingMeshProcessor::ProcessPathTracing(
 	const FMeshBatch& RESTRICT MeshBatch,
 	uint64 BatchElementMask,
