Fixes static analisis in TranslucentRendering.cpp

diff --git a/Engine/Source/Runtime/Renderer/Private/TranslucentRendering.cpp b/Engine/Source/Runtime/Renderer/Private/TranslucentRendering.cpp
--- a/Engine/Source/Runtime/Renderer/Private/TranslucentRendering.cpp
+++ b/Engine/Source/Runtime/Renderer/Private/TranslucentRendering.cpp
@@ -607,7 +607,7 @@ FScreenPassTexture FTranslucencyComposition::AddPass(
 		RDG_EVENT_NAME(
 			"%s(%s%s%s) %dx%d -> %dx%d",
 			OpName,
-			kTranslucencyPassName[int32(TranslucencyTextures.Pass)],
+			kTranslucencyPassName[FMath::Clamp(int32(TranslucencyTextures.Pass), 0, int32(ETranslucencyPass::TPT_MAX) - 1)],
 			bApplyModulateOnly ? TEXT(" ModulateOnly") : TEXT(""),
 			DepthUpscampling ? TEXT(" DepthUpsampling") : TEXT(""),
 			TranslucencyTextures.ViewRect.Width(), TranslucencyTextures.ViewRect.Height(),
