Fix crash when running hair iin forward with AA set to MSAA.

diff --git a/Engine/Source/Runtime/Renderer/Private/HairStrands/HairStrandsComposition.cpp b/Engine/Source/Runtime/Renderer/Private/HairStrands/HairStrandsComposition.cpp
--- a/Engine/Source/Runtime/Renderer/Private/HairStrands/HairStrandsComposition.cpp
+++ b/Engine/Source/Runtime/Renderer/Private/HairStrands/HairStrandsComposition.cpp
@@ -336,6 +336,7 @@ static void AddHairVisibilityFastResolveMaskPass(
 	FRDGTextureRef DummyTexture;
 	{
 		FRDGTextureDesc Desc = FRDGTextureDesc::Create2D(Resolution, PF_R8G8B8A8, FClearValueBinding::Black, TexCreate_RenderTargetable);
+		Desc.NumSamples = OutDepthTexture->Desc.NumSamples;
 		DummyTexture = GraphBuilder.CreateTexture(Desc, TEXT("Hair.DummyTexture"));
 	}
 
diff --git a/Engine/Source/Runtime/Renderer/Private/HairStrands/HairStrandsVisibility.cpp b/Engine/Source/Runtime/Renderer/Private/HairStrands/HairStrandsVisibility.cpp
--- a/Engine/Source/Runtime/Renderer/Private/HairStrands/HairStrandsVisibility.cpp
+++ b/Engine/Source/Runtime/Renderer/Private/HairStrands/HairStrandsVisibility.cpp
@@ -4679,6 +4679,12 @@ void RenderHairStrandsVisibilityBuffer(
 					SceneDepthTexture);
 			}
 		#endif
+
+			const bool bIsMSAAForwardEnabled = IsForwardShadingEnabled(View.GetShaderPlatform()) && SceneColorTexture->Desc.NumSamples > 1;
+			if (bIsMSAAForwardEnabled)
+			{		
+				VisibilityData.ResolveMaskTexture = nullptr;
+			}
 		}
 	}
 }
diff --git a/Engine/Source/Runtime/Renderer/Private/LightGridInjection.cpp b/Engine/Source/Runtime/Renderer/Private/LightGridInjection.cpp
--- a/Engine/Source/Runtime/Renderer/Private/LightGridInjection.cpp
+++ b/Engine/Source/Runtime/Renderer/Private/LightGridInjection.cpp
@@ -885,6 +885,7 @@ void FDeferredShadingSceneRenderer::RenderForwardShadowProjections(
 			ForwardScreenSpaceShadowMask = CreateTextureMSAA(GraphBuilder, Desc, TEXT("ShadowMaskTexture"), GFastVRamConfig.ScreenSpaceShadowMask);
 			if (bIsHairEnable)
 			{
+				Desc.NumSamples = 1;
 				ForwardScreenSpaceShadowMaskSubPixel = CreateTextureMSAA(GraphBuilder, Desc, TEXT("ShadowMaskSubPixelTexture"), GFastVRamConfig.ScreenSpaceShadowMask);
 			}
 		}
@@ -948,8 +949,7 @@ void FDeferredShadingSceneRenderer::RenderForwardShadowProjections(
 
 		if (bIsHairEnable)
 		{
-			PassParameters->RenderTargets[1] = FRenderTargetBinding(ForwardScreenSpaceShadowMaskSubPixel.Target, ForwardScreenSpaceShadowMaskSubPixel.Resolve, ERenderTargetLoadAction::ELoad);
-			OutForwardScreenSpaceShadowMaskSubPixel = ForwardScreenSpaceShadowMaskSubPixel.Resolve;
+			OutForwardScreenSpaceShadowMaskSubPixel = ForwardScreenSpaceShadowMaskSubPixel.Target;
 		}
 
 		GraphBuilder.AddPass(RDG_EVENT_NAME("ResolveScreenSpaceShadowMask"), PassParameters, ERDGPassFlags::Raster, [](FRHICommandList&) {});
