Make VSM clipmap CornerOffset relative and upgrade WorldCenter to 64-bit type to fix directional shadow LWC glitches

diff --git a/Engine/Shaders/Private/VirtualShadowMaps/VirtualShadowMapPageCacheCommon.ush b/Engine/Shaders/Private/VirtualShadowMaps/VirtualShadowMapPageCacheCommon.ush
--- a/Engine/Shaders/Private/VirtualShadowMaps/VirtualShadowMapPageCacheCommon.ush
+++ b/Engine/Shaders/Private/VirtualShadowMaps/VirtualShadowMapPageCacheCommon.ush
@@ -10,6 +10,8 @@ struct FShadowMapCacheData
 	// ID of the corresponding virtual shadow map from the previous frame, or INDEX_NONE if
 	// there wasn't one.
 	int PrevVirtualShadowMapId;
+
+	int2 ClipmapCornerOffsetDelta;
 };
 
 StructuredBuffer<FShadowMapCacheData> ShadowMapCacheData;
diff --git a/Engine/Shaders/Private/VirtualShadowMaps/VirtualShadowMapPageManagement.usf b/Engine/Shaders/Private/VirtualShadowMaps/VirtualShadowMapPageManagement.usf
--- a/Engine/Shaders/Private/VirtualShadowMaps/VirtualShadowMapPageManagement.usf
+++ b/Engine/Shaders/Private/VirtualShadowMaps/VirtualShadowMapPageManagement.usf
@@ -493,6 +493,8 @@ void CreateCachedPageMappings(uint VirtualShadowMapId, uint GlobalPageOffset, ui
 		const bool bCacheStaticSeparately = VirtualShadowMapShouldCacheStaticSeparately();
 
 		const uint PrevVirtualShadowMapId = ShadowMapCacheData[VirtualShadowMapId].PrevVirtualShadowMapId;
+		const int OffsetScale = (VSM_LEVEL0_DIM_PAGES_XY >> 2);
+		const int2 ClipmapCornerOffsetDelta = OffsetScale * ShadowMapCacheData[VirtualShadowMapId].ClipmapCornerOffsetDelta;
 		if (PrevVirtualShadowMapId != INDEX_NONE)
 		{
 			const FVirtualShadowMapProjectionShaderData Projection = GetVirtualShadowMapProjectionData(VirtualShadowMapId);
@@ -508,9 +510,7 @@ void CreateCachedPageMappings(uint VirtualShadowMapId, uint GlobalPageOffset, ui
 			if (Projection.LightType == LIGHT_TYPE_DIRECTIONAL)
 			{
 				// Clipmap panning
-				const FVirtualShadowMapProjectionShaderData PrevProjection = DecodeVirtualShadowMapProjectionData(PrevProjectionData, PrevVirtualShadowMapId);
-				const int OffsetScale = (VSM_LEVEL0_DIM_PAGES_XY >> 2);
-				PrevPageAddress += OffsetScale * (PrevProjection.ClipmapCornerOffset - Projection.ClipmapCornerOffset);
+				PrevPageAddress += ClipmapCornerOffsetDelta;
 			}
 
 			if (all(PrevPageAddress >= 0) && all(PrevPageAddress < (uint(VSM_LEVEL0_DIM_PAGES_XY) >> MipLevel)))
@@ -859,7 +859,7 @@ void PropagateMappedMips(uint2 ThreadId : SV_DispatchThreadID, uint2 GroupId : S
 		if (!pPage0.bThisLODValid)
 		{
 			const int OffsetScale = (VSM_LEVEL0_DIM_PAGES_XY >> 2);
-			int2 BaseOffset = OffsetScale * ProjectionData.ClipmapCornerOffset;
+			int2 BaseOffset = OffsetScale * ProjectionData.ClipmapCornerRelativeOffset;
 			int2 BasePage   = int2(Level0Page) - BaseOffset;
 
 			// Search for first mapped page past this one
@@ -869,7 +869,7 @@ void PropagateMappedMips(uint2 ThreadId : SV_DispatchThreadID, uint2 GroupId : S
 				const int ClipmapLevelId = VirtualShadowMapID + int(ClipmapOffset);
 
 				FVirtualShadowMapProjectionShaderData LevelProjectionData = GetVirtualShadowMapProjectionData(ClipmapLevelId);								
-				int2 LevelOffset = OffsetScale * LevelProjectionData.ClipmapCornerOffset;
+				int2 LevelOffset = OffsetScale * LevelProjectionData.ClipmapCornerRelativeOffset;
 
 				int2 LevelPage = (BasePage + (LevelOffset << ClipmapOffset)) >> ClipmapOffset;
 
diff --git a/Engine/Shaders/Private/VirtualShadowMaps/VirtualShadowMapProjectionCommon.ush b/Engine/Shaders/Private/VirtualShadowMaps/VirtualShadowMapProjectionCommon.ush
--- a/Engine/Shaders/Private/VirtualShadowMaps/VirtualShadowMapProjectionCommon.ush
+++ b/Engine/Shaders/Private/VirtualShadowMaps/VirtualShadowMapProjectionCommon.ush
@@ -137,8 +137,8 @@ FVirtualShadowMapClipmapRelativeTransform CalcClipmapRelativeTransform(int Clipm
 	const FVirtualShadowMapProjectionShaderData ProjectionDataA = GetVirtualShadowMapProjectionData(ClipmapId);
 	const FVirtualShadowMapProjectionShaderData ProjectionDataB = GetVirtualShadowMapProjectionData(ClipmapId + LevelOffset);
 
-	float2 OffsetA = float2(ProjectionDataA.ClipmapCornerOffset);
-	float2 OffsetB = float2(ProjectionDataB.ClipmapCornerOffset);
+	float2 OffsetA = float2(ProjectionDataA.ClipmapCornerRelativeOffset);
+	float2 OffsetB = float2(ProjectionDataB.ClipmapCornerRelativeOffset);
 
 	FVirtualShadowMapClipmapRelativeTransform Result;
 	Result.Scale = LevelOffset >= 0 ? rcp(float(1U << LevelOffset)) : float(1U << (-LevelOffset));
@@ -161,8 +161,8 @@ uint2 CalcClipmapOffsetLevelPage(uint2 BasePage, int ClipmapId, uint LevelOffset
 	const FVirtualShadowMapProjectionShaderData ProjectionDataB = GetVirtualShadowMapProjectionData(ClipmapId + LevelOffset);
 
 	const int OffsetScale = (VSM_LEVEL0_DIM_PAGES_XY >> 2);
-	int2 BasePageOffset  = OffsetScale * ProjectionDataA.ClipmapCornerOffset;
-	int2 LevelPageOffset = OffsetScale * ProjectionDataB.ClipmapCornerOffset;
+	int2 BasePageOffset  = OffsetScale * ProjectionDataA.ClipmapCornerRelativeOffset;
+	int2 LevelPageOffset = OffsetScale * ProjectionDataB.ClipmapCornerRelativeOffset;
 	return (BasePage - BasePageOffset + (LevelPageOffset << LevelOffset)) >> LevelOffset;
 }
 
diff --git a/Engine/Shaders/Private/VirtualShadowMaps/VirtualShadowMapProjectionStructs.ush b/Engine/Shaders/Private/VirtualShadowMaps/VirtualShadowMapProjectionStructs.ush
--- a/Engine/Shaders/Private/VirtualShadowMaps/VirtualShadowMapProjectionStructs.ush
+++ b/Engine/Shaders/Private/VirtualShadowMaps/VirtualShadowMapProjectionStructs.ush
@@ -22,7 +22,7 @@ struct FVirtualShadowMapProjectionShaderData
 	FLWCVector3 ClipmapWorldOrigin;
 	float ResolutionLodBias;
 
-	int2 ClipmapCornerOffset;
+	int2 ClipmapCornerRelativeOffset;
 	int ClipmapIndex;		// 0..ClipmapLevelCount-1
 	int ClipmapLevel;
 
@@ -74,7 +74,7 @@ FVirtualShadowMapProjectionShaderData DecodeVirtualShadowMapProjectionData(ByteA
 	float3 NegativeClipmapWorldOriginOffset = asfloat(ProjectionData.Load3(Base + (16*18 + 4*0)));
 	Result.ResolutionLodBias				= asfloat(ProjectionData.Load (Base + (16*18 + 4*3)));
 		
-	Result.ClipmapCornerOffset				= asint  (ProjectionData.Load2(Base + (16*19 + 4*0)));
+	Result.ClipmapCornerRelativeOffset		= asint(ProjectionData.Load2(Base + (16*19 + 4*0)));
 	Result.ClipmapIndex						= asint  (ProjectionData.Load (Base + (16*19 + 4*2)));
 	Result.ClipmapLevel						= asint  (ProjectionData.Load (Base + (16*19 + 4*3)));
 
diff --git a/Engine/Source/Runtime/Renderer/Private/VirtualShadowMaps/VirtualShadowMapArray.cpp b/Engine/Source/Runtime/Renderer/Private/VirtualShadowMaps/VirtualShadowMapArray.cpp
--- a/Engine/Source/Runtime/Renderer/Private/VirtualShadowMaps/VirtualShadowMapArray.cpp
+++ b/Engine/Source/Runtime/Renderer/Private/VirtualShadowMaps/VirtualShadowMapArray.cpp
@@ -28,6 +28,8 @@ IMPLEMENT_STATIC_UNIFORM_BUFFER_STRUCT(FVirtualShadowMapUniformParameters, "Virt
 struct FShadowMapCacheData
 {
 	int32 PrevVirtualShadowMapId = INDEX_NONE;
+
+	FIntPoint ClipmapCornerOffsetDelta = FIntPoint(0, 0);
 };
 
 
@@ -324,10 +326,16 @@ static void SetCacheDataShaderParameters(FRDGBuilder& GraphBuilder, const TArray
 		if (VirtualShadowMapCacheEntry != nullptr && VirtualShadowMapCacheEntry->IsValid())
 		{
 			ShadowMapCacheData[SmIndex].PrevVirtualShadowMapId = VirtualShadowMapCacheEntry->PrevVirtualShadowMapId;
+			
+			const int64 TileSize = FLargeWorldRenderScalar::GetTileSize();
+			const FInt64Point& CurOffset = VirtualShadowMapCacheEntry->Clipmap.CurrentClipmapCornerOffset;
+			const FInt64Point& PrevOffset = VirtualShadowMapCacheEntry->Clipmap.PrevClipmapCornerOffset;
+			ShadowMapCacheData[SmIndex].ClipmapCornerOffsetDelta = FInt32Point(PrevOffset - CurOffset);
 		}
 		else
 		{
 			ShadowMapCacheData[SmIndex].PrevVirtualShadowMapId = INDEX_NONE;
+			ShadowMapCacheData[SmIndex].ClipmapCornerOffsetDelta = FInt32Point(0, 0);
 		}
 	}
 	CacheDataParameters.ShadowMapCacheData = GraphBuilder.CreateSRV(CreateStructuredBuffer(GraphBuilder, TEXT("Shadow.Virtual.ShadowMapCacheData"), ShadowMapCacheData));
diff --git a/Engine/Source/Runtime/Renderer/Private/VirtualShadowMaps/VirtualShadowMapArray.h b/Engine/Source/Runtime/Renderer/Private/VirtualShadowMaps/VirtualShadowMapArray.h
--- a/Engine/Source/Runtime/Renderer/Private/VirtualShadowMaps/VirtualShadowMapArray.h
+++ b/Engine/Source/Runtime/Renderer/Private/VirtualShadowMaps/VirtualShadowMapArray.h
@@ -97,7 +97,7 @@ struct FVirtualShadowMapProjectionShaderData
 	// Slightly different meaning for clipmaps (includes camera pixel size scaling stuff) and local lights (raw bias)
 	float ResolutionLodBias = 0.0f;
 
-	FIntPoint ClipmapCornerOffset = FIntPoint(0, 0);
+	FIntPoint ClipmapCornerRelativeOffset = FIntPoint(0, 0);
 	int32 ClipmapIndex = 0;					// 0 .. ClipmapLevelCount-1
 	int32 ClipmapLevel = 0;					// "Absolute" level, can be negative
 
diff --git a/Engine/Source/Runtime/Renderer/Private/VirtualShadowMaps/VirtualShadowMapCacheManager.cpp b/Engine/Source/Runtime/Renderer/Private/VirtualShadowMaps/VirtualShadowMapCacheManager.cpp
--- a/Engine/Source/Runtime/Renderer/Private/VirtualShadowMaps/VirtualShadowMapCacheManager.cpp
+++ b/Engine/Source/Runtime/Renderer/Private/VirtualShadowMaps/VirtualShadowMapCacheManager.cpp
@@ -73,7 +73,8 @@ static FAutoConsoleVariableRef  CVarForceInvalidateDirectionalVSM(
 void FVirtualShadowMapCacheEntry::UpdateClipmap(
 	int32 VirtualShadowMapId,
 	const FMatrix &WorldToLight,
-	FIntPoint PageSpaceLocation,
+	FInt64Point PageSpaceLocation,
+	FInt64Point ClipmapCornerOffset,
 	double LevelRadius,
 	double ViewCenterZ,
 	// NOTE: ViewRadiusZ must be constant for a given clipmap level
@@ -141,12 +142,15 @@ void FVirtualShadowMapCacheEntry::UpdateClipmap(
 	
 	CurrentVirtualShadowMapId = VirtualShadowMapId;
 	CurrentPageSpaceLocation = PageSpaceLocation;
+
+	Clipmap.PrevClipmapCornerOffset = Clipmap.CurrentClipmapCornerOffset;
+	Clipmap.CurrentClipmapCornerOffset = ClipmapCornerOffset;
 }
 
 void FVirtualShadowMapCacheEntry::UpdateLocal(int32 VirtualShadowMapId, const FVirtualShadowMapPerLightCacheEntry& PerLightEntry)
 {
 	// Swap previous frame data over.
-	PrevPageSpaceLocation = FIntPoint(0, 0);		// Not used for local lights
+	PrevPageSpaceLocation = FInt64Point(0, 0);		// Not used for local lights
 	PrevVirtualShadowMapId = CurrentVirtualShadowMapId;
 
 	// Not valid if it was never rendered
@@ -162,7 +166,7 @@ void FVirtualShadowMapCacheEntry::UpdateLocal(int32 VirtualShadowMapId, const FV
 	}
 
 	CurrentVirtualShadowMapId = VirtualShadowMapId;
-	CurrentPageSpaceLocation = FIntPoint(0, 0);		// Not used for local lights
+	CurrentPageSpaceLocation = FInt64Point(0, 0);		// Not used for local lights
 }
 
 void FVirtualShadowMapCacheEntry::Invalidate()
diff --git a/Engine/Source/Runtime/Renderer/Private/VirtualShadowMaps/VirtualShadowMapCacheManager.h b/Engine/Source/Runtime/Renderer/Private/VirtualShadowMaps/VirtualShadowMapCacheManager.h
--- a/Engine/Source/Runtime/Renderer/Private/VirtualShadowMaps/VirtualShadowMapCacheManager.h
+++ b/Engine/Source/Runtime/Renderer/Private/VirtualShadowMaps/VirtualShadowMapCacheManager.h
@@ -29,7 +29,8 @@ public:
 
 	void UpdateClipmap(int32 VirtualShadowMapId,
 		const FMatrix &WorldToLight,
-		FIntPoint PageSpaceLocation,
+		FInt64Point PageSpaceLocation,
+		FInt64Point ClipmapCornerOffset,
 		double LevelRadius,
 		double ViewCenterZ,
 		double ViewRadiusZ, 
@@ -38,11 +39,11 @@ public:
 	void Invalidate();
 
 	// Previous frame data
-	FIntPoint PrevPageSpaceLocation = FIntPoint(0, 0);
+	FInt64Point PrevPageSpaceLocation = FInt64Point(0, 0);
 	int32 PrevVirtualShadowMapId = INDEX_NONE;
 
 	// Current frame data
-	FIntPoint CurrentPageSpaceLocation = FIntPoint(0, 0);
+	FInt64Point CurrentPageSpaceLocation = FInt64Point(0, 0);
 	int32 CurrentVirtualShadowMapId = INDEX_NONE;
 
 	struct FClipmapInfo
@@ -50,6 +51,12 @@ public:
 		FMatrix	WorldToLight;
 		double ViewCenterZ;
 		double ViewRadiusZ;
+
+		// Previous frame data
+		FInt64Point PrevClipmapCornerOffset = FInt64Point(0, 0);
+
+		// Current frame data
+		FInt64Point CurrentClipmapCornerOffset = FInt64Point(0, 0);
 	};
 	FClipmapInfo Clipmap;
 };
diff --git a/Engine/Source/Runtime/Renderer/Private/VirtualShadowMaps/VirtualShadowMapClipmap.cpp b/Engine/Source/Runtime/Renderer/Private/VirtualShadowMaps/VirtualShadowMapClipmap.cpp
--- a/Engine/Source/Runtime/Renderer/Private/VirtualShadowMaps/VirtualShadowMapClipmap.cpp
+++ b/Engine/Source/Runtime/Renderer/Private/VirtualShadowMaps/VirtualShadowMapClipmap.cpp
@@ -122,6 +122,17 @@ FVirtualShadowMapClipmap::FVirtualShadowMapClipmap(
 		}
 	}
 
+	const int RadiusLn = GetLevelRadius(LastLevel);
+	const int RadiiPerLevel = 4;
+	FVector SnappedOriginLn = WorldToLightViewRotationMatrix.TransformPosition(WorldOrigin);
+	{
+		FInt64Point OriginSnapUnitsLn(
+			FMath::RoundToInt64(SnappedOriginLn.X / RadiusLn),
+			FMath::RoundToInt64(SnappedOriginLn.Y / RadiusLn));
+		SnappedOriginLn.X = OriginSnapUnitsLn.X * RadiusLn;
+		SnappedOriginLn.Y = OriginSnapUnitsLn.Y * RadiusLn;
+	}
+
 	for (int32 Index = 0; Index < LevelCount; ++Index)
 	{
 		FLevelData& Level = LevelData[Index];
@@ -137,21 +148,30 @@ FVirtualShadowMapClipmap::FVirtualShadowMapClipmap(
 		double SnapSize = RawLevelRadius;
 
 		FVector ViewCenter = WorldToLightViewRotationMatrix.TransformPosition(WorldOrigin);
-		FIntPoint CenterSnapUnits(
-			FMath::RoundToInt(ViewCenter.X / SnapSize),
-			FMath::RoundToInt(ViewCenter.Y / SnapSize));
+		FVector2D CenterSnapUnits(
+			FMath::RoundToDouble(ViewCenter.X / SnapSize),
+			FMath::RoundToDouble(ViewCenter.Y / SnapSize));
 		ViewCenter.X = CenterSnapUnits.X * SnapSize;
 		ViewCenter.Y = CenterSnapUnits.Y * SnapSize;
 
-		FIntPoint CornerOffset;
-		CornerOffset.X = -CenterSnapUnits.X + 2;
-		CornerOffset.Y =  CenterSnapUnits.Y + 2;
+		FInt64Point CornerOffset;
+		CornerOffset.X = -(int64_t)CenterSnapUnits.X + (RadiiPerLevel/2);
+		CornerOffset.Y =  (int64_t)CenterSnapUnits.Y + (RadiiPerLevel/2);
 
 		const FVector SnappedWorldCenter = ViewToWorldRotationMatrix.TransformPosition(ViewCenter);
 
 		Level.WorldCenter = SnappedWorldCenter;
 		Level.CornerOffset = CornerOffset;
 
+		// A relative corner offset is used for LWC reasons.
+		// The reference point is WorldOrigin snapped to a grid of GetLevelRadius(LastLevel),
+		// because points on this grid are guaranteed to also be present on lower levels,
+		// therefore allowing the offsets to represented as factors of level radii without precision loss.
+		const FInt64Point SnappedPageOriginLi(-ViewCenter.X, ViewCenter.Y);
+		const FInt64Point SnappedPageOriginLn(-SnappedOriginLn.X, SnappedOriginLn.Y);
+		const FInt64Point RelativeCornerOffset = SnappedPageOriginLi - SnappedPageOriginLn + ((RadiiPerLevel / 2) * (int64_t)SnapSize);
+		Level.RelativeCornerOffset = FIntPoint(RelativeCornerOffset / SnapSize);
+
 		// Check if we have a cache entry for this level
 		// If we do and it covers our required depth range, we can use cached pages. Otherwise we need to invalidate.
 		TSharedPtr<FVirtualShadowMapCacheEntry> CacheEntry = nullptr;
@@ -175,11 +195,12 @@ FVirtualShadowMapClipmap::FVirtualShadowMapClipmap(
 		{
 			// We snap to half the size of the VSM at each level
 			check((FVirtualShadowMap::Level0DimPagesXY & 1) == 0);
-			FIntPoint PageOffset(CornerOffset * (FVirtualShadowMap::Level0DimPagesXY >> 2));
+			FInt64Point PageOffset(CornerOffset * (FVirtualShadowMap::Level0DimPagesXY >> 2));
 
 			CacheEntry->UpdateClipmap(Level.VirtualShadowMap->ID,
 				WorldToLightRotationMatrix,
 				PageOffset,
+				CornerOffset,
 				RawLevelRadius,
 				ViewCenter.Z,
 				ViewRadiusZ,
@@ -270,7 +291,7 @@ FVirtualShadowMapProjectionShaderData FVirtualShadowMapClipmap::GetProjectionSha
 	Data.ClipmapLevel = FirstLevel + ClipmapIndex;
 	Data.ClipmapLevelCount = LevelData.Num();
 	Data.ResolutionLodBias = ResolutionLodBias;
-	Data.ClipmapCornerOffset = Level.CornerOffset;
+	Data.ClipmapCornerRelativeOffset = Level.RelativeCornerOffset;
 	Data.LightSourceRadius = GetLightSceneInfo().Proxy->GetSourceRadius();
 	Data.Flags = GForceInvalidateDirectionalVSM ? VSM_PROJ_FLAG_UNCACHED : 0U;
 
diff --git a/Engine/Source/Runtime/Renderer/Private/VirtualShadowMaps/VirtualShadowMapClipmap.h b/Engine/Source/Runtime/Renderer/Private/VirtualShadowMaps/VirtualShadowMapClipmap.h
--- a/Engine/Source/Runtime/Renderer/Private/VirtualShadowMaps/VirtualShadowMapClipmap.h
+++ b/Engine/Source/Runtime/Renderer/Private/VirtualShadowMaps/VirtualShadowMapClipmap.h
@@ -128,7 +128,10 @@ private:
 		FVirtualShadowMap* VirtualShadowMap;
 		FMatrix ViewToClip;
 		FVector WorldCenter;
-		FIntPoint CornerOffset;
+		//Offset from (0,0) to clipmap corner, in level radii
+		FInt64Point CornerOffset;
+		//Offset from LastLevel-snapped WorldCenter to clipmap corner, in level radii
+		FIntPoint RelativeCornerOffset;
 	};
 	TArray< FLevelData, TInlineAllocator<32> > LevelData;
 
