Fix CIS static analysis warning.

diff --git a/Engine/Source/Runtime/Renderer/Private/HairStrands/HairStrandsVisibility.cpp b/Engine/Source/Runtime/Renderer/Private/HairStrands/HairStrandsVisibility.cpp
--- a/Engine/Source/Runtime/Renderer/Private/HairStrands/HairStrandsVisibility.cpp
+++ b/Engine/Source/Runtime/Renderer/Private/HairStrands/HairStrandsVisibility.cpp
@@ -4680,7 +4680,7 @@ void RenderHairStrandsVisibilityBuffer(
 			}
 		#endif
 
-			const bool bIsMSAAForwardEnabled = IsForwardShadingEnabled(View.GetShaderPlatform()) && SceneColorTexture->Desc.NumSamples > 1;
+			const bool bIsMSAAForwardEnabled = IsForwardShadingEnabled(View.GetShaderPlatform()) && SceneColorTexture && SceneColorTexture->Desc.NumSamples > 1;
 			if (bIsMSAAForwardEnabled)
 			{		
 				VisibilityData.ResolveMaskTexture = nullptr;
