Fix packing alignment of FShadowMapCacheData

diff --git a/Engine/Shaders/Private/VirtualShadowMaps/VirtualShadowMapPageCacheCommon.ush b/Engine/Shaders/Private/VirtualShadowMaps/VirtualShadowMapPageCacheCommon.ush
--- a/Engine/Shaders/Private/VirtualShadowMaps/VirtualShadowMapPageCacheCommon.ush
+++ b/Engine/Shaders/Private/VirtualShadowMaps/VirtualShadowMapPageCacheCommon.ush
@@ -12,6 +12,7 @@ struct FShadowMapCacheData
 	int PrevVirtualShadowMapId;
 
 	int2 ClipmapCornerOffsetDelta;
+	int Padding[1];
 };
 
 StructuredBuffer<FShadowMapCacheData> ShadowMapCacheData;
diff --git a/Engine/Source/Runtime/Renderer/Private/VirtualShadowMaps/VirtualShadowMapArray.cpp b/Engine/Source/Runtime/Renderer/Private/VirtualShadowMaps/VirtualShadowMapArray.cpp
--- a/Engine/Source/Runtime/Renderer/Private/VirtualShadowMaps/VirtualShadowMapArray.cpp
+++ b/Engine/Source/Runtime/Renderer/Private/VirtualShadowMaps/VirtualShadowMapArray.cpp
@@ -30,6 +30,7 @@ struct FShadowMapCacheData
 	int32 PrevVirtualShadowMapId = INDEX_NONE;
 
 	FIntPoint ClipmapCornerOffsetDelta = FIntPoint(0, 0);
+	int32 Padding[1];
 };
 
 
