Handle underwater in path tracing:
* For single layer water absorbtion/scattering parameters are only known at the surface
* Add a pass to initialize extinction coefficient by shooting an up-ray from camera to hit the water surface.
* Evaluate material and pass extinction coefficient in a buffer for initializing path state's extinction coefficient.

diff --git a/Engine/Shaders/Private/PathTracing/PathTracingCommon.ush b/Engine/Shaders/Private/PathTracing/PathTracingCommon.ush
--- a/Engine/Shaders/Private/PathTracing/PathTracingCommon.ush
+++ b/Engine/Shaders/Private/PathTracing/PathTracingCommon.ush
@@ -16,7 +16,6 @@
 #define PATH_TRACING_PAYLOAD_INPUT_FLAG_IGNORE_SKY_MATERIALS          (1 << 0) // Indicates that the AHS/CHS should skip processing if the material is tagged as sky
 #define PATH_TRACING_PAYLOAD_INPUT_FLAG_PATHTRACER_VISIBILITY         (1 << 1) // Indicates that this is a path tracer visibility ray (which supports material evaluation for transparent shadows and fake caustics)
 
-
 #define PATH_TRACING_PAYLOAD_NORMAL_BITS  24
 
 // shading model extensions beyond the default count
@@ -178,6 +177,7 @@ struct FPathTracingPayload : FMinimalPayload
 			   ShadingModelID == SHADINGMODELID_THIN_TRANSLUCENT || 
 			   ShadingModelID == SHADINGMODELID_HAIR ||
 			   ShadingModelID == SHADINGMODELID_MEDIUM ||
+			   ShadingModelID == SHADINGMODELID_SINGLELAYERWATER ||	
 			   IsMaterialSolidGlass();
 	}
 
diff --git a/Engine/Shaders/Private/PathTracing/PathTracingCore.ush b/Engine/Shaders/Private/PathTracing/PathTracingCore.ush
--- a/Engine/Shaders/Private/PathTracing/PathTracingCore.ush
+++ b/Engine/Shaders/Private/PathTracing/PathTracingCore.ush
@@ -67,6 +67,8 @@ RWTexture2D<float4> NormalTexture;
 RaytracingAccelerationStructure TLAS;
 uint SceneVisibleLightCount;
 
+Buffer<float> StartingExtinctionCoefficient;
+
 struct FPathState
 {                                   // packed size
 	uint2 PixelIndex;               //  4 bytes
@@ -818,7 +820,7 @@ FPathState CreatePathState(int2 PixelIndex)
 	// path state variables (these cary information between bounces)
 	Output.PathThroughput = 1.0;
 	Output.PathRoughness = 0;
-	Output.SigmaT = 0; // TODO: figure out the right start value for camera rays
+	Output.SigmaT = float3(StartingExtinctionCoefficient[0], StartingExtinctionCoefficient[1], StartingExtinctionCoefficient[2]);
 
 	return Output;
 }
diff --git a/Engine/Shaders/Private/PathTracing/PathTracingInitExtinctionCoefficient.usf b/Engine/Shaders/Private/PathTracing/PathTracingInitExtinctionCoefficient.usf
new file mode 100644
--- /dev/null
+++ b/Engine/Shaders/Private/PathTracing/PathTracingInitExtinctionCoefficient.usf
@@ -0,0 +1,58 @@
+// Copyright Epic Games, Inc. All Rights Reserved.
+
+#include "../Common.ush"
+#include "PathTracingCommon.ush"
+
+RaytracingAccelerationStructure TLAS;
+RWBuffer<float> RWStartingExtinctionCoefficient;
+
+// TODO: figure out the right start value for camera rays inside glass
+RAY_TRACING_ENTRY_RAYGEN(PathTracingInitExtinctionCoefficientRG)
+{
+	FPackedPathTracingPayload PackedPayload = InitPathTracingPayload(0, true);
+
+	const uint RayFlags = RAY_FLAG_FORCE_OPAQUE;
+	const uint MissShaderIndex = 0;
+
+	FRayDesc UpRay;
+	UpRay.Origin = PrimaryView.TranslatedWorldCameraOrigin;
+	UpRay.Direction = float3(0, 0, 1);
+	UpRay.TMin = 0.0;
+	UpRay.TMax = POSITIVE_INFINITY;
+		
+	TraceRay(
+		TLAS,
+		RayFlags,
+		RAY_TRACING_MASK_ALL,
+		RAY_TRACING_SHADER_SLOT_MATERIAL,
+		RAY_TRACING_NUM_SHADER_SLOTS,
+		MissShaderIndex,
+		UpRay.GetNativeDesc(),
+		PackedPayload);
+
+	float3 StartingSigma = 0.0;
+	while (PackedPayload.IsHit())
+	{
+		FPathTracingPayload HitPayload = UnpackPathTracingPayload(PackedPayload, UpRay);
+		if (HitPayload.ShadingModelID == SHADINGMODELID_SINGLELAYERWATER)
+		{			
+			StartingSigma = max(HitPayload.CustomData.xyz, 1e-8f);
+			break;
+		}
+
+		UpRay.TMin = asfloat(asuint(HitPayload.HitT) + 1);
+		TraceRay(
+			TLAS,
+			RayFlags,
+			RAY_TRACING_MASK_ALL,
+			RAY_TRACING_SHADER_SLOT_MATERIAL,
+			RAY_TRACING_NUM_SHADER_SLOTS,
+			MissShaderIndex,
+			UpRay.GetNativeDesc(),
+			PackedPayload);
+	}
+
+	RWStartingExtinctionCoefficient[0] = StartingSigma.x;
+	RWStartingExtinctionCoefficient[1] = StartingSigma.y;
+	RWStartingExtinctionCoefficient[2] = StartingSigma.z;
+}
diff --git a/Engine/Shaders/Private/PathTracing/PathTracingMaterialHitShader.usf b/Engine/Shaders/Private/PathTracing/PathTracingMaterialHitShader.usf
--- a/Engine/Shaders/Private/PathTracing/PathTracingMaterialHitShader.usf
+++ b/Engine/Shaders/Private/PathTracing/PathTracingMaterialHitShader.usf
@@ -220,9 +220,14 @@ RAY_TRACING_ENTRY_CLOSEST_HIT(PathTracingMaterialCHS,
 #if MATERIAL_SHADINGMODEL_SINGLELAYERWATER
 	if (Payload.ShadingModelID == SHADINGMODELID_SINGLELAYERWATER)
 	{
-		const float3 ScatteringCoeff = max(0.0f, LWCToFloat(GetSingleLayerWaterMaterialOutput0(MaterialParameters)));
-		const float3 AbsorptionCoeff = max(0.0f, LWCToFloat(GetSingleLayerWaterMaterialOutput1(MaterialParameters)));
-		const float PhaseG = clamp(LWCToFloat(GetSingleLayerWaterMaterialOutput2(MaterialParameters)), -1.0f, 1.0f);
+		// For single layer water we always want to treat it as front facing when evaluating parameters 
+		// as hitting back face is treated as being underwater and we get completely different parameters.
+		FMaterialPixelParameters WaterMaterialParameters = MaterialParameters;
+		WaterMaterialParameters.TwoSidedSign = 1; 
+
+		const float3 ScatteringCoeff = max(0.0f, LWCToFloat(GetSingleLayerWaterMaterialOutput0(WaterMaterialParameters)));
+		const float3 AbsorptionCoeff = max(0.0f, LWCToFloat(GetSingleLayerWaterMaterialOutput1(WaterMaterialParameters)));
+		const float PhaseG = clamp(LWCToFloat(GetSingleLayerWaterMaterialOutput2(WaterMaterialParameters)), -1.0f, 1.0f);
 
 		const float3 WaterExtinction = ScatteringCoeff + AbsorptionCoeff;
 		const float3 WaterAlbedo = ScatteringCoeff / WaterExtinction;
diff --git a/Engine/Source/Runtime/Renderer/Private/PathTracing.cpp b/Engine/Source/Runtime/Renderer/Private/PathTracing.cpp
--- a/Engine/Source/Runtime/Renderer/Private/PathTracing.cpp
+++ b/Engine/Source/Runtime/Renderer/Private/PathTracing.cpp
@@ -742,6 +742,9 @@ class FPathTracingRG : public FGlobalShader
 		// scene decals
 		SHADER_PARAMETER_RDG_UNIFORM_BUFFER(FRayTracingDecals, DecalParameters)
 
+		// camera ray starting extinction coefficient
+		SHADER_PARAMETER_RDG_BUFFER_SRV(Buffer<float>, StartingExtinctionCoefficient)
+
 		// Used by multi-GPU rendering and TDR-avoidance tiling
 		SHADER_PARAMETER(FIntPoint, TilePixelOffset)
 		SHADER_PARAMETER(FIntPoint, TileTextureOffset)
@@ -760,6 +763,24 @@ class FPathTracingRG : public FGlobalShader
 };
 IMPLEMENT_GLOBAL_SHADER(FPathTracingRG, "/Engine/Private/PathTracing/PathTracing.usf", "PathTracingMainRG", SF_RayGen);
 
+class FPathTracingInitExtinctionCoefficientRG : public FGlobalShader
+{
+	DECLARE_GLOBAL_SHADER(FPathTracingInitExtinctionCoefficientRG)
+	SHADER_USE_ROOT_PARAMETER_STRUCT(FPathTracingInitExtinctionCoefficientRG, FGlobalShader)
+
+	static bool ShouldCompilePermutation(const FGlobalShaderPermutationParameters& Parameters)
+	{
+		return ShouldCompilePathTracingShadersForProject(Parameters.Platform);
+	}
+
+	BEGIN_SHADER_PARAMETER_STRUCT(FParameters, )
+		SHADER_PARAMETER_SRV(RaytracingAccelerationStructure, TLAS)
+		SHADER_PARAMETER_STRUCT_REF(FViewUniformShaderParameters, ViewUniformBuffer)
+		SHADER_PARAMETER_RDG_BUFFER_UAV(RWBuffer<float>, RWStartingExtinctionCoefficient)
+	END_SHADER_PARAMETER_STRUCT()
+};
+IMPLEMENT_GLOBAL_SHADER(FPathTracingInitExtinctionCoefficientRG, "/Engine/Private/PathTracing/PathTracingInitExtinctionCoefficient.usf", "PathTracingInitExtinctionCoefficientRG", SF_RayGen);
+
 class FPathTracingIESAtlasCS : public FGlobalShader
 {
 	DECLARE_GLOBAL_SHADER(FPathTracingIESAtlasCS)
@@ -1901,6 +1922,11 @@ void FDeferredShadingSceneRenderer::PreparePathTracing(const FSceneViewFamily& V
 			auto RayGenShader = GetGlobalShaderMap(ViewFamily.GetShaderPlatform())->GetShader<FPathTracingRG>(PermutationVector);
 			OutRayGenShaders.Add(RayGenShader.GetRayTracingShader());
 		}
+
+		{
+			auto RayGenShader = GetGlobalShaderMap(ViewFamily.GetShaderPlatform())->GetShader<FPathTracingInitExtinctionCoefficientRG>();
+			OutRayGenShaders.Add(RayGenShader.GetRayTracingShader());
+		}		
 	}
 }
 
@@ -2051,6 +2077,43 @@ void FDeferredShadingSceneRenderer::RenderPathTracing(
 		
 	}
 
+	// prepare extinction coefficient for camera rays
+	FRDGBufferRef StartingExtinctionCoefficient = GraphBuilder.CreateBuffer(FRDGBufferDesc::CreateBufferDesc(sizeof(float), 3), TEXT("PathTracer.StartingExtinctionCoefficient"));
+	FRDGBufferUAVRef StartingExtinctionCoefficientUAV = GraphBuilder.CreateUAV(StartingExtinctionCoefficient, PF_R32_FLOAT);
+
+	if (View.IsUnderwater())
+	{
+		auto RayGenShader = GetGlobalShaderMap(View.FeatureLevel)->GetShader<FPathTracingInitExtinctionCoefficientRG>();
+
+		FPathTracingInitExtinctionCoefficientRG::FParameters* PassParameters = GraphBuilder.AllocParameters<FPathTracingInitExtinctionCoefficientRG::FParameters>();
+		PassParameters->TLAS = Scene->RayTracingScene.GetLayerSRVChecked(ERayTracingSceneLayer::Base);
+		PassParameters->ViewUniformBuffer = View.ViewUniformBuffer;
+		PassParameters->RWStartingExtinctionCoefficient = StartingExtinctionCoefficientUAV;
+
+		GraphBuilder.AddPass(
+			RDG_EVENT_NAME("Path Tracer Init Sigma"),
+			PassParameters,
+			ERDGPassFlags::Compute,
+			[PassParameters, RayGenShader, &View](FRHIRayTracingCommandList& RHICmdList)
+			{
+				FRHIRayTracingScene* RayTracingSceneRHI = View.GetRayTracingSceneChecked();
+
+				FRayTracingShaderBindingsWriter GlobalResources;
+				SetShaderParameters(GlobalResources, RayGenShader, *PassParameters);
+				
+				RHICmdList.RayTraceDispatch(
+					View.RayTracingMaterialPipeline,
+					RayGenShader.GetRayTracingShader(),
+					RayTracingSceneRHI, GlobalResources,
+					1, 1
+				);
+			});
+	}
+	else
+	{
+		AddClearUAVPass(GraphBuilder, StartingExtinctionCoefficientUAV, 0);
+	}
+
 	// prepare atmosphere optical depth lookup texture (if needed)
 	FRDGTexture* AtmosphereOpticalDepthLUT = nullptr;
 	if (Config.PathTracingData.EnableAtmosphere)
@@ -2244,6 +2307,7 @@ void FDeferredShadingSceneRenderer::RenderPathTracing(
 							PassParameters->TLAS = Scene->RayTracingScene.GetLayerSRVChecked(ERayTracingSceneLayer::Base);
 							PassParameters->ViewUniformBuffer = View.ViewUniformBuffer;
 							PassParameters->PathTracingData = Config.PathTracingData;
+							PassParameters->StartingExtinctionCoefficient = GraphBuilder.CreateSRV(StartingExtinctionCoefficient, PF_R32_FLOAT);
 							if (PreviousPassParameters == nullptr)
 							{
 								// upload sky/lights data
