Fixed validation errors due to async compute distance field targets.

diff --git a/Engine/Source/Runtime/Renderer/Private/DeferredShadingRenderer.cpp b/Engine/Source/Runtime/Renderer/Private/DeferredShadingRenderer.cpp
--- a/Engine/Source/Runtime/Renderer/Private/DeferredShadingRenderer.cpp
+++ b/Engine/Source/Runtime/Renderer/Private/DeferredShadingRenderer.cpp
@@ -2345,6 +2345,8 @@ void FDeferredShadingSceneRenderer::Render(FRDGBuilder& GraphBuilder)
 		}
 	}
 
+	FRDGExternalAccessQueue ExternalAccessQueue;
+
 	{
 		RDG_CSV_STAT_EXCLUSIVE_SCOPE(GraphBuilder, UpdateGPUScene);
 		RDG_GPU_STAT_SCOPE(GraphBuilder, GPUSceneUpdate);
@@ -2354,8 +2356,6 @@ void FDeferredShadingSceneRenderer::Render(FRDGBuilder& GraphBuilder)
 			GraphBuilder.SetFlushResourcesRHI();
 		}
 
-		FRDGExternalAccessQueue ExternalAccessQueue;
-
 		Scene->GPUScene.Update(GraphBuilder, *Scene, ExternalAccessQueue);
 
 		for (int32 ViewIndex = 0; ViewIndex < Views.Num(); ViewIndex++)
@@ -2368,8 +2368,6 @@ void FDeferredShadingSceneRenderer::Render(FRDGBuilder& GraphBuilder)
 			Scene->GPUScene.DebugRender(GraphBuilder, *Scene, View);
 		}
 
-		ExternalAccessQueue.Submit(GraphBuilder);
-
 		InstanceCullingManager.BeginDeferredCulling(GraphBuilder, Scene->GPUScene);
 
 		if (Views.Num() > 0)
@@ -2522,7 +2520,8 @@ void FDeferredShadingSceneRenderer::Render(FRDGBuilder& GraphBuilder)
 		}
 	}
 
-	PrepareDistanceFieldScene(GraphBuilder, false);
+	PrepareDistanceFieldScene(GraphBuilder, ExternalAccessQueue, false);
+	ExternalAccessQueue.Submit(GraphBuilder);
 
 	const bool bShouldRenderVelocities = ShouldRenderVelocities();
 	const EShaderPlatform Platform = GetViewFamilyInfo(Views).GetShaderPlatform();
diff --git a/Engine/Source/Runtime/Renderer/Private/DistanceFieldObjectManagement.cpp b/Engine/Source/Runtime/Renderer/Private/DistanceFieldObjectManagement.cpp
--- a/Engine/Source/Runtime/Renderer/Private/DistanceFieldObjectManagement.cpp
+++ b/Engine/Source/Runtime/Renderer/Private/DistanceFieldObjectManagement.cpp
@@ -558,7 +558,8 @@ bool ProcessHeightFieldPrimitiveUpdate(
 bool bVerifySceneIntegrity = false;
 
 void FDistanceFieldSceneData::UpdateDistanceFieldObjectBuffers(
-	FRDGBuilder& GraphBuilder, 
+	FRDGBuilder& GraphBuilder,
+	FRDGExternalAccessQueue& ExternalAccessQueue,
 	FScene* Scene,
 	const DistanceField::FUpdateTrackingBounds& UpdateTrackingBounds,
 	TArray<FDistanceFieldAssetMipId>& DistanceFieldAssetAdds,
@@ -799,8 +800,8 @@ void FDistanceFieldSceneData::UpdateDistanceFieldObjectBuffers(
 					UploadDistanceFieldDataBuffer.ResourceUploadTo(GraphBuilder, DFObjectDataBuffer);
 					UploadDistanceFieldBoundsBuffer.ResourceUploadTo(GraphBuilder, DFObjectBoundsBuffer);
 
-					GraphBuilder.UseExternalAccessMode(DFObjectDataBuffer, ERHIAccess::SRVMask);
-					GraphBuilder.UseExternalAccessMode(DFObjectBoundsBuffer, ERHIAccess::SRVMask);
+					ExternalAccessQueue.Add(DFObjectDataBuffer, ERHIAccess::SRVMask, ERHIPipeline::All);
+					ExternalAccessQueue.Add(DFObjectBoundsBuffer, ERHIAccess::SRVMask, ERHIPipeline::All);
 				}
 			}
 		}
@@ -999,7 +1000,7 @@ void FSceneRenderer::UpdateGlobalHeightFieldObjectBuffers(FRDGBuilder& GraphBuil
 	}
 }
 
-void FSceneRenderer::PrepareDistanceFieldScene(FRDGBuilder& GraphBuilder, bool bSplitDispatch)
+void FSceneRenderer::PrepareDistanceFieldScene(FRDGBuilder& GraphBuilder, FRDGExternalAccessQueue& ExternalAccessQueue, bool bSplitDispatch)
 {
 	RDG_CSV_STAT_EXCLUSIVE_SCOPE(GraphBuilder, PrepareDistanceFieldScene);
 	TRACE_CPUPROFILER_EVENT_SCOPE(FSceneRenderer::PrepareDistanceFieldScene);
@@ -1033,9 +1034,9 @@ void FSceneRenderer::PrepareDistanceFieldScene(FRDGBuilder& GraphBuilder, bool b
 	{
 		TArray<FDistanceFieldAssetMipId> DistanceFieldAssetAdds;
 		TArray<FSetElementId> DistanceFieldAssetRemoves;
-		Scene->DistanceFieldSceneData.UpdateDistanceFieldObjectBuffers(GraphBuilder, Scene, UpdateTrackingBounds, DistanceFieldAssetAdds, DistanceFieldAssetRemoves);
+		Scene->DistanceFieldSceneData.UpdateDistanceFieldObjectBuffers(GraphBuilder, ExternalAccessQueue, Scene, UpdateTrackingBounds, DistanceFieldAssetAdds, DistanceFieldAssetRemoves);
 
-		Scene->DistanceFieldSceneData.UpdateDistanceFieldAtlas(GraphBuilder, Views[0], Scene, IsLumenEnabled(Views[0]), Views[0].ShaderMap, DistanceFieldAssetAdds, DistanceFieldAssetRemoves);
+		Scene->DistanceFieldSceneData.UpdateDistanceFieldAtlas(GraphBuilder, ExternalAccessQueue, Views[0], Scene, IsLumenEnabled(Views[0]), Views[0].ShaderMap, DistanceFieldAssetAdds, DistanceFieldAssetRemoves);
 
 		if (bSplitDispatch)
 		{
@@ -1058,7 +1059,7 @@ void FSceneRenderer::PrepareDistanceFieldScene(FRDGBuilder& GraphBuilder, bool b
 					OcclusionMaxDistance = Scene->SkyLight->OcclusionMaxDistance;
 				}
 
-				UpdateGlobalDistanceFieldVolume(GraphBuilder, View, Scene, OcclusionMaxDistance, IsLumenEnabled(View), View.GlobalDistanceFieldInfo);
+				UpdateGlobalDistanceFieldVolume(GraphBuilder, ExternalAccessQueue, View, Scene, OcclusionMaxDistance, IsLumenEnabled(View), View.GlobalDistanceFieldInfo);
 			}
 		}
 		if (!bSplitDispatch)
diff --git a/Engine/Source/Runtime/Renderer/Private/DistanceFieldStreaming.cpp b/Engine/Source/Runtime/Renderer/Private/DistanceFieldStreaming.cpp
--- a/Engine/Source/Runtime/Renderer/Private/DistanceFieldStreaming.cpp
+++ b/Engine/Source/Runtime/Renderer/Private/DistanceFieldStreaming.cpp
@@ -852,7 +852,7 @@ void FDistanceFieldSceneData::ProcessReadRequests(
 	AssetDataUploads.Append(DistanceFieldAssetMipAdds);
 }
 
-void FDistanceFieldSceneData::ResizeBrickAtlasIfNeeded(FRDGBuilder& GraphBuilder, FGlobalShaderMap* GlobalShaderMap)
+FRDGTexture* FDistanceFieldSceneData::ResizeBrickAtlasIfNeeded(FRDGBuilder& GraphBuilder, FGlobalShaderMap* GlobalShaderMap)
 {
 	// Mask should be set in FSceneRenderer::PrepareDistanceFieldScene before calling this
 	check(GraphBuilder.RHICmdList.GetGPUMask() == FRHIGPUMask::All());
@@ -903,7 +903,10 @@ void FDistanceFieldSceneData::ResizeBrickAtlasIfNeeded(FRDGBuilder& GraphBuilder
 
 		BrickTextureDimensionsInBricks = DesiredBrickTextureDimensionsInBricks;
 		DistanceFieldBrickVolumeTexture = GraphBuilder.ConvertToExternalTexture(DistanceFieldBrickVolumeTextureRDG);
+		return DistanceFieldBrickVolumeTextureRDG;
 	}
+
+	return GraphBuilder.RegisterExternalTexture(DistanceFieldBrickVolumeTexture);
 }
 
 static FIntVector CalculateDesiredSize(FIntVector CurrentSize, FIntVector RequiredSize)
@@ -1277,7 +1280,8 @@ void FDistanceFieldSceneData::UploadAllAssetData(FRDGBuilder& GraphBuilder, FRDG
 }
 
 void FDistanceFieldSceneData::UpdateDistanceFieldAtlas(
-	FRDGBuilder& GraphBuilder, 
+	FRDGBuilder& GraphBuilder,
+	FRDGExternalAccessQueue& ExternalAccessQueue,
 	const FViewInfo& View,
 	FScene* Scene,
 	bool bLumenEnabled,
@@ -1294,8 +1298,6 @@ void FDistanceFieldSceneData::UpdateDistanceFieldAtlas(
 
 	TArray<FDistanceFieldAssetMipId> AssetDataUploads;
 
-	FRDGExternalAccessQueue ExternalAccessQueue;
-
 	for (FSetElementId AssetSetId : DistanceFieldAssetRemoves)
 	{
 		const FDistanceFieldAssetState& AssetState = AssetStateArray[AssetSetId];
@@ -1375,7 +1377,7 @@ void FDistanceFieldSceneData::UpdateDistanceFieldAtlas(
 	}
 
 	// Now that DistanceFieldAtlasBlockAllocator has been modified, potentially resize the atlas
-	ResizeBrickAtlasIfNeeded(GraphBuilder, GlobalShaderMap);
+	FRDGTextureRef DistanceFieldBrickVolumeTextureRDG = ResizeBrickAtlasIfNeeded(GraphBuilder, GlobalShaderMap);
 
 	const uint32 NumAssets = AssetStateArray.GetMaxIndex();
 	const int32 AssetDataStrideFloat4s = DistanceField::NumMips * AssetDataMipStrideFloat4s;
@@ -1511,8 +1513,6 @@ void FDistanceFieldSceneData::UpdateDistanceFieldAtlas(
 
 		if (NumBrickUploads > 0)
 		{
-			FRDGTextureRef DistanceFieldBrickVolumeTextureRDG = GraphBuilder.RegisterExternalTexture(DistanceFieldBrickVolumeTexture, TEXT("DistanceFields.DistanceFieldBrickVolumeTexture"));
-
 			// GRHIMaxDispatchThreadGroupsPerDimension can be MAX_int32 so we need to do this math in 64-bit.
 			const int32 MaxBrickUploadsPerPass = (int32)FMath::Min<int64>((int64)GRHIMaxDispatchThreadGroupsPerDimension.Z * FScatterUploadDistanceFieldAtlasCS::GetGroupSize() / DistanceField::BrickSize, MAX_int32);
 
@@ -1537,8 +1537,6 @@ void FDistanceFieldSceneData::UpdateDistanceFieldAtlas(
 					PassParameters,
 					FComputeShaderUtils::GetGroupCount(FIntVector(DistanceField::BrickSize, DistanceField::BrickSize, NumBrickUploadsThisPass * DistanceField::BrickSize), FScatterUploadDistanceFieldAtlasCS::GetGroupSize()));
 			}
-
-			DistanceFieldBrickVolumeTexture = ConvertToExternalAccessTexture(GraphBuilder, ExternalAccessQueue, DistanceFieldBrickVolumeTextureRDG);
 		}
 	}
 	
@@ -1560,8 +1558,8 @@ void FDistanceFieldSceneData::UpdateDistanceFieldAtlas(
 		GDistanceFieldAtlasLogStats = 0;
 	}
 
+	ExternalAccessQueue.Add(DistanceFieldBrickVolumeTextureRDG, ERHIAccess::SRVMask, ERHIPipeline::All);
 	ExternalAccessQueue.Add(AssetDataBufferRDG, ERHIAccess::SRVMask, ERHIPipeline::All);
-	ExternalAccessQueue.Submit(GraphBuilder);
 }
 
 void FDistanceFieldSceneData::ListMeshDistanceFields(bool bDumpAssetStats) const
diff --git a/Engine/Source/Runtime/Renderer/Private/GlobalDistanceField.cpp b/Engine/Source/Runtime/Renderer/Private/GlobalDistanceField.cpp
--- a/Engine/Source/Runtime/Renderer/Private/GlobalDistanceField.cpp
+++ b/Engine/Source/Runtime/Renderer/Private/GlobalDistanceField.cpp
@@ -1827,6 +1827,7 @@ IMPLEMENT_GLOBAL_SHADER(FGlobalDistanceFieldDebugCS, "/Engine/Private/DistanceFi
  **/
 void UpdateGlobalDistanceFieldVolume(
 	FRDGBuilder& GraphBuilder,
+	FRDGExternalAccessQueue& ExternalAccessQueue,
 	FViewInfo& View,
 	FScene* Scene,
 	float MaxOcclusionDistance,
@@ -2712,8 +2713,6 @@ void UpdateGlobalDistanceFieldVolume(
 				}
 			}
 
-			FRDGExternalAccessQueue ExternalAccessQueue;
-
 			for (int32 CacheType = StartCacheType; CacheType < GDF_Num; CacheType++)
 			{
 				if (PageTableLayerTextures[CacheType])
@@ -2756,8 +2755,6 @@ void UpdateGlobalDistanceFieldVolume(
 			{
 				GlobalDistanceFieldInfo.MipTexture = ConvertToExternalAccessTexture(GraphBuilder, ExternalAccessQueue, MipTexture);
 			}
-
-			ExternalAccessQueue.Submit(GraphBuilder);
 		}
 
 		if (CVarGlobalDistanceFieldDebug.GetValueOnRenderThread() != 0 && GlobalDistanceFieldInfo.PageFreeListAllocatorBuffer)
diff --git a/Engine/Source/Runtime/Renderer/Private/GlobalDistanceField.h b/Engine/Source/Runtime/Renderer/Private/GlobalDistanceField.h
--- a/Engine/Source/Runtime/Renderer/Private/GlobalDistanceField.h
+++ b/Engine/Source/Runtime/Renderer/Private/GlobalDistanceField.h
@@ -43,6 +43,7 @@ namespace GlobalDistanceField
  **/
 extern void UpdateGlobalDistanceFieldVolume(
 	FRDGBuilder& GraphBuilder,
+	FRDGExternalAccessQueue& ExternalAccessQueue,
 	FViewInfo& View, 
 	FScene* Scene, 
 	float MaxOcclusionDistance, 
diff --git a/Engine/Source/Runtime/Renderer/Private/MobileShadingRenderer.cpp b/Engine/Source/Runtime/Renderer/Private/MobileShadingRenderer.cpp
--- a/Engine/Source/Runtime/Renderer/Private/MobileShadingRenderer.cpp
+++ b/Engine/Source/Runtime/Renderer/Private/MobileShadingRenderer.cpp
@@ -630,24 +630,25 @@ void FMobileSceneRenderer::InitViews(FRDGBuilder& GraphBuilder, FSceneTexturesCo
 		View.InitRHIResources();
 	}
 
+	FRDGExternalAccessQueue ExternalAccessQueue;
+
 	{
 		RDG_CSV_STAT_EXCLUSIVE_SCOPE(GraphBuilder, UpdateGPUScene);
-		FRDGExternalAccessQueue ExternalAccessQueue;
 
 		Scene->GPUScene.Update(GraphBuilder, *Scene, ExternalAccessQueue);
 		for (int32 ViewIndex = 0; ViewIndex < Views.Num(); ViewIndex++)
 		{
 			Scene->GPUScene.UploadDynamicPrimitiveShaderDataForView(GraphBuilder, Scene, Views[ViewIndex], ExternalAccessQueue);
 		}
-
-		ExternalAccessQueue.Submit(GraphBuilder);
 	}
 
 	if (bRequiresDistanceField)
 	{
-		PrepareDistanceFieldScene(GraphBuilder, false);
+		PrepareDistanceFieldScene(GraphBuilder, ExternalAccessQueue, false);
 	}
 
+	ExternalAccessQueue.Submit(GraphBuilder);
+
 	InstanceCullingManager.BeginDeferredCulling(GraphBuilder, Scene->GPUScene);
 
 	extern TSet<IPersistentViewUniformBufferExtension*> PersistentViewUniformBufferExtensions;
diff --git a/Engine/Source/Runtime/Renderer/Private/ScenePrivate.h b/Engine/Source/Runtime/Renderer/Private/ScenePrivate.h
--- a/Engine/Source/Runtime/Renderer/Private/ScenePrivate.h
+++ b/Engine/Source/Runtime/Renderer/Private/ScenePrivate.h
@@ -1962,14 +1962,16 @@ public:
 	void ListMeshDistanceFields(bool bDumpAssetStats) const;
 
 	void UpdateDistanceFieldObjectBuffers(
-		FRDGBuilder& GraphBuilder, 
+		FRDGBuilder& GraphBuilder,
+		FRDGExternalAccessQueue& ExternalAccessQueue,
 		FScene* Scene,
 		const DistanceField::FUpdateTrackingBounds& UpdateTrackingBounds,
 		TArray<FDistanceFieldAssetMipId>& DistanceFieldAssetAdds,
 		TArray<FSetElementId>& DistanceFieldAssetRemoves);
 
 	void UpdateDistanceFieldAtlas(
-		FRDGBuilder& GraphBuilder, 
+		FRDGBuilder& GraphBuilder,
+		FRDGExternalAccessQueue& ExternalAccessQueue,
 		const FViewInfo& View,
 		FScene* Scene,
 		bool bLumenEnabled,
@@ -2104,7 +2106,7 @@ private:
 		TArray<FDistanceFieldReadRequest>& ReadRequestsToUpload,
 		TArray<FDistanceFieldReadRequest>& ReadRequestsToCleanUp);
 
-	void ResizeBrickAtlasIfNeeded(FRDGBuilder& GraphBuilder, FGlobalShaderMap* GlobalShaderMap);
+	FRDGTexture* ResizeBrickAtlasIfNeeded(FRDGBuilder& GraphBuilder, FGlobalShaderMap* GlobalShaderMap);
 
 	bool ResizeIndirectionAtlasIfNeeded(FRDGBuilder& GraphBuilder, FGlobalShaderMap* GlobalShaderMap, FRDGTexture*& OutTexture);
 
diff --git a/Engine/Source/Runtime/Renderer/Private/SceneRendering.h b/Engine/Source/Runtime/Renderer/Private/SceneRendering.h
--- a/Engine/Source/Runtime/Renderer/Private/SceneRendering.h
+++ b/Engine/Source/Runtime/Renderer/Private/SceneRendering.h
@@ -2119,7 +2119,7 @@ public:
 	void UpdateGlobalDistanceFieldObjectBuffers(FRDGBuilder& GraphBuilder);
 	void UpdateGlobalHeightFieldObjectBuffers(FRDGBuilder& GraphBuilder);
 	void AddOrRemoveSceneHeightFieldPrimitives(const DistanceField::FUpdateTrackingBounds& UpdateTrackingBounds, bool bSkipAdd = false);
-	void PrepareDistanceFieldScene(FRDGBuilder& GraphBuilder, bool bSplitDispatch);
+	void PrepareDistanceFieldScene(FRDGBuilder& GraphBuilder, FRDGExternalAccessQueue& ExternalAccessQueue, bool bSplitDispatch);
 
 	void DrawGPUSkinCacheVisualizationInfoText();
 
