Properly filter out shading models for a mobile deferred shading

diff --git a/Engine/Source/Runtime/Engine/Classes/Engine/EngineTypes.h b/Engine/Source/Runtime/Engine/Classes/Engine/EngineTypes.h
--- a/Engine/Source/Runtime/Engine/Classes/Engine/EngineTypes.h
+++ b/Engine/Source/Runtime/Engine/Classes/Engine/EngineTypes.h
@@ -651,6 +651,7 @@ public:
 	bool IsLit() const														{ return !IsUnlit(); }
 	bool IsValid() const													{ return (ShadingModelField > 0) && (ShadingModelField < (1 << MSM_NUM)); }
 	uint16 GetShadingModelField() const										{ return ShadingModelField; }
+	void SetShadingModelField(uint16 InShadingModelField)					{ ShadingModelField = InShadingModelField; }
 	int32 CountShadingModels() const										{ return FMath::CountBits(ShadingModelField); }
 	EMaterialShadingModel GetFirstShadingModel() const						{ check(IsValid()); return (EMaterialShadingModel)FMath::CountTrailingZeros(ShadingModelField); }
 
diff --git a/Engine/Source/Runtime/Renderer/Private/EditorPrimitivesRendering.cpp b/Engine/Source/Runtime/Renderer/Private/EditorPrimitivesRendering.cpp
--- a/Engine/Source/Runtime/Renderer/Private/EditorPrimitivesRendering.cpp
+++ b/Engine/Source/Runtime/Renderer/Private/EditorPrimitivesRendering.cpp
@@ -150,7 +150,7 @@ bool FEditorPrimitivesBasePassMeshProcessor::ProcessMobileShadingPath(const FMes
 
 	if (bTranslucentBasePass)
 	{
-		MobileBasePass::SetTranslucentRenderState(DrawRenderState, Material);
+		MobileBasePass::SetTranslucentRenderState(DrawRenderState, Material, Material.GetShadingModels());
 	}
 
 	const FMeshDrawingPolicyOverrideSettings OverrideSettings = ComputeMeshOverrideSettings(MeshBatch);
diff --git a/Engine/Source/Runtime/Renderer/Private/MobileBasePass.cpp b/Engine/Source/Runtime/Renderer/Private/MobileBasePass.cpp
--- a/Engine/Source/Runtime/Renderer/Private/MobileBasePass.cpp
+++ b/Engine/Source/Runtime/Renderer/Private/MobileBasePass.cpp
@@ -338,7 +338,7 @@ ELightMapPolicyType MobileBasePass::SelectMeshLightmapPolicy(
 	return SelectedLightmapPolicy;
 }
 
-void MobileBasePass::SetOpaqueRenderState(FMeshPassProcessorRenderState& DrawRenderState, const FPrimitiveSceneProxy* PrimitiveSceneProxy, const FMaterial& Material, bool bEnableReceiveDecalOutput, bool bUsesDeferredShading)
+void MobileBasePass::SetOpaqueRenderState(FMeshPassProcessorRenderState& DrawRenderState, const FPrimitiveSceneProxy* PrimitiveSceneProxy, const FMaterial& Material, FMaterialShadingModelField ShadingModels, bool bEnableReceiveDecalOutput, bool bUsesDeferredShading)
 {
 	uint8 StencilValue = 0;
 	if (bEnableReceiveDecalOutput)
@@ -349,7 +349,7 @@ void MobileBasePass::SetOpaqueRenderState(FMeshPassProcessorRenderState& DrawRen
 	
 	if (bUsesDeferredShading)
 	{
-		uint8 ShadingModel = GetMobileShadingModelStencilValue(Material.GetShadingModels());
+		uint8 ShadingModel = GetMobileShadingModelStencilValue(ShadingModels);
 		StencilValue |= GET_STENCIL_MOBILE_SM_MASK(ShadingModel);
 		StencilValue |= STENCIL_LIGHTING_CHANNELS_MASK(PrimitiveSceneProxy ? PrimitiveSceneProxy->GetLightingChannelStencilValue() : 0x00);
 	}
@@ -384,11 +384,11 @@ void MobileBasePass::SetOpaqueRenderState(FMeshPassProcessorRenderState& DrawRen
 	}
 }
 
-void MobileBasePass::SetTranslucentRenderState(FMeshPassProcessorRenderState& DrawRenderState, const FMaterial& Material)
+void MobileBasePass::SetTranslucentRenderState(FMeshPassProcessorRenderState& DrawRenderState, const FMaterial& Material, FMaterialShadingModelField ShadingModels)
 {
 	const bool bIsUsingMobilePixelProjectedReflection = Material.IsUsingPlanarForwardReflections() && IsUsingMobilePixelProjectedReflection(GetFeatureLevelShaderPlatform(Material.GetFeatureLevel()));
 
-	if (Material.GetShadingModels().HasShadingModel(MSM_ThinTranslucent))
+	if (ShadingModels.HasShadingModel(MSM_ThinTranslucent))
 	{
 		// the mobile thin translucent fallback uses a similar mode as BLEND_Translucent, but multiplies color by 1 insead of SrcAlpha.
 		DrawRenderState.SetBlendState(TStaticBlendState<CW_RGBA, BO_Add, BF_One, BF_InverseSourceAlpha, BO_Add, BF_Zero, BF_InverseSourceAlpha,
@@ -473,7 +473,7 @@ void MobileBasePass::SetTranslucentRenderState(FMeshPassProcessorRenderState& Dr
 															CW_NONE, BO_Add, BF_One, BF_Zero, BO_Add, BF_One, BF_Zero>::GetRHI());
 			break;
 		default:
-			if (Material.GetShadingModels().HasShadingModel(MSM_SingleLayerWater))
+			if (ShadingModels.HasShadingModel(MSM_SingleLayerWater))
 			{
 				// Single layer water is an opaque marerial rendered as translucent on Mobile. We force pre-multiplied alpha to achieve water depth based transmittance.
 				DrawRenderState.SetBlendState(TStaticBlendState<CW_RGBA, BO_Add, BF_One, BF_InverseSourceAlpha, BO_Add, BF_Zero, BF_InverseSourceAlpha,
@@ -573,15 +573,29 @@ FMobileBasePassMeshProcessor::FMobileBasePassMeshProcessor(
 	, PassDrawRenderState(InDrawRenderState)
 	, TranslucencyPassType(InTranslucencyPassType)
 	, Flags(InFlags)
+	, ShadingModelsMask(GetPlatformShadingModelsMask(GetFeatureLevelShaderPlatform(InFeatureLevel)))
 	, bTranslucentBasePass(InTranslucencyPassType != ETranslucencyPass::TPT_MAX)
 	, bUsesDeferredShading(!bTranslucentBasePass && IsMobileDeferredShadingEnabled(GetFeatureLevelShaderPlatform(InFeatureLevel)))
 {
 }
 
+FMaterialShadingModelField FMobileBasePassMeshProcessor::FilterShadingModelsMask(const FMaterialShadingModelField& ShadingModels) const
+{
+	uint16 ShadingModelField = ShadingModels.GetShadingModelField();
+	uint16 FilteredShadingModels = (ShadingModelField & ShadingModelsMask);
+	FMaterialShadingModelField NewShadingModelField;
+	NewShadingModelField.SetShadingModelField(FilteredShadingModels);
+	if (!NewShadingModelField.IsValid())
+	{
+		NewShadingModelField.AddShadingModel(MSM_DefaultLit);
+	}
+	return NewShadingModelField;
+}
+
 bool FMobileBasePassMeshProcessor::TryAddMeshBatch(const FMeshBatch& RESTRICT MeshBatch, uint64 BatchElementMask, const FPrimitiveSceneProxy* RESTRICT PrimitiveSceneProxy, int32 StaticMeshId, const FMaterialRenderProxy& MaterialRenderProxy, const FMaterial& Material)
 {
 	const EBlendMode BlendMode = Material.GetBlendMode();
-	const FMaterialShadingModelField ShadingModels = Material.GetShadingModels();
+	const FMaterialShadingModelField ShadingModels = FilterShadingModelsMask(Material.GetShadingModels());
 	const bool bIsTranslucent = IsTranslucentBlendMode(BlendMode);
 	const bool bUsesWaterMaterial = ShadingModels.HasShadingModel(MSM_SingleLayerWater); // Water goes into the translucent pass
 	const bool bCanReceiveCSM = ((Flags & EFlags::CanReceiveCSM) == EFlags::CanReceiveCSM);
@@ -697,7 +711,7 @@ bool FMobileBasePassMeshProcessor::Process(
 	{
 		if (bTranslucentBasePass)
 		{
-			MobileBasePass::SetTranslucentRenderState(DrawRenderState, MaterialResource);
+			MobileBasePass::SetTranslucentRenderState(DrawRenderState, MaterialResource, ShadingModels);
 		}
 		else if((MeshBatch.bUseForDepthPass && Scene->EarlyZPassMode == DDM_AllOpaque) || bMaskedInEarlyPass)
 		{
@@ -706,7 +720,7 @@ bool FMobileBasePassMeshProcessor::Process(
 		else
 		{
 			const bool bEnableReceiveDecalOutput = ((Flags & EFlags::CanUseDepthStencil) == EFlags::CanUseDepthStencil);
-			MobileBasePass::SetOpaqueRenderState(DrawRenderState, PrimitiveSceneProxy, MaterialResource, bEnableReceiveDecalOutput && IsMobileHDR(), bUsesDeferredShading);
+			MobileBasePass::SetOpaqueRenderState(DrawRenderState, PrimitiveSceneProxy, MaterialResource, ShadingModels, bEnableReceiveDecalOutput && IsMobileHDR(), bUsesDeferredShading);
 		}
 	}
 
diff --git a/Engine/Source/Runtime/Renderer/Private/MobileBasePassRendering.h b/Engine/Source/Runtime/Renderer/Private/MobileBasePassRendering.h
--- a/Engine/Source/Runtime/Renderer/Private/MobileBasePassRendering.h
+++ b/Engine/Source/Runtime/Renderer/Private/MobileBasePassRendering.h
@@ -335,7 +335,7 @@ namespace MobileBasePass
 	bool StaticCanReceiveCSM(const FLightSceneInfo* LightSceneInfo, const FPrimitiveSceneProxy* PrimitiveSceneProxy);
 
 	void SetOpaqueRenderState(FMeshPassProcessorRenderState& DrawRenderState, const FPrimitiveSceneProxy* PrimitiveSceneProxy, const FMaterial& Material, bool bEnableReceiveDecalOutput, bool bUsesDeferredShading);
-	void SetTranslucentRenderState(FMeshPassProcessorRenderState& DrawRenderState, const FMaterial& Material);
+	void SetTranslucentRenderState(FMeshPassProcessorRenderState& DrawRenderState, const FMaterial& Material, FMaterialShadingModelField ShadingModels);
 };
 
 
@@ -458,6 +458,8 @@ public:
 	FMeshPassProcessorRenderState PassDrawRenderState;
 
 private:
+	FMaterialShadingModelField FilterShadingModelsMask(const FMaterialShadingModelField& ShadingModels) const;
+
 	bool TryAddMeshBatch(const FMeshBatch& RESTRICT MeshBatch, uint64 BatchElementMask, const FPrimitiveSceneProxy* RESTRICT PrimitiveSceneProxy, int32 StaticMeshId, const FMaterialRenderProxy& MaterialRenderProxy, const FMaterial& Material);
 
 	bool Process(
@@ -475,6 +477,7 @@ private:
 
 	const ETranslucencyPass::Type TranslucencyPassType;
 	const EFlags Flags;
+	uint16 ShadingModelsMask;
 	const bool bTranslucentBasePass;
 	const bool bUsesDeferredShading;
 };
