Make distant light threshold respect resolution LOD bias for local lights such that more lights are treated as distant if LOD bias is increased.

diff --git a/Engine/Source/Runtime/Renderer/Private/Shadows/ShadowSceneRenderer.cpp b/Engine/Source/Runtime/Renderer/Private/Shadows/ShadowSceneRenderer.cpp
--- a/Engine/Source/Runtime/Renderer/Private/Shadows/ShadowSceneRenderer.cpp
+++ b/Engine/Source/Runtime/Renderer/Private/Shadows/ShadowSceneRenderer.cpp
@@ -64,7 +64,7 @@ TSharedPtr<FVirtualShadowMapPerLightCacheEntry> FShadowSceneRenderer::AddLocalLi
 	//       Also possible to feed back from gpu - which would be more accurate wrt partially visible lights (e.g., a spot going through the ground).
 	//       Of course this creates jumps if visibility changes, which may or may not create unsolvable artifacts.
 	const bool bIsDistantLight = CVarDistantLightMode.GetValueOnRenderThread() != 0
-		&& (MaxScreenRadius <= FVirtualShadowMap::PageSize || CVarDistantLightMode.GetValueOnRenderThread() == 2);
+		&& (MaxScreenRadius <= float(FVirtualShadowMap::PageSize) * FMath::Exp2(VirtualShadowMapArray.GetResolutionLODBiasLocal()) || CVarDistantLightMode.GetValueOnRenderThread() == 2);
 
 
 	const int32 NumMaps = ProjectedShadowInitializer.bOnePassPointLightShadow ? 6 : 1;
diff --git a/Engine/Source/Runtime/Renderer/Private/VirtualShadowMaps/VirtualShadowMapArray.cpp b/Engine/Source/Runtime/Renderer/Private/VirtualShadowMaps/VirtualShadowMapArray.cpp
--- a/Engine/Source/Runtime/Renderer/Private/VirtualShadowMaps/VirtualShadowMapArray.cpp
+++ b/Engine/Source/Runtime/Renderer/Private/VirtualShadowMaps/VirtualShadowMapArray.cpp
@@ -479,6 +479,11 @@ FVirtualShadowMap* FVirtualShadowMapArray::Allocate(bool bSinglePageShadowMap)
 	return SM;
 }
 
+float FVirtualShadowMapArray::GetResolutionLODBiasLocal() const
+{
+	return CVarResolutionLodBiasLocal.GetValueOnRenderThread();
+}
+
 FVirtualShadowMapArray::~FVirtualShadowMapArray()
 {
 }
@@ -1176,7 +1181,7 @@ void FVirtualShadowMapArray::BuildPageAllocations(
 			}
 		}
 
-		const float ResolutionLodBiasLocal = CVarResolutionLodBiasLocal.GetValueOnRenderThread();
+		const float ResolutionLodBiasLocal = GetResolutionLODBiasLocal();
 
 		for (FProjectedShadowInfo* ProjectedShadowInfo : VisibleLightInfo.AllProjectedShadows)
 		{
diff --git a/Engine/Source/Runtime/Renderer/Private/VirtualShadowMaps/VirtualShadowMapArray.h b/Engine/Source/Runtime/Renderer/Private/VirtualShadowMaps/VirtualShadowMapArray.h
--- a/Engine/Source/Runtime/Renderer/Private/VirtualShadowMaps/VirtualShadowMapArray.h
+++ b/Engine/Source/Runtime/Renderer/Private/VirtualShadowMaps/VirtualShadowMapArray.h
@@ -227,6 +227,11 @@ public:
 		return GetNumFullShadowMaps() + GetNumSinglePageShadowMaps();
 	}
 
+	/**
+	 * Get configured LOD bias for the local lights (maps to cvar, 
+	 */
+	float GetResolutionLODBiasLocal() const;
+
 	// Raw size of the physical pool, including both static and dynamic pages (if enabled)
 	FIntPoint GetPhysicalPoolSize() const;
 	// Size of HZB (level 0)
