Remove deprecation warning when no mesh type is provided - will fallback to num type which means to stat tracking for that pass

diff --git a/Engine/Source/Runtime/Renderer/Public/MeshPassProcessor.h b/Engine/Source/Runtime/Renderer/Public/MeshPassProcessor.h
--- a/Engine/Source/Runtime/Renderer/Public/MeshPassProcessor.h
+++ b/Engine/Source/Runtime/Renderer/Public/MeshPassProcessor.h
@@ -1935,7 +1935,6 @@ public:
 	const FSceneView* ViewIfDynamicMeshCommand;
 	FMeshPassDrawListContext* DrawListContext;
 		
-	UE_DEPRECATED(5.1, "FMeshPassProcessor constructor without EMeshPass::Type is deprecated.")
 	FMeshPassProcessor(const FScene* InScene, ERHIFeatureLevel::Type InFeatureLevel, const FSceneView* InViewIfDynamicMeshCommand, FMeshPassDrawListContext* InDrawListContext) :
 		FMeshPassProcessor(EMeshPass::Num, InScene, InFeatureLevel, InViewIfDynamicMeshCommand, InDrawListContext) { }
 
