Call public renamed function

diff --git a/Engine/Source/Runtime/Renderer/Public/MeshPassProcessor.inl b/Engine/Source/Runtime/Renderer/Public/MeshPassProcessor.inl
--- a/Engine/Source/Runtime/Renderer/Public/MeshPassProcessor.inl
+++ b/Engine/Source/Runtime/Renderer/Public/MeshPassProcessor.inl
@@ -210,7 +210,7 @@ void FMeshPassProcessor::AddGraphicsPipelineStateInitializer(
 	if ((MeshPassFeatures & EMeshPassFeatures::PositionAndNormalOnly) != EMeshPassFeatures::Default)	InputStreamType = EVertexInputStreamType::PositionAndNormalOnly;
 
 	FVertexDeclarationElementList Elements;
-	VertexFactoryType->GetPSOPrecacheVertexFetchElements(InputStreamType, Elements);
+	VertexFactoryType->GetShaderPSOPrecacheVertexFetchElements(InputStreamType, Elements);
 	FRHIVertexDeclaration* VertexDeclaration = PipelineStateCache::GetOrCreateVertexDeclaration(Elements);
 	
 	FMeshProcessorShaders MeshProcessorShaders = PassShaders.GetUntypedShaders();
