Rolled back 21899912 and integrated pieces of 20806454 to properly fix distance field shadow + CSM in forward in mobile renderer

#rnx

diff --git a/Engine/Source/Runtime/Renderer/Private/ShadowRendering.cpp b/Engine/Source/Runtime/Renderer/Private/ShadowRendering.cpp
--- a/Engine/Source/Runtime/Renderer/Private/ShadowRendering.cpp
+++ b/Engine/Source/Runtime/Renderer/Private/ShadowRendering.cpp
@@ -2484,12 +2484,11 @@ void FMobileSceneRenderer::RenderMobileShadowProjections(
 		// Local light shadows don't render to shadow mask texture on mobile deferred
 		if (LightSceneProxy->GetLightType() == LightType_Directional || !IsMobileDeferredShadingEnabled(ShaderPlatform))
 		{
-			bool bProjectingForForwardShading = !IsMobileDistanceFieldEnabled(ShaderPlatform);
 			RenderShadowProjections(GraphBuilder, SceneTextures,
 				ScreenShadowMaskTexture,
 				nullptr,
 				LightSceneInfo,
-				bProjectingForForwardShading);
+				true);
 		}
 	}
 }
