Allocate virtual texture page table texture up front for private page tables.
For adaptive virtual texture this fixes a bug where deferred page table updates during a page table texture reallocation frame could lead to invalid page table samples.
Also fix missing function overrides in FVirtualTextureAddressRedirect which could lead to adaptive virtual textures not waiting on producer tasks.

diff --git a/Engine/Source/Runtime/Renderer/Private/VT/AdaptiveVirtualTexture.cpp b/Engine/Source/Runtime/Renderer/Private/VT/AdaptiveVirtualTexture.cpp
--- a/Engine/Source/Runtime/Renderer/Private/VT/AdaptiveVirtualTexture.cpp
+++ b/Engine/Source/Runtime/Renderer/Private/VT/AdaptiveVirtualTexture.cpp
@@ -99,6 +99,20 @@ public:
 		return VirtualTexture->ProducePageData(RHICmdList, FeatureLevel, Flags, ProducerHandle, LayerMask, vLevel, vAddress, RequestHandle, TargetLayers);
 	}
 
+	virtual void GatherProducePageDataTasks(
+		FVirtualTextureProducerHandle const& ProducerHandle,
+		FGraphEventArray& InOutTasks) const override
+	{
+		VirtualTexture->GatherProducePageDataTasks(ProducerHandle, InOutTasks);
+	}
+
+	virtual void GatherProducePageDataTasks(
+		uint64 RequestHandle,
+		FGraphEventArray& InOutTasks) const override
+	{
+		VirtualTexture->GatherProducePageDataTasks(RequestHandle, InOutTasks);
+	}
+
 private:
 	IVirtualTexture* VirtualTexture;
 	FIntPoint AddressOffset;
diff --git a/Engine/Source/Runtime/Renderer/Private/VT/VirtualTextureSpace.cpp b/Engine/Source/Runtime/Renderer/Private/VT/VirtualTextureSpace.cpp
--- a/Engine/Source/Runtime/Renderer/Private/VT/VirtualTextureSpace.cpp
+++ b/Engine/Source/Runtime/Renderer/Private/VT/VirtualTextureSpace.cpp
@@ -87,8 +87,11 @@ FVirtualTextureSpace::~FVirtualTextureSpace()
 
 uint32 FVirtualTextureSpace::AllocateVirtualTexture(FAllocatedVirtualTexture* VirtualTexture)
 {
-	uint32 vAddress = Allocator.Alloc(VirtualTexture);
-	if (Allocator.GetAllocatedWidth() > CachedPageTableWidth || Allocator.GetAllocatedHeight() > CachedPageTableHeight)
+	const uint32 vAddress = Allocator.Alloc(VirtualTexture);
+	
+	// After allocation, check if we need to reallocate the page table texture.
+	const FUintPoint RequiredPageTableSize = GetRequiredPageTableAllocationSize(); 
+	if (RequiredPageTableSize.X > CachedPageTableWidth || RequiredPageTableSize.Y > CachedPageTableHeight)
 	{
 		bNeedToAllocatePageTable = true;
 	}
@@ -128,15 +131,26 @@ void FVirtualTextureSpace::ReleaseRHI()
 	UpdateBufferSRV.SafeRelease();
 }
 
+FUintPoint FVirtualTextureSpace::GetRequiredPageTableAllocationSize() const
+{
+	// Private spaces should allocate the full page table texture up front.
+	const uint32 Width = Description.bPrivateSpace ? Description.MaxSpaceSize : Allocator.GetAllocatedWidth();
+	const uint32 Height = Description.bPrivateSpace ? Description.MaxSpaceSize : Allocator.GetAllocatedHeight();
+	// We align on some minimum size. Maybe minimum, and align sizes should be different? But OK for now.
+	const uint32 WidthAligned = Align(Width, VIRTUALTEXTURE_MIN_PAGETABLE_SIZE);
+	const uint32 HeightAligned = Align(Height, VIRTUALTEXTURE_MIN_PAGETABLE_SIZE);
+	return FUintPoint(Width, Height);
+}
+
 uint32 FVirtualTextureSpace::GetSizeInBytes() const
 {
+	const FUintPoint RequiredPageTableSize = GetRequiredPageTableAllocationSize();
+	const uint32 NumPageTableLevels = FMath::FloorLog2(FMath::Max(RequiredPageTableSize.X, RequiredPageTableSize.Y)) + 1u;
+
 	uint32 TotalSize = 0u;
 	for (uint32 TextureIndex = 0u; TextureIndex < GetNumPageTableTextures(); ++TextureIndex)
 	{
-		const uint32 PageTableWidth = Align(Allocator.GetAllocatedWidth(), VIRTUALTEXTURE_MIN_PAGETABLE_SIZE);
-		const uint32 PageTableHeight = Align(Allocator.GetAllocatedHeight(), VIRTUALTEXTURE_MIN_PAGETABLE_SIZE);
-		const uint32 NumPageTableLevels = FMath::FloorLog2(FMath::Max(PageTableWidth, PageTableHeight)) + 1u;
-		const SIZE_T TextureSize = CalcTextureSize(PageTableWidth, PageTableHeight, TexturePixelFormat[TextureIndex], NumPageTableLevels);
+		const SIZE_T TextureSize = CalcTextureSize(RequiredPageTableSize.X, RequiredPageTableSize.Y, TexturePixelFormat[TextureIndex], NumPageTableLevels);
 		TotalSize += TextureSize;
 	}
 	
@@ -269,8 +283,9 @@ void FVirtualTextureSpace::AllocateTextures(FRDGBuilder& GraphBuilder)
 		const TCHAR* TextureNames[] = { TEXT("VirtualPageTable_0"), TEXT("VirtualPageTable_1") };
 		static_assert(UE_ARRAY_COUNT(TextureNames) == TextureCapacity, "");
 
-		CachedPageTableWidth = Align(Allocator.GetAllocatedWidth(), VIRTUALTEXTURE_MIN_PAGETABLE_SIZE);
-		CachedPageTableHeight = Align(Allocator.GetAllocatedHeight(), VIRTUALTEXTURE_MIN_PAGETABLE_SIZE);
+		const FUintPoint RequiredPageTableSize = GetRequiredPageTableAllocationSize();
+		CachedPageTableWidth = RequiredPageTableSize.X;
+		CachedPageTableHeight = RequiredPageTableSize.Y;
 		CachedNumPageTableLevels = FMath::FloorLog2(FMath::Max(CachedPageTableWidth, CachedPageTableHeight)) + 1u;
 
 		for (uint32 TextureIndex = 0u; TextureIndex < GetNumPageTableTextures(); ++TextureIndex)
@@ -342,8 +357,12 @@ void FVirtualTextureSpace::ApplyUpdates(FVirtualTextureSystem* System, FRDGBuild
 
 	if (bNeedToAllocatePageTable)
 	{
-		// Defer updates until next frame if page table needs to be re-allocated
-		// We can't update page table at this point in frame, as RHIUpdateTextureReference can't be called during RHIBegin/EndScene
+		// Defer updates until next frame if page table texture needs to be re-allocated
+		// We can't update the page table texture at this point in frame, as RHIUpdateTextureReference can't be called during RHIBegin/EndScene
+		// Note that the virtual texture system doesn't account for page table updates being deferred. So this can potentially lead to sampling invalid page table addresses.
+		// This could cause a glitch if we sample a VT on the first frame it is allocated. That's usually not the case (we usually sample some time after loading).
+		// But it can be the case for Adaptive Virtual Texture which does a lot of dynamic page table allocation during the texture life.
+		// However Adaptive Virtual Texture is OK because it always sets bPrivateSpace which gives fixed allocation of the actual page table texture.
 		return;
 	}
 
diff --git a/Engine/Source/Runtime/Renderer/Private/VT/VirtualTextureSpace.h b/Engine/Source/Runtime/Renderer/Private/VT/VirtualTextureSpace.h
--- a/Engine/Source/Runtime/Renderer/Private/VT/VirtualTextureSpace.h
+++ b/Engine/Source/Runtime/Renderer/Private/VT/VirtualTextureSpace.h
@@ -103,6 +103,8 @@ public:
 #endif // WITH_EDITOR
 
 private:
+	FUintPoint GetRequiredPageTableAllocationSize() const;
+
 	static const uint32 TextureCapacity = (VIRTUALTEXTURE_SPACE_MAXLAYERS + LayersPerPageTableTexture - 1u) / LayersPerPageTableTexture;
 
 	struct FTextureEntry
