Fix invalid check when there are loose parameters present

diff --git a/Engine/Source/Runtime/Renderer/Private/MeshPassProcessor.cpp b/Engine/Source/Runtime/Renderer/Private/MeshPassProcessor.cpp
--- a/Engine/Source/Runtime/Renderer/Private/MeshPassProcessor.cpp
+++ b/Engine/Source/Runtime/Renderer/Private/MeshPassProcessor.cpp
@@ -419,17 +419,15 @@ FRayTracingLocalShaderBindings* FMeshDrawShaderBindings::SetRayTracingShaderBind
 
 	// Measure parameter memory requirements
 
-	int32 MaxUniformBufferUsed = -1;
+	int32 MaxUniformBufferIndex = -1;
 	for (int32 UniformBufferIndex = 0; UniformBufferIndex < NumUniformBufferParameters; UniformBufferIndex++)
 	{
 		FShaderUniformBufferParameterInfo Parameter = UniformBufferParameters[UniformBufferIndex];
 		const FRHIUniformBuffer* UniformBuffer = UniformBufferBindings[UniformBufferIndex];
-		MaxUniformBufferUsed = FMath::Max((int32)Parameter.BaseIndex, MaxUniformBufferUsed);
+		MaxUniformBufferIndex = FMath::Max((int32)Parameter.BaseIndex, MaxUniformBufferIndex);
 	}
 
-	const uint32 NumUniformBuffersToSet = MaxUniformBufferUsed + 1;
-
-	check(NumUniformBuffersToSet == NumUniformBufferParameters);
+	const uint32 NumUniformBuffersToSet = MaxUniformBufferIndex + 1;
 
 	const TMemoryImageArray<FShaderLooseParameterBufferInfo>& LooseParameterBuffers = SingleShaderBindings.ParameterMapInfo.LooseParameterBuffers;
 	uint32 LooseParameterDataSize = 0;
@@ -448,6 +446,12 @@ FRayTracingLocalShaderBindings* FMeshDrawShaderBindings::SetRayTracingShaderBind
 		}
 	}
 
+	checkf(MaxUniformBufferIndex + 1 == NumUniformBufferParameters + LooseParameterBuffers.Num(),
+		TEXT("Highest index of a uniform buffer was %d, but there were %d uniform buffer parameters and %d loose parameters"),
+		MaxUniformBufferIndex,
+		NumUniformBufferParameters,
+		LooseParameterBuffers.Num());
+
 	// Allocate and fill bindings
 
 	const uint32 UserData = 0; // UserData could be used to store material ID or any other kind of per-material constant. This can be retrieved in hit shaders via GetHitGroupUserData().
