Fix for hair PrimMat texture in some cases not being written before reading it in a later pass.

#rnx

diff --git a/Engine/Source/Runtime/Renderer/Private/HairStrands/HairStrandsVisibility.cpp b/Engine/Source/Runtime/Renderer/Private/HairStrands/HairStrandsVisibility.cpp
--- a/Engine/Source/Runtime/Renderer/Private/HairStrands/HairStrandsVisibility.cpp
+++ b/Engine/Source/Runtime/Renderer/Private/HairStrands/HairStrandsVisibility.cpp
@@ -3404,6 +3404,7 @@ static FRasterComputeOutput AddVisibilityComputeRasterPass(
 		}
 	}
 
+	bool bAnyHairRasterized = false;
 	uint32 PrimitiveInfoIndex = 0;
 	for (const FHairStrandsMacroGroupData& MacroGroup : MacroGroupDatas)
 	{
@@ -3675,9 +3676,17 @@ static FRasterComputeOutput AddVisibilityComputeRasterPass(
 			}
 
 			++PrimitiveInfoIndex;
+			bAnyHairRasterized = true;
 		}
 	}
 
+	// If no groom was rasterized, clear the PrimMat texture to invalid values so that RDG doesn't complain about unwritten resources.
+	// Usually we don't have to clear this texture because it is only accessed if the corresponding texels in the hair count texture indicate valid samples.
+	if (!bAnyHairRasterized)
+	{
+		uint32 PrimMatClearValues[4] = { 0xFFFFFFFF, 0xFFFFFFFF, 0xFFFFFFFF, 0xFFFFFFFF };
+		AddClearUAVPass(GraphBuilder, PrimMatTextureUAV, PrimMatClearValues);
+	}
 
 	return Out;
 }
