StaticLighting: Replace GetLevelVolumetricLightmap() nullptr checks with HasData() as the former one returns &GlobalVolumetricLightmap which is never nullptr

diff --git a/Engine/Source/Runtime/Renderer/Private/LightMapRendering.cpp b/Engine/Source/Runtime/Renderer/Private/LightMapRendering.cpp
--- a/Engine/Source/Runtime/Renderer/Private/LightMapRendering.cpp
+++ b/Engine/Source/Runtime/Renderer/Private/LightMapRendering.cpp
@@ -282,7 +282,7 @@ void GetIndirectLightingCacheParameters(
 {
 	// FCachedVolumeIndirectLightingPolicy, FCachedPointIndirectLightingPolicy
 	{
-		if (VolumetricLightmapSceneData)
+		if (VolumetricLightmapSceneData && VolumetricLightmapSceneData->HasData())
 		{
 			FVolumetricLightmapInterpolation* Interpolation = VolumetricLightmapSceneData->CPUInterpolationCache.Find(VolumetricLightmapLookupPosition);
 
diff --git a/Engine/Source/Runtime/Renderer/Private/SceneRendering.cpp b/Engine/Source/Runtime/Renderer/Private/SceneRendering.cpp
--- a/Engine/Source/Runtime/Renderer/Private/SceneRendering.cpp
+++ b/Engine/Source/Runtime/Renderer/Private/SceneRendering.cpp
@@ -1120,7 +1120,7 @@ void UpdateNoiseTextureParameters(FViewUniformShaderParameters& ViewUniformShade
 
 void SetupPrecomputedVolumetricLightmapUniformBufferParameters(const FScene* Scene, FEngineShowFlags EngineShowFlags, FViewUniformShaderParameters& ViewUniformShaderParameters)
 {
-	if (Scene && Scene->VolumetricLightmapSceneData.GetLevelVolumetricLightmap() && EngineShowFlags.VolumetricLightmap)
+	if (Scene && Scene->VolumetricLightmapSceneData.HasData() && EngineShowFlags.VolumetricLightmap)
 	{
 		const FPrecomputedVolumetricLightmapData* VolumetricLightmapData = Scene->VolumetricLightmapSceneData.GetLevelVolumetricLightmap()->Data;
 
diff --git a/Engine/Source/Runtime/Renderer/Private/VisualizeVolumetricLightmap.cpp b/Engine/Source/Runtime/Renderer/Private/VisualizeVolumetricLightmap.cpp
--- a/Engine/Source/Runtime/Renderer/Private/VisualizeVolumetricLightmap.cpp
+++ b/Engine/Source/Runtime/Renderer/Private/VisualizeVolumetricLightmap.cpp
@@ -101,9 +101,7 @@ void FDeferredShadingSceneRenderer::VisualizeVolumetricLightmap(
 		return;
 	}
 
-	const FPrecomputedVolumetricLightmap* VolumetricLightmap = Scene->VolumetricLightmapSceneData.GetLevelVolumetricLightmap();
-
-	if (!VolumetricLightmap)
+	if (!Scene->VolumetricLightmapSceneData.HasData())
 	{
 		return;
 	}
