Fixed SkyAtmosphere not updating correctly the cached camera volume buffer when cvars are updated.

diff --git a/Engine/Source/Runtime/Renderer/Private/RendererScene.cpp b/Engine/Source/Runtime/Renderer/Private/RendererScene.cpp
--- a/Engine/Source/Runtime/Renderer/Private/RendererScene.cpp
+++ b/Engine/Source/Runtime/Renderer/Private/RendererScene.cpp
@@ -281,12 +281,15 @@ struct FSpeedTreeWindComputation
 
 FPersistentSkyAtmosphereData::FPersistentSkyAtmosphereData()
 	: bInitialised(false)
+	, CurrentScreenResolution(0)
+	, CurrentDepthResolution(0)
+	, CurrentTextureAerialLUTFormat(PF_Unknown)
 	, CameraAerialPerspectiveVolumeIndex(0)
 {
 }
 void FPersistentSkyAtmosphereData::InitialiseOrNextFrame(ERHIFeatureLevel::Type FeatureLevel, FPooledRenderTargetDesc& AerialPerspectiveDesc, FRHICommandListImmediate& RHICmdList)
 {
-	if (!bInitialised)
+	if (!bInitialised || (bInitialised && ((AerialPerspectiveDesc.Extent.X != CurrentScreenResolution) || (AerialPerspectiveDesc.Depth != CurrentDepthResolution) || (AerialPerspectiveDesc.Format != CurrentTextureAerialLUTFormat))))
 	{
 		CameraAerialPerspectiveVolumeCount = FeatureLevel == ERHIFeatureLevel::ES3_1 ? 2 : 1;
 		for (int i = 0; i < CameraAerialPerspectiveVolumeCount; ++i)
@@ -295,6 +298,9 @@ void FPersistentSkyAtmosphereData::InitialiseOrNextFrame(ERHIFeatureLevel::Type
 				i==0 ? TEXT("SkyAtmosphere.CameraAPVolume0") : TEXT("SkyAtmosphere.CameraAPVolume1"));
 		}
 		bInitialised = true;
+		CurrentScreenResolution = AerialPerspectiveDesc.Extent.X;
+		CurrentDepthResolution = AerialPerspectiveDesc.Depth;
+		CurrentTextureAerialLUTFormat = AerialPerspectiveDesc.Format;
 	}
 
 	CameraAerialPerspectiveVolumeIndex = (CameraAerialPerspectiveVolumeIndex + 1) % CameraAerialPerspectiveVolumeCount;
diff --git a/Engine/Source/Runtime/Renderer/Private/ScenePrivate.h b/Engine/Source/Runtime/Renderer/Private/ScenePrivate.h
--- a/Engine/Source/Runtime/Renderer/Private/ScenePrivate.h
+++ b/Engine/Source/Runtime/Renderer/Private/ScenePrivate.h
@@ -605,6 +605,9 @@ struct FPersistentSkyAtmosphereData
 
 private:
 	bool bInitialised;
+	int32 CurrentScreenResolution;
+	int32 CurrentDepthResolution;
+	EPixelFormat CurrentTextureAerialLUTFormat;
 	TRefCountPtr<IPooledRenderTarget> CameraAerialPerspectiveVolumes[2];
 	uint8 CameraAerialPerspectiveVolumeCount;
 	uint8 CameraAerialPerspectiveVolumeIndex;
