Change Strata byte per pixel budget to be driven per platform, rather than a global project settings.

diff --git a/Engine/Source/Runtime/Engine/Classes/Engine/RendererSettings.h b/Engine/Source/Runtime/Engine/Classes/Engine/RendererSettings.h
--- a/Engine/Source/Runtime/Engine/Classes/Engine/RendererSettings.h
+++ b/Engine/Source/Runtime/Engine/Classes/Engine/RendererSettings.h
@@ -917,15 +917,6 @@ class ENGINE_API URendererSettings : public UDeveloperSettings
 		ConfigRestartRequired = true))
 		uint32 bEnableStrata : 1;
 
-	/**
-	"Strata allocated byte per pixel to store materials data. Higher value means more complex material can be represented."
-	*/
-	UPROPERTY(config, EditAnywhere, Category = Strata, meta = (
-		ConsoleVariable = "r.Strata.BytesPerPixel", DisplayName = "Strata bytes per pixel",
-		ToolTip = "Strata allocated byte per pixel to store materials data. Higher value means more complex material can be represented.",
-		ConfigRestartRequired = true))
-		uint32 StrataBytePerPixel;
-
 	/**
 	"Enable Strata opaque material rough refractions effect from top layers over layers below."
 	*/
diff --git a/Engine/Source/Runtime/Engine/Private/Materials/HLSLMaterialTranslator.cpp b/Engine/Source/Runtime/Engine/Private/Materials/HLSLMaterialTranslator.cpp
--- a/Engine/Source/Runtime/Engine/Private/Materials/HLSLMaterialTranslator.cpp
+++ b/Engine/Source/Runtime/Engine/Private/Materials/HLSLMaterialTranslator.cpp
@@ -69,12 +69,13 @@ static inline bool IsDebugTextureSampleEnabled()
 
 bool Engine_IsStrataEnabled();
 
-static uint32 GetStrataBytePerPixel()
+static uint32 GetStrataBytePerPixel(EShaderPlatform InPlatform)
 {
-	static const auto CVarStrataBytePerPixel = IConsoleManager::Get().FindTConsoleVariableDataInt(TEXT("r.Strata.BytesPerPixel"));
-	check(CVarStrataBytePerPixel);
-	const uint32 StrataBytePerPixel = CVarStrataBytePerPixel ? CVarStrataBytePerPixel->GetValueOnAnyThread() : 0;
-	return StrataBytePerPixel;
+	// We enforce at least 20 bytes per pixel because this is the minimal Strata GBuffer footprint of the simplest material.
+	const uint32 MinStrataBytePerPixel = 20u;
+	const uint32 MaxStrataBytePerPixel = 256u;
+	static FShaderPlatformCachedIniValue<int32> CVarBudget(TEXT("r.Strata.BytesPerPixel"));
+	return FMath::Clamp(uint32(CVarBudget.Get(InPlatform)), MinStrataBytePerPixel, MaxStrataBytePerPixel);
 }
 
 /** @return the vector type containing a given number of components. */
@@ -1951,7 +1952,7 @@ void FHLSLMaterialTranslator::GetMaterialEnvironment(EShaderPlatform InPlatform,
 		// Now write some feedback to the user
 		{
 			// Output some debug info as comment in code and in the material stat window
-			const uint32 StrataBytePerPixel = GetStrataBytePerPixel();
+			const uint32 StrataBytePerPixel = GetStrataBytePerPixel(InPlatform);
 
 			OutEnvironment.SetDefine(TEXT("STRATA_CLAMPED_BSDF_COUNT"), StrataMaterialBSDFCount);
 
@@ -9850,7 +9851,7 @@ bool FHLSLMaterialTranslator::StrataGenerateDerivedMaterialOperatorData()
 		}
 	}
 
-	const uint32 StrataBytePerPixel = GetStrataBytePerPixel();
+	const uint32 StrataBytePerPixel = GetStrataBytePerPixel(Platform);
 	do 
 	{
 		// Reset some data for simplifiucation iteration
diff --git a/Engine/Source/Runtime/Engine/Private/RendererSettings.cpp b/Engine/Source/Runtime/Engine/Private/RendererSettings.cpp
--- a/Engine/Source/Runtime/Engine/Private/RendererSettings.cpp
+++ b/Engine/Source/Runtime/Engine/Private/RendererSettings.cpp
@@ -246,18 +246,6 @@ void URendererSettings::PostEditChangeProperty(FPropertyChangedEvent& PropertyCh
 			}
 		}
 
-		if ((PropertyChangedEvent.Property->GetFName() == GET_MEMBER_NAME_CHECKED(URendererSettings, StrataBytePerPixel)))
-		{
-			const uint32 MinStrataBytePerPixel = 20u;
-			const uint32 MaxStrataBytePerPixel = 128u;
-			if (StrataBytePerPixel < MinStrataBytePerPixel || StrataBytePerPixel > MaxStrataBytePerPixel)
-			{
-				FMessageDialog::Open(EAppMsgType::Ok, LOCTEXT("Strata Byte Per Pixel", "Strata `byte-per-pixel` must be between 20 and 128."));
-			}
-			// We enforce at least 20 bytes per pixel because this is the minimal Strata GBuffer footprint of the simplest material.
-			StrataBytePerPixel = FMath::Clamp(StrataBytePerPixel, MinStrataBytePerPixel, MaxStrataBytePerPixel);
-		}
-
 		ExportValuesToConsoleVariables(PropertyChangedEvent.Property);
 
 		if (PropertyChangedEvent.Property->GetFName() == GET_MEMBER_NAME_CHECKED(URendererSettings, ReflectionCaptureResolution) &&
diff --git a/Engine/Source/Runtime/RenderCore/Private/Shader.cpp b/Engine/Source/Runtime/RenderCore/Private/Shader.cpp
--- a/Engine/Source/Runtime/RenderCore/Private/Shader.cpp
+++ b/Engine/Source/Runtime/RenderCore/Private/Shader.cpp
@@ -1828,10 +1828,14 @@ void ShaderMapAppendKeyString(EShaderPlatform Platform, FString& KeyString)
 			KeyString += TEXT("_STRATA");
 		}
 
-		static const auto CVarBudget = IConsoleManager::Get().FindTConsoleVariableDataInt(TEXT("r.Strata.BytesPerPixel"));
-		if (bStrataEnabled && CVarBudget)
+		if (bStrataEnabled)
 		{
-			KeyString += FString::Printf(TEXT("_BUDGET%u"), CVarBudget->GetValueOnAnyThread());
+			// We enforce at least 20 bytes per pixel because this is the minimal Strata GBuffer footprint of the simplest material.
+			const uint32 MinStrataBytePerPixel = 20u;
+			const uint32 MaxStrataBytePerPixel = 256u;
+			static FShaderPlatformCachedIniValue<int32> CVarBudget(TEXT("r.Strata.BytesPerPixel"));
+			const uint32 BytesPerPixel = FMath::Clamp(uint32(CVarBudget.Get(Platform)), MinStrataBytePerPixel, MaxStrataBytePerPixel);
+			KeyString += FString::Printf(TEXT("_BUDGET%u"), BytesPerPixel);
 		}
 
 		static const auto CVarBackCompatibility = IConsoleManager::Get().FindTConsoleVariableDataInt(TEXT("r.StrataBackCompatibility"));
diff --git a/Engine/Source/Runtime/Renderer/Private/Strata/Strata.cpp b/Engine/Source/Runtime/Renderer/Private/Strata/Strata.cpp
--- a/Engine/Source/Runtime/Renderer/Private/Strata/Strata.cpp
+++ b/Engine/Source/Runtime/Renderer/Private/Strata/Strata.cpp
@@ -280,16 +280,21 @@ static bool NeedBSDFOffsets(const FScene* Scene, const FViewInfo& View)
 	return  ShouldRenderLumenDiffuseGI(Scene, View) || ShouldRenderLumenReflections(View) || Strata::ShouldRenderStrataDebugPasses(View);
 }
 
+static uint32 StrataGetBytePerPixel()
+{
+	// We enforce at least 20 bytes per pixel because this is the minimal Strata GBuffer footprint of the simplest material.
+	const uint32 MinStrataBytePerPixel = 20u;
+	const uint32 MaxStrataBytePerPixel = 256u;
+	return FMath::Clamp(uint32(CVarStrataBytePerPixel.GetValueOnAnyThread()), MinStrataBytePerPixel, MaxStrataBytePerPixel);
+}
 
 static void RecordStrataAnalytics()
 {
 	if (FEngineAnalytics::IsAvailable())
 	{
-		const uint32 MaterialConservativeByteCountPerPixel = CVarStrataBytePerPixel.GetValueOnAnyThread();
-
 		TArray<FAnalyticsEventAttribute> EventAttributes;
 		EventAttributes.Add(FAnalyticsEventAttribute(TEXT("Enabled"), 1));
-		EventAttributes.Add(FAnalyticsEventAttribute(TEXT("BytesPerPixel"), MaterialConservativeByteCountPerPixel));
+		EventAttributes.Add(FAnalyticsEventAttribute(TEXT("BytesPerPixel"), StrataGetBytePerPixel()));
 
 		FString OutStr(TEXT("Strata"));
 		FEngineAnalytics::GetProvider().RecordEvent(OutStr, EventAttributes);
@@ -332,9 +337,8 @@ void InitialiseStrataFrameSceneData(FRDGBuilder& GraphBuilder, FSceneRenderer& S
 		// We need to allocate enough for the tiled memory addressing of material data to always work
 		UpdateMaterialBufferToTiledResolution(SceneTextureExtent, MaterialBufferSizeXY);
 
-		const uint32 MaterialConservativeByteCountPerPixel = CVarStrataBytePerPixel.GetValueOnAnyThread();
 		const uint32 RoundToValue = 4u;
-		Out.MaxBytesPerPixel = FMath::DivideAndRoundUp(MaterialConservativeByteCountPerPixel, RoundToValue) * RoundToValue;
+		Out.MaxBytesPerPixel = FMath::DivideAndRoundUp(StrataGetBytePerPixel(), RoundToValue) * RoundToValue;
 
 		// Top layer texture
 		{
