Fix hair selection issue under certain camera position/view resolution.

diff --git a/Engine/Shaders/Private/HairStrands/HairStrandsVisibilityTile.usf b/Engine/Shaders/Private/HairStrands/HairStrandsVisibilityTile.usf
--- a/Engine/Shaders/Private/HairStrands/HairStrandsVisibilityTile.usf
+++ b/Engine/Shaders/Private/HairStrands/HairStrandsVisibilityTile.usf
@@ -41,7 +41,7 @@ void TileMainCS(uint2 DispatchThreadId : SV_DispatchThreadID, uint LinearIndex :
 		{
 			// When bForceOutputAllTiles, we generate tiles for all pixels.
 			fHair = true;
-			pHair = true;
+			pHair = false;
 		}
 		else if (bUintTexture)
 		{
diff --git a/Engine/Source/Runtime/Renderer/Private/HairStrands/HairStrandsTile.cpp b/Engine/Source/Runtime/Renderer/Private/HairStrands/HairStrandsTile.cpp
--- a/Engine/Source/Runtime/Renderer/Private/HairStrands/HairStrandsTile.cpp
+++ b/Engine/Source/Runtime/Renderer/Private/HairStrands/HairStrandsTile.cpp
@@ -182,7 +182,7 @@ static FHairStrandsTiles AddHairStrandsGenerateTilesPass_Internal(
 
 	TShaderMapRef<FHairStrandsTileGenerationPassCS> ComputeShader(View.ShaderMap);
 	FHairStrandsTileGenerationPassCS::FParameters* PassParameters = GraphBuilder.AllocParameters<FHairStrandsTileGenerationPassCS::FParameters>();
-	PassParameters->BufferResolution		= InputResolution;//View.ViewRect.Size();
+	PassParameters->BufferResolution		= View.ViewRect.Size();
 	PassParameters->bForceOutputAllTiles	= bHasValidInput ? 0 : 1;
 	PassParameters->bUintTexture			= bUintTexture ? 1u : 0u;
 	PassParameters->TransmittanceThreshold	= 1.f - GetHairStrandsFullCoverageThreshold();	
