Strata Dbuffer pass now respect material decal responsness.

This CL uses stentil buffer bit to mark the material decal responsness, which is later use during the DBuffer pass to blend decals material data.

diff --git a/Engine/Shaders/Private/DBufferDecalShared.ush b/Engine/Shaders/Private/DBufferDecalShared.ush
--- a/Engine/Shaders/Private/DBufferDecalShared.ush
+++ b/Engine/Shaders/Private/DBufferDecalShared.ush
@@ -6,17 +6,20 @@
 
 #pragma once
 
+#ifndef DBUFFER_BASEPASS_RESOURCE
+#define DBUFFER_BASEPASS_RESOURCE 1
+#endif
+
+#if DBUFFER_BASEPASS_RESOURCE
+#define OpaqueBasePassDBufferATexture OpaqueBasePass.DBufferATexture
+#define OpaqueBasePassDBufferBTexture OpaqueBasePass.DBufferBTexture
+#define OpaqueBasePassDBufferCTexture OpaqueBasePass.DBufferCTexture
+#define OpaqueBasePassDBufferRenderMask OpaqueBasePass.DBufferRenderMask
 
-#if SUPPORTS_INDEPENDENT_SAMPLERS
-// Share samplers to give artists more
-#define OpaqueBasePassDBufferATextureSampler OpaqueBasePass.DBufferATextureSampler
-#define OpaqueBasePassDBufferBTextureSampler OpaqueBasePass.DBufferATextureSampler
-#define OpaqueBasePassDBufferCTextureSampler OpaqueBasePass.DBufferATextureSampler
-#else
 #define OpaqueBasePassDBufferATextureSampler OpaqueBasePass.DBufferATextureSampler
 #define OpaqueBasePassDBufferBTextureSampler OpaqueBasePass.DBufferBTextureSampler
 #define OpaqueBasePassDBufferCTextureSampler OpaqueBasePass.DBufferCTextureSampler
-#endif
+#endif // DBUFFER_BASEPASS_RESOURCE
 
 // Returns a bitmask of DBuffer targets that contain valid data for this pixel.
 // @param BufferUV - UV space in the DBuffer textures
@@ -25,9 +28,9 @@ uint GetDBufferTargetMask(uint2 PixelPos)
 #if USE_DBUFFER && ENABLE_DBUFFER_TEXTURES && !RAYHITGROUPSHADER && !LUMEN_CARD_CAPTURE
 
 #if PLATFORM_SUPPORTS_RENDERTARGET_WRITE_MASK
-	return DecodeRTWriteMask(PixelPos, OpaqueBasePass.DBufferRenderMask, 3);
+	return DecodeRTWriteMask(PixelPos, OpaqueBasePassDBufferRenderMask, 3);
 #elif PLATFORM_SUPPORTS_PER_PIXEL_DBUFFER_MASK
-	uint Mask = OpaqueBasePass.DBufferRenderMask.Load(uint3(PixelPos, 0));
+	uint Mask = OpaqueBasePassDBufferRenderMask.Load(uint3(PixelPos, 0));
 	return Mask > 0 ? 0x07 : 0x00;
 #else
 	return 0x07;
@@ -74,11 +77,11 @@ FStrataDBuffer GetDBufferData(float2 BufferUV, uint RTMaskBit)
 
 #if USE_DBUFFER && ENABLE_DBUFFER_TEXTURES && !RAYHITGROUPSHADER && !LUMEN_CARD_CAPTURE
 	BRANCH
-	if (RTMaskBit & 0x1) { DBufferA = Texture2DSampleLevel(OpaqueBasePass.DBufferATexture, OpaqueBasePassDBufferATextureSampler, BufferUV, 0);	}
+	if (RTMaskBit & 0x1) { DBufferA = Texture2DSampleLevel(OpaqueBasePassDBufferATexture, OpaqueBasePassDBufferATextureSampler, BufferUV, 0);	}
 	BRANCH
-	if (RTMaskBit & 0x2) { DBufferB = Texture2DSampleLevel(OpaqueBasePass.DBufferBTexture, OpaqueBasePassDBufferBTextureSampler, BufferUV, 0);	}
+	if (RTMaskBit & 0x2) { DBufferB = Texture2DSampleLevel(OpaqueBasePassDBufferBTexture, OpaqueBasePassDBufferBTextureSampler, BufferUV, 0);	}
 	BRANCH
-	if (RTMaskBit & 0x4) { DBufferC = Texture2DSampleLevel(OpaqueBasePass.DBufferCTexture, OpaqueBasePassDBufferCTextureSampler, BufferUV, 0);	}
+	if (RTMaskBit & 0x4) { DBufferC = Texture2DSampleLevel(OpaqueBasePassDBufferCTexture, OpaqueBasePassDBufferCTextureSampler, BufferUV, 0);	}
 #endif
 	return StrataUnpackDBuffer(DBufferA, DBufferB, DBufferC);
 }
@@ -92,9 +95,13 @@ FStrataDBuffer GetDBufferData(float2 BufferUV, uint RTMaskBit)
 #define STRATA_DBUFFER_BLENDING_MODE STRATA_DBUFFER_LEGACY
 
 // Define which DBuffers are valid to read from, for this material
-#define STRATA_DBUFFER_RESPONSE_BASECOLOR (MATERIALDECALRESPONSEMASK & 0x1)
-#define STRATA_DBUFFER_RESPONSE_NORMAL    (MATERIALDECALRESPONSEMASK & 0x2)
-#define STRATA_DBUFFER_RESPONSE_ROUGHNESS (MATERIALDECALRESPONSEMASK & 0x4)
+#define STRATA_DBUFFER_RESPONSE_BASECOLOR_BIT_OFFSET 0
+#define STRATA_DBUFFER_RESPONSE_NORMAL_BIT_OFFSET    1
+#define STRATA_DBUFFER_RESPONSE_ROUGHNESS_BIT_OFFSET 2
+
+#define STRATA_DBUFFER_RESPONSE_BASECOLOR (MATERIALDECALRESPONSEMASK & (1<<STRATA_DBUFFER_RESPONSE_BASECOLOR_BIT_OFFSET))
+#define STRATA_DBUFFER_RESPONSE_NORMAL    (MATERIALDECALRESPONSEMASK & (1<<STRATA_DBUFFER_RESPONSE_NORMAL_BIT_OFFSET))
+#define STRATA_DBUFFER_RESPONSE_ROUGHNESS (MATERIALDECALRESPONSEMASK & (1<<STRATA_DBUFFER_RESPONSE_ROUGHNESS_BIT_OFFSET))
 
 #if MATERIAL_IS_STRATA
 void ApplyDBufferData(in FStrataDBuffer StrataDBufferData, inout FStrataPixelHeader StrataPixelHeader, inout FStrataData StrataData)
diff --git a/Engine/Shaders/Private/Strata/StrataDbuffer.usf b/Engine/Shaders/Private/Strata/StrataDbuffer.usf
--- a/Engine/Shaders/Private/Strata/StrataDbuffer.usf
+++ b/Engine/Shaders/Private/Strata/StrataDbuffer.usf
@@ -18,11 +18,9 @@
 #define STRATA_INLINE_SHADING 0
 #define STRATA_SSS_MATERIAL_OVERRIDE 0
 #include "/Engine/Private/Strata/Strata.ush"
-
-#define MATERIALDECALRESPONSEMASK 0x7
-#include "../DBufferDecalShared.ush"
 #include "StrataTile.ush"
 
+
 ////////////////////////////////////////////////////////////////////////////////////////////////////////////
 // Post Dbuffer application
 #if SHADER_DBUFFER
@@ -35,6 +33,35 @@ SamplerState DBufferATextureSampler;
 SamplerState DBufferBTextureSampler;
 SamplerState DBufferCTextureSampler;
 
+#ifndef STRATA_STENCIL_DBUFFER_MASK
+#error STRATA_STENCIL_DBUFFER_MASK needs to be defined
+#endif
+#ifndef STENCIL_STRATA_RECEIVE_DBUFFER_NORMAL_BIT_ID
+#error STENCIL_STRATA_RECEIVE_DBUFFER_NORMAL_BIT_ID needs to be defined
+#endif
+#ifndef STENCIL_STRATA_RECEIVE_DBUFFER_DIFFUSE_BIT_ID
+#error STENCIL_STRATA_RECEIVE_DBUFFER_DIFFUSE_BIT_ID needs to be defined
+#endif
+#ifndef STENCIL_STRATA_RECEIVE_DBUFFER_ROUGHNESS_BIT_ID
+#error STENCIL_STRATA_RECEIVE_DBUFFER_ROUGHNESS_BIT_ID needs to be defined
+#endif
+
+#define MATERIALDECALRESPONSEMASK 0x7
+#define USE_DBUFFER 1
+#define ENABLE_DBUFFER_TEXTURES 1
+#define DBUFFER_BASEPASS_RESOURCE 0
+
+#define OpaqueBasePassDBufferATexture DBufferATexture
+#define OpaqueBasePassDBufferBTexture DBufferBTexture
+#define OpaqueBasePassDBufferCTexture DBufferCTexture
+#define OpaqueBasePassDBufferRenderMask DBufferRenderMask
+
+#define OpaqueBasePassDBufferBTextureSampler DBufferBTextureSampler
+#define OpaqueBasePassDBufferATextureSampler DBufferATextureSampler
+#define OpaqueBasePassDBufferCTextureSampler DBufferCTextureSampler
+
+#include "../DBufferDecalShared.ush"
+
 int2 ViewResolution;
 uint MaxBytesPerPixel;
 uint FirstSliceStoringStrataSSSData;
@@ -45,39 +72,9 @@ Buffer<uint> TileListBuffer;
 Buffer<uint> TileIndirectBuffer;
 
 // TODO
-// * Respect material decal response
 // * Add support Eye, Hair, Water, ... ?
 // * Refactor encoding to share code between strata export
 
-uint StrataGetDBufferTargetMask(uint2 PixelPos)
-{
-#if PLATFORM_SUPPORTS_RENDERTARGET_WRITE_MASK
-	return DecodeRTWriteMask(PixelPos, DBufferRenderMask, 3);
-#elif PLATFORM_SUPPORTS_PER_PIXEL_DBUFFER_MASK
-	uint Mask = DBufferRenderMask.Load(uint3(PixelPos, 0));
-	return Mask > 0 ? 0x07 : 0x00;
-#else
-	return 0x07;
-#endif
-}
-
-// @param BufferUV - UV space in the DBuffer textures
-FStrataDBuffer StrataGetDBufferData(float2 BufferUV, uint RTMaskBit)
-{
-	// Setup default values, which mean that no decals are present in DBuffer
-	float4 DBufferA = float4(0.0, 0.0, 0.0, 1.0);
-	float4 DBufferB = float4(128.0f / 255.0f, 128.f / 255.5f, 128.f / 255.5f, 1.0);
-	float4 DBufferC = float4(0.0, 0.0, 0.0, 1.0);
-
-	BRANCH
-	if (RTMaskBit & 0x1) { DBufferA = Texture2DSampleLevel(DBufferATexture, DBufferATextureSampler, BufferUV, 0); }
-	BRANCH
-	if (RTMaskBit & 0x2) { DBufferB = Texture2DSampleLevel(DBufferBTexture, DBufferBTextureSampler, BufferUV, 0); }
-	BRANCH
-	if (RTMaskBit & 0x4) { DBufferC = Texture2DSampleLevel(DBufferCTexture, DBufferCTextureSampler, BufferUV, 0); }
-	return StrataUnpackDBuffer(DBufferA, DBufferB, DBufferC);
-}
-
 [numthreads(STRATA_TILE_SIZE, STRATA_TILE_SIZE, 1)]
 void MainCS(
 	uint2 DispatchThreadId : SV_DispatchThreadID, 
@@ -97,15 +94,25 @@ void MainCS(
 		return;
 	}
 
-	uint ValidDBufferTargetMask = StrataGetDBufferTargetMask(PixelCoord);
+	// Early out of no decal has been written or if the pixel has no 'decal responseness'
+	const uint ResponseMask = SceneStencilTexture.Load(uint3(PixelCoord, 0)) STENCIL_COMPONENT_SWIZZLE;
+	uint ValidDBufferTargetMask = GetDBufferTargetMask(PixelCoord);
 	BRANCH
-	if (!ValidDBufferTargetMask)
+	if (!ValidDBufferTargetMask || (ResponseMask & STRATA_STENCIL_DBUFFER_MASK) == 0)
 	{
 		return;
 	}
 
+	// Load Dbuffer data only attributes for which the material has some 'responseness'
+	ValidDBufferTargetMask = ValidDBufferTargetMask & 
+	(
+		(((ResponseMask >> STENCIL_STRATA_RECEIVE_DBUFFER_DIFFUSE_BIT_ID)   & 0x1) << STRATA_DBUFFER_RESPONSE_BASECOLOR_BIT_OFFSET) |
+		(((ResponseMask >> STENCIL_STRATA_RECEIVE_DBUFFER_NORMAL_BIT_ID)    & 0x1) << STRATA_DBUFFER_RESPONSE_NORMAL_BIT_OFFSET)    |
+		(((ResponseMask >> STENCIL_STRATA_RECEIVE_DBUFFER_ROUGHNESS_BIT_ID) & 0x1) << STRATA_DBUFFER_RESPONSE_ROUGHNESS_BIT_OFFSET)
+	);
+
 	const float2 BufferUV = SvPositionToBufferUV(float4(PixelCoord + 0.5f, 0, 0));
-	const FStrataDBuffer DBufferData = StrataGetDBufferData(BufferUV, ValidDBufferTargetMask);
+	const FStrataDBuffer DBufferData = GetDBufferData(BufferUV, ValidDBufferTargetMask);
 
 	// Early out if DBuffer has no coverage
 	BRANCH
diff --git a/Engine/Source/Runtime/Renderer/Private/BasePassRendering.cpp b/Engine/Source/Runtime/Renderer/Private/BasePassRendering.cpp
--- a/Engine/Source/Runtime/Renderer/Private/BasePassRendering.cpp
+++ b/Engine/Source/Runtime/Renderer/Private/BasePassRendering.cpp
@@ -367,6 +367,30 @@ FMeshDrawCommandSortKey CalculateBasePassMeshStaticSortKey(EDepthDrawingMode Ear
 	return SortKey;
 }
 
+template<bool bDepthTest, ECompareFunction CompareFunction, uint32 StencilWriteMask>
+void SetDepthStencilStateForBasePass_Internal(FMeshPassProcessorRenderState& InDrawRenderState)
+{
+	InDrawRenderState.SetDepthStencilState(TStaticDepthStencilState<
+		bDepthTest, CompareFunction,
+		true, CF_Always, SO_Keep, SO_Keep, SO_Replace,
+		false, CF_Always, SO_Keep, SO_Keep, SO_Keep,
+		0xFF, StencilWriteMask>::GetRHI());
+}
+
+template<bool bDepthTest, ECompareFunction CompareFunction>
+void SetDepthStencilStateForBasePass_Internal(FMeshPassProcessorRenderState& InDrawRenderState)
+{	
+	const static bool bStrataDufferPassEnabled = Strata::IsStrataEnabled() && Strata::IsStrataDbufferPassEnabled();
+	if (bStrataDufferPassEnabled)
+	{
+		SetDepthStencilStateForBasePass_Internal<bDepthTest, CompareFunction, GET_STENCIL_BIT_MASK(STRATA_RECEIVE_DBUFFER_NORMAL, 1) | GET_STENCIL_BIT_MASK(STRATA_RECEIVE_DBUFFER_DIFFUSE, 1) | GET_STENCIL_BIT_MASK(STRATA_RECEIVE_DBUFFER_ROUGHNESS, 1) | GET_STENCIL_BIT_MASK(DISTANCE_FIELD_REPRESENTATION, 1) | STENCIL_LIGHTING_CHANNELS_MASK(0x7)>(InDrawRenderState);
+	}
+	else
+	{
+		SetDepthStencilStateForBasePass_Internal<bDepthTest, CompareFunction, GET_STENCIL_BIT_MASK(RECEIVE_DECAL, 1) | GET_STENCIL_BIT_MASK(DISTANCE_FIELD_REPRESENTATION, 1) | STENCIL_LIGHTING_CHANNELS_MASK(0x7)>(InDrawRenderState);
+	}
+}
+
 void SetDepthStencilStateForBasePass(
 	FMeshPassProcessorRenderState& DrawRenderState,
 	ERHIFeatureLevel::Type FeatureLevel,
@@ -380,30 +404,15 @@ void SetDepthStencilStateForBasePass(
 	{
 		if (bMaskedInEarlyPass)
 		{
-			DrawRenderState.SetDepthStencilState(TStaticDepthStencilState<
-				false, CF_Equal,
-				true, CF_Always, SO_Keep, SO_Keep, SO_Replace,
-				false, CF_Always, SO_Keep, SO_Keep, SO_Keep,
-				0xFF, GET_STENCIL_BIT_MASK(RECEIVE_DECAL, 1) | GET_STENCIL_BIT_MASK(DISTANCE_FIELD_REPRESENTATION, 1) | STENCIL_LIGHTING_CHANNELS_MASK(0x7)
-			>::GetRHI());
+			SetDepthStencilStateForBasePass_Internal<false, CF_Equal>(DrawRenderState);
 		}
 		else if (DrawRenderState.GetDepthStencilAccess() & FExclusiveDepthStencil::DepthWrite)
 		{
-			DrawRenderState.SetDepthStencilState(TStaticDepthStencilState<
-				true, CF_GreaterEqual,
-				true, CF_Always, SO_Keep, SO_Keep, SO_Replace,
-				false, CF_Always, SO_Keep, SO_Keep, SO_Keep,
-				0xFF, GET_STENCIL_BIT_MASK(RECEIVE_DECAL, 1) | GET_STENCIL_BIT_MASK(DISTANCE_FIELD_REPRESENTATION, 1) | STENCIL_LIGHTING_CHANNELS_MASK(0x7)
-			>::GetRHI());
+			SetDepthStencilStateForBasePass_Internal<true, CF_GreaterEqual>(DrawRenderState);
 		}
 		else
 		{
-			DrawRenderState.SetDepthStencilState(TStaticDepthStencilState<
-				false, CF_GreaterEqual,
-				true, CF_Always, SO_Keep, SO_Keep, SO_Replace,
-				false, CF_Always, SO_Keep, SO_Keep, SO_Keep,
-				0xFF, GET_STENCIL_BIT_MASK(RECEIVE_DECAL, 1) | GET_STENCIL_BIT_MASK(DISTANCE_FIELD_REPRESENTATION, 1) | STENCIL_LIGHTING_CHANNELS_MASK(0x7)
-			>::GetRHI());
+			SetDepthStencilStateForBasePass_Internal<false, CF_GreaterEqual>(DrawRenderState);
 		}
 	}
 	else if (bMaskedInEarlyPass)
@@ -413,13 +422,7 @@ void SetDepthStencilStateForBasePass(
 
 	if (bForceEnableStencilDitherState)
 	{
-		DrawRenderState.SetDepthStencilState(
-			TStaticDepthStencilState<
-			false, CF_Equal,
-			true, CF_Always, SO_Keep, SO_Keep, SO_Replace,
-			false, CF_Always, SO_Keep, SO_Keep, SO_Keep,
-			0xFF, GET_STENCIL_BIT_MASK(RECEIVE_DECAL, 1) | GET_STENCIL_BIT_MASK(DISTANCE_FIELD_REPRESENTATION, 1) | STENCIL_LIGHTING_CHANNELS_MASK(0x7)
-			>::GetRHI());
+		SetDepthStencilStateForBasePass_Internal<false, CF_Equal>(DrawRenderState);
 	}
 }
 
@@ -1509,11 +1512,32 @@ bool FBasePassMeshProcessor::Process(
 	
 	if (bEnableReceiveDecalOutput)
 	{
+		static const bool bStrataDufferPassEnabled = Strata::IsStrataEnabled() && Strata::IsStrataDbufferPassEnabled();
+
 		// Set stencil value for this draw call
 		// This is effectively extending the GBuffer using the stencil bits
-		const uint8 StencilValue = GET_STENCIL_BIT_MASK(RECEIVE_DECAL, PrimitiveSceneProxy ? !!PrimitiveSceneProxy->ReceivesDecals() : 0x00)
+		uint8 StencilValue = 0;
+		if (bStrataDufferPassEnabled)
+		{
+			const uint32 DecalResponse = MaterialResource.GetMaterialDecalResponse();
+			const bool bNormalResponse		= DecalResponse == MDR_ColorNormalRoughness || DecalResponse == MDR_ColorNormal		|| DecalResponse == MDR_NormalRoughness || DecalResponse == MDR_Normal;
+			const bool bRoughnessResponse	= DecalResponse == MDR_ColorNormalRoughness || DecalResponse == MDR_ColorRoughness	|| DecalResponse == MDR_NormalRoughness || DecalResponse == MDR_Roughness;
+			const bool bColorResponse		= DecalResponse == MDR_ColorNormalRoughness || DecalResponse == MDR_ColorNormal		|| DecalResponse == MDR_ColorRoughness	|| DecalResponse == MDR_Color;
+
+			StencilValue =
+			  GET_STENCIL_BIT_MASK(STRATA_RECEIVE_DBUFFER_NORMAL, bNormalResponse ? 0x1 : 0x0)
+			| GET_STENCIL_BIT_MASK(STRATA_RECEIVE_DBUFFER_DIFFUSE, bColorResponse ? 0x1 : 0x0)
+			| GET_STENCIL_BIT_MASK(STRATA_RECEIVE_DBUFFER_ROUGHNESS, bRoughnessResponse ? 0x1 : 0x0)
 			| GET_STENCIL_BIT_MASK(DISTANCE_FIELD_REPRESENTATION, PrimitiveSceneProxy ? PrimitiveSceneProxy->HasDistanceFieldRepresentation() : 0x00)
 			| STENCIL_LIGHTING_CHANNELS_MASK(PrimitiveSceneProxy ? PrimitiveSceneProxy->GetLightingChannelStencilValue() : 0x00);
+		}
+		else
+		{
+			StencilValue = 
+			  GET_STENCIL_BIT_MASK(RECEIVE_DECAL, PrimitiveSceneProxy ? !!PrimitiveSceneProxy->ReceivesDecals() : 0x00)
+			| GET_STENCIL_BIT_MASK(DISTANCE_FIELD_REPRESENTATION, PrimitiveSceneProxy ? PrimitiveSceneProxy->HasDistanceFieldRepresentation() : 0x00)
+			| STENCIL_LIGHTING_CHANNELS_MASK(PrimitiveSceneProxy ? PrimitiveSceneProxy->GetLightingChannelStencilValue() : 0x00);
+		}
 		DrawRenderState.SetStencilRef(StencilValue);
 	}
 
diff --git a/Engine/Source/Runtime/Renderer/Private/PostProcess/SceneRenderTargets.h b/Engine/Source/Runtime/Renderer/Private/PostProcess/SceneRenderTargets.h
--- a/Engine/Source/Runtime/Renderer/Private/PostProcess/SceneRenderTargets.h
+++ b/Engine/Source/Runtime/Renderer/Private/PostProcess/SceneRenderTargets.h
@@ -18,16 +18,20 @@
 *
 * After deferred decals, stencil is cleared to 0 and no longer packed in this way, to ensure use of fast hardware clears and HiStencil.
 */
-#define STENCIL_SANDBOX_BIT_ID				0
+#define STENCIL_SANDBOX_BIT_ID							0
 // Must match usf
-#define STENCIL_DISTANCE_FIELD_REPRESENTATION_BIT_ID 2
-#define STENCIL_TEMPORAL_RESPONSIVE_AA_BIT_ID 3
-#define STENCIL_LIGHTING_CHANNELS_BIT_ID	4
-#define STENCIL_RECEIVE_DECAL_BIT_ID		7
+#define STENCIL_DISTANCE_FIELD_REPRESENTATION_BIT_ID	2
+#define STENCIL_TEMPORAL_RESPONSIVE_AA_BIT_ID			3
+#define STENCIL_LIGHTING_CHANNELS_BIT_ID				4
+#define STENCIL_RECEIVE_DECAL_BIT_ID					7
 // Used only during the lighting pass - alias/reuse light channels (which copied from stencil to a texture prior to lighting pass)
-#define STENCIL_STRATA_FASTPATH				4 
-#define STENCIL_STRATA_SINGLEPATH			5
-#define STENCIL_STRATA_COMPLEX				6
+#define STENCIL_STRATA_FASTPATH							4 
+#define STENCIL_STRATA_SINGLEPATH						5
+#define STENCIL_STRATA_COMPLEX							6
+// Used only by Strata during the base pass when bUseDBufferPass is enabled (to mark material STRATA_DBUFFER_RESPONSE_xxx Normal/BaseColor/Roughness)
+#define STENCIL_STRATA_RECEIVE_DBUFFER_NORMAL_BIT_ID	1
+#define STENCIL_STRATA_RECEIVE_DBUFFER_DIFFUSE_BIT_ID	3
+#define STENCIL_STRATA_RECEIVE_DBUFFER_ROUGHNESS_BIT_ID	7
 
 // Outputs a compile-time constant stencil's bit mask ready to be used
 // in TStaticDepthStencilState<> template parameter. It also takes care
diff --git a/Engine/Source/Runtime/Renderer/Private/Strata/Strata.cpp b/Engine/Source/Runtime/Renderer/Private/Strata/Strata.cpp
--- a/Engine/Source/Runtime/Renderer/Private/Strata/Strata.cpp
+++ b/Engine/Source/Runtime/Renderer/Private/Strata/Strata.cpp
@@ -674,6 +674,7 @@ class FStrataDBufferPassCS : public FGlobalShader
 		SHADER_PARAMETER_RDG_TEXTURE_UAV(Texture2D, TopLayerTexture)
 		SHADER_PARAMETER_RDG_TEXTURE_UAV(RWTexture2DArray<uint>, MaterialTextureArrayUAV)
 		SHADER_PARAMETER_RDG_BUFFER_SRV(Buffer<uint>, TileListBuffer)
+		SHADER_PARAMETER_RDG_TEXTURE_SRV(Texture2D<float>, SceneStencilTexture)
 		RDG_BUFFER_ACCESS(TileIndirectBuffer, ERHIAccess::IndirectArgs)
 	END_SHADER_PARAMETER_STRUCT()
 
@@ -684,8 +685,17 @@ class FStrataDBufferPassCS : public FGlobalShader
 
 	static void ModifyCompilationEnvironment(const FGlobalShaderPermutationParameters& Parameters, FShaderCompilerEnvironment& OutEnvironment)
 	{
+		const uint32 StrataStencilDbufferMask 
+			= GET_STENCIL_BIT_MASK(STRATA_RECEIVE_DBUFFER_NORMAL, 1) 
+			| GET_STENCIL_BIT_MASK(STRATA_RECEIVE_DBUFFER_DIFFUSE, 1) 
+			| GET_STENCIL_BIT_MASK(STRATA_RECEIVE_DBUFFER_ROUGHNESS, 1);
+
 		FGlobalShader::ModifyCompilationEnvironment(Parameters, OutEnvironment);
 		OutEnvironment.SetDefine(TEXT("SHADER_DBUFFER"), 1);
+		OutEnvironment.SetDefine(TEXT("STRATA_STENCIL_DBUFFER_MASK"), StrataStencilDbufferMask);
+		OutEnvironment.SetDefine(TEXT("STENCIL_STRATA_RECEIVE_DBUFFER_NORMAL_BIT_ID"), STENCIL_STRATA_RECEIVE_DBUFFER_NORMAL_BIT_ID);
+		OutEnvironment.SetDefine(TEXT("STENCIL_STRATA_RECEIVE_DBUFFER_DIFFUSE_BIT_ID"), STENCIL_STRATA_RECEIVE_DBUFFER_DIFFUSE_BIT_ID);
+		OutEnvironment.SetDefine(TEXT("STENCIL_STRATA_RECEIVE_DBUFFER_ROUGHNESS_BIT_ID"), STENCIL_STRATA_RECEIVE_DBUFFER_ROUGHNESS_BIT_ID);
 	}
 };
 IMPLEMENT_GLOBAL_SHADER(FStrataDBufferPassCS, "/Engine/Private/Strata/StrataDBuffer.usf", "MainCS", SF_Compute);
@@ -1241,6 +1251,7 @@ void AddStrataDBufferPass(FRDGBuilder& GraphBuilder, const FMinimalSceneTextures
 			PassParameters->TopLayerTexture = GraphBuilder.CreateUAV(StrataSceneData->TopLayerTexture);
 			PassParameters->MaterialTextureArrayUAV = StrataSceneData->MaterialTextureArrayUAV;
 			PassParameters->FirstSliceStoringStrataSSSData = StrataSceneData->FirstSliceStoringStrataSSSData;
+			PassParameters->SceneStencilTexture = SceneTextures.Stencil;
 
 			PassParameters->TileListBuffer = StrataViewData->ClassificationTileListBufferSRV[TileType];
 			PassParameters->TileIndirectBuffer = StrataViewData->ClassificationTileDispatchIndirectBuffer;
