Fix warning on shader parameter binding.

diff --git a/Engine/Source/Runtime/Renderer/Private/Strata/Strata.cpp b/Engine/Source/Runtime/Renderer/Private/Strata/Strata.cpp
--- a/Engine/Source/Runtime/Renderer/Private/Strata/Strata.cpp
+++ b/Engine/Source/Runtime/Renderer/Private/Strata/Strata.cpp
@@ -671,7 +671,7 @@ class FStrataDBufferPassCS : public FGlobalShader
 		SHADER_PARAMETER(uint32, FirstSliceStoringStrataSSSData)
 		SHADER_PARAMETER_STRUCT_REF(FViewUniformShaderParameters, ViewUniformBuffer)
 		SHADER_PARAMETER_STRUCT_INCLUDE(FDBufferParameters, DBuffer)
-		SHADER_PARAMETER_RDG_TEXTURE_UAV(Texture2D, TopLayerTexture)
+		SHADER_PARAMETER_RDG_TEXTURE_UAV(RWTexture2D<uint>, TopLayerTexture)
 		SHADER_PARAMETER_RDG_TEXTURE_UAV(RWTexture2DArray<uint>, MaterialTextureArrayUAV)
 		SHADER_PARAMETER_RDG_BUFFER_SRV(Buffer<uint>, TileListBuffer)
 		SHADER_PARAMETER_RDG_TEXTURE_SRV(Texture2D<float>, SceneStencilTexture)
