Fix the iterator used for the Primitives bit array (it was using FDefaultBitArrayAllocator, while the bit array itself uses SceneRenderingAllocator).

#rnx

diff --git a/Engine/Source/Runtime/Renderer/Private/VirtualShadowMaps/VirtualShadowMapCacheManager.cpp b/Engine/Source/Runtime/Renderer/Private/VirtualShadowMaps/VirtualShadowMapCacheManager.cpp
--- a/Engine/Source/Runtime/Renderer/Private/VirtualShadowMaps/VirtualShadowMapCacheManager.cpp
+++ b/Engine/Source/Runtime/Renderer/Private/VirtualShadowMaps/VirtualShadowMapCacheManager.cpp
@@ -288,7 +288,7 @@ FVirtualShadowMapArrayCacheManager::FInvalidatingPrimitiveCollector::FInvalidati
 				// TODO: Mark that we've done this buffer? Not really harmful to redo it
 			}
 
-			for (TConstSetBitIterator<> PrimitivesIt(Primitives); PrimitivesIt; ++PrimitivesIt)
+			for (TConstSetBitIterator<SceneRenderingAllocator> PrimitivesIt(Primitives); PrimitivesIt; ++PrimitivesIt)
 			{
 				const FPersistentPrimitiveIndex PersistentPrimitiveIndex = FPersistentPrimitiveIndex{PrimitivesIt.GetIndex()};
 				
