Ensure that Strata complex material never get decal apply twice.

When DBuffer pass is enabled complex materials receive their decals during the base pass while simple & single materials receive decal during the DBuffer pass. However the material type (Simple/Single/Complex) refers to the type known at compile type, which could differs from the output of the base, due to dynamic evaluation (i.e. a material flagged as complex on CPU, might end-up being 'simple' at the end of the base pass). To avoid the DBuffer to be applied twice, the decal responsness stored in the stencil, is only fill in for simple & single materials.

diff --git a/Engine/Source/Runtime/Engine/Private/Materials/HLSLMaterialTranslator.cpp b/Engine/Source/Runtime/Engine/Private/Materials/HLSLMaterialTranslator.cpp
--- a/Engine/Source/Runtime/Engine/Private/Materials/HLSLMaterialTranslator.cpp
+++ b/Engine/Source/Runtime/Engine/Private/Materials/HLSLMaterialTranslator.cpp
@@ -10371,6 +10371,8 @@ bool FHLSLMaterialTranslator::StrataGenerateDerivedMaterialOperatorData()
 				StrataSimplificationStatus.OriginalRequestedByteSize = StrataMaterialRequestedSizeByte;
 			}
 			StrataSimplificationStatus.bRunFullSimplification = !StrataSimplificationStatus.bMaterialFitsInMemoryBudget;
+
+			MaterialCompilationOutput.StrataMaterialType = bStrataMaterialIsSimple ? 0 : bStrataMaterialIsSingle? 1 : 2;
 		}
 	}
 	while (!StrataSimplificationStatus.bMaterialFitsInMemoryBudget);
diff --git a/Engine/Source/Runtime/Engine/Private/Materials/MaterialShared.cpp b/Engine/Source/Runtime/Engine/Private/Materials/MaterialShared.cpp
--- a/Engine/Source/Runtime/Engine/Private/Materials/MaterialShared.cpp
+++ b/Engine/Source/Runtime/Engine/Private/Materials/MaterialShared.cpp
@@ -1098,6 +1098,17 @@ bool FMaterial::MaterialUsesAnisotropy_RenderThread() const
 	return RenderingThreadShaderMap ? RenderingThreadShaderMap->UsesAnisotropy() : false;
 }
 
+uint8 FMaterial::MaterialGetStrataMaterialType_GameThread() const
+{
+	return GameThreadShaderMap ? GameThreadShaderMap->GetStrataMaterialType() : false;
+}
+
+uint8 FMaterial::MaterialGetStrataMaterialType_RenderThread() const
+{
+	check(IsInParallelRenderingThread());
+	return RenderingThreadShaderMap ? RenderingThreadShaderMap->GetStrataMaterialType() : false;
+}
+
 void FMaterial::SetGameThreadShaderMap(FMaterialShaderMap* InMaterialShaderMap)
 {
 	checkSlow(IsInGameThread() || IsInAsyncLoadingThread());
diff --git a/Engine/Source/Runtime/Engine/Public/MaterialShared.h b/Engine/Source/Runtime/Engine/Public/MaterialShared.h
--- a/Engine/Source/Runtime/Engine/Public/MaterialShared.h
+++ b/Engine/Source/Runtime/Engine/Public/MaterialShared.h
@@ -688,7 +688,8 @@ public:
 		bUsesPerInstanceRandom(false),
 		bUsesVertexInterpolator(false),
 		bHasRuntimeVirtualTextureOutputNode(false),
-		bUsesAnisotropy(false)
+		bUsesAnisotropy(false),
+		StrataMaterialType(0)
 	{}
 
 	ENGINE_API bool IsSceneTextureUsed(ESceneTextureId TexId) const { return (UsedSceneTextures & (1 << TexId)) != 0; }
@@ -791,6 +792,9 @@ public:
 
 	/** true if the material uses non 0 anisotropy value */
 	LAYOUT_BITFIELD(uint8, bUsesAnisotropy, 1);
+
+	/** Strata material type, at compile type (0:simple, 1:single, 2: complex) */
+	LAYOUT_BITFIELD(uint8, StrataMaterialType, 2);
 };
 
 struct FDebugShaderPipelineInfo
@@ -1369,6 +1373,7 @@ public:
 	bool UsesVelocitySceneTexture() const { return GetContent()->MaterialCompilationOutput.UsesVelocitySceneTexture(); }
 	bool UsesDistanceCullFade() const { return GetContent()->MaterialCompilationOutput.bUsesDistanceCullFade; }
 	bool UsesAnisotropy() const { return GetContent()->MaterialCompilationOutput.bUsesAnisotropy; }
+	uint8 GetStrataMaterialType() const { return GetContent()->MaterialCompilationOutput.StrataMaterialType; }
 #if WITH_EDITOR
 	uint32 GetNumUsedUVScalars() const { return GetContent()->MaterialCompilationOutput.NumUsedUVScalars; }
 	uint32 GetNumUsedCustomInterpolatorScalars() const { return GetContent()->MaterialCompilationOutput.NumUsedCustomInterpolatorScalars; }
@@ -2093,6 +2098,10 @@ public:
 	ENGINE_API bool MaterialUsesAnisotropy_GameThread() const;
 	ENGINE_API bool MaterialUsesAnisotropy_RenderThread() const;
 
+	/** Get Strata material type (single, single, complex slab). */
+	ENGINE_API uint8 MaterialGetStrataMaterialType_GameThread() const;
+	ENGINE_API uint8 MaterialGetStrataMaterialType_RenderThread() const;
+
 	class FMaterialShaderMap* GetGameThreadShaderMap() const 
 	{ 
 		checkSlow(IsInGameThread() || IsInAsyncLoadingThread());
diff --git a/Engine/Source/Runtime/Renderer/Private/BasePassRendering.cpp b/Engine/Source/Runtime/Renderer/Private/BasePassRendering.cpp
--- a/Engine/Source/Runtime/Renderer/Private/BasePassRendering.cpp
+++ b/Engine/Source/Runtime/Renderer/Private/BasePassRendering.cpp
@@ -1519,10 +1519,15 @@ bool FBasePassMeshProcessor::Process(
 		uint8 StencilValue = 0;
 		if (bStrataDufferPassEnabled)
 		{
-			const uint32 DecalResponse = MaterialResource.GetMaterialDecalResponse();
-			const bool bNormalResponse		= DecalResponse == MDR_ColorNormalRoughness || DecalResponse == MDR_ColorNormal		|| DecalResponse == MDR_NormalRoughness || DecalResponse == MDR_Normal;
-			const bool bRoughnessResponse	= DecalResponse == MDR_ColorNormalRoughness || DecalResponse == MDR_ColorRoughness	|| DecalResponse == MDR_NormalRoughness || DecalResponse == MDR_Roughness;
-			const bool bColorResponse		= DecalResponse == MDR_ColorNormalRoughness || DecalResponse == MDR_ColorNormal		|| DecalResponse == MDR_ColorRoughness	|| DecalResponse == MDR_Color;
+			// Set material's decal responsness through stencil bit. This is only used when Strata DBuffer pass is enabled.
+			// This 'stencil marking' allows Strata's DBuffer pass to blend decal appropriately. It is only used for Simple 
+			// and Single materials, as other materials (Complex, ...) get decal applied during the base pass.
+			const uint8 StrataMaterialType	= MaterialResource.MaterialGetStrataMaterialType_RenderThread();
+			const uint32 DecalResponse		= MaterialResource.GetMaterialDecalResponse();
+			const bool bSupportDBufferPass	= StrataMaterialType == 0 || StrataMaterialType == 1; // Simple or single
+			const bool bRoughnessResponse	= bSupportDBufferPass && (DecalResponse == MDR_ColorNormalRoughness || DecalResponse == MDR_ColorRoughness	|| DecalResponse == MDR_NormalRoughness || DecalResponse == MDR_Roughness);
+			const bool bColorResponse		= bSupportDBufferPass && (DecalResponse == MDR_ColorNormalRoughness || DecalResponse == MDR_ColorNormal		|| DecalResponse == MDR_ColorRoughness	|| DecalResponse == MDR_Color);
+			const bool bNormalResponse		= bSupportDBufferPass && (DecalResponse == MDR_ColorNormalRoughness || DecalResponse == MDR_ColorNormal		|| DecalResponse == MDR_NormalRoughness || DecalResponse == MDR_Normal);
 
 			StencilValue =
 			  GET_STENCIL_BIT_MASK(STRATA_RECEIVE_DBUFFER_NORMAL, bNormalResponse ? 0x1 : 0x0)
