Add a separate VSM cache for the views of each SceneCapture component

diff --git a/Engine/Source/Runtime/Renderer/Private/DeferredShadingRenderer.cpp b/Engine/Source/Runtime/Renderer/Private/DeferredShadingRenderer.cpp
--- a/Engine/Source/Runtime/Renderer/Private/DeferredShadingRenderer.cpp
+++ b/Engine/Source/Runtime/Renderer/Private/DeferredShadingRenderer.cpp
@@ -2203,7 +2203,7 @@ void FDeferredShadingSceneRenderer::Render(FRDGBuilder& GraphBuilder)
 	// Important that this uses consistent logic throughout the frame, so evaluate once and pass in the flag from here
 	// NOTE: Must be done after  system texture initialization
 	// TODO: This doesn't take into account the potential for split screen views with separate shadow caches
-	VirtualShadowMapArray.Initialize(GraphBuilder, Scene->GetVirtualShadowMapCache(Views[0]), UseVirtualShadowMaps(ShaderPlatform, FeatureLevel));
+	VirtualShadowMapArray.Initialize(GraphBuilder, Scene->GetVirtualShadowMapCache(Views[0]), UseVirtualShadowMaps(ShaderPlatform, FeatureLevel), Views[0].bIsSceneCapture);
 
 	// if DDM_AllOpaqueNoVelocity was used, then velocity should have already been rendered as well
 	const bool bIsEarlyDepthComplete = (DepthPass.EarlyZPassMode == DDM_AllOpaque || DepthPass.EarlyZPassMode == DDM_AllOpaqueNoVelocity);
diff --git a/Engine/Source/Runtime/Renderer/Private/MobileShadingRenderer.cpp b/Engine/Source/Runtime/Renderer/Private/MobileShadingRenderer.cpp
--- a/Engine/Source/Runtime/Renderer/Private/MobileShadingRenderer.cpp
+++ b/Engine/Source/Runtime/Renderer/Private/MobileShadingRenderer.cpp
@@ -815,7 +815,7 @@ void FMobileSceneRenderer::Render(FRDGBuilder& GraphBuilder)
 	// Important that this uses consistent logic throughout the frame, so evaluate once and pass in the flag from here
 	// NOTE: Must be done after  system texture initialization
 	// TODO: This doesn't take into account the potential for split screen views with separate shadow caches
-	VirtualShadowMapArray.Initialize(GraphBuilder, Scene->GetVirtualShadowMapCache(Views[0]), UseVirtualShadowMaps(ShaderPlatform, FeatureLevel));
+	VirtualShadowMapArray.Initialize(GraphBuilder, Scene->GetVirtualShadowMapCache(Views[0]), UseVirtualShadowMaps(ShaderPlatform, FeatureLevel), Views[0].bIsSceneCapture);
 
 	GraphBuilder.SetCommandListStat(GET_STATID(STAT_CLMM_InitViews));
 
diff --git a/Engine/Source/Runtime/Renderer/Private/SceneCaptureRendering.cpp b/Engine/Source/Runtime/Renderer/Private/SceneCaptureRendering.cpp
--- a/Engine/Source/Runtime/Renderer/Private/SceneCaptureRendering.cpp
+++ b/Engine/Source/Runtime/Renderer/Private/SceneCaptureRendering.cpp
@@ -880,6 +880,25 @@ void FScene::UpdateSceneCaptureContents(USceneCaptureComponent2D* CaptureCompone
 			ViewStateInterface->RemoveLumenSceneData(this);
 		}
 
+		if(UseVirtualShadowMaps(SceneRenderer->ShaderPlatform, FeatureLevel))
+		{
+			// A separate VSM cache is added to the SceneCapture views. 
+			// This is needed to prevent the cache being invalidated on every frame.
+			// In a future iteration, it may be more efficient to share this cache with the main viewfamily.
+			if (ViewStateInterface)
+			{
+				ViewStateInterface->AddVirtualShadowMapCache(this);
+			}
+			else
+			{
+				GEngine->AddOnScreenDebugMessage((uint64)this, 10.0f, FColor::Yellow, TEXT("SceneCapture has no viewstate. This may affect Virtual Shadow Map caching performance. Consider enabling bAlwaysPersistRenderingState or bCaptureEveryFrame."));
+			}
+		}
+		else if (ViewStateInterface)
+		{
+			ViewStateInterface->RemoveVirtualShadowMapCache(this);
+		}
+
 		// Ensure that the views for this scene capture reflect any simulated camera motion for this frame
 		TOptional<FTransform> PreviousTransform = FMotionVectorSimulation::Get().GetPreviousTransform(CaptureComponent);
 
diff --git a/Engine/Source/Runtime/Renderer/Private/VirtualShadowMaps/VirtualShadowMapArray.cpp b/Engine/Source/Runtime/Renderer/Private/VirtualShadowMaps/VirtualShadowMapArray.cpp
--- a/Engine/Source/Runtime/Renderer/Private/VirtualShadowMaps/VirtualShadowMapArray.cpp
+++ b/Engine/Source/Runtime/Renderer/Private/VirtualShadowMaps/VirtualShadowMapArray.cpp
@@ -78,6 +78,13 @@ TAutoConsoleVariable<int32> CVarMaxPhysicalPages(
 	ECVF_Scalability | ECVF_RenderThreadSafe
 );
 
+TAutoConsoleVariable<int32> CVarMaxPhysicalPagesSceneCapture(
+	TEXT("r.Shadow.Virtual.MaxPhysicalPagesSceneCapture"),
+	512,
+	TEXT("Maximum number of physical pages in the pool for each scene capture component."),
+	ECVF_Scalability | ECVF_RenderThreadSafe
+);
+
 TAutoConsoleVariable<int32> CVarCacheStaticSeparate(
 	TEXT("r.Shadow.Virtual.Cache.StaticSeparate"),
 	1,
@@ -362,7 +369,7 @@ static void SetCacheDataShaderParameters(FRDGBuilder& GraphBuilder, const TArray
 }
 
 
-void FVirtualShadowMapArray::Initialize(FRDGBuilder& GraphBuilder, FVirtualShadowMapArrayCacheManager* InCacheManager, bool bInEnabled)
+void FVirtualShadowMapArray::Initialize(FRDGBuilder& GraphBuilder, FVirtualShadowMapArrayCacheManager* InCacheManager, bool bInEnabled, bool bIsSceneCapture)
 {
 	bInitialized = true;
 	bEnabled = bInEnabled;
@@ -400,7 +407,8 @@ void FVirtualShadowMapArray::Initialize(FRDGBuilder& GraphBuilder, FVirtualShado
 		// NOTE: This assumes GetMax2DTextureDimension() is a power of two on supported platforms
 		const uint32 PhysicalPagesX = FMath::DivideAndRoundDown(GetMax2DTextureDimension(), FVirtualShadowMap::PageSize);
 		check(FMath::IsPowerOfTwo(PhysicalPagesX));
-		uint32 PhysicalPagesY = FMath::DivideAndRoundUp((uint32)FMath::Max(1, CVarMaxPhysicalPages.GetValueOnRenderThread()), PhysicalPagesX);	
+		const int32 MaxPhysicalPages = (bIsSceneCapture ? CVarMaxPhysicalPagesSceneCapture : CVarMaxPhysicalPages).GetValueOnRenderThread();
+		uint32 PhysicalPagesY = FMath::DivideAndRoundUp((uint32)FMath::Max(1, MaxPhysicalPages), PhysicalPagesX);	
 
 		UniformParameters.MaxPhysicalPages = PhysicalPagesX * PhysicalPagesY;
 				
diff --git a/Engine/Source/Runtime/Renderer/Private/VirtualShadowMaps/VirtualShadowMapArray.h b/Engine/Source/Runtime/Renderer/Private/VirtualShadowMaps/VirtualShadowMapArray.h
--- a/Engine/Source/Runtime/Renderer/Private/VirtualShadowMaps/VirtualShadowMapArray.h
+++ b/Engine/Source/Runtime/Renderer/Private/VirtualShadowMaps/VirtualShadowMapArray.h
@@ -198,7 +198,7 @@ public:
 	FVirtualShadowMapArray(FScene& InScene);
 	~FVirtualShadowMapArray();
 
-	void Initialize(FRDGBuilder& GraphBuilder, FVirtualShadowMapArrayCacheManager* InCacheManager, bool bInEnabled);
+	void Initialize(FRDGBuilder& GraphBuilder, FVirtualShadowMapArrayCacheManager* InCacheManager, bool bInEnabled, bool bIsSceneCapture);
 
 	// Returns true if virtual shadow maps are enabled
 	bool IsEnabled() const
diff --git a/Engine/Source/Runtime/Renderer/Private/VirtualShadowMaps/VirtualShadowMapCacheManager.cpp b/Engine/Source/Runtime/Renderer/Private/VirtualShadowMaps/VirtualShadowMapCacheManager.cpp
--- a/Engine/Source/Runtime/Renderer/Private/VirtualShadowMaps/VirtualShadowMapCacheManager.cpp
+++ b/Engine/Source/Runtime/Renderer/Private/VirtualShadowMaps/VirtualShadowMapCacheManager.cpp
@@ -483,13 +483,17 @@ FVirtualShadowMapArrayCacheManager::FVirtualShadowMapArrayCacheManager(FScene* I
 				static const auto* CVarMaxPhysicalPagesPtr = IConsoleManager::Get().FindTConsoleVariableDataInt(TEXT("r.Shadow.Virtual.MaxPhysicalPages"));
 				const int32 MaxPhysicalPages = CVarMaxPhysicalPagesPtr->GetValueOnRenderThread();
 
+				static const auto* CVarMaxPhysicalPagesSceneCapturePtr = IConsoleManager::Get().FindTConsoleVariableDataInt(TEXT("r.Shadow.Virtual.MaxPhysicalPagesSceneCapture"));
+				const int32 MaxPhysicalPagesSceneCapture = CVarMaxPhysicalPagesSceneCapturePtr->GetValueOnRenderThread();
+
 #if !UE_BUILD_SHIPPING
 				if (!bLoggedPageOverflow)
 				{
 					UE_LOG(LogRenderer, Warning, TEXT("Virtual Shadow Map Page Pool overflow (%d page allocations were not served), this will produce visual artifacts (missing shadow), increase the page pool limit or reduce resolution bias to avoid.\n")
-						TEXT(" See r.Shadow.Virtual.MaxPhysicalPages (%d), r.Shadow.Virtual.ResolutionLodBiasLocal (%.2f), and r.Shadow.Virtual.ResolutionLodBiasDirectional (%.2f)"),
+						TEXT(" See r.Shadow.Virtual.MaxPhysicalPages (%d), r.Shadow.Virtual.MaxPhysicalPagesSceneCapture (%d), r.Shadow.Virtual.ResolutionLodBiasLocal (%.2f), and r.Shadow.Virtual.ResolutionLodBiasDirectional (%.2f)"),
 						-NumPagesFree,
 						MaxPhysicalPages,
+						MaxPhysicalPagesSceneCapture,
 						LodBiasLocal,
 						LodBiasDirectional);
 					bLoggedPageOverflow = true;
