Use provided primitive type during PSO precaching of depth pass

diff --git a/Engine/Source/Runtime/Renderer/Private/DepthRendering.cpp b/Engine/Source/Runtime/Renderer/Private/DepthRendering.cpp
--- a/Engine/Source/Runtime/Renderer/Private/DepthRendering.cpp
+++ b/Engine/Source/Runtime/Renderer/Private/DepthRendering.cpp
@@ -795,7 +795,8 @@ void FDepthPassMeshProcessor::CollectPSOInitializersInternal(
 	const FMaterial& RESTRICT MaterialResource,
 	ERasterizerFillMode MeshFillMode,
 	ERasterizerCullMode MeshCullMode,
-	bool bDitheredLODTransition, 
+	bool bDitheredLODTransition,
+	EPrimitiveType PrimitiveType,
 	TArray<FPSOPrecacheData>& PSOInitializers)
 {
 	TMeshProcessorShaders<
@@ -839,7 +840,7 @@ void FDepthPassMeshProcessor::CollectPSOInitializersInternal(
 		DepthPassShaders,
 		MeshFillMode,
 		MeshCullMode,
-		PT_TriangleList,
+		PrimitiveType,
 		bPositionOnly ? EMeshPassFeatures::PositionOnly : EMeshPassFeatures::Default,
 		PSOInitializers);
 }
@@ -1036,11 +1037,11 @@ void FDepthPassMeshProcessor::CollectPSOInitializers(const FSceneTexturesConfig&
 		const bool bAllowDitheredLODTransition = !bIsMoveable && Material.IsDitheredLODTransition();
 
 		bool bDitheredLODTransition = false;
-		CollectPSOInitializersInternal<false>(SceneTexturesConfig, VertexFactoryType, Material, MeshFillMode, MeshCullMode, bDitheredLODTransition, PSOInitializers);
+		CollectPSOInitializersInternal<false>(SceneTexturesConfig, VertexFactoryType, Material, MeshFillMode, MeshCullMode, bDitheredLODTransition, (EPrimitiveType)PreCacheParams.PrimitiveType, PSOInitializers);
 		if (bAllowDitheredLODTransition)
 		{
 			bDitheredLODTransition = true;
-			CollectPSOInitializersInternal<false>(SceneTexturesConfig, VertexFactoryType, Material, MeshFillMode, MeshCullMode, bDitheredLODTransition, PSOInitializers);
+			CollectPSOInitializersInternal<false>(SceneTexturesConfig, VertexFactoryType, Material, MeshFillMode, MeshCullMode, bDitheredLODTransition, (EPrimitiveType)PreCacheParams.PrimitiveType, PSOInitializers);
 		}
 	}
 }
@@ -1058,34 +1059,34 @@ void FDepthPassMeshProcessor::CollectDefaultMaterialPSOInitializers(
 		ERasterizerCullMode MeshCullMode = CM_None;
 		bool bDitheredLODTransition = false;
 
-		CollectPSOInitializersInternal<true>(SceneTexturesConfig, VertexFactoryType, Material, MeshFillMode, MeshCullMode, bDitheredLODTransition, PSOInitializers);
-		CollectPSOInitializersInternal<false>(SceneTexturesConfig, VertexFactoryType, Material, MeshFillMode, MeshCullMode, bDitheredLODTransition, PSOInitializers);
+		CollectPSOInitializersInternal<true>(SceneTexturesConfig, VertexFactoryType, Material, MeshFillMode, MeshCullMode, bDitheredLODTransition, PT_TriangleList, PSOInitializers);
+		CollectPSOInitializersInternal<false>(SceneTexturesConfig, VertexFactoryType, Material, MeshFillMode, MeshCullMode, bDitheredLODTransition, PT_TriangleList, PSOInitializers);
 
 		bDitheredLODTransition = true;
-		CollectPSOInitializersInternal<true>(SceneTexturesConfig, VertexFactoryType, Material, MeshFillMode, MeshCullMode, bDitheredLODTransition, PSOInitializers);
-		CollectPSOInitializersInternal<false>(SceneTexturesConfig, VertexFactoryType, Material, MeshFillMode, MeshCullMode, bDitheredLODTransition, PSOInitializers);
+		CollectPSOInitializersInternal<true>(SceneTexturesConfig, VertexFactoryType, Material, MeshFillMode, MeshCullMode, bDitheredLODTransition, PT_TriangleList, PSOInitializers);
+		CollectPSOInitializersInternal<false>(SceneTexturesConfig, VertexFactoryType, Material, MeshFillMode, MeshCullMode, bDitheredLODTransition, PT_TriangleList, PSOInitializers);
 	}
 	{
 		ERasterizerCullMode MeshCullMode = CM_CW;
 		bool bDitheredLODTransition = false;
 
-		CollectPSOInitializersInternal<true>(SceneTexturesConfig, VertexFactoryType, Material, MeshFillMode, MeshCullMode, bDitheredLODTransition, PSOInitializers);
-		CollectPSOInitializersInternal<false>(SceneTexturesConfig, VertexFactoryType, Material, MeshFillMode, MeshCullMode, bDitheredLODTransition, PSOInitializers);
+		CollectPSOInitializersInternal<true>(SceneTexturesConfig, VertexFactoryType, Material, MeshFillMode, MeshCullMode, bDitheredLODTransition, PT_TriangleList, PSOInitializers);
+		CollectPSOInitializersInternal<false>(SceneTexturesConfig, VertexFactoryType, Material, MeshFillMode, MeshCullMode, bDitheredLODTransition, PT_TriangleList, PSOInitializers);
 
 		bDitheredLODTransition = true;
-		CollectPSOInitializersInternal<true>(SceneTexturesConfig, VertexFactoryType, Material, MeshFillMode, MeshCullMode, bDitheredLODTransition, PSOInitializers);
-		CollectPSOInitializersInternal<false>(SceneTexturesConfig, VertexFactoryType, Material, MeshFillMode, MeshCullMode, bDitheredLODTransition, PSOInitializers);
+		CollectPSOInitializersInternal<true>(SceneTexturesConfig, VertexFactoryType, Material, MeshFillMode, MeshCullMode, bDitheredLODTransition, PT_TriangleList, PSOInitializers);
+		CollectPSOInitializersInternal<false>(SceneTexturesConfig, VertexFactoryType, Material, MeshFillMode, MeshCullMode, bDitheredLODTransition, PT_TriangleList, PSOInitializers);
 	}
 	{
 		ERasterizerCullMode MeshCullMode = CM_CCW;
 		bool bDitheredLODTransition = false;
 
-		CollectPSOInitializersInternal<true>(SceneTexturesConfig, VertexFactoryType, Material, MeshFillMode, MeshCullMode, bDitheredLODTransition, PSOInitializers);
-		CollectPSOInitializersInternal<false>(SceneTexturesConfig, VertexFactoryType, Material, MeshFillMode, MeshCullMode, bDitheredLODTransition, PSOInitializers);
+		CollectPSOInitializersInternal<true>(SceneTexturesConfig, VertexFactoryType, Material, MeshFillMode, MeshCullMode, bDitheredLODTransition, PT_TriangleList, PSOInitializers);
+		CollectPSOInitializersInternal<false>(SceneTexturesConfig, VertexFactoryType, Material, MeshFillMode, MeshCullMode, bDitheredLODTransition, PT_TriangleList, PSOInitializers);
 
 		bDitheredLODTransition = true;
-		CollectPSOInitializersInternal<true>(SceneTexturesConfig, VertexFactoryType, Material, MeshFillMode, MeshCullMode, bDitheredLODTransition, PSOInitializers);
-		CollectPSOInitializersInternal<false>(SceneTexturesConfig, VertexFactoryType, Material, MeshFillMode, MeshCullMode, bDitheredLODTransition, PSOInitializers);
+		CollectPSOInitializersInternal<true>(SceneTexturesConfig, VertexFactoryType, Material, MeshFillMode, MeshCullMode, bDitheredLODTransition, PT_TriangleList, PSOInitializers);
+		CollectPSOInitializersInternal<false>(SceneTexturesConfig, VertexFactoryType, Material, MeshFillMode, MeshCullMode, bDitheredLODTransition, PT_TriangleList, PSOInitializers);
 	}
 }
 
diff --git a/Engine/Source/Runtime/Renderer/Private/DepthRendering.h b/Engine/Source/Runtime/Renderer/Private/DepthRendering.h
--- a/Engine/Source/Runtime/Renderer/Private/DepthRendering.h
+++ b/Engine/Source/Runtime/Renderer/Private/DepthRendering.h
@@ -219,6 +219,7 @@ private:
 		ERasterizerFillMode MeshFillMode,
 		ERasterizerCullMode MeshCullMode,
 		bool bDitheredLODTransition, 
+		EPrimitiveType PrimitiveType,
 		TArray<FPSOPrecacheData>& PSOInitializers);
 
 	FMeshPassProcessorRenderState PassDrawRenderState;
