Mobile: Use 3 shadow cascades on Epic shadow quality (Android, iOS)
Remove r.Shadow.CSM.MaxMobileCascades, use r.Shadow.CSM.MaxCascades to limit number of cascades

diff --git a/Engine/Config/Android/AndroidScalability.ini b/Engine/Config/Android/AndroidScalability.ini
--- a/Engine/Config/Android/AndroidScalability.ini
+++ b/Engine/Config/Android/AndroidScalability.ini
@@ -70,7 +70,7 @@ r.Shadow.RadiusThreshold=0.06
 r.ShadowQuality=2
 r.Shadow.MaxCSMResolution=2048
 r.Shadow.DistanceScale=1.0
-r.Shadow.CSM.MaxCascades=2
+r.Shadow.CSM.MaxCascades=3
 r.Shadow.RadiusThreshold=0.06
 
 ; ------------------------------------------------------------------------------
diff --git a/Engine/Config/IOS/IOSScalability.ini b/Engine/Config/IOS/IOSScalability.ini
--- a/Engine/Config/IOS/IOSScalability.ini
+++ b/Engine/Config/IOS/IOSScalability.ini
@@ -70,7 +70,7 @@ r.Shadow.RadiusThreshold=0.06
 r.ShadowQuality=2
 r.Shadow.MaxCSMResolution=2048
 r.Shadow.DistanceScale=1.0
-r.Shadow.CSM.MaxCascades=2
+r.Shadow.CSM.MaxCascades=3
 r.Shadow.RadiusThreshold=0.06
 
 ; ------------------------------------------------------------------------------
diff --git a/Engine/Shaders/Private/MobileLightingCommon.ush b/Engine/Shaders/Private/MobileLightingCommon.ush
--- a/Engine/Shaders/Private/MobileLightingCommon.ush
+++ b/Engine/Shaders/Private/MobileLightingCommon.ush
@@ -140,7 +140,7 @@ half MobileDirectionalLightCSM(float2 ScreenPosition, float SceneDepth, inout fl
 	Settings.ProjectionDepthBiasParameters = 0;
 
 	float4 ShadowPosition = float4(0, 0, 0, 0);
-	for (int i = 0; i < MAX_MOBILE_SHADOWCASCADES; i++)
+	for (int i = 0; i < MobileDirectionalLight.DirectionalLightShadowMapNumCascades; i++)
 	{
 		if (SceneDepth < MobileDirectionalLight.DirectionalLightShadowDistances[i])
 		{
diff --git a/Engine/Shaders/Private/ShadingModels.ush b/Engine/Shaders/Private/ShadingModels.ush
--- a/Engine/Shaders/Private/ShadingModels.ush
+++ b/Engine/Shaders/Private/ShadingModels.ush
@@ -15,10 +15,6 @@
 #if SHADING_PATH_MOBILE
 #include "MobileGGX.ush"
 
-#ifndef MAX_MOBILE_SHADOWCASCADES
-#define MAX_MOBILE_SHADOWCASCADES 0
-#endif 
-
 #ifndef MOBILE_SHADOW_QUALITY
 #define MOBILE_SHADOW_QUALITY 2
 #endif
diff --git a/Engine/Source/Runtime/Engine/Classes/Engine/RendererSettings.h b/Engine/Source/Runtime/Engine/Classes/Engine/RendererSettings.h
--- a/Engine/Source/Runtime/Engine/Classes/Engine/RendererSettings.h
+++ b/Engine/Source/Runtime/Engine/Classes/Engine/RendererSettings.h
@@ -300,12 +300,6 @@ class ENGINE_API URendererSettings : public UDeveloperSettings
 		ConfigRestartRequired = true))
 	uint32 bMobileSupportGPUScene : 1;
 
-	UPROPERTY(config, EditAnywhere, Category = Mobile, meta = (
-		ConsoleVariable = "r.Shadow.CSM.MaxMobileCascades", DisplayName = "Maximum number of CSM cascades to render", ClampMin = 1, ClampMax = 4,
-		ToolTip = "The maximum number of cascades with which to render dynamic directional light shadows when using the mobile renderer.",
-		ConfigRestartRequired = true))
-	int32 MaxMobileCascades;
-
 	UPROPERTY(config, EditAnywhere, Category = Mobile, meta = (
 		ConsoleVariable = "r.Mobile.AntiAliasing", DisplayName = "Mobile Anti-Aliasing Method",
 		ToolTip = "The mobile default anti-aliasing method."))
diff --git a/Engine/Source/Runtime/Engine/Private/SceneManagement.cpp b/Engine/Source/Runtime/Engine/Private/SceneManagement.cpp
--- a/Engine/Source/Runtime/Engine/Private/SceneManagement.cpp
+++ b/Engine/Source/Runtime/Engine/Private/SceneManagement.cpp
@@ -728,6 +728,7 @@ FMobileDirectionalLightShaderParameters::FMobileDirectionalLightShaderParameters
 	DirectionalLightColor = FLinearColor::Black;
 	DirectionalLightDirectionAndShadowTransition = FVector4f(EForceInit::ForceInitToZero);
 	DirectionalLightShadowMapChannelMask = 0xFF;
+	DirectionalLightShadowMapNumCascades = 0;
 
 	// white texture should act like a shadowmap cleared to the farplane.
 	DirectionalLightShadowTexture = GWhiteTexture->TextureRHI;
diff --git a/Engine/Source/Runtime/Engine/Public/SceneView.h b/Engine/Source/Runtime/Engine/Public/SceneView.h
--- a/Engine/Source/Runtime/Engine/Public/SceneView.h
+++ b/Engine/Source/Runtime/Engine/Public/SceneView.h
@@ -570,6 +570,7 @@ BEGIN_GLOBAL_SHADER_PARAMETER_STRUCT_WITH_CONSTRUCTOR(FMobileDirectionalLightSha
 	SHADER_PARAMETER_EX(FVector4f, DirectionalLightShadowDistances, EShaderPrecisionModifier::Half)
 	SHADER_PARAMETER_ARRAY(FMatrix44f, DirectionalLightScreenToShadow, [MAX_MOBILE_SHADOWCASCADES])
 	SHADER_PARAMETER(uint32, DirectionalLightShadowMapChannelMask)
+	SHADER_PARAMETER(uint32, DirectionalLightShadowMapNumCascades)
 	SHADER_PARAMETER_TEXTURE(Texture2D, DirectionalLightShadowTexture)
 	SHADER_PARAMETER_SAMPLER(SamplerState, DirectionalLightShadowSampler)
 END_GLOBAL_SHADER_PARAMETER_STRUCT()
diff --git a/Engine/Source/Runtime/RenderCore/Private/Shader.cpp b/Engine/Source/Runtime/RenderCore/Private/Shader.cpp
--- a/Engine/Source/Runtime/RenderCore/Private/Shader.cpp
+++ b/Engine/Source/Runtime/RenderCore/Private/Shader.cpp
@@ -1541,11 +1541,6 @@ void ShaderMapAppendKeyString(EShaderPlatform Platform, FString& KeyString)
 			static IConsoleVariable* CVar = IConsoleManager::Get().FindConsoleVariable(TEXT("r.Mobile.DisableVertexFog"));
 			KeyString += (CVar && CVar->GetInt() != 0) ? TEXT("_NoVFog") : TEXT("");
 		}
-
-		{
-			static const auto* CVar = IConsoleManager::Get().FindTConsoleVariableDataInt(TEXT("r.Shadow.CSM.MaxMobileCascades"));
-			KeyString += (CVar) ? FString::Printf(TEXT("MMC%d"), CVar->GetValueOnAnyThread()) : TEXT("");
-		}	
 		
 		{
 			static IConsoleVariable* CVar = IConsoleManager::Get().FindConsoleVariable(TEXT("r.Mobile.FloatPrecisionMode"));
diff --git a/Engine/Source/Runtime/Renderer/Private/LightMapRendering.h b/Engine/Source/Runtime/Renderer/Private/LightMapRendering.h
--- a/Engine/Source/Runtime/Renderer/Private/LightMapRendering.h
+++ b/Engine/Source/Runtime/Renderer/Private/LightMapRendering.h
@@ -292,12 +292,6 @@ struct FCachedPointIndirectLightingPolicy
 	static void ModifyCompilationEnvironment(const FMaterialShaderPermutationParameters& Parameters, FShaderCompilerEnvironment& OutEnvironment);
 };
 
-FORCEINLINE_DEBUGGABLE int32 GetMobileMaxShadowCascades()
-{
-	static const auto* CVar = IConsoleManager::Get().FindTConsoleVariableDataInt(TEXT("r.Shadow.CSM.MaxMobileCascades"));
-	return CVar ? FMath::Clamp(CVar->GetValueOnAnyThread(), 0, MAX_MOBILE_SHADOWCASCADES) : MAX_MOBILE_SHADOWCASCADES;
-}
-
 FORCEINLINE_DEBUGGABLE bool MobileUseCSMShaderBranch()
 {
 	static const auto* CVar = IConsoleManager::Get().FindTConsoleVariableDataInt(TEXT("r.Mobile.UseCSMShaderBranch"));
@@ -310,7 +304,6 @@ struct FMobileDirectionalLightAndCSMPolicy
 	static void ModifyCompilationEnvironment(const FMaterialShaderPermutationParameters& Parameters, FShaderCompilerEnvironment& OutEnvironment)
 	{
 		OutEnvironment.SetDefine(TEXT("DIRECTIONAL_LIGHT_CSM"), TEXT("1"));
-		OutEnvironment.SetDefine(TEXT(PREPROCESSOR_TO_STRING(MAX_MOBILE_SHADOWCASCADES)), GetMobileMaxShadowCascades());
 		OutEnvironment.SetDefine(TEXT("MOBILE_USE_CSM_BRANCH"), MobileUseCSMShaderBranch());
 	}
 	
@@ -375,7 +368,6 @@ public:
 	static void ModifyCompilationEnvironment(const FMaterialShaderPermutationParameters& Parameters, FShaderCompilerEnvironment& OutEnvironment)
 	{
 		OutEnvironment.SetDefine(TEXT("DIRECTIONAL_LIGHT_CSM"), TEXT("1"));
-		OutEnvironment.SetDefine(TEXT(PREPROCESSOR_TO_STRING(MAX_MOBILE_SHADOWCASCADES)), GetMobileMaxShadowCascades());
 
 		Super::ModifyCompilationEnvironment(Parameters, OutEnvironment);
 	}
@@ -401,7 +393,6 @@ public:
 	static void ModifyCompilationEnvironment(const FMaterialShaderPermutationParameters& Parameters, FShaderCompilerEnvironment& OutEnvironment)
 	{
 		OutEnvironment.SetDefine(TEXT("DIRECTIONAL_LIGHT_CSM"), TEXT("1"));
-		OutEnvironment.SetDefine(TEXT(PREPROCESSOR_TO_STRING(MAX_MOBILE_SHADOWCASCADES)), GetMobileMaxShadowCascades());
 
 		Super::ModifyCompilationEnvironment(Parameters, OutEnvironment);
 	}
@@ -447,7 +438,6 @@ public:
 	static void ModifyCompilationEnvironment(const FMaterialShaderPermutationParameters& Parameters, FShaderCompilerEnvironment& OutEnvironment)
 	{
 		OutEnvironment.SetDefine(TEXT("DIRECTIONAL_LIGHT_CSM"), TEXT("1"));
-		OutEnvironment.SetDefine(TEXT(PREPROCESSOR_TO_STRING(MAX_MOBILE_SHADOWCASCADES)), GetMobileMaxShadowCascades());
 
 		Super::ModifyCompilationEnvironment(Parameters, OutEnvironment);
 	}
@@ -492,7 +482,6 @@ struct FMobileMovableDirectionalLightCSMWithLightmapPolicy : public FMobileMovab
 	static void ModifyCompilationEnvironment(const FMaterialShaderPermutationParameters& Parameters, FShaderCompilerEnvironment& OutEnvironment)
 	{
 		OutEnvironment.SetDefine(TEXT("DIRECTIONAL_LIGHT_CSM"), TEXT("1"));
-		OutEnvironment.SetDefine(TEXT(PREPROCESSOR_TO_STRING(MAX_MOBILE_SHADOWCASCADES)), GetMobileMaxShadowCascades());
 
 		Super::ModifyCompilationEnvironment(Parameters, OutEnvironment);
 	}
diff --git a/Engine/Source/Runtime/Renderer/Private/MobileBasePassRendering.cpp b/Engine/Source/Runtime/Renderer/Private/MobileBasePassRendering.cpp
--- a/Engine/Source/Runtime/Renderer/Private/MobileBasePassRendering.cpp
+++ b/Engine/Source/Runtime/Renderer/Private/MobileBasePassRendering.cpp
@@ -292,6 +292,8 @@ void SetupMobileDirectionalLightUniformParameters(
 					OutShadowIndex++;
 				}
 			}
+
+			Params.DirectionalLightShadowMapNumCascades = SceneView.MaxShadowCascades;
 		}
 	}
 }
diff --git a/Engine/Source/Runtime/Renderer/Private/MobileDeferredShadingPass.cpp b/Engine/Source/Runtime/Renderer/Private/MobileDeferredShadingPass.cpp
--- a/Engine/Source/Runtime/Renderer/Private/MobileDeferredShadingPass.cpp
+++ b/Engine/Source/Runtime/Renderer/Private/MobileDeferredShadingPass.cpp
@@ -72,7 +72,6 @@ class FMobileDirectionalLightFunctionPS : public FMaterialShader
 	{
 		FMaterialShader::ModifyCompilationEnvironment(Parameters, OutEnvironment);
 		FForwardLightingParameters::ModifyCompilationEnvironment(Parameters.Platform, OutEnvironment);
-		OutEnvironment.SetDefine(TEXT(PREPROCESSOR_TO_STRING(MAX_MOBILE_SHADOWCASCADES)), GetMobileMaxShadowCascades());
 		OutEnvironment.SetDefine(TEXT("USE_LIGHT_FUNCTION"), Parameters.MaterialParameters.bIsDefaultMaterial ? 0 : 1);
 		OutEnvironment.SetDefine(TEXT("USE_SHADOWMASKTEXTURE"), MobileUsesShadowMaskTexture(Parameters.Platform) ? 1u : 0u);
 		OutEnvironment.SetDefine(TEXT("MATERIAL_SHADER"), 1);
diff --git a/Engine/Source/Runtime/Renderer/Private/SceneRendering.cpp b/Engine/Source/Runtime/Renderer/Private/SceneRendering.cpp
--- a/Engine/Source/Runtime/Renderer/Private/SceneRendering.cpp
+++ b/Engine/Source/Runtime/Renderer/Private/SceneRendering.cpp
@@ -246,13 +246,6 @@ static TAutoConsoleVariable<int32> CVarMaxShadowCascades(
 	ECVF_Scalability | ECVF_RenderThreadSafe
 	);
 
-static TAutoConsoleVariable<int32> CVarMaxMobileShadowCascades(
-	TEXT("r.Shadow.CSM.MaxMobileCascades"),
-	2,
-	TEXT("The maximum number of cascades with which to render dynamic directional light shadows when using the mobile renderer."),
-	ECVF_Scalability | ECVF_RenderThreadSafe
-);
-
 static TAutoConsoleVariable<float> CVarNormalCurvatureToRoughnessBias(
 	TEXT("r.NormalCurvatureToRoughnessBias"),
 	0.0f,
@@ -947,7 +940,7 @@ void FViewInfo::Init()
 		TranslucencyLightingVolumeSize[CascadeIndex] = FVector(0);
 	}
 
-	const int32 MaxMobileShadowCascadeCount = FMath::Clamp(CVarMaxMobileShadowCascades.GetValueOnAnyThread(), 0, MAX_MOBILE_SHADOWCASCADES / FMath::Max(Family->Views.Num(), 1));
+	const int32 MaxMobileShadowCascadeCount = MAX_MOBILE_SHADOWCASCADES / FMath::Max(Family->Views.Num(), 1);
 	const int32 MaxShadowCascadeCountUpperBound = GetFeatureLevel() >= ERHIFeatureLevel::SM5 ? 10 : MaxMobileShadowCascadeCount;
 
 	MaxShadowCascades = FMath::Clamp<int32>(CVarMaxShadowCascades.GetValueOnAnyThread(), 0, MaxShadowCascadeCountUpperBound);
diff --git a/Templates/TP_AEC_HandheldARBP/Config/DefaultEngine.ini b/Templates/TP_AEC_HandheldARBP/Config/DefaultEngine.ini
--- a/Templates/TP_AEC_HandheldARBP/Config/DefaultEngine.ini
+++ b/Templates/TP_AEC_HandheldARBP/Config/DefaultEngine.ini
@@ -2,7 +2,7 @@
 [/Script/Engine.RendererSettings]
 r.MobileHDR=True
 r.Mobile.DisableVertexFog=True
-r.Shadow.CSM.MaxMobileCascades=2
+
 r.MobileMSAA=1
 r.DiscardUnusedQuality=False
 r.AllowOcclusionQueries=True
diff --git a/Templates/TP_HandheldARBP/Config/DefaultEngine.ini b/Templates/TP_HandheldARBP/Config/DefaultEngine.ini
--- a/Templates/TP_HandheldARBP/Config/DefaultEngine.ini
+++ b/Templates/TP_HandheldARBP/Config/DefaultEngine.ini
@@ -2,7 +2,6 @@
 [/Script/Engine.RendererSettings]
 r.MobileHDR=True
 r.Mobile.DisableVertexFog=True
-r.Shadow.CSM.MaxMobileCascades=2
 r.DiscardUnusedQuality=False
 r.AllowOcclusionQueries=True
 r.MinScreenRadiusForLights=0.030000
diff --git a/Templates/TP_MFG_HandheldARBP/Config/DefaultEngine.ini b/Templates/TP_MFG_HandheldARBP/Config/DefaultEngine.ini
--- a/Templates/TP_MFG_HandheldARBP/Config/DefaultEngine.ini
+++ b/Templates/TP_MFG_HandheldARBP/Config/DefaultEngine.ini
@@ -2,7 +2,6 @@
 [/Script/Engine.RendererSettings]
 r.MobileHDR=True
 r.Mobile.DisableVertexFog=True
-r.Shadow.CSM.MaxMobileCascades=2
 r.MobileMSAA=1
 r.DiscardUnusedQuality=False
 r.AllowOcclusionQueries=True
diff --git a/Templates/TP_VirtualRealityBP/Config/DefaultEngine.ini b/Templates/TP_VirtualRealityBP/Config/DefaultEngine.ini
--- a/Templates/TP_VirtualRealityBP/Config/DefaultEngine.ini
+++ b/Templates/TP_VirtualRealityBP/Config/DefaultEngine.ini
@@ -2,7 +2,6 @@
 
 [/Script/Engine.RendererSettings]
 r.Mobile.DisableVertexFog=True
-r.Shadow.CSM.MaxMobileCascades=2
 r.Mobile.AllowDitheredLODTransition=False
 r.Mobile.AllowSoftwareOcclusion=False
 r.Mobile.VirtualTextures=False
