Limit shadow light culling octree bounds to OLD_HALF_WORLD_MAX. Elements at LWC ranges (outside the octree bounds) are stored at the edge/corner of the structure and have reduced performance. This works because octree hit-testing does not actually check the node bounds, but rather only compares the query bounding box with the octants relative to the node center.

diff --git a/Engine/Source/Runtime/Renderer/Private/RendererScene.cpp b/Engine/Source/Runtime/Renderer/Private/RendererScene.cpp
--- a/Engine/Source/Runtime/Renderer/Private/RendererScene.cpp
+++ b/Engine/Source/Runtime/Renderer/Private/RendererScene.cpp
@@ -1847,7 +1847,7 @@ FScene::FScene(UWorld* InWorld, bool bInRequiresHitProxies, bool bInIsEditorScen
 ,	SkyAtmosphere(NULL)
 ,	VolumetricCloud(NULL)
 ,	PrecomputedVisibilityHandler(NULL)
-,	LocalShadowCastingLightOctree(FVector::ZeroVector, HALF_WORLD_MAX)
+,	LocalShadowCastingLightOctree(FVector::ZeroVector, UE_OLD_HALF_WORLD_MAX)
 ,	PrimitiveOctree(FVector::ZeroVector, HALF_WORLD_MAX)
 ,	bRequiresHitProxies(bInRequiresHitProxies)
 ,	bIsEditorScene(bInIsEditorScene)
