Minor PSO precaching fixes:
- use paramater FeatureLevel instead of touching possible nullptr scene
- move static stat members to global c++ files instead of storing them in the header
- fix check when PSO got rescheduled during precache background task
- only enable PSO precache validation when actual PSO precaching is enabled as well

diff --git a/Engine/Source/Runtime/Engine/Private/PSOPrecache.cpp b/Engine/Source/Runtime/Engine/Private/PSOPrecache.cpp
--- a/Engine/Source/Runtime/Engine/Private/PSOPrecache.cpp
+++ b/Engine/Source/Runtime/Engine/Private/PSOPrecache.cpp
@@ -39,9 +39,16 @@ static FAutoConsoleVariableRef CVarValidatePSOPrecaching(
 	ECVF_ReadOnly
 );
 
+// Stats for the complete PSO
+static FCriticalSection FullPSOPrecacheLock;
+static PSOCollectorStats::FPrecacheStats FullPSOPrecacheStats;
+
+// Cache stats on hash - TODO: should use FGraphicsPipelineStateInitializer as key and use FGraphicsPSOInitializerKeyFuncs for search
+static Experimental::TRobinHoodHashMap<uint64, PSOCollectorStats::FShaderStateUsage> FullPSOPrecacheMap;
+
 int32 PSOCollectorStats::IsPrecachingValidationEnabled()
 {
-	return GValidatePrecaching;
+	return PipelineStateCache::IsPSOPrecachingEnabled() && GValidatePrecaching;
 }
 
 void PSOCollectorStats::AddPipelineStateToCache(const FGraphicsPipelineStateInitializer& PSOInitializer, uint32 MeshPassType, const FVertexFactoryType* VertexFactoryType)
@@ -58,7 +65,7 @@ void PSOCollectorStats::AddPipelineStateToCache(const FGraphicsPipelineStateInit
 	static bool bFirstTime = true;
 	if (bFirstTime)
 	{
-		PSOCollectorStats::FullPSOPrecacheStats.Empty();
+		FullPSOPrecacheStats.Empty();
 		bFirstTime = false;
 	}
 
diff --git a/Engine/Source/Runtime/Engine/Public/PSOPrecache.h b/Engine/Source/Runtime/Engine/Public/PSOPrecache.h
--- a/Engine/Source/Runtime/Engine/Public/PSOPrecache.h
+++ b/Engine/Source/Runtime/Engine/Public/PSOPrecache.h
@@ -33,6 +33,7 @@ struct FPSOPrecacheParams
 		bCastShadow = true;
 		bAffectDynamicIndirectLighting = true;
 		bReverseCulling = false;
+		bDisableBackFaceCulling = false;
 		bCastShadowAsTwoSided = false;
 		bForceLODModel = false;
 		Mobility = (uint8)EComponentMobility::Static;
@@ -93,13 +94,14 @@ struct FPSOPrecacheParams
 
 			uint32 bAffectDynamicIndirectLighting : 1;
 			uint32 bReverseCulling : 1;
+			uint32 bDisableBackFaceCulling : 1;
 			uint32 bCastShadowAsTwoSided : 1;
 			uint32 bForceLODModel : 1;
 
 			uint32 Mobility : 4;
 			uint32 StencilWriteMask : 4;
 
-			uint32 Unused : 10;
+			uint32 Unused : 9;
 		};
 		uint32 Data;
 	};
@@ -255,13 +257,6 @@ namespace PSOCollectorStats
 		FPrecacheUsageData UntrackedData;		//< PSOs which are used during rendering but are currently not precached because for example the MeshPassProcessor or VertexFactory type don't support PSO precaching yet
 	};
 
-	// Stats for the complete PSO
-	static FCriticalSection FullPSOPrecacheLock;
-	static FPrecacheStats FullPSOPrecacheStats;
-
-	// Cache stats on hash - TODO: should use FGraphicsPipelineStateInitializer as key and use FGraphicsPSOInitializerKeyFuncs for search
-	static Experimental::TRobinHoodHashMap<uint64, FShaderStateUsage> FullPSOPrecacheMap;
-
 	/**
 	 * Add PSO initializer to cache for validation & state tracking
 	 */
diff --git a/Engine/Source/Runtime/RHI/Private/PipelineStateCache.cpp b/Engine/Source/Runtime/RHI/Private/PipelineStateCache.cpp
--- a/Engine/Source/Runtime/RHI/Private/PipelineStateCache.cpp
+++ b/Engine/Source/Runtime/RHI/Private/PipelineStateCache.cpp
@@ -961,9 +961,12 @@ public:
 
 		FRWScopeLock WriteLock(ActivePrecachingTasksRWLock, SLT_Write);
 		FPrecacheTask* FindResult = ActivePrecachingTasks.Find(PrecachePSOHash);
-		check(FindResult);
-		verify(ActivePrecachingTasks.Remove(PrecachePSOHash) == 1);
-		PrecachedPSOs.Add(FindResult->PSO);
+		// Might have already been processed because the task got rescheduled
+		if (FindResult != nullptr)
+		{
+			verify(ActivePrecachingTasks.Remove(PrecachePSOHash) == 1);
+			PrecachedPSOs.Add(FindResult->PSO);
+		}
 	}
 
 	void ProcessDelayedCleanup()
diff --git a/Engine/Source/Runtime/Renderer/Private/DistortionRendering.cpp b/Engine/Source/Runtime/Renderer/Private/DistortionRendering.cpp
--- a/Engine/Source/Runtime/Renderer/Private/DistortionRendering.cpp
+++ b/Engine/Source/Runtime/Renderer/Private/DistortionRendering.cpp
@@ -1030,7 +1030,7 @@ FDistortionMeshProcessor::FDistortionMeshProcessor(
 	const FMeshPassProcessorRenderState& InPassDrawRenderState,
 	const FMeshPassProcessorRenderState& InDistortionPassStateNoDepthTest,
 	FMeshPassDrawListContext* InDrawListContext)
-	: FMeshPassProcessor(EMeshPass::Distortion, Scene, Scene->GetFeatureLevel(), InViewIfDynamicMeshCommand, InDrawListContext)
+	: FMeshPassProcessor(EMeshPass::Distortion, Scene, FeatureLevel, InViewIfDynamicMeshCommand, InDrawListContext)
 	, PassDrawRenderState(InPassDrawRenderState)
 	, PassDrawRenderStateNoDepthTest(InDistortionPassStateNoDepthTest)
 {}
diff --git a/Engine/Source/Runtime/Renderer/Private/MeshPassProcessor.cpp b/Engine/Source/Runtime/Renderer/Private/MeshPassProcessor.cpp
--- a/Engine/Source/Runtime/Renderer/Private/MeshPassProcessor.cpp
+++ b/Engine/Source/Runtime/Renderer/Private/MeshPassProcessor.cpp
@@ -1809,10 +1809,8 @@ FMeshPassProcessor::FMeshDrawingPolicyOverrideSettings FMeshPassProcessor::Compu
 {
 	FMeshDrawingPolicyOverrideSettings OverrideSettings;
 	OverrideSettings.MeshPrimitiveType = PT_TriangleList;
-
-	//OverrideSettings.MeshOverrideFlags |= Mesh.bDisableBackfaceCulling ? EDrawingPolicyOverrideFlags::TwoSided : EDrawingPolicyOverrideFlags::None;
-	//OverrideSettings.MeshOverrideFlags |= Mesh.bDitheredLODTransition ? EDrawingPolicyOverrideFlags::DitheredLODTransition : EDrawingPolicyOverrideFlags::None;
-	//OverrideSettings.MeshOverrideFlags |= Mesh.bWireframe ? EDrawingPolicyOverrideFlags::Wireframe : EDrawingPolicyOverrideFlags::None;
+		
+	OverrideSettings.MeshOverrideFlags |= PrecachePSOParams.bDisableBackFaceCulling ? EDrawingPolicyOverrideFlags::TwoSided : EDrawingPolicyOverrideFlags::None;
 	OverrideSettings.MeshOverrideFlags |= PrecachePSOParams.bReverseCulling ? EDrawingPolicyOverrideFlags::ReverseCullMode : EDrawingPolicyOverrideFlags::None;
 	return OverrideSettings;
 }
@@ -2253,6 +2251,16 @@ void SetupGBufferRenderTargetInfo(const FSceneTexturesConfig& SceneTexturesConfi
 
 #if PSO_PRECACHING_VALIDATE
 
+// Stats for PSO shaders only (no state, only VS/PS/GS usage)
+static FCriticalSection PSOShadersPrecacheLock;
+static PSOCollectorStats::FPrecacheStats PSOShadersPrecacheStats;
+static Experimental::TRobinHoodHashMap<FGraphicsMinimalPipelineStateInitializer, PSOCollectorStats::FShaderStateUsage> PSOShadersPrecacheMap;
+
+// Stats for minimal PSO initializer data during MeshDrawCommand building (doesn't contain render target info)
+static FCriticalSection MinimalPSOPrecacheLock;
+static PSOCollectorStats::FPrecacheStats MinimalPSOPrecacheStats;
+static Experimental::TRobinHoodHashMap<uint64, PSOCollectorStats::FShaderStateUsage> MinimalPSOPrecacheMap;
+
 // Helper function to remove certain members of the PSO initializer to help track down issues with PSO precache misses
 // eg Remove BlendState and check if the misses are gone for example
 static FGraphicsMinimalPipelineStateInitializer PatchMinimalPipelineStateToCheck(const FGraphicsMinimalPipelineStateInitializer& MinimalPSO)
diff --git a/Engine/Source/Runtime/Renderer/Public/MeshPassProcessor.h b/Engine/Source/Runtime/Renderer/Public/MeshPassProcessor.h
--- a/Engine/Source/Runtime/Renderer/Public/MeshPassProcessor.h
+++ b/Engine/Source/Runtime/Renderer/Public/MeshPassProcessor.h
@@ -55,6 +55,7 @@ namespace EMeshPass
 		LumenTranslucencyRadianceCacheMark,
 		LumenFrontLayerTranslucencyGBuffer,
 		DitheredLODFadingOutMaskPass, /** A mini depth pass used to mark pixels with dithered LOD fading out. Currently only used by ray tracing shadows. */
+		NaniteMeshPass,
 
 #if WITH_EDITOR
 		HitProxy,
@@ -99,6 +100,7 @@ inline const TCHAR* GetMeshPassName(EMeshPass::Type MeshPass)
 	case EMeshPass::LumenTranslucencyRadianceCacheMark: return TEXT("LumenTranslucencyRadianceCacheMark");
 	case EMeshPass::LumenFrontLayerTranslucencyGBuffer: return TEXT("LumenFrontLayerTranslucencyGBuffer");
 	case EMeshPass::DitheredLODFadingOutMaskPass: return TEXT("DitheredLODFadingOutMaskPass");
+	case EMeshPass::NaniteMeshPass: return TEXT("NaniteMeshPass");
 #if WITH_EDITOR
 	case EMeshPass::HitProxy: return TEXT("HitProxy");
 	case EMeshPass::HitProxyOpaqueOnly: return TEXT("HitProxyOpaqueOnly");
@@ -2035,16 +2037,6 @@ protected:
 
 namespace PSOCollectorStats
 {
-	// Stats for PSO shaders only (no state, only VS/PS/GS usage)
-	static FCriticalSection PSOShadersPrecacheLock;
-	static FPrecacheStats PSOShadersPrecacheStats;
-	static Experimental::TRobinHoodHashMap<FGraphicsMinimalPipelineStateInitializer, FShaderStateUsage> PSOShadersPrecacheMap;
-
-	// Stats for minimal PSO initializer data during MeshDrawCommand building (doesn't contain render target info)
-	static FCriticalSection MinimalPSOPrecacheLock;
-	static FPrecacheStats MinimalPSOPrecacheStats;		
-	static Experimental::TRobinHoodHashMap<uint64, FShaderStateUsage> MinimalPSOPrecacheMap;
-
 	// Add, check & track the minimal PSO initializer during precaching and MeshDrawCommand building
 	RENDERER_API extern void AddMinimalPipelineStateToCache(const FGraphicsMinimalPipelineStateInitializer& PipelineStateInitializer, uint32 MeshPassType, const FVertexFactoryType* VertexFactoryType);
 	RENDERER_API extern void CheckMinimalPipelineStateInCache(const FGraphicsMinimalPipelineStateInitializer& PSOInitialize, uint32 MeshPassType, const FVertexFactoryType* VertexFactoryType);
