Fix robomerge issue with incorrect auto resolve

diff --git a/Engine/Source/Runtime/Renderer/Public/MeshPassProcessor.h b/Engine/Source/Runtime/Renderer/Public/MeshPassProcessor.h
--- a/Engine/Source/Runtime/Renderer/Public/MeshPassProcessor.h
+++ b/Engine/Source/Runtime/Renderer/Public/MeshPassProcessor.h
@@ -110,9 +110,9 @@ inline const TCHAR* GetMeshPassName(EMeshPass::Type MeshPass)
 	}
 
 #if WITH_EDITOR
-	static_assert(EMeshPass::Num == 26 + 4, "Need to update switch(MeshPass) after changing EMeshPass");
+	static_assert(EMeshPass::Num == 27 + 4, "Need to update switch(MeshPass) after changing EMeshPass");
 #else
-	static_assert(EMeshPass::Num == 26, "Need to update switch(MeshPass) after changing EMeshPass");
+	static_assert(EMeshPass::Num == 27, "Need to update switch(MeshPass) after changing EMeshPass");
 #endif
 
 	checkf(0, TEXT("Missing case for EMeshPass %u"), (uint32)MeshPass);
