Add dedicated 'decal tile' data to Strata (optional), to dispatch decal work only where there are decals on the frame.

This makes the dbuffer pass 0.05ms faster, but the makes the classification slower by 0.1ms. So no win for now, but further optimization are coming. Disabled for now.

diff --git a/Engine/Shaders/Private/Strata/StrataDbuffer.usf b/Engine/Shaders/Private/Strata/StrataDbuffer.usf
--- a/Engine/Shaders/Private/Strata/StrataDbuffer.usf
+++ b/Engine/Shaders/Private/Strata/StrataDbuffer.usf
@@ -1,6 +1,7 @@
 // Copyright Epic Games, Inc. All Rights Reserved.
 
 #include "/Engine/Private/Common.ush"
+#include "/Engine/Shared/StrataDefinitions.h"   
 
 #if PERMUTATION_TILETYPE == 0
 #define STRATA_FASTPATH 1
@@ -9,7 +10,7 @@
 #elif PERMUTATION_TILETYPE == 2
   // COMPLEX PATH
 #else
-#error Strata tile type non-implemented
+#error Strata tile type non-implemented 
 #endif
 
 #define OpaqueBasePass
diff --git a/Engine/Shaders/Private/Strata/StrataMaterialClassification.usf b/Engine/Shaders/Private/Strata/StrataMaterialClassification.usf
--- a/Engine/Shaders/Private/Strata/StrataMaterialClassification.usf
+++ b/Engine/Shaders/Private/Strata/StrataMaterialClassification.usf
@@ -29,11 +29,45 @@ RWBuffer<uint> SimpleTileListDataBuffer;
 RWBuffer<uint> SingleTileListDataBuffer;
 RWBuffer<uint> ComplexTileListDataBuffer;
 
+#if PERMUTATION_DECAL
+RWBuffer<uint> DecalSimpleTileListDataBuffer;
+RWBuffer<uint> DecalSingleTileListDataBuffer;
+RWBuffer<uint> DecalComplexTileListDataBuffer;
+
+Texture2D<float4> DBufferATexture;
+Texture2D<float4> DBufferBTexture;
+Texture2D<float4> DBufferCTexture;
+Texture2D<uint> DBufferRenderMask;
+
+SamplerState DBufferATextureSampler;
+SamplerState DBufferBTextureSampler;
+SamplerState DBufferCTextureSampler;
+
+// @param BufferUV - UV space in the DBuffer textures
+uint GetDBufferTargetMask(uint2 PixelPos)
+{
+	#if PLATFORM_SUPPORTS_RENDERTARGET_WRITE_MASK
+	return DecodeRTWriteMask(PixelPos, DBufferRenderMask, 3);
+	#elif PLATFORM_SUPPORTS_PER_PIXEL_DBUFFER_MASK
+	uint Mask = DBufferRenderMask.Load(uint3(PixelPos, 0));
+	return Mask > 0 ? 0x07 : 0x00;
+	#else
+	// For debug purpose:
+	// return 
+	//	(DBufferATexture.Load(uint3(PixelPos, 0)).a < 1.f ? 0x1 : 0x0) |
+	//	(DBufferBTexture.Load(uint3(PixelPos, 0)).a < 1.f ? 0x2 : 0x0) |
+	//	(DBufferCTexture.Load(uint3(PixelPos, 0)).a < 1.f ? 0x3 : 0x0) ;
+	return 0x07;
+	#endif
+}
+
+#endif // PERMUTATION_DECAL
+
 #if STRATA_OPAQUE_ROUGH_REFRACTION_ENABLED
 RWBuffer<uint> OpaqueRoughRefractionTileListDataBuffer;
 RWBuffer<uint> SSSWithoutOpaqueRoughRefractionTileListDataBuffer;
 Texture2D<float3> OpaqueRoughRefractionTexture;
-#endif
+#endif // STRATA_OPAQUE_ROUGH_REFRACTION_ENABLED
 
 #if !PERMUTATION_WAVE_OPS
 groupshared uint s_TileFlags[GROUP_THREAD_COUNT];
@@ -49,11 +83,17 @@ void TileMainCS(uint2 DispatchThreadId : SV_DispatchThreadID, uint LinearIndex :
 		TileDrawIndirectDataBuffer[GetStrataTileTypeDrawIndirectArgOffset_DWord(STRATA_TILE_TYPE_SINGLE)					+ 0] = IndexCountPerInstance;
 		TileDrawIndirectDataBuffer[GetStrataTileTypeDrawIndirectArgOffset_DWord(STRATA_TILE_TYPE_COMPLEX)					+ 0] = IndexCountPerInstance;
 		TileDrawIndirectDataBuffer[GetStrataTileTypeDrawIndirectArgOffset_DWord(STRATA_TILE_TYPE_ROUGH_REFRACT)				+ 0] = IndexCountPerInstance;
-		TileDrawIndirectDataBuffer[GetStrataTileTypeDrawIndirectArgOffset_DWord(STRATA_TILE_TYPE_SSS_WITHOUT_ROUGH_REFRACT) + 0] = IndexCountPerInstance;
+		TileDrawIndirectDataBuffer[GetStrataTileTypeDrawIndirectArgOffset_DWord(STRATA_TILE_TYPE_ROUGH_REFRACT_SSS_WITHOUT) + 0] = IndexCountPerInstance;
+
+	#if PERMUTATION_DECAL
+		TileDrawIndirectDataBuffer[GetStrataTileTypeDrawIndirectArgOffset_DWord(STRATA_TILE_TYPE_DECAL_SIMPLE) + 0]				 = IndexCountPerInstance;
+		TileDrawIndirectDataBuffer[GetStrataTileTypeDrawIndirectArgOffset_DWord(STRATA_TILE_TYPE_DECAL_SINGLE) + 0]				 = IndexCountPerInstance;
+		TileDrawIndirectDataBuffer[GetStrataTileTypeDrawIndirectArgOffset_DWord(STRATA_TILE_TYPE_DECAL_COMPLEX) + 0]			 = IndexCountPerInstance;
+	#endif
 	}
 
 	const uint2 PixelCoord = DispatchThreadId.xy + View.ViewRectMin.xy;
-	const bool bIsValid = all(DispatchThreadId.xy < View.ViewSizeAndInvSize.xy);
+	const bool bIsValid = all(DispatchThreadId.xy < uint2(View.ViewSizeAndInvSize.xy));
 
 	// If CMask data are available, we use it as a coarse evaluation to know if a tile contains any data. 
 	// * If the tile is entirely empty: we clear the header & SSS data
@@ -74,12 +114,14 @@ void TileMainCS(uint2 DispatchThreadId : SV_DispatchThreadID, uint LinearIndex :
 #endif
 
 	// Pixels outside of the view area are considered simple to enable screen borders to receive the simple permutation when not aligned to shader group size.
-	bool bContainsComplexMaterial = false;
-	bool bContainsSimpleMaterial  = false;
-	bool bContainsSingleMaterial  = false;
-	bool bContainsStrataMaterial  = false;
-	bool bContainsOpaqueRoughRefraction = false;
-	bool bContainsScreenSpaceSubsurfaceScattering = false;
+	bool bContainsComplexMaterial					= false;
+	bool bContainsSimpleMaterial					= false;
+	bool bContainsSingleMaterial					= false;
+	bool bContainsStrataMaterial					= false;
+	bool bContainsDecals							= false;
+	bool bContainsOpaqueRoughRefraction				= false;
+	bool bContainsScreenSpaceSubsurfaceScattering	= false;
+
 	FStrataOpaqueRoughRefractionData OpaqueRoughRefractionData = (FStrataOpaqueRoughRefractionData)0;
 	if (bIsValid)
 	{
@@ -97,6 +139,12 @@ void TileMainCS(uint2 DispatchThreadId : SV_DispatchThreadID, uint LinearIndex :
 		bContainsComplexMaterial = !bIsSingle && !bIsSimple;
 		bContainsScreenSpaceSubsurfaceScattering = HasSubsurface(StrataPixelHeader);
 
+	#if PERMUTATION_DECAL
+		const uint DBufferResponseMask = SceneStencilTexture.Load(uint3(PixelCoord, 0)) STENCIL_COMPONENT_SWIZZLE;
+		const uint DBufferTargetMask = GetDBufferTargetMask(PixelCoord);
+		bContainsDecals = DBufferResponseMask != 0 && DBufferTargetMask != 0;
+	#endif
+
 	#if STRATA_OPAQUE_ROUGH_REFRACTION_ENABLED
 		OpaqueRoughRefractionData = StrataUnpackOpaqueRoughRefractionData(OpaqueRoughRefractionTexture[PixelCoord]);
 		bContainsOpaqueRoughRefraction = OpaqueRoughRefractionData.OpaqueRoughRefractionEnabled > 0.0f;
@@ -119,8 +167,6 @@ void TileMainCS(uint2 DispatchThreadId : SV_DispatchThreadID, uint LinearIndex :
 			}
 			else if (bIsLegacySSSProfile)
 			{
-				const bool bDebugEnable = all(PixelCoord == View.CursorPosition);
-
 				const uint Data1 = MaterialTextureArrayUAV[uint3(PixelCoord, 1)];
 				const uint Data2 = MaterialTextureArrayUAV[uint3(PixelCoord, 2)];
 				float  RadiusScale = UnpackR8(Data1 >> 24);
@@ -169,20 +215,23 @@ void TileMainCS(uint2 DispatchThreadId : SV_DispatchThreadID, uint LinearIndex :
 
 #if PERMUTATION_WAVE_OPS
 	const bool bTileContainsStrata							= WaveActiveAnyTrue(bContainsStrataMaterial);
-//	const bool bTileContainsSimple							= WaveActiveAnyTrue(bContainsSimpleMaterial);
+	const bool bTileContainsSimple							= WaveActiveAnyTrue(bContainsSimpleMaterial);
 	const bool bTileContainsSingle							= WaveActiveAnyTrue(bContainsSingleMaterial);
 	const bool bTileContainsComplex							= WaveActiveAnyTrue(bContainsComplexMaterial);
 	const bool bTileContainsOpaqueRoughRefraction			= WaveActiveAnyTrue(bContainsOpaqueRoughRefraction);
 	const bool bTileContainsScreenSpaceSubsurfaceScattering	= WaveActiveAnyTrue(bContainsScreenSpaceSubsurfaceScattering);
+	const bool bTileContainsDecals							= WaveActiveAnyTrue(bContainsDecals);
 #else // PERMUTATION_WAVE_OPS
 
 	s_TileFlags[LinearIndex] = 
-		(bContainsStrataMaterial					? 0x1u	: 0u) | 
-//		(bContainsSimpleMaterial					? 0x2u	: 0u) | 
-		(bContainsSingleMaterial					? 0x4u	: 0u) | 
-		(bContainsComplexMaterial					? 0x8u	: 0u) | 
-		(bContainsOpaqueRoughRefraction				? 0x10u : 0u) |
-		(bContainsScreenSpaceSubsurfaceScattering	? 0x20u : 0u);
+		  (bContainsStrataMaterial					? 0x1u	: 0u)
+		| (bContainsSimpleMaterial					? 0x2u	: 0u)
+		| (bContainsSingleMaterial					? 0x4u	: 0u) 
+		| (bContainsComplexMaterial					? 0x8u	: 0u) 
+		| (bContainsOpaqueRoughRefraction			? 0x10u : 0u)
+		| (bContainsScreenSpaceSubsurfaceScattering	? 0x20u : 0u)
+		| (bContainsScreenSpaceSubsurfaceScattering	? 0x20u : 0u)
+		| (bContainsDecals							? 0x40u : 0u);
 
 	GroupMemoryBarrierWithGroupSync();
 	if (LinearIndex < 32)
@@ -215,11 +264,12 @@ void TileMainCS(uint2 DispatchThreadId : SV_DispatchThreadID, uint LinearIndex :
 	const uint FinalTileFlags = s_TileFlags[LinearIndex] | s_TileFlags[LinearIndex + 1];
 
 	const bool bTileContainsStrata							= (FinalTileFlags & 0x1u) > 0;
-//	const bool bTileContainsSimple							= (FinalTileFlags & 0x2u) > 0;
+	const bool bTileContainsSimple							= (FinalTileFlags & 0x2u) > 0;
 	const bool bTileContainsSingle							= (FinalTileFlags & 0x4u) > 0;
 	const bool bTileContainsComplex							= (FinalTileFlags & 0x8u) > 0;
 	const bool bTileContainsOpaqueRoughRefraction			= (FinalTileFlags & 0x10u) > 0;
 	const bool bTileContainsScreenSpaceSubsurfaceScattering	= (FinalTileFlags & 0x20u) > 0;
+	const bool bTileContainsDecals							= (FinalTileFlags & 0x40u) > 0;
 #endif // PERMUTATION_WAVE_OPS
 
 	if (LinearIndex < 1 && bTileContainsStrata)
@@ -255,10 +305,34 @@ void TileMainCS(uint2 DispatchThreadId : SV_DispatchThreadID, uint LinearIndex :
 		if(bTileContainsScreenSpaceSubsurfaceScattering && !bTileContainsOpaqueRoughRefraction)
 		{
 			uint WriteToIndex;
-			InterlockedAdd(TileDrawIndirectDataBuffer[GetStrataTileTypeDrawIndirectArgOffset_DWord(STRATA_TILE_TYPE_SSS_WITHOUT_ROUGH_REFRACT) + 1], 1, WriteToIndex);
+			InterlockedAdd(TileDrawIndirectDataBuffer[GetStrataTileTypeDrawIndirectArgOffset_DWord(STRATA_TILE_TYPE_ROUGH_REFRACT_SSS_WITHOUT) + 1], 1, WriteToIndex);
 			SSSWithoutOpaqueRoughRefractionTileListDataBuffer[WriteToIndex] = EncodedTile;
 		}
 	#endif
+
+	#if PERMUTATION_DECAL
+		if (bTileContainsDecals)
+		{
+			if (bTileContainsComplex)
+			{
+				uint WriteToIndex;
+				InterlockedAdd(TileDrawIndirectDataBuffer[GetStrataTileTypeDrawIndirectArgOffset_DWord(STRATA_TILE_TYPE_DECAL_COMPLEX) + 1], 1, WriteToIndex);
+				DecalComplexTileListDataBuffer[WriteToIndex] = EncodedTile;
+			}
+			else if (bTileContainsSingle)
+			{
+				uint WriteToIndex;
+				InterlockedAdd(TileDrawIndirectDataBuffer[GetStrataTileTypeDrawIndirectArgOffset_DWord(STRATA_TILE_TYPE_DECAL_SINGLE) + 1], 1, WriteToIndex);
+				DecalSingleTileListDataBuffer[WriteToIndex] = EncodedTile;
+			}
+			else // (bTileContainsSimple)
+			{
+				uint WriteToIndex;
+				InterlockedAdd(TileDrawIndirectDataBuffer[GetStrataTileTypeDrawIndirectArgOffset_DWord(STRATA_TILE_TYPE_DECAL_SIMPLE) + 1], 1, WriteToIndex);
+				DecalSimpleTileListDataBuffer[WriteToIndex] = EncodedTile;
+			}
+		}
+	#endif
 	}
 
 #if PERMUTATION_CMASK
@@ -283,19 +357,28 @@ void ArgsMainCS(uint2 DispatchThreadId : SV_DispatchThreadID)
 		TileDispatchIndirectDataBuffer[GetStrataTileTypeDispatchIndirectArgOffset_DWord(STRATA_TILE_TYPE_SINGLE)					+ 0] = TileDrawIndirectDataBuffer[GetStrataTileTypeDrawIndirectArgOffset_DWord(STRATA_TILE_TYPE_SINGLE)						+ 1];
 		TileDispatchIndirectDataBuffer[GetStrataTileTypeDispatchIndirectArgOffset_DWord(STRATA_TILE_TYPE_COMPLEX)					+ 0] = TileDrawIndirectDataBuffer[GetStrataTileTypeDrawIndirectArgOffset_DWord(STRATA_TILE_TYPE_COMPLEX)					+ 1];
 		TileDispatchIndirectDataBuffer[GetStrataTileTypeDispatchIndirectArgOffset_DWord(STRATA_TILE_TYPE_ROUGH_REFRACT)				+ 0] = TileDrawIndirectDataBuffer[GetStrataTileTypeDrawIndirectArgOffset_DWord(STRATA_TILE_TYPE_ROUGH_REFRACT)				+ 1];
-		TileDispatchIndirectDataBuffer[GetStrataTileTypeDispatchIndirectArgOffset_DWord(STRATA_TILE_TYPE_SSS_WITHOUT_ROUGH_REFRACT)	+ 0] = TileDrawIndirectDataBuffer[GetStrataTileTypeDrawIndirectArgOffset_DWord(STRATA_TILE_TYPE_SSS_WITHOUT_ROUGH_REFRACT)	+ 1];
+		TileDispatchIndirectDataBuffer[GetStrataTileTypeDispatchIndirectArgOffset_DWord(STRATA_TILE_TYPE_ROUGH_REFRACT_SSS_WITHOUT)	+ 0] = TileDrawIndirectDataBuffer[GetStrataTileTypeDrawIndirectArgOffset_DWord(STRATA_TILE_TYPE_ROUGH_REFRACT_SSS_WITHOUT)	+ 1];
+		TileDispatchIndirectDataBuffer[GetStrataTileTypeDispatchIndirectArgOffset_DWord(STRATA_TILE_TYPE_DECAL_SIMPLE)				+ 0] = TileDrawIndirectDataBuffer[GetStrataTileTypeDrawIndirectArgOffset_DWord(STRATA_TILE_TYPE_DECAL_SIMPLE)				+ 1];
+		TileDispatchIndirectDataBuffer[GetStrataTileTypeDispatchIndirectArgOffset_DWord(STRATA_TILE_TYPE_DECAL_SINGLE)				+ 0] = TileDrawIndirectDataBuffer[GetStrataTileTypeDrawIndirectArgOffset_DWord(STRATA_TILE_TYPE_DECAL_SINGLE)				+ 1];
+		TileDispatchIndirectDataBuffer[GetStrataTileTypeDispatchIndirectArgOffset_DWord(STRATA_TILE_TYPE_DECAL_COMPLEX)				+ 0] = TileDrawIndirectDataBuffer[GetStrataTileTypeDrawIndirectArgOffset_DWord(STRATA_TILE_TYPE_DECAL_COMPLEX)				+ 1];
 
 		TileDispatchIndirectDataBuffer[GetStrataTileTypeDispatchIndirectArgOffset_DWord(STRATA_TILE_TYPE_SIMPLE)					+ 1] = 1;
 		TileDispatchIndirectDataBuffer[GetStrataTileTypeDispatchIndirectArgOffset_DWord(STRATA_TILE_TYPE_SINGLE)					+ 1] = 1;
 		TileDispatchIndirectDataBuffer[GetStrataTileTypeDispatchIndirectArgOffset_DWord(STRATA_TILE_TYPE_COMPLEX)					+ 1] = 1;
 		TileDispatchIndirectDataBuffer[GetStrataTileTypeDispatchIndirectArgOffset_DWord(STRATA_TILE_TYPE_ROUGH_REFRACT)				+ 1] = 1;
-		TileDispatchIndirectDataBuffer[GetStrataTileTypeDispatchIndirectArgOffset_DWord(STRATA_TILE_TYPE_SSS_WITHOUT_ROUGH_REFRACT) + 1] = 1;
+		TileDispatchIndirectDataBuffer[GetStrataTileTypeDispatchIndirectArgOffset_DWord(STRATA_TILE_TYPE_ROUGH_REFRACT_SSS_WITHOUT) + 1] = 1;
+		TileDispatchIndirectDataBuffer[GetStrataTileTypeDispatchIndirectArgOffset_DWord(STRATA_TILE_TYPE_DECAL_SIMPLE)				+ 1] = 1;
+		TileDispatchIndirectDataBuffer[GetStrataTileTypeDispatchIndirectArgOffset_DWord(STRATA_TILE_TYPE_DECAL_SINGLE)				+ 1] = 1;
+		TileDispatchIndirectDataBuffer[GetStrataTileTypeDispatchIndirectArgOffset_DWord(STRATA_TILE_TYPE_DECAL_COMPLEX)				+ 1] = 1;
 
 		TileDispatchIndirectDataBuffer[GetStrataTileTypeDispatchIndirectArgOffset_DWord(STRATA_TILE_TYPE_SIMPLE)					+ 2] = 1;
 		TileDispatchIndirectDataBuffer[GetStrataTileTypeDispatchIndirectArgOffset_DWord(STRATA_TILE_TYPE_SINGLE)					+ 2] = 1;
 		TileDispatchIndirectDataBuffer[GetStrataTileTypeDispatchIndirectArgOffset_DWord(STRATA_TILE_TYPE_COMPLEX)					+ 2] = 1;
 		TileDispatchIndirectDataBuffer[GetStrataTileTypeDispatchIndirectArgOffset_DWord(STRATA_TILE_TYPE_ROUGH_REFRACT)				+ 2] = 1;
-		TileDispatchIndirectDataBuffer[GetStrataTileTypeDispatchIndirectArgOffset_DWord(STRATA_TILE_TYPE_SSS_WITHOUT_ROUGH_REFRACT)	+ 2] = 1;
+		TileDispatchIndirectDataBuffer[GetStrataTileTypeDispatchIndirectArgOffset_DWord(STRATA_TILE_TYPE_ROUGH_REFRACT_SSS_WITHOUT)	+ 2] = 1;
+		TileDispatchIndirectDataBuffer[GetStrataTileTypeDispatchIndirectArgOffset_DWord(STRATA_TILE_TYPE_DECAL_SIMPLE)				+ 2] = 1;
+		TileDispatchIndirectDataBuffer[GetStrataTileTypeDispatchIndirectArgOffset_DWord(STRATA_TILE_TYPE_DECAL_SINGLE)				+ 2] = 1;
+		TileDispatchIndirectDataBuffer[GetStrataTileTypeDispatchIndirectArgOffset_DWord(STRATA_TILE_TYPE_DECAL_COMPLEX)				+ 2] = 1;
 	}
 }
 
diff --git a/Engine/Shaders/Shared/StrataDefinitions.h b/Engine/Shaders/Shared/StrataDefinitions.h
--- a/Engine/Shaders/Shared/StrataDefinitions.h
+++ b/Engine/Shaders/Shared/StrataDefinitions.h
@@ -56,8 +56,11 @@
 #define STRATA_TILE_TYPE_SINGLE						1
 #define STRATA_TILE_TYPE_COMPLEX					2
 #define STRATA_TILE_TYPE_ROUGH_REFRACT				3
-#define STRATA_TILE_TYPE_SSS_WITHOUT_ROUGH_REFRACT	4
-#define STRATA_TILE_TYPE_COUNT						5
+#define STRATA_TILE_TYPE_ROUGH_REFRACT_SSS_WITHOUT	4
+#define STRATA_TILE_TYPE_DECAL_SIMPLE				5
+#define STRATA_TILE_TYPE_DECAL_SINGLE				6
+#define STRATA_TILE_TYPE_DECAL_COMPLEX				7
+#define STRATA_TILE_TYPE_COUNT						8
 
 // sizeof(FRHIDrawIndirectParameters) = 4 uints = 16 bytes
 #define GetStrataTileTypeDrawIndirectArgOffset_Byte(x)  (x * 16)
diff --git a/Engine/Source/Runtime/Renderer/Private/DeferredShadingRenderer.cpp b/Engine/Source/Runtime/Renderer/Private/DeferredShadingRenderer.cpp
--- a/Engine/Source/Runtime/Renderer/Private/DeferredShadingRenderer.cpp
+++ b/Engine/Source/Runtime/Renderer/Private/DeferredShadingRenderer.cpp
@@ -3055,7 +3055,7 @@ void FDeferredShadingSceneRenderer::Render(FRDGBuilder& GraphBuilder)
 	// This needs to run before virtual shadow map, in order to have ready&cleared classified SSS data
 	if (Strata::IsStrataEnabled())
 	{
-		Strata::AddStrataMaterialClassificationPass(GraphBuilder, SceneTextures, Views);
+		Strata::AddStrataMaterialClassificationPass(GraphBuilder, SceneTextures, DBufferTextures, Views);
 		Strata::AddStrataDBufferPass(GraphBuilder, SceneTextures, DBufferTextures, Views);
 	}
 
diff --git a/Engine/Source/Runtime/Renderer/Private/Strata/Strata.cpp b/Engine/Source/Runtime/Renderer/Private/Strata/Strata.cpp
--- a/Engine/Source/Runtime/Renderer/Private/Strata/Strata.cpp
+++ b/Engine/Source/Runtime/Renderer/Private/Strata/Strata.cpp
@@ -92,6 +92,12 @@ static TAutoConsoleVariable<int32> CVarStrataDBufferPass(
 	TEXT("Apply DBuffer after the base-pass as a separate pass."),
 	ECVF_RenderThreadSafe);
 
+static TAutoConsoleVariable<int32> CVarStrataDBufferPassDedicatedTiles(
+	TEXT("r.Strata.DBufferPass.DedicatedTiles"),
+	0,
+	TEXT("Use dedicated tile for DBuffer application when DBuffer pass is enabled."),
+	ECVF_RenderThreadSafe);
+
 IMPLEMENT_GLOBAL_SHADER_PARAMETER_STRUCT(FStrataGlobalUniformParameters, "Strata");
 
 extern bool IsStrataAdvancedVisualizationShadersEnabled();
@@ -114,8 +120,12 @@ const TCHAR* ToString(EStrataTileType Type)
 	case EStrataTileType::ESimple:							return TEXT("Simple");
 	case EStrataTileType::ESingle:							return TEXT("Single");
 	case EStrataTileType::EComplex:							return TEXT("Complex");
-	case EStrataTileType::EOpaqueRoughRefraction:			return TEXT("OpaqueRoughRefraction");
-	case EStrataTileType::ESSSWithoutOpaqueRoughRefraction:	return TEXT("SSSWithoutOpaqueRoughRefraction");
+	case EStrataTileType::EOpaqueRoughRefraction:			return TEXT("Opaque/RoughRefraction");
+	case EStrataTileType::EOpaqueRoughRefractionSSSWithout:	return TEXT("Opaque/RoughRefraction/SSSWithout");
+	case EStrataTileType::EDecalSimple:						return TEXT("Decal/Simple");
+	case EStrataTileType::EDecalSingle:						return TEXT("Decal/Single");
+	case EStrataTileType::EDecalComplex:					return TEXT("Decal/Complex");
+
 	}
 	return TEXT("Unknown");
 }
@@ -205,7 +215,10 @@ static void InitialiseStrataViewData(FRDGBuilder& GraphBuilder, FViewInfo& View,
 			TEXT("Strata.StrataTileListBuffer(Single)"),
 			TEXT("Strata.StrataTileListBuffer(Complex)"),
 			TEXT("Strata.StrataTileListBuffer(OpaqueRoughRefraction)"),
-			TEXT("Strata.StrataTileListBuffer(SSSWithoutOpaqueRoughRefraction)")
+			TEXT("Strata.StrataTileListBuffer(OpaqueRoughRefractionSSSWithout)"),
+			TEXT("Strata.StrataTileListBuffer(DecalSimple)"),
+			TEXT("Strata.StrataTileListBuffer(DecalSingle)"),
+			TEXT("Strata.StrataTileListBuffer(DecalComplex)"),
 		};
 
 		// Tile classification buffers
@@ -220,29 +233,40 @@ static void InitialiseStrataViewData(FRDGBuilder& GraphBuilder, FViewInfo& View,
 			Out.ClassificationTileDispatchIndirectBufferUAV = GraphBuilder.CreateUAV(Out.ClassificationTileDispatchIndirectBuffer, PF_R32_UINT);
 			AddClearUAVPass(GraphBuilder, Out.ClassificationTileDispatchIndirectBufferUAV, 0);
 
+			// Separated subsurface & rough refraction textures (tile data)
+			const uint32 RoughTileCount = IsStrataOpaqueMaterialRoughRefractionEnabled() ? TileResolution.X * TileResolution.Y : 4;
+			const uint32 DecalTileCount = IsStrataDbufferPassEnabled() ? TileResolution.X * TileResolution.Y : 4;
+			const uint32 RegularTileCount = TileResolution.X * TileResolution.Y;
+
 			const EPixelFormat ClassificationTileFormat = GetClassificationTileFormat(ViewResolution);
-			for (uint32 i = 0; i <= EStrataTileType::EComplex; ++i)
+			for (uint32 i = 0; i < EStrataTileType::ECount; ++i)
 			{
-				Out.ClassificationTileListBuffer[i] = GraphBuilder.CreateBuffer(FRDGBufferDesc::CreateBufferDesc(ClassificationTileFormat == PF_R16_UINT ? sizeof(uint16) : sizeof(uint32), TileResolution.X * TileResolution.Y), StrataTileListBufferNames[i]);
+				const EStrataTileType TileType = EStrataTileType(i);
+				const uint32 FormatBytes = ClassificationTileFormat == PF_R16_UINT ? sizeof(uint16) : sizeof(uint32);
+
+				uint32 TileCount = 0;
+				switch (TileType)
+				{
+				case EStrataTileType::ESimple:
+				case EStrataTileType::ESingle:
+				case EStrataTileType::EComplex:
+					TileCount = RegularTileCount; break;
+				case EStrataTileType::EOpaqueRoughRefraction:
+				case EStrataTileType::EOpaqueRoughRefractionSSSWithout:
+					TileCount = RoughTileCount; break;
+				case EStrataTileType::EDecalSimple:
+				case EStrataTileType::EDecalSingle:
+				case EStrataTileType::EDecalComplex:
+					TileCount = DecalTileCount; break;
+				}
+				check(TileCount > 0);
+
+				Out.ClassificationTileListBuffer[i] = GraphBuilder.CreateBuffer(FRDGBufferDesc::CreateBufferDesc(FormatBytes, TileCount), StrataTileListBufferNames[i]);
 				Out.ClassificationTileListBufferSRV[i] = GraphBuilder.CreateSRV(Out.ClassificationTileListBuffer[i], ClassificationTileFormat);
 				Out.ClassificationTileListBufferUAV[i] = GraphBuilder.CreateUAV(Out.ClassificationTileListBuffer[i], ClassificationTileFormat);
 			}
 		}
 
-		// Separated subsurface & rough refraction textures (tile data)
-		{
-			const bool bIsStrataOpaqueMaterialRoughRefractionEnabled= IsStrataOpaqueMaterialRoughRefractionEnabled();
-			const int32 TileListBufferElementCount					= bIsStrataOpaqueMaterialRoughRefractionEnabled ? TileResolution.X * TileResolution.Y : 4;
-			
-			Out.ClassificationTileListBuffer[EStrataTileType::EOpaqueRoughRefraction] = GraphBuilder.CreateBuffer(FRDGBufferDesc::CreateBufferDesc(sizeof(uint32), TileListBufferElementCount), StrataTileListBufferNames[EStrataTileType::EOpaqueRoughRefraction]);
-			Out.ClassificationTileListBufferSRV[EStrataTileType::EOpaqueRoughRefraction] = GraphBuilder.CreateSRV(Out.ClassificationTileListBuffer[EStrataTileType::EOpaqueRoughRefraction], PF_R32_UINT);
-			Out.ClassificationTileListBufferUAV[EStrataTileType::EOpaqueRoughRefraction] = GraphBuilder.CreateUAV(Out.ClassificationTileListBuffer[EStrataTileType::EOpaqueRoughRefraction], PF_R32_UINT);
-
-			Out.ClassificationTileListBuffer[EStrataTileType::ESSSWithoutOpaqueRoughRefraction] = GraphBuilder.CreateBuffer(FRDGBufferDesc::CreateBufferDesc(sizeof(uint32), TileListBufferElementCount), StrataTileListBufferNames[EStrataTileType::ESSSWithoutOpaqueRoughRefraction]);
-			Out.ClassificationTileListBufferSRV[EStrataTileType::ESSSWithoutOpaqueRoughRefraction] = GraphBuilder.CreateSRV(Out.ClassificationTileListBuffer[EStrataTileType::ESSSWithoutOpaqueRoughRefraction], PF_R32_UINT);
-			Out.ClassificationTileListBufferUAV[EStrataTileType::ESSSWithoutOpaqueRoughRefraction] = GraphBuilder.CreateUAV(Out.ClassificationTileListBuffer[EStrataTileType::ESSSWithoutOpaqueRoughRefraction], PF_R32_UINT);
-		}
-
 		// BSDF tiles
 		if (bNeedBSDFOffets)
 		{			
@@ -610,7 +634,8 @@ class FStrataMaterialTileClassificationPassCS : public FGlobalShader
 
 	class FCmask : SHADER_PERMUTATION_BOOL("PERMUTATION_CMASK");
 	class FWaveOps : SHADER_PERMUTATION_BOOL("PERMUTATION_WAVE_OPS");
-	using FPermutationDomain = TShaderPermutationDomain<FWaveOps, FCmask>;
+	class FDecal : SHADER_PERMUTATION_BOOL("PERMUTATION_DECAL"); 
+	using FPermutationDomain = TShaderPermutationDomain<FCmask, FWaveOps, FDecal>;
 
 	BEGIN_SHADER_PARAMETER_STRUCT(FParameters, )
 		SHADER_PARAMETER_STRUCT_REF(FViewUniformShaderParameters, ViewUniformBuffer)
@@ -625,9 +650,14 @@ class FStrataMaterialTileClassificationPassCS : public FGlobalShader
 		SHADER_PARAMETER_RDG_BUFFER_UAV(RWBuffer, SingleTileListDataBuffer)
 		SHADER_PARAMETER_RDG_BUFFER_UAV(RWBuffer, ComplexTileListDataBuffer)
 		SHADER_PARAMETER_RDG_BUFFER_UAV(RWBuffer, OpaqueRoughRefractionTileListDataBuffer)
-		SHADER_PARAMETER_RDG_BUFFER_UAV(RWBuffer, SSSWithoutOpaqueRoughRefractionTileListDataBuffer)
+		SHADER_PARAMETER_RDG_BUFFER_UAV(RWBuffer, OpaqueRoughRefractionSSSWithoutTileListDataBuffer)
+		SHADER_PARAMETER_RDG_BUFFER_UAV(RWBuffer, DecalSimpleTileListDataBuffer)
+		SHADER_PARAMETER_RDG_BUFFER_UAV(RWBuffer, DecalSingleTileListDataBuffer)
+		SHADER_PARAMETER_RDG_BUFFER_UAV(RWBuffer, DecalComplexTileListDataBuffer)
 		SHADER_PARAMETER_RDG_TEXTURE_UAV(RWTexture2DArray<uint>, MaterialTextureArrayUAV)
 		SHADER_PARAMETER_RDG_TEXTURE(Texture2D<float3>, OpaqueRoughRefractionTexture)
+		SHADER_PARAMETER_STRUCT_INCLUDE(FDBufferParameters, DBuffer)
+		SHADER_PARAMETER_RDG_TEXTURE_SRV(Texture2D<float>, SceneStencilTexture)
 	END_SHADER_PARAMETER_STRUCT()
 
 	static bool ShouldCompilePermutation(const FGlobalShaderPermutationParameters& Parameters)
@@ -662,7 +692,7 @@ class FStrataDBufferPassCS : public FGlobalShader
 	DECLARE_GLOBAL_SHADER(FStrataDBufferPassCS);
 	SHADER_USE_PARAMETER_STRUCT(FStrataDBufferPassCS, FGlobalShader);
 
-	class FTileType : SHADER_PERMUTATION_INT("PERMUTATION_TILETYPE", 3); // Only support Simple/Single/Complex
+	class FTileType : SHADER_PERMUTATION_INT("PERMUTATION_TILETYPE", 3);
 	using FPermutationDomain = TShaderPermutationDomain<FTileType>;
 
 	BEGIN_SHADER_PARAMETER_STRUCT(FParameters, )
@@ -899,8 +929,13 @@ void AddStrataInternalClassificationTilePass(
 		case EStrataTileType::ESimple:							ParametersPS->DebugTileColor = FVector4f(0.0f, 1.0f, 0.0f, 1.0); break;
 		case EStrataTileType::ESingle:							ParametersPS->DebugTileColor = FVector4f(1.0f, 1.0f, 0.0f, 1.0); break;
 		case EStrataTileType::EComplex:							ParametersPS->DebugTileColor = FVector4f(1.0f, 0.0f, 0.0f, 1.0); break;
+
 		case EStrataTileType::EOpaqueRoughRefraction:			ParametersPS->DebugTileColor = FVector4f(0.0f, 1.0f, 1.0f, 1.0); break;
-		case EStrataTileType::ESSSWithoutOpaqueRoughRefraction:	ParametersPS->DebugTileColor = FVector4f(0.0f, 0.0f, 1.0f, 1.0); break;
+		case EStrataTileType::EOpaqueRoughRefractionSSSWithout:	ParametersPS->DebugTileColor = FVector4f(0.0f, 0.0f, 1.0f, 1.0); break;
+
+		case EStrataTileType::EDecalSingle:						ParametersPS->DebugTileColor = FVector4f(0.0f, 1.0f, 0.0f, 1.0); break;
+		case EStrataTileType::EDecalSimple:						ParametersPS->DebugTileColor = FVector4f(1.0f, 1.0f, 0.0f, 1.0); break;
+		case EStrataTileType::EDecalComplex:					ParametersPS->DebugTileColor = FVector4f(1.0f, 0.0f, 0.0f, 1.0); break;
 		default: check(false);
 		}
 	}
@@ -934,7 +969,7 @@ void AddStrataInternalClassificationTilePass(
 			}
 			else
 			{
-				check(TileMaterialType != EStrataTileType::ECount && TileMaterialType != EStrataTileType::EOpaqueRoughRefraction && TileMaterialType != EStrataTileType::ESSSWithoutOpaqueRoughRefraction);
+				check(TileMaterialType != EStrataTileType::ECount && TileMaterialType != EStrataTileType::EOpaqueRoughRefraction && TileMaterialType != EStrataTileType::EOpaqueRoughRefractionSSSWithout);
 
 				// No blending and no pixel shader required. Stencil will be written to.
 				GraphicsPSOInit.BoundShaderState.PixelShaderRHI = nullptr;
@@ -1068,7 +1103,7 @@ void SetBasePassRenderTargetOutputFormat(const EShaderPlatform Platform, const F
 
 //#include "RendererUtils.h"
 
-void AddStrataMaterialClassificationPass(FRDGBuilder& GraphBuilder, const FMinimalSceneTextures& SceneTextures, const TArray<FViewInfo>& Views)
+void AddStrataMaterialClassificationPass(FRDGBuilder& GraphBuilder, const FMinimalSceneTextures& SceneTextures, const FDBufferTextures& DBufferTextures, const TArray<FViewInfo>& Views)
 {
 	RDG_EVENT_SCOPE_CONDITIONAL(GraphBuilder, IsStrataEnabled() && Views.Num() > 0, "Strata::MaterialClassification");
 	if (!IsStrataEnabled())
@@ -1102,9 +1137,13 @@ void AddStrataMaterialClassificationPass(FRDGBuilder& GraphBuilder, const FMinim
 				FRenderTargetWriteMask::Decode(GraphBuilder, View.ShaderMap, MakeArrayView(SourceCMaskTextures), TopLayerCmaskTexture, GFastVRamConfig.DBufferMask, TEXT("Strata::TopLayerCmask"));
 			}
 
+			// If Dbuffer pass (i.e. apply DBuffer data after the base-pass) is enabled, run special classification for outputing tile with/without tiles
+			const bool bDBufferTiles = CVarStrataDBufferPass.GetValueOnRenderThread() > 0 && CVarStrataDBufferPassDedicatedTiles.GetValueOnRenderThread() > 0 && DBufferTextures.IsValid();
+
 			FStrataMaterialTileClassificationPassCS::FPermutationDomain PermutationVector;
 			PermutationVector.Set< FStrataMaterialTileClassificationPassCS::FCmask >(bSupportCMask);
 			PermutationVector.Set< FStrataMaterialTileClassificationPassCS::FWaveOps >(bWaveOps);
+			PermutationVector.Set< FStrataMaterialTileClassificationPassCS::FDecal>(bDBufferTiles);
 			TShaderMapRef<FStrataMaterialTileClassificationPassCS> ComputeShader(View.ShaderMap, PermutationVector);
 			FStrataMaterialTileClassificationPassCS::FParameters* PassParameters = GraphBuilder.AllocParameters<FStrataMaterialTileClassificationPassCS::FParameters>();
 			PassParameters->ViewUniformBuffer = View.ViewUniformBuffer;
@@ -1117,11 +1156,19 @@ void AddStrataMaterialClassificationPass(FRDGBuilder& GraphBuilder, const FMinim
 			PassParameters->MaterialTextureArrayUAV = StrataSceneData->MaterialTextureArrayUAV;
 			PassParameters->OpaqueRoughRefractionTexture = StrataSceneData->OpaqueRoughRefractionTexture;
 			PassParameters->TileDrawIndirectDataBuffer = StrataViewData->ClassificationTileDrawIndirectBufferUAV;
+			PassParameters->DBuffer = GetDBufferParameters(GraphBuilder, DBufferTextures, View.GetShaderPlatform());
+			PassParameters->SceneStencilTexture = SceneTextures.Stencil;
+			// Regular tiles
 			PassParameters->SimpleTileListDataBuffer = StrataViewData->ClassificationTileListBufferUAV[EStrataTileType::ESimple];
 			PassParameters->SingleTileListDataBuffer = StrataViewData->ClassificationTileListBufferUAV[EStrataTileType::ESingle];
 			PassParameters->ComplexTileListDataBuffer = StrataViewData->ClassificationTileListBufferUAV[EStrataTileType::EComplex];
+			// Decal tiles
+			PassParameters->DecalSimpleTileListDataBuffer = StrataViewData->ClassificationTileListBufferUAV[EStrataTileType::EDecalSimple];
+			PassParameters->DecalSingleTileListDataBuffer = StrataViewData->ClassificationTileListBufferUAV[EStrataTileType::EDecalSingle];
+			PassParameters->DecalComplexTileListDataBuffer = StrataViewData->ClassificationTileListBufferUAV[EStrataTileType::EDecalComplex];
+			// Rough refraction tiles
 			PassParameters->OpaqueRoughRefractionTileListDataBuffer = StrataViewData->ClassificationTileListBufferUAV[EStrataTileType::EOpaqueRoughRefraction];
-			PassParameters->SSSWithoutOpaqueRoughRefractionTileListDataBuffer = StrataViewData->ClassificationTileListBufferUAV[EStrataTileType::ESSSWithoutOpaqueRoughRefraction];
+			PassParameters->OpaqueRoughRefractionSSSWithoutTileListDataBuffer = StrataViewData->ClassificationTileListBufferUAV[EStrataTileType::EOpaqueRoughRefractionSSSWithout];
 
 			const uint32 GroupSize = 8;
 			FComputeShaderUtils::AddPass(
@@ -1236,10 +1283,26 @@ void AddStrataDBufferPass(FRDGBuilder& GraphBuilder, const FMinimalSceneTextures
 		{
 			// Only simple & single material are support but also dispatch complex tiles, 
 			// as they can contain simple/single material pixels
-			check(TileType == EStrataTileType::ESimple || TileType == EStrataTileType::ESingle || TileType == EStrataTileType::EComplex);
+
+			uint32 TilePermutation = 0;
+			switch(TileType)
+			{
+			case EStrataTileType::EComplex:
+			case EStrataTileType::EDecalComplex:
+				TilePermutation = 2;
+				break;
+			case EStrataTileType::ESingle:
+			case EStrataTileType::EDecalSingle:
+				TilePermutation = 1;
+				break;
+			case EStrataTileType::ESimple:
+			case EStrataTileType::EDecalSimple:
+				TilePermutation = 0;
+				break;
+			}
 
 			FStrataDBufferPassCS::FPermutationDomain PermutationVector;
-			PermutationVector.Set<FStrataDBufferPassCS::FTileType>(uint32(TileType));
+			PermutationVector.Set<FStrataDBufferPassCS::FTileType>(TilePermutation);
 
 			TShaderMapRef<FStrataDBufferPassCS> ComputeShader(View.ShaderMap, PermutationVector);
 			FStrataDBufferPassCS::FParameters* PassParameters = GraphBuilder.AllocParameters<FStrataDBufferPassCS::FParameters>();
@@ -1268,9 +1331,10 @@ void AddStrataDBufferPass(FRDGBuilder& GraphBuilder, const FMinimalSceneTextures
 				TileTypeDispatchIndirectArgOffset(TileType));
 		};
 
-		DBufferPass(EStrataTileType::ESimple);
-		DBufferPass(EStrataTileType::ESingle);
-		DBufferPass(EStrataTileType::EComplex);
+		const bool bDbufferTiles = CVarStrataDBufferPassDedicatedTiles.GetValueOnRenderThread() > 0;
+		DBufferPass(bDbufferTiles ? EStrataTileType::EDecalSimple : EStrataTileType::ESimple);
+		DBufferPass(bDbufferTiles ? EStrataTileType::EDecalSingle : EStrataTileType::ESingle);
+		DBufferPass(bDbufferTiles ? EStrataTileType::EDecalComplex : EStrataTileType::EComplex);
 	}
 }
 
diff --git a/Engine/Source/Runtime/Renderer/Private/Strata/Strata.h b/Engine/Source/Runtime/Renderer/Private/Strata/Strata.h
--- a/Engine/Source/Runtime/Renderer/Private/Strata/Strata.h
+++ b/Engine/Source/Runtime/Renderer/Private/Strata/Strata.h
@@ -66,11 +66,14 @@ END_GLOBAL_SHADER_PARAMETER_STRUCT()
 // This must map to the STRATA_TILE_TYPE defines.
 enum EStrataTileType : uint32
 {
-	ESimple  = STRATA_TILE_TYPE_SIMPLE,
-	ESingle = STRATA_TILE_TYPE_SINGLE,
-	EComplex = STRATA_TILE_TYPE_COMPLEX,
-	EOpaqueRoughRefraction = STRATA_TILE_TYPE_ROUGH_REFRACT,
-	ESSSWithoutOpaqueRoughRefraction = STRATA_TILE_TYPE_SSS_WITHOUT_ROUGH_REFRACT,
+	ESimple								= STRATA_TILE_TYPE_SIMPLE,
+	ESingle								= STRATA_TILE_TYPE_SINGLE,
+	EComplex							= STRATA_TILE_TYPE_COMPLEX,
+	EOpaqueRoughRefraction				= STRATA_TILE_TYPE_ROUGH_REFRACT,
+	EOpaqueRoughRefractionSSSWithout	= STRATA_TILE_TYPE_ROUGH_REFRACT_SSS_WITHOUT,
+	EDecalSimple						= STRATA_TILE_TYPE_DECAL_SIMPLE,
+	EDecalSingle						= STRATA_TILE_TYPE_DECAL_SINGLE,
+	EDecalComplex						= STRATA_TILE_TYPE_DECAL_COMPLEX,
 	ECount
 };
 
@@ -159,7 +162,7 @@ void AppendStrataMRTs(const FSceneRenderer& SceneRenderer, uint32& BasePassTextu
 void SetBasePassRenderTargetOutputFormat(const EShaderPlatform Platform, const FMaterialShaderParameters& MaterialParameters, FShaderCompilerEnvironment& OutEnvironment);
 
 
-void AddStrataMaterialClassificationPass(FRDGBuilder& GraphBuilder, const FMinimalSceneTextures& SceneTextures, const TArray<FViewInfo>& Views);
+void AddStrataMaterialClassificationPass(FRDGBuilder& GraphBuilder, const FMinimalSceneTextures& SceneTextures, const FDBufferTextures& DBufferTextures, const TArray<FViewInfo>& Views);
 void AddStrataDBufferPass(FRDGBuilder& GraphBuilder, const FMinimalSceneTextures& SceneTextures, const FDBufferTextures& DBufferTextures, const TArray<FViewInfo>& Views);
 
 void AddStrataStencilPass(FRDGBuilder& GraphBuilder, const TArray<FViewInfo>& Views, const FMinimalSceneTextures& SceneTextures);
diff --git a/Engine/Source/Runtime/Renderer/Private/Strata/StrataRoughRefraction.cpp b/Engine/Source/Runtime/Renderer/Private/Strata/StrataRoughRefraction.cpp
--- a/Engine/Source/Runtime/Renderer/Private/Strata/StrataRoughRefraction.cpp
+++ b/Engine/Source/Runtime/Renderer/Private/Strata/StrataRoughRefraction.cpp
@@ -210,7 +210,7 @@ void AddStrataOpaqueRoughRefractionPasses(
 	//
 	// 3. Add remaining tiles with subsurface scattering that did not have rough refractions on them, resulting in a complete scene color texture.
 	//
-	StrataTileType = EStrataTileType::ESSSWithoutOpaqueRoughRefraction;
+	StrataTileType = EStrataTileType::EOpaqueRoughRefractionSSSWithout;
 	for (uint32 ViewIndex = 0; ViewIndex < ViewCount; ++ViewIndex)
 	{
 		const FViewInfo& View = Views[ViewIndex];
diff --git a/Engine/Source/Runtime/Renderer/Private/Strata/StrataVisualize.cpp b/Engine/Source/Runtime/Renderer/Private/Strata/StrataVisualize.cpp
--- a/Engine/Source/Runtime/Renderer/Private/Strata/StrataVisualize.cpp
+++ b/Engine/Source/Runtime/Renderer/Private/Strata/StrataVisualize.cpp
@@ -319,10 +319,16 @@ FScreenPassTexture AddStrataDebugPasses(FRDGBuilder& GraphBuilder, const FViewIn
 	{
 		RDG_EVENT_SCOPE(GraphBuilder, "Strata::VisualizeClassification");
 		const bool bDebugPass = true;
-		if (StrataClassificationDebug > 1)
+		if (StrataClassificationDebug == 3)
+		{
+			AddStrataInternalClassificationTilePass(GraphBuilder, View, nullptr, &ScreenPassSceneColor.Texture, EStrataTileType::EDecalSimple, bDebugPass);
+			AddStrataInternalClassificationTilePass(GraphBuilder, View, nullptr, &ScreenPassSceneColor.Texture, EStrataTileType::EDecalSingle, bDebugPass);
+			AddStrataInternalClassificationTilePass(GraphBuilder, View, nullptr, &ScreenPassSceneColor.Texture, EStrataTileType::EDecalComplex, bDebugPass);
+		}
+		else if (StrataClassificationDebug == 2)
 		{
 			AddStrataInternalClassificationTilePass(GraphBuilder, View, nullptr, &ScreenPassSceneColor.Texture, EStrataTileType::EOpaqueRoughRefraction, bDebugPass);
-			AddStrataInternalClassificationTilePass(GraphBuilder, View, nullptr, &ScreenPassSceneColor.Texture, EStrataTileType::ESSSWithoutOpaqueRoughRefraction, bDebugPass);
+			AddStrataInternalClassificationTilePass(GraphBuilder, View, nullptr, &ScreenPassSceneColor.Texture, EStrataTileType::EOpaqueRoughRefractionSSSWithout, bDebugPass);
 		}
 		else
 		{
