Add UAV overlap for Strata dbuffer pass to avoid bubble.

diff --git a/Engine/Source/Runtime/Renderer/Private/Strata/Strata.cpp b/Engine/Source/Runtime/Renderer/Private/Strata/Strata.cpp
--- a/Engine/Source/Runtime/Renderer/Private/Strata/Strata.cpp
+++ b/Engine/Source/Runtime/Renderer/Private/Strata/Strata.cpp
@@ -1279,6 +1279,9 @@ void AddStrataDBufferPass(FRDGBuilder& GraphBuilder, const FMinimalSceneTextures
 		const FStrataViewData* StrataViewData = &View.StrataViewData;
 		const FStrataSceneData* StrataSceneData = View.StrataViewData.SceneData;
 
+		FRDGTextureUAVRef RWMaterialTexture = GraphBuilder.CreateUAV(StrataSceneData->MaterialTextureArray, ERDGUnorderedAccessViewFlags::SkipBarrier);
+		FRDGTextureUAVRef RWTopLayerTexture = GraphBuilder.CreateUAV(StrataSceneData->TopLayerTexture, ERDGUnorderedAccessViewFlags::SkipBarrier);
+
 		auto DBufferPass = [&](EStrataTileType TileType)
 		{
 			// Only simple & single material are support but also dispatch complex tiles, 
@@ -1311,8 +1314,8 @@ void AddStrataDBufferPass(FRDGBuilder& GraphBuilder, const FMinimalSceneTextures
 			PassParameters->ViewUniformBuffer = View.ViewUniformBuffer;
 			PassParameters->ViewResolution = View.ViewRect.Size();
 			PassParameters->MaxBytesPerPixel = StrataSceneData->MaxBytesPerPixel;
-			PassParameters->TopLayerTexture = GraphBuilder.CreateUAV(StrataSceneData->TopLayerTexture);
-			PassParameters->MaterialTextureArrayUAV = StrataSceneData->MaterialTextureArrayUAV;
+			PassParameters->TopLayerTexture = RWTopLayerTexture;
+			PassParameters->MaterialTextureArrayUAV = RWMaterialTexture;
 			PassParameters->FirstSliceStoringStrataSSSData = StrataSceneData->FirstSliceStoringStrataSSSData;
 			PassParameters->SceneStencilTexture = SceneTextures.Stencil;
 
@@ -1332,9 +1335,9 @@ void AddStrataDBufferPass(FRDGBuilder& GraphBuilder, const FMinimalSceneTextures
 		};
 
 		const bool bDbufferTiles = CVarStrataDBufferPassDedicatedTiles.GetValueOnRenderThread() > 0;
-		DBufferPass(bDbufferTiles ? EStrataTileType::EDecalSimple : EStrataTileType::ESimple);
-		DBufferPass(bDbufferTiles ? EStrataTileType::EDecalSingle : EStrataTileType::ESingle);
 		DBufferPass(bDbufferTiles ? EStrataTileType::EDecalComplex : EStrataTileType::EComplex);
+		DBufferPass(bDbufferTiles ? EStrataTileType::EDecalSingle : EStrataTileType::ESingle);
+		DBufferPass(bDbufferTiles ? EStrataTileType::EDecalSimple : EStrataTileType::ESimple);
 	}
 }
 
