Run UpdateDeferredCachedUniformExpressions() before BuildRayTracingMeshCommands()

Uniform expression cache explicit update may not run for objects that are off-screen, but it must be done before objects are added to the ray tracing scene. Normally this work is done through FMaterialRenderProxy::UpdateUniformExpressionCacheIfNeeded, however ray tracing material processing (FMaterialShader::GetShaderBindings, which accesses UniformExpressionCache) is done on task threads, therefore all work must be done here up-front as UpdateUniformExpressionCacheIfNeeded is not free-threaded.

diff --git a/Engine/Source/Runtime/Renderer/Private/DeferredShadingRenderer.cpp b/Engine/Source/Runtime/Renderer/Private/DeferredShadingRenderer.cpp
--- a/Engine/Source/Runtime/Renderer/Private/DeferredShadingRenderer.cpp
+++ b/Engine/Source/Runtime/Renderer/Private/DeferredShadingRenderer.cpp
@@ -829,6 +829,12 @@ bool FDeferredShadingSceneRenderer::GatherRayTracingWorldInstancesForView(FRDGBu
 	TRACE_CPUPROFILER_EVENT_SCOPE(FDeferredShadingSceneRenderer::GatherRayTracingWorldInstances);
 	SCOPE_CYCLE_COUNTER(STAT_GatherRayTracingWorldInstances);
 
+	// Ensure that any invalidated cached uniform expressions have been updated on the rendering thread.
+	// Normally this work is done through FMaterialRenderProxy::UpdateUniformExpressionCacheIfNeeded,
+	// however ray tracing material processing (FMaterialShader::GetShaderBindings, which accesses UniformExpressionCache)
+	// is done on task threads, therefore all work must be done here up-front as UpdateUniformExpressionCacheIfNeeded is not free-threaded.
+	FMaterialRenderProxy::UpdateDeferredCachedUniformExpressions();
+
 	RayTracingCollector.ClearViewMeshArrays();
 
 	FGPUScenePrimitiveCollector DummyDynamicPrimitiveCollector;
