Prototype moire detection heuristic.

diff --git a/Engine/Config/BaseScalability.ini b/Engine/Config/BaseScalability.ini
--- a/Engine/Config/BaseScalability.ini
+++ b/Engine/Config/BaseScalability.ini
@@ -43,6 +43,7 @@ r.TSR.History.R11G11B10=1
 r.TSR.History.ScreenPercentage=100
 r.TSR.History.UpdateQuality=0
 r.TSR.ShadingRejection.SpatialFilter=0
+r.TSR.ShadingRejection.Flickering=0
 r.TSR.Velocity.Extrapolation=0
 r.TSR.Velocity.HoleFill=0
 r.TSR.RejectionAntiAliasingQuality=0
@@ -54,6 +55,7 @@ r.TSR.History.R11G11B10=1
 r.TSR.History.ScreenPercentage=100
 r.TSR.History.UpdateQuality=1
 r.TSR.ShadingRejection.SpatialFilter=1
+r.TSR.ShadingRejection.Flickering=0
 r.TSR.Velocity.Extrapolation=0
 r.TSR.Velocity.HoleFill=1
 r.TSR.RejectionAntiAliasingQuality=1
@@ -65,6 +67,7 @@ r.TSR.History.R11G11B10=1
 r.TSR.History.ScreenPercentage=100
 r.TSR.History.UpdateQuality=2
 r.TSR.ShadingRejection.SpatialFilter=1
+r.TSR.ShadingRejection.Flickering=0
 r.TSR.Velocity.Extrapolation=0
 r.TSR.Velocity.HoleFill=1
 r.TSR.RejectionAntiAliasingQuality=1
@@ -76,6 +79,7 @@ r.TSR.History.R11G11B10=1
 r.TSR.History.ScreenPercentage=100
 r.TSR.History.UpdateQuality=3
 r.TSR.ShadingRejection.SpatialFilter=2
+r.TSR.ShadingRejection.Flickering=1
 r.TSR.Velocity.Extrapolation=1
 r.TSR.Velocity.HoleFill=1
 r.TSR.RejectionAntiAliasingQuality=2
@@ -87,6 +91,7 @@ r.TSR.History.R11G11B10=0
 r.TSR.History.ScreenPercentage=200
 r.TSR.History.UpdateQuality=3
 r.TSR.ShadingRejection.SpatialFilter=2
+r.TSR.ShadingRejection.Flickering=1
 r.TSR.Velocity.Extrapolation=1
 r.TSR.Velocity.HoleFill=1
 r.TSR.RejectionAntiAliasingQuality=2
diff --git a/Engine/Shaders/Private/TemporalSuperResolution/TSRCommon.ush b/Engine/Shaders/Private/TemporalSuperResolution/TSRCommon.ush
--- a/Engine/Shaders/Private/TemporalSuperResolution/TSRCommon.ush
+++ b/Engine/Shaders/Private/TemporalSuperResolution/TSRCommon.ush
@@ -820,14 +820,14 @@ bool TakeOnlyOneSamplePair(float2 Offset)
 	return Offset.y > 0.0 || (Offset.x > 0.0 && Offset.y == 0.0);
 }
 
-float2 ComputeStaticVelocity(float2 ScreenPos, float DeviceZ)
+float3 ComputeStaticVelocity(float2 ScreenPos, float DeviceZ)
 {
 	float3 PosN = float3(ScreenPos, DeviceZ);
 
 	float4 ThisClip = float4(PosN, 1);
 	float4 PrevClip = mul( ThisClip, View.ClipToPrevClip );
-	float2 PrevScreen = PrevClip.xy / PrevClip.w;
-	return PosN.xy - PrevScreen;
+	float3 PrevScreen = PrevClip.xyz / PrevClip.w;
+	return PosN - PrevScreen;
 }
 
 tsr_half ComputePredictionCompleteness(tsr_half SampleHistoryValidity)
diff --git a/Engine/Shaders/Private/TemporalSuperResolution/TSRDecimateHistory.usf b/Engine/Shaders/Private/TemporalSuperResolution/TSRDecimateHistory.usf
--- a/Engine/Shaders/Private/TemporalSuperResolution/TSRDecimateHistory.usf
+++ b/Engine/Shaders/Private/TemporalSuperResolution/TSRDecimateHistory.usf
@@ -26,6 +26,9 @@
 	#endif
 #endif // CONFIG_VELOCITY_SAMPLES != 1
 
+/** Whether to compile the moire reprojection. */
+#define CONFIG_MOIRE_REPROJECTION (DIM_MOIRE_REPROJECTION)
+
 
 //------------------------------------------------------- PARAMETERS
 
@@ -40,6 +43,7 @@ Texture2D<tsr_subpixel_payload> ParallaxFactorTexture;
 Texture2D<tsr_subpixel_details> PrevHistorySubpixelDetails;
 
 Texture2D<tsr_halfC> PrevHistoryGuide;
+Texture2D<tsr_half4> PrevHistoryMoire;
 float2 PrevGuideInfo_Extent;
 float2 PrevGuideInfo_ExtentInverse;
 uint2  PrevGuideInfo_ViewportMin;
@@ -54,6 +58,7 @@ float3 HistoryGuideQuantizationError;
 
 
 RWTexture2D<tsr_halfC> ReprojectedHistoryGuideOutput;
+RWTexture2D<tsr_half4> ReprojectedHistoryMoireOutput;
 
 RWTexture2D<float2> HoleFilledVelocityOutput;
 RWTexture2D<tsr_half2> HoleFilledVelocityMaskOutput;
@@ -262,6 +267,9 @@ void MainCS(
 	uint HistoryClosestDeviceZSamples[4];
 	
 	tsr_halfC ReprojectedHistoryGuideColor;
+	#if CONFIG_MOIRE_REPROJECTION
+		tsr_half4 ReprojectedHistoryMoireColor;
+	#endif
 	FSubpixelNeighborhood SubpixelNeigborhood;
 	bool bIsOffScreen;
 
@@ -327,6 +335,15 @@ void MainCS(
 				ReprojectedHistoryGuideColor += PrevHistoryGuide.SampleLevel(GlobalBilinearClampedSampler, SampleUV, 0) * tsr_half(Samples.Weight[i]);
 			}
 			ReprojectedHistoryGuideColor *= tsr_half(Samples.FinalMultiplier);
+
+			#if CONFIG_MOIRE_REPROJECTION
+			{
+				uint2 Random = Rand3DPCG16(int3(InputPixelPos - InputInfo_ViewportMin, View.StateFrameIndexMod8)).xy;
+				float2 E = Hammersley16(0, 1, Random).xy;
+
+				ReprojectedHistoryMoireColor = PrevHistoryMoire.SampleLevel(GlobalPointClampedSampler, KernelHistoryBufferUV + (E - 0.5) * PrevGuideInfo_ExtentInverse, 0);
+			}
+			#endif
 		}
 
 		{
@@ -562,6 +579,9 @@ void MainCS(
 		tsr_short2 OutputPixelPos = InvalidateOutputPixelPos(InputPixelPos, InputInfo_ViewportMax);
 
 		ReprojectedHistoryGuideOutput[OutputPixelPos] = ReprojectedHistoryGuideColor;
+		#if CONFIG_MOIRE_REPROJECTION
+			ReprojectedHistoryMoireOutput[OutputPixelPos] = ReprojectedHistoryMoireColor;
+		#endif
 		HoleFilledVelocityOutput[OutputPixelPos] = EncodedScreenVelocity;
 		HoleFilledVelocityMaskOutput[InputPixelPos] = tsr_half2(0.0, RejectHighFrequency);
 		ParallaxRejectionMaskOutput[OutputPixelPos] = ParallaxRejectionMask;
diff --git a/Engine/Shaders/Private/TemporalSuperResolution/TSRDilateVelocity.usf b/Engine/Shaders/Private/TemporalSuperResolution/TSRDilateVelocity.usf
--- a/Engine/Shaders/Private/TemporalSuperResolution/TSRDilateVelocity.usf
+++ b/Engine/Shaders/Private/TemporalSuperResolution/TSRDilateVelocity.usf
@@ -31,6 +31,7 @@ float2 PrevOutputBufferUVMin;
 float2 PrevOutputBufferUVMax;
 float WorldDepthToDepthError;
 float VelocityExtrapolationMultiplier;
+uint bOutputIsMovingTexture;
 
 Texture2D<float> SceneDepthTexture;
 Texture2D SceneVelocityTexture;
@@ -39,7 +40,7 @@ RWTexture2D<float2> DilatedVelocityOutput;
 RWTexture2D<float> ClosestDepthOutput;
 RWTexture2D<uint> PrevUseCountOutput;
 RWTexture2D<uint> PrevClosestDepthOutput;
-RWTexture2D<uint> ParallaxFactorOutput;
+RWTexture2DArray<uint> R8Output;
 
 RWTexture2D<float3>	VelocityFlattenOutput;
 RWTexture2D<float4>	VelocityTileOutput0;
@@ -151,7 +152,7 @@ void MainCS(
 	#endif
 	
 	// Camera motion for pixel.
-	float2 FinalScreenVelocity = 0.0;
+	float3 FinalScreenVelocity = 0.0;
 
 	// Compute final screen vecocity.
 	#if CONFIG_VELOCITY_EXTRAPOLATION
@@ -169,24 +170,24 @@ void MainCS(
 		float4 EncodedVelocity1 = SceneVelocityTexture[PixelCoord1];
 		float4 EncodedVelocity2 = SceneVelocityTexture[PixelCoord2];
 
-		float2 ScreenVelocity1 = ComputeStaticVelocity(ScreenPos, ClosestDeviceZ1);
-		float2 ScreenVelocity2 = ComputeStaticVelocity(ScreenPos, ClosestDeviceZ2);
+		float3 ScreenVelocity1 = ComputeStaticVelocity(ScreenPos, ClosestDeviceZ1);
+		float3 ScreenVelocity2 = ComputeStaticVelocity(ScreenPos, ClosestDeviceZ2);
 		
 		#if CONFIG_USE_SCENE_VELOCITY
 		{
 			if (EncodedVelocity1.x > 0.0)
 			{
-				ScreenVelocity1 = DecodeVelocityFromTexture(EncodedVelocity1).xy;
+				ScreenVelocity1 = DecodeVelocityFromTexture(EncodedVelocity1);
 			}
 			if (EncodedVelocity2.x > 0.0)
 			{
-				ScreenVelocity2 = DecodeVelocityFromTexture(EncodedVelocity2).xy;
+				ScreenVelocity2 = DecodeVelocityFromTexture(EncodedVelocity2);
 			}
 		}
 		#endif
 
-		float2 VelocityDelta = ScreenVelocity1 - ScreenVelocity2;
-		float ExtrapolationBilateral = ComputeScreenVelocityBilateralWeight(VelocityDelta);
+		float3 VelocityDelta = ScreenVelocity1 - ScreenVelocity2;
+		float ExtrapolationBilateral = ComputeScreenVelocityBilateralWeight(VelocityDelta.xy);
 
 		FinalScreenVelocity = ScreenVelocity1 + VelocityDelta * ExtrapolationBilateral * VelocityExtrapolationMultiplier;
 	}
@@ -201,17 +202,38 @@ void MainCS(
 		{
 			if (EncodedVelocity.x > 0.0)
 			{
-				FinalScreenVelocity = DecodeVelocityFromTexture(EncodedVelocity).xy;
+				FinalScreenVelocity = DecodeVelocityFromTexture(EncodedVelocity);
 			}
 		}
 		#endif
 	}
 
+	tsr_half IsMovingMask = tsr_half(0.0);
+	BRANCH
+	if (bOutputIsMovingTexture)
+	{
+		float Depth = ConvertFromDeviceZ(ClosestDeviceZ);
+		float PrevDepth = ConvertFromDeviceZ(ClosestDeviceZ - FinalScreenVelocity.z);
+		float2 PrevScreenPos = ScreenPos - FinalScreenVelocity.xy;
+
+		float4 ClipPosition = float4(ScreenPos * (View.ViewToClip[3][3] < 1.0f ? Depth : 1.0f), Depth, 1);
+		float4 PrevClipPosition = float4(PrevScreenPos * (View.ViewToClip[3][3] < 1.0f ? PrevDepth : 1.0f), PrevDepth, 1);
+		
+		float3 PreViewTranslationOffset = LWCToFloat(LWCSubtract(PrimaryView.PreViewTranslation, PrimaryView.PrevPreViewTranslation));
+		float3 TranslatedWorldPosition = mul(ClipPosition, View.ScreenToTranslatedWorld).xyz;
+		float3 PrevTranslatedWorldPosition = mul(PrevClipPosition, View.PrevScreenToTranslatedWorld).xyz + PreViewTranslationOffset;
+		
+		float WorldDepthToPixelWorldRadius = 2.0 * GetTanHalfFieldOfView().x * View.ViewSizeAndInvSize.z;
+		float StaticWorldRadius = Depth * WorldDepthToPixelWorldRadius;
+
+		IsMovingMask = tsr_half(saturate(length(TranslatedWorldPosition - PrevTranslatedWorldPosition) * rcp(StaticWorldRadius) - 1.0));
+	}
+
 	// Do motion blur velocity flatten.
 	#if DIM_MOTION_BLUR_DIRECTIONS > 0
 		FVelocityRange VelocityPolarRange;
 		float2 VelocityPolar;
-		ReduceVelocityFlattenTile(GroupThreadIndex, FinalScreenVelocity, /* out */ VelocityPolarRange, /* out */ VelocityPolar);
+		ReduceVelocityFlattenTile(GroupThreadIndex, FinalScreenVelocity.xy, /* out */ VelocityPolarRange, /* out */ VelocityPolar);
 	#endif
 
 	Debug.rg = VelocityPixelOffset * 0.5 + 0.5;
@@ -222,10 +244,10 @@ void MainCS(
 		// Scatter parallax rejection
 		#if 1
 		{
-			float2 PrevScreenPos = ScreenPos - FinalScreenVelocity;
+			float2 PrevScreenPos = ScreenPos - FinalScreenVelocity.xy;
 			float2 PrevInputBufferUV = (InputInfo_ScreenPosToViewportScale * PrevScreenPos + InputInfo_ScreenPosToViewportBias) * InputInfo_ExtentInverse;
 
-			float2 PixelVelocity = FinalScreenVelocity * 0.5 * InputInfo_ViewportSize;
+			float2 PixelVelocity = FinalScreenVelocity.xy * 0.5 * InputInfo_ViewportSize;
 			uint EncodedHoleFillingVelocity = EncodeHoleFillingVelocity(PixelVelocity);
 
 			uint PrevClosestDepth = (f32tof16(ClosestDeviceZ) << VELOCITY_HOLE_FILLING_BITS) | EncodedHoleFillingVelocity;
@@ -264,10 +286,16 @@ void MainCS(
 		
 		uint2 OutputPixelPos = bValidOutputPixel ? InputPixelPos : uint(~0).xx;
 
-		//DilatedVelocityOutput[OutputPixelPos] = EncodedVelocity.xy; //EncodeVelocityToTexture(float3(FinalScreenVelocity, 0.0)).xy;
-		DilatedVelocityOutput[OutputPixelPos] = EncodeVelocityToTexture(float3(FinalScreenVelocity, 0.0)).xy;
+		//DilatedVelocityOutput[OutputPixelPos] = EncodedVelocity.xy; //EncodeVelocityToTexture(float3(FinalScreenVelocity.xy, 0.0)).xy;
+		DilatedVelocityOutput[OutputPixelPos] = EncodeVelocityToTexture(float3(FinalScreenVelocity.xy, 0.0)).xy;
 		ClosestDepthOutput[OutputPixelPos] = ClosestDeviceZ;
-		ParallaxFactorOutput[OutputPixelPos] = CompressParallaxFactor(tsr_half(ParallaxFactorForHistory)) | (CompressParallaxFactor(tsr_half(ParallaxFactorForHistoryRejection)) << 4);
+		R8Output[uint3(OutputPixelPos, 0)] = CompressParallaxFactor(tsr_half(ParallaxFactorForHistory)) | (CompressParallaxFactor(tsr_half(ParallaxFactorForHistoryRejection)) << 4);
+
+		BRANCH
+		if (bOutputIsMovingTexture)
+		{
+			R8Output[uint3(OutputPixelPos, 1)] = uint(round(IsMovingMask * 255.0));
+		}
 
 		// Output motion blur velocity flatten
 		#if DIM_MOTION_BLUR_DIRECTIONS > 0
diff --git a/Engine/Shaders/Private/TemporalSuperResolution/TSRRejectShading.usf b/Engine/Shaders/Private/TemporalSuperResolution/TSRRejectShading.usf
--- a/Engine/Shaders/Private/TemporalSuperResolution/TSRRejectShading.usf
+++ b/Engine/Shaders/Private/TemporalSuperResolution/TSRRejectShading.usf
@@ -29,7 +29,10 @@
 #define CONFIG_FIREFLY_REMOVAL 1
 
 /** Whether to compose translucency to the input. */
-#define CONFIG_COMPOSE_TRANSLUCENCY (!DIM_SEPARATE_TRANSLUCENCY)
+#define CONFIG_COMPOSE_TRANSLUCENCY 1
+
+/** Whether to detect flickering. */
+#define CONFIG_MOIRE_DETECTION (DIM_FLICKERING_DETECTION)
 
 
 //------------------------------------------------------- DERIVED
@@ -81,13 +84,17 @@ FScreenTransform InputPixelPosToTranslucencyTextureUV;
 float2 TranslucencyTextureUVMin;
 float2 TranslucencyTextureUVMax;
 float3 HistoryGuideQuantizationError;
+float FlickeringFramePeriod;
 
 Texture2D<tsr_halfC> InputTexture;
 Texture2D<tsr_half4> InputSceneTranslucencyTexture;
 Texture2D<tsr_halfC> ReprojectedHistoryGuideTexture;
+Texture2D<tsr_half4> ReprojectedHistoryMoireTexture;
 Texture2D<tsr_half> ParallaxRejectionMaskTexture;
+Texture2D<tsr_ushort> IsMovingMaskTexture;
 
 RWTexture2D<tsr_halfC> HistoryGuideOutput;
+RWTexture2D<tsr_half4> HistoryMoireOutput;
 RWTexture2D<tsr_half> HistoryRejectionOutput;
 RWTexture2D<tsr_half> InputSceneColorLdrLumaOutput;
 
@@ -346,12 +353,16 @@ void MainCS(
 		tsr_short(-TILE_OVERSCAN_PADDING).xx);
 		
 	// Load
-	tsr_simd_halfC OriginalInput;
+	tsr_simd_halfC OriginalOpaqueInput;
 	#if CONFIG_COMPOSE_TRANSLUCENCY
 		tsr_simd_half4 OriginalTranslucencyInput;
 	#endif
 	tsr_simd_halfC OriginalHistory;
 	tsr_simd_half bIsValid;
+	#if CONFIG_MOIRE_DETECTION
+		tsr_simd_half4 PrevMoire;
+		tsr_simd_half bIsStatic;
+	#endif
 	ISOLATE
 	{
 		UNROLL_N(SIMD_SIZE)
@@ -363,12 +374,16 @@ void MainCS(
 			ClampedFetchPixelPos = ClampPixelOffset(ClampedFetchPixelPos, InputPixelPosMin, InputPixelPosMax);
 			TranslucencyTextureUV = clamp(TranslucencyTextureUV, TranslucencyTextureUVMin, TranslucencyTextureUVMax);
 
-			OriginalInput.SetElement(SimdIndex, InputTexture[ClampedFetchPixelPos]);
+			OriginalOpaqueInput.SetElement(SimdIndex, InputTexture[ClampedFetchPixelPos]);
 			#if CONFIG_COMPOSE_TRANSLUCENCY
 				OriginalTranslucencyInput.SetElement(SimdIndex, InputSceneTranslucencyTexture.SampleLevel(GlobalBilinearClampedSampler, TranslucencyTextureUV, 0));
 			#endif
 			OriginalHistory.SetElement(SimdIndex, ReprojectedHistoryGuideTexture[ClampedFetchPixelPos]);
 			bIsValid.SetElement(SimdIndex, ParallaxRejectionMaskTexture[ClampedFetchPixelPos] > tsr_half(PARALLAX_REJECTION_MASK_THRESHOLD) ? tsr_half(1.0) : tsr_half(0.0));
+			#if CONFIG_MOIRE_DETECTION
+				PrevMoire.SetElement(SimdIndex, ReprojectedHistoryMoireTexture[ClampedFetchPixelPos]);
+				bIsStatic.SetElement(SimdIndex, tsr_half(1.0) - tsr_half(IsMovingMaskTexture[ClampedFetchPixelPos]) * rcp(tsr_half(255.0)));
+			#endif
 		}
 	}
 	
@@ -379,7 +394,7 @@ void MainCS(
 		UNROLL_N(SIMD_SIZE)
 		for (uint SimdIndex = 0; SimdIndex < SIMD_SIZE; SimdIndex++)
 		{
-			tsr_half3 Color = OriginalInput.GetElement(SimdIndex).rgb;
+			tsr_half3 Color = OriginalOpaqueInput.GetElement(SimdIndex).rgb;
 			
 			#if CONFIG_COMPOSE_TRANSLUCENCY
 				tsr_half4 Translucency = OriginalTranslucencyInput.GetElement(SimdIndex);
@@ -390,6 +405,8 @@ void MainCS(
 			CenterLumaLDR.SetElement(SimdIndex, CenterLuma / (SpatialAAExposure + CenterLuma));
 		}
 	}
+	
+	tsr_simd_halfC OriginalInput = OriginalOpaqueInput;
 
 	#if CONFIG_COMPOSE_TRANSLUCENCY
 	{
@@ -398,7 +415,7 @@ void MainCS(
 		UNROLL_N(SIMD_SIZE)
 		for (uint SimdIndex = 0; SimdIndex < SIMD_SIZE; SimdIndex++)
 		{
-			tsr_halfC Color = OriginalInput.GetElement(SimdIndex);
+			tsr_halfC Color = OriginalOpaqueInput.GetElement(SimdIndex);
 			tsr_half4 Translucency = FilteredTranslucencyInput.GetElement(SimdIndex);
 
 			Color.rgb = Color.rgb * Translucency.a + Translucency.rgb;
@@ -412,8 +429,155 @@ void MainCS(
 	#endif
 	
 	// Convert input to LDR
+	OriginalOpaqueInput = OriginalOpaqueInput * rcp(OriginalOpaqueInput + tsr_simd_halfC::Const(1.0));
 	OriginalInput = OriginalInput * rcp(OriginalInput + tsr_simd_halfC::Const(1.0));
 	
+	// Analyses the moire pattern overtime.
+	#if CONFIG_MOIRE_DETECTION
+	tsr_simd_half4 FinalMoire = tsr_simd_half4::Const(0.0);
+	tsr_simd_half MoireError = tsr_simd_half::Const(0.0);
+	{
+		const tsr_half BlendFinal = tsr_half(0.05);
+		//const tsr_half BlendFinal = tsr_half(0.2);
+		const tsr_half MoireEncodingError = tsr_half(rcp(127.0));
+		const tsr_half MaxPrevTotalVariationCount = tsr_half(20.0);
+
+		// Compute the threshold at which the variation count should quick in
+		const tsr_half TotalVariationCountThreshold = tsr_half(1.0 / (1.0 - pow(1.0 - BlendFinal, tsr_half(FlickeringFramePeriod))));
+
+
+		tsr_simd_half InputLuma;
+		{
+			UNROLL_N(SIMD_SIZE)
+			for (uint SimdIndex = 0; SimdIndex < SIMD_SIZE; SimdIndex++)
+			{
+				const tsr_half3 LumaWeights = tsr_half3(0.333f, 0.333f, 0.333f);
+				
+				InputLuma.SetElement(SimdIndex, dot(OriginalOpaqueInput.GetElement(SimdIndex).rgb, LumaWeights));
+			}
+		}
+		
+		// Convolve the input luma to have the exact same moire pattern as FilteredInput below.
+		InputLuma = Convolve3x3HV<TAccumulatorAvg<tsr_simd_half>, tsr_simd_half>(InputLuma, bIsValid, GroupThreadIndex);
+
+		// Unpack history
+		tsr_simd_half PrevFlickeringHistory = PrevMoire[0];
+		tsr_simd_half PrevGradient = PrevMoire[1] * tsr_half(2.0) - tsr_half(1.0);
+		//tsr_simd_half PrevTotalVariation = PrevMoire[2] * tsr_simd_half::Const(1.0 - BlendFinal);
+		tsr_simd_half PrevTotalVariation = PrevMoire[2] * tsr_simd_half::Const(tsr_half(1.0) - BlendFinal);
+		tsr_simd_half PrevTotalVariationCount = PrevMoire[3] * MaxPrevTotalVariationCount * tsr_simd_half::Const(tsr_half(1.0) - BlendFinal);
+		
+		// Discard history on moving objects.
+		PrevTotalVariation = PrevTotalVariation * bIsStatic;
+		PrevTotalVariationCount = PrevTotalVariationCount * bIsStatic;
+
+
+		// Compute clamping box of the limance.
+		tsr_simd_half InputLumaMin = Convolve3x3HV<TAccumulatorMin<tsr_simd_half>, tsr_simd_half>(InputLuma, bIsValid, GroupThreadIndex);
+		tsr_simd_half InputLumaMax = Convolve3x3HV<TAccumulatorMax<tsr_simd_half>, tsr_simd_half>(InputLuma, bIsValid, GroupThreadIndex);
+
+		// Clamp and contribue the history.		
+		tsr_simd_half ClampedPrevFlickeringHistory = clamp(PrevFlickeringHistory, InputLumaMin, InputLumaMax);
+		tsr_simd_half FinalFlickeringHistory = InputLuma * tsr_simd_half::Const(BlendFinal) - ClampedPrevFlickeringHistory * tsr_simd_half::Const(BlendFinal - tsr_half(1.0));
+	
+		// Measure the flickering gradient of the current frame.
+		//tsr_simd_half CurrentGradient = FinalFlickeringHistory - PrevFlickeringHistory;
+		tsr_simd_half CurrentGradient = ClampedPrevFlickeringHistory - PrevFlickeringHistory;
+		
+		#if 0
+			tsr_simd_half InputBoxSize = InputLumaMax - InputLumaMin;
+			CurrentGradient = sign(CurrentGradient) * max(abs(CurrentGradient) - tsr_simd_half::Const(0.5) * InputBoxSize, tsr_simd_half::Const(0.0));
+		#endif
+
+		tsr_simd_half IsFlicker;
+		UNROLL_N(SIMD_SIZE)
+		for (uint SimdIndex = 0; SimdIndex < SIMD_SIZE; SimdIndex++)
+		{
+			tsr_half LocalCurrentGradient = CurrentGradient.GetElement(SimdIndex);
+			tsr_half LocalPrevGradient = PrevGradient.GetElement(SimdIndex);
+
+			bool bGradientsAreSameSign = LocalCurrentGradient * LocalPrevGradient > tsr_half(0.0);
+			bool bGradientsAreWithinError = abs(LocalCurrentGradient) < MoireEncodingError || abs(LocalPrevGradient) < MoireEncodingError;
+			//bool bCurrentGradientSubstentialEnough = abs(LocalCurrentGradient) > View.GeneralPurposeTweak * abs(LocalPrevGradient); // - 2.0 * MoireEncodingError;
+
+			tsr_half LocalIsFlicker = select((
+				bGradientsAreSameSign ||
+				bGradientsAreWithinError ||
+				//!bCurrentGradientSubstentialEnough ||
+				bCameraCut), tsr_half(0.0), tsr_half(1.0));
+
+			IsFlicker.SetElement(SimdIndex, LocalIsFlicker);
+		}
+
+		// Compute the total variation of the gradient over time.
+		tsr_simd_half GradientVariation = min(abs(PrevGradient), abs(CurrentGradient)) * IsFlicker;
+
+		// Dilates the total variation and contribution to keep stability despite inacuracies in the history reprojection
+		// of these need a nearest filter which isn't precise at all.
+		tsr_simd_half TotalVariationContribution = Convolve3x3HV<TAccumulatorMax<tsr_simd_half>, tsr_simd_half>(GradientVariation, tsr_simd_half::Const(1.0), GroupThreadIndex);
+		tsr_simd_half TotalVariationCountContribution = Convolve3x3HV<TAccumulatorMax<tsr_simd_half>, tsr_simd_half>(IsFlicker, tsr_simd_half::Const(1.0), GroupThreadIndex);
+
+		UNROLL_N(SIMD_SIZE)
+		for (uint SimdIndex = 0; SimdIndex < SIMD_SIZE; SimdIndex++)
+		{
+			tsr_half LocalCurrentGradient = CurrentGradient.GetElement(SimdIndex);
+
+			tsr_half LocalPrevGradient = PrevGradient.GetElement(SimdIndex);
+			tsr_half LocalPrevTotalVariation = PrevTotalVariation.GetElement(SimdIndex);
+			tsr_half LocalPrevTotalVariationCount = PrevTotalVariationCount.GetElement(SimdIndex);
+			tsr_half LocalIsValid = bIsValid.GetElement(SimdIndex);
+
+			tsr_half LocalIsFlicker = IsFlicker.GetElement(SimdIndex);
+
+			tsr_half LocalTotalVariationContribution = TotalVariationContribution.GetElement(SimdIndex);
+			tsr_half LocalTotalVariationCountContribution = TotalVariationCountContribution.GetElement(SimdIndex);
+			
+			// Blend out previous gradient
+			tsr_half LocalPrevBlendedGradient = LocalPrevGradient * (tsr_half(1.0) - BlendFinal) * (tsr_half(1.0) - LocalIsFlicker);
+
+			// Compute new gradient.
+			#if 0
+				tsr_half LocalNewGradient = LocalPrevBlendedGradient + LocalCurrentGradient;
+			#else
+				tsr_half LocalNewGradient = abs(LocalPrevBlendedGradient) > abs(LocalCurrentGradient) ? LocalPrevBlendedGradient : LocalCurrentGradient;
+			#endif
+			
+			// Accumulate current frame variation.
+			tsr_half LocalNewTotalVariation = LocalPrevTotalVariation + LocalTotalVariationContribution;
+			tsr_half LocalNewTotalVariationCount = LocalPrevTotalVariationCount + LocalTotalVariationCountContribution;
+
+			// Discard all moire history on parallax disocclusion.
+			LocalNewGradient *= LocalIsValid;
+			LocalNewTotalVariation *= LocalIsValid;
+			LocalNewTotalVariationCount *= LocalIsValid;
+
+			// Quantise total variation and count to history bit depth to variation drift in the history over time.
+			tsr_half LocalQuantizedNewTotalVariationCount = floor(LocalNewTotalVariationCount * (tsr_half(255.0) / MaxPrevTotalVariationCount)) * (MaxPrevTotalVariationCount / tsr_half(255.0));
+			LocalNewTotalVariation *= LocalQuantizedNewTotalVariationCount * SafeRcp(LocalNewTotalVariationCount);
+
+			// Compute the final luminance 
+			#if 1
+				tsr_half LocalCountFadeIn = saturate(LocalNewTotalVariationCount / TotalVariationCountThreshold - tsr_half(0.5));
+			#else
+				tsr_half LocalCountFadeIn = saturate(LocalNewTotalVariationCount - TotalVariationCountThreshold);
+			#endif
+			tsr_half LocalMoireError = abs(LocalNewTotalVariation * SafeRcp(LocalNewTotalVariationCount)) * LocalCountFadeIn;
+
+			FinalMoire.RegisterArray[SimdIndex][0] = FinalFlickeringHistory.GetElement(SimdIndex);
+			FinalMoire.RegisterArray[SimdIndex][1] = LocalNewGradient * tsr_half(0.5) + tsr_half(0.5);
+			FinalMoire.RegisterArray[SimdIndex][2] = LocalNewTotalVariation;
+			FinalMoire.RegisterArray[SimdIndex][3] = LocalQuantizedNewTotalVariationCount * rcp(MaxPrevTotalVariationCount);
+
+			MoireError.SetElement(SimdIndex, LocalMoireError);
+		}
+		
+		// Moving object in the scene have very low chance of creating moire pattern with TAA jitter, so discard the moire error to still
+		// have character animation ghost free evem as environement shadow project different things on character.
+		MoireError = MoireError * bIsStatic;
+	}
+	#endif
+
+	// Top secret sauce with paper and salt.
 	tsr_simd_halfC ClampedInput = OriginalInput;
 	tsr_simd_halfC ClampedHistory = OriginalHistory;
 	#if CONFIG_FIREFLY_REMOVAL
@@ -432,22 +596,51 @@ void MainCS(
 	Rejection.SetAllElements(tsr_half(1.0));
 	#if 1
 	{
+		const tsr_half FilteringWeight = rcp(tsr_half(1.0 + 4 * 0.5 + 4 * 0.25));
+
 		tsr_simd_halfC BackbufferError = tsr_simd_halfC::Const(MeasureBackbufferLDRQuantizationError());
 
 		tsr_simd_halfC FilteredBoxMin = Convolve3x3HV<TAccumulatorMin<tsr_simd_halfC>, tsr_simd_halfC>(FilteredInput, bIsValid, GroupThreadIndex) - BackbufferError;
 		tsr_simd_halfC FilteredBoxMax = Convolve3x3HV<TAccumulatorMax<tsr_simd_halfC>, tsr_simd_halfC>(FilteredInput, bIsValid, GroupThreadIndex) + BackbufferError;
-		tsr_simd_halfC ClampedFilteredHistory = clamp(FilteredHistory, FilteredBoxMin, FilteredBoxMax);
+		
+		tsr_simd_halfC ClampedFilteredHistory;
+		#if CONFIG_MOIRE_DETECTION
+		{
+			UNROLL_N(SIMD_SIZE)
+			for (uint SimdIndex = 0; SimdIndex < SIMD_SIZE; SimdIndex++)
+			{
+				// TODO: improve this 3.0 *
+				ClampedFilteredHistory.SetElement(SimdIndex, clamp(
+					FilteredHistory.GetElement(SimdIndex),
+					FilteredBoxMin.GetElement(SimdIndex) - tsr_half(3.0) * MoireError.GetElement(SimdIndex),
+					FilteredBoxMax.GetElement(SimdIndex) + tsr_half(3.0) * MoireError.GetElement(SimdIndex)));
+			}
+		}
+		#else
+			ClampedFilteredHistory = clamp(FilteredHistory, FilteredBoxMin, FilteredBoxMax);
+		#endif
 		
 		tsr_simd_halfC OriginalBoxMin = Convolve3x3HV<TAccumulatorMin<tsr_simd_halfC>, tsr_simd_halfC>(OriginalInput, bIsValid, GroupThreadIndex);
 		tsr_simd_halfC OriginalBoxMax = Convolve3x3HV<TAccumulatorMax<tsr_simd_halfC>, tsr_simd_halfC>(OriginalInput, bIsValid, GroupThreadIndex);
-		tsr_simd_halfC BoxSize = (OriginalBoxMax - OriginalBoxMin + (BackbufferError + BackbufferError)) * rcp(tsr_half(1.0 + 4 * 0.5 + 4 * 0.25));
+		tsr_simd_halfC BoxSize = (OriginalBoxMax - OriginalBoxMin + (BackbufferError + BackbufferError)) * FilteringWeight;
 		
 		UNROLL_N(SIMD_SIZE)
 		for (uint SimdIndex = 0; SimdIndex < SIMD_SIZE; SimdIndex++)
 		{
+			#if CONFIG_MOIRE_DETECTION
+				tsr_halfC ClampedEnergy = abs(ClampedFilteredHistory.GetElement(SimdIndex) - FilteredHistory.GetElement(SimdIndex));
+				tsr_halfC ClampedEnergyAmount = ClampedEnergy * SafeRcp(ClampedEnergy[0] + ClampedEnergy[1] + ClampedEnergy[2]);
+
+				// Apply the moire error on each individual channel based on how much they need to clamp.
+				tsr_halfC MoireErrorSize = ClampedEnergyAmount * ((FilteringWeight * tsr_half(3.0)) * MoireError.GetElement(SimdIndex));
+				//tsr_halfC MoireErrorSize = FilteringWeight * 3.0 * MoireError.GetElement(SimdIndex);
+			#else
+				const tsr_halfC MoireErrorSize = tsr_half(0.0);
+			#endif
+
 			Rejection.SetElement(SimdIndex, MeasureRejectionFactor(
 				FilteredHistory.GetElement(SimdIndex), ClampedFilteredHistory.GetElement(SimdIndex),
-				FilteredInput.GetElement(SimdIndex), tsr_half(0.0), BoxSize.GetElement(SimdIndex),
+				FilteredInput.GetElement(SimdIndex), tsr_half(0.0), max(BoxSize.GetElement(SimdIndex), MoireErrorSize),
 				MeasureBackbufferLDRQuantizationError()));
 		}
 		
@@ -521,6 +714,9 @@ void MainCS(
 			#endif
 
 			HistoryGuideOutput[OutputPixelPos] = FinalGuideColor;
+			#if CONFIG_MOIRE_DETECTION
+				HistoryMoireOutput[OutputPixelPos] = FinalMoire.GetElement(SimdIndex);
+			#endif
 			HistoryRejectionOutput[OutputPixelPos] = Rejection.GetElement(SimdIndex);
 			InputSceneColorLdrLumaOutput[OutputPixelPos] = CenterLumaLDR.GetElement(SimdIndex);
 			
diff --git a/Engine/Source/Runtime/Renderer/Private/PostProcess/TemporalSuperResolution.cpp b/Engine/Source/Runtime/Renderer/Private/PostProcess/TemporalSuperResolution.cpp
--- a/Engine/Source/Runtime/Renderer/Private/PostProcess/TemporalSuperResolution.cpp
+++ b/Engine/Source/Runtime/Renderer/Private/PostProcess/TemporalSuperResolution.cpp
@@ -41,6 +41,16 @@ TAutoConsoleVariable<int32> CVarTSRWaveOps(
 	TEXT("Whether to use wave ops in the shading rejection heuristics"),
 	ECVF_RenderThreadSafe);
 
+TAutoConsoleVariable<int32> CVarTSRFlickeringEnable(
+	TEXT("r.TSR.ShadingRejection.Flickering"), 1,
+	TEXT("Whether to enable the flickering detection heuristic.\n"),
+	ECVF_Scalability | ECVF_RenderThreadSafe);
+
+TAutoConsoleVariable<int32> CVarTSRFlickeringPeriod(
+	TEXT("r.TSR.ShadingRejection.Flickering.Period"), 3,
+	TEXT("Periode in frames in which luma oscilations at equal or greater frequency is considered flickering and should ghost (Default=3.0).\n"),
+	ECVF_RenderThreadSafe);
+
 TAutoConsoleVariable<int32> CVarTSRFilterShadingRejection(
 	TEXT("r.TSR.ShadingRejection.SpatialFilter"), 1,
 	TEXT("Whether the shading rejection should have spatial statistical filtering pass to reduce flickering (default = 1).\n")
@@ -132,6 +142,7 @@ BEGIN_SHADER_PARAMETER_STRUCT(FTSRHistoryTextures, )
 	SHADER_PARAMETER_RDG_TEXTURE(Texture2D, SubpixelDetails)
 	SHADER_PARAMETER_RDG_TEXTURE(Texture2D, TranslucencyAlpha)
 	SHADER_PARAMETER_RDG_TEXTURE(Texture2D, Guide)
+	SHADER_PARAMETER_RDG_TEXTURE(Texture2D, Moire)
 END_SHADER_PARAMETER_STRUCT()
 
 BEGIN_SHADER_PARAMETER_STRUCT(FTSRHistorySRVs, )
@@ -141,6 +152,7 @@ BEGIN_SHADER_PARAMETER_STRUCT(FTSRHistorySRVs, )
 	SHADER_PARAMETER_RDG_TEXTURE_SRV(Texture2D, Translucency)
 	SHADER_PARAMETER_RDG_TEXTURE_SRV(Texture2D, TranslucencyAlpha)
 	SHADER_PARAMETER_RDG_TEXTURE(Texture2D, Guide)
+	SHADER_PARAMETER_RDG_TEXTURE(Texture2D, Moire)
 END_SHADER_PARAMETER_STRUCT()
 
 BEGIN_SHADER_PARAMETER_STRUCT(FTSRHistoryUAVs, )
@@ -162,7 +174,8 @@ END_SHADER_PARAMETER_STRUCT()
 enum class ETSRHistoryFormatBits : uint32
 {
 	None = 0,
-	Translucency = 1 << 1,
+	Translucency = 1 << 0,
+	Moire = 1 << 1,
 };
 ENUM_CLASS_FLAGS(ETSRHistoryFormatBits);
 
@@ -193,6 +206,7 @@ FTSRHistorySRVs CreateSRVs(FRDGBuilder& GraphBuilder, const FTSRHistoryTextures&
 		SRVs.TranslucencyAlpha = GraphBuilder.CreateSRV(FRDGTextureSRVDesc::Create(Textures.TranslucencyAlpha));
 	}
 	SRVs.Guide = Textures.Guide;
+	SRVs.Moire = Textures.Moire;
 	return SRVs;
 }
 
@@ -297,6 +311,7 @@ class FTSRDilateVelocityCS : public FTSRShader
 		SHADER_PARAMETER(FVector2f, PrevOutputBufferUVMax)
 		SHADER_PARAMETER(float, WorldDepthToDepthError)
 		SHADER_PARAMETER(float, VelocityExtrapolationMultiplier)
+		SHADER_PARAMETER(int32, bOutputIsMovingTexture)
 
 		SHADER_PARAMETER_RDG_TEXTURE(Texture2D, SceneDepthTexture)
 		SHADER_PARAMETER_RDG_TEXTURE(Texture2D, SceneVelocityTexture)
@@ -305,7 +320,7 @@ class FTSRDilateVelocityCS : public FTSRShader
 		SHADER_PARAMETER_RDG_TEXTURE_UAV(RWTexture2D, ClosestDepthOutput)
 		SHADER_PARAMETER_RDG_TEXTURE_UAV(RWTexture2D, PrevUseCountOutput)
 		SHADER_PARAMETER_RDG_TEXTURE_UAV(RWTexture2D, PrevClosestDepthOutput)
-		SHADER_PARAMETER_RDG_TEXTURE_UAV(RWTexture2D, ParallaxFactorOutput)
+		SHADER_PARAMETER_RDG_TEXTURE_UAV(RWTexture2DArray, R8Output)
 
 		SHADER_PARAMETER_RDG_TEXTURE_UAV(RWTexture2D, VelocityFlattenOutput)
 		SHADER_PARAMETER_RDG_TEXTURE_UAV(RWTexture2D, VelocityTileOutput0)
@@ -320,6 +335,9 @@ class FTSRDecimateHistoryCS : public FTSRShader
 	DECLARE_GLOBAL_SHADER(FTSRDecimateHistoryCS);
 	SHADER_USE_PARAMETER_STRUCT(FTSRDecimateHistoryCS, FTSRShader);
 
+	class FMoireReprojectionDim : SHADER_PERMUTATION_BOOL("DIM_MOIRE_REPROJECTION");
+	using FPermutationDomain = TShaderPermutationDomain<FMoireReprojectionDim>;
+
 	BEGIN_SHADER_PARAMETER_STRUCT(FParameters, )
 		SHADER_PARAMETER_STRUCT_INCLUDE(FTSRCommonParameters, CommonParameters)
 		SHADER_PARAMETER(FMatrix44f, RotationalClipToPrevClip)
@@ -330,18 +348,20 @@ class FTSRDecimateHistoryCS : public FTSRShader
 		SHADER_PARAMETER_RDG_TEXTURE(Texture2D, ClosestDepthTexture)
 		SHADER_PARAMETER_RDG_TEXTURE(Texture2D, PrevUseCountTexture)
 		SHADER_PARAMETER_RDG_TEXTURE(Texture2D, PrevClosestDepthTexture)
-		SHADER_PARAMETER_RDG_TEXTURE(Texture2D, ParallaxFactorTexture)
+		SHADER_PARAMETER_RDG_TEXTURE_SRV(Texture2D, ParallaxFactorTexture)
 
 		SHADER_PARAMETER_STRUCT_INCLUDE(FTSRPrevHistoryParameters, PrevHistoryParameters)
 		SHADER_PARAMETER_RDG_TEXTURE(Texture2D, PrevHistorySubpixelDetails)
 
 		SHADER_PARAMETER_RDG_TEXTURE(Texture2D, PrevHistoryGuide)
+		SHADER_PARAMETER_RDG_TEXTURE(Texture2D, PrevHistoryMoire)
 		SHADER_PARAMETER_STRUCT(FScreenPassTextureViewportParameters, PrevGuideInfo)
 		SHADER_PARAMETER(FScreenTransform, InputPixelPosToReprojectScreenPos)
 		SHADER_PARAMETER(FScreenTransform, ScreenPosToPrevHistoryGuideBufferUV)
 		SHADER_PARAMETER(FVector3f, HistoryGuideQuantizationError)
 
 		SHADER_PARAMETER_RDG_TEXTURE_UAV(RWTexture2D, ReprojectedHistoryGuideOutput)
+		SHADER_PARAMETER_RDG_TEXTURE_UAV(RWTexture2D, ReprojectedHistoryMoireOutput)
 		SHADER_PARAMETER_RDG_TEXTURE_UAV(RWTexture2D, HoleFilledVelocityOutput)
 		SHADER_PARAMETER_RDG_TEXTURE_UAV(RWTexture2D, HoleFilledVelocityMaskOutput)
 		SHADER_PARAMETER_RDG_TEXTURE_UAV(RWTexture2D, ParallaxRejectionMaskOutput)
@@ -381,9 +401,9 @@ class FTSRRejectShadingCS : public FTSRShader
 	SHADER_USE_PARAMETER_STRUCT(FTSRRejectShadingCS, FTSRShader);
 
 	class FWaveSizeOps : SHADER_PERMUTATION_SPARSE_INT("DIM_WAVE_SIZE", 0, 32, 64);
-	class FSeparateTranslucencyDim : SHADER_PERMUTATION_BOOL("DIM_SEPARATE_TRANSLUCENCY");
+	class FFlickeringDetectionDim : SHADER_PERMUTATION_BOOL("DIM_FLICKERING_DETECTION");
 
-	using FPermutationDomain = TShaderPermutationDomain<FWaveSizeOps, FSeparateTranslucencyDim>;
+	using FPermutationDomain = TShaderPermutationDomain<FWaveSizeOps, FFlickeringDetectionDim>;
 
 	BEGIN_SHADER_PARAMETER_STRUCT(FParameters, )
 		SHADER_PARAMETER_STRUCT_INCLUDE(FTSRCommonParameters, CommonParameters)
@@ -391,13 +411,17 @@ class FTSRRejectShadingCS : public FTSRShader
 		SHADER_PARAMETER(FVector2f, TranslucencyTextureUVMin)
 		SHADER_PARAMETER(FVector2f, TranslucencyTextureUVMax)
 		SHADER_PARAMETER(FVector3f, HistoryGuideQuantizationError)
+		SHADER_PARAMETER(float, FlickeringFramePeriod)
 
 		SHADER_PARAMETER_RDG_TEXTURE(Texture2D, InputTexture)
 		SHADER_PARAMETER_RDG_TEXTURE(Texture2D, InputSceneTranslucencyTexture)
 		SHADER_PARAMETER_RDG_TEXTURE(Texture2D, ReprojectedHistoryGuideTexture)
+		SHADER_PARAMETER_RDG_TEXTURE(Texture2D, ReprojectedHistoryMoireTexture)
 		SHADER_PARAMETER_RDG_TEXTURE(Texture2D, ParallaxRejectionMaskTexture)
+		SHADER_PARAMETER_RDG_TEXTURE_SRV(Texture2D, IsMovingMaskTexture)
 
 		SHADER_PARAMETER_RDG_TEXTURE_UAV(RWTexture2D, HistoryGuideOutput)
+		SHADER_PARAMETER_RDG_TEXTURE_UAV(RWTexture2D, HistoryMoireOutput)
 		SHADER_PARAMETER_RDG_TEXTURE_UAV(RWTexture2D, HistoryRejectionOutput)
 		SHADER_PARAMETER_RDG_TEXTURE_UAV(RWTexture2D, InputSceneColorLdrLumaOutput)
 		SHADER_PARAMETER_RDG_TEXTURE_UAV(RWTexture2D, DebugOutput)
@@ -561,7 +585,7 @@ class FTSRUpdateHistoryCS : public FTSRShader
 		SHADER_PARAMETER_RDG_TEXTURE(Texture2D, TranslucencyRejectionTexture)
 
 		SHADER_PARAMETER_RDG_TEXTURE(Texture2D, DilatedVelocityTexture)
-		SHADER_PARAMETER_RDG_TEXTURE(Texture2D, ParallaxFactorTexture)
+		SHADER_PARAMETER_RDG_TEXTURE_SRV(Texture2D, ParallaxFactorTexture)
 		SHADER_PARAMETER_RDG_TEXTURE(Texture2D, ParallaxRejectionMaskTexture)
 		SHADER_PARAMETER_RDG_TEXTURE(Texture2D, AntiAliasingTexture)
 		SHADER_PARAMETER_RDG_TEXTURE(Texture2D, NoiseFilteringTexture)
@@ -711,11 +735,19 @@ ITemporalUpscaler::FOutputs AddTemporalSuperResolutionPasses(
 	// whether TSR passes can run on async compute.
 	int32 AsyncComputePasses = GSupportsEfficientAsyncCompute ? CVarTSRAsyncCompute.GetValueOnRenderThread() : 0;
 
+	// period at which history changes is considered too distracting.
+	const int32 FlickeringFramePeriod = CVarTSRFlickeringEnable.GetValueOnRenderThread() ? CVarTSRFlickeringPeriod.GetValueOnRenderThread() : 0;
+
 	ETSRHistoryFormatBits HistoryFormatBits = ETSRHistoryFormatBits::None;
 	{
 		if (bAccumulateTranslucencySeparately)
 		{
-			HistoryFormatBits = ETSRHistoryFormatBits::Translucency;
+			HistoryFormatBits |= ETSRHistoryFormatBits::Translucency;
+		}
+
+		if (FlickeringFramePeriod > 0)
+		{
+			HistoryFormatBits |= ETSRHistoryFormatBits::Moire;
 		}
 	}
 
@@ -898,9 +930,14 @@ ITemporalUpscaler::FOutputs AddTemporalSuperResolutionPasses(
 	// Dilate the velocity texture & scatter reprojection into previous frame
 	FRDGTextureRef DilatedVelocityTexture;
 	FRDGTextureRef ClosestDepthTexture;
-	FRDGTextureRef ParallaxFactorTexture;
+	FRDGTextureSRVRef ParallaxFactorTexture;
+	FRDGTextureSRVRef IsMovingMaskTexture = nullptr;
 	FVelocityFlattenTextures VelocityFlattenTextures;
 	{
+		const bool bOutputIsMovingTexture = FlickeringFramePeriod > 0;
+
+		FRDGTextureRef R8OutputTexture;
+
 		{
 			FRDGTextureDesc Desc = FRDGTextureDesc::Create2D(
 				InputExtent,
@@ -912,9 +949,21 @@ ITemporalUpscaler::FOutputs AddTemporalSuperResolutionPasses(
 
 			Desc.Format = PF_R16F;
 			ClosestDepthTexture = GraphBuilder.CreateTexture(Desc, TEXT("TSR.ClosestDepthTexture"));
+		}
+		{
+			FRDGTextureDesc Desc = FRDGTextureDesc::Create2DArray(
+				InputExtent,
+				PF_R8_UINT,
+				FClearValueBinding::None,
+				/* InFlags = */ TexCreate_ShaderResource | TexCreate_UAV,
+				bOutputIsMovingTexture ? 2 : 1);
 
-			Desc.Format = PF_R8_UINT;
-			ParallaxFactorTexture = GraphBuilder.CreateTexture(Desc, TEXT("TSR.ParallaxFactor"));
+			R8OutputTexture = GraphBuilder.CreateTexture(Desc, TEXT("TSR.ParallaxFactor"));
+			ParallaxFactorTexture = GraphBuilder.CreateSRV(FRDGTextureSRVDesc::CreateForSlice(R8OutputTexture, 0));
+			if (bOutputIsMovingTexture)
+			{
+				IsMovingMaskTexture = GraphBuilder.CreateSRV(FRDGTextureSRVDesc::CreateForSlice(R8OutputTexture, 1));
+			}
 		}
 
 		int32 TileSize = 8;
@@ -932,6 +981,7 @@ ITemporalUpscaler::FOutputs AddTemporalSuperResolutionPasses(
 			PassParameters->WorldDepthToDepthError = WorldDepthToPixelWorldRadius * 2.0f;
 		}
 		PassParameters->VelocityExtrapolationMultiplier = FMath::Clamp(CVarTSRVelocityExtrapolation.GetValueOnRenderThread(), 0.0f, 1.0f);
+		PassParameters->bOutputIsMovingTexture = bOutputIsMovingTexture;
 
 		PassParameters->SceneDepthTexture = PassInputs.SceneDepthTexture;
 		PassParameters->SceneVelocityTexture = PassInputs.SceneVelocityTexture;
@@ -940,7 +990,7 @@ ITemporalUpscaler::FOutputs AddTemporalSuperResolutionPasses(
 		PassParameters->ClosestDepthOutput = GraphBuilder.CreateUAV(ClosestDepthTexture);
 		PassParameters->PrevUseCountOutput = GraphBuilder.CreateUAV(PrevUseCountTexture);
 		PassParameters->PrevClosestDepthOutput = GraphBuilder.CreateUAV(PrevClosestDepthTexture);
-		PassParameters->ParallaxFactorOutput = GraphBuilder.CreateUAV(ParallaxFactorTexture);
+		PassParameters->R8Output = GraphBuilder.CreateUAV(R8OutputTexture);
 
 		// Setup up the motion blur's velocity flatten pass.
 		if (PassInputs.bGenerateVelocityFlattenTextures)
@@ -986,8 +1036,9 @@ ITemporalUpscaler::FOutputs AddTemporalSuperResolutionPasses(
 		TShaderMapRef<FTSRDilateVelocityCS> ComputeShader(View.ShaderMap, PermutationVector);
 		FComputeShaderUtils::AddPass(
 			GraphBuilder,
-			RDG_EVENT_NAME("TSR DilateVelocity(MotionBlurDirections=%d) %dx%d",
+			RDG_EVENT_NAME("TSR DilateVelocity(MotionBlurDirections=%d%s) %dx%d",
 				int32(PermutationVector.Get<FTSRDilateVelocityCS::FMotionBlurDirectionsDim>()),
+				bOutputIsMovingTexture ? TEXT(" OutputIsMoving") : TEXT(""),
 				InputRect.Width(), InputRect.Height()),
 			AsyncComputePasses >= 2 ? ERDGPassFlags::AsyncCompute : ERDGPassFlags::Compute,
 			ComputeShader,
@@ -1041,6 +1092,8 @@ ITemporalUpscaler::FOutputs AddTemporalSuperResolutionPasses(
 			TexCreate_ShaderResource | TexCreate_UAV);
 
 		History.Guide = GraphBuilder.CreateTexture(Desc, TEXT("TSR.History.Guide"));
+		Desc.Format = PF_R8G8B8A8;
+		History.Moire = GraphBuilder.CreateTexture(Desc, TEXT("TSR.History.Moire"));
 	}
 
 	// Setup a dummy history
@@ -1051,6 +1104,7 @@ ITemporalUpscaler::FOutputs AddTemporalSuperResolutionPasses(
 		DummyHistorySRVs.Translucency = GraphBuilder.CreateSRV(FRDGTextureSRVDesc::Create(BlackDummy));
 		DummyHistorySRVs.TranslucencyAlpha = GraphBuilder.CreateSRV(FRDGTextureSRVDesc::Create(WhiteDummy));
 		DummyHistorySRVs.Guide = BlackDummy;
+		DummyHistorySRVs.Moire = BlackDummy;
 
 		DummyHistorySRVs.SubpixelDetails = BlackUintDummy;
 	}
@@ -1067,6 +1121,7 @@ ITemporalUpscaler::FOutputs AddTemporalSuperResolutionPasses(
 		PrevHistory.Metadata = GraphBuilder.RegisterExternalTexture(InputHistory.Metadata);
 		PrevHistory.TranslucencyAlpha = InputHistory.TranslucencyAlpha.IsValid() ? GraphBuilder.RegisterExternalTexture(InputHistory.TranslucencyAlpha) : DummyHistorySRVs.TranslucencyAlpha->Desc.Texture;
 		PrevHistory.Guide = GraphBuilder.RegisterExternalTexture(InputHistory.Guide);
+		PrevHistory.Moire = InputHistory.Moire.IsValid() ? GraphBuilder.RegisterExternalTexture(InputHistory.Moire) : DummyHistorySRVs.Moire;
 
 		// Register non-filterable history
 		PrevHistory.SubpixelDetails = GraphBuilder.RegisterExternalTexture(InputHistory.SubpixelDetails);
@@ -1104,6 +1159,7 @@ ITemporalUpscaler::FOutputs AddTemporalSuperResolutionPasses(
 
 	// Decimate input to flicker at same frequency as input.
 	FRDGTextureRef ReprojectedHistoryGuideTexture = nullptr;
+	FRDGTextureRef ReprojectedHistoryMoireTexture = nullptr;
 	FRDGTextureRef ParallaxRejectionMaskTexture = nullptr;
 	FRDGTextureRef HoleFilledVelocityMaskTexture = nullptr;
 	{
@@ -1134,6 +1190,16 @@ ITemporalUpscaler::FOutputs AddTemporalSuperResolutionPasses(
 			ReprojectedHistoryGuideTexture = GraphBuilder.CreateTexture(Desc, TEXT("TSR.ReprojectedHistoryGuide"));
 		}
 
+		{
+			FRDGTextureDesc Desc = FRDGTextureDesc::Create2D(
+				InputExtent,
+				PrevHistorySRVs.Moire->Desc.Format,
+				FClearValueBinding::None,
+				/* InFlags = */ TexCreate_ShaderResource | TexCreate_UAV);
+
+			ReprojectedHistoryMoireTexture = GraphBuilder.CreateTexture(Desc, TEXT("TSR.ReprojectedHistoryMoire"));
+		}
+
 		FTSRDecimateHistoryCS::FParameters* PassParameters = GraphBuilder.AllocParameters<FTSRDecimateHistoryCS::FParameters>();
 		PassParameters->CommonParameters = CommonParameters;
 		{
@@ -1165,6 +1231,7 @@ ITemporalUpscaler::FOutputs AddTemporalSuperResolutionPasses(
 		{
 			FScreenPassTextureViewport PrevHistoryColorViewport(PrevHistorySRVs.Guide->Desc.Extent, InputHistory.InputViewportRect);
 			PassParameters->PrevHistoryGuide = PrevHistorySRVs.Guide;
+			PassParameters->PrevHistoryMoire = PrevHistorySRVs.Moire;
 			PassParameters->PrevGuideInfo = GetScreenPassTextureViewportParameters(PrevHistoryColorViewport);
 			PassParameters->InputPixelPosToReprojectScreenPos = ((FScreenTransform::Identity - InputRect.Min + 0.5f) / InputRect.Size()) * FScreenTransform::ViewportUVToScreenPos;
 			PassParameters->ScreenPosToPrevHistoryGuideBufferUV = FScreenTransform::ChangeTextureBasisFromTo(
@@ -1175,16 +1242,22 @@ ITemporalUpscaler::FOutputs AddTemporalSuperResolutionPasses(
 		}
 
 		PassParameters->ReprojectedHistoryGuideOutput = GraphBuilder.CreateUAV(ReprojectedHistoryGuideTexture);
+		PassParameters->ReprojectedHistoryMoireOutput = GraphBuilder.CreateUAV(ReprojectedHistoryMoireTexture);
 		PassParameters->HoleFilledVelocityOutput = GraphBuilder.CreateUAV(HoleFilledVelocityTexture);
 		PassParameters->HoleFilledVelocityMaskOutput = GraphBuilder.CreateUAV(HoleFilledVelocityMaskTexture);
 		PassParameters->ParallaxRejectionMaskOutput = GraphBuilder.CreateUAV(ParallaxRejectionMaskTexture);
 		PassParameters->HistorySubpixelDetailsOutput = GraphBuilder.CreateUAV(History.SubpixelDetails);
 		PassParameters->DebugOutput = CreateDebugUAV(InputExtent, TEXT("Debug.TSR.DecimateHistory"));
 
-		TShaderMapRef<FTSRDecimateHistoryCS> ComputeShader(View.ShaderMap);
+		FTSRDecimateHistoryCS::FPermutationDomain PermutationVector;
+		PermutationVector.Set<FTSRDecimateHistoryCS::FMoireReprojectionDim>(FlickeringFramePeriod > 0);
+
+		TShaderMapRef<FTSRDecimateHistoryCS> ComputeShader(View.ShaderMap, PermutationVector);
 		FComputeShaderUtils::AddPass(
 			GraphBuilder,
-			RDG_EVENT_NAME("TSR DecimateHistory %dx%d", InputRect.Width(), InputRect.Height()),
+			RDG_EVENT_NAME("TSR DecimateHistory(%s) %dx%d",
+				PermutationVector.Get<FTSRDecimateHistoryCS::FMoireReprojectionDim>() ? TEXT("ReprojectMoire") : TEXT(""),
+				InputRect.Width(), InputRect.Height()),
 			AsyncComputePasses >= 2 ? ERDGPassFlags::AsyncCompute : ERDGPassFlags::Compute,
 			ComputeShader,
 			PassParameters,
@@ -1289,27 +1362,39 @@ ITemporalUpscaler::FOutputs AddTemporalSuperResolutionPasses(
 		PassParameters->TranslucencyTextureUVMin = GetScreenPassTextureViewportParameters(TranslucencyViewport).UVViewportBilinearMin;
 		PassParameters->TranslucencyTextureUVMax = GetScreenPassTextureViewportParameters(TranslucencyViewport).UVViewportBilinearMax;
 		PassParameters->HistoryGuideQuantizationError = ComputePixelFormatQuantizationError(History.Guide->Desc.Format);
+		PassParameters->FlickeringFramePeriod = float(FlickeringFramePeriod);
 
 		PassParameters->InputTexture = PassInputs.SceneColorTexture;
-		PassParameters->InputSceneTranslucencyTexture = SeparateTranslucencyTexture;
+		if (bAccumulateTranslucencySeparately)
+		{
+			PassParameters->InputSceneTranslucencyTexture = BlackAlphaOneDummy;
+		}
+		else
+		{
+			PassParameters->InputSceneTranslucencyTexture = SeparateTranslucencyTexture;
+		}
 		PassParameters->ReprojectedHistoryGuideTexture = ReprojectedHistoryGuideTexture;
+		PassParameters->ReprojectedHistoryMoireTexture = ReprojectedHistoryMoireTexture;
 		PassParameters->ParallaxRejectionMaskTexture = ParallaxRejectionMaskTexture;
+		PassParameters->IsMovingMaskTexture = IsMovingMaskTexture;
 
 		PassParameters->HistoryGuideOutput = GraphBuilder.CreateUAV(History.Guide);
+		PassParameters->HistoryMoireOutput = GraphBuilder.CreateUAV(History.Moire);
 		PassParameters->HistoryRejectionOutput = GraphBuilder.CreateUAV(HistoryRejectionTexture);
 		PassParameters->InputSceneColorLdrLumaOutput = GraphBuilder.CreateUAV(InputSceneColorLdrLumaTexture);
 		PassParameters->DebugOutput = CreateDebugUAV(InputExtent, TEXT("Debug.TSR.RejectShading"));
 
 		FTSRRejectShadingCS::FPermutationDomain PermutationVector;
 		PermutationVector.Set<FTSRRejectShadingCS::FWaveSizeOps>(bUseWaveOps && GRHIMinimumWaveSize >= 32 && GRHIMinimumWaveSize <= 64 ? GRHIMinimumWaveSize : 0);
-		PermutationVector.Set<FTSRRejectShadingCS::FSeparateTranslucencyDim>(bAccumulateTranslucencySeparately);
+		PermutationVector.Set<FTSRRejectShadingCS::FFlickeringDetectionDim>(FlickeringFramePeriod > 0);
 
 		TShaderMapRef<FTSRRejectShadingCS> ComputeShader(View.ShaderMap, PermutationVector);
 		FComputeShaderUtils::AddPass(
 			GraphBuilder,
-			RDG_EVENT_NAME("TSR RejectShading(WaveSize=%d%s) %dx%d",
+			RDG_EVENT_NAME("TSR RejectShading(WaveSize=%d FlickeringFramePeriod=%d %s) %dx%d",
 				int32(PermutationVector.Get<FTSRRejectShadingCS::FWaveSizeOps>()),
-				PermutationVector.Get<FTSRRejectShadingCS::FSeparateTranslucencyDim>() ? TEXT(" SeparateTranslucency") : TEXT(""),
+				FlickeringFramePeriod,
+				!bAccumulateTranslucencySeparately ? TEXT(" ComposeTranslucency") : TEXT(""),
 				InputRect.Width(), InputRect.Height()),
 			AsyncComputePasses >= 3 ? ERDGPassFlags::AsyncCompute : ERDGPassFlags::Compute,
 			ComputeShader,
@@ -1642,6 +1727,11 @@ ITemporalUpscaler::FOutputs AddTemporalSuperResolutionPasses(
 		// Extract history guide
 		GraphBuilder.QueueTextureExtraction(History.Guide, &OutputHistory.Guide);
 
+		if (FlickeringFramePeriod > 0)
+		{
+			GraphBuilder.QueueTextureExtraction(History.Moire, &OutputHistory.Moire);
+		}
+
 		// Extract the translucency buffer to compare it with next frame
 		if (bRejectSeparateTranslucency)
 		{
diff --git a/Engine/Source/Runtime/Renderer/Private/SceneRendering.h b/Engine/Source/Runtime/Renderer/Private/SceneRendering.h
--- a/Engine/Source/Runtime/Renderer/Private/SceneRendering.h
+++ b/Engine/Source/Runtime/Renderer/Private/SceneRendering.h
@@ -917,6 +917,7 @@ struct FTSRHistory
 
 	// Filterable input resolution representation of the output
 	TRefCountPtr<IPooledRenderTarget> Guide;
+	TRefCountPtr<IPooledRenderTarget> Moire;
 
 	FIntRect InputViewportRect;
 	FIntRect OutputViewportRect;
