Move Strata material buffer slice 0 into ESRAM.

This CL remove SceneColor from ESRAM without noticeable perf impact.

Other perf differences:
* No impact on base pass cost
* Classification 0.373 -> 0.240
* DBufferPass 0.510 -> 0.420

diff --git a/Engine/Source/Runtime/Renderer/Private/SceneRendering.cpp b/Engine/Source/Runtime/Renderer/Private/SceneRendering.cpp
--- a/Engine/Source/Runtime/Renderer/Private/SceneRendering.cpp
+++ b/Engine/Source/Runtime/Renderer/Private/SceneRendering.cpp
@@ -639,6 +639,12 @@ void FFastVramConfig::Update()
 	bDirty |= UpdateBufferFlagFromCVar(CVarFastVRam_DistanceFieldAOScreenGridResources, DistanceFieldAOScreenGridResources);
 	bDirty |= UpdateBufferFlagFromCVar(CVarFastVRam_ForwardLightingCullingResources, ForwardLightingCullingResources);
 	bDirty |= UpdateBufferFlagFromCVar(CVarFastVRam_GlobalDistanceFieldCullGridBuffers, GlobalDistanceFieldCullGridBuffers);
+
+	// When strata is enable, remove Scene color from fast VRAM to leave space for material buffer which has more impact on performance
+	if (Strata::IsStrataEnabled() && !IsForwardShadingEnabled(GMaxRHIShaderPlatform))
+	{
+		SceneColor = SceneColor & (~(TexCreate_FastVRAM | TexCreate_FastVRAMPartialAlloc));
+	}
 }
 
 bool FFastVramConfig::UpdateTextureFlagFromCVar(TAutoConsoleVariable<int32>& CVar, ETextureCreateFlags& InOutValue)
diff --git a/Engine/Source/Runtime/Renderer/Private/Strata/Strata.cpp b/Engine/Source/Runtime/Renderer/Private/Strata/Strata.cpp
--- a/Engine/Source/Runtime/Renderer/Private/Strata/Strata.cpp
+++ b/Engine/Source/Runtime/Renderer/Private/Strata/Strata.cpp
@@ -422,8 +422,9 @@ void InitialiseStrataFrameSceneData(FRDGBuilder& GraphBuilder, FSceneRenderer& S
 	const uint32 SliceCountSSS = STRATA_SSS_DATA_UINT_COUNT;
 	const uint32 SliceCountAdvDebug = IsStrataAdvancedVisualizationShadersEnabled() ? 1 : 0;
 	const uint32 SliceCount = FMath::DivideAndRoundUp(Out.MaxBytesPerPixel, 4u) + SliceCountSSS + SliceCountAdvDebug;
-	const FRDGTextureDesc MaterialTextureDesc = FRDGTextureDesc::Create2DArray(SceneTextureExtent, PF_R32_UINT, FClearValueBinding::Transparent,
-		TexCreate_TargetArraySlicesIndependently | TexCreate_DisableDCC | TexCreate_NoFastClear | TexCreate_RenderTargetable | TexCreate_ShaderResource | TexCreate_UAV, SliceCount, 1, 1);
+	FRDGTextureDesc MaterialTextureDesc = FRDGTextureDesc::Create2DArray(SceneTextureExtent, PF_R32_UINT, FClearValueBinding::Transparent, TexCreate_TargetArraySlicesIndependently | TexCreate_DisableDCC | TexCreate_NoFastClear | TexCreate_RenderTargetable | TexCreate_ShaderResource | TexCreate_UAV | TexCreate_FastVRAM, SliceCount, 1, 1);
+	MaterialTextureDesc.FastVRAMPercentage = (1.0f / SliceCount) * 0xFF; // Only allocate the first slice into ESRAM
+
 	Out.MaterialTextureArray = GraphBuilder.CreateTexture(MaterialTextureDesc, TEXT("Strata.Material"));
 	Out.MaterialTextureArraySRV = GraphBuilder.CreateSRV(FRDGTextureSRVDesc::Create(Out.MaterialTextureArray));
 	Out.MaterialTextureArrayUAV = GraphBuilder.CreateUAV(FRDGTextureUAVDesc(Out.MaterialTextureArray, 0));
