Use DrawsVelocity instead of IsMovable to check whether primitive can draw velocity.
This fixes case where static primitives were not rendering velocity, even with WPO.
Might need to follow up with fixes to PSO precache coverage.
#rnx

diff --git a/Engine/Source/Runtime/Renderer/Private/VelocityRendering.cpp b/Engine/Source/Runtime/Renderer/Private/VelocityRendering.cpp
--- a/Engine/Source/Runtime/Renderer/Private/VelocityRendering.cpp
+++ b/Engine/Source/Runtime/Renderer/Private/VelocityRendering.cpp
@@ -455,7 +455,7 @@ bool FVelocityMeshProcessor::PrimitiveHasVelocityForView(const FViewInfo& View,
 
 bool FOpaqueVelocityMeshProcessor::PrimitiveCanHaveVelocity(EShaderPlatform ShaderPlatform, const FPrimitiveSceneProxy* PrimitiveSceneProxy)
 {
-	return PrimitiveCanHaveVelocity(ShaderPlatform, PrimitiveSceneProxy->IsMovable(), PrimitiveSceneProxy->HasStaticLighting());
+	return PrimitiveCanHaveVelocity(ShaderPlatform, PrimitiveSceneProxy->DrawsVelocity(), PrimitiveSceneProxy->HasStaticLighting());
 }
 
 bool FOpaqueVelocityMeshProcessor::PrimitiveCanHaveVelocity(EShaderPlatform ShaderPlatform, bool bDrawVelocity, bool bHasStaticLighting)
