Fix translucent part of Gizmos shows up in shifted location when the camera viewport does not cover whole viewport

diff --git a/Engine/Source/Runtime/Renderer/Private/PostProcess/PostProcessCompositeEditorPrimitives.cpp b/Engine/Source/Runtime/Renderer/Private/PostProcess/PostProcessCompositeEditorPrimitives.cpp
--- a/Engine/Source/Runtime/Renderer/Private/PostProcess/PostProcessCompositeEditorPrimitives.cpp
+++ b/Engine/Source/Runtime/Renderer/Private/PostProcess/PostProcessCompositeEditorPrimitives.cpp
@@ -629,9 +629,9 @@ FScreenPassTexture AddEditorPrimitivePass(
 				NumSamples),
 			PassParameters,
 			ERDGPassFlags::Raster,
-			[&View, &InstanceCullingManager, PassParameters, EditorView, EditorPrimitivesViewport, BasePassType, NumSamples](FRHICommandListImmediate& RHICmdList)
+			[&View, &InstanceCullingManager, PassParameters, EditorView, EditorPrimitivesViewport, OutputViewportRect = Output.ViewRect, BasePassType, NumSamples](FRHICommandListImmediate& RHICmdList)
 			{
-				RHICmdList.SetViewport(EditorPrimitivesViewport.Rect.Min.X, EditorPrimitivesViewport.Rect.Min.Y, 0.0f, EditorPrimitivesViewport.Rect.Max.X, EditorPrimitivesViewport.Rect.Max.Y, 1.0f);
+				RHICmdList.SetViewport(OutputViewportRect.Min.X, OutputViewportRect.Min.Y, 0.0f, OutputViewportRect.Max.X, OutputViewportRect.Max.Y, 1.0f);
 
 				FMeshPassProcessorRenderState DrawRenderState;
 				DrawRenderState.SetDepthStencilAccess(FExclusiveDepthStencil::DepthRead_StencilNop);
