OpenXR PixelDensity: Allow XR scene textures to use requested extents more uniformly

We previously refactored some of the code in FSceneTextureExtentState to allow for retaining VR scene textures for VR scene capture  (planar reflection) usage, assuming those scene capture extents were smaller than the 'main' VR extents.

In the previous change (12158192), there's logic where the VR extent request will 'grow' if the same VR extent is requested for multiple frames. The thinking might be to get the desired size into the history, but it's not obvious code. This 'breaks' when using a reduced PixelDensity + non-VR SceneCapture. SceneCapture will 'grow' into the bigger texture from history, which will be PD 1.0. The following usage of a sub-1.0 PD VR scene texture (with now a matching extent to the retained stereo extent) will also grow. And when it looks into the history, it will find the PD 1.0 texture, and grow into it, which wasn't the original intent.

<Resize to PD 0.5>
Frame 0 SC - Grow into PD 1.0 VR texture
Frame 0 Main Render - Use PD 0.5 size as extent because stereo history mismatch
Frame 1 SC - Grow into PD 1.0 VR texture still in history
Frame 1 Main Render - Grow into PD 1.0 VR texture because SC retained it in history

The code here has rotted from the VR perspective, and should have further updates. The original logic was meant to manage the interactions between VR main render + VR-enabled planar reflections. We need to account for all non-main render interactions. But we'll defer the bulk of risk to Main.

diff --git a/Engine/Source/Runtime/Renderer/Private/SceneTextures.cpp b/Engine/Source/Runtime/Renderer/Private/SceneTextures.cpp
--- a/Engine/Source/Runtime/Renderer/Private/SceneTextures.cpp
+++ b/Engine/Source/Runtime/Renderer/Private/SceneTextures.cpp
@@ -181,23 +181,18 @@ public:
 					{
 						if (!bIsSceneCapture && !bIsReflectionCapture)
 						{
-							// If this isn't a scene capture, and it's a VR scene, and the size has changed since the last time we
-							// rendered a VR scene (or this is the first time), use the requested size method.
+							// If this is VR, but not a capture (only current XR capture is for Planar Reflections), then we want
+							// to use the requested size. Ideally, capture targets will be able to 'grow' into the VR extents.
 							if (DesiredFamilyExtent.X != LastStereoExtent.X || DesiredFamilyExtent.Y != LastStereoExtent.Y)
 							{
 								LastStereoExtent = DesiredFamilyExtent;
-								SceneTargetsSizingMethod = RequestedSize;
 								UE_LOG(LogRenderer, Warning, TEXT("Resizing VR buffer to %d by %d"), DesiredFamilyExtent.X, DesiredFamilyExtent.Y);
 							}
-							else
-							{
-								// Otherwise use the grow method.
-								SceneTargetsSizingMethod = Grow;
-							}
+							SceneTargetsSizingMethod = RequestedSize;
 						}
 						else
 						{
-							// If this is a scene capture, and it's smaller than the VR view size, then don't re-allocate buffers, just use the "grow" method.
+							// If this is a VR scene capture (i.e planar reflection), and it's smaller than the VR view size, then don't re-allocate buffers, just use the "grow" method.
 							// If it's bigger than the VR view, then log a warning, and use resize method.
 							if (DesiredFamilyExtent.X > LastStereoExtent.X || DesiredFamilyExtent.Y > LastStereoExtent.Y)
 							{
