Fixes TSR responsiveness under cinematic scalability settings

diff --git a/Engine/Shaders/Private/TemporalSuperResolution/TSRUpdateHistory.usf b/Engine/Shaders/Private/TemporalSuperResolution/TSRUpdateHistory.usf
--- a/Engine/Shaders/Private/TemporalSuperResolution/TSRUpdateHistory.usf
+++ b/Engine/Shaders/Private/TemporalSuperResolution/TSRUpdateHistory.usf
@@ -27,6 +27,9 @@
 	#define CONFIG_INTERPOLATE_VELOCITY 1
 	#define CONFIG_REJECTION_ANTI_ALIASING 1
 
+	// To support sg.AntiAliasingQuality=4 with same temporal response as sg.AntiAliasingQuality=3
+	#define CONFIG_APPLY_CONTRIBUTION_MULTIPLIER 1
+
 #else
 	#error Unknown history update quality
 #endif
@@ -72,6 +75,10 @@
 	#define CONFIG_ENABLE_DUAL_PIXEL_VECTORIZATION 0
 #endif
 
+#ifndef CONFIG_APPLY_CONTRIBUTION_MULTIPLIER
+	#define CONFIG_APPLY_CONTRIBUTION_MULTIPLIER 0
+#endif
+
 #ifndef PLATFORM_SPECIFIC_ISOLATE
 	#define PLATFORM_SPECIFIC_ISOLATE ISOLATE
 #endif
@@ -122,6 +129,7 @@ float3 HistoryQuantizationError;
 float MinTranslucencyRejection;
 float InvWeightClampingPixelSpeed;
 float InputToHistoryFactor;
+float InputContributionMultiplier;
 uint ResponsiveStencilMask;
 uint bGenerateOutputMip1;
 uint bHasSeparateTranslucency;
@@ -764,6 +772,10 @@ void MainCS(
 			
 			InputPixelAlignement = ComputeSampleWeigth(KernelInputToHistoryAlignmentFactor, dInputKO, /* MinimalContribution = */ float(0.0));
 		
+			#if CONFIG_APPLY_CONTRIBUTION_MULTIPLIER
+				InputPixelAlignement *= tsr_half(InputContributionMultiplier);
+			#endif
+
 			UNROLL_N(CONFIG_SAMPLES)
 			for (uint SampleId = 0; SampleId < CONFIG_SAMPLES; SampleId++)
 			{
diff --git a/Engine/Source/Runtime/Renderer/Private/PostProcess/TemporalSuperResolution.cpp b/Engine/Source/Runtime/Renderer/Private/PostProcess/TemporalSuperResolution.cpp
--- a/Engine/Source/Runtime/Renderer/Private/PostProcess/TemporalSuperResolution.cpp
+++ b/Engine/Source/Runtime/Renderer/Private/PostProcess/TemporalSuperResolution.cpp
@@ -604,6 +604,7 @@ class FTSRUpdateHistoryCS : public FTSRShader
 		SHADER_PARAMETER(float, MinTranslucencyRejection)
 		SHADER_PARAMETER(float, InvWeightClampingPixelSpeed)
 		SHADER_PARAMETER(float, InputToHistoryFactor)
+		SHADER_PARAMETER(float, InputContributionMultiplier)
 		SHADER_PARAMETER(int32, ResponsiveStencilMask)
 		SHADER_PARAMETER(int32, bGenerateOutputMip1)
 		SHADER_PARAMETER(int32, bHasSeparateTranslucency)
@@ -1585,6 +1586,7 @@ ITemporalUpscaler::FOutputs AddTemporalSuperResolutionPasses(
 		PassParameters->MinTranslucencyRejection = TranslucencyRejectionTexture == nullptr ? 1.0 : 0.0;
 		PassParameters->InvWeightClampingPixelSpeed = 1.0f / CVarTSRWeightClampingPixelSpeed.GetValueOnRenderThread();
 		PassParameters->InputToHistoryFactor = float(HistorySize.X) / float(InputRect.Width());
+		PassParameters->InputContributionMultiplier = FMath::Pow(float(HistorySize.X) / float(OutputRect.Width()), 2.0f); 
 		PassParameters->ResponsiveStencilMask = CVarTSREnableResponiveAA.GetValueOnRenderThread() ? (STENCIL_TEMPORAL_RESPONSIVE_AA_MASK) : 0;
 		PassParameters->bGenerateOutputMip1 = ((PassInputs.bGenerateOutputMip1 || PassInputs.bAllowDownsampleSceneColor) && HistorySize == OutputRect.Size()) ? 1 : 0;
 		PassParameters->bHasSeparateTranslucency = bHasSeparateTranslucency;
