Lumen - don't allocate and update GlobalDistanceField.PageObjectGridBuffer, when Lumen is disabled.

diff --git a/Engine/Source/Runtime/Renderer/Private/GlobalDistanceField.cpp b/Engine/Source/Runtime/Renderer/Private/GlobalDistanceField.cpp
--- a/Engine/Source/Runtime/Renderer/Private/GlobalDistanceField.cpp
+++ b/Engine/Source/Runtime/Renderer/Private/GlobalDistanceField.cpp
@@ -820,7 +820,7 @@ static void ComputeUpdateRegionsAndUpdateViewState(
 				GlobalDistanceFieldData.PageFreeListBuffer = AllocatePooledBuffer(FRDGBufferDesc::CreateStructuredDesc(sizeof(uint32), MaxPageNum), TEXT("GlobalDistanceField.PageFreeList"));
 			}
 
-			if (Lumen::UseGlobalSDFObjectGrid(*View.Family))
+			if (bLumenEnabled && Lumen::UseGlobalSDFObjectGrid(*View.Family))
 			{
 				if (!GlobalDistanceFieldData.PageObjectGridBuffer
 					|| GlobalDistanceFieldData.PageObjectGridBuffer->Desc.NumElements != PageObjectGridBufferSize)
