HeterogeneousVolumes: Storing in-scattering computation in a general-purpose lighting cache.

This submission repurposes the transmittance volume as a general-purpose lighting cache. A secondary pipeline has been created to cache in-scattering contribution in the lighting cache, which can be accuulated per-light in independent passes Accumumulating in-scattering contribution eliminates redundant direct volume rendering passes, which provides a significant runtime performance boost, scaling per-light.

r.HeterogeneousVolumes.LightingCache (Default =2)
0: Disabled
1: Cache transmittance
2: Cache in-scattering

Performance comparison between caching transmittance v in-scattering, when rendering "Grid3D Gas Cine Explosion Cine" with three lights, with no discernible difference in quality:

r.HeterogeneousVolumes.LightingCache 1
r.HeterogeneousVolumes.LightingCache.DownsampleFactor 2

47.2% 4.42ms   HeterogeneousVolumes 8 dispatches
    0.6% 0.06ms   ClearTextureFloat(HeterogeneousVolumes) 1924x1092 1 dispatch 241x137 groups
   46.6% 4.36ms   NS_Grid3D_Gas_Explosion_Cine_5_1 7 dispatches
       0.3% 0.02ms   ClearTextureFloat(HeterogeneousVolumes.LightingCacheTexture) 128x128 1 dispatch 32x32x32 groups
       4.2% 0.40ms   RenderLightingCacheWithLiveShadingCS [Transmittance] (Light = Untitled_1.PointLight) 1 dispatch 32x32x32 groups
      11.3% 1.06ms   RenderSingleScatteringWithLiveShadingCS (Light = Untitled_1.PointLight) 1 dispatch 240x135 groups
       4.2% 0.39ms   RenderLightingCacheWithLiveShadingCS [Transmittance] (Light = Untitled_1.DirectionalLight) 1 dispatch 32x32x32 groups
      11.1% 1.03ms   RenderSingleScatteringWithLiveShadingCS (Light = Untitled_1.DirectionalLight) 1 dispatch 240x135 groups
       4.2% 0.39ms   RenderLightingCacheWithLiveShadingCS [Transmittance] (Light = Untitled_1.SpotLight) 1 dispatch 32x32x32 groups
      11.2% 1.05ms   RenderSingleScatteringWithLiveShadingCS (Light = Untitled_1.SpotLight) 1 dispatch 240x135 groups

r.HeterogeneousVolumes.LightingCache 2
r.HeterogeneousVolumes.LightingCache.DownsampleFactor 2

32.3% 2.37ms   HeterogeneousVolumes 6 dispatches
    0.8% 0.06ms   ClearTextureFloat(HeterogeneousVolumes) 1924x1092 1 dispatch 241x137 groups
   31.5% 2.31ms   NS_Grid3D_Gas_Explosion_Cine_5_1 5 dispatches
       0.3% 0.02ms   ClearTextureFloat(HeterogeneousVolumes.LightingCacheTexture) 128x128 1 dispatch 32x32x32 groups
       5.8% 0.42ms   RenderLightingCacheWithLiveShadingCS [In-Scattering] (Light = Untitled_1.PointLight) 1 dispatch 32x32x32 groups
       5.8% 0.42ms   RenderLightingCacheWithLiveShadingCS [In-Scattering] (Light = Untitled_1.DirectionalLight) 1 dispatch 32x32x32 groups
       5.7% 0.42ms   RenderLightingCacheWithLiveShadingCS [In-Scattering] (Light = Untitled_1.SpotLight) 1 dispatch 32x32x32 groups
      13.8% 1.01ms   RenderSingleScatteringWithLiveShadingCS (Light = none) 1 dispatch 240x135 groups

diff --git a/Engine/Shaders/Private/HeterogeneousVolumes/HeterogeneousVolumesHardwareRayTracing.usf b/Engine/Shaders/Private/HeterogeneousVolumes/HeterogeneousVolumesHardwareRayTracing.usf
--- a/Engine/Shaders/Private/HeterogeneousVolumes/HeterogeneousVolumesHardwareRayTracing.usf
+++ b/Engine/Shaders/Private/HeterogeneousVolumes/HeterogeneousVolumesHardwareRayTracing.usf
@@ -372,8 +372,10 @@ RAY_TRACING_ENTRY_RAYGEN(SparseVoxelsRayGenerationShader1)
 	}
 }
 
+// Lighting
 int bApplyEmissionAndTransmittance;
 int bApplyDirectLighting;
+int bApplyShadowTransmittance;
 int LightType;
 
 float3 ComputeTransmittanceHardwareRayTracing(
@@ -395,24 +397,24 @@ float3 ComputeTransmittanceHardwareRayTracing(
 
 uint3 DispatchRaysToVoxelIndex(uint2 DispatchRaysIndex)
 {
-	uint3 VolumeResolution = GetTransmittanceVolumeResolution();
+	uint3 VolumeResolution = GetLightingCacheResolution();
 	uint3 VoxelIndex = uint3(DispatchRaysIndex.x,
-		DispatchRaysIndex.y % TransmittanceVolumeResolution.y,
-		DispatchRaysIndex.y / TransmittanceVolumeResolution.y);
+		DispatchRaysIndex.y % LightingCacheResolution.y,
+		DispatchRaysIndex.y / LightingCacheResolution.y);
 	return VoxelIndex;
 }
 
-RWTexture3D<float3> RWTransmittanceVolumeTexture;
+RWTexture3D<float3> RWLightingCacheTexture;
 
-RAY_TRACING_ENTRY_RAYGEN(RenderTransmittanceVolumeWithPreshadingRGS)
+RAY_TRACING_ENTRY_RAYGEN(RenderLightingCacheWithPreshadingRGS)
 {
 	uint3 VoxelIndex = DispatchRaysToVoxelIndex(DispatchRaysIndex().xy);
-	if (any(VoxelIndex >= TransmittanceVolumeResolution))
+	if (any(VoxelIndex >= LightingCacheResolution))
 	{
 		return;
 	}
 
-	float3 UVW = (VoxelIndex + 0.5) / float3(TransmittanceVolumeResolution);
+	float3 UVW = (VoxelIndex + 0.5) / float3(LightingCacheResolution);
 
 	float3 LocalBoundsMin = GetLocalBoundsOrigin() - GetLocalBoundsExtent();
 	float3 LocalBoundsMax = GetLocalBoundsOrigin() + GetLocalBoundsExtent();
@@ -426,11 +428,17 @@ RAY_TRACING_ENTRY_RAYGEN(RenderTransmittanceVolumeWithPreshadingRGS)
 	float Epsilon = 1.0e-7;
 	if (all(Extinction < Epsilon))
 	{
-		RWTransmittanceVolumeTexture[VoxelIndex] = 1.0f;
+#if DIM_LIGHTING_CACHE_MODE == LIGHTING_CACHE_TRANSMITTANCE
+		RWLightingCacheTexture[VoxelIndex] = 1.0f;
+
+#elif DIM_LIGHTING_CACHE_MODE == LIGHTING_CACHE_INSCATTERING
+		RWLightingCacheTexture[VoxelIndex] = 0.0f;
+#endif // DIM_LIGHTING_CACHE_MODE
 		return;
 	}
 
 	FDeferredLightData LightData = LoadLightData(LightType);
+#if DIM_LIGHTING_CACHE_MODE == LIGHTING_CACHE_TRANSMITTANCE
 	float3 L = LightData.Direction;
 	float3 ToLight = L * GetMaxShadowTraceDistance();
 	if (LightType != LIGHT_TYPE_DIRECTIONAL)
@@ -448,44 +456,14 @@ RAY_TRACING_ENTRY_RAYGEN(RenderTransmittanceVolumeWithPreshadingRGS)
 			FCapsuleLight Capsule = GetCapsule(ToLight, LightData);
 		}
 	}
-	float3 WorldRayEnd = WorldRayOrigin + ToLight;
-	float3 LocalRayEnd = mul(float4(WorldRayEnd, 1.0), GetWorldToLocal()).xyz;
 
-	float3 WorldRayDirection = normalize(WorldRayEnd - WorldRayOrigin);
-	float3 LocalRayDirection = LocalRayEnd - LocalRayOrigin;
-
-	float LocalRayTMin = 0.0f;
-	float LocalRayTMax = length(LocalRayDirection);
-	LocalRayDirection /= LocalRayTMax;
-
-	float RayBias = 0.5;
-	float StepSize = GetShadowStepSize();
-	int bApplyEmissionAndTransmittance = 0;
-	int bApplyDirectLighting = 0;
-	int bApplyShadowTransmittance = 0;
-	FRayMarchingContext RayMarchingContext = (FRayMarchingContext)0;
-	RayMarchingContext = CreateRayMarchingContext(
-		// Local-space
-		LocalRayOrigin,
-		LocalRayDirection,
-		LocalRayTMin,
-		LocalRayTMax,
-		// World-space
-		WorldRayOrigin,
-		WorldRayDirection,
-		// Ray-step attributes
-		RayBias,
-		StepSize,
-		MaxStepCount,
-		// Lighting
-		bApplyEmissionAndTransmittance,
-		bApplyDirectLighting,
-		bApplyShadowTransmittance
-	);
-	RayMarchingContext.MipLevel = 0;
+	float3 Transmittance = ComputeTransmittance(WorldRayOrigin, ToLight, MaxStepCount);
+	RWLightingCacheTexture[VoxelIndex] = Transmittance;
 
-	float3 Transmittance = ComputeTransmittance(RayMarchingContext);
-	RWTransmittanceVolumeTexture[VoxelIndex] = Transmittance;
+#elif DIM_LIGHTING_CACHE_MODE == LIGHTING_CACHE_INSCATTERING
+	float3 Inscattering = ComputeInscattering(WorldRayOrigin, LightData, LightType, MaxStepCount, bApplyShadowTransmittance);
+	RWLightingCacheTexture[VoxelIndex] += Inscattering;
+#endif // DIM_LIGHTING_CACHE_MODE
 }
 
 RAY_TRACING_ENTRY_RAYGEN(RenderSingleScatteringWithPreshadingRGS)
diff --git a/Engine/Shaders/Private/HeterogeneousVolumes/HeterogeneousVolumesLiveShadingPipeline.usf b/Engine/Shaders/Private/HeterogeneousVolumes/HeterogeneousVolumesLiveShadingPipeline.usf
--- a/Engine/Shaders/Private/HeterogeneousVolumes/HeterogeneousVolumesLiveShadingPipeline.usf
+++ b/Engine/Shaders/Private/HeterogeneousVolumes/HeterogeneousVolumesLiveShadingPipeline.usf
@@ -168,17 +168,17 @@ float3 SampleAlbedo(inout FVolumeSampleContext Context)
 
 #include "HeterogeneousVolumesRayMarchingUtils.ush"
 
-RWTexture3D<float3> RWTransmittanceVolumeTexture;
+RWTexture3D<float3> RWLightingCacheTexture;
 
 [numthreads(THREADGROUP_SIZE_3D, THREADGROUP_SIZE_3D, THREADGROUP_SIZE_3D)]
-void RenderTransmittanceVolumeWithLiveShadingCS(
+void RenderLightingCacheWithLiveShadingCS(
 	uint3 GroupId : SV_GroupID,
 	uint3 GroupThreadId : SV_GroupThreadID,
 	uint3 DispatchThreadId : SV_DispatchThreadID
 )
 {
 	float3 VoxelIndex = DispatchThreadId;
-	if (any(DispatchThreadId > TransmittanceVolumeResolution))
+	if (any(DispatchThreadId > LightingCacheResolution))
 	{
 		return;
 	}
@@ -189,8 +189,8 @@ void RenderTransmittanceVolumeWithLiveShadingCS(
 		uint3 Rand32Bits = Rand4DPCG32(uint4(VoxelIndex, View.StateFrameIndexMod8)).xyz;
 		VoxelJitter = float3(Rand32Bits) / float(uint(0xffffffff));
 	}
-	//float3 UVW = (VoxelIndex + VoxelJitter) / float3(GetTransmittanceVolumeResolution());
-	float3 UVW = (VoxelIndex + 0.5) / float3(GetTransmittanceVolumeResolution());
+	//float3 UVW = (VoxelIndex + VoxelJitter) / float3(GetLightingCacheResolution());
+	float3 UVW = (VoxelIndex + 0.5) / float3(GetLightingCacheResolution());
 
 	float3 LocalBoundsMin = GetLocalBoundsOrigin() - GetLocalBoundsExtent();
 	float3 LocalBoundsMax = GetLocalBoundsOrigin() + GetLocalBoundsExtent();
@@ -204,11 +204,17 @@ void RenderTransmittanceVolumeWithLiveShadingCS(
 	float Epsilon = 1.0e-7;
 	if (all(Extinction < Epsilon))
 	{
-		RWTransmittanceVolumeTexture[VoxelIndex] = 1.0f;
+#if DIM_LIGHTING_CACHE_MODE == LIGHTING_CACHE_TRANSMITTANCE
+		RWLightingCacheTexture[VoxelIndex] = 1.0f;
+
+#elif DIM_LIGHTING_CACHE_MODE == LIGHTING_CACHE_INSCATTERING
+		RWLightingCacheTexture[VoxelIndex] = 0.0f;
+#endif // DIM_LIGHTING_CACHE_MODE
 		return;
 	}
 
 	FDeferredLightData LightData = LoadLightData(LightType);
+#if DIM_LIGHTING_CACHE_MODE == LIGHTING_CACHE_TRANSMITTANCE
 	float3 L = LightData.Direction;
 	float3 ToLight = L * GetMaxShadowTraceDistance();
 	if (LightType != LIGHT_TYPE_DIRECTIONAL)
@@ -226,43 +232,14 @@ void RenderTransmittanceVolumeWithLiveShadingCS(
 			FCapsuleLight Capsule = GetCapsule(ToLight, LightData);
 		}
 	}
-	float3 WorldRayEnd = WorldRayOrigin + ToLight;
-	float3 LocalRayEnd = mul(float4(WorldRayEnd, 1.0), GetWorldToLocal()).xyz;
 
-	float3 WorldRayDirection = normalize(WorldRayEnd - WorldRayOrigin);
-	float3 LocalRayDirection = LocalRayEnd - LocalRayOrigin;
+	float3 Transmittance = ComputeTransmittance(WorldRayOrigin, ToLight, MaxStepCount);
+	RWLightingCacheTexture[VoxelIndex] = Transmittance;
 
-	float LocalRayTMin = 0.0f;
-	float LocalRayTMax = length(LocalRayDirection);
-	LocalRayDirection /= LocalRayTMax;
-
-	float RayBias = 0.5;
-	float StepSize = GetShadowStepSize();
-	int bApplyEmissionAndTransmittance = 0;
-	int bApplyDirectLighting = 0;
-	int bApplyShadowTransmittance = 0;
-	FRayMarchingContext RayMarchingContext = (FRayMarchingContext)0;
-	RayMarchingContext = CreateRayMarchingContext(
-		// Local-space
-		LocalRayOrigin,
-		LocalRayDirection,
-		LocalRayTMin,
-		LocalRayTMax,
-		// World-space
-		WorldRayOrigin,
-		WorldRayDirection,
-		// Ray-step attributes
-		RayBias,
-		StepSize,
-		MaxStepCount,
-		// Lighting
-		bApplyEmissionAndTransmittance,
-		bApplyDirectLighting,
-		bApplyShadowTransmittance
-	);
-
-	float3 Transmittance = ComputeTransmittance(RayMarchingContext);
-	RWTransmittanceVolumeTexture[VoxelIndex] = Transmittance;
+#elif DIM_LIGHTING_CACHE_MODE == LIGHTING_CACHE_INSCATTERING
+	float3 Inscattering = ComputeInscattering(WorldRayOrigin, LightData, LightType, MaxStepCount, bApplyShadowTransmittance);
+	RWLightingCacheTexture[VoxelIndex] += Inscattering;
+#endif // DIM_LIGHTING_CACHE_MODE
 }
 
 [numthreads(THREADGROUP_SIZE_2D, THREADGROUP_SIZE_2D, 1)]
diff --git a/Engine/Shaders/Private/HeterogeneousVolumes/HeterogeneousVolumesPreshadingPipeline.usf b/Engine/Shaders/Private/HeterogeneousVolumes/HeterogeneousVolumesPreshadingPipeline.usf
--- a/Engine/Shaders/Private/HeterogeneousVolumes/HeterogeneousVolumesPreshadingPipeline.usf
+++ b/Engine/Shaders/Private/HeterogeneousVolumes/HeterogeneousVolumesPreshadingPipeline.usf
@@ -43,6 +43,7 @@
 // Lighting
 int bApplyEmissionAndTransmittance;
 int bApplyDirectLighting;
+int bApplyShadowTransmittance;
 int LightType;
 
 // Ray data
@@ -411,17 +412,17 @@ FTraceVoxelResult TraceVoxels(float3 LocalRayOrigin, float3 LocalRayDirection, f
 	return TraceVoxelResult;
 }
 
-RWTexture3D<float3> RWTransmittanceVolumeTexture;
+RWTexture3D<float3> RWLightingCacheTexture;
 
 [numthreads(THREADGROUP_SIZE_3D, THREADGROUP_SIZE_3D, THREADGROUP_SIZE_3D)]
-void RenderTransmittanceVolumeWithPreshadingCS(
+void RenderLightingCacheWithPreshadingCS(
 	uint3 GroupId : SV_GroupID,
 	uint3 GroupThreadId : SV_GroupThreadID,
 	uint3 DispatchThreadId : SV_DispatchThreadID
 )
 {
 	float3 VoxelIndex = DispatchThreadId;
-	if (any(VoxelIndex >= TransmittanceVolumeResolution))
+	if (any(VoxelIndex >= LightingCacheResolution))
 	{
 		return;
 	}
@@ -432,8 +433,8 @@ void RenderTransmittanceVolumeWithPreshadingCS(
 		uint3 Rand32Bits = Rand4DPCG32(uint4(VoxelIndex, View.StateFrameIndexMod8)).xyz;
 		VoxelJitter = float3(Rand32Bits) / float(uint(0xffffffff));
 	}
-	//float3 UVW = (VoxelIndex + VectorJitter) / float3(GetTransmittanceVolumeResolution());
-	float3 UVW = (VoxelIndex + 0.5) / float3(GetTransmittanceVolumeResolution());
+	//float3 UVW = (VoxelIndex + VectorJitter) / float3(GetLightingCacheResolution());
+	float3 UVW = (VoxelIndex + 0.5) / float3(GetLightingCacheResolution());
 
 	float3 LocalBoundsMin = GetLocalBoundsOrigin() - GetLocalBoundsExtent();
 	float3 LocalBoundsMax = GetLocalBoundsOrigin() + GetLocalBoundsExtent();
@@ -447,12 +448,18 @@ void RenderTransmittanceVolumeWithPreshadingCS(
 	float Epsilon = 1.0e-7;
 	if (all(Extinction < Epsilon))
 	{
-		RWTransmittanceVolumeTexture[VoxelIndex] = 1.0f;
+#if DIM_LIGHTING_CACHE_MODE == LIGHTING_CACHE_TRANSMITTANCE
+		RWLightingCacheTexture[VoxelIndex] = 1.0f;
+
+#elif DIM_LIGHTING_CACHE_MODE == LIGHTING_CACHE_INSCATTERING
+		RWLightingCacheTexture[VoxelIndex] = 0.0f;
+#endif // DIM_LIGHTING_CACHE_MODE
 		return;
 	}
 
 	// Build shadow ray
 	FDeferredLightData LightData = LoadLightData(LightType);
+#if DIM_LIGHTING_CACHE_MODE == LIGHTING_CACHE_TRANSMITTANCE
 	float3 L = LightData.Direction;
 	float3 ToLight = L * GetMaxShadowTraceDistance();
 	if (LightType != LIGHT_TYPE_DIRECTIONAL)
@@ -470,44 +477,14 @@ void RenderTransmittanceVolumeWithPreshadingCS(
 			FCapsuleLight Capsule = GetCapsule(ToLight, LightData);
 		}
 	}
-	float3 WorldRayEnd = WorldRayOrigin + ToLight;
-	float3 LocalRayEnd = mul(float4(WorldRayEnd, 1.0), GetWorldToLocal()).xyz;
 
-	float3 WorldRayDirection = normalize(WorldRayEnd - WorldRayOrigin);
-	float3 LocalRayDirection = LocalRayEnd - LocalRayOrigin;
+	float3 Transmittance = ComputeTransmittance(WorldRayOrigin, ToLight, MaxStepCount);
+	RWLightingCacheTexture[VoxelIndex] = Transmittance;
 
-	float LocalRayTMin = 0.0f;
-	float LocalRayTMax = length(LocalRayDirection);
-	LocalRayDirection /= LocalRayTMax;
-
-	// Build ray-marching context
-	float RayJitter = 0.5;
-	float StepSize = GetShadowStepSize();
-	int bApplyEmissionAndTransmittance = 0;
-	int bApplyDirectLighting = 0;
-	int bApplyShadowTransmittance = 0;
-	FRayMarchingContext RayMarchingContext = (FRayMarchingContext) 0;
-	RayMarchingContext = CreateRayMarchingContext(
-		// Local-space
-		LocalRayOrigin,
-		LocalRayDirection,
-		LocalRayTMin,
-		LocalRayTMax,
-		// World-space
-		WorldRayOrigin,
-		WorldRayDirection,
-		// Ray-step attributes
-		RayJitter,
-		StepSize,
-		MaxStepCount,
-		// Lighting
-		bApplyEmissionAndTransmittance,
-		bApplyDirectLighting,
-		bApplyShadowTransmittance
-	);
-
-	float3 Transmittance = ComputeTransmittance(RayMarchingContext);
-	RWTransmittanceVolumeTexture[VoxelIndex] = Transmittance;
+#elif DIM_LIGHTING_CACHE_MODE == LIGHTING_CACHE_INSCATTERING
+	float3 Inscattering = ComputeInscattering(WorldRayOrigin, LightData, LightType, MaxStepCount, bApplyShadowTransmittance);
+	RWLightingCacheTexture[VoxelIndex] += Inscattering;
+#endif // DIM_LIGHTING_CACHE_MODE
 }
 
 RWStructuredBuffer<FVoxelDataPacked> RWVoxelOutputBuffer;
diff --git a/Engine/Shaders/Private/HeterogeneousVolumes/HeterogeneousVolumesRayMarchingUtils.ush b/Engine/Shaders/Private/HeterogeneousVolumes/HeterogeneousVolumesRayMarchingUtils.ush
--- a/Engine/Shaders/Private/HeterogeneousVolumes/HeterogeneousVolumesRayMarchingUtils.ush
+++ b/Engine/Shaders/Private/HeterogeneousVolumes/HeterogeneousVolumesRayMarchingUtils.ush
@@ -98,6 +98,42 @@ uint CalcStepCount(inout FRayMarchingContext Context)
 	return StepCount;
 }
 
+float ComputeHardSurfaceShadowFactor(
+	float3 TranslatedWorldPosition,
+	inout FDeferredLightData LightData,
+	uint LightType
+)
+{
+	float HardSurfaceShadowFactor = 1.0;
+
+#if HARD_SURFACE_SHADOWING
+	// Evaluate hard-surface shadow term
+	if (LightType == LIGHT_TYPE_DIRECTIONAL)
+	{
+		float SceneDepth = dot(TranslatedWorldPosition - PrimaryView.TranslatedWorldCameraOrigin, View.ViewForward);
+		bool bShadowingFromValidUVArea = false;
+		float ShadowFactor = ComputeDirectionalLightDynamicShadowing(TranslatedWorldPosition, SceneDepth, bShadowingFromValidUVArea);
+
+		if (bShadowingFromValidUVArea)
+		{
+			HardSurfaceShadowFactor *= ShadowFactor;
+		}
+	}
+	else // Local lights
+	{
+		bool bShadowingFromValidUVArea = false;
+		float ShadowFactor = ComputeVolumeShadowing(TranslatedWorldPosition, LightData.bRadialLight && !LightData.bSpotLight, LightData.bSpotLight, bShadowingFromValidUVArea);
+
+		if (bShadowingFromValidUVArea)
+		{
+			HardSurfaceShadowFactor *= ShadowFactor;
+		}
+	}
+#endif // HARD_SURFACE_SHADOWING
+
+	return HardSurfaceShadowFactor;
+}
+
 void RayMarchTransmittance(
 	inout FRayMarchingContext RayMarchingContext,
 	uint StepCount,
@@ -152,6 +188,112 @@ float3 ComputeTransmittance(
 #endif // HARDWARE_RAY_TRACING
 }
 
+float3 ComputeTransmittance(
+	float3 WorldRayOrigin,
+	float3 ToLight,
+	uint MaxStepCount
+)
+{
+#if DIM_USE_TRANSMITTANCE_VOLUME
+	float3 LocalRayOrigin = mul(float4(WorldRayOrigin, 1.0), GetWorldToLocal()).xyz;
+	float3 LocalBoundsMin = GetLocalBoundsOrigin() - GetLocalBoundsExtent();
+	float3 TransmittanceUVW = saturate((LocalRayOrigin - LocalBoundsMin) / (2.0 * GetLocalBoundsExtent()));
+	float MipLevel = 0;
+
+	float3 Transmittance = SampleLightingCache(TransmittanceUVW, MipLevel);
+
+#else // DIM_USE_TRANSMITTANCE_VOLUME
+
+	float3 WorldRayEnd = WorldRayOrigin + ToLight;
+	float3 WorldRayDirection = normalize(ToLight);
+
+	float3 LocalRayOrigin = mul(float4(WorldRayOrigin, 1.0), GetWorldToLocal()).xyz;
+	float3 LocalRayEnd = mul(float4(WorldRayEnd, 1.0), GetWorldToLocal()).xyz;
+	float3 LocalRayDirection = LocalRayEnd - LocalRayOrigin;
+	float LocalRayTMin = 0.0;
+	float LocalRayTMax = length(LocalRayDirection);
+	LocalRayDirection /= LocalRayTMax;
+
+	float ShadowBias = 0.5;
+	float ShadowStepSize = GetShadowStepSize();
+	int bApplyEmissionAndTransmittance = 0;
+	int bApplyDirectLighting = 0;
+	int bApplyShadowTransmittance = 0;
+
+	FRayMarchingContext ShadowRayMarchingContext = CreateRayMarchingContext(
+		LocalRayOrigin,
+		LocalRayDirection,
+		LocalRayTMin,
+		LocalRayTMax,
+		WorldRayOrigin,
+		WorldRayDirection,
+		ShadowBias,
+		ShadowStepSize,
+		MaxStepCount,
+		bApplyEmissionAndTransmittance,
+		bApplyDirectLighting,
+		bApplyShadowTransmittance
+		//RayMarchingContext.MaxShadowTraceDistance
+	);
+
+	float3 Transmittance = ComputeTransmittance(ShadowRayMarchingContext);
+#endif // DIM_USE_INSCATTERING_VOLUME
+
+	return Transmittance;
+}
+
+float3 ComputeInscattering(
+	float3 WorldPosition,
+	inout FDeferredLightData LightData,
+	uint LightType,
+	uint MaxStepCount,
+	bool bApplyShadowTransmittance
+)
+{
+	float3 Inscattering = 0.0;
+
+#if DIM_USE_INSCATTERING_VOLUME
+	float3 LocalShadowRayOrigin = mul(float4(WorldPosition, 1.0), GetWorldToLocal()).xyz;
+	float3 LocalBoundsMin = GetLocalBoundsOrigin() - GetLocalBoundsExtent();
+	float3 TransmittanceUVW = saturate((LocalShadowRayOrigin - LocalBoundsMin) / (2.0 * GetLocalBoundsExtent()));
+	float MipLevel = 0;
+
+	Inscattering = SampleLightingCache(TransmittanceUVW, MipLevel);
+
+#else // DIM_USE_INSCATTERING_VOLUME
+	float3 L = LightData.Direction;
+	//float3 ToLight = L * RayMarchingContext.MaxShadowTraceDistance;
+	float3 ToLight = L * 10000;
+
+	float LightAttenuation = 1.0;
+	float3 TranslatedWorldPosition = LWCHackToFloat(LWCAdd(WorldPosition, PrimaryView.PreViewTranslation));
+	if (LightType != LIGHT_TYPE_DIRECTIONAL)
+	{
+		LightAttenuation = GetLocalLightAttenuation(TranslatedWorldPosition, LightData, ToLight, L);
+		if (LightData.bRectLight)
+		{
+			FRect Rect = GetRect(ToLight, LightData);
+			LightAttenuation *= IntegrateLight(Rect);
+		}
+		else
+		{
+			FCapsuleLight Capsule = GetCapsule(ToLight, LightData);
+			Capsule.DistBiasSqr = 0;
+			LightAttenuation *= IntegrateLight(Capsule, LightData.bInverseSquared);
+		}
+	}
+
+	float HardSurfaceShadowFactor = ComputeHardSurfaceShadowFactor(TranslatedWorldPosition, LightData, LightType);
+	Inscattering = LightData.Color * LightAttenuation * HardSurfaceShadowFactor;
+	if (any(Inscattering > 0.0) && bApplyShadowTransmittance)
+	{
+		Inscattering *= ComputeTransmittance(WorldPosition, ToLight, MaxStepCount);
+	}
+#endif // DIM_USE_INSCATTERING_VOLUME
+
+	return Inscattering;
+}
+
 void RayMarchEmissionAbsorption(
 	inout FRayMarchingContext RayMarchingContext,
 	uint StepCount,
@@ -192,7 +334,7 @@ void RayMarchSingleScattering(
 
 		if (RayMarchingContext.bApplyEmissionAndTransmittance)
 		{
-			float3 Emission = SampleEmission(SampleContext);// +10;
+			float3 Emission = SampleEmission(SampleContext);
 			Radiance += Emission * RayMarchingContext.StepSize * Transmittance;
 		}
 
@@ -202,108 +344,11 @@ void RayMarchSingleScattering(
 			float3 Albedo = SampleAlbedo(SampleContext);
 			if (any(Albedo > 0.0))
 			{
-				float3 TranslatedWorldPosition = LWCHackToFloat(LWCAdd(WorldPosition, PrimaryView.PreViewTranslation));
-
-				float3 L = LightData.Direction;
-				//float3 ToLight = L * RayMarchingContext.MaxShadowTraceDistance;
-				float3 ToLight = L * 10000;
-
-				float LightAttenuation = 1.0;
-				if (LightType != LIGHT_TYPE_DIRECTIONAL)
-				{
-					LightAttenuation = GetLocalLightAttenuation(TranslatedWorldPosition, LightData, ToLight, L);
-					if (LightData.bRectLight)
-					{
-						FRect Rect = GetRect(ToLight, LightData);
-						LightAttenuation *= IntegrateLight(Rect);
-					}
-					else
-					{
-						FCapsuleLight Capsule = GetCapsule(ToLight, LightData);
-						Capsule.DistBiasSqr = 0;
-						LightAttenuation *= IntegrateLight(Capsule, LightData.bInverseSquared);
-					}
-				}
-
-				float HardSurfaceShadowFactor = 1.0;
-#if HARD_SURFACE_SHADOWING
-				// Evaluate hard-surface shadow term
-				if (LightType == LIGHT_TYPE_DIRECTIONAL)
-				{
-					float SceneDepth = dot(TranslatedWorldPosition - PrimaryView.TranslatedWorldCameraOrigin, View.ViewForward);
-					bool bShadowingFromValidUVArea = false;
-					float ShadowFactor = ComputeDirectionalLightDynamicShadowing(TranslatedWorldPosition, SceneDepth, bShadowingFromValidUVArea);
-
-					if (bShadowingFromValidUVArea)
-					{
-						HardSurfaceShadowFactor *= ShadowFactor;
-					}
-				}
-				else // Local lights
-				{
-					bool bShadowingFromValidUVArea = false;
-					float ShadowFactor = ComputeVolumeShadowing(TranslatedWorldPosition, LightData.bRadialLight && !LightData.bSpotLight, LightData.bSpotLight, bShadowingFromValidUVArea);
-
-					if (bShadowingFromValidUVArea)
-					{
-						HardSurfaceShadowFactor *= ShadowFactor;
-					}
-				}
-#endif // HARD_SURFACE_SHADOWING
+				float3 Inscattering = ComputeInscattering(WorldPosition, LightData, LightType, RayMarchingContext.MaxStepCount, RayMarchingContext.bApplyShadowTransmittance);
 
-				float3 Inscattering = LightData.Color * LightAttenuation * HardSurfaceShadowFactor;
-				if (any(Inscattering > 0.0))
-				{
-					float3 ShadowTransmittance = 1.0;
-
-					if (RayMarchingContext.bApplyShadowTransmittance)
-					{
-						float3 LocalShadowRayOrigin = mul(float4(WorldPosition, 1.0), GetWorldToLocal()).xyz;
-
-#if DIM_USE_TRANSMITTANCE_VOLUME
-						float3 LocalBoundsMin = GetLocalBoundsOrigin() - GetLocalBoundsExtent();
-						float3 TransmittanceUVW = saturate((LocalShadowRayOrigin - LocalBoundsMin) / (2.0 * GetLocalBoundsExtent()));
-						float MipLevel = 0;
-						ShadowTransmittance = SampleTransmittanceVolume(TransmittanceUVW, MipLevel);
-#else
-						// Build shadow ray and find intersection
-						float3 LocalShadowRayEnd = mul(float4(WorldPosition + ToLight, 1.0), GetWorldToLocal()).xyz;
-						float3 LocalShadowRayDirection = LocalShadowRayEnd - LocalShadowRayOrigin;
-						float LocalShadowRayTMin = 0.0;
-						float LocalShadowRayTMax = length(LocalShadowRayDirection);
-						LocalShadowRayDirection /= LocalShadowRayTMax;
-
-						float ShadowBias = 0.5;
-						float ShadowStepSize = GetShadowStepSize();
-						int bApplyEmissionAndTransmittance = 0;
-						int bApplyDirectLighting = 0;
-						int bApplyShadowTransmittance = 0;
-
-						FRayMarchingContext ShadowRayMarchingContext = CreateRayMarchingContext(
-							LocalShadowRayOrigin,
-							LocalShadowRayDirection,
-							LocalShadowRayTMin,
-							LocalShadowRayTMax,
-							WorldPosition, // WorldShadowRayOrigin
-							normalize(ToLight), // WorldShadowRayDirection
-							ShadowBias,
-							ShadowStepSize,
-							RayMarchingContext.MaxStepCount,
-							bApplyEmissionAndTransmittance,
-							bApplyDirectLighting,
-							bApplyShadowTransmittance
-							//RayMarchingContext.MaxShadowTraceDistance
-						);
-
-						ShadowTransmittance = ComputeTransmittance(ShadowRayMarchingContext);
-#endif // DIM_USE_TRANSMITTANCE_VOLUME
-					}
-
-					// Convert in-scattering to exitant radiance
-					float3 ScatteringCoefficient = Albedo * Extinction;
-					float IsotropicPhase = 1.0 / (4.0 * PI);
-					Radiance += Inscattering * ShadowTransmittance * ScatteringCoefficient * IsotropicPhase * Transmittance * RayMarchingContext.StepSize;
-				}
+				float3 ScatteringCoefficient = Albedo * Extinction;
+				float IsotropicPhase = 1.0 / (4.0 * PI);
+				Radiance += Inscattering * ScatteringCoefficient * IsotropicPhase * Transmittance * RayMarchingContext.StepSize;
 			}
 		}
 
diff --git a/Engine/Shaders/Private/HeterogeneousVolumes/HeterogeneousVolumesTransmittanceVolumeUtils.ush b/Engine/Shaders/Private/HeterogeneousVolumes/HeterogeneousVolumesTransmittanceVolumeUtils.ush
--- a/Engine/Shaders/Private/HeterogeneousVolumes/HeterogeneousVolumesTransmittanceVolumeUtils.ush
+++ b/Engine/Shaders/Private/HeterogeneousVolumes/HeterogeneousVolumesTransmittanceVolumeUtils.ush
@@ -4,18 +4,21 @@
 #define DIM_USE_TRANSMITTANCE_VOLUME 0
 #endif // DIM_USE_TRANSMITTANCE_VOLUME
 
-uint3 TransmittanceVolumeResolution;
-Texture3D TransmittanceVolumeTexture;
+#define LIGHTING_CACHE_TRANSMITTANCE 0
+#define LIGHTING_CACHE_INSCATTERING 1
 
-uint3 GetTransmittanceVolumeResolution()
+uint3 LightingCacheResolution;
+Texture3D LightingCacheTexture;
+
+uint3 GetLightingCacheResolution()
 {
-	return TransmittanceVolumeResolution;
+	return LightingCacheResolution;
 }
 
-float3 SampleTransmittanceVolume(float3 UVW, float MipLevel)
+float3 SampleLightingCache(float3 UVW, float MipLevel)
 {
 	float3 Transmittance = Texture3DSampleLevel(
-		TransmittanceVolumeTexture,
+		LightingCacheTexture,
 		View.SharedTrilinearClampedSampler,
 		UVW,
 		MipLevel).r;
diff --git a/Engine/Source/Runtime/Renderer/Private/HeterogeneousVolumes/HeterogeneousVolumes.cpp b/Engine/Source/Runtime/Renderer/Private/HeterogeneousVolumes/HeterogeneousVolumes.cpp
--- a/Engine/Source/Runtime/Renderer/Private/HeterogeneousVolumes/HeterogeneousVolumes.cpp
+++ b/Engine/Source/Runtime/Renderer/Private/HeterogeneousVolumes/HeterogeneousVolumes.cpp
@@ -136,15 +136,18 @@ static TAutoConsoleVariable<float> CVarHeterogeneousVolumesStepSize(
 	ECVF_RenderThreadSafe
 );
 
-static TAutoConsoleVariable<int32> CVarHeterogeneousVolumesTransmittanceVolume(
-	TEXT("r.HeterogeneousVolumes.TransmittanceVolume"),
-	1,
-	TEXT("Enalbes transmittance volume optimized pre-pass (Default = 1)"),
+static TAutoConsoleVariable<int32> CVarHeterogeneousVolumesLightingCache(
+	TEXT("r.HeterogeneousVolumes.LightingCache"),
+	2,
+	TEXT("Enables an optimized pre-pass, caching certain volumetric rendering lighting quantities (Default = 2)\n")
+	TEXT("0: Disabled\n")
+	TEXT("1: Cache transmittance\n")
+	TEXT("2: Cache in-scattering\n"),
 	ECVF_RenderThreadSafe
 );
 
-static TAutoConsoleVariable<int32> CVarHeterogeneousVolumesTransmittanceVolumeDownsampleFactor(
-	TEXT("r.HeterogeneousVolumes.TransmittanceVolume.DownsampleFactor"),
+static TAutoConsoleVariable<int32> CVarHeterogeneousVolumesLightingCacheDownsampleFactor(
+	TEXT("r.HeterogeneousVolumes.LightingCache.DownsampleFactor"),
 	1,
 	TEXT("Determines the downsample factor, relative to the preshading volume resolution (Default = 1)"),
 	ECVF_RenderThreadSafe
@@ -254,9 +257,19 @@ namespace HeterogeneousVolumes
 		return CVarHeterogeneousVolumesSparseVoxelPerTileCulling.GetValueOnRenderThread() != 0;
 	}
 
-	bool UseTransmittanceVolume()
+	int32 GetLightingCacheMode()
+	{
+		return CVarHeterogeneousVolumesLightingCache.GetValueOnRenderThread();
+	}
+
+	bool UseLightingCacheForInscattering()
+	{
+		return CVarHeterogeneousVolumesLightingCache.GetValueOnRenderThread() == 2;
+	}
+
+	bool UseLightingCacheForTransmittance()
 	{
-		return CVarHeterogeneousVolumesTransmittanceVolume.GetValueOnRenderThread() != 0;
+		return CVarHeterogeneousVolumesLightingCache.GetValueOnRenderThread() == 1;
 	}
 
 	bool ShouldJitter()
@@ -291,12 +304,12 @@ namespace HeterogeneousVolumes
 		);
 	}
 
-	FIntVector GetTransmittanceVolumeResolution()
+	FIntVector GetLightingCacheResolution()
 	{
-		float DownsampleFactor = CVarHeterogeneousVolumesTransmittanceVolumeDownsampleFactor.GetValueOnRenderThread();
-		FIntVector TransmittanceVolumeResolution = GetVolumeResolution() / DownsampleFactor;
+		float DownsampleFactor = CVarHeterogeneousVolumesLightingCacheDownsampleFactor.GetValueOnRenderThread();
+		FIntVector LightingCacheResolution = GetVolumeResolution() / DownsampleFactor;
 
-		return TransmittanceVolumeResolution;
+		return LightingCacheResolution;
 	}
 }
 
@@ -341,16 +354,17 @@ void FDeferredShadingSceneRenderer::RenderHeterogeneousVolumes(
 
 				// Allocate transmittance volume
 				// TODO: Allow option for scalar transmittance to conserve bandwidth
-				FIntVector TransmittanceVolumeResolution = HeterogeneousVolumes::GetTransmittanceVolumeResolution();
-				uint32 NumMips = FMath::Log2(float(FMath::Min(FMath::Min(TransmittanceVolumeResolution.X, TransmittanceVolumeResolution.Y), TransmittanceVolumeResolution.Z))) + 1;
-				FRDGTextureDesc TransmittanceVolumeDesc = FRDGTextureDesc::Create3D(
-					TransmittanceVolumeResolution,
+				FIntVector LightingCacheResolution = HeterogeneousVolumes::GetLightingCacheResolution();
+				uint32 NumMips = FMath::Log2(float(FMath::Min(FMath::Min(LightingCacheResolution.X, LightingCacheResolution.Y), LightingCacheResolution.Z))) + 1;
+				FRDGTextureDesc LightingCacheDesc = FRDGTextureDesc::Create3D(
+					LightingCacheResolution,
 					PF_FloatR11G11B10,
 					FClearValueBinding::Black,
 					TexCreate_ShaderResource | TexCreate_UAV | TexCreate_3DTiling,
 					NumMips
 				);
-				FRDGTextureRef TransmittanceVolumeTexture = GraphBuilder.CreateTexture(TransmittanceVolumeDesc, TEXT("HeterogeneousVolumes.TransmittanceVolumeTexture"));
+				FRDGTextureRef LightingCacheTexture = GraphBuilder.CreateTexture(LightingCacheDesc, TEXT("HeterogeneousVolumes.LightingCacheTexture"));
+				AddClearUAVPass(GraphBuilder, GraphBuilder.CreateUAV(LightingCacheTexture), FLinearColor::Black);
 
 				// Material baking executes a pre-shading pipeline
 				if (CVarHeterogeneousVolumesPreshading.GetValueOnRenderThread())
@@ -369,7 +383,7 @@ void FDeferredShadingSceneRenderer::RenderHeterogeneousVolumes(
 						PrimitiveId,
 						LocalBoxSphereBounds,
 						// Transmittance accleration
-						TransmittanceVolumeTexture,
+						LightingCacheTexture,
 						// Output
 						HeterogeneousVolumeRadiance
 					);
@@ -390,7 +404,7 @@ void FDeferredShadingSceneRenderer::RenderHeterogeneousVolumes(
 						PrimitiveId,
 						LocalBoxSphereBounds,
 						// Transmittance accleration
-						TransmittanceVolumeTexture,
+						LightingCacheTexture,
 						// Output
 						HeterogeneousVolumeRadiance
 					);
diff --git a/Engine/Source/Runtime/Renderer/Private/HeterogeneousVolumes/HeterogeneousVolumes.h b/Engine/Source/Runtime/Renderer/Private/HeterogeneousVolumes/HeterogeneousVolumes.h
--- a/Engine/Source/Runtime/Renderer/Private/HeterogeneousVolumes/HeterogeneousVolumes.h
+++ b/Engine/Source/Runtime/Renderer/Private/HeterogeneousVolumes/HeterogeneousVolumes.h
@@ -26,7 +26,7 @@ namespace HeterogeneousVolumes
 {
 	// CVars
 	FIntVector GetVolumeResolution();
-	FIntVector GetTransmittanceVolumeResolution();
+	FIntVector GetLightingCacheResolution();
 
 	float GetShadowStepSize();
 	float GetMaxTraceDistance();
@@ -36,6 +36,7 @@ namespace HeterogeneousVolumes
 
 	int32 GetMipLevel();
 	int32 GetDebugMode();
+	int32 GetLightingCacheMode();
 	uint32 GetSparseVoxelMipBias();
 
 	bool ShouldJitter();
@@ -43,7 +44,8 @@ namespace HeterogeneousVolumes
 	bool UseHardwareRayTracing();
 	bool UseSparseVoxelPipeline();
 	bool UseSparseVoxelPerTileCulling();
-	bool UseTransmittanceVolume();
+	bool UseLightingCacheForInscattering();
+	bool UseLightingCacheForTransmittance();
 
 	// Convenience Utils
 	int GetVoxelCount(FIntVector VolumeResolution);
@@ -72,7 +74,7 @@ BEGIN_UNIFORM_BUFFER_STRUCT(FSparseVoxelUniformBufferParameters, )
 	SHADER_PARAMETER_SAMPLER(SamplerState, TextureSampler)
 
 	// Resolution
-	SHADER_PARAMETER(FIntVector, TransmittanceVolumeResolution)
+	SHADER_PARAMETER(FIntVector, LightingCacheResolution)
 
 	// Sparse voxel data
 	SHADER_PARAMETER_RDG_BUFFER_SRV(Buffer<uint>, NumVoxelsBuffer)
@@ -86,9 +88,9 @@ BEGIN_UNIFORM_BUFFER_STRUCT(FSparseVoxelUniformBufferParameters, )
 	SHADER_PARAMETER(float, ShadowStepSize)
 END_UNIFORM_BUFFER_STRUCT()
 
-BEGIN_SHADER_PARAMETER_STRUCT(FTransmittanceVolumeParameters, )
-	SHADER_PARAMETER(FIntVector, TransmittanceVolumeResolution)
-	SHADER_PARAMETER_RDG_TEXTURE(Texture3D, TransmittanceVolumeTexture)
+BEGIN_SHADER_PARAMETER_STRUCT(FLightingCacheParameters, )
+	SHADER_PARAMETER(FIntVector, LightingCacheResolution)
+	SHADER_PARAMETER_RDG_TEXTURE(Texture3D, LightingCacheTexture)
 END_SHADER_PARAMETER_STRUCT()
 
 // Render specializations
@@ -106,7 +108,7 @@ void RenderWithLiveShading(
 	const int32 PrimitiveId,
 	const FBoxSphereBounds LocalBoxSphereBounds,
 	// Transmittance acceleration
-	FRDGTextureRef TransmittanceVolumeTexture,
+	FRDGTextureRef LightingCacheTexture,
 	// Output
 	FRDGTextureRef& HeterogeneousVolumeRadiance
 );
@@ -126,7 +128,7 @@ void RenderWithPreshading(
 	const int32 PrimitiveId,
 	const FBoxSphereBounds LocalBoxSphereBounds,
 	// Transmittance acceleration
-	FRDGTextureRef TransmittanceVolumeTexture,
+	FRDGTextureRef LightingCacheTexture,
 	// Output
 	FRDGTextureRef& HeterogeneousVolumeRadiance
 );
@@ -201,7 +203,7 @@ void GenerateRayTracingScene(
 	FRayTracingScene& RayTracingScene
 );
 
-void RenderTransmittanceVolumeWithPreshadingHardwareRayTracing(
+void RenderLightingCacheWithPreshadingHardwareRayTracing(
 	FRDGBuilder& GraphBuilder,
 	// Scene data
 	const FScene* Scene,
@@ -213,6 +215,8 @@ void RenderTransmittanceVolumeWithPreshadingHardwareRayTracing(
 	bool bApplyShadowTransmittance,
 	uint32 LightType,
 	const FLightSceneInfo* LightSceneInfo,
+	// Shadow data
+	const FVisibleLightInfo* VisibleLightInfo,
 	// Object data
 	const FPrimitiveSceneProxy* PrimitiveSceneProxy,
 	// Sparse voxel data
@@ -220,7 +224,7 @@ void RenderTransmittanceVolumeWithPreshadingHardwareRayTracing(
 	// Ray tracing data
 	FRayTracingScene& RayTracingScene,
 	// Output
-	FRDGTextureRef& TransmittanceVolumeTexture
+	FRDGTextureRef& LightingCacheTexture
 );
 
 void RenderSingleScatteringWithPreshadingHardwareRayTracing(
@@ -244,7 +248,7 @@ void RenderSingleScatteringWithPreshadingHardwareRayTracing(
 	// Ray tracing data
 	FRayTracingScene& RayTracingScene,
 	// Transmittance volume
-	FRDGTextureRef TransmittanceVolumeTexture,
+	FRDGTextureRef LightingCacheTexture,
 	// Output
 	FRDGTextureRef& HeterogeneousVolumeTexture
 );
diff --git a/Engine/Source/Runtime/Renderer/Private/HeterogeneousVolumes/HeterogeneousVolumesHardwareRayTracing.cpp b/Engine/Source/Runtime/Renderer/Private/HeterogeneousVolumes/HeterogeneousVolumesHardwareRayTracing.cpp
--- a/Engine/Source/Runtime/Renderer/Private/HeterogeneousVolumes/HeterogeneousVolumesHardwareRayTracing.cpp
+++ b/Engine/Source/Runtime/Renderer/Private/HeterogeneousVolumes/HeterogeneousVolumesHardwareRayTracing.cpp
@@ -320,13 +320,13 @@ class FHeterogeneousVolumesSparseVoxelMS : public FGlobalShader
 
 IMPLEMENT_GLOBAL_SHADER(FHeterogeneousVolumesSparseVoxelMS, "/Engine/Private/HeterogeneousVolumes/HeterogeneousVolumesHardwareRayTracing.usf", "SparseVoxelsMissShader", SF_RayMiss);
 
-class FRenderTransmittanceVolumeWithPreshadingRGS : public FGlobalShader
+class FRenderLightingCacheWithPreshadingRGS : public FGlobalShader
 {
-	DECLARE_GLOBAL_SHADER(FRenderTransmittanceVolumeWithPreshadingRGS)
-	SHADER_USE_ROOT_PARAMETER_STRUCT(FRenderTransmittanceVolumeWithPreshadingRGS, FGlobalShader)
+	DECLARE_GLOBAL_SHADER(FRenderLightingCacheWithPreshadingRGS)
+	SHADER_USE_ROOT_PARAMETER_STRUCT(FRenderLightingCacheWithPreshadingRGS, FGlobalShader)
 
-	class FApplyShadowTransmittanceDim : SHADER_PERMUTATION_BOOL("DIM_APPLY_SHADOW_TRANSMITTANCE");
-	using FPermutationDomain = TShaderPermutationDomain<FApplyShadowTransmittanceDim>;
+	class FLightingCacheMode : SHADER_PERMUTATION_INT("DIM_LIGHTING_CACHE_MODE", 2);
+	using FPermutationDomain = TShaderPermutationDomain<FLightingCacheMode>;
 
 	BEGIN_SHADER_PARAMETER_STRUCT(FParameters, )
 		// Scene 
@@ -335,9 +335,16 @@ class FRenderTransmittanceVolumeWithPreshadingRGS : public FGlobalShader
 		SHADER_PARAMETER_STRUCT_INCLUDE(FSceneTextureParameters, SceneTextures)
 
 		// Lighting data
+		SHADER_PARAMETER(int, bApplyEmissionAndTransmittance)
+		SHADER_PARAMETER(int, bApplyDirectLighting)
+		SHADER_PARAMETER(int, bApplyShadowTransmittance)
 		SHADER_PARAMETER(int, LightType)
 		SHADER_PARAMETER_STRUCT_REF(FDeferredLightUniformStruct, DeferredLight)
 
+		// Shadow data
+		SHADER_PARAMETER_RDG_UNIFORM_BUFFER(FForwardLightData, ForwardLightData)
+		SHADER_PARAMETER_STRUCT_INCLUDE(FVolumeShadowingShaderParameters, VolumeShadowingShaderParameters)
+
 		// Sparse Volume
 		SHADER_PARAMETER_RDG_UNIFORM_BUFFER(FSparseVoxelUniformBufferParameters, SparseVoxelUniformBuffer)
 
@@ -350,10 +357,10 @@ class FRenderTransmittanceVolumeWithPreshadingRGS : public FGlobalShader
 		SHADER_PARAMETER(int, bJitter)
 
 		// Transmittance volume data
-		SHADER_PARAMETER_STRUCT_INCLUDE(FTransmittanceVolumeParameters, TransmittanceVolume)
+		SHADER_PARAMETER_STRUCT_INCLUDE(FLightingCacheParameters, LightingCache)
 
 		// Output
-		SHADER_PARAMETER_RDG_TEXTURE_UAV(RWTexture2D<float4>, RWTransmittanceVolumeTexture)
+		SHADER_PARAMETER_RDG_TEXTURE_UAV(RWTexture2D<float4>, RWLightingCacheTexture)
 	END_SHADER_PARAMETER_STRUCT()
 
 	static bool ShouldCompilePermutation(const FGlobalShaderPermutationParameters& Parameters)
@@ -363,7 +370,7 @@ class FRenderTransmittanceVolumeWithPreshadingRGS : public FGlobalShader
 	}
 };
 
-IMPLEMENT_GLOBAL_SHADER(FRenderTransmittanceVolumeWithPreshadingRGS, "/Engine/Private/HeterogeneousVolumes/HeterogeneousVolumesHardwareRayTracing.usf", "RenderTransmittanceVolumeWithPreshadingRGS", SF_RayGen);
+IMPLEMENT_GLOBAL_SHADER(FRenderLightingCacheWithPreshadingRGS, "/Engine/Private/HeterogeneousVolumes/HeterogeneousVolumesHardwareRayTracing.usf", "RenderLightingCacheWithPreshadingRGS", SF_RayGen);
 
 class FRenderSingleScatteringWithPreshadingRGS : public FGlobalShader
 {
@@ -372,7 +379,8 @@ class FRenderSingleScatteringWithPreshadingRGS : public FGlobalShader
 
 	class FApplyShadowTransmittanceDim : SHADER_PERMUTATION_BOOL("DIM_APPLY_SHADOW_TRANSMITTANCE");
 	class FUseTransmittanceVolume : SHADER_PERMUTATION_BOOL("DIM_USE_TRANSMITTANCE_VOLUME");
-	using FPermutationDomain = TShaderPermutationDomain<FApplyShadowTransmittanceDim, FUseTransmittanceVolume>;
+	class FUseInscatteringVolume : SHADER_PERMUTATION_BOOL("DIM_USE_INSCATTERING_VOLUME");
+	using FPermutationDomain = TShaderPermutationDomain<FApplyShadowTransmittanceDim, FUseTransmittanceVolume, FUseInscatteringVolume>;
 
 	BEGIN_SHADER_PARAMETER_STRUCT(FParameters, )
 		// Scene 
@@ -397,7 +405,7 @@ class FRenderSingleScatteringWithPreshadingRGS : public FGlobalShader
 		SHADER_PARAMETER(int, MipLevel)
 
 		// Transmittance volume
-		SHADER_PARAMETER_STRUCT_INCLUDE(FTransmittanceVolumeParameters, TransmittanceVolume)
+		SHADER_PARAMETER_STRUCT_INCLUDE(FLightingCacheParameters, LightingCache)
 
 		// Ray
 		SHADER_PARAMETER(float, MaxTraceDistance)
@@ -500,7 +508,7 @@ FRayTracingPipelineState* BuildRayTracingPipelineState(
 	return RayTracingPipelineState;
 }
 
-void RenderTransmittanceVolumeWithPreshadingHardwareRayTracing(
+void RenderLightingCacheWithPreshadingHardwareRayTracing(
 	FRDGBuilder& GraphBuilder,
 	// Scene data
 	const FScene* Scene,
@@ -512,6 +520,8 @@ void RenderTransmittanceVolumeWithPreshadingHardwareRayTracing(
 	bool bApplyShadowTransmittance,
 	uint32 LightType,
 	const FLightSceneInfo* LightSceneInfo,
+	// Shadow data
+	const FVisibleLightInfo* VisibleLightInfo,
 	// Object data
 	const FPrimitiveSceneProxy* PrimitiveSceneProxy,
 	// Sparse voxel data
@@ -519,10 +529,10 @@ void RenderTransmittanceVolumeWithPreshadingHardwareRayTracing(
 	// Ray tracing data
 	FRayTracingScene& RayTracingScene,
 	// Output
-	FRDGTextureRef& TransmittanceVolumeTexture
+	FRDGTextureRef& LightingCacheTexture
 )
 {
-	FRenderTransmittanceVolumeWithPreshadingRGS::FParameters* PassParameters = GraphBuilder.AllocParameters<FRenderTransmittanceVolumeWithPreshadingRGS::FParameters>();
+	FRenderLightingCacheWithPreshadingRGS::FParameters* PassParameters = GraphBuilder.AllocParameters<FRenderLightingCacheWithPreshadingRGS::FParameters>();
 	{
 		// Scene
 		PassParameters->TLAS = RayTracingScene.GetLayerSRVChecked(ERayTracingSceneLayer::Base);
@@ -531,6 +541,9 @@ void RenderTransmittanceVolumeWithPreshadingHardwareRayTracing(
 
 		// Light data
 		check(LightSceneInfo != nullptr);
+		PassParameters->bApplyEmissionAndTransmittance = bApplyEmissionAndTransmittance;
+		PassParameters->bApplyDirectLighting = bApplyDirectLighting;
+		PassParameters->bApplyShadowTransmittance = bApplyShadowTransmittance;
 		FDeferredLightUniformStruct DeferredLightUniform = GetDeferredLightParameters(View, *LightSceneInfo);
 		PassParameters->DeferredLight = CreateUniformBufferImmediate(DeferredLightUniform, UniformBuffer_SingleDraw);
 		PassParameters->LightType = LightType;
@@ -539,8 +552,8 @@ void RenderTransmittanceVolumeWithPreshadingHardwareRayTracing(
 		PassParameters->SparseVoxelUniformBuffer = SparseVoxelUniformBuffer;
 
 		// Transmittance volume
-		PassParameters->TransmittanceVolume.TransmittanceVolumeResolution = HeterogeneousVolumes::GetTransmittanceVolumeResolution();
-		//PassParameters->TransmittanceVolume.TransmittanceVolumeTexture = GraphBuilder.CreateSRV(TransmittanceVolumeTexture);
+		PassParameters->LightingCache.LightingCacheResolution = HeterogeneousVolumes::GetLightingCacheResolution();
+		//PassParameters->LightingCache.LightingCacheTexture = GraphBuilder.CreateSRV(LightingCacheTexture);
 
 		// Ray data
 		PassParameters->MaxTraceDistance = HeterogeneousVolumes::GetMaxTraceDistance();
@@ -548,23 +561,52 @@ void RenderTransmittanceVolumeWithPreshadingHardwareRayTracing(
 		PassParameters->bJitter = HeterogeneousVolumes::ShouldJitter();
 		PassParameters->MipLevel = HeterogeneousVolumes::GetMipLevel();
 
+		// Shadow data
+		PassParameters->ForwardLightData = View.ForwardLightingResources.ForwardLightUniformBuffer;
+		if (VisibleLightInfo != nullptr)
+		{
+			const FProjectedShadowInfo* ProjectedShadowInfo = GetShadowForInjectionIntoVolumetricFog(*VisibleLightInfo);
+			bool bDynamicallyShadowed = ProjectedShadowInfo != NULL;
+			if (bDynamicallyShadowed)
+			{
+				GetVolumeShadowingShaderParameters(GraphBuilder, View, LightSceneInfo, ProjectedShadowInfo, PassParameters->VolumeShadowingShaderParameters);
+			}
+			else
+			{
+				SetVolumeShadowingDefaultShaderParametersGlobal(GraphBuilder, PassParameters->VolumeShadowingShaderParameters);
+			}
+		}
+		else
+		{
+			SetVolumeShadowingDefaultShaderParametersGlobal(GraphBuilder, PassParameters->VolumeShadowingShaderParameters);
+		}
+
 		// Output
-		PassParameters->RWTransmittanceVolumeTexture = GraphBuilder.CreateUAV(TransmittanceVolumeTexture);
+		PassParameters->RWLightingCacheTexture = GraphBuilder.CreateUAV(LightingCacheTexture);
 	}
 
-	FRenderTransmittanceVolumeWithPreshadingRGS::FPermutationDomain PermutationVector;
-	TShaderRef<FRenderTransmittanceVolumeWithPreshadingRGS> RayGenerationShader = View.ShaderMap->GetShader<FRenderTransmittanceVolumeWithPreshadingRGS>(PermutationVector);
-	FIntVector VolumeResolution = HeterogeneousVolumes::GetTransmittanceVolumeResolution();
-	FIntPoint DispatchResolution = FIntPoint(VolumeResolution.X, VolumeResolution.Y * VolumeResolution.Z);
-
-	FString LightName = TEXT("none");
-	if (LightSceneInfo != nullptr)
+	FString PassName;
+#if WANTS_DRAW_MESH_EVENTS
+	if (GetEmitDrawEvents())
 	{
-		FSceneRenderer::GetLightNameForDrawEvent(LightSceneInfo->Proxy, LightName);
+		FString LightName = TEXT("none");
+		if (LightSceneInfo != nullptr)
+		{
+			FSceneRenderer::GetLightNameForDrawEvent(LightSceneInfo->Proxy, LightName);
+		}
+		FString ModeName = HeterogeneousVolumes::UseLightingCacheForInscattering() ? TEXT("In-Scattering") : TEXT("Transmittance");
+		PassName = FString::Printf(TEXT("RenderLightingCacheWithPreshadingRGS [%s] (Light = %s)"), *ModeName, *LightName);
 	}
+#endif // WANTS_DRAW_MESH_EVENTS
+
+	FRenderLightingCacheWithPreshadingRGS::FPermutationDomain PermutationVector;
+	PermutationVector.Set<FRenderLightingCacheWithPreshadingRGS::FLightingCacheMode>(HeterogeneousVolumes::GetLightingCacheMode() - 1);
+	TShaderRef<FRenderLightingCacheWithPreshadingRGS> RayGenerationShader = View.ShaderMap->GetShader<FRenderLightingCacheWithPreshadingRGS>(PermutationVector);
+	FIntVector VolumeResolution = HeterogeneousVolumes::GetLightingCacheResolution();
+	FIntPoint DispatchResolution = FIntPoint(VolumeResolution.X, VolumeResolution.Y * VolumeResolution.Z);
 
 	GraphBuilder.AddPass(
-		RDG_EVENT_NAME("RenderTransmittanceVolumeWithPreshadingRGS (Light = %s) %ux%u", *LightName, DispatchResolution.X, DispatchResolution.Y),
+		RDG_EVENT_NAME("%s %ux%u", *PassName, DispatchResolution.X, DispatchResolution.Y),
 		PassParameters,
 		ERDGPassFlags::Compute | ERDGPassFlags::NeverCull,
 		[
@@ -620,7 +662,7 @@ void RenderSingleScatteringWithPreshadingHardwareRayTracing(
 	// Ray tracing data
 	FRayTracingScene& RayTracingScene,
 	// Transmittance volume
-	FRDGTextureRef TransmittanceVolumeTexture,
+	FRDGTextureRef LightingCacheTexture,
 	// Output
 	FRDGTextureRef& HeterogeneousVolumeTexture
 )
@@ -663,10 +705,10 @@ void RenderSingleScatteringWithPreshadingHardwareRayTracing(
 		PassParameters->MipLevel = HeterogeneousVolumes::GetMipLevel();
 
 		// Transmittance volume
-		if (HeterogeneousVolumes::UseTransmittanceVolume() && bApplyShadowTransmittance)
+		if ((HeterogeneousVolumes::UseLightingCacheForTransmittance() && bApplyShadowTransmittance) || HeterogeneousVolumes::UseLightingCacheForInscattering())
 		{
-			PassParameters->TransmittanceVolume.TransmittanceVolumeResolution = HeterogeneousVolumes::GetTransmittanceVolumeResolution();
-			PassParameters->TransmittanceVolume.TransmittanceVolumeTexture = TransmittanceVolumeTexture;
+			PassParameters->LightingCache.LightingCacheResolution = HeterogeneousVolumes::GetLightingCacheResolution();
+			PassParameters->LightingCache.LightingCacheTexture = LightingCacheTexture;
 		}
 
 		// Ray data
@@ -680,7 +722,8 @@ void RenderSingleScatteringWithPreshadingHardwareRayTracing(
 
 	FRenderSingleScatteringWithPreshadingRGS::FPermutationDomain PermutationVector;
 	PermutationVector.Set<FRenderSingleScatteringWithPreshadingRGS::FApplyShadowTransmittanceDim>(bApplyShadowTransmittance);
-	PermutationVector.Set<FRenderSingleScatteringWithPreshadingRGS::FUseTransmittanceVolume>(HeterogeneousVolumes::UseTransmittanceVolume());
+	PermutationVector.Set<FRenderSingleScatteringWithPreshadingRGS::FUseTransmittanceVolume>(HeterogeneousVolumes::UseLightingCacheForTransmittance());
+	PermutationVector.Set<FRenderSingleScatteringWithPreshadingRGS::FUseInscatteringVolume>(HeterogeneousVolumes::UseLightingCacheForInscattering());
 	TShaderRef<FRenderSingleScatteringWithPreshadingRGS> RayGenerationShader = View.ShaderMap->GetShader<FRenderSingleScatteringWithPreshadingRGS>(PermutationVector);
 	FIntPoint DispatchResolution = View.ViewRect.Size();
 
diff --git a/Engine/Source/Runtime/Renderer/Private/HeterogeneousVolumes/HeterogeneousVolumesLiveShadingPipeline.cpp b/Engine/Source/Runtime/Renderer/Private/HeterogeneousVolumes/HeterogeneousVolumesLiveShadingPipeline.cpp
--- a/Engine/Source/Runtime/Renderer/Private/HeterogeneousVolumes/HeterogeneousVolumesLiveShadingPipeline.cpp
+++ b/Engine/Source/Runtime/Renderer/Private/HeterogeneousVolumes/HeterogeneousVolumesLiveShadingPipeline.cpp
@@ -11,9 +11,12 @@
 #include "SceneManagement.h"
 #include "VolumetricFog.h"
 
-class FRenderTransmittanceVolumeWithLiveShadingCS : public FMeshMaterialShader
+class FRenderLightingCacheWithLiveShadingCS : public FMeshMaterialShader
 {
-	DECLARE_SHADER_TYPE(FRenderTransmittanceVolumeWithLiveShadingCS, MeshMaterial);
+	DECLARE_SHADER_TYPE(FRenderLightingCacheWithLiveShadingCS, MeshMaterial);
+
+	class FLightingCacheMode : SHADER_PERMUTATION_INT("DIM_LIGHTING_CACHE_MODE", 2);
+	using FPermutationDomain = TShaderPermutationDomain<FLightingCacheMode>;
 
 	BEGIN_SHADER_PARAMETER_STRUCT(FParameters, )
 		// Scene data
@@ -29,6 +32,8 @@ class FRenderTransmittanceVolumeWithLiveShadingCS : public FMeshMaterialShader
 
 		// Shadow data
 		SHADER_PARAMETER(float, ShadowStepSize)
+		SHADER_PARAMETER_RDG_UNIFORM_BUFFER(FForwardLightData, ForwardLightData)
+		SHADER_PARAMETER_STRUCT_INCLUDE(FVolumeShadowingShaderParameters, VolumeShadowingShaderParameters)
 
 		// Object data
 		SHADER_PARAMETER(FMatrix44f, LocalToWorld)
@@ -45,15 +50,15 @@ class FRenderTransmittanceVolumeWithLiveShadingCS : public FMeshMaterialShader
 		SHADER_PARAMETER(int, bJitter)
 
 		// Volume data
-		SHADER_PARAMETER_STRUCT_INCLUDE(FTransmittanceVolumeParameters, TransmittanceVolume)
+		SHADER_PARAMETER_STRUCT_INCLUDE(FLightingCacheParameters, LightingCache)
 
 		// Output
-		SHADER_PARAMETER_RDG_TEXTURE_UAV(RWTexture3D<float>, RWTransmittanceVolumeTexture)
+		SHADER_PARAMETER_RDG_TEXTURE_UAV(RWTexture3D<float>, RWLightingCacheTexture)
 	END_SHADER_PARAMETER_STRUCT()
 
-	FRenderTransmittanceVolumeWithLiveShadingCS() = default;
+	FRenderLightingCacheWithLiveShadingCS() = default;
 
-	FRenderTransmittanceVolumeWithLiveShadingCS(
+	FRenderLightingCacheWithLiveShadingCS(
 		const FMeshMaterialShaderType::CompiledShaderInitializerType & Initializer
 	)
 		: FMeshMaterialShader(Initializer)
@@ -117,14 +122,15 @@ class FRenderTransmittanceVolumeWithLiveShadingCS : public FMeshMaterialShader
 	static int32 GetThreadGroupSize3D() { return 4; }
 };
 
-IMPLEMENT_MATERIAL_SHADER_TYPE(, FRenderTransmittanceVolumeWithLiveShadingCS, TEXT("/Engine/Private/HeterogeneousVolumes/HeterogeneousVolumesLiveShadingPipeline.usf"), TEXT("RenderTransmittanceVolumeWithLiveShadingCS"), SF_Compute);
+IMPLEMENT_MATERIAL_SHADER_TYPE(, FRenderLightingCacheWithLiveShadingCS, TEXT("/Engine/Private/HeterogeneousVolumes/HeterogeneousVolumesLiveShadingPipeline.usf"), TEXT("RenderLightingCacheWithLiveShadingCS"), SF_Compute);
 
 class FRenderSingleScatteringWithLiveShadingCS : public FMeshMaterialShader
 {
 	DECLARE_SHADER_TYPE(FRenderSingleScatteringWithLiveShadingCS, MeshMaterial);
 
 	class FUseTransmittanceVolume : SHADER_PERMUTATION_BOOL("DIM_USE_TRANSMITTANCE_VOLUME");
-	using FPermutationDomain = TShaderPermutationDomain<FUseTransmittanceVolume>;
+	class FUseInscatteringVolume : SHADER_PERMUTATION_BOOL("DIM_USE_INSCATTERING_VOLUME");
+	using FPermutationDomain = TShaderPermutationDomain<FUseTransmittanceVolume, FUseInscatteringVolume>;
 
 	BEGIN_SHADER_PARAMETER_STRUCT(FParameters, )
 		// Scene data
@@ -151,7 +157,7 @@ class FRenderSingleScatteringWithLiveShadingCS : public FMeshMaterialShader
 		SHADER_PARAMETER(int32, PrimitiveId)
 
 		// Volume data
-		SHADER_PARAMETER_STRUCT_INCLUDE(FTransmittanceVolumeParameters, TransmittanceVolume)
+		SHADER_PARAMETER_STRUCT_INCLUDE(FLightingCacheParameters, LightingCache)
 
 		// Ray data
 		SHADER_PARAMETER(float, MaxTraceDistance)
@@ -279,7 +285,7 @@ void AddComputePass(
 	);
 }
 
-void RenderTransmittanceVolumeWithLiveShading(
+void RenderLightingCacheWithLiveShading(
 	FRDGBuilder& GraphBuilder,
 	// Scene data
 	const FScene* Scene,
@@ -291,13 +297,15 @@ void RenderTransmittanceVolumeWithLiveShading(
 	bool bApplyShadowTransmittance,
 	uint32 LightType,
 	const FLightSceneInfo* LightSceneInfo,
+	// Shadow data
+	const FVisibleLightInfo* VisibleLightInfo,
 	// Object data
 	const FPrimitiveSceneProxy* PrimitiveSceneProxy,
 	const FMaterialRenderProxy* DefaultMaterialRenderProxy,
 	const int32 PrimitiveId,
 	const FBoxSphereBounds LocalBoxSphereBounds,
 	// Output
-	FRDGTextureRef TransmittanceVolumeTexture
+	FRDGTextureRef LightingCacheTexture
 )
 {
 	const FMaterialRenderProxy* MaterialRenderProxy = nullptr;
@@ -306,7 +314,7 @@ void RenderTransmittanceVolumeWithLiveShading(
 
 	check(Material.GetMaterialDomain() == MD_Volume);
 
-	FRenderTransmittanceVolumeWithLiveShadingCS::FParameters* PassParameters = GraphBuilder.AllocParameters<FRenderTransmittanceVolumeWithLiveShadingCS::FParameters>();
+	FRenderLightingCacheWithLiveShadingCS::FParameters* PassParameters = GraphBuilder.AllocParameters<FRenderLightingCacheWithLiveShadingCS::FParameters>();
 	{
 		// Scene data
 		PassParameters->View = View.ViewUniformBuffer;
@@ -314,6 +322,9 @@ void RenderTransmittanceVolumeWithLiveShading(
 
 		// Light data
 		check(LightSceneInfo != nullptr)
+		PassParameters->bApplyEmissionAndTransmittance = bApplyEmissionAndTransmittance;
+		PassParameters->bApplyDirectLighting = bApplyDirectLighting;
+		PassParameters->bApplyShadowTransmittance = bApplyShadowTransmittance;
 		FDeferredLightUniformStruct DeferredLightUniform = GetDeferredLightParameters(View, *LightSceneInfo);
 		PassParameters->DeferredLight = CreateUniformBufferImmediate(DeferredLightUniform, UniformBuffer_SingleDraw);
 		PassParameters->LightType = LightType;
@@ -327,8 +338,8 @@ void RenderTransmittanceVolumeWithLiveShading(
 		PassParameters->PrimitiveId = PrimitiveId;
 
 		// Transmittance volume
-		PassParameters->TransmittanceVolume.TransmittanceVolumeResolution = HeterogeneousVolumes::GetTransmittanceVolumeResolution();
-		//PassParameters->TransmittanceVolume.TransmittanceVolumeTexture = GraphBuilder.CreateSRV(TransmittanceVolumeTexture);
+		PassParameters->LightingCache.LightingCacheResolution = HeterogeneousVolumes::GetLightingCacheResolution();
+		//PassParameters->LightingCache.LightingCacheTexture = GraphBuilder.CreateSRV(LightingCacheTexture);
 
 		// Ray data
 		PassParameters->MaxTraceDistance = HeterogeneousVolumes::GetMaxTraceDistance();
@@ -338,24 +349,52 @@ void RenderTransmittanceVolumeWithLiveShading(
 		PassParameters->MaxStepCount = HeterogeneousVolumes::GetMaxStepCount();
 		PassParameters->bJitter = HeterogeneousVolumes::ShouldJitter();
 
+		// Shadow data
+		PassParameters->ForwardLightData = View.ForwardLightingResources.ForwardLightUniformBuffer;
+		if (VisibleLightInfo != nullptr)
+		{
+			const FProjectedShadowInfo* ProjectedShadowInfo = GetShadowForInjectionIntoVolumetricFog(*VisibleLightInfo);
+			bool bDynamicallyShadowed = ProjectedShadowInfo != NULL;
+			if (bDynamicallyShadowed)
+			{
+				GetVolumeShadowingShaderParameters(GraphBuilder, View, LightSceneInfo, ProjectedShadowInfo, PassParameters->VolumeShadowingShaderParameters);
+			}
+			else
+			{
+				SetVolumeShadowingDefaultShaderParametersGlobal(GraphBuilder, PassParameters->VolumeShadowingShaderParameters);
+			}
+		}
+		else
+		{
+			SetVolumeShadowingDefaultShaderParametersGlobal(GraphBuilder, PassParameters->VolumeShadowingShaderParameters);
+		}
+
 		// Output
-		PassParameters->RWTransmittanceVolumeTexture = GraphBuilder.CreateUAV(TransmittanceVolumeTexture);
+		PassParameters->RWLightingCacheTexture = GraphBuilder.CreateUAV(LightingCacheTexture);
 	}
 
-	FString LightName = TEXT("none");
-	if (LightSceneInfo != nullptr)
+	FString PassName;
+#if WANTS_DRAW_MESH_EVENTS
+	if (GetEmitDrawEvents())
 	{
-		FSceneRenderer::GetLightNameForDrawEvent(LightSceneInfo->Proxy, LightName);
+		FString LightName = TEXT("none");
+		if (LightSceneInfo != nullptr)
+		{
+			FSceneRenderer::GetLightNameForDrawEvent(LightSceneInfo->Proxy, LightName);
+		}
+		FString ModeName = HeterogeneousVolumes::UseLightingCacheForInscattering() ? TEXT("In-Scattering") : TEXT("Transmittance");
+		PassName = FString::Printf(TEXT("RenderLightingCacheWithLiveShadingCS [%s] (Light = %s)"), *ModeName, *LightName);
 	}
-	FString PassName = FString::Printf(TEXT("RenderTransmittanceVolumeWithLiveShadingCS (Light = %s)"), *LightName);
+#endif // WANTS_DRAW_MESH_EVENTS
 
-	FIntVector GroupCount = HeterogeneousVolumes::GetTransmittanceVolumeResolution();
-	GroupCount.X = FMath::DivideAndRoundUp(GroupCount.X, FRenderTransmittanceVolumeWithLiveShadingCS::GetThreadGroupSize3D());
-	GroupCount.Y = FMath::DivideAndRoundUp(GroupCount.Y, FRenderTransmittanceVolumeWithLiveShadingCS::GetThreadGroupSize3D());
-	GroupCount.Z = FMath::DivideAndRoundUp(GroupCount.Z, FRenderTransmittanceVolumeWithLiveShadingCS::GetThreadGroupSize3D());
+	FIntVector GroupCount = HeterogeneousVolumes::GetLightingCacheResolution();
+	GroupCount.X = FMath::DivideAndRoundUp(GroupCount.X, FRenderLightingCacheWithLiveShadingCS::GetThreadGroupSize3D());
+	GroupCount.Y = FMath::DivideAndRoundUp(GroupCount.Y, FRenderLightingCacheWithLiveShadingCS::GetThreadGroupSize3D());
+	GroupCount.Z = FMath::DivideAndRoundUp(GroupCount.Z, FRenderLightingCacheWithLiveShadingCS::GetThreadGroupSize3D());
 
-	FRenderTransmittanceVolumeWithLiveShadingCS::FPermutationDomain PermutationVector;
-	TShaderRef<FRenderTransmittanceVolumeWithLiveShadingCS> ComputeShader = Material.GetShader<FRenderTransmittanceVolumeWithLiveShadingCS>(&FLocalVertexFactory::StaticType, PermutationVector, false);
+	FRenderLightingCacheWithLiveShadingCS::FPermutationDomain PermutationVector;
+	PermutationVector.Set<FRenderLightingCacheWithLiveShadingCS::FLightingCacheMode>(HeterogeneousVolumes::GetLightingCacheMode() - 1);
+	TShaderRef<FRenderLightingCacheWithLiveShadingCS> ComputeShader = Material.GetShader<FRenderLightingCacheWithLiveShadingCS>(&FLocalVertexFactory::StaticType, PermutationVector, false);
 	if (!ComputeShader.IsNull())
 	{
 		AddComputePass(GraphBuilder, ComputeShader, PassParameters, Scene, View, MaterialRenderProxy, Material, PassName, GroupCount);
@@ -382,7 +421,7 @@ void RenderSingleScatteringWithLiveShading(
 	const int32 PrimitiveId,
 	const FBoxSphereBounds LocalBoxSphereBounds,
 	// Transmittance acceleration
-	FRDGTextureRef TransmittanceVolumeTexture,
+	FRDGTextureRef LightingCacheTexture,
 	// Output
 	FRDGTextureRef& HeterogeneousVolumeTexture
 )
@@ -457,11 +496,10 @@ void RenderSingleScatteringWithLiveShading(
 		}
 
 		// Volume data
-		if (HeterogeneousVolumes::UseTransmittanceVolume() && bApplyShadowTransmittance)
+		if ((HeterogeneousVolumes::UseLightingCacheForTransmittance() && bApplyShadowTransmittance) || HeterogeneousVolumes::UseLightingCacheForInscattering())
 		{
-			PassParameters->TransmittanceVolume.TransmittanceVolumeResolution = HeterogeneousVolumes::GetTransmittanceVolumeResolution();
-			//PassParameters->TransmittanceVolume.TransmittanceVolumeTexture = GraphBuilder.CreateSRV(TransmittanceVolumeTexture);
-			PassParameters->TransmittanceVolume.TransmittanceVolumeTexture = TransmittanceVolumeTexture;
+			PassParameters->LightingCache.LightingCacheResolution = HeterogeneousVolumes::GetLightingCacheResolution();
+			PassParameters->LightingCache.LightingCacheTexture = LightingCacheTexture;
 		}
 
 		// Dispatch data
@@ -472,15 +510,22 @@ void RenderSingleScatteringWithLiveShading(
 		//PassParameters->RWDebugOutputBuffer = GraphBuilder.CreateUAV(DebugOutputBuffer);
 	}
 
-	FString LightName = TEXT("none");
-	if (LightSceneInfo != nullptr)
+	FString PassName;
+#if WANTS_DRAW_MESH_EVENTS
+	if (GetEmitDrawEvents())
 	{
-		FSceneRenderer::GetLightNameForDrawEvent(LightSceneInfo->Proxy, LightName);
+		FString LightName = TEXT("none");
+		if (LightSceneInfo != nullptr)
+		{
+			FSceneRenderer::GetLightNameForDrawEvent(LightSceneInfo->Proxy, LightName);
+		}
+		PassName = FString::Printf(TEXT("RenderSingleScatteringWithLiveShadingCS (Light = %s)"), *LightName);
 	}
-	FString PassName = FString::Printf(TEXT("RenderSingleScatteringWithLiveShadingCS (Light = %s)"), *LightName);
+#endif // WANTS_DRAW_MESH_EVENTS
 
 	FRenderSingleScatteringWithLiveShadingCS::FPermutationDomain PermutationVector;
-	PermutationVector.Set<FRenderSingleScatteringWithLiveShadingCS::FUseTransmittanceVolume>(HeterogeneousVolumes::UseTransmittanceVolume() && PassParameters->bApplyShadowTransmittance);
+	PermutationVector.Set<FRenderSingleScatteringWithLiveShadingCS::FUseTransmittanceVolume>(HeterogeneousVolumes::UseLightingCacheForTransmittance() && PassParameters->bApplyShadowTransmittance);
+	PermutationVector.Set<FRenderSingleScatteringWithLiveShadingCS::FUseInscatteringVolume>(HeterogeneousVolumes::UseLightingCacheForInscattering());
 	TShaderRef<FRenderSingleScatteringWithLiveShadingCS> ComputeShader = Material.GetShader<FRenderSingleScatteringWithLiveShadingCS>(&FLocalVertexFactory::StaticType, PermutationVector, false);
 	if (!ComputeShader.IsNull())
 	{
@@ -488,7 +533,7 @@ void RenderSingleScatteringWithLiveShading(
 	}
 }
 
-void RenderWithLiveShading(
+void RenderWithTransmittanceVolumePipeline(
 	FRDGBuilder& GraphBuilder,
 	const FSceneTextures& SceneTextures,
 	const FScene* Scene,
@@ -501,7 +546,7 @@ void RenderWithLiveShading(
 	const int32 PrimitiveId,
 	const FBoxSphereBounds LocalBoxSphereBounds,
 	// Transmittance acceleration
-	FRDGTextureRef TransmittanceVolumeTexture,
+	FRDGTextureRef LightingCacheTexture,
 	// Output
 	FRDGTextureRef& HeterogeneousVolumeRadiance
 )
@@ -541,9 +586,9 @@ void RenderWithLiveShading(
 			}
 		}
 
-		if (HeterogeneousVolumes::UseTransmittanceVolume() && bApplyShadowTransmittance)
+		if (HeterogeneousVolumes::UseLightingCacheForTransmittance() && bApplyShadowTransmittance)
 		{
-			RenderTransmittanceVolumeWithLiveShading(
+			RenderLightingCacheWithLiveShading(
 				GraphBuilder,
 				// Scene data
 				Scene,
@@ -555,13 +600,15 @@ void RenderWithLiveShading(
 				bApplyShadowTransmittance,
 				LightType,
 				LightSceneInfo,
+				// Shadow data
+				VisibleLightInfo,
 				// Object data
 				PrimitiveSceneProxy,
 				MaterialRenderProxy,
 				PrimitiveId,
 				LocalBoxSphereBounds,
 				// Output
-				TransmittanceVolumeTexture
+				LightingCacheTexture
 			);
 		}
 
@@ -585,7 +632,181 @@ void RenderWithLiveShading(
 			PrimitiveId,
 			LocalBoxSphereBounds,
 			// Transmittance acceleration
-			TransmittanceVolumeTexture,
+			LightingCacheTexture,
+			// Output
+			HeterogeneousVolumeRadiance
+		);
+	}
+}
+
+void RenderWithInscatteringVolumePipeline(
+	FRDGBuilder& GraphBuilder,
+	const FSceneTextures& SceneTextures,
+	const FScene* Scene,
+	const FViewInfo& View,
+	// Shadow data
+	TArray<FVisibleLightInfo, SceneRenderingAllocator>& VisibleLightInfos,
+	// Object data
+	const FPrimitiveSceneProxy* PrimitiveSceneProxy,
+	const FMaterialRenderProxy* MaterialRenderProxy,
+	const int32 PrimitiveId,
+	const FBoxSphereBounds LocalBoxSphereBounds,
+	// Transmittance acceleration
+	FRDGTextureRef LightingCacheTexture,
+	// Output
+	FRDGTextureRef& HeterogeneousVolumeRadiance
+)
+{
+	// Light culling
+	TArray<FLightSceneInfoCompact, TInlineAllocator<64>> LightSceneInfoCompact;
+	for (auto LightIt = Scene->Lights.CreateConstIterator(); LightIt; ++LightIt)
+	{
+		if (LightIt->AffectsPrimitive(PrimitiveSceneProxy->GetBounds(), PrimitiveSceneProxy))
+		{
+			LightSceneInfoCompact.Add(*LightIt);
+		}
+	}
+
+	// Light loop:
+	int32 NumPasses = LightSceneInfoCompact.Num();
+	for (int32 PassIndex = 0; PassIndex < NumPasses; ++PassIndex)
+	{
+		bool bApplyEmissionAndTransmittance = PassIndex == 0;
+		bool bApplyDirectLighting = !LightSceneInfoCompact.IsEmpty();
+		bool bApplyShadowTransmittance = false;
+
+		uint32 LightType = 0;
+		FLightSceneInfo* LightSceneInfo = nullptr;
+		const FVisibleLightInfo* VisibleLightInfo = nullptr;
+		if (bApplyDirectLighting)
+		{
+			LightType = LightSceneInfoCompact[PassIndex].LightType;
+			LightSceneInfo = LightSceneInfoCompact[PassIndex].LightSceneInfo;
+			check(LightSceneInfo != nullptr);
+
+			bApplyDirectLighting = (LightSceneInfo != nullptr);
+			if (LightSceneInfo)
+			{
+				VisibleLightInfo = &VisibleLightInfos[LightSceneInfo->Id];
+				bApplyShadowTransmittance = LightSceneInfo->Proxy->CastsVolumetricShadow();
+			}
+		}
+
+		RenderLightingCacheWithLiveShading(
+			GraphBuilder,
+			// Scene data
+			Scene,
+			View,
+			SceneTextures,
+			// Light data
+			bApplyEmissionAndTransmittance,
+			bApplyDirectLighting,
+			bApplyShadowTransmittance,
+			LightType,
+			LightSceneInfo,
+			// Shadow data
+			VisibleLightInfo,
+			// Object data
+			PrimitiveSceneProxy,
+			MaterialRenderProxy,
+			PrimitiveId,
+			LocalBoxSphereBounds,
+			// Output
+			LightingCacheTexture
+		);
+	}
+
+	// Direct volume integrator
+	{
+		bool bApplyEmissionAndTransmittance = true;
+		bool bApplyDirectLighting = true;
+		bool bApplyShadowTransmittance = true;
+
+		uint32 LightType = 0;
+		FLightSceneInfo* LightSceneInfo = nullptr;
+		const FVisibleLightInfo* VisibleLightInfo = nullptr;
+
+		RenderSingleScatteringWithLiveShading(
+			GraphBuilder,
+			// Scene data
+			Scene,
+			View,
+			SceneTextures,
+			// Light data
+			bApplyEmissionAndTransmittance,
+			bApplyDirectLighting,
+			bApplyShadowTransmittance,
+			LightType,
+			LightSceneInfo,
+			// Shadow data
+			VisibleLightInfo,
+			// Object data
+			PrimitiveSceneProxy,
+			MaterialRenderProxy,
+			PrimitiveId,
+			LocalBoxSphereBounds,
+			// Transmittance acceleration
+			LightingCacheTexture,
+			// Output
+			HeterogeneousVolumeRadiance
+		);
+	}
+}
+
+void RenderWithLiveShading(
+	FRDGBuilder& GraphBuilder,
+	const FSceneTextures& SceneTextures,
+	const FScene* Scene,
+	const FViewInfo& View,
+	// Shadow data
+	TArray<FVisibleLightInfo, SceneRenderingAllocator>& VisibleLightInfos,
+	// Object data
+	const FPrimitiveSceneProxy* PrimitiveSceneProxy,
+	const FMaterialRenderProxy* MaterialRenderProxy,
+	const int32 PrimitiveId,
+	const FBoxSphereBounds LocalBoxSphereBounds,
+	// Transmittance acceleration
+	FRDGTextureRef LightingCacheTexture,
+	// Output
+	FRDGTextureRef& HeterogeneousVolumeRadiance
+)
+{
+	if (HeterogeneousVolumes::UseLightingCacheForInscattering())
+	{
+		RenderWithInscatteringVolumePipeline(
+			GraphBuilder,
+			SceneTextures,
+			Scene,
+			View,
+			// Shadow data
+			VisibleLightInfos,
+			// Object data
+			PrimitiveSceneProxy,
+			MaterialRenderProxy,
+			PrimitiveId,
+			LocalBoxSphereBounds,
+			// Transmittance acceleration
+			LightingCacheTexture,
+			// Output
+			HeterogeneousVolumeRadiance
+		);
+	}
+	else
+	{
+		RenderWithTransmittanceVolumePipeline(
+			GraphBuilder,
+			SceneTextures,
+			Scene,
+			View,
+			// Shadow data
+			VisibleLightInfos,
+			// Object data
+			PrimitiveSceneProxy,
+			MaterialRenderProxy,
+			PrimitiveId,
+			LocalBoxSphereBounds,
+			// Transmittance acceleration
+			LightingCacheTexture,
 			// Output
 			HeterogeneousVolumeRadiance
 		);
diff --git a/Engine/Source/Runtime/Renderer/Private/HeterogeneousVolumes/HeterogeneousVolumesPreshadingPipeline.cpp b/Engine/Source/Runtime/Renderer/Private/HeterogeneousVolumes/HeterogeneousVolumesPreshadingPipeline.cpp
--- a/Engine/Source/Runtime/Renderer/Private/HeterogeneousVolumes/HeterogeneousVolumesPreshadingPipeline.cpp
+++ b/Engine/Source/Runtime/Renderer/Private/HeterogeneousVolumes/HeterogeneousVolumesPreshadingPipeline.cpp
@@ -200,12 +200,13 @@ void GenerateRayMarchingTiles(
 	);
 }
 
-class FRenderTransmittanceVolumeWithPreshadingCS : public FGlobalShader
+class FRenderLightingCacheWithPreshadingCS : public FGlobalShader
 {
-	DECLARE_GLOBAL_SHADER(FRenderTransmittanceVolumeWithPreshadingCS);
-	SHADER_USE_PARAMETER_STRUCT(FRenderTransmittanceVolumeWithPreshadingCS, FGlobalShader);
+	DECLARE_GLOBAL_SHADER(FRenderLightingCacheWithPreshadingCS);
+	SHADER_USE_PARAMETER_STRUCT(FRenderLightingCacheWithPreshadingCS, FGlobalShader);
 
-	using FPermutationDomain = TShaderPermutationDomain<>;
+	class FLightingCacheMode : SHADER_PERMUTATION_INT("DIM_LIGHTING_CACHE_MODE", 2);
+	using FPermutationDomain = TShaderPermutationDomain<FLightingCacheMode>;
 
 	BEGIN_SHADER_PARAMETER_STRUCT(FParameters, )
 		// Scene data
@@ -213,12 +214,19 @@ class FRenderTransmittanceVolumeWithPreshadingCS : public FGlobalShader
 		SHADER_PARAMETER_STRUCT_INCLUDE(FSceneTextureParameters, SceneTextures)
 
 		// Light data
+		SHADER_PARAMETER(int, bApplyEmissionAndTransmittance)
+		SHADER_PARAMETER(int, bApplyDirectLighting)
+		SHADER_PARAMETER(int, bApplyShadowTransmittance)
 		SHADER_PARAMETER(int, LightType)
 		SHADER_PARAMETER_STRUCT_REF(FDeferredLightUniformStruct, DeferredLight)
 
+		// Shadow data
+		SHADER_PARAMETER_RDG_UNIFORM_BUFFER(FForwardLightData, ForwardLightData)
+		SHADER_PARAMETER_STRUCT_INCLUDE(FVolumeShadowingShaderParameters, VolumeShadowingShaderParameters)
+
 		// Volume structures
 		SHADER_PARAMETER_RDG_UNIFORM_BUFFER(FSparseVoxelUniformBufferParameters, SparseVoxelUniformBuffer)
-		SHADER_PARAMETER_STRUCT_INCLUDE(FTransmittanceVolumeParameters, TransmittanceVolume)
+		SHADER_PARAMETER_STRUCT_INCLUDE(FLightingCacheParameters, LightingCache)
 
 		// Ray data
 		SHADER_PARAMETER(float, MaxShadowTraceDistance)
@@ -228,7 +236,7 @@ class FRenderTransmittanceVolumeWithPreshadingCS : public FGlobalShader
 		SHADER_PARAMETER(int, bJitter)
 
 		// Output
-		SHADER_PARAMETER_RDG_TEXTURE_UAV(RWTexture3D<float>, RWTransmittanceVolumeTexture)
+		SHADER_PARAMETER_RDG_TEXTURE_UAV(RWTexture3D<float>, RWLightingCacheTexture)
 	END_SHADER_PARAMETER_STRUCT()
 
 	static bool ShouldCompilePermutation(
@@ -260,7 +268,7 @@ class FRenderTransmittanceVolumeWithPreshadingCS : public FGlobalShader
 	static int32 GetThreadGroupSize3D() { return 4; }
 };
 
-IMPLEMENT_MATERIAL_SHADER_TYPE(, FRenderTransmittanceVolumeWithPreshadingCS, TEXT("/Engine/Private/HeterogeneousVolumes/HeterogeneousVolumesPreshadingPipeline.usf"), TEXT("RenderTransmittanceVolumeWithPreshadingCS"), SF_Compute);
+IMPLEMENT_MATERIAL_SHADER_TYPE(, FRenderLightingCacheWithPreshadingCS, TEXT("/Engine/Private/HeterogeneousVolumes/HeterogeneousVolumesPreshadingPipeline.usf"), TEXT("RenderLightingCacheWithPreshadingCS"), SF_Compute);
 
 class FRenderSingleScatteringWithPreshadingCS : public FGlobalShader
 {
@@ -271,8 +279,9 @@ class FRenderSingleScatteringWithPreshadingCS : public FGlobalShader
 	class FVoxelCullingDim : SHADER_PERMUTATION_BOOL("DIM_VOXEL_CULLING");
 	class FSparseVoxelTracingDim : SHADER_PERMUTATION_BOOL("DIM_SPARSE_VOXEL_TRACING");
 	class FUseTransmittanceVolume : SHADER_PERMUTATION_BOOL("DIM_USE_TRANSMITTANCE_VOLUME");
+	class FUseInscatteringVolume : SHADER_PERMUTATION_BOOL("DIM_USE_INSCATTERING_VOLUME");
 	class FDebugDim : SHADER_PERMUTATION_BOOL("DIM_DEBUG");
-	using FPermutationDomain = TShaderPermutationDomain<FApplyShadowTransmittanceDim, FVoxelCullingDim, FSparseVoxelTracingDim, FUseTransmittanceVolume, FDebugDim>;
+	using FPermutationDomain = TShaderPermutationDomain<FApplyShadowTransmittanceDim, FVoxelCullingDim, FSparseVoxelTracingDim, FUseTransmittanceVolume, FUseInscatteringVolume, FDebugDim>;
 
 	BEGIN_SHADER_PARAMETER_STRUCT(FParameters, )
 		// Scene data
@@ -294,7 +303,7 @@ class FRenderSingleScatteringWithPreshadingCS : public FGlobalShader
 
 		// Volume structures
 		SHADER_PARAMETER_RDG_UNIFORM_BUFFER(FSparseVoxelUniformBufferParameters, SparseVoxelUniformBuffer)
-		SHADER_PARAMETER_STRUCT_INCLUDE(FTransmittanceVolumeParameters, TransmittanceVolume)
+		SHADER_PARAMETER_STRUCT_INCLUDE(FLightingCacheParameters, LightingCache)
 
 		// Ray data
 		SHADER_PARAMETER(float, MaxTraceDistance)
@@ -343,7 +352,7 @@ class FRenderSingleScatteringWithPreshadingCS : public FGlobalShader
 
 IMPLEMENT_MATERIAL_SHADER_TYPE(, FRenderSingleScatteringWithPreshadingCS, TEXT("/Engine/Private/HeterogeneousVolumes/HeterogeneousVolumesPreshadingPipeline.usf"), TEXT("RenderSingleScatteringWithPreshadingCS"), SF_Compute);
 
-void RenderTransmitanceVolumeWithPreshadingCompute(
+void RenderLightingCacheWithPreshadingCompute(
 	FRDGBuilder& GraphBuilder,
 	// Scene data
 	const FScene* Scene,
@@ -355,6 +364,8 @@ void RenderTransmitanceVolumeWithPreshadingCompute(
 	bool bApplyShadowTransmittance,
 	uint32 LightType,
 	const FLightSceneInfo* LightSceneInfo,
+	// Shadow data
+	const FVisibleLightInfo* VisibleLightInfo,
 	// Object data
 	const FPrimitiveSceneProxy* PrimitiveSceneProxy,
 	// Sparse voxel data
@@ -365,10 +376,10 @@ void RenderTransmitanceVolumeWithPreshadingCompute(
 	FRDGBufferRef RayMarchingTilesBuffer,
 	FRDGBufferRef VoxelsPerTileBuffer,
 	// Output
-	FRDGTextureRef& TransmittanceVolumeTexture
+	FRDGTextureRef& LightingCacheTexture
 )
 {
-	FRenderTransmittanceVolumeWithPreshadingCS::FParameters* PassParameters = GraphBuilder.AllocParameters<FRenderTransmittanceVolumeWithPreshadingCS::FParameters>();
+	FRenderLightingCacheWithPreshadingCS::FParameters* PassParameters = GraphBuilder.AllocParameters<FRenderLightingCacheWithPreshadingCS::FParameters>();
 	{
 		// Scene data
 		PassParameters->View = View.ViewUniformBuffer;
@@ -376,6 +387,9 @@ void RenderTransmitanceVolumeWithPreshadingCompute(
 
 		// Light data
 		check(LightSceneInfo != nullptr);
+		PassParameters->bApplyEmissionAndTransmittance = bApplyEmissionAndTransmittance;
+		PassParameters->bApplyDirectLighting = bApplyDirectLighting;
+		PassParameters->bApplyShadowTransmittance = bApplyShadowTransmittance;
 		FDeferredLightUniformStruct DeferredLightUniform = GetDeferredLightParameters(View, *LightSceneInfo);
 		PassParameters->DeferredLight = CreateUniformBufferImmediate(DeferredLightUniform, UniformBuffer_SingleDraw);
 		PassParameters->LightType = LightType;
@@ -384,8 +398,8 @@ void RenderTransmitanceVolumeWithPreshadingCompute(
 		PassParameters->SparseVoxelUniformBuffer = SparseVoxelUniformBuffer;
 
 		// Transmittance volume
-		PassParameters->TransmittanceVolume.TransmittanceVolumeResolution = HeterogeneousVolumes::GetTransmittanceVolumeResolution();
-		PassParameters->TransmittanceVolume.TransmittanceVolumeTexture = TransmittanceVolumeTexture;
+		PassParameters->LightingCache.LightingCacheResolution = HeterogeneousVolumes::GetLightingCacheResolution();
+		PassParameters->LightingCache.LightingCacheTexture = LightingCacheTexture;
 
 		// Ray data
 		//PassParameters->StepSize = HeterogeneousVolumes::GetStepSize();
@@ -393,27 +407,56 @@ void RenderTransmitanceVolumeWithPreshadingCompute(
 		PassParameters->bJitter = HeterogeneousVolumes::ShouldJitter();
 		PassParameters->MipLevel = HeterogeneousVolumes::GetMipLevel();
 
+		// Shadow data
+		PassParameters->ForwardLightData = View.ForwardLightingResources.ForwardLightUniformBuffer;
+		if (VisibleLightInfo != nullptr)
+		{
+			const FProjectedShadowInfo* ProjectedShadowInfo = GetShadowForInjectionIntoVolumetricFog(*VisibleLightInfo);
+			bool bDynamicallyShadowed = ProjectedShadowInfo != NULL;
+			if (bDynamicallyShadowed)
+			{
+				GetVolumeShadowingShaderParameters(GraphBuilder, View, LightSceneInfo, ProjectedShadowInfo, PassParameters->VolumeShadowingShaderParameters);
+			}
+			else
+			{
+				SetVolumeShadowingDefaultShaderParametersGlobal(GraphBuilder, PassParameters->VolumeShadowingShaderParameters);
+			}
+		}
+		else
+		{
+			SetVolumeShadowingDefaultShaderParametersGlobal(GraphBuilder, PassParameters->VolumeShadowingShaderParameters);
+		}
+
 		// Output
-		PassParameters->RWTransmittanceVolumeTexture = GraphBuilder.CreateUAV(TransmittanceVolumeTexture);
+		PassParameters->RWLightingCacheTexture = GraphBuilder.CreateUAV(LightingCacheTexture);
 	}
 
-	FString LightName = TEXT("none");
-	if (LightSceneInfo != nullptr)
+	FString PassName;
+#if WANTS_DRAW_MESH_EVENTS
+	if (GetEmitDrawEvents())
 	{
-		FSceneRenderer::GetLightNameForDrawEvent(LightSceneInfo->Proxy, LightName);
+		FString LightName = TEXT("none");
+		if (LightSceneInfo != nullptr)
+		{
+			FSceneRenderer::GetLightNameForDrawEvent(LightSceneInfo->Proxy, LightName);
+		}
+		FString ModeName = HeterogeneousVolumes::UseLightingCacheForInscattering() ? TEXT("In-Scattering") : TEXT("Transmittance");
+		PassName = FString::Printf(TEXT("RenderLightingCacheWithPreshadingCS [%s] (Light = %s)"), *ModeName, *LightName);
 	}
+#endif // WANTS_DRAW_MESH_EVENTS
 
-	FRenderTransmittanceVolumeWithPreshadingCS::FPermutationDomain PermutationVector;
-	TShaderRef<FRenderTransmittanceVolumeWithPreshadingCS> ComputeShader = View.ShaderMap->GetShader<FRenderTransmittanceVolumeWithPreshadingCS>(PermutationVector);
+	FRenderLightingCacheWithPreshadingCS::FPermutationDomain PermutationVector;
+	PermutationVector.Set<FRenderLightingCacheWithPreshadingCS::FLightingCacheMode>(HeterogeneousVolumes::GetLightingCacheMode() - 1);
+	TShaderRef<FRenderLightingCacheWithPreshadingCS> ComputeShader = View.ShaderMap->GetShader<FRenderLightingCacheWithPreshadingCS>(PermutationVector);
 
-	FIntVector GroupCount = HeterogeneousVolumes::GetTransmittanceVolumeResolution();
-	GroupCount.X = FMath::DivideAndRoundUp(GroupCount.X, FRenderTransmittanceVolumeWithPreshadingCS::GetThreadGroupSize3D());
-	GroupCount.Y = FMath::DivideAndRoundUp(GroupCount.Y, FRenderTransmittanceVolumeWithPreshadingCS::GetThreadGroupSize3D());
-	GroupCount.Z = FMath::DivideAndRoundUp(GroupCount.Z, FRenderTransmittanceVolumeWithPreshadingCS::GetThreadGroupSize3D());
+	FIntVector GroupCount = HeterogeneousVolumes::GetLightingCacheResolution();
+	GroupCount.X = FMath::DivideAndRoundUp(GroupCount.X, FRenderLightingCacheWithPreshadingCS::GetThreadGroupSize3D());
+	GroupCount.Y = FMath::DivideAndRoundUp(GroupCount.Y, FRenderLightingCacheWithPreshadingCS::GetThreadGroupSize3D());
+	GroupCount.Z = FMath::DivideAndRoundUp(GroupCount.Z, FRenderLightingCacheWithPreshadingCS::GetThreadGroupSize3D());
 
 	FComputeShaderUtils::AddPass(
 		GraphBuilder,
-		RDG_EVENT_NAME("RenderTransmittanceVolumeWithPreshadingCS (Light = %s)", *LightName),
+		RDG_EVENT_NAME("%s", *PassName),
 		ComputeShader,
 		PassParameters,
 		GroupCount);
@@ -438,7 +481,7 @@ void RenderSingleScatteringWithPreshadingCompute(
 	// Sparse voxel data
 	FRDGBufferRef NumVoxelsBuffer,
 	const TRDGUniformBufferRef<FSparseVoxelUniformBufferParameters>& SparseVoxelUniformBuffer,
-	FRDGTextureRef TransmittanceVolumeTexture,
+	FRDGTextureRef LightingCacheTexture,
 	// Ray marching tiles
 	FRDGBufferRef NumRayMarchingTilesBuffer,
 	FRDGBufferRef RayMarchingTilesBuffer,
@@ -496,10 +539,10 @@ void RenderSingleScatteringWithPreshadingCompute(
 		PassParameters->SparseVoxelUniformBuffer = SparseVoxelUniformBuffer;
 
 		// Transmittance volume
-		if (HeterogeneousVolumes::UseTransmittanceVolume() && bApplyShadowTransmittance)
+		if ((HeterogeneousVolumes::UseLightingCacheForTransmittance() && bApplyShadowTransmittance) || HeterogeneousVolumes::UseLightingCacheForInscattering())
 		{
-			PassParameters->TransmittanceVolume.TransmittanceVolumeResolution = HeterogeneousVolumes::GetTransmittanceVolumeResolution();
-			PassParameters->TransmittanceVolume.TransmittanceVolumeTexture = TransmittanceVolumeTexture;
+			PassParameters->LightingCache.LightingCacheResolution = HeterogeneousVolumes::GetLightingCacheResolution();
+			PassParameters->LightingCache.LightingCacheTexture = LightingCacheTexture;
 		}
 
 		// Ray data
@@ -530,7 +573,8 @@ void RenderSingleScatteringWithPreshadingCompute(
 	PermutationVector.Set<FRenderSingleScatteringWithPreshadingCS::FApplyShadowTransmittanceDim>(bApplyShadowTransmittance);
 	PermutationVector.Set<FRenderSingleScatteringWithPreshadingCS::FVoxelCullingDim>(HeterogeneousVolumes::UseSparseVoxelPerTileCulling());
 	PermutationVector.Set<FRenderSingleScatteringWithPreshadingCS::FSparseVoxelTracingDim>(HeterogeneousVolumes::UseSparseVoxelPipeline());
-	PermutationVector.Set<FRenderSingleScatteringWithPreshadingCS::FUseTransmittanceVolume>(HeterogeneousVolumes::UseTransmittanceVolume());
+	PermutationVector.Set<FRenderSingleScatteringWithPreshadingCS::FUseTransmittanceVolume>(HeterogeneousVolumes::UseLightingCacheForTransmittance());
+	PermutationVector.Set<FRenderSingleScatteringWithPreshadingCS::FUseInscatteringVolume>(HeterogeneousVolumes::UseLightingCacheForInscattering());
 	PermutationVector.Set<FRenderSingleScatteringWithPreshadingCS::FDebugDim>(HeterogeneousVolumes::GetDebugMode() != 0);
 	TShaderRef<FRenderSingleScatteringWithPreshadingCS> ComputeShader = View.ShaderMap->GetShader<FRenderSingleScatteringWithPreshadingCS>(PermutationVector);
 	FComputeShaderUtils::AddPass(
@@ -542,7 +586,7 @@ void RenderSingleScatteringWithPreshadingCompute(
 		0);
 }
 
-void RenderWithPreshadingCompute(
+void RenderWithInscatteringVolumePipelineWithPreshadingCompute(
 	FRDGBuilder& GraphBuilder,
 	// Scene data
 	const FSceneTextures& SceneTextures,
@@ -557,142 +601,250 @@ void RenderWithPreshadingCompute(
 	// Sparse voxel data
 	FRDGBufferRef NumVoxelsBuffer,
 	TRDGUniformBufferRef<FSparseVoxelUniformBufferParameters>& SparseVoxelUniformBuffer,
+	// Render tile data
+	FRDGBufferRef NumRayMarchingTilesBuffer,
+	FRDGBufferRef RayMarchingTilesBuffer,
+	FRDGBufferRef VoxelsPerTileBuffer,
 	// Output
-	FRDGTextureRef& TransmittanceVolumeTexture,
+	FRDGTextureRef& LightingCacheTexture,
 	FRDGTextureRef& HeterogeneousVolumeRadiance
 )
 {
-	RDG_EVENT_SCOPE(GraphBuilder, "Software Ray Tracing");
+	RDG_EVENT_SCOPE(GraphBuilder, "Direct Volume Rendering");
 
-	FRDGBufferRef NumRayMarchingTilesBuffer;
-	FRDGBufferRef RayMarchingTilesBuffer;
-	FRDGBufferRef VoxelsPerTileBuffer;
+	// Light culling
+	TArray<FLightSceneInfoCompact, TInlineAllocator<64>> LightSceneInfoCompact;
+	for (auto LightIt = Scene->Lights.CreateConstIterator(); LightIt; ++LightIt)
 	{
-		RDG_EVENT_SCOPE(GraphBuilder, "Ray Tile Generation");
-		GenerateRayMarchingTiles(
+		if (LightIt->AffectsPrimitive(PrimitiveSceneProxy->GetBounds(), PrimitiveSceneProxy))
+		{
+			LightSceneInfoCompact.Add(*LightIt);
+		}
+	}
+
+	// Light loop:
+	int32 NumPasses = LightSceneInfoCompact.Num();
+	for (int32 PassIndex = 0; PassIndex < NumPasses; ++PassIndex)
+	{
+		bool bApplyEmissionAndTransmittance = PassIndex == 0;
+		bool bApplyDirectLighting = !LightSceneInfoCompact.IsEmpty();
+		bool bApplyShadowTransmittance = false;
+
+		uint32 LightType = 0;
+		FLightSceneInfo* LightSceneInfo = nullptr;
+		const FVisibleLightInfo* VisibleLightInfo = nullptr;
+		if (bApplyDirectLighting)
+		{
+			LightType = LightSceneInfoCompact[PassIndex].LightType;
+			LightSceneInfo = LightSceneInfoCompact[PassIndex].LightSceneInfo;
+			check(LightSceneInfo != nullptr);
+
+			bApplyDirectLighting = (LightSceneInfo != nullptr);
+			if (LightSceneInfo)
+			{
+				VisibleLightInfo = &VisibleLightInfos[LightSceneInfo->Id];
+				bApplyShadowTransmittance = LightSceneInfo->Proxy->CastsVolumetricShadow();
+			}
+		}
+
+		RenderLightingCacheWithPreshadingCompute(
+			GraphBuilder,
+			// Scene data
+			Scene,
+			View,
+			SceneTextures,
+			// Light data
+			bApplyEmissionAndTransmittance,
+			bApplyDirectLighting,
+			bApplyShadowTransmittance,
+			LightType,
+			LightSceneInfo,
+			// Shadow data
+			VisibleLightInfo,
+			// Object data
+			PrimitiveSceneProxy,
+			// Volume data
+			NumVoxelsBuffer,
+			// Sparse voxel data
+			SparseVoxelUniformBuffer,
+			// Ray marching tile
+			NumRayMarchingTilesBuffer,
+			RayMarchingTilesBuffer,
+			VoxelsPerTileBuffer,
+			// Output
+			LightingCacheTexture
+		);
+	}
+
+	// Direct volume integrator
+	{
+		bool bApplyEmissionAndTransmittance = true;
+		bool bApplyDirectLighting = true;
+		bool bApplyShadowTransmittance = true;
+
+		uint32 LightType = 0;
+		FLightSceneInfo* LightSceneInfo = nullptr;
+		const FVisibleLightInfo* VisibleLightInfo = nullptr;
+
+		RenderSingleScatteringWithPreshadingCompute(
 			GraphBuilder,
 			// Scene
 			Scene,
 			View,
 			SceneTextures,
+			// Light
+			bApplyEmissionAndTransmittance,
+			bApplyDirectLighting,
+			bApplyShadowTransmittance,
+			LightType,
+			LightSceneInfo,
+			// Shadow
+			VisibleLightInfo,
 			// Object
 			PrimitiveSceneProxy,
 			// Volume data
 			NumVoxelsBuffer,
 			// Sparse voxel data
 			SparseVoxelUniformBuffer,
-			// Output
+			LightingCacheTexture,
+			// Ray marching tile
 			NumRayMarchingTilesBuffer,
 			RayMarchingTilesBuffer,
-			VoxelsPerTileBuffer
+			VoxelsPerTileBuffer,
+			// Output
+			HeterogeneousVolumeRadiance
 		);
 	}
+}
 
-	// Direct Volume Rendering
-	{
-		RDG_EVENT_SCOPE(GraphBuilder, "Direct Volume Rendering");
+void RenderWithTransmittanceVolumePipelineWithPreshadingCompute(
+	FRDGBuilder& GraphBuilder,
+	// Scene data
+	const FSceneTextures& SceneTextures,
+	const FScene* Scene,
+	const FSceneViewFamily& ViewFamily,
+	FViewInfo& View,
+	// Shadow data
+	TArray<FVisibleLightInfo, SceneRenderingAllocator>& VisibleLightInfos,
+	// Object data
+	const FPrimitiveSceneProxy* PrimitiveSceneProxy,
+	const FMaterialRenderProxy* MaterialRenderProxy,
+	// Sparse voxel data
+	FRDGBufferRef NumVoxelsBuffer,
+	TRDGUniformBufferRef<FSparseVoxelUniformBufferParameters>& SparseVoxelUniformBuffer,
+	// Render tile data
+	FRDGBufferRef NumRayMarchingTilesBuffer,
+	FRDGBufferRef RayMarchingTilesBuffer,
+	FRDGBufferRef VoxelsPerTileBuffer,
+	// Output
+	FRDGTextureRef& LightingCacheTexture,
+	FRDGTextureRef& HeterogeneousVolumeRadiance
+)
+{
+	RDG_EVENT_SCOPE(GraphBuilder, "Direct Volume Rendering");
 
-		// Light culling
-		TArray<FLightSceneInfoCompact, TInlineAllocator<64>> LightSceneInfoCompact;
-		for (auto LightIt = Scene->Lights.CreateConstIterator(); LightIt; ++LightIt)
+	// Light culling
+	TArray<FLightSceneInfoCompact, TInlineAllocator<64>> LightSceneInfoCompact;
+	for (auto LightIt = Scene->Lights.CreateConstIterator(); LightIt; ++LightIt)
+	{
+		if (LightIt->AffectsPrimitive(PrimitiveSceneProxy->GetBounds(), PrimitiveSceneProxy))
 		{
-			if (LightIt->AffectsPrimitive(PrimitiveSceneProxy->GetBounds(), PrimitiveSceneProxy))
-			{
-				LightSceneInfoCompact.Add(*LightIt);
-			}
+			LightSceneInfoCompact.Add(*LightIt);
 		}
+	}
 
-		// Single-scattering
-		int32 NumPasses = FMath::Max(LightSceneInfoCompact.Num(), 1);
-		for (int32 PassIndex = 0; PassIndex < NumPasses; ++PassIndex)
+	// Single-scattering
+	int32 NumPasses = FMath::Max(LightSceneInfoCompact.Num(), 1);
+	for (int32 PassIndex = 0; PassIndex < NumPasses; ++PassIndex)
+	{
+		bool bApplyEmissionAndTransmittance = PassIndex == 0;
+		bool bApplyDirectLighting = !LightSceneInfoCompact.IsEmpty();
+		bool bApplyShadowTransmittance = false;
+
+		uint32 LightType = 0;
+		FLightSceneInfo* LightSceneInfo = nullptr;
+		const FVisibleLightInfo* VisibleLightInfo = nullptr;
+		if (bApplyDirectLighting)
 		{
-			bool bApplyEmissionAndTransmittance = PassIndex == 0;
-			bool bApplyDirectLighting = !LightSceneInfoCompact.IsEmpty();
-			bool bApplyShadowTransmittance = false;
-
-			uint32 LightType = 0;
-			FLightSceneInfo* LightSceneInfo = nullptr;
-			const FVisibleLightInfo* VisibleLightInfo = nullptr;
-			if (bApplyDirectLighting)
-			{
-				LightType = LightSceneInfoCompact[PassIndex].LightType;
-				LightSceneInfo = LightSceneInfoCompact[PassIndex].LightSceneInfo;
-				check(LightSceneInfo != nullptr);
-
-				bApplyDirectLighting = (LightSceneInfo != nullptr);
-				if (LightSceneInfo)
-				{
-					VisibleLightInfo = &VisibleLightInfos[LightSceneInfo->Id];
-					bApplyShadowTransmittance = LightSceneInfo->Proxy->CastsVolumetricShadow();
-				}
-			}
+			LightType = LightSceneInfoCompact[PassIndex].LightType;
+			LightSceneInfo = LightSceneInfoCompact[PassIndex].LightSceneInfo;
+			check(LightSceneInfo != nullptr);
 
-			if (HeterogeneousVolumes::UseTransmittanceVolume() && bApplyShadowTransmittance)
+			bApplyDirectLighting = (LightSceneInfo != nullptr);
+			if (LightSceneInfo)
 			{
-				RenderTransmitanceVolumeWithPreshadingCompute(
-					GraphBuilder,
-					// Scene
-					Scene,
-					View,
-					SceneTextures,
-					// Light
-					bApplyEmissionAndTransmittance,
-					bApplyDirectLighting,
-					bApplyShadowTransmittance,
-					LightType,
-					LightSceneInfo,
-					// Object
-					PrimitiveSceneProxy,
-					// Volume data
-					NumVoxelsBuffer,
-					// Sparse voxel data
-					SparseVoxelUniformBuffer,
-					// Ray marching tile
-					NumRayMarchingTilesBuffer,
-					RayMarchingTilesBuffer,
-					VoxelsPerTileBuffer,
-					// Output
-					TransmittanceVolumeTexture
-				);
+				VisibleLightInfo = &VisibleLightInfos[LightSceneInfo->Id];
+				bApplyShadowTransmittance = LightSceneInfo->Proxy->CastsVolumetricShadow();
 			}
+		}
 
-			RenderSingleScatteringWithPreshadingCompute(
+		if (HeterogeneousVolumes::UseLightingCacheForTransmittance() && bApplyShadowTransmittance)
+		{
+			RenderLightingCacheWithPreshadingCompute(
 				GraphBuilder,
-				// Scene
+				// Scene data
 				Scene,
 				View,
 				SceneTextures,
-				// Light
+				// Light data
 				bApplyEmissionAndTransmittance,
 				bApplyDirectLighting,
 				bApplyShadowTransmittance,
 				LightType,
 				LightSceneInfo,
-				// Shadow
+				// Shadow data
 				VisibleLightInfo,
-				// Object
+				// Object data
 				PrimitiveSceneProxy,
 				// Volume data
 				NumVoxelsBuffer,
 				// Sparse voxel data
 				SparseVoxelUniformBuffer,
-				TransmittanceVolumeTexture,
 				// Ray marching tile
 				NumRayMarchingTilesBuffer,
 				RayMarchingTilesBuffer,
 				VoxelsPerTileBuffer,
 				// Output
-				HeterogeneousVolumeRadiance
+				LightingCacheTexture
 			);
 		}
+
+		RenderSingleScatteringWithPreshadingCompute(
+			GraphBuilder,
+			// Scene
+			Scene,
+			View,
+			SceneTextures,
+			// Light
+			bApplyEmissionAndTransmittance,
+			bApplyDirectLighting,
+			bApplyShadowTransmittance,
+			LightType,
+			LightSceneInfo,
+			// Shadow
+			VisibleLightInfo,
+			// Object
+			PrimitiveSceneProxy,
+			// Volume data
+			NumVoxelsBuffer,
+			// Sparse voxel data
+			SparseVoxelUniformBuffer,
+			LightingCacheTexture,
+			// Ray marching tile
+			NumRayMarchingTilesBuffer,
+			RayMarchingTilesBuffer,
+			VoxelsPerTileBuffer,
+			// Output
+			HeterogeneousVolumeRadiance
+		);
 	}
 }
 
-void RenderWithPreshadingHardwareRayTracing(
+void RenderWithPreshadingCompute(
 	FRDGBuilder& GraphBuilder,
 	// Scene data
 	const FSceneTextures& SceneTextures,
-	FScene* Scene,
+	const FScene* Scene,
 	const FSceneViewFamily& ViewFamily,
 	FViewInfo& View,
 	// Shadow data
@@ -702,117 +854,277 @@ void RenderWithPreshadingHardwareRayTracing(
 	const FMaterialRenderProxy* MaterialRenderProxy,
 	// Sparse voxel data
 	FRDGBufferRef NumVoxelsBuffer,
-	const TRDGUniformBufferRef<FSparseVoxelUniformBufferParameters>& SparseVoxelUniformBuffer,
-	// Transmittance acceleration
-	FRDGTextureRef TransmittanceVolumeTexture,
+	TRDGUniformBufferRef<FSparseVoxelUniformBufferParameters>& SparseVoxelUniformBuffer,
 	// Output
+	FRDGTextureRef& LightingCacheTexture,
 	FRDGTextureRef& HeterogeneousVolumeRadiance
 )
 {
-#if RHI_RAYTRACING
-	RDG_EVENT_SCOPE(GraphBuilder, "Hardware Ray Tracing");
+	RDG_EVENT_SCOPE(GraphBuilder, "Software Ray Tracing");
 
-	// WARNING: Currently works, but I'm skeptical if all RHI resources have the correct lifetime management
-	TArray<FRayTracingGeometryRHIRef> RayTracingGeometries;
-	TArray<FMatrix> RayTracingTransforms;
+	FRDGBufferRef NumRayMarchingTilesBuffer;
+	FRDGBufferRef RayMarchingTilesBuffer;
+	FRDGBufferRef VoxelsPerTileBuffer;
 	{
-		RDG_EVENT_SCOPE(GraphBuilder, "Acceleration Structure Build");
-
-		GenerateRayTracingGeometryInstance(
+		RDG_EVENT_SCOPE(GraphBuilder, "Ray Tile Generation");
+		GenerateRayMarchingTiles(
 			GraphBuilder,
 			// Scene
 			Scene,
 			View,
+			SceneTextures,
 			// Object
 			PrimitiveSceneProxy,
-			// Sparse voxel
+			// Volume data
 			NumVoxelsBuffer,
+			// Sparse voxel data
 			SparseVoxelUniformBuffer,
 			// Output
-			RayTracingGeometries,
-			RayTracingTransforms
+			NumRayMarchingTilesBuffer,
+			RayMarchingTilesBuffer,
+			VoxelsPerTileBuffer
 		);
+	}
 
-		GenerateRayTracingScene(
+	if (HeterogeneousVolumes::UseLightingCacheForInscattering())
+	{
+		RenderWithInscatteringVolumePipelineWithPreshadingCompute(
 			GraphBuilder,
-			// Scene
+			// Scene data
+			SceneTextures,
 			Scene,
+			ViewFamily,
 			View,
-			// Object
+			// Shadow data
+			VisibleLightInfos,
+			// Object data
 			PrimitiveSceneProxy,
-			// Ray tracing data
-			RayTracingGeometries,
-			RayTracingTransforms,
+			MaterialRenderProxy,
+			// Sparse voxel data
+			NumVoxelsBuffer,
+			SparseVoxelUniformBuffer,
+			// Render tile data
+			NumRayMarchingTilesBuffer,
+			RayMarchingTilesBuffer,
+			VoxelsPerTileBuffer,
 			// Output
-			Scene->HeterogeneousVolumesRayTracingScene
+			LightingCacheTexture,
+			HeterogeneousVolumeRadiance
 		);
 	}
+	else
+	{
+		RenderWithTransmittanceVolumePipelineWithPreshadingCompute(
+			GraphBuilder,
+			// Scene data
+			SceneTextures,
+			Scene,
+			ViewFamily,
+			View,
+			// Shadow data
+			VisibleLightInfos,
+			// Object data
+			PrimitiveSceneProxy,
+			MaterialRenderProxy,
+			// Sparse voxel data
+			NumVoxelsBuffer,
+			SparseVoxelUniformBuffer,
+			// Render tile data
+			NumRayMarchingTilesBuffer,
+			RayMarchingTilesBuffer,
+			VoxelsPerTileBuffer,
+			// Output
+			LightingCacheTexture,
+			HeterogeneousVolumeRadiance
+		);
+	}
+}
+
+void RenderWithInscatteringVolumePipelineWithPreshadingHardwareRayTracing(
+	FRDGBuilder& GraphBuilder,
+	// Scene data
+	const FSceneTextures& SceneTextures,
+	FScene* Scene,
+	const FSceneViewFamily& ViewFamily,
+	FViewInfo& View,
+	// Shadow data
+	TArray<FVisibleLightInfo, SceneRenderingAllocator>& VisibleLightInfos,
+	// Object data
+	const FPrimitiveSceneProxy* PrimitiveSceneProxy,
+	const FMaterialRenderProxy* MaterialRenderProxy,
+	// Sparse voxel data
+	FRDGBufferRef NumVoxelsBuffer,
+	const TRDGUniformBufferRef<FSparseVoxelUniformBufferParameters>& SparseVoxelUniformBuffer,
+	// Transmittance acceleration
+	FRDGTextureRef LightingCacheTexture,
+	// Output
+	FRDGTextureRef& HeterogeneousVolumeRadiance
+)
+{
+#if RHI_RAYTRACING
+	RDG_EVENT_SCOPE(GraphBuilder, "Direct Volume Rendering");
 
-	// Direct Volume Rendering
+	// Light culling
+	TArray<FLightSceneInfoCompact, TInlineAllocator<64>> LightSceneInfoCompact;
+	for (auto LightIt = Scene->Lights.CreateConstIterator(); LightIt; ++LightIt)
 	{
-		RDG_EVENT_SCOPE(GraphBuilder, "Direct Volume Rendering");
+		if (LightIt->AffectsPrimitive(PrimitiveSceneProxy->GetBounds(), PrimitiveSceneProxy))
+		{
+			LightSceneInfoCompact.Add(*LightIt);
+		}
+	}
 
-		// Light culling
-		TArray<FLightSceneInfoCompact, TInlineAllocator<64>> LightSceneInfoCompact;
-		for (auto LightIt = Scene->Lights.CreateConstIterator(); LightIt; ++LightIt)
+	// Single-scattering
+	int32 NumPasses = FMath::Max(LightSceneInfoCompact.Num(), 1);
+	for (int32 PassIndex = 0; PassIndex < NumPasses; ++PassIndex)
+	{
+		bool bApplyEmissionAndTransmittance = PassIndex == 0;
+		bool bApplyDirectLighting = !LightSceneInfoCompact.IsEmpty();
+		bool bApplyShadowTransmittance = false;
+
+		uint32 LightType = 0;
+		FLightSceneInfo* LightSceneInfo = nullptr;
+		const FVisibleLightInfo* VisibleLightInfo = nullptr;
+		if (bApplyDirectLighting)
 		{
-			if (LightIt->AffectsPrimitive(PrimitiveSceneProxy->GetBounds(), PrimitiveSceneProxy))
+			LightType = LightSceneInfoCompact[PassIndex].LightType;
+			LightSceneInfo = LightSceneInfoCompact[PassIndex].LightSceneInfo;
+			check(LightSceneInfo != nullptr);
+
+			bApplyDirectLighting = (LightSceneInfo != nullptr);
+			if (LightSceneInfo)
 			{
-				LightSceneInfoCompact.Add(*LightIt);
+				VisibleLightInfo = &VisibleLightInfos[LightSceneInfo->Id];
+				bApplyShadowTransmittance = LightSceneInfo->Proxy->CastsVolumetricShadow();
 			}
 		}
 
-		// Single-scattering
-		int32 NumPasses = FMath::Max(LightSceneInfoCompact.Num(), 1);
-		for (int32 PassIndex = 0; PassIndex < NumPasses; ++PassIndex)
+		RenderLightingCacheWithPreshadingHardwareRayTracing(
+			GraphBuilder,
+			// Scene data
+			Scene,
+			View,
+			SceneTextures,
+			// Light data
+			bApplyEmissionAndTransmittance,
+			bApplyDirectLighting,
+			bApplyShadowTransmittance,
+			LightType,
+			LightSceneInfo,
+			// Shadow
+			VisibleLightInfo,
+			// Object data
+			PrimitiveSceneProxy,
+			// Sparse voxel
+			SparseVoxelUniformBuffer,
+			// Ray tracing data
+			Scene->HeterogeneousVolumesRayTracingScene,
+			// Transmittance volume
+			LightingCacheTexture
+		);
+	}
+
+	// Direct volume integrator
+	{
+		bool bApplyEmissionAndTransmittance = true;
+		bool bApplyDirectLighting = true;
+		bool bApplyShadowTransmittance = true;
+
+		uint32 LightType = 0;
+		FLightSceneInfo* LightSceneInfo = nullptr;
+		const FVisibleLightInfo* VisibleLightInfo = nullptr;
+
+		RenderSingleScatteringWithPreshadingHardwareRayTracing(
+			GraphBuilder,
+			// Scene data
+			Scene,
+			View,
+			SceneTextures,
+			// Light data
+			bApplyEmissionAndTransmittance,
+			bApplyDirectLighting,
+			bApplyShadowTransmittance,
+			LightType,
+			LightSceneInfo,
+			// Shadow
+			VisibleLightInfo,
+			// Object data
+			PrimitiveSceneProxy,
+			// Sparse voxel
+			SparseVoxelUniformBuffer,
+			// Ray tracing data
+			Scene->HeterogeneousVolumesRayTracingScene,
+			// Transmittance volume
+			LightingCacheTexture,
+			// Output
+			HeterogeneousVolumeRadiance
+		);
+	}
+#endif // RHI_RAYTRACING
+}
+
+void RenderWithTransmittanceVolumePipelineWithPreshadingHardwareRayTracing(
+	FRDGBuilder& GraphBuilder,
+	// Scene data
+	const FSceneTextures& SceneTextures,
+	FScene* Scene,
+	const FSceneViewFamily& ViewFamily,
+	FViewInfo& View,
+	// Shadow data
+	TArray<FVisibleLightInfo, SceneRenderingAllocator>& VisibleLightInfos,
+	// Object data
+	const FPrimitiveSceneProxy* PrimitiveSceneProxy,
+	const FMaterialRenderProxy* MaterialRenderProxy,
+	// Sparse voxel data
+	FRDGBufferRef NumVoxelsBuffer,
+	const TRDGUniformBufferRef<FSparseVoxelUniformBufferParameters>& SparseVoxelUniformBuffer,
+	// Transmittance acceleration
+	FRDGTextureRef LightingCacheTexture,
+	// Output
+	FRDGTextureRef& HeterogeneousVolumeRadiance
+)
+{
+#if RHI_RAYTRACING
+	RDG_EVENT_SCOPE(GraphBuilder, "Direct Volume Rendering");
+
+	// Light culling
+	TArray<FLightSceneInfoCompact, TInlineAllocator<64>> LightSceneInfoCompact;
+	for (auto LightIt = Scene->Lights.CreateConstIterator(); LightIt; ++LightIt)
+	{
+		if (LightIt->AffectsPrimitive(PrimitiveSceneProxy->GetBounds(), PrimitiveSceneProxy))
 		{
-			bool bApplyEmissionAndTransmittance = PassIndex == 0;
-			bool bApplyDirectLighting = !LightSceneInfoCompact.IsEmpty();
-			bool bApplyShadowTransmittance = false;
-
-			uint32 LightType = 0;
-			FLightSceneInfo* LightSceneInfo = nullptr;
-			const FVisibleLightInfo* VisibleLightInfo = nullptr;
-			if (bApplyDirectLighting)
-			{
-				LightType = LightSceneInfoCompact[PassIndex].LightType;
-				LightSceneInfo = LightSceneInfoCompact[PassIndex].LightSceneInfo;
-				check(LightSceneInfo != nullptr);
-
-				bApplyDirectLighting = (LightSceneInfo != nullptr);
-				if (LightSceneInfo)
-				{
-					VisibleLightInfo = &VisibleLightInfos[LightSceneInfo->Id];
-					bApplyShadowTransmittance = LightSceneInfo->Proxy->CastsVolumetricShadow();
-				}
-			}
+			LightSceneInfoCompact.Add(*LightIt);
+		}
+	}
+
+	// Single-scattering
+	int32 NumPasses = FMath::Max(LightSceneInfoCompact.Num(), 1);
+	for (int32 PassIndex = 0; PassIndex < NumPasses; ++PassIndex)
+	{
+		bool bApplyEmissionAndTransmittance = PassIndex == 0;
+		bool bApplyDirectLighting = !LightSceneInfoCompact.IsEmpty();
+		bool bApplyShadowTransmittance = false;
+
+		uint32 LightType = 0;
+		FLightSceneInfo* LightSceneInfo = nullptr;
+		const FVisibleLightInfo* VisibleLightInfo = nullptr;
+		if (bApplyDirectLighting)
+		{
+			LightType = LightSceneInfoCompact[PassIndex].LightType;
+			LightSceneInfo = LightSceneInfoCompact[PassIndex].LightSceneInfo;
+			check(LightSceneInfo != nullptr);
 
-			if (HeterogeneousVolumes::UseTransmittanceVolume() && bApplyShadowTransmittance)
+			bApplyDirectLighting = (LightSceneInfo != nullptr);
+			if (LightSceneInfo)
 			{
-				RenderTransmittanceVolumeWithPreshadingHardwareRayTracing(
-					GraphBuilder,
-					// Scene data
-					Scene,
-					View,
-					SceneTextures,
-					// Light data
-					bApplyEmissionAndTransmittance,
-					bApplyDirectLighting,
-					bApplyShadowTransmittance,
-					LightType,
-					LightSceneInfo,
-					// Object data
-					PrimitiveSceneProxy,
-					// Sparse voxel
-					SparseVoxelUniformBuffer,
-					// Ray tracing data
-					Scene->HeterogeneousVolumesRayTracingScene,
-					// Transmittance volume
-					TransmittanceVolumeTexture
-				);
+				VisibleLightInfo = &VisibleLightInfos[LightSceneInfo->Id];
+				bApplyShadowTransmittance = LightSceneInfo->Proxy->CastsVolumetricShadow();
 			}
+		}
 
-			RenderSingleScatteringWithPreshadingHardwareRayTracing(
+		if (HeterogeneousVolumes::UseLightingCacheForTransmittance() && bApplyShadowTransmittance)
+		{
+			RenderLightingCacheWithPreshadingHardwareRayTracing(
 				GraphBuilder,
 				// Scene data
 				Scene,
@@ -833,11 +1145,140 @@ void RenderWithPreshadingHardwareRayTracing(
 				// Ray tracing data
 				Scene->HeterogeneousVolumesRayTracingScene,
 				// Transmittance volume
-				TransmittanceVolumeTexture,
-				// Output
-				HeterogeneousVolumeRadiance
+				LightingCacheTexture
 			);
 		}
+
+		RenderSingleScatteringWithPreshadingHardwareRayTracing(
+			GraphBuilder,
+			// Scene data
+			Scene,
+			View,
+			SceneTextures,
+			// Light data
+			bApplyEmissionAndTransmittance,
+			bApplyDirectLighting,
+			bApplyShadowTransmittance,
+			LightType,
+			LightSceneInfo,
+			// Shadow
+			VisibleLightInfo,
+			// Object data
+			PrimitiveSceneProxy,
+			// Sparse voxel
+			SparseVoxelUniformBuffer,
+			// Ray tracing data
+			Scene->HeterogeneousVolumesRayTracingScene,
+			// Transmittance volume
+			LightingCacheTexture,
+			// Output
+			HeterogeneousVolumeRadiance
+		);
+	}
+#endif // RHI_RAYTRACING
+}
+
+void RenderWithPreshadingHardwareRayTracing(
+	FRDGBuilder& GraphBuilder,
+	// Scene data
+	const FSceneTextures& SceneTextures,
+	FScene* Scene,
+	const FSceneViewFamily& ViewFamily,
+	FViewInfo& View,
+	// Shadow data
+	TArray<FVisibleLightInfo, SceneRenderingAllocator>& VisibleLightInfos,
+	// Object data
+	const FPrimitiveSceneProxy* PrimitiveSceneProxy,
+	const FMaterialRenderProxy* MaterialRenderProxy,
+	// Sparse voxel data
+	FRDGBufferRef NumVoxelsBuffer,
+	const TRDGUniformBufferRef<FSparseVoxelUniformBufferParameters>& SparseVoxelUniformBuffer,
+	// Transmittance acceleration
+	FRDGTextureRef LightingCacheTexture,
+	// Output
+	FRDGTextureRef& HeterogeneousVolumeRadiance
+)
+{
+#if RHI_RAYTRACING
+	RDG_EVENT_SCOPE(GraphBuilder, "Hardware Ray Tracing");
+
+	// WARNING: Currently works, but I'm skeptical if all RHI resources have the correct lifetime management
+	TArray<FRayTracingGeometryRHIRef> RayTracingGeometries;
+	TArray<FMatrix> RayTracingTransforms;
+	{
+		RDG_EVENT_SCOPE(GraphBuilder, "Acceleration Structure Build");
+
+		GenerateRayTracingGeometryInstance(
+			GraphBuilder,
+			// Scene
+			Scene,
+			View,
+			// Object
+			PrimitiveSceneProxy,
+			// Sparse voxel
+			NumVoxelsBuffer,
+			SparseVoxelUniformBuffer,
+			// Output
+			RayTracingGeometries,
+			RayTracingTransforms
+		);
+
+		GenerateRayTracingScene(
+			GraphBuilder,
+			// Scene
+			Scene,
+			View,
+			// Object
+			PrimitiveSceneProxy,
+			// Ray tracing data
+			RayTracingGeometries,
+			RayTracingTransforms,
+			// Output
+			Scene->HeterogeneousVolumesRayTracingScene
+		);
+	}
+
+	if (HeterogeneousVolumes::UseLightingCacheForInscattering())
+	{
+		RenderWithInscatteringVolumePipelineWithPreshadingHardwareRayTracing(
+			GraphBuilder,
+			SceneTextures,
+			Scene,
+			ViewFamily,
+			View,
+			// Shadow data
+			VisibleLightInfos,
+			// Object data
+			PrimitiveSceneProxy,
+			MaterialRenderProxy,
+			NumVoxelsBuffer,
+			SparseVoxelUniformBuffer,
+			// Transmittance acceleration
+			LightingCacheTexture,
+			// Output
+			HeterogeneousVolumeRadiance
+		);
+	}
+	else
+	{
+		RenderWithTransmittanceVolumePipelineWithPreshadingHardwareRayTracing(
+			GraphBuilder,
+			SceneTextures,
+			Scene,
+			ViewFamily,
+			View,
+			// Shadow data
+			VisibleLightInfos,
+			// Object data
+			PrimitiveSceneProxy,
+			MaterialRenderProxy,
+			NumVoxelsBuffer,
+			SparseVoxelUniformBuffer,
+			// Transmittance acceleration
+			LightingCacheTexture,
+			// Output
+			HeterogeneousVolumeRadiance
+		);
 	}
 
 	// Tear-down ray tracing scene
@@ -977,7 +1418,7 @@ void RenderWithPreshading(
 	const int32 PrimitiveId,
 	const FBoxSphereBounds LocalBoxSphereBounds,
 	// Transmittance acceleration
-	FRDGTextureRef TransmittanceVolumeTexture,
+	FRDGTextureRef LightingCacheTexture,
 	// Output
 	FRDGTextureRef& HeterogeneousVolumeRadiance
 )
@@ -1120,7 +1561,7 @@ void RenderWithPreshading(
 			NumVoxelsBuffer,
 			SparseVoxelUniformBuffer,
 			// Transmittance acceleration
-			TransmittanceVolumeTexture,
+			LightingCacheTexture,
 			// Output
 			HeterogeneousVolumeRadiance
 		);
@@ -1144,7 +1585,7 @@ void RenderWithPreshading(
 			NumVoxelsBuffer,
 			SparseVoxelUniformBuffer,
 			// Transmittance acceleration
-			TransmittanceVolumeTexture,
+			LightingCacheTexture,
 			// Output
 			HeterogeneousVolumeRadiance
 		);
