Global Distance Field - fixed "Unable to bind a shader parameter type Texture3D<uint> with a SHADER_PARAMETER_RDG_TEXTURE_UAV" cook warning.

diff --git a/Engine/Source/Runtime/Renderer/Private/GlobalDistanceField.cpp b/Engine/Source/Runtime/Renderer/Private/GlobalDistanceField.cpp
--- a/Engine/Source/Runtime/Renderer/Private/GlobalDistanceField.cpp
+++ b/Engine/Source/Runtime/Renderer/Private/GlobalDistanceField.cpp
@@ -1534,7 +1534,7 @@ class FCompositeObjectsIntoPagesCS : public FGlobalShader
 		RDG_BUFFER_ACCESS(ComposeIndirectArgBuffer, ERHIAccess::IndirectArgs)
 		SHADER_PARAMETER_RDG_BUFFER_SRV(StructuredBuffer<uint>, ComposeTileBuffer)
 		SHADER_PARAMETER_RDG_BUFFER_SRV(StructuredBuffer<uint>, HeightfieldMarkedPageBuffer)
-		SHADER_PARAMETER_RDG_TEXTURE_UAV(Texture3D<uint>, RWPageTableCombinedTexture)
+		SHADER_PARAMETER_RDG_TEXTURE_UAV(RWTexture3D<uint>, RWPageTableCombinedTexture)
 		SHADER_PARAMETER_RDG_TEXTURE(Texture3D<uint>, ParentPageTableLayerTexture)
 		SHADER_PARAMETER_RDG_BUFFER_SRV(StructuredBuffer<uint>, CullGridObjectHeader)
 		SHADER_PARAMETER_RDG_BUFFER_SRV(StructuredBuffer<uint>, CullGridObjectArray)
