Fix incorrect AO when using SSGI

diff --git a/Engine/Source/Runtime/Renderer/Private/IndirectLightRendering.cpp b/Engine/Source/Runtime/Renderer/Private/IndirectLightRendering.cpp
--- a/Engine/Source/Runtime/Renderer/Private/IndirectLightRendering.cpp
+++ b/Engine/Source/Runtime/Renderer/Private/IndirectLightRendering.cpp
@@ -416,7 +416,7 @@ void FDeferredShadingSceneRenderer::CommitIndirectLightingState()
 			DiffuseIndirectDenoiser = IScreenSpaceDenoiser::GetDenoiserMode(CVarDiffuseIndirectDenoiser);
 		}
 		
-		if (DiffuseIndirectMethod == EDiffuseIndirectMethod::Disabled && ScreenSpaceRayTracing::IsScreenSpaceDiffuseIndirectSupported(View))
+		if (DiffuseIndirectMethod == EDiffuseIndirectMethod::SSGI)
 		{
 			if (CVarLumenProbeHierarchy.GetValueOnRenderThread() && CVarStandaloneSSGIAllowLumenProbeHierarchy.GetValueOnRenderThread())
 			{
