Removed Mobile FSR from Source/Shaders

diff --git a/Engine/Source/Runtime/Renderer/Private/PostProcess/PostProcessMobile.cpp b/Engine/Source/Runtime/Renderer/Private/PostProcess/PostProcessMobile.cpp
--- a/Engine/Source/Runtime/Renderer/Private/PostProcess/PostProcessMobile.cpp
+++ b/Engine/Source/Runtime/Renderer/Private/PostProcess/PostProcessMobile.cpp
@@ -16,19 +16,6 @@
 #include "ClearQuad.h"
 #include "PostProcess/PostProcessing.h"
 
-#define A_CPU 1
-// colliding ffx_a.ush/UE4Engine definitions
-// ffx_a.ush will override these definitions
-#ifdef A_STATIC
-#undef A_STATIC
-#endif
-#ifdef A_RESTRICT
-#undef A_RESTRICT
-#endif
-#include "../Shaders/Private/Mobile/FSR/ffx_a.ush"
-#include "../Shaders/Private/Mobile/FSR/ffx_fsr1.ush"
-
-
 static TAutoConsoleVariable<int32> CVarMobileSupportBloomSetupRareCases(
 	TEXT("r.Mobile.MobileSupportBloomSetupRareCases"),
 	0,
@@ -2014,400 +2001,3 @@ void AddMobileEyeAdaptationPass(FRDGBuilder& GraphBuilder, const FViewInfo& View
 
 	View.EnqueueEyeAdaptationExposureBufferReadback(GraphBuilder);
 }
-
-//////////////////////////////////////////////////////////////////////////
-//! FSR
-//////////////////////////////////////////////////////////////////////////
-
-// permutation domains
-class FFSR_UseFP16Dim : SHADER_PERMUTATION_BOOL("ENABLE_FP16");
-class FRCAS_DenoiseDim : SHADER_PERMUTATION_BOOL("USE_RCAS_DENOISE");
-
-BEGIN_SHADER_PARAMETER_STRUCT(FFSRPassParameters, )
-	SHADER_PARAMETER(FUintVector4, Const0)
-	SHADER_PARAMETER(FUintVector4, Const1)
-	SHADER_PARAMETER(FUintVector4, Const2)
-	SHADER_PARAMETER(FUintVector4, Const3)
-	SHADER_PARAMETER(FVector2f, VPColor_ExtentInverse)
-	SHADER_PARAMETER(FVector2f, VPColor_ViewportMin)
-	SHADER_PARAMETER_SAMPLER(SamplerState, samLinearClamp)
-	SHADER_PARAMETER_RDG_TEXTURE(Texture2D, InputTexture)
-END_SHADER_PARAMETER_STRUCT()
-
-BEGIN_SHADER_PARAMETER_STRUCT(FRCASPassParameters, )
-	SHADER_PARAMETER_RDG_TEXTURE(Texture2D, InputTexture)
-	SHADER_PARAMETER(FUintVector4, Const0)
-	SHADER_PARAMETER(FVector2f, VPColor_ExtentInverse)
-END_SHADER_PARAMETER_STRUCT()
-
-static float GMobileFSRCASSharpness = 0.2; // 0.2 stops = 1 / (2^N) ~= 0.87 in linear [0-1] sharpening of CAS
-static FAutoConsoleVariableRef CVarMobileFSRCASSharpness(
-	TEXT("r.Mobile.FSR.RCAS.Sharpness"),
-	GMobileFSRCASSharpness,
-	TEXT("FSR RCAS Sharpness in stops (0: sharpest, 1: 1/2 as sharp, 2: 1/4 as sharp, 3: 1/8 as sharp, etc.). A value of 0.2 would correspond to a ~0.87 sharpness in [0-1] linear scale"),
-	ECVF_Default);
-
-static int32 GMobileFSRCASDenoise = 0;
-static FAutoConsoleVariableRef CVarMobileFSRCASDenoise(
-	TEXT("r.Mobile.FSR.RCAS.Denoise"),
-	GMobileFSRCASDenoise,
-	TEXT("FSR RCAS Denoise support for grainy input such as dithered images or input with custom film grain effects applied prior to FSR. 1:On, 0:Off"),
-	ECVF_Default);
-
-static int32 GMobileFSRDisableCompute = 0;
-static FAutoConsoleVariableRef CVarMobileFSRDisableCompute(
-	TEXT("r.Mobile.FSR.DisableCompute"),
-	GMobileFSRDisableCompute,
-	TEXT("FSR Disable compute shaders for FSR (Recommended for Adreno). 1: Compute Disabled, 0: Compute Enabled (default)"),
-	ECVF_Default);
-
-///
-/// FSR COMPUTE SHADER
-///
-class FFSRCS : public FGlobalShader
-{
-public:
-	DECLARE_GLOBAL_SHADER(FFSRCS);
-	SHADER_USE_PARAMETER_STRUCT(FFSRCS, FGlobalShader);
-
-	using FPermutationDomain = TShaderPermutationDomain<FFSR_UseFP16Dim>;
-
-	BEGIN_SHADER_PARAMETER_STRUCT(FParameters, )
-		SHADER_PARAMETER_STRUCT_INCLUDE(FFSRPassParameters, FSR)
-		SHADER_PARAMETER_RDG_TEXTURE_UAV(RWTexture2D, OutputTexture)
-	END_SHADER_PARAMETER_STRUCT()
-
-	static bool ShouldCompilePermutation(const FGlobalShaderPermutationParameters& Parameters)
-	{
-		FPermutationDomain PermutationVector(Parameters.PermutationId);
-
-		if (PermutationVector.Get<FFSR_UseFP16Dim>())
-		{
-			if (IsSimulatedPlatform(Parameters.Platform) || !IsDxcEnabledForPlatform(Parameters.Platform))
-			{
-				return false;
-			}
-		}
-
-		return IsMobilePlatform(Parameters.Platform);
-	}
-
-	static void ModifyCompilationEnvironment(const FGlobalShaderPermutationParameters& Parameters, FShaderCompilerEnvironment& OutEnvironment)
-	{
-		FPermutationDomain PermutationVector(Parameters.PermutationId);
-
-		OutEnvironment.SetDefine(TEXT("COMPUTE_SHADER"), 1);
-
-		OutEnvironment.SetDefine(TEXT("THREADGROUP_SIZEX"), 64);
-		OutEnvironment.SetDefine(TEXT("THREADGROUP_SIZEY"), 1);
-		OutEnvironment.SetDefine(TEXT("THREADGROUP_SIZEZ"), 1);
-	}
-};
-
-IMPLEMENT_GLOBAL_SHADER(FFSRCS, "/Engine/Private/Mobile/FSR/PostProcessFFX_FSR.usf", "MainCS", SF_Compute);
-
-class FFSRPS : public FGlobalShader
-{
-public:
-	DECLARE_GLOBAL_SHADER(FFSRPS);
-	SHADER_USE_PARAMETER_STRUCT(FFSRPS, FGlobalShader);
-
-	using FPermutationDomain = TShaderPermutationDomain<FFSR_UseFP16Dim>;
-
-	BEGIN_SHADER_PARAMETER_STRUCT(FParameters, )
-		SHADER_PARAMETER_STRUCT_INCLUDE(FFSRPassParameters, FSR)
-		RENDER_TARGET_BINDING_SLOTS()
-	END_SHADER_PARAMETER_STRUCT()
-
-	static bool ShouldCompilePermutation(const FGlobalShaderPermutationParameters& Parameters)
-	{
-		FPermutationDomain PermutationVector(Parameters.PermutationId);
-
-		if (PermutationVector.Get<FFSR_UseFP16Dim>())
-		{
-			if (IsSimulatedPlatform(Parameters.Platform) || !IsDxcEnabledForPlatform(Parameters.Platform))
-			{
-				return false; 
-			}
-		}
-
-		return IsMobilePlatform(Parameters.Platform);
-	}
-
-	static void ModifyCompilationEnvironment(const FGlobalShaderPermutationParameters& Parameters, FShaderCompilerEnvironment& OutEnvironment)
-	{
-	}
-};
-
-IMPLEMENT_GLOBAL_SHADER(FFSRPS, "/Engine/Private/Mobile/FSR/PostProcessFFX_FSR.usf", "MainPS", SF_Pixel);
-
-FScreenPassTexture AddEASUPass(FRDGBuilder& GraphBuilder, const FViewInfo& View, const ISpatialUpscaler::FInputs& PassInputs)
-{
-	FScreenPassTexture OutputTexture;
-	bool bUseFP16 = !IsSimulatedPlatform(View.GetShaderPlatform()) && IsDxcEnabledForPlatform(View.GetShaderPlatform());
-
-	FRDGTextureDesc FSROutputTextureDesc = PassInputs.SceneColor.Texture->Desc;
-	FSROutputTextureDesc.Reset();
-	FSROutputTextureDesc.Extent = View.UnscaledViewRect.Max;
-	FSROutputTextureDesc.ClearValue = FClearValueBinding::Black;
-	FSROutputTextureDesc.Flags = TexCreate_ShaderResource | TexCreate_UAV | TexCreate_RenderTargetable;
-
-	if (PassInputs.OverrideOutput.IsValid())
-	{
-		OutputTexture = PassInputs.OverrideOutput;
-	}
-	else
-	{
-		OutputTexture = FScreenPassTexture(PassInputs.SceneColor);
-		OutputTexture.Texture = GraphBuilder.CreateTexture(FSROutputTextureDesc, TEXT("FSR-Output"), ERDGTextureFlags::MultiFrame);
-	}
-
-	OutputTexture.ViewRect = View.UnscaledViewRect;
-	FScreenPassTextureViewport OutputViewport = FScreenPassTextureViewport(OutputTexture);
-	FScreenPassTextureViewport InputViewport = FScreenPassTextureViewport(PassInputs.SceneColor);
-	FScreenPassTextureViewportParameters PassOutputViewportParams = GetScreenPassTextureViewportParameters(OutputViewport);
-
-	AU1 const0[4];
-	AU1 const1[4];
-	AU1 const2[4];
-	AU1 const3[4]; // Configure FSR
-
-	FsrEasuConOffset(const0, const1, const2, const3,
-		static_cast<AF1>(InputViewport.Rect.Width())
-		, static_cast<AF1>(InputViewport.Rect.Height()) // current frame render resolution
-		, static_cast<AF1>(PassInputs.SceneColor.Texture->Desc.Extent.X)
-		, static_cast<AF1>(PassInputs.SceneColor.Texture->Desc.Extent.Y) // input container resolution (for DRS)
-		, static_cast<AF1>(OutputViewport.Rect.Width())
-		, static_cast<AF1>(OutputViewport.Rect.Height()) // upscaled-to resolution
-		, static_cast<AF1>(InputViewport.Rect.Min.X)
-		, static_cast<AF1>(InputViewport.Rect.Min.Y)
-	);
-
-	const bool bOutputSupportsUAV = (OutputTexture.Texture->Desc.Flags & TexCreate_UAV) == TexCreate_UAV;
-
-	if (!bOutputSupportsUAV || GMobileFSRDisableCompute)
-	{
-		FFSRPS::FParameters* PassParameters = GraphBuilder.AllocParameters<FFSRPS::FParameters>();
-		for (int i = 0; i < 4; i++)
-		{
-			PassParameters->FSR.Const0[i] = const0[i];
-			PassParameters->FSR.Const1[i] = const1[i];
-			PassParameters->FSR.Const2[i] = const2[i];
-			PassParameters->FSR.Const3[i] = const3[i];
-		}
-		PassParameters->FSR.InputTexture = PassInputs.SceneColor.Texture;
-		PassParameters->FSR.samLinearClamp = TStaticSamplerState<SF_Bilinear, AM_Clamp, AM_Clamp, AM_Clamp>::GetRHI();
-		PassParameters->FSR.VPColor_ExtentInverse = PassOutputViewportParams.ExtentInverse;
-		PassParameters->FSR.VPColor_ViewportMin = PassOutputViewportParams.ViewportMin;
-		PassParameters->RenderTargets[0] = FRenderTargetBinding(OutputTexture.Texture, ERenderTargetLoadAction::EClear);
-
-		FFSRPS::FPermutationDomain PSPermutationVector;
-		PSPermutationVector.Set<FFSR_UseFP16Dim>(bUseFP16);
-
-		TShaderMapRef<FFSRPS> PixelShader(View.ShaderMap, PSPermutationVector);
-
-		AddDrawScreenPass(GraphBuilder,
-			RDG_EVENT_NAME("FSR/Upscale %dx%d -> %dx%d (PS)"
-				, InputViewport.Rect.Width(), InputViewport.Rect.Height()
-				, OutputViewport.Rect.Width(), OutputViewport.Rect.Height()),
-			View, OutputViewport, InputViewport,
-			PixelShader, PassParameters,
-			EScreenPassDrawFlags::None
-		);
-	}
-	// CS
-	else
-	{
-		FFSRCS::FParameters* PassParameters = GraphBuilder.AllocParameters<FFSRCS::FParameters>();
-
-		for (int i = 0; i < 4; i++)
-		{
-			PassParameters->FSR.Const0[i] = const0[i];
-			PassParameters->FSR.Const1[i] = const1[i];
-			PassParameters->FSR.Const2[i] = const2[i];
-			PassParameters->FSR.Const3[i] = const3[i];
-		}
-		PassParameters->FSR.InputTexture = PassInputs.SceneColor.Texture;
-		PassParameters->FSR.samLinearClamp = TStaticSamplerState<SF_Bilinear, AM_Clamp, AM_Clamp, AM_Clamp>::GetRHI();
-		PassParameters->FSR.VPColor_ExtentInverse = PassOutputViewportParams.ExtentInverse;
-		PassParameters->FSR.VPColor_ViewportMin = PassOutputViewportParams.ViewportMin;
-		PassParameters->OutputTexture = GraphBuilder.CreateUAV(OutputTexture.Texture);
-
-		FFSRCS::FPermutationDomain CSPermutationVector;
-		CSPermutationVector.Set<FFSR_UseFP16Dim>(bUseFP16);
-
-		TShaderMapRef<FFSRCS> ComputeShaderFSR(View.ShaderMap, CSPermutationVector);
-		FComputeShaderUtils::AddPass(
-			GraphBuilder,
-			RDG_EVENT_NAME("FSR/Upscale %dx%d -> %dx%d (CS)"
-				, InputViewport.Rect.Width(), InputViewport.Rect.Height()
-				, OutputViewport.Rect.Width(), OutputViewport.Rect.Height()),
-			ComputeShaderFSR,
-			PassParameters,
-			FComputeShaderUtils::GetGroupCount(OutputViewport.Rect.Size(), 16));
-	}
-
-	return MoveTemp(OutputTexture);
-}
-
-///
-/// RCAS COMPUTE SHADER
-///
-class FRCASCS : public FGlobalShader
-{
-public:
-	DECLARE_GLOBAL_SHADER(FRCASCS);
-	SHADER_USE_PARAMETER_STRUCT(FRCASCS, FGlobalShader);
-
-	using FPermutationDomain = TShaderPermutationDomain<FFSR_UseFP16Dim, FRCAS_DenoiseDim>;
-
-	BEGIN_SHADER_PARAMETER_STRUCT(FParameters, )
-		SHADER_PARAMETER_STRUCT_INCLUDE(FRCASPassParameters, RCAS)
-		SHADER_PARAMETER_RDG_TEXTURE_UAV(RWTexture2D, OutputTexture)
-	END_SHADER_PARAMETER_STRUCT()
-
-	static bool ShouldCompilePermutation(const FGlobalShaderPermutationParameters& Parameters)
-	{
-		FPermutationDomain PermutationVector(Parameters.PermutationId);
-
-		if (PermutationVector.Get<FFSR_UseFP16Dim>())
-		{
-			if (IsSimulatedPlatform(Parameters.Platform) || !IsDxcEnabledForPlatform(Parameters.Platform))
-			{
-				return false;
-			}
-		}
-
-		return IsMobilePlatform(Parameters.Platform);
-	}
-
-	static void ModifyCompilationEnvironment(const FGlobalShaderPermutationParameters& Parameters, FShaderCompilerEnvironment& OutEnvironment)
-	{
-		FPermutationDomain PermutationVector(Parameters.PermutationId);
-
-		OutEnvironment.SetDefine(TEXT("COMPUTE_SHADER"), 1);
-
-		OutEnvironment.SetDefine(TEXT("THREADGROUP_SIZEX"), 64);
-		OutEnvironment.SetDefine(TEXT("THREADGROUP_SIZEY"), 1);
-		OutEnvironment.SetDefine(TEXT("THREADGROUP_SIZEZ"), 1);
-	}
-};
-
-IMPLEMENT_GLOBAL_SHADER(FRCASCS, "/Engine/Private/Mobile/FSR/PostProcessFFX_RCAS.usf", "MainCS", SF_Compute);
-
-
-///
-/// RCAS PIXEL SHADER
-///
-class FRCASPS : public FGlobalShader
-{
-public:
-	DECLARE_GLOBAL_SHADER(FRCASPS);
-	SHADER_USE_PARAMETER_STRUCT(FRCASPS, FGlobalShader);
-
-	using FPermutationDomain = TShaderPermutationDomain<FFSR_UseFP16Dim, FRCAS_DenoiseDim>;
-
-	BEGIN_SHADER_PARAMETER_STRUCT(FParameters, )
-		SHADER_PARAMETER_STRUCT_INCLUDE(FRCASPassParameters, RCAS)
-		RENDER_TARGET_BINDING_SLOTS()
-	END_SHADER_PARAMETER_STRUCT()
-
-	static bool ShouldCompilePermutation(const FGlobalShaderPermutationParameters& Parameters)
-	{
-		FPermutationDomain PermutationVector(Parameters.PermutationId);
-		if (PermutationVector.Get<FFSR_UseFP16Dim>())
-		{
-			if (IsSimulatedPlatform(Parameters.Platform) || !IsDxcEnabledForPlatform(Parameters.Platform))
-			{
-				return false;
-			}
-		}
-
-		return IsMobilePlatform(Parameters.Platform);
-	}
-
-	static void ModifyCompilationEnvironment(const FGlobalShaderPermutationParameters& Parameters, FShaderCompilerEnvironment& OutEnvironment)
-	{
-	}
-};
-
-IMPLEMENT_GLOBAL_SHADER(FRCASPS, "/Engine/Private/Mobile/FSR/PostProcessFFX_RCAS.usf", "MainPS", SF_Pixel);
-
-FScreenPassTexture AddCASPass(FRDGBuilder& GraphBuilder, const FViewInfo& View, const ISpatialUpscaler::FInputs& PassInputs)
-{
-	FScreenPassTexture OutputTexture;
-	bool bUseFP16 = !IsSimulatedPlatform(View.GetShaderPlatform()) && IsDxcEnabledForPlatform(View.GetShaderPlatform());
-
-	if (PassInputs.OverrideOutput.IsValid())
-	{
-		OutputTexture = PassInputs.OverrideOutput;
-	}
-	else
-	{
-		FRDGTextureDesc RCASOutputTextureDesc = PassInputs.SceneColor.Texture->Desc;
-
-		OutputTexture = FScreenPassTexture(PassInputs.SceneColor);
-		OutputTexture.Texture = GraphBuilder.CreateTexture(RCASOutputTextureDesc, TEXT("RCAS-Output"), ERDGTextureFlags::MultiFrame);
-	}
-
-	FScreenPassTextureViewport OutputViewport = FScreenPassTextureViewport(OutputTexture);
-	FScreenPassTextureViewportParameters PassOutputViewportParams = GetScreenPassTextureViewportParameters(OutputViewport);
-
-	AU1 const0[4]; // Configure FSR
-	FsrRcasCon(const0, GMobileFSRCASSharpness);
-
-	const bool bOutputSupportsUAV = (OutputTexture.Texture->Desc.Flags & TexCreate_UAV) == TexCreate_UAV;
-
-	if (!bOutputSupportsUAV || GMobileFSRDisableCompute)
-	{
-		FRCASPS::FParameters* PassParameters = GraphBuilder.AllocParameters<FRCASPS::FParameters>();
-
-		// set pass inputs
-		for (int i = 0; i < 4; i++) { PassParameters->RCAS.Const0[i] = const0[i]; }
-		PassParameters->RCAS.InputTexture = PassInputs.SceneColor.Texture;
-		PassParameters->RCAS.VPColor_ExtentInverse = PassOutputViewportParams.ExtentInverse;
-		PassParameters->RenderTargets[0] = FRenderTargetBinding(OutputTexture.Texture, ERenderTargetLoadAction::EClear);
-
-		// grab shaders
-		FRCASPS::FPermutationDomain PSPermutationVector;
-		PSPermutationVector.Set<FFSR_UseFP16Dim>(bUseFP16);
-		PSPermutationVector.Set<FRCAS_DenoiseDim>(GMobileFSRCASDenoise > 0);
-
-		TShaderMapRef<FRCASPS> PixelShader(View.ShaderMap, PSPermutationVector);
-
-		AddDrawScreenPass(GraphBuilder,
-			RDG_EVENT_NAME("FSR/RCAS Sharpness=%.2f(PS)"
-				, GMobileFSRCASSharpness),
-			View, OutputViewport, OutputViewport,
-			PixelShader, PassParameters,
-			EScreenPassDrawFlags::None
-		);
-	}
-	// CS
-	else
-	{
-		FRCASCS::FParameters* PassParameters = GraphBuilder.AllocParameters<FRCASCS::FParameters>();
-
-		for (int i = 0; i < 4; i++)
-		{
-			PassParameters->RCAS.Const0[i] = const0[i];
-		}
-		PassParameters->RCAS.InputTexture = PassInputs.SceneColor.Texture;
-		PassParameters->RCAS.VPColor_ExtentInverse = PassOutputViewportParams.ExtentInverse;
-		PassParameters->OutputTexture = GraphBuilder.CreateUAV(OutputTexture.Texture);
-
-		FRCASCS::FPermutationDomain CSPermutationVector;
-		CSPermutationVector.Set<FFSR_UseFP16Dim>(bUseFP16);
-		CSPermutationVector.Set<FRCAS_DenoiseDim>(GMobileFSRCASDenoise > 0);
-
-		TShaderMapRef<FRCASCS> ComputeShaderRCASPass(View.ShaderMap, CSPermutationVector);
-		FComputeShaderUtils::AddPass(
-			GraphBuilder,
-			RDG_EVENT_NAME("FSR/RCAS Sharpness=%.2f (CS)", GMobileFSRCASSharpness),
-			ComputeShaderRCASPass,
-			PassParameters,
-			FComputeShaderUtils::GetGroupCount(OutputViewport.Rect.Size(), 16));
-	}
-
-	return MoveTemp(OutputTexture);
-}
\ No newline at end of file
diff --git a/Engine/Source/Runtime/Renderer/Private/PostProcess/PostProcessing.cpp b/Engine/Source/Runtime/Renderer/Private/PostProcess/PostProcessing.cpp
--- a/Engine/Source/Runtime/Renderer/Private/PostProcess/PostProcessing.cpp
+++ b/Engine/Source/Runtime/Renderer/Private/PostProcess/PostProcessing.cpp
@@ -130,19 +130,7 @@ TAutoConsoleVariable<int32> CVarPostProcessingForceAsyncDispatch(
 	TEXT("Only available for testing in non-shipping builds."),
 	ECVF_RenderThreadSafe);
 #endif
-
-TAutoConsoleVariable<int32> CVarMobileFSRCASEnabled(
-	TEXT("r.Mobile.FSR.RCAS.Enabled"),
-	0,
-	TEXT("FidelityFX FSR/RCAS : Robust Contrast Adaptive Sharpening Filter. Requires r.FidelityFX.FSR.PrimaryUpscale 1 or r.FidelityFX.FSR.SecondaryUpscale 1"),
-	ECVF_RenderThreadSafe);
-
-TAutoConsoleVariable<int32>CVarMobileFSRUpsamplingEnabled(
-	TEXT("r.Mobile.FSR.Upsampling.Enabled"),
-	0,
-	TEXT("FidelityFX FSR/RCAS : Robust Contrast Adaptive Sharpening Filter. Requires r.FidelityFX.FSR.PrimaryUpscale 1 or r.FidelityFX.FSR.SecondaryUpscale 1"),
-	ECVF_Default);
-} //! namespace
+}
 
 EDownsampleQuality GetDownsampleQuality(const TAutoConsoleVariable<int32>& CVar)
 {
@@ -1796,7 +1784,7 @@ void AddMobilePostProcessingPasses(FRDGBuilder& GraphBuilder, FScene* Scene, con
 		SelectionOutline,
 		EditorPrimitive,
 		PrimaryUpscale,
-		Sharpen,
+		SecondaryUpscale,
 		Visualize,
 		HMDDistortion,
 		MAX
@@ -1860,10 +1848,8 @@ void AddMobilePostProcessingPasses(FRDGBuilder& GraphBuilder, FScene* Scene, con
 
 	bool bUseHighResolutionScreenshotMask = IsHighResolutionScreenshotMaskEnabled(View);
 
-	bool bEnabledFSRCAS = CVarMobileFSRCASEnabled.GetValueOnRenderThread() > 0;
-	bool bEnabledFSRUpsampling = CVarMobileFSRUpsamplingEnabled.GetValueOnRenderThread() > 0;
-	bool bShouldPrimaryUpscale = (View.PrimaryScreenPercentageMethod == EPrimaryScreenPercentageMethod::SpatialUpscale && View.UnscaledViewRect != View.ViewRect)
-								|| PaniniConfig.IsEnabled();
+	bool bShouldPrimaryUpscale = (View.PrimaryScreenPercentageMethod == EPrimaryScreenPercentageMethod::SpatialUpscale && View.UnscaledViewRect != View.ViewRect) || PaniniConfig.IsEnabled();
+	bShouldPrimaryUpscale |= View.Family->GetPrimarySpatialUpscalerInterface() != nullptr;
 
 	PassSequence.SetEnabled(EPass::Tonemap, bUseToneMapper);
 	PassSequence.SetEnabled(EPass::HighResolutionScreenshotMask, bUseHighResolutionScreenshotMask);
@@ -1875,7 +1861,7 @@ void AddMobilePostProcessingPasses(FRDGBuilder& GraphBuilder, FScene* Scene, con
 	PassSequence.SetEnabled(EPass::EditorPrimitive, false);
 #endif
 	PassSequence.SetEnabled(EPass::PrimaryUpscale, bShouldPrimaryUpscale);
-	PassSequence.SetEnabled(EPass::Sharpen, bEnabledFSRCAS);
+	PassSequence.SetEnabled(EPass::SecondaryUpscale, View.Family->GetSecondarySpatialUpscalerInterface() != nullptr);
 
 	PassSequence.SetEnabled(EPass::Visualize, View.Family->EngineShowFlags.ShaderComplexity);
 
@@ -2464,25 +2450,30 @@ void AddMobilePostProcessingPasses(FRDGBuilder& GraphBuilder, FScene* Scene, con
 		}
 		else
 		{
-			if (!bEnabledFSRUpsampling)
-			{
-				SceneColor = ISpatialUpscaler::AddDefaultUpscalePass(GraphBuilder, View, PassInputs, EUpscaleMethod::Bilinear, PaniniConfig);
-			}
-			else
-			{
-				SceneColor = AddEASUPass(GraphBuilder, View, PassInputs);
-			}
+			SceneColor = ISpatialUpscaler::AddDefaultUpscalePass(GraphBuilder, View, PassInputs, EUpscaleMethod::Bilinear, PaniniConfig);
 		}
 	}
 
-	if (PassSequence.IsEnabled(EPass::Sharpen))
+	if (PassSequence.IsEnabled(EPass::SecondaryUpscale))
 	{
 		ISpatialUpscaler::FInputs PassInputs;
-		PassSequence.AcceptOverrideIfLastPass(EPass::Sharpen, PassInputs.OverrideOutput);
+		PassSequence.AcceptOverrideIfLastPass(EPass::SecondaryUpscale, PassInputs.OverrideOutput);
 		PassInputs.SceneColor = SceneColor;
-		PassInputs.OverrideOutput.LoadAction = View.IsFirstInFamily() ? ERenderTargetLoadAction::EClear : ERenderTargetLoadAction::ELoad;
+		PassInputs.Stage = EUpscaleStage::SecondaryToOutput;
+
+		const ISpatialUpscaler* CustomUpscaler = View.Family ? View.Family->GetSecondarySpatialUpscalerInterface() : nullptr;
+		if (CustomUpscaler)
+		{
+			RDG_EVENT_SCOPE(
+				GraphBuilder,
+				"ThirdParty SecondaryUpscale %s %dx%d -> %dx%d",
+				CustomUpscaler->GetDebugName(),
+				SceneColor.ViewRect.Width(), SceneColor.ViewRect.Height(),
+				View.UnscaledViewRect.Width(), View.UnscaledViewRect.Height());
 
-		SceneColor = AddCASPass(GraphBuilder, View, PassInputs);
+			SceneColor = CustomUpscaler->AddPasses(GraphBuilder, View, PassInputs);
+			check(SceneColor == ViewFamilyOutput);
+		}
 	}
 
 	if (PassSequence.IsEnabled(EPass::Visualize))
