Fix groom not receiving Lumen lighting when there is no sky.

diff --git a/Engine/Source/Runtime/Renderer/Private/HairStrands/HairStrandsEnvironment.cpp b/Engine/Source/Runtime/Renderer/Private/HairStrands/HairStrandsEnvironment.cpp
--- a/Engine/Source/Runtime/Renderer/Private/HairStrands/HairStrandsEnvironment.cpp
+++ b/Engine/Source/Runtime/Renderer/Private/HairStrands/HairStrandsEnvironment.cpp
@@ -354,8 +354,8 @@ static void AddHairStrandsEnvironmentLightingPassPS(
 	// Sanity check
 	if (bHasStaticLighting) { check(LightingType != EHairLightingSourceType::Lumen); }
 
-	// Early out if there is no static lighting, no sky lighting, nor reflection probes
-	if (!bSkyLight && !bDynamicSkyLight && !bHasStaticLighting && !bReflectionEnv)
+	// Early out if there is no static lighting, no sky lighting, nor reflection probes, nor Lumen
+	if (!bSkyLight && !bDynamicSkyLight && !bHasStaticLighting && !bReflectionEnv && !bLumenActive)
 		return;
 
 	FSceneTextureParameters SceneTextures = GetSceneTextureParameters(GraphBuilder, View);
