Fix static analysis warning for dereferencing null pointer in PostProcessing.cpp

diff --git a/Engine/Source/Runtime/Renderer/Private/PostProcess/PostProcessing.cpp b/Engine/Source/Runtime/Renderer/Private/PostProcess/PostProcessing.cpp
--- a/Engine/Source/Runtime/Renderer/Private/PostProcess/PostProcessing.cpp
+++ b/Engine/Source/Runtime/Renderer/Private/PostProcess/PostProcessing.cpp
@@ -2461,7 +2461,7 @@ void AddMobilePostProcessingPasses(FRDGBuilder& GraphBuilder, FScene* Scene, con
 		PassInputs.SceneColor = SceneColor;
 		PassInputs.Stage = EUpscaleStage::SecondaryToOutput;
 
-		const ISpatialUpscaler* CustomUpscaler = View.Family ? View.Family->GetSecondarySpatialUpscalerInterface() : nullptr;
+		const ISpatialUpscaler* CustomUpscaler = View.Family->GetSecondarySpatialUpscalerInterface();
 		if (CustomUpscaler)
 		{
 			RDG_EVENT_SCOPE(
