Fixed: crash in CitySample related to mobile auto-instancing changes

diff --git a/Engine/Source/Runtime/Renderer/Private/MeshDrawCommands.cpp b/Engine/Source/Runtime/Renderer/Private/MeshDrawCommands.cpp
--- a/Engine/Source/Runtime/Renderer/Private/MeshDrawCommands.cpp
+++ b/Engine/Source/Runtime/Renderer/Private/MeshDrawCommands.cpp
@@ -940,8 +940,14 @@ public:
 
 			if (Context.bUseGPUScene)
 			{
+				TArrayView<const FStateBucketAuxData> StateBucketsAuxData;
+				if (Context.PassType != EMeshPass::Num)
+				{
+					StateBucketsAuxData = Context.Scene->CachedStateBucketsAuxData[Context.PassType];
+				}
+				
 				Context.InstanceCullingContext.SetupDrawCommands(
-					Context.Scene->CachedStateBucketsAuxData[Context.PassType],
+					StateBucketsAuxData,
 					Context.MeshDrawCommands, 
 					true, 
 					Context.MaxInstances, 
