Small PSO precache fixes:
- reduce the amount PSOs precached which won't be used during rendering
- fixes a few missing PSOs
- add some profile tags
- Minor fixes when running with PSO precaching turned on in the editor (won't do a lot right now but also won't fail or assert)

#rnx

diff --git a/Engine/Source/Runtime/Engine/Classes/Materials/MaterialInterface.h b/Engine/Source/Runtime/Engine/Classes/Materials/MaterialInterface.h
--- a/Engine/Source/Runtime/Engine/Classes/Materials/MaterialInterface.h
+++ b/Engine/Source/Runtime/Engine/Classes/Materials/MaterialInterface.h
@@ -506,7 +506,7 @@ public:
 	{
 		return PrecachePSOs(MakeArrayView(&VertexFactoryType, 1), PreCacheParams);
 	}
-	virtual FGraphEventArray PrecachePSOs(const TConstArrayView<const FVertexFactoryType*>& VertexFactoryTypes, const struct FPSOPrecacheParams& PreCacheParams) PURE_VIRTUAL(UMaterialInterface::PrecachePSOs, return FGraphEventArray(););
+	virtual FGraphEventArray PrecachePSOs(const TConstArrayView<const FVertexFactoryType*>& VertexFactoryTypes, const struct FPSOPrecacheParams& PreCacheParams) { return FGraphEventArray(); }
 
 #if WITH_EDITORONLY_DATA
 
diff --git a/Engine/Source/Runtime/Engine/Private/Components/PrimitiveComponent.cpp b/Engine/Source/Runtime/Engine/Private/Components/PrimitiveComponent.cpp
--- a/Engine/Source/Runtime/Engine/Private/Components/PrimitiveComponent.cpp
+++ b/Engine/Source/Runtime/Engine/Private/Components/PrimitiveComponent.cpp
@@ -4249,6 +4249,7 @@ void UPrimitiveComponent::SetupPrecachePSOParams(FPSOPrecacheParams& Params)
 	Params.bStaticLighting = HasStaticLighting();
 	Params.bAffectDynamicIndirectLighting = bAffectDynamicIndirectLighting;
 	Params.bCastShadow = CastShadow;
+	Params.bRenderCustomDepth = bRenderCustomDepth;
 	Params.bCastShadowAsTwoSided = bCastShadowAsTwoSided;
 	Params.SetMobility(Mobility);	
 	Params.SetStencilWriteMask(FRendererStencilMaskEvaluation::ToStencilMask(CustomDepthStencilWriteMask));
diff --git a/Engine/Source/Runtime/Engine/Private/Components/StaticMeshComponent.cpp b/Engine/Source/Runtime/Engine/Private/Components/StaticMeshComponent.cpp
--- a/Engine/Source/Runtime/Engine/Private/Components/StaticMeshComponent.cpp
+++ b/Engine/Source/Runtime/Engine/Private/Components/StaticMeshComponent.cpp
@@ -1543,6 +1543,9 @@ void UStaticMeshComponent::PrecachePSOs()
 		return;
 	}
 
+	UWorld* World = GetWorld();
+	ERHIFeatureLevel::Type FeatureLevel = World ? World->FeatureLevel.GetValue() : GMaxRHIFeatureLevel;
+
 	bool bAnySectionCastsShadows = false;
 	TArray<int16, TInlineAllocator<2>> UsedMaterialIndices;
 	for (FStaticMeshLODResources& LODRenderData : StaticMesh->GetRenderData()->LODResources)
@@ -1569,6 +1572,9 @@ void UStaticMeshComponent::PrecachePSOs()
 		UMaterialInterface* MaterialInterface = GetMaterial(MaterialIndex);
 		if (MaterialInterface)
 		{
+			PrecachePSOParams.bHasWorldPositionOffsetVelocity = SupportsWorldPositionOffsetVelocity()
+				&& MaterialInterface->GetRelevance_Concurrent(FeatureLevel).bUsesWorldPositionOffset;
+
 			MaterialInterface->PrecachePSOs(VFType, PrecachePSOParams);
 		}
 	}
diff --git a/Engine/Source/Runtime/Engine/Private/Components/TextRenderComponent.cpp b/Engine/Source/Runtime/Engine/Private/Components/TextRenderComponent.cpp
--- a/Engine/Source/Runtime/Engine/Private/Components/TextRenderComponent.cpp
+++ b/Engine/Source/Runtime/Engine/Private/Components/TextRenderComponent.cpp
@@ -1363,6 +1363,12 @@ void UTextRenderComponent::PostLoad()
 	if (TextMaterial)
 	{
 		FPSOPrecacheParams PrecachePSOParams;
+		SetupPrecachePSOParams(PrecachePSOParams);
+
+		// Overwrite the cast shadow flag because MeshBatch setup in DrawStaticElements doesn't set the correct CastShadow flag from the proxy or component
+		// and leaves the default CastShadow value which is true
+		PrecachePSOParams.bCastShadow = true;
+
 		TextMaterial->PrecachePSOs(&FLocalVertexFactory::StaticType, PrecachePSOParams);
 	}
 
diff --git a/Engine/Source/Runtime/Engine/Private/Materials/Material.cpp b/Engine/Source/Runtime/Engine/Private/Materials/Material.cpp
--- a/Engine/Source/Runtime/Engine/Private/Materials/Material.cpp
+++ b/Engine/Source/Runtime/Engine/Private/Materials/Material.cpp
@@ -2503,11 +2503,8 @@ bool UMaterial::IsComplete() const
 FGraphEventArray UMaterial::PrecachePSOs(const TConstArrayView<const FVertexFactoryType*>& VertexFactoryTypes, const FPSOPrecacheParams& InPreCacheParams)
 {
 	FGraphEventArray GraphEvents;
-	if (FApp::CanEverRender() && PipelineStateCache::IsPSOPrecachingEnabled())
+	if (FApp::CanEverRender() && MaterialResources.Num() > 0 && PipelineStateCache::IsPSOPrecachingEnabled())
 	{
-		// Assume post load has been called on the UMaterial already and material resources are created
-		check(MaterialResources.Num() > 0);
-
 		EMaterialQualityLevel::Type ActiveQualityLevel = GetCachedScalabilityCVars().MaterialQualityLevel;
 		uint32 FeatureLevelsToCompile = GetFeatureLevelsToCompileForRendering();
 		while (FeatureLevelsToCompile != 0)
diff --git a/Engine/Source/Runtime/Engine/Private/Materials/MaterialShader.cpp b/Engine/Source/Runtime/Engine/Private/Materials/MaterialShader.cpp
--- a/Engine/Source/Runtime/Engine/Private/Materials/MaterialShader.cpp
+++ b/Engine/Source/Runtime/Engine/Private/Materials/MaterialShader.cpp
@@ -2732,6 +2732,8 @@ bool FMaterialShaderMap::IsComplete(const FMaterial* Material, bool bSilent)
 
 FGraphEventArray FMaterialShaderMap::CollectPSOs(ERHIFeatureLevel::Type InFeatureLevel, const FMaterial* Material, const TConstArrayView<const FVertexFactoryType*>& VertexFactoryTypes, const FPSOPrecacheParams& PreCacheParams)
 {
+	TRACE_CPUPROFILER_EVENT_SCOPE(FMaterialShaderMap::CollectPSOs);
+
 	// Only feature level is currently set as init settings - rest is default
 	// (multiview & alpha channel not taken into account here)
 	FSceneTexturesConfigInitSettings SceneTexturesConfigInitSettings;
diff --git a/Engine/Source/Runtime/Engine/Private/Materials/MaterialShared.cpp b/Engine/Source/Runtime/Engine/Private/Materials/MaterialShared.cpp
--- a/Engine/Source/Runtime/Engine/Private/Materials/MaterialShared.cpp
+++ b/Engine/Source/Runtime/Engine/Private/Materials/MaterialShared.cpp
@@ -2642,6 +2642,8 @@ bool FMaterial::CacheShaders(const FMaterialShaderMapId& ShaderMapId, EShaderPla
 
 FGraphEventArray FMaterial::CollectPSOs(ERHIFeatureLevel::Type InFeatureLevel, const TConstArrayView<const FVertexFactoryType*>& VertexFactoryTypes, const FPSOPrecacheParams& PreCacheParams)
 {
+	TRACE_CPUPROFILER_EVENT_SCOPE(FMaterial::CollectPSOs);
+
 	FGraphEventArray GraphEvents;
 
 	// Only care about inline shaders for now
diff --git a/Engine/Source/Runtime/Engine/Public/PSOPrecache.h b/Engine/Source/Runtime/Engine/Public/PSOPrecache.h
--- a/Engine/Source/Runtime/Engine/Public/PSOPrecache.h
+++ b/Engine/Source/Runtime/Engine/Public/PSOPrecache.h
@@ -31,12 +31,14 @@ struct FPSOPrecacheParams
 		bRenderInDepthPass = true;
 		bStaticLighting = true;
 		bCastShadow = true;
+		bRenderCustomDepth = false;
 		bAffectDynamicIndirectLighting = true;
 		bReverseCulling = false;
 		bDisableBackFaceCulling = false;
 		bCastShadowAsTwoSided = false;
 		bForceLODModel = false;
 		Mobility = (uint8)EComponentMobility::Static;
+		bHasWorldPositionOffsetVelocity = false;
 		StencilWriteMask = (uint8)EStencilMask::SM_Default;
 		Unused = 0;
 	}
@@ -91,6 +93,7 @@ struct FPSOPrecacheParams
 			uint32 bRenderInDepthPass : 1;
 			uint32 bStaticLighting : 1;
 			uint32 bCastShadow : 1;
+			uint32 bRenderCustomDepth : 1;
 
 			uint32 bAffectDynamicIndirectLighting : 1;
 			uint32 bReverseCulling : 1;
@@ -99,9 +102,10 @@ struct FPSOPrecacheParams
 			uint32 bForceLODModel : 1;
 
 			uint32 Mobility : 4;
+			uint32 bHasWorldPositionOffsetVelocity : 1;
 			uint32 StencilWriteMask : 4;
 
-			uint32 Unused : 9;
+			uint32 Unused : 7;
 		};
 		uint32 Data;
 	};
diff --git a/Engine/Source/Runtime/RenderCore/Private/ShaderCodeArchive.cpp b/Engine/Source/Runtime/RenderCore/Private/ShaderCodeArchive.cpp
--- a/Engine/Source/Runtime/RenderCore/Private/ShaderCodeArchive.cpp
+++ b/Engine/Source/Runtime/RenderCore/Private/ShaderCodeArchive.cpp
@@ -1997,6 +1997,8 @@ TRefCountPtr<FRHIShader> FIoStoreShaderCodeArchive::CreateShader(int32 ShaderInd
 	};
 #endif
 
+	TRACE_CPUPROFILER_EVENT_SCOPE(FIoStoreShaderCodeArchive::CreateShader);
+
 	TRefCountPtr<FRHIShader> Shader;
 
 	const FIoStoreShaderCodeEntry& ShaderEntry = Header.ShaderEntries[ShaderIndex];
diff --git a/Engine/Source/Runtime/Renderer/Private/CustomDepthRendering.cpp b/Engine/Source/Runtime/Renderer/Private/CustomDepthRendering.cpp
--- a/Engine/Source/Runtime/Renderer/Private/CustomDepthRendering.cpp
+++ b/Engine/Source/Runtime/Renderer/Private/CustomDepthRendering.cpp
@@ -451,7 +451,7 @@ void FCustomDepthPassMeshProcessor::CollectPSOInitializers(const FSceneTexturesC
 
 	bool bPositionOnly = false;
 	bool bUseDefaultMaterial = UseDefaultMaterial(Material, Material.MaterialModifiesMeshPosition_GameThread(), bSupportPositionOnlyStream, bPositionOnly);
-	if (!bUseDefaultMaterial)
+	if (!bUseDefaultMaterial && PreCacheParams.bRenderCustomDepth)
 	{
 		check(!bPositionOnly);
 
diff --git a/Engine/Source/Runtime/Renderer/Private/ShadowDepthRendering.cpp b/Engine/Source/Runtime/Renderer/Private/ShadowDepthRendering.cpp
--- a/Engine/Source/Runtime/Renderer/Private/ShadowDepthRendering.cpp
+++ b/Engine/Source/Runtime/Renderer/Private/ShadowDepthRendering.cpp
@@ -2132,7 +2132,8 @@ void FShadowDepthPassMeshProcessor::CollectPSOInitializers(const FSceneTexturesC
 	// Early out if possible
 	if (!Material.ShouldCastDynamicShadows() ||
 		!ShouldIncludeDomainInMeshPass(Material.GetMaterialDomain()) ||
-		!ShouldIncludeMaterialInDefaultOpaquePass(Material))
+		!ShouldIncludeMaterialInDefaultOpaquePass(Material) ||
+		!EnumHasAnyFlags(MeshSelectionMask, VertexFactoryType->SupportsPrimitiveIdStream() ? EShadowMeshSelection::VSM : EShadowMeshSelection::SM))
 	{
 		return;
 	}
@@ -2148,7 +2149,8 @@ void FShadowDepthPassMeshProcessor::CollectPSOInitializers(const FSceneTexturesC
 		CollectPSOInitializersForEachShadowDepthType(VertexFactoryType, Material, MeshFillMode, CM_CW, bCastShadowsAsTwoSided, PSOInitializers);
 		CollectPSOInitializersForEachShadowDepthType(VertexFactoryType, Material, MeshFillMode, CM_CCW, bCastShadowsAsTwoSided, PSOInitializers);
 	}
-	else if (PreCacheParams.bCastShadow && !UseDefaultMaterialForShadowDepth(Material, Material.MaterialModifiesMeshPosition_GameThread()))
+	else if (PreCacheParams.bCastShadow && 
+		!UseDefaultMaterialForShadowDepth(Material, Material.MaterialModifiesMeshPosition_GameThread()))
 	{
 		const FMeshDrawingPolicyOverrideSettings OverrideSettings = ComputeMeshOverrideSettings(PreCacheParams);
 		const ERasterizerFillMode MeshFillMode = ComputeMeshFillMode(Material, OverrideSettings);
diff --git a/Engine/Source/Runtime/Renderer/Private/VelocityRendering.cpp b/Engine/Source/Runtime/Renderer/Private/VelocityRendering.cpp
--- a/Engine/Source/Runtime/Renderer/Private/VelocityRendering.cpp
+++ b/Engine/Source/Runtime/Renderer/Private/VelocityRendering.cpp
@@ -582,7 +582,7 @@ void FOpaqueVelocityMeshProcessor::AddMeshBatch(
 void FOpaqueVelocityMeshProcessor::CollectPSOInitializers(const FSceneTexturesConfig& SceneTexturesConfig, const FMaterial& Material, const FVertexFactoryType* VertexFactoryType, const FPSOPrecacheParams& PreCacheParams, TArray<FPSOPrecacheData>& PSOInitializers)
 {
 	const EShaderPlatform ShaderPlatform = GetFeatureLevelShaderPlatform(FeatureLevel);
-	bool bDrawsVelocity = (PreCacheParams.Mobility == EComponentMobility::Movable || PreCacheParams.Mobility == EComponentMobility::Stationary);
+	bool bDrawsVelocity = (PreCacheParams.Mobility == EComponentMobility::Movable || PreCacheParams.Mobility == EComponentMobility::Stationary || PreCacheParams.bHasWorldPositionOffsetVelocity);
 	if (!PrimitiveCanHaveVelocity(ShaderPlatform, bDrawsVelocity, PreCacheParams.bStaticLighting))
 	{
 		return;
diff --git a/Engine/Source/Runtime/Renderer/Public/MeshPassProcessor.inl b/Engine/Source/Runtime/Renderer/Public/MeshPassProcessor.inl
--- a/Engine/Source/Runtime/Renderer/Public/MeshPassProcessor.inl
+++ b/Engine/Source/Runtime/Renderer/Public/MeshPassProcessor.inl
@@ -199,6 +199,8 @@ void FMeshPassProcessor::AddGraphicsPipelineStateInitializer(
 	EMeshPassFeatures MeshPassFeatures, 
 	TArray<FPSOPrecacheData>& PSOInitializers)
 {
+	TRACE_CPUPROFILER_EVENT_SCOPE(FMeshPassProcessor::AddGraphicsPipelineStateInitializer);
+
 	FGraphicsMinimalPipelineStateInitializer MinimalPipelineStateInitializer;
 	MinimalPipelineStateInitializer.PrimitiveType = PrimitiveType;
 	
