Enable non-blocking RTPSO creation by default

diff --git a/Engine/Source/Runtime/Renderer/Private/RayTracing/RayTracingMaterialHitShaders.cpp b/Engine/Source/Runtime/Renderer/Private/RayTracing/RayTracingMaterialHitShaders.cpp
--- a/Engine/Source/Runtime/Renderer/Private/RayTracing/RayTracingMaterialHitShaders.cpp
+++ b/Engine/Source/Runtime/Renderer/Private/RayTracing/RayTracingMaterialHitShaders.cpp
@@ -40,7 +40,7 @@ static FAutoConsoleVariableRef CVarCompileRayTracingMaterialAHS(
 	ECVF_ReadOnly
 );
 
-static int32 GRayTracingNonBlockingPipelineCreation = 0;
+static int32 GRayTracingNonBlockingPipelineCreation = 1;
 static FAutoConsoleVariableRef CVarRayTracingNonBlockingPipelineCreation(
 	TEXT("r.RayTracing.NonBlockingPipelineCreation"),
 	GRayTracingNonBlockingPipelineCreation,
