Add TAA pixel shader to utilize the hardward frame buffer compression on mobile devices.

diff --git a/Engine/Config/Android/AndroidEngine.ini b/Engine/Config/Android/AndroidEngine.ini
--- a/Engine/Config/Android/AndroidEngine.ini
+++ b/Engine/Config/Android/AndroidEngine.ini
@@ -100,6 +100,7 @@ bSupportsInAppPurchasing=true
 
 [/Script/Engine.RendererSettings]
 r.DistanceFields=0
+r.TemporalAA.Mobile.UseCompute=0
 
 [AlternateTextureCompression]
 TextureCompressionFormat=""
diff --git a/Engine/Config/IOS/IOSEngine.ini b/Engine/Config/IOS/IOSEngine.ini
--- a/Engine/Config/IOS/IOSEngine.ini
+++ b/Engine/Config/IOS/IOSEngine.ini
@@ -90,6 +90,7 @@ fx.Niagara.QualityLevel.Max=1
 
 [/Script/Engine.RendererSettings]
 r.DistanceFields=0
+r.TemporalAA.Mobile.UseCompute=0
 
 [AlternateTextureCompression]
 TextureCompressionFormat=""
diff --git a/Engine/Shaders/Private/TemporalAA.usf b/Engine/Shaders/Private/TemporalAA.usf
--- a/Engine/Shaders/Private/TemporalAA.usf
+++ b/Engine/Shaders/Private/TemporalAA.usf
@@ -2291,7 +2291,9 @@ FTAAHistoryPayload TemporalAASample(uint2 GroupId, uint2 GroupThreadId, uint Gro
 	return OutputPayload;
 }
 
-//------------------------------------------------------- ENTRY POINT
+//------------------------------------------------------- ENTRY POINTS
+
+#if COMPUTESHADER
 
 [numthreads(THREADGROUP_SIZEX, THREADGROUP_SIZEY, 1)]
 void MainCS(
@@ -2374,3 +2376,28 @@ void MainCS(
 	}
 	#endif //TAA_DOWNSAMPLE
 }
+
+#elif PIXELSHADER // Mobile Only.
+
+#if AA_UPSAMPLE
+#error Up sample only available to CS.
+#endif
+
+void MainPS(
+	float4 SvPosition : SV_POSITION,
+	out float4 OutColor0 : SV_Target0
+	)
+{
+	float2 ViewportUV = (SvPosition.xy - OutputViewportRect.xy) * OutputViewportSize.zw;
+
+	uint2 GroupId = int2(0, 0);
+	uint2 GroupThreadId = int2(0, 0);
+	uint GroupThreadIndex = 0;
+	float FrameExposureScale = EyeAdaptationLookup();
+	
+	FTAAHistoryPayload OutputPayload = TemporalAASample(GroupId, GroupThreadId, GroupThreadIndex, ViewportUV, FrameExposureScale);
+
+	OutColor0 = OutputPayload.Color;
+}
+
+#endif
diff --git a/Engine/Source/Runtime/Renderer/Private/PostProcess/TemporalAA.cpp b/Engine/Source/Runtime/Renderer/Private/PostProcess/TemporalAA.cpp
--- a/Engine/Source/Runtime/Renderer/Private/PostProcess/TemporalAA.cpp
+++ b/Engine/Source/Runtime/Renderer/Private/PostProcess/TemporalAA.cpp
@@ -76,6 +76,13 @@
 	TEXT(" 1: enabled;\n"),
 	ECVF_ReadOnly);
 
+TAutoConsoleVariable<int32> CVarTemporalAAMobileUseCompute(
+	TEXT("r.TemporalAA.Mobile.UseCompute"),
+	1,
+	TEXT(" 0: Uses pixel shader to save bandwidth with FBC on tiled gpu;\n")
+	TEXT(" 1: Uses compute shader (default);\n"),
+	ECVF_RenderThreadSafe);
+
 #if WITH_MGPU
 const FName TAAEffectName("TAA");
 #endif
@@ -86,22 +93,11 @@
 		&& IsFeatureLevelSupported(Platform, ERHIFeatureLevel::SM5);
 }
 
-class FTemporalAACS : public FGlobalShader
+class FTemporalAA : public FGlobalShader
 {
-	DECLARE_GLOBAL_SHADER(FTemporalAACS);
-	SHADER_USE_PARAMETER_STRUCT(FTemporalAACS, FGlobalShader);
-
+public:
 	class FTAAPassConfigDim : SHADER_PERMUTATION_ENUM_CLASS("TAA_PASS_CONFIG", ETAAPassConfig);
 	class FTAAQualityDim : SHADER_PERMUTATION_ENUM_CLASS("TAA_QUALITY", ETAAQuality);
-	class FTAAResponsiveDim : SHADER_PERMUTATION_BOOL("TAA_RESPONSIVE");
-	class FTAAScreenPercentageDim : SHADER_PERMUTATION_INT("TAA_SCREEN_PERCENTAGE_RANGE", 4);
-	class FTAADownsampleDim : SHADER_PERMUTATION_BOOL("TAA_DOWNSAMPLE");
 
-	using FPermutationDomain = TShaderPermutationDomain<
-		FTAAPassConfigDim,
-		FTAAQualityDim,
-		FTAAScreenPercentageDim,
-		FTAADownsampleDim>;
-
 	BEGIN_SHADER_PARAMETER_STRUCT(FParameters, )
 		SHADER_PARAMETER(FVector4f, ViewportUVToInputBufferUV)
@@ -148,6 +144,80 @@
 		SHADER_PARAMETER_RDG_TEXTURE_SRV(Texture2D, StencilTexture)
 
 		SHADER_PARAMETER_STRUCT_REF(FViewUniformShaderParameters, ViewUniformBuffer)
+	END_SHADER_PARAMETER_STRUCT()
 
+	static void ModifyCompilationEnvironment(const FGlobalShaderPermutationParameters& Parameters, FShaderCompilerEnvironment& OutEnvironment)
+	{
+		static FShaderPlatformCachedIniValue<bool> UseMobileConfig(TEXT("r.TemporalAA.UseMobileConfig"));
+		bool bUseMobileConfig = (UseMobileConfig.Get((EShaderPlatform)Parameters.Platform) != 0);
+
+		bool bIsMobileTiledGPU = RHIHasTiledGPU(Parameters.Platform) || IsSimulatedPlatform(Parameters.Platform);
+
+		// There are some mobile specific shader optimizations need to be set in the shader, such as disable shared memory usage, disable stencil texture sampling.
+		OutEnvironment.SetDefine(TEXT("AA_MOBILE_CONFIG"), (bIsMobileTiledGPU || bUseMobileConfig) ? 1 : 0);
+	}
+
+	FTemporalAA() = default;
+	FTemporalAA(const CompiledShaderInitializerType & Initializer)
+	:	FGlobalShader(Initializer)
+	{}
+}; // class FTemporalAA
+
+class FTemporalAAPS : public FTemporalAA
+{
+	DECLARE_GLOBAL_SHADER(FTemporalAAPS);
+	SHADER_USE_PARAMETER_STRUCT(FTemporalAAPS, FTemporalAA);
+
+	using FPermutationDomain = TShaderPermutationDomain<
+		FTemporalAA::FTAAPassConfigDim,
+		FTemporalAA::FTAAQualityDim>;
+
+	BEGIN_SHADER_PARAMETER_STRUCT(FParameters, )
+		SHADER_PARAMETER_STRUCT_INCLUDE(FTemporalAA::FParameters, Common)
+		RENDER_TARGET_BINDING_SLOTS()
+	END_SHADER_PARAMETER_STRUCT()
+
+	static bool ShouldCompilePermutation(const FGlobalShaderPermutationParameters& Parameters)
+	{
+		FPermutationDomain PermutationVector(Parameters.PermutationId);
+
+		// Pixel shader is only used on mobile to utilize the hardware frame buffer compression to save bandwidth
+		if (!IsMobilePlatform(Parameters.Platform))
+		{
+			return false;
+		}
+
+		// Only Main config is supported on pixel shader.
+		if (PermutationVector.Get<FTemporalAA::FTAAPassConfigDim>() != ETAAPassConfig::Main)
+		{
+			return false;
+		}
+
+		return SupportsGen4TAA(Parameters.Platform);
+	}
+
+	static void ModifyCompilationEnvironment(const FGlobalShaderPermutationParameters& Parameters, FShaderCompilerEnvironment& OutEnvironment)
+	{
+		FTemporalAA::ModifyCompilationEnvironment(Parameters, OutEnvironment);
+	}
+}; // class FTemporalAAPS
+
+class FTemporalAACS : public FTemporalAA
+{
+	DECLARE_GLOBAL_SHADER(FTemporalAACS);
+	SHADER_USE_PARAMETER_STRUCT(FTemporalAACS, FTemporalAA);
+
+	class FTAAScreenPercentageDim : SHADER_PERMUTATION_INT("TAA_SCREEN_PERCENTAGE_RANGE", 4);
+	class FTAADownsampleDim : SHADER_PERMUTATION_BOOL("TAA_DOWNSAMPLE");
+
+	using FPermutationDomain = TShaderPermutationDomain<
+		FTemporalAA::FTAAPassConfigDim,
+		FTemporalAA::FTAAQualityDim,
+		FTAAScreenPercentageDim,
+		FTAADownsampleDim>;
+
+	BEGIN_SHADER_PARAMETER_STRUCT(FParameters, )
+		SHADER_PARAMETER_STRUCT_INCLUDE(FTemporalAA::FParameters, Common)
+
 		// Temporal upsample specific parameters.
 		SHADER_PARAMETER(FVector4f, InputViewSize)
@@ -164,21 +234,21 @@
 
 	static FPermutationDomain RemapPermutation(FPermutationDomain PermutationVector)
 	{
-		if (PermutationVector.Get<FTAAPassConfigDim>() == ETAAPassConfig::Main ||
-			PermutationVector.Get<FTAAPassConfigDim>() == ETAAPassConfig::MainUpsampling)
+		if (PermutationVector.Get<FTemporalAA::FTAAPassConfigDim>() == ETAAPassConfig::Main ||
+			PermutationVector.Get<FTemporalAA::FTAAPassConfigDim>() == ETAAPassConfig::MainUpsampling)
 		{
 			// No point downsampling if not using the faster quality permutation already.
-			if (PermutationVector.Get<FTAAQualityDim>() == ETAAQuality::High)
+			if (PermutationVector.Get<FTemporalAA::FTAAQualityDim>() == ETAAQuality::High)
 			{
 				PermutationVector.Set<FTAADownsampleDim>(false);
 			}
 		}
 		else if (
-			PermutationVector.Get<FTAAPassConfigDim>() == ETAAPassConfig::DiaphragmDOF ||
-			PermutationVector.Get<FTAAPassConfigDim>() == ETAAPassConfig::DiaphragmDOFUpsampling)
+			PermutationVector.Get<FTemporalAA::FTAAPassConfigDim>() == ETAAPassConfig::DiaphragmDOF ||
+			PermutationVector.Get<FTemporalAA::FTAAPassConfigDim>() == ETAAPassConfig::DiaphragmDOFUpsampling)
 		{
 			// DOF pass allow only quality 1 and 2
-			PermutationVector.Set<FTAAQualityDim>(ETAAQuality(FMath::Max(int(PermutationVector.Get<FTAAQualityDim>()), int(ETAAQuality::Medium))));
+			PermutationVector.Set<FTemporalAA::FTAAQualityDim>(ETAAQuality(FMath::Max(int(PermutationVector.Get<FTemporalAA::FTAAQualityDim>()), int(ETAAQuality::Medium))));
 
 			// Only the Main and Main Upsampling can downsample the output.
 			PermutationVector.Set<FTAADownsampleDim>(false);
@@ -186,7 +256,7 @@
 		else
 		{
 			// Only the Main and Main Upsampling have quality options 0, 1, 2.
-			PermutationVector.Set<FTAAQualityDim>(ETAAQuality::High);
+			PermutationVector.Set<FTemporalAA::FTAAQualityDim>(ETAAQuality::High);
 
 			// Only the Main and Main Upsampling can downsample the output.
 			PermutationVector.Set<FTAADownsampleDim>(false);
@@ -206,7 +276,7 @@
 		}
 
 		// Screen percentage dimension is only for upsampling permutation.
-		if (!IsTAAUpsamplingConfig(PermutationVector.Get<FTAAPassConfigDim>()) &&
+		if (!IsTAAUpsamplingConfig(PermutationVector.Get<FTemporalAA::FTAAPassConfigDim>()) &&
 			PermutationVector.Get<FTAAScreenPercentageDim>() != 0)
 		{
 			return false;
@@ -222,31 +292,36 @@
 		}
 
 		// Screen percentage range 3 is only for super sampling.
-		if (PermutationVector.Get<FTAAPassConfigDim>() != ETAAPassConfig::MainSuperSampling &&
+		if (PermutationVector.Get<FTemporalAA::FTAAPassConfigDim>() != ETAAPassConfig::MainSuperSampling &&
 			PermutationVector.Get<FTAAScreenPercentageDim>() == 3)
 		{
 			return false;
 		}
 		
-		//Only Main and MainUpsampling config without DownSample permutations are supported on mobile platform.
-		return SupportsGen4TAA(Parameters.Platform) && (!IsMobilePlatform(Parameters.Platform) || ((PermutationVector.Get<FTAAPassConfigDim>() == ETAAPassConfig::Main || PermutationVector.Get<FTAAPassConfigDim>() == ETAAPassConfig::MainUpsampling) && !PermutationVector.Get<FTAADownsampleDim>()));
+		if (IsMobilePlatform(Parameters.Platform))
+		{
+			// Only Main and MainUpsampling config are supported on mobile platform.
+			if ((PermutationVector.Get<FTAAPassConfigDim>() != ETAAPassConfig::Main && PermutationVector.Get<FTAAPassConfigDim>() != ETAAPassConfig::MainUpsampling)
+				// DownSample is not supported on mobile platform
+				|| PermutationVector.Get<FTAADownsampleDim>())
+			{
+				return false;
+			}
+		}
+
+		return SupportsGen4TAA(Parameters.Platform);
 	}
 
 	static void ModifyCompilationEnvironment(const FGlobalShaderPermutationParameters& Parameters, FShaderCompilerEnvironment& OutEnvironment)
 	{
+		FTemporalAA::ModifyCompilationEnvironment(Parameters, OutEnvironment);
+
 		OutEnvironment.SetDefine(TEXT("THREADGROUP_SIZEX"), GTemporalAATileSizeX);
 		OutEnvironment.SetDefine(TEXT("THREADGROUP_SIZEY"), GTemporalAATileSizeY);
-
-		static FShaderPlatformCachedIniValue<bool> UseMobileConfig(TEXT("r.TemporalAA.UseMobileConfig"));
-		bool bUseMobileConfig = (UseMobileConfig.Get((EShaderPlatform)Parameters.Platform) != 0);
-
-		bool bIsMobileTiledGPU = RHIHasTiledGPU(Parameters.Platform) || IsSimulatedPlatform(Parameters.Platform);
-
-		// There are some mobile specific shader optimizations need to be set in the shader, such as disable shared memory usage, disable stencil texture sampling.
-		OutEnvironment.SetDefine(TEXT("AA_MOBILE_CONFIG"), (bIsMobileTiledGPU || bUseMobileConfig) ? 1 : 0);
 	}
 }; // class FTemporalAACS
 
+IMPLEMENT_GLOBAL_SHADER(FTemporalAAPS, "/Engine/Private/TemporalAA.usf", "MainPS", SF_Pixel);
 IMPLEMENT_GLOBAL_SHADER(FTemporalAACS, "/Engine/Private/TemporalAA.usf", "MainCS", SF_Compute);
 
 float CatmullRom(float x)
@@ -258,7 +333,7 @@
 		return (1.5f * ax - 2.5f) * ax*ax + 1.0f;
 }
 
-void SetupSampleWeightParameters(FTemporalAACS::FParameters* OutTAAParameters, const FTAAPassParameters& PassParameters, FVector2f TemporalJitterPixels)
+void SetupSampleWeightParameters(FTemporalAA::FParameters* OutTAAParameters, const FTAAPassParameters& PassParameters, FVector2f TemporalJitterPixels)
 {
 	float JitterX = TemporalJitterPixels.X;
 	float JitterY = TemporalJitterPixels.Y;
@@ -426,6 +501,11 @@
 	return UpscaleFactor;
 }
 
+bool DoesTemporalAAUseComputeShader(EShaderPlatform Platform)
+{
+	return !IsMobilePlatform(Platform) || CVarTemporalAAMobileUseCompute.GetValueOnAnyThread() > 0;
+}
+
 FIntPoint FTAAPassParameters::GetOutputExtent() const
 {
 	check(Validate());
@@ -498,6 +578,8 @@
 
 	TStaticArray<FRDGTextureRef, FTemporalAAHistory::kRenderTargetCount> NewHistoryTexture;
 
+	const bool bIsComputePass = DoesTemporalAAUseComputeShader(View.GetShaderPlatform());
+
 	{
 		EPixelFormat HistoryPixelFormat = PF_FloatRGBA;
 		if (bIsMainPass && Inputs.Quality != ETAAQuality::High && !bSupportsAlpha && CVarTAAR11G11B10History.GetValueOnRenderThread())
@@ -509,9 +591,14 @@
 			OutputExtent,
 			HistoryPixelFormat,
 			FClearValueBinding::Black,
-			TexCreate_ShaderResource | TexCreate_UAV);
+			TexCreate_ShaderResource);
 
-		if (Inputs.bOutputRenderTargetable)
+		if (bIsComputePass)
 		{
+			SceneColorDesc.Flags |= TexCreate_UAV;
+		}
+
+		if (Inputs.bOutputRenderTargetable || !bIsComputePass)
+		{
 			SceneColorDesc.Flags |= TexCreate_RenderTargetable;
 		}
@@ -535,6 +622,7 @@
 
 		if (Inputs.bDownsample)
 		{
+			check(bIsComputePass);
 			const FRDGTextureDesc HalfResSceneColorDesc = FRDGTextureDesc::Create2D(
 				SceneColorDesc.Extent / 2,
 				Inputs.DownsampleOverrideFormat != PF_Unknown ? Inputs.DownsampleOverrideFormat : Inputs.SceneColorInput->Desc.Format,
@@ -557,44 +645,13 @@
 
 	TStaticArray<bool, FTemporalAAHistory::kRenderTargetCount> bUseHistoryTexture;
 
+	// Setups common shader parameters
+	const FIntPoint InputExtent = Inputs.SceneColorInput->Desc.Extent;
+	const FIntRect InputViewRect = Inputs.InputViewRect;
+	const FIntRect OutputViewRect = Inputs.OutputViewRect;
+
+	auto SetupTemporalAACommonPassParameters = [&](FTemporalAA::FParameters* PassParameters)
 	{
-		FTemporalAACS::FPermutationDomain PermutationVector;
-		PermutationVector.Set<FTemporalAACS::FTAAPassConfigDim>(Inputs.Pass);
-		PermutationVector.Set<FTemporalAACS::FTAAQualityDim>(Inputs.Quality);
-		PermutationVector.Set<FTemporalAACS::FTAADownsampleDim>(Inputs.bDownsample);
-
-		if (IsTAAUpsamplingConfig(Inputs.Pass))
-		{
-			// If screen percentage > 100% on X or Y axes, then use screen percentage range = 2 shader permutation to disable LDS caching.
-			if (SrcRect.Width() > DestRect.Width() ||
-				SrcRect.Height() > DestRect.Height())
-			{
-				PermutationVector.Set<FTemporalAACS::FTAAScreenPercentageDim>(2);
-			}
-			// If screen percentage < 50% on X and Y axes, then use screen percentage range = 3 shader permutation.
-			else if (SrcRect.Width() * 100 < 50 * DestRect.Width() &&
-				SrcRect.Height() * 100 < 50 * DestRect.Height() &&
-				Inputs.Pass == ETAAPassConfig::MainSuperSampling)
-			{
-				PermutationVector.Set<FTemporalAACS::FTAAScreenPercentageDim>(3);
-			}
-			// If screen percentage < 71% on X and Y axes, then use screen percentage range = 1 shader permutation to have smaller LDS caching.
-			else if (SrcRect.Width() * 100 < 71 * DestRect.Width() &&
-				SrcRect.Height() * 100 < 71 * DestRect.Height())
-			{
-				PermutationVector.Set<FTemporalAACS::FTAAScreenPercentageDim>(1);
-			}
-		}
-
-		PermutationVector = FTemporalAACS::RemapPermutation(PermutationVector);
-
-		FTemporalAACS::FParameters* PassParameters = GraphBuilder.AllocParameters<FTemporalAACS::FParameters>();
-
-		// Setups common shader parameters
-		const FIntPoint InputExtent = Inputs.SceneColorInput->Desc.Extent;
-		const FIntRect InputViewRect = Inputs.InputViewRect;
-		const FIntRect OutputViewRect = Inputs.OutputViewRect;
-
 		if (!IsTAAUpsamplingConfig(Inputs.Pass))
 		{
 			SetupSampleWeightParameters(PassParameters, Inputs, FVector2f(View.TemporalJitterPixels));
@@ -739,6 +796,43 @@
 		{
 			PassParameters->EyeAdaptationTexture = GetEyeAdaptationTexture(GraphBuilder, View);
 		}
+	};
 
+	if (bIsComputePass)
+	{
+		FTemporalAACS::FPermutationDomain PermutationVector;
+		PermutationVector.Set<FTemporalAA::FTAAPassConfigDim>(Inputs.Pass);
+		PermutationVector.Set<FTemporalAA::FTAAQualityDim>(Inputs.Quality);
+		PermutationVector.Set<FTemporalAACS::FTAADownsampleDim>(Inputs.bDownsample);
+
+		if (IsTAAUpsamplingConfig(Inputs.Pass))
+		{
+			// If screen percentage > 100% on X or Y axes, then use screen percentage range = 2 shader permutation to disable LDS caching.
+			if (SrcRect.Width() > DestRect.Width() ||
+				SrcRect.Height() > DestRect.Height())
+			{
+				PermutationVector.Set<FTemporalAACS::FTAAScreenPercentageDim>(2);
+			}
+			// If screen percentage < 50% on X and Y axes, then use screen percentage range = 3 shader permutation.
+			else if (SrcRect.Width() * 100 < 50 * DestRect.Width() &&
+				SrcRect.Height() * 100 < 50 * DestRect.Height() &&
+				Inputs.Pass == ETAAPassConfig::MainSuperSampling)
+			{
+				PermutationVector.Set<FTemporalAACS::FTAAScreenPercentageDim>(3);
+			}
+			// If screen percentage < 71% on X and Y axes, then use screen percentage range = 1 shader permutation to have smaller LDS caching.
+			else if (SrcRect.Width() * 100 < 71 * DestRect.Width() &&
+				SrcRect.Height() * 100 < 71 * DestRect.Height())
+			{
+				PermutationVector.Set<FTemporalAACS::FTAAScreenPercentageDim>(1);
+			}
+		}
+
+		PermutationVector = FTemporalAACS::RemapPermutation(PermutationVector);
+
+		FTemporalAACS::FParameters* PassParameters = GraphBuilder.AllocParameters<FTemporalAACS::FParameters>();
+
+		SetupTemporalAACommonPassParameters(&PassParameters->Common);
+
 		// Temporal upsample specific shader parameters.
 		{
@@ -785,7 +879,7 @@
 		ClearUnusedGraphResources(ComputeShader, PassParameters);
 		for (int32 i = 0; i < FTemporalAAHistory::kRenderTargetCount; i++)
 		{
-			bUseHistoryTexture[i] = PassParameters->HistoryBuffer[i] != nullptr;
+			bUseHistoryTexture[i] = PassParameters->Common.HistoryBuffer[i] != nullptr;
 		}
 
 		{
@@ -797,7 +891,7 @@
 			GraphBuilder.AddPass(
 				RDG_EVENT_NAME("TAA(%s Quality=%s) %dx%d -> %dx%d",
 					PassName,
-					kTAAQualityNames[int32(PermutationVector.Get<FTemporalAACS::FTAAQualityDim>())],
+					kTAAQualityNames[int32(PermutationVector.Get<FTemporalAA::FTAAQualityDim>())],
 					PracticableSrcRect.Width(), PracticableSrcRect.Height(),
 					PracticableDestRect.Width(), PracticableDestRect.Height()),
 				ParametersMetadata,
@@ -825,6 +919,40 @@
 				});
 		}
 	}
+	else
+	{
+		check(IsMobilePlatform(View.GetShaderPlatform()) && !IsTAAUpsamplingConfig(Inputs.Pass));
 
+		FTemporalAAPS::FPermutationDomain PermutationVector;
+		PermutationVector.Set<FTemporalAA::FTAAPassConfigDim>(Inputs.Pass);
+		PermutationVector.Set<FTemporalAA::FTAAQualityDim>(Inputs.Quality);
+
+		FTemporalAAPS::FParameters* PassParameters = GraphBuilder.AllocParameters<FTemporalAAPS::FParameters>();
+
+		SetupTemporalAACommonPassParameters(&PassParameters->Common);
+
+		PassParameters->RenderTargets[0] = FRenderTargetBinding(
+			Outputs.SceneColor,
+			ERenderTargetLoadAction::EClear);
+
+		TShaderMapRef<FTemporalAAPS> PixelShader(View.ShaderMap, PermutationVector);
+
+		FPixelShaderUtils::AddFullscreenPass(
+			GraphBuilder,
+			View.ShaderMap,
+			RDG_EVENT_NAME("TAA(%s Quality=%s) %dx%d",
+				PassName,
+				kTAAQualityNames[int32(PermutationVector.Get<FTemporalAA::FTAAQualityDim>())],
+				PracticableDestRect.Width(), PracticableDestRect.Height()),
+			PixelShader,
+			PassParameters,
+			PracticableDestRect);
+
+		for (int32 i = 0; i < FTemporalAAHistory::kRenderTargetCount; i++)
+		{
+			bUseHistoryTexture[i] = PassParameters->Common.HistoryBuffer[i] != nullptr;
+		}
+	}
+
 	if (!View.bStatePrevViewInfoIsReadOnly)
 	{
diff --git a/Engine/Source/Runtime/Renderer/Private/SceneRendering.cpp b/Engine/Source/Runtime/Renderer/Private/SceneRendering.cpp
--- a/Engine/Source/Runtime/Renderer/Private/SceneRendering.cpp
+++ b/Engine/Source/Runtime/Renderer/Private/SceneRendering.cpp
@@ -2845,11 +2845,13 @@ void FSceneRenderer::PrepareViewRectsForRendering(FRHICommandListImmediate& RHIC
 
 		// Automatic screen percentage fallback.
 		{
+			extern bool DoesTemporalAAUseComputeShader(EShaderPlatform Platform);
 			// Tenmporal upsample is supported only if TAA is turned on.
 			if (View.PrimaryScreenPercentageMethod == EPrimaryScreenPercentageMethod::TemporalUpscale &&
 				(!IsTemporalAccumulationBasedMethod(View.AntiAliasingMethod) ||
 				 ViewFamily.EngineShowFlags.VisualizeBuffer || 
-				 ViewFamily.EngineShowFlags.VisualizeSSS))
+				 ViewFamily.EngineShowFlags.VisualizeSSS ||
+				 !DoesTemporalAAUseComputeShader(View.GetShaderPlatform())))
 			{
 				View.PrimaryScreenPercentageMethod = EPrimaryScreenPercentageMethod::SpatialUpscale;
 			}
