Restore changes lost when the DevPR copy-up robomerged back to UE5 Main

diff --git a/Engine/Source/Runtime/Core/Private/Async/TaskGraph.cpp b/Engine/Source/Runtime/Core/Private/Async/TaskGraph.cpp
--- a/Engine/Source/Runtime/Core/Private/Async/TaskGraph.cpp
+++ b/Engine/Source/Runtime/Core/Private/Async/TaskGraph.cpp
@@ -2481,6 +2481,7 @@ bool FTaskGraphInterface::IsMultithread()
 #if TASKGRAPH_NEW_FRONTEND
 
 FGraphEventImplAllocator GraphEventImplAllocator;
+FGraphTaskAllocator SmallTaskAllocator;
 
 #else
 
diff --git a/Engine/Source/Runtime/Engine/Private/Materials/MaterialShared.cpp b/Engine/Source/Runtime/Engine/Private/Materials/MaterialShared.cpp
--- a/Engine/Source/Runtime/Engine/Private/Materials/MaterialShared.cpp
+++ b/Engine/Source/Runtime/Engine/Private/Materials/MaterialShared.cpp
@@ -1098,6 +1098,17 @@ bool FMaterial::MaterialUsesAnisotropy_RenderThread() const
 	return RenderingThreadShaderMap ? RenderingThreadShaderMap->UsesAnisotropy() : false;
 }
 
+uint8 FMaterial::MaterialGetStrataMaterialType_GameThread() const
+{
+	return GameThreadShaderMap ? GameThreadShaderMap->GetStrataMaterialType() : false;
+}
+
+uint8 FMaterial::MaterialGetStrataMaterialType_RenderThread() const
+{
+	check(IsInParallelRenderingThread());
+	return RenderingThreadShaderMap ? RenderingThreadShaderMap->GetStrataMaterialType() : false;
+}
+
 void FMaterial::SetGameThreadShaderMap(FMaterialShaderMap* InMaterialShaderMap)
 {
 	checkSlow(IsInGameThread() || IsInAsyncLoadingThread());
diff --git a/Engine/Source/Runtime/Engine/Private/UnrealEngine.cpp b/Engine/Source/Runtime/Engine/Private/UnrealEngine.cpp
--- a/Engine/Source/Runtime/Engine/Private/UnrealEngine.cpp
+++ b/Engine/Source/Runtime/Engine/Private/UnrealEngine.cpp
@@ -5303,34 +5303,54 @@ static void DumpHelp(UWorld* InWorld)
 	UE_LOG(LogEngine, Display, TEXT("To browse console variables open this: '%s'"), *FilePath);
 	UE_LOG(LogEngine, Display, TEXT(" "));
 
-	ConsoleCommandLibrary_DumpLibraryHTML(InWorld, *GEngine, FilePath);
+	const bool bSuccess = ConsoleCommandLibrary_DumpLibraryHTML(InWorld, *GEngine, FilePath);
 
 	// Notification in editor
 #if WITH_EDITOR
+	if (bSuccess && FPaths::FileExists(FilePath))
 	{
-		const FText Message = NSLOCTEXT("UnrealEd", "ConsoleHelpExported", "ConsoleHelp.html was saved as");
+		const FText Message = NSLOCTEXT("UnrealEd", "ConsoleHelpExported", "Console help file saved at:");
 		FNotificationInfo Info(Message);
 		Info.bFireAndForget = true;
 		Info.ExpireDuration = 5.0f;
 		Info.bUseSuccessFailIcons = false;
 		Info.bUseLargeFont = false;
 
-		const FString HyperLinkText = FPaths::ConvertRelativePathToFull(FilePath);
-		Info.Hyperlink = FSimpleDelegate::CreateStatic([](FString SourceFilePath) 
+		const FString HyperLinkPath = FPaths::ConvertRelativePathToFull(FilePath);
+		Info.SubText = FText::FromString(HyperLinkPath);
+		Info.Hyperlink = FSimpleDelegate::CreateStatic([](FString SourceFilePath)
 		{
 			// open folder, you can choose the browser yourself
 			FPlatformProcess::ExploreFolder(*(FPaths::GetPath(SourceFilePath)));
-		}, HyperLinkText);
-		Info.HyperlinkText = FText::FromString(HyperLinkText);
+		}, HyperLinkPath);
+		Info.HyperlinkText = NSLOCTEXT("UnrealEd", "ConsoleHelpFolderLink", "Open file location...");
+		Info.WidthOverride = FOptionalSize();
 
 		FSlateNotificationManager::Get().AddNotification(Info);
 
 		// Always try to open the help file on Windows (including in -game, etc...)
 #if PLATFORM_WINDOWS
-		const FString LaunchableURL = FString(TEXT("file://")) + HyperLinkText;
+		const FString LaunchableURL = FString(TEXT("file://")) + HyperLinkPath;
 		FPlatformProcess::LaunchURL(*LaunchableURL, nullptr, nullptr);
 #endif
 	}
+	else
+	{
+		const FText Message = NSLOCTEXT("UnrealEd", "ConsoleHelpExportFailed", "Console help file generation failed.");
+		FNotificationInfo Info(Message);
+
+		const FString DocTemplateFile = FPaths::Combine(FPaths::EngineDir(), "Documentation/Extras/ConsoleHelpTemplate.html");
+		if (!FPaths::FileExists(DocTemplateFile))
+		{
+				Info.SubText = NSLOCTEXT("UnrealEd", "ConsoleHelpExportFailedTip", "The template HTML file could not be found in the Engine/Documentation/Extras folder.\n\nIf you use a version control system or Unreal Game Sync (UGS), check that this folder is not excluded by your sync settings.");
+		}
+		Info.bFireAndForget = true;
+		Info.ExpireDuration = 10.0f;
+		Info.bUseSuccessFailIcons = false;
+		Info.bUseLargeFont = false;
+
+		FSlateNotificationManager::Get().AddNotification(Info);
+	}
 #endif// WITH_EDITOR
 }
 static FAutoConsoleCommandWithWorld GConsoleCommandHelp(
diff --git a/Engine/Source/Runtime/Launch/Private/LaunchEngineLoop.cpp b/Engine/Source/Runtime/Launch/Private/LaunchEngineLoop.cpp
--- a/Engine/Source/Runtime/Launch/Private/LaunchEngineLoop.cpp
+++ b/Engine/Source/Runtime/Launch/Private/LaunchEngineLoop.cpp
@@ -4159,6 +4159,10 @@ bool FEngineLoop::LoadStartupCoreModules()
 	// Required during serialization of AnimSequence which could happen from async loading thread.
 	// See UAnimSequence::UpdateFrameRate().
 	FModuleManager::Get().LoadModule(TEXT("TimeManagement"));
+
+	// Required during construction of UAnimBlueprint which could happen from async loading thread.
+	// See UAnimBlueprint::UAnimBlueprint().
+	FModuleManager::Get().LoadModule(TEXT("AnimGraph"));
 #endif
 
 	FModuleManager::Get().LoadModule(TEXT("PacketHandler"));
diff --git a/Engine/Source/Runtime/RenderCore/Private/RenderingThread.cpp b/Engine/Source/Runtime/RenderCore/Private/RenderingThread.cpp
--- a/Engine/Source/Runtime/RenderCore/Private/RenderingThread.cpp
+++ b/Engine/Source/Runtime/RenderCore/Private/RenderingThread.cpp
@@ -992,11 +992,8 @@ void StartRenderCommandFenceBundler()
 	STAT_FNullGraphTask_FenceRenderCommandBundled,
 		STATGROUP_TaskGraphTasks);
 
-	FGraphEventRef Task = TGraphTask<FNullGraphTask>::CreateTask(&Prereqs, ENamedThreads::GameThread).ConstructAndDispatchWhenReady(
+	BundledCompletionEvent = TGraphTask<FNullGraphTask>::CreateTask(&Prereqs, ENamedThreads::GameThread).ConstructAndDispatchWhenReady(
 		GET_STATID(STAT_FNullGraphTask_FenceRenderCommandBundled), ENamedThreads::GetRenderThread());
-	// this reference can live long, store completion handle instead of a reference to the task to reduce peak mem usage 
-	BundledCompletionEvent = Task->CreateCompletionHandle();
-	
 
 	StartBatchedRelease();
 }
diff --git a/Engine/Source/Runtime/Renderer/Private/DeferredShadingRenderer.cpp b/Engine/Source/Runtime/Renderer/Private/DeferredShadingRenderer.cpp
--- a/Engine/Source/Runtime/Renderer/Private/DeferredShadingRenderer.cpp
+++ b/Engine/Source/Runtime/Renderer/Private/DeferredShadingRenderer.cpp
@@ -651,14 +651,16 @@ static void GatherRayTracingRelevantPrimitives(const FScene& Scene, const FViewI
 
 		for (int PrimitiveIndex = 0; PrimitiveIndex < Scene.PrimitiveSceneProxies.Num(); PrimitiveIndex++)
 		{
-			// Find the next TypeOffsetTable entry that's relevant to this primitive idnex.
+			// Find the next TypeOffsetTable entry that's relevant to this primitive index.
 			while (PrimitiveIndex >= int(Scene.TypeOffsetTable[BroadIndex].Offset))
 			{
 				BroadIndex++;
 			}
 
+			const ERayTracingPrimitiveFlags Flags = Scene.PrimitiveRayTracingFlags[PrimitiveIndex];
+
 			// Skip before dereferencing SceneInfo
-			if (Scene.PrimitiveRayTracingFlags[PrimitiveIndex] == ERayTracingPrimitiveFlags::UnsupportedProxyType)
+			if (Flags == ERayTracingPrimitiveFlags::UnsupportedProxyType)
 			{
 				// Find the index of a proxy of the next type, skipping over a batch of proxies that are the same type as current.
 				// This assumes that FPrimitiveSceneProxy::IsRayTracingRelevant() is consistent for all proxies of the same type.
@@ -697,7 +699,7 @@ static void GatherRayTracingRelevantPrimitives(const FScene& Scene, const FViewI
 			}
 
 			// Marked visible and used after point, check if streaming then mark as used in the TLAS (so it can be streamed in)
-			if (EnumHasAnyFlags(Scene.PrimitiveRayTracingFlags[PrimitiveIndex], ERayTracingPrimitiveFlags::Streaming))
+			if (EnumHasAnyFlags(Flags, ERayTracingPrimitiveFlags::Streaming))
 			{
 				// Is the cached data dirty?
 				if (SceneInfo->bCachedRaytracingDataDirty)
@@ -716,7 +718,7 @@ static void GatherRayTracingRelevantPrimitives(const FScene& Scene, const FViewI
 			FRayTracingRelevantPrimitive Item;
 			Item.PrimitiveIndex = PrimitiveIndex;
 
-			if (EnumHasAnyFlags(Scene.PrimitiveRayTracingFlags[PrimitiveIndex], ERayTracingPrimitiveFlags::StaticMesh)
+			if (EnumHasAnyFlags(Flags, ERayTracingPrimitiveFlags::StaticMesh)
 				&& View.Family->EngineShowFlags.StaticMeshes
 				&& RayTracingStaticMeshesCVar && RayTracingStaticMeshesCVar->GetValueOnRenderThread() > 0)
 			{
@@ -725,7 +727,7 @@ static void GatherRayTracingRelevantPrimitives(const FScene& Scene, const FViewI
 			}
 			else if (View.Family->EngineShowFlags.SkeletalMeshes)
 			{
-				checkf(!EnumHasAllFlags(Scene.PrimitiveRayTracingFlags[PrimitiveIndex], ERayTracingPrimitiveFlags::CacheInstances), 
+				checkf(!EnumHasAllFlags(Flags, ERayTracingPrimitiveFlags::CacheInstances),
 					TEXT("Only static primitives are expected to use CacheInstances flag."));
 
 				Item.bStatic = false;
@@ -750,10 +752,11 @@ static void GatherRayTracingRelevantPrimitives(const FScene& Scene, const FViewI
 
 			const int32 PrimitiveIndex = RelevantPrimitive.PrimitiveIndex;
 			const FPrimitiveSceneInfo* SceneInfo = Scene.Primitives[PrimitiveIndex];
+			const ERayTracingPrimitiveFlags Flags = Scene.PrimitiveRayTracingFlags[PrimitiveIndex];
 
 			int8 LODIndex = 0;
 
-			if (EnumHasAnyFlags(Scene.PrimitiveRayTracingFlags[PrimitiveIndex], ERayTracingPrimitiveFlags::ComputeLOD))
+			if (EnumHasAnyFlags(Flags, ERayTracingPrimitiveFlags::ComputeLOD))
 			{
 				const FPrimitiveBounds& Bounds = Scene.PrimitiveBounds[PrimitiveIndex];
 				const FPrimitiveSceneInfo* RESTRICT PrimitiveSceneInfo = Scene.Primitives[PrimitiveIndex];
@@ -770,7 +773,7 @@ static void GatherRayTracingRelevantPrimitives(const FScene& Scene, const FViewI
 				LODIndex = LODToRender.GetRayTracedLOD();
 			}
 
-			if (!EnumHasAllFlags(Scene.PrimitiveRayTracingFlags[PrimitiveIndex], ERayTracingPrimitiveFlags::CacheInstances))
+			if (!EnumHasAllFlags(Flags, ERayTracingPrimitiveFlags::CacheInstances))
 			{
 				FRHIRayTracingGeometry* RayTracingGeometryInstance = SceneInfo->GetStaticRayTracingGeometryInstance(LODIndex);
 				if (RayTracingGeometryInstance == nullptr)
@@ -780,6 +783,7 @@ static void GatherRayTracingRelevantPrimitives(const FScene& Scene, const FViewI
 
 				// Sometimes LODIndex is out of range because it is clamped by ClampToFirstLOD, like the requested LOD is being streamed in and hasn't been available
 				// According to InitViews, we should hide the static mesh instance
+				check(EnumHasAnyFlags(Flags, ERayTracingPrimitiveFlags::CacheMeshCommands));
 				if (SceneInfo->CachedRayTracingMeshCommandIndicesPerLOD.IsValidIndex(LODIndex))
 				{
 					RelevantPrimitive.LODIndex = LODIndex;
@@ -811,7 +815,7 @@ static void GatherRayTracingRelevantPrimitives(const FScene& Scene, const FViewI
 
 					RelevantPrimitive.InstanceMask |= RelevantPrimitive.bAnySegmentsCastShadow ? RAY_TRACING_MASK_SHADOW : 0;
 
-					if (EnumHasAllFlags(Scene.PrimitiveRayTracingFlags[PrimitiveIndex], ERayTracingPrimitiveFlags::FarField))
+					if (EnumHasAllFlags(Flags, ERayTracingPrimitiveFlags::FarField))
 					{
 						RelevantPrimitive.InstanceMask = RAY_TRACING_MASK_FAR_FIELD;
 					}
@@ -2270,7 +2274,7 @@ void FDeferredShadingSceneRenderer::Render(FRDGBuilder& GraphBuilder)
 	// Important that this uses consistent logic throughout the frame, so evaluate once and pass in the flag from here
 	// NOTE: Must be done after  system texture initialization
 	// TODO: This doesn't take into account the potential for split screen views with separate shadow caches
-	VirtualShadowMapArray.Initialize(GraphBuilder, Scene->GetVirtualShadowMapCache(Views[0]), UseVirtualShadowMaps(ShaderPlatform, FeatureLevel));
+	VirtualShadowMapArray.Initialize(GraphBuilder, Scene->GetVirtualShadowMapCache(Views[0]), UseVirtualShadowMaps(ShaderPlatform, FeatureLevel), Views[0].bIsSceneCapture);
 
 	// if DDM_AllOpaqueNoVelocity was used, then velocity should have already been rendered as well
 	const bool bIsEarlyDepthComplete = (DepthPass.EarlyZPassMode == DDM_AllOpaque || DepthPass.EarlyZPassMode == DDM_AllOpaqueNoVelocity);
@@ -3127,7 +3131,8 @@ void FDeferredShadingSceneRenderer::Render(FRDGBuilder& GraphBuilder)
 	// This needs to run before virtual shadow map, in order to have ready&cleared classified SSS data
 	if (Strata::IsStrataEnabled())
 	{
-		Strata::AddStrataMaterialClassificationPass(GraphBuilder, SceneTextures, Views);
+		Strata::AddStrataMaterialClassificationPass(GraphBuilder, SceneTextures, DBufferTextures, Views);
+		Strata::AddStrataDBufferPass(GraphBuilder, SceneTextures, DBufferTextures, Views);
 	}
 
 #if RHI_RAYTRACING
@@ -3412,7 +3417,8 @@ void FDeferredShadingSceneRenderer::Render(FRDGBuilder& GraphBuilder)
 			RDG_CSV_STAT_EXCLUSIVE_SCOPE(GraphBuilder, RenderTranslucency);
 			SCOPE_CYCLE_COUNTER(STAT_TranslucencyDrawTime);
 			GraphBuilder.SetCommandListStat(GET_STATID(STAT_CLM_Translucency));
-			RenderTranslucency(GraphBuilder, SceneTextures, TranslucencyLightingVolumeTextures, &TranslucencyResourceMap, ETranslucencyView::UnderWater, InstanceCullingManager);
+			const bool bStandardTranslucentCanRenderSeparate = false;
+			RenderTranslucency(GraphBuilder, SceneTextures, TranslucencyLightingVolumeTextures, &TranslucencyResourceMap, ETranslucencyView::UnderWater, InstanceCullingManager, bStandardTranslucentCanRenderSeparate);
 			EnumRemoveFlags(TranslucencyViewsToRender, ETranslucencyView::UnderWater);
 		}
 
@@ -3474,7 +3480,7 @@ void FDeferredShadingSceneRenderer::Render(FRDGBuilder& GraphBuilder)
 	if (GetHairStrandsComposition() == EHairStrandsCompositionType::BeforeTranslucent)
 	{
 		RDG_GPU_STAT_SCOPE(GraphBuilder, HairRendering);
-		RenderHairComposition(GraphBuilder, Views, SceneTextures.Color.Target, SceneTextures.Depth.Target, SceneTextures.Velocity);
+		RenderHairComposition(GraphBuilder, Views, SceneTextures.Color.Target, SceneTextures.Depth.Target, SceneTextures.Velocity, TranslucencyResourceMap);
 	}
 
 	// Draw translucency.
@@ -3486,7 +3492,7 @@ void FDeferredShadingSceneRenderer::Render(FRDGBuilder& GraphBuilder)
 		RDG_EVENT_SCOPE(GraphBuilder, "Translucency");
 
 		// Raytracing doesn't need the distortion effect.
-		const bool bShouldRenderDistortion = TranslucencyViewsToRender != ETranslucencyView::RayTracing;
+		const bool bShouldRenderDistortion = TranslucencyViewsToRender != ETranslucencyView::RayTracing && ShouldRenderDistortion();
 
 #if RHI_RAYTRACING
 		if (EnumHasAnyFlags(TranslucencyViewsToRender, ETranslucencyView::RayTracing))
@@ -3513,23 +3519,26 @@ void FDeferredShadingSceneRenderer::Render(FRDGBuilder& GraphBuilder)
 			}
 		}
 
-		// Render all remaining translucency views.
-		GraphBuilder.SetCommandListStat(GET_STATID(STAT_CLM_Translucency));
-		RenderTranslucency(GraphBuilder, SceneTextures, TranslucencyLightingVolumeTextures, &TranslucencyResourceMap, TranslucencyViewsToRender, InstanceCullingManager);
-		TranslucencyViewsToRender = ETranslucencyView::None;
+		{
+			// Render all remaining translucency views.
+			GraphBuilder.SetCommandListStat(GET_STATID(STAT_CLM_Translucency));
+			const bool bStandardTranslucentCanRenderSeparate = bShouldRenderDistortion; // It is only needed to render standard translucent as separate when there is distortion (non self distortion of transmittance/specular/etc.)
+			RenderTranslucency(GraphBuilder, SceneTextures, TranslucencyLightingVolumeTextures, &TranslucencyResourceMap, TranslucencyViewsToRender, InstanceCullingManager, bStandardTranslucentCanRenderSeparate);
+			TranslucencyViewsToRender = ETranslucencyView::None;
+		}
 
 		// Compose hair before velocity/distortion pass since these pass write depth value, 
 		// and this would make the hair composition fails in this cases.
 		if (GetHairStrandsComposition() == EHairStrandsCompositionType::AfterTranslucent)
 		{
 			RDG_GPU_STAT_SCOPE(GraphBuilder, HairRendering);
-			RenderHairComposition(GraphBuilder, Views, SceneTextures.Color.Target, SceneTextures.Depth.Target, SceneTextures.Velocity);
+			RenderHairComposition(GraphBuilder, Views, SceneTextures.Color.Target, SceneTextures.Depth.Target, SceneTextures.Velocity, TranslucencyResourceMap);
 		}
 
 		if (bShouldRenderDistortion)
 		{
 			GraphBuilder.SetCommandListStat(GET_STATID(STAT_CLM_Distortion));
-			RenderDistortion(GraphBuilder, SceneTextures.Color.Target, SceneTextures.Depth.Target);
+			RenderDistortion(GraphBuilder, SceneTextures.Color.Target, SceneTextures.Depth.Target, TranslucencyResourceMap);
 		}
 
 		if (bShouldRenderVelocities)
@@ -3551,7 +3560,7 @@ void FDeferredShadingSceneRenderer::Render(FRDGBuilder& GraphBuilder)
 	else if (GetHairStrandsComposition() == EHairStrandsCompositionType::AfterTranslucent)
 	{
 		RDG_GPU_STAT_SCOPE(GraphBuilder, HairRendering);
-		RenderHairComposition(GraphBuilder, Views, SceneTextures.Color.Target, SceneTextures.Depth.Target, SceneTextures.Velocity);
+		RenderHairComposition(GraphBuilder, Views, SceneTextures.Color.Target, SceneTextures.Depth.Target, SceneTextures.Velocity, TranslucencyResourceMap);
 	}
 
 #if !UE_BUILD_SHIPPING
diff --git a/Engine/Source/Runtime/Renderer/Private/MeshDrawCommands.cpp b/Engine/Source/Runtime/Renderer/Private/MeshDrawCommands.cpp
--- a/Engine/Source/Runtime/Renderer/Private/MeshDrawCommands.cpp
+++ b/Engine/Source/Runtime/Renderer/Private/MeshDrawCommands.cpp
@@ -1196,6 +1196,7 @@ void FParallelMeshDrawCommandPass::DispatchPassSetup(
 	switch (PassType)
 	{
 		case EMeshPass::TranslucencyStandard: TaskContext.TranslucencyPass = ETranslucencyPass::TPT_StandardTranslucency; break;
+		case EMeshPass::TranslucencyStandardModulate: TaskContext.TranslucencyPass = ETranslucencyPass::TPT_StandardTranslucencyModulate; break;
 		case EMeshPass::TranslucencyAfterDOF: TaskContext.TranslucencyPass = ETranslucencyPass::TPT_TranslucencyAfterDOF; break;
 		case EMeshPass::TranslucencyAfterDOFModulate: TaskContext.TranslucencyPass = ETranslucencyPass::TPT_TranslucencyAfterDOFModulate; break;
 		case EMeshPass::TranslucencyAfterMotionBlur: TaskContext.TranslucencyPass = ETranslucencyPass::TPT_TranslucencyAfterMotionBlur; break;
diff --git a/Engine/Source/Runtime/Renderer/Private/SceneRendering.h b/Engine/Source/Runtime/Renderer/Private/SceneRendering.h
--- a/Engine/Source/Runtime/Renderer/Private/SceneRendering.h
+++ b/Engine/Source/Runtime/Renderer/Private/SceneRendering.h
@@ -1391,11 +1391,17 @@ public:
 	 */
 	uint32 bHasSingleLayerWaterMaterial : 1;
 	/**
-	 * true if the scene has at least one mesh with a material that needs dual blending AND is applied post DOF. If true,
-	 * that means we need to run the separate modulation render pass.
+	 * true if the scene has at least one mesh with a material that needs dual blending AND is applied post DOF. 
+	 * If true, that means we need to run the post-dof separate modulation render pass.
 	 */
 	uint32 bHasTranslucencySeparateModulation : 1;
 
+	/**
+	 * true if the scene has at least one mesh with a material that needs dual blending AND is applied before DOF. 
+	 * If true, that means we need to run the before-dof separate modulation render pass.
+	 */
+	uint32 bHasStandardTranslucencyModulation : 1;
+
 	/** Bitmask of all shading models used by primitives in this view */
 	uint16 ShadingModelMaskInView;
 
