Fixed: Occlusion queries related assert when switching to mobile preview
Do FenceOcclusionTests in a separate pass on mobile

diff --git a/Engine/Source/Runtime/Renderer/Private/MobileShadingRenderer.cpp b/Engine/Source/Runtime/Renderer/Private/MobileShadingRenderer.cpp
--- a/Engine/Source/Runtime/Renderer/Private/MobileShadingRenderer.cpp
+++ b/Engine/Source/Runtime/Renderer/Private/MobileShadingRenderer.cpp
@@ -744,6 +744,8 @@ void FMobileSceneRenderer::RenderFullDepthPrepass(FRDGBuilder& GraphBuilder, FSc
 				RenderOcclusion(RHICmdList, View);
 			});
 	}
+
+	FenceOcclusionTests(GraphBuilder);
 }
 
 void FMobileSceneRenderer::RenderMaskedPrePass(FRHICommandListImmediate& RHICmdList, const FViewInfo& View)
@@ -946,6 +948,11 @@ void FMobileSceneRenderer::Render(FRDGBuilder& GraphBuilder)
 		RenderForward(GraphBuilder, ViewFamilyTexture, SceneTextures);
 	}
 
+	if (!bIsFullDepthPrepassEnabled)
+	{
+		FenceOcclusionTests(GraphBuilder);
+	}
+
 	SceneTextures.MobileSetupMode = EMobileSceneTextureSetupMode::All;
 	SceneTextures.MobileSetupMode &= ~EMobileSceneTextureSetupMode::SceneVelocity;
 	SceneTextures.MobileUniformBuffer = CreateMobileSceneTextureUniformBuffer(GraphBuilder, &SceneTextures, SceneTextures.MobileSetupMode);
diff --git a/Engine/Source/Runtime/Renderer/Private/SceneOcclusion.cpp b/Engine/Source/Runtime/Renderer/Private/SceneOcclusion.cpp
--- a/Engine/Source/Runtime/Renderer/Private/SceneOcclusion.cpp
+++ b/Engine/Source/Runtime/Renderer/Private/SceneOcclusion.cpp
@@ -1456,11 +1456,6 @@ void FMobileSceneRenderer::RenderOcclusion(FRHICommandListImmediate& RHICmdList,
 			BeginOcclusionTests(RHICmdList, Views, FeatureLevel, QueriesPerView, 1.0f);
 		}
 	}
-
-	if (IsRunningRHIInSeparateThread())
-	{
-		FenceOcclusionTestsInternal(RHICmdList);
-	}
 }
 
 DECLARE_CYCLE_STAT(TEXT("OcclusionSubmittedFence Dispatch"), STAT_OcclusionSubmittedFence_Dispatch, STATGROUP_SceneRendering);
@@ -1509,7 +1504,7 @@ void FSceneRenderer::FenceOcclusionTestsInternal(FRHICommandListImmediate& RHICm
 
 void FSceneRenderer::FenceOcclusionTests(FRDGBuilder& GraphBuilder)
 {
-	if (IsRunningRHIInSeparateThread())
+	if (DoOcclusionQueries() && IsRunningRHIInSeparateThread())
 	{
 		AddPass(GraphBuilder, RDG_EVENT_NAME("FenceOcclusionTests"), [this](FRHICommandListImmediate& RHICmdList)
 		{
