Force disabling deprecated RT reflection when Strata is enabled, as it is not supported.

diff --git a/Engine/Source/Runtime/Renderer/Private/RayTracing/RayTracingReflections.cpp b/Engine/Source/Runtime/Renderer/Private/RayTracing/RayTracingReflections.cpp
--- a/Engine/Source/Runtime/Renderer/Private/RayTracing/RayTracingReflections.cpp
+++ b/Engine/Source/Runtime/Renderer/Private/RayTracing/RayTracingReflections.cpp
@@ -397,7 +397,7 @@ bool ShouldRenderRayTracingReflections(const FViewInfo& View)
 		? bThisViewHasRaytracingReflections 
 		: GRayTracingReflections != 0;
 
-	const bool bReflectionPassEnabled = bReflectionsCvarEnabled && (GetRayTracingReflectionsSamplesPerPixel(View) > 0) && !View.bIsReflectionCapture;
+	const bool bReflectionPassEnabled = bReflectionsCvarEnabled && (GetRayTracingReflectionsSamplesPerPixel(View) > 0) && !View.bIsReflectionCapture && !Strata::IsStrataEnabled();
 		
 	return ShouldRenderRayTracingEffect(bReflectionPassEnabled, ERayTracingPipelineCompatibilityFlags::FullPipeline, &View);
 }
