Fix static meshes going through dynamic raytracing primitive codepath when EngineShowFlags.StaticMeshes or r.RayTracing.Geometry.StaticMeshes are disabled.

diff --git a/Engine/Source/Runtime/Renderer/Private/DeferredShadingRenderer.cpp b/Engine/Source/Runtime/Renderer/Private/DeferredShadingRenderer.cpp
--- a/Engine/Source/Runtime/Renderer/Private/DeferredShadingRenderer.cpp
+++ b/Engine/Source/Runtime/Renderer/Private/DeferredShadingRenderer.cpp
@@ -711,19 +711,20 @@ static void GatherRayTracingRelevantPrimitives(const FScene& Scene, const FViewI
 				Result.UsedCoarseMeshStreamingHandles.Add(SceneInfo->CoarseMeshStreamingHandle);
 			}
 
-			//#dxr_todo UE-68621  The Raytracing code path does not support ShowFlags since data moved to the SceneInfo. 
-			//Touching the SceneProxy to determine this would simply cost too much
-			static const auto RayTracingStaticMeshesCVar = IConsoleManager::Get().FindTConsoleVariableDataInt(TEXT("r.RayTracing.Geometry.StaticMeshes"));
-
 			FRayTracingRelevantPrimitive Item;
 			Item.PrimitiveIndex = PrimitiveIndex;
 
-			if (EnumHasAnyFlags(Flags, ERayTracingPrimitiveFlags::StaticMesh)
-				&& View.Family->EngineShowFlags.StaticMeshes
-				&& RayTracingStaticMeshesCVar && RayTracingStaticMeshesCVar->GetValueOnRenderThread() > 0)
+			if (EnumHasAnyFlags(Flags, ERayTracingPrimitiveFlags::StaticMesh))
 			{
-				Item.bStatic = true;
-				Result.StaticPrimitives.AddElement(Item);
+				//#dxr_todo UE-68621  The Raytracing code path does not support ShowFlags since data moved to the SceneInfo. 
+				//Touching the SceneProxy to determine this would simply cost too much
+				static const auto RayTracingStaticMeshesCVar = IConsoleManager::Get().FindTConsoleVariableDataInt(TEXT("r.RayTracing.Geometry.StaticMeshes"));
+
+				if (View.Family->EngineShowFlags.StaticMeshes && RayTracingStaticMeshesCVar && RayTracingStaticMeshesCVar->GetValueOnRenderThread() > 0)
+				{
+					Item.bStatic = true;
+					Result.StaticPrimitives.AddElement(Item);
+				}
 			}
 			else if (View.Family->EngineShowFlags.SkeletalMeshes)
 			{
