Fixed single frame uniform buffer semantics being used on multi frame uniform buffer.

diff --git a/Engine/Source/Runtime/Renderer/Private/RendererScene.cpp b/Engine/Source/Runtime/Renderer/Private/RendererScene.cpp
--- a/Engine/Source/Runtime/Renderer/Private/RendererScene.cpp
+++ b/Engine/Source/Runtime/Renderer/Private/RendererScene.cpp
@@ -1744,7 +1744,7 @@ void FPersistentUniformBuffers::Initialize()
 	for (int32 Index = 0; Index < UE_ARRAY_COUNT(MobileDirectionalLightUniformBuffers); ++Index)
 	{
 		// UniformBuffer_SingleFrame here is an optimization as this buffer gets uploaded everyframe
-		MobileDirectionalLightUniformBuffers[Index] = TUniformBufferRef<FMobileDirectionalLightShaderParameters>::CreateUniformBufferImmediate(MobileDirectionalLightShaderParameters, UniformBuffer_SingleFrame, EUniformBufferValidation::None);
+		MobileDirectionalLightUniformBuffers[Index] = TUniformBufferRef<FMobileDirectionalLightShaderParameters>::CreateUniformBufferImmediate(MobileDirectionalLightShaderParameters, UniformBuffer_MultiFrame, EUniformBufferValidation::None);
 	}
 
 	const FMobileReflectionCaptureShaderParameters* DefaultMobileSkyReflectionParameters = (const FMobileReflectionCaptureShaderParameters*)GDefaultMobileReflectionCaptureUniformBuffer.GetContents();
