Use a single shared sampler for Landscape weightmap textures

diff --git a/Engine/Source/Runtime/Engine/Classes/Engine/EngineTypes.h b/Engine/Source/Runtime/Engine/Classes/Engine/EngineTypes.h
--- a/Engine/Source/Runtime/Engine/Classes/Engine/EngineTypes.h
+++ b/Engine/Source/Runtime/Engine/Classes/Engine/EngineTypes.h
@@ -296,7 +296,9 @@ enum ESamplerSourceMode
 	/** Shared sampler source that does not consume a sampler slot.  Uses wrap addressing and gets filter mode from the world texture group. */
 	SSM_Wrap_WorldGroupSettings UMETA(DisplayName="Shared: Wrap"),
 	/** Shared sampler source that does not consume a sampler slot.  Uses clamp addressing and gets filter mode from the world texture group. */
-	SSM_Clamp_WorldGroupSettings UMETA(DisplayName="Shared: Clamp")
+	SSM_Clamp_WorldGroupSettings UMETA(DisplayName="Shared: Clamp"),
+	/** Shared sampler source that does not consume a sampler slot, used to sample the terrain weightmap.  Gets filter mode from the terrain weightmap texture group. */
+	SSM_TerrainWeightmapGroupSettings UMETA(Hidden)
 };
 
 /** defines how MipValue is used */
diff --git a/Engine/Source/Runtime/Engine/Private/Materials/HLSLMaterialTranslator.cpp b/Engine/Source/Runtime/Engine/Private/Materials/HLSLMaterialTranslator.cpp
--- a/Engine/Source/Runtime/Engine/Private/Materials/HLSLMaterialTranslator.cpp
+++ b/Engine/Source/Runtime/Engine/Private/Materials/HLSLMaterialTranslator.cpp
@@ -6012,6 +6012,7 @@ int32 FHLSLMaterialTranslator::TextureSample(
 			StaticAddressZ = TA_Wrap;
 			break;
 		case SSM_Clamp_WorldGroupSettings:
+		case SSM_TerrainWeightmapGroupSettings:
 			StaticAddressX = TA_Clamp;
 			StaticAddressY = TA_Clamp;
 			StaticAddressZ = TA_Clamp;
@@ -6129,6 +6130,12 @@ int32 FHLSLMaterialTranslator::TextureSample(
 				*TextureName, AutomaticViewMipBias ? TEXT("View.MaterialTextureBilinearClampedSampler") : TEXT("Material.Clamp_WorldGroupSettings"));
 			RequiresManualViewMipBias = false;
 		}
+		else if (SamplerSource == SSM_TerrainWeightmapGroupSettings)
+		{
+			SamplerStateCode = FString::Printf(TEXT("GetMaterialSharedSampler(%sSampler,%s)"),
+				*TextureName, TEXT("View.LandscapeWeightmapSampler"));
+			RequiresManualViewMipBias = false;
+		}
 	}
 
 	const EDerivativeStatus UvDerivativeStatus = GetDerivativeStatus(CoordinateIndex);
@@ -7092,11 +7099,11 @@ int32 FHLSLMaterialTranslator::StaticTerrainLayerWeight(FName LayerName,int32 De
 	}
 	else
 	{			
-		const EMaterialSamplerType SamplerType = SAMPLERTYPE_Masks;
+		constexpr EMaterialSamplerType SamplerType = SAMPLERTYPE_Masks;
 		FString WeightmapName = FString::Printf(TEXT("Weightmap%d"),WeightmapIndex);
 		int32 TextureReferenceIndex = INDEX_NONE;
 		int32 TextureCodeIndex = TextureParameter(FName(*WeightmapName), GEngine->WeightMapPlaceholderTexture, TextureReferenceIndex, SamplerType);
-		int32 WeightmapCode = TextureSample(TextureCodeIndex, TextureCoordinate(3, false, false), SamplerType);
+		int32 WeightmapCode = TextureSample(TextureCodeIndex, TextureCoordinate(3, false, false), SamplerType, /*MipValue0Index = */INDEX_NONE, /*MipValue1Index = */INDEX_NONE, /*MipValueMode = */TMVM_None, /*SamplerSource = */SSM_TerrainWeightmapGroupSettings);
 
 		FString LayerMaskName = FString::Printf(TEXT("LayerMask_%s"),*LayerName.ToString());
 		return Dot(WeightmapCode,VectorParameter(FName(*LayerMaskName), FLinearColor(1.f,0.f,0.f,0.f)));
diff --git a/Engine/Source/Runtime/Engine/Private/Materials/MaterialHLSLTree.cpp b/Engine/Source/Runtime/Engine/Private/Materials/MaterialHLSLTree.cpp
--- a/Engine/Source/Runtime/Engine/Private/Materials/MaterialHLSLTree.cpp
+++ b/Engine/Source/Runtime/Engine/Private/Materials/MaterialHLSLTree.cpp
@@ -882,6 +882,7 @@ void FExpressionTextureSample::EmitValueShader(FEmitContext& Context, FEmitScope
 			StaticAddressZ = TA_Wrap;
 			break;
 		case SSM_Clamp_WorldGroupSettings:
+		case SSM_TerrainWeightmapGroupSettings:
 			StaticAddressX = TA_Clamp;
 			StaticAddressY = TA_Clamp;
 			StaticAddressZ = TA_Clamp;
@@ -1025,6 +1026,11 @@ void FExpressionTextureSample::EmitValueShader(FEmitContext& Context, FEmitScope
 				FormattedTexture.ToString(),
 				AutomaticViewMipBias ? TEXT("View.MaterialTextureBilinearClampedSampler") : TEXT("Material.Clamp_WorldGroupSettings"));
 			break;
+		case SSM_TerrainWeightmapGroupSettings:
+			FormattedSampler.Appendf(TEXT("GetMaterialSharedSampler(%sSampler,%s)"),
+				FormattedTexture.ToString(),
+				TEXT("View.LandscapeWeightmapSampler"));
+			break;
 		default:
 			checkNoEntry();
 			break;
diff --git a/Engine/Source/Runtime/Engine/Private/Materials/MaterialUniformExpressions.cpp b/Engine/Source/Runtime/Engine/Private/Materials/MaterialUniformExpressions.cpp
--- a/Engine/Source/Runtime/Engine/Private/Materials/MaterialUniformExpressions.cpp
+++ b/Engine/Source/Runtime/Engine/Private/Materials/MaterialUniformExpressions.cpp
@@ -973,6 +973,7 @@ void FUniformExpressionSet::FillUniformBuffer(const FMaterialRenderContext& Mate
 				{
 					SamplerSource = &Clamp_WorldGroupSettings->SamplerStateRHI;
 				}
+				// NOTE: for SSM_TerrainWeightmapGroupSettings, the generated code always tries to use the per-view sampler, but we can fallback to the texture-associated sampler if necessary
 
 				if (*SamplerSource)
 				{
@@ -1026,6 +1027,7 @@ void FUniformExpressionSet::FillUniformBuffer(const FMaterialRenderContext& Mate
 				{
 					SamplerSource = &Clamp_WorldGroupSettings->SamplerStateRHI;
 				}
+				check(SourceMode != SSM_TerrainWeightmapGroupSettings); // not allowed for cubemaps
 
 				check(*SamplerSource);
 				*ResourceTableSamplerPtr = *SamplerSource;
@@ -1064,6 +1066,7 @@ void FUniformExpressionSet::FillUniformBuffer(const FMaterialRenderContext& Mate
 				{
 					SamplerSource = &Clamp_WorldGroupSettings->SamplerStateRHI;
 				}
+				check(SourceMode != SSM_TerrainWeightmapGroupSettings); // not allowed for texture arrays
 
 				check(*SamplerSource);
 				*ResourceTableSamplerPtr = *SamplerSource;
@@ -1102,6 +1105,7 @@ void FUniformExpressionSet::FillUniformBuffer(const FMaterialRenderContext& Mate
 				{
 					SamplerSource = &Clamp_WorldGroupSettings->SamplerStateRHI;
 				}
+				check(SourceMode != SSM_TerrainWeightmapGroupSettings); // not allowed for texture cube arrays
 
 				check(*SamplerSource);
 				*ResourceTableSamplerPtr = *SamplerSource;
@@ -1142,6 +1146,7 @@ void FUniformExpressionSet::FillUniformBuffer(const FMaterialRenderContext& Mate
 				{
 					SamplerSource = &Clamp_WorldGroupSettings->SamplerStateRHI;
 				}
+				check(SourceMode != SSM_TerrainWeightmapGroupSettings); // not allowed for texture volumes
 
 				check(*SamplerSource);
 				*ResourceTableSamplerPtr = *SamplerSource;
diff --git a/Engine/Source/Runtime/Engine/Private/SceneManagement.cpp b/Engine/Source/Runtime/Engine/Private/SceneManagement.cpp
--- a/Engine/Source/Runtime/Engine/Private/SceneManagement.cpp
+++ b/Engine/Source/Runtime/Engine/Private/SceneManagement.cpp
@@ -829,6 +829,8 @@ FViewUniformShaderParameters::FViewUniformShaderParameters()
 	WaterIndirection = GWhiteVertexBufferWithSRV->ShaderResourceViewRHI;
 	WaterData = GWhiteVertexBufferWithSRV->ShaderResourceViewRHI;
 
+	// Landscape
+	LandscapeWeightmapSampler = TStaticSamplerState<SF_AnisotropicPoint, AM_Clamp, AM_Clamp, AM_Clamp>::GetRHI();
 	LandscapeIndirection = GWhiteVertexBufferWithSRV->ShaderResourceViewRHI;
 	LandscapePerComponentData = GWhiteVertexBufferWithSRV->ShaderResourceViewRHI;
 
diff --git a/Engine/Source/Runtime/Engine/Public/SceneView.h b/Engine/Source/Runtime/Engine/Public/SceneView.h
--- a/Engine/Source/Runtime/Engine/Public/SceneView.h
+++ b/Engine/Source/Runtime/Engine/Public/SceneView.h
@@ -909,7 +909,8 @@ BEGIN_GLOBAL_SHADER_PARAMETER_STRUCT_WITH_CONSTRUCTOR(FViewUniformShaderParamete
 	SHADER_PARAMETER(float, RectLightAtlasMaxMipLevel)
 	SHADER_PARAMETER_TEXTURE(Texture2D<float4>, RectLightAtlasTexture)
 	SHADER_PARAMETER_SAMPLER(SamplerState, RectLightAtlasSampler)
-
+	// Landscape
+	SHADER_PARAMETER_SAMPLER(SamplerState, LandscapeWeightmapSampler)
 	SHADER_PARAMETER_SRV(Buffer<uint>, LandscapeIndirection)
 	SHADER_PARAMETER_SRV(Buffer<float>, LandscapePerComponentData)
 
diff --git a/Engine/Source/Runtime/Renderer/Private/RendererScene.cpp b/Engine/Source/Runtime/Renderer/Private/RendererScene.cpp
--- a/Engine/Source/Runtime/Renderer/Private/RendererScene.cpp
+++ b/Engine/Source/Runtime/Renderer/Private/RendererScene.cpp
@@ -355,6 +355,7 @@ FSceneViewState::FSceneViewState(ERHIFeatureLevel::Type FeatureLevel, FSceneView
 	
 	// Sets the mipbias to invalid large number.
 	MaterialTextureCachedMipBias = BIG_NUMBER;
+	LandscapeCachedMipBias = BIG_NUMBER;
 
 	SequencerState = ESS_None;
 
diff --git a/Engine/Source/Runtime/Renderer/Private/ScenePrivate.h b/Engine/Source/Runtime/Renderer/Private/ScenePrivate.h
--- a/Engine/Source/Runtime/Renderer/Private/ScenePrivate.h
+++ b/Engine/Source/Runtime/Renderer/Private/ScenePrivate.h
@@ -963,8 +963,10 @@ public:
 
 	// Cached material texture samplers
 	float MaterialTextureCachedMipBias;
+	float LandscapeCachedMipBias;
 	FSamplerStateRHIRef MaterialTextureBilinearWrapedSamplerCache;
 	FSamplerStateRHIRef MaterialTextureBilinearClampedSamplerCache;
+	FSamplerStateRHIRef LandscapeWeightmapSamplerCache;
 
 #if RHI_RAYTRACING
 	// Invalidates cached results related to the path tracer so accumulated rendering can start over
diff --git a/Engine/Source/Runtime/Renderer/Private/SceneRendering.cpp b/Engine/Source/Runtime/Renderer/Private/SceneRendering.cpp
--- a/Engine/Source/Runtime/Renderer/Private/SceneRendering.cpp
+++ b/Engine/Source/Runtime/Renderer/Private/SceneRendering.cpp
@@ -907,15 +907,16 @@ void FViewInfo::Init()
 
 	PreExposure = 1.0f;
 
-	// Cache TEXTUREGROUP_World's for the render thread to create the material textures' shared sampler.
+	// Cache TEXTUREGROUP filter settings for the render thread to create shared samplers.
 	if (IsInGameThread())
 	{
 		WorldTextureGroupSamplerFilter = (ESamplerFilter)UDeviceProfileManager::Get().GetActiveProfile()->GetTextureLODSettings()->GetSamplerFilter(TEXTUREGROUP_World);
-		bIsValidWorldTextureGroupSamplerFilter = true;
+		TerrainWeightmapTextureGroupSamplerFilter = (ESamplerFilter)UDeviceProfileManager::Get().GetActiveProfile()->GetTextureLODSettings()->GetSamplerFilter(TEXTUREGROUP_Terrain_Weightmap);
+		bIsValidTextureGroupSamplerFilters = true;
 	}
 	else
 	{
-		bIsValidWorldTextureGroupSamplerFilter = false;
+		bIsValidTextureGroupSamplerFilters = false;
 	}
 
 	PrimitiveSceneDataOverrideSRV = nullptr;
@@ -1431,12 +1432,13 @@ void FViewInfo::SetupUniformBufferParameters(
 	SetupPhysicsFieldUniformBufferParameters(Scene, Family->EngineShowFlags, ViewUniformShaderParameters);
 
 	// Setup view's shared sampler for material texture sampling.
+	float FinalMaterialTextureMipBias;
 	{
 		const float GlobalMipBias = UTexture2D::GetGlobalMipMapLODBias();
 
-		float FinalMaterialTextureMipBias = GlobalMipBias;
+		FinalMaterialTextureMipBias = GlobalMipBias;
 
-		if (bIsValidWorldTextureGroupSamplerFilter && !FMath::IsNearlyZero(MaterialTextureMipBias))
+		if (bIsValidTextureGroupSamplerFilters && !FMath::IsNearlyZero(MaterialTextureMipBias))
 		{
 			ViewUniformShaderParameters.MaterialTextureMipBias = MaterialTextureMipBias;
 			ViewUniformShaderParameters.MaterialTextureDerivativeMultiply = FMath::Pow(2.0f, MaterialTextureMipBias);
@@ -1459,7 +1461,7 @@ void FViewInfo::SetupUniformBufferParameters(
 		}
 		else
 		{
-			check(bIsValidWorldTextureGroupSamplerFilter);
+			check(bIsValidTextureGroupSamplerFilters);
 
 			WrappedSampler = RHICreateSamplerState(FSamplerStateInitializerRHI(WorldTextureGroupSamplerFilter, AM_Wrap,  AM_Wrap,  AM_Wrap,  FinalMaterialTextureMipBias));
 			ClampedSampler = RHICreateSamplerState(FSamplerStateInitializerRHI(WorldTextureGroupSamplerFilter, AM_Clamp, AM_Clamp, AM_Clamp, FinalMaterialTextureMipBias));
@@ -1879,6 +1881,29 @@ void FViewInfo::SetupUniformBufferParameters(
 	}
 
 	// Landscape global resources
+	{
+		FSamplerStateRHIRef WeightmapSampler = nullptr;
+		if (ViewState && FMath::Abs(ViewState->LandscapeCachedMipBias - FinalMaterialTextureMipBias) < KINDA_SMALL_NUMBER)
+		{
+			// use cached sampler
+			WeightmapSampler = ViewState->LandscapeWeightmapSamplerCache;
+		}
+		else
+		{
+			// create a new one
+			ESamplerFilter Filter = bIsValidTextureGroupSamplerFilters ? TerrainWeightmapTextureGroupSamplerFilter : SF_AnisotropicPoint;
+			WeightmapSampler = RHICreateSamplerState(FSamplerStateInitializerRHI(Filter, AM_Clamp, AM_Clamp, AM_Clamp, FinalMaterialTextureMipBias));
+		}
+		check(WeightmapSampler.IsValid());
+		ViewUniformShaderParameters.LandscapeWeightmapSampler = WeightmapSampler;
+
+		if (ViewState)
+		{
+			ViewState->LandscapeCachedMipBias = FinalMaterialTextureMipBias;
+			ViewState->LandscapeWeightmapSamplerCache = WeightmapSampler;
+		}
+	}
+
 	if (LandscapePerComponentDataBuffer.IsValid() && LandscapeIndirectionBuffer.IsValid())
 	{
 		ViewUniformShaderParameters.LandscapeIndirection = LandscapeIndirectionBuffer.GetReference();
diff --git a/Engine/Source/Runtime/Renderer/Private/SceneRendering.h b/Engine/Source/Runtime/Renderer/Private/SceneRendering.h
--- a/Engine/Source/Runtime/Renderer/Private/SceneRendering.h
+++ b/Engine/Source/Runtime/Renderer/Private/SceneRendering.h
@@ -1736,7 +1736,8 @@ private:
 	// Cache of TEXTUREGROUP_World to create view's samplers on render thread.
 	// may not have a valid value if FViewInfo is created on the render thread.
 	ESamplerFilter WorldTextureGroupSamplerFilter;
-	bool bIsValidWorldTextureGroupSamplerFilter;
+	ESamplerFilter TerrainWeightmapTextureGroupSamplerFilter;
+	bool bIsValidTextureGroupSamplerFilters;
 
 	FSceneViewState* GetEyeAdaptationViewState() const;
 
