Fixed ViewFamily->FrameNumber not being properly initialized when a new ViewFamily was created for ScreenCapture, solving flickering LOD issues with skeletal meshes in editor and PIE when SceneCaptureComponents are in use. The initialization is now done in SetupViewFamilyForSceneCapture.

diff --git a/Engine/Source/Runtime/Renderer/Private/SceneCaptureRendering.cpp b/Engine/Source/Runtime/Renderer/Private/SceneCaptureRendering.cpp
--- a/Engine/Source/Runtime/Renderer/Private/SceneCaptureRendering.cpp
+++ b/Engine/Source/Runtime/Renderer/Private/SceneCaptureRendering.cpp
@@ -596,6 +596,9 @@ void SetupViewFamilyForSceneCapture(
 	// For cube map capture, CubeMapFaceIndex takes precedence over view index, so we must have only one view for that case
 	check(CubemapFaceIndex == INDEX_NONE || Views.Num() == 1);
 
+	// Initialize frame number
+	ViewFamily.FrameNumber = ViewFamily.Scene->GetFrameNumber();
+
 	// Disable Nanite in planar reflections for now because Nanite doesn't support the global clip plane
 	if (bIsPlanarReflection)
 	{
