Fixed SetupRayTracingLightDataForViews being called after WaitForRayTracingScene, when Lumen is setup to use HW tracing and the 'Reuse Shadowmaps' show flag is disabled.  This showed up as a crash due to null uniform buffer ptr.

diff --git a/Engine/Source/Runtime/Renderer/Private/DeferredShadingRenderer.cpp b/Engine/Source/Runtime/Renderer/Private/DeferredShadingRenderer.cpp
--- a/Engine/Source/Runtime/Renderer/Private/DeferredShadingRenderer.cpp
+++ b/Engine/Source/Runtime/Renderer/Private/DeferredShadingRenderer.cpp
@@ -2994,6 +2994,9 @@ void FDeferredShadingSceneRenderer::Render(FRDGBuilder& GraphBuilder)
 	{
 		bRayTracingSceneReady = true;
 	}
+
+	/** Should be called somewhere before "WaitForRayTracingScene" */
+	SetupRayTracingLightDataForViews(GraphBuilder);
 #endif
 
 	if (!bHasRayTracedOverlay && !LumenSceneDirectLighting::AllowShadowMaps(ViewFamily.EngineShowFlags))
@@ -3146,11 +3149,6 @@ void FDeferredShadingSceneRenderer::Render(FRDGBuilder& GraphBuilder)
 		Strata::AddStrataDBufferPass(GraphBuilder, SceneTextures, DBufferTextures, Views);
 	}
 
-#if RHI_RAYTRACING
-	/** Should be called somewhere before "WaitForRayTracingScene" */
-	SetupRayTracingLightDataForViews(GraphBuilder);
-#endif // RHI_RAYTRACING
-
 	// Copy lighting channels out of stencil before deferred decals which overwrite those values
 	FRDGTextureRef LightingChannelsTexture = CopyStencilToLightingChannelTexture(GraphBuilder, SceneTextures.Stencil);
 
