Fix motion blur appearing on moving mesh that has mobility changed from movable to static.
Issue is that RemoveFromScene() is called with bImmediate false.
This defers removal from the map, and sets PrimitiveSceneInfo to nullptr instead to signify that we should ignore.
That needs to be respected in GetComponentPreviousLocalToWorld() or we will set velocity to true when we reinit the primitive uniform buffer.

diff --git a/Engine/Source/Runtime/Renderer/Private/ScenePrivate.h b/Engine/Source/Runtime/Renderer/Private/ScenePrivate.h
--- a/Engine/Source/Runtime/Renderer/Private/ScenePrivate.h
+++ b/Engine/Source/Runtime/Renderer/Private/ScenePrivate.h
@@ -2426,7 +2426,7 @@ public:
 	{
 		const FComponentVelocityData* VelocityData = ComponentData.Find(PrimitiveComponentId);
 
-		if (VelocityData)
+		if (VelocityData && VelocityData->PrimitiveSceneInfo)
 		{
 			check(VelocityData->bPreviousLocalToWorldValid);
 			VelocityData->LastFrameUsed = InternalFrameIndex;
