Proper color conversion for editor selection outline in SDR/HDR

diff --git a/Engine/Shaders/Private/PostProcessSelectionOutline.usf b/Engine/Shaders/Private/PostProcessSelectionOutline.usf
--- a/Engine/Shaders/Private/PostProcessSelectionOutline.usf
+++ b/Engine/Shaders/Private/PostProcessSelectionOutline.usf
@@ -2,6 +2,8 @@
 
 #include "Common.ush"
 #include "ScreenPass.ush"
+#include "GammaCorrectionCommon.ush"
+#include "ACES.ush"
 
 #if MSAA_SAMPLE_COUNT > 1
 Texture2DMS<float, MSAA_SAMPLE_COUNT> EditorPrimitivesDepth;
@@ -26,6 +28,7 @@ float3 OutlineColor;
 float3 SubduedOutlineColor;
 float SelectionHighlightIntensity;
 float BSPSelectionIntensity;
+float UILuminanceAndIsSCRGB;
 
 struct FPixelInfo
 {
@@ -267,10 +270,38 @@ void MainPS(
 		// Scene color has gamma applied. Need to remove and re-apply after linear operation.
 		if (Avg.a > 0)
 		{
-			OutColor.rgb = max(OutColor.rgb, 0.f);
-			OutColor.rgb = pow(OutColor.rgb, 2.2f );
-		    OutColor.rgb = lerp(OutColor.rgb, Avg.rgb, Avg.a);
-		    OutColor.rgb = pow(OutColor.rgb, 1.0f / 2.2f);
+#if SELECTION_OUTLINE_HDR
+			static const float ScRGBScaleFactor = 80.0f; // used to convert between ScRGB and nits
+			Avg.rgb *= abs(UILuminanceAndIsSCRGB);
+			if (UILuminanceAndIsSCRGB > 0.0)
+			{
+				OutColor.rgb *= ScRGBScaleFactor;
+			}
+			else
+			{
+				const float3x3 Rec2020_2_sRGB = mul(XYZ_2_sRGB_MAT, Rec2020_2_XYZ_MAT);
+				OutColor.rgb = ST2084ToLinear(OutColor.rgb);
+				OutColor.rgb = mul(Rec2020_2_sRGB, OutColor.rgb);
+			}
+#else
+			OutColor.rgb = sRGBToLinear(OutColor.rgb);
+#endif
+			OutColor.rgb = lerp(OutColor.rgb, Avg.rgb, Avg.a);
+
+#if SELECTION_OUTLINE_HDR
+			if (UILuminanceAndIsSCRGB > 0.0)
+			{
+				OutColor.rgb /= ScRGBScaleFactor;
+			}
+			else
+			{
+				const float3x3 sRGB_2_Rec2020 = mul(XYZ_2_Rec2020_MAT, sRGB_2_XYZ_MAT);
+				OutColor.xyz = mul(sRGB_2_Rec2020, OutColor.xyz);
+				OutColor.rgb = LinearToST2084(OutColor.rgb);
+			}
+#else
+			OutColor.rgb = LinearToSrgbBranchless(OutColor.rgb);
+#endif
 		}
 	}
 	else
diff --git a/Engine/Source/Runtime/Renderer/Private/PostProcess/PostProcessCompositeEditorPrimitives.h b/Engine/Source/Runtime/Renderer/Private/PostProcess/PostProcessCompositeEditorPrimitives.h
--- a/Engine/Source/Runtime/Renderer/Private/PostProcess/PostProcessCompositeEditorPrimitives.h
+++ b/Engine/Source/Runtime/Renderer/Private/PostProcess/PostProcessCompositeEditorPrimitives.h
@@ -16,9 +16,8 @@ public:
 	class FSampleCountDimension : SHADER_PERMUTATION_RANGE_INT("MSAA_SAMPLE_COUNT", 1, kEditorMSAASampleCountMax + 1);
 	using FPermutationDomain = TShaderPermutationDomain<FSampleCountDimension>;
 
-	static bool ShouldCompilePermutation(const FGlobalShaderPermutationParameters& Parameters)
+	static bool ShouldCompilePermutation(const FPermutationDomain& PermutationVector, const EShaderPlatform Platform)
 	{
-		const FPermutationDomain PermutationVector(Parameters.PermutationId);
 		const int32 SampleCount = PermutationVector.Get<FSampleCountDimension>();
 
 		// Only use permutations with valid MSAA sample counts.
@@ -28,7 +27,13 @@ public:
 		}
 
 		// Only PC platforms render editor primitives.
-		return IsPCPlatform(Parameters.Platform);
+		return IsPCPlatform(Platform);
+	}
+
+	static bool ShouldCompilePermutation(const FGlobalShaderPermutationParameters& Parameters)
+	{
+		const FPermutationDomain PermutationVector(Parameters.PermutationId);
+		return ShouldCompilePermutation(PermutationVector, Parameters.Platform);
 	}
 
 	FEditorPrimitiveShader() = default;
diff --git a/Engine/Source/Runtime/Renderer/Private/PostProcess/PostProcessSelectionOutline.cpp b/Engine/Source/Runtime/Renderer/Private/PostProcess/PostProcessSelectionOutline.cpp
--- a/Engine/Source/Runtime/Renderer/Private/PostProcess/PostProcessSelectionOutline.cpp
+++ b/Engine/Source/Runtime/Renderer/Private/PostProcess/PostProcessSelectionOutline.cpp
@@ -19,6 +19,8 @@ class FSelectionOutlinePS : public FEditorPrimitiveShader
 public:
 	DECLARE_GLOBAL_SHADER(FSelectionOutlinePS);
 	SHADER_USE_PARAMETER_STRUCT(FSelectionOutlinePS, FEditorPrimitiveShader);
+	class FSelectionOutlineHDRDim : SHADER_PERMUTATION_BOOL("SELECTION_OUTLINE_HDR");
+	using FPermutationDomain = TShaderPermutationDomain< FEditorPrimitiveShader::FPermutationDomain, FSelectionOutlineHDRDim>;
 
 	BEGIN_SHADER_PARAMETER_STRUCT(FParameters, )
 		SHADER_PARAMETER_STRUCT_REF(FViewUniformShaderParameters, View)
@@ -35,8 +37,26 @@ public:
 		SHADER_PARAMETER(float, SelectionHighlightIntensity)
 		SHADER_PARAMETER(FVector3f, SubduedOutlineColor)
 		SHADER_PARAMETER(float, BSPSelectionIntensity)
+		SHADER_PARAMETER(float, UILuminanceAndIsSCRGB)
 		RENDER_TARGET_BINDING_SLOTS()
 	END_SHADER_PARAMETER_STRUCT()
+
+	static bool SupportHDR(const EShaderPlatform Platform)
+	{
+		return IsFeatureLevelSupported(Platform, ERHIFeatureLevel::SM5);
+	}
+
+	static bool ShouldCompilePermutation(const FGlobalShaderPermutationParameters& Parameters)
+	{
+		FPermutationDomain CurrentPermutation(Parameters.PermutationId);
+		FEditorPrimitiveShader::FPermutationDomain EditorPrimitiveShaderDomain = CurrentPermutation.Get<FEditorPrimitiveShader::FPermutationDomain>();
+		bool bIsHDR = CurrentPermutation.Get<FSelectionOutlineHDRDim>();
+		if (bIsHDR && !SupportHDR(Parameters.Platform))
+		{
+			return false;
+		}
+		return FEditorPrimitiveShader::ShouldCompilePermutation(EditorPrimitiveShaderDomain, Parameters.Platform);
+	}
 };
 
 IMPLEMENT_GLOBAL_SHADER(FSelectionOutlinePS, "/Engine/Private/PostProcessSelectionOutline.usf", "MainPS", SF_Pixel);
@@ -188,8 +208,44 @@ FScreenPassTexture AddSelectionOutlinePass(FRDGBuilder& GraphBuilder, const FVie
 		PassParameters->SubduedOutlineColor = FVector3f(View.SubduedSelectionOutlineColor);
 		PassParameters->BSPSelectionIntensity = GEngine->BSPSelectionHighlightIntensity;
 
+		EDisplayOutputFormat DisplayOutputFormat = View.Family->RenderTarget->GetDisplayOutputFormat();
+		bool bIsHDR = false;
+		bool bIsSCRGB = false;
+		if (FSelectionOutlinePS::SupportHDR(GetFeatureLevelShaderPlatform(GMaxRHIFeatureLevel)))
+		{
+			switch (DisplayOutputFormat)
+			{
+			case EDisplayOutputFormat::HDR_ACES_1000nit_ScRGB:
+			case EDisplayOutputFormat::HDR_ACES_2000nit_ScRGB:
+				bIsHDR = true;
+				bIsSCRGB = true;
+				break;
+			case EDisplayOutputFormat::HDR_ACES_1000nit_ST2084:
+			case EDisplayOutputFormat::HDR_ACES_2000nit_ST2084:
+				bIsHDR = true;
+				break;
+			case EDisplayOutputFormat::SDR_sRGB:
+			case EDisplayOutputFormat::SDR_Rec709:
+			case EDisplayOutputFormat::SDR_ExplicitGammaMapping:
+			case EDisplayOutputFormat::HDR_LinearEXR:
+			case EDisplayOutputFormat::HDR_LinearNoToneCurve:
+			case EDisplayOutputFormat::HDR_LinearWithToneCurve:
+				break;
+			default:
+				checkNoEntry();
+				break;
+			}
+		}
+
+		static const auto CVarHDRUILuminance = IConsoleManager::Get().FindTConsoleVariableDataFloat(TEXT("r.HDR.UI.Luminance"));
+		float UILuminance = CVarHDRUILuminance ? CVarHDRUILuminance->GetValueOnRenderThread() : 300.0f;
+		PassParameters->UILuminanceAndIsSCRGB = bIsSCRGB ? UILuminance : -UILuminance;
+
 		FSelectionOutlinePS::FPermutationDomain PermutationVector;
-		PermutationVector.Set<FSelectionOutlinePS::FSampleCountDimension>(NumSamples);
+		FEditorPrimitiveShader::FPermutationDomain PermutationVectorBase;
+		PermutationVectorBase.Set<FSelectionOutlinePS::FSampleCountDimension>(NumSamples);
+		PermutationVector.Set<FEditorPrimitiveShader::FPermutationDomain>(PermutationVectorBase);
+		PermutationVector.Set<FSelectionOutlinePS::FSelectionOutlineHDRDim>(bIsHDR);
 
 		TShaderMapRef<FSelectionOutlinePS> PixelShader(View.ShaderMap, PermutationVector);
 
