Fix primitives marked as ERayTracingPrimitiveFlags::Excluded going through dynamic raytracing primitive codepath.
- CullPrimitiveByFlags must run even when r.RayTracing.Culling=0

diff --git a/Engine/Source/Runtime/Renderer/Private/DeferredShadingRenderer.cpp b/Engine/Source/Runtime/Renderer/Private/DeferredShadingRenderer.cpp
--- a/Engine/Source/Runtime/Renderer/Private/DeferredShadingRenderer.cpp
+++ b/Engine/Source/Runtime/Renderer/Private/DeferredShadingRenderer.cpp
@@ -675,6 +675,8 @@ static void GatherRayTracingRelevantPrimitives(const FScene& Scene, const FViewI
 				continue;
 			}
 
+			check(!EnumHasAnyFlags(Scene.PrimitiveRayTracingFlags[PrimitiveIndex], ERayTracingPrimitiveFlags::Excluded));
+
 			const FPrimitiveSceneInfo* SceneInfo = Scene.Primitives[PrimitiveIndex];
 
 			// #dxr_todo: ray tracing in scene captures should re-use the persistent RT scene. (UE-112448)
diff --git a/Engine/Source/Runtime/Renderer/Private/SceneVisibility.cpp b/Engine/Source/Runtime/Renderer/Private/SceneVisibility.cpp
--- a/Engine/Source/Runtime/Renderer/Private/SceneVisibility.cpp
+++ b/Engine/Source/Runtime/Renderer/Private/SceneVisibility.cpp
@@ -579,10 +579,6 @@ struct FPrimitiveCullingFlags
 	bool bNaniteAlwaysVisible;
 	bool bHasHiddenPrimitives;
 	bool bHasShowOnlyPrimitives;
-
-#if RHI_RAYTRACING
-	bool bCullInRayTracing;
-#endif
 };
 
 // Returns true if the frustum and bounds intersect
@@ -750,7 +746,7 @@ static void PrimitiveCullTask(FThreadSafeCounter& NumCulledPrimitives, const FSc
 #if RHI_RAYTRACING
 			bool bIsVisibleInRayTracing = true;
 
-			if (bPrimitiveIsHidden || (Flags.bCullInRayTracing && ShouldCullForRayTracing(Scene, View, Index)))
+			if (bPrimitiveIsHidden || ShouldCullForRayTracing(Scene, View, Index))
 			{
 				bIsVisibleInRayTracing = false;
 			}
@@ -918,7 +914,6 @@ static int32 PrimitiveCull(const FScene* RESTRICT Scene, FViewInfo& View, bool b
 
 #if RHI_RAYTRACING
 	View.RayTracingCullingParameters.Init(View);
-	Flags.bCullInRayTracing = View.RayTracingCullingParameters.CullInRayTracing > 0;
 #endif
 
 	SCOPE_CYCLE_COUNTER(STAT_PrimitiveCull);
@@ -1004,7 +999,7 @@ static void UpdatePrimitiveFading(const FScene* Scene, FViewInfo& View)
 					// This should be a very rare occurrence, so the hit is not worrisome.
 					// TODO:  Could this be moved into the actual culling phase?
 
-					if (View.RayTracingCullingParameters.CullInRayTracing == 0 || !ShouldCullForRayTracing(Scene, View, BitIt.GetIndex()))
+					if (!ShouldCullForRayTracing(Scene, View, BitIt.GetIndex()))
 					{
 						View.PrimitiveRayTracingVisibilityMap.AccessCorrespondingBit(BitIt) = true;
 					}
