Implemented depth prepass for SingleLayerWater. Depth prepass output is fed into VSM page marking/allocation shader to get correct VSM pages for sampling VSM on water. The depth prepass is also used in the main water pass by setting the main pass depth test to EQUAL and forcing earlyZ. CopyTexture() is currently not supported for depth buffers on all platforms, so I had to write a custom copy shader to create a copy of the main depth buffer to render the water depth prepass into.

#rnx

diff --git a/Engine/Shaders/Private/BasePassPixelShader.usf b/Engine/Shaders/Private/BasePassPixelShader.usf
--- a/Engine/Shaders/Private/BasePassPixelShader.usf
+++ b/Engine/Shaders/Private/BasePassPixelShader.usf
@@ -2149,6 +2149,13 @@ void FPixelShaderInOut_MainPS(
 
 #endif // STRATA_OPAQUE_DEFERRED
 
+// Force early Z on single layer water main pass.
+// This assumes that a full depth prepass was run for SLW and that the depth test is set to equal in the main pass.
+#ifndef PIXELSHADER_EARLYDEPTHSTENCIL
+	#if SINGLE_LAYER_WATER_NO_DISCARD && !OUTPUT_PIXEL_DEPTH_OFFSET
+		#define PIXELSHADER_EARLYDEPTHSTENCIL EARLYDEPTHSTENCIL
+	#endif // SINGLE_LAYER_WATER_NO_DISCARD && !OUTPUT_PIXEL_DEPTH_OFFSET
+#endif // PIXELSHADER_EARLYDEPTHSTENCIL
 
 // the following needs to match to the code in FSceneTextures::GetGBufferRenderTargets()
 #define PIXELSHADEROUTPUT_BASEPASS 1
diff --git a/Engine/Shaders/Private/CopyDepthTexture.usf b/Engine/Shaders/Private/CopyDepthTexture.usf
new file mode 100644
--- /dev/null
+++ b/Engine/Shaders/Private/CopyDepthTexture.usf
@@ -0,0 +1,27 @@
+// Copyright Epic Games, Inc. All Rights Reserved.
+
+/*=============================================================================
+	CopyDepthTexture.usf
+=============================================================================*/
+
+#include "Common.ush"
+
+#if MSAA_SAMPLE_COUNT > 1
+Texture2DMS<float, MSAA_SAMPLE_COUNT> DepthTextureMS;
+#else
+Texture2D<float> DepthTexture;
+#endif
+
+float CopyDepthPS(
+	float4 Position : SV_POSITION
+#if MSAA_SAMPLE_COUNT > 1
+	, uint SampleIndex : SV_SampleIndex
+#endif
+) : SV_DEPTH
+{
+#if MSAA_SAMPLE_COUNT > 1
+	return DepthTextureMS.Load(int2(Position.xy), SampleIndex).x;
+#else
+	return DepthTexture.Load(int3(Position.xy, 0)).x;
+#endif
+}
\ No newline at end of file
diff --git a/Engine/Shaders/Private/Definitions.usf b/Engine/Shaders/Private/Definitions.usf
--- a/Engine/Shaders/Private/Definitions.usf
+++ b/Engine/Shaders/Private/Definitions.usf
@@ -108,6 +108,10 @@
 #define SINGLE_LAYER_WATER_DF_SHADOW_ENABLED			0 
 #endif
 
+#ifndef SINGLE_LAYER_WATER_NO_DISCARD
+#define SINGLE_LAYER_WATER_NO_DISCARD					0
+#endif
+
 #ifndef MATERIAL_SINGLE_SHADINGMODEL
 #define	MATERIAL_SINGLE_SHADINGMODEL					0
 #endif
diff --git a/Engine/Shaders/Private/MaterialTemplate.ush b/Engine/Shaders/Private/MaterialTemplate.ush
--- a/Engine/Shaders/Private/MaterialTemplate.ush
+++ b/Engine/Shaders/Private/MaterialTemplate.ush
@@ -3179,7 +3179,7 @@ void GetMaterialCoverageAndClipping(FMaterialPixelParameters Parameters, FPixelM
 {
 	ClipLODTransition(Parameters);
 
-#if MATERIALBLENDING_MASKED
+#if MATERIALBLENDING_MASKED && !SINGLE_LAYER_WATER_NO_DISCARD
 	#if MATERIAL_DITHER_OPACITY_MASK
 		/*
 		5 value dither. Every value present in +
diff --git a/Engine/Shaders/Private/VirtualShadowMaps/VirtualShadowMapPageManagement.usf b/Engine/Shaders/Private/VirtualShadowMaps/VirtualShadowMapPageManagement.usf
--- a/Engine/Shaders/Private/VirtualShadowMaps/VirtualShadowMapPageManagement.usf
+++ b/Engine/Shaders/Private/VirtualShadowMaps/VirtualShadowMapPageManagement.usf
@@ -32,6 +32,7 @@
 // Type of input data consume by the page allocation (i.e., data read from the source buffer: Gbuffer, HairStrands data, ...)
 #define INPUT_TYPE_GBUFFER 0
 #define INPUT_TYPE_HAIRSTRANDS 1
+#define INPUT_TYPE_GBUFFER_AND_WATER_DEPTH 2
 
 struct FPhysicalPageRequest
 {
@@ -121,6 +122,22 @@ void MarkPage(uint VirtualShadowMapId, uint MipLevel, float3 TranslatedWorldPosi
 	}
 }
 
+void MarkPageClipmap(
+	FVirtualShadowMapProjectionShaderData BaseProjectionData, 
+	int ClipmapStartId,
+	bool bUsePageDilation, 
+	float2 PageDilationOffset, 
+	FLWCVector3 WorldPosition,
+	float3 TranslatedWorldPosition)
+{
+	const int ClipmapLevel = CalcClipmapLevel(BaseProjectionData, WorldPosition);
+	int ClipmapIndex = max(0, ClipmapLevel - BaseProjectionData.ClipmapLevel);
+	if (ClipmapIndex < BaseProjectionData.ClipmapLevelCount)
+	{
+		MarkPage(ClipmapStartId + ClipmapIndex, 0, TranslatedWorldPosition, bUsePageDilation, PageDilationOffset);
+	}
+}
+
 // 
 RWStructuredBuffer<uint> OutPrunedLightGridData;
 RWStructuredBuffer<uint> OutPrunedNumCulledLightsGrid;
@@ -165,6 +182,7 @@ void PruneLightGridCS(uint GridLinearIndex : SV_DispatchThreadID)
 
 StructuredBuffer<uint> PrunedLightGridData;
 StructuredBuffer<uint> PrunedNumCulledLightsGrid;
+Texture2D<float> SingeLayerWaterDepthTexture;
 
 
 [numthreads(VSM_DEFAULT_CS_GROUP_XY, VSM_DEFAULT_CS_GROUP_XY, 1)]
@@ -186,7 +204,7 @@ void GeneratePageFlagsFromPixels(
 		}
 		GroupId = HairStrands.HairTileData[LinearIndex];
 	}
-#else // PERMUTATION_INPUT_TYPE == INPUT_TYPE_GBUFFER
+#else // PERMUTATION_INPUT_TYPE == INPUT_TYPE_GBUFFER || PERMUTATION_INPUT_TYPE == INPUT_TYPE_GBUFFER_AND_WATER_DEPTH
 	const uint2 GroupId = InGroupId.xy;
 #endif
 	// Morton order within a group so page access/atomics are more coherent and wave-swizzled gradients are possible.
@@ -197,17 +215,19 @@ void GeneratePageFlagsFromPixels(
 		return;
 	}
 
-	half3  WorldNormal   = half3(0, 0, 0);
-	float  DeviceZ		 = 0;
-	bool   bIsValid		 = true;
-	bool   bBackfaceCull = bCullBackfacingPixels != 0;
+	half3  WorldNormal			= half3(0, 0, 0);
+	float  DeviceZ				= 0;
+	bool   bIsValid				= true;
+	bool   bBackfaceCull		= bCullBackfacingPixels != 0;
+	float  DeviceZWater			= 0;
+	bool   bIsWaterDepthValid	= false;
 #if PERMUTATION_INPUT_TYPE == INPUT_TYPE_HAIRSTRANDS
 	{
 		DeviceZ = HairStrands.HairOnlyDepthTexture.Load(uint3(PixelPos, 0)).x;
 		bIsValid = DeviceZ > 0;
 		bBackfaceCull = false;
 	}
-#else // PERMUTATION_INPUT_TYPE == INPUT_TYPE_GBUFFER
+#else // PERMUTATION_INPUT_TYPE == INPUT_TYPE_GBUFFER || PERMUTATION_INPUT_TYPE == INPUT_TYPE_GBUFFER_AND_WATER_DEPTH
 	{
 		DeviceZ = LookupDeviceZ(PixelPos);
 #if STRATA_ENABLED
@@ -222,10 +242,16 @@ void GeneratePageFlagsFromPixels(
 		bIsValid = GBuffer.ShadingModelID != SHADINGMODELID_UNLIT;
 		bBackfaceCull = bBackfaceCull && !IsSubsurfaceModel(GBuffer.ShadingModelID);
 #endif // STRATA_ENABLED
+
+#if PERMUTATION_INPUT_TYPE == INPUT_TYPE_GBUFFER_AND_WATER_DEPTH
+		DeviceZWater = SingeLayerWaterDepthTexture.Load(uint3(PixelPos, 0)).x;
+		bIsWaterDepthValid = DeviceZWater != DeviceZ;
+#endif // PERMUTATION_INPUT_TYPE == INPUT_TYPE_GBUFFER_AND_WATER_DEPTH
+
 	}
-#endif // PERMUTATION_INPUT_TYPE == INPUT_TYPE_GBUFFER
+#endif // PERMUTATION_INPUT_TYPE == INPUT_TYPE_GBUFFER || PERMUTATION_INPUT_TYPE == INPUT_TYPE_GBUFFER_AND_WATER_DEPTH
 	
-	if (!bIsValid)
+	if (!bIsValid && !bIsWaterDepthValid)
 	{
 		return;
 	}
@@ -235,6 +261,21 @@ void GeneratePageFlagsFromPixels(
 	const float3 TranslatedWorldPosition = SvPositionToTranslatedWorld(SvPosition);
 
 	const FLWCVector3 WorldPosition = LWCSubtract(TranslatedWorldPosition, PrimaryView.PreViewTranslation);
+
+#if PERMUTATION_INPUT_TYPE == INPUT_TYPE_GBUFFER_AND_WATER_DEPTH
+	float SceneDepthWater;
+	float4 SvPositionWater;
+	float3 TranslatedWorldPositionWater;
+	FLWCVector3 WorldPositionWater;
+
+	if (bIsWaterDepthValid)
+	{
+		SceneDepthWater = ConvertFromDeviceZ(DeviceZWater);
+		SvPositionWater = float4(float2(PixelPos) + 0.5f, DeviceZWater, 1.0f);
+		TranslatedWorldPositionWater = SvPositionToTranslatedWorld(SvPositionWater);
+		WorldPositionWater = LWCSubtract(TranslatedWorldPositionWater, PrimaryView.PreViewTranslation);
+	}
+#endif
 	
 	// Dither pattern for page dilation
 	// We don't need to to check all 8 adjacent pages; as long as there's at least a single pixel near the edge
@@ -254,21 +295,27 @@ void GeneratePageFlagsFromPixels(
 			int ClipmapStartId = DirectionalLightIds[Index];
 			FVirtualShadowMapProjectionShaderData BaseProjectionData = GetVirtualShadowMapProjectionData(ClipmapStartId);
 
+			bool bBackfaceCulledGBuffer = false;
 #if PERMUTATION_INPUT_TYPE != INPUT_TYPE_HAIRSTRANDS
 			// Backface test if requested
 			float3 LightDirection = normalize(-BaseProjectionData.TranslatedWorldToShadowViewMatrix._13_23_33);
 			if (bBackfaceCull && IsBackfaceToDirectionalLight(WorldNormal, LightDirection, BaseProjectionData.LightSourceRadius))
 			{
-				continue;
+				bBackfaceCulledGBuffer = true;
 			}
 #endif
 
-			const int ClipmapLevel = CalcClipmapLevel(BaseProjectionData, WorldPosition);
-			int ClipmapIndex = max(0, ClipmapLevel - BaseProjectionData.ClipmapLevel);
-			if (ClipmapIndex < BaseProjectionData.ClipmapLevelCount)
+			if (!bBackfaceCulledGBuffer)
+			{
+				MarkPageClipmap(BaseProjectionData, ClipmapStartId, bUsePageDilation, PageDilationOffset, WorldPosition, TranslatedWorldPosition);
+			}
+			
+#if PERMUTATION_INPUT_TYPE == INPUT_TYPE_GBUFFER_AND_WATER_DEPTH
+			if (bIsWaterDepthValid)
 			{
-				MarkPage(ClipmapStartId + ClipmapIndex, 0, TranslatedWorldPosition, bUsePageDilation, PageDilationOffset);
+				MarkPageClipmap(BaseProjectionData, ClipmapStartId, bUsePageDilation, PageDilationOffset, WorldPositionWater, TranslatedWorldPositionWater);
 			}
+#endif // PERMUTATION_INPUT_TYPE == INPUT_TYPE_GBUFFER_AND_WATER_DEPTH
 		}
 	}
 
diff --git a/Engine/Source/Runtime/RenderCore/Private/Shader.cpp b/Engine/Source/Runtime/RenderCore/Private/Shader.cpp
--- a/Engine/Source/Runtime/RenderCore/Private/Shader.cpp
+++ b/Engine/Source/Runtime/RenderCore/Private/Shader.cpp
@@ -2039,6 +2039,11 @@ void ShaderMapAppendKeyString(EShaderPlatform Platform, FString& KeyString)
 
 	uint64 ShaderPlatformPropertiesHash = FDataDrivenShaderPlatformInfo::GetShaderPlatformPropertiesHash(Platform);
 	KeyString += FString::Printf(TEXT("_%u"), ShaderPlatformPropertiesHash);
+
+	if (IsSingleLayerWaterDepthPrepassEnabled(Platform, GetMaxSupportedFeatureLevel(Platform)))
+	{
+		KeyString += TEXT("_SLWDP");
+	}
 }
 
 EShaderPermutationFlags GetShaderPermutationFlags(const FPlatformTypeLayoutParameters& LayoutParams)
diff --git a/Engine/Source/Runtime/RenderCore/Public/RenderUtils.h b/Engine/Source/Runtime/RenderCore/Public/RenderUtils.h
--- a/Engine/Source/Runtime/RenderCore/Public/RenderUtils.h
+++ b/Engine/Source/Runtime/RenderCore/Public/RenderUtils.h
@@ -793,6 +793,21 @@ inline bool UseNonNaniteVirtualShadowMaps(EShaderPlatform ShaderPlatform, const
 	return EnableCVar->GetInt() != 0 && UseVirtualShadowMaps(ShaderPlatform, FeatureLevel);
 }
 
+/**
+*	(Non-runtime) Checks if the depth prepass for single layer water is enabled. This also depends on virtual shadow maps to be supported on the platform.
+*/
+inline bool IsSingleLayerWaterDepthPrepassEnabled(const FStaticShaderPlatform& Platform, const FStaticFeatureLevel& FeatureLevel)
+{
+	static const auto CVarWaterSingleLayerDepthPrepass = IConsoleManager::Get().FindConsoleVariable(TEXT("r.Water.SingleLayer.DepthPrepass"));
+	const bool bPrepassEnabled = CVarWaterSingleLayerDepthPrepass && CVarWaterSingleLayerDepthPrepass->GetInt() > 0;
+	// Currently VSM is the only feature dependent on the depth prepass which is why we only enable it if VSM could also be enabled.
+	// VSM can be toggled at runtime, but we need a compile time value here, so we fall back to DoesRuntimeSupportNanite() to check if
+	// VSM *could* be enabled.
+	const bool bVSMSupported = DoesPlatformSupportNanite(Platform, false /* check project setting */);
+
+	return bPrepassEnabled && bVSMSupported;
+}
+
 /**
 *	Checks if virtual texturing lightmap enabled and supported
 */
diff --git a/Engine/Source/Runtime/Renderer/Private/BasePassRendering.h b/Engine/Source/Runtime/Renderer/Private/BasePassRendering.h
--- a/Engine/Source/Runtime/Renderer/Private/BasePassRendering.h
+++ b/Engine/Source/Runtime/Renderer/Private/BasePassRendering.h
@@ -482,6 +482,15 @@ public:
 			OutEnvironment.SetRenderTargetOutputFormat(TargetSeparatedMainDirLight, PF_FloatR11G11B10);
 		}
 
+		if (IsSingleLayerWater)
+		{
+			const bool bHasDepthPrepass = IsSingleLayerWaterDepthPrepassEnabled(Parameters.Platform, GetMaxSupportedFeatureLevel(Parameters.Platform));
+			if (bHasDepthPrepass)
+			{
+				OutEnvironment.SetDefine(TEXT("SINGLE_LAYER_WATER_NO_DISCARD"), TEXT("1"));
+			}
+		}
+
 		TBasePassPixelShaderBaseType<LightMapPolicyType>::ModifyCompilationEnvironment(Parameters, OutEnvironment);
 	}
 	
diff --git a/Engine/Source/Runtime/Renderer/Private/DeferredShadingRenderer.cpp b/Engine/Source/Runtime/Renderer/Private/DeferredShadingRenderer.cpp
--- a/Engine/Source/Runtime/Renderer/Private/DeferredShadingRenderer.cpp
+++ b/Engine/Source/Runtime/Renderer/Private/DeferredShadingRenderer.cpp
@@ -3157,6 +3157,14 @@ void FDeferredShadingSceneRenderer::Render(FRDGBuilder& GraphBuilder)
 	// Copy lighting channels out of stencil before deferred decals which overwrite those values
 	FRDGTextureRef LightingChannelsTexture = CopyStencilToLightingChannelTexture(GraphBuilder, SceneTextures.Stencil);
 
+	// Single layer water depth prepass. Needs to run before VSM page allocation.
+	FRDGTextureMSAA SingleLayerWaterDepthPrepassOutput;
+	const bool bShouldRenderSingleLayerWaterDepthPrepass = !bHasRayTracedOverlay && ShouldRenderSingleLayerWaterDepthPrepass(Views);
+	if (bShouldRenderSingleLayerWaterDepthPrepass)
+	{
+		RenderSingleLayerWaterDepthPrepass(GraphBuilder, SceneTextures, SingleLayerWaterDepthPrepassOutput);
+	}
+
 	FAsyncLumenIndirectLightingOutputs AsyncLumenIndirectLightingOutputs;
 	const bool bHasLumenLights = SortedLightSet.LumenLightStart < SortedLightSet.SortedLights.Num();
 
@@ -3181,7 +3189,7 @@ void FDeferredShadingSceneRenderer::Render(FRDGBuilder& GraphBuilder)
 				RDG_GPU_STAT_SCOPE(GraphBuilder, ShadowDepths);
 
 				ensureMsgf(AreLightsInLightGrid(), TEXT("Virtual shadow map setup requires local lights to be injected into the light grid (this may be caused by 'r.LightCulling.Quality=0')."));
-				VirtualShadowMapArray.BuildPageAllocations(GraphBuilder, SceneTextures, Views, ViewFamily.EngineShowFlags, SortedLightSet, VisibleLightInfos, NaniteRasterResults);
+				VirtualShadowMapArray.BuildPageAllocations(GraphBuilder, SceneTextures, Views, ViewFamily.EngineShowFlags, SortedLightSet, VisibleLightInfos, NaniteRasterResults, SingleLayerWaterDepthPrepassOutput.Resolve);
 			}
 
 			RenderShadowDepthMaps(GraphBuilder, InstanceCullingManager);
@@ -3437,7 +3445,15 @@ void FDeferredShadingSceneRenderer::Render(FRDGBuilder& GraphBuilder)
 		}
 
 		GraphBuilder.SetCommandListStat(GET_STATID(STAT_CLM_WaterPass));
-		RenderSingleLayerWater(GraphBuilder, SceneTextures, bShouldRenderVolumetricCloud, SceneWithoutWaterTextures, LumenFrameTemporaries);
+		RenderSingleLayerWater(GraphBuilder, SceneTextures, SingleLayerWaterDepthPrepassOutput, bShouldRenderVolumetricCloud, SceneWithoutWaterTextures, LumenFrameTemporaries);
+
+		// Replace main depth texture with the output of the SLW depth prepass which contains the scene + water.
+		// Note: Stencil now has all water bits marked with 1. As long as no other passes after this point want to read the depth buffer,
+		// a stencil clear should not be necessary here.
+		if (bShouldRenderSingleLayerWaterDepthPrepass)
+		{
+			SceneTextures.Depth = SingleLayerWaterDepthPrepassOutput;
+		}
 	}
 
 	// Rebuild scene textures to include scene color.
diff --git a/Engine/Source/Runtime/Renderer/Private/DeferredShadingRenderer.h b/Engine/Source/Runtime/Renderer/Private/DeferredShadingRenderer.h
--- a/Engine/Source/Runtime/Renderer/Private/DeferredShadingRenderer.h
+++ b/Engine/Source/Runtime/Renderer/Private/DeferredShadingRenderer.h
@@ -341,9 +341,15 @@ public:
 		FSceneTextures& SceneTextures,
 		bool bDoParallelPass);
 
+	void RenderSingleLayerWaterDepthPrepass(
+		FRDGBuilder& GraphBuilder,
+		const FSceneTextures& SceneTextures,
+		FRDGTextureMSAA& OutDepthPrepassTexture);
+
 	void RenderSingleLayerWater(
 		FRDGBuilder& GraphBuilder,
 		const FSceneTextures& SceneTextures,
+		const FRDGTextureMSAA& SingleLayerWaterDepthPrepassTexture,
 		bool bShouldRenderVolumetricCloud,
 		FSceneWithoutWaterTextures& SceneWithoutWaterTextures,
 		FLumenSceneFrameTemporaries& LumenFrameTemporaries);
@@ -351,7 +357,8 @@ public:
 	void RenderSingleLayerWaterInner(
 		FRDGBuilder& GraphBuilder,
 		const FSceneTextures& SceneTextures,
-		const FSceneWithoutWaterTextures& SceneWithoutWaterTextures);
+		const FSceneWithoutWaterTextures& SceneWithoutWaterTextures,
+		const FRDGTextureMSAA& SingleLayerWaterDepthPrepassTexture);
 
 	void RenderSingleLayerWaterReflections(
 		FRDGBuilder& GraphBuilder,
diff --git a/Engine/Source/Runtime/Renderer/Private/DepthRendering.h b/Engine/Source/Runtime/Renderer/Private/DepthRendering.h
--- a/Engine/Source/Runtime/Renderer/Private/DepthRendering.h
+++ b/Engine/Source/Runtime/Renderer/Private/DepthRendering.h
@@ -57,6 +57,8 @@ extern FDepthPassInfo GetDepthPassInfo(const FScene* Scene);
 
 void AddDitheredStencilFillPass(FRDGBuilder& GraphBuilder, TConstArrayView<FViewInfo> Views, FRDGTextureRef DepthTexture, const FDepthPassInfo& DepthPass);
 
+FMeshDrawCommandSortKey CalculateDepthPassMeshStaticSortKey(EBlendMode BlendMode, const FMeshMaterialShader* VertexShader, const FMeshMaterialShader* PixelShader);
+
 /**
  * A vertex shader for rendering the depth of a mesh.
  */
diff --git a/Engine/Source/Runtime/Renderer/Private/SceneVisibility.cpp b/Engine/Source/Runtime/Renderer/Private/SceneVisibility.cpp
--- a/Engine/Source/Runtime/Renderer/Private/SceneVisibility.cpp
+++ b/Engine/Source/Runtime/Renderer/Private/SceneVisibility.cpp
@@ -2724,6 +2724,7 @@ struct FRelevancePacket : public FSceneRenderingAllocatorObject<FRelevancePacket
 									if (StaticMeshRelevance.bUseSingleLayerWaterMaterial)
 									{
 										DrawCommandPacket.AddCommandsForMesh(PrimitiveIndex, PrimitiveSceneInfo, StaticMeshRelevance, StaticMesh, Scene, bCanCache, EMeshPass::SingleLayerWaterPass);
+										DrawCommandPacket.AddCommandsForMesh(PrimitiveIndex, PrimitiveSceneInfo, StaticMeshRelevance, StaticMesh, Scene, bCanCache, EMeshPass::SingleLayerWaterDepthPrepass);
 									}
 								}
 
@@ -3248,6 +3249,8 @@ void ComputeDynamicMeshRelevance(EShadingPath ShadingPath, bool bAddLightmapDens
 			{
 				PassMask.Set(EMeshPass::SingleLayerWaterPass);
 				View.NumVisibleDynamicMeshElements[EMeshPass::SingleLayerWaterPass] += NumElements;
+				PassMask.Set(EMeshPass::SingleLayerWaterDepthPrepass);
+				View.NumVisibleDynamicMeshElements[EMeshPass::SingleLayerWaterDepthPrepass] += NumElements;
 			}
 		}
 	}
diff --git a/Engine/Source/Runtime/Renderer/Private/SingleLayerWaterRendering.cpp b/Engine/Source/Runtime/Renderer/Private/SingleLayerWaterRendering.cpp
--- a/Engine/Source/Runtime/Renderer/Private/SingleLayerWaterRendering.cpp
+++ b/Engine/Source/Runtime/Renderer/Private/SingleLayerWaterRendering.cpp
@@ -23,6 +23,7 @@
 
 DECLARE_GPU_STAT_NAMED(RayTracingWaterReflections, TEXT("Ray Tracing Water Reflections"));
 
+DECLARE_GPU_DRAWCALL_STAT(SingleLayerWaterDepthPrepass);
 DECLARE_GPU_DRAWCALL_STAT(SingleLayerWater);
 DECLARE_CYCLE_STAT(TEXT("WaterSingleLayer"), STAT_CLP_WaterSingleLayerPass, STATGROUP_ParallelCommandListMarkers);
 
@@ -96,6 +97,11 @@ static TAutoConsoleVariable<int32> CVarRHICmdFlushRenderThreadTasksSingleLayerWa
 	0,
 	TEXT("Wait for completion of parallel render thread tasks at the end of Single layer water. A more granular version of r.RHICmdFlushRenderThreadTasks. If either r.RHICmdFlushRenderThreadTasks or r.RHICmdFlushRenderThreadTasksSingleLayerWater is > 0 we will flush."));
 
+static TAutoConsoleVariable<int32> CVarWaterSingleLayerDepthPrepass(
+	TEXT("r.Water.SingleLayer.DepthPrepass"), 0,
+	TEXT("Enable a depth prepass for single layer water. Necessary for proper Virtual Shadow Maps support."),
+	ECVF_ReadOnly | ECVF_RenderThreadSafe);
+
 // This is to have platforms use the simple single layer water shading similar to mobile: no dynamic lights, only sun and sky, no distortion, no colored transmittance on background, no custom depth read.
 bool SingleLayerWaterUsesSimpleShading(EShaderPlatform ShaderPlatform)
 {
@@ -132,6 +138,15 @@ bool ShouldRenderSingleLayerWaterSkippedRenderEditorNotification(TArrayView<cons
 	return false;
 }
 
+bool ShouldRenderSingleLayerWaterDepthPrepass(TArrayView<const FViewInfo> Views)
+{
+	check(Views.Num() > 0);
+	const bool bPrepassEnabled = IsSingleLayerWaterDepthPrepassEnabled(Views[0].GetShaderPlatform(), Views[0].GetFeatureLevel());
+	const bool bShouldRenderWater = ShouldRenderSingleLayerWater(Views);
+	
+	return bPrepassEnabled && bShouldRenderWater;
+}
+
 bool ShouldUseBilinearSamplerForDepthWithoutSingleLayerWater(EPixelFormat DepthTextureFormat)
 {
 	const bool bHasDownsampling = GSingleLayerWaterRefractionDownsampleFactor > 1;
@@ -291,6 +306,160 @@ class FWaterRefractionCopyPS : public FGlobalShader
 
 IMPLEMENT_GLOBAL_SHADER(FWaterRefractionCopyPS, "/Engine/Private/SingleLayerWaterComposite.usf", "WaterRefractionCopyPS", SF_Pixel);
 
+class FCopyDepthPS : public FGlobalShader
+{
+public:
+	DECLARE_GLOBAL_SHADER(FCopyDepthPS);
+	SHADER_USE_PARAMETER_STRUCT(FCopyDepthPS, FGlobalShader);
+
+	class FMSAASampleCount : SHADER_PERMUTATION_SPARSE_INT("MSAA_SAMPLE_COUNT", 1, 2, 4, 8);
+	using FPermutationDomain = TShaderPermutationDomain<FMSAASampleCount>;
+
+	BEGIN_SHADER_PARAMETER_STRUCT(FParameters, )
+		SHADER_PARAMETER_RDG_TEXTURE(Texture2D, DepthTexture)
+		SHADER_PARAMETER_RDG_TEXTURE(Texture2DMS, DepthTextureMS)
+		RENDER_TARGET_BINDING_SLOTS()
+		END_SHADER_PARAMETER_STRUCT()
+
+		static bool ShouldCompilePermutation(const FGlobalShaderPermutationParameters& Parameters)
+	{
+		return IsFeatureLevelSupported(Parameters.Platform, ERHIFeatureLevel::SM5);
+	}
+};
+
+IMPLEMENT_GLOBAL_SHADER(FCopyDepthPS, "/Engine/Private/CopyDepthTexture.usf", "CopyDepthPS", SF_Pixel);
+
+BEGIN_SHADER_PARAMETER_STRUCT(FSingleLayerWaterDepthPassParameters, )
+	SHADER_PARAMETER_STRUCT_INCLUDE(FViewShaderParameters, View)
+	SHADER_PARAMETER_STRUCT_INCLUDE(FInstanceCullingDrawParams, InstanceCullingDrawParams)
+	RENDER_TARGET_BINDING_SLOTS()
+END_SHADER_PARAMETER_STRUCT()
+
+static FSingleLayerWaterDepthPassParameters* GetDepthPassParameters(FRDGBuilder& GraphBuilder, const FViewInfo& View, FRDGTextureRef DepthTexture)
+{
+	FSingleLayerWaterDepthPassParameters* PassParameters = GraphBuilder.AllocParameters<FSingleLayerWaterDepthPassParameters>();
+	PassParameters->View = View.GetShaderParameters();
+	PassParameters->RenderTargets.DepthStencil = FDepthStencilBinding(DepthTexture, ERenderTargetLoadAction::ELoad, ERenderTargetLoadAction::ELoad, FExclusiveDepthStencil::DepthWrite_StencilWrite);
+	return PassParameters;
+}
+
+void FDeferredShadingSceneRenderer::RenderSingleLayerWaterDepthPrepass(FRDGBuilder& GraphBuilder, const FSceneTextures& SceneTextures, FRDGTextureMSAA& OutDepthPrepassTexture)
+{
+	RDG_CSV_STAT_EXCLUSIVE_SCOPE(GraphBuilder, Water);
+	SCOPED_NAMED_EVENT(FDeferredShadingSceneRenderer_RenderSingleLayerWaterDepthPrepass, FColor::Emerald);
+	SCOPE_CYCLE_COUNTER(STAT_WaterPassDrawTime);
+	RDG_EVENT_SCOPE(GraphBuilder, "SingleLayerWaterDepthPrepass");
+	RDG_GPU_STAT_SCOPE(GraphBuilder, SingleLayerWaterDepthPrepass);
+
+	// Create an identical copy of the main depth buffer
+	{
+		const FRDGTextureDesc& DepthPrepassTextureDesc = SceneTextures.Depth.Target->Desc;
+		OutDepthPrepassTexture = GraphBuilder.CreateTexture(DepthPrepassTextureDesc, TEXT("SLW.DepthPrepassOutput"));
+		if (DepthPrepassTextureDesc.NumSamples > 1)
+		{
+			FRDGTextureDesc DepthPrepassResolveTextureDesc = DepthPrepassTextureDesc;
+			DepthPrepassResolveTextureDesc.NumSamples = 1;
+			OutDepthPrepassTexture.Resolve = GraphBuilder.CreateTexture(DepthPrepassResolveTextureDesc, TEXT("SLW.DepthPrepassOutputResolve"));
+		}
+
+		//AddCopyTexturePass(GraphBuilder, SceneTextures.Depth.Target, OutDepthPrepassTexture.Target);
+		//AddClearDepthStencilPass(GraphBuilder, OutDepthPrepassTexture.Target, false, 0.0f, true, 0);
+
+		// Copy main depth buffer content to our prepass depth buffer and clear stencil to 0
+		// TODO: replace with AddCopyTexturePass() and AddClearDepthStencilPass() once CopyTexture() supports depth buffer copies on all platforms.
+		{
+			FCopyDepthPS::FParameters* PassParameters = GraphBuilder.AllocParameters<FCopyDepthPS::FParameters>();
+			if (DepthPrepassTextureDesc.NumSamples > 1)
+			{
+				PassParameters->DepthTextureMS = SceneTextures.Depth.Target;
+			}
+			else
+			{
+				PassParameters->DepthTexture = SceneTextures.Depth.Target;
+			}
+			PassParameters->RenderTargets.DepthStencil = FDepthStencilBinding(OutDepthPrepassTexture.Target, ERenderTargetLoadAction::ENoAction, ERenderTargetLoadAction::ELoad, FExclusiveDepthStencil::DepthWrite_StencilNop);
+
+			FGlobalShaderMap* ShaderMap = GetGlobalShaderMap(GMaxRHIFeatureLevel);
+
+			FCopyDepthPS::FPermutationDomain PermutationVector;
+			PermutationVector.Set<FCopyDepthPS::FMSAASampleCount>(DepthPrepassTextureDesc.NumSamples);
+			TShaderMapRef<FCopyDepthPS> PixelShader(ShaderMap, PermutationVector);
+
+			FIntRect Viewport(0, 0, DepthPrepassTextureDesc.Extent.X, DepthPrepassTextureDesc.Extent.Y);
+
+			// Set depth test to always pass and stencil test to replace all pixels with zero, essentially also clearing stencil while doing the depth copy.
+			FRHIDepthStencilState* DepthStencilState = TStaticDepthStencilState<
+				true, CF_Always,										// depth
+				true, CF_Always, SO_Replace, SO_Replace, SO_Replace,	// frontface stencil
+				true, CF_Always, SO_Replace, SO_Replace, SO_Replace		// backface stencil
+			>::GetRHI();
+
+			FPixelShaderUtils::AddFullscreenPass(
+				GraphBuilder,
+				ShaderMap,
+				RDG_EVENT_NAME("SLW::DepthBufferCopy"),
+				PixelShader,
+				PassParameters,
+				Viewport,
+				nullptr, /*BlendState*/
+				nullptr, /*RasterizerState*/
+				DepthStencilState,
+				0 /*StencilRef*/);
+
+			// The above copy technique loses HTILE data during the copy, so until AddCopyTexturePass() supports depth buffer copies on all platforms,
+			// this is the best we can do.
+			AddResummarizeHTilePass(GraphBuilder, OutDepthPrepassTexture.Target);
+		}
+	}
+
+	const bool bRenderInParallel = GRHICommandList.UseParallelAlgorithms() && CVarParallelSingleLayerWaterPass.GetValueOnRenderThread() == 1;
+
+	for (int32 ViewIndex = 0; ViewIndex < Views.Num(); ++ViewIndex)
+	{
+		FViewInfo& View = Views[ViewIndex];
+
+		if (!View.ShouldRenderView())
+		{
+			continue;
+		}
+
+		RDG_GPU_MASK_SCOPE(GraphBuilder, View.GPUMask);
+		RDG_EVENT_SCOPE_CONDITIONAL(GraphBuilder, Views.Num() > 1, "View%d", ViewIndex);
+		View.BeginRenderView();
+
+		FSingleLayerWaterDepthPassParameters* PassParameters = GetDepthPassParameters(GraphBuilder, View, OutDepthPrepassTexture.Target);
+
+		View.ParallelMeshDrawCommandPasses[EMeshPass::SingleLayerWaterDepthPrepass].BuildRenderingCommands(GraphBuilder, Scene->GPUScene, PassParameters->InstanceCullingDrawParams);
+
+		if (bRenderInParallel)
+		{
+			GraphBuilder.AddPass(
+				RDG_EVENT_NAME("SingleLayerWaterDepthPrepassParallel"),
+				PassParameters,
+				ERDGPassFlags::Raster | ERDGPassFlags::SkipRenderPass,
+				[this, &View, PassParameters](const FRDGPass* InPass, FRHICommandListImmediate& RHICmdList)
+				{
+					FRDGParallelCommandListSet ParallelCommandListSet(InPass, RHICmdList, GET_STATID(STAT_CLP_WaterSingleLayerPass), *this, View, FParallelCommandListBindings(PassParameters));
+					View.ParallelMeshDrawCommandPasses[EMeshPass::SingleLayerWaterDepthPrepass].DispatchDraw(&ParallelCommandListSet, RHICmdList, &PassParameters->InstanceCullingDrawParams);
+				});
+		}
+		else
+		{
+			GraphBuilder.AddPass(
+				RDG_EVENT_NAME("SingleLayerWaterDepthPrepass"),
+				PassParameters,
+				ERDGPassFlags::Raster,
+				[this, &View, PassParameters](FRHICommandList& RHICmdList)
+				{
+					SetStereoViewport(RHICmdList, View, 1.0f);
+					View.ParallelMeshDrawCommandPasses[EMeshPass::SingleLayerWaterDepthPrepass].DispatchDraw(nullptr, RHICmdList, &PassParameters->InstanceCullingDrawParams);
+				});
+		}
+	}
+
+	AddResolveSceneDepthPass(GraphBuilder, Views, OutDepthPrepassTexture);
+}
+
 static FSceneWithoutWaterTextures AddCopySceneWithoutWaterPass(
 	FRDGBuilder& GraphBuilder,
 	const FSceneViewFamily& ViewFamily, 
@@ -809,6 +978,7 @@ void FDeferredShadingSceneRenderer::RenderSingleLayerWaterReflections(
 void FDeferredShadingSceneRenderer::RenderSingleLayerWater(
 	FRDGBuilder& GraphBuilder,
 	const FSceneTextures& SceneTextures,
+	const FRDGTextureMSAA& SingleLayerWaterDepthPrepassTexture,
 	bool bShouldRenderVolumetricCloud,
 	FSceneWithoutWaterTextures& SceneWithoutWaterTextures,
 	FLumenSceneFrameTemporaries& LumenFrameTemporaries)
@@ -830,13 +1000,21 @@ void FDeferredShadingSceneRenderer::RenderSingleLayerWater(
 		ComposeVolumetricRenderTargetOverSceneUnderWater(GraphBuilder, Views, SceneWithoutWaterTextures, SceneTextures);
 	}
 
-	RenderSingleLayerWaterInner(GraphBuilder, SceneTextures, SceneWithoutWaterTextures);
+	RenderSingleLayerWaterInner(GraphBuilder, SceneTextures, SceneWithoutWaterTextures, SingleLayerWaterDepthPrepassTexture);
 
 	// No SSR or composite needed in Forward. Reflections are applied in the WaterGBuffer pass.
 	if (!IsForwardShadingEnabled(ShaderPlatform))
 	{
+		// Reflection composite expects the depth buffer in FSceneTextures to contain water but the swap of the main depth buffer with the water prepass depth buffer
+		// is only done at the call site after this function returns (for visibility and to keep SceneTextures const), so we need to swap the depth buffers on an internal copy.
+		FSceneTextures SceneTexturesInternal = SceneTextures;
+		if (SingleLayerWaterDepthPrepassTexture.Target != nullptr)
+		{
+			SceneTexturesInternal.Depth = SingleLayerWaterDepthPrepassTexture;
+		}
+
 		// If supported render SSR, the composite pass in non deferred and/or under water effect.
-		RenderSingleLayerWaterReflections(GraphBuilder, SceneTextures, SceneWithoutWaterTextures, LumenFrameTemporaries);
+		RenderSingleLayerWaterReflections(GraphBuilder, SceneTexturesInternal, SceneWithoutWaterTextures, LumenFrameTemporaries);
 	}
 }
 
@@ -872,7 +1050,8 @@ END_SHADER_PARAMETER_STRUCT()
 void FDeferredShadingSceneRenderer::RenderSingleLayerWaterInner(
 	FRDGBuilder& GraphBuilder,
 	const FSceneTextures& SceneTextures,
-	const FSceneWithoutWaterTextures& SceneWithoutWaterTextures)
+	const FSceneWithoutWaterTextures& SceneWithoutWaterTextures,
+	const FRDGTextureMSAA& SingleLayerWaterDepthPrepassTexture)
 {
 	RDG_CSV_STAT_EXCLUSIVE_SCOPE(GraphBuilder, Water);
 	SCOPED_NAMED_EVENT(FDeferredShadingSceneRenderer_RenderSingleLayerWaterPass, FColor::Emerald);
@@ -895,6 +1074,17 @@ void FDeferredShadingSceneRenderer::RenderSingleLayerWaterInner(
 
 	FRDGTextureRef WhiteForwardScreenSpaceShadowMask = SystemTextures.White;
 
+	const bool bHasDepthPrepass = SingleLayerWaterDepthPrepassTexture.Target != nullptr;
+	FDepthStencilBinding DepthStencilBinding;
+	if (bHasDepthPrepass)
+	{
+		DepthStencilBinding = FDepthStencilBinding(SingleLayerWaterDepthPrepassTexture.Target, ERenderTargetLoadAction::ELoad, ERenderTargetLoadAction::ELoad, FExclusiveDepthStencil::DepthRead_StencilRead);
+	}
+	else
+	{
+		DepthStencilBinding = FDepthStencilBinding(SceneTextures.Depth.Target, ERenderTargetLoadAction::ELoad, ERenderTargetLoadAction::ELoad, FExclusiveDepthStencil::DepthWrite_StencilNop);
+	}
+
 	for (int32 ViewIndex = 0; ViewIndex < Views.Num(); ++ViewIndex)
 	{
 		FViewInfo& View = Views[ViewIndex];
@@ -908,7 +1098,7 @@ void FDeferredShadingSceneRenderer::RenderSingleLayerWaterInner(
 		RDG_EVENT_SCOPE_CONDITIONAL(GraphBuilder, Views.Num() > 1, "View%d", ViewIndex);
 		View.BeginRenderView();
 
-		FSingleLayerWaterPassUniformParameters &SLWUniformParameters = *GraphBuilder.AllocParameters<FSingleLayerWaterPassUniformParameters>();
+		FSingleLayerWaterPassUniformParameters& SLWUniformParameters = *GraphBuilder.AllocParameters<FSingleLayerWaterPassUniformParameters>();
 		{
 			const bool bShouldUseBilinearSamplerForDepth = ShouldUseBilinearSamplerForDepthWithoutSingleLayerWater(SceneWithoutWaterTextures.DepthTexture->Desc.Format);
 			const bool bCustomDepthTextureProduced = HasBeenProduced(SceneTextures.CustomDepth.Depth);
@@ -926,7 +1116,7 @@ void FDeferredShadingSceneRenderer::RenderSingleLayerWaterInner(
 			SLWUniformParameters.SceneWithoutSingleLayerWaterTextureSize = FVector2f(DepthTextureSize.X, DepthTextureSize.Y);
 			SLWUniformParameters.SceneWithoutSingleLayerWaterInvTextureSize = FVector2f(1.0f / DepthTextureSize.X, 1.0f / DepthTextureSize.Y);
 
-			const FLightSceneProxy *SelectedForwardDirectionalLightProxy = View.ForwardLightingResources.SelectedForwardDirectionalLightProxy;
+			const FLightSceneProxy* SelectedForwardDirectionalLightProxy = View.ForwardLightingResources.SelectedForwardDirectionalLightProxy;
 			SetupLightCloudTransmittanceParameters(GraphBuilder, Scene, View, SelectedForwardDirectionalLightProxy ? SelectedForwardDirectionalLightProxy->GetLightSceneInfo() : nullptr, SLWUniformParameters.ForwardDirLightCloudShadow);
 		}
 
@@ -937,7 +1127,7 @@ void FDeferredShadingSceneRenderer::RenderSingleLayerWaterInner(
 		PassParameters->VirtualShadowMapSamplingParameters = VirtualShadowMapArray.GetSamplingParameters(GraphBuilder);
 		PassParameters->SingleLayerWater = GraphBuilder.CreateUniformBuffer(&SLWUniformParameters);
 		PassParameters->RenderTargets = GetRenderTargetBindings(ERenderTargetLoadAction::ELoad, BasePassTexturesView);
-		PassParameters->RenderTargets.DepthStencil = FDepthStencilBinding(SceneTextures.Depth.Target, ERenderTargetLoadAction::ELoad, ERenderTargetLoadAction::ELoad, FExclusiveDepthStencil::DepthWrite_StencilWrite);
+		PassParameters->RenderTargets.DepthStencil = DepthStencilBinding;
 
 		View.ParallelMeshDrawCommandPasses[EMeshPass::SingleLayerWaterPass].BuildRenderingCommands(GraphBuilder, Scene->GPUScene, PassParameters->InstanceCullingDrawParams);
 
@@ -967,7 +1157,10 @@ void FDeferredShadingSceneRenderer::RenderSingleLayerWaterInner(
 		}
 	}
 
-	AddResolveSceneDepthPass(GraphBuilder, Views, SceneTextures.Depth);
+	if (!bHasDepthPrepass)
+	{
+		AddResolveSceneDepthPass(GraphBuilder, Views, SceneTextures.Depth);
+	}
 }
 
 class FSingleLayerWaterPassMeshProcessor : public FSceneRenderingAllocatorObject<FSingleLayerWaterPassMeshProcessor>, public FMeshPassProcessor
@@ -1142,6 +1335,9 @@ void FSingleLayerWaterPassMeshProcessor::CollectPSOInitializers(const FSceneText
 			AddRenderTargetInfo(PF_FloatR11G11B10, TexCreate_ShaderResource | TexCreate_RenderTargetable, RenderTargetsInfo);
 		}
 
+		const bool bHasDepthPrepass = IsSingleLayerWaterDepthPrepassEnabled(GetFeatureLevelShaderPlatform(FeatureLevel), FeatureLevel);
+		RenderTargetsInfo.DepthStencilAccess = bHasDepthPrepass ? FExclusiveDepthStencil::DepthRead_StencilRead : FExclusiveDepthStencil::DepthRead_StencilNop;
+
 		AddGraphicsPipelineStateInitializer(
 			VertexFactoryType,
 			Material,
@@ -1158,15 +1354,304 @@ void FSingleLayerWaterPassMeshProcessor::CollectPSOInitializers(const FSceneText
 
 FMeshPassProcessor* CreateSingleLayerWaterPassProcessor(ERHIFeatureLevel::Type FeatureLevel, const FScene* Scene, const FSceneView* InViewIfDynamicMeshCommand, FMeshPassDrawListContext* InDrawListContext)
 {
+	const bool bHasDepthPrepass = IsSingleLayerWaterDepthPrepassEnabled(GetFeatureLevelShaderPlatform(FeatureLevel), FeatureLevel);
 	const FExclusiveDepthStencil::Type SceneBasePassDepthStencilAccess = FScene::GetDefaultBasePassDepthStencilAccess(FeatureLevel);
 
 	FMeshPassProcessorRenderState DrawRenderState;
 
-	// Make sure depth write is enabled.
-	FExclusiveDepthStencil::Type BasePassDepthStencilAccess_DepthWrite = FExclusiveDepthStencil::Type(SceneBasePassDepthStencilAccess | FExclusiveDepthStencil::DepthWrite);
+	// Make sure depth write is enabled if no prepass is used.
+	FExclusiveDepthStencil::Type BasePassDepthStencilAccess_DepthWrite = FExclusiveDepthStencil::Type(bHasDepthPrepass ? FExclusiveDepthStencil::DepthRead : SceneBasePassDepthStencilAccess | FExclusiveDepthStencil::DepthWrite);
 	SetupBasePassState(BasePassDepthStencilAccess_DepthWrite, false, DrawRenderState);
+	if (bHasDepthPrepass)
+	{
+		// Set depth stencil test to only pass if depth and stencil are equal to the values written by the prepass
+		DrawRenderState.SetDepthStencilState(TStaticDepthStencilState<
+			false, CF_Equal,							// Depth test
+			true, CF_Equal, SO_Keep, SO_Keep, SO_Keep,	// Front face stencil 
+			true, CF_Equal, SO_Keep, SO_Keep, SO_Keep,	// Back face stencil
+			0xFF, 0x0									// Stencil read/write masks
+		>::GetRHI());
+		DrawRenderState.SetStencilRef(1);
+	}
 
 	return new FSingleLayerWaterPassMeshProcessor(Scene, FeatureLevel, InViewIfDynamicMeshCommand, DrawRenderState, InDrawListContext);
 }
 
 REGISTER_MESHPASSPROCESSOR_AND_PSOCOLLECTOR(SingleLayerWater, CreateSingleLayerWaterPassProcessor, EShadingPath::Deferred, EMeshPass::SingleLayerWaterPass, EMeshPassFlags::MainView);
+
+
+class FSingleLayerWaterDepthPrepassMeshProcessor : public FSceneRenderingAllocatorObject<FSingleLayerWaterDepthPrepassMeshProcessor>, public FMeshPassProcessor
+{
+public:
+	FSingleLayerWaterDepthPrepassMeshProcessor(const FScene* Scene, ERHIFeatureLevel::Type FeatureLevel, const FSceneView* InViewIfDynamicMeshCommand, const FMeshPassProcessorRenderState& InPassDrawRenderState, FMeshPassDrawListContext* InDrawListContext);
+
+	virtual void AddMeshBatch(const FMeshBatch& RESTRICT MeshBatch, uint64 BatchElementMask, const FPrimitiveSceneProxy* RESTRICT PrimitiveSceneProxy, int32 StaticMeshId = -1) override final;
+	virtual void CollectPSOInitializers(const FSceneTexturesConfig& SceneTexturesConfig, const FMaterial& Material, const FVertexFactoryType* VertexFactoryType, const FPSOPrecacheParams& PreCacheParams, TArray<FPSOPrecacheData>& PSOInitializers) override final;
+
+private:
+	bool TryAddMeshBatch(
+		const FMeshBatch& RESTRICT MeshBatch,
+		uint64 BatchElementMask,
+		const FPrimitiveSceneProxy* RESTRICT PrimitiveSceneProxy,
+		int32 StaticMeshId,
+		const FMaterialRenderProxy& MaterialRenderProxy,
+		const FMaterial& Material);
+
+	template<bool bPositionOnly>
+	bool Process(
+		const FMeshBatch& RESTRICT MeshBatch,
+		uint64 BatchElementMask,
+		const FPrimitiveSceneProxy* RESTRICT PrimitiveSceneProxy,
+		int32 StaticMeshId,
+		const FMaterialRenderProxy& RESTRICT MaterialRenderProxy,
+		const FMaterial& RESTRICT MaterialResource,
+		ERasterizerFillMode MeshFillMode,
+		ERasterizerCullMode MeshCullMode,
+		EBlendMode BlendMode);
+
+	template<bool bPositionOnly>
+	void CollectPSOInitializersInternal(
+		const FSceneTexturesConfig& SceneTexturesConfig,
+		const FVertexFactoryType* VertexFactoryType,
+		const FMaterial& RESTRICT MaterialResource,
+		ERasterizerFillMode MeshFillMode,
+		ERasterizerCullMode MeshCullMode,
+		const FPSOPrecacheParams& PreCacheParams, 
+		TArray<FPSOPrecacheData>& PSOInitializers);
+
+	FMeshPassProcessorRenderState PassDrawRenderState;
+};
+
+FSingleLayerWaterDepthPrepassMeshProcessor::FSingleLayerWaterDepthPrepassMeshProcessor(const FScene* Scene, ERHIFeatureLevel::Type FeatureLevel, const FSceneView* InViewIfDynamicMeshCommand, const FMeshPassProcessorRenderState& InPassDrawRenderState, FMeshPassDrawListContext* InDrawListContext)
+	: FMeshPassProcessor(EMeshPass::SingleLayerWaterPass, Scene, FeatureLevel, InViewIfDynamicMeshCommand, InDrawListContext)
+	, PassDrawRenderState(InPassDrawRenderState)
+{
+
+}
+
+void FSingleLayerWaterDepthPrepassMeshProcessor::AddMeshBatch(const FMeshBatch& RESTRICT MeshBatch, uint64 BatchElementMask, const FPrimitiveSceneProxy* RESTRICT PrimitiveSceneProxy, int32 StaticMeshId)
+{
+	// Early out if the depth prepass for water is disabled
+	if (!IsSingleLayerWaterDepthPrepassEnabled(GetFeatureLevelShaderPlatform(FeatureLevel), FeatureLevel))
+	{
+		return;
+	}
+
+	const FMaterialRenderProxy* MaterialRenderProxy = MeshBatch.MaterialRenderProxy;
+	while (MaterialRenderProxy)
+	{
+		const FMaterial* Material = MaterialRenderProxy->GetMaterialNoFallback(FeatureLevel);
+		if (Material)
+		{
+			if (TryAddMeshBatch(MeshBatch, BatchElementMask, PrimitiveSceneProxy, StaticMeshId, *MaterialRenderProxy, *Material))
+			{
+				break;
+			}
+		}
+
+		MaterialRenderProxy = MaterialRenderProxy->GetFallback(FeatureLevel);
+	}
+}
+
+bool FSingleLayerWaterDepthPrepassMeshProcessor::TryAddMeshBatch(
+	const FMeshBatch& RESTRICT MeshBatch,
+	uint64 BatchElementMask,
+	const FPrimitiveSceneProxy* RESTRICT PrimitiveSceneProxy,
+	int32 StaticMeshId,
+	const FMaterialRenderProxy& MaterialRenderProxy,
+	const FMaterial& Material)
+{
+	if (Material.GetShadingModels().HasShadingModel(MSM_SingleLayerWater))
+	{
+		// Determine the mesh's material and blend mode.
+		const EBlendMode BlendMode = Material.GetBlendMode();
+		const FMeshDrawingPolicyOverrideSettings OverrideSettings = ComputeMeshOverrideSettings(MeshBatch);
+		const ERasterizerFillMode MeshFillMode = ComputeMeshFillMode(Material, OverrideSettings);
+		const ERasterizerCullMode MeshCullMode = ComputeMeshCullMode(Material, OverrideSettings);
+
+		if (BlendMode == BLEND_Opaque
+			&& MeshBatch.VertexFactory->SupportsPositionOnlyStream()
+			&& !Material.MaterialModifiesMeshPosition_RenderThread()
+			&& Material.WritesEveryPixel())
+		{
+			const FMaterialRenderProxy& DefaultProxy = *UMaterial::GetDefaultMaterial(MD_Surface)->GetRenderProxy();
+			const FMaterial& DefaultMaterial = *DefaultProxy.GetMaterialNoFallback(FeatureLevel);
+			return Process<true>(MeshBatch, BatchElementMask, PrimitiveSceneProxy, StaticMeshId, DefaultProxy, DefaultMaterial, MeshFillMode, MeshCullMode, BlendMode);
+		}
+		else
+		{
+			const bool bMaterialMasked = !Material.WritesEveryPixel();
+			const FMaterialRenderProxy* EffectiveMaterialRenderProxy = &MaterialRenderProxy;
+			const FMaterial* EffectiveMaterial = &Material;
+
+			if (!bMaterialMasked && !Material.MaterialModifiesMeshPosition_RenderThread())
+			{
+				// Override with the default material for opaque materials that are not two sided
+				EffectiveMaterialRenderProxy = UMaterial::GetDefaultMaterial(MD_Surface)->GetRenderProxy();
+				EffectiveMaterial = EffectiveMaterialRenderProxy->GetMaterialNoFallback(FeatureLevel);
+				check(EffectiveMaterial);
+			}
+
+			return Process<false>(MeshBatch, BatchElementMask, PrimitiveSceneProxy, StaticMeshId, *EffectiveMaterialRenderProxy, *EffectiveMaterial, MeshFillMode, MeshCullMode, BlendMode);
+		}
+	}
+
+	return true;
+}
+
+template<bool bPositionOnly>
+bool FSingleLayerWaterDepthPrepassMeshProcessor::Process(
+	const FMeshBatch& RESTRICT MeshBatch,
+	uint64 BatchElementMask,
+	const FPrimitiveSceneProxy* RESTRICT PrimitiveSceneProxy,
+	int32 StaticMeshId,
+	const FMaterialRenderProxy& RESTRICT MaterialRenderProxy,
+	const FMaterial& RESTRICT MaterialResource,
+	ERasterizerFillMode MeshFillMode,
+	ERasterizerCullMode MeshCullMode,
+	EBlendMode BlendMode)
+{
+	TMeshProcessorShaders<TDepthOnlyVS<bPositionOnly>, FDepthOnlyPS> DepthPassShaders;
+	FShaderPipelineRef ShaderPipeline;
+
+	if (!GetDepthPassShaders<bPositionOnly>(
+		MaterialResource,
+		MeshBatch.VertexFactory->GetType(),
+		FeatureLevel,
+		MaterialResource.MaterialUsesPixelDepthOffset_GameThread(),
+		DepthPassShaders.VertexShader,
+		DepthPassShaders.PixelShader,
+		ShaderPipeline))
+	{
+		return false;
+	}
+
+	FMeshMaterialShaderElementData ShaderElementData;
+	ShaderElementData.InitializeMeshMaterialData(ViewIfDynamicMeshCommand, PrimitiveSceneProxy, MeshBatch, StaticMeshId, true);
+
+	FMeshDrawCommandSortKey SortKey = CalculateDepthPassMeshStaticSortKey(BlendMode, DepthPassShaders.VertexShader.GetShader(), DepthPassShaders.PixelShader.GetShader());
+
+	BuildMeshDrawCommands(
+		MeshBatch,
+		BatchElementMask,
+		PrimitiveSceneProxy,
+		MaterialRenderProxy,
+		MaterialResource,
+		PassDrawRenderState,
+		DepthPassShaders,
+		MeshFillMode,
+		MeshCullMode,
+		SortKey,
+		bPositionOnly ? EMeshPassFeatures::PositionOnly : EMeshPassFeatures::Default,
+		ShaderElementData);
+
+	return true;
+}
+
+void FSingleLayerWaterDepthPrepassMeshProcessor::CollectPSOInitializers(const FSceneTexturesConfig& SceneTexturesConfig, const FMaterial& Material, const FVertexFactoryType* VertexFactoryType, const FPSOPrecacheParams& PreCacheParams, TArray<FPSOPrecacheData>& PSOInitializers)
+{
+	if (Material.GetShadingModels().HasShadingModel(MSM_SingleLayerWater) && IsSingleLayerWaterDepthPrepassEnabled(GetFeatureLevelShaderPlatform(FeatureLevel), FeatureLevel))
+	{
+		// Determine the mesh's material and blend mode.
+		const FMeshDrawingPolicyOverrideSettings OverrideSettings = ComputeMeshOverrideSettings(PreCacheParams);
+		const ERasterizerFillMode MeshFillMode = ComputeMeshFillMode(Material, OverrideSettings);
+		const ERasterizerCullMode MeshCullMode = ComputeMeshCullMode(Material, OverrideSettings);
+		const EBlendMode BlendMode = Material.GetBlendMode();
+		const bool bSupportPositionOnlyStream = VertexFactoryType->SupportsPositionOnly();
+
+		if (BlendMode == BLEND_Opaque
+			&& bSupportPositionOnlyStream
+			&& !Material.MaterialModifiesMeshPosition_RenderThread()
+			&& Material.WritesEveryPixel())
+		{
+			const FMaterialRenderProxy& DefaultProxy = *UMaterial::GetDefaultMaterial(MD_Surface)->GetRenderProxy();
+			const FMaterial& DefaultMaterial = *DefaultProxy.GetMaterialNoFallback(FeatureLevel);
+
+			CollectPSOInitializersInternal<true>(SceneTexturesConfig, VertexFactoryType, DefaultMaterial, MeshFillMode, MeshCullMode, PreCacheParams, PSOInitializers);
+		}
+		else
+		{
+			const bool bMaterialMasked = !Material.WritesEveryPixel();
+			const FMaterial* EffectiveMaterial = &Material;
+
+			if (!bMaterialMasked && !Material.MaterialModifiesMeshPosition_RenderThread())
+			{
+				// Override with the default material for opaque materials that are not two sided
+				FMaterialRenderProxy* EffectiveMaterialRenderProxy = UMaterial::GetDefaultMaterial(MD_Surface)->GetRenderProxy();
+				EffectiveMaterial = EffectiveMaterialRenderProxy->GetMaterialNoFallback(FeatureLevel);
+				check(EffectiveMaterial);
+			}
+
+			CollectPSOInitializersInternal<false>(SceneTexturesConfig, VertexFactoryType, *EffectiveMaterial, MeshFillMode, MeshCullMode, PreCacheParams, PSOInitializers);
+		}
+	}
+}
+
+template<bool bPositionOnly>
+void FSingleLayerWaterDepthPrepassMeshProcessor::CollectPSOInitializersInternal(
+	const FSceneTexturesConfig& SceneTexturesConfig,
+	const FVertexFactoryType* VertexFactoryType,
+	const FMaterial& RESTRICT MaterialResource,
+	ERasterizerFillMode MeshFillMode,
+	ERasterizerCullMode MeshCullMode,
+	const FPSOPrecacheParams& PreCacheParams,
+	TArray<FPSOPrecacheData>& PSOInitializers)
+{
+	TMeshProcessorShaders<TDepthOnlyVS<bPositionOnly>, FDepthOnlyPS> DepthPassShaders;
+	FShaderPipelineRef ShaderPipeline;
+
+	if (!GetDepthPassShaders<bPositionOnly>(
+		MaterialResource,
+		VertexFactoryType,
+		FeatureLevel,
+		MaterialResource.MaterialUsesPixelDepthOffset_GameThread(),
+		DepthPassShaders.VertexShader,
+		DepthPassShaders.PixelShader,
+		ShaderPipeline))
+	{
+		return;
+	}
+
+	FGraphicsPipelineRenderTargetsInfo RenderTargetsInfo;
+	RenderTargetsInfo.NumSamples = SceneTexturesConfig.NumSamples;
+
+	ETextureCreateFlags DepthStencilCreateFlags = SceneTexturesConfig.DepthCreateFlags;
+	SetupDepthStencilInfo(PF_DepthStencil, DepthStencilCreateFlags, ERenderTargetLoadAction::ELoad,
+		ERenderTargetLoadAction::ELoad, FExclusiveDepthStencil::DepthWrite_StencilWrite, RenderTargetsInfo);
+
+	AddGraphicsPipelineStateInitializer(
+		VertexFactoryType,
+		MaterialResource,
+		PassDrawRenderState,
+		RenderTargetsInfo,
+		DepthPassShaders,
+		MeshFillMode,
+		MeshCullMode,
+		(EPrimitiveType)PreCacheParams.PrimitiveType,
+		bPositionOnly ? EMeshPassFeatures::PositionOnly : EMeshPassFeatures::Default,
+		PSOInitializers);
+}
+
+FMeshPassProcessor* CreateSingleLayerWaterDepthPrepassProcessor(ERHIFeatureLevel::Type FeatureLevel, const FScene* Scene, const FSceneView* InViewIfDynamicMeshCommand, FMeshPassDrawListContext* InDrawListContext)
+{
+	const FExclusiveDepthStencil::Type SceneBasePassDepthStencilAccess = FScene::GetDefaultBasePassDepthStencilAccess(FeatureLevel);
+
+	FMeshPassProcessorRenderState DrawRenderState;
+
+	// Make sure depth write is enabled.
+	FExclusiveDepthStencil::Type BasePassDepthStencilAccess_DepthWrite = FExclusiveDepthStencil::Type(SceneBasePassDepthStencilAccess | FExclusiveDepthStencil::DepthWrite);
+	
+	// Disable color writes, enable depth tests and writes.
+	DrawRenderState.SetBlendState(TStaticBlendState<CW_NONE>::GetRHI());
+	DrawRenderState.SetDepthStencilState(TStaticDepthStencilState<
+		true, CF_DepthNearOrEqual,						// Depth test
+		true, CF_Always, SO_Keep, SO_Keep, SO_Replace,	// Front face stencil 
+		true, CF_Always, SO_Keep, SO_Keep, SO_Replace	// Back face stencil
+	>::GetRHI());
+	DrawRenderState.SetDepthStencilAccess(BasePassDepthStencilAccess_DepthWrite);
+	DrawRenderState.SetStencilRef(1);
+
+	return new FSingleLayerWaterDepthPrepassMeshProcessor(Scene, FeatureLevel, InViewIfDynamicMeshCommand, DrawRenderState, InDrawListContext);
+}
+
+REGISTER_MESHPASSPROCESSOR_AND_PSOCOLLECTOR(SingleLayerWaterDepthPrepass, CreateSingleLayerWaterDepthPrepassProcessor, EShadingPath::Deferred, EMeshPass::SingleLayerWaterDepthPrepass, EMeshPassFlags::MainView);
diff --git a/Engine/Source/Runtime/Renderer/Private/SingleLayerWaterRendering.h b/Engine/Source/Runtime/Renderer/Private/SingleLayerWaterRendering.h
--- a/Engine/Source/Runtime/Renderer/Private/SingleLayerWaterRendering.h
+++ b/Engine/Source/Runtime/Renderer/Private/SingleLayerWaterRendering.h
@@ -25,6 +25,7 @@ struct FSceneWithoutWaterTextures
 
 bool ShouldRenderSingleLayerWater(TArrayView<const FViewInfo> Views);
 bool ShouldRenderSingleLayerWaterSkippedRenderEditorNotification(TArrayView<const FViewInfo> Views);
+bool ShouldRenderSingleLayerWaterDepthPrepass(TArrayView<const FViewInfo> Views);
 bool ShouldUseBilinearSamplerForDepthWithoutSingleLayerWater(EPixelFormat DepthTextureFormat);
 
 class FWaterTileVS : public FGlobalShader
diff --git a/Engine/Source/Runtime/Renderer/Private/VirtualShadowMaps/VirtualShadowMapArray.cpp b/Engine/Source/Runtime/Renderer/Private/VirtualShadowMaps/VirtualShadowMapArray.cpp
--- a/Engine/Source/Runtime/Renderer/Private/VirtualShadowMaps/VirtualShadowMapArray.cpp
+++ b/Engine/Source/Runtime/Renderer/Private/VirtualShadowMaps/VirtualShadowMapArray.cpp
@@ -655,6 +655,7 @@ class FGeneratePageFlagsFromPixelsCS : public FVirtualPageManagementShader
 		SHADER_PARAMETER_RDG_BUFFER_SRV(StructuredBuffer< uint >, DirectionalLightIds)
 		SHADER_PARAMETER_RDG_BUFFER_SRV(StructuredBuffer< uint >, PrunedLightGridData)
 		SHADER_PARAMETER_RDG_BUFFER_SRV(StructuredBuffer< uint >, PrunedNumCulledLightsGrid)
+		SHADER_PARAMETER_RDG_TEXTURE(Texture2D<float>, SingeLayerWaterDepthTexture)
 		SHADER_PARAMETER_RDG_UNIFORM_BUFFER(FStrataGlobalUniformParameters, Strata)
 		RDG_BUFFER_ACCESS(IndirectBufferArgs, ERHIAccess::IndirectArgs)
 		SHADER_PARAMETER(uint32, InputType)
@@ -1107,7 +1108,8 @@ void FVirtualShadowMapArray::BuildPageAllocations(
 	const FEngineShowFlags& EngineShowFlags,
 	const FSortedLightSetSceneInfo& SortedLightsInfo,
 	const TArray<FVisibleLightInfo, SceneRenderingAllocator>& VisibleLightInfos,
-	const TArray<Nanite::FRasterResults, TInlineAllocator<2>>& NaniteRasterResults)
+	const TArray<Nanite::FRasterResults, TInlineAllocator<2>>& NaniteRasterResults,
+	FRDGTextureRef SingleLayerWaterDepthTexture)
 {
 	check(IsEnabled());
 
@@ -1394,7 +1396,7 @@ void FVirtualShadowMapArray::BuildPageAllocations(
 				auto GeneratePageFlags = [&](const EVirtualShadowMapProjectionInputType InputType)
 				{
 					FGeneratePageFlagsFromPixelsCS::FPermutationDomain PermutationVector;
-					PermutationVector.Set<FGeneratePageFlagsFromPixelsCS::FInputType>(InputType == EVirtualShadowMapProjectionInputType::HairStrands ? 1u : 0u);
+					PermutationVector.Set<FGeneratePageFlagsFromPixelsCS::FInputType>(static_cast<uint32>(InputType));
 					FGeneratePageFlagsFromPixelsCS::FParameters* PassParameters = GraphBuilder.AllocParameters< FGeneratePageFlagsFromPixelsCS::FParameters >();
 					PassParameters->VirtualShadowMap = GetUniformBuffer(GraphBuilder);
 
@@ -1407,6 +1409,7 @@ void FVirtualShadowMapArray::BuildPageAllocations(
 					PassParameters->DirectionalLightIds = GraphBuilder.CreateSRV(DirectionalLightIdsRDG);
 					PassParameters->PrunedLightGridData = GraphBuilder.CreateSRV(PrunedLightGridDataRDG);
 					PassParameters->PrunedNumCulledLightsGrid = GraphBuilder.CreateSRV(PrunedNumCulledLightsGridRDG);
+					PassParameters->SingeLayerWaterDepthTexture = SingleLayerWaterDepthTexture;
 					PassParameters->NumDirectionalLightSmInds = uint32(DirectionalLightIds.Num());
 					PassParameters->PageDilationBorderSizeLocal = CVarPageDilationBorderSizeLocal.GetValueOnRenderThread();
 					PassParameters->PageDilationBorderSizeDirectional = CVarPageDilationBorderSizeDirectional.GetValueOnRenderThread();
@@ -1441,7 +1444,8 @@ void FVirtualShadowMapArray::BuildPageAllocations(
 					}
 				};
 
-				GeneratePageFlags(EVirtualShadowMapProjectionInputType::GBuffer);
+				const bool bHasValidSingleLayerWaterDepth = SingleLayerWaterDepthTexture != nullptr;
+				GeneratePageFlags(bHasValidSingleLayerWaterDepth ? EVirtualShadowMapProjectionInputType::GBufferAndSingleLayerWaterDepth : EVirtualShadowMapProjectionInputType::GBuffer);
 				if (HairStrands::HasViewHairStrandsData(View))
 				{
 					GeneratePageFlags(EVirtualShadowMapProjectionInputType::HairStrands);
diff --git a/Engine/Source/Runtime/Renderer/Private/VirtualShadowMaps/VirtualShadowMapArray.h b/Engine/Source/Runtime/Renderer/Private/VirtualShadowMaps/VirtualShadowMapArray.h
--- a/Engine/Source/Runtime/Renderer/Private/VirtualShadowMaps/VirtualShadowMapArray.h
+++ b/Engine/Source/Runtime/Renderer/Private/VirtualShadowMaps/VirtualShadowMapArray.h
@@ -256,7 +256,8 @@ public:
 		const FEngineShowFlags& EngineShowFlags,
 		const FSortedLightSetSceneInfo& SortedLights, 
 		const TArray<FVisibleLightInfo, SceneRenderingAllocator> &VisibleLightInfos, 
-		const TArray<Nanite::FRasterResults, TInlineAllocator<2>> &NaniteRasterResults);
+		const TArray<Nanite::FRasterResults, TInlineAllocator<2>> &NaniteRasterResults,
+		FRDGTextureRef SingleLayerWaterDepthTexture);
 
 	bool IsAllocated() const
 	{
diff --git a/Engine/Source/Runtime/Renderer/Private/VirtualShadowMaps/VirtualShadowMapProjection.h b/Engine/Source/Runtime/Renderer/Private/VirtualShadowMaps/VirtualShadowMapProjection.h
--- a/Engine/Source/Runtime/Renderer/Private/VirtualShadowMaps/VirtualShadowMapProjection.h
+++ b/Engine/Source/Runtime/Renderer/Private/VirtualShadowMaps/VirtualShadowMapProjection.h
@@ -12,10 +12,12 @@
 
 class FVirtualShadowMapClipmap;
 
+// Note: Must match the definitions in VirtualShadowMapPageManagement.usf!
 enum class EVirtualShadowMapProjectionInputType
 {
 	GBuffer = 0,
-	HairStrands = 1
+	HairStrands = 1,
+	GBufferAndSingleLayerWaterDepth = 2
 };
 const TCHAR* ToString(EVirtualShadowMapProjectionInputType In);
 
diff --git a/Engine/Source/Runtime/Renderer/Public/MeshPassProcessor.h b/Engine/Source/Runtime/Renderer/Public/MeshPassProcessor.h
--- a/Engine/Source/Runtime/Renderer/Public/MeshPassProcessor.h
+++ b/Engine/Source/Runtime/Renderer/Public/MeshPassProcessor.h
@@ -34,6 +34,7 @@ namespace EMeshPass
 		AnisotropyPass,
 		SkyPass,
 		SingleLayerWaterPass,
+		SingleLayerWaterDepthPrepass,
 		CSMShadowDepth,
 		VSMShadowDepth,
 		Distortion,
@@ -79,6 +80,7 @@ inline const TCHAR* GetMeshPassName(EMeshPass::Type MeshPass)
 	case EMeshPass::AnisotropyPass: return TEXT("AnisotropyPass");
 	case EMeshPass::SkyPass: return TEXT("SkyPass");
 	case EMeshPass::SingleLayerWaterPass: return TEXT("SingleLayerWaterPass");
+	case EMeshPass::SingleLayerWaterDepthPrepass: return TEXT("SingleLayerWaterDepthPrepass");
 	case EMeshPass::CSMShadowDepth: return TEXT("CSMShadowDepth");
 	case EMeshPass::VSMShadowDepth: return TEXT("VSMShadowDepth");
 	case EMeshPass::Distortion: return TEXT("Distortion");
