Fixed missing update for uniform buffer in FSceneTextures after swapping depth buffers in SingleLayerWater rendering. Fixes broken reflections on water.

#rnx

diff --git a/Engine/Source/Runtime/Renderer/Private/SingleLayerWaterRendering.cpp b/Engine/Source/Runtime/Renderer/Private/SingleLayerWaterRendering.cpp
--- a/Engine/Source/Runtime/Renderer/Private/SingleLayerWaterRendering.cpp
+++ b/Engine/Source/Runtime/Renderer/Private/SingleLayerWaterRendering.cpp
@@ -1011,6 +1011,8 @@ void FDeferredShadingSceneRenderer::RenderSingleLayerWater(
 		if (SingleLayerWaterDepthPrepassTexture.Target != nullptr)
 		{
 			SceneTexturesInternal.Depth = SingleLayerWaterDepthPrepassTexture;
+			// Rebuild scene textures uniform buffer to include new depth buffer.
+			SceneTexturesInternal.UniformBuffer = CreateSceneTextureUniformBuffer(GraphBuilder, &SceneTexturesInternal, FeatureLevel, SceneTexturesInternal.SetupMode);
 		}
 
 		// If supported render SSR, the composite pass in non deferred and/or under water effect.
