Fix robomerge issue with incorrect auto resolve

diff --git a/Engine/Source/Runtime/Renderer/Public/MeshPassProcessor.h b/Engine/Source/Runtime/Renderer/Public/MeshPassProcessor.h
--- a/Engine/Source/Runtime/Renderer/Public/MeshPassProcessor.h
+++ b/Engine/Source/Runtime/Renderer/Public/MeshPassProcessor.h
@@ -112,9 +112,9 @@ inline const TCHAR* GetMeshPassName(EMeshPass::Type MeshPass)
 	}
 
 #if WITH_EDITOR
-	static_assert(EMeshPass::Num == 27 + 4, "Need to update switch(MeshPass) after changing EMeshPass");
+	static_assert(EMeshPass::Num == 28 + 4, "Need to update switch(MeshPass) after changing EMeshPass"); // GUID to prevent incorrect auto-resolves, please change when changing the expression: {FED9577B-F826-478C-841F-6A01BCB4415E}
 #else
-	static_assert(EMeshPass::Num == 27, "Need to update switch(MeshPass) after changing EMeshPass");
+	static_assert(EMeshPass::Num == 28, "Need to update switch(MeshPass) after changing EMeshPass"); // GUID to prevent incorrect auto-resolves, please change when changing the expression: {FED9577B-F826-478C-841F-6A01BCB4415E}
 #endif
 
 	checkf(0, TEXT("Missing case for EMeshPass %u"), (uint32)MeshPass);
