Add explicit unsigned uint8 storage annotation for EMeshPass::Type enum

* Fixes FCachedMeshDrawCommandInfo::MeshPass bitfield when EMeshPass::Num is exactly 1<<EMeshPass::NumBits and current compiler uses signed default enum storage
* Improved error message in FPrimitiveSceneInfo::CacheMeshDrawCommands that checks FCachedMeshDrawCommandInfo::MeshPass field

diff --git a/Engine/Source/Runtime/Renderer/Private/PrimitiveSceneInfo.cpp b/Engine/Source/Runtime/Renderer/Private/PrimitiveSceneInfo.cpp
--- a/Engine/Source/Runtime/Renderer/Private/PrimitiveSceneInfo.cpp
+++ b/Engine/Source/Runtime/Renderer/Private/PrimitiveSceneInfo.cpp
@@ -354,8 +354,11 @@ void FPrimitiveSceneInfo::CacheMeshDrawCommands(FRHICommandListImmediate& RHICmd
 							MeshRelevance.CommandInfosBase++;
 
 							int CommandInfoIndex = MeshAndInfo.MeshIndex * EMeshPass::Num + PassType;
-							check(SceneInfo->StaticMeshCommandInfos[CommandInfoIndex].MeshPass == EMeshPass::Num);
-							SceneInfo->StaticMeshCommandInfos[CommandInfoIndex] = CommandInfo;
+							FCachedMeshDrawCommandInfo& CurrentCommandInfo = SceneInfo->StaticMeshCommandInfos[CommandInfoIndex];
+							checkf(CurrentCommandInfo.MeshPass == EMeshPass::Num,
+								TEXT("SceneInfo->StaticMeshCommandInfos[%d] is not expected to be initialized yet. MeshPass is %d, but expected EMeshPass::Num (%d)."),
+								CommandInfoIndex, (int32)EMeshPass::Num, CurrentCommandInfo.MeshPass);
+							CurrentCommandInfo = CommandInfo;
 						}
 					}
 
diff --git a/Engine/Source/Runtime/Renderer/Public/MeshPassProcessor.h b/Engine/Source/Runtime/Renderer/Public/MeshPassProcessor.h
--- a/Engine/Source/Runtime/Renderer/Public/MeshPassProcessor.h
+++ b/Engine/Source/Runtime/Renderer/Public/MeshPassProcessor.h
@@ -27,7 +27,7 @@ struct FSceneTexturesConfig;
 /** Mesh pass types supported. */
 namespace EMeshPass
 {
-	enum Type
+	enum Type : uint8
 	{
 		DepthPass,
 		BasePass,
@@ -70,6 +70,7 @@ namespace EMeshPass
 	};
 }
 static_assert(EMeshPass::Num <= (1 << EMeshPass::NumBits), "EMeshPass::Num will not fit in EMeshPass::NumBits");
+static_assert(EMeshPass::NumBits <= sizeof(EMeshPass::Type) * 8, "EMeshPass::Type storage is too small");
 
 inline const TCHAR* GetMeshPassName(EMeshPass::Type MeshPass)
 {
