Compute shader base pass and lightmap policy support

diff --git a/Engine/Shaders/Private/BasePassPixelShader.usf b/Engine/Shaders/Private/BasePassPixelShader.usf
--- a/Engine/Shaders/Private/BasePassPixelShader.usf
+++ b/Engine/Shaders/Private/BasePassPixelShader.usf
@@ -50,6 +50,10 @@
 #define USE_DBUFFER 0
 #endif
 
+#ifndef COMPUTE_SHADED
+#define COMPUTE_SHADED 0
+#endif
+
 // For the base pass STRATA_INLINE_SHADING is defined in BasePassRendering.h
 #include "SHCommon.ush"
 #include "/Engine/Generated/Material.ush"
@@ -981,33 +985,33 @@ void FPixelShaderInOut_MainPS(
 	float DBufferOpacity = 1.0f;
 #if USE_DBUFFER && MATERIALDECALRESPONSEMASK && !MATERIALBLENDING_ANY_TRANSLUCENT && !MATERIAL_SHADINGMODEL_SINGLELAYERWATER
 	// apply decals from the DBuffer
-	#if PC_D3D
+#if PC_D3D
 	//Temporary workaround to avoid crashes on AMD, revert back to BRANCH
 	FLATTEN
-	#else
+#else
 	BRANCH
-	#endif
-	if ((GetPrimitiveData(MaterialParameters).Flags & PRIMITIVE_SCENE_DATA_FLAG_DECAL_RECEIVER) != 0 && View.ShowDecalsMask > 0)
-	{
-		uint ValidDBufferTargetMask = GetDBufferTargetMask(uint2(In.SvPosition.xy)) & MATERIALDECALRESPONSEMASK;
-
-		BRANCH
-		if (ValidDBufferTargetMask)
+#endif
+		if ((GetPrimitiveData(MaterialParameters).Flags & PRIMITIVE_SCENE_DATA_FLAG_DECAL_RECEIVER) != 0 && View.ShowDecalsMask > 0)
 		{
-			float2 BufferUV = SvPositionToBufferUV(In.SvPosition);
-		#if STRATA_ENABLED 
-			#if STRATA_INLINE_SHADING && !STRATA_OPTIMIZED_UNLIT
-			const FStrataDBuffer StrataBufferData = GetDBufferData(BufferUV, ValidDBufferTargetMask);
-			ApplyDBufferData(StrataBufferData, StrataPixelHeader, StrataData);
-			DBufferOpacity = StrataBufferData.Coverage;
-			#endif
-		#else
-			FDBufferData DBufferData = GetDBufferData(BufferUV, ValidDBufferTargetMask);
-			ApplyDBufferData(DBufferData, MaterialParameters.WorldNormal, SubsurfaceColor, Roughness, BaseColor, Metallic, Specular);
-			DBufferOpacity = (DBufferData.ColorOpacity + DBufferData.NormalOpacity + DBufferData.RoughnessOpacity) * (1.0f / 3.0f);
-		#endif // STRATA_ENABLED
+			uint ValidDBufferTargetMask = GetDBufferTargetMask(uint2(In.SvPosition.xy)) & MATERIALDECALRESPONSEMASK;
+
+			BRANCH
+			if (ValidDBufferTargetMask)
+			{
+				float2 BufferUV = SvPositionToBufferUV(In.SvPosition);
+			#if STRATA_ENABLED 
+			   #if STRATA_INLINE_SHADING && !STRATA_OPTIMIZED_UNLIT
+				const FStrataDBuffer StrataBufferData = GetDBufferData(BufferUV, ValidDBufferTargetMask);
+				ApplyDBufferData(StrataBufferData, StrataPixelHeader, StrataData);
+				DBufferOpacity = StrataBufferData.Coverage;
+			   #endif
+			#else
+				FDBufferData DBufferData = GetDBufferData(BufferUV, ValidDBufferTargetMask);
+				ApplyDBufferData(DBufferData, MaterialParameters.WorldNormal, SubsurfaceColor, Roughness, BaseColor, Metallic, Specular);
+				DBufferOpacity = (DBufferData.ColorOpacity + DBufferData.NormalOpacity + DBufferData.RoughnessOpacity) * (1.0f / 3.0f);
+			#endif // STRATA_ENABLED
+			}
 		}
-	}
 #endif
 
 	const float BaseMaterialCoverageOverWater = Opacity;
@@ -1643,7 +1647,7 @@ void FPixelShaderInOut_MainPS(
 		FStrataIntegrationSettings Settings = InitStrataIntegrationSettings(false /*bForceFullyRough*/, StrataStruct.bRoughDiffuse, StrataStruct.PeelLayersAboveDepth);
 
 	#if STRATA_ADVANCED_DEBUG_ENABLED
-		Settings.SliceStoringDebugStrataTreeData = all(uint2(float2(View.CursorPosition) * View.ViewResolutionFraction) == PixelPos) ? StrataStruct.SliceStoringDebugStrataTreeDataWithoutMRT : Settings.SliceStoringDebugStrataTreeData;
+		Settings.SliceStoringDebugStrataTree = all(uint2(float2(View.CursorPosition) * View.ViewResolutionFraction) == PixelPos) ? StrataStruct.SliceStoringDebugStrataTree : Settings.SliceStoringDebugStrataTree;
 	#endif
 
 		// Generate the strata material data to write out
@@ -1667,13 +1671,12 @@ void FPixelShaderInOut_MainPS(
 		);
 
 		// Write out MRT data
-	#if STRATA_BASE_PASS_MRT_OUTPUT_COUNT != 3
+	#if STRATA_BASE_PASS_MRT_OUTPUT_COUNT != 2
 	#error Strata STRATA_BASE_PASS_MRT_OUTPUT_COUNT has been update but not StrataOutput
 	#endif
 		Out.StrataOutput[0] = RWStrataMaterialContainer.MaterialRenderTargets[0];
 		Out.StrataOutput[1] = RWStrataMaterialContainer.MaterialRenderTargets[1];
-		Out.StrataOutput[2] = RWStrataMaterialContainer.MaterialRenderTargets[2];
-		Out.StrataOutput[3] = StrataPackTopLayerData(TopLayerData);
+		Out.StrataOutput[2] = StrataPackTopLayerData(TopLayerData);
 	#if STRATA_OUTPUT_ROUGH_REFRACTION
 		OpaqueRoughRefractionData.OpaqueRoughRefractionEnabled = true; // As long as we output to rough refraction, we want it to be enabled for the tiling process.
 	#if STRATA_USES_CONVERSION_FROM_LEGACY
diff --git a/Engine/Shaders/Private/ComputeShaderOutputCommon.ush b/Engine/Shaders/Private/ComputeShaderOutputCommon.ush
new file mode 100644
--- /dev/null
+++ b/Engine/Shaders/Private/ComputeShaderOutputCommon.ush
@@ -0,0 +1,84 @@
+// Copyright Epic Games, Inc. All Rights Reserved.
+
+/*=============================================================================
+ComputeShaderOutputCommon.ush: To allow CS input/output passed into functions 
+through a single struct, allowing for a more readable code
+(less #ifdefs, reducing the boolean hell)
+=============================================================================*/
+
+#include "ShaderOutputCommon.ush"
+
+#if PIXELSHADEROUTPUT_MRT0
+	#if DUAL_SOURCE_COLOR_BLENDING_ENABLED && MATERIAL_WORKS_WITH_DUAL_SOURCE_COLOR_BLENDING
+		RWTexture2D<float4> OutTarget0UAV; // DUAL_SOURCE_BLENDING_SLOT(0)
+	#else
+		RWTexture2D<float4> OutTarget0UAV;
+	#endif
+#endif
+
+#if PIXELSHADEROUTPUT_MRT1
+	#if DUAL_SOURCE_COLOR_BLENDING_ENABLED && MATERIAL_WORKS_WITH_DUAL_SOURCE_COLOR_BLENDING
+		RWTexture2D<float4> OutTarget1UAV; // DUAL_SOURCE_BLENDING_SLOT(1)
+	#else
+		RWTexture2D<float4> OutTarget1UAV;
+	#endif
+#endif
+
+#if PIXELSHADEROUTPUT_MRT2
+	RWTexture2D<float4> OutTarget2UAV;
+#endif
+
+#if STRATA_OPAQUE_DEFERRED
+
+	#if STRATA_BASE_PASS_MRT_OUTPUT_COUNT != 2
+		#error Strata STRATA_BASE_PASS_MRT_OUTPUT_COUNT has been updated but not the uint MRTs
+	#endif
+	#if PIXELSHADEROUTPUT_MRT3
+		#error Strata cannot map to such a case
+	#endif
+	#if PIXELSHADEROUTPUT_MRT2
+		RWTexture2D<uint> StrataOutput3UAV;
+		RWTexture2D<uint> StrataOutput4UAV;
+		RWTexture2D<uint> StrataOutput5UAV;
+	#elif PIXELSHADEROUTPUT_MRT1
+		RWTexture2D<uint> StrataOutput2UAV;
+		RWTexture2D<uint> StrataOutput3UAV;
+		RWTexture2D<uint> StrataOutput4UAV;
+	#else
+		RWTexture2D<uint> StrataOutput1UAV;
+		RWTexture2D<uint> StrataOutput2UAV;
+		RWTexture2D<uint> StrataOutput3UAV;
+	#endif
+
+#else // STRATA_OPAQUE_DEFERRED
+
+	#if PIXELSHADEROUTPUT_MRT3
+		RWTexture2D<float4> OutTarget3UAV;
+	#endif
+
+	#if PIXELSHADEROUTPUT_MRT4
+		RWTexture2D<float4> OutTarget4UAV;
+	#endif
+
+	#if PIXELSHADEROUTPUT_MRT5
+		RWTexture2D<float4> OutTarget5UAV;
+	#endif
+
+	#if PIXELSHADEROUTPUT_MRT6
+		RWTexture2D<float4> OutTarget6UAV;
+	#endif
+
+	#if PIXELSHADEROUTPUT_MRT7
+		RWTexture2D<float4> OutTarget7UAV;
+	#endif
+
+#endif // STRATA_OPAQUE_DEFERRED
+
+// .xy = min, .zw = max
+uint4 ViewRect;
+
+[numthreads(16, 16, 1)]
+void MainCS(uint3 ThreadIndex : SV_GroupThreadID, uint3 GroupID : SV_GroupID)
+{
+	// TODO
+}
diff --git a/Engine/Shaders/Private/PixelShaderOutputCommon.ush b/Engine/Shaders/Private/PixelShaderOutputCommon.ush
--- a/Engine/Shaders/Private/PixelShaderOutputCommon.ush
+++ b/Engine/Shaders/Private/PixelShaderOutputCommon.ush
@@ -1,53 +1,12 @@
 // Copyright Epic Games, Inc. All Rights Reserved.
 
 /*=============================================================================
-PixelShaderOutputCommon.usf: To allow PS input/output passed int functions through a single struct, allowing for a more readable code (less #ifdefs, reducing the boolean hell)
+PixelShaderOutputCommon.ush: To allow PS input/output passed into functions
+through a single struct, allowing for a more readable code
+(less #ifdefs, reducing the boolean hell)
 =============================================================================*/
 
-#ifndef PIXELSHADEROUTPUT_INTERPOLANTS
-	#define PIXELSHADEROUTPUT_INTERPOLANTS 0
-#endif
-#ifndef PIXELSHADEROUTPUT_BASEPASS
-	#define PIXELSHADEROUTPUT_BASEPASS 0
-#endif
-#ifndef PIXELSHADEROUTPUT_MESHDECALPASS
-	#define PIXELSHADEROUTPUT_MESHDECALPASS 0
-#endif
-#ifndef PIXELSHADEROUTPUT_MRT0
-	#define PIXELSHADEROUTPUT_MRT0 0
-#endif
-#ifndef PIXELSHADEROUTPUT_MRT1
-	#define PIXELSHADEROUTPUT_MRT1 0
-#endif
-#ifndef PIXELSHADEROUTPUT_MRT2
-	#define PIXELSHADEROUTPUT_MRT2 0
-#endif
-#ifndef PIXELSHADEROUTPUT_MRT3
-	#define PIXELSHADEROUTPUT_MRT3 0
-#endif
-#ifndef PIXELSHADEROUTPUT_MRT4
-	#define PIXELSHADEROUTPUT_MRT4 0
-#endif
-#ifndef PIXELSHADEROUTPUT_MRT5
-	#define PIXELSHADEROUTPUT_MRT5 0
-#endif
-#ifndef PIXELSHADEROUTPUT_MRT6
-	#define PIXELSHADEROUTPUT_MRT6 0
-#endif
-#ifndef PIXELSHADEROUTPUT_MRT7
-	#define PIXELSHADEROUTPUT_MRT7 0
-#endif
-#ifndef PIXELSHADEROUTPUT_COVERAGE
-	#define PIXELSHADEROUTPUT_COVERAGE 0
-#endif
-#ifndef PIXELSHADEROUTPUT_A2C
-	#define PIXELSHADEROUTPUT_A2C 0
-#endif
-#ifndef PIXELSHADER_EARLYDEPTHSTENCIL
-#define PIXELSHADER_EARLYDEPTHSTENCIL
-#endif
-
-// ---------------------------------------------------------------------------------
+#include "ShaderOutputCommon.ush"
 
 PIXELSHADER_EARLYDEPTHSTENCIL
 void MainPS
@@ -95,7 +54,7 @@ void MainPS
 
 #if STRATA_OPAQUE_DEFERRED
 #if STRATA_BASE_PASS_MRT_OUTPUT_COUNT != 3
-#error Strata STRATA_BASE_PASS_MRT_OUTPUT_COUNT has been update but not the uint MRTs
+#error Strata STRATA_BASE_PASS_MRT_OUTPUT_COUNT has been updated but not the uint MRTs
 #endif
 #if PIXELSHADEROUTPUT_MRT4
 #error Strata cannot map to such a case
diff --git a/Engine/Shaders/Private/ShaderOutputCommon.ush b/Engine/Shaders/Private/ShaderOutputCommon.ush
new file mode 100644
--- /dev/null
+++ b/Engine/Shaders/Private/ShaderOutputCommon.ush
@@ -0,0 +1,57 @@
+// Copyright Epic Games, Inc. All Rights Reserved.
+
+#ifndef PIXELSHADEROUTPUT_INTERPOLANTS
+	#define PIXELSHADEROUTPUT_INTERPOLANTS 0
+#endif
+
+#ifndef PIXELSHADEROUTPUT_BASEPASS
+	#define PIXELSHADEROUTPUT_BASEPASS 0
+#endif
+
+#ifndef PIXELSHADEROUTPUT_MESHDECALPASS
+	#define PIXELSHADEROUTPUT_MESHDECALPASS 0
+#endif
+
+#ifndef PIXELSHADEROUTPUT_MRT0
+	#define PIXELSHADEROUTPUT_MRT0 0
+#endif
+
+#ifndef PIXELSHADEROUTPUT_MRT1
+	#define PIXELSHADEROUTPUT_MRT1 0
+#endif
+
+#ifndef PIXELSHADEROUTPUT_MRT2
+	#define PIXELSHADEROUTPUT_MRT2 0
+#endif
+
+#ifndef PIXELSHADEROUTPUT_MRT3
+	#define PIXELSHADEROUTPUT_MRT3 0
+#endif
+
+#ifndef PIXELSHADEROUTPUT_MRT4
+	#define PIXELSHADEROUTPUT_MRT4 0
+#endif
+
+#ifndef PIXELSHADEROUTPUT_MRT5
+	#define PIXELSHADEROUTPUT_MRT5 0
+#endif
+
+#ifndef PIXELSHADEROUTPUT_MRT6
+	#define PIXELSHADEROUTPUT_MRT6 0
+#endif
+
+#ifndef PIXELSHADEROUTPUT_MRT7
+	#define PIXELSHADEROUTPUT_MRT7 0
+#endif
+
+#ifndef PIXELSHADEROUTPUT_COVERAGE
+	#define PIXELSHADEROUTPUT_COVERAGE 0
+#endif
+
+#ifndef PIXELSHADEROUTPUT_A2C
+	#define PIXELSHADEROUTPUT_A2C 0
+#endif
+
+#ifndef PIXELSHADER_EARLYDEPTHSTENCIL
+#define PIXELSHADER_EARLYDEPTHSTENCIL
+#endif
diff --git a/Engine/Source/Runtime/RenderCore/Public/VertexFactory.h b/Engine/Source/Runtime/RenderCore/Public/VertexFactory.h
--- a/Engine/Source/Runtime/RenderCore/Public/VertexFactory.h
+++ b/Engine/Source/Runtime/RenderCore/Public/VertexFactory.h
@@ -125,7 +125,8 @@ enum class EVertexFactoryFlags : uint32
 	SupportsPSOPrecaching				  = 1u << 13,
 	SupportsManualVertexFetch			  = 1u << 14,
 	DoesNotSupportNullPixelShader		  = 1u << 15,
-	SupportsGPUSkinPassThrough			  = 1u << 16
+	SupportsGPUSkinPassThrough			  = 1u << 16,
+	SupportsComputeShading                = 1u << 17
 };
 ENUM_CLASS_FLAGS(EVertexFactoryFlags);
 
@@ -435,6 +436,7 @@ public:
 	bool SupportsPSOPrecaching() const                 { return HasFlags(EVertexFactoryFlags::SupportsPSOPrecaching); }
 	bool SupportsNullPixelShader() const			   { return !HasFlags(EVertexFactoryFlags::DoesNotSupportNullPixelShader); }
 	bool SupportsGPUSkinPassThrough() const			   { return HasFlags(EVertexFactoryFlags::SupportsGPUSkinPassThrough); }
+	bool SupportsComputeShading() const                { return HasFlags(EVertexFactoryFlags::SupportsComputeShading); }
 
 	bool SupportsManualVertexFetch(ERHIFeatureLevel::Type InFeatureLevel) const 
 	{
diff --git a/Engine/Source/Runtime/Renderer/Private/BasePassRendering.cpp b/Engine/Source/Runtime/Renderer/Private/BasePassRendering.cpp
--- a/Engine/Source/Runtime/Renderer/Private/BasePassRendering.cpp
+++ b/Engine/Source/Runtime/Renderer/Private/BasePassRendering.cpp
@@ -105,11 +105,17 @@ IMPLEMENT_STATIC_UNIFORM_BUFFER_STRUCT(FTranslucentBasePassUniformParameters, "T
 	typedef TBasePassPS<LightMapPolicyType, bEnableSkyLight> TBasePassPS##LightMapPolicyName##SkyLightName; \
 	IMPLEMENT_MATERIAL_SHADER_TYPE(template<>,TBasePassPS##LightMapPolicyName##SkyLightName,TEXT("/Engine/Private/BasePassPixelShader.usf"),TEXT("MainPS"),SF_Pixel);
 
-// Implement a pixel shader type for skylights and one without, and one vertex shader that will be shared between them
+#define IMPLEMENT_BASEPASS_COMPUTESHADER_TYPE(LightMapPolicyType,LightMapPolicyName,bEnableSkyLight,SkyLightName) \
+	typedef TBasePassCS<LightMapPolicyType, bEnableSkyLight> TBasePassCS##LightMapPolicyName##SkyLightName; \
+	IMPLEMENT_MATERIAL_SHADER_TYPE(template<>,TBasePassCS##LightMapPolicyName##SkyLightName,TEXT("/Engine/Private/BasePassPixelShader.usf"),TEXT("MainCS"),SF_Compute);
+
+// Implement a pixel and compute shader type for skylights and one without, and one vertex shader that will be shared between them
 #define IMPLEMENT_BASEPASS_LIGHTMAPPED_SHADER_TYPE(LightMapPolicyType,LightMapPolicyName) \
 	IMPLEMENT_BASEPASS_VERTEXSHADER_TYPE(LightMapPolicyType,LightMapPolicyName) \
 	IMPLEMENT_BASEPASS_PIXELSHADER_TYPE(LightMapPolicyType,LightMapPolicyName,true,Skylight) \
-	IMPLEMENT_BASEPASS_PIXELSHADER_TYPE(LightMapPolicyType,LightMapPolicyName,false,);
+	IMPLEMENT_BASEPASS_PIXELSHADER_TYPE(LightMapPolicyType,LightMapPolicyName,false,) \
+	IMPLEMENT_BASEPASS_COMPUTESHADER_TYPE(LightMapPolicyType,LightMapPolicyName,true,Skylight) \
+	IMPLEMENT_BASEPASS_COMPUTESHADER_TYPE(LightMapPolicyType,LightMapPolicyName,false,)
 
 // Implement shader types per lightmap policy
 // If renaming or refactoring these, remember to update FMaterialResource::GetRepresentativeInstructionCounts and FPreviewMaterial::ShouldCache().
@@ -546,6 +552,71 @@ bool GetBasePassShaders<FUniformLightMapPolicy>(
 	}
 }
 
+template <ELightMapPolicyType Policy>
+bool GetUniformBasePassShader(
+	const FMaterial& Material,
+	FVertexFactoryType* VertexFactoryType, 
+	ERHIFeatureLevel::Type FeatureLevel,
+	bool bEnableSkyLight,
+	TShaderRef<TBasePassComputeShaderPolicyParamType<FUniformLightMapPolicy>>* ComputeShader
+)
+{
+	FMaterialShaderTypes ShaderTypes;
+
+	if (ComputeShader)
+	{
+		if (bEnableSkyLight)
+		{
+			ShaderTypes.AddShaderType<TBasePassCS<TUniformLightMapPolicy<Policy>, true>>();
+		}
+		else
+		{
+			ShaderTypes.AddShaderType<TBasePassCS<TUniformLightMapPolicy<Policy>, false>>();
+		}
+	}
+
+	FMaterialShaders Shaders;
+	if (!Material.TryGetShaders(ShaderTypes, VertexFactoryType, Shaders))
+	{
+		return false;
+	}
+
+	Shaders.TryGetComputeShader(ComputeShader);
+	return true;
+}
+
+template <>
+bool GetBasePassShader<FUniformLightMapPolicy>(
+	const FMaterial& Material, 
+	FVertexFactoryType* VertexFactoryType, 
+	FUniformLightMapPolicy LightMapPolicy, 
+	ERHIFeatureLevel::Type FeatureLevel,
+	bool bEnableSkyLight,
+	TShaderRef<TBasePassComputeShaderPolicyParamType<FUniformLightMapPolicy>>* ComputeShader
+	)
+{
+	switch (LightMapPolicy.GetIndirectPolicy())
+	{
+	case LMP_PRECOMPUTED_IRRADIANCE_VOLUME_INDIRECT_LIGHTING:
+		return GetUniformBasePassShader<LMP_PRECOMPUTED_IRRADIANCE_VOLUME_INDIRECT_LIGHTING>(Material, VertexFactoryType, FeatureLevel, bEnableSkyLight, ComputeShader);
+	case LMP_CACHED_VOLUME_INDIRECT_LIGHTING:
+		return GetUniformBasePassShader<LMP_CACHED_VOLUME_INDIRECT_LIGHTING>(Material, VertexFactoryType, FeatureLevel, bEnableSkyLight, ComputeShader);
+	case LMP_CACHED_POINT_INDIRECT_LIGHTING:
+		return GetUniformBasePassShader<LMP_CACHED_POINT_INDIRECT_LIGHTING>(Material, VertexFactoryType, FeatureLevel, bEnableSkyLight, ComputeShader);
+	case LMP_LQ_LIGHTMAP:
+		return GetUniformBasePassShader<LMP_LQ_LIGHTMAP>(Material, VertexFactoryType, FeatureLevel, bEnableSkyLight, ComputeShader);
+	case LMP_HQ_LIGHTMAP:
+		return GetUniformBasePassShader<LMP_HQ_LIGHTMAP>(Material, VertexFactoryType, FeatureLevel, bEnableSkyLight, ComputeShader);
+	case LMP_DISTANCE_FIELD_SHADOWS_AND_HQ_LIGHTMAP:
+		return GetUniformBasePassShader<LMP_DISTANCE_FIELD_SHADOWS_AND_HQ_LIGHTMAP>(Material, VertexFactoryType, FeatureLevel, bEnableSkyLight, ComputeShader);
+	case LMP_NO_LIGHTMAP:
+		return GetUniformBasePassShader<LMP_NO_LIGHTMAP>(Material, VertexFactoryType, FeatureLevel, bEnableSkyLight, ComputeShader);
+	default:
+		check(false);
+		return false;
+	}
+}
+
 void SetupSharedBasePassParameters(
 	FRDGBuilder& GraphBuilder,
 	const FViewInfo& View,
diff --git a/Engine/Source/Runtime/Renderer/Private/BasePassRendering.h b/Engine/Source/Runtime/Renderer/Private/BasePassRendering.h
--- a/Engine/Source/Runtime/Renderer/Private/BasePassRendering.h
+++ b/Engine/Source/Runtime/Renderer/Private/BasePassRendering.h
@@ -297,6 +297,198 @@ public:
 	}
 };
 
+// BEGIN COMPUTE
+
+/**
+ * The base type for compute shaders that render the emissive color, and light-mapped/ambient lighting of a mesh.
+ * The base type is shared between the versions with and without sky light.
+ */
+template<typename LightMapPolicyType>
+class TBasePassComputeShaderPolicyParamType : public FMeshMaterialShader, public LightMapPolicyType::ComputeParametersType
+{
+	DECLARE_INLINE_TYPE_LAYOUT_EXPLICIT_BASES(TBasePassComputeShaderPolicyParamType, NonVirtual, FMeshMaterialShader, typename LightMapPolicyType::ComputeParametersType);
+
+public:
+	static void ModifyCompilationEnvironment(const FMaterialShaderPermutationParameters& Parameters, FShaderCompilerEnvironment& OutEnvironment)
+	{
+		FMaterialShader::ModifyCompilationEnvironment(Parameters, OutEnvironment);
+		Strata::SetBasePassRenderTargetOutputFormat(Parameters.Platform, OutEnvironment);
+		FForwardLightingParameters::ModifyCompilationEnvironment(Parameters.Platform, OutEnvironment);
+	}
+
+	static bool ValidateCompiledResult(EShaderPlatform Platform, const FShaderParameterMap& ParameterMap, TArray<FString>& OutError)
+	{
+		if (ParameterMap.ContainsParameterAllocation(FSceneTextureUniformParameters::StaticStructMetadata.GetShaderVariableName()))
+		{
+			OutError.Add(TEXT("Base pass shaders cannot read from the SceneTexturesStruct."));
+			return false;
+		}
+
+		return true;
+	}
+
+	/** Initialization constructor. */
+	TBasePassComputeShaderPolicyParamType(const FMeshMaterialShaderType::CompiledShaderInitializerType& Initializer):
+		FMeshMaterialShader(Initializer)
+	{
+		LightMapPolicyType::ComputeParametersType::Bind(Initializer.ParameterMap);
+		ReflectionCaptureBuffer.Bind(Initializer.ParameterMap, TEXT("ReflectionCapture"));
+
+		// These parameters should only be used nested in the base pass uniform buffer
+		check(!Initializer.ParameterMap.ContainsParameterAllocation(FFogUniformParameters::StaticStructMetadata.GetShaderVariableName()));
+		check(!Initializer.ParameterMap.ContainsParameterAllocation(FReflectionUniformParameters::StaticStructMetadata.GetShaderVariableName()));
+		check(!Initializer.ParameterMap.ContainsParameterAllocation(FPlanarReflectionUniformParameters::StaticStructMetadata.GetShaderVariableName()));
+	}
+	TBasePassComputeShaderPolicyParamType() {}
+
+	void GetShaderBindings(
+		const FScene* Scene,
+		ERHIFeatureLevel::Type FeatureLevel,
+		const FPrimitiveSceneProxy* PrimitiveSceneProxy,
+		const FMaterialRenderProxy& MaterialRenderProxy,
+		const FMaterial& Material,
+		const FMeshPassProcessorRenderState& DrawRenderState,
+		const TBasePassShaderElementData<LightMapPolicyType>& ShaderElementData,
+		FMeshDrawSingleShaderBindings& ShaderBindings) const;
+
+private:
+	LAYOUT_FIELD(FShaderUniformBufferParameter, ReflectionCaptureBuffer);
+};
+
+/**
+ * The base type for compute shaders that render the emissive color, and light-mapped/ambient lighting of a mesh.
+ * The base type is shared between the versions with and without sky light.
+ */
+template<typename LightMapPolicyType>
+class TBasePassComputeShaderBaseType : public TBasePassComputeShaderPolicyParamType<LightMapPolicyType>
+{
+	typedef TBasePassComputeShaderPolicyParamType<LightMapPolicyType> Super;
+	DECLARE_INLINE_TYPE_LAYOUT(TBasePassComputeShaderBaseType, NonVirtual);
+
+public:
+	static bool ShouldCompilePermutation(const FMeshMaterialShaderPermutationParameters& Parameters)
+	{
+		return LightMapPolicyType::ShouldCompilePermutation(Parameters);
+	}
+
+	static void ModifyCompilationEnvironment(const FMaterialShaderPermutationParameters& Parameters, FShaderCompilerEnvironment& OutEnvironment)
+	{
+		LightMapPolicyType::ModifyCompilationEnvironment(Parameters, OutEnvironment);
+		Super::ModifyCompilationEnvironment(Parameters, OutEnvironment);
+	}
+
+	/** Initialization constructor. */
+	TBasePassComputeShaderBaseType(const FMeshMaterialShaderType::CompiledShaderInitializerType& Initializer) : Super(Initializer) {}
+
+	TBasePassComputeShaderBaseType() {}
+};
+
+/** The concrete base pass compute shader type. */
+template<typename LightMapPolicyType, bool bEnableSkyLight>
+class TBasePassCS : public TBasePassComputeShaderBaseType<LightMapPolicyType>
+{
+	DECLARE_SHADER_TYPE(TBasePassCS,MeshMaterial);
+public:
+
+	static bool ShouldCompilePermutation(const FMeshMaterialShaderPermutationParameters& Parameters)
+	{
+		// Only compile skylight version for lit materials, and if the project allows them.
+		static const auto SupportStationarySkylight = IConsoleManager::Get().FindTConsoleVariableDataInt(TEXT("r.SupportStationarySkylight"));
+		static const auto SupportAllShaderPermutations = IConsoleManager::Get().FindTConsoleVariableDataInt(TEXT("r.SupportAllShaderPermutations"));
+
+		const bool IsSingleLayerWater = Parameters.MaterialParameters.ShadingModels.HasShadingModel(MSM_SingleLayerWater);
+
+		const bool bTranslucent = IsTranslucentBlendMode(Parameters.MaterialParameters.BlendMode);
+		const bool bForceAllPermutations = SupportAllShaderPermutations && SupportAllShaderPermutations->GetValueOnAnyThread() != 0;
+		const bool bProjectSupportsStationarySkylight = !SupportStationarySkylight || SupportStationarySkylight->GetValueOnAnyThread() != 0 || bForceAllPermutations;
+
+		const bool bCacheShaders = !bEnableSkyLight
+			//translucent materials need to compile skylight support to support MOVABLE skylights also.
+			|| bTranslucent
+			// Some lightmap policies (eg Simple Forward) always require skylight support
+			|| IsSingleLayerWater
+			|| ((bProjectSupportsStationarySkylight || IsForwardShadingEnabled(Parameters.Platform)) && Parameters.MaterialParameters.ShadingModels.IsLit());
+		return bCacheShaders
+			&& (IsFeatureLevelSupported(Parameters.Platform, ERHIFeatureLevel::SM5))
+			&& TBasePassComputeShaderBaseType<LightMapPolicyType>::ShouldCompilePermutation(Parameters);
+	}
+
+	static void ModifyCompilationEnvironment(const FMaterialShaderPermutationParameters& Parameters, FShaderCompilerEnvironment& OutEnvironment)
+	{
+		OutEnvironment.SetDefine(TEXT("SCENE_TEXTURES_DISABLED"), Parameters.MaterialParameters.MaterialDomain != MD_Surface);
+		OutEnvironment.SetDefine(TEXT("ENABLE_DBUFFER_TEXTURES"), Parameters.MaterialParameters.MaterialDomain == MD_Surface);
+		OutEnvironment.SetDefine(TEXT("COMPILE_BASEPASS_PIXEL_VOLUMETRIC_FOGGING"), DoesPlatformSupportVolumetricFog(Parameters.Platform));
+		OutEnvironment.SetDefine(TEXT("ENABLE_SKY_LIGHT"), bEnableSkyLight);
+		OutEnvironment.SetDefine(TEXT("PLATFORM_FORCE_SIMPLE_SKY_DIFFUSE"), ForceSimpleSkyDiffuse(Parameters.Platform));
+		OutEnvironment.SetDefine(TEXT("COMPUTE_SHADED"), 1);
+
+		const bool bTranslucent = IsTranslucentBlendMode(Parameters.MaterialParameters.BlendMode);
+		const bool bIsSingleLayerWater = Parameters.MaterialParameters.ShadingModels.HasShadingModel(MSM_SingleLayerWater);
+		const bool bSupportVirtualShadowMap = IsFeatureLevelSupported(Parameters.Platform, ERHIFeatureLevel::SM5);
+		if (bSupportVirtualShadowMap && (bTranslucent || bIsSingleLayerWater))
+		{
+			FVirtualShadowMapArray::SetShaderDefines(OutEnvironment);
+			OutEnvironment.SetDefine(TEXT("VIRTUAL_SHADOW_MAP"), 1);
+		}
+
+		// This define simply lets the compilation environment know that we are using BasePassPixelShader.usf, so that we can check for more
+		// complicated defines later in the compilation pipe.
+		OutEnvironment.SetDefine(TEXT("IS_BASE_PASS"), 1);
+		OutEnvironment.SetDefine(TEXT("IS_MOBILE_BASE_PASS"), 0);
+
+		OutEnvironment.SetDefine(TEXT("STRATA_INLINE_SHADING"), 1);
+
+		const bool IsSingleLayerWater = Parameters.MaterialParameters.ShadingModels.HasShadingModel(MSM_SingleLayerWater);
+		if (IsSingleLayerWater && IsWaterDistanceFieldShadowEnabled(Parameters.Platform))
+		{
+			// See FShaderCompileUtilities::FetchGBufferParamsRuntime for the details
+			const bool bOutputVelocity = FVelocityRendering::BasePassCanOutputVelocity(Parameters.Platform);
+			const bool bHasTangent = false;
+			static const auto CVar = IConsoleManager::Get().FindTConsoleVariableDataInt(TEXT("r.AllowStaticLighting"));
+			bool bHasPrecShadowFactor = (CVar ? (CVar->GetValueOnAnyThread() != 0) : 1);
+
+			uint32 TargetSeparatedMainDirLight = 5;
+			if (bOutputVelocity == false && bHasTangent == false)
+			{
+				TargetSeparatedMainDirLight = 5;
+				if (bHasPrecShadowFactor)
+				{
+					TargetSeparatedMainDirLight = 6;
+				}
+			}
+			else if (bOutputVelocity)
+			{
+				TargetSeparatedMainDirLight = 6;
+				if (bHasPrecShadowFactor)
+				{
+					TargetSeparatedMainDirLight = 7;
+				}
+			}
+			else if (bHasTangent)
+			{
+				TargetSeparatedMainDirLight = 6;
+				if (bHasPrecShadowFactor)
+				{
+					TargetSeparatedMainDirLight = 7;
+				}
+			}
+			OutEnvironment.SetRenderTargetOutputFormat(TargetSeparatedMainDirLight, PF_FloatR11G11B10);
+		}
+
+		TBasePassComputeShaderBaseType<LightMapPolicyType>::ModifyCompilationEnvironment(Parameters, OutEnvironment);
+	}
+	
+	/** Initialization constructor. */
+	TBasePassCS(const ShaderMetaType::CompiledShaderInitializerType& Initializer):
+		TBasePassComputeShaderBaseType<LightMapPolicyType>(Initializer)
+	{}
+
+	/** Default constructor. */
+	TBasePassCS() {}
+};
+
+// END COMPUTE
+
 /**
  * The base type for pixel shaders that render the emissive color, and light-mapped/ambient lighting of a mesh.
  * The base type is shared between the versions with and without sky light.
@@ -307,8 +499,6 @@ class TBasePassPixelShaderPolicyParamType : public FMeshMaterialShader, public L
 	DECLARE_INLINE_TYPE_LAYOUT_EXPLICIT_BASES(TBasePassPixelShaderPolicyParamType, NonVirtual, FMeshMaterialShader, typename LightMapPolicyType::PixelParametersType);
 public:
 
-	// static bool ShouldCompilePermutation(const FMeshMaterialShaderPermutationParameters& Parameters)
-
 	static void ModifyCompilationEnvironment(const FMaterialShaderPermutationParameters& Parameters, FShaderCompilerEnvironment& OutEnvironment)
 	{
 		FMeshMaterialShader::ModifyCompilationEnvironment(Parameters, OutEnvironment);
@@ -529,6 +719,54 @@ public:
 	F128BitRTBasePassPS() {}
 };
 
+/**
+ * Get shader templates allowing to redirect between compatible shaders.
+ */
+
+template <typename LightMapPolicyType>
+bool GetBasePassShader(
+	const FMaterial& Material,
+	FVertexFactoryType* VertexFactoryType,
+	LightMapPolicyType LightMapPolicy,
+	ERHIFeatureLevel::Type FeatureLevel,
+	bool bEnableSkyLight,
+	TShaderRef<TBasePassComputeShaderPolicyParamType<LightMapPolicyType>>* ComputeShader
+)
+{
+	FMaterialShaderTypes ShaderTypes;
+
+	if (ComputeShader)
+	{
+		if (bEnableSkyLight)
+		{
+			ShaderTypes.AddShaderType<TBasePassCS<LightMapPolicyType, true>>();
+		}
+		else
+		{
+			ShaderTypes.AddShaderType<TBasePassCS<LightMapPolicyType, false>>();
+		}
+	}
+
+	FMaterialShaders Shaders;
+	if (!Material.TryGetShaders(ShaderTypes, VertexFactoryType, Shaders))
+	{
+		return false;
+	}
+
+	Shaders.TryGetComputeShader(ComputeShader);
+	return true;
+}
+
+template <>
+bool GetBasePassShader<FUniformLightMapPolicy>(
+	const FMaterial& Material,
+	FVertexFactoryType* VertexFactoryType,
+	FUniformLightMapPolicy LightMapPolicy,
+	ERHIFeatureLevel::Type FeatureLevel,
+	bool bEnableSkyLight,
+	TShaderRef<TBasePassComputeShaderPolicyParamType<FUniformLightMapPolicy>>* ComputeShader
+);
+
 /**
  * Get shader templates allowing to redirect between compatible shaders.
  */
diff --git a/Engine/Source/Runtime/Renderer/Private/BasePassRendering.inl b/Engine/Source/Runtime/Renderer/Private/BasePassRendering.inl
--- a/Engine/Source/Runtime/Renderer/Private/BasePassRendering.inl
+++ b/Engine/Source/Runtime/Renderer/Private/BasePassRendering.inl
@@ -75,4 +75,24 @@ void TBasePassPixelShaderPolicyParamType<LightMapPolicyType>::GetShaderBindings(
 		ShaderElementData.LightMapPolicyElementData,
 		this,
 		ShaderBindings);
+}
+
+template<typename LightMapPolicyType>
+void TBasePassComputeShaderPolicyParamType<LightMapPolicyType>::GetShaderBindings(
+	const FScene* Scene,
+	ERHIFeatureLevel::Type FeatureLevel,
+	const FPrimitiveSceneProxy* PrimitiveSceneProxy,
+	const FMaterialRenderProxy& MaterialRenderProxy,
+	const FMaterial& Material,
+	const FMeshPassProcessorRenderState& DrawRenderState,
+	const TBasePassShaderElementData<LightMapPolicyType>& ShaderElementData,
+	FMeshDrawSingleShaderBindings& ShaderBindings) const
+{
+	FMeshMaterialShader::GetShaderBindings(Scene, FeatureLevel, PrimitiveSceneProxy, MaterialRenderProxy, Material, DrawRenderState, ShaderElementData, ShaderBindings);
+
+	LightMapPolicyType::GetComputeShaderBindings(
+		PrimitiveSceneProxy,
+		ShaderElementData.LightMapPolicyElementData,
+		this,
+		ShaderBindings);
 }
\ No newline at end of file
diff --git a/Engine/Source/Runtime/Renderer/Private/LightMapRendering.cpp b/Engine/Source/Runtime/Renderer/Private/LightMapRendering.cpp
--- a/Engine/Source/Runtime/Renderer/Private/LightMapRendering.cpp
+++ b/Engine/Source/Runtime/Renderer/Private/LightMapRendering.cpp
@@ -88,6 +88,25 @@ void FSelfShadowedCachedPointIndirectLightingPolicy::GetPixelShaderBindings(
 	FSelfShadowedTranslucencyPolicy::GetPixelShaderBindings(PrimitiveSceneProxy, ShaderElementData.SelfShadowTranslucencyUniformBuffer, PixelShaderParameters, ShaderBindings);
 }
 
+void FSelfShadowedCachedPointIndirectLightingPolicy::GetComputeShaderBindings(
+	const FPrimitiveSceneProxy* PrimitiveSceneProxy,
+	const ElementDataType& ShaderElementData,
+	const ComputeParametersType* ComputeShaderParameters,
+	FMeshDrawSingleShaderBindings& ShaderBindings)
+{
+	FRHIUniformBuffer* PrecomputedLightingBuffer = nullptr;
+	FRHIUniformBuffer* LightmapResourceClusterBuffer = nullptr;
+	FRHIUniformBuffer* IndirectLightingCacheBuffer = nullptr;
+
+	SetupLCIUniformBuffers(PrimitiveSceneProxy, ShaderElementData.LCI, PrecomputedLightingBuffer, LightmapResourceClusterBuffer, IndirectLightingCacheBuffer);
+
+	ShaderBindings.Add(ComputeShaderParameters->PrecomputedLightingBufferParameter, PrecomputedLightingBuffer);
+	ShaderBindings.Add(ComputeShaderParameters->IndirectLightingCacheParameter, IndirectLightingCacheBuffer);
+	ShaderBindings.Add(ComputeShaderParameters->LightmapResourceCluster, LightmapResourceClusterBuffer);
+
+	FSelfShadowedTranslucencyPolicy::GetComputeShaderBindings(PrimitiveSceneProxy, ShaderElementData.SelfShadowTranslucencyUniformBuffer, ComputeShaderParameters, ShaderBindings);
+}
+
 void FSelfShadowedVolumetricLightmapPolicy::GetPixelShaderBindings(
 	const FPrimitiveSceneProxy* PrimitiveSceneProxy,
 	const ElementDataType& ShaderElementData,
@@ -107,6 +126,25 @@ void FSelfShadowedVolumetricLightmapPolicy::GetPixelShaderBindings(
 	FSelfShadowedTranslucencyPolicy::GetPixelShaderBindings(PrimitiveSceneProxy, ShaderElementData.SelfShadowTranslucencyUniformBuffer, PixelShaderParameters, ShaderBindings);
 }
 
+void FSelfShadowedVolumetricLightmapPolicy::GetComputeShaderBindings(
+	const FPrimitiveSceneProxy* PrimitiveSceneProxy,
+	const ElementDataType& ShaderElementData,
+	const ComputeParametersType* ComputeShaderParameters,
+	FMeshDrawSingleShaderBindings& ShaderBindings)
+{
+	FRHIUniformBuffer* PrecomputedLightingBuffer = nullptr;
+	FRHIUniformBuffer* LightmapResourceClusterBuffer = nullptr;
+	FRHIUniformBuffer* IndirectLightingCacheBuffer = nullptr;
+
+	SetupLCIUniformBuffers(PrimitiveSceneProxy, ShaderElementData.LCI, PrecomputedLightingBuffer, LightmapResourceClusterBuffer, IndirectLightingCacheBuffer);
+
+	ShaderBindings.Add(ComputeShaderParameters->PrecomputedLightingBufferParameter, PrecomputedLightingBuffer);
+	ShaderBindings.Add(ComputeShaderParameters->IndirectLightingCacheParameter, IndirectLightingCacheBuffer);
+	ShaderBindings.Add(ComputeShaderParameters->LightmapResourceCluster, LightmapResourceClusterBuffer);
+
+	FSelfShadowedTranslucencyPolicy::GetComputeShaderBindings(PrimitiveSceneProxy, ShaderElementData.SelfShadowTranslucencyUniformBuffer, ComputeShaderParameters, ShaderBindings);
+}
+
 void FUniformLightMapPolicy::GetVertexShaderBindings(
 	const FPrimitiveSceneProxy* PrimitiveSceneProxy,
 	const ElementDataType& ShaderElementData,
@@ -161,6 +199,23 @@ void FUniformLightMapPolicy::GetRayHitGroupShaderBindings(
 }
 #endif
 
+void FUniformLightMapPolicy::GetComputeShaderBindings(
+	const FPrimitiveSceneProxy* PrimitiveSceneProxy,
+	const ElementDataType& ShaderElementData,
+	const ComputeParametersType* ComputeShaderParameters,
+	FMeshDrawSingleShaderBindings& ShaderBindings)
+{
+	FRHIUniformBuffer* PrecomputedLightingBuffer = nullptr;
+	FRHIUniformBuffer* LightmapResourceClusterBuffer = nullptr;
+	FRHIUniformBuffer* IndirectLightingCacheBuffer = nullptr;
+
+	SetupLCIUniformBuffers(PrimitiveSceneProxy, ShaderElementData, PrecomputedLightingBuffer, LightmapResourceClusterBuffer, IndirectLightingCacheBuffer);
+
+	ShaderBindings.Add(ComputeShaderParameters->PrecomputedLightingBufferParameter, PrecomputedLightingBuffer);
+	ShaderBindings.Add(ComputeShaderParameters->IndirectLightingCacheParameter, IndirectLightingCacheBuffer);
+	ShaderBindings.Add(ComputeShaderParameters->LightmapResourceCluster, LightmapResourceClusterBuffer);
+}
+
 void InterpolateVolumetricLightmap(
 	FVector LookupPosition,
 	const FVolumetricLightmapSceneData& VolumetricLightmapSceneData,
diff --git a/Engine/Source/Runtime/Renderer/Private/LightMapRendering.h b/Engine/Source/Runtime/Renderer/Private/LightMapRendering.h
--- a/Engine/Source/Runtime/Renderer/Private/LightMapRendering.h
+++ b/Engine/Source/Runtime/Renderer/Private/LightMapRendering.h
@@ -199,6 +199,23 @@ public:
 		LAYOUT_FIELD(FShaderUniformBufferParameter, TranslucentSelfShadowBufferParameter);
 	};
 
+	class ComputeParametersType
+	{
+		DECLARE_INLINE_TYPE_LAYOUT(ComputeParametersType, NonVirtual);
+	public:
+		void Bind(const FShaderParameterMap& ParameterMap)
+		{
+			TranslucentSelfShadowBufferParameter.Bind(ParameterMap, TEXT("TranslucentSelfShadow"));
+		}
+
+		void Serialize(FArchive& Ar)
+		{
+			Ar << TranslucentSelfShadowBufferParameter;
+		}
+
+		LAYOUT_FIELD(FShaderUniformBufferParameter, TranslucentSelfShadowBufferParameter);
+	};
+
 	static bool ShouldCompilePermutation(const FMeshMaterialShaderPermutationParameters& Parameters)
 	{
 		return Parameters.MaterialParameters.ShadingModels.IsLit() &&
@@ -229,6 +246,15 @@ public:
 		ShaderBindings.Add(PixelShaderParameters->TranslucentSelfShadowBufferParameter, ShaderElementData);
 	}
 
+	static void GetComputeShaderBindings(
+		const FPrimitiveSceneProxy* PrimitiveSceneProxy,
+		const ElementDataType& ShaderElementData,
+		const ComputeParametersType* ComputeShaderParameters,
+		FMeshDrawSingleShaderBindings& ShaderBindings)
+	{
+		ShaderBindings.Add(ComputeShaderParameters->TranslucentSelfShadowBufferParameter, ShaderElementData);
+	}
+
 	friend bool operator==(const FSelfShadowedTranslucencyPolicy A,const FSelfShadowedTranslucencyPolicy B)
 	{
 		return true;
@@ -545,6 +571,7 @@ public:
 #if RHI_RAYTRACING
 	typedef FUniformLightMapPolicyShaderParametersType RayHitGroupParametersType;
 #endif
+	typedef FUniformLightMapPolicyShaderParametersType ComputeParametersType;
 
 	static bool ShouldCompilePermutation(const FMeshMaterialShaderPermutationParameters& Parameters)
 	{
@@ -577,6 +604,12 @@ public:
 	) const;
 #endif // RHI_RAYTRACING
 
+	static void GetComputeShaderBindings(
+		const FPrimitiveSceneProxy* PrimitiveSceneProxy,
+		const ElementDataType& ShaderElementData,
+		const ComputeParametersType* PixelShaderParameters,
+		FMeshDrawSingleShaderBindings& ShaderBindings);
+
 	friend bool operator==(const FUniformLightMapPolicy A,const FUniformLightMapPolicy B)
 	{
 		return A.IndirectPolicy == B.IndirectPolicy;
@@ -744,6 +777,23 @@ public:
 		}
 	};
 
+	class ComputeParametersType : public FUniformLightMapPolicyShaderParametersType, public FSelfShadowedTranslucencyPolicy::ComputeParametersType
+	{
+		DECLARE_INLINE_TYPE_LAYOUT_EXPLICIT_BASES(ComputeParametersType, NonVirtual, FUniformLightMapPolicyShaderParametersType, FSelfShadowedTranslucencyPolicy::ComputeParametersType);
+	public:
+		void Bind(const FShaderParameterMap& ParameterMap)
+		{
+			FUniformLightMapPolicyShaderParametersType::Bind(ParameterMap);
+			FSelfShadowedTranslucencyPolicy::ComputeParametersType::Bind(ParameterMap);
+		}
+
+		void Serialize(FArchive& Ar)
+		{
+			FUniformLightMapPolicyShaderParametersType::Serialize(Ar);
+			FSelfShadowedTranslucencyPolicy::ComputeParametersType::Serialize(Ar);
+		}
+	};
+
 	static bool ShouldCompilePermutation(const FMeshMaterialShaderPermutationParameters& Parameters)
 	{
 		static IConsoleVariable* AllowStaticLightingVar = IConsoleManager::Get().FindConsoleVariable(TEXT("r.AllowStaticLighting"));
@@ -770,6 +820,12 @@ public:
 		const ElementDataType& ShaderElementData,
 		const PixelParametersType* PixelShaderParameters,
 		FMeshDrawSingleShaderBindings& ShaderBindings);
+
+	static void GetComputeShaderBindings(
+		const FPrimitiveSceneProxy* PrimitiveSceneProxy,
+		const ElementDataType& ShaderElementData,
+		const ComputeParametersType* ComputeShaderParameters,
+		FMeshDrawSingleShaderBindings& ShaderBindings);
 };
 
 class FSelfShadowedVolumetricLightmapPolicy : public FSelfShadowedTranslucencyPolicy
@@ -794,6 +850,23 @@ public:
 		}
 	};
 
+	class ComputeParametersType : public FUniformLightMapPolicyShaderParametersType, public FSelfShadowedTranslucencyPolicy::ComputeParametersType
+	{
+		DECLARE_INLINE_TYPE_LAYOUT_EXPLICIT_BASES(ComputeParametersType, NonVirtual, FUniformLightMapPolicyShaderParametersType, FSelfShadowedTranslucencyPolicy::ComputeParametersType);
+	public:
+		void Bind(const FShaderParameterMap& ParameterMap)
+		{
+			FUniformLightMapPolicyShaderParametersType::Bind(ParameterMap);
+			FSelfShadowedTranslucencyPolicy::ComputeParametersType::Bind(ParameterMap);
+		}
+
+		void Serialize(FArchive& Ar)
+		{
+			FUniformLightMapPolicyShaderParametersType::Serialize(Ar);
+			FSelfShadowedTranslucencyPolicy::ComputeParametersType::Serialize(Ar);
+		}
+	};
+
 	static bool ShouldCompilePermutation(const FMeshMaterialShaderPermutationParameters& Parameters)
 	{
 		static IConsoleVariable* AllowStaticLightingVar = IConsoleManager::Get().FindConsoleVariable(TEXT("r.AllowStaticLighting"));
@@ -824,4 +897,10 @@ public:
 		const ElementDataType& ShaderElementData,
 		const PixelParametersType* PixelShaderParameters,
 		FMeshDrawSingleShaderBindings& ShaderBindings);
+
+	static void GetComputeShaderBindings(
+		const FPrimitiveSceneProxy* PrimitiveSceneProxy,
+		const ElementDataType& ShaderElementData,
+		const ComputeParametersType* ComputeShaderParameters,
+		FMeshDrawSingleShaderBindings& ShaderBindings);
 };
\ No newline at end of file
diff --git a/Engine/Source/Runtime/Renderer/Private/ShaderBaseClasses.cpp b/Engine/Source/Runtime/Renderer/Private/ShaderBaseClasses.cpp
--- a/Engine/Source/Runtime/Renderer/Private/ShaderBaseClasses.cpp
+++ b/Engine/Source/Runtime/Renderer/Private/ShaderBaseClasses.cpp
@@ -27,7 +27,7 @@ FAutoConsoleVariableRef FMaterialShader::CVarAllowCachedUniformExpressions(
 	TEXT("Allow uniform expressions to be cached."),
 	ECVF_RenderThreadSafe);
 
-void FMeshMaterialShaderElementData::InitializeMeshMaterialData(const FSceneView* SceneView, const FPrimitiveSceneProxy* RESTRICT PrimitiveSceneProxy, const FMeshBatch& RESTRICT MeshBatch, int32 StaticMeshId, bool bAllowStencilDither)
+void FMeshMaterialShaderElementData::InitializeMeshMaterialData(const FSceneView* SceneView, const FPrimitiveSceneProxy* RESTRICT PrimitiveSceneProxy, int32 StaticMeshId, bool bDitheredLODTransition, bool bAllowStencilDither)
 {
 	FadeUniformBuffer = GDistanceCullFadedInUniformBuffer.GetUniformBufferRHI();
 	DitherUniformBuffer = GDitherFadedInUniformBuffer.GetUniformBufferRHI();
@@ -37,7 +37,7 @@ void FMeshMaterialShaderElementData::InitializeMeshMaterialData(const FSceneView
 		checkSlow(SceneView->bIsViewInfo);
 		const FViewInfo* ViewInfo = (FViewInfo*)SceneView;
 
-		if (StaticMeshId >= 0 && MeshBatch.bDitheredLODTransition && !(bAllowStencilDither && ViewInfo->bAllowStencilDither))
+		if (StaticMeshId >= 0 && bDitheredLODTransition && !(bAllowStencilDither && ViewInfo->bAllowStencilDither))
 		{
 			if (ViewInfo->StaticMeshFadeOutDitheredLODMap[StaticMeshId])
 			{
@@ -61,6 +61,11 @@ void FMeshMaterialShaderElementData::InitializeMeshMaterialData(const FSceneView
 	}
 }
 
+void FMeshMaterialShaderElementData::InitializeMeshMaterialData(const FSceneView* SceneView, const FPrimitiveSceneProxy* RESTRICT PrimitiveSceneProxy, const FMeshBatch& RESTRICT MeshBatch, int32 StaticMeshId, bool bAllowStencilDither)
+{
+	InitializeMeshMaterialData(SceneView, PrimitiveSceneProxy, StaticMeshId, MeshBatch.bDitheredLODTransition, bAllowStencilDither);
+}
+
 FName FMaterialShader::UniformBufferLayoutName(TEXT("Material"));
 
 FMaterialShader::FMaterialShader(const FMaterialShaderType::CompiledShaderInitializerType& Initializer)
diff --git a/Engine/Source/Runtime/Renderer/Public/MeshMaterialShader.h b/Engine/Source/Runtime/Renderer/Public/MeshMaterialShader.h
--- a/Engine/Source/Runtime/Renderer/Public/MeshMaterialShader.h
+++ b/Engine/Source/Runtime/Renderer/Public/MeshMaterialShader.h
@@ -27,6 +27,7 @@ public:
 	FRHIUniformBuffer* DitherUniformBuffer = nullptr;
 
 	RENDERER_API void InitializeMeshMaterialData(const FSceneView* SceneView, const FPrimitiveSceneProxy* RESTRICT PrimitiveSceneProxy, const FMeshBatch& RESTRICT MeshBatch, int32 StaticMeshId, bool bAllowStencilDither);
+	RENDERER_API void InitializeMeshMaterialData(const FSceneView* SceneView, const FPrimitiveSceneProxy* RESTRICT PrimitiveSceneProxy, int32 StaticMeshId, bool bDitheredLODTransition, bool bAllowStencilDither);
 };
 
 struct FMeshMaterialShaderPermutationParameters : public FMaterialShaderPermutationParameters
