Fix compile error due to merge

diff --git a/Engine/Source/Runtime/Renderer/Private/BasePassRendering.h b/Engine/Source/Runtime/Renderer/Private/BasePassRendering.h
--- a/Engine/Source/Runtime/Renderer/Private/BasePassRendering.h
+++ b/Engine/Source/Runtime/Renderer/Private/BasePassRendering.h
@@ -312,7 +312,7 @@ public:
 	static void ModifyCompilationEnvironment(const FMaterialShaderPermutationParameters& Parameters, FShaderCompilerEnvironment& OutEnvironment)
 	{
 		FMaterialShader::ModifyCompilationEnvironment(Parameters, OutEnvironment);
-		Strata::SetBasePassRenderTargetOutputFormat(Parameters.Platform, OutEnvironment);
+		Strata::SetBasePassRenderTargetOutputFormat(Parameters.Platform, Parameters.MaterialParameters, OutEnvironment);
 		FForwardLightingParameters::ModifyCompilationEnvironment(Parameters.Platform, OutEnvironment);
 	}
 
