Fix function name being reused breaking StressTestUnity
#horde 220119

diff --git a/Engine/Source/Runtime/Renderer/Private/SingleLayerWaterRendering.cpp b/Engine/Source/Runtime/Renderer/Private/SingleLayerWaterRendering.cpp
--- a/Engine/Source/Runtime/Renderer/Private/SingleLayerWaterRendering.cpp
+++ b/Engine/Source/Runtime/Renderer/Private/SingleLayerWaterRendering.cpp
@@ -335,7 +335,7 @@ BEGIN_SHADER_PARAMETER_STRUCT(FSingleLayerWaterDepthPassParameters, )
 	RENDER_TARGET_BINDING_SLOTS()
 END_SHADER_PARAMETER_STRUCT()
 
-static FSingleLayerWaterDepthPassParameters* GetDepthPassParameters(FRDGBuilder& GraphBuilder, const FViewInfo& View, FRDGTextureRef DepthTexture)
+static FSingleLayerWaterDepthPassParameters* GetSingleLayerWaterDepthPassParameters(FRDGBuilder& GraphBuilder, const FViewInfo& View, FRDGTextureRef DepthTexture)
 {
 	FSingleLayerWaterDepthPassParameters* PassParameters = GraphBuilder.AllocParameters<FSingleLayerWaterDepthPassParameters>();
 	PassParameters->View = View.GetShaderParameters();
@@ -427,7 +427,7 @@ void FDeferredShadingSceneRenderer::RenderSingleLayerWaterDepthPrepass(FRDGBuild
 		RDG_EVENT_SCOPE_CONDITIONAL(GraphBuilder, Views.Num() > 1, "View%d", ViewIndex);
 		View.BeginRenderView();
 
-		FSingleLayerWaterDepthPassParameters* PassParameters = GetDepthPassParameters(GraphBuilder, View, OutDepthPrepassTexture.Target);
+		FSingleLayerWaterDepthPassParameters* PassParameters = GetSingleLayerWaterDepthPassParameters(GraphBuilder, View, OutDepthPrepassTexture.Target);
 
 		View.ParallelMeshDrawCommandPasses[EMeshPass::SingleLayerWaterDepthPrepass].BuildRenderingCommands(GraphBuilder, Scene->GPUScene, PassParameters->InstanceCullingDrawParams);
 
