PSOPrecache for FMeshDecalMeshProcessor

#rnx

diff --git a/Engine/Source/Runtime/Renderer/Private/CompositionLighting/PostProcessAmbientOcclusion.cpp b/Engine/Source/Runtime/Renderer/Private/CompositionLighting/PostProcessAmbientOcclusion.cpp
--- a/Engine/Source/Runtime/Renderer/Private/CompositionLighting/PostProcessAmbientOcclusion.cpp
+++ b/Engine/Source/Runtime/Renderer/Private/CompositionLighting/PostProcessAmbientOcclusion.cpp
@@ -307,10 +307,14 @@ EGTAOType FSSAOHelper::GetGTAOPassType(const FViewInfo& View, uint32 Levels)
 	return EGTAOType::EOff;
 }
 
-FRDGTextureRef CreateScreenSpaceAOTexture(FRDGBuilder& GraphBuilder, FIntPoint Extent)
+FRDGTextureDesc GetScreenSpaceAOTextureDesc(FIntPoint Extent)
 {
-	const FRDGTextureDesc Desc(FRDGTextureDesc::Create2D(Extent, PF_G8, FClearValueBinding::White, TexCreate_UAV | TexCreate_RenderTargetable | TexCreate_ShaderResource | GFastVRamConfig.ScreenSpaceAO));
-	return GraphBuilder.CreateTexture(Desc, TEXT("ScreenSpaceAO"));
+	return FRDGTextureDesc(FRDGTextureDesc::Create2D(Extent, PF_G8, FClearValueBinding::White, TexCreate_UAV | TexCreate_RenderTargetable | TexCreate_ShaderResource | GFastVRamConfig.ScreenSpaceAO));
+}
+
+FRDGTextureRef CreateScreenSpaceAOTexture(FRDGBuilder& GraphBuilder, FIntPoint Extent)
+{	
+	return GraphBuilder.CreateTexture(GetScreenSpaceAOTextureDesc(Extent), TEXT("ScreenSpaceAO"));
 }
 
 FRDGTextureRef GetScreenSpaceAOFallback(const FRDGSystemTextures& SystemTextures)
diff --git a/Engine/Source/Runtime/Renderer/Private/CompositionLighting/PostProcessAmbientOcclusion.h b/Engine/Source/Runtime/Renderer/Private/CompositionLighting/PostProcessAmbientOcclusion.h
--- a/Engine/Source/Runtime/Renderer/Private/CompositionLighting/PostProcessAmbientOcclusion.h
+++ b/Engine/Source/Runtime/Renderer/Private/CompositionLighting/PostProcessAmbientOcclusion.h
@@ -52,6 +52,8 @@ enum EGTAOPass
 	EGTAOPass_Upsample					= 0x20,
 };
 
+FRDGTextureDesc GetScreenSpaceAOTextureDesc(FIntPoint Extent);
+
 FRDGTextureRef CreateScreenSpaceAOTexture(FRDGBuilder& GraphBuilder, FIntPoint Extent);
 
 FRDGTextureRef GetScreenSpaceAOFallback(const FRDGSystemTextures& SystemTextures);
diff --git a/Engine/Source/Runtime/Renderer/Private/CompositionLighting/PostProcessDeferredDecals.cpp b/Engine/Source/Runtime/Renderer/Private/CompositionLighting/PostProcessDeferredDecals.cpp
--- a/Engine/Source/Runtime/Renderer/Private/CompositionLighting/PostProcessDeferredDecals.cpp
+++ b/Engine/Source/Runtime/Renderer/Private/CompositionLighting/PostProcessDeferredDecals.cpp
@@ -1,6 +1,7 @@
 // Copyright Epic Games, Inc. All Rights Reserved.
 
 #include "CompositionLighting/PostProcessDeferredDecals.h"
+#include "CompositionLighting/PostProcessAmbientOcclusion.h"
 
 #include "ClearQuad.h"
 #include "DBufferTextures.h"
@@ -80,6 +81,58 @@ FDeferredDecalPassTextures GetDeferredDecalPassTextures(
 	return PassTextures;
 }
 
+void GetDeferredDecalRenderTargetsInfo(
+	const FSceneTexturesConfig& Config,
+	EShaderPlatform ShaderPlatform,
+	EDecalRenderTargetMode RenderTargetMode,
+	FGraphicsPipelineRenderTargetsInfo& RenderTargetsInfo)
+{
+	switch (RenderTargetMode)
+	{
+	case EDecalRenderTargetMode::SceneColorAndGBuffer:
+		AddRenderTargetInfo(Config.ColorFormat, Config.ColorCreateFlags, RenderTargetsInfo);
+		AddRenderTargetInfo(Config.GBufferA.Format, Config.GBufferA.Flags, RenderTargetsInfo);
+		AddRenderTargetInfo(Config.GBufferB.Format, Config.GBufferB.Flags, RenderTargetsInfo);
+		AddRenderTargetInfo(Config.GBufferC.Format, Config.GBufferC.Flags, RenderTargetsInfo);
+		break;
+	case EDecalRenderTargetMode::SceneColorAndGBufferNoNormal:
+		AddRenderTargetInfo(Config.ColorFormat, Config.ColorCreateFlags, RenderTargetsInfo);
+		AddRenderTargetInfo(Config.GBufferB.Format, Config.GBufferB.Flags, RenderTargetsInfo);
+		AddRenderTargetInfo(Config.GBufferC.Format, Config.GBufferC.Flags, RenderTargetsInfo);
+		break;
+	case EDecalRenderTargetMode::SceneColor:
+		AddRenderTargetInfo(Config.ColorFormat, Config.ColorCreateFlags, RenderTargetsInfo);
+		break;
+
+	case EDecalRenderTargetMode::DBuffer:
+	{
+		const FDBufferTexturesDesc DBufferTexturesDesc = GetDBufferTexturesDesc(Config.Extent, ShaderPlatform);
+
+		AddRenderTargetInfo(DBufferTexturesDesc.DBufferADesc.Format, DBufferTexturesDesc.DBufferADesc.Flags, RenderTargetsInfo);
+		AddRenderTargetInfo(DBufferTexturesDesc.DBufferBDesc.Format, DBufferTexturesDesc.DBufferBDesc.Flags, RenderTargetsInfo);
+		AddRenderTargetInfo(DBufferTexturesDesc.DBufferCDesc.Format, DBufferTexturesDesc.DBufferCDesc.Flags, RenderTargetsInfo);
+
+		if (DBufferTexturesDesc.DBufferMaskDesc.Format != PF_Unknown)
+		{
+			AddRenderTargetInfo(DBufferTexturesDesc.DBufferMaskDesc.Format, DBufferTexturesDesc.DBufferMaskDesc.Flags, RenderTargetsInfo);
+		}
+		break;
+	}
+	case EDecalRenderTargetMode::AmbientOcclusion:
+	{		
+		const FRDGTextureDesc AOTextureDesc = GetScreenSpaceAOTextureDesc(Config.Extent);
+		AddRenderTargetInfo(AOTextureDesc.Format, AOTextureDesc.Flags, RenderTargetsInfo);
+		break;
+	}
+
+	default:
+		checkNoEntry();
+	}
+
+	SetupDepthStencilInfo(PF_DepthStencil, Config.DepthCreateFlags, ERenderTargetLoadAction::ELoad,
+		ERenderTargetLoadAction::ELoad, FExclusiveDepthStencil::DepthRead_StencilWrite, RenderTargetsInfo);
+}
+
 void GetDeferredDecalPassParameters(
 	FRDGBuilder &GraphBuilder,
 	const FViewInfo& View,
diff --git a/Engine/Source/Runtime/Renderer/Private/CompositionLighting/PostProcessDeferredDecals.h b/Engine/Source/Runtime/Renderer/Private/CompositionLighting/PostProcessDeferredDecals.h
--- a/Engine/Source/Runtime/Renderer/Private/CompositionLighting/PostProcessDeferredDecals.h
+++ b/Engine/Source/Runtime/Renderer/Private/CompositionLighting/PostProcessDeferredDecals.h
@@ -66,6 +66,12 @@ BEGIN_SHADER_PARAMETER_STRUCT(FDeferredDecalPassParameters, )
 	RENDER_TARGET_BINDING_SLOTS()
 END_SHADER_PARAMETER_STRUCT()
 
+void GetDeferredDecalRenderTargetsInfo(
+	const FSceneTexturesConfig& Config,
+	EShaderPlatform ShaderPlatform,
+	EDecalRenderTargetMode RenderTargetMode,
+	FGraphicsPipelineRenderTargetsInfo& RenderTargetsInfo);
+
 void GetDeferredDecalPassParameters(
 	FRDGBuilder& GraphBuilder,
 	const FViewInfo& View,
diff --git a/Engine/Source/Runtime/Renderer/Private/CompositionLighting/PostProcessMeshDecals.cpp b/Engine/Source/Runtime/Renderer/Private/CompositionLighting/PostProcessMeshDecals.cpp
--- a/Engine/Source/Runtime/Renderer/Private/CompositionLighting/PostProcessMeshDecals.cpp
+++ b/Engine/Source/Runtime/Renderer/Private/CompositionLighting/PostProcessMeshDecals.cpp
@@ -133,6 +133,13 @@ public:
 
 	virtual void AddMeshBatch(const FMeshBatch& RESTRICT MeshBatch, uint64 BatchElementMask, const FPrimitiveSceneProxy* RESTRICT PrimitiveSceneProxy, int32 StaticMeshId = -1) override final;
 
+	virtual void CollectPSOInitializers(
+		const FSceneTexturesConfig& SceneTexturesConfig,
+		const FMaterial& Material,
+		const FVertexFactoryType* VertexFactoryType,
+		const FPSOPrecacheParams& PreCacheParams, 
+		TArray<FPSOPrecacheData>& PSOInitializers) override final;
+
 private:
 	bool TryAddMeshBatch(
 		const FMeshBatch& RESTRICT MeshBatch,
@@ -161,12 +168,12 @@ IMPLEMENT_STATIC_UNIFORM_BUFFER_SLOT(DeferredDecals);
 IMPLEMENT_STATIC_UNIFORM_BUFFER_STRUCT(FDeferredDecalUniformParameters, "DeferredDecal", DeferredDecals);
 
 FMeshDecalMeshProcessor::FMeshDecalMeshProcessor(const FScene* Scene, 
-	ERHIFeatureLevel::Type FeatureLevel,
+	ERHIFeatureLevel::Type InFeatureLevel,
 	const FSceneView* InViewIfDynamicMeshCommand, 
 	EDecalRenderStage InPassDecalStage, 
 	EDecalRenderTargetMode InRenderTargetMode,
 	FMeshPassDrawListContext* InDrawListContext)
-	: FMeshPassProcessor(EMeshPass::Num, Scene, FeatureLevel, InViewIfDynamicMeshCommand, InDrawListContext)
+	: FMeshPassProcessor(EMeshPass::MeshDecal, Scene, InFeatureLevel, InViewIfDynamicMeshCommand, InDrawListContext)
 	, PassDecalStage(InPassDecalStage)
 	, RenderTargetMode(InRenderTargetMode)
 {
@@ -315,6 +322,98 @@ bool FMeshDecalMeshProcessor::Process(
 	return true;
 }
 
+void FMeshDecalMeshProcessor::CollectPSOInitializers(
+	const FSceneTexturesConfig& SceneTexturesConfig,
+	const FMaterial& Material,
+	const FVertexFactoryType* VertexFactoryType,
+	const FPSOPrecacheParams& PreCacheParams, 
+	TArray<FPSOPrecacheData>& PSOInitializers)
+{
+	if (!Material.IsDeferredDecal())
+	{
+		return;
+	}
+
+	const EShaderPlatform ShaderPlatform = GetFeatureLevelShaderPlatform(FeatureLevel);
+	const FDecalBlendDesc DecalBlendDesc = DecalRendering::ComputeDecalBlendDesc(ShaderPlatform, Material);
+
+	for (uint32 DecalStageIter = 0; DecalStageIter < (uint32)EDecalRenderStage::Num; ++DecalStageIter)
+	{
+		EDecalRenderStage LocalDecalRenderState = EDecalRenderStage(DecalStageIter);
+
+		const bool bShouldRender = DecalRendering::IsCompatibleWithRenderStage(DecalBlendDesc, LocalDecalRenderState);
+		if (!bShouldRender)
+		{
+			continue;
+		}
+		
+		EDecalRenderTargetMode LocalRenderTargetMode = DecalRendering::GetRenderTargetMode(DecalBlendDesc, LocalDecalRenderState);
+		PassDrawRenderState.SetBlendState(DecalRendering::GetDecalBlendState(DecalBlendDesc, LocalDecalRenderState, LocalRenderTargetMode));
+
+		const FMeshDrawingPolicyOverrideSettings OverrideSettings = ComputeMeshOverrideSettings(PreCacheParams);
+		ERasterizerFillMode MeshFillMode = ComputeMeshFillMode(Material, OverrideSettings);
+		ERasterizerCullMode MeshCullMode = ComputeMeshCullMode(Material, OverrideSettings);
+
+		FMaterialShaderTypes ShaderTypes;
+		ShaderTypes.AddShaderType<FMeshDecalsVS>();
+		if (LocalDecalRenderState == EDecalRenderStage::Emissive)
+		{
+			ShaderTypes.AddShaderType<FMeshDecalsEmissivePS>();
+		}
+		else if (LocalDecalRenderState == EDecalRenderStage::AmbientOcclusion)
+		{
+			ShaderTypes.AddShaderType<FMeshDecalsAmbientOcclusionPS>();
+		}
+		else
+		{
+			ShaderTypes.AddShaderType<FMeshDecalsPS>();
+		}
+
+		FMaterialShaders Shaders;
+		if (!Material.TryGetShaders(ShaderTypes, VertexFactoryType, Shaders))
+		{
+			continue;
+		}
+
+		TMeshProcessorShaders<
+			FMeshDecalsVS,
+			FMeshDecalsPS> MeshDecalPassShaders;
+		Shaders.TryGetVertexShader(MeshDecalPassShaders.VertexShader);
+		Shaders.TryGetPixelShader(MeshDecalPassShaders.PixelShader);
+
+		FGraphicsPipelineRenderTargetsInfo RenderTargetsInfo;
+		RenderTargetsInfo.NumSamples = 1;
+		GetDeferredDecalRenderTargetsInfo(SceneTexturesConfig, ShaderPlatform, LocalRenderTargetMode, RenderTargetsInfo);
+		
+		AddGraphicsPipelineStateInitializer(
+			VertexFactoryType,
+			Material,
+			PassDrawRenderState,
+			RenderTargetsInfo,
+			MeshDecalPassShaders,
+			MeshFillMode,
+			MeshCullMode,
+			PT_TriangleList,
+			EMeshPassFeatures::Default,
+			PSOInitializers);
+	}
+}
+
+IPSOCollector* CreateMeshDecalMeshProcessor(ERHIFeatureLevel::Type FeatureLevel)
+{
+	if (DoesPlatformSupportNanite(GetFeatureLevelShaderPlatform(FeatureLevel)))
+	{
+		return new FMeshDecalMeshProcessor(nullptr, FeatureLevel, nullptr, EDecalRenderStage::None, EDecalRenderTargetMode::None, nullptr);
+	}
+	else
+	{
+		return nullptr;
+	}
+}
+
+// Only register for PSO Collection
+FRegisterPSOCollectorCreateFunction RegisterPSOCollectorMeshDecal(&CreateMeshDecalMeshProcessor, EShadingPath::Deferred, (uint32)EMeshPass::MeshDecal);
+
 void DrawDecalMeshCommands(
 	FRDGBuilder& GraphBuilder,
 	const FViewInfo& View,
diff --git a/Engine/Source/Runtime/Renderer/Private/DBufferTextures.cpp b/Engine/Source/Runtime/Renderer/Private/DBufferTextures.cpp
--- a/Engine/Source/Runtime/Renderer/Private/DBufferTextures.cpp
+++ b/Engine/Source/Runtime/Renderer/Private/DBufferTextures.cpp
@@ -30,32 +30,29 @@ EDecalDBufferMaskTechnique GetDBufferMaskTechnique(EShaderPlatform ShaderPlatfor
 	return EDecalDBufferMaskTechnique::Disabled;
 }
 
-FDBufferTextures CreateDBufferTextures(FRDGBuilder& GraphBuilder, FIntPoint Extent, EShaderPlatform ShaderPlatform)
+FDBufferTexturesDesc GetDBufferTexturesDesc(FIntPoint Extent, EShaderPlatform ShaderPlatform)
 {
-	FDBufferTextures DBufferTextures;
+	FDBufferTexturesDesc DBufferTexturesDesc;
 
 	if (IsUsingDBuffers(ShaderPlatform))
 	{
 		const EDecalDBufferMaskTechnique DBufferMaskTechnique = GetDBufferMaskTechnique(ShaderPlatform);
 		const ETextureCreateFlags WriteMaskFlags = DBufferMaskTechnique == EDecalDBufferMaskTechnique::WriteMask ? TexCreate_NoFastClearFinalize | TexCreate_DisableDCC : TexCreate_None;
 		const ETextureCreateFlags BaseFlags = WriteMaskFlags | TexCreate_ShaderResource | TexCreate_RenderTargetable;
-		const ERDGTextureFlags TextureFlags = DBufferMaskTechnique != EDecalDBufferMaskTechnique::Disabled
-			? ERDGTextureFlags::MaintainCompression
-			: ERDGTextureFlags::None;
-
+		
 		FRDGTextureDesc Desc = FRDGTextureDesc::Create2D(Extent, PF_B8G8R8A8, FClearValueBinding::None, BaseFlags);
 
 		Desc.Flags = BaseFlags | GFastVRamConfig.DBufferA;
 		Desc.ClearValue = FClearValueBinding::Black;
-		DBufferTextures.DBufferA = GraphBuilder.CreateTexture(Desc, TEXT("DBufferA"), TextureFlags);
+		DBufferTexturesDesc.DBufferADesc = Desc;
 
 		Desc.Flags = BaseFlags | GFastVRamConfig.DBufferB;
 		Desc.ClearValue = FClearValueBinding(FLinearColor(128.0f / 255.0f, 128.0f / 255.0f, 128.0f / 255.0f, 1));
-		DBufferTextures.DBufferB = GraphBuilder.CreateTexture(Desc, TEXT("DBufferB"), TextureFlags);
+		DBufferTexturesDesc.DBufferBDesc = Desc;
 
 		Desc.Flags = BaseFlags | GFastVRamConfig.DBufferC;
 		Desc.ClearValue = FClearValueBinding(FLinearColor(0, 0, 0, 1));
-		DBufferTextures.DBufferC = GraphBuilder.CreateTexture(Desc, TEXT("DBufferC"), TextureFlags);
+		DBufferTexturesDesc.DBufferCDesc = Desc;
 
 		if (DBufferMaskTechnique == EDecalDBufferMaskTechnique::PerPixel)
 		{
@@ -65,7 +62,33 @@ FDBufferTextures CreateDBufferTextures(FRDGBuilder& GraphBuilder, FIntPoint Exte
 			check(Desc.Format == PF_B8G8R8A8);
 			Desc.Flags = TexCreate_ShaderResource | TexCreate_RenderTargetable;
 			Desc.ClearValue = FClearValueBinding::Transparent;
-			DBufferTextures.DBufferMask = GraphBuilder.CreateTexture(Desc, TEXT("DBufferMask"));
+			DBufferTexturesDesc.DBufferMaskDesc = Desc;
+		}
+	}
+
+	return DBufferTexturesDesc;
+}
+
+FDBufferTextures CreateDBufferTextures(FRDGBuilder& GraphBuilder, FIntPoint Extent, EShaderPlatform ShaderPlatform)
+{
+	FDBufferTextures DBufferTextures;
+
+	if (IsUsingDBuffers(ShaderPlatform))
+	{
+		FDBufferTexturesDesc TexturesDesc = GetDBufferTexturesDesc(Extent, ShaderPlatform);
+
+		const EDecalDBufferMaskTechnique DBufferMaskTechnique = GetDBufferMaskTechnique(ShaderPlatform);
+		const ERDGTextureFlags TextureFlags = DBufferMaskTechnique != EDecalDBufferMaskTechnique::Disabled
+			? ERDGTextureFlags::MaintainCompression
+			: ERDGTextureFlags::None;
+				
+		DBufferTextures.DBufferA = GraphBuilder.CreateTexture(TexturesDesc.DBufferADesc, TEXT("DBufferA"), TextureFlags);
+		DBufferTextures.DBufferB = GraphBuilder.CreateTexture(TexturesDesc.DBufferBDesc, TEXT("DBufferB"), TextureFlags);
+		DBufferTextures.DBufferC = GraphBuilder.CreateTexture(TexturesDesc.DBufferCDesc, TEXT("DBufferC"), TextureFlags);
+
+		if (DBufferMaskTechnique == EDecalDBufferMaskTechnique::PerPixel)
+		{
+			DBufferTextures.DBufferMask = GraphBuilder.CreateTexture(TexturesDesc.DBufferMaskDesc, TEXT("DBufferMask"));
 		}
 	}
 
diff --git a/Engine/Source/Runtime/Renderer/Private/DBufferTextures.h b/Engine/Source/Runtime/Renderer/Private/DBufferTextures.h
--- a/Engine/Source/Runtime/Renderer/Private/DBufferTextures.h
+++ b/Engine/Source/Runtime/Renderer/Private/DBufferTextures.h
@@ -20,6 +20,14 @@ enum class EDecalDBufferMaskTechnique
 
 EDecalDBufferMaskTechnique GetDBufferMaskTechnique(EShaderPlatform ShaderPlatform);
 
+struct FDBufferTexturesDesc
+{
+	FRDGTextureDesc DBufferADesc;
+	FRDGTextureDesc DBufferBDesc;
+	FRDGTextureDesc DBufferCDesc;
+	FRDGTextureDesc DBufferMaskDesc;
+};
+
 struct FDBufferTextures
 {
 	bool IsValid() const;
@@ -30,6 +38,7 @@ struct FDBufferTextures
 	FRDGTextureRef DBufferMask = nullptr;
 };
 
+FDBufferTexturesDesc GetDBufferTexturesDesc(FIntPoint Extent, EShaderPlatform ShaderPlatform);
 FDBufferTextures CreateDBufferTextures(FRDGBuilder& GraphBuilder, FIntPoint Extent, EShaderPlatform ShaderPlatform);
 
 BEGIN_SHADER_PARAMETER_STRUCT(FDBufferParameters, )
diff --git a/Engine/Source/Runtime/Renderer/Private/DecalRenderingCommon.h b/Engine/Source/Runtime/Renderer/Private/DecalRenderingCommon.h
--- a/Engine/Source/Runtime/Renderer/Private/DecalRenderingCommon.h
+++ b/Engine/Source/Runtime/Renderer/Private/DecalRenderingCommon.h
@@ -49,6 +49,8 @@ enum class EDecalRenderStage : uint8
 	// Ambient occlusion decal pass.
 	// A decal can write regular attributes in another pass and then AO in this pass.
 	AmbientOcclusion = 6,
+
+	Num,
 };
 
 /** Enumeration of the render target layouts for decal rendering. */
diff --git a/Engine/Source/Runtime/Renderer/Public/MeshPassProcessor.h b/Engine/Source/Runtime/Renderer/Public/MeshPassProcessor.h
--- a/Engine/Source/Runtime/Renderer/Public/MeshPassProcessor.h
+++ b/Engine/Source/Runtime/Renderer/Public/MeshPassProcessor.h
@@ -57,6 +57,7 @@ namespace EMeshPass
 		LumenFrontLayerTranslucencyGBuffer,
 		DitheredLODFadingOutMaskPass, /** A mini depth pass used to mark pixels with dithered LOD fading out. Currently only used by ray tracing shadows. */
 		NaniteMeshPass,
+		MeshDecal,
 
 #if WITH_EDITOR
 		HitProxy,
@@ -104,6 +105,7 @@ inline const TCHAR* GetMeshPassName(EMeshPass::Type MeshPass)
 	case EMeshPass::LumenFrontLayerTranslucencyGBuffer: return TEXT("LumenFrontLayerTranslucencyGBuffer");
 	case EMeshPass::DitheredLODFadingOutMaskPass: return TEXT("DitheredLODFadingOutMaskPass");
 	case EMeshPass::NaniteMeshPass: return TEXT("NaniteMeshPass");
+	case EMeshPass::MeshDecal: return TEXT("MeshDecal");
 #if WITH_EDITOR
 	case EMeshPass::HitProxy: return TEXT("HitProxy");
 	case EMeshPass::HitProxyOpaqueOnly: return TEXT("HitProxyOpaqueOnly");
@@ -113,9 +115,9 @@ inline const TCHAR* GetMeshPassName(EMeshPass::Type MeshPass)
 	}
 
 #if WITH_EDITOR
-	static_assert(EMeshPass::Num == 28 + 4, "Need to update switch(MeshPass) after changing EMeshPass"); // GUID to prevent incorrect auto-resolves, please change when changing the expression: {FED9577B-F826-478C-841F-6A01BCB4415E}
+	static_assert(EMeshPass::Num == 29 + 4, "Need to update switch(MeshPass) after changing EMeshPass"); // GUID to prevent incorrect auto-resolves, please change when changing the expression: {A6E82589-44B3-4DAD-AC57-8AF6BD50DF43}
 #else
-	static_assert(EMeshPass::Num == 28, "Need to update switch(MeshPass) after changing EMeshPass"); // GUID to prevent incorrect auto-resolves, please change when changing the expression: {FED9577B-F826-478C-841F-6A01BCB4415E}
+	static_assert(EMeshPass::Num == 29, "Need to update switch(MeshPass) after changing EMeshPass"); // GUID to prevent incorrect auto-resolves, please change when changing the expression: {A6E82589-44B3-4DAD-AC57-8AF6BD50DF43}
 #endif
 
 	checkf(0, TEXT("Missing case for EMeshPass %u"), (uint32)MeshPass);
