Fix build because EMeshPass::Type now needs 6 bits of storage and doesn't fit 32 bit mask anymore

#rnx

diff --git a/Engine/Source/Runtime/Engine/Public/PSOPrecache.h b/Engine/Source/Runtime/Engine/Public/PSOPrecache.h
--- a/Engine/Source/Runtime/Engine/Public/PSOPrecache.h
+++ b/Engine/Source/Runtime/Engine/Public/PSOPrecache.h
@@ -152,7 +152,7 @@ class ENGINE_API FPSOCollectorCreateManager
 {
 public:
 
-	constexpr static uint32 MaxPSOCollectorCount = 32;
+	constexpr static uint32 MaxPSOCollectorCount = 33;
 
 	static PSOCollectorCreateFunction GetCreateFunction(EShadingPath ShadingPath, uint32 Index)
 	{
diff --git a/Engine/Source/Runtime/Renderer/Public/MeshPassProcessor.h b/Engine/Source/Runtime/Renderer/Public/MeshPassProcessor.h
--- a/Engine/Source/Runtime/Renderer/Public/MeshPassProcessor.h
+++ b/Engine/Source/Runtime/Renderer/Public/MeshPassProcessor.h
@@ -67,7 +67,7 @@ namespace EMeshPass
 #endif
 
 		Num,
-		NumBits = 5,
+		NumBits = 6,
 	};
 }
 static_assert(EMeshPass::Num <= (1 << EMeshPass::NumBits), "EMeshPass::Num will not fit in EMeshPass::NumBits");
@@ -135,18 +135,18 @@ public:
 
 	void Set(EMeshPass::Type Pass) 
 	{ 
-		Data |= (1 << Pass); 
+		Data |= (uint64(1) << Pass); 
 	}
 
 	bool Get(EMeshPass::Type Pass) const 
 	{ 
-		return !!(Data & (1 << Pass)); 
+		return !!(Data & (uint64(1) << Pass)); 
 	}
 
 	EMeshPass::Type SkipEmpty(EMeshPass::Type Pass) const 
 	{
-		uint32 Mask = 0xFFffFFff << Pass;
-		return EMeshPass::Type(FMath::Min<uint32>(EMeshPass::Num, FMath::CountTrailingZeros(Data & Mask)));
+		uint64 Mask = 0xFFffFFffFFffFFffULL << Pass;
+		return EMeshPass::Type(FMath::Min<uint64>(EMeshPass::Num, FMath::CountTrailingZeros(Data & Mask)));
 	}
 
 	int GetNum() 
@@ -169,7 +169,7 @@ public:
 		return Data == 0; 
 	}
 
-	uint32 Data;
+	uint64 Data;
 };
 
 struct FMinimalBoundShaderStateInput
