HeterogeneousVolumes: Integration with virtual shadow maps.

diff --git a/Engine/Shaders/Private/HeterogeneousVolumes/HeterogeneousVolumesHardwareRayTracing.usf b/Engine/Shaders/Private/HeterogeneousVolumes/HeterogeneousVolumesHardwareRayTracing.usf
--- a/Engine/Shaders/Private/HeterogeneousVolumes/HeterogeneousVolumesHardwareRayTracing.usf
+++ b/Engine/Shaders/Private/HeterogeneousVolumes/HeterogeneousVolumesHardwareRayTracing.usf
@@ -12,6 +12,10 @@
 #define SHADOW_QUALITY 2
 #define NO_TRANSLUCENCY_AVAILABLE
 
+#if VIRTUAL_SHADOW_MAP
+#include "../VirtualShadowMaps/VirtualShadowMapProjectionCommon.ush"
+#endif
+
 #include "HeterogeneousVolumesSparseVoxelUniformBufferUtils.ush"
 #include "HeterogeneousVolumesTracingUtils.ush"
 #include "HeterogeneousVolumesLightingUtils.ush"
@@ -378,6 +382,9 @@ int bApplyDirectLighting;
 int bApplyShadowTransmittance;
 int LightType;
 
+// Shadowing
+uint VirtualShadowMapId;
+
 float3 ComputeTransmittanceHardwareRayTracing(
 	float3 Origin,
 	float3 Direction,
diff --git a/Engine/Shaders/Private/HeterogeneousVolumes/HeterogeneousVolumesLiveShadingPipeline.usf b/Engine/Shaders/Private/HeterogeneousVolumes/HeterogeneousVolumesLiveShadingPipeline.usf
--- a/Engine/Shaders/Private/HeterogeneousVolumes/HeterogeneousVolumesLiveShadingPipeline.usf
+++ b/Engine/Shaders/Private/HeterogeneousVolumes/HeterogeneousVolumesLiveShadingPipeline.usf
@@ -16,6 +16,10 @@
 #include "HeterogeneousVolumesTracingUtils.ush"
 #include "HeterogeneousVolumesLightingUtils.ush"
 
+#if VIRTUAL_SHADOW_MAP
+#include "../VirtualShadowMaps/VirtualShadowMapProjectionCommon.ush"
+#endif
+
 #ifndef THREADGROUP_SIZE_1D
 #define THREADGROUP_SIZE_1D 1
 #endif // THREADGROUP_SIZE_1D
@@ -40,6 +44,9 @@ int bApplyEmissionAndTransmittance;
 int bApplyDirectLighting;
 int bApplyShadowTransmittance;
 int LightType;
+
+// Shadowing
+uint VirtualShadowMapId;
 float ShadowStepSize;
 
 // Ray data
diff --git a/Engine/Shaders/Private/HeterogeneousVolumes/HeterogeneousVolumesPreshadingPipeline.usf b/Engine/Shaders/Private/HeterogeneousVolumes/HeterogeneousVolumesPreshadingPipeline.usf
--- a/Engine/Shaders/Private/HeterogeneousVolumes/HeterogeneousVolumesPreshadingPipeline.usf
+++ b/Engine/Shaders/Private/HeterogeneousVolumes/HeterogeneousVolumesPreshadingPipeline.usf
@@ -11,6 +11,10 @@
 #define SHADOW_QUALITY 2
 #define NO_TRANSLUCENCY_AVAILABLE
 
+#if VIRTUAL_SHADOW_MAP
+#include "../VirtualShadowMaps/VirtualShadowMapProjectionCommon.ush"
+#endif
+
 #include "HeterogeneousVolumesSparseVoxelUniformBufferUtils.ush"
 #include "HeterogeneousVolumesTracingUtils.ush"
 #include "HeterogeneousVolumesLightingUtils.ush"
@@ -46,6 +50,9 @@ int bApplyDirectLighting;
 int bApplyShadowTransmittance;
 int LightType;
 
+// Shadowing
+uint VirtualShadowMapId;
+
 // Ray data
 float MaxTraceDistance;
 float MaxShadowTraceDistance;
diff --git a/Engine/Shaders/Private/HeterogeneousVolumes/HeterogeneousVolumesRayMarchingUtils.ush b/Engine/Shaders/Private/HeterogeneousVolumes/HeterogeneousVolumesRayMarchingUtils.ush
--- a/Engine/Shaders/Private/HeterogeneousVolumes/HeterogeneousVolumesRayMarchingUtils.ush
+++ b/Engine/Shaders/Private/HeterogeneousVolumes/HeterogeneousVolumesRayMarchingUtils.ush
@@ -131,6 +131,15 @@ float ComputeHardSurfaceShadowFactor(
 			HardSurfaceShadowFactor *= ShadowFactor;
 		}
 	}
+
+#if VIRTUAL_SHADOW_MAP
+	if (VirtualShadowMapId != INDEX_NONE)
+	{
+		FVirtualShadowMapSampleResult VirtualShadowMapSample = SampleVirtualShadowMapTranslatedWorld(VirtualShadowMapId, TranslatedWorldPosition);
+		HardSurfaceShadowFactor *= VirtualShadowMapSample.ShadowFactor;
+	}
+#endif // VIRTUALSHADOW_MAP
+
 #endif // HARD_SURFACE_SHADOWING
 
 	return HardSurfaceShadowFactor;
diff --git a/Engine/Source/Runtime/Renderer/Private/HeterogeneousVolumes/HeterogeneousVolumes.cpp b/Engine/Source/Runtime/Renderer/Private/HeterogeneousVolumes/HeterogeneousVolumes.cpp
--- a/Engine/Source/Runtime/Renderer/Private/HeterogeneousVolumes/HeterogeneousVolumes.cpp
+++ b/Engine/Source/Runtime/Renderer/Private/HeterogeneousVolumes/HeterogeneousVolumes.cpp
@@ -377,6 +377,7 @@ void FDeferredShadingSceneRenderer::RenderHeterogeneousVolumes(
 						View,
 						// Shadow Data
 						VisibleLightInfos,
+						VirtualShadowMapArray,
 						// Object data
 						PrimitiveSceneProxy,
 						MaterialRenderProxy,
@@ -398,6 +399,7 @@ void FDeferredShadingSceneRenderer::RenderHeterogeneousVolumes(
 						View,
 						// Shadow Data
 						VisibleLightInfos,
+						VirtualShadowMapArray,
 						// Object Data
 						PrimitiveSceneProxy,
 						MaterialRenderProxy,
diff --git a/Engine/Source/Runtime/Renderer/Private/HeterogeneousVolumes/HeterogeneousVolumes.h b/Engine/Source/Runtime/Renderer/Private/HeterogeneousVolumes/HeterogeneousVolumes.h
--- a/Engine/Source/Runtime/Renderer/Private/HeterogeneousVolumes/HeterogeneousVolumes.h
+++ b/Engine/Source/Runtime/Renderer/Private/HeterogeneousVolumes/HeterogeneousVolumes.h
@@ -102,6 +102,7 @@ void RenderWithLiveShading(
 	const FViewInfo& View,
 	// Shadow data
 	TArray<FVisibleLightInfo, SceneRenderingAllocator>& VisibleLightInfos,
+	const FVirtualShadowMapArray& VirtualShadowMapArray,
 	// Object data
 	const FPrimitiveSceneProxy* PrimitiveSceneProxy,
 	const FMaterialRenderProxy* MaterialRenderProxy,
@@ -122,6 +123,7 @@ void RenderWithPreshading(
 	FViewInfo& View,
 	// Shadow data
 	TArray<FVisibleLightInfo, SceneRenderingAllocator>& VisibleLightInfos,
+	const FVirtualShadowMapArray& VirtualShadowMapArray,
 	// Object data
 	const FPrimitiveSceneProxy* PrimitiveSceneProxy,
 	const FMaterialRenderProxy* MaterialRenderProxy,
@@ -217,6 +219,7 @@ void RenderLightingCacheWithPreshadingHardwareRayTracing(
 	const FLightSceneInfo* LightSceneInfo,
 	// Shadow data
 	const FVisibleLightInfo* VisibleLightInfo,
+	const FVirtualShadowMapArray& VirtualShadowMapArray,
 	// Object data
 	const FPrimitiveSceneProxy* PrimitiveSceneProxy,
 	// Sparse voxel data
@@ -241,6 +244,7 @@ void RenderSingleScatteringWithPreshadingHardwareRayTracing(
 	const FLightSceneInfo* LightSceneInfo,
 	// Shadow data
 	const FVisibleLightInfo* VisibleLightInfo,
+	const FVirtualShadowMapArray& VirtualShadowMapArray,
 	// Object data
 	const FPrimitiveSceneProxy* PrimitiveSceneProxy,
 	// Sparse voxel data
diff --git a/Engine/Source/Runtime/Renderer/Private/HeterogeneousVolumes/HeterogeneousVolumesHardwareRayTracing.cpp b/Engine/Source/Runtime/Renderer/Private/HeterogeneousVolumes/HeterogeneousVolumesHardwareRayTracing.cpp
--- a/Engine/Source/Runtime/Renderer/Private/HeterogeneousVolumes/HeterogeneousVolumesHardwareRayTracing.cpp
+++ b/Engine/Source/Runtime/Renderer/Private/HeterogeneousVolumes/HeterogeneousVolumesHardwareRayTracing.cpp
@@ -344,6 +344,8 @@ class FRenderLightingCacheWithPreshadingRGS : public FGlobalShader
 		// Shadow data
 		SHADER_PARAMETER_RDG_UNIFORM_BUFFER(FForwardLightData, ForwardLightData)
 		SHADER_PARAMETER_STRUCT_INCLUDE(FVolumeShadowingShaderParameters, VolumeShadowingShaderParameters)
+		SHADER_PARAMETER_STRUCT_INCLUDE(FVirtualShadowMapSamplingParameters, VirtualShadowMapSamplingParameters)
+		SHADER_PARAMETER(int32, VirtualShadowMapId)
 
 		// Sparse Volume
 		SHADER_PARAMETER_RDG_UNIFORM_BUFFER(FSparseVoxelUniformBufferParameters, SparseVoxelUniformBuffer)
@@ -368,6 +370,21 @@ class FRenderLightingCacheWithPreshadingRGS : public FGlobalShader
 		return ShouldCompileRayTracingShadersForProject(Parameters.Platform) && DoesPlatformSupportHeterogeneousVolumes(Parameters.Platform) &&
 			FDataDrivenShaderPlatformInfo::GetSupportsRayTracingProceduralPrimitive(Parameters.Platform);
 	}
+
+	static void ModifyCompilationEnvironment(
+		const FGlobalShaderPermutationParameters& Parameters,
+		FShaderCompilerEnvironment& OutEnvironment
+	)
+	{
+		FMaterialShader::ModifyCompilationEnvironment(Parameters, OutEnvironment);
+
+		bool bSupportVirtualShadowMap = IsFeatureLevelSupported(Parameters.Platform, ERHIFeatureLevel::SM5);
+		if (bSupportVirtualShadowMap)
+		{
+			OutEnvironment.SetDefine(TEXT("VIRTUAL_SHADOW_MAP"), 1);
+			FVirtualShadowMapArray::SetShaderDefines(OutEnvironment);
+		}
+	}
 };
 
 IMPLEMENT_GLOBAL_SHADER(FRenderLightingCacheWithPreshadingRGS, "/Engine/Private/HeterogeneousVolumes/HeterogeneousVolumesHardwareRayTracing.usf", "RenderLightingCacheWithPreshadingRGS", SF_RayGen);
@@ -397,6 +414,8 @@ class FRenderSingleScatteringWithPreshadingRGS : public FGlobalShader
 		// Shadow data
 		SHADER_PARAMETER_RDG_UNIFORM_BUFFER(FForwardLightData, ForwardLightData)
 		SHADER_PARAMETER_STRUCT_INCLUDE(FVolumeShadowingShaderParameters, VolumeShadowingShaderParameters)
+		SHADER_PARAMETER_STRUCT_INCLUDE(FVirtualShadowMapSamplingParameters, VirtualShadowMapSamplingParameters)
+		SHADER_PARAMETER(int32, VirtualShadowMapId)
 
 		// Sparse Volume
 		SHADER_PARAMETER_RDG_UNIFORM_BUFFER(FSparseVoxelUniformBufferParameters, SparseVoxelUniformBuffer)
@@ -428,6 +447,13 @@ class FRenderSingleScatteringWithPreshadingRGS : public FGlobalShader
 	)
 	{
 		FMaterialShader::ModifyCompilationEnvironment(Parameters, OutEnvironment);
+
+		bool bSupportVirtualShadowMap = IsFeatureLevelSupported(Parameters.Platform, ERHIFeatureLevel::SM5);
+		if (bSupportVirtualShadowMap)
+		{
+			OutEnvironment.SetDefine(TEXT("VIRTUAL_SHADOW_MAP"), 1);
+			FVirtualShadowMapArray::SetShaderDefines(OutEnvironment);
+		}
 	}
 };
 
@@ -522,6 +548,7 @@ void RenderLightingCacheWithPreshadingHardwareRayTracing(
 	const FLightSceneInfo* LightSceneInfo,
 	// Shadow data
 	const FVisibleLightInfo* VisibleLightInfo,
+	const FVirtualShadowMapArray& VirtualShadowMapArray,
 	// Object data
 	const FPrimitiveSceneProxy* PrimitiveSceneProxy,
 	// Sparse voxel data
@@ -575,11 +602,14 @@ void RenderLightingCacheWithPreshadingHardwareRayTracing(
 			{
 				SetVolumeShadowingDefaultShaderParametersGlobal(GraphBuilder, PassParameters->VolumeShadowingShaderParameters);
 			}
+			PassParameters->VirtualShadowMapId = VisibleLightInfo->GetVirtualShadowMapId(&View);
 		}
 		else
 		{
 			SetVolumeShadowingDefaultShaderParametersGlobal(GraphBuilder, PassParameters->VolumeShadowingShaderParameters);
+			PassParameters->VirtualShadowMapId = -1;
 		}
+		PassParameters->VirtualShadowMapSamplingParameters = VirtualShadowMapArray.GetSamplingParameters(GraphBuilder);
 
 		// Output
 		PassParameters->RWLightingCacheTexture = GraphBuilder.CreateUAV(LightingCacheTexture);
@@ -655,6 +685,7 @@ void RenderSingleScatteringWithPreshadingHardwareRayTracing(
 	const FLightSceneInfo* LightSceneInfo,
 	// Shadow data
 	const FVisibleLightInfo* VisibleLightInfo,
+	const FVirtualShadowMapArray& VirtualShadowMapArray,
 	// Object data
 	const FPrimitiveSceneProxy* PrimitiveSceneProxy,
 	// Sparse voxel data
@@ -691,12 +722,22 @@ void RenderSingleScatteringWithPreshadingHardwareRayTracing(
 		{
 			const FProjectedShadowInfo* ProjectedShadowInfo = GetShadowForInjectionIntoVolumetricFog(*VisibleLightInfo);
 			bool bDynamicallyShadowed = ProjectedShadowInfo != NULL;
-			GetVolumeShadowingShaderParameters(GraphBuilder, View, LightSceneInfo, ProjectedShadowInfo, PassParameters->VolumeShadowingShaderParameters);
+			if (bDynamicallyShadowed)
+			{
+				GetVolumeShadowingShaderParameters(GraphBuilder, View, LightSceneInfo, ProjectedShadowInfo, PassParameters->VolumeShadowingShaderParameters);
+			}
+			else
+			{
+				SetVolumeShadowingDefaultShaderParametersGlobal(GraphBuilder, PassParameters->VolumeShadowingShaderParameters);
+			}
+			PassParameters->VirtualShadowMapId = VisibleLightInfo->GetVirtualShadowMapId(&View);
 		}
 		else
 		{
 			SetVolumeShadowingDefaultShaderParametersGlobal(GraphBuilder, PassParameters->VolumeShadowingShaderParameters);
+			PassParameters->VirtualShadowMapId = -1;
 		}
+		PassParameters->VirtualShadowMapSamplingParameters = VirtualShadowMapArray.GetSamplingParameters(GraphBuilder);
 
 		// Sparse Voxel data
 		PassParameters->SparseVoxelUniformBuffer = SparseVoxelUniformBuffer;
diff --git a/Engine/Source/Runtime/Renderer/Private/HeterogeneousVolumes/HeterogeneousVolumesLiveShadingPipeline.cpp b/Engine/Source/Runtime/Renderer/Private/HeterogeneousVolumes/HeterogeneousVolumesLiveShadingPipeline.cpp
--- a/Engine/Source/Runtime/Renderer/Private/HeterogeneousVolumes/HeterogeneousVolumesLiveShadingPipeline.cpp
+++ b/Engine/Source/Runtime/Renderer/Private/HeterogeneousVolumes/HeterogeneousVolumesLiveShadingPipeline.cpp
@@ -34,6 +34,8 @@ class FRenderLightingCacheWithLiveShadingCS : public FMeshMaterialShader
 		SHADER_PARAMETER(float, ShadowStepSize)
 		SHADER_PARAMETER_RDG_UNIFORM_BUFFER(FForwardLightData, ForwardLightData)
 		SHADER_PARAMETER_STRUCT_INCLUDE(FVolumeShadowingShaderParameters, VolumeShadowingShaderParameters)
+		SHADER_PARAMETER_STRUCT_INCLUDE(FVirtualShadowMapSamplingParameters, VirtualShadowMapSamplingParameters)
+		SHADER_PARAMETER(int32, VirtualShadowMapId)
 
 		// Object data
 		SHADER_PARAMETER(FMatrix44f, LocalToWorld)
@@ -98,6 +100,13 @@ class FRenderLightingCacheWithLiveShadingCS : public FMeshMaterialShader
 		OutEnvironment.SetDefine(TEXT("THREADGROUP_SIZE_2D"), GetThreadGroupSize2D());
 		OutEnvironment.SetDefine(TEXT("THREADGROUP_SIZE_3D"), GetThreadGroupSize3D());
 
+		bool bSupportVirtualShadowMap = IsFeatureLevelSupported(Parameters.Platform, ERHIFeatureLevel::SM5);
+		if (bSupportVirtualShadowMap)
+		{
+			OutEnvironment.SetDefine(TEXT("VIRTUAL_SHADOW_MAP"), 1);
+			FVirtualShadowMapArray::SetShaderDefines(OutEnvironment);
+		}
+
 		// This shader takes a very long time to compile with FXC, so we pre-compile it with DXC first and then forward the optimized HLSL to FXC.
 		OutEnvironment.CompilerFlags.Add(CFLAG_PrecompileWithDXC);
 		OutEnvironment.CompilerFlags.Add(CFLAG_AllowTypedUAVLoads);
@@ -148,6 +157,8 @@ class FRenderSingleScatteringWithLiveShadingCS : public FMeshMaterialShader
 		SHADER_PARAMETER(float, ShadowStepSize)
 		SHADER_PARAMETER_RDG_UNIFORM_BUFFER(FForwardLightData, ForwardLightData)
 		SHADER_PARAMETER_STRUCT_INCLUDE(FVolumeShadowingShaderParameters, VolumeShadowingShaderParameters)
+		SHADER_PARAMETER_STRUCT_INCLUDE(FVirtualShadowMapSamplingParameters, VirtualShadowMapSamplingParameters)
+		SHADER_PARAMETER(int32, VirtualShadowMapId)
 
 		// Object data
 		SHADER_PARAMETER(FMatrix44f, LocalToWorld)
@@ -213,6 +224,13 @@ class FRenderSingleScatteringWithLiveShadingCS : public FMeshMaterialShader
 		OutEnvironment.SetDefine(TEXT("THREADGROUP_SIZE_1D"), GetThreadGroupSize1D());
 		OutEnvironment.SetDefine(TEXT("THREADGROUP_SIZE_2D"), GetThreadGroupSize2D());
 
+		bool bSupportVirtualShadowMap = IsFeatureLevelSupported(Parameters.Platform, ERHIFeatureLevel::SM5);
+		if (bSupportVirtualShadowMap)
+		{
+			OutEnvironment.SetDefine(TEXT("VIRTUAL_SHADOW_MAP"), 1);
+			FVirtualShadowMapArray::SetShaderDefines(OutEnvironment);
+		}
+
 		// This shader takes a very long time to compile with FXC, so we pre-compile it with DXC first and then forward the optimized HLSL to FXC.
 		OutEnvironment.CompilerFlags.Add(CFLAG_PrecompileWithDXC);
 		OutEnvironment.CompilerFlags.Add(CFLAG_AllowTypedUAVLoads);
@@ -299,6 +317,7 @@ void RenderLightingCacheWithLiveShading(
 	const FLightSceneInfo* LightSceneInfo,
 	// Shadow data
 	const FVisibleLightInfo* VisibleLightInfo,
+	const FVirtualShadowMapArray& VirtualShadowMapArray,
 	// Object data
 	const FPrimitiveSceneProxy* PrimitiveSceneProxy,
 	const FMaterialRenderProxy* DefaultMaterialRenderProxy,
@@ -363,11 +382,14 @@ void RenderLightingCacheWithLiveShading(
 			{
 				SetVolumeShadowingDefaultShaderParametersGlobal(GraphBuilder, PassParameters->VolumeShadowingShaderParameters);
 			}
+			PassParameters->VirtualShadowMapId = VisibleLightInfo->GetVirtualShadowMapId(&View);
 		}
 		else
 		{
 			SetVolumeShadowingDefaultShaderParametersGlobal(GraphBuilder, PassParameters->VolumeShadowingShaderParameters);
+			PassParameters->VirtualShadowMapId = -1;
 		}
+		PassParameters->VirtualShadowMapSamplingParameters = VirtualShadowMapArray.GetSamplingParameters(GraphBuilder);
 
 		// Output
 		PassParameters->RWLightingCacheTexture = GraphBuilder.CreateUAV(LightingCacheTexture);
@@ -415,6 +437,7 @@ void RenderSingleScatteringWithLiveShading(
 	const FLightSceneInfo* LightSceneInfo,
 	// Shadow data
 	const FVisibleLightInfo* VisibleLightInfo,
+	const FVirtualShadowMapArray& VirtualShadowMapArray,
 	// Object data
 	const FPrimitiveSceneProxy* PrimitiveSceneProxy,
 	const FMaterialRenderProxy* DefaultMaterialRenderProxy,
@@ -489,11 +512,13 @@ void RenderSingleScatteringWithLiveShading(
 			{
 				SetVolumeShadowingDefaultShaderParametersGlobal(GraphBuilder, PassParameters->VolumeShadowingShaderParameters);
 			}
+			PassParameters->VirtualShadowMapId = VisibleLightInfo->GetVirtualShadowMapId(&View);
 		}
 		else
 		{
 			SetVolumeShadowingDefaultShaderParametersGlobal(GraphBuilder, PassParameters->VolumeShadowingShaderParameters);
 		}
+		PassParameters->VirtualShadowMapSamplingParameters = VirtualShadowMapArray.GetSamplingParameters(GraphBuilder);
 
 		// Volume data
 		if ((HeterogeneousVolumes::UseLightingCacheForTransmittance() && bApplyShadowTransmittance) || HeterogeneousVolumes::UseLightingCacheForInscattering())
@@ -540,6 +565,7 @@ void RenderWithTransmittanceVolumePipeline(
 	const FViewInfo& View,
 	// Shadow data
 	TArray<FVisibleLightInfo, SceneRenderingAllocator>& VisibleLightInfos,
+	const FVirtualShadowMapArray& VirtualShadowMapArray,
 	// Object data
 	const FPrimitiveSceneProxy* PrimitiveSceneProxy,
 	const FMaterialRenderProxy* MaterialRenderProxy,
@@ -602,6 +628,7 @@ void RenderWithTransmittanceVolumePipeline(
 				LightSceneInfo,
 				// Shadow data
 				VisibleLightInfo,
+				VirtualShadowMapArray,
 				// Object data
 				PrimitiveSceneProxy,
 				MaterialRenderProxy,
@@ -626,6 +653,7 @@ void RenderWithTransmittanceVolumePipeline(
 			LightSceneInfo,
 			// Shadow data
 			VisibleLightInfo,
+			VirtualShadowMapArray,
 			// Object data
 			PrimitiveSceneProxy,
 			MaterialRenderProxy,
@@ -646,6 +674,7 @@ void RenderWithInscatteringVolumePipeline(
 	const FViewInfo& View,
 	// Shadow data
 	TArray<FVisibleLightInfo, SceneRenderingAllocator>& VisibleLightInfos,
+	const FVirtualShadowMapArray& VirtualShadowMapArray,
 	// Object data
 	const FPrimitiveSceneProxy* PrimitiveSceneProxy,
 	const FMaterialRenderProxy* MaterialRenderProxy,
@@ -706,6 +735,7 @@ void RenderWithInscatteringVolumePipeline(
 			LightSceneInfo,
 			// Shadow data
 			VisibleLightInfo,
+			VirtualShadowMapArray,
 			// Object data
 			PrimitiveSceneProxy,
 			MaterialRenderProxy,
@@ -740,6 +770,7 @@ void RenderWithInscatteringVolumePipeline(
 			LightSceneInfo,
 			// Shadow data
 			VisibleLightInfo,
+			VirtualShadowMapArray,
 			// Object data
 			PrimitiveSceneProxy,
 			MaterialRenderProxy,
@@ -760,6 +791,7 @@ void RenderWithLiveShading(
 	const FViewInfo& View,
 	// Shadow data
 	TArray<FVisibleLightInfo, SceneRenderingAllocator>& VisibleLightInfos,
+	const FVirtualShadowMapArray& VirtualShadowMapArray,
 	// Object data
 	const FPrimitiveSceneProxy* PrimitiveSceneProxy,
 	const FMaterialRenderProxy* MaterialRenderProxy,
@@ -780,6 +812,7 @@ void RenderWithLiveShading(
 			View,
 			// Shadow data
 			VisibleLightInfos,
+			VirtualShadowMapArray,
 			// Object data
 			PrimitiveSceneProxy,
 			MaterialRenderProxy,
@@ -800,6 +833,7 @@ void RenderWithLiveShading(
 			View,
 			// Shadow data
 			VisibleLightInfos,
+			VirtualShadowMapArray,
 			// Object data
 			PrimitiveSceneProxy,
 			MaterialRenderProxy,
diff --git a/Engine/Source/Runtime/Renderer/Private/HeterogeneousVolumes/HeterogeneousVolumesPreshadingPipeline.cpp b/Engine/Source/Runtime/Renderer/Private/HeterogeneousVolumes/HeterogeneousVolumesPreshadingPipeline.cpp
--- a/Engine/Source/Runtime/Renderer/Private/HeterogeneousVolumes/HeterogeneousVolumesPreshadingPipeline.cpp
+++ b/Engine/Source/Runtime/Renderer/Private/HeterogeneousVolumes/HeterogeneousVolumesPreshadingPipeline.cpp
@@ -223,6 +223,8 @@ class FRenderLightingCacheWithPreshadingCS : public FGlobalShader
 		// Shadow data
 		SHADER_PARAMETER_RDG_UNIFORM_BUFFER(FForwardLightData, ForwardLightData)
 		SHADER_PARAMETER_STRUCT_INCLUDE(FVolumeShadowingShaderParameters, VolumeShadowingShaderParameters)
+		SHADER_PARAMETER_STRUCT_INCLUDE(FVirtualShadowMapSamplingParameters, VirtualShadowMapSamplingParameters)
+		SHADER_PARAMETER(int32, VirtualShadowMapId)
 
 		// Volume structures
 		SHADER_PARAMETER_RDG_UNIFORM_BUFFER(FSparseVoxelUniformBufferParameters, SparseVoxelUniformBuffer)
@@ -256,6 +258,13 @@ class FRenderLightingCacheWithPreshadingCS : public FGlobalShader
 		OutEnvironment.SetDefine(TEXT("THREADGROUP_SIZE_2D"), GetThreadGroupSize2D());
 		OutEnvironment.SetDefine(TEXT("THREADGROUP_SIZE_3D"), GetThreadGroupSize3D());
 
+		bool bSupportVirtualShadowMap = IsFeatureLevelSupported(Parameters.Platform, ERHIFeatureLevel::SM5);
+		if (bSupportVirtualShadowMap)
+		{
+			OutEnvironment.SetDefine(TEXT("VIRTUAL_SHADOW_MAP"), 1);
+			FVirtualShadowMapArray::SetShaderDefines(OutEnvironment);
+		}
+
 		// This shader takes a very long time to compile with FXC, so we pre-compile it with DXC first and then forward the optimized HLSL to FXC.
 		//OutEnvironment.CompilerFlags.Add(CFLAG_PrecompileWithDXC);
 		OutEnvironment.CompilerFlags.Add(CFLAG_AllowTypedUAVLoads);
@@ -297,6 +306,8 @@ class FRenderSingleScatteringWithPreshadingCS : public FGlobalShader
 		// Shadow data
 		SHADER_PARAMETER_RDG_UNIFORM_BUFFER(FForwardLightData, ForwardLightData)
 		SHADER_PARAMETER_STRUCT_INCLUDE(FVolumeShadowingShaderParameters, VolumeShadowingShaderParameters)
+		SHADER_PARAMETER_STRUCT_INCLUDE(FVirtualShadowMapSamplingParameters, VirtualShadowMapSamplingParameters)
+		SHADER_PARAMETER(int32, VirtualShadowMapId)
 
 		// Volume data
 		SHADER_PARAMETER(int, MipLevel)
@@ -339,6 +350,13 @@ class FRenderSingleScatteringWithPreshadingCS : public FGlobalShader
 		OutEnvironment.SetDefine(TEXT("THREADGROUP_SIZE_1D"), GetThreadGroupSize1D());
 		OutEnvironment.SetDefine(TEXT("THREADGROUP_SIZE_2D"), GetThreadGroupSize2D());
 
+		bool bSupportVirtualShadowMap = IsFeatureLevelSupported(Parameters.Platform, ERHIFeatureLevel::SM5);
+		if (bSupportVirtualShadowMap)
+		{
+			OutEnvironment.SetDefine(TEXT("VIRTUAL_SHADOW_MAP"), 1);
+			FVirtualShadowMapArray::SetShaderDefines(OutEnvironment);
+		}
+
 		// This shader takes a very long time to compile with FXC, so we pre-compile it with DXC first and then forward the optimized HLSL to FXC.
 		//OutEnvironment.CompilerFlags.Add(CFLAG_PrecompileWithDXC);
 		OutEnvironment.CompilerFlags.Add(CFLAG_AllowTypedUAVLoads);
@@ -366,6 +384,7 @@ void RenderLightingCacheWithPreshadingCompute(
 	const FLightSceneInfo* LightSceneInfo,
 	// Shadow data
 	const FVisibleLightInfo* VisibleLightInfo,
+	const FVirtualShadowMapArray& VirtualShadowMapArray,
 	// Object data
 	const FPrimitiveSceneProxy* PrimitiveSceneProxy,
 	// Sparse voxel data
@@ -421,11 +440,14 @@ void RenderLightingCacheWithPreshadingCompute(
 			{
 				SetVolumeShadowingDefaultShaderParametersGlobal(GraphBuilder, PassParameters->VolumeShadowingShaderParameters);
 			}
+			PassParameters->VirtualShadowMapId = VisibleLightInfo->GetVirtualShadowMapId(&View);
 		}
 		else
 		{
 			SetVolumeShadowingDefaultShaderParametersGlobal(GraphBuilder, PassParameters->VolumeShadowingShaderParameters);
+			PassParameters->VirtualShadowMapId = -1;
 		}
+		PassParameters->VirtualShadowMapSamplingParameters = VirtualShadowMapArray.GetSamplingParameters(GraphBuilder);
 
 		// Output
 		PassParameters->RWLightingCacheTexture = GraphBuilder.CreateUAV(LightingCacheTexture);
@@ -476,6 +498,7 @@ void RenderSingleScatteringWithPreshadingCompute(
 	const FLightSceneInfo* LightSceneInfo,
 	// Shadow data
 	const FVisibleLightInfo* VisibleLightInfo,
+	const FVirtualShadowMapArray& VirtualShadowMapArray,
 	// Object data
 	const FPrimitiveSceneProxy* PrimitiveSceneProxy,
 	// Sparse voxel data
@@ -526,11 +549,14 @@ void RenderSingleScatteringWithPreshadingCompute(
 			{
 				SetVolumeShadowingDefaultShaderParametersGlobal(GraphBuilder, PassParameters->VolumeShadowingShaderParameters);
 			}
+			PassParameters->VirtualShadowMapId = VisibleLightInfo->GetVirtualShadowMapId(&View);
 		}
 		else
 		{
 			SetVolumeShadowingDefaultShaderParametersGlobal(GraphBuilder, PassParameters->VolumeShadowingShaderParameters);
+			PassParameters->VirtualShadowMapId = -1;
 		}
+		PassParameters->VirtualShadowMapSamplingParameters = VirtualShadowMapArray.GetSamplingParameters(GraphBuilder);
 
 		// Volume data
 		PassParameters->MipLevel = HeterogeneousVolumes::GetMipLevel();
@@ -595,6 +621,7 @@ void RenderWithInscatteringVolumePipelineWithPreshadingCompute(
 	FViewInfo& View,
 	// Shadow data
 	TArray<FVisibleLightInfo, SceneRenderingAllocator>& VisibleLightInfos,
+	const FVirtualShadowMapArray& VirtualShadowMapArray,
 	// Object data
 	const FPrimitiveSceneProxy* PrimitiveSceneProxy,
 	const FMaterialRenderProxy* MaterialRenderProxy,
@@ -661,6 +688,7 @@ void RenderWithInscatteringVolumePipelineWithPreshadingCompute(
 			LightSceneInfo,
 			// Shadow data
 			VisibleLightInfo,
+			VirtualShadowMapArray,
 			// Object data
 			PrimitiveSceneProxy,
 			// Volume data
@@ -700,6 +728,7 @@ void RenderWithInscatteringVolumePipelineWithPreshadingCompute(
 			LightSceneInfo,
 			// Shadow
 			VisibleLightInfo,
+			VirtualShadowMapArray,
 			// Object
 			PrimitiveSceneProxy,
 			// Volume data
@@ -726,6 +755,7 @@ void RenderWithTransmittanceVolumePipelineWithPreshadingCompute(
 	FViewInfo& View,
 	// Shadow data
 	TArray<FVisibleLightInfo, SceneRenderingAllocator>& VisibleLightInfos,
+	const FVirtualShadowMapArray& VirtualShadowMapArray,
 	// Object data
 	const FPrimitiveSceneProxy* PrimitiveSceneProxy,
 	const FMaterialRenderProxy* MaterialRenderProxy,
@@ -794,6 +824,7 @@ void RenderWithTransmittanceVolumePipelineWithPreshadingCompute(
 				LightSceneInfo,
 				// Shadow data
 				VisibleLightInfo,
+				VirtualShadowMapArray,
 				// Object data
 				PrimitiveSceneProxy,
 				// Volume data
@@ -823,6 +854,7 @@ void RenderWithTransmittanceVolumePipelineWithPreshadingCompute(
 			LightSceneInfo,
 			// Shadow
 			VisibleLightInfo,
+			VirtualShadowMapArray,
 			// Object
 			PrimitiveSceneProxy,
 			// Volume data
@@ -849,6 +881,7 @@ void RenderWithPreshadingCompute(
 	FViewInfo& View,
 	// Shadow data
 	TArray<FVisibleLightInfo, SceneRenderingAllocator>& VisibleLightInfos,
+	const FVirtualShadowMapArray& VirtualShadowMapArray,
 	// Object data
 	const FPrimitiveSceneProxy* PrimitiveSceneProxy,
 	const FMaterialRenderProxy* MaterialRenderProxy,
@@ -897,6 +930,7 @@ void RenderWithPreshadingCompute(
 			View,
 			// Shadow data
 			VisibleLightInfos,
+			VirtualShadowMapArray,
 			// Object data
 			PrimitiveSceneProxy,
 			MaterialRenderProxy,
@@ -923,6 +957,7 @@ void RenderWithPreshadingCompute(
 			View,
 			// Shadow data
 			VisibleLightInfos,
+			VirtualShadowMapArray,
 			// Object data
 			PrimitiveSceneProxy,
 			MaterialRenderProxy,
@@ -949,6 +984,7 @@ void RenderWithInscatteringVolumePipelineWithPreshadingHardwareRayTracing(
 	FViewInfo& View,
 	// Shadow data
 	TArray<FVisibleLightInfo, SceneRenderingAllocator>& VisibleLightInfos,
+	const FVirtualShadowMapArray& VirtualShadowMapArray,
 	// Object data
 	const FPrimitiveSceneProxy* PrimitiveSceneProxy,
 	const FMaterialRenderProxy* MaterialRenderProxy,
@@ -1013,6 +1049,7 @@ void RenderWithInscatteringVolumePipelineWithPreshadingHardwareRayTracing(
 			LightSceneInfo,
 			// Shadow
 			VisibleLightInfo,
+			VirtualShadowMapArray,
 			// Object data
 			PrimitiveSceneProxy,
 			// Sparse voxel
@@ -1048,6 +1085,7 @@ void RenderWithInscatteringVolumePipelineWithPreshadingHardwareRayTracing(
 			LightSceneInfo,
 			// Shadow
 			VisibleLightInfo,
+			VirtualShadowMapArray,
 			// Object data
 			PrimitiveSceneProxy,
 			// Sparse voxel
@@ -1072,6 +1110,7 @@ void RenderWithTransmittanceVolumePipelineWithPreshadingHardwareRayTracing(
 	FViewInfo& View,
 	// Shadow data
 	TArray<FVisibleLightInfo, SceneRenderingAllocator>& VisibleLightInfos,
+	const FVirtualShadowMapArray& VirtualShadowMapArray,
 	// Object data
 	const FPrimitiveSceneProxy* PrimitiveSceneProxy,
 	const FMaterialRenderProxy* MaterialRenderProxy,
@@ -1138,6 +1177,7 @@ void RenderWithTransmittanceVolumePipelineWithPreshadingHardwareRayTracing(
 				LightSceneInfo,
 				// Shadow
 				VisibleLightInfo,
+				VirtualShadowMapArray,
 				// Object data
 				PrimitiveSceneProxy,
 				// Sparse voxel
@@ -1163,6 +1203,7 @@ void RenderWithTransmittanceVolumePipelineWithPreshadingHardwareRayTracing(
 			LightSceneInfo,
 			// Shadow
 			VisibleLightInfo,
+			VirtualShadowMapArray,
 			// Object data
 			PrimitiveSceneProxy,
 			// Sparse voxel
@@ -1187,6 +1228,7 @@ void RenderWithPreshadingHardwareRayTracing(
 	FViewInfo& View,
 	// Shadow data
 	TArray<FVisibleLightInfo, SceneRenderingAllocator>& VisibleLightInfos,
+	const FVirtualShadowMapArray& VirtualShadowMapArray,
 	// Object data
 	const FPrimitiveSceneProxy* PrimitiveSceneProxy,
 	const FMaterialRenderProxy* MaterialRenderProxy,
@@ -1248,6 +1290,7 @@ void RenderWithPreshadingHardwareRayTracing(
 			View,
 			// Shadow data
 			VisibleLightInfos,
+			VirtualShadowMapArray,
 			// Object data
 			PrimitiveSceneProxy,
 			MaterialRenderProxy,
@@ -1269,6 +1312,7 @@ void RenderWithPreshadingHardwareRayTracing(
 			View,
 			// Shadow data
 			VisibleLightInfos,
+			VirtualShadowMapArray,
 			// Object data
 			PrimitiveSceneProxy,
 			MaterialRenderProxy,
@@ -1412,6 +1456,7 @@ void RenderWithPreshading(
 	FViewInfo& View,
 	// Shadow data
 	TArray<FVisibleLightInfo, SceneRenderingAllocator>& VisibleLightInfos,
+	const FVirtualShadowMapArray& VirtualShadowMapArray,
 	// Object data
 	const FPrimitiveSceneProxy* PrimitiveSceneProxy,
 	const FMaterialRenderProxy* MaterialRenderProxy,
@@ -1554,6 +1599,7 @@ void RenderWithPreshading(
 			View,
 			// Shadow data
 			VisibleLightInfos,
+			VirtualShadowMapArray,
 			// Object data
 			PrimitiveSceneProxy,
 			MaterialRenderProxy,
@@ -1578,6 +1624,7 @@ void RenderWithPreshading(
 			View,
 			// Shadow data
 			VisibleLightInfos,
+			VirtualShadowMapArray,
 			// Object data
 			PrimitiveSceneProxy,
 			MaterialRenderProxy,
