Fix VSM when r.Nanite.ProjectEnabled=0

diff --git a/Engine/Source/Runtime/RenderCore/Public/RenderUtils.h b/Engine/Source/Runtime/RenderCore/Public/RenderUtils.h
--- a/Engine/Source/Runtime/RenderCore/Public/RenderUtils.h
+++ b/Engine/Source/Runtime/RenderCore/Public/RenderUtils.h
@@ -774,6 +774,15 @@ inline bool UseVirtualShadowMaps(EShaderPlatform ShaderPlatform, const FStaticFe
 	return bVirtualShadowMapsEnabled && DoesRuntimeSupportNanite(ShaderPlatform, true /* check for atomics */, false /* check project setting */);
 }
 
+/**
+* Returns true if Virtual Shadow Mapsare supported for the given shader platform.
+* Note: Virtual Shadow Maps require Nanite platform support.
+*/
+inline bool DoesPlatformSupportVirtualShadowMaps(EShaderPlatform Platform)
+{
+	return DoesPlatformSupportNanite(Platform, false /* check project setting */);
+}
+
 /**
 * Returns true if non-Nanite virtual shadow maps are enabled by CVar r.Shadow.Virtual.NonNaniteVSM
 * and the runtime supports Nanite/virtual shadow maps.
diff --git a/Engine/Source/Runtime/Renderer/Private/VirtualShadowMaps/VirtualShadowMapArray.cpp b/Engine/Source/Runtime/Renderer/Private/VirtualShadowMaps/VirtualShadowMapArray.cpp
--- a/Engine/Source/Runtime/Renderer/Private/VirtualShadowMaps/VirtualShadowMapArray.cpp
+++ b/Engine/Source/Runtime/Renderer/Private/VirtualShadowMaps/VirtualShadowMapArray.cpp
@@ -604,7 +604,7 @@ public:
 	static bool ShouldCompilePermutation(const FGlobalShaderPermutationParameters& Parameters)
 	{
 		return IsFeatureLevelSupported(Parameters.Platform, ERHIFeatureLevel::SM5) &&
-			DoesPlatformSupportNanite(Parameters.Platform);
+			DoesPlatformSupportVirtualShadowMaps(Parameters.Platform);
 	}
 
 	/**
@@ -1983,7 +1983,7 @@ public:
 	static bool ShouldCompilePermutation(const FGlobalShaderPermutationParameters& Parameters)
 	{
 		return IsFeatureLevelSupported(Parameters.Platform, ERHIFeatureLevel::SM5) &&
-			DoesPlatformSupportNanite(Parameters.Platform);
+			DoesPlatformSupportVirtualShadowMaps(Parameters.Platform);
 	}
 
 	/**
@@ -2073,7 +2073,7 @@ public:
 	static bool ShouldCompilePermutation(const FGlobalShaderPermutationParameters& Parameters)
 	{
 		return IsFeatureLevelSupported(Parameters.Platform, ERHIFeatureLevel::SM5) &&
-			DoesPlatformSupportNanite(Parameters.Platform);
+			DoesPlatformSupportVirtualShadowMaps(Parameters.Platform);
 	}
 
 	/**
@@ -2111,7 +2111,7 @@ public:
 	static bool ShouldCompilePermutation(const FGlobalShaderPermutationParameters& Parameters)
 	{
 		return IsFeatureLevelSupported(Parameters.Platform, ERHIFeatureLevel::SM5) &&
-			DoesPlatformSupportNanite(Parameters.Platform);
+			DoesPlatformSupportVirtualShadowMaps(Parameters.Platform);
 	}
 
 	/**
diff --git a/Engine/Source/Runtime/Renderer/Private/VirtualShadowMaps/VirtualShadowMapCacheManager.cpp b/Engine/Source/Runtime/Renderer/Private/VirtualShadowMaps/VirtualShadowMapCacheManager.cpp
--- a/Engine/Source/Runtime/Renderer/Private/VirtualShadowMaps/VirtualShadowMapCacheManager.cpp
+++ b/Engine/Source/Runtime/Renderer/Private/VirtualShadowMaps/VirtualShadowMapCacheManager.cpp
@@ -679,7 +679,7 @@ public:
 	static bool ShouldCompilePermutation(const FGlobalShaderPermutationParameters& Parameters)
 	{
 		return IsFeatureLevelSupported(Parameters.Platform, ERHIFeatureLevel::SM5) &&
-			DoesPlatformSupportNanite(Parameters.Platform);
+			DoesPlatformSupportVirtualShadowMaps(Parameters.Platform);
 	}
 
 	static void ModifyCompilationEnvironment(const FGlobalShaderPermutationParameters& Parameters, FShaderCompilerEnvironment& OutEnvironment)
@@ -1062,7 +1062,7 @@ public:
 	static bool ShouldCompilePermutation(const FGlobalShaderPermutationParameters& Parameters)
 	{
 		return IsFeatureLevelSupported(Parameters.Platform, ERHIFeatureLevel::SM5) &&
-			DoesPlatformSupportNanite(Parameters.Platform);
+			DoesPlatformSupportVirtualShadowMaps(Parameters.Platform);
 	}
 
 	static void ModifyCompilationEnvironment(const FGlobalShaderPermutationParameters& Parameters, FShaderCompilerEnvironment& OutEnvironment)
diff --git a/Engine/Source/Runtime/Renderer/Private/VirtualShadowMaps/VirtualShadowMapProjection.cpp b/Engine/Source/Runtime/Renderer/Private/VirtualShadowMaps/VirtualShadowMapProjection.cpp
--- a/Engine/Source/Runtime/Renderer/Private/VirtualShadowMaps/VirtualShadowMapProjection.cpp
+++ b/Engine/Source/Runtime/Renderer/Private/VirtualShadowMaps/VirtualShadowMapProjection.cpp
@@ -270,7 +270,7 @@ class FVirtualShadowMapProjectionCS : public FGlobalShader
 			return false;
 		}
 		
-		return DoesPlatformSupportNanite(Parameters.Platform);
+		return DoesPlatformSupportVirtualShadowMaps(Parameters.Platform);
 	}
 };
 IMPLEMENT_GLOBAL_SHADER(FVirtualShadowMapProjectionCS, "/Engine/Private/VirtualShadowMaps/VirtualShadowMapProjection.usf", "VirtualShadowMapProjection", SF_Compute);
@@ -540,7 +540,7 @@ class FVirtualShadowMapProjectionCompositePS : public FGlobalShader
 
 	static bool ShouldCompilePermutation(const FGlobalShaderPermutationParameters& Parameters)
 	{
-		return DoesPlatformSupportNanite(Parameters.Platform);
+		return DoesPlatformSupportVirtualShadowMaps(Parameters.Platform);
 	}
 
 	static void ModifyCompilationEnvironment(const FGlobalShaderPermutationParameters& Parameters, FShaderCompilerEnvironment& OutEnvironment)
@@ -603,7 +603,7 @@ class FVirtualShadowMapProjectionCompositeFromMaskBitsPS : public FGlobalShader
 
 	static bool ShouldCompilePermutation(const FGlobalShaderPermutationParameters& Parameters)
 	{
-		return DoesPlatformSupportNanite(Parameters.Platform);
+		return DoesPlatformSupportVirtualShadowMaps(Parameters.Platform);
 	}
 
 	static void ModifyCompilationEnvironment(const FGlobalShaderPermutationParameters& Parameters, FShaderCompilerEnvironment& OutEnvironment)
