Fix crash/assert in cached mesh draw commands

diff --git a/Engine/Source/Runtime/Renderer/Public/MeshPassProcessor.h b/Engine/Source/Runtime/Renderer/Public/MeshPassProcessor.h
--- a/Engine/Source/Runtime/Renderer/Public/MeshPassProcessor.h
+++ b/Engine/Source/Runtime/Renderer/Public/MeshPassProcessor.h
@@ -146,7 +146,7 @@ public:
 	EMeshPass::Type SkipEmpty(EMeshPass::Type Pass) const 
 	{
 		uint64 Mask = 0xFFffFFffFFffFFffULL << Pass;
-		return EMeshPass::Type(FMath::Min<uint64>(EMeshPass::Num, FMath::CountTrailingZeros(Data & Mask)));
+		return EMeshPass::Type(FMath::Min<uint64>(EMeshPass::Num, FMath::CountTrailingZeros64(Data & Mask)));
 	}
 
 	int GetNum() 
