Work around O0 crash in DXC when compiling FPathTracingRG

diff --git a/Engine/Source/Runtime/Renderer/Private/PathTracing.cpp b/Engine/Source/Runtime/Renderer/Private/PathTracing.cpp
--- a/Engine/Source/Runtime/Renderer/Private/PathTracing.cpp
+++ b/Engine/Source/Runtime/Renderer/Private/PathTracing.cpp
@@ -705,6 +705,8 @@ class FPathTracingRG : public FGlobalShader
 		FGlobalShader::ModifyCompilationEnvironment(Parameters, OutEnvironment);
 		OutEnvironment.SetDefine(TEXT("USE_RECT_LIGHT_TEXTURES"), 1);
 		OutEnvironment.CompilerFlags.Add(CFLAG_WarningsAsErrors);
+		// @todo - Working around DXC compiler crash UE-165154, should be removed after DXC is updated with a fix.
+		OutEnvironment.CompilerFlags.Add(CFLAG_ForceOptimization);
 	}
 
 	BEGIN_SHADER_PARAMETER_STRUCT(FParameters, )
