Add local light data packing functions.

diff --git a/Engine/Source/Runtime/Renderer/Private/LightGridInjection.cpp b/Engine/Source/Runtime/Renderer/Private/LightGridInjection.cpp
--- a/Engine/Source/Runtime/Renderer/Private/LightGridInjection.cpp
+++ b/Engine/Source/Runtime/Renderer/Private/LightGridInjection.cpp
@@ -279,6 +279,66 @@ FVector GetLightGridZParams(float NearPlane, float FarPlane)
 	return FVector(B, O, S);
 }
 
+static void PackLocalLightData(
+	FForwardLocalLightData& Out,
+	const FViewInfo& View,
+	const FSimpleLightEntry& SimpleLight,
+	const FSimpleLightPerViewEntry& SimpleLightPerViewData)
+{
+	// Pack both values into a single float to keep float4 alignment
+	const FFloat16 SimpleLightSourceLength16f = FFloat16(0);
+
+	// Put simple lights in all lighting channels
+	FLightingChannels SimpleLightLightingChannels;
+	SimpleLightLightingChannels.bChannel0 = SimpleLightLightingChannels.bChannel1 = SimpleLightLightingChannels.bChannel2 = true;
+
+	const uint32 SimpleLightLightingChannelMask = GetLightingChannelMaskForStruct(SimpleLightLightingChannels);
+	const FVector3f LightTranslatedWorldPosition(View.ViewMatrices.GetPreViewTranslation() + SimpleLightPerViewData.Position);
+
+	// No shadowmap channels for simple lights
+	uint32 ShadowMapChannelMask = 0;
+	ShadowMapChannelMask |= SimpleLightLightingChannelMask << 8;
+
+	// Pack both values into a single float to keep float4 alignment
+	const FFloat16 VolumetricScatteringIntensity16f = FFloat16(SimpleLight.VolumetricScatteringIntensity);
+	const uint32 PackedWInt = ((uint32)SimpleLightSourceLength16f.Encoded) | ((uint32)VolumetricScatteringIntensity16f.Encoded << 16);
+
+	Out.LightPositionAndInvRadius							= FVector4f(LightTranslatedWorldPosition, 1.0f / FMath::Max(SimpleLight.Radius, KINDA_SMALL_NUMBER));
+	Out.LightColorAndFalloffExponent						= FVector4f((FVector3f)SimpleLight.Color, SimpleLight.Exponent);
+	Out.LightDirectionAndShadowMapChannelMask				= FVector4f(FVector3f(1, 0, 0), *((float*)&ShadowMapChannelMask));
+	Out.SpotAnglesAndSourceRadiusPacked						= FVector4f(-2, 1, 0, *(float*)&PackedWInt);
+	Out.LightTangentAndSoftSourceRadius						= FVector4f(1.0f, 0.0f, 0.0f, 0.0f);
+	Out.RectBarnDoorAndVirtualShadowMapIdAndSpecularScale	= FVector4f(0, -2, -1, 1);
+}
+
+static void PackLocalLightData(
+	FForwardLocalLightData& Out, 
+	const FViewInfo& View,
+	const FLightRenderParameters& LightParameters,
+	const uint32 LightTypeAndShadowMapChannelMaskPacked,
+	const int32 VirtualShadowMapId,
+	const float VolumetricScatteringIntensity)
+{
+	const FVector3f LightTranslatedWorldPosition(View.ViewMatrices.GetPreViewTranslation() + LightParameters.WorldPosition);
+
+	// Pack both values into a single float to keep float4 alignment
+	const FFloat16 SourceLength16f = FFloat16(LightParameters.SourceLength);
+	const FFloat16 VolumetricScatteringIntensity16f = FFloat16(VolumetricScatteringIntensity);
+	const uint32 PackedWInt = ((uint32)SourceLength16f.Encoded) | ((uint32)VolumetricScatteringIntensity16f.Encoded << 16);
+
+	// NOTE: This cast of VirtualShadowMapId to float is not ideal, but bitcast has issues here with INDEX_NONE -> NaN
+	// and 32-bit floats have enough mantissa to cover all reasonable numbers here for now.
+	Out.LightPositionAndInvRadius							= FVector4f(LightTranslatedWorldPosition, LightParameters.InvRadius);
+	Out.LightColorAndFalloffExponent						= FVector4f(LightParameters.Color, LightParameters.FalloffExponent);
+	Out.LightDirectionAndShadowMapChannelMask				= FVector4f(LightParameters.Direction, *((float*)&LightTypeAndShadowMapChannelMaskPacked));
+	Out.SpotAnglesAndSourceRadiusPacked						= FVector4f(LightParameters.SpotAngles.X, LightParameters.SpotAngles.Y, LightParameters.SourceRadius, *(float*)&PackedWInt);
+	Out.LightTangentAndSoftSourceRadius						= FVector4f(LightParameters.Tangent, LightParameters.SoftSourceRadius);
+	Out.RectBarnDoorAndVirtualShadowMapIdAndSpecularScale	= FVector4f(LightParameters.RectLightBarnCosAngle, LightParameters.RectLightBarnLength, float(VirtualShadowMapId), LightParameters.SpecularScale);
+
+	checkSlow(int32(Out.RectBarnDoorAndVirtualShadowMapIdAndSpecularScale.Z) == VirtualShadowMapId);
+}
+
+
 void FSceneRenderer::ComputeLightGrid(FRDGBuilder& GraphBuilder, bool bCullLightsToGrid, FSortedLightSetSceneInfo& SortedLightSet)
 {
 	RDG_CSV_STAT_EXCLUSIVE_SCOPE(GraphBuilder, ComputeLightGrid);
@@ -352,39 +412,24 @@ void FSceneRenderer::ComputeLightGrid(FRDGBuilder& GraphBuilder, bool bCullLight
 					check(SortedLightSet.SortedLights[SortedIndex].LightSceneInfo == nullptr);
 					check(!SortedLightSet.SortedLights[SortedIndex].SortKey.Fields.bIsNotSimpleLight);
 
+
 					int32 SimpleLightIndex = SortedLightSet.SortedLights[SortedIndex].SimpleLightIndex;
 
 					ForwardLocalLightData.AddUninitialized(1);
 					FForwardLocalLightData& LightData = ForwardLocalLightData.Last();
+
 					// Simple lights have no 'VisibleLight' info
 					LocalLightVisibleLightInfosIndex.Add(INDEX_NONE);
 
 					const FSimpleLightEntry& SimpleLight = SimpleLights.InstanceData[SimpleLightIndex];
 					const FSimpleLightPerViewEntry& SimpleLightPerViewData = SimpleLights.GetViewDependentData(SimpleLightIndex, ViewIndex, Views.Num());
+					PackLocalLightData(LightData, View, SimpleLight, SimpleLightPerViewData);
 
-					const FVector3f LightTranslatedWorldPosition(View.ViewMatrices.GetPreViewTranslation() + SimpleLightPerViewData.Position);
-					LightData.LightPositionAndInvRadius = FVector4f(LightTranslatedWorldPosition, 1.0f / FMath::Max(SimpleLight.Radius, KINDA_SMALL_NUMBER));
-					LightData.LightColorAndFalloffExponent = FVector4f((FVector3f)SimpleLight.Color, SimpleLight.Exponent);
-
-					// No shadowmap channels for simple lights
-					uint32 ShadowMapChannelMask = 0;
-					ShadowMapChannelMask |= SimpleLightLightingChannelMask << 8;
-
-					LightData.LightDirectionAndShadowMapChannelMask = FVector4f(FVector3f(1, 0, 0), *((float*)&ShadowMapChannelMask));
-
-					// Pack both values into a single float to keep float4 alignment
-					const FFloat16 VolumetricScatteringIntensity16f = FFloat16(SimpleLight.VolumetricScatteringIntensity);
-					const uint32 PackedWInt = ((uint32)SimpleLightSourceLength16f.Encoded) | ((uint32)VolumetricScatteringIntensity16f.Encoded << 16);
-
-					LightData.SpotAnglesAndSourceRadiusPacked = FVector4f(-2, 1, 0, *(float*)&PackedWInt);
-					LightData.LightTangentAndSoftSourceRadius = FVector4f(1.0f, 0.0f, 0.0f, 0.0f);
-					LightData.RectBarnDoorAndVirtualShadowMapIdAndSpecularScale = FVector4f(0, -2, -1, 1);
-
-#if ENABLE_LIGHT_CULLING_VIEW_SPACE_BUILD_DATA
+				#if ENABLE_LIGHT_CULLING_VIEW_SPACE_BUILD_DATA
 					FVector4f ViewSpacePosAndRadius(FVector4f(View.ViewMatrices.GetViewMatrix().TransformPosition(SimpleLightPerViewData.Position)), SimpleLight.Radius);
 					ViewSpacePosAndRadiusData.Add(ViewSpacePosAndRadius);
 					ViewSpaceDirAndPreprocAngleData.AddZeroed();
-#endif // ENABLE_LIGHT_CULLING_VIEW_SPACE_BUILD_DATA
+				#endif // ENABLE_LIGHT_CULLING_VIEW_SPACE_BUILD_DATA
 				}
 			}
 
@@ -466,25 +511,9 @@ void FSceneRenderer::ComputeLightGrid(FRDGBuilder& GraphBuilder, bool bCullLight
 						}
 						const float LightFade = GetLightFadeFactor(View, LightProxy);
 						LightParameters.Color *= LightFade;
-
 						LightParameters.Color *= LightParameters.GetLightExposureScale(Exposure);
 
-						const FVector3f LightTranslatedWorldPosition(View.ViewMatrices.GetPreViewTranslation() + LightParameters.WorldPosition);
-						LightData.LightPositionAndInvRadius = FVector4f(LightTranslatedWorldPosition, LightParameters.InvRadius);
-						LightData.LightColorAndFalloffExponent = FVector4f(LightParameters.Color, LightParameters.FalloffExponent);
-						LightData.LightDirectionAndShadowMapChannelMask = FVector4f(LightParameters.Direction, *((float*)&LightTypeAndShadowMapChannelMaskPacked));
-
-						LightData.SpotAnglesAndSourceRadiusPacked = FVector4f(LightParameters.SpotAngles.X, LightParameters.SpotAngles.Y, LightParameters.SourceRadius, 0);
-
-						LightData.LightTangentAndSoftSourceRadius = FVector4f(LightParameters.Tangent, LightParameters.SoftSourceRadius);
-
-						// NOTE: This cast of VirtualShadowMapId to float is not ideal, but bitcast has issues here with INDEX_NONE -> NaN
-						// and 32-bit floats have enough mantissa to cover all reasonable numbers here for now.
-						LightData.RectBarnDoorAndVirtualShadowMapIdAndSpecularScale = FVector4f(LightParameters.RectLightBarnCosAngle, LightParameters.RectLightBarnLength, float(VirtualShadowMapId), LightParameters.SpecularScale);
-						checkSlow(int32(LightData.RectBarnDoorAndVirtualShadowMapIdAndSpecularScale.Z) == VirtualShadowMapId);
-
 						float VolumetricScatteringIntensity = LightProxy->GetVolumetricScatteringIntensity();
-
 						if (LightNeedsSeparateInjectionIntoVolumetricFogForOpaqueShadow(View, LightSceneInfo, VisibleLightInfos[LightSceneInfo->Id])
 							|| (LightNeedsSeparateInjectionIntoVolumetricFogForLightFunction(LightSceneInfo) && CheckForLightFunction(LightSceneInfo)))
 						{
@@ -492,27 +521,22 @@ void FSceneRenderer::ComputeLightGrid(FRDGBuilder& GraphBuilder, bool bCullLight
 							VolumetricScatteringIntensity = 0;
 						}
 
-						// Pack both values into a single float to keep float4 alignment
-						const FFloat16 SourceLength16f = FFloat16(LightParameters.SourceLength);
-						const FFloat16 VolumetricScatteringIntensity16f = FFloat16(VolumetricScatteringIntensity);
-						const uint32 PackedWInt = ((uint32)SourceLength16f.Encoded) | ((uint32)VolumetricScatteringIntensity16f.Encoded << 16);
-						LightData.SpotAnglesAndSourceRadiusPacked.W = *(float*)&PackedWInt;
+						PackLocalLightData(LightData, View, LightParameters, LightTypeAndShadowMapChannelMaskPacked, VirtualShadowMapId, VolumetricScatteringIntensity);
 
 						const FSphere BoundingSphere = LightProxy->GetBoundingSphere();
 						const float Distance = View.ViewMatrices.GetViewMatrix().TransformPosition(BoundingSphere.Center).Z + BoundingSphere.W;
 						FurthestLight = FMath::Max(FurthestLight, Distance);
 
-#if ENABLE_LIGHT_CULLING_VIEW_SPACE_BUILD_DATA
+					#if ENABLE_LIGHT_CULLING_VIEW_SPACE_BUILD_DATA
 						// Note: inverting radius twice seems stupid (but done in shader anyway otherwise)
 						const FVector3f LightViewPosition = FVector4f(View.ViewMatrices.GetViewMatrix().TransformPosition(LightParameters.WorldPosition)); // LWC_TODO: precision loss
 						FVector4f ViewSpacePosAndRadius(LightViewPosition, 1.0f / LightParameters.InvRadius);
 						ViewSpacePosAndRadiusData.Add(ViewSpacePosAndRadius);
 
-						float PreProcAngle = SortedLightInfo.SortKey.Fields.LightType == LightType_Spot ? GetTanRadAngleOrZero(LightSceneInfo->Proxy->GetOuterConeAngle()) : 0.0f;
-
+						const float PreProcAngle = SortedLightInfo.SortKey.Fields.LightType == LightType_Spot ? GetTanRadAngleOrZero(LightSceneInfo->Proxy->GetOuterConeAngle()) : 0.0f;
 						FVector4f ViewSpaceDirAndPreprocAngle(FVector4f(View.ViewMatrices.GetViewMatrix().TransformVector((FVector)LightParameters.Direction)), PreProcAngle); // LWC_TODO: precision loss
 						ViewSpaceDirAndPreprocAngleData.Add(ViewSpaceDirAndPreprocAngle);
-#endif // ENABLE_LIGHT_CULLING_VIEW_SPACE_BUILD_DATA
+					#endif // ENABLE_LIGHT_CULLING_VIEW_SPACE_BUILD_DATA
 					}
 					else if (SortedLightInfo.SortKey.Fields.LightType == LightType_Directional && ViewFamily.EngineShowFlags.DirectionalLights)
 					{
