HeterogeneousVolumes: Respect "Volumetric Inscattering Intensity" parameter when computing direct lighting.

diff --git a/Engine/Shaders/Private/HeterogeneousVolumes/HeterogeneousVolumesHardwareRayTracing.usf b/Engine/Shaders/Private/HeterogeneousVolumes/HeterogeneousVolumesHardwareRayTracing.usf
--- a/Engine/Shaders/Private/HeterogeneousVolumes/HeterogeneousVolumesHardwareRayTracing.usf
+++ b/Engine/Shaders/Private/HeterogeneousVolumes/HeterogeneousVolumesHardwareRayTracing.usf
@@ -381,6 +381,7 @@ int bApplyEmissionAndTransmittance;
 int bApplyDirectLighting;
 int bApplyShadowTransmittance;
 int LightType;
+float VolumetricScatteringIntensity;
 
 // Shadowing
 uint VirtualShadowMapId;
@@ -444,7 +445,7 @@ RAY_TRACING_ENTRY_RAYGEN(RenderLightingCacheWithPreshadingRGS)
 		return;
 	}
 
-	const FDeferredLightData LightData = InitDeferredLightFromUniforms(LightType);
+	const FDeferredLightData LightData = InitDeferredLightFromUniforms(LightType, VolumetricScatteringIntensity);
 #if DIM_LIGHTING_CACHE_MODE == LIGHTING_CACHE_TRANSMITTANCE
 	float3 L = LightData.Direction;
 	float3 ToLight = L * GetMaxShadowTraceDistance();
@@ -551,7 +552,7 @@ RAY_TRACING_ENTRY_RAYGEN(RenderSingleScatteringWithPreshadingRGS)
 			DIM_APPLY_SHADOW_TRANSMITTANCE
 		);
 
-		const FDeferredLightData LightData = InitDeferredLightFromUniforms(LightType);
+		const FDeferredLightData LightData = InitDeferredLightFromUniforms(LightType, VolumetricScatteringIntensity);
 		uint StepCount = CalcStepCount(RayMarchingContext);
 
 		RayMarchSingleScattering(
diff --git a/Engine/Shaders/Private/HeterogeneousVolumes/HeterogeneousVolumesLightingUtils.ush b/Engine/Shaders/Private/HeterogeneousVolumes/HeterogeneousVolumesLightingUtils.ush
--- a/Engine/Shaders/Private/HeterogeneousVolumes/HeterogeneousVolumesLightingUtils.ush
+++ b/Engine/Shaders/Private/HeterogeneousVolumes/HeterogeneousVolumesLightingUtils.ush
@@ -4,3 +4,11 @@
 #include "../DeferredLightingCommon.ush"
 #include "../LightGridCommon.ush"
 #include "../LightDataUniforms.ush"
+
+FDeferredLightData InitDeferredLightFromUniforms(uint LightType, float VolumetricScatteringIntensity)
+{
+    FDeferredLightData DeferredLightData = InitDeferredLightFromUniforms(LightType);
+	DeferredLightData.Color *= VolumetricScatteringIntensity;
+
+    return DeferredLightData;
+}
diff --git a/Engine/Shaders/Private/HeterogeneousVolumes/HeterogeneousVolumesLiveShadingPipeline.usf b/Engine/Shaders/Private/HeterogeneousVolumes/HeterogeneousVolumesLiveShadingPipeline.usf
--- a/Engine/Shaders/Private/HeterogeneousVolumes/HeterogeneousVolumesLiveShadingPipeline.usf
+++ b/Engine/Shaders/Private/HeterogeneousVolumes/HeterogeneousVolumesLiveShadingPipeline.usf
@@ -44,6 +44,7 @@ int bApplyEmissionAndTransmittance;
 int bApplyDirectLighting;
 int bApplyShadowTransmittance;
 int LightType;
+float VolumetricScatteringIntensity;
 
 // Shadowing
 uint VirtualShadowMapId;
@@ -220,7 +221,7 @@ void RenderLightingCacheWithLiveShadingCS(
 		return;
 	}
 
-	const FDeferredLightData LightData = InitDeferredLightFromUniforms(LightType);
+	const FDeferredLightData LightData = InitDeferredLightFromUniforms(LightType, VolumetricScatteringIntensity);
 #if DIM_LIGHTING_CACHE_MODE == LIGHTING_CACHE_TRANSMITTANCE
 	float3 L = LightData.Direction;
 	float3 ToLight = L * GetMaxShadowTraceDistance();
@@ -320,7 +321,7 @@ void RenderSingleScatteringWithLiveShadingCS(
 			bApplyShadowTransmittance
 		);
 
-		const FDeferredLightData LightData = InitDeferredLightFromUniforms(LightType);
+		const FDeferredLightData LightData = InitDeferredLightFromUniforms(LightType, VolumetricScatteringIntensity);
 		uint StepCount = CalcStepCount(RayMarchingContext);
 
 		RayMarchSingleScattering(
diff --git a/Engine/Shaders/Private/HeterogeneousVolumes/HeterogeneousVolumesPreshadingPipeline.usf b/Engine/Shaders/Private/HeterogeneousVolumes/HeterogeneousVolumesPreshadingPipeline.usf
--- a/Engine/Shaders/Private/HeterogeneousVolumes/HeterogeneousVolumesPreshadingPipeline.usf
+++ b/Engine/Shaders/Private/HeterogeneousVolumes/HeterogeneousVolumesPreshadingPipeline.usf
@@ -48,6 +48,7 @@ int bApplyEmissionAndTransmittance;
 int bApplyDirectLighting;
 int bApplyShadowTransmittance;
 int LightType;
+float VolumetricScatteringIntensity;
 
 // Shadowing
 uint VirtualShadowMapId;
@@ -466,7 +467,7 @@ void RenderLightingCacheWithPreshadingCS(
 	}
 
 	// Build shadow ray
-	const FDeferredLightData LightData = InitDeferredLightFromUniforms(LightType);
+	const FDeferredLightData LightData = InitDeferredLightFromUniforms(LightType, VolumetricScatteringIntensity);
 #if DIM_LIGHTING_CACHE_MODE == LIGHTING_CACHE_TRANSMITTANCE
 	float3 L = LightData.Direction;
 	float3 ToLight = L * GetMaxShadowTraceDistance();
@@ -594,7 +595,7 @@ void RenderSingleScatteringWithPreshadingCS(
 			DIM_APPLY_SHADOW_TRANSMITTANCE
 		);
 
-		const FDeferredLightData LightData = InitDeferredLightFromUniforms(LightType);
+		const FDeferredLightData LightData = InitDeferredLightFromUniforms(LightType, VolumetricScatteringIntensity);
 		uint StepCount = CalcStepCount(RayMarchingContext);
 
 		RayMarchSingleScattering(
diff --git a/Engine/Source/Runtime/Renderer/Private/HeterogeneousVolumes/HeterogeneousVolumesHardwareRayTracing.cpp b/Engine/Source/Runtime/Renderer/Private/HeterogeneousVolumes/HeterogeneousVolumesHardwareRayTracing.cpp
--- a/Engine/Source/Runtime/Renderer/Private/HeterogeneousVolumes/HeterogeneousVolumesHardwareRayTracing.cpp
+++ b/Engine/Source/Runtime/Renderer/Private/HeterogeneousVolumes/HeterogeneousVolumesHardwareRayTracing.cpp
@@ -340,6 +340,7 @@ class FRenderLightingCacheWithPreshadingRGS : public FGlobalShader
 		SHADER_PARAMETER(int, bApplyShadowTransmittance)
 		SHADER_PARAMETER(int, LightType)
 		SHADER_PARAMETER_STRUCT_REF(FDeferredLightUniformStruct, DeferredLight)
+		SHADER_PARAMETER(float, VolumetricScatteringIntensity)
 
 		// Shadow data
 		SHADER_PARAMETER_RDG_UNIFORM_BUFFER(FForwardLightData, ForwardLightData)
@@ -411,6 +412,7 @@ class FRenderSingleScatteringWithPreshadingRGS : public FGlobalShader
 		SHADER_PARAMETER(int, bApplyDirectLighting)
 		SHADER_PARAMETER(int, LightType)
 		SHADER_PARAMETER_STRUCT_REF(FDeferredLightUniformStruct, DeferredLight)
+		SHADER_PARAMETER(float, VolumetricScatteringIntensity)
 
 		// Shadow data
 		SHADER_PARAMETER_RDG_UNIFORM_BUFFER(FForwardLightData, ForwardLightData)
@@ -578,6 +580,7 @@ void RenderLightingCacheWithPreshadingHardwareRayTracing(
 		FDeferredLightUniformStruct DeferredLightUniform = GetDeferredLightParameters(View, *LightSceneInfo);
 		PassParameters->DeferredLight = CreateUniformBufferImmediate(DeferredLightUniform, UniformBuffer_SingleDraw);
 		PassParameters->LightType = LightType;
+		PassParameters->VolumetricScatteringIntensity = LightSceneInfo->Proxy->GetVolumetricScatteringIntensity();
 
 		// Sparse Voxel data
 		PassParameters->SparseVoxelUniformBuffer = SparseVoxelUniformBuffer;
@@ -716,6 +719,7 @@ void RenderSingleScatteringWithPreshadingHardwareRayTracing(
 		if (PassParameters->bApplyDirectLighting && (LightSceneInfo != nullptr))
 		{
 			DeferredLightUniform = GetDeferredLightParameters(View, *LightSceneInfo);
+			PassParameters->VolumetricScatteringIntensity = LightSceneInfo->Proxy->GetVolumetricScatteringIntensity();
 		}
 		PassParameters->DeferredLight = CreateUniformBufferImmediate(DeferredLightUniform, UniformBuffer_SingleDraw);
 		PassParameters->LightType = LightType;
diff --git a/Engine/Source/Runtime/Renderer/Private/HeterogeneousVolumes/HeterogeneousVolumesLiveShadingPipeline.cpp b/Engine/Source/Runtime/Renderer/Private/HeterogeneousVolumes/HeterogeneousVolumesLiveShadingPipeline.cpp
--- a/Engine/Source/Runtime/Renderer/Private/HeterogeneousVolumes/HeterogeneousVolumesLiveShadingPipeline.cpp
+++ b/Engine/Source/Runtime/Renderer/Private/HeterogeneousVolumes/HeterogeneousVolumesLiveShadingPipeline.cpp
@@ -29,6 +29,7 @@ class FRenderLightingCacheWithLiveShadingCS : public FMeshMaterialShader
 		SHADER_PARAMETER(int, bApplyShadowTransmittance)
 		SHADER_PARAMETER(int, LightType)
 		SHADER_PARAMETER_STRUCT_REF(FDeferredLightUniformStruct, DeferredLight)
+		SHADER_PARAMETER(float, VolumetricScatteringIntensity)
 
 		// Shadow data
 		SHADER_PARAMETER(float, ShadowStepSize)
@@ -153,6 +154,7 @@ class FRenderSingleScatteringWithLiveShadingCS : public FMeshMaterialShader
 		SHADER_PARAMETER(int, bApplyShadowTransmittance)
 		SHADER_PARAMETER(int, LightType)
 		SHADER_PARAMETER_STRUCT_REF(FDeferredLightUniformStruct, DeferredLight)
+		SHADER_PARAMETER(float, VolumetricScatteringIntensity)
 
 		// Shadow data
 		SHADER_PARAMETER(float, ShadowStepSize)
@@ -351,6 +353,7 @@ void RenderLightingCacheWithLiveShading(
 		FDeferredLightUniformStruct DeferredLightUniform = GetDeferredLightParameters(View, *LightSceneInfo);
 		PassParameters->DeferredLight = CreateUniformBufferImmediate(DeferredLightUniform, UniformBuffer_SingleDraw);
 		PassParameters->LightType = LightType;
+		PassParameters->VolumetricScatteringIntensity = LightSceneInfo->Proxy->GetVolumetricScatteringIntensity();
 
 		// Object data
 		FMatrix44f LocalToWorld = FMatrix44f(PrimitiveSceneProxy->GetLocalToWorld());
@@ -483,6 +486,7 @@ void RenderSingleScatteringWithLiveShading(
 		if (PassParameters->bApplyDirectLighting && (LightSceneInfo != nullptr))
 		{
 			DeferredLightUniform = GetDeferredLightParameters(View, *LightSceneInfo);
+			PassParameters->VolumetricScatteringIntensity = LightSceneInfo->Proxy->GetVolumetricScatteringIntensity();
 		}
 		PassParameters->DeferredLight = CreateUniformBufferImmediate(DeferredLightUniform, UniformBuffer_SingleDraw);
 		PassParameters->LightType = LightType;
diff --git a/Engine/Source/Runtime/Renderer/Private/HeterogeneousVolumes/HeterogeneousVolumesPreshadingPipeline.cpp b/Engine/Source/Runtime/Renderer/Private/HeterogeneousVolumes/HeterogeneousVolumesPreshadingPipeline.cpp
--- a/Engine/Source/Runtime/Renderer/Private/HeterogeneousVolumes/HeterogeneousVolumesPreshadingPipeline.cpp
+++ b/Engine/Source/Runtime/Renderer/Private/HeterogeneousVolumes/HeterogeneousVolumesPreshadingPipeline.cpp
@@ -219,6 +219,7 @@ class FRenderLightingCacheWithPreshadingCS : public FGlobalShader
 		SHADER_PARAMETER(int, bApplyShadowTransmittance)
 		SHADER_PARAMETER(int, LightType)
 		SHADER_PARAMETER_STRUCT_REF(FDeferredLightUniformStruct, DeferredLight)
+		SHADER_PARAMETER(float, VolumetricScatteringIntensity)
 
 		// Shadow data
 		SHADER_PARAMETER_RDG_UNIFORM_BUFFER(FForwardLightData, ForwardLightData)
@@ -303,6 +304,7 @@ class FRenderSingleScatteringWithPreshadingCS : public FGlobalShader
 		SHADER_PARAMETER(int, bApplyDirectLighting)
 		SHADER_PARAMETER(int, LightType)
 		SHADER_PARAMETER_STRUCT_REF(FDeferredLightUniformStruct, DeferredLight)
+		SHADER_PARAMETER(float, VolumetricScatteringIntensity)
 
 		// Shadow data
 		SHADER_PARAMETER_RDG_UNIFORM_BUFFER(FForwardLightData, ForwardLightData)
@@ -416,6 +418,7 @@ void RenderLightingCacheWithPreshadingCompute(
 		FDeferredLightUniformStruct DeferredLightUniform = GetDeferredLightParameters(View, *LightSceneInfo);
 		PassParameters->DeferredLight = CreateUniformBufferImmediate(DeferredLightUniform, UniformBuffer_SingleDraw);
 		PassParameters->LightType = LightType;
+		PassParameters->VolumetricScatteringIntensity = LightSceneInfo->Proxy->GetVolumetricScatteringIntensity();
 
 		// Sparse voxel data
 		PassParameters->SparseVoxelUniformBuffer = SparseVoxelUniformBuffer;
@@ -535,6 +538,7 @@ void RenderSingleScatteringWithPreshadingCompute(
 		if (PassParameters->bApplyDirectLighting && (LightSceneInfo != nullptr))
 		{
 			DeferredLightUniform = GetDeferredLightParameters(View, *LightSceneInfo);
+			PassParameters->VolumetricScatteringIntensity = LightSceneInfo->Proxy->GetVolumetricScatteringIntensity();
 		}
 		PassParameters->DeferredLight = CreateUniformBufferImmediate(DeferredLightUniform, UniformBuffer_SingleDraw);
 		PassParameters->LightType = LightType;
