Fix volumetric cloud rendering in stereo (UE-156178).

diff --git a/Engine/Source/Runtime/Renderer/Private/VolumetricCloudRendering.cpp b/Engine/Source/Runtime/Renderer/Private/VolumetricCloudRendering.cpp
--- a/Engine/Source/Runtime/Renderer/Private/VolumetricCloudRendering.cpp
+++ b/Engine/Source/Runtime/Renderer/Private/VolumetricCloudRendering.cpp
@@ -1649,6 +1649,8 @@ void FSceneRenderer::InitVolumetricCloudsForViews(FRDGBuilder& GraphBuilder, boo
 
 				for (int32 ViewIndex = 0; ViewIndex < Views.Num(); ViewIndex++)
 				{
+					RDG_EVENT_SCOPE_CONDITIONAL(GraphBuilder, Views.Num() > 1, "View%d", ViewIndex);
+
 					FViewInfo& ViewInfo = Views[ViewIndex];
 					FVector3f ViewOrigin(ViewInfo.ViewMatrices.GetViewOrigin());
 
@@ -2257,6 +2259,8 @@ bool FSceneRenderer::RenderVolumetricCloud(
 
 			for (int32 ViewIndex = 0; ViewIndex < ViewsToProcess.Num(); ViewIndex++)
 			{
+				RDG_EVENT_SCOPE_CONDITIONAL(GraphBuilder, ViewsToProcess.Num() > 1, "View%d", ViewIndex);
+
 				FLocalCloudView& ViewToProcess = ViewsToProcess[ViewIndex];
 				FViewInfo& ViewInfo = *ViewToProcess.ViewInfo;
 
@@ -2417,7 +2421,9 @@ bool FSceneRenderer::RenderVolumetricCloud(
 					SkyRC.ViewMatrices = &ViewInfo.ViewMatrices;
 					SkyRC.ViewUniformBuffer = bShouldViewRenderVolumetricCloudRenderTarget ? ViewInfo.VolumetricRenderTargetViewUniformBuffer : ViewInfo.ViewUniformBuffer;
 
-					SkyRC.Viewport = ViewInfo.ViewRect;
+					SkyRC.Viewport = bShouldViewRenderVolumetricCloudRenderTarget ?
+						FIntRect(FIntPoint(0, 0), FIntPoint(DestinationRT->Desc.GetSize().X, DestinationRT->Desc.GetSize().Y)) // this texture is per view, so we don't need to use viewrect inside it
+						: ViewInfo.ViewRect;
 					SkyRC.bLightDiskEnabled = !ViewInfo.bIsReflectionCapture;
 					SkyRC.AerialPerspectiveStartDepthInCm = GetValidAerialPerspectiveStartDepthInCm(ViewInfo, SkyAtmosphereSceneProxy);
 					SkyRC.NearClippingDistance = ViewInfo.NearClippingDistance;
