Remove unused dummy rect light texture.

diff --git a/Engine/Source/Runtime/Renderer/Private/LightGridInjection.cpp b/Engine/Source/Runtime/Renderer/Private/LightGridInjection.cpp
--- a/Engine/Source/Runtime/Renderer/Private/LightGridInjection.cpp
+++ b/Engine/Source/Runtime/Renderer/Private/LightGridInjection.cpp
@@ -143,7 +143,6 @@ FForwardLightData::FForwardLightData()
 	ShadowmapSampler = TStaticSamplerState<SF_Point,AM_Clamp,AM_Clamp,AM_Clamp>::GetRHI();
 	DirectionalLightStaticShadowmap = GBlackTexture->TextureRHI;
 	StaticShadowmapSampler = TStaticSamplerState<SF_Bilinear,AM_Clamp,AM_Clamp,AM_Clamp>::GetRHI();
-	DummyRectLightSourceTexture = GWhiteTexture->TextureRHI;
 }
 
 int32 NumCulledLightsGridStride = 2;
@@ -737,9 +736,6 @@ void FSceneRenderer::ComputeLightGrid(FRDGBuilder& GraphBuilder, bool bCullLight
 
 		const int32 MaxNumCells = MaxLightGridSizeXY.X * MaxLightGridSizeXY.Y * GLightGridSizeZ * NumCulledGridPrimitiveTypes;
 
-		// Used to pass to the GetDynamicLighting but not actually used, since USE_SOURCE_TEXTURE is 0
-		ForwardLightData->DummyRectLightSourceTexture = GWhiteTexture->TextureRHI;
-
 		const FIntVector NumGroups = FIntVector::DivideAndRoundUp(FIntVector(LightGridSizeXY.X, LightGridSizeXY.Y, GLightGridSizeZ), LightGridInjectionGroupSize);
 
 		{
diff --git a/Engine/Source/Runtime/Renderer/Private/SceneRendering.h b/Engine/Source/Runtime/Renderer/Private/SceneRendering.h
--- a/Engine/Source/Runtime/Renderer/Private/SceneRendering.h
+++ b/Engine/Source/Runtime/Renderer/Private/SceneRendering.h
@@ -726,7 +726,6 @@ BEGIN_GLOBAL_SHADER_PARAMETER_STRUCT_WITH_CONSTRUCTOR(FForwardLightData, )
 	SHADER_PARAMETER_RDG_BUFFER_SRV(Buffer<float4>, ForwardLocalLightBuffer)
 	SHADER_PARAMETER_RDG_BUFFER_SRV(Buffer<uint>, NumCulledLightsGrid)
 	SHADER_PARAMETER_RDG_BUFFER_SRV(Buffer<uint>, CulledLightDataGrid)
-	SHADER_PARAMETER_TEXTURE(Texture2D, DummyRectLightSourceTexture)
 END_GLOBAL_SHADER_PARAMETER_STRUCT()
 
 extern TRDGUniformBufferRef<FForwardLightData> CreateDummyForwardLightUniformBuffer(FRDGBuilder& GraphBuilder);
