* Change r.RectLightAtlas.Translucent to be readonly
* Enforce shader change when the CVar is toggled

diff --git a/Engine/Source/Runtime/RenderCore/Private/Shader.cpp b/Engine/Source/Runtime/RenderCore/Private/Shader.cpp
--- a/Engine/Source/Runtime/RenderCore/Private/Shader.cpp
+++ b/Engine/Source/Runtime/RenderCore/Private/Shader.cpp
@@ -1816,6 +1816,14 @@ void ShaderMapAppendKeyString(EShaderPlatform Platform, FString& KeyString)
 		}
 	}
 
+	{
+		static const auto CVar = IConsoleManager::Get().FindTConsoleVariableDataInt(TEXT("r.RectLightAtlas.Translucent"));
+		if (CVar && CVar->GetValueOnAnyThread() > 0)
+		{
+			KeyString += TEXT("_RECTTRANS");
+		}
+	}
+
 	{
 		const bool bStrataEnabled = RenderCore_IsStrataEnabled();
 		if (bStrataEnabled)
diff --git a/Engine/Source/Runtime/Renderer/Private/RectLightTextureManager.cpp b/Engine/Source/Runtime/Renderer/Private/RectLightTextureManager.cpp
--- a/Engine/Source/Runtime/Renderer/Private/RectLightTextureManager.cpp
+++ b/Engine/Source/Runtime/Renderer/Private/RectLightTextureManager.cpp
@@ -54,8 +54,8 @@ static TAutoConsoleVariable<int32> CVarRectLighForceUpdate(
 static TAutoConsoleVariable<int32> CVarRectLighTranslucent(
 	TEXT("r.RectLightAtlas.Translucent"),
 	0,
-	TEXT("Enable rect light support for translucent surfaces."),
-	ECVF_RenderThreadSafe);
+	TEXT("Enable rect light support for translucent surfaces. When enabled, it will add an extrat sampler to the pixel shader (limited to 16 on dx11 based system)"),
+	ECVF_ReadOnly | ECVF_RenderThreadSafe);
 
 namespace RectLightAtlas
 {
