VRS GetVariableRateShadingImage returns nullptr rather than crashing w/ a pure virtual function call when passed a ViewFamily with no RenderTarget. VRS does not support non-standard view families/renderers.

diff --git a/Engine/Source/Runtime/Renderer/Private/VariableRateShadingImageManager.cpp b/Engine/Source/Runtime/Renderer/Private/VariableRateShadingImageManager.cpp
--- a/Engine/Source/Runtime/Renderer/Private/VariableRateShadingImageManager.cpp
+++ b/Engine/Source/Runtime/Renderer/Private/VariableRateShadingImageManager.cpp
@@ -178,6 +178,13 @@ FRDGTextureRef FVariableRateShadingImageManager::GetVariableRateShadingImage(FRD
 		return nullptr;
 	}
 
+	// Also bail if we're given a ViewFamily with no valid RenderTarget
+	if (ViewFamily.RenderTarget == nullptr)
+	{
+		ensureMsgf(0, TEXT("VRS Image Manager does not support ViewFamilies with no valid RenderTarget"));
+		return nullptr;
+	}
+
 	// Always want to make sure we tick every frame, even if we're not going to be generating any VRS images.
 	Tick();
 
