Changes r.TSR.History.SeparateTranslucency's default to 0

diff --git a/Engine/Source/Runtime/Renderer/Private/PostProcess/TemporalSuperResolution.cpp b/Engine/Source/Runtime/Renderer/Private/PostProcess/TemporalSuperResolution.cpp
--- a/Engine/Source/Runtime/Renderer/Private/PostProcess/TemporalSuperResolution.cpp
+++ b/Engine/Source/Runtime/Renderer/Private/PostProcess/TemporalSuperResolution.cpp
@@ -32,7 +32,7 @@ TAutoConsoleVariable<int32> CVarTSRHistoryUpdateQuality(
 	ECVF_Scalability | ECVF_RenderThreadSafe);
 
 TAutoConsoleVariable<int32> CVarTSRHistorySeparateTranslucency(
-	TEXT("r.TSR.History.SeparateTranslucency"), 1,
+	TEXT("r.TSR.History.SeparateTranslucency"), 0,
 	TEXT("Whether separate translucency should be accumulated separatly."),
 	ECVF_RenderThreadSafe);
 
