Dont set VSM resource for the VolumetricCloud when the compute shader dont sample the VSM. This remove the unecessary wait on the graphic pipe added by RDG.

diff --git a/Engine/Source/Runtime/Renderer/Private/VolumetricCloudRendering.cpp b/Engine/Source/Runtime/Renderer/Private/VolumetricCloudRendering.cpp
--- a/Engine/Source/Runtime/Renderer/Private/VolumetricCloudRendering.cpp
+++ b/Engine/Source/Runtime/Renderer/Private/VolumetricCloudRendering.cpp
@@ -2061,13 +2061,18 @@ void FSceneRenderer::RenderVolumetricCloudsInternal(FRDGBuilder& GraphBuilder, F
 		FRDGTextureRef CloudDepthTexture;
 		GetOutputTexturesWithFallback(GraphBuilder, CloudRC, CloudColorCubeTexture, CloudColorTexture, CloudDepthTexture);
 
+		bool bSampleVirtualShadowMap = (!bSkipAtmosphericLightShadowmap && CVarVolumetricCloudShadowSampleAtmosphericLightShadowmap.GetValueOnRenderThread() > 0);
+
 		FRenderVolumetricCloudRenderViewCS::FParameters* PassParameters = GraphBuilder.AllocParameters<FRenderVolumetricCloudRenderViewCS::FParameters>();
 		PassParameters->OutputViewRect = FVector4f(0.f, 0.f, Desc.Extent.X, Desc.Extent.Y);
 		PassParameters->bBlendCloudColor = !bShouldViewRenderVolumetricRenderTarget;
 		PassParameters->TargetCubeFace = CloudRC.RenderTargets.Output[0].GetArraySlice();
 		PassParameters->VolumetricCloudRenderViewParamsUB = CloudPassUniformBuffer;
 		PassParameters->View = ViewUniformBuffer;
-		PassParameters->VirtualShadowMap = VirtualShadowMapArray.GetSamplingParameters(GraphBuilder);
+		if (bSampleVirtualShadowMap)
+		{
+			PassParameters->VirtualShadowMap = VirtualShadowMapArray.GetSamplingParameters(GraphBuilder);
+		}
 		PassParameters->OutCloudColor = GraphBuilder.CreateUAV(CloudColorTexture);
 		PassParameters->OutCloudDepth = GraphBuilder.CreateUAV(CloudDepthTexture);
 		PassParameters->OutCloudColorCube = GraphBuilder.CreateUAV(CloudColorCubeTexture);
@@ -2078,7 +2083,7 @@ void FSceneRenderer::RenderVolumetricCloudsInternal(FRDGBuilder& GraphBuilder, F
 
 		typename FRenderVolumetricCloudRenderViewCS::FPermutationDomain PermutationVector;
 		PermutationVector.Set<typename FRenderVolumetricCloudRenderViewCS::FCloudPerSampleAtmosphereTransmittance>(ShouldUsePerSampleAtmosphereTransmittance(Scene, &MainView));
-		PermutationVector.Set<typename FRenderVolumetricCloudRenderViewCS::FCloudSampleAtmosphericLightShadowmap>(!bSkipAtmosphericLightShadowmap && CVarVolumetricCloudShadowSampleAtmosphericLightShadowmap.GetValueOnRenderThread() > 0);
+		PermutationVector.Set<typename FRenderVolumetricCloudRenderViewCS::FCloudSampleAtmosphericLightShadowmap>(bSampleVirtualShadowMap);
 		PermutationVector.Set<typename FRenderVolumetricCloudRenderViewCS::FCloudSampleSecondLight>(bSecondAtmosphereLightEnabled);
 		PermutationVector.Set<typename FRenderVolumetricCloudRenderViewCS::FCloudSampleLocalLights>(bCloudEnableLocalLightSampling && !CloudRC.bIsSkyRealTimeReflectionRendering);
 
