Fix light functions being rendered but not sampled with VSM one pass projection

diff --git a/Engine/Source/Runtime/Renderer/Private/LightRendering.cpp b/Engine/Source/Runtime/Renderer/Private/LightRendering.cpp
--- a/Engine/Source/Runtime/Renderer/Private/LightRendering.cpp
+++ b/Engine/Source/Runtime/Renderer/Private/LightRendering.cpp
@@ -1741,11 +1741,8 @@ void FDeferredShadingSceneRenderer::RenderLights(
 					{
 						const FViewInfo& View = Views[ViewIndex];
 
-						// If the light wrote into the packed shadow mask, sample directly from that
-						// TODO: Work this out higher in the logic and decide on shadow mask composite vs. kernel fusion
-						const bool bShadowMaskBitsValid = ShadowSceneRenderer->UsePackedShadowMaskBits();
-						const bool bWroteIntoPackedShadowMask = bShadowMaskBitsValid && !SortedLightInfo.SortKey.Fields.bDoesNotWriteIntoPackedShadowMask;
-						const int32 VirtualShadowMapId = bWroteIntoPackedShadowMask ? VisibleLightInfo.GetVirtualShadowMapId(&View) : INDEX_NONE;
+						// If the light elided the screen space shadow mask, sample directly from the packed shadow mask
+						const int32 VirtualShadowMapId = bElideScreenShadowMask ? VisibleLightInfo.GetVirtualShadowMapId(&View) : INDEX_NONE;
 
 						RDG_EVENT_SCOPE_CONDITIONAL(GraphBuilder, ViewCount > 1, "View%d", ViewIndex);
 						SCOPED_GPU_MASK(GraphBuilder.RHICmdList, View.GPUMask);
