SkyLight: Fix mGPU case in real time capture where 2nd GPU view is turned on.

The main issue was that the render target index was being changed between viewfamilies of different GPUs in the same frame.

There was a secondary issue related to the skycloud substeps, when SkyCloudCubeFacePerFrame is less than the default of 6 faces per frame, because the substep was being incremented between viewfamilies of different GPUs in the same frame.

We now keep track of which GPUs have a full cube map generated, to ensure that they all get it at least once since time slicing is enabled.

We also keep track of the GPUs processed each frame and avoid repeating the calculations for multiple view families in the same GPU.

diff --git a/Engine/Source/Runtime/Renderer/Private/ReflectionEnvironmentRealTimeCapture.cpp b/Engine/Source/Runtime/Renderer/Private/ReflectionEnvironmentRealTimeCapture.cpp
--- a/Engine/Source/Runtime/Renderer/Private/ReflectionEnvironmentRealTimeCapture.cpp
+++ b/Engine/Source/Runtime/Renderer/Private/ReflectionEnvironmentRealTimeCapture.cpp
@@ -305,8 +305,26 @@ void FScene::AllocateAndCaptureFrameSkyEnvMap(
 		return;
 	}
 
-	const bool bIsNewFrame = GFrameNumberRenderThread != RealTimeSlicedReflectionCaptureFrameNumber;
-	RealTimeSlicedReflectionCaptureFrameNumber = GFrameNumberRenderThread;
+	FRealTimeSlicedReflectionCapture& Capture = RealTimeSlicedReflectionCapture;
+
+	const bool bIsNewFrame = GFrameNumberRenderThread != Capture.FrameNumber;
+	Capture.FrameNumber = GFrameNumberRenderThread;
+
+	// Clear record of GPUs handled this frame if this is a new frame
+	if (bIsNewFrame)
+	{
+		Capture.GpusHandledThisFrame = 0;
+	}
+
+	// If this GPU has already been handled this frame, return, because we want to process the
+	// sky capture update for each RenderScene, but only once per GPU.
+	if ((Capture.GpusHandledThisFrame & MainView.GPUMask.GetNative()) == MainView.GPUMask.GetNative())
+	{
+		return;
+	}
+
+	// Record that we are handling the GPU in the MainView
+	Capture.GpusHandledThisFrame |= MainView.GPUMask.GetNative();
 
 	RDG_EVENT_SCOPE(GraphBuilder, "CaptureConvolveSkyEnvMap");
 	RDG_GPU_STAT_SCOPE(GraphBuilder, CaptureConvolveSkyEnvMap);
@@ -877,30 +895,36 @@ void FScene::AllocateAndCaptureFrameSkyEnvMap(
 		// Go to next state iff this is a new frame
 		if (bIsNewFrame)
 		{
-			switch (RealTimeSlicedReflectionCaptureFirstFrameState)
+			switch (Capture.FirstFrameState)
 			{
-				case ERealTimeSlicedReflectionCaptureFirstFrameState::INIT:
-					RealTimeSlicedReflectionCaptureFirstFrameState = ERealTimeSlicedReflectionCaptureFirstFrameState::FIRST_FRAME;
-					break;
+			case FRealTimeSlicedReflectionCapture::EFirstFrameState::INIT:
+				Capture.FirstFrameState = FRealTimeSlicedReflectionCapture::EFirstFrameState::FIRST_FRAME;
+				Capture.GpusWithFullCube = 0;
+				break;
 
-				case ERealTimeSlicedReflectionCaptureFirstFrameState::FIRST_FRAME:
-					RealTimeSlicedReflectionCaptureFirstFrameState = ERealTimeSlicedReflectionCaptureFirstFrameState::BEYOND_FIRST_FRAME;
-					break;
+			case FRealTimeSlicedReflectionCapture::EFirstFrameState::FIRST_FRAME:
+				Capture.FirstFrameState = FRealTimeSlicedReflectionCapture::EFirstFrameState::BEYOND_FIRST_FRAME;
+				break;
 
-				default:
-					break;
+			default:
+				break;
 			}
 		}
 	}
 	else
 	{
 		// Reset the time-slicing first frame detection state when not time-slicing.
-		RealTimeSlicedReflectionCaptureFirstFrameState = ERealTimeSlicedReflectionCaptureFirstFrameState::INIT;
+		Capture.FirstFrameState = FRealTimeSlicedReflectionCapture::EFirstFrameState::INIT;
 	}
 
+	const bool bGpuNeedsFullCube = Capture.GpusWithFullCube != (Capture.GpusWithFullCube | MainView.GPUMask.GetNative());
+
 	if (!bTimeSlicedRealTimeCapture 
-		|| (RealTimeSlicedReflectionCaptureFirstFrameState < ERealTimeSlicedReflectionCaptureFirstFrameState::BEYOND_FIRST_FRAME))
+		|| (Capture.FirstFrameState < FRealTimeSlicedReflectionCapture::EFirstFrameState::BEYOND_FIRST_FRAME)
+		|| bGpuNeedsFullCube)
 	{
+		Capture.GpusWithFullCube |= MainView.GPUMask.GetNative();
+
 		// Generate a full cube map in a single frame for the first frame.
 		// Perf number are for a 128x128x6 a cubemap on PS4 with sky and cloud and default settings
 
@@ -920,7 +944,8 @@ void FScene::AllocateAndCaptureFrameSkyEnvMap(
 		RenderCubeFaces_DiffuseIrradiance(ConvolvedSkyRenderTarget[ConvolvedSkyRenderTargetReadyIndex]);
 
 		// Reset Scene time slicing state so that it starts from the beginning if/when we get out of non-time-sliced.
-		RealTimeSlicedReflectionCaptureState = -1; // Value of -1 indicates this is the first time-sliced iteration.
+		Capture.State = -1; // Value of -1 indicates this is the first time-sliced iteration.
+		Capture.StateSubStep = 0;
 
 		// The sky just changed, so invalidate these textures, so that the path tracer can rebuild them
 		PathTracingSkylightTexture.SafeRelease();
@@ -933,7 +958,6 @@ void FScene::AllocateAndCaptureFrameSkyEnvMap(
 
 		// On the first frame, we always fully initialise the convolution so ConvolvedSkyRenderTargetReadyIndex should already be valid.
 		check(ConvolvedSkyRenderTargetReadyIndex >= 0 && ConvolvedSkyRenderTargetReadyIndex <= 1);
-		const int32 ConvolvedSkyRenderTargetWorkIndex = 1 - ConvolvedSkyRenderTargetReadyIndex;
 		const int32 TimeSliceCount = 12;
 
 #define DEBUG_TIME_SLICE 0
@@ -948,59 +972,76 @@ void FScene::AllocateAndCaptureFrameSkyEnvMap(
 			}
 #endif 
 
-		// Update the current time-slicing state if this is a new frame and if the current step is done.
-		// Note: RealTimeSlicedReflectionCaptureState will initially be -1.
-		if (bIsNewFrame && bRealTimeSlicedReflectionCaptureStateStepDone)
+		const int32 SkyCloudFrameStepCount = FMath::Clamp(CVarRealTimeReflectionCaptureTimeSlicingSkyCloudCubeFacePerFrame.GetValueOnRenderThread(), int32(1), int32(CubeFace_MAX));
+
+		// Because we want all GPUs to do the time slicing in lockstep, we only update the state when a new frame is starting
+		if (bIsNewFrame)
 		{
-			if (++RealTimeSlicedReflectionCaptureState >= TimeSliceCount)
+			int32 LastSkyCloudEndSubStep = FMath::Clamp(Capture.StateSubStep + SkyCloudFrameStepCount, int32(0), int32(CubeFace_MAX));
+
+			bool bStateFaceStepsDone = true;
+
+			if (Capture.State == 0 || Capture.State == 1)
+			{
+				bStateFaceStepsDone = LastSkyCloudEndSubStep >= CubeFace_MAX;
+				Capture.StateSubStep = bStateFaceStepsDone ? 0 : LastSkyCloudEndSubStep;
+			}
+
+			// Update the current time-slicing state if this is a new frame and if the current step is done.
+			// Note: RealTimeSlicedReflectionCaptureState will initially be -1.
+			if (bStateFaceStepsDone)
 			{
-				RealTimeSlicedReflectionCaptureState = 0;
-				RealTimeSlicedReflectionCaptureStateStep = 0;
+				if (++Capture.State >= TimeSliceCount)
+				{
+					// Now use the new cubemap
+					ConvolvedSkyRenderTargetReadyIndex = 1 - ConvolvedSkyRenderTargetReadyIndex;
+
+					// The sky just changed, so invalidate these textures, so that the path tracer can rebuild them
+					PathTracingSkylightTexture.SafeRelease();
+					PathTracingSkylightPdf.SafeRelease();
+
+					Capture.State = 0;
+					Capture.StateSubStep = 0;
+				}
 			}
 		}
-		bRealTimeSlicedReflectionCaptureStateStepDone = true;
 
-		const int32 SkyCloudFrameStepCount	= FMath::Clamp(CVarRealTimeReflectionCaptureTimeSlicingSkyCloudCubeFacePerFrame.GetValueOnRenderThread(), int32(1), int32(CubeFace_MAX));
-		const int32 SkyCloudStartSubStep	= FMath::Clamp(RealTimeSlicedReflectionCaptureStateStep, int32(0), int32(CubeFace_MAX - 1));
-		const int32 SkyCloudEndSubStep		= FMath::Clamp(RealTimeSlicedReflectionCaptureStateStep + SkyCloudFrameStepCount, int32(0), int32(CubeFace_MAX));
+		const int32 ConvolvedSkyRenderTargetWorkIndex = 1 - ConvolvedSkyRenderTargetReadyIndex;
+
+		const int32 SkyCloudStartSubStep = FMath::Clamp(Capture.StateSubStep, int32(0), int32(CubeFace_MAX - 1));
+		const int32 SkyCloudEndSubStep = FMath::Clamp(Capture.StateSubStep + SkyCloudFrameStepCount, int32(0), int32(CubeFace_MAX));
 
-		if (RealTimeSlicedReflectionCaptureState <= 0)
+		if (Capture.State <= 0)
 		{
 			RDG_EVENT_SCOPE(GraphBuilder, "RenderSky StartFace=%d EndFace=%d", SkyCloudStartSubStep, SkyCloudEndSubStep);
 			RenderCubeFaces_SkyCloud(true, false, CapturedSkyRenderTarget, SkyCloudStartSubStep, SkyCloudEndSubStep);
-
-			bRealTimeSlicedReflectionCaptureStateStepDone = SkyCloudEndSubStep >= CubeFace_MAX;
-			RealTimeSlicedReflectionCaptureStateStep = bRealTimeSlicedReflectionCaptureStateStepDone ? 0 : SkyCloudEndSubStep;
 		}
-		else if (RealTimeSlicedReflectionCaptureState == 1)
+		else if (Capture.State == 1)
 		{
 			RDG_EVENT_SCOPE(GraphBuilder, "RenderCloud StartFace=%d EndFace=%d", SkyCloudStartSubStep, SkyCloudEndSubStep);
 			RenderCubeFaces_SkyCloud(false, true, CapturedSkyRenderTarget, SkyCloudStartSubStep, SkyCloudEndSubStep);
-
-			bRealTimeSlicedReflectionCaptureStateStepDone = SkyCloudEndSubStep >= CubeFace_MAX;
-			RealTimeSlicedReflectionCaptureStateStep = bRealTimeSlicedReflectionCaptureStateStepDone ? 0 : SkyCloudEndSubStep;
 		}
-		else if (RealTimeSlicedReflectionCaptureState == 2)
+		else if (Capture.State == 2)
 		{
 			RDG_EVENT_SCOPE(GraphBuilder, "GenCubeMips");
 			RenderCubeFaces_GenCubeMips(1, LastMipLevel, CapturedSkyRenderTarget);
 		}
-		else if (RealTimeSlicedReflectionCaptureState == 3)
+		else if (Capture.State == 3)
 		{
 			RDG_EVENT_SCOPE(GraphBuilder, "ConvolutionMip0Face01");
 			RenderCubeFaces_SpecularConvolution(0, 0, 0, 2, ConvolvedSkyRenderTarget[ConvolvedSkyRenderTargetWorkIndex], CapturedSkyRenderTarget); // convolution of mip0, face 0, 1
 		}
-		else if (RealTimeSlicedReflectionCaptureState == 4)
+		else if (Capture.State == 4)
 		{
 			RDG_EVENT_SCOPE(GraphBuilder, "ConvolutionMip0Face23");
 			RenderCubeFaces_SpecularConvolution(0, 0, 2, 2, ConvolvedSkyRenderTarget[ConvolvedSkyRenderTargetWorkIndex], CapturedSkyRenderTarget); // convolution of mip0, face 2, 3
 		}
-		else if (RealTimeSlicedReflectionCaptureState == 5)
+		else if (Capture.State == 5)
 		{
 			RDG_EVENT_SCOPE(GraphBuilder, "ConvolutionMip0Face45");
 			RenderCubeFaces_SpecularConvolution(0, 0, 4, 2, ConvolvedSkyRenderTarget[ConvolvedSkyRenderTargetWorkIndex], CapturedSkyRenderTarget); // convolution of mip0, face 4, 5
 		}
-		else if (RealTimeSlicedReflectionCaptureState == 6)
+		else if (Capture.State == 6)
 		{
 			if (LastMipLevel >= 1)
 			{
@@ -1008,7 +1049,7 @@ void FScene::AllocateAndCaptureFrameSkyEnvMap(
 				RenderCubeFaces_SpecularConvolution(1, 1, 0, 6, ConvolvedSkyRenderTarget[ConvolvedSkyRenderTargetWorkIndex], CapturedSkyRenderTarget);
 			}
 		}
-		else if (RealTimeSlicedReflectionCaptureState == 7)
+		else if (Capture.State == 7)
 		{
 			if (LastMipLevel >= 2)
 			{
@@ -1016,7 +1057,7 @@ void FScene::AllocateAndCaptureFrameSkyEnvMap(
 				RenderCubeFaces_SpecularConvolution(2, 2, 0, 6, ConvolvedSkyRenderTarget[ConvolvedSkyRenderTargetWorkIndex], CapturedSkyRenderTarget);
 			}
 		}
-		else if (RealTimeSlicedReflectionCaptureState == 8)
+		else if (Capture.State == 8)
 		{
 			if (LastMipLevel >= 3)
 			{
@@ -1024,7 +1065,7 @@ void FScene::AllocateAndCaptureFrameSkyEnvMap(
 				RenderCubeFaces_SpecularConvolution(3, 3, 0, 6, ConvolvedSkyRenderTarget[ConvolvedSkyRenderTargetWorkIndex], CapturedSkyRenderTarget);
 			}
 		}
-		else if (RealTimeSlicedReflectionCaptureState == 9)
+		else if (Capture.State == 9)
 		{
 			if (LastMipLevel >= 5)
 			{
@@ -1037,7 +1078,7 @@ void FScene::AllocateAndCaptureFrameSkyEnvMap(
 				RenderCubeFaces_SpecularConvolution(4, 4, 0, 6, ConvolvedSkyRenderTarget[ConvolvedSkyRenderTargetWorkIndex], CapturedSkyRenderTarget);
 			}
 		}
-		else if (RealTimeSlicedReflectionCaptureState == 10)
+		else if (Capture.State == 10)
 		{
 			if (LastMipLevel >= 6)
 			{
@@ -1045,19 +1086,12 @@ void FScene::AllocateAndCaptureFrameSkyEnvMap(
 				RenderCubeFaces_SpecularConvolution(6, LastMipLevel, 0, 6, ConvolvedSkyRenderTarget[ConvolvedSkyRenderTargetWorkIndex], CapturedSkyRenderTarget);
 			}
 		}
-		else if (RealTimeSlicedReflectionCaptureState == 11)
+		else if (Capture.State == 11)
 		{
 			RDG_EVENT_SCOPE(GraphBuilder, "DiffuseIrradiance");
 
 			// Update the sky irradiance SH buffer.
 			RenderCubeFaces_DiffuseIrradiance(ConvolvedSkyRenderTarget[ConvolvedSkyRenderTargetWorkIndex]);
-
-			// Now use the new cubemap
-			ConvolvedSkyRenderTargetReadyIndex = ConvolvedSkyRenderTargetWorkIndex;
-
-			// The sky just changed, so invalidate these textures, so that the path tracer can rebuild them
-			PathTracingSkylightTexture.SafeRelease();
-			PathTracingSkylightPdf.SafeRelease();
 		}
 
 #if DEBUG_TIME_SLICE
diff --git a/Engine/Source/Runtime/Renderer/Private/RendererScene.cpp b/Engine/Source/Runtime/Renderer/Private/RendererScene.cpp
--- a/Engine/Source/Runtime/Renderer/Private/RendererScene.cpp
+++ b/Engine/Source/Runtime/Renderer/Private/RendererScene.cpp
@@ -1833,11 +1833,6 @@ FScene::FScene(UWorld* InWorld, bool bInRequiresHitProxies, bool bInIsEditorScen
 #endif
 ,	SkyLight(NULL)
 ,	ConvolvedSkyRenderTargetReadyIndex(-1)
-,	RealTimeSlicedReflectionCaptureFirstFrameState(ERealTimeSlicedReflectionCaptureFirstFrameState::INIT)
-,	RealTimeSlicedReflectionCaptureState(-1)
-,	RealTimeSlicedReflectionCaptureStateStep(0)
-,	bRealTimeSlicedReflectionCaptureStateStepDone(true)
-,	RealTimeSlicedReflectionCaptureFrameNumber(0)
 ,   PathTracingSkylightColor(0, 0, 0, 0)
 ,	SimpleDirectionalLight(NULL)
 ,	ReflectionSceneData(InFeatureLevel)
diff --git a/Engine/Source/Runtime/Renderer/Private/ScenePrivate.h b/Engine/Source/Runtime/Renderer/Private/ScenePrivate.h
--- a/Engine/Source/Runtime/Renderer/Private/ScenePrivate.h
+++ b/Engine/Source/Runtime/Renderer/Private/ScenePrivate.h
@@ -2867,23 +2867,31 @@ public:
 	/** The index of the ConvolvedSkyRenderTarget to use when rendering meshes. -1 when not initialised. */
 	int32 ConvolvedSkyRenderTargetReadyIndex;
 
-	/** We always enforce a complete one the first frame even with time slicing for correct start up lighting.*/
-	enum class ERealTimeSlicedReflectionCaptureFirstFrameState
-	{
-		INIT = 0,
-		FIRST_FRAME = 1,
-		BEYOND_FIRST_FRAME = 2,
-	} RealTimeSlicedReflectionCaptureFirstFrameState;
-
-	/** The current progress of the real time reflection capture when time sliced. */
-	int32 RealTimeSlicedReflectionCaptureState;
-	/** The current progress of each sub step of a state capture of the real time reflection capture when time sliced. */
-	int32 RealTimeSlicedReflectionCaptureStateStep;
-	/** When a state is sub divided, StepDone tells us if we can move onto the next step*/
-	bool bRealTimeSlicedReflectionCaptureStateStepDone;
-
-	/** Cache of Frame Number, used to detect first viewfamily */
-	uint64 RealTimeSlicedReflectionCaptureFrameNumber;
+	struct FRealTimeSlicedReflectionCapture
+	{
+		/** We always enforce a complete one the first frame even with time slicing for correct start up lighting.*/
+		enum class EFirstFrameState
+		{
+			INIT = 0,
+			FIRST_FRAME = 1,
+			BEYOND_FIRST_FRAME = 2,
+		} FirstFrameState = EFirstFrameState::INIT;
+
+		/** The current progress of the real time reflection capture when time sliced. */
+		int32 State = -1;
+
+		/** The current progress of each sub step of a state capture of the real time reflection capture when time sliced. */
+		int32 StateSubStep = 0;
+
+		/** Keeps track of which GPUs have been initialized with a full cube map */
+		uint32 GpusWithFullCube = 0;
+
+		/** Keeps track of which GPUs calculations have been done in the frame */
+		uint32 GpusHandledThisFrame = 0;
+
+		/** Cache of Frame Number, used to detect first viewfamily */
+		uint64 FrameNumber = uint64(-1);
+	} RealTimeSlicedReflectionCapture;
 
 	/**
 	 * The path tracer uses its own representation of the skylight. These textures
