Remove some shader print debug shaders for platforms that shouldn't have them.

diff --git a/Engine/Source/Runtime/Renderer/Private/LightGridInjection.cpp b/Engine/Source/Runtime/Renderer/Private/LightGridInjection.cpp
--- a/Engine/Source/Runtime/Renderer/Private/LightGridInjection.cpp
+++ b/Engine/Source/Runtime/Renderer/Private/LightGridInjection.cpp
@@ -1021,7 +1021,7 @@ class FDebugLightGridPS : public FGlobalShader
 
 	static bool ShouldCompilePermutation(const FGlobalShaderPermutationParameters& Parameters)
 	{
-		return EnumHasAllFlags(Parameters.Flags, EShaderPermutationFlags::HasEditorOnlyData);
+		return EnumHasAllFlags(Parameters.Flags, EShaderPermutationFlags::HasEditorOnlyData) && ShaderPrint::IsSupported(Parameters.Platform);
 	}
 
 	static void ModifyCompilationEnvironment(const FGlobalShaderPermutationParameters& Parameters, FShaderCompilerEnvironment& OutEnvironment)
@@ -1039,29 +1039,32 @@ IMPLEMENT_GLOBAL_SHADER(FDebugLightGridPS, "/Engine/Private/LightGridInjection.u
 
 FScreenPassTexture AddVisualizeLightGridPass(FRDGBuilder& GraphBuilder, const FViewInfo& View, FScreenPassTexture& ScreenPassSceneColor, FRDGTextureRef SceneDepthTexture)
 {
-	RDG_EVENT_SCOPE(GraphBuilder, "VisualizeLightGrid");
+	if (ShaderPrint::IsSupported(View.Family->GetShaderPlatform()))
+	{
+		RDG_EVENT_SCOPE(GraphBuilder, "VisualizeLightGrid");
 
-	// Force ShaderPrint on.
-	ShaderPrint::SetEnabled(true);
+		// Force ShaderPrint on.
+		ShaderPrint::SetEnabled(true);
 
-	ShaderPrint::RequestSpaceForLines(128);
-	ShaderPrint::RequestSpaceForCharacters(128);
+		ShaderPrint::RequestSpaceForLines(128);
+		ShaderPrint::RequestSpaceForCharacters(128);
 
-	FDebugLightGridPS::FPermutationDomain PermutationVector;
-	TShaderMapRef<FDebugLightGridPS> PixelShader(View.ShaderMap, PermutationVector);
-	FDebugLightGridPS::FParameters* PassParameters = GraphBuilder.AllocParameters<FDebugLightGridPS::FParameters>();
-	PassParameters->ViewUniformBuffer = View.ViewUniformBuffer;
-	PassParameters->Forward = View.ForwardLightingResources.ForwardLightUniformBuffer;
-	ShaderPrint::SetParameters(GraphBuilder, View.ShaderPrintData, PassParameters->ShaderPrintParameters);
-	PassParameters->DepthTexture = SceneDepthTexture ? SceneDepthTexture : GSystemTextures.GetMaxFP16Depth(GraphBuilder);
-	PassParameters->MiniFontTexture = GetMiniFontTexture();
-	PassParameters->RenderTargets[0] = FRenderTargetBinding(ScreenPassSceneColor.Texture, ERenderTargetLoadAction::ELoad);
-	PassParameters->DebugMode = GForwardLightGridDebug;
+		FDebugLightGridPS::FPermutationDomain PermutationVector;
+		TShaderMapRef<FDebugLightGridPS> PixelShader(View.ShaderMap, PermutationVector);
+		FDebugLightGridPS::FParameters* PassParameters = GraphBuilder.AllocParameters<FDebugLightGridPS::FParameters>();
+		PassParameters->ViewUniformBuffer = View.ViewUniformBuffer;
+		PassParameters->Forward = View.ForwardLightingResources.ForwardLightUniformBuffer;
+		ShaderPrint::SetParameters(GraphBuilder, View.ShaderPrintData, PassParameters->ShaderPrintParameters);
+		PassParameters->DepthTexture = SceneDepthTexture ? SceneDepthTexture : GSystemTextures.GetMaxFP16Depth(GraphBuilder);
+		PassParameters->MiniFontTexture = GetMiniFontTexture();
+		PassParameters->RenderTargets[0] = FRenderTargetBinding(ScreenPassSceneColor.Texture, ERenderTargetLoadAction::ELoad);
+		PassParameters->DebugMode = GForwardLightGridDebug;
 
-	FRHIBlendState* PreMultipliedColorTransmittanceBlend = TStaticBlendState<CW_RGB, BO_Add, BF_One, BF_SourceAlpha, BO_Add, BF_Zero, BF_One>::GetRHI();
+		FRHIBlendState* PreMultipliedColorTransmittanceBlend = TStaticBlendState<CW_RGB, BO_Add, BF_One, BF_SourceAlpha, BO_Add, BF_Zero, BF_One>::GetRHI();
 
-	FPixelShaderUtils::AddFullscreenPass<FDebugLightGridPS>(GraphBuilder, View.ShaderMap, RDG_EVENT_NAME("DebugLightGridCS"), PixelShader, PassParameters,
-		ScreenPassSceneColor.ViewRect, PreMultipliedColorTransmittanceBlend);
+		FPixelShaderUtils::AddFullscreenPass<FDebugLightGridPS>(GraphBuilder, View.ShaderMap, RDG_EVENT_NAME("DebugLightGridCS"), PixelShader, PassParameters,
+			ScreenPassSceneColor.ViewRect, PreMultipliedColorTransmittanceBlend);
+	}
 
 	return MoveTemp(ScreenPassSceneColor);
 }
diff --git a/Engine/Source/Runtime/Renderer/Private/OIT/OIT.cpp b/Engine/Source/Runtime/Renderer/Private/OIT/OIT.cpp
--- a/Engine/Source/Runtime/Renderer/Private/OIT/OIT.cpp
+++ b/Engine/Source/Runtime/Renderer/Private/OIT/OIT.cpp
@@ -475,6 +475,17 @@ class FOITSortTriangleIndex_ScanCS : public FGlobalShader
 
 	END_SHADER_PARAMETER_STRUCT()
 public:
+
+	static bool ShouldCompilePermutation(const FGlobalShaderPermutationParameters& Parameters) 
+	{ 
+		FPermutationDomain PermutationVector(Parameters.PermutationId);
+		if (PermutationVector.Get<FDebug>() == 1 && !ShaderPrint::IsSupported(Parameters.Platform))
+		{
+			return false;
+		}
+		return true;
+	}
+
 	static void ModifyCompilationEnvironment(const FGlobalShaderPermutationParameters& Parameters, FShaderCompilerEnvironment& OutEnvironment)
 	{
 		FGlobalShader::ModifyCompilationEnvironment(Parameters, OutEnvironment);
@@ -558,6 +569,12 @@ class FOITSortTriangleIndex_Debug : public FGlobalShader
 		SHADER_PARAMETER_RDG_BUFFER_SRV(Buffer<uint>, DebugData)
 	END_SHADER_PARAMETER_STRUCT()
 public:
+
+	static bool ShouldCompilePermutation(const FGlobalShaderPermutationParameters& Parameters) 
+	{
+		return ShaderPrint::IsSupported(Parameters.Platform); 
+	}
+
 	static void ModifyCompilationEnvironment(const FGlobalShaderPermutationParameters& Parameters, FShaderCompilerEnvironment& OutEnvironment)
 	{
 		FGlobalShader::ModifyCompilationEnvironment(Parameters, OutEnvironment);
@@ -676,7 +693,7 @@ static void AddOITSortTriangleIndexPass(
 			Parameters->WorldBound_Max	= (FVector3f)Bounds.GetBox().Max;
 			Parameters->ViewBound_Min	= (FVector3f)ViewBounds.GetBox().Min;
 			Parameters->ViewBound_Max	= (FVector3f)ViewBounds.GetBox().Max;
-			ShaderPrint::SetParameters(GraphBuilder, View.ShaderPrintData, Parameters->ShaderPrintParameters);
+				ShaderPrint::SetParameters(GraphBuilder, View.ShaderPrintData, Parameters->ShaderPrintParameters);
 
 			check(DebugData.Buffer);
 
diff --git a/Engine/Source/Runtime/Renderer/Private/RectLightTextureManager.cpp b/Engine/Source/Runtime/Renderer/Private/RectLightTextureManager.cpp
--- a/Engine/Source/Runtime/Renderer/Private/RectLightTextureManager.cpp
+++ b/Engine/Source/Runtime/Renderer/Private/RectLightTextureManager.cpp
@@ -256,7 +256,7 @@ class FRectLightAtlasDebugInfoCS : public FGlobalShader
 	END_SHADER_PARAMETER_STRUCT()
 
 public:
-	static bool ShouldCompilePermutation(const FGlobalShaderPermutationParameters& Parameters) { return true; }
+	static bool ShouldCompilePermutation(const FGlobalShaderPermutationParameters& Parameters) { return ShaderPrint::IsSupported(Parameters.Platform); }
 	static void ModifyCompilationEnvironment(const FGlobalShaderPermutationParameters& Parameters, FShaderCompilerEnvironment& OutEnvironment)
 	{
 		FGlobalShader::ModifyCompilationEnvironment(Parameters, OutEnvironment);
@@ -1624,7 +1624,7 @@ void UpdateRectLightAtlasTexture(FRDGBuilder& GraphBuilder, const ERHIFeatureLev
 
 void AddRectLightAtlasDebugPass(FRDGBuilder& GraphBuilder, const FViewInfo& View, FRDGTextureRef OutputTexture)
 {
-	if (CVarRectLighTextureDebug.GetValueOnRenderThread() > 0)
+	if (CVarRectLighTextureDebug.GetValueOnRenderThread() > 0 && ShaderPrint::IsSupported(View.Family->GetShaderPlatform()))
 	{
 		if (FRDGTextureRef DebugOutput = AddRectLightDebugInfoPass(GraphBuilder, View, OutputTexture->Desc))
 		{
diff --git a/Engine/Source/Runtime/Renderer/Private/ShaderPrint.cpp b/Engine/Source/Runtime/Renderer/Private/ShaderPrint.cpp
--- a/Engine/Source/Runtime/Renderer/Private/ShaderPrint.cpp
+++ b/Engine/Source/Runtime/Renderer/Private/ShaderPrint.cpp
@@ -145,29 +145,33 @@ namespace ShaderPrint
 	IMPLEMENT_GLOBAL_SHADER_PARAMETER_STRUCT(FShaderPrintCommonParameters, "ShaderPrintData");
 	
 	// Fill the uniform buffer parameters
+	void GetParameters(FShaderPrintSetup const& InSetup, FShaderPrintCommonParameters& OutParameters)
+	{
+		const FVector2D ViewSize(FMath::Max(InSetup.ViewRect.Size().X, 1), FMath::Max(InSetup.ViewRect.Size().Y, 1));
+		const float FontWidth = float(InSetup.FontSize.X) * InSetup.DPIScale / ViewSize.X;
+		const float FontHeight = float(InSetup.FontSize.Y) * InSetup.DPIScale / ViewSize.Y;
+		const float SpaceWidth = float(InSetup.FontSpacing.X) * InSetup.DPIScale / ViewSize.X;
+		const float SpaceHeight = float(InSetup.FontSpacing.Y) * InSetup.DPIScale / ViewSize.Y;
+
+		OutParameters.FontSize = FVector2f(FontWidth, FontHeight);
+		OutParameters.FontSpacing = FVector2f(FontWidth + SpaceWidth, FontHeight + SpaceHeight);
+		OutParameters.Resolution = InSetup.ViewRect.Size();
+		OutParameters.CursorCoord = InSetup.CursorCoord;
+		OutParameters.MaxValueCount = InSetup.MaxValueCount;
+		OutParameters.MaxSymbolCount = GetMaxSymbolCountFromValueCount(InSetup.MaxValueCount);
+		OutParameters.MaxStateCount = InSetup.MaxStateCount;
+		OutParameters.MaxLineCount = InSetup.MaxLineCount;
+		OutParameters.MaxTriangleCount = InSetup.MaxTriangleCount;
+		OutParameters.TranslatedWorldOffset = FVector3f(InSetup.PreViewTranslation);
+	}
+
 	// Return a uniform buffer with values filled and with single frame lifetime
-	static TUniformBufferRef<ShaderPrint::FShaderPrintCommonParameters> CreateUniformBuffer(const FShaderPrintSetup& Setup)
+	static TUniformBufferRef<ShaderPrint::FShaderPrintCommonParameters> CreateUniformBuffer(const FShaderPrintSetup& InSetup)
 	{		
-		FShaderPrintCommonParameters Out;
+		FShaderPrintCommonParameters Parameters;
+		GetParameters(InSetup, Parameters);
 
-		const FVector2D ViewSize(FMath::Max(Setup.ViewRect.Size().X, 1), FMath::Max(Setup.ViewRect.Size().Y, 1));
-		const float FontWidth = float(Setup.FontSize.X) * Setup.DPIScale / ViewSize.X;
-		const float FontHeight = float(Setup.FontSize.Y) * Setup.DPIScale / ViewSize.Y;
-		const float SpaceWidth = float(Setup.FontSpacing.X) * Setup.DPIScale / ViewSize.X;
-		const float SpaceHeight = float(Setup.FontSpacing.Y) * Setup.DPIScale / ViewSize.Y;
-		
-		Out.FontSize = FVector2f(FontWidth, FontHeight);
-		Out.FontSpacing = FVector2f(FontWidth + SpaceWidth, FontHeight + SpaceHeight);
-		Out.Resolution = Setup.ViewRect.Size();
-		Out.CursorCoord = Setup.CursorCoord;
-		Out.MaxValueCount = Setup.MaxValueCount;
-		Out.MaxSymbolCount = GetMaxSymbolCountFromValueCount(Setup.MaxValueCount);
-		Out.MaxStateCount = Setup.MaxStateCount;
-		Out.MaxLineCount = Setup.MaxLineCount;
-		Out.MaxTriangleCount = Setup.MaxTriangleCount;
-		Out.TranslatedWorldOffset = FVector3f(Setup.PreViewTranslation);
-
-		return TUniformBufferRef<ShaderPrint::FShaderPrintCommonParameters>::CreateUniformBufferImmediate(Out, UniformBuffer_SingleFrame);
+		return TUniformBufferRef<ShaderPrint::FShaderPrintCommonParameters>::CreateUniformBufferImmediate(Parameters, UniformBuffer_SingleFrame);
 	}
 
 	//////////////////////////////////////////////////////////////////////////////////////////////////
@@ -196,7 +200,7 @@ namespace ShaderPrint
 
 	bool IsSupported(EShaderPlatform InShaderPlatform)
 	{
-		return !IsHlslccShaderPlatform(InShaderPlatform);
+		return !IsMobilePlatform(InShaderPlatform) && !IsHlslccShaderPlatform(InShaderPlatform);
 	}
 
 	bool IsEnabled()
@@ -563,7 +567,6 @@ namespace ShaderPrint
 			return ShaderPrintData;
 		}
 
-
 		// Characters/Widgets/Primitives/Lines
 		{
 			const bool bLockBufferThisFrame = IsDrawLocked() && InViewState != nullptr && !InViewState->ShaderPrintStateData.bIsLocked;
diff --git a/Engine/Source/Runtime/Renderer/Public/ShaderPrintParameters.h b/Engine/Source/Runtime/Renderer/Public/ShaderPrintParameters.h
--- a/Engine/Source/Runtime/Renderer/Public/ShaderPrintParameters.h
+++ b/Engine/Source/Runtime/Renderer/Public/ShaderPrintParameters.h
@@ -116,6 +116,9 @@ namespace ShaderPrint
 		uint32 MaxTriangleCount = 0;
 	};
 
+	/** Fill the ShaderPrintCommonParameters uniform buffer structure for our setup. */
+	RENDERER_API void GetParameters(FShaderPrintSetup const& InSetup, FShaderPrintCommonParameters& OutParameters);
+
 	/** Create the shader print render data. This allocates and clears the render buffers. */
 	RENDERER_API FShaderPrintData CreateShaderPrintData(FRDGBuilder& GraphBuilder, FShaderPrintSetup const& InSetup);
 	
