Fix integer overflow causing FVolumeBounds to be populated with extreme bounds leading to checks and potential GPU issues.

Fixed Forward lights FarPlane max distance to match UE4, this was increased due to the change to WORLD_MAX.

#rnx

diff --git a/Engine/Source/Runtime/Renderer/Private/LightGridInjection.cpp b/Engine/Source/Runtime/Renderer/Private/LightGridInjection.cpp
--- a/Engine/Source/Runtime/Renderer/Private/LightGridInjection.cpp
+++ b/Engine/Source/Runtime/Renderer/Private/LightGridInjection.cpp
@@ -693,7 +693,7 @@ void FSceneRenderer::ComputeLightGrid(FRDGBuilder& GraphBuilder, bool bCullLight
 
 		// Clamp far plane to something reasonable
 		const float KilometersToCentimeters = 100000.0f;
-		const float LightCullingMaxDistance = GLightCullingMaxDistanceOverrideKilometers <= 0.0f ? (float)HALF_WORLD_MAX / 5.0f : GLightCullingMaxDistanceOverrideKilometers * KilometersToCentimeters;
+		const float LightCullingMaxDistance = GLightCullingMaxDistanceOverrideKilometers <= 0.0f ? (float)UE_OLD_HALF_WORLD_MAX / 5.0f : GLightCullingMaxDistanceOverrideKilometers * KilometersToCentimeters;
 		float FarPlane = FMath::Min(FMath::Max(FurthestLight, View.FurthestReflectionCaptureDistance), LightCullingMaxDistance);
 		FVector ZParams = GetLightGridZParams(View.NearClippingDistance, FarPlane + 10.f);
 		ForwardLightData->LightGridZParams = (FVector3f)ZParams;
diff --git a/Engine/Source/Runtime/Renderer/Private/TranslucentLighting.cpp b/Engine/Source/Runtime/Renderer/Private/TranslucentLighting.cpp
--- a/Engine/Source/Runtime/Renderer/Private/TranslucentLighting.cpp
+++ b/Engine/Source/Runtime/Renderer/Private/TranslucentLighting.cpp
@@ -984,13 +984,13 @@ FVolumeBounds CalculateLightVolumeBounds(const FSphere& LightBounds, const FView
 		const FVector MinPosition = (LightBounds.Center - LightBounds.W - View.TranslucencyLightingVolumeMin[VolumeCascadeIndex]) / View.TranslucencyVolumeVoxelSize[VolumeCascadeIndex];
 		const FVector MaxPosition = (LightBounds.Center + LightBounds.W - View.TranslucencyLightingVolumeMin[VolumeCascadeIndex]) / View.TranslucencyVolumeVoxelSize[VolumeCascadeIndex];
 
-		VolumeBounds.MinX = FMath::Max(FMath::TruncToInt(MinPosition.X), 0);
-		VolumeBounds.MinY = FMath::Max(FMath::TruncToInt(MinPosition.Y), 0);
-		VolumeBounds.MinZ = FMath::Max(FMath::TruncToInt(MinPosition.Z), 0);
+		VolumeBounds.MinX = FMath::Max(FMath::TruncToInt32(MinPosition.X), 0);
+		VolumeBounds.MinY = FMath::Max(FMath::TruncToInt32(MinPosition.Y), 0);
+		VolumeBounds.MinZ = FMath::Max(FMath::TruncToInt32(MinPosition.Z), 0);
 
-		VolumeBounds.MaxX = FMath::Min(FMath::TruncToInt(MaxPosition.X) + 1, TranslucencyLightingVolumeDim);
-		VolumeBounds.MaxY = FMath::Min(FMath::TruncToInt(MaxPosition.Y) + 1, TranslucencyLightingVolumeDim);
-		VolumeBounds.MaxZ = FMath::Min(FMath::TruncToInt(MaxPosition.Z) + 1, TranslucencyLightingVolumeDim);
+		VolumeBounds.MaxX = FMath::Min(FMath::TruncToInt32(MaxPosition.X) + 1, TranslucencyLightingVolumeDim);
+		VolumeBounds.MaxY = FMath::Min(FMath::TruncToInt32(MaxPosition.Y) + 1, TranslucencyLightingVolumeDim);
+		VolumeBounds.MaxZ = FMath::Min(FMath::TruncToInt32(MaxPosition.Z) + 1, TranslucencyLightingVolumeDim);
 	}
 
 	return VolumeBounds;
