Rendering Engine change to support the following two kind of path tracing denoiser plugins at the same time:
  1. Spatial denoising only plugin.
  2. Spatial and temporal denoising plugin.

When the user enables both plugins:
  1. The spatial denoising only plugin + builtin temporal denoiser can be selected as (default):

        `r.PathTracing.SpatialDenoiser.Type 0`
  2. The spatial and temporal denoser plugin can be selected as:

        `r.PathTracing.SpatialDenoiser.Type 1`

Note:
  * There is a modification in FRDGBuilder::IsTransient() to make TexCreate_Shared  return false to make it committed resource. Otherwise, the interop will fail when trying to get the device E_Invalidarg (0x80070057).

diff --git a/Engine/Shaders/Private/PathTracing/PathTracingSpatialTemporalDenoising.usf b/Engine/Shaders/Private/PathTracing/PathTracingSpatialTemporalDenoising.usf
--- a/Engine/Shaders/Private/PathTracing/PathTracingSpatialTemporalDenoising.usf
+++ b/Engine/Shaders/Private/PathTracing/PathTracingSpatialTemporalDenoising.usf
@@ -69,6 +69,10 @@ PathTracingSpatialTemporalDenoising.usf: Spatial temporal denoising for path tra
 #define TOTAL_VARIATION					0
 #endif
 
+#ifndef PREPROCESS_BUFFER
+#define PREPROCESS_BUFFER				0
+#endif
+
 // Define the distance metrics to use
 #define METRICS_LUMINANCE				0
 #define METRICS_EUCLIDEAN				1
@@ -94,6 +98,10 @@ PathTracingSpatialTemporalDenoising.usf: Spatial temporal denoising for path tra
 // the mip difference is 2 instead of one to make use of coarser mips
 #define MIP_DIFF_DELTA 2
 
+// The motion vector estimation algorithm is inspired from
+// Hanika, J., Tessari, L., & Dachsbacher, C. (2021). 
+// Fast Temporal Reprojection without Motion Vectors. Journal of Computer Graphics Techniques Vol, 10(3).
+
 // Distance type (L1, L2, logL1, logL1^C)
 #define DISTANCE_L1						0
 #define DISTANCE_L2						1
@@ -123,11 +131,6 @@ void Max3(in float3 Value, out float MaxValue, out int MaxIndex)
 	MaxValue = Value[MaxIndex];
 }
 
-float2 GetPositionScaleFromMipLevel(uint FromMipLevel, uint ToMipLevel)
-{
-	return (float2((uint2)TargetViewport_Extent >> FromMipLevel)) * (float2((uint2)TargetViewport_Extent >> ToMipLevel));
-}
-
 #if K_NUM_OF_TEXTURES_PER_PASS <= 7 && (TEMPORAL_REPROJECTION_ALIGN||TEMPORAL_REPROJECTION_MERGE)
 
 static int2 Shifts[25] = {
@@ -1274,7 +1277,7 @@ float3 hsv2rgb(float h, float s, float v)
 	return RGB;
 }
 
-//:adaptive from https://www.shadertoy.com/view/MsS3Wc by iq
+//:adapted from https://www.shadertoy.com/view/MsS3Wc by iq
 float3 hsv2rgb_smooth(float3 c)
 {
 	float3 rgb = clamp(abs( (c.x * 6.0 + float3(0.0, 4.0, 2.0)) % 6.0 - 3.0) - 1.0, 0.0, 1.0);
@@ -1322,7 +1325,7 @@ void VisualizePathTracingMotionVector(
 		h = h < 0 ? ((h + 2) / 2) : h / 2.0;
 		s = 1.0f;
 		OutColor = float4(hsv2rgb_smooth(float3(h, s, v)), 1.0f);// float4(s, h, Velocity.b, 1.0f);
-		OutColor.rgb = pow(OutColor.rgb, Velocity.z);
+		OutColor.rgb = pow(OutColor.rgb, Velocity.z == 0? 1.0f: Velocity.z);
 #endif
 	}
 	else
@@ -1472,4 +1475,23 @@ void SpatialDenoiserCS(uint3 DT_ID : SV_DispatchThreadID)
 		OutputTexture[Position] = InputTexture[Position];
 	}
 }
-#endif // SPATIAL_DENOISING
\ No newline at end of file
+#endif // SPATIAL_DENOISING
+
+#if PREPROCESS_BUFFER
+
+RWTexture2D<float4> NormalTexture;
+float Width;
+float Height;
+
+[numthreads(THREAD_SIZE_X, THREAD_SIZE_Y, 1)]
+void ConvertWorldSpaceNormalToCameraSpaceCS(uint3 DT_ID : SV_DispatchThreadID)
+{
+	int2 Position = DT_ID.xy;
+	int2 MaxPosition = int2(Width, Height);
+	if (all(Position < MaxPosition))
+	{
+		float4 Normal = float4(NormalTexture[Position].xyz, 0.0f);
+		NormalTexture[Position] = float4(mul(Normal, View.TranslatedWorldToCameraView).xyz, 0.0f);
+	}
+}
+#endif
\ No newline at end of file
diff --git a/Engine/Source/Runtime/RenderCore/Private/RenderGraphBuilder.cpp b/Engine/Source/Runtime/RenderCore/Private/RenderGraphBuilder.cpp
--- a/Engine/Source/Runtime/RenderCore/Private/RenderGraphBuilder.cpp
+++ b/Engine/Source/Runtime/RenderCore/Private/RenderGraphBuilder.cpp
@@ -477,6 +477,10 @@ bool FRDGBuilder::IsTransient(FRDGBufferRef Buffer) const
 
 bool FRDGBuilder::IsTransient(FRDGTextureRef Texture) const
 {
+	if (EnumHasAnyFlags(Texture->Desc.Flags, ETextureCreateFlags::Shared))
+	{
+		return false;
+	}
 	return IsTransientInternal(Texture, EnumHasAnyFlags(Texture->Desc.Flags, ETextureCreateFlags::FastVRAM));
 }
 
diff --git a/Engine/Source/Runtime/Renderer/Private/PathTracing.cpp b/Engine/Source/Runtime/Renderer/Private/PathTracing.cpp
--- a/Engine/Source/Runtime/Renderer/Private/PathTracing.cpp
+++ b/Engine/Source/Runtime/Renderer/Private/PathTracing.cpp
@@ -2206,7 +2206,7 @@ void FDeferredShadingSceneRenderer::RenderPathTracing(
 			View.ViewRect.Size(),
 			PF_A32B32G32R32F,
 			FClearValueBinding::None,
-			TexCreate_ShaderResource | TexCreate_UAV);
+			TexCreate_ShaderResource | TexCreate_UAV | GetExtraTextureCreateFlagsForDenoiser());
 		RadianceTexture = GraphBuilder.CreateTexture(Desc, TEXT("PathTracer.Radiance"), ERDGTextureFlags::MultiFrame);
 		AlbedoTexture   = GraphBuilder.CreateTexture(Desc, TEXT("PathTracer.Albedo")  , ERDGTextureFlags::MultiFrame);
 		NormalTexture   = GraphBuilder.CreateTexture(Desc, TEXT("PathTracer.Normal")  , ERDGTextureFlags::MultiFrame);
diff --git a/Engine/Source/Runtime/Renderer/Private/PathTracingSpatialTemporalDenoising.cpp b/Engine/Source/Runtime/Renderer/Private/PathTracingSpatialTemporalDenoising.cpp
--- a/Engine/Source/Runtime/Renderer/Private/PathTracingSpatialTemporalDenoising.cpp
+++ b/Engine/Source/Runtime/Renderer/Private/PathTracingSpatialTemporalDenoising.cpp
@@ -7,7 +7,8 @@
 #include "RHIDefinitions.h"
 
 PathTracingDenoiserFunction* GPathTracingDenoiserFunc = nullptr;
-
+PathTracingSpatialTemporalDenoiserFunction* GPathTracingSpatialTemporalDenoiserFunc = nullptr;
+PathTracingMotionVectorFunction* GPathTracingMotionVectorFunc = nullptr;
 #if RHI_RAYTRACING
 
 #include "RendererPrivate.h"
@@ -31,6 +32,21 @@ BEGIN_SHADER_PARAMETER_STRUCT(FDenoiseTextureParameters, )
 	RDG_TEXTURE_ACCESS(OutputTexture, ERHIAccess::CopyDest)
 END_SHADER_PARAMETER_STRUCT()
 
+BEGIN_SHADER_PARAMETER_STRUCT(FDenoiseTextureExtParameters, )
+	RDG_TEXTURE_ACCESS(InputTexture, ERHIAccess::CopySrc)
+	RDG_TEXTURE_ACCESS(InputAlbedo, ERHIAccess::CopySrc)
+	RDG_TEXTURE_ACCESS(InputNormal, ERHIAccess::CopySrc)
+	RDG_TEXTURE_ACCESS(InputFlow, ERHIAccess::CopySrc)
+	RDG_TEXTURE_ACCESS(InputPreviousOutput, ERHIAccess::CopySrc)
+	RDG_TEXTURE_ACCESS(OutputTexture, ERHIAccess::CopyDest)
+END_SHADER_PARAMETER_STRUCT()
+
+BEGIN_SHADER_PARAMETER_STRUCT(FMotionVectorParameters, )
+	RDG_TEXTURE_ACCESS(InputFrameTexture, ERHIAccess::CopySrc)
+	RDG_TEXTURE_ACCESS(ReferenceFrameTexture, ERHIAccess::CopySrc)
+	RDG_TEXTURE_ACCESS(OutputTexture, ERHIAccess::CopyDest)
+END_SHADER_PARAMETER_STRUCT()
+
 namespace {
 
 	TAutoConsoleVariable<int32> CVarPathTracingDenoiser(
@@ -43,6 +59,15 @@ namespace {
 		ECVF_RenderThreadSafe
 	);
 
+	TAutoConsoleVariable<int32> CVarPathTracingDenoiserNormalSpace(
+		TEXT("r.PathTracing.Denoiser.NormalSpace"),
+		0,
+		TEXT("The space normal is in\n")
+		TEXT("0: World space (default)\n")
+		TEXT("1: Camera space. Some denoisers require camera space normal\n"),
+		ECVF_RenderThreadSafe
+	);
+
 	TAutoConsoleVariable<int32> CVarPathTracingSpatialDenoiser(
 		TEXT("r.PathTracing.SpatialDenoiser"),
 		1,
@@ -53,6 +78,15 @@ namespace {
 		ECVF_RenderThreadSafe
 	);
 
+	TAutoConsoleVariable<int32> CVarPathTracingSpatialDenoiserType(
+		TEXT("r.PathTracing.SpatialDenoiser.Type"),
+		0,
+		TEXT("The type of spatial denoiser\n")
+		TEXT("0: Use spatial denoiser only plugin\n")
+		TEXT("1: Use spatial denoiser plugin that also provides temporal denoising\n"),
+		ECVF_RenderThreadSafe
+	);
+
 	TAutoConsoleVariable<int32> CVarPathTracingTemporalDenoiser(
 		TEXT("r.PathTracing.TemporalDenoiser"),
 		0,
@@ -63,6 +97,24 @@ namespace {
 		ECVF_RenderThreadSafe
 	);
 
+	TAutoConsoleVariable<int32> CVarPathTracingTemporalDenoiserType(
+		TEXT("r.PathTracing.TemporalDenoiser.Type"),
+		0,
+		TEXT("The type of temporal denoiser\n")
+		TEXT("0: Use the built-in temporal denoiser\n")
+		TEXT("1: Use the temporal denoiser from plugin\n"),
+		ECVF_RenderThreadSafe
+	);
+
+	TAutoConsoleVariable<int32> CVarPathTracingTemporalDenoiserMotionVectorType(
+		TEXT("r.PathTracing.TemporalDenoiser.MotionVector.Type"),
+		0,
+		TEXT("The type of motion vecotr estimation algorithm\n")
+		TEXT("0: Built-in motion vector estimator\n")
+		TEXT("1: Motion vector estimator from the plugin\n"),
+		ECVF_RenderThreadSafe
+	);
+
 	TAutoConsoleVariable<int32> CVarPathTracingTemporalDenoiserVisualizeMotionVector(
 		TEXT("r.PathTracing.TemporalDenoiser.VisualizeMotionVector"),
 		0,
@@ -212,11 +264,25 @@ int GetPathTracingDenoiserMode(const FViewInfo& View)
 bool IsPathTracingDenoiserEnabled(const FViewInfo& View)
 {
 	int DenoiserMode = GetPathTracingDenoiserMode(View);
-	const bool IsDenoiserEnabled = DenoiserMode != 0 && GPathTracingDenoiserFunc != nullptr;
+	const bool IsDenoiserEnabled = DenoiserMode != 0 && (GPathTracingDenoiserFunc != nullptr || GPathTracingSpatialTemporalDenoiserFunc != nullptr);
 
 	return IsDenoiserEnabled;
 }
 
+static bool ShouldDenoiseWithNormalInCameraSpace()
+{
+	int NormalSpace = CVarPathTracingDenoiserNormalSpace.GetValueOnRenderThread();
+	return NormalSpace != 0;
+}
+
+enum class ESpatialDenoiserType : int32
+{
+	NONE = -1,
+	SPATIAL_DENOISER_PLUGIN,
+	SPATIAL_TEMPORAL_DENOISER_PLUGIN,
+	MAX
+};
+
 enum class ETemporalDenoisingMode : uint32 
 {
 	ETDM_DISABLED,
@@ -225,6 +291,22 @@ enum class ETemporalDenoisingMode : uint32
 	ETDM_MAX
 };
 
+enum class ETemporalDenoiserType : int32
+{
+	NONE = -1,
+	BUILTIN_TEMPORAL_DENOISER,
+	SPATIAL_TEMPORAL_DENOISER_PLUGIN,
+	MAX
+};
+
+enum class ETemporalDenoiserMotionVectorType : int32
+{
+	NONE = -1,
+	BUILTIN,
+	PLUGIN,
+	MAX
+};
+
 bool ShouldEnablePathTracingDenoiserRealtimeDebug()
 {
 	int TemporalDenoiserMode = CVarPathTracingTemporalDenoiserMode.GetValueOnRenderThread();
@@ -239,16 +321,109 @@ static bool ShouldApplySpatialDenoiser()
 	return CVarPathTracingSpatialDenoiser.GetValueOnRenderThread() != 0;
 }
 
-static bool ShouldApplyTemporalDenoiser(const FPathTracingSpatialTemporalDenoisingContext& DenoisingContext)
+
+TArray<ESpatialDenoiserType> GetAvailableSpatialDenoiserTypes()
+{
+	TArray<ESpatialDenoiserType> Types;
+	if (GPathTracingSpatialTemporalDenoiserFunc)
+	{
+		Types.Add(ESpatialDenoiserType::SPATIAL_TEMPORAL_DENOISER_PLUGIN);
+	}
+
+	if (GPathTracingDenoiserFunc)
+	{
+		Types.Add(ESpatialDenoiserType::SPATIAL_DENOISER_PLUGIN);
+	}
+
+	if (Types.Num() == 0)
+	{
+		Types.Add(ESpatialDenoiserType::NONE);
+	}
+	return Types;
+}
+
+ESpatialDenoiserType GetSpatialDenosierType()
+{
+	int32 Type = CVarPathTracingSpatialDenoiserType.GetValueOnRenderThread();
+
+	Type = FMath::Clamp(Type,
+		static_cast<int32>(ESpatialDenoiserType::NONE) + 1,
+		static_cast<int32>(ESpatialDenoiserType::MAX) - 1);
+
+	ESpatialDenoiserType DenoiserType = static_cast<ESpatialDenoiserType>(Type);
+
+	if (DenoiserType == ESpatialDenoiserType::SPATIAL_DENOISER_PLUGIN && GPathTracingDenoiserFunc)
+	{
+		return DenoiserType;
+	}
+	else if (DenoiserType == ESpatialDenoiserType::SPATIAL_TEMPORAL_DENOISER_PLUGIN && GPathTracingSpatialTemporalDenoiserFunc)
+	{
+		return DenoiserType;
+	}
+	else
+	{
+		// Fallback to the one that is available instead of the requested one
+		DenoiserType = GetAvailableSpatialDenoiserTypes()[0];
+	}
+
+	return DenoiserType;
+}
+
+static bool ShouldApplyTemporalDenoiser(const FPathTracingSpatialTemporalDenoisingContext& DenoisingContext, const FViewInfo& View)
 {
 	const bool bApplyTemporalDenoiser =
-		CVarPathTracingTemporalDenoiser.GetValueOnAnyThread() == 1 &&
+		(CVarPathTracingTemporalDenoiser.GetValueOnAnyThread() == 1) &&
 		DenoisingContext.LastDenoisedRadianceTexture &&
 		DenoisingContext.LastDenoisedRadianceTexture != DenoisingContext.RadianceTexture;
 
 	return bApplyTemporalDenoiser;
 }
 
+ETemporalDenoiserType GetTemporalDenoiserType()
+{
+
+	ESpatialDenoiserType SpatialDenoiserType = GetSpatialDenosierType();
+
+	if ( SpatialDenoiserType == ESpatialDenoiserType::SPATIAL_DENOISER_PLUGIN)
+	{
+		return ETemporalDenoiserType::BUILTIN_TEMPORAL_DENOISER;
+	}
+	else if (SpatialDenoiserType == ESpatialDenoiserType::SPATIAL_TEMPORAL_DENOISER_PLUGIN)
+	{
+		return ETemporalDenoiserType::SPATIAL_TEMPORAL_DENOISER_PLUGIN;
+	}
+
+	return ETemporalDenoiserType::NONE;
+
+}
+
+ETemporalDenoiserMotionVectorType GetTemporalDenoiserMotionVectorType()
+{
+	int32 Type = CVarPathTracingTemporalDenoiserMotionVectorType.GetValueOnRenderThread();
+
+	Type = FMath::Clamp(Type,
+		static_cast<int32>(ETemporalDenoiserMotionVectorType::NONE) + 1,
+		static_cast<int32>(ETemporalDenoiserMotionVectorType::MAX) - 1);
+
+	return static_cast<ETemporalDenoiserMotionVectorType>(Type);
+}
+
+ETextureCreateFlags GetExtraTextureCreateFlagsForDenoiser()
+{
+	ESpatialDenoiserType SpatialDenoiserType = GetSpatialDenosierType();
+	ETemporalDenoiserType TemporalDenoiserType = GetTemporalDenoiserType();
+	ETextureCreateFlags TextureCreateFlags = ETextureCreateFlags::None;
+
+	if (SpatialDenoiserType == ESpatialDenoiserType::SPATIAL_TEMPORAL_DENOISER_PLUGIN ||
+		TemporalDenoiserType == ETemporalDenoiserType::SPATIAL_TEMPORAL_DENOISER_PLUGIN ||
+		GPathTracingSpatialTemporalDenoiserFunc)
+	{
+		TextureCreateFlags |= TexCreate_Shared | TexCreate_RenderTargetable;
+	}
+
+	return TextureCreateFlags;
+}
+
 static float GetHighFrequencyCutoffDeltaE()
 {
 	return FMath::Max(1.0f, CVarPathTracingTemporalDenoiserDeltaEHighFrequency.GetValueOnRenderThread());
@@ -619,7 +794,7 @@ public:
 		SHADER_PARAMETER_RDG_TEXTURE_UAV(RWTexture2D, OutputTexture)
 		END_SHADER_PARAMETER_STRUCT()
 
-		static bool ShouldCompilePermutation(const FGlobalShaderPermutationParameters& Parameters)
+	static bool ShouldCompilePermutation(const FGlobalShaderPermutationParameters& Parameters)
 	{
 		return ShouldCompilePathTracingDenoiserShadersForProject(Parameters.Platform);
 	}
@@ -649,7 +824,7 @@ public:
 		SHADER_PARAMETER_STRUCT(FScreenPassTextureViewportParameters, TargetViewport)
 	END_SHADER_PARAMETER_STRUCT()
 
-		static bool ShouldCompilePermutation(const FGlobalShaderPermutationParameters& Parameters)
+	static bool ShouldCompilePermutation(const FGlobalShaderPermutationParameters& Parameters)
 	{
 		return ShouldCompilePathTracingDenoiserShadersForProject(Parameters.Platform);
 	}
@@ -662,6 +837,57 @@ public:
 
 IMPLEMENT_GLOBAL_SHADER(FSpatialDenoiserCS, "/Engine/Private/PathTracing/PathTracingSpatialTemporalDenoising.usf", "SpatialDenoiserCS", SF_Compute);
 
+class FConvertWorldSpaceNormalToCameraSpaceCS : public FGlobalShader
+{
+public:
+	DECLARE_GLOBAL_SHADER(FConvertWorldSpaceNormalToCameraSpaceCS);
+	SHADER_USE_PARAMETER_STRUCT(FConvertWorldSpaceNormalToCameraSpaceCS, FGlobalShader);
+
+	BEGIN_SHADER_PARAMETER_STRUCT(FParameters, )
+		SHADER_PARAMETER_RDG_TEXTURE_UAV(RWTexture2D, NormalTexture)
+		SHADER_PARAMETER_STRUCT_REF(FViewUniformShaderParameters, ViewUniformBuffer)
+		SHADER_PARAMETER(float, Width)
+		SHADER_PARAMETER(float, Height)
+		END_SHADER_PARAMETER_STRUCT()
+
+		static bool ShouldCompilePermutation(const FGlobalShaderPermutationParameters& Parameters)
+	{
+		return ShouldCompilePathTracingDenoiserShadersForProject(Parameters.Platform);
+	}
+
+	static void ModifyCompilationEnvironment(const FShaderPermutationParameters&, FShaderCompilerEnvironment& OutEnvironment)
+	{
+		OutEnvironment.SetDefine(TEXT("PREPROCESS_BUFFER"), 1);
+		OutEnvironment.CompilerFlags.Add(CFLAG_AllowTypedUAVLoads);
+	}
+};
+
+IMPLEMENT_GLOBAL_SHADER(FConvertWorldSpaceNormalToCameraSpaceCS, "/Engine/Private/PathTracing/PathTracingSpatialTemporalDenoising.usf", "ConvertWorldSpaceNormalToCameraSpaceCS", SF_Compute);
+
+static void ConvertNormalSpace(FRDGBuilder& GraphBuilder,
+	const FViewInfo& View, 
+	FRDGTextureRef NormalTexture)
+{
+	typedef FConvertWorldSpaceNormalToCameraSpaceCS SHADER;
+	SHADER::FParameters* PassParameters = GraphBuilder.AllocParameters<SHADER::FParameters>();
+	PassParameters->NormalTexture = GraphBuilder.CreateUAV(FRDGTextureUAVDesc(NormalTexture));
+	PassParameters->ViewUniformBuffer = View.ViewUniformBuffer;
+	PassParameters->Width = NormalTexture->Desc.GetSize().X;
+	PassParameters->Height = NormalTexture->Desc.GetSize().Y;
+
+
+	TShaderMapRef<SHADER> ComputeShader(View.ShaderMap);
+
+	FComputeShaderUtils::AddPass(
+		GraphBuilder,
+		RDG_EVENT_NAME("ConvertWorldSpaceNormalToCameraSpace %dx%d",
+			View.ViewRect.Width(),
+			View.ViewRect.Height()),
+		ComputeShader,
+		PassParameters,
+		FComputeShaderUtils::GetGroupCount(View.ViewRect.Size(), 8));
+}
+
 static void PathTracingDenoiserPlugin(FRDGBuilder& GraphBuilder,
 	const FViewInfo& View,
 	int DenoiserMode,
@@ -670,10 +896,28 @@ static void PathTracingDenoiserPlugin(FRDGBuilder& GraphBuilder,
 	FRDGTextureRef NormalTexture,
 	FRDGTextureRef OutputTexture)
 {
+	FRDGTextureRef ProcessedNormalTexture = NormalTexture;
+
+	if (ShouldDenoiseWithNormalInCameraSpace())
+	{
+		const FRDGTextureDesc& Desc = NormalTexture->Desc;
+		ProcessedNormalTexture = GraphBuilder.CreateTexture(Desc, TEXT("PathTracing.CameraSpaceNormal"));
+		{
+			FRHICopyTextureInfo CopyInfo;
+			CopyInfo.Size.X = Desc.Extent.X;
+			CopyInfo.Size.Y = Desc.Extent.Y;
+			CopyInfo.Size.Z = 1;
+			CopyInfo.NumMips = Desc.NumMips;
+
+			AddCopyTexturePass(GraphBuilder, NormalTexture, ProcessedNormalTexture, CopyInfo);
+		}
+		ConvertNormalSpace(GraphBuilder, View, ProcessedNormalTexture);
+	}
+
 	FDenoiseTextureParameters* DenoiseParameters = GraphBuilder.AllocParameters<FDenoiseTextureParameters>();
 	DenoiseParameters->InputTexture = InputTexture;
 	DenoiseParameters->InputAlbedo = AlbedoTexture;
-	DenoiseParameters->InputNormal = NormalTexture;
+	DenoiseParameters->InputNormal = ProcessedNormalTexture;
 	DenoiseParameters->OutputTexture = OutputTexture;
 	// Need to read GPU mask outside Pass function, as the value is not refreshed inside the pass
 	GraphBuilder.AddPass(RDG_EVENT_NAME("Path Tracer Denoiser Plugin"), DenoiseParameters, ERDGPassFlags::Readback,
@@ -689,6 +933,111 @@ static void PathTracingDenoiserPlugin(FRDGBuilder& GraphBuilder,
 	);
 }
 
+static void PathTracingSpatialTemporalDenoiserPlugin(FRDGBuilder& GraphBuilder,
+	const FViewInfo& View,
+	int DenoiserMode,
+	FRDGTextureRef InputTexture,
+	FRDGTextureRef AlbedoTexture,
+	FRDGTextureRef NormalTexture,
+	FRDGTextureRef FlowTexture,
+	FRDGTextureRef PreviousOutputFrameTexture,
+	FRDGTextureRef OutputTexture,
+	int DenoisingFrameId,
+	bool bForceSpatialDenoiserOnly)
+{
+
+	FRDGTextureRef ProcessedNormalTexture = NormalTexture;
+
+	if (ShouldDenoiseWithNormalInCameraSpace())
+	{
+		const FRDGTextureDesc& Desc = NormalTexture->Desc;
+		ProcessedNormalTexture = GraphBuilder.CreateTexture(Desc, TEXT("PathTracing.CameraSpaceNormal"));
+		{
+			FRHICopyTextureInfo CopyInfo;
+			CopyInfo.Size.X = Desc.Extent.X;
+			CopyInfo.Size.Y = Desc.Extent.Y;
+			CopyInfo.Size.Z = 1;
+			CopyInfo.NumMips = Desc.NumMips;
+
+			AddCopyTexturePass(GraphBuilder, NormalTexture, ProcessedNormalTexture, CopyInfo);
+		}
+		ConvertNormalSpace(GraphBuilder, View, ProcessedNormalTexture);
+	}
+
+	FDenoiseTextureExtParameters* DenoiseParameters = GraphBuilder.AllocParameters<FDenoiseTextureExtParameters>();
+	DenoiseParameters->InputTexture = InputTexture;
+	DenoiseParameters->InputAlbedo = AlbedoTexture;
+	DenoiseParameters->InputNormal = ProcessedNormalTexture;
+	DenoiseParameters->InputFlow = FlowTexture;
+	DenoiseParameters->InputPreviousOutput = PreviousOutputFrameTexture;
+	DenoiseParameters->OutputTexture = OutputTexture;
+
+	// Need to read GPU mask outside Pass function, as the value is not refreshed inside the pass
+	GraphBuilder.AddPass(RDG_EVENT_NAME("Path Tracer Denoiser Ext Plugin"), DenoiseParameters, ERDGPassFlags::Readback,
+		[DenoiseParameters, DenoiserMode, DenoisingFrameId, bForceSpatialDenoiserOnly, GPUMask = View.GPUMask](FRHICommandListImmediate& RHICmdList)
+	{
+		check(GPathTracingSpatialTemporalDenoiserFunc);
+
+		FDenoisingArgumentsExt DenoisingArgumentsExt;
+		DenoisingArgumentsExt.FlowTex = DenoiseParameters->InputFlow ?
+			DenoiseParameters->InputFlow->GetRHI()->GetTexture2D() : nullptr;
+		DenoisingArgumentsExt.PreviousOutputTex = DenoiseParameters->InputPreviousOutput ?
+			DenoiseParameters->InputPreviousOutput->GetRHI()->GetTexture2D() : nullptr;
+		DenoisingArgumentsExt.DenoisingFrameId = DenoisingFrameId;
+		DenoisingArgumentsExt.bForceSpatialDenoiserOnly = bForceSpatialDenoiserOnly;
+
+		GPathTracingSpatialTemporalDenoiserFunc(RHICmdList,
+			DenoiseParameters->InputTexture->GetRHI()->GetTexture2D(),
+			DenoiseParameters->InputAlbedo->GetRHI()->GetTexture2D(),
+			DenoiseParameters->InputNormal->GetRHI()->GetTexture2D(),
+			DenoiseParameters->OutputTexture->GetRHI()->GetTexture2D(),
+			&DenoisingArgumentsExt,
+			GPUMask);
+	});
+}
+
+static bool ShouldApplyPreExposureToMotionVectorEstimation()
+{
+	int32 DenoiserSource = CVarPathTracingTemporalDenoiserSource.GetValueOnRenderThread();
+	if (DenoiserSource == 0 || DenoiserSource == 3)
+	{
+		return true;
+	}
+
+	return false;
+}
+
+static void PathTracingMotionVectorPlugin(FRDGBuilder& GraphBuilder,
+	const FViewInfo& View,
+	FRDGTextureRef InputFrameTexture,
+	FRDGTextureRef ReferenceFrameTexture,
+	FRDGTextureRef OutputTexture)
+{
+	check(GPathTracingMotionVectorFunc);
+
+	FMotionVectorParameters* DenoiseParameters = GraphBuilder.AllocParameters<FMotionVectorParameters>();
+	DenoiseParameters->InputFrameTexture = InputFrameTexture;
+	DenoiseParameters->ReferenceFrameTexture = ReferenceFrameTexture;
+	DenoiseParameters->OutputTexture = OutputTexture;
+	
+	bool bShouldApplyPreExposure = ShouldApplyPreExposureToMotionVectorEstimation();
+	float PreExposure = bShouldApplyPreExposure ? View.PreExposure : 1.0f;
+
+
+	GraphBuilder.AddPass(RDG_EVENT_NAME("Path Tracer Motion Vector Plugin"), DenoiseParameters, ERDGPassFlags::Readback,
+		[DenoiseParameters, PreExposure, GPUMask = View.GPUMask](FRHICommandListImmediate& RHICmdList)
+	{
+
+		GPathTracingMotionVectorFunc(RHICmdList,
+			DenoiseParameters->InputFrameTexture->GetRHI()->GetTexture2D(),
+			DenoiseParameters->ReferenceFrameTexture->GetRHI()->GetTexture2D(),
+			DenoiseParameters->OutputTexture->GetRHI()->GetTexture2D(),
+			PreExposure,
+			GPUMask);
+	});
+
+}
+
 class FMotionVectorEstimationContext
 {
 public:
@@ -703,6 +1052,7 @@ public:
 
 	FRDGTextureDesc TextureWithNMipsDescriptor;
 	FRDGTextureDesc TextureDescriptor;
+	FRDGTextureDesc TextureDescriptorFinalOutput;
 
 	FRDGTextureRef SourceTexture;
 	FRDGTextureRef TargetTexture;
@@ -754,6 +1104,11 @@ public:
 			FClearValueBinding(),
 			TexCreate_RenderTargetable | TexCreate_ShaderResource | TexCreate_UAV);
 
+		TextureDescriptorFinalOutput = FRDGTextureDesc::Create2D(
+			TargetViewport.Extent,
+			PF_A32B32G32R32F,
+			FClearValueBinding(),
+			TexCreate_RenderTargetable | TexCreate_ShaderResource | TexCreate_UAV | GetExtraTextureCreateFlagsForDenoiser());
 		
 		PixelOffsetTextures[0] = GraphBuilder.CreateTexture(TextureWithNMipsDescriptor, PixelOffsetTextureNames[0]);
 		PixelOffsetTextures[1] = GraphBuilder.CreateTexture(TextureWithNMipsDescriptor, PixelOffsetTextureNames[1]);
@@ -1088,7 +1443,7 @@ static void FuseTemporalFeature(FRDGBuilder& GraphBuilder,
 		TargetViewport.Extent,
 		PF_A32B32G32R32F,
 		FClearValueBinding(),
-		TexCreate_RenderTargetable | TexCreate_ShaderResource | TexCreate_UAV);
+		TexCreate_RenderTargetable | TexCreate_ShaderResource | TexCreate_UAV | GetExtraTextureCreateFlagsForDenoiser());
 
 	OutSourceTexture = GraphBuilder.CreateTexture(TextureDescriptor, TEXT("PathTracing.Denoising.Feature.Source"));
 	OutTargetTexture = GraphBuilder.CreateTexture(TextureDescriptor, TEXT("PathTracing.Denoising.Feature.Target"));
@@ -1126,7 +1481,7 @@ static void FuseTemporalFeature(FRDGBuilder& GraphBuilder,
 
 }
 
-static void SelectMotionSourceAndTargetTextures(
+static bool SelectMotionSourceAndTargetTextures(
 	FRDGBuilder& GraphBuilder,
 	const FViewInfo& View,
 	FRDGTextureRef InSourceTexture,
@@ -1160,6 +1515,7 @@ static void SelectMotionSourceAndTargetTextures(
 		break;
 	}
 
+	return  MotionEstimationSourceTexture->Desc.Extent == MotionEstimationTargetTexture->Desc.Extent;
 }
 
 // Estimate the motion vector from source to target
@@ -1167,22 +1523,11 @@ static FRDGTextureRef TemporalReprojectionWithoutMotionVector(FRDGBuilder& Graph
 	const FViewInfo& View,
 	FRDGTextureRef InSourceTexture,
 	FRDGTextureRef InTargetTexture,
+	FRDGTextureRef MotionEstimationSourceTexture,
+	FRDGTextureRef MotionEstimationTargetTexture,
 	FPathTracingSpatialTemporalDenoisingContext& DenoisingContext)
 {
 
-	FRDGTextureRef MotionEstimationSourceTexture = nullptr;
-	FRDGTextureRef MotionEstimationTargetTexture = nullptr;
-
-	SelectMotionSourceAndTargetTextures(
-		GraphBuilder, View, InSourceTexture, InTargetTexture, DenoisingContext, MotionEstimationSourceTexture, MotionEstimationTargetTexture);
-
-	// TODO: when the source and target texture size is different
-	if (!(MotionEstimationSourceTexture->Desc.Extent == MotionEstimationTargetTexture->Desc.Extent))
-	{
-		UE_LOG(LogPathTracingDenoising, Log, TEXT("TODO: Add motion vector estimation when source and dest dimensions do not match"));
-		return nullptr;
-	}
-
 	FMotionVectorEstimationContext EstimationContext(MotionEstimationSourceTexture, MotionEstimationTargetTexture);
 
 	if (!EstimationContext.InitContext(GraphBuilder, View))
@@ -1191,7 +1536,7 @@ static FRDGTextureRef TemporalReprojectionWithoutMotionVector(FRDGBuilder& Graph
 	}
 
 	FRDGTextureRef  HighFrequencyRejectMap = GraphBuilder.CreateTexture(EstimationContext.TextureDescriptor, TEXT("PathTracing.EstimateMotion.HighFrequencyRejectMap"));
-	FRDGTextureRef  FinalAccumulation = GraphBuilder.CreateTexture(EstimationContext.TextureDescriptor, TEXT("PathTracing.EstimateMotion.FinalAccumulation"));
+	FRDGTextureRef  FinalAccumulation = GraphBuilder.CreateTexture(EstimationContext.TextureDescriptorFinalOutput, TEXT("PathTracing.EstimateMotion.FinalAccumulation"));
 	int PatchCount = EstimationContext.GetPatchCount();
 	float HighFrequencyCutoffDeltaE = GetHighFrequencyCutoffDeltaE();
 
@@ -1364,44 +1709,126 @@ void PathTracingSpatialTemporalDenoising(FRDGBuilder& GraphBuilder,
 	RDG_GPU_STAT_SCOPE(GraphBuilder, PathTracingSpatialTemporalDenoising);
 	RDG_EVENT_SCOPE(GraphBuilder, "PathTracingSpatialTemporalDenoising");
 
+	ETextureCreateFlags ExtraFlags = GetExtraTextureCreateFlagsForDenoiser();
+
 	FRDGTextureDesc RadianceTextureDesc = FRDGTextureDesc::Create2D(
 		View.ViewRect.Size(),
 		PF_A32B32G32R32F,
 		FClearValueBinding::None,
-		TexCreate_ShaderResource | TexCreate_UAV);
+		TexCreate_ShaderResource | TexCreate_UAV | ExtraFlags);
 
+	const ESpatialDenoiserType SpatialDenoiserType = GetSpatialDenosierType();
+	const ETemporalDenoiserType TemporalDenoiserType = GetTemporalDenoiserType();
 	const bool bApplySpatialDenoiser = ShouldApplySpatialDenoiser();
+	const bool bApplyTemporalDenoiser = ShouldApplyTemporalDenoiser(SpatialTemporalDenoisingContext, View);
+
 	FRDGTextureRef TargetTexture = SpatialTemporalDenoisingContext.RadianceTexture;
+	FRDGTextureRef SourceTexture = SpatialTemporalDenoisingContext.LastDenoisedRadianceTexture;
+	FRDGTextureRef TemporalDenoisedTexture = nullptr;
 
-	if (bApplySpatialDenoiser)
+	if (SpatialDenoiserType == ESpatialDenoiserType::SPATIAL_DENOISER_PLUGIN)
 	{
-		TargetTexture = GraphBuilder.CreateTexture(RadianceTextureDesc, TEXT("PathTracer.SpatialDenoiser.Output"));
-		PathTracingDenoiserPlugin(
-			GraphBuilder,
-			View,
-			DenoiserMode,
-			SpatialTemporalDenoisingContext.RadianceTexture,
-			SpatialTemporalDenoisingContext.AlbedoTexture,
-			SpatialTemporalDenoisingContext.NormalTexture,
-			TargetTexture);
-	}
+		if (bApplySpatialDenoiser)
+		{
+			TargetTexture = GraphBuilder.CreateTexture(RadianceTextureDesc, TEXT("PathTracer.SpatialDenoiser.Output"));
+			PathTracingDenoiserPlugin(
+				GraphBuilder,
+				View,
+				DenoiserMode,
+				SpatialTemporalDenoisingContext.RadianceTexture,
+				SpatialTemporalDenoisingContext.AlbedoTexture,
+				SpatialTemporalDenoisingContext.NormalTexture,
+				TargetTexture);
+		}
+
+		if (bApplyTemporalDenoiser)
+		{
+			// Image space temporal and spatial denoising.
+			// 
+			// Select motion source and target for temporal denoising
+			FRDGTextureRef MotionEstimationSourceTexture = nullptr;
+			FRDGTextureRef MotionEstimationTargetTexture = nullptr;
+			bool IsSourceTargetDimensionMatch = SelectMotionSourceAndTargetTextures(
+					GraphBuilder,
+					View,
+					SourceTexture,
+					TargetTexture, SpatialTemporalDenoisingContext, MotionEstimationSourceTexture, MotionEstimationTargetTexture);
 
-	const bool bApplyTemporalDenoiser = ShouldApplyTemporalDenoiser(SpatialTemporalDenoisingContext);
-	FRDGTextureRef TemporalDenoisedTexture = nullptr;
 
-	if (bApplyTemporalDenoiser)
+			if (IsSourceTargetDimensionMatch)
+			{
+				UE_LOG(LogPathTracingDenoising, Log, TEXT("Using temporal denoising for frame %i"), SpatialTemporalDenoisingContext.FrameIndex);
+
+				if (TemporalDenoiserType == ETemporalDenoiserType::BUILTIN_TEMPORAL_DENOISER)
+				{
+					TemporalDenoisedTexture = TemporalReprojectionWithoutMotionVector(GraphBuilder,
+						View,
+						SourceTexture,
+						TargetTexture,
+						MotionEstimationSourceTexture,
+						MotionEstimationTargetTexture,
+						SpatialTemporalDenoisingContext);
+				}
+				else
+				{
+					// not supported
+				}
+			}
+		}
+	}
+	else if (SpatialDenoiserType == ESpatialDenoiserType::SPATIAL_TEMPORAL_DENOISER_PLUGIN)
 	{
-		UE_LOG(LogPathTracingDenoising, Log, TEXT("Using temporal denoising for frame %i"), SpatialTemporalDenoisingContext.FrameIndex);
+		FRDGTextureRef MotionTexture = nullptr;
+		bool bIsInitialFrame = true;
+		
+		MotionTexture = GraphBuilder.CreateTexture(RadianceTextureDesc, TEXT("PathTracing.EstimateMotion.OpticalFlow"));
 
-		// Image space temporal and spatial denoising.
-		// if we have the last available frame, do a fast reprojection to use history for frame stabilization
-		FRDGTextureRef SourceTexture = SpatialTemporalDenoisingContext.LastDenoisedRadianceTexture;
+		if (bApplyTemporalDenoiser)
+		{
+			// Select motion source and target for temporal denoising
+			FRDGTextureRef MotionEstimationSourceTexture = nullptr;
+			FRDGTextureRef MotionEstimationTargetTexture = nullptr;
+			bool IsSourceTargetDimensionMatch = SelectMotionSourceAndTargetTextures(
+					GraphBuilder,
+					View,
+					SourceTexture,
+					TargetTexture, SpatialTemporalDenoisingContext, MotionEstimationSourceTexture, MotionEstimationTargetTexture);
+
+			if (IsSourceTargetDimensionMatch)
+			{
+				if (TemporalDenoiserType == ETemporalDenoiserType::SPATIAL_TEMPORAL_DENOISER_PLUGIN)
+				{
+					FScreenPassTextureViewport TargetViewport = FScreenPassTextureViewport(View.ViewRect);
+					
+
+					PathTracingMotionVectorPlugin(
+						GraphBuilder,
+						View,
+						MotionEstimationSourceTexture,
+						MotionEstimationTargetTexture,
+						MotionTexture);
+					bIsInitialFrame = false;
+					
+				}
+			}
+		}
+		else
+		{
+			FLinearColor ClearColor = FLinearColor::Black;
+			AddClearRenderTargetPass(GraphBuilder, MotionTexture, ClearColor);
+		}
+
+		SpatialTemporalDenoisingContext.MotionVector = MotionTexture;
+
+		TemporalDenoisedTexture = GraphBuilder.CreateTexture(RadianceTextureDesc, TEXT("PathTracing.EstimateMotion.FinalAccumulation"));
+		
+		PathTracingSpatialTemporalDenoiserPlugin(
+			GraphBuilder, View, DenoiserMode, TargetTexture,
+			SpatialTemporalDenoisingContext.AlbedoTexture, SpatialTemporalDenoisingContext.NormalTexture,
+			MotionTexture, SourceTexture, TemporalDenoisedTexture,
+			SpatialTemporalDenoisingContext.FrameIndex,
+			bIsInitialFrame);	// zero frame will denoise without temporal
 
-		TemporalDenoisedTexture = TemporalReprojectionWithoutMotionVector(GraphBuilder,
-			View,
-			SourceTexture,
-			TargetTexture,
-			SpatialTemporalDenoisingContext);
 	}
 
 	SpatialTemporalDenoisedTexture = TemporalDenoisedTexture ? TemporalDenoisedTexture : TargetTexture;
diff --git a/Engine/Source/Runtime/Renderer/Private/PathTracingSpatialTemporalDenoising.h b/Engine/Source/Runtime/Renderer/Private/PathTracingSpatialTemporalDenoising.h
--- a/Engine/Source/Runtime/Renderer/Private/PathTracingSpatialTemporalDenoising.h
+++ b/Engine/Source/Runtime/Renderer/Private/PathTracingSpatialTemporalDenoising.h
@@ -45,6 +45,7 @@ struct FPathTracingSpatialTemporalDenoisingContext
 int	GetPathTracingDenoiserMode(const FViewInfo& View);
 bool IsPathTracingDenoiserEnabled(const FViewInfo& View);
 bool ShouldEnablePathTracingDenoiserRealtimeDebug();
+ETextureCreateFlags GetExtraTextureCreateFlagsForDenoiser();
 
 // Calculate the variance per-pixel.
 void PathTracingSpatialTemporalDenoisingPrePass(
diff --git a/Engine/Source/Runtime/Renderer/Public/PathTracingDenoiser.h b/Engine/Source/Runtime/Renderer/Public/PathTracingDenoiser.h
--- a/Engine/Source/Runtime/Renderer/Public/PathTracingDenoiser.h
+++ b/Engine/Source/Runtime/Renderer/Public/PathTracingDenoiser.h
@@ -7,6 +7,25 @@
 class FRHITexture;
 class FRHICommandListImmediate;
 
+// Spatial denoiser only Plugin
 using PathTracingDenoiserFunction = void(FRHICommandListImmediate& RHICmdList, FRHITexture* ColorTex, FRHITexture* AlbedoTex, FRHITexture* NormalTex, FRHITexture* OutputTex, FRHIGPUMask GPUMask);
 
 extern RENDERER_API PathTracingDenoiserFunction* GPathTracingDenoiserFunc;
+
+struct FDenoisingArgumentsExt
+{
+	FRHITexture* FlowTex;
+	FRHITexture* PreviousOutputTex;
+
+	int Width;
+	int Height;
+	int DenoisingFrameId;
+	bool bForceSpatialDenoiserOnly;
+};
+
+// Spatial-temporal denoiser plugin
+using PathTracingSpatialTemporalDenoiserFunction = void(FRHICommandListImmediate& RHICmdList, FRHITexture* ColorTex, FRHITexture* AlbedoTex, FRHITexture* NormalTex, FRHITexture* OutputTex, const FDenoisingArgumentsExt* DenoisingArgumentExt, FRHIGPUMask GPUMask);
+using PathTracingMotionVectorFunction = void(FRHICommandListImmediate& RHICmdList, FRHITexture* InputFrameTex, FRHITexture* ReferenceFrameTex, FRHITexture* FlowTex, float PreExposure, FRHIGPUMask GPUMask);
+
+extern RENDERER_API PathTracingSpatialTemporalDenoiserFunction* GPathTracingSpatialTemporalDenoiserFunc;
+extern RENDERER_API PathTracingMotionVectorFunction* GPathTracingMotionVectorFunc;
\ No newline at end of file
