Fix FillRayTracingInstanceUploadBuffer reading from write-combined memory.
- Setup InstanceDesc in a local variable and then copy to upload memory.

diff --git a/Engine/Source/Runtime/Renderer/Private/RayTracing/RayTracingInstanceBufferUtil.cpp b/Engine/Source/Runtime/Renderer/Private/RayTracing/RayTracingInstanceBufferUtil.cpp
--- a/Engine/Source/Runtime/Renderer/Private/RayTracing/RayTracingInstanceBufferUtil.cpp
+++ b/Engine/Source/Runtime/Renderer/Private/RayTracing/RayTracingInstanceBufferUtil.cpp
@@ -209,12 +209,7 @@ void FillRayTracingInstanceUploadBuffer(
 			int32 NumInactiveNativeInstancesThisSceneInstance = 0;
 			for (uint32 TransformIndex = 0; TransformIndex < NumTransforms; ++TransformIndex)
 			{
-				FRayTracingInstanceDescriptorInput& InstanceDesc = OutInstanceUploadData[BaseDescriptorIndex + TransformIndex];
-
-				InstanceDesc.InstanceMaskAndFlags = SceneInstance.Mask | ((uint32)SceneInstance.Flags << 8);
-				InstanceDesc.InstanceContributionToHitGroupIndex = SceneInitializer.SegmentPrefixSum[SceneInstanceIndex] * SceneInitializer.ShaderSlotsPerGeometrySegment;
-				InstanceDesc.InstanceId = bUseUniqueUserData ? SceneInstance.UserData[TransformIndex] : SceneInstance.DefaultUserData;
-				InstanceDesc.OutputDescriptorIndex = LayerBaseIndex + BaseInstanceIndex + TransformIndex;
+				FRayTracingInstanceDescriptorInput InstanceDesc;
 
 				if (bGpuSceneInstance)
 				{
@@ -232,7 +227,11 @@ void FillRayTracingInstanceUploadBuffer(
 					InstanceDesc.GPUSceneInstanceOrTransformIndex = BaseTransformIndex + TransformIndex;
 				}
 
+				InstanceDesc.OutputDescriptorIndex = LayerBaseIndex + BaseInstanceIndex + TransformIndex;
 				InstanceDesc.AccelerationStructureIndex = AccelerationStructureIndex;
+				InstanceDesc.InstanceId = bUseUniqueUserData ? SceneInstance.UserData[TransformIndex] : SceneInstance.DefaultUserData;
+				InstanceDesc.InstanceMaskAndFlags = SceneInstance.Mask | ((uint32)SceneInstance.Flags << 8);
+				InstanceDesc.InstanceContributionToHitGroupIndex = SceneInitializer.SegmentPrefixSum[SceneInstanceIndex] * SceneInitializer.ShaderSlotsPerGeometrySegment;
 				InstanceDesc.bApplyLocalBoundsTransform = SceneInstance.bApplyLocalBoundsTransform;
 
 				if (!SceneInstance.ActivationMask.IsEmpty() && (SceneInstance.ActivationMask[TransformIndex / 32] & (1 << (TransformIndex % 32))) == 0)
@@ -240,10 +239,7 @@ void FillRayTracingInstanceUploadBuffer(
 					// Set flag for deactivated instances
 					InstanceDesc.AccelerationStructureIndex = 0xFFFFFFFF;
 					NumInactiveNativeInstancesThisSceneInstance++;
-					continue;
-				}
-
-				if (bCpuInstance)
+				} else if (bCpuInstance)
 				{
 					const uint32 TransformDataOffset = InstanceDesc.GPUSceneInstanceOrTransformIndex * 3;
 					FMatrix LocalToTranslatedWorld = SceneInstance.Transforms[TransformIndex].ConcatTranslation(PreViewTranslation);
@@ -252,6 +248,8 @@ void FillRayTracingInstanceUploadBuffer(
 					OutTransformData[TransformDataOffset + 1] = *(FVector4f*)&LocalToTranslatedWorldF.M[1];
 					OutTransformData[TransformDataOffset + 2] = *(FVector4f*)&LocalToTranslatedWorldF.M[2];
 				}
+
+				OutInstanceUploadData[BaseDescriptorIndex + TransformIndex] = InstanceDesc;
 			}
 
 #if STATS
