Fixed stale NaniteRaytracingDataOffset in GPUScene.
- Primitive update was requested on a task causing a race condition.
- Moved Nanite::GRayTracingManager.Update() out of RayTracingRelevantPrimitiveListTask since it needs to run on RenderThread before GPUScene is updated.

diff --git a/Engine/Source/Runtime/Renderer/Private/DeferredShadingRenderer.cpp b/Engine/Source/Runtime/Renderer/Private/DeferredShadingRenderer.cpp
--- a/Engine/Source/Runtime/Renderer/Private/DeferredShadingRenderer.cpp
+++ b/Engine/Source/Runtime/Renderer/Private/DeferredShadingRenderer.cpp
@@ -719,7 +719,7 @@ static void GatherRayTracingRelevantPrimitives(const FScene& Scene, const FViewI
 				continue;
 			}
 
-			check(!EnumHasAnyFlags(Scene.PrimitiveRayTracingFlags[PrimitiveIndex], ERayTracingPrimitiveFlags::Excluded));
+			check(!EnumHasAnyFlags(Flags, ERayTracingPrimitiveFlags::Excluded));
 
 			const FPrimitiveSceneInfo* SceneInfo = Scene.Primitives[PrimitiveIndex];
 
@@ -1210,8 +1210,6 @@ bool FDeferredShadingSceneRenderer::GatherRayTracingWorldInstancesForView(FRDGBu
 
 		KickRayTracingMeshBatchTask();
 	}
-
-	Nanite::GRayTracingManager.Update();
 	
 	// Task to iterate over static ray tracing instances, perform auto-instancing and culling.
 	// This adds final instances to the ray tracing scene and must be done before FRayTracingScene::BuildInitializationData().
@@ -2427,6 +2425,8 @@ void FDeferredShadingSceneRenderer::Render(FRDGBuilder& GraphBuilder)
 
 	if (IsRayTracingEnabled() && RHISupportsRayTracingShaders(ViewFamily.GetShaderPlatform()))
 	{
+		Nanite::GRayTracingManager.Update();
+
 		if (!ViewFamily.EngineShowFlags.PathTracing)
 		{
 			// get the default lighting miss shader (to implicitly fill in the MissShader library before the RT pipeline is created)
