Ensure Strata Dbuffer pass shader is only compiled if DBuffer is enabled.

diff --git a/Engine/Shaders/Private/Strata/StrataDbuffer.usf b/Engine/Shaders/Private/Strata/StrataDbuffer.usf
--- a/Engine/Shaders/Private/Strata/StrataDbuffer.usf
+++ b/Engine/Shaders/Private/Strata/StrataDbuffer.usf
@@ -47,8 +47,12 @@ SamplerState DBufferCTextureSampler;
 #error STENCIL_STRATA_RECEIVE_DBUFFER_ROUGHNESS_BIT_ID needs to be defined
 #endif
 
+// Sanity check
+#if !defined(USE_DBUFFER) || USE_DBUFFER != 1
+#error This shader should only be compiled when USE_DBUFFER == 1
+#endif
+
 #define MATERIALDECALRESPONSEMASK 0x7
-#define USE_DBUFFER 1
 #define ENABLE_DBUFFER_TEXTURES 1
 #define DBUFFER_BASEPASS_RESOURCE 0
 
diff --git a/Engine/Source/Runtime/Renderer/Private/BasePassRendering.cpp b/Engine/Source/Runtime/Renderer/Private/BasePassRendering.cpp
--- a/Engine/Source/Runtime/Renderer/Private/BasePassRendering.cpp
+++ b/Engine/Source/Runtime/Renderer/Private/BasePassRendering.cpp
@@ -384,9 +384,9 @@ void SetDepthStencilStateForBasePass_Internal(FMeshPassProcessorRenderState& InD
 }
 
 template<bool bDepthTest, ECompareFunction CompareFunction>
-void SetDepthStencilStateForBasePass_Internal(FMeshPassProcessorRenderState& InDrawRenderState)
+void SetDepthStencilStateForBasePass_Internal(FMeshPassProcessorRenderState& InDrawRenderState, ERHIFeatureLevel::Type FeatureLevel)
 {	
-	const static bool bStrataDufferPassEnabled = Strata::IsStrataEnabled() && Strata::IsStrataDbufferPassEnabled();
+	const static bool bStrataDufferPassEnabled = Strata::IsStrataEnabled() && Strata::IsStrataDbufferPassEnabled(GShaderPlatformForFeatureLevel[FeatureLevel]);
 	if (bStrataDufferPassEnabled)
 	{
 		SetDepthStencilStateForBasePass_Internal<bDepthTest, CompareFunction, GET_STENCIL_BIT_MASK(STRATA_RECEIVE_DBUFFER_NORMAL, 1) | GET_STENCIL_BIT_MASK(STRATA_RECEIVE_DBUFFER_DIFFUSE, 1) | GET_STENCIL_BIT_MASK(STRATA_RECEIVE_DBUFFER_ROUGHNESS, 1) | GET_STENCIL_BIT_MASK(DISTANCE_FIELD_REPRESENTATION, 1) | STENCIL_LIGHTING_CHANNELS_MASK(0x7)>(InDrawRenderState);
@@ -410,15 +410,15 @@ void SetDepthStencilStateForBasePass(
 	{
 		if (bMaskedInEarlyPass)
 		{
-			SetDepthStencilStateForBasePass_Internal<false, CF_Equal>(DrawRenderState);
+			SetDepthStencilStateForBasePass_Internal<false, CF_Equal>(DrawRenderState, FeatureLevel);
 		}
 		else if (DrawRenderState.GetDepthStencilAccess() & FExclusiveDepthStencil::DepthWrite)
 		{
-			SetDepthStencilStateForBasePass_Internal<true, CF_GreaterEqual>(DrawRenderState);
+			SetDepthStencilStateForBasePass_Internal<true, CF_GreaterEqual>(DrawRenderState, FeatureLevel);
 		}
 		else
 		{
-			SetDepthStencilStateForBasePass_Internal<false, CF_GreaterEqual>(DrawRenderState);
+			SetDepthStencilStateForBasePass_Internal<false, CF_GreaterEqual>(DrawRenderState, FeatureLevel);
 		}
 	}
 	else if (bMaskedInEarlyPass)
@@ -428,7 +428,7 @@ void SetDepthStencilStateForBasePass(
 
 	if (bForceEnableStencilDitherState)
 	{
-		SetDepthStencilStateForBasePass_Internal<false, CF_Equal>(DrawRenderState);
+		SetDepthStencilStateForBasePass_Internal<false, CF_Equal>(DrawRenderState, FeatureLevel);
 	}
 }
 
@@ -1583,7 +1583,7 @@ bool FBasePassMeshProcessor::Process(
 	
 	if (bEnableReceiveDecalOutput)
 	{
-		static const bool bStrataDufferPassEnabled = Strata::IsStrataEnabled() && Strata::IsStrataDbufferPassEnabled();
+		static const bool bStrataDufferPassEnabled = Strata::IsStrataEnabled() && Strata::IsStrataDbufferPassEnabled(GShaderPlatformForFeatureLevel[FeatureLevel]);
 
 		// Set stencil value for this draw call
 		// This is effectively extending the GBuffer using the stencil bits
diff --git a/Engine/Source/Runtime/Renderer/Private/CompositionLighting/PostProcessDeferredDecals.cpp b/Engine/Source/Runtime/Renderer/Private/CompositionLighting/PostProcessDeferredDecals.cpp
--- a/Engine/Source/Runtime/Renderer/Private/CompositionLighting/PostProcessDeferredDecals.cpp
+++ b/Engine/Source/Runtime/Renderer/Private/CompositionLighting/PostProcessDeferredDecals.cpp
@@ -512,7 +512,11 @@ void AddDeferredDecalPass(
 		// Sanity check - Strata only support DBuffer, SceneColor, or AO decals
 		if (Strata::IsStrataEnabled())
 		{
-			check(RenderTargetMode == EDecalRenderTargetMode::DBuffer || RenderTargetMode == EDecalRenderTargetMode::SceneColor || RenderTargetMode == EDecalRenderTargetMode::AmbientOcclusion);
+			const bool bSupported = RenderTargetMode == EDecalRenderTargetMode::DBuffer || RenderTargetMode == EDecalRenderTargetMode::SceneColor || RenderTargetMode == EDecalRenderTargetMode::AmbientOcclusion;
+			if (!bSupported)
+			{
+				return;
+			}
 		}
 
 		auto* PassParameters = GraphBuilder.AllocParameters<FDeferredDecalPassParameters>();
diff --git a/Engine/Source/Runtime/Renderer/Private/Strata/Strata.cpp b/Engine/Source/Runtime/Renderer/Private/Strata/Strata.cpp
--- a/Engine/Source/Runtime/Renderer/Private/Strata/Strata.cpp
+++ b/Engine/Source/Runtime/Renderer/Private/Strata/Strata.cpp
@@ -143,9 +143,9 @@ bool IsStrataEnabled()
 	return CVarStrata.GetValueOnAnyThread() > 0;
 }
 
-bool IsStrataDbufferPassEnabled()
+bool IsStrataDbufferPassEnabled(const EShaderPlatform InPlatform)
 {
-	return IsStrataEnabled() && CVarStrataDBufferPass.GetValueOnAnyThread() > 0;
+	return IsStrataEnabled() && IsUsingDBuffers(InPlatform) && CVarStrataDBufferPass.GetValueOnAnyThread() > 0;
 }
 
 enum EStrataTileSpace
@@ -235,7 +235,7 @@ static void InitialiseStrataViewData(FRDGBuilder& GraphBuilder, FViewInfo& View,
 
 			// Separated subsurface & rough refraction textures (tile data)
 			const uint32 RoughTileCount = IsStrataOpaqueMaterialRoughRefractionEnabled() ? TileResolution.X * TileResolution.Y : 4;
-			const uint32 DecalTileCount = IsStrataDbufferPassEnabled() ? TileResolution.X * TileResolution.Y : 4;
+			const uint32 DecalTileCount = IsStrataDbufferPassEnabled(View.GetShaderPlatform()) ? TileResolution.X * TileResolution.Y : 4;
 			const uint32 RegularTileCount = TileResolution.X * TileResolution.Y;
 
 			const EPixelFormat ClassificationTileFormat = GetClassificationTileFormat(ViewResolution);
@@ -711,7 +711,7 @@ class FStrataDBufferPassCS : public FGlobalShader
 
 	static bool ShouldCompilePermutation(const FGlobalShaderPermutationParameters& Parameters)
 	{
-		return GetMaxSupportedFeatureLevel(Parameters.Platform) >= ERHIFeatureLevel::SM5 && Strata::IsStrataEnabled();
+		return GetMaxSupportedFeatureLevel(Parameters.Platform) >= ERHIFeatureLevel::SM5 && Strata::IsStrataEnabled() && IsUsingDBuffers(Parameters.Platform);
 	}
 
 	static void ModifyCompilationEnvironment(const FGlobalShaderPermutationParameters& Parameters, FShaderCompilerEnvironment& OutEnvironment)
diff --git a/Engine/Source/Runtime/Renderer/Private/Strata/Strata.h b/Engine/Source/Runtime/Renderer/Private/Strata/Strata.h
--- a/Engine/Source/Runtime/Renderer/Private/Strata/Strata.h
+++ b/Engine/Source/Runtime/Renderer/Private/Strata/Strata.h
@@ -148,7 +148,7 @@ constexpr uint32 StencilBit_Single = 0x10; // In sync with SceneRenderTargets.h
 constexpr uint32 StencilBit_Complex= 0x20; // In sync with SceneRenderTargets.h - GET_STENCIL_BIT_MASK(STENCIL_STRATA_COMPLEX)
 
 bool IsStrataEnabled();
-bool IsStrataDbufferPassEnabled();
+bool IsStrataDbufferPassEnabled(const EShaderPlatform Platform);
 
 FIntPoint GetStrataTextureResolution(const FIntPoint& InResolution);
 
