Strata - fixed Debug Visualization.

diff --git a/Engine/Shaders/Private/Strata/StrataVisualize.usf b/Engine/Shaders/Private/Strata/StrataVisualize.usf
--- a/Engine/Shaders/Private/Strata/StrataVisualize.usf
+++ b/Engine/Shaders/Private/Strata/StrataVisualize.usf
@@ -595,17 +595,17 @@ void VisualizeMaterialPS(
 	float4 SVPos : SV_POSITION,
 	out float4 OutColor : SV_Target0)
 {
-	const uint2 PixelPos = uint2(SVPos.xy);
-	float2 BufferUV = SvPositionToBufferUV(float4(PixelPos, SVPos.zw));
+	const uint2 PixelPosDynRes = uint2(float2(SVPos.xy) * View.ViewResolutionFraction);
+	float2 BufferUV = SvPositionToBufferUV(float4(PixelPosDynRes, SVPos.zw));
 
 	const float2 ScreenPosition = SvPositionToScreenPosition(SVPos).xy;
-	const float2 ScreenMaterialPosition = SvPositionToScreenPosition(float4(PixelPos, 0.5f, 1)).xy;
+	const float2 ScreenMaterialPosition = SvPositionToScreenPosition(float4(PixelPosDynRes, 0.5f, 1)).xy;
 	float DeviceZ = SampleDeviceZFromSceneTextures(BufferUV);
 	float SceneDepth = ConvertFromDeviceZ(DeviceZ);
 	const float3 WorldPosition = mul(float4(ScreenMaterialPosition * SceneDepth, SceneDepth, 1), PrimaryView.ScreenToTranslatedWorld).xyz;
 	const float3 V = -normalize(WorldPosition - PrimaryView.TranslatedWorldCameraOrigin);
 
-	FStrataAddressing StrataAddressing = GetStrataPixelDataByteOffset(PixelPos, uint2(View.BufferSizeAndInvSize.xy), Strata.MaxBytesPerPixel);
+	FStrataAddressing StrataAddressing = GetStrataPixelDataByteOffset(PixelPosDynRes, uint2(View.BufferSizeAndInvSize.xy), Strata.MaxBytesPerPixel);
 	FStrataPixelHeader Header = UnpackStrataHeaderIn(Strata.MaterialTextureArray, StrataAddressing, Strata.TopLayerTexture);
 	bool bStrataMaterial = Header.BSDFCount > 0;
 
diff --git a/Engine/Source/Runtime/Renderer/Private/Strata/Strata.cpp b/Engine/Source/Runtime/Renderer/Private/Strata/Strata.cpp
--- a/Engine/Source/Runtime/Renderer/Private/Strata/Strata.cpp
+++ b/Engine/Source/Runtime/Renderer/Private/Strata/Strata.cpp
@@ -908,7 +908,7 @@ void AddStrataInternalClassificationTilePass(
 	const bool bDebug = false)
 {
 	EPrimitiveType StrataTilePrimitiveType = PT_TriangleList;
-	const FIntPoint OutputResolution = View.ViewRect.Size();
+	FIntPoint DebugOutputResolution = FIntPoint(View.UnscaledViewRect.Width(), View.UnscaledViewRect.Height());
 	const FIntRect ViewRect = View.ViewRect;
 
 	FStrataMaterialStencilTaggingPassPS::FParameters* ParametersPS = GraphBuilder.AllocParameters<FStrataMaterialStencilTaggingPassPS::FParameters>();
@@ -923,6 +923,13 @@ void AddStrataInternalClassificationTilePass(
 	// For debug purpose
 	if (bDebug)
 	{
+		// ViewRect contains the scaled resolution according to TSR screen percentage.
+		// The ColorTexture can be larger than the screen resolution if the screen percentage has be manipulated to be >100%.
+		// So we simply re-use the previously computed ViewResolutionFraction to recover the targeted resolution in the editor.
+		// TODO fix this for split screen.
+		const float InvViewResolutionFraction = 1.0f / View.CachedViewUniformShaderParameters->ViewResolutionFraction;
+		DebugOutputResolution = FIntPoint(float(ViewRect.Width()) * InvViewResolutionFraction, float(ViewRect.Height()) * InvViewResolutionFraction);
+
 		check(ColorTexture);
 		ParametersPS->RenderTargets[0] = FRenderTargetBinding(*ColorTexture, ERenderTargetLoadAction::ELoad);
 		switch (TileMaterialType)
@@ -955,7 +962,7 @@ void AddStrataInternalClassificationTilePass(
 		RDG_EVENT_NAME("Strata::%sClassificationPass(%s)", bDebug ? TEXT("Debug") : TEXT("Stencil"), ToString(TileMaterialType)),
 		ParametersPS,
 		ERDGPassFlags::Raster,
-		[ParametersPS, VertexShader, PixelShader, ViewRect, OutputResolution, StrataTilePrimitiveType, TileMaterialType, bDebug](FRHICommandList& RHICmdList)
+		[ParametersPS, VertexShader, PixelShader, ViewRect, DebugOutputResolution, StrataTilePrimitiveType, TileMaterialType, bDebug](FRHICommandList& RHICmdList)
 		{
 			FGraphicsPipelineStateInitializer GraphicsPSOInit;
 			RHICmdList.ApplyCachedRenderTargets(GraphicsPSOInit);
@@ -1018,7 +1025,7 @@ void AddStrataInternalClassificationTilePass(
 			{
 				// Debug rendering is aways done during the post-processing stage, which has an ViewMinRect set to (0,0)
 				SetShaderParameters(RHICmdList, PixelShader, PixelShader.GetPixelShader(), *ParametersPS);
-				RHICmdList.SetViewport(0, 0, 0.0f, OutputResolution.X, OutputResolution.Y, 1.0f);
+				RHICmdList.SetViewport(0, 0, 0.0f, DebugOutputResolution.X, DebugOutputResolution.Y, 1.0f);
 			}
 			else
 			{
