Fix crash when enabling the shader complexity view mode.

The debug view mode shaders weren't binding the instanced view UB.

#rnx

diff --git a/Engine/Source/Runtime/Renderer/Private/DebugViewModeRendering.cpp b/Engine/Source/Runtime/Renderer/Private/DebugViewModeRendering.cpp
--- a/Engine/Source/Runtime/Renderer/Private/DebugViewModeRendering.cpp
+++ b/Engine/Source/Runtime/Renderer/Private/DebugViewModeRendering.cpp
@@ -156,7 +156,7 @@ bool FDebugViewModeVS::ShouldCompilePermutation(const FMeshMaterialShaderPermuta
 }
 
 BEGIN_SHADER_PARAMETER_STRUCT(FDebugViewModePassParameters, )
-	SHADER_PARAMETER_STRUCT_REF(FViewUniformShaderParameters, View)
+	SHADER_PARAMETER_STRUCT_INCLUDE(FViewShaderParameters, View)
 	SHADER_PARAMETER_RDG_UNIFORM_BUFFER(FDebugViewModePassUniformParameters, Pass)
 	SHADER_PARAMETER_STRUCT_INCLUDE(FInstanceCullingDrawParams, InstanceCullingDrawParams)
 	RENDER_TARGET_BINDING_SLOTS()
@@ -173,7 +173,7 @@ void RenderDebugViewMode(FRDGBuilder& GraphBuilder, TArrayView<FViewInfo> Views,
 		RDG_EVENT_SCOPE_CONDITIONAL(GraphBuilder, Views.Num() > 1, "View%d", ViewIndex);
 
 		auto* PassParameters = GraphBuilder.AllocParameters<FDebugViewModePassParameters>();
-		PassParameters->View = View.ViewUniformBuffer;
+		PassParameters->View = View.GetShaderParameters();
 		PassParameters->Pass = CreateDebugViewModePassUniformBuffer(GraphBuilder, View, QuadOverdrawTexture);
 		PassParameters->RenderTargets = RenderTargets;
 
