Change EnergyConservationTable to always use the offline-built tables, instead of the runtime generation.
* This avoids startup overhead
* This fixed missing energy tables for the very first frames (tables not ready).

diff --git a/Engine/Source/Runtime/Renderer/Private/ShadingEnergyConservation.cpp b/Engine/Source/Runtime/Renderer/Private/ShadingEnergyConservation.cpp
--- a/Engine/Source/Runtime/Renderer/Private/ShadingEnergyConservation.cpp
+++ b/Engine/Source/Runtime/Renderer/Private/ShadingEnergyConservation.cpp
@@ -229,11 +229,8 @@ void Init(FRDGBuilder& GraphBuilder, FViewInfo& View)
 	const bool bBindEnergyData = (View.ViewState != nullptr) && (bMaterialEnergyConservationEnabled || Strata::IsStrataEnabled() || (View.Family->EngineShowFlags.PathTracing && RHI_RAYTRACING)) && (bIsEnergyPreservationEnabled || bIsEnergyConservationEnabled);
 	if (bBindEnergyData)
 	{
-		#if UE_EDITOR
-		const bool bRuntimeGeneration = true;
-		#else
+		// Change this to true in order to regenerate the energy tables, and manually copy the coefficients into ShadingEnergyConservationData.h
 		const bool bRuntimeGeneration = false;
-		#endif
 
 		const int Size = SHADING_ENERGY_CONSERVATION_TABLE_RESOLUTION;
 		const EPixelFormat SpecFormat = bRuntimeGeneration && CVarShadingFurnaceTest_TableFormat.GetValueOnRenderThread() > 0 ? PF_G32R32F : PF_G16R16;
