Fix missing lighting on hair when there is a mixture of shadow casting & non-shadow casting lights.

diff --git a/Engine/Source/Runtime/Renderer/Private/LightRendering.cpp b/Engine/Source/Runtime/Renderer/Private/LightRendering.cpp
--- a/Engine/Source/Runtime/Renderer/Private/LightRendering.cpp
+++ b/Engine/Source/Runtime/Renderer/Private/LightRendering.cpp
@@ -1754,9 +1754,15 @@ void FDeferredShadingSceneRenderer::RenderLights(
 				{
 					for (FViewInfo& View : Views)
 					{
-						if (bDrawHairShadow && HairStrands::HasViewHairStrandsData(View))
+						if (HairStrands::HasViewHairStrandsData(View))
 						{
-							FHairStrandsTransmittanceMaskData TransmittanceMaskData = RenderHairStrandsTransmittanceMask(GraphBuilder, View, &LightSceneInfo, false, ScreenShadowMaskSubPixelTexture);
+							FHairStrandsTransmittanceMaskData TransmittanceMaskData;
+							FRDGTextureRef HairShadowMask = nullptr;
+							if (bDrawHairShadow)
+							{
+								TransmittanceMaskData = RenderHairStrandsTransmittanceMask(GraphBuilder, View, &LightSceneInfo, false, ScreenShadowMaskSubPixelTexture);
+								HairShadowMask = ScreenShadowMaskSubPixelTexture;
+							}
 							if (TransmittanceMaskData.TransmittanceMask == nullptr)
 							{
 								TransmittanceMaskData = DummyTransmittanceMaskData;
@@ -1766,7 +1772,7 @@ void FDeferredShadingSceneRenderer::RenderLights(
 
 							// Note: ideally the light should still be evaluated for hair when not casting shadow, but for preserving the old behavior, and not adding 
 							// any perf. regression, we disable this light for hair rendering 
-							RenderLightForHair(GraphBuilder, View, SceneTextures, &LightSceneInfo, ScreenShadowMaskSubPixelTexture, LightingChannelsTexture, TransmittanceMaskData, false /*bForwardRendering*/);
+							RenderLightForHair(GraphBuilder, View, SceneTextures, &LightSceneInfo, HairShadowMask, LightingChannelsTexture, TransmittanceMaskData, false /*bForwardRendering*/);
 						}
 					}
 				}
