Renamed TPT_Transluceny enum for consistency.
Follow up from https://p4-swarm.epicgames.net/reviews/21888399

diff --git a/Engine/Source/Runtime/Renderer/Private/BasePassRendering.cpp b/Engine/Source/Runtime/Renderer/Private/BasePassRendering.cpp
--- a/Engine/Source/Runtime/Renderer/Private/BasePassRendering.cpp
+++ b/Engine/Source/Runtime/Renderer/Private/BasePassRendering.cpp
@@ -185,7 +185,7 @@ void SetTranslucentRenderState(FMeshPassProcessorRenderState& DrawRenderState, c
 	{
 		if (Material.IsDualBlendingEnabled(Platform))
 		{
-			if (InTranslucencyPassType == ETranslucencyPass::TPT_StandardTranslucency || InTranslucencyPassType == ETranslucencyPass::TPT_AllTranslucency)
+			if (InTranslucencyPassType == ETranslucencyPass::TPT_TranslucencyStandard || InTranslucencyPassType == ETranslucencyPass::TPT_AllTranslucency)
 			{
 				// If we are in the transparancy pass (before DoF) we do standard dual blending, and the alpha gets ignored
 				// Blend by putting add in target 0 and multiply by background in target 1.
@@ -197,7 +197,7 @@ void SetTranslucentRenderState(FMeshPassProcessorRenderState& DrawRenderState, c
 				// Alpha is BF_Zero for source and BF_One for dest, which leaves alpha unchanged
 				DrawRenderState.SetBlendState(TStaticBlendState<CW_RGBA, BO_Add, BF_One, BF_Source1Color, BO_Add, BF_Zero, BF_One>::GetRHI());
 			}
-			else if (InTranslucencyPassType == ETranslucencyPass::TPT_TranslucencyAfterDOFModulate || InTranslucencyPassType == ETranslucencyPass::TPT_StandardTranslucencyModulate)
+			else if (InTranslucencyPassType == ETranslucencyPass::TPT_TranslucencyAfterDOFModulate || InTranslucencyPassType == ETranslucencyPass::TPT_TranslucencyStandardModulate)
 			{
 				// In the separate pass (after DoF) modulate, we want to only darken the target by our multiplication term, and ignore the addition term.
 				// For regular dual blending, our function is:
@@ -237,7 +237,7 @@ void SetTranslucentRenderState(FMeshPassProcessorRenderState& DrawRenderState, c
 		// Special case for dual blending, which is not exposed as a parameter in the material editor
 		if (Material.IsDualBlendingEnabled(Platform))
 		{
-			if (InTranslucencyPassType == ETranslucencyPass::TPT_StandardTranslucency || InTranslucencyPassType == ETranslucencyPass::TPT_AllTranslucency)
+			if (InTranslucencyPassType == ETranslucencyPass::TPT_TranslucencyStandard || InTranslucencyPassType == ETranslucencyPass::TPT_AllTranslucency)
 			{
 				// If we are in the transparancy pass (before DoF) we do standard dual blending, and the alpha gets ignored
 
@@ -250,7 +250,7 @@ void SetTranslucentRenderState(FMeshPassProcessorRenderState& DrawRenderState, c
 				// Alpha is BF_Zero for source and BF_One for dest, which leaves alpha unchanged
 				DrawRenderState.SetBlendState(TStaticBlendState<CW_RGBA, BO_Add, BF_One, BF_Source1Color, BO_Add, BF_Zero, BF_One>::GetRHI());
 			}
-			else if (InTranslucencyPassType == ETranslucencyPass::TPT_TranslucencyAfterDOFModulate || InTranslucencyPassType == ETranslucencyPass::TPT_StandardTranslucencyModulate)
+			else if (InTranslucencyPassType == ETranslucencyPass::TPT_TranslucencyAfterDOFModulate || InTranslucencyPassType == ETranslucencyPass::TPT_TranslucencyStandardModulate)
 			{
 				// In the separate pass (after DoF) modulate, we want to only darken the target by our multiplication term, and ignore the addition term.
 				// For regular dual blending, our function is:
@@ -313,7 +313,7 @@ void SetTranslucentRenderState(FMeshPassProcessorRenderState& DrawRenderState, c
 	const bool bIsPostMotionBlur = Material.IsTranslucencyAfterMotionBlurEnabled();
 
 	// When separate standard translucent are used, we must mark the distoprtion bit for the composition to happen correctly for any BeforeDoF translucent.
-	const bool bSeparatedStandardTranslucent = (InTranslucencyPassType == ETranslucencyPass::TPT_StandardTranslucency || InTranslucencyPassType == ETranslucencyPass::TPT_AllTranslucency) && IsStandardTranslucenyPassSeparated();
+	const bool bSeparatedStandardTranslucent = (InTranslucencyPassType == ETranslucencyPass::TPT_TranslucencyStandard || InTranslucencyPassType == ETranslucencyPass::TPT_AllTranslucency) && IsStandardTranslucenyPassSeparated();
 
 	if (bEnableResponsiveAA && !bIsPostMotionBlur)
 	{
@@ -1695,11 +1695,11 @@ bool FBasePassMeshProcessor::ShouldDraw(const FMaterial& Material)
 		{
 			switch (TranslucencyPassType)
 			{
-			case ETranslucencyPass::TPT_StandardTranslucency:
+			case ETranslucencyPass::TPT_TranslucencyStandard:
 				bShouldDraw = !Material.IsTranslucencyAfterDOFEnabled() && !Material.IsTranslucencyAfterMotionBlurEnabled();
 				break;
 
-			case ETranslucencyPass::TPT_StandardTranslucencyModulate:
+			case ETranslucencyPass::TPT_TranslucencyStandardModulate:
 				bShouldDraw = !Material.IsTranslucencyAfterDOFEnabled() && !Material.IsTranslucencyAfterMotionBlurEnabled() 
 					&& (Material.IsDualBlendingEnabled(GetFeatureLevelShaderPlatform(FeatureLevel)) || BlendMode == BLEND_Modulate || StrataBlendMode == SBM_ColoredTransmittanceOnly);
 				break;
@@ -2166,7 +2166,7 @@ FMeshPassProcessor* CreateTranslucencyStandardPassProcessor(ERHIFeatureLevel::Ty
 
 	const FBasePassMeshProcessor::EFlags Flags = FBasePassMeshProcessor::EFlags::CanUseDepthStencil;
 
-	return new FBasePassMeshProcessor(EMeshPass::TranslucencyStandard, Scene, FeatureLevel, InViewIfDynamicMeshCommand, PassDrawRenderState, InDrawListContext, Flags, ETranslucencyPass::TPT_StandardTranslucency);
+	return new FBasePassMeshProcessor(EMeshPass::TranslucencyStandard, Scene, FeatureLevel, InViewIfDynamicMeshCommand, PassDrawRenderState, InDrawListContext, Flags, ETranslucencyPass::TPT_TranslucencyStandard);
 }
 
 FMeshPassProcessor* CreateTranslucencyStandardModulatePassProcessor(ERHIFeatureLevel::Type FeatureLevel, const FScene* Scene, const FSceneView* InViewIfDynamicMeshCommand, FMeshPassDrawListContext* InDrawListContext)
@@ -2176,7 +2176,7 @@ FMeshPassProcessor* CreateTranslucencyStandardModulatePassProcessor(ERHIFeatureL
 
 	const FBasePassMeshProcessor::EFlags Flags = FBasePassMeshProcessor::EFlags::CanUseDepthStencil;
 
-	return new FBasePassMeshProcessor(EMeshPass::TranslucencyStandardModulate, Scene, FeatureLevel, InViewIfDynamicMeshCommand, PassDrawRenderState, InDrawListContext, Flags, ETranslucencyPass::TPT_StandardTranslucencyModulate);
+	return new FBasePassMeshProcessor(EMeshPass::TranslucencyStandardModulate, Scene, FeatureLevel, InViewIfDynamicMeshCommand, PassDrawRenderState, InDrawListContext, Flags, ETranslucencyPass::TPT_TranslucencyStandardModulate);
 }
 
 FMeshPassProcessor* CreateTranslucencyAfterDOFProcessor(ERHIFeatureLevel::Type FeatureLevel, const FScene* Scene, const FSceneView* InViewIfDynamicMeshCommand, FMeshPassDrawListContext* InDrawListContext)
diff --git a/Engine/Source/Runtime/Renderer/Private/DistortionRendering.cpp b/Engine/Source/Runtime/Renderer/Private/DistortionRendering.cpp
--- a/Engine/Source/Runtime/Renderer/Private/DistortionRendering.cpp
+++ b/Engine/Source/Runtime/Renderer/Private/DistortionRendering.cpp
@@ -834,7 +834,7 @@ void FDeferredShadingSceneRenderer::RenderDistortion(
 			PassParameters->RenderTargets[0] = FRenderTargetBinding(SceneColorTexture, ERenderTargetLoadAction::ELoad);
 
 			// Fetch the separated before dof translucent data if available for composition
-			FTranslucencyPassResources& TranslucencyPassResources = TranslucencyResourceMap.Get(ViewIndex, ETranslucencyPass::TPT_StandardTranslucency);
+			FTranslucencyPassResources& TranslucencyPassResources = TranslucencyResourceMap.Get(ViewIndex, ETranslucencyPass::TPT_TranslucencyStandard);
 			PassParameters->StandardTranslucentColorAlphaTexture = TranslucencyPassResources.GetColorForRead(GraphBuilder);
 			PassParameters->StandardTranslucentColorAlphaMSAATexture = TranslucencyPassResources.ColorTexture.Target;
 			PassParameters->StandardTranslucentTransmittanceTexture = TranslucencyPassResources.GetColorModulateForRead(GraphBuilder);
diff --git a/Engine/Source/Runtime/Renderer/Private/EditorPrimitivesRendering.cpp b/Engine/Source/Runtime/Renderer/Private/EditorPrimitivesRendering.cpp
--- a/Engine/Source/Runtime/Renderer/Private/EditorPrimitivesRendering.cpp
+++ b/Engine/Source/Runtime/Renderer/Private/EditorPrimitivesRendering.cpp
@@ -93,7 +93,7 @@ bool FEditorPrimitivesBasePassMeshProcessor::ProcessDeferredShadingPath(const FM
 	if (bTranslucentBasePass)
 	{
 		extern void SetTranslucentRenderState(FMeshPassProcessorRenderState& DrawRenderState, const FMaterial& Material, const EShaderPlatform Platform, ETranslucencyPass::Type InTranslucencyPassType);
-		SetTranslucentRenderState(DrawRenderState, Material, Scene->GetShaderPlatform(), ETranslucencyPass::TPT_StandardTranslucency);
+		SetTranslucentRenderState(DrawRenderState, Material, Scene->GetShaderPlatform(), ETranslucencyPass::TPT_TranslucencyStandard);
 	}
 
 	const FMeshDrawingPolicyOverrideSettings OverrideSettings = ComputeMeshOverrideSettings(MeshBatch);
diff --git a/Engine/Source/Runtime/Renderer/Private/HairStrands/HairStrandsComposition.cpp b/Engine/Source/Runtime/Renderer/Private/HairStrands/HairStrandsComposition.cpp
--- a/Engine/Source/Runtime/Renderer/Private/HairStrands/HairStrandsComposition.cpp
+++ b/Engine/Source/Runtime/Renderer/Private/HairStrands/HairStrandsComposition.cpp
@@ -622,7 +622,7 @@ void RenderHairComposition(
 	{		
 		if (View.Family && HairStrands::HasViewHairStrandsData(View))
 		{
-			FTranslucencyPassResources& TranslucencyPassResources = TranslucencyResourceMap.Get(ViewIndex, ETranslucencyPass::TPT_StandardTranslucency);
+			FTranslucencyPassResources& TranslucencyPassResources = TranslucencyResourceMap.Get(ViewIndex, ETranslucencyPass::TPT_TranslucencyStandard);
 			FRDGTextureRef SceneColorTarget = TranslucencyPassResources.ColorTexture.IsValid() ? TranslucencyPassResources.ColorTexture.Target : SceneColorTexture;
 
 			InternalRenderHairComposition(
diff --git a/Engine/Source/Runtime/Renderer/Private/MeshDrawCommands.cpp b/Engine/Source/Runtime/Renderer/Private/MeshDrawCommands.cpp
--- a/Engine/Source/Runtime/Renderer/Private/MeshDrawCommands.cpp
+++ b/Engine/Source/Runtime/Renderer/Private/MeshDrawCommands.cpp
@@ -1195,12 +1195,12 @@ void FParallelMeshDrawCommandPass::DispatchPassSetup(
 
 	switch (PassType)
 	{
-		case EMeshPass::TranslucencyStandard: TaskContext.TranslucencyPass = ETranslucencyPass::TPT_StandardTranslucency; break;
-		case EMeshPass::TranslucencyStandardModulate: TaskContext.TranslucencyPass = ETranslucencyPass::TPT_StandardTranslucencyModulate; break;
-		case EMeshPass::TranslucencyAfterDOF: TaskContext.TranslucencyPass = ETranslucencyPass::TPT_TranslucencyAfterDOF; break;
-		case EMeshPass::TranslucencyAfterDOFModulate: TaskContext.TranslucencyPass = ETranslucencyPass::TPT_TranslucencyAfterDOFModulate; break;
-		case EMeshPass::TranslucencyAfterMotionBlur: TaskContext.TranslucencyPass = ETranslucencyPass::TPT_TranslucencyAfterMotionBlur; break;
-		case EMeshPass::TranslucencyAll: TaskContext.TranslucencyPass = ETranslucencyPass::TPT_AllTranslucency; break;
+		case EMeshPass::TranslucencyStandard: TaskContext.TranslucencyPass			= ETranslucencyPass::TPT_TranslucencyStandard; break;
+		case EMeshPass::TranslucencyStandardModulate: TaskContext.TranslucencyPass	= ETranslucencyPass::TPT_TranslucencyStandardModulate; break;
+		case EMeshPass::TranslucencyAfterDOF: TaskContext.TranslucencyPass			= ETranslucencyPass::TPT_TranslucencyAfterDOF; break;
+		case EMeshPass::TranslucencyAfterDOFModulate: TaskContext.TranslucencyPass	= ETranslucencyPass::TPT_TranslucencyAfterDOFModulate; break;
+		case EMeshPass::TranslucencyAfterMotionBlur: TaskContext.TranslucencyPass	= ETranslucencyPass::TPT_TranslucencyAfterMotionBlur; break;
+		case EMeshPass::TranslucencyAll: TaskContext.TranslucencyPass				= ETranslucencyPass::TPT_AllTranslucency; break;
 	}
 
 	FMemory::Memswap(&TaskContext.MeshDrawCommands, &InOutMeshDrawCommands, sizeof(InOutMeshDrawCommands));
diff --git a/Engine/Source/Runtime/Renderer/Private/MobileBasePass.cpp b/Engine/Source/Runtime/Renderer/Private/MobileBasePass.cpp
--- a/Engine/Source/Runtime/Renderer/Private/MobileBasePass.cpp
+++ b/Engine/Source/Runtime/Renderer/Private/MobileBasePass.cpp
@@ -606,8 +606,8 @@ bool FMobileBasePassMeshProcessor::TryAddMeshBatch(const FMeshBatch& RESTRICT Me
 	{
 		// Skipping TPT_TranslucencyAfterDOFModulate. That pass is only needed for Dual Blending, which is not supported on Mobile.
 		bool bShouldDraw = (bIsTranslucent || bUsesWaterMaterial) &&
-		(TranslucencyPassType == ETranslucencyPass::TPT_AllTranslucency
-		|| (TranslucencyPassType == ETranslucencyPass::TPT_StandardTranslucency && !Material.IsMobileSeparateTranslucencyEnabled())
+		(   TranslucencyPassType == ETranslucencyPass::TPT_AllTranslucency
+		|| (TranslucencyPassType == ETranslucencyPass::TPT_TranslucencyStandard && !Material.IsMobileSeparateTranslucencyEnabled())
 		|| (TranslucencyPassType == ETranslucencyPass::TPT_TranslucencyAfterDOF && Material.IsMobileSeparateTranslucencyEnabled()));
 
 		if (bShouldDraw)
@@ -808,7 +808,7 @@ FMeshPassProcessor* CreateMobileTranslucencyStandardPassProcessor(ERHIFeatureLev
 
 	const FMobileBasePassMeshProcessor::EFlags Flags = FMobileBasePassMeshProcessor::EFlags::CanUseDepthStencil;
 
-	return new FMobileBasePassMeshProcessor(EMeshPass::TranslucencyStandard, Scene, FeatureLevel, InViewIfDynamicMeshCommand, PassDrawRenderState, InDrawListContext, Flags, ETranslucencyPass::TPT_StandardTranslucency);
+	return new FMobileBasePassMeshProcessor(EMeshPass::TranslucencyStandard, Scene, FeatureLevel, InViewIfDynamicMeshCommand, PassDrawRenderState, InDrawListContext, Flags, ETranslucencyPass::TPT_TranslucencyStandard);
 }
 
 FMeshPassProcessor* CreateMobileTranslucencyAfterDOFProcessor(ERHIFeatureLevel::Type FeatureLevel, const FScene* Scene, const FSceneView* InViewIfDynamicMeshCommand, FMeshPassDrawListContext* InDrawListContext)
diff --git a/Engine/Source/Runtime/Renderer/Private/MobileShadingRenderer.cpp b/Engine/Source/Runtime/Renderer/Private/MobileShadingRenderer.cpp
--- a/Engine/Source/Runtime/Renderer/Private/MobileShadingRenderer.cpp
+++ b/Engine/Source/Runtime/Renderer/Private/MobileShadingRenderer.cpp
@@ -271,7 +271,7 @@ FMobileSceneRenderer::FMobileSceneRenderer(const FSceneViewFamily* InViewFamily,
 	bEnableClusteredLocalLights = MobileForwardEnableLocalLights(ShaderPlatform);
 	bEnableClusteredReflections = MobileForwardEnableClusteredReflections(ShaderPlatform);
 	
-	StandardTranslucencyPass = ViewFamily.AllowTranslucencyAfterDOF() ? ETranslucencyPass::TPT_StandardTranslucency : ETranslucencyPass::TPT_AllTranslucency;
+	StandardTranslucencyPass = ViewFamily.AllowTranslucencyAfterDOF() ? ETranslucencyPass::TPT_TranslucencyStandard : ETranslucencyPass::TPT_AllTranslucency;
 	StandardTranslucencyMeshPass = TranslucencyPassToMeshPass(StandardTranslucencyPass);
 
 	// Don't do occlusion queries when doing scene captures
diff --git a/Engine/Source/Runtime/Renderer/Private/PostProcess/PostProcessCompositeEditorPrimitives.cpp b/Engine/Source/Runtime/Renderer/Private/PostProcess/PostProcessCompositeEditorPrimitives.cpp
--- a/Engine/Source/Runtime/Renderer/Private/PostProcess/PostProcessCompositeEditorPrimitives.cpp
+++ b/Engine/Source/Runtime/Renderer/Private/PostProcess/PostProcessCompositeEditorPrimitives.cpp
@@ -173,9 +173,9 @@ void RenderForegroundTranslucentEditorPrimitives(FRHICommandList& RHICmdList, co
 	const auto ShaderPlatform = GShaderPlatformForFeatureLevel[FeatureLevel];
 
 	// Force all translucent editor primitives to standard translucent rendering
-	const ETranslucencyPass::Type TranslucencyPass = ETranslucencyPass::TPT_StandardTranslucency;
+	const ETranslucencyPass::Type TranslucencyPass = ETranslucencyPass::TPT_TranslucencyStandard;
 
-	if (TranslucencyPass == ETranslucencyPass::TPT_StandardTranslucency)
+	if (TranslucencyPass == ETranslucencyPass::TPT_TranslucencyStandard)
 	{
 		DrawRenderState.SetDepthStencilState(TStaticDepthStencilState<false, CF_DepthNearOrEqual>::GetRHI());
 	}
diff --git a/Engine/Source/Runtime/Renderer/Private/RayTracing/RayTracingTranslucency.cpp b/Engine/Source/Runtime/Renderer/Private/RayTracing/RayTracingTranslucency.cpp
--- a/Engine/Source/Runtime/Renderer/Private/RayTracing/RayTracingTranslucency.cpp
+++ b/Engine/Source/Runtime/Renderer/Private/RayTracing/RayTracingTranslucency.cpp
@@ -144,8 +144,8 @@ bool ShouldRenderRayTracingTranslucency(const FViewInfo& View)
 
 void FDeferredShadingSceneRenderer::RenderRayTracingTranslucency(FRDGBuilder& GraphBuilder, FRDGTextureMSAA SceneColorTexture)
 {
-	if (   !ShouldRenderTranslucency(ETranslucencyPass::TPT_StandardTranslucency)
-		&& !ShouldRenderTranslucency(ETranslucencyPass::TPT_StandardTranslucencyModulate)
+	if (   !ShouldRenderTranslucency(ETranslucencyPass::TPT_TranslucencyStandard)
+		&& !ShouldRenderTranslucency(ETranslucencyPass::TPT_TranslucencyStandardModulate)
 		&& !ShouldRenderTranslucency(ETranslucencyPass::TPT_TranslucencyAfterDOF)
 		&& !ShouldRenderTranslucency(ETranslucencyPass::TPT_TranslucencyAfterDOFModulate)
 		&& !ShouldRenderTranslucency(ETranslucencyPass::TPT_TranslucencyAfterMotionBlur)
diff --git a/Engine/Source/Runtime/Renderer/Private/SceneVisibility.cpp b/Engine/Source/Runtime/Renderer/Private/SceneVisibility.cpp
--- a/Engine/Source/Runtime/Renderer/Private/SceneVisibility.cpp
+++ b/Engine/Source/Runtime/Renderer/Private/SceneVisibility.cpp
@@ -2429,12 +2429,12 @@ struct FRelevancePacket : public FSceneRenderingAllocatorObject<FRelevancePacket
 				{
 					if (ViewRelevance.bNormalTranslucency)
 					{
-						TranslucentPrimCount.Add(ETranslucencyPass::TPT_StandardTranslucency, ViewRelevance.bUsesSceneColorCopy);
+						TranslucentPrimCount.Add(ETranslucencyPass::TPT_TranslucencyStandard, ViewRelevance.bUsesSceneColorCopy);
 					}
 
 					if (ViewRelevance.bNormalTranslucency && ViewRelevance.bTranslucencyModulate && View.Family->AllowStandardTranslucencySeparated())
 					{
-						TranslucentPrimCount.Add(ETranslucencyPass::TPT_StandardTranslucencyModulate, ViewRelevance.bUsesSceneColorCopy);
+						TranslucentPrimCount.Add(ETranslucencyPass::TPT_TranslucencyStandardModulate, ViewRelevance.bUsesSceneColorCopy);
 					}
 
 					if (ViewRelevance.bSeparateTranslucency)
diff --git a/Engine/Source/Runtime/Renderer/Private/TranslucencyPass.h b/Engine/Source/Runtime/Renderer/Private/TranslucencyPass.h
--- a/Engine/Source/Runtime/Renderer/Private/TranslucencyPass.h
+++ b/Engine/Source/Runtime/Renderer/Private/TranslucencyPass.h
@@ -7,8 +7,8 @@ namespace ETranslucencyPass
 {
 	enum Type
 	{
-		TPT_StandardTranslucency,
-		TPT_StandardTranslucencyModulate,
+		TPT_TranslucencyStandard,
+		TPT_TranslucencyStandardModulate,
 		TPT_TranslucencyAfterDOF,
 		TPT_TranslucencyAfterDOFModulate,
 		TPT_TranslucencyAfterMotionBlur,
diff --git a/Engine/Source/Runtime/Renderer/Private/TranslucentRendering.cpp b/Engine/Source/Runtime/Renderer/Private/TranslucentRendering.cpp
--- a/Engine/Source/Runtime/Renderer/Private/TranslucentRendering.cpp
+++ b/Engine/Source/Runtime/Renderer/Private/TranslucentRendering.cpp
@@ -149,12 +149,12 @@ EMeshPass::Type TranslucencyPassToMeshPass(ETranslucencyPass::Type TranslucencyP
 
 	switch (TranslucencyPass)
 	{
-	case ETranslucencyPass::TPT_StandardTranslucency: TranslucencyMeshPass = EMeshPass::TranslucencyStandard; break;
-	case ETranslucencyPass::TPT_StandardTranslucencyModulate: TranslucencyMeshPass = EMeshPass::TranslucencyStandardModulate; break;
-	case ETranslucencyPass::TPT_TranslucencyAfterDOF: TranslucencyMeshPass = EMeshPass::TranslucencyAfterDOF; break;
-	case ETranslucencyPass::TPT_TranslucencyAfterDOFModulate: TranslucencyMeshPass = EMeshPass::TranslucencyAfterDOFModulate; break;
-	case ETranslucencyPass::TPT_TranslucencyAfterMotionBlur: TranslucencyMeshPass = EMeshPass::TranslucencyAfterMotionBlur; break;
-	case ETranslucencyPass::TPT_AllTranslucency: TranslucencyMeshPass = EMeshPass::TranslucencyAll; break;
+	case ETranslucencyPass::TPT_TranslucencyStandard:			TranslucencyMeshPass = EMeshPass::TranslucencyStandard; break;
+	case ETranslucencyPass::TPT_TranslucencyStandardModulate:	TranslucencyMeshPass = EMeshPass::TranslucencyStandardModulate; break;
+	case ETranslucencyPass::TPT_TranslucencyAfterDOF:			TranslucencyMeshPass = EMeshPass::TranslucencyAfterDOF; break;
+	case ETranslucencyPass::TPT_TranslucencyAfterDOFModulate:	TranslucencyMeshPass = EMeshPass::TranslucencyAfterDOFModulate; break;
+	case ETranslucencyPass::TPT_TranslucencyAfterMotionBlur:	TranslucencyMeshPass = EMeshPass::TranslucencyAfterMotionBlur; break;
+	case ETranslucencyPass::TPT_AllTranslucency:				TranslucencyMeshPass = EMeshPass::TranslucencyAll; break;
 	}
 
 	check(TranslucencyMeshPass != EMeshPass::Num);
@@ -186,7 +186,7 @@ ETranslucencyView GetTranslucencyViews(TArrayView<const FViewInfo> Views)
 /** Mostly used to know if debug rendering should be drawn in this pass */
 static bool IsMainTranslucencyPass(ETranslucencyPass::Type TranslucencyPass)
 {
-	return TranslucencyPass == ETranslucencyPass::TPT_AllTranslucency || TranslucencyPass == ETranslucencyPass::TPT_StandardTranslucency;
+	return TranslucencyPass == ETranslucencyPass::TPT_AllTranslucency || TranslucencyPass == ETranslucencyPass::TPT_TranslucencyStandard;
 }
 
 static bool IsParallelTranslucencyEnabled()
@@ -202,9 +202,9 @@ static bool IsTranslucencyWaitForTasksEnabled()
 bool IsSeparateTranslucencyEnabled(ETranslucencyPass::Type TranslucencyPass, float DownsampleScale)
 {
 	// Currently AfterDOF is rendered earlier in the frame and must be rendered in a separate texture.
-	if (TranslucencyPass == ETranslucencyPass::TPT_TranslucencyAfterDOF
+	if (   TranslucencyPass == ETranslucencyPass::TPT_TranslucencyAfterDOF
 		|| TranslucencyPass == ETranslucencyPass::TPT_TranslucencyAfterDOFModulate
-		|| TranslucencyPass == ETranslucencyPass::TPT_StandardTranslucencyModulate
+		|| TranslucencyPass == ETranslucencyPass::TPT_TranslucencyStandardModulate
 		|| TranslucencyPass == ETranslucencyPass::TPT_TranslucencyAfterMotionBlur
 		)
 	{
@@ -685,7 +685,7 @@ bool FSceneRenderer::ShouldRenderTranslucency(ETranslucencyPass::Type Translucen
 		}
 	}
 
-	// If lightshafts are rendered in low res, we must reset the offscreen buffer in case is was also used in TPT_StandardTranslucency.
+	// If lightshafts are rendered in low res, we must reset the offscreen buffer in case is was also used in TPT_TranslucencyStandard.
 	if (GLightShaftRenderAfterDOF && TranslucencyPass == ETranslucencyPass::TPT_TranslucencyAfterDOF)
 	{
 		return true;
@@ -1258,11 +1258,11 @@ void FDeferredShadingSceneRenderer::RenderTranslucencyInner(
 
 	RDG_WAIT_FOR_TASKS_CONDITIONAL(GraphBuilder, IsTranslucencyWaitForTasksEnabled());
 
-	const bool bIsModulate = TranslucencyPass == ETranslucencyPass::TPT_TranslucencyAfterDOFModulate || TranslucencyPass == ETranslucencyPass::TPT_StandardTranslucencyModulate;
+	const bool bIsModulate = TranslucencyPass == ETranslucencyPass::TPT_TranslucencyAfterDOFModulate || TranslucencyPass == ETranslucencyPass::TPT_TranslucencyStandardModulate;
 	const bool bDepthTest = TranslucencyPass != ETranslucencyPass::TPT_TranslucencyAfterMotionBlur;
 	const bool bRenderInParallel = IsParallelTranslucencyEnabled();
 	const bool bIsScalingTranslucency = SeparateTranslucencyDimensions.Scale < 1.0f;
-	const bool bIsStandardSeparatedTranslucency = bStandardTranslucentCanRenderSeparate && TranslucencyPass == ETranslucencyPass::TPT_StandardTranslucency && ViewFamily.AllowStandardTranslucencySeparated();
+	const bool bIsStandardSeparatedTranslucency = bStandardTranslucentCanRenderSeparate && TranslucencyPass == ETranslucencyPass::TPT_TranslucencyStandard && ViewFamily.AllowStandardTranslucencySeparated();
 	const bool bRenderInSeparateTranslucency = IsSeparateTranslucencyEnabled(TranslucencyPass, SeparateTranslucencyDimensions.Scale) || bIsStandardSeparatedTranslucency;
 
 	// Can't reference scene color in scene textures. Scene color copy is used instead.
@@ -1378,9 +1378,9 @@ void FDeferredShadingSceneRenderer::RenderTranslucencyInner(
 					ensure(TranslucencyPassResources.DepthTexture == SharedDepthTexture);
 					TranslucencyPassResources.ColorModulateTexture = SharedColorTexture;
 				}
-				else if (TranslucencyPass == ETranslucencyPass::TPT_StandardTranslucencyModulate)
+				else if (TranslucencyPass == ETranslucencyPass::TPT_TranslucencyStandardModulate)
 				{
-					FTranslucencyPassResources& TranslucencyPassResources = OutTranslucencyResourceMap->Get(ViewIndex, ETranslucencyPass::TPT_StandardTranslucency);
+					FTranslucencyPassResources& TranslucencyPassResources = OutTranslucencyResourceMap->Get(ViewIndex, ETranslucencyPass::TPT_TranslucencyStandard);
 					ensure(TranslucencyPassResources.ViewRect == ScaledViewRect);
 					ensure(TranslucencyPassResources.DepthTexture == SharedDepthTexture);
 					TranslucencyPassResources.ColorModulateTexture = SharedColorTexture;
@@ -1514,10 +1514,10 @@ void FDeferredShadingSceneRenderer::RenderTranslucency(
 
 	if (ViewFamily.AllowTranslucencyAfterDOF())
 	{
-		RenderTranslucencyInner(GraphBuilder, SceneTextures, TranslucentLightingVolumeTextures, OutTranslucencyResourceMap, SharedDepthTexture, ViewsToRender, SceneColorCopyTexture, ETranslucencyPass::TPT_StandardTranslucency, InstanceCullingManager, bStandardTranslucentCanRenderSeparate);
+		RenderTranslucencyInner(GraphBuilder, SceneTextures, TranslucentLightingVolumeTextures, OutTranslucencyResourceMap, SharedDepthTexture, ViewsToRender, SceneColorCopyTexture, ETranslucencyPass::TPT_TranslucencyStandard, InstanceCullingManager, bStandardTranslucentCanRenderSeparate);
 		if (ViewFamily.AllowStandardTranslucencySeparated() && bStandardTranslucentCanRenderSeparate)
 		{
-			RenderTranslucencyInner(GraphBuilder, SceneTextures, TranslucentLightingVolumeTextures, OutTranslucencyResourceMap, SharedDepthTexture, ViewsToRender, SceneColorCopyTexture, ETranslucencyPass::TPT_StandardTranslucencyModulate, InstanceCullingManager, bStandardTranslucentCanRenderSeparate);
+			RenderTranslucencyInner(GraphBuilder, SceneTextures, TranslucentLightingVolumeTextures, OutTranslucencyResourceMap, SharedDepthTexture, ViewsToRender, SceneColorCopyTexture, ETranslucencyPass::TPT_TranslucencyStandardModulate, InstanceCullingManager, bStandardTranslucentCanRenderSeparate);
 		}
 
 		if (GetHairStrandsComposition() == EHairStrandsCompositionType::AfterTranslucentBeforeTranslucentAfterDOF)
