Cloud samples debug option.
Disabled by default using #define.

diff --git a/Engine/Shaders/Private/VolumetricCloud.usf b/Engine/Shaders/Private/VolumetricCloud.usf
--- a/Engine/Shaders/Private/VolumetricCloud.usf
+++ b/Engine/Shaders/Private/VolumetricCloud.usf
@@ -39,6 +39,15 @@
 #define CLOUD_CONSERVATIVE_DENSITY_DEBUG 0
 #endif
 
+#ifndef CLOUD_DEBUG_SAMPLES
+#define CLOUD_DEBUG_SAMPLES 0
+#endif
+
+// Draw debug sample positions
+#if CLOUD_DEBUG_SAMPLES
+#include "ShaderPrint.ush"
+#endif
+
 #if defined(SHADER_MAINVS) || defined(SHADER_SHADOW_PS) || SHADER_RENDERVIEW_PS || SHADER_RENDERVIEW_CS
 
 #include "/Engine/Generated/Material.ush"
@@ -614,6 +623,11 @@ void MainCommon(in FMaterialPixelParameters MaterialParameters, in float4 SvPosi
 		// TRACING LOOP
 		//
 
+	#if CLOUD_DEBUG_SAMPLES
+		int2 DebugPixelCoord = float2(ResolvedView.CursorPosition - RenderVolumetricCloudParameters.TracingCoordToFullResPixelCoordScaleBias.zw) / float2(RenderVolumetricCloudParameters.TracingCoordToFullResPixelCoordScaleBias.xy);
+		const bool bDrawDebugEnabled = all(int2(SvPosition.xy) == DebugPixelCoord);
+	#endif
+
 		float t = TMin + 0.5 * StepT;
 		if (RenderVolumetricCloudParameters.IsReflectionRendering == 0)
 		{
@@ -644,8 +658,18 @@ void MainCommon(in FMaterialPixelParameters MaterialParameters, in float4 SvPosi
 				i += RenderVolumetricCloudParameters.StepSizeOnZeroConservativeDensity - 1;
 				t += RenderVolumetricCloudParameters.StepSizeOnZeroConservativeDensity * StepT;
 #endif
+
+			#if CLOUD_DEBUG_SAMPLES
+				if (bDrawDebugEnabled) AddCrossTWS(LWCToFloat(SampleTranslatedWorldPosition), 60000.0f, float4(1, 0, 0, 0.5));
+			#endif
 				continue; // Conservative density is 0 so skip and go to the next sample
 			}
+			#if CLOUD_DEBUG_SAMPLES
+			else
+			{
+				if(bDrawDebugEnabled) AddCrossTWS(LWCToFloat(SampleTranslatedWorldPosition), 60000.0f, float4(0, 1, 0, 0.5));
+			}
+			#endif // CLOUD_DEBUG_SAMPLES
 #endif
 
 
@@ -841,7 +865,8 @@ void MainCommon(in FMaterialPixelParameters MaterialParameters, in float4 SvPosi
 					const float DetlaNormT = CurrentNormT - PreviousNormT;
 					const float ExtinctionFactor = DetlaNormT;
 					const float ShadowSampleDistance = ShadowLengthTest * (PreviousNormT + DetlaNormT * ShadowJitterNorm);
-					UpdateMaterialCloudParam(MaterialParameters, LWCAdd(SampleWorldPosition, LWCPromote(Light0Direction * ShadowSampleDistance)),
+					const FLWCVector3 ShadowSampleWorldPos = LWCAdd(SampleWorldPosition, LWCPromote(Light0Direction * ShadowSampleDistance));
+					UpdateMaterialCloudParam(MaterialParameters, ShadowSampleWorldPos,
 						ResolvedView, CloudLayerParams, ShadowSampleDistance);
 					PreviousNormT = CurrentNormT;
 #endif
@@ -849,8 +874,17 @@ void MainCommon(in FMaterialPixelParameters MaterialParameters, in float4 SvPosi
 #if MATERIAL_VOLUMETRIC_ADVANCED_CONSERVATIVE_DENSITY
 					if (MaterialParameters.VolumeSampleConservativeDensity.x <= 0.0f)
 					{
+					#if CLOUD_DEBUG_SAMPLES
+						if (bDrawDebugEnabled) AddCrossWS(LWCToFloat(ShadowSampleWorldPos), 30000.0f, float4(1, 0, 1, 0.5));
+					#endif
 						continue; // Conservative density is 0 so skip and go to the next sample
 					}
+					#if CLOUD_DEBUG_SAMPLES
+					else
+					{
+						if (bDrawDebugEnabled) AddCrossWS(LWCToFloat(ShadowSampleWorldPos), 30000.0f, float4(0, 1, 1, 0.5));
+					}
+					#endif // CLOUD_DEBUG_SAMPLES
 #endif
 
 					if (MaterialParameters.CloudSampleNormAltitudeInLayer < 0.0f || MaterialParameters.CloudSampleNormAltitudeInLayer > 1.0f)
diff --git a/Engine/Source/Runtime/Renderer/Private/VolumetricCloudRendering.cpp b/Engine/Source/Runtime/Renderer/Private/VolumetricCloudRendering.cpp
--- a/Engine/Source/Runtime/Renderer/Private/VolumetricCloudRendering.cpp
+++ b/Engine/Source/Runtime/Renderer/Private/VolumetricCloudRendering.cpp
@@ -26,9 +26,8 @@
 #include "SkyAtmosphereRendering.h"
 #include "BasePassRendering.h"
 
-
-//PRAGMA_DISABLE_OPTIMIZATION
-
+//  If this is enabled, you also need to touch VolumetricCloud.usf for shaders to be recompiled.
+#define CLOUD_DEBUG_SAMPLES 0 /*!Never check in enabled!*/
 
 ////////////////////////////////////////////////////////////////////////// Cloud rendering and tracing
 
@@ -524,7 +523,6 @@ BEGIN_GLOBAL_SHADER_PARAMETER_STRUCT(FRenderVolumetricCloudGlobalParameters, )
 	SHADER_PARAMETER_SAMPLER(SamplerState, CloudBilinearTextureSampler)
 	SHADER_PARAMETER_STRUCT_INCLUDE(FVolumeShadowingShaderParametersGlobal0, Light0Shadow)
 	SHADER_PARAMETER(int32, VirtualShadowMapId0)
-//	SHADER_PARAMETER_STRUCT(FBlueNoise, BlueNoise)
 	SHADER_PARAMETER(FUintVector4, TracingCoordToZbufferCoordScaleBias)
 	SHADER_PARAMETER(FUintVector4, TracingCoordToFullResPixelCoordScaleBias)
 	SHADER_PARAMETER(FUintVector4, SceneDepthTextureMinMaxCoord)
@@ -565,13 +563,6 @@ void SetupDefaultRenderVolumetricCloudGlobalParameters(FRDGBuilder& GraphBuilder
 	VolumetricCloudParams.CloudBilinearTextureSampler = TStaticSamplerState<SF_Bilinear>::GetRHI();
 	// Light0Shadow
 	VolumetricCloudParams.VirtualShadowMapId0 = INDEX_NONE;
-/*#if RHI_RAYTRACING
-	VolumetricCloudParams.BlueNoise = GetBlueNoiseParameters();
-#else
-	// Blue noise texture is undified for some configuration so replace by other noise for now.
-	VolumetricCloudParams.BlueNoise.Dimensions = FIntVector(16, 16, 4); // 16 is the size of the tile, so 4 dimension for the 64x64 HighFrequencyNoiseTexture.
-	VolumetricCloudParams.BlueNoise.Texture = GEngine->HighFrequencyNoiseTexture->Resource->TextureRHI;
-#endif*/
 	VolumetricCloudParams.TracingCoordToZbufferCoordScaleBias = FUintVector4(1, 1, 0, 0);
 	VolumetricCloudParams.TracingCoordToFullResPixelCoordScaleBias = FUintVector4(1, 1, 0, 0);
 	VolumetricCloudParams.NoiseFrameIndexModPattern = 0;
@@ -788,6 +779,9 @@ class FRenderVolumetricCloudRenderViewCS : public FMeshMaterialShader
 		SHADER_PARAMETER_RDG_TEXTURE_UAV(RWTexture2D, OutCloudColor)
 		SHADER_PARAMETER_RDG_TEXTURE_UAV(RWTexture2D, OutCloudDepth)
 		SHADER_PARAMETER_RDG_TEXTURE_UAV(RWTexture2D, OutCloudColorCube)
+#if CLOUD_DEBUG_SAMPLES
+		SHADER_PARAMETER_STRUCT_INCLUDE(ShaderPrint::FShaderParameters, ShaderPrintParameters)
+#endif
 	END_SHADER_PARAMETER_STRUCT()
 
 	static const int32 ThreadGroupSizeX = 8;
@@ -822,6 +816,9 @@ class FRenderVolumetricCloudRenderViewCS : public FMeshMaterialShader
 			OutEnvironment.SetDefine(TEXT("VIRTUAL_SHADOW_MAP"), TEXT("1"));
 		}
 
+		// If this is enabled
+		OutEnvironment.SetDefine(TEXT("CLOUD_DEBUG_SAMPLES"), CLOUD_DEBUG_SAMPLES);
+
 		// This shader must support typed UAV load and we are testing if it is supported at runtime using RHIIsTypedUAVLoadSupported
 		OutEnvironment.CompilerFlags.Add(CFLAG_AllowTypedUAVLoads);
 
@@ -2077,6 +2074,12 @@ void FSceneRenderer::RenderVolumetricCloudsInternal(FRDGBuilder& GraphBuilder, F
 		PassParameters->OutCloudDepth = GraphBuilder.CreateUAV(CloudDepthTexture);
 		PassParameters->OutCloudColorCube = GraphBuilder.CreateUAV(CloudColorCubeTexture);
 
+#if CLOUD_DEBUG_SAMPLES
+		ShaderPrint::SetEnabled(true);
+		ShaderPrint::RequestSpaceForLines(65536);
+		ShaderPrint::SetParameters(GraphBuilder, MainView.ShaderPrintData, PassParameters->ShaderPrintParameters);
+#endif
+
 		const FMaterialRenderProxy* MaterialRenderProxy = nullptr;
 		const FMaterial* MaterialResource = &CloudVolumeMaterialProxy->GetMaterialWithFallback(Scene->GetFeatureLevel(), MaterialRenderProxy);
 		MaterialRenderProxy = MaterialRenderProxy ? MaterialRenderProxy : CloudVolumeMaterialProxy;
