Fixes warning Unknown TypeLayout FTSRDebugHistoryCS

diff --git a/Engine/Source/Runtime/Renderer/Private/PostProcess/TemporalSuperResolution.cpp b/Engine/Source/Runtime/Renderer/Private/PostProcess/TemporalSuperResolution.cpp
--- a/Engine/Source/Runtime/Renderer/Private/PostProcess/TemporalSuperResolution.cpp
+++ b/Engine/Source/Runtime/Renderer/Private/PostProcess/TemporalSuperResolution.cpp
@@ -714,8 +714,6 @@ class FTSRResolveHistoryCS : public FTSRShader
 	}
 }; // class FTSRResolveHistoryCS
 
-#if COMPILE_TSR_DEBUG_PASSES
-
 class FTSRDebugHistoryCS : public FTSRShader
 {
 	DECLARE_GLOBAL_SHADER(FTSRDebugHistoryCS);
@@ -729,8 +727,6 @@ class FTSRDebugHistoryCS : public FTSRShader
 	END_SHADER_PARAMETER_STRUCT()
 }; // class FTSRDebugHistoryCS
 
-#endif
-
 IMPLEMENT_GLOBAL_SHADER(FTSRClearPrevTexturesCS,     "/Engine/Private/TemporalSuperResolution/TSRClearPrevTextures.usf",     "MainCS", SF_Compute);
 IMPLEMENT_GLOBAL_SHADER(FTSRDilateVelocityCS,        "/Engine/Private/TemporalSuperResolution/TSRDilateVelocity.usf",        "MainCS", SF_Compute);
 IMPLEMENT_GLOBAL_SHADER(FTSRDecimateHistoryCS,       "/Engine/Private/TemporalSuperResolution/TSRDecimateHistory.usf",       "MainCS", SF_Compute);
@@ -742,10 +738,7 @@ IMPLEMENT_GLOBAL_SHADER(FTSRSpatialAntiAliasingCS,   "/Engine/Private/TemporalSu
 IMPLEMENT_GLOBAL_SHADER(FTSRFilterAntiAliasingCS,    "/Engine/Private/TemporalSuperResolution/TSRFilterAntiAliasing.usf",    "MainCS", SF_Compute);
 IMPLEMENT_GLOBAL_SHADER(FTSRUpdateHistoryCS,         "/Engine/Private/TemporalSuperResolution/TSRUpdateHistory.usf",         "MainCS", SF_Compute);
 IMPLEMENT_GLOBAL_SHADER(FTSRResolveHistoryCS,        "/Engine/Private/TemporalSuperResolution/TSRResolveHistory.usf",        "MainCS", SF_Compute);
-
-#if COMPILE_TSR_DEBUG_PASSES
 IMPLEMENT_GLOBAL_SHADER(FTSRDebugHistoryCS,          "/Engine/Private/TemporalSuperResolution/TSRDebugHistory.usf",          "MainCS", SF_Compute);
-#endif
 
 DECLARE_GPU_STAT(TemporalSuperResolution)
 
