Swapping FPrimitiveOctreeSemantics::ElementAllocator from TInlineAllocator to FDefaultAllocator.

- Lowers primitive octree mem usage by ~20MB over the course of the match.
- Reduces realloc hitches (100ms on console) by 10x or more. On high end pc I measured 2ms down to 0.2ms.

diff --git a/Engine/Source/Runtime/Renderer/Public/PrimitiveSceneInfo.h b/Engine/Source/Runtime/Renderer/Public/PrimitiveSceneInfo.h
--- a/Engine/Source/Runtime/Renderer/Public/PrimitiveSceneInfo.h
+++ b/Engine/Source/Runtime/Renderer/Public/PrimitiveSceneInfo.h
@@ -722,7 +722,7 @@ struct FPrimitiveOctreeSemantics
 	enum { MinInclusiveElementsPerNode = 7 };
 	enum { MaxNodeDepth = 12 };
 
-	typedef TInlineAllocator<MaxElementsPerLeaf> ElementAllocator;
+	typedef FDefaultAllocator ElementAllocator;
 
 	FORCEINLINE static const FCompactBoxSphereBounds& GetBoundingBox(const FPrimitiveSceneInfoCompact& PrimitiveSceneInfoCompact)
 	{
