Fixed real time capture DEBUG_TIME_SLICE due to recent code update to fix multi gpu.

diff --git a/Engine/Source/Runtime/Renderer/Private/ReflectionEnvironmentRealTimeCapture.cpp b/Engine/Source/Runtime/Renderer/Private/ReflectionEnvironmentRealTimeCapture.cpp
--- a/Engine/Source/Runtime/Renderer/Private/ReflectionEnvironmentRealTimeCapture.cpp
+++ b/Engine/Source/Runtime/Renderer/Private/ReflectionEnvironmentRealTimeCapture.cpp
@@ -962,11 +962,11 @@ void FScene::AllocateAndCaptureFrameSkyEnvMap(
 
 #define DEBUG_TIME_SLICE 0
 #if DEBUG_TIME_SLICE
-		RealTimeSlicedReflectionCaptureState = -1;
-		RealTimeSlicedReflectionCaptureStateStep = 0;
+		Capture = FRealTimeSlicedReflectionCapture();
+		Capture.FirstFrameState = FRealTimeSlicedReflectionCapture::EFirstFrameState::FIRST_FRAME;
 		while(true)
 		{
-			if (RealTimeSlicedReflectionCaptureState+1 >= TimeSliceCount)
+			if (Capture .State+1 >= TimeSliceCount)
 			{
 				break;
 			}
@@ -988,7 +988,7 @@ void FScene::AllocateAndCaptureFrameSkyEnvMap(
 			}
 
 			// Update the current time-slicing state if this is a new frame and if the current step is done.
-			// Note: RealTimeSlicedReflectionCaptureState will initially be -1.
+			// Note: Capture.State will initially be -1.
 			if (bStateFaceStepsDone)
 			{
 				if (++Capture.State >= TimeSliceCount)
