More DEBUG_TIME_SLICE fixes to have it run every frame instead of everyother frames.

diff --git a/Engine/Source/Runtime/Renderer/Private/ReflectionEnvironmentRealTimeCapture.cpp b/Engine/Source/Runtime/Renderer/Private/ReflectionEnvironmentRealTimeCapture.cpp
--- a/Engine/Source/Runtime/Renderer/Private/ReflectionEnvironmentRealTimeCapture.cpp
+++ b/Engine/Source/Runtime/Renderer/Private/ReflectionEnvironmentRealTimeCapture.cpp
@@ -963,10 +963,11 @@ void FScene::AllocateAndCaptureFrameSkyEnvMap(
 #define DEBUG_TIME_SLICE 0
 #if DEBUG_TIME_SLICE
 		Capture = FRealTimeSlicedReflectionCapture();
-		Capture.FirstFrameState = FRealTimeSlicedReflectionCapture::EFirstFrameState::FIRST_FRAME;
+		Capture.FirstFrameState = FRealTimeSlicedReflectionCapture::EFirstFrameState::BEYOND_FIRST_FRAME;
+		Capture.GpusWithFullCube |= MainView.GPUMask.GetNative();
 		while(true)
 		{
-			if (Capture .State+1 >= TimeSliceCount)
+			if (Capture.State+1 >= TimeSliceCount)
 			{
 				break;
 			}
@@ -1096,6 +1097,9 @@ void FScene::AllocateAndCaptureFrameSkyEnvMap(
 
 #if DEBUG_TIME_SLICE
 		}
+		ConvolvedSkyRenderTargetReadyIndex = 1 - ConvolvedSkyRenderTargetReadyIndex;
+		Capture.State = 0;
+		Capture.StateSubStep = 0;
 #endif
 	}
 
