Fix translucent materials without `allow custom depth writes` still writes to the custom depth buffer.

diff --git a/Engine/Source/Runtime/Renderer/Private/CustomDepthRendering.cpp b/Engine/Source/Runtime/Renderer/Private/CustomDepthRendering.cpp
--- a/Engine/Source/Runtime/Renderer/Private/CustomDepthRendering.cpp
+++ b/Engine/Source/Runtime/Renderer/Private/CustomDepthRendering.cpp
@@ -234,7 +234,7 @@ private:
 		ERasterizerFillMode MeshFillMode,
 		ERasterizerCullMode MeshCullMode);
 
-	bool UseDefaultMaterial(const FMaterial& Material, bool bMaterialModifiesMeshPosition, bool bSupportPositionOnlyStream, bool& bPositionOnly);
+	bool UseDefaultMaterial(const FMaterial& Material, bool bMaterialModifiesMeshPosition, bool bSupportPositionOnlyStream, bool& bPositionOnly, bool& bIgnoreThisMaterial);
 
 	void CollectDefaultMaterialPSOInitializers(
 		const FSceneTexturesConfig& SceneTexturesConfig,
@@ -307,9 +307,10 @@ FRHIDepthStencilState* GetCustomDepthStencilState(bool bWriteCustomStencilValues
 	}
 }
 
-bool FCustomDepthPassMeshProcessor::UseDefaultMaterial(const FMaterial& Material, bool bMaterialModifiesMeshPosition, bool bSupportPositionOnlyStream, bool& bPositionOnly)
+bool FCustomDepthPassMeshProcessor::UseDefaultMaterial(const FMaterial& Material, bool bMaterialModifiesMeshPosition, bool bSupportPositionOnlyStream, bool& bPositionOnly, bool& bIgnoreThisMaterial)
 {
 	bool bUseDefaultMaterial = false;
+	bIgnoreThisMaterial = false;
 
 	const EBlendMode BlendMode = Material.GetBlendMode();
 	if (BlendMode == BLEND_Opaque
@@ -329,6 +330,11 @@ bool FCustomDepthPassMeshProcessor::UseDefaultMaterial(const FMaterial& Material
 			bPositionOnly = false;
 		}
 	}
+	else
+	{
+		// E.g., ignore translucent materials without allowing custom depth writes.
+		bIgnoreThisMaterial = true;
+	}
 
 	return bUseDefaultMaterial;
 }
@@ -351,9 +357,13 @@ bool FCustomDepthPassMeshProcessor::TryAddMeshBatch(
 	}
 
 	// Using default material?
+	bool bIgnoreThisMaterial = false;
 	bool bPositionOnly = false;
-	bool bUseDefaultMaterial = UseDefaultMaterial(Material, Material.MaterialModifiesMeshPosition_RenderThread(), MeshBatch.VertexFactory->SupportsPositionOnlyStream(), bPositionOnly);
-
+	bool bUseDefaultMaterial = UseDefaultMaterial(Material, Material.MaterialModifiesMeshPosition_RenderThread(), MeshBatch.VertexFactory->SupportsPositionOnlyStream(), bPositionOnly, bIgnoreThisMaterial);
+	if (bIgnoreThisMaterial)
+	{
+		return true;
+	}
 	// Swap to default material
 	const FMaterialRenderProxy* EffectiveMaterialRenderProxy = &MaterialRenderProxy;
 	const FMaterial* EffectiveMaterial = &Material;
@@ -448,18 +458,22 @@ void FCustomDepthPassMeshProcessor::CollectPSOInitializers(const FSceneTexturesC
 
 	// assume we can always do this when collecting PSO's for now (vertex factory instance might actually not support it)
 	bool bSupportPositionOnlyStream = VertexFactoryType->SupportsPositionOnly();
-
+	bool bIgnoreThisMaterial = false;
 	bool bPositionOnly = false;
-	bool bUseDefaultMaterial = UseDefaultMaterial(Material, Material.MaterialModifiesMeshPosition_GameThread(), bSupportPositionOnlyStream, bPositionOnly);
-	if (!bUseDefaultMaterial && PreCacheParams.bRenderCustomDepth)
+	bool bUseDefaultMaterial = UseDefaultMaterial(Material, Material.MaterialModifiesMeshPosition_GameThread(), bSupportPositionOnlyStream, bPositionOnly, bIgnoreThisMaterial);
+
+	if (!bIgnoreThisMaterial)
 	{
-		check(!bPositionOnly);
+		if (!bUseDefaultMaterial && PreCacheParams.bRenderCustomDepth)
+		{
+			check(!bPositionOnly);
 
-		const FMeshDrawingPolicyOverrideSettings OverrideSettings = ComputeMeshOverrideSettings(PreCacheParams);
-		const ERasterizerFillMode MeshFillMode = ComputeMeshFillMode(Material, OverrideSettings);
-		const ERasterizerCullMode MeshCullMode = ComputeMeshCullMode(Material, OverrideSettings);
+			const FMeshDrawingPolicyOverrideSettings OverrideSettings = ComputeMeshOverrideSettings(PreCacheParams);
+			const ERasterizerFillMode MeshFillMode = ComputeMeshFillMode(Material, OverrideSettings);
+			const ERasterizerCullMode MeshCullMode = ComputeMeshCullMode(Material, OverrideSettings);
 
-		CollectPSOInitializers<false>(VertexFactoryType, Material, MeshFillMode, MeshCullMode, PSOInitializers);
+			CollectPSOInitializers<false>(VertexFactoryType, Material, MeshFillMode, MeshCullMode, PSOInitializers);
+		}
 	}
 }
 
