Disable "conservative culling" for VSM by default to reduce CPU/empty draw overhead for non-nanite geometry (host side tracking makes this not needed anymore).

diff --git a/Engine/Source/Runtime/Renderer/Private/ShadowSetup.cpp b/Engine/Source/Runtime/Renderer/Private/ShadowSetup.cpp
--- a/Engine/Source/Runtime/Renderer/Private/ShadowSetup.cpp
+++ b/Engine/Source/Runtime/Renderer/Private/ShadowSetup.cpp
@@ -281,8 +281,9 @@ static TAutoConsoleVariable<int32> CVarVsmUseFarShadowRules(
 
 static TAutoConsoleVariable<int32> CVarClipmapUseConservativeCulling(
 	TEXT("r.Shadow.Virtual.Clipmap.UseConservativeCulling"),
-	1,
-	TEXT(""),
+	0,
+	TEXT("Conservative culling removes the frustum-clipped culling volume for the non-nanite geometry for VSM rendering. This means a lot more geometry is submitted, and also marked as rendered.\n")
+	TEXT("Useful to diagnose if there are culling artifacts in virtual shadow map clip maps due to errors in the tracking code."),
 	ECVF_RenderThreadSafe
 );
 
