Improve cloud render step when rendering time sliced cloud in real time sky light capture.
- default textures and UAV are now allocated only once (preventing some fast clear barriers)
- When render into a cubemap texture, a special UAV can be use if specified from the caller on the ClourRenderContext.

diff --git a/Engine/Source/Runtime/Renderer/Private/ReflectionEnvironmentRealTimeCapture.cpp b/Engine/Source/Runtime/Renderer/Private/ReflectionEnvironmentRealTimeCapture.cpp
--- a/Engine/Source/Runtime/Renderer/Private/ReflectionEnvironmentRealTimeCapture.cpp
+++ b/Engine/Source/Runtime/Renderer/Private/ReflectionEnvironmentRealTimeCapture.cpp
@@ -485,6 +485,9 @@ void FScene::AllocateAndCaptureFrameSkyEnvMap(
 					{
 						SetVolumeShadowingDefaultShaderParameters(GraphBuilder, CloudRC.LightShadowShaderParams0);
 					}
+
+					// Create default textures once for each faces
+					CloudRC.CreateDefaultTexturesIfNeeded(GraphBuilder);
 				}
 				else
 				{
@@ -729,6 +732,15 @@ void FScene::AllocateAndCaptureFrameSkyEnvMap(
 					CloudRC.RenderTargets[0] = SkyRC.RenderTargets[0];
 					//	CloudRC.RenderTargets[1] = Null target will skip export
 
+					if (CloudRC.ComputeOverlapCloudColorCubeTextureUAVWithoutBarrier == nullptr)
+					{
+						// Create a UAV skipping barrier to make sure each compute pass processing each faces can overlap. 
+						// Here in ReflectionEnvironmentRealTimeCapture we know it is a cube texture
+						FRDGTextureRef CubeTexture = CloudRC.RenderTargets[0].GetTexture();
+						check(CubeTexture->Desc.IsTextureCube());
+						CloudRC.ComputeOverlapCloudColorCubeTextureUAVWithoutBarrier = GraphBuilder.CreateUAV(CubeTexture, ERDGUnorderedAccessViewFlags::SkipBarrier);
+					}
+
 					CloudRC.VolumetricCloudShadowTexture[0] = CloudShadowAOData.VolumetricCloudShadowMap[0];
 					CloudRC.VolumetricCloudShadowTexture[1] = CloudShadowAOData.VolumetricCloudShadowMap[1];
 
diff --git a/Engine/Source/Runtime/Renderer/Private/VolumetricCloudRendering.cpp b/Engine/Source/Runtime/Renderer/Private/VolumetricCloudRendering.cpp
--- a/Engine/Source/Runtime/Renderer/Private/VolumetricCloudRendering.cpp
+++ b/Engine/Source/Runtime/Renderer/Private/VolumetricCloudRendering.cpp
@@ -1921,6 +1921,29 @@ FCloudRenderContext::FCloudRenderContext()
 	bVisualizeConservativeDensity = false;
 }
 
+void FCloudRenderContext::CreateDefaultTexturesIfNeeded(FRDGBuilder& GraphBuilder)
+{
+	if (DefaultCloudColorCubeTexture == nullptr)
+	{
+		DefaultCloudColorCubeTexture = GraphBuilder.CreateTexture(
+			FRDGTextureDesc::CreateCube(1, PF_FloatRGBA, FClearValueBinding::Black, TexCreate_ShaderResource | TexCreate_RenderTargetable | TexCreate_UAV),
+			TEXT("Cloud.ColorCubeDummy"));
+
+		DefaultCloudColor2DTexture = GraphBuilder.CreateTexture(
+			FRDGTextureDesc::Create2D(FIntPoint(1, 1), PF_FloatRGBA, FClearValueBinding::None, TexCreate_ShaderResource | TexCreate_RenderTargetable | TexCreate_UAV),
+			TEXT("Cloud.ColorDummy"));
+
+		DefaultCloudDepthTexture = GraphBuilder.CreateTexture(
+			FRDGTextureDesc::Create2D(FIntPoint(1, 1), PF_G16R16F, FClearValueBinding::None, TexCreate_ShaderResource | TexCreate_RenderTargetable | TexCreate_UAV),
+			TEXT("Cloud.DepthDummy"));
+
+
+		DefaultCloudColorCubeTextureUAV = GraphBuilder.CreateUAV(DefaultCloudColorCubeTexture, ERDGUnorderedAccessViewFlags::SkipBarrier);
+		DefaultCloudColor2DTextureUAV = GraphBuilder.CreateUAV(DefaultCloudColor2DTexture, ERDGUnorderedAccessViewFlags::SkipBarrier);
+		DefaultCloudDepthTextureUAV = GraphBuilder.CreateUAV(DefaultCloudDepthTexture, ERDGUnorderedAccessViewFlags::SkipBarrier);
+	}
+}
+
 static TRDGUniformBufferRef<FRenderVolumetricCloudGlobalParameters> CreateCloudPassUniformBuffer(FRDGBuilder& GraphBuilder, FCloudRenderContext& CloudRC)
 {
 	FViewInfo& MainView = *CloudRC.MainView;
@@ -2003,30 +2026,46 @@ static TRDGUniformBufferRef<FRenderVolumetricCloudGlobalParameters> CreateCloudP
 	return GraphBuilder.CreateUniformBuffer(&VolumetricCloudParams);
 }
 
-static void GetOutputTexturesWithFallback(FRDGBuilder& GraphBuilder, FCloudRenderContext& CloudRC, FRDGTextureRef& CloudColorCubeTexture, FRDGTextureRef& CloudColorTexture, FRDGTextureRef& CloudDepthTexture)
+static void GetOutputTexturesWithFallback(FRDGBuilder& GraphBuilder, FCloudRenderContext& CloudRC, FRDGTextureUAVRef& CloudColorCubeTextureUAV, FRDGTextureUAVRef& CloudColorTextureUAV, FRDGTextureUAVRef& CloudDepthTextureUAV)
 {
-	CloudColorCubeTexture = CloudRC.RenderTargets[0].GetTexture();
+	CloudRC.CreateDefaultTexturesIfNeeded(GraphBuilder);
+
+	FRDGTextureRef CloudColorCubeTexture = CloudRC.RenderTargets[0].GetTexture();
 	if (!CloudColorCubeTexture || !CloudColorCubeTexture->Desc.IsTextureCube())
 	{
-		CloudColorCubeTexture = GraphBuilder.CreateTexture(
-			FRDGTextureDesc::CreateCube(1, PF_FloatRGBA, FClearValueBinding::None, TexCreate_ShaderResource | TexCreate_RenderTargetable | TexCreate_UAV),
-			TEXT("Cloud.ColorCubeDummy"));
+		CloudColorCubeTextureUAV = CloudRC.DefaultCloudColorCubeTextureUAV;
+	}
+	else
+	{
+		if (CloudRC.ComputeOverlapCloudColorCubeTextureUAVWithoutBarrier)
+		{
+			// If specified, this is to make sure a single UAV without barrier so that compute work items can overlap.
+			CloudColorCubeTextureUAV = CloudRC.ComputeOverlapCloudColorCubeTextureUAVWithoutBarrier;
+		}
+		else
+		{
+			CloudColorCubeTextureUAV = GraphBuilder.CreateUAV(CloudColorCubeTexture, ERDGUnorderedAccessViewFlags::SkipBarrier);
+		}
 	}
 
-	CloudColorTexture = CloudRC.RenderTargets[0].GetTexture();
+	FRDGTextureRef CloudColorTexture = CloudRC.RenderTargets[0].GetTexture();
 	if (!CloudColorTexture || CloudColorTexture->Desc.IsTextureCube())
 	{
-		CloudColorTexture = GraphBuilder.CreateTexture(
-			FRDGTextureDesc::Create2D(FIntPoint(1, 1), PF_FloatRGBA, FClearValueBinding::None, TexCreate_ShaderResource | TexCreate_RenderTargetable | TexCreate_UAV),
-			TEXT("Cloud.ColorDummy"));
+		CloudColorTextureUAV = CloudRC.DefaultCloudColor2DTextureUAV;
+	}
+	else
+	{
+		CloudColorTextureUAV = GraphBuilder.CreateUAV(CloudColorTexture, ERDGUnorderedAccessViewFlags::SkipBarrier);
 	}
 
-	CloudDepthTexture = CloudRC.RenderTargets[1].GetTexture();
+	FRDGTextureRef CloudDepthTexture = CloudRC.RenderTargets[1].GetTexture();
 	if (!CloudDepthTexture)
 	{
-		CloudDepthTexture = GraphBuilder.CreateTexture(
-			FRDGTextureDesc::Create2D(FIntPoint(1, 1), PF_G16R16F, FClearValueBinding::None, TexCreate_ShaderResource | TexCreate_RenderTargetable | TexCreate_UAV),
-			TEXT("Cloud.DepthDummy"));
+		CloudDepthTextureUAV = CloudRC.DefaultCloudDepthTextureUAV;
+	}
+	else
+	{
+		CloudDepthTextureUAV = GraphBuilder.CreateUAV(CloudDepthTexture, ERDGUnorderedAccessViewFlags::SkipBarrier);
 	}
 }
 
@@ -2053,10 +2092,10 @@ void FSceneRenderer::RenderVolumetricCloudsInternal(FRDGBuilder& GraphBuilder, F
 	const bool bShouldUseComputeForCloudTracing = !bVisualizeConservativeDensity && ShouldUseComputeForCloudTracing(Scene->GetFeatureLevel());
 	if (bShouldUseComputeForCloudTracing)
 	{
-		FRDGTextureRef CloudColorCubeTexture;
-		FRDGTextureRef CloudColorTexture;
-		FRDGTextureRef CloudDepthTexture;
-		GetOutputTexturesWithFallback(GraphBuilder, CloudRC, CloudColorCubeTexture, CloudColorTexture, CloudDepthTexture);
+		FRDGTextureUAVRef CloudColorCubeTextureUAV;
+		FRDGTextureUAVRef CloudColorTextureUAV;
+		FRDGTextureUAVRef CloudDepthTextureUAV;
+		GetOutputTexturesWithFallback(GraphBuilder, CloudRC, CloudColorCubeTextureUAV, CloudColorTextureUAV, CloudDepthTextureUAV);
 
 		bool bSampleVirtualShadowMap = (!bSkipAtmosphericLightShadowmap && CVarVolumetricCloudShadowSampleAtmosphericLightShadowmap.GetValueOnRenderThread() > 0);
 
@@ -2070,9 +2109,9 @@ void FSceneRenderer::RenderVolumetricCloudsInternal(FRDGBuilder& GraphBuilder, F
 		{
 			PassParameters->VirtualShadowMap = VirtualShadowMapArray.GetSamplingParameters(GraphBuilder);
 		}
-		PassParameters->OutCloudColor = GraphBuilder.CreateUAV(CloudColorTexture);
-		PassParameters->OutCloudDepth = GraphBuilder.CreateUAV(CloudDepthTexture);
-		PassParameters->OutCloudColorCube = GraphBuilder.CreateUAV(CloudColorCubeTexture);
+		PassParameters->OutCloudColor = CloudColorTextureUAV;
+		PassParameters->OutCloudDepth = CloudDepthTextureUAV;
+		PassParameters->OutCloudColorCube = CloudColorCubeTextureUAV;
 
 #if CLOUD_DEBUG_SAMPLES
 		ShaderPrint::SetEnabled(true);
diff --git a/Engine/Source/Runtime/Renderer/Private/VolumetricCloudRendering.h b/Engine/Source/Runtime/Renderer/Private/VolumetricCloudRendering.h
--- a/Engine/Source/Runtime/Renderer/Private/VolumetricCloudRendering.h
+++ b/Engine/Source/Runtime/Renderer/Private/VolumetricCloudRendering.h
@@ -145,8 +145,19 @@ struct FCloudRenderContext
 
 	int VirtualShadowMapId0 = INDEX_NONE;
 
+	FRDGTextureRef DefaultCloudColorCubeTexture = nullptr;
+	FRDGTextureRef DefaultCloudColor2DTexture = nullptr;
+	FRDGTextureRef DefaultCloudDepthTexture = nullptr;
+	FRDGTextureUAVRef DefaultCloudColorCubeTextureUAV = nullptr;
+	FRDGTextureUAVRef DefaultCloudColor2DTextureUAV = nullptr;
+	FRDGTextureUAVRef DefaultCloudDepthTextureUAV = nullptr;
+
+	FRDGTextureUAVRef ComputeOverlapCloudColorCubeTextureUAVWithoutBarrier = nullptr;
+
 	FCloudRenderContext();
 
+	void CreateDefaultTexturesIfNeeded(FRDGBuilder& GraphBuilder);
+
 private:
 };
 
