Add new ray tracing debug visualization mode to show heatmap for primary rays instance overlap:
* Collect instance AABBs when building RayTracingScene
* Output depth from RayTracingDebug pass for depth testing the instance AABBs
* Render instance AABBs as boxes with depth testing to accumulate the overlap which is then converted to heatmap and blended on top of RayTracingDebug output
* Non-GPUScene instances are not yet supported in this mode.

diff --git a/Engine/Plugins/Experimental/GPULightmass/Source/GPULightmass/Private/LightmapRenderer.cpp b/Engine/Plugins/Experimental/GPULightmass/Source/GPULightmass/Private/LightmapRenderer.cpp
--- a/Engine/Plugins/Experimental/GPULightmass/Source/GPULightmass/Private/LightmapRenderer.cpp
+++ b/Engine/Plugins/Experimental/GPULightmass/Source/GPULightmass/Private/LightmapRenderer.cpp
@@ -1187,7 +1187,8 @@ bool FSceneRenderState::SetupRayTracingScene(int32 LODIndex)
 				TransformUploadSRV,
 				SceneWithGeometryInstances.NumNativeGPUSceneInstances,
 				SceneWithGeometryInstances.NumNativeCPUInstances,
-				{});
+				{},
+				nullptr);
 
 			RHICmdList.BindAccelerationStructureMemory(RayTracingScene, RayTracingSceneBuffer, 0);
 
diff --git a/Engine/Shaders/Private/RayTracing/RayTracingDebug.usf b/Engine/Shaders/Private/RayTracing/RayTracingDebug.usf
--- a/Engine/Shaders/Private/RayTracing/RayTracingDebug.usf
+++ b/Engine/Shaders/Private/RayTracing/RayTracingDebug.usf
@@ -16,6 +16,8 @@
 #include "/Engine/Shared/RayTracingDebugTypes.h"
 
 RWTexture2D<float4> Output;
+RWTexture2D<float> OutputDepth;
+
 RaytracingAccelerationStructure TLAS;
 
 uint VisualizationMode;
@@ -167,7 +169,7 @@ struct FRayTracingDebugPayload : FMinimalPayload
 	uint GeometryIndex;
 };
 
-float4 DebugRayTracingInstance(FRayDesc Ray)
+float4 DebugRayTracingInstance(FRayDesc Ray, uint2 PixelCoord, out float OutDepth)
 {
 	float4 Result = float4(0, 0, 0, 1);
 
@@ -195,6 +197,9 @@ float4 DebugRayTracingInstance(FRayDesc Ray)
 		case RAY_TRACING_DEBUG_VIZ_INSTANCES:
 			Result.rgb = InstanceColor;
 			break;
+		case RAY_TRACING_DEBUG_VIZ_INSTANCE_OVERLAP:
+			Result.rgb = 1.0; // We only use fake lighting as instance heatmap is blended in
+			break;
 		case RAY_TRACING_DEBUG_VIZ_TRIANGLES:
 			Result.rgb = TriangleColor;
 			break;
@@ -234,8 +239,16 @@ float4 DebugRayTracingInstance(FRayDesc Ray)
 		}
 #endif // ENABLE_PICKING
 
-		// apply some shading based on the primitive normal (map dot product result from [0,1] to [0.75,1]).
-		Result.rgb *= dot(Payload.WorldNormal, normalize(float3(1.0f, 1.0f, 1.0f))) * 0.125f + 0.875;
+		// apply some shading based on the primitive normal (map dot product result from [0,1] to [0.75,1] unless we are showing instance heatmap then just use [0,1]).
+		const float FakeLighting = dot(Payload.WorldNormal, normalize(float3(1.0f, 1.0f, 1.0f)));
+		const float FakeLightingScale = (VisualizationMode == RAY_TRACING_DEBUG_VIZ_INSTANCE_OVERLAP) ? 0.5f : 0.125f;
+		const float FakeLightingBias = (VisualizationMode == RAY_TRACING_DEBUG_VIZ_INSTANCE_OVERLAP) ? 0.5f : 0.875f;
+
+		Result.rgb *= FakeLighting * FakeLightingScale + FakeLightingBias;
+
+		float3 TranslatedWorldPosition = Ray.Origin + Payload.HitT * Ray.Direction;
+		float4 ClipPosition = mul(float4(TranslatedWorldPosition, 1.0), View.TranslatedWorldToClip);
+		OutDepth = ClipPosition.z / ClipPosition.w;
 	}
 
 	return Result;
@@ -250,10 +263,16 @@ RAY_TRACING_ENTRY_RAYGEN(RayTracingDebugMainRGS)
 	FRayDesc Ray = CreatePrimaryRay(RenderTargetUV);
 
 	float4 Result = (float4)0;
+	
+#if	HAS_INVERTED_Z_BUFFER
+	float OutDepth = 0.0f;
+#else
+	float OutDepth = 1.0f;
+#endif
 
 	if (UseDebugCHS)
 	{
-		Result = DebugRayTracingInstance(Ray);
+		Result = DebugRayTracingInstance(Ray, PixelCoord, OutDepth);
 	}
 	else
 	{
@@ -285,6 +304,11 @@ RAY_TRACING_ENTRY_RAYGEN(RayTracingDebugMainRGS)
 	}
 	
 	Output[PixelCoord] = Result;
+
+	if (VisualizationMode == RAY_TRACING_DEBUG_VIZ_INSTANCE_OVERLAP)
+	{
+		OutputDepth[PixelCoord] = OutDepth;
+	}	
 }
 
 float3 ComputeWorldNormal()
diff --git a/Engine/Shaders/Private/RayTracing/RayTracingDebugInstanceOverlap.usf b/Engine/Shaders/Private/RayTracing/RayTracingDebugInstanceOverlap.usf
new file mode 100644
--- /dev/null
+++ b/Engine/Shaders/Private/RayTracing/RayTracingDebugInstanceOverlap.usf
@@ -0,0 +1,118 @@
+// Copyright Epic Games, Inc. All Rights Reserved.
+
+#include "../Common.ush"
+#include "../ShadingCommon.ush"
+#include "../ColorMap.ush"
+#include "../Hash.ush"
+#include "../Visualization.ush"
+#include "../SceneData.ush"
+
+#define MAX_FLT (+3.402823466e+38f)
+#define MIN_FLT (-3.402823466e+38f)
+#define INVALID_INSTANCE_GPUSCENE_INDEX (0xFFFFFFFF)
+
+float BoundingBoxExtentScale;
+
+StructuredBuffer<uint> InstanceGPUSceneIndices;
+
+struct FInstanceBounds
+{
+	float3 Center;
+	float3 Extent;
+};
+
+FInstanceBounds GetTranslatedWorldAABBForInstance(uint GPUSceneInstanceIndex)
+{
+	FInstanceSceneData InstanceSceneData = GetInstanceSceneData(GPUSceneInstanceIndex, View.InstanceSceneDataSOAStride);
+
+	const FLWCVector3 PreViewTranslation = MakeLWCVector3(-View.TLASViewTilePosition, View.TLASRelativePreViewTranslation);
+	const float4x4 LocalToTranslatedWorld = LWCMultiplyTranslation(InstanceSceneData.LocalToWorld, PreViewTranslation);	
+
+	const float3 BoundingBoxCenter = InstanceSceneData.LocalBoundsCenter.xyz;
+	const float3 BoundingBoxExtent = InstanceSceneData.LocalBoundsExtent.xyz;
+
+	float3 TranslatedWorldPositionMin = MAX_FLT;
+	float3 TranslatedWorldPositionMax = MIN_FLT;
+
+	for (int x = 0; x < 2; x++)
+	{
+		for (int y = 0; y < 2; y++)
+		{
+			for (int z = 0; z < 2; z++)
+			{
+				float3 Pos;
+				Pos.x = x ? -1.0f : 1.0f;
+				Pos.y = y ? -1.0f : 1.0f;
+				Pos.z = z ? -1.0f : 1.0f;
+
+				const float3 LocalPosition = Pos.xyz * BoundingBoxExtent + BoundingBoxCenter;
+
+				const float3 TranslatedWorldPosition = mul(float4(LocalPosition, 1.0), LocalToTranslatedWorld);
+
+				TranslatedWorldPositionMin = min(TranslatedWorldPosition, TranslatedWorldPositionMin);
+				TranslatedWorldPositionMax = max(TranslatedWorldPosition, TranslatedWorldPositionMax);
+			}
+		}
+	}		
+	
+	FInstanceBounds InstanceBounds;
+	InstanceBounds.Center = (TranslatedWorldPositionMax + TranslatedWorldPositionMin) * 0.5;
+	InstanceBounds.Extent = (TranslatedWorldPositionMax - TranslatedWorldPositionMin) * 0.5;
+
+	return InstanceBounds;
+}
+
+void InstanceOverlapMainVS(
+	float4 InPosition : ATTRIBUTE0,
+	uint InstanceIndex : SV_InstanceID,	
+	out float4 OutPosition : SV_POSITION
+)
+{
+	const uint GPUSceneIndex = InstanceGPUSceneIndices[InstanceIndex];
+
+	if (GPUSceneIndex != INVALID_INSTANCE_GPUSCENE_INDEX)
+	{
+		const FInstanceBounds InstanceBounds = GetTranslatedWorldAABBForInstance(GPUSceneIndex);
+
+		const float3 TranslatedWorldPosition = InPosition.xyz * InstanceBounds.Extent * BoundingBoxExtentScale + InstanceBounds.Center;
+
+		OutPosition = mul(float4(TranslatedWorldPosition, 1.0), PrimaryView.TranslatedWorldToClip);
+	}
+	else
+	{
+		OutPosition = float4(0, 0, 0, -1);
+	}
+}
+
+void InstanceOverlapMainPS(
+	in float4 Position : SV_POSITION, 
+	out float4 OutColor : SV_Target0
+)
+{
+	OutColor = 1.0;
+}
+
+Texture2D<float> InputDepth;
+
+void ConvertToDeviceDepthPS(
+	in  float4 SvPosition : SV_POSITION,
+	out float  OutDepth : SV_DEPTH)
+{
+	OutDepth = InputDepth.Load(int3(SvPosition.xy, 0)).r;
+}
+
+float HeatmapScale;
+Texture2D<float> InstanceOverlap;
+
+void BlendInstanceOverlapPS(
+	in float4 SvPosition : SV_POSITION,
+	out float4 OutColor : SV_Target0
+)
+{
+	const uint2 PixelCoord = uint2(SvPosition.xy);
+
+	const float3 HeatmapColor = InstanceOverlap[PixelCoord] == 0.0 ? 0.0 : ColorMapTurbo(InstanceOverlap[PixelCoord] / HeatmapScale);
+	const float BlendValue = 0.5;
+	
+	OutColor = float4(HeatmapColor, BlendValue);
+}
\ No newline at end of file
diff --git a/Engine/Shaders/Private/RayTracing/RayTracingInstanceBufferUtil.usf b/Engine/Shaders/Private/RayTracing/RayTracingInstanceBufferUtil.usf
--- a/Engine/Shaders/Private/RayTracing/RayTracingInstanceBufferUtil.usf
+++ b/Engine/Shaders/Private/RayTracing/RayTracingInstanceBufferUtil.usf
@@ -41,6 +41,9 @@ ByteAddressBuffer AccelerationStructureAddresses;
 // because of FXC matrix packing issues when using StructuredBuffer<float3x4>
 StructuredBuffer<float4> InstanceTransforms;
 
+// Extra per instance debug data
+RWStructuredBuffer<uint> RWDebugInstanceGPUSceneIndices;
+
 [numthreads(THREADGROUP_SIZE, 1, 1)]
 void RayTracingBuildInstanceBufferCS(uint3 DispatchThreadId : SV_DispatchThreadID)
 {
@@ -79,12 +82,19 @@ void RayTracingBuildInstanceBufferCS(uint3 DispatchThreadId : SV_DispatchThreadI
 		LocalToTranslatedWorld[0] = LocalToTranslatedWorld4x4[0];
 		LocalToTranslatedWorld[1] = LocalToTranslatedWorld4x4[1];
 		LocalToTranslatedWorld[2] = LocalToTranslatedWorld4x4[2];
+
+#if OUTPUT_INSTANCE_GPUSCENE_INDEX
+		{
+			RWDebugInstanceGPUSceneIndices[InputDesc.OutputDescriptorIndex] = InputDesc.GPUSceneInstanceOrTransformIndex;
+		}
+#endif // OUTPUT_INSTANCE_GPUSCENE_INDEX
+
 #else
 		float3x4 LocalToTranslatedWorld;
 		LocalToTranslatedWorld[0] = InstanceTransforms[InputDesc.GPUSceneInstanceOrTransformIndex * 3 + 0];
 		LocalToTranslatedWorld[1] = InstanceTransforms[InputDesc.GPUSceneInstanceOrTransformIndex * 3 + 1];
 		LocalToTranslatedWorld[2] = InstanceTransforms[InputDesc.GPUSceneInstanceOrTransformIndex * 3 + 2];
-#endif
+#endif // USE_GPUSCENE
 
 		uint2 BlasAddress = 0;
 
diff --git a/Engine/Shaders/Shared/RaytracingDebugDefinitions.h b/Engine/Shaders/Shared/RaytracingDebugDefinitions.h
--- a/Engine/Shaders/Shared/RaytracingDebugDefinitions.h
+++ b/Engine/Shaders/Shared/RaytracingDebugDefinitions.h
@@ -41,6 +41,7 @@
 #define RAY_TRACING_DEBUG_VIZ_DYNAMIC_INSTANCES			31
 #define RAY_TRACING_DEBUG_VIZ_PROXY_TYPE				32
 #define RAY_TRACING_DEBUG_VIZ_PICKER					33
+#define RAY_TRACING_DEBUG_VIZ_INSTANCE_OVERLAP			34
 
 #define RAY_TRACING_DEBUG_PICKER_DOMAIN_TRIANGLE		0
 #define RAY_TRACING_DEBUG_PICKER_DOMAIN_INSTANCE		1
diff --git a/Engine/Source/Editor/UnrealEd/Private/RayTracingDebugVisualizationMenuCommands.cpp b/Engine/Source/Editor/UnrealEd/Private/RayTracingDebugVisualizationMenuCommands.cpp
--- a/Engine/Source/Editor/UnrealEd/Private/RayTracingDebugVisualizationMenuCommands.cpp
+++ b/Engine/Source/Editor/UnrealEd/Private/RayTracingDebugVisualizationMenuCommands.cpp
@@ -54,6 +54,7 @@ void FRayTracingDebugVisualizationMenuCommands::CreateRayTracingDebugVisualizati
 	RayTracingDebugModeNames.Add(LOCTEXT("World Tangent", "World Tangent"));
 	RayTracingDebugModeNames.Add(LOCTEXT("Anisotropy", "Anisotropy"));
 	RayTracingDebugModeNames.Add(LOCTEXT("Instances", "Instances"));
+	RayTracingDebugModeNames.Add(LOCTEXT("Instance Overlap", "Instance Overlap"));
 
 	if (GRHISupportsShaderTimestamp)
 	{
diff --git a/Engine/Source/Runtime/Renderer/Private/RayTracing/RayTracingDebug.cpp b/Engine/Source/Runtime/Renderer/Private/RayTracing/RayTracingDebug.cpp
--- a/Engine/Source/Runtime/Renderer/Private/RayTracing/RayTracingDebug.cpp
+++ b/Engine/Source/Runtime/Renderer/Private/RayTracing/RayTracingDebug.cpp
@@ -16,6 +16,7 @@
 #include "RayTracing/RayTracingLighting.h"
 #include "RayTracing/RaytracingOptions.h"
 #include "RayTracing/RayTracingTraversalStatistics.h"
+#include "PixelShaderUtils.h"
 
 #define LOCTEXT_NAMESPACE "RayTracingDebugVisualizationMenuCommands"
 
@@ -61,6 +62,24 @@ static TAutoConsoleVariable<float> CVarRayTracingDebugTraversalClusterScale(
 	TEXT("Scaling factor for cluster traversal heat map visualization. (default = 2500)\n")
 );
 
+static TAutoConsoleVariable<float> CVarRayTracingDebugInstanceOverlapScale(
+	TEXT("r.RayTracing.Debug.InstanceOverlap.Scale"),
+	16.0f,
+	TEXT("Scaling factor for instance traversal heat map visualization. (default = 16)\n")
+);
+
+static TAutoConsoleVariable<float> CVarRayTracingDebugInstanceOverlapBoundingBoxScale(
+	TEXT("r.RayTracing.Debug.InstanceOverlap.BoundingBoxScale"),
+	1.001f,
+	TEXT("Scaling factor for instance bounding box extent for avoiding z-fighting. (default = 1.001)\n")
+);
+
+static TAutoConsoleVariable<int32> CVarRayTracingDebugInstanceOverlapShowWireframe(
+	TEXT("r.RayTracing.Debug.InstanceOverlap.ShowWireframe"),
+	1,
+	TEXT("Show instance bounding boxes in wireframe in Instances Overlap mode. (default = 1)\n")
+);
+
 static TAutoConsoleVariable<float> CVarRayTracingDebugTraversalTriangleScale(
 	TEXT("r.RayTracing.DebugTraversalScale.Triangle"),
 	30.0f,
@@ -97,6 +116,7 @@ class FRayTracingDebugRGS : public FGlobalShader
 		SHADER_PARAMETER(int32, OpaqueOnly)
 		SHADER_PARAMETER_SRV(RaytracingAccelerationStructure, TLAS)
 		SHADER_PARAMETER_RDG_TEXTURE_UAV(RWTexture2D<float4>, Output)
+		SHADER_PARAMETER_RDG_TEXTURE_UAV(RWTexture2D<float>, OutputDepth)
 		SHADER_PARAMETER_RDG_BUFFER_SRV(StructuredBuffer, InstancesDebugData)
 		SHADER_PARAMETER_RDG_BUFFER_SRV(StructuredBuffer<FRayTracingPickingFeedback>, PickingBuffer)
 		SHADER_PARAMETER_STRUCT_REF(FViewUniformShaderParameters, ViewUniformBuffer)
@@ -209,6 +229,139 @@ class FRayTracingPickingRGS : public FGlobalShader
 };
 IMPLEMENT_GLOBAL_SHADER(FRayTracingPickingRGS, "/Engine/Private/RayTracing/RayTracingDebug.usf", "RayTracingDebugPickingRGS", SF_RayGen);
 
+class FRayTracingDebugInstanceOverlapVS : public FGlobalShader
+{
+	DECLARE_GLOBAL_SHADER(FRayTracingDebugInstanceOverlapVS);
+	SHADER_USE_PARAMETER_STRUCT(FRayTracingDebugInstanceOverlapVS, FGlobalShader);
+
+	BEGIN_SHADER_PARAMETER_STRUCT(FParameters, )
+		SHADER_PARAMETER_STRUCT_REF(FViewUniformShaderParameters, View)
+		SHADER_PARAMETER_RDG_BUFFER_SRV(StructuredBuffer<float4>, GPUSceneInstanceSceneData)
+		SHADER_PARAMETER_RDG_BUFFER_SRV(StructuredBuffer<float4>, GPUSceneInstancePayloadData)
+		SHADER_PARAMETER_RDG_BUFFER_SRV(StructuredBuffer<float4>, GPUScenePrimitiveSceneData)
+		SHADER_PARAMETER_RDG_BUFFER_SRV(StructuredBuffer<uint>, InstanceGPUSceneIndices)
+		SHADER_PARAMETER(float, BoundingBoxExtentScale)
+	END_SHADER_PARAMETER_STRUCT()
+
+	static inline void ModifyCompilationEnvironment(const FGlobalShaderPermutationParameters& Parameters, FShaderCompilerEnvironment& OutEnvironment)
+	{
+		FGlobalShader::ModifyCompilationEnvironment(Parameters, OutEnvironment);
+
+		OutEnvironment.SetDefine(TEXT("VF_SUPPORTS_PRIMITIVE_SCENE_DATA"), 1);
+		OutEnvironment.SetDefine(TEXT("USE_GLOBAL_GPU_SCENE_DATA"), 1);
+	}
+
+	static bool ShouldCompilePermutation(const FGlobalShaderPermutationParameters& Parameters)
+	{
+		return IsRayTracingEnabledForProject(Parameters.Platform);
+	}
+};
+
+IMPLEMENT_GLOBAL_SHADER(FRayTracingDebugInstanceOverlapVS, "/Engine/Private/RayTracing/RayTracingDebugInstanceOverlap.usf", "InstanceOverlapMainVS", SF_Vertex);
+
+class FRayTracingDebugInstanceOverlapPS : public FGlobalShader
+{
+	DECLARE_GLOBAL_SHADER(FRayTracingDebugInstanceOverlapPS);
+	SHADER_USE_PARAMETER_STRUCT(FRayTracingDebugInstanceOverlapPS, FGlobalShader);
+
+	BEGIN_SHADER_PARAMETER_STRUCT(FParameters, )
+		SHADER_PARAMETER_STRUCT_REF(FViewUniformShaderParameters, View)		
+	END_SHADER_PARAMETER_STRUCT()
+
+	static bool ShouldCompilePermutation(const FGlobalShaderPermutationParameters& Parameters)
+	{
+		return IsRayTracingEnabledForProject(Parameters.Platform);
+	}
+};
+
+IMPLEMENT_GLOBAL_SHADER(FRayTracingDebugInstanceOverlapPS, "/Engine/Private/RayTracing/RayTracingDebugInstanceOverlap.usf", "InstanceOverlapMainPS", SF_Pixel);
+
+BEGIN_SHADER_PARAMETER_STRUCT(FRayTracingDebugInstanceOverlapVSPSParameters, )
+	SHADER_PARAMETER_STRUCT_INCLUDE(FRayTracingDebugInstanceOverlapVS::FParameters, VS)
+	SHADER_PARAMETER_STRUCT_INCLUDE(FRayTracingDebugInstanceOverlapPS::FParameters, PS)
+	RENDER_TARGET_BINDING_SLOTS()
+END_SHADER_PARAMETER_STRUCT()
+
+class FRayTracingDebugConvertToDeviceDepthPS : public FGlobalShader
+{
+	DECLARE_GLOBAL_SHADER(FRayTracingDebugConvertToDeviceDepthPS);
+	SHADER_USE_PARAMETER_STRUCT(FRayTracingDebugConvertToDeviceDepthPS, FGlobalShader);
+
+	BEGIN_SHADER_PARAMETER_STRUCT(FParameters, )
+		SHADER_PARAMETER_STRUCT_REF(FViewUniformShaderParameters, View)
+		SHADER_PARAMETER_RDG_TEXTURE_SRV(Texture2D<float>, InputDepth)
+		RENDER_TARGET_BINDING_SLOTS()
+	END_SHADER_PARAMETER_STRUCT()
+
+	static bool ShouldCompilePermutation(const FGlobalShaderPermutationParameters& Parameters)
+	{
+		return IsRayTracingEnabledForProject(Parameters.Platform);
+	}
+};
+
+IMPLEMENT_GLOBAL_SHADER(FRayTracingDebugConvertToDeviceDepthPS, "/Engine/Private/RayTracing/RayTracingDebugInstanceOverlap.usf", "ConvertToDeviceDepthPS", SF_Pixel);
+
+class FRayTracingDebugBlendInstanceOverlapPS : public FGlobalShader
+{
+	DECLARE_GLOBAL_SHADER(FRayTracingDebugBlendInstanceOverlapPS);
+	SHADER_USE_PARAMETER_STRUCT(FRayTracingDebugBlendInstanceOverlapPS, FGlobalShader);
+
+	BEGIN_SHADER_PARAMETER_STRUCT(FParameters, )
+		SHADER_PARAMETER_STRUCT_REF(FViewUniformShaderParameters, View)
+		SHADER_PARAMETER_RDG_TEXTURE_SRV(Texture2D<float>, InstanceOverlap)
+		SHADER_PARAMETER(float, HeatmapScale)
+		RENDER_TARGET_BINDING_SLOTS()
+	END_SHADER_PARAMETER_STRUCT()
+
+	static bool ShouldCompilePermutation(const FGlobalShaderPermutationParameters& Parameters)
+	{
+		return IsRayTracingEnabledForProject(Parameters.Platform);
+	}
+};
+
+IMPLEMENT_GLOBAL_SHADER(FRayTracingDebugBlendInstanceOverlapPS, "/Engine/Private/RayTracing/RayTracingDebugInstanceOverlap.usf", "BlendInstanceOverlapPS", SF_Pixel);
+
+class FRayTracingDebugLineAABBIndexBuffer : public FIndexBuffer
+{
+public:
+	/**
+	* Initialize the RHI for this rendering resource
+	*/
+	void InitRHI() override
+	{
+		TResourceArray<uint16, INDEXBUFFER_ALIGNMENT> Indices;
+
+		static const uint16 LineIndices[12 * 2] =
+		{
+			0, 1,
+			0, 2,
+			0, 4,
+			2, 3,
+			3, 1,
+			1, 5,
+			3, 7,
+			2, 6,
+			6, 7,
+			6, 4,
+			7, 5,
+			4, 5
+		};
+
+		int32 NumIndices = UE_ARRAY_COUNT(LineIndices);
+		Indices.AddUninitialized(NumIndices);
+		FMemory::Memcpy(Indices.GetData(), LineIndices, NumIndices * sizeof(uint16));
+
+		const uint32 Size = Indices.GetResourceDataSize();
+		const uint32 Stride = sizeof(uint16);
+
+		// Create index buffer. Fill buffer with initial data upon creation
+		FRHIResourceCreateInfo CreateInfo(TEXT("FRayTracingDebugLineAABBIndexBuffer"), &Indices);
+		IndexBufferRHI = RHICreateIndexBuffer(Stride, Size, BUF_Static, CreateInfo);
+	}
+};
+
+TGlobalResource<FRayTracingDebugLineAABBIndexBuffer> GRayTracingInstanceLineAABBIndexBuffer;
+
 struct FRayTracingDebugResources : public FRenderResource
 {
 	const int32 MaxPickingBuffers = 4;
@@ -288,7 +441,8 @@ static bool RequiresRayTracingDebugCHS(uint32 DebugVisualizationMode)
 		DebugVisualizationMode == RAY_TRACING_DEBUG_VIZ_TRIANGLES ||
 		DebugVisualizationMode == RAY_TRACING_DEBUG_VIZ_DYNAMIC_INSTANCES ||
 		DebugVisualizationMode == RAY_TRACING_DEBUG_VIZ_PROXY_TYPE ||
-		DebugVisualizationMode == RAY_TRACING_DEBUG_VIZ_PICKER;
+		DebugVisualizationMode == RAY_TRACING_DEBUG_VIZ_PICKER ||
+		DebugVisualizationMode == RAY_TRACING_DEBUG_VIZ_INSTANCE_OVERLAP;
 }
 
 static bool IsRayTracingDebugTraversalMode(uint32 DebugVisualizationMode)
@@ -427,6 +581,130 @@ static FRDGBufferRef RayTracingPerformPicking(FRDGBuilder& GraphBuilder, const F
 	return PickingBuffer;
 }
 
+static void RayTracingDrawInstances(FRDGBuilder& GraphBuilder, const FViewInfo& View, const FGPUScene& GPUScene, FRDGTextureRef OutputTexture, FRDGTextureRef SceneDepthTexture, FRDGBufferRef InstanceGPUSceneIndexBuffer, bool bWireframe)
+{
+	TShaderMapRef<FRayTracingDebugInstanceOverlapVS> VertexShader(View.ShaderMap);
+	TShaderMapRef<FRayTracingDebugInstanceOverlapPS> PixelShader(View.ShaderMap);
+
+	FRayTracingDebugInstanceOverlapVSPSParameters* PassParameters = GraphBuilder.AllocParameters<FRayTracingDebugInstanceOverlapVSPSParameters>();
+	PassParameters->VS.View = View.ViewUniformBuffer;
+	PassParameters->VS.InstanceGPUSceneIndices = GraphBuilder.CreateSRV(InstanceGPUSceneIndexBuffer);
+	PassParameters->VS.BoundingBoxExtentScale = CVarRayTracingDebugInstanceOverlapBoundingBoxScale.GetValueOnRenderThread();
+
+	const FGPUSceneResourceParameters GPUSceneParameters = GPUScene.GetShaderParameters();
+	PassParameters->VS.GPUSceneInstanceSceneData = GPUSceneParameters.GPUSceneInstanceSceneData;
+	PassParameters->VS.GPUSceneInstancePayloadData = GPUSceneParameters.GPUSceneInstancePayloadData;
+	PassParameters->VS.GPUScenePrimitiveSceneData = GPUSceneParameters.GPUScenePrimitiveSceneData;
+
+	PassParameters->PS.View = View.ViewUniformBuffer;
+
+	const uint32 NumInstances = InstanceGPUSceneIndexBuffer->GetSize() / InstanceGPUSceneIndexBuffer->GetStride();
+
+	PassParameters->RenderTargets[0] = FRenderTargetBinding(OutputTexture, bWireframe ? ERenderTargetLoadAction::ELoad : ERenderTargetLoadAction::EClear);
+	PassParameters->RenderTargets.DepthStencil = FDepthStencilBinding(SceneDepthTexture, ERenderTargetLoadAction::ELoad, FExclusiveDepthStencil::DepthRead_StencilNop);
+
+	ValidateShaderParameters(PixelShader, PassParameters->PS);
+	ClearUnusedGraphResources(PixelShader, &PassParameters->PS);
+	ValidateShaderParameters(VertexShader, PassParameters->VS);
+	ClearUnusedGraphResources(VertexShader, &PassParameters->VS);
+
+	GraphBuilder.AddPass(
+		RDG_EVENT_NAME("RayTracingDebug::DrawInstances"),
+		PassParameters,
+		ERDGPassFlags::Raster,
+		[&View, VertexShader, PixelShader, PassParameters, NumInstances, bWireframe](FRHICommandList& RHICmdList)
+		{
+			FGraphicsPipelineStateInitializer GraphicsPSOInit;
+			RHICmdList.ApplyCachedRenderTargets(GraphicsPSOInit);
+
+			RHICmdList.SetViewport(View.ViewRect.Min.X, View.ViewRect.Min.Y, 0.0f, View.ViewRect.Max.X, View.ViewRect.Max.Y, 1.0f);
+
+			GraphicsPSOInit.DepthStencilState = TStaticDepthStencilState<false, CF_DepthNearOrEqual>::GetRHI();
+			GraphicsPSOInit.BlendState = bWireframe ? TStaticBlendState<CW_RGB>::GetRHI() : TStaticBlendState<CW_RED, BO_Add, BF_One, BF_One>::GetRHI();
+			GraphicsPSOInit.RasterizerState = bWireframe ? TStaticRasterizerState<FM_Wireframe, CM_None>::GetRHI() : TStaticRasterizerState<FM_Solid, CM_CCW>::GetRHI();
+			GraphicsPSOInit.PrimitiveType = bWireframe ? PT_LineList : PT_TriangleList;
+			GraphicsPSOInit.BoundShaderState.VertexDeclarationRHI = GetVertexDeclarationFVector4();
+			GraphicsPSOInit.BoundShaderState.VertexShaderRHI = VertexShader.GetVertexShader();
+			GraphicsPSOInit.BoundShaderState.PixelShaderRHI = PixelShader.GetPixelShader();
+			SetGraphicsPipelineState(RHICmdList, GraphicsPSOInit, 0);
+
+			SetShaderParameters(RHICmdList, VertexShader, VertexShader.GetVertexShader(), PassParameters->VS);
+			SetShaderParameters(RHICmdList, PixelShader, PixelShader.GetPixelShader(), PassParameters->PS);
+
+			RHICmdList.SetStreamSource(0, GetUnitCubeVertexBuffer(), 0);
+
+			const FBufferRHIRef IndexBufferRHI = bWireframe ? GRayTracingInstanceLineAABBIndexBuffer.IndexBufferRHI  : GetUnitCubeIndexBuffer();
+			RHICmdList.DrawIndexedPrimitive(IndexBufferRHI, 0, 0, 8, 0, 12, NumInstances);
+		});
+}
+
+static void DrawInstanceOverlap(FRDGBuilder& GraphBuilder, const FScene* Scene, const FViewInfo& View, FRDGTextureRef SceneColorTexture, FRDGTextureRef InputDepthTexture)
+{
+	FRDGTextureRef SceneDepthTexture = GraphBuilder.CreateTexture(
+		FRDGTextureDesc::Create2D(
+			SceneColorTexture->Desc.Extent,
+			PF_DepthStencil,
+			FClearValueBinding::DepthFar,
+			TexCreate_DepthStencilTargetable | TexCreate_InputAttachmentRead | TexCreate_ShaderResource),
+		TEXT("RayTracingDebug::SceneDepth"));
+
+	// Convert from depth texture to depth buffer for depth testing
+	{
+		FRayTracingDebugConvertToDeviceDepthPS::FParameters* PassParameters = GraphBuilder.AllocParameters<FRayTracingDebugConvertToDeviceDepthPS::FParameters>();
+		PassParameters->RenderTargets.DepthStencil = FDepthStencilBinding(SceneDepthTexture, ERenderTargetLoadAction::EClear, FExclusiveDepthStencil::DepthWrite_StencilNop);
+
+		PassParameters->InputDepth = GraphBuilder.CreateSRV(InputDepthTexture);
+
+		TShaderMapRef<FRayTracingDebugConvertToDeviceDepthPS> PixelShader(View.ShaderMap);
+
+		FPixelShaderUtils::AddFullscreenPass(
+			GraphBuilder,
+			View.ShaderMap,
+			RDG_EVENT_NAME("RayTracingDebug::ConvertToDeviceDepth"),
+			PixelShader,
+			PassParameters,
+			View.ViewRect,
+			TStaticBlendState<>::GetRHI(),
+			TStaticRasterizerState<FM_Solid, CM_None>::GetRHI(),
+			TStaticDepthStencilState<true, CF_Always>::GetRHI());
+	}
+
+	// Accumulate instance overlap
+	FRDGTextureDesc InstanceOverlapTextureDesc = FRDGTextureDesc::Create2D(SceneColorTexture->Desc.Extent, PF_R32_FLOAT, FClearValueBinding::Black, TexCreate_ShaderResource | TexCreate_RenderTargetable);
+	FRDGTextureRef InstanceOverlapTexture = GraphBuilder.CreateTexture(InstanceOverlapTextureDesc, TEXT("RayTracingDebug::InstanceOverlap"));
+	
+	RayTracingDrawInstances(GraphBuilder, View, Scene->GPUScene, InstanceOverlapTexture, SceneDepthTexture, Scene->RayTracingScene.DebugInstanceGPUSceneIndexBuffer, false);
+
+	// Calculate heatmap of instance overlap and blend it on top of ray tracing debug output
+	{
+		FRayTracingDebugBlendInstanceOverlapPS::FParameters* PassParameters = GraphBuilder.AllocParameters<FRayTracingDebugBlendInstanceOverlapPS::FParameters>();
+
+		PassParameters->RenderTargets[0] = FRenderTargetBinding(SceneColorTexture, ERenderTargetLoadAction::ELoad, 0);
+
+		PassParameters->InstanceOverlap = GraphBuilder.CreateSRV(InstanceOverlapTexture);
+		PassParameters->HeatmapScale = CVarRayTracingDebugInstanceOverlapScale.GetValueOnRenderThread();
+
+		TShaderMapRef<FRayTracingDebugBlendInstanceOverlapPS> PixelShader(View.ShaderMap);
+
+		FPixelShaderUtils::AddFullscreenPass(
+			GraphBuilder,
+			View.ShaderMap,
+			RDG_EVENT_NAME("RayTracingDebug::BlendInstanceOverlap"),
+			PixelShader,
+			PassParameters,
+			View.ViewRect,
+			TStaticBlendState<CW_RGBA, BO_Add, BF_SourceAlpha, BF_InverseSourceAlpha>::GetRHI(),
+			TStaticRasterizerState<FM_Solid, CM_None>::GetRHI(),
+			TStaticDepthStencilState<false, CF_Always>::GetRHI());
+	}
+
+	// Draw instance AABB with lines
+	if (CVarRayTracingDebugInstanceOverlapShowWireframe.GetValueOnRenderThread() != 0)
+	{
+		RayTracingDrawInstances(GraphBuilder, View, Scene->GPUScene, SceneColorTexture, SceneDepthTexture, Scene->RayTracingScene.DebugInstanceGPUSceneIndexBuffer, true);
+	}
+}
+
 void FDeferredShadingSceneRenderer::RenderRayTracingDebug(FRDGBuilder& GraphBuilder, const FViewInfo& View, FRDGTextureRef SceneColorTexture, FRayTracingPickingFeedback& PickingFeedback)
 {
 	static TMap<FName, uint32> RayTracingDebugVisualizationModes;
@@ -466,6 +744,7 @@ void FDeferredShadingSceneRenderer::RenderRayTracingDebug(FRDGBuilder& GraphBuil
 		RayTracingDebugVisualizationModes.Emplace(FName(*LOCTEXT("Dynamic Instances", "Dynamic Instances").ToString()),							RAY_TRACING_DEBUG_VIZ_DYNAMIC_INSTANCES);
 		RayTracingDebugVisualizationModes.Emplace(FName(*LOCTEXT("Proxy Type", "Proxy Type").ToString()),										RAY_TRACING_DEBUG_VIZ_PROXY_TYPE);
 		RayTracingDebugVisualizationModes.Emplace(FName(*LOCTEXT("Picker", "Picker").ToString()),												RAY_TRACING_DEBUG_VIZ_PICKER);
+		RayTracingDebugVisualizationModes.Emplace(FName(*LOCTEXT("Instance Overlap", "Instance Overlap").ToString()),							RAY_TRACING_DEBUG_VIZ_INSTANCE_OVERLAP);
 	}
 
 	uint32 DebugVisualizationMode;
@@ -602,6 +881,17 @@ void FDeferredShadingSceneRenderer::RenderRayTracingDebug(FRDGBuilder& GraphBuil
 	RayGenParameters->TimingScale = CVarRayTracingDebugTimingScale.GetValueOnAnyThread() / 25000.0f;
 	RayGenParameters->OpaqueOnly = CVarRayTracingDebugModeOpaqueOnly.GetValueOnRenderThread();
 	
+	// If we don't output depth, create dummy 1x1 texture
+	const bool bOutputDepth = DebugVisualizationMode == RAY_TRACING_DEBUG_VIZ_INSTANCE_OVERLAP;
+
+	FRDGTextureDesc OutputDepthTextureDesc = FRDGTextureDesc::Create2D(
+		bOutputDepth ? SceneColorTexture->Desc.Extent : FIntPoint(1, 1),
+		PF_R32_FLOAT, 
+		FClearValueBinding::Black, 
+		TexCreate_ShaderResource | TexCreate_UAV);
+	FRDGTextureRef OutputDepthTexture = GraphBuilder.CreateTexture(OutputDepthTextureDesc, TEXT("RayTracingDebug::Depth"));
+	RayGenParameters->OutputDepth = GraphBuilder.CreateUAV(OutputDepthTexture);
+
 	if (Scene->RayTracingScene.InstanceDebugBuffer)
 	{
 		RayGenParameters->InstancesDebugData = GraphBuilder.CreateSRV(Scene->RayTracingScene.InstanceDebugBuffer);
@@ -650,6 +940,11 @@ void FDeferredShadingSceneRenderer::RenderRayTracingDebug(FRDGBuilder& GraphBuil
 
 		RHICmdList.RayTraceDispatch(Pipeline, RayGenShader.GetRayTracingShader(), View.GetRayTracingSceneChecked(), GlobalResources, ViewRect.Size().X, ViewRect.Size().Y);
 	});
+
+	if (DebugVisualizationMode == RAY_TRACING_DEBUG_VIZ_INSTANCE_OVERLAP)
+	{
+		DrawInstanceOverlap(GraphBuilder, Scene, View, SceneColorTexture, OutputDepthTexture);
+	}
 }
 
 void FDeferredShadingSceneRenderer::RayTracingDisplayPicking(const FRayTracingPickingFeedback& PickingFeedback, FScreenMessageWriter& Writer)
@@ -798,6 +1093,30 @@ bool IsRayTracingInstanceDebugDataEnabled(const FViewInfo& View)
 	}
 }
 
+bool IsRayTracingInstanceOverlapEnabled(const FViewInfo& View)
+{
+#if !UE_BUILD_SHIPPING	
+	static TArray<FName> RayTracingDebugVisualizationModes;
+	if (RayTracingDebugVisualizationModes.Num() == 0)
+	{
+		RayTracingDebugVisualizationModes.Add(*LOCTEXT("Instance Overlap", "Instance Overlap").ToString());
+	}
+
+	FString ConsoleViewMode = CVarRayTracingDebugMode.GetValueOnRenderThread();
+	if (!ConsoleViewMode.IsEmpty())
+	{
+		return RayTracingDebugVisualizationModes.Contains(FName(*ConsoleViewMode));
+	}
+	else if (View.CurrentRayTracingDebugVisualizationMode != NAME_None)
+	{
+		return RayTracingDebugVisualizationModes.Contains(View.CurrentRayTracingDebugVisualizationMode);
+	}
+#endif
+	{
+		return false;
+	}
+}
+
 #undef LOCTEXT_NAMESPACE
 
 #endif
diff --git a/Engine/Source/Runtime/Renderer/Private/RayTracing/RayTracingInstanceBufferUtil.cpp b/Engine/Source/Runtime/Renderer/Private/RayTracing/RayTracingInstanceBufferUtil.cpp
--- a/Engine/Source/Runtime/Renderer/Private/RayTracing/RayTracingInstanceBufferUtil.cpp
+++ b/Engine/Source/Runtime/Renderer/Private/RayTracing/RayTracingInstanceBufferUtil.cpp
@@ -292,10 +292,14 @@ struct FRayTracingBuildInstanceBufferCS : public FGlobalShader
 
 		SHADER_PARAMETER(FVector3f, ViewTilePosition)
 		SHADER_PARAMETER(FVector3f, RelativePreViewTranslation)
+
+		// Debug parameters
+		SHADER_PARAMETER_UAV(RWStructuredBuffer<uint>, RWDebugInstanceGPUSceneIndices)
 	END_SHADER_PARAMETER_STRUCT()
 
 	class FUseGPUSceneDim : SHADER_PERMUTATION_BOOL("USE_GPUSCENE");
-	using FPermutationDomain = TShaderPermutationDomain<FUseGPUSceneDim>;
+	class FOutputInstanceGPUSceneIndexDim : SHADER_PERMUTATION_BOOL("OUTPUT_INSTANCE_GPUSCENE_INDEX");
+	using FPermutationDomain = TShaderPermutationDomain<FUseGPUSceneDim, FOutputInstanceGPUSceneIndexDim>;
 		
 	static constexpr uint32 ThreadGroupSize = 64;
 
@@ -326,7 +330,8 @@ void BuildRayTracingInstanceBuffer(
 	FUnorderedAccessViewRHIRef InstancesUAV,
 	FShaderResourceViewRHIRef InstanceUploadSRV,
 	FShaderResourceViewRHIRef AccelerationStructureAddressesSRV,
-	FShaderResourceViewRHIRef InstanceTransformSRV)
+	FShaderResourceViewRHIRef InstanceTransformSRV,
+	FUnorderedAccessViewRHIRef DebugInstanceGPUSceneIndexUAV)
 {
 	FRayTracingBuildInstanceBufferCS::FParameters PassParams;
 	PassParams.InstancesDescriptors = InstancesUAV;
@@ -347,13 +352,16 @@ void BuildRayTracingInstanceBuffer(
 		PassParams.GPUScenePrimitiveSceneData = GPUScene->PrimitiveBuffer->GetSRV();
 	}
 
+	PassParams.RWDebugInstanceGPUSceneIndices = DebugInstanceGPUSceneIndexUAV;
+
 	FRayTracingBuildInstanceBufferCS::FPermutationDomain PermutationVector;
 	PermutationVector.Set<FRayTracingBuildInstanceBufferCS::FUseGPUSceneDim>(InstanceTransformSRV == nullptr);
+	PermutationVector.Set<FRayTracingBuildInstanceBufferCS::FOutputInstanceGPUSceneIndexDim>(DebugInstanceGPUSceneIndexUAV != nullptr);
 
 	auto ComputeShader = GetGlobalShaderMap(GMaxRHIFeatureLevel)->GetShader<FRayTracingBuildInstanceBufferCS>(PermutationVector);
 	const int32 GroupSize = FMath::DivideAndRoundUp(PassParams.NumInstances, FRayTracingBuildInstanceBufferCS::ThreadGroupSize);
 
-	//ClearUnusedGraphResources(ComputeShader, PassParams);
+	//ClearUnusedGraphResources(ComputeShader, &PassParams);
 
 	SetComputePipelineState(RHICmdList, ComputeShader.GetComputeShader());
 
@@ -375,7 +383,8 @@ void BuildRayTracingInstanceBuffer(
 	FShaderResourceViewRHIRef CPUInstanceTransformSRV,
 	uint32 NumNativeGPUSceneInstances,
 	uint32 NumNativeCPUInstances,
-	TConstArrayView<FRayTracingGPUInstance> GPUInstances)
+	TConstArrayView<FRayTracingGPUInstance> GPUInstances,
+	FUnorderedAccessViewRHIRef DebugInstanceGPUSceneIndexUAV)
 {
 	if (NumNativeGPUSceneInstances > 0)
 	{
@@ -389,7 +398,8 @@ void BuildRayTracingInstanceBuffer(
 			InstancesUAV,
 			InstanceUploadSRV,
 			AccelerationStructureAddressesSRV,
-			nullptr);
+			nullptr,
+			DebugInstanceGPUSceneIndexUAV);
 	}
 
 	if (NumNativeCPUInstances > 0)
@@ -404,7 +414,8 @@ void BuildRayTracingInstanceBuffer(
 			InstancesUAV,
 			InstanceUploadSRV,
 			AccelerationStructureAddressesSRV,
-			CPUInstanceTransformSRV);
+			CPUInstanceTransformSRV,
+			nullptr);
 	}
 
 	for (const auto& GPUInstance : GPUInstances)
@@ -422,7 +433,8 @@ void BuildRayTracingInstanceBuffer(
 			InstancesUAV,
 			InstanceUploadSRV,
 			AccelerationStructureAddressesSRV,
-			GPUInstance.TransformSRV);
+			GPUInstance.TransformSRV,
+			DebugInstanceGPUSceneIndexUAV);
 	}
 }
 
diff --git a/Engine/Source/Runtime/Renderer/Private/RayTracing/RayTracingScene.cpp b/Engine/Source/Runtime/Renderer/Private/RayTracing/RayTracingScene.cpp
--- a/Engine/Source/Runtime/Renderer/Private/RayTracing/RayTracingScene.cpp
+++ b/Engine/Source/Runtime/Renderer/Private/RayTracing/RayTracingScene.cpp
@@ -13,6 +13,7 @@
 
 BEGIN_SHADER_PARAMETER_STRUCT(FBuildInstanceBufferPassParams, )
 	SHADER_PARAMETER_RDG_BUFFER_UAV(RWStructuredBuffer, InstanceBuffer)
+	SHADER_PARAMETER_RDG_BUFFER_UAV(RWStructuredBuffer, DebugInstanceGPUSceneIndexBuffer)
 END_SHADER_PARAMETER_STRUCT()
 
 FRayTracingScene::FRayTracingScene()
@@ -196,6 +197,20 @@ void FRayTracingScene::CreateWithInitializationData(FRDGBuilder& GraphBuilder, c
 
 		FBuildInstanceBufferPassParams* PassParams = GraphBuilder.AllocParameters<FBuildInstanceBufferPassParams>();
 		PassParams->InstanceBuffer = GraphBuilder.CreateUAV(InstanceBuffer);
+		PassParams->DebugInstanceGPUSceneIndexBuffer = nullptr;
+
+		if (bNeedsDebugInstanceGPUSceneIndexBuffer)
+		{
+			FRDGBufferDesc DebugInstanceGPUSceneIndexBufferDesc;
+			DebugInstanceGPUSceneIndexBufferDesc.Usage = EBufferUsageFlags::UnorderedAccess | EBufferUsageFlags::ShaderResource | EBufferUsageFlags::StructuredBuffer;
+			DebugInstanceGPUSceneIndexBufferDesc.BytesPerElement = sizeof(uint32);
+			DebugInstanceGPUSceneIndexBufferDesc.NumElements = NumNativeInstances;
+
+			DebugInstanceGPUSceneIndexBuffer = GraphBuilder.CreateBuffer(DebugInstanceGPUSceneIndexBufferDesc, TEXT("FRayTracingScene::DebugInstanceGPUSceneIndexBuffer"));
+			PassParams->DebugInstanceGPUSceneIndexBuffer = GraphBuilder.CreateUAV(DebugInstanceGPUSceneIndexBuffer);
+
+			AddClearUAVPass(GraphBuilder, PassParams->DebugInstanceGPUSceneIndexBuffer, 0xFFFFFFFF);
+		}
 
 		GraphBuilder.AddPass(
 			RDG_EVENT_NAME("BuildTLASInstanceBuffer"),
@@ -252,7 +267,8 @@ void FRayTracingScene::CreateWithInitializationData(FRDGBuilder& GraphBuilder, c
 					TransformUploadSRV,
 					NumNativeGPUSceneInstances,
 					NumNativeCPUInstances,
-					GPUInstances);
+					GPUInstances,
+					PassParams->DebugInstanceGPUSceneIndexBuffer ? PassParams->DebugInstanceGPUSceneIndexBuffer->GetRHI() : nullptr);
 			});
 	}
 }
@@ -327,6 +343,7 @@ void FRayTracingScene::Reset()
 
 	BuildScratchBuffer = nullptr;
 	InstanceDebugBuffer = nullptr;
+	DebugInstanceGPUSceneIndexBuffer = nullptr;
 }
 
 void FRayTracingScene::ResetAndReleaseResources()
diff --git a/Engine/Source/Runtime/Renderer/Private/RayTracing/RayTracingScene.h b/Engine/Source/Runtime/Renderer/Private/RayTracing/RayTracingScene.h
--- a/Engine/Source/Runtime/Renderer/Private/RayTracing/RayTracingScene.h
+++ b/Engine/Source/Runtime/Renderer/Private/RayTracing/RayTracingScene.h
@@ -123,6 +123,9 @@ public:
 	FRDGBufferRef InstanceDebugBuffer = nullptr;
 	TArray<FRayTracingInstanceDebugData> InstancesDebugData;
 
+	FRDGBufferRef DebugInstanceGPUSceneIndexBuffer = nullptr;
+	bool bNeedsDebugInstanceGPUSceneIndexBuffer = false;
+
 	// Used for transforming to translated world space in which TLAS was built.
 	FVector RelativePreViewTranslation = FVector::Zero();
 	FVector3f ViewTilePosition = FVector3f::ZeroVector;
diff --git a/Engine/Source/Runtime/Renderer/Private/RayTracing/RaytracingOptions.h b/Engine/Source/Runtime/Renderer/Private/RayTracing/RaytracingOptions.h
--- a/Engine/Source/Runtime/Renderer/Private/RayTracing/RaytracingOptions.h
+++ b/Engine/Source/Runtime/Renderer/Private/RayTracing/RaytracingOptions.h
@@ -76,6 +76,7 @@ extern float GetRaytracingMaxNormalBias();
 extern int32 GetRayTracingCulling();
 extern float GetRayTracingCullingRadius();
 extern bool IsRayTracingInstanceDebugDataEnabled(const FViewInfo& View);
+extern bool IsRayTracingInstanceOverlapEnabled(const FViewInfo& View);
 
 extern bool CanUseRayTracingAMDHitToken();
 
@@ -156,7 +157,12 @@ FORCEINLINE bool CanUseRayTracingAMDHitToken()
 	return false;
 }
 
-FORCEINLINE bool IsRayTracingInstanceDebugDataEnabled()
+FORCEINLINE bool IsRayTracingInstanceDebugDataEnabled(const FViewInfo& View)
+{
+	return false;
+}
+
+FORCEINLINE bool IsRayTracingInstanceOverlapEnabled(const FViewInfo& View)
 {
 	return false;
 }
diff --git a/Engine/Source/Runtime/Renderer/Private/SceneVisibility.cpp b/Engine/Source/Runtime/Renderer/Private/SceneVisibility.cpp
--- a/Engine/Source/Runtime/Renderer/Private/SceneVisibility.cpp
+++ b/Engine/Source/Runtime/Renderer/Private/SceneVisibility.cpp
@@ -4045,10 +4045,10 @@ void FSceneRenderer::PreVisibilityFrameSetup(FRDGBuilder& GraphBuilder, const FS
 	if (Scene && Views.Num())
 	{
 		const int32 ReferenceViewIndex = 0;
-		FViewInfo& ReferenceView = Views[ReferenceViewIndex];
-		FRayTracingScene& RayTracingScene = Scene->RayTracingScene;
+		const FViewInfo& ReferenceView = Views[ReferenceViewIndex];
 
 		Scene->RayTracingScene.InitPreViewTranslation(ReferenceView.ViewMatrices);
+		Scene->RayTracingScene.bNeedsDebugInstanceGPUSceneIndexBuffer = IsRayTracingInstanceOverlapEnabled(ReferenceView);
 	}
 #endif
 
diff --git a/Engine/Source/Runtime/Renderer/Private/Tests/RayTracingTestbed.cpp b/Engine/Source/Runtime/Renderer/Private/Tests/RayTracingTestbed.cpp
--- a/Engine/Source/Runtime/Renderer/Private/Tests/RayTracingTestbed.cpp
+++ b/Engine/Source/Runtime/Renderer/Private/Tests/RayTracingTestbed.cpp
@@ -208,7 +208,8 @@ bool RunRayTracingTestbed_RenderThread(const FString& Parameters)
 		TransformUploadSRV,
 		RayTracingScene.NumNativeGPUSceneInstances,
 		RayTracingScene.NumNativeCPUInstances,
-		{});
+		{},
+		nullptr);
 
 	RHICmdList.BindAccelerationStructureMemory(RayTracingScene.Scene, SceneBuffer, 0);
 
diff --git a/Engine/Source/Runtime/Renderer/Public/RayTracingInstanceBufferUtil.h b/Engine/Source/Runtime/Renderer/Public/RayTracingInstanceBufferUtil.h
--- a/Engine/Source/Runtime/Renderer/Public/RayTracingInstanceBufferUtil.h
+++ b/Engine/Source/Runtime/Renderer/Public/RayTracingInstanceBufferUtil.h
@@ -94,6 +94,7 @@ RENDERER_API void BuildRayTracingInstanceBuffer(
 	FShaderResourceViewRHIRef CPUInstanceTransformSRV,
 	uint32 NumNativeGPUSceneInstances,
 	uint32 NumNativeCPUInstances,
-	TConstArrayView<FRayTracingGPUInstance> GPUInstances);
+	TConstArrayView<FRayTracingGPUInstance> GPUInstances,
+	FUnorderedAccessViewRHIRef DebugInstanceGPUSceneIndexUAV);
 
 #endif // RHI_RAYTRACING
\ No newline at end of file
