Fix crash in RectLightTextureAtlas when using world partitioning stream.

This CL fixes incorrect use of UTexture on the rendering thread causing crashes when a world partition is streamed out.

diff --git a/Engine/Source/Runtime/Renderer/Private/RectLightTextureManager.cpp b/Engine/Source/Runtime/Renderer/Private/RectLightTextureManager.cpp
--- a/Engine/Source/Runtime/Renderer/Private/RectLightTextureManager.cpp
+++ b/Engine/Source/Runtime/Renderer/Private/RectLightTextureManager.cpp
@@ -127,10 +127,11 @@ struct FAtlasSlot
 {
 	uint32 Id = InvalidSlotIndex;
 	FAtlasRect Rect;
-	UTexture* SourceTexture = nullptr;
+	FTextureReferenceRHIRef SourceTexture = nullptr;
 	uint32 RefCount = 0;
 	bool bForceRefresh = false;
 	bool IsValid() const { return SourceTexture != nullptr; }
+	FRHITexture* GetTextureRHI() const { return SourceTexture->GetReferencedTexture(); }
 };
 
 // Store info for copying one atlas slot to another one, when a new layout is created
@@ -140,7 +141,7 @@ struct FAtlasCopySlot
 	FIntPoint SrcOrigin = InvalidOrigin;
 	FIntPoint DstOrigin = InvalidOrigin;
 	FIntPoint Resolution = FIntPoint::ZeroValue;
-	UTexture* SourceTexture = nullptr;
+	FTextureReferenceRHIRef SourceTexture = nullptr;
 };
 
 // Texture manager, holding all atlas data & description
@@ -430,14 +431,14 @@ static void AddSlotsPass(
 
 			switch (SlotIt)
 			{
-			case 0: Parameters->InTexture0 = Slot.SourceTexture->GetResource()->TextureRHI; break;
-			case 1: Parameters->InTexture1 = Slot.SourceTexture->GetResource()->TextureRHI; break;
-			case 2: Parameters->InTexture2 = Slot.SourceTexture->GetResource()->TextureRHI; break;
-			case 3: Parameters->InTexture3 = Slot.SourceTexture->GetResource()->TextureRHI; break;
-			case 4: Parameters->InTexture4 = Slot.SourceTexture->GetResource()->TextureRHI; break;
-			case 5: Parameters->InTexture5 = Slot.SourceTexture->GetResource()->TextureRHI; break;
-			case 6: Parameters->InTexture6 = Slot.SourceTexture->GetResource()->TextureRHI; break;
-			case 7: Parameters->InTexture7 = Slot.SourceTexture->GetResource()->TextureRHI; break;
+			case 0: Parameters->InTexture0 = Slot.GetTextureRHI(); break;
+			case 1: Parameters->InTexture1 = Slot.GetTextureRHI(); break;
+			case 2: Parameters->InTexture2 = Slot.GetTextureRHI(); break;
+			case 3: Parameters->InTexture3 = Slot.GetTextureRHI(); break;
+			case 4: Parameters->InTexture4 = Slot.GetTextureRHI(); break;
+			case 5: Parameters->InTexture5 = Slot.GetTextureRHI(); break;
+			case 6: Parameters->InTexture6 = Slot.GetTextureRHI(); break;
+			case 7: Parameters->InTexture7 = Slot.GetTextureRHI(); break;
 			}
 		}
 
@@ -647,14 +648,14 @@ static void FilterSlotsPass(
 					FAtlasSlot& DstMIPSlot		= DstMIPSlots.AddDefaulted_GetRef();
 					DstMIPSlot.Rect.Origin		= ToMIP(Slot.Rect.Origin, DstMip);
 					DstMIPSlot.Rect.Resolution	= ToMIP(Slot.Rect.Resolution, DstMip);
-					DstMIPSlot.SourceTexture= Slot.SourceTexture;
+					DstMIPSlot.SourceTexture	= Slot.SourceTexture;
 				}
 
 				{
 					FAtlasSlot& SrcMIPSlot		= SrcMIPSlots.AddDefaulted_GetRef();
 					SrcMIPSlot.Rect.Origin		= ToMIP(Slot.Rect.Origin, SrcMip);
 					SrcMIPSlot.Rect.Resolution	= ToMIP(Slot.Rect.Resolution, SrcMip);
-					SrcMIPSlot.SourceTexture= Slot.SourceTexture;
+					SrcMIPSlot.SourceTexture	= Slot.SourceTexture;
 				}
 			}
 		}
@@ -1229,7 +1230,7 @@ static void PackAtlas(
 
 			// Check if the texture resolution has changed (due to streaming)
 			// If the texture resolution has change, reset the slot to be handled as a new slot
-			const FIntVector TextureResolution = ValidSlot.SourceTexture->GetResource()->TextureRHI->GetSizeXYZ();
+			const FIntVector TextureResolution = ValidSlot.GetTextureRHI()->GetSizeXYZ();
 			const bool bHasTextureResolutionChanged = ValidSlot.Rect.Resolution != FIntPoint(TextureResolution.X, TextureResolution.Y);
 			if (bHasTextureResolutionChanged || ValidSlot.bForceRefresh)
 			{
@@ -1336,7 +1337,7 @@ static void PackAtlas(
 					CurrentSourceTextureMIPBias++;
 					for (FAtlasSlot& Slot : ValidSlots)
 					{
-						const FRHITexture* Tex = Slot.SourceTexture->GetResource()->TextureRHI;
+						const FRHITexture* Tex = Slot.GetTextureRHI();
 						Slot.Rect.Origin = InvalidOrigin;
 						Slot.Rect.Resolution = ToMIP(FIntPoint(Tex->GetSizeXYZ().X, Tex->GetSizeXYZ().Y), CurrentSourceTextureMIPBias);
 					}
@@ -1367,7 +1368,7 @@ uint32 AddRectLightTexture(UTexture* In)
 		bool bFound = false;
 		for (FAtlasSlot& Slot : GRectLightTextureManager.AtlasSlots)
 		{
-			if (Slot.SourceTexture == In)
+			if (Slot.SourceTexture == In->TextureReference.TextureReferenceRHI)
 			{
 				Slot.RefCount++;
 				SlotIndex = Slot.Id;
@@ -1390,10 +1391,10 @@ uint32 AddRectLightTexture(UTexture* In)
 				Slot = &GRectLightTextureManager.AtlasSlots.AddDefaulted_GetRef();
 			}
 
-			const FRHITexture* Tex = In->GetResource()->TextureRHI;
+			const FRHITexture* Tex = In->TextureReference.TextureReferenceRHI;
 
 			*Slot = FAtlasSlot();
-			Slot->SourceTexture = In;
+			Slot->SourceTexture = In->TextureReference.TextureReferenceRHI;
 			Slot->Id = SlotIndex;
 			Slot->Rect.Origin = InvalidOrigin;
 			Slot->Rect.Resolution = FIntPoint(Tex->GetSizeXYZ().X, Tex->GetSizeXYZ().Y);
@@ -1495,9 +1496,12 @@ void UpdateRectLightAtlasTexture(FRDGBuilder& GraphBuilder, const ERHIFeatureLev
 	{
 		for (FAtlasSlot& Slot : GRectLightTextureManager.AtlasSlots)
 		{
-			if (Slot.IsValid())
+			if (Slot.IsValid() && Slot.GetTextureRHI())
 			{
-				const FIntVector SourceResolution = Slot.SourceTexture->GetResource()->TextureRHI->GetSizeXYZ();
+				FRHITexture* TextureRHI = Slot.GetTextureRHI();
+				check(TextureRHI != nullptr);
+
+				const FIntVector SourceResolution = TextureRHI->GetSizeXYZ();
 				const FIntPoint  SourceResolutionMIPed = ToMIP(FIntPoint(SourceResolution.X, SourceResolution.Y), GRectLightTextureManager.AtlasLayout.SourceTextureMIPBias);
 				if (SourceResolutionMIPed.X > Slot.Rect.Resolution.X || SourceResolutionMIPed.Y > Slot.Rect.Resolution.Y)
 				{
@@ -1665,7 +1669,7 @@ FAtlasTextureInvalidationScope::FAtlasTextureInvalidationScope(UTexture* In)
 			// then we lock the atlas to not update the atlas during 
 			// the capture, and force the target texture to refresh 
 			// its data during the next update
-			if (Slot.SourceTexture == In)
+			if (Slot.SourceTexture == In->TextureReference.TextureReferenceRHI)
 			{
 				// Sanity check, allow a single lock/capture refresh at a time.
 				check(GRectLightTextureManager.bLock == false);
