Volumetric Cloud now uses mode 2 in order to test how it would look and cost. Should fix reprojection and opaque edge artefact.
Tested in the editor and was fine.

diff --git a/Engine/Config/BaseScalability.ini b/Engine/Config/BaseScalability.ini
--- a/Engine/Config/BaseScalability.ini
+++ b/Engine/Config/BaseScalability.ini
@@ -604,6 +604,7 @@ r.SkyAtmosphere.TransmittanceLUT.SampleCount=10.0
 r.SkyAtmosphere.MultiScatteringLUT.SampleCount=15.0
 r.SkyLight.RealTimeReflectionCapture=0
 r.VolumetricCloud.ShadowMap.SpatialFiltering=0
+r.VolumetricRenderTarget.Mode=2
 fx.Niagara.QualityLevel=0
 
 [EffectsQuality@1]
@@ -635,6 +636,7 @@ r.SkyAtmosphere.TransmittanceLUT.SampleCount=10.0
 r.SkyAtmosphere.MultiScatteringLUT.SampleCount=15.0
 r.SkyLight.RealTimeReflectionCapture=0
 r.VolumetricCloud.ShadowMap.SpatialFiltering=0
+r.VolumetricRenderTarget.Mode=2
 fx.Niagara.QualityLevel=1
 
 [EffectsQuality@2]
@@ -666,6 +668,7 @@ r.SkyAtmosphere.TransmittanceLUT.SampleCount=10.0
 r.SkyAtmosphere.MultiScatteringLUT.SampleCount=15.0
 r.SkyLight.RealTimeReflectionCapture=1
 r.VolumetricCloud.ShadowMap.SpatialFiltering=0
+r.VolumetricRenderTarget.Mode=2
 fx.Niagara.QualityLevel=2
 
 [EffectsQuality@3]
@@ -697,6 +700,7 @@ r.SkyAtmosphere.TransmittanceLUT.SampleCount=10.0
 r.SkyAtmosphere.MultiScatteringLUT.SampleCount=15.0
 r.SkyLight.RealTimeReflectionCapture=1
 r.VolumetricCloud.ShadowMap.SpatialFiltering=0
+r.VolumetricRenderTarget.Mode=2
 fx.Niagara.QualityLevel=3
 
 [EffectsQuality@Cine]
@@ -728,6 +732,7 @@ r.SkyAtmosphere.TransmittanceLUT.SampleCount=30.0
 r.SkyAtmosphere.MultiScatteringLUT.SampleCount=20.0
 r.SkyLight.RealTimeReflectionCapture=1
 r.VolumetricCloud.ShadowMap.SpatialFiltering=1
+r.VolumetricRenderTarget.Mode=2
 fx.Niagara.QualityLevel=4
 
 ;-----------------------------------------------------------------------------------------------------------------
diff --git a/Engine/Source/Runtime/Renderer/Private/VolumetricRenderTarget.cpp b/Engine/Source/Runtime/Renderer/Private/VolumetricRenderTarget.cpp
--- a/Engine/Source/Runtime/Renderer/Private/VolumetricRenderTarget.cpp
+++ b/Engine/Source/Runtime/Renderer/Private/VolumetricRenderTarget.cpp
@@ -19,37 +19,37 @@
 static TAutoConsoleVariable<int32> CVarVolumetricRenderTarget(
 	TEXT("r.VolumetricRenderTarget"), 1,
 	TEXT(""),
-	ECVF_SetByScalability | ECVF_RenderThreadSafe);
+	ECVF_RenderThreadSafe | ECVF_Scalability);
 
 static TAutoConsoleVariable<float> CVarVolumetricRenderTargetUvNoiseScale(
 	TEXT("r.VolumetricRenderTarget.UvNoiseScale"), 0.5f,
 	TEXT("Used when r.VolumetricRenderTarget.UpsamplingMode is in a mode using jitter - this value scales the amount of jitter."),
-	ECVF_SetByScalability);
+	ECVF_RenderThreadSafe | ECVF_Scalability);
 
 static TAutoConsoleVariable<float> CVarVolumetricRenderTargetUvNoiseSampleAcceptanceWeight(
 	TEXT("r.VolumetricRenderTarget.UvNoiseSampleAcceptanceWeight"), 20.0f,
 	TEXT("Used when r.VolumetricRenderTarget.UpsamplingMode is in a mode using jitter - this value control the acceptance of noisy cloud samples according to their similarities. A higher value means large differences will be less accepted for blending."),
-	ECVF_SetByScalability);
+	ECVF_RenderThreadSafe | ECVF_Scalability);
 
 static TAutoConsoleVariable<int32> CVarVolumetricRenderTargetMode(
 	TEXT("r.VolumetricRenderTarget.Mode"), 0,
 	TEXT("[0] trace quarter resolution + reconstruct at half resolution + upsample [1] trace half res + reconstruct full res + upsample [2] trace at quarter resolution + reconstruct full resolution (cannot intersect with opaque meshes and forces UpsamplingMode=2 [3] trace 1/8 resolution + reconstruct at half resolution + upsample)"),
-	ECVF_SetByScalability);
+	ECVF_RenderThreadSafe | ECVF_Scalability);
 
 static TAutoConsoleVariable<int32> CVarVolumetricRenderTargetUpsamplingMode(
 	TEXT("r.VolumetricRenderTarget.UpsamplingMode"), 4,
 	TEXT("Used in compositing volumetric RT over the scene. [0] bilinear [1] bilinear + jitter [2] nearest + depth test [3] bilinear + jitter + keep closest [4] bilaterial upsampling"),
-	ECVF_SetByScalability);
+	ECVF_RenderThreadSafe | ECVF_Scalability);
 
 static TAutoConsoleVariable<int32> CVarVolumetricRenderTargetPreferAsyncCompute(
 	TEXT("r.VolumetricRenderTarget.PreferAsyncCompute"), 0,
 	TEXT("Whether to prefer using async compute to generate volumetric cloud render targets."),
-	ECVF_SetByScalability | ECVF_RenderThreadSafe);
+	ECVF_RenderThreadSafe | ECVF_Scalability);
 
 static TAutoConsoleVariable<int32> CVarVolumetricRenderTargetReprojectionBoxConstraint(
 	TEXT("r.VolumetricRenderTarget.ReprojectionBoxConstraint"), 0,
 	TEXT("Whether reprojected data should be constrained to the new incoming cloud data neighborhod value."),
-	ECVF_SetByScalability | ECVF_RenderThreadSafe);
+	ECVF_RenderThreadSafe | ECVF_Scalability);
 
 
 static float GetUvNoiseSampleAcceptanceWeight()
