Use GameThread function during PSO Collection instead of RenderThread function

diff --git a/Engine/Source/Runtime/Renderer/Private/SingleLayerWaterRendering.cpp b/Engine/Source/Runtime/Renderer/Private/SingleLayerWaterRendering.cpp
--- a/Engine/Source/Runtime/Renderer/Private/SingleLayerWaterRendering.cpp
+++ b/Engine/Source/Runtime/Renderer/Private/SingleLayerWaterRendering.cpp
@@ -1563,7 +1563,7 @@ void FSingleLayerWaterDepthPrepassMeshProcessor::CollectPSOInitializers(const FS
 
 		if (BlendMode == BLEND_Opaque
 			&& bSupportPositionOnlyStream
-			&& !Material.MaterialModifiesMeshPosition_RenderThread()
+			&& !Material.MaterialModifiesMeshPosition_GameThread()
 			&& Material.WritesEveryPixel())
 		{
 			const FMaterialRenderProxy& DefaultProxy = *UMaterial::GetDefaultMaterial(MD_Surface)->GetRenderProxy();
@@ -1576,7 +1576,7 @@ void FSingleLayerWaterDepthPrepassMeshProcessor::CollectPSOInitializers(const FS
 			const bool bMaterialMasked = !Material.WritesEveryPixel();
 			const FMaterial* EffectiveMaterial = &Material;
 
-			if (!bMaterialMasked && !Material.MaterialModifiesMeshPosition_RenderThread())
+			if (!bMaterialMasked && !Material.MaterialModifiesMeshPosition_GameThread())
 			{
 				// Override with the default material for opaque materials that are not two sided
 				FMaterialRenderProxy* EffectiveMaterialRenderProxy = UMaterial::GetDefaultMaterial(MD_Surface)->GetRenderProxy();
