Fix crash in MaterialShader occurring in a Shipping build when bUseChecksInShipping is set to true.

diff --git a/Engine/Source/Runtime/Renderer/Private/ShaderBaseClasses.cpp b/Engine/Source/Runtime/Renderer/Private/ShaderBaseClasses.cpp
--- a/Engine/Source/Runtime/Renderer/Private/ShaderBaseClasses.cpp
+++ b/Engine/Source/Runtime/Renderer/Private/ShaderBaseClasses.cpp
@@ -76,7 +76,7 @@ FMaterialShader::FMaterialShader(const FMaterialShaderType::CompiledShaderInitia
 #endif // WITH_EDITORONLY_DATA
 {
 #if WITH_EDITORONLY_DATA
-	check(!DebugDescription.IsEmpty());
+	check(UE_BUILD_SHIPPING || !DebugDescription.IsEmpty());
 	DebugUniformExpressionUBLayout.CopyFrom(Initializer.UniformExpressionSet.GetUniformBufferLayoutInitializer());
 #endif // WITH_EDITORONLY_DATA
 
