Refactored distance field streaming to use bulk data streaming API

#rnx

diff --git a/Engine/Source/Runtime/Renderer/Private/DistanceFieldStreaming.cpp b/Engine/Source/Runtime/Renderer/Private/DistanceFieldStreaming.cpp
--- a/Engine/Source/Runtime/Renderer/Private/DistanceFieldStreaming.cpp
+++ b/Engine/Source/Runtime/Renderer/Private/DistanceFieldStreaming.cpp
@@ -438,13 +438,13 @@ void FDistanceFieldBlockAllocator::Free(const TArray<int32, TInlineAllocator<4>>
 class FDistanceFieldStreamingUpdateTask
 {
 public:
-	explicit FDistanceFieldStreamingUpdateTask(const FDistanceFieldAsyncUpdateParameters& InParams) : Parameters(InParams) {}
+	explicit FDistanceFieldStreamingUpdateTask(FDistanceFieldAsyncUpdateParameters&& InParams) : Parameters(MoveTemp(InParams)) {}
 
 	FDistanceFieldAsyncUpdateParameters Parameters;
 
 	void DoTask(ENamedThreads::Type CurrentThread, const FGraphEventRef& MyCompletionGraphEvent)
 	{
-		Parameters.DistanceFieldSceneData->AsyncUpdate(Parameters);
+		Parameters.DistanceFieldSceneData->AsyncUpdate(MoveTemp(Parameters));
 	}
 
 	static ESubsequentsMode::Type	GetSubsequentsMode()	{ return ESubsequentsMode::TrackSubsequents; }
@@ -452,7 +452,7 @@ public:
 	FORCEINLINE TStatId				GetStatId() const		{ return TStatId(); }
 };
 
-void FDistanceFieldSceneData::AsyncUpdate(FDistanceFieldAsyncUpdateParameters UpdateParameters)
+void FDistanceFieldSceneData::AsyncUpdate(FDistanceFieldAsyncUpdateParameters&& UpdateParameters)
 {
 	QUICK_SCOPE_CYCLE_COUNTER(STAT_FDistanceFieldSceneData_AsyncUpdate);
 	TRACE_CPUPROFILER_EVENT_SCOPE(FDistanceFieldSceneData::AsyncUpdate);
@@ -464,7 +464,7 @@ void FDistanceFieldSceneData::AsyncUpdate(FDistanceFieldAsyncUpdateParameters Up
 	int32 BrickUploadIndex = 0;
 	int32 IndirectionUploadIndex = 0;
 
-	for (FDistanceFieldReadRequest ReadRequest : UpdateParameters.ReadRequestsToUpload)
+	for (FDistanceFieldReadRequest& ReadRequest : UpdateParameters.ReadRequestsToUpload)
 	{
 		const FDistanceFieldAssetState& AssetState = AssetStateArray[ReadRequest.AssetSetId];
 		const int32 ReversedMipIndex = ReadRequest.ReversedMipIndex;
@@ -472,7 +472,9 @@ void FDistanceFieldSceneData::AsyncUpdate(FDistanceFieldAsyncUpdateParameters Up
 		const int32 MipIndex = AssetState.BuiltData->Mips.Num() - ReversedMipIndex - 1;
 		const FSparseDistanceFieldMip& MipBuiltData = AssetState.BuiltData->Mips[MipIndex];
 
-		const uint8* BulkDataReadPtr = ReadRequest.BulkData ? ReadRequest.ReadOutputDataPtr : ReadRequest.AlwaysLoadedDataPtr;
+		const uint8* BulkDataReadPtr = ReadRequest.BulkData != nullptr
+			? ReadRequest.RequestBuffer.GetData()
+			: ReadRequest.AlwaysLoadedDataPtr;
 
 #if WITH_EDITOR
 		if (ReadRequest.BulkData)
@@ -581,53 +583,19 @@ void FDistanceFieldSceneData::AsyncUpdate(FDistanceFieldAsyncUpdateParameters Up
 	}
 
 #if !WITH_EDITOR
-
-	for (FDistanceFieldReadRequest ReadRequest : UpdateParameters.ReadRequestsToCleanUp)
+	if (UpdateParameters.NewReadRequests.Num() > 0)
 	{
-		if (ReadRequest.AsyncRequest)
-		{
-			check(ReadRequest.AsyncRequest->PollCompletion());	
-			delete ReadRequest.AsyncRequest;
-			delete ReadRequest.AsyncHandle;
-		}
-		else
-		{
-			check(ReadRequest.Request.Status().IsCompleted());
-		}
-
-		FMemory::Free(ReadRequest.ReadOutputDataPtr);
-	}
-
-	FIoBatch Batch;
-
-	for (FDistanceFieldReadRequest& ReadRequest : UpdateParameters.NewReadRequests)
-	{
-		check(ReadRequest.BulkSize > 0);
-		ReadRequest.ReadOutputDataPtr = (uint8*)FMemory::Malloc(ReadRequest.BulkSize);
-		const bool bIODispatcher = ReadRequest.BulkData->IsUsingIODispatcher();
-
-		if (bIODispatcher)
-		{
-			// Use IODispatcher when available
-			FIoChunkId ChunkID = ReadRequest.BulkData->CreateChunkId();
-			FIoReadOptions ReadOptions;
-			ReadOptions.SetRange(ReadRequest.BulkData->GetBulkDataOffsetInFile() + ReadRequest.BulkOffset, ReadRequest.BulkSize);
-			ReadOptions.SetTargetVa(ReadRequest.ReadOutputDataPtr);
-			ReadRequest.Request = Batch.Read(ChunkID, ReadOptions, IoDispatcherPriority_Low);
-		}
-		else
+		FBulkDataBatchRequest::FBatchBuilder Batch = FBulkDataBatchRequest::NewBatch(UpdateParameters.NewReadRequests.Num());
+		for (FDistanceFieldReadRequest& ReadRequest : UpdateParameters.NewReadRequests)
 		{
-			// Compatibility path without IODispatcher
-			ReadRequest.AsyncHandle = ReadRequest.BulkData->OpenAsyncReadHandle();
-			ReadRequest.AsyncRequest = ReadRequest.AsyncHandle->ReadRequest(ReadRequest.BulkData->GetBulkDataOffsetInFile() + ReadRequest.BulkOffset, ReadRequest.BulkSize, AIOP_Low, nullptr, ReadRequest.ReadOutputDataPtr);
+			check(ReadRequest.BulkSize > 0);
+			Batch.Read(*ReadRequest.BulkData, uint64(ReadRequest.BulkOffset), uint64(ReadRequest.BulkSize), AIOP_Low, ReadRequest.RequestBuffer, ReadRequest.RequestHandle);
+			ReadRequests.Add(MoveTemp(ReadRequest));
 		}
+		FBulkDataRequest::EStatus Status = Batch.Issue();
+		UE_CLOG(Status == FBulkDataRequest::EStatus::Ok, LogDistanceField, Error, TEXT("Failed to issue bulk data I/O request"));
 	}
-
-	Batch.Issue();
-
 #endif
-
-	ReadRequests.Append(UpdateParameters.NewReadRequests);
 }
 
 bool AssetHasOutstandingRequest(FSetElementId AssetSetId, const TArray<FDistanceFieldReadRequest>& ReadRequests)
@@ -772,8 +740,7 @@ void FDistanceFieldSceneData::ProcessStreamingRequestsFromGPU(
 void FDistanceFieldSceneData::ProcessReadRequests(
 	TArray<FDistanceFieldAssetMipId>& AssetDataUploads,
 	TArray<FDistanceFieldAssetMipId>& DistanceFieldAssetMipAdds,
-	TArray<FDistanceFieldReadRequest>& ReadRequestsToUpload,
-	TArray<FDistanceFieldReadRequest>& ReadRequestsToCleanUp)
+	TArray<FDistanceFieldReadRequest>& ReadRequestsToUpload)
 {
 	const uint32 BrickSizeBytes = GPixelFormats[DistanceField::DistanceFieldFormat].BlockBytes * DistanceField::BrickSize * DistanceField::BrickSize * DistanceField::BrickSize;
 	const SIZE_T TextureUploadLimitBytes = (SIZE_T)CVarTextureUploadLimitKBytes.GetValueOnRenderThread() * 1024;
@@ -803,41 +770,29 @@ void FDistanceFieldSceneData::ProcessReadRequests(
 
 	for (int32 RequestIndex = 0; RequestIndex < ReadRequests.Num(); RequestIndex++)
 	{
-		const FDistanceFieldReadRequest ReadRequest = ReadRequests[RequestIndex];
-
-		bool bReady = true;
-
-#if !WITH_EDITOR
-		if (ReadRequest.AsyncRequest)
-		{
-			bReady = bReady && ReadRequest.AsyncRequest->PollCompletion();
-		}
-		else
+		if (ReadRequests[RequestIndex].RequestHandle.IsCompleted() == false)
 		{
-			bReady = bReady && ReadRequest.Request.Status().IsCompleted();
+			continue;
 		}
-#endif
 
-		if (bReady)
+		FDistanceFieldReadRequest CompletedRequest = MoveTemp(ReadRequests[RequestIndex]);
+		CompletedRequest.RequestHandle = FBulkDataRequest();
+
+		ReadRequests.RemoveAtSwap(RequestIndex--);
+
+		if (AssetStateArray.IsValidId(CompletedRequest.AssetSetId) 
+			// Prevent attempting to upload after a different asset has been allocated at the same index
+			&& CompletedRequest.BuiltDataId == AssetStateArray[CompletedRequest.AssetSetId].BuiltData->GetId()
+			// Shader requires sequential reversed mips starting from 0, skip upload if the IO request got out of sync with the streaming feedback requests
+			&& CompletedRequest.ReversedMipIndex == AssetStateArray[CompletedRequest.AssetSetId].ReversedMips.Num())
 		{
-			ReadRequests.RemoveAtSwap(RequestIndex);
-			RequestIndex--;
-
-			if (AssetStateArray.IsValidId(ReadRequest.AssetSetId) 
-				// Prevent attempting to upload after a different asset has been allocated at the same index
-				&& ReadRequest.BuiltDataId == AssetStateArray[ReadRequest.AssetSetId].BuiltData->GetId()
-				// Shader requires sequential reversed mips starting from 0, skip upload if the IO request got out of sync with the streaming feedback requests
-				&& ReadRequest.ReversedMipIndex == AssetStateArray[ReadRequest.AssetSetId].ReversedMips.Num())
-			{
-				TextureUploadBytes += ReadRequest.NumDistanceFieldBricks * BrickSizeBytes;
+			TextureUploadBytes += CompletedRequest.NumDistanceFieldBricks * BrickSizeBytes;
 
-				DistanceFieldAssetMipAdds.Add(FDistanceFieldAssetMipId(ReadRequest.AssetSetId, ReadRequest.ReversedMipIndex));
-				// Re-upload mip0 to push the new NumMips to the shader
-				AssetDataUploads.Add(FDistanceFieldAssetMipId(ReadRequest.AssetSetId, 0));
-				ReadRequestsToUpload.Add(ReadRequest);
-			}
+			DistanceFieldAssetMipAdds.Add(FDistanceFieldAssetMipId(CompletedRequest.AssetSetId, CompletedRequest.ReversedMipIndex));
+			// Re-upload mip0 to push the new NumMips to the shader
+			AssetDataUploads.Add(FDistanceFieldAssetMipId(CompletedRequest.AssetSetId, 0));
 
-			ReadRequestsToCleanUp.Add(ReadRequest);
+			ReadRequestsToUpload.Add(MoveTemp(CompletedRequest));
 		}
 
 		// Stop uploading when we reach the limit
@@ -1333,9 +1288,8 @@ void FDistanceFieldSceneData::UpdateDistanceFieldAtlas(
 	ProcessStreamingRequestsFromGPU(NewReadRequests, AssetDataUploads);
 
 	TArray<FDistanceFieldReadRequest> ReadRequestsToUpload;
-	TArray<FDistanceFieldReadRequest> ReadRequestsToCleanUp;
 	// Build a list of completed read requests that should be uploaded to the GPU this frame
-	ProcessReadRequests(AssetDataUploads, DistanceFieldAssetMipAdds, ReadRequestsToUpload, ReadRequestsToCleanUp);
+	ProcessReadRequests(AssetDataUploads, DistanceFieldAssetMipAdds, ReadRequestsToUpload);
 
 	int32 NumIndirectionTableAdds = 0;
 	int32 NumBrickUploads = 0;
@@ -1459,14 +1413,13 @@ void FDistanceFieldSceneData::UpdateDistanceFieldAtlas(
 			UpdateParameters.BrickUploadCoordinatesPtr = AtlasUpload.BrickUploadCoordinatesPtr;
 		}
 
-		if (NewReadRequests.Num() || ReadRequestsToUpload.Num() || ReadRequestsToCleanUp.Num())
+		if (NewReadRequests.Num() || ReadRequestsToUpload.Num())
 		{
 			UpdateParameters.NewReadRequests = MoveTemp(NewReadRequests);
 			UpdateParameters.ReadRequestsToUpload = MoveTemp(ReadRequestsToUpload);
-			UpdateParameters.ReadRequestsToCleanUp = MoveTemp(ReadRequestsToCleanUp);
 
 			// TODO: We actually run this synchronously now after the RDG conversion, as it would otherwise immediately sync.
-			AsyncUpdate(UpdateParameters);
+			AsyncUpdate(MoveTemp(UpdateParameters));
 		}
 
 		if (NumBrickUploads > 0 || NumIndirectionTableAdds > 0)
diff --git a/Engine/Source/Runtime/Renderer/Private/ScenePrivate.h b/Engine/Source/Runtime/Renderer/Private/ScenePrivate.h
--- a/Engine/Source/Runtime/Renderer/Private/ScenePrivate.h
+++ b/Engine/Source/Runtime/Renderer/Private/ScenePrivate.h
@@ -1927,15 +1927,13 @@ struct FDistanceFieldReadRequest
 	const uint8* AlwaysLoadedDataPtr = nullptr;
 
 	// Inputs of read request
-	const FByteBulkData* BulkData = nullptr;
+	const FBulkData* BulkData = nullptr;
 	uint32 BulkOffset = 0;
 	uint32 BulkSize = 0;
 
 	// Outputs of read request
-	uint8* ReadOutputDataPtr = nullptr;
-	FIoRequest Request;
-	IAsyncReadFileHandle* AsyncHandle = nullptr;
-	IAsyncReadRequest* AsyncRequest = nullptr;
+	FBulkDataBatchReadRequest RequestHandle;
+	FIoBuffer RequestBuffer;
 };
 
 struct FDistanceFieldAsyncUpdateParameters
@@ -1949,7 +1947,6 @@ struct FDistanceFieldAsyncUpdateParameters
 
 	TArray<FDistanceFieldReadRequest> NewReadRequests;
 	TArray<FDistanceFieldReadRequest> ReadRequestsToUpload;
-	TArray<FDistanceFieldReadRequest> ReadRequestsToCleanUp;
 };
 
 /** Scene data used to manage distance field object buffers on the GPU. */
@@ -2109,8 +2106,7 @@ private:
 	void ProcessReadRequests(
 		TArray<FDistanceFieldAssetMipId>& AssetDataUploads,
 		TArray<FDistanceFieldAssetMipId>& DistanceFieldAssetMipAdds,
-		TArray<FDistanceFieldReadRequest>& ReadRequestsToUpload,
-		TArray<FDistanceFieldReadRequest>& ReadRequestsToCleanUp);
+		TArray<FDistanceFieldReadRequest>& ReadRequestsToUpload);
 
 	FRDGTexture* ResizeBrickAtlasIfNeeded(FRDGBuilder& GraphBuilder, FGlobalShaderMap* GlobalShaderMap);
 
@@ -2122,7 +2118,7 @@ private:
 	
 	void UploadAllAssetData(FRDGBuilder& GraphBuilder, FRDGBuffer* AssetDataBufferRDG);
 
-	void AsyncUpdate(FDistanceFieldAsyncUpdateParameters UpdateParameters);
+	void AsyncUpdate(FDistanceFieldAsyncUpdateParameters&& UpdateParameters);
 
 	void GenerateStreamingRequests(
 		FRDGBuilder& GraphBuilder, 
