Fix stationary local lights missing static shadow on mobile.

#UE-166949

diff --git a/Engine/Source/Runtime/Renderer/Private/MobileBasePass.cpp b/Engine/Source/Runtime/Renderer/Private/MobileBasePass.cpp
--- a/Engine/Source/Runtime/Renderer/Private/MobileBasePass.cpp
+++ b/Engine/Source/Runtime/Renderer/Private/MobileBasePass.cpp
@@ -272,47 +272,51 @@ ELightMapPolicyType MobileBasePass::SelectMeshLightmapPolicy(
 
 			if (LightMapInteraction.GetType() == LMIT_Texture && ReadOnlyCVARCache.bEnableLowQualityLightmaps)
 			{
-				// Lightmap path
-				if (bUseMovableLight)
+				const FShadowMapInteraction ShadowMapInteraction = (Mesh.LCI != nullptr)
+					? Mesh.LCI->GetShadowMapInteraction(FeatureLevel)
+					: FShadowMapInteraction();
+
+				if ((bUseStaticAndCSM || bMovableWithCSM) && !bUsesDeferredShading)
 				{
-					if (bUsesDeferredShading)
-					{
-						SelectedLightmapPolicy = LMP_LQ_LIGHTMAP;
-					}
-					else if (bMovableWithCSM && !bUsesDeferredShading)
+					if (ShadowMapInteraction.GetType() == SMIT_Texture &&
+						MobileDirectionalLight->ShouldRenderViewIndependentWholeSceneShadows() &&
+						ReadOnlyCVARCache.bMobileAllowDistanceFieldShadows)
 					{
-						SelectedLightmapPolicy = LMP_MOBILE_MOVABLE_DIRECTIONAL_LIGHT_CSM_WITH_LIGHTMAP;
+						SelectedLightmapPolicy = LMP_MOBILE_DISTANCE_FIELD_SHADOWS_LIGHTMAP_AND_CSM;
 					}
 					else
 					{
-						SelectedLightmapPolicy = LMP_MOBILE_MOVABLE_DIRECTIONAL_LIGHT_WITH_LIGHTMAP;						
-					}
-				}
-				else
-				{
-					const FShadowMapInteraction ShadowMapInteraction = (Mesh.LCI != nullptr)
-						? Mesh.LCI->GetShadowMapInteraction(FeatureLevel)
-						: FShadowMapInteraction();
-
-					if (bUseStaticAndCSM && !bUsesDeferredShading)
-					{
-						if (ShadowMapInteraction.GetType() == SMIT_Texture && 
-							MobileDirectionalLight->ShouldRenderViewIndependentWholeSceneShadows() && 
-							ReadOnlyCVARCache.bMobileAllowDistanceFieldShadows)
+						// Lightmap path
+						if (bMovableWithCSM)
 						{
-							SelectedLightmapPolicy = LMP_MOBILE_DISTANCE_FIELD_SHADOWS_LIGHTMAP_AND_CSM;
+							SelectedLightmapPolicy = LMP_MOBILE_MOVABLE_DIRECTIONAL_LIGHT_CSM_WITH_LIGHTMAP;
 						}
 						else
 						{
 							SelectedLightmapPolicy = LMP_MOBILE_DIRECTIONAL_LIGHT_CSM_AND_LIGHTMAP;
 						}
 					}
+				}
+				else
+				{
+					if (ShadowMapInteraction.GetType() == SMIT_Texture &&
+						ReadOnlyCVARCache.bMobileAllowDistanceFieldShadows)
+					{
+						SelectedLightmapPolicy = LMP_MOBILE_DISTANCE_FIELD_SHADOWS_AND_LQ_LIGHTMAP;
+					}
 					else
 					{
-						if (ShadowMapInteraction.GetType() == SMIT_Texture &&
-							ReadOnlyCVARCache.bMobileAllowDistanceFieldShadows)
+						// Lightmap path
+						if (bUseMovableLight)
 						{
-							SelectedLightmapPolicy = LMP_MOBILE_DISTANCE_FIELD_SHADOWS_AND_LQ_LIGHTMAP;
+							if (bUsesDeferredShading)
+							{
+								SelectedLightmapPolicy = LMP_LQ_LIGHTMAP;
+							}
+							else
+							{
+								SelectedLightmapPolicy = LMP_MOBILE_MOVABLE_DIRECTIONAL_LIGHT_WITH_LIGHTMAP;
+							}
 						}
 						else
 						{
