Remove permutations from RayTracingDebug by using dummy buffers for Picker and Instance Data.

diff --git a/Engine/Shaders/Private/RayTracing/RayTracingDebug.usf b/Engine/Shaders/Private/RayTracing/RayTracingDebug.usf
--- a/Engine/Shaders/Private/RayTracing/RayTracingDebug.usf
+++ b/Engine/Shaders/Private/RayTracing/RayTracingDebug.usf
@@ -203,17 +203,14 @@ float4 DebugRayTracingInstance(FRayDesc Ray, uint2 PixelCoord, out float OutDept
 		case RAY_TRACING_DEBUG_VIZ_TRIANGLES:
 			Result.rgb = TriangleColor;
 			break;
-#if ENABLE_INSTANCE_DEBUG_DATA
 		case RAY_TRACING_DEBUG_VIZ_DYNAMIC_INSTANCES:
 			Result.rgb = InstancesDebugData[Payload.GeometryInstanceIndex].Flags ? float3(0.25f, 0.50f, 0.25f) : float3(0.50f, 0.25f, 0.25f);
 			break;
 		case RAY_TRACING_DEBUG_VIZ_PROXY_TYPE:
 			Result.rgb = IntToColor(InstancesDebugData[Payload.GeometryInstanceIndex].ProxyHash);
 			break;
-#endif // ENABLE_INSTANCE_DEBUG_DATA
 		}
 	
-#if ENABLE_PICKING
 		if (VisualizationMode == RAY_TRACING_DEBUG_VIZ_PICKER)
 		{
 			const float ColorScale = 0.1f;
@@ -237,7 +234,6 @@ float4 DebugRayTracingInstance(FRayDesc Ray, uint2 PixelCoord, out float OutDept
 				break;
 			}
 		}
-#endif // ENABLE_PICKING
 
 		// apply some shading based on the primitive normal (map dot product result from [0,1] to [0.75,1] unless we are showing instance heatmap then just use [0,1]).
 		const float FakeLighting = dot(Payload.WorldNormal, normalize(float3(1.0f, 1.0f, 1.0f)));
diff --git a/Engine/Source/Runtime/Renderer/Private/RayTracing/RayTracingDebug.cpp b/Engine/Source/Runtime/Renderer/Private/RayTracing/RayTracingDebug.cpp
--- a/Engine/Source/Runtime/Renderer/Private/RayTracing/RayTracingDebug.cpp
+++ b/Engine/Source/Runtime/Renderer/Private/RayTracing/RayTracingDebug.cpp
@@ -17,6 +17,7 @@
 #include "RayTracing/RaytracingOptions.h"
 #include "RayTracing/RayTracingTraversalStatistics.h"
 #include "PixelShaderUtils.h"
+#include "SystemTextures.h"
 
 #define LOCTEXT_NAMESPACE "RayTracingDebugVisualizationMenuCommands"
 
@@ -100,10 +101,6 @@ class FRayTracingDebugRGS : public FGlobalShader
 	DECLARE_GLOBAL_SHADER(FRayTracingDebugRGS)
 	SHADER_USE_ROOT_PARAMETER_STRUCT(FRayTracingDebugRGS, FGlobalShader)
 
-	class FEnablePicking : SHADER_PERMUTATION_BOOL("ENABLE_PICKING");
-	class FEnableInstanceDebugData : SHADER_PERMUTATION_BOOL("ENABLE_INSTANCE_DEBUG_DATA");
-	using FPermutationDomain = TShaderPermutationDomain<FEnablePicking, FEnableInstanceDebugData>;
-
 	BEGIN_SHADER_PARAMETER_STRUCT(FParameters, )
 		SHADER_PARAMETER(uint32, VisualizationMode)
 		SHADER_PARAMETER(uint32, UseDebugCHS)
@@ -465,21 +462,15 @@ void FDeferredShadingSceneRenderer::PrepareRayTracingDebug(const FSceneViewFamil
 	bool bEnabled = ViewFamily.EngineShowFlags.RayTracingDebug && ShouldRenderRayTracingEffect(ERayTracingPipelineCompatibilityFlags::FullPipeline);
 	if (bEnabled)
 	{
-		for (int32 EnablePicking = 0; EnablePicking < 2; ++EnablePicking)
 		{
-			for (int32 EnableInstanceDebugData = 0; EnableInstanceDebugData < 2; ++EnableInstanceDebugData)
-			{
-				FRayTracingDebugRGS::FPermutationDomain PermutationVector;
-				PermutationVector.Set<FRayTracingDebugRGS::FEnablePicking>((bool)EnablePicking);
-				PermutationVector.Set<FRayTracingDebugRGS::FEnableInstanceDebugData>((bool)EnableInstanceDebugData);
+			auto RayGenShader = GetGlobalShaderMap(ViewFamily.GetShaderPlatform())->GetShader<FRayTracingDebugRGS>();
+			OutRayGenShaders.Add(RayGenShader.GetRayTracingShader());
+		}
 
-				auto RayGenShader = GetGlobalShaderMap(ViewFamily.GetShaderPlatform())->GetShader<FRayTracingDebugRGS>(PermutationVector);
-				OutRayGenShaders.Add(RayGenShader.GetRayTracingShader());
-			}			
+		{
+			auto PickingRayGenShader = GetGlobalShaderMap(ViewFamily.GetShaderPlatform())->GetShader<FRayTracingPickingRGS>();
+			OutRayGenShaders.Add(PickingRayGenShader.GetRayTracingShader());
 		}		
-
-		auto PickingRayGenShader = GetGlobalShaderMap(ViewFamily.GetShaderPlatform())->GetShader<FRayTracingPickingRGS>();
-		OutRayGenShaders.Add(PickingRayGenShader.GetRayTracingShader());
 	}
 }
 
@@ -843,13 +834,19 @@ void FDeferredShadingSceneRenderer::RenderRayTracingDebug(FRDGBuilder& GraphBuil
 	{
 		PickingBuffer = RayTracingPerformPicking(GraphBuilder, Scene, View, PickingFeedback);
 	}
+	else
+	{
+		PickingBuffer = GSystemTextures.GetDefaultStructuredBuffer(GraphBuilder, sizeof(FRayTracingPickingFeedback));
+	}
 
-	FRayTracingDebugRGS::FPermutationDomain PermutationVector;
-	PermutationVector.Set<FRayTracingDebugRGS::FEnablePicking>(PickingBuffer != nullptr);
-	PermutationVector.Set<FRayTracingDebugRGS::FEnableInstanceDebugData>(Scene->RayTracingScene.InstanceDebugBuffer != nullptr);
+	FRDGBufferRef InstanceDebugBuffer = Scene->RayTracingScene.InstanceDebugBuffer;
+	if (InstanceDebugBuffer == nullptr)
+	{
+		InstanceDebugBuffer = GSystemTextures.GetDefaultStructuredBuffer(GraphBuilder, sizeof(uint32));
+	}
 
 	FGlobalShaderMap* ShaderMap = GetGlobalShaderMap(FeatureLevel);
-	auto RayGenShader = ShaderMap->GetShader<FRayTracingDebugRGS>(PermutationVector);
+	auto RayGenShader = ShaderMap->GetShader<FRayTracingDebugRGS>();
 
 	FRayTracingPipelineState* Pipeline = View.RayTracingMaterialPipeline;
 	bool bRequiresBindings = false;
@@ -892,16 +889,6 @@ void FDeferredShadingSceneRenderer::RenderRayTracingDebug(FRDGBuilder& GraphBuil
 	FRDGTextureRef OutputDepthTexture = GraphBuilder.CreateTexture(OutputDepthTextureDesc, TEXT("RayTracingDebug::Depth"));
 	RayGenParameters->OutputDepth = GraphBuilder.CreateUAV(OutputDepthTexture);
 
-	if (Scene->RayTracingScene.InstanceDebugBuffer)
-	{
-		RayGenParameters->InstancesDebugData = GraphBuilder.CreateSRV(Scene->RayTracingScene.InstanceDebugBuffer);
-	}
-
-	if (PickingBuffer)
-	{		
-		RayGenParameters->PickingBuffer = GraphBuilder.CreateSRV(PickingBuffer);
-	}	
-
 	if (Lumen::UseFarField(ViewFamily))
 	{
 		RayGenParameters->MaxTraceDistance = Lumen::GetMaxTraceDistance(View);
@@ -915,6 +902,8 @@ void FDeferredShadingSceneRenderer::RenderRayTracingDebug(FRDGBuilder& GraphBuil
 		RayGenParameters->FarFieldReferencePos = FVector3f(0.0f);
 	}
 	
+	RayGenParameters->InstancesDebugData = GraphBuilder.CreateSRV(InstanceDebugBuffer);
+	RayGenParameters->PickingBuffer = GraphBuilder.CreateSRV(PickingBuffer);
 	RayGenParameters->TLAS = Scene->RayTracingScene.GetLayerSRVChecked(ERayTracingSceneLayer::Base);
 	RayGenParameters->ViewUniformBuffer = View.ViewUniformBuffer;
 	RayGenParameters->Output = GraphBuilder.CreateUAV(SceneColorTexture);
