Make Scene parameter of UploadDynamicPrimitiveShaderDataForView a reference to indicate this must be non-null

diff --git a/Engine/Source/Runtime/Renderer/Private/DeferredShadingRenderer.cpp b/Engine/Source/Runtime/Renderer/Private/DeferredShadingRenderer.cpp
--- a/Engine/Source/Runtime/Renderer/Private/DeferredShadingRenderer.cpp
+++ b/Engine/Source/Runtime/Renderer/Private/DeferredShadingRenderer.cpp
@@ -2500,7 +2500,7 @@ void FDeferredShadingSceneRenderer::Render(FRDGBuilder& GraphBuilder)
 			FViewInfo& View = Views[ViewIndex];
 			RDG_GPU_MASK_SCOPE(GraphBuilder, View.GPUMask);
 
-			Scene->GPUScene.UploadDynamicPrimitiveShaderDataForView(GraphBuilder, Scene, View, ExternalAccessQueue);
+			Scene->GPUScene.UploadDynamicPrimitiveShaderDataForView(GraphBuilder, *Scene, View, ExternalAccessQueue);
 
 			Scene->GPUScene.DebugRender(GraphBuilder, *Scene, View);
 		}
diff --git a/Engine/Source/Runtime/Renderer/Private/GPUScene.cpp b/Engine/Source/Runtime/Renderer/Private/GPUScene.cpp
--- a/Engine/Source/Runtime/Renderer/Private/GPUScene.cpp
+++ b/Engine/Source/Runtime/Renderer/Private/GPUScene.cpp
@@ -645,7 +645,7 @@ void FGPUScene::UpdateInternal(FRDGBuilder& GraphBuilder, FScene& Scene, FRDGExt
 	check(!BufferState.IsValid());
 
 	FUploadDataSourceAdapterScenePrimitives& Adapter = *GraphBuilder.AllocObject<FUploadDataSourceAdapterScenePrimitives>(Scene, SceneFrameNumber, MoveTemp(PrimitivesToUpdate), MoveTemp(PrimitiveDirtyState));
-	UpdateBufferState(GraphBuilder, &Scene, Adapter);
+	UpdateBufferState(GraphBuilder, Scene, Adapter);
 
 	// Run a pass that clears (Sets ID to invalid) any instances that need it
 	AddClearInstancesPass(GraphBuilder);
@@ -677,23 +677,20 @@ void FGPUScene::UpdateInternal(FRDGBuilder& GraphBuilder, FScene& Scene, FRDGExt
 		QUICK_SCOPE_CYCLE_COUNTER(STAT_UpdateGPUScene);
 		SCOPE_CYCLE_COUNTER(STAT_UpdateGPUSceneTime);
 
-		UploadGeneral<FUploadDataSourceAdapterScenePrimitives>(GraphBuilder, &Scene, ExternalAccessQueue, Adapter);
+		UploadGeneral<FUploadDataSourceAdapterScenePrimitives>(GraphBuilder, Scene, ExternalAccessQueue, Adapter);
 	}
 
 	UseExternalAccessMode(ExternalAccessQueue, ERHIAccess::SRVMask, ERHIPipeline::All);
 }
 
 template<typename FUploadDataSourceAdapter>
-void FGPUScene::UpdateBufferState(FRDGBuilder& GraphBuilder, FScene* Scene, const FUploadDataSourceAdapter& UploadDataSourceAdapter)
+void FGPUScene::UpdateBufferState(FRDGBuilder& GraphBuilder, FScene& Scene, const FUploadDataSourceAdapter& UploadDataSourceAdapter)
 {
 	LLM_SCOPE_BYTAG(GPUScene);
 
 	ensure(bInBeginEndBlock);
-	if (Scene != nullptr)
-	{
-		ensure(bIsEnabled == UseGPUScene(GMaxRHIShaderPlatform, Scene->GetFeatureLevel()));
-		ensure(NumScenePrimitives == Scene->Primitives.Num());
-	}
+	ensure(bIsEnabled == UseGPUScene(GMaxRHIShaderPlatform, Scene.GetFeatureLevel()));
+	ensure(NumScenePrimitives == Scene.Primitives.Num());
 
 	// Multi-GPU support : Updating on all GPUs is inefficient for AFR. Work is wasted
 	// for any primitives that update on consecutive frames.
@@ -717,21 +714,18 @@ void FGPUScene::UpdateBufferState(FRDGBuilder& GraphBuilder, FScene* Scene, cons
 	const uint32 InstancePayloadDataSizeReserve = FMath::RoundUpToPowerOfTwo(PayloadFloat4Count * sizeof(FVector4f));
 	BufferState.InstancePayloadDataBuffer = ResizeStructuredBufferIfNeeded(GraphBuilder, InstancePayloadDataBuffer, InstancePayloadDataSizeReserve, TEXT("GPUScene.InstancePayloadData"));
 
-	if (Scene != nullptr)
-	{
-		const uint32 NumNodes = FMath::RoundUpToPowerOfTwo(FMath::Max(Scene->InstanceBVH.GetNumNodes(), InitialBufferSize));
-		BufferState.InstanceBVHBuffer = ResizeStructuredBufferIfNeeded(GraphBuilder, InstanceBVHBuffer, NumNodes * sizeof(FBVHNode), TEXT("InstanceBVH"));
+	const uint32 NumNodes = FMath::RoundUpToPowerOfTwo(FMath::Max(Scene.InstanceBVH.GetNumNodes(), InitialBufferSize));
+	BufferState.InstanceBVHBuffer = ResizeStructuredBufferIfNeeded(GraphBuilder, InstanceBVHBuffer, NumNodes * sizeof(FBVHNode), TEXT("InstanceBVH"));
 
-		const bool bNaniteEnabled = DoesPlatformSupportNanite(GMaxRHIShaderPlatform);
-		if (UploadDataSourceAdapter.bUpdateNaniteMaterialTables && bNaniteEnabled)
+	const bool bNaniteEnabled = DoesPlatformSupportNanite(GMaxRHIShaderPlatform);
+	if (UploadDataSourceAdapter.bUpdateNaniteMaterialTables && bNaniteEnabled)
+	{
+		for (int32 NaniteMeshPassIndex = 0; NaniteMeshPassIndex < ENaniteMeshPass::Num; ++NaniteMeshPassIndex)
 		{
-			for (int32 NaniteMeshPassIndex = 0; NaniteMeshPassIndex < ENaniteMeshPass::Num; ++NaniteMeshPassIndex)
-			{
-				Scene->NaniteMaterials[NaniteMeshPassIndex].UpdateBufferState(GraphBuilder, Scene->Primitives.Num());
-			}
+			Scene.NaniteMaterials[NaniteMeshPassIndex].UpdateBufferState(GraphBuilder, Scene.Primitives.Num());
 		}
 	}
-
+	
 	const uint32 LightMapDataBufferSize = FMath::RoundUpToPowerOfTwo(FMath::Max(LightmapDataAllocator.GetMaxSize(), InitialBufferSize));
 	BufferState.LightmapDataBuffer = ResizeStructuredBufferIfNeeded(GraphBuilder, LightmapDataBuffer, LightMapDataBufferSize * sizeof(FLightmapSceneShaderData::Data), TEXT("GPUScene.LightmapData"));
 	BufferState.LightMapDataBufferSize = LightMapDataBufferSize;
@@ -864,15 +858,12 @@ void FGPUScene::UseExternalAccessMode(FRDGExternalAccessQueue& ExternalAccessQue
 }
 
 template<typename FUploadDataSourceAdapter>
-void FGPUScene::UploadGeneral(FRDGBuilder& GraphBuilder, FScene *Scene, FRDGExternalAccessQueue& ExternalAccessQueue, const FUploadDataSourceAdapter& UploadDataSourceAdapter)
+void FGPUScene::UploadGeneral(FRDGBuilder& GraphBuilder, FScene& Scene, FRDGExternalAccessQueue& ExternalAccessQueue, const FUploadDataSourceAdapter& UploadDataSourceAdapter)
 {
 	LLM_SCOPE_BYTAG(GPUScene);
 
-	if (Scene != nullptr)
-	{
-		ensure(bIsEnabled == UseGPUScene(GMaxRHIShaderPlatform, Scene->GetFeatureLevel()));
-		ensure(NumScenePrimitives == Scene->Primitives.Num());
-	}
+	ensure(bIsEnabled == UseGPUScene(GMaxRHIShaderPlatform, Scene.GetFeatureLevel()));
+	ensure(NumScenePrimitives == Scene.Primitives.Num());
 
 	const int32 NumPrimitiveDataUploads = UploadDataSourceAdapter.NumPrimitivesToUpload();
 
@@ -941,9 +932,9 @@ void FGPUScene::UploadGeneral(FRDGBuilder& GraphBuilder, FScene *Scene, FRDGExte
 		TaskContext.InstanceSceneUploader = InstanceSceneUploadBuffer.BeginPreSized(GraphBuilder, BufferState.InstanceSceneDataBuffer, TaskContext.NumInstanceSceneDataUploads * FInstanceSceneShaderData::GetDataStrideInFloat4s(), sizeof(FVector4f), TEXT("InstanceSceneUploadBuffer"));
 	}
 
-	if (Scene && Scene->InstanceBVH.GetNumDirty() > 0)
+	if (Scene.InstanceBVH.GetNumDirty() > 0)
 	{
-		TaskContext.InstanceBVHUploader = InstanceBVHUploadBuffer.Begin(GraphBuilder, BufferState.InstanceBVHBuffer, Scene->InstanceBVH.GetNumDirty(), sizeof(FBVHNode), TEXT("InstanceSceneUploadBuffer"));
+		TaskContext.InstanceBVHUploader = InstanceBVHUploadBuffer.Begin(GraphBuilder, BufferState.InstanceBVHBuffer, Scene.InstanceBVH.GetNumDirty(), sizeof(FBVHNode), TEXT("InstanceSceneUploadBuffer"));
 	}
 
 	if (TaskContext.NumLightmapDataUploads > 0)
@@ -951,17 +942,17 @@ void FGPUScene::UploadGeneral(FRDGBuilder& GraphBuilder, FScene *Scene, FRDGExte
 		TaskContext.LightmapUploader = LightmapUploadBuffer.Begin(GraphBuilder, BufferState.LightmapDataBuffer, TaskContext.NumLightmapDataUploads, sizeof(FLightmapSceneShaderData::Data), TEXT("LightmapUploadBuffer"));
 	}
 
-	if (Scene != nullptr && UploadDataSourceAdapter.bUpdateNaniteMaterialTables && bNaniteEnabled)
+	if (UploadDataSourceAdapter.bUpdateNaniteMaterialTables && bNaniteEnabled)
 	{
 		for (int32 NaniteMeshPassIndex = 0; NaniteMeshPassIndex < ENaniteMeshPass::Num; ++NaniteMeshPassIndex)
 		{
-			TaskContext.NaniteMaterialUploaders[NaniteMeshPassIndex] = Scene->NaniteMaterials[NaniteMeshPassIndex].Begin(GraphBuilder, Scene->Primitives.Num(), NumPrimitiveDataUploads);
+			TaskContext.NaniteMaterialUploaders[NaniteMeshPassIndex] = Scene.NaniteMaterials[NaniteMeshPassIndex].Begin(GraphBuilder, Scene.Primitives.Num(), NumPrimitiveDataUploads);
 		}
 
 		TaskContext.bUseNaniteMaterialUploaders = true;
 	}
 
-	GraphBuilder.AddCommandListSetupTask([&TaskContext, &UploadDataSourceAdapter, Scene, bNaniteEnabled, bExecuteInParallel, FeatureLevel = FeatureLevel](FRHICommandListBase& RHICmdList)
+	GraphBuilder.AddCommandListSetupTask([&TaskContext, &UploadDataSourceAdapter, &Scene, bNaniteEnabled, bExecuteInParallel, FeatureLevel = FeatureLevel](FRHICommandListBase& RHICmdList)
 	{
 		SCOPED_NAMED_EVENT(UpdateGPUScene, FColor::Green);
 
@@ -1199,7 +1190,7 @@ void FGPUScene::UploadGeneral(FRDGBuilder& GraphBuilder, FScene *Scene, FRDGExte
 		{
 			SCOPED_NAMED_EVENT(InstanceBVH, FColor::Green);
 
-			Scene->InstanceBVH.ForAllDirty(
+			Scene.InstanceBVH.ForAllDirty(
 				[&](uint32 NodeIndex, const auto& Node)
 			{
 				FBVHNode GPUNode;
@@ -1274,7 +1265,7 @@ void FGPUScene::UploadGeneral(FRDGBuilder& GraphBuilder, FScene *Scene, FRDGExte
 	{
 		for (int32 NaniteMeshPassIndex = 0; NaniteMeshPassIndex < ENaniteMeshPass::Num; ++NaniteMeshPassIndex)
 		{
-			Scene->NaniteMaterials[NaniteMeshPassIndex].Finish(GraphBuilder, ExternalAccessQueue, TaskContext.NaniteMaterialUploaders[NaniteMeshPassIndex]);
+			Scene.NaniteMaterials[NaniteMeshPassIndex].Finish(GraphBuilder, ExternalAccessQueue, TaskContext.NaniteMaterialUploaders[NaniteMeshPassIndex]);
 		}
 	}
 
@@ -1441,14 +1432,14 @@ struct FUploadDataSourceAdapterDynamicPrimitives
 	TArray<uint32, SceneRenderingAllocator> PrimitivesIds;
 };
 
-void FGPUScene::UploadDynamicPrimitiveShaderDataForViewInternal(FRDGBuilder& GraphBuilder, FScene *Scene, FViewInfo& View, FRDGExternalAccessQueue& ExternalAccessQueue, bool bIsShadowView)
+void FGPUScene::UploadDynamicPrimitiveShaderDataForViewInternal(FRDGBuilder& GraphBuilder, FScene& Scene, FViewInfo& View, FRDGExternalAccessQueue& ExternalAccessQueue, bool bIsShadowView)
 {
 	LLM_SCOPE_BYTAG(GPUScene);
 
 	RDG_EVENT_SCOPE(GraphBuilder, "GPUScene.UploadDynamicPrimitiveShaderDataForView");
 
 	ensure(bInBeginEndBlock);
-	ensure(Scene == nullptr || DynamicPrimitivesOffset >= Scene->Primitives.Num());
+	ensure(DynamicPrimitivesOffset >= Scene.Primitives.Num());
 
 	FGPUScenePrimitiveCollector& Collector = View.DynamicPrimitiveCollector;
 
@@ -1481,7 +1472,7 @@ void FGPUScene::UploadDynamicPrimitiveShaderDataForViewInternal(FRDGBuilder& Gra
 		ensure(UploadIdStart < DynamicPrimitivesOffset);
 		ensure(InstanceIdStart != INDEX_NONE);
 
-		if (bIsShadowView && Scene != nullptr && Scene->GetVirtualShadowMapCache(View) != nullptr)
+		if (bIsShadowView && Scene.GetVirtualShadowMapCache(View) != nullptr)
 		{
 			// Enqueue cache invalidations for all dynamic primitives' instances, as they will be removed this frame and are not associated
 			// with any particular FPrimitiveSceneInfo. Will occur on the next call to UpdateAllPrimitiveSceneInfos
@@ -1617,7 +1608,7 @@ void FGPUScene::Update(FRDGBuilder& GraphBuilder, FScene& Scene, FRDGExternalAcc
 	}
 }
 
-void FGPUScene::UploadDynamicPrimitiveShaderDataForView(FRDGBuilder& GraphBuilder, FScene *Scene, FViewInfo& View, FRDGExternalAccessQueue& ExternalAccessQueue, bool bIsShadowView)
+void FGPUScene::UploadDynamicPrimitiveShaderDataForView(FRDGBuilder& GraphBuilder, FScene& Scene, FViewInfo& View, FRDGExternalAccessQueue& ExternalAccessQueue, bool bIsShadowView)
 {
 	if (bIsEnabled)
 	{
diff --git a/Engine/Source/Runtime/Renderer/Private/GPUScene.h b/Engine/Source/Runtime/Renderer/Private/GPUScene.h
--- a/Engine/Source/Runtime/Renderer/Private/GPUScene.h
+++ b/Engine/Source/Runtime/Renderer/Private/GPUScene.h
@@ -218,7 +218,7 @@ public:
 	/**
 	 * Upload primitives from View.DynamicPrimitiveCollector.
 	 */
-	void UploadDynamicPrimitiveShaderDataForView(FRDGBuilder& GraphBuilder, FScene* Scene, FViewInfo& View, FRDGExternalAccessQueue& ExternalAccessQueue, bool bIsShadowView = false);
+	void UploadDynamicPrimitiveShaderDataForView(FRDGBuilder& GraphBuilder, FScene& Scene, FViewInfo& View, FRDGExternalAccessQueue& ExternalAccessQueue, bool bIsShadowView = false);
 
 	/**
 	 * Pull all pending updates from Scene and upload primitive & instance data.
@@ -394,16 +394,16 @@ private:
 	ERHIFeatureLevel::Type FeatureLevel;
 
 	template<typename FUploadDataSourceAdapter>
-	void UpdateBufferState(FRDGBuilder& GraphBuilder, FScene* Scene, const FUploadDataSourceAdapter& UploadDataSourceAdapter);
+	void UpdateBufferState(FRDGBuilder& GraphBuilder, FScene& Scene, const FUploadDataSourceAdapter& UploadDataSourceAdapter);
 
 	/**
 	 * Generalized upload that uses an adapter to abstract the data souce. Enables uploading scene primitives & dynamic primitives using a single path.
 	 * @parameter Scene may be null, as it is only needed for the Nanite material table update (which is coupled to the Scene at the moment).
 	 */
 	template<typename FUploadDataSourceAdapter>
-	void UploadGeneral(FRDGBuilder& GraphBuilder, FScene* Scene, FRDGExternalAccessQueue& ExternalAccessQueue, const FUploadDataSourceAdapter& UploadDataSourceAdapter);
+	void UploadGeneral(FRDGBuilder& GraphBuilder, FScene& Scene, FRDGExternalAccessQueue& ExternalAccessQueue, const FUploadDataSourceAdapter& UploadDataSourceAdapter);
 
-	void UploadDynamicPrimitiveShaderDataForViewInternal(FRDGBuilder& GraphBuilder, FScene* Scene, FViewInfo& View, FRDGExternalAccessQueue& ExternalAccessQueue, bool bIsShadowView);
+	void UploadDynamicPrimitiveShaderDataForViewInternal(FRDGBuilder& GraphBuilder, FScene& Scene, FViewInfo& View, FRDGExternalAccessQueue& ExternalAccessQueue, bool bIsShadowView);
 
 	void UpdateInternal(FRDGBuilder& GraphBuilder, FScene& Scene, FRDGExternalAccessQueue& ExternalAccessQueue);
 
diff --git a/Engine/Source/Runtime/Renderer/Private/MobileShadingRenderer.cpp b/Engine/Source/Runtime/Renderer/Private/MobileShadingRenderer.cpp
--- a/Engine/Source/Runtime/Renderer/Private/MobileShadingRenderer.cpp
+++ b/Engine/Source/Runtime/Renderer/Private/MobileShadingRenderer.cpp
@@ -639,7 +639,7 @@ void FMobileSceneRenderer::InitViews(FRDGBuilder& GraphBuilder, FSceneTexturesCo
 		Scene->GPUScene.Update(GraphBuilder, *Scene, ExternalAccessQueue);
 		for (int32 ViewIndex = 0; ViewIndex < Views.Num(); ViewIndex++)
 		{
-			Scene->GPUScene.UploadDynamicPrimitiveShaderDataForView(GraphBuilder, Scene, Views[ViewIndex], ExternalAccessQueue);
+			Scene->GPUScene.UploadDynamicPrimitiveShaderDataForView(GraphBuilder, *Scene, Views[ViewIndex], ExternalAccessQueue);
 		}
 	}
 
diff --git a/Engine/Source/Runtime/Renderer/Private/ReflectionEnvironmentRealTimeCapture.cpp b/Engine/Source/Runtime/Renderer/Private/ReflectionEnvironmentRealTimeCapture.cpp
--- a/Engine/Source/Runtime/Renderer/Private/ReflectionEnvironmentRealTimeCapture.cpp
+++ b/Engine/Source/Runtime/Renderer/Private/ReflectionEnvironmentRealTimeCapture.cpp
@@ -568,7 +568,7 @@ void FScene::AllocateAndCaptureFrameSkyEnvMap(
 				CubeView.CreateViewUniformBuffers(*CubeView.CachedViewUniformShaderParameters);
 				if (CubeView.bSceneHasSkyMaterial)
 				{
-					GPUScene.UploadDynamicPrimitiveShaderDataForView(GraphBuilder, this, CubeView, ExternalAccessQueue);
+					GPUScene.UploadDynamicPrimitiveShaderDataForView(GraphBuilder, *this, CubeView, ExternalAccessQueue);
 				}
 
 				SkyRC.ViewUniformBuffer = CubeView.ViewUniformBuffer;
diff --git a/Engine/Source/Runtime/Renderer/Private/SceneHitProxyRendering.cpp b/Engine/Source/Runtime/Renderer/Private/SceneHitProxyRendering.cpp
--- a/Engine/Source/Runtime/Renderer/Private/SceneHitProxyRendering.cpp
+++ b/Engine/Source/Runtime/Renderer/Private/SceneHitProxyRendering.cpp
@@ -659,7 +659,7 @@ void FDeferredShadingSceneRenderer::RenderHitProxies(FRDGBuilder& GraphBuilder)
 
 		for (int32 ViewIndex = 0; ViewIndex < Views.Num(); ViewIndex++)
 		{
-			Scene->GPUScene.UploadDynamicPrimitiveShaderDataForView(GraphBuilder, Scene, Views[ViewIndex], ExternalAccessQueue);
+			Scene->GPUScene.UploadDynamicPrimitiveShaderDataForView(GraphBuilder, *Scene, Views[ViewIndex], ExternalAccessQueue);
 		}
 
 		ExternalAccessQueue.Submit(GraphBuilder);
diff --git a/Engine/Source/Runtime/Renderer/Private/ShadowDepthRendering.cpp b/Engine/Source/Runtime/Renderer/Private/ShadowDepthRendering.cpp
--- a/Engine/Source/Runtime/Renderer/Private/ShadowDepthRendering.cpp
+++ b/Engine/Source/Runtime/Renderer/Private/ShadowDepthRendering.cpp
@@ -1650,32 +1650,32 @@ void FSceneRenderer::RenderShadowDepthMaps(FRDGBuilder& GraphBuilder, FInstanceC
 	{
 		for (FProjectedShadowInfo* ProjectedShadowInfo : ShadowMapAtlas.Shadows)
 		{
-			Scene->GPUScene.UploadDynamicPrimitiveShaderDataForView(GraphBuilder, Scene, *ProjectedShadowInfo->ShadowDepthView, ExternalAccessQueue, true);
+			Scene->GPUScene.UploadDynamicPrimitiveShaderDataForView(GraphBuilder, *Scene, *ProjectedShadowInfo->ShadowDepthView, ExternalAccessQueue, true);
 		}
 	}
 	for (FSortedShadowMapAtlas& ShadowMap : SortedShadowsForShadowDepthPass.ShadowMapCubemaps)
 	{
 		check(ShadowMap.Shadows.Num() == 1);
 		FProjectedShadowInfo* ProjectedShadowInfo = ShadowMap.Shadows[0];
-		Scene->GPUScene.UploadDynamicPrimitiveShaderDataForView(GraphBuilder, Scene, *ProjectedShadowInfo->ShadowDepthView, ExternalAccessQueue, true);
+		Scene->GPUScene.UploadDynamicPrimitiveShaderDataForView(GraphBuilder, *Scene, *ProjectedShadowInfo->ShadowDepthView, ExternalAccessQueue, true);
 	}
 	for (FProjectedShadowInfo* ProjectedShadowInfo : SortedShadowsForShadowDepthPass.PreshadowCache.Shadows)
 	{
 		if (!ProjectedShadowInfo->bDepthsCached)
 		{
-			Scene->GPUScene.UploadDynamicPrimitiveShaderDataForView(GraphBuilder, Scene, *ProjectedShadowInfo->ShadowDepthView, ExternalAccessQueue, true);
+			Scene->GPUScene.UploadDynamicPrimitiveShaderDataForView(GraphBuilder, *Scene, *ProjectedShadowInfo->ShadowDepthView, ExternalAccessQueue, true);
 		}
 	}
 	for (const FSortedShadowMapAtlas& ShadowMapAtlas : SortedShadowsForShadowDepthPass.TranslucencyShadowMapAtlases)
 	{
 		for (FProjectedShadowInfo* ProjectedShadowInfo : ShadowMapAtlas.Shadows)
 		{
-			Scene->GPUScene.UploadDynamicPrimitiveShaderDataForView(GraphBuilder, Scene, *ProjectedShadowInfo->ShadowDepthView, ExternalAccessQueue, true);
+			Scene->GPUScene.UploadDynamicPrimitiveShaderDataForView(GraphBuilder, *Scene, *ProjectedShadowInfo->ShadowDepthView, ExternalAccessQueue, true);
 		}
 	}
 	for (FProjectedShadowInfo* ProjectedShadowInfo : SortedShadowsForShadowDepthPass.VirtualShadowMapShadows)
 	{
-		Scene->GPUScene.UploadDynamicPrimitiveShaderDataForView(GraphBuilder, Scene, *ProjectedShadowInfo->ShadowDepthView, ExternalAccessQueue, true);
+		Scene->GPUScene.UploadDynamicPrimitiveShaderDataForView(GraphBuilder, *Scene, *ProjectedShadowInfo->ShadowDepthView, ExternalAccessQueue, true);
 	}
 
 	ExternalAccessQueue.Submit(GraphBuilder);
diff --git a/Engine/Source/Runtime/Renderer/Private/VT/RuntimeVirtualTextureRender.cpp b/Engine/Source/Runtime/Renderer/Private/VT/RuntimeVirtualTextureRender.cpp
--- a/Engine/Source/Runtime/Renderer/Private/VT/RuntimeVirtualTextureRender.cpp
+++ b/Engine/Source/Runtime/Renderer/Private/VT/RuntimeVirtualTextureRender.cpp
@@ -1256,7 +1256,7 @@ namespace RuntimeVirtualTexture
 		View->ViewUniformBuffer = TUniformBufferRef<FViewUniformShaderParameters>::CreateUniformBufferImmediate(*View->CachedViewUniformShaderParameters, UniformBuffer_SingleFrame);
 
 		FRDGExternalAccessQueue ExternalAccessQueue;
-		Scene->GPUScene.UploadDynamicPrimitiveShaderDataForView(GraphBuilder, (const_cast<FScene*>(Scene)), *View, ExternalAccessQueue);
+		Scene->GPUScene.UploadDynamicPrimitiveShaderDataForView(GraphBuilder, *Scene, *View, ExternalAccessQueue);
 		ExternalAccessQueue.Submit(GraphBuilder);
 
 		// Build graph
