Added support for filtered virtual shadow maps on SingleLayerWater. Implemented by applying the VSM shadow mask to the SeparateMainDirLight texture that is written out by the main water pass in the presence of distance field shadows (or VSM filtering) and then later composing it back into the scene. Also removed EncodeLightAttenuationFromMask() function from VirtualShadowMapMaskBitsCommon.ush.

#rnx

diff --git a/Engine/Shaders/Private/BasePassPixelShader.usf b/Engine/Shaders/Private/BasePassPixelShader.usf
--- a/Engine/Shaders/Private/BasePassPixelShader.usf
+++ b/Engine/Shaders/Private/BasePassPixelShader.usf
@@ -174,10 +174,6 @@
 #include "/Engine/Private/Strata/StrataExport.ush"
 #endif
 
-// Separate main directional light from scene color buffer if water needs to receive distance field shadows.
-// Sperated dir light should match with FShaderCompileUtilities::FetchGBufferParamsRuntime.
-#define SINGLE_LAYER_WATER_SEPARATED_MAIN_LIGHT  (SINGLE_LAYER_WATER_DF_SHADOW_ENABLED)
-
 #if !FORWARD_SHADING
 
 bool ManualDepthTestEqual(float4 SvPosition, float PSDeviceZWithOffset)
@@ -1181,7 +1177,7 @@ void FPixelShaderInOut_MainPS(
 	float3 DiffuseIndirectLighting = 0;
 	float3 SubsurfaceIndirectLighting = 0;
 
-	bool bSeparateWaterMainDirLightLuminance = SINGLE_LAYER_WATER_SEPARATED_MAIN_LIGHT > 0 && View.SeparateWaterMainDirLightLuminance > 0.0f; 
+	bool bSeparateWaterMainDirLightLuminance = SINGLE_LAYER_WATER_SEPARATED_MAIN_LIGHT > 0 && SingleLayerWater.bSeparateMainDirLightLuminance; 
 	float3 SeparatedWaterMainDirLightLuminance = float3(0, 0, 0); 
 
 #if !STRATA_ENABLED || STRATA_INLINE_SINGLELAYERWATER
@@ -1282,10 +1278,13 @@ void FPixelShaderInOut_MainPS(
 					}
 					#endif // NEEDS_BASEPASS_PIXEL_FOGGING && PROJECT_SUPPORT_SKY_ATMOSPHERE
 
+					const bool bSkipMainDirLightVirtualShadowMapEvaluation = SingleLayerWater.bMainDirectionalLightVSMFiltering;
+
 					FDeferredLightingSplit ForwardDirectLighting = GetForwardDirectLightingSplit(
 						GridIndex, MaterialParameters.WorldPosition_CamRelative, MaterialParameters.CameraVector, GBuffer, NearestResolvedDepthScreenUV, MaterialParameters.PrimitiveId, EyeIndex, Dither, 
 						DirectionalLightCloudShadow, DirectionalLightAtmosphereTransmittance, DirectionalLightShadow,
-						bSeparateWaterMainDirLightLuminance, SeparatedWaterMainDirLightLuminance);
+						bSeparateWaterMainDirLightLuminance, SeparatedWaterMainDirLightLuminance,
+						bSkipMainDirLightVirtualShadowMapEvaluation);
 					
 					#if MATERIAL_SHADINGMODEL_THIN_TRANSLUCENT
 						DiffuseColor += ForwardDirectLighting.DiffuseLighting.rgb;
diff --git a/Engine/Shaders/Private/Definitions.usf b/Engine/Shaders/Private/Definitions.usf
--- a/Engine/Shaders/Private/Definitions.usf
+++ b/Engine/Shaders/Private/Definitions.usf
@@ -104,8 +104,8 @@
 #define MATERIAL_SHADINGMODEL_SINGLELAYERWATER			0 
 #endif
 
-#ifndef SINGLE_LAYER_WATER_DF_SHADOW_ENABLED
-#define SINGLE_LAYER_WATER_DF_SHADOW_ENABLED			0 
+#ifndef SINGLE_LAYER_WATER_SEPARATED_MAIN_LIGHT
+#define SINGLE_LAYER_WATER_SEPARATED_MAIN_LIGHT			0 
 #endif
 
 #ifndef SINGLE_LAYER_WATER_NO_DISCARD
diff --git a/Engine/Shaders/Private/ForwardLightingCommon.ush b/Engine/Shaders/Private/ForwardLightingCommon.ush
--- a/Engine/Shaders/Private/ForwardLightingCommon.ush
+++ b/Engine/Shaders/Private/ForwardLightingCommon.ush
@@ -122,7 +122,7 @@ float GetIndirectOcclusion(float2 ScreenUV, bool bHasDynamicIndirectShadowCaster
 FDeferredLightingSplit GetForwardDirectLightingSplit(
 	uint GridIndex, float3 TranslatedWorldPosition, float3 CameraVector, FGBufferData GBufferData, float2 ScreenUV, uint PrimitiveId, uint EyeIndex, float Dither, 
 	float InDirectionalLightCloudShadow, float3 InDirectionalLightAtmosphereTransmittance, inout float OutDirectionalLightShadow,
-	bool bSeparateMainDirLightLuminance, inout float3 SeparatedMainDirLightLuminance)
+	bool bSeparateMainDirLightLuminance, inout float3 SeparatedMainDirLightLuminance, bool bSkipDirLightVirtualShadowMapEvaluation)
 {
 	float3 WorldPosition = TranslatedWorldPosition - LWCHackToFloat(PrimaryView.PreViewTranslation);
 
@@ -185,7 +185,7 @@ FDeferredLightingSplit GetForwardDirectLightingSplit(
 			
 			#if VIRTUAL_SHADOW_MAP
 			BRANCH
-			if ( ForwardLightData.DirectionalLightVSM != INDEX_NONE )
+			if ( !bSkipDirLightVirtualShadowMapEvaluation && ForwardLightData.DirectionalLightVSM != INDEX_NONE )
 			{
 				FVirtualShadowMapSampleResult VirtualShadowMapSample = SampleVirtualShadowMapTranslatedWorld( ForwardLightData.DirectionalLightVSM, TranslatedWorldPosition );
 				DynamicShadowFactor *= VirtualShadowMapSample.ShadowFactor;
diff --git a/Engine/Shaders/Private/VirtualShadowMaps/VirtualShadowMapMaskBitsCommon.ush b/Engine/Shaders/Private/VirtualShadowMaps/VirtualShadowMapMaskBitsCommon.ush
--- a/Engine/Shaders/Private/VirtualShadowMaps/VirtualShadowMapMaskBitsCommon.ush
+++ b/Engine/Shaders/Private/VirtualShadowMaps/VirtualShadowMapMaskBitsCommon.ush
@@ -11,21 +11,6 @@ VirtualShadowMapMaskBitsCommon.ush:
 #include "../LightShaderParameters.ush"
 #include "../LightGridCommon.ush"
 
-float4 EncodeLightAttenuationFromMask(float ShadowMask, float SSSTransmission)
-{
-	const float ShadowFadeFraction = 1;
-
-	// 0 is shadowed, 1 is unshadowed
-	// RETURN_COLOR not needed unless writing to SceneColor;
-	//float FadedShadow = lerp(1.0f, Square(ShadowMask), ShadowFadeFraction);
-	//float FadedSSSShadow = lerp(1.0f, Square(SSSTransmission), ShadowFadeFraction);
-	float FadedShadow = lerp(1.0f, ShadowMask, ShadowFadeFraction);
-	float FadedSSSShadow = lerp(1.0f, SSSTransmission, ShadowFadeFraction);
-
-	// the channel assignment is documented in ShadowRendering.cpp (look for Light Attenuation channel assignment)
-	return EncodeLightAttenuation(float4(FadedShadow, FadedSSSShadow, FadedShadow, FadedSSSShadow));
-}
-
 float GetVirtualShadowMapMaskForLight(
 	Texture2D<uint4> ShadowMaskBits,
 	uint2 PixelPos,
diff --git a/Engine/Shaders/Private/VirtualShadowMaps/VirtualShadowMapProjection.usf b/Engine/Shaders/Private/VirtualShadowMaps/VirtualShadowMapProjection.usf
--- a/Engine/Shaders/Private/VirtualShadowMaps/VirtualShadowMapProjection.usf
+++ b/Engine/Shaders/Private/VirtualShadowMaps/VirtualShadowMapProjection.usf
@@ -534,15 +534,23 @@ void OutputVisualize( uint2 PixelPos, FVirtualShadowMapSampleResult Result )
 	}
 }
 
-[numthreads(8, 8, 1)]
+#if USE_TILE_LIST
+Buffer<uint> TileListData;
+#endif
+
+[numthreads(WORK_TILE_SIZE, WORK_TILE_SIZE, 1)]
 void VirtualShadowMapProjection(
 	uint3	GroupId				: SV_GroupID,
 	uint	GroupIndex			: SV_GroupIndex,
 	uint3	DispatchThreadId	: SV_DispatchThreadID )
 {
+#if USE_TILE_LIST
+	uint PackedTile = TileListData.Load(GroupId.x);
+	uint2 LocalPixelPos = WORK_TILE_SIZE * uint2(PackedTile & 0xFFFF, PackedTile >> 16u) + MortonDecode(GroupIndex);
+#else
 	// Morton order within a group so page access/atomics are more coherent and wave-swizzled gradients are possible.
-	uint2 PixelCoord = DispatchThreadId.xy;
-	uint2 LocalPixelPos = 8 * GroupId.xy + MortonDecode( GroupIndex );
+	uint2 LocalPixelPos = WORK_TILE_SIZE * GroupId.xy + MortonDecode( GroupIndex );
+#endif // USE_TILE_LIST
 	uint2 PixelPos = LocalPixelPos + uint2( ProjectionRect.xy );
 	if ( any( PixelPos >= uint2( ProjectionRect.zw ) ) )
 	{
diff --git a/Engine/Shaders/Private/VirtualShadowMaps/VirtualShadowMapProjectionComposite.usf b/Engine/Shaders/Private/VirtualShadowMaps/VirtualShadowMapProjectionComposite.usf
--- a/Engine/Shaders/Private/VirtualShadowMaps/VirtualShadowMapProjectionComposite.usf
+++ b/Engine/Shaders/Private/VirtualShadowMaps/VirtualShadowMapProjectionComposite.usf
@@ -5,17 +5,40 @@
 #include "../LightShaderParameters.ush"
 #include "VirtualShadowMapMaskBitsCommon.ush"
 
+#ifndef WORK_TILE_SIZE
+#define WORK_TILE_SIZE 8
+#endif
+
+Buffer<uint> TileListData;
+
+void VirtualShadowMapCompositeTileVS(
+	in uint InstanceId : SV_InstanceID,
+	in uint VertexId : SV_VertexID,
+	out float4 Position : SV_POSITION)
+{
+	const uint TileData = TileListData[InstanceId.x];
+	const uint2 TileOrigin = WORK_TILE_SIZE * uint2((TileData & 0xFFFF), (TileData >> 16u));
+
+	uint2 TileVertex = TileOrigin;
+	TileVertex.x += VertexId == 1 || VertexId == 2 || VertexId == 4 ? WORK_TILE_SIZE : 0;
+	TileVertex.y += VertexId == 2 || VertexId == 4 || VertexId == 5 ? WORK_TILE_SIZE : 0;
+
+	// Viewport is set on the view rect so no offset is needed.
+	Position = float4(float2(TileVertex) * View.ViewSizeAndInvSize.zw * float2(2.0f, -2.0f) + float2(-1.0, 1.0f), 0.5f, 1.0f);
+}
+
 Texture2D<float2> InputShadowFactor;
 Texture2D<float> InputDepthTexture;
 Texture2D<uint4> ShadowMaskBits;
 int CompositeVirtualShadowMapId;
+uint bModulateRGB;
 
 void VirtualShadowMapCompositePS(
 	in float4 SvPosition : SV_Position,
 	out float4 OutShadowMask : SV_Target)
 {
 	float2 ShadowFactor = InputShadowFactor.Load(int3(SvPosition.xy, 0));
-	OutShadowMask = EncodeLightAttenuationFromMask(ShadowFactor.x, ShadowFactor.y);
+	OutShadowMask = (bModulateRGB != 0) ? ShadowFactor.xxxx : EncodeLightAttenuation(float4(ShadowFactor.x, ShadowFactor.y, ShadowFactor.x, ShadowFactor.y));
 }
 
 void VirtualShadowMapCompositeFromMaskBitsPS(
@@ -33,5 +56,5 @@ void VirtualShadowMapCompositeFromMaskBitsPS(
 	
 	// TODO: Subsurface...?
 		
-	OutShadowMask = EncodeLightAttenuationFromMask(ShadowMask, ShadowMask);
+	OutShadowMask = EncodeLightAttenuation(ShadowMask.xxxx);
 }
diff --git a/Engine/Source/Runtime/Engine/Private/Materials/HLSLMaterialTranslator.cpp b/Engine/Source/Runtime/Engine/Private/Materials/HLSLMaterialTranslator.cpp
--- a/Engine/Source/Runtime/Engine/Private/Materials/HLSLMaterialTranslator.cpp
+++ b/Engine/Source/Runtime/Engine/Private/Materials/HLSLMaterialTranslator.cpp
@@ -1858,9 +1858,10 @@ void FHLSLMaterialTranslator::GetMaterialEnvironment(EShaderPlatform InPlatform,
 		}
 
 		const bool bIsWaterDistanceFieldShadowEnabled = IsWaterDistanceFieldShadowEnabled(InPlatform);
-		if (ShadingModels.HasShadingModel(MSM_SingleLayerWater) && bIsWaterDistanceFieldShadowEnabled)
+		const bool bIsWaterVSMFilteringEnabled = IsWaterVirtualShadowMapFilteringEnabled(InPlatform);
+		if (ShadingModels.HasShadingModel(MSM_SingleLayerWater) && (bIsWaterDistanceFieldShadowEnabled || bIsWaterVSMFilteringEnabled))
 		{
-			OutEnvironment.SetDefine(TEXT("SINGLE_LAYER_WATER_DF_SHADOW_ENABLED"), TEXT("1"));
+			OutEnvironment.SetDefine(TEXT("SINGLE_LAYER_WATER_SEPARATED_MAIN_LIGHT"), TEXT("1"));
 		}
 
 		if (NumSetMaterials == 1)
diff --git a/Engine/Source/Runtime/Engine/Private/ShaderCompiler/ShaderGenerationUtil.cpp b/Engine/Source/Runtime/Engine/Private/ShaderCompiler/ShaderGenerationUtil.cpp
--- a/Engine/Source/Runtime/Engine/Private/ShaderCompiler/ShaderGenerationUtil.cpp
+++ b/Engine/Source/Runtime/Engine/Private/ShaderCompiler/ShaderGenerationUtil.cpp
@@ -156,7 +156,7 @@ void FShaderCompileUtilities::ApplyFetchEnvironment(FShaderMaterialPropertyDefin
 	FETCH_COMPILE_BOOL(MATERIAL_SHADINGMODEL_SINGLELAYERWATER);
 	FETCH_COMPILE_BOOL(MATERIAL_SHADINGMODEL_THIN_TRANSLUCENT);
 
-	FETCH_COMPILE_BOOL(SINGLE_LAYER_WATER_DF_SHADOW_ENABLED);
+	FETCH_COMPILE_BOOL(SINGLE_LAYER_WATER_SEPARATED_MAIN_LIGHT);
 
 	FETCH_COMPILE_BOOL(MATERIAL_FULLY_ROUGH);
 
@@ -1795,7 +1795,7 @@ static void DetermineUsedMaterialSlots(
 	{
 		// single layer water uses standard slots
 		SetStandardGBufferSlots(Slots, bWriteEmissive, bHasTangent, bHasVelocity, bHasStaticLighting, bIsStrataMaterial);
-		if (Mat.SINGLE_LAYER_WATER_DF_SHADOW_ENABLED)
+		if (Mat.SINGLE_LAYER_WATER_SEPARATED_MAIN_LIGHT)
 		{
 			Slots[GBS_SeparatedMainDirLight] = true;
 		}
@@ -2157,8 +2157,8 @@ FGBufferParams FShaderCompileUtilities::FetchGBufferParamsRuntime(EShaderPlatfor
 	static const auto CVarFormat = IConsoleManager::Get().FindTConsoleVariableDataInt(TEXT("r.GBufferFormat"));
 	Ret.LegacyFormatIndex = CVarFormat->GetValueOnAnyThread();
 
-	// This should match with SINGLE_LAYER_WATER_SEPARATED_DIR_LIGHT
-	Ret.bHasSingleLayerWaterSeparatedMainLight = IsWaterDistanceFieldShadowEnabled(Platform);
+	// This should match with SINGLE_LAYER_WATER_SEPARATED_MAIN_LIGHT
+	Ret.bHasSingleLayerWaterSeparatedMainLight = IsWaterDistanceFieldShadowEnabled(Platform) || IsWaterVirtualShadowMapFilteringEnabled(Platform);
 
 	return Ret;
 }
diff --git a/Engine/Source/Runtime/Engine/Public/SceneView.h b/Engine/Source/Runtime/Engine/Public/SceneView.h
--- a/Engine/Source/Runtime/Engine/Public/SceneView.h
+++ b/Engine/Source/Runtime/Engine/Public/SceneView.h
@@ -653,7 +653,6 @@ enum ETranslucencyVolumeCascade
 	VIEW_UNIFORM_BUFFER_MEMBER(FVector2f, BufferToSceneTextureScale) \
 	VIEW_UNIFORM_BUFFER_MEMBER(FVector2f, ResolutionFractionAndInv) \
 	VIEW_UNIFORM_BUFFER_MEMBER(int32, NumSceneColorMSAASamples) \
-	VIEW_UNIFORM_BUFFER_MEMBER(float, SeparateWaterMainDirLightLuminance) \
 	VIEW_UNIFORM_BUFFER_MEMBER(float, PreExposure) \
 	VIEW_UNIFORM_BUFFER_MEMBER(float, OneOverPreExposure) \
 	VIEW_UNIFORM_BUFFER_MEMBER_EX(FVector4f, DiffuseOverrideParameter, EShaderPrecisionModifier::Half) \
diff --git a/Engine/Source/Runtime/RenderCore/Private/Shader.cpp b/Engine/Source/Runtime/RenderCore/Private/Shader.cpp
--- a/Engine/Source/Runtime/RenderCore/Private/Shader.cpp
+++ b/Engine/Source/Runtime/RenderCore/Private/Shader.cpp
@@ -1803,9 +1803,10 @@ void ShaderMapAppendKeyString(EShaderPlatform Platform, FString& KeyString)
 		}
 	}
 
-	if (IsWaterDistanceFieldShadowEnabled(Platform))
+	const bool bNeedsSeparateMainDirLightTexture = IsWaterDistanceFieldShadowEnabled(Platform) || IsWaterVirtualShadowMapFilteringEnabled(Platform);
+	if (bNeedsSeparateMainDirLightTexture)
 	{
-		KeyString += TEXT("_SLWDFS");
+		KeyString += TEXT("_SLWSMDLT");
 	}
 
 	{
diff --git a/Engine/Source/Runtime/RenderCore/Public/RenderUtils.h b/Engine/Source/Runtime/RenderCore/Public/RenderUtils.h
--- a/Engine/Source/Runtime/RenderCore/Public/RenderUtils.h
+++ b/Engine/Source/Runtime/RenderCore/Public/RenderUtils.h
@@ -802,6 +802,17 @@ inline bool UseNonNaniteVirtualShadowMaps(EShaderPlatform ShaderPlatform, const
 	return EnableCVar->GetInt() != 0 && UseVirtualShadowMaps(ShaderPlatform, FeatureLevel);
 }
 
+/** Returns if water should evaluate virtual shadow maps a second time for the water surface. This is for a platform so can be used at cook time. */
+inline bool IsWaterVirtualShadowMapFilteringEnabled(const FStaticShaderPlatform Platform)
+{
+	static const auto CVarWaterSingleLayerShaderSupportVSMFiltering = IConsoleManager::Get().FindConsoleVariable(TEXT("r.Water.SingleLayer.ShadersSupportVSMFiltering"));
+	const bool bWaterVSMFilteringSupported = CVarWaterSingleLayerShaderSupportVSMFiltering && (CVarWaterSingleLayerShaderSupportVSMFiltering->GetInt() > 0);
+
+	const bool bVirtualShadowMapsSupported = DoesPlatformSupportVirtualShadowMaps(Platform);
+
+	return !IsForwardShadingEnabled(Platform) && bWaterVSMFilteringSupported && bVirtualShadowMapsSupported;
+}
+
 /**
 *	(Non-runtime) Checks if the depth prepass for single layer water is enabled. This also depends on virtual shadow maps to be supported on the platform.
 */
@@ -810,9 +821,9 @@ inline bool IsSingleLayerWaterDepthPrepassEnabled(const FStaticShaderPlatform& P
 	static const auto CVarWaterSingleLayerDepthPrepass = IConsoleManager::Get().FindConsoleVariable(TEXT("r.Water.SingleLayer.DepthPrepass"));
 	const bool bPrepassEnabled = CVarWaterSingleLayerDepthPrepass && CVarWaterSingleLayerDepthPrepass->GetInt() > 0;
 	// Currently VSM is the only feature dependent on the depth prepass which is why we only enable it if VSM could also be enabled.
-	// VSM can be toggled at runtime, but we need a compile time value here, so we fall back to DoesRuntimeSupportNanite() to check if
+	// VSM can be toggled at runtime, but we need a compile time value here, so we fall back to DoesPlatformSupportVirtualShadowMaps() to check if
 	// VSM *could* be enabled.
-	const bool bVSMSupported = DoesPlatformSupportNanite(Platform, false /* check project setting */);
+	const bool bVSMSupported = DoesPlatformSupportVirtualShadowMaps(Platform);
 
 	return bPrepassEnabled && bVSMSupported;
 }
diff --git a/Engine/Source/Runtime/RenderCore/Public/ShaderMaterial.h b/Engine/Source/Runtime/RenderCore/Public/ShaderMaterial.h
--- a/Engine/Source/Runtime/RenderCore/Public/ShaderMaterial.h
+++ b/Engine/Source/Runtime/RenderCore/Public/ShaderMaterial.h
@@ -101,7 +101,7 @@ struct FShaderMaterialPropertyDefines
 	uint8 MATERIAL_SHADINGMODEL_CLOTH : 1;
 	uint8 MATERIAL_SHADINGMODEL_EYE : 1;
 	uint8 MATERIAL_SHADINGMODEL_SINGLELAYERWATER : 1;
-	uint8 SINGLE_LAYER_WATER_DF_SHADOW_ENABLED : 1;
+	uint8 SINGLE_LAYER_WATER_SEPARATED_MAIN_LIGHT : 1;
 	uint8 MATERIAL_SHADINGMODEL_THIN_TRANSLUCENT : 1;
 
 	uint8 TRANSLUCENCY_LIGHTING_VOLUMETRIC_NONDIRECTIONAL : 1;
diff --git a/Engine/Source/Runtime/Renderer/Private/BasePassRendering.h b/Engine/Source/Runtime/Renderer/Private/BasePassRendering.h
--- a/Engine/Source/Runtime/Renderer/Private/BasePassRendering.h
+++ b/Engine/Source/Runtime/Renderer/Private/BasePassRendering.h
@@ -432,7 +432,8 @@ public:
 		OutEnvironment.SetDefine(TEXT("STRATA_INLINE_SHADING"), 1);
 
 		const bool IsSingleLayerWater = Parameters.MaterialParameters.ShadingModels.HasShadingModel(MSM_SingleLayerWater);
-		if (IsSingleLayerWater && IsWaterDistanceFieldShadowEnabled(Parameters.Platform))
+		const bool bNeedsSeparateMainDirLightTexture = IsWaterDistanceFieldShadowEnabled(Parameters.Platform) || IsWaterVirtualShadowMapFilteringEnabled(Parameters.Platform);
+		if (IsSingleLayerWater && bNeedsSeparateMainDirLightTexture)
 		{
 			// See FShaderCompileUtilities::FetchGBufferParamsRuntime for the details
 			const bool bOutputVelocity = FVelocityRendering::BasePassCanOutputVelocity(Parameters.Platform);
@@ -629,7 +630,8 @@ public:
 		OutEnvironment.SetDefine(TEXT("STRATA_INLINE_SHADING"), 1);
 
 		const bool IsSingleLayerWater = Parameters.MaterialParameters.ShadingModels.HasShadingModel(MSM_SingleLayerWater);
-		if (IsSingleLayerWater && IsWaterDistanceFieldShadowEnabled(Parameters.Platform))
+		const bool bNeedsSeparateMainDirLightTexture = IsWaterDistanceFieldShadowEnabled(Parameters.Platform) || IsWaterVirtualShadowMapFilteringEnabled(Parameters.Platform);
+		if (IsSingleLayerWater && bNeedsSeparateMainDirLightTexture)
 		{
 			// See FShaderCompileUtilities::FetchGBufferParamsRuntime for the details
 			const bool bOutputVelocity = FVelocityRendering::BasePassCanOutputVelocity(Parameters.Platform);
diff --git a/Engine/Source/Runtime/Renderer/Private/SceneRendering.cpp b/Engine/Source/Runtime/Renderer/Private/SceneRendering.cpp
--- a/Engine/Source/Runtime/Renderer/Private/SceneRendering.cpp
+++ b/Engine/Source/Runtime/Renderer/Private/SceneRendering.cpp
@@ -1271,11 +1271,6 @@ void FViewInfo::SetupUniformBufferParameters(
 
 	ViewUniformShaderParameters.BufferToSceneTextureScale = FVector2f(1.0f, 1.0f);
 
-	{
-		extern bool IsWaterDistanceFieldShadowEnabled_Runtime(const FStaticShaderPlatform Platform);
-		ViewUniformShaderParameters.SeparateWaterMainDirLightLuminance = IsWaterDistanceFieldShadowEnabled_Runtime(ShaderPlatform) ? 1.0f : 0.0f;
-	}
-
 	FRHITexture* TransmittanceLutTextureFound = nullptr;
 	FRHITexture* SkyViewLutTextureFound = nullptr;
 	FRHITexture* CameraAerialPerspectiveVolumeFound = nullptr;
diff --git a/Engine/Source/Runtime/Renderer/Private/Shadows/ShadowSceneRenderer.cpp b/Engine/Source/Runtime/Renderer/Private/Shadows/ShadowSceneRenderer.cpp
--- a/Engine/Source/Runtime/Renderer/Private/Shadows/ShadowSceneRenderer.cpp
+++ b/Engine/Source/Runtime/Renderer/Private/Shadows/ShadowSceneRenderer.cpp
@@ -398,6 +398,8 @@ void FShadowSceneRenderer::ApplyVirtualShadowMapProjectionForLight(
 						ScissorRect,
 						EVirtualShadowMapProjectionInputType::GBuffer,
 						VisibleLightInfo.FindShadowClipmapForView(&View),
+						false, // bModulateRGB
+						nullptr, // TiledVSMProjection
 						OutputScreenShadowMaskTexture);
 				}
 				else if (bShouldUseVirtualShadowMapOnePassProjection)
@@ -443,6 +445,8 @@ void FShadowSceneRenderer::ApplyVirtualShadowMapProjectionForLight(
 						ScissorRect,
 						EVirtualShadowMapProjectionInputType::HairStrands,
 						VisibleLightInfo.FindShadowClipmapForView(&View),
+						false, // bModulateRGB
+						nullptr, // TiledVSMProjection
 						OutputScreenShadowMaskSubPixelTexture);
 				}
 				else if (bShouldUseVirtualShadowMapOnePassProjection)
diff --git a/Engine/Source/Runtime/Renderer/Private/SingleLayerWaterRendering.cpp b/Engine/Source/Runtime/Renderer/Private/SingleLayerWaterRendering.cpp
--- a/Engine/Source/Runtime/Renderer/Private/SingleLayerWaterRendering.cpp
+++ b/Engine/Source/Runtime/Renderer/Private/SingleLayerWaterRendering.cpp
@@ -18,6 +18,7 @@
 #include "SceneTextureParameters.h"
 #include "Strata/Strata.h"
 #include "VirtualShadowMaps/VirtualShadowMapArray.h"
+#include "VirtualShadowMaps/VirtualShadowMapProjection.h"
 #include "Lumen/LumenSceneData.h"
 #include "Lumen/LumenTracingUtils.h"
 
@@ -71,6 +72,21 @@ static TAutoConsoleVariable<int32> CVarWaterSingleLayerShadersSupportDistanceFie
 	TEXT("Whether or not the single layer water material shaders are compiled with support for distance field shadow, i.e. output main directional light luminance in a separate render target. This is preconditioned on using deferred shading and having distance field support enabled in the project."),
 	ECVF_ReadOnly | ECVF_RenderThreadSafe);
 
+static TAutoConsoleVariable<int32> CVarWaterSingleLayerDistanceFieldShadow(
+	TEXT("r.Water.SingleLayer.DistanceFieldShadow"), 1,
+	TEXT("When using deferred, distance field shadow tracing is supported on single layer water. This cvar can be used to toggle it on/off at runtime."),
+	ECVF_RenderThreadSafe | ECVF_Scalability);
+
+static TAutoConsoleVariable<int32> CVarWaterSingleLayerShadersSupportVSMFiltering(
+	TEXT("r.Water.SingleLayer.ShadersSupportVSMFiltering"), 1,
+	TEXT("Whether or not the single layer water material shaders are compiled with support for virtual shadow map filter, i.e. output main directional light luminance in a separate render target. This is preconditioned on using deferred shading and having VSM support enabled in the project."),
+	ECVF_ReadOnly | ECVF_RenderThreadSafe);
+
+static TAutoConsoleVariable<int32> CVarWaterSingleLayerVSMFiltering(
+	TEXT("r.Water.SingleLayer.VSMFiltering"), 1,
+	TEXT("When using deferred, virtual shadow map filtering is supported on single layer water. This cvar can be used to toggle it on/off at runtime."),
+	ECVF_RenderThreadSafe | ECVF_Scalability);
+
 // The project setting for the cloud shadow to affect SingleLayerWater (enable/disable runtime and shader code).  This is not implemented on mobile as VolumetricClouds are not available on these platforms.
 static TAutoConsoleVariable<int32> CVarSupportCloudShadowOnSingleLayerWater(
 	TEXT("r.Water.SingleLayerWater.SupportCloudShadow"),
@@ -78,11 +94,6 @@ static TAutoConsoleVariable<int32> CVarSupportCloudShadowOnSingleLayerWater(
 	TEXT("Enables cloud shadows on SingleLayerWater materials."),
 	ECVF_ReadOnly | ECVF_RenderThreadSafe);
 
-static TAutoConsoleVariable<int32> CVarWaterSingleLayerDistanceFieldShadow(
-	TEXT("r.Water.SingleLayer.DistanceFieldShadow"), 1,
-	TEXT("When using deferred, distance field shadow tracing is supported on single layer water. This cvar can be used to toggle it on/off at runtime."),
-	ECVF_RenderThreadSafe | ECVF_Scalability);
-
 static TAutoConsoleVariable<int32> CVarWaterSingleLayerRTR(
 	TEXT("r.Water.SingleLayer.RTR"), 1,
 	TEXT("Enable RTR for the single water renderring system."),
@@ -169,6 +180,21 @@ bool IsWaterDistanceFieldShadowEnabled_Runtime(const FStaticShaderPlatform Platf
 	return IsWaterDistanceFieldShadowEnabled(Platform) && CVarWaterSingleLayerDistanceFieldShadow.GetValueOnAnyThread() > 0;
 }
 
+bool IsWaterVirtualShadowMapFilteringEnabled_Runtime(const FStaticShaderPlatform Platform)
+{
+	return IsWaterVirtualShadowMapFilteringEnabled(Platform) && UseVirtualShadowMaps(Platform, GetMaxSupportedFeatureLevel(Platform)) && CVarWaterSingleLayerVSMFiltering.GetValueOnRenderThread() > 0;
+}
+
+bool NeedsSeparatedMainDirectionalLightTexture(const FStaticShaderPlatform Platform)
+{
+	return IsWaterDistanceFieldShadowEnabled(Platform) || IsWaterVirtualShadowMapFilteringEnabled(Platform);
+}
+
+bool NeedsSeparatedMainDirectionalLightTexture_Runtime(const FStaticShaderPlatform Platform)
+{
+	return IsWaterDistanceFieldShadowEnabled_Runtime(Platform) || IsWaterVirtualShadowMapFilteringEnabled_Runtime(Platform);
+}
+
 BEGIN_SHADER_PARAMETER_STRUCT(FSingleLayerWaterCommonShaderParameters, )
 	SHADER_PARAMETER_RDG_TEXTURE(Texture2D, ScreenSpaceReflectionsTexture)
 	SHADER_PARAMETER_SAMPLER(SamplerState, ScreenSpaceReflectionsSampler)
@@ -379,7 +405,7 @@ void FDeferredShadingSceneRenderer::RenderSingleLayerWaterDepthPrepass(FRDGBuild
 			}
 			PassParameters->RenderTargets.DepthStencil = FDepthStencilBinding(OutDepthPrepassTexture.Target, ERenderTargetLoadAction::ENoAction, ERenderTargetLoadAction::ELoad, FExclusiveDepthStencil::DepthWrite_StencilNop);
 
-			FGlobalShaderMap* ShaderMap = GetGlobalShaderMap(GMaxRHIFeatureLevel);
+			FGlobalShaderMap* ShaderMap = GetGlobalShaderMap(FeatureLevel);
 
 			FCopyDepthPS::FPermutationDomain PermutationVector;
 			PermutationVector.Set<FCopyDepthPS::FMSAASampleCount>(DepthPrepassTextureDesc.NumSamples);
@@ -686,34 +712,17 @@ void FDeferredShadingSceneRenderer::RenderSingleLayerWaterReflections(
 		}
 
 		const FPerViewPipelineState& ViewPipelineState = GetViewPipelineState(View);
+		const FStaticShaderPlatform StaticShaderPlatform = View.GetShaderPlatform();
+		const bool bWaterVSMFiltering = IsWaterVirtualShadowMapFilteringEnabled_Runtime(StaticShaderPlatform);
+		const bool bWaterDistanceFieldShadow = IsWaterDistanceFieldShadowEnabled_Runtime(StaticShaderPlatform);
 
-		if (IsWaterDistanceFieldShadowEnabled_Runtime(View.GetShaderPlatform()))
+		if (bWaterVSMFiltering || bWaterDistanceFieldShadow)
 		{
-
-			FProjectedShadowInfo* DistanceFieldShadowInfo = nullptr;
-
-			// Try to find the ProjectedShadowInfo corresponding to ray trace shadow info for the main directional light.
 			const FLightSceneProxy* SelectedForwardDirectionalLightProxy = View.ForwardLightingResources.SelectedForwardDirectionalLightProxy;
-			if (SelectedForwardDirectionalLightProxy)
-			{
-				FLightSceneInfo* LightSceneInfo = SelectedForwardDirectionalLightProxy->GetLightSceneInfo();
-				FVisibleLightInfo& VisibleLightViewInfo = VisibleLightInfos[LightSceneInfo->Id];
 
-				for (int32 ShadowIndex = 0; ShadowIndex < VisibleLightViewInfo.ShadowsToProject.Num(); ShadowIndex++)
-				{
-					FProjectedShadowInfo* ProjectedShadowInfo = VisibleLightViewInfo.ShadowsToProject[ShadowIndex];
-					if (ProjectedShadowInfo->bRayTracedDistanceField)
-					{
-						DistanceFieldShadowInfo = ProjectedShadowInfo;
-					}
-				}
-			}
-
-			// If DFShadow data has been found, then combine it with the separate main directional light luminance texture.
-			FRDGTextureRef ScreenShadowMaskTexture = SystemTextures.White;
-			if (DistanceFieldShadowInfo)
+			if (bWaterVSMFiltering && SelectedForwardDirectionalLightProxy)
 			{
-				RDG_EVENT_SCOPE(GraphBuilder, "SLW::DistanceFieldShadow");
+				RDG_EVENT_SCOPE(GraphBuilder, "SLW::VirtualShadowMaps");
 
 				FIntRect ScissorRect;
 				if (!SelectedForwardDirectionalLightProxy->GetScissorRect(ScissorRect, View, View.ViewRect))
@@ -721,28 +730,82 @@ void FDeferredShadingSceneRenderer::RenderSingleLayerWaterReflections(
 					ScissorRect = View.ViewRect;
 				}
 
-				// Reset the cached texture to create a new one mapping to the water depth buffer
-				DistanceFieldShadowInfo->ResetRayTracedDistanceFieldShadow(&View);
+				const FVisibleLightInfo& VisibleLightInfo = VisibleLightInfos[SelectedForwardDirectionalLightProxy->GetLightSceneInfo()->Id];
 
-				FTiledShadowRendering TiledShadowRendering;
-				if (bRunTiled)
-				{
-					TiledShadowRendering.DrawIndirectParametersBuffer = TiledScreenSpaceReflection.DrawIndirectParametersBuffer;
-					TiledShadowRendering.TileListDataBufferSRV = TiledScreenSpaceReflection.TileListDataBufferSRV;
-					TiledShadowRendering.TileSize = TiledScreenSpaceReflection.TileSize;
-				}
+				FTiledVSMProjection TiledVSMProjection{};
+				TiledVSMProjection.DrawIndirectParametersBuffer = TiledScreenSpaceReflection.DrawIndirectParametersBuffer;
+				TiledVSMProjection.DispatchIndirectParametersBuffer = TiledScreenSpaceReflection.DispatchIndirectParametersBuffer;
+				TiledVSMProjection.TileListDataBufferSRV = TiledScreenSpaceReflection.TileListDataBufferSRV;
+				TiledVSMProjection.TileSize = TiledScreenSpaceReflection.TileSize;
 
-				const bool bProjectingForForwardShading = false;
-				const bool bForceRGBModulation = true;
-				DistanceFieldShadowInfo->RenderRayTracedDistanceFieldProjection(
+				RenderVirtualShadowMapProjection(
 					GraphBuilder,
 					SceneTextures,
-					SceneWithoutWaterTextures.SeparatedMainDirLightTexture,
-					View,
+					View, ViewIndex,
+					VirtualShadowMapArray,
 					ScissorRect,
-					bProjectingForForwardShading,
-					bForceRGBModulation,
-					bRunTiled ? &TiledShadowRendering : nullptr);
+					EVirtualShadowMapProjectionInputType::GBuffer,
+					VisibleLightInfo.FindShadowClipmapForView(&View),
+					true, // bModulateRGB
+					&TiledVSMProjection,
+					SceneWithoutWaterTextures.SeparatedMainDirLightTexture);
+			}
+
+			if (bWaterDistanceFieldShadow)
+			{
+				FProjectedShadowInfo* DistanceFieldShadowInfo = nullptr;
+
+				// Try to find the ProjectedShadowInfo corresponding to ray trace shadow info for the main directional light.
+				if (SelectedForwardDirectionalLightProxy)
+				{
+					FLightSceneInfo* LightSceneInfo = SelectedForwardDirectionalLightProxy->GetLightSceneInfo();
+					FVisibleLightInfo& VisibleLightViewInfo = VisibleLightInfos[LightSceneInfo->Id];
+
+					for (int32 ShadowIndex = 0; ShadowIndex < VisibleLightViewInfo.ShadowsToProject.Num(); ShadowIndex++)
+					{
+						FProjectedShadowInfo* ProjectedShadowInfo = VisibleLightViewInfo.ShadowsToProject[ShadowIndex];
+						if (ProjectedShadowInfo->bRayTracedDistanceField)
+						{
+							DistanceFieldShadowInfo = ProjectedShadowInfo;
+						}
+					}
+				}
+
+				// If DFShadow data has been found, then combine it with the separate main directional light luminance texture.
+				FRDGTextureRef ScreenShadowMaskTexture = SystemTextures.White;
+				if (DistanceFieldShadowInfo)
+				{
+					RDG_EVENT_SCOPE(GraphBuilder, "SLW::DistanceFieldShadow");
+
+					FIntRect ScissorRect;
+					if (!SelectedForwardDirectionalLightProxy->GetScissorRect(ScissorRect, View, View.ViewRect))
+					{
+						ScissorRect = View.ViewRect;
+					}
+
+					// Reset the cached texture to create a new one mapping to the water depth buffer
+					DistanceFieldShadowInfo->ResetRayTracedDistanceFieldShadow(&View);
+
+					FTiledShadowRendering TiledShadowRendering;
+					if (bRunTiled)
+					{
+						TiledShadowRendering.DrawIndirectParametersBuffer = TiledScreenSpaceReflection.DrawIndirectParametersBuffer;
+						TiledShadowRendering.TileListDataBufferSRV = TiledScreenSpaceReflection.TileListDataBufferSRV;
+						TiledShadowRendering.TileSize = TiledScreenSpaceReflection.TileSize;
+					}
+
+					const bool bProjectingForForwardShading = false;
+					const bool bForceRGBModulation = true;
+					DistanceFieldShadowInfo->RenderRayTracedDistanceFieldProjection(
+						GraphBuilder,
+						SceneTextures,
+						SceneWithoutWaterTextures.SeparatedMainDirLightTexture,
+						View,
+						ScissorRect,
+						bProjectingForForwardShading,
+						bForceRGBModulation,
+						bRunTiled ? &TiledShadowRendering : nullptr);
+				}
 			}
 		}
 
@@ -906,7 +969,7 @@ void FDeferredShadingSceneRenderer::RenderSingleLayerWaterReflections(
 			PassParameters->VS.TileListData = TiledScreenSpaceReflection.TileListDataBufferSRV;
 
 			SetCommonParameters(PassParameters->PS.CommonParameters);
-			if (IsWaterDistanceFieldShadowEnabled_Runtime(Scene->GetShaderPlatform()))
+			if (NeedsSeparatedMainDirectionalLightTexture_Runtime(Scene->GetShaderPlatform()))
 			{
 				PassParameters->PS.CommonParameters.SeparatedMainDirLightTexture = SceneWithoutWaterTextures.SeparatedMainDirLightTexture;
 				PassParameters->PS.CommonParameters.UseSeparatedMainDirLightTexture = 1.0f;
@@ -1032,6 +1095,8 @@ BEGIN_UNIFORM_BUFFER_STRUCT(FSingleLayerWaterPassUniformParameters,)
 	SHADER_PARAMETER(FVector4f, DistortionParams)
 	SHADER_PARAMETER(FVector2f, SceneWithoutSingleLayerWaterTextureSize)
 	SHADER_PARAMETER(FVector2f, SceneWithoutSingleLayerWaterInvTextureSize)
+	SHADER_PARAMETER(uint32, bMainDirectionalLightVSMFiltering)
+	SHADER_PARAMETER(uint32, bSeparateMainDirLightLuminance)
 	SHADER_PARAMETER_STRUCT(FLightCloudTransmittanceParameters, ForwardDirLightCloudShadow)
 END_UNIFORM_BUFFER_STRUCT()
 
@@ -1066,7 +1131,7 @@ void FDeferredShadingSceneRenderer::RenderSingleLayerWaterInner(
 
 	TStaticArray<FTextureRenderTargetBinding, MaxSimultaneousRenderTargets> BasePassTextures;
 	uint32 BasePassTextureCount = SceneTextures.GetGBufferRenderTargets(BasePassTextures);
-	if(IsWaterDistanceFieldShadowEnabled_Runtime(Scene->GetShaderPlatform()))
+	if(NeedsSeparatedMainDirectionalLightTexture(Scene->GetShaderPlatform())
 	{
 		const bool bNeverClear = true;
 		BasePassTextures[BasePassTextureCount++] = FTextureRenderTargetBinding(SceneWithoutWaterTextures.SeparatedMainDirLightTexture, bNeverClear);
@@ -1117,6 +1182,8 @@ void FDeferredShadingSceneRenderer::RenderSingleLayerWaterInner(
 			SetupDistortionParams(SLWUniformParameters.DistortionParams, View);
 			SLWUniformParameters.SceneWithoutSingleLayerWaterTextureSize = FVector2f(DepthTextureSize.X, DepthTextureSize.Y);
 			SLWUniformParameters.SceneWithoutSingleLayerWaterInvTextureSize = FVector2f(1.0f / DepthTextureSize.X, 1.0f / DepthTextureSize.Y);
+			SLWUniformParameters.bMainDirectionalLightVSMFiltering = IsWaterVirtualShadowMapFilteringEnabled_Runtime(View.GetShaderPlatform());
+			SLWUniformParameters.bSeparateMainDirLightLuminance = NeedsSeparatedMainDirectionalLightTexture_Runtime(View.GetShaderPlatform());
 
 			const FLightSceneProxy* SelectedForwardDirectionalLightProxy = View.ForwardLightingResources.SelectedForwardDirectionalLightProxy;
 			SetupLightCloudTransmittanceParameters(GraphBuilder, Scene, View, SelectedForwardDirectionalLightProxy ? SelectedForwardDirectionalLightProxy->GetLightSceneInfo() : nullptr, SLWUniformParameters.ForwardDirLightCloudShadow);
@@ -1331,8 +1398,7 @@ void FSingleLayerWaterPassMeshProcessor::CollectPSOInitializers(const FSceneText
 
 		FGraphicsPipelineRenderTargetsInfo RenderTargetsInfo;
 		SetupGBufferRenderTargetInfo(SceneTexturesConfig, RenderTargetsInfo, true /*bSetupDepthStencil*/);
-		if (IsWaterDistanceFieldShadowEnabled_Runtime(GMaxRHIShaderPlatform)
-			&& !Strata::IsStrataEnabled())	// We do not bind that texture if Strata is enabled as the data will go through the Strata material buffer.
+		if (NeedsSeparatedMainDirectionalLightTexture(GMaxRHIShaderPlatform))
 		{
 			AddRenderTargetInfo(PF_FloatR11G11B10, TexCreate_ShaderResource | TexCreate_RenderTargetable, RenderTargetsInfo);
 		}
diff --git a/Engine/Source/Runtime/Renderer/Private/VirtualShadowMaps/VirtualShadowMapProjection.cpp b/Engine/Source/Runtime/Renderer/Private/VirtualShadowMaps/VirtualShadowMapProjection.cpp
--- a/Engine/Source/Runtime/Renderer/Private/VirtualShadowMaps/VirtualShadowMapProjection.cpp
+++ b/Engine/Source/Runtime/Renderer/Private/VirtualShadowMaps/VirtualShadowMapProjection.cpp
@@ -174,6 +174,10 @@ static TAutoConsoleVariable<int32> CVarTestPermutation(
 );
 #endif
 
+// The tile size in pixels for VSM projection with tile list.
+// Is also used as the workgroup size for the CS without tile list.
+static constexpr int32 VSMProjectionWorkTileSize = 8;
+
 const TCHAR* ToString(EVirtualShadowMapProjectionInputType In)
 {
 	switch (In)
@@ -195,6 +199,7 @@ class FVirtualShadowMapProjectionCS : public FGlobalShader
 	class FHairStrandsDim			: SHADER_PERMUTATION_BOOL("HAS_HAIR_STRANDS");
 	class FVisualizeOutputDim		: SHADER_PERMUTATION_BOOL("VISUALIZE_OUTPUT");
 	class FExtrapolateSlopeDim		: SHADER_PERMUTATION_BOOL("SMRT_EXTRAPOLATE_SLOPE");
+	class FUseTileList				: SHADER_PERMUTATION_BOOL("USE_TILE_LIST");
 	// -1 means dynamic count
 	class FSMRTStaticSampleCount	: SHADER_PERMUTATION_RANGE_INT("SMRT_TEMPLATE_STATIC_SAMPLES_PER_RAY", -1, 2);
 	// Used for A/B testing a change that affects reg allocation, etc.
@@ -206,6 +211,7 @@ class FVirtualShadowMapProjectionCS : public FGlobalShader
 		FHairStrandsDim,
 		FVisualizeOutputDim,
 		FExtrapolateSlopeDim,
+		FUseTileList,
 		FSMRTStaticSampleCount
 #if MAX_TEST_PERMUTATION > 0
 		, FTestDim
@@ -246,6 +252,9 @@ class FVirtualShadowMapProjectionCS : public FGlobalShader
 		SHADER_PARAMETER(int32, VisualizeModeId)
 		SHADER_PARAMETER(int32, VisualizeVirtualShadowMapId)
 		SHADER_PARAMETER_RDG_TEXTURE_UAV(RWTexture2D, OutVisualize)
+		// Optional tile list
+		SHADER_PARAMETER_RDG_BUFFER_SRV(Buffer, TileListData)
+		RDG_BUFFER_ACCESS(IndirectDispatchArgs, ERHIAccess::IndirectArgs)
 	END_SHADER_PARAMETER_STRUCT()
 
 	static void ModifyCompilationEnvironment( const FGlobalShaderPermutationParameters& Parameters, FShaderCompilerEnvironment& OutEnvironment )
@@ -264,6 +273,8 @@ class FVirtualShadowMapProjectionCS : public FGlobalShader
 		{
 			OutEnvironment.CompilerFlags.Add(CFLAG_AllowRealTypes);
 		}
+
+		OutEnvironment.SetDefine(TEXT("WORK_TILE_SIZE"), VSMProjectionWorkTileSize);
 	}
 
 	static bool ShouldCompilePermutation(const FGlobalShaderPermutationParameters& Parameters)
@@ -296,10 +307,14 @@ static void RenderVirtualShadowMapProjectionCommon(
 	EVirtualShadowMapProjectionInputType InputType,
 	FRDGTextureRef OutputTexture,
 	const FLightSceneProxy* LightProxy = nullptr,
-	int32 VirtualShadowMapId = INDEX_NONE)
+	int32 VirtualShadowMapId = INDEX_NONE,
+	FTiledVSMProjection* TiledVSMProjection = nullptr)
 {
 	check(GRHISupportsWaveOperations);
 
+	const bool bUseTileList = TiledVSMProjection != nullptr;
+	check(!bUseTileList || TiledVSMProjection->TileSize == VSMProjectionWorkTileSize);
+
 	// Use hair strands data (i.e., hair voxel tracing) only for Gbuffer input for casting hair shadow onto opaque geometry.
 	const bool bHasHairStrandsData = HairStrands::HasViewHairStrandsData(View);
 
@@ -316,6 +331,11 @@ static void RenderVirtualShadowMapProjectionCommon(
 	PassParameters->bCullBackfacingPixels = VirtualShadowMapArray.ShouldCullBackfacingPixels() ? 1 : 0;
 	PassParameters->bSMRTUseAdaptiveRayCount = CVarSMRTAdaptiveRayCount.GetValueOnRenderThread() != 0 ? 1 : 0;
 	PassParameters->Strata = Strata::BindStrataGlobalUniformParameters(View);
+	if (bUseTileList)
+	{
+		PassParameters->TileListData = TiledVSMProjection->TileListDataBufferSRV;
+		PassParameters->IndirectDispatchArgs = TiledVSMProjection->DispatchIndirectParametersBuffer;
+	}
 	if (bHasHairStrandsData)
 	{
 		PassParameters->HairStrands = HairStrands::BindHairStrandsViewUniformParameters(View);
@@ -386,6 +406,7 @@ static void RenderVirtualShadowMapProjectionCommon(
 	PermutationVector.Set< FVirtualShadowMapProjectionCS::FHairStrandsDim >( bHasHairStrandsData );
 	PermutationVector.Set< FVirtualShadowMapProjectionCS::FVisualizeOutputDim >( bDebugOutput );
 	PermutationVector.Set< FVirtualShadowMapProjectionCS::FExtrapolateSlopeDim >( PassParameters->SMRTExtrapolateSlope > 0.0f );
+	PermutationVector.Set< FVirtualShadowMapProjectionCS::FUseTileList >(bUseTileList);
 	PermutationVector.Set< FVirtualShadowMapProjectionCS::FSMRTStaticSampleCount >( StaticSamplesPerRay );
 #if MAX_TEST_PERMUTATION > 0
 	{
@@ -398,19 +419,38 @@ static void RenderVirtualShadowMapProjectionCommon(
 	ClearUnusedGraphResources( ComputeShader, PassParameters );
 	ValidateShaderParameters( ComputeShader, *PassParameters );
 
-	const FIntPoint GroupCount = FIntPoint::DivideAndRoundUp( ProjectionRect.Size(), 8 );
-	FComputeShaderUtils::AddPass(
-		GraphBuilder,
-		RDG_EVENT_NAME("VirtualShadowMapProjection(RayCount:%u(%s),SamplesPerRay:%u,Input:%s%s)",
-			PassParameters->SMRTRayCount,
-			PassParameters->bSMRTUseAdaptiveRayCount ? TEXT("Adaptive") : TEXT("Static"), 
-			PassParameters->SMRTSamplesPerRay,
-			ToString(InputType),
-			bDebugOutput ? TEXT(",Debug") : TEXT("")),
-		ComputeShader,
-		PassParameters,
-		FIntVector( GroupCount.X, GroupCount.Y, 1 )
-	);
+	if (bUseTileList)
+	{
+		GraphBuilder.AddPass(
+			RDG_EVENT_NAME("VirtualShadowMapProjection(RayCount:%u(%s),SamplesPerRay:%u,Input:%s%s,TileList)",
+				PassParameters->SMRTRayCount,
+				PassParameters->bSMRTUseAdaptiveRayCount ? TEXT("Adaptive") : TEXT("Static"),
+				PassParameters->SMRTSamplesPerRay,
+				ToString(InputType),
+				bDebugOutput ? TEXT(",Debug") : TEXT("")),
+			PassParameters,
+			ERDGPassFlags::Compute,
+			[PassParameters, ComputeShader](FRHICommandList& RHICmdList)
+			{
+				FComputeShaderUtils::DispatchIndirect(RHICmdList, ComputeShader, *PassParameters, PassParameters->IndirectDispatchArgs->GetIndirectRHICallBuffer(), 0);
+			});
+	}
+	else
+	{
+		const FIntPoint GroupCount = FIntPoint::DivideAndRoundUp(ProjectionRect.Size(), VSMProjectionWorkTileSize);
+		FComputeShaderUtils::AddPass(
+			GraphBuilder,
+			RDG_EVENT_NAME("VirtualShadowMapProjection(RayCount:%u(%s),SamplesPerRay:%u,Input:%s%s)",
+				PassParameters->SMRTRayCount,
+				PassParameters->bSMRTUseAdaptiveRayCount ? TEXT("Adaptive") : TEXT("Static"),
+				PassParameters->SMRTSamplesPerRay,
+				ToString(InputType),
+				bDebugOutput ? TEXT(",Debug") : TEXT("")),
+			ComputeShader,
+			PassParameters,
+			FIntVector(GroupCount.X, GroupCount.Y, 1)
+		);
+	}
 }
 
 FRDGTextureRef CreateVirtualShadowMapMaskBits(
@@ -499,6 +539,8 @@ void RenderVirtualShadowMapProjection(
 		ScissorRect,
 		VirtualShadowMaskTexture,
 		false,	// bDirectionalLight
+		false, // bModulateRGB
+		nullptr, //TiledVSMProjection
 		OutputShadowMaskTexture);
 }
 
@@ -510,6 +552,8 @@ void RenderVirtualShadowMapProjection(
 	const FIntRect ScissorRect,
 	EVirtualShadowMapProjectionInputType InputType,
 	const TSharedPtr<FVirtualShadowMapClipmap>& Clipmap,
+	bool bModulateRGB,
+	FTiledVSMProjection* TiledVSMProjection,
 	FRDGTextureRef OutputShadowMaskTexture)
 {
 	FRDGTextureRef VirtualShadowMaskTexture = CreateShadowMaskTexture(GraphBuilder, View.ViewRect.Max);
@@ -523,7 +567,8 @@ void RenderVirtualShadowMapProjection(
 		InputType,
 		VirtualShadowMaskTexture,
 		Clipmap->GetLightSceneInfo().Proxy,
-		Clipmap->GetVirtualShadowMap()->ID);
+		Clipmap->GetVirtualShadowMap()->ID,
+		TiledVSMProjection);
 	
 	CompositeVirtualShadowMapMask(
 		GraphBuilder,
@@ -531,9 +576,35 @@ void RenderVirtualShadowMapProjection(
 		ScissorRect,
 		VirtualShadowMaskTexture,
 		true,	// bDirectionalLight
+		bModulateRGB,
+		TiledVSMProjection,
 		OutputShadowMaskTexture);
 }
 
+class FVirtualShadowMapProjectionCompositeTileVS : public FGlobalShader
+{
+	DECLARE_GLOBAL_SHADER(FVirtualShadowMapProjectionCompositeTileVS);
+	SHADER_USE_PARAMETER_STRUCT(FVirtualShadowMapProjectionCompositeTileVS, FGlobalShader);
+
+	BEGIN_SHADER_PARAMETER_STRUCT(FParameters, )
+		SHADER_PARAMETER_STRUCT_REF(FViewUniformShaderParameters, ViewUniformBuffer)
+		SHADER_PARAMETER_RDG_BUFFER_SRV(Buffer<uint>, TileListData)
+	END_SHADER_PARAMETER_STRUCT()
+
+	static bool ShouldCompilePermutation(const FGlobalShaderPermutationParameters& Parameters)
+	{
+		return DoesPlatformSupportVirtualShadowMaps(Parameters.Platform);
+	}
+
+	static void ModifyCompilationEnvironment(const FGlobalShaderPermutationParameters& Parameters, FShaderCompilerEnvironment& OutEnvironment)
+	{
+		// Required right now due to where the shader function lives, but not actually used
+		FVirtualShadowMapArray::SetShaderDefines(OutEnvironment);
+		OutEnvironment.SetDefine(TEXT("WORK_TILE_SIZE"), VSMProjectionWorkTileSize);
+	}
+};
+IMPLEMENT_GLOBAL_SHADER(FVirtualShadowMapProjectionCompositeTileVS, "/Engine/Private/VirtualShadowMaps/VirtualShadowMapProjectionComposite.usf", "VirtualShadowMapCompositeTileVS", SF_Vertex);
+
 // Composite denoised shadow projection mask onto the light's shadow mask
 // Basically just a copy shader with a special blend mode
 class FVirtualShadowMapProjectionCompositePS : public FGlobalShader
@@ -543,6 +614,7 @@ class FVirtualShadowMapProjectionCompositePS : public FGlobalShader
 
 	BEGIN_SHADER_PARAMETER_STRUCT(FParameters, )
 		SHADER_PARAMETER_RDG_TEXTURE(Texture2D<float4>, InputShadowFactor)
+		SHADER_PARAMETER(uint32, bModulateRGB)
 		RENDER_TARGET_BINDING_SLOTS()
 	END_SHADER_PARAMETER_STRUCT()
 
@@ -555,41 +627,107 @@ class FVirtualShadowMapProjectionCompositePS : public FGlobalShader
 	{
 		// Required right now due to where the shader function lives, but not actually used
 		FVirtualShadowMapArray::SetShaderDefines(OutEnvironment);
+		OutEnvironment.SetDefine(TEXT("WORK_TILE_SIZE"), VSMProjectionWorkTileSize);
 	}
 };
 IMPLEMENT_GLOBAL_SHADER(FVirtualShadowMapProjectionCompositePS, "/Engine/Private/VirtualShadowMaps/VirtualShadowMapProjectionComposite.usf", "VirtualShadowMapCompositePS", SF_Pixel);
 
+BEGIN_SHADER_PARAMETER_STRUCT(FVirtualShadowMapProjectionCompositeTile, )
+	SHADER_PARAMETER_STRUCT_INCLUDE(FVirtualShadowMapProjectionCompositePS::FParameters, PS)
+	SHADER_PARAMETER_STRUCT_INCLUDE(FVirtualShadowMapProjectionCompositeTileVS::FParameters, VS)
+	RDG_BUFFER_ACCESS(IndirectDrawParameter, ERHIAccess::IndirectArgs)
+END_SHADER_PARAMETER_STRUCT()
+
 void CompositeVirtualShadowMapMask(
 	FRDGBuilder& GraphBuilder,
 	const FViewInfo& View,
 	const FIntRect ScissorRect,
 	const FRDGTextureRef Input,
 	bool bDirectionalLight,
+	bool bModulateRGB,
+	FTiledVSMProjection* TiledVSMProjection,
 	FRDGTextureRef OutputShadowMaskTexture)
 {
-	FVirtualShadowMapProjectionCompositePS::FParameters* PassParameters = GraphBuilder.AllocParameters<FVirtualShadowMapProjectionCompositePS::FParameters>();
-	PassParameters->InputShadowFactor = Input;
-
-	PassParameters->RenderTargets[0] = FRenderTargetBinding(OutputShadowMaskTexture, ERenderTargetLoadAction::ELoad);
-
-	FRHIBlendState* BlendState = FProjectedShadowInfo::GetBlendStateForProjection(
-		0,					// ShadowMapChannel
-		bDirectionalLight,	// bIsWholeSceneDirectionalShadow,
-		false,				// bUseFadePlane
-		false,				// bProjectingForForwardShading, 
-		false				// bMobileModulatedProjections
-	);
+	const bool bUseTileList = TiledVSMProjection != nullptr;
+	check(!bUseTileList || TiledVSMProjection->TileSize == VSMProjectionWorkTileSize);
 
 	auto PixelShader = View.ShaderMap->GetShader<FVirtualShadowMapProjectionCompositePS>();
-	ValidateShaderParameters(PixelShader, *PassParameters);
 
-	FPixelShaderUtils::AddFullscreenPass(GraphBuilder,
-		View.ShaderMap,
-		RDG_EVENT_NAME("CompositeVirtualShadowMapMask"),
-		PixelShader,
-		PassParameters,
-		ScissorRect,
-		BlendState);
+	FRHIBlendState* BlendState;
+	if (bModulateRGB)
+	{
+		// This has the shadow contribution modulate all the channels, e.g. used for water rendering to apply VSM on the main light RGB luminance for the updated depth buffer with water in it.
+		BlendState = TStaticBlendState<CW_RGBA, BO_Add, BF_Zero, BF_SourceColor, BO_Add, BF_Zero, BF_One>::GetRHI();
+	}
+	else
+	{
+		BlendState = FProjectedShadowInfo::GetBlendStateForProjection(
+			0,					// ShadowMapChannel
+			bDirectionalLight,	// bIsWholeSceneDirectionalShadow,
+			false,				// bUseFadePlane
+			false,				// bProjectingForForwardShading, 
+			false				// bMobileModulatedProjections
+		);
+	}
+
+	if (bUseTileList)
+	{
+		FVirtualShadowMapProjectionCompositeTile* PassParameters = GraphBuilder.AllocParameters<FVirtualShadowMapProjectionCompositeTile>();
+		PassParameters->VS.ViewUniformBuffer = View.ViewUniformBuffer;
+		PassParameters->VS.TileListData = TiledVSMProjection->TileListDataBufferSRV;
+		PassParameters->PS.InputShadowFactor = Input;
+		PassParameters->PS.bModulateRGB = bModulateRGB;
+		PassParameters->PS.RenderTargets[0] = FRenderTargetBinding(OutputShadowMaskTexture, ERenderTargetLoadAction::ELoad);
+		PassParameters->IndirectDrawParameter = TiledVSMProjection->DrawIndirectParametersBuffer;
+
+		auto VertexShader = View.ShaderMap->GetShader<FVirtualShadowMapProjectionCompositeTileVS>();
+
+		GraphBuilder.AddPass(
+			RDG_EVENT_NAME("CompositeVirtualShadowMapMask(TileList)"),
+			PassParameters,
+			ERDGPassFlags::Raster,
+			[BlendState, VertexShader, PixelShader, ScissorRect, PassParameters](FRHICommandList& RHICmdList)
+			{
+				RHICmdList.SetViewport(ScissorRect.Min.X, ScissorRect.Min.Y, 0.0f, ScissorRect.Max.X, ScissorRect.Max.Y, 1.0f);
+				RHICmdList.SetScissorRect(true, ScissorRect.Min.X, ScissorRect.Min.Y, ScissorRect.Max.X, ScissorRect.Max.Y);
+
+				FGraphicsPipelineStateInitializer GraphicsPSOInit;
+				RHICmdList.ApplyCachedRenderTargets(GraphicsPSOInit);
+				GraphicsPSOInit.RasterizerState = TStaticRasterizerState<FM_Solid, CM_None>::GetRHI();
+				GraphicsPSOInit.DepthStencilState = TStaticDepthStencilState<false, CF_Always>::GetRHI();
+				GraphicsPSOInit.BlendState = BlendState;
+				GraphicsPSOInit.bDepthBounds = false;
+				GraphicsPSOInit.PrimitiveType = GRHISupportsRectTopology ? PT_RectList : PT_TriangleList;
+				GraphicsPSOInit.BoundShaderState.VertexDeclarationRHI = GFilterVertexDeclaration.VertexDeclarationRHI;
+				GraphicsPSOInit.BoundShaderState.PixelShaderRHI = PixelShader.GetPixelShader();
+				GraphicsPSOInit.BoundShaderState.VertexShaderRHI = VertexShader.GetVertexShader();
+				SetGraphicsPipelineState(RHICmdList, GraphicsPSOInit, 0);
+
+				SetShaderParameters(RHICmdList, PixelShader, PixelShader.GetPixelShader(), PassParameters->PS);
+				SetShaderParameters(RHICmdList, VertexShader, VertexShader.GetVertexShader(), PassParameters->VS);
+
+				RHICmdList.DrawPrimitiveIndirect(PassParameters->IndirectDrawParameter->GetIndirectRHICallBuffer(), 0);
+
+				RHICmdList.SetScissorRect(false, 0, 0, 0, 0);
+			});
+	}
+	else
+	{
+		FVirtualShadowMapProjectionCompositePS::FParameters* PassParameters = GraphBuilder.AllocParameters<FVirtualShadowMapProjectionCompositePS::FParameters>();
+		PassParameters->InputShadowFactor = Input;
+		PassParameters->bModulateRGB = bModulateRGB;
+		PassParameters->RenderTargets[0] = FRenderTargetBinding(OutputShadowMaskTexture, ERenderTargetLoadAction::ELoad);
+
+		ValidateShaderParameters(PixelShader, *PassParameters);
+
+		FPixelShaderUtils::AddFullscreenPass(GraphBuilder,
+			View.ShaderMap,
+			RDG_EVENT_NAME("CompositeVirtualShadowMapMask"),
+			PixelShader,
+			PassParameters,
+			ScissorRect,
+			BlendState);
+	}
 }
 
 class FVirtualShadowMapProjectionCompositeFromMaskBitsPS : public FGlobalShader
diff --git a/Engine/Source/Runtime/Renderer/Private/VirtualShadowMaps/VirtualShadowMapProjection.h b/Engine/Source/Runtime/Renderer/Private/VirtualShadowMaps/VirtualShadowMapProjection.h
--- a/Engine/Source/Runtime/Renderer/Private/VirtualShadowMaps/VirtualShadowMapProjection.h
+++ b/Engine/Source/Runtime/Renderer/Private/VirtualShadowMaps/VirtualShadowMapProjection.h
@@ -12,6 +12,14 @@
 
 class FVirtualShadowMapClipmap;
 
+struct FTiledVSMProjection
+{
+	FRDGBufferRef DrawIndirectParametersBuffer;
+	FRDGBufferRef DispatchIndirectParametersBuffer;
+	FRDGBufferSRVRef TileListDataBufferSRV;
+	uint32 TileSize;
+};
+
 // Note: Must match the definitions in VirtualShadowMapPageManagement.usf!
 enum class EVirtualShadowMapProjectionInputType
 {
@@ -29,6 +37,8 @@ void RenderVirtualShadowMapProjection(
 	const FIntRect ScissorRect,
 	EVirtualShadowMapProjectionInputType InputType,
 	const TSharedPtr<FVirtualShadowMapClipmap>& Clipmap,
+	bool bModulateRGB,
+	FTiledVSMProjection* TiledVSMProjection,
 	FRDGTextureRef OutputShadowMaskTexture);
 
 void RenderVirtualShadowMapProjection(
@@ -62,6 +72,8 @@ void CompositeVirtualShadowMapMask(
 	const FIntRect ScissorRect,
 	const FRDGTextureRef Input,
 	bool bDirectionalLight,
+	bool bModulateRGB,
+	FTiledVSMProjection* TiledVSMProjection,
 	FRDGTextureRef OutputShadowMaskTexture);
 
 void CompositeVirtualShadowMapFromMaskBits(
