Only calculate scene downsample chain with log luma in alpha channel if using basic eye adaptation.
- Avoids running BasicEyeAdaptationSetup pass when using histogram eye adaptation combined with local exposure.

diff --git a/Engine/Source/Runtime/Renderer/Private/PostProcess/PostProcessing.cpp b/Engine/Source/Runtime/Renderer/Private/PostProcess/PostProcessing.cpp
--- a/Engine/Source/Runtime/Renderer/Private/PostProcess/PostProcessing.cpp
+++ b/Engine/Source/Runtime/Renderer/Private/PostProcess/PostProcessing.cpp
@@ -805,7 +805,7 @@ void AddPostProcessingPasses(
 		FSceneDownsampleChain SceneDownsampleChain;
 		if (bProduceSceneColorChain)
 		{
-			const bool bLogLumaInAlpha = true;
+			const bool bLogLumaInAlpha = bBasicEyeAdaptationEnabled;
 			SceneDownsampleChain.Init(
 				GraphBuilder, View,
 				EyeAdaptationParameters,
