Only build PhysicsFieldRayMarching shader which uses ShaderPrint, if ShaderPrint is supported.

diff --git a/Engine/Source/Runtime/Renderer/Private/PhysicsFieldRendering.cpp b/Engine/Source/Runtime/Renderer/Private/PhysicsFieldRendering.cpp
--- a/Engine/Source/Runtime/Renderer/Private/PhysicsFieldRendering.cpp
+++ b/Engine/Source/Runtime/Renderer/Private/PhysicsFieldRendering.cpp
@@ -74,7 +74,10 @@ class FPhysicsFieldRayMarchingCS : public FGlobalShader
 	END_SHADER_PARAMETER_STRUCT()
 
 public:
-	static bool ShouldCompilePermutation(const FGlobalShaderPermutationParameters& Parameters) { return GetMaxSupportedFeatureLevel(Parameters.Platform) == ERHIFeatureLevel::SM5; }
+	static bool ShouldCompilePermutation(const FGlobalShaderPermutationParameters& Parameters) 
+	{
+		return ShaderPrint::IsSupported(Parameters.Platform) && GetMaxSupportedFeatureLevel(Parameters.Platform) == ERHIFeatureLevel::SM5;
+	}
 	static void ModifyCompilationEnvironment(const FGlobalShaderPermutationParameters& Parameters, FShaderCompilerEnvironment& OutEnvironment)
 	{
 		// Skip optimization for avoiding long compilation time due to large UAV writes
