Fix a bug where Lumen reflections output is assigned to the wrong slot when the async code path is used

diff --git a/Engine/Source/Runtime/Renderer/Private/IndirectLightRendering.cpp b/Engine/Source/Runtime/Renderer/Private/IndirectLightRendering.cpp
--- a/Engine/Source/Runtime/Renderer/Private/IndirectLightRendering.cpp
+++ b/Engine/Source/Runtime/Renderer/Private/IndirectLightRendering.cpp
@@ -822,7 +822,7 @@ void FDeferredShadingSceneRenderer::DispatchAsyncLumenIndirectLightingWork(
 			{
 				if (GLumenReflectionsAsyncCompute != 0)
 				{
-					ViewOutputs.IndirectLightingTextures.Textures[2] = RenderLumenReflections(
+					ViewOutputs.IndirectLightingTextures.Textures[3] = RenderLumenReflections(
 						GraphBuilder,
 						View,
 						SceneTextures,
@@ -839,8 +839,8 @@ void FDeferredShadingSceneRenderer::DispatchAsyncLumenIndirectLightingWork(
 			else
 			{
 				const FRDGSystemTextures& SystemTextures = FRDGSystemTextures::Get(GraphBuilder);
-				ViewOutputs.IndirectLightingTextures.Textures[1] = SystemTextures.Black;
 				ViewOutputs.IndirectLightingTextures.Textures[2] = SystemTextures.Black;
+				ViewOutputs.IndirectLightingTextures.Textures[3] = SystemTextures.Black;
 			}
 
 			// Lumen needs its own depth history because things like Translucency velocities write to depth
