Register GlobalDistanceField.PageObjectGridBuffer buffer for use on async compute pipe as well as graphics pipe to resolve ensure when utilized by async Lumen passes

Resolves the following ensure:
Resource GlobalDistanceField.PageObjectGridBuffer is in external access mode and is valid for access on the following pipelines: Graphics, but is being used on the AsyncCompute pipe in pass Scene.DiffuseIndirectAndAO.LumenScreenProbeGather.UpdateRadianceCaches.TraceFromProbes Res=16x16.

diff --git a/Engine/Source/Runtime/Renderer/Private/GlobalDistanceField.cpp b/Engine/Source/Runtime/Renderer/Private/GlobalDistanceField.cpp
--- a/Engine/Source/Runtime/Renderer/Private/GlobalDistanceField.cpp
+++ b/Engine/Source/Runtime/Renderer/Private/GlobalDistanceField.cpp
@@ -2779,7 +2779,7 @@ void UpdateGlobalDistanceFieldVolume(
 
 			if (PageObjectGridBuffer)
 			{
-				GlobalDistanceFieldInfo.PageObjectGridBuffer = ConvertToExternalAccessBuffer(GraphBuilder, ExternalAccessQueue, PageObjectGridBuffer);
+				GlobalDistanceFieldInfo.PageObjectGridBuffer = ConvertToExternalAccessBuffer(GraphBuilder, ExternalAccessQueue, PageObjectGridBuffer, ERHIAccess::SRVMask, ERHIPipeline::All);
 			}
 
 			if (PageTableCombinedTexture)
