Propertly filter out shading models field in UMaterial according to a current shader platform

diff --git a/Engine/Source/Runtime/Engine/Classes/Materials/MaterialInterface.h b/Engine/Source/Runtime/Engine/Classes/Materials/MaterialInterface.h
--- a/Engine/Source/Runtime/Engine/Classes/Materials/MaterialInterface.h
+++ b/Engine/Source/Runtime/Engine/Classes/Materials/MaterialInterface.h
@@ -1048,6 +1048,9 @@ protected:
 
 	void UpdateMaterialRenderProxy(FMaterialRenderProxy& Proxy);
 
+	/** Filter out ShadingModels field to a shader platform settings */
+	static void FilterOutPlatformShadingModels(const FStaticShaderPlatform Platform, FMaterialShadingModelField& ShadingModels);
+
 	/**
 	 * Cached data generated from the material's expressions, may be nullptr
 	 * UMaterials should always have cached data
diff --git a/Engine/Source/Runtime/Engine/Private/Materials/Material.cpp b/Engine/Source/Runtime/Engine/Private/Materials/Material.cpp
--- a/Engine/Source/Runtime/Engine/Private/Materials/Material.cpp
+++ b/Engine/Source/Runtime/Engine/Private/Materials/Material.cpp
@@ -3675,6 +3675,12 @@ void UMaterial::PostLoad()
 		bUseFullPrecision_DEPRECATED = false;
 	}
 
+	if (!GIsEditor)
+	{
+		// Filter out ShadingModels field to a current platform settings
+		FilterOutPlatformShadingModels(GMaxRHIShaderPlatform, ShadingModels);
+	}
+
 #if WITH_EDITOR
 	// Create exec flow expressions, if needed
 	CreateExecutionFlowExpressions();
diff --git a/Engine/Source/Runtime/Engine/Private/Materials/MaterialInstance.cpp b/Engine/Source/Runtime/Engine/Private/Materials/MaterialInstance.cpp
--- a/Engine/Source/Runtime/Engine/Private/Materials/MaterialInstance.cpp
+++ b/Engine/Source/Runtime/Engine/Private/Materials/MaterialInstance.cpp
@@ -2926,6 +2926,12 @@ void UMaterialInstance::PostLoad()
 	}
 #endif // WITH_EDITORONLY_DATA
 
+	if (!GIsEditor)
+	{
+		// Filter out ShadingModels field to a current platform settings
+		FilterOutPlatformShadingModels(GMaxRHIShaderPlatform, ShadingModels);
+	}
+
 #if WITH_EDITOR
 	UpdateCachedData();
 #endif
diff --git a/Engine/Source/Runtime/Engine/Private/Materials/MaterialInterface.cpp b/Engine/Source/Runtime/Engine/Private/Materials/MaterialInterface.cpp
--- a/Engine/Source/Runtime/Engine/Private/Materials/MaterialInterface.cpp
+++ b/Engine/Source/Runtime/Engine/Private/Materials/MaterialInterface.cpp
@@ -1486,6 +1486,25 @@ void UMaterialInterface::EnsureIsComplete()
 #endif
 }
 
+void UMaterialInterface::FilterOutPlatformShadingModels(const FStaticShaderPlatform Platform, FMaterialShadingModelField& ShadingModels)
+{
+	if (ShadingModels.CountShadingModels() > 1 && !AllowPerPixelShadingModels(Platform))
+	{
+		ShadingModels = FMaterialShadingModelField(ShadingModels.GetFirstShadingModel());
+	}
+
+	uint32 ShadingModelsMask = GetPlatformShadingModelsMask(Platform);
+	if (ShadingModelsMask != 0xFFFFFFFF)
+	{
+		uint16 FilteredShadingModels = (ShadingModels.GetShadingModelField() & ShadingModelsMask);
+		ShadingModels.SetShadingModelField(FilteredShadingModels);
+		if (!ShadingModels.IsValid())
+		{
+			ShadingModels.AddShadingModel(MSM_DefaultLit);
+		}
+	}
+}
+
 #if WITH_EDITORONLY_DATA
 
 namespace MaterialInterface
diff --git a/Engine/Source/Runtime/Renderer/Private/MobileBasePass.cpp b/Engine/Source/Runtime/Renderer/Private/MobileBasePass.cpp
--- a/Engine/Source/Runtime/Renderer/Private/MobileBasePass.cpp
+++ b/Engine/Source/Runtime/Renderer/Private/MobileBasePass.cpp
@@ -578,29 +578,15 @@ FMobileBasePassMeshProcessor::FMobileBasePassMeshProcessor(
 	, PassDrawRenderState(InDrawRenderState)
 	, TranslucencyPassType(InTranslucencyPassType)
 	, Flags(InFlags)
-	, ShadingModelsMask(GetPlatformShadingModelsMask(GetFeatureLevelShaderPlatform(InFeatureLevel)))
 	, bTranslucentBasePass(InTranslucencyPassType != ETranslucencyPass::TPT_MAX)
 	, bUsesDeferredShading(!bTranslucentBasePass && IsMobileDeferredShadingEnabled(GetFeatureLevelShaderPlatform(InFeatureLevel)))
 {
 }
 
-FMaterialShadingModelField FMobileBasePassMeshProcessor::FilterShadingModelsMask(const FMaterialShadingModelField& ShadingModels) const
-{
-	uint16 ShadingModelField = ShadingModels.GetShadingModelField();
-	uint16 FilteredShadingModels = (ShadingModelField & ShadingModelsMask);
-	FMaterialShadingModelField NewShadingModelField;
-	NewShadingModelField.SetShadingModelField(FilteredShadingModels);
-	if (!NewShadingModelField.IsValid())
-	{
-		NewShadingModelField.AddShadingModel(MSM_DefaultLit);
-	}
-	return NewShadingModelField;
-}
-
 bool FMobileBasePassMeshProcessor::TryAddMeshBatch(const FMeshBatch& RESTRICT MeshBatch, uint64 BatchElementMask, const FPrimitiveSceneProxy* RESTRICT PrimitiveSceneProxy, int32 StaticMeshId, const FMaterialRenderProxy& MaterialRenderProxy, const FMaterial& Material)
 {
 	const EBlendMode BlendMode = Material.GetBlendMode();
-	const FMaterialShadingModelField ShadingModels = FilterShadingModelsMask(Material.GetShadingModels());
+	const FMaterialShadingModelField ShadingModels = Material.GetShadingModels();
 	const bool bIsTranslucent = IsTranslucentBlendMode(BlendMode);
 	const bool bUsesWaterMaterial = ShadingModels.HasShadingModel(MSM_SingleLayerWater); // Water goes into the translucent pass
 	const bool bCanReceiveCSM = ((Flags & EFlags::CanReceiveCSM) == EFlags::CanReceiveCSM);
diff --git a/Engine/Source/Runtime/Renderer/Private/MobileBasePassRendering.h b/Engine/Source/Runtime/Renderer/Private/MobileBasePassRendering.h
--- a/Engine/Source/Runtime/Renderer/Private/MobileBasePassRendering.h
+++ b/Engine/Source/Runtime/Renderer/Private/MobileBasePassRendering.h
@@ -478,7 +478,6 @@ private:
 
 	const ETranslucencyPass::Type TranslucencyPassType;
 	const EFlags Flags;
-	uint16 ShadingModelsMask;
 	const bool bTranslucentBasePass;
 	const bool bUsesDeferredShading;
 };
