Change Velocity ShouldCompile to always compile for default material.
Wasn't compiling default material when base pass velocity was enabled, but it's still required as fallback for transparents writing velocity.
Fixes issue in cook on the fly where shader map is rewritten by request for shader that shouldn't be compiled.

diff --git a/Engine/Source/Runtime/Renderer/Private/VelocityRendering.cpp b/Engine/Source/Runtime/Renderer/Private/VelocityRendering.cpp
--- a/Engine/Source/Runtime/Renderer/Private/VelocityRendering.cpp
+++ b/Engine/Source/Runtime/Renderer/Private/VelocityRendering.cpp
@@ -114,12 +114,12 @@ public:
 		 */
 		const bool bIsSeparateVelocityPassRequired = 
 			!FVelocityRendering::BasePassCanOutputVelocity(Parameters.Platform) &&
-			(bIsDefault || bIsMasked || bIsOpaqueAndTwoSided || bMayModifyMeshes);
+			(bIsMasked || bIsOpaqueAndTwoSided || bMayModifyMeshes);
 
 		// The material may explicitly request that it be rendered into the translucent velocity pass.
 		const bool bIsSeparateVelocityPassRequiredByMaterial = Parameters.MaterialParameters.bIsTranslucencyWritingVelocity;
 
-		return bHasPlatformSupport && (bIsSeparateVelocityPassRequired || bIsSeparateVelocityPassRequiredByMaterial);
+		return bHasPlatformSupport && (bIsDefault || bIsSeparateVelocityPassRequired || bIsSeparateVelocityPassRequiredByMaterial);
 	}
 
 	FVelocityVS() = default;
