Only build shaders that use ShaderPrint if ShaderPrint is supported.

diff --git a/Engine/Source/Runtime/Renderer/Private/GlobalDistanceField.cpp b/Engine/Source/Runtime/Renderer/Private/GlobalDistanceField.cpp
--- a/Engine/Source/Runtime/Renderer/Private/GlobalDistanceField.cpp
+++ b/Engine/Source/Runtime/Renderer/Private/GlobalDistanceField.cpp
@@ -1838,7 +1838,7 @@ class FGlobalDistanceFieldDebugCS : public FGlobalShader
 
 	static bool ShouldCompilePermutation(const FGlobalShaderPermutationParameters& Parameters)
 	{
-		return ShouldCompileDistanceFieldShaders(Parameters.Platform);
+		return ShouldCompileDistanceFieldShaders(Parameters.Platform) && ShaderPrint::IsSupported(Parameters.Platform);
 	}
 
 	static int32 GetGroupSize()
@@ -2826,23 +2826,25 @@ void UpdateGlobalDistanceFieldVolume(
 					GroupSize);
 			}
 
+			if (ShaderPrint::IsEnabled(View.ShaderPrintData))
+			{
+				FRDGBufferRef PageFreeListAllocatorBuffer = GraphBuilder.RegisterExternalBuffer(GlobalDistanceFieldInfo.PageFreeListAllocatorBuffer);
 
-			FRDGBufferRef PageFreeListAllocatorBuffer = GraphBuilder.RegisterExternalBuffer(GlobalDistanceFieldInfo.PageFreeListAllocatorBuffer);
-
-			FGlobalDistanceFieldDebugCS::FParameters* PassParameters = GraphBuilder.AllocParameters<FGlobalDistanceFieldDebugCS::FParameters>();
-			ShaderPrint::SetParameters(GraphBuilder, View.ShaderPrintData, PassParameters->ShaderPrintUniformBuffer);
-			PassParameters->GlobalDistanceFieldPageFreeListAllocatorBuffer = GraphBuilder.CreateSRV(PageFreeListAllocatorBuffer, PF_R32_UINT);
-			PassParameters->GlobalDistanceFieldMaxPageNum = View.GlobalDistanceFieldInfo.ParameterData.MaxPageNum;
-			PassParameters->PageStatsBuffer = GraphBuilder.CreateSRV(PageStatsBuffer);
+				FGlobalDistanceFieldDebugCS::FParameters* PassParameters = GraphBuilder.AllocParameters<FGlobalDistanceFieldDebugCS::FParameters>();
+				ShaderPrint::SetParameters(GraphBuilder, View.ShaderPrintData, PassParameters->ShaderPrintUniformBuffer);
+				PassParameters->GlobalDistanceFieldPageFreeListAllocatorBuffer = GraphBuilder.CreateSRV(PageFreeListAllocatorBuffer, PF_R32_UINT);
+				PassParameters->GlobalDistanceFieldMaxPageNum = View.GlobalDistanceFieldInfo.ParameterData.MaxPageNum;
+				PassParameters->PageStatsBuffer = GraphBuilder.CreateSRV(PageStatsBuffer);
 
-			auto ComputeShader = View.ShaderMap->GetShader<FGlobalDistanceFieldDebugCS>();
+				auto ComputeShader = View.ShaderMap->GetShader<FGlobalDistanceFieldDebugCS>();
 
-			FComputeShaderUtils::AddPass(
-				GraphBuilder,
-				RDG_EVENT_NAME("GlobalDistanceFieldDebug"),
-				ComputeShader,
-				PassParameters,
-				FIntVector(1, 1, 1));
+				FComputeShaderUtils::AddPass(
+					GraphBuilder,
+					RDG_EVENT_NAME("GlobalDistanceFieldDebug"),
+					ComputeShader,
+					PassParameters,
+					FIntVector(1, 1, 1));
+			}
 		}
 	}
 
diff --git a/Engine/Source/Runtime/Renderer/Private/HairStrands/HairStrandsDebug.cpp b/Engine/Source/Runtime/Renderer/Private/HairStrands/HairStrandsDebug.cpp
--- a/Engine/Source/Runtime/Renderer/Private/HairStrands/HairStrandsDebug.cpp
+++ b/Engine/Source/Runtime/Renderer/Private/HairStrands/HairStrandsDebug.cpp
@@ -242,7 +242,10 @@ class FHairPrintLODInfoCS : public FGlobalShader
 		END_SHADER_PARAMETER_STRUCT()
 
 public:
-	static bool ShouldCompilePermutation(const FGlobalShaderPermutationParameters& Parameters) { return IsHairStrandsSupported(EHairStrandsShaderType::All, Parameters.Platform); }
+	static bool ShouldCompilePermutation(const FGlobalShaderPermutationParameters& Parameters) 
+	{
+		return IsHairStrandsSupported(EHairStrandsShaderType::All, Parameters.Platform) && ShaderPrint::IsSupported(Parameters.Platform);
+	}
 	static void ModifyCompilationEnvironment(const FGlobalShaderPermutationParameters& Parameters, FShaderCompilerEnvironment& OutEnvironment)
 	{
 		// Skip optimization for avoiding long compilation time due to large UAV writes
diff --git a/Engine/Source/Runtime/Renderer/Private/HairStrands/HairStrandsTile.cpp b/Engine/Source/Runtime/Renderer/Private/HairStrands/HairStrandsTile.cpp
--- a/Engine/Source/Runtime/Renderer/Private/HairStrands/HairStrandsTile.cpp
+++ b/Engine/Source/Runtime/Renderer/Private/HairStrands/HairStrandsTile.cpp
@@ -272,9 +272,9 @@ class FHairStrandsTileDebugPrintPassCS : public FGlobalShader
 		SHADER_PARAMETER_RDG_UNIFORM_BUFFER(FHairStrandsViewUniformParameters, HairStrands)
 	END_SHADER_PARAMETER_STRUCT()
 
-		static bool ShouldCompilePermutation(const FGlobalShaderPermutationParameters& Parameters)
+	static bool ShouldCompilePermutation(const FGlobalShaderPermutationParameters& Parameters)
 	{
-		return GetMaxSupportedFeatureLevel(Parameters.Platform) >= ERHIFeatureLevel::SM5; //TODO
+		return GetMaxSupportedFeatureLevel(Parameters.Platform) >= ERHIFeatureLevel::SM5 && ShaderPrint::IsSupported(Parameters.Platform); //TODO
 	}
 
 	static void ModifyCompilationEnvironment(const FGlobalShaderPermutationParameters& Parameters, FShaderCompilerEnvironment& OutEnvironment)
diff --git a/Engine/Source/Runtime/Renderer/Private/VolumetricCloudRendering.cpp b/Engine/Source/Runtime/Renderer/Private/VolumetricCloudRendering.cpp
--- a/Engine/Source/Runtime/Renderer/Private/VolumetricCloudRendering.cpp
+++ b/Engine/Source/Runtime/Renderer/Private/VolumetricCloudRendering.cpp
@@ -1278,7 +1278,10 @@ class FDrawDebugCloudShadowCS : public FGlobalShader
 	END_SHADER_PARAMETER_STRUCT()
 
 public:
-	static bool ShouldCompilePermutation(const FGlobalShaderPermutationParameters& Parameters) { return IsVolumetricCloudMaterialSupported(Parameters.Platform); }
+	static bool ShouldCompilePermutation(const FGlobalShaderPermutationParameters& Parameters)
+	{
+		return IsVolumetricCloudMaterialSupported(Parameters.Platform) && ShaderPrint::IsSupported(Parameters.Platform);
+	}
 	static void ModifyCompilationEnvironment(const FGlobalShaderPermutationParameters& Parameters, FShaderCompilerEnvironment& OutEnvironment)
 	{
 		FGlobalShader::ModifyCompilationEnvironment(Parameters, OutEnvironment);
