Fixed incorrect local-light shadow occlusion culling for VSM that causes caching of empty pages.

diff --git a/Engine/Source/Runtime/Renderer/Private/ShadowSetup.cpp b/Engine/Source/Runtime/Renderer/Private/ShadowSetup.cpp
--- a/Engine/Source/Runtime/Renderer/Private/ShadowSetup.cpp
+++ b/Engine/Source/Runtime/Renderer/Private/ShadowSetup.cpp
@@ -3908,7 +3908,8 @@ void FSceneRenderer::InitProjectedShadowVisibility()
 
 					bool bShadowIsOccluded = false;
 
-					if (!View.bIgnoreExistingQueries && View.State)
+					// Skip occlusion result if using VSM - can result in incorrectly cached pages otherwise.
+					if (!View.bIgnoreExistingQueries && View.State && !ProjectedShadowInfo.HasVirtualShadowMap())
 					{
 						// Check if the shadow is occluded.
 						bShadowIsOccluded =
