LightRendering.cpp cleanup

diff --git a/Engine/Source/Runtime/Renderer/Private/LightRendering.cpp b/Engine/Source/Runtime/Renderer/Private/LightRendering.cpp
--- a/Engine/Source/Runtime/Renderer/Private/LightRendering.cpp
+++ b/Engine/Source/Runtime/Renderer/Private/LightRendering.cpp
@@ -1179,30 +1179,31 @@ void FDeferredShadingSceneRenderer::RenderLights(
 			//#dxr_todo: support multiview for the batching case
 			const bool bDoShadowBatching = (bDoShadowDenoisingBatching || MaxRTShadowBatchSize > 1) && Views.Num() == 1;
 
-			// Optimisations: batches all shadow ray tracing denoising. Definitely could be smarter to avoid high VGPR pressure if this entire
-			// function was converted to render graph, and want least intrusive change as possible. So right not it trades render target memory pressure
+			// Optimizations: batches all shadow ray tracing denoising. Definitely could be smarter to avoid high VGPR pressure if this entire
+			// function was converted to render graph, and want least intrusive change as possible. So right now it trades render target memory pressure
 			// for denoising perf.
 			if (RHI_RAYTRACING && bDoShadowBatching)
 			{
 				const uint32 ViewIndex = 0;
 				FViewInfo& View = Views[ViewIndex];
 
+				const int32 NumShadowedLights = LumenLightStart - UnbatchedLightStart;
 				// Allocate PreprocessedShadowMaskTextures once so QueueTextureExtraction can deferred write.
 				{
 					if (!View.bStatePrevViewInfoIsReadOnly)
 					{
 						View.ViewState->PrevFrameViewInfo.ShadowHistories.Empty();
-						View.ViewState->PrevFrameViewInfo.ShadowHistories.Reserve(SortedLights.Num());
+						View.ViewState->PrevFrameViewInfo.ShadowHistories.Reserve(NumShadowedLights);
 					}
 
-					PreprocessedShadowMaskTextures.SetNum(SortedLights.Num());
+					PreprocessedShadowMaskTextures.SetNum(NumShadowedLights);
 				}
 
-				PreprocessedShadowMaskTextures.SetNum(SortedLights.Num());
+				PreprocessedShadowMaskTextures.SetNum(NumShadowedLights);
 
 				if (HairStrands::HasViewHairStrandsData(View))
 				{ 
-					PreprocessedShadowMaskSubPixelTextures.SetNum(SortedLights.Num());
+					PreprocessedShadowMaskSubPixelTextures.SetNum(NumShadowedLights);
 				}
 			} // if (RHI_RAYTRACING)
 
@@ -1220,8 +1221,6 @@ void FDeferredShadingSceneRenderer::RenderLights(
 				const FVisibleLightInfo& VisibleLightInfo = VisibleLightInfos[LightSceneInfo.Id];
 				const FLightOcclusionType OcclusionType = GetLightOcclusionType(LightSceneProxy);
 
-				// Note: Skip shadow mask generation for rect light if direct illumination is computed
-				//		 stochastically (rather than analytically + shadow mask)
 				const bool bDrawShadows = SortedLightInfo.SortKey.Fields.bShadowed;
 				const bool bDrawLightFunction = SortedLightInfo.SortKey.Fields.bLightFunction;
 				const bool bDrawPreviewIndicator = ViewFamily.EngineShowFlags.PreviewShadowsIndicator && !LightSceneInfo.IsPrecomputedLightingValid() && LightSceneProxy.HasStaticShadowing();
@@ -1292,9 +1291,8 @@ void FDeferredShadingSceneRenderer::RenderLights(
 							bDenoiserCompatible &&
 							SortedLightInfo.SortKey.Fields.bShadowed;
 
-						// determine if this light doesn't yet have a precomuted shadow and execute a batch to amortize costs if one is needed
-						if (
-							RHI_RAYTRACING &&
+						// determine if this light doesn't yet have a precomputed shadow and execute a batch to amortize costs if one is needed
+						if (RHI_RAYTRACING &&
 							bWantsBatchedShadow &&
 							(PreprocessedShadowMaskTextures.Num() == 0 || !PreprocessedShadowMaskTextures[LightIndex - UnbatchedLightStart]))
 						{
@@ -1345,21 +1343,25 @@ void FDeferredShadingSceneRenderer::RenderLights(
 								}
 							}; // QuickOffDenoisingBatch
 
-							// Ray trace shadows of light that needs, and quick off denoising batch.
+							// Ray trace shadows of lights, and quick off denoising batch.
 							for (int32 LightBatchIndex = LightIndex; LightBatchIndex < LumenLightStart; LightBatchIndex++)
 							{
 								const FSortedLightSceneInfo& BatchSortedLightInfo = SortedLights[LightBatchIndex];
 								const FLightSceneInfo& BatchLightSceneInfo = *BatchSortedLightInfo.LightSceneInfo;
 
-								// Denoiser do not support texture rect light important sampling.
+								// Denoiser does not support texture rect light important sampling.
 								const bool bBatchDrawShadows = BatchSortedLightInfo.SortKey.Fields.bShadowed;
 
 								if (!bBatchDrawShadows)
+								{
 									continue;
+								}
 
 								const FLightOcclusionType BatchOcclusionType = GetLightOcclusionType(*BatchLightSceneInfo.Proxy);
 								if (BatchOcclusionType != FLightOcclusionType::Raytraced)
+								{
 									continue;
+								}
 
 								const bool bRequiresDenoiser = LightRequiresDenosier(BatchLightSceneInfo) && DenoiserMode > 0;
 
@@ -1498,7 +1500,7 @@ void FDeferredShadingSceneRenderer::RenderLights(
 						ScreenShadowMaskTexture = PreprocessedShadowMaskTextures[ShadowMaskIndex];
 						PreprocessedShadowMaskTextures[ShadowMaskIndex] = nullptr;
 
-						// Subp-ixel shadow for hair strands geometries
+						// Sub-pixel shadow for hair strands geometries
 						if (bUseHairLighting && ShadowMaskIndex < uint32(PreprocessedShadowMaskSubPixelTextures.Num()))
 						{
 							ScreenShadowMaskSubPixelTexture = PreprocessedShadowMaskSubPixelTextures[ShadowMaskIndex];
