Fix shader compilation issue on some platforms

diff --git a/Engine/Source/Runtime/Renderer/Private/HairStrands/HairStrandsTile.cpp b/Engine/Source/Runtime/Renderer/Private/HairStrands/HairStrandsTile.cpp
--- a/Engine/Source/Runtime/Renderer/Private/HairStrands/HairStrandsTile.cpp
+++ b/Engine/Source/Runtime/Renderer/Private/HairStrands/HairStrandsTile.cpp
@@ -274,7 +274,7 @@ class FHairStrandsTileDebugPrintPassCS : public FGlobalShader
 
 	static bool ShouldCompilePermutation(const FGlobalShaderPermutationParameters& Parameters)
 	{
-		return GetMaxSupportedFeatureLevel(Parameters.Platform) >= ERHIFeatureLevel::SM5 && ShaderPrint::IsSupported(Parameters.Platform); //TODO
+		return IsHairStrandsSupported(EHairStrandsShaderType::Strands, Parameters.Platform) && ShaderPrint::IsSupported(Parameters.Platform);
 	}
 
 	static void ModifyCompilationEnvironment(const FGlobalShaderPermutationParameters& Parameters, FShaderCompilerEnvironment& OutEnvironment)
