Enable r.Bloom.ApplyLocalExposure by default

diff --git a/Engine/Source/Runtime/Renderer/Private/PostProcess/PostProcessing.cpp b/Engine/Source/Runtime/Renderer/Private/PostProcess/PostProcessing.cpp
--- a/Engine/Source/Runtime/Renderer/Private/PostProcess/PostProcessing.cpp
+++ b/Engine/Source/Runtime/Renderer/Private/PostProcess/PostProcessing.cpp
@@ -83,8 +83,8 @@ TAutoConsoleVariable<float> CVarDepthOfFieldMaxSize(
 
 TAutoConsoleVariable<bool> CVarBloomApplyLocalExposure(
 	TEXT("r.Bloom.ApplyLocalExposure"),
-	false,
-	TEXT("Whether to apply local exposure when calculating bloom, default: false"),
+	true,
+	TEXT("Whether to apply local exposure when calculating bloom, default: true"),
 	ECVF_Scalability | ECVF_RenderThreadSafe);
 
 TAutoConsoleVariable<int32> CVarPostProcessingPropagateAlpha(
