Fix crash when accessing RayTracingDynamicGeometryCollection when ray tracing is disabled but r.RayTracing.Enable has changed.

diff --git a/Engine/Source/Runtime/Renderer/Private/DeferredShadingRenderer.cpp b/Engine/Source/Runtime/Renderer/Private/DeferredShadingRenderer.cpp
--- a/Engine/Source/Runtime/Renderer/Private/DeferredShadingRenderer.cpp
+++ b/Engine/Source/Runtime/Renderer/Private/DeferredShadingRenderer.cpp
@@ -2192,7 +2192,7 @@ void FDeferredShadingSceneRenderer::Render(FRDGBuilder& GraphBuilder)
 
 	GRayTracingGeometryManager.Tick();
 
-	if (bHasRayTracingEnableChanged)
+	if ((GetRayTracingMode() == ERayTracingMode::Dynamic) && bHasRayTracingEnableChanged)
 	{
 		Scene->GetRayTracingDynamicGeometryCollection()->Clear();
 	}
