Distance Field - fix condition when logging I/O error

#rnx

diff --git a/Engine/Source/Runtime/Renderer/Private/DistanceFieldStreaming.cpp b/Engine/Source/Runtime/Renderer/Private/DistanceFieldStreaming.cpp
--- a/Engine/Source/Runtime/Renderer/Private/DistanceFieldStreaming.cpp
+++ b/Engine/Source/Runtime/Renderer/Private/DistanceFieldStreaming.cpp
@@ -593,7 +593,7 @@ void FDistanceFieldSceneData::AsyncUpdate(FDistanceFieldAsyncUpdateParameters&&
 			ReadRequests.Add(MoveTemp(ReadRequest));
 		}
 		FBulkDataRequest::EStatus Status = Batch.Issue();
-		UE_CLOG(Status == FBulkDataRequest::EStatus::Ok, LogDistanceField, Error, TEXT("Failed to issue bulk data I/O request"));
+		UE_CLOG(Status != FBulkDataRequest::EStatus::Ok, LogDistanceField, Error, TEXT("Failed to issue bulk data I/O request"));
 	}
 #endif
 }
