Volumetric cloud - new empty space skipping technique.
Basically evalaute non 0 material in froxel from the camera 32x32 cells on screen and 64 slices in depth (from froxel postiion and radius). The tracing is then moved to the start of the first slice to resum regular tracing.
This removes 1/4 to 1/3 of the cost in some tests.
Split screen has been tested and is working.
Also contains Cloud samples debug change from ue5/maim 22483386
Disabled for now, for Daniel to try and use it.

diff --git a/Engine/Shaders/Private/MaterialTemplate.ush b/Engine/Shaders/Private/MaterialTemplate.ush
--- a/Engine/Shaders/Private/MaterialTemplate.ush
+++ b/Engine/Shaders/Private/MaterialTemplate.ush
@@ -460,6 +460,9 @@ struct FMaterialPixelParameters
 	float CloudSampleNormAltitudeInLayer;
 	float4 VolumeSampleConservativeDensity;
 	float ShadowSampleDistance;
+
+	float3 CloudEmptySpaceSkippingSphereCenterWorldPosition;
+	float  CloudEmptySpaceSkippingSphereRadius;
 #endif
 
 #if TEMPLATE_USES_STRATA
@@ -1788,6 +1791,24 @@ float MaterialExpressionVolumeSampleShadowSampleDistance(FMaterialPixelParameter
 #endif
 }
 
+float3 MaterialExpressionCloudEmptySpaceSkippingSphereCenterWorldPosition(FMaterialPixelParameters Parameters)
+{
+#if CLOUD_LAYER_PIXEL_SHADER
+	return Parameters.CloudEmptySpaceSkippingSphereCenterWorldPosition;
+#else
+	return 0.0f;
+#endif
+}
+
+float MaterialExpressionCloudEmptySpaceSkippingSphereRadius(FMaterialPixelParameters Parameters)
+{
+#if CLOUD_LAYER_PIXEL_SHADER
+	return Parameters.CloudEmptySpaceSkippingSphereRadius;
+#else
+	return 0.0f;
+#endif
+}
+
 // We use a forward declaration towards the real implementation of functions we need. This is always defined for base passes (opaque and translucent)
 #if defined(IS_BASE_PASS) && IS_MOBILE_BASE_PASS==0
 float3 GetSkyLightReflectionSupportingBlend(float3 ReflectionVector, float Roughness, out float OutSkyAverageBrightness);
diff --git a/Engine/Shaders/Private/VolumetricCloud.usf b/Engine/Shaders/Private/VolumetricCloud.usf
--- a/Engine/Shaders/Private/VolumetricCloud.usf
+++ b/Engine/Shaders/Private/VolumetricCloud.usf
@@ -27,6 +27,10 @@
 #define SHADER_RENDERVIEW_PS 0
 #endif
 
+#ifndef SHADER_EMPTY_SPACE_SKIPPING_CS
+#define SHADER_EMPTY_SPACE_SKIPPING_CS 0
+#endif
+
 #ifndef VIRTUAL_SHADOW_MAP
 #define VIRTUAL_SHADOW_MAP 0
 #endif
@@ -35,6 +39,14 @@
 #define CLOUD_SAMPLE_LOCAL_LIGHTS 0
 #endif
 
+#ifndef EMPTY_SPACE_SKIPPING_DEBUG
+#define EMPTY_SPACE_SKIPPING_DEBUG 0
+#endif
+
+#ifndef MATERIAL_VOLUMETRIC_CLOUD_EMPTY_SPACE_SKIPPING_OUTPUT
+#define MATERIAL_VOLUMETRIC_CLOUD_EMPTY_SPACE_SKIPPING_OUTPUT 0
+#endif
+
 #ifndef CLOUD_CONSERVATIVE_DENSITY_DEBUG
 #define CLOUD_CONSERVATIVE_DENSITY_DEBUG 0
 #endif
@@ -43,17 +55,21 @@
 #define CLOUD_DEBUG_SAMPLES 0
 #endif
 
+#ifndef CLOUD_SAMPLE_COUNT_DEBUG_MODE
+#define CLOUD_SAMPLE_COUNT_DEBUG_MODE 0
+#endif
+
 // Draw debug sample positions
-#if CLOUD_DEBUG_SAMPLES
+#if CLOUD_DEBUG_SAMPLES || EMPTY_SPACE_SKIPPING_DEBUG
 #include "ShaderPrint.ush"
 #endif
 
-#if defined(SHADER_MAINVS) || defined(SHADER_SHADOW_PS) || SHADER_RENDERVIEW_PS || SHADER_RENDERVIEW_CS
+#if defined(SHADER_MAINVS) || defined(SHADER_SHADOW_PS) || SHADER_RENDERVIEW_PS || SHADER_RENDERVIEW_CS || SHADER_EMPTY_SPACE_SKIPPING_CS
 
 #include "/Engine/Generated/Material.ush"
 #include "VolumetricCloudMaterialPixelCommon.ush"
 
-#if !SHADER_RENDERVIEW_CS
+#if !SHADER_RENDERVIEW_CS && !SHADER_EMPTY_SPACE_SKIPPING_CS
 #include "/Engine/Generated/VertexFactory.ush"
 #endif
 
@@ -61,7 +77,7 @@
 #include "ParticipatingMediaCommon.ush"
 
 
-#if defined(SHADER_SHADOW_PS) || SHADER_RENDERVIEW_PS || SHADER_RENDERVIEW_CS
+#if defined(SHADER_SHADOW_PS) || SHADER_RENDERVIEW_PS || SHADER_RENDERVIEW_CS || SHADER_EMPTY_SPACE_SKIPPING_CS
 
 float3 SampleExtinctionCoefficients(in FPixelMaterialInputs PixelMaterialInputs)
 {
@@ -230,7 +246,7 @@ bool RayIntersectSphereSolution(float3 RayOrigin, float3 RayDirection, float4 Sp
 // Single full screen triangle vertex shader
 ////////////////////////////////////////////////////////////
 
-#if SHADER_RENDERVIEW_PS || SHADER_RENDERVIEW_CS
+#if SHADER_RENDERVIEW_PS || SHADER_RENDERVIEW_CS || SHADER_EMPTY_SPACE_SKIPPING_CS
 
 #include "/Engine/Private/SkyAtmosphereCommon.ush"
 
@@ -265,6 +281,9 @@ bool RayIntersectSphereSolution(float3 RayOrigin, float3 RayDirection, float4 Sp
 #ifndef MATERIAL_VOLUMETRIC_ADVANCED_OVERRIDE_AMBIENT_OCCLUSION
 #define MATERIAL_VOLUMETRIC_ADVANCED_OVERRIDE_AMBIENT_OCCLUSION 0
 #endif
+#ifndef MATERIAL_VOLUMETRIC_CLOUD_EMPTY_SPACE_SKIPPING_OUTPUT
+#define MATERIAL_VOLUMETRIC_CLOUD_EMPTY_SPACE_SKIPPING_OUTPUT 0
+#endif
 
 
 #if CLOUD_SAMPLE_ATMOSPHERIC_LIGHT_SHADOWMAP
@@ -389,6 +408,16 @@ ParticipatingMediaPhaseContext SetupParticipatingMediaPhaseContext(float BasePha
 	return PMPC;
 }
 
+#define SHADER_RENDERVIEW_CS_WITH_EMPTY_SPACE_SKIPPING (SHADER_RENDERVIEW_CS && MATERIAL_VOLUMETRIC_CLOUD_EMPTY_SPACE_SKIPPING_OUTPUT)
+
+#if SHADER_RENDERVIEW_CS_WITH_EMPTY_SPACE_SKIPPING
+
+float2	StartTracingDistanceTextureResolution;
+Texture2D<float> StartTracingDistanceTexture;
+
+#endif // SHADER_RENDERVIEW_CS_WITH_EMPTY_SPACE_SKIPPING
+
+
 void MainCommon(in FMaterialPixelParameters MaterialParameters, in float4 SvPosition, inout float4 OutColor0, inout float2 OutColor1)
 {
 	ResolvedView = ResolveView();
@@ -641,7 +670,6 @@ void MainCommon(in FMaterialPixelParameters MaterialParameters, in float4 SvPosi
 		const float3 GroundHemisphereLuminanceIsotropic = (2.0f * PI) * IsotropicPhase(); // Assumes the ground is uniform luminance to the cloud and solid angle is bottom hemisphere 2PI
 		const float3 GroundToCloudTransfertIsoScatter = GroundBrdfNdotL * GroundHemisphereLuminanceIsotropic;
 
-
 		//
 		// TRACING LOOP
 		//
@@ -658,8 +686,31 @@ void MainCommon(in FMaterialPixelParameters MaterialParameters, in float4 SvPosi
 			t = TMin + (float(Rand3DPCG16(int3(SvPosition.xy, View.StateFrameIndexMod8)).x) * rcp(65536.0)) * StepT;
 		}
 
+	#if SHADER_RENDERVIEW_CS_WITH_EMPTY_SPACE_SKIPPING
+		// In this case we cannot run variable step size to make sure all our samples are aligned and empty space skipping low resolution texture texel do nto become visible.
+		const float StepSizeOnZeroConservativeDensity = 1.0f;
+	#else
+		const float StepSizeOnZeroConservativeDensity = RenderVolumetricCloudParameters.StepSizeOnZeroConservativeDensity;
+	#endif
+
+		//
+		// Read from empty space skipping
+		//
+		uint StartStepCount = 0;
+	#if SHADER_RENDERVIEW_CS_WITH_EMPTY_SPACE_SKIPPING
+		const float StartDepthKm = StartTracingDistanceTexture[SvPositionToViewportUV(SvPosition) * StartTracingDistanceTextureResolution];
+		const float StartDepthCm = StartDepthKm * KILOMETER_TO_CENTIMETER;
+		if (t < StartDepthCm)
+		{
+			StartStepCount = uint(StartDepthCm / StepT);
+			t += float(StartStepCount) * StepSizeOnZeroConservativeDensity * StepT;
+		}
+		//OutColor0 = float4(StartDepthKm / 50.0f, 0.0f, 0.0f, 0.5f);
+		//return;
+	#endif
+
 		float3 S0 = 0.0f;
-		for (uint i = 0; i < IStepCount; ++i)
+		for (uint i = StartStepCount; i < IStepCount; ++i)
 		{
 			int ms;
 			float3	  SampleTranslatedWorldPosition = RayTranslatedWorldOriginKm + t * Raydir;
@@ -678,8 +729,8 @@ void MainCommon(in FMaterialPixelParameters MaterialParameters, in float4 SvPosi
 				++DebugZeroConservativeDensitySampleCount;
 				t += StepT;
 #else
-				i += RenderVolumetricCloudParameters.StepSizeOnZeroConservativeDensity - 1;
-				t += RenderVolumetricCloudParameters.StepSizeOnZeroConservativeDensity * StepT;
+				i += StepSizeOnZeroConservativeDensity - 1;
+				t += StepSizeOnZeroConservativeDensity * StepT;
 #endif
 
 			#if CLOUD_DEBUG_SAMPLES
@@ -1373,6 +1424,150 @@ void MainCS(uint2 DispatchThreadId : SV_DispatchThreadID)
 
 
 
+#if SHADER_EMPTY_SPACE_SKIPPING_CS
+
+float2	StartTracingDistanceTextureResolution;
+float	StartTracingSampleVolumeDepth;
+
+RWTexture2D<float> OutStartTracingDistanceTexture;
+
+struct FThreadData
+{
+	float FrontDepthKm;
+};
+groupshared FThreadData GroupData[THREADGROUP_SIZEX][THREADGROUP_SIZEY][THREADGROUP_SIZEZ];
+
+float GetDefaultDepth()
+{
+	return StartTracingSampleVolumeDepth * CENTIMETER_TO_KILOMETER;
+}
+
+FThreadData GetDefaultThreadData()
+{
+	FThreadData ThreadData;
+	ThreadData.FrontDepthKm = GetDefaultDepth();
+	return ThreadData;
+}
+
+FThreadData Reduce(FThreadData A, FThreadData B)
+{
+	FThreadData Out;
+	Out.FrontDepthKm = min(A.FrontDepthKm, B.FrontDepthKm);
+	return Out;
+}
+
+[numthreads(THREADGROUP_SIZEX, THREADGROUP_SIZEY, THREADGROUP_SIZEZ)]
+void MainCS(uint3 DispatchThreadId : SV_DispatchThreadID, uint3 GroupThreadId : SV_GroupThreadID)
+{
+#if MATERIAL_VOLUMETRIC_CLOUD_EMPTY_SPACE_SKIPPING_OUTPUT
+
+	ResolvedView = ResolveView();
+	const float SliceDepth = StartTracingSampleVolumeDepth * (1.0f / THREADGROUP_SIZEZ);
+
+	float2 PixelUV = float2(DispatchThreadId.x + 0.5, DispatchThreadId.y + 0.5) / StartTracingDistanceTextureResolution;
+	float4 SvPosition = float4(PixelUV * ResolvedView.ViewSizeAndInvSize.xy, 0.1, 1);
+
+	FMaterialPixelParameters MaterialParameters = MakeInitializedMaterialPixelParameters();
+	FPixelMaterialInputs PixelMaterialInputs;
+	CalcMaterialParameters(MaterialParameters, PixelMaterialInputs, SvPosition, true);
+
+	FCloudLayerParameters CloudLayerParams = GetCloudLayerParams(
+		RenderVolumetricCloudParameters.CloudLayerCenterKm, RenderVolumetricCloudParameters.PlanetRadiusKm,
+		RenderVolumetricCloudParameters.BottomRadiusKm, RenderVolumetricCloudParameters.TopRadiusKm);
+
+	const float SampleDistanceFromCamera = (float(GroupThreadId.z) + 0.5f) * SliceDepth;
+	const float3 Raydir =-MaterialParameters.CameraVector;
+	UpdateMaterialCloudParam(MaterialParameters, LWCAdd(ResolvedView.WorldCameraOrigin, LWCPromote(Raydir * SampleDistanceFromCamera)),
+		ResolvedView, CloudLayerParams, SampleDistanceFromCamera);
+
+	MaterialParameters.CloudEmptySpaceSkippingSphereCenterWorldPosition = LWCHackToFloat(MaterialParameters.AbsoluteWorldPosition); // This is the same as the WorldPOsitionNode but float3
+	MaterialParameters.CloudEmptySpaceSkippingSphereRadius = 0.5f * SliceDepth * 1.73f;	// 1.73 = sqrt(1*1+1*1+1*1), the lenght of the diagonal of a cube as compared to its unit edge length
+
+	float CloudEmptySpaceSkippingOutput0 = GetVolumetricCloudEmptySpaceSkippingOutput0(MaterialParameters);
+	GroupData[GroupThreadId.x][GroupThreadId.y][GroupThreadId.z].FrontDepthKm = CloudEmptySpaceSkippingOutput0 > 0.0f ? (float(GroupThreadId.z) + 0.0f) * SliceDepth * CENTIMETER_TO_KILOMETER : GetDefaultDepth();
+
+	GroupMemoryBarrierWithGroupSync();
+
+	// Reduction
+#if THREADGROUP_SIZEZ > 64
+#error Reduction does not support THREADGROUP_SIZEZ > 64
+#endif
+
+#if THREADGROUP_SIZEZ >= 64
+	if (GroupThreadId.z < 32)
+	{
+		GroupData[GroupThreadId.x][GroupThreadId.y][GroupThreadId.z] = Reduce(
+			GroupData[GroupThreadId.x][GroupThreadId.y][GroupThreadId.z],
+			GroupData[GroupThreadId.x][GroupThreadId.y][GroupThreadId.z + 32]);
+	}
+	GroupMemoryBarrierWithGroupSync();
+#endif
+
+#if THREADGROUP_SIZEZ >= 32
+	if (GroupThreadId.z < 16)
+	{
+		GroupData[GroupThreadId.x][GroupThreadId.y][GroupThreadId.z] = Reduce(
+			GroupData[GroupThreadId.x][GroupThreadId.y][GroupThreadId.z],
+			GroupData[GroupThreadId.x][GroupThreadId.y][GroupThreadId.z + 16]);
+	}
+	GroupMemoryBarrierWithGroupSync();
+#endif
+
+#if THREADGROUP_SIZEZ >= 16
+	if (GroupThreadId.z < 8)
+	{
+		GroupData[GroupThreadId.x][GroupThreadId.y][GroupThreadId.z] = Reduce(
+			GroupData[GroupThreadId.x][GroupThreadId.y][GroupThreadId.z],
+			GroupData[GroupThreadId.x][GroupThreadId.y][GroupThreadId.z + 8]);
+	}
+	GroupMemoryBarrierWithGroupSync();
+#endif
+
+#if THREADGROUP_SIZEZ >= 8
+	if (GroupThreadId.z < 4)
+	{
+		GroupData[GroupThreadId.x][GroupThreadId.y][GroupThreadId.z] = Reduce(
+			GroupData[GroupThreadId.x][GroupThreadId.y][GroupThreadId.z],
+			GroupData[GroupThreadId.x][GroupThreadId.y][GroupThreadId.z + 4]);
+	}
+	GroupMemoryBarrierWithGroupSync();
+#endif
+
+	// The smallest wave size is 4 on Mali G-71 hardware. So now we can do simple math operations without group sync.
+
+#if THREADGROUP_SIZEZ >= 4
+	if (GroupThreadId.z < 2)
+	{
+		GroupData[GroupThreadId.x][GroupThreadId.y][GroupThreadId.z] = Reduce(
+			GroupData[GroupThreadId.x][GroupThreadId.y][GroupThreadId.z],
+			GroupData[GroupThreadId.x][GroupThreadId.y][GroupThreadId.z + 2]);
+	}
+#endif
+
+	if (GroupThreadId.z < 1)
+	{
+		GroupData[GroupThreadId.x][GroupThreadId.y][GroupThreadId.z] = Reduce(
+			GroupData[GroupThreadId.x][GroupThreadId.y][GroupThreadId.z],
+			GroupData[GroupThreadId.x][GroupThreadId.y][GroupThreadId.z + 1]);
+
+		BRANCH
+		if (all(DispatchThreadId.xy < uint2(StartTracingDistanceTextureResolution.xy)))
+		{
+			OutStartTracingDistanceTexture[DispatchThreadId.xy] = GroupData[GroupThreadId.x][GroupThreadId.y][0].FrontDepthKm;
+		}
+	}
+
+#else // MATERIAL_VOLUMETRIC_CLOUD_EMPTY_SPACE_SKIPPING_OUTPUT
+
+	OutStartTracingDistanceTexture[DispatchThreadId.xy] = 0.0f;	// No empty space skipping
+
+#endif // MATERIAL_VOLUMETRIC_CLOUD_EMPTY_SPACE_SKIPPING_OUTPUT
+}
+
+#endif // SHADER_SHADOW_PS
+
+
+
 #if SHADER_SHADOW_PS
 
 struct CloudShadowTraceContext
diff --git a/Engine/Shaders/Private/VolumetricCloudMaterialPixelCommon.ush b/Engine/Shaders/Private/VolumetricCloudMaterialPixelCommon.ush
--- a/Engine/Shaders/Private/VolumetricCloudMaterialPixelCommon.ush
+++ b/Engine/Shaders/Private/VolumetricCloudMaterialPixelCommon.ush
@@ -61,6 +61,9 @@ void UpdateMaterialCloudParam(inout FMaterialPixelParameters MaterialParameters,
 	MaterialParameters.VolumeSampleConservativeDensity = GetVolumetricAdvancedMaterialOutput6(MaterialParameters);	// Evaluate conservative density
 #endif
 
+	MaterialParameters.CloudEmptySpaceSkippingSphereCenterWorldPosition = 0.0f;
+	MaterialParameters.CloudEmptySpaceSkippingSphereRadius = 0.0f;
+
 	MaterialParameters.ShadowSampleDistance = ShadowSampleDistance;
 }
 
diff --git a/Engine/Source/Runtime/Engine/Classes/Materials/MaterialExpressionVolumetricAdvancedMaterialInput.h b/Engine/Source/Runtime/Engine/Classes/Materials/MaterialExpressionVolumetricAdvancedMaterialInput.h
--- a/Engine/Source/Runtime/Engine/Classes/Materials/MaterialExpressionVolumetricAdvancedMaterialInput.h
+++ b/Engine/Source/Runtime/Engine/Classes/Materials/MaterialExpressionVolumetricAdvancedMaterialInput.h
@@ -21,4 +21,20 @@ class UMaterialExpressionVolumetricAdvancedMaterialInput : public UMaterialExpre
 	//~ End UMaterialExpression Interface
 };
 
+UCLASS()
+class UMaterialExpressionVolumetricCloudEmptySpaceSkippingInput : public UMaterialExpression
+{
+	GENERATED_UCLASS_BODY()
+
+	//~ Begin UMaterialExpression Interface
+#if WITH_EDITOR
+	virtual int32 Compile(class FMaterialCompiler* Compiler, int32 OutputIndex) override;
+	virtual void GetCaption(TArray<FString>& OutCaptions) const override;
+	virtual void GetConnectorToolTip(int32 InputIndex, int32 OutputIndex, TArray<FString>& OutToolTip) override;
+	virtual uint32 GetInputType(int32 InputIndex) override;
+	virtual bool GenerateHLSLExpression(FMaterialHLSLGenerator& Generator, UE::HLSLTree::FScope& Scope, int32 OutputIndex, UE::HLSLTree::FExpression const*& OutExpression) const override;
+#endif
+	//~ End UMaterialExpression Interface
+};
+
 
diff --git a/Engine/Source/Runtime/Engine/Classes/Materials/MaterialExpressionVolumetricAdvancedMaterialOutput.h b/Engine/Source/Runtime/Engine/Classes/Materials/MaterialExpressionVolumetricAdvancedMaterialOutput.h
--- a/Engine/Source/Runtime/Engine/Classes/Materials/MaterialExpressionVolumetricAdvancedMaterialOutput.h
+++ b/Engine/Source/Runtime/Engine/Classes/Materials/MaterialExpressionVolumetricAdvancedMaterialOutput.h
@@ -114,4 +114,31 @@ public:
 	//~ End UMaterialExpressionCustomOutput Interface
 };
 
+/** USed to help the cloud system to fast skip empty space areas when ray marching. */
+UCLASS(MinimalAPI, collapsecategories, hidecategories = Object)
+class UMaterialExpressionVolumetricCloudEmptySpaceSkippingOutput : public UMaterialExpressionCustomOutput
+{
+	GENERATED_UCLASS_BODY()
+
+	/** ContainsMatter must be 1 when the volume is occupied by any matter. This is for the tracing to later not miss it. Otherwise it can be 0 to accelerate the tracing by skipping that area. */
+	UPROPERTY(meta = (RequiredInput = "false", ToolTip = "Specify 0 if no matter (cloud or participating media) can be found within the area, otherwise should be set > 0."))
+	FExpressionInput ContainsMatter;
+
+public:
+#if WITH_EDITOR
+	//~ Begin UMaterialExpression Interface
+	virtual int32 Compile(class FMaterialCompiler* Compiler, int32 OutputIndex) override;
+	virtual void GetCaption(TArray<FString>& OutCaptions) const override;
+	virtual uint32 GetOutputType(int32 OutputIndex) override;
+	//~ End UMaterialExpression Interface
 
+	virtual bool GenerateHLSLExpression(FMaterialHLSLGenerator& Generator, UE::HLSLTree::FScope& Scope, int32 OutputIndex, UE::HLSLTree::FExpression const*& OutExpression) const override;
+	virtual UE::Shader::EValueType GetCustomOutputType(int32 OutputIndex) const override;
+#endif
+
+	//~ Begin UMaterialExpressionCustomOutput Interface
+	virtual int32 GetNumOutputs() const override;
+	virtual FString GetFunctionName() const override;
+	virtual FString GetDisplayName() const override;
+	//~ End UMaterialExpressionCustomOutput Interface
+};
diff --git a/Engine/Source/Runtime/Engine/Private/Materials/HLSLMaterialTranslator.cpp b/Engine/Source/Runtime/Engine/Private/Materials/HLSLMaterialTranslator.cpp
--- a/Engine/Source/Runtime/Engine/Private/Materials/HLSLMaterialTranslator.cpp
+++ b/Engine/Source/Runtime/Engine/Private/Materials/HLSLMaterialTranslator.cpp
@@ -1883,7 +1883,7 @@ void FHLSLMaterialTranslator::GetMaterialEnvironment(EShaderPlatform InPlatform,
 		OutEnvironment.SetDefine(TEXT("MATERIAL_SHADINGMODEL_UNLIT"), TEXT("1"));
 	}
 
-	if (Material->GetMaterialDomain() == MD_Volume ) // && Material->HasN)
+	if (Material->GetMaterialDomain() == MD_Volume)
 	{
 		TArray<const UMaterialExpressionVolumetricAdvancedMaterialOutput*> VolumetricAdvancedExpressions;
 		Material->GetMaterialInterface()->GetMaterial()->GetAllExpressionsOfType(VolumetricAdvancedExpressions);
@@ -1921,6 +1921,17 @@ void FHLSLMaterialTranslator::GetMaterialEnvironment(EShaderPlatform InPlatform,
 			OutEnvironment.SetDefine(TEXT("MATERIAL_VOLUMETRIC_ADVANCED_GROUND_CONTRIBUTION"),
 				VolumetricAdvancedNode->bGroundContribution ? TEXT("1") : TEXT("0"));
 		}
+
+		TArray<const UMaterialExpressionVolumetricCloudEmptySpaceSkippingOutput*> EmptySpaceSkippingOutputExpressions;
+		Material->GetMaterialInterface()->GetMaterial()->GetAllExpressionsOfType(EmptySpaceSkippingOutputExpressions);
+		if (EmptySpaceSkippingOutputExpressions.Num() > 0)
+		{
+			if (VolumetricAdvancedExpressions.Num() > 1)
+			{
+				UE_LOG(LogMaterial, Fatal, TEXT("Only a single UMaterialExpressionEmptySpaceSkippingOutput node is supported."));
+			}
+			OutEnvironment.SetDefine(TEXT("MATERIAL_VOLUMETRIC_CLOUD_EMPTY_SPACE_SKIPPING_OUTPUT"), TEXT("1"));
+		}
 	}
 
 	OutEnvironment.SetDefine(TEXT("MATERIAL_IS_STRATA"), bMaterialIsStrata ? TEXT("1") : TEXT("0"));
@@ -9477,6 +9488,16 @@ int32 FHLSLMaterialTranslator::GetVolumeSampleConservativeDensity()
 	return AddCodeChunk(MCT_Float4, TEXT("MaterialExpressionVolumeSampleConservativeDensity(Parameters)"));
 }
 
+int32 FHLSLMaterialTranslator::GetCloudEmptySpaceSkippingSphereCenterWorldPosition()
+{
+	return AddCodeChunk(MCT_Float3, TEXT("MaterialExpressionCloudEmptySpaceSkippingSphereCenterWorldPosition(Parameters)"));
+}
+
+int32 FHLSLMaterialTranslator::GetCloudEmptySpaceSkippingSphereRadius()
+{
+	return AddCodeChunk(MCT_Float1, TEXT("MaterialExpressionCloudEmptySpaceSkippingSphereRadius(Parameters)"));
+}
+
 int32 FHLSLMaterialTranslator::CustomPrimitiveData(int32 OutputIndex, EMaterialValueType Type)
 {
 	check(OutputIndex < FCustomPrimitiveData::NumCustomPrimitiveDataFloats);
diff --git a/Engine/Source/Runtime/Engine/Private/Materials/HLSLMaterialTranslator.h b/Engine/Source/Runtime/Engine/Private/Materials/HLSLMaterialTranslator.h
--- a/Engine/Source/Runtime/Engine/Private/Materials/HLSLMaterialTranslator.h
+++ b/Engine/Source/Runtime/Engine/Private/Materials/HLSLMaterialTranslator.h
@@ -1018,6 +1018,8 @@ protected:
 	virtual int32 GetCloudSampleNormAltitudeInLayer() override;
 	virtual int32 GetCloudSampleShadowSampleDistance() override;
 	virtual int32 GetVolumeSampleConservativeDensity() override;
+	virtual int32 GetCloudEmptySpaceSkippingSphereCenterWorldPosition() override;
+	virtual int32 GetCloudEmptySpaceSkippingSphereRadius() override;
 
 	virtual int32 CustomPrimitiveData(int32 OutputIndex, EMaterialValueType Type) override;
 
diff --git a/Engine/Source/Runtime/Engine/Private/Materials/MaterialExpressionHLSL.cpp b/Engine/Source/Runtime/Engine/Private/Materials/MaterialExpressionHLSL.cpp
--- a/Engine/Source/Runtime/Engine/Private/Materials/MaterialExpressionHLSL.cpp
+++ b/Engine/Source/Runtime/Engine/Private/Materials/MaterialExpressionHLSL.cpp
@@ -2263,6 +2263,27 @@ UE::Shader::EValueType UMaterialExpressionVolumetricAdvancedMaterialOutput::GetC
 	else return UE::Shader::EValueType::Void;
 }
 
+bool UMaterialExpressionVolumetricCloudEmptySpaceSkippingInput::GenerateHLSLExpression(FMaterialHLSLGenerator& Generator, UE::HLSLTree::FScope& Scope, int32 OutputIndex, UE::HLSLTree::FExpression const*& OutExpression) const
+{
+	check(false);	// TODO implement when the new compiler is enabled
+	return OutExpression != nullptr;
+}
+
+bool UMaterialExpressionVolumetricCloudEmptySpaceSkippingOutput::GenerateHLSLExpression(FMaterialHLSLGenerator& Generator, UE::HLSLTree::FScope& Scope, int32 OutputIndex, UE::HLSLTree::FExpression const*& OutExpression) const
+{
+	check(false);	// TODO implement when the new compiler is enabled
+	return OutExpression != nullptr;
+}
+
+UE::Shader::EValueType UMaterialExpressionVolumetricCloudEmptySpaceSkippingOutput::GetCustomOutputType(int32 OutputIndex) const
+{
+	if (OutputIndex == 0)
+	{
+		return UE::Shader::EValueType::Float1;
+	}
+	return UE::Shader::EValueType::Void;
+}
+
 bool UMaterialExpressionHairAttributes::GenerateHLSLExpression(FMaterialHLSLGenerator& Generator, UE::HLSLTree::FScope& Scope, int32 OutputIndex, UE::HLSLTree::FExpression const*& OutExpression) const
 {
 	using namespace UE::HLSLTree;
diff --git a/Engine/Source/Runtime/Engine/Private/Materials/MaterialExpressions.cpp b/Engine/Source/Runtime/Engine/Private/Materials/MaterialExpressions.cpp
--- a/Engine/Source/Runtime/Engine/Private/Materials/MaterialExpressions.cpp
+++ b/Engine/Source/Runtime/Engine/Private/Materials/MaterialExpressions.cpp
@@ -21360,6 +21360,154 @@ void UMaterialExpressionVolumetricAdvancedMaterialInput::GetCaption(TArray<FStri
 }
 #endif // WITH_EDITOR
 
+
+///////////////////////////////////////////////////////////////////////////////
+// UMaterialExpressionVolumetricCloudEmptySpaceSkippingOutput
+///////////////////////////////////////////////////////////////////////////////
+
+UMaterialExpressionVolumetricCloudEmptySpaceSkippingOutput::UMaterialExpressionVolumetricCloudEmptySpaceSkippingOutput(const FObjectInitializer& ObjectInitializer)
+	: Super(ObjectInitializer)
+{
+	// Structure to hold one-time initialization
+	struct FConstructorStatics
+	{
+		FText NAME_VolumetricCloudEmptySpaceSkippingOutput;
+		FConstructorStatics()
+			: NAME_VolumetricCloudEmptySpaceSkippingOutput(LOCTEXT("Volume", "Volume"))
+		{
+		}
+	};
+	static FConstructorStatics ConstructorStatics;
+
+#if WITH_EDITORONLY_DATA
+	MenuCategories.Add(ConstructorStatics.NAME_VolumetricCloudEmptySpaceSkippingOutput);
+#endif
+
+#if WITH_EDITOR
+	Outputs.Reset();
+#endif
+}
+
+#if WITH_EDITOR
+
+int32 UMaterialExpressionVolumetricCloudEmptySpaceSkippingOutput::Compile(class FMaterialCompiler* Compiler, int32 OutputIndex)
+{
+	int32 CodeInput = INDEX_NONE;
+
+	if (OutputIndex == 0)
+	{
+		CodeInput = ContainsMatter.IsConnected() ? ContainsMatter.Compile(Compiler) : Compiler->Constant(1.0f);
+	}
+
+	return Compiler->CustomOutput(this, OutputIndex, CodeInput);
+}
+
+void UMaterialExpressionVolumetricCloudEmptySpaceSkippingOutput::GetCaption(TArray<FString>& OutCaptions) const
+{
+	OutCaptions.Add(FString(TEXT("Volumetric Cloud Empty Space Skipping Output")));
+}
+
+uint32 UMaterialExpressionVolumetricCloudEmptySpaceSkippingOutput::GetOutputType(int32 OutputIndex)
+{
+	return MCT_Float1;
+}
+
+#endif // WITH_EDITOR
+
+int32 UMaterialExpressionVolumetricCloudEmptySpaceSkippingOutput::GetNumOutputs() const
+{
+	return 1;
+}
+
+FString UMaterialExpressionVolumetricCloudEmptySpaceSkippingOutput::GetFunctionName() const
+{
+	return TEXT("GetVolumetricCloudEmptySpaceSkippingOutput");
+}
+
+FString UMaterialExpressionVolumetricCloudEmptySpaceSkippingOutput::GetDisplayName() const
+{
+	return TEXT("Volumetric Cloud Empty Space Skipping Output");
+}
+
+
+///////////////////////////////////////////////////////////////////////////////
+// UMaterialExpressionVolumetricCloudEmptySpaceSkippingInput
+///////////////////////////////////////////////////////////////////////////////
+
+UMaterialExpressionVolumetricCloudEmptySpaceSkippingInput::UMaterialExpressionVolumetricCloudEmptySpaceSkippingInput(const FObjectInitializer& ObjectInitializer)
+	: Super(ObjectInitializer)
+{
+#if WITH_EDITORONLY_DATA
+	// Structure to hold one-time initialization
+	struct FConstructorStatics
+	{
+		FText NAME_UMaterialExpressionVolumetricCloudEmptySpaceSkippingInputInput;
+		FConstructorStatics()
+			: NAME_UMaterialExpressionVolumetricCloudEmptySpaceSkippingInputInput(LOCTEXT("Volume", "Volume"))
+		{
+		}
+	};
+	static FConstructorStatics ConstructorStatics;
+	MenuCategories.Add(ConstructorStatics.NAME_UMaterialExpressionVolumetricCloudEmptySpaceSkippingInputInput);
+
+	bShowOutputNameOnPin = true;
+
+	Outputs.Reset();
+	Outputs.Add(FExpressionOutput(TEXT("Sphere Center")));
+	Outputs.Add(FExpressionOutput(TEXT("Sphere Radius")));
+#endif
+}
+
+#if WITH_EDITOR
+int32 UMaterialExpressionVolumetricCloudEmptySpaceSkippingInput::Compile(class FMaterialCompiler* Compiler, int32 OutputIndex)
+{
+	if (OutputIndex == 0)
+	{
+		return Compiler->GetCloudEmptySpaceSkippingSphereCenterWorldPosition();
+	}
+	else if (OutputIndex == 1)
+	{
+		return Compiler->GetCloudEmptySpaceSkippingSphereRadius();
+	}
+
+	return Compiler->Errorf(TEXT("Invalid input parameter"));
+}
+
+void UMaterialExpressionVolumetricCloudEmptySpaceSkippingInput::GetCaption(TArray<FString>& OutCaptions) const
+{
+	OutCaptions.Add(TEXT("Volumetric Cloud Empty Space Skipping Input"));
+}
+
+void UMaterialExpressionVolumetricCloudEmptySpaceSkippingInput::GetConnectorToolTip(int32 InputIndex, int32 OutputIndex, TArray<FString>& OutToolTip)
+{
+	if (OutputIndex == 0)
+	{
+		OutToolTip.Add(TEXT("The position of the center of the spherical area considered for empty space skipping. When evaluating cloud for empty space skipping, this output returns the same as the WorldPosition node."));
+		return;
+	}
+	else if (OutputIndex == 1)
+	{
+		OutToolTip.Add(TEXT("The radius of the spherical area considered for empty space skipping."));
+		return;
+	}
+	Super::GetConnectorToolTip(InputIndex, OutputIndex, OutToolTip);
+}
+
+uint32 UMaterialExpressionVolumetricCloudEmptySpaceSkippingInput::GetInputType(int32 InputIndex)
+{
+	if (InputIndex == 0)
+	{
+		return MCT_Float3;
+	}
+	else if (InputIndex == 1)
+	{
+		return MCT_Float1;
+	}
+	return MCT_Unknown;
+}
+#endif // WITH_EDITOR
+
+
 ///////////////////////////////////////////////////////////////////////////////
 // UMaterialExpressionReflectionCapturePassSwitch
 ///////////////////////////////////////////////////////////////////////////////
diff --git a/Engine/Source/Runtime/Engine/Public/MaterialCompiler.h b/Engine/Source/Runtime/Engine/Public/MaterialCompiler.h
--- a/Engine/Source/Runtime/Engine/Public/MaterialCompiler.h
+++ b/Engine/Source/Runtime/Engine/Public/MaterialCompiler.h
@@ -592,6 +592,9 @@ public:
 	virtual int32 GetCloudSampleShadowSampleDistance() = 0;
 	virtual int32 GetVolumeSampleConservativeDensity() = 0;
 
+	virtual int32 GetCloudEmptySpaceSkippingSphereCenterWorldPosition() = 0;
+	virtual int32 GetCloudEmptySpaceSkippingSphereRadius() = 0;
+
 	virtual int32 GetHairUV() = 0;
 	virtual int32 GetHairDimensions() = 0;
 	virtual int32 GetHairSeed() = 0;
@@ -1093,6 +1096,16 @@ public:
 		return Compiler->GetVolumeSampleConservativeDensity();
 	}
 
+	virtual int32 GetCloudEmptySpaceSkippingSphereCenterWorldPosition() override
+	{
+		return Compiler->GetCloudEmptySpaceSkippingSphereCenterWorldPosition();
+	}
+
+	virtual int32 GetCloudEmptySpaceSkippingSphereRadius() override
+	{
+		return Compiler->GetCloudEmptySpaceSkippingSphereRadius();
+	}
+
 	virtual int32 SceneDepthWithoutWater(int32 Offset, int32 ViewportUV, bool bUseOffset, float FallbackDepth) override
 	{
 		return Compiler->SceneDepthWithoutWater(Offset, ViewportUV, bUseOffset, FallbackDepth);
diff --git a/Engine/Source/Runtime/Renderer/Private/VolumetricCloudRendering.cpp b/Engine/Source/Runtime/Renderer/Private/VolumetricCloudRendering.cpp
--- a/Engine/Source/Runtime/Renderer/Private/VolumetricCloudRendering.cpp
+++ b/Engine/Source/Runtime/Renderer/Private/VolumetricCloudRendering.cpp
@@ -209,6 +209,18 @@ static TAutoConsoleVariable<float> CVarVolumetricCloudLocalLightsShadowSampleCou
 
 ////////////////////////////////////////////////////////////////////////// 
 
+static TAutoConsoleVariable<int32> CVarVolumetricCloudEmptySpaceSkipping(
+	TEXT("r.VolumetricCloud.EmptySpaceSkipping"), 0,
+	TEXT("Enable/disable empty space skipping to accelerate cloud tracing through emty areas."),
+	ECVF_RenderThreadSafe);
+
+static TAutoConsoleVariable<int32> CVarVolumetricCloudEmptySpaceSkippingVolumeDepth(
+	TEXT("r.VolumetricCloud.EmptySpaceSkipping.VolumeDepth"), 40.0f,
+	TEXT("Set the distance in kilometer over which empty space can be evalauted."),
+	ECVF_RenderThreadSafe);
+
+////////////////////////////////////////////////////////////////////////// 
+
 static TAutoConsoleVariable<int32> CVarVolumetricCloudDebugSampleCountMode(
 	TEXT("r.VolumetricCloud.Debug.SampleCountMode"), 0,
 	TEXT("Only for developers. [0] Disabled [1] Primary material sample count [2] Advanced:raymarched shadow sample count [3] Shadow material sample count [4] Advanced:ground shadow sample count [5] Advanced:ground shadow material sample count"));
@@ -230,6 +242,11 @@ static bool ShouldUseComputeForCloudTracing(const FStaticFeatureLevel InFeatureL
 		&& !bMSAAEnabled;
 }
 
+static bool ShouldUseComputeCloudEmptySpaceSkipping()
+{
+	return !CVarVolumetricCloudDisableCompute.GetValueOnRenderThread() && RHIIsTypedUAVLoadSupported(PF_R16F) && CVarVolumetricCloudEmptySpaceSkipping.GetValueOnRenderThread() > 0;
+}
+
 bool ShouldRenderVolumetricCloud(const FScene* Scene, const FEngineShowFlags& EngineShowFlags)
 {
 	if (Scene && Scene->HasVolumetricCloud() && EngineShowFlags.Atmosphere && EngineShowFlags.Cloud)
@@ -773,6 +790,7 @@ class FRenderVolumetricCloudRenderViewCS : public FMeshMaterialShader
 
 	BEGIN_SHADER_PARAMETER_STRUCT(FParameters, )
 		SHADER_PARAMETER(FVector4f, OutputViewRect)
+		SHADER_PARAMETER(FVector2f, StartTracingDistanceTextureResolution)
 		SHADER_PARAMETER(int32, bBlendCloudColor)
 		SHADER_PARAMETER(int32, TargetCubeFace)
 		SHADER_PARAMETER_RDG_UNIFORM_BUFFER(FRenderVolumetricCloudGlobalParameters, VolumetricCloudRenderViewParamsUB)
@@ -782,6 +800,7 @@ class FRenderVolumetricCloudRenderViewCS : public FMeshMaterialShader
 		SHADER_PARAMETER_RDG_TEXTURE_UAV(RWTexture2D, OutCloudColor)
 		SHADER_PARAMETER_RDG_TEXTURE_UAV(RWTexture2D, OutCloudDepth)
 		SHADER_PARAMETER_RDG_TEXTURE_UAV(RWTexture2D, OutCloudColorCube)
+		SHADER_PARAMETER_RDG_TEXTURE(Texture2D, StartTracingDistanceTexture)
 #if CLOUD_DEBUG_SAMPLES
 		SHADER_PARAMETER_STRUCT_INCLUDE(ShaderPrint::FShaderParameters, ShaderPrintParameters)
 #endif
@@ -838,6 +857,68 @@ IMPLEMENT_MATERIAL_SHADER_TYPE(, FRenderVolumetricCloudRenderViewCS, TEXT("/Engi
 
 //////////////////////////////////////////////////////////////////////////
 
+class FRenderVolumetricCloudEmptySpaceSkippingCS : public FMeshMaterialShader
+{
+	DECLARE_SHADER_TYPE(FRenderVolumetricCloudEmptySpaceSkippingCS, MeshMaterial)
+	SHADER_USE_PARAMETER_STRUCT_WITH_LEGACY_BASE(FRenderVolumetricCloudEmptySpaceSkippingCS, FMeshMaterialShader)
+
+	class FCloudEmptySpaceSkippingDebug : SHADER_PERMUTATION_BOOL("EMPTY_SPACE_SKIPPING_DEBUG");
+
+	using FPermutationDomain = TShaderPermutationDomain<FCloudEmptySpaceSkippingDebug>;
+
+	BEGIN_SHADER_PARAMETER_STRUCT(FParameters, )
+		SHADER_PARAMETER(FVector2f, StartTracingDistanceTextureResolution)
+		SHADER_PARAMETER(float, StartTracingSampleVolumeDepth)
+		SHADER_PARAMETER_RDG_UNIFORM_BUFFER(FRenderVolumetricCloudGlobalParameters, VolumetricCloudRenderViewParamsUB)
+		SHADER_PARAMETER_STRUCT_REF(FViewUniformShaderParameters, View)
+		SHADER_PARAMETER_RDG_TEXTURE_UAV(RWTexture2D, OutStartTracingDistanceTexture)
+//		SHADER_PARAMETER_STRUCT_INCLUDE(ShaderPrint::FShaderParameters, ShaderPrintParameters)
+	END_SHADER_PARAMETER_STRUCT()
+
+	static const int32 ThreadGroupSizeX = 1;
+	static const int32 ThreadGroupSizeY = 1;
+	static const int32 ThreadGroupSizeZ = 64;
+
+	static bool ShouldCompilePermutation(const FMaterialShaderPermutationParameters& Parameters)
+	{
+		const FPermutationDomain PermutationVector(Parameters.PermutationId);
+		const bool bIsCompatible = IsMaterialCompatibleWithVolumetricCloud(Parameters.MaterialParameters, Parameters.Platform);
+
+		if (PermutationVector.Get<FCloudEmptySpaceSkippingDebug>())
+		{
+			return false;	// TODO implement the debug version using ShaderPrint::FShaderParameters
+		}
+
+		return bIsCompatible;
+	}
+
+	static void ModifyCompilationEnvironment(const FMaterialShaderPermutationParameters& Parameters, FShaderCompilerEnvironment& OutEnvironment)
+	{
+		FMaterialShader::ModifyCompilationEnvironment(Parameters, OutEnvironment);
+
+		const FPermutationDomain PermutationVector(Parameters.PermutationId);
+
+		OutEnvironment.SetDefine(TEXT("SHADER_EMPTY_SPACE_SKIPPING_CS"), TEXT("1"));
+		OutEnvironment.SetDefine(TEXT("CLOUD_LAYER_PIXEL_SHADER"), TEXT("1"));
+
+		// Force texture fetches to not use automatic mip generation because the shader is compute and using a dynamic loop to evaluate the material multiple times.
+		OutEnvironment.SetDefine(TEXT("USE_FORCE_TEXTURE_MIP"), TEXT("1"));
+
+		OutEnvironment.SetDefine(TEXT("THREADGROUP_SIZEX"), ThreadGroupSizeX);
+		OutEnvironment.SetDefine(TEXT("THREADGROUP_SIZEY"), ThreadGroupSizeY);
+		OutEnvironment.SetDefine(TEXT("THREADGROUP_SIZEZ"), ThreadGroupSizeZ);
+
+		// This shader must support typed UAV load and we are testing if it is supported at runtime using RHIIsTypedUAVLoadSupported
+		OutEnvironment.CompilerFlags.Add(CFLAG_AllowTypedUAVLoads);
+
+		FForwardLightingParameters::ModifyCompilationEnvironment(Parameters.Platform, OutEnvironment);
+	}
+};
+
+IMPLEMENT_MATERIAL_SHADER_TYPE(, FRenderVolumetricCloudEmptySpaceSkippingCS, TEXT("/Engine/Private/VolumetricCloud.usf"), TEXT("MainCS"), SF_Compute);
+
+//////////////////////////////////////////////////////////////////////////
+
 class FSingleTriangleMeshVertexBuffer : public FRenderResource
 {
 public:
@@ -2128,6 +2209,93 @@ void FSceneRenderer::RenderVolumetricCloudsInternal(FRDGBuilder& GraphBuilder, F
 	const bool bShouldUseComputeForCloudTracing = !bVisualizeConservativeDensity && ShouldUseComputeForCloudTracing(Scene->GetFeatureLevel());
 	if (bShouldUseComputeForCloudTracing)
 	{
+		//
+		// Cloud empty space skipping
+		// TODO: add support for reflection rendering without breaking compute work item overlap for each faces.
+		//
+
+		FRDGTextureRef EmptySpaceSkippingTexture = GSystemTextures.GetBlackDummy(GraphBuilder);
+		if(ShouldUseComputeCloudEmptySpaceSkipping() && !CloudRC.bIsReflectionRendering)
+		{
+			const FIntPoint EmptySpaceSkippingTextureResolution = FIntPoint(32, 32);
+			EmptySpaceSkippingTexture = GraphBuilder.CreateTexture(
+				FRDGTextureDesc::Create2D(EmptySpaceSkippingTextureResolution, PF_R16F, FClearValueBinding(EClearBinding::ENoneBound),
+					TexCreate_ShaderResource | TexCreate_UAV), TEXT("Cloud.EmptySpaceSkippingTexture"));
+			
+			FRenderVolumetricCloudEmptySpaceSkippingCS::FParameters* PassParameters = GraphBuilder.AllocParameters<FRenderVolumetricCloudEmptySpaceSkippingCS::FParameters>();
+			PassParameters->VolumetricCloudRenderViewParamsUB = CloudPassUniformBuffer;
+			PassParameters->View = ViewUniformBuffer;
+			PassParameters->OutStartTracingDistanceTexture = GraphBuilder.CreateUAV(FRDGTextureUAVDesc(EmptySpaceSkippingTexture));
+			PassParameters->StartTracingDistanceTextureResolution = FVector2f(EmptySpaceSkippingTextureResolution.X, EmptySpaceSkippingTextureResolution.Y);
+			const float KilometersToCentimeters = 100000.0f;
+			PassParameters->StartTracingSampleVolumeDepth = FMath::Max(1.0f, CVarVolumetricCloudEmptySpaceSkippingVolumeDepth.GetValueOnRenderThread()) * KilometersToCentimeters;
+
+			const bool bCloudEmptySpaceSkippingDebug = false;
+			//if (bCloudEmptySpaceSkippingDebug)
+			//{
+			//	ShaderPrint::SetEnabled(true);
+			//	ShaderPrint::RequestSpaceForLines(65536);
+			//	ShaderPrint::SetParameters(GraphBuilder, MainView.ShaderPrintData, PassParameters->ShaderPrintParameters);
+			//}
+			
+			const FMaterialRenderProxy* MaterialRenderProxy = nullptr;
+			const FMaterial* MaterialResource = &CloudVolumeMaterialProxy->GetMaterialWithFallback(Scene->GetFeatureLevel(), MaterialRenderProxy);
+			MaterialRenderProxy = MaterialRenderProxy ? MaterialRenderProxy : CloudVolumeMaterialProxy;
+
+			typename FRenderVolumetricCloudEmptySpaceSkippingCS::FPermutationDomain PermutationVector;
+			PermutationVector.Set<typename FRenderVolumetricCloudEmptySpaceSkippingCS::FCloudEmptySpaceSkippingDebug>(bCloudEmptySpaceSkippingDebug);
+
+			TShaderRef<FRenderVolumetricCloudEmptySpaceSkippingCS> ComputeShader = MaterialResource->GetShader<FRenderVolumetricCloudEmptySpaceSkippingCS>(&FLocalVertexFactory::StaticType, PermutationVector, false);
+
+			const FIntVector GroupCount = FComputeShaderUtils::GetGroupCount(EmptySpaceSkippingTextureResolution, FIntPoint(FRenderVolumetricCloudEmptySpaceSkippingCS::ThreadGroupSizeX, FRenderVolumetricCloudEmptySpaceSkippingCS::ThreadGroupSizeY));
+
+			if (ComputeShader.IsValid())
+			{
+				ClearUnusedGraphResources(ComputeShader, PassParameters);
+			}
+
+			GraphBuilder.AddPass(
+				RDG_EVENT_NAME("CloudViewEmptySpaceSkipping (CS) %dx%d", EmptySpaceSkippingTextureResolution.X, EmptySpaceSkippingTextureResolution.Y),
+				PassParameters,
+				CloudRC.bAsyncCompute ? ERDGPassFlags::AsyncCompute : ERDGPassFlags::Compute,
+				[LocalScene = Scene, MaterialRenderProxy, MaterialResource, PassParameters, ComputeShader, GroupCount](FRHIComputeCommandList& RHICmdList)
+				{
+					if (MaterialResource->GetMaterialDomain() != MD_Volume)
+					{
+						return;
+					}
+
+					FMeshPassProcessorRenderState DrawRenderState;
+
+					FMeshMaterialShaderElementData ShaderElementData;
+					ShaderElementData.FadeUniformBuffer = GDistanceCullFadedInUniformBuffer.GetUniformBufferRHI();
+					ShaderElementData.DitherUniformBuffer = GDitherFadedInUniformBuffer.GetUniformBufferRHI();
+
+					FMeshProcessorShaders PassShaders;
+					PassShaders.ComputeShader = ComputeShader;
+
+					FMeshDrawShaderBindings ShaderBindings;
+					ShaderBindings.Initialize(PassShaders);
+
+					int32 DataOffset = 0;
+					FMeshDrawSingleShaderBindings SingleShaderBindings = ShaderBindings.GetSingleShaderBindings(SF_Compute, DataOffset);
+					ComputeShader->GetShaderBindings(LocalScene, LocalScene->GetFeatureLevel(), nullptr, *MaterialRenderProxy, *MaterialResource, DrawRenderState, ShaderElementData, SingleShaderBindings);
+
+					ShaderBindings.Finalize(&PassShaders);
+
+					FRHIComputeShader* ComputeShaderRHI = ComputeShader.GetComputeShader();
+					SetComputePipelineState(RHICmdList, ComputeShaderRHI);
+					ShaderBindings.SetOnCommandList(RHICmdList, ComputeShaderRHI);
+					SetShaderParameters(RHICmdList, ComputeShader, ComputeShaderRHI, *PassParameters);
+					RHICmdList.DispatchComputeShader(GroupCount.X, GroupCount.Y, GroupCount.Z);
+					UnsetShaderUAVs(RHICmdList, ComputeShader, ComputeShaderRHI);
+				});
+		}
+		
+		//
+		// Cloud view tracing
+		//
+
 		FRDGTextureUAVRef CloudColorCubeTextureUAV;
 		FRDGTextureUAVRef CloudColorTextureUAV;
 		FRDGTextureUAVRef CloudDepthTextureUAV;
@@ -2149,6 +2317,9 @@ void FSceneRenderer::RenderVolumetricCloudsInternal(FRDGBuilder& GraphBuilder, F
 		PassParameters->OutCloudDepth = CloudDepthTextureUAV;
 		PassParameters->OutCloudColorCube = CloudColorCubeTextureUAV;
 
+		PassParameters->StartTracingDistanceTextureResolution = FVector2f(EmptySpaceSkippingTexture->Desc.Extent.X, EmptySpaceSkippingTexture->Desc.Extent.Y);
+		PassParameters->StartTracingDistanceTexture = EmptySpaceSkippingTexture;
+
 #if CLOUD_DEBUG_SAMPLES
 		ShaderPrint::SetEnabled(true);
 		ShaderPrint::RequestSpaceForLines(65536);
