Invalidate CSM caching if FOV has been changed.

diff --git a/Engine/Source/Runtime/Renderer/Private/ShadowSetup.cpp b/Engine/Source/Runtime/Renderer/Private/ShadowSetup.cpp
--- a/Engine/Source/Runtime/Renderer/Private/ShadowSetup.cpp
+++ b/Engine/Source/Runtime/Renderer/Private/ShadowSetup.cpp
@@ -3272,15 +3272,14 @@ void ComputeViewDependentWholeSceneShadowCacheModes(
 			bool bIsValidCachedViewDependentWholeSceneShadow = ProjectedShadowInitializer.WorldToLight == CachedShadowInitializer.WorldToLight
 				&& ProjectedShadowInitializer.WAxis == CachedShadowInitializer.WAxis
 				&& ProjectedShadowInitializer.bOnePassPointLightShadow == CachedShadowInitializer.bOnePassPointLightShadow
-				&& ProjectedShadowInitializer.bRayTracedDistanceField == CachedShadowInitializer.bRayTracedDistanceField;
+				&& ProjectedShadowInitializer.bRayTracedDistanceField == CachedShadowInitializer.bRayTracedDistanceField
+				&& ProjectedShadowInitializer.SubjectBounds.SphereRadius == CachedShadowInitializer.SubjectBounds.SphereRadius;
 
 			if (bIsValidCachedViewDependentWholeSceneShadow)
 			{
 				bool bExactlyEqual = ProjectedShadowInitializer.PreShadowTranslation == CachedShadowInitializer.PreShadowTranslation
 					&& ProjectedShadowInitializer.Scales == CachedShadowInitializer.Scales
 					&& ProjectedShadowInitializer.SubjectBounds.Origin == CachedShadowInitializer.SubjectBounds.Origin
-					&& ProjectedShadowInitializer.SubjectBounds.BoxExtent == CachedShadowInitializer.SubjectBounds.BoxExtent
-					&& ProjectedShadowInitializer.SubjectBounds.SphereRadius == CachedShadowInitializer.SubjectBounds.SphereRadius
 					&& ProjectedShadowInitializer.MinLightW == CachedShadowInitializer.MinLightW
 					&& ProjectedShadowInitializer.MaxDistanceToCastInLightW == CachedShadowInitializer.MaxDistanceToCastInLightW;
 
