Fix for RenderSingleLayerWaterDepthPrepass incorrectly set to not write stencil during depth copy

Issue:
RenderSingleLayerWaterDepthPrepass intends to clear the stencil buffer during the depth copy according to this comment:
// Copy main depth buffer content to our prepass depth buffer and clear stencil to 0

The depth stencil state is set correctly but the depth stencil binding was set for depth write only.  This caused the following validation error:

Assertion failed: DSV.ExclusiveDepthStencil.IsUsingStencil() [Engine\Source\Runtime\RHI\Public\RHIValidationContext.h] [Line: 1157]
Graphics PSO is using stencil but it's not enabled on the RenderPass.

Fix:
Set the depth stencil binding to also write stencil as intended.

diff --git a/Engine/Source/Runtime/Renderer/Private/SingleLayerWaterRendering.cpp b/Engine/Source/Runtime/Renderer/Private/SingleLayerWaterRendering.cpp
--- a/Engine/Source/Runtime/Renderer/Private/SingleLayerWaterRendering.cpp
+++ b/Engine/Source/Runtime/Renderer/Private/SingleLayerWaterRendering.cpp
@@ -377,7 +377,7 @@ void FDeferredShadingSceneRenderer::RenderSingleLayerWaterDepthPrepass(FRDGBuild
 			{
 				PassParameters->DepthTexture = SceneTextures.Depth.Target;
 			}
-			PassParameters->RenderTargets.DepthStencil = FDepthStencilBinding(OutDepthPrepassTexture.Target, ERenderTargetLoadAction::ENoAction, ERenderTargetLoadAction::ELoad, FExclusiveDepthStencil::DepthWrite_StencilNop);
+			PassParameters->RenderTargets.DepthStencil = FDepthStencilBinding(OutDepthPrepassTexture.Target, ERenderTargetLoadAction::ENoAction, ERenderTargetLoadAction::ELoad, FExclusiveDepthStencil::DepthWrite_StencilWrite);
 
 			FGlobalShaderMap* ShaderMap = GetGlobalShaderMap(GMaxRHIFeatureLevel);
 
