Register Cloud.ShadowTexture buffer for use on async compute pipe as well as graphics pipe to resolve ensure when utilized by async Lumen passes

Resolves the following ensure:
Resource Cloud.ShadowTexture is in external access mode and is valid for access on the following pipelines: Graphics, but is being used on the AsyncCompute pipe in pass Scene.VolumetricCloud.CloudView (CS) 720x405.

diff --git a/Engine/Source/Runtime/Renderer/Private/VolumetricCloudRendering.cpp b/Engine/Source/Runtime/Renderer/Private/VolumetricCloudRendering.cpp
--- a/Engine/Source/Runtime/Renderer/Private/VolumetricCloudRendering.cpp
+++ b/Engine/Source/Runtime/Renderer/Private/VolumetricCloudRendering.cpp
@@ -1923,13 +1923,12 @@ void FSceneRenderer::InitVolumetricCloudsForViews(FRDGBuilder& GraphBuilder, boo
 										FilterTracedCloudTexture(&SpatiallyFilteredShadowTexture, VolumetricCloudParams.VolumetricCloud.CloudShadowmapSizeInvSize[LightIndex], CloudShadowTextureTexelWorldSizeInvSize, false);
 									}
 									ViewInfo.VolumetricCloudShadowRenderTarget[LightIndex] = SpatiallyFilteredShadowTexture;
-									ViewInfo.VolumetricCloudShadowExtractedRenderTarget[LightIndex] = ConvertToExternalAccessTexture(GraphBuilder, ExternalAccessQueue, ViewInfo.VolumetricCloudShadowRenderTarget[LightIndex]);
 								}
 								else
 								{
 									ViewInfo.VolumetricCloudShadowRenderTarget[LightIndex] = GraphBuilder.RegisterExternalTexture(CloudShadowTemporalRT.CurrentRenderTarget());
-									ViewInfo.VolumetricCloudShadowExtractedRenderTarget[LightIndex] = ConvertToExternalAccessTexture(GraphBuilder, ExternalAccessQueue, ViewInfo.VolumetricCloudShadowRenderTarget[LightIndex]);
 								}
+								ViewInfo.VolumetricCloudShadowExtractedRenderTarget[LightIndex] = ConvertToExternalAccessTexture(GraphBuilder, ExternalAccessQueue, ViewInfo.VolumetricCloudShadowRenderTarget[LightIndex], ERHIAccess::SRVMask, ERHIPipeline::All);
 
 							}
 							else
@@ -1944,13 +1943,12 @@ void FSceneRenderer::InitVolumetricCloudsForViews(FRDGBuilder& GraphBuilder, boo
 										FilterTracedCloudTexture(&SpatiallyFilteredShadowTexture, VolumetricCloudParams.VolumetricCloud.CloudShadowmapSizeInvSize[LightIndex], CloudShadowTextureTexelWorldSizeInvSize, false);
 									}
 									ViewInfo.VolumetricCloudShadowRenderTarget[LightIndex] = SpatiallyFilteredShadowTexture;
-									ViewInfo.VolumetricCloudShadowExtractedRenderTarget[LightIndex] = ConvertToExternalAccessTexture(GraphBuilder, ExternalAccessQueue, ViewInfo.VolumetricCloudShadowRenderTarget[LightIndex]);
 								}
 								else
 								{
 									ViewInfo.VolumetricCloudShadowRenderTarget[LightIndex] = NewCloudShadowTexture;
-									ViewInfo.VolumetricCloudShadowExtractedRenderTarget[LightIndex] = ConvertToExternalAccessTexture(GraphBuilder, ExternalAccessQueue, ViewInfo.VolumetricCloudShadowRenderTarget[LightIndex]);
 								}
+								ViewInfo.VolumetricCloudShadowExtractedRenderTarget[LightIndex] = ConvertToExternalAccessTexture(GraphBuilder, ExternalAccessQueue, ViewInfo.VolumetricCloudShadowRenderTarget[LightIndex], ERHIAccess::SRVMask, ERHIPipeline::All);
 							}
 						}
 						else
