Global Distance Field - fixed missing updates to Global SDF when there are more than one SceneRender() calls per frame by removing UpdateTrackingBounds optimization:
* We update primitives once per scene inside a SceneRender()
* GlobalSDF is cached per view
* GlobalSDF update is based on whether current primitive is in UpdateTrackingBounds, where UpdateTrackingBounds are built from current views
* If render gets called with a non main view (like a cubemap capture) then scene primitive will be updated, but GlobalSDF for the main view won't

diff --git a/Engine/Source/Runtime/Renderer/Private/DistanceFieldLightingShared.h b/Engine/Source/Runtime/Renderer/Private/DistanceFieldLightingShared.h
--- a/Engine/Source/Runtime/Renderer/Private/DistanceFieldLightingShared.h
+++ b/Engine/Source/Runtime/Renderer/Private/DistanceFieldLightingShared.h
@@ -118,15 +118,8 @@ namespace DistanceField
 {
 	constexpr uint32 MinPrimitiveModifiedBoundsAllocation = 16 * 1024;
 
-	// Bounds indicating where distance field object updates should be tracked
-	struct FUpdateTrackingBounds
-	{
-		FBox GlobalDistanceFieldBounds;
-	};
-
 	RENDERER_API FDistanceFieldObjectBufferParameters SetupObjectBufferParameters(FRDGBuilder& GraphBuilder, const FDistanceFieldSceneData& DistanceFieldSceneData);
 	RENDERER_API FDistanceFieldAtlasParameters SetupAtlasParameters(FRDGBuilder& GraphBuilder, const FDistanceFieldSceneData& DistanceFieldSceneData);
-	void BuildUpdateTrackingBounds(const TArray<FViewInfo>& Views, DistanceField::FUpdateTrackingBounds& UpdateTrackingBounds);
 };
 
 BEGIN_SHADER_PARAMETER_STRUCT(FDistanceFieldCulledObjectBufferParameters, )
diff --git a/Engine/Source/Runtime/Renderer/Private/DistanceFieldObjectManagement.cpp b/Engine/Source/Runtime/Renderer/Private/DistanceFieldObjectManagement.cpp
--- a/Engine/Source/Runtime/Renderer/Private/DistanceFieldObjectManagement.cpp
+++ b/Engine/Source/Runtime/Renderer/Private/DistanceFieldObjectManagement.cpp
@@ -217,18 +217,13 @@ static int32 PartitionUpdateRangesDFO(FParallelUpdateRangesDFO& Ranges, int32 It
 	return Ranges.Range[3].ItemCount > 0 ? 4 : 3;
 }
 
-void AddModifiedBounds(FScene* Scene, FGlobalDFCacheType CacheType, const FBox& Bounds, const DistanceField::FUpdateTrackingBounds& UpdateTrackingBounds)
+void AddModifiedBounds(FScene* Scene, FGlobalDFCacheType CacheType, const FBox& Bounds)
 {
 	FDistanceFieldSceneData& DistanceFieldData = Scene->DistanceFieldSceneData;
-
-	if (UpdateTrackingBounds.GlobalDistanceFieldBounds.IsValid 
-		&& UpdateTrackingBounds.GlobalDistanceFieldBounds.Intersect(Bounds))
-	{
-		DistanceFieldData.PrimitiveModifiedBounds[CacheType].Add(Bounds);
-	}
+	DistanceFieldData.PrimitiveModifiedBounds[CacheType].Add(Bounds);
 }
 
-void UpdateGlobalDistanceFieldObjectRemoves(FScene* Scene, const DistanceField::FUpdateTrackingBounds& UpdateTrackingBounds, TArray<FSetElementId>& DistanceFieldAssetRemoves)
+void UpdateGlobalDistanceFieldObjectRemoves(FScene* Scene, TArray<FSetElementId>& DistanceFieldAssetRemoves)
 {
 	FDistanceFieldSceneData& DistanceFieldSceneData = Scene->DistanceFieldSceneData;
 
@@ -261,7 +256,7 @@ void UpdateGlobalDistanceFieldObjectRemoves(FScene* Scene, const DistanceField::
 				if (InstanceIndex >= 0)
 				{
 					FGlobalDFCacheType CacheType = PrimitiveRemoveInfo.bOftenMoving ? GDF_Full : GDF_MostlyStatic;
-					AddModifiedBounds(Scene, CacheType, DistanceFieldSceneData.PrimitiveInstanceMapping[InstanceIndex].WorldBounds, UpdateTrackingBounds);
+					AddModifiedBounds(Scene, CacheType, DistanceFieldSceneData.PrimitiveInstanceMapping[InstanceIndex].WorldBounds);
 					PendingRemoveOperations.Add(InstanceIndex);
 				}
 			}
@@ -299,7 +294,7 @@ void UpdateGlobalDistanceFieldObjectRemoves(FScene* Scene, const DistanceField::
 	}
 }
 
-void UpdateGlobalHeightFieldObjectRemoves(FScene* Scene, const DistanceField::FUpdateTrackingBounds& UpdateTrackingBounds)
+void UpdateGlobalHeightFieldObjectRemoves(FScene* Scene)
 {
 	FDistanceFieldSceneData& SceneData = Scene->DistanceFieldSceneData;
 
@@ -316,7 +311,7 @@ void UpdateGlobalHeightFieldObjectRemoves(FScene* Scene, const DistanceField::FU
 			if (ObjectIdx >= 0)
 			{
 				const FGlobalDFCacheType CacheType = RemoveInfo.bOftenMoving ? GDF_Full : GDF_MostlyStatic;
-				AddModifiedBounds(Scene, CacheType, RemoveInfo.WorldBounds, UpdateTrackingBounds);
+				AddModifiedBounds(Scene, CacheType, RemoveInfo.WorldBounds);
 				PendingRemoveObjectIndices.Add(ObjectIdx);
 			}
 		}
@@ -376,7 +371,6 @@ void LogDistanceFieldUpdate(FPrimitiveSceneInfo const* PrimitiveSceneInfo, float
 bool ProcessPrimitiveUpdate(
 	bool bIsAddOperation,
 	FScene* Scene,
-	const DistanceField::FUpdateTrackingBounds& UpdateTrackingBounds,
 	FPrimitiveSceneInfo* PrimitiveSceneInfo,
 	TArray<FRenderTransform>& InstanceLocalToPrimitiveTransforms,
 	TArray<int32>& IndicesToUpdateInObjectBuffers, 
@@ -460,7 +454,7 @@ bool ProcessPrimitiveUpdate(
 						const int32 MappingIndex = DistanceFieldSceneData.PrimitiveInstanceMapping.Add(FPrimitiveAndInstance(LocalToWorld, WorldBounds, PrimitiveSceneInfo, TransformIndex));
 						PrimitiveSceneInfo->DistanceFieldInstanceIndices[TransformIndex] = AddIndex;
 
-						AddModifiedBounds(Scene, CacheType, WorldBounds, UpdateTrackingBounds);
+						AddModifiedBounds(Scene, CacheType, WorldBounds);
 						LogDistanceFieldUpdate(PrimitiveSceneInfo, BoundingRadius, bIsAddOperation);
 					}
 					else 
@@ -480,12 +474,12 @@ bool ProcessPrimitiveUpdate(
 								const FVector MergedExtentIncrease = MergedBounds.GetExtent() - Mapping.WorldBounds.GetExtent() - WorldBounds.GetExtent();
 								if (MergedExtentIncrease.GetMax() < 100.0f)
 								{
-									AddModifiedBounds(Scene, CacheType, MergedBounds, UpdateTrackingBounds);
+									AddModifiedBounds(Scene, CacheType, MergedBounds);
 								}
 								else
 								{
-									AddModifiedBounds(Scene, CacheType, Mapping.WorldBounds, UpdateTrackingBounds);
-									AddModifiedBounds(Scene, CacheType, WorldBounds, UpdateTrackingBounds);
+									AddModifiedBounds(Scene, CacheType, Mapping.WorldBounds);
+									AddModifiedBounds(Scene, CacheType, WorldBounds);
 								}
 								LogDistanceFieldUpdate(PrimitiveSceneInfo, BoundingRadius, bIsAddOperation);
 
@@ -513,7 +507,6 @@ bool ProcessPrimitiveUpdate(
 bool ProcessHeightFieldPrimitiveUpdate(
 	bool bIsAddOperation,
 	FScene* Scene,
-	const DistanceField::FUpdateTrackingBounds& UpdateTrackingBounds,
 	FPrimitiveSceneInfo* PrimitiveSceneInfo,
 	TArray<int32>& IndicesToUpdateInObjectBuffers)
 {
@@ -540,7 +533,7 @@ bool ProcessHeightFieldPrimitiveUpdate(
 
 		const FGlobalDFCacheType CacheType = PrimitiveSceneInfo->Proxy->IsOftenMoving() ? GDF_Full : GDF_MostlyStatic;
 		const FBoxSphereBounds Bounds = PrimitiveSceneInfo->Proxy->GetBounds();
-		AddModifiedBounds(Scene, CacheType, Bounds.GetBox(), UpdateTrackingBounds);
+		AddModifiedBounds(Scene, CacheType, Bounds.GetBox());
 
 		PrimitiveSceneInfo->DistanceFieldInstanceIndices.Empty(1);
 		PrimitiveSceneInfo->DistanceFieldInstanceIndices.Add(UploadIdx);
@@ -561,7 +554,6 @@ void FDistanceFieldSceneData::UpdateDistanceFieldObjectBuffers(
 	FRDGBuilder& GraphBuilder,
 	FRDGExternalAccessQueue& ExternalAccessQueue,
 	FScene* Scene,
-	const DistanceField::FUpdateTrackingBounds& UpdateTrackingBounds,
 	TArray<FDistanceFieldAssetMipId>& DistanceFieldAssetAdds,
 	TArray<FSetElementId>& DistanceFieldAssetRemoves)
 {
@@ -590,7 +582,7 @@ void FDistanceFieldSceneData::UpdateDistanceFieldObjectBuffers(
 
 		// Process removes before adds, as the adds will overwrite primitive allocation info
 		// This also prevents re-uploading distance fields on render state recreation
-		UpdateGlobalDistanceFieldObjectRemoves(Scene, UpdateTrackingBounds, DistanceFieldAssetRemoves);
+		UpdateGlobalDistanceFieldObjectRemoves(Scene, DistanceFieldAssetRemoves);
 
 		if ((PendingAddOperations.Num() > 0 || PendingUpdateOperations.Num() > 0) && GDFReverseAtlasAllocationOrder == GDFPreviousReverseAtlasAllocationOrder)
 		{
@@ -602,7 +594,6 @@ void FDistanceFieldSceneData::UpdateDistanceFieldObjectBuffers(
 				if (!ProcessPrimitiveUpdate(
 					true,
 					Scene,
-					UpdateTrackingBounds,
 					PrimitiveSceneInfo,
 					InstanceLocalToPrimitiveTransforms,
 					IndicesToUpdateInObjectBuffers,
@@ -618,7 +609,6 @@ void FDistanceFieldSceneData::UpdateDistanceFieldObjectBuffers(
 				ProcessPrimitiveUpdate(
 					false,
 					Scene,
-					UpdateTrackingBounds,
 					PrimitiveSceneInfo,
 					InstanceLocalToPrimitiveTransforms,
 					IndicesToUpdateInObjectBuffers,
@@ -835,12 +825,9 @@ void FSceneRenderer::UpdateGlobalHeightFieldObjectBuffers(FRDGBuilder& GraphBuil
 		QUICK_SCOPE_CYCLE_COUNTER(STAT_UpdateHeightFieldSceneObjectData);
 		RDG_EVENT_SCOPE(GraphBuilder, "UpdateHeightFieldSceneObjectData");
 
-		DistanceField::FUpdateTrackingBounds UpdateTrackingBounds;
-		BuildUpdateTrackingBounds(Views, UpdateTrackingBounds);
-
 		if (!DistanceFieldSceneData.HeightFieldObjectBuffers)
 		{
-			AddOrRemoveSceneHeightFieldPrimitives(UpdateTrackingBounds, true);
+			AddOrRemoveSceneHeightFieldPrimitives(true);
 
 			for (int32 Idx = 0; Idx < DistanceFieldSceneData.HeightfieldPrimitives.Num(); ++Idx)
 			{
@@ -871,18 +858,18 @@ void FSceneRenderer::UpdateGlobalHeightFieldObjectBuffers(FRDGBuilder& GraphBuil
 			}
 		}
 
-		UpdateGlobalHeightFieldObjectRemoves(Scene, UpdateTrackingBounds);
+		UpdateGlobalHeightFieldObjectRemoves(Scene);
 
 		if (DistanceFieldSceneData.PendingHeightFieldAddOps.Num() || DistanceFieldSceneData.PendingHeightFieldUpdateOps.Num())
 		{
 			for (FPrimitiveSceneInfo* PrimitiveSceneInfo : DistanceFieldSceneData.PendingHeightFieldAddOps)
 			{
-				ProcessHeightFieldPrimitiveUpdate(true, Scene, UpdateTrackingBounds, PrimitiveSceneInfo, DistanceFieldSceneData.IndicesToUpdateInHeightFieldObjectBuffers);
+				ProcessHeightFieldPrimitiveUpdate(true, Scene, PrimitiveSceneInfo, DistanceFieldSceneData.IndicesToUpdateInHeightFieldObjectBuffers);
 			}
 
 			for (FPrimitiveSceneInfo* PrimitiveSceneInfo : DistanceFieldSceneData.PendingHeightFieldUpdateOps)
 			{
-				ProcessHeightFieldPrimitiveUpdate(false, Scene, UpdateTrackingBounds, PrimitiveSceneInfo, DistanceFieldSceneData.IndicesToUpdateInHeightFieldObjectBuffers);
+				ProcessHeightFieldPrimitiveUpdate(false, Scene, PrimitiveSceneInfo, DistanceFieldSceneData.IndicesToUpdateInHeightFieldObjectBuffers);
 			}
 
 			DistanceFieldSceneData.PendingHeightFieldAddOps.Reset();
@@ -1014,9 +1001,6 @@ void FSceneRenderer::PrepareDistanceFieldScene(FRDGBuilder& GraphBuilder, FRDGEx
 	const bool bShouldPrepareHeightFieldScene = ShouldPrepareHeightFieldScene();
 	const bool bShouldPrepareDistanceFieldScene = ShouldPrepareDistanceFieldScene();
 
-	DistanceField::FUpdateTrackingBounds UpdateTrackingBounds;
-	DistanceField::BuildUpdateTrackingBounds(Views, UpdateTrackingBounds);
-
 	if (bShouldPrepareHeightFieldScene)
 	{
 		extern int32 GHFShadowQuality;
@@ -1029,14 +1013,14 @@ void FSceneRenderer::PrepareDistanceFieldScene(FRDGBuilder& GraphBuilder, FRDGEx
 	}
 	else if (bShouldPrepareDistanceFieldScene)
 	{
-		AddOrRemoveSceneHeightFieldPrimitives(UpdateTrackingBounds);
+		AddOrRemoveSceneHeightFieldPrimitives();
 	}
 
 	if (bShouldPrepareDistanceFieldScene)
 	{
 		TArray<FDistanceFieldAssetMipId> DistanceFieldAssetAdds;
 		TArray<FSetElementId> DistanceFieldAssetRemoves;
-		Scene->DistanceFieldSceneData.UpdateDistanceFieldObjectBuffers(GraphBuilder, ExternalAccessQueue, Scene, UpdateTrackingBounds, DistanceFieldAssetAdds, DistanceFieldAssetRemoves);
+		Scene->DistanceFieldSceneData.UpdateDistanceFieldObjectBuffers(GraphBuilder, ExternalAccessQueue, Scene, DistanceFieldAssetAdds, DistanceFieldAssetRemoves);
 
 		Scene->DistanceFieldSceneData.UpdateDistanceFieldAtlas(GraphBuilder, ExternalAccessQueue, Views[0], Scene, IsLumenEnabled(Views[0]), Views[0].ShaderMap, DistanceFieldAssetAdds, DistanceFieldAssetRemoves);
 
@@ -1071,7 +1055,7 @@ void FSceneRenderer::PrepareDistanceFieldScene(FRDGBuilder& GraphBuilder, FRDGEx
 	}
 }
 
-void FSceneRenderer::AddOrRemoveSceneHeightFieldPrimitives(const DistanceField::FUpdateTrackingBounds& UpdateTrackingBounds, bool bSkipAdd)
+void FSceneRenderer::AddOrRemoveSceneHeightFieldPrimitives(bool bSkipAdd)
 {
 	FDistanceFieldSceneData& SceneData = Scene->DistanceFieldSceneData;
 
@@ -1091,7 +1075,7 @@ void FSceneRenderer::AddOrRemoveSceneHeightFieldPrimitives(const DistanceField::
 		check(RemoveInfo.DistanceFieldInstanceIndices.Num() == 1);
 		PendingRemoveIndices.Add(RemoveInfo.DistanceFieldInstanceIndices[0]);
 		const FGlobalDFCacheType CacheType = RemoveInfo.bOftenMoving ? GDF_Full : GDF_MostlyStatic;
-		AddModifiedBounds(Scene, CacheType, RemoveInfo.WorldBounds, UpdateTrackingBounds);
+		AddModifiedBounds(Scene, CacheType, RemoveInfo.WorldBounds);
 	}
 	SceneData.PendingHeightFieldRemoveOps.Reset();
 	Algo::Sort(PendingRemoveIndices);
@@ -1115,27 +1099,10 @@ void FSceneRenderer::AddOrRemoveSceneHeightFieldPrimitives(const DistanceField::
 			Primitive->DistanceFieldInstanceIndices.Add(HFIdx);
 			const FGlobalDFCacheType CacheType = Primitive->Proxy->IsOftenMoving() ? GDF_Full : GDF_MostlyStatic;
 			const FBoxSphereBounds& Bounds = Primitive->Proxy->GetBounds();
-			AddModifiedBounds(Scene, CacheType, Bounds.GetBox(), UpdateTrackingBounds);
+			AddModifiedBounds(Scene, CacheType, Bounds.GetBox());
 		}
 		SceneData.PendingHeightFieldAddOps.Reset();
 	}
 
 	SceneData.PendingHeightFieldUpdateOps.Empty();
 }
-
-void DistanceField::BuildUpdateTrackingBounds(const TArray<FViewInfo>& Views, DistanceField::FUpdateTrackingBounds& UpdateTrackingBounds)
-{
-	// Global Distance Field
-	{
-		UpdateTrackingBounds.GlobalDistanceFieldBounds.Init();
-
-		for (const FViewInfo& View : Views)
-		{
-			const FSceneViewState* ViewState = View.ViewState;
-			if (ViewState)
-			{
-				GlobalDistanceField::ExpandDistanceFieldUpdateTrackingBounds(ViewState, UpdateTrackingBounds);
-			}
-		}
-	}
-}
diff --git a/Engine/Source/Runtime/Renderer/Private/GlobalDistanceField.cpp b/Engine/Source/Runtime/Renderer/Private/GlobalDistanceField.cpp
--- a/Engine/Source/Runtime/Renderer/Private/GlobalDistanceField.cpp
+++ b/Engine/Source/Runtime/Renderer/Private/GlobalDistanceField.cpp
@@ -2866,21 +2866,4 @@ void UpdateGlobalDistanceFieldVolume(
 		View.ViewState->GlobalDistanceFieldData->bFirstFrame = false;
 		View.ViewState->GlobalDistanceFieldData->UpdateFrame = View.Family->FrameNumber;
 	}
-}
-
-void GlobalDistanceField::ExpandDistanceFieldUpdateTrackingBounds(const FSceneViewState* ViewState, DistanceField::FUpdateTrackingBounds& UpdateTrackingBounds)
-{
-	// Global Distance Field is interested in any updates which in ClipmapInfluenceBounds range of it's clipmaps
-
-	const int32 NumClipmaps = FMath::Clamp<int32>(GetNumGlobalDistanceFieldClipmaps(false, 1.0f), 0, GlobalDistanceField::MaxClipmaps);
-	for (int32 ClipmapIndex = 0; ClipmapIndex < NumClipmaps; ClipmapIndex++)
-	{
-		const FGlobalDistanceFieldClipmapState& ClipmapViewState = ViewState->GlobalDistanceFieldData->ClipmapState[ClipmapIndex];
-
-		const FVector3f ClipmapCenter = ClipmapViewState.CachedClipmapCenter;
-		const float ClipmapExtent = ClipmapViewState.CachedClipmapExtent + ClipmapViewState.CacheClipmapInfluenceRadius;
-
-		const FBox ClipmapInfluenceBounds(ClipmapCenter - ClipmapExtent, ClipmapCenter + ClipmapExtent);
-		UpdateTrackingBounds.GlobalDistanceFieldBounds += ClipmapInfluenceBounds;
-	}
-}
+}
\ No newline at end of file
diff --git a/Engine/Source/Runtime/Renderer/Private/GlobalDistanceField.h b/Engine/Source/Runtime/Renderer/Private/GlobalDistanceField.h
--- a/Engine/Source/Runtime/Renderer/Private/GlobalDistanceField.h
+++ b/Engine/Source/Runtime/Renderer/Private/GlobalDistanceField.h
@@ -33,7 +33,6 @@ namespace GlobalDistanceField
 	uint32 GetPageTableClipmapResolution(bool bLumenEnabled);
 	FIntVector GetPageTableTextureResolution(bool bLumenEnabled, float LumenSceneViewDistance);
 	int32 GetMaxPageNum(bool bLumenEnabled, float LumenSceneViewDistance);
-	void ExpandDistanceFieldUpdateTrackingBounds(const FSceneViewState* ViewState, DistanceField::FUpdateTrackingBounds& UpdateTrackingBounds);
 };
 
 /** 
diff --git a/Engine/Source/Runtime/Renderer/Private/ScenePrivate.h b/Engine/Source/Runtime/Renderer/Private/ScenePrivate.h
--- a/Engine/Source/Runtime/Renderer/Private/ScenePrivate.h
+++ b/Engine/Source/Runtime/Renderer/Private/ScenePrivate.h
@@ -1968,7 +1968,6 @@ public:
 		FRDGBuilder& GraphBuilder,
 		FRDGExternalAccessQueue& ExternalAccessQueue,
 		FScene* Scene,
-		const DistanceField::FUpdateTrackingBounds& UpdateTrackingBounds,
 		TArray<FDistanceFieldAssetMipId>& DistanceFieldAssetAdds,
 		TArray<FSetElementId>& DistanceFieldAssetRemoves);
 
diff --git a/Engine/Source/Runtime/Renderer/Private/SceneRendering.h b/Engine/Source/Runtime/Renderer/Private/SceneRendering.h
--- a/Engine/Source/Runtime/Renderer/Private/SceneRendering.h
+++ b/Engine/Source/Runtime/Renderer/Private/SceneRendering.h
@@ -74,11 +74,6 @@ struct FSceneTextures;
 struct FCustomDepthTextures;
 struct FDynamicShadowsTaskData;
 
-namespace DistanceField
-{
-	struct FUpdateTrackingBounds;
-};
-
 DECLARE_STATS_GROUP(TEXT("Command List Markers"), STATGROUP_CommandListMarkers, STATCAT_Advanced);
 
 DECLARE_GPU_DRAWCALL_STAT_EXTERN(VirtualTextureUpdate);
@@ -2123,7 +2118,7 @@ public:
 
 	void UpdateGlobalDistanceFieldObjectBuffers(FRDGBuilder& GraphBuilder);
 	void UpdateGlobalHeightFieldObjectBuffers(FRDGBuilder& GraphBuilder);
-	void AddOrRemoveSceneHeightFieldPrimitives(const DistanceField::FUpdateTrackingBounds& UpdateTrackingBounds, bool bSkipAdd = false);
+	void AddOrRemoveSceneHeightFieldPrimitives(bool bSkipAdd = false);
 	void PrepareDistanceFieldScene(FRDGBuilder& GraphBuilder, FRDGExternalAccessQueue& ExternalAccessQueue, bool bSplitDispatch);
 
 	void DrawGPUSkinCacheVisualizationInfoText();
