Fix RHI validation errors on global distance field resources

diff --git a/Engine/Source/Runtime/Renderer/Private/GlobalDistanceField.cpp b/Engine/Source/Runtime/Renderer/Private/GlobalDistanceField.cpp
--- a/Engine/Source/Runtime/Renderer/Private/GlobalDistanceField.cpp
+++ b/Engine/Source/Runtime/Renderer/Private/GlobalDistanceField.cpp
@@ -1855,6 +1855,127 @@ class FGlobalDistanceFieldDebugCS : public FGlobalShader
 
 IMPLEMENT_GLOBAL_SHADER(FGlobalDistanceFieldDebugCS, "/Engine/Private/DistanceField/GlobalDistanceFieldDebug.usf", "GlobalDistanceFieldDebugCS", SF_Compute);
 
+struct FGlobalDistanceFieldInfoRDG
+{
+	FRDGBufferRef PageFreeListAllocatorBuffer;
+	FRDGBufferRef PageFreeListBuffer;
+	FRDGTextureRef PageAtlasTexture;
+	FRDGTextureRef CoverageAtlasTexture;
+	FRDGBufferRef PageObjectGridBuffer;
+	FRDGTextureRef PageTableCombinedTexture;
+	FRDGTextureRef MipTexture;
+	FRDGTextureRef PageTableLayerTextures[GDF_Num];
+};
+
+static void RegisterGlobalDistanceFieldExternalResources(
+	FRDGBuilder& GraphBuilder,
+	const FGlobalDistanceFieldInfo& GlobalDistanceFieldInfo,
+	const FGlobalDFCacheType StartCacheType,
+	FGlobalDistanceFieldInfoRDG& Out)
+{
+	Out.PageFreeListAllocatorBuffer = nullptr;
+	if (GlobalDistanceFieldInfo.PageFreeListAllocatorBuffer)
+	{
+		Out.PageFreeListAllocatorBuffer = GraphBuilder.RegisterExternalBuffer(GlobalDistanceFieldInfo.PageFreeListAllocatorBuffer, TEXT("GlobalDistanceField.PageFreeListAllocator"));
+	}
+
+	Out.PageFreeListBuffer = nullptr;
+	if (GlobalDistanceFieldInfo.PageFreeListBuffer)
+	{
+		Out.PageFreeListBuffer = GraphBuilder.RegisterExternalBuffer(GlobalDistanceFieldInfo.PageFreeListBuffer, TEXT("GlobalDistanceField.PageFreeList"));
+	}
+
+	Out.PageAtlasTexture = nullptr;
+	if (GlobalDistanceFieldInfo.PageAtlasTexture)
+	{
+		Out.PageAtlasTexture = GraphBuilder.RegisterExternalTexture(GlobalDistanceFieldInfo.PageAtlasTexture, TEXT("GlobalDistanceField.PageAtlas"));
+	}
+
+	Out.CoverageAtlasTexture = nullptr;
+	if (GlobalDistanceFieldInfo.CoverageAtlasTexture)
+	{
+		Out.CoverageAtlasTexture = GraphBuilder.RegisterExternalTexture(GlobalDistanceFieldInfo.CoverageAtlasTexture, TEXT("GlobalDistanceField.CoverageAtlas"));
+	}
+
+	Out.PageObjectGridBuffer = nullptr;
+	if (GlobalDistanceFieldInfo.PageObjectGridBuffer)
+	{
+		Out.PageObjectGridBuffer = GraphBuilder.RegisterExternalBuffer(GlobalDistanceFieldInfo.PageObjectGridBuffer, TEXT("GlobalDistanceField.PageObjectGridBuffer"));
+	}
+
+	Out.PageTableCombinedTexture = nullptr;
+	if (GlobalDistanceFieldInfo.PageTableCombinedTexture)
+	{
+		Out.PageTableCombinedTexture = GraphBuilder.RegisterExternalTexture(GlobalDistanceFieldInfo.PageTableCombinedTexture, TEXT("GlobalDistanceField.PageTableCombined"));
+	}
+
+	Out.MipTexture = nullptr;
+	if (GlobalDistanceFieldInfo.MipTexture)
+	{
+		Out.MipTexture = GraphBuilder.RegisterExternalTexture(GlobalDistanceFieldInfo.MipTexture, TEXT("GlobalDistanceField.SDFMips"));
+	}
+
+	for (int32 CacheType = StartCacheType; CacheType < GDF_Num; CacheType++)
+	{
+		Out.PageTableLayerTextures[CacheType] = nullptr;
+		if (GlobalDistanceFieldInfo.PageTableLayerTextures[CacheType])
+		{
+			Out.PageTableLayerTextures[CacheType] = GraphBuilder.RegisterExternalTexture(GlobalDistanceFieldInfo.PageTableLayerTextures[CacheType], TEXT("GlobalDistanceFieldPageTableLayer"));
+		}
+	}
+}
+
+static void FinalizeGlobalDistanceFieldExternalResourceAccess(
+	FRDGBuilder& GraphBuilder,
+	FRDGExternalAccessQueue& ExternalAccessQueue,
+	FGlobalDistanceFieldInfo& GlobalDistanceFieldInfo,
+	const FGlobalDFCacheType StartCacheType,
+	const FGlobalDistanceFieldInfoRDG& In)
+{
+	for (int32 CacheType = StartCacheType; CacheType < GDF_Num; CacheType++)
+	{
+		if (In.PageTableLayerTextures[CacheType])
+		{
+			GlobalDistanceFieldInfo.PageTableLayerTextures[CacheType] = ConvertToExternalAccessTexture(GraphBuilder, ExternalAccessQueue, In.PageTableLayerTextures[CacheType]);
+		}
+	}
+
+	if (In.PageFreeListAllocatorBuffer)
+	{
+		GlobalDistanceFieldInfo.PageFreeListAllocatorBuffer = ConvertToExternalAccessBuffer(GraphBuilder, ExternalAccessQueue, In.PageFreeListAllocatorBuffer);
+	}
+
+	if (In.PageFreeListBuffer)
+	{
+		GlobalDistanceFieldInfo.PageFreeListBuffer = ConvertToExternalAccessBuffer(GraphBuilder, ExternalAccessQueue, In.PageFreeListBuffer);
+	}
+
+	if (In.PageAtlasTexture)
+	{
+		GlobalDistanceFieldInfo.PageAtlasTexture = ConvertToExternalAccessTexture(GraphBuilder, ExternalAccessQueue, In.PageAtlasTexture, ERHIAccess::SRVMask, ERHIPipeline::All);
+	}
+
+	if (In.CoverageAtlasTexture)
+	{
+		GlobalDistanceFieldInfo.CoverageAtlasTexture = ConvertToExternalAccessTexture(GraphBuilder, ExternalAccessQueue, In.CoverageAtlasTexture, ERHIAccess::SRVMask, ERHIPipeline::All);
+	}
+
+	if (In.PageObjectGridBuffer)
+	{
+		GlobalDistanceFieldInfo.PageObjectGridBuffer = ConvertToExternalAccessBuffer(GraphBuilder, ExternalAccessQueue, In.PageObjectGridBuffer, ERHIAccess::SRVMask, ERHIPipeline::All);
+	}
+
+	if (In.PageTableCombinedTexture)
+	{
+		GlobalDistanceFieldInfo.PageTableCombinedTexture = ConvertToExternalAccessTexture(GraphBuilder, ExternalAccessQueue, In.PageTableCombinedTexture, ERHIAccess::SRVMask, ERHIPipeline::All);
+	}
+
+	if (In.MipTexture)
+	{
+		GlobalDistanceFieldInfo.MipTexture = ConvertToExternalAccessTexture(GraphBuilder, ExternalAccessQueue, In.MipTexture, ERHIAccess::SRVMask, ERHIPipeline::All);
+	}
+}
+
 /** 
  * Updates the global distance field for a view.  
  * Typically issues updates for just the newly exposed regions of the volume due to camera movement.
@@ -1872,6 +1993,7 @@ void UpdateGlobalDistanceFieldVolume(
 	RDG_RHI_GPU_STAT_SCOPE(GraphBuilder, GlobalDistanceFieldUpdate);
 
 	const FDistanceFieldSceneData& DistanceFieldSceneData = Scene->DistanceFieldSceneData;
+	bool bNeedsFinalizeAccess = true;
 
 	UpdateGlobalDistanceFieldViewOrigin(View, bLumenEnabled);
 
@@ -1900,49 +2022,20 @@ void UpdateGlobalDistanceFieldVolume(
 		{
 			RDG_EVENT_SCOPE(GraphBuilder, "UpdateGlobalDistanceField");
 
+			bNeedsFinalizeAccess = false;
 			const FGlobalDFCacheType StartCacheType = GAOGlobalDistanceFieldCacheMostlyStaticSeparately ? GDF_MostlyStatic : GDF_Full;
+			FGlobalDistanceFieldInfoRDG GlobalDistanceFieldInfoRDG;
 
-			FRDGBufferRef PageFreeListAllocatorBuffer = nullptr;
-			if (GlobalDistanceFieldInfo.PageFreeListAllocatorBuffer)
-			{
-				PageFreeListAllocatorBuffer = GraphBuilder.RegisterExternalBuffer(GlobalDistanceFieldInfo.PageFreeListAllocatorBuffer, TEXT("GlobalDistanceField.PageFreeListAllocator"));
-			}
-
-			FRDGBufferRef PageFreeListBuffer = nullptr;
-			if (GlobalDistanceFieldInfo.PageFreeListBuffer)
-			{
-				PageFreeListBuffer = GraphBuilder.RegisterExternalBuffer(GlobalDistanceFieldInfo.PageFreeListBuffer, TEXT("GlobalDistanceField.PageFreeList"));
-			}
-
-			FRDGTextureRef PageAtlasTexture = nullptr;
-			if (GlobalDistanceFieldInfo.PageAtlasTexture)
-			{
-				PageAtlasTexture = GraphBuilder.RegisterExternalTexture(GlobalDistanceFieldInfo.PageAtlasTexture, TEXT("GlobalDistanceField.PageAtlas"));
-			}
-
-			FRDGTextureRef CoverageAtlasTexture = nullptr;
-			if (GlobalDistanceFieldInfo.CoverageAtlasTexture)
-			{
-				CoverageAtlasTexture = GraphBuilder.RegisterExternalTexture(GlobalDistanceFieldInfo.CoverageAtlasTexture, TEXT("GlobalDistanceField.CoverageAtlas"));
-			}
-
-			FRDGBufferRef PageObjectGridBuffer = nullptr;
-			if (GlobalDistanceFieldInfo.PageObjectGridBuffer)
-			{
-				PageObjectGridBuffer = GraphBuilder.RegisterExternalBuffer(GlobalDistanceFieldInfo.PageObjectGridBuffer, TEXT("GlobalDistanceField.PageObjectGridBuffer"));
-			}
-
-			FRDGTextureRef PageTableCombinedTexture = nullptr;
-			if (GlobalDistanceFieldInfo.PageTableCombinedTexture)
-			{
-				PageTableCombinedTexture = GraphBuilder.RegisterExternalTexture(GlobalDistanceFieldInfo.PageTableCombinedTexture, TEXT("GlobalDistanceField.PageTableCombined"));
-			}
+			RegisterGlobalDistanceFieldExternalResources(GraphBuilder, GlobalDistanceFieldInfo, StartCacheType, GlobalDistanceFieldInfoRDG);
 
-			FRDGTextureRef MipTexture = nullptr;
-			if (GlobalDistanceFieldInfo.MipTexture)
-			{
-				MipTexture = GraphBuilder.RegisterExternalTexture(GlobalDistanceFieldInfo.MipTexture, TEXT("GlobalDistanceField.SDFMips"));
-			}
+			FRDGBufferRef PageFreeListAllocatorBuffer = GlobalDistanceFieldInfoRDG.PageFreeListAllocatorBuffer;
+			FRDGBufferRef PageFreeListBuffer = GlobalDistanceFieldInfoRDG.PageFreeListBuffer;
+			FRDGTextureRef PageAtlasTexture = GlobalDistanceFieldInfoRDG.PageAtlasTexture;
+			FRDGTextureRef CoverageAtlasTexture = GlobalDistanceFieldInfoRDG.CoverageAtlasTexture;
+			FRDGBufferRef PageObjectGridBuffer = GlobalDistanceFieldInfoRDG.PageObjectGridBuffer;
+			FRDGTextureRef PageTableCombinedTexture = GlobalDistanceFieldInfoRDG.PageTableCombinedTexture;
+			FRDGTextureRef MipTexture = GlobalDistanceFieldInfoRDG.MipTexture;
+			FRDGTextureRef (&PageTableLayerTextures)[GDF_Num] = GlobalDistanceFieldInfoRDG.PageTableLayerTextures;
 
 			FRDGTextureRef TempMipTexture = nullptr;
 			{
@@ -1956,15 +2049,6 @@ void UpdateGlobalDistanceFieldVolume(
 				TempMipTexture = GraphBuilder.CreateTexture(TempMipDesc, TEXT("GlobalDistanceField.TempMip"));
 			}
 
-			FRDGTextureRef PageTableLayerTextures[GDF_Num] = {};
-			for (int32 CacheType = StartCacheType; CacheType < GDF_Num; CacheType++)
-			{
-				if (GlobalDistanceFieldInfo.PageTableLayerTextures[CacheType])
-				{
-					PageTableLayerTextures[CacheType] = GraphBuilder.RegisterExternalTexture(GlobalDistanceFieldInfo.PageTableLayerTextures[CacheType], TEXT("GlobalDistanceFieldPageTableLayer"));
-				}
-			}
-
 			if (View.ViewState && View.ViewState->GlobalDistanceFieldData->bPendingReset)
 			{
 				// Reset all allocators to default
@@ -2749,48 +2833,7 @@ void UpdateGlobalDistanceFieldVolume(
 				}
 			}
 
-			for (int32 CacheType = StartCacheType; CacheType < GDF_Num; CacheType++)
-			{
-				if (PageTableLayerTextures[CacheType])
-				{
-					GlobalDistanceFieldInfo.PageTableLayerTextures[CacheType] = ConvertToExternalAccessTexture(GraphBuilder, ExternalAccessQueue, PageTableLayerTextures[CacheType]);
-				}
-			}
-
-			if (PageFreeListAllocatorBuffer)
-			{
-				GlobalDistanceFieldInfo.PageFreeListAllocatorBuffer = ConvertToExternalAccessBuffer(GraphBuilder, ExternalAccessQueue, PageFreeListAllocatorBuffer);
-			}
-
-			if (PageFreeListBuffer)
-			{
-				GlobalDistanceFieldInfo.PageFreeListBuffer = ConvertToExternalAccessBuffer(GraphBuilder, ExternalAccessQueue, PageFreeListBuffer);
-			}
-
-			if (PageAtlasTexture)
-			{
-				GlobalDistanceFieldInfo.PageAtlasTexture = ConvertToExternalAccessTexture(GraphBuilder, ExternalAccessQueue, PageAtlasTexture);
-			}
-
-			if (CoverageAtlasTexture)
-			{
-				GlobalDistanceFieldInfo.CoverageAtlasTexture = ConvertToExternalAccessTexture(GraphBuilder, ExternalAccessQueue, CoverageAtlasTexture);
-			}
-
-			if (PageObjectGridBuffer)
-			{
-				GlobalDistanceFieldInfo.PageObjectGridBuffer = ConvertToExternalAccessBuffer(GraphBuilder, ExternalAccessQueue, PageObjectGridBuffer, ERHIAccess::SRVMask, ERHIPipeline::All);
-			}
-
-			if (PageTableCombinedTexture)
-			{
-				GlobalDistanceFieldInfo.PageTableCombinedTexture = ConvertToExternalAccessTexture(GraphBuilder, ExternalAccessQueue, PageTableCombinedTexture);
-			}
-
-			if (MipTexture)
-			{
-				GlobalDistanceFieldInfo.MipTexture = ConvertToExternalAccessTexture(GraphBuilder, ExternalAccessQueue, MipTexture);
-			}
+			FinalizeGlobalDistanceFieldExternalResourceAccess(GraphBuilder, ExternalAccessQueue, GlobalDistanceFieldInfo, StartCacheType, GlobalDistanceFieldInfoRDG);
 		}
 
 		if (CVarGlobalDistanceFieldDebug.GetValueOnRenderThread() != 0 
@@ -2848,6 +2891,17 @@ void UpdateGlobalDistanceFieldVolume(
 		}
 	}
 
+#if ENABLE_RHI_VALIDATION
+	if (bNeedsFinalizeAccess)
+	{
+		const FGlobalDFCacheType StartCacheType = GAOGlobalDistanceFieldCacheMostlyStaticSeparately ? GDF_MostlyStatic : GDF_Full;
+		FGlobalDistanceFieldInfoRDG GlobalDistanceFieldInfoRDG;
+
+		RegisterGlobalDistanceFieldExternalResources(GraphBuilder, GlobalDistanceFieldInfo, StartCacheType, GlobalDistanceFieldInfoRDG);
+		FinalizeGlobalDistanceFieldExternalResourceAccess(GraphBuilder, ExternalAccessQueue, GlobalDistanceFieldInfo, StartCacheType, GlobalDistanceFieldInfoRDG);
+	}
+#endif
+
 	if (GDFReadbackRequest && GlobalDistanceFieldInfo.Clipmaps.Num() > 0)
 	{
 		// Read back a clipmap
