Added support for using a stride to skip marking VSM every pixel, this improves performance on page marking (at the potential loss of some quality).
 - Also moved Single Layer Water tile classification pass to the depth pre-pass and generated a mask to let the VSM page marking skip fetching water data for every pixel.

diff --git a/Engine/Shaders/Private/SingleLayerWaterComposite.usf b/Engine/Shaders/Private/SingleLayerWaterComposite.usf
--- a/Engine/Shaders/Private/SingleLayerWaterComposite.usf
+++ b/Engine/Shaders/Private/SingleLayerWaterComposite.usf
@@ -11,6 +11,7 @@
 #include "ShadingModels.ush"
 #include "SceneTextureParameters.ush"
 #include "LightGridCommon.ush"
+#include "/Engine/Shared/SingleLayerWaterDefinitions.h"
 #include "SingleLayerWaterCommon.ush"
 
 #define ENABLE_SKY_LIGHT 1
@@ -268,20 +269,42 @@ SingleLayerWaterCompositeOutput SingleLayerWaterComposite(float2 BufferUV, float
 
 #ifdef TILE_CATERGORISATION_SHADER
 
+#if USE_WATER_PRE_PASS_STENCIL
+Texture2D<uint2> WaterDepthStencilTexture;
+#endif
 
 uint VertexCountPerInstanceIndirect;
-
+int2 TiledViewRes;
 RWBuffer<uint> DrawIndirectDataUAV;
 RWBuffer<uint> DispatchIndirectDataUAV;
 RWBuffer<uint> WaterTileListDataUAV;
+RWStructuredBuffer<uint> TileMaskBufferOut;
 
 #if COMPILER_SUPPORTS_WAVE_VOTE
 groupshared uint bAnyWaterPixels;
 #else
-groupshared bool ContainsWater[WORK_TILE_SIZE * WORK_TILE_SIZE];
+groupshared bool ContainsWater[SLW_TILE_SIZE_XY * SLW_TILE_SIZE_XY];
 #endif
 
-[numthreads(WORK_TILE_SIZE, WORK_TILE_SIZE, 1)]
+bool DoesPixelContainWater(uint2 PixelPos)
+{
+#if USE_WATER_PRE_PASS_STENCIL
+	uint Stencil = WaterDepthStencilTexture.Load(uint3(PixelPos, 0)) STENCIL_COMPONENT_SWIZZLE;
+	return (Stencil & 1U) != 0U;
+#else // !USE_WATER_PRE_PASS_STENCIL
+#if STRATA_ENABLED
+	FStrataAddressing StrataAddressing = GetStrataPixelDataByteOffset(PixelPos, uint2(View.BufferSizeAndInvSize.xy), Strata.MaxBytesPerPixel);
+	FStrataPixelHeader StrataPixelHeader = UnpackStrataHeaderIn(Strata.MaterialTextureArray, StrataAddressing, Strata.TopLayerTexture);
+	return (StrataPixelHeader.BSDFCount > 0 && IsSingleLayerWater(StrataPixelHeader));
+#else
+	float2 BufferUV = (PixelPos + 0.5f) * View.BufferSizeAndInvSize.zw;
+	FGBufferData GBuffer = GetGBufferDataFromSceneTextures(BufferUV);
+	return (GBuffer.ShadingModelID == SHADINGMODELID_SINGLELAYERWATER);
+#endif
+#endif // USE_WATER_PRE_PASS_STENCIL
+}
+
+[numthreads(SLW_TILE_SIZE_XY, SLW_TILE_SIZE_XY, 1)]
 void WaterTileCatergorisationCS(uint3 ThreadId : SV_DispatchThreadID, uint3 GroupId : SV_GroupID, uint3 GroupThreadId : SV_GroupThreadID)
 {
 	if (ThreadId.x == 0 && ThreadId.y == 0 && GroupThreadId.x == 0 && GroupThreadId.y == 0)
@@ -298,25 +321,15 @@ void WaterTileCatergorisationCS(uint3 ThreadId : SV_DispatchThreadID, uint3 Grou
 
 #if 0
 	// Slow reference path
-	const int LinearIndex = GroupThreadId.y * WORK_TILE_SIZE + GroupThreadId.x;
+	const int LinearIndex = GroupThreadId.y * SLW_TILE_SIZE_XY + GroupThreadId.x;
 	if (LinearIndex < 1)
 	{
 		bool bContainsWater = false;
-		for (uint i = 0; i < WORK_TILE_SIZE; ++i)
+		for (uint i = 0; i < SLW_TILE_SIZE_XY; ++i)
 		{
-			for (uint j = 0; j < WORK_TILE_SIZE; ++j)
+			for (uint j = 0; j < SLW_TILE_SIZE_XY; ++j)
 			{
-				uint2 PixelPos = (ThreadId.xy + uint2(i, j)) + View.ViewRectMin.xy;
-				float2 BufferUV = (PixelPos + 0.5f) * View.BufferSizeAndInvSize.zw;
-
-			#if STRATA_ENABLED
-				FStrataAddressing StrataAddressing = GetStrataPixelDataByteOffset(PixelPos, uint2(View.BufferSizeAndInvSize.xy), Strata.MaxBytesPerPixel);
-				FStrataPixelHeader StrataPixelHeader = UnpackStrataHeaderIn(Strata.MaterialTextureArray, StrataAddressing, Strata.TopLayerTexture);
-				bContainsWater = bContainsWater || (StrataPixelHeader.BSDFCount > 0 && IsSingleLayerWater(StrataPixelHeader));
-			#else
-				FGBufferData GBuffer = GetGBufferDataFromSceneTextures(BufferUV);
-				bContainsWater = bContainsWater || GBuffer.ShadingModelID == SHADINGMODELID_SINGLELAYERWATER;
-			#endif
+				bContainsWater = bContainsWater || DoesPixelContainWater((ThreadId.xy + uint2(i, j)) + View.ViewRectMin.xy);
 			}
 		}
 		if (bContainsWater)
@@ -333,17 +346,7 @@ void WaterTileCatergorisationCS(uint3 ThreadId : SV_DispatchThreadID, uint3 Grou
 	bool bContainsWater = false;
 	if(ThreadId.x < uint(View.BufferSizeAndInvSize.x) && ThreadId.y < uint(View.BufferSizeAndInvSize.y))
 	{
-		uint2 PixelPos = ThreadId.xy + View.ViewRectMin.xy;
-		float2 BufferUV = (PixelPos + 0.5f) * View.BufferSizeAndInvSize.zw;
-
-	#if STRATA_ENABLED
-		FStrataAddressing StrataAddressing = GetStrataPixelDataByteOffset(PixelPos, uint2(View.BufferSizeAndInvSize.xy), Strata.MaxBytesPerPixel);
-		FStrataPixelHeader StrataPixelHeader = UnpackStrataHeaderIn(Strata.MaterialTextureArray, StrataAddressing, Strata.TopLayerTexture);
-		bContainsWater = bContainsWater || (StrataPixelHeader.BSDFCount > 0 && IsSingleLayerWater(StrataPixelHeader));
-	#else
-		FGBufferData GBuffer = GetGBufferDataFromSceneTextures(BufferUV);
-		bContainsWater = bContainsWater || (GBuffer.ShadingModelID == SHADINGMODELID_SINGLELAYERWATER);
-	#endif
+		bContainsWater = bContainsWater || DoesPixelContainWater(ThreadId.xy + View.ViewRectMin.xy);
 	}
 
 	bool bWriteTile = false;
@@ -370,7 +373,7 @@ void WaterTileCatergorisationCS(uint3 ThreadId : SV_DispatchThreadID, uint3 Grou
 
 #else
 
-	const int LinearIndex = GroupThreadId.y * WORK_TILE_SIZE + GroupThreadId.x;
+	const int LinearIndex = GroupThreadId.y * SLW_TILE_SIZE_XY + GroupThreadId.x;
 	ContainsWater[LinearIndex] = bContainsWater;
 	GroupMemoryBarrierWithGroupSync();
 
@@ -416,6 +419,10 @@ void WaterTileCatergorisationCS(uint3 ThreadId : SV_DispatchThreadID, uint3 Grou
 		InterlockedAdd(DrawIndirectDataUAV[1], 1, WriteToIndex);
 		InterlockedAdd(DispatchIndirectDataUAV[0], 1);
 		WaterTileListDataUAV[WriteToIndex] = GroupId.x | (GroupId.y << 16); // assumes 16bit is enought to represent a tiled resolution up to 65,535 :)
+
+		// Set bit to indicate tile is occupied.
+		uint MaskLinearIndex = TiledViewRes.x * GroupId.y + GroupId.x;
+		InterlockedOr(TileMaskBufferOut[MaskLinearIndex / 32U], 1U << (MaskLinearIndex % 32U));
 	}
 #endif
 } 
@@ -433,11 +440,11 @@ void WaterTileVS(
 	out float4 Position : SV_POSITION) 
 {
 	uint TileData = TileListData[InstanceId.x];
-	uint2 TileOrigin = uint2((TileData & 0x0000FFFF) * WORK_TILE_SIZE, (TileData >> 16) * WORK_TILE_SIZE);
+	uint2 TileOrigin = uint2((TileData & 0x0000FFFF) * SLW_TILE_SIZE_XY, (TileData >> 16) * SLW_TILE_SIZE_XY);
 
 	uint2 TileVertex = TileOrigin;
-	TileVertex.x += VertexId == 1 || VertexId == 2 || VertexId == 4 ? WORK_TILE_SIZE : 0;
-	TileVertex.y += VertexId == 2 || VertexId == 4 || VertexId == 5 ? WORK_TILE_SIZE : 0;
+	TileVertex.x += VertexId == 1 || VertexId == 2 || VertexId == 4 ? SLW_TILE_SIZE_XY : 0;
+	TileVertex.y += VertexId == 2 || VertexId == 4 || VertexId == 5 ? SLW_TILE_SIZE_XY : 0;
 
 	// View port is set on the view rect. So no offset are needed.
 	Position = float4(float2(TileVertex) * View.ViewSizeAndInvSize.zw * float2(2.0f, -2.0f) + float2(-1.0, 1.0f), 0.5f, 1.0f);
diff --git a/Engine/Shaders/Private/VirtualShadowMaps/VirtualShadowMapPageManagement.usf b/Engine/Shaders/Private/VirtualShadowMaps/VirtualShadowMapPageManagement.usf
--- a/Engine/Shaders/Private/VirtualShadowMaps/VirtualShadowMapPageManagement.usf
+++ b/Engine/Shaders/Private/VirtualShadowMaps/VirtualShadowMapPageManagement.usf
@@ -4,6 +4,7 @@
 	VirtualShadowMapPageManagement.usf: 
 =============================================================================*/
 
+#include "/Engine/Shared/SingleLayerWaterDefinitions.h"
 #include "../HairStrands/HairStrandsVisibilityCommonStruct.ush"
 #include "../Common.ush"
 #include "../WaveOpUtil.ush"
@@ -182,10 +183,16 @@ void PruneLightGridCS(uint GridLinearIndex : SV_DispatchThreadID)
 
 StructuredBuffer<uint> PrunedLightGridData;
 StructuredBuffer<uint> PrunedNumCulledLightsGrid;
-Texture2D<float> SingeLayerWaterDepthTexture;
+uint2 PixelStride;
 
+#if PERMUTATION_INPUT_TYPE == INPUT_TYPE_GBUFFER_AND_WATER_DEPTH
+Texture2D<float> SingleLayerWaterDepthTexture;
+StructuredBuffer< uint > SingleLayerWaterTileMask;
+int2 SingleLayerWaterTileViewRes;
+#endif
 
-[numthreads(VSM_DEFAULT_CS_GROUP_XY, VSM_DEFAULT_CS_GROUP_XY, 1)]
+// Note: we use the tile size defined by the water as the group-size - this is needed because the tile mask testing code relies on the size being the same to scalarize efficiently.
+[numthreads(SLW_TILE_SIZE_XY, SLW_TILE_SIZE_XY, 1)]
 void GeneratePageFlagsFromPixels(
 	uint3 InGroupId : SV_GroupID,
 	uint  GroupIndex : SV_GroupIndex,
@@ -207,8 +214,8 @@ void GeneratePageFlagsFromPixels(
 #else // PERMUTATION_INPUT_TYPE == INPUT_TYPE_GBUFFER || PERMUTATION_INPUT_TYPE == INPUT_TYPE_GBUFFER_AND_WATER_DEPTH
 	const uint2 GroupId = InGroupId.xy;
 #endif
-	// Morton order within a group so page access/atomics are more coherent and wave-swizzled gradients are possible.
-	const uint2 PixelPos = uint2(View.ViewRectMin.xy) + VSM_DEFAULT_CS_GROUP_XY * GroupId.xy + MortonDecode(GroupIndex);
+	const uint2 PixelLocalPos = DispatchThreadId.xy * PixelStride;
+	const uint2 PixelPos = uint2(View.ViewRectMin.xy) + PixelLocalPos;
 
 	if (any(PixelPos >= uint2(View.ViewRectMin.xy + View.ViewSizeAndInvSize.xy)))
 	{
@@ -244,8 +251,38 @@ void GeneratePageFlagsFromPixels(
 #endif // STRATA_ENABLED
 
 #if PERMUTATION_INPUT_TYPE == INPUT_TYPE_GBUFFER_AND_WATER_DEPTH
-		DeviceZWater = SingeLayerWaterDepthTexture.Load(uint3(PixelPos, 0)).x;
-		bIsWaterDepthValid = DeviceZWater != DeviceZ;
+
+		// uint2 WaterTileIndex = DispatchThreadId.xy * PixelStride / SLW_TILE_SIZE_XY; <=>
+		// uint2 WaterTileIndex = (groupId*SLW_TILE_SIZE_XY + threadgroupIndex) * PixelStride / SLW_TILE_SIZE_XY; <=>
+		// uint2 WaterTileIndex = groupId * PixelStride + threadgroupIndex * PixelStride / SLW_TILE_SIZE_XY; <=>
+		// uint2 WaterTileMin = groupId * PixelStride + {0,0} * PixelStride / SLW_TILE_SIZE_XY; <=> groupId * PixelStride
+		// uint2 WaterTileMax = groupId * PixelStride + (SLW_TILE_SIZE_XY-1) * PixelStride / SLW_TILE_SIZE_XY; <=> 
+		// uint2 WaterTileMax = groupId * PixelStride + (SLW_TILE_SIZE_XY * PixelStride - PixelStride)  / SLW_TILE_SIZE_XY; <=> 
+		// uint2 WaterTileMax = groupId * PixelStride + PixelStride - PixelStride / SLW_TILE_SIZE_XY; <~> 
+		// uint2 WaterTileMax = groupId * PixelStride + PixelStride; <~> 
+
+		// Bit-scalarized version, must test rect
+		uint2 WaterTileMin = GroupId * PixelStride;
+		uint2 WaterTileMax = WaterTileMin + PixelStride;
+		
+		bIsWaterDepthValid = false;
+		for (uint WaterTileY = WaterTileMin.y; WaterTileY < WaterTileMax.y && !bIsWaterDepthValid; ++WaterTileY)
+		{
+			for (uint WaterTileX = WaterTileMin.x; WaterTileX < WaterTileMax.x; ++WaterTileX)
+			{
+				uint MaskLinearIndex = uint(SingleLayerWaterTileViewRes.x) * WaterTileY + WaterTileX;
+				if ((SingleLayerWaterTileMask[MaskLinearIndex / 32U] & (1U << (MaskLinearIndex % 32U))) != 0U)
+				{
+					bIsWaterDepthValid = true;
+					break;
+				}
+			}
+		}
+		if (bIsWaterDepthValid)
+		{
+			DeviceZWater = SingleLayerWaterDepthTexture.Load(uint3(PixelPos, 0)).x;
+			bIsWaterDepthValid = DeviceZWater != DeviceZ;
+		}
 #endif // PERMUTATION_INPUT_TYPE == INPUT_TYPE_GBUFFER_AND_WATER_DEPTH
 
 	}
diff --git a/Engine/Shaders/Shared/SingleLayerWaterDefinitions.h b/Engine/Shaders/Shared/SingleLayerWaterDefinitions.h
new file mode 100644
--- /dev/null
+++ b/Engine/Shaders/Shared/SingleLayerWaterDefinitions.h
@@ -0,0 +1,10 @@
+// Copyright Epic Games, Inc. All Rights Reserved.
+
+/*================================================================================================
+	SingleLayerWaterDefinitions.h: 
+	Used in shaders and C++ code to define common constants
+	!!! Changing this file requires recompilation of the engine !!!
+=================================================================================================*/
+#pragma once
+
+#define SLW_TILE_SIZE_XY (8U)
\ No newline at end of file
diff --git a/Engine/Source/Runtime/Renderer/Private/DeferredShadingRenderer.cpp b/Engine/Source/Runtime/Renderer/Private/DeferredShadingRenderer.cpp
--- a/Engine/Source/Runtime/Renderer/Private/DeferredShadingRenderer.cpp
+++ b/Engine/Source/Runtime/Renderer/Private/DeferredShadingRenderer.cpp
@@ -3241,11 +3241,12 @@ void FDeferredShadingSceneRenderer::Render(FRDGBuilder& GraphBuilder)
 	FRDGTextureRef LightingChannelsTexture = CopyStencilToLightingChannelTexture(GraphBuilder, SceneTextures.Stencil);
 
 	// Single layer water depth prepass. Needs to run before VSM page allocation.
-	FRDGTextureMSAA SingleLayerWaterDepthPrepassOutput;
+	const FSingleLayerWaterPrePassResult* SingleLayerWaterPrePassResult = nullptr;
+
 	const bool bShouldRenderSingleLayerWaterDepthPrepass = !bHasRayTracedOverlay && ShouldRenderSingleLayerWaterDepthPrepass(Views);
 	if (bShouldRenderSingleLayerWaterDepthPrepass)
 	{
-		RenderSingleLayerWaterDepthPrepass(GraphBuilder, SceneTextures, SingleLayerWaterDepthPrepassOutput);
+		SingleLayerWaterPrePassResult = RenderSingleLayerWaterDepthPrepass(GraphBuilder, SceneTextures);
 	}
 
 	FAsyncLumenIndirectLightingOutputs AsyncLumenIndirectLightingOutputs;
@@ -3272,7 +3273,7 @@ void FDeferredShadingSceneRenderer::Render(FRDGBuilder& GraphBuilder)
 				RDG_GPU_STAT_SCOPE(GraphBuilder, ShadowDepths);
 
 				ensureMsgf(AreLightsInLightGrid(), TEXT("Virtual shadow map setup requires local lights to be injected into the light grid (this may be caused by 'r.LightCulling.Quality=0')."));
-				VirtualShadowMapArray.BuildPageAllocations(GraphBuilder, SceneTextures, Views, ViewFamily.EngineShowFlags, SortedLightSet, VisibleLightInfos, NaniteRasterResults, SingleLayerWaterDepthPrepassOutput.Resolve);
+				VirtualShadowMapArray.BuildPageAllocations(GraphBuilder, SceneTextures, Views, ViewFamily.EngineShowFlags, SortedLightSet, VisibleLightInfos, NaniteRasterResults, SingleLayerWaterPrePassResult);
 			}
 
 			RenderShadowDepthMaps(GraphBuilder, InstanceCullingManager);
@@ -3534,14 +3535,14 @@ void FDeferredShadingSceneRenderer::Render(FRDGBuilder& GraphBuilder)
 		}
 
 		GraphBuilder.SetCommandListStat(GET_STATID(STAT_CLM_WaterPass));
-		RenderSingleLayerWater(GraphBuilder, SceneTextures, SingleLayerWaterDepthPrepassOutput, bShouldRenderVolumetricCloud, SceneWithoutWaterTextures, LumenFrameTemporaries);
+		RenderSingleLayerWater(GraphBuilder, SceneTextures, SingleLayerWaterPrePassResult, bShouldRenderVolumetricCloud, SceneWithoutWaterTextures, LumenFrameTemporaries);
 
 		// Replace main depth texture with the output of the SLW depth prepass which contains the scene + water.
 		// Note: Stencil now has all water bits marked with 1. As long as no other passes after this point want to read the depth buffer,
 		// a stencil clear should not be necessary here.
-		if (bShouldRenderSingleLayerWaterDepthPrepass)
+		if (SingleLayerWaterPrePassResult)
 		{
-			SceneTextures.Depth = SingleLayerWaterDepthPrepassOutput;
+			SceneTextures.Depth = SingleLayerWaterPrePassResult->DepthPrepassTexture;
 		}
 	}
 
diff --git a/Engine/Source/Runtime/Renderer/Private/DeferredShadingRenderer.h b/Engine/Source/Runtime/Renderer/Private/DeferredShadingRenderer.h
--- a/Engine/Source/Runtime/Renderer/Private/DeferredShadingRenderer.h
+++ b/Engine/Source/Runtime/Renderer/Private/DeferredShadingRenderer.h
@@ -48,6 +48,7 @@ struct FHairStrandsTransmittanceMaskData;
 struct FVolumetricFogLocalLightFunctionInfo;
 struct FTranslucencyLightingVolumeTextures;
 struct FLumenSceneFrameTemporaries;
+struct FSingleLayerWaterPrePassResult;
 
 /**   
  * Data for rendering meshes into Lumen Lighting Cards.
@@ -364,16 +365,17 @@ public:
 		FRDGBuilder& GraphBuilder,
 		FSceneTextures& SceneTextures,
 		bool bDoParallelPass);
-
-	void RenderSingleLayerWaterDepthPrepass(
+	/**
+	 * Runs water pre-pass if enabled and returns an RDG-allocated object with intermediates, or null.
+	 */
+	FSingleLayerWaterPrePassResult* RenderSingleLayerWaterDepthPrepass(
 		FRDGBuilder& GraphBuilder,
-		const FSceneTextures& SceneTextures,
-		FRDGTextureMSAA& OutDepthPrepassTexture);
+		const FSceneTextures& SceneTextures);
 
 	void RenderSingleLayerWater(
 		FRDGBuilder& GraphBuilder,
 		const FSceneTextures& SceneTextures,
-		const FRDGTextureMSAA& SingleLayerWaterDepthPrepassTexture,
+		const FSingleLayerWaterPrePassResult* SingleLayerWaterPrePassResult,
 		bool bShouldRenderVolumetricCloud,
 		FSceneWithoutWaterTextures& SceneWithoutWaterTextures,
 		FLumenSceneFrameTemporaries& LumenFrameTemporaries);
@@ -382,12 +384,13 @@ public:
 		FRDGBuilder& GraphBuilder,
 		const FSceneTextures& SceneTextures,
 		const FSceneWithoutWaterTextures& SceneWithoutWaterTextures,
-		const FRDGTextureMSAA& SingleLayerWaterDepthPrepassTexture);
+		const FSingleLayerWaterPrePassResult* SingleLayerWaterPrePassResult);
 
 	void RenderSingleLayerWaterReflections(
 		FRDGBuilder& GraphBuilder,
 		const FSceneTextures& SceneTextures,
 		const FSceneWithoutWaterTextures& SceneWithoutWaterTextures,
+		const FSingleLayerWaterPrePassResult* SingleLayerWaterPrePassResult,
 		FLumenSceneFrameTemporaries& LumenFrameTemporaries);
 
 	void RenderOcclusion(
diff --git a/Engine/Source/Runtime/Renderer/Private/SingleLayerWaterRendering.cpp b/Engine/Source/Runtime/Renderer/Private/SingleLayerWaterRendering.cpp
--- a/Engine/Source/Runtime/Renderer/Private/SingleLayerWaterRendering.cpp
+++ b/Engine/Source/Runtime/Renderer/Private/SingleLayerWaterRendering.cpp
@@ -226,17 +226,20 @@ class FWaterTileCategorisationCS : public FGlobalShader
 	DECLARE_GLOBAL_SHADER(FWaterTileCategorisationCS);
 	SHADER_USE_PARAMETER_STRUCT(FWaterTileCategorisationCS, FGlobalShader)
 
-	static int32 GetTileSize()
-	{
-		return 8;
-	}
+	class FUsePrepassStencil : SHADER_PERMUTATION_BOOL("USE_WATER_PRE_PASS_STENCIL");
+	using FPermutationDomain = TShaderPermutationDomain<FUsePrepassStencil>;
 
 	BEGIN_SHADER_PARAMETER_STRUCT(FParameters, )
-		SHADER_PARAMETER_STRUCT_INCLUDE(FSingleLayerWaterCommonShaderParameters, CommonParameters)
+		SHADER_PARAMETER_STRUCT_INCLUDE(FSceneTextureParameters, SceneTextures)	// Water scene texture
+		SHADER_PARAMETER_RDG_UNIFORM_BUFFER(FStrataGlobalUniformParameters, Strata)
+		SHADER_PARAMETER_STRUCT_INCLUDE(FViewShaderParameters, View)
+		SHADER_PARAMETER_RDG_TEXTURE_SRV(Texture2D, WaterDepthStencilTexture)
 		SHADER_PARAMETER(uint32, VertexCountPerInstanceIndirect)
+		SHADER_PARAMETER(FIntPoint, TiledViewRes)
 		SHADER_PARAMETER_RDG_BUFFER_UAV(RWBuffer<uint>, DrawIndirectDataUAV)
 		SHADER_PARAMETER_RDG_BUFFER_UAV(RWBuffer<uint>, DispatchIndirectDataUAV)
 		SHADER_PARAMETER_RDG_BUFFER_UAV(RWBuffer<uint>, WaterTileListDataUAV)
+		SHADER_PARAMETER_RDG_BUFFER_UAV(RWStructuredBuffer<uint>, TileMaskBufferOut)
 	END_SHADER_PARAMETER_STRUCT()
 
 	static FPermutationDomain RemapPermutation(FPermutationDomain PermutationVector)
@@ -252,7 +255,6 @@ class FWaterTileCategorisationCS : public FGlobalShader
 	static void ModifyCompilationEnvironment(const FGlobalShaderPermutationParameters& Parameters, FShaderCompilerEnvironment& OutEnvironment)
 	{
 		OutEnvironment.SetDefine(TEXT("TILE_CATERGORISATION_SHADER"), 1.0f);
-		OutEnvironment.SetDefine(TEXT("WORK_TILE_SIZE"), GetTileSize());
 		FGlobalShader::ModifyCompilationEnvironment(Parameters, OutEnvironment);
 	}
 };
@@ -267,7 +269,6 @@ bool FWaterTileVS::ShouldCompilePermutation(const FGlobalShaderPermutationParame
 void FWaterTileVS::ModifyCompilationEnvironment(const FGlobalShaderPermutationParameters& Parameters, FShaderCompilerEnvironment& OutEnvironment)
 	{
 		OutEnvironment.SetDefine(TEXT("TILE_VERTEX_SHADER"), 1.0f);
-		OutEnvironment.SetDefine(TEXT("WORK_TILE_SIZE"), FWaterTileCategorisationCS::GetTileSize());
 		FGlobalShader::ModifyCompilationEnvironment(Parameters, OutEnvironment);
 	}
 
@@ -319,9 +320,9 @@ public:
 		SHADER_PARAMETER_RDG_TEXTURE(Texture2D, DepthTexture)
 		SHADER_PARAMETER_RDG_TEXTURE(Texture2DMS, DepthTextureMS)
 		RENDER_TARGET_BINDING_SLOTS()
-		END_SHADER_PARAMETER_STRUCT()
+	END_SHADER_PARAMETER_STRUCT()
 
-		static bool ShouldCompilePermutation(const FGlobalShaderPermutationParameters& Parameters)
+	static bool ShouldCompilePermutation(const FGlobalShaderPermutationParameters& Parameters)
 	{
 		return IsFeatureLevelSupported(Parameters.Platform, ERHIFeatureLevel::SM5);
 	}
@@ -343,7 +344,63 @@ static FSingleLayerWaterDepthPassParameters* GetSingleLayerWaterDepthPassParamet
 	return PassParameters;
 }
 
-void FDeferredShadingSceneRenderer::RenderSingleLayerWaterDepthPrepass(FRDGBuilder& GraphBuilder, const FSceneTextures& SceneTextures, FRDGTextureMSAA& OutDepthPrepassTexture)
+
+static FSingleLayerWaterTileClassification ClassifyTiles(FRDGBuilder& GraphBuilder, const FViewInfo &View, const FSceneTextures& SceneTextures, const FRDGTextureRef& DepthPrepassTexture)
+{
+	FSingleLayerWaterTileClassification Result;
+	const bool bRunTiled = UseSingleLayerWaterIndirectDraw(View.GetShaderPlatform()) && CVarWaterSingleLayerTiledComposite.GetValueOnRenderThread();
+	if (bRunTiled)
+	{
+		FIntPoint ViewRes(View.ViewRect.Width(), View.ViewRect.Height());
+		Result.TiledViewRes = FIntPoint::DivideAndRoundUp(ViewRes, SLW_TILE_SIZE_XY);
+
+		Result.TiledReflection.DrawIndirectParametersBuffer = GraphBuilder.CreateBuffer(FRDGBufferDesc::CreateIndirectDesc<FRHIDrawIndirectParameters>(), TEXT("SLW.WaterIndirectDrawParameters"));
+		Result.TiledReflection.DispatchIndirectParametersBuffer = GraphBuilder.CreateBuffer(FRDGBufferDesc::CreateIndirectDesc<FRHIDispatchIndirectParameters>(1), TEXT("SLW.WaterIndirectDispatchParameters"));
+
+		FRDGBufferRef TileListDataBuffer = GraphBuilder.CreateBuffer(FRDGBufferDesc::CreateBufferDesc(sizeof(uint32), Result.TiledViewRes.X * Result.TiledViewRes.Y), TEXT("SLW.TileListDataBuffer"));
+		Result.TiledReflection.TileListDataBufferSRV = GraphBuilder.CreateSRV(TileListDataBuffer, PF_R32_UINT);
+
+		FRDGBufferUAVRef DrawIndirectParametersBufferUAV = GraphBuilder.CreateUAV(Result.TiledReflection.DrawIndirectParametersBuffer);
+		FRDGBufferUAVRef DispatchIndirectParametersBufferUAV = GraphBuilder.CreateUAV(Result.TiledReflection.DispatchIndirectParametersBuffer);
+
+		// Allocate buffer with 1 bit / tile
+		Result.TileMaskBuffer = GraphBuilder.CreateBuffer(FRDGBufferDesc::CreateStructuredDesc(sizeof(uint32), FMath::DivideAndRoundUp(Result.TiledViewRes.X * Result.TiledViewRes.Y, 32)), TEXT("SLW.TileMaskBuffer"));
+		FRDGBufferUAVRef TileMaskBufferUAV = GraphBuilder.CreateUAV(Result.TileMaskBuffer);
+		AddClearUAVPass(GraphBuilder, TileMaskBufferUAV, 0);
+
+		// Clear DrawIndirectParametersBuffer
+		AddClearUAVPass(GraphBuilder, DrawIndirectParametersBufferUAV, 0);
+		AddClearUAVPass(GraphBuilder, DispatchIndirectParametersBufferUAV, 0);
+
+		// Categorization based on SHADING_MODEL_ID
+		{
+			FWaterTileCategorisationCS::FPermutationDomain PermutationVector;
+			PermutationVector.Set<FWaterTileCategorisationCS::FUsePrepassStencil>(DepthPrepassTexture != nullptr);
+			TShaderMapRef<FWaterTileCategorisationCS> ComputeShader(View.ShaderMap, PermutationVector);
+
+			FWaterTileCategorisationCS::FParameters* PassParameters = GraphBuilder.AllocParameters<FWaterTileCategorisationCS::FParameters>();
+
+			PassParameters->SceneTextures = GetSceneTextureParameters(GraphBuilder, SceneTextures);
+			PassParameters->View = View.GetShaderParameters();
+			PassParameters->Strata = Strata::BindStrataGlobalUniformParameters(View);
+			PassParameters->TiledViewRes = Result.TiledViewRes;
+
+			PassParameters->VertexCountPerInstanceIndirect = GRHISupportsRectTopology ? 3 : 6;
+			PassParameters->DrawIndirectDataUAV = DrawIndirectParametersBufferUAV;
+			PassParameters->DispatchIndirectDataUAV = DispatchIndirectParametersBufferUAV;
+			PassParameters->WaterTileListDataUAV = GraphBuilder.CreateUAV(TileListDataBuffer, PF_R32_UINT);
+
+			PassParameters->WaterDepthStencilTexture = DepthPrepassTexture ? GraphBuilder.CreateSRV(FRDGTextureSRVDesc::CreateWithPixelFormat(DepthPrepassTexture, PF_X24_G8)) : nullptr;
+
+			PassParameters->TileMaskBufferOut = TileMaskBufferUAV;
+
+			FComputeShaderUtils::AddPass(GraphBuilder, RDG_EVENT_NAME("SLW::TileCategorisation"), ComputeShader, PassParameters, FIntVector(Result.TiledViewRes.X, Result.TiledViewRes.Y, 1));
+		}
+	}
+	return Result;
+}
+
+FSingleLayerWaterPrePassResult* FDeferredShadingSceneRenderer::RenderSingleLayerWaterDepthPrepass(FRDGBuilder& GraphBuilder, const FSceneTextures& SceneTextures)
 {
 	RDG_CSV_STAT_EXCLUSIVE_SCOPE(GraphBuilder, Water);
 	SCOPED_NAMED_EVENT(FDeferredShadingSceneRenderer_RenderSingleLayerWaterDepthPrepass, FColor::Emerald);
@@ -351,6 +408,10 @@ void FDeferredShadingSceneRenderer::RenderSingleLayerWaterDepthPrepass(FRDGBuild
 	RDG_EVENT_SCOPE(GraphBuilder, "SingleLayerWaterDepthPrepass");
 	RDG_GPU_STAT_SCOPE(GraphBuilder, SingleLayerWaterDepthPrepass);
 
+	FSingleLayerWaterPrePassResult* Result = GraphBuilder.AllocObject<FSingleLayerWaterPrePassResult>();
+	Result->ViewTileClassification.SetNum(Views.Num());
+
+	FRDGTextureMSAA &OutDepthPrepassTexture = Result->DepthPrepassTexture;
 	// Create an identical copy of the main depth buffer
 	{
 		const FRDGTextureDesc& DepthPrepassTextureDesc = SceneTextures.Depth.Target->Desc;
@@ -458,6 +519,18 @@ void FDeferredShadingSceneRenderer::RenderSingleLayerWaterDepthPrepass(FRDGBuild
 	}
 
 	AddResolveSceneDepthPass(GraphBuilder, Views, OutDepthPrepassTexture);
+
+	// Run classification pass.
+	for (int32 ViewIndex = 0; ViewIndex < Views.Num(); ++ViewIndex)
+	{
+		FViewInfo& View = Views[ViewIndex];
+		if (UseSingleLayerWaterIndirectDraw(View.GetShaderPlatform()) && CVarWaterSingleLayerTiledComposite.GetValueOnRenderThread())
+		{
+			Result->ViewTileClassification[ViewIndex] = ClassifyTiles(GraphBuilder, View, SceneTextures, OutDepthPrepassTexture.Resolve);
+		}
+	}
+
+	return Result;
 }
 
 static FSceneWithoutWaterTextures AddCopySceneWithoutWaterPass(
@@ -590,6 +663,7 @@ void FDeferredShadingSceneRenderer::RenderSingleLayerWaterReflections(
 	FRDGBuilder& GraphBuilder,
 	const FSceneTextures& SceneTextures,
 	const FSceneWithoutWaterTextures& SceneWithoutWaterTextures,
+	const FSingleLayerWaterPrePassResult* SingleLayerWaterPrePassResult,
 	FLumenSceneFrameTemporaries& LumenFrameTemporaries)
 {
 	if (CVarWaterSingleLayer.GetValueOnRenderThread() <= 0 || CVarWaterSingleLayerReflection.GetValueOnRenderThread() <= 0)
@@ -650,41 +724,20 @@ void FDeferredShadingSceneRenderer::RenderSingleLayerWaterReflections(
 		};
 
 		const bool bRunTiled = UseSingleLayerWaterIndirectDraw(View.GetShaderPlatform()) && CVarWaterSingleLayerTiledComposite.GetValueOnRenderThread();
-		FTiledReflection TiledScreenSpaceReflection = {nullptr, nullptr, nullptr, 8};
-		FIntVector ViewRes(View.ViewRect.Width(), View.ViewRect.Height(), 1);
-		FIntVector TiledViewRes = FIntVector::DivideAndRoundUp(ViewRes, TiledScreenSpaceReflection.TileSize);
 
+		FSingleLayerWaterTileClassification SingleLayerWaterTileClassification;
 		if (bRunTiled)
 		{
-			TiledScreenSpaceReflection.DrawIndirectParametersBuffer = GraphBuilder.CreateBuffer(FRDGBufferDesc::CreateIndirectDesc<FRHIDrawIndirectParameters>(), TEXT("SLW.WaterIndirectDrawParameters"));
-			TiledScreenSpaceReflection.DispatchIndirectParametersBuffer = GraphBuilder.CreateBuffer(FRDGBufferDesc::CreateIndirectDesc<FRHIDispatchIndirectParameters>(1), TEXT("SLW.WaterIndirectDispatchParameters"));
-
-			FRDGBufferRef TileListDataBuffer = GraphBuilder.CreateBuffer(FRDGBufferDesc::CreateBufferDesc(sizeof(uint32), TiledViewRes.X * TiledViewRes.Y), TEXT("SLW.TileListDataBuffer"));
-			TiledScreenSpaceReflection.TileListDataBufferSRV = GraphBuilder.CreateSRV(TileListDataBuffer, PF_R32_UINT);
-
-			FRDGBufferUAVRef DrawIndirectParametersBufferUAV = GraphBuilder.CreateUAV(TiledScreenSpaceReflection.DrawIndirectParametersBuffer);
-			FRDGBufferUAVRef DispatchIndirectParametersBufferUAV = GraphBuilder.CreateUAV(TiledScreenSpaceReflection.DispatchIndirectParametersBuffer);
-			FRDGBufferUAVRef TileListDataBufferUAV = GraphBuilder.CreateUAV(TileListDataBuffer, PF_R32_UINT);
-
-			// Clear DrawIndirectParametersBuffer
-			AddClearUAVPass(GraphBuilder, DrawIndirectParametersBufferUAV, 0);
-			AddClearUAVPass(GraphBuilder, DispatchIndirectParametersBufferUAV, 0);
-
-			// Categorization based on SHADING_MODEL_ID
+			if (SingleLayerWaterPrePassResult)
 			{
-				TShaderMapRef<FWaterTileCategorisationCS> ComputeShader(View.ShaderMap);
-
-				FWaterTileCategorisationCS::FParameters* PassParameters = GraphBuilder.AllocParameters<FWaterTileCategorisationCS::FParameters>();
-				SetCommonParameters(PassParameters->CommonParameters);
-				PassParameters->VertexCountPerInstanceIndirect = GRHISupportsRectTopology ? 3 : 6;
-				PassParameters->DrawIndirectDataUAV = DrawIndirectParametersBufferUAV;
-				PassParameters->DispatchIndirectDataUAV = DispatchIndirectParametersBufferUAV;
-				PassParameters->WaterTileListDataUAV = TileListDataBufferUAV;
-
-				FComputeShaderUtils::AddPass(GraphBuilder, RDG_EVENT_NAME("SLW::TileCategorisation"), ComputeShader, PassParameters, TiledViewRes);
+				SingleLayerWaterTileClassification = SingleLayerWaterPrePassResult->ViewTileClassification[ViewIndex];
+			}
+			else
+			{
+				SingleLayerWaterTileClassification = ClassifyTiles(GraphBuilder, View, SceneTextures, nullptr);
 			}
 		}
-
+		FTiledReflection& TiledScreenSpaceReflection = SingleLayerWaterTileClassification.TiledReflection;
 		const FPerViewPipelineState& ViewPipelineState = GetViewPipelineState(View);
 
 		if (IsWaterDistanceFieldShadowEnabled_Runtime(View.GetShaderPlatform()))
@@ -976,7 +1029,7 @@ void FDeferredShadingSceneRenderer::RenderSingleLayerWaterReflections(
 void FDeferredShadingSceneRenderer::RenderSingleLayerWater(
 	FRDGBuilder& GraphBuilder,
 	const FSceneTextures& SceneTextures,
-	const FRDGTextureMSAA& SingleLayerWaterDepthPrepassTexture,
+	const FSingleLayerWaterPrePassResult* SingleLayerWaterPrePassResult,
 	bool bShouldRenderVolumetricCloud,
 	FSceneWithoutWaterTextures& SceneWithoutWaterTextures,
 	FLumenSceneFrameTemporaries& LumenFrameTemporaries)
@@ -998,7 +1051,7 @@ void FDeferredShadingSceneRenderer::RenderSingleLayerWater(
 		ComposeVolumetricRenderTargetOverSceneUnderWater(GraphBuilder, Views, SceneWithoutWaterTextures, SceneTextures);
 	}
 
-	RenderSingleLayerWaterInner(GraphBuilder, SceneTextures, SceneWithoutWaterTextures, SingleLayerWaterDepthPrepassTexture);
+	RenderSingleLayerWaterInner(GraphBuilder, SceneTextures, SceneWithoutWaterTextures, SingleLayerWaterPrePassResult);
 
 	// No SSR or composite needed in Forward. Reflections are applied in the WaterGBuffer pass.
 	if (!IsForwardShadingEnabled(ShaderPlatform))
@@ -1006,15 +1059,15 @@ void FDeferredShadingSceneRenderer::RenderSingleLayerWater(
 		// Reflection composite expects the depth buffer in FSceneTextures to contain water but the swap of the main depth buffer with the water prepass depth buffer
 		// is only done at the call site after this function returns (for visibility and to keep SceneTextures const), so we need to swap the depth buffers on an internal copy.
 		FSceneTextures SceneTexturesInternal = SceneTextures;
-		if (SingleLayerWaterDepthPrepassTexture.Target != nullptr)
+		if (SingleLayerWaterPrePassResult)
 		{
-			SceneTexturesInternal.Depth = SingleLayerWaterDepthPrepassTexture;
+			SceneTexturesInternal.Depth = SingleLayerWaterPrePassResult->DepthPrepassTexture;
 			// Rebuild scene textures uniform buffer to include new depth buffer.
 			SceneTexturesInternal.UniformBuffer = CreateSceneTextureUniformBuffer(GraphBuilder, &SceneTexturesInternal, FeatureLevel, SceneTexturesInternal.SetupMode);
 		}
 
 		// If supported render SSR, the composite pass in non deferred and/or under water effect.
-		RenderSingleLayerWaterReflections(GraphBuilder, SceneTexturesInternal, SceneWithoutWaterTextures, LumenFrameTemporaries);
+		RenderSingleLayerWaterReflections(GraphBuilder, SceneTexturesInternal, SceneWithoutWaterTextures, SingleLayerWaterPrePassResult, LumenFrameTemporaries);
 	}
 }
 
@@ -1051,7 +1104,7 @@ void FDeferredShadingSceneRenderer::RenderSingleLayerWaterInner(
 	FRDGBuilder& GraphBuilder,
 	const FSceneTextures& SceneTextures,
 	const FSceneWithoutWaterTextures& SceneWithoutWaterTextures,
-	const FRDGTextureMSAA& SingleLayerWaterDepthPrepassTexture)
+	const FSingleLayerWaterPrePassResult* SingleLayerWaterPrePassResult)
 {
 	RDG_CSV_STAT_EXCLUSIVE_SCOPE(GraphBuilder, Water);
 	SCOPED_NAMED_EVENT(FDeferredShadingSceneRenderer_RenderSingleLayerWaterPass, FColor::Emerald);
@@ -1074,11 +1127,11 @@ void FDeferredShadingSceneRenderer::RenderSingleLayerWaterInner(
 
 	FRDGTextureRef WhiteForwardScreenSpaceShadowMask = SystemTextures.White;
 
-	const bool bHasDepthPrepass = SingleLayerWaterDepthPrepassTexture.Target != nullptr;
+	const bool bHasDepthPrepass = SingleLayerWaterPrePassResult != nullptr;
 	FDepthStencilBinding DepthStencilBinding;
 	if (bHasDepthPrepass)
 	{
-		DepthStencilBinding = FDepthStencilBinding(SingleLayerWaterDepthPrepassTexture.Target, ERenderTargetLoadAction::ELoad, ERenderTargetLoadAction::ELoad, FExclusiveDepthStencil::DepthRead_StencilRead);
+		DepthStencilBinding = FDepthStencilBinding(SingleLayerWaterPrePassResult->DepthPrepassTexture.Target, ERenderTargetLoadAction::ELoad, ERenderTargetLoadAction::ELoad, FExclusiveDepthStencil::DepthRead_StencilRead);
 	}
 	else
 	{
diff --git a/Engine/Source/Runtime/Renderer/Private/SingleLayerWaterRendering.h b/Engine/Source/Runtime/Renderer/Private/SingleLayerWaterRendering.h
--- a/Engine/Source/Runtime/Renderer/Private/SingleLayerWaterRendering.h
+++ b/Engine/Source/Runtime/Renderer/Private/SingleLayerWaterRendering.h
@@ -5,9 +5,24 @@
 #include "RenderGraph.h"
 #include "PixelShaderUtils.h"
 #include "ScreenSpaceRayTracing.h"
+#include "SingleLayerWaterDefinitions.h"
 
 class FViewInfo;
 
+struct FSingleLayerWaterTileClassification
+{
+	FTiledReflection TiledReflection = FTiledReflection{ nullptr, nullptr, nullptr, SLW_TILE_SIZE_XY };
+	FRDGBufferRef TileMaskBuffer = nullptr;
+	FIntPoint TiledViewRes;
+};
+
+struct FSingleLayerWaterPrePassResult
+{
+	FRDGTextureMSAA DepthPrepassTexture;
+	TArray<FSingleLayerWaterTileClassification> ViewTileClassification;
+};
+
+
 struct FSceneWithoutWaterTextures
 {
 	struct FView
diff --git a/Engine/Source/Runtime/Renderer/Private/VirtualShadowMaps/VirtualShadowMapArray.cpp b/Engine/Source/Runtime/Renderer/Private/VirtualShadowMaps/VirtualShadowMapArray.cpp
--- a/Engine/Source/Runtime/Renderer/Private/VirtualShadowMaps/VirtualShadowMapArray.cpp
+++ b/Engine/Source/Runtime/Renderer/Private/VirtualShadowMaps/VirtualShadowMapArray.cpp
@@ -18,6 +18,7 @@
 #include "VirtualShadowMapCacheManager.h"
 #include "VirtualShadowMapClipmap.h"
 #include "VirtualShadowMapVisualizationData.h"
+#include "SingleLayerWaterRendering.h"
 
 #define DEBUG_ALLOW_STATIC_SEPARATE_WITHOUT_CACHING 0
 
@@ -316,6 +317,18 @@ static TAutoConsoleVariable<int32> CVarVirtualShadowSinglePassBatched(
 	TEXT("."),
 	ECVF_RenderThreadSafe);
 
+static TAutoConsoleVariable<int32> CVarVirtualShadowMapPageMarkingPixelStrideX(
+	TEXT("r.Shadow.Virtual.PageMarkingPixelStrideX"),
+	1,
+	TEXT("."),
+	ECVF_RenderThreadSafe);
+
+static TAutoConsoleVariable<int32> CVarVirtualShadowMapPageMarkingPixelStrideY(
+	TEXT("r.Shadow.Virtual.PageMarkingPixelStrideY"),
+	1,
+	TEXT("."),
+	ECVF_RenderThreadSafe);
+
 FMatrix CalcTranslatedWorldToShadowUVMatrix(
 	const FMatrix& TranslatedWorldToShadowView,
 	const FMatrix& ViewToClip)
@@ -617,7 +630,6 @@ public:
 	// Kernel launch group sizes
 	static constexpr uint32 DefaultCSGroupXY = 8;
 	static constexpr uint32 DefaultCSGroupX = 256;
-	static constexpr uint32 GeneratePageFlagsGroupXYZ = 4;
 	static constexpr uint32 BuildExplicitBoundsGroupXY = 16;
 
 	FVirtualPageManagementShader()
@@ -646,7 +658,6 @@ public:
 
 		OutEnvironment.SetDefine(TEXT("VSM_DEFAULT_CS_GROUP_X"), DefaultCSGroupX);
 		OutEnvironment.SetDefine(TEXT("VSM_DEFAULT_CS_GROUP_XY"), DefaultCSGroupXY);
-		OutEnvironment.SetDefine(TEXT("VSM_GENERATE_PAGE_FLAGS_CS_GROUP_XYZ"), GeneratePageFlagsGroupXYZ);
 		OutEnvironment.SetDefine(TEXT("VSM_BUILD_EXPLICIT_BOUNDS_CS_XY"), BuildExplicitBoundsGroupXY);
 
 		FForwardLightingParameters::ModifyCompilationEnvironment(Parameters.Platform, OutEnvironment);
@@ -688,7 +699,9 @@ class FGeneratePageFlagsFromPixelsCS : public FVirtualPageManagementShader
 		SHADER_PARAMETER_RDG_BUFFER_SRV(StructuredBuffer< uint >, DirectionalLightIds)
 		SHADER_PARAMETER_RDG_BUFFER_SRV(StructuredBuffer< uint >, PrunedLightGridData)
 		SHADER_PARAMETER_RDG_BUFFER_SRV(StructuredBuffer< uint >, PrunedNumCulledLightsGrid)
-		SHADER_PARAMETER_RDG_TEXTURE(Texture2D<float>, SingeLayerWaterDepthTexture)
+		SHADER_PARAMETER_RDG_TEXTURE(Texture2D<float>, SingleLayerWaterDepthTexture)
+		SHADER_PARAMETER_RDG_BUFFER_SRV(StructuredBuffer< uint >, SingleLayerWaterTileMask)
+		SHADER_PARAMETER(FIntPoint, SingleLayerWaterTileViewRes)
 		SHADER_PARAMETER_RDG_UNIFORM_BUFFER(FStrataGlobalUniformParameters, Strata)
 		RDG_BUFFER_ACCESS(IndirectBufferArgs, ERHIAccess::IndirectArgs)
 		SHADER_PARAMETER(uint32, InputType)
@@ -698,6 +711,7 @@ class FGeneratePageFlagsFromPixelsCS : public FVirtualPageManagementShader
 		SHADER_PARAMETER(float, PageDilationBorderSizeDirectional)
 		SHADER_PARAMETER(float, PageDilationBorderSizeLocal)
 		SHADER_PARAMETER(uint32, bCullBackfacingPixels)
+		SHADER_PARAMETER(FIntPoint, PixelStride)
 	END_SHADER_PARAMETER_STRUCT()
 };
 IMPLEMENT_GLOBAL_SHADER(FGeneratePageFlagsFromPixelsCS, "/Engine/Private/VirtualShadowMaps/VirtualShadowMapPageManagement.usf", "GeneratePageFlagsFromPixels", SF_Compute);
@@ -1142,7 +1156,7 @@ void FVirtualShadowMapArray::BuildPageAllocations(
 	const FSortedLightSetSceneInfo& SortedLightsInfo,
 	const TArray<FVisibleLightInfo, SceneRenderingAllocator>& VisibleLightInfos,
 	const TArray<Nanite::FRasterResults, TInlineAllocator<2>>& NaniteRasterResults,
-	FRDGTextureRef SingleLayerWaterDepthTexture)
+	const FSingleLayerWaterPrePassResult* SingleLayerWaterPrePassResult)
 {
 	check(IsEnabled());
 
@@ -1440,6 +1454,10 @@ void FVirtualShadowMapArray::BuildPageAllocations(
 
 					PassParameters->SceneTexturesStruct = SceneTextures.UniformBuffer;
 
+					const FIntPoint PixelStride(
+						FMath::Clamp(CVarVirtualShadowMapPageMarkingPixelStrideX.GetValueOnRenderThread(), 1, 128),
+						FMath::Clamp(CVarVirtualShadowMapPageMarkingPixelStrideY.GetValueOnRenderThread(), 1, 128));
+
 					PassParameters->HairStrands = HairStrands::BindHairStrandsViewUniformParameters(View);
 					PassParameters->View = View.ViewUniformBuffer;
 					PassParameters->OutPageRequestFlags = PageRequestFlagsUAV;
@@ -1447,17 +1465,24 @@ void FVirtualShadowMapArray::BuildPageAllocations(
 					PassParameters->DirectionalLightIds = GraphBuilder.CreateSRV(DirectionalLightIdsRDG);
 					PassParameters->PrunedLightGridData = GraphBuilder.CreateSRV(PrunedLightGridDataRDG);
 					PassParameters->PrunedNumCulledLightsGrid = GraphBuilder.CreateSRV(PrunedNumCulledLightsGridRDG);
-					PassParameters->SingeLayerWaterDepthTexture = SingleLayerWaterDepthTexture;
+					if (SingleLayerWaterPrePassResult)
+					{
+						PassParameters->SingleLayerWaterDepthTexture = SingleLayerWaterPrePassResult->DepthPrepassTexture.Resolve;
+						PassParameters->SingleLayerWaterTileMask = GraphBuilder.CreateSRV(SingleLayerWaterPrePassResult->ViewTileClassification[ViewIndex].TileMaskBuffer);
+						PassParameters->SingleLayerWaterTileViewRes = SingleLayerWaterPrePassResult->ViewTileClassification[ViewIndex].TiledViewRes;
+					}
 					PassParameters->NumDirectionalLightSmInds = uint32(DirectionalLightIds.Num());
 					PassParameters->PageDilationBorderSizeLocal = CVarPageDilationBorderSizeLocal.GetValueOnRenderThread();
 					PassParameters->PageDilationBorderSizeDirectional = CVarPageDilationBorderSizeDirectional.GetValueOnRenderThread();
 					PassParameters->bCullBackfacingPixels = ShouldCullBackfacingPixels() ? 1 : 0;
 					PassParameters->Strata = Strata::BindStrataGlobalUniformParameters(View);
+					PassParameters->PixelStride = PixelStride;
 
 					auto ComputeShader = View.ShaderMap->GetShader<FGeneratePageFlagsFromPixelsCS>(PermutationVector);
 
-					static_assert((FVirtualPageManagementShader::DefaultCSGroupXY % 2) == 0, "GeneratePageFlagsFromPixels requires even-sized CS groups for quad swizzling.");
-					const FIntPoint GridSize = FIntPoint::DivideAndRoundUp(View.ViewRect.Size(), FVirtualPageManagementShader::DefaultCSGroupXY);
+					const FIntPoint StridedPixelSize = FIntPoint::DivideAndRoundUp(View.ViewRect.Size(), PixelStride);
+					// Note: we use the tile size defined by the water as the group-size - this is needed because the tile mask testing code relies on the size being the same to scalarize efficiently.
+					const FIntPoint GridSize = FIntPoint::DivideAndRoundUp(StridedPixelSize, SLW_TILE_SIZE_XY);
 
 					if (InputType == EVirtualShadowMapProjectionInputType::HairStrands)
 					{
@@ -1475,14 +1500,14 @@ void FVirtualShadowMapArray::BuildPageAllocations(
 					{
 						FComputeShaderUtils::AddPass(
 							GraphBuilder,
-							RDG_EVENT_NAME("GeneratePageFlagsFromPixels(%s,NumShadowMaps=%d)", ToString(InputType), GetNumFullShadowMaps()),
+							RDG_EVENT_NAME("GeneratePageFlagsFromPixels(%s,NumShadowMaps=%d,{%d,%d})", ToString(InputType), GetNumFullShadowMaps(), GridSize.X, GridSize.Y),
 							ComputeShader,
 							PassParameters,
 							FIntVector(GridSize.X, GridSize.Y, 1));
 					}
 				};
 
-				const bool bHasValidSingleLayerWaterDepth = SingleLayerWaterDepthTexture != nullptr;
+				const bool bHasValidSingleLayerWaterDepth = SingleLayerWaterPrePassResult != nullptr;
 				GeneratePageFlags(bHasValidSingleLayerWaterDepth ? EVirtualShadowMapProjectionInputType::GBufferAndSingleLayerWaterDepth : EVirtualShadowMapProjectionInputType::GBuffer);
 				if (HairStrands::HasViewHairStrandsData(View))
 				{
diff --git a/Engine/Source/Runtime/Renderer/Private/VirtualShadowMaps/VirtualShadowMapArray.h b/Engine/Source/Runtime/Renderer/Private/VirtualShadowMaps/VirtualShadowMapArray.h
--- a/Engine/Source/Runtime/Renderer/Private/VirtualShadowMaps/VirtualShadowMapArray.h
+++ b/Engine/Source/Runtime/Renderer/Private/VirtualShadowMaps/VirtualShadowMapArray.h
@@ -19,6 +19,7 @@ class FVirtualShadowMapArrayCacheManager;
 struct FSortedLightSetSceneInfo;
 class FVirtualShadowMapClipmap;
 struct FScreenPassTexture;
+struct FSingleLayerWaterPrePassResult;
 
 // TODO: does this exist?
 constexpr uint32 ILog2Const(uint32 n)
@@ -263,7 +264,7 @@ public:
 		const FSortedLightSetSceneInfo& SortedLights, 
 		const TArray<FVisibleLightInfo, SceneRenderingAllocator> &VisibleLightInfos, 
 		const TArray<Nanite::FRasterResults, TInlineAllocator<2>> &NaniteRasterResults,
-		FRDGTextureRef SingleLayerWaterDepthTexture);
+		const FSingleLayerWaterPrePassResult* SingleLayerWaterPrePassResult);
 
 	bool IsAllocated() const
 	{
