Fix RHI validation errors on some sky atmosphere resources

diff --git a/Engine/Source/Runtime/Renderer/Private/ReflectionEnvironment.cpp b/Engine/Source/Runtime/Renderer/Private/ReflectionEnvironment.cpp
--- a/Engine/Source/Runtime/Renderer/Private/ReflectionEnvironment.cpp
+++ b/Engine/Source/Runtime/Renderer/Private/ReflectionEnvironment.cpp
@@ -631,7 +631,7 @@ void SetupSkyIrradianceEnvironmentMapConstantsFromSkyIrradiance(FVector4f* OutSk
 	OutSkyIrradianceEnvironmentMap[6].W = 1;
 }
 
-void UpdateSkyIrradianceGpuBuffer(FRHICommandListImmediate& RHICmdList, const FEngineShowFlags& EngineShowFlags, const FSkyLightSceneProxy* SkyLight, TRefCountPtr<FRDGPooledBuffer>& Buffer)
+void UpdateSkyIrradianceGpuBuffer(FRDGBuilder& GraphBuilder, const FEngineShowFlags& EngineShowFlags, const FSkyLightSceneProxy* SkyLight, TRefCountPtr<FRDGPooledBuffer>& Buffer)
 {
 	TRACE_CPUPROFILER_EVENT_SCOPE(UpdateSkyIrradianceGpuBuffer);
 
@@ -642,30 +642,35 @@ void UpdateSkyIrradianceGpuBuffer(FRHICommandListImmediate& RHICmdList, const FE
 		&& EngineShowFlags.SkyLighting
 		&& !SkyLight->bRealTimeCaptureEnabled; // When bRealTimeCaptureEnabled is true, the buffer will be setup on GPU directly in this case
 
-	if (AllocatePooledBuffer(FRDGBufferDesc::CreateStructuredDesc(sizeof(FVector4f), SKY_IRRADIANCE_ENVIRONMENT_MAP_VEC4_COUNT), Buffer, TEXT("SkyIrradianceEnvironmentMap"), ERDGPooledBufferAlignment::None))
-	{
-		if (!bUploadIrradiance)
-		{
-			// Ensure that sky irradiance SH buffer contains sensible initial values (zero init).
-			// If there is no sky in the level, then nothing else may fill this buffer.
-			void* DataPtr = RHICmdList.LockBuffer(Buffer->GetRHI(), 0, Buffer->GetSize(), RLM_WriteOnly);
-			FPlatformMemory::Memset(DataPtr, 0, Buffer->GetSize());
-			RHICmdList.UnlockBuffer(Buffer->GetRHI());
-		}
+	const bool bNewAlloc = AllocatePooledBuffer(
+		FRDGBufferDesc::CreateStructuredDesc(sizeof(FVector4f), SKY_IRRADIANCE_ENVIRONMENT_MAP_VEC4_COUNT),
+		Buffer,
+		TEXT("SkyIrradianceEnvironmentMap"),
+		ERDGPooledBufferAlignment::None);
 
-		RHICmdList.Transition(FRHITransitionInfo(Buffer->GetRHI(), ERHIAccess::Unknown, ERHIAccess::SRVMask));
+	FRDGBufferRef BufferRDG = GraphBuilder.RegisterExternalBuffer(Buffer, TEXT("SkyIrradianceEnvironmentMap"));
+	const uint32 BufferSizeInBytes = BufferRDG->GetSize();
+	check(BufferRDG);
+
+	if (bNewAlloc && !bUploadIrradiance)
+	{
+		// Ensure that sky irradiance SH buffer contains sensible initial values (zero init).
+		// If there is no sky in the level, then nothing else may fill this buffer.
+		void* InitialValue = GraphBuilder.Alloc(BufferSizeInBytes);
+		FMemory::Memset(InitialValue, 0, BufferSizeInBytes);
+		GraphBuilder.QueueBufferUpload(BufferRDG, InitialValue, BufferSizeInBytes, ERDGInitialDataFlags::NoCopy);
 	}
 
 	if (bUploadIrradiance)
 	{
-		FVector4f UploadData[SKY_IRRADIANCE_ENVIRONMENT_MAP_VEC4_COUNT];
+		FVector4f* UploadData = (FVector4f*)GraphBuilder.Alloc(BufferSizeInBytes);
 		SetupSkyIrradianceEnvironmentMapConstantsFromSkyIrradiance(UploadData, SkyLight->IrradianceEnvironmentMap);
 
 		const float SkyLightAverageBrightness = SkyLight->AverageBrightness;
 		UploadData[7] = FVector4f(SkyLightAverageBrightness, SkyLightAverageBrightness, SkyLightAverageBrightness, SkyLightAverageBrightness);
 
-		void* DataPtr = RHICmdList.LockBuffer(Buffer->GetRHI(), 0, Buffer->GetSize(), RLM_WriteOnly);
-		FPlatformMemory::Memcpy(DataPtr, &UploadData, sizeof(UploadData));
-		RHICmdList.UnlockBuffer(Buffer->GetRHI());
+		GraphBuilder.QueueBufferUpload(BufferRDG, UploadData, BufferSizeInBytes, ERDGInitialDataFlags::NoCopy);
 	}
+
+	Buffer = ConvertToExternalAccessBuffer(GraphBuilder, BufferRDG, ERHIAccess::SRVMask, ERHIPipeline::All);
 }
\ No newline at end of file
diff --git a/Engine/Source/Runtime/Renderer/Private/ReflectionEnvironment.h b/Engine/Source/Runtime/Renderer/Private/ReflectionEnvironment.h
--- a/Engine/Source/Runtime/Renderer/Private/ReflectionEnvironment.h
+++ b/Engine/Source/Runtime/Renderer/Private/ReflectionEnvironment.h
@@ -36,4 +36,4 @@ TUniformBufferRef<FReflectionUniformParameters> CreateReflectionUniformBuffer(co
 
 RENDERER_API void SetupSkyIrradianceEnvironmentMapConstantsFromSkyIrradiance(FVector4f* OutSkyIrradianceEnvironmentMap, const FSHVectorRGB3 SkyIrradiance);
 
-extern void UpdateSkyIrradianceGpuBuffer(FRHICommandListImmediate& RHICmdList, const FEngineShowFlags& EngineShowFlags, const FSkyLightSceneProxy* SkyLight, TRefCountPtr<FRDGPooledBuffer>& Buffer);
+extern void UpdateSkyIrradianceGpuBuffer(FRDGBuilder& GraphBuilder, const FEngineShowFlags& EngineShowFlags, const FSkyLightSceneProxy* SkyLight, TRefCountPtr<FRDGPooledBuffer>& Buffer);
diff --git a/Engine/Source/Runtime/Renderer/Private/ReflectionEnvironmentRealTimeCapture.cpp b/Engine/Source/Runtime/Renderer/Private/ReflectionEnvironmentRealTimeCapture.cpp
--- a/Engine/Source/Runtime/Renderer/Private/ReflectionEnvironmentRealTimeCapture.cpp
+++ b/Engine/Source/Runtime/Renderer/Private/ReflectionEnvironmentRealTimeCapture.cpp
@@ -872,6 +872,7 @@ void FScene::AllocateAndCaptureFrameSkyEnvMap(
 		// the flag, RDG will split the transition to UAV to the start of the graph, which results in a validation error. With the flag, RDG
 		// will transition to UAV at the start of the pass instead.
 		FRDGBuffer* SkyIrradianceEnvironmentMapRDG = GraphBuilder.RegisterExternalBuffer(SkyIrradianceEnvironmentMap, ERDGBufferFlags::ForceImmediateFirstBarrier);
+		GraphBuilder.UseInternalAccessMode(SkyIrradianceEnvironmentMapRDG);
 
 		TShaderMapRef<FComputeSkyEnvMapDiffuseIrradianceCS> ComputeShader(GetGlobalShaderMap(FeatureLevel));
 
@@ -893,7 +894,7 @@ void FScene::AllocateAndCaptureFrameSkyEnvMap(
 		const FIntVector NumGroups = FIntVector(1, 1, 1);
 		FComputeShaderUtils::AddPass(GraphBuilder, RDG_EVENT_NAME("ComputeSkyEnvMapDiffuseIrradianceCS"), ComputeShader, PassParameters, NumGroups);
 
-		ExternalAccessQueue.Add(SkyIrradianceEnvironmentMapRDG, ERHIAccess::SRVMask);
+		ExternalAccessQueue.Add(SkyIrradianceEnvironmentMapRDG, ERHIAccess::SRVMask, ERHIPipeline::All);
 	};
 
 	const uint32 LastMipLevel = CubeMipCount - 1;
@@ -1117,7 +1118,7 @@ void FScene::AllocateAndCaptureFrameSkyEnvMap(
 
 	if (ConvolvedSkyRenderTarget[ConvolvedSkyRenderTargetReadyIndex])
 	{
-		ExternalAccessQueue.Add(GraphBuilder.RegisterExternalTexture(ConvolvedSkyRenderTarget[ConvolvedSkyRenderTargetReadyIndex]), ERHIAccess::SRVMask);
+		ExternalAccessQueue.Add(GraphBuilder.RegisterExternalTexture(ConvolvedSkyRenderTarget[ConvolvedSkyRenderTargetReadyIndex]), ERHIAccess::SRVMask, ERHIPipeline::All);
 	}
 
 	ExternalAccessQueue.Submit(GraphBuilder);
diff --git a/Engine/Source/Runtime/Renderer/Private/SceneVisibility.cpp b/Engine/Source/Runtime/Renderer/Private/SceneVisibility.cpp
--- a/Engine/Source/Runtime/Renderer/Private/SceneVisibility.cpp
+++ b/Engine/Source/Runtime/Renderer/Private/SceneVisibility.cpp
@@ -5203,7 +5203,7 @@ void FDeferredShadingSceneRenderer::InitViews(FRDGBuilder& GraphBuilder, const F
 	// This must happen before we start initialising and using views.
 	if (Scene)
 	{
-		UpdateSkyIrradianceGpuBuffer(RHICmdList, ViewFamily.EngineShowFlags, Scene->SkyLight, Scene->SkyIrradianceEnvironmentMap);
+		UpdateSkyIrradianceGpuBuffer(GraphBuilder, ViewFamily.EngineShowFlags, Scene->SkyLight, Scene->SkyIrradianceEnvironmentMap);
 	}
 
 	// Initialise Sky/View resources before the view global uniform buffer is built.
diff --git a/Engine/Source/Runtime/Renderer/Private/SkyAtmosphereRendering.cpp b/Engine/Source/Runtime/Renderer/Private/SkyAtmosphereRendering.cpp
--- a/Engine/Source/Runtime/Renderer/Private/SkyAtmosphereRendering.cpp
+++ b/Engine/Source/Runtime/Renderer/Private/SkyAtmosphereRendering.cpp
@@ -1398,7 +1398,7 @@ void FSceneRenderer::RenderSkyAtmosphereLookUpTables(FRDGBuilder& GraphBuilder)
 		const FIntVector NumGroups = FIntVector::DivideAndRoundUp(TextureSize, FRenderDistantSkyLightLutCS::GroupSize);
 		FComputeShaderUtils::AddPass(GraphBuilder, RDG_EVENT_NAME("DistantSkyLightLut"), PassFlag, ComputeShader, PassParameters, NumGroups);
 
-		SkyInfo.GetDistantSkyLightLutTexture() = ConvertToExternalAccessTexture(GraphBuilder, ExternalAccessQueue, DistantSkyLightLut);
+		SkyInfo.GetDistantSkyLightLutTexture() = ConvertToExternalAccessTexture(GraphBuilder, ExternalAccessQueue, DistantSkyLightLut, ERHIAccess::SRVMask, ERHIPipeline::All);
 	}
 
 	SkyAtmosphereLightShadowData LightShadowData;
@@ -1672,7 +1672,7 @@ void FSceneRenderer::RenderSkyAtmosphereLookUpTables(FRDGBuilder& GraphBuilder)
 		}
 
 		View.SkyAtmosphereViewLutTexture = ConvertToExternalAccessTexture(GraphBuilder, ExternalAccessQueue, SkyAtmosphereViewLutTexture);
-		View.SkyAtmosphereCameraAerialPerspectiveVolume = ConvertToExternalAccessTexture(GraphBuilder, ExternalAccessQueue, SkyAtmosphereCameraAerialPerspectiveVolume);
+		View.SkyAtmosphereCameraAerialPerspectiveVolume = ConvertToExternalAccessTexture(GraphBuilder, ExternalAccessQueue, SkyAtmosphereCameraAerialPerspectiveVolume, ERHIAccess::SRVMask, ERHIPipeline::All);
 	}
 
 	SkyInfo.GetTransmittanceLutTexture() = ConvertToExternalAccessTexture(GraphBuilder, ExternalAccessQueue, TransmittanceLut);
