Fix a non unity compile error

diff --git a/Engine/Source/Runtime/Renderer/Private/ReflectionEnvironment.h b/Engine/Source/Runtime/Renderer/Private/ReflectionEnvironment.h
--- a/Engine/Source/Runtime/Renderer/Private/ReflectionEnvironment.h
+++ b/Engine/Source/Runtime/Renderer/Private/ReflectionEnvironment.h
@@ -36,4 +36,4 @@ TUniformBufferRef<FReflectionUniformParameters> CreateReflectionUniformBuffer(co
 
 RENDERER_API void SetupSkyIrradianceEnvironmentMapConstantsFromSkyIrradiance(FVector4f* OutSkyIrradianceEnvironmentMap, const FSHVectorRGB3 SkyIrradiance);
 
-extern void UpdateSkyIrradianceGpuBuffer(FRDGBuilder& GraphBuilder, const FEngineShowFlags& EngineShowFlags, const FSkyLightSceneProxy* SkyLight, TRefCountPtr<FRDGPooledBuffer>& Buffer);
+extern void UpdateSkyIrradianceGpuBuffer(class FRDGBuilder& GraphBuilder, const FEngineShowFlags& EngineShowFlags, const FSkyLightSceneProxy* SkyLight, TRefCountPtr<FRDGPooledBuffer>& Buffer);
