Implements r.TSR.ShadingRejection.Flickering.MaxParralaxVelocity to minimise ghosting on reflection and heavy parallaxed pixel wizardry

diff --git a/Engine/Shaders/Private/TemporalSuperResolution/TSRDilateVelocity.usf b/Engine/Shaders/Private/TemporalSuperResolution/TSRDilateVelocity.usf
--- a/Engine/Shaders/Private/TemporalSuperResolution/TSRDilateVelocity.usf
+++ b/Engine/Shaders/Private/TemporalSuperResolution/TSRDilateVelocity.usf
@@ -27,10 +27,12 @@
 
 //------------------------------------------------------- PARAMETERS
 
+float4x4 RotationalClipToPrevClip;
 float2 PrevOutputBufferUVMin;
 float2 PrevOutputBufferUVMax;
 float WorldDepthToDepthError;
 float VelocityExtrapolationMultiplier;
+float InvFlickeringMaxParralaxVelocity;
 uint bOutputIsMovingTexture;
 
 Texture2D<float> SceneDepthTexture;
@@ -215,7 +217,16 @@ void MainCS(
 		float Depth = ConvertFromDeviceZ(ClosestDeviceZ);
 		float PrevDepth = ConvertFromDeviceZ(ClosestDeviceZ - FinalScreenVelocity.z);
 		float2 PrevScreenPos = ScreenPos - FinalScreenVelocity.xy;
-
+		
+		float4 ThisClip = float4(ScreenPos, DeviceZ, 1);
+		float4 PrevClip = mul(ThisClip, View.ClipToPrevClip); 
+		float2 PrevScreen = PrevClip.xy / PrevClip.w;
+			
+		float4 RotationalPrevClip = mul(ThisClip, RotationalClipToPrevClip);
+		float2 RotationalPrevScreen = RotationalPrevClip.xy / RotationalPrevClip.w;
+			
+		tsr_half PixelParallaxVelocity = tsr_half(0.5) * length(tsr_half2(RotationalPrevScreen - PrevScreen) * tsr_half2(InputInfo_ViewportSize));
+			
 		float4 ClipPosition = float4(ScreenPos * (View.ViewToClip[3][3] < 1.0f ? Depth : 1.0f), Depth, 1);
 		float4 PrevClipPosition = float4(PrevScreenPos * (View.ViewToClip[3][3] < 1.0f ? PrevDepth : 1.0f), PrevDepth, 1);
 		
@@ -226,7 +237,9 @@ void MainCS(
 		float WorldDepthToPixelWorldRadius = 2.0 * GetTanHalfFieldOfView().x * View.ViewSizeAndInvSize.z;
 		float StaticWorldRadius = Depth * WorldDepthToPixelWorldRadius;
 
-		IsMovingMask = tsr_half(saturate(length(TranslatedWorldPosition - PrevTranslatedWorldPosition) * rcp(StaticWorldRadius) - 1.0));
+		IsMovingMask = max(
+			tsr_half(saturate(length(TranslatedWorldPosition - PrevTranslatedWorldPosition) * rcp(StaticWorldRadius) - 1.0)),
+			saturate(PixelParallaxVelocity * tsr_half(InvFlickeringMaxParralaxVelocity) - tsr_half(0.5)));
 	}
 
 	// Do motion blur velocity flatten.
diff --git a/Engine/Source/Runtime/Renderer/Private/PostProcess/TemporalSuperResolution.cpp b/Engine/Source/Runtime/Renderer/Private/PostProcess/TemporalSuperResolution.cpp
--- a/Engine/Source/Runtime/Renderer/Private/PostProcess/TemporalSuperResolution.cpp
+++ b/Engine/Source/Runtime/Renderer/Private/PostProcess/TemporalSuperResolution.cpp
@@ -56,6 +56,11 @@ TAutoConsoleVariable<int32> CVarTSRFlickeringPeriod(
 	TEXT("Periode in frames in which luma oscilations at equal or greater frequency is considered flickering and should ghost (Default=3.0).\n"),
 	ECVF_RenderThreadSafe);
 
+TAutoConsoleVariable<float> CVarTSRFlickeringMaxParralaxVelocity(
+	TEXT("r.TSR.ShadingRejection.Flickering.MaxParralaxVelocity"), 10.0,
+	TEXT("Maximum parralax velocity in 1080p 60hz pixels allowed before diminishing flickering.\n"),
+	ECVF_RenderThreadSafe);
+
 TAutoConsoleVariable<int32> CVarTSRFilterShadingRejection(
 	TEXT("r.TSR.ShadingRejection.SpatialFilter"), 1,
 	TEXT("Whether the shading rejection should have spatial statistical filtering pass to reduce flickering (default = 1).\n")
@@ -345,10 +350,12 @@ class FTSRDilateVelocityCS : public FTSRShader
 		SHADER_PARAMETER_STRUCT_INCLUDE(FTSRCommonParameters, CommonParameters)
 		SHADER_PARAMETER_STRUCT_INCLUDE(FVelocityFlattenParameters, VelocityFlattenParameters)
 
+		SHADER_PARAMETER(FMatrix44f, RotationalClipToPrevClip)
 		SHADER_PARAMETER(FVector2f, PrevOutputBufferUVMin)
 		SHADER_PARAMETER(FVector2f, PrevOutputBufferUVMax)
 		SHADER_PARAMETER(float, WorldDepthToDepthError)
 		SHADER_PARAMETER(float, VelocityExtrapolationMultiplier)
+		SHADER_PARAMETER(float, InvFlickeringMaxParralaxVelocity)
 		SHADER_PARAMETER(int32, bOutputIsMovingTexture)
 
 		SHADER_PARAMETER_RDG_TEXTURE(Texture2D, SceneDepthTexture)
@@ -789,6 +796,8 @@ ITemporalUpscaler::FOutputs AddTemporalSuperResolutionPasses(
 
 	const bool bGrandReprojection = !bAccumulateTranslucencySeparately && CVarTSRHistoryGrandReprojection.GetValueOnRenderThread() != 0;
 
+	const float RefreshRateTo60Hz = View.Family->Time.GetDeltaRealTimeSeconds() * 60.0f;
+
 	// whether TSR passes can run on async compute.
 	int32 AsyncComputePasses = GSupportsEfficientAsyncCompute ? CVarTSRAsyncCompute.GetValueOnRenderThread() : 0;
 
@@ -920,6 +929,17 @@ ITemporalUpscaler::FOutputs AddTemporalSuperResolutionPasses(
 		SeparateTranslucencyRect = PassInputs.PostDOFTranslucencyResources.ViewRect;
 	}
 
+	FMatrix44f RotationalClipToPrevClip;
+	{
+		const FViewMatrices& ViewMatrices = View.ViewMatrices;
+		const FViewMatrices& PrevViewMatrices = View.PrevViewInfo.ViewMatrices;
+
+		FMatrix RotationalInvViewProj = ViewMatrices.ComputeInvProjectionNoAAMatrix() * (ViewMatrices.GetTranslatedViewMatrix().RemoveTranslation().GetTransposed());
+		FMatrix RotationalPrevViewProj = (PrevViewMatrices.GetTranslatedViewMatrix().RemoveTranslation()) * PrevViewMatrices.ComputeProjectionNoAAMatrix();
+
+		RotationalClipToPrevClip = FMatrix44f(RotationalInvViewProj * RotationalPrevViewProj);
+	}
+
 	FTSRCommonParameters CommonParameters;
 	{
 		CommonParameters.InputInfo = GetScreenPassTextureViewportParameters(FScreenPassTextureViewport(
@@ -1032,6 +1052,7 @@ ITemporalUpscaler::FOutputs AddTemporalSuperResolutionPasses(
 		FTSRDilateVelocityCS::FPermutationDomain PermutationVector;
 		FTSRDilateVelocityCS::FParameters* PassParameters = GraphBuilder.AllocParameters<FTSRDilateVelocityCS::FParameters>();
 		PassParameters->CommonParameters = CommonParameters;
+		PassParameters->RotationalClipToPrevClip = RotationalClipToPrevClip;
 		PassParameters->PrevOutputBufferUVMin = CommonParameters.InputInfo.UVViewportBilinearMin - CommonParameters.InputInfo.ExtentInverse;
 		PassParameters->PrevOutputBufferUVMax = CommonParameters.InputInfo.UVViewportBilinearMax + CommonParameters.InputInfo.ExtentInverse;
 		{
@@ -1043,6 +1064,10 @@ ITemporalUpscaler::FOutputs AddTemporalSuperResolutionPasses(
 			PassParameters->WorldDepthToDepthError = WorldDepthToPixelWorldRadius * 2.0f;
 		}
 		PassParameters->VelocityExtrapolationMultiplier = FMath::Clamp(CVarTSRVelocityExtrapolation.GetValueOnRenderThread(), 0.0f, 1.0f);
+		{
+			float FlickeringMaxParralaxVelocity = RefreshRateTo60Hz * CVarTSRFlickeringMaxParralaxVelocity.GetValueOnRenderThread() * float(View.ViewRect.Width()) / 1920.0f;
+			PassParameters->InvFlickeringMaxParralaxVelocity = 1.0f / FlickeringMaxParralaxVelocity;
+		}
 		PassParameters->bOutputIsMovingTexture = bOutputIsMovingTexture;
 
 		PassParameters->SceneDepthTexture = PassInputs.SceneDepthTexture;
@@ -1279,18 +1304,8 @@ ITemporalUpscaler::FOutputs AddTemporalSuperResolutionPasses(
 
 		FTSRDecimateHistoryCS::FParameters* PassParameters = GraphBuilder.AllocParameters<FTSRDecimateHistoryCS::FParameters>();
 		PassParameters->CommonParameters = CommonParameters;
-		{
-			const FViewMatrices& ViewMatrices = View.ViewMatrices;
-			const FViewMatrices& PrevViewMatrices = View.PrevViewInfo.ViewMatrices;
-
-			FMatrix RotationalInvViewProj = ViewMatrices.ComputeInvProjectionNoAAMatrix() * (ViewMatrices.GetTranslatedViewMatrix().RemoveTranslation().GetTransposed());
-			FMatrix RotationalPrevViewProj = (PrevViewMatrices.GetTranslatedViewMatrix().RemoveTranslation()) * PrevViewMatrices.ComputeProjectionNoAAMatrix();
-
-			PassParameters->RotationalClipToPrevClip = FMatrix44f(RotationalInvViewProj * RotationalPrevViewProj);		// LWC_TODO: Precision loss?
-		}
-		{
-			PassParameters->CurrentToPrevInputJitter = InputHistory.PrevTemporalJitterPixels - CommonParameters.InputJitter;
-		}
+		PassParameters->RotationalClipToPrevClip = RotationalClipToPrevClip;
+		PassParameters->CurrentToPrevInputJitter = InputHistory.PrevTemporalJitterPixels - CommonParameters.InputJitter;
 		{
 			float TanHalfFieldOfView = View.ViewMatrices.GetInvProjectionMatrix().M[0][0];
 
