Fix up velocity buffer not being used during LumenScreenProbeGather.TemporalReprojection when async gather is enabled. Fix up LumenRoughSpecularIndirect not being used during DiffuseIndirectComposite when async lumen is enabled and SSR is used in place of Lumen reflections

diff --git a/Engine/Source/Runtime/Renderer/Private/DeferredShadingRenderer.cpp b/Engine/Source/Runtime/Renderer/Private/DeferredShadingRenderer.cpp
--- a/Engine/Source/Runtime/Renderer/Private/DeferredShadingRenderer.cpp
+++ b/Engine/Source/Runtime/Renderer/Private/DeferredShadingRenderer.cpp
@@ -3203,6 +3203,10 @@ void FDeferredShadingSceneRenderer::Render(FRDGBuilder& GraphBuilder)
 
 	// Rebuild scene textures to include GBuffers.
 	SceneTextures.SetupMode |= ESceneTextureSetupMode::GBuffers;
+	if (bShouldRenderVelocities && (bBasePassCanOutputVelocity || Scene->EarlyZPassMode == DDM_AllOpaqueNoVelocity))
+	{
+		SceneTextures.SetupMode |= ESceneTextureSetupMode::SceneVelocity;
+	}
 	SceneTextures.UniformBuffer = CreateSceneTextureUniformBuffer(GraphBuilder, &SceneTextures, FeatureLevel, SceneTextures.SetupMode);
 
 	if (bRealTimeSkyCaptureEnabled)
diff --git a/Engine/Source/Runtime/Renderer/Private/IndirectLightRendering.cpp b/Engine/Source/Runtime/Renderer/Private/IndirectLightRendering.cpp
--- a/Engine/Source/Runtime/Renderer/Private/IndirectLightRendering.cpp
+++ b/Engine/Source/Runtime/Renderer/Private/IndirectLightRendering.cpp
@@ -819,28 +819,19 @@ void FDeferredShadingSceneRenderer::DispatchAsyncLumenIndirectLightingWork(
 				ViewOutputs.ScreenBentNormalParameters,
 				ERDGPassFlags::AsyncCompute);
 
-			if (ViewPipelineState.ReflectionsMethod == EReflectionsMethod::Lumen)
-			{
-				if (GLumenReflectionsAsyncCompute != 0)
-				{
-					ViewOutputs.IndirectLightingTextures.Textures[3] = RenderLumenReflections(
-						GraphBuilder,
-						View,
-						SceneTextures,
-						LumenFrameTemporaries,
-						ViewOutputs.MeshSDFGridParameters,
-						ViewOutputs.RadianceCacheParameters,
-						ELumenReflectionPass::Opaque,
-						nullptr,
-						nullptr,
-						ERDGPassFlags::AsyncCompute);
-				}
-			}
-			else
+			if (ViewPipelineState.ReflectionsMethod == EReflectionsMethod::Lumen && GLumenReflectionsAsyncCompute != 0)
 			{
-				const FRDGSystemTextures& SystemTextures = FRDGSystemTextures::Get(GraphBuilder);
-				ViewOutputs.IndirectLightingTextures.Textures[2] = SystemTextures.Black;
-				ViewOutputs.IndirectLightingTextures.Textures[3] = SystemTextures.Black;
+				ViewOutputs.IndirectLightingTextures.Textures[3] = RenderLumenReflections(
+					GraphBuilder,
+					View,
+					SceneTextures,
+					LumenFrameTemporaries,
+					ViewOutputs.MeshSDFGridParameters,
+					ViewOutputs.RadianceCacheParameters,
+					ELumenReflectionPass::Opaque,
+					nullptr,
+					nullptr,
+					ERDGPassFlags::AsyncCompute);
 			}
 
 			// Lumen needs its own depth history because things like Translucency velocities write to depth
@@ -961,15 +952,16 @@ void FDeferredShadingSceneRenderer::RenderDiffuseIndirectAndAmbientOcclusion(
 					ERDGPassFlags::Compute);
 			}
 
-			if (EnumHasAnyFlags(StepsLeft, ELumenIndirectLightingSteps::Reflections))
-			{
-				auto& ViewOutputs = AsyncLumenIndirectLightingOutputs.ViewOutputs;
-				const auto& MeshSDFGridParams = bDoComposite ? MeshSDFGridParameters : ViewOutputs[ViewIndex].MeshSDFGridParameters;
-				const auto& RadianceCacheParams = bDoComposite ? RadianceCacheParameters : ViewOutputs[ViewIndex].RadianceCacheParameters;
-				auto& OutTextures = bDoComposite ? DenoiserOutputs : ViewOutputs[ViewIndex].IndirectLightingTextures;
+			auto& AllViewOutputs = AsyncLumenIndirectLightingOutputs.ViewOutputs;
+			auto& OutTextures = bDoComposite ? DenoiserOutputs : AllViewOutputs[ViewIndex].IndirectLightingTextures;
 
-				if (ViewPipelineState.ReflectionsMethod == EReflectionsMethod::Lumen)
+			if (ViewPipelineState.ReflectionsMethod == EReflectionsMethod::Lumen)
+			{
+				if (EnumHasAnyFlags(StepsLeft, ELumenIndirectLightingSteps::Reflections))
 				{
+					const auto& MeshSDFGridParams = bDoComposite ? MeshSDFGridParameters : AllViewOutputs[ViewIndex].MeshSDFGridParameters;
+					const auto& RadianceCacheParams = bDoComposite ? RadianceCacheParameters : AllViewOutputs[ViewIndex].RadianceCacheParameters;
+
 					OutTextures.Textures[3] = RenderLumenReflections(
 						GraphBuilder,
 						View,
@@ -982,23 +974,23 @@ void FDeferredShadingSceneRenderer::RenderDiffuseIndirectAndAmbientOcclusion(
 						nullptr,
 						ERDGPassFlags::Compute);
 				}
-				else if (ViewPipelineState.ReflectionsMethod == EReflectionsMethod::SSR)
-				{
-					ESSRQuality SSRQuality;
-					IScreenSpaceDenoiser::FReflectionsRayTracingConfig DenoiserConfig;
-					ScreenSpaceRayTracing::GetSSRQualityForView(View, &SSRQuality, &DenoiserConfig);
-
-					RDG_EVENT_SCOPE(GraphBuilder, "ScreenSpaceReflections(Quality=%d)", int32(SSRQuality));
-					IScreenSpaceDenoiser::FReflectionsInputs SSRDenoiserInputs;
-					ScreenSpaceRayTracing::RenderScreenSpaceReflections(GraphBuilder, SceneTextureParameters, SceneColorTexture, View, SSRQuality, /*bDenoise*/ false, &SSRDenoiserInputs);
-					OutTextures.Textures[3] = SSRDenoiserInputs.Color;
-				}
-				else
-				{
-					// Remove Lumen rough specular when ReflectionsMethod != EReflectionsMethod::Lumen, so that we won't get double specular when other reflection methods are used
-					OutTextures.Textures[2] = SystemTextures.Black;
-					OutTextures.Textures[3] = SystemTextures.Black;
-				}
+			}
+			else if (ViewPipelineState.ReflectionsMethod == EReflectionsMethod::SSR)
+			{
+				ESSRQuality SSRQuality;
+				IScreenSpaceDenoiser::FReflectionsRayTracingConfig DenoiserConfig;
+				ScreenSpaceRayTracing::GetSSRQualityForView(View, &SSRQuality, &DenoiserConfig);
+
+				RDG_EVENT_SCOPE(GraphBuilder, "ScreenSpaceReflections(Quality=%d)", int32(SSRQuality));
+				IScreenSpaceDenoiser::FReflectionsInputs SSRDenoiserInputs;
+				ScreenSpaceRayTracing::RenderScreenSpaceReflections(GraphBuilder, SceneTextureParameters, SceneColorTexture, View, SSRQuality, /*bDenoise*/ false, &SSRDenoiserInputs);
+				OutTextures.Textures[3] = SSRDenoiserInputs.Color;
+			}
+			else
+			{
+				// Remove Lumen rough specular when ReflectionsMethod != EReflectionsMethod::Lumen, so that we won't get double specular when other reflection methods are used
+				OutTextures.Textures[2] = SystemTextures.Black;
+				OutTextures.Textures[3] = SystemTextures.Black;
 			}
 
 			if (EnumHasAnyFlags(StepsLeft, ELumenIndirectLightingSteps::StoreDepthHistory))
