Fixed several issues with split-screen rendering on mobile
Crash in CSM setup
Wrong shadows for a second view
Wrong occlusion queries rendering for a second view
Crash in deferred shading mode

diff --git a/Engine/Source/Runtime/Renderer/Private/MobileBasePassRendering.cpp b/Engine/Source/Runtime/Renderer/Private/MobileBasePassRendering.cpp
--- a/Engine/Source/Runtime/Renderer/Private/MobileBasePassRendering.cpp
+++ b/Engine/Source/Runtime/Renderer/Private/MobileBasePassRendering.cpp
@@ -251,7 +251,6 @@ TRDGUniformBufferRef<FMobileBasePassUniformParameters> CreateMobileBasePassUnifo
 
 void SetupMobileDirectionalLightUniformParameters(
 	const FScene& Scene,
-	int32 ViewIndex,
 	const FViewInfo& SceneView,
 	const TArray<FVisibleLightInfo,SceneRenderingAllocator>& VisibleLightInfos,
 	int32 ChannelIdx,
@@ -301,9 +300,9 @@ void SetupMobileDirectionalLightUniformParameters(
 			int32_t OutShadowIndex = 0;
 			for (int32 i = 0; i < NumShadowsToCopy && OutShadowIndex < SceneView.MaxShadowCascades; ++i)
 			{
-				const FProjectedShadowInfo* ShadowInfo = DirectionalLightShadowInfos[i + ViewIndex * NumShadowsToCopy];
+				const FProjectedShadowInfo* ShadowInfo = DirectionalLightShadowInfos[i];
 
-				if (ShadowInfo->ShadowDepthView && !ShadowInfo->bRayTracedDistanceField && ShadowInfo->CacheMode != SDCM_StaticPrimitivesOnly)
+				if (ShadowInfo->ShadowDepthView && !ShadowInfo->bRayTracedDistanceField && ShadowInfo->CacheMode != SDCM_StaticPrimitivesOnly && ShadowInfo->DependentView == &SceneView)
 				{
 					if (OutShadowIndex == 0)
 					{
diff --git a/Engine/Source/Runtime/Renderer/Private/MobileBasePassRendering.h b/Engine/Source/Runtime/Renderer/Private/MobileBasePassRendering.h
--- a/Engine/Source/Runtime/Renderer/Private/MobileBasePassRendering.h
+++ b/Engine/Source/Runtime/Renderer/Private/MobileBasePassRendering.h
@@ -70,7 +70,6 @@ extern TRDGUniformBufferRef<FMobileBasePassUniformParameters> CreateMobileBasePa
 
 extern void SetupMobileDirectionalLightUniformParameters(
 	const FScene& Scene,
-	int32 ViewIndex,
 	const FViewInfo& View,
 	const TArray<FVisibleLightInfo,SceneRenderingAllocator>& VisibleLightInfos,
 	int32 ChannelIdx,
diff --git a/Engine/Source/Runtime/Renderer/Private/MobileShadingRenderer.cpp b/Engine/Source/Runtime/Renderer/Private/MobileShadingRenderer.cpp
--- a/Engine/Source/Runtime/Renderer/Private/MobileShadingRenderer.cpp
+++ b/Engine/Source/Runtime/Renderer/Private/MobileShadingRenderer.cpp
@@ -254,6 +254,38 @@ static void PollOcclusionQueriesPass(FRDGBuilder& GraphBuilder)
 	});
 }
 
+struct FRenderViewContext
+{
+	FViewInfo* ViewInfo;
+	int32 ViewIndex;
+	bool bIsFirstView;
+	bool bIsLastView;
+};
+using FRenderViewContextArray = TArray<FRenderViewContext, TInlineAllocator<2, SceneRenderingAllocator>>;
+
+static void GetRenderViews(TArray<FViewInfo>& InViews, FRenderViewContextArray& RenderViews)
+{
+	for (int32 ViewIndex = 0; ViewIndex < InViews.Num(); ViewIndex++)
+	{
+		FViewInfo& View = InViews[ViewIndex];
+		if (View.ShouldRenderView())
+		{
+			FRenderViewContext RenderView;
+			RenderView.ViewInfo = &View;
+			RenderView.ViewIndex = ViewIndex;
+			RenderView.bIsFirstView = (RenderViews.Num() == 0);
+			RenderView.bIsLastView = false;
+
+			RenderViews.Add(RenderView);
+		}
+	}
+
+	if (RenderViews.Num())
+	{
+		RenderViews.Last().bIsLastView = true;
+	}
+}
+
 FMobileSceneRenderer::FMobileSceneRenderer(const FSceneViewFamily* InViewFamily, FHitProxyConsumer* HitProxyConsumer)
 	: FSceneRenderer(InViewFamily, HitProxyConsumer)
 	, bGammaSpace(!IsMobileHDR())
@@ -542,14 +574,14 @@ void FMobileSceneRenderer::InitViews(FRDGBuilder& GraphBuilder, FSceneTexturesCo
 	// Update the bKeepDepthContent based on the mobile renderer status.
 	SceneTexturesConfig.bKeepDepthContent = bKeepDepthContent;
 	// If we render in a single pass MSAA targets can be memoryless
-    SceneTexturesConfig.bMemorylessMSAA = !(bRequiresMultiPass || bShouldCompositeEditorPrimitives);
+    SceneTexturesConfig.bMemorylessMSAA = !(bRequiresMultiPass || bShouldCompositeEditorPrimitives || bRequireSeparateViewPass);
     SceneTexturesConfig.NumSamples = NumMSAASamples;
     
     SceneTexturesConfig.BuildSceneColorAndDepthFlags();
     
 	if (bDeferredShading) 
 	{
-		SetupGBufferFlags(SceneTexturesConfig, bRequiresMultiPass || GraphBuilder.IsDumpingFrame());
+		SetupGBufferFlags(SceneTexturesConfig, bRequiresMultiPass || GraphBuilder.IsDumpingFrame() || bRequireSeparateViewPass);
 	}
 
 	// Update the pixel projected reflection extent according to the settings in the PlanarReflectionComponent.
@@ -708,15 +740,14 @@ void FMobileSceneRenderer::RenderFullDepthPrepass(FRDGBuilder& GraphBuilder, FSc
 	BasePassRenderTargets.DepthStencil = FDepthStencilBinding(SceneTextures.Depth.Target, ERenderTargetLoadAction::EClear, ERenderTargetLoadAction::EClear, FExclusiveDepthStencil::DepthWrite_StencilWrite);
 	BasePassRenderTargets.NumOcclusionQueries = ComputeNumOcclusionQueriesToBatch();
 
-	for (int32 ViewIndex = 0; ViewIndex < Views.Num(); ViewIndex++)
+	FRenderViewContextArray RenderViews;
+	GetRenderViews(Views, RenderViews);
+
+	for (FRenderViewContext& ViewContext : RenderViews)
 	{
-		FViewInfo& View = Views[ViewIndex];
-		if (!View.ShouldRenderView())
-		{
-			continue;
-		}
+		FViewInfo& View = *ViewContext.ViewInfo;
 
-		if (ViewIndex > 0)
+		if (!ViewContext.bIsFirstView)
 		{
 			BasePassRenderTargets.DepthStencil.SetDepthLoadAction(ERenderTargetLoadAction::ELoad);
 			BasePassRenderTargets.DepthStencil.SetStencilLoadAction(ERenderTargetLoadAction::ELoad);
@@ -734,17 +765,22 @@ void FMobileSceneRenderer::RenderFullDepthPrepass(FRDGBuilder& GraphBuilder, FSc
 			BuildMeshPassInstanceCullingDrawParams(GraphBuilder, Scene->GPUScene, View.ParallelMeshDrawCommandPasses[EMeshPass::DepthPass], *PassParameters, MeshPassInstanceCullingDrawParams[EMeshPass::DepthPass]);
 		}
 
+		// Render occlusion at the last view pass only, as they already loop through all views
+		bool bDoOcclusionQueries = (ViewContext.bIsLastView && DoOcclusionQueries());
+
 		GraphBuilder.AddPass(
 			RDG_EVENT_NAME("FullDepthPrepass"),
 			PassParameters,
 			ERDGPassFlags::Raster,
-			[this, PassParameters, &View](FRHICommandListImmediate& RHICmdList)
+			[this, PassParameters, &View, bDoOcclusionQueries](FRHICommandListImmediate& RHICmdList)
 			{
 				RenderPrePass(RHICmdList, View);
 
-				// Issue occlusion queries
-				RHICmdList.SetCurrentStat(GET_STATID(STAT_CLMM_Occlusion));
-				RenderOcclusion(RHICmdList, View);
+				if (bDoOcclusionQueries)
+				{
+					RHICmdList.SetCurrentStat(GET_STATID(STAT_CLMM_Occlusion));
+					RenderOcclusion(RHICmdList);
+				}
 			});
 	}
 
@@ -1107,19 +1143,17 @@ void FMobileSceneRenderer::RenderForward(FRDGBuilder& GraphBuilder, FRDGTextureR
 
 	const FRDGSystemTextures& SystemTextures = FRDGSystemTextures::Get(GraphBuilder);
 
-	for (int32 ViewIndex = 0; ViewIndex < Views.Num(); ViewIndex++)
+	FRenderViewContextArray RenderViews;
+	GetRenderViews(Views, RenderViews);
+
+	for (FRenderViewContext& ViewContext : RenderViews)
 	{
-		FViewInfo& View = Views[ViewIndex];
+		FViewInfo& View = *ViewContext.ViewInfo;
 
 		SCOPED_GPU_MASK(GraphBuilder.RHICmdList, !View.IsInstancedStereoPass() ? View.GPUMask : (View.GPUMask | View.GetInstancedView()->GPUMask));
-		SCOPED_CONDITIONAL_DRAW_EVENTF(GraphBuilder.RHICmdList, EventView, Views.Num() > 1, TEXT("View%d"), ViewIndex);
+		SCOPED_CONDITIONAL_DRAW_EVENTF(GraphBuilder.RHICmdList, EventView, RenderViews.Num() > 1, TEXT("View%d"), ViewContext.ViewIndex);
 
-		if (!View.ShouldRenderView())
-		{
-			continue;
-		}
-
-		if (ViewIndex > 0)
+		if (!ViewContext.bIsFirstView)
 		{
 			BasePassRenderTargets[0].SetLoadAction(ERenderTargetLoadAction::ELoad);
 			if (bRequiresSceneDepthAux)
@@ -1133,7 +1167,7 @@ void FMobileSceneRenderer::RenderForward(FRDGBuilder& GraphBuilder, FRDGTextureR
 
 		View.BeginRenderView();
 
-		UpdateDirectionalLightUniformBuffers(GraphBuilder, ViewIndex, View);
+		UpdateDirectionalLightUniformBuffers(GraphBuilder, View);
 
 		FMobileBasePassTextures MobileBasePassTextures{};
 		MobileBasePassTextures.ScreenSpaceAO = bRequiresAmbientOcclusionPass ? SceneTextures.ScreenSpaceAO : SystemTextures.White;
@@ -1150,31 +1184,32 @@ void FMobileSceneRenderer::RenderForward(FRDGBuilder& GraphBuilder, FRDGTextureR
 		// Split if we need to render translucency in a separate render pass
 		if (bRequiresMultiPass)
 		{
-			RenderForwardMultiPass(GraphBuilder, PassParameters, BasePassRenderTargets, ViewIndex, View, SceneTextures);
+			RenderForwardMultiPass(GraphBuilder, PassParameters, BasePassRenderTargets, ViewContext, SceneTextures);
 		}
 		else
 		{
-			RenderForwardSinglePass(GraphBuilder, PassParameters, ViewIndex, View, SceneTextures);
+			RenderForwardSinglePass(GraphBuilder, PassParameters, ViewContext, SceneTextures);
 		}
 	}
 }
 
-void FMobileSceneRenderer::RenderForwardSinglePass(FRDGBuilder& GraphBuilder, FMobileRenderPassParameters* PassParameters, int32 ViewIndex, FViewInfo& View, FSceneTextures& SceneTextures)
+void FMobileSceneRenderer::RenderForwardSinglePass(FRDGBuilder& GraphBuilder, FMobileRenderPassParameters* PassParameters, FRenderViewContext& ViewContext, FSceneTextures& SceneTextures)
 {
 	PassParameters->RenderTargets.SubpassHint = ESubpassHint::DepthReadSubpass;
-	if (!bIsFullDepthPrepassEnabled)
-	{
-		PassParameters->RenderTargets.NumOcclusionQueries = ComputeNumOcclusionQueriesToBatch();
-	}
+	
+	const bool bDoOcclusionQueries = (!bIsFullDepthPrepassEnabled && ViewContext.bIsLastView && DoOcclusionQueries());
+	PassParameters->RenderTargets.NumOcclusionQueries = bDoOcclusionQueries ? ComputeNumOcclusionQueriesToBatch() : 0u;
 	
 	GraphBuilder.AddPass(
 		RDG_EVENT_NAME("SceneColorRendering"),
 		PassParameters,
 		// the second view pass should not be merged with the first view pass on mobile since the subpass would not work properly.
 		ERDGPassFlags::Raster | ERDGPassFlags::NeverMerge,
-		[this, PassParameters, ViewIndex, &View, &SceneTextures](FRHICommandListImmediate& RHICmdList)
+		[this, PassParameters, ViewContext, bDoOcclusionQueries, &SceneTextures](FRHICommandListImmediate& RHICmdList)
 	{
-		if (GIsEditor && !View.bIsSceneCapture && ViewIndex == 0)
+		FViewInfo& View = *ViewContext.ViewInfo;
+			
+		if (GIsEditor && !View.bIsSceneCapture && ViewContext.bIsFirstView)
 		{
 			DrawClearQuad(RHICmdList, View.BackgroundColor);
 		}
@@ -1192,12 +1227,12 @@ void FMobileSceneRenderer::RenderForwardSinglePass(FRDGBuilder& GraphBuilder, FM
 		RHICmdList.NextSubpass();
 		RHICmdList.SetCurrentStat(GET_STATID(STAT_CLMM_Translucency));
 		RenderDecals(RHICmdList, View);
-		RenderModulatedShadowProjections(RHICmdList, ViewIndex, View);
+		RenderModulatedShadowProjections(RHICmdList, ViewContext.ViewIndex, View);
 		RenderFog(RHICmdList, View);
 		// Draw translucency.
 		RenderTranslucency(RHICmdList, View);
 
-		if (!bIsFullDepthPrepassEnabled)
+		if (bDoOcclusionQueries)
 		{
 			// Issue occlusion queries
 			RHICmdList.SetCurrentStat(GET_STATID(STAT_CLMM_Occlusion));
@@ -1207,7 +1242,7 @@ void FMobileSceneRenderer::RenderForwardSinglePass(FRDGBuilder& GraphBuilder, FM
 				// flush
 				RHICmdList.SubmitCommandsHint();
 			}
-			RenderOcclusion(RHICmdList, View);
+			RenderOcclusion(RHICmdList);
 		}
 
 		// Pre-tonemap before MSAA resolve (iOS only)
@@ -1217,19 +1252,21 @@ void FMobileSceneRenderer::RenderForwardSinglePass(FRDGBuilder& GraphBuilder, FM
 	// resolve MSAA depth
 	if (!bIsFullDepthPrepassEnabled)
 	{
-		AddResolveSceneDepthPass(GraphBuilder, View, SceneTextures.Depth);
+		AddResolveSceneDepthPass(GraphBuilder, *ViewContext.ViewInfo, SceneTextures.Depth);
 	}
 }
 
-void FMobileSceneRenderer::RenderForwardMultiPass(FRDGBuilder& GraphBuilder, FMobileRenderPassParameters* PassParameters, FRenderTargetBindingSlots& BasePassRenderTargets, int32 ViewIndex, FViewInfo& View, FSceneTextures& SceneTextures)
+void FMobileSceneRenderer::RenderForwardMultiPass(FRDGBuilder& GraphBuilder, FMobileRenderPassParameters* PassParameters, FRenderTargetBindingSlots& BasePassRenderTargets, FRenderViewContext& ViewContext, FSceneTextures& SceneTextures)
 {
 	GraphBuilder.AddPass(
 		RDG_EVENT_NAME("SceneColorRendering"),
 		PassParameters,
 		ERDGPassFlags::Raster,
-		[this, PassParameters, ViewIndex, &View, &SceneTextures](FRHICommandListImmediate& RHICmdList)
+		[this, PassParameters, ViewContext, &SceneTextures](FRHICommandListImmediate& RHICmdList)
 	{
-		if (GIsEditor && !View.bIsSceneCapture && ViewIndex == 0)
+		FViewInfo& View = *ViewContext.ViewInfo;
+			
+		if (GIsEditor && !View.bIsSceneCapture && ViewContext.bIsFirstView)
 		{
 			DrawClearQuad(RHICmdList, View.BackgroundColor);
 		}
@@ -1245,6 +1282,8 @@ void FMobileSceneRenderer::RenderForwardMultiPass(FRDGBuilder& GraphBuilder, FMo
 		PostRenderBasePass(RHICmdList, View);
 	});
 
+	FViewInfo& View = *ViewContext.ViewInfo;
+
 	// resolve MSAA depth
 	if (!bIsFullDepthPrepassEnabled)
 	{
@@ -1277,30 +1316,31 @@ void FMobileSceneRenderer::RenderForwardMultiPass(FRDGBuilder& GraphBuilder, FMo
 	SecondPassParameters->MobileBasePass = CreateMobileBasePassUniformBuffer(GraphBuilder, View, EMobileBasePass::Translucent, SetupMode);
 	SecondPassParameters->ReflectionCapture = View.MobileReflectionCaptureUniformBuffer;
 	SecondPassParameters->RenderTargets = BasePassRenderTargets;
-	if (!bIsFullDepthPrepassEnabled)
-	{
-		SecondPassParameters->RenderTargets.NumOcclusionQueries = ComputeNumOcclusionQueriesToBatch();
-	}
+	
+	const bool bDoOcclusionQueries = (!bIsFullDepthPrepassEnabled && ViewContext.bIsLastView && DoOcclusionQueries());
+	SecondPassParameters->RenderTargets.NumOcclusionQueries = bDoOcclusionQueries ? ComputeNumOcclusionQueriesToBatch() : 0u;
 
 	GraphBuilder.AddPass(
 		RDG_EVENT_NAME("DecalsAndTranslucency"),
 		SecondPassParameters,
 		ERDGPassFlags::Raster,
-		[this, SecondPassParameters, ViewIndex, &View, &SceneTextures](FRHICommandListImmediate& RHICmdList)
+		[this, SecondPassParameters, ViewContext, bDoOcclusionQueries, &SceneTextures](FRHICommandListImmediate& RHICmdList)
 	{
+		FViewInfo& View = *ViewContext.ViewInfo;
+			
 		// scene depth is read only and can be fetched
 		RHICmdList.SetCurrentStat(GET_STATID(STAT_CLMM_Translucency));
 		RenderDecals(RHICmdList, View);
-		RenderModulatedShadowProjections(RHICmdList, ViewIndex, View);
+		RenderModulatedShadowProjections(RHICmdList, ViewContext.ViewIndex, View);
 		RenderFog(RHICmdList, View);
 		// Draw translucency.
 		RenderTranslucency(RHICmdList, View);
 
-		if (!bIsFullDepthPrepassEnabled)
+		if (bDoOcclusionQueries)
 		{
 			// Issue occlusion queries
 			RHICmdList.SetCurrentStat(GET_STATID(STAT_CLMM_Occlusion));
-			RenderOcclusion(RHICmdList, View);
+			RenderOcclusion(RHICmdList);
 		}
 
 		// Pre-tonemap before MSAA resolve (iOS only)
@@ -1435,23 +1475,31 @@ void FMobileSceneRenderer::RenderDeferred(FRDGBuilder& GraphBuilder, const FSort
 
 	const FRDGSystemTextures& SystemTextures = FRDGSystemTextures::Get(GraphBuilder);
 
-	int32 NumViews = Views.Num();
+	FRenderViewContextArray RenderViews;
+	GetRenderViews(Views, RenderViews);
 
-	for (int32 ViewIndex = 0; ViewIndex < NumViews; ViewIndex++)
+	for (FRenderViewContext& ViewContext : RenderViews)
 	{
-		FViewInfo& View = Views[ViewIndex];
+		FViewInfo& View = *ViewContext.ViewInfo;
 
 		SCOPED_GPU_MASK(GraphBuilder.RHICmdList, !View.IsInstancedStereoPass() ? View.GPUMask : (View.GPUMask | View.GetInstancedView()->GPUMask));
-		SCOPED_CONDITIONAL_DRAW_EVENTF(GraphBuilder.RHICmdList, EventView, NumViews > 1, TEXT("View%d"), ViewIndex);
+		SCOPED_CONDITIONAL_DRAW_EVENTF(GraphBuilder.RHICmdList, EventView, RenderViews.Num() > 1, TEXT("View%d"), ViewContext.ViewIndex);
 
-		if (!View.ShouldRenderView())
+		if (!ViewContext.bIsFirstView)
 		{
-			continue;
+			// Load targets for a non-first view 
+			for (int32 i = 0; i < ColorTargets.Num(); ++i)
+			{
+				BasePassRenderTargets[i].SetLoadAction(ERenderTargetLoadAction::ELoad);
+			}
+			BasePassRenderTargets.DepthStencil.SetDepthLoadAction(ERenderTargetLoadAction::ELoad);
+			BasePassRenderTargets.DepthStencil.SetStencilLoadAction(ERenderTargetLoadAction::ELoad);
+			BasePassRenderTargets.DepthStencil.SetDepthStencilAccess(bIsFullDepthPrepassEnabled ? FExclusiveDepthStencil::DepthRead_StencilWrite : FExclusiveDepthStencil::DepthWrite_StencilWrite);
 		}
 
 		View.BeginRenderView();
 
-		UpdateDirectionalLightUniformBuffers(GraphBuilder, ViewIndex, View);
+		UpdateDirectionalLightUniformBuffers(GraphBuilder, View);
 
 		FMobileBasePassTextures MobileBasePassTextures{};
 		MobileBasePassTextures.ScreenSpaceAO = bRequiresAmbientOcclusionPass ? SceneTextures.ScreenSpaceAO : SystemTextures.White;
@@ -1467,29 +1515,31 @@ void FMobileSceneRenderer::RenderDeferred(FRDGBuilder& GraphBuilder, const FSort
 
 		if (bRequiresMultiPass)
 		{
-			RenderDeferredMultiPass(GraphBuilder, PassParameters, BasePassRenderTargets, ColorTargets.Num(), ViewIndex, NumViews, View, SceneTextures, SortedLightSet);
+			RenderDeferredMultiPass(GraphBuilder, PassParameters, BasePassRenderTargets, ColorTargets.Num(), ViewContext, SceneTextures, SortedLightSet);
 		}
 		else
 		{
-			RenderDeferredSinglePass(GraphBuilder, PassParameters, ViewIndex, NumViews, View, SceneTextures, SortedLightSet, bUsingPixelLocalStorage);
+			RenderDeferredSinglePass(GraphBuilder, PassParameters, ViewContext, SceneTextures, SortedLightSet, bUsingPixelLocalStorage);
 		}
 	}
 }
 
-void FMobileSceneRenderer::RenderDeferredSinglePass(FRDGBuilder& GraphBuilder, class FMobileRenderPassParameters* PassParameters, int32 ViewIndex, int32 NumViews, FViewInfo& View, FSceneTextures& SceneTextures, const FSortedLightSetSceneInfo& SortedLightSet, bool bUsingPixelLocalStorage)
+void FMobileSceneRenderer::RenderDeferredSinglePass(FRDGBuilder& GraphBuilder, class FMobileRenderPassParameters* PassParameters, FRenderViewContext& ViewContext, FSceneTextures& SceneTextures, const FSortedLightSetSceneInfo& SortedLightSet, bool bUsingPixelLocalStorage)
 {
 	PassParameters->RenderTargets.SubpassHint = ESubpassHint::DeferredShadingSubpass;
-	if (!bIsFullDepthPrepassEnabled)
-	{
-		PassParameters->RenderTargets.NumOcclusionQueries = ComputeNumOcclusionQueriesToBatch();
-	}
-		
+
+	const bool bDoOcclusionQueires = (!bIsFullDepthPrepassEnabled && ViewContext.bIsLastView && DoOcclusionQueries());
+	PassParameters->RenderTargets.NumOcclusionQueries = bDoOcclusionQueires ? ComputeNumOcclusionQueriesToBatch() : 0u;
+				
 	GraphBuilder.AddPass(
 		RDG_EVENT_NAME("SceneColorRendering"),
 		PassParameters,
-		ERDGPassFlags::Raster,
-		[this, PassParameters, ViewIndex, NumViews, &View, &SceneTextures, &SortedLightSet, bUsingPixelLocalStorage](FRHICommandListImmediate& RHICmdList)
+		// the second view pass should not be merged with the first view pass on mobile since the subpass would not work properly.
+		ERDGPassFlags::Raster | ERDGPassFlags::NeverMerge,
+		[this, PassParameters, ViewContext, bDoOcclusionQueires, &SceneTextures, &SortedLightSet, bUsingPixelLocalStorage](FRHICommandListImmediate& RHICmdList)
 	{
+		FViewInfo& View = *ViewContext.ViewInfo;
+			
 		// Depth pre-pass
 		RHICmdList.SetCurrentStat(GET_STATID(STAT_CLM_MobilePrePass));
 		RenderMaskedPrePass(RHICmdList, View);
@@ -1504,7 +1554,7 @@ void FMobileSceneRenderer::RenderDeferredSinglePass(FRDGBuilder& GraphBuilder, c
 		RenderDecals(RHICmdList, View);
 		// SceneColor write, SceneDepth is read only
 		RHICmdList.NextSubpass();
-		MobileDeferredShadingPass(RHICmdList, ViewIndex, NumViews, View, *Scene, SortedLightSet, VisibleLightInfos);
+		MobileDeferredShadingPass(RHICmdList, ViewContext.ViewIndex, Views.Num(), View, *Scene, SortedLightSet, VisibleLightInfos);
 		if (bUsingPixelLocalStorage)
 		{
 			MobileDeferredCopyBuffer<FMobileDeferredCopyPLSPS>(RHICmdList, View);
@@ -1513,23 +1563,27 @@ void FMobileSceneRenderer::RenderDeferredSinglePass(FRDGBuilder& GraphBuilder, c
 		// Draw translucency.
 		RenderTranslucency(RHICmdList, View);
 
-		if (!bIsFullDepthPrepassEnabled)
+		if (bDoOcclusionQueires)
 		{
 			// Issue occlusion queries
 			RHICmdList.SetCurrentStat(GET_STATID(STAT_CLMM_Occlusion));
-			RenderOcclusion(RHICmdList, View);
+			RenderOcclusion(RHICmdList);
 		}
 	});
 }
 
-void FMobileSceneRenderer::RenderDeferredMultiPass(FRDGBuilder& GraphBuilder, class FMobileRenderPassParameters* PassParameters, FRenderTargetBindingSlots& BasePassRenderTargets, int32 NumColorTargets, int32 ViewIndex, int32 NumViews, FViewInfo& View, FSceneTextures& SceneTextures, const FSortedLightSetSceneInfo& SortedLightSet)
+void FMobileSceneRenderer::RenderDeferredMultiPass(FRDGBuilder& GraphBuilder, class FMobileRenderPassParameters* PassParameters, const FRenderTargetBindingSlots& InBasePassRenderTargets, int32 NumColorTargets, FRenderViewContext& ViewContext, FSceneTextures& SceneTextures, const FSortedLightSetSceneInfo& SortedLightSet)
 {
+	FRenderTargetBindingSlots BasePassRenderTargets = InBasePassRenderTargets;
+	
 	GraphBuilder.AddPass(
 		RDG_EVENT_NAME("SceneColorRendering"),
 		PassParameters,
 		ERDGPassFlags::Raster,
-		[this, PassParameters, &View, &SceneTextures](FRHICommandListImmediate& RHICmdList)
+		[this, PassParameters, ViewContext, &SceneTextures](FRHICommandListImmediate& RHICmdList)
 	{
+		FViewInfo& View = *ViewContext.ViewInfo;
+		
 		// Depth pre-pass
 		RHICmdList.SetCurrentStat(GET_STATID(STAT_CLM_MobilePrePass));
 		RenderMaskedPrePass(RHICmdList, View);
@@ -1550,6 +1604,8 @@ void FMobileSceneRenderer::RenderDeferredMultiPass(FRDGBuilder& GraphBuilder, cl
 	BasePassRenderTargets.DepthStencil.SetStencilLoadAction(ERenderTargetLoadAction::ELoad);
 	BasePassRenderTargets.DepthStencil.SetDepthStencilAccess(FExclusiveDepthStencil::DepthRead_StencilRead);
 	
+	FViewInfo& View = *ViewContext.ViewInfo;
+	
 	EMobileSceneTextureSetupMode SetupMode = EMobileSceneTextureSetupMode::SceneDepth | EMobileSceneTextureSetupMode::SceneDepthAux | EMobileSceneTextureSetupMode::CustomDepth;
 	auto* SecondPassParameters = GraphBuilder.AllocParameters<FMobileRenderPassParameters>();
 	*SecondPassParameters = *PassParameters;
@@ -1561,8 +1617,9 @@ void FMobileSceneRenderer::RenderDeferredMultiPass(FRDGBuilder& GraphBuilder, cl
 		RDG_EVENT_NAME("Decals"),
 		SecondPassParameters,
 		ERDGPassFlags::Raster,
-		[this, SecondPassParameters, &View](FRHICommandListImmediate& RHICmdList)
+		[this, SecondPassParameters, ViewContext](FRHICommandListImmediate& RHICmdList)
 	{
+		FViewInfo& View = *ViewContext.ViewInfo;
 		RenderDecals(RHICmdList, View);
 	});
 
@@ -1579,28 +1636,29 @@ void FMobileSceneRenderer::RenderDeferredMultiPass(FRDGBuilder& GraphBuilder, cl
 	ThirdPassParameters->MobileBasePass = CreateMobileBasePassUniformBuffer(GraphBuilder, View, EMobileBasePass::Translucent, SetupMode);
 	ThirdPassParameters->ReflectionCapture = View.MobileReflectionCaptureUniformBuffer;
 	ThirdPassParameters->RenderTargets = BasePassRenderTargets;
-	if (!bIsFullDepthPrepassEnabled)
-	{
-		ThirdPassParameters->RenderTargets.NumOcclusionQueries = ComputeNumOcclusionQueriesToBatch();
-	}
+	
+	const bool bDoOcclusionQueires = (!bIsFullDepthPrepassEnabled && ViewContext.bIsLastView && DoOcclusionQueries());
+	ThirdPassParameters->RenderTargets.NumOcclusionQueries = bDoOcclusionQueires ? ComputeNumOcclusionQueriesToBatch() : 0u;
 
 	GraphBuilder.AddPass(
 		RDG_EVENT_NAME("LightingAndTranslucency"),
 		ThirdPassParameters,
 		ERDGPassFlags::Raster,
-		[this, ThirdPassParameters, ViewIndex, NumViews, &View, &SceneTextures, &SortedLightSet](FRHICommandListImmediate& RHICmdList)
+		[this, ThirdPassParameters, ViewContext, bDoOcclusionQueires, &SceneTextures, &SortedLightSet](FRHICommandListImmediate& RHICmdList)
 	{
+		FViewInfo& View = *ViewContext.ViewInfo;
+			
 		RHICmdList.SetCurrentStat(GET_STATID(STAT_CLMM_Translucency));
-		MobileDeferredShadingPass(RHICmdList, ViewIndex, NumViews, View, *Scene, SortedLightSet, VisibleLightInfos);
+		MobileDeferredShadingPass(RHICmdList, ViewContext.ViewIndex, Views.Num(), View, *Scene, SortedLightSet, VisibleLightInfos);
 		RenderFog(RHICmdList, View);
 		// Draw translucency.
 		RenderTranslucency(RHICmdList, View);
 
-		if (!bIsFullDepthPrepassEnabled)
+		if (bDoOcclusionQueires)
 		{
 			// Issue occlusion queries
 			RHICmdList.SetCurrentStat(GET_STATID(STAT_CLMM_Occlusion));
-			RenderOcclusion(RHICmdList, View);
+			RenderOcclusion(RHICmdList);
 		}
 	});
 }
@@ -1698,7 +1756,7 @@ bool FMobileSceneRenderer::RequiresMultiPass(FRHICommandListImmediate& RHICmdLis
 	return true;
 }
 
-void FMobileSceneRenderer::UpdateDirectionalLightUniformBuffers(FRDGBuilder& GraphBuilder, int32 ViewIndex, const FViewInfo& View)
+void FMobileSceneRenderer::UpdateDirectionalLightUniformBuffers(FRDGBuilder& GraphBuilder, const FViewInfo& View)
 {
 	if (CachedView == &View)
 	{
@@ -1706,14 +1764,14 @@ void FMobileSceneRenderer::UpdateDirectionalLightUniformBuffers(FRDGBuilder& Gra
 	}
 	CachedView = &View;
 
-	AddPass(GraphBuilder, RDG_EVENT_NAME("UpdateDirectionalLightUniformBuffers"), [this, ViewIndex, &View](FRHICommandListImmediate&)
+	AddPass(GraphBuilder, RDG_EVENT_NAME("UpdateDirectionalLightUniformBuffers"), [this, &View](FRHICommandListImmediate&)
 	{
 		const bool bDynamicShadows = ViewFamily.EngineShowFlags.DynamicShadows;
 		// Fill in the other entries based on the lights
 		for (int32 ChannelIdx = 0; ChannelIdx < UE_ARRAY_COUNT(Scene->MobileDirectionalLights); ChannelIdx++)
 		{
 			FMobileDirectionalLightShaderParameters Params;
-			SetupMobileDirectionalLightUniformParameters(*Scene, ViewIndex, View, VisibleLightInfos, ChannelIdx, bDynamicShadows, Params);
+			SetupMobileDirectionalLightUniformParameters(*Scene, View, VisibleLightInfos, ChannelIdx, bDynamicShadows, Params);
 			Scene->UniformBuffers.MobileDirectionalLightUniformBuffers[ChannelIdx + 1].UpdateUniformBufferImmediate(Params);
 		}
 	});
diff --git a/Engine/Source/Runtime/Renderer/Private/SceneOcclusion.cpp b/Engine/Source/Runtime/Renderer/Private/SceneOcclusion.cpp
--- a/Engine/Source/Runtime/Renderer/Private/SceneOcclusion.cpp
+++ b/Engine/Source/Runtime/Renderer/Private/SceneOcclusion.cpp
@@ -1441,7 +1441,7 @@ void FDeferredShadingSceneRenderer::RenderOcclusion(
 	}
 }
 
-void FMobileSceneRenderer::RenderOcclusion(FRHICommandListImmediate& RHICmdList, const FViewInfo& View)
+void FMobileSceneRenderer::RenderOcclusion(FRHICommandListImmediate& RHICmdList)
 {
 	if (!DoOcclusionQueries())
 	{
diff --git a/Engine/Source/Runtime/Renderer/Private/SceneRendering.h b/Engine/Source/Runtime/Renderer/Private/SceneRendering.h
--- a/Engine/Source/Runtime/Renderer/Private/SceneRendering.h
+++ b/Engine/Source/Runtime/Renderer/Private/SceneRendering.h
@@ -2487,7 +2487,7 @@ protected:
 	void RenderModulatedShadowProjections(FRHICommandListImmediate& RHICmdList, int32 ViewIndex, const FViewInfo& View);
 
 	/** Issues occlusion queries */
-	void RenderOcclusion(FRHICommandListImmediate& RHICmdList, const FViewInfo& View);
+	void RenderOcclusion(FRHICommandListImmediate& RHICmdList);
 	
 	bool ShouldRenderHZB();
 
@@ -2517,18 +2517,18 @@ protected:
 	void SortMobileBasePassAfterShadowInit(FExclusiveDepthStencil::Type BasePassDepthStencilAccess, FViewVisibleCommandsPerView& ViewCommandsPerView);
 	void SetupMobileBasePassAfterShadowInit(FExclusiveDepthStencil::Type BasePassDepthStencilAccess, FViewVisibleCommandsPerView& ViewCommandsPerView, FInstanceCullingManager& InstanceCullingManager);
 
-	void UpdateDirectionalLightUniformBuffers(FRDGBuilder& GraphBuilder, int32 ViewIndex, const FViewInfo& View);
+	void UpdateDirectionalLightUniformBuffers(FRDGBuilder& GraphBuilder, const FViewInfo& View);
 	void UpdateSkyReflectionUniformBuffer();
 	
 	void BuildInstanceCullingDrawParams(FRDGBuilder& GraphBuilder, FViewInfo& View, class FMobileRenderPassParameters* PassParameters);
 
 	void RenderForward(FRDGBuilder& GraphBuilder, FRDGTextureRef ViewFamilyTexture, FSceneTextures& SceneTextures);
-	void RenderForwardSinglePass(FRDGBuilder& GraphBuilder, class FMobileRenderPassParameters* PassParameters, int32 ViewIndex, FViewInfo& View, FSceneTextures& SceneTextures);
-	void RenderForwardMultiPass(FRDGBuilder& GraphBuilder, class FMobileRenderPassParameters* PassParameters, FRenderTargetBindingSlots& BasePassRenderTargets, int32 ViewIndex, FViewInfo& View, FSceneTextures& SceneTextures);
+	void RenderForwardSinglePass(FRDGBuilder& GraphBuilder, class FMobileRenderPassParameters* PassParameters, struct FRenderViewContext& ViewContext, FSceneTextures& SceneTextures);
+	void RenderForwardMultiPass(FRDGBuilder& GraphBuilder, class FMobileRenderPassParameters* PassParameters, FRenderTargetBindingSlots& BasePassRenderTargets, struct FRenderViewContext& ViewContext, FSceneTextures& SceneTextures);
 	
 	void RenderDeferred(FRDGBuilder& GraphBuilder, const FSortedLightSetSceneInfo& SortedLightSet, FRDGTextureRef ViewFamilyTexture, FSceneTextures& SceneTextures);
-	void RenderDeferredSinglePass(FRDGBuilder& GraphBuilder, class FMobileRenderPassParameters* PassParameters, int32 ViewIndex, int32 NumViews, FViewInfo& View, FSceneTextures& SceneTextures, const FSortedLightSetSceneInfo& SortedLightSet, bool bUsingPixelLocalStorage);
-	void RenderDeferredMultiPass(FRDGBuilder& GraphBuilder, class FMobileRenderPassParameters* PassParameters, FRenderTargetBindingSlots& BasePassRenderTargets, int32 NumColorTargets, int32 ViewIndex, int32 NumViews, FViewInfo& View, FSceneTextures& SceneTextures, const FSortedLightSetSceneInfo& SortedLightSet);
+	void RenderDeferredSinglePass(FRDGBuilder& GraphBuilder, class FMobileRenderPassParameters* PassParameters, struct FRenderViewContext& ViewContext, FSceneTextures& SceneTextures, const FSortedLightSetSceneInfo& SortedLightSet, bool bUsingPixelLocalStorage);
+	void RenderDeferredMultiPass(FRDGBuilder& GraphBuilder, class FMobileRenderPassParameters* PassParameters, const FRenderTargetBindingSlots& BasePassRenderTargets, int32 NumColorTargets, struct FRenderViewContext& ViewContext, FSceneTextures& SceneTextures, const FSortedLightSetSceneInfo& SortedLightSet);
 	
 	void RenderAmbientOcclusion(FRDGBuilder& GraphBuilder, FRDGTextureRef SceneDepthTexture, FRDGTextureRef AmbientOcclusionTexture);
 
