Optimises TSR UpdateHistory pass' write bandwidth when output and history's high frequencies are identical

diff --git a/Engine/Shaders/Private/TemporalSuperResolution/TSRUpdateHistory.usf b/Engine/Shaders/Private/TemporalSuperResolution/TSRUpdateHistory.usf
--- a/Engine/Shaders/Private/TemporalSuperResolution/TSRUpdateHistory.usf
+++ b/Engine/Shaders/Private/TemporalSuperResolution/TSRUpdateHistory.usf
@@ -94,6 +94,12 @@
 	#define CONFIG_INTERPOLATE_VELOCITY 0
 #endif
 
+#if CONFIG_GRAND_REPROJECTION == 2 || CONFIG_ACCUMULATE_TRANSLUCENCY_SEPARATLY || CONFIG_PREMULTIPLY_HISTORY
+	#define CONFIG_OUTPUT_HIGH_FREQUENCY 1
+#else
+	#define CONFIG_OUTPUT_HIGH_FREQUENCY 0
+#endif
+
 
 //------------------------------------------------------- CONSTANTS
 
@@ -1568,8 +1574,9 @@ void MainCS(
 			uint2 Random = Rand3DPCG16(int3(dpv_lo(LocalHistoryPixelPos), View.StateFrameIndexMod8)).xy;
 			tsr_half E = tsr_half(Hammersley16(0, 1, Random).x);
 			
+			FinalOutputColor = QuantizeForFloatRenderTarget(FinalOutputColor, E, HistoryQuantizationError);
 			PremultipliedHighFrequencyColor = QuantizeForFloatRenderTarget(PremultipliedHighFrequencyColor, E, HistoryQuantizationError);
-			
+
 			#if CONFIG_GRAND_REPROJECTION
 				FinalPrevHighFrequencyColor = QuantizeForFloatRenderTarget(FinalPrevHighFrequencyColor, E, HistoryQuantizationError);
 				#if CONFIG_GRAND_REPROJECTION == 2
@@ -1595,8 +1602,10 @@ void MainCS(
 		{
 			// Output final history lo pixel.
 			{
-				HistoryOutput_ColorArray[tsr_short3(dpv_lo(LocalHistoryPixelPos), HistoryOutput_ArrayIndices_HighFrequency)] = dpv_lo(PremultipliedHighFrequencyColor);
-				HistoryOutput_Metadata[dpv_lo(LocalHistoryPixelPos)]      = dpv_lo(FinalMetadata);
+				#if CONFIG_OUTPUT_HIGH_FREQUENCY
+					HistoryOutput_ColorArray[tsr_short3(dpv_lo(LocalHistoryPixelPos), HistoryOutput_ArrayIndices_HighFrequency)] = dpv_lo(PremultipliedHighFrequencyColor);
+				#endif
+				HistoryOutput_Metadata[dpv_lo(LocalHistoryPixelPos)]  = dpv_lo(FinalMetadata);
 				#if CONFIG_ACCUMULATE_TRANSLUCENCY_SEPARATLY
 					#if CONFIG_SCENE_COLOR_APLHA
 						HistoryOutput_ColorArray[tsr_short3(dpv_lo(LocalHistoryPixelPos), HistoryOutput_ArrayIndices_Translucency)] = dpv_lo(FinalTranslucencyColor);
@@ -1617,8 +1626,10 @@ void MainCS(
 			// Output final history hi pixel.
 			#if CONFIG_ENABLE_DUAL_PIXEL_VECTORIZATION
 			{
-				HistoryOutput_ColorArray[tsr_short3(dpv_hi(LocalHistoryPixelPos), HistoryOutput_ArrayIndices_HighFrequency)] = dpv_hi(PremultipliedHighFrequencyColor);
-				HistoryOutput_Metadata[dpv_hi(LocalHistoryPixelPos)]      = dpv_hi(FinalMetadata);
+				#if CONFIG_OUTPUT_HIGH_FREQUENCY
+					HistoryOutput_ColorArray[tsr_short3(dpv_hi(LocalHistoryPixelPos), HistoryOutput_ArrayIndices_HighFrequency)] = dpv_hi(PremultipliedHighFrequencyColor);
+				#endif
+				HistoryOutput_Metadata[dpv_hi(LocalHistoryPixelPos)] = dpv_hi(FinalMetadata);
 				#if CONFIG_ACCUMULATE_TRANSLUCENCY_SEPARATLY
 					#if CONFIG_SCENE_COLOR_APLHA
 						HistoryOutput_ColorArray[tsr_short3(dpv_hi(LocalHistoryPixelPos), HistoryOutput_ArrayIndices_Translucency)] = dpv_hi(FinalTranslucencyColor);
diff --git a/Engine/Source/Runtime/Renderer/Private/PostProcess/TemporalSuperResolution.cpp b/Engine/Source/Runtime/Renderer/Private/PostProcess/TemporalSuperResolution.cpp
--- a/Engine/Source/Runtime/Renderer/Private/PostProcess/TemporalSuperResolution.cpp
+++ b/Engine/Source/Runtime/Renderer/Private/PostProcess/TemporalSuperResolution.cpp
@@ -152,6 +152,7 @@ END_SHADER_PARAMETER_STRUCT()
 
 BEGIN_SHADER_PARAMETER_STRUCT(FTSRHistoryTextures, )
 	SHADER_PARAMETER_STRUCT(FTSRHistoryArrayIndices, ArrayIndices)
+	SHADER_PARAMETER_RDG_TEXTURE(Texture2D, Output)
 	SHADER_PARAMETER_RDG_TEXTURE(Texture2DArray, ColorArray)
 	SHADER_PARAMETER_RDG_TEXTURE(Texture2D, Metadata)
 	SHADER_PARAMETER_RDG_TEXTURE(Texture2D, SubpixelDetails)
@@ -201,16 +202,26 @@ enum class ETSRHistoryFormatBits : uint32
 };
 ENUM_CLASS_FLAGS(ETSRHistoryFormatBits);
 
+bool IsOutputDifferentThanHighFrequency(ETSRHistoryFormatBits HistoryFormatBits)
+{
+	return EnumHasAnyFlags(HistoryFormatBits, ETSRHistoryFormatBits::Translucency);
+}
+
 FTSRHistoryArrayIndices TranslateHistoryFormatBitsToArrayIndices(ETSRHistoryFormatBits HistoryFormatBits)
 {
 	FTSRHistoryArrayIndices ArrayIndices;
-	ArrayIndices.Size = 1;
-	ArrayIndices.HighFrequency = 0;
+	ArrayIndices.Size = 0;
+	ArrayIndices.HighFrequency = -1;
 	ArrayIndices.Translucency = -1;
 	ArrayIndices.PrevHighFrequency = -1;
 	ArrayIndices.PrevHighFrequencyResultant = -1;
 	ArrayIndices.HighFrequencyOverblur = -1;
 
+	if (IsOutputDifferentThanHighFrequency(HistoryFormatBits))
+	{
+		ArrayIndices.HighFrequency = ArrayIndices.Size++;
+	}
+
 	if (EnumHasAnyFlags(HistoryFormatBits, ETSRHistoryFormatBits::Translucency))
 	{
 		ArrayIndices.Translucency = ArrayIndices.Size++;
@@ -229,7 +240,14 @@ FTSRHistoryArrayIndices TranslateHistoryFormatBitsToArrayIndices(ETSRHistoryForm
 FTSRHistorySRVs CreateSRVs(FRDGBuilder& GraphBuilder, const FTSRHistoryTextures& Textures)
 {
 	FTSRHistorySRVs SRVs;
-	SRVs.HighFrequency = GraphBuilder.CreateSRV(FRDGTextureSRVDesc::CreateForSlice(Textures.ColorArray, Textures.ArrayIndices.HighFrequency));
+	if (Textures.ArrayIndices.HighFrequency >= 0)
+	{
+		SRVs.HighFrequency = GraphBuilder.CreateSRV(FRDGTextureSRVDesc::CreateForSlice(Textures.ColorArray, Textures.ArrayIndices.HighFrequency));
+	}
+	else
+	{
+		SRVs.HighFrequency = GraphBuilder.CreateSRV(Textures.Output);
+	}
 	SRVs.Metadata = GraphBuilder.CreateSRV(FRDGTextureSRVDesc::Create(Textures.Metadata));
 	SRVs.SubpixelDetails = Textures.SubpixelDetails;
 	if (Textures.ArrayIndices.Translucency >= 0)
@@ -253,7 +271,10 @@ FTSRHistoryUAVs CreateUAVs(FRDGBuilder& GraphBuilder, const FTSRHistoryTextures&
 {
 	FTSRHistoryUAVs UAVs;
 	UAVs.ArrayIndices = Textures.ArrayIndices;
-	UAVs.ColorArray = GraphBuilder.CreateUAV(Textures.ColorArray);
+	if (Textures.ArrayIndices.Size > 0)
+	{
+		UAVs.ColorArray = GraphBuilder.CreateUAV(Textures.ColorArray);
+	}
 	UAVs.Metadata = GraphBuilder.CreateUAV(Textures.Metadata);
 	UAVs.SubpixelDetails = GraphBuilder.CreateUAV(Textures.SubpixelDetails);
 	if (Textures.TranslucencyAlpha)
@@ -818,6 +839,8 @@ ITemporalUpscaler::FOutputs AddTemporalSuperResolutionPasses(
 		}
 	}
 
+	const bool bIsOutputDifferentThanHighFrequency = IsOutputDifferentThanHighFrequency(HistoryFormatBits);
+
 	// Whether to use camera cut shader permutation or not.
 	bool bCameraCut = !InputHistory.IsValid() || View.bCameraCut || ETSRHistoryFormatBits(InputHistory.FormatBit) != HistoryFormatBits;
 
@@ -1130,18 +1153,52 @@ ITemporalUpscaler::FOutputs AddTemporalSuperResolutionPasses(
 	}
 
 	// Create new history.
+	FRDGTextureRef UpdateHistoryOutputTexture;
+	FRDGTextureRef SceneColorOutputTexture;
+	FRDGTextureRef SceneColorOutputHalfResTexture = nullptr;
 	FTSRHistoryTextures History;
 	{
-		History.ArrayIndices = TranslateHistoryFormatBitsToArrayIndices(HistoryFormatBits);
-
-		FRDGTextureDesc ArrayDesc = FRDGTextureDesc::Create2DArray(
+		check(!(PassInputs.bGenerateOutputMip1 && PassInputs.bAllowDownsampleSceneColor));
+		FRDGTextureDesc Desc = FRDGTextureDesc::Create2D(
 			HistoryExtent,
-			(CVarTSRR11G11B10History.GetValueOnRenderThread() != 0 && !bSupportsAlpha) ? PF_FloatR11G11B10 : PF_FloatRGBA,
+			ColorFormat,
 			FClearValueBinding::None,
-			TexCreate_ShaderResource | TexCreate_UAV,
-			History.ArrayIndices.Size);
+			/* InFlags = */ TexCreate_ShaderResource | TexCreate_UAV,
+			/* NumMips = */ PassInputs.bGenerateOutputMip1 ? 2 : 1);
+
+		UpdateHistoryOutputTexture = GraphBuilder.CreateTexture(Desc, TEXT("TSR.Output"));
+		History.Output = UpdateHistoryOutputTexture;
+
+		if (OutputRect.Size() != HistorySize)
+		{
+			Desc.Extent = OutputExtent;
+			SceneColorOutputTexture = GraphBuilder.CreateTexture(Desc, TEXT("TSR.Output"));
+		}
+		else
+		{
+			SceneColorOutputTexture = UpdateHistoryOutputTexture;
+		}
 
-		History.ColorArray = GraphBuilder.CreateTexture(ArrayDesc, TEXT("TSR.History.ColorArray"));
+		if (PassInputs.bAllowDownsampleSceneColor)
+		{
+			Desc.Extent /= 2;
+			SceneColorOutputHalfResTexture = GraphBuilder.CreateTexture(Desc, TEXT("TSR.HalfResOutput"));
+		}
+	}
+	{
+		History.ArrayIndices = TranslateHistoryFormatBitsToArrayIndices(HistoryFormatBits);
+
+		if (History.ArrayIndices.Size > 0)
+		{
+			FRDGTextureDesc ArrayDesc = FRDGTextureDesc::Create2DArray(
+				HistoryExtent,
+				(CVarTSRR11G11B10History.GetValueOnRenderThread() != 0 && !bSupportsAlpha) ? PF_FloatR11G11B10 : PF_FloatRGBA,
+				FClearValueBinding::None,
+				TexCreate_ShaderResource | TexCreate_UAV,
+				History.ArrayIndices.Size);
+
+			History.ColorArray = GraphBuilder.CreateTexture(ArrayDesc, TEXT("TSR.History.ColorArray"));
+		}
 	}
 	{
 		FRDGTextureDesc Desc = FRDGTextureDesc::Create2D(
@@ -1212,7 +1269,8 @@ ITemporalUpscaler::FOutputs AddTemporalSuperResolutionPasses(
 		PrevHistory.ArrayIndices = History.ArrayIndices;
 
 		// Register filterable history
-		PrevHistory.ColorArray = GraphBuilder.RegisterExternalTexture(InputHistory.ColorArray, TEXT("TSR.PrevHistory.ColorArray"));
+		PrevHistory.Output = InputHistory.Output.IsValid() ? GraphBuilder.RegisterExternalTexture(InputHistory.Output, TEXT("TSR.PrevHistory.Output")) : nullptr;
+		PrevHistory.ColorArray = InputHistory.ColorArray.IsValid() ? GraphBuilder.RegisterExternalTexture(InputHistory.ColorArray, TEXT("TSR.PrevHistory.ColorArray")) : nullptr;
 		PrevHistory.Metadata = GraphBuilder.RegisterExternalTexture(InputHistory.Metadata, TEXT("TSR.PrevHistory.Metadata"));
 		PrevHistory.TranslucencyAlpha = InputHistory.TranslucencyAlpha.IsValid() ? GraphBuilder.RegisterExternalTexture(InputHistory.TranslucencyAlpha, TEXT("TSR.PrevHistory.TranslucencyAlpha")) : DummyHistorySRVs.TranslucencyAlpha->Desc.Texture;
 		PrevHistory.Guide = GraphBuilder.RegisterExternalTexture(InputHistory.Guide, TEXT("TSR.PrevHistory.Guide"));
@@ -1223,8 +1281,18 @@ ITemporalUpscaler::FOutputs AddTemporalSuperResolutionPasses(
 
 		PrevHistorySRVs = CreateSRVs(GraphBuilder, PrevHistory);
 
-		PrevVelocity = InputHistory.Velocity.IsValid() ? GraphBuilder.RegisterExternalTexture(InputHistory.Velocity) : PrevVelocity;
-		GrandPrevColorTexture = InputHistory.PrevColorArray.IsValid() ?  GraphBuilder.CreateSRV(FRDGTextureSRVDesc::CreateForSlice(GraphBuilder.RegisterExternalTexture(InputHistory.PrevColorArray, TEXT("TSR.GrandPrevHistory.ColorArray")), PrevHistory.ArrayIndices.HighFrequency)) : GrandPrevColorTexture;
+		if (bGrandReprojection)
+		{
+			PrevVelocity = GraphBuilder.RegisterExternalTexture(InputHistory.Velocity);
+			if (InputHistory.PrevColorArray.IsValid() && bIsOutputDifferentThanHighFrequency)
+			{
+				GrandPrevColorTexture = GraphBuilder.CreateSRV(FRDGTextureSRVDesc::CreateForSlice(GraphBuilder.RegisterExternalTexture(InputHistory.PrevColorArray, TEXT("TSR.GrandPrevHistory.ColorArray")), PrevHistory.ArrayIndices.HighFrequency));
+			}
+			else if (InputHistory.PrevOutput.IsValid() && !bIsOutputDifferentThanHighFrequency)
+			{
+				GrandPrevColorTexture = GraphBuilder.CreateSRV(GraphBuilder.RegisterExternalTexture(InputHistory.PrevOutput, TEXT("TSR.GrandPrevHistory.Output")));
+			}
+		}
 	}
 	else
 	{
@@ -1625,27 +1693,6 @@ ITemporalUpscaler::FOutputs AddTemporalSuperResolutionPasses(
 			FComputeShaderUtils::GetGroupCount(RejectionRect.Size(), 8));
 	}
 
-	// Allocate output
-	FRDGTextureRef SceneColorOutputTexture;
-	FRDGTextureRef SceneColorOutputHalfResTexture = nullptr;
-	{
-		check(!(PassInputs.bGenerateOutputMip1 && PassInputs.bAllowDownsampleSceneColor));
-		FRDGTextureDesc Desc = FRDGTextureDesc::Create2D(
-			OutputExtent,
-			ColorFormat,
-			FClearValueBinding::None,
-			/* InFlags = */ TexCreate_ShaderResource | TexCreate_UAV,
-			/* NumMips = */ PassInputs.bGenerateOutputMip1 ? 2 : 1);
-
-		SceneColorOutputTexture = GraphBuilder.CreateTexture(Desc, TEXT("TSR.Output"));
-
-		if (PassInputs.bAllowDownsampleSceneColor)
-		{
-			Desc.Extent /= 2;
-			SceneColorOutputHalfResTexture = GraphBuilder.CreateTexture(Desc, TEXT("TSR.HalfResOutput"));
-		}
-	}
-
 	// Update temporal history.
 	{
 		static const TCHAR* const kUpdateQualityNames[] = {
@@ -1682,7 +1729,7 @@ ITemporalUpscaler::FOutputs AddTemporalSuperResolutionPasses(
 		PassParameters->HistoryPixelPosToScreenPos = HistoryPixelPosToViewportUV * FScreenTransform::ViewportUVToScreenPos;
 		PassParameters->HistoryPixelPosToInputPPCo = HistoryPixelPosToViewportUV * CommonParameters.InputInfo.ViewportSize + CommonParameters.InputJitter + CommonParameters.InputPixelPosMin;
 		PassParameters->HistoryPixelPosToTranslucencyPPCo = HistoryPixelPosToViewportUV * PassParameters->TranslucencyInfo.ViewportSize + CommonParameters.InputJitter * PassParameters->TranslucencyInfo.ViewportSize / CommonParameters.InputInfo.ViewportSize + SeparateTranslucencyRect.Min;
-		PassParameters->HistoryQuantizationError = ComputePixelFormatQuantizationError(History.ColorArray->Desc.Format);
+		PassParameters->HistoryQuantizationError = ComputePixelFormatQuantizationError((History.ColorArray ? History.ColorArray : History.Output)->Desc.Format);
 		PassParameters->MinTranslucencyRejection = TranslucencyRejectionTexture == nullptr ? 1.0 : 0.0;
 		PassParameters->InvWeightClampingPixelSpeed = 1.0f / CVarTSRWeightClampingPixelSpeed.GetValueOnRenderThread();
 		PassParameters->InputToHistoryFactor = float(HistorySize.X) / float(InputRect.Width());
@@ -1697,13 +1744,13 @@ ITemporalUpscaler::FOutputs AddTemporalSuperResolutionPasses(
 		PassParameters->GrandPrevColorTexture = GrandPrevColorTexture;
 
 		PassParameters->HistoryOutput = CreateUAVs(GraphBuilder, History);
-		if (HistorySize != OutputRect.Size())
+		if (HistorySize != OutputRect.Size() && bIsOutputDifferentThanHighFrequency)
 		{
 			PassParameters->SceneColorOutputMip0 = CreateDummyUAV(GraphBuilder, PF_FloatR11G11B10);
 		}
 		else
 		{
-			PassParameters->SceneColorOutputMip0 = GraphBuilder.CreateUAV(FRDGTextureUAVDesc(SceneColorOutputTexture, /* InMipLevel = */ 0));
+			PassParameters->SceneColorOutputMip0 = GraphBuilder.CreateUAV(FRDGTextureUAVDesc(UpdateHistoryOutputTexture, /* InMipLevel = */ 0));
 		}
 		
 		if (!PassParameters->bGenerateOutputMip1)
@@ -1735,7 +1782,7 @@ ITemporalUpscaler::FOutputs AddTemporalSuperResolutionPasses(
 			RDG_EVENT_NAME("TSR UpdateHistory(Quality=%s%s%s%s) %dx%d",
 				kUpdateQualityNames[int32(PermutationVector.Get<FTSRUpdateHistoryCS::FQualityDim>())],
 				PermutationVector.Get<FTSRUpdateHistoryCS::FSeparateTranslucencyDim>() ? TEXT(" SeparateTranslucency") : TEXT(""),
-				History.ColorArray->Desc.Format == PF_FloatR11G11B10 ? TEXT(" R11G11B10") : TEXT(""),
+				(History.ColorArray ? History.ColorArray : History.Output)->Desc.Format == PF_FloatR11G11B10 ? TEXT(" R11G11B10") : TEXT(""),
 				PassParameters->bGenerateOutputMip1 ? TEXT(" OutputMip1") : TEXT(""),
 				HistorySize.X, HistorySize.Y),
 			AsyncComputePasses >= 3 ? ERDGPassFlags::AsyncCompute : ERDGPassFlags::Compute,
@@ -1824,7 +1871,14 @@ ITemporalUpscaler::FOutputs AddTemporalSuperResolutionPasses(
 		OutputHistory.PrevOutputViewportRect = InputHistory.OutputViewportRect;
 
 		// Extract filterable history
-		GraphBuilder.QueueTextureExtraction(History.ColorArray, &OutputHistory.ColorArray);
+		if (!bIsOutputDifferentThanHighFrequency)
+		{
+			GraphBuilder.QueueTextureExtraction(History.Output, &OutputHistory.Output);
+		}
+		if (History.ColorArray)
+		{
+			GraphBuilder.QueueTextureExtraction(History.ColorArray, &OutputHistory.ColorArray);
+		}
 		GraphBuilder.QueueTextureExtraction(History.Metadata, &OutputHistory.Metadata);
 		if (bAccumulateTranslucencySeparately)
 		{
@@ -1845,13 +1899,27 @@ ITemporalUpscaler::FOutputs AddTemporalSuperResolutionPasses(
 		if (bGrandReprojection)
 		{
 			GraphBuilder.QueueTextureExtraction(DilatedVelocityTexture, &OutputHistory.Velocity);
-			if (PrevHistory.ColorArray)
+			if (!bIsOutputDifferentThanHighFrequency)
 			{
-				GraphBuilder.QueueTextureExtraction(PrevHistory.ColorArray, &OutputHistory.PrevColorArray);
+				if (PrevHistory.Output)
+				{
+					GraphBuilder.QueueTextureExtraction(PrevHistory.Output, &OutputHistory.PrevOutput);
+				}
+				else
+				{
+					GraphBuilder.QueueTextureExtraction(History.Output, &OutputHistory.PrevOutput);
+				}
 			}
 			else
 			{
-				GraphBuilder.QueueTextureExtraction(History.ColorArray, &OutputHistory.PrevColorArray);
+				if (PrevHistory.ColorArray)
+				{
+					GraphBuilder.QueueTextureExtraction(PrevHistory.ColorArray, &OutputHistory.PrevColorArray);
+				}
+				else
+				{
+					GraphBuilder.QueueTextureExtraction(History.ColorArray, &OutputHistory.PrevColorArray);
+				}
 			}
 		}
 
diff --git a/Engine/Source/Runtime/Renderer/Private/SceneRendering.h b/Engine/Source/Runtime/Renderer/Private/SceneRendering.h
--- a/Engine/Source/Runtime/Renderer/Private/SceneRendering.h
+++ b/Engine/Source/Runtime/Renderer/Private/SceneRendering.h
@@ -896,6 +896,7 @@ struct FTemporalAAHistory
 struct FTSRHistory
 {
 	// Output resolution.
+	TRefCountPtr<IPooledRenderTarget> Output;
 	TRefCountPtr<IPooledRenderTarget> ColorArray;
 	TRefCountPtr<IPooledRenderTarget> Metadata;
 	TRefCountPtr<IPooledRenderTarget> TranslucencyAlpha;
@@ -907,6 +908,7 @@ struct FTSRHistory
 	TRefCountPtr<IPooledRenderTarget> Velocity;
 
 	// Previous frame's history.
+	TRefCountPtr<IPooledRenderTarget> PrevOutput;
 	TRefCountPtr<IPooledRenderTarget> PrevColorArray;
 
 	// Frame's input and output resolution.
@@ -929,7 +931,7 @@ struct FTSRHistory
 
 	bool IsValid() const
 	{
-		return ColorArray.IsValid();
+		return Metadata.IsValid();
 	}
 
 	uint64 GetGPUSizeBytes(bool bLogSizes) const;
