Add Strata debug mode into the level view menu.

This gives direct access to all Strata view mode, which helps with discoverability.

All strata view mode are also available with CVar: r.Strata.ViewMode

diff --git a/Engine/Source/Editor/CommonMenuExtensions/Private/CommonMenuExtensionsModule.cpp b/Engine/Source/Editor/CommonMenuExtensions/Private/CommonMenuExtensionsModule.cpp
--- a/Engine/Source/Editor/CommonMenuExtensions/Private/CommonMenuExtensionsModule.cpp
+++ b/Engine/Source/Editor/CommonMenuExtensions/Private/CommonMenuExtensionsModule.cpp
@@ -4,6 +4,7 @@
 #include "BufferVisualizationMenuCommands.h"
 #include "NaniteVisualizationMenuCommands.h"
 #include "LumenVisualizationMenuCommands.h"
+#include "StrataVisualizationMenuCommands.h"
 #include "VirtualShadowMapVisualizationMenuCommands.h"
 #include "ShowFlagMenuCommands.h"
 #include "Modules/ModuleManager.h"
@@ -15,6 +16,7 @@ void FCommonMenuExtensionsModule::StartupModule()
 	FBufferVisualizationMenuCommands::Register();
 	FNaniteVisualizationMenuCommands::Register();
 	FLumenVisualizationMenuCommands::Register();
+	FStrataVisualizationMenuCommands::Register();
 	FVirtualShadowMapVisualizationMenuCommands::Register();
 	FShowFlagMenuCommands::Register();
 }
@@ -24,6 +26,7 @@ void FCommonMenuExtensionsModule::ShutdownModule()
 	FShowFlagMenuCommands::Unregister();
 	FVirtualShadowMapVisualizationMenuCommands::Unregister();
 	FNaniteVisualizationMenuCommands::Unregister();
+	FStrataVisualizationMenuCommands::Unregister();
 	FLumenVisualizationMenuCommands::Unregister();
 	FBufferVisualizationMenuCommands::Unregister();
 }
\ No newline at end of file
diff --git a/Engine/Source/Editor/CommonMenuExtensions/Private/StrataVisualizationMenuCommands.cpp b/Engine/Source/Editor/CommonMenuExtensions/Private/StrataVisualizationMenuCommands.cpp
new file mode 100644
--- /dev/null
+++ b/Engine/Source/Editor/CommonMenuExtensions/Private/StrataVisualizationMenuCommands.cpp
@@ -0,0 +1,141 @@
+// Copyright Epic Games, Inc. All Rights Reserved.
+
+#include "StrataVisualizationMenuCommands.h"
+
+#include "Delegates/Delegate.h"
+#include "EditorViewportClient.h"
+#include "Framework/Commands/InputChord.h"
+#include "Framework/Commands/UIAction.h"
+#include "Framework/Commands/UICommandInfo.h"
+#include "Framework/Commands/UICommandList.h"
+#include "Framework/MultiBox/MultiBoxBuilder.h"
+#include "Internationalization/Internationalization.h"
+#include "Internationalization/Text.h"
+#include "StrataVisualizationData.h"
+#include "Misc/AssertionMacros.h"
+#include "Styling/AppStyle.h"
+#include "UObject/UnrealNames.h"
+
+#define LOCTEXT_NAMESPACE "StrataVisualizationMenuCommands"
+
+FStrataVisualizationMenuCommands::FStrataVisualizationMenuCommands()
+	: TCommands<FStrataVisualizationMenuCommands>
+	(
+		TEXT("StrataVisualizationMenu"), // Context name for fast lookup
+		NSLOCTEXT("Contexts", "StrataVisualizationMenu", "Strata"), // Localized context name for displaying
+		NAME_None, // Parent context name.  
+		FAppStyle::GetAppStyleSetName() // Icon Style Set
+	),
+	CommandMap()
+{
+}
+
+void FStrataVisualizationMenuCommands::BuildCommandMap()
+{
+	const FStrataVisualizationData& VisualizationData = GetStrataVisualizationData();
+	const FStrataVisualizationData::TModeMap& ModeMap = VisualizationData.GetModeMap();
+
+	CommandMap.Empty();
+	for (FStrataVisualizationData::TModeMap::TConstIterator It = ModeMap.CreateConstIterator(); It; ++It)
+	{
+		const FStrataVisualizationData::FModeRecord& Entry = It.Value();
+		FStrataVisualizationRecord& Record = CommandMap.Add(Entry.ModeName, FStrataVisualizationRecord());
+		Record.Name = Entry.ModeName;
+		Record.Command = FUICommandInfoDecl(
+			this->AsShared(),
+			Entry.ModeName,
+			Entry.ModeText,
+			Entry.ModeDesc)
+			.UserInterfaceType(EUserInterfaceActionType::RadioButton)
+			.DefaultChord(FInputChord()
+		);
+
+		switch (Entry.Mode)
+		{
+		default:
+		case FStrataVisualizationData::FViewMode::MaterialProperties:
+			Record.Type = FStrataVisualizationType::MaterialProperties;
+			break;
+
+		case FStrataVisualizationData::FViewMode::MaterialCount:
+			Record.Type = FStrataVisualizationType::MaterialCount;
+			break;
+
+		case FStrataVisualizationData::FViewMode::AdvancedMaterialProperties:
+			Record.Type = FStrataVisualizationType::AdvancedMaterialProperties;
+			break;
+		}
+	}
+}
+
+void FStrataVisualizationMenuCommands::BuildVisualisationSubMenu(FMenuBuilder& Menu)
+{
+	const FStrataVisualizationMenuCommands& Commands = FStrataVisualizationMenuCommands::Get();
+	if (Commands.IsPopulated())
+	{
+		Menu.BeginSection("LevelViewportStrataVisualizationMode", LOCTEXT("StrataVisualizationHeader", "Strata Visualization Mode"));
+
+		Commands.AddCommandTypeToMenu(Menu, FStrataVisualizationType::MaterialProperties);
+		Commands.AddCommandTypeToMenu(Menu, FStrataVisualizationType::MaterialCount);
+		Commands.AddCommandTypeToMenu(Menu, FStrataVisualizationType::AdvancedMaterialProperties);
+
+		Menu.EndSection();
+	}
+}
+
+bool FStrataVisualizationMenuCommands::AddCommandTypeToMenu(FMenuBuilder& Menu, const FStrataVisualizationType Type) const
+{
+	bool bAddedCommands = false;
+
+	const TStrataVisualizationModeCommandMap& Commands = CommandMap;
+	for (TCommandConstIterator It = CreateCommandConstIterator(); It; ++It)
+	{
+		const FStrataVisualizationRecord& Record = It.Value();
+		if (Record.Type == Type)
+		{
+			Menu.AddMenuEntry(Record.Command, NAME_None, Record.Command->GetLabel());
+			bAddedCommands = true;
+		}
+	}
+
+	return bAddedCommands;
+}
+
+FStrataVisualizationMenuCommands::TCommandConstIterator FStrataVisualizationMenuCommands::CreateCommandConstIterator() const
+{
+	return CommandMap.CreateConstIterator();
+}
+
+void FStrataVisualizationMenuCommands::RegisterCommands()
+{
+	BuildCommandMap();
+}
+
+void FStrataVisualizationMenuCommands::BindCommands(FUICommandList& CommandList, const TSharedPtr<FEditorViewportClient>& Client) const
+{
+	// Map Strata visualization mode actions
+	for (FStrataVisualizationMenuCommands::TCommandConstIterator It = FStrataVisualizationMenuCommands::Get().CreateCommandConstIterator(); It; ++It)
+	{
+		const FStrataVisualizationMenuCommands::FStrataVisualizationRecord& Record = It.Value();
+		CommandList.MapAction(
+			Record.Command,
+			FExecuteAction::CreateStatic<const TSharedPtr<FEditorViewportClient>&>(&FStrataVisualizationMenuCommands::ChangeStrataVisualizationMode, Client, Record.Name),
+			FCanExecuteAction(),
+			FIsActionChecked::CreateStatic<const TSharedPtr<FEditorViewportClient>&>(&FStrataVisualizationMenuCommands::IsStrataVisualizationModeSelected, Client, Record.Name)
+		);
+	}
+}
+
+void FStrataVisualizationMenuCommands::ChangeStrataVisualizationMode(const TSharedPtr<FEditorViewportClient>& Client, FName InName)
+{
+	check(Client.IsValid());
+	Client->ChangeStrataVisualizationMode(InName);
+}
+
+bool FStrataVisualizationMenuCommands::IsStrataVisualizationModeSelected(const TSharedPtr<FEditorViewportClient>& Client, FName InName)
+{
+	check(Client.IsValid());
+	return Client->IsStrataVisualizationModeSelected(InName);
+}
+
+#undef LOCTEXT_NAMESPACE
diff --git a/Engine/Source/Editor/CommonMenuExtensions/Public/StrataVisualizationMenuCommands.h b/Engine/Source/Editor/CommonMenuExtensions/Public/StrataVisualizationMenuCommands.h
new file mode 100644
--- /dev/null
+++ b/Engine/Source/Editor/CommonMenuExtensions/Public/StrataVisualizationMenuCommands.h
@@ -0,0 +1,62 @@
+// Copyright Epic Games, Inc. All Rights Reserved.
+
+#pragma once
+
+#include "Containers/Map.h"
+#include "Framework/Commands/Commands.h"
+#include "Framework/Commands/UICommandInfo.h"
+#include "Framework/MultiBox/MultiBoxBuilder.h"
+#include "HAL/Platform.h"
+#include "Templates/SharedPointer.h"
+#include "UObject/NameTypes.h"
+
+class FEditorViewportClient;
+class FMenuBuilder;
+class FUICommandInfo;
+class FUICommandList;
+
+class COMMONMENUEXTENSIONS_API FStrataVisualizationMenuCommands : public TCommands<FStrataVisualizationMenuCommands>
+{
+public:
+	enum class FStrataVisualizationType : uint8
+	{
+		MaterialProperties,
+		MaterialCount,
+		AdvancedMaterialProperties,
+	};
+
+	struct FStrataVisualizationRecord
+	{
+		FName Name;
+		TSharedPtr<FUICommandInfo> Command;
+		FStrataVisualizationType Type = FStrataVisualizationType::MaterialProperties;
+	};
+
+	typedef TMultiMap<FName, FStrataVisualizationRecord> TStrataVisualizationModeCommandMap;
+	typedef TStrataVisualizationModeCommandMap::TConstIterator TCommandConstIterator;
+
+	FStrataVisualizationMenuCommands();
+
+	TCommandConstIterator CreateCommandConstIterator() const;
+
+	static void BuildVisualisationSubMenu(FMenuBuilder& Menu);
+
+	virtual void RegisterCommands() override;
+
+	void BindCommands(FUICommandList& CommandList, const TSharedPtr<FEditorViewportClient>& Client) const;
+	
+	inline bool IsPopulated() const
+	{
+		return CommandMap.Num() > 0;
+	}
+
+private:
+	void BuildCommandMap();
+	bool AddCommandTypeToMenu(FMenuBuilder& Menu, const FStrataVisualizationType Type) const;
+
+	static void ChangeStrataVisualizationMode(const TSharedPtr<FEditorViewportClient>& Client, FName InName);
+	static bool IsStrataVisualizationModeSelected(const TSharedPtr<FEditorViewportClient>& Client, FName InName);
+
+private:
+	TStrataVisualizationModeCommandMap CommandMap;
+};
\ No newline at end of file
diff --git a/Engine/Source/Editor/DetailCustomizations/Private/DebugCameraControllerSettingsCustomization.cpp b/Engine/Source/Editor/DetailCustomizations/Private/DebugCameraControllerSettingsCustomization.cpp
--- a/Engine/Source/Editor/DetailCustomizations/Private/DebugCameraControllerSettingsCustomization.cpp
+++ b/Engine/Source/Editor/DetailCustomizations/Private/DebugCameraControllerSettingsCustomization.cpp
@@ -45,6 +45,7 @@ void FDebugCameraControllerSettingsViewModeIndexCustomization::CustomizeHeader(T
 	EnumRestriction->AddHiddenValue(ViewModeIndexEnum->GetNameStringByValue((uint8)EViewModeIndex::VMI_VisualizeBuffer));
 	EnumRestriction->AddHiddenValue(ViewModeIndexEnum->GetNameStringByValue((uint8)EViewModeIndex::VMI_VisualizeNanite));
 	EnumRestriction->AddHiddenValue(ViewModeIndexEnum->GetNameStringByValue((uint8)EViewModeIndex::VMI_VisualizeLumen));
+	EnumRestriction->AddHiddenValue(ViewModeIndexEnum->GetNameStringByValue((uint8)EViewModeIndex::VMI_VisualizeStrata));
 	EnumRestriction->AddHiddenValue(ViewModeIndexEnum->GetNameStringByValue((uint8)EViewModeIndex::VMI_VisualizeVirtualShadowMap));
 	EnumRestriction->AddHiddenValue(ViewModeIndexEnum->GetNameStringByValue((uint8)EViewModeIndex::VMI_VisualizeGPUSkinCache));
 	EnumRestriction->AddHiddenValue(ViewModeIndexEnum->GetNameStringByValue((uint8)EViewModeIndex::VMI_StationaryLightOverlap));
diff --git a/Engine/Source/Editor/EditorStyle/Private/StarshipStyle.cpp b/Engine/Source/Editor/EditorStyle/Private/StarshipStyle.cpp
--- a/Engine/Source/Editor/EditorStyle/Private/StarshipStyle.cpp
+++ b/Engine/Source/Editor/EditorStyle/Private/StarshipStyle.cpp
@@ -4067,6 +4067,7 @@ void FStarshipEditorStyle::FStyle::SetupLevelEditorStyle()
 		Set( "EditorViewport.VisualizeBufferMode",    new IMAGE_BRUSH_SVG("Starship/Common/BufferVisualization", Icon16x16 ) );
 		Set( "EditorViewport.VisualizeNaniteMode",    new IMAGE_BRUSH_SVG("Starship/Common/BufferVisualization", Icon16x16 ) );
 		Set( "EditorViewport.VisualizeLumenMode",    new IMAGE_BRUSH_SVG("Starship/Common/BufferVisualization", Icon16x16 ) );
+		Set( "EditorViewport.VisualizeStrataMode",    new IMAGE_BRUSH_SVG("Starship/Common/BufferVisualization", Icon16x16 ) );
 		Set( "EditorViewport.VisualizeVirtualShadowMapMode", new IMAGE_BRUSH_SVG("Starship/Common/BufferVisualization", Icon16x16 ) );
 		Set( "EditorViewport.CollisionPawn",          new IMAGE_BRUSH_SVG("Starship/Common/PlayerCollision", Icon16x16 ) );
 		Set( "EditorViewport.CollisionVisibility",    new IMAGE_BRUSH_SVG("Starship/Common/VisibilityCollision", Icon16x16 ) );
diff --git a/Engine/Source/Editor/LevelEditor/Private/LevelViewportActions.cpp b/Engine/Source/Editor/LevelEditor/Private/LevelViewportActions.cpp
--- a/Engine/Source/Editor/LevelEditor/Private/LevelViewportActions.cpp
+++ b/Engine/Source/Editor/LevelEditor/Private/LevelViewportActions.cpp
@@ -11,6 +11,7 @@
 #include "BufferVisualizationData.h"
 #include "NaniteVisualizationData.h"
 #include "LumenVisualizationData.h"
+#include "StrataVisualizationData.h"
 #include "VirtualShadowMapVisualizationData.h"
 #include "Bookmarks/BookmarkUI.h"
 #include "UnrealEdGlobals.h"
diff --git a/Engine/Source/Editor/LevelEditor/Private/SLevelViewport.cpp b/Engine/Source/Editor/LevelEditor/Private/SLevelViewport.cpp
--- a/Engine/Source/Editor/LevelEditor/Private/SLevelViewport.cpp
+++ b/Engine/Source/Editor/LevelEditor/Private/SLevelViewport.cpp
@@ -62,6 +62,7 @@
 #include "BufferVisualizationData.h"
 #include "NaniteVisualizationData.h"
 #include "LumenVisualizationData.h"
+#include "StrataVisualizationData.h"
 #include "VirtualShadowMapVisualizationData.h"
 #include "Framework/Notifications/NotificationManager.h"
 #include "Widgets/Notifications/SNotificationList.h"
@@ -76,6 +77,7 @@
 #include "BufferVisualizationMenuCommands.h"
 #include "NaniteVisualizationMenuCommands.h"
 #include "LumenVisualizationMenuCommands.h"
+#include "StrataVisualizationMenuCommands.h"
 #include "VirtualShadowMapVisualizationMenuCommands.h"
 #include "EditorLevelUtils.h"
 #include "Engine/LevelStreaming.h"
@@ -101,6 +103,13 @@
 static const FName LevelEditorName("LevelEditor");
 static FAutoConsoleCommand EnableInViewportMenu(TEXT("Editor.EnableInViewportMenu"), TEXT("Enables the new in-viewport property menu"), FConsoleCommandDelegate::CreateStatic(&SLevelViewport::EnableInViewportMenu));
 bool SLevelViewport::bInViewportMenuEnabled = false;
+
+bool LevelEditor_IsStrataEnabled()
+{
+	static const auto CVar = IConsoleManager::Get().FindTConsoleVariableDataInt(TEXT("r.Strata"));
+	return CVar && CVar->GetValueOnAnyThread() > 0;
+}
+
 #define LOCTEXT_NAMESPACE "LevelViewport"
 
 // @todo Slate Hack: Disallow game UI to be used in play in viewport until GWorld problem is fixed
@@ -490,6 +499,7 @@ void SLevelViewport::ConstructLevelEditorViewportClient(FLevelEditorViewportInst
 	LevelViewportClient->CurrentBufferVisualizationMode = ViewportInstanceSettings.BufferVisualizationMode;
 	LevelViewportClient->CurrentNaniteVisualizationMode = ViewportInstanceSettings.NaniteVisualizationMode;
 	LevelViewportClient->CurrentLumenVisualizationMode = ViewportInstanceSettings.LumenVisualizationMode;
+	LevelViewportClient->CurrentStrataVisualizationMode = ViewportInstanceSettings.StrataVisualizationMode;
 	LevelViewportClient->CurrentVirtualShadowMapVisualizationMode = ViewportInstanceSettings.VirtualShadowMapVisualizationMode;
 	LevelViewportClient->CurrentRayTracingDebugVisualizationMode = ViewportInstanceSettings.RayTracingDebugVisualizationMode;
 	LevelViewportClient->CurrentGPUSkinCacheVisualizationMode = ViewportInstanceSettings.GPUSkinCacheVisualizationMode;
@@ -1539,6 +1549,10 @@ void SLevelViewport::BindViewCommands( FUICommandList& OutCommandList )
 	FBufferVisualizationMenuCommands::Get().BindCommands(OutCommandList, Client);
 	FNaniteVisualizationMenuCommands::Get().BindCommands(OutCommandList, Client);
 	FLumenVisualizationMenuCommands::Get().BindCommands(OutCommandList, Client);
+	if (LevelEditor_IsStrataEnabled())
+	{
+		FStrataVisualizationMenuCommands::Get().BindCommands(OutCommandList, Client);
+	}
 	FVirtualShadowMapVisualizationMenuCommands::Get().BindCommands(OutCommandList, Client);
 }
 
@@ -2238,6 +2252,7 @@ void SLevelViewport::SaveConfig(const FString& ConfigName) const
 		ViewportInstanceSettings.BufferVisualizationMode = LevelViewportClient->CurrentBufferVisualizationMode;
 		ViewportInstanceSettings.NaniteVisualizationMode = LevelViewportClient->CurrentNaniteVisualizationMode;
 		ViewportInstanceSettings.LumenVisualizationMode = LevelViewportClient->CurrentLumenVisualizationMode;
+		ViewportInstanceSettings.StrataVisualizationMode = LevelViewportClient->CurrentStrataVisualizationMode;
 		ViewportInstanceSettings.VirtualShadowMapVisualizationMode = LevelViewportClient->CurrentVirtualShadowMapVisualizationMode;
 		ViewportInstanceSettings.RayTracingDebugVisualizationMode = LevelViewportClient->CurrentRayTracingDebugVisualizationMode;
 		ViewportInstanceSettings.GPUSkinCacheVisualizationMode = LevelViewportClient->CurrentGPUSkinCacheVisualizationMode;
diff --git a/Engine/Source/Editor/LevelEditor/Private/SLevelViewportToolBar.cpp b/Engine/Source/Editor/LevelEditor/Private/SLevelViewportToolBar.cpp
--- a/Engine/Source/Editor/LevelEditor/Private/SLevelViewportToolBar.cpp
+++ b/Engine/Source/Editor/LevelEditor/Private/SLevelViewportToolBar.cpp
@@ -34,6 +34,7 @@
 #include "BufferVisualizationData.h"
 #include "NaniteVisualizationData.h"
 #include "LumenVisualizationData.h"
+#include "StrataVisualizationData.h"
 #include "VirtualShadowMapVisualizationData.h"
 #include "FoliageType.h"
 #include "ShowFlagMenuCommands.h"
@@ -49,6 +50,8 @@
 
 #define LOCTEXT_NAMESPACE "LevelViewportToolBar"
 
+bool LevelEditor_IsStrataEnabled();
+
 /** Override the view menu, just so we can specify the level viewport as active when the button is clicked */
 class SLevelEditorViewportViewMenu : public SEditorViewportViewMenu
 {
@@ -1604,6 +1607,31 @@ void SLevelViewportToolBar::FillViewMenu(UToolMenu* Menu)
 						);
 	}
 
+	if (LevelEditor_IsStrataEnabled())
+	{
+		FToolMenuSection& Section = Menu->FindOrAddSection("ViewMode");
+		Section.AddSubMenu(
+			"VisualizeStrataViewMode",
+			LOCTEXT("VisualizeStrataViewModeDisplayName", "Strata"),
+			LOCTEXT("StrataVisualizationMenu_ToolTip", "Select a mode for Strata visualization"),
+			FNewMenuDelegate::CreateStatic(&FStrataVisualizationMenuCommands::BuildVisualisationSubMenu),
+			FUIAction(
+				FExecuteAction(),
+				FCanExecuteAction(),
+				FIsActionChecked::CreateLambda([this]()
+					{
+						const TSharedRef<SEditorViewport> ViewportRef = Viewport.Pin().ToSharedRef();
+						const TSharedPtr<FEditorViewportClient> ViewportClient = ViewportRef->GetViewportClient();
+						check(ViewportClient.IsValid());
+						return ViewportClient->IsViewModeEnabled(VMI_VisualizeStrata);
+					})
+			),
+			EUserInterfaceActionType::RadioButton,
+						/* bInOpenSubMenuOnClick = */ false,
+						FSlateIcon(FAppStyle::GetAppStyleSetName(), "EditorViewport.VisualizeStrataMode")
+						);
+	}
+
 	{
 		FToolMenuSection& Section = Menu->FindOrAddSection("ViewMode");
 		Section.AddSubMenu(
diff --git a/Engine/Source/Editor/LevelEditor/Public/LevelViewportActions.h b/Engine/Source/Editor/LevelEditor/Public/LevelViewportActions.h
--- a/Engine/Source/Editor/LevelEditor/Public/LevelViewportActions.h
+++ b/Engine/Source/Editor/LevelEditor/Public/LevelViewportActions.h
@@ -10,6 +10,7 @@
 #include "BufferVisualizationMenuCommands.h"
 #include "NaniteVisualizationMenuCommands.h"
 #include "LumenVisualizationMenuCommands.h"
+#include "StrataVisualizationMenuCommands.h"
 #include "VirtualShadowMapVisualizationMenuCommands.h"
 
 /**
diff --git a/Engine/Source/Editor/MaterialEditor/Private/MaterialEditor.cpp b/Engine/Source/Editor/MaterialEditor/Private/MaterialEditor.cpp
--- a/Engine/Source/Editor/MaterialEditor/Private/MaterialEditor.cpp
+++ b/Engine/Source/Editor/MaterialEditor/Private/MaterialEditor.cpp
@@ -175,7 +175,7 @@ static TAutoConsoleVariable<int32> CVarMaterialEdUseDevShaders(
 	TEXT("Toggles whether the material editor will use shaders that include extra overhead incurred by the editor. Material editor must be re-opened if changed at runtime."),
 	ECVF_RenderThreadSafe);
 
-static bool Editor_IsStrataEnabled()
+bool Editor_IsStrataEnabled()
 {
 	static const auto CVar = IConsoleManager::Get().FindTConsoleVariableDataInt(TEXT("r.Strata"));
 	return CVar && CVar->GetValueOnAnyThread() > 0;
diff --git a/Engine/Source/Editor/UnrealEd/Classes/Settings/LevelEditorViewportSettings.h b/Engine/Source/Editor/UnrealEd/Classes/Settings/LevelEditorViewportSettings.h
--- a/Engine/Source/Editor/UnrealEd/Classes/Settings/LevelEditorViewportSettings.h
+++ b/Engine/Source/Editor/UnrealEd/Classes/Settings/LevelEditorViewportSettings.h
@@ -98,6 +98,7 @@ struct UNREALED_API FLevelEditorViewportInstanceSettings
 		, BufferVisualizationMode()
 		, NaniteVisualizationMode()
 		, LumenVisualizationMode()
+		, StrataVisualizationMode()
 		, VirtualShadowMapVisualizationMode()
 		, GPUSkinCacheVisualizationMode()
 		, ExposureSettings()
@@ -148,6 +149,10 @@ struct UNREALED_API FLevelEditorViewportInstanceSettings
 	UPROPERTY(config)
 	FName LumenVisualizationMode;
 
+	/** The Strata visualization mode for the viewport. */
+	UPROPERTY(config)
+	FName StrataVisualizationMode;
+
 	/** The virtual shadow map visualization mode for the viewport. */
 	UPROPERTY(config)
 	FName VirtualShadowMapVisualizationMode;
diff --git a/Engine/Source/Editor/UnrealEd/Private/EditorViewportClient.cpp b/Engine/Source/Editor/UnrealEd/Private/EditorViewportClient.cpp
--- a/Engine/Source/Editor/UnrealEd/Private/EditorViewportClient.cpp
+++ b/Engine/Source/Editor/UnrealEd/Private/EditorViewportClient.cpp
@@ -55,6 +55,7 @@
 #include "BufferVisualizationData.h"
 #include "NaniteVisualizationData.h"
 #include "LumenVisualizationData.h"
+#include "StrataVisualizationData.h"
 #include "VirtualShadowMapVisualizationData.h"
 #include "GPUSkinCacheVisualizationData.h"
 #include "UnrealWidget.h"
@@ -389,6 +390,7 @@ FEditorViewportClient::FEditorViewportClient(FEditorModeTools* InModeTools, FPre
 	, CurrentBufferVisualizationMode(NAME_None)
 	, CurrentNaniteVisualizationMode(NAME_None)
 	, CurrentLumenVisualizationMode(NAME_None)
+	, CurrentStrataVisualizationMode(NAME_None)
 	, CurrentVirtualShadowMapVisualizationMode(NAME_None)
 	, CurrentRayTracingDebugVisualizationMode(NAME_None)
 	, CurrentGPUSkinCacheVisualizationMode(NAME_None)
@@ -2676,6 +2678,23 @@ FText FEditorViewportClient::GetCurrentVirtualShadowMapVisualizationModeDisplayN
 	return GetVirtualShadowMapVisualizationData().GetModeDisplayName(CurrentVirtualShadowMapVisualizationMode);
 }
 
+void FEditorViewportClient::ChangeStrataVisualizationMode(FName InName)
+{
+	SetViewMode(VMI_VisualizeStrata);
+	CurrentStrataVisualizationMode = InName;
+}
+
+bool FEditorViewportClient::IsStrataVisualizationModeSelected(FName InName) const
+{
+	return IsViewModeEnabled(VMI_VisualizeStrata) && CurrentStrataVisualizationMode == InName;
+}
+
+FText FEditorViewportClient::GetCurrentStrataVisualizationModeDisplayName() const
+{
+	checkf(IsViewModeEnabled(VMI_VisualizeStrata), TEXT("In order to call GetCurrentStrataVisualizationMode(), first you must set ViewMode to VMI_VisualizeStrata."));
+	return GetStrataVisualizationData().GetModeDisplayName(CurrentStrataVisualizationMode);
+}
+
 bool FEditorViewportClient::IsVisualizeCalibrationMaterialEnabled() const
 {
 	// Get the list of requested buffers from the console
@@ -3840,6 +3859,7 @@ void FEditorViewportClient::SetupViewForRendering(FSceneViewFamily& ViewFamily,
 	View.CurrentBufferVisualizationMode = CurrentBufferVisualizationMode;
 	View.CurrentNaniteVisualizationMode = CurrentNaniteVisualizationMode;
 	View.CurrentLumenVisualizationMode = CurrentLumenVisualizationMode;
+	View.CurrentStrataVisualizationMode = CurrentStrataVisualizationMode;
 	View.CurrentVirtualShadowMapVisualizationMode = CurrentVirtualShadowMapVisualizationMode;
 	View.CurrentGPUSkinCacheVisualizationMode = CurrentGPUSkinCacheVisualizationMode;
 #if RHI_RAYTRACING
diff --git a/Engine/Source/Editor/UnrealEd/Private/EditorViewportCommands.cpp b/Engine/Source/Editor/UnrealEd/Private/EditorViewportCommands.cpp
--- a/Engine/Source/Editor/UnrealEd/Private/EditorViewportCommands.cpp
+++ b/Engine/Source/Editor/UnrealEd/Private/EditorViewportCommands.cpp
@@ -120,6 +120,7 @@ void FEditorViewportCommands::RegisterCommands()
 	UI_COMMAND( VisualizeBufferMode, "Buffer Visualization View Mode", "Renders a set of selected post process materials, which visualize various intermediate render buffers (material attributes)", EUserInterfaceActionType::RadioButton, FInputChord());
 	UI_COMMAND( VisualizeNaniteMode, "Nanite Visualization View Mode", "Visualizes various rendering aspects of Nanite.", EUserInterfaceActionType::RadioButton, FInputChord());
 	UI_COMMAND( VisualizeLumenMode, "Lumen Visualization View Mode", "Visualizes Lumen debug views.", EUserInterfaceActionType::RadioButton, FInputChord());
+	UI_COMMAND( VisualizeStrataMode, "Strata Visualization View Mode", "Visualizes Strata debug views.", EUserInterfaceActionType::RadioButton, FInputChord());
 	UI_COMMAND( VisualizeVirtualShadowMapMode, "Virtual Shadow Map Visualization View Mode", "Visualizes various rendering aspects of virtual shadow maps.", EUserInterfaceActionType::RadioButton, FInputChord());
 
 	UI_COMMAND( ReflectionOverrideMode, "Reflections View Mode", "Renders the scene with reflections only", EUserInterfaceActionType::RadioButton, FInputChord() );
diff --git a/Engine/Source/Editor/UnrealEd/Private/SEditorViewport.cpp b/Engine/Source/Editor/UnrealEd/Private/SEditorViewport.cpp
--- a/Engine/Source/Editor/UnrealEd/Private/SEditorViewport.cpp
+++ b/Engine/Source/Editor/UnrealEd/Private/SEditorViewport.cpp
@@ -429,6 +429,7 @@ void SEditorViewport::BindCommands()
 	MAP_VIEWMODE_ACTION( Commands.VisualizeBufferMode, VMI_VisualizeBuffer );
 	MAP_VIEWMODE_ACTION( Commands.VisualizeNaniteMode, VMI_VisualizeNanite );
 	MAP_VIEWMODE_ACTION( Commands.VisualizeLumenMode, VMI_VisualizeLumen );
+	MAP_VIEWMODE_ACTION( Commands.VisualizeStrataMode, VMI_VisualizeStrata);
 	MAP_VIEWMODE_ACTION( Commands.VisualizeVirtualShadowMapMode, VMI_VisualizeVirtualShadowMap );
 	MAP_VIEWMODE_ACTION( Commands.CollisionPawn, VMI_CollisionPawn);
 	MAP_VIEWMODE_ACTION( Commands.CollisionVisibility, VMI_CollisionVisibility);
diff --git a/Engine/Source/Editor/UnrealEd/Private/SEditorViewportViewMenu.cpp b/Engine/Source/Editor/UnrealEd/Private/SEditorViewportViewMenu.cpp
--- a/Engine/Source/Editor/UnrealEd/Private/SEditorViewportViewMenu.cpp
+++ b/Engine/Source/Editor/UnrealEd/Private/SEditorViewportViewMenu.cpp
@@ -55,6 +55,10 @@ FText SEditorViewportViewMenu::GetViewMenuLabel() const
 		{
 			Label = ViewportClient->GetCurrentLumenVisualizationModeDisplayName();
 		}
+		else if (ViewMode == VMI_VisualizeStrata)
+		{
+			Label = ViewportClient->GetCurrentStrataVisualizationModeDisplayName();
+		}
 		else if (ViewMode == VMI_VisualizeVirtualShadowMap)
 		{
 			Label = ViewportClient->GetCurrentVirtualShadowMapVisualizationModeDisplayName();
diff --git a/Engine/Source/Editor/UnrealEd/Public/EditorViewportClient.h b/Engine/Source/Editor/UnrealEd/Public/EditorViewportClient.h
--- a/Engine/Source/Editor/UnrealEd/Public/EditorViewportClient.h
+++ b/Engine/Source/Editor/UnrealEd/Public/EditorViewportClient.h
@@ -1288,6 +1288,26 @@ public:
 	 */
 	FText GetCurrentLumenVisualizationModeDisplayName() const;
 
+	/**
+	 * Changes the Strata visualization mode for this viewport.
+	 *
+	 * @param InName	The ID of the required visualization mode
+	 */
+	void ChangeStrataVisualizationMode(FName InName);
+
+	/**
+	 * Checks if a Strata visualization mode is selected.
+	 *
+	 * @param InName	The ID of the required visualization mode
+	 * @return	true if the supplied Strata visualization mode is checked
+	 */
+	bool IsStrataVisualizationModeSelected(FName InName) const;
+
+	/**
+	 * Returns the FText display name associated with CurrentStrataVisualizationMode.
+	 */
+	FText GetCurrentStrataVisualizationModeDisplayName() const;
+
 	/**
 	* Changes the virtual shadow map visualization mode for this viewport.
 	*
@@ -1654,6 +1674,7 @@ public:
 	FName CurrentBufferVisualizationMode;
 	FName CurrentNaniteVisualizationMode;
 	FName CurrentLumenVisualizationMode;
+	FName CurrentStrataVisualizationMode;
 	FName CurrentVirtualShadowMapVisualizationMode;
 
 	FName CurrentRayTracingDebugVisualizationMode;
diff --git a/Engine/Source/Editor/UnrealEd/Public/EditorViewportCommands.h b/Engine/Source/Editor/UnrealEd/Public/EditorViewportCommands.h
--- a/Engine/Source/Editor/UnrealEd/Public/EditorViewportCommands.h
+++ b/Engine/Source/Editor/UnrealEd/Public/EditorViewportCommands.h
@@ -136,6 +136,9 @@ public:
 	/** Changes the viewport to visualize Lumen debug views */
 	TSharedPtr< FUICommandInfo > VisualizeLumenMode;
 
+	/** Changes the viewport to visualize Strata debug views */
+	TSharedPtr< FUICommandInfo > VisualizeStrataMode;
+
 	/** Changes the viewport to visualize virtual shadow map content */
 	TSharedPtr< FUICommandInfo > VisualizeVirtualShadowMapMode;
 
diff --git a/Engine/Source/Runtime/Engine/Classes/Engine/EngineBaseTypes.h b/Engine/Source/Runtime/Engine/Classes/Engine/EngineBaseTypes.h
--- a/Engine/Source/Runtime/Engine/Classes/Engine/EngineBaseTypes.h
+++ b/Engine/Source/Runtime/Engine/Classes/Engine/EngineBaseTypes.h
@@ -922,6 +922,9 @@ enum EViewModeIndex
 	/** Visualize Skin Cache. */
 	VMI_VisualizeGPUSkinCache = 33 UMETA(DisplayName = "GPU Skin Cache Visualization"),
 
+	/** Visualize Strata debug views */
+	VMI_VisualizeStrata = 34 UMETA(DisplayName = "Strata Visualization"),
+
 	VMI_Max UMETA(Hidden),
 
 	// VMI_Unknown - The value assigned to VMI_Unknown must be the highest possible of any member of EViewModeIndex, or GetViewModeName might seg-fault
diff --git a/Engine/Source/Runtime/Engine/Classes/Engine/GameViewportClient.h b/Engine/Source/Runtime/Engine/Classes/Engine/GameViewportClient.h
--- a/Engine/Source/Runtime/Engine/Classes/Engine/GameViewportClient.h
+++ b/Engine/Source/Runtime/Engine/Classes/Engine/GameViewportClient.h
@@ -658,6 +658,9 @@ protected:
 	void SetCurrentLumenVisualizationMode(FName NewLumenVisualizationMode) { CurrentLumenVisualizationMode = NewLumenVisualizationMode; }
 	FName GetCurrentLumenVisualizationMode() const { return CurrentLumenVisualizationMode; }
 
+	void SetCurrentStrataVisualizationMode(FName NewStrataVisualizationMode) { CurrentStrataVisualizationMode = NewStrataVisualizationMode; }
+	FName GetCurrentStrataVisualizationMode() const { return CurrentStrataVisualizationMode; }
+
 	void SetCurrentVirtualShadowMapVisualizationMode(FName NewVirtualShadowMapVisualizationMode) { CurrentVirtualShadowMapVisualizationMode = NewVirtualShadowMapVisualizationMode; }
 	FName GetCurrentVirtualShadowMapVisualizationMode() const { return CurrentVirtualShadowMapVisualizationMode; }
 
@@ -965,6 +968,9 @@ private:
 	/** Current Lumen visualization mode for this game viewport */
 	FName CurrentLumenVisualizationMode;
 
+	/** Current Strata visualization mode for this game viewport */
+	FName CurrentStrataVisualizationMode;
+
 	/** Current virtual shadow map visualization mode for this game viewport */
 	FName CurrentVirtualShadowMapVisualizationMode;
 
diff --git a/Engine/Source/Runtime/Engine/Private/GameViewportClient.cpp b/Engine/Source/Runtime/Engine/Private/GameViewportClient.cpp
--- a/Engine/Source/Runtime/Engine/Private/GameViewportClient.cpp
+++ b/Engine/Source/Runtime/Engine/Private/GameViewportClient.cpp
@@ -46,6 +46,7 @@
 #include "BufferVisualizationData.h"
 #include "NaniteVisualizationData.h"
 #include "LumenVisualizationData.h"
+#include "StrataVisualizationData.h"
 #include "VirtualShadowMapVisualizationData.h"
 #include "GameFramework/InputSettings.h"
 #include "Components/LineBatchComponent.h"
@@ -237,6 +238,7 @@ UGameViewportClient::UGameViewportClient(const FObjectInitializer& ObjectInitial
 	, CurrentBufferVisualizationMode(NAME_None)
 	, CurrentNaniteVisualizationMode(NAME_None)
 	, CurrentLumenVisualizationMode(NAME_None)
+	, CurrentStrataVisualizationMode(NAME_None)
 	, CurrentVirtualShadowMapVisualizationMode(NAME_None)
 	, HighResScreenshotDialog(nullptr)
 	, bUseSoftwareCursorWidgets(true)
@@ -331,6 +333,7 @@ UGameViewportClient::UGameViewportClient(FVTableHelper& Helper)
 	, CurrentBufferVisualizationMode(NAME_None)
 	, CurrentNaniteVisualizationMode(NAME_None)
 	, CurrentLumenVisualizationMode(NAME_None)
+	, CurrentStrataVisualizationMode(NAME_None)
 	, CurrentVirtualShadowMapVisualizationMode(NAME_None)
 	, HighResScreenshotDialog(nullptr)
 	, bIgnoreInput(false)
@@ -1495,6 +1498,7 @@ void UGameViewportClient::Draw(FViewport* InViewport, FCanvas* SceneCanvas)
 					View->CurrentBufferVisualizationMode = CurrentBufferVisualizationMode;
 					View->CurrentNaniteVisualizationMode = CurrentNaniteVisualizationMode;
 					View->CurrentLumenVisualizationMode = CurrentLumenVisualizationMode;
+					View->CurrentStrataVisualizationMode = CurrentStrataVisualizationMode;
 					View->CurrentVirtualShadowMapVisualizationMode = CurrentVirtualShadowMapVisualizationMode;
 
 					View->CameraConstrainedViewRect = View->UnscaledViewRect;
diff --git a/Engine/Source/Runtime/Engine/Private/SceneView.cpp b/Engine/Source/Runtime/Engine/Private/SceneView.cpp
--- a/Engine/Source/Runtime/Engine/Private/SceneView.cpp
+++ b/Engine/Source/Runtime/Engine/Private/SceneView.cpp
@@ -19,6 +19,7 @@
 #include "BufferVisualizationData.h"
 #include "NaniteVisualizationData.h"
 #include "LumenVisualizationData.h"
+#include "StrataVisualizationData.h"
 #include "VirtualShadowMapVisualizationData.h"
 #include "Interfaces/Interface_PostProcessVolume.h"
 #include "Engine/TextureCube.h"
diff --git a/Engine/Source/Runtime/Engine/Private/ShowFlags.cpp b/Engine/Source/Runtime/Engine/Private/ShowFlags.cpp
--- a/Engine/Source/Runtime/Engine/Private/ShowFlags.cpp
+++ b/Engine/Source/Runtime/Engine/Private/ShowFlags.cpp
@@ -343,6 +343,9 @@ void ApplyViewMode(EViewModeIndex ViewModeIndex, bool bPerspective, FEngineShowF
 		case VMI_VisualizeLumen:
 			bPostProcessing = true;
 			break;
+		case VMI_VisualizeStrata:
+			bPostProcessing = true;
+			break;
 		case VMI_VisualizeVirtualShadowMap:
 			bPostProcessing = true;
 			break;
@@ -380,6 +383,7 @@ void ApplyViewMode(EViewModeIndex ViewModeIndex, bool bPerspective, FEngineShowF
 	EngineShowFlags.SetVisualizeBuffer(ViewModeIndex == VMI_VisualizeBuffer);
 	EngineShowFlags.SetVisualizeNanite(ViewModeIndex == VMI_VisualizeNanite);
 	EngineShowFlags.SetVisualizeLumen(ViewModeIndex == VMI_VisualizeLumen);
+	EngineShowFlags.SetVisualizeStrata(ViewModeIndex == VMI_VisualizeStrata);
 	EngineShowFlags.SetVisualizeVirtualShadowMap(ViewModeIndex == VMI_VisualizeVirtualShadowMap);
 	EngineShowFlags.SetVisualizeLightCulling(ViewModeIndex == VMI_LightComplexity);
 	EngineShowFlags.SetShaderComplexity(ViewModeIndex == VMI_ShaderComplexity || ViewModeIndex == VMI_QuadOverdraw || ViewModeIndex == VMI_ShaderComplexityWithQuadOverdraw);
@@ -702,6 +706,10 @@ EViewModeIndex FindViewMode(const FEngineShowFlags& EngineShowFlags)
 	{
 		return VMI_VisualizeLumen;
 	}
+	else if (EngineShowFlags.VisualizeStrata)
+	{
+		return VMI_VisualizeStrata;
+	}
 	else if (EngineShowFlags.VisualizeVirtualShadowMap)
 	{
 		return VMI_VisualizeVirtualShadowMap;
@@ -844,6 +852,7 @@ const TCHAR* GetViewModeName(EViewModeIndex ViewModeIndex)
 		case VMI_VisualizeBuffer:			return TEXT("VisualizeBuffer");
 		case VMI_VisualizeNanite:			return TEXT("VisualizeNanite");
 		case VMI_VisualizeLumen:			return TEXT("VisualizeLumen");
+		case VMI_VisualizeStrata:			return TEXT("VisualizeStrata");
 		case VMI_VisualizeVirtualShadowMap:	return TEXT("VisualizeVirtualShadowMap");
 		case VMI_RayTracingDebug:			return TEXT("RayTracingDebug");
 		case VMI_PathTracing:				return TEXT("PathTracing");
diff --git a/Engine/Source/Runtime/Engine/Private/StrataVisualizationData.cpp b/Engine/Source/Runtime/Engine/Private/StrataVisualizationData.cpp
new file mode 100644
--- /dev/null
+++ b/Engine/Source/Runtime/Engine/Private/StrataVisualizationData.cpp
@@ -0,0 +1,187 @@
+// Copyright Epic Games, Inc. All Rights Reserved.
+
+#include "StrataVisualizationData.h"
+#include "HAL/IConsoleManager.h"
+#include "Materials/Material.h"
+#include "Misc/ConfigCacheIni.h"
+
+#define LOCTEXT_NAMESPACE "FStrataVisualizationData"
+
+static FStrataVisualizationData GStrataVisualizationData;
+
+bool Engine_IsStrataEnabled();
+
+static bool Engine_IsStrataDBufferPassEnabled()
+{
+	static const auto CVar = IConsoleManager::Get().FindTConsoleVariableDataInt(TEXT("r.Strata.DBufferPass"));
+	return CVar && CVar->GetValueOnAnyThread() > 0;
+}
+
+static bool Engine_IsStrataRoughRefractionEnabled()
+{
+	static const auto CVar = IConsoleManager::Get().FindTConsoleVariableDataInt(TEXT("r.Strata.OpaqueMaterialRoughRefraction"));
+	return CVar && CVar->GetValueOnAnyThread() > 0;
+}
+
+static bool Engine_IsStrataAdvancedDebugShaderEnabled()
+{
+	static const auto CVar = IConsoleManager::Get().FindTConsoleVariableDataInt(TEXT("r.Strata.Debug.AdvancedVisualizationShaders"));
+	return CVar && CVar->GetValueOnAnyThread() > 0;
+}
+
+static FString ConfigureConsoleCommand(FStrataVisualizationData::TModeMap& ModeMap)
+{
+	FString AvailableVisualizationModes;
+	for (FStrataVisualizationData::TModeMap::TConstIterator It = ModeMap.CreateConstIterator(); It; ++It)
+	{
+		const FStrataVisualizationData::FModeRecord& Record = It.Value();
+		AvailableVisualizationModes += FString(TEXT("\n  "));
+		AvailableVisualizationModes += Record.ModeString;
+	}
+
+	FString Out;
+	Out = TEXT("When the viewport view-mode is set to 'Strata Visualization', this command specifies which of the various channels to display. Values entered other than the allowed values shown below will be ignored.");
+	Out += AvailableVisualizationModes;
+
+	IConsoleManager::Get().RegisterConsoleVariable(
+		FStrataVisualizationData::GetVisualizeConsoleCommandName(),
+		0,
+		*Out,
+		ECVF_Cheat);
+
+	return Out;
+}
+
+static void AddVisualizationMode(
+	FStrataVisualizationData::TModeMap& ModeMap,
+	const TCHAR* ModeString,
+	const FText& ModeText,
+	const FText& ModeDesc,
+	const FStrataVisualizationData::FViewMode Mode,
+	bool DefaultComposited
+)
+{
+	const FName ModeName = FName(ModeString);
+
+	FStrataVisualizationData::FModeRecord& Record = ModeMap.Emplace(ModeName);
+	Record.ModeString			= FString(ModeString);
+	Record.ModeName				= ModeName;
+	Record.ModeText				= ModeText;
+	Record.ModeDesc				= ModeDesc;
+	Record.Mode					= Mode;
+	Record.DefaultComposited	= DefaultComposited;
+}
+
+void FStrataVisualizationData::Initialize()
+{
+	if (!bIsInitialized && Engine_IsStrataEnabled())
+	{
+		AddVisualizationMode(
+			ModeMap,
+			TEXT("MaterialProperties"),
+			LOCTEXT("MaterialProperties", "Material Properites"),
+			LOCTEXT("MaterialPropertiesDesc", "Visualizes Strata material properties under mouse cursor"),
+			FViewMode::MaterialProperties,
+			true);
+
+		AddVisualizationMode(
+			ModeMap,
+			TEXT("MaterialCount"),
+			LOCTEXT("MaterialCount", "Material Count"),
+			LOCTEXT("MaterialCountDesc", "Visualizes Strata material count per pixel"),
+			FViewMode::MaterialCount,
+			true);
+
+		if (Engine_IsStrataAdvancedDebugShaderEnabled())
+		{
+			AddVisualizationMode(
+				ModeMap,
+				TEXT("AdvancedMaterialProperties"),
+				LOCTEXT("AdvancedMaterialProperties", "Advanced Material Properties"),
+				LOCTEXT("AdvancedMaterialPropertiesDesc", "Visualizes Strata advanced material properties"),
+				FViewMode::AdvancedMaterialProperties,
+				true);
+		}
+
+		AddVisualizationMode(
+			ModeMap,
+			TEXT("MaterialClassification"),
+			LOCTEXT("MaterialClassification", "Material Classification"),
+			LOCTEXT("MaterialClassificationDesc", "Visualizes Strata material classification"),
+			FViewMode::MaterialClassification,
+			true);
+
+		if (Engine_IsStrataDBufferPassEnabled())
+		{
+			AddVisualizationMode(
+				ModeMap,
+				TEXT("DecalClassification"),
+				LOCTEXT("DecalClassification", "Decal classification"),
+				LOCTEXT("DecalClassificationDesc", "Visualizes Strata decal classification"),
+				FViewMode::DecalClassification,
+				true);
+		}
+
+		if (Engine_IsStrataRoughRefractionEnabled())
+		{
+			AddVisualizationMode(
+				ModeMap,
+				TEXT("RoughRefractionClassification"),
+				LOCTEXT("RoughRefractionClassification", "Material Count"),
+				LOCTEXT("RoughRefractionClassificationDesc", "Visualizes Strata rough refraction classification"),
+				FViewMode::RoughRefractionClassification,
+				true);
+		}
+
+		ConsoleDocumentationVisualizationMode = ConfigureConsoleCommand(ModeMap);
+	}
+	bIsInitialized = true;
+}
+
+FText FStrataVisualizationData::GetModeDisplayName(const FName& InModeName) const
+{
+	if (const FModeRecord* Record = ModeMap.Find(InModeName))
+	{
+		return Record->ModeText;
+	}
+	else
+	{
+		return FText::GetEmpty();
+	}
+}
+
+FStrataVisualizationData::FViewMode FStrataVisualizationData::GetViewMode(const FName& InModeName) const
+{
+	if (const FModeRecord* Record = ModeMap.Find(InModeName))
+	{
+		return Record->Mode;
+	}
+	else
+	{
+		return FStrataVisualizationData::FViewMode::None;
+	}
+}
+
+bool FStrataVisualizationData::GetModeDefaultComposited(const FName& InModeName) const
+{
+	if (const FModeRecord* Record = ModeMap.Find(InModeName))
+	{
+		return Record->DefaultComposited;
+	}
+	else
+	{
+		return false;
+	}
+}
+
+FStrataVisualizationData& GetStrataVisualizationData()
+{
+	if (!GStrataVisualizationData.IsInitialized())
+	{
+		GStrataVisualizationData.Initialize();
+	}
+
+	return GStrataVisualizationData;
+}
+
+#undef LOCTEXT_NAMESPACE
diff --git a/Engine/Source/Runtime/Engine/Private/ViewModeNames.cpp b/Engine/Source/Runtime/Engine/Private/ViewModeNames.cpp
--- a/Engine/Source/Runtime/Engine/Private/ViewModeNames.cpp
+++ b/Engine/Source/Runtime/Engine/Private/ViewModeNames.cpp
@@ -281,6 +281,10 @@ TArray<const FSlateBrush*> FillViewModeDisplayIcons()
 		{
 			ViewModeDisplayIcons.Emplace(FAppStyle::Get().GetBrush("EditorViewport.VisualizeLumenMode"));
 		}
+		else if (ViewModeIndex == VMI_VisualizeStrata)
+		{
+			ViewModeDisplayIcons.Emplace(FAppStyle::Get().GetBrush("EditorViewport.VisualizeStrataMode"));
+		}
 		else if (ViewModeIndex == VMI_VisualizeVirtualShadowMap)
 		{
 			ViewModeDisplayIcons.Emplace(FAppStyle::Get().GetBrush("EditorViewport.VisualizeVirtualShadowMapMode"));
diff --git a/Engine/Source/Runtime/Engine/Public/SceneView.h b/Engine/Source/Runtime/Engine/Public/SceneView.h
--- a/Engine/Source/Runtime/Engine/Public/SceneView.h
+++ b/Engine/Source/Runtime/Engine/Public/SceneView.h
@@ -1075,10 +1075,13 @@ public:
 	/** Current Nanite visualization mode */
 	FName CurrentNaniteVisualizationMode;
 
-	/** Current Nanite visualization mode */
+	/** Current Lumen visualization mode */
 	FName CurrentLumenVisualizationMode;
 
-	/** Current Nanite visualization mode */
+	/** Current Strata visualization mode */
+	FName CurrentStrataVisualizationMode;
+
+	/** Current Virtual Shadow Map visualization mode */
 	FName CurrentVirtualShadowMapVisualizationMode;
 
 	/** Current visualize calibration color material name */
diff --git a/Engine/Source/Runtime/Engine/Public/ShowFlags.h b/Engine/Source/Runtime/Engine/Public/ShowFlags.h
--- a/Engine/Source/Runtime/Engine/Public/ShowFlags.h
+++ b/Engine/Source/Runtime/Engine/Public/ShowFlags.h
@@ -474,7 +474,6 @@ private:
 		SetVisualizePostProcessStack(false);
 		SetVirtualTexturePrimitives(false);
 		SetVisualizeVolumetricCloudConservativeDensity(false);
-		SetVisualizeStrataMaterial(false);
 		SetDrawOnlyVSMInvalidatingGeo(false);
 		SetDebugDrawDistantVirtualSMLights(false);
 
diff --git a/Engine/Source/Runtime/Engine/Public/ShowFlagsValues.inl b/Engine/Source/Runtime/Engine/Public/ShowFlagsValues.inl
--- a/Engine/Source/Runtime/Engine/Public/ShowFlagsValues.inl
+++ b/Engine/Source/Runtime/Engine/Public/ShowFlagsValues.inl
@@ -66,6 +66,8 @@ SHOWFLAG_FIXED_IN_SHIPPING(0, VisualizeBuffer, SFG_Hidden, NSLOCTEXT("UnrealEd",
 SHOWFLAG_ALWAYS_ACCESSIBLE(VisualizeNanite, SFG_Hidden, NSLOCTEXT("UnrealEd", "VisualizeNaniteSF", "Nanite Visualization"))
 /** Needed for VMI_VisualizeLumen, Whether to enable the Lumen visualization mode. */
 SHOWFLAG_ALWAYS_ACCESSIBLE(VisualizeLumen, SFG_Hidden, NSLOCTEXT("UnrealEd", "VisualizeLumenSF", "Lumen Visualization"))
+/** Needed for VMI_VisualizeStrata, Whether to enable the Strata visualization mode. */
+SHOWFLAG_ALWAYS_ACCESSIBLE(VisualizeStrata, SFG_Hidden, NSLOCTEXT("UnrealEd", "VisualizeStrataSF", "Strata Visualization"))
 /** Needed for VMI_VisualizeVirtualShadowMap, Whether to enable the virtual shadow map visualization mode. */
 SHOWFLAG_ALWAYS_ACCESSIBLE(VisualizeVirtualShadowMap, SFG_Hidden, NSLOCTEXT("UnrealEd", "VisualizeVirtualShadowMapSF", "Virtual Shadow Map Visualization"))
 /** Allows to disable all direct lighting (does not affect indirect light) */
@@ -395,8 +397,6 @@ SHOWFLAG_FIXED_IN_SHIPPING(0, VirtualTexturePrimitives, SFG_Developer, NSLOCTEXT
 
 /** Visualize volumetric cloud conservative density. */
 SHOWFLAG_FIXED_IN_SHIPPING(0, VisualizeVolumetricCloudConservativeDensity, SFG_Visualize, NSLOCTEXT("UnrealEd", "VisualizeVolumetricCloudConservativeDensitySF", "Volumetric Cloud Conservative Density"))
-/** Debug the Strata material buffer content */
-SHOWFLAG_FIXED_IN_SHIPPING(0, VisualizeStrataMaterial, SFG_Visualize, NSLOCTEXT("UnrealEd", "StrataMaterial", "Strata Material"))
 
 SHOWFLAG_FIXED_IN_SHIPPING(1, VirtualShadowMapCaching, SFG_Developer, NSLOCTEXT("UnrealEd", "VirtualShadowMapCaching", "Cache Virtual Shadow Maps"))
 
diff --git a/Engine/Source/Runtime/Engine/Public/StrataVisualizationData.h b/Engine/Source/Runtime/Engine/Public/StrataVisualizationData.h
new file mode 100644
--- /dev/null
+++ b/Engine/Source/Runtime/Engine/Public/StrataVisualizationData.h
@@ -0,0 +1,79 @@
+// Copyright Epic Games, Inc. All Rights Reserved.
+
+#pragma once
+
+#include "CoreMinimal.h"
+
+class FStrataVisualizationData
+{
+public:
+	enum class FViewMode : uint8
+	{
+		None,
+		MaterialProperties,
+		MaterialCount,
+		AdvancedMaterialProperties,
+		MaterialClassification,
+		DecalClassification,
+		RoughRefractionClassification,
+	};
+
+	/** Describes a single available visualization mode. */
+	struct FModeRecord
+	{
+		FString   ModeString;
+		FName     ModeName;
+		FText     ModeText;
+		FText     ModeDesc;
+		FViewMode Mode;
+
+		// Whether or not this mode (by default) composites with regular scene depth.
+		bool      DefaultComposited;
+	};
+
+	/** Mapping of FName to a visualization mode record. */
+	typedef TMultiMap<FName, FModeRecord> TModeMap;
+
+public:
+	FStrataVisualizationData()
+	: bIsInitialized(false)
+	{
+	}
+
+	/** Initialize the system. */
+	void Initialize();
+
+	/** Check if system was initialized. */
+	inline bool IsInitialized() const { return bIsInitialized; }
+
+	/** Get the display name of a named mode from the available mode map. **/
+	ENGINE_API FText GetModeDisplayName(const FName& InModeName) const;
+
+	ENGINE_API FViewMode GetViewMode(const FName& InModeName) const;
+
+	ENGINE_API bool GetModeDefaultComposited(const FName& InModeName) const;
+
+	inline const TModeMap& GetModeMap() const
+	{
+		return ModeMap;
+	}
+
+	/** Return the console command name for enabling single mode visualization. */
+	static const TCHAR* GetVisualizeConsoleCommandName()
+	{
+		return TEXT("r.Strata.ViewMode");
+	}
+
+private:
+
+	/** The name->mode mapping table */
+	TModeMap ModeMap;
+
+	/** Storage for console variable documentation strings. **/
+	FString ConsoleDocumentationVisualizationMode;
+
+	/** Flag indicating if system is initialized. **/
+	bool bIsInitialized;
+};
+
+ENGINE_API FStrataVisualizationData& GetStrataVisualizationData();
diff --git a/Engine/Source/Runtime/Renderer/Private/PostProcess/TemporalSuperResolution.cpp b/Engine/Source/Runtime/Renderer/Private/PostProcess/TemporalSuperResolution.cpp
--- a/Engine/Source/Runtime/Renderer/Private/PostProcess/TemporalSuperResolution.cpp
+++ b/Engine/Source/Runtime/Renderer/Private/PostProcess/TemporalSuperResolution.cpp
@@ -940,7 +940,8 @@ ITemporalUpscaler::FOutputs AddTemporalSuperResolutionPasses(
 	bHasSeparateTranslucency &= 
 		   (!View.Family->EngineShowFlags.VisualizeBuffer || (View.Family->EngineShowFlags.VisualizeBuffer && View.CurrentBufferVisualizationMode == OverviewName))
 		&& (!View.Family->EngineShowFlags.VisualizeNanite || (View.Family->EngineShowFlags.VisualizeNanite && View.CurrentNaniteVisualizationMode == OverviewName))
-		&& (!View.Family->EngineShowFlags.VisualizeLumen  || (View.Family->EngineShowFlags.VisualizeLumen  && View.CurrentLumenVisualizationMode  == OverviewName));
+		&& (!View.Family->EngineShowFlags.VisualizeLumen  || (View.Family->EngineShowFlags.VisualizeLumen  && View.CurrentLumenVisualizationMode  == OverviewName))
+		&& (!View.Family->EngineShowFlags.VisualizeStrata || (View.Family->EngineShowFlags.VisualizeStrata && View.CurrentStrataVisualizationMode == OverviewName));
 #endif
 	if (bHasSeparateTranslucency)
 	{
diff --git a/Engine/Source/Runtime/Renderer/Private/Strata/Strata.h b/Engine/Source/Runtime/Renderer/Private/Strata/Strata.h
--- a/Engine/Source/Runtime/Renderer/Private/Strata/Strata.h
+++ b/Engine/Source/Runtime/Renderer/Private/Strata/Strata.h
@@ -215,8 +215,4 @@ FStrataTilePassVS::FParameters SetTileParameters(FRDGBuilder& GraphBuilder, cons
 FStrataTilePassVS::FParameters SetTileParameters(const FViewInfo& View, const EStrataTileType Type, EPrimitiveType& PrimitiveType);
 uint32 TileTypeDrawIndirectArgOffset(const EStrataTileType Type);
 uint32 TileTypeDispatchIndirectArgOffset(const EStrataTileType Type);
-
-bool ShouldRenderStrataRoughRefractionRnD();
-void StrataRoughRefractionRnD(FRDGBuilder& GraphBuilder, const FViewInfo& View, FScreenPassTexture& ScreenPassSceneColor);
-
 };
diff --git a/Engine/Source/Runtime/Renderer/Private/Strata/StrataVisualize.cpp b/Engine/Source/Runtime/Renderer/Private/Strata/StrataVisualize.cpp
--- a/Engine/Source/Runtime/Renderer/Private/Strata/StrataVisualize.cpp
+++ b/Engine/Source/Runtime/Renderer/Private/Strata/StrataVisualize.cpp
@@ -8,20 +8,9 @@
 #include "PixelShaderUtils.h"
 #include "BasePassRendering.h"
 #include "IndirectLightRendering.h"
+#include "StrataVisualizationData.h"
 
 
-static TAutoConsoleVariable<int32> CVarStrataClassificationDebug(
-	TEXT("r.Strata.Classification.Debug"),
-	0,
-	TEXT("Enable strata classification visualization: 1 shows simple material tiles in green and complex material tiles in red."),
-	ECVF_RenderThreadSafe);
-
-static TAutoConsoleVariable<int32> CVarStrataDebugVisualizeMode(
-	TEXT("r.Strata.Debug.VisualizeMode"),
-	1,
-	TEXT("Strata debug view mode."),
-	ECVF_RenderThreadSafe);
-
 static TAutoConsoleVariable<int32> CVarStrataDebugAdvancedVisualizationShaders(
 	TEXT("r.Strata.Debug.AdvancedVisualizationShaders"),
 	0,
@@ -196,141 +185,170 @@ class FMaterialDebugStrataTreePS : public FGlobalShader
 };
 IMPLEMENT_GLOBAL_SHADER(FMaterialDebugStrataTreePS, "/Engine/Private/Strata/StrataVisualize.usf", "MaterialDebugStrataTreePS", SF_Pixel);
 
-static void AddVisualizeMaterialPasses(FRDGBuilder& GraphBuilder, const FViewInfo& View, FScreenPassTexture& ScreenPassSceneColor, EShaderPlatform Platform)
+static void AddVisualizeMaterialPropertiesPasses(FRDGBuilder& GraphBuilder, const FViewInfo& View, FScreenPassTexture& ScreenPassSceneColor, EShaderPlatform Platform)
+{
+	// Force ShaderPrint on.
+	ShaderPrint::SetEnabled(true);
+	ShaderPrint::RequestSpaceForLines(1024);
+	ShaderPrint::RequestSpaceForCharacters(1024);
+	FRDGBufferUAVRef PrintOffsetBufferUAV = nullptr;
+
+	FRDGBufferRef PrintOffsetBuffer = GraphBuilder.CreateBuffer(FRDGBufferDesc::CreateBufferDesc(4, 2), TEXT("Strata.DebugPrintPositionOffset"));
+	PrintOffsetBufferUAV = GraphBuilder.CreateUAV(PrintOffsetBuffer, PF_R32_UINT);
+	AddClearUAVPass(GraphBuilder, PrintOffsetBufferUAV, 50u);
+	const uint32 MaxBSDFCount = 8;
+
+	for (uint32 BSDFIndex=0; BSDFIndex < MaxBSDFCount; ++BSDFIndex)
+	{
+		FMaterialPrintInfoCS::FParameters* PassParameters = GraphBuilder.AllocParameters<FMaterialPrintInfoCS::FParameters>();
+		PassParameters->BSDFIndex = BSDFIndex;
+		PassParameters->RWPositionOffsetBuffer = PrintOffsetBufferUAV;
+		PassParameters->ViewUniformBuffer = View.ViewUniformBuffer;
+		PassParameters->Strata = Strata::BindStrataGlobalUniformParameters(View);
+		PassParameters->SceneTextures = GetSceneTextureParameters(GraphBuilder, View);
+		ShaderPrint::SetParameters(GraphBuilder, View.ShaderPrintData, PassParameters->ShaderPrintParameters);
+
+		TShaderMapRef<FMaterialPrintInfoCS> ComputeShader(View.ShaderMap);
+		FComputeShaderUtils::AddPass(GraphBuilder, RDG_EVENT_NAME("Strata::VisualizeMaterial(Print, BSDF=%d)", BSDFIndex), ComputeShader, PassParameters, FIntVector(1, 1, 1));
+	}
+}
+
+static void AddVisualizeMaterialCountPasses(FRDGBuilder & GraphBuilder, const FViewInfo & View, FScreenPassTexture & ScreenPassSceneColor, EShaderPlatform Platform)
 {
+	ShaderPrint::SetEnabled(true);
+	ShaderPrint::RequestSpaceForLines(1024);
+	ShaderPrint::RequestSpaceForCharacters(1024);
+
 	FRDGTextureRef SceneColorTexture = ScreenPassSceneColor.Texture;
 	FRHIBlendState* PreMultipliedColorTransmittanceBlend = TStaticBlendState<CW_RGB, BO_Add, BF_One, BF_SourceAlpha, BO_Add, BF_Zero, BF_One>::GetRHI();
-	if (View.Family->EngineShowFlags.VisualizeStrataMaterial)
+
+	FVisualizeMaterialPS::FParameters* PassParameters = GraphBuilder.AllocParameters<FVisualizeMaterialPS::FParameters>();
+	PassParameters->ViewUniformBuffer = View.ViewUniformBuffer;
+	PassParameters->ViewMode = 2;
+	PassParameters->Strata = Strata::BindStrataGlobalUniformParameters(View);
+	PassParameters->SceneTextures = GetSceneTextureParameters(GraphBuilder, View);
+	PassParameters->RenderTargets[0] = FRenderTargetBinding(SceneColorTexture, ERenderTargetLoadAction::ELoad);
+	ShaderPrint::SetParameters(GraphBuilder, View.ShaderPrintData, PassParameters->ShaderPrintParameters);
+
+	FVisualizeMaterialPS::FPermutationDomain PermutationVector;
+	TShaderMapRef<FVisualizeMaterialPS> PixelShader(View.ShaderMap, PermutationVector);
+
+	FPixelShaderUtils::AddFullscreenPass<FVisualizeMaterialPS>(GraphBuilder, View.ShaderMap, RDG_EVENT_NAME("Strata::VisualizeMaterial(Draw)"), PixelShader, PassParameters, ScreenPassSceneColor.ViewRect, PreMultipliedColorTransmittanceBlend);
+}
+
+// Draw each material layer independently
+static void AddVisualizeAdvancedMaterialPasses(FRDGBuilder& GraphBuilder, const FViewInfo& View, FScreenPassTexture& ScreenPassSceneColor, EShaderPlatform Platform)
+{
+	if (!IsStrataAdvancedVisualizationShadersEnabled())
 	{
-		// Force ShaderPrint on.
-		ShaderPrint::SetEnabled(true);
-		ShaderPrint::RequestSpaceForLines(1024);
-		ShaderPrint::RequestSpaceForCharacters(1024);
+		return;
+	}
 
-		const uint32 ViewMode = FMath::Max(0, CVarStrataDebugVisualizeMode.GetValueOnRenderThread());
+	ShaderPrint::SetEnabled(true);
+	ShaderPrint::RequestSpaceForLines(1024);
+	ShaderPrint::RequestSpaceForCharacters(1024);
 
-		FSceneTextureParameters SceneTextureParameters = GetSceneTextureParameters(GraphBuilder, View);
+	FRDGTextureRef SceneColorTexture = ScreenPassSceneColor.Texture;
+	FRHIBlendState* PreMultipliedColorTransmittanceBlend = TStaticBlendState<CW_RGB, BO_Add, BF_One, BF_SourceAlpha, BO_Add, BF_Zero, BF_One>::GetRHI();
+	{
+		FMaterialDebugStrataTreeCS::FParameters* PassParameters = GraphBuilder.AllocParameters<FMaterialDebugStrataTreeCS::FParameters>();
+		PassParameters->ViewUniformBuffer = View.ViewUniformBuffer;
+		PassParameters->Strata = Strata::BindStrataGlobalUniformParameters(View);
+		ShaderPrint::SetParameters(GraphBuilder, View.ShaderPrintData, PassParameters->ShaderPrintParameters);
 
-		// Print Material info
-		if(ViewMode >= 1 && ViewMode != 3)
-		{
-			FRDGBufferUAVRef PrintOffsetBufferUAV = nullptr;
-
-			FRDGBufferRef PrintOffsetBuffer = GraphBuilder.CreateBuffer(FRDGBufferDesc::CreateBufferDesc(4, 2), TEXT("Strata.DebugPrintPositionOffset"));
-			PrintOffsetBufferUAV = GraphBuilder.CreateUAV(PrintOffsetBuffer, PF_R32_UINT);
-			AddClearUAVPass(GraphBuilder, PrintOffsetBufferUAV, 50u);
-			const uint32 MaxBSDFCount = 8;
-
-			for (uint32 BSDFIndex=0; BSDFIndex < MaxBSDFCount; ++BSDFIndex)
-			{
-				FMaterialPrintInfoCS::FParameters* PassParameters = GraphBuilder.AllocParameters<FMaterialPrintInfoCS::FParameters>();
-				PassParameters->BSDFIndex = BSDFIndex;
-				PassParameters->RWPositionOffsetBuffer = PrintOffsetBufferUAV;
-				PassParameters->ViewUniformBuffer = View.ViewUniformBuffer;
-				PassParameters->Strata = Strata::BindStrataGlobalUniformParameters(View);
-				PassParameters->SceneTextures = GetSceneTextureParameters(GraphBuilder, View);
-				ShaderPrint::SetParameters(GraphBuilder, View.ShaderPrintData, PassParameters->ShaderPrintParameters);
-
-				TShaderMapRef<FMaterialPrintInfoCS> ComputeShader(View.ShaderMap);
-				FComputeShaderUtils::AddPass(GraphBuilder, RDG_EVENT_NAME("Strata::VisualizeMaterial(Print, BSDF=%d)", BSDFIndex), ComputeShader, PassParameters, FIntVector(1, 1, 1));
-			}
-		}
+		TShaderMapRef<FMaterialDebugStrataTreeCS> ComputeShader(View.ShaderMap);
+		FComputeShaderUtils::AddPass(GraphBuilder, RDG_EVENT_NAME("Strata::StrataAdvancedVisualization(Print)"), ComputeShader, PassParameters, FIntVector(1, 1, 1));
+	}
 
-		// Draw material debug
-		if (ViewMode >= 2 && ViewMode != 3)
+	{
+		FMaterialDebugStrataTreePS::FParameters* PassParameters = GraphBuilder.AllocParameters<FMaterialDebugStrataTreePS::FParameters>();
+		PassParameters->ViewUniformBuffer = View.ViewUniformBuffer;
+		PassParameters->Strata = Strata::BindStrataGlobalUniformParameters(View);
+		PassParameters->ReflectionStruct = CreateReflectionUniformBuffer(View, UniformBuffer_SingleFrame);
+		PassParameters->ReflectionCapture = View.ReflectionCaptureUniformBuffer;
+		PassParameters->ForwardLightData = View.ForwardLightingResources.ForwardLightUniformBuffer;
+		PassParameters->RenderTargets[0] = FRenderTargetBinding(SceneColorTexture, ERenderTargetLoadAction::ELoad);
+
+		const float DynamicBentNormalAO = 0.0f;
+		FSkyLightSceneProxy* NullSkyLight = nullptr;
+		PassParameters->SkyDiffuseLighting = GetSkyDiffuseLightingParameters(NullSkyLight, DynamicBentNormalAO);
+
+		FMaterialDebugStrataTreePS::FPermutationDomain PermutationVector;
+		TShaderMapRef<FMaterialDebugStrataTreePS> PixelShader(View.ShaderMap, PermutationVector);
+
+		FPixelShaderUtils::AddFullscreenPass<FMaterialDebugStrataTreePS>(GraphBuilder, View.ShaderMap, RDG_EVENT_NAME("Strata::StrataAdvancedVisualization(Draw)"), PixelShader, PassParameters, ScreenPassSceneColor.ViewRect, PreMultipliedColorTransmittanceBlend);
+	}
+}
+
+bool ShouldRenderStrataRoughRefractionRnD();
+void StrataRoughRefractionRnD(FRDGBuilder& GraphBuilder, const FViewInfo& View, FScreenPassTexture& ScreenPassSceneColor);
+
+static FStrataVisualizationData::FViewMode GetStrataVisualizeMode(const FViewInfo & View)
+{
+	FStrataVisualizationData::FViewMode Out = FStrataVisualizationData::FViewMode::None;
+	if (IsStrataEnabled() && StrataDebugVisualizationCanRunOnPlatform(View.GetShaderPlatform()))
+	{
+		// Variable defined in StrataVisualizationData.h/.cpp
+		static const auto CVarStrataViewMode = IConsoleManager::Get().FindConsoleVariable(FStrataVisualizationData::GetVisualizeConsoleCommandName());
+		const uint32 ViewMode = CVarStrataViewMode && CVarStrataViewMode->AsVariableInt() ? CVarStrataViewMode->AsVariableInt()->GetValueOnRenderThread() : 0;
+		switch (ViewMode)
 		{
-			FVisualizeMaterialPS::FParameters* PassParameters = GraphBuilder.AllocParameters<FVisualizeMaterialPS::FParameters>();
-			PassParameters->ViewUniformBuffer = View.ViewUniformBuffer;
-			PassParameters->ViewMode = ViewMode;
-			PassParameters->Strata = Strata::BindStrataGlobalUniformParameters(View);
-			PassParameters->SceneTextures = GetSceneTextureParameters(GraphBuilder, View);
-			PassParameters->RenderTargets[0] = FRenderTargetBinding(SceneColorTexture, ERenderTargetLoadAction::ELoad);
-			ShaderPrint::SetParameters(GraphBuilder, View.ShaderPrintData, PassParameters->ShaderPrintParameters);
-
-			FVisualizeMaterialPS::FPermutationDomain PermutationVector;
-			TShaderMapRef<FVisualizeMaterialPS> PixelShader(View.ShaderMap, PermutationVector);
-
-			FPixelShaderUtils::AddFullscreenPass<FVisualizeMaterialPS>(GraphBuilder, View.ShaderMap, RDG_EVENT_NAME("Strata::VisualizeMaterial(Draw)"), PixelShader, PassParameters, ScreenPassSceneColor.ViewRect, PreMultipliedColorTransmittanceBlend);
+			case 1: return FStrataVisualizationData::FViewMode::MaterialProperties;
+			case 2: return FStrataVisualizationData::FViewMode::MaterialCount;
+			case 3: return FStrataVisualizationData::FViewMode::AdvancedMaterialProperties;
+			case 4: return FStrataVisualizationData::FViewMode::MaterialClassification;
+			case 5: return FStrataVisualizationData::FViewMode::DecalClassification;
+			case 6: return FStrataVisualizationData::FViewMode::RoughRefractionClassification;
 		}
 
-		// Draw each material layer independently
-		if (ViewMode == 3)
+		const FStrataVisualizationData& VisualizationData = GetStrataVisualizationData();
+		if (View.Family && View.Family->EngineShowFlags.VisualizeStrata)
 		{
-			if (IsStrataAdvancedVisualizationShadersEnabled())
-			{
-				{
-					FMaterialDebugStrataTreeCS::FParameters* PassParameters = GraphBuilder.AllocParameters<FMaterialDebugStrataTreeCS::FParameters>();
-					PassParameters->ViewUniformBuffer = View.ViewUniformBuffer;
-					PassParameters->Strata = Strata::BindStrataGlobalUniformParameters(View);
-					ShaderPrint::SetParameters(GraphBuilder, View.ShaderPrintData, PassParameters->ShaderPrintParameters);
-
-					TShaderMapRef<FMaterialDebugStrataTreeCS> ComputeShader(View.ShaderMap);
-					FComputeShaderUtils::AddPass(GraphBuilder, RDG_EVENT_NAME("Strata::StrataAdvancedVisualization(Print)"), ComputeShader, PassParameters, FIntVector(1, 1, 1));
-				}
-
-				{
-					FMaterialDebugStrataTreePS::FParameters* PassParameters = GraphBuilder.AllocParameters<FMaterialDebugStrataTreePS::FParameters>();
-					PassParameters->ViewUniformBuffer = View.ViewUniformBuffer;
-					PassParameters->Strata = Strata::BindStrataGlobalUniformParameters(View);
-					PassParameters->ReflectionStruct = CreateReflectionUniformBuffer(View, UniformBuffer_SingleFrame);
-					PassParameters->ReflectionCapture = View.ReflectionCaptureUniformBuffer;
-					PassParameters->ForwardLightData = View.ForwardLightingResources.ForwardLightUniformBuffer;
-					PassParameters->RenderTargets[0] = FRenderTargetBinding(SceneColorTexture, ERenderTargetLoadAction::ELoad);
-
-					const float DynamicBentNormalAO = 0.0f;
-					FSkyLightSceneProxy* NullSkyLight = nullptr;
-					PassParameters->SkyDiffuseLighting = GetSkyDiffuseLightingParameters(NullSkyLight, DynamicBentNormalAO);
-
-					FMaterialDebugStrataTreePS::FPermutationDomain PermutationVector;
-					TShaderMapRef<FMaterialDebugStrataTreePS> PixelShader(View.ShaderMap, PermutationVector);
-
-					FPixelShaderUtils::AddFullscreenPass<FMaterialDebugStrataTreePS>(GraphBuilder, View.ShaderMap, RDG_EVENT_NAME("Strata::StrataAdvancedVisualization(Draw)"), PixelShader, PassParameters, ScreenPassSceneColor.ViewRect, PreMultipliedColorTransmittanceBlend);
-				}
-			}
-			else
-			{
-				// TODO warn the user that this mode is not available with delegate SceneRenderer.OnGetOnScreenMessages.AddLambda([](FScreenMessageWriter& ScreenMessageWriter)->void
-			}
+			Out = VisualizationData.GetViewMode(View.CurrentStrataVisualizationMode);
 		}
 	}
+	return Out;
 }
 
 bool ShouldRenderStrataDebugPasses(const FViewInfo& View)
 {
-	return IsStrataEnabled() && (
-		(StrataDebugVisualizationCanRunOnPlatform(View.GetShaderPlatform()) && ( (View.Family && View.Family->EngineShowFlags.VisualizeStrataMaterial) || CVarStrataClassificationDebug.GetValueOnAnyThread() > 0 ))
-		|| ShouldRenderStrataRoughRefractionRnD()
-		);
+	return GetStrataVisualizeMode(View) != FStrataVisualizationData::FViewMode::None || ShouldRenderStrataRoughRefractionRnD();
 }
 
 FScreenPassTexture AddStrataDebugPasses(FRDGBuilder& GraphBuilder, const FViewInfo& View, FScreenPassTexture& ScreenPassSceneColor)
 {
 	check(IsStrataEnabled());
-	EShaderPlatform Platform = View.GetShaderPlatform();
 
-	if (StrataDebugVisualizationCanRunOnPlatform(Platform))
+	const FStrataVisualizationData::FViewMode DebugMode = GetStrataVisualizeMode(View);
+	if (DebugMode != FStrataVisualizationData::FViewMode::None)
 	{
 		RDG_EVENT_SCOPE(GraphBuilder, "Strata::VisualizeMaterial");
-		AddVisualizeMaterialPasses(GraphBuilder, View, ScreenPassSceneColor, Platform);
-	}
 
-	const int32 StrataClassificationDebug = CVarStrataClassificationDebug.GetValueOnAnyThread();
-	if (StrataClassificationDebug > 0)
-	{
-		RDG_EVENT_SCOPE(GraphBuilder, "Strata::VisualizeClassification");
 		const bool bDebugPass = true;
-		if (StrataClassificationDebug == 3)
+		if (DebugMode == FStrataVisualizationData::FViewMode::MaterialProperties)
+		{
+			AddVisualizeMaterialPropertiesPasses(GraphBuilder, View, ScreenPassSceneColor, View.GetShaderPlatform());
+		}
+		if (DebugMode == FStrataVisualizationData::FViewMode::MaterialCount)
+		{
+			AddVisualizeMaterialCountPasses(GraphBuilder, View, ScreenPassSceneColor, View.GetShaderPlatform());
+		}
+		if (DebugMode == FStrataVisualizationData::FViewMode::AdvancedMaterialProperties)
+		{
+			AddVisualizeAdvancedMaterialPasses(GraphBuilder, View, ScreenPassSceneColor, View.GetShaderPlatform());
+		}
+		else if (DebugMode == FStrataVisualizationData::FViewMode::DecalClassification)
 		{
 			AddStrataInternalClassificationTilePass(GraphBuilder, View, nullptr, &ScreenPassSceneColor.Texture, EStrataTileType::EDecalSimple, bDebugPass);
 			AddStrataInternalClassificationTilePass(GraphBuilder, View, nullptr, &ScreenPassSceneColor.Texture, EStrataTileType::EDecalSingle, bDebugPass);
 			AddStrataInternalClassificationTilePass(GraphBuilder, View, nullptr, &ScreenPassSceneColor.Texture, EStrataTileType::EDecalComplex, bDebugPass);
 		}
-		else if (StrataClassificationDebug == 2)
+		else if (DebugMode == FStrataVisualizationData::FViewMode::RoughRefractionClassification)
 		{
 			AddStrataInternalClassificationTilePass(GraphBuilder, View, nullptr, &ScreenPassSceneColor.Texture, EStrataTileType::EOpaqueRoughRefraction, bDebugPass);
 			AddStrataInternalClassificationTilePass(GraphBuilder, View, nullptr, &ScreenPassSceneColor.Texture, EStrataTileType::EOpaqueRoughRefractionSSSWithout, bDebugPass);
 		}
-		else
+		else if (DebugMode == FStrataVisualizationData::FViewMode::MaterialClassification)
 		{
 			AddStrataInternalClassificationTilePass(GraphBuilder, View, nullptr, &ScreenPassSceneColor.Texture, EStrataTileType::EComplex, bDebugPass);
 			AddStrataInternalClassificationTilePass(GraphBuilder, View, nullptr, &ScreenPassSceneColor.Texture, EStrataTileType::ESingle, bDebugPass);
