Simplified cloud debug to a single CloudDebugViewMode.
Removed sample count view mode now super seeded by CLOUD_DEBUG_SAMPLES.
Only used by compute shader.
Made code common in UpdateCloudParam, empty space skipping slice depth is now computed on cpu and shared.

#swarm https://p4-swarm.epicgames.net/reviews/22690110
Ask me for the perf test link

manual merge from staging to main

diff --git a/Engine/Shaders/Private/VolumetricCloud.usf b/Engine/Shaders/Private/VolumetricCloud.usf
--- a/Engine/Shaders/Private/VolumetricCloud.usf
+++ b/Engine/Shaders/Private/VolumetricCloud.usf
@@ -55,8 +55,8 @@
 #define CLOUD_DEBUG_SAMPLES 0
 #endif
 
-#ifndef CLOUD_SAMPLE_COUNT_DEBUG_MODE
-#define CLOUD_SAMPLE_COUNT_DEBUG_MODE 0
+#ifndef CLOUD_DEBUG_VIEW_MODE
+#define CLOUD_DEBUG_VIEW_MODE 0
 #endif
 
 // Draw debug sample positions
@@ -635,7 +635,7 @@ void MainCommon(in FMaterialPixelParameters MaterialParameters, in float4 SvPosi
 		const float MsPhaseFactor = 1.0f;
 #endif
 
-#if CLOUD_CONSERVATIVE_DENSITY_DEBUG
+#if CLOUD_DEBUG_VIEW_MODE
 		uint DebugZeroConservativeDensitySampleCount = 0;
 #endif
 
@@ -697,7 +697,7 @@ void MainCommon(in FMaterialPixelParameters MaterialParameters, in float4 SvPosi
 		// Read from empty space skipping
 		//
 		uint StartStepCount = 0;
-	#if SHADER_RENDERVIEW_CS_WITH_EMPTY_SPACE_SKIPPING
+	#if SHADER_RENDERVIEW_CS_WITH_EMPTY_SPACE_SKIPPING && !CLOUD_DEBUG_VIEW_MODE
 		const float StartDepthKm = StartTracingDistanceTexture[SvPositionToViewportUV(SvPosition) * StartTracingDistanceTextureResolution];
 		const float StartDepthCm = StartDepthKm * KILOMETER_TO_CENTIMETER;
 		if (t < StartDepthCm)
@@ -720,18 +720,13 @@ void MainCommon(in FMaterialPixelParameters MaterialParameters, in float4 SvPosi
 			// Update FMaterialPixelParameters according to current sample.
 			//////////////////////////////
 			UpdateMaterialCloudParam(MaterialParameters, SampleWorldPosition,
-				ResolvedView, CloudLayerParams, TRACING_SHADOW_DISTANCE_OFF);
+				ResolvedView, CloudLayerParams, TRACING_SHADOW_DISTANCE_OFF, RenderVolumetricCloudParameters.EmptySpaceSkippingSliceDepth);
 
-#if MATERIAL_VOLUMETRIC_ADVANCED_CONSERVATIVE_DENSITY
+#if MATERIAL_VOLUMETRIC_ADVANCED_CONSERVATIVE_DENSITY && !CLOUD_DEBUG_VIEW_MODE
 			if (MaterialParameters.VolumeSampleConservativeDensity.x <= 0.0f)
 			{
-#if CLOUD_CONSERVATIVE_DENSITY_DEBUG
-				++DebugZeroConservativeDensitySampleCount;
-				t += StepT;
-#else
 				i += StepSizeOnZeroConservativeDensity - 1;
 				t += StepSizeOnZeroConservativeDensity * StepT;
-#endif
 
 			#if CLOUD_DEBUG_SAMPLES
 				if (bDrawDebugEnabled) AddCrossTWS(LWCToFloat(SampleTranslatedWorldPosition), 60000.0f, float4(1, 0, 0, 0.5));
@@ -746,6 +741,18 @@ void MainCommon(in FMaterialPixelParameters MaterialParameters, in float4 SvPosi
 			#endif // CLOUD_DEBUG_SAMPLES
 #endif
 
+#if CLOUD_DEBUG_VIEW_MODE
+			if (RenderVolumetricCloudParameters.CloudDebugViewMode == 1 && MaterialParameters.VolumeSampleConservativeDensity.x <= 0.0f)
+			{
+				++DebugZeroConservativeDensitySampleCount;
+			}
+		#if SHADER_RENDERVIEW_CS_WITH_EMPTY_SPACE_SKIPPING
+			else if (RenderVolumetricCloudParameters.CloudDebugViewMode == 2 && GetVolumetricCloudEmptySpaceSkippingOutput0(MaterialParameters).x <= 0.0f)
+			{
+				++DebugZeroConservativeDensitySampleCount;
+			}
+		#endif
+#endif
 
 			//////////////////////////////
 			// Sample the participating media material at the sample position.
@@ -837,7 +844,7 @@ void MainCommon(in FMaterialPixelParameters MaterialParameters, in float4 SvPosi
 						const float DetlaNormT = CurrentNormT - PreviousNormT;
 						const float ShadowSampleDistance = ShadowLengthTest * (CurrentNormT - 0.5 * DetlaNormT);
 						UpdateMaterialCloudParam(MaterialParameters, LWCAdd(SampleWorldPosition, LWCPromote(float3(0.0f, 0.0f, -1.0f) * ShadowSampleDistance)),
-							ResolvedView, CloudLayerParams, ShadowSampleDistance);
+							ResolvedView, CloudLayerParams, ShadowSampleDistance, SPACE_SKIPPING_SLICE_DEPTH_OFF);
 						PreviousNormT = CurrentNormT;
 #if MATERIAL_VOLUMETRIC_ADVANCED_CONSERVATIVE_DENSITY
 						if (MaterialParameters.VolumeSampleConservativeDensity.x <= 0.0f)
@@ -857,7 +864,7 @@ void MainCommon(in FMaterialPixelParameters MaterialParameters, in float4 SvPosi
 					{
 						const float ShadowSampleDistance = ShadowLengthTest * (ShadowT * InvShadowStepCount);
 						UpdateMaterialCloudParam(MaterialParameters, LWCAdd(SampleWorldPosition, LWCPromote(float3(0.0f, 0.0f, -1.0f) * ShadowSampleDistance)),
-							ResolvedView, CloudLayerParams, ShadowSampleDistance);
+							ResolvedView, CloudLayerParams, ShadowSampleDistance, SPACE_SKIPPING_SLICE_DEPTH_OFF);
 #if MATERIAL_VOLUMETRIC_ADVANCED_CONSERVATIVE_DENSITY
 						if (MaterialParameters.VolumeSampleConservativeDensity.x <= 0.0f)
 						{
@@ -913,7 +920,7 @@ void MainCommon(in FMaterialPixelParameters MaterialParameters, in float4 SvPosi
 				{
 					PMC.TransmittanceToLight0[ms] *= exp(-OutOpticalDepth * (ms == 0 ? 1.0f : pow(MsExtinFactor, ms)));
 				}
-#else 
+#else			
 				// Use raymarched shadows
 				UNROLL
 				for (ms = 0; ms < MSCOUNT; ++ms)
@@ -927,7 +934,7 @@ void MainCommon(in FMaterialPixelParameters MaterialParameters, in float4 SvPosi
 				{
 					const float ShadowSampleDistance = ShadowLengthTest * (ShadowT * InvShadowStepCount);
 					UpdateMaterialCloudParam(MaterialParameters, LWCAdd(SampleWorldPosition, LWCPromote(Light0Direction * ShadowSampleDistance)),
-						ResolvedView, CloudLayerParams, ShadowSampleDistance);
+						ResolvedView, CloudLayerParams, ShadowSampleDistance, SPACE_SKIPPING_SLICE_DEPTH_OFF);
 					const float ExtinctionFactor = 1.0f;
 #else
 				// Non-linear shadow sample distribution, hardcoded to x^2
@@ -941,7 +948,7 @@ void MainCommon(in FMaterialPixelParameters MaterialParameters, in float4 SvPosi
 					const float ShadowSampleDistance = ShadowLengthTest * (PreviousNormT + DetlaNormT * ShadowJitterNorm);
 					const FLWCVector3 ShadowSampleWorldPos = LWCAdd(SampleWorldPosition, LWCPromote(Light0Direction * ShadowSampleDistance));
 					UpdateMaterialCloudParam(MaterialParameters, ShadowSampleWorldPos,
-						ResolvedView, CloudLayerParams, ShadowSampleDistance);
+						ResolvedView, CloudLayerParams, ShadowSampleDistance, SPACE_SKIPPING_SLICE_DEPTH_OFF);
 					PreviousNormT = CurrentNormT;
 #endif
 
@@ -1015,7 +1022,7 @@ void MainCommon(in FMaterialPixelParameters MaterialParameters, in float4 SvPosi
 						const float ExtinctionFactor = DetlaNormT;
 						const float ShadowSampleDistance = ShadowLengthTest * (PreviousNormT + DetlaNormT * ShadowJitterNorm);
 						UpdateMaterialCloudParam(MaterialParameters, LWCAdd(SampleWorldPosition, LWCPromote(Light1Direction * ShadowSampleDistance)),
-							ResolvedView, CloudLayerParams, ShadowSampleDistance);
+							ResolvedView, CloudLayerParams, ShadowSampleDistance, SPACE_SKIPPING_SLICE_DEPTH_OFF);
 						PreviousNormT = CurrentNormT;
 
 #if MATERIAL_VOLUMETRIC_ADVANCED_CONSERVATIVE_DENSITY
@@ -1145,7 +1152,7 @@ void MainCommon(in FMaterialPixelParameters MaterialParameters, in float4 SvPosi
 							for (float ShadowT = ShadowJitterNorm; ShadowT < ShadowStepCount; ShadowT += 1.0f)
 							{
 								UpdateMaterialCloudParam(MaterialParameters, LWCAdd(SampleWorldPosition, LWCPromote(SampleToLightNorm * ShadowLengthTest * (ShadowT * InvShadowStepCount))),
-									ResolvedView, CloudLayerParams, TRACING_SHADOW_DISTANCE_OFF);
+									ResolvedView, CloudLayerParams, TRACING_SHADOW_DISTANCE_OFF, SPACE_SKIPPING_SLICE_DEPTH_OFF);
 								const float ExtinctionFactor = 1.0f;
 
 								CalcPixelMaterialInputs(MaterialParameters, PixelMaterialInputs);
@@ -1252,8 +1259,7 @@ void MainCommon(in FMaterialPixelParameters MaterialParameters, in float4 SvPosi
 			}
 
 			// This is helpful for performance. Can result in less light pucnhing through clouds. Should be a setting.
-			if (all(TransmittanceToView < RenderVolumetricCloudParameters.StopTracingTransmittanceThreshold) 
-				&& !CLOUD_CONSERVATIVE_DENSITY_DEBUG)
+			if (all(TransmittanceToView < RenderVolumetricCloudParameters.StopTracingTransmittanceThreshold) && (!CLOUD_DEBUG_VIEW_MODE))
 			{
 				break;
 			}
@@ -1261,8 +1267,11 @@ void MainCommon(in FMaterialPixelParameters MaterialParameters, in float4 SvPosi
 			t += StepT;
 			}
 
-#if CLOUD_CONSERVATIVE_DENSITY_DEBUG
-		OutColor0.r = 1.0 - saturate(DebugZeroConservativeDensitySampleCount / float(IStepCount));
+#if CLOUD_DEBUG_VIEW_MODE
+		if (RenderVolumetricCloudParameters.CloudDebugViewMode >= 1)
+		{
+			OutColor0.r = 1.0 - saturate(DebugZeroConservativeDensitySampleCount / float(IStepCount));
+		}
 #endif
 	}
 	////////////////////////////////////////////////////////////////////////////////////////////////////
@@ -1338,7 +1347,7 @@ void MainCommon(in FMaterialPixelParameters MaterialParameters, in float4 SvPosi
 	//
 
 	float GrayScaleTransmittance = MeanTransmittance < RenderVolumetricCloudParameters.StopTracingTransmittanceThreshold ? 0.0f : MeanTransmittance;
-#if CLOUD_CONSERVATIVE_DENSITY_DEBUG
+#if CLOUD_DEBUG_VIEW_MODE
 	OutColor0 = float4(OutColor0.r * OutputPreExposure, Luminance.gb * OutputPreExposure, GrayScaleTransmittance);
 #else
 	OutColor0 = float4(Luminance * OutputPreExposure, GrayScaleTransmittance);
@@ -1461,8 +1470,8 @@ void MainCS(uint3 DispatchThreadId : SV_DispatchThreadID, uint3 GroupThreadId :
 {
 #if MATERIAL_VOLUMETRIC_CLOUD_EMPTY_SPACE_SKIPPING_OUTPUT
 
+	const float SliceDepth = RenderVolumetricCloudParameters.EmptySpaceSkippingSliceDepth;
 	ResolvedView = ResolveView();
-	const float SliceDepth = StartTracingSampleVolumeDepth * (1.0f / THREADGROUP_SIZEZ);
 
 	float2 PixelUV = float2(DispatchThreadId.x + 0.5, DispatchThreadId.y + 0.5) / StartTracingDistanceTextureResolution;
 	float4 SvPosition = float4(PixelUV * ResolvedView.ViewSizeAndInvSize.xy, 0.1, 1);
@@ -1478,10 +1487,7 @@ void MainCS(uint3 DispatchThreadId : SV_DispatchThreadID, uint3 GroupThreadId :
 	const float SampleDistanceFromCamera = (float(GroupThreadId.z) + 0.5f) * SliceDepth;
 	const float3 Raydir =-MaterialParameters.CameraVector;
 	UpdateMaterialCloudParam(MaterialParameters, LWCAdd(ResolvedView.WorldCameraOrigin, LWCPromote(Raydir * SampleDistanceFromCamera)),
-		ResolvedView, CloudLayerParams, SampleDistanceFromCamera);
-
-	MaterialParameters.CloudEmptySpaceSkippingSphereCenterWorldPosition = LWCHackToFloat(MaterialParameters.AbsoluteWorldPosition); // This is the same as the WorldPOsitionNode but float3
-	MaterialParameters.CloudEmptySpaceSkippingSphereRadius = 0.5f * SliceDepth * 1.73f;	// 1.73 = sqrt(1*1+1*1+1*1), the lenght of the diagonal of a cube as compared to its unit edge length
+		ResolvedView, CloudLayerParams, TRACING_SHADOW_DISTANCE_OFF, RenderVolumetricCloudParameters.EmptySpaceSkippingSliceDepth);
 
 	float CloudEmptySpaceSkippingOutput0 = GetVolumetricCloudEmptySpaceSkippingOutput0(MaterialParameters);
 	GroupData[GroupThreadId.x][GroupThreadId.y][GroupThreadId.z].FrontDepthKm = CloudEmptySpaceSkippingOutput0 > 0.0f ? (float(GroupThreadId.z) + 0.0f) * SliceDepth * CENTIMETER_TO_KILOMETER : GetDefaultDepth();
@@ -1706,7 +1712,7 @@ void MainPS(
 		{
 			const float SampleT = ShadowLengthTest * (ShadowT * InvShadowStepCount);
 			UpdateMaterialCloudParam(MaterialParameters, LWCPromote(WorldPosOnLayer + LightDirection * SampleT),			// LWC TODO FIX ME
-				ResolvedView, CloudLayerParams, SampleT);
+				ResolvedView, CloudLayerParams, SampleT, SPACE_SKIPPING_SLICE_DEPTH_OFF);
 
 		#if MATERIAL_VOLUMETRIC_ADVANCED_CONSERVATIVE_DENSITY
 			if (MaterialParameters.VolumeSampleConservativeDensity.x <= 0.0f)
diff --git a/Engine/Shaders/Private/VolumetricCloudMaterialPixelCommon.ush b/Engine/Shaders/Private/VolumetricCloudMaterialPixelCommon.ush
--- a/Engine/Shaders/Private/VolumetricCloudMaterialPixelCommon.ush
+++ b/Engine/Shaders/Private/VolumetricCloudMaterialPixelCommon.ush
@@ -39,10 +39,12 @@ FCloudLayerParameters GetCloudLayerParams(
 
 #define TRACING_SHADOW_DISTANCE_OFF	0.0f
 
+#define SPACE_SKIPPING_SLICE_DEPTH_OFF 0.0f
+
 // Function to update the material parameter structure for an evaluation
 void UpdateMaterialCloudParam(inout FMaterialPixelParameters MaterialParameters,
 	FLWCVector3 AbsoluteWorldPosition, in ViewState InputView, in FCloudLayerParameters CloudLayerParams,
-	in float ShadowSampleDistance)
+	in float ShadowSampleDistance, in float EmptySpaceSkippingSliceDepth)
 {
 	MaterialParameters.AbsoluteWorldPosition = AbsoluteWorldPosition;
 	MaterialParameters.WorldPosition_CamRelative = LWCToFloat(LWCAdd(MaterialParameters.AbsoluteWorldPosition, InputView.PreViewTranslation));
@@ -61,8 +63,8 @@ void UpdateMaterialCloudParam(inout FMaterialPixelParameters MaterialParameters,
 	MaterialParameters.VolumeSampleConservativeDensity = GetVolumetricAdvancedMaterialOutput6(MaterialParameters);	// Evaluate conservative density
 #endif
 
-	MaterialParameters.CloudEmptySpaceSkippingSphereCenterWorldPosition = 0.0f;
-	MaterialParameters.CloudEmptySpaceSkippingSphereRadius = 0.0f;
+	MaterialParameters.CloudEmptySpaceSkippingSphereCenterWorldPosition = LWCHackToFloat(MaterialParameters.AbsoluteWorldPosition); // This is the same as the WorldPosition node but float3
+	MaterialParameters.CloudEmptySpaceSkippingSphereRadius = 0.5f * EmptySpaceSkippingSliceDepth * 1.73f;			// 1.73 = sqrt(1*1+1*1+1*1), the lenght of the diagonal of a cube as compared to its unit edge length
 
 	MaterialParameters.ShadowSampleDistance = ShadowSampleDistance;
 }
diff --git a/Engine/Shaders/Private/VolumetricFogVoxelization.usf b/Engine/Shaders/Private/VolumetricFogVoxelization.usf
--- a/Engine/Shaders/Private/VolumetricFogVoxelization.usf
+++ b/Engine/Shaders/Private/VolumetricFogVoxelization.usf
@@ -310,9 +310,8 @@ void VoxelizePS(
 	FCloudLayerParameters CloudLayerParams = GetCloudLayerParams(
 		VoxelizeVolumePass.RenderVolumetricCloudParametersCloudLayerCenterKm, VoxelizeVolumePass.RenderVolumetricCloudParametersPlanetRadiusKm,
 		VoxelizeVolumePass.RenderVolumetricCloudParametersBottomRadiusKm, VoxelizeVolumePass.RenderVolumetricCloudParametersTopRadiusKm);
-	const float NotShadowSample = 0.0f;
 	FLWCVector3 AbsoluteWorldPosition = SvPositionToWorld(SvPosition);
-	UpdateMaterialCloudParam(MaterialParameters, AbsoluteWorldPosition, ResolvedView, CloudLayerParams, NotShadowSample);
+	UpdateMaterialCloudParam(MaterialParameters, AbsoluteWorldPosition, ResolvedView, CloudLayerParams, TRACING_SHADOW_DISTANCE_OFF, SPACE_SKIPPING_SLICE_DEPTH_OFF);
 #endif
 
 	FPixelMaterialInputs PixelMaterialInputs;
diff --git a/Engine/Source/Runtime/Engine/Public/ShowFlags.h b/Engine/Source/Runtime/Engine/Public/ShowFlags.h
--- a/Engine/Source/Runtime/Engine/Public/ShowFlags.h
+++ b/Engine/Source/Runtime/Engine/Public/ShowFlags.h
@@ -474,6 +474,7 @@ private:
 		SetVisualizePostProcessStack(false);
 		SetVirtualTexturePrimitives(false);
 		SetVisualizeVolumetricCloudConservativeDensity(false);
+		SetVisualizeVolumetricCloudEmptySpaceSkipping(false);
 		SetDrawOnlyVSMInvalidatingGeo(false);
 		SetDebugDrawDistantVirtualSMLights(false);
 
diff --git a/Engine/Source/Runtime/Engine/Public/ShowFlagsValues.inl b/Engine/Source/Runtime/Engine/Public/ShowFlagsValues.inl
--- a/Engine/Source/Runtime/Engine/Public/ShowFlagsValues.inl
+++ b/Engine/Source/Runtime/Engine/Public/ShowFlagsValues.inl
@@ -398,6 +398,9 @@ SHOWFLAG_FIXED_IN_SHIPPING(0, VirtualTexturePrimitives, SFG_Developer, NSLOCTEXT
 /** Visualize volumetric cloud conservative density. */
 SHOWFLAG_FIXED_IN_SHIPPING(0, VisualizeVolumetricCloudConservativeDensity, SFG_Visualize, NSLOCTEXT("UnrealEd", "VisualizeVolumetricCloudConservativeDensitySF", "Volumetric Cloud Conservative Density"))
 
+/** Visualize volumetric cloud density for empty space skipping. */
+SHOWFLAG_FIXED_IN_SHIPPING(0, VisualizeVolumetricCloudEmptySpaceSkipping, SFG_Visualize, NSLOCTEXT("UnrealEd", "VisualizeVolumetricCloudEmptySpaceSkippingSF", "Volumetric Cloud Empty Space Skipping Density"))
+
 SHOWFLAG_FIXED_IN_SHIPPING(1, VirtualShadowMapCaching, SFG_Developer, NSLOCTEXT("UnrealEd", "VirtualShadowMapCaching", "Cache Virtual Shadow Maps"))
 
 SHOWFLAG_FIXED_IN_SHIPPING(0, DrawOnlyVSMInvalidatingGeo, SFG_Visualize, NSLOCTEXT("UnrealEd", "DrawOnlyVSMInvalidatingGeo", "Draw Only Geometry Causing VSM Invalidation"))
diff --git a/Engine/Source/Runtime/Renderer/Private/DeferredShadingRenderer.cpp b/Engine/Source/Runtime/Renderer/Private/DeferredShadingRenderer.cpp
--- a/Engine/Source/Runtime/Renderer/Private/DeferredShadingRenderer.cpp
+++ b/Engine/Source/Runtime/Renderer/Private/DeferredShadingRenderer.cpp
@@ -2980,8 +2980,8 @@ void FDeferredShadingSceneRenderer::Render(FRDGBuilder& GraphBuilder)
 	const bool bShouldRenderSkyAtmosphere = ShouldRenderSkyAtmosphere(Scene, ViewFamily.EngineShowFlags);
 	const ESkyAtmospherePassLocation SkyAtmospherePassLocation = GetSkyAtmospherePassLocation();
 	const bool bShouldRenderVolumetricCloudBase = ShouldRenderVolumetricCloud(Scene, ViewFamily.EngineShowFlags);
-	const bool bShouldRenderVolumetricCloud = bShouldRenderVolumetricCloudBase && !ViewFamily.EngineShowFlags.VisualizeVolumetricCloudConservativeDensity;
-	const bool bShouldVisualizeVolumetricCloud = bShouldRenderVolumetricCloudBase && !!ViewFamily.EngineShowFlags.VisualizeVolumetricCloudConservativeDensity;
+	const bool bShouldRenderVolumetricCloud = bShouldRenderVolumetricCloudBase && (!ViewFamily.EngineShowFlags.VisualizeVolumetricCloudConservativeDensity && !ViewFamily.EngineShowFlags.VisualizeVolumetricCloudEmptySpaceSkipping);
+	const bool bShouldVisualizeVolumetricCloud = bShouldRenderVolumetricCloudBase && (!!ViewFamily.EngineShowFlags.VisualizeVolumetricCloudConservativeDensity || !!ViewFamily.EngineShowFlags.VisualizeVolumetricCloudEmptySpaceSkipping);
 	bool bAsyncComputeVolumetricCloud = IsVolumetricRenderTargetEnabled() && IsVolumetricRenderTargetAsyncCompute();
 	bool bHasHalfResCheckerboardMinMaxDepth = false;
 	bool bVolumetricRenderTargetRequired = bShouldRenderVolumetricCloud && !bHasRayTracedOverlay;
diff --git a/Engine/Source/Runtime/Renderer/Private/PostProcess/PostProcessEyeAdaptation.cpp b/Engine/Source/Runtime/Renderer/Private/PostProcess/PostProcessEyeAdaptation.cpp
--- a/Engine/Source/Runtime/Renderer/Private/PostProcess/PostProcessEyeAdaptation.cpp
+++ b/Engine/Source/Runtime/Renderer/Private/PostProcess/PostProcessEyeAdaptation.cpp
@@ -241,6 +241,7 @@ bool IsAutoExposureDebugMode(const FViewInfo& View)
 		EngineShowFlags.RayTracingDebug ||
 		EngineShowFlags.VisualizeDistanceFieldAO ||
 		EngineShowFlags.VisualizeVolumetricCloudConservativeDensity ||
+		EngineShowFlags.VisualizeVolumetricCloudEmptySpaceSkipping ||
 		EngineShowFlags.CollisionVisibility ||
 		EngineShowFlags.CollisionPawn ||
 		!EngineShowFlags.PostProcessing;
diff --git a/Engine/Source/Runtime/Renderer/Private/PostProcess/PostProcessing.cpp b/Engine/Source/Runtime/Renderer/Private/PostProcess/PostProcessing.cpp
--- a/Engine/Source/Runtime/Renderer/Private/PostProcess/PostProcessing.cpp
+++ b/Engine/Source/Runtime/Renderer/Private/PostProcess/PostProcessing.cpp
@@ -161,6 +161,7 @@ bool IsPostProcessingEnabled(const FViewInfo& View)
 			!View.Family->EngineShowFlags.VisualizeDistanceFieldAO &&
 			!View.Family->EngineShowFlags.VisualizeShadingModels &&
 			!View.Family->EngineShowFlags.VisualizeVolumetricCloudConservativeDensity &&
+			!View.Family->EngineShowFlags.VisualizeVolumetricCloudEmptySpaceSkipping &&
 			!View.Family->EngineShowFlags.ShaderComplexity;
 	}
 	else
diff --git a/Engine/Source/Runtime/Renderer/Private/VolumetricCloudRendering.cpp b/Engine/Source/Runtime/Renderer/Private/VolumetricCloudRendering.cpp
--- a/Engine/Source/Runtime/Renderer/Private/VolumetricCloudRendering.cpp
+++ b/Engine/Source/Runtime/Renderer/Private/VolumetricCloudRendering.cpp
@@ -221,12 +221,6 @@ static TAutoConsoleVariable<int32> CVarVolumetricCloudEmptySpaceSkippingVolumeDe
 
 ////////////////////////////////////////////////////////////////////////// 
 
-static TAutoConsoleVariable<int32> CVarVolumetricCloudDebugSampleCountMode(
-	TEXT("r.VolumetricCloud.Debug.SampleCountMode"), 0,
-	TEXT("Only for developers. [0] Disabled [1] Primary material sample count [2] Advanced:raymarched shadow sample count [3] Shadow material sample count [4] Advanced:ground shadow sample count [5] Advanced:ground shadow material sample count"));
-
-////////////////////////////////////////////////////////////////////////// 
-
 
 static bool ShouldUseComputeForCloudTracing(const FStaticFeatureLevel InFeatureLevel)
 {
@@ -270,7 +264,7 @@ bool ShouldRenderVolumetricCloud(const FScene* Scene, const FEngineShowFlags& En
 
 bool ShouldViewVisualizeVolumetricCloudConservativeDensity(const FViewInfo& ViewInfo, const FEngineShowFlags& EngineShowFlags)
 {
-	return !!EngineShowFlags.VisualizeVolumetricCloudConservativeDensity
+	return (!!EngineShowFlags.VisualizeVolumetricCloudConservativeDensity || !!EngineShowFlags.VisualizeVolumetricCloudEmptySpaceSkipping)
 		&& !ViewInfo.bIsReflectionCapture
 		&& !ViewInfo.bIsSceneCapture
 		&& GIsEditor;
@@ -549,6 +543,7 @@ BEGIN_GLOBAL_SHADER_PARAMETER_STRUCT(FRenderVolumetricCloudGlobalParameters, )
 	SHADER_PARAMETER(uint32, NoiseFrameIndexModPattern)
 	SHADER_PARAMETER(int32, OpaqueIntersectionMode)
 	SHADER_PARAMETER(uint32, VolumetricRenderTargetMode)
+	SHADER_PARAMETER(uint32, CloudDebugViewMode)
 	SHADER_PARAMETER(uint32, IsReflectionRendering)
 	SHADER_PARAMETER(uint32, HasValidHZB)
 	SHADER_PARAMETER(uint32, ClampRayTToDepthBufferPostHZB)
@@ -564,6 +559,7 @@ BEGIN_GLOBAL_SHADER_PARAMETER_STRUCT(FRenderVolumetricCloudGlobalParameters, )
 	SHADER_PARAMETER(int32, EnableDistantSkyLightSampling)
 	SHADER_PARAMETER(int32, EnableAtmosphericLightsSampling)
 	SHADER_PARAMETER(int32, EnableHeightFog)
+	SHADER_PARAMETER(float, EmptySpaceSkippingSliceDepth)
 	SHADER_PARAMETER_STRUCT_INCLUDE(FFogUniformParameters, FogStruct)
 	SHADER_PARAMETER_STRUCT(FForwardLightData, ForwardLightData)
 END_GLOBAL_SHADER_PARAMETER_STRUCT()
@@ -572,6 +568,21 @@ IMPLEMENT_STATIC_UNIFORM_BUFFER_STRUCT(FRenderVolumetricCloudGlobalParameters, "
 
 IMPLEMENT_GLOBAL_SHADER_PARAMETER_STRUCT(FVolumetricCloudCommonGlobalShaderParameters, "VolumetricCloudCommonParameters");
 
+uint32 GetVolumetricCloudDebugViewMode(const FEngineShowFlags& ShowFlags)
+{
+	if (ShowFlags.VisualizeVolumetricCloudConservativeDensity)
+	{
+		return 1;
+	}
+	if (ShowFlags.VisualizeVolumetricCloudEmptySpaceSkipping)
+	{
+		return 2;
+	}
+	return 0;
+}
+
+static float GetEmptySpaceSkippingSliceDepth();
+
 // When calling this, you still need to setup Light0Shadow yourself.
 void SetupDefaultRenderVolumetricCloudGlobalParameters(FRDGBuilder& GraphBuilder, FRenderVolumetricCloudGlobalParameters& VolumetricCloudParams, FVolumetricCloudRenderSceneInfo& CloudInfo, FViewInfo& ViewInfo)
 {
@@ -587,6 +598,7 @@ void SetupDefaultRenderVolumetricCloudGlobalParameters(FRDGBuilder& GraphBuilder
 	VolumetricCloudParams.TracingCoordToFullResPixelCoordScaleBias = FUintVector4(1, 1, 0, 0);
 	VolumetricCloudParams.NoiseFrameIndexModPattern = 0;
 	VolumetricCloudParams.VolumetricRenderTargetMode = ViewInfo.ViewState ? ViewInfo.ViewState->VolumetricCloudRenderTarget.GetMode() : 0;
+	VolumetricCloudParams.CloudDebugViewMode = GetVolumetricCloudDebugViewMode(ViewInfo.Family->EngineShowFlags);
 
 	VolumetricCloudParams.HasValidHZB = 0;
 	VolumetricCloudParams.ClampRayTToDepthBufferPostHZB = 0;
@@ -595,6 +607,8 @@ void SetupDefaultRenderVolumetricCloudGlobalParameters(FRDGBuilder& GraphBuilder
 
 	VolumetricCloudParams.StepSizeOnZeroConservativeDensity = FMath::Max(CVarVolumetricCloudStepSizeOnZeroConservativeDensity.GetValueOnRenderThread(), 1);
 
+	VolumetricCloudParams.EmptySpaceSkippingSliceDepth = GetEmptySpaceSkippingSliceDepth();
+
 	VolumetricCloudParams.EnableHeightFog = ViewInfo.Family->Scene->HasAnyExponentialHeightFog() && ShouldRenderFog(*ViewInfo.Family);
 	SetupFogUniformParameters(GraphBuilder, ViewInfo, VolumetricCloudParams.FogStruct);
 
@@ -689,7 +703,7 @@ BEGIN_SHADER_PARAMETER_STRUCT(FRenderVolumetricCloudRenderViewParametersPS, )
 	RENDER_TARGET_BINDING_SLOTS()
 END_SHADER_PARAMETER_STRUCT()
 
-template<EVolumetricCloudRenderViewPsPermutations Permutation, bool bConvervativeDensityDebug>
+template<EVolumetricCloudRenderViewPsPermutations Permutation>
 class FRenderVolumetricCloudRenderViewPs : public FMeshMaterialShader
 {
 	DECLARE_SHADER_TYPE(FRenderVolumetricCloudRenderViewPs, MeshMaterial);
@@ -707,13 +721,6 @@ public:
 	static bool ShouldCompilePermutation(const FMeshMaterialShaderPermutationParameters& Parameters)
 	{
 		const bool bIsCompatible = IsMaterialCompatibleWithVolumetricCloud(Parameters.MaterialParameters, Parameters.Platform);
-		
-		if (bConvervativeDensityDebug)
-		{
-			return bIsCompatible
-				&& EnumHasAllFlags(Parameters.Flags, EShaderPermutationFlags::HasEditorOnlyData)
-				&& Permutation == CloudPs_AtmoTrans0_SampleShad0_2ndLight0;
-		}
 
 		return bIsCompatible;
 	}
@@ -745,8 +752,6 @@ public:
 										Permutation == CloudPs_AtmoTrans1_SampleShad1_2ndLight1 || Permutation == CloudPs_AtmoTrans1_SampleShad0_2ndLight1;
 		OutEnvironment.SetDefine(TEXT("CLOUD_SAMPLE_SECOND_LIGHT"), bSampleSecondLight ? TEXT("1") : TEXT("0"));
 
-		OutEnvironment.SetDefine(TEXT("CLOUD_CONSERVATIVE_DENSITY_DEBUG"), bConvervativeDensityDebug);
-
 		// This shader takes a very long time to compile with FXC, so we pre-compile it with DXC first and then forward the optimized HLSL to FXC.
 		OutEnvironment.CompilerFlags.Add(CFLAG_PrecompileWithDXC);
 	}
@@ -754,19 +759,18 @@ public:
 private:
 };
 
-#define IMPLEMENT_CLOUD_RENDERVIEW_PS(A, B, C, D) \
-	typedef FRenderVolumetricCloudRenderViewPs<CloudPs_AtmoTrans##A##_SampleShad##B##_2ndLight##C, D> FCloudRenderViewPS##A##B##C##D; \
-	IMPLEMENT_MATERIAL_SHADER_TYPE(template<>, FCloudRenderViewPS##A##B##C##D, TEXT("/Engine/Private/VolumetricCloud.usf"), TEXT("MainPS"), SF_Pixel)
+#define IMPLEMENT_CLOUD_RENDERVIEW_PS(A, B, C) \
+	typedef FRenderVolumetricCloudRenderViewPs<CloudPs_AtmoTrans##A##_SampleShad##B##_2ndLight##C> FCloudRenderViewPS##A##B##C; \
+	IMPLEMENT_MATERIAL_SHADER_TYPE(template<>, FCloudRenderViewPS##A##B##C, TEXT("/Engine/Private/VolumetricCloud.usf"), TEXT("MainPS"), SF_Pixel)
 
-IMPLEMENT_CLOUD_RENDERVIEW_PS(0, 0, 0, false);
-IMPLEMENT_CLOUD_RENDERVIEW_PS(1, 0, 0, false);
-IMPLEMENT_CLOUD_RENDERVIEW_PS(0, 1, 0, false);
-IMPLEMENT_CLOUD_RENDERVIEW_PS(1, 1, 0, false);
-IMPLEMENT_CLOUD_RENDERVIEW_PS(0, 0, 1, false);
-IMPLEMENT_CLOUD_RENDERVIEW_PS(1, 0, 1, false);
-IMPLEMENT_CLOUD_RENDERVIEW_PS(0, 1, 1, false);
-IMPLEMENT_CLOUD_RENDERVIEW_PS(1, 1, 1, false);
-IMPLEMENT_CLOUD_RENDERVIEW_PS(0, 0, 0, true);
+IMPLEMENT_CLOUD_RENDERVIEW_PS(0, 0, 0);
+IMPLEMENT_CLOUD_RENDERVIEW_PS(1, 0, 0);
+IMPLEMENT_CLOUD_RENDERVIEW_PS(0, 1, 0);
+IMPLEMENT_CLOUD_RENDERVIEW_PS(1, 1, 0);
+IMPLEMENT_CLOUD_RENDERVIEW_PS(0, 0, 1);
+IMPLEMENT_CLOUD_RENDERVIEW_PS(1, 0, 1);
+IMPLEMENT_CLOUD_RENDERVIEW_PS(0, 1, 1);
+IMPLEMENT_CLOUD_RENDERVIEW_PS(1, 1, 1);
 
 #undef IMPLEMENT_CLOUD_RENDERVIEW_PS
 
@@ -781,12 +785,14 @@ class FRenderVolumetricCloudRenderViewCS : public FMeshMaterialShader
 	class FCloudSampleAtmosphericLightShadowmap : SHADER_PERMUTATION_BOOL("CLOUD_SAMPLE_ATMOSPHERIC_LIGHT_SHADOWMAP");
 	class FCloudSampleSecondLight : SHADER_PERMUTATION_BOOL("CLOUD_SAMPLE_SECOND_LIGHT");
 	class FCloudSampleLocalLights : SHADER_PERMUTATION_BOOL("CLOUD_SAMPLE_LOCAL_LIGHTS");
+	class FCloudDebugViewMode : SHADER_PERMUTATION_BOOL("CLOUD_DEBUG_VIEW_MODE");	// Only the CS feature this because CS is typically used in editor.
 
 	using FPermutationDomain = TShaderPermutationDomain<
 		FCloudPerSampleAtmosphereTransmittance,
 		FCloudSampleAtmosphericLightShadowmap,
 		FCloudSampleSecondLight,
-		FCloudSampleLocalLights>;
+		FCloudSampleLocalLights,
+		FCloudDebugViewMode>;
 
 	BEGIN_SHADER_PARAMETER_STRUCT(FParameters, )
 		SHADER_PARAMETER(FVector4f, OutputViewRect)
@@ -814,6 +820,17 @@ class FRenderVolumetricCloudRenderViewCS : public FMeshMaterialShader
 		const FPermutationDomain PermutationVector(Parameters.PermutationId);
 		const bool bIsCompatible = IsMaterialCompatibleWithVolumetricCloud(Parameters.MaterialParameters, Parameters.Platform);
 
+		if (PermutationVector.Get<FCloudDebugViewMode>())
+		{
+			// Only compile visualization shaders for editor builds
+			return bIsCompatible
+				&& EnumHasAllFlags(Parameters.Flags, EShaderPermutationFlags::HasEditorOnlyData)
+				&& !PermutationVector.Get<FCloudPerSampleAtmosphereTransmittance>()
+				&& !PermutationVector.Get<FCloudSampleAtmosphericLightShadowmap>()
+				&& !PermutationVector.Get<FCloudSampleSecondLight>()
+				&& !PermutationVector.Get<FCloudSampleLocalLights>();
+		}
+
 		return bIsCompatible;
 	}
 
@@ -915,6 +932,19 @@ class FRenderVolumetricCloudEmptySpaceSkippingCS : public FMeshMaterialShader
 	}
 };
 
+static float GetEmptySpaceSkippingVolumeDepth()
+{
+	const float KilometersToCentimeters = 100000.0f;
+	return FMath::Max(1.0f, CVarVolumetricCloudEmptySpaceSkippingVolumeDepth.GetValueOnRenderThread()) * KilometersToCentimeters;
+}
+
+static float GetEmptySpaceSkippingSliceDepth()
+{
+	const float KilometersToCentimeters = 100000.0f;
+	float EmptySpaceSkippingVolumeDepth = GetEmptySpaceSkippingVolumeDepth();
+	return EmptySpaceSkippingVolumeDepth * (1.0f / FRenderVolumetricCloudEmptySpaceSkippingCS::ThreadGroupSizeZ);
+}
+
 IMPLEMENT_MATERIAL_SHADER_TYPE(, FRenderVolumetricCloudEmptySpaceSkippingCS, TEXT("/Engine/Private/VolumetricCloud.usf"), TEXT("MainCS"), SF_Compute);
 
 //////////////////////////////////////////////////////////////////////////
@@ -1021,16 +1051,16 @@ class FVolumetricCloudRenderViewMeshProcessor : public FMeshPassProcessor
 public:
 	FVolumetricCloudRenderViewMeshProcessor(const FScene* Scene, ERHIFeatureLevel::Type FeatureLevel, const FViewInfo* InViewIfDynamicMeshCommand,
 		bool bShouldViewRenderVolumetricRenderTarget, bool bSkipAtmosphericLightShadowmap, bool bSecondAtmosphereLightEnabled,
-		bool bInVisualizeConservativeDensity, FMeshPassDrawListContext* InDrawListContext)
+		bool bCloudDebugViewModeEnabledIn, FMeshPassDrawListContext* InDrawListContext)
 		: FMeshPassProcessor(EMeshPass::Num, Scene, FeatureLevel, InViewIfDynamicMeshCommand, InDrawListContext)
 		, bVolumetricCloudPerSampleAtmosphereTransmittance(ShouldUsePerSampleAtmosphereTransmittance(Scene, InViewIfDynamicMeshCommand))
 		, bVolumetricCloudSampleLightShadowmap(!bSkipAtmosphericLightShadowmap && CVarVolumetricCloudShadowSampleAtmosphericLightShadowmap.GetValueOnAnyThread() > 0)
 		, bVolumetricCloudSecondLight(bSecondAtmosphereLightEnabled)
-		, bVisualizeConservativeDensity(bInVisualizeConservativeDensity)
+		, bCloudDebugViewModeEnabled(bCloudDebugViewModeEnabledIn)
 	{
 		PassDrawRenderState.SetDepthStencilState(TStaticDepthStencilState<false, CF_Always>::GetRHI());
 
-		if (bShouldViewRenderVolumetricRenderTarget || bVisualizeConservativeDensity)
+		if (bShouldViewRenderVolumetricRenderTarget || bCloudDebugViewModeEnabled)
 		{
 			// No blending as we only render clouds in that render target today. Avoids clearing for now.
 			PassDrawRenderState.SetBlendState(TStaticBlendState<>::GetRHI());
@@ -1078,23 +1108,18 @@ public:
 		const ERasterizerFillMode MeshFillMode = FM_Solid;
 		const ERasterizerCullMode MeshCullMode = CM_None;
 
-		if (bVisualizeConservativeDensity)
-		{
-			return TemplatedProcess<FRenderVolumetricCloudRenderViewPs<CloudPs_AtmoTrans0_SampleShad0_2ndLight0, true>>(
-				MeshBatch, BatchElementMask, PrimitiveSceneProxy, MaterialRenderProxy, Material, StaticMeshId, MeshFillMode, MeshCullMode);
-		}
-		else if (bVolumetricCloudSecondLight)
+		if (bVolumetricCloudSecondLight)
 		{
 			if (bVolumetricCloudSampleLightShadowmap)
 			{
 				if (bVolumetricCloudPerSampleAtmosphereTransmittance)
 				{
-					return TemplatedProcess<FRenderVolumetricCloudRenderViewPs<CloudPs_AtmoTrans1_SampleShad1_2ndLight1, false>>(
+					return TemplatedProcess<FRenderVolumetricCloudRenderViewPs<CloudPs_AtmoTrans1_SampleShad1_2ndLight1>>(
 						MeshBatch, BatchElementMask, PrimitiveSceneProxy, MaterialRenderProxy, Material, StaticMeshId, MeshFillMode, MeshCullMode);
 				}
 				else
 				{
-					return TemplatedProcess<FRenderVolumetricCloudRenderViewPs<CloudPs_AtmoTrans0_SampleShad1_2ndLight1, false>>(
+					return TemplatedProcess<FRenderVolumetricCloudRenderViewPs<CloudPs_AtmoTrans0_SampleShad1_2ndLight1>>(
 						MeshBatch, BatchElementMask, PrimitiveSceneProxy, MaterialRenderProxy, Material, StaticMeshId, MeshFillMode, MeshCullMode);
 				}
 			}
@@ -1102,12 +1127,12 @@ public:
 			{
 				if (bVolumetricCloudPerSampleAtmosphereTransmittance)
 				{
-					return TemplatedProcess<FRenderVolumetricCloudRenderViewPs<CloudPs_AtmoTrans1_SampleShad0_2ndLight1, false>>(
+					return TemplatedProcess<FRenderVolumetricCloudRenderViewPs<CloudPs_AtmoTrans1_SampleShad0_2ndLight1>>(
 						MeshBatch, BatchElementMask, PrimitiveSceneProxy, MaterialRenderProxy, Material, StaticMeshId, MeshFillMode, MeshCullMode);
 				}
 				else
 				{
-					return TemplatedProcess<FRenderVolumetricCloudRenderViewPs<CloudPs_AtmoTrans0_SampleShad0_2ndLight1, false>>(
+					return TemplatedProcess<FRenderVolumetricCloudRenderViewPs<CloudPs_AtmoTrans0_SampleShad0_2ndLight1>>(
 						MeshBatch, BatchElementMask, PrimitiveSceneProxy, MaterialRenderProxy, Material, StaticMeshId, MeshFillMode, MeshCullMode);
 				}
 			}
@@ -1118,12 +1143,12 @@ public:
 			{
 				if (bVolumetricCloudPerSampleAtmosphereTransmittance)
 				{
-					return TemplatedProcess<FRenderVolumetricCloudRenderViewPs<CloudPs_AtmoTrans1_SampleShad1_2ndLight0, false>>(
+					return TemplatedProcess<FRenderVolumetricCloudRenderViewPs<CloudPs_AtmoTrans1_SampleShad1_2ndLight0>>(
 						MeshBatch, BatchElementMask, PrimitiveSceneProxy, MaterialRenderProxy, Material, StaticMeshId, MeshFillMode, MeshCullMode);
 				}
 				else
 				{
-					return TemplatedProcess<FRenderVolumetricCloudRenderViewPs<CloudPs_AtmoTrans0_SampleShad1_2ndLight0, false>>(
+					return TemplatedProcess<FRenderVolumetricCloudRenderViewPs<CloudPs_AtmoTrans0_SampleShad1_2ndLight0>>(
 						MeshBatch, BatchElementMask, PrimitiveSceneProxy, MaterialRenderProxy, Material, StaticMeshId, MeshFillMode, MeshCullMode);
 				}
 			}
@@ -1131,12 +1156,12 @@ public:
 			{
 				if (bVolumetricCloudPerSampleAtmosphereTransmittance)
 				{
-					return TemplatedProcess<FRenderVolumetricCloudRenderViewPs<CloudPs_AtmoTrans1_SampleShad0_2ndLight0, false>>(
+					return TemplatedProcess<FRenderVolumetricCloudRenderViewPs<CloudPs_AtmoTrans1_SampleShad0_2ndLight0>>(
 						MeshBatch, BatchElementMask, PrimitiveSceneProxy, MaterialRenderProxy, Material, StaticMeshId, MeshFillMode, MeshCullMode);
 				}
 				else
 				{
-					return TemplatedProcess<FRenderVolumetricCloudRenderViewPs<CloudPs_AtmoTrans0_SampleShad0_2ndLight0, false>>(
+					return TemplatedProcess<FRenderVolumetricCloudRenderViewPs<CloudPs_AtmoTrans0_SampleShad0_2ndLight0>>(
 						MeshBatch, BatchElementMask, PrimitiveSceneProxy, MaterialRenderProxy, Material, StaticMeshId, MeshFillMode, MeshCullMode);
 				}
 			}
@@ -1199,7 +1224,7 @@ private:
 	bool bVolumetricCloudPerSampleAtmosphereTransmittance;
 	bool bVolumetricCloudSampleLightShadowmap;
 	bool bVolumetricCloudSecondLight;
-	bool bVisualizeConservativeDensity;
+	bool bCloudDebugViewModeEnabled;
 };
 
 
@@ -2004,7 +2029,7 @@ FCloudRenderContext::FCloudRenderContext()
 	bSkipAerialPerspective = false;
 
 	bAsyncCompute = false;
-	bVisualizeConservativeDensity = false;
+	bCloudDebugViewModeEnabled = false;
 }
 
 void FCloudRenderContext::CreateDefaultTexturesIfNeeded(FRDGBuilder& GraphBuilder)
@@ -2196,16 +2221,14 @@ void FSceneRenderer::RenderVolumetricCloudsInternal(FRDGBuilder& GraphBuilder, F
 	const bool bShouldViewRenderVolumetricRenderTarget = CloudRC.bShouldViewRenderVolumetricRenderTarget;
 	const bool bSkipAtmosphericLightShadowmap = CloudRC.bSkipAtmosphericLightShadowmap;
 	const bool bSecondAtmosphereLightEnabled = CloudRC.bSecondAtmosphereLightEnabled;
-	const bool bVisualizeConservativeDensity = CloudRC.bVisualizeConservativeDensity;
+	const bool bCloudDebugViewModeEnabled = CloudRC.bCloudDebugViewModeEnabled;
 	const FRDGTextureDesc& Desc = CloudRC.RenderTargets.Output[0].GetTexture()->Desc;
 	FViewInfo& MainView = *CloudRC.MainView;
 	TUniformBufferRef<FViewUniformShaderParameters> ViewUniformBuffer = CloudRC.ViewUniformBuffer;
 	TRDGUniformBufferRef<FRenderVolumetricCloudGlobalParameters> CloudPassUniformBuffer = CreateCloudPassUniformBuffer(GraphBuilder, CloudRC);
 
 	const bool bCloudEnableLocalLightSampling = CVarVolumetricCloudEnableLocalLightsSampling.GetValueOnRenderThread() > 0;
-	// Debug permutations are only generated for pixel shader
-	const bool bShouldUseComputeForCloudTracing = !bVisualizeConservativeDensity && ShouldUseComputeForCloudTracing(Scene->GetFeatureLevel());
-	if (bShouldUseComputeForCloudTracing)
+	if (ShouldUseComputeForCloudTracing(Scene->GetFeatureLevel()))
 	{
 		//
 		// Cloud empty space skipping
@@ -2225,8 +2248,7 @@ void FSceneRenderer::RenderVolumetricCloudsInternal(FRDGBuilder& GraphBuilder, F
 			PassParameters->View = ViewUniformBuffer;
 			PassParameters->OutStartTracingDistanceTexture = GraphBuilder.CreateUAV(FRDGTextureUAVDesc(EmptySpaceSkippingTexture));
 			PassParameters->StartTracingDistanceTextureResolution = FVector2f(EmptySpaceSkippingTextureResolution.X, EmptySpaceSkippingTextureResolution.Y);
-			const float KilometersToCentimeters = 100000.0f;
-			PassParameters->StartTracingSampleVolumeDepth = FMath::Max(1.0f, CVarVolumetricCloudEmptySpaceSkippingVolumeDepth.GetValueOnRenderThread()) * KilometersToCentimeters;
+			PassParameters->StartTracingSampleVolumeDepth = GetEmptySpaceSkippingVolumeDepth();
 
 			const bool bCloudEmptySpaceSkippingDebug = false;
 			//if (bCloudEmptySpaceSkippingDebug)
@@ -2299,11 +2321,11 @@ void FSceneRenderer::RenderVolumetricCloudsInternal(FRDGBuilder& GraphBuilder, F
 		FRDGTextureUAVRef CloudDepthTextureUAV;
 		GetOutputTexturesWithFallback(GraphBuilder, CloudRC, CloudColorCubeTextureUAV, CloudColorTextureUAV, CloudDepthTextureUAV);
 
-		bool bSampleVirtualShadowMap = (!bSkipAtmosphericLightShadowmap && CVarVolumetricCloudShadowSampleAtmosphericLightShadowmap.GetValueOnRenderThread() > 0);
+		bool bSampleVirtualShadowMap = (!bCloudDebugViewModeEnabled && !bSkipAtmosphericLightShadowmap && CVarVolumetricCloudShadowSampleAtmosphericLightShadowmap.GetValueOnRenderThread() > 0);
 
 		FRenderVolumetricCloudRenderViewCS::FParameters* PassParameters = GraphBuilder.AllocParameters<FRenderVolumetricCloudRenderViewCS::FParameters>();
 		PassParameters->OutputViewRect = FVector4f(0.f, 0.f, Desc.Extent.X, Desc.Extent.Y);
-		PassParameters->bBlendCloudColor = !bShouldViewRenderVolumetricRenderTarget;
+		PassParameters->bBlendCloudColor = !bShouldViewRenderVolumetricRenderTarget && !bCloudDebugViewModeEnabled;
 		PassParameters->TargetCubeFace = CloudRC.RenderTargets.Output[0].GetArraySlice();
 		PassParameters->VolumetricCloudRenderViewParamsUB = CloudPassUniformBuffer;
 		PassParameters->View = ViewUniformBuffer;
@@ -2329,10 +2351,11 @@ void FSceneRenderer::RenderVolumetricCloudsInternal(FRDGBuilder& GraphBuilder, F
 		MaterialRenderProxy = MaterialRenderProxy ? MaterialRenderProxy : CloudVolumeMaterialProxy;
 
 		typename FRenderVolumetricCloudRenderViewCS::FPermutationDomain PermutationVector;
-		PermutationVector.Set<typename FRenderVolumetricCloudRenderViewCS::FCloudPerSampleAtmosphereTransmittance>(ShouldUsePerSampleAtmosphereTransmittance(Scene, &MainView));
+		PermutationVector.Set<typename FRenderVolumetricCloudRenderViewCS::FCloudPerSampleAtmosphereTransmittance>(!bCloudDebugViewModeEnabled && ShouldUsePerSampleAtmosphereTransmittance(Scene, &MainView));
 		PermutationVector.Set<typename FRenderVolumetricCloudRenderViewCS::FCloudSampleAtmosphericLightShadowmap>(bSampleVirtualShadowMap);
-		PermutationVector.Set<typename FRenderVolumetricCloudRenderViewCS::FCloudSampleSecondLight>(bSecondAtmosphereLightEnabled);
-		PermutationVector.Set<typename FRenderVolumetricCloudRenderViewCS::FCloudSampleLocalLights>(bCloudEnableLocalLightSampling && !CloudRC.bIsSkyRealTimeReflectionRendering);
+		PermutationVector.Set<typename FRenderVolumetricCloudRenderViewCS::FCloudSampleSecondLight>(!bCloudDebugViewModeEnabled && bSecondAtmosphereLightEnabled);
+		PermutationVector.Set<typename FRenderVolumetricCloudRenderViewCS::FCloudSampleLocalLights>(!bCloudDebugViewModeEnabled && bCloudEnableLocalLightSampling && !CloudRC.bIsSkyRealTimeReflectionRendering);
+		PermutationVector.Set<typename FRenderVolumetricCloudRenderViewCS::FCloudDebugViewMode>(bCloudDebugViewModeEnabled);
 
 		TShaderRef<FRenderVolumetricCloudRenderViewCS> ComputeShader = MaterialResource->GetShader<FRenderVolumetricCloudRenderViewCS>(&FLocalVertexFactory::StaticType, PermutationVector, false);
 
@@ -2397,16 +2420,16 @@ void FSceneRenderer::RenderVolumetricCloudsInternal(FRDGBuilder& GraphBuilder, F
 			RDG_EVENT_NAME("CloudView (PS) %dx%d", Desc.Extent.X, Desc.Extent.Y),
 			RenderViewPassParameters,
 			ERDGPassFlags::Raster,
-			[RenderViewPassParameters, Scene = Scene, &MainView, bVisualizeConservativeDensity,
+			[RenderViewPassParameters, Scene = Scene, &MainView, bCloudDebugViewModeEnabled,
 			bShouldViewRenderVolumetricRenderTarget, CloudVolumeMaterialProxy, bSkipAtmosphericLightShadowmap, bSecondAtmosphereLightEnabled](FRHICommandListImmediate& RHICmdList)
 			{
 				DrawDynamicMeshPass(MainView, RHICmdList,
 					[&MainView, bShouldViewRenderVolumetricRenderTarget, bSkipAtmosphericLightShadowmap, bSecondAtmosphereLightEnabled,
-					bVisualizeConservativeDensity, CloudVolumeMaterialProxy](FDynamicPassMeshDrawListContext* DynamicMeshPassContext)
+					bCloudDebugViewModeEnabled, CloudVolumeMaterialProxy](FDynamicPassMeshDrawListContext* DynamicMeshPassContext)
 					{
 						FVolumetricCloudRenderViewMeshProcessor PassMeshProcessor(
 							MainView.Family->Scene->GetRenderScene(), MainView.GetFeatureLevel(), & MainView, bShouldViewRenderVolumetricRenderTarget,
-							bSkipAtmosphericLightShadowmap, bSecondAtmosphereLightEnabled, bVisualizeConservativeDensity, DynamicMeshPassContext);
+							bSkipAtmosphericLightShadowmap, bSecondAtmosphereLightEnabled, bCloudDebugViewModeEnabled, DynamicMeshPassContext);
 
 						FMeshBatch LocalSingleTriangleMesh;
 						GetSingleTriangleMeshBatch(LocalSingleTriangleMesh, CloudVolumeMaterialProxy, MainView.GetFeatureLevel());
@@ -2463,7 +2486,7 @@ bool FSceneRenderer::RenderVolumetricCloud(
 				bool bShouldViewRenderVolumetricCloudRenderTarget;
 				bool bEnableAerialPerspectiveSampling;
 				bool bShouldUseHighQualityAerialPerspective;
-				bool bVisualizeConservativeDensity;
+				bool bCloudDebugViewModeEnabled;
 			};
 			TArray<FLocalCloudView, TInlineAllocator<2>> ViewsToProcess;
 			
@@ -2477,14 +2500,14 @@ bool FSceneRenderer::RenderVolumetricCloud(
 					continue;
 				}
 
-				const bool bVisualizeConservativeDensity = ShouldViewVisualizeVolumetricCloudConservativeDensity(ViewInfo, ViewFamily.EngineShowFlags);
+				const bool bCloudDebugViewModeEnabled = ShouldViewVisualizeVolumetricCloudConservativeDensity(ViewInfo, ViewFamily.EngineShowFlags);
 				const bool bEnableAerialPerspectiveSampling = CVarVolumetricCloudEnableAerialPerspectiveSampling.GetValueOnAnyThread() > 0;
 				const bool bShouldUseHighQualityAerialPerspective =
 					bEnableAerialPerspectiveSampling
 					&& Scene->HasSkyAtmosphere()
 					&& CVarVolumetricCloudHighQualityAerialPerspective.GetValueOnAnyThread() > 0
 					&& !ViewInfo.bIsReflectionCapture
-					&& !bVisualizeConservativeDensity;
+					&& !bCloudDebugViewModeEnabled;
 
 				if (bAsyncCompute
 					&& bShouldViewRenderVolumetricCloudRenderTarget
@@ -2506,7 +2529,7 @@ bool FSceneRenderer::RenderVolumetricCloud(
 				ViewToProcess.bShouldViewRenderVolumetricCloudRenderTarget = bShouldViewRenderVolumetricCloudRenderTarget;
 				ViewToProcess.bEnableAerialPerspectiveSampling = bEnableAerialPerspectiveSampling;
 				ViewToProcess.bShouldUseHighQualityAerialPerspective = bShouldUseHighQualityAerialPerspective;
-				ViewToProcess.bVisualizeConservativeDensity = bVisualizeConservativeDensity;
+				ViewToProcess.bCloudDebugViewModeEnabled = bCloudDebugViewModeEnabled;
 			}
 
 			for (int32 ViewIndex = 0; ViewIndex < ViewsToProcess.Num(); ViewIndex++)
@@ -2518,7 +2541,7 @@ bool FSceneRenderer::RenderVolumetricCloud(
 
 				CloudRC.MainView = ViewToProcess.ViewInfo;
 				CloudRC.bAsyncCompute = bAsyncCompute;
-				CloudRC.bVisualizeConservativeDensity = ViewToProcess.bVisualizeConservativeDensity;
+				CloudRC.bCloudDebugViewModeEnabled = ViewToProcess.bCloudDebugViewModeEnabled;
 
 				const bool bShouldViewRenderVolumetricCloudRenderTarget = ViewToProcess.bShouldViewRenderVolumetricCloudRenderTarget;
 				const bool bEnableAerialPerspectiveSampling = ViewToProcess.bEnableAerialPerspectiveSampling;
@@ -2563,7 +2586,7 @@ bool FSceneRenderer::RenderVolumetricCloud(
 
 					CloudRC.NoiseFrameIndexModPattern = VRT.GetNoiseFrameIndexModPattern();
 
-					check(VRT.GetMode() != 0 || CloudRC.bVisualizeConservativeDensity || HalfResolutionDepthCheckerboardMinMaxTexture);
+					check(VRT.GetMode() != 0 || CloudRC.bCloudDebugViewModeEnabled || HalfResolutionDepthCheckerboardMinMaxTexture);
 					CloudRC.SceneDepthZ = (VRT.GetMode() == 0 || VRT.GetMode() == 3) ? HalfResolutionDepthCheckerboardMinMaxTexture : SceneDepthZ;
 				}
 				else
@@ -2588,8 +2611,7 @@ bool FSceneRenderer::RenderVolumetricCloud(
 					DestinationRTDepth = GraphBuilder.CreateTexture(FRDGTextureDesc::Create2D(FIntPoint(RtSize.X, RtSize.Y), PF_G16R16F, FClearValueBinding::Black,
 						DepthTexCreateFlags, DepthNumMips, DepthNumSamples), TEXT("Cloud.DummyDepth"));
 
-					const bool bShouldUseComputeForCloudTracing = !CloudRC.bVisualizeConservativeDensity && ShouldUseComputeForCloudTracing(Scene->GetFeatureLevel());
-					if (bShouldUseHighQualityAerialPerspective && bShouldUseComputeForCloudTracing)
+					if (bShouldUseHighQualityAerialPerspective && ShouldUseComputeForCloudTracing(Scene->GetFeatureLevel()))
 					{
 						// If using the compute path, we then need to clear the intermediate render target manually as RDG won't do it for us in this case.
 						AddClearRenderTargetPass(GraphBuilder, IntermediateRT);
@@ -2602,7 +2624,7 @@ bool FSceneRenderer::RenderVolumetricCloud(
 					CloudRC.SceneDepthZ = SceneDepthZ;
 				}
 
-				if (CloudRC.bVisualizeConservativeDensity)
+				if (CloudRC.bCloudDebugViewModeEnabled)
 				{
 					CloudRC.SceneDepthZ = (bool)ERHIZBuffer::IsInverted ? GSystemTextures.GetBlackDummy(GraphBuilder) : GSystemTextures.GetWhiteDummy(GraphBuilder);
 				}
diff --git a/Engine/Source/Runtime/Renderer/Private/VolumetricCloudRendering.h b/Engine/Source/Runtime/Renderer/Private/VolumetricCloudRendering.h
--- a/Engine/Source/Runtime/Renderer/Private/VolumetricCloudRendering.h
+++ b/Engine/Source/Runtime/Renderer/Private/VolumetricCloudRendering.h
@@ -107,6 +107,7 @@ private:
 bool ShouldRenderVolumetricCloud(const FScene* Scene, const FEngineShowFlags& EngineShowFlags);
 bool ShouldViewVisualizeVolumetricCloudConservativeDensity(const FViewInfo& ViewInfo, const FEngineShowFlags& EngineShowFlags);
 bool VolumetricCloudWantsToSampleLocalLights(const FScene* Scene, const FEngineShowFlags& EngineShowFlags);
+uint32 GetVolumetricCloudDebugViewMode(const FEngineShowFlags& ShowFlags);
 
 
 // Structure with data necessary to specify a cloud render.
@@ -135,7 +136,7 @@ struct FCloudRenderContext
 	bool bSecondAtmosphereLightEnabled;
 
 	bool bAsyncCompute;
-	bool bVisualizeConservativeDensity;
+	bool bCloudDebugViewModeEnabled;
 
 	FUintVector4 TracingCoordToZbufferCoordScaleBias;
 	FUintVector4 TracingCoordToFullResPixelCoordScaleBias;
