Fix CIS compilation error.

diff --git a/Engine/Source/Runtime/Renderer/Private/IESTextureManager.cpp b/Engine/Source/Runtime/Renderer/Private/IESTextureManager.cpp
--- a/Engine/Source/Runtime/Renderer/Private/IESTextureManager.cpp
+++ b/Engine/Source/Runtime/Renderer/Private/IESTextureManager.cpp
@@ -252,7 +252,7 @@ static void AddSlotsPass(
 ///////////////////////////////////////////////////////////////////////////////////////////////////
 // API
 
-uint32 AddTexture(UTexture* In)
+uint32 AddTexture(UTextureLightProfile* In)
 {
 	check(IsInRenderingThread());
 
diff --git a/Engine/Source/Runtime/Renderer/Private/IESTextureManager.h b/Engine/Source/Runtime/Renderer/Private/IESTextureManager.h
--- a/Engine/Source/Runtime/Renderer/Private/IESTextureManager.h
+++ b/Engine/Source/Runtime/Renderer/Private/IESTextureManager.h
@@ -6,14 +6,14 @@
 #include "RenderGraphDefinitions.h"
 
 class FRHITexture;
-class UTexture;
+class UTextureLightProfile;
 class FViewInfo;
 
 namespace IESAtlas
 {
 	
 // Add an IES light texture to the atlas
-RENDERER_API uint32 AddTexture(UTexture* IESTexture);
+RENDERER_API uint32 AddTexture(UTextureLightProfile* IESTexture);
 
 // Remove an IES texture from the atlas
 RENDERER_API void RemoveTexture(uint32 InSlotId);
diff --git a/Engine/Source/Runtime/Renderer/Private/RendererScene.cpp b/Engine/Source/Runtime/Renderer/Private/RendererScene.cpp
--- a/Engine/Source/Runtime/Renderer/Private/RendererScene.cpp
+++ b/Engine/Source/Runtime/Renderer/Private/RendererScene.cpp
@@ -2637,7 +2637,7 @@ void FScene::AddLightSceneInfo_RenderThread(FLightSceneInfo* LightSceneInfo)
 	}
 
 	// Register IES texture
-	if (UTexture* IESTexture = LightSceneInfo->Proxy->GetIESTexture())
+	if (UTextureLightProfile* IESTexture = LightSceneInfo->Proxy->GetIESTexture())
 	{
 		LightSceneInfo->Proxy->IESAtlasId = IESAtlas::AddTexture(IESTexture);
 	}
@@ -3928,7 +3928,7 @@ void FScene::RemoveLightSceneInfo_RenderThread(FLightSceneInfo* LightSceneInfo)
 		RectLightAtlas::RemoveTexture(RectProxy->RectAtlasId);
 	}
 
-	if (UTexture* IESTexture = LightSceneInfo->Proxy->GetIESTexture())
+	if (UTextureLightProfile* IESTexture = LightSceneInfo->Proxy->GetIESTexture())
 	{
 		IESAtlas::RemoveTexture(LightSceneInfo->Proxy->IESAtlasId);
 	}
