Add cvar to enable clamping of the camera viewport size when computing VSM clipmap resolution.
This can be important for avoiding artifacts due to excessively show shadow resolutions if dynamic resolution falls too low.

diff --git a/Engine/Source/Runtime/Renderer/Private/VirtualShadowMaps/VirtualShadowMapClipmap.cpp b/Engine/Source/Runtime/Renderer/Private/VirtualShadowMaps/VirtualShadowMapClipmap.cpp
--- a/Engine/Source/Runtime/Renderer/Private/VirtualShadowMaps/VirtualShadowMapClipmap.cpp
+++ b/Engine/Source/Runtime/Renderer/Private/VirtualShadowMaps/VirtualShadowMapClipmap.cpp
@@ -54,6 +54,14 @@ TAutoConsoleVariable<float> CVarVirtualShadowMapClipmapZRangeScale(
 	ECVF_RenderThreadSafe
 );
 
+TAutoConsoleVariable<int32> CVarVirtualShadowMapClipmapMinCameraViewportWidth(
+	TEXT("r.Shadow.Virtual.Clipmap.MinCameraViewportWidth"),
+	0,
+	TEXT("If greater than zero, clamps the camera viewport dimensions used to adjust the clipmap resolution.\n")
+	TEXT("This can be useful to avoid dynamic resolution indirectly dropping the shadow resolution far too low."),
+	ECVF_RenderThreadSafe
+);
+
 // "Virtual" clipmap level to clipmap radius
 // NOTE: This is the radius of around the clipmap origin that this level must cover
 // The actual clipmap dimensions will be larger due to snapping and other accomodations
@@ -90,10 +98,14 @@ FVirtualShadowMapClipmap::FVirtualShadowMapClipmap(
 	// Pure rotation matrix
 	FMatrix ViewToWorldRotationMatrix = WorldToLightViewRotationMatrix.GetTransposed();
 	
+	// Optionally clamp camera viewport to avoid excessively low resolution shadows with dynamic resolution
+	const int32 MinCameraViewportWidth = CVarVirtualShadowMapClipmapMinCameraViewportWidth.GetValueOnRenderThread();
+	const int32 CameraViewportWidth = FMath::Max(MinCameraViewportWidth, CameraViewRectSize.X);
+
 	// NOTE: Rotational (roll) invariance of the directional light depends on square pixels so we just base everything on the camera X scales/resolution
 	// NOTE: 0.5 because we double the size of the clipmap region below to handle snapping
 	float LodScale = 0.5f / CameraViewMatrices.GetProjectionScale().X;
-	LodScale *= float(FVirtualShadowMap::VirtualMaxResolutionXY) / float(CameraViewRectSize.X);
+	LodScale *= float(FVirtualShadowMap::VirtualMaxResolutionXY) / float(CameraViewportWidth);
 	
 	// For now we adjust resolution by just biasing the page we look up in. This is wasteful in terms of page table vs.
 	// just resizing the virtual shadow maps for each clipmap, but convenient for now. This means we need to additionally bias
