Fix intermittent crash when building streaming virtual texture mips.
Issue is that null request handles from an RVT were being sent to the FUploadingVirtualTexture where they hit an assertion.

diff --git a/Engine/Source/Runtime/Renderer/Private/VT/VirtualTextureSystem.cpp b/Engine/Source/Runtime/Renderer/Private/VT/VirtualTextureSystem.cpp
--- a/Engine/Source/Runtime/Renderer/Private/VT/VirtualTextureSystem.cpp
+++ b/Engine/Source/Runtime/Renderer/Private/VT/VirtualTextureSystem.cpp
@@ -2518,7 +2518,10 @@ void FVirtualTextureSystem::SubmitRequests(FRDGBuilder& GraphBuilder, ERHIFeatur
 
 				for (FProducePageDataPrepareTask& Task : PrepareTasks)
 				{
-					Task.VirtualTexture->GatherProducePageDataTasks(Task.RequestHandle, ProducePageTasks);
+					if (Task.RequestHandle != 0)
+					{
+						Task.VirtualTexture->GatherProducePageDataTasks(Task.RequestHandle, ProducePageTasks);
+					}
 				}
 
 				{
