Fix compilation & non-unity CIS issue.

diff --git a/Engine/Source/Runtime/Renderer/Private/IESTextureManager.cpp b/Engine/Source/Runtime/Renderer/Private/IESTextureManager.cpp
--- a/Engine/Source/Runtime/Renderer/Private/IESTextureManager.cpp
+++ b/Engine/Source/Runtime/Renderer/Private/IESTextureManager.cpp
@@ -18,6 +18,7 @@
 #include "TextureLayout.h"
 #include "CommonRenderResources.h"
 #include "ScreenPass.h"
+#include "Engine/TextureLightProfile.h"
 
 ///////////////////////////////////////////////////////////////////////////////////////////////////
 // Config
