Add diagnostic logging / screen message for showing large objects overlapping many non-nanite VSM pages.

diff --git a/Engine/Shaders/Private/VirtualShadowMaps/VirtualShadowMapBuildPerPageDrawCommands.usf b/Engine/Shaders/Private/VirtualShadowMaps/VirtualShadowMapBuildPerPageDrawCommands.usf
--- a/Engine/Shaders/Private/VirtualShadowMaps/VirtualShadowMapBuildPerPageDrawCommands.usf
+++ b/Engine/Shaders/Private/VirtualShadowMaps/VirtualShadowMapBuildPerPageDrawCommands.usf
@@ -59,9 +59,6 @@ StructuredBuffer<uint> DrawCommandDescs;
 
 RWBuffer<uint> DrawIndirectArgsBufferOut;
 
-RWStructuredBuffer<uint> OutInvalidatingInstances;
-uint NumInvalidatingInstanceSlots;
-
 StructuredBuffer<uint> PrimitiveRevealedMask;
 uint PrimitiveRevealedNum;
 
@@ -73,6 +70,17 @@ RWStructuredBuffer<uint> OutDirtyPageFlags;
 uint HZBMode;
 #endif
 
+// TODO: Move to common header
+// Get the area of an "inclusive" rect (which means that the max is inside the rect), also guards against negative area (where min > max)
+uint GetInclusiveRectArea(uint4 Rect)
+{
+	if (all(Rect.zw >= Rect.xy))
+	{
+		uint2 Size = Rect.zw - Rect.xy;
+		return (Size.x + 1) * (Size.y + 1);
+	}
+	return 0;
+}
 
 bool WasPrimitiveRevealed(FVSMCullingBatchInfo VSMCullingBatchInfo, uint PersistentPrimitiveId)
 {
@@ -184,6 +192,10 @@ bool MarkPageDirty(FVirtualSMLevelOffset PageTableLevelOffset, uint2 vPage, uint
 	return false;
 }
 
+#if VSM_GENERATE_STATS
+uint NumPageAreaDiagnosticSlots;
+uint LargeInstancePageAreaThreshold;
+#endif
 
 
 [numthreads(NUM_THREADS_PER_GROUP, 1, 1)]
@@ -263,6 +275,10 @@ void CullPerPageDrawCommandsCs(uint3 GroupId : SV_GroupID, int GroupThreadIndex
 
 		bool bShouldCacheAsStaticForAnyView = false;
 		uint ThreadTotalForAllViews = 0;
+#if VSM_GENERATE_STATS
+		uint TotalPageArea = 0U;
+		float NumVisibleLevels = 0.0f;
+#endif
 		// Loop over views and output visible instance (i.e., those that overlap a valid page)
 		for (uint PrimaryViewId = VSMCullingBatchInfo.FirstPrimaryView; PrimaryViewId < VSMCullingBatchInfo.FirstPrimaryView + VSMCullingBatchInfo.NumPrimaryViews; ++PrimaryViewId)
 		{
@@ -309,6 +325,10 @@ void CullPerPageDrawCommandsCs(uint3 GroupId : SV_GroupID, int GroupThreadIndex
 
 					if (OverlapsAnyValidPage(VirtualShadowMapId, MipLevel, RectPages, FlagMask))
 					{
+#if VSM_GENERATE_STATS
+						TotalPageArea += float(GetInclusiveRectArea(RectPages));
+						NumVisibleLevels += 1.0f;
+#endif
 #if USE_HZB_OCCLUSION
 						if (bViewHZB)
 						{
@@ -406,6 +426,25 @@ void CullPerPageDrawCommandsCs(uint3 GroupId : SV_GroupID, int GroupThreadIndex
 
 		// Accumulate total number of instances for each indirect argument, is also used to allocate space and output compact range of instances later
 		InterlockedAdd(DrawIndirectArgsBufferOut[Payload.IndirectArgIndex * INDIRECT_ARGS_NUM_WORDS + 1], ThreadTotalForAllViews);
+
+#if VSM_GENERATE_STATS
+		if (NumPageAreaDiagnosticSlots > 0U && TotalPageArea > LargeInstancePageAreaThreshold)
+		{
+			for (uint Index = 0U; Index < NumPageAreaDiagnosticSlots; ++Index)
+			{
+				uint PrevArea = 0U;
+				// Store these after the stats slots
+				InterlockedMax(OutStatsBuffer[VSM_NUM_STATS + 2U * Index], TotalPageArea, PrevArea);
+				// If the area was greater then we store the persistent Primitive ID after the area.
+				// This obviously has a data race, so numbers & IDs will only be approximately correct, could use 64-bit atomics instead...
+				if (PrevArea < TotalPageArea)
+				{
+					OutStatsBuffer[VSM_NUM_STATS + 2U * Index + 1U] = PrimitiveData.PersistentPrimitiveIndex;
+					break;
+				}
+			}
+		}
+#endif
 	}
 	GroupMemoryBarrierWithGroupSync();
 	uint NumMarkingJobs = NumSharedMarkingJobs;
diff --git a/Engine/Shaders/Private/VirtualShadowMaps/VirtualShadowMapPrintStats.usf b/Engine/Shaders/Private/VirtualShadowMaps/VirtualShadowMapPrintStats.usf
--- a/Engine/Shaders/Private/VirtualShadowMaps/VirtualShadowMapPrintStats.usf
+++ b/Engine/Shaders/Private/VirtualShadowMaps/VirtualShadowMapPrintStats.usf
@@ -4,6 +4,7 @@
 #include "../ShaderPrint.ush"
 #include "VirtualShadowMapPageAccessCommon.ush"
 #include "VirtualShadowMapStats.ush"
+#include "../GPUMessaging.ush"
 
 StructuredBuffer<uint>	InStatsBuffer;
 // Rectangles bounding (allocated) pages for each mip level
@@ -54,10 +55,26 @@ void PrintUnits(inout FShaderPrintContext Context, uint Value)
 		Print(Context, Value);
 	}
 }
+uint StatsMessageId;
 
 [numthreads(1, 1, 1)]
-void PrintStats()
+void PrintVirtualSmStatsCS()
 {
+	if (StatsMessageId != INDEX_NONE)
+	{
+		// Just post back the whole buffer
+		FGPUMessageWriter Mw = GPUMessageBegin(StatsMessageId, MAX_PAGE_AREA_DIAGNOSTIC_SLOTS * 2U);
+		for (uint Index = 0; Index < MAX_PAGE_AREA_DIAGNOSTIC_SLOTS * 2U; ++Index)
+		{
+			GPUMessageWriteItem(Mw, InStatsBuffer[VSM_NUM_STATS + Index]);
+		}
+	}
+
+	if (ShowStatsValue == 0U)
+	{
+		return;
+	}
+
 	const bool bShowAll                = ShowStatsValue == 1;
 	// Work around potential compiler bug...
 	const bool bShowPhysicalPages      = (ShowStatsValue == 1) || ShowStatsValue == 2;
diff --git a/Engine/Source/Runtime/RenderCore/Public/GPUMessaging.h b/Engine/Source/Runtime/RenderCore/Public/GPUMessaging.h
--- a/Engine/Source/Runtime/RenderCore/Public/GPUMessaging.h
+++ b/Engine/Source/Runtime/RenderCore/Public/GPUMessaging.h
@@ -92,6 +92,8 @@ public:
 
 	uint32 GetPayloadSize() const { return PayloadSize; }
 
+	TConstArrayView<uint32> GetPayLoad() const { return TConstArrayView<uint32>(PayloadData, PayloadSize); }
+
 private:
 	FReader(FMessageId InMessageId, uint32 InPayloadSize, const uint32* InPayloadData)
 		: MessageId(InMessageId)
diff --git a/Engine/Source/Runtime/Renderer/Private/DeferredShadingRenderer.cpp b/Engine/Source/Runtime/Renderer/Private/DeferredShadingRenderer.cpp
--- a/Engine/Source/Runtime/Renderer/Private/DeferredShadingRenderer.cpp
+++ b/Engine/Source/Runtime/Renderer/Private/DeferredShadingRenderer.cpp
@@ -3748,10 +3748,6 @@ void FDeferredShadingSceneRenderer::Render(FRDGBuilder& GraphBuilder)
 	if (VirtualShadowMapArray.IsEnabled())
 	{
 		VirtualShadowMapArray.RenderDebugInfo(GraphBuilder, Views);
-		if (Views.Num() > 0)
-		{
-			VirtualShadowMapArray.PrintStats(GraphBuilder, Views[0]);
-		}
 	}
 
 	for (FViewInfo& View : Views)
diff --git a/Engine/Source/Runtime/Renderer/Private/PrimitiveSceneInfo.cpp b/Engine/Source/Runtime/Renderer/Private/PrimitiveSceneInfo.cpp
--- a/Engine/Source/Runtime/Renderer/Private/PrimitiveSceneInfo.cpp
+++ b/Engine/Source/Runtime/Renderer/Private/PrimitiveSceneInfo.cpp
@@ -2165,6 +2165,18 @@ FString FPrimitiveSceneInfo::GetFullnameForDebuggingOnly() const
 	return FString(TEXT("Unknown Object"));
 }
 
+FString FPrimitiveSceneInfo::GetOwnerActorNameOrLabelForDebuggingOnly() const
+{
+	if (ComponentForDebuggingOnly)
+	{
+#if ACTOR_HAS_LABELS
+		return ComponentForDebuggingOnly->GetOwner() ? ComponentForDebuggingOnly->GetOwner()->GetActorNameOrLabel() : ComponentForDebuggingOnly->GetName();
+#else
+		return ComponentForDebuggingOnly->GetName();
+#endif
+	}
+	return FString(TEXT("Unknown Object"));
+}
 void FPrimitiveSceneInfo::SetCacheShadowAsStatic(bool bStatic)
 {
 	if (bCacheShadowAsStatic != bStatic)
diff --git a/Engine/Source/Runtime/Renderer/Private/VirtualShadowMaps/VirtualShadowMapArray.cpp b/Engine/Source/Runtime/Renderer/Private/VirtualShadowMaps/VirtualShadowMapArray.cpp
--- a/Engine/Source/Runtime/Renderer/Private/VirtualShadowMaps/VirtualShadowMapArray.cpp
+++ b/Engine/Source/Runtime/Renderer/Private/VirtualShadowMaps/VirtualShadowMapArray.cpp
@@ -284,6 +284,20 @@ TAutoConsoleVariable<int32> CVarDebugSkipDynamicPageInvalidation(
 	TEXT("Skip invalidation of cached pages when geometry moves for debugging purposes. This will create obvious visual artifacts when disabled.\n"),
 	ECVF_RenderThreadSafe
 );
+
+TAutoConsoleVariable<int32> CVarNumPageAreaDiagSlots(
+	TEXT("r.Shadow.Virtual.NonNanite.NumPageAreaDiagSlots"),
+	0,
+	TEXT("Feed back diagnostics to host to report page area coverage for the K first occurrences, < 0 uses the max number allowed, 0 disables."),
+	ECVF_RenderThreadSafe
+);
+
+TAutoConsoleVariable<int32> CVarLargeInstancePageAreaThreshold(
+	TEXT("r.Shadow.Virtual.NonNanite.LargeInstancePageAreaThreshold"),
+	-1,
+	TEXT("How large area is considered a 'large' footprint, summed over all overlapped levels, if set to -1 uses the physical page pool size / 8."),
+	ECVF_RenderThreadSafe
+);
 #endif // !UE_BUILD_SHIPPING
 
 static TAutoConsoleVariable<float> CVarMaxMaterialPositionInvalidationRange(
@@ -607,6 +621,7 @@ void FVirtualShadowMapArray::SetShaderDefines(FShaderCompilerEnvironment& OutEnv
 	OutEnvironment.SetDefine(TEXT("VSM_RASTER_WINDOW_PAGES"), FVirtualShadowMap::RasterWindowPages);
 	OutEnvironment.SetDefine(TEXT("VSM_PAGE_TABLE_SIZE"), FVirtualShadowMap::PageTableSize);
 	OutEnvironment.SetDefine(TEXT("VSM_NUM_STATS"), NumStats);
+	OutEnvironment.SetDefine(TEXT("MAX_PAGE_AREA_DIAGNOSTIC_SLOTS"), MaxPageAreaDiagnosticSlots);
 	OutEnvironment.SetDefine(TEXT("INDEX_NONE"), INDEX_NONE);
 }
 
@@ -1321,9 +1336,15 @@ void FVirtualShadowMapArray::BuildPageAllocations(
 	UniformParameters.CoarsePagePixelThresholdStatic = CVarCoarsePagePixelThresholdStatic.GetValueOnRenderThread();
 	UniformParameters.CoarsePagePixelThresholdDynamicNanite = CVarCoarsePagePixelThresholdDynamicNanite.GetValueOnRenderThread();
 
-	if (CVarShowStats.GetValueOnRenderThread() || CacheManager->IsAccumulatingStats())
+#if !UE_BUILD_SHIPPING
+	const bool bRunPageAreaDiagnostics = CVarNumPageAreaDiagSlots.GetValueOnRenderThread() != 0;
+#else
+	constexpr bool bRunPageAreaDiagnostics = false;
+#endif
+
+	if (CVarShowStats.GetValueOnRenderThread() || CacheManager->IsAccumulatingStats() || bRunPageAreaDiagnostics)
 	{
-		StatsBufferRDG = GraphBuilder.CreateBuffer(FRDGBufferDesc::CreateStructuredDesc(sizeof(uint32), NumStats), TEXT("Shadow.Virtual.StatsBuffer"));
+		StatsBufferRDG = GraphBuilder.CreateBuffer(FRDGBufferDesc::CreateStructuredDesc(sizeof(uint32), NumStats + MaxPageAreaDiagnosticSlots * 2), TEXT("Shadow.Virtual.StatsBuffer"));
 		AddClearUAVPass(GraphBuilder, GraphBuilder.CreateUAV(StatsBufferRDG), 0);
 	}
 	
@@ -1802,6 +1823,11 @@ void FVirtualShadowMapArray::RenderDebugInfo(FRDGBuilder& GraphBuilder, TArrayVi
 {
 	check(IsEnabled());
 			
+	if (Views.Num() > 0)
+	{
+		PrintStats(GraphBuilder, Views[0]);
+	}
+
 	if (DebugVisualizationOutput.IsEmpty() || VisualizeLight.IsEmpty())
 	{
 		return;
@@ -1857,10 +1883,12 @@ class FVirtualSmPrintStatsCS : public FVirtualPageManagementShader
 
 	BEGIN_SHADER_PARAMETER_STRUCT(FParameters, )
 		SHADER_PARAMETER_RDG_UNIFORM_BUFFER(FVirtualShadowMapUniformParameters, VirtualShadowMap)
+		SHADER_PARAMETER_STRUCT_INCLUDE(GPUMessage::FParameters, GPUMessageParams)
 		SHADER_PARAMETER_STRUCT_INCLUDE( ShaderPrint::FShaderParameters, ShaderPrintStruct )
 		SHADER_PARAMETER_RDG_BUFFER_SRV(StructuredBuffer< uint >, InStatsBuffer)
 		SHADER_PARAMETER_RDG_BUFFER_SRV(StructuredBuffer<FIntVector4>, AllocatedPageRectBounds)
 		SHADER_PARAMETER(int, ShowStatsValue)
+		SHADER_PARAMETER(uint32, StatsMessageId)
 	END_SHADER_PARAMETER_STRUCT()
 
 
@@ -1873,7 +1901,7 @@ class FVirtualSmPrintStatsCS : public FVirtualPageManagementShader
 
 		
 };
-IMPLEMENT_GLOBAL_SHADER(FVirtualSmPrintStatsCS, "/Engine/Private/VirtualShadowMaps/VirtualShadowMapPrintStats.usf", "PrintStats", SF_Compute);
+IMPLEMENT_GLOBAL_SHADER(FVirtualSmPrintStatsCS, "/Engine/Private/VirtualShadowMaps/VirtualShadowMapPrintStats.usf", "PrintVirtualSmStatsCS", SF_Compute);
 
 void FVirtualShadowMapArray::PrintStats(FRDGBuilder& GraphBuilder, const FViewInfo& View)
 {
@@ -1882,7 +1910,7 @@ void FVirtualShadowMapArray::PrintStats(FRDGBuilder& GraphBuilder, const FViewIn
 
 	// Print stats
 	int ShowStatsValue = CVarShowStats.GetValueOnRenderThread();
-	if (ShowStatsValue != 0 && StatsBufferRDG)
+	if (StatsBufferRDG)
 	{
 		{
 			FVirtualSmPrintStatsCS::FParameters* PassParameters = GraphBuilder.AllocParameters<FVirtualSmPrintStatsCS::FParameters>();
@@ -1891,6 +1919,12 @@ void FVirtualShadowMapArray::PrintStats(FRDGBuilder& GraphBuilder, const FViewIn
 			PassParameters->InStatsBuffer = GraphBuilder.CreateSRV(StatsBufferRDG);
 			PassParameters->VirtualShadowMap = GetUncachedUniformBuffer(GraphBuilder);
 			PassParameters->ShowStatsValue = ShowStatsValue;
+#if !UE_BUILD_SHIPPING
+			PassParameters->StatsMessageId = CacheManager->StatsFeedbackSocket.GetMessageId().IsValid() ? CacheManager->StatsFeedbackSocket.GetMessageId().GetIndex() : INDEX_NONE;
+#else
+			PassParameters->StatsMessageId = INDEX_NONE;
+#endif
+			PassParameters->GPUMessageParams = GPUMessage::GetShaderParameters(GraphBuilder);
 
 			auto ComputeShader = View.ShaderMap->GetShader<FVirtualSmPrintStatsCS>();
 
@@ -2116,8 +2150,8 @@ public:
 
 		SHADER_PARAMETER_RDG_BUFFER_UAV(RWStructuredBuffer<uint>, OutStaticInvalidatingPrimitives)
 
-		SHADER_PARAMETER_RDG_BUFFER_UAV(RWStructuredBuffer<uint>, OutInvalidatingInstances)
-		SHADER_PARAMETER(uint32, NumInvalidatingInstanceSlots)
+		SHADER_PARAMETER(uint32, NumPageAreaDiagnosticSlots)
+		SHADER_PARAMETER(uint32, LargeInstancePageAreaThreshold)
 
 		SHADER_PARAMETER_RDG_BUFFER_UAV(RWStructuredBuffer< uint >, OutStatsBuffer)
 	END_SHADER_PARAMETER_STRUCT()
@@ -2333,12 +2367,19 @@ static FCullingResult AddCullingPasses(FRDGBuilder& GraphBuilder,
 		PassParameters->VisibleInstancesOut = GraphBuilder.CreateUAV(VisibleInstancesRdg, ERDGUnorderedAccessViewFlags::SkipBarrier);
 		PassParameters->VisibleInstanceCountBufferOut = GraphBuilder.CreateUAV(VisibleInstanceWriteOffsetRDG, ERDGUnorderedAccessViewFlags::SkipBarrier);
 		PassParameters->OutStaticInvalidatingPrimitives = GraphBuilder.CreateUAV(VirtualShadowMapArray.StaticInvalidatingPrimitivesRDG, ERDGUnorderedAccessViewFlags::SkipBarrier);
+
+		PassParameters->NumPageAreaDiagnosticSlots = 0U;
+
 		PassParameters->HZBShaderParameters = HZBShaderParameters;
 
 		bool bGenerateStats = VirtualShadowMapArray.StatsBufferRDG != nullptr;
 		if (bGenerateStats)
 		{
 			PassParameters->OutStatsBuffer = GraphBuilder.CreateUAV(VirtualShadowMapArray.StatsBufferRDG, ERDGUnorderedAccessViewFlags::SkipBarrier);
+#if !UE_BUILD_SHIPPING
+			PassParameters->NumPageAreaDiagnosticSlots = CVarNumPageAreaDiagSlots.GetValueOnRenderThread() < 0 ? FVirtualShadowMapArray::MaxPageAreaDiagnosticSlots : FMath::Min(FVirtualShadowMapArray::MaxPageAreaDiagnosticSlots, uint32(CVarNumPageAreaDiagSlots.GetValueOnRenderThread()));
+			PassParameters->LargeInstancePageAreaThreshold = CVarLargeInstancePageAreaThreshold.GetValueOnRenderThread() >= 0 ? CVarLargeInstancePageAreaThreshold.GetValueOnRenderThread() : (VirtualShadowMapArray.GetMaxPhysicalPages() / 8);
+#endif
 		}
 
 		FCullPerPageDrawCommandsCs::FPermutationDomain PermutationVector;
@@ -2364,7 +2405,7 @@ static FCullingResult AddCullingPasses(FRDGBuilder& GraphBuilder,
 		AddClearUAVPass(GraphBuilder, GraphBuilder.CreateUAV(InstanceIdOutOffsetBufferRDG), 0);
 
 		PassParameters->NumIndirectArgs = NumIndirectArgs;
-		PassParameters->InstanceIdOffsetBufferOut = GraphBuilder.CreateUAV(CullingResult.InstanceIdOffsetBufferRDG, PF_R32_UINT);;
+		PassParameters->InstanceIdOffsetBufferOut = GraphBuilder.CreateUAV(CullingResult.InstanceIdOffsetBufferRDG, PF_R32_UINT);
 		PassParameters->OutputOffsetBufferOut = GraphBuilder.CreateUAV(InstanceIdOutOffsetBufferRDG);
 		PassParameters->TmpInstanceIdOffsetBufferOut = GraphBuilder.CreateUAV(TmpInstanceIdOffsetBufferRDG);
 		PassParameters->DrawIndirectArgsBuffer = GraphBuilder.CreateSRV(CullingResult.DrawIndirectArgsRDG, PF_R32_UINT);
diff --git a/Engine/Source/Runtime/Renderer/Private/VirtualShadowMaps/VirtualShadowMapArray.h b/Engine/Source/Runtime/Renderer/Private/VirtualShadowMaps/VirtualShadowMapArray.h
--- a/Engine/Source/Runtime/Renderer/Private/VirtualShadowMaps/VirtualShadowMapArray.h
+++ b/Engine/Source/Runtime/Renderer/Private/VirtualShadowMaps/VirtualShadowMapArray.h
@@ -349,6 +349,7 @@ public:
 
 	// See Engine\Shaders\Private\VirtualShadowMaps\VirtualShadowMapStats.ush for definitions of the different stat indexes
 	static constexpr uint32 NumStats = 16;
+	static constexpr uint32 MaxPageAreaDiagnosticSlots = 32;
 
 	FRDGBufferRef StatsBufferRDG = nullptr;
 
diff --git a/Engine/Source/Runtime/Renderer/Private/VirtualShadowMaps/VirtualShadowMapCacheManager.cpp b/Engine/Source/Runtime/Renderer/Private/VirtualShadowMaps/VirtualShadowMapCacheManager.cpp
--- a/Engine/Source/Runtime/Renderer/Private/VirtualShadowMaps/VirtualShadowMapCacheManager.cpp
+++ b/Engine/Source/Runtime/Renderer/Private/VirtualShadowMaps/VirtualShadowMapCacheManager.cpp
@@ -498,7 +498,7 @@ FVirtualShadowMapArrayCacheManager::FVirtualShadowMapArrayCacheManager(FScene* I
 						LodBiasDirectional);
 					bLoggedPageOverflow = true;
 				}
-				LastOverflowFrame = Scene->GetFrameNumber();
+				LastOverflowTime = float(FGameTime::GetTimeSinceAppStart().GetRealTimeSeconds());
 #endif
 			}
 #if !UE_BUILD_SHIPPING
@@ -510,15 +510,54 @@ FVirtualShadowMapArrayCacheManager::FVirtualShadowMapArrayCacheManager(FScene* I
 		}
 	});
 
+#if !UE_BUILD_SHIPPING
+	// Handle message with stats sent back from GPU whenever stats are enabled
+	StatsFeedbackSocket = GPUMessage::RegisterHandler(TEXT("Shadow.Virtual.StatsFeedback"), [this](GPUMessage::FReader Message)
+	{
+		LastLoggedPageOverlapAppTime.SetNumZeroed(Scene->GetMaxPersistentPrimitiveIndex());
+		float RealTimeSeconds = float(FGameTime::GetTimeSinceAppStart().GetRealTimeSeconds());
+
+		TConstArrayView<uint32> PageAreaDiags = Message.GetPayLoad();
+		for (int32 Index = 0; Index < PageAreaDiags.Num(); Index += 2)
+		{
+			uint32 Overlap = PageAreaDiags[Index];
+			uint32 PersistentPrimitiveId = PageAreaDiags[Index + 1];
+			int32 PrimtiveIndex = Scene->GetPrimitiveIndex(FPersistentPrimitiveIndex{ int32(PersistentPrimitiveId) });
+			if (Overlap > 0 && PrimtiveIndex != INDEX_NONE)
+			{
+				if (RealTimeSeconds - LastLoggedPageOverlapAppTime[PersistentPrimitiveId] > 5.0f)
+				{
+					LastLoggedPageOverlapAppTime[PersistentPrimitiveId] = RealTimeSeconds;
+					UE_LOG(LogRenderer, Warning, TEXT("Non-Nanite VSM page overlap performance Warning: Primitive '%s' overlapped %d Pages (full name: %s)"), *Scene->Primitives[PrimtiveIndex]->GetOwnerActorNameOrLabelForDebuggingOnly(), Overlap, *Scene->Primitives[PrimtiveIndex]->GetFullnameForDebuggingOnly());
+				}
+				LargePageAreaItems.Add(PersistentPrimitiveId, FLargePageAreaItem{ Overlap, RealTimeSeconds });
+			}
+		}
+	});
+#endif
+
 #if !UE_BUILD_SHIPPING
 	ScreenMessageDelegate = FRendererOnScreenNotification::Get().AddLambda([this](TMultiMap<FCoreDelegates::EOnScreenMessageSeverity, FText >& OutMessages)
 	{
+		float RealTimeSeconds = float(FGameTime::GetTimeSinceAppStart().GetRealTimeSeconds());
+
 		// Show for ~5s after last overflow
 		int32 CurrentFrameNumber = Scene->GetFrameNumber();
-		if (LastOverflowFrame >= 0 && CurrentFrameNumber - LastOverflowFrame < 30 * 5)
+		if (LastOverflowTime >= 0.0f && RealTimeSeconds - LastOverflowTime < 5.0f)
 		{
-			OutMessages.Add(FCoreDelegates::EOnScreenMessageSeverity::Warning, FText::FromString(FString::Printf(TEXT("Virtual Shadow Map Page Pool overflow detected (%d frames ago)"), CurrentFrameNumber - LastOverflowFrame)));
+			OutMessages.Add(FCoreDelegates::EOnScreenMessageSeverity::Warning, FText::FromString(FString::Printf(TEXT("Virtual Shadow Map Page Pool overflow detected (%0.0f seconds ago)"), RealTimeSeconds - LastOverflowTime)));
 		}
+
+		for (const auto& Item : LargePageAreaItems)
+		{
+			int32 PrimtiveIndex = Scene->GetPrimitiveIndex(FPersistentPrimitiveIndex{ int32(Item.Key) });
+			uint32 Overlap = Item.Value.PageArea;
+			if (PrimtiveIndex != INDEX_NONE && RealTimeSeconds - Item.Value.LastTimeSeen < 2.5f)
+			{
+				OutMessages.Add(FCoreDelegates::EOnScreenMessageSeverity::Warning, FText::FromString(FString::Printf(TEXT("Non-Nanite VSM page overlap performance Warning: Primitive '%s' overlapped %d Pages"), *Scene->Primitives[PrimtiveIndex]->GetOwnerActorNameOrLabelForDebuggingOnly(), Overlap)));
+			}
+		}
+		TrimLoggingInfo();
 	});
 #endif
 }
@@ -697,6 +736,8 @@ void FVirtualShadowMapArrayCacheManager::ExtractFrameData(
 	const FSceneRenderer& SceneRenderer,
 	bool bEnableCaching)
 {
+	TrimLoggingInfo();
+
 	const bool bNewShadowData = VirtualShadowMapArray.IsAllocated();
 	const bool bDropAll = !bEnableCaching;
 	const bool bDropPrevBuffers = bDropAll || bNewShadowData;
@@ -1210,3 +1251,15 @@ void FVirtualShadowMapArrayCacheManager::ProcessInvalidations(
 	);
 
 }
+// Remove old info used to track logging.
+void FVirtualShadowMapArrayCacheManager::TrimLoggingInfo()
+{
+#if !UE_BUILD_SHIPPING
+	// Remove old items
+	float RealTimeSeconds = float(FGameTime::GetTimeSinceAppStart().GetRealTimeSeconds());
+	LargePageAreaItems = LargePageAreaItems.FilterByPredicate([RealTimeSeconds](const TMap<uint32, FLargePageAreaItem>::ElementType& Element)
+	{
+		return RealTimeSeconds - Element.Value.LastTimeSeen < 5.0f;
+	});
+#endif
+}
diff --git a/Engine/Source/Runtime/Renderer/Private/VirtualShadowMaps/VirtualShadowMapCacheManager.h b/Engine/Source/Runtime/Renderer/Private/VirtualShadowMaps/VirtualShadowMapCacheManager.h
--- a/Engine/Source/Runtime/Renderer/Private/VirtualShadowMaps/VirtualShadowMapCacheManager.h
+++ b/Engine/Source/Runtime/Renderer/Private/VirtualShadowMaps/VirtualShadowMapCacheManager.h
@@ -287,6 +287,10 @@ public:
 	uint64 GetGPUSizeBytes(bool bLogSizes) const;
 
 	GPUMessage::FSocket StatusFeedbackSocket;
+#if !UE_BUILD_SHIPPING
+	// Socket for optional stats that are only sent back if enabled
+	GPUMessage::FSocket StatsFeedbackSocket;
+#endif 
 
 	FVirtualShadowMapFeedback StaticGPUInvalidationsFeedback;
 
@@ -309,6 +313,8 @@ private:
 			RecentlyRemovedPrimitives[RecentlyRemovedReadIndex][PersistentPrimitiveIndex.Index];
 	}
 
+	// Remove old info used to track logging.
+	void TrimLoggingInfo();
 	// The actual physical texture data is stored here rather than in VirtualShadowMapArray (which is recreated each frame)
 	// This allows us to (optionally) persist cached pages between frames. Regardless of whether caching is enabled,
 	// we store the physical pool here.
@@ -334,8 +340,19 @@ private:
 	FRHIGPUBufferReadback* GPUBufferReadback = nullptr;
 #if !UE_BUILD_SHIPPING
 	FDelegateHandle ScreenMessageDelegate;
-	int32 LastOverflowFrame = -1;
+	float LastOverflowTime = -1.0f;
 	bool bLoggedPageOverflow = false;
+
+	// Stores the last time (wall-clock seconds since app-start) that an non-nanite page area message was logged,
+	TArray<float> LastLoggedPageOverlapAppTime;
+
+	// Map to track non-nanite page area items that are shown on screen
+	struct FLargePageAreaItem
+	{
+		uint32 PageArea;
+		float LastTimeSeen;
+	};
+	TMap<uint32, FLargePageAreaItem> LargePageAreaItems;
 #endif
 #if WITH_MGPU
 	FRHIGPUMask LastGPUMask;
diff --git a/Engine/Source/Runtime/Renderer/Public/PrimitiveSceneInfo.h b/Engine/Source/Runtime/Renderer/Public/PrimitiveSceneInfo.h
--- a/Engine/Source/Runtime/Renderer/Public/PrimitiveSceneInfo.h
+++ b/Engine/Source/Runtime/Renderer/Public/PrimitiveSceneInfo.h
@@ -599,7 +599,8 @@ public:
 
 	/** Return primitive fullname (for debugging only). */
 	FString GetFullnameForDebuggingOnly() const;
-
+	/** Return primitive Owner actor name (for debugging only). */
+	FString GetOwnerActorNameOrLabelForDebuggingOnly() const;
 	FORCEINLINE bool ShouldCacheShadowAsStatic() const
 	{
 		return bCacheShadowAsStatic;
