Don't use the fast clear flag for the bloom render targets, since all pixels will be written. If not specify, bloom render targets can be aliased with an RT using a cmask. Discard/ResourceBarriers on each bloom passes decrease performance. Saves 0.15ms on old gen console

diff --git a/Engine/Source/Runtime/Renderer/Private/PostProcess/PostProcessWeightedSampleSum.cpp b/Engine/Source/Runtime/Renderer/Private/PostProcess/PostProcessWeightedSampleSum.cpp
--- a/Engine/Source/Runtime/Renderer/Private/PostProcess/PostProcessWeightedSampleSum.cpp
+++ b/Engine/Source/Runtime/Renderer/Private/PostProcess/PostProcessWeightedSampleSum.cpp
@@ -384,7 +384,7 @@ FScreenPassTexture AddGaussianBlurPass(
 	FRDGTextureDesc OutputDesc = Filter.Texture->Desc;
 	OutputDesc.Reset();
 	OutputDesc.Extent = OutputViewport.Extent;
-	OutputDesc.Flags |= bIsComputePass ? TexCreate_UAV : TexCreate_None;
+	OutputDesc.Flags |= bIsComputePass ? TexCreate_UAV : TexCreate_NoFastClear;
 	OutputDesc.ClearValue = FClearValueBinding(FLinearColor::Transparent);
 
 	const FScreenPassRenderTarget Output(GraphBuilder.CreateTexture(OutputDesc, Name), OutputViewport.Rect, ERenderTargetLoadAction::ENoAction);
@@ -559,4 +559,4 @@ FScreenPassTexture AddGaussianBlurPass(
 			TArrayView<const FLinearColor>(SampleWeights, SampleCount),
 			Inputs.UseMirrorAddressMode);
 	}
-}
\ No newline at end of file
+}
