Works arround crash on D3D11 driver that doesn't support UAV on B8G8R8A8

diff --git a/Engine/Source/Runtime/Renderer/Private/PostProcess/TemporalSuperResolution.cpp b/Engine/Source/Runtime/Renderer/Private/PostProcess/TemporalSuperResolution.cpp
--- a/Engine/Source/Runtime/Renderer/Private/PostProcess/TemporalSuperResolution.cpp
+++ b/Engine/Source/Runtime/Renderer/Private/PostProcess/TemporalSuperResolution.cpp
@@ -1357,7 +1357,7 @@ ITemporalUpscaler::FOutputs AddTemporalSuperResolutionPasses(
 		{
 			FRDGTextureDesc Desc = FRDGTextureDesc::Create2D(
 				InputExtent,
-				PrevHistorySRVs.Guide->Desc.Format,
+				History.Guide->Desc.Format,
 				FClearValueBinding::None,
 				/* InFlags = */ TexCreate_ShaderResource | TexCreate_UAV);
 
@@ -1367,7 +1367,7 @@ ITemporalUpscaler::FOutputs AddTemporalSuperResolutionPasses(
 		{
 			FRDGTextureDesc Desc = FRDGTextureDesc::Create2D(
 				InputExtent,
-				PrevHistorySRVs.Moire->Desc.Format,
+				History.Moire->Desc.Format,
 				FClearValueBinding::None,
 				/* InFlags = */ TexCreate_ShaderResource | TexCreate_UAV);
 
