Experimental support for calculating auto exposure without being affected by materials.
- In this initial version luminance used for auto exposure is divided by base color luminance.
    - this doesn't properly take into account specular, reflections, volumetrics, emissive, etc.
- Can be enabled using r.AutoExposure.IgnoreMaterials.

diff --git a/Engine/Shaders/Private/Histogram.usf b/Engine/Shaders/Private/Histogram.usf
--- a/Engine/Shaders/Private/Histogram.usf
+++ b/Engine/Shaders/Private/Histogram.usf
@@ -1,5 +1,7 @@
 // Copyright Epic Games, Inc. All Rights Reserved.
 
+#include "DeferredShadingCommon.ush"
+
 #include "/Engine/Private/Common.ush"
 #include "/Engine/Private/ScreenPass.ush"
 #include "/Engine/Private/PostProcessHistogramCommon.ush"
@@ -41,6 +43,12 @@ uint2 ThreadGroupCount;
 Texture2D InputTexture;
 Texture2D EyeAdaptationTexture;
 
+#if DIVIDE_ALBEDO
+	SCREEN_PASS_TEXTURE_VIEWPORT(InputGBuffer)
+	Texture2D InputTextureB;
+	Texture2D InputTextureC;
+#endif
+
 ///////////////////////////////////////////////////////
 // * LOCAL DATA STORE *
 // This is where each ThreadGroup accumulates (via InterlockedAdd) a local (per thread group) copy of the histogram
@@ -74,6 +82,10 @@ RWTexture2D<uint> HistogramScatter32Output;
 //  we should pack up our Temp UINT textures, convert back to float and stick the histogram in the final texture
 RWTexture2D<float4> HistogramOutput;
 
+///////////////////////////////////////////////////////
+//  * DEBUG TEXTURES *
+RWTexture2D<float> DebugOutput;
+
 
 ///////////////////////////////////////////////////////
 //  * PACKING HELPER FUNCTIONS *
@@ -144,13 +156,8 @@ void ClearLDS(uint GroupThreadId)
 
 ///////////////////////////////////////////////////////
 //Given a Texel pos this reads in the screen value and calculates the two fractional bucket it lies in.
-void CalculateBucketsAndWeights(in uint2 TexelPos, out float Weight0, out float Weight1, out uint Bucket0, out uint Bucket1)
+void CalculateBucketsAndWeights(in float LuminanceVal, out float Weight0, out float Weight1, out uint Bucket0, out uint Bucket1)
 {
-
-	float4 SceneColor = InputTexture.Load(int3(TexelPos, 0));
-	SceneColor.xyz *= View.OneOverPreExposure;
-
-	float LuminanceVal = CalculateEyeAdaptationLuminance(SceneColor.xyz);
 	float LogLuminance = ComputeHistogramPositionFromLuminance(LuminanceVal);
 	// Map the normalized histogram position into texels.
 	float fBucket = saturate(LogLuminance) * (HISTOGRAM_SIZE-1);
@@ -310,13 +317,40 @@ void MainAtomicCS(
 
 		if (TexelPos.x < Input_ViewportMax.x)
 		{
+			float2 InvViewSize = Input_ViewportSizeInverse.xy;
+			float2 ScreenUV = (TexelPos.xy - Input_ViewportMin) * InvViewSize;
+
+			float4 SceneColor = InputTexture.Load(int3(TexelPos, 0));
+			SceneColor *= View.OneOverPreExposure;
+
+			float LuminanceVal = CalculateEyeAdaptationLuminance(SceneColor.rgb);
+
+#if DIVIDE_ALBEDO
+			uint2 GBufferTexelPos = InputGBuffer_ViewportMin + ScreenUV * InputGBuffer_ViewportSize;
+			float4 GBufferB = InputTextureB.Load(int3(GBufferTexelPos, 0));
+			uint ShadingModelID = DecodeShadingModelId(GBufferB.a);
+
+			if (ShadingModelID != SHADINGMODELID_UNLIT)
+			{
+				float4 GBufferC = InputTextureC.Load(int3(GBufferTexelPos, 0));
+				float3 BaseColor = DecodeBaseColor(GBufferC.rgb);
+
+				float BaseColorLuminanceVal = CalculateEyeAdaptationLuminance(BaseColor.rgb);
+
+				LuminanceVal /= max(BaseColorLuminanceVal, EyeAdaptation_IgnoreMaterialsMinBaseColorLuminance);
+				LuminanceVal *= lerp(EyeAdaptation_IgnoreMaterialsLuminanceScale, 1.0f, EyeAdaptation_IgnoreMaterialsMinBaseColorLuminance);
+			}
+#endif
+
+#if DIVIDE_ALBEDO_DEBUG
+			DebugOutput[TexelPos] = LuminanceVal;
+#endif
+
 			float Weight0,Weight1;
 			uint Bucket0,Bucket1;
 			
-			CalculateBucketsAndWeights(TexelPos, Weight0, Weight1, Bucket0, Bucket1);
+			CalculateBucketsAndWeights(LuminanceVal, Weight0, Weight1, Bucket0, Bucket1);
 
-			float2 InvViewSize = Input_ViewportSizeInverse.xy;
-			float2 ScreenUV = (TexelPos.xy - Input_ViewportMin) * InvViewSize;
 			float ScreenWeight = AdaptationWeightTexture(ScreenUV);
 
 			ScatterWeightsToLDS(Bucket0, Bucket1, Weight0, Weight1, ScreenWeight, GroupThreadId.x);
diff --git a/Engine/Shaders/Private/PostProcessHistogramCommon.ush b/Engine/Shaders/Private/PostProcessHistogramCommon.ush
--- a/Engine/Shaders/Private/PostProcessHistogramCommon.ush
+++ b/Engine/Shaders/Private/PostProcessHistogramCommon.ush
@@ -22,6 +22,8 @@ float EyeAdaptation_LocalExposureDetailStrength;
 float EyeAdaptation_LocalExposureBlurredLuminanceBlend;
 float EyeAdaptation_LocalExposureMiddleGreyExposureCompensation;
 float EyeAdaptation_BlackHistogramBucketInfluence;
+float EyeAdaptation_IgnoreMaterialsLuminanceScale;
+float EyeAdaptation_IgnoreMaterialsMinBaseColorLuminance;
 float EyeAdaptation_GreyMult;
 float EyeAdaptation_ExponentialUpM;
 float EyeAdaptation_ExponentialDownM;
diff --git a/Engine/Source/Runtime/Renderer/Private/PostProcess/PostProcessEyeAdaptation.cpp b/Engine/Source/Runtime/Renderer/Private/PostProcess/PostProcessEyeAdaptation.cpp
--- a/Engine/Source/Runtime/Renderer/Private/PostProcess/PostProcessEyeAdaptation.cpp
+++ b/Engine/Source/Runtime/Renderer/Private/PostProcess/PostProcessEyeAdaptation.cpp
@@ -66,6 +66,31 @@ namespace
 		TEXT("When set to 1.0, fully dark pixels will accumulate normally, whereas when set to 0.0 fully dark pixels\n")
 		TEXT("will have no influence.\n"),
 		ECVF_RenderThreadSafe);
+
+	TAutoConsoleVariable<bool> CVarAutoExposureIgnoreMaterials(
+		TEXT("r.AutoExposure.IgnoreMaterials"),
+		false,
+		TEXT("(Experimental) Whether to calculate auto exposure assuming every surface uses a perfectly diffuse white material.\n")
+		TEXT("(default: false)"),
+		ECVF_Scalability | ECVF_RenderThreadSafe);
+
+	TAutoConsoleVariable<float> CVarAutoExposureIgnoreMaterialsLuminanceScale(
+		TEXT("r.AutoExposure.IgnoreMaterials.LuminanceScale"),
+		0.18f,
+		TEXT(""),
+		ECVF_RenderThreadSafe);
+
+	TAutoConsoleVariable<float> CVarAutoExposureIgnoreMaterialsMinBaseColorLuminance(
+		TEXT("r.AutoExposure.IgnoreMaterials.MinBaseColorLuminance"),
+		0.01f,
+		TEXT(""),
+		ECVF_RenderThreadSafe);
+
+	TAutoConsoleVariable<bool> CVarAutoExposureIgnoreMaterialsDebug(
+		TEXT("r.AutoExposure.IgnoreMaterials.Debug"),
+		false,
+		TEXT(""),
+		ECVF_RenderThreadSafe);
 }
 
 // Basic eye adaptation is supported everywhere except mobile when MobileHDR is disabled
@@ -441,6 +466,8 @@ FEyeAdaptationParameters GetEyeAdaptationParameters(const FViewInfo& View, ERHIF
 	Parameters.ExponentialUpM = ExponentialUpM;
 	Parameters.StartDistance = StartDistance;
 	Parameters.LuminanceMax = LuminanceMax;
+	Parameters.IgnoreMaterialsLuminanceScale = CVarAutoExposureIgnoreMaterialsLuminanceScale.GetValueOnRenderThread();
+	Parameters.IgnoreMaterialsMinBaseColorLuminance = CVarAutoExposureIgnoreMaterialsMinBaseColorLuminance.GetValueOnRenderThread();
 	Parameters.ForceTarget = ForceTarget;
 	Parameters.VisualizeDebugType = CVarEyeAdaptationVisualizeDebugType.GetValueOnRenderThread();
 	Parameters.MeterMaskTexture = MeterMask;
diff --git a/Engine/Source/Runtime/Renderer/Private/PostProcess/PostProcessEyeAdaptation.h b/Engine/Source/Runtime/Renderer/Private/PostProcess/PostProcessEyeAdaptation.h
--- a/Engine/Source/Runtime/Renderer/Private/PostProcess/PostProcessEyeAdaptation.h
+++ b/Engine/Source/Runtime/Renderer/Private/PostProcess/PostProcessEyeAdaptation.h
@@ -76,6 +76,8 @@ BEGIN_SHADER_PARAMETER_STRUCT(FEyeAdaptationParameters, )
 	SHADER_PARAMETER(float, ExponentialDownM)
 	SHADER_PARAMETER(float, StartDistance)
 	SHADER_PARAMETER(float, LuminanceMax)
+	SHADER_PARAMETER(float, IgnoreMaterialsLuminanceScale)
+	SHADER_PARAMETER(float, IgnoreMaterialsMinBaseColorLuminance)
 	SHADER_PARAMETER(float, ForceTarget)
 	SHADER_PARAMETER(int, VisualizeDebugType)
 	SHADER_PARAMETER_TEXTURE(Texture2D, MeterMaskTexture)
diff --git a/Engine/Source/Runtime/Renderer/Private/PostProcess/PostProcessHistogram.cpp b/Engine/Source/Runtime/Renderer/Private/PostProcess/PostProcessHistogram.cpp
--- a/Engine/Source/Runtime/Renderer/Private/PostProcess/PostProcessHistogram.cpp
+++ b/Engine/Source/Runtime/Renderer/Private/PostProcess/PostProcessHistogram.cpp
@@ -8,6 +8,8 @@
 #include "PostProcess/PostProcessEyeAdaptation.h"
 #include "ShaderCompilerCore.h"
 
+#include "SceneTextureParameters.h"
+
 TAutoConsoleVariable<int32> CVarUseAtomicHistogram(
 	TEXT("r.Histogram.UseAtomic"), 1,
 	TEXT("Uses atomic to speed up the generation of the histogram."),
@@ -132,14 +134,22 @@ public:
 	BEGIN_SHADER_PARAMETER_STRUCT(FParameters, )
 		SHADER_PARAMETER_STRUCT_REF(FViewUniformShaderParameters, View)
 		SHADER_PARAMETER_STRUCT(FScreenPassTextureViewportParameters, Input)
+		SHADER_PARAMETER_STRUCT(FScreenPassTextureViewportParameters, InputGBuffer)
 		SHADER_PARAMETER_STRUCT(FEyeAdaptationParameters, EyeAdaptation)
 		SHADER_PARAMETER_RDG_TEXTURE(Texture2D, InputTexture)
+		SHADER_PARAMETER_RDG_TEXTURE(Texture2D, InputTextureB)
+		SHADER_PARAMETER_RDG_TEXTURE(Texture2D, InputTextureC)
 		SHADER_PARAMETER_RDG_TEXTURE(Texture2D, EyeAdaptationTexture)
 		SHADER_PARAMETER(FIntPoint, ThreadGroupCount)
 		SHADER_PARAMETER_RDG_TEXTURE_UAV(RWTexture2D<float4>, HistogramScatter64Output)
 		SHADER_PARAMETER_RDG_TEXTURE_UAV(RWTexture2D<float4>, HistogramScatter32Output)
+		SHADER_PARAMETER_RDG_TEXTURE_UAV(RWTexture2D<float>, DebugOutput)
 	END_SHADER_PARAMETER_STRUCT()
 
+	class FDivideAlbedoDim : SHADER_PERMUTATION_BOOL("DIVIDE_ALBEDO");
+	class FDivideAlbedoDebugDim : SHADER_PERMUTATION_BOOL("DIVIDE_ALBEDO_DEBUG");
+	using FPermutationDomain = TShaderPermutationDomain<FDivideAlbedoDim, FDivideAlbedoDebugDim>;
+
 	static bool ShouldCompilePermutation(const FGlobalShaderPermutationParameters& Parameters)
 	{
 		return IsFeatureLevelSupported(Parameters.Platform, ERHIFeatureLevel::SM5);
@@ -292,6 +302,7 @@ static FRDGTextureRef AddHistogramAtomicPass(
 	const FViewInfo& View,
 	const FEyeAdaptationParameters& EyeAdaptationParameters,
 	const FScreenPassTexture& SceneColor,
+	const FSceneTextureParameters& SceneTextures,
 	FRDGTextureRef EyeAdaptationTexture)
 {
 	check(SceneColor.IsValid());
@@ -326,18 +337,35 @@ static FRDGTextureRef AddHistogramAtomicPass(
 		PassParameters->View = View.ViewUniformBuffer;
 		PassParameters->Input = GetScreenPassTextureViewportParameters(FScreenPassTextureViewport(SceneColor));
 		PassParameters->InputTexture = SceneColor.Texture;
+		PassParameters->InputGBuffer = GetScreenPassTextureViewportParameters(FScreenPassTextureViewport(View.ViewRect));
+		PassParameters->InputTextureB = SceneTextures.GBufferBTexture;
+		PassParameters->InputTextureC = SceneTextures.GBufferCTexture;
 		PassParameters->EyeAdaptationTexture = EyeAdaptationTexture;
 		PassParameters->ThreadGroupCount = FIntPoint(SceneColor.ViewRect.Size().Y, 1);
 		PassParameters->EyeAdaptation = EyeAdaptationParameters;
 		// PassParameters->HistogramScatter64Output = GraphBuilder->CreateUAV(HistogramScatter64Texture);
 		PassParameters->HistogramScatter32Output = GraphBuilder.CreateUAV(HistogramScatter32Texture);
 
+		{
+			const FRDGTextureDesc Desc = FRDGTextureDesc::Create2D(SceneColor.ViewRect.Size(), PF_R16F, FClearValueBinding::Black, TexCreate_UAV | TexCreate_ShaderResource);
+			FRDGTextureRef DebugOutputTexture = GraphBuilder.CreateTexture(Desc, TEXT("EyeAdaptation_DebugOutput"));
+			PassParameters->DebugOutput = GraphBuilder.CreateUAV(DebugOutputTexture);
+		}
+
 		//clear the temp textures
 		uint32 ClearValues[4] = { 0, 0, 0, 0 };
 		// AddClearUAVPass(GraphBuilder, PassParameters->HistogramScatter64Output, ClearValues);
 		AddClearUAVPass(GraphBuilder, PassParameters->HistogramScatter32Output, ClearValues);
 
-		TShaderMapRef<FHistogramAtomicCS> ComputeShader(View.ShaderMap);
+		static const auto CVarIgnoreMaterials = IConsoleManager::Get().FindTConsoleVariableDataBool(TEXT("r.AutoExposure.IgnoreMaterials"));
+		static const auto CVarIgnoreMaterialsDebug = IConsoleManager::Get().FindTConsoleVariableDataBool(TEXT("r.AutoExposure.IgnoreMaterials.Debug"));
+
+		FHistogramAtomicCS::FPermutationDomain PermutationVector;
+		PermutationVector.Set<FHistogramAtomicCS::FDivideAlbedoDim>(CVarIgnoreMaterials->GetValueOnRenderThread());
+		PermutationVector.Set<FHistogramAtomicCS::FDivideAlbedoDebugDim>(CVarIgnoreMaterialsDebug->GetValueOnRenderThread());
+
+		auto ComputeShader = View.ShaderMap->GetShader<FHistogramAtomicCS>(PermutationVector);
+
 		FComputeShaderUtils::AddPass(
 			GraphBuilder,
 			RDG_EVENT_NAME("Histogram Atomic %dx%d (CS)", SceneColor.ViewRect.Width(), SceneColor.ViewRect.Height()),
@@ -387,6 +415,7 @@ FRDGTextureRef AddHistogramPass(
 	const FViewInfo& View,
 	const FEyeAdaptationParameters& EyeAdaptationParameters,
 	FScreenPassTexture SceneColor,
+	const FSceneTextureParameters& SceneTextures,
 	FRDGTextureRef EyeAdaptationTexture)
 {
 	if (CVarUseAtomicHistogram.GetValueOnRenderThread() == 1)
@@ -396,6 +425,7 @@ FRDGTextureRef AddHistogramPass(
 			View,
 			EyeAdaptationParameters,
 			SceneColor,
+			SceneTextures,
 			EyeAdaptationTexture);
 	}
 	else
diff --git a/Engine/Source/Runtime/Renderer/Private/PostProcess/PostProcessHistogram.h b/Engine/Source/Runtime/Renderer/Private/PostProcess/PostProcessHistogram.h
--- a/Engine/Source/Runtime/Renderer/Private/PostProcess/PostProcessHistogram.h
+++ b/Engine/Source/Runtime/Renderer/Private/PostProcess/PostProcessHistogram.h
@@ -11,6 +11,7 @@ FRDGTextureRef AddHistogramPass(
 	const FViewInfo& View,
 	const FEyeAdaptationParameters& EyeAdaptationParameters,
 	FScreenPassTexture SceneColor,
+	const FSceneTextureParameters& SceneTextures,
 	FRDGTextureRef EyeAdaptationTexture);
 
 FRDGTextureRef AddLocalExposurePass(
diff --git a/Engine/Source/Runtime/Renderer/Private/PostProcess/PostProcessing.cpp b/Engine/Source/Runtime/Renderer/Private/PostProcess/PostProcessing.cpp
--- a/Engine/Source/Runtime/Renderer/Private/PostProcess/PostProcessing.cpp
+++ b/Engine/Source/Runtime/Renderer/Private/PostProcess/PostProcessing.cpp
@@ -798,10 +798,20 @@ void AddPostProcessingPasses(
 
 		if (bHistogramEnabled)
 		{
+			FScreenPassTexture HistogramSceneColor = bProcessQuarterResolution ? QuarterResSceneColor : HalfResSceneColor;
+
+			static const auto CVarIgnoreMaterials = IConsoleManager::Get().FindTConsoleVariableDataBool(TEXT("r.AutoExposure.IgnoreMaterials"));
+
+			if (CVarIgnoreMaterials->GetValueOnRenderThread())
+			{
+				HistogramSceneColor = OriginalSceneColor;
+			}
+
 			HistogramTexture = AddHistogramPass(
 				GraphBuilder, View,
 				EyeAdaptationParameters,
-				bProcessQuarterResolution ? QuarterResSceneColor : HalfResSceneColor,
+				HistogramSceneColor,
+				SceneTextureParameters,
 				LastEyeAdaptationTexture);
 		}
 
