band-aid for double removal of a non-directional light from the octree during the same frame due to transform update

diff --git a/Engine/Source/Runtime/Renderer/Private/LightSceneInfo.cpp b/Engine/Source/Runtime/Renderer/Private/LightSceneInfo.cpp
--- a/Engine/Source/Runtime/Renderer/Private/LightSceneInfo.cpp
+++ b/Engine/Source/Runtime/Renderer/Private/LightSceneInfo.cpp
@@ -132,8 +132,14 @@ void FLightSceneInfo::RemoveFromScene()
 {
 	if (OctreeId.IsValidId())
 	{
+		if (!Scene->LocalShadowCastingLightOctree.IsValidElementId(OctreeId))
+		{
+			UE_LOG(LogTemp, Warning, TEXT("Octree Element INVALID, possibly already removed this frame"));
+		}
+
 		// Remove the light from the octree.
 		Scene->LocalShadowCastingLightOctree.RemoveElement(OctreeId);
+		OctreeId = FOctreeElementId2();
 	}
 	else
 	{
