Fixed: A few issues with depth access in a mobile renderer
Make SceneDepthAux a platform constant configuration, it can't depend on runtime vars

diff --git a/Engine/Shaders/Private/MobileBasePassPixelShader.usf b/Engine/Shaders/Private/MobileBasePassPixelShader.usf
--- a/Engine/Shaders/Private/MobileBasePassPixelShader.usf
+++ b/Engine/Shaders/Private/MobileBasePassPixelShader.usf
@@ -69,7 +69,7 @@
 #define DEFERRED_SHADING_PATH (MOBILE_DEFERRED_SHADING && ((MATERIALBLENDING_SOLID || MATERIALBLENDING_MASKED) && !MATERIAL_SHADINGMODEL_SINGLELAYERWATER))
 
 // SceneDepthAux always enabled for IOS, on other platforms only enabled when in forward and MobileHDR=true
-#define USE_SCENE_DEPTH_AUX ((METAL_PROFILE || (!DEFERRED_SHADING_PATH && OUTPUT_MOBILE_HDR) || (DEFERRED_SHADING_PATH && USE_GLES_FBF_DEFERRED)) && (MATERIALBLENDING_SOLID || MATERIALBLENDING_MASKED) && !COMPILER_SWITCH)
+#define HAS_SCENE_DEPTH_AUX_OUTPUT (USE_SCENE_DEPTH_AUX && (MATERIALBLENDING_SOLID || MATERIALBLENDING_MASKED))
 
 #if DEFERRED_SHADING_PATH && MOBILE_EXTENDED_GBUFFER
 	#define SV_TargetDepthAux SV_Target5	
@@ -247,7 +247,7 @@ void Main(
 #else
 	, out HALF4_TYPE OutColor : SV_Target0
 #endif
-#if USE_SCENE_DEPTH_AUX
+#if HAS_SCENE_DEPTH_AUX_OUTPUT
 	, out float OutSceneDepthAux : SV_TargetDepthAux
 #endif
 #if OUTPUT_PIXEL_DEPTH_OFFSET
@@ -679,7 +679,7 @@ void Main(
 		#endif
 	#endif
 
-#if USE_SCENE_DEPTH_AUX
+#if HAS_SCENE_DEPTH_AUX_OUTPUT
 	OutSceneDepthAux = SvPosition.z;
 #endif
 
diff --git a/Engine/Shaders/Private/SceneTexturesCommon.ush b/Engine/Shaders/Private/SceneTexturesCommon.ush
--- a/Engine/Shaders/Private/SceneTexturesCommon.ush
+++ b/Engine/Shaders/Private/SceneTexturesCommon.ush
@@ -200,8 +200,8 @@ float LookupDeviceZ( float2 ScreenUV )
 	return VulkanSubpassDepthFetch();
 #elif MOBILE_DEPTHFECTH && (METAL_PROFILE && !MAC)
 	return DepthbufferFetchES2();
-#elif (POST_PROCESS_MATERIAL && !MOBILE_DEFERRED_SHADING)
-	// SceneDepth texture is not accessible during post-processing as we discard it at the end of mobile BasePass (with forward shading) 
+#elif (USE_SCENE_DEPTH_AUX && !MOBILE_DEFERRED_SHADING)
+	// SceneDepth texture is discarded after BasePass (with forward shading) 
 	// instead fetch DeviceZ from SceneDepthAuxTexture
 	return Texture2DSampleLevel(MobileSceneTextures.SceneDepthAuxTexture, MobileSceneTextures.SceneDepthAuxTextureSampler, ScreenUV, 0).r;
 #else
diff --git a/Engine/Source/Developer/ShaderFormatOpenGL/Private/OpenGLShaderCompiler.cpp b/Engine/Source/Developer/ShaderFormatOpenGL/Private/OpenGLShaderCompiler.cpp
--- a/Engine/Source/Developer/ShaderFormatOpenGL/Private/OpenGLShaderCompiler.cpp
+++ b/Engine/Source/Developer/ShaderFormatOpenGL/Private/OpenGLShaderCompiler.cpp
@@ -2291,6 +2291,7 @@ bool GenerateGlslShader(std::string& OutString, GLSLCompileParameters& GLSLCompi
 	}
 
 	// If we are rendering deferred, then we only need SceneDepthAux on devices that don't support framebuffer fetch depth
+	/*
 	if (bIsDeferred)
 	{
 		std::string SceneDepthAux = "layout(location = 4) out highp float out_var_SV_Target4;";
@@ -2317,6 +2318,7 @@ bool GenerateGlslShader(std::string& OutString, GLSLCompileParameters& GLSLCompi
 			}
 		}
 	}
+	*/
 
 	// Fixup packed globals
 	{
diff --git a/Engine/Source/Runtime/Engine/Private/ShaderCompiler/ShaderCompiler.cpp b/Engine/Source/Runtime/Engine/Private/ShaderCompiler/ShaderCompiler.cpp
--- a/Engine/Source/Runtime/Engine/Private/ShaderCompiler/ShaderCompiler.cpp
+++ b/Engine/Source/Runtime/Engine/Private/ShaderCompiler/ShaderCompiler.cpp
@@ -5767,6 +5767,8 @@ void GlobalBeginCompileShader(
 			Input.Environment.SetDefine(TEXT("USE_GLES_FBF_DEFERRED"), bGLESDeferredShading ? 1 : 0);
 			Input.Environment.SetDefine(TEXT("MOBILE_EXTENDED_GBUFFER"), MobileUsesExtenedGBuffer((EShaderPlatform)Target.Platform) ? 1 : 0);
 		}
+
+		Input.Environment.SetDefine(TEXT("USE_SCENE_DEPTH_AUX"), MobileRequiresSceneDepthAux(ShaderPlatform) ? 1 : 0);
 	}
 
 	if (ShaderPlatform == SP_VULKAN_ES3_1_ANDROID || ShaderPlatform == SP_VULKAN_SM5_ANDROID)
diff --git a/Engine/Source/Runtime/RenderCore/Private/RenderUtils.cpp b/Engine/Source/Runtime/RenderCore/Private/RenderUtils.cpp
--- a/Engine/Source/Runtime/RenderCore/Private/RenderUtils.cpp
+++ b/Engine/Source/Runtime/RenderCore/Private/RenderUtils.cpp
@@ -1095,19 +1095,14 @@ RENDERCORE_API bool MobileRequiresSceneDepthAux(const FStaticShaderPlatform Plat
 	static const auto CVarMobileHDR = IConsoleManager::Get().FindTConsoleVariableDataInt(TEXT("r.MobileHDR"));
 	const bool bMobileHDR = (CVarMobileHDR && CVarMobileHDR->GetValueOnAnyThread() != 0);
 
-	// SceneDepth is used on most mobile platforms when forward shading is enabled and always on IOS.
+	// SceneDepthAux is used on tiled GPUs (iOS, Android) with forward shading and on iOS only with deferred
 	if (IsMetalMobilePlatform(Platform))
 	{
 		return true;
 	}
-	else if (IsMobileDeferredShadingEnabled(Platform) && IsAndroidOpenGLESPlatform(Platform) && !GSupportsShaderDepthStencilFetch)
+	else if (bMobileHDR && !IsMobileDeferredShadingEnabled(Platform))
 	{
-		return true;
-	}
-	else if (!IsMobileDeferredShadingEnabled(Platform) && bMobileHDR)
-	{
-		// SceneDepthAux disabled when MobileHDR=false for non-IOS
-		return IsAndroidOpenGLESPlatform(Platform) || IsVulkanMobilePlatform(Platform) || IsSimulatedPlatform(Platform);
+		return (IsAndroidOpenGLESPlatform(Platform) || IsVulkanMobilePlatform(Platform));
 	}
 	return false;
 }
diff --git a/Engine/Source/Runtime/Renderer/Private/MobileSeparateTranslucencyPass.cpp b/Engine/Source/Runtime/Renderer/Private/MobileSeparateTranslucencyPass.cpp
--- a/Engine/Source/Runtime/Renderer/Private/MobileSeparateTranslucencyPass.cpp
+++ b/Engine/Source/Runtime/Renderer/Private/MobileSeparateTranslucencyPass.cpp
@@ -36,12 +36,17 @@ bool IsMobileSeparateTranslucencyActive(const FViewInfo& View)
 
 void AddMobileSeparateTranslucencyPass(FRDGBuilder& GraphBuilder, FScene* Scene, const FViewInfo& View, const FMobileSeparateTranslucencyInputs& Inputs)
 {
-	// GPUCULL_TODO: View.ParallelMeshDrawCommandPasses[EMeshPass::TranslucencyAfterDOF].BuildRenderingCommands(GraphBuilder, Scene->GPUScene);
+	// Whether BasePass uses single pass rendering with framebuffer fetch
+	const bool bSinglePassRendering = !Inputs.bRequiresMultiPass;
 
 	FMobileSeparateTranslucencyPassParameters* PassParameters = GraphBuilder.AllocParameters<FMobileSeparateTranslucencyPassParameters>();
 	PassParameters->RenderTargets[0] = FRenderTargetBinding(Inputs.SceneColor.Texture, ERenderTargetLoadAction::ELoad);
+	if (bSinglePassRendering)
+	{
+		PassParameters->RenderTargets[1] = FRenderTargetBinding(Inputs.SceneDepthAux.Texture, ERenderTargetLoadAction::ELoad);
+		PassParameters->RenderTargets.SubpassHint = ESubpassHint::DepthReadSubpass;
+	}
 	PassParameters->RenderTargets.DepthStencil = FDepthStencilBinding(Inputs.SceneDepth.Texture, ERenderTargetLoadAction::ELoad, ERenderTargetLoadAction::ELoad, FExclusiveDepthStencil::DepthRead_StencilRead);
-	PassParameters->RenderTargets.SubpassHint = ESubpassHint::DepthReadSubpass;
 	PassParameters->View = View.ViewUniformBuffer;
 	EMobileSceneTextureSetupMode SetupMode = EMobileSceneTextureSetupMode::SceneDepth | EMobileSceneTextureSetupMode::SceneDepthAux | EMobileSceneTextureSetupMode::CustomDepth;
 	PassParameters->MobileBasePass = CreateMobileBasePassUniformBuffer(GraphBuilder, View, EMobileBasePass::Translucent, SetupMode);
@@ -52,10 +57,13 @@ void AddMobileSeparateTranslucencyPass(FRDGBuilder& GraphBuilder, FScene* Scene,
 		RDG_EVENT_NAME("SeparateTranslucency %dx%d", View.ViewRect.Width(), View.ViewRect.Height()),
 		PassParameters,
 		ERDGPassFlags::Raster,
-		[&View, PassParameters](FRHICommandList& RHICmdList)
+		[&View, PassParameters, bSinglePassRendering](FRHICommandList& RHICmdList)
 	{
+		if (bSinglePassRendering)
+		{
+			RHICmdList.NextSubpass();
+		}
 		// Set the view family's render target/viewport.
-		RHICmdList.NextSubpass();
 		RHICmdList.SetViewport(View.ViewRect.Min.X, View.ViewRect.Min.Y, 0.0f, View.ViewRect.Max.X, View.ViewRect.Max.Y, 1.0f);
 		View.ParallelMeshDrawCommandPasses[EMeshPass::TranslucencyAfterDOF].DispatchDraw(nullptr, RHICmdList, &PassParameters->InstanceCullingDrawParams);
 	});
diff --git a/Engine/Source/Runtime/Renderer/Private/MobileSeparateTranslucencyPass.h b/Engine/Source/Runtime/Renderer/Private/MobileSeparateTranslucencyPass.h
--- a/Engine/Source/Runtime/Renderer/Private/MobileSeparateTranslucencyPass.h
+++ b/Engine/Source/Runtime/Renderer/Private/MobileSeparateTranslucencyPass.h
@@ -8,7 +8,9 @@
 struct FMobileSeparateTranslucencyInputs
 {
 	FScreenPassTexture SceneColor;
+	FScreenPassTexture SceneDepthAux;
 	FScreenPassTexture SceneDepth;
+	bool bRequiresMultiPass;
 };
 
 void AddMobileSeparateTranslucencyPass(FRDGBuilder& GraphBuilder, FScene* Scene, const FViewInfo& View, const FMobileSeparateTranslucencyInputs& Inputs);
diff --git a/Engine/Source/Runtime/Renderer/Private/MobileShadingRenderer.cpp b/Engine/Source/Runtime/Renderer/Private/MobileShadingRenderer.cpp
--- a/Engine/Source/Runtime/Renderer/Private/MobileShadingRenderer.cpp
+++ b/Engine/Source/Runtime/Renderer/Private/MobileShadingRenderer.cpp
@@ -1049,6 +1049,7 @@ void FMobileSceneRenderer::Render(FRDGBuilder& GraphBuilder)
 				FMobilePostProcessingInputs PostProcessingInputs;
 				PostProcessingInputs.ViewFamilyTexture = ViewFamilyTexture;
 				PostProcessingInputs.SceneTextures = CreateMobileSceneTextureUniformBuffer(GraphBuilder, &SceneTextures, EMobileSceneTextureSetupMode::All);
+				PostProcessingInputs.bRequiresMultiPass = bRequiresMultiPass;
 
 				for (int32 ViewIndex = 0; ViewIndex < Views.Num(); ViewIndex++)
 				{
@@ -1192,7 +1193,7 @@ void FMobileSceneRenderer::RenderForward(FRDGBuilder& GraphBuilder, FRDGTextureR
 		// Split if we need to render translucency in a separate render pass
 		if (bRequiresMultiPass)
 		{
-			RenderForwardMultiPass(GraphBuilder, PassParameters, BasePassRenderTargets, ViewContext, SceneTextures);
+			RenderForwardMultiPass(GraphBuilder, PassParameters, ViewContext, SceneTextures);
 		}
 		else
 		{
@@ -1264,7 +1265,7 @@ void FMobileSceneRenderer::RenderForwardSinglePass(FRDGBuilder& GraphBuilder, FM
 	}
 }
 
-void FMobileSceneRenderer::RenderForwardMultiPass(FRDGBuilder& GraphBuilder, FMobileRenderPassParameters* PassParameters, FRenderTargetBindingSlots& BasePassRenderTargets, FRenderViewContext& ViewContext, FSceneTextures& SceneTextures)
+void FMobileSceneRenderer::RenderForwardMultiPass(FRDGBuilder& GraphBuilder, FMobileRenderPassParameters* PassParameters, FRenderViewContext& ViewContext, FSceneTextures& SceneTextures)
 {
 	GraphBuilder.AddPass(
 		RDG_EVENT_NAME("SceneColorRendering"),
@@ -1309,21 +1310,16 @@ void FMobileSceneRenderer::RenderForwardMultiPass(FRDGBuilder& GraphBuilder, FMo
 		ExclusiveDepthStencil = FExclusiveDepthStencil::DepthRead_StencilWrite;
 	}
 
-	BasePassRenderTargets[0].SetLoadAction(ERenderTargetLoadAction::ELoad);
-	if (bRequiresSceneDepthAux)
-	{
-		BasePassRenderTargets[1].SetLoadAction(ERenderTargetLoadAction::ELoad);
-	}
-	BasePassRenderTargets.DepthStencil.SetDepthLoadAction(ERenderTargetLoadAction::ELoad);
-	BasePassRenderTargets.DepthStencil.SetStencilLoadAction(ERenderTargetLoadAction::ELoad);
-	BasePassRenderTargets.DepthStencil.SetDepthStencilAccess(ExclusiveDepthStencil);
-
 	EMobileSceneTextureSetupMode SetupMode = EMobileSceneTextureSetupMode::SceneDepth | EMobileSceneTextureSetupMode::SceneDepthAux | EMobileSceneTextureSetupMode::CustomDepth;
 	FMobileRenderPassParameters* SecondPassParameters = GraphBuilder.AllocParameters<FMobileRenderPassParameters>();
 	*SecondPassParameters = *PassParameters;
 	SecondPassParameters->MobileBasePass = CreateMobileBasePassUniformBuffer(GraphBuilder, View, EMobileBasePass::Translucent, SetupMode);
 	SecondPassParameters->ReflectionCapture = View.MobileReflectionCaptureUniformBuffer;
-	SecondPassParameters->RenderTargets = BasePassRenderTargets;
+	SecondPassParameters->RenderTargets[0].SetLoadAction(ERenderTargetLoadAction::ELoad);
+	SecondPassParameters->RenderTargets[1] = FRenderTargetBinding();
+	SecondPassParameters->RenderTargets.DepthStencil.SetDepthLoadAction(ERenderTargetLoadAction::ELoad);
+	SecondPassParameters->RenderTargets.DepthStencil.SetStencilLoadAction(ERenderTargetLoadAction::ELoad);
+	SecondPassParameters->RenderTargets.DepthStencil.SetDepthStencilAccess(ExclusiveDepthStencil);
 	
 	const bool bDoOcclusionQueries = (!bIsFullDepthPrepassEnabled && ViewContext.bIsLastView && DoOcclusionQueries());
 	SecondPassParameters->RenderTargets.NumOcclusionQueries = bDoOcclusionQueries ? ComputeNumOcclusionQueriesToBatch() : 0u;
@@ -1523,7 +1519,7 @@ void FMobileSceneRenderer::RenderDeferred(FRDGBuilder& GraphBuilder, const FSort
 
 		if (bRequiresMultiPass)
 		{
-			RenderDeferredMultiPass(GraphBuilder, PassParameters, BasePassRenderTargets, ColorTargets.Num(), ViewContext, SceneTextures, SortedLightSet);
+			RenderDeferredMultiPass(GraphBuilder, PassParameters, ColorTargets.Num(), ViewContext, SceneTextures, SortedLightSet);
 		}
 		else
 		{
@@ -1580,10 +1576,8 @@ void FMobileSceneRenderer::RenderDeferredSinglePass(FRDGBuilder& GraphBuilder, c
 	});
 }
 
-void FMobileSceneRenderer::RenderDeferredMultiPass(FRDGBuilder& GraphBuilder, class FMobileRenderPassParameters* PassParameters, const FRenderTargetBindingSlots& InBasePassRenderTargets, int32 NumColorTargets, FRenderViewContext& ViewContext, FSceneTextures& SceneTextures, const FSortedLightSetSceneInfo& SortedLightSet)
+void FMobileSceneRenderer::RenderDeferredMultiPass(FRDGBuilder& GraphBuilder, class FMobileRenderPassParameters* PassParameters, int32 NumColorTargets, FRenderViewContext& ViewContext, FSceneTextures& SceneTextures, const FSortedLightSetSceneInfo& SortedLightSet)
 {
-	FRenderTargetBindingSlots BasePassRenderTargets = InBasePassRenderTargets;
-	
 	GraphBuilder.AddPass(
 		RDG_EVENT_NAME("SceneColorRendering"),
 		PassParameters,
@@ -1602,16 +1596,6 @@ void FMobileSceneRenderer::RenderDeferredMultiPass(FRDGBuilder& GraphBuilder, cl
 		PostRenderBasePass(RHICmdList, View);
 	});
 
-	// SceneColor + GBuffer write, SceneDepth is read only
-	for (int32 i = 0; i < NumColorTargets; ++i)
-	{
-		BasePassRenderTargets[i].SetLoadAction(ERenderTargetLoadAction::ELoad);
-	}
-
-	BasePassRenderTargets.DepthStencil.SetDepthLoadAction(ERenderTargetLoadAction::ELoad);
-	BasePassRenderTargets.DepthStencil.SetStencilLoadAction(ERenderTargetLoadAction::ELoad);
-	BasePassRenderTargets.DepthStencil.SetDepthStencilAccess(FExclusiveDepthStencil::DepthRead_StencilRead);
-	
 	FViewInfo& View = *ViewContext.ViewInfo;
 	
 	EMobileSceneTextureSetupMode SetupMode = EMobileSceneTextureSetupMode::SceneDepth | EMobileSceneTextureSetupMode::SceneDepthAux | EMobileSceneTextureSetupMode::CustomDepth;
@@ -1619,7 +1603,14 @@ void FMobileSceneRenderer::RenderDeferredMultiPass(FRDGBuilder& GraphBuilder, cl
 	*SecondPassParameters = *PassParameters;
 	SecondPassParameters->MobileBasePass = CreateMobileBasePassUniformBuffer(GraphBuilder, View, EMobileBasePass::Translucent, SetupMode);
 	SecondPassParameters->ReflectionCapture = View.MobileReflectionCaptureUniformBuffer;
-	SecondPassParameters->RenderTargets = BasePassRenderTargets;
+	// SceneColor + GBuffer write, SceneDepth is read only
+	for (int32 i = 0; i < NumColorTargets; ++i)
+	{
+		SecondPassParameters->RenderTargets[i].SetLoadAction(ERenderTargetLoadAction::ELoad);
+	}
+	SecondPassParameters->RenderTargets.DepthStencil.SetDepthLoadAction(ERenderTargetLoadAction::ELoad);
+	SecondPassParameters->RenderTargets.DepthStencil.SetStencilLoadAction(ERenderTargetLoadAction::ELoad);
+	SecondPassParameters->RenderTargets.DepthStencil.SetDepthStencilAccess(FExclusiveDepthStencil::DepthRead_StencilRead);
 
 	GraphBuilder.AddPass(
 		RDG_EVENT_NAME("Decals"),
@@ -1631,20 +1622,18 @@ void FMobileSceneRenderer::RenderDeferredMultiPass(FRDGBuilder& GraphBuilder, cl
 		RenderDecals(RHICmdList, View);
 	});
 
-	// SceneColor write, SceneDepth is read only
-	for (int32 i = 1; i < NumColorTargets; ++i)
-	{
-		BasePassRenderTargets[i] = FRenderTargetBinding();
-	}
-	BasePassRenderTargets.DepthStencil.SetDepthStencilAccess(FExclusiveDepthStencil::DepthRead_StencilWrite);
-
 	SetupMode = EMobileSceneTextureSetupMode::SceneDepth | EMobileSceneTextureSetupMode::SceneDepthAux | EMobileSceneTextureSetupMode::GBuffers | EMobileSceneTextureSetupMode::CustomDepth;
 	auto* ThirdPassParameters = GraphBuilder.AllocParameters<FMobileRenderPassParameters>();
 	*ThirdPassParameters = *PassParameters;
 	ThirdPassParameters->MobileBasePass = CreateMobileBasePassUniformBuffer(GraphBuilder, View, EMobileBasePass::Translucent, SetupMode);
 	ThirdPassParameters->ReflectionCapture = View.MobileReflectionCaptureUniformBuffer;
-	ThirdPassParameters->RenderTargets = BasePassRenderTargets;
-	
+	// SceneColor write, SceneDepth is read only
+	for (int32 i = 1; i < NumColorTargets; ++i)
+	{
+		ThirdPassParameters->RenderTargets[i] = FRenderTargetBinding();
+	}
+	ThirdPassParameters->RenderTargets.DepthStencil.SetDepthStencilAccess(FExclusiveDepthStencil::DepthRead_StencilWrite);
+		
 	const bool bDoOcclusionQueires = (!bIsFullDepthPrepassEnabled && ViewContext.bIsLastView && DoOcclusionQueries());
 	ThirdPassParameters->RenderTargets.NumOcclusionQueries = bDoOcclusionQueires ? ComputeNumOcclusionQueriesToBatch() : 0u;
 
diff --git a/Engine/Source/Runtime/Renderer/Private/PostProcess/PostProcessing.cpp b/Engine/Source/Runtime/Renderer/Private/PostProcess/PostProcessing.cpp
--- a/Engine/Source/Runtime/Renderer/Private/PostProcess/PostProcessing.cpp
+++ b/Engine/Source/Runtime/Renderer/Private/PostProcess/PostProcessing.cpp
@@ -2252,8 +2252,10 @@ void AddMobilePostProcessingPasses(FRDGBuilder& GraphBuilder, FScene* Scene, con
 			PassSequence.AcceptPass(EPass::SeparateTranslucency);
 			FMobileSeparateTranslucencyInputs SeparateTranslucencyInputs;
 			SeparateTranslucencyInputs.SceneColor = SceneColor;
+			SeparateTranslucencyInputs.SceneDepthAux = SceneDepthAux;
 			SeparateTranslucencyInputs.SceneDepth = SceneDepth;
-
+			SeparateTranslucencyInputs.bRequiresMultiPass = Inputs.bRequiresMultiPass;
+			
 			AddMobileSeparateTranslucencyPass(GraphBuilder, Scene, View, SeparateTranslucencyInputs);
 		}
 
diff --git a/Engine/Source/Runtime/Renderer/Private/PostProcess/PostProcessing.h b/Engine/Source/Runtime/Renderer/Private/PostProcess/PostProcessing.h
--- a/Engine/Source/Runtime/Renderer/Private/PostProcess/PostProcessing.h
+++ b/Engine/Source/Runtime/Renderer/Private/PostProcess/PostProcessing.h
@@ -62,6 +62,7 @@ struct FMobilePostProcessingInputs
 {
 	TRDGUniformBufferRef<FMobileSceneTextureUniformParameters> SceneTextures = nullptr;
 	FRDGTextureRef ViewFamilyTexture = nullptr;
+	bool bRequiresMultiPass = false;
 
 	void Validate() const
 	{
diff --git a/Engine/Source/Runtime/Renderer/Private/SceneRendering.h b/Engine/Source/Runtime/Renderer/Private/SceneRendering.h
--- a/Engine/Source/Runtime/Renderer/Private/SceneRendering.h
+++ b/Engine/Source/Runtime/Renderer/Private/SceneRendering.h
@@ -2526,11 +2526,11 @@ protected:
 
 	void RenderForward(FRDGBuilder& GraphBuilder, FRDGTextureRef ViewFamilyTexture, FSceneTextures& SceneTextures);
 	void RenderForwardSinglePass(FRDGBuilder& GraphBuilder, class FMobileRenderPassParameters* PassParameters, struct FRenderViewContext& ViewContext, FSceneTextures& SceneTextures);
-	void RenderForwardMultiPass(FRDGBuilder& GraphBuilder, class FMobileRenderPassParameters* PassParameters, FRenderTargetBindingSlots& BasePassRenderTargets, struct FRenderViewContext& ViewContext, FSceneTextures& SceneTextures);
+	void RenderForwardMultiPass(FRDGBuilder& GraphBuilder, class FMobileRenderPassParameters* PassParameters, struct FRenderViewContext& ViewContext, FSceneTextures& SceneTextures);
 	
 	void RenderDeferred(FRDGBuilder& GraphBuilder, const FSortedLightSetSceneInfo& SortedLightSet, FRDGTextureRef ViewFamilyTexture, FSceneTextures& SceneTextures);
 	void RenderDeferredSinglePass(FRDGBuilder& GraphBuilder, class FMobileRenderPassParameters* PassParameters, struct FRenderViewContext& ViewContext, FSceneTextures& SceneTextures, const FSortedLightSetSceneInfo& SortedLightSet, bool bUsingPixelLocalStorage);
-	void RenderDeferredMultiPass(FRDGBuilder& GraphBuilder, class FMobileRenderPassParameters* PassParameters, const FRenderTargetBindingSlots& BasePassRenderTargets, int32 NumColorTargets, struct FRenderViewContext& ViewContext, FSceneTextures& SceneTextures, const FSortedLightSetSceneInfo& SortedLightSet);
+	void RenderDeferredMultiPass(FRDGBuilder& GraphBuilder, class FMobileRenderPassParameters* PassParameters, int32 NumColorTargets, struct FRenderViewContext& ViewContext, FSceneTextures& SceneTextures, const FSortedLightSetSceneInfo& SortedLightSet);
 	
 	void RenderAmbientOcclusion(FRDGBuilder& GraphBuilder, FRDGTextureRef SceneDepthTexture, FRDGTextureRef AmbientOcclusionTexture);
 
