Fixed: Depth validation assert with a mobile deferred shading enabled

diff --git a/Engine/Source/Runtime/Renderer/Private/MobileShadingRenderer.cpp b/Engine/Source/Runtime/Renderer/Private/MobileShadingRenderer.cpp
--- a/Engine/Source/Runtime/Renderer/Private/MobileShadingRenderer.cpp
+++ b/Engine/Source/Runtime/Renderer/Private/MobileShadingRenderer.cpp
@@ -1622,11 +1622,8 @@ void FMobileSceneRenderer::RenderDeferredMultiPass(FRDGBuilder& GraphBuilder, cl
 		RenderDecals(RHICmdList, View);
 	});
 
-	SetupMode = EMobileSceneTextureSetupMode::SceneDepth | EMobileSceneTextureSetupMode::SceneDepthAux | EMobileSceneTextureSetupMode::GBuffers | EMobileSceneTextureSetupMode::CustomDepth;
 	auto* ThirdPassParameters = GraphBuilder.AllocParameters<FMobileRenderPassParameters>();
-	*ThirdPassParameters = *PassParameters;
-	ThirdPassParameters->MobileBasePass = CreateMobileBasePassUniformBuffer(GraphBuilder, View, EMobileBasePass::Translucent, SetupMode);
-	ThirdPassParameters->ReflectionCapture = View.MobileReflectionCaptureUniformBuffer;
+	*ThirdPassParameters = *SecondPassParameters;
 	// SceneColor write, SceneDepth is read only
 	for (int32 i = 1; i < NumColorTargets; ++i)
 	{
