Fixed crash in a separate translucency rendering with mobile deferred shading enabled

diff --git a/Engine/Source/Runtime/Renderer/Private/MobileSeparateTranslucencyPass.cpp b/Engine/Source/Runtime/Renderer/Private/MobileSeparateTranslucencyPass.cpp
--- a/Engine/Source/Runtime/Renderer/Private/MobileSeparateTranslucencyPass.cpp
+++ b/Engine/Source/Runtime/Renderer/Private/MobileSeparateTranslucencyPass.cpp
@@ -43,7 +43,10 @@ void AddMobileSeparateTranslucencyPass(FRDGBuilder& GraphBuilder, FScene* Scene,
 	PassParameters->RenderTargets[0] = FRenderTargetBinding(Inputs.SceneColor.Texture, ERenderTargetLoadAction::ELoad);
 	if (bSinglePassRendering)
 	{
-		PassParameters->RenderTargets[1] = FRenderTargetBinding(Inputs.SceneDepthAux.Texture, ERenderTargetLoadAction::ELoad);
+		if (MobileRequiresSceneDepthAux(View.GetShaderPlatform()))
+		{
+			PassParameters->RenderTargets[1] = FRenderTargetBinding(Inputs.SceneDepthAux.Texture, ERenderTargetLoadAction::ELoad);
+		}
 		PassParameters->RenderTargets.SubpassHint = ESubpassHint::DepthReadSubpass;
 	}
 	PassParameters->RenderTargets.DepthStencil = FDepthStencilBinding(Inputs.SceneDepth.Texture, ERenderTargetLoadAction::ELoad, ERenderTargetLoadAction::ELoad, FExclusiveDepthStencil::DepthRead_StencilRead);
