Volumetric cloud - added cvar to control a slice bias for empty space skipping start depth.

diff --git a/Engine/Shaders/Private/VolumetricCloud.usf b/Engine/Shaders/Private/VolumetricCloud.usf
--- a/Engine/Shaders/Private/VolumetricCloud.usf
+++ b/Engine/Shaders/Private/VolumetricCloud.usf
@@ -1455,6 +1455,7 @@ void MainCS(uint2 DispatchThreadId : SV_DispatchThreadID)
 
 float2	StartTracingDistanceTextureResolution;
 float	StartTracingSampleVolumeDepth;
+float	StartTracingSliceBias;
 
 RWTexture2D<float> OutStartTracingDistanceTexture;
 
@@ -1508,7 +1509,7 @@ void MainCS(uint3 DispatchThreadId : SV_DispatchThreadID, uint3 GroupThreadId :
 		ResolvedView, CloudLayerParams, TRACING_SHADOW_DISTANCE_OFF, RenderVolumetricCloudParameters.EmptySpaceSkippingSliceDepth);
 
 	float CloudEmptySpaceSkippingOutput0 = GetVolumetricCloudEmptySpaceSkippingOutput0(MaterialParameters);
-	GroupData[GroupThreadId.x][GroupThreadId.y][GroupThreadId.z].FrontDepthKm = CloudEmptySpaceSkippingOutput0 > 0.0f ? (float(GroupThreadId.z) + 0.0f) * SliceDepth * CENTIMETER_TO_KILOMETER : GetDefaultDepth();
+	GroupData[GroupThreadId.x][GroupThreadId.y][GroupThreadId.z].FrontDepthKm = max(0.0f, CloudEmptySpaceSkippingOutput0 > 0.0f ? (float(GroupThreadId.z) + StartTracingSliceBias) * SliceDepth * CENTIMETER_TO_KILOMETER : GetDefaultDepth());
 
 	GroupMemoryBarrierWithGroupSync();
 
diff --git a/Engine/Source/Runtime/Renderer/Private/VolumetricCloudRendering.cpp b/Engine/Source/Runtime/Renderer/Private/VolumetricCloudRendering.cpp
--- a/Engine/Source/Runtime/Renderer/Private/VolumetricCloudRendering.cpp
+++ b/Engine/Source/Runtime/Renderer/Private/VolumetricCloudRendering.cpp
@@ -212,12 +212,17 @@ static TAutoConsoleVariable<float> CVarVolumetricCloudLocalLightsShadowSampleCou
 static TAutoConsoleVariable<int32> CVarVolumetricCloudEmptySpaceSkipping(
 	TEXT("r.VolumetricCloud.EmptySpaceSkipping"), 0,
 	TEXT("Enable/disable empty space skipping to accelerate cloud tracing through emty areas."),
-	ECVF_RenderThreadSafe);
+	ECVF_RenderThreadSafe | ECVF_Scalability);
 
-static TAutoConsoleVariable<int32> CVarVolumetricCloudEmptySpaceSkippingVolumeDepth(
+static TAutoConsoleVariable<float> CVarVolumetricCloudEmptySpaceSkippingVolumeDepth(
 	TEXT("r.VolumetricCloud.EmptySpaceSkipping.VolumeDepth"), 40.0f,
 	TEXT("Set the distance in kilometer over which empty space can be evalauted."),
-	ECVF_RenderThreadSafe);
+	ECVF_RenderThreadSafe | ECVF_Scalability);
+
+static TAutoConsoleVariable<float> CVarVolumetricCloudEmptySpaceSkippingStartTracingSliceBias(
+	TEXT("r.VolumetricCloud.EmptySpaceSkipping.StartTracingSliceBias"), 0.0f,
+	TEXT("The number of slices to bias the start depth with. A valuie of -1 means a bias of one slice towards the view point."),
+	ECVF_RenderThreadSafe | ECVF_Scalability);
 
 ////////////////////////////////////////////////////////////////////////// 
 
@@ -901,6 +906,7 @@ class FRenderVolumetricCloudEmptySpaceSkippingCS : public FMeshMaterialShader
 	BEGIN_SHADER_PARAMETER_STRUCT(FParameters, )
 		SHADER_PARAMETER(FVector2f, StartTracingDistanceTextureResolution)
 		SHADER_PARAMETER(float, StartTracingSampleVolumeDepth)
+		SHADER_PARAMETER(float, StartTracingSliceBias)
 		SHADER_PARAMETER_RDG_UNIFORM_BUFFER(FRenderVolumetricCloudGlobalParameters, VolumetricCloudRenderViewParamsUB)
 		SHADER_PARAMETER_STRUCT_REF(FViewUniformShaderParameters, View)
 		SHADER_PARAMETER_RDG_TEXTURE_UAV(RWTexture2D, OutStartTracingDistanceTexture)
@@ -2269,6 +2275,7 @@ void FSceneRenderer::RenderVolumetricCloudsInternal(FRDGBuilder& GraphBuilder, F
 			PassParameters->OutStartTracingDistanceTexture = GraphBuilder.CreateUAV(FRDGTextureUAVDesc(EmptySpaceSkippingTexture));
 			PassParameters->StartTracingDistanceTextureResolution = FVector2f(EmptySpaceSkippingTextureResolution.X, EmptySpaceSkippingTextureResolution.Y);
 			PassParameters->StartTracingSampleVolumeDepth = GetEmptySpaceSkippingVolumeDepth();
+			PassParameters->StartTracingSliceBias = CVarVolumetricCloudEmptySpaceSkippingStartTracingSliceBias.GetValueOnRenderThread();
 
 			const bool bCloudEmptySpaceSkippingDebug = false;
 			//if (bCloudEmptySpaceSkippingDebug)
