Render movable skylight in a separate render pass on mobile deferred.

diff --git a/Engine/Shaders/Private/MobileBasePassPixelShader.usf b/Engine/Shaders/Private/MobileBasePassPixelShader.usf
--- a/Engine/Shaders/Private/MobileBasePassPixelShader.usf
+++ b/Engine/Shaders/Private/MobileBasePassPixelShader.usf
@@ -130,16 +130,13 @@ void ApplyPixelDepthOffsetForMobileBasePass(inout FMaterialPixelParameters Mater
 	#define PIXELSHADER_EARLYDEPTHSTENCIL 	
 #endif
 
-void GetSkyLighting(bool bApplySkyLighting, half3 SkyLightColor, half3 WorldNormal, bool bEvaluateBackface, inout half IndirectIrradiance, out half3 OutSkyDiffuseLighting, out half3 OutSkySubsurfaceLighting, out half3 OutWaterDiffuseIndirectLuminance)
+void GetSkyLighting(half3 SkyLightColor, half3 WorldNormal, bool bEvaluateBackface, inout half IndirectIrradiance, out half3 OutSkyDiffuseLighting, out half3 OutSkySubsurfaceLighting, out half3 OutWaterDiffuseIndirectLuminance)
 {
 	OutSkyDiffuseLighting = 0;
 	OutSkySubsurfaceLighting = 0;
 	OutWaterDiffuseIndirectLuminance = 0;
 
 #if ENABLE_SKY_LIGHT
-	//@mw todo
-	// TODO: Also need to do specular.
-	if (bApplySkyLighting)
 	{
 		half3 SkyDiffuse = GetSkySHDiffuseSimple(WorldNormal);
 		OutSkyDiffuseLighting = SkyDiffuse * SkyLightColor;
@@ -159,7 +156,7 @@ void GetSkyLighting(bool bApplySkyLighting, half3 SkyLightColor, half3 WorldNorm
 #endif
 }
 
-void GetPrecomputedIndirectLightingAndSkyLight(VTPageTableResult LightmapVTPageTableResult, FVertexFactoryInterpolantsVSToPS Interpolants, bool bEvaluateBackface, half3 DiffuseDir, bool bApplySkyLighting, half3 SkyLightColor, out half3 OutDiffuseIndirectLighting, out half3 OutSubsurfaceIndirectLighting, out half OutIndirectIrradiance, out half3 OutWaterDiffuseIndirectLuminance)
+void GetPrecomputedIndirectLightingAndSkyLight(VTPageTableResult LightmapVTPageTableResult, FVertexFactoryInterpolantsVSToPS Interpolants, bool bEvaluateBackface, half3 DiffuseDir, half3 SkyLightColor, out half3 OutDiffuseIndirectLighting, out half3 OutSubsurfaceIndirectLighting, out half OutIndirectIrradiance, out half3 OutWaterDiffuseIndirectLuminance)
 {
 	//To keep IndirectLightingCache conherence with PC, initialize the OutIndirectIrradiance to zero.
 	OutIndirectIrradiance = 0;
@@ -220,7 +217,7 @@ void GetPrecomputedIndirectLightingAndSkyLight(VTPageTableResult LightmapVTPageT
 
 	half3 SkyDiffuseLighting;
 	half3 SkySubsurfaceLighting;
-	GetSkyLighting(bApplySkyLighting, SkyLightColor, DiffuseDir, bEvaluateBackface, OutIndirectIrradiance, SkyDiffuseLighting, SkySubsurfaceLighting, OutWaterDiffuseIndirectLuminance);
+	GetSkyLighting(SkyLightColor, DiffuseDir, bEvaluateBackface, OutIndirectIrradiance, SkyDiffuseLighting, SkySubsurfaceLighting, OutWaterDiffuseIndirectLuminance);
 
 	OutDiffuseIndirectLighting += SkyDiffuseLighting;
 	OutSubsurfaceIndirectLighting += SkySubsurfaceLighting;
@@ -529,11 +526,10 @@ void Main(
 #endif
 
 	half NoL = max(0, dot(GBuffer.WorldNormal, MobileDirectionalLight.DirectionalLightDirectionAndShadowTransition.xyz));
-	bool bApplySkyLighting = !(MATERIAL_TWOSIDED && LQ_TEXTURE_LIGHTMAP) || NoL == 0;
 
 	const bool bEvaluateBackface = GetShadingModelRequiresBackfaceLighting(GBuffer.ShadingModelID);
 
-	GetPrecomputedIndirectLightingAndSkyLight(LightmapVTPageTableResult, Interpolants, bEvaluateBackface, DiffuseDir, bApplySkyLighting, ResolvedView.SkyLightColor.rgb, DiffuseIndirectLighting, SubsurfaceIndirectLighting, IndirectIrradiance, WaterDiffuseIndirectLuminance);
+	GetPrecomputedIndirectLightingAndSkyLight(LightmapVTPageTableResult, Interpolants, bEvaluateBackface, DiffuseDir, ResolvedView.SkyLightColor.rgb, DiffuseIndirectLighting, SubsurfaceIndirectLighting, IndirectIrradiance, WaterDiffuseIndirectLuminance);
 
 	// Apply MaterialAO since we don't have the DiffuseIndirectComposite pass on mobile deferred.
 	IndirectIrradiance *= GBuffer.GBufferAO;
diff --git a/Engine/Shaders/Private/MobileDeferredShading.usf b/Engine/Shaders/Private/MobileDeferredShading.usf
--- a/Engine/Shaders/Private/MobileDeferredShading.usf
+++ b/Engine/Shaders/Private/MobileDeferredShading.usf
@@ -51,6 +51,54 @@ half ComputeLightFunctionMultiplier(float3 TranslatedWorldPosition)
 #endif
 }
 
+half3 SkyLightDiffuseMobile(FGBufferData GBuffer, half3 DiffuseColor, half3 V, inout half IndirectIrradiance)
+{
+	half3 Lighting = 0;
+	half3 SkyDiffuseLookUpNormal = GBuffer.WorldNormal;
+
+	if (GBuffer.ShadingModelID == SHADINGMODELID_TWOSIDED_FOLIAGE)
+	{
+		half3 SubsurfaceLookup = GetSkySHDiffuseSimple(-GBuffer.WorldNormal) * View.SkyLightColor.rgb;
+		half3 SubsurfaceColor = ExtractSubsurfaceColor(GBuffer);
+		Lighting += SubsurfaceLookup * SubsurfaceColor;
+	}
+
+	if (GBuffer.ShadingModelID == SHADINGMODELID_SUBSURFACE || GBuffer.ShadingModelID == SHADINGMODELID_PREINTEGRATED_SKIN)
+	{
+		half3 SubsurfaceColor = ExtractSubsurfaceColor(GBuffer);
+		// Add subsurface energy to diffuse
+		DiffuseColor += SubsurfaceColor;
+	}
+
+	if (GBuffer.ShadingModelID == SHADINGMODELID_HAIR)
+	{
+		half3 N = GBuffer.WorldNormal;
+		half3 L = normalize(V - N * dot(V, N));
+		SkyDiffuseLookUpNormal = L;
+		FHairTransmittanceData TransmittanceData = InitHairTransmittanceData(true);
+		DiffuseColor = 2 * PI * HairShading(GBuffer, L, V, N, 1, TransmittanceData, 0, 0.2, uint2(0, 0));
+	}
+
+	if (GBuffer.ShadingModelID == SHADINGMODELID_CLOTH)
+	{
+		half3 ClothFuzz = ExtractSubsurfaceColor(GBuffer);
+		DiffuseColor += ClothFuzz * GBuffer.CustomData.a;
+	}
+
+	// Compute the preconvolved incoming lighting with the bent normal direction
+	half3 DiffuseLookup = GetSkySHDiffuseSimple(SkyDiffuseLookUpNormal) * View.SkyLightColor.rgb;
+
+	// And accumulate the lighting
+	Lighting += DiffuseLookup  * DiffuseColor;
+	Lighting *= View.PreExposure;
+
+#if ALLOW_STATIC_LIGHTING
+	IndirectIrradiance += Luminance(DiffuseLookup);
+#endif
+
+	return Lighting;
+}
+
 void ReflectionEnvironmentSkyLighting(
 	FGBufferData GBuffer,
 	half3 CameraVector,
@@ -59,12 +107,30 @@ void ReflectionEnvironmentSkyLighting(
 	uint GridIndex,
 	inout FLightAccumulator DirectLighting)
 {
+	half IndirectIrradiance = GBuffer.IndirectIrradiance;
+
+#if ENABLE_DYNAMIC_SKY_LIGHT
+
+	half3 N = GBuffer.WorldNormal;
+	half3 V = -CameraVector;
+	half NoV = saturate(abs(dot(N, V)) + 1e-5);
+
+	half3 DiffuseColor = GBuffer.DiffuseColor;
+	half3 SpecularColor = GBuffer.SpecularColor;
+	RemapClearCoatDiffuseAndSpecularColor(GBuffer, NoV, DiffuseColor, SpecularColor);
+
+	half3 SkyLighting = SkyLightDiffuseMobile(GBuffer, DiffuseColor, V, IndirectIrradiance);
+
+	const bool bNeedsSeparateSubsurfaceLightAccumulation = UseSubsurfaceProfile(GBuffer.ShadingModelID);
+	LightAccumulator_Add(DirectLighting, SkyLighting, SkyLighting, 1.0f, bNeedsSeparateSubsurfaceLightAccumulation);
+#endif // ENABLE_DYNAMIC_SKY_LIGHT 
+
 	// IBL
 	AccumulateReflection(GBuffer
 		, CameraVector
 		, TranslatedWorldPosition
 		, ReflectionVector
-		, GBuffer.IndirectIrradiance
+		, IndirectIrradiance
 		, GridIndex
 		, DirectLighting);
 }
diff --git a/Engine/Source/Runtime/Renderer/Private/MobileBasePass.cpp b/Engine/Source/Runtime/Renderer/Private/MobileBasePass.cpp
--- a/Engine/Source/Runtime/Renderer/Private/MobileBasePass.cpp
+++ b/Engine/Source/Runtime/Renderer/Private/MobileBasePass.cpp
@@ -669,7 +669,8 @@ bool FMobileBasePassMeshProcessor::Process(
 	
 	if (Scene && Scene->SkyLight)
 	{
-		bEnableSkyLight = ShadingModels.IsLit() && Scene->ShouldRenderSkylightInBasePass(BlendMode);
+		// Uses bTranslucentBasePass instead of BlendMode to handle single layer water meshes.
+		bEnableSkyLight = ShadingModels.IsLit() && Scene->ShouldRenderSkylightInBasePass(bTranslucentBasePass ? BLEND_Translucent : BLEND_Opaque);
 	}
 
 	bool bEnableLocalLights = false;
diff --git a/Engine/Source/Runtime/Renderer/Private/MobileDeferredShadingPass.cpp b/Engine/Source/Runtime/Renderer/Private/MobileDeferredShadingPass.cpp
--- a/Engine/Source/Runtime/Renderer/Private/MobileDeferredShadingPass.cpp
+++ b/Engine/Source/Runtime/Renderer/Private/MobileDeferredShadingPass.cpp
@@ -46,6 +46,7 @@ class FMobileDirectionalLightFunctionPS : public FMaterialShader
 	class FEnableClustredReflection		: SHADER_PERMUTATION_BOOL("ENABLE_CLUSTERED_REFLECTION");
 	class FEnablePlanarReflection		: SHADER_PERMUTATION_BOOL("ENABLE_PLANAR_REFLECTION");
 	class FEnableSkyLight				: SHADER_PERMUTATION_BOOL("ENABLE_SKY_LIGHT");
+	class FEnableDynamicSkyLight		: SHADER_PERMUTATION_BOOL("ENABLE_DYNAMIC_SKY_LIGHT");
 	class FEnableCSM					: SHADER_PERMUTATION_BOOL("ENABLE_MOBILE_CSM");
 	class FShadowQuality				: SHADER_PERMUTATION_RANGE_INT("MOBILE_SHADOW_QUALITY", 1, 3); // not using Quality=0
 	
@@ -55,6 +56,7 @@ class FMobileDirectionalLightFunctionPS : public FMaterialShader
 		FEnableClustredReflection, 
 		FEnablePlanarReflection,
 		FEnableSkyLight,
+		FEnableDynamicSkyLight,
 		FEnableCSM, 
 		FShadowQuality>;
 
@@ -95,6 +97,11 @@ class FMobileDirectionalLightFunctionPS : public FMaterialShader
 			PermutationVector.Set<FEnableShadingModelSupport>(false);
 		}
 
+		if (!PermutationVector.Get<FEnableSkyLight>())
+		{
+			PermutationVector.Set<FEnableDynamicSkyLight>(false);
+		}
+
 		return PermutationVector;
 	}
 
@@ -117,11 +124,12 @@ class FMobileDirectionalLightFunctionPS : public FMaterialShader
 		return true;
 	}
 
-	static FPermutationDomain BuildPermutationVector(const FViewInfo& View, bool bInlineReflectionAndSky, bool bShadingModelSupport, bool bDynamicShadows, bool bSkyLight, bool bPlanarReflection)
+	static FPermutationDomain BuildPermutationVector(const FViewInfo& View, bool bInlineReflectionAndSky, bool bShadingModelSupport, bool bDynamicShadows, bool bSkyLight, bool bDynamicSkyLight, bool bPlanarReflection)
 	{
 		bool bUseClusteredLights = UseClusteredDeferredShading(View.GetShaderPlatform());
 		bool bClustredReflection = bInlineReflectionAndSky && (View.NumBoxReflectionCaptures + View.NumSphereReflectionCaptures) > 0;
 		bool bEnableSkyLight = bInlineReflectionAndSky && bSkyLight;
+		bool bEnableDynamicSkyLight = bInlineReflectionAndSky && bDynamicSkyLight;
 		const bool bMobileUsesShadowMaskTexture = MobileUsesShadowMaskTexture(View.GetShaderPlatform());
 		int32 ShadowQuality = bDynamicShadows && !bMobileUsesShadowMaskTexture ? (int32)GetShadowQuality() : 0;
 				
@@ -131,6 +139,7 @@ class FMobileDirectionalLightFunctionPS : public FMaterialShader
 		PermutationVector.Set<FMobileDirectionalLightFunctionPS::FEnableClustredReflection>(bClustredReflection);
 		PermutationVector.Set<FMobileDirectionalLightFunctionPS::FEnablePlanarReflection>(bPlanarReflection);
 		PermutationVector.Set<FMobileDirectionalLightFunctionPS::FEnableSkyLight>(bEnableSkyLight);
+		PermutationVector.Set<FMobileDirectionalLightFunctionPS::FEnableDynamicSkyLight>(bEnableDynamicSkyLight);
 		PermutationVector.Set<FMobileDirectionalLightFunctionPS::FEnableCSM>(ShadowQuality > 0);
 		PermutationVector.Set<FMobileDirectionalLightFunctionPS::FShadowQuality>(FMath::Clamp(ShadowQuality, 1, 3));
 		return PermutationVector;
@@ -240,12 +249,14 @@ public:
 	class FEnableClustredReflection		: SHADER_PERMUTATION_BOOL("ENABLE_CLUSTERED_REFLECTION");
 	class FEnablePlanarReflection		: SHADER_PERMUTATION_BOOL("ENABLE_PLANAR_REFLECTION");
 	class FEnableSkyLight				: SHADER_PERMUTATION_BOOL("ENABLE_SKY_LIGHT");
+	class FEnableDynamicSkyLight : SHADER_PERMUTATION_BOOL("ENABLE_DYNAMIC_SKY_LIGHT");
 	
 	using FPermutationDomain = TShaderPermutationDomain<
 		FEnableShadingModelSupport, 
 		FEnableClustredReflection, 
 		FEnablePlanarReflection,
-		FEnableSkyLight
+		FEnableSkyLight,
+		FEnableDynamicSkyLight
 	>;
 
 	BEGIN_SHADER_PARAMETER_STRUCT(FParameters, )
@@ -267,6 +278,11 @@ public:
 			return false;
 		}
 
+		if (!PermutationVector.Get<FEnableSkyLight>() && PermutationVector.Get<FEnableDynamicSkyLight>())
+		{
+			return false;
+		}
+
 		return true;
 	}
 
@@ -331,6 +347,7 @@ void RenderReflectionEnvironmentSkyLighting(FRHICommandListImmediate& RHICmdList
 {
 	// Skylights with static lighting already had their diffuse contribution baked into lightmaps
 	const bool bSkyLight = Scene.SkyLight && !Scene.SkyLight->bHasStaticLighting && View.Family->EngineShowFlags.SkyLighting;
+	const bool bDynamicSkyLight = bSkyLight && !Scene.SkyLight->bWantsStaticShadowing;
 	const bool bClustredReflection = (View.NumBoxReflectionCaptures + View.NumSphereReflectionCaptures) > 0;
 	const bool bPlanarReflection = Scene.GetForwardPassGlobalPlanarReflection() != nullptr;
 	if (!(bSkyLight || bClustredReflection || bPlanarReflection))
@@ -374,6 +391,7 @@ void RenderReflectionEnvironmentSkyLighting(FRHICommandListImmediate& RHICmdList
 		PermutationVector.Set<FMobileReflectionEnvironmentSkyLightingPS::FEnableClustredReflection>(bClustredReflection);
 		PermutationVector.Set<FMobileReflectionEnvironmentSkyLightingPS::FEnablePlanarReflection>(bPlanarReflection);
 		PermutationVector.Set<FMobileReflectionEnvironmentSkyLightingPS::FEnableSkyLight>(bSkyLight);
+		PermutationVector.Set<FMobileReflectionEnvironmentSkyLightingPS::FEnableDynamicSkyLight>(bDynamicSkyLight);
 		TShaderMapRef<FMobileReflectionEnvironmentSkyLightingPS> PixelShader(View.ShaderMap, PermutationVector);
 		
 		GraphicsPSOInit.BoundShaderState.VertexDeclarationRHI = GFilterVertexDeclaration.VertexDeclarationRHI;
@@ -473,6 +491,7 @@ static void RenderDirectionalLight(FRHICommandListImmediate& RHICmdList, const F
 
 	// Skylights with static lighting already had their diffuse contribution baked into lightmaps
 	const bool bSkyLight = Scene.SkyLight && !Scene.SkyLight->bHasStaticLighting && View.Family->EngineShowFlags.SkyLighting;
+	const bool bDynamicSkyLight = bSkyLight && !Scene.SkyLight->bWantsStaticShadowing;
 	const bool bDynamicShadows = DirectionalLight.Proxy->CastsDynamicShadow() && (LightingChannel == 0u) && View.Family->EngineShowFlags.DynamicShadows;
 	const bool bPlanarReflection = Scene.GetForwardPassGlobalPlanarReflection() != nullptr;
 
@@ -494,6 +513,7 @@ static void RenderDirectionalLight(FRHICommandListImmediate& RHICmdList, const F
 			bEnableShadingModelSupport,
 			bDynamicShadows,
 			bSkyLight,
+			bDynamicSkyLight,
 			bPlanarReflection
 		);
 		FCachedLightMaterial LightMaterial;
diff --git a/Engine/Source/Runtime/Renderer/Private/ScenePrivate.h b/Engine/Source/Runtime/Renderer/Private/ScenePrivate.h
--- a/Engine/Source/Runtime/Renderer/Private/ScenePrivate.h
+++ b/Engine/Source/Runtime/Renderer/Private/ScenePrivate.h
@@ -3402,7 +3402,7 @@ public:
 			bRenderSkyLight = bRenderSkyLight
 				&& ((ReadOnlyCVARCache.bEnableStationarySkylight && SkyLight->bWantsStaticShadowing)
 					|| (!SkyLight->bWantsStaticShadowing
-						&& (IsForwardShadingEnabled(GetShaderPlatform()) || IsMobilePlatform(GetShaderPlatform()))));
+						&& (IsForwardShadingEnabled(GetShaderPlatform()) || (IsMobilePlatform(GetShaderPlatform()) && !IsMobileDeferredShadingEnabled(GetShaderPlatform())))));
 		}
 
 		return bRenderSkyLight;
