Don't bind view uniform buffer to Global SDF update shaders

diff --git a/Engine/Shaders/Private/DistanceField/GlobalDistanceFieldDebug.usf b/Engine/Shaders/Private/DistanceField/GlobalDistanceFieldDebug.usf
--- a/Engine/Shaders/Private/DistanceField/GlobalDistanceFieldDebug.usf
+++ b/Engine/Shaders/Private/DistanceField/GlobalDistanceFieldDebug.usf
@@ -14,6 +14,7 @@ RWStructuredBuffer<uint> RWPageStatsBuffer;
 StructuredBuffer<uint> PageStatsBuffer;
 Texture3D<uint> PageTableCombinedTexture;
 Texture3D PageAtlasTexture;
+SamplerState DistanceFieldSampler;
 
 [numthreads(THREADGROUP_SIZE, THREADGROUP_SIZE, THREADGROUP_SIZE)]
 void GlobalDistanceFieldPageStatsCS(
@@ -45,7 +46,7 @@ void GlobalDistanceFieldPageStatsCS(
 
 						uint3 PageAtlasOffset = GlobalDistanceFieldPageLinearIndexToPageAtlasOffset(Page) * GLOBAL_DISTANCE_FIELD_PAGE_RESOLUTION_IN_ATLAS;
 						float3 PageUV = (PageAtlasOffset + CoordInPage) * InvPageAtlasSize;
-						float DistanceFieldValue = Texture3DSampleLevel(PageAtlasTexture, GlobalTrilinearWrappedSampler, PageUV, 0).x;
+						float DistanceFieldValue = Texture3DSampleLevel(PageAtlasTexture, DistanceFieldSampler, PageUV, 0).x;
 
 						if (DistanceFieldValue > 0.0f && DistanceFieldValue < 1.0f)
 						{
diff --git a/Engine/Source/Runtime/Renderer/Private/DistanceFieldStreaming.cpp b/Engine/Source/Runtime/Renderer/Private/DistanceFieldStreaming.cpp
--- a/Engine/Source/Runtime/Renderer/Private/DistanceFieldStreaming.cpp
+++ b/Engine/Source/Runtime/Renderer/Private/DistanceFieldStreaming.cpp
@@ -1472,8 +1472,8 @@ void FDistanceFieldSceneData::UpdateDistanceFieldAtlas(
 			for (int32 StartBrickIndex = 0; StartBrickIndex < NumBrickUploads; StartBrickIndex += MaxBrickUploadsPerPass)
 			{
 				const int32 NumBrickUploadsThisPass = FMath::Min(MaxBrickUploadsPerPass, NumBrickUploads - StartBrickIndex);
-				FScatterUploadDistanceFieldAtlasCS::FParameters* PassParameters = GraphBuilder.AllocParameters<FScatterUploadDistanceFieldAtlasCS::FParameters>();
 
+				auto* PassParameters = GraphBuilder.AllocParameters<FScatterUploadDistanceFieldAtlasCS::FParameters>();
 				PassParameters->RWDistanceFieldBrickAtlas = GraphBuilder.CreateUAV(DistanceFieldBrickVolumeTextureRDG);
 				PassParameters->BrickUploadCoordinates = AtlasUpload.BrickUploadCoordinatesBuffer.SRV;
 				PassParameters->BrickUploadData = AtlasUpload.BrickUploadDataBuffer.SRV;
diff --git a/Engine/Source/Runtime/Renderer/Private/GlobalDistanceField.cpp b/Engine/Source/Runtime/Renderer/Private/GlobalDistanceField.cpp
--- a/Engine/Source/Runtime/Renderer/Private/GlobalDistanceField.cpp
+++ b/Engine/Source/Runtime/Renderer/Private/GlobalDistanceField.cpp
@@ -1427,7 +1427,6 @@ class FCullObjectsToGridCS : public FGlobalShader
 	SHADER_USE_PARAMETER_STRUCT(FCullObjectsToGridCS, FGlobalShader);
 
 	BEGIN_SHADER_PARAMETER_STRUCT(FParameters, )
-		SHADER_PARAMETER_STRUCT_REF(FViewUniformShaderParameters, View)
 		SHADER_PARAMETER_RDG_BUFFER_UAV(RWStructuredBuffer<uint>, RWCullGridAllocator)
 		SHADER_PARAMETER_RDG_BUFFER_UAV(RWStructuredBuffer<uint>, RWCullGridObjectHeader)
 		SHADER_PARAMETER_RDG_BUFFER_UAV(RWStructuredBuffer<uint>, RWCullGridObjectArray)
@@ -1486,7 +1485,6 @@ class FCompositeObjectsIntoObjectGridPagesCS : public FGlobalShader
 	SHADER_USE_PARAMETER_STRUCT(FCompositeObjectsIntoObjectGridPagesCS, FGlobalShader);
 
 	BEGIN_SHADER_PARAMETER_STRUCT(FParameters, )
-		SHADER_PARAMETER_STRUCT_REF(FViewUniformShaderParameters, View)
 		SHADER_PARAMETER_RDG_BUFFER_UAV(RWStructuredBuffer<uint4>, RWPageObjectGridBuffer)
 		RDG_BUFFER_ACCESS(IndirectArgBuffer, ERHIAccess::IndirectArgs)
 		SHADER_PARAMETER_RDG_TEXTURE(Texture3D<uint>, PageTableLayerTexture)
@@ -1529,7 +1527,6 @@ class FCompositeObjectsIntoPagesCS : public FGlobalShader
 	SHADER_USE_PARAMETER_STRUCT(FCompositeObjectsIntoPagesCS, FGlobalShader);
 
 	BEGIN_SHADER_PARAMETER_STRUCT(FParameters, )
-		SHADER_PARAMETER_STRUCT_REF(FViewUniformShaderParameters, View)
 		SHADER_PARAMETER_RDG_TEXTURE_UAV(RWTexture3D<UNORM float>, RWPageAtlasTexture)
 		SHADER_PARAMETER_RDG_TEXTURE_UAV(RWTexture3D<UNORM float>, RWCoverageAtlasTexture)
 		RDG_BUFFER_ACCESS(ComposeIndirectArgBuffer, ERHIAccess::IndirectArgs)
@@ -1797,10 +1794,10 @@ class FGlobalDistanceFieldPageStatsCS : public FGlobalShader
 	SHADER_USE_PARAMETER_STRUCT(FGlobalDistanceFieldPageStatsCS, FGlobalShader)
 
 	BEGIN_SHADER_PARAMETER_STRUCT(FParameters, )
-		SHADER_PARAMETER_STRUCT_REF(FViewUniformShaderParameters, View)
 		SHADER_PARAMETER_RDG_BUFFER_UAV(RWStructuredBuffer<uint>, RWPageStatsBuffer)
 		SHADER_PARAMETER_RDG_TEXTURE(Texture3D<uint>, PageTableCombinedTexture)
 		SHADER_PARAMETER_RDG_TEXTURE(Texture3D<float>, PageAtlasTexture)
+		SHADER_PARAMETER_SAMPLER(SamplerState, DistanceFieldSampler)
 		SHADER_PARAMETER(uint32, ClipmapSizeInPages)
 		SHADER_PARAMETER(uint32, ClipmapIndex)
 		SHADER_PARAMETER(FVector3f, InvPageAtlasSize)
@@ -2359,7 +2356,6 @@ void UpdateGlobalDistanceFieldVolume(
 									UTexture2D* VisibilityTexture = It.Key().Visibility;
 
 									FMarkHeightfieldPagesCS::FParameters* PassParameters = GraphBuilder.AllocParameters<FMarkHeightfieldPagesCS::FParameters>();
-									PassParameters->View = View.ViewUniformBuffer;
 									PassParameters->RWMarkedHeightfieldPageBuffer = GraphBuilder.CreateUAV(MarkedHeightfieldPageBuffer, PF_R32_UINT);
 									PassParameters->PageUpdateIndirectArgBuffer = PageUpdateIndirectArgBuffer;
 									PassParameters->PageUpdateTileBuffer = GraphBuilder.CreateSRV(PageUpdateTileBuffer, PF_R32_UINT);
@@ -2449,7 +2445,6 @@ void UpdateGlobalDistanceFieldVolume(
 							AddClearUAVPass(GraphBuilder, GraphBuilder.CreateUAV(CullGridObjectHeader, PF_R32_UINT), 0);
 
 							FCullObjectsToGridCS::FParameters* PassParameters = GraphBuilder.AllocParameters<FCullObjectsToGridCS::FParameters>();
-							PassParameters->View = View.ViewUniformBuffer;
 							PassParameters->RWCullGridAllocator = GraphBuilder.CreateUAV(CullGridAllocator, PF_R32_UINT);
 							PassParameters->RWCullGridObjectHeader = GraphBuilder.CreateUAV(CullGridObjectHeader, PF_R32_UINT);
 							PassParameters->RWCullGridObjectArray = GraphBuilder.CreateUAV(CullGridObjectArray, PF_R32_UINT);
@@ -2612,7 +2607,6 @@ void UpdateGlobalDistanceFieldVolume(
 							// Composite mesh SDFs into allocated global distance field pages
 							{
 								FCompositeObjectsIntoPagesCS::FParameters* PassParameters = GraphBuilder.AllocParameters<FCompositeObjectsIntoPagesCS::FParameters>();
-								PassParameters->View = View.ViewUniformBuffer;
 								PassParameters->RWPageAtlasTexture = GraphBuilder.CreateUAV(PageAtlasTexture);
 								PassParameters->RWCoverageAtlasTexture = CoverageAtlasTexture ? GraphBuilder.CreateUAV(CoverageAtlasTexture) : nullptr;
 								PassParameters->ComposeIndirectArgBuffer = PageComposeIndirectArgBuffer;
@@ -2647,7 +2641,6 @@ void UpdateGlobalDistanceFieldVolume(
 							if (PageObjectGridBuffer)
 							{
 								FCompositeObjectsIntoObjectGridPagesCS::FParameters* PassParameters = GraphBuilder.AllocParameters<FCompositeObjectsIntoObjectGridPagesCS::FParameters>();
-								PassParameters->View = View.ViewUniformBuffer;
 								PassParameters->RWPageObjectGridBuffer = GraphBuilder.CreateUAV(PageObjectGridBuffer);
 								PassParameters->IndirectArgBuffer = BuildObjectGridIndirectArgBuffer;
 								PassParameters->DistanceFieldObjectBuffers = DistanceField::SetupObjectBufferParameters(GraphBuilder, DistanceFieldSceneData);
@@ -2694,7 +2687,6 @@ void UpdateGlobalDistanceFieldVolume(
 									// Compose heightfields into global SDF pages
 									{
 										FComposeHeightfieldsIntoPagesCS::FParameters* PassParameters = GraphBuilder.AllocParameters<FComposeHeightfieldsIntoPagesCS::FParameters>();
-										PassParameters->View = View.ViewUniformBuffer;
 										PassParameters->RWPageAtlasTexture = GraphBuilder.CreateUAV(PageAtlasTexture);
 										PassParameters->RWCoverageAtlasTexture = CoverageAtlasTexture ? GraphBuilder.CreateUAV(CoverageAtlasTexture) : nullptr;
 										PassParameters->ComposeIndirectArgBuffer = PageComposeHeightfieldIndirectArgBuffer;
@@ -2738,7 +2730,6 @@ void UpdateGlobalDistanceFieldVolume(
 									if (PageObjectGridBuffer)
 									{
 										FCompositeHeightfieldsIntoObjectGridPagesCS::FParameters* PassParameters = GraphBuilder.AllocParameters<FCompositeHeightfieldsIntoObjectGridPagesCS::FParameters>();
-										PassParameters->View = View.ViewUniformBuffer;
 										PassParameters->RWPageObjectGridBuffer = GraphBuilder.CreateUAV(PageObjectGridBuffer);
 										PassParameters->ComposeIndirectArgBuffer = PageComposeHeightfieldIndirectArgBuffer;
 										PassParameters->ComposeTileBuffer = GraphBuilder.CreateSRV(PageComposeHeightfieldTileBuffer, PF_R32_UINT);
@@ -2853,9 +2844,9 @@ void UpdateGlobalDistanceFieldVolume(
 			{
 				FGlobalDistanceFieldPageStatsCS::FParameters* PassParameters = GraphBuilder.AllocParameters<FGlobalDistanceFieldPageStatsCS::FParameters>();
 				PassParameters->RWPageStatsBuffer = GraphBuilder.CreateUAV(PageStatsBuffer);
-				PassParameters->View = View.ViewUniformBuffer;
 				PassParameters->PageTableCombinedTexture = PageTableCombinedTexture;
 				PassParameters->PageAtlasTexture = PageAtlasTexture;
+				PassParameters->DistanceFieldSampler = TStaticSamplerState<SF_Bilinear, AM_Wrap, AM_Wrap, AM_Wrap>::GetRHI();
 				PassParameters->ClipmapSizeInPages = GlobalDistanceField::GetPageTableTextureResolution(bLumenEnabled, View.FinalPostProcessSettings.LumenSceneViewDistance).X;
 				PassParameters->ClipmapIndex = ClipmapIndex;
 				PassParameters->InvPageAtlasSize = FVector3f(GlobalDistanceFieldInfo.ParameterData.InvPageAtlasSize);
diff --git a/Engine/Source/Runtime/Renderer/Private/GlobalDistanceFieldHeightfields.h b/Engine/Source/Runtime/Renderer/Private/GlobalDistanceFieldHeightfields.h
--- a/Engine/Source/Runtime/Renderer/Private/GlobalDistanceFieldHeightfields.h
+++ b/Engine/Source/Runtime/Renderer/Private/GlobalDistanceFieldHeightfields.h
@@ -12,7 +12,6 @@ class FMarkHeightfieldPagesCS : public FGlobalShader
 	SHADER_USE_PARAMETER_STRUCT(FMarkHeightfieldPagesCS, FGlobalShader);
 
 	BEGIN_SHADER_PARAMETER_STRUCT(FParameters, )
-		SHADER_PARAMETER_STRUCT_REF(FViewUniformShaderParameters, View)
 		SHADER_PARAMETER_RDG_BUFFER_UAV(RWStructuredBuffer<uint>, RWMarkedHeightfieldPageBuffer)
 		RDG_BUFFER_ACCESS(PageUpdateIndirectArgBuffer, ERHIAccess::IndirectArgs)
 		SHADER_PARAMETER_RDG_BUFFER_SRV(StructuredBuffer<uint>, PageUpdateTileBuffer)
@@ -118,7 +117,6 @@ class FComposeHeightfieldsIntoPagesCS : public FGlobalShader
 	SHADER_USE_PARAMETER_STRUCT(FComposeHeightfieldsIntoPagesCS, FGlobalShader);
 
 	BEGIN_SHADER_PARAMETER_STRUCT(FParameters, )
-		SHADER_PARAMETER_STRUCT_REF(FViewUniformShaderParameters, View)
 		SHADER_PARAMETER_RDG_TEXTURE_UAV(RWTexture3D<float>, RWPageAtlasTexture)
 		SHADER_PARAMETER_RDG_TEXTURE_UAV(RWTexture3D<UNORM float>, RWCoverageAtlasTexture)
 		RDG_BUFFER_ACCESS(ComposeIndirectArgBuffer, ERHIAccess::IndirectArgs)
@@ -173,7 +171,6 @@ class FCompositeHeightfieldsIntoObjectGridPagesCS : public FGlobalShader
 	SHADER_USE_PARAMETER_STRUCT(FCompositeHeightfieldsIntoObjectGridPagesCS, FGlobalShader);
 
 	BEGIN_SHADER_PARAMETER_STRUCT(FParameters, )
-		SHADER_PARAMETER_STRUCT_REF(FViewUniformShaderParameters, View)
 		SHADER_PARAMETER_RDG_BUFFER_UAV(RWStructuredBuffer<uint4>, RWPageObjectGridBuffer)
 		RDG_BUFFER_ACCESS(ComposeIndirectArgBuffer, ERHIAccess::IndirectArgs)
 		SHADER_PARAMETER_RDG_BUFFER_SRV(StructuredBuffer<uint>, ComposeTileBuffer)
