Use view combined shading model mask to optimize lighting passes in a mobile deferred shading

diff --git a/Engine/Source/Runtime/Engine/Private/Materials/MaterialInstance.cpp b/Engine/Source/Runtime/Engine/Private/Materials/MaterialInstance.cpp
--- a/Engine/Source/Runtime/Engine/Private/Materials/MaterialInstance.cpp
+++ b/Engine/Source/Runtime/Engine/Private/Materials/MaterialInstance.cpp
@@ -2097,6 +2097,12 @@ void UMaterialInstance::UpdateOverridableBaseProperties()
 		}
 	}
 
+	if (!GIsEditor)
+	{
+		// Filter out ShadingModels field to a current platform settings
+		FilterOutPlatformShadingModels(GMaxRHIShaderPlatform, ShadingModels);
+	}
+
 	if (BasePropertyOverrides.bOverride_TwoSided)
 	{
 		TwoSided = BasePropertyOverrides.TwoSided != 0;
diff --git a/Engine/Source/Runtime/Renderer/Private/MobileDeferredShadingPass.cpp b/Engine/Source/Runtime/Renderer/Private/MobileDeferredShadingPass.cpp
--- a/Engine/Source/Runtime/Renderer/Private/MobileDeferredShadingPass.cpp
+++ b/Engine/Source/Runtime/Renderer/Private/MobileDeferredShadingPass.cpp
@@ -308,6 +308,13 @@ constexpr uint8 GetLightingChannelStencilValue(uint32 LightingChannel)
 	return (LightingChannel == 0u ? 0u : (1u << LightingChannel));
 }
 
+constexpr bool IsOnlyDefaultLitShadingModel(uint32 ShadingModelMask)
+{
+	constexpr uint32 LitOpaqueMask = ~(1u << MSM_Unlit | 1u << MSM_SingleLayerWater | 1u << MSM_ThinTranslucent);
+	constexpr uint32 DefaultLitMask = (1u << MSM_DefaultLit);
+	return (ShadingModelMask & LitOpaqueMask) == DefaultLitMask;
+}
+
 struct FCachedLightMaterial
 {
 	const FMaterial* Material;
@@ -375,7 +382,8 @@ void RenderReflectionEnvironmentSkyLighting(FRHICommandListImmediate& RHICmdList
 	TShaderMapRef<FPostProcessVS> VertexShader(View.ShaderMap);
 
 	// Do two passes, first masking DefautLit, second masking all other shading models
-	int32 NumPasses = MobileUsesGBufferCustomData(Scene.GetShaderPlatform()) ? 2 : 1;
+	const bool bOnlyDefaultLitInView = IsOnlyDefaultLitShadingModel(View.ShadingModelMaskInView);
+	int32 NumPasses = !bOnlyDefaultLitInView && MobileUsesGBufferCustomData(Scene.GetShaderPlatform()) ? 2 : 1;
 	uint8 PassShadingModelStencilValue[2] =
 	{
 		GetMobileShadingModelStencilValue(MSM_DefaultLit),
@@ -496,7 +504,8 @@ static void RenderDirectionalLight(FRHICommandListImmediate& RHICmdList, const F
 	const bool bPlanarReflection = Scene.GetForwardPassGlobalPlanarReflection() != nullptr;
 
 	// Do two passes, first masking DefautLit, second masking all other shading models
-	int32 NumPasses = MobileUsesGBufferCustomData(Scene.GetShaderPlatform()) ? 2 : 1;
+	const bool bOnlyDefaultLitInView = IsOnlyDefaultLitShadingModel(View.ShadingModelMaskInView);
+	int32 NumPasses = !bOnlyDefaultLitInView && MobileUsesGBufferCustomData(Scene.GetShaderPlatform()) ? 2 : 1;
 	uint8 PassShadingModelStencilValue[2] =
 	{
 		GetMobileShadingModelStencilValue(MSM_DefaultLit),
@@ -751,7 +760,8 @@ static void RenderLocalLight(
 	PassParameters.TranslatedWorldToLight = FMatrix44f(FTranslationMatrix(-View.ViewMatrices.GetPreViewTranslation()) * WorldToLight);
 
 	// Do two passes, first masking DefautLit, second masking all other shading models
-	int32 NumPasses = MobileUsesGBufferCustomData(Scene.GetShaderPlatform()) ? 2 : 1;
+	const bool bOnlyDefaultLitInView = IsOnlyDefaultLitShadingModel(View.ShadingModelMaskInView);
+	int32 NumPasses = !bOnlyDefaultLitInView && MobileUsesGBufferCustomData(Scene.GetShaderPlatform()) ? 2 : 1;
 	uint8 PassShadingModelStencilValue[2] =
 	{
 		GetMobileShadingModelStencilValue(MSM_DefaultLit),
@@ -838,7 +848,8 @@ static void RenderSimpleLights(
 
 	// Setup PSOs we going to use for light rendering 
 	// Do two passes, first masking DefautLit, second masking all other shading models
-	int32 NumPasses = MobileUsesGBufferCustomData(Scene.GetShaderPlatform()) ? 2 : 1;
+	const bool bOnlyDefaultLitInView = IsOnlyDefaultLitShadingModel(View.ShadingModelMaskInView);
+	int32 NumPasses = !bOnlyDefaultLitInView && MobileUsesGBufferCustomData(Scene.GetShaderPlatform()) ? 2 : 1;
 	uint8 PassShadingModelStencilValue[2] =
 	{
 		GetMobileShadingModelStencilValue(MSM_DefaultLit),
