Remove a warning on light removal as it can't be hit anymore

diff --git a/Engine/Source/Runtime/Renderer/Private/LightSceneInfo.cpp b/Engine/Source/Runtime/Renderer/Private/LightSceneInfo.cpp
--- a/Engine/Source/Runtime/Renderer/Private/LightSceneInfo.cpp
+++ b/Engine/Source/Runtime/Renderer/Private/LightSceneInfo.cpp
@@ -132,11 +132,6 @@ void FLightSceneInfo::RemoveFromScene()
 {
 	if (OctreeId.IsValidId())
 	{
-		if (!Scene->LocalShadowCastingLightOctree.IsValidElementId(OctreeId))
-		{
-			UE_LOG(LogTemp, Warning, TEXT("Octree Element INVALID, possibly already removed this frame"));
-		}
-
 		// Remove the light from the octree.
 		Scene->LocalShadowCastingLightOctree.RemoveElement(OctreeId);
 		OctreeId = FOctreeElementId2();
