Fix for crash in volumetric cloud compute path.

diff --git a/Engine/Source/Runtime/Renderer/Private/VolumetricCloudRendering.cpp b/Engine/Source/Runtime/Renderer/Private/VolumetricCloudRendering.cpp
--- a/Engine/Source/Runtime/Renderer/Private/VolumetricCloudRendering.cpp
+++ b/Engine/Source/Runtime/Renderer/Private/VolumetricCloudRendering.cpp
@@ -637,6 +637,7 @@ void SetupDefaultRenderVolumetricCloudGlobalParameters(FRDGBuilder& GraphBuilder
 	FRDGTextureRef BlackDummy = GSystemTextures.GetBlackDummy(GraphBuilder);
 	VolumetricCloudParams.VolumetricCloud = CloudInfo.GetVolumetricCloudCommonShaderParameters();
 	VolumetricCloudParams.SceneDepthTexture = BlackDummy;
+	VolumetricCloudParams.SceneDepthMinAndMaxTexture = ((bool)ERHIZBuffer::IsInverted ? GSystemTextures.GetBlackDummy(GraphBuilder) : GSystemTextures.GetWhiteDummy(GraphBuilder));
 	VolumetricCloudParams.CloudShadowTexture0 = BlackDummy;
 	VolumetricCloudParams.CloudShadowTexture1 = BlackDummy;
 	VolumetricCloudParams.CloudBilinearTextureSampler = TStaticSamplerState<SF_Bilinear>::GetRHI();
@@ -2122,7 +2123,7 @@ static TRDGUniformBufferRef<FRenderVolumetricCloudGlobalParameters> CreateCloudP
 	SetupDefaultRenderVolumetricCloudGlobalParameters(GraphBuilder, VolumetricCloudParams, CloudInfo, MainView);
 
 	VolumetricCloudParams.SceneDepthTexture = CloudRC.SceneDepthZ;
-	VolumetricCloudParams.SceneDepthMinAndMaxTexture = CloudRC.SceneDepthMinAndMax ? CloudRC.SceneDepthMinAndMax : ((bool)ERHIZBuffer::IsInverted ? GSystemTextures.GetBlackDummy(GraphBuilder) : GSystemTextures.GetWhiteDummy(GraphBuilder));
+	VolumetricCloudParams.SceneDepthMinAndMaxTexture = CloudRC.SceneDepthMinAndMax ? CloudRC.SceneDepthMinAndMax : VolumetricCloudParams.SceneDepthMinAndMaxTexture;
 	if (CloudRC.bShouldViewRenderVolumetricRenderTarget && MainView.ViewState)
 	{
 		FVolumetricRenderTargetViewStateData& VRT = MainView.ViewState->VolumetricCloudRenderTarget;
