Allows TSR to output quarter res output

diff --git a/Engine/Shaders/Private/TemporalSuperResolution/TSRUpdateHistory.usf b/Engine/Shaders/Private/TemporalSuperResolution/TSRUpdateHistory.usf
--- a/Engine/Shaders/Private/TemporalSuperResolution/TSRUpdateHistory.usf
+++ b/Engine/Shaders/Private/TemporalSuperResolution/TSRUpdateHistory.usf
@@ -172,6 +172,7 @@ float InputContributionMultiplier;
 float GrandPrevPreExposureCorrection;
 uint ResponsiveStencilMask;
 uint bGenerateOutputMip1;
+uint bGenerateOutputMip2;
 uint bHasSeparateTranslucency;
 
 Texture2D<tsr_halfC> PrevHistory_HighFrequency;
@@ -222,6 +223,32 @@ groupshared tsr_half4 SharedArray10[GROUP_SHARED_SIZE];
 groupshared tsr_half4 SharedArray11[GROUP_SHARED_SIZE];
 
 
+//------------------------------------------------------- WAVE
+
+#if PLATFORM_SUPPORTS_WAVE_BROADCAST
+
+CALL_SITE_DEBUGLOC
+tsr_half3x2 WaveBroadcast(const FWaveBroadcastSettings Settings, tsr_half3x2 v)
+{
+	return tsr_half3x2(
+		WaveBroadcast(Settings, v[0]),
+		WaveBroadcast(Settings, v[1]),
+		WaveBroadcast(Settings, v[2]));
+} 
+
+CALL_SITE_DEBUGLOC
+tsr_half4x2 WaveBroadcast(const FWaveBroadcastSettings Settings, tsr_half4x2 v)
+{
+	return tsr_half4x2(
+		WaveBroadcast(Settings, v[0]),
+		WaveBroadcast(Settings, v[1]),
+		WaveBroadcast(Settings, v[2]),
+		WaveBroadcast(Settings, v[3]));
+} 
+
+#endif
+
+
 //------------------------------------------------------- FUNCTIONS
 
 template<typename T>
@@ -351,6 +378,91 @@ T BilinearSampleColorHistory(Texture2D<T> Texture, float2 SampleUV)
 }
 #endif
 
+tsr_halfCx2 DownsampleSceneColor(tsr_halfCx2 SceneColor, uint LocalGroupThreadIndex, const uint XorButterFly)
+{
+	SceneColor = SceneColor * tsr_half(0.25);
+
+	#if PLATFORM_SUPPORTS_WAVE_BROADCAST && CONFIG_ENABLE_DUAL_PIXEL_VECTORIZATION
+	if (uint(XorButterFly * 4) <= WaveGetLaneCount())
+	{
+		FWaveBroadcastSettings Horizontal = InitWaveXorButterfly(XorButterFly * (XorButterFly * 0x1));
+		FWaveBroadcastSettings Vertical   = InitWaveXorButterfly(XorButterFly * (XorButterFly * 0x2));
+
+		SceneColor += WaveBroadcast(Horizontal, SceneColor);
+		SceneColor += WaveBroadcast(Vertical, SceneColor);
+	}
+	else
+	#endif
+	{
+		#if CONFIG_SCENE_COLOR_ALPHA
+			SharedArray0[LocalGroupThreadIndex] = dpv_lo(SceneColor);
+			#if CONFIG_ENABLE_DUAL_PIXEL_VECTORIZATION
+				SharedArray1[LocalGroupThreadIndex] = dpv_hi(SceneColor);
+			#endif
+		#else
+			SharedArray0[LocalGroupThreadIndex] = tsr_half4(dpv_lo(SceneColor), 0.0);
+			#if CONFIG_ENABLE_DUAL_PIXEL_VECTORIZATION
+				SharedArray1[LocalGroupThreadIndex] = tsr_half4(dpv_hi(SceneColor), 0.0);
+			#endif
+		#endif
+		
+		if (XorButterFly * 0x2 <= 0x8)
+		{
+			// NOP
+		}
+		#if PLATFORM_SUPPORTS_WAVE_BROADCAST
+		else if (XorButterFly * 0x2 <= WaveGetLaneCount())
+		{
+			// NOP
+		}
+		#endif
+		else
+		{
+			GroupMemoryBarrierWithGroupSync();
+		}
+
+		SharedArray0[LocalGroupThreadIndex] += SharedArray0[LocalGroupThreadIndex ^ (XorButterFly * 0x1)];
+		#if CONFIG_ENABLE_DUAL_PIXEL_VECTORIZATION
+			SharedArray1[LocalGroupThreadIndex] += SharedArray1[LocalGroupThreadIndex ^ (XorButterFly * 0x1)];
+		#endif
+			
+		if (XorButterFly * 0x4 <= 0x8)
+		{
+			// NOP
+		}
+		#if PLATFORM_SUPPORTS_WAVE_BROADCAST
+		else if (XorButterFly * 0x4 <= WaveGetLaneCount())
+		{
+			// NOP
+		}
+		#endif
+		else
+		{
+			GroupMemoryBarrierWithGroupSync();
+		}
+		
+		SharedArray0[LocalGroupThreadIndex] += SharedArray0[LocalGroupThreadIndex ^ (XorButterFly * 0x2)];
+		#if CONFIG_ENABLE_DUAL_PIXEL_VECTORIZATION
+			SharedArray1[LocalGroupThreadIndex] += SharedArray1[LocalGroupThreadIndex ^ (XorButterFly * 0x2)];
+		#endif
+					
+		#if CONFIG_SCENE_COLOR_ALPHA
+			#if CONFIG_ENABLE_DUAL_PIXEL_VECTORIZATION
+				SceneColor = dpv_interleave_registers(SharedArray0[LocalGroupThreadIndex], SharedArray1[LocalGroupThreadIndex]);
+			#else
+				SceneColor = dpv_interleave_mono_registers(SharedArray0[LocalGroupThreadIndex]);
+			#endif
+		#else
+			#if CONFIG_ENABLE_DUAL_PIXEL_VECTORIZATION
+				SceneColor = dpv_interleave_registers(SharedArray0[LocalGroupThreadIndex].rgb, SharedArray1[LocalGroupThreadIndex].rgb);
+			#else
+				SceneColor = dpv_interleave_mono_registers(SharedArray0[LocalGroupThreadIndex].rgb);
+			#endif
+		#endif
+	}
+	return SceneColor;
+}
+
 
 //------------------------------------------------------- ENTRY POINT
 
@@ -1552,10 +1664,6 @@ void MainCS(
 			
 		LocalHistoryPixelPos = InvalidateOutputPixelPos(LocalHistoryPixelPos, HistoryInfo_ViewportMax);
 		
-		tsr_short2x2 HalfResOutputPixelPos;
-		HalfResOutputPixelPos[0] = (LocalHistoryPixelPos[0] >> tsr_short(1)) | (((LocalHistoryPixelPos[0] & tsr_short(0x1))) * tsr_short(~0));
-		HalfResOutputPixelPos[1] = (LocalHistoryPixelPos[1] >> tsr_short(1)) | (((LocalHistoryPixelPos[1] & tsr_short(0x1))) * tsr_short(~0));
-		
 		#if CONFIG_METADATA_CHANNELS == 1
 			tsr_halfMx2 FinalMetadata = FinalHistoryValidity;
 		#else
@@ -1678,62 +1786,28 @@ void MainCS(
 		
 		// Output final scene color Mip1
 		{
-			tsr_halfCx2 HalfResOutput = OutputColor * tsr_half(0.25);
-			
-			#if PLATFORM_SUPPORTS_WAVE_BROADCAST && 0 // Support WaveBroadcast with halfs
-			BRANCH
-			if (bGenerateOutputMip1)
-			{
-				FWaveBroadcastSettings Horizontal = InitWaveXorButterfly(/* XorButterFly = */ 0x1);
-				FWaveBroadcastSettings Vertical = InitWaveXorButterfly(/* XorButterFly = */ 0x2);
+			tsr_halfCx2 HalfResOutput = OutputColor;
+			tsr_short2x2 HalfResOutputPixelPos = dpv_interleave_mono_registers(tsr_short(-1).xx);
 
-				HalfResOutput += WaveBroadcast(Horizontal, HalfResOutput);
-				HalfResOutput += WaveBroadcast(Vertical, HalfResOutput);
-			}
-			#else
 			BRANCH
-			if (bGenerateOutputMip1)
+			if (bGenerateOutputMip1 || bGenerateOutputMip2)
 			{
-				#if CONFIG_SCENE_COLOR_ALPHA
-					SharedArray0[LocalGroupThreadIndex] = dpv_lo(HalfResOutput);
-					#if CONFIG_ENABLE_DUAL_PIXEL_VECTORIZATION
-						SharedArray1[LocalGroupThreadIndex] = dpv_hi(HalfResOutput);
-					#endif
-				#else
-					SharedArray0[LocalGroupThreadIndex] = tsr_half4(dpv_lo(HalfResOutput), 0.0);
-					#if CONFIG_ENABLE_DUAL_PIXEL_VECTORIZATION
-						SharedArray1[LocalGroupThreadIndex] = tsr_half4(dpv_hi(HalfResOutput), 0.0);
-					#endif
-				#endif
+				HalfResOutput = DownsampleSceneColor(HalfResOutput, LocalGroupThreadIndex, /* XorButterFly = */ 0x1);
 
-				SharedArray0[LocalGroupThreadIndex] += SharedArray0[LocalGroupThreadIndex ^ 0x1];
-				#if CONFIG_ENABLE_DUAL_PIXEL_VECTORIZATION
-					SharedArray1[LocalGroupThreadIndex] += SharedArray1[LocalGroupThreadIndex ^ 0x1];
-				#endif
-					
-				SharedArray0[LocalGroupThreadIndex] += SharedArray0[LocalGroupThreadIndex ^ 0x2];
-				#if CONFIG_ENABLE_DUAL_PIXEL_VECTORIZATION
-					SharedArray1[LocalGroupThreadIndex] += SharedArray1[LocalGroupThreadIndex ^ 0x2];
-				#endif
-					
-				#if CONFIG_SCENE_COLOR_ALPHA
-					#if CONFIG_ENABLE_DUAL_PIXEL_VECTORIZATION
-						HalfResOutput = dpv_interleave_registers(SharedArray0[LocalGroupThreadIndex], SharedArray1[LocalGroupThreadIndex]);
-					#else
-						HalfResOutput = dpv_interleave_mono_registers(SharedArray0[LocalGroupThreadIndex]);
-					#endif
-				#else
-					#if CONFIG_ENABLE_DUAL_PIXEL_VECTORIZATION
-						HalfResOutput = dpv_interleave_registers(SharedArray0[LocalGroupThreadIndex].rgb, SharedArray1[LocalGroupThreadIndex].rgb);
-					#else
-						HalfResOutput = dpv_interleave_mono_registers(SharedArray0[LocalGroupThreadIndex].rgb);
-					#endif
-				#endif
-			}
-			#endif
+				BRANCH
+				if (bGenerateOutputMip2)
+				{
+					HalfResOutput = DownsampleSceneColor(HalfResOutput, LocalGroupThreadIndex, /* XorButterFly = */ 0x4);
 
-			tsr_short2x2 SceneColorOutputMip1PixelPos = HalfResOutputPixelPos;
-			SceneColorOutputMip1PixelPos[0] = select(bGenerateOutputMip1 != 0, SceneColorOutputMip1PixelPos[0], tsr_short(~0));
+					HalfResOutputPixelPos[0] = (LocalHistoryPixelPos[0] >> tsr_short(2)) | (((LocalHistoryPixelPos[0] & tsr_short(0x3))) * tsr_short(~0));
+					HalfResOutputPixelPos[1] = (LocalHistoryPixelPos[1] >> tsr_short(2)) | (((LocalHistoryPixelPos[1] & tsr_short(0x3))) * tsr_short(~0));
+				}
+				else
+				{
+					HalfResOutputPixelPos[0] = (LocalHistoryPixelPos[0] >> tsr_short(1)) | (((LocalHistoryPixelPos[0] & tsr_short(0x1))) * tsr_short(~0));
+					HalfResOutputPixelPos[1] = (LocalHistoryPixelPos[1] >> tsr_short(1)) | (((LocalHistoryPixelPos[1] & tsr_short(0x1))) * tsr_short(~0));
+				}
+			}
 
 			SceneColorOutputMip1[dpv_lo(HalfResOutputPixelPos)] = dpv_lo(HalfResOutput);
 			
diff --git a/Engine/Source/Runtime/Renderer/Private/PostProcess/PostProcessing.cpp b/Engine/Source/Runtime/Renderer/Private/PostProcess/PostProcessing.cpp
--- a/Engine/Source/Runtime/Renderer/Private/PostProcess/PostProcessing.cpp
+++ b/Engine/Source/Runtime/Renderer/Private/PostProcess/PostProcessing.cpp
@@ -628,9 +628,13 @@ void AddPostProcessingPasses(
 			}
 			else
 			{
-				UpscalerPassInputs.bAllowDownsampleSceneColor =
-					(bNeedPostMotionBlurHalfRes || bNeedPostMotionBlurQuarterRes) &&
-					DownsampleQuality == EDownsampleQuality::Low;;
+				UpscalerPassInputs.bGenerateSceneColorHalfRes =
+					bNeedPostMotionBlurHalfRes &&
+					DownsampleQuality == EDownsampleQuality::Low;
+
+				UpscalerPassInputs.bGenerateSceneColorQuarterRes =
+					bNeedPostMotionBlurQuarterRes &&
+					DownsampleQuality == EDownsampleQuality::Low;
 			}
 			UpscalerPassInputs.DownsampleOverrideFormat = DownsampleOverrideFormat;
 			UpscalerPassInputs.SceneColorTexture = SceneColor.Texture;
@@ -645,6 +649,7 @@ void AddPostProcessingPasses(
 
 			SceneColor = Outputs.FullRes;
 			HalfResSceneColor = Outputs.HalfRes;
+			QuarterResSceneColor = Outputs.QuarterRes;
 			VelocityFlattenTextures = Outputs.VelocityFlattenTextures;
 		}
 		else if (ReflectionsMethod == EReflectionsMethod::SSR)
diff --git a/Engine/Source/Runtime/Renderer/Private/PostProcess/TemporalAA.cpp b/Engine/Source/Runtime/Renderer/Private/PostProcess/TemporalAA.cpp
--- a/Engine/Source/Runtime/Renderer/Private/PostProcess/TemporalAA.cpp
+++ b/Engine/Source/Runtime/Renderer/Private/PostProcess/TemporalAA.cpp
@@ -1017,7 +1017,7 @@ static void AddGen4MainTemporalAAPasses(
 
 	TAAParameters.DownsampleOverrideFormat = PassInputs.DownsampleOverrideFormat;
 
-	TAAParameters.bDownsample = PassInputs.bAllowDownsampleSceneColor && TAAParameters.Quality != ETAAQuality::High;
+	TAAParameters.bDownsample = (PassInputs.bGenerateSceneColorHalfRes || PassInputs.bGenerateSceneColorQuarterRes) && TAAParameters.Quality != ETAAQuality::High;
 
 	TAAParameters.SceneDepthTexture = PassInputs.SceneDepthTexture;
 	TAAParameters.SceneVelocityTexture = PassInputs.SceneVelocityTexture;
diff --git a/Engine/Source/Runtime/Renderer/Private/PostProcess/TemporalAA.h b/Engine/Source/Runtime/Renderer/Private/PostProcess/TemporalAA.h
--- a/Engine/Source/Runtime/Renderer/Private/PostProcess/TemporalAA.h
+++ b/Engine/Source/Runtime/Renderer/Private/PostProcess/TemporalAA.h
@@ -179,7 +179,8 @@ public:
 
 	struct FPassInputs
 	{
-		bool bAllowDownsampleSceneColor = false;
+		bool bGenerateSceneColorHalfRes = false;
+		bool bGenerateSceneColorQuarterRes = false;
 		bool bGenerateOutputMip1 = false;
 		bool bGenerateVelocityFlattenTextures = false;
 		EPixelFormat DownsampleOverrideFormat;
@@ -193,6 +194,7 @@ public:
 	{
 		FScreenPassTexture FullRes;
 		FScreenPassTexture HalfRes;
+		FScreenPassTexture QuarterRes;
 		FVelocityFlattenTextures VelocityFlattenTextures;
 	};
 
diff --git a/Engine/Source/Runtime/Renderer/Private/PostProcess/TemporalSuperResolution.cpp b/Engine/Source/Runtime/Renderer/Private/PostProcess/TemporalSuperResolution.cpp
--- a/Engine/Source/Runtime/Renderer/Private/PostProcess/TemporalSuperResolution.cpp
+++ b/Engine/Source/Runtime/Renderer/Private/PostProcess/TemporalSuperResolution.cpp
@@ -639,6 +639,7 @@ class FTSRUpdateHistoryCS : public FTSRShader
 		SHADER_PARAMETER(float, GrandPrevPreExposureCorrection)
 		SHADER_PARAMETER(int32, ResponsiveStencilMask)
 		SHADER_PARAMETER(int32, bGenerateOutputMip1)
+		SHADER_PARAMETER(int32, bGenerateOutputMip2)
 		SHADER_PARAMETER(int32, bHasSeparateTranslucency)
 
 		SHADER_PARAMETER_STRUCT_INCLUDE(FTSRPrevHistoryParameters, PrevHistoryParameters)
@@ -1105,9 +1106,10 @@ ITemporalUpscaler::FOutputs AddTemporalSuperResolutionPasses(
 	FRDGTextureRef UpdateHistoryOutputTexture;
 	FRDGTextureRef SceneColorOutputTexture;
 	FRDGTextureRef SceneColorOutputHalfResTexture = nullptr;
+	FRDGTextureRef SceneColorOutputQuarterResTexture = nullptr;
 	FTSRHistoryTextures History;
 	{
-		check(!(PassInputs.bGenerateOutputMip1 && PassInputs.bAllowDownsampleSceneColor));
+		check(!(PassInputs.bGenerateOutputMip1 && (PassInputs.bGenerateSceneColorHalfRes || PassInputs.bGenerateSceneColorQuarterRes)));
 		FRDGTextureDesc Desc = FRDGTextureDesc::Create2D(
 			HistoryExtent,
 			ColorFormat,
@@ -1127,11 +1129,20 @@ ITemporalUpscaler::FOutputs AddTemporalSuperResolutionPasses(
 		{
 			SceneColorOutputTexture = UpdateHistoryOutputTexture;
 		}
-
-		if (PassInputs.bAllowDownsampleSceneColor)
+		
+		// Generate quarter res output if only one needed, otherwise only output half res and let caller downscale to quarter res. 
+		// This is motivated to saves UAV slots on FTSRUpdateHistoryCS
+		if (PassInputs.bGenerateSceneColorQuarterRes && !PassInputs.bGenerateSceneColorHalfRes && OutputRect.Size() == HistorySize)
+		{
+			FRDGTextureDesc QuarterResDesc = Desc;
+			QuarterResDesc.Extent /= 4;
+			SceneColorOutputQuarterResTexture = GraphBuilder.CreateTexture(QuarterResDesc, TEXT("TSR.QuarterResOutput"));
+		}
+		else if (PassInputs.bGenerateSceneColorHalfRes || PassInputs.bGenerateSceneColorQuarterRes)
 		{
-			Desc.Extent /= 2;
-			SceneColorOutputHalfResTexture = GraphBuilder.CreateTexture(Desc, TEXT("TSR.HalfResOutput"));
+			FRDGTextureDesc HalfResDesc = Desc;
+			HalfResDesc.Extent /= 2;
+			SceneColorOutputHalfResTexture = GraphBuilder.CreateTexture(HalfResDesc, TEXT("TSR.HalfResOutput"));
 		}
 	}
 	{
@@ -1622,7 +1633,8 @@ ITemporalUpscaler::FOutputs AddTemporalSuperResolutionPasses(
 		PassParameters->InputContributionMultiplier = FMath::Pow(float(HistorySize.X) / float(OutputRect.Width()), 2.0f); 
 		PassParameters->GrandPrevPreExposureCorrection = bCameraCut ? 1.0f : View.PreExposure / InputHistory.PrevSceneColorPreExposure;
 		PassParameters->ResponsiveStencilMask = CVarTSREnableResponiveAA.GetValueOnRenderThread() ? (STENCIL_TEMPORAL_RESPONSIVE_AA_MASK) : 0;
-		PassParameters->bGenerateOutputMip1 = ((PassInputs.bGenerateOutputMip1 || PassInputs.bAllowDownsampleSceneColor) && HistorySize == OutputRect.Size()) ? 1 : 0;
+		PassParameters->bGenerateOutputMip1 = false;
+		PassParameters->bGenerateOutputMip2 = false;
 		PassParameters->bHasSeparateTranslucency = bHasSeparateTranslucency;
 
 		PassParameters->PrevHistoryParameters = PrevHistoryParameters;
@@ -1639,21 +1651,24 @@ ITemporalUpscaler::FOutputs AddTemporalSuperResolutionPasses(
 			PassParameters->SceneColorOutputMip0 = GraphBuilder.CreateUAV(FRDGTextureUAVDesc(UpdateHistoryOutputTexture, /* InMipLevel = */ 0));
 		}
 		
-		if (!PassParameters->bGenerateOutputMip1)
-		{
-			PassParameters->SceneColorOutputMip1 = CreateDummyUAV(GraphBuilder, PF_FloatR11G11B10);
-		}
-		else if (PassInputs.bGenerateOutputMip1)
+		if (PassInputs.bGenerateOutputMip1 && HistorySize == OutputRect.Size())
 		{
+			PassParameters->bGenerateOutputMip1 = true;
 			PassParameters->SceneColorOutputMip1 = GraphBuilder.CreateUAV(FRDGTextureUAVDesc(SceneColorOutputTexture, /* InMipLevel = */ 1));
 		}
-		else if (PassInputs.bAllowDownsampleSceneColor)
+		else if (SceneColorOutputHalfResTexture && HistorySize == OutputRect.Size())
 		{
+			PassParameters->bGenerateOutputMip1 = true;
 			PassParameters->SceneColorOutputMip1 = GraphBuilder.CreateUAV(FRDGTextureUAVDesc(SceneColorOutputHalfResTexture));
 		}
+		else if (SceneColorOutputQuarterResTexture && HistorySize == OutputRect.Size())
+		{
+			PassParameters->bGenerateOutputMip2 = true;
+			PassParameters->SceneColorOutputMip1 = GraphBuilder.CreateUAV(FRDGTextureUAVDesc(SceneColorOutputQuarterResTexture));
+		}
 		else
 		{
-			unimplemented();
+			PassParameters->SceneColorOutputMip1 = CreateDummyUAV(GraphBuilder, PF_FloatR11G11B10);
 		}
 		PassParameters->DebugOutput = CreateDebugUAV(HistoryExtent, TEXT("Debug.TSR.UpdateHistory"));
 
@@ -1665,11 +1680,12 @@ ITemporalUpscaler::FOutputs AddTemporalSuperResolutionPasses(
 		TShaderMapRef<FTSRUpdateHistoryCS> ComputeShader(View.ShaderMap, PermutationVector);
 		FComputeShaderUtils::AddPass(
 			GraphBuilder,
-			RDG_EVENT_NAME("TSR UpdateHistory(Quality=%s%s%s%s) %dx%d",
+			RDG_EVENT_NAME("TSR UpdateHistory(Quality=%s%s%s%s%s) %dx%d",
 				kUpdateQualityNames[int32(PermutationVector.Get<FTSRUpdateHistoryCS::FQualityDim>())],
 				PermutationVector.Get<FTSRUpdateHistoryCS::FSeparateTranslucencyDim>() ? TEXT(" SeparateTranslucency") : TEXT(""),
+				PermutationVector.Get<FTSRUpdateHistoryCS::FGrandReprojectionDim>() ? TEXT(" GrandReprojection") : TEXT(""),
 				(History.ColorArray ? History.ColorArray : History.Output)->Desc.Format == PF_FloatR11G11B10 ? TEXT(" R11G11B10") : TEXT(""),
-				PassParameters->bGenerateOutputMip1 ? TEXT(" OutputMip1") : TEXT(""),
+				PassParameters->bGenerateOutputMip2 ? TEXT(" OutputMip2") : (PassParameters->bGenerateOutputMip1 ? TEXT(" OutputMip1") : TEXT("")),
 				HistorySize.X, HistorySize.Y),
 			AsyncComputePasses >= 3 ? ERDGPassFlags::AsyncCompute : ERDGPassFlags::Compute,
 			ComputeShader,
@@ -1703,6 +1719,7 @@ ITemporalUpscaler::FOutputs AddTemporalSuperResolutionPasses(
 	// If we upscaled the history buffer, downsize back to the secondary screen percentage size.
 	if (HistorySize != OutputRect.Size())
 	{
+		check(!SceneColorOutputQuarterResTexture);
 		FTSRResolveHistoryCS::FParameters* PassParameters = GraphBuilder.AllocParameters<FTSRResolveHistoryCS::FParameters>();
 		PassParameters->CommonParameters = CommonParameters;
 		PassParameters->DispatchThreadToHistoryPixelPos = (
@@ -1712,7 +1729,7 @@ ITemporalUpscaler::FOutputs AddTemporalSuperResolutionPasses(
 				FScreenTransform::ETextureBasis::ViewportUV, FScreenTransform::ETextureBasis::TexelPosition));
 		PassParameters->OutputViewRectMin = OutputRect.Min;
 		PassParameters->OutputViewRectMax = OutputRect.Max;
-		PassParameters->bGenerateOutputMip1 = (PassInputs.bGenerateOutputMip1 || PassInputs.bAllowDownsampleSceneColor) ? 1 : 0;
+		PassParameters->bGenerateOutputMip1 = (PassInputs.bGenerateOutputMip1 || PassInputs.bGenerateSceneColorHalfRes) ? 1 : 0;
 		PassParameters->HistoryValidityMultiply = float(HistorySize.X * HistorySize.Y) / float(OutputRect.Width() * OutputRect.Height());
 
 		PassParameters->History = CreateSRVs(GraphBuilder, History);
@@ -1722,7 +1739,7 @@ ITemporalUpscaler::FOutputs AddTemporalSuperResolutionPasses(
 		{
 			PassParameters->SceneColorOutputMip1 = GraphBuilder.CreateUAV(FRDGTextureUAVDesc(SceneColorOutputTexture, /* InMipLevel = */ 1));
 		}
-		else if (PassInputs.bAllowDownsampleSceneColor)
+		else if (SceneColorOutputHalfResTexture)
 		{
 			PassParameters->SceneColorOutputMip1 = GraphBuilder.CreateUAV(FRDGTextureUAVDesc(SceneColorOutputHalfResTexture));
 		}
@@ -1831,12 +1848,18 @@ ITemporalUpscaler::FOutputs AddTemporalSuperResolutionPasses(
 	ITemporalUpscaler::FOutputs Outputs;
 	Outputs.FullRes.Texture = SceneColorOutputTexture;
 	Outputs.FullRes.ViewRect = OutputRect;
-	if (PassInputs.bAllowDownsampleSceneColor)
+	if (SceneColorOutputHalfResTexture)
 	{
 		Outputs.HalfRes.Texture = SceneColorOutputHalfResTexture;
 		Outputs.HalfRes.ViewRect.Min = OutputRect.Min / 2;
 		Outputs.HalfRes.ViewRect.Max = Outputs.HalfRes.ViewRect.Min + FIntPoint::DivideAndRoundUp(OutputRect.Size(), 2);
 	}
+	if (SceneColorOutputQuarterResTexture)
+	{
+		Outputs.QuarterRes.Texture = SceneColorOutputQuarterResTexture;
+		Outputs.QuarterRes.ViewRect.Min = OutputRect.Min / 4;
+		Outputs.QuarterRes.ViewRect.Max = Outputs.HalfRes.ViewRect.Min + FIntPoint::DivideAndRoundUp(OutputRect.Size(), 4);
+	}
 	Outputs.VelocityFlattenTextures = VelocityFlattenTextures;
 	return Outputs;
 } // AddTemporalSuperResolutionPasses()
