Render all skylights in a separate render pass when static lighting is disabled on mobile deferred.
Strip mobile base pass shader permutations with SkyLight enabled and static lighting disabled for mobile deferred.

diff --git a/Engine/Shaders/Private/MobileDeferredShading.usf b/Engine/Shaders/Private/MobileDeferredShading.usf
--- a/Engine/Shaders/Private/MobileDeferredShading.usf
+++ b/Engine/Shaders/Private/MobileDeferredShading.usf
@@ -90,7 +90,6 @@ half3 SkyLightDiffuseMobile(FGBufferData GBuffer, half3 DiffuseColor, half3 V, i
 
 	// And accumulate the lighting
 	Lighting += DiffuseLookup  * DiffuseColor;
-	Lighting *= View.PreExposure;
 
 #if ALLOW_STATIC_LIGHTING
 	IndirectIrradiance += Luminance(DiffuseLookup);
diff --git a/Engine/Source/Runtime/Renderer/Private/MobileBasePassRendering.h b/Engine/Source/Runtime/Renderer/Private/MobileBasePassRendering.h
--- a/Engine/Source/Runtime/Renderer/Private/MobileBasePassRendering.h
+++ b/Engine/Source/Runtime/Renderer/Private/MobileBasePassRendering.h
@@ -340,6 +340,8 @@ public:
 	{		
 		// We compile the point light shader combinations based on the project settings
 		static auto* MobileSkyLightPermutationCVar = IConsoleManager::Get().FindTConsoleVariableDataInt(TEXT("r.Mobile.SkyLightPermutation"));
+		static auto AllowStaticLightingVar = IConsoleManager::Get().FindTConsoleVariableDataInt(TEXT("r.AllowStaticLighting"));
+		const bool bAllowStaticLighting = (!AllowStaticLightingVar || AllowStaticLightingVar->GetValueOnAnyThread() != 0);
 
 		const int32 MobileSkyLightPermutationOptions = MobileSkyLightPermutationCVar->GetValueOnAnyThread();
 		const bool bDeferredShading = IsMobileDeferredShadingEnabled(Parameters.Platform);
@@ -351,7 +353,9 @@ public:
 		// Translucent materials always support clustered shading on mobile deferred
 		const bool bSupportsLocalLights = (!bDeferredShading && MobileForwardEnableLocalLights(Parameters.Platform)) || (bDeferredShading && bMaterialUsesForwardShading);
 		// Only compile skylight version for lit materials
-		const bool bShouldCacheBySkylight = !bEnableSkyLight || bIsLit;
+		const bool bForwardShading = !bDeferredShading || bMaterialUsesForwardShading;
+		const bool bRenderSkyLightInBasePass = bForwardShading || bAllowStaticLighting;
+		const bool bShouldCacheBySkylight = !bEnableSkyLight || (bIsLit && bRenderSkyLightInBasePass);
 
 		// Only compile skylight permutations when they are enabled
 		if (bIsLit && !UseSkylightPermutation(bEnableSkyLight, MobileSkyLightPermutationOptions))
@@ -360,8 +364,7 @@ public:
 		}
 
 		// Deferred shading does not need SkyLight and LocalLight permutations
-		// TODO: skip skylight permutations for deferred
-		const bool bForwardShading = !bDeferredShading || bMaterialUsesForwardShading;
+		// TODO: skip skylight permutations for deferred	
 		const bool bShouldCacheByShading = (bForwardShading || !bEnableLocalLights);
 		const bool bShouldCacheByLocalLights = !bEnableLocalLights || (bIsLit && bEnableLocalLights == bSupportsLocalLights);
 
diff --git a/Engine/Source/Runtime/Renderer/Private/MobileDeferredShadingPass.cpp b/Engine/Source/Runtime/Renderer/Private/MobileDeferredShadingPass.cpp
--- a/Engine/Source/Runtime/Renderer/Private/MobileDeferredShadingPass.cpp
+++ b/Engine/Source/Runtime/Renderer/Private/MobileDeferredShadingPass.cpp
@@ -353,8 +353,9 @@ static void GetLightMaterial(const FCachedLightMaterial& DefaultLightMaterial, c
 void RenderReflectionEnvironmentSkyLighting(FRHICommandListImmediate& RHICmdList, const FScene& Scene, const FViewInfo& View)
 {
 	// Skylights with static lighting already had their diffuse contribution baked into lightmaps
+	const bool bIsDynamicLighting = !FReadOnlyCVARCache::Get().bAllowStaticLighting || Scene.GetForceNoPrecomputedLighting();
 	const bool bSkyLight = Scene.SkyLight && !Scene.SkyLight->bHasStaticLighting && View.Family->EngineShowFlags.SkyLighting;
-	const bool bDynamicSkyLight = bSkyLight && !Scene.SkyLight->bWantsStaticShadowing;
+	const bool bDynamicSkyLight = bSkyLight && (!Scene.SkyLight->bWantsStaticShadowing || bIsDynamicLighting);
 	const bool bClustredReflection = (View.NumBoxReflectionCaptures + View.NumSphereReflectionCaptures) > 0;
 	const bool bPlanarReflection = Scene.GetForwardPassGlobalPlanarReflection() != nullptr;
 	if (!(bSkyLight || bClustredReflection || bPlanarReflection))
@@ -498,8 +499,9 @@ static void RenderDirectionalLight(FRHICommandListImmediate& RHICmdList, const F
 	}
 
 	// Skylights with static lighting already had their diffuse contribution baked into lightmaps
+	const bool bIsDynamicLighting = !FReadOnlyCVARCache::Get().bAllowStaticLighting || Scene.GetForceNoPrecomputedLighting();
 	const bool bSkyLight = Scene.SkyLight && !Scene.SkyLight->bHasStaticLighting && View.Family->EngineShowFlags.SkyLighting;
-	const bool bDynamicSkyLight = bSkyLight && !Scene.SkyLight->bWantsStaticShadowing;
+	const bool bDynamicSkyLight = bSkyLight && (!Scene.SkyLight->bWantsStaticShadowing || bIsDynamicLighting);
 	const bool bDynamicShadows = DirectionalLight.Proxy->CastsDynamicShadow() && (LightingChannel == 0u) && View.Family->EngineShowFlags.DynamicShadows;
 	const bool bPlanarReflection = Scene.GetForwardPassGlobalPlanarReflection() != nullptr;
 
diff --git a/Engine/Source/Runtime/Renderer/Private/RendererScene.cpp b/Engine/Source/Runtime/Renderer/Private/RendererScene.cpp
--- a/Engine/Source/Runtime/Renderer/Private/RendererScene.cpp
+++ b/Engine/Source/Runtime/Renderer/Private/RendererScene.cpp
@@ -6093,6 +6093,42 @@ bool FScene::IsPrimitiveBeingRemoved(FPrimitiveSceneInfo* PrimitiveSceneInfo) co
 	return RemovedPrimitiveSceneInfos.Find(PrimitiveSceneInfo) != nullptr;
 }
 
+bool FScene::ShouldRenderSkylightInBasePass(EBlendMode BlendMode) const
+{
+	if (IsMobilePlatform(GetShaderPlatform()))
+	{
+		bool bRenderSkyLight = SkyLight && !SkyLight->bHasStaticLighting;
+		const bool bIsForwardShading = IsTranslucentBlendMode(BlendMode) || !IsMobileDeferredShadingEnabled(GetShaderPlatform());
+		const bool bIsDynamicLighting = !ReadOnlyCVARCache.bAllowStaticLighting || GetForceNoPrecomputedLighting();
+
+		// Both stationary and movable skylights are applied in base pass for forward shading
+		// Only stationary skylights are applied in base pass for deferred with static lighting enabled
+		return bRenderSkyLight && (bIsForwardShading || (!bIsDynamicLighting && SkyLight->bWantsStaticShadowing));
+	}
+	else
+	{
+		bool bRenderSkyLight = SkyLight && !SkyLight->bHasStaticLighting && !(ShouldRenderRayTracingSkyLight(SkyLight) && !IsForwardShadingEnabled(GetShaderPlatform()));
+
+		if (IsTranslucentBlendMode(BlendMode))
+		{
+			// Both stationary and movable skylights are applied in base pass for translucent materials
+			bRenderSkyLight = bRenderSkyLight
+				&& (ReadOnlyCVARCache.bEnableStationarySkylight || !SkyLight->bWantsStaticShadowing);
+		}
+		else
+		{
+			// For opaque materials, stationary skylight is applied in base pass but movable skylight
+			// is applied in a separate render pass (bWantssStaticShadowing means stationary skylight)
+			bRenderSkyLight = bRenderSkyLight
+				&& ((ReadOnlyCVARCache.bEnableStationarySkylight && SkyLight->bWantsStaticShadowing)
+					|| (!SkyLight->bWantsStaticShadowing
+						&& IsForwardShadingEnabled(GetShaderPlatform())));
+		}
+
+		return bRenderSkyLight;
+	}
+}
+
 /**
  * Dummy NULL scene interface used by dedicated servers.
  */
diff --git a/Engine/Source/Runtime/Renderer/Private/ScenePrivate.h b/Engine/Source/Runtime/Renderer/Private/ScenePrivate.h
--- a/Engine/Source/Runtime/Renderer/Private/ScenePrivate.h
+++ b/Engine/Source/Runtime/Renderer/Private/ScenePrivate.h
@@ -3385,28 +3385,7 @@ public:
 
 	virtual bool IsEditorScene() const override { return bIsEditorScene; }
 
-	bool ShouldRenderSkylightInBasePass(EBlendMode BlendMode) const
-	{
-		bool bRenderSkyLight = SkyLight && !SkyLight->bHasStaticLighting && !(ShouldRenderRayTracingSkyLight(SkyLight) && !IsForwardShadingEnabled(GetShaderPlatform()));
-
-		if (IsTranslucentBlendMode(BlendMode))
-		{
-			// Both stationary and movable skylights are applied in base pass for translucent materials
-			bRenderSkyLight = bRenderSkyLight
-				&& (ReadOnlyCVARCache.bEnableStationarySkylight || !SkyLight->bWantsStaticShadowing);
-		}
-		else
-		{
-			// For opaque materials, stationary skylight is applied in base pass but movable skylight
-			// is applied in a separate render pass (bWantssStaticShadowing means stationary skylight)
-			bRenderSkyLight = bRenderSkyLight
-				&& ((ReadOnlyCVARCache.bEnableStationarySkylight && SkyLight->bWantsStaticShadowing)
-					|| (!SkyLight->bWantsStaticShadowing
-						&& (IsForwardShadingEnabled(GetShaderPlatform()) || (IsMobilePlatform(GetShaderPlatform()) && !IsMobileDeferredShadingEnabled(GetShaderPlatform())))));
-		}
-
-		return bRenderSkyLight;
-	}
+	bool ShouldRenderSkylightInBasePass(EBlendMode BlendMode) const;
 
 	virtual TArray<FPrimitiveComponentId> GetScenePrimitiveComponentIds() const override
 	{
