Take subsurface color and EnvBRDF into account when calculating approximate illuminance for auto exposure.

diff --git a/Engine/Shaders/Private/Histogram.usf b/Engine/Shaders/Private/Histogram.usf
--- a/Engine/Shaders/Private/Histogram.usf
+++ b/Engine/Shaders/Private/Histogram.usf
@@ -1,6 +1,8 @@
 // Copyright Epic Games, Inc. All Rights Reserved.
 
 #include "DeferredShadingCommon.ush"
+#include "BRDF.ush"
+#include "PositionReconstructionCommon.ush"
 
 #include "/Engine/Private/Common.ush"
 #include "/Engine/Private/ScreenPass.ush"
@@ -43,12 +45,6 @@ uint2 ThreadGroupCount;
 Texture2D InputTexture;
 Texture2D EyeAdaptationTexture;
 
-#if DIVIDE_ALBEDO
-	SCREEN_PASS_TEXTURE_VIEWPORT(InputGBuffer)
-	Texture2D InputTextureB;
-	Texture2D InputTextureC;
-#endif
-
 ///////////////////////////////////////////////////////
 // * LOCAL DATA STORE *
 // This is where each ThreadGroup accumulates (via InterlockedAdd) a local (per thread group) copy of the histogram
@@ -317,32 +313,53 @@ void MainAtomicCS(
 
 		if (TexelPos.x < Input_ViewportMax.x)
 		{
-			float2 InvViewSize = Input_ViewportSizeInverse.xy;
-			float2 ScreenUV = (TexelPos.xy - Input_ViewportMin) * InvViewSize;
+			const float2 InvViewSize = Input_ViewportSizeInverse.xy;
+			const float2 ScreenUV = (TexelPos.xy - Input_ViewportMin + 0.5f) * InvViewSize;
 
 			float4 SceneColor = InputTexture.Load(int3(TexelPos, 0));
 			SceneColor *= View.OneOverPreExposure;
 
 			float LuminanceVal = CalculateEyeAdaptationLuminance(SceneColor.rgb);
 
-#if DIVIDE_ALBEDO
-			uint2 GBufferTexelPos = InputGBuffer_ViewportMin + ScreenUV * InputGBuffer_ViewportSize;
-			float4 GBufferB = InputTextureB.Load(int3(GBufferTexelPos, 0));
-			uint ShadingModelID = DecodeShadingModelId(GBufferB.a);
+#if USE_APPROX_ILLUMINANCE
+			const uint2 GBufferTexelPos = View.ViewRectMinAndSize.xy + ScreenUV * View.ViewRectMinAndSize.zw;
+			const float2 GBufferUV = (GBufferTexelPos + 0.5f) * View.BufferSizeAndInvSize.zw;
 
-			if (ShadingModelID != SHADINGMODELID_UNLIT)
+			FGBufferData GBufferData = GetGBufferDataFromSceneTextures(GBufferUV);
+			
+			if (GBufferData.ShadingModelID != SHADINGMODELID_UNLIT)
 			{
-				float4 GBufferC = InputTextureC.Load(int3(GBufferTexelPos, 0));
-				float3 BaseColor = DecodeBaseColor(GBufferC.rgb);
+				float3 Denominator = GBufferData.DiffuseColor;
+
+				if (GBufferData.ShadingModelID == SHADINGMODELID_SUBSURFACE
+					|| GBufferData.ShadingModelID == SHADINGMODELID_PREINTEGRATED_SKIN
+					|| GBufferData.ShadingModelID == SHADINGMODELID_TWOSIDED_FOLIAGE
+					|| GBufferData.ShadingModelID == SHADINGMODELID_CLOTH
+					|| GBufferData.ShadingModelID == SHADINGMODELID_EYE)
+				{
+					const float3 SubsurfaceColor = ExtractSubsurfaceColor(GBufferData);
+					Denominator += SubsurfaceColor;
+				}
+
+				{
+					const float3 TranslatedWorldPosition = ReconstructTranslatedWorldPositionFromDepth(GBufferUV, GBufferData.Depth);
+					const float3 CameraVector = normalize(View.TranslatedWorldCameraOrigin - TranslatedWorldPosition);
+
+					const float3 N = GBufferData.WorldNormal;
+					const float3 V = CameraVector;
+					const float3 EnvBrdf = EnvBRDF(GBufferData.SpecularColor, GBufferData.Roughness, max(0.0, dot(N, V)));
+
+					Denominator += EnvBrdf;
+				}
 
-				float BaseColorLuminanceVal = CalculateEyeAdaptationLuminance(BaseColor.rgb);
+				float BaseColorLuminanceVal = CalculateEyeAdaptationLuminance(Denominator);
 
 				LuminanceVal /= max(BaseColorLuminanceVal, EyeAdaptation_IgnoreMaterialsMinBaseColorLuminance);
 				LuminanceVal *= lerp(EyeAdaptation_IgnoreMaterialsLuminanceScale, 1.0f, EyeAdaptation_IgnoreMaterialsMinBaseColorLuminance);
 			}
 #endif
 
-#if DIVIDE_ALBEDO_DEBUG
+#if USE_APPROX_ILLUMINANCE_DEBUG
 			DebugOutput[TexelPos] = LuminanceVal;
 #endif
 
diff --git a/Engine/Source/Runtime/Renderer/Private/PostProcess/PostProcessHistogram.cpp b/Engine/Source/Runtime/Renderer/Private/PostProcess/PostProcessHistogram.cpp
--- a/Engine/Source/Runtime/Renderer/Private/PostProcess/PostProcessHistogram.cpp
+++ b/Engine/Source/Runtime/Renderer/Private/PostProcess/PostProcessHistogram.cpp
@@ -9,6 +9,7 @@
 #include "ShaderCompilerCore.h"
 
 #include "SceneTextureParameters.h"
+#include "SystemTextures.h"
 
 TAutoConsoleVariable<int32> CVarUseAtomicHistogram(
 	TEXT("r.Histogram.UseAtomic"), 1,
@@ -133,22 +134,23 @@ public:
 
 	BEGIN_SHADER_PARAMETER_STRUCT(FParameters, )
 		SHADER_PARAMETER_STRUCT_REF(FViewUniformShaderParameters, View)
+		SHADER_PARAMETER_STRUCT_INCLUDE(FSceneTextureParameters, SceneTextures)
 		SHADER_PARAMETER_STRUCT(FScreenPassTextureViewportParameters, Input)
-		SHADER_PARAMETER_STRUCT(FScreenPassTextureViewportParameters, InputGBuffer)
 		SHADER_PARAMETER_STRUCT(FEyeAdaptationParameters, EyeAdaptation)
 		SHADER_PARAMETER_RDG_TEXTURE(Texture2D, InputTexture)
-		SHADER_PARAMETER_RDG_TEXTURE(Texture2D, InputTextureB)
-		SHADER_PARAMETER_RDG_TEXTURE(Texture2D, InputTextureC)
 		SHADER_PARAMETER_RDG_TEXTURE(Texture2D, EyeAdaptationTexture)
 		SHADER_PARAMETER(FIntPoint, ThreadGroupCount)
 		SHADER_PARAMETER_RDG_TEXTURE_UAV(RWTexture2D<float4>, HistogramScatter64Output)
 		SHADER_PARAMETER_RDG_TEXTURE_UAV(RWTexture2D<float4>, HistogramScatter32Output)
 		SHADER_PARAMETER_RDG_TEXTURE_UAV(RWTexture2D<float>, DebugOutput)
+
+		SHADER_PARAMETER_TEXTURE(Texture2D, PreIntegratedGF)
+		SHADER_PARAMETER_SAMPLER(SamplerState, PreIntegratedGFSampler)
 	END_SHADER_PARAMETER_STRUCT()
 
-	class FDivideAlbedoDim : SHADER_PERMUTATION_BOOL("DIVIDE_ALBEDO");
-	class FDivideAlbedoDebugDim : SHADER_PERMUTATION_BOOL("DIVIDE_ALBEDO_DEBUG");
-	using FPermutationDomain = TShaderPermutationDomain<FDivideAlbedoDim, FDivideAlbedoDebugDim>;
+	class FUseApproxIlluminanceDim : SHADER_PERMUTATION_BOOL("USE_APPROX_ILLUMINANCE");
+	class FUseApproxIlluminanceDebugDim : SHADER_PERMUTATION_BOOL("USE_APPROX_ILLUMINANCE_DEBUG");
+	using FPermutationDomain = TShaderPermutationDomain<FUseApproxIlluminanceDim, FUseApproxIlluminanceDebugDim>;
 
 	static bool ShouldCompilePermutation(const FGlobalShaderPermutationParameters& Parameters)
 	{
@@ -335,11 +337,9 @@ static FRDGTextureRef AddHistogramAtomicPass(
 
 		FHistogramAtomicCS::FParameters* PassParameters = GraphBuilder.AllocParameters<FHistogramAtomicCS::FParameters>();
 		PassParameters->View = View.ViewUniformBuffer;
+		PassParameters->SceneTextures = SceneTextures;
 		PassParameters->Input = GetScreenPassTextureViewportParameters(FScreenPassTextureViewport(SceneColor));
 		PassParameters->InputTexture = SceneColor.Texture;
-		PassParameters->InputGBuffer = GetScreenPassTextureViewportParameters(FScreenPassTextureViewport(View.ViewRect));
-		PassParameters->InputTextureB = SceneTextures.GBufferBTexture;
-		PassParameters->InputTextureC = SceneTextures.GBufferCTexture;
 		PassParameters->EyeAdaptationTexture = EyeAdaptationTexture;
 		PassParameters->ThreadGroupCount = FIntPoint(SceneColor.ViewRect.Size().Y, 1);
 		PassParameters->EyeAdaptation = EyeAdaptationParameters;
@@ -352,6 +352,10 @@ static FRDGTextureRef AddHistogramAtomicPass(
 			PassParameters->DebugOutput = GraphBuilder.CreateUAV(DebugOutputTexture);
 		}
 
+
+		PassParameters->PreIntegratedGF = GSystemTextures.PreintegratedGF->GetRHI();
+		PassParameters->PreIntegratedGFSampler = TStaticSamplerState<SF_Bilinear, AM_Clamp, AM_Clamp, AM_Clamp>::GetRHI();
+
 		//clear the temp textures
 		uint32 ClearValues[4] = { 0, 0, 0, 0 };
 		// AddClearUAVPass(GraphBuilder, PassParameters->HistogramScatter64Output, ClearValues);
@@ -361,8 +365,8 @@ static FRDGTextureRef AddHistogramAtomicPass(
 		static const auto CVarIgnoreMaterialsDebug = IConsoleManager::Get().FindTConsoleVariableDataBool(TEXT("r.AutoExposure.IgnoreMaterials.Debug"));
 
 		FHistogramAtomicCS::FPermutationDomain PermutationVector;
-		PermutationVector.Set<FHistogramAtomicCS::FDivideAlbedoDim>(CVarIgnoreMaterials->GetValueOnRenderThread());
-		PermutationVector.Set<FHistogramAtomicCS::FDivideAlbedoDebugDim>(CVarIgnoreMaterialsDebug->GetValueOnRenderThread());
+		PermutationVector.Set<FHistogramAtomicCS::FUseApproxIlluminanceDim>(CVarIgnoreMaterials->GetValueOnRenderThread());
+		PermutationVector.Set<FHistogramAtomicCS::FUseApproxIlluminanceDebugDim>(CVarIgnoreMaterialsDebug->GetValueOnRenderThread());
 
 		auto ComputeShader = View.ShaderMap->GetShader<FHistogramAtomicCS>(PermutationVector);
 
