#UE Qucik fix for memory corruption in DistanceFieldObjectManagement index logic

diff --git a/Engine/Source/Runtime/Renderer/Private/DistanceFieldObjectManagement.cpp b/Engine/Source/Runtime/Renderer/Private/DistanceFieldObjectManagement.cpp
--- a/Engine/Source/Runtime/Renderer/Private/DistanceFieldObjectManagement.cpp
+++ b/Engine/Source/Runtime/Renderer/Private/DistanceFieldObjectManagement.cpp
@@ -340,7 +340,10 @@ void UpdateGlobalHeightFieldObjectRemoves(FScene* Scene)
 					PrimitiveBeingMoved->DistanceFieldInstanceIndices[0] = RemoveIndex;
 				}
 
-				SceneData.HeightfieldPrimitives.RemoveAtSwap(RemoveIndex, 1, false);
+				if (SceneData.HeightfieldPrimitives.IsValidIndex(RemoveIndex))
+				{
+					SceneData.HeightfieldPrimitives.RemoveAtSwap(RemoveIndex, 1, false);
+				}
 
 				SceneData.IndicesToUpdateInHeightFieldObjectBuffers.Add(RemoveIndex);
 			}
