Add CSV counters for VSM.

diff --git a/Engine/Shaders/Private/VirtualShadowMaps/VirtualShadowMapPageManagement.usf b/Engine/Shaders/Private/VirtualShadowMaps/VirtualShadowMapPageManagement.usf
--- a/Engine/Shaders/Private/VirtualShadowMaps/VirtualShadowMapPageManagement.usf
+++ b/Engine/Shaders/Private/VirtualShadowMaps/VirtualShadowMapPageManagement.usf
@@ -1548,16 +1548,14 @@ void BuildHZBPerPageTopCS(uint2 GroupThreadId : SV_GroupThreadID, uint PageInput
 }
 
 uint StatusMessageId;
-StructuredBuffer<int> FreePhysicalPages;
 
+StructuredBuffer<int> FreePhysicalPages;
 
 [numthreads(1, 1, 1)]
 void FeedbackStatusCS()
 {
 	{
-		FGPUMessageWriter Mw = GPUMessageBegin(StatusMessageId, 2U);
-		// 1. Write out frame number
-		GPUMessageWriteItem(Mw, 0U);
+		FGPUMessageWriter Mw = GPUMessageBegin(StatusMessageId, 1U);
 		// 2. Write out the free list status
 		GPUMessageWriteItem(Mw, FreePhysicalPages[VirtualShadowMap.MaxPhysicalPages]);
 	}
diff --git a/Engine/Shaders/Private/VirtualShadowMaps/VirtualShadowMapPrintStats.usf b/Engine/Shaders/Private/VirtualShadowMaps/VirtualShadowMapPrintStats.usf
--- a/Engine/Shaders/Private/VirtualShadowMaps/VirtualShadowMapPrintStats.usf
+++ b/Engine/Shaders/Private/VirtualShadowMaps/VirtualShadowMapPrintStats.usf
@@ -4,6 +4,7 @@
 #include "../ShaderPrint.ush"
 #include "VirtualShadowMapPageAccessCommon.ush"
 #include "VirtualShadowMapStats.ush"
+#include "../Nanite/NaniteDataDecode.ush"
 #include "../GPUMessaging.ush"
 
 StructuredBuffer<uint>	InStatsBuffer;
@@ -56,18 +57,33 @@ void PrintUnits(inout FShaderPrintContext Context, uint Value)
 	}
 }
 uint StatsMessageId;
+uint StatsMessageTimestampHigh;
+uint StatsMessageTimestampLow;
+
+StructuredBuffer<FNaniteStats> NaniteStats;
+
+void SendStatsMessage()
+{
+	FGPUMessageWriter Mw = GPUMessageBegin(StatsMessageId, 3U + (MAX_PAGE_AREA_DIAGNOSTIC_SLOTS * 2U));
+
+	// Culling stats
+	GPUMessageWriteItem(Mw, NaniteStats[0].NumTris);
+	GPUMessageWriteItem(Mw, NaniteStats[0].NumMainInstancesPostCull);
+	GPUMessageWriteItem(Mw, InStatsBuffer[VSM_STAT_NON_NANITE_INSTANCES_DRAWN]);
+
+	// Large page area instances
+	for (uint Index = 0; Index < MAX_PAGE_AREA_DIAGNOSTIC_SLOTS * 2U; ++Index)
+	{
+		GPUMessageWriteItem(Mw, InStatsBuffer[VSM_NUM_STATS + Index]);
+	}
+}
 
 [numthreads(1, 1, 1)]
-void PrintVirtualSmStatsCS()
+void LogVirtualSmStatsCS()
 {
 	if (StatsMessageId != INDEX_NONE)
 	{
-		// Just post back the whole buffer
-		FGPUMessageWriter Mw = GPUMessageBegin(StatsMessageId, MAX_PAGE_AREA_DIAGNOSTIC_SLOTS * 2U);
-		for (uint Index = 0; Index < MAX_PAGE_AREA_DIAGNOSTIC_SLOTS * 2U; ++Index)
-		{
-			GPUMessageWriteItem(Mw, InStatsBuffer[VSM_NUM_STATS + Index]);
-		}
+		SendStatsMessage();
 	}
 
 	if (ShowStatsValue == 0U)
diff --git a/Engine/Source/Runtime/RenderCore/Public/GPUMessaging.h b/Engine/Source/Runtime/RenderCore/Public/GPUMessaging.h
--- a/Engine/Source/Runtime/RenderCore/Public/GPUMessaging.h
+++ b/Engine/Source/Runtime/RenderCore/Public/GPUMessaging.h
@@ -90,9 +90,16 @@ public:
 		return DefaultValue;
 	}
 
-	uint32 GetPayloadSize() const { return PayloadSize; }
+	// Read up to N elements
+	TConstArrayView<uint32> ReadCount(int32 Count)
+	{
+		TConstArrayView<uint32> Payload(PayloadData, PayloadSize);
+		Payload.MidInline(ReadOffset, Count);
+		ReadOffset += Payload.Num();
+		return Payload;
+	}
 
-	TConstArrayView<uint32> GetPayLoad() const { return TConstArrayView<uint32>(PayloadData, PayloadSize); }
+	uint32 GetPayloadSize() const { return PayloadSize; }
 
 private:
 	FReader(FMessageId InMessageId, uint32 InPayloadSize, const uint32* InPayloadData)
diff --git a/Engine/Source/Runtime/Renderer/Private/GPUScene.cpp b/Engine/Source/Runtime/Renderer/Private/GPUScene.cpp
--- a/Engine/Source/Runtime/Renderer/Private/GPUScene.cpp
+++ b/Engine/Source/Runtime/Renderer/Private/GPUScene.cpp
@@ -756,6 +756,8 @@ void FGPUScene::UpdateBufferState(FRDGBuilder& GraphBuilder, FScene& Scene, cons
 	const uint32 SizeReserve = FMath::RoundUpToPowerOfTwo(FMath::Max(DynamicPrimitivesOffset, InitialBufferSize));
 	BufferState.PrimitiveBuffer = ResizeStructuredBufferIfNeeded(GraphBuilder, PrimitiveBuffer, SizeReserve * sizeof(FPrimitiveSceneShaderData::Data), TEXT("GPUScene.PrimitiveData"));
 
+	CSV_CUSTOM_STAT_GLOBAL(GPUSceneInstanceCount, float(InstanceSceneDataAllocator.GetMaxSize()), ECsvCustomStatOp::Accumulate);
+
 	const uint32 InstanceSceneDataSizeReserve = FMath::RoundUpToPowerOfTwo(FMath::Max(InstanceSceneDataAllocator.GetMaxSize(), InitialBufferSize));
 	FResizeResourceSOAParams ResizeParams;
 	ResizeParams.NumBytes = InstanceSceneDataSizeReserve * FInstanceSceneShaderData::GetEffectiveNumBytes();
diff --git a/Engine/Source/Runtime/Renderer/Private/Shadows/ShadowSceneRenderer.cpp b/Engine/Source/Runtime/Renderer/Private/Shadows/ShadowSceneRenderer.cpp
--- a/Engine/Source/Runtime/Renderer/Private/Shadows/ShadowSceneRenderer.cpp
+++ b/Engine/Source/Runtime/Renderer/Private/Shadows/ShadowSceneRenderer.cpp
@@ -12,6 +12,8 @@
 #include "DynamicPrimitiveDrawing.h"
 #endif
 
+CSV_DECLARE_CATEGORY_EXTERN(VSM);
+
 TAutoConsoleVariable<int32> CVarMaxDistantLightsPerFrame(
 	TEXT("r.Shadow.Virtual.MaxDistantUpdatePerFrame"),
 	1,
@@ -33,11 +35,6 @@ static TAutoConsoleVariable<float> CVarNaniteShadowsLODBias(
 	TEXT("LOD bias for nanite geometry in shadows. 0 = full detail. >0 = reduced detail."),
 	ECVF_RenderThreadSafe);
 
-namespace Nanite
-{
-	extern bool IsStatFilterActive(const FString& FilterName);
-}
-
 FShadowSceneRenderer::FShadowSceneRenderer(FDeferredShadingSceneRenderer& InSceneRenderer)
 	: SceneRenderer(InSceneRenderer)
 	, Scene(*InSceneRenderer.Scene)
@@ -120,107 +117,10 @@ void FShadowSceneRenderer::RenderVirtualShadowMaps(FRDGBuilder& GraphBuilder, bo
 		return;
 	}
 
-	FVirtualShadowMapArrayCacheManager* CacheManager = VirtualShadowMapArray.CacheManager;
-
-	// TODO: Separate out the decision about nanite using HZB and stuff like HZB culling invalidations?
-	const bool bVSMUseHZB = VirtualShadowMapArray.UseHzbOcclusion();
-
-	const FIntPoint VirtualShadowSize = VirtualShadowMapArray.GetPhysicalPoolSize();
-	const FIntRect VirtualShadowViewRect = FIntRect(0, 0, VirtualShadowSize.X, VirtualShadowSize.Y);
-
-	Nanite::FSharedContext SharedContext{};
-	SharedContext.FeatureLevel = SceneRenderer.FeatureLevel;
-	SharedContext.ShaderMap = GetGlobalShaderMap(SharedContext.FeatureLevel);
-	SharedContext.Pipeline = Nanite::EPipeline::Shadows;
-
 	if (bNaniteEnabled)
 	{
-		const TRefCountPtr<IPooledRenderTarget> PrevHZBPhysical = bVSMUseHZB ? CacheManager->PrevBuffers.HZBPhysical : nullptr;
-
-		RDG_EVENT_SCOPE(GraphBuilder, "RenderVirtualShadowMaps(Nanite)");
-
-		check(VirtualShadowMapArray.PhysicalPagePoolRDG != nullptr);
-
-		Nanite::FRasterContext RasterContext = Nanite::InitRasterContext(
-			GraphBuilder,
-			SharedContext,
-			VirtualShadowSize,
-			VirtualShadowViewRect,
-			false,
-			Nanite::EOutputBufferMode::DepthOnly,
-			false,	// Clear entire texture
-			nullptr, 0,
-			VirtualShadowMapArray.PhysicalPagePoolRDG);
-
-		const FViewInfo& SceneView = SceneRenderer.Views[0];
-
-		static FString VirtualFilterName = TEXT("VirtualShadowMaps");
-
-		TArray<Nanite::FPackedView, SceneRenderingAllocator> VirtualShadowViews;
-
-		for (FProjectedShadowInfo* ProjectedShadowInfo : VirtualShadowMapShadows)
-		{
-			if (ProjectedShadowInfo->bShouldRenderVSM)
-			{
-				VirtualShadowMapArray.AddRenderViews(
-					ProjectedShadowInfo,
-					SceneRenderer.Views,
-					ComputeNaniteShadowsLODScaleFactor(),
-					PrevHZBPhysical.IsValid(),
-					bVSMUseHZB,
-					ProjectedShadowInfo->ShouldClampToNearPlane(),
-					VirtualShadowViews);
-			}
-		}
-
-		if (VirtualShadowViews.Num() > 0)
-		{
-			int32 NumPrimaryViews = VirtualShadowViews.Num();
-			VirtualShadowMapArray.CreateMipViews(VirtualShadowViews);
-
-			Nanite::FRasterState RasterState;
-
-			FNaniteVisibilityResults VisibilityResults; // TODO: Hook up culling for shadows
-
-			Nanite::FCullingContext::FConfiguration CullingConfig = { 0 };
-			CullingConfig.bUpdateStreaming = bUpdateNaniteStreaming;
-			CullingConfig.bTwoPassOcclusion = VirtualShadowMapArray.UseTwoPassHzbOcclusion();
-			CullingConfig.bProgrammableRaster = bNaniteProgrammableRaster;
-			CullingConfig.SetViewFlags(SceneView);
-
-			Nanite::FCullingContext CullingContext = Nanite::InitCullingContext(
-				GraphBuilder,
-				SharedContext,
-				Scene,
-				PrevHZBPhysical,
-				VirtualShadowViewRect,
-				CullingConfig
-			);
-
-			const bool bExtractStats = Nanite::IsStatFilterActive(VirtualFilterName);
-
-			Nanite::CullRasterize(
-				GraphBuilder,
-				Scene.NaniteRasterPipelines[ENaniteMeshPass::BasePass],
-				VisibilityResults,
-				Scene,
-				SceneView,
-				VirtualShadowViews,
-				NumPrimaryViews,
-				SharedContext,
-				CullingContext,
-				RasterContext,
-				RasterState,
-				nullptr,
-				&VirtualShadowMapArray,
-				bExtractStats
-			);
-		}
-
-		if (bVSMUseHZB)
-		{
-			VirtualShadowMapArray.UpdateHZB(GraphBuilder);
-		}
+		const float ShadowsLODScaleFactor = ComputeNaniteShadowsLODScaleFactor();
+		VirtualShadowMapArray.RenderVirtualShadowMapsNanite(GraphBuilder, SceneRenderer, ShadowsLODScaleFactor, bUpdateNaniteStreaming, bNaniteProgrammableRaster);
 	}
 
 	if (UseNonNaniteVirtualShadowMaps(SceneRenderer.ShaderPlatform, SceneRenderer.FeatureLevel))
diff --git a/Engine/Source/Runtime/Renderer/Private/VirtualShadowMaps/VirtualShadowMapArray.cpp b/Engine/Source/Runtime/Renderer/Private/VirtualShadowMaps/VirtualShadowMapArray.cpp
--- a/Engine/Source/Runtime/Renderer/Private/VirtualShadowMaps/VirtualShadowMapArray.cpp
+++ b/Engine/Source/Runtime/Renderer/Private/VirtualShadowMaps/VirtualShadowMapArray.cpp
@@ -9,6 +9,7 @@
 #include "GPUMessaging.h"
 #include "HairStrands/HairStrandsData.h"
 #include "InstanceCulling/InstanceCullingMergedContext.h"
+#include "Nanite/Nanite.h"
 #include "RendererModule.h"
 #include "Rendering/NaniteResources.h"
 #include "SceneTextureReductions.h"
@@ -26,6 +27,8 @@ IMPLEMENT_STATIC_UNIFORM_BUFFER_SLOT(VirtualShadowMapUbSlot);
 
 IMPLEMENT_STATIC_UNIFORM_BUFFER_STRUCT(FVirtualShadowMapUniformParameters, "VirtualShadowMap", VirtualShadowMapUbSlot);
 
+CSV_DEFINE_CATEGORY(VSM, false);
+
 struct FShadowMapCacheData
 {
 	int32 PrevVirtualShadowMapId = INDEX_NONE;
@@ -343,6 +346,11 @@ static TAutoConsoleVariable<int32> CVarVirtualShadowMapPageMarkingPixelStrideY(
 	TEXT("."),
 	ECVF_RenderThreadSafe);
 
+namespace Nanite
+{
+	extern bool IsStatFilterActive(const FString& FilterName);
+}
+
 FMatrix CalcTranslatedWorldToShadowUVMatrix(
 	const FMatrix& TranslatedWorldToShadowView,
 	const FMatrix& ViewToClip)
@@ -1338,11 +1346,15 @@ void FVirtualShadowMapArray::BuildPageAllocations(
 
 #if !UE_BUILD_SHIPPING
 	const bool bRunPageAreaDiagnostics = CVarNumPageAreaDiagSlots.GetValueOnRenderThread() != 0;
+#if CSV_PROFILER
+	bool bCsvLogEnabled = FCsvProfiler::Get()->IsCapturing_Renderthread() && FCsvProfiler::Get()->IsCategoryEnabled(CSV_CATEGORY_INDEX(VSM));
+#endif
 #else
 	constexpr bool bRunPageAreaDiagnostics = false;
+	constexpr bool bCsvLogEnabled = false;
 #endif
 
-	if (CVarShowStats.GetValueOnRenderThread() || CacheManager->IsAccumulatingStats() || bRunPageAreaDiagnostics)
+	if (CVarShowStats.GetValueOnRenderThread() || CacheManager->IsAccumulatingStats() || bRunPageAreaDiagnostics || bCsvLogEnabled)
 	{
 		StatsBufferRDG = GraphBuilder.CreateBuffer(FRDGBufferDesc::CreateStructuredDesc(sizeof(uint32), NumStats + MaxPageAreaDiagnosticSlots * 2), TEXT("Shadow.Virtual.StatsBuffer"));
 		AddClearUAVPass(GraphBuilder, GraphBuilder.CreateUAV(StatsBufferRDG), 0);
@@ -1825,7 +1837,7 @@ void FVirtualShadowMapArray::RenderDebugInfo(FRDGBuilder& GraphBuilder, TArrayVi
 			
 	if (Views.Num() > 0)
 	{
-		PrintStats(GraphBuilder, Views[0]);
+		LogStats(GraphBuilder, Views[0]);
 	}
 
 	if (DebugVisualizationOutput.IsEmpty() || VisualizeLight.IsEmpty())
@@ -1876,16 +1888,17 @@ void FVirtualShadowMapArray::RenderDebugInfo(FRDGBuilder& GraphBuilder, TArrayVi
 }
 
 
-class FVirtualSmPrintStatsCS : public FVirtualPageManagementShader
+class FVirtualSmLogStatsCS : public FVirtualPageManagementShader
 {
-	DECLARE_GLOBAL_SHADER(FVirtualSmPrintStatsCS);
-	SHADER_USE_PARAMETER_STRUCT(FVirtualSmPrintStatsCS, FVirtualPageManagementShader)
+	DECLARE_GLOBAL_SHADER(FVirtualSmLogStatsCS);
+	SHADER_USE_PARAMETER_STRUCT(FVirtualSmLogStatsCS, FVirtualPageManagementShader)
 
 	BEGIN_SHADER_PARAMETER_STRUCT(FParameters, )
 		SHADER_PARAMETER_RDG_UNIFORM_BUFFER(FVirtualShadowMapUniformParameters, VirtualShadowMap)
 		SHADER_PARAMETER_STRUCT_INCLUDE(GPUMessage::FParameters, GPUMessageParams)
 		SHADER_PARAMETER_STRUCT_INCLUDE( ShaderPrint::FShaderParameters, ShaderPrintStruct )
-		SHADER_PARAMETER_RDG_BUFFER_SRV(StructuredBuffer< uint >, InStatsBuffer)
+		SHADER_PARAMETER_RDG_BUFFER_SRV(StructuredBuffer<uint>, InStatsBuffer)
+		SHADER_PARAMETER_RDG_BUFFER_SRV(StructuredBuffer<FNaniteStats>, NaniteStats)
 		SHADER_PARAMETER_RDG_BUFFER_SRV(StructuredBuffer<FIntVector4>, AllocatedPageRectBounds)
 		SHADER_PARAMETER(int, ShowStatsValue)
 		SHADER_PARAMETER(uint32, StatsMessageId)
@@ -1898,47 +1911,55 @@ class FVirtualSmPrintStatsCS : public FVirtualPageManagementShader
 		// Disable optimizations as shader print causes long compile times
 		OutEnvironment.CompilerFlags.Add(CFLAG_SkipOptimizations);
 	}
-
-		
 };
-IMPLEMENT_GLOBAL_SHADER(FVirtualSmPrintStatsCS, "/Engine/Private/VirtualShadowMaps/VirtualShadowMapPrintStats.usf", "PrintVirtualSmStatsCS", SF_Compute);
+IMPLEMENT_GLOBAL_SHADER(FVirtualSmLogStatsCS, "/Engine/Private/VirtualShadowMaps/VirtualShadowMapPrintStats.usf", "LogVirtualSmStatsCS", SF_Compute);
 
-void FVirtualShadowMapArray::PrintStats(FRDGBuilder& GraphBuilder, const FViewInfo& View)
+void FVirtualShadowMapArray::LogStats(FRDGBuilder& GraphBuilder, const FViewInfo& View)
 {
 	check(IsEnabled());
 	LLM_SCOPE_BYTAG(Nanite);
 
-	// Print stats
-	int ShowStatsValue = CVarShowStats.GetValueOnRenderThread();
 	if (StatsBufferRDG)
 	{
-		{
-			FVirtualSmPrintStatsCS::FParameters* PassParameters = GraphBuilder.AllocParameters<FVirtualSmPrintStatsCS::FParameters>();
+		int ShowStatsValue = CVarShowStats.GetValueOnRenderThread();
 
-			ShaderPrint::SetParameters(GraphBuilder, View.ShaderPrintData, PassParameters->ShaderPrintStruct);
-			PassParameters->InStatsBuffer = GraphBuilder.CreateSRV(StatsBufferRDG);
-			PassParameters->VirtualShadowMap = GetUncachedUniformBuffer(GraphBuilder);
-			PassParameters->ShowStatsValue = ShowStatsValue;
+		FVirtualSmLogStatsCS::FParameters* PassParameters = GraphBuilder.AllocParameters<FVirtualSmLogStatsCS::FParameters>();
+
+		ShaderPrint::SetParameters(GraphBuilder, View.ShaderPrintData, PassParameters->ShaderPrintStruct);
+		PassParameters->InStatsBuffer = GraphBuilder.CreateSRV(StatsBufferRDG);
+		PassParameters->VirtualShadowMap = GetUncachedUniformBuffer(GraphBuilder);
+		PassParameters->ShowStatsValue = ShowStatsValue;
+		
 #if !UE_BUILD_SHIPPING
-			PassParameters->StatsMessageId = CacheManager->StatsFeedbackSocket.GetMessageId().IsValid() ? CacheManager->StatsFeedbackSocket.GetMessageId().GetIndex() : INDEX_NONE;
+		PassParameters->StatsMessageId = CacheManager->StatsFeedbackSocket.GetMessageId().IsValid() ? CacheManager->StatsFeedbackSocket.GetMessageId().GetIndex() : INDEX_NONE;
+		bool bBindNaniteStatsBuffer = StatsNaniteBufferRDG != nullptr;
 #else
-			PassParameters->StatsMessageId = INDEX_NONE;
+		PassParameters->StatsMessageId = INDEX_NONE;
+		constexpr bool bBindNaniteStatsBuffer = false;
 #endif
-			PassParameters->GPUMessageParams = GPUMessage::GetShaderParameters(GraphBuilder);
-
-			auto ComputeShader = View.ShaderMap->GetShader<FVirtualSmPrintStatsCS>();
-
-			FComputeShaderUtils::AddPass(
-				GraphBuilder,
-				RDG_EVENT_NAME("Print Stats"),
-				ComputeShader,
-				PassParameters,
-				FIntVector(1, 1, 1)
-			);
+		if (bBindNaniteStatsBuffer)
+		{
+			PassParameters->NaniteStats = GraphBuilder.CreateSRV(StatsNaniteBufferRDG);
+		}
+		else
+		{
+			PassParameters->NaniteStats = GraphBuilder.CreateSRV(GSystemTextures.GetDefaultStructuredBuffer(GraphBuilder, sizeof(FNaniteStats)));
 		}
+		PassParameters->GPUMessageParams = GPUMessage::GetShaderParameters(GraphBuilder);
+	
+		auto ComputeShader = View.ShaderMap->GetShader<FVirtualSmLogStatsCS>();
+
+		FComputeShaderUtils::AddPass(
+			GraphBuilder,
+			RDG_EVENT_NAME("VSM Log Stats"),
+			ComputeShader,
+			PassParameters,
+			FIntVector(1, 1, 1)
+		);
 	}
 }
 
+
 void FVirtualShadowMapArray::CreateMipViews( TArray<Nanite::FPackedView, SceneRenderingAllocator>& Views ) const
 {
 	// strategy: 
@@ -2493,6 +2514,125 @@ static void AddRasterPass(
 		});
 }
 
+void FVirtualShadowMapArray::RenderVirtualShadowMapsNanite(FRDGBuilder& GraphBuilder, FSceneRenderer& SceneRenderer, float ShadowsLODScaleFactor, bool bUpdateNaniteStreaming, bool bNaniteProgrammableRaster)
+{
+	bool bCsvLogEnabled = false;
+#if CSV_PROFILER
+	bCsvLogEnabled = FCsvProfiler::Get()->IsCapturing_Renderthread() && FCsvProfiler::Get()->IsCategoryEnabled(CSV_CATEGORY_INDEX(VSM));
+#endif
+
+	const TArray<FProjectedShadowInfo*, SceneRenderingAllocator>& VirtualShadowMapShadows = SceneRenderer.SortedShadowsForShadowDepthPass.VirtualShadowMapShadows;
+	TArray<TSharedPtr<FVirtualShadowMapClipmap>, SceneRenderingAllocator>& VirtualShadowMapClipmaps = SceneRenderer.SortedShadowsForShadowDepthPass.VirtualShadowMapClipmaps;
+
+	// TODO: Separate out the decision about nanite using HZB and stuff like HZB culling invalidations?
+	const bool bVSMUseHZB = UseHzbOcclusion();
+
+	const FIntPoint VirtualShadowSize = GetPhysicalPoolSize();
+	const FIntRect VirtualShadowViewRect = FIntRect(0, 0, VirtualShadowSize.X, VirtualShadowSize.Y);
+
+	Nanite::FSharedContext SharedContext{};
+	SharedContext.FeatureLevel = SceneRenderer.FeatureLevel;
+	SharedContext.ShaderMap = GetGlobalShaderMap(SharedContext.FeatureLevel);
+	SharedContext.Pipeline = Nanite::EPipeline::Shadows;
+
+	const TRefCountPtr<IPooledRenderTarget> PrevHZBPhysical = bVSMUseHZB ? CacheManager->PrevBuffers.HZBPhysical : nullptr;
+
+	RDG_EVENT_SCOPE(GraphBuilder, "RenderVirtualShadowMaps(Nanite)");
+
+	check(PhysicalPagePoolRDG != nullptr);
+
+	Nanite::FRasterContext RasterContext = Nanite::InitRasterContext(
+		GraphBuilder,
+		SharedContext,
+		VirtualShadowSize,
+		VirtualShadowViewRect,
+		false,
+		Nanite::EOutputBufferMode::DepthOnly,
+		false,	// Clear entire texture
+		nullptr, 0,
+		PhysicalPagePoolRDG);
+
+	const FViewInfo& SceneView = SceneRenderer.Views[0];
+
+	static FString VirtualFilterName = TEXT("VirtualShadowMaps");
+
+	TArray<Nanite::FPackedView, SceneRenderingAllocator> VirtualShadowViews;
+
+	for (FProjectedShadowInfo* ProjectedShadowInfo : VirtualShadowMapShadows)
+	{
+		if (ProjectedShadowInfo->bShouldRenderVSM)
+		{
+			AddRenderViews(
+				ProjectedShadowInfo,
+				SceneRenderer.Views,
+				ShadowsLODScaleFactor,
+				PrevHZBPhysical.IsValid(),
+				bVSMUseHZB,
+				ProjectedShadowInfo->ShouldClampToNearPlane(),
+				VirtualShadowViews);
+		}
+	}
+
+	if (VirtualShadowViews.Num() > 0)
+	{
+		int32 NumPrimaryViews = VirtualShadowViews.Num();
+		CreateMipViews(VirtualShadowViews);
+
+		Nanite::FRasterState RasterState;
+
+		FNaniteVisibilityResults VisibilityResults; // TODO: Hook up culling for shadows
+
+		Nanite::FCullingContext::FConfiguration CullingConfig = { 0 };
+		CullingConfig.bUpdateStreaming = bUpdateNaniteStreaming;
+		CullingConfig.bTwoPassOcclusion = UseTwoPassHzbOcclusion();
+		CullingConfig.bProgrammableRaster = bNaniteProgrammableRaster;
+		CullingConfig.SetViewFlags(SceneView);
+
+		Nanite::FCullingContext CullingContext = Nanite::InitCullingContext(
+			GraphBuilder,
+			SharedContext,
+			Scene,
+			PrevHZBPhysical,
+			VirtualShadowViewRect,
+			CullingConfig
+		);
+
+		if (bCsvLogEnabled)
+		{
+			CullingContext.DebugFlags |= NANITE_DEBUG_FLAG_WRITE_STATS;
+		}
+
+		const bool bExtractStats = Nanite::IsStatFilterActive(VirtualFilterName);
+
+		Nanite::CullRasterize(
+			GraphBuilder,
+			Scene.NaniteRasterPipelines[ENaniteMeshPass::BasePass],
+			VisibilityResults,
+			Scene,
+			SceneView,
+			VirtualShadowViews,
+			NumPrimaryViews,
+			SharedContext,
+			CullingContext,
+			RasterContext,
+			RasterState,
+			nullptr,
+			this,
+			bExtractStats
+		);
+
+		if (bCsvLogEnabled)
+		{
+			StatsNaniteBufferRDG = CullingContext.StatsBuffer;
+		}
+	}
+
+	if (bVSMUseHZB)
+	{
+		UpdateHZB(GraphBuilder);
+	}
+}
+
 void FVirtualShadowMapArray::RenderVirtualShadowMapsNonNanite(FRDGBuilder& GraphBuilder, const TArray<FProjectedShadowInfo *, SceneRenderingAllocator>& VirtualSmMeshCommandPasses, TArrayView<FViewInfo> Views)
 {
 	if (VirtualSmMeshCommandPasses.Num() == 0)
@@ -2576,6 +2716,7 @@ void FVirtualShadowMapArray::RenderVirtualShadowMapsNonNanite(FRDGBuilder& Graph
 	FInstanceCullingMergedContext InstanceCullingMergedContext(GPUScene.GetFeatureLevel());
 	// We don't use the registered culling views (this redundancy should probably be addressed at some point), set the number to disable index range checking
 	InstanceCullingMergedContext.NumCullingViews = -1;
+	int32 TotalPreCullInstanceCount = 0;
 	for (int32 Index = 0; Index < VirtualSmMeshCommandPasses.Num(); ++Index)
 	{
 		FProjectedShadowInfo* ProjectedShadowInfo = VirtualSmMeshCommandPasses[Index];
@@ -2613,6 +2754,7 @@ void FVirtualShadowMapArray::RenderVirtualShadowMapsNonNanite(FRDGBuilder& Graph
 			MeshCommandPass.WaitForSetupTask();
 
 			FInstanceCullingContext* InstanceCullingContext = MeshCommandPass.GetInstanceCullingContext();
+			TotalPreCullInstanceCount += InstanceCullingContext->TotalInstances;
 
 			if (InstanceCullingContext->HasCullingCommands())
 			{
@@ -2641,6 +2783,7 @@ void FVirtualShadowMapArray::RenderVirtualShadowMapsNonNanite(FRDGBuilder& Graph
 			}
 		}
 	}
+	CSV_CUSTOM_STAT(VSM, NonNanitePreCullInstanceCount, TotalPreCullInstanceCount, ECsvCustomStatOp::Set);
 	uint32 TotalPrimaryViews = uint32(VirtualShadowViews.Num());
 	CreateMipViews(VirtualShadowViews);
 	FRDGBufferRef VirtualShadowViewsRDG = CreateStructuredBuffer(GraphBuilder, TEXT("Shadow.Virtual.VirtualShadowViews"), VirtualShadowViews);
diff --git a/Engine/Source/Runtime/Renderer/Private/VirtualShadowMaps/VirtualShadowMapArray.h b/Engine/Source/Runtime/Renderer/Private/VirtualShadowMaps/VirtualShadowMapArray.h
--- a/Engine/Source/Runtime/Renderer/Private/VirtualShadowMaps/VirtualShadowMapArray.h
+++ b/Engine/Source/Runtime/Renderer/Private/VirtualShadowMaps/VirtualShadowMapArray.h
@@ -278,6 +278,11 @@ public:
 
 	void CreateMipViews( TArray<Nanite::FPackedView, SceneRenderingAllocator>& Views ) const;
 
+	/**
+	 * Draw Nanite geometry into the VSMs.
+	 */
+	void RenderVirtualShadowMapsNanite(FRDGBuilder& GraphBuilder, FSceneRenderer& SceneRenderer, float ShadowsLODScaleFactor, bool bUpdateNaniteStreaming, bool bNaniteProgrammableRaster);
+
 	/**
 	 * Draw Non-Nanite geometry into the VSMs.
 	 */
@@ -285,7 +290,7 @@ public:
 
 	void RenderDebugInfo(FRDGBuilder& GraphBuilder, TArrayView<FViewInfo> Views);
 	
-	void PrintStats(FRDGBuilder& GraphBuilder, const FViewInfo& View);
+	void LogStats(FRDGBuilder& GraphBuilder, const FViewInfo& View);
 
 	// Get shader parameters necessary to sample virtual shadow maps
 	// It is safe to bind this buffer even if VSMs are disabled, but the sampling should be branched around in the shader.
@@ -352,6 +357,7 @@ public:
 	static constexpr uint32 MaxPageAreaDiagnosticSlots = 32;
 
 	FRDGBufferRef StatsBufferRDG = nullptr;
+	FRDGBufferRef StatsNaniteBufferRDG = nullptr;
 
 	// Debug visualization
 	TArray<FRDGTextureRef> DebugVisualizationOutput;
diff --git a/Engine/Source/Runtime/Renderer/Private/VirtualShadowMaps/VirtualShadowMapCacheManager.cpp b/Engine/Source/Runtime/Renderer/Private/VirtualShadowMaps/VirtualShadowMapCacheManager.cpp
--- a/Engine/Source/Runtime/Renderer/Private/VirtualShadowMaps/VirtualShadowMapCacheManager.cpp
+++ b/Engine/Source/Runtime/Renderer/Private/VirtualShadowMaps/VirtualShadowMapCacheManager.cpp
@@ -17,6 +17,8 @@
 
 #define VSM_LOG_STATIC_CACHING 0
 
+CSV_DECLARE_CATEGORY_EXTERN(VSM);
+
 static TAutoConsoleVariable<int32> CVarAccumulateStats(
 	TEXT("r.Shadow.Virtual.AccumulateStats"),
 	0,
@@ -464,60 +466,66 @@ FVirtualShadowMapArrayCacheManager::FVirtualShadowMapArrayCacheManager(FScene* I
 	// Handle message with status sent back from GPU
 	StatusFeedbackSocket = GPUMessage::RegisterHandler(TEXT("Shadow.Virtual.StatusFeedback"), [this](GPUMessage::FReader Message)
 	{
-		// Only process status messages that came from this specific cache manager
-		if (Message.MessageId == this->StatusFeedbackSocket.GetMessageId())
-		{
-			// Get the frame that the message was sent.
-			uint32 FrameNumber = Message.Read<uint32>(0);
-			// Goes negative on underflow
-			int32 NumPagesFree = Message.Read<int32>(0);
+		// Goes negative on underflow
+		int32 NumPagesFree = Message.Read<int32>(0);
+			
+		CSV_CUSTOM_STAT(VSM, FreePages, NumPagesFree, ECsvCustomStatOp::Set);
 
-			if (NumPagesFree < 0)
-			{
-				static const auto* CVarResolutionLodBiasLocalPtr = IConsoleManager::Get().FindTConsoleVariableDataFloat(TEXT("r.Shadow.Virtual.ResolutionLodBiasLocal"));
-				const float LodBiasLocal = CVarResolutionLodBiasLocalPtr->GetValueOnRenderThread();
+		if (NumPagesFree < 0)
+		{
+			static const auto* CVarResolutionLodBiasLocalPtr = IConsoleManager::Get().FindTConsoleVariableDataFloat(TEXT("r.Shadow.Virtual.ResolutionLodBiasLocal"));
+			const float LodBiasLocal = CVarResolutionLodBiasLocalPtr->GetValueOnRenderThread();
 
-				static const auto* CVarResolutionLodBiasDirectionalPtr = IConsoleManager::Get().FindTConsoleVariableDataFloat(TEXT("r.Shadow.Virtual.ResolutionLodBiasDirectional"));
-				const float LodBiasDirectional = CVarResolutionLodBiasDirectionalPtr->GetValueOnRenderThread();
+			static const auto* CVarResolutionLodBiasDirectionalPtr = IConsoleManager::Get().FindTConsoleVariableDataFloat(TEXT("r.Shadow.Virtual.ResolutionLodBiasDirectional"));
+			const float LodBiasDirectional = CVarResolutionLodBiasDirectionalPtr->GetValueOnRenderThread();
 
-				static const auto* CVarMaxPhysicalPagesPtr = IConsoleManager::Get().FindTConsoleVariableDataInt(TEXT("r.Shadow.Virtual.MaxPhysicalPages"));
-				const int32 MaxPhysicalPages = CVarMaxPhysicalPagesPtr->GetValueOnRenderThread();
+			static const auto* CVarMaxPhysicalPagesPtr = IConsoleManager::Get().FindTConsoleVariableDataInt(TEXT("r.Shadow.Virtual.MaxPhysicalPages"));
+			const int32 MaxPhysicalPages = CVarMaxPhysicalPagesPtr->GetValueOnRenderThread();
 
 				static const auto* CVarMaxPhysicalPagesSceneCapturePtr = IConsoleManager::Get().FindTConsoleVariableDataInt(TEXT("r.Shadow.Virtual.MaxPhysicalPagesSceneCapture"));
 				const int32 MaxPhysicalPagesSceneCapture = CVarMaxPhysicalPagesSceneCapturePtr->GetValueOnRenderThread();
 
 #if !UE_BUILD_SHIPPING
-				if (!bLoggedPageOverflow)
-				{
-					UE_LOG(LogRenderer, Warning, TEXT("Virtual Shadow Map Page Pool overflow (%d page allocations were not served), this will produce visual artifacts (missing shadow), increase the page pool limit or reduce resolution bias to avoid.\n")
-						TEXT(" See r.Shadow.Virtual.MaxPhysicalPages (%d), r.Shadow.Virtual.MaxPhysicalPagesSceneCapture (%d), r.Shadow.Virtual.ResolutionLodBiasLocal (%.2f), and r.Shadow.Virtual.ResolutionLodBiasDirectional (%.2f)"),
-						-NumPagesFree,
-						MaxPhysicalPages,
-						MaxPhysicalPagesSceneCapture,
-						LodBiasLocal,
-						LodBiasDirectional);
-					bLoggedPageOverflow = true;
-				}
-				LastOverflowTime = float(FGameTime::GetTimeSinceAppStart().GetRealTimeSeconds());
-#endif
-			}
-#if !UE_BUILD_SHIPPING
-			else
+			if (!bLoggedPageOverflow)
 			{
-				bLoggedPageOverflow = false;
+				UE_LOG(LogRenderer, Warning, TEXT("Virtual Shadow Map Page Pool overflow (%d page allocations were not served), this will produce visual artifacts (missing shadow), increase the page pool limit or reduce resolution bias to avoid.\n")
+					TEXT(" See r.Shadow.Virtual.MaxPhysicalPages (%d), r.Shadow.Virtual.MaxPhysicalPagesSceneCapture (%d), r.Shadow.Virtual.ResolutionLodBiasLocal (%.2f), and r.Shadow.Virtual.ResolutionLodBiasDirectional (%.2f)"),
+					-NumPagesFree,
+					MaxPhysicalPages,
+					MaxPhysicalPagesSceneCapture,
+					LodBiasLocal,
+					LodBiasDirectional);
+				bLoggedPageOverflow = true;
 			}
+			LastOverflowTime = float(FGameTime::GetTimeSinceAppStart().GetRealTimeSeconds());
 #endif
 		}
+#if !UE_BUILD_SHIPPING
+		else
+		{
+			bLoggedPageOverflow = false;
+		}
+#endif
 	});
 
 #if !UE_BUILD_SHIPPING
 	// Handle message with stats sent back from GPU whenever stats are enabled
 	StatsFeedbackSocket = GPUMessage::RegisterHandler(TEXT("Shadow.Virtual.StatsFeedback"), [this](GPUMessage::FReader Message)
 	{
+		// Culling stats
+		int32 NaniteNumTris = Message.Read<int32>(0);
+		int32 NanitePostCullNodeCount = Message.Read<int32>(0);
+		int32 NonNanitePostCullInstanceCount = Message.Read<int32>(0);
+
+		CSV_CUSTOM_STAT(VSM, NaniteNumTris, NaniteNumTris, ECsvCustomStatOp::Set);
+		CSV_CUSTOM_STAT(VSM, NanitePostCullNodeCount, NanitePostCullNodeCount, ECsvCustomStatOp::Set);
+		CSV_CUSTOM_STAT(VSM, NonNanitePostCullInstanceCount, NonNanitePostCullInstanceCount, ECsvCustomStatOp::Set);
+
+		// Large page area items
 		LastLoggedPageOverlapAppTime.SetNumZeroed(Scene->GetMaxPersistentPrimitiveIndex());
 		float RealTimeSeconds = float(FGameTime::GetTimeSinceAppStart().GetRealTimeSeconds());
 
-		TConstArrayView<uint32> PageAreaDiags = Message.GetPayLoad();
+		TConstArrayView<uint32> PageAreaDiags = Message.ReadCount(FVirtualShadowMapArray::MaxPageAreaDiagnosticSlots * 2);
 		for (int32 Index = 0; Index < PageAreaDiags.Num(); Index += 2)
 		{
 			uint32 Overlap = PageAreaDiags[Index];
