Volumetric cloud - using blue noise for tracing offset to increase visual quality.

diff --git a/Engine/Shaders/Private/VolumetricCloud.usf b/Engine/Shaders/Private/VolumetricCloud.usf
--- a/Engine/Shaders/Private/VolumetricCloud.usf
+++ b/Engine/Shaders/Private/VolumetricCloud.usf
@@ -254,6 +254,9 @@ bool RayIntersectSphereSolution(float3 RayOrigin, float3 RayDirection, float4 Sp
 
 #include "/Engine/Private/SkyAtmosphereCommon.ush"
 
+#define BlueNoise RenderVolumetricCloudParameters.BlueNoise
+#include "BlueNoise.ush"
+
 #define FogStruct RenderVolumetricCloudParameters
 #undef MATERIALBLENDING_ADDITIVE	// we need to override this for height fog to appear
 #define MATERIALBLENDING_ADDITIVE 0
@@ -738,8 +741,13 @@ void MainCommon(in FMaterialPixelParameters MaterialParameters, in float4 SvPosi
 		float t = TMin + 0.5 * StepT;
 		if (RenderVolumetricCloudParameters.IsReflectionRendering == 0)
 		{
+		#if 0
 			// Using simple ALU based noise as it seems to work the best now. Mod8 is also used to stabilise the result when the amount of sample cannot converge to the correct integration.
 			t = TMin + (float(Rand3DPCG16(int3(SvPosition.xy, View.StateFrameIndexMod8)).x) * rcp(65536.0)) * StepT;
+		#else
+			uint2 FullResPixelCoord = SvPosition.xy * RenderVolumetricCloudParameters.TracingCoordToFullResPixelCoordScaleBias.xy + RenderVolumetricCloudParameters.TracingCoordToFullResPixelCoordScaleBias.xy;
+			t = TMin + BlueNoiseScalar(FullResPixelCoord, View.StateFrameIndexMod8) * StepT;
+		#endif
 		}
 
 	#if SHADER_RENDERVIEW_CS_WITH_EMPTY_SPACE_SKIPPING
diff --git a/Engine/Source/Runtime/Renderer/Private/SceneRendering.cpp b/Engine/Source/Runtime/Renderer/Private/SceneRendering.cpp
--- a/Engine/Source/Runtime/Renderer/Private/SceneRendering.cpp
+++ b/Engine/Source/Runtime/Renderer/Private/SceneRendering.cpp
@@ -2515,7 +2515,7 @@ FSceneRenderer::FSceneRenderer(const FSceneViewFamily* InViewFamily, FHitProxyCo
 		}
 #endif
 
-		if (ShouldRenderLumenDiffuseGI(Scene, *ViewInfo) || ShouldRenderLumenReflections(*ViewInfo))
+		if (ShouldRenderLumenDiffuseGI(Scene, *ViewInfo) || ShouldRenderLumenReflections(*ViewInfo) || ShouldRenderVolumetricCloudWithBlueNoise(Scene, ViewFamily.EngineShowFlags))
 		{
 			GEngine->LoadBlueNoiseTexture();
 		}
diff --git a/Engine/Source/Runtime/Renderer/Private/VolumetricCloudRendering.cpp b/Engine/Source/Runtime/Renderer/Private/VolumetricCloudRendering.cpp
--- a/Engine/Source/Runtime/Renderer/Private/VolumetricCloudRendering.cpp
+++ b/Engine/Source/Runtime/Renderer/Private/VolumetricCloudRendering.cpp
@@ -272,6 +272,11 @@ bool ShouldRenderVolumetricCloud(const FScene* Scene, const FEngineShowFlags& En
 	return false;
 }
 
+bool ShouldRenderVolumetricCloudWithBlueNoise(const FScene* Scene, const FEngineShowFlags& EngineShowFlags)
+{
+	return Scene && Scene->HasVolumetricCloud() && EngineShowFlags.Atmosphere && EngineShowFlags.Cloud;
+}
+
 bool ShouldViewVisualizeVolumetricCloudConservativeDensity(const FViewInfo& ViewInfo, const FEngineShowFlags& EngineShowFlags)
 {
 	return (!!EngineShowFlags.VisualizeVolumetricCloudConservativeDensity || !!EngineShowFlags.VisualizeVolumetricCloudEmptySpaceSkipping)
@@ -610,6 +615,7 @@ BEGIN_GLOBAL_SHADER_PARAMETER_STRUCT(FRenderVolumetricCloudGlobalParameters, )
 	SHADER_PARAMETER(float, AerialPerspectiveMieOnlyFadeDistanceKmInv)
 	SHADER_PARAMETER_STRUCT_INCLUDE(FFogUniformParameters, FogStruct)
 	SHADER_PARAMETER_STRUCT(FForwardLightData, ForwardLightData)
+	SHADER_PARAMETER_STRUCT(FBlueNoise, BlueNoise)
 END_GLOBAL_SHADER_PARAMETER_STRUCT()
 
 IMPLEMENT_STATIC_UNIFORM_BUFFER_STRUCT(FRenderVolumetricCloudGlobalParameters, "RenderVolumetricCloudParameters", SceneTextures);
@@ -658,6 +664,8 @@ void SetupDefaultRenderVolumetricCloudGlobalParameters(FRDGBuilder& GraphBuilder
 
 	VolumetricCloudParams.EmptySpaceSkippingSliceDepth = GetEmptySpaceSkippingSliceDepth();
 
+	VolumetricCloudParams.BlueNoise = GetBlueNoiseParameters();
+
 	VolumetricCloudParams.EnableHeightFog = ViewInfo.Family->Scene->HasAnyExponentialHeightFog() && ShouldRenderFog(*ViewInfo.Family);
 	SetupFogUniformParameters(GraphBuilder, ViewInfo, VolumetricCloudParams.FogStruct);
 
diff --git a/Engine/Source/Runtime/Renderer/Private/VolumetricCloudRendering.h b/Engine/Source/Runtime/Renderer/Private/VolumetricCloudRendering.h
--- a/Engine/Source/Runtime/Renderer/Private/VolumetricCloudRendering.h
+++ b/Engine/Source/Runtime/Renderer/Private/VolumetricCloudRendering.h
@@ -105,6 +105,7 @@ private:
 
 
 bool ShouldRenderVolumetricCloud(const FScene* Scene, const FEngineShowFlags& EngineShowFlags);
+bool ShouldRenderVolumetricCloudWithBlueNoise(const FScene* Scene, const FEngineShowFlags& EngineShowFlags);
 bool ShouldViewVisualizeVolumetricCloudConservativeDensity(const FViewInfo& ViewInfo, const FEngineShowFlags& EngineShowFlags);
 bool VolumetricCloudWantsToSampleLocalLights(const FScene* Scene, const FEngineShowFlags& EngineShowFlags);
 uint32 GetVolumetricCloudDebugViewMode(const FEngineShowFlags& ShowFlags);
