Fix compilation error due to var missing when precompiler flag is not set.

diff --git a/Engine/Source/Runtime/Renderer/Private/VirtualShadowMaps/VirtualShadowMapArray.cpp b/Engine/Source/Runtime/Renderer/Private/VirtualShadowMaps/VirtualShadowMapArray.cpp
--- a/Engine/Source/Runtime/Renderer/Private/VirtualShadowMaps/VirtualShadowMapArray.cpp
+++ b/Engine/Source/Runtime/Renderer/Private/VirtualShadowMaps/VirtualShadowMapArray.cpp
@@ -1344,14 +1344,14 @@ void FVirtualShadowMapArray::BuildPageAllocations(
 	UniformParameters.CoarsePagePixelThresholdStatic = CVarCoarsePagePixelThresholdStatic.GetValueOnRenderThread();
 	UniformParameters.CoarsePagePixelThresholdDynamicNanite = CVarCoarsePagePixelThresholdDynamicNanite.GetValueOnRenderThread();
 
+	bool bCsvLogEnabled = false;
 #if !UE_BUILD_SHIPPING
 	const bool bRunPageAreaDiagnostics = CVarNumPageAreaDiagSlots.GetValueOnRenderThread() != 0;
 #if CSV_PROFILER
-	bool bCsvLogEnabled = FCsvProfiler::Get()->IsCapturing_Renderthread() && FCsvProfiler::Get()->IsCategoryEnabled(CSV_CATEGORY_INDEX(VSM));
+	bCsvLogEnabled = FCsvProfiler::Get()->IsCapturing_Renderthread() && FCsvProfiler::Get()->IsCategoryEnabled(CSV_CATEGORY_INDEX(VSM));
 #endif
 #else
 	constexpr bool bRunPageAreaDiagnostics = false;
-	constexpr bool bCsvLogEnabled = false;
 #endif
 
 	if (CVarShowStats.GetValueOnRenderThread() || CacheManager->IsAccumulatingStats() || bRunPageAreaDiagnostics || bCsvLogEnabled)
