PR #9727: FIX: Custom culling query was being provided wrong VisibilityId value (Contributed by [REDACTED_PERSON])

diff --git a/Engine/Source/Runtime/Renderer/Private/SceneVisibility.cpp b/Engine/Source/Runtime/Renderer/Private/SceneVisibility.cpp
--- a/Engine/Source/Runtime/Renderer/Private/SceneVisibility.cpp
+++ b/Engine/Source/Runtime/Renderer/Private/SceneVisibility.cpp
@@ -795,7 +795,7 @@ static void PrimitiveCullTask(FThreadSafeCounter& NumCulledPrimitives, const FSc
 						if (Flags.bUseCustomCulling &&
 							((Scene->PrimitiveOcclusionFlags[Index] & CustomVisibilityFlags) == CustomVisibilityFlags))
 						{
-							VisibilityId = Scene->PrimitiveVisibilityIds[Index].ByteIndex;
+							VisibilityId = Scene->Primitives[Index]->Proxy->GetVisibilityId();
 						}
 
 						bIsVisible = !bPartiallyOutside || IsPrimitiveVisible(View, PermutedPlanePtr, Bounds, VisibilityId, Flags);
