Don't include emissive and translucent materials when calculating illuminance for auto exposure.
- Extract emissive, diffuse and subsurface color right after base pass before lighting is calculated.
- Calculate illuminance before translucent pass.

diff --git a/Engine/Shaders/Private/Histogram.usf b/Engine/Shaders/Private/Histogram.usf
--- a/Engine/Shaders/Private/Histogram.usf
+++ b/Engine/Shaders/Private/Histogram.usf
@@ -316,12 +316,15 @@ void MainAtomicCS(
 			const float2 InvViewSize = Input_ViewportSizeInverse.xy;
 			const float2 ScreenUV = (TexelPos.xy - Input_ViewportMin + 0.5f) * InvViewSize;
 
-			float4 SceneColor = InputTexture.Load(int3(TexelPos, 0));
+			float3 SceneColor = InputTexture.Load(int3(TexelPos, 0)).rgb;
 			SceneColor *= View.OneOverPreExposure;
 
-			float LuminanceVal = CalculateEyeAdaptationLuminance(SceneColor.rgb);
+			float LuminanceVal = CalculateEyeAdaptationLuminance(SceneColor);
 
 #if USE_APPROX_ILLUMINANCE
+#	if USE_PRECALCULATED_ILLUMINANCE
+			LuminanceVal = SceneColor.r;
+#	else
 			const uint2 GBufferTexelPos = View.ViewRectMinAndSize.xy + ScreenUV * View.ViewRectMinAndSize.zw;
 			const float2 GBufferUV = (GBufferTexelPos + 0.5f) * View.BufferSizeAndInvSize.zw;
 
@@ -357,6 +360,7 @@ void MainAtomicCS(
 				LuminanceVal /= max(BaseColorLuminanceVal, EyeAdaptation_IgnoreMaterialsMinBaseColorLuminance);
 				LuminanceVal *= lerp(EyeAdaptation_IgnoreMaterialsLuminanceScale, 1.0f, EyeAdaptation_IgnoreMaterialsMinBaseColorLuminance);
 			}
+#	endif
 #endif
 
 #if USE_APPROX_ILLUMINANCE_DEBUG
diff --git a/Engine/Shaders/Private/PostProcessEyeAdaptation.usf b/Engine/Shaders/Private/PostProcessEyeAdaptation.usf
--- a/Engine/Shaders/Private/PostProcessEyeAdaptation.usf
+++ b/Engine/Shaders/Private/PostProcessEyeAdaptation.usf
@@ -8,6 +8,10 @@
 #include "PostProcessCommon.ush"
 #include "PostProcessHistogramCommon.ush"
 
+#include "DeferredShadingCommon.ush"
+#include "BRDF.ush"
+#include "PositionReconstructionCommon.ush"
+
 Texture2D EyeAdaptationTexture;
 
 Texture2D HistogramTexture;
@@ -261,4 +265,88 @@ void BasicEyeAdaptationPS(in float4 UVAndScreenPos : TEXCOORD0, out float4 OutCo
 		RWEyeAdaptationTexture[uint2(0,0)] = OutColor;
 	}
 #endif
-}
\ No newline at end of file
+}
+
+#if COMPUTESHADER
+
+RWTexture2D<float> RWIlluminanceTexture;
+
+// same weights as CalculateEyeAdaptationLuminance
+// but skip the min since we apply EyeAdaptation_IgnoreMaterialsMinBaseColorLuminance here
+float CalculateEyeAdaptationLuminanceWithoutMin(float3 Color)
+{
+	return dot(Color, float3(1.0f, 1.0f, 1.0f) / 3.0f);
+}
+
+[numthreads(8, 8, 1)]
+void SetupExposureIlluminance(uint2 ThreadId : SV_DispatchThreadID)
+{
+	const uint2 TexelPos = ThreadId + View.ViewRectMinAndSize.xy;
+
+	float3 Denominator = 0.0f;
+
+	const float2 GBufferUV = (TexelPos + 0.5f) * View.BufferSizeAndInvSize.zw;
+	const FGBufferData GBufferData = GetGBufferDataFromSceneTextures(GBufferUV);
+
+	if (GBufferData.ShadingModelID != SHADINGMODELID_UNLIT)
+	{
+		Denominator += GBufferData.DiffuseColor;
+
+		// Emissive
+		Denominator += ColorTexture.Load(int3(TexelPos, 0)).rgb * View.OneOverPreExposure;
+
+		if (GBufferData.ShadingModelID == SHADINGMODELID_SUBSURFACE
+			|| GBufferData.ShadingModelID == SHADINGMODELID_PREINTEGRATED_SKIN
+			|| GBufferData.ShadingModelID == SHADINGMODELID_TWOSIDED_FOLIAGE
+			|| GBufferData.ShadingModelID == SHADINGMODELID_CLOTH
+			|| GBufferData.ShadingModelID == SHADINGMODELID_EYE)
+		{
+			float3 SubsurfaceColor = ExtractSubsurfaceColor(GBufferData);
+
+			if (GBufferData.ShadingModelID == SHADINGMODELID_TWOSIDED_FOLIAGE)
+			{
+				SubsurfaceColor *= 1.0f / PI;
+			}
+
+			Denominator += SubsurfaceColor;
+		}
+	}
+
+	RWIlluminanceTexture[TexelPos] = CalculateEyeAdaptationLuminanceWithoutMin(Denominator) * View.PreExposure;
+}
+
+[numthreads(8, 8, 1)]
+void CalculateExposureIlluminance(uint2 ThreadId : SV_DispatchThreadID)
+{
+	const uint2 TexelPos = ThreadId + View.ViewRectMinAndSize.xy;
+
+	const float3 SceneColor = ColorTexture.Load(int3(TexelPos, 0)).rgb * View.OneOverPreExposure;
+
+	float Illuminance = CalculateEyeAdaptationLuminance(SceneColor);
+
+	const float2 GBufferUV = (TexelPos + 0.5f) * View.BufferSizeAndInvSize.zw;
+	const FGBufferData GBufferData = GetGBufferDataFromSceneTextures(GBufferUV);
+
+	if (GBufferData.ShadingModelID != SHADINGMODELID_UNLIT)
+	{
+		float Denominator = RWIlluminanceTexture[TexelPos] * View.OneOverPreExposure;
+
+		{
+			const float3 TranslatedWorldPosition = ReconstructTranslatedWorldPositionFromDepth(GBufferUV, GBufferData.Depth);
+			const float3 CameraVector = normalize(View.TranslatedWorldCameraOrigin - TranslatedWorldPosition);
+
+			const float3 N = GBufferData.WorldNormal;
+			const float3 V = CameraVector;
+			const float3 EnvBrdf = EnvBRDF(GBufferData.SpecularColor, GBufferData.Roughness, max(0.0, dot(N, V)));
+
+			Denominator += CalculateEyeAdaptationLuminanceWithoutMin(EnvBrdf);
+		}
+
+		Illuminance /= max(Denominator, EyeAdaptation_IgnoreMaterialsMinBaseColorLuminance);
+		Illuminance *= lerp(EyeAdaptation_IgnoreMaterialsLuminanceScale, 1.0f, EyeAdaptation_IgnoreMaterialsMinBaseColorLuminance);
+	}
+	
+	RWIlluminanceTexture[TexelPos] = Illuminance * View.PreExposure;
+}
+
+#endif // COMPUTESHADER
diff --git a/Engine/Source/Runtime/Renderer/Private/DeferredShadingRenderer.cpp b/Engine/Source/Runtime/Renderer/Private/DeferredShadingRenderer.cpp
--- a/Engine/Source/Runtime/Renderer/Private/DeferredShadingRenderer.cpp
+++ b/Engine/Source/Runtime/Renderer/Private/DeferredShadingRenderer.cpp
@@ -26,6 +26,7 @@
 #include "DistanceFieldAmbientOcclusion.h"
 #include "GlobalDistanceField.h"
 #include "PostProcess/PostProcessing.h"
+#include "PostProcess/PostProcessEyeAdaptation.h"
 #include "DistanceFieldAtlas.h"
 #include "EngineModule.h"
 #include "SceneViewExtension.h"
@@ -3230,6 +3231,9 @@ void FDeferredShadingSceneRenderer::Render(FRDGBuilder& GraphBuilder)
 		// VisualizeVirtualShadowMap TODO
 	}
 
+	// Extract emissive from SceneColor (before lighting is applied) + material diffuse and subsurface colors
+	FRDGTextureRef ExposureIlluminanceSetup = AddSetupExposureIlluminancePass(GraphBuilder, Views, SceneTextures);
+
 	if (ViewFamily.EngineShowFlags.VisualizeLightCulling)
 	{
 		FRDGTextureRef VisualizeLightCullingTexture = GraphBuilder.CreateTexture(SceneTextures.Color.Target->Desc, TEXT("SceneColorVisualizeLightCulling"));
@@ -3649,6 +3653,8 @@ void FDeferredShadingSceneRenderer::Render(FRDGBuilder& GraphBuilder)
 		ComposeVolumetricRenderTargetOverScene(GraphBuilder, Views, SceneTextures.Color.Target, SceneTextures.Depth.Target, bShouldRenderSingleLayerWater, SceneWithoutWaterTextures, SceneTextures);
 	}
 
+	FRDGTextureRef ExposureIlluminance = AddCalculateExposureIlluminancePass(GraphBuilder, Views, SceneTextures, ExposureIlluminanceSetup);
+
 	RenderOpaqueFX(GraphBuilder, Views, FXSystem, SceneTextures.UniformBuffer);
 
 	FRendererModule& RendererModule = static_cast<FRendererModule&>(GetRendererModule());
@@ -3917,6 +3923,7 @@ void FDeferredShadingSceneRenderer::Render(FRDGBuilder& GraphBuilder)
 		FPostProcessingInputs PostProcessingInputs;
 		PostProcessingInputs.ViewFamilyTexture = ViewFamilyTexture;
 		PostProcessingInputs.CustomDepthTexture = SceneTextures.CustomDepth.Depth;
+		PostProcessingInputs.ExposureIlluminance = ExposureIlluminance;
 		PostProcessingInputs.SceneTextures = SceneTextures.UniformBuffer;
 
 		GraphBuilder.FlushSetupQueue();
diff --git a/Engine/Source/Runtime/Renderer/Private/PostProcess/PostProcessEyeAdaptation.cpp b/Engine/Source/Runtime/Renderer/Private/PostProcess/PostProcessEyeAdaptation.cpp
--- a/Engine/Source/Runtime/Renderer/Private/PostProcess/PostProcessEyeAdaptation.cpp
+++ b/Engine/Source/Runtime/Renderer/Private/PostProcess/PostProcessEyeAdaptation.cpp
@@ -1,6 +1,10 @@
 // Copyright Epic Games, Inc. All Rights Reserved.
 
 #include "PostProcess/PostProcessEyeAdaptation.h"
+
+#include "SceneTextureParameters.h"
+#include "SystemTextures.h"
+
 #include "RHIGPUReadback.h"
 #include "RendererUtils.h"
 #include "ScenePrivate.h"
@@ -86,6 +90,12 @@ namespace
 		TEXT(""),
 		ECVF_RenderThreadSafe);
 
+	TAutoConsoleVariable<bool> CVarAutoExposureIgnoreMaterialsUsePrecalculatedIlluminance(
+		TEXT("r.AutoExposure.IgnoreMaterials.UsePrecalculatedIlluminance"),
+		true,
+		TEXT(""),
+		ECVF_Scalability | ECVF_RenderThreadSafe);
+
 	TAutoConsoleVariable<bool> CVarAutoExposureIgnoreMaterialsDebug(
 		TEXT("r.AutoExposure.IgnoreMaterials.Debug"),
 		false,
@@ -119,6 +129,11 @@ bool IsExtendLuminanceRangeEnabled()
 	return VarDefaultAutoExposureExtendDefaultLuminanceRange->GetValueOnRenderThread() == 1;
 }
 
+bool IsAutoExposureUsingIlluminanceEnabled(const FViewInfo& View)
+{
+	return CVarAutoExposureIgnoreMaterials.GetValueOnRenderThread() && GetAutoExposureMethod(View) == AEM_Histogram;
+}
+
 static EAutoExposureMethod ApplyEyeAdaptationQuality(EAutoExposureMethod AutoExposureMethod)
 {
 	static const auto CVarEyeAdaptationQuality = IConsoleManager::Get().FindTConsoleVariableDataInt(TEXT("r.EyeAdaptationQuality"));
@@ -488,6 +503,148 @@ float GetEyeAdaptationFixedExposure(const FViewInfo& View)
 	return ExposureScale * Parameters.ExposureCompensationSettings;
 }
 
+//////////////////////////////////////////////////////////////////////////
+//! Setup
+//////////////////////////////////////////////////////////////////////////
+
+class FSetupExposureIlluminanceCS : public FGlobalShader
+{
+	DECLARE_GLOBAL_SHADER(FSetupExposureIlluminanceCS);
+	SHADER_USE_PARAMETER_STRUCT(FSetupExposureIlluminanceCS, FGlobalShader);
+
+	BEGIN_SHADER_PARAMETER_STRUCT(FParameters, )
+		SHADER_PARAMETER_STRUCT_REF(FViewUniformShaderParameters, View)
+		SHADER_PARAMETER_STRUCT_INCLUDE(FSceneTextureParameters, SceneTextures)
+		SHADER_PARAMETER_RDG_TEXTURE(Texture2D, ColorTexture)
+		SHADER_PARAMETER_RDG_TEXTURE_UAV(RWTexture2D<float4>, RWIlluminanceTexture)
+	END_SHADER_PARAMETER_STRUCT()
+
+	static bool ShouldCompilePermutation(const FGlobalShaderPermutationParameters& Parameters)
+	{
+		return IsFeatureLevelSupported(Parameters.Platform, ERHIFeatureLevel::SM5);
+	}
+};
+
+IMPLEMENT_GLOBAL_SHADER(FSetupExposureIlluminanceCS, "/Engine/Private/PostProcessEyeAdaptation.usf", "SetupExposureIlluminance", SF_Compute);
+
+FRDGTextureRef AddSetupExposureIlluminancePass(
+	FRDGBuilder& GraphBuilder,
+	TArrayView<const FViewInfo> Views,
+	const FSceneTextures& SceneTextures)
+{
+	const uint32 ViewCount = Views.Num();
+
+	bool bAnyViewEnabled = false;
+
+	for (uint32 ViewIndex = 0; ViewIndex < ViewCount; ++ViewIndex)
+	{
+		bAnyViewEnabled |= IsAutoExposureUsingIlluminanceEnabled(Views[ViewIndex]) && CVarAutoExposureIgnoreMaterialsUsePrecalculatedIlluminance.GetValueOnRenderThread();
+	}
+
+	if (!bAnyViewEnabled)
+	{
+		return nullptr;
+	}
+
+	const FRDGTextureDesc Desc = FRDGTextureDesc::Create2D(SceneTextures.Config.Extent, PF_R16F, FClearValueBinding::Black, TexCreate_UAV | TexCreate_ShaderResource);
+	FRDGTextureRef OutputTexture = GraphBuilder.CreateTexture(Desc, TEXT("EyeAdaptation_Illuminance"));
+
+	for (uint32 ViewIndex = 0; ViewIndex < ViewCount; ++ViewIndex)
+	{
+		const FViewInfo& View = Views[ViewIndex];
+
+		if (IsAutoExposureUsingIlluminanceEnabled(View) && CVarAutoExposureIgnoreMaterialsUsePrecalculatedIlluminance.GetValueOnRenderThread())
+		{
+			auto* PassParameters = GraphBuilder.AllocParameters<FSetupExposureIlluminanceCS::FParameters>();
+			PassParameters->View = View.ViewUniformBuffer;
+			PassParameters->SceneTextures = GetSceneTextureParameters(GraphBuilder, SceneTextures);
+			PassParameters->ColorTexture = SceneTextures.Color.Resolve;
+			PassParameters->RWIlluminanceTexture = GraphBuilder.CreateUAV(OutputTexture);
+
+			auto ComputeShader = View.ShaderMap->GetShader<FSetupExposureIlluminanceCS>();
+
+			FComputeShaderUtils::AddPass(
+				GraphBuilder,
+				RDG_EVENT_NAME("SetupExposureIlluminance (ViewId=%d)", ViewIndex),
+				ComputeShader,
+				PassParameters,
+				FComputeShaderUtils::GetGroupCount(View.ViewRect.Size(), FIntPoint(8, 8)));
+		}
+	}
+
+	return OutputTexture;
+}
+
+class FCalculateExposureIlluminanceCS : public FGlobalShader
+{
+	DECLARE_GLOBAL_SHADER(FCalculateExposureIlluminanceCS);
+	SHADER_USE_PARAMETER_STRUCT(FCalculateExposureIlluminanceCS, FGlobalShader);
+
+	BEGIN_SHADER_PARAMETER_STRUCT(FParameters, )
+		SHADER_PARAMETER_STRUCT_REF(FViewUniformShaderParameters, View)
+		SHADER_PARAMETER_STRUCT_INCLUDE(FSceneTextureParameters, SceneTextures)
+		SHADER_PARAMETER_RDG_TEXTURE(Texture2D, ColorTexture)
+		SHADER_PARAMETER_RDG_TEXTURE_UAV(RWTexture2D<float4>, RWIlluminanceTexture)
+
+		SHADER_PARAMETER_STRUCT(FEyeAdaptationParameters, EyeAdaptation)
+
+		SHADER_PARAMETER_TEXTURE(Texture2D, PreIntegratedGF)
+		SHADER_PARAMETER_SAMPLER(SamplerState, PreIntegratedGFSampler)
+	END_SHADER_PARAMETER_STRUCT()
+
+	static bool ShouldCompilePermutation(const FGlobalShaderPermutationParameters& Parameters)
+	{
+		return IsFeatureLevelSupported(Parameters.Platform, ERHIFeatureLevel::SM5);
+	}
+};
+
+IMPLEMENT_GLOBAL_SHADER(FCalculateExposureIlluminanceCS, "/Engine/Private/PostProcessEyeAdaptation.usf", "CalculateExposureIlluminance", SF_Compute);
+
+FRDGTextureRef AddCalculateExposureIlluminancePass(
+	FRDGBuilder& GraphBuilder,
+	TArrayView<const FViewInfo> Views,
+	const FSceneTextures& SceneTextures,
+	FRDGTextureRef ExposureIlluminanceSetup)
+{
+	if (!ExposureIlluminanceSetup)
+	{
+		return nullptr;
+	}
+
+	const uint32 ViewCount = Views.Num();
+
+	for (uint32 ViewIndex = 0; ViewIndex < ViewCount; ++ViewIndex)
+	{
+		const FViewInfo& View = Views[ViewIndex];
+
+		if (IsAutoExposureUsingIlluminanceEnabled(View))
+		{
+			check(CVarAutoExposureIgnoreMaterialsUsePrecalculatedIlluminance.GetValueOnRenderThread());
+
+			auto* PassParameters = GraphBuilder.AllocParameters<FCalculateExposureIlluminanceCS::FParameters>();
+			PassParameters->View = View.ViewUniformBuffer;
+			PassParameters->SceneTextures = GetSceneTextureParameters(GraphBuilder, SceneTextures);
+			PassParameters->ColorTexture = SceneTextures.Color.Resolve;
+			PassParameters->RWIlluminanceTexture = GraphBuilder.CreateUAV(ExposureIlluminanceSetup);
+			PassParameters->EyeAdaptation = GetEyeAdaptationParameters(View, ERHIFeatureLevel::SM5);
+
+			PassParameters->PreIntegratedGF = GSystemTextures.PreintegratedGF->GetRHI();
+			PassParameters->PreIntegratedGFSampler = TStaticSamplerState<SF_Bilinear, AM_Clamp, AM_Clamp, AM_Clamp>::GetRHI();
+
+			auto ComputeShader = View.ShaderMap->GetShader<FCalculateExposureIlluminanceCS>();
+
+			FComputeShaderUtils::AddPass(
+				GraphBuilder,
+				RDG_EVENT_NAME("CalculateExposureIlluminance (ViewId=%d)", ViewIndex),
+				ComputeShader,
+				PassParameters,
+				FComputeShaderUtils::GetGroupCount(View.ViewRect.Size(), FIntPoint(8, 8)));
+		}
+	}
+
+	return ExposureIlluminanceSetup;
+}
+
 //////////////////////////////////////////////////////////////////////////
 //! Histogram Eye Adaptation
 //////////////////////////////////////////////////////////////////////////
diff --git a/Engine/Source/Runtime/Renderer/Private/PostProcess/PostProcessEyeAdaptation.h b/Engine/Source/Runtime/Renderer/Private/PostProcess/PostProcessEyeAdaptation.h
--- a/Engine/Source/Runtime/Renderer/Private/PostProcess/PostProcessEyeAdaptation.h
+++ b/Engine/Source/Runtime/Renderer/Private/PostProcess/PostProcessEyeAdaptation.h
@@ -49,6 +49,8 @@ float GetAutoExposureCompensationFromSettings(const FViewInfo& View);
 
 bool IsExtendLuminanceRangeEnabled();
 
+bool IsAutoExposureUsingIlluminanceEnabled(const FViewInfo& View);
+
 // Returns the auto exposure method enabled by the view (including CVar override).
 EAutoExposureMethod GetAutoExposureMethod(const FViewInfo& View);
 
@@ -89,6 +91,17 @@ FEyeAdaptationParameters GetEyeAdaptationParameters(const FViewInfo& ViewInfo, E
 // Computes the a fixed exposure to be used to replace the dynamic exposure when it's not supported (< SM5).
 float GetEyeAdaptationFixedExposure(const FViewInfo& View);
 
+FRDGTextureRef AddSetupExposureIlluminancePass(
+	FRDGBuilder& GraphBuilder,
+	TArrayView<const FViewInfo> Views,
+	const FSceneTextures& SceneTextures);
+
+FRDGTextureRef AddCalculateExposureIlluminancePass(
+	FRDGBuilder& GraphBuilder,
+	TArrayView<const FViewInfo> Views,
+	const FSceneTextures& SceneTextures,
+	FRDGTextureRef ExposureIlluminanceSetup);
+
 // Returns the updated 1x1 eye adaptation texture.
 FRDGTextureRef AddHistogramEyeAdaptationPass(
 	FRDGBuilder& GraphBuilder,
diff --git a/Engine/Source/Runtime/Renderer/Private/PostProcess/PostProcessHistogram.cpp b/Engine/Source/Runtime/Renderer/Private/PostProcess/PostProcessHistogram.cpp
--- a/Engine/Source/Runtime/Renderer/Private/PostProcess/PostProcessHistogram.cpp
+++ b/Engine/Source/Runtime/Renderer/Private/PostProcess/PostProcessHistogram.cpp
@@ -149,8 +149,9 @@ public:
 	END_SHADER_PARAMETER_STRUCT()
 
 	class FUseApproxIlluminanceDim : SHADER_PERMUTATION_BOOL("USE_APPROX_ILLUMINANCE");
+	class FUsePrecalculatedIlluminanceDim : SHADER_PERMUTATION_BOOL("USE_PRECALCULATED_ILLUMINANCE");
 	class FUseApproxIlluminanceDebugDim : SHADER_PERMUTATION_BOOL("USE_APPROX_ILLUMINANCE_DEBUG");
-	using FPermutationDomain = TShaderPermutationDomain<FUseApproxIlluminanceDim, FUseApproxIlluminanceDebugDim>;
+	using FPermutationDomain = TShaderPermutationDomain<FUseApproxIlluminanceDim, FUsePrecalculatedIlluminanceDim, FUseApproxIlluminanceDebugDim>;
 
 	static bool ShouldCompilePermutation(const FGlobalShaderPermutationParameters& Parameters)
 	{
@@ -363,9 +364,11 @@ static FRDGTextureRef AddHistogramAtomicPass(
 
 		static const auto CVarIgnoreMaterials = IConsoleManager::Get().FindTConsoleVariableDataBool(TEXT("r.AutoExposure.IgnoreMaterials"));
 		static const auto CVarIgnoreMaterialsDebug = IConsoleManager::Get().FindTConsoleVariableDataBool(TEXT("r.AutoExposure.IgnoreMaterials.Debug"));
+		static const auto CVarUsePrecalculatedIlluminance = IConsoleManager::Get().FindTConsoleVariableDataBool(TEXT("r.AutoExposure.IgnoreMaterials.UsePrecalculatedIlluminance"));
 
 		FHistogramAtomicCS::FPermutationDomain PermutationVector;
 		PermutationVector.Set<FHistogramAtomicCS::FUseApproxIlluminanceDim>(CVarIgnoreMaterials->GetValueOnRenderThread());
+		PermutationVector.Set<FHistogramAtomicCS::FUsePrecalculatedIlluminanceDim>(CVarUsePrecalculatedIlluminance->GetValueOnRenderThread());
 		PermutationVector.Set<FHistogramAtomicCS::FUseApproxIlluminanceDebugDim>(CVarIgnoreMaterialsDebug->GetValueOnRenderThread());
 
 		auto ComputeShader = View.ShaderMap->GetShader<FHistogramAtomicCS>(PermutationVector);
diff --git a/Engine/Source/Runtime/Renderer/Private/PostProcess/PostProcessing.cpp b/Engine/Source/Runtime/Renderer/Private/PostProcess/PostProcessing.cpp
--- a/Engine/Source/Runtime/Renderer/Private/PostProcess/PostProcessing.cpp
+++ b/Engine/Source/Runtime/Renderer/Private/PostProcess/PostProcessing.cpp
@@ -805,11 +805,16 @@ void AddPostProcessingPasses(
 		{
 			FScreenPassTexture HistogramSceneColor = bProcessQuarterResolution ? QuarterResSceneColor : HalfResSceneColor;
 
-			static const auto CVarIgnoreMaterials = IConsoleManager::Get().FindTConsoleVariableDataBool(TEXT("r.AutoExposure.IgnoreMaterials"));
-
-			if (CVarIgnoreMaterials->GetValueOnRenderThread())
+			if (IsAutoExposureUsingIlluminanceEnabled(View))
 			{
-				HistogramSceneColor = OriginalSceneColor;
+				if (Inputs.ExposureIlluminance)
+				{
+					HistogramSceneColor = FScreenPassTexture(Inputs.ExposureIlluminance, PrimaryViewRect);
+				}
+				else
+				{
+					HistogramSceneColor = OriginalSceneColor;
+				}
 			}
 
 			HistogramTexture = AddHistogramPass(
diff --git a/Engine/Source/Runtime/Renderer/Private/PostProcess/PostProcessing.h b/Engine/Source/Runtime/Renderer/Private/PostProcess/PostProcessing.h
--- a/Engine/Source/Runtime/Renderer/Private/PostProcess/PostProcessing.h
+++ b/Engine/Source/Runtime/Renderer/Private/PostProcess/PostProcessing.h
@@ -32,6 +32,7 @@ struct FPostProcessingInputs
 	TRDGUniformBufferRef<FSceneTextureUniformParameters> SceneTextures = nullptr;
 	FRDGTextureRef ViewFamilyTexture = nullptr;
 	FRDGTextureRef CustomDepthTexture = nullptr;
+	FRDGTextureRef ExposureIlluminance = nullptr;
 	FTranslucencyViewResourcesMap TranslucencyViewResourcesMap;
 
 	void Validate() const
