Fix multiple uploads to GPUScene in one graph

diff --git a/Engine/Source/Runtime/Renderer/Private/GPUScene.cpp b/Engine/Source/Runtime/Renderer/Private/GPUScene.cpp
--- a/Engine/Source/Runtime/Renderer/Private/GPUScene.cpp
+++ b/Engine/Source/Runtime/Renderer/Private/GPUScene.cpp
@@ -788,7 +788,7 @@ void FGPUScene::UpdateBufferState(FRDGBuilder& GraphBuilder, FScene& Scene, cons
 	BufferState.LightMapDataBufferSize = LightMapDataBufferSize;
 	
 	const uint32 LightDataBufferSize = FMath::RoundUpToPowerOfTwo(FMath::Max(Scene.Lights.Num(), InitialBufferSize));
-	BufferState.LightDataBuffer = ResizeStructuredBufferIfNeeded(GraphBuilder, LightDataBuffer, LightDataBufferSize * sizeof(FLightSceneData), TEXT("GPUScene.LightData"));
+	BufferState.LightDataBuffer = GraphBuilder.CreateBuffer(FRDGBufferDesc::CreateStructuredDesc(sizeof(FLightSceneData), Scene.Lights.Num()), TEXT("GPUScene.LightData"));
 
 	ShaderParameters.GPUSceneInstanceSceneData = GraphBuilder.CreateSRV(BufferState.InstanceSceneDataBuffer);
 	ShaderParameters.GPUSceneInstancePayloadData = GraphBuilder.CreateSRV(BufferState.InstancePayloadDataBuffer);
diff --git a/Engine/Source/Runtime/Renderer/Private/GPUScene.h b/Engine/Source/Runtime/Renderer/Private/GPUScene.h
--- a/Engine/Source/Runtime/Renderer/Private/GPUScene.h
+++ b/Engine/Source/Runtime/Renderer/Private/GPUScene.h
@@ -348,9 +348,6 @@ public:
 	TRefCountPtr<FRDGPooledBuffer> LightmapDataBuffer;
 	FRDGAsyncScatterUploadBuffer        LightmapUploadBuffer;
 
-	/** View-independent GPU light data */
-	TRefCountPtr<FRDGPooledBuffer> LightDataBuffer;
-
 	struct FInstanceRange
 	{
 		uint32 InstanceSceneDataOffset;
