Fix 0-length buffer assert when there are no lights in the scene

diff --git a/Engine/Source/Runtime/Renderer/Private/GPUScene.cpp b/Engine/Source/Runtime/Renderer/Private/GPUScene.cpp
--- a/Engine/Source/Runtime/Renderer/Private/GPUScene.cpp
+++ b/Engine/Source/Runtime/Renderer/Private/GPUScene.cpp
@@ -788,7 +788,7 @@ void FGPUScene::UpdateBufferState(FRDGBuilder& GraphBuilder, FScene& Scene, cons
 	BufferState.LightMapDataBufferSize = LightMapDataBufferSize;
 	
 	const uint32 LightDataBufferSize = FMath::RoundUpToPowerOfTwo(FMath::Max(Scene.Lights.Num(), InitialBufferSize));
-	BufferState.LightDataBuffer = GraphBuilder.CreateBuffer(FRDGBufferDesc::CreateStructuredDesc(sizeof(FLightSceneData), Scene.Lights.Num()), TEXT("GPUScene.LightData"));
+	BufferState.LightDataBuffer = GraphBuilder.CreateBuffer(FRDGBufferDesc::CreateStructuredDesc(sizeof(FLightSceneData), FMath::Max(1, Scene.Lights.Num())), TEXT("GPUScene.LightData"));
 
 	ShaderParameters.GPUSceneInstanceSceneData = GraphBuilder.CreateSRV(BufferState.InstanceSceneDataBuffer);
 	ShaderParameters.GPUSceneInstancePayloadData = GraphBuilder.CreateSRV(BufferState.InstancePayloadDataBuffer);
