Disable fast clear on the DistortedSceneColor since we dont clear the
RTV. Disable fast clear on DistortedSceneColor to save on the Fast eliminate when the RTV is aliased with the transient allocator. The RTV is used in the Apply and Merge passes, but uses stencil to only write and read on pixels with distortion.

diff --git a/Engine/Source/Runtime/Renderer/Private/DistortionRendering.cpp b/Engine/Source/Runtime/Renderer/Private/DistortionRendering.cpp
--- a/Engine/Source/Runtime/Renderer/Private/DistortionRendering.cpp
+++ b/Engine/Source/Runtime/Renderer/Private/DistortionRendering.cpp
@@ -719,6 +719,9 @@ void FDeferredShadingSceneRenderer::RenderDistortion(
 	}
 
 	FRDGTextureDesc DistortedSceneColorDesc = SceneColorTexture->Desc;
+	//Remove fast clear flag on the DistoredSceneColor which is used in the Apply and Merge passes. 
+	// This can save the Fast clear eliminate in the Merge pass when the RTV is transient allocated.
+	EnumAddFlags(DistortedSceneColorDesc.Flags, TexCreate_NoFastClear);
 	EnumRemoveFlags(DistortedSceneColorDesc.Flags, TexCreate_FastVRAM);
 
 	FRDGTextureRef DistortionSceneColorTexture = GraphBuilder.CreateTexture(DistortedSceneColorDesc, TEXT("DistortedSceneColor"));
