Manual backout of changelists 22943247 AND 22944158 to fix Vulkan rendering.
#tests preflight Android Test client

diff --git a/Engine/Shaders/Private/MobileBasePassPixelShader.usf b/Engine/Shaders/Private/MobileBasePassPixelShader.usf
--- a/Engine/Shaders/Private/MobileBasePassPixelShader.usf
+++ b/Engine/Shaders/Private/MobileBasePassPixelShader.usf
@@ -130,13 +130,16 @@ void ApplyPixelDepthOffsetForMobileBasePass(inout FMaterialPixelParameters Mater
 	#define PIXELSHADER_EARLYDEPTHSTENCIL 	
 #endif
 
-void GetSkyLighting(half3 SkyLightColor, half3 WorldNormal, bool bEvaluateBackface, inout half IndirectIrradiance, out half3 OutSkyDiffuseLighting, out half3 OutSkySubsurfaceLighting, out half3 OutWaterDiffuseIndirectLuminance)
+void GetSkyLighting(bool bApplySkyLighting, half3 SkyLightColor, half3 WorldNormal, bool bEvaluateBackface, inout half IndirectIrradiance, out half3 OutSkyDiffuseLighting, out half3 OutSkySubsurfaceLighting, out half3 OutWaterDiffuseIndirectLuminance)
 {
 	OutSkyDiffuseLighting = 0;
 	OutSkySubsurfaceLighting = 0;
 	OutWaterDiffuseIndirectLuminance = 0;
 
 #if ENABLE_SKY_LIGHT
+	//@mw todo
+	// TODO: Also need to do specular.
+	if (bApplySkyLighting)
 	{
 		half3 SkyDiffuse = GetSkySHDiffuseSimple(WorldNormal);
 		OutSkyDiffuseLighting = SkyDiffuse * SkyLightColor;
@@ -156,7 +159,7 @@ void GetSkyLighting(half3 SkyLightColor, half3 WorldNormal, bool bEvaluateBackfa
 #endif
 }
 
-void GetPrecomputedIndirectLightingAndSkyLight(VTPageTableResult LightmapVTPageTableResult, FVertexFactoryInterpolantsVSToPS Interpolants, bool bEvaluateBackface, half3 DiffuseDir, half3 SkyLightColor, out half3 OutDiffuseIndirectLighting, out half3 OutSubsurfaceIndirectLighting, out half OutIndirectIrradiance, out half3 OutWaterDiffuseIndirectLuminance)
+void GetPrecomputedIndirectLightingAndSkyLight(VTPageTableResult LightmapVTPageTableResult, FVertexFactoryInterpolantsVSToPS Interpolants, bool bEvaluateBackface, half3 DiffuseDir, bool bApplySkyLighting, half3 SkyLightColor, out half3 OutDiffuseIndirectLighting, out half3 OutSubsurfaceIndirectLighting, out half OutIndirectIrradiance, out half3 OutWaterDiffuseIndirectLuminance)
 {
 	//To keep IndirectLightingCache conherence with PC, initialize the OutIndirectIrradiance to zero.
 	OutIndirectIrradiance = 0;
@@ -217,7 +220,7 @@ void GetPrecomputedIndirectLightingAndSkyLight(VTPageTableResult LightmapVTPageT
 
 	half3 SkyDiffuseLighting;
 	half3 SkySubsurfaceLighting;
-	GetSkyLighting(SkyLightColor, DiffuseDir, bEvaluateBackface, OutIndirectIrradiance, SkyDiffuseLighting, SkySubsurfaceLighting, OutWaterDiffuseIndirectLuminance);
+	GetSkyLighting(bApplySkyLighting, SkyLightColor, DiffuseDir, bEvaluateBackface, OutIndirectIrradiance, SkyDiffuseLighting, SkySubsurfaceLighting, OutWaterDiffuseIndirectLuminance);
 
 	OutDiffuseIndirectLighting += SkyDiffuseLighting;
 	OutSubsurfaceIndirectLighting += SkySubsurfaceLighting;
@@ -526,10 +529,11 @@ void Main(
 #endif
 
 	half NoL = max(0, dot(GBuffer.WorldNormal, MobileDirectionalLight.DirectionalLightDirectionAndShadowTransition.xyz));
+	bool bApplySkyLighting = !(MATERIAL_TWOSIDED && LQ_TEXTURE_LIGHTMAP) || NoL == 0;
 
 	const bool bEvaluateBackface = GetShadingModelRequiresBackfaceLighting(GBuffer.ShadingModelID);
 
-	GetPrecomputedIndirectLightingAndSkyLight(LightmapVTPageTableResult, Interpolants, bEvaluateBackface, DiffuseDir, ResolvedView.SkyLightColor.rgb, DiffuseIndirectLighting, SubsurfaceIndirectLighting, IndirectIrradiance, WaterDiffuseIndirectLuminance);
+	GetPrecomputedIndirectLightingAndSkyLight(LightmapVTPageTableResult, Interpolants, bEvaluateBackface, DiffuseDir, bApplySkyLighting, ResolvedView.SkyLightColor.rgb, DiffuseIndirectLighting, SubsurfaceIndirectLighting, IndirectIrradiance, WaterDiffuseIndirectLuminance);
 
 	// Apply MaterialAO since we don't have the DiffuseIndirectComposite pass on mobile deferred.
 	IndirectIrradiance *= GBuffer.GBufferAO;
diff --git a/Engine/Shaders/Private/MobileDeferredShading.usf b/Engine/Shaders/Private/MobileDeferredShading.usf
--- a/Engine/Shaders/Private/MobileDeferredShading.usf
+++ b/Engine/Shaders/Private/MobileDeferredShading.usf
@@ -51,53 +51,6 @@ half ComputeLightFunctionMultiplier(float3 TranslatedWorldPosition)
 #endif
 }
 
-half3 SkyLightDiffuseMobile(FGBufferData GBuffer, half3 DiffuseColor, half3 V, inout half IndirectIrradiance)
-{
-	half3 Lighting = 0;
-	half3 SkyDiffuseLookUpNormal = GBuffer.WorldNormal;
-
-	if (GBuffer.ShadingModelID == SHADINGMODELID_TWOSIDED_FOLIAGE)
-	{
-		half3 SubsurfaceLookup = GetSkySHDiffuseSimple(-GBuffer.WorldNormal) * View.SkyLightColor.rgb;
-		half3 SubsurfaceColor = ExtractSubsurfaceColor(GBuffer);
-		Lighting += SubsurfaceLookup * SubsurfaceColor;
-	}
-
-	if (GBuffer.ShadingModelID == SHADINGMODELID_SUBSURFACE || GBuffer.ShadingModelID == SHADINGMODELID_PREINTEGRATED_SKIN)
-	{
-		half3 SubsurfaceColor = ExtractSubsurfaceColor(GBuffer);
-		// Add subsurface energy to diffuse
-		DiffuseColor += SubsurfaceColor;
-	}
-
-	if (GBuffer.ShadingModelID == SHADINGMODELID_HAIR)
-	{
-		half3 N = GBuffer.WorldNormal;
-		half3 L = normalize(V - N * dot(V, N));
-		SkyDiffuseLookUpNormal = L;
-		FHairTransmittanceData TransmittanceData = InitHairTransmittanceData(true);
-		DiffuseColor = 2 * PI * HairShading(GBuffer, L, V, N, 1, TransmittanceData, 0, 0.2, uint2(0, 0));
-	}
-
-	if (GBuffer.ShadingModelID == SHADINGMODELID_CLOTH)
-	{
-		half3 ClothFuzz = ExtractSubsurfaceColor(GBuffer);
-		DiffuseColor += ClothFuzz * GBuffer.CustomData.a;
-	}
-
-	// Compute the preconvolved incoming lighting with the bent normal direction
-	half3 DiffuseLookup = GetSkySHDiffuseSimple(SkyDiffuseLookUpNormal) * View.SkyLightColor.rgb;
-
-	// And accumulate the lighting
-	Lighting += DiffuseLookup  * DiffuseColor;
-
-#if ALLOW_STATIC_LIGHTING
-	IndirectIrradiance += Luminance(DiffuseLookup);
-#endif
-
-	return Lighting;
-}
-
 void ReflectionEnvironmentSkyLighting(
 	FGBufferData GBuffer,
 	half3 CameraVector,
@@ -106,30 +59,12 @@ void ReflectionEnvironmentSkyLighting(
 	uint GridIndex,
 	inout FLightAccumulator DirectLighting)
 {
-	half IndirectIrradiance = GBuffer.IndirectIrradiance;
-
-#if ENABLE_DYNAMIC_SKY_LIGHT
-
-	half3 N = GBuffer.WorldNormal;
-	half3 V = -CameraVector;
-	half NoV = saturate(abs(dot(N, V)) + 1e-5);
-
-	half3 DiffuseColor = GBuffer.DiffuseColor;
-	half3 SpecularColor = GBuffer.SpecularColor;
-	RemapClearCoatDiffuseAndSpecularColor(GBuffer, NoV, DiffuseColor, SpecularColor);
-
-	half3 SkyLighting = SkyLightDiffuseMobile(GBuffer, DiffuseColor, V, IndirectIrradiance);
-
-	const bool bNeedsSeparateSubsurfaceLightAccumulation = UseSubsurfaceProfile(GBuffer.ShadingModelID);
-	LightAccumulator_Add(DirectLighting, SkyLighting, SkyLighting, 1.0f, bNeedsSeparateSubsurfaceLightAccumulation);
-#endif // ENABLE_DYNAMIC_SKY_LIGHT 
-
 	// IBL
 	AccumulateReflection(GBuffer
 		, CameraVector
 		, TranslatedWorldPosition
 		, ReflectionVector
-		, IndirectIrradiance
+		, GBuffer.IndirectIrradiance
 		, GridIndex
 		, DirectLighting);
 }
diff --git a/Engine/Source/Runtime/Renderer/Private/MobileBasePass.cpp b/Engine/Source/Runtime/Renderer/Private/MobileBasePass.cpp
--- a/Engine/Source/Runtime/Renderer/Private/MobileBasePass.cpp
+++ b/Engine/Source/Runtime/Renderer/Private/MobileBasePass.cpp
@@ -669,8 +669,7 @@ bool FMobileBasePassMeshProcessor::Process(
 	
 	if (Scene && Scene->SkyLight)
 	{
-		// Uses bTranslucentBasePass instead of BlendMode to handle single layer water meshes.
-		bEnableSkyLight = ShadingModels.IsLit() && Scene->ShouldRenderSkylightInBasePass(bTranslucentBasePass ? BLEND_Translucent : BLEND_Opaque);
+		bEnableSkyLight = ShadingModels.IsLit() && Scene->ShouldRenderSkylightInBasePass(BlendMode);
 	}
 
 	bool bEnableLocalLights = false;
diff --git a/Engine/Source/Runtime/Renderer/Private/MobileBasePassRendering.h b/Engine/Source/Runtime/Renderer/Private/MobileBasePassRendering.h
--- a/Engine/Source/Runtime/Renderer/Private/MobileBasePassRendering.h
+++ b/Engine/Source/Runtime/Renderer/Private/MobileBasePassRendering.h
@@ -340,8 +340,6 @@ public:
 	{		
 		// We compile the point light shader combinations based on the project settings
 		static auto* MobileSkyLightPermutationCVar = IConsoleManager::Get().FindTConsoleVariableDataInt(TEXT("r.Mobile.SkyLightPermutation"));
-		static auto AllowStaticLightingVar = IConsoleManager::Get().FindTConsoleVariableDataInt(TEXT("r.AllowStaticLighting"));
-		const bool bAllowStaticLighting = (!AllowStaticLightingVar || AllowStaticLightingVar->GetValueOnAnyThread() != 0);
 
 		const int32 MobileSkyLightPermutationOptions = MobileSkyLightPermutationCVar->GetValueOnAnyThread();
 		const bool bDeferredShading = IsMobileDeferredShadingEnabled(Parameters.Platform);
@@ -353,9 +351,7 @@ public:
 		// Translucent materials always support clustered shading on mobile deferred
 		const bool bSupportsLocalLights = (!bDeferredShading && MobileForwardEnableLocalLights(Parameters.Platform)) || (bDeferredShading && bMaterialUsesForwardShading);
 		// Only compile skylight version for lit materials
-		const bool bForwardShading = !bDeferredShading || bMaterialUsesForwardShading;
-		const bool bRenderSkyLightInBasePass = bForwardShading || bAllowStaticLighting;
-		const bool bShouldCacheBySkylight = !bEnableSkyLight || (bIsLit && bRenderSkyLightInBasePass);
+		const bool bShouldCacheBySkylight = !bEnableSkyLight || bIsLit;
 
 		// Only compile skylight permutations when they are enabled
 		if (bIsLit && !UseSkylightPermutation(bEnableSkyLight, MobileSkyLightPermutationOptions))
@@ -364,7 +360,8 @@ public:
 		}
 
 		// Deferred shading does not need SkyLight and LocalLight permutations
-		// TODO: skip skylight permutations for deferred	
+		// TODO: skip skylight permutations for deferred
+		const bool bForwardShading = !bDeferredShading || bMaterialUsesForwardShading;
 		const bool bShouldCacheByShading = (bForwardShading || !bEnableLocalLights);
 		const bool bShouldCacheByLocalLights = !bEnableLocalLights || (bIsLit && bEnableLocalLights == bSupportsLocalLights);
 
diff --git a/Engine/Source/Runtime/Renderer/Private/MobileDeferredShadingPass.cpp b/Engine/Source/Runtime/Renderer/Private/MobileDeferredShadingPass.cpp
--- a/Engine/Source/Runtime/Renderer/Private/MobileDeferredShadingPass.cpp
+++ b/Engine/Source/Runtime/Renderer/Private/MobileDeferredShadingPass.cpp
@@ -46,7 +46,6 @@ class FMobileDirectionalLightFunctionPS : public FMaterialShader
 	class FEnableClustredReflection		: SHADER_PERMUTATION_BOOL("ENABLE_CLUSTERED_REFLECTION");
 	class FEnablePlanarReflection		: SHADER_PERMUTATION_BOOL("ENABLE_PLANAR_REFLECTION");
 	class FEnableSkyLight				: SHADER_PERMUTATION_BOOL("ENABLE_SKY_LIGHT");
-	class FEnableDynamicSkyLight		: SHADER_PERMUTATION_BOOL("ENABLE_DYNAMIC_SKY_LIGHT");
 	class FEnableCSM					: SHADER_PERMUTATION_BOOL("ENABLE_MOBILE_CSM");
 	class FShadowQuality				: SHADER_PERMUTATION_RANGE_INT("MOBILE_SHADOW_QUALITY", 1, 3); // not using Quality=0
 	
@@ -56,7 +55,6 @@ class FMobileDirectionalLightFunctionPS : public FMaterialShader
 		FEnableClustredReflection, 
 		FEnablePlanarReflection,
 		FEnableSkyLight,
-		FEnableDynamicSkyLight,
 		FEnableCSM, 
 		FShadowQuality>;
 
@@ -97,11 +95,6 @@ class FMobileDirectionalLightFunctionPS : public FMaterialShader
 			PermutationVector.Set<FEnableShadingModelSupport>(false);
 		}
 
-		if (!PermutationVector.Get<FEnableSkyLight>())
-		{
-			PermutationVector.Set<FEnableDynamicSkyLight>(false);
-		}
-
 		return PermutationVector;
 	}
 
@@ -124,12 +117,11 @@ class FMobileDirectionalLightFunctionPS : public FMaterialShader
 		return true;
 	}
 
-	static FPermutationDomain BuildPermutationVector(const FViewInfo& View, bool bInlineReflectionAndSky, bool bShadingModelSupport, bool bDynamicShadows, bool bSkyLight, bool bDynamicSkyLight, bool bPlanarReflection)
+	static FPermutationDomain BuildPermutationVector(const FViewInfo& View, bool bInlineReflectionAndSky, bool bShadingModelSupport, bool bDynamicShadows, bool bSkyLight, bool bPlanarReflection)
 	{
 		bool bUseClusteredLights = UseClusteredDeferredShading(View.GetShaderPlatform());
 		bool bClustredReflection = bInlineReflectionAndSky && (View.NumBoxReflectionCaptures + View.NumSphereReflectionCaptures) > 0;
 		bool bEnableSkyLight = bInlineReflectionAndSky && bSkyLight;
-		bool bEnableDynamicSkyLight = bInlineReflectionAndSky && bDynamicSkyLight;
 		const bool bMobileUsesShadowMaskTexture = MobileUsesShadowMaskTexture(View.GetShaderPlatform());
 		int32 ShadowQuality = bDynamicShadows && !bMobileUsesShadowMaskTexture ? (int32)GetShadowQuality() : 0;
 				
@@ -139,7 +131,6 @@ class FMobileDirectionalLightFunctionPS : public FMaterialShader
 		PermutationVector.Set<FMobileDirectionalLightFunctionPS::FEnableClustredReflection>(bClustredReflection);
 		PermutationVector.Set<FMobileDirectionalLightFunctionPS::FEnablePlanarReflection>(bPlanarReflection);
 		PermutationVector.Set<FMobileDirectionalLightFunctionPS::FEnableSkyLight>(bEnableSkyLight);
-		PermutationVector.Set<FMobileDirectionalLightFunctionPS::FEnableDynamicSkyLight>(bEnableDynamicSkyLight);
 		PermutationVector.Set<FMobileDirectionalLightFunctionPS::FEnableCSM>(ShadowQuality > 0);
 		PermutationVector.Set<FMobileDirectionalLightFunctionPS::FShadowQuality>(FMath::Clamp(ShadowQuality, 1, 3));
 		return PermutationVector;
@@ -249,14 +240,12 @@ public:
 	class FEnableClustredReflection		: SHADER_PERMUTATION_BOOL("ENABLE_CLUSTERED_REFLECTION");
 	class FEnablePlanarReflection		: SHADER_PERMUTATION_BOOL("ENABLE_PLANAR_REFLECTION");
 	class FEnableSkyLight				: SHADER_PERMUTATION_BOOL("ENABLE_SKY_LIGHT");
-	class FEnableDynamicSkyLight : SHADER_PERMUTATION_BOOL("ENABLE_DYNAMIC_SKY_LIGHT");
 	
 	using FPermutationDomain = TShaderPermutationDomain<
 		FEnableShadingModelSupport, 
 		FEnableClustredReflection, 
 		FEnablePlanarReflection,
-		FEnableSkyLight,
-		FEnableDynamicSkyLight
+		FEnableSkyLight
 	>;
 
 	BEGIN_SHADER_PARAMETER_STRUCT(FParameters, )
@@ -278,11 +267,6 @@ public:
 			return false;
 		}
 
-		if (!PermutationVector.Get<FEnableSkyLight>() && PermutationVector.Get<FEnableDynamicSkyLight>())
-		{
-			return false;
-		}
-
 		return true;
 	}
 
@@ -353,9 +337,7 @@ static void GetLightMaterial(const FCachedLightMaterial& DefaultLightMaterial, c
 void RenderReflectionEnvironmentSkyLighting(FRHICommandListImmediate& RHICmdList, const FScene& Scene, const FViewInfo& View)
 {
 	// Skylights with static lighting already had their diffuse contribution baked into lightmaps
-	const bool bIsDynamicLighting = !FReadOnlyCVARCache::Get().bAllowStaticLighting || Scene.GetForceNoPrecomputedLighting();
 	const bool bSkyLight = Scene.SkyLight && !Scene.SkyLight->bHasStaticLighting && View.Family->EngineShowFlags.SkyLighting;
-	const bool bDynamicSkyLight = bSkyLight && (!Scene.SkyLight->bWantsStaticShadowing || bIsDynamicLighting);
 	const bool bClustredReflection = (View.NumBoxReflectionCaptures + View.NumSphereReflectionCaptures) > 0;
 	const bool bPlanarReflection = Scene.GetForwardPassGlobalPlanarReflection() != nullptr;
 	if (!(bSkyLight || bClustredReflection || bPlanarReflection))
@@ -400,7 +382,6 @@ void RenderReflectionEnvironmentSkyLighting(FRHICommandListImmediate& RHICmdList
 		PermutationVector.Set<FMobileReflectionEnvironmentSkyLightingPS::FEnableClustredReflection>(bClustredReflection);
 		PermutationVector.Set<FMobileReflectionEnvironmentSkyLightingPS::FEnablePlanarReflection>(bPlanarReflection);
 		PermutationVector.Set<FMobileReflectionEnvironmentSkyLightingPS::FEnableSkyLight>(bSkyLight);
-		PermutationVector.Set<FMobileReflectionEnvironmentSkyLightingPS::FEnableDynamicSkyLight>(bDynamicSkyLight);
 		TShaderMapRef<FMobileReflectionEnvironmentSkyLightingPS> PixelShader(View.ShaderMap, PermutationVector);
 		
 		GraphicsPSOInit.BoundShaderState.VertexDeclarationRHI = GFilterVertexDeclaration.VertexDeclarationRHI;
@@ -499,9 +480,7 @@ static void RenderDirectionalLight(FRHICommandListImmediate& RHICmdList, const F
 	}
 
 	// Skylights with static lighting already had their diffuse contribution baked into lightmaps
-	const bool bIsDynamicLighting = !FReadOnlyCVARCache::Get().bAllowStaticLighting || Scene.GetForceNoPrecomputedLighting();
 	const bool bSkyLight = Scene.SkyLight && !Scene.SkyLight->bHasStaticLighting && View.Family->EngineShowFlags.SkyLighting;
-	const bool bDynamicSkyLight = bSkyLight && (!Scene.SkyLight->bWantsStaticShadowing || bIsDynamicLighting);
 	const bool bDynamicShadows = DirectionalLight.Proxy->CastsDynamicShadow() && (LightingChannel == 0u) && View.Family->EngineShowFlags.DynamicShadows;
 	const bool bPlanarReflection = Scene.GetForwardPassGlobalPlanarReflection() != nullptr;
 
@@ -524,7 +503,6 @@ static void RenderDirectionalLight(FRHICommandListImmediate& RHICmdList, const F
 			bEnableShadingModelSupport,
 			bDynamicShadows,
 			bSkyLight,
-			bDynamicSkyLight,
 			bPlanarReflection
 		);
 		FCachedLightMaterial LightMaterial;
diff --git a/Engine/Source/Runtime/Renderer/Private/RendererScene.cpp b/Engine/Source/Runtime/Renderer/Private/RendererScene.cpp
--- a/Engine/Source/Runtime/Renderer/Private/RendererScene.cpp
+++ b/Engine/Source/Runtime/Renderer/Private/RendererScene.cpp
@@ -6092,42 +6092,6 @@ bool FScene::IsPrimitiveBeingRemoved(FPrimitiveSceneInfo* PrimitiveSceneInfo) co
 	return RemovedPrimitiveSceneInfos.Find(PrimitiveSceneInfo) != nullptr;
 }
 
-bool FScene::ShouldRenderSkylightInBasePass(EBlendMode BlendMode) const
-{
-	if (IsMobilePlatform(GetShaderPlatform()))
-	{
-		bool bRenderSkyLight = SkyLight && !SkyLight->bHasStaticLighting;
-		const bool bIsForwardShading = IsTranslucentBlendMode(BlendMode) || !IsMobileDeferredShadingEnabled(GetShaderPlatform());
-		const bool bIsDynamicLighting = !ReadOnlyCVARCache.bAllowStaticLighting || GetForceNoPrecomputedLighting();
-
-		// Both stationary and movable skylights are applied in base pass for forward shading
-		// Only stationary skylights are applied in base pass for deferred with static lighting enabled
-		return bRenderSkyLight && (bIsForwardShading || (!bIsDynamicLighting && SkyLight->bWantsStaticShadowing));
-	}
-	else
-	{
-		bool bRenderSkyLight = SkyLight && !SkyLight->bHasStaticLighting && !(ShouldRenderRayTracingSkyLight(SkyLight) && !IsForwardShadingEnabled(GetShaderPlatform()));
-
-		if (IsTranslucentBlendMode(BlendMode))
-		{
-			// Both stationary and movable skylights are applied in base pass for translucent materials
-			bRenderSkyLight = bRenderSkyLight
-				&& (ReadOnlyCVARCache.bEnableStationarySkylight || !SkyLight->bWantsStaticShadowing);
-		}
-		else
-		{
-			// For opaque materials, stationary skylight is applied in base pass but movable skylight
-			// is applied in a separate render pass (bWantssStaticShadowing means stationary skylight)
-			bRenderSkyLight = bRenderSkyLight
-				&& ((ReadOnlyCVARCache.bEnableStationarySkylight && SkyLight->bWantsStaticShadowing)
-					|| (!SkyLight->bWantsStaticShadowing
-						&& IsForwardShadingEnabled(GetShaderPlatform())));
-		}
-
-		return bRenderSkyLight;
-	}
-}
-
 /**
  * Dummy NULL scene interface used by dedicated servers.
  */
diff --git a/Engine/Source/Runtime/Renderer/Private/ScenePrivate.h b/Engine/Source/Runtime/Renderer/Private/ScenePrivate.h
--- a/Engine/Source/Runtime/Renderer/Private/ScenePrivate.h
+++ b/Engine/Source/Runtime/Renderer/Private/ScenePrivate.h
@@ -3386,7 +3386,28 @@ public:
 
 	virtual bool IsEditorScene() const override { return bIsEditorScene; }
 
-	bool ShouldRenderSkylightInBasePass(EBlendMode BlendMode) const;
+	bool ShouldRenderSkylightInBasePass(EBlendMode BlendMode) const
+	{
+		bool bRenderSkyLight = SkyLight && !SkyLight->bHasStaticLighting && !(ShouldRenderRayTracingSkyLight(SkyLight) && !IsForwardShadingEnabled(GetShaderPlatform()));
+
+		if (IsTranslucentBlendMode(BlendMode))
+		{
+			// Both stationary and movable skylights are applied in base pass for translucent materials
+			bRenderSkyLight = bRenderSkyLight
+				&& (ReadOnlyCVARCache.bEnableStationarySkylight || !SkyLight->bWantsStaticShadowing);
+		}
+		else
+		{
+			// For opaque materials, stationary skylight is applied in base pass but movable skylight
+			// is applied in a separate render pass (bWantssStaticShadowing means stationary skylight)
+			bRenderSkyLight = bRenderSkyLight
+				&& ((ReadOnlyCVARCache.bEnableStationarySkylight && SkyLight->bWantsStaticShadowing)
+					|| (!SkyLight->bWantsStaticShadowing
+						&& (IsForwardShadingEnabled(GetShaderPlatform()) || IsMobilePlatform(GetShaderPlatform()))));
+		}
+
+		return bRenderSkyLight;
+	}
 
 	virtual TArray<FPrimitiveComponentId> GetScenePrimitiveComponentIds() const override
 	{
