Changed EyeAdaptationBuffer to be a StructuredBuffer.
- Avoid need for a sampler.
- On certain platforms exposure values can end up in SGPR instead of VGPR.
- Moved EyeAdaptationBuffer read from tonemapper VS to PS due to limitation of mobile GPUs.

diff --git a/Engine/Shaders/Private/EyeAdaptationCommon.ush b/Engine/Shaders/Private/EyeAdaptationCommon.ush
--- a/Engine/Shaders/Private/EyeAdaptationCommon.ush
+++ b/Engine/Shaders/Private/EyeAdaptationCommon.ush
@@ -7,7 +7,7 @@
 #if !EYE_ADAPTATION_DISABLED
 
 Texture2D EyeAdaptationTexture;
-Buffer<float4> EyeAdaptationBuffer;
+StructuredBuffer<float4> EyeAdaptationBuffer;
 
 #if SHADING_PATH_DEFERRED
 float4 EyeAdaptationLookupTexture2D(Texture2D InEyeAdaptation)
@@ -15,7 +15,7 @@ float4 EyeAdaptationLookupTexture2D(Texture2D InEyeAdaptation)
 	return InEyeAdaptation.Load(int3(0, 0, 0));
 }
 #elif SHADING_PATH_MOBILE
-float4 EyeAdaptationLookupBuffer(Buffer<float4> InEyeAdaptation)
+float4 EyeAdaptationLookupBuffer(StructuredBuffer<float4> InEyeAdaptation)
 {
 	return InEyeAdaptation[0];
 }
diff --git a/Engine/Shaders/Private/PostProcessMobile.usf b/Engine/Shaders/Private/PostProcessMobile.usf
--- a/Engine/Shaders/Private/PostProcessMobile.usf
+++ b/Engine/Shaders/Private/PostProcessMobile.usf
@@ -1122,7 +1122,7 @@ void SunMergePS_Mobile(
 
 // EyeAdaptation
 
-Buffer<float4> EyeAdaptationBuffer;
+StructuredBuffer<float4> EyeAdaptationBuffer;
 
 static const float FLOAT_PRECISION = 1e+5;
 
@@ -1272,7 +1272,7 @@ float4 BasicEyeAdaptation_Mobile()
 }
 
 #if BASIC_EYEADAPTATION_COMPUTE_SHADER
-RWBuffer<float4> OutputBuffer;
+RWStructuredBuffer<float4> OutputBuffer;
 
 [numthreads(1, 1, 1)]
 void BasicEyeAdaptationCS_Mobile()
@@ -1490,7 +1490,7 @@ float ComputeEyeAdaptationExposure_Mobile()
 	return LumAve;
 }
 
-RWBuffer<float4> OutputBuffer;
+RWStructuredBuffer<float4> OutputBuffer;
 
 [numthreads(1, 1, 1)]
 void HistogramEyeAdaptationCS(uint2 DispatchThreadId : SV_DispatchThreadID)
diff --git a/Engine/Shaders/Private/PostProcessTonemap.usf b/Engine/Shaders/Private/PostProcessTonemap.usf
--- a/Engine/Shaders/Private/PostProcessTonemap.usf
+++ b/Engine/Shaders/Private/PostProcessTonemap.usf
@@ -200,7 +200,7 @@ float4 ChromaticAberrationParams;
 void TonemapCommonVS(
 	in	float4 Position,
 	in	float2 TexCoord,
-	out	float4 OutExposureScaleVignette,
+	out	float2 OutVignette,
 	out	float4 OutGrainUV,
 	out	float2 OutScreenPos,
 	out	float2 OutFullViewUV
@@ -210,28 +210,13 @@ void TonemapCommonVS(
 	// TODO: Looks to be Ronin specific.. check this:
 	OutFullViewUV.xy = Position.xy * float2(0.5,-0.5) + 0.5;
 
-#if EYEADAPTATION_EXPOSURE_FIX == 1
-	#if FEATURE_LEVEL >= FEATURE_LEVEL_SM5
-		// texture can be GWhiteTexture which is 1x1. It's important we don't read outside bounds.
-		OutExposureScaleVignette.xy = EyeAdaptationTexture.Load(int3(0, 0, 0)).xw;
-	#else
-		OutExposureScaleVignette.xy = EyeAdaptationBuffer[0].xw;
-	#endif
-#else
-	// If eye adaptation is disabled (e.g. fixed exposure level ) or  not available.  
-	OutExposureScaleVignette.xy = DefaultEyeExposure;
-#endif
-
-	// Middle grey lum value adjusted by exposure compensation
-	OutExposureScaleVignette.y = log2(0.18 * OutExposureScaleVignette.y * EyeAdaptation_LocalExposureMiddleGreyExposureCompensation);
-
 	const float AspectRatio = Output_ViewportSize.y * Output_ViewportSizeInverse.x;
 
 	float2 ColorViewportPos = UVToScreenPos(TexCoord, Color_Extent);
 
 	// Scale vignette to always be a circle with consistent corner intensity.
 	float2 LensViewportPos = ConvertScreenViewportSpaceToLensViewportSpace(ColorViewportPos);
-	OutExposureScaleVignette.zw = VignetteSpace(LensViewportPos, AspectRatio);
+	OutVignette = VignetteSpace(LensViewportPos, AspectRatio);
 
 	// Grain
 	OutGrainUV.xy = TexCoord + Color_ExtentInverse * float2(-0.5,0.5);
@@ -246,7 +231,7 @@ void MainVS(
 	in 					float4 InPosition 				: ATTRIBUTE0,
 	in 					float2 InTexCoord 				: ATTRIBUTE1,
 	out noperspective 	float2 OutTexCoord 				: TEXCOORD0,
-	out noperspective 	float4 OutExposureScaleVignette : TEXCOORD1,
+	out noperspective 	float2 OutVignette				: TEXCOORD1,
 	out noperspective 	float4 OutGrainUV 				: TEXCOORD2,
 	out noperspective 	float2 OutScreenPos 			: TEXCOORD3,
 	out noperspective 	float2 OutFullViewUV 			: TEXCOORD4,
@@ -254,7 +239,7 @@ void MainVS(
 	)
 {
 	DrawRectangle(InPosition, InTexCoord, OutPosition, OutTexCoord);
-    TonemapCommonVS(OutPosition, OutTexCoord, OutExposureScaleVignette, OutGrainUV, OutScreenPos, OutFullViewUV);
+    TonemapCommonVS(OutPosition, OutTexCoord, OutVignette, OutGrainUV, OutScreenPos, OutFullViewUV);
 }
 
 // Function graphing
@@ -363,7 +348,7 @@ half3 LookupSceneColor(float2 SceneUV, float2 PixelOffset)
 
 float4 TonemapCommonPS(
 	float2 UV,
-	float4 ExposureScaleVignette,
+	float2 Vignette,
 	float4 GrainUV,
 	float2 ScreenPos, // [-1, 1]x[-1, 1]
 	float2 FullViewUV,
@@ -433,7 +418,19 @@ float4 TonemapCommonPS(
 
 	SceneColor.rgb *= View.OneOverPreExposure;
 
-	float ExposureScale = ExposureScaleVignette.x;
+#if EYEADAPTATION_EXPOSURE_FIX == 1
+#if FEATURE_LEVEL >= FEATURE_LEVEL_SM5
+	// texture can be GWhiteTexture which is 1x1. It's important we don't read outside bounds.
+	float2 EyeAdaptationData = EyeAdaptationTexture.Load(int3(0, 0, 0)).xw;
+#else
+	float2 EyeAdaptationData = EyeAdaptationBuffer[0].xw;
+#endif
+#else
+	// If eye adaptation is disabled (e.g. fixed exposure level ) or not available.  
+	float2 EyeAdaptationData = DefaultEyeExposure;
+#endif
+
+	float ExposureScale = EyeAdaptationData.x;
 
 	float LocalExposureOneOverPreExposure = View.OneOverPreExposure;
 
@@ -441,7 +438,9 @@ float4 TonemapCommonPS(
 	{
 		float LuminanceVal = CalculateEyeAdaptationLuminance(SceneColor.rgb);
 		float LogLuminance = log2(LuminanceVal);
-		float MiddleGreyLumValue = ExposureScaleVignette.y;
+
+		// Middle grey lum value adjusted by exposure compensation
+		float MiddleGreyLumValue = log2(0.18 * EyeAdaptationData.y * EyeAdaptation_LocalExposureMiddleGreyExposureCompensation);
 
 		float BaseLogLum = CalculateBaseLogLuminance(LogLuminance, EyeAdaptation_LocalExposureBlurredLuminanceBlend, ExposureScale, FullViewUV, LumBilateralGrid, BlurredLogLum, LumBilateralGridSampler, BlurredLogLumSampler);
 		float LocalExposure = CalculateLocalExposure(LogLuminance + log2(ExposureScale), BaseLogLum, MiddleGreyLumValue, EyeAdaptation_LocalExposureHighlightContrastScale, EyeAdaptation_LocalExposureShadowContrastScale, EyeAdaptation_LocalExposureDetailStrength);
@@ -604,7 +603,7 @@ float4 TonemapCommonPS(
 		LinearColor *= ExposureScale;
 
 		#if USE_VIGNETTE
-			LinearColor.rgb *= ComputeVignetteMask( ExposureScaleVignette.zw, TonemapperParams.x );
+			LinearColor.rgb *= ComputeVignetteMask( Vignette, TonemapperParams.x );
 		#endif
 
 		#if USE_FILM_GRAIN
@@ -703,7 +702,7 @@ float4 TonemapCommonPS(
 // pixel shader entry point
 void MainPS(
 	in noperspective float2 UV : TEXCOORD0,
-	in noperspective float4 InExposureScaleVignette : TEXCOORD1,
+	in noperspective float2 InVignette : TEXCOORD1,
 	in noperspective float4 GrainUV : TEXCOORD2,
 	in noperspective float2 ScreenPos : TEXCOORD3,
 	in noperspective float2 FullViewUV : TEXCOORD4,
@@ -711,7 +710,7 @@ void MainPS(
 	out float4 OutColor : SV_Target0
 	)
 {
-    OutColor = TonemapCommonPS(UV, InExposureScaleVignette, GrainUV, ScreenPos, FullViewUV, SvPosition);
+    OutColor = TonemapCommonPS(UV, InVignette, GrainUV, ScreenPos, FullViewUV, SvPosition);
 }
 
 #if COMPUTESHADER
@@ -729,12 +728,12 @@ void MainCS(uint2 DispatchThreadId : SV_DispatchThreadID)
 		return;
 	}
 
-	float4 ExposureScaleVignette;
+	float2 Vignette;
 	float4 GrainUV;
 	float2 FullViewUV, ScreenPos;	
-	TonemapCommonVS(InScreenPos, UV, ExposureScaleVignette, GrainUV, ScreenPos, FullViewUV);
+	TonemapCommonVS(InScreenPos, UV, Vignette, GrainUV, ScreenPos, FullViewUV);
 
-	float4 OutColor = TonemapCommonPS(UV, ExposureScaleVignette, GrainUV, ScreenPos, FullViewUV, SvPosition);
+	float4 OutColor = TonemapCommonPS(UV, Vignette, GrainUV, ScreenPos, FullViewUV, SvPosition);
 
 	uint2 PixelPos = DispatchThreadId + Output_ViewportMin;
 
diff --git a/Engine/Source/Runtime/Renderer/Private/MobileBasePassRendering.cpp b/Engine/Source/Runtime/Renderer/Private/MobileBasePassRendering.cpp
--- a/Engine/Source/Runtime/Renderer/Private/MobileBasePassRendering.cpp
+++ b/Engine/Source/Runtime/Renderer/Private/MobileBasePassRendering.cpp
@@ -201,7 +201,7 @@ void SetupMobileBasePassUniformParameters(
 
 	BasePassParameters.PreIntegratedGFTexture = GSystemTextures.PreintegratedGF->GetRHI();
 	BasePassParameters.PreIntegratedGFSampler = TStaticSamplerState<SF_Bilinear, AM_Clamp, AM_Clamp, AM_Clamp>::GetRHI();
-	BasePassParameters.EyeAdaptationBuffer = GraphBuilder.CreateSRV(GetEyeAdaptationBuffer(GraphBuilder, View), PF_A32B32G32R32F);
+	BasePassParameters.EyeAdaptationBuffer = GraphBuilder.CreateSRV(GetEyeAdaptationBuffer(GraphBuilder, View));
 
 	FRDGTextureRef AmbientOcclusionTexture = SystemTextures.White;
 	if (BasePass == EMobileBasePass::Opaque && MobileBasePassTextures.ScreenSpaceAO != nullptr)
diff --git a/Engine/Source/Runtime/Renderer/Private/MobileBasePassRendering.h b/Engine/Source/Runtime/Renderer/Private/MobileBasePassRendering.h
--- a/Engine/Source/Runtime/Renderer/Private/MobileBasePassRendering.h
+++ b/Engine/Source/Runtime/Renderer/Private/MobileBasePassRendering.h
@@ -41,7 +41,7 @@ BEGIN_GLOBAL_SHADER_PARAMETER_STRUCT(FMobileBasePassUniformParameters, )
 	SHADER_PARAMETER_STRUCT(FReflectionUniformParameters, ReflectionsParameters)
 	SHADER_PARAMETER_TEXTURE(Texture2D, PreIntegratedGFTexture)
 	SHADER_PARAMETER_SAMPLER(SamplerState, PreIntegratedGFSampler)
-	SHADER_PARAMETER_RDG_BUFFER_SRV(Buffer<float4>, EyeAdaptationBuffer)
+	SHADER_PARAMETER_RDG_BUFFER_SRV(StructuredBuffer<float4>, EyeAdaptationBuffer)
 	SHADER_PARAMETER_RDG_TEXTURE(Texture2D, AmbientOcclusionTexture)
 	SHADER_PARAMETER_SAMPLER(SamplerState, AmbientOcclusionSampler)
 	SHADER_PARAMETER_RDG_TEXTURE(Texture2D, ScreenSpaceShadowMaskTexture)
diff --git a/Engine/Source/Runtime/Renderer/Private/PostProcess/PostProcessEyeAdaptation.cpp b/Engine/Source/Runtime/Renderer/Private/PostProcess/PostProcessEyeAdaptation.cpp
--- a/Engine/Source/Runtime/Renderer/Private/PostProcess/PostProcessEyeAdaptation.cpp
+++ b/Engine/Source/Runtime/Renderer/Private/PostProcess/PostProcessEyeAdaptation.cpp
@@ -1045,7 +1045,7 @@ const TRefCountPtr<FRDGPooledBuffer>& FSceneViewState::FEyeAdaptationManager::Ge
 	// Create textures if needed.
 	if (!ExposureBufferData[BufferIndex].IsValid())
 	{
-		FRDGBufferDesc RDGBufferDesc = FRDGBufferDesc::CreateBufferDesc(sizeof(FVector4f), 1);
+		FRDGBufferDesc RDGBufferDesc = FRDGBufferDesc::CreateStructuredDesc(sizeof(FVector4f), 1);
 		RDGBufferDesc.Usage |= BUF_SourceCopy;
 		FRDGBufferRef RDGBuffer = GraphBuilder.CreateBuffer(RDGBufferDesc, TEXT("EyeAdaptationBuffer"), ERDGBufferFlags::MultiFrame);
 
diff --git a/Engine/Source/Runtime/Renderer/Private/PostProcess/PostProcessMaterial.cpp b/Engine/Source/Runtime/Renderer/Private/PostProcess/PostProcessMaterial.cpp
--- a/Engine/Source/Runtime/Renderer/Private/PostProcess/PostProcessMaterial.cpp
+++ b/Engine/Source/Runtime/Renderer/Private/PostProcess/PostProcessMaterial.cpp
@@ -478,7 +478,7 @@ FScreenPassTexture AddPostProcessMaterialPass(
 	PostProcessMaterialParameters->View = View.ViewUniformBuffer;
 	if (bMobilePlatform)
 	{
-		PostProcessMaterialParameters->EyeAdaptationBuffer = GraphBuilder.CreateSRV(GetEyeAdaptationBuffer(GraphBuilder, View), PF_A32B32G32R32F);
+		PostProcessMaterialParameters->EyeAdaptationBuffer = GraphBuilder.CreateSRV(GetEyeAdaptationBuffer(GraphBuilder, View));
 	}
 	else
 	{
diff --git a/Engine/Source/Runtime/Renderer/Private/PostProcess/PostProcessMaterial.h b/Engine/Source/Runtime/Renderer/Private/PostProcess/PostProcessMaterial.h
--- a/Engine/Source/Runtime/Renderer/Private/PostProcess/PostProcessMaterial.h
+++ b/Engine/Source/Runtime/Renderer/Private/PostProcess/PostProcessMaterial.h
@@ -45,7 +45,7 @@ BEGIN_SHADER_PARAMETER_STRUCT(FPostProcessMaterialParameters, )
 	SHADER_PARAMETER_STRUCT_ARRAY(FScreenPassTextureInput, PostProcessInput, [kPostProcessMaterialInputCountMax])
 	SHADER_PARAMETER_SAMPLER(SamplerState, PostProcessInput_BilinearSampler)
 	SHADER_PARAMETER_RDG_TEXTURE(Texture2D, EyeAdaptationTexture)
-	SHADER_PARAMETER_RDG_BUFFER_SRV(Buffer<float4>, EyeAdaptationBuffer)
+	SHADER_PARAMETER_RDG_BUFFER_SRV(StructuredBuffer<float4>, EyeAdaptationBuffer)
 	SHADER_PARAMETER(uint32, bMetalMSAAHDRDecode)
 	SHADER_PARAMETER(uint32, bSceneDepthWithoutWaterTextureAvailable)
 	SHADER_PARAMETER_RDG_TEXTURE(Texture2D, SceneDepthWithoutSingleLayerWaterTexture)
diff --git a/Engine/Source/Runtime/Renderer/Private/PostProcess/PostProcessMobile.cpp b/Engine/Source/Runtime/Renderer/Private/PostProcess/PostProcessMobile.cpp
--- a/Engine/Source/Runtime/Renderer/Private/PostProcess/PostProcessMobile.cpp
+++ b/Engine/Source/Runtime/Renderer/Private/PostProcess/PostProcessMobile.cpp
@@ -1897,9 +1897,9 @@ public:
 	BEGIN_SHADER_PARAMETER_STRUCT(FParameters, )
 		SHADER_PARAMETER_STRUCT_REF(FViewUniformShaderParameters, View)
 		SHADER_PARAMETER_STRUCT(FEyeAdaptationParameters, EyeAdaptation)
-		SHADER_PARAMETER_RDG_BUFFER_SRV(Buffer<float4>, EyeAdaptationBuffer)
+		SHADER_PARAMETER_RDG_BUFFER_SRV(StructuredBuffer<float4>, EyeAdaptationBuffer)
 		SHADER_PARAMETER_RDG_BUFFER_SRV(Buffer<uint>, LogLuminanceWeightBuffer)
-		SHADER_PARAMETER_RDG_BUFFER_UAV(RWBuffer<float4>, OutputBuffer)
+		SHADER_PARAMETER_RDG_BUFFER_UAV(RWStructuredBuffer<float4>, OutputBuffer)
 	END_SHADER_PARAMETER_STRUCT()
 
 	/** Static Shader boilerplate */
@@ -1930,9 +1930,9 @@ class FMobileHistogramEyeAdaptationCS : public FGlobalShader
 public:
 	BEGIN_SHADER_PARAMETER_STRUCT(FParameters, )
 		SHADER_PARAMETER_STRUCT(FEyeAdaptationParameters, EyeAdaptation)
-		SHADER_PARAMETER_RDG_BUFFER_SRV(Buffer<float4>, EyeAdaptationBuffer)
+		SHADER_PARAMETER_RDG_BUFFER_SRV(StructuredBuffer<float4>, EyeAdaptationBuffer)
 		SHADER_PARAMETER_RDG_BUFFER_SRV(Buffer<uint>, HistogramBuffer)
-		SHADER_PARAMETER_RDG_BUFFER_UAV(RWBuffer<float4>, OutputBuffer)
+		SHADER_PARAMETER_RDG_BUFFER_UAV(RWStructuredBuffer<float4>, OutputBuffer)
 	END_SHADER_PARAMETER_STRUCT()
 
 	static bool ShouldCompilePermutation(const FGlobalShaderPermutationParameters& Parameters)
@@ -1957,10 +1957,10 @@ void AddMobileEyeAdaptationPass(FRDGBuilder& GraphBuilder, const FViewInfo& View
 	View.SwapEyeAdaptationBuffers();
 
 	FRDGBufferRef EyeAdaptationBuffer = Inputs.EyeAdaptationBuffer;
-	FRDGBufferSRVRef EyeAdaptationBufferSRV = GraphBuilder.CreateSRV(EyeAdaptationBuffer, PF_A32B32G32R32F);
+	FRDGBufferSRVRef EyeAdaptationBufferSRV = GraphBuilder.CreateSRV(EyeAdaptationBuffer);
 
 	FRDGBufferRef OutputBuffer = GraphBuilder.RegisterExternalBuffer(View.GetEyeAdaptationBuffer(GraphBuilder), ERDGBufferFlags::MultiFrame);
-	FRDGBufferUAVRef OutputBufferUAV = GraphBuilder.CreateUAV(OutputBuffer, PF_A32B32G32R32F);
+	FRDGBufferUAVRef OutputBufferUAV = GraphBuilder.CreateUAV(OutputBuffer);
 
 	if (Inputs.bUseBasicEyeAdaptation)
 	{
diff --git a/Engine/Source/Runtime/Renderer/Private/PostProcess/PostProcessTonemap.cpp b/Engine/Source/Runtime/Renderer/Private/PostProcess/PostProcessTonemap.cpp
--- a/Engine/Source/Runtime/Renderer/Private/PostProcess/PostProcessTonemap.cpp
+++ b/Engine/Source/Runtime/Renderer/Private/PostProcess/PostProcessTonemap.cpp
@@ -67,7 +67,7 @@ namespace TonemapperPermutation
 // Shared permutation dimensions between deferred and mobile renderer.
 class FTonemapperBloomDim          : SHADER_PERMUTATION_BOOL("USE_BLOOM");
 class FTonemapperGammaOnlyDim      : SHADER_PERMUTATION_BOOL("USE_GAMMA_ONLY");
-class FTonemapperLocalExposureDim : SHADER_PERMUTATION_BOOL("USE_LOCAL_EXPOSURE");
+class FTonemapperLocalExposureDim  : SHADER_PERMUTATION_BOOL("USE_LOCAL_EXPOSURE");
 class FTonemapperVignetteDim       : SHADER_PERMUTATION_BOOL("USE_VIGNETTE");
 class FTonemapperSharpenDim        : SHADER_PERMUTATION_BOOL("USE_SHARPEN");
 class FTonemapperFilmGrainDim      : SHADER_PERMUTATION_BOOL("USE_FILM_GRAIN");
@@ -81,7 +81,8 @@ using FCommonDomain = TShaderPermutationDomain<
 	FTonemapperVignetteDim,
 	FTonemapperSharpenDim,
 	FTonemapperFilmGrainDim,
-	FTonemapperMsaaDim>;
+	FTonemapperMsaaDim,
+	FTonemapperEyeAdaptationDim>;
 
 bool ShouldCompileCommonPermutation(const FGlobalShaderPermutationParameters& Parameters, const FCommonDomain& PermutationVector)
 {
@@ -99,13 +100,14 @@ bool ShouldCompileCommonPermutation(const FGlobalShaderPermutationParameters& Pa
 			!PermutationVector.Get<FTonemapperVignetteDim>() &&
 			!PermutationVector.Get<FTonemapperSharpenDim>() &&
 			!PermutationVector.Get<FTonemapperFilmGrainDim>() &&
-			!PermutationVector.Get<FTonemapperMsaaDim>();
+			!PermutationVector.Get<FTonemapperMsaaDim>() &&
+			!PermutationVector.Get<FTonemapperEyeAdaptationDim>();
 	}
 	return true;
 }
 
 // Common conversion of engine settings into.
-FCommonDomain BuildCommonPermutationDomain(const FViewInfo& View, bool bGammaOnly, bool bLocalExposure, bool bMetalMSAAHDRDecode)
+FCommonDomain BuildCommonPermutationDomain(const FViewInfo& View, bool bGammaOnly, bool bEyeAdaptation, bool bLocalExposure, bool bMetalMSAAHDRDecode)
 {
 	const FSceneViewFamily* Family = View.Family;
 
@@ -125,8 +127,9 @@ FCommonDomain BuildCommonPermutationDomain(const FViewInfo& View, bool bGammaOnl
 	PermutationVector.Set<FTonemapperBloomDim>(Settings.BloomIntensity > 0.0);
 	PermutationVector.Set<FTonemapperLocalExposureDim>(bLocalExposure);
 	PermutationVector.Set<FTonemapperFilmGrainDim>(View.FilmGrainTexture != nullptr);
-	PermutationVector.Set<FTonemapperSharpenDim>(CVarTonemapperSharpen.GetValueOnRenderThread() > 0.0f);	
+	PermutationVector.Set<FTonemapperSharpenDim>(CVarTonemapperSharpen.GetValueOnRenderThread() > 0.0f);
 	PermutationVector.Set<FTonemapperMsaaDim>(bMetalMSAAHDRDecode);
+	PermutationVector.Set<FTonemapperEyeAdaptationDim>(bEyeAdaptation);
 	return PermutationVector;
 }
 
@@ -319,7 +322,7 @@ BEGIN_SHADER_PARAMETER_STRUCT(FTonemapParameters, )
 	SHADER_PARAMETER(float, EditorNITLevel)
 	SHADER_PARAMETER(uint32, bOutputInHDR)
 	// ES3_1 uses EyeAdaptationBuffer
-	SHADER_PARAMETER_RDG_BUFFER_SRV(Buffer<float4>, EyeAdaptationBuffer)
+	SHADER_PARAMETER_RDG_BUFFER_SRV(StructuredBuffer<float4>, EyeAdaptationBuffer)
 END_SHADER_PARAMETER_STRUCT()
 
 class FFilmGrainReduceCS : public FGlobalShader
@@ -366,7 +369,6 @@ public:
 	// FDrawRectangleParameters is filled by DrawScreenPass.
 	SHADER_USE_PARAMETER_STRUCT_WITH_LEGACY_BASE(FTonemapVS, FGlobalShader);
 
-	using FPermutationDomain = TShaderPermutationDomain<TonemapperPermutation::FTonemapperEyeAdaptationDim>;
 	using FParameters = FTonemapParameters;
 
 	static bool ShouldCompilePermutation(const FGlobalShaderPermutationParameters& Parameters)
@@ -413,9 +415,7 @@ public:
 	DECLARE_GLOBAL_SHADER(FTonemapCS);
 	SHADER_USE_PARAMETER_STRUCT(FTonemapCS, FGlobalShader);
 
-	using FPermutationDomain = TShaderPermutationDomain<
-		TonemapperPermutation::FDesktopDomain,
-		TonemapperPermutation::FTonemapperEyeAdaptationDim>;
+	using FPermutationDomain = TonemapperPermutation::FDesktopDomain;
 
 	BEGIN_SHADER_PARAMETER_STRUCT(FParameters, )
 		SHADER_PARAMETER_STRUCT_INCLUDE(FTonemapParameters, Tonemap)
@@ -431,7 +431,7 @@ public:
 
 		FPermutationDomain PermutationVector(Parameters.PermutationId);
 
-		return TonemapperPermutation::ShouldCompileDesktopPermutation(Parameters, PermutationVector.Get<TonemapperPermutation::FDesktopDomain>());
+		return TonemapperPermutation::ShouldCompileDesktopPermutation(Parameters, PermutationVector);
 	}
 
 	static void ModifyCompilationEnvironment(const FGlobalShaderPermutationParameters& Parameters, FShaderCompilerEnvironment& OutEnvironment)
@@ -808,7 +808,7 @@ FScreenPassTexture AddTonemapPass(FRDGBuilder& GraphBuilder, const FViewInfo& Vi
 	CommonParameters.LensPrincipalPointOffsetScaleInverse.Z = 1.0f / View.LensPrincipalPointOffsetScale.Z;
 	CommonParameters.LensPrincipalPointOffsetScaleInverse.W = 1.0f / View.LensPrincipalPointOffsetScale.W;
 	
-	FRDGBufferSRVRef EyeAdaptationBufferSRV = Inputs.EyeAdaptationBuffer != nullptr ? GraphBuilder.CreateSRV(Inputs.EyeAdaptationBuffer, PF_A32B32G32R32F) : nullptr;
+	FRDGBufferSRVRef EyeAdaptationBufferSRV = Inputs.EyeAdaptationBuffer != nullptr ? GraphBuilder.CreateSRV(Inputs.EyeAdaptationBuffer) : nullptr;
 
 	CommonParameters.EyeAdaptationBuffer = EyeAdaptationBufferSRV;
 
@@ -816,7 +816,7 @@ FScreenPassTexture AddTonemapPass(FRDGBuilder& GraphBuilder, const FViewInfo& Vi
 	TonemapperPermutation::FDesktopDomain DesktopPermutationVector;
 
 	{
-		TonemapperPermutation::FCommonDomain CommonDomain = TonemapperPermutation::BuildCommonPermutationDomain(View, Inputs.bGammaOnly, Inputs.LocalExposureTexture != nullptr, Inputs.bMetalMSAAHDRDecode);
+		TonemapperPermutation::FCommonDomain CommonDomain = TonemapperPermutation::BuildCommonPermutationDomain(View, Inputs.bGammaOnly, bEyeAdaptation, Inputs.LocalExposureTexture != nullptr, Inputs.bMetalMSAAHDRDecode);
 		DesktopPermutationVector.Set<TonemapperPermutation::FCommonDomain>(CommonDomain);
 
 		if (!CommonDomain.Get<TonemapperPermutation::FTonemapperGammaOnlyDim>())
@@ -845,11 +845,7 @@ FScreenPassTexture AddTonemapPass(FRDGBuilder& GraphBuilder, const FViewInfo& Vi
 		PassParameters->Tonemap = CommonParameters;
 		PassParameters->RWOutputTexture = GraphBuilder.CreateUAV(Output.Texture);
 
-		FTonemapCS::FPermutationDomain PermutationVector;
-		PermutationVector.Set<TonemapperPermutation::FDesktopDomain>(DesktopPermutationVector);
-		PermutationVector.Set<TonemapperPermutation::FTonemapperEyeAdaptationDim>(bEyeAdaptation);
-
-		TShaderMapRef<FTonemapCS> ComputeShader(View.ShaderMap, PermutationVector);
+		TShaderMapRef<FTonemapCS> ComputeShader(View.ShaderMap, DesktopPermutationVector);
 
 		FComputeShaderUtils::AddPass(
 			GraphBuilder,
@@ -864,10 +860,7 @@ FScreenPassTexture AddTonemapPass(FRDGBuilder& GraphBuilder, const FViewInfo& Vi
 		PassParameters->Tonemap = CommonParameters;
 		PassParameters->RenderTargets[0] = Output.GetRenderTargetBinding();
 
-		FTonemapVS::FPermutationDomain VertexPermutationVector;
-		VertexPermutationVector.Set<TonemapperPermutation::FTonemapperEyeAdaptationDim>(bEyeAdaptation);
-
-		TShaderMapRef<FTonemapVS> VertexShader(View.ShaderMap, VertexPermutationVector);
+		TShaderMapRef<FTonemapVS> VertexShader(View.ShaderMap);
 		TShaderMapRef<FTonemapPS> PixelShader(View.ShaderMap, DesktopPermutationVector);
 
 		// If this is a stereo view, there's a good chance we need alpha out of the tonemapper
diff --git a/Engine/Source/Runtime/Renderer/Private/PostProcess/TemporalAA.cpp b/Engine/Source/Runtime/Renderer/Private/PostProcess/TemporalAA.cpp
--- a/Engine/Source/Runtime/Renderer/Private/PostProcess/TemporalAA.cpp
+++ b/Engine/Source/Runtime/Renderer/Private/PostProcess/TemporalAA.cpp
@@ -123,7 +123,7 @@ public:
 		SHADER_PARAMETER(FVector4f, ScreenPosToHistoryBufferUV)
 
 		SHADER_PARAMETER_RDG_TEXTURE(Texture2D, EyeAdaptationTexture)
-		SHADER_PARAMETER_RDG_BUFFER_SRV(Buffer<float4>, EyeAdaptationBuffer)
+		SHADER_PARAMETER_RDG_BUFFER_SRV(StructuredBuffer<float4>, EyeAdaptationBuffer)
 
 		// Inputs
 		SHADER_PARAMETER_RDG_TEXTURE(Texture2D, InputSceneColor)
@@ -790,7 +790,7 @@ FTAAOutputs AddTemporalAAPass(
 
 		if (View.GetFeatureLevel() <= ERHIFeatureLevel::ES3_1)
 		{
-			PassParameters->EyeAdaptationBuffer = GraphBuilder.CreateSRV(GetEyeAdaptationBuffer(GraphBuilder, View), PF_A32B32G32R32F);
+			PassParameters->EyeAdaptationBuffer = GraphBuilder.CreateSRV(GetEyeAdaptationBuffer(GraphBuilder, View));
 		}
 		else
 		{
diff --git a/Engine/Source/Runtime/Renderer/Private/SceneTextureParameters.cpp b/Engine/Source/Runtime/Renderer/Private/SceneTextureParameters.cpp
--- a/Engine/Source/Runtime/Renderer/Private/SceneTextureParameters.cpp
+++ b/Engine/Source/Runtime/Renderer/Private/SceneTextureParameters.cpp
@@ -92,6 +92,6 @@ FRDGBufferRef GetEyeAdaptationBuffer(FRDGBuilder& GraphBuilder, const FSceneView
 	}
 	else
 	{
-		return GraphBuilder.RegisterExternalBuffer(GWhiteVertexBufferWithRDG->Buffer);
+		return GSystemTextures.GetDefaultStructuredBuffer(GraphBuilder, sizeof(FVector4f), FVector4f::One());
 	}
 }
\ No newline at end of file
