Clamped SLW depth buffer copy to viewport to save some GPU time - more savings if a HiZ copy pass is implemented

diff --git a/Engine/Source/Runtime/Renderer/Private/SingleLayerWaterRendering.cpp b/Engine/Source/Runtime/Renderer/Private/SingleLayerWaterRendering.cpp
--- a/Engine/Source/Runtime/Renderer/Private/SingleLayerWaterRendering.cpp
+++ b/Engine/Source/Runtime/Renderer/Private/SingleLayerWaterRendering.cpp
@@ -56,6 +56,12 @@ static TAutoConsoleVariable<int32> CVarParallelSingleLayerWaterPass(
 	TEXT("Toggles parallel single layer water pass rendering. Parallel rendering must be enabled for this to have an effect."),
 	ECVF_RenderThreadSafe);
 
+static TAutoConsoleVariable<int32> CVarSingleLayerWaterPassOptimizedClear(
+	TEXT("r.Water.SingleLayer.OptimizedClear"),
+	1,
+	TEXT("Toggles optimized depth clear"),
+	ECVF_RenderThreadSafe);
+
 static TAutoConsoleVariable<int32> CVarWaterSingleLayerSSR(
 	TEXT("r.Water.SingleLayer.SSR"), 1,
 	TEXT("Enable SSR for the single water rendering system."),
@@ -454,6 +460,13 @@ FSingleLayerWaterPrePassResult* FDeferredShadingSceneRenderer::RenderSingleLayer
 
 		// Copy main depth buffer content to our prepass depth buffer and clear stencil to 0
 		// TODO: replace with AddCopyTexturePass() and AddClearDepthStencilPass() once CopyTexture() supports depth buffer copies on all platforms.
+
+		const bool bOptimizedClear = CVarSingleLayerWaterPassOptimizedClear.GetValueOnRenderThread() == 1;
+		if (false)//bOptimizedClear && GRHISupportsDepthUAV && GRHISupportsExplicitHTile)
+		{
+			// TODO: Implement optimized copy path
+		}
+		else
 		{
 			FCopyDepthPS::FParameters* PassParameters = GraphBuilder.AllocParameters<FCopyDepthPS::FParameters>();
 			if (DepthPrepassTextureDesc.NumSamples > 1)
@@ -473,6 +486,10 @@ FSingleLayerWaterPrePassResult* FDeferredShadingSceneRenderer::RenderSingleLayer
 			TShaderMapRef<FCopyDepthPS> PixelShader(ShaderMap, PermutationVector);
 
 			FIntRect Viewport(0, 0, DepthPrepassTextureDesc.Extent.X, DepthPrepassTextureDesc.Extent.Y);
+			if (bOptimizedClear && Views.Num() == 1)
+			{
+				Viewport = Views[0].ViewRect;
+			}
 
 			// Set depth test to always pass and stencil test to replace all pixels with zero, essentially also clearing stencil while doing the depth copy.
 			FRHIDepthStencilState* DepthStencilState = TStaticDepthStencilState<
