Fixed cloud PS path PSO creation on Metal failing (ps output component count must match RT).

diff --git a/Engine/Shaders/Private/VolumetricCloud.usf b/Engine/Shaders/Private/VolumetricCloud.usf
--- a/Engine/Shaders/Private/VolumetricCloud.usf
+++ b/Engine/Shaders/Private/VolumetricCloud.usf
@@ -1548,13 +1548,13 @@ void MainCommon(in FMaterialPixelParameters MaterialParameters, in float4 SvPosi
 void MainPS(
 	in float4 SvPosition : SV_Position
 	, out float4 OutColor0 : SV_Target0
-	, out float2 OutDepth : SV_Target1)
+	, out float4 OutDepth : SV_Target1)
 {
 	FMaterialPixelParameters MaterialParameters = MakeInitializedMaterialPixelParameters();
 	
 	float4 OutDepth4 = 0;
 	MainCommon(MaterialParameters, SvPosition, OutColor0, OutDepth4);
-	OutDepth = OutDepth4.xy;
+	OutDepth = float4(OutDepth4.xy, 0.0f, 0.0f);	// PS path never writes min and max tracing depth in zw.
 }
 
 #endif // SHADER_RENDERVIEW_PS
diff --git a/Engine/Source/Runtime/Renderer/Private/VolumetricRenderTarget.cpp b/Engine/Source/Runtime/Renderer/Private/VolumetricRenderTarget.cpp
--- a/Engine/Source/Runtime/Renderer/Private/VolumetricRenderTarget.cpp
+++ b/Engine/Source/Runtime/Renderer/Private/VolumetricRenderTarget.cpp
@@ -290,7 +290,7 @@ FRDGTextureRef FVolumetricRenderTargetViewStateData::GetOrCreateVolumetricTracin
 
 	if (!VolumetricTracingRTDepth.IsValid())
 	{
-		EPixelFormat DepthDataFormat = Mode == 0 ? PF_FloatRGBA : PF_G16R16F;	// TODO This should accoutn for the new mode 0 variation cvar. Or mode 5?
+		EPixelFormat DepthDataFormat = Mode == 0 ? PF_FloatRGBA : PF_G16R16F; // Mode 0 supports MinAndMax depth tracing when the compute path is used so always allocate a 4-components texture in this case.
 
 		FPooledRenderTargetDesc Desc = FPooledRenderTargetDesc::Create2DDesc(
 			VolumetricTracingRTResolution, DepthDataFormat, FClearValueBinding(FLinearColor(63000.0f, 63000.0f, 63000.0f, 63000.0f)),
