Sky Atmosphere now automatically transtion to forced tracing as a function of the plant radius.

diff --git a/Engine/Shaders/Private/SkyAtmosphere.usf b/Engine/Shaders/Private/SkyAtmosphere.usf
--- a/Engine/Shaders/Private/SkyAtmosphere.usf
+++ b/Engine/Shaders/Private/SkyAtmosphere.usf
@@ -883,7 +883,7 @@ void RenderSkyAtmosphereRayMarchingPS(
 
 	float ViewHeight = length(WorldPos);
 #if FASTSKY_ENABLED && RENDERSKY_ENABLED
-	if (ViewHeight < (Atmosphere.TopRadiusKm + PLANET_RADIUS_SAFE_TRACE_EDGE) && DeviceZ == FarDepthValue)
+	if (ViewHeight < (Atmosphere.TopRadiusKm * PLANET_RADIUS_RATIO_SAFE_EDGE) && DeviceZ == FarDepthValue)
 	{
 		float2 UV;
 
diff --git a/Engine/Shaders/Private/SkyAtmosphereCommon.ush b/Engine/Shaders/Private/SkyAtmosphereCommon.ush
--- a/Engine/Shaders/Private/SkyAtmosphereCommon.ush
+++ b/Engine/Shaders/Private/SkyAtmosphereCommon.ush
@@ -23,8 +23,10 @@
 #define SKY_UNIT_TO_CM (1.0f/CM_TO_SKY_UNIT)
 // Float accuracy offset in Sky unit (km, so this is 1m). Should match the one in FAtmosphereSetup::ComputeViewData
 #define PLANET_RADIUS_OFFSET 0.001f
+
 // Planet radius safe edge to make sure ray does intersect with the atmosphere, for it to traverse the atmosphere. Must match the one in FSceneRenderer::RenderSkyAtmosphereInternal.
-#define PLANET_RADIUS_SAFE_TRACE_EDGE 0.01f
+// This is (0.01km/6420km).
+#define PLANET_RADIUS_RATIO_SAFE_EDGE 1.00000155763f
 
 
 // The number of killometer per slice in the aerial pespective camera volume texture. (assuming a uniform depth distribution)
diff --git a/Engine/Source/Runtime/Renderer/Private/SkyAtmosphereRendering.cpp b/Engine/Source/Runtime/Renderer/Private/SkyAtmosphereRendering.cpp
--- a/Engine/Source/Runtime/Renderer/Private/SkyAtmosphereRendering.cpp
+++ b/Engine/Source/Runtime/Renderer/Private/SkyAtmosphereRendering.cpp
@@ -1744,8 +1744,8 @@ void FSceneRenderer::RenderSkyAtmosphereInternal(
 	const FVector3f ViewOrigin = (FVector3f)ViewMatrices.GetViewOrigin();
 	const FVector3f PlanetCenter = (FVector3f)Atmosphere.PlanetCenterKm * KM_TO_CM;	// LWC_TODO: Precision Loss
 	const float TopOfAtmosphere = Atmosphere.TopRadiusKm * KM_TO_CM;
-	const float PlanetRadiusTraceSafeEdgeCm = 1000.0f;	// 10 meters, must match PLANET_RADIUS_SAFE_TRACE_EDGE
-	const bool ForceRayMarching = SkyRC.bForceRayMarching || (FVector3f::Distance(ViewOrigin, PlanetCenter) - TopOfAtmosphere - PlanetRadiusTraceSafeEdgeCm) > 0.0f;
+	const float PLANET_RADIUS_RATIO_SAFE_EDGE = 1.00000155763f; // must match PLANET_RADIUS_SAFE_TRACE_EDGE
+	const bool ForceRayMarching = SkyRC.bForceRayMarching || (FVector3f::Distance(ViewOrigin, PlanetCenter) >= (TopOfAtmosphere * PLANET_RADIUS_RATIO_SAFE_EDGE));
 	const bool bDisableBlending = SkyRC.bDisableBlending;
 
 	// We only support MSAA up to 8 sample and in forward
