RVT: Compress direct to VT physical texture when UAV aliasing is supported.
Saves some GPU time transitionaling and copying in the RVT write work.

diff --git a/Engine/Shaders/Private/BCCompressionCommon.ush b/Engine/Shaders/Private/BCCompressionCommon.ush
--- a/Engine/Shaders/Private/BCCompressionCommon.ush
+++ b/Engine/Shaders/Private/BCCompressionCommon.ush
@@ -167,7 +167,7 @@ uint2 GetPackedAlphaBlockIndices( in float Block[16], in float MinAlpha, in floa
 }
 
 // Compress a BC1 block
-uint4 CompressBC1Block( in float3 Block[16] )
+uint2 CompressBC1Block( in float3 Block[16] )
 {
 	float3 MinColor, MaxColor;
 	GetMinMax( Block, MinColor, MaxColor );
@@ -181,12 +181,12 @@ uint4 CompressBC1Block( in float3 Block[16] )
 		Indices = GetPackedColorBlockIndices( Block, MinColor, MaxColor );
 	}
 
-	return uint4(MaxColor565, MinColor565, Indices & 0xffff, Indices >> 16);
+	return uint2((MinColor565 << 16) | MaxColor565, Indices);
 }
 
 // Compress a BC1 block that will be sampled as sRGB
 // We expect linear colors as input and convert internally
-uint4 CompressBC1Block_SRGB( in float3 Block[16] )
+uint2 CompressBC1Block_SRGB( in float3 Block[16] )
 {
 #if (FEATURE_LEVEL == FEATURE_LEVEL_ES3_1)
 	// mobile uses linear color for RVT
@@ -205,7 +205,7 @@ uint4 CompressBC1Block_SRGB( in float3 Block[16] )
 		Indices = GetPackedColorBlockIndices( Block, MinColor, MaxColor );
 	}
 
-	return uint4(MaxColor565, MinColor565, Indices & 0xffff, Indices >> 16);
+	return uint2((MinColor565 << 16) | MaxColor565, Indices);
 }
 
 // Compress a BC3 block
diff --git a/Engine/Shaders/Private/VirtualTextureCompress.usf b/Engine/Shaders/Private/VirtualTextureCompress.usf
--- a/Engine/Shaders/Private/VirtualTextureCompress.usf
+++ b/Engine/Shaders/Private/VirtualTextureCompress.usf
@@ -15,18 +15,39 @@
 #define ETC_PROFILE 0
 #endif
 
-uint4 DestRect;
+#ifndef PACK_RG32_RGBA16
+#define PACK_RG32_RGBA16 0
+#endif
+
+#if PACK_RG32_RGBA16
+#define TYPE_64BIT uint4
+#else
+#define TYPE_64BIT uint2
+#endif
+
+uint4 SourceRect;
+uint2 DestPos0;
+uint2 DestPos1;
+uint2 DestPos2;
+
 Texture2D<float4> RenderTexture0;
 SamplerState TextureSampler0;
 Texture2D<float4> RenderTexture1;
 SamplerState TextureSampler1;
 Texture2D<float4> RenderTexture2;
 SamplerState TextureSampler2;
-RWTexture2D<uint4> OutCompressTexture0;
-RWTexture2D<uint4> OutCompressTexture1;
-RWTexture2D<uint4> OutCompressTexture2;
 
-/** Pack ETC RG32 to RGBA16 for output. */
+// Outputs can be 64bit or 128 bit depending on formats.
+// The 64 bit type can be uint2 or uint4 depending on platform support.
+RWTexture2D<TYPE_64BIT> OutCompressTexture0_64bit;
+RWTexture2D<TYPE_64BIT> OutCompressTexture1_64bit;
+RWTexture2D<TYPE_64BIT> OutCompressTexture2_64bit;
+RWTexture2D<uint4> OutCompressTexture0_128bit;
+RWTexture2D<uint4> OutCompressTexture1_128bit;
+RWTexture2D<uint4> OutCompressTexture2_128bit;
+
+#if PACK_RG32_RGBA16
+/** Pack RG32 to RGBA16 for output. */
 uint4 PackRG32(in uint2 u2)
 {
 	uint x0 = (u2.x >> 16) & 0x0000ffff;
@@ -35,6 +56,12 @@ uint4 PackRG32(in uint2 u2)
 	uint y1 = (u2.y) & 0x0000ffff;
 	return uint4(x1,x0,y1,y0);
 }
+#else
+uint2 PackRG32(in uint2 u2)
+{
+	return u2;
+}
+#endif // PACK_RG32_RGBA16
 
 /** Trivial vertex shader for copy shaders. */
 void CopyVS(in uint VertexId : SV_VertexID, out float4 OutPosition : SV_POSITION, out float2 OutTexCoord : TEXCOORD0)
@@ -47,22 +74,23 @@ void CopyVS(in uint VertexId : SV_VertexID, out float4 OutPosition : SV_POSITION
 /** Compress base color to a single BC1 target. */
 [numthreads(8,8,1)]
 void CompressBaseColorCS(
-	uint3 ThreadId : SV_DispatchThreadID 
-	)
+	uint3 ThreadId : SV_DispatchThreadID)
 {
-	uint2 SamplePos = ThreadId.xy * 4;
-	if (any(ThreadId.xy >= DestRect.zw))
+	if (any(ThreadId.xy * 4 >= SourceRect.zw))
+	{
 		return;
+	}
 
-	float2 TexelUVSize = 1.f / float2(DestRect.zw);
+	uint2 SamplePos = SourceRect.xy + ThreadId.xy * 4;
+	float2 TexelUVSize = 1.f / float2(SourceRect.zw);
 	float2 SampleUV = (float2(SamplePos) + 0.5f) * TexelUVSize;
 
 	float3 BlockBaseColor[16];
 	ReadBlockRGB(RenderTexture0, TextureSampler0, SampleUV, TexelUVSize, BlockBaseColor);
 #if ETC_PROFILE
-	OutCompressTexture0[ThreadId.xy + DestRect.xy] = PackRG32(CompressBlock_ETC2_SRGB(BlockBaseColor));
+	OutCompressTexture0_64bit[ThreadId.xy + DestPos0] = PackRG32(CompressBlock_ETC2_SRGB(BlockBaseColor));
 #else
-	OutCompressTexture0[ThreadId.xy + DestRect.xy] = CompressBC1Block_SRGB(BlockBaseColor);
+	OutCompressTexture0_64bit[ThreadId.xy + DestPos0] = PackRG32(CompressBC1Block_SRGB(BlockBaseColor));
 #endif
 
 }
@@ -80,14 +108,15 @@ void CopyBaseColorPS(
 /** Compress base color, normal, roughness and specular to BC3 target pair. */
 [numthreads(8,8,1)]
 void CompressBaseColorNormalSpecularCS(
-	uint3 ThreadId : SV_DispatchThreadID 
-	)
+	uint3 ThreadId : SV_DispatchThreadID)
 {
-	uint2 SamplePos = ThreadId.xy * 4;
-	if (any(SamplePos >= DestRect.zw))
+	if (any(ThreadId.xy * 4 >= SourceRect.zw))
+	{
 		return;
+	}
 
-	float2 TexelUVSize = 1.f / float2(DestRect.zw);
+	uint2 SamplePos = SourceRect.xy + ThreadId.xy * 4;
+	float2 TexelUVSize = 1.f / float2(SourceRect.zw);
 	float2 SampleUV = (float2(SamplePos) + 0.5f) * TexelUVSize;
 
 	float3 BlockBaseColor[16];
@@ -105,11 +134,11 @@ void CompressBaseColorNormalSpecularCS(
 		BlockSpecular[i].z = round(BlockSpecular[i].z);
 	}
 #if ETC_PROFILE
-	OutCompressTexture0[ThreadId.xy + DestRect.xy] = CompressBlock_ETC2_RGBA(BlockBaseColor, BlockNormalX);
-	OutCompressTexture1[ThreadId.xy + DestRect.xy] = CompressBlock_ETC2_RGBA(BlockSpecular, BlockNormalY);
+	OutCompressTexture0_128bit[ThreadId.xy + DestPos0] = CompressBlock_ETC2_RGBA(BlockBaseColor, BlockNormalX);
+	OutCompressTexture1_128bit[ThreadId.xy + DestPos1] = CompressBlock_ETC2_RGBA(BlockSpecular, BlockNormalY);
 #else
-	OutCompressTexture0[ThreadId.xy + DestRect.xy] = CompressBC3Block_SRGB(BlockBaseColor, BlockNormalX);
-	OutCompressTexture1[ThreadId.xy + DestRect.xy] = CompressBC3Block(BlockSpecular, BlockNormalY);
+	OutCompressTexture0_128bit[ThreadId.xy + DestPos0] = CompressBC3Block_SRGB(BlockBaseColor, BlockNormalX);
+	OutCompressTexture1_128bit[ThreadId.xy + DestPos1] = CompressBC3Block(BlockSpecular, BlockNormalY);
 #endif
 
 }
@@ -117,14 +146,15 @@ void CompressBaseColorNormalSpecularCS(
 /** Compress base color, normal, roughness to BC1 & BC3. */
 [numthreads(8, 8, 1)]
 void CompressBaseColorNormalRoughnessCS(
-	uint3 ThreadId : SV_DispatchThreadID
-)
+	uint3 ThreadId : SV_DispatchThreadID)
 {
-	uint2 SamplePos = ThreadId.xy * 4;
-	if (any(SamplePos >= DestRect.zw))
+	if (any(ThreadId.xy * 4 >= SourceRect.zw))
+	{
 		return;
+	}
 
-	float2 TexelUVSize = 1.f / float2(DestRect.zw);
+	uint2 SamplePos = SourceRect.xy + ThreadId.xy * 4;
+	float2 TexelUVSize = 1.f / float2(SourceRect.zw);
 	float2 SampleUV = (float2(SamplePos) + 0.5f) * TexelUVSize;
 
 	float3 BlockBaseColor[16];
@@ -135,11 +165,11 @@ void CompressBaseColorNormalRoughnessCS(
 								   1.0, 1.0, 1.0, 1.0, 1.0, 1.0, 1.0, 1.0};
 	ReadBlockRGB(RenderTexture1, TextureSampler2, SampleUV, TexelUVSize, BlockNormalXRougnnessY);
 #if ETC_PROFILE
-	OutCompressTexture0[ThreadId.xy + DestRect.xy] = PackRG32(CompressBlock_ETC2_RGB(BlockBaseColor));
-	OutCompressTexture1[ThreadId.xy + DestRect.xy] = CompressBlock_ETC2_RGBA(BlockNormalXRougnnessY, BlockNormalZ);
+	OutCompressTexture0_64bit[ThreadId.xy + DestPos0] = PackRG32(CompressBlock_ETC2_RGB(BlockBaseColor));
+	OutCompressTexture1_128bit[ThreadId.xy + DestPos1] = CompressBlock_ETC2_RGBA(BlockNormalXRougnnessY, BlockNormalZ);
 #else
-	OutCompressTexture0[ThreadId.xy + DestRect.xy] = CompressBC1Block_SRGB(BlockBaseColor);
-	OutCompressTexture1[ThreadId.xy + DestRect.xy] = CompressBC3Block(BlockNormalXRougnnessY, BlockNormalZ);
+	OutCompressTexture0_64bit[ThreadId.xy + DestPos0] = PackRG32(CompressBC1Block_SRGB(BlockBaseColor));
+	OutCompressTexture1_128bit[ThreadId.xy + DestPos1] = CompressBC3Block(BlockNormalXRougnnessY, BlockNormalZ);
 #endif
 
 }
@@ -164,14 +194,15 @@ void CopyBaseColorNormalSpecularPS(
 /** Compress base color, normal, roughness and specular to BC3, BC5, BC1 target set. */
 [numthreads(8,8,1)]
 void CompressBaseColorNormalSpecularYCoCgCS(
-	uint3 ThreadId : SV_DispatchThreadID 
-	)
+	uint3 ThreadId : SV_DispatchThreadID)
 {
-	uint2 SamplePos = ThreadId.xy * 4;
-	if (any(SamplePos >= DestRect.zw))
+	if (any(ThreadId.xy * 4 >= SourceRect.zw))
+	{
 		return;
+	}
 
-	float2 TexelUVSize = 1.f / float2(DestRect.zw);
+	uint2 SamplePos = SourceRect.xy + ThreadId.xy * 4;
+	float2 TexelUVSize = 1.f / float2(SourceRect.zw);
 	float2 SampleUV = (float2(SamplePos) + 0.5f) * TexelUVSize;
 
 	float3 BlockBaseColor[16];
@@ -190,13 +221,13 @@ void CompressBaseColorNormalSpecularYCoCgCS(
 	}
 
 #if ETC_PROFILE
-	OutCompressTexture0[ThreadId.xy + DestRect.xy] = CompressBlock_ETC2_YCoCg(BlockBaseColor);
-	OutCompressTexture1[ThreadId.xy + DestRect.xy] = CompressBlock_ETC2_RG(BlockNormalX, BlockNormalY);
-	OutCompressTexture2[ThreadId.xy + DestRect.xy] = PackRG32(CompressBlock_ETC2_RGB(BlockSpecular));
+	OutCompressTexture0_128bit[ThreadId.xy + DestPos0] = CompressBlock_ETC2_YCoCg(BlockBaseColor);
+	OutCompressTexture1_128bit[ThreadId.xy + DestPos1] = CompressBlock_ETC2_RG(BlockNormalX, BlockNormalY);
+	OutCompressTexture2_64bit[ThreadId.xy + DestPos2] = PackRG32(CompressBlock_ETC2_RGB(BlockSpecular));
 #else
-	OutCompressTexture0[ThreadId.xy + DestRect.xy] = CompressBC3BlockYCoCg(BlockBaseColor);
-	OutCompressTexture1[ThreadId.xy + DestRect.xy] = CompressBC5Block(BlockNormalX, BlockNormalY);
-	OutCompressTexture2[ThreadId.xy + DestRect.xy] = CompressBC1Block(BlockSpecular);
+	OutCompressTexture0_128bit[ThreadId.xy + DestPos0] = CompressBC3BlockYCoCg(BlockBaseColor);
+	OutCompressTexture1_128bit[ThreadId.xy + DestPos1] = CompressBC5Block(BlockNormalX, BlockNormalY);
+	OutCompressTexture2_64bit[ThreadId.xy + DestPos2] = PackRG32(CompressBC1Block(BlockSpecular));
 #endif
 }
 
@@ -222,14 +253,13 @@ void CopyBaseColorNormalSpecularYCoCgPS(
 /** Compress base color, normal, roughness, specular and mask to BC3, BC5, BC3 target set. */
 [numthreads(8,8,1)]
 void CompressBaseColorNormalSpecularMaskYCoCgCS(
-	uint3 ThreadId : SV_DispatchThreadID 
-	)
+	uint3 ThreadId : SV_DispatchThreadID)
 {
-	uint2 SamplePos = ThreadId.xy * 4;
-	if (any(SamplePos >= DestRect.zw))
+	if (any(ThreadId.xy * 4 >= SourceRect.zw))
 		return;
 
-	float2 TexelUVSize = 1.f / float2(DestRect.zw);
+	uint2 SamplePos = SourceRect.xy + ThreadId.xy * 4;
+	float2 TexelUVSize = 1.f / float2(SourceRect.zw);
 	float2 SampleUV = (float2(SamplePos) + 0.5f) * TexelUVSize;
 
 	float3 BlockBaseColor[16];
@@ -249,13 +279,13 @@ void CompressBaseColorNormalSpecularMaskYCoCgCS(
 	}
 
 #if ETC_PROFILE
-	OutCompressTexture0[ThreadId.xy + DestRect.xy] = CompressBlock_ETC2_YCoCg(BlockBaseColor);
-	OutCompressTexture1[ThreadId.xy + DestRect.xy] = CompressBlock_ETC2_RG(BlockNormalX, BlockNormalY);
-	OutCompressTexture2[ThreadId.xy + DestRect.xy] = CompressBlock_ETC2_RGBA(BlockSpecular, BlockMask);
+	OutCompressTexture0_128bit[ThreadId.xy + DestPos0] = CompressBlock_ETC2_YCoCg(BlockBaseColor);
+	OutCompressTexture1_128bit[ThreadId.xy + DestPos1] = CompressBlock_ETC2_RG(BlockNormalX, BlockNormalY);
+	OutCompressTexture2_128bit[ThreadId.xy + DestPos2] = CompressBlock_ETC2_RGBA(BlockSpecular, BlockMask);
 #else
-	OutCompressTexture0[ThreadId.xy + DestRect.xy] = CompressBC3BlockYCoCg(BlockBaseColor);
-	OutCompressTexture1[ThreadId.xy + DestRect.xy] = CompressBC5Block(BlockNormalX, BlockNormalY);
-	OutCompressTexture2[ThreadId.xy + DestRect.xy] = CompressBC3Block(BlockSpecular, BlockMask);
+	OutCompressTexture0_128bit[ThreadId.xy + DestPos0] = CompressBC3BlockYCoCg(BlockBaseColor);
+	OutCompressTexture1_128bit[ThreadId.xy + DestPos1] = CompressBC5Block(BlockNormalX, BlockNormalY);
+	OutCompressTexture2_128bit[ThreadId.xy + DestPos2] = CompressBC3Block(BlockSpecular, BlockMask);
 #endif
 }
 
diff --git a/Engine/Source/Runtime/Renderer/Private/VT/RuntimeVirtualTextureRender.cpp b/Engine/Source/Runtime/Renderer/Private/VT/RuntimeVirtualTextureRender.cpp
--- a/Engine/Source/Runtime/Renderer/Private/VT/RuntimeVirtualTextureRender.cpp
+++ b/Engine/Source/Runtime/Renderer/Private/VT/RuntimeVirtualTextureRender.cpp
@@ -20,6 +20,13 @@
 
 namespace RuntimeVirtualTexture
 {
+	static TAutoConsoleVariable<int32> CVarVTDirectCompress(
+		TEXT("r.VT.RVT.DirectCompress"),
+		1,
+		TEXT("Compress texture data direct to the physical texture on platforms that support it."),
+		ECVF_RenderThreadSafe
+	);
+
     int32 RenderCaptureNextRVTPagesDraws = 0;
     static FAutoConsoleVariableRef CVarRenderCaptureNextRVTPagesDraws(
 	    TEXT("r.VT.RenderCaptureNextPagesDraws"),
@@ -95,6 +102,12 @@ namespace RuntimeVirtualTexture
 		return false;
 	}
 
+	/** For platforms that do not support 2-channel images, write 64bit compressed texture outputs into RGBA16 instead of RG32. */
+	bool UseRGBA16(EShaderPlatform ShaderPlatform)
+	{
+		return IsOpenGLPlatform(ShaderPlatform);
+	}
+
 	/** Mesh material shader for writing to the virtual texture. */
 	class FShader_VirtualTextureMaterialDraw : public FMeshMaterialShader
 	{
@@ -407,8 +420,6 @@ namespace RuntimeVirtualTexture
 		static void SetupRenderTargetsInfo(ERuntimeVirtualTextureMaterialType MaterialType, ERHIFeatureLevel::Type FeatureLevel, bool bLQFormat, FGraphicsPipelineRenderTargetsInfo& RenderTargetsInfo)
 		{
 			const ETextureCreateFlags VT_SRGB = FeatureLevel > ERHIFeatureLevel::ES3_1 ? TexCreate_SRGB : TexCreate_None;
-			const EPixelFormat Compressed64BitFormat = PF_R16G16B16A16_UINT;
-			const EPixelFormat Compressed128BitFormat = PF_R32G32B32A32_UINT;
 
 			switch (MaterialType)
 			{
@@ -440,18 +451,21 @@ namespace RuntimeVirtualTexture
 			}
 		}
 
-		FRenderGraphSetup(FRDGBuilder& GraphBuilder, ERHIFeatureLevel::Type FeatureLevel, ERuntimeVirtualTextureMaterialType MaterialType, FRHITexture2D* OutputTexture0, FIntPoint TextureSize, bool bIsThumbnails)
+		FRenderGraphSetup(FRDGBuilder& GraphBuilder, ERHIFeatureLevel::Type FeatureLevel, ERuntimeVirtualTextureMaterialType MaterialType, TArrayView<TRefCountPtr<IPooledRenderTarget>> OutputTargets, EPixelFormat OutputFormat0, FIntPoint TextureSize, bool bIsThumbnails)
 		{
-			bRenderPass = OutputTexture0 != nullptr;
+			bRenderPass = OutputFormat0 != PF_Unknown;
 			bCopyThumbnailPass = bRenderPass && bIsThumbnails;
-			EPixelFormat OutputFormat = (OutputTexture0 != nullptr ? OutputTexture0->GetFormat() : PF_Unknown);
-			const bool bCompressedFormat = GPixelFormats[OutputFormat].BlockSizeX == 4 && GPixelFormats[OutputFormat].BlockSizeY == 4;
-			const bool bLQFormat = OutputTexture0 != nullptr && OutputTexture0->GetFormat() == PF_R5G6B5_UNORM;
+			const bool bCompressedFormat = GPixelFormats[OutputFormat0].BlockSizeX == 4 && GPixelFormats[OutputFormat0].BlockSizeY == 4;
+			const bool bLQFormat = OutputFormat0 == PF_R5G6B5_UNORM;
 			bCompressPass = bRenderPass && !bCopyThumbnailPass && bCompressedFormat;
 			bCopyPass = bRenderPass && !bCopyThumbnailPass && !bCompressPass && (MaterialType == ERuntimeVirtualTextureMaterialType::BaseColor_Normal_Specular || MaterialType == ERuntimeVirtualTextureMaterialType::BaseColor_Normal_Specular_YCoCg || MaterialType == ERuntimeVirtualTextureMaterialType::BaseColor_Normal_Specular_Mask_YCoCg);
+			
+			// Use direct aliasing for compression pass on platforms that support it.
+			bDirectAliasing = bCompressedFormat && GRHISupportsUAVFormatAliasing && CVarVTDirectCompress.GetValueOnRenderThread() != 0;
+
 			// Not all mobile RHIs support sRGB texture views/aliasing, use only linear targets on mobile
 			const ETextureCreateFlags VT_SRGB = FeatureLevel > ERHIFeatureLevel::ES3_1 ? TexCreate_SRGB : TexCreate_None;
-			const EPixelFormat Compressed64BitFormat = PF_R16G16B16A16_UINT;
+			const EPixelFormat Compressed64BitFormat = UseRGBA16(GMaxRHIShaderPlatform) ? PF_R16G16B16A16_UINT : PF_R32G32_UINT;
 			const EPixelFormat Compressed128BitFormat = PF_R32G32B32A32_UINT;
 
 			switch (MaterialType)
@@ -463,7 +477,17 @@ namespace RuntimeVirtualTexture
 				}
 				if (bCompressPass)
 				{
-					OutputAlias0 = CompressTexture0 = GraphBuilder.CreateTexture(FRDGTextureDesc::Create2D(TextureSize / 4, Compressed64BitFormat, FClearValueBinding::None, TexCreate_UAV), TEXT("CompressTexture0"));
+					if (bDirectAliasing)
+					{
+						CompressTexture0 = GraphBuilder.RegisterExternalTexture(OutputTargets[0], ERDGTextureFlags::None);
+						CompressTextureUAV0_64bit = GraphBuilder.CreateUAV(FRDGTextureUAVDesc(CompressTexture0, 0, Compressed64BitFormat));
+						OutputAlias0 = nullptr;
+					}
+					else
+					{
+						OutputAlias0 = CompressTexture0 = GraphBuilder.CreateTexture(FRDGTextureDesc::Create2D(TextureSize / 4, Compressed64BitFormat, FClearValueBinding::None, TexCreate_UAV), TEXT("CompressTexture0"));
+						CompressTextureUAV0_64bit = GraphBuilder.CreateUAV(FRDGTextureUAVDesc(CompressTexture0));
+					}
 				}
 				if (bCopyThumbnailPass)
 				{
@@ -474,12 +498,25 @@ namespace RuntimeVirtualTexture
 				if (bRenderPass)
 				{
 					OutputAlias0 = RenderTexture0 = GraphBuilder.CreateTexture(FRDGTextureDesc::Create2D(TextureSize, bLQFormat ? PF_R5G6B5_UNORM : PF_B8G8R8A8, FClearValueBinding::Black, TexCreate_RenderTargetable | TexCreate_ShaderResource), TEXT("RenderTexture0"));
-					OutputAlias1 = RenderTexture1 = GraphBuilder.CreateTexture(FRDGTextureDesc::Create2D(TextureSize, bLQFormat ? PF_R5G6B5_UNORM : PF_B8G8R8A8, FClearValueBinding::Black, TexCreate_RenderTargetable | TexCreate_ShaderResource), TEXT("RenderTexture1"));
+					RenderTexture1 = GraphBuilder.CreateTexture(FRDGTextureDesc::Create2D(TextureSize, bLQFormat ? PF_R5G6B5_UNORM : PF_B8G8R8A8, FClearValueBinding::Black, TexCreate_RenderTargetable | TexCreate_ShaderResource), TEXT("RenderTexture1"));
 				}
 				if (bCompressPass)
 				{
-					OutputAlias0 = CompressTexture0 = GraphBuilder.CreateTexture(FRDGTextureDesc::Create2D(TextureSize / 4, Compressed64BitFormat, FClearValueBinding::None, TexCreate_UAV), TEXT("CompressTexture0"));
-					OutputAlias1 = CompressTexture1 = GraphBuilder.CreateTexture(FRDGTextureDesc::Create2D(TextureSize / 4, Compressed128BitFormat, FClearValueBinding::None, TexCreate_UAV), TEXT("CompressTexture1"));
+					if (bDirectAliasing)
+					{
+						CompressTexture0 = GraphBuilder.RegisterExternalTexture(OutputTargets[0], ERDGTextureFlags::None);
+						CompressTextureUAV0_64bit = GraphBuilder.CreateUAV(FRDGTextureUAVDesc(CompressTexture0, 0, Compressed64BitFormat));
+						CompressTexture1 = GraphBuilder.RegisterExternalTexture(OutputTargets[1], ERDGTextureFlags::None);
+						CompressTextureUAV1_128bit = GraphBuilder.CreateUAV(FRDGTextureUAVDesc(CompressTexture0, 0, Compressed128BitFormat));
+						OutputAlias0 = OutputAlias1 = nullptr;
+					}
+					else
+					{
+						OutputAlias0 = CompressTexture0 = GraphBuilder.CreateTexture(FRDGTextureDesc::Create2D(TextureSize / 4, Compressed64BitFormat, FClearValueBinding::None, TexCreate_UAV), TEXT("CompressTexture0"));
+						CompressTextureUAV0_64bit = GraphBuilder.CreateUAV(FRDGTextureUAVDesc(CompressTexture0));
+						OutputAlias1 = CompressTexture1 = GraphBuilder.CreateTexture(FRDGTextureDesc::Create2D(TextureSize / 4, Compressed128BitFormat, FClearValueBinding::None, TexCreate_UAV), TEXT("CompressTexture1"));
+						CompressTextureUAV1_128bit = GraphBuilder.CreateUAV(FRDGTextureUAVDesc(CompressTexture1));
+					}
 				}
 				if (bCopyThumbnailPass)
 				{
@@ -496,8 +533,21 @@ namespace RuntimeVirtualTexture
 				}
 				if (bCompressPass)
 				{
-					OutputAlias0 = CompressTexture0 = GraphBuilder.CreateTexture(FRDGTextureDesc::Create2D(TextureSize / 4, Compressed128BitFormat, FClearValueBinding::None, TexCreate_UAV), TEXT("CompressTexture0"));
-					OutputAlias1 = CompressTexture1 = GraphBuilder.CreateTexture(FRDGTextureDesc::Create2D(TextureSize / 4, Compressed128BitFormat, FClearValueBinding::None, TexCreate_UAV), TEXT("CompressTexture1"));
+					if (bDirectAliasing)
+					{
+						CompressTexture0 = GraphBuilder.RegisterExternalTexture(OutputTargets[0], ERDGTextureFlags::None);
+						CompressTextureUAV0_128bit = GraphBuilder.CreateUAV(FRDGTextureUAVDesc(CompressTexture0, 0, Compressed128BitFormat));
+						CompressTexture1 = GraphBuilder.RegisterExternalTexture(OutputTargets[1], ERDGTextureFlags::None);
+						CompressTextureUAV1_128bit = GraphBuilder.CreateUAV(FRDGTextureUAVDesc(CompressTexture1, 0, Compressed128BitFormat));
+						OutputAlias0 = OutputAlias1 = nullptr;
+					}
+					else
+					{
+						OutputAlias0 = CompressTexture0 = GraphBuilder.CreateTexture(FRDGTextureDesc::Create2D(TextureSize / 4, Compressed128BitFormat, FClearValueBinding::None, TexCreate_UAV), TEXT("CompressTexture0"));
+						CompressTextureUAV0_128bit = GraphBuilder.CreateUAV(FRDGTextureUAVDesc(CompressTexture0));
+						OutputAlias1 = CompressTexture1 = GraphBuilder.CreateTexture(FRDGTextureDesc::Create2D(TextureSize / 4, Compressed128BitFormat, FClearValueBinding::None, TexCreate_UAV), TEXT("CompressTexture1"));
+						CompressTextureUAV1_128bit = GraphBuilder.CreateUAV(FRDGTextureUAVDesc(CompressTexture1));
+					}
 				}
 				if (bCopyPass)
 				{
@@ -518,9 +568,26 @@ namespace RuntimeVirtualTexture
 				}
 				if (bCompressPass)
 				{
-					OutputAlias0 = CompressTexture0 = GraphBuilder.CreateTexture(FRDGTextureDesc::Create2D(TextureSize / 4, Compressed128BitFormat, FClearValueBinding::None, TexCreate_UAV), TEXT("CompressTexture0"));
-					OutputAlias1 = CompressTexture1 = GraphBuilder.CreateTexture(FRDGTextureDesc::Create2D(TextureSize / 4, Compressed128BitFormat, FClearValueBinding::None, TexCreate_UAV), TEXT("CompressTexture1"));
-					OutputAlias2 = CompressTexture2 = GraphBuilder.CreateTexture(FRDGTextureDesc::Create2D(TextureSize / 4, Compressed64BitFormat, FClearValueBinding::None, TexCreate_UAV), TEXT("CompressTexture2"));
+					if (bDirectAliasing)
+					{
+						CompressTexture0 = GraphBuilder.RegisterExternalTexture(OutputTargets[0], ERDGTextureFlags::None);
+						CompressTextureUAV0_128bit = GraphBuilder.CreateUAV(FRDGTextureUAVDesc(CompressTexture0, 0, Compressed128BitFormat));
+						CompressTexture1 = GraphBuilder.RegisterExternalTexture(OutputTargets[1], ERDGTextureFlags::None);
+						CompressTextureUAV1_128bit = GraphBuilder.CreateUAV(FRDGTextureUAVDesc(CompressTexture1, 0, Compressed128BitFormat));
+						CompressTexture2 = GraphBuilder.RegisterExternalTexture(OutputTargets[2], ERDGTextureFlags::None);
+						CompressTextureUAV2_64bit = GraphBuilder.CreateUAV(FRDGTextureUAVDesc(CompressTexture2, 0, Compressed64BitFormat));
+
+						OutputAlias0 = OutputAlias1 = OutputAlias2 = nullptr;
+					}
+					else
+					{
+						OutputAlias0 = CompressTexture0 = GraphBuilder.CreateTexture(FRDGTextureDesc::Create2D(TextureSize / 4, Compressed128BitFormat, FClearValueBinding::None, TexCreate_UAV), TEXT("CompressTexture0"));
+						CompressTextureUAV0_128bit = GraphBuilder.CreateUAV(FRDGTextureUAVDesc(CompressTexture0));
+						OutputAlias1 = CompressTexture1 = GraphBuilder.CreateTexture(FRDGTextureDesc::Create2D(TextureSize / 4, Compressed128BitFormat, FClearValueBinding::None, TexCreate_UAV), TEXT("CompressTexture1"));
+						CompressTextureUAV1_128bit = GraphBuilder.CreateUAV(FRDGTextureUAVDesc(CompressTexture1));
+						OutputAlias2 = CompressTexture2 = GraphBuilder.CreateTexture(FRDGTextureDesc::Create2D(TextureSize / 4, Compressed64BitFormat, FClearValueBinding::None, TexCreate_UAV), TEXT("CompressTexture2"));
+						CompressTextureUAV2_64bit = GraphBuilder.CreateUAV(FRDGTextureUAVDesc(CompressTexture2));
+					}
 				}
 				if (bCopyPass)
 				{
@@ -542,9 +609,26 @@ namespace RuntimeVirtualTexture
 				}
 				if (bCompressPass)
 				{
-					OutputAlias0 = CompressTexture0 = GraphBuilder.CreateTexture(FRDGTextureDesc::Create2D(TextureSize / 4, Compressed128BitFormat, FClearValueBinding::None, TexCreate_UAV), TEXT("CompressTexture0"));
-					OutputAlias1 = CompressTexture1 = GraphBuilder.CreateTexture(FRDGTextureDesc::Create2D(TextureSize / 4, Compressed128BitFormat, FClearValueBinding::None, TexCreate_UAV), TEXT("CompressTexture1"));
-					OutputAlias2 = CompressTexture2 = GraphBuilder.CreateTexture(FRDGTextureDesc::Create2D(TextureSize / 4, Compressed128BitFormat, FClearValueBinding::None, TexCreate_UAV), TEXT("CompressTexture2"));
+					if (bDirectAliasing)
+					{
+						CompressTexture0 = GraphBuilder.RegisterExternalTexture(OutputTargets[0], ERDGTextureFlags::None);
+						CompressTextureUAV0_128bit = GraphBuilder.CreateUAV(FRDGTextureUAVDesc(CompressTexture0, 0, Compressed128BitFormat));
+						CompressTexture1 = GraphBuilder.RegisterExternalTexture(OutputTargets[1], ERDGTextureFlags::None);
+						CompressTextureUAV1_128bit = GraphBuilder.CreateUAV(FRDGTextureUAVDesc(CompressTexture1, 0, Compressed128BitFormat));
+						CompressTexture2 = GraphBuilder.RegisterExternalTexture(OutputTargets[2], ERDGTextureFlags::None);
+						CompressTextureUAV2_128bit = GraphBuilder.CreateUAV(FRDGTextureUAVDesc(CompressTexture2, 0, Compressed128BitFormat));
+
+						OutputAlias0 = OutputAlias1 = OutputAlias2 = nullptr;
+					}
+					else
+					{
+						OutputAlias0 = CompressTexture0 = GraphBuilder.CreateTexture(FRDGTextureDesc::Create2D(TextureSize / 4, Compressed128BitFormat, FClearValueBinding::None, TexCreate_UAV), TEXT("CompressTexture0"));
+						CompressTextureUAV0_128bit = GraphBuilder.CreateUAV(FRDGTextureUAVDesc(CompressTexture0));
+						OutputAlias1 = CompressTexture1 = GraphBuilder.CreateTexture(FRDGTextureDesc::Create2D(TextureSize / 4, Compressed128BitFormat, FClearValueBinding::None, TexCreate_UAV), TEXT("CompressTexture1"));
+						CompressTextureUAV1_128bit = GraphBuilder.CreateUAV(FRDGTextureUAVDesc(CompressTexture1));
+						OutputAlias2 = CompressTexture2 = GraphBuilder.CreateTexture(FRDGTextureDesc::Create2D(TextureSize / 4, Compressed128BitFormat, FClearValueBinding::None, TexCreate_UAV), TEXT("CompressTexture2"));
+						CompressTextureUAV2_128bit = GraphBuilder.CreateUAV(FRDGTextureUAVDesc(CompressTexture2));
+					}
 				}
 				if (bCopyPass)
 				{
@@ -575,7 +659,8 @@ namespace RuntimeVirtualTexture
 		bool bCompressPass = false;
 		bool bCopyPass = false;
 		bool bCopyThumbnailPass = false;
-
+		bool bDirectAliasing = false;
+		
 		/** Render graph textures needed for this virtual texture layout. */
 		FRDGTextureRef RenderTexture0 = nullptr;
 		FRDGTextureRef RenderTexture1 = nullptr;
@@ -583,12 +668,17 @@ namespace RuntimeVirtualTexture
 		FRDGTextureRef CompressTexture0 = nullptr;
 		FRDGTextureRef CompressTexture1 = nullptr;
 		FRDGTextureRef CompressTexture2 = nullptr;
+		FRDGTextureUAVRef CompressTextureUAV0_64bit = nullptr;
+		FRDGTextureUAVRef CompressTextureUAV1_64bit = nullptr;
+		FRDGTextureUAVRef CompressTextureUAV2_64bit = nullptr;
+		FRDGTextureUAVRef CompressTextureUAV0_128bit = nullptr;
+		FRDGTextureUAVRef CompressTextureUAV1_128bit = nullptr;
+		FRDGTextureUAVRef CompressTextureUAV2_128bit = nullptr;
 		FRDGTextureRef CopyTexture0 = nullptr;
 		FRDGTextureRef CopyTexture1 = nullptr;
 		FRDGTextureRef CopyTexture2 = nullptr;
 
 		/** Aliases to one of the render/compress/copy textures. This is what we will Copy into the final physical texture. */
-		//todo[vt]: On platforms that support direct aliasing we can not set these and compress direct to the final destination
 		FRDGTextureRef OutputAlias0 = nullptr;
 		FRDGTextureRef OutputAlias1 = nullptr;
 		FRDGTextureRef OutputAlias2 = nullptr;
@@ -921,7 +1011,10 @@ namespace RuntimeVirtualTexture
 	{
 	public:
 		BEGIN_SHADER_PARAMETER_STRUCT(FParameters, )
-			SHADER_PARAMETER(FIntVector4, DestRect)
+			SHADER_PARAMETER(FUintVector4, SourceRect)
+			SHADER_PARAMETER(FUintVector2, DestPos0)
+			SHADER_PARAMETER(FUintVector2, DestPos1)
+			SHADER_PARAMETER(FUintVector2, DestPos2)
 			SHADER_PARAMETER_STRUCT_REF(FEtcParameters, EtcParameters)
 			SHADER_PARAMETER_RDG_TEXTURE(Texture2D<float4>, RenderTexture0)
 			SHADER_PARAMETER_SAMPLER(SamplerState, TextureSampler0)
@@ -929,15 +1022,19 @@ namespace RuntimeVirtualTexture
 			SHADER_PARAMETER_SAMPLER(SamplerState, TextureSampler1)
 			SHADER_PARAMETER_RDG_TEXTURE(Texture2D<float4>, RenderTexture2)
 			SHADER_PARAMETER_SAMPLER(SamplerState, TextureSampler2)
-			SHADER_PARAMETER_RDG_TEXTURE_UAV(RWTexture2D<uint4>, OutCompressTexture0)
-			SHADER_PARAMETER_RDG_TEXTURE_UAV(RWTexture2D<uint4>, OutCompressTexture1)
-			SHADER_PARAMETER_RDG_TEXTURE_UAV(RWTexture2D<uint4>, OutCompressTexture2)
+			SHADER_PARAMETER_RDG_TEXTURE_UAV(RWTexture2D<uint2>, OutCompressTexture0_64bit)
+			SHADER_PARAMETER_RDG_TEXTURE_UAV(RWTexture2D<uint2>, OutCompressTexture1_64bit)
+			SHADER_PARAMETER_RDG_TEXTURE_UAV(RWTexture2D<uint2>, OutCompressTexture2_64bit)
+			SHADER_PARAMETER_RDG_TEXTURE_UAV(RWTexture2D<uint4>, OutCompressTexture0_128bit)
+			SHADER_PARAMETER_RDG_TEXTURE_UAV(RWTexture2D<uint4>, OutCompressTexture1_128bit)
+			SHADER_PARAMETER_RDG_TEXTURE_UAV(RWTexture2D<uint4>, OutCompressTexture2_128bit)
 		END_SHADER_PARAMETER_STRUCT()
 
 		static void ModifyCompilationEnvironment(const FGlobalShaderPermutationParameters& Parameters, FShaderCompilerEnvironment& OutEnvironment)
 		{
 			FGlobalShader::ModifyCompilationEnvironment(Parameters, OutEnvironment);
 			OutEnvironment.SetDefine(TEXT("ETC_PROFILE"), UseEtcProfile(Parameters.Platform) ? 1 : 0);
+			OutEnvironment.SetDefine(TEXT("PACK_RG32_RGBA16"), UseRGBA16(Parameters.Platform) ? 1 : 0);
 		}
 
 		FShader_VirtualTextureCompress()
@@ -1179,7 +1276,7 @@ namespace RuntimeVirtualTexture
 		ERuntimeVirtualTextureMaterialType MaterialType,
 		bool bClearTextures,
 		bool bIsThumbnails,
-		FRHITexture2D* OutputTexture0,
+		FRHITexture2D* OutputTexture0,		// todo[vt]: Only use IPooledRenderTarget or FRDGTextureRef, not raw RHI textures.
 		IPooledRenderTarget* OutputTarget0,
 		FBox2D const& DestBox0,
 		FRHITexture2D* OutputTexture1,
@@ -1260,7 +1357,12 @@ namespace RuntimeVirtualTexture
 		ExternalAccessQueue.Submit(GraphBuilder);
 
 		// Build graph
-		FRenderGraphSetup GraphSetup(GraphBuilder, Scene->GetFeatureLevel(), MaterialType, OutputTexture0, TextureSize, bIsThumbnails);
+		TArray<TRefCountPtr<IPooledRenderTarget>> Targets;
+		Targets.Add(OutputTarget0);
+		Targets.Add(OutputTarget1);
+		Targets.Add(OutputTarget2);
+		EPixelFormat OutputFormat0 = OutputTexture0 != nullptr ? OutputTexture0->GetFormat() : PF_Unknown;
+		FRenderGraphSetup GraphSetup(GraphBuilder, Scene->GetFeatureLevel(), MaterialType, Targets, OutputFormat0, TextureSize, bIsThumbnails);
 
 		// Draw Pass
 		if (GraphSetup.bRenderPass)
@@ -1283,7 +1385,10 @@ namespace RuntimeVirtualTexture
 		if (GraphSetup.bCompressPass)
 		{
 			FShader_VirtualTextureCompress::FParameters* PassParameters = GraphBuilder.AllocParameters<FShader_VirtualTextureCompress::FParameters>();
-			PassParameters->DestRect = FIntVector4(0, 0, TextureSize.X, TextureSize.Y);
+			PassParameters->SourceRect = FUintVector4(0, 0, TextureSize.X, TextureSize.Y);
+			PassParameters->DestPos0 = GraphSetup.bDirectAliasing ? FUintVector2(DestBox0.Min.X / 4, DestBox0.Min.Y / 4) : FUintVector2(0, 0);
+			PassParameters->DestPos1 = GraphSetup.bDirectAliasing ? FUintVector2(DestBox1.Min.X / 4, DestBox1.Min.Y / 4) : FUintVector2(0, 0);
+			PassParameters->DestPos2 = GraphSetup.bDirectAliasing ? FUintVector2(DestBox2.Min.X / 4, DestBox2.Min.Y / 4) : FUintVector2(0, 0);
 			PassParameters->EtcParameters = GetEtcParametersUniformBufferRef();
 			PassParameters->RenderTexture0 = GraphSetup.RenderTexture0;
 			PassParameters->TextureSampler0 = TStaticSamplerState<SF_Point, AM_Clamp, AM_Clamp, AM_Clamp>::GetRHI();
@@ -1291,9 +1396,12 @@ namespace RuntimeVirtualTexture
 			PassParameters->TextureSampler1 = TStaticSamplerState<SF_Point, AM_Clamp, AM_Clamp, AM_Clamp>::GetRHI();
 			PassParameters->RenderTexture2 = GraphSetup.RenderTexture2;
 			PassParameters->TextureSampler2 = TStaticSamplerState<SF_Point, AM_Clamp, AM_Clamp, AM_Clamp>::GetRHI();
-			PassParameters->OutCompressTexture0 = GraphSetup.CompressTexture0 ? GraphBuilder.CreateUAV(FRDGTextureUAVDesc(GraphSetup.CompressTexture0)) : nullptr;
-			PassParameters->OutCompressTexture1 = GraphSetup.CompressTexture1 ? GraphBuilder.CreateUAV(FRDGTextureUAVDesc(GraphSetup.CompressTexture1)) : nullptr;
-			PassParameters->OutCompressTexture2 = GraphSetup.CompressTexture2 ? GraphBuilder.CreateUAV(FRDGTextureUAVDesc(GraphSetup.CompressTexture2)) : nullptr;
+			PassParameters->OutCompressTexture0_64bit = GraphSetup.CompressTextureUAV0_64bit;
+			PassParameters->OutCompressTexture1_64bit = GraphSetup.CompressTextureUAV1_64bit;
+			PassParameters->OutCompressTexture2_64bit = GraphSetup.CompressTextureUAV2_64bit;
+			PassParameters->OutCompressTexture0_128bit = GraphSetup.CompressTextureUAV0_128bit;
+			PassParameters->OutCompressTexture1_128bit = GraphSetup.CompressTextureUAV1_128bit;
+			PassParameters->OutCompressTexture2_128bit = GraphSetup.CompressTextureUAV2_128bit;
 
 			AddCompressPass(GraphBuilder, Scene->GetFeatureLevel(), PassParameters, TextureSize, MaterialType);
 		}
@@ -1402,13 +1510,6 @@ namespace RuntimeVirtualTexture
 		}
 	}
 
-	void RenderPages(FRHICommandListImmediate& RHICmdList, FRenderPageBatchDesc const& InDesc)
-	{
-		FRDGBuilder GraphBuilder(RHICmdList);
-		RenderPages(GraphBuilder, InDesc);
-		GraphBuilder.Execute();
-	}
-
 	uint32 GetRuntimeVirtualTextureSceneIndex_GameThread(class URuntimeVirtualTextureComponent* InComponent)
 	{
 		int32 SceneIndex = 0;
diff --git a/Engine/Source/Runtime/Renderer/Public/VT/RuntimeVirtualTextureRender.h b/Engine/Source/Runtime/Renderer/Public/VT/RuntimeVirtualTextureRender.h
--- a/Engine/Source/Runtime/Renderer/Public/VT/RuntimeVirtualTextureRender.h
+++ b/Engine/Source/Runtime/Renderer/Public/VT/RuntimeVirtualTextureRender.h
@@ -13,6 +13,7 @@ class FRDGTexture;
 class FRDGTextureUAV;
 class FScene;
 class FSceneInterface;
+struct IPooledRenderTarget;
 class URuntimeVirtualTextureComponent;
 
 namespace RuntimeVirtualTexture
@@ -38,7 +39,7 @@ namespace RuntimeVirtualTexture
 		UE_DEPRECATED(5.1, "UAV is deprecated. Register the pooled render target with RDG instead.")
 		FRHIUnorderedAccessView* UAV = nullptr;
 
-		struct IPooledRenderTarget* PooledRenderTarget = nullptr;
+		IPooledRenderTarget* PooledRenderTarget = nullptr;
 	};
 
 	/** A single page description. Multiple of these can be placed in a single FRenderPageBatchDesc batch description. */
@@ -94,7 +95,4 @@ namespace RuntimeVirtualTexture
 	 * Performs the required scene rendering setup such that it can be called from a render thread task.
 	 */
 	RENDERER_API void RenderPagesStandAlone(FRDGBuilder& GraphBuilder, FRenderPageBatchDesc const& InDesc);
-
-	UE_DEPRECATED(5.0, "This method has been refactored to use an FRDGBuilder instead.")
-	RENDERER_API void RenderPages(FRHICommandListImmediate& RHICmdList, FRenderPageBatchDesc const& InDesc);
 }
