Fixed: black screen with post-processing disabled

diff --git a/Engine/Source/Runtime/Renderer/Private/PostProcess/PostProcessTonemap.cpp b/Engine/Source/Runtime/Renderer/Private/PostProcess/PostProcessTonemap.cpp
--- a/Engine/Source/Runtime/Renderer/Private/PostProcess/PostProcessTonemap.cpp
+++ b/Engine/Source/Runtime/Renderer/Private/PostProcess/PostProcessTonemap.cpp
@@ -588,7 +588,7 @@ FScreenPassTexture AddTonemapPass(FRDGBuilder& GraphBuilder, const FViewInfo& Vi
 	FRHISamplerState* BilinearClampSampler = TStaticSamplerState<SF_Bilinear, AM_Clamp, AM_Clamp, AM_Clamp>::GetRHI();
 	FRHISamplerState* PointClampSampler = TStaticSamplerState<SF_Point, AM_Clamp, AM_Clamp, AM_Clamp>::GetRHI();
 
-	const float DefaultEyeExposure = bEyeAdaptation ? 0.0f : GetEyeAdaptationFixedExposure(View);
+	const float DefaultEyeExposure = bEyeAdaptation ? 1.0f : GetEyeAdaptationFixedExposure(View);
 
 	const float SharpenDiv6 = FMath::Clamp(CVarTonemapperSharpen.GetValueOnRenderThread(), 0.0f, 10.0f) / 6.0f;
 
