Dont use a cmask in the Downsampling pass since we will override all pixels and we dont clear. This will make sure we dont call fast clear eliminate on the RTV when its coming from the transient allocator and is potentially binded to a cmask.

diff --git a/Engine/Source/Runtime/Renderer/Private/PostProcess/PostProcessDownsample.cpp b/Engine/Source/Runtime/Renderer/Private/PostProcess/PostProcessDownsample.cpp
--- a/Engine/Source/Runtime/Renderer/Private/PostProcess/PostProcessDownsample.cpp
+++ b/Engine/Source/Runtime/Renderer/Private/PostProcess/PostProcessDownsample.cpp
@@ -112,7 +112,7 @@ FScreenPassTexture AddDownsamplePass(
 		Desc.Extent.X = FMath::Max(1, Desc.Extent.X);
 		Desc.Extent.Y = FMath::Max(1, Desc.Extent.Y);
 		Desc.Flags &= ~(TexCreate_RenderTargetable | TexCreate_UAV | TexCreate_Presentable);
-		Desc.Flags |= bIsComputePass ? TexCreate_UAV : TexCreate_RenderTargetable;
+		Desc.Flags |= bIsComputePass ? TexCreate_UAV : (TexCreate_RenderTargetable | TexCreate_NoFastClear);
 		Desc.Flags |= GFastVRamConfig.Downsample;
 		Desc.ClearValue = FClearValueBinding(FLinearColor(0, 0, 0, 0));
 
