When rendering to RVT with bClearTextures off, set TexCreate_NoFastClear on the render target.
This ensures that we don't pay the unwanted fast clear cost on transient resources.

diff --git a/Engine/Source/Runtime/Renderer/Private/VT/RuntimeVirtualTextureRender.cpp b/Engine/Source/Runtime/Renderer/Private/VT/RuntimeVirtualTextureRender.cpp
--- a/Engine/Source/Runtime/Renderer/Private/VT/RuntimeVirtualTextureRender.cpp
+++ b/Engine/Source/Runtime/Renderer/Private/VT/RuntimeVirtualTextureRender.cpp
@@ -419,236 +419,258 @@ namespace RuntimeVirtualTexture
 	{
 		static void SetupRenderTargetsInfo(ERuntimeVirtualTextureMaterialType MaterialType, ERHIFeatureLevel::Type FeatureLevel, bool bLQFormat, FGraphicsPipelineRenderTargetsInfo& RenderTargetsInfo)
 		{
-			const ETextureCreateFlags VT_SRGB = FeatureLevel > ERHIFeatureLevel::ES3_1 ? TexCreate_SRGB : TexCreate_None;
+			const ETextureCreateFlags RTCreateFlags = TexCreate_RenderTargetable | TexCreate_ShaderResource;
+			const ETextureCreateFlags RTSrgbFlags = FeatureLevel > ERHIFeatureLevel::ES3_1 ? TexCreate_SRGB : TexCreate_None;
 
 			switch (MaterialType)
 			{
 			case ERuntimeVirtualTextureMaterialType::BaseColor:
-				AddRenderTargetInfo(PF_B8G8R8A8, VT_SRGB | TexCreate_RenderTargetable | TexCreate_ShaderResource, RenderTargetsInfo);
+				AddRenderTargetInfo(PF_B8G8R8A8, RTCreateFlags | RTSrgbFlags, RenderTargetsInfo);
 				break;
 			case ERuntimeVirtualTextureMaterialType::BaseColor_Normal_Roughness:
-				AddRenderTargetInfo(bLQFormat ? PF_R5G6B5_UNORM : PF_B8G8R8A8, TexCreate_RenderTargetable | TexCreate_ShaderResource, RenderTargetsInfo);
-				AddRenderTargetInfo(bLQFormat ? PF_R5G6B5_UNORM : PF_B8G8R8A8, TexCreate_RenderTargetable | TexCreate_ShaderResource, RenderTargetsInfo);
+				AddRenderTargetInfo(bLQFormat ? PF_R5G6B5_UNORM : PF_B8G8R8A8, RTCreateFlags, RenderTargetsInfo);
+				AddRenderTargetInfo(bLQFormat ? PF_R5G6B5_UNORM : PF_B8G8R8A8, RTCreateFlags, RenderTargetsInfo);
 				break;
 			case ERuntimeVirtualTextureMaterialType::BaseColor_Normal_Specular:
-				AddRenderTargetInfo(PF_B8G8R8A8, VT_SRGB | TexCreate_RenderTargetable | TexCreate_ShaderResource, RenderTargetsInfo);
-				AddRenderTargetInfo(PF_B8G8R8A8, TexCreate_RenderTargetable | TexCreate_ShaderResource, RenderTargetsInfo);
-				AddRenderTargetInfo(PF_B8G8R8A8, TexCreate_RenderTargetable | TexCreate_ShaderResource, RenderTargetsInfo);
+				AddRenderTargetInfo(PF_B8G8R8A8, RTCreateFlags | RTSrgbFlags, RenderTargetsInfo);
+				AddRenderTargetInfo(PF_B8G8R8A8, RTCreateFlags, RenderTargetsInfo);
+				AddRenderTargetInfo(PF_B8G8R8A8, RTCreateFlags, RenderTargetsInfo);
 				break;
 			case ERuntimeVirtualTextureMaterialType::BaseColor_Normal_Specular_YCoCg:
-				AddRenderTargetInfo(PF_B8G8R8A8, VT_SRGB | TexCreate_RenderTargetable | TexCreate_ShaderResource, RenderTargetsInfo);
-				AddRenderTargetInfo(PF_B8G8R8A8, TexCreate_RenderTargetable | TexCreate_ShaderResource, RenderTargetsInfo);
-				AddRenderTargetInfo(PF_B8G8R8A8, TexCreate_RenderTargetable | TexCreate_ShaderResource, RenderTargetsInfo);
+				AddRenderTargetInfo(PF_B8G8R8A8, RTCreateFlags | RTSrgbFlags, RenderTargetsInfo);
+				AddRenderTargetInfo(PF_B8G8R8A8, RTCreateFlags, RenderTargetsInfo);
+				AddRenderTargetInfo(PF_B8G8R8A8, RTCreateFlags, RenderTargetsInfo);
 				break;
 			case ERuntimeVirtualTextureMaterialType::BaseColor_Normal_Specular_Mask_YCoCg:
-				AddRenderTargetInfo(PF_B8G8R8A8, VT_SRGB | TexCreate_RenderTargetable | TexCreate_ShaderResource, RenderTargetsInfo);
-				AddRenderTargetInfo(PF_B8G8R8A8, TexCreate_RenderTargetable | TexCreate_ShaderResource, RenderTargetsInfo);
-				AddRenderTargetInfo(PF_B8G8R8A8, TexCreate_RenderTargetable | TexCreate_ShaderResource, RenderTargetsInfo);
+				AddRenderTargetInfo(PF_B8G8R8A8, RTCreateFlags | RTSrgbFlags, RenderTargetsInfo);
+				AddRenderTargetInfo(PF_B8G8R8A8, RTCreateFlags, RenderTargetsInfo);
+				AddRenderTargetInfo(PF_B8G8R8A8, RTCreateFlags, RenderTargetsInfo);
 				break;
 			case ERuntimeVirtualTextureMaterialType::WorldHeight:
-				AddRenderTargetInfo(PF_G16, TexCreate_RenderTargetable | TexCreate_ShaderResource, RenderTargetsInfo);
+				AddRenderTargetInfo(PF_G16, RTCreateFlags, RenderTargetsInfo);
 				break;
 			}
 		}
 
-		FRenderGraphSetup(FRDGBuilder& GraphBuilder, ERHIFeatureLevel::Type FeatureLevel, ERuntimeVirtualTextureMaterialType MaterialType, TArrayView<TRefCountPtr<IPooledRenderTarget>> OutputTargets, EPixelFormat OutputFormat0, FIntPoint TextureSize, bool bIsThumbnails)
+		struct FInitDesc
 		{
-			bRenderPass = OutputFormat0 != PF_Unknown;
-			bCopyThumbnailPass = bRenderPass && bIsThumbnails;
-			const bool bCompressedFormat = GPixelFormats[OutputFormat0].BlockSizeX == 4 && GPixelFormats[OutputFormat0].BlockSizeY == 4;
-			const bool bLQFormat = OutputFormat0 == PF_R5G6B5_UNORM;
+			ERHIFeatureLevel::Type FeatureLevel;
+			ERuntimeVirtualTextureMaterialType MaterialType;
+			FIntPoint TextureSize;
+			EPixelFormat OutputFormat0 = PF_Unknown;
+			TArray<TRefCountPtr<IPooledRenderTarget>> OutputTargets;
+			bool bClearTextures = false;
+			bool bIsThumbnails = false;
+
+			FInitDesc(ERHIFeatureLevel::Type InFeatureLevel, ERuntimeVirtualTextureMaterialType InMaterialType, FIntPoint InTextureSize)
+				: FeatureLevel(InFeatureLevel)
+				, MaterialType(InMaterialType)
+				, TextureSize(InTextureSize)
+			{}
+		};
+
+		FRenderGraphSetup(FRDGBuilder& GraphBuilder, FInitDesc const& Desc)
+		{
+			bRenderPass = Desc.OutputFormat0 != PF_Unknown;
+			bCopyThumbnailPass = bRenderPass && Desc.bIsThumbnails;
+			const bool bCompressedFormat = GPixelFormats[Desc.OutputFormat0].BlockSizeX == 4 && GPixelFormats[Desc.OutputFormat0].BlockSizeY == 4;
+			const bool bLQFormat = Desc.OutputFormat0 == PF_R5G6B5_UNORM;
 			bCompressPass = bRenderPass && !bCopyThumbnailPass && bCompressedFormat;
-			bCopyPass = bRenderPass && !bCopyThumbnailPass && !bCompressPass && (MaterialType == ERuntimeVirtualTextureMaterialType::BaseColor_Normal_Specular || MaterialType == ERuntimeVirtualTextureMaterialType::BaseColor_Normal_Specular_YCoCg || MaterialType == ERuntimeVirtualTextureMaterialType::BaseColor_Normal_Specular_Mask_YCoCg);
+			bCopyPass = bRenderPass && !bCopyThumbnailPass && !bCompressPass && (Desc.MaterialType == ERuntimeVirtualTextureMaterialType::BaseColor_Normal_Specular || Desc.MaterialType == ERuntimeVirtualTextureMaterialType::BaseColor_Normal_Specular_YCoCg || Desc.MaterialType == ERuntimeVirtualTextureMaterialType::BaseColor_Normal_Specular_Mask_YCoCg);
 			
 			// Use direct aliasing for compression pass on platforms that support it.
 			bDirectAliasing = bCompressedFormat && GRHISupportsUAVFormatAliasing && CVarVTDirectCompress.GetValueOnRenderThread() != 0;
 
+			// Some RHI need explicit flag to avoid a fast clear, even when we don't use ERenderTargetLoadAction::EClear
+			const ETextureCreateFlags RTClearFlags = Desc.bClearTextures ? TexCreate_None : TexCreate_NoFastClear;
+			const ETextureCreateFlags RTCreateFlags = TexCreate_RenderTargetable | TexCreate_ShaderResource | RTClearFlags;
 			// Not all mobile RHIs support sRGB texture views/aliasing, use only linear targets on mobile
-			const ETextureCreateFlags VT_SRGB = FeatureLevel > ERHIFeatureLevel::ES3_1 ? TexCreate_SRGB : TexCreate_None;
+			const ETextureCreateFlags RTSrgbFlags = Desc.FeatureLevel > ERHIFeatureLevel::ES3_1 ? TexCreate_SRGB : TexCreate_None;
+
 			const EPixelFormat Compressed64BitFormat = UseRGBA16(GMaxRHIShaderPlatform) ? PF_R16G16B16A16_UINT : PF_R32G32_UINT;
 			const EPixelFormat Compressed128BitFormat = PF_R32G32B32A32_UINT;
 
-			switch (MaterialType)
+			switch (Desc.MaterialType)
 			{
 			case ERuntimeVirtualTextureMaterialType::BaseColor:
 				if (bRenderPass)
 				{
-					OutputAlias0 = RenderTexture0 = GraphBuilder.CreateTexture(FRDGTextureDesc::Create2D(TextureSize, PF_B8G8R8A8, FClearValueBinding::Black, VT_SRGB | TexCreate_RenderTargetable | TexCreate_ShaderResource), TEXT("RenderTexture0"));
+					OutputAlias0 = RenderTexture0 = GraphBuilder.CreateTexture(FRDGTextureDesc::Create2D(Desc.TextureSize, PF_B8G8R8A8, FClearValueBinding::Black, RTCreateFlags | RTSrgbFlags), TEXT("RenderTexture0"));
 				}
 				if (bCompressPass)
 				{
 					if (bDirectAliasing)
 					{
-						CompressTexture0 = GraphBuilder.RegisterExternalTexture(OutputTargets[0], ERDGTextureFlags::None);
+						CompressTexture0 = GraphBuilder.RegisterExternalTexture(Desc.OutputTargets[0], ERDGTextureFlags::None);
 						CompressTextureUAV0_64bit = GraphBuilder.CreateUAV(FRDGTextureUAVDesc(CompressTexture0, 0, Compressed64BitFormat));
 						OutputAlias0 = nullptr;
 					}
 					else
 					{
-						OutputAlias0 = CompressTexture0 = GraphBuilder.CreateTexture(FRDGTextureDesc::Create2D(TextureSize / 4, Compressed64BitFormat, FClearValueBinding::None, TexCreate_UAV), TEXT("CompressTexture0"));
+						OutputAlias0 = CompressTexture0 = GraphBuilder.CreateTexture(FRDGTextureDesc::Create2D(Desc.TextureSize / 4, Compressed64BitFormat, FClearValueBinding::None, TexCreate_UAV), TEXT("CompressTexture0"));
 						CompressTextureUAV0_64bit = GraphBuilder.CreateUAV(FRDGTextureUAVDesc(CompressTexture0));
 					}
 				}
 				if (bCopyThumbnailPass)
 				{
-					OutputAlias0 = CopyTexture0 = GraphBuilder.CreateTexture(FRDGTextureDesc::Create2D(TextureSize, PF_B8G8R8A8, FClearValueBinding::Black, VT_SRGB | TexCreate_RenderTargetable | TexCreate_ShaderResource), TEXT("CopyTexture0"));
+					OutputAlias0 = CopyTexture0 = GraphBuilder.CreateTexture(FRDGTextureDesc::Create2D(Desc.TextureSize, PF_B8G8R8A8, FClearValueBinding::Black, RTCreateFlags | RTSrgbFlags), TEXT("CopyTexture0"));
 				}
 				break;
 			case ERuntimeVirtualTextureMaterialType::BaseColor_Normal_Roughness:
 				if (bRenderPass)
 				{
-					OutputAlias0 = RenderTexture0 = GraphBuilder.CreateTexture(FRDGTextureDesc::Create2D(TextureSize, bLQFormat ? PF_R5G6B5_UNORM : PF_B8G8R8A8, FClearValueBinding::Black, TexCreate_RenderTargetable | TexCreate_ShaderResource), TEXT("RenderTexture0"));
-					RenderTexture1 = GraphBuilder.CreateTexture(FRDGTextureDesc::Create2D(TextureSize, bLQFormat ? PF_R5G6B5_UNORM : PF_B8G8R8A8, FClearValueBinding::Black, TexCreate_RenderTargetable | TexCreate_ShaderResource), TEXT("RenderTexture1"));
+					OutputAlias0 = RenderTexture0 = GraphBuilder.CreateTexture(FRDGTextureDesc::Create2D(Desc.TextureSize, bLQFormat ? PF_R5G6B5_UNORM : PF_B8G8R8A8, FClearValueBinding::Black, RTCreateFlags), TEXT("RenderTexture0"));
+					RenderTexture1 = GraphBuilder.CreateTexture(FRDGTextureDesc::Create2D(Desc.TextureSize, bLQFormat ? PF_R5G6B5_UNORM : PF_B8G8R8A8, FClearValueBinding::Black, RTCreateFlags), TEXT("RenderTexture1"));
 				}
 				if (bCompressPass)
 				{
 					if (bDirectAliasing)
 					{
-						CompressTexture0 = GraphBuilder.RegisterExternalTexture(OutputTargets[0], ERDGTextureFlags::None);
+						CompressTexture0 = GraphBuilder.RegisterExternalTexture(Desc.OutputTargets[0], ERDGTextureFlags::None);
 						CompressTextureUAV0_64bit = GraphBuilder.CreateUAV(FRDGTextureUAVDesc(CompressTexture0, 0, Compressed64BitFormat));
-						CompressTexture1 = GraphBuilder.RegisterExternalTexture(OutputTargets[1], ERDGTextureFlags::None);
+						CompressTexture1 = GraphBuilder.RegisterExternalTexture(Desc.OutputTargets[1], ERDGTextureFlags::None);
 						CompressTextureUAV1_128bit = GraphBuilder.CreateUAV(FRDGTextureUAVDesc(CompressTexture0, 0, Compressed128BitFormat));
 						OutputAlias0 = OutputAlias1 = nullptr;
 					}
 					else
 					{
-						OutputAlias0 = CompressTexture0 = GraphBuilder.CreateTexture(FRDGTextureDesc::Create2D(TextureSize / 4, Compressed64BitFormat, FClearValueBinding::None, TexCreate_UAV), TEXT("CompressTexture0"));
+						OutputAlias0 = CompressTexture0 = GraphBuilder.CreateTexture(FRDGTextureDesc::Create2D(Desc.TextureSize / 4, Compressed64BitFormat, FClearValueBinding::None, TexCreate_UAV), TEXT("CompressTexture0"));
 						CompressTextureUAV0_64bit = GraphBuilder.CreateUAV(FRDGTextureUAVDesc(CompressTexture0));
-						OutputAlias1 = CompressTexture1 = GraphBuilder.CreateTexture(FRDGTextureDesc::Create2D(TextureSize / 4, Compressed128BitFormat, FClearValueBinding::None, TexCreate_UAV), TEXT("CompressTexture1"));
+						OutputAlias1 = CompressTexture1 = GraphBuilder.CreateTexture(FRDGTextureDesc::Create2D(Desc.TextureSize / 4, Compressed128BitFormat, FClearValueBinding::None, TexCreate_UAV), TEXT("CompressTexture1"));
 						CompressTextureUAV1_128bit = GraphBuilder.CreateUAV(FRDGTextureUAVDesc(CompressTexture1));
 					}
 				}
 				if (bCopyThumbnailPass)
 				{
-					OutputAlias0 = CopyTexture0 = GraphBuilder.CreateTexture(FRDGTextureDesc::Create2D(TextureSize, bLQFormat ? PF_R5G6B5_UNORM : PF_B8G8R8A8, FClearValueBinding::Black, TexCreate_RenderTargetable | TexCreate_ShaderResource), TEXT("CopyTexture0"));
+					OutputAlias0 = CopyTexture0 = GraphBuilder.CreateTexture(FRDGTextureDesc::Create2D(Desc.TextureSize, bLQFormat ? PF_R5G6B5_UNORM : PF_B8G8R8A8, FClearValueBinding::Black, RTCreateFlags), TEXT("CopyTexture0"));
 					OutputAlias1 = nullptr;
 				}
 				break;
 			case ERuntimeVirtualTextureMaterialType::BaseColor_Normal_Specular:
 				if (bRenderPass)
 				{
-					OutputAlias0 = RenderTexture0 = GraphBuilder.CreateTexture(FRDGTextureDesc::Create2D(TextureSize, PF_B8G8R8A8, FClearValueBinding::Black, VT_SRGB | TexCreate_RenderTargetable | TexCreate_ShaderResource), TEXT("RenderTexture0"));
-					RenderTexture1 = GraphBuilder.CreateTexture(FRDGTextureDesc::Create2D(TextureSize, PF_B8G8R8A8, FClearValueBinding::DefaultNormal8Bit, TexCreate_RenderTargetable | TexCreate_ShaderResource), TEXT("RenderTexture1"));
-					RenderTexture2 = GraphBuilder.CreateTexture(FRDGTextureDesc::Create2D(TextureSize, PF_B8G8R8A8, FClearValueBinding::DefaultNormal8Bit, TexCreate_RenderTargetable | TexCreate_ShaderResource), TEXT("RenderTexture2"));
+					OutputAlias0 = RenderTexture0 = GraphBuilder.CreateTexture(FRDGTextureDesc::Create2D(Desc.TextureSize, PF_B8G8R8A8, FClearValueBinding::Black, RTCreateFlags | RTSrgbFlags), TEXT("RenderTexture0"));
+					RenderTexture1 = GraphBuilder.CreateTexture(FRDGTextureDesc::Create2D(Desc.TextureSize, PF_B8G8R8A8, FClearValueBinding::DefaultNormal8Bit, RTCreateFlags), TEXT("RenderTexture1"));
+					RenderTexture2 = GraphBuilder.CreateTexture(FRDGTextureDesc::Create2D(Desc.TextureSize, PF_B8G8R8A8, FClearValueBinding::DefaultNormal8Bit, RTCreateFlags), TEXT("RenderTexture2"));
 				}
 				if (bCompressPass)
 				{
 					if (bDirectAliasing)
 					{
-						CompressTexture0 = GraphBuilder.RegisterExternalTexture(OutputTargets[0], ERDGTextureFlags::None);
+						CompressTexture0 = GraphBuilder.RegisterExternalTexture(Desc.OutputTargets[0], ERDGTextureFlags::None);
 						CompressTextureUAV0_128bit = GraphBuilder.CreateUAV(FRDGTextureUAVDesc(CompressTexture0, 0, Compressed128BitFormat));
-						CompressTexture1 = GraphBuilder.RegisterExternalTexture(OutputTargets[1], ERDGTextureFlags::None);
+						CompressTexture1 = GraphBuilder.RegisterExternalTexture(Desc.OutputTargets[1], ERDGTextureFlags::None);
 						CompressTextureUAV1_128bit = GraphBuilder.CreateUAV(FRDGTextureUAVDesc(CompressTexture1, 0, Compressed128BitFormat));
 						OutputAlias0 = OutputAlias1 = nullptr;
 					}
 					else
 					{
-						OutputAlias0 = CompressTexture0 = GraphBuilder.CreateTexture(FRDGTextureDesc::Create2D(TextureSize / 4, Compressed128BitFormat, FClearValueBinding::None, TexCreate_UAV), TEXT("CompressTexture0"));
+						OutputAlias0 = CompressTexture0 = GraphBuilder.CreateTexture(FRDGTextureDesc::Create2D(Desc.TextureSize / 4, Compressed128BitFormat, FClearValueBinding::None, TexCreate_UAV), TEXT("CompressTexture0"));
 						CompressTextureUAV0_128bit = GraphBuilder.CreateUAV(FRDGTextureUAVDesc(CompressTexture0));
-						OutputAlias1 = CompressTexture1 = GraphBuilder.CreateTexture(FRDGTextureDesc::Create2D(TextureSize / 4, Compressed128BitFormat, FClearValueBinding::None, TexCreate_UAV), TEXT("CompressTexture1"));
+						OutputAlias1 = CompressTexture1 = GraphBuilder.CreateTexture(FRDGTextureDesc::Create2D(Desc.TextureSize / 4, Compressed128BitFormat, FClearValueBinding::None, TexCreate_UAV), TEXT("CompressTexture1"));
 						CompressTextureUAV1_128bit = GraphBuilder.CreateUAV(FRDGTextureUAVDesc(CompressTexture1));
 					}
 				}
 				if (bCopyPass)
 				{
-					OutputAlias0 = CopyTexture0 = GraphBuilder.CreateTexture(FRDGTextureDesc::Create2D(TextureSize, PF_B8G8R8A8, FClearValueBinding::Black, VT_SRGB | TexCreate_RenderTargetable | TexCreate_ShaderResource), TEXT("CopyTexture0"));
-					OutputAlias1 = CopyTexture1 = GraphBuilder.CreateTexture(FRDGTextureDesc::Create2D(TextureSize, PF_B8G8R8A8, FClearValueBinding::Black, TexCreate_RenderTargetable | TexCreate_ShaderResource), TEXT("CopyTexture1"));
+					OutputAlias0 = CopyTexture0 = GraphBuilder.CreateTexture(FRDGTextureDesc::Create2D(Desc.TextureSize, PF_B8G8R8A8, FClearValueBinding::Black, RTCreateFlags | RTSrgbFlags), TEXT("CopyTexture0"));
+					OutputAlias1 = CopyTexture1 = GraphBuilder.CreateTexture(FRDGTextureDesc::Create2D(Desc.TextureSize, PF_B8G8R8A8, FClearValueBinding::Black, RTCreateFlags), TEXT("CopyTexture1"));
 				}
 				if (bCopyThumbnailPass)
 				{
-					OutputAlias0 = CopyTexture0 = GraphBuilder.CreateTexture(FRDGTextureDesc::Create2D(TextureSize, PF_B8G8R8A8, FClearValueBinding::Black, VT_SRGB | TexCreate_RenderTargetable | TexCreate_ShaderResource), TEXT("CopyTexture0"));
+					OutputAlias0 = CopyTexture0 = GraphBuilder.CreateTexture(FRDGTextureDesc::Create2D(Desc.TextureSize, PF_B8G8R8A8, FClearValueBinding::Black, RTCreateFlags | RTSrgbFlags), TEXT("CopyTexture0"));
 				}
 				break;
 			case ERuntimeVirtualTextureMaterialType::BaseColor_Normal_Specular_YCoCg:
 				if (bRenderPass)
 				{
-					OutputAlias0 = RenderTexture0 = GraphBuilder.CreateTexture(FRDGTextureDesc::Create2D(TextureSize, PF_B8G8R8A8, FClearValueBinding::Black, VT_SRGB | TexCreate_RenderTargetable | TexCreate_ShaderResource), TEXT("RenderTexture0"));
-					RenderTexture1 = GraphBuilder.CreateTexture(FRDGTextureDesc::Create2D(TextureSize, PF_B8G8R8A8, FClearValueBinding::DefaultNormal8Bit, TexCreate_RenderTargetable | TexCreate_ShaderResource), TEXT("RenderTexture1"));
-					RenderTexture2 = GraphBuilder.CreateTexture(FRDGTextureDesc::Create2D(TextureSize, PF_B8G8R8A8, FClearValueBinding::DefaultNormal8Bit, TexCreate_RenderTargetable | TexCreate_ShaderResource), TEXT("RenderTexture2"));
+					OutputAlias0 = RenderTexture0 = GraphBuilder.CreateTexture(FRDGTextureDesc::Create2D(Desc.TextureSize, PF_B8G8R8A8, FClearValueBinding::Black, RTCreateFlags | RTSrgbFlags), TEXT("RenderTexture0"));
+					RenderTexture1 = GraphBuilder.CreateTexture(FRDGTextureDesc::Create2D(Desc.TextureSize, PF_B8G8R8A8, FClearValueBinding::DefaultNormal8Bit, RTCreateFlags), TEXT("RenderTexture1"));
+					RenderTexture2 = GraphBuilder.CreateTexture(FRDGTextureDesc::Create2D(Desc.TextureSize, PF_B8G8R8A8, FClearValueBinding::DefaultNormal8Bit, RTCreateFlags), TEXT("RenderTexture2"));
 				}
 				if (bCompressPass)
 				{
 					if (bDirectAliasing)
 					{
-						CompressTexture0 = GraphBuilder.RegisterExternalTexture(OutputTargets[0], ERDGTextureFlags::None);
+						CompressTexture0 = GraphBuilder.RegisterExternalTexture(Desc.OutputTargets[0], ERDGTextureFlags::None);
 						CompressTextureUAV0_128bit = GraphBuilder.CreateUAV(FRDGTextureUAVDesc(CompressTexture0, 0, Compressed128BitFormat));
-						CompressTexture1 = GraphBuilder.RegisterExternalTexture(OutputTargets[1], ERDGTextureFlags::None);
+						CompressTexture1 = GraphBuilder.RegisterExternalTexture(Desc.OutputTargets[1], ERDGTextureFlags::None);
 						CompressTextureUAV1_128bit = GraphBuilder.CreateUAV(FRDGTextureUAVDesc(CompressTexture1, 0, Compressed128BitFormat));
-						CompressTexture2 = GraphBuilder.RegisterExternalTexture(OutputTargets[2], ERDGTextureFlags::None);
+						CompressTexture2 = GraphBuilder.RegisterExternalTexture(Desc.OutputTargets[2], ERDGTextureFlags::None);
 						CompressTextureUAV2_64bit = GraphBuilder.CreateUAV(FRDGTextureUAVDesc(CompressTexture2, 0, Compressed64BitFormat));
 
 						OutputAlias0 = OutputAlias1 = OutputAlias2 = nullptr;
 					}
 					else
 					{
-						OutputAlias0 = CompressTexture0 = GraphBuilder.CreateTexture(FRDGTextureDesc::Create2D(TextureSize / 4, Compressed128BitFormat, FClearValueBinding::None, TexCreate_UAV), TEXT("CompressTexture0"));
+						OutputAlias0 = CompressTexture0 = GraphBuilder.CreateTexture(FRDGTextureDesc::Create2D(Desc.TextureSize / 4, Compressed128BitFormat, FClearValueBinding::None, TexCreate_UAV), TEXT("CompressTexture0"));
 						CompressTextureUAV0_128bit = GraphBuilder.CreateUAV(FRDGTextureUAVDesc(CompressTexture0));
-						OutputAlias1 = CompressTexture1 = GraphBuilder.CreateTexture(FRDGTextureDesc::Create2D(TextureSize / 4, Compressed128BitFormat, FClearValueBinding::None, TexCreate_UAV), TEXT("CompressTexture1"));
+						OutputAlias1 = CompressTexture1 = GraphBuilder.CreateTexture(FRDGTextureDesc::Create2D(Desc.TextureSize / 4, Compressed128BitFormat, FClearValueBinding::None, TexCreate_UAV), TEXT("CompressTexture1"));
 						CompressTextureUAV1_128bit = GraphBuilder.CreateUAV(FRDGTextureUAVDesc(CompressTexture1));
-						OutputAlias2 = CompressTexture2 = GraphBuilder.CreateTexture(FRDGTextureDesc::Create2D(TextureSize / 4, Compressed64BitFormat, FClearValueBinding::None, TexCreate_UAV), TEXT("CompressTexture2"));
+						OutputAlias2 = CompressTexture2 = GraphBuilder.CreateTexture(FRDGTextureDesc::Create2D(Desc.TextureSize / 4, Compressed64BitFormat, FClearValueBinding::None, TexCreate_UAV), TEXT("CompressTexture2"));
 						CompressTextureUAV2_64bit = GraphBuilder.CreateUAV(FRDGTextureUAVDesc(CompressTexture2));
 					}
 				}
 				if (bCopyPass)
 				{
-					OutputAlias0 = CopyTexture0 = GraphBuilder.CreateTexture(FRDGTextureDesc::Create2D(TextureSize, PF_B8G8R8A8, FClearValueBinding::Black, TexCreate_RenderTargetable | TexCreate_ShaderResource), TEXT("CopyTexture0"));
-					OutputAlias1 = CopyTexture1 = GraphBuilder.CreateTexture(FRDGTextureDesc::Create2D(TextureSize, PF_B8G8R8A8, FClearValueBinding::Black, TexCreate_RenderTargetable | TexCreate_ShaderResource), TEXT("CopyTexture1"));
-					OutputAlias2 = CopyTexture2 = GraphBuilder.CreateTexture(FRDGTextureDesc::Create2D(TextureSize, PF_B8G8R8A8, FClearValueBinding::Black, TexCreate_RenderTargetable | TexCreate_ShaderResource), TEXT("CopyTexture2"));
+					OutputAlias0 = CopyTexture0 = GraphBuilder.CreateTexture(FRDGTextureDesc::Create2D(Desc.TextureSize, PF_B8G8R8A8, FClearValueBinding::Black, RTCreateFlags), TEXT("CopyTexture0"));
+					OutputAlias1 = CopyTexture1 = GraphBuilder.CreateTexture(FRDGTextureDesc::Create2D(Desc.TextureSize, PF_B8G8R8A8, FClearValueBinding::Black, RTCreateFlags), TEXT("CopyTexture1"));
+					OutputAlias2 = CopyTexture2 = GraphBuilder.CreateTexture(FRDGTextureDesc::Create2D(Desc.TextureSize, PF_B8G8R8A8, FClearValueBinding::Black, RTCreateFlags), TEXT("CopyTexture2"));
 				}
 				if (bCopyThumbnailPass)
 				{
-					OutputAlias0 = CopyTexture0 = GraphBuilder.CreateTexture(FRDGTextureDesc::Create2D(TextureSize, PF_B8G8R8A8, FClearValueBinding::Black, VT_SRGB | TexCreate_RenderTargetable | TexCreate_ShaderResource), TEXT("CopyTexture0"));
+					OutputAlias0 = CopyTexture0 = GraphBuilder.CreateTexture(FRDGTextureDesc::Create2D(Desc.TextureSize, PF_B8G8R8A8, FClearValueBinding::Black, RTCreateFlags | RTSrgbFlags), TEXT("CopyTexture0"));
 				}
 				break;
 			case ERuntimeVirtualTextureMaterialType::BaseColor_Normal_Specular_Mask_YCoCg:
 				if (bRenderPass)
 				{
-					OutputAlias0 = RenderTexture0 = GraphBuilder.CreateTexture(FRDGTextureDesc::Create2D(TextureSize, PF_B8G8R8A8, FClearValueBinding::Black, VT_SRGB | TexCreate_RenderTargetable | TexCreate_ShaderResource), TEXT("RenderTexture0"));
-					RenderTexture1 = GraphBuilder.CreateTexture(FRDGTextureDesc::Create2D(TextureSize, PF_B8G8R8A8, FClearValueBinding::DefaultNormal8Bit, TexCreate_RenderTargetable | TexCreate_ShaderResource), TEXT("RenderTexture1"));
-					RenderTexture2 = GraphBuilder.CreateTexture(FRDGTextureDesc::Create2D(TextureSize, PF_B8G8R8A8, FClearValueBinding::DefaultNormal8Bit, TexCreate_RenderTargetable | TexCreate_ShaderResource), TEXT("RenderTexture2"));
+					OutputAlias0 = RenderTexture0 = GraphBuilder.CreateTexture(FRDGTextureDesc::Create2D(Desc.TextureSize, PF_B8G8R8A8, FClearValueBinding::Black, RTCreateFlags | RTSrgbFlags), TEXT("RenderTexture0"));
+					RenderTexture1 = GraphBuilder.CreateTexture(FRDGTextureDesc::Create2D(Desc.TextureSize, PF_B8G8R8A8, FClearValueBinding::DefaultNormal8Bit, RTCreateFlags), TEXT("RenderTexture1"));
+					RenderTexture2 = GraphBuilder.CreateTexture(FRDGTextureDesc::Create2D(Desc.TextureSize, PF_B8G8R8A8, FClearValueBinding::DefaultNormal8Bit, RTCreateFlags), TEXT("RenderTexture2"));
 				}
 				if (bCompressPass)
 				{
 					if (bDirectAliasing)
 					{
-						CompressTexture0 = GraphBuilder.RegisterExternalTexture(OutputTargets[0], ERDGTextureFlags::None);
+						CompressTexture0 = GraphBuilder.RegisterExternalTexture(Desc.OutputTargets[0], ERDGTextureFlags::None);
 						CompressTextureUAV0_128bit = GraphBuilder.CreateUAV(FRDGTextureUAVDesc(CompressTexture0, 0, Compressed128BitFormat));
-						CompressTexture1 = GraphBuilder.RegisterExternalTexture(OutputTargets[1], ERDGTextureFlags::None);
+						CompressTexture1 = GraphBuilder.RegisterExternalTexture(Desc.OutputTargets[1], ERDGTextureFlags::None);
 						CompressTextureUAV1_128bit = GraphBuilder.CreateUAV(FRDGTextureUAVDesc(CompressTexture1, 0, Compressed128BitFormat));
-						CompressTexture2 = GraphBuilder.RegisterExternalTexture(OutputTargets[2], ERDGTextureFlags::None);
+						CompressTexture2 = GraphBuilder.RegisterExternalTexture(Desc.OutputTargets[2], ERDGTextureFlags::None);
 						CompressTextureUAV2_128bit = GraphBuilder.CreateUAV(FRDGTextureUAVDesc(CompressTexture2, 0, Compressed128BitFormat));
 
 						OutputAlias0 = OutputAlias1 = OutputAlias2 = nullptr;
 					}
 					else
 					{
-						OutputAlias0 = CompressTexture0 = GraphBuilder.CreateTexture(FRDGTextureDesc::Create2D(TextureSize / 4, Compressed128BitFormat, FClearValueBinding::None, TexCreate_UAV), TEXT("CompressTexture0"));
+						OutputAlias0 = CompressTexture0 = GraphBuilder.CreateTexture(FRDGTextureDesc::Create2D(Desc.TextureSize / 4, Compressed128BitFormat, FClearValueBinding::None, TexCreate_UAV), TEXT("CompressTexture0"));
 						CompressTextureUAV0_128bit = GraphBuilder.CreateUAV(FRDGTextureUAVDesc(CompressTexture0));
-						OutputAlias1 = CompressTexture1 = GraphBuilder.CreateTexture(FRDGTextureDesc::Create2D(TextureSize / 4, Compressed128BitFormat, FClearValueBinding::None, TexCreate_UAV), TEXT("CompressTexture1"));
+						OutputAlias1 = CompressTexture1 = GraphBuilder.CreateTexture(FRDGTextureDesc::Create2D(Desc.TextureSize / 4, Compressed128BitFormat, FClearValueBinding::None, TexCreate_UAV), TEXT("CompressTexture1"));
 						CompressTextureUAV1_128bit = GraphBuilder.CreateUAV(FRDGTextureUAVDesc(CompressTexture1));
-						OutputAlias2 = CompressTexture2 = GraphBuilder.CreateTexture(FRDGTextureDesc::Create2D(TextureSize / 4, Compressed128BitFormat, FClearValueBinding::None, TexCreate_UAV), TEXT("CompressTexture2"));
+						OutputAlias2 = CompressTexture2 = GraphBuilder.CreateTexture(FRDGTextureDesc::Create2D(Desc.TextureSize / 4, Compressed128BitFormat, FClearValueBinding::None, TexCreate_UAV), TEXT("CompressTexture2"));
 						CompressTextureUAV2_128bit = GraphBuilder.CreateUAV(FRDGTextureUAVDesc(CompressTexture2));
 					}
 				}
 				if (bCopyPass)
 				{
-					OutputAlias0 = CopyTexture0 = GraphBuilder.CreateTexture(FRDGTextureDesc::Create2D(TextureSize, PF_B8G8R8A8, FClearValueBinding::Black, TexCreate_RenderTargetable | TexCreate_ShaderResource), TEXT("CopyTexture0"));
-					OutputAlias1 = CopyTexture1 = GraphBuilder.CreateTexture(FRDGTextureDesc::Create2D(TextureSize, PF_B8G8R8A8, FClearValueBinding::Black, TexCreate_RenderTargetable | TexCreate_ShaderResource), TEXT("CopyTexture1"));
-					OutputAlias2 = CopyTexture2 = GraphBuilder.CreateTexture(FRDGTextureDesc::Create2D(TextureSize, PF_B8G8R8A8, FClearValueBinding::Black, TexCreate_RenderTargetable | TexCreate_ShaderResource), TEXT("CopyTexture2"));
+					OutputAlias0 = CopyTexture0 = GraphBuilder.CreateTexture(FRDGTextureDesc::Create2D(Desc.TextureSize, PF_B8G8R8A8, FClearValueBinding::Black, RTCreateFlags), TEXT("CopyTexture0"));
+					OutputAlias1 = CopyTexture1 = GraphBuilder.CreateTexture(FRDGTextureDesc::Create2D(Desc.TextureSize, PF_B8G8R8A8, FClearValueBinding::Black, RTCreateFlags), TEXT("CopyTexture1"));
+					OutputAlias2 = CopyTexture2 = GraphBuilder.CreateTexture(FRDGTextureDesc::Create2D(Desc.TextureSize, PF_B8G8R8A8, FClearValueBinding::Black, RTCreateFlags), TEXT("CopyTexture2"));
 				}
 				if (bCopyThumbnailPass)
 				{
-					OutputAlias0 = CopyTexture0 = GraphBuilder.CreateTexture(FRDGTextureDesc::Create2D(TextureSize, PF_B8G8R8A8, FClearValueBinding::Black, VT_SRGB | TexCreate_RenderTargetable | TexCreate_ShaderResource), TEXT("CopyTexture0"));
+					OutputAlias0 = CopyTexture0 = GraphBuilder.CreateTexture(FRDGTextureDesc::Create2D(Desc.TextureSize, PF_B8G8R8A8, FClearValueBinding::Black, RTCreateFlags | RTSrgbFlags), TEXT("CopyTexture0"));
 				}
 				break;
 			case ERuntimeVirtualTextureMaterialType::WorldHeight:
 				if (bRenderPass)
 				{
-					OutputAlias0 = RenderTexture0 = GraphBuilder.CreateTexture(FRDGTextureDesc::Create2D(TextureSize, PF_G16, FClearValueBinding::Black, TexCreate_RenderTargetable | TexCreate_ShaderResource), TEXT("RenderTexture0"));
+					OutputAlias0 = RenderTexture0 = GraphBuilder.CreateTexture(FRDGTextureDesc::Create2D(Desc.TextureSize, PF_G16, FClearValueBinding::Black, RTCreateFlags), TEXT("RenderTexture0"));
 				}
 				if (bCopyThumbnailPass)
 				{
-					OutputAlias0 = CopyTexture0 = GraphBuilder.CreateTexture(FRDGTextureDesc::Create2D(TextureSize, PF_B8G8R8A8, FClearValueBinding::Black, TexCreate_RenderTargetable | TexCreate_ShaderResource), TEXT("CopyTexture0"));
+					OutputAlias0 = CopyTexture0 = GraphBuilder.CreateTexture(FRDGTextureDesc::Create2D(Desc.TextureSize, PF_B8G8R8A8, FClearValueBinding::Black, RTCreateFlags), TEXT("CopyTexture0"));
 				}
 				break;
 			}
@@ -1357,12 +1379,14 @@ namespace RuntimeVirtualTexture
 		ExternalAccessQueue.Submit(GraphBuilder);
 
 		// Build graph
-		TArray<TRefCountPtr<IPooledRenderTarget>> Targets;
-		Targets.Add(OutputTarget0);
-		Targets.Add(OutputTarget1);
-		Targets.Add(OutputTarget2);
-		EPixelFormat OutputFormat0 = OutputTexture0 != nullptr ? OutputTexture0->GetFormat() : PF_Unknown;
-		FRenderGraphSetup GraphSetup(GraphBuilder, Scene->GetFeatureLevel(), MaterialType, Targets, OutputFormat0, TextureSize, bIsThumbnails);
+		FRenderGraphSetup::FInitDesc GraphSetupDesc(Scene->GetFeatureLevel(), MaterialType, TextureSize);
+		GraphSetupDesc.OutputFormat0 = OutputTexture0 != nullptr ? OutputTexture0->GetFormat() : PF_Unknown;
+		GraphSetupDesc.OutputTargets.Add(OutputTarget0);
+		GraphSetupDesc.OutputTargets.Add(OutputTarget1);
+		GraphSetupDesc.OutputTargets.Add(OutputTarget2);
+		GraphSetupDesc.bClearTextures = bClearTextures;
+		GraphSetupDesc.bIsThumbnails = bIsThumbnails;
+		FRenderGraphSetup GraphSetup(GraphBuilder, GraphSetupDesc);
 
 		// Draw Pass
 		if (GraphSetup.bRenderPass)
