Support calculating AutoExposure Illuminance at lower resolution.

diff --git a/Engine/Shaders/Private/PostProcessEyeAdaptation.usf b/Engine/Shaders/Private/PostProcessEyeAdaptation.usf
--- a/Engine/Shaders/Private/PostProcessEyeAdaptation.usf
+++ b/Engine/Shaders/Private/PostProcessEyeAdaptation.usf
@@ -11,6 +11,7 @@
 #include "DeferredShadingCommon.ush"
 #include "BRDF.ush"
 #include "PositionReconstructionCommon.ush"
+#include "/Engine/Private/ScreenPass.ush"
 
 Texture2D EyeAdaptationTexture;
 
@@ -270,6 +271,8 @@ void BasicEyeAdaptationPS(in float4 UVAndScreenPos : TEXCOORD0, out float4 OutCo
 #if COMPUTESHADER
 
 RWTexture2D<float> RWIlluminanceTexture;
+SCREEN_PASS_TEXTURE_VIEWPORT(Illuminance)
+uint IllumiananceDownscaleFactor;
 
 // same weights as CalculateEyeAdaptationLuminance
 // but skip the min since we apply EyeAdaptation_IgnoreMaterialsMinBaseColorLuminance here
@@ -281,25 +284,19 @@ float CalculateEyeAdaptationLuminanceWithoutMin(float3 Color)
 [numthreads(8, 8, 1)]
 void SetupExposureIlluminance(uint2 ThreadId : SV_DispatchThreadID)
 {
-	const uint2 TexelPos = ThreadId + View.ViewRectMinAndSize.xy;
+	const uint2 OutputTexelPos = ThreadId + Illuminance_ViewportMin;
+	const uint2 TexelPos = ThreadId * IllumiananceDownscaleFactor + View.ViewRectMinAndSize.xy;
 
-	float3 Emissive = 0.0f;
+	const float3 Emissive = ColorTexture.Load(int3(TexelPos, 0)).rgb;
 
-	const float2 GBufferUV = (TexelPos + 0.5f) * View.BufferSizeAndInvSize.zw;
-	const FGBufferData GBufferData = GetGBufferDataFromSceneTextures(GBufferUV);
-
-	if (GBufferData.ShadingModelID != SHADINGMODELID_UNLIT)
-	{
-		Emissive = ColorTexture.Load(int3(TexelPos, 0)).rgb;
-	}
-
-	RWIlluminanceTexture[TexelPos] = CalculateEyeAdaptationLuminanceWithoutMin(Emissive);
+	RWIlluminanceTexture[OutputTexelPos] = CalculateEyeAdaptationLuminanceWithoutMin(Emissive);
 }
 
 [numthreads(8, 8, 1)]
 void CalculateExposureIlluminance(uint2 ThreadId : SV_DispatchThreadID)
 {
-	const uint2 TexelPos = ThreadId + View.ViewRectMinAndSize.xy;
+	const uint2 OutputTexelPos = ThreadId + Illuminance_ViewportMin;
+	const uint2 TexelPos = ThreadId * IllumiananceDownscaleFactor + View.ViewRectMinAndSize.xy;
 
 	const float3 SceneColor = ColorTexture.Load(int3(TexelPos, 0)).rgb * View.OneOverPreExposure;
 
@@ -310,7 +307,7 @@ void CalculateExposureIlluminance(uint2 ThreadId : SV_DispatchThreadID)
 
 	if (GBufferData.ShadingModelID != SHADINGMODELID_UNLIT)
 	{
-		const float Emissive = RWIlluminanceTexture[TexelPos] * View.OneOverPreExposure;
+		const float Emissive = RWIlluminanceTexture[OutputTexelPos] * View.OneOverPreExposure;
 
 		float3 Denominator = 0.0f;
 
@@ -350,7 +347,7 @@ void CalculateExposureIlluminance(uint2 ThreadId : SV_DispatchThreadID)
 		Illuminance *= lerp(EyeAdaptation_IgnoreMaterialsLuminanceScale, 1.0f, EyeAdaptation_IgnoreMaterialsMinBaseColorLuminance);
 	}
 	
-	RWIlluminanceTexture[TexelPos] = Illuminance * View.PreExposure;
+	RWIlluminanceTexture[OutputTexelPos] = Illuminance * View.PreExposure;
 }
 
 #endif // COMPUTESHADER
diff --git a/Engine/Source/Runtime/Renderer/Private/PostProcess/PostProcessEyeAdaptation.cpp b/Engine/Source/Runtime/Renderer/Private/PostProcess/PostProcessEyeAdaptation.cpp
--- a/Engine/Source/Runtime/Renderer/Private/PostProcess/PostProcessEyeAdaptation.cpp
+++ b/Engine/Source/Runtime/Renderer/Private/PostProcess/PostProcessEyeAdaptation.cpp
@@ -96,6 +96,12 @@ namespace
 		TEXT(""),
 		ECVF_Scalability | ECVF_RenderThreadSafe);
 
+	TAutoConsoleVariable<int32> CVarAutoExposureIgnoreMaterialsDownscaleFactor(
+		TEXT("r.AutoExposure.IgnoreMaterials.DownscaleFactor"),
+		2,
+		TEXT(""),
+		ECVF_Scalability | ECVF_RenderThreadSafe);
+
 	TAutoConsoleVariable<bool> CVarAutoExposureIgnoreMaterialsDebug(
 		TEXT("r.AutoExposure.IgnoreMaterials.Debug"),
 		false,
@@ -134,6 +140,13 @@ bool IsAutoExposureUsingIlluminanceEnabled(const FViewInfo& View)
 	return CVarAutoExposureIgnoreMaterials.GetValueOnRenderThread() && GetAutoExposureMethod(View) == AEM_Histogram;
 }
 
+int32 GetAutoExposureIlluminanceDownscaleFactor()
+{
+	int32 DownscaleFactor = CVarAutoExposureIgnoreMaterialsDownscaleFactor.GetValueOnRenderThread();
+
+	return FMath::Clamp(DownscaleFactor, 1, 16);
+}
+
 static EAutoExposureMethod ApplyEyeAdaptationQuality(EAutoExposureMethod AutoExposureMethod)
 {
 	static const auto CVarEyeAdaptationQuality = IConsoleManager::Get().FindTConsoleVariableDataInt(TEXT("r.EyeAdaptationQuality"));
@@ -517,6 +530,8 @@ class FSetupExposureIlluminanceCS : public FGlobalShader
 		SHADER_PARAMETER_STRUCT_INCLUDE(FSceneTextureParameters, SceneTextures)
 		SHADER_PARAMETER_RDG_TEXTURE(Texture2D, ColorTexture)
 		SHADER_PARAMETER_RDG_TEXTURE_UAV(RWTexture2D<float4>, RWIlluminanceTexture)
+		SHADER_PARAMETER_STRUCT(FScreenPassTextureViewportParameters, Illuminance)
+		SHADER_PARAMETER(uint32, IllumiananceDownscaleFactor)
 	END_SHADER_PARAMETER_STRUCT()
 
 	static bool ShouldCompilePermutation(const FGlobalShaderPermutationParameters& Parameters)
@@ -546,7 +561,9 @@ FRDGTextureRef AddSetupExposureIlluminancePass(
 		return nullptr;
 	}
 
-	const FRDGTextureDesc Desc = FRDGTextureDesc::Create2D(SceneTextures.Config.Extent, PF_R16F, FClearValueBinding::Black, TexCreate_UAV | TexCreate_ShaderResource);
+	const FIntPoint OutputExtent = GetDownscaledExtent(SceneTextures.Config.Extent, GetAutoExposureIlluminanceDownscaleFactor());
+
+	const FRDGTextureDesc Desc = FRDGTextureDesc::Create2D(OutputExtent, PF_R16F, FClearValueBinding::Black, TexCreate_UAV | TexCreate_ShaderResource);
 	FRDGTextureRef OutputTexture = GraphBuilder.CreateTexture(Desc, TEXT("EyeAdaptation_Illuminance"));
 
 	for (uint32 ViewIndex = 0; ViewIndex < ViewCount; ++ViewIndex)
@@ -555,11 +572,16 @@ FRDGTextureRef AddSetupExposureIlluminancePass(
 
 		if (IsAutoExposureUsingIlluminanceEnabled(View) && CVarAutoExposureIgnoreMaterialsUsePrecalculatedIlluminance.GetValueOnRenderThread())
 		{
+			const FScreenPassTextureViewport SceneViewport(SceneTextures.Config.Extent, View.ViewRect);
+			const FScreenPassTextureViewport OutputViewport(GetDownscaledViewport(SceneViewport, GetAutoExposureIlluminanceDownscaleFactor()));
+
 			auto* PassParameters = GraphBuilder.AllocParameters<FSetupExposureIlluminanceCS::FParameters>();
 			PassParameters->View = View.ViewUniformBuffer;
 			PassParameters->SceneTextures = GetSceneTextureParameters(GraphBuilder, SceneTextures);
 			PassParameters->ColorTexture = SceneTextures.Color.Resolve;
 			PassParameters->RWIlluminanceTexture = GraphBuilder.CreateUAV(OutputTexture);
+			PassParameters->Illuminance = GetScreenPassTextureViewportParameters(OutputViewport);
+			PassParameters->IllumiananceDownscaleFactor = GetAutoExposureIlluminanceDownscaleFactor();
 
 			auto ComputeShader = View.ShaderMap->GetShader<FSetupExposureIlluminanceCS>();
 
@@ -568,7 +590,7 @@ FRDGTextureRef AddSetupExposureIlluminancePass(
 				RDG_EVENT_NAME("SetupExposureIlluminance (ViewId=%d)", ViewIndex),
 				ComputeShader,
 				PassParameters,
-				FComputeShaderUtils::GetGroupCount(View.ViewRect.Size(), FIntPoint(8, 8)));
+				FComputeShaderUtils::GetGroupCount(OutputViewport.Rect.Size(), FIntPoint(8, 8)));
 		}
 	}
 
@@ -585,6 +607,8 @@ class FCalculateExposureIlluminanceCS : public FGlobalShader
 		SHADER_PARAMETER_STRUCT_INCLUDE(FSceneTextureParameters, SceneTextures)
 		SHADER_PARAMETER_RDG_TEXTURE(Texture2D, ColorTexture)
 		SHADER_PARAMETER_RDG_TEXTURE_UAV(RWTexture2D<float4>, RWIlluminanceTexture)
+		SHADER_PARAMETER_STRUCT(FScreenPassTextureViewportParameters, Illuminance)
+		SHADER_PARAMETER(uint32, IllumiananceDownscaleFactor)
 
 		SHADER_PARAMETER_STRUCT(FEyeAdaptationParameters, EyeAdaptation)
 
@@ -621,11 +645,16 @@ FRDGTextureRef AddCalculateExposureIlluminancePass(
 		{
 			check(CVarAutoExposureIgnoreMaterialsUsePrecalculatedIlluminance.GetValueOnRenderThread());
 
+			const FScreenPassTextureViewport SceneViewport(SceneTextures.Config.Extent, View.ViewRect);
+			const FScreenPassTextureViewport OutputViewport(GetDownscaledViewport(SceneViewport, GetAutoExposureIlluminanceDownscaleFactor()));
+
 			auto* PassParameters = GraphBuilder.AllocParameters<FCalculateExposureIlluminanceCS::FParameters>();
 			PassParameters->View = View.ViewUniformBuffer;
 			PassParameters->SceneTextures = GetSceneTextureParameters(GraphBuilder, SceneTextures);
 			PassParameters->ColorTexture = SceneTextures.Color.Resolve;
 			PassParameters->RWIlluminanceTexture = GraphBuilder.CreateUAV(ExposureIlluminanceSetup);
+			PassParameters->Illuminance = GetScreenPassTextureViewportParameters(OutputViewport);
+			PassParameters->IllumiananceDownscaleFactor = GetAutoExposureIlluminanceDownscaleFactor();
 			PassParameters->EyeAdaptation = GetEyeAdaptationParameters(View, ERHIFeatureLevel::SM5);
 
 			PassParameters->PreIntegratedGF = GSystemTextures.PreintegratedGF->GetRHI();
@@ -638,7 +667,7 @@ FRDGTextureRef AddCalculateExposureIlluminancePass(
 				RDG_EVENT_NAME("CalculateExposureIlluminance (ViewId=%d)", ViewIndex),
 				ComputeShader,
 				PassParameters,
-				FComputeShaderUtils::GetGroupCount(View.ViewRect.Size(), FIntPoint(8, 8)));
+				FComputeShaderUtils::GetGroupCount(OutputViewport.Rect.Size(), FIntPoint(8, 8)));
 		}
 	}
 
diff --git a/Engine/Source/Runtime/Renderer/Private/PostProcess/PostProcessEyeAdaptation.h b/Engine/Source/Runtime/Renderer/Private/PostProcess/PostProcessEyeAdaptation.h
--- a/Engine/Source/Runtime/Renderer/Private/PostProcess/PostProcessEyeAdaptation.h
+++ b/Engine/Source/Runtime/Renderer/Private/PostProcess/PostProcessEyeAdaptation.h
@@ -50,6 +50,7 @@ float GetAutoExposureCompensationFromSettings(const FViewInfo& View);
 bool IsExtendLuminanceRangeEnabled();
 
 bool IsAutoExposureUsingIlluminanceEnabled(const FViewInfo& View);
+int32 GetAutoExposureIlluminanceDownscaleFactor();
 
 // Returns the auto exposure method enabled by the view (including CVar override).
 EAutoExposureMethod GetAutoExposureMethod(const FViewInfo& View);
diff --git a/Engine/Source/Runtime/Renderer/Private/PostProcess/PostProcessing.cpp b/Engine/Source/Runtime/Renderer/Private/PostProcess/PostProcessing.cpp
--- a/Engine/Source/Runtime/Renderer/Private/PostProcess/PostProcessing.cpp
+++ b/Engine/Source/Runtime/Renderer/Private/PostProcess/PostProcessing.cpp
@@ -809,7 +809,9 @@ void AddPostProcessingPasses(
 			{
 				if (Inputs.ExposureIlluminance)
 				{
-					HistogramSceneColor = FScreenPassTexture(Inputs.ExposureIlluminance, PrimaryViewRect);
+					const FIntRect IlluminanceRect = GetDownscaledRect(PrimaryViewRect, GetAutoExposureIlluminanceDownscaleFactor());
+
+					HistogramSceneColor = FScreenPassTexture(Inputs.ExposureIlluminance, IlluminanceRect);
 				}
 				else
 				{
