Flush RDG pass queue slightly earlier to avoid stalls later.

diff --git a/Engine/Source/Runtime/Renderer/Private/DeferredShadingRenderer.cpp b/Engine/Source/Runtime/Renderer/Private/DeferredShadingRenderer.cpp
--- a/Engine/Source/Runtime/Renderer/Private/DeferredShadingRenderer.cpp
+++ b/Engine/Source/Runtime/Renderer/Private/DeferredShadingRenderer.cpp
@@ -3311,6 +3311,8 @@ void FDeferredShadingSceneRenderer::Render(FRDGBuilder& GraphBuilder)
 	FAsyncLumenIndirectLightingOutputs AsyncLumenIndirectLightingOutputs;
 	const bool bHasLumenLights = SortedLightSet.LumenLightStart < SortedLightSet.SortedLights.Num();
 
+	GraphBuilder.FlushSetupQueue();
+
 	// Shadows, lumen and fog after base pass
 	if (!bHasRayTracedOverlay)
 	{
