Fix crashe when allocating IES Atlas.

diff --git a/Engine/Source/Runtime/Renderer/Private/IESTextureManager.cpp b/Engine/Source/Runtime/Renderer/Private/IESTextureManager.cpp
--- a/Engine/Source/Runtime/Renderer/Private/IESTextureManager.cpp
+++ b/Engine/Source/Runtime/Renderer/Private/IESTextureManager.cpp
@@ -338,7 +338,7 @@ static FRDGTextureRef CreateAtlasTexture(FRDGBuilder& GraphBuilder, const FIntPo
 		PF_R16F,
 		FClearValueBinding::Transparent,
 		ETextureCreateFlags::UAV | ETextureCreateFlags::ShaderResource | ETextureCreateFlags::RenderTargetable,
-		SliceCount, 
+		FMath::Max(1u, SliceCount),
 		1),
 		TEXT("IES.AtlasTexture"),
 		ERDGTextureFlags::MultiFrame);
