Fix for editor crashing after executing the command r.DisplayInternals 1 or r.DisplayInternals = 1. Also add an RDGEventName to UpdateAllPrimitiveSceneInfos pass.

diff --git a/Engine/Source/Runtime/Engine/Private/SceneInterface.cpp b/Engine/Source/Runtime/Engine/Private/SceneInterface.cpp
--- a/Engine/Source/Runtime/Engine/Private/SceneInterface.cpp
+++ b/Engine/Source/Runtime/Engine/Private/SceneInterface.cpp
@@ -10,7 +10,7 @@ TArray<FPrimitiveComponentId> FSceneInterface::GetScenePrimitiveComponentIds() c
 
 void FSceneInterface::UpdateAllPrimitiveSceneInfos(FRHICommandListImmediate& RHICmdList, bool bAsyncCreateLPIs)
 {
-	FRDGBuilder GraphBuilder(RHICmdList);
+	FRDGBuilder GraphBuilder(RHICmdList, FRDGEventName(TEXT("UpdateAllPrimitiveSceneInfos")));
 	UpdateAllPrimitiveSceneInfos(GraphBuilder, bAsyncCreateLPIs);
 	GraphBuilder.Execute();
 }
\ No newline at end of file
diff --git a/Engine/Source/Runtime/Renderer/Private/SceneRendering.cpp b/Engine/Source/Runtime/Renderer/Private/SceneRendering.cpp
--- a/Engine/Source/Runtime/Renderer/Private/SceneRendering.cpp
+++ b/Engine/Source/Runtime/Renderer/Private/SceneRendering.cpp
@@ -5048,17 +5048,17 @@ static void DisplayInternals(FRDGBuilder& GraphBuilder, FViewInfo& InView)
 	// if r.DisplayInternals != 0
 	if(Family->EngineShowFlags.OnScreenDebug && Family->DisplayInternalsData.IsValid())
 	{
-		AddPass(GraphBuilder, RDG_EVENT_NAME("DisplayInternals"), [Family, &InView] (FRHICommandListImmediate& RHICmdList)
+		FRDGTextureRef OutputTexture = GraphBuilder.FindExternalTexture(Family->RenderTarget->GetRenderTargetTexture());
+		FScreenPassRenderTarget Output = FScreenPassRenderTarget::CreateViewFamilyOutput(OutputTexture, InView);
+		AddDrawCanvasPass(GraphBuilder, RDG_EVENT_NAME("DisplayInternals"), InView, Output, [Family, &InView](FCanvas& Canvas)
 		{
 			// could be 0
 			auto State = InView.ViewState;
 
-			FCanvas Canvas((FRenderTarget*)Family->RenderTarget, NULL, Family->Time, InView.GetFeatureLevel());
 			Canvas.SetRenderTargetRect(FIntRect(0, 0, Family->RenderTarget->GetSizeXY().X, Family->RenderTarget->GetSizeXY().Y));
 
 
 			FRHIRenderPassInfo RenderPassInfo(Family->RenderTarget->GetRenderTargetTexture(), ERenderTargetActions::Load_Store);
-			RHICmdList.BeginRenderPass(RenderPassInfo, TEXT("DisplayInternalsRenderPass"));
 
 			// further down to not intersect with "LIGHTING NEEDS TO BE REBUILT"
 			FVector2D Pos(30, 140);
@@ -5143,9 +5143,6 @@ static void DisplayInternals(FRDGBuilder& GraphBuilder, FViewInfo& InView)
 
 	#undef CANVAS_LINE
 	#undef CANVAS_HEADER
-
-			RHICmdList.EndRenderPass();
-			Canvas.Flush_RenderThread(RHICmdList);
 		});
 	}
 #endif
