GPULM: Fix ray tracing payload type runtime error

GPULM flagged some errors at runtime due to a mismatch in payload type with respect to the default hit shader. The path tracing default hit shader should be used instead of the raytracing one.

Refactor functions that find default shaders to be independent of the DeferredRenderer and of the View object (since only a ShaderMap is needed)

Make the binding to the miss shader for GPULM explicit since this should also match the right payload type (and we want to eventually deprecate implicit miss shader bindings)

diff --git a/Engine/Plugins/Experimental/GPULightmass/Shaders/Private/LightmapCommon.ush b/Engine/Plugins/Experimental/GPULightmass/Shaders/Private/LightmapCommon.ush
new file mode 100644
--- /dev/null
+++ b/Engine/Plugins/Experimental/GPULightmass/Shaders/Private/LightmapCommon.ush
@@ -0,0 +1,30 @@
+// Copyright Epic Games, Inc. All Rights Reserved.
+
+#pragma once
+
+#include "/Engine/Private/PathTracing/PathTracingCommon.ush"
+
+FMinimalPayload TraceLightmapVisibilityRay(
+	in RaytracingAccelerationStructure TLAS,
+	in uint RayFlags,
+	in FRayDesc Ray)
+{
+	// GPULightmass only binds path tracing shaders, so we need to use it to trace rays, even though this function only cares about the hit distance
+	FPackedPathTracingPayload PackedPayload = InitPathTracingVisibilityPayload(uint2(0, 0), 1.0f);
+
+	TraceRay(
+		TLAS,
+		RayFlags,
+		RAY_TRACING_MASK_SHADOW,
+		RAY_TRACING_SHADER_SLOT_SHADOW,
+		RAY_TRACING_NUM_SHADER_SLOTS,
+		0,
+		Ray.GetNativeDesc(),
+		PackedPayload);
+
+	// Unpack the payload
+
+	FMinimalPayload MinimalPayload = (FMinimalPayload)0;
+	MinimalPayload.HitT = PackedPayload.HitT;
+	return MinimalPayload;
+}
\ No newline at end of file
diff --git a/Engine/Plugins/Experimental/GPULightmass/Shaders/Private/LightmapPathTracing.usf b/Engine/Plugins/Experimental/GPULightmass/Shaders/Private/LightmapPathTracing.usf
--- a/Engine/Plugins/Experimental/GPULightmass/Shaders/Private/LightmapPathTracing.usf
+++ b/Engine/Plugins/Experimental/GPULightmass/Shaders/Private/LightmapPathTracing.usf
@@ -37,6 +37,7 @@ Texture2D<float> RayGuidingCDFY;
 #include "/Engine/Private/PathTracing/Material/PathTracingMaterialSampling.ush"
 
 #include "LightmapEncoding.ush"
+#include "LightmapCommon.ush"
 
 int LastInvalidationFrame;
 int NumTotalSamples;
@@ -1400,11 +1401,9 @@ void VolumetricLightmapPathTracingMainRG()
 		{
 			uint RayFlags = RAY_FLAG_SKIP_CLOSEST_HIT_SHADER | RAY_FLAG_ACCEPT_FIRST_HIT_AND_END_SEARCH;
 	
-			FMinimalPayload MinimalPayload = TraceVisibilityRay(
+			FMinimalPayload MinimalPayload = TraceLightmapVisibilityRay(
 				TLAS,
 				RayFlags,
-				RAY_TRACING_MASK_SHADOW,
-				uint2(0, 0),
 				Ray);
 			
 			DirectionalLightShadowing[VoxelPos] += (MinimalPayload.IsMiss() ? 1 : 0) / 32.0f;
@@ -1477,11 +1476,9 @@ void StationaryLightShadowTracingMainRG()
 			{
 				uint RayFlags = RAY_FLAG_SKIP_CLOSEST_HIT_SHADER | RAY_FLAG_ACCEPT_FIRST_HIT_AND_END_SEARCH;
 
-				FMinimalPayload MinimalPayload = TraceVisibilityRay(
+				FMinimalPayload MinimalPayload = TraceLightmapVisibilityRay(
 					TLAS,
 					RayFlags,
-					RAY_TRACING_MASK_SHADOW,
-					uint2(0, 0),
 					Ray);
 				
 				uint Visibility = asuint(ShadowMask[TexelIndexInPool][ChannelIndex]);
diff --git a/Engine/Plugins/Experimental/GPULightmass/Shaders/Private/StaticShadowDepthMap.usf b/Engine/Plugins/Experimental/GPULightmass/Shaders/Private/StaticShadowDepthMap.usf
--- a/Engine/Plugins/Experimental/GPULightmass/Shaders/Private/StaticShadowDepthMap.usf
+++ b/Engine/Plugins/Experimental/GPULightmass/Shaders/Private/StaticShadowDepthMap.usf
@@ -1,6 +1,7 @@
 // Copyright Epic Games, Inc. All Rights Reserved.
 
 #include "/Engine/Private/RayTracing/RayTracingCommon.ush"
+#include "LightmapCommon.ush"
 
 RaytracingAccelerationStructure TLAS;
 
@@ -50,11 +51,9 @@ void StaticShadowDepthMapTracingRG()
 
 	const uint RayFlags = 0;
 	
-	FMinimalPayload MinimalPayload = TraceVisibilityRay(
+	FMinimalPayload MinimalPayload = TraceLightmapVisibilityRay(
 		TLAS,
 		RayFlags,
-		RAY_TRACING_MASK_SHADOW,
-		uint2(0, 0),
 		Ray);
 
 	if (MinimalPayload.IsHit())
diff --git a/Engine/Plugins/Experimental/GPULightmass/Source/GPULightmass/Private/LightmapRayTracing.h b/Engine/Plugins/Experimental/GPULightmass/Source/GPULightmass/Private/LightmapRayTracing.h
--- a/Engine/Plugins/Experimental/GPULightmass/Source/GPULightmass/Private/LightmapRayTracing.h
+++ b/Engine/Plugins/Experimental/GPULightmass/Source/GPULightmass/Private/LightmapRayTracing.h
@@ -71,7 +71,7 @@ class FLightmapPathTracingRGS : public FGlobalShader
 
 	static ERayTracingPayloadType GetRayTracingPayloadType(const int32 PermutationId)
 	{
-		return ERayTracingPayloadType::Minimal | ERayTracingPayloadType::PathTracingMaterial;
+		return ERayTracingPayloadType::PathTracingMaterial;
 	}
 
 	BEGIN_SHADER_PARAMETER_STRUCT(FParameters, )
@@ -136,7 +136,7 @@ class FVolumetricLightmapPathTracingRGS : public FGlobalShader
 
 	static ERayTracingPayloadType GetRayTracingPayloadType(const int32 PermutationId)
 	{
-		return ERayTracingPayloadType::Minimal | ERayTracingPayloadType::PathTracingMaterial;
+		return ERayTracingPayloadType::PathTracingMaterial;
 	}
 
 	BEGIN_SHADER_PARAMETER_STRUCT(FParameters, )
@@ -197,7 +197,7 @@ class FStationaryLightShadowTracingRGS : public FGlobalShader
 
 	static ERayTracingPayloadType GetRayTracingPayloadType(const int32 PermutationId)
 	{
-		return ERayTracingPayloadType::Minimal | ERayTracingPayloadType::PathTracingMaterial;
+		return ERayTracingPayloadType::PathTracingMaterial;
 	}
 
 	BEGIN_SHADER_PARAMETER_STRUCT(FParameters, )
@@ -269,7 +269,7 @@ class FStaticShadowDepthMapTracingRGS : public FGlobalShader
 
 	static ERayTracingPayloadType GetRayTracingPayloadType(const int32 PermutationId)
 	{
-		return ERayTracingPayloadType::Minimal | ERayTracingPayloadType::PathTracingMaterial;
+		return ERayTracingPayloadType::PathTracingMaterial;
 	}
 
 	BEGIN_SHADER_PARAMETER_STRUCT(FParameters, )
diff --git a/Engine/Plugins/Experimental/GPULightmass/Source/GPULightmass/Private/LightmapRenderer.cpp b/Engine/Plugins/Experimental/GPULightmass/Source/GPULightmass/Private/LightmapRenderer.cpp
--- a/Engine/Plugins/Experimental/GPULightmass/Source/GPULightmass/Private/LightmapRenderer.cpp
+++ b/Engine/Plugins/Experimental/GPULightmass/Source/GPULightmass/Private/LightmapRenderer.cpp
@@ -29,6 +29,7 @@
 #include "ShaderCompiler.h"
 #include "RectLightTextureManager.h"
 #include "IESTextureManager.h"
+#include "PathTracing.h"
 
 class FCopyConvergedLightmapTilesCS : public FGlobalShader
 {
@@ -1236,13 +1237,16 @@ bool FSceneRenderState::SetupRayTracingScene(int32 LODIndex)
 			}
 			PSOInitializer.SetRayGenShaderTable(RayGenShaderTable);
 
-			auto DefaultClosestHitShader = GlobalShaderMap->GetShader<FOpaqueShadowHitGroup>().GetRayTracingShader();
+			auto DefaultClosestHitShader = GetPathTracingDefaultOpaqueHitShader(GlobalShaderMap);
 			TArray<FRHIRayTracingShader*> RayTracingHitGroupLibrary;
 			FShaderMapResource::GetRayTracingHitGroupLibrary(RayTracingHitGroupLibrary, DefaultClosestHitShader);
 
 			PSOInitializer.SetHitGroupTable(RayTracingHitGroupLibrary);
-			
-			// TODO(UE-157946): This pipeline does not bind any miss shader and relies on the pipeline to do this automatically. This should be made explicit.
+
+			TArray<FRHIRayTracingShader*> RayTracingMissLibrary;
+			RayTracingMissLibrary.Add(GetPathTracingDefaultMissShader(GlobalShaderMap));
+			PSOInitializer.SetMissShaderTable(RayTracingMissLibrary);
+
 			RayTracingPipelineState = PipelineStateCache::GetAndOrCreateRayTracingPipelineState(RHICmdList, PSOInitializer);
 
 			TUniquePtr<FRayTracingLocalShaderBindingWriter> BindingWriter = MakeUnique<FRayTracingLocalShaderBindingWriter>();
@@ -1328,6 +1332,8 @@ bool FSceneRenderState::SetupRayTracingScene(int32 LODIndex)
 						NumTotalBindings, MergedBindings,
 						bCopyDataToInlineStorage);
 				}
+
+				// there is only one miss shader, so it must be at index 0 by definition
 				RHICmdList.SetRayTracingMissShader(RayTracingScene, 0, RayTracingPipelineState, 0 /* ShaderIndexInPipeline */, 0, nullptr, 0);
 
 				// Move the ray tracing binding container ownership to the command list, so that memory will be
diff --git a/Engine/Source/Runtime/Renderer/Private/DeferredShadingRenderer.cpp b/Engine/Source/Runtime/Renderer/Private/DeferredShadingRenderer.cpp
--- a/Engine/Source/Runtime/Renderer/Private/DeferredShadingRenderer.cpp
+++ b/Engine/Source/Runtime/Renderer/Private/DeferredShadingRenderer.cpp
@@ -2466,7 +2466,7 @@ void FDeferredShadingSceneRenderer::Render(FRDGBuilder& GraphBuilder)
 		if (!ViewFamily.EngineShowFlags.PathTracing)
 		{
 			// get the default lighting miss shader (to implicitly fill in the MissShader library before the RT pipeline is created)
-			GetRayTracingLightingMissShader(ReferenceView);
+			GetRayTracingLightingMissShader(ReferenceView.ShaderMap);
 			RayTracingScene.NumMissShaderSlots++;
 		}
 
diff --git a/Engine/Source/Runtime/Renderer/Private/DeferredShadingRenderer.h b/Engine/Source/Runtime/Renderer/Private/DeferredShadingRenderer.h
--- a/Engine/Source/Runtime/Renderer/Private/DeferredShadingRenderer.h
+++ b/Engine/Source/Runtime/Renderer/Private/DeferredShadingRenderer.h
@@ -1156,15 +1156,6 @@ private:
 	static void PrepareLumenHardwareRayTracingRadiosityLumenMaterial(const FViewInfo& View, TArray<FRHIRayTracingShader*>& OutRayGenShaders);
 	static void PrepareLumenHardwareRayTracingDirectLightingLumenMaterial(const FViewInfo& View, TArray<FRHIRayTracingShader*>& OutRayGenShaders);
 
-	/** Lighting evaluation shader registration */
-	static FRHIRayTracingShader* GetRayTracingDefaultMissShader(const FViewInfo& View);
-	static FRHIRayTracingShader* GetRayTracingLightingMissShader(const FViewInfo& View);
-	static FRHIRayTracingShader* GetRayTracingDefaultOpaqueShader(const FViewInfo& View);
-	static FRHIRayTracingShader* GetRayTracingDefaultHiddenShader(const FViewInfo& View);
-	static FRHIRayTracingShader* GetPathTracingDefaultMissShader(const FViewInfo& View);
-	static FRHIRayTracingShader* GetPathTracingDefaultOpaqueHitShader(const FViewInfo& View);
-	static FRHIRayTracingShader* GetPathTracingDefaultHiddenHitShader(const FViewInfo& View);
-
 	const FRHITransition* RayTracingDynamicGeometryUpdateEndTransition = nullptr; // Signaled when all AS for this frame are built
 #endif // RHI_RAYTRACING
 
diff --git a/Engine/Source/Runtime/Renderer/Private/PathTracing.cpp b/Engine/Source/Runtime/Renderer/Private/PathTracing.cpp
--- a/Engine/Source/Runtime/Renderer/Private/PathTracing.cpp
+++ b/Engine/Source/Runtime/Renderer/Private/PathTracing.cpp
@@ -871,14 +871,14 @@ public:
 };
 IMPLEMENT_SHADER_TYPE(, FPathTracingDefaultMS, TEXT("/Engine/Private/PathTracing/PathTracingMissShader.usf"), TEXT("PathTracingDefaultMS"), SF_RayMiss);
 
-FRHIRayTracingShader* FDeferredShadingSceneRenderer::GetPathTracingDefaultMissShader(const FViewInfo& View)
+FRHIRayTracingShader* GetPathTracingDefaultMissShader(const FGlobalShaderMap* ShaderMap)
 {
-	return View.ShaderMap->GetShader<FPathTracingDefaultMS>().GetRayTracingShader();
+	return ShaderMap->GetShader<FPathTracingDefaultMS>().GetRayTracingShader();
 }
 
 void FDeferredShadingSceneRenderer::SetupPathTracingDefaultMissShader(FRHICommandListImmediate& RHICmdList, const FViewInfo& View)
 {
-	int32 MissShaderPipelineIndex = FindRayTracingMissShaderIndex(View.RayTracingMaterialPipeline, GetPathTracingDefaultMissShader(View), true);
+	int32 MissShaderPipelineIndex = FindRayTracingMissShaderIndex(View.RayTracingMaterialPipeline, GetPathTracingDefaultMissShader(View.ShaderMap), true);
 
 	RHICmdList.SetRayTracingMissShader(View.GetRayTracingSceneChecked(),
 		RAY_TRACING_MISS_SHADER_SLOT_DEFAULT,
@@ -1251,14 +1251,14 @@ class FPathTracingDefaultHiddenHitGroup : public FGlobalShader
 IMPLEMENT_SHADER_TYPE(, FPathTracingDefaultOpaqueHitGroup, TEXT("/Engine/Private/PathTracing/PathTracingDefaultHitShader.usf"), TEXT("closesthit=PathTracingDefaultOpaqueCHS"), SF_RayHitGroup);
 IMPLEMENT_SHADER_TYPE(, FPathTracingDefaultHiddenHitGroup, TEXT("/Engine/Private/PathTracing/PathTracingDefaultHitShader.usf"), TEXT("closesthit=PathTracingDefaultHiddenCHS anyhit=PathTracingDefaultHiddenAHS"), SF_RayHitGroup);
 
-FRHIRayTracingShader* FDeferredShadingSceneRenderer::GetPathTracingDefaultOpaqueHitShader(const FViewInfo& View)
+FRHIRayTracingShader* GetPathTracingDefaultOpaqueHitShader(const FGlobalShaderMap* ShaderMap)
 {
-	return View.ShaderMap->GetShader<FPathTracingDefaultOpaqueHitGroup>().GetRayTracingShader();
+	return ShaderMap->GetShader<FPathTracingDefaultOpaqueHitGroup>().GetRayTracingShader();
 }
 
-FRHIRayTracingShader* FDeferredShadingSceneRenderer::GetPathTracingDefaultHiddenHitShader(const FViewInfo& View)
+FRHIRayTracingShader* GetPathTracingDefaultHiddenHitShader(const FGlobalShaderMap* ShaderMap)
 {
-	return View.ShaderMap->GetShader<FPathTracingDefaultHiddenHitGroup>().GetRayTracingShader();
+	return ShaderMap->GetShader<FPathTracingDefaultHiddenHitGroup>().GetRayTracingShader();
 }
 
 bool FRayTracingMeshProcessor::ProcessPathTracing(
diff --git a/Engine/Source/Runtime/Renderer/Private/PathTracing.h b/Engine/Source/Runtime/Renderer/Private/PathTracing.h
--- a/Engine/Source/Runtime/Renderer/Private/PathTracing.h
+++ b/Engine/Source/Runtime/Renderer/Private/PathTracing.h
@@ -17,8 +17,13 @@ class FRDGBuilder;
 class FScene;
 class FViewInfo;
 class FSceneViewFamily;
+class FRHIRayTracingShader;
+class FGlobalShaderMap;
 
 bool PrepareSkyTexture(FRDGBuilder& GraphBuilder, FScene* Scene, const FViewInfo& View, bool SkylightEnabled, bool UseMISCompensation, FPathTracingSkylight* SkylightParameters);
+RENDERER_API FRHIRayTracingShader* GetPathTracingDefaultMissShader(const FGlobalShaderMap* ShaderMap);
+RENDERER_API FRHIRayTracingShader* GetPathTracingDefaultOpaqueHitShader(const FGlobalShaderMap* ShaderMap);
+RENDERER_API FRHIRayTracingShader* GetPathTracingDefaultHiddenHitShader(const FGlobalShaderMap* ShaderMap);
 
 namespace PathTracing
 {
diff --git a/Engine/Source/Runtime/Renderer/Private/RayTracing/RayTracingLighting.cpp b/Engine/Source/Runtime/Renderer/Private/RayTracing/RayTracingLighting.cpp
--- a/Engine/Source/Runtime/Renderer/Private/RayTracing/RayTracingLighting.cpp
+++ b/Engine/Source/Runtime/Renderer/Private/RayTracing/RayTracingLighting.cpp
@@ -7,6 +7,7 @@
 #if RHI_RAYTRACING
 
 #include "SceneRendering.h"
+#include "RayTracingMaterialHitShaders.h"
 
 static TAutoConsoleVariable<int32> CVarRayTracingLightingCells(
 	TEXT("r.RayTracing.LightCulling.Cells"),
@@ -507,9 +508,9 @@ static void BindLightFunction(
 	ShaderBindings.SetRayTracingShaderBindingsForMissShader(RHICmdList, RTScene, Pipeline, MissShaderPipelineIndex, Index);
 }
 
-FRHIRayTracingShader* FDeferredShadingSceneRenderer::GetRayTracingLightingMissShader(const FViewInfo& View)
+FRHIRayTracingShader* GetRayTracingLightingMissShader(const FGlobalShaderMap* ShaderMap)
 {
-	return View.ShaderMap->GetShader<FRayTracingLightingMS>().GetRayTracingShader();
+	return ShaderMap->GetShader<FRayTracingLightingMS>().GetRayTracingShader();
 }
 
 void BindLightFunctionShaders(
@@ -582,7 +583,7 @@ static int32 BindParameters(const TShaderRef<ShaderClass>& Shader, typename Shad
 
 void FDeferredShadingSceneRenderer::SetupRayTracingDefaultMissShader(FRHICommandListImmediate& RHICmdList, const FViewInfo& View)
 {
-	int32 MissShaderPipelineIndex = FindRayTracingMissShaderIndex(View.RayTracingMaterialPipeline, GetRayTracingDefaultMissShader(View), true);
+	int32 MissShaderPipelineIndex = FindRayTracingMissShaderIndex(View.RayTracingMaterialPipeline, GetRayTracingDefaultMissShader(View.ShaderMap), true);
 
 	RHICmdList.SetRayTracingMissShader(View.GetRayTracingSceneChecked(),
 		RAY_TRACING_MISS_SHADER_SLOT_DEFAULT,
diff --git a/Engine/Source/Runtime/Renderer/Private/RayTracing/RayTracingLighting.h b/Engine/Source/Runtime/Renderer/Private/RayTracing/RayTracingLighting.h
--- a/Engine/Source/Runtime/Renderer/Private/RayTracing/RayTracingLighting.h
+++ b/Engine/Source/Runtime/Renderer/Private/RayTracing/RayTracingLighting.h
@@ -12,6 +12,8 @@
 
 #if RHI_RAYTRACING
 
+FRHIRayTracingShader* GetRayTracingLightingMissShader(const FGlobalShaderMap* ShaderMap);
+
 BEGIN_GLOBAL_SHADER_PARAMETER_STRUCT(FRaytracingLightDataPacked, RENDERER_API)
 	SHADER_PARAMETER(uint32, Count)
 	SHADER_PARAMETER(uint32, CellCount)
diff --git a/Engine/Source/Runtime/Renderer/Private/RayTracing/RayTracingMaterialHitShaders.cpp b/Engine/Source/Runtime/Renderer/Private/RayTracing/RayTracingMaterialHitShaders.cpp
--- a/Engine/Source/Runtime/Renderer/Private/RayTracing/RayTracingMaterialHitShaders.cpp
+++ b/Engine/Source/Runtime/Renderer/Private/RayTracing/RayTracingMaterialHitShaders.cpp
@@ -12,6 +12,7 @@
 #include "BuiltInRayTracingShaders.h"
 #include "RaytracingOptions.h"
 #include "RayTracingLighting.h"
+#include "PathTracing.h"
 
 int32 GEnableRayTracingMaterials = 1;
 static FAutoConsoleVariableRef CVarEnableRayTracingMaterials(
@@ -592,19 +593,19 @@ static bool PipelineContainsHitShaders(FRayTracingPipelineState* Pipeline, const
 	return true;
 }
 
-FRHIRayTracingShader* FDeferredShadingSceneRenderer::GetRayTracingDefaultMissShader(const FViewInfo& View)
+FRHIRayTracingShader* GetRayTracingDefaultMissShader(const FGlobalShaderMap* ShaderMap)
 {
-	return View.ShaderMap->GetShader<FPackedMaterialClosestHitPayloadMS>().GetRayTracingShader();
+	return ShaderMap->GetShader<FPackedMaterialClosestHitPayloadMS>().GetRayTracingShader();
 }
 
-FRHIRayTracingShader* FDeferredShadingSceneRenderer::GetRayTracingDefaultOpaqueShader(const FViewInfo& View)
+FRHIRayTracingShader* GetRayTracingDefaultOpaqueShader(const FGlobalShaderMap* ShaderMap)
 {
-	return View.ShaderMap->GetShader<FOpaqueShadowHitGroup>().GetRayTracingShader();
+	return ShaderMap->GetShader<FOpaqueShadowHitGroup>().GetRayTracingShader();
 }
 
-FRHIRayTracingShader* FDeferredShadingSceneRenderer::GetRayTracingDefaultHiddenShader(const FViewInfo& View)
+FRHIRayTracingShader* GetRayTracingDefaultHiddenShader(const FGlobalShaderMap* ShaderMap)
 {
-	return View.ShaderMap->GetShader<FHiddenMaterialHitGroup>().GetRayTracingShader();
+	return ShaderMap->GetShader<FHiddenMaterialHitGroup>().GetRayTracingShader();
 }
 
 
@@ -625,7 +626,7 @@ FRayTracingPipelineState* FDeferredShadingSceneRenderer::CreateRayTracingMateria
 	Initializer.MaxPayloadSizeInBytes = RAY_TRACING_MAX_ALLOWED_PAYLOAD_SIZE; // sizeof(FPackedMaterialClosestHitPayload)
 	Initializer.bAllowHitGroupIndexing = true;
 
-	FRHIRayTracingShader* DefaultMissShader = bIsPathTracing ? GetPathTracingDefaultMissShader(View) : GetRayTracingDefaultMissShader(View);
+	FRHIRayTracingShader* DefaultMissShader = bIsPathTracing ? GetPathTracingDefaultMissShader(View.ShaderMap) : GetRayTracingDefaultMissShader(View.ShaderMap);
 
 	TArray<FRHIRayTracingShader*> RayTracingMissShaderLibrary;
 	FShaderMapResource::GetRayTracingMissShaderLibrary(RayTracingMissShaderLibrary, DefaultMissShader);
@@ -642,8 +643,8 @@ FRayTracingPipelineState* FDeferredShadingSceneRenderer::CreateRayTracingMateria
 	static auto CVarEnableShadowMaterials = IConsoleManager::Get().FindConsoleVariable(TEXT("r.RayTracing.Shadows.EnableMaterials"));
 	const bool bEnableShadowMaterials = bMaterialsCompiled && (CVarEnableShadowMaterials ? CVarEnableShadowMaterials->GetInt() != 0 : true);
 
-	FRHIRayTracingShader* OpaqueShadowShader   = bIsPathTracing ? GetPathTracingDefaultOpaqueHitShader(View) : GetRayTracingDefaultOpaqueShader(View);
-	FRHIRayTracingShader* HiddenMaterialShader = bIsPathTracing ? GetPathTracingDefaultHiddenHitShader(View) : GetRayTracingDefaultHiddenShader(View);
+	FRHIRayTracingShader* OpaqueShadowShader   = bIsPathTracing ? GetPathTracingDefaultOpaqueHitShader(View.ShaderMap) : GetRayTracingDefaultOpaqueShader(View.ShaderMap);
+	FRHIRayTracingShader* HiddenMaterialShader = bIsPathTracing ? GetPathTracingDefaultHiddenHitShader(View.ShaderMap) : GetRayTracingDefaultHiddenShader(View.ShaderMap);
 	
 	TArray<FRHIRayTracingShader*> RayTracingHitGroupLibrary;
 	if (bEnableMaterials)
diff --git a/Engine/Source/Runtime/Renderer/Private/RayTracing/RayTracingMaterialHitShaders.h b/Engine/Source/Runtime/Renderer/Private/RayTracing/RayTracingMaterialHitShaders.h
--- a/Engine/Source/Runtime/Renderer/Private/RayTracing/RayTracingMaterialHitShaders.h
+++ b/Engine/Source/Runtime/Renderer/Private/RayTracing/RayTracingMaterialHitShaders.h
@@ -11,6 +11,9 @@
 #include "RayTracingPayloadType.h"
 
 ENGINE_API uint8 ComputeBlendModeMask(const EBlendMode BlendMode);
+FRHIRayTracingShader* GetRayTracingDefaultMissShader(const FGlobalShaderMap* ShaderMap);
+FRHIRayTracingShader* GetRayTracingDefaultOpaqueShader(const FGlobalShaderMap* ShaderMap);
+FRHIRayTracingShader* GetRayTracingDefaultHiddenShader(const FGlobalShaderMap* ShaderMap);
 
 class RENDERER_API FRayTracingMeshProcessor
 {
