Fixes r.TSR.ShadingRejection.Flickering.Period letting before distortion translucency ghosts based on framerate

diff --git a/Engine/Source/Runtime/Renderer/Private/PostProcess/TemporalSuperResolution.cpp b/Engine/Source/Runtime/Renderer/Private/PostProcess/TemporalSuperResolution.cpp
--- a/Engine/Source/Runtime/Renderer/Private/PostProcess/TemporalSuperResolution.cpp
+++ b/Engine/Source/Runtime/Renderer/Private/PostProcess/TemporalSuperResolution.cpp
@@ -51,9 +51,14 @@ TAutoConsoleVariable<int32> CVarTSRFlickeringEnable(
 	TEXT("Whether to enable the flickering detection heuristic.\n"),
 	ECVF_Scalability | ECVF_RenderThreadSafe);
 
-TAutoConsoleVariable<int32> CVarTSRFlickeringPeriod(
-	TEXT("r.TSR.ShadingRejection.Flickering.Period"), 3,
-	TEXT("Periode in frames in which luma oscilations at equal or greater frequency is considered flickering and should ghost (Default=3.0).\n"),
+TAutoConsoleVariable<int32> CVarTSRFlickeringAdjustToFrameRate(
+	TEXT("r.TSR.ShadingRejection.Flickering.AdjustToFrameRate"), 1,
+	TEXT("Whether TSR settings should adjust to frame rate (Enabled by default, meant for exclusive use of EngineTest).\n"),
+	ECVF_Scalability | ECVF_RenderThreadSafe);
+
+TAutoConsoleVariable<float> CVarTSRFlickeringPeriod(
+	TEXT("r.TSR.ShadingRejection.Flickering.Period"), 3.0f,
+	TEXT("Periode in 60hz frames in which luma oscilations at equal or greater frequency is considered flickering and should ghost (Default=3.0).\n"),
 	ECVF_RenderThreadSafe);
 
 TAutoConsoleVariable<float> CVarTSRFlickeringMaxParralaxVelocity(
@@ -774,13 +779,13 @@ ITemporalUpscaler::FOutputs AddTemporalSuperResolutionPasses(
 
 	const bool bGrandReprojection = !bAccumulateTranslucencySeparately && CVarTSRHistoryGrandReprojection.GetValueOnRenderThread() != 0;
 
-	const float RefreshRateTo60Hz = View.Family->Time.GetDeltaRealTimeSeconds() * 60.0f;
+	const float RefreshRateTo60Hz = View.Family->Time.GetDeltaRealTimeSeconds() > 0.0f && CVarTSRFlickeringAdjustToFrameRate.GetValueOnRenderThread() ? View.Family->Time.GetDeltaRealTimeSeconds() * 60.0f : 1.0f;
 
 	// whether TSR passes can run on async compute.
 	int32 AsyncComputePasses = GSupportsEfficientAsyncCompute ? CVarTSRAsyncCompute.GetValueOnRenderThread() : 0;
 
 	// period at which history changes is considered too distracting.
-	const int32 FlickeringFramePeriod = CVarTSRFlickeringEnable.GetValueOnRenderThread() ? CVarTSRFlickeringPeriod.GetValueOnRenderThread() : 0;
+	const float FlickeringFramePeriod = CVarTSRFlickeringEnable.GetValueOnRenderThread() ? (CVarTSRFlickeringPeriod.GetValueOnRenderThread() / RefreshRateTo60Hz) : 0.0f;
 
 	ETSRHistoryFormatBits HistoryFormatBits = ETSRHistoryFormatBits::None;
 	{
@@ -985,7 +990,7 @@ ITemporalUpscaler::FOutputs AddTemporalSuperResolutionPasses(
 	FRDGTextureSRVRef IsMovingMaskTexture = nullptr;
 	FVelocityFlattenTextures VelocityFlattenTextures;
 	{
-		const bool bOutputIsMovingTexture = FlickeringFramePeriod > 0;
+		const bool bOutputIsMovingTexture = FlickeringFramePeriod > 0.0f;
 
 		FRDGTextureRef R8OutputTexture;
 
@@ -1378,7 +1383,7 @@ ITemporalUpscaler::FOutputs AddTemporalSuperResolutionPasses(
 		PassParameters->DebugOutput = CreateDebugUAV(InputExtent, TEXT("Debug.TSR.DecimateHistory"));
 
 		FTSRDecimateHistoryCS::FPermutationDomain PermutationVector;
-		PermutationVector.Set<FTSRDecimateHistoryCS::FMoireReprojectionDim>(FlickeringFramePeriod > 0);
+		PermutationVector.Set<FTSRDecimateHistoryCS::FMoireReprojectionDim>(FlickeringFramePeriod > 0.0f);
 		PermutationVector.Set<FTSRDecimateHistoryCS::FGrandReprojectionDim>(bGrandReprojection);
 
 		TShaderMapRef<FTSRDecimateHistoryCS> ComputeShader(View.ShaderMap, PermutationVector);
@@ -1488,7 +1493,7 @@ ITemporalUpscaler::FOutputs AddTemporalSuperResolutionPasses(
 		PassParameters->TranslucencyTextureUVMin = GetScreenPassTextureViewportParameters(TranslucencyViewport).UVViewportBilinearMin;
 		PassParameters->TranslucencyTextureUVMax = GetScreenPassTextureViewportParameters(TranslucencyViewport).UVViewportBilinearMax;
 		PassParameters->HistoryGuideQuantizationError = ComputePixelFormatQuantizationError(History.Guide->Desc.Format);
-		PassParameters->FlickeringFramePeriod = float(FlickeringFramePeriod);
+		PassParameters->FlickeringFramePeriod = FlickeringFramePeriod;
 
 		PassParameters->InputTexture = PassInputs.SceneColorTexture;
 		if (bAccumulateTranslucencySeparately)
@@ -1512,12 +1517,12 @@ ITemporalUpscaler::FOutputs AddTemporalSuperResolutionPasses(
 
 		FTSRRejectShadingCS::FPermutationDomain PermutationVector;
 		PermutationVector.Set<FTSRRejectShadingCS::FWaveSizeOps>(bUseWaveOps && GRHIMinimumWaveSize >= 32 && GRHIMinimumWaveSize <= 64 ? GRHIMinimumWaveSize : 0);
-		PermutationVector.Set<FTSRRejectShadingCS::FFlickeringDetectionDim>(FlickeringFramePeriod > 0);
+		PermutationVector.Set<FTSRRejectShadingCS::FFlickeringDetectionDim>(FlickeringFramePeriod > 0.0f);
 
 		TShaderMapRef<FTSRRejectShadingCS> ComputeShader(View.ShaderMap, PermutationVector);
 		FComputeShaderUtils::AddPass(
 			GraphBuilder,
-			RDG_EVENT_NAME("TSR RejectShading(WaveSize=%d FlickeringFramePeriod=%d %s) %dx%d",
+			RDG_EVENT_NAME("TSR RejectShading(WaveSize=%d FlickeringFramePeriod=%f %s) %dx%d",
 				int32(PermutationVector.Get<FTSRRejectShadingCS::FWaveSizeOps>()),
 				FlickeringFramePeriod,
 				!bAccumulateTranslucencySeparately ? TEXT(" ComposeTranslucency") : TEXT(""),
@@ -1794,7 +1799,7 @@ ITemporalUpscaler::FOutputs AddTemporalSuperResolutionPasses(
 		// Extract history guide
 		GraphBuilder.QueueTextureExtraction(History.Guide, &OutputHistory.Guide);
 
-		if (FlickeringFramePeriod > 0)
+		if (FlickeringFramePeriod > 0.0f)
 		{
 			GraphBuilder.QueueTextureExtraction(History.Moire, &OutputHistory.Moire);
 		}
