Clean up Global SDF update frequency logic

diff --git a/Engine/Source/Runtime/Renderer/Private/GlobalDistanceField.cpp b/Engine/Source/Runtime/Renderer/Private/GlobalDistanceField.cpp
--- a/Engine/Source/Runtime/Renderer/Private/GlobalDistanceField.cpp
+++ b/Engine/Source/Runtime/Renderer/Private/GlobalDistanceField.cpp
@@ -82,7 +82,7 @@ int32 GAOGlobalDistanceFieldClipmapUpdatesPerFrame = 2;
 FAutoConsoleVariableRef CVarAOGlobalDistanceFieldClipmapUpdatesPerFrame(
 	TEXT("r.AOGlobalDistanceFieldClipmapUpdatesPerFrame"),
 	GAOGlobalDistanceFieldClipmapUpdatesPerFrame,
-	TEXT("How many clipmaps to update each frame, only 1 or 2 supported.  With values less than 2, the first clipmap is only updated every other frame, which can cause incorrect self occlusion during movement."),
+	TEXT("How many clipmaps to update each frame.  With values less than 2, the first clipmap is only updated every other frame, which can cause incorrect self occlusion during movement."),
 	ECVF_Scalability | ECVF_RenderThreadSafe
 	);
 
@@ -585,89 +585,33 @@ static void AddUpdateBoundsForAxis(FIntVector MovementInPages,
 
 static void GetUpdateFrequencyForClipmap(int32 ClipmapIndex, int32 NumClipmaps, int32& OutFrequency, int32& OutPhase)
 {
-	if (!GAOGlobalDistanceFieldStaggeredUpdates)
+	const int32 NumClipmapUpdatesPerFrame = GAOGlobalDistanceFieldStaggeredUpdates ? FMath::Min(GetNumClipmapUpdatesPerFrame(), NumClipmaps) : NumClipmaps;
+
+	if (ClipmapIndex < NumClipmapUpdatesPerFrame - 1)
 	{
+		// update the first N-1 clipmaps every frame
 		OutFrequency = 1;
 		OutPhase = 0;
+		return;
 	}
-	else if (GetNumClipmapUpdatesPerFrame() == 1)
+
+	// remaining clipmaps update at different frequencies so that only one is updated each frame
+
+	const int32 RemIndex = ClipmapIndex - (NumClipmapUpdatesPerFrame - 1);
+
+	if (ClipmapIndex == NumClipmaps - 1)
 	{
-		if (ClipmapIndex == 0)
-		{
-			OutFrequency = 2;
-			OutPhase = 0;
-		}
-		else if (ClipmapIndex == 1)
-		{
-			OutFrequency = 4;
-			OutPhase = 1;
-		}
-		else if (ClipmapIndex == 2)
-		{
-			OutFrequency = 8;
-			OutPhase = 3;
-		}
-		else
-		{
-			if (NumClipmaps > 4)
-			{
-				if (ClipmapIndex == 3)
-				{
-					OutFrequency = 16;
-					OutPhase = 7;
-				}
-				else
-				{
-					OutFrequency = 16;
-					OutPhase = 15;
-				}
-			}
-			else
-			{
-				OutFrequency = 8;
-				OutPhase = 7;
-			}
-		}
+		// last clipmap updates at same frequency as previous
+		OutFrequency = 1 << RemIndex;
 	}
 	else
 	{
-		if (ClipmapIndex == 0)
-		{
-			OutFrequency = 1;
-			OutPhase = 0;
-		}
-		else if (ClipmapIndex == 1)
-		{
-			OutFrequency = 2;
-			OutPhase = 0;
-		}
-		else if (ClipmapIndex == 2)
-		{
-			OutFrequency = 4;
-			OutPhase = 1;
-		}
-		else
-		{
-			if (NumClipmaps > 4)
-			{
-				if (ClipmapIndex == 3)
-				{
-					OutFrequency = 8;
-					OutPhase = 3;
-				}
-				else
-				{
-					OutFrequency = 8;
-					OutPhase = 7;
-				}
-			}
-			else
-			{
-				OutFrequency = 4;
-				OutPhase = 3;
-			}
-		}
+		// each clipmap updates with half frequency of previous
+		OutFrequency = 2 << RemIndex;
 	}
+		
+	// every clipmap uses a different phase
+	OutPhase = (1 << RemIndex) - 1; // 2^n - 1 sequence
 }
 
 /** Staggers clipmap updates so there are only 2 per frame */
