Fix a bug with uploading a Nanite material hit proxy table when there are more materials than mesh slots, and a mesh slot refers to a material index beyond that of the mesh slot indices.

#rnx

diff --git a/Engine/Source/Runtime/Renderer/Private/GPUScene.cpp b/Engine/Source/Runtime/Renderer/Private/GPUScene.cpp
--- a/Engine/Source/Runtime/Renderer/Private/GPUScene.cpp
+++ b/Engine/Source/Runtime/Renderer/Private/GPUScene.cpp
@@ -1084,7 +1084,7 @@ void FGPUScene::UploadGeneral(FRDGBuilder& GraphBuilder, FScene& Scene, FRDGExte
 							if (NaniteMeshPass == ENaniteMeshPass::BasePass && NaniteSceneProxy->GetHitProxyMode() == Nanite::FSceneProxyBase::EHitProxyMode::MaterialSection)
 							{
 								const TArray<uint32>& PassHitProxyIds = PrimitiveSceneInfo->NaniteHitProxyIds;
-								uint32* HitProxyTable = static_cast<uint32*>(NaniteMaterialUploader->GetHitProxyTablePtr(UploadInfo.PrimitiveID, MaterialSlotCount));
+								uint32* HitProxyTable = static_cast<uint32*>(NaniteMaterialUploader->GetHitProxyTablePtr(UploadInfo.PrimitiveID, UploadEntryCount));
 								for (int32 Entry = 0; Entry < PassHitProxyIds.Num(); ++Entry)
 								{
 									HitProxyTable[PassMaterials[Entry].MaterialIndex] = PassHitProxyIds[Entry];
