Renderer: Add proper include path for HaltonUtilities.ush

#rnx

diff --git a/Engine/Source/Runtime/Renderer/Private/HaltonUtilities.cpp b/Engine/Source/Runtime/Renderer/Private/HaltonUtilities.cpp
--- a/Engine/Source/Runtime/Renderer/Private/HaltonUtilities.cpp
+++ b/Engine/Source/Runtime/Renderer/Private/HaltonUtilities.cpp
@@ -41,7 +41,7 @@ float asfloat(unsigned int integer)
 	return value.f;
 }
 
-#include "../../../Shaders/Private/HaltonUtilities.ush"
+#include "HaltonUtilities.ush"
 
 /**
  * Binds dimensions to consecutive Halton sequences. Note that sequences based on primes
diff --git a/Engine/Source/Runtime/Renderer/Renderer.Build.cs b/Engine/Source/Runtime/Renderer/Renderer.Build.cs
--- a/Engine/Source/Runtime/Renderer/Renderer.Build.cs
+++ b/Engine/Source/Runtime/Renderer/Renderer.Build.cs
@@ -8,7 +8,8 @@ public class Renderer : ModuleRules
 	{
 		PrivateIncludePaths.AddRange(
 			new string[] {
-				"../Shaders/Shared"
+				EngineDirectory + "/Shaders/Shared",
+				EngineDirectory + "/Shaders/Private", // For HaltonUtilities.ush
 				}
 			);
 
