Fix merge error.

diff --git a/Engine/Source/Runtime/Renderer/Private/DeferredShadingRenderer.cpp b/Engine/Source/Runtime/Renderer/Private/DeferredShadingRenderer.cpp
--- a/Engine/Source/Runtime/Renderer/Private/DeferredShadingRenderer.cpp
+++ b/Engine/Source/Runtime/Renderer/Private/DeferredShadingRenderer.cpp
@@ -3780,11 +3780,6 @@ void FDeferredShadingSceneRenderer::Render(FRDGBuilder& GraphBuilder)
 		ShadingEnergyConservation::Debug(GraphBuilder, View, SceneTextures);
 	}
 
-	for (FViewInfo& View : Views)
-	{
-		ShadingEnergyConservation::Debug(GraphBuilder, View, SceneTextures);
-	}
-
 	if (!bHasRayTracedOverlay && ViewFamily.EngineShowFlags.LightShafts)
 	{
 		SCOPE_CYCLE_COUNTER(STAT_FDeferredShadingSceneRenderer_RenderLightShaftBloom);
