Volumetric cloud
- The game thread function no longer access render thread data.
- The blue noise texture is now loaded when volumetric clouds can potentially be enabled.
- r.volumetriccloud is set to 0 on platform not using volumetric cloud features as an optimisation to no load the blue noise texture for our game on some platform.

diff --git a/Engine/Source/Runtime/Renderer/Private/SceneRendering.cpp b/Engine/Source/Runtime/Renderer/Private/SceneRendering.cpp
--- a/Engine/Source/Runtime/Renderer/Private/SceneRendering.cpp
+++ b/Engine/Source/Runtime/Renderer/Private/SceneRendering.cpp
@@ -2515,7 +2515,7 @@ FSceneRenderer::FSceneRenderer(const FSceneViewFamily* InViewFamily, FHitProxyCo
 		}
 #endif
 
-		if (ShouldRenderLumenDiffuseGI(Scene, *ViewInfo) || ShouldRenderLumenReflections(*ViewInfo) || ShouldRenderVolumetricCloudWithBlueNoise(Scene, ViewFamily.EngineShowFlags))
+		if (ShouldRenderLumenDiffuseGI(Scene, *ViewInfo) || ShouldRenderLumenReflections(*ViewInfo) || ShouldRenderVolumetricCloudWithBlueNoise_GameThread(Scene, *ViewInfo))
 		{
 			GEngine->LoadBlueNoiseTexture();
 		}
diff --git a/Engine/Source/Runtime/Renderer/Private/VolumetricCloudRendering.cpp b/Engine/Source/Runtime/Renderer/Private/VolumetricCloudRendering.cpp
--- a/Engine/Source/Runtime/Renderer/Private/VolumetricCloudRendering.cpp
+++ b/Engine/Source/Runtime/Renderer/Private/VolumetricCloudRendering.cpp
@@ -272,9 +272,10 @@ bool ShouldRenderVolumetricCloud(const FScene* Scene, const FEngineShowFlags& En
 	return false;
 }
 
-bool ShouldRenderVolumetricCloudWithBlueNoise(const FScene* Scene, const FEngineShowFlags& EngineShowFlags)
+bool ShouldRenderVolumetricCloudWithBlueNoise_GameThread(const FScene* Scene, const FSceneView& View)
 {
-	return Scene && Scene->HasVolumetricCloud() && EngineShowFlags.Atmosphere && EngineShowFlags.Cloud;
+	// We cannot use Scene->HasVolumetricCloud() here because the proxy is set form the render thread.
+	return CVarVolumetricCloud.GetValueOnGameThread() > 0 && View.Family->EngineShowFlags.Atmosphere && View.Family->EngineShowFlags.Cloud;
 }
 
 bool ShouldViewVisualizeVolumetricCloudConservativeDensity(const FViewInfo& ViewInfo, const FEngineShowFlags& EngineShowFlags)
diff --git a/Engine/Source/Runtime/Renderer/Private/VolumetricCloudRendering.h b/Engine/Source/Runtime/Renderer/Private/VolumetricCloudRendering.h
--- a/Engine/Source/Runtime/Renderer/Private/VolumetricCloudRendering.h
+++ b/Engine/Source/Runtime/Renderer/Private/VolumetricCloudRendering.h
@@ -105,7 +105,8 @@ private:
 
 
 bool ShouldRenderVolumetricCloud(const FScene* Scene, const FEngineShowFlags& EngineShowFlags);
-bool ShouldRenderVolumetricCloudWithBlueNoise(const FScene* Scene, const FEngineShowFlags& EngineShowFlags);
+bool ShouldRenderVolumetricCloudWithBlueNoise_GameThread(const FScene* Scene, const FSceneView& View);
+
 bool ShouldViewVisualizeVolumetricCloudConservativeDensity(const FViewInfo& ViewInfo, const FEngineShowFlags& EngineShowFlags);
 bool VolumetricCloudWantsToSampleLocalLights(const FScene* Scene, const FEngineShowFlags& EngineShowFlags);
 uint32 GetVolumetricCloudDebugViewMode(const FEngineShowFlags& ShowFlags);
