Added missing TexCreate_ShaderResource flag to RayTracedShadows texture.

diff --git a/Engine/Source/Runtime/Renderer/Private/DistanceFieldShadowing.cpp b/Engine/Source/Runtime/Renderer/Private/DistanceFieldShadowing.cpp
--- a/Engine/Source/Runtime/Renderer/Private/DistanceFieldShadowing.cpp
+++ b/Engine/Source/Runtime/Renderer/Private/DistanceFieldShadowing.cpp
@@ -929,7 +929,7 @@ FRDGTextureRef FProjectedShadowInfo::RenderRayTracedDistanceFieldProjection(
 
 			{
 				const FIntPoint BufferSize = GetBufferSizeForDFShadows(View);
-				FRDGTextureDesc Desc(FRDGTextureDesc::Create2D(BufferSize, PF_G16R16F, FClearValueBinding::None, TexCreate_UAV));
+				FRDGTextureDesc Desc(FRDGTextureDesc::Create2D(BufferSize, PF_G16R16F, FClearValueBinding::None, TexCreate_UAV | TexCreate_ShaderResource));
 				Desc.Flags |= GFastVRamConfig.DistanceFieldShadows;
 				SDFShadowViewGPUData.RayTracedShadowsTexture = GraphBuilder.CreateTexture(Desc, TEXT("RayTracedShadows"));
 			}
@@ -994,7 +994,7 @@ FRDGTextureRef FProjectedShadowInfo::RenderRayTracedDistanceFieldProjection(
 		if (!SDFShadowViewGPUData.RayTracedShadowsTexture)
 		{
 			const FIntPoint BufferSize = GetBufferSizeForDFShadows(View);
-			FRDGTextureDesc Desc(FRDGTextureDesc::Create2D(BufferSize, PF_G16R16F, FClearValueBinding::None, TexCreate_UAV));
+			FRDGTextureDesc Desc(FRDGTextureDesc::Create2D(BufferSize, PF_G16R16F, FClearValueBinding::None, TexCreate_UAV | TexCreate_ShaderResource));
 			Desc.Flags |= GFastVRamConfig.DistanceFieldShadows;
 			SDFShadowViewGPUData.RayTracedShadowsTexture = GraphBuilder.CreateTexture(Desc, TEXT("RayTracedShadows"));
 		}
