Clear DistanceFieldInstanceIndices when resetting HeightfieldPrimitives to prevent stale indices.

diff --git a/Engine/Source/Runtime/Renderer/Private/DistanceFieldObjectManagement.cpp b/Engine/Source/Runtime/Renderer/Private/DistanceFieldObjectManagement.cpp
--- a/Engine/Source/Runtime/Renderer/Private/DistanceFieldObjectManagement.cpp
+++ b/Engine/Source/Runtime/Renderer/Private/DistanceFieldObjectManagement.cpp
@@ -340,10 +340,7 @@ void UpdateGlobalHeightFieldObjectRemoves(FScene* Scene)
 					PrimitiveBeingMoved->DistanceFieldInstanceIndices[0] = RemoveIndex;
 				}
 
-				if (SceneData.HeightfieldPrimitives.IsValidIndex(RemoveIndex))
-				{
-					SceneData.HeightfieldPrimitives.RemoveAtSwap(RemoveIndex, 1, false);
-				}
+				SceneData.HeightfieldPrimitives.RemoveAtSwap(RemoveIndex, 1, false);
 
 				SceneData.IndicesToUpdateInHeightFieldObjectBuffers.Add(RemoveIndex);
 			}
@@ -837,6 +834,7 @@ void FSceneRenderer::UpdateGlobalHeightFieldObjectBuffers(FRDGBuilder& GraphBuil
 				FPrimitiveSceneInfo* Primitive = DistanceFieldSceneData.HeightfieldPrimitives[Idx];
 				check(!DistanceFieldSceneData.PendingHeightFieldAddOps.Contains(Primitive));
 				DistanceFieldSceneData.PendingHeightFieldAddOps.Add(Primitive);
+				Primitive->DistanceFieldInstanceIndices.Empty(1);
 			}
 			DistanceFieldSceneData.HeightfieldPrimitives.Reset();
 			DistanceFieldSceneData.HeightFieldObjectBuffers = new FHeightFieldObjectBuffers;
