Missed file from previous CL23061754

diff --git a/Engine/Source/Runtime/Renderer/Private/SceneRendering.cpp b/Engine/Source/Runtime/Renderer/Private/SceneRendering.cpp
--- a/Engine/Source/Runtime/Renderer/Private/SceneRendering.cpp
+++ b/Engine/Source/Runtime/Renderer/Private/SceneRendering.cpp
@@ -2515,7 +2515,7 @@ FSceneRenderer::FSceneRenderer(const FSceneViewFamily* InViewFamily, FHitProxyCo
 		}
 #endif
 
-		if (ShouldRenderLumenDiffuseGI(Scene, *ViewInfo) || ShouldRenderLumenReflections(*ViewInfo) || ShouldRenderVolumetricCloudWithBlueNoise_GameThread(Scene, *ViewInfo))
+		if (ShouldRenderLumenDiffuseGI(Scene, *ViewInfo) || ShouldRenderLumenReflections(*ViewInfo) || ShouldRenderVolumetricCloudWithBlueNoise_GameThread(Scene, *ViewInfo) || UseVirtualShadowMaps(Scene->GetShaderPlatform(), Scene->GetFeatureLevel()))
 		{
 			GEngine->LoadBlueNoiseTexture();
 		}
