Fix typo/merge conflict

diff --git a/Engine/Source/Runtime/Renderer/Private/VirtualShadowMaps/VirtualShadowMapProjection.cpp b/Engine/Source/Runtime/Renderer/Private/VirtualShadowMaps/VirtualShadowMapProjection.cpp
--- a/Engine/Source/Runtime/Renderer/Private/VirtualShadowMaps/VirtualShadowMapProjection.cpp
+++ b/Engine/Source/Runtime/Renderer/Private/VirtualShadowMaps/VirtualShadowMapProjection.cpp
@@ -429,7 +429,7 @@ static void RenderVirtualShadowMapProjectionCommon(
 		GraphBuilder.AddPass(
 			RDG_EVENT_NAME("VirtualShadowMapProjection(RayCount:%u(%s),SamplesPerRay:%u,Input:%s%s,TileList)",
 				PassParameters->SMRTRayCount,
-				PassParameters->SMRTUseAdaptiveRayCount > 0 ? TEXT("Adaptive") : TEXT("Static"),
+				PassParameters->SMRTAdaptiveRayCount > 0 ? TEXT("Adaptive") : TEXT("Static"),
 				PassParameters->SMRTSamplesPerRay,
 				ToString(InputType),
 				bDebugOutput ? TEXT(",Debug") : TEXT("")),
@@ -447,7 +447,7 @@ static void RenderVirtualShadowMapProjectionCommon(
 			GraphBuilder,
 			RDG_EVENT_NAME("VirtualShadowMapProjection(RayCount:%u(%s),SamplesPerRay:%u,Input:%s%s)",
 				PassParameters->SMRTRayCount,
-				PassParameters->SMRTUseAdaptiveRayCount > 0 ? TEXT("Adaptive") : TEXT("Static"),
+				PassParameters->SMRTAdaptiveRayCount > 0 ? TEXT("Adaptive") : TEXT("Static"),
 				PassParameters->SMRTSamplesPerRay,
 				ToString(InputType),
 				bDebugOutput ? TEXT(",Debug") : TEXT("")),
