Fixed crash with volumetric cloud test and separated mie/rayleigh scattering.

diff --git a/Engine/Source/Runtime/Renderer/Private/VolumetricCloudRendering.cpp b/Engine/Source/Runtime/Renderer/Private/VolumetricCloudRendering.cpp
--- a/Engine/Source/Runtime/Renderer/Private/VolumetricCloudRendering.cpp
+++ b/Engine/Source/Runtime/Renderer/Private/VolumetricCloudRendering.cpp
@@ -2819,8 +2819,8 @@ bool FSceneRenderer::RenderVolumetricCloud(
 					{
 						SkyRC.SkyAtmosphereViewLutTexture = GraphBuilder.RegisterExternalTexture(ViewInfo.SkyAtmosphereViewLutTexture);
 						SkyRC.SkyAtmosphereCameraAerialPerspectiveVolume = GraphBuilder.RegisterExternalTexture(ViewInfo.SkyAtmosphereCameraAerialPerspectiveVolume);
-						SkyRC.SkyAtmosphereCameraAerialPerspectiveVolumeMieOnly = GraphBuilder.RegisterExternalTexture(ViewInfo.SkyAtmosphereCameraAerialPerspectiveVolumeMieOnly);
-						SkyRC.SkyAtmosphereCameraAerialPerspectiveVolumeRayOnly = GraphBuilder.RegisterExternalTexture(ViewInfo.SkyAtmosphereCameraAerialPerspectiveVolumeRayOnly);
+						SkyRC.SkyAtmosphereCameraAerialPerspectiveVolumeMieOnly = GraphBuilder.RegisterExternalTexture(ViewInfo.SkyAtmosphereCameraAerialPerspectiveVolumeMieOnly.IsValid() ? ViewInfo.SkyAtmosphereCameraAerialPerspectiveVolumeMieOnly : GSystemTextures.VolumetricBlackAlphaOneDummy);
+						SkyRC.SkyAtmosphereCameraAerialPerspectiveVolumeRayOnly = GraphBuilder.RegisterExternalTexture(ViewInfo.SkyAtmosphereCameraAerialPerspectiveVolumeRayOnly.IsValid() ? ViewInfo.SkyAtmosphereCameraAerialPerspectiveVolumeRayOnly : GSystemTextures.VolumetricBlackAlphaOneDummy);
 					}
 					else
 					{
