Use GetSize() instead of 2 * GetExtent() in GlobalDistanceField.cpp

diff --git a/Engine/Source/Runtime/Renderer/Private/GlobalDistanceField.cpp b/Engine/Source/Runtime/Renderer/Private/GlobalDistanceField.cpp
--- a/Engine/Source/Runtime/Renderer/Private/GlobalDistanceField.cpp
+++ b/Engine/Source/Runtime/Renderer/Private/GlobalDistanceField.cpp
@@ -508,11 +508,11 @@ void FGlobalDistanceFieldInfo::UpdateParameterData(float MaxOcclusionDistance, b
 				// GlobalUV = (WorldPosition - GlobalVolumeCenterAndExtent[ClipmapIndex].xyz + GlobalVolumeScollOffset[ClipmapIndex].xyz) / (GlobalVolumeCenterAndExtent[ClipmapIndex].w * 2) + .5f;
 				// WorldToUVMul = 1.0f / (GlobalVolumeCenterAndExtent[ClipmapIndex].w * 2)
 				// WorldToUVAdd = (GlobalVolumeScollOffset[ClipmapIndex].xyz - GlobalVolumeCenterAndExtent[ClipmapIndex].xyz) / (GlobalVolumeCenterAndExtent[ClipmapIndex].w * 2) + .5f
-				const FVector WorldToUVAdd = (Clipmap.ScrollOffset - Clipmap.Bounds.GetCenter()) / (Clipmap.Bounds.GetExtent().X * 2) + FVector(.5f);
-				ParameterData.WorldToUVAddAndMul[ClipmapIndex] = FVector4f((FVector3f)WorldToUVAdd, 1.0f / (Clipmap.Bounds.GetExtent().X * 2));
+				const FVector WorldToUVAdd = (Clipmap.ScrollOffset - Clipmap.Bounds.GetCenter()) / Clipmap.Bounds.GetSize().X + FVector(.5f);
+				ParameterData.WorldToUVAddAndMul[ClipmapIndex] = FVector4f((FVector3f)WorldToUVAdd, 1.0f / Clipmap.Bounds.GetSize().X);
 
-				ParameterData.MipWorldToUVScale[ClipmapIndex] = FVector3f(FVector(1.0f) / (2.0f * Clipmap.Bounds.GetExtent())); // LWC_TODO: precision loss
-				ParameterData.MipWorldToUVBias[ClipmapIndex] = FVector3f((-Clipmap.Bounds.Min) / (2.0f * Clipmap.Bounds.GetExtent())); // LWC_TODO: precision loss
+				ParameterData.MipWorldToUVScale[ClipmapIndex] = FVector3f(FVector(1.0f) / Clipmap.Bounds.GetSize());
+				ParameterData.MipWorldToUVBias[ClipmapIndex] = FVector3f((-Clipmap.Bounds.Min) / Clipmap.Bounds.GetSize()); // LWC_TODO: precision loss
 
 				ParameterData.MipWorldToUVScale[ClipmapIndex].Z = ParameterData.MipWorldToUVScale[ClipmapIndex].Z / Clipmaps.Num();
 				ParameterData.MipWorldToUVBias[ClipmapIndex].Z = (ParameterData.MipWorldToUVBias[ClipmapIndex].Z + ClipmapIndex) / Clipmaps.Num();
@@ -541,8 +541,8 @@ void FGlobalDistanceFieldInfo::UpdateParameterData(float MaxOcclusionDistance, b
 		ParameterData.GlobalDFResolution = GlobalDistanceField::GetClipmapResolution(bLumenEnabled);
 
 		extern float GAOConeHalfAngle;
-		const float MaxClipmapExtentX = Clipmaps[Clipmaps.Num() - 1].Bounds.GetExtent().X;
-		const float MaxClipmapVoxelSize =  (2.0f * MaxClipmapExtentX) / GlobalDistanceField::GetClipmapResolution(bLumenEnabled);
+		const float MaxClipmapSizeX = Clipmaps[Clipmaps.Num() - 1].Bounds.GetSize().X;
+		const float MaxClipmapVoxelSize =  MaxClipmapSizeX / GlobalDistanceField::GetClipmapResolution(bLumenEnabled);
 		float MaxClipmapInfluenceRadius = GGlobalDistanceFieldInfluenceRangeInVoxels * MaxClipmapVoxelSize;
 		const float GlobalMaxSphereQueryRadius = FMath::Min(MaxOcclusionDistance / (1.0f + FMath::Tan(GAOConeHalfAngle)), MaxClipmapInfluenceRadius);
 		ParameterData.MaxDFAOConeDistance = GlobalMaxSphereQueryRadius;
@@ -2071,8 +2071,8 @@ void UpdateGlobalDistanceFieldVolume(
 					const FVector ClipmapVoxelExtent = 0.5f * ClipmapVoxelSize;
 					const float ClipmapInfluenceRadius = (GGlobalDistanceFieldInfluenceRangeInVoxels * ClipmapSize.X) / ClipmapResolution;
 
-					const FVector TranslatedWorldToUVAdd = (Clipmap.ScrollOffset - TranslatedBounds.GetCenter()) / (TranslatedBounds.GetExtent().X * 2.0f) + FVector(0.5f);
-					const FVector4f ClipmapVolumeTranslatedWorldToUVAddAndMul = FVector4f((FVector3f)TranslatedWorldToUVAdd, 1.0f / (TranslatedBounds.GetExtent().X * 2.0f));
+					const FVector TranslatedWorldToUVAdd = (Clipmap.ScrollOffset - TranslatedBounds.GetCenter()) / TranslatedBounds.GetSize().X + FVector(0.5f);
+					const FVector4f ClipmapVolumeTranslatedWorldToUVAddAndMul = FVector4f((FVector3f)TranslatedWorldToUVAdd, 1.0f / TranslatedBounds.GetSize().X);
 
 					int32 MaxSDFMeshObjects = FMath::RoundUpToPowerOfTwo(DistanceFieldSceneData.NumObjectsInBuffer);
 					FRDGBufferRef ObjectIndexBuffer = GraphBuilder.CreateBuffer(FRDGBufferDesc::CreateStructuredDesc(sizeof(uint32), MaxSDFMeshObjects), TEXT("GlobalDistanceField.ObjectIndices"));
