Fix assert on IsInActualRenderingThread() by using IsInParallelRenderingThread() instead and inverting condition

#rnx

diff --git a/Engine/Source/Runtime/Renderer/Private/MeshDrawCommands.cpp b/Engine/Source/Runtime/Renderer/Private/MeshDrawCommands.cpp
--- a/Engine/Source/Runtime/Renderer/Private/MeshDrawCommands.cpp
+++ b/Engine/Source/Runtime/Renderer/Private/MeshDrawCommands.cpp
@@ -1483,7 +1483,7 @@ void FParallelMeshDrawCommandPass::DispatchDraw(FParallelCommandListSet* Paralle
 	{
 		QUICK_SCOPE_CYCLE_COUNTER(STAT_MeshPassDrawImmediate);
 
-		WaitForMeshPassSetupTask(IsInActualRenderingThread() ? EWaitThread::Render : EWaitThread::Task);
+		WaitForMeshPassSetupTask(IsInParallelRenderingThread() ? EWaitThread::Task : EWaitThread::Render);
 
 		if (TaskContext.bUseGPUScene)
 		{
