Fix distance field streaming on builds WITH_EDITOR.

diff --git a/Engine/Source/Runtime/Renderer/Private/DistanceFieldStreaming.cpp b/Engine/Source/Runtime/Renderer/Private/DistanceFieldStreaming.cpp
--- a/Engine/Source/Runtime/Renderer/Private/DistanceFieldStreaming.cpp
+++ b/Engine/Source/Runtime/Renderer/Private/DistanceFieldStreaming.cpp
@@ -472,17 +472,18 @@ void FDistanceFieldSceneData::AsyncUpdate(FDistanceFieldAsyncUpdateParameters&&
 		const int32 MipIndex = AssetState.BuiltData->Mips.Num() - ReversedMipIndex - 1;
 		const FSparseDistanceFieldMip& MipBuiltData = AssetState.BuiltData->Mips[MipIndex];
 
-		const uint8* BulkDataReadPtr = ReadRequest.BulkData != nullptr
-			? ReadRequest.RequestBuffer.GetData()
-			: ReadRequest.AlwaysLoadedDataPtr;
+		const uint8* BulkDataReadPtr = ReadRequest.AlwaysLoadedDataPtr;
 
-#if WITH_EDITOR
 		if (ReadRequest.BulkData)
-		{	
+		{
+#if WITH_EDITOR
 			check((ReadRequest.BulkData->IsBulkDataLoaded() ||ReadRequest.BulkData->CanLoadFromDisk()) && ReadRequest.BulkData->GetBulkDataSize() > 0);
 			BulkDataReadPtr = (const uint8*)ReadRequest.BulkData->LockReadOnly() + ReadRequest.BulkOffset;
-		}
+#else
+			BulkDataReadPtr = ReadRequest.RequestBuffer.GetData();
 #endif
+		}
+
 		const int32 NumIndirectionEntries = MipBuiltData.IndirectionDimensions.X * MipBuiltData.IndirectionDimensions.Y * MipBuiltData.IndirectionDimensions.Z;
 		const uint32 ExpectedBulkSize = NumIndirectionEntries * sizeof(uint32) + ReadRequest.NumDistanceFieldBricks * BrickSizeBytes;
 
@@ -595,6 +596,8 @@ void FDistanceFieldSceneData::AsyncUpdate(FDistanceFieldAsyncUpdateParameters&&
 		FBulkDataRequest::EStatus Status = Batch.Issue();
 		UE_CLOG(Status != FBulkDataRequest::EStatus::Ok, LogDistanceField, Error, TEXT("Failed to issue bulk data I/O request"));
 	}
+#else
+	ReadRequests.Append(UpdateParameters.NewReadRequests);
 #endif
 }
 
@@ -770,13 +773,17 @@ void FDistanceFieldSceneData::ProcessReadRequests(
 
 	for (int32 RequestIndex = 0; RequestIndex < ReadRequests.Num(); RequestIndex++)
 	{
+#if !WITH_EDITOR
 		if (ReadRequests[RequestIndex].RequestHandle.IsCompleted() == false)
 		{
 			continue;
 		}
+#endif
 
 		FDistanceFieldReadRequest CompletedRequest = MoveTemp(ReadRequests[RequestIndex]);
+#if !WITH_EDITOR
 		CompletedRequest.RequestHandle = FBulkDataRequest();
+#endif
 
 		ReadRequests.RemoveAtSwap(RequestIndex--);
 
diff --git a/Engine/Source/Runtime/Renderer/Private/ScenePrivate.h b/Engine/Source/Runtime/Renderer/Private/ScenePrivate.h
--- a/Engine/Source/Runtime/Renderer/Private/ScenePrivate.h
+++ b/Engine/Source/Runtime/Renderer/Private/ScenePrivate.h
@@ -1938,9 +1938,11 @@ struct FDistanceFieldReadRequest
 	uint32 BulkOffset = 0;
 	uint32 BulkSize = 0;
 
+#if !WITH_EDITOR
 	// Outputs of read request
 	FBulkDataBatchReadRequest RequestHandle;
 	FIoBuffer RequestBuffer;
+#endif
 };
 
 struct FDistanceFieldAsyncUpdateParameters
