Track PrimitiveModifiedBounds using FBox instead of FRenderBounds to avoid precision issues in LWC.

diff --git a/Engine/Source/Runtime/Renderer/Private/GlobalDistanceField.cpp b/Engine/Source/Runtime/Renderer/Private/GlobalDistanceField.cpp
--- a/Engine/Source/Runtime/Renderer/Private/GlobalDistanceField.cpp
+++ b/Engine/Source/Runtime/Renderer/Private/GlobalDistanceField.cpp
@@ -1008,18 +1008,18 @@ static void ComputeUpdateRegionsAndUpdateViewState(
 						? &GlobalDistanceFieldInfo.MostlyStaticClipmaps[ClipmapIndex]
 						: &GlobalDistanceFieldInfo.Clipmaps[ClipmapIndex]);
 
-					const TArray<FRenderBounds>& PrimitiveModifiedBounds = ClipmapViewState.Cache[CacheType].PrimitiveModifiedBounds;
+					const TArray<FBox>& PrimitiveModifiedBounds = ClipmapViewState.Cache[CacheType].PrimitiveModifiedBounds;
 
-					TArray<FRenderBounds, SceneRenderingAllocator> CulledPrimitiveModifiedBounds;
-					CulledPrimitiveModifiedBounds.Empty(ClipmapViewState.Cache[CacheType].PrimitiveModifiedBounds.Num() / 2);
+					TArray<FBox, SceneRenderingAllocator> CulledPrimitiveModifiedBounds;
+					CulledPrimitiveModifiedBounds.Empty(PrimitiveModifiedBounds.Num() / 2);
 
-					Clipmap.UpdateBounds.Empty(ClipmapViewState.Cache[CacheType].PrimitiveModifiedBounds.Num() / 2);
+					Clipmap.UpdateBounds.Empty(PrimitiveModifiedBounds.Num() / 2);
 
-					for (int32 BoundsIndex = 0; BoundsIndex < ClipmapViewState.Cache[CacheType].PrimitiveModifiedBounds.Num(); BoundsIndex++)
+					for (int32 BoundsIndex = 0; BoundsIndex < PrimitiveModifiedBounds.Num(); BoundsIndex++)
 					{
-						const FRenderBounds PrimBounds = ClipmapViewState.Cache[CacheType].PrimitiveModifiedBounds[BoundsIndex];
-						const FVector PrimWorldCenter = (FVector)PrimBounds.GetCenter();
-						const FVector PrimWorldExtent = (FVector)PrimBounds.GetExtent();
+						const FBox PrimBounds = PrimitiveModifiedBounds[BoundsIndex];
+						const FVector PrimWorldCenter = PrimBounds.GetCenter();
+						const FVector PrimWorldExtent = PrimBounds.GetExtent();
 						const FBox ModifiedBounds(PrimWorldCenter - PrimWorldExtent, PrimWorldCenter + PrimWorldExtent);
 
 						if (ModifiedBounds.ComputeSquaredDistanceToBox(ClipmapBounds) < ClipmapInfluenceRadius * ClipmapInfluenceRadius)
diff --git a/Engine/Source/Runtime/Renderer/Private/ScenePrivate.h b/Engine/Source/Runtime/Renderer/Private/ScenePrivate.h
--- a/Engine/Source/Runtime/Renderer/Private/ScenePrivate.h
+++ b/Engine/Source/Runtime/Renderer/Private/ScenePrivate.h
@@ -433,7 +433,7 @@ public:
 class FGlobalDistanceFieldCacheTypeState
 {
 public:
-	TArray<FRenderBounds> PrimitiveModifiedBounds;
+	TArray<FBox> PrimitiveModifiedBounds;
 };
 
 class FGlobalDistanceFieldClipmapState
@@ -2093,7 +2093,7 @@ public:
 	TSet<FPrimitiveSceneInfo*> PendingThrottledOperations;
 	TSet<FPrimitiveSceneInfo*> PendingUpdateOperations;
 	TArray<FPrimitiveRemoveInfo> PendingRemoveOperations;
-	TArray<FRenderBounds> PrimitiveModifiedBounds[GDF_Num];
+	TArray<FBox> PrimitiveModifiedBounds[GDF_Num];
 
 	TSet<FPrimitiveSceneInfo*> PendingHeightFieldAddOps;
 	TSet<FPrimitiveSceneInfo*> PendingHeightFieldUpdateOps;
diff --git a/Engine/Source/Runtime/Renderer/Private/SceneRendering.cpp b/Engine/Source/Runtime/Renderer/Private/SceneRendering.cpp
--- a/Engine/Source/Runtime/Renderer/Private/SceneRendering.cpp
+++ b/Engine/Source/Runtime/Renderer/Private/SceneRendering.cpp
@@ -4318,7 +4318,7 @@ void FSceneRenderer::WaitForCleanUpTasks(FRHICommandListImmediate& RHICmdList)
 	GSceneRenderCleanUpState.bWaitForTasksComplete = true;
 }
 
-void ResetAndShrinkModifiedBounds(TArray<FRenderBounds>& Bounds)
+void ResetAndShrinkModifiedBounds(TArray<FBox>& Bounds)
 {
 	const int32 MaxAllocatedSize = FMath::RoundUpToPowerOfTwo(FMath::Max<uint32>(DistanceField::MinPrimitiveModifiedBoundsAllocation, Bounds.Num()));
 
