Fixed bad depth state

diff --git a/Engine/Source/Runtime/Renderer/Private/ScreenPass.cpp b/Engine/Source/Runtime/Renderer/Private/ScreenPass.cpp
--- a/Engine/Source/Runtime/Renderer/Private/ScreenPass.cpp
+++ b/Engine/Source/Runtime/Renderer/Private/ScreenPass.cpp
@@ -239,7 +239,7 @@ void AddDownsampleDepthPass(
 
 	if (bIsMinAndMaxDepthFilter)
 	{
-		DepthStencilState = TStaticDepthStencilState<false>::GetRHI();
+		DepthStencilState = TStaticDepthStencilState<false, CF_Always>::GetRHI();
 		PassParameters->RenderTargets[0] = FRenderTargetBinding(Output.Texture, Output.LoadAction);
 	}
 	else
