Fix for blue noise texture being loaded synchronously while not on loading screen.
Fixed potential crash with volumetric cloud test and separated mie/rayleigh scattering.
Tested locally on 23087647.

diff --git a/Engine/Source/Runtime/Renderer/Private/VolumetricCloudRendering.cpp b/Engine/Source/Runtime/Renderer/Private/VolumetricCloudRendering.cpp
--- a/Engine/Source/Runtime/Renderer/Private/VolumetricCloudRendering.cpp
+++ b/Engine/Source/Runtime/Renderer/Private/VolumetricCloudRendering.cpp
@@ -275,7 +275,8 @@ bool ShouldRenderVolumetricCloud(const FScene* Scene, const FEngineShowFlags& En
 bool ShouldRenderVolumetricCloudWithBlueNoise_GameThread(const FScene* Scene, const FSceneView& View)
 {
 	// We cannot use Scene->HasVolumetricCloud() here because the proxy is set form the render thread.
-	return CVarVolumetricCloud.GetValueOnGameThread() > 0 && View.Family->EngineShowFlags.Atmosphere && View.Family->EngineShowFlags.Cloud;
+	// We cannot use View.Family->EngineShowFlags.Atmosphere && View.Family->EngineShowFlags.Cloud because those could cause sync loading texture during gameplay when toggled.
+	return CVarVolumetricCloud.GetValueOnGameThread() > 0;
 }
 
 bool ShouldViewVisualizeVolumetricCloudConservativeDensity(const FViewInfo& ViewInfo, const FEngineShowFlags& EngineShowFlags)
