Works arround a source of ghosting in TSR when Effect is on low quality but TSR is on high or epic quality

diff --git a/Engine/Source/Runtime/Renderer/Private/PostProcess/TemporalSuperResolution.cpp b/Engine/Source/Runtime/Renderer/Private/PostProcess/TemporalSuperResolution.cpp
--- a/Engine/Source/Runtime/Renderer/Private/PostProcess/TemporalSuperResolution.cpp
+++ b/Engine/Source/Runtime/Renderer/Private/PostProcess/TemporalSuperResolution.cpp
@@ -59,7 +59,7 @@ TAutoConsoleVariable<int32> CVarTSRFlickeringAdjustToFrameRate(
 TAutoConsoleVariable<float> CVarTSRFlickeringPeriod(
 	TEXT("r.TSR.ShadingRejection.Flickering.Period"), 3.0f,
 	TEXT("Periode in 60hz frames in which luma oscilations at equal or greater frequency is considered flickering and should ghost (Default=3.0).\n"),
-	ECVF_RenderThreadSafe);
+	ECVF_Scalability | ECVF_RenderThreadSafe);
 
 TAutoConsoleVariable<float> CVarTSRFlickeringMaxParralaxVelocity(
 	TEXT("r.TSR.ShadingRejection.Flickering.MaxParralaxVelocity"), 10.0,
