Support Nanite material visibility culling for VSM, r.Shadow.Virtual.Nanite.MaterialVisibility, default off.

diff --git a/Engine/Source/Runtime/Renderer/Private/ShadowSetup.cpp b/Engine/Source/Runtime/Renderer/Private/ShadowSetup.cpp
--- a/Engine/Source/Runtime/Renderer/Private/ShadowSetup.cpp
+++ b/Engine/Source/Runtime/Renderer/Private/ShadowSetup.cpp
@@ -4910,6 +4910,10 @@ void FSceneRenderer::AddViewDependentWholeSceneShadowsForView(
 					ProjectedShadowInfo->SetupClipmapProjection(&LightSceneInfo, &View, VirtualShadowMapClipmap, CVarVsmUseFarShadowRules.GetValueOnRenderThread() != 0 ? MaxNonFarCascadeDistance : -1.0f);
 					VisibleLightInfo.AllProjectedShadows.Add(ProjectedShadowInfo);
 					ShadowInfosThatNeedCulling.Add(ProjectedShadowInfo);
+
+					// TODO: disentangle from legacy setup alltogether
+					// TODO2: This needs to not depend on the non-nanite shadows!
+					ShadowSceneRenderer->AddDirectionalLightShadow(ProjectedShadowInfo);
 				}
 
 				// NOTE: If there are multiple camera views this will simply be associated with "one of them"
diff --git a/Engine/Source/Runtime/Renderer/Private/Shadows/ShadowSceneRenderer.cpp b/Engine/Source/Runtime/Renderer/Private/Shadows/ShadowSceneRenderer.cpp
--- a/Engine/Source/Runtime/Renderer/Private/Shadows/ShadowSceneRenderer.cpp
+++ b/Engine/Source/Runtime/Renderer/Private/Shadows/ShadowSceneRenderer.cpp
@@ -14,6 +14,15 @@
 
 CSV_DECLARE_CATEGORY_EXTERN(VSM);
 
+extern int32 GForceInvalidateDirectionalVSM;
+
+TAutoConsoleVariable<int32> CVarVSMMaterialVisibility(
+	TEXT("r.Shadow.Virtual.Nanite.MaterialVisibility"),
+	0,
+	TEXT("Enable Nanite CPU-side visibility filtering of draw commands, depends on r.Nanite.MaterialVisibility being enabled."),
+	ECVF_RenderThreadSafe
+);
+
 TAutoConsoleVariable<int32> CVarMaxDistantLightsPerFrame(
 	TEXT("r.Shadow.Virtual.MaxDistantUpdatePerFrame"),
 	1,
@@ -100,10 +109,62 @@ TSharedPtr<FVirtualShadowMapPerLightCacheEntry> FShadowSceneRenderer::AddLocalLi
 	return PerLightCacheEntry;
 }
 
+void FShadowSceneRenderer::AddDirectionalLightShadow(FProjectedShadowInfo* ProjectedShadowInfo)
+{
+	FDirectionalLightShadowFrameSetup& DirectionalLightShadowFrameSetup = DirectionalLights.AddDefaulted_GetRef();
+	DirectionalLightShadowFrameSetup.ProjectedShadowInfo = ProjectedShadowInfo;
+}
+
 void FShadowSceneRenderer::PostInitDynamicShadowsSetup()
 {
 	UpdateDistantLightPriorityRender();
 
+	// Dispatch async Nanite culling job if appropriate
+	if (CVarVSMMaterialVisibility.GetValueOnRenderThread() != 0)
+	{
+		TArray<FConvexVolume, SceneRenderingAllocator> NaniteCullingViewsVolumes;
+		// If we have a clipmap that can't be culled, it'd be a complete waste of time to cull the local lights.
+		bool bUnboundedClipmap = false;
+		
+		for (const FDirectionalLightShadowFrameSetup& DirectionalLightShadowFrameSetup : DirectionalLights)
+		{
+			FProjectedShadowInfo* ProjectedShadowInfo = DirectionalLightShadowFrameSetup.ProjectedShadowInfo;
+			if (!bUnboundedClipmap && ProjectedShadowInfo->bShouldRenderVSM)
+			{
+				const bool bIsCached = VirtualShadowMapArray.CacheManager->IsValid() && GForceInvalidateDirectionalVSM == 0;
+
+				// We can only do this culling if the light is both uncached & it is using the accurate bounds (i.e., r.Shadow.Virtual.Clipmap.UseConservativeCulling is turned off).
+				if (!bIsCached && !ProjectedShadowInfo->CascadeSettings.ShadowBoundsAccurate.Planes.IsEmpty())
+				{
+					NaniteCullingViewsVolumes.Add(ProjectedShadowInfo->CascadeSettings.ShadowBoundsAccurate);
+				}
+				else
+				{
+					bUnboundedClipmap = true;
+				}
+			}
+		}
+		if (!bUnboundedClipmap)
+		{
+			for (const FLocalLightShadowFrameSetup& LocalLightShadowFrameSetup : LocalLights)
+			{
+				FProjectedShadowInfo* ProjectedShadowInfo = LocalLightShadowFrameSetup.ProjectedShadowInfo;
+				if (ProjectedShadowInfo->bShouldRenderVSM)
+				{
+					NaniteCullingViewsVolumes.Add(ProjectedShadowInfo->CasterOuterFrustum);
+				}
+			}
+
+			if (!NaniteCullingViewsVolumes.IsEmpty())
+			{
+				NaniteVisibilityQuery = Scene.NaniteVisibility[ENaniteMeshPass::BasePass].BeginVisibilityQuery(
+					NaniteCullingViewsVolumes,
+					&Scene.NaniteRasterPipelines[ENaniteMeshPass::BasePass],
+					&Scene.NaniteMaterials[ENaniteMeshPass::BasePass]
+				);
+			}
+		}
+	}
 	PostSetupDebugRender();
 }
 
@@ -120,7 +181,7 @@ void FShadowSceneRenderer::RenderVirtualShadowMaps(FRDGBuilder& GraphBuilder, bo
 	if (bNaniteEnabled)
 	{
 		const float ShadowsLODScaleFactor = ComputeNaniteShadowsLODScaleFactor();
-		VirtualShadowMapArray.RenderVirtualShadowMapsNanite(GraphBuilder, SceneRenderer, ShadowsLODScaleFactor, bUpdateNaniteStreaming, bNaniteProgrammableRaster);
+		VirtualShadowMapArray.RenderVirtualShadowMapsNanite(GraphBuilder, SceneRenderer, ShadowsLODScaleFactor, bUpdateNaniteStreaming, bNaniteProgrammableRaster, NaniteVisibilityQuery);
 	}
 
 	if (UseNonNaniteVirtualShadowMaps(SceneRenderer.ShaderPlatform, SceneRenderer.FeatureLevel))
diff --git a/Engine/Source/Runtime/Renderer/Private/Shadows/ShadowSceneRenderer.h b/Engine/Source/Runtime/Renderer/Private/Shadows/ShadowSceneRenderer.h
--- a/Engine/Source/Runtime/Renderer/Private/Shadows/ShadowSceneRenderer.h
+++ b/Engine/Source/Runtime/Renderer/Private/Shadows/ShadowSceneRenderer.h
@@ -37,6 +37,12 @@ public:
 	 * TODO: Return reference to FLocalLightShadowFrameSetup ?
 	 */
 	TSharedPtr<FVirtualShadowMapPerLightCacheEntry> AddLocalLightShadow(const FWholeSceneProjectedShadowInitializer& Initializer, FProjectedShadowInfo* ProjectedShadowInfo, FLightSceneInfo* LightSceneInfo, float MaxScreenRadius);
+	/**
+	 * Add a directional light for processing this frame.
+	 * TODO: Don't use legacy FProjectedShadowInfo or other params, instead info should flow from persistent setup & update.
+	 * TODO: Return reference to FLocalLightShadowFrameSetup ?
+	 */
+	void AddDirectionalLightShadow(FProjectedShadowInfo* ProjectedShadowInfo);
 
 	/**
 	 * Call after view-dependent setup has been processed (InitView etc) but before any rendering activity has been kicked off.
@@ -97,6 +103,13 @@ private:
 	//       tradeoff is easy to look up (given light ID) but not compact, but OTOH can keep compact lists of indices for various purposes
 	TArray<FLocalLightShadowFrameSetup, SceneRenderingAllocator> LocalLights;
 
+
+	struct FDirectionalLightShadowFrameSetup
+	{
+		FProjectedShadowInfo* ProjectedShadowInfo = nullptr;
+	};
+	TArray<FDirectionalLightShadowFrameSetup, SceneRenderingAllocator> DirectionalLights;
+
 	// Priority queue of distant lights to update.
 	FBinaryHeap<int32, uint32> DistantLightUpdateQueue;
 
@@ -107,4 +120,6 @@ private:
 	FDeferredShadingSceneRenderer& SceneRenderer;
 	FScene& Scene;
 	FVirtualShadowMapArray& VirtualShadowMapArray;
+
+	FNaniteVisibilityQuery* NaniteVisibilityQuery = nullptr;
 };
diff --git a/Engine/Source/Runtime/Renderer/Private/VirtualShadowMaps/VirtualShadowMapArray.cpp b/Engine/Source/Runtime/Renderer/Private/VirtualShadowMaps/VirtualShadowMapArray.cpp
--- a/Engine/Source/Runtime/Renderer/Private/VirtualShadowMaps/VirtualShadowMapArray.cpp
+++ b/Engine/Source/Runtime/Renderer/Private/VirtualShadowMaps/VirtualShadowMapArray.cpp
@@ -21,6 +21,13 @@
 #include "VirtualShadowMapVisualizationData.h"
 #include "SingleLayerWaterRendering.h"
 
+DECLARE_DWORD_COUNTER_STAT(TEXT("VSM Total Raster Bins"), STAT_VSMNaniteBasePassTotalRasterBins, STATGROUP_ShadowRendering);
+DECLARE_DWORD_COUNTER_STAT(TEXT("VSM Total Shading Draws"), STAT_VSMNaniteBasePassTotalShadingDraws, STATGROUP_ShadowRendering);
+
+DECLARE_DWORD_COUNTER_STAT(TEXT("VSM Visible Raster Bins"), STAT_VSMNaniteBasePassVisibleRasterBins, STATGROUP_ShadowRendering);
+DECLARE_DWORD_COUNTER_STAT(TEXT("VSM Visible Shading Draws"), STAT_VSMNaniteBassPassVisibleShadingDraws, STATGROUP_ShadowRendering);
+
+
 #define DEBUG_ALLOW_STATIC_SEPARATE_WITHOUT_CACHING 0
 
 IMPLEMENT_STATIC_UNIFORM_BUFFER_SLOT(VirtualShadowMapUbSlot);
@@ -2514,7 +2521,7 @@ static void AddRasterPass(
 		});
 }
 
-void FVirtualShadowMapArray::RenderVirtualShadowMapsNanite(FRDGBuilder& GraphBuilder, FSceneRenderer& SceneRenderer, float ShadowsLODScaleFactor, bool bUpdateNaniteStreaming, bool bNaniteProgrammableRaster)
+void FVirtualShadowMapArray::RenderVirtualShadowMapsNanite(FRDGBuilder& GraphBuilder, FSceneRenderer& SceneRenderer, float ShadowsLODScaleFactor, bool bUpdateNaniteStreaming, bool bNaniteProgrammableRaster, FNaniteVisibilityQuery* NaniteVisibilityQuery)
 {
 	bool bCsvLogEnabled = false;
 #if CSV_PROFILER
@@ -2581,6 +2588,24 @@ void FVirtualShadowMapArray::RenderVirtualShadowMapsNanite(FRDGBuilder& GraphBui
 		Nanite::FRasterState RasterState;
 
 		FNaniteVisibilityResults VisibilityResults; // TODO: Hook up culling for shadows
+		if (NaniteVisibilityQuery != nullptr)
+		{
+			Scene.NaniteVisibility[ENaniteMeshPass::BasePass].FinishVisibilityQuery(NaniteVisibilityQuery, VisibilityResults);
+
+			uint32 TotalRasterBins = 0;
+			uint32 VisibleRasterBins = 0;
+			VisibilityResults.GetRasterBinStats(VisibleRasterBins, TotalRasterBins);
+
+			uint32 TotalShadingDraws = 0;
+			uint32 VisibleShadingDraws = 0;
+			VisibilityResults.GetShadingDrawStats(VisibleShadingDraws, TotalShadingDraws);
+
+			SET_DWORD_STAT(STAT_VSMNaniteBasePassTotalRasterBins, TotalRasterBins);
+			SET_DWORD_STAT(STAT_VSMNaniteBasePassTotalShadingDraws, TotalShadingDraws);
+								
+			SET_DWORD_STAT(STAT_VSMNaniteBasePassVisibleRasterBins, VisibleRasterBins);
+			SET_DWORD_STAT(STAT_VSMNaniteBassPassVisibleShadingDraws, VisibleShadingDraws);
+		}
 
 		Nanite::FCullingContext::FConfiguration CullingConfig = { 0 };
 		CullingConfig.bUpdateStreaming = bUpdateNaniteStreaming;
diff --git a/Engine/Source/Runtime/Renderer/Private/VirtualShadowMaps/VirtualShadowMapArray.h b/Engine/Source/Runtime/Renderer/Private/VirtualShadowMaps/VirtualShadowMapArray.h
--- a/Engine/Source/Runtime/Renderer/Private/VirtualShadowMaps/VirtualShadowMapArray.h
+++ b/Engine/Source/Runtime/Renderer/Private/VirtualShadowMaps/VirtualShadowMapArray.h
@@ -281,7 +281,7 @@ public:
 	/**
 	 * Draw Nanite geometry into the VSMs.
 	 */
-	void RenderVirtualShadowMapsNanite(FRDGBuilder& GraphBuilder, FSceneRenderer& SceneRenderer, float ShadowsLODScaleFactor, bool bUpdateNaniteStreaming, bool bNaniteProgrammableRaster);
+	void RenderVirtualShadowMapsNanite(FRDGBuilder& GraphBuilder, FSceneRenderer& SceneRenderer, float ShadowsLODScaleFactor, bool bUpdateNaniteStreaming, bool bNaniteProgrammableRaster, FNaniteVisibilityQuery* NaniteVisibilityQuery);
 
 	/**
 	 * Draw Non-Nanite geometry into the VSMs.
