Fixing NonUnity errors caused by CL 23104907 and CL 23104389

diff --git a/Engine/Plugins/Developer/PixWinPlugin/Source/PixWinPlugin/Private/PixWinPluginModule.cpp b/Engine/Plugins/Developer/PixWinPlugin/Source/PixWinPlugin/Private/PixWinPluginModule.cpp
--- a/Engine/Plugins/Developer/PixWinPlugin/Source/PixWinPlugin/Private/PixWinPluginModule.cpp
+++ b/Engine/Plugins/Developer/PixWinPlugin/Source/PixWinPlugin/Private/PixWinPluginModule.cpp
@@ -7,6 +7,7 @@
 #include "RHI.h"
 #include "UnrealClient.h"
 #include "Engine/GameViewportClient.h"
+#include "Engine/Engine.h"
 
 #if !defined(WITH_PIX_EVENT_RUNTIME)
 #define WITH_PIX_EVENT_RUNTIME 0
diff --git a/Engine/Source/Runtime/Renderer/Private/GPUBenchmark.cpp b/Engine/Source/Runtime/Renderer/Private/GPUBenchmark.cpp
--- a/Engine/Source/Runtime/Renderer/Private/GPUBenchmark.cpp
+++ b/Engine/Source/Runtime/Renderer/Private/GPUBenchmark.cpp
@@ -23,6 +23,7 @@
 #include "LongGPUTask.h"
 #include "VisualizeTexture.h"
 #include "CommonRenderResources.h"
+#include "SceneView.h"
 
 static const uint32 GBenchmarkResolution = 512;
 static const uint32 GBenchmarkPrimitives = 200000;
diff --git a/Engine/Source/Runtime/Renderer/Private/PostProcess/SceneFilterRendering.h b/Engine/Source/Runtime/Renderer/Private/PostProcess/SceneFilterRendering.h
--- a/Engine/Source/Runtime/Renderer/Private/PostProcess/SceneFilterRendering.h
+++ b/Engine/Source/Runtime/Renderer/Private/PostProcess/SceneFilterRendering.h
@@ -15,6 +15,8 @@
 #include "ShaderParameterUtils.h"
 #include "SceneRenderTargetParameters.h"
 
+enum class EStereoscopicPass;
+
 /** Uniform buffer for computing the vertex positional and UV adjustments in the vertex shader. */
 BEGIN_GLOBAL_SHADER_PARAMETER_STRUCT( FDrawRectangleParameters, RENDERER_API)
 	SHADER_PARAMETER( FVector4f, PosScaleBias )
