[Backout] - CL23092056 - causing hangs on PC performance mode startup
Original CL Desc
-----------------------------------------------------------------
Fix assert on IsInActualRenderingThread() by using IsInParallelRenderingThread() instead and inverting condition

#rnx

diff --git a/Engine/Source/Runtime/Renderer/Private/MeshDrawCommands.cpp b/Engine/Source/Runtime/Renderer/Private/MeshDrawCommands.cpp
--- a/Engine/Source/Runtime/Renderer/Private/MeshDrawCommands.cpp
+++ b/Engine/Source/Runtime/Renderer/Private/MeshDrawCommands.cpp
@@ -1483,7 +1483,7 @@ void FParallelMeshDrawCommandPass::DispatchDraw(FParallelCommandListSet* Paralle
 	{
 		QUICK_SCOPE_CYCLE_COUNTER(STAT_MeshPassDrawImmediate);
 
-		WaitForMeshPassSetupTask(IsInParallelRenderingThread() ? EWaitThread::Task : EWaitThread::Render);
+		WaitForMeshPassSetupTask(IsInActualRenderingThread() ? EWaitThread::Render : EWaitThread::Task);
 
 		if (TaskContext.bUseGPUScene)
 		{
