Make sure reflection capture View knows it's a capture during the construction.

diff --git a/Engine/Source/Runtime/Renderer/Private/ReflectionEnvironmentCapture.cpp b/Engine/Source/Runtime/Renderer/Private/ReflectionEnvironmentCapture.cpp
--- a/Engine/Source/Runtime/Renderer/Private/ReflectionEnvironmentCapture.cpp
+++ b/Engine/Source/Runtime/Renderer/Private/ReflectionEnvironmentCapture.cpp
@@ -1283,6 +1283,7 @@ void CaptureSceneIntoScratchCubemap(
 
 		ViewInitOptions.ViewOrigin = CapturePosition;
 		ViewInitOptions.ViewRotationMatrix = CalcCubeFaceViewRotationMatrix((ECubeFace)CubeFace);
+		ViewInitOptions.bIsReflectionCapture = true;
 
 		FSceneView* View = new FSceneView(ViewInitOptions);
 
@@ -1295,7 +1296,6 @@ void CaptureSceneIntoScratchCubemap(
 			View->SpecularOverrideParameter = FVector4f(0, 0, 0, 0);
 		}
 
-		View->bIsReflectionCapture = true;
 		View->bStaticSceneOnly = bStaticSceneOnly;
 		View->StartFinalPostprocessSettings(CapturePosition);
 		View->EndFinalPostprocessSettings(ViewInitOptions);
