Fix shadowsdepths blocking waiting for async compute work to finish due to GPUScene resources ping-ponging between access modes. Using the shared ExternalResourceQueue avoids this.

diff --git a/Engine/Source/Runtime/Renderer/Private/DeferredShadingRenderer.cpp b/Engine/Source/Runtime/Renderer/Private/DeferredShadingRenderer.cpp
--- a/Engine/Source/Runtime/Renderer/Private/DeferredShadingRenderer.cpp
+++ b/Engine/Source/Runtime/Renderer/Private/DeferredShadingRenderer.cpp
@@ -265,6 +265,12 @@ static TAutoConsoleVariable<int32> CVarSceneDepthHZBAsyncCompute(
 	TEXT(" 2: Use async compute, start after ComputeLightGrid.CompactLinks pass"),
 	ECVF_RenderThreadSafe);
 
+static TAutoConsoleVariable<int32> CVarShadowsUseSharedExternalAccessQueue(
+	TEXT("r.ShadowsUseSharedExternalAccessQueue"), 1,
+	TEXT("If enabled, shadows will use the shared external access queue, minimizing unnecessary transitions"),
+	ECVF_RenderThreadSafe);
+
+
 #if RHI_RAYTRACING
 
 static bool bRefreshRayTracingInstances = false;
@@ -2706,7 +2712,10 @@ void FDeferredShadingSceneRenderer::Render(FRDGBuilder& GraphBuilder)
 	}
 
 	PrepareDistanceFieldScene(GraphBuilder, ExternalAccessQueue, false);
-	ExternalAccessQueue.Submit(GraphBuilder);
+	if (!CVarShadowsUseSharedExternalAccessQueue.GetValueOnRenderThread())
+	{
+		ExternalAccessQueue.Submit(GraphBuilder);
+	}
 
 	const bool bShouldRenderVelocities = ShouldRenderVelocities();
 	const EShaderPlatform Platform = GetViewFamilyInfo(Views).GetShaderPlatform();
@@ -3110,7 +3119,7 @@ void FDeferredShadingSceneRenderer::Render(FRDGBuilder& GraphBuilder)
 	{
 		// With forward shading we need to render shadow maps early
 		ensureMsgf(!VirtualShadowMapArray.IsEnabled(), TEXT("Virtual shadow maps are not supported in the forward shading path"));
-		RenderShadowDepthMaps(GraphBuilder, InstanceCullingManager);
+		RenderShadowDepthMaps(GraphBuilder, InstanceCullingManager, ExternalAccessQueue);
 
 		if (bHairStrandsEnable && !bHasRayTracedOverlay)
 		{
@@ -3359,7 +3368,7 @@ void FDeferredShadingSceneRenderer::Render(FRDGBuilder& GraphBuilder)
 				VirtualShadowMapArray.BuildPageAllocations(GraphBuilder, SceneTextures, Views, ViewFamily.EngineShowFlags, SortedLightSet, VisibleLightInfos, NaniteRasterResults, SingleLayerWaterPrePassResult);
 			}
 
-			RenderShadowDepthMaps(GraphBuilder, InstanceCullingManager);
+			RenderShadowDepthMaps(GraphBuilder, InstanceCullingManager, ExternalAccessQueue);
 		}
 		CheckShadowDepthRenderCompleted();
 
@@ -3379,6 +3388,10 @@ void FDeferredShadingSceneRenderer::Render(FRDGBuilder& GraphBuilder)
 			RenderLumenSceneLighting(GraphBuilder, LumenFrameTemporaries);
 		}
 	}
+	if (CVarShadowsUseSharedExternalAccessQueue.GetValueOnRenderThread())
+	{
+		ExternalAccessQueue.Submit(GraphBuilder);
+	}
 	// End shadow and fog after base pass
 
 	// Trigger a command submit here, to avoid GPU bubbles
diff --git a/Engine/Source/Runtime/Renderer/Private/MobileShadingRenderer.cpp b/Engine/Source/Runtime/Renderer/Private/MobileShadingRenderer.cpp
--- a/Engine/Source/Runtime/Renderer/Private/MobileShadingRenderer.cpp
+++ b/Engine/Source/Runtime/Renderer/Private/MobileShadingRenderer.cpp
@@ -937,7 +937,9 @@ void FMobileSceneRenderer::Render(FRDGBuilder& GraphBuilder)
 	}
 	
 	GraphBuilder.SetCommandListStat(GET_STATID(STAT_CLMM_Shadows));
-	RenderShadowDepthMaps(GraphBuilder, InstanceCullingManager);
+	FRDGExternalAccessQueue ExternalAccessQueue; // Ideally we'd use the shared access queue and pass that in, but keeping this minimal for now
+	RenderShadowDepthMaps(GraphBuilder, InstanceCullingManager, ExternalAccessQueue);
+	ExternalAccessQueue.Submit(GraphBuilder);
 	
 	PollOcclusionQueriesPass(GraphBuilder);
 	GraphBuilder.AddDispatchHint();
diff --git a/Engine/Source/Runtime/Renderer/Private/SceneRendering.h b/Engine/Source/Runtime/Renderer/Private/SceneRendering.h
--- a/Engine/Source/Runtime/Renderer/Private/SceneRendering.h
+++ b/Engine/Source/Runtime/Renderer/Private/SceneRendering.h
@@ -2270,7 +2270,7 @@ protected:
 	void BeginGatherShadowPrimitives(FDynamicShadowsTaskData* TaskData);
 	void FinishGatherShadowPrimitives(FDynamicShadowsTaskData* TaskData);
 
-	void RenderShadowDepthMaps(FRDGBuilder& GraphBuilder, FInstanceCullingManager& InstanceCullingManager);
+	void RenderShadowDepthMaps(FRDGBuilder& GraphBuilder, FInstanceCullingManager& InstanceCullingManager, FRDGExternalAccessQueue& ExternalAccessQueue);
 	void RenderVirtualShadowMaps(FRDGBuilder& GraphBuilder, bool bNaniteEnabled);
 	void RenderShadowDepthMapAtlases(FRDGBuilder& GraphBuilder);
 
diff --git a/Engine/Source/Runtime/Renderer/Private/ShadowDepthRendering.cpp b/Engine/Source/Runtime/Renderer/Private/ShadowDepthRendering.cpp
--- a/Engine/Source/Runtime/Renderer/Private/ShadowDepthRendering.cpp
+++ b/Engine/Source/Runtime/Renderer/Private/ShadowDepthRendering.cpp
@@ -1627,7 +1627,7 @@ void FSceneRenderer::RenderVirtualShadowMaps(FRDGBuilder& GraphBuilder, bool bNa
 	}
 }
 
-void FSceneRenderer::RenderShadowDepthMaps(FRDGBuilder& GraphBuilder, FInstanceCullingManager &InstanceCullingManager)
+void FSceneRenderer::RenderShadowDepthMaps(FRDGBuilder& GraphBuilder, FInstanceCullingManager &InstanceCullingManager, FRDGExternalAccessQueue& ExternalAccessQueue)
 {
 	ensureMsgf(!bShadowDepthRenderCompleted, TEXT("RenderShadowDepthMaps called twice in the same frame"));
 
@@ -1639,8 +1639,6 @@ void FSceneRenderer::RenderShadowDepthMaps(FRDGBuilder& GraphBuilder, FInstanceC
 	RDG_EVENT_SCOPE(GraphBuilder, "ShadowDepths");
 	RDG_GPU_STAT_SCOPE(GraphBuilder, ShadowDepths);
 
-	FRDGExternalAccessQueue ExternalAccessQueue;
-
 	// Ensure all shadow view dynamic primitives are uploaded before shadow-culling batching pass.
 	// TODO: automate this such that:
 	//  1. we only process views that need it (have dynamic primitives)
@@ -1868,8 +1866,6 @@ void FSceneRenderer::RenderShadowDepthMaps(FRDGBuilder& GraphBuilder, FInstanceC
 		ShadowMap.RenderTargets.DepthTarget = ConvertToExternalAccessTexture(GraphBuilder, ExternalAccessQueue, ShadowDepthTexture);
 	}
 
-	ExternalAccessQueue.Submit(GraphBuilder);
-
 	if (SortedShadowsForShadowDepthPass.PreshadowCache.Shadows.Num() > 0)
 	{
 		RDG_EVENT_SCOPE(GraphBuilder, "PreshadowCache");
