Implements TSR's ComputeMoireLuma from scene color beforedistortion translucency

diff --git a/Engine/Shaders/Private/TemporalSuperResolution/TSRComputeMoireLuma.usf b/Engine/Shaders/Private/TemporalSuperResolution/TSRComputeMoireLuma.usf
new file mode 100644
--- /dev/null
+++ b/Engine/Shaders/Private/TemporalSuperResolution/TSRComputeMoireLuma.usf
@@ -0,0 +1,70 @@
+// Copyright Epic Games, Inc. All Rights Reserved.
+
+#include "TSRCommon.ush"
+
+
+//------------------------------------------------------- CONFIG
+
+#define TILE_SIZE 8
+
+#define TILE_ROWS 2
+
+
+//------------------------------------------------------- PARAMETERS
+
+Texture2D<tsr_half3> SceneColorTexture;
+RWTexture2D<tsr_half> MoireLumaOutput;
+
+
+//------------------------------------------------------- ENTRY POINT
+
+[numthreads(TILE_SIZE * TILE_SIZE, 1, 1)]
+void MainCS(
+	uint2 GroupId : SV_GroupID,
+	uint GroupThreadIndex : SV_GroupIndex)
+{
+	uint2 GroupThreadId = uint2(GroupThreadIndex % TILE_SIZE, GroupThreadIndex / TILE_SIZE);
+	uint2 OutputPixelOrigin = (InputInfo_ViewportMin + GroupId * uint(TILE_SIZE * TILE_ROWS)) + GroupThreadId;
+
+	tsr_half3 ColorArray[TILE_ROWS * TILE_ROWS];
+
+	ISOLATE
+	{
+		UNROLL_N(TILE_ROWS)
+		for (uint y = 0; y < TILE_ROWS; y++)
+		{
+			UNROLL_N(TILE_ROWS)
+			for (uint x = 0; x < TILE_ROWS; x++)
+			{
+				uint2 PixelOffset = uint2(x, y) * TILE_SIZE;
+				uint2 InputPixelPos = OutputPixelOrigin + PixelOffset;
+			
+				ColorArray[x + y * TILE_ROWS] = SceneColorTexture[InputPixelPos];
+			}
+		}
+	}
+
+	ISOLATE
+	{
+		const tsr_half MoireLumaToChannel = tsr_half(3.0);
+		const tsr_half3 LumaWeights = rcp(MoireLumaToChannel.xxx);
+
+		UNROLL_N(TILE_ROWS)
+		for (uint y = 0; y < TILE_ROWS; y++)
+		{
+			UNROLL_N(TILE_ROWS)
+			for (uint x = 0; x < TILE_ROWS; x++)
+			{
+				uint2 PixelOffset = uint2(x, y) * TILE_SIZE;
+				uint2 InputPixelPos = OutputPixelOrigin + PixelOffset;
+			
+				tsr_half3 Color = ColorArray[x + y * TILE_ROWS];
+				Color *= rcp(tsr_half(1.0) + Color);
+
+				InputPixelPos.x = select(all(InputPixelPos < InputInfo_ViewportMax), InputPixelPos.x, uint(~0));
+				
+				MoireLumaOutput[InputPixelPos] = dot(Color, LumaWeights);
+			}
+		}
+	}
+}
diff --git a/Engine/Shaders/Private/TemporalSuperResolution/TSRRejectShading.usf b/Engine/Shaders/Private/TemporalSuperResolution/TSRRejectShading.usf
--- a/Engine/Shaders/Private/TemporalSuperResolution/TSRRejectShading.usf
+++ b/Engine/Shaders/Private/TemporalSuperResolution/TSRRejectShading.usf
@@ -100,6 +100,7 @@ float3 HistoryGuideQuantizationError;
 float FlickeringFramePeriod;
 
 Texture2D<tsr_halfC> InputTexture;
+Texture2D<tsr_half>  InputMoireLumaTexture;
 Texture2D<tsr_half4> InputSceneTranslucencyTexture;
 Texture2D<tsr_halfC> ReprojectedHistoryGuideTexture;
 Texture2D<tsr_half4> ReprojectedHistoryMoireTexture;
@@ -611,6 +612,7 @@ void MainCS(
 	tsr_simd_halfC OriginalHistory;
 	tsr_simd_half bIsValid;
 	#if CONFIG_MOIRE_DETECTION
+		tsr_simd_half MoireInput;
 		tsr_simd_half4 PrevMoire;
 		tsr_simd_half bIsStatic;
 	#endif
@@ -632,6 +634,7 @@ void MainCS(
 			OriginalHistory.SetElement(SimdIndex, ReprojectedHistoryGuideTexture[ClampedFetchPixelPos]);
 			bIsValid.SetElement(SimdIndex, ParallaxRejectionMaskTexture[ClampedFetchPixelPos] > tsr_half(PARALLAX_REJECTION_MASK_THRESHOLD) ? tsr_half(1.0) : tsr_half(0.0));
 			#if CONFIG_MOIRE_DETECTION
+				MoireInput.SetElement(SimdIndex, InputMoireLumaTexture[ClampedFetchPixelPos]);
 				PrevMoire.SetElement(SimdIndex, ReprojectedHistoryMoireTexture[ClampedFetchPixelPos]);
 				bIsStatic.SetElement(SimdIndex, tsr_half(1.0) - tsr_half(IsMovingMaskTexture[ClampedFetchPixelPos]) * rcp(tsr_half(255.0)));
 			#endif
@@ -705,6 +708,11 @@ void MainCS(
 
 
 		tsr_simd_half InputLuma;
+		#if 1
+		{
+			InputLuma = MoireInput;
+		}
+		#else
 		{
 			UNROLL_N(SIMD_SIZE)
 			for (uint SimdIndex = 0; SimdIndex < SIMD_SIZE; SimdIndex++)
@@ -714,6 +722,7 @@ void MainCS(
 				InputLuma.SetElement(SimdIndex, dot(OriginalOpaqueInput.GetElement(SimdIndex).rgb, LumaWeights));
 			}
 		}
+		#endif
 		
 		// Convolve the input luma to have the exact same moire pattern as FilteredInput below.
 		InputLuma = Convolve3x3HV<TAccumulatorAvg<tsr_simd_half>, tsr_simd_half>(InputLuma, bIsValid, GroupThreadIndex);
diff --git a/Engine/Source/Runtime/Renderer/Private/DeferredShadingRenderer.cpp b/Engine/Source/Runtime/Renderer/Private/DeferredShadingRenderer.cpp
--- a/Engine/Source/Runtime/Renderer/Private/DeferredShadingRenderer.cpp
+++ b/Engine/Source/Runtime/Renderer/Private/DeferredShadingRenderer.cpp
@@ -3709,6 +3709,7 @@ void FDeferredShadingSceneRenderer::Render(FRDGBuilder& GraphBuilder)
 	}
 
 	// Draw translucency.
+	TArray<FScreenPassTexture> TSRMoireInputTextures;
 	if (!bHasRayTracedOverlay && TranslucencyViewsToRender != ETranslucencyView::None)
 	{
 		RDG_CSV_STAT_EXCLUSIVE_SCOPE(GraphBuilder, RenderTranslucency);
@@ -3735,6 +3736,21 @@ void FDeferredShadingSceneRenderer::Render(FRDGBuilder& GraphBuilder)
 			}
 		}
 
+		// Extract TSR's moire heuristic luminance before renderering translucency into the scene color.
+		for (int32 ViewIndex = 0; ViewIndex < Views.Num(); ++ViewIndex)
+		{
+			FViewInfo& View = Views[ViewIndex];
+			if (ITemporalUpscaler::GetMainTAAPassConfig(View) == EMainTAAPassConfig::TSR)
+			{
+				if (TSRMoireInputTextures.Num() == 0)
+				{
+					TSRMoireInputTextures.SetNum(Views.Num());
+				}
+
+				TSRMoireInputTextures[ViewIndex] = AddTSRComputeMoireLuma(GraphBuilder, View.ShaderMap, FScreenPassTexture(SceneTextures.Color.Target, View.ViewRect));
+			}
+		}
+
 		// Sort objects' triangles
 		for (FViewInfo& View : Views)
 		{
@@ -4005,7 +4021,24 @@ void FDeferredShadingSceneRenderer::Render(FRDGBuilder& GraphBuilder)
 					const FPerViewPipelineState& ViewPipelineState = GetViewPipelineState(View);
 					const bool bAnyLumenActive = ViewPipelineState.DiffuseIndirectMethod == EDiffuseIndirectMethod::Lumen || ViewPipelineState.ReflectionsMethod == EReflectionsMethod::Lumen;
 
-					AddPostProcessingPasses(GraphBuilder, View, ViewIndex, bAnyLumenActive, ViewPipelineState.ReflectionsMethod, PostProcessingInputs, NaniteResults, InstanceCullingManager, &VirtualShadowMapArray, LumenFrameTemporaries, SceneWithoutWaterTextures);
+					FScreenPassTexture TSRMoireInput;
+					if (ViewIndex < TSRMoireInputTextures.Num())
+					{
+						TSRMoireInput = TSRMoireInputTextures[ViewIndex];
+					}
+
+					AddPostProcessingPasses(
+						GraphBuilder,
+						View, ViewIndex,
+						bAnyLumenActive,
+						ViewPipelineState.ReflectionsMethod,
+						PostProcessingInputs,
+						NaniteResults,
+						InstanceCullingManager,
+						&VirtualShadowMapArray,
+						LumenFrameTemporaries,
+						SceneWithoutWaterTextures,
+						TSRMoireInput);
 				}
 			}
 		}
diff --git a/Engine/Source/Runtime/Renderer/Private/PostProcess/PostProcessing.cpp b/Engine/Source/Runtime/Renderer/Private/PostProcess/PostProcessing.cpp
--- a/Engine/Source/Runtime/Renderer/Private/PostProcess/PostProcessing.cpp
+++ b/Engine/Source/Runtime/Renderer/Private/PostProcess/PostProcessing.cpp
@@ -205,7 +205,8 @@ void AddPostProcessingPasses(
 	FInstanceCullingManager& InstanceCullingManager,
 	FVirtualShadowMapArray* VirtualShadowMapArray, 
 	FLumenSceneFrameTemporaries& LumenFrameTemporaries,
-	const FSceneWithoutWaterTextures& SceneWithoutWaterTextures)
+	const FSceneWithoutWaterTextures& SceneWithoutWaterTextures,
+	FScreenPassTexture TSRMoireInput)
 {
 	RDG_CSV_STAT_EXCLUSIVE_SCOPE(GraphBuilder, RenderPostProcessing);
 	QUICK_SCOPE_CYCLE_COUNTER(STAT_PostProcessing_Process);
@@ -641,6 +642,7 @@ void AddPostProcessingPasses(
 			UpscalerPassInputs.SceneDepthTexture = SceneDepth.Texture;
 			UpscalerPassInputs.SceneVelocityTexture = Velocity.Texture;
 			UpscalerPassInputs.PostDOFTranslucencyResources = PostDOFTranslucencyResources;
+			UpscalerPassInputs.MoireInputTexture = TSRMoireInput;
 
 			ITemporalUpscaler::FOutputs Outputs = UpscalerToUse->AddPasses(
 				GraphBuilder,
diff --git a/Engine/Source/Runtime/Renderer/Private/PostProcess/PostProcessing.h b/Engine/Source/Runtime/Renderer/Private/PostProcess/PostProcessing.h
--- a/Engine/Source/Runtime/Renderer/Private/PostProcess/PostProcessing.h
+++ b/Engine/Source/Runtime/Renderer/Private/PostProcess/PostProcessing.h
@@ -53,7 +53,8 @@ void AddPostProcessingPasses(
 	FInstanceCullingManager& InstanceCullingManager,
 	FVirtualShadowMapArray* VirtualShadowMapArray,
 	struct FLumenSceneFrameTemporaries& LumenFrameTemporaries,
-	const FSceneWithoutWaterTextures& SceneWithoutWaterTextures);
+	const FSceneWithoutWaterTextures& SceneWithoutWaterTextures,
+	FScreenPassTexture TSRMoireInput);
 
 void AddDebugViewPostProcessingPasses(FRDGBuilder& GraphBuilder, const FViewInfo& View, const FPostProcessingInputs& Inputs, const Nanite::FRasterResults* NaniteRasterResults);
 
diff --git a/Engine/Source/Runtime/Renderer/Private/PostProcess/TemporalAA.h b/Engine/Source/Runtime/Renderer/Private/PostProcess/TemporalAA.h
--- a/Engine/Source/Runtime/Renderer/Private/PostProcess/TemporalAA.h
+++ b/Engine/Source/Runtime/Renderer/Private/PostProcess/TemporalAA.h
@@ -172,6 +172,8 @@ extern RENDERER_API FTAAOutputs AddTemporalAAPass(
 	const FTemporalAAHistory& InputHistory,
 	FTemporalAAHistory* OutputHistory);
 
+extern RENDERER_API FScreenPassTexture AddTSRComputeMoireLuma(FRDGBuilder& GraphBuilder, FGlobalShaderMap* ShaderMap, FScreenPassTexture SceneColor);
+
 /** Interface for the main temporal upscaling algorithm. */
 class RENDERER_API ITemporalUpscaler : public ISceneViewFamilyExtention
 {
@@ -188,6 +190,7 @@ public:
 		FRDGTextureRef SceneDepthTexture = nullptr;
 		FRDGTextureRef SceneVelocityTexture = nullptr;
 		FTranslucencyPassResources PostDOFTranslucencyResources;
+		FScreenPassTexture MoireInputTexture;
 	};
 
 	struct FOutputs
diff --git a/Engine/Source/Runtime/Renderer/Private/PostProcess/TemporalSuperResolution.cpp b/Engine/Source/Runtime/Renderer/Private/PostProcess/TemporalSuperResolution.cpp
--- a/Engine/Source/Runtime/Renderer/Private/PostProcess/TemporalSuperResolution.cpp
+++ b/Engine/Source/Runtime/Renderer/Private/PostProcess/TemporalSuperResolution.cpp
@@ -59,7 +59,7 @@ TAutoConsoleVariable<int32> CVarTSRFlickeringAdjustToFrameRate(
 TAutoConsoleVariable<float> CVarTSRFlickeringPeriod(
 	TEXT("r.TSR.ShadingRejection.Flickering.Period"), 3.0f,
 	TEXT("Periode in 60hz frames in which luma oscilations at equal or greater frequency is considered flickering and should ghost (Default=3.0).\n"),
-	ECVF_Scalability | ECVF_RenderThreadSafe);
+	ECVF_RenderThreadSafe);
 
 TAutoConsoleVariable<float> CVarTSRFlickeringMaxParralaxVelocity(
 	TEXT("r.TSR.ShadingRejection.Flickering.MaxParralaxVelocity"), 10.0,
@@ -339,6 +339,18 @@ public:
 	}
 }; // class FTemporalSuperResolutionShader
 
+class FTSRComputeMoireLumaCS : public FTSRShader
+{
+	DECLARE_GLOBAL_SHADER(FTSRComputeMoireLumaCS);
+	SHADER_USE_PARAMETER_STRUCT(FTSRComputeMoireLumaCS, FTSRShader);
+
+	BEGIN_SHADER_PARAMETER_STRUCT(FParameters, )
+		SHADER_PARAMETER_STRUCT(FScreenPassTextureViewportParameters, InputInfo)
+		SHADER_PARAMETER_RDG_TEXTURE(Texture2D, SceneColorTexture)
+		SHADER_PARAMETER_RDG_TEXTURE_UAV(RWTexture2D, MoireLumaOutput)
+	END_SHADER_PARAMETER_STRUCT()
+}; // class FTSRComputeMoireLumaCS
+
 class FTSRClearPrevTexturesCS : public FTSRShader
 {
 	DECLARE_GLOBAL_SHADER(FTSRClearPrevTexturesCS);
@@ -478,6 +490,7 @@ class FTSRRejectShadingCS : public FTSRShader
 		SHADER_PARAMETER(float, FlickeringFramePeriod)
 
 		SHADER_PARAMETER_RDG_TEXTURE(Texture2D, InputTexture)
+		SHADER_PARAMETER_RDG_TEXTURE(Texture2D, InputMoireLumaTexture)
 		SHADER_PARAMETER_RDG_TEXTURE(Texture2D, InputSceneTranslucencyTexture)
 		SHADER_PARAMETER_RDG_TEXTURE(Texture2D, ReprojectedHistoryGuideTexture)
 		SHADER_PARAMETER_RDG_TEXTURE(Texture2D, ReprojectedHistoryMoireTexture)
@@ -719,6 +732,7 @@ class FTSRDebugHistoryCS : public FTSRShader
 	END_SHADER_PARAMETER_STRUCT()
 }; // class FTSRDebugHistoryCS
 
+IMPLEMENT_GLOBAL_SHADER(FTSRComputeMoireLumaCS,      "/Engine/Private/TemporalSuperResolution/TSRComputeMoireLuma.usf",      "MainCS", SF_Compute);
 IMPLEMENT_GLOBAL_SHADER(FTSRClearPrevTexturesCS,     "/Engine/Private/TemporalSuperResolution/TSRClearPrevTextures.usf",     "MainCS", SF_Compute);
 IMPLEMENT_GLOBAL_SHADER(FTSRDilateVelocityCS,        "/Engine/Private/TemporalSuperResolution/TSRDilateVelocity.usf",        "MainCS", SF_Compute);
 IMPLEMENT_GLOBAL_SHADER(FTSRDecimateHistoryCS,       "/Engine/Private/TemporalSuperResolution/TSRDecimateHistory.usf",       "MainCS", SF_Compute);
@@ -755,6 +769,40 @@ static FRDGTextureUAVRef CreateDummyUAV(FRDGBuilder& GraphBuilder, EPixelFormat
 	return GraphBuilder.CreateUAV(DummyTexture);
 };
 
+FScreenPassTexture AddTSRComputeMoireLuma(FRDGBuilder& GraphBuilder, FGlobalShaderMap* ShaderMap, FScreenPassTexture SceneColor)
+{
+	check(SceneColor.Texture)
+;	RDG_GPU_STAT_SCOPE(GraphBuilder, TemporalSuperResolution);
+
+	FScreenPassTexture MoireLuma;
+	{
+		FRDGTextureDesc Desc = FRDGTextureDesc::Create2D(
+			SceneColor.Texture->Desc.Extent,
+			PF_R8,
+			FClearValueBinding::None,
+			/* InFlags = */ TexCreate_ShaderResource | TexCreate_UAV);
+
+		MoireLuma.Texture = GraphBuilder.CreateTexture(Desc, TEXT("TSR.Moire.Luma"));
+		MoireLuma.ViewRect = SceneColor.ViewRect;
+	}
+
+	FTSRComputeMoireLumaCS::FParameters* PassParameters = GraphBuilder.AllocParameters<FTSRComputeMoireLumaCS::FParameters>();
+	PassParameters->InputInfo = GetScreenPassTextureViewportParameters(FScreenPassTextureViewport(
+		SceneColor.Texture->Desc.Extent, SceneColor.ViewRect));
+	PassParameters->SceneColorTexture = SceneColor.Texture;
+	PassParameters->MoireLumaOutput = GraphBuilder.CreateUAV(MoireLuma.Texture);
+
+	TShaderMapRef<FTSRComputeMoireLumaCS> ComputeShader(ShaderMap);
+	FComputeShaderUtils::AddPass(
+		GraphBuilder,
+		RDG_EVENT_NAME("TSR ComputeMoireLuma %dx%d", SceneColor.ViewRect.Width(), SceneColor.ViewRect.Height()),
+		ComputeShader,
+		PassParameters,
+		FComputeShaderUtils::GetGroupCount(MoireLuma.ViewRect.Size(), 8 * 2));
+
+	return MoireLuma;
+}
+
 ITemporalUpscaler::FOutputs AddTemporalSuperResolutionPasses(
 	FRDGBuilder& GraphBuilder,
 	const FViewInfo& View,
@@ -1496,6 +1544,15 @@ ITemporalUpscaler::FOutputs AddTemporalSuperResolutionPasses(
 		PassParameters->FlickeringFramePeriod = FlickeringFramePeriod;
 
 		PassParameters->InputTexture = PassInputs.SceneColorTexture;
+		if (PassInputs.MoireInputTexture.IsValid())
+		{
+			ensure(InputRect == PassInputs.MoireInputTexture.ViewRect);
+			PassParameters->InputMoireLumaTexture = PassInputs.MoireInputTexture.Texture;
+		}
+		else
+		{
+			PassParameters->InputMoireLumaTexture = BlackDummy;
+		}
 		if (bAccumulateTranslucencySeparately)
 		{
 			PassParameters->InputSceneTranslucencyTexture = BlackAlphaOneDummy;
