Remove fast clear on some resources that do not need it

diff --git a/Engine/Source/Runtime/Renderer/Private/SceneRendering.cpp b/Engine/Source/Runtime/Renderer/Private/SceneRendering.cpp
--- a/Engine/Source/Runtime/Renderer/Private/SceneRendering.cpp
+++ b/Engine/Source/Runtime/Renderer/Private/SceneRendering.cpp
@@ -5514,11 +5514,20 @@ void VirtualTextureFeedbackEnd(FRDGBuilder& GraphBuilder)
 	GVirtualTextureFeedbackBuffer.End(GraphBuilder);
 }
 
+static TAutoConsoleVariable<int32> CVarHalfResDepthNoFastClear(
+	TEXT("r.HalfResDepthNoFastClear"),
+	1,
+	TEXT("Remove fast clear on half resolution depth buffer (checkerboard and minmax)"),
+	ECVF_RenderThreadSafe);
+
 FRDGTextureRef CreateHalfResolutionDepthCheckerboardMinMax(FRDGBuilder& GraphBuilder, TArrayView<const FViewInfo> Views, FRDGTextureRef SceneDepthTexture)
 {
 	const uint32 DownscaleFactor = 2;
 	const FIntPoint SmallDepthExtent = GetDownscaledExtent(SceneDepthTexture->Desc.Extent, DownscaleFactor);
-	const FRDGTextureDesc SmallDepthDesc = FRDGTextureDesc::Create2D(SmallDepthExtent, PF_DepthStencil, FClearValueBinding::None, TexCreate_DepthStencilTargetable | TexCreate_ShaderResource);
+
+	const ETextureCreateFlags NoFastClearFlags = (CVarHalfResDepthNoFastClear.GetValueOnAnyThread() != 0) ? TexCreate_NoFastClear : TexCreate_None;
+
+	const FRDGTextureDesc SmallDepthDesc = FRDGTextureDesc::Create2D(SmallDepthExtent, PF_DepthStencil, FClearValueBinding::None, TexCreate_DepthStencilTargetable | TexCreate_ShaderResource | NoFastClearFlags);
 	FRDGTextureRef SmallDepthTexture = GraphBuilder.CreateTexture(SmallDepthDesc, TEXT("HalfResolutionDepthCheckerboardMinMax"));
 
 	for (const FViewInfo& View : Views)
@@ -5536,7 +5545,8 @@ FRDGTextureRef CreateHalfResolutionDepthCheckerboardMinMax(FRDGBuilder& GraphBui
 FRDGTextureRef CreateQuarterResolutionDepthMinAndMax(FRDGBuilder& GraphBuilder, TArrayView<const FViewInfo> Views, FRDGTextureRef InputDepthTexture)
 {
 	const FIntPoint SmallDepthExtent = GetDownscaledExtent(InputDepthTexture->Desc.Extent, 2);
-	const FRDGTextureDesc SmallTextureDesc = FRDGTextureDesc::Create2D(SmallDepthExtent, PF_G16R16F, FClearValueBinding::None, TexCreate_RenderTargetable | TexCreate_ShaderResource);
+	const ETextureCreateFlags NoFastClearFlags = (CVarHalfResDepthNoFastClear.GetValueOnAnyThread() != 0) ? TexCreate_NoFastClear : TexCreate_None;
+	const FRDGTextureDesc SmallTextureDesc = FRDGTextureDesc::Create2D(SmallDepthExtent, PF_G16R16F, FClearValueBinding::None, TexCreate_RenderTargetable | TexCreate_ShaderResource | NoFastClearFlags);
 	FRDGTextureRef SmallTexture = GraphBuilder.CreateTexture(SmallTextureDesc, TEXT("HalfResolutionDepthMinAndMax"));
 
 	for (const FViewInfo& View : Views)
diff --git a/Engine/Source/Runtime/Renderer/Private/VolumetricCloudRendering.cpp b/Engine/Source/Runtime/Renderer/Private/VolumetricCloudRendering.cpp
--- a/Engine/Source/Runtime/Renderer/Private/VolumetricCloudRendering.cpp
+++ b/Engine/Source/Runtime/Renderer/Private/VolumetricCloudRendering.cpp
@@ -2097,24 +2097,32 @@ FCloudRenderContext::FCloudRenderContext()
 	bCloudDebugViewModeEnabled = false;
 }
 
+static TAutoConsoleVariable<int32> CVarCloudDefaultTexturesNoFastClear(
+	TEXT("r.CloudDefaultTexturesNoFastClear"),
+	1,
+	TEXT("Remove fast clear on default cloud textures"),
+	ECVF_RenderThreadSafe);
+
 void FCloudRenderContext::CreateDefaultTexturesIfNeeded(FRDGBuilder& GraphBuilder)
 {
 	if (DefaultCloudColorCubeTexture == nullptr)
 	{
+		const ETextureCreateFlags NoFastClearFlags = (CVarCloudDefaultTexturesNoFastClear.GetValueOnAnyThread() != 0) ? TexCreate_NoFastClear : TexCreate_None;
+
 		DefaultCloudColorCubeTexture = GraphBuilder.CreateTexture(
-			FRDGTextureDesc::CreateCube(1, PF_FloatRGBA, FClearValueBinding::Black, TexCreate_ShaderResource | TexCreate_RenderTargetable | TexCreate_UAV),
+			FRDGTextureDesc::CreateCube(1, PF_FloatRGBA, FClearValueBinding::Black, TexCreate_ShaderResource | TexCreate_RenderTargetable | TexCreate_UAV | NoFastClearFlags),
 			TEXT("Cloud.ColorCubeDummy"));
 
 		DefaultCloudColor02DTexture = GraphBuilder.CreateTexture(
-			FRDGTextureDesc::Create2D(FIntPoint(1, 1), PF_FloatRGBA, FClearValueBinding::None, TexCreate_ShaderResource | TexCreate_RenderTargetable | TexCreate_UAV),
+			FRDGTextureDesc::Create2D(FIntPoint(1, 1), PF_FloatRGBA, FClearValueBinding::None, TexCreate_ShaderResource | TexCreate_RenderTargetable | TexCreate_UAV | NoFastClearFlags),
 			TEXT("Cloud.ColorDummy"));
 
 		DefaultCloudColor12DTexture = GraphBuilder.CreateTexture(
-			FRDGTextureDesc::Create2D(FIntPoint(1, 1), PF_FloatRGBA, FClearValueBinding::None, TexCreate_ShaderResource | TexCreate_RenderTargetable | TexCreate_UAV),
+			FRDGTextureDesc::Create2D(FIntPoint(1, 1), PF_FloatRGBA, FClearValueBinding::None, TexCreate_ShaderResource | TexCreate_RenderTargetable | TexCreate_UAV | NoFastClearFlags),
 			TEXT("Cloud.ColorDummy"));
 
 		DefaultCloudDepthTexture = GraphBuilder.CreateTexture(
-			FRDGTextureDesc::Create2D(FIntPoint(1, 1), PF_G16R16F, FClearValueBinding::None, TexCreate_ShaderResource | TexCreate_RenderTargetable | TexCreate_UAV),
+			FRDGTextureDesc::Create2D(FIntPoint(1, 1), PF_G16R16F, FClearValueBinding::None, TexCreate_ShaderResource | TexCreate_RenderTargetable | TexCreate_UAV | NoFastClearFlags),
 			TEXT("Cloud.DepthDummy"));
 
 
