Fixed Strata looking bad on some platforms requiring render target format specification (not account for Strata target count, motion blur not specifying the correct index)

diff --git a/Engine/Shaders/Shared/StrataDefinitions.h b/Engine/Shaders/Shared/StrataDefinitions.h
--- a/Engine/Shaders/Shared/StrataDefinitions.h
+++ b/Engine/Shaders/Shared/StrataDefinitions.h
@@ -8,7 +8,7 @@
 #pragma once
 
 // Change this to force recompilation of all strata dependent shaders (use https://www.random.org/cgi-bin/randbyte?nbytes=4&format=h)
-#define STRATA_SHADER_VERSION 0xc5265453
+#define STRATA_SHADER_VERSION 0xcb0f647de 
 
 
 
diff --git a/Engine/Source/Runtime/Renderer/Private/BasePassRendering.cpp b/Engine/Source/Runtime/Renderer/Private/BasePassRendering.cpp
--- a/Engine/Source/Runtime/Renderer/Private/BasePassRendering.cpp
+++ b/Engine/Source/Runtime/Renderer/Private/BasePassRendering.cpp
@@ -869,7 +869,8 @@ void ModifyBasePassCSPSCompilationEnvironment(const FMeshMaterialShaderPermutati
 		FVelocityRendering::BasePassCanOutputVelocity(Parameters.Platform);
 	if (bOutputVelocity)
 	{
-		const int32 VelocityIndex = IsForwardShadingEnabled(Parameters.Platform) ? 1 : 4; // As defined in BasePassPixelShader.usf
+		// As defined in BasePassPixelShader.usf. Also account for Strata setting velocity in slot 1 as described in FetchLegacyGBufferInfo.
+		const int32 VelocityIndex = Strata::IsStrataEnabled() ? 1 : (IsForwardShadingEnabled(Parameters.Platform) ? 1 : 4);
 		OutEnvironment.SetRenderTargetOutputFormat(VelocityIndex, PF_G16R16);
 	}
 
diff --git a/Engine/Source/Runtime/Renderer/Private/Strata/Strata.cpp b/Engine/Source/Runtime/Renderer/Private/Strata/Strata.cpp
--- a/Engine/Source/Runtime/Renderer/Private/Strata/Strata.cpp
+++ b/Engine/Source/Runtime/Renderer/Private/Strata/Strata.cpp
@@ -1136,12 +1136,14 @@ void SetBasePassRenderTargetOutputFormat(const EShaderPlatform Platform, const F
 		}
 		const FGBufferInfo BufferInfo = FetchFullGBufferInfo(GBufferParams);
 
-		// Add 2 uint for Strata fast path
-		OutEnvironment.SetRenderTargetOutputFormat(BufferInfo.NumTargets + 0, PF_R32_UINT);
-		OutEnvironment.SetRenderTargetOutputFormat(BufferInfo.NumTargets + 1, PF_R32_UINT);
+		// Add N uint for Strata fast path
+		for (int i = 0; i < STRATA_BASE_PASS_MRT_OUTPUT_COUNT; ++i)
+		{
+			OutEnvironment.SetRenderTargetOutputFormat(BufferInfo.NumTargets + i, PF_R32_UINT);
+		}
 
 		// Add another MRT for Strata top layer information
-		OutEnvironment.SetRenderTargetOutputFormat(BufferInfo.NumTargets + 2, PF_R32_UINT);
+		OutEnvironment.SetRenderTargetOutputFormat(BufferInfo.NumTargets + STRATA_BASE_PASS_MRT_OUTPUT_COUNT, PF_R32_UINT);
 	}
 }
 
