Process scheduled BLAS builds even in frames that don't need raytracing scene.
- Nanite::FRayTracingManager expects builds to be processed in the same frame their scheduled since it also impacts GPUScene updates.

diff --git a/Engine/Source/Runtime/Renderer/Private/DeferredShadingRenderer.cpp b/Engine/Source/Runtime/Renderer/Private/DeferredShadingRenderer.cpp
--- a/Engine/Source/Runtime/Renderer/Private/DeferredShadingRenderer.cpp
+++ b/Engine/Source/Runtime/Renderer/Private/DeferredShadingRenderer.cpp
@@ -1800,14 +1800,19 @@ bool FDeferredShadingSceneRenderer::DispatchRayTracingWorldUpdates(FRDGBuilder&
 {
 	OutDynamicGeometryScratchBuffer = nullptr;
 
-	if (!IsRayTracingEnabled() || !bAnyRayTracingPassEnabled || Views.Num() == 0)
+	// We only need to update ray tracing scene for the first view family, if multiple are rendered in a single scene render call.
+	if (!bShouldUpdateRayTracingScene)
 	{
 		// This needs to happen even when ray tracing is not enabled
-		// because importers might batch BVH creation requests that need to be resolved in any case
+		// - importers might batch BVH creation requests that need to be resolved in any case
 		GRayTracingGeometryManager.ProcessBuildRequests(GraphBuilder.RHICmdList);
+		// - Nanite ray tracing instances are already pointing at the new BLASes and RayTracingDataOffsets in GPUScene have been updated
+		Nanite::GRayTracingManager.ProcessBuildRequests(GraphBuilder);
 		return false;
 	}
 
+	check(IsRayTracingEnabled() && bAnyRayTracingPassEnabled && !Views.IsEmpty());
+
 	TRACE_CPUPROFILER_EVENT_SCOPE(FDeferredShadingSceneRenderer::DispatchRayTracingWorldUpdates);
 
 	// Make sure there are no pending skin cache builds and updates anymore:
@@ -1966,10 +1971,7 @@ static void ReleaseRaytracingResources(FRDGBuilder& GraphBuilder, TArrayView<FVi
 
 void FDeferredShadingSceneRenderer::WaitForRayTracingScene(FRDGBuilder& GraphBuilder, FRDGBufferRef DynamicGeometryScratchBuffer)
 {
-	if (!bAnyRayTracingPassEnabled)
-	{
-		return;
-	}
+	check(bAnyRayTracingPassEnabled);
 
 	TRACE_CPUPROFILER_EVENT_SCOPE(FDeferredShadingSceneRenderer::WaitForRayTracingScene);
 
@@ -2477,6 +2479,7 @@ void FDeferredShadingSceneRenderer::Render(FRDGBuilder& GraphBuilder)
 
 	if (IsRayTracingEnabled() && RHISupportsRayTracingShaders(ViewFamily.GetShaderPlatform()))
 	{
+		// Nanite raytracing manager update must run before GPUScene update since it can modify primitive data
 		Nanite::GRayTracingManager.Update();
 
 		if (!ViewFamily.EngineShowFlags.PathTracing)
@@ -3035,8 +3038,6 @@ void FDeferredShadingSceneRenderer::Render(FRDGBuilder& GraphBuilder)
 	// Early occlusion queries
 	const bool bOcclusionBeforeBasePass = ((DepthPass.EarlyZPassMode == EDepthDrawingMode::DDM_AllOccluders) || bIsEarlyDepthComplete);
 
-	bool bRayTracingSceneReady = false;
-
 	if (bOcclusionBeforeBasePass)
 	{
 		RenderOcclusionLambda();
@@ -3046,9 +3047,6 @@ void FDeferredShadingSceneRenderer::Render(FRDGBuilder& GraphBuilder)
 
 	BeginAsyncDistanceFieldShadowProjections(GraphBuilder, SceneTextures);
 
-
-	
-
 	if (bShouldRenderVolumetricCloudBase)
 	{
 		InitVolumetricRenderTargetForViews(GraphBuilder, Views);
@@ -3155,16 +3153,8 @@ void FDeferredShadingSceneRenderer::Render(FRDGBuilder& GraphBuilder)
 
 	FRDGBufferRef DynamicGeometryScratchBuffer = nullptr;
 #if RHI_RAYTRACING
-	// Async AS builds can potentially overlap with BasePass.  We only need to update ray tracing scene for the first view family,
-	// if multiple are rendered in a single scene render call.
-	if (bShouldUpdateRayTracingScene)
-	{
-		DispatchRayTracingWorldUpdates(GraphBuilder, DynamicGeometryScratchBuffer);
-	}
-	else
-	{
-		bRayTracingSceneReady = true;
-	}
+	// Async AS builds can potentially overlap with BasePass.
+	bool bNeedToWaitForRayTracingScene = DispatchRayTracingWorldUpdates(GraphBuilder, DynamicGeometryScratchBuffer);
 
 	/** Should be called somewhere before "WaitForRayTracingScene" */
 	SetupRayTracingLightDataForViews(GraphBuilder);
@@ -3174,10 +3164,10 @@ void FDeferredShadingSceneRenderer::Render(FRDGBuilder& GraphBuilder)
 	{
 #if RHI_RAYTRACING
 		// Lumen scene lighting requires ray tracing scene to be ready if HWRT shadows are desired
-		if (!bRayTracingSceneReady && Lumen::UseHardwareRayTracedSceneLighting(ViewFamily))
+		if (bNeedToWaitForRayTracingScene && Lumen::UseHardwareRayTracedSceneLighting(ViewFamily))
 		{
 			WaitForRayTracingScene(GraphBuilder, DynamicGeometryScratchBuffer);
-			bRayTracingSceneReady = true;
+			bNeedToWaitForRayTracingScene = false;
 		}
 #endif
 
@@ -3374,10 +3364,10 @@ void FDeferredShadingSceneRenderer::Render(FRDGBuilder& GraphBuilder)
 
 #if RHI_RAYTRACING
 		// Lumen scene lighting requires ray tracing scene to be ready if HWRT shadows are desired
-		if (!bRayTracingSceneReady && Lumen::UseHardwareRayTracedSceneLighting(ViewFamily))
+		if (bNeedToWaitForRayTracingScene && Lumen::UseHardwareRayTracedSceneLighting(ViewFamily))
 		{
 			WaitForRayTracingScene(GraphBuilder, DynamicGeometryScratchBuffer);
-			bRayTracingSceneReady = true;
+			bNeedToWaitForRayTracingScene = false;
 		}
 #endif // RHI_RAYTRACING
 
@@ -3480,10 +3470,10 @@ void FDeferredShadingSceneRenderer::Render(FRDGBuilder& GraphBuilder)
 
 #if RHI_RAYTRACING
 	// If Lumen did not force an earlier ray tracing scene sync, we must wait for it here.
-	if (!bRayTracingSceneReady)
+	if (bNeedToWaitForRayTracingScene)
 	{
 		WaitForRayTracingScene(GraphBuilder, DynamicGeometryScratchBuffer);
-		bRayTracingSceneReady = true;
+		bNeedToWaitForRayTracingScene = false;
 	}
 #endif // RHI_RAYTRACING
 
