Fix incorrect exposure when using raytracing debug visualization.

diff --git a/Engine/Shaders/Private/PostProcessTonemap.usf b/Engine/Shaders/Private/PostProcessTonemap.usf
--- a/Engine/Shaders/Private/PostProcessTonemap.usf
+++ b/Engine/Shaders/Private/PostProcessTonemap.usf
@@ -427,7 +427,7 @@ float4 TonemapCommonPS(
 #endif
 #else
 	// If eye adaptation is disabled (e.g. fixed exposure level ) or not available.  
-	float2 EyeAdaptationData = DefaultEyeExposure;
+	float2 EyeAdaptationData = DefaultEyeExposure.xx;
 #endif
 
 	float ExposureScale = EyeAdaptationData.x;
@@ -594,11 +594,6 @@ float4 TonemapCommonPS(
 			#endif
 		}
 		#endif
-
-		// TODO: doesn't need a shader permutation for a scalar.
-		#if NO_EYEADAPTATION_EXPOSURE_FIX == 1
-			ExposureScale = 1.0;
-		#endif
 		
 		LinearColor *= ExposureScale;
 
diff --git a/Engine/Source/Runtime/Renderer/Private/PostProcess/PostProcessEyeAdaptation.cpp b/Engine/Source/Runtime/Renderer/Private/PostProcess/PostProcessEyeAdaptation.cpp
--- a/Engine/Source/Runtime/Renderer/Private/PostProcess/PostProcessEyeAdaptation.cpp
+++ b/Engine/Source/Runtime/Renderer/Private/PostProcess/PostProcessEyeAdaptation.cpp
@@ -1195,18 +1195,19 @@ void FSceneViewState::UpdatePreExposure(FViewInfo& View)
 	// but this is too limiting for certain view. For instance shader preview doesn't have 
 	// volumetric lighting enabled, which makes the view be flagged as rich, and not updating 
 	// the pre-exposition value.
-	const bool bIsPreExposureRelevant =
-		ViewFamily.EngineShowFlags.EyeAdaptation && // Controls whether scene luminance is computed at all.
-		ViewFamily.EngineShowFlags.Lighting &&
-		ViewFamily.EngineShowFlags.PostProcessing &&
-		ViewFamily.bResolveScene &&
-		!ViewFamily.EngineShowFlags.LightMapDensity &&
-		!ViewFamily.EngineShowFlags.StationaryLightOverlap &&
-		!ViewFamily.EngineShowFlags.LightComplexity &&
-		!ViewFamily.EngineShowFlags.LODColoration &&
-		!ViewFamily.EngineShowFlags.HLODColoration &&
-		!ViewFamily.EngineShowFlags.LevelColoration &&
-		((!ViewFamily.EngineShowFlags.VisualizeBuffer) || View.CurrentBufferVisualizationMode != NAME_None); // disable pre-exposure for the buffer visualization modes
+	const bool bIsPreExposureRelevant = true
+		&& ViewFamily.EngineShowFlags.EyeAdaptation // Controls whether scene luminance is computed at all.
+		&& ViewFamily.EngineShowFlags.Lighting
+		&& ViewFamily.EngineShowFlags.PostProcessing
+		&& ViewFamily.bResolveScene
+		&& !ViewFamily.EngineShowFlags.LightMapDensity
+		&& !ViewFamily.EngineShowFlags.StationaryLightOverlap
+		&& !ViewFamily.EngineShowFlags.LightComplexity
+		&& !ViewFamily.EngineShowFlags.LODColoration
+		&& !ViewFamily.EngineShowFlags.HLODColoration
+		&& !ViewFamily.EngineShowFlags.LevelColoration
+		&& ((!ViewFamily.EngineShowFlags.VisualizeBuffer) || View.CurrentBufferVisualizationMode != NAME_None) // disable pre-exposure for the buffer visualization modes
+		&& !ViewFamily.EngineShowFlags.RayTracingDebug;
 
 	PreExposure = 1.f;
 	bUpdateLastExposure = false;
