Removed cloud noisy gathering. Cause issue at edge and in fact was deemed not useful by artists. Tracing sample count is what matters.
This will also make the upsampling half=>full res faster.

diff --git a/Engine/Shaders/Private/VolumetricRenderTarget.usf b/Engine/Shaders/Private/VolumetricRenderTarget.usf
--- a/Engine/Shaders/Private/VolumetricRenderTarget.usf
+++ b/Engine/Shaders/Private/VolumetricRenderTarget.usf
@@ -509,7 +509,6 @@ float2 SafeSampleVolumetricDepthTexture(float2 UV)
 }
 
 float4 VolumetricTextureSizeAndInvSize;
-float UvOffsetScale;
 float UvOffsetSampleAcceptanceWeight;
 float2 FullResolutionToVolumetricBufferResolutionScale;
 
@@ -541,6 +540,7 @@ void ComposeVolumetricRTOverScenePS(
 	Offset2Sample = normalize(Offset2Sample);
 	Offset3Sample = normalize(Offset3Sample);
 
+	const float UvOffsetScale = 1.0f;
 	float2 Offset0 = Offset0Sample * FullResOffsetUVScale * UvOffsetScale;
 	float2 Offset1 = Offset1Sample * FullResOffsetUVScale * UvOffsetScale;
 	float2 Offset2 = Offset2Sample * FullResOffsetUVScale * UvOffsetScale;
@@ -858,64 +858,13 @@ void ComposeVolumetricRTOverScenePS(
 		//DataAcc = float4(0, 1.0, 0, 0);
 	}
 
-		// we are going to fetch far sample so let's skip water pixel because depth cannot be trusted and used in this case.
-	if (UvOffsetScale > 0.0 && !bIsWaterPixel &&	
-		// Also only blur background to not leak cloud over opaque (this can happen when cloud data happear over meshes due to temporal jittering)
-#if HAS_INVERTED_Z_BUFFER
-		PixelDeviceZ < 0.0001				
-#else
-		PixelDeviceZ > 0.9999
-#endif
-		)
-	{
-		// Add some jittered samples to hide upsampling and noise through TAA filtering
-		// We use only Offset0 to save ALU and go faster. Variation across pixels is enough noise for taa to soften it out.
-		float4 VolumeRGBT0 = SafeSampleVolumetricTexture(VolumeUV + Offset0);
-		float4 VolumeRGBT1 = SafeSampleVolumetricTexture(VolumeUV - Offset0);
-		float4 VolumeRGBT2 = SafeSampleVolumetricTexture(VolumeUV + float2(Offset0.y, -Offset0.x));
-		float4 VolumeRGBT3 = SafeSampleVolumetricTexture(VolumeUV - float2(Offset0.y, -Offset0.x));
-
-		// Resolve the OutputRt0 from DataAcc to know what we want to blend into for the bilaterial upsampling based on a=transmittance.
-		// Why using transmittance? Because it cause less artefact than depth for far away pixels and 
-		// pixel with similar transmittance or more likely to be a good choice to mix together to produce noise (that is our assumption here)
-		OutputRt0 = float4(0.0f, 0.0f, 0.0f, 1.0f);
-		if (ValidSampleCount > 0.0f)
-		{
-			OutputRt0 = DataAcc / ValidSampleCount;
-		}
-
-		// Now continue to accumulate in DataAcc
-		const float Extinction = UvOffsetSampleAcceptanceWeight;
-		{
-			const float weight = exp2(-Extinction * abs(VolumeRGBT0.a - OutputRt0.a));
-			ValidSampleCount += weight;
-			DataAcc += VolumeRGBT0 * weight;
-		}
-		{
-			const float weight = exp2(-Extinction * abs(VolumeRGBT1.a - OutputRt0.a));
-			ValidSampleCount += weight;
-			DataAcc += VolumeRGBT1 * weight;
-		}
-		{
-			const float weight = exp2(-Extinction * abs(VolumeRGBT2.a - OutputRt0.a));
-			ValidSampleCount += weight;
-			DataAcc += VolumeRGBT2 * weight;
-		}
-		{
-			const float weight = exp2(-Extinction * abs(VolumeRGBT3.a - OutputRt0.a));
-			ValidSampleCount += weight;
-			DataAcc += VolumeRGBT3 * weight;
-		}
-	}
-
-	OutputRt0 = float4(0, 0, 0, 1);
+	OutputRt0 = float4(0.0f, 0.0f, 0.0f, 1.0f);
 	if (ValidSampleCount > 0.0f)
 	{
 		OutputRt0 = DataAcc / ValidSampleCount;
 	}
 	else
 	{
-		OutputRt0 = float4(0.0f, 0.0f, 0.0f, 1.0f);
 		clip(-1.0f);
 	}
 
diff --git a/Engine/Source/Runtime/Renderer/Private/VolumetricRenderTarget.cpp b/Engine/Source/Runtime/Renderer/Private/VolumetricRenderTarget.cpp
--- a/Engine/Source/Runtime/Renderer/Private/VolumetricRenderTarget.cpp
+++ b/Engine/Source/Runtime/Renderer/Private/VolumetricRenderTarget.cpp
@@ -21,11 +21,6 @@ static TAutoConsoleVariable<int32> CVarVolumetricRenderTarget(
 	TEXT(""),
 	ECVF_RenderThreadSafe | ECVF_Scalability);
 
-static TAutoConsoleVariable<float> CVarVolumetricRenderTargetUvNoiseScale(
-	TEXT("r.VolumetricRenderTarget.UvNoiseScale"), 0.5f,
-	TEXT("Used when r.VolumetricRenderTarget.UpsamplingMode is in a mode using jitter - this value scales the amount of jitter."),
-	ECVF_RenderThreadSafe | ECVF_Scalability);
-
 static TAutoConsoleVariable<float> CVarVolumetricRenderTargetUvNoiseSampleAcceptanceWeight(
 	TEXT("r.VolumetricRenderTarget.UvNoiseSampleAcceptanceWeight"), 20.0f,
 	TEXT("Used when r.VolumetricRenderTarget.UpsamplingMode is in a mode using jitter - this value control the acceptance of noisy cloud samples according to their similarities. A higher value means large differences will be less accepted for blending."),
@@ -172,12 +167,10 @@ FVolumetricRenderTargetViewStateData::~FVolumetricRenderTargetViewStateData()
 
 void FVolumetricRenderTargetViewStateData::Initialise(
 	FIntPoint& ViewRectResolutionIn,
-	float InUvNoiseScale,
 	int32 InMode,
 	int32 InUpsamplingMode)
 {
 	// Update internal settings
-	UvNoiseScale = InUvNoiseScale;
 	Mode = FMath::Clamp(InMode, 0, 3);
 	UpsamplingMode = Mode == 2 ? 2 : FMath::Clamp(InUpsamplingMode, 0, 4); // if we are using mode 2 then we cannot intersect with depth and upsampling should be 2 (simple on/off intersection)
 
@@ -407,7 +400,6 @@ void InitVolumetricRenderTargetForViews(FRDGBuilder& GraphBuilder, TArrayView<FV
 		FIntPoint ViewRect = ViewInfo.ViewRect.Size();
 		VolumetricCloudRT.Initialise(	// TODO this is going to reallocate a buffer each time dynamic resolution scaling is applied 
 			ViewRect,
-			CVarVolumetricRenderTargetUvNoiseScale.GetValueOnAnyThread(), 
 			CVarVolumetricRenderTargetMode.GetValueOnRenderThread(),
 			CVarVolumetricRenderTargetUpsamplingMode.GetValueOnAnyThread());
 
@@ -601,7 +593,6 @@ class FComposeVolumetricRTOverScenePS : public FGlobalShader
 		SHADER_PARAMETER_SAMPLER(SamplerState, LinearTextureSampler)
 		SHADER_PARAMETER_SAMPLER(SamplerState, WaterLinearDepthSampler)
 		RENDER_TARGET_BINDING_SLOTS()
-		SHADER_PARAMETER(float, UvOffsetScale)
 		SHADER_PARAMETER(float, UvOffsetSampleAcceptanceWeight)
 		SHADER_PARAMETER(FVector4f, VolumetricTextureSizeAndInvSize)
 		SHADER_PARAMETER(FVector2f, FullResolutionToVolumetricBufferResolutionScale)
@@ -693,7 +684,6 @@ void ComposeVolumetricRenderTargetOverScene(
 		PassParameters->VolumetricTexture = VolumetricTexture;
 		PassParameters->VolumetricDepthTexture = VolumetricDepthTexture;
 		PassParameters->LinearTextureSampler = TStaticSamplerState<SF_Bilinear>::GetRHI();
-		PassParameters->UvOffsetScale = VolumetricCloudRT.GetUvNoiseScale();
 		PassParameters->UvOffsetSampleAcceptanceWeight = GetUvNoiseSampleAcceptanceWeight();
 		PassParameters->FullResolutionToVolumetricBufferResolutionScale = FVector2f(1.0f / float(GetMainDownsampleFactor(VRTMode)), float(GetMainDownsampleFactor(VRTMode)));
 		PassParameters->SceneTextures = SceneTextures.UniformBuffer;
@@ -766,7 +756,6 @@ void ComposeVolumetricRenderTargetOverSceneUnderWater(
 		PassParameters->WaterLinearDepthTexture = WaterPassData.DepthTexture;
 		PassParameters->LinearTextureSampler = TStaticSamplerState<SF_Bilinear>::GetRHI();
 		PassParameters->WaterLinearDepthSampler = TStaticSamplerState<SF_Point>::GetRHI();
-		PassParameters->UvOffsetScale = VolumetricCloudRT.GetUvNoiseScale();
 		PassParameters->UvOffsetSampleAcceptanceWeight = GetUvNoiseSampleAcceptanceWeight();
 		PassParameters->FullResolutionToVolumetricBufferResolutionScale = FVector2f(1.0f / float(GetMainDownsampleFactor(VRTMode)), float(GetMainDownsampleFactor(VRTMode)));
 		PassParameters->FullResolutionToWaterBufferScale = FVector2f(1.0f / WaterPassData.RefractionDownsampleFactor, WaterPassData.RefractionDownsampleFactor);
@@ -824,7 +813,6 @@ void ComposeVolumetricRenderTargetOverSceneForVisualization(
 		PassParameters->VolumetricTexture = VolumetricTexture;
 		PassParameters->VolumetricDepthTexture = VolumetricDepthTexture;
 		PassParameters->LinearTextureSampler = TStaticSamplerState<SF_Bilinear>::GetRHI();
-		PassParameters->UvOffsetScale = VolumetricCloudRT.GetUvNoiseScale();
 		PassParameters->UvOffsetSampleAcceptanceWeight = GetUvNoiseSampleAcceptanceWeight();
 		PassParameters->FullResolutionToVolumetricBufferResolutionScale = FVector2f(1.0f / float(GetMainDownsampleFactor(VRTMode)), float(GetMainDownsampleFactor(VRTMode)));
 		PassParameters->SceneTextures = SceneTextures.UniformBuffer;
diff --git a/Engine/Source/Runtime/Renderer/Private/VolumetricRenderTargetViewStateData.h b/Engine/Source/Runtime/Renderer/Private/VolumetricRenderTargetViewStateData.h
--- a/Engine/Source/Runtime/Renderer/Private/VolumetricRenderTargetViewStateData.h
+++ b/Engine/Source/Runtime/Renderer/Private/VolumetricRenderTargetViewStateData.h
@@ -23,7 +23,6 @@ public:
 
 	void Initialise(
 		FIntPoint& ViewRectResolutionIn,
-		float UvNoiseScale,
 		int32 Mode,
 		int32 UpsamplingMode);
 
@@ -53,7 +52,6 @@ public:
 	FUintVector4 GetTracingCoordToZbufferCoordScaleBias() const;
 	FUintVector4 GetTracingCoordToFullResPixelCoordScaleBias() const;
 
-	float GetUvNoiseScale()		const { return UvNoiseScale; }
 	int32 GetMode()				const { return Mode; }
 	int32 GetUpsamplingMode()	const { return UpsamplingMode; }
 
@@ -85,7 +83,6 @@ private:
 	TRefCountPtr<IPooledRenderTarget> VolumetricSecondaryTracingRT;
 	TRefCountPtr<IPooledRenderTarget> VolumetricTracingRTDepth;
 
-	float UvNoiseScale;
 	int32 Mode;
 	int32 UpsamplingMode;
 };
