Add r.VT.RVT.StreamingMips CVar.
Should be useful for debugging whether SVT is giving same result as RVT.

#localization none
#tests local

diff --git a/Engine/Source/Runtime/Renderer/Private/VT/RuntimeVirtualTextureSceneProxy.cpp b/Engine/Source/Runtime/Renderer/Private/VT/RuntimeVirtualTextureSceneProxy.cpp
--- a/Engine/Source/Runtime/Renderer/Private/VT/RuntimeVirtualTextureSceneProxy.cpp
+++ b/Engine/Source/Runtime/Renderer/Private/VT/RuntimeVirtualTextureSceneProxy.cpp
@@ -2,6 +2,7 @@
 
 #include "VT/RuntimeVirtualTextureSceneProxy.h"
 
+#include "ComponentRecreateRenderStateContext.h"
 #include "Components/RuntimeVirtualTextureComponent.h"
 #include "RendererOnScreenNotification.h"
 #include "VirtualTextureSystem.h"
@@ -13,6 +14,16 @@
 
 #define LOCTEXT_NAMESPACE "VirtualTexture"
 
+static TAutoConsoleVariable<int32> CVarVTStreamingMips(
+	TEXT("r.VT.RVT.StreamingMips"),
+	1,
+	TEXT("Enable streaming mips for RVT"),
+	FConsoleVariableDelegate::CreateLambda([](IConsoleVariable* InVariable)
+	{
+		FGlobalComponentRecreateRenderStateContext Context;
+	}),
+	ECVF_Default);
+
 int32 FRuntimeVirtualTextureSceneProxy::ProducerIdGenerator = 1;
 
 FRuntimeVirtualTextureSceneProxy::FRuntimeVirtualTextureSceneProxy(URuntimeVirtualTextureComponent* InComponent)
@@ -56,7 +67,19 @@ FRuntimeVirtualTextureSceneProxy::FRuntimeVirtualTextureSceneProxy(URuntimeVirtu
 		// This is bound with the main producer so that one allocated VT can use both runtime or streaming producers dependent on mip level.
 		if (InComponent->IsStreamingLowMips())
 		{
-			if (InComponent->IsStreamingTextureInvalid())
+			if (CVarVTStreamingMips.GetValueOnAnyThread() == 0)
+			{
+#if !UE_BUILD_SHIPPING
+				// Notify that streaming texture is turned off.
+				OnScreenWarningDelegateHandle = FRendererOnScreenNotification::Get().AddLambda([](FCoreDelegates::FSeverityMessageMap& OutMessages)
+				{
+					OutMessages.Add(
+						FCoreDelegates::EOnScreenMessageSeverity::Warning,
+						LOCTEXT("SVTDisabled", "Runtime Virtual Texture streaming mips disabled."));
+				});
+#endif
+			}
+			else if (InComponent->IsStreamingTextureInvalid())
 			{
 #if !UE_BUILD_SHIPPING
 				// Notify that streaming texture is invalid since this can cause performance regression.
