VCloud on translucent fading in region depth can now be controled via a cvar.

diff --git a/Engine/Shaders/Private/BasePassPixelShader.usf b/Engine/Shaders/Private/BasePassPixelShader.usf
--- a/Engine/Shaders/Private/BasePassPixelShader.usf
+++ b/Engine/Shaders/Private/BasePassPixelShader.usf
@@ -1338,7 +1338,7 @@ void FPixelShaderInOut_MainPS(
 			NDCPosition, LWCHackToFloat(MaterialParameters.AbsoluteWorldPosition), LWCHackToFloat(ResolvedView.WorldCameraOrigin),
 			TranslucentBasePass.VolumetricCloudColor, TranslucentBasePass.VolumetricCloudColorSampler,
 			TranslucentBasePass.VolumetricCloudDepth, TranslucentBasePass.VolumetricCloudDepthSampler,
-			OneOverPreExposure, Fogging);
+			OneOverPreExposure, Fogging, TranslucentBasePass.SoftBlendingDistanceKm);
 	}
 
 #endif
diff --git a/Engine/Shaders/Private/BasePassVertexShader.usf b/Engine/Shaders/Private/BasePassVertexShader.usf
--- a/Engine/Shaders/Private/BasePassVertexShader.usf
+++ b/Engine/Shaders/Private/BasePassVertexShader.usf
@@ -108,7 +108,7 @@ void Main(
 			Output.Position, WorldPosition.xyz, ResolvedView.TranslatedWorldCameraOrigin,
 			TranslucentBasePass.VolumetricCloudColor, TranslucentBasePass.VolumetricCloudColorSampler,
 			TranslucentBasePass.VolumetricCloudDepth, TranslucentBasePass.VolumetricCloudDepthSampler,
-			OneOverPreExposure, Output.BasePassInterpolants.VertexFog);
+			OneOverPreExposure, Output.BasePassInterpolants.VertexFog, TranslucentBasePass.SoftBlendingDistanceKm);
 	}
 #endif
 
diff --git a/Engine/Shaders/Private/VolumetricCloudCommon.ush b/Engine/Shaders/Private/VolumetricCloudCommon.ush
--- a/Engine/Shaders/Private/VolumetricCloudCommon.ush
+++ b/Engine/Shaders/Private/VolumetricCloudCommon.ush
@@ -77,7 +77,7 @@ float4 GetCloudLuminanceTransmittanceOverFog(
 	float4 NDC, float3 SampledWorldPos, float3 CameraWorldPos,
 	Texture2D<float4> VolumetricCloudColor, SamplerState VolumetricCloudColorSampler,
 	Texture2D<float4> VolumetricCloudDepth, SamplerState VolumetricCloudDepthSampler,
-	float OneOverPreExposure, float4 CurrentVertexFog)
+	float OneOverPreExposure, float4 CurrentVertexFog, float SoftBlendingDistanceKm)
 {
 	float2 ScreenUv = (NDC.xy / NDC.ww) * float2(0.5f, -0.5f) + 0.5f;
 
@@ -99,7 +99,6 @@ float4 GetCloudLuminanceTransmittanceOverFog(
 		float4 CloudFog = Texture2DSample(VolumetricCloudColor, VolumetricCloudColorSampler, ScreenUv);
 		CloudFog.rgb *= OneOverPreExposure;
 
-		const float SoftBlendingDistanceKm = 0.5f;// 500 meters
 		const float BlendFactor = saturate(DepthInsideKm / SoftBlendingDistanceKm);
 		const float4 FogToApplyUnder = CurrentVertexFog;
 		const float4 FogToApplyOver = lerp(float4(0.0, 0.0, 0.0, 1.0), CloudFog, BlendFactor);
diff --git a/Engine/Source/Runtime/Renderer/Private/BasePassRendering.h b/Engine/Source/Runtime/Renderer/Private/BasePassRendering.h
--- a/Engine/Source/Runtime/Renderer/Private/BasePassRendering.h
+++ b/Engine/Source/Runtime/Renderer/Private/BasePassRendering.h
@@ -114,6 +114,7 @@ BEGIN_GLOBAL_SHADER_PARAMETER_STRUCT(FTranslucentBasePassUniformParameters,)
 	SHADER_PARAMETER_TEXTURE(Texture2D, VolumetricCloudDepth)
 	SHADER_PARAMETER_SAMPLER(SamplerState, VolumetricCloudDepthSampler)
 	SHADER_PARAMETER(float, ApplyVolumetricCloudOnTransparent)
+	SHADER_PARAMETER(float, SoftBlendingDistanceKm)
 	// Translucency Lighting Volume
 	SHADER_PARAMETER_STRUCT_INCLUDE(FTranslucencyLightingVolumeParameters, TranslucencyLightingVolume)
 	SHADER_PARAMETER_STRUCT_INCLUDE(FLumenTranslucencyLightingParameters, LumenParameters)
diff --git a/Engine/Source/Runtime/Renderer/Private/TranslucentRendering.cpp b/Engine/Source/Runtime/Renderer/Private/TranslucentRendering.cpp
--- a/Engine/Source/Runtime/Renderer/Private/TranslucentRendering.cpp
+++ b/Engine/Source/Runtime/Renderer/Private/TranslucentRendering.cpp
@@ -81,6 +81,10 @@ static TAutoConsoleVariable<int32> CVarTranslucencyUpperBoundQuantization(
 	TEXT("Only recommended for use with the transient allocator (on supported platforms) with a large transient texture cache (e.g RHI.TransientAllocator.TextureCacheSize=512)"),
 	ECVF_RenderThreadSafe | ECVF_Default);
 
+static TAutoConsoleVariable<float> CVarVolumetricCloudSoftBlendingDistanceOnTranslucent(
+	TEXT("r.VolumetricCloud.SoftBlendingDistanceOnTranslucent"), 0.5,
+	TEXT("The soft blending in distance in kilometer used to soft blend in cloud over translucent from the evaluated start depth."),
+	ECVF_RenderThreadSafe | ECVF_Scalability);
 
 int32 GSeparateTranslucencyUpsampleMode = 1;
 static FAutoConsoleVariableRef CVarSeparateTranslucencyUpsampleMode(
@@ -876,6 +880,7 @@ TRDGUniformBufferRef<FTranslucentBasePassUniformParameters> CreateTranslucentBas
 			BasePassParameters.PrevSceneColorBilinearUVMax = FVector2f(1.0f, 1.0f);
 		}
 
+		BasePassParameters.SoftBlendingDistanceKm = FMath::Max(0.0001f, CVarVolumetricCloudSoftBlendingDistanceOnTranslucent.GetValueOnRenderThread());
 		BasePassParameters.ApplyVolumetricCloudOnTransparent = 0.0f;
 		BasePassParameters.VolumetricCloudColor = nullptr;
 		BasePassParameters.VolumetricCloudDepth = nullptr;
