Post-processing: Reconnect r.LUT.Size with the combined LUT/tonemapping texture size.

Increasing the LUT resolution can become necessary to mitigate precision artifacts with wide-gamut working color spaces.

diff --git a/Engine/Shaders/Private/PostProcessCombineLUTs.usf b/Engine/Shaders/Private/PostProcessCombineLUTs.usf
--- a/Engine/Shaders/Private/PostProcessCombineLUTs.usf
+++ b/Engine/Shaders/Private/PostProcessCombineLUTs.usf
@@ -34,7 +34,7 @@ half4 OverlayColor;
 uint bUseMobileTonemapper;
 uint bIsTemperatureWhiteBalance;
 
-static const float LUTSize = 32;
+float LUTSize;
 
 float WhiteTemp;
 float WhiteTint;
@@ -284,19 +284,19 @@ float4 CombineLUTsCommon(float2 InUV, uint InLayerIndex)
 
 		// BLENDCOUNT is the number of LUT that are blended together including the neutral one
 	#if BLENDCOUNT >= 2
-		GammaColor += GET_SCALAR_ARRAY_ELEMENT(LUTWeights, 1) * UnwrappedTexture3DSample( Textures_1, Samplers_1, UVW, 16 ).rgb;
+		GammaColor += GET_SCALAR_ARRAY_ELEMENT(LUTWeights, 1) * UnwrappedTexture3DSample( Textures_1, Samplers_1, UVW, 16, 1 / 16.0 ).rgb;
 	#endif
 
 	#if BLENDCOUNT >= 3
-		GammaColor += GET_SCALAR_ARRAY_ELEMENT(LUTWeights, 2) * UnwrappedTexture3DSample( Textures_2, Samplers_2, UVW, 16 ).rgb;
+		GammaColor += GET_SCALAR_ARRAY_ELEMENT(LUTWeights, 2) * UnwrappedTexture3DSample( Textures_2, Samplers_2, UVW, 16, 1 / 16.0 ).rgb;
 	#endif
 
 	#if BLENDCOUNT >= 4
-		GammaColor += GET_SCALAR_ARRAY_ELEMENT(LUTWeights, 3) * UnwrappedTexture3DSample( Textures_3, Samplers_3, UVW, 16 ).rgb;
+		GammaColor += GET_SCALAR_ARRAY_ELEMENT(LUTWeights, 3) * UnwrappedTexture3DSample( Textures_3, Samplers_3, UVW, 16, 1 / 16.0 ).rgb;
 	#endif
 
 	#if BLENDCOUNT >= 5
-		GammaColor += GET_SCALAR_ARRAY_ELEMENT(LUTWeights, 4) * UnwrappedTexture3DSample( Textures_4, Samplers_4, UVW, 16 ).rgb;
+		GammaColor += GET_SCALAR_ARRAY_ELEMENT(LUTWeights, 4) * UnwrappedTexture3DSample( Textures_4, Samplers_4, UVW, 16, 1 / 16.0 ).rgb;
 	#endif
 
 		// Back to linear space, AP1 gamut
diff --git a/Engine/Shaders/Private/PostProcessCommon.ush b/Engine/Shaders/Private/PostProcessCommon.ush
--- a/Engine/Shaders/Private/PostProcessCommon.ush
+++ b/Engine/Shaders/Private/PostProcessCommon.ush
@@ -164,18 +164,18 @@ float2 VignetteSpace(float2 Pos)
 	return VignetteSpace(Pos, ViewportSize.y * ViewportSize.z);
 }
 
-half4 UnwrappedTexture3DSample( Texture2D Texture, SamplerState Sampler, float3 UVW, float Size )
+half4 UnwrappedTexture3DSample( Texture2D Texture, SamplerState Sampler, float3 UVW, float Size, float InvSize )
 {
 	// a volume texture 16x16x16 would be unwrapped to a 2d texture 256x16
 
 	float IntW = floor( UVW.z * Size - 0.5 );
 	half FracW = UVW.z * Size - 0.5 - IntW;
 
-	float U = ( UVW.x + IntW ) / Size;
+	float U = ( UVW.x + IntW ) * InvSize;
 	float V = UVW.y;
 
 	half4 RG0 = Texture2DSample( Texture, Sampler, float2(U, V) );
-	half4 RG1 = Texture2DSample( Texture, Sampler, float2(U + 1.0f / Size, V) );
+	half4 RG1 = Texture2DSample( Texture, Sampler, float2(U + InvSize, V) );
 
 	return lerp(RG0, RG1, FracW);
 }
diff --git a/Engine/Shaders/Private/PostProcessTonemap.usf b/Engine/Shaders/Private/PostProcessTonemap.usf
--- a/Engine/Shaders/Private/PostProcessTonemap.usf
+++ b/Engine/Shaders/Private/PostProcessTonemap.usf
@@ -156,7 +156,10 @@ float ComputeFilmGrainIntensity(float3 LinearColorRGBGamut)
 #endif 
 SamplerState        ColorGradingLUTSampler;
 
-static const float LUTSize = 32;
+float LUTSize;
+float InvLUTSize; // 1 / LUTSize
+float LUTScale; // (LUTSize - 1) / LUTSize
+float LUTOffset; // 0.5 / LUTSize
 
 half3 ColorLookupTable( half3 LinearColor )
 {
@@ -171,12 +174,12 @@ half3 ColorLookupTable( half3 LinearColor )
 
 	#endif
 
-	float3 UVW = LUTEncodedColor * ((LUTSize - 1) / LUTSize) + (0.5f / LUTSize);
+	float3 UVW = LUTEncodedColor * LUTScale + LUTOffset;
 
 #if USE_VOLUME_LUT == 1
 	half3 OutDeviceColor = Texture3DSample( ColorGradingLUT, ColorGradingLUTSampler, UVW ).rgb;
 #else
-	half3 OutDeviceColor = UnwrappedTexture3DSample( ColorGradingLUT, ColorGradingLUTSampler, UVW, LUTSize ).rgb;
+	half3 OutDeviceColor = UnwrappedTexture3DSample( ColorGradingLUT, ColorGradingLUTSampler, UVW, LUTSize, InvLUTSize ).rgb;
 #endif
 	
 	return OutDeviceColor * 1.05;
diff --git a/Engine/Source/Runtime/Renderer/Private/PostProcess/PostProcessCombineLUTs.cpp b/Engine/Source/Runtime/Renderer/Private/PostProcess/PostProcessCombineLUTs.cpp
--- a/Engine/Source/Runtime/Renderer/Private/PostProcess/PostProcessCombineLUTs.cpp
+++ b/Engine/Source/Runtime/Renderer/Private/PostProcess/PostProcessCombineLUTs.cpp
@@ -30,10 +30,9 @@ TAutoConsoleVariable<float> CVarColorMax(
 	TEXT("Value should be around 1, smaller values darken the highlights, larger values move more colors towards white, Default: 1"),
 	ECVF_RenderThreadSafe);
 
-int32 GLUTSize = 32;
-FAutoConsoleVariableRef CVarLUTSize(
+TAutoConsoleVariable<int32> CVarLUTSize(
 	TEXT("r.LUT.Size"),
-	GLUTSize,
+	32,
 	TEXT("Size of film LUT"),
 	ECVF_RenderThreadSafe);
 
@@ -65,7 +64,7 @@ struct FColorTransform
 };
 } //! namespace
 
-// false:use 256x16 texture / true:use volume texture (faster, requires geometry shader)
+// false:use 1024x32 texture / true:use volume texture (faster, requires geometry shader)
 // USE_VOLUME_LUT: needs to be the same for C++ and HLSL.
 // Safe to use at pipeline and run time.
 bool PipelineVolumeTextureLUTSupportGuaranteedAtRuntime(EShaderPlatform Platform)
@@ -131,6 +130,7 @@ BEGIN_SHADER_PARAMETER_STRUCT(FCombineLUTParameters, )
 	SHADER_PARAMETER(FVector4f, ColorGammaHighlights)
 	SHADER_PARAMETER(FVector4f, ColorGainHighlights)
 	SHADER_PARAMETER(FVector4f, ColorOffsetHighlights)
+	SHADER_PARAMETER(float, LUTSize)
 	SHADER_PARAMETER(float, WhiteTemp)
 	SHADER_PARAMETER(float, WhiteTint)
 	SHADER_PARAMETER(float, ColorCorrectionShadowsMax)
@@ -192,6 +192,7 @@ void GetCombineLUTParameters(
 
 	// White balance
 	Parameters.bIsTemperatureWhiteBalance = Settings.TemperatureType == ETemperatureMethod::TEMP_WhiteBalance;
+	Parameters.LUTSize = (float)CVarLUTSize->GetInt();
 	Parameters.WhiteTemp = Settings.WhiteTemp;
 	Parameters.WhiteTint = Settings.WhiteTint;
 
@@ -427,8 +428,9 @@ FRDGTextureRef AddCombineLUTPass(FRDGBuilder& GraphBuilder, const FViewInfo& Vie
 	const bool bUseFloatOutput = ViewFamily.SceneCaptureSource == SCS_FinalColorHDR || ViewFamily.SceneCaptureSource == SCS_FinalToneCurveHDR;
 
 	// Attempt to register the persistent view LUT texture.
+	const int32 LUTSize = CVarLUTSize->GetInt();
 	FRDGTextureRef OutputTexture = TryRegisterExternalTexture(GraphBuilder,
-		View.GetTonemappingLUT(GraphBuilder.RHICmdList, GLUTSize, bUseVolumeTextureLUT, bUseComputePass, bUseFloatOutput));
+		View.GetTonemappingLUT(GraphBuilder.RHICmdList, LUTSize, bUseVolumeTextureLUT, bUseComputePass, bUseFloatOutput));
 
 	View.SetValidTonemappingLUT();
 
@@ -436,14 +438,14 @@ FRDGTextureRef AddCombineLUTPass(FRDGBuilder& GraphBuilder, const FViewInfo& Vie
 	if (!OutputTexture)
 	{
 		OutputTexture = GraphBuilder.CreateTexture(
-			Translate(FSceneViewState::CreateLUTRenderTarget(GLUTSize, bUseVolumeTextureLUT, bUseComputePass, bUseFloatOutput)),
+			Translate(FSceneViewState::CreateLUTRenderTarget(LUTSize, bUseVolumeTextureLUT, bUseComputePass, bUseFloatOutput)),
 			TEXT("CombineLUT"));
 	}
 
 	check(OutputTexture);
 
-	// For a 3D texture, the viewport is 16x16 (per slice); for a 2D texture, it's unwrapped to 256x16.
-	const FIntPoint OutputViewSize(bUseVolumeTextureLUT ? GLUTSize : GLUTSize * GLUTSize, GLUTSize);
+	// For a 3D texture, the viewport defaults to 32x32 (per slice); for a 2D texture, it's unwrapped to 1024x32.
+	const FIntPoint OutputViewSize(bUseVolumeTextureLUT ? LUTSize : LUTSize * LUTSize, LUTSize);
 
 	FLUTBlenderShader::FPermutationDomain PermutationVector;
 	PermutationVector.Set<FLUTBlenderShader::FBlendCount>(LocalCount);
@@ -468,7 +470,7 @@ FRDGTextureRef AddCombineLUTPass(FRDGBuilder& GraphBuilder, const FViewInfo& Vie
 
 		FComputeShaderUtils::AddPass(
 			GraphBuilder,
-			RDG_EVENT_NAME("CombineLUTs (CS)"),
+			RDG_EVENT_NAME("CombineLUTs %d (CS)", LUTSize),
 			ComputeShader,
 			PassParameters,
 			FIntVector(GroupSizeXY, GroupSizeXY, GroupSizeZ));
@@ -485,10 +487,10 @@ FRDGTextureRef AddCombineLUTPass(FRDGBuilder& GraphBuilder, const FViewInfo& Vie
 		TShaderMapRef<FLUTBlenderPS> PixelShader(View.ShaderMap, PermutationVector);
 
 		GraphBuilder.AddPass(
-			RDG_EVENT_NAME("CombineLUTS (PS)"),
+			RDG_EVENT_NAME("CombineLUTS %d (PS)", LUTSize),
 			PassParameters,
 			ERDGPassFlags::Raster,
-			[&View, PixelShader, PassParameters, bUseVolumeTextureLUT] (FRHICommandList& RHICmdList)
+			[&View, PixelShader, PassParameters, bUseVolumeTextureLUT, LUTSize] (FRHICommandList& RHICmdList)
 		{
 			FGraphicsPipelineStateInitializer GraphicsPSOInit;
 			RHICmdList.ApplyCachedRenderTargets(GraphicsPSOInit);
@@ -498,7 +500,7 @@ FRDGTextureRef AddCombineLUTPass(FRDGBuilder& GraphBuilder, const FViewInfo& Vie
 
 			if (bUseVolumeTextureLUT)
 			{
-				const FVolumeBounds VolumeBounds(GLUTSize);
+				const FVolumeBounds VolumeBounds(LUTSize);
 
 				TShaderMapRef<FWriteToSliceVS> VertexShader(View.ShaderMap);
 				TOptionalShaderMapRef<FWriteToSliceGS> GeometryShader(View.ShaderMap);
@@ -531,11 +533,11 @@ FRDGTextureRef AddCombineLUTPass(FRDGBuilder& GraphBuilder, const FViewInfo& Vie
 				DrawRectangle(
 					RHICmdList,
 					0, 0,										// XY
-					GLUTSize * GLUTSize, GLUTSize,				// SizeXY
+					LUTSize * LUTSize, LUTSize,				// SizeXY
 					0, 0,										// UV
-					GLUTSize * GLUTSize, GLUTSize,				// SizeUV
-					FIntPoint(GLUTSize * GLUTSize, GLUTSize),	// TargetSize
-					FIntPoint(GLUTSize * GLUTSize, GLUTSize),	// TextureSize
+					LUTSize * LUTSize, LUTSize,				// SizeUV
+					FIntPoint(LUTSize * LUTSize, LUTSize),	// TargetSize
+					FIntPoint(LUTSize * LUTSize, LUTSize),	// TextureSize
 					VertexShader,
 					EDRF_UseTriangleOptimization);
 			}
diff --git a/Engine/Source/Runtime/Renderer/Private/PostProcess/PostProcessTonemap.cpp b/Engine/Source/Runtime/Renderer/Private/PostProcess/PostProcessTonemap.cpp
--- a/Engine/Source/Runtime/Renderer/Private/PostProcess/PostProcessTonemap.cpp
+++ b/Engine/Source/Runtime/Renderer/Private/PostProcess/PostProcessTonemap.cpp
@@ -318,6 +318,10 @@ BEGIN_SHADER_PARAMETER_STRUCT(FTonemapParameters, )
 	SHADER_PARAMETER(FVector4f, TonemapperParams)
 	SHADER_PARAMETER(FVector4f, LensPrincipalPointOffsetScale)
 	SHADER_PARAMETER(FVector4f, LensPrincipalPointOffsetScaleInverse)
+	SHADER_PARAMETER(float, LUTSize)
+	SHADER_PARAMETER(float, InvLUTSize)
+	SHADER_PARAMETER(float, LUTScale)
+	SHADER_PARAMETER(float, LUTOffset)
 	SHADER_PARAMETER(float, DefaultEyeExposure)
 	SHADER_PARAMETER(float, EditorNITLevel)
 	SHADER_PARAMETER(uint32, bOutputInHDR)
@@ -726,6 +730,8 @@ FScreenPassTexture AddTonemapPass(FRDGBuilder& GraphBuilder, const FViewInfo& Vi
 			CommonParameters.FilmGrain.FilmGrainTextureConstants = nullptr;
 		}
 	}
+	
+	const float LUTSize = Inputs.ColorGradingTexture ? (float)Inputs.ColorGradingTexture->Desc.Extent.Y : /* unused (default): */ 32.0f;
 
 	CommonParameters.OutputDevice = GetTonemapperOutputDeviceParameters(ViewFamily);
 	CommonParameters.Color = GetScreenPassTextureViewportParameters(SceneColorViewport);
@@ -746,6 +752,10 @@ FScreenPassTexture AddTonemapPass(FRDGBuilder& GraphBuilder, const FViewInfo& Vi
 	CommonParameters.DefaultEyeExposure = DefaultEyeExposure;
 	CommonParameters.EditorNITLevel = EditorNITLevel;
 	CommonParameters.bOutputInHDR = ViewFamily.bIsHDR;
+	CommonParameters.LUTSize = LUTSize;
+	CommonParameters.InvLUTSize = 1.0f / LUTSize;
+	CommonParameters.LUTScale = (LUTSize - 1.0f) / LUTSize;
+	CommonParameters.LUTOffset = 0.5f / LUTSize;
 	CommonParameters.LensPrincipalPointOffsetScale = View.LensPrincipalPointOffsetScale;
 
 	// TODO: PostProcessSettings.BloomDirtMask->GetResource() is not thread safe
