Fixed issues with mesh previews

diff --git a/Engine/Source/Runtime/Renderer/Private/SceneRendering.cpp b/Engine/Source/Runtime/Renderer/Private/SceneRendering.cpp
--- a/Engine/Source/Runtime/Renderer/Private/SceneRendering.cpp
+++ b/Engine/Source/Runtime/Renderer/Private/SceneRendering.cpp
@@ -1773,7 +1773,13 @@ void FViewInfo::SetupUniformBufferParameters(
 		ViewUniformShaderParameters.FarShadowStaticMeshLODBias = 0;
 	}
 
-	ViewUniformShaderParameters.PreIntegratedBRDF = GEngine->PreIntegratedSkinBRDFTexture->GetResource()->TextureRHI;
+	{
+		const FTextureResource* TextureResource = GEngine->PreIntegratedSkinBRDFTexture->GetResource();
+		if (TextureResource)
+		{
+			ViewUniformShaderParameters.PreIntegratedBRDF = TextureResource->TextureRHI;
+		}
+	}
 
 	ViewUniformShaderParameters.GlobalVirtualTextureMipBias = FVirtualTextureSystem::Get().GetGlobalMipBias();
 
