Includes the material name in post process material's draw event even on test builds

diff --git a/Engine/Source/Runtime/Renderer/Private/PostProcess/PostProcessMaterial.cpp b/Engine/Source/Runtime/Renderer/Private/PostProcess/PostProcessMaterial.cpp
--- a/Engine/Source/Runtime/Renderer/Private/PostProcess/PostProcessMaterial.cpp
+++ b/Engine/Source/Runtime/Renderer/Private/PostProcess/PostProcessMaterial.cpp
@@ -588,7 +588,11 @@ FScreenPassTexture AddPostProcessMaterialPass(
 	{
 		AddDrawScreenPass(
 			GraphBuilder,
-			RDG_EVENT_NAME("PostProcessMaterial"),
+#if RDG_EVENTS == RDG_EVENTS_NONE
+			RDG_EVENT_NAME(""),
+#else
+			FRDGEventName(*Material->GetAssetName()),
+#endif
 			View,
 			OutputViewport,
 			SceneColorViewport,
