Updated the mobile renderer to render the geometry when using "Sort Triangles".

diff --git a/Engine/Source/Runtime/Renderer/Private/MobileShadingRenderer.cpp b/Engine/Source/Runtime/Renderer/Private/MobileShadingRenderer.cpp
--- a/Engine/Source/Runtime/Renderer/Private/MobileShadingRenderer.cpp
+++ b/Engine/Source/Runtime/Renderer/Private/MobileShadingRenderer.cpp
@@ -954,6 +954,15 @@ void FMobileSceneRenderer::Render(FRDGBuilder& GraphBuilder)
 		RenderCustomDepthPass(GraphBuilder, SceneTextures.CustomDepth, SceneTextures.GetSceneTextureShaderParameters(FeatureLevel));
 	}
 
+	// Sort objects' triangles
+	for (FViewInfo& View : Views)
+	{
+		if (View.ShouldRenderView() && OIT::IsEnabled(EOITSortingType::SortedTriangles, View))
+		{
+			OIT::AddSortTrianglesPass(GraphBuilder, View, Scene->OITSceneData, FTriangleSortingOrder::BackToFront);
+		}
+	}
+
 	FRDGTextureRef ViewFamilyTexture = TryCreateViewFamilyTexture(GraphBuilder, ViewFamily);
 
 	const FRDGSystemTextures& SystemTextures = FRDGSystemTextures::Get(GraphBuilder);
