Fix crash in path tracer state invalidation when building nanite

Ran into this crash when nanite RT is enabled, a static mesh gets re-built for nanite for the first time and path tracing is enabled. Also added a condition to skip invalidation for offline renders since we wouldn't want to loose samples in progress from a streaming request (for example in the context of motion blur).

diff --git a/Engine/Source/Runtime/Renderer/Private/DeferredShadingRenderer.cpp b/Engine/Source/Runtime/Renderer/Private/DeferredShadingRenderer.cpp
--- a/Engine/Source/Runtime/Renderer/Private/DeferredShadingRenderer.cpp
+++ b/Engine/Source/Runtime/Renderer/Private/DeferredShadingRenderer.cpp
@@ -1845,10 +1845,13 @@ bool FDeferredShadingSceneRenderer::DispatchRayTracingWorldUpdates(FRDGBuilder&
 		const bool bAnyBlasRebuilt = Nanite::GRayTracingManager.ProcessBuildRequests(GraphBuilder);
 		if (bAnyBlasRebuilt)
 		{
-			for (int32 ViewIndex = 0; ViewIndex < Views.Num(); ViewIndex++)
+			for (FViewInfo& View : Views)
 			{
-				FViewInfo& View = Views[ViewIndex];
-				View.ViewState->PathTracingInvalidate();
+				if (View.ViewState != nullptr && !View.bIsOfflineRender)
+				{
+					// don't invalidate in the offline case because we only get one attempt at rendering each sample
+					View.ViewState->PathTracingInvalidate();
+				}
 			}
 		}
 	}
@@ -2242,10 +2245,13 @@ void FDeferredShadingSceneRenderer::Render(FRDGBuilder& GraphBuilder)
 
 		if (bNaniteRayTracingModeChanged)
 		{
-			for (int32 ViewIndex = 0; ViewIndex < Views.Num(); ViewIndex++)
+			for (FViewInfo& View : Views)
 			{
-				FViewInfo& View = Views[ViewIndex];
-				View.ViewState->PathTracingInvalidate();
+				if (View.ViewState != nullptr && !View.bIsOfflineRender)
+				{
+					// don't invalidate in the offline case because we only get one attempt at rendering each sample
+					View.ViewState->PathTracingInvalidate();
+				}
 			}
 		}
 	}
