Updated cloud mode1 to not do any temporal reconstruction. It now traces at half resolution and then upsample at full resolution, leveraging checkerboarded minmax depth.
It is a good mid ground between mode 0 (temporal reprojection) and mode 2 (no depth intersection). It has not temporal reprojection issue and look sharper while it intersects with opaque depth with nice crisp edges.

diff --git a/Engine/Source/Runtime/Renderer/Private/BasePassRendering.h b/Engine/Source/Runtime/Renderer/Private/BasePassRendering.h
--- a/Engine/Source/Runtime/Renderer/Private/BasePassRendering.h
+++ b/Engine/Source/Runtime/Renderer/Private/BasePassRendering.h
@@ -109,9 +109,9 @@ BEGIN_GLOBAL_SHADER_PARAMETER_STRUCT(FTranslucentBasePassUniformParameters,)
 	SHADER_PARAMETER_RDG_TEXTURE(Texture2D, PrevSceneColor)
 	SHADER_PARAMETER_SAMPLER(SamplerState, PrevSceneColorSampler)
 	// Volumetric cloud
-	SHADER_PARAMETER_TEXTURE(Texture2D, VolumetricCloudColor)
+	SHADER_PARAMETER_RDG_TEXTURE(Texture2D, VolumetricCloudColor)
 	SHADER_PARAMETER_SAMPLER(SamplerState, VolumetricCloudColorSampler)
-	SHADER_PARAMETER_TEXTURE(Texture2D, VolumetricCloudDepth)
+	SHADER_PARAMETER_RDG_TEXTURE(Texture2D, VolumetricCloudDepth)
 	SHADER_PARAMETER_SAMPLER(SamplerState, VolumetricCloudDepthSampler)
 	SHADER_PARAMETER(float, ApplyVolumetricCloudOnTransparent)
 	SHADER_PARAMETER(float, SoftBlendingDistanceKm)
diff --git a/Engine/Source/Runtime/Renderer/Private/TranslucentRendering.cpp b/Engine/Source/Runtime/Renderer/Private/TranslucentRendering.cpp
--- a/Engine/Source/Runtime/Renderer/Private/TranslucentRendering.cpp
+++ b/Engine/Source/Runtime/Renderer/Private/TranslucentRendering.cpp
@@ -888,20 +888,32 @@ TRDGUniformBufferRef<FTranslucentBasePassUniformParameters> CreateTranslucentBas
 		BasePassParameters.VolumetricCloudDepthSampler = TStaticSamplerState<SF_Bilinear, AM_Clamp, AM_Clamp, AM_Clamp>::GetRHI();
 		if (IsVolumetricRenderTargetEnabled() && View.ViewState)
 		{
-			TRefCountPtr<IPooledRenderTarget> VolumetricReconstructRT = View.ViewState->VolumetricCloudRenderTarget.GetDstVolumetricReconstructRT();
-			if (VolumetricReconstructRT.IsValid())
+			if (View.ViewState->VolumetricCloudRenderTarget.GetMode() == 1)
 			{
-				TRefCountPtr<IPooledRenderTarget> VolumetricReconstructRTDepth = View.ViewState->VolumetricCloudRenderTarget.GetDstVolumetricReconstructRTDepth();
-
-				BasePassParameters.VolumetricCloudColor = VolumetricReconstructRT->GetRHI();
-				BasePassParameters.VolumetricCloudDepth = VolumetricReconstructRTDepth->GetRHI();
-				BasePassParameters.ApplyVolumetricCloudOnTransparent = 1.0f;
+				FRDGTextureRef VolumetricReconstructRT = View.ViewState->VolumetricCloudRenderTarget.GetOrCreateVolumetricTracingRT(GraphBuilder);
+				if (VolumetricReconstructRT)
+				{
+					BasePassParameters.VolumetricCloudColor = VolumetricReconstructRT;
+					BasePassParameters.VolumetricCloudDepth = View.ViewState->VolumetricCloudRenderTarget.GetOrCreateVolumetricTracingRTDepth(GraphBuilder);
+					BasePassParameters.ApplyVolumetricCloudOnTransparent = 1.0f;
+				}
+			}
+			else
+			{
+				TRefCountPtr<IPooledRenderTarget> VolumetricReconstructRT = View.ViewState->VolumetricCloudRenderTarget.GetDstVolumetricReconstructRT();
+				if (VolumetricReconstructRT.IsValid())
+				{
+					TRefCountPtr<IPooledRenderTarget> VolumetricReconstructRTDepth = View.ViewState->VolumetricCloudRenderTarget.GetDstVolumetricReconstructRTDepth();
+					BasePassParameters.VolumetricCloudColor = GraphBuilder.RegisterExternalTexture(VolumetricReconstructRT);
+					BasePassParameters.VolumetricCloudDepth = GraphBuilder.RegisterExternalTexture(VolumetricReconstructRTDepth);
+					BasePassParameters.ApplyVolumetricCloudOnTransparent = 1.0f;
+				}
 			}
 		}
 		if (BasePassParameters.VolumetricCloudColor == nullptr)
 		{
-			BasePassParameters.VolumetricCloudColor = GSystemTextures.BlackAlphaOneDummy->GetRHI();
-			BasePassParameters.VolumetricCloudDepth = GSystemTextures.BlackDummy->GetRHI();
+			BasePassParameters.VolumetricCloudColor = GSystemTextures.GetBlackAlphaOneDummy(GraphBuilder);
+			BasePassParameters.VolumetricCloudDepth = GSystemTextures.GetBlackDummy(GraphBuilder);
 		}
 
 		FIntPoint ViewportOffset = View.ViewRect.Min;
diff --git a/Engine/Source/Runtime/Renderer/Private/VolumetricCloudRendering.cpp b/Engine/Source/Runtime/Renderer/Private/VolumetricCloudRendering.cpp
--- a/Engine/Source/Runtime/Renderer/Private/VolumetricCloudRendering.cpp
+++ b/Engine/Source/Runtime/Renderer/Private/VolumetricCloudRendering.cpp
@@ -2703,7 +2703,7 @@ bool FSceneRenderer::RenderVolumetricCloud(
 					CloudRC.NoiseFrameIndexModPattern = VRT.GetNoiseFrameIndexModPattern();
 
 					check(VRT.GetMode() != 0 || CloudRC.bCloudDebugViewModeEnabled || HalfResolutionDepthCheckerboardMinMaxTexture);
-					CloudRC.SceneDepthZ = (VRT.GetMode() == 0 || VRT.GetMode() == 3) ? HalfResolutionDepthCheckerboardMinMaxTexture : SceneDepthZ;
+					CloudRC.SceneDepthZ = (VRT.GetMode() == 0 || VRT.GetMode() == 1 || VRT.GetMode() == 3) ? HalfResolutionDepthCheckerboardMinMaxTexture : SceneDepthZ;
 					CloudRC.SceneDepthMinAndMax = bShouldVolumetricCloudTraceWithMinMaxDepth ? QuarterResolutionDepthMinMaxTexture : SceneDepthZ;
 				}
 				else
diff --git a/Engine/Source/Runtime/Renderer/Private/VolumetricRenderTarget.cpp b/Engine/Source/Runtime/Renderer/Private/VolumetricRenderTarget.cpp
--- a/Engine/Source/Runtime/Renderer/Private/VolumetricRenderTarget.cpp
+++ b/Engine/Source/Runtime/Renderer/Private/VolumetricRenderTarget.cpp
@@ -14,8 +14,6 @@
 #include "VolumetricCloudRendering.h"
 #include "RendererUtils.h"
 
-//PRAGMA_DISABLE_OPTIMIZATION
-
 static TAutoConsoleVariable<int32> CVarVolumetricRenderTarget(
 	TEXT("r.VolumetricRenderTarget"), 1,
 	TEXT(""),
@@ -244,6 +242,12 @@ void FVolumetricRenderTargetViewStateData::Initialise(
 				CurrentPixelOffset = FIntPoint(FrameId % VolumetricTracingRTDownsampleFactor, FrameId / VolumetricTracingRTDownsampleFactor);
 			}
 		}
+
+		if (Mode == 1)
+		{
+			// No need to jitter in this case. Mode on is tracing half res and then upsample without reconstruction.
+			CurrentPixelOffset = FIntPoint::ZeroValue;
+		}
 	}
 }
 
@@ -357,7 +361,7 @@ FRDGTextureRef FVolumetricRenderTargetViewStateData::GetOrCreateSrcVolumetricRec
 
 FUintVector4 FVolumetricRenderTargetViewStateData::GetTracingCoordToZbufferCoordScaleBias() const
 {
-	if (Mode == 2 || Mode == 1)
+	if (Mode == 2)
 	{
 		// In this case, the source depth buffer full resolution depth buffer is the full resolution scene one
 		const uint32 CombinedDownsampleFactor = VolumetricReconstructRTDownsampleFactor * VolumetricTracingRTDownsampleFactor;
@@ -516,6 +520,14 @@ void ReconstructVolumetricRenderTarget(
 
 		FVolumetricRenderTargetViewStateData& VolumetricCloudRT = ViewInfo.ViewState->VolumetricCloudRenderTarget;
 
+		if (VolumetricCloudRT.GetMode() == 1)
+		{
+			// In this case, we trace at half resolution using checker boarded min max depth.
+			// We will then directly up sample on screen from half resolution to full resolution. 
+			// No reconstruction needed.
+			continue;
+		}
+
 		FRDGTextureRef DstVolumetric = VolumetricCloudRT.GetOrCreateDstVolumetricReconstructRT(GraphBuilder);
 		FRDGTextureRef DstVolumetricDepth = VolumetricCloudRT.GetOrCreateDstVolumetricReconstructRTDepth(GraphBuilder);
 		FRDGTextureRef SrcTracingVolumetric = VolumetricCloudRT.GetOrCreateVolumetricTracingRT(GraphBuilder);
@@ -632,6 +644,37 @@ IMPLEMENT_GLOBAL_SHADER(FComposeVolumetricRTOverScenePS, "/Engine/Private/Volume
 
 //////////////////////////////////////////////////////////////////////////
 
+static FVector2f GetCompositionFullResolutionToVolumetricBufferResolutionScale(uint32 VRTMode)
+{
+	if (VRTMode == 1)
+	{
+		return FVector2f(0.5f, 2.0f);
+	}
+	return FVector2f(1.0f / float(GetMainDownsampleFactor(VRTMode)), float(GetMainDownsampleFactor(VRTMode)));
+}
+
+static void GetCompositionCloudTextures(uint32 VRTMode, FVolumetricRenderTargetViewStateData& VolumetricCloudRT, FRDGBuilder& GraphBuilder, FRDGTextureRef& VolumetricTexture, FRDGTextureRef& VolumetricDepthTexture)
+{
+	if (VRTMode == 1)
+	{
+		// In this case, we trace at half resolution using checker boarded min max depth.
+		// We will then directly up sample on screen from half resolution to full resolution. 
+		// No reconstruction needed.
+		VolumetricTexture = VolumetricCloudRT.GetOrCreateVolumetricTracingRT(GraphBuilder);
+		VolumetricDepthTexture = VolumetricCloudRT.GetOrCreateVolumetricTracingRTDepth(GraphBuilder);
+	}
+	else
+	{
+		VolumetricTexture = VolumetricCloudRT.GetOrCreateDstVolumetricReconstructRT(GraphBuilder);
+		VolumetricDepthTexture = VolumetricCloudRT.GetOrCreateDstVolumetricReconstructRTDepth(GraphBuilder);
+	}
+}
+
+static int32 GetCompositionUpsamplingMode(uint32 VRTMode, int32 UpsamplingMode)
+{
+	return UpsamplingMode == 3 && (VRTMode == 1 || VRTMode == 2) ? 2 : UpsamplingMode;
+}
+
 void ComposeVolumetricRenderTargetOverScene(
 	FRDGBuilder& GraphBuilder,
 	TArrayView<FViewInfo> Views,
@@ -656,13 +699,14 @@ void ComposeVolumetricRenderTargetOverScene(
 			continue;
 		}
 		FVolumetricRenderTargetViewStateData& VolumetricCloudRT = ViewInfo.ViewState->VolumetricCloudRenderTarget;
-		FRDGTextureRef VolumetricTexture = VolumetricCloudRT.GetOrCreateDstVolumetricReconstructRT(GraphBuilder);
-		FRDGTextureRef VolumetricDepthTexture = VolumetricCloudRT.GetOrCreateDstVolumetricReconstructRTDepth(GraphBuilder);
 
 		// When reconstructed and back buffer resolution matches, force using a pixel perfect upsampling.
 		const uint32 VRTMode = VolumetricCloudRT.GetMode();
-		int UpsamplingMode = VolumetricCloudRT.GetUpsamplingMode();
-		UpsamplingMode = UpsamplingMode == 3 && (VRTMode == 1 || VRTMode == 2) ? 2 : UpsamplingMode;
+		int32 UpsamplingMode = GetCompositionUpsamplingMode(VRTMode, VolumetricCloudRT.GetUpsamplingMode());
+
+		FRDGTextureRef VolumetricTexture = nullptr;
+		FRDGTextureRef VolumetricDepthTexture = nullptr;
+		GetCompositionCloudTextures(VRTMode, VolumetricCloudRT, GraphBuilder, VolumetricTexture, VolumetricDepthTexture);
 
 		// We only support MSAA up to 8 sample and in forward
 		check(SceneDepthTexture->Desc.NumSamples <= 8);
@@ -685,7 +729,7 @@ void ComposeVolumetricRenderTargetOverScene(
 		PassParameters->VolumetricDepthTexture = VolumetricDepthTexture;
 		PassParameters->LinearTextureSampler = TStaticSamplerState<SF_Bilinear>::GetRHI();
 		PassParameters->UvOffsetSampleAcceptanceWeight = GetUvNoiseSampleAcceptanceWeight();
-		PassParameters->FullResolutionToVolumetricBufferResolutionScale = FVector2f(1.0f / float(GetMainDownsampleFactor(VRTMode)), float(GetMainDownsampleFactor(VRTMode)));
+		PassParameters->FullResolutionToVolumetricBufferResolutionScale = GetCompositionFullResolutionToVolumetricBufferResolutionScale(VRTMode);
 		PassParameters->SceneTextures = SceneTextures.UniformBuffer;
 		GetTextureSafeUvCoordBound(PassParameters->VolumetricTexture, PassParameters->VolumetricTextureValidCoordRect, PassParameters->VolumetricTextureValidUvRect);
 
@@ -731,15 +775,16 @@ void ComposeVolumetricRenderTargetOverSceneUnderWater(
 			continue;
 		}
 
-		FVolumetricRenderTargetViewStateData& VolumetricCloudRT = ViewInfo.ViewState->VolumetricCloudRenderTarget;
-		FRDGTextureRef VolumetricTexture = VolumetricCloudRT.GetOrCreateDstVolumetricReconstructRT(GraphBuilder);
-		FRDGTextureRef VolumetricDepthTexture = VolumetricCloudRT.GetOrCreateDstVolumetricReconstructRTDepth(GraphBuilder);
 		const FSceneWithoutWaterTextures::FView& WaterPassViewData = WaterPassData.Views[ViewIndex];
+		FVolumetricRenderTargetViewStateData& VolumetricCloudRT = ViewInfo.ViewState->VolumetricCloudRenderTarget;
 
 		// When reconstructed and back buffer resolution matches, force using a pixel perfect upsampling.
 		const uint32 VRTMode = VolumetricCloudRT.GetMode();
-		int UpsamplingMode = VolumetricCloudRT.GetUpsamplingMode();
-		UpsamplingMode = UpsamplingMode == 3 && (VRTMode == 1 || VRTMode == 2) ? 2 : UpsamplingMode;
+		int32 UpsamplingMode = GetCompositionUpsamplingMode(VRTMode, VolumetricCloudRT.GetUpsamplingMode());
+
+		FRDGTextureRef VolumetricTexture = nullptr;
+		FRDGTextureRef VolumetricDepthTexture = nullptr;
+		GetCompositionCloudTextures(VRTMode, VolumetricCloudRT, GraphBuilder, VolumetricTexture, VolumetricDepthTexture);
 
 		FComposeVolumetricRTOverScenePS::FPermutationDomain PermutationVector;
 		PermutationVector.Set<FComposeVolumetricRTOverScenePS::FUpsamplingMode>(UpsamplingMode);
@@ -757,7 +802,7 @@ void ComposeVolumetricRenderTargetOverSceneUnderWater(
 		PassParameters->LinearTextureSampler = TStaticSamplerState<SF_Bilinear>::GetRHI();
 		PassParameters->WaterLinearDepthSampler = TStaticSamplerState<SF_Point>::GetRHI();
 		PassParameters->UvOffsetSampleAcceptanceWeight = GetUvNoiseSampleAcceptanceWeight();
-		PassParameters->FullResolutionToVolumetricBufferResolutionScale = FVector2f(1.0f / float(GetMainDownsampleFactor(VRTMode)), float(GetMainDownsampleFactor(VRTMode)));
+		PassParameters->FullResolutionToVolumetricBufferResolutionScale = GetCompositionFullResolutionToVolumetricBufferResolutionScale(VRTMode);
 		PassParameters->FullResolutionToWaterBufferScale = FVector2f(1.0f / WaterPassData.RefractionDownsampleFactor, WaterPassData.RefractionDownsampleFactor);
 		PassParameters->SceneWithoutSingleLayerWaterViewRect = FVector4f(WaterPassViewData.ViewRect.Min.X, WaterPassViewData.ViewRect.Min.Y,
 																		WaterPassViewData.ViewRect.Max.X, WaterPassViewData.ViewRect.Max.Y);
@@ -794,11 +839,14 @@ void ComposeVolumetricRenderTargetOverSceneForVisualization(
 			continue;
 		}
 		FVolumetricRenderTargetViewStateData& VolumetricCloudRT = ViewInfo.ViewState->VolumetricCloudRenderTarget;
-		FRDGTextureRef VolumetricTexture = VolumetricCloudRT.GetOrCreateDstVolumetricReconstructRT(GraphBuilder);
-		FRDGTextureRef VolumetricDepthTexture = VolumetricCloudRT.GetOrCreateDstVolumetricReconstructRTDepth(GraphBuilder);
 
 		// When reconstructed and back buffer resolution matches, force using a pixel perfect upsampling.
 		const uint32 VRTMode = VolumetricCloudRT.GetMode();
+		int32 UpsamplingMode = GetCompositionUpsamplingMode(VRTMode, VolumetricCloudRT.GetUpsamplingMode());
+
+		FRDGTextureRef VolumetricTexture = nullptr;
+		FRDGTextureRef VolumetricDepthTexture = nullptr;
+		GetCompositionCloudTextures(VRTMode, VolumetricCloudRT, GraphBuilder, VolumetricTexture, VolumetricDepthTexture);
 
 		FComposeVolumetricRTOverScenePS::FPermutationDomain PermutationVector;
 		PermutationVector.Set<FComposeVolumetricRTOverScenePS::FUpsamplingMode>(0);
@@ -814,7 +862,7 @@ void ComposeVolumetricRenderTargetOverSceneForVisualization(
 		PassParameters->VolumetricDepthTexture = VolumetricDepthTexture;
 		PassParameters->LinearTextureSampler = TStaticSamplerState<SF_Bilinear>::GetRHI();
 		PassParameters->UvOffsetSampleAcceptanceWeight = GetUvNoiseSampleAcceptanceWeight();
-		PassParameters->FullResolutionToVolumetricBufferResolutionScale = FVector2f(1.0f / float(GetMainDownsampleFactor(VRTMode)), float(GetMainDownsampleFactor(VRTMode)));
+		PassParameters->FullResolutionToVolumetricBufferResolutionScale = GetCompositionFullResolutionToVolumetricBufferResolutionScale(VRTMode);
 		PassParameters->SceneTextures = SceneTextures.UniformBuffer;
 		GetTextureSafeUvCoordBound(PassParameters->VolumetricTexture, PassParameters->VolumetricTextureValidCoordRect, PassParameters->VolumetricTextureValidUvRect);
 
