Fix assert in IES atlas function when running with forward lighting.

diff --git a/Engine/Source/Runtime/Renderer/Private/RendererScene.cpp b/Engine/Source/Runtime/Renderer/Private/RendererScene.cpp
--- a/Engine/Source/Runtime/Renderer/Private/RendererScene.cpp
+++ b/Engine/Source/Runtime/Renderer/Private/RendererScene.cpp
@@ -3044,7 +3044,7 @@ void FScene::GetLightIESAtlasSlot(const FLightSceneProxy* Proxy, FLightRenderPar
 {
 	if (Proxy)
 	{
-		check(IsInRenderingThread());
+		check(!IsInGameThread());
 		Out->IESAtlasIndex = IESAtlas::GetAtlasSlot(Proxy->IESAtlasId);
 	}
 }
