VulkanRHI: Use array dummies for AtlasTexture to fix Vulkan validation warnings.

#rnx

diff --git a/Engine/Source/Runtime/Renderer/Private/SceneRendering.cpp b/Engine/Source/Runtime/Renderer/Private/SceneRendering.cpp
--- a/Engine/Source/Runtime/Renderer/Private/SceneRendering.cpp
+++ b/Engine/Source/Runtime/Renderer/Private/SceneRendering.cpp
@@ -1873,9 +1873,9 @@ void FViewInfo::SetupUniformBufferParameters(
 	// IES atlas
 	{
 		FRHITexture* AtlasTexture = IESAtlas::GetAtlasTexture();
-		if (!AtlasTexture && GSystemTextures.BlackDummy.IsValid())
+		if (!AtlasTexture && GSystemTextures.BlackArrayDummy.IsValid())
 		{
-			AtlasTexture = GSystemTextures.BlackDummy->GetRHI();
+			AtlasTexture = GSystemTextures.BlackArrayDummy->GetRHI();
 		}
 				
 		if (AtlasTexture)
@@ -1885,7 +1885,7 @@ void FViewInfo::SetupUniformBufferParameters(
 			ViewUniformShaderParameters.IESAtlasSampler = TStaticSamplerState<SF_Bilinear, AM_Clamp, AM_Clamp, AM_Clamp>::GetRHI();
 			ViewUniformShaderParameters.IESAtlasSizeAndInvSize = FVector4f(AtlasSize.X, AtlasSize.Y, 1.0f / AtlasSize.X, 1.0f / AtlasSize.Y);
 		}
-		ViewUniformShaderParameters.IESAtlasTexture = OrBlack2DIfNull(ViewUniformShaderParameters.IESAtlasTexture);
+		ViewUniformShaderParameters.IESAtlasTexture = OrBlack2DArrayIfNull(ViewUniformShaderParameters.IESAtlasTexture);
 	}
 
 	// Hair global resources 
diff --git a/Engine/Source/Runtime/Renderer/Private/SceneRendering.h b/Engine/Source/Runtime/Renderer/Private/SceneRendering.h
--- a/Engine/Source/Runtime/Renderer/Private/SceneRendering.h
+++ b/Engine/Source/Runtime/Renderer/Private/SceneRendering.h
@@ -2597,6 +2597,13 @@ inline FRHITexture* OrBlack2DIfNull(FRHITexture* Tex)
 	return Result;
 }
 
+inline FRHITexture* OrBlack2DArrayIfNull(FRHITexture* Tex)
+{
+	FRHITexture* Result = Tex ? Tex : GBlackArrayTexture->TextureRHI.GetReference();
+	check(Result);
+	return Result;
+}
+
 inline FRHITexture* OrBlack3DIfNull(FRHITexture* Tex)
 {
 	// we fall back to 2D which are unbound mobile parameters
