Chaning more code to use SetShaderParameters
- MeshPassProcessor converted to use SetShaderParameters instead of setting individual resources. Caching of SRVs and Samplers was also removed.
- UnsetShaderSRVs and UnsetShaderUAVs now set all the null binds in a single SetShaderParameters command.
- Made some FRHIShaderParameterResource utilities more public.

diff --git a/Engine/Source/Runtime/RHI/Public/RHIResources.h b/Engine/Source/Runtime/RHI/Public/RHIResources.h
--- a/Engine/Source/Runtime/RHI/Public/RHIResources.h
+++ b/Engine/Source/Runtime/RHI/Public/RHIResources.h
@@ -4355,14 +4355,35 @@ struct FRHIShaderParameterResource
 		UniformBuffer,
 	};
 
+	FRHIShaderParameterResource() = default;
 	FRHIShaderParameterResource(EType InType, FRHIResource* InResource, uint16 InIndex)
 		: Resource(InResource)
 		, Index(InIndex)
 		, Type(InType)
 	{
 	}
+	FRHIShaderParameterResource(FRHITexture* InTexture, uint16 InIndex)
+		: FRHIShaderParameterResource(FRHIShaderParameterResource::EType::Texture, InTexture, InIndex)
+	{
+	}
+	FRHIShaderParameterResource(FRHIShaderResourceView* InView, uint16 InIndex)
+		: FRHIShaderParameterResource(FRHIShaderParameterResource::EType::ResourceView, InView, InIndex)
+	{
+	}
+	FRHIShaderParameterResource(FRHIUnorderedAccessView* InUAV, uint16 InIndex)
+		: FRHIShaderParameterResource(FRHIShaderParameterResource::EType::UnorderedAccessView, InUAV, InIndex)
+	{
+	}
+	FRHIShaderParameterResource(FRHISamplerState* InSamplerState, uint16 InIndex)
+		: FRHIShaderParameterResource(FRHIShaderParameterResource::EType::Sampler, InSamplerState, InIndex)
+	{
+	}
+	FRHIShaderParameterResource(FRHIUniformBuffer* InUniformBuffer, uint16 InIndex)
+		: FRHIShaderParameterResource(FRHIShaderParameterResource::EType::UniformBuffer, InUniformBuffer, InIndex)
+	{
+	}
 
-	FRHIResource* Resource;
-	uint16        Index;
-	EType         Type;
+	FRHIResource* Resource = nullptr;
+	uint16        Index = 0;
+	EType         Type = EType::Texture;
 };
diff --git a/Engine/Source/Runtime/RenderCore/Private/ShaderParameterStruct.cpp b/Engine/Source/Runtime/RenderCore/Private/ShaderParameterStruct.cpp
--- a/Engine/Source/Runtime/RenderCore/Private/ShaderParameterStruct.cpp
+++ b/Engine/Source/Runtime/RenderCore/Private/ShaderParameterStruct.cpp
@@ -623,41 +623,8 @@ static void ExtractShaderParameters(
 	}
 }
 
-/** Utility class to populate FRHIShaderParameterResource based on types */
-struct FRHIShaderParameterResourceEx : public FRHIShaderParameterResource
-{
-	FRHIShaderParameterResourceEx()
-		: FRHIShaderParameterResource(FRHIShaderParameterResource::EType::Texture, nullptr, 0)
-	{
-	}
-	FRHIShaderParameterResourceEx(FRHITexture* InTexture, uint16 InIndex)
-		: FRHIShaderParameterResource(FRHIShaderParameterResource::EType::Texture, InTexture, InIndex)
-	{
-	}
-	FRHIShaderParameterResourceEx(FRHIShaderResourceView* InView, uint16 InIndex)
-		: FRHIShaderParameterResource(FRHIShaderParameterResource::EType::ResourceView, InView, InIndex)
-	{
-	}
-	FRHIShaderParameterResourceEx(FRHIUnorderedAccessView* InUAV, uint16 InIndex)
-		: FRHIShaderParameterResource(FRHIShaderParameterResource::EType::UnorderedAccessView, InUAV, InIndex)
-	{
-	}
-	FRHIShaderParameterResourceEx(FRHISamplerState* InSamplerState, uint16 InIndex)
-		: FRHIShaderParameterResource(FRHIShaderParameterResource::EType::Sampler, InSamplerState, InIndex)
-	{
-	}
-	FRHIShaderParameterResourceEx(FRHIUniformBuffer* InUniformBuffer, uint16 InIndex)
-		: FRHIShaderParameterResource(FRHIShaderParameterResource::EType::UniformBuffer, InUniformBuffer, InIndex)
-	{
-	}
-};
-
 #define UE_SINGLE_SHADER_PARAMETERS_COMMAND_INCLUDE_UB (0)
 
-inline int16 GetParameterIndex(const FShaderParameterBindings::FResourceParameter& Parameter) { return Parameter.BaseIndex; }
-inline int16 GetParameterIndex(const FShaderParameterBindings::FBindlessResourceParameter& Parameter) { return Parameter.GlobalConstantOffset; }
-inline int16 GetParameterIndex(const FShaderParameterBindings::FParameterStructReference& Parameter) { return Parameter.BufferIndex; }
-
 template<typename BindingParameterType>
 FRHIShaderParameterResource ExtractShaderParameterResource(FShaderParameterReader Reader, const BindingParameterType& Parameter)
 {
@@ -669,32 +636,32 @@ FRHIShaderParameterResource ExtractShaderParameterResource(FShaderParameterReade
 	{
 		FRHITexture* Texture = Reader.Read<FRHITexture*>(Parameter);
 		checkSlow(Texture);
-		return FRHIShaderParameterResourceEx(Texture, GetParameterIndex(Parameter));
+		return FRHIShaderParameterResource(Texture, GetParameterIndex(Parameter));
 	}
 	case UBMT_SRV:
 	{
 		FRHIShaderResourceView* ShaderResourceView = Reader.Read<FRHIShaderResourceView*>(Parameter);
 		checkSlow(ShaderResourceView);
-		return FRHIShaderParameterResourceEx(ShaderResourceView, GetParameterIndex(Parameter));
+		return FRHIShaderParameterResource(ShaderResourceView, GetParameterIndex(Parameter));
 	}
 	case UBMT_UAV:
 	{
 		FRHIUnorderedAccessView* UnorderedAccessView = Reader.Read<FRHIUnorderedAccessView*>(Parameter);
 		checkSlow(UnorderedAccessView);
-		return FRHIShaderParameterResourceEx(UnorderedAccessView, GetParameterIndex(Parameter));
+		return FRHIShaderParameterResource(UnorderedAccessView, GetParameterIndex(Parameter));
 	}
 	case UBMT_SAMPLER:
 	{
 		FRHISamplerState* SamplerState = Reader.Read<FRHISamplerState*>(Parameter);
 		checkSlow(SamplerState);
-		return FRHIShaderParameterResourceEx(SamplerState, GetParameterIndex(Parameter));
+		return FRHIShaderParameterResource(SamplerState, GetParameterIndex(Parameter));
 	}
 	case UBMT_RDG_TEXTURE:
 	{
 		FRDGTexture* RDGTexture = Reader.Read<FRDGTexture*>(Parameter);
 		checkSlow(RDGTexture);
 		RDGTexture->MarkResourceAsUsed();
-		return FRHIShaderParameterResourceEx(RDGTexture->GetRHI(), GetParameterIndex(Parameter));
+		return FRHIShaderParameterResource(RDGTexture->GetRHI(), GetParameterIndex(Parameter));
 	}
 	case UBMT_RDG_TEXTURE_SRV:
 	case UBMT_RDG_BUFFER_SRV:
@@ -702,7 +669,7 @@ FRHIShaderParameterResource ExtractShaderParameterResource(FShaderParameterReade
 		FRDGShaderResourceView* RDGShaderResourceView = Reader.Read<FRDGShaderResourceView*>(Parameter);
 		checkSlow(RDGShaderResourceView);
 		RDGShaderResourceView->MarkResourceAsUsed();
-		return FRHIShaderParameterResourceEx(RDGShaderResourceView->GetRHI(), GetParameterIndex(Parameter));
+		return FRHIShaderParameterResource(RDGShaderResourceView->GetRHI(), GetParameterIndex(Parameter));
 	}
 	case UBMT_RDG_TEXTURE_UAV:
 	case UBMT_RDG_BUFFER_UAV:
@@ -710,11 +677,11 @@ FRHIShaderParameterResource ExtractShaderParameterResource(FShaderParameterReade
 		FRDGUnorderedAccessView* RDGUnorderedAccessView = Reader.Read<FRDGUnorderedAccessView*>(Parameter);
 		checkSlow(RDGUnorderedAccessView);
 		RDGUnorderedAccessView->MarkResourceAsUsed();
-		return FRHIShaderParameterResourceEx(RDGUnorderedAccessView->GetRHI(), GetParameterIndex(Parameter));
+		return FRHIShaderParameterResource(RDGUnorderedAccessView->GetRHI(), GetParameterIndex(Parameter));
 	}
 	default:
 		checkf(false, TEXT("Unhandled resource type?"));
-		return FRHIShaderParameterResourceEx();
+		return FRHIShaderParameterResource();
 	}
 }
 
@@ -763,8 +730,7 @@ static void ExtractShaderParameterResources(
 			{
 				UniformBufferBinding->MarkResourceAsUsed();
 
-				FRHIShaderParameterResourceEx Binding(UniformBufferBinding->GetRHI(), GetParameterIndex(Parameter));
-				OutResourceParameters.Emplace(Binding);
+				OutResourceParameters.Emplace(UniformBufferBinding->GetRHI(), GetParameterIndex(Parameter));
 			}
 		}
 
@@ -773,8 +739,7 @@ static void ExtractShaderParameterResources(
 			const FUniformBufferBinding& UniformBufferBinding = Reader.Read<FUniformBufferBinding>(Parameter);
 			if (UniformBufferBinding.IsShader())
 			{
-				FRHIShaderParameterResourceEx Binding(UniformBufferBinding.GetUniformBuffer(), GetParameterIndex(Parameter));
-				OutResourceParameters.Emplace(Binding);
+				OutResourceParameters.Emplace(UniformBufferBinding.GetUniformBuffer(), GetParameterIndex(Parameter));
 			}
 		}
 #endif // UE_SINGLE_SHADER_PARAMETERS_COMMAND_INCLUDE_UB
diff --git a/Engine/Source/Runtime/RenderCore/Public/Shader.h b/Engine/Source/Runtime/RenderCore/Public/Shader.h
--- a/Engine/Source/Runtime/RenderCore/Public/Shader.h
+++ b/Engine/Source/Runtime/RenderCore/Public/Shader.h
@@ -771,6 +771,10 @@ public:
 
 }; // FShaderParameterBindings
 
+inline int16 GetParameterIndex(const FShaderParameterBindings::FResourceParameter& Parameter) { return Parameter.BaseIndex; }
+inline int16 GetParameterIndex(const FShaderParameterBindings::FBindlessResourceParameter& Parameter) { return Parameter.GlobalConstantOffset; }
+inline int16 GetParameterIndex(const FShaderParameterBindings::FParameterStructReference& Parameter) { return Parameter.BufferIndex; }
+
 // Flags that can specialize shader permutations compiled for specific platforms
 enum class EShaderPermutationFlags : uint32
 {
diff --git a/Engine/Source/Runtime/RenderCore/Public/ShaderParameterStruct.h b/Engine/Source/Runtime/RenderCore/Public/ShaderParameterStruct.h
--- a/Engine/Source/Runtime/RenderCore/Public/ShaderParameterStruct.h
+++ b/Engine/Source/Runtime/RenderCore/Public/ShaderParameterStruct.h
@@ -128,43 +128,69 @@ FORCEINLINE void ValidateShaderParameters(const TShaderRef<TShaderClass>& Shader
 	return ValidateShaderParameters(Shader, TShaderClass::FParameters::FTypeInfo::GetStructMetadata(), &Parameters);
 }
 
+template<typename TParameterType>
+inline void CollectUAVsToUnset(TArray<FRHIShaderParameterResource, TInlineAllocator<16>>& UAVsToUnset, const TMemoryImageArray<TParameterType>& InParameters)
+{
+	FRHIUnorderedAccessView* NullUAV = nullptr;
+	for (const TParameterType& Parameter : InParameters)
+	{
+		if (Parameter.BaseType == UBMT_UAV ||
+			Parameter.BaseType == UBMT_RDG_TEXTURE_UAV ||
+			Parameter.BaseType == UBMT_RDG_BUFFER_UAV)
+		{
+			UAVsToUnset.Emplace(NullUAV, GetParameterIndex(Parameter));
+		}
+	}
+}
+
+template<typename TParameterType>
+inline void CollectSRVsToUnset(TArray<FRHIShaderParameterResource, TInlineAllocator<16>>& SRVsToUnset, const TMemoryImageArray<TParameterType>& InParameters)
+{
+	FRHIShaderResourceView* NullSRV = nullptr;
+	for (const TParameterType& Parameter : InParameters)
+	{
+		if (Parameter.BaseType == UBMT_SRV ||
+			Parameter.BaseType == UBMT_RDG_TEXTURE_SRV ||
+			Parameter.BaseType == UBMT_RDG_BUFFER_SRV)
+		{
+			SRVsToUnset.Emplace(NullSRV, GetParameterIndex(Parameter));
+		}
+	}
+}
+
 /** Unset compute shader UAVs. */
 template<typename TRHICmdList, typename TShaderClass>
-inline void UnsetShaderUAVs(TRHICmdList& RHICmdList, const TShaderRef<TShaderClass>& Shader, FRHIComputeShader* ShadeRHI)
+inline void UnsetShaderUAVs(TRHICmdList& RHICmdList, const TShaderRef<TShaderClass>& Shader, FRHIComputeShader* ShaderRHI)
 {
 	// TODO(RDG): Once all shader sets their parameter through this, can refactor RHI so all UAVs of a shader get unset through a single RHI function call.
 	const FShaderParameterBindings& Bindings = Shader->Bindings;
 
 	checkf(Bindings.RootParameterBufferIndex == FShaderParameterBindings::kInvalidBufferIndex, TEXT("Can't use UnsetShaderUAVs() for root parameter buffer index."));
 
-	for (const FShaderParameterBindings::FResourceParameter& ParameterBinding : Bindings.ResourceParameters)
-	{
-		if (ParameterBinding.BaseType == UBMT_UAV ||
-			ParameterBinding.BaseType == UBMT_RDG_TEXTURE_UAV ||
-			ParameterBinding.BaseType == UBMT_RDG_BUFFER_UAV)
-		{
-			RHICmdList.SetUAVParameter(ShadeRHI, ParameterBinding.BaseIndex, nullptr);
-		}
-	}
+	TArray<FRHIShaderParameterResource, TInlineAllocator<16>> ResourceUAVsToUnset;
+	TArray<FRHIShaderParameterResource, TInlineAllocator<16>> BindlessUAVsToUnset;
+
+	CollectUAVsToUnset(ResourceUAVsToUnset, Bindings.ResourceParameters);
+	CollectUAVsToUnset(BindlessUAVsToUnset, Bindings.BindlessResourceParameters);
+
+	RHICmdList.SetShaderParameters(ShaderRHI, {}, {}, ResourceUAVsToUnset, BindlessUAVsToUnset);
 }
 
 /** Unset compute shader SRVs. */
 template<typename TRHICmdList, typename TShaderClass>
-inline void UnsetShaderSRVs(TRHICmdList& RHICmdList, const TShaderRef<TShaderClass>& Shader, FRHIComputeShader* ShadeRHI)
+inline void UnsetShaderSRVs(TRHICmdList& RHICmdList, const TShaderRef<TShaderClass>& Shader, FRHIComputeShader* ShaderRHI)
 {
 	const FShaderParameterBindings& Bindings = Shader->Bindings;
 
 	checkf(Bindings.RootParameterBufferIndex == FShaderParameterBindings::kInvalidBufferIndex, TEXT("Can't use UnsetShaderSRVs() for root parameter buffer index."));
 
-	for (const FShaderParameterBindings::FResourceParameter& ParameterBinding : Bindings.ResourceParameters)
-	{
-		if (ParameterBinding.BaseType == UBMT_SRV ||
-			ParameterBinding.BaseType == UBMT_RDG_TEXTURE_SRV ||
-			ParameterBinding.BaseType == UBMT_RDG_BUFFER_SRV)
-		{
-			RHICmdList.SetShaderResourceViewParameter(ShadeRHI, ParameterBinding.BaseIndex, nullptr);
-		}
-	}
+	TArray<FRHIShaderParameterResource, TInlineAllocator<16>> ResourceSRVsToUnset;
+	TArray<FRHIShaderParameterResource, TInlineAllocator<16>> BindlessSRVsToUnset;
+
+	CollectSRVsToUnset(ResourceSRVsToUnset, Bindings.ResourceParameters);
+	CollectSRVsToUnset(BindlessSRVsToUnset, Bindings.BindlessResourceParameters);
+
+	RHICmdList.SetShaderParameters(ShaderRHI, {}, {}, ResourceSRVsToUnset, BindlessSRVsToUnset);
 }
 
 RENDERCORE_API void SetShaderParameters(
diff --git a/Engine/Source/Runtime/Renderer/Private/MeshPassProcessor.cpp b/Engine/Source/Runtime/Renderer/Private/MeshPassProcessor.cpp
--- a/Engine/Source/Runtime/Renderer/Private/MeshPassProcessor.cpp
+++ b/Engine/Source/Runtime/Renderer/Private/MeshPassProcessor.cpp
@@ -111,133 +111,73 @@ private:
 	const uint8* Data;
 };
 
-template<typename TRHICmdList, typename TRHIShader>
-FORCEINLINE void SetBindlessParameter(TRHICmdList& RHICmdList, TRHIShader* Shader, FShaderResourceParameterInfo Parameter, FRHIDescriptorHandle Handle)
+struct FAllShaderParameters
 {
-	if (Handle.IsValid())
-	{
-		const uint32 BindlessIndex = Handle.GetIndex();
-		RHICmdList.SetShaderParameter(Shader, Parameter.BufferIndex, Parameter.BaseIndex, sizeof(BindlessIndex), &BindlessIndex);
-	}
-}
-
-template<class TRHICmdList, class TRHIShader>
-inline void SetTextureParameter(TRHICmdList& RHICmdList, TRHIShader* Shader, FShaderBindingState& RESTRICT ShaderBindingState, FShaderResourceParameterInfo Parameter, FRHITexture* TextureRHI)
-{
-#if PLATFORM_SUPPORTS_BINDLESS_RENDERING
-	if (Parameter.Type == EShaderParameterType::BindlessResourceIndex)
-	{
-		const FRHIDescriptorHandle Handle = TextureRHI ? TextureRHI->GetDefaultBindlessHandle() : FRHIDescriptorHandle();
-		SetBindlessParameter(RHICmdList, Shader, Parameter, Handle);
-	}
-	else
-#endif
-	{
-		checkSlow(Parameter.BaseIndex < UE_ARRAY_COUNT(ShaderBindingState.Textures));
-		if (TextureRHI != ShaderBindingState.Textures[Parameter.BaseIndex])
-		{
-			RHICmdList.SetShaderTexture(Shader, Parameter.BaseIndex, TextureRHI);
-
-			ShaderBindingState.Textures[Parameter.BaseIndex] = TextureRHI;
-			ShaderBindingState.MaxTextureUsed = FMath::Max((int32)Parameter.BaseIndex, ShaderBindingState.MaxTextureUsed);
-		}
-	}
-}
+	TArray<uint8> ParametersData;
+	TArray<FRHIShaderParameter> Parameters;
+	TArray<FRHIShaderParameterResource> ResourceParameters;
+	TArray<FRHIShaderParameterResource> BindlessParameters;
+};
 
-template<class TRHICmdList, class TRHIShader>
-inline void SetTextureParameter(TRHICmdList& RHICmdList, TRHIShader* Shader, FShaderResourceParameterInfo Parameter, FRHITexture* TextureRHI)
+inline void SetTextureParameter(FAllShaderParameters& AllParameters, const FShaderResourceParameterInfo& Parameter, FRHITexture* TextureRHI)
 {
 #if PLATFORM_SUPPORTS_BINDLESS_RENDERING
 	if (Parameter.Type == EShaderParameterType::BindlessResourceIndex)
 	{
-		const FRHIDescriptorHandle Handle = TextureRHI ? TextureRHI->GetDefaultBindlessHandle() : FRHIDescriptorHandle();
-		SetBindlessParameter(RHICmdList, Shader, Parameter, Handle);
+		check(Parameter.BufferIndex == 0);
+		AllParameters.BindlessParameters.Emplace(TextureRHI, Parameter.BaseIndex);
 	}
 	else
 #endif
 	{
-		RHICmdList.SetShaderTexture(Shader, Parameter.BaseIndex, TextureRHI);
+		AllParameters.ResourceParameters.Emplace(TextureRHI, Parameter.BaseIndex);
 	}
 }
 
-template<class TRHICmdList, class TRHIShader>
-inline void SetSrvParameter(TRHICmdList& RHICmdList, TRHIShader* Shader, FShaderBindingState& RESTRICT ShaderBindingState, FShaderResourceParameterInfo Parameter, FRHIShaderResourceView* SrvRHI)
+inline void SetSrvParameter(FAllShaderParameters& AllParameters, const FShaderResourceParameterInfo& Parameter, FRHIShaderResourceView* SrvRHI)
 {
 #if PLATFORM_SUPPORTS_BINDLESS_RENDERING
 	if (Parameter.Type == EShaderParameterType::BindlessResourceIndex)
 	{
-		const FRHIDescriptorHandle Handle = SrvRHI ? SrvRHI->GetBindlessHandle() : FRHIDescriptorHandle();
-		SetBindlessParameter(RHICmdList, Shader, Parameter, Handle);
+		check(Parameter.BufferIndex == 0);
+		AllParameters.BindlessParameters.Emplace(SrvRHI, Parameter.BaseIndex);
 	}
 	else
 #endif
 	{
-		checkSlow(Parameter.BaseIndex < UE_ARRAY_COUNT(ShaderBindingState.SRVs));
-		if (SrvRHI != ShaderBindingState.SRVs[Parameter.BaseIndex])
-		{
-			RHICmdList.SetShaderResourceViewParameter(Shader, Parameter.BaseIndex, SrvRHI);
-
-			ShaderBindingState.SRVs[Parameter.BaseIndex] = SrvRHI;
-			ShaderBindingState.MaxSRVUsed = FMath::Max((int32)Parameter.BaseIndex, ShaderBindingState.MaxSRVUsed);
-		}
+		AllParameters.ResourceParameters.Emplace(SrvRHI, Parameter.BaseIndex);
 	}
 }
 
-template<class TRHICmdList, class TRHIShader>
-inline void SetSrvParameter(TRHICmdList& RHICmdList, TRHIShader* Shader, FShaderResourceParameterInfo Parameter, FRHIShaderResourceView* SrvRHI)
+inline void SetSamplerParameter(FAllShaderParameters& AllParameters, const FShaderResourceParameterInfo& Parameter, FRHISamplerState* SamplerStateRHI)
 {
 #if PLATFORM_SUPPORTS_BINDLESS_RENDERING
-	if (Parameter.Type == EShaderParameterType::BindlessResourceIndex)
+	if (Parameter.Type == EShaderParameterType::BindlessSamplerIndex)
 	{
-		const FRHIDescriptorHandle Handle = SrvRHI ? SrvRHI->GetBindlessHandle() : FRHIDescriptorHandle();
-		SetBindlessParameter(RHICmdList, Shader, Parameter, Handle);
+		AllParameters.BindlessParameters.Emplace(SamplerStateRHI, Parameter.BaseIndex);
 	}
 	else
 #endif
 	{
-		RHICmdList.SetShaderResourceViewParameter(Shader, Parameter.BaseIndex, SrvRHI);
+		AllParameters.ResourceParameters.Emplace(SamplerStateRHI, Parameter.BaseIndex);
 	}
 }
 
-template<class TRHICmdList, class TRHIShader>
-inline void SetSamplerParameter(TRHICmdList& RHICmdList, TRHIShader* Shader, FShaderBindingState& RESTRICT ShaderBindingState, FShaderResourceParameterInfo Parameter, FRHISamplerState* SamplerStateRHI)
+inline void SetLooseParameters(FAllShaderParameters& AllParameters, const FShaderParameterMapInfo& ParameterMapInfo, const uint8* LooseDataStart)
 {
-#if PLATFORM_SUPPORTS_BINDLESS_RENDERING
-	if (Parameter.Type == EShaderParameterType::BindlessSamplerIndex)
+	for (const FShaderLooseParameterBufferInfo& LooseParameterBuffer : ParameterMapInfo.LooseParameterBuffers)
 	{
-		const FRHIDescriptorHandle Handle = SamplerStateRHI ? SamplerStateRHI->GetBindlessHandle() : FRHIDescriptorHandle();
-		SetBindlessParameter(RHICmdList, Shader, Parameter, Handle);
-	}
-	else
-#endif
-	{
-		checkSlow(Parameter.BaseIndex < UE_ARRAY_COUNT(ShaderBindingState.Samplers));
-		if (SamplerStateRHI != ShaderBindingState.Samplers[Parameter.BaseIndex])
+		for (const FShaderLooseParameterInfo& Parameter : LooseParameterBuffer.Parameters)
 		{
-			RHICmdList.SetShaderSampler(Shader, Parameter.BaseIndex, SamplerStateRHI);
+			const int32 DataOffset = AllParameters.ParametersData.Num();
+			AllParameters.ParametersData.Append((const uint8*)LooseDataStart, Parameter.Size);
+			AllParameters.Parameters.Emplace(LooseParameterBuffer.BaseIndex, Parameter.BaseIndex, DataOffset, Parameter.Size);
 
-			ShaderBindingState.Samplers[Parameter.BaseIndex] = SamplerStateRHI;
-			ShaderBindingState.MaxSamplerUsed = FMath::Max((int32)Parameter.BaseIndex, ShaderBindingState.MaxSamplerUsed);
+			LooseDataStart += Parameter.Size;
 		}
 	}
 }
 
-template<class TRHICmdList, class TRHIShader>
-inline void SetSamplerParameter(TRHICmdList& RHICmdList, TRHIShader* Shader, FShaderResourceParameterInfo Parameter, FRHISamplerState* SamplerStateRHI)
-{
-#if PLATFORM_SUPPORTS_BINDLESS_RENDERING
-	if (Parameter.Type == EShaderParameterType::BindlessSamplerIndex)
-	{
-		const FRHIDescriptorHandle Handle = SamplerStateRHI ? SamplerStateRHI->GetBindlessHandle() : FRHIDescriptorHandle();
-		SetBindlessParameter(RHICmdList, Shader, Parameter, Handle);
-	}
-	else
-#endif
-	{
-		RHICmdList.SetShaderSampler(Shader, Parameter.BaseIndex, SamplerStateRHI);
-	}
-}
-
 template<class RHICmdListType, class RHIShaderType>
 void FMeshDrawShaderBindings::SetShaderBindings(
 	RHICmdListType& RHICmdList,
@@ -249,6 +189,8 @@ void FMeshDrawShaderBindings::SetShaderBindings(
 	const FShaderUniformBufferParameterInfo* RESTRICT UniformBufferParameters = SingleShaderBindings.ParameterMapInfo.UniformBuffers.GetData();
 	const int32 NumUniformBuffers = SingleShaderBindings.ParameterMapInfo.UniformBuffers.Num();
 
+	FAllShaderParameters AllParameters;
+
 	for (int32 UniformBufferIndex = 0; UniformBufferIndex < NumUniformBuffers; UniformBufferIndex++)
 	{
 		FShaderUniformBufferParameterInfo Parameter = UniformBufferParameters[UniformBufferIndex];
@@ -272,7 +214,7 @@ void FMeshDrawShaderBindings::SetShaderBindings(
 		FShaderResourceParameterInfo Parameter = TextureSamplerParameters[SamplerIndex];
 		FRHISamplerState* Sampler = (FRHISamplerState*)SamplerBindings[SamplerIndex];
 
-		SetSamplerParameter(RHICmdList, Shader, ShaderBindingState, Parameter, Sampler);
+		SetSamplerParameter(AllParameters, Parameter, Sampler);
 	}
 
 	const uint8* RESTRICT SRVType = SingleShaderBindings.GetSRVTypeStart();
@@ -290,32 +232,24 @@ void FMeshDrawShaderBindings::SetShaderBindings(
 		if (SRVType[TypeByteIndex] & (1 << TypeBitIndex))
 		{
 			FRHIShaderResourceView* SRV = (FRHIShaderResourceView*)SRVBindings[SRVIndex];
-			SetSrvParameter(RHICmdList, Shader, ShaderBindingState, Parameter, SRV);
+			SetSrvParameter(AllParameters, Parameter, SRV);
 		}
 		else
 		{
 			FRHITexture* Texture = (FRHITexture*)SRVBindings[SRVIndex];
-			SetTextureParameter(RHICmdList, Shader, ShaderBindingState, Parameter, Texture);
+			SetTextureParameter(AllParameters, Parameter, Texture);
 		}
 	}
 
-	const uint8* LooseDataStart = SingleShaderBindings.GetLooseDataStart();
+	SetLooseParameters(AllParameters, SingleShaderBindings.ParameterMapInfo, SingleShaderBindings.GetLooseDataStart());
 
-	for (const FShaderLooseParameterBufferInfo& LooseParameterBuffer : SingleShaderBindings.ParameterMapInfo.LooseParameterBuffers)
-	{
-		for (FShaderLooseParameterInfo Parameter : LooseParameterBuffer.Parameters)
-		{
-			RHICmdList.SetShaderParameter(
-				Shader,
-				LooseParameterBuffer.BaseIndex,
-				Parameter.BaseIndex,
-				Parameter.Size,
-				LooseDataStart
-			);
-
-			LooseDataStart += Parameter.Size;
-		}
-	}
+	RHICmdList.SetShaderParameters(
+		Shader
+		, AllParameters.ParametersData
+		, AllParameters.Parameters
+		, AllParameters.ResourceParameters
+		, AllParameters.BindlessParameters
+	);
 }
 
 template<class RHICmdListType, class RHIShaderType>
@@ -328,6 +262,8 @@ void FMeshDrawShaderBindings::SetShaderBindings(
 	const FShaderUniformBufferParameterInfo* RESTRICT UniformBufferParameters = SingleShaderBindings.ParameterMapInfo.UniformBuffers.GetData();
 	const int32 NumUniformBuffers = SingleShaderBindings.ParameterMapInfo.UniformBuffers.Num();
 
+	FAllShaderParameters AllParameters;
+
 	for (int32 UniformBufferIndex = 0; UniformBufferIndex < NumUniformBuffers; UniformBufferIndex++)
 	{
 		FShaderUniformBufferParameterInfo Parameter = UniformBufferParameters[UniformBufferIndex];
@@ -350,7 +286,7 @@ void FMeshDrawShaderBindings::SetShaderBindings(
 
 		if (Sampler)
 		{
-			SetSamplerParameter(RHICmdList, Shader, Parameter, Sampler);
+			SetSamplerParameter(AllParameters, Parameter, Sampler);
 		}
 	}
 
@@ -369,32 +305,24 @@ void FMeshDrawShaderBindings::SetShaderBindings(
 		if (SRVType[TypeByteIndex] & (1 << TypeBitIndex))
 		{
 			FRHIShaderResourceView* SRV = (FRHIShaderResourceView*)SRVBindings[SRVIndex];
-			SetSrvParameter(RHICmdList, Shader, Parameter, SRV);
+			SetSrvParameter(AllParameters, Parameter, SRV);
 		}
 		else
 		{
 			FRHITexture* Texture = (FRHITexture*)SRVBindings[SRVIndex];
-			SetTextureParameter(RHICmdList, Shader, Parameter, Texture);
+			SetTextureParameter(AllParameters, Parameter, Texture);
 		}
 	}
-	
-	const uint8* LooseDataStart = SingleShaderBindings.GetLooseDataStart();
 
-	for (const FShaderLooseParameterBufferInfo& LooseParameterBuffer : SingleShaderBindings.ParameterMapInfo.LooseParameterBuffers)
-	{
-		for (FShaderLooseParameterInfo Parameter : LooseParameterBuffer.Parameters)
-		{
-			RHICmdList.SetShaderParameter(
-				Shader,
-				LooseParameterBuffer.BaseIndex,
-				Parameter.BaseIndex,
-				Parameter.Size,
-				LooseDataStart
-			);
+	SetLooseParameters(AllParameters, SingleShaderBindings.ParameterMapInfo, SingleShaderBindings.GetLooseDataStart());
 
-			LooseDataStart += Parameter.Size;
-		}
-	}
+	RHICmdList.SetShaderParameters(
+		Shader
+		, AllParameters.ParametersData
+		, AllParameters.Parameters
+		, AllParameters.ResourceParameters
+		, AllParameters.BindlessParameters
+	);
 }
 
 #if RHI_RAYTRACING
diff --git a/Engine/Source/Runtime/Renderer/Public/MeshPassProcessor.h b/Engine/Source/Runtime/Renderer/Public/MeshPassProcessor.h
--- a/Engine/Source/Runtime/Renderer/Public/MeshPassProcessor.h
+++ b/Engine/Source/Runtime/Renderer/Public/MeshPassProcessor.h
@@ -615,19 +615,11 @@ private:
 
 class FShaderBindingState
 {
-	enum { MAX_SRVS_PER_STAGE = 128 };
 	enum { MAX_UNIFORM_BUFFERS_PER_STAGE = 14 };
-	enum { MAX_SAMPLERS_PER_STAGE = 32 };
 
 public:
-	int32 MaxSRVUsed = -1;
-	FRHIShaderResourceView* SRVs[MAX_SRVS_PER_STAGE] = {};
 	int32 MaxUniformBufferUsed = -1;
 	FRHIUniformBuffer* UniformBuffers[MAX_UNIFORM_BUFFERS_PER_STAGE] = {};
-	int32 MaxTextureUsed = -1;
-	FRHITexture* Textures[MAX_SRVS_PER_STAGE] = {};
-	int32 MaxSamplerUsed = -1;
-	FRHISamplerState* Samplers[MAX_SAMPLERS_PER_STAGE] = {};
 };
 
 struct FMeshProcessorShaders
@@ -732,33 +724,12 @@ public:
 		{
 			FShaderBindingState& RESTRICT ShaderBinding = ShaderBindings[FrequencyIndex];
 
-			for (int32 SlotIndex = 0; SlotIndex <= ShaderBinding.MaxSRVUsed; SlotIndex++)
-			{
-				ShaderBinding.SRVs[SlotIndex] = nullptr;
-			}
-
-			ShaderBinding.MaxSRVUsed = -1;
-
 			for (int32 SlotIndex = 0; SlotIndex <= ShaderBinding.MaxUniformBufferUsed; SlotIndex++)
 			{
 				ShaderBinding.UniformBuffers[SlotIndex] = nullptr;
 			}
 
 			ShaderBinding.MaxUniformBufferUsed = -1;
-			
-			for (int32 SlotIndex = 0; SlotIndex <= ShaderBinding.MaxTextureUsed; SlotIndex++)
-			{
-				ShaderBinding.Textures[SlotIndex] = nullptr;
-			}
-
-			ShaderBinding.MaxTextureUsed = -1;
-
-			for (int32 SlotIndex = 0; SlotIndex <= ShaderBinding.MaxSamplerUsed; SlotIndex++)
-			{
-				ShaderBinding.Samplers[SlotIndex] = nullptr;
-			}
-
-			ShaderBinding.MaxSamplerUsed = -1;
 		}
 	}
 
