Temporary workaround for cached raytracing instances not being updated when meshes are streamed out.
- check if the cached instance is valid.
- request cached instance to be updated when necessary.

diff --git a/Engine/Source/Runtime/Renderer/Private/DeferredShadingRenderer.cpp b/Engine/Source/Runtime/Renderer/Private/DeferredShadingRenderer.cpp
--- a/Engine/Source/Runtime/Renderer/Private/DeferredShadingRenderer.cpp
+++ b/Engine/Source/Runtime/Renderer/Private/DeferredShadingRenderer.cpp
@@ -787,16 +787,17 @@ static void GatherRayTracingRelevantPrimitives(const FScene& Scene, const FViewI
 			// Marked visible and used after point, check if streaming then mark as used in the TLAS (so it can be streamed in)
 			if (EnumHasAnyFlags(Flags, ERayTracingPrimitiveFlags::Streaming))
 			{
-				// Is the cached data dirty?
-				if (SceneInfo->bCachedRaytracingDataDirty)
-				{
-					Result.DirtyCachedRayTracingPrimitives.Add(Scene.Primitives[PrimitiveIndex]);
-				}
-
 				check(SceneInfo->CoarseMeshStreamingHandle != INDEX_NONE);
 				Result.UsedCoarseMeshStreamingHandles.Add(SceneInfo->CoarseMeshStreamingHandle);
 			}
 
+			// Is the cached data dirty?
+			// eg: mesh was streamed in/out
+			if (SceneInfo->bCachedRaytracingDataDirty)
+			{
+				Result.DirtyCachedRayTracingPrimitives.Add(Scene.Primitives[PrimitiveIndex]);
+			}
+
 			FRayTracingRelevantPrimitive Item;
 			Item.PrimitiveIndex = PrimitiveIndex;
 
@@ -1272,17 +1273,20 @@ bool FDeferredShadingSceneRenderer::GatherRayTracingWorldInstancesForView(FRDGBu
 
 		FRayTracingScene& RayTracingScene; // New instances are added into FRayTracingScene::Instances and FRayTracingScene::Allocator is used for temporary data
 		TArray<FVisibleRayTracingMeshCommand>& VisibleRayTracingMeshCommands; // New elements are added here by this task
+		TArray<FPrimitiveSceneProxy*>& ProxiesWithDirtyCachedInstance;
 
 		FRayTracingSceneAddInstancesTask(const FScene& InScene,
 											TChunkedArray<FRayTracingRelevantPrimitive>& InRelevantStaticPrimitives,
 											const FRayTracingCullingParameters& InCullingParameters, bool InEnableInstanceDebugData,
-											FRayTracingScene& InRayTracingScene, TArray<FVisibleRayTracingMeshCommand>& InVisibleRayTracingMeshCommands)
+											FRayTracingScene& InRayTracingScene, TArray<FVisibleRayTracingMeshCommand>& InVisibleRayTracingMeshCommands,
+											TArray<FPrimitiveSceneProxy*>& InProxiesWithDirtyCachedInstance)
 			: Scene(InScene)
 			, RelevantStaticPrimitives(InRelevantStaticPrimitives)
 			, CullingParameters(InCullingParameters)
 			, bEnableInstanceDebugData(InEnableInstanceDebugData)
 			, RayTracingScene(InRayTracingScene)
 			, VisibleRayTracingMeshCommands(InVisibleRayTracingMeshCommands)
+			, ProxiesWithDirtyCachedInstance(InProxiesWithDirtyCachedInstance)
 		{
 			VisibleRayTracingMeshCommands.Reserve(RelevantStaticPrimitives.Num());
 		}
@@ -1375,6 +1379,20 @@ bool FDeferredShadingSceneRenderer::GatherRayTracingWorldInstancesForView(FRDGBu
 							Nanite::GRayTracingManager.AddVisiblePrimitive(SceneInfo);
 						}
 					}
+
+					// For primitives with ERayTracingPrimitiveFlags::CacheInstances flag we only cache the instance/mesh commands of the current LOD
+					// (see FPrimitiveSceneInfo::UpdateCachedRayTracingInstance(...) and CacheRayTracingPrimitive(...))
+					const int32 LODIndex = 0;
+
+					FRHIRayTracingGeometry* RayTracingGeometryInstance = SceneInfo->GetStaticRayTracingGeometryInstance(LODIndex);
+					if (!RayTracingGeometryInstance)
+					{
+						// cached instance is not valid (eg: was streamed out) need to invalidate for next frame
+						ProxiesWithDirtyCachedInstance.Add(Scene.PrimitiveSceneProxies[PrimitiveIndex]);
+						continue;
+					}
+
+					check(RayTracingGeometryInstance == SceneInfo->CachedRayTracingInstance.GeometryRHI);
 					
 					// TODO: support GRayTracingExcludeDecals, but not in the form of RayTracingMeshCommand.bDecal as that requires looping over all cached MDCs
 					// Instead, either make r.RayTracing.ExcludeDecals read only or request a recache of all ray tracing commands during which decals are excluded
@@ -1383,9 +1401,9 @@ bool FDeferredShadingSceneRenderer::GatherRayTracingWorldInstancesForView(FRDGBu
 
 					// At the moment we only support SM & ISMs on this path
 					check(EnumHasAnyFlags(Flags, ERayTracingPrimitiveFlags::CacheMeshCommands));
-					if (SceneInfo->CachedRayTracingMeshCommandIndicesPerLOD.Num() > 0 && SceneInfo->CachedRayTracingMeshCommandIndicesPerLOD[0].Num() > 0)
+					if (SceneInfo->CachedRayTracingMeshCommandIndicesPerLOD.Num() > 0 && SceneInfo->CachedRayTracingMeshCommandIndicesPerLOD[LODIndex].Num() > 0)
 					{
-						for (int32 CommandIndex : SceneInfo->CachedRayTracingMeshCommandIndicesPerLOD[0])
+						for (int32 CommandIndex : SceneInfo->CachedRayTracingMeshCommandIndicesPerLOD[LODIndex])
 						{
 							FVisibleRayTracingMeshCommand NewVisibleMeshCommand;
 
@@ -1549,7 +1567,7 @@ bool FDeferredShadingSceneRenderer::GatherRayTracingWorldInstancesForView(FRDGBu
 
 	FGraphEventRef AddInstancesTask = TGraphTask<FRayTracingSceneAddInstancesTask>::CreateTask(&AddInstancesTaskPrerequisites).ConstructAndDispatchWhenReady(
 		*Scene, RelevantPrimitiveList.StaticPrimitives, View.RayTracingCullingParameters, bEnableInstanceDebugData, // inputs 
-		RayTracingScene, View.VisibleRayTracingMeshCommands // outputs
+		RayTracingScene, View.VisibleRayTracingMeshCommands, View.ProxiesWithDirtyCachedInstance // outputs
 	);
 
 	// Scene init task can run only when all pre-init tasks are complete (including culling tasks that are spawned while adding instances)
@@ -1840,6 +1858,11 @@ bool FDeferredShadingSceneRenderer::DispatchRayTracingWorldUpdates(FRDGBuilder&
 
 	ReferenceView.RayTracingSceneInitTask = {};
 
+	for (FPrimitiveSceneProxy* SceneProxy : ReferenceView.ProxiesWithDirtyCachedInstance)
+	{
+		SceneProxy->GetScene().UpdateCachedRayTracingState(SceneProxy);
+	}
+
 	{
 		Nanite::GRayTracingManager.ProcessUpdateRequests(GraphBuilder, Scene->GPUScene.PrimitiveBuffer->GetSRV());
 		const bool bAnyBlasRebuilt = Nanite::GRayTracingManager.ProcessBuildRequests(GraphBuilder);
diff --git a/Engine/Source/Runtime/Renderer/Private/PrimitiveSceneInfo.cpp b/Engine/Source/Runtime/Renderer/Private/PrimitiveSceneInfo.cpp
--- a/Engine/Source/Runtime/Renderer/Private/PrimitiveSceneInfo.cpp
+++ b/Engine/Source/Runtime/Renderer/Private/PrimitiveSceneInfo.cpp
@@ -232,6 +232,7 @@ FPrimitiveSceneInfo::FPrimitiveSceneInfo(UPrimitiveComponent* InComponent,FScene
 
 #if RHI_RAYTRACING
 	RayTracingGeometries = InComponent->SceneProxy->MoveRayTracingGeometries();
+	CachedRayTracingGeometry = nullptr;
 #endif
 }
 
@@ -266,6 +267,11 @@ FRHIRayTracingGeometry* FPrimitiveSceneInfo::GetStaticRayTracingGeometryInstance
 	}
 	else
 	{
+		if (LodLevel == 0 && CachedRayTracingGeometry && CachedRayTracingGeometry->IsValid())
+		{
+			return CachedRayTracingGeometry->RayTracingGeometryRHI;
+		}
+
 		return nullptr;
 	}
 }
@@ -992,6 +998,8 @@ void FPrimitiveSceneInfo::UpdateCachedRayTracingInstance(FPrimitiveSceneInfo* Sc
 
 		SceneInfo->CachedRayTracingInstance.GeometryRHI = CachedRayTracingInstance.Geometry->RayTracingGeometryRHI;
 
+		SceneInfo->CachedRayTracingGeometry = CachedRayTracingInstance.Geometry;
+
 		if (Nanite::GetRayTracingMode() != Nanite::ERayTracingMode::Fallback && SceneProxy->IsNaniteMesh())
 		{
 			FRHIRayTracingGeometry* NaniteRayTracingGeometry = Nanite::GRayTracingManager.GetRayTracingGeometry(SceneInfo);
diff --git a/Engine/Source/Runtime/Renderer/Private/SceneRendering.h b/Engine/Source/Runtime/Renderer/Private/SceneRendering.h
--- a/Engine/Source/Runtime/Renderer/Private/SceneRendering.h
+++ b/Engine/Source/Runtime/Renderer/Private/SceneRendering.h
@@ -1310,6 +1310,8 @@ public:
 	FRayTracingSceneWithGeometryInstances RayTracingSceneInitData;
 	FGraphEventRef RayTracingSceneInitTask; // Task to asynchronously create RayTracingSceneInitData
 
+	TArray<FPrimitiveSceneProxy*> ProxiesWithDirtyCachedInstance;
+
 	FGraphEventArray AddRayTracingMeshBatchTaskList;
 	TArray<FRayTracingMeshCommandOneFrameArray*, SceneRenderingAllocator> VisibleRayTracingMeshCommandsPerTask;
 	TArray<FDynamicRayTracingMeshCommandStorage*, SceneRenderingAllocator> DynamicRayTracingMeshCommandStoragePerTask;
diff --git a/Engine/Source/Runtime/Renderer/Public/PrimitiveSceneInfo.h b/Engine/Source/Runtime/Renderer/Public/PrimitiveSceneInfo.h
--- a/Engine/Source/Runtime/Renderer/Public/PrimitiveSceneInfo.h
+++ b/Engine/Source/Runtime/Renderer/Public/PrimitiveSceneInfo.h
@@ -695,6 +695,10 @@ private:
 #if RHI_RAYTRACING
 	TArray<FRayTracingGeometry*> RayTracingGeometries;
 
+	// Cache pointer to FRayTracingGeometry used by cached ray tracing instance
+	// since primitives using ERayTracingPrimitiveFlags::CacheInstances don't fill the RayTracingGeometries array above
+	const FRayTracingGeometry* CachedRayTracingGeometry;
+
 	/** Creates cached ray tracing representations for all meshes. */
 	static void CacheRayTracingPrimitives(FScene* Scene, const TArrayView<FPrimitiveSceneInfo*>& SceneInfos);
 
