Unshelved from pending changelist '23216677':
Fixed crash due to bCachedRaytracingDataDirty not being initialized.
- the flag used to only affect a code path used by nanite coarse meshes in which it didn't cause a crash.
- however now it's used more generally resulting in a crash.
Note: Swarm is down, committing on behalf of Tiago Costa

diff --git a/Engine/Source/Runtime/Renderer/Private/PrimitiveSceneInfo.cpp b/Engine/Source/Runtime/Renderer/Private/PrimitiveSceneInfo.cpp
--- a/Engine/Source/Runtime/Renderer/Private/PrimitiveSceneInfo.cpp
+++ b/Engine/Source/Runtime/Renderer/Private/PrimitiveSceneInfo.cpp
@@ -180,6 +180,7 @@ FPrimitiveSceneInfo::FPrimitiveSceneInfo(UPrimitiveComponent* InComponent,FScene
 	bIsRayTracingRelevant(InComponent->SceneProxy->IsRayTracingRelevant()),
 	bIsRayTracingStaticRelevant(InComponent->SceneProxy->IsRayTracingStaticRelevant()),
 	bIsVisibleInRayTracing(InComponent->SceneProxy->IsVisibleInRayTracing()),
+	bCachedRaytracingDataDirty(false),
 	CoarseMeshStreamingHandle(INDEX_NONE),
 #endif
 	PackedIndex(INDEX_NONE),
