Fix compilation error/warning when Strata is enabled on dx11.

diff --git a/Engine/Source/Runtime/Renderer/Private/Strata/Strata.cpp b/Engine/Source/Runtime/Renderer/Private/Strata/Strata.cpp
--- a/Engine/Source/Runtime/Renderer/Private/Strata/Strata.cpp
+++ b/Engine/Source/Runtime/Renderer/Private/Strata/Strata.cpp
@@ -787,6 +787,9 @@ class FStrataDBufferPassCS : public FGlobalShader
 		OutEnvironment.SetDefine(TEXT("STENCIL_STRATA_RECEIVE_DBUFFER_NORMAL_BIT_ID"), STENCIL_STRATA_RECEIVE_DBUFFER_NORMAL_BIT_ID);
 		OutEnvironment.SetDefine(TEXT("STENCIL_STRATA_RECEIVE_DBUFFER_DIFFUSE_BIT_ID"), STENCIL_STRATA_RECEIVE_DBUFFER_DIFFUSE_BIT_ID);
 		OutEnvironment.SetDefine(TEXT("STENCIL_STRATA_RECEIVE_DBUFFER_ROUGHNESS_BIT_ID"), STENCIL_STRATA_RECEIVE_DBUFFER_ROUGHNESS_BIT_ID);
+
+		// Needed as top layer texture can be a uint2
+		OutEnvironment.CompilerFlags.Add(CFLAG_AllowTypedUAVLoads);
 	}
 };
 IMPLEMENT_GLOBAL_SHADER(FStrataDBufferPassCS, "/Engine/Private/Strata/StrataDBuffer.usf", "MainCS", SF_Compute);
