Strata - added a check for a case that is not handled well.

diff --git a/Engine/Source/Runtime/Renderer/Private/Strata/Strata.cpp b/Engine/Source/Runtime/Renderer/Private/Strata/Strata.cpp
--- a/Engine/Source/Runtime/Renderer/Private/Strata/Strata.cpp
+++ b/Engine/Source/Runtime/Renderer/Private/Strata/Strata.cpp
@@ -376,8 +376,16 @@ static void RecordStrataAnalytics()
 
 static EPixelFormat GetTopLayerTextureFormat()
 {
-	static const auto CVarStrataGBufferFormat = IConsoleManager::Get().FindTConsoleVariableDataInt(TEXT("r.GBufferFormat"));
-	return CVarStrataGBufferFormat && CVarStrataGBufferFormat->GetValueOnAnyThread() > 1 ? PF_R32G32_UINT : PF_R32_UINT;
+	static const auto CVarStrataGBufferFormat = IConsoleManager::Get().FindTConsoleVariableDataInt(TEXT("r.GBufferFormat")); 
+
+	const bool bStrataHighQualityNormal = CVarStrataGBufferFormat && CVarStrataGBufferFormat->GetValueOnAnyThread() > 1;
+
+	// High quality normal is not supported on platforms that do not support R32G32 UAV load.
+	// This is dues to the way Strata account for decals. See FStrataDBufferPassCS, updating TopLayerTexture this way.
+	// If you encounter this check, you must disable high quality normal for Strata (material shaders must be recompiled to account for that).
+	check(!bStrataHighQualityNormal || (bStrataHighQualityNormal && UE::PixelFormat::HasCapabilities(PF_R32G32_UINT, EPixelFormatCapabilities::TypedUAVLoad)));
+
+	return bStrataHighQualityNormal ? PF_R32G32_UINT : PF_R32_UINT;
 }
 
 void InitialiseStrataFrameSceneData(FRDGBuilder& GraphBuilder, FSceneRenderer& SceneRenderer)
