Fix use-after-free bug on CSM RT DebugName. Rendertarget names are not copied, and so must be statically allocated to survive the function scope.

diff --git a/Engine/Source/Runtime/Renderer/Private/ShadowSetup.cpp b/Engine/Source/Runtime/Renderer/Private/ShadowSetup.cpp
--- a/Engine/Source/Runtime/Renderer/Private/ShadowSetup.cpp
+++ b/Engine/Source/Runtime/Renderer/Private/ShadowSetup.cpp
@@ -5369,6 +5369,38 @@ void FSceneRenderer::AllocateAtlasedShadowDepthTargets(
 	}
 }
 
+/**
+* Helper class to get the name of an indexed rendertarget, keeping the pointers around (this is required by the rendertarget pool)
+*/
+class RenderTargetNameSet
+{
+private:
+	FString Prefix;
+	TArray<FString> Names;
+
+public:
+	RenderTargetNameSet(const TCHAR* Prefix)
+		: Prefix(Prefix), Names()
+	{ }
+
+	const TCHAR* Get(const int32 Index)
+	{
+		const int32 Count = Index + 1;
+		while (Names.Num() < Count)
+		{
+			if (Index == 0)
+			{
+				Names.Add(Prefix);
+			}
+			else
+			{
+				Names.Emplace(FString::Printf(TEXT("%s%d"), *Prefix, Names.Num()));
+			}
+		}
+		return *Names[Index];
+	}
+};
+
 void FSceneRenderer::AllocateCachedShadowDepthTargets(FRHICommandListImmediate& RHICmdList, TArray<FProjectedShadowInfo*, SceneRenderingAllocator>& CachedShadows)
 {
 	for (int32 ShadowIndex = 0; ShadowIndex < CachedShadows.Num(); ShadowIndex++)
@@ -5382,7 +5414,9 @@ void FSceneRenderer::AllocateCachedShadowDepthTargets(FRHICommandListImmediate&
 
 		FIntPoint ShadowResolution(ProjectedShadowInfo->ResolutionX + ProjectedShadowInfo->BorderSize * 2, ProjectedShadowInfo->ResolutionY + ProjectedShadowInfo->BorderSize * 2);
 		FPooledRenderTargetDesc ShadowMapDesc2D = FPooledRenderTargetDesc::Create2DDesc(ShadowResolution, PF_ShadowDepth, FClearValueBinding::DepthOne, TexCreate_None, TexCreate_DepthStencilTargetable | TexCreate_ShaderResource, false, 1, false);
-		GRenderTargetPool.FindFreeElement(RHICmdList, ShadowMapDesc2D, ShadowMap.RenderTargets.DepthTarget, bIsWholeSceneDirectionalShadow ? *FString::Printf(TEXT("CachedShadowDepthMap_Split%d"), ProjectedShadowInfo->CascadeSettings.ShadowSplitIndex) : TEXT("CachedShadowDepthMap"));
+		static RenderTargetNameSet CachedSplitDepthRTNames(TEXT("CachedShadowDepthMap_Split"));
+		const TCHAR* RTDebugName = bIsWholeSceneDirectionalShadow ? CachedSplitDepthRTNames.Get(ProjectedShadowInfo->CascadeSettings.ShadowSplitIndex) : TEXT("CachedShadowDepthMap");
+		GRenderTargetPool.FindFreeElement(RHICmdList, ShadowMapDesc2D, ShadowMap.RenderTargets.DepthTarget, RTDebugName);
 
 		check(ProjectedShadowInfo->CacheMode == SDCM_StaticPrimitivesOnly);
 		FCachedShadowMapData& CachedShadowMapData = Scene->GetCachedShadowMapDataRef(ProjectedShadowInfo->GetLightSceneInfo().Id, FMath::Max(ProjectedShadowInfo->CascadeSettings.ShadowSplitIndex, 0));
@@ -5397,28 +5431,6 @@ void FSceneRenderer::AllocateCachedShadowDepthTargets(FRHICommandListImmediate&
 	}
 }
 
-/**
-* Helper function to get the name of a CSM rendertarget, keeping the pointers around (this is required by the rendertarget pool)
-* @param ShadowMapIndex - the index of the shadow map cascade
-*/
-const TCHAR* GetCSMRenderTargetName(int32 ShadowMapIndex)
-{
-	// Render target names require string pointers not to be released, so we cache them in a static array and grow as necessary
-	static TArray<FString*> ShadowmapNames;
-	while (ShadowmapNames.Num() < ShadowMapIndex + 1)
-	{
-		if (ShadowMapIndex == 0)
-		{
-			ShadowmapNames.Add(new FString(TEXT("WholeSceneShadowmap")));
-		}
-		else
-		{
-			ShadowmapNames.Add(new FString(FString::Printf(TEXT("WholeSceneShadowmap%d"), ShadowmapNames.Num())));
-		}
-	}
-	return **ShadowmapNames[ShadowMapIndex];
-}
-
 void FSceneRenderer::AllocateCSMDepthTargets(
 	FRHICommandListImmediate& RHICmdList,
 	const TArray<FProjectedShadowInfo*,	SceneRenderingAllocator>& WholeSceneDirectionalShadows,
@@ -5467,7 +5479,8 @@ void FSceneRenderer::AllocateCSMDepthTargets(
 			FIntPoint WholeSceneAtlasSize(CurrentLayout.TextureLayout.GetSizeX(), CurrentLayout.TextureLayout.GetSizeY());
 			FPooledRenderTargetDesc WholeSceneShadowMapDesc2D(FPooledRenderTargetDesc::Create2DDesc(WholeSceneAtlasSize, PF_ShadowDepth, FClearValueBinding::DepthOne, TexCreate_None, TexCreate_DepthStencilTargetable | TexCreate_ShaderResource, false));
 			WholeSceneShadowMapDesc2D.Flags |= GFastVRamConfig.ShadowCSM;
-			GRenderTargetPool.FindFreeElement(RHICmdList, WholeSceneShadowMapDesc2D, ShadowMapAtlas.RenderTargets.DepthTarget, GetCSMRenderTargetName(LayoutIndex));
+			static RenderTargetNameSet CsmRtNames(TEXT("WholeSceneShadowmap"));
+			GRenderTargetPool.FindFreeElement(RHICmdList, WholeSceneShadowMapDesc2D, ShadowMapAtlas.RenderTargets.DepthTarget, CsmRtNames.Get(LayoutIndex));
 
 			for (FProjectedShadowInfo* ProjectedShadowInfo : CurrentLayout.Shadows)
 			{
