Add depth bound test when we only apply AO in the fullscreen pass DiffuseIndirectAndAO. Use the same DepthFar value computed in PostProcessAmbientOcclusion.

diff --git a/Engine/Source/Runtime/Renderer/Private/IndirectLightRendering.cpp b/Engine/Source/Runtime/Renderer/Private/IndirectLightRendering.cpp
--- a/Engine/Source/Runtime/Renderer/Private/IndirectLightRendering.cpp
+++ b/Engine/Source/Runtime/Renderer/Private/IndirectLightRendering.cpp
@@ -81,6 +81,11 @@ static TAutoConsoleVariable<int32> CVarProbeSamplePerPixel(
 	TEXT("Number of sample to do per full res pixel."),
 	ECVF_RenderThreadSafe);
 
+static TAutoConsoleVariable<bool> CVarDiffuseIndirectOffUseDepthBoundsAO(
+	TEXT("r.DiffuseIndirectOffUseDepthBoundsAO"), true,
+	TEXT("Use depth bounds when we apply the AO when DiffuseIndirect is disabled."),
+	ECVF_RenderThreadSafe);
+
 
 DECLARE_GPU_STAT_NAMED(ReflectionEnvironment, TEXT("Reflection Environment"));
 DECLARE_GPU_STAT_NAMED(RayTracingReflections, TEXT("Ray Tracing Reflections"));
@@ -1267,19 +1272,55 @@ void FDeferredShadingSceneRenderer::RenderDiffuseIndirectAndAmbientOcclusion(
 			{
 				ClearUnusedGraphResources(PixelShader, PassParameters);
 
-				FPixelShaderUtils::AddFullscreenPass(
-					GraphBuilder,
-					View.ShaderMap,
+				// only use depth bound optimization when diffuse indirect is disable
+				bool bUseDepthBounds = CVarDiffuseIndirectOffUseDepthBoundsAO.GetValueOnRenderThread() &&
+										bool(ERHIZBuffer::IsInverted) && // Inverted depth buffer is assumed when setting depth bounds test for AO.
+										!PermutationVector.Get<FDiffuseIndirectCompositePS::FUpscaleDiffuseIndirectDim>() && AmbientOcclusionMask;
+
+				if (bUseDepthBounds)
+				{
+					PassParameters->RenderTargets.DepthStencil = FDepthStencilBinding(SceneTextures.Depth.Target, ERenderTargetLoadAction::ELoad, FExclusiveDepthStencil::DepthRead_StencilNop);
+				}
+
+				GraphBuilder.AddPass(
 					RDG_EVENT_NAME(
 						"DiffuseIndirectComposite(DiffuseIndirect=%s%s%s) %dx%d",
 						DiffuseIndirectSampling,
 						PermutationVector.Get<FDiffuseIndirectCompositePS::FUpscaleDiffuseIndirectDim>() ? TEXT(" UpscaleDiffuseIndirect") : TEXT(""),
 						AmbientOcclusionMask ? TEXT(" ApplyAOToSceneColor") : TEXT(""),
-						View.ViewRect.Width(), View.ViewRect.Height()),
-					PixelShader,
+						View.ViewRect.Width(), View.ViewRect.Height()), 
 					PassParameters,
-					View.ViewRect,
-					BlendState);
+					ERDGPassFlags::Raster,
+					[&View, PassParameters,PixelShader, BlendState, bUseDepthBounds](FRHICommandList& RHICmdList)
+					{
+						check(PixelShader.IsValid());
+						RHICmdList.SetViewport((float)View.ViewRect.Min.X, (float)View.ViewRect.Min.Y, 0.0f, (float)View.ViewRect.Max.X, (float)View.ViewRect.Max.Y, 1.0f);
+
+						FGraphicsPipelineStateInitializer GraphicsPSOInit;
+						FPixelShaderUtils::InitFullscreenPipelineState(RHICmdList, View.ShaderMap, PixelShader, /* out */GraphicsPSOInit);
+						GraphicsPSOInit.BlendState = BlendState;
+						GraphicsPSOInit.bDepthBounds = false;
+
+						if (bUseDepthBounds)
+						{
+							GraphicsPSOInit.bDepthBounds = true;
+							const FFinalPostProcessSettings& Settings = View.FinalPostProcessSettings;
+							const FMatrix& ProjectionMatrix = View.ViewMatrices.GetProjectionMatrix();
+							const FVector4f Far = (FVector4f)ProjectionMatrix.TransformFVector4(FVector4(0, 0, Settings.AmbientOcclusionFadeDistance));
+							float DepthFar = FMath::Min(1.0f, Far.Z / Far.W);
+							RHICmdList.SetDepthBounds(DepthFar, 1.0f);
+						}
+
+						SetGraphicsPipelineState(RHICmdList, GraphicsPSOInit, 0);
+						SetShaderParameters(RHICmdList, PixelShader, PixelShader.GetPixelShader(), *PassParameters);
+
+						FPixelShaderUtils::DrawFullscreenTriangle(RHICmdList);
+
+						if (bUseDepthBounds)
+						{
+							RHICmdList.SetDepthBounds(0.f, 1.0f);
+						}
+					});
 			}
 			else
 			{
