Fix for unity compile error where a local variable in the water rendering code hides a global one.

#rnx

diff --git a/Engine/Source/Runtime/Renderer/Private/SingleLayerWaterRendering.cpp b/Engine/Source/Runtime/Renderer/Private/SingleLayerWaterRendering.cpp
--- a/Engine/Source/Runtime/Renderer/Private/SingleLayerWaterRendering.cpp
+++ b/Engine/Source/Runtime/Renderer/Private/SingleLayerWaterRendering.cpp
@@ -189,8 +189,8 @@ bool ShouldRenderScreenSpaceReflectionsWater(const FViewInfo& View)
 		return false;
 	}
 
-	static const auto CVarSSRQuality = IConsoleManager::Get().FindTConsoleVariableDataInt(TEXT("r.SSR.Quality"));
-	int SSRQuality = CVarSSRQuality ? CVarSSRQuality->GetValueOnRenderThread() : 0;
+	static const auto SSRQualityCVar = IConsoleManager::Get().FindTConsoleVariableDataInt(TEXT("r.SSR.Quality"));
+	int SSRQuality = SSRQualityCVar ? SSRQualityCVar->GetValueOnRenderThread() : 0;
 	if (SSRQuality <= 0 
 		|| View.FinalPostProcessSettings.ScreenSpaceReflectionIntensity < 1.0f 
 		|| IsForwardShadingEnabled(View.GetShaderPlatform()))
