Fix bug in GTAO when there's a camera cut.

diff --git a/Engine/Source/Runtime/Renderer/Private/CompositionLighting/CompositionLighting.cpp b/Engine/Source/Runtime/Renderer/Private/CompositionLighting/CompositionLighting.cpp
--- a/Engine/Source/Runtime/Renderer/Private/CompositionLighting/CompositionLighting.cpp
+++ b/Engine/Source/Runtime/Renderer/Private/CompositionLighting/CompositionLighting.cpp
@@ -278,6 +278,7 @@ static FScreenPassTexture AddPostProcessingGTAOAllPasses(
 			else
 			{
 				HistoryColor = FScreenPassTexture(GraphBuilder.RegisterExternalTexture(GSystemTextures.WhiteDummy, TEXT("GTAODummyTexture")));
+				HistoryViewport = FScreenPassTextureViewport(HistoryColor);
 			}
 
 			FGTAOTemporalOutputs TemporalOutputs =
@@ -377,6 +378,7 @@ static FScreenPassTexture AddPostProcessingGTAOPostAsync(
 			else
 			{
 				HistoryColor = FScreenPassTexture(GraphBuilder.RegisterExternalTexture(GSystemTextures.WhiteDummy, TEXT("GTAODummyTexture")));
+				HistoryViewport = FScreenPassTextureViewport(HistoryColor);
 			}
 
 			FGTAOTemporalOutputs TemporalOutputs =
