Fix for NaN in CreateDefaultFogPassParameters(). Fixed by properly initializing a function argument.

#rnx

robomerge[FN_EngineMerge] -UE5-RES

diff --git a/Engine/Source/Runtime/Renderer/Private/FogRendering.cpp b/Engine/Source/Runtime/Renderer/Private/FogRendering.cpp
--- a/Engine/Source/Runtime/Renderer/Private/FogRendering.cpp
+++ b/Engine/Source/Runtime/Renderer/Private/FogRendering.cpp
@@ -255,7 +255,7 @@ static FFogPassParameters* CreateDefaultFogPassParameters(
 	TRDGUniformBufferRef<FSceneTextureUniformParameters> SceneTexturesUniformbuffer,
 	TRDGUniformBufferRef<FFogUniformParameters>& FogUniformBuffer,
 	FRDGTextureRef LightShaftOcclusionTexture,
-	const FScreenPassTextureViewportParameters LightShaftParameters)
+	const FScreenPassTextureViewportParameters& LightShaftParameters)
 {
 	extern int32 GVolumetricFogGridPixelSize;
 
@@ -430,7 +430,7 @@ void FDeferredShadingSceneRenderer::RenderUnderWaterFog(
 
 				// TODO add support for occlusion texture on water
 				FRDGTextureRef LightShaftOcclusionTexture = nullptr;	
-				FScreenPassTextureViewportParameters LightShaftParameters = FScreenPassTextureViewportParameters();
+				FScreenPassTextureViewportParameters LightShaftParameters = GetScreenPassTextureViewportParameters(FScreenPassTextureViewport(FIntRect(0, 0, 1, 1)));
 
 				FFogPassParameters* PassParameters = CreateDefaultFogPassParameters(GraphBuilder, View, SceneTexturesWithDepth, FogUniformBuffer, LightShaftOcclusionTexture, LightShaftParameters);
 				PassParameters->PS.WaterDepthTexture = WaterDepthTexture;
