Fix issues where Nanite primitives' custom data would not update on the GPU if it had no draw visibility, but was otherwise visible in shadow.

diff --git a/Engine/Source/Runtime/Engine/Private/PrimitiveSceneProxy.cpp b/Engine/Source/Runtime/Engine/Private/PrimitiveSceneProxy.cpp
--- a/Engine/Source/Runtime/Engine/Private/PrimitiveSceneProxy.cpp
+++ b/Engine/Source/Runtime/Engine/Private/PrimitiveSceneProxy.cpp
@@ -1045,6 +1045,7 @@ void FPrimitiveSceneProxy::SetLightingChannels_GameThread(FLightingChannels Ligh
 	{
 		PrimitiveSceneProxy->LightingChannelMask = LocalLightingChannelMask;
 		PrimitiveSceneProxy->GetPrimitiveSceneInfo()->SetNeedsUniformBufferUpdate(true);
+		PrimitiveSceneProxy->GetScene().RequestGPUSceneUpdate(*PrimitiveSceneProxy->GetPrimitiveSceneInfo(), EPrimitiveDirtyState::ChangedOther);
 	});
 }
 
@@ -1172,6 +1173,7 @@ void FPrimitiveSceneProxy::SetEvaluateWorldPositionOffset_GameThread(bool bEvalu
 		if (PrimitiveSceneProxy->bEvaluateWorldPositionOffset != bWPOEvaluate)
 		{
 			PrimitiveSceneProxy->bEvaluateWorldPositionOffset = bWPOEvaluate;
+			PrimitiveSceneProxy->GetPrimitiveSceneInfo()->SetNeedsUniformBufferUpdate(true);
 			PrimitiveSceneProxy->GetScene().RequestGPUSceneUpdate(*PrimitiveSceneProxy->GetPrimitiveSceneInfo(), EPrimitiveDirtyState::ChangedOther);
 			PrimitiveSceneProxy->OnEvaluateWorldPositionOffsetChanged_RenderThread();
 		}
diff --git a/Engine/Source/Runtime/Renderer/Private/RendererScene.cpp b/Engine/Source/Runtime/Renderer/Private/RendererScene.cpp
--- a/Engine/Source/Runtime/Renderer/Private/RendererScene.cpp
+++ b/Engine/Source/Runtime/Renderer/Private/RendererScene.cpp
@@ -3686,8 +3686,8 @@ void FSceneVelocityData::StartFrame(FScene* Scene)
 
 		if ((InternalFrameIndex - VelocityData.LastFrameUpdated == 1) && VelocityData.PrimitiveSceneInfo)
 		{
-			// Recreate PrimitiveUniformBuffer on the frame after the primitive moved, since it contains PreviousLocalToWorld
-			VelocityData.PrimitiveSceneInfo->SetNeedsUniformBufferUpdate(true);
+			// Force an update of the primitive data on the frame after the primitive moved, since it contains PreviousLocalToWorld
+			VelocityData.PrimitiveSceneInfo->MarkGPUStateDirty(EPrimitiveDirtyState::ChangedTransform);
 		}
 
 		if (bTrimOld && (InternalFrameIndex - VelocityData.LastFrameUsed) > 10)
@@ -5973,8 +5973,8 @@ void FScene::UpdateAllPrimitiveSceneInfos(FRDGBuilder& GraphBuilder, bool bAsync
 		FScopeCycleCounter Context(PrimitiveSceneProxy->GetStatId());
 		PrimitiveSceneProxy->CustomPrimitiveData = CustomParams.Value;
 
-		// Make sure the uniform buffer is updated before rendering
-		PrimitiveSceneProxy->GetPrimitiveSceneInfo()->SetNeedsUniformBufferUpdate(true);
+		// Ensure an update of primitive data before rendering
+		PrimitiveSceneProxy->GetPrimitiveSceneInfo()->MarkGPUStateDirty(EPrimitiveDirtyState::ChangedOther);
 	}
 
 	for (FPrimitiveSceneInfo* PrimitiveSceneInfo : DistanceFieldSceneDataUpdates)
diff --git a/Engine/Source/Runtime/Renderer/Private/SceneVisibility.cpp b/Engine/Source/Runtime/Renderer/Private/SceneVisibility.cpp
--- a/Engine/Source/Runtime/Renderer/Private/SceneVisibility.cpp
+++ b/Engine/Source/Runtime/Renderer/Private/SceneVisibility.cpp
@@ -2957,11 +2957,8 @@ struct FRelevancePacket : public FSceneRenderingAllocatorObject<FRelevancePacket
 		{
 			FPrimitiveSceneInfo* PrimitiveSceneInfo = RecachedReflectionCapturePrimitives.Prims[Index];
 
-			PrimitiveSceneInfo->SetNeedsUniformBufferUpdate(true);
+			PrimitiveSceneInfo->MarkGPUStateDirty(EPrimitiveDirtyState::ChangedAll);
 			PrimitiveSceneInfo->ConditionalUpdateUniformBuffer(RHICmdList);
-
-			FScene& WriteScene = *const_cast<FScene*>(Scene);
-			WriteScene.GPUScene.AddPrimitiveToUpdate(PrimitiveSceneInfo->GetIndex(), EPrimitiveDirtyState::ChangedAll);
 		}
 
 		for (int32 Index = 0; Index < LazyUpdatePrimitives.NumPrims; Index++)
diff --git a/Engine/Source/Runtime/Renderer/Public/PrimitiveSceneInfo.h b/Engine/Source/Runtime/Renderer/Public/PrimitiveSceneInfo.h
--- a/Engine/Source/Runtime/Renderer/Public/PrimitiveSceneInfo.h
+++ b/Engine/Source/Runtime/Renderer/Public/PrimitiveSceneInfo.h
@@ -486,6 +486,13 @@ public:
 	/** Adds a request to update GPU scene representation. */
 	RENDERER_API bool RequestGPUSceneUpdate(EPrimitiveDirtyState PrimitiveDirtyState = EPrimitiveDirtyState::ChangedAll);
 
+	/** Marks the primitive UB as needing updated and requests a GPU scene update */
+	void MarkGPUStateDirty(EPrimitiveDirtyState PrimitiveDirtyState = EPrimitiveDirtyState::ChangedAll)
+	{
+		SetNeedsUniformBufferUpdate(true);
+		RequestGPUSceneUpdate(PrimitiveDirtyState);
+	}
+
 	/** 
 	 * Builds an array of all primitive scene info's in this primitive's attachment group. 
 	 * This only works on potential parents (!LightingAttachmentRoot.IsValid()) and will include the current primitive in the output array.
