Remove pixel shader version of BasicEyeAdaptation.
- Compute shader was already default.
- Part of work to move from EyeAdaptationTexture to EyeAdaptationBuffer on every platform.

diff --git a/Engine/Shaders/Private/PostProcessEyeAdaptation.usf b/Engine/Shaders/Private/PostProcessEyeAdaptation.usf
--- a/Engine/Shaders/Private/PostProcessEyeAdaptation.usf
+++ b/Engine/Shaders/Private/PostProcessEyeAdaptation.usf
@@ -20,6 +20,8 @@
 #include "SHCommon.ush"
 #include "/Engine/Private/ScreenPass.ush"
 
+RWTexture2D<float4> RWEyeAdaptationTexture;
+
 #endif
 
 Texture2D EyeAdaptationTexture;
@@ -30,12 +32,6 @@ Texture2D ColorTexture;
 
 SamplerState ColorSampler;
 
-#if COMPUTESHADER
-
-RWTexture2D<float4> RWEyeAdaptationTexture;
-
-#endif
-
 uint2 Color_ViewportMin;
 uint2 Color_ViewportMax;
 
@@ -177,17 +173,10 @@ groupshared float2 SubRectValueWeight[TGSIZE*TGSIZE];
 [numthreads(TGSIZE, TGSIZE, 1)]
 void BasicEyeAdaptationCS(uint GIndex : SV_GroupIndex, uint2 GTId : SV_GroupThreadID)
 {
-	float4 OutColor=0;
-
-#else
-
-void BasicEyeAdaptationPS(in float4 UVAndScreenPos : TEXCOORD0, out float4 OutColor : SV_Target0)
-{
-
-#endif
-
+	float4 OutColor = 0;
+	
 	// Compute scaled Log Luminance Average
-#if COMPUTESHADER
+
 	// There are TGSIZE*TGSIZE threads. Each thread will calculate the luminance for its own subregions from a TGSIZE*TGSIZE screen grid
 	const uint2 SubRectMin = uint2(
 		((Color_ViewportMax.x - Color_ViewportMin.x) * GTId.x) / TGSIZE,
@@ -229,9 +218,6 @@ void BasicEyeAdaptationPS(in float4 UVAndScreenPos : TEXCOORD0, out float4 OutCo
 	GroupMemoryBarrierWithGroupSync();
 
 	float LogLumAve = SubRectValueWeight[0].x /  SubRectValueWeight[0].y;
-#else
-	float LogLumAve = ComputeWeightedTextureAverageAlpha(ColorTexture, Color_ViewportMin, Color_ViewportMax);
-#endif
 
 	// Correct for [0,1] scaling
 	LogLumAve = (LogLumAve - EyeAdaptation_HistogramBias) / EyeAdaptation_HistogramScale;
@@ -269,14 +255,14 @@ void BasicEyeAdaptationPS(in float4 UVAndScreenPos : TEXCOORD0, out float4 OutCo
 	OutColor = !all(IsFinite(OutColor)) ? float4(1, 1, 1, 0) : OutColor;
 #endif
 	
-#if COMPUTESHADER
 	if(GIndex==0)
 	{
 		RWEyeAdaptationTexture[uint2(0,0)] = OutColor;
 	}
-#endif
 }
 
+#endif
+
 #if COMPUTESHADER
 
 RWTexture2D<float> RWIlluminanceTexture;
diff --git a/Engine/Source/Runtime/Renderer/Private/PostProcess/PostProcessEyeAdaptation.cpp b/Engine/Source/Runtime/Renderer/Private/PostProcess/PostProcessEyeAdaptation.cpp
--- a/Engine/Source/Runtime/Renderer/Private/PostProcess/PostProcessEyeAdaptation.cpp
+++ b/Engine/Source/Runtime/Renderer/Private/PostProcess/PostProcessEyeAdaptation.cpp
@@ -36,14 +36,6 @@ namespace
 		TEXT(" 3: Manual"),
 		ECVF_Scalability | ECVF_RenderThreadSafe);
 
-	TAutoConsoleVariable<int32> CVarEyeAdaptationBasicCompute(
-		TEXT("r.EyeAdaptation.Basic.Compute"),
-		1,
-		TEXT("Use Pixel or Compute Shader to compute the basic eye adaptation. \n")
-		TEXT("= 0 : Pixel Shader\n")
-		TEXT("> 0 : Compute Shader (default) \n"),
-		ECVF_Scalability | ECVF_RenderThreadSafe);
-
 	TAutoConsoleVariable<float> CVarEyeAdaptationExponentialTransitionDistance(
 		TEXT("r.EyeAdaptation.ExponentialTransitionDistance"),
 		1.5,
@@ -770,9 +762,9 @@ public:
 		SHADER_PARAMETER_RDG_TEXTURE(Texture2D, ColorTexture)
 		SHADER_PARAMETER_SAMPLER(SamplerState, ColorSampler)
 		RENDER_TARGET_BINDING_SLOTS()
-		END_SHADER_PARAMETER_STRUCT()
+	END_SHADER_PARAMETER_STRUCT()
 
-		static bool ShouldCompilePermutation(const FGlobalShaderPermutationParameters& Parameters)
+	static bool ShouldCompilePermutation(const FGlobalShaderPermutationParameters& Parameters)
 	{
 		return IsFeatureLevelSupported(Parameters.Platform, ERHIFeatureLevel::Type(GetBasicEyeAdaptationMinFeatureLevel()));
 	}
@@ -819,64 +811,22 @@ FScreenPassTexture AddBasicEyeAdaptationSetupPass(
 	return FScreenPassTexture(OutputTexture, SceneColor.ViewRect);
 }
 
-class FBasicEyeAdaptationShader : public FGlobalShader
+class FBasicEyeAdaptationCS : public FGlobalShader
 {
 public:
+	DECLARE_GLOBAL_SHADER(FBasicEyeAdaptationCS);
+	SHADER_USE_PARAMETER_STRUCT(FBasicEyeAdaptationCS, FGlobalShader);
+
 	BEGIN_SHADER_PARAMETER_STRUCT(FParameters, )
 		SHADER_PARAMETER_STRUCT_REF(FViewUniformShaderParameters, View)
 		SHADER_PARAMETER_STRUCT(FEyeAdaptationParameters, EyeAdaptation)
 		SHADER_PARAMETER_STRUCT(FScreenPassTextureViewportParameters, Color)
 		SHADER_PARAMETER_RDG_TEXTURE(Texture2D, ColorTexture)
 		SHADER_PARAMETER_RDG_TEXTURE(Texture2D, EyeAdaptationTexture)
-		END_SHADER_PARAMETER_STRUCT()
-
-		static const EPixelFormat OutputFormat = PF_A32B32G32R32F;
-
-	static void ModifyCompilationEnvironment(const FGlobalShaderPermutationParameters& Parameters, FShaderCompilerEnvironment& OutEnvironment)
-	{
-		FGlobalShader::ModifyCompilationEnvironment(Parameters, OutEnvironment);
-		OutEnvironment.SetRenderTargetOutputFormat(0, OutputFormat);
-	}
-
-	static bool ShouldCompilePermutation(const FGlobalShaderPermutationParameters& Parameters)
-	{
-		return IsFeatureLevelSupported(Parameters.Platform, ERHIFeatureLevel::Type(GetBasicEyeAdaptationMinFeatureLevel()));
-	}
-
-	FBasicEyeAdaptationShader() = default;
-	FBasicEyeAdaptationShader(const CompiledShaderInitializerType& Initializer)
-		: FGlobalShader(Initializer)
-	{}
-};
-
-class FBasicEyeAdaptationPS : public FBasicEyeAdaptationShader
-{
-	using Super = FBasicEyeAdaptationShader;
-public:
-	DECLARE_GLOBAL_SHADER(FBasicEyeAdaptationPS);
-	SHADER_USE_PARAMETER_STRUCT(FBasicEyeAdaptationPS, Super);
-
-	BEGIN_SHADER_PARAMETER_STRUCT(FParameters, )
-		SHADER_PARAMETER_STRUCT_INCLUDE(Super::FParameters, Base)
-		RENDER_TARGET_BINDING_SLOTS()
-		END_SHADER_PARAMETER_STRUCT()
-};
-
-IMPLEMENT_GLOBAL_SHADER(FBasicEyeAdaptationPS, "/Engine/Private/PostProcessEyeAdaptation.usf", "BasicEyeAdaptationPS", SF_Pixel);
-
-class FBasicEyeAdaptationCS : public FBasicEyeAdaptationShader
-{
-	using Super = FBasicEyeAdaptationShader;
-public:
-	DECLARE_GLOBAL_SHADER(FBasicEyeAdaptationCS);
-	SHADER_USE_PARAMETER_STRUCT(FBasicEyeAdaptationCS, Super);
-
-	BEGIN_SHADER_PARAMETER_STRUCT(FParameters, )
-		SHADER_PARAMETER_STRUCT_INCLUDE(Super::FParameters, Base)
 		SHADER_PARAMETER_RDG_TEXTURE_UAV(RWTexture2D<float4>, RWEyeAdaptationTexture)
-		END_SHADER_PARAMETER_STRUCT()
+	END_SHADER_PARAMETER_STRUCT()
 
-		static bool ShouldCompilePermutation(const FGlobalShaderPermutationParameters& Parameters)
+	static bool ShouldCompilePermutation(const FGlobalShaderPermutationParameters& Parameters)
 	{
 		return IsFeatureLevelSupported(Parameters.Platform, ERHIFeatureLevel::SM5);
 	}
@@ -898,47 +848,22 @@ FRDGTextureRef AddBasicEyeAdaptationPass(
 
 	FRDGTextureRef OutputTexture = GraphBuilder.RegisterExternalTexture(View.GetEyeAdaptationTexture(GraphBuilder.RHICmdList), ERDGTextureFlags::MultiFrame);
 
-	FBasicEyeAdaptationShader::FParameters PassBaseParameters;
-	PassBaseParameters.View = View.ViewUniformBuffer;
-	PassBaseParameters.EyeAdaptation = EyeAdaptationParameters;
-	PassBaseParameters.Color = GetScreenPassTextureViewportParameters(SceneColorViewport);
-	PassBaseParameters.ColorTexture = SceneColor.Texture;
-	PassBaseParameters.EyeAdaptationTexture = EyeAdaptationTexture;
+	FBasicEyeAdaptationCS::FParameters* PassParameters = GraphBuilder.AllocParameters<FBasicEyeAdaptationCS::FParameters>();
+	PassParameters->View = View.ViewUniformBuffer;
+	PassParameters->EyeAdaptation = EyeAdaptationParameters;
+	PassParameters->Color = GetScreenPassTextureViewportParameters(SceneColorViewport);
+	PassParameters->ColorTexture = SceneColor.Texture;
+	PassParameters->EyeAdaptationTexture = EyeAdaptationTexture;
+	PassParameters->RWEyeAdaptationTexture = GraphBuilder.CreateUAV(OutputTexture);
 
-	if (View.bUseComputePasses || CVarEyeAdaptationBasicCompute.GetValueOnRenderThread())
-	{
-		FBasicEyeAdaptationCS::FParameters* PassParameters = GraphBuilder.AllocParameters<FBasicEyeAdaptationCS::FParameters>();
-		PassParameters->Base = PassBaseParameters;
-		PassParameters->RWEyeAdaptationTexture = GraphBuilder.CreateUAV(OutputTexture);
-
-		TShaderMapRef<FBasicEyeAdaptationCS> ComputeShader(View.ShaderMap);
-
-		FComputeShaderUtils::AddPass(
-			GraphBuilder,
-			RDG_EVENT_NAME("BasicEyeAdaptation (CS)"),
-			ComputeShader,
-			PassParameters,
-			FIntVector(1, 1, 1));
-	}
-	else
-	{
-		FBasicEyeAdaptationPS::FParameters* PassParameters = GraphBuilder.AllocParameters<FBasicEyeAdaptationPS::FParameters>();
-		PassParameters->Base = PassBaseParameters;
-		PassParameters->RenderTargets[0] = FRenderTargetBinding(OutputTexture, ERenderTargetLoadAction::ENoAction);
-
-		TShaderMapRef<FBasicEyeAdaptationPS> PixelShader(View.ShaderMap);
-
-		const FScreenPassTextureViewport OutputViewport(OutputTexture);
-
-		AddDrawScreenPass(
-			GraphBuilder,
-			RDG_EVENT_NAME("BasicEyeAdaptation (PS)"),
-			View,
-			OutputViewport,
-			OutputViewport,
-			PixelShader,
-			PassParameters);
-	}
+	auto ComputeShader = View.ShaderMap->GetShader<FBasicEyeAdaptationCS>();
+
+	FComputeShaderUtils::AddPass(
+		GraphBuilder,
+		RDG_EVENT_NAME("BasicEyeAdaptation (CS)"),
+		ComputeShader,
+		PassParameters,
+		FIntVector(1, 1, 1));
 
 	View.EnqueueEyeAdaptationExposureTextureReadback(GraphBuilder);
 
