Fix misplaced elements in VisualizeHDR.
- Apply DPIScale when using Canvas.DrawShadowedString.
- Apply ResolutionFractionAndInv when rendering illuminance meter square.
- Small tweaks to position of some elements.

diff --git a/Engine/Shaders/Private/DebugProbes.usf b/Engine/Shaders/Private/DebugProbes.usf
--- a/Engine/Shaders/Private/DebugProbes.usf
+++ b/Engine/Shaders/Private/DebugProbes.usf
@@ -253,14 +253,14 @@ void MainPS(
 #if PERMUTATION_ILLUMINANCEMETER /////////////////////////////////////////////////////////////////////////////////////////////////////////////////////////////////////
 
 	// In this case we want to show a small patch in front of the camera with pure white albedo.
-	// That is a way to measure the illuminance at the camera position, on the hemisphere pointing otwards the back of the camera.
+	// That is a way to measure the illuminance at the camera position, on the hemisphere pointing towards the back of the camera.
 
 	const float DeviceZ = 0.999999f;
 
 	// IlluminanceMeter variables must match the one in PostProcessVisualizeHDR.usf
-	const float  IlluminanceMeterSize = 20;
+	const float  IlluminanceMeterSize = 20 * View.ResolutionFractionAndInv.x;
 	const float  IlluminanceMeterHalfSize = IlluminanceMeterSize / 2;
-	const float2 IlluminanceMeterCenter = View.ViewSizeAndInvSize.xy * 0.5f - uint2(0, 100);
+	const float2 IlluminanceMeterCenter = View.ViewSizeAndInvSize.xy * 0.5f - uint2(0, 100) * View.ResolutionFractionAndInv.x;
 
 	if(all(PixelCoord > (IlluminanceMeterCenter - IlluminanceMeterHalfSize)) && all(PixelCoord < (IlluminanceMeterCenter + IlluminanceMeterHalfSize)))
 	{
diff --git a/Engine/Shaders/Private/PostProcessVisualizeHDR.usf b/Engine/Shaders/Private/PostProcessVisualizeHDR.usf
--- a/Engine/Shaders/Private/PostProcessVisualizeHDR.usf
+++ b/Engine/Shaders/Private/PostProcessVisualizeHDR.usf
@@ -313,16 +313,14 @@ void MainPS(noperspective float4 UVAndScreenPos : TEXCOORD0, float4 SvPosition :
 		// IlluminanceMeter variables must match the one in DebugProbes.usf
 		const float  IlluminanceMeterSize = 20;
 		const float  IlluminanceMeterHalfSize = IlluminanceMeterSize / 2;
-		const float2 IlluminanceMeterCenter = View.ViewSizeAndInvSize.xy * 0.5f - uint2(0, 100);
-
-		const float2 IlluminanceMeterCenterCoord = Input_UVViewportMin + IlluminanceMeterCenter;
-		const float2 IlluminanceMeterCenterUV = IlluminanceMeterCenterCoord * Input_ExtentInverse;
+		const float2 IlluminanceMeterCenter = OutputViewportCenter - uint2(0, 100);
+		const float2 IlluminanceMeterCenterUV = IlluminanceMeterCenter * Input_ExtentInverse;
 
 		const float PixelDx = PixelPos.x - IlluminanceMeterCenter.x;
 		const float PixelDy = PixelPos.y - IlluminanceMeterCenter.y;
 		const float AbsPixelDx = abs(PixelDx);
 		const float AbsPixelDy = abs(PixelDy);
-		if (PixelDx > - 20 && PixelDx < 300 && PixelDy < 40 && PixelDy > -20)
+		if (PixelDx > -20 && PixelDx < 300 && PixelDy < 40 && PixelDy > -20)
 		{
 			float3 LuminanceAvg = 0.0f;
 			float SampleCount = 0.0f;
@@ -354,7 +352,7 @@ void MainPS(noperspective float4 UVAndScreenPos : TEXCOORD0, float4 SvPosition :
 			float3 Illuminance = LuminanceAvg * PI;
 			float MeanIlluminance = Luminance(Illuminance);
 
-			int2 TopLeft = IlluminanceMeterCenterCoord + int2(11, 11);
+			int2 TopLeft = IlluminanceMeterCenter + int2(11, 11);
 			// Mean illuminance
 			if (MeanIlluminance < 1000.0f)
 			{
@@ -376,7 +374,7 @@ void MainPS(noperspective float4 UVAndScreenPos : TEXCOORD0, float4 SvPosition :
 			PrintCharacter(PixelPos, OutColor.xyz, NumberColor, TopLeft, _X_);
 
 			// Illuminance as RGB
-			TopLeft = IlluminanceMeterCenterCoord + int2(11, 22);
+			TopLeft = IlluminanceMeterCenter + int2(11, 22);
 			if (Illuminance.r < 1000.0f)
 			{
 				PrintFloat(PixelPos, OutColor.xyz, float3(1, 0.5, 0.5), TopLeft, Illuminance.r);
diff --git a/Engine/Source/Runtime/Renderer/Private/PostProcess/PostProcessVisualizeHDR.cpp b/Engine/Source/Runtime/Renderer/Private/PostProcess/PostProcessVisualizeHDR.cpp
--- a/Engine/Source/Runtime/Renderer/Private/PostProcess/PostProcessVisualizeHDR.cpp
+++ b/Engine/Source/Runtime/Renderer/Private/PostProcess/PostProcessVisualizeHDR.cpp
@@ -115,6 +115,8 @@ FScreenPassTexture AddVisualizeHDRPass(FRDGBuilder& GraphBuilder, const FViewInf
 	AddDrawCanvasPass(GraphBuilder, RDG_EVENT_NAME("Overlay"), View, Output,
 		[Output, &View](FCanvas& Canvas)
 	{
+		const float DPIScale = Canvas.GetDPIScale();
+
 		const EAutoExposureMethod AutoExposureMethod = GetAutoExposureMethod(View);
 		const bool bExtendedLuminanceRange = IsExtendLuminanceRangeEnabled();
 
@@ -132,25 +134,6 @@ FScreenPassTexture AddVisualizeHDRPass(FRDGBuilder& GraphBuilder, const FViewInf
 
 		Y += 160;
 
-		float MinX = Output.ViewRect.Min.X + 64 + 10;
-		float MaxY = Output.ViewRect.Max.Y - 64;
-		float SizeX = Output.ViewRect.Size().X - 64 * 2 - 20;
-
-		for (uint32 i = 0; i <= 4; ++i)
-		{
-			int XAdd = (int)(i * SizeX / 4);
-			float HistogramPosition = i / 4.0f;
-			float EV100Value = FMath::Lerp(View.FinalPostProcessSettings.HistogramLogMin, View.FinalPostProcessSettings.HistogramLogMax, HistogramPosition);
-			if (!bExtendedLuminanceRange)
-			{
-				// In this case the post process settings are actually Log2 values.
-				EV100Value = Log2ToEV100(LuminanceMax,EV100Value);
-			}
-
-			Line = FString::Printf(TEXT("%.2g"), EV100Value);
-			Canvas.DrawShadowedString(MinX + XAdd - 5, MaxY + YStep, *Line, GetStatsFont(), FLinearColor(1, 0.3f, 0.3f));
-		}
-		Y += 3 * YStep;
 		switch (AutoExposureMethod)
 		{
 		case EAutoExposureMethod::AEM_Basic:
@@ -216,8 +199,11 @@ FScreenPassTexture AddVisualizeHDRPass(FRDGBuilder& GraphBuilder, const FViewInf
 			Canvas.DrawShadowedString(X, Y += YStep, TEXT("Exposure Compensation (Curve):"), GetStatsFont(), FLinearColor(1, 1, 1));
 			Canvas.DrawShadowedString(X + ColumnWidth, Y, *Line, GetStatsFont(), FLinearColor(1, 1, 1));
 
-			const float IlluminanceMeterTextX = Output.ViewRect.Min.X + Output.ViewRect.Size().X * 0.5f - 180.0f;
-			const float IlluminanceMeterTextY = Output.ViewRect.Min.Y + Output.ViewRect.Size().Y * 0.5f - 130.0f;
+			const float ScreenCenterX = Output.ViewRect.Min.X + Output.ViewRect.Width() * 0.5f;
+			const float ScreenCenterY = Output.ViewRect.Min.Y + Output.ViewRect.Height() * 0.5f;
+
+			const float IlluminanceMeterTextX = (ScreenCenterX - 10.0f) / DPIScale;
+			const float IlluminanceMeterTextY = (ScreenCenterY - 140.0f) / DPIScale;
 			if(IsIlluminanceMeterSupportedByView(View))
 			{
 				Canvas.DrawShadowedString(IlluminanceMeterTextX, IlluminanceMeterTextY, TEXT("Illuminance meter - over the hemisphere of the surface patch"), GetStatsFont(), FLinearColor(1, 1, 1));
@@ -226,15 +212,40 @@ FScreenPassTexture AddVisualizeHDRPass(FRDGBuilder& GraphBuilder, const FViewInf
 			{
 				Canvas.DrawShadowedString(IlluminanceMeterTextX, IlluminanceMeterTextY, TEXT("Illuminance meter - not available with Forward Shading"), GetStatsFont(), FLinearColor(1, 1, 1));
 			}
-			
 
-			const float LuminanceMeterTextX = Output.ViewRect.Min.X + Output.ViewRect.Size().X * 0.5f - 50.0f;
-			const float LuminanceMeterTextY = Output.ViewRect.Min.Y + Output.ViewRect.Size().Y * 0.5f - 30.0f;
+			const float LuminanceMeterTextX = (ScreenCenterX - 10.0f) / DPIScale;
+			const float LuminanceMeterTextY = (ScreenCenterY - 40.0f) / DPIScale;
 			Canvas.DrawShadowedString(LuminanceMeterTextX, LuminanceMeterTextY, TEXT("Luminance meter"), GetStatsFont(), FLinearColor(1, 1, 1));
 
 			AutoExposureBias += CurveExposureBias;
 		}
 
+		// Draw EV100 values at the bottom of the histogram
+		{
+			const float MinX = Output.ViewRect.Min.X + 64 + 8;
+			const float MaxY = Output.ViewRect.Max.Y - 68;
+			const float SizeX = Output.ViewRect.Width() - 64 * 2 - 20;
+
+			for (uint32 i = 0; i <= 4; ++i)
+			{
+				int XAdd = (int)(i * SizeX / 4);
+				float HistogramPosition = i / 4.0f;
+				float EV100Value = FMath::Lerp(View.FinalPostProcessSettings.HistogramLogMin, View.FinalPostProcessSettings.HistogramLogMax, HistogramPosition);
+				if (!bExtendedLuminanceRange)
+				{
+					// In this case the post process settings are actually Log2 values.
+					EV100Value = Log2ToEV100(LuminanceMax, EV100Value);
+				}
+
+				Line = FString::Printf(TEXT("%.2g"), EV100Value);
+
+				const float PosX = (MinX + XAdd - 5) / DPIScale;
+				const float PosY = (MaxY + YStep) / DPIScale;
+
+				Canvas.DrawShadowedString(PosX, PosY, *Line, GetStatsFont(), FLinearColor(1, 0.3f, 0.3f));
+			}
+		}
+
 		Line = FString::Printf(TEXT("%.3g"), AutoExposureBias);
 		Canvas.DrawShadowedString(X, Y += YStep, TEXT("Exposure Compensation (All): "), GetStatsFont(), FLinearColor(1, 1, 1));
 		Canvas.DrawShadowedString(X + ColumnWidth, Y, *Line, GetStatsFont(), FLinearColor(1, 0.3f, 0.3f));
