Add mobile byte per pixel explicit limit.

diff --git a/Engine/Source/Runtime/Engine/Private/Materials/HLSLMaterialTranslator.cpp b/Engine/Source/Runtime/Engine/Private/Materials/HLSLMaterialTranslator.cpp
--- a/Engine/Source/Runtime/Engine/Private/Materials/HLSLMaterialTranslator.cpp
+++ b/Engine/Source/Runtime/Engine/Private/Materials/HLSLMaterialTranslator.cpp
@@ -88,8 +88,8 @@ static uint32 GetStrataBytePerPixel(EShaderPlatform InPlatform)
 {
 	// We enforce at least 20 bytes per pixel because this is the minimal Strata GBuffer footprint of the simplest material.
 	const uint32 MinStrataBytePerPixel = 20u;
-	const uint32 MaxStrataBytePerPixel = 256u;
-	static FShaderPlatformCachedIniValue<int32> CVarBudget(TEXT("r.Strata.BytesPerPixel"));
+	const uint32 MaxStrataBytePerPixel = IsMobilePlatform(InPlatform) ? 24u : 256u;
+	static FShaderPlatformCachedIniValue<int32> CVarBudget(TEXT("r.Strata.BytesPerPixel"));	
 	return FMath::Clamp(uint32(CVarBudget.Get(InPlatform)), MinStrataBytePerPixel, MaxStrataBytePerPixel);
 }
 
diff --git a/Engine/Source/Runtime/RenderCore/Private/Shader.cpp b/Engine/Source/Runtime/RenderCore/Private/Shader.cpp
--- a/Engine/Source/Runtime/RenderCore/Private/Shader.cpp
+++ b/Engine/Source/Runtime/RenderCore/Private/Shader.cpp
@@ -1909,7 +1909,7 @@ void ShaderMapAppendKeyString(EShaderPlatform Platform, FString& KeyString)
 		{
 			// We enforce at least 20 bytes per pixel because this is the minimal Strata GBuffer footprint of the simplest material.
 			const uint32 MinStrataBytePerPixel = 20u;
-			const uint32 MaxStrataBytePerPixel = 256u;
+			const uint32 MaxStrataBytePerPixel = IsMobilePlatform(Platform) ? 24u : 256u;
 			static FShaderPlatformCachedIniValue<int32> CVarBudget(TEXT("r.Strata.BytesPerPixel"));
 			const uint32 BytesPerPixel = FMath::Clamp(uint32(CVarBudget.Get(Platform)), MinStrataBytePerPixel, MaxStrataBytePerPixel);
 			KeyString += FString::Printf(TEXT("_BUDGET%u"), BytesPerPixel);
diff --git a/Engine/Source/Runtime/Renderer/Private/Strata/Strata.cpp b/Engine/Source/Runtime/Renderer/Private/Strata/Strata.cpp
--- a/Engine/Source/Runtime/Renderer/Private/Strata/Strata.cpp
+++ b/Engine/Source/Runtime/Renderer/Private/Strata/Strata.cpp
@@ -353,21 +353,21 @@ static bool NeedBSDFOffsets(const FScene* Scene, const FViewInfo& View)
 	return  ShouldRenderLumenDiffuseGI(Scene, View) || ShouldRenderLumenReflections(View) || Strata::ShouldRenderStrataDebugPasses(View);
 }
 
-static uint32 StrataGetBytePerPixel()
+static uint32 StrataGetBytePerPixel(EShaderPlatform InPlatform)
 {
 	// We enforce at least 20 bytes per pixel because this is the minimal Strata GBuffer footprint of the simplest material.
 	const uint32 MinStrataBytePerPixel = 20u;
-	const uint32 MaxStrataBytePerPixel = 256u;
+	const uint32 MaxStrataBytePerPixel = IsMobilePlatform(InPlatform) ? 24u : 256u;
 	return FMath::Clamp(uint32(CVarStrataBytePerPixel.GetValueOnAnyThread()), MinStrataBytePerPixel, MaxStrataBytePerPixel);
 }
 
-static void RecordStrataAnalytics()
+static void RecordStrataAnalytics(EShaderPlatform InPlatform)
 {
 	if (FEngineAnalytics::IsAvailable())
 	{
 		TArray<FAnalyticsEventAttribute> EventAttributes;
 		EventAttributes.Add(FAnalyticsEventAttribute(TEXT("Enabled"), 1));
-		EventAttributes.Add(FAnalyticsEventAttribute(TEXT("BytesPerPixel"), StrataGetBytePerPixel()));
+		EventAttributes.Add(FAnalyticsEventAttribute(TEXT("BytesPerPixel"), StrataGetBytePerPixel(InPlatform)));
 
 		FString OutStr(TEXT("Strata"));
 		FEngineAnalytics::GetProvider().RecordEvent(OutStr, EventAttributes);
@@ -414,7 +414,7 @@ void InitialiseStrataFrameSceneData(FRDGBuilder& GraphBuilder, FSceneRenderer& S
 		static bool bAnalyticsInitialized = false;
 		if (!bAnalyticsInitialized)
 		{
-			RecordStrataAnalytics();
+			RecordStrataAnalytics(SceneRenderer.ShaderPlatform);
 			bAnalyticsInitialized = true;
 		}
 
@@ -424,7 +424,7 @@ void InitialiseStrataFrameSceneData(FRDGBuilder& GraphBuilder, FSceneRenderer& S
 		UpdateMaterialBufferToTiledResolution(SceneTextureExtent, MaterialBufferSizeXY);
 
 		const uint32 RoundToValue = 4u;
-		Out.MaxBytesPerPixel = FMath::DivideAndRoundUp(StrataGetBytePerPixel(), RoundToValue) * RoundToValue;
+		Out.MaxBytesPerPixel = FMath::DivideAndRoundUp(StrataGetBytePerPixel(SceneRenderer.ShaderPlatform), RoundToValue) * RoundToValue;
 
 		// Top layer texture
 		{
