VT: Return aligned sizes in GetRequiredPageTableAllocationSize

#rnx

diff --git a/Engine/Source/Runtime/Renderer/Private/VT/VirtualTextureSpace.cpp b/Engine/Source/Runtime/Renderer/Private/VT/VirtualTextureSpace.cpp
--- a/Engine/Source/Runtime/Renderer/Private/VT/VirtualTextureSpace.cpp
+++ b/Engine/Source/Runtime/Renderer/Private/VT/VirtualTextureSpace.cpp
@@ -140,7 +140,7 @@ FUintPoint FVirtualTextureSpace::GetRequiredPageTableAllocationSize() const
 	// We align on some minimum size. Maybe minimum, and align sizes should be different? But OK for now.
 	const uint32 WidthAligned = Align(Width, VIRTUALTEXTURE_MIN_PAGETABLE_SIZE);
 	const uint32 HeightAligned = Align(Height, VIRTUALTEXTURE_MIN_PAGETABLE_SIZE);
-	return FUintPoint(Width, Height);
+	return FUintPoint(WidthAligned, HeightAligned);
 }
 
 uint32 FVirtualTextureSpace::GetSizeInBytes() const
