Implemented IsCachedRayTracingGeometryValid() instead of using hack in GetStaticRayTracingGeometryInstance().

diff --git a/Engine/Source/Runtime/Renderer/Private/DeferredShadingRenderer.cpp b/Engine/Source/Runtime/Renderer/Private/DeferredShadingRenderer.cpp
--- a/Engine/Source/Runtime/Renderer/Private/DeferredShadingRenderer.cpp
+++ b/Engine/Source/Runtime/Renderer/Private/DeferredShadingRenderer.cpp
@@ -1384,15 +1384,12 @@ bool FDeferredShadingSceneRenderer::GatherRayTracingWorldInstancesForView(FRDGBu
 					// (see FPrimitiveSceneInfo::UpdateCachedRayTracingInstance(...) and CacheRayTracingPrimitive(...))
 					const int32 LODIndex = 0;
 
-					FRHIRayTracingGeometry* RayTracingGeometryInstance = SceneInfo->GetStaticRayTracingGeometryInstance(LODIndex);
-					if (!RayTracingGeometryInstance)
+					if (!SceneInfo->IsCachedRayTracingGeometryValid())
 					{
 						// cached instance is not valid (eg: was streamed out) need to invalidate for next frame
 						ProxiesWithDirtyCachedInstance.Add(Scene.PrimitiveSceneProxies[PrimitiveIndex]);
 						continue;
 					}
-
-					check(RayTracingGeometryInstance == SceneInfo->CachedRayTracingInstance.GeometryRHI);
 					
 					// TODO: support GRayTracingExcludeDecals, but not in the form of RayTracingMeshCommand.bDecal as that requires looping over all cached MDCs
 					// Instead, either make r.RayTracing.ExcludeDecals read only or request a recache of all ray tracing commands during which decals are excluded
diff --git a/Engine/Source/Runtime/Renderer/Private/PrimitiveSceneInfo.cpp b/Engine/Source/Runtime/Renderer/Private/PrimitiveSceneInfo.cpp
--- a/Engine/Source/Runtime/Renderer/Private/PrimitiveSceneInfo.cpp
+++ b/Engine/Source/Runtime/Renderer/Private/PrimitiveSceneInfo.cpp
@@ -247,6 +247,18 @@ FPrimitiveSceneInfo::~FPrimitiveSceneInfo()
 }
 
 #if RHI_RAYTRACING
+bool FPrimitiveSceneInfo::IsCachedRayTracingGeometryValid() const
+{
+	if (CachedRayTracingGeometry)
+	{
+		// TODO: Doesn't take Nanite Ray Tracing into account
+		//check(CachedRayTracingGeometry->RayTracingGeometryRHI == CachedRayTracingInstance.GeometryRHI);
+		return CachedRayTracingGeometry->IsValid();
+	}
+
+	return false;
+}
+
 FRHIRayTracingGeometry* FPrimitiveSceneInfo::GetStaticRayTracingGeometryInstance(int LodLevel) const
 {
 	if (RayTracingGeometries.Num() > LodLevel)
@@ -268,11 +280,6 @@ FRHIRayTracingGeometry* FPrimitiveSceneInfo::GetStaticRayTracingGeometryInstance
 	}
 	else
 	{
-		if (LodLevel == 0 && CachedRayTracingGeometry && CachedRayTracingGeometry->IsValid())
-		{
-			return CachedRayTracingGeometry->RayTracingGeometryRHI;
-		}
-
 		return nullptr;
 	}
 }
diff --git a/Engine/Source/Runtime/Renderer/Public/PrimitiveSceneInfo.h b/Engine/Source/Runtime/Renderer/Public/PrimitiveSceneInfo.h
--- a/Engine/Source/Runtime/Renderer/Public/PrimitiveSceneInfo.h
+++ b/Engine/Source/Runtime/Renderer/Public/PrimitiveSceneInfo.h
@@ -599,6 +599,9 @@ public:
 
 #if RHI_RAYTRACING
 	static void UpdateCachedRaytracingData(FScene* Scene, const TArrayView<FPrimitiveSceneInfo*>& SceneInfos);
+
+	bool IsCachedRayTracingGeometryValid() const;
+
 	RENDERER_API FRHIRayTracingGeometry* GetStaticRayTracingGeometryInstance(int LodLevel) const;
 
 	int GetRayTracingGeometryNum() const { return RayTracingGeometries.Num(); }
