Added a new occlusion culling method based on GPU feedback from occlusion geometry
Can be activated with r.OcclusionFeedback.Enable=1
Works only with mobile rendering at the moment

diff --git a/Engine/Shaders/Private/OcclusionFeedbackShaders.usf b/Engine/Shaders/Private/OcclusionFeedbackShaders.usf
new file mode 100644
--- /dev/null
+++ b/Engine/Shaders/Private/OcclusionFeedbackShaders.usf
@@ -0,0 +1,40 @@
+// Copyright Epic Games, Inc. All Rights Reserved.
+#include "Common.ush"
+
+#if SHADING_PATH_MOBILE
+	#define OcclusionBuffer MobileBasePass.RWOcclusionBufferUAV
+#else
+	#define OcclusionBuffer DummyBuffer
+#endif
+
+uint StartIndex;
+
+void MainVS(
+	float4 InPosition	: ATTRIBUTE0,
+	float4 InOrigin		: ATTRIBUTE1,
+	float4 InExtent		: ATTRIBUTE2,
+	uint InstanceId		: SV_InstanceID,
+	out nointerpolation uint PrimitiveIdx : PRIMITIVE_ID,
+	out float4 OutPosition : SV_POSITION
+)
+{
+	ResolvedView = ResolveView();
+	float3 TransformedPosition = (InPosition.xyz * InExtent.xyz) + InOrigin.xyz;
+	OutPosition = mul(float4(TransformedPosition, 1), ResolvedView.TranslatedWorldToClip);
+	PrimitiveIdx = StartIndex + InstanceId;
+}
+
+EARLYDEPTHSTENCIL
+void MainPS(
+	in nointerpolation uint PrimitiveIdx : PRIMITIVE_ID,
+	out float4 OutColor : SV_Target0
+)
+{
+	// TODO: write only some pixels?
+	OcclusionBuffer[PrimitiveIdx] = 1u;
+
+#if PERMUTATION_PIXEL_DISCARD
+	clip(-1.0f);
+#endif
+	OutColor = float4(1, 0, 0, 1);
+}
\ No newline at end of file
diff --git a/Engine/Source/Runtime/OpenGLDrv/Private/OpenGLDevice.cpp b/Engine/Source/Runtime/OpenGLDrv/Private/OpenGLDevice.cpp
--- a/Engine/Source/Runtime/OpenGLDrv/Private/OpenGLDevice.cpp
+++ b/Engine/Source/Runtime/OpenGLDrv/Private/OpenGLDevice.cpp
@@ -154,7 +154,7 @@ bool IsUniformBufferBound( GLuint Buffer )
 extern void BeginFrame_UniformBufferPoolCleanup();
 extern void BeginFrame_VertexBufferCleanup();
 extern void BeginFrame_QueryBatchCleanup();
-extern void BeginFrame_PollAllFences();
+extern void OpenGL_PollAllFences();
 
 
 FOpenGLContextState& FOpenGLDynamicRHI::GetContextStateForCurrentContext(bool bAssertIfInvalid)
@@ -193,12 +193,14 @@ void FOpenGLDynamicRHI::RHIBeginFrame()
 	PendingState.DepthStencil = 0 ;
 #endif
 
-	BeginFrame_PollAllFences();
+	OpenGL_PollAllFences();
 }
 
 void FOpenGLDynamicRHI::RHIEndFrame()
 {
 	GPUProfilingData.EndFrame();
+
+	OpenGL_PollAllFences();
 }
 
 void FOpenGLDynamicRHI::RHIPerFrameRHIFlushComplete()
@@ -206,6 +208,8 @@ void FOpenGLDynamicRHI::RHIPerFrameRHIFlushComplete()
 	BeginFrame_UniformBufferPoolCleanup();
 	BeginFrame_VertexBufferCleanup();
 	BeginFrame_QueryBatchCleanup();
+
+	OpenGL_PollAllFences();
 }
 
 
diff --git a/Engine/Source/Runtime/OpenGLDrv/Private/OpenGLQuery.cpp b/Engine/Source/Runtime/OpenGLDrv/Private/OpenGLQuery.cpp
--- a/Engine/Source/Runtime/OpenGLDrv/Private/OpenGLQuery.cpp
+++ b/Engine/Source/Runtime/OpenGLDrv/Private/OpenGLQuery.cpp
@@ -229,7 +229,6 @@ void BeginOcclusionQueryBatch(uint32 NumOcclusionQueries)
 {
 	if (IsRunningRHIInSeparateThread())
 	{
-
 		GBatcher.StartNewBatch(*GetDynamicRHI<FOpenGLDynamicRHI>());
 	}
 }
@@ -242,11 +241,15 @@ void EndOcclusionQueryBatch()
 	}
 }
 
+void OpenGL_PollAllFences();
+
 void FOpenGLDynamicRHI::RHIPollOcclusionQueries()
 {
 	if (IsRunningRHIInSeparateThread())
 	{
 		GBatcher.SoftFlush(*GetDynamicRHI<FOpenGLDynamicRHI>());
+
+		OpenGL_PollAllFences();
 	}
 }
 
@@ -1197,11 +1200,15 @@ struct FOpenGLGPUFenceProxy
 
 TArray<FOpenGLGPUFenceProxy*> FOpenGLGPUFenceProxy::AllOpenGLGPUFences;
 
-void BeginFrame_PollAllFences()
+void OpenGL_PollAllFences()
 {
 	if (IsRunningRHIInSeparateThread())
 	{
-		FOpenGLGPUFenceProxy::PollAllFences();
+		RunOnGLRenderContextThread([]()
+		{
+			VERIFY_GL_SCOPE();
+			FOpenGLGPUFenceProxy::PollAllFences();
+		});
 	}
 }
 
diff --git a/Engine/Source/Runtime/OpenGLDrv/Private/OpenGLUAV.cpp b/Engine/Source/Runtime/OpenGLDrv/Private/OpenGLUAV.cpp
--- a/Engine/Source/Runtime/OpenGLDrv/Private/OpenGLUAV.cpp
+++ b/Engine/Source/Runtime/OpenGLDrv/Private/OpenGLUAV.cpp
@@ -311,6 +311,7 @@ void FOpenGLDynamicRHI::RHIClearUAVFloat(FRHIUnorderedAccessView* UnorderedAcces
 		{
 			int32 NumComponents = 0;
 			uint32 NumElements = 0;
+			bool bStorageBuffer = false;
 
 			if (Texture->UnrealFormat != 0)
 			{
@@ -321,14 +322,23 @@ void FOpenGLDynamicRHI::RHIClearUAVFloat(FRHIUnorderedAccessView* UnorderedAcces
 			{
 				NumElements = Texture->GetBufferSize() / sizeof(float);
 				NumComponents = 1;
+				bStorageBuffer = true;
 			}
 					
 			switch (NumComponents)
 			{
 			case 1:
-				ClearUAVShader_T<EClearReplacementResourceType::Buffer, EClearReplacementValueType::Float, 1, false>(RHICmdList, UnorderedAccessViewRHI, NumElements, 1, 1, *reinterpret_cast<const float(*)[1]>(&Values));
+				if (bStorageBuffer)
+				{
+					ClearUAVShader_T<EClearReplacementResourceType::StructuredBuffer, EClearReplacementValueType::Float, 1, false>(RHICmdList, UnorderedAccessViewRHI, NumElements, 1, 1, *reinterpret_cast<const float(*)[1]>(&Values));
+				}
+				else
+				{
+					ClearUAVShader_T<EClearReplacementResourceType::Buffer, EClearReplacementValueType::Float, 1, false>(RHICmdList, UnorderedAccessViewRHI, NumElements, 1, 1, *reinterpret_cast<const float(*)[1]>(&Values));
+				}
 				break;
 			case 4:
+				check(!bStorageBuffer);
 				ClearUAVShader_T<EClearReplacementResourceType::Buffer, EClearReplacementValueType::Float, 4, false>(RHICmdList, UnorderedAccessViewRHI, NumElements, 1, 1, *reinterpret_cast<const float(*)[4]>(&Values));
 				break;
 			default:
@@ -365,6 +375,7 @@ void FOpenGLDynamicRHI::RHIClearUAVUint(FRHIUnorderedAccessView* UnorderedAccess
 		{
 			int32 NumComponents = 0;
 			uint32 NumElements = 0;
+			bool bStorageBuffer = false;
 
 			if (Texture->UnrealFormat != 0)
 			{
@@ -375,14 +386,23 @@ void FOpenGLDynamicRHI::RHIClearUAVUint(FRHIUnorderedAccessView* UnorderedAccess
 			{
 				NumElements = Texture->GetBufferSize() / sizeof(uint32);
 				NumComponents = 1;
+				bStorageBuffer = true;
 			}
 
 			switch (NumComponents)
 			{
 			case 1:
-				ClearUAVShader_T<EClearReplacementResourceType::Buffer, EClearReplacementValueType::Uint32, 1, false>(RHICmdList, UnorderedAccessViewRHI, NumElements, 1, 1, *reinterpret_cast<const uint32(*)[1]>(&Values));
+				if (bStorageBuffer)
+				{
+					ClearUAVShader_T<EClearReplacementResourceType::StructuredBuffer, EClearReplacementValueType::Uint32, 1, false>(RHICmdList, UnorderedAccessViewRHI, NumElements, 1, 1, *reinterpret_cast<const uint32(*)[1]>(&Values));
+				}
+				else
+				{
+					ClearUAVShader_T<EClearReplacementResourceType::Buffer, EClearReplacementValueType::Uint32, 1, false>(RHICmdList, UnorderedAccessViewRHI, NumElements, 1, 1, *reinterpret_cast<const uint32(*)[1]>(&Values));
+				}
 				break;
 			case 4:
+				check(!bStorageBuffer)
 				ClearUAVShader_T<EClearReplacementResourceType::Buffer, EClearReplacementValueType::Uint32, 4, false>(RHICmdList, UnorderedAccessViewRHI, NumElements, 1, 1, *reinterpret_cast<const uint32(*)[4]>(&Values));
 				break;
 			default:
diff --git a/Engine/Source/Runtime/Renderer/Private/MobileBasePassRendering.cpp b/Engine/Source/Runtime/Renderer/Private/MobileBasePassRendering.cpp
--- a/Engine/Source/Runtime/Renderer/Private/MobileBasePassRendering.cpp
+++ b/Engine/Source/Runtime/Renderer/Private/MobileBasePassRendering.cpp
@@ -227,6 +227,14 @@ void SetupMobileBasePassUniformParameters(
 		BasePassParameters.ScreenSpaceShadowMaskSampler = TStaticSamplerState<SF_Point, AM_Clamp, AM_Clamp, AM_Clamp>::GetRHI();
 	}
 
+	FRDGBuffer* OcclusionBuffer = View.ViewState ? View.ViewState->OcclusionFeedback.GetGPUFeedbackBuffer() : nullptr;
+	if (!OcclusionBuffer)
+	{
+		OcclusionBuffer = GraphBuilder.CreateBuffer(FRDGBufferDesc::CreateStructuredDesc(sizeof(uint32), 1), TEXT("OcclusionBufferFallback"));
+		
+	}
+	BasePassParameters.RWOcclusionBufferUAV = GraphBuilder.CreateUAV(OcclusionBuffer);
+
 	// Strata
 	Strata::BindStrataMobileForwardPasslUniformParameters(GraphBuilder, View, BasePassParameters.Strata);
 
diff --git a/Engine/Source/Runtime/Renderer/Private/MobileBasePassRendering.h b/Engine/Source/Runtime/Renderer/Private/MobileBasePassRendering.h
--- a/Engine/Source/Runtime/Renderer/Private/MobileBasePassRendering.h
+++ b/Engine/Source/Runtime/Renderer/Private/MobileBasePassRendering.h
@@ -43,6 +43,7 @@ BEGIN_GLOBAL_SHADER_PARAMETER_STRUCT(FMobileBasePassUniformParameters, )
 	SHADER_PARAMETER_TEXTURE(Texture2D, PreIntegratedGFTexture)
 	SHADER_PARAMETER_SAMPLER(SamplerState, PreIntegratedGFSampler)
 	SHADER_PARAMETER_RDG_BUFFER_SRV(StructuredBuffer<float4>, EyeAdaptationBuffer)
+	SHADER_PARAMETER_RDG_BUFFER_UAV(RWStructuredBuffer<uint>, RWOcclusionBufferUAV)
 	SHADER_PARAMETER_RDG_TEXTURE(Texture2D, AmbientOcclusionTexture)
 	SHADER_PARAMETER_SAMPLER(SamplerState, AmbientOcclusionSampler)
 	SHADER_PARAMETER_RDG_TEXTURE(Texture2D, ScreenSpaceShadowMaskTexture)
diff --git a/Engine/Source/Runtime/Renderer/Private/MobileShadingRenderer.cpp b/Engine/Source/Runtime/Renderer/Private/MobileShadingRenderer.cpp
--- a/Engine/Source/Runtime/Renderer/Private/MobileShadingRenderer.cpp
+++ b/Engine/Source/Runtime/Renderer/Private/MobileShadingRenderer.cpp
@@ -733,6 +733,30 @@ static void BuildMeshPassInstanceCullingDrawParams(FRDGBuilder& GraphBuilder, co
 	}
 }
 
+static void BeginOcclusionScope(FRDGBuilder& GraphBuilder, TArray<FViewInfo>& Views)
+{
+	for (int32 ViewIndex = 0; ViewIndex < Views.Num(); ViewIndex++)
+	{
+		FViewInfo& View = Views[ViewIndex];
+		if (View.ShouldRenderView() && View.ViewState && View.ViewState->OcclusionFeedback.IsInitialized())
+		{
+			View.ViewState->OcclusionFeedback.BeginOcclusionScope(GraphBuilder);
+		}
+	}
+}
+
+static void EndOcclusionScope(FRDGBuilder& GraphBuilder, TArray<FViewInfo>& Views)
+{
+	for (int32 ViewIndex = 0; ViewIndex < Views.Num(); ViewIndex++)
+	{
+		FViewInfo& View = Views[ViewIndex];
+		if (View.ShouldRenderView() && View.ViewState && View.ViewState->OcclusionFeedback.IsInitialized())
+		{
+			View.ViewState->OcclusionFeedback.EndOcclusionScope(GraphBuilder);
+		}
+	}
+}
+
 /*
 * Renders the Full Depth Prepass
 */
@@ -761,6 +785,7 @@ void FMobileSceneRenderer::RenderFullDepthPrepass(FRDGBuilder& GraphBuilder, FSc
 		auto* PassParameters = GraphBuilder.AllocParameters<FMobileRenderPassParameters>();
 		PassParameters->RenderTargets = BasePassRenderTargets;
 		PassParameters->View = View.GetShaderParameters();
+		PassParameters->MobileBasePass = CreateMobileBasePassUniformBuffer(GraphBuilder, View, EMobileBasePass::Opaque, EMobileSceneTextureSetupMode::None);
 
 		if (Scene->GPUScene.IsEnabled())
 		{
@@ -961,6 +986,8 @@ void FMobileSceneRenderer::Render(FRDGBuilder& GraphBuilder)
 		RenderCustomDepthPass(GraphBuilder, SceneTextures.CustomDepth, SceneTextures.GetSceneTextureShaderParameters(FeatureLevel));
 	}
 
+	BeginOcclusionScope(GraphBuilder, Views);
+	
 	// Sort objects' triangles
 	for (FViewInfo& View : Views)
 	{
@@ -1013,6 +1040,8 @@ void FMobileSceneRenderer::Render(FRDGBuilder& GraphBuilder)
 		RenderForward(GraphBuilder, ViewFamilyTexture, SceneTextures);
 	}
 
+	EndOcclusionScope(GraphBuilder, Views);
+
 	if (!bIsFullDepthPrepassEnabled)
 	{
 		FenceOcclusionTests(GraphBuilder);
diff --git a/Engine/Source/Runtime/Renderer/Private/RendererScene.cpp b/Engine/Source/Runtime/Renderer/Private/RendererScene.cpp
--- a/Engine/Source/Runtime/Renderer/Private/RendererScene.cpp
+++ b/Engine/Source/Runtime/Renderer/Private/RendererScene.cpp
@@ -427,6 +427,16 @@ FSceneViewState::FSceneViewState(ERHIFeatureLevel::Type FeatureLevel, FSceneView
 	LumenSurfaceCacheResolution = 1.0f;
 
 	ViewVirtualShadowMapCache = nullptr;
+
+	// OcclusionFeedback works only with mobile rendering atm
+	if (FeatureLevel == ERHIFeatureLevel::ES3_1)
+	{
+		extern int32 GOcclusionFeedback_Enable;
+		if (GOcclusionFeedback_Enable != 0)
+		{
+			BeginInitResource(&OcclusionFeedback);
+		}
+	}
 }
 
 void DestroyRenderResource(FRenderResource* RenderResource)
diff --git a/Engine/Source/Runtime/Renderer/Private/SceneOcclusion.cpp b/Engine/Source/Runtime/Renderer/Private/SceneOcclusion.cpp
--- a/Engine/Source/Runtime/Renderer/Private/SceneOcclusion.cpp
+++ b/Engine/Source/Runtime/Renderer/Private/SceneOcclusion.cpp
@@ -1362,6 +1362,12 @@ static void BeginOcclusionTests(
 					View.GroupedOcclusionQueries.Flush(RHICmdList);
 				}
 			}
+
+			if (View.ViewState && View.ViewState->OcclusionFeedback.IsInitialized())
+			{
+				SCOPED_DRAW_EVENT(RHICmdList, IndividualQueries);
+				View.ViewState->OcclusionFeedback.SubmitOcclusionDraws(RHICmdList, View);
+			}
 		}
 	}
 }
@@ -1597,3 +1603,293 @@ void FSceneRenderer::WaitOcclusionTests(FRHICommandListImmediate& RHICmdList)
 		}
 	}
 }
+
+FOcclusionFeedback::FOcclusionFeedback()
+	: FRenderResource()
+	, CurrentBufferIndex(0)
+{
+	OcclusionVertexDeclarationRHI = nullptr;
+}
+
+FOcclusionFeedback::~FOcclusionFeedback()
+{
+}
+
+void FOcclusionFeedback::InitDynamicRHI()
+{
+	for (int32 i = 0; i < UE_ARRAY_COUNT(OcclusionBuffers); ++i)
+	{
+		OcclusionBuffers[i].ReadbackBuffer = new FRHIGPUBufferReadback(TEXT("FeedbackOcclusionBuffer.Readback"));
+	}
+
+	{
+		FVertexDeclarationElementList Elements;
+		Elements.Add(FVertexElement(0, 0,					VET_Float4, 0, sizeof(FVector4f)*1, false));
+		Elements.Add(FVertexElement(1, 0,					VET_Float4, 1, sizeof(FVector4f)*2, true));
+		Elements.Add(FVertexElement(1, sizeof(FVector4f),	VET_Float4, 2, sizeof(FVector4f)*2, true));
+		OcclusionVertexDeclarationRHI = PipelineStateCache::GetOrCreateVertexDeclaration(Elements);
+	}
+}
+void FOcclusionFeedback::ReleaseDynamicRHI()
+{
+	BatchOcclusionQueries.Empty();
+
+	for (int32 i = 0; i < UE_ARRAY_COUNT(OcclusionBuffers); ++i)
+	{
+		delete OcclusionBuffers[i].ReadbackBuffer;
+		OcclusionBuffers[i].ReadbackBuffer = nullptr;
+		OcclusionBuffers[i].BatchedPrimitives.Empty();
+	}
+	CurrentBufferIndex = 0;
+	OcclusionVertexDeclarationRHI.SafeRelease();
+	GPUFeedbackBuffer = nullptr;
+	LatestOcclusionResults.Empty();
+}
+
+void FOcclusionFeedback::BeginOcclusionScope(FRDGBuilder& GraphBuilder)
+{
+	check(GPUFeedbackBuffer == nullptr);
+
+	FOcclusionBuffer& OcclusionBuffer = OcclusionBuffers[CurrentBufferIndex];
+	uint32 NumPrimitives = FMath::Max(1, OcclusionBuffer.BatchedPrimitives.Num());
+
+	FRDGBufferDesc Desc = FRDGBufferDesc::CreateStructuredDesc(sizeof(uint32), NumPrimitives);
+	Desc.Usage |= EBufferUsageFlags::SourceCopy;
+	GPUFeedbackBuffer = GraphBuilder.CreateBuffer(Desc, TEXT("FeedbackOcclusionBuffer"));
+	AddClearUAVPass(GraphBuilder, GraphBuilder.CreateUAV(GPUFeedbackBuffer), 0u);
+}
+
+void FOcclusionFeedback::EndOcclusionScope(FRDGBuilder& GraphBuilder)
+{
+	FOcclusionBuffer& OcclusionBuffer = OcclusionBuffers[CurrentBufferIndex];
+
+	FRHIGPUBufferReadback* ReadbackBuffer = OcclusionBuffer.ReadbackBuffer;
+	check(ReadbackBuffer);
+	uint32 ReadbackBytes = OcclusionBuffer.BatchedPrimitives.Num() * sizeof(uint32);
+	AddEnqueueCopyPass(GraphBuilder, ReadbackBuffer, GPUFeedbackBuffer, ReadbackBytes);
+	GPUFeedbackBuffer = nullptr;
+}
+
+void FOcclusionFeedback::ReadbackResults(FRHICommandList& RHICmdList)
+{
+	constexpr uint32 NumBuffers = UE_ARRAY_COUNT(OcclusionBuffers);
+
+	CurrentBufferIndex = (CurrentBufferIndex + 1u) % NumBuffers;
+
+	FOcclusionBuffer& OcclusionBuffer = OcclusionBuffers[CurrentBufferIndex];
+	uint32 NumPrimitives = OcclusionBuffer.BatchedPrimitives.Num();
+
+	bool bReady = (NumPrimitives > 0 && OcclusionBuffer.ReadbackBuffer->IsReady());
+	if (!bReady)
+	{
+		if (OcclusionBuffer.BatchedPrimitives.Num() > 0)
+		{
+			// We'll do manual wait
+			double StartTime = FPlatformTime::Seconds();
+			FRenderThreadIdleScope IdleScope(ERenderThreadIdleTypes::WaitingForGPUQuery);
+			ENamedThreads::Type RenderThread_Local = ENamedThreads::GetRenderThread_Local();
+
+			while (!bReady)
+			{
+				FPlatformProcess::SleepNoStats(0);
+				// pump RHIThread to make sure these queries have actually been submitted to the GPU.
+				if (IsInActualRenderingThread())
+				{
+					FTaskGraphInterface::Get().ProcessThreadUntilIdle(RenderThread_Local);
+				}
+				const double TimeoutValue = 0.5;
+				// look for gpu stuck/crashed
+				if ((FPlatformTime::Seconds() - StartTime) > TimeoutValue)
+				{
+					UE_LOG(LogRHI, Log, TEXT("Timed out while waiting for GPU to catch up on occlusion readback. (%.1f s)"), TimeoutValue);
+					break;
+				}
+
+				bReady = OcclusionBuffer.ReadbackBuffer->IsReady();
+			}
+		}
+	}
+
+	if (bReady)
+	{
+		LatestOcclusionResults.Reset();
+
+		uint32 NumBytes = NumPrimitives * sizeof(uint32);
+		uint32* Results = (uint32*)OcclusionBuffer.ReadbackBuffer->Lock(NumBytes);
+
+		for (uint32 i = 0; i < NumPrimitives; ++i)
+		{
+			if (Results[i] == 0u)
+			{
+				LatestOcclusionResults.Add(OcclusionBuffer.BatchedPrimitives[i]);
+			}
+		}
+
+		OcclusionBuffer.ReadbackBuffer->Unlock();
+	}
+
+	OcclusionBuffer.BatchedPrimitives.Empty();
+}
+
+void FOcclusionFeedback::AddPrimitive(FPrimitiveComponentId PrimitiveId, const FVector& BoundsOrigin, const FVector& BoundsBoxExtent, FGlobalDynamicVertexBuffer& DynamicVertexBuffer)
+{
+	constexpr uint32 MaxBatchedPrimitives = 512;
+	constexpr uint32 PrimitiveStride = sizeof(FVector4f) * 2u;
+
+	if (BatchOcclusionQueries.Num() == 0 ||
+		BatchOcclusionQueries.Last().NumBatchedPrimitives >= MaxBatchedPrimitives)
+	{
+		FOcclusionBatch OcclusionBatch;
+		OcclusionBatch.NumBatchedPrimitives = 0u;
+		OcclusionBatch.VertexAllocation = DynamicVertexBuffer.Allocate(MaxBatchedPrimitives * PrimitiveStride);
+		check(OcclusionBatch.VertexAllocation.IsValid());
+		BatchOcclusionQueries.Add(OcclusionBatch);
+	}
+
+	FOcclusionBatch& OcclusionBatch = BatchOcclusionQueries.Last();
+
+	// TODO: encode Tile into .w, to support LWC?
+	FVector3f BoundsOrigin3f = FVector3f(BoundsOrigin);
+	FVector3f BoundsBoxExtent3f = FVector3f(BoundsBoxExtent);
+
+	FVector4f* RESTRICT PrimitiveData = (FVector4f*)OcclusionBatch.VertexAllocation.Buffer;
+	PrimitiveData[0] = FVector4f(BoundsOrigin3f, 0.f);
+	PrimitiveData[1] = FVector4f(BoundsBoxExtent3f, 0.f);
+
+	OcclusionBatch.VertexAllocation.Buffer += PrimitiveStride;
+	OcclusionBatch.NumBatchedPrimitives++;
+
+	FOcclusionBuffer& OcclusionBuffer = OcclusionBuffers[CurrentBufferIndex];
+	OcclusionBuffer.BatchedPrimitives.Add(PrimitiveId);
+}
+
+class FOcclusionFeedbackVS : public FGlobalShader
+{
+	DECLARE_GLOBAL_SHADER(FOcclusionFeedbackVS);
+	SHADER_USE_PARAMETER_STRUCT(FOcclusionFeedbackVS, FGlobalShader);
+
+public:
+	static bool ShouldCompilePermutation(const FGlobalShaderPermutationParameters& Parameters)
+	{
+		return IsMobilePlatform(Parameters.Platform);
+	}
+
+	BEGIN_SHADER_PARAMETER_STRUCT(FParameters, )
+		SHADER_PARAMETER_STRUCT_REF(FViewUniformShaderParameters, View)
+		SHADER_PARAMETER(uint32, StartIndex)
+	END_SHADER_PARAMETER_STRUCT()
+};
+
+class FOcclusionFeedbackPS : public FGlobalShader
+{
+	DECLARE_GLOBAL_SHADER(FOcclusionFeedbackPS);
+	SHADER_USE_PARAMETER_STRUCT(FOcclusionFeedbackPS, FGlobalShader);
+
+public:
+	static bool ShouldCompilePermutation(const FGlobalShaderPermutationParameters& Parameters)
+	{
+		return IsMobilePlatform(Parameters.Platform);
+	}
+
+	class FPixelDiscard : SHADER_PERMUTATION_BOOL("PERMUTATION_PIXEL_DISCARD");
+	using FPermutationDomain = TShaderPermutationDomain<FPixelDiscard>;
+
+	using FParameters = FEmptyShaderParameters;
+};
+
+IMPLEMENT_GLOBAL_SHADER(FOcclusionFeedbackVS,"/Engine/Private/OcclusionFeedbackShaders.usf","MainVS",SF_Vertex);
+IMPLEMENT_GLOBAL_SHADER(FOcclusionFeedbackPS,"/Engine/Private/OcclusionFeedbackShaders.usf","MainPS",SF_Pixel);
+
+int32 GOcclusionFeeback_Blending = 1;
+static FAutoConsoleVariableRef CVarOcclusionFeeback_Blending(
+	TEXT("r.OcclusionFeeback.Blending"),
+	GOcclusionFeeback_Blending,
+	TEXT("0: Opaque \n")
+	TEXT("1: Enable blending  (default) \n")
+	TEXT("2: Enable pixel discard"),
+	ECVF_RenderThreadSafe
+);
+
+void FOcclusionFeedback::SubmitOcclusionDraws(FRHICommandList& RHICmdList, FViewInfo& View)
+{
+	if (BatchOcclusionQueries.Num() == 0)
+	{
+		return;
+	}
+
+	FRHIBuffer* CubeIndexBuffer = GetUnitCubeIndexBuffer();
+	FRHIBuffer* CubeVertexBuffer = GetUnitCubeVertexBuffer();
+
+	const FIntRect ViewRect = View.ViewRect;
+	RHICmdList.SetViewport(ViewRect.Min.X, ViewRect.Min.Y, 0.0f, ViewRect.Max.X, ViewRect.Max.Y, 1.0f);
+
+	TShaderMapRef<FOcclusionFeedbackVS> VertexShader(View.ShaderMap);
+
+	FOcclusionFeedbackPS::FPermutationDomain PsPermutationVector;
+	PsPermutationVector.Set<FOcclusionFeedbackPS::FPixelDiscard>(GOcclusionFeeback_Blending == 2);
+	TShaderMapRef<FOcclusionFeedbackPS> PixelShader(View.ShaderMap, PsPermutationVector);
+
+	FGraphicsPipelineStateInitializer GraphicsPSOInit;
+	RHICmdList.ApplyCachedRenderTargets(GraphicsPSOInit);
+	GraphicsPSOInit.PrimitiveType = PT_TriangleList;
+
+	if (GOcclusionFeeback_Blending == 1)
+	{
+		GraphicsPSOInit.BlendState = TStaticBlendState<
+			CW_ALPHA, BO_Add, BF_One, BF_One, BO_Add, BF_Zero, BF_One,
+			CW_NONE, BO_Add, BF_One, BF_Zero, BO_Add, BF_One, BF_Zero,
+			CW_NONE, BO_Add, BF_One, BF_Zero, BO_Add, BF_One, BF_Zero,
+			CW_NONE, BO_Add, BF_One, BF_Zero, BO_Add, BF_One, BF_Zero,
+			CW_NONE, BO_Add, BF_One, BF_Zero, BO_Add, BF_One, BF_Zero,
+			CW_NONE, BO_Add, BF_One, BF_Zero, BO_Add, BF_One, BF_Zero,
+			CW_NONE, BO_Add, BF_One, BF_Zero, BO_Add, BF_One, BF_Zero,
+			CW_NONE, BO_Add, BF_One, BF_Zero, BO_Add, BF_One, BF_Zero>::GetRHI();
+	}
+	else
+	{
+		GraphicsPSOInit.BlendState = TStaticBlendStateWriteMask<CW_NONE, CW_NONE, CW_NONE, CW_NONE, CW_NONE, CW_NONE, CW_NONE, CW_NONE>::GetRHI();
+	}
+
+	// Depth tests, no depth writes, no color writes, opaque
+	GraphicsPSOInit.DepthStencilState = TStaticDepthStencilState<false, CF_DepthNearOrEqual>::GetRHI();
+	GraphicsPSOInit.BoundShaderState.VertexDeclarationRHI = OcclusionVertexDeclarationRHI;
+	GraphicsPSOInit.BoundShaderState.VertexShaderRHI = VertexShader.GetVertexShader();
+	GraphicsPSOInit.BoundShaderState.PixelShaderRHI = PixelShader.GetPixelShader();
+	// We only need to render the front-faces of the culling geometry (this halves the amount of pixels we touch)
+	GraphicsPSOInit.RasterizerState = View.bReverseCulling ? TStaticRasterizerState<FM_Solid, CM_CW>::GetRHI() : TStaticRasterizerState<FM_Solid, CM_CCW>::GetRHI();
+	SetGraphicsPipelineState(RHICmdList, GraphicsPSOInit, 0);
+
+	RHICmdList.SetStreamSource(0, CubeVertexBuffer, 0);
+
+	uint32 PrimitiveStartIndex = 0u;
+	// Draw the batches.
+	for (int32 BatchIndex = 0, NumBatches = BatchOcclusionQueries.Num(); BatchIndex < NumBatches; BatchIndex++)
+	{
+		FOcclusionBatch& Batch = BatchOcclusionQueries[BatchIndex];
+		
+		FOcclusionFeedbackVS::FParameters VSParams;
+		VSParams.View = View.ViewUniformBuffer;
+		VSParams.StartIndex = PrimitiveStartIndex;
+		SetShaderParameters(RHICmdList, VertexShader, VertexShader.GetVertexShader(), VSParams);
+				
+		FRHIBuffer* PrimitiveBufferRHI = Batch.VertexAllocation.VertexBuffer->VertexBufferRHI;
+		uint32 VertexBufferOffset = Batch.VertexAllocation.VertexOffset;
+		RHICmdList.SetStreamSource(1, PrimitiveBufferRHI, VertexBufferOffset);
+
+		RHICmdList.DrawIndexedPrimitive(
+			CubeIndexBuffer,
+			/*BaseVertexIndex=*/ 0,
+			/*MinIndex=*/ 0,
+			/*NumVertices=*/ 8,
+			/*StartIndex=*/ 0,
+			/*NumPrimitives=*/ 12,
+			/*NumInstances=*/ Batch.NumBatchedPrimitives
+		);
+
+		PrimitiveStartIndex += Batch.NumBatchedPrimitives;
+	}
+	INC_DWORD_STAT_BY(STAT_OcclusionQueries,BatchOcclusionQueries.Num());
+
+	// Reset the batch state.
+	BatchOcclusionQueries.Empty(BatchOcclusionQueries.Num());
+}
diff --git a/Engine/Source/Runtime/Renderer/Private/ScenePrivate.h b/Engine/Source/Runtime/Renderer/Private/ScenePrivate.h
--- a/Engine/Source/Runtime/Renderer/Private/ScenePrivate.h
+++ b/Engine/Source/Runtime/Renderer/Private/ScenePrivate.h
@@ -714,6 +714,7 @@ public:
 	FFrameBasedOcclusionQueryPool PrimitiveOcclusionQueryPool;
 
 	FHZBOcclusionTester HZBOcclusionTests;
+	FOcclusionFeedback OcclusionFeedback;
 
 	FPersistentSkyAtmosphereData PersistentSkyAtmosphereData;
 
@@ -1336,6 +1337,7 @@ public:
 	{
 		HZBOcclusionTests.ReleaseDynamicRHI();
 		EyeAdaptationManager.SafeRelease();
+		OcclusionFeedback.ReleaseResource();
 		bValidEyeAdaptationTexture = false;
 		bValidEyeAdaptationBuffer = false;
 	}
diff --git a/Engine/Source/Runtime/Renderer/Private/SceneRendering.h b/Engine/Source/Runtime/Renderer/Private/SceneRendering.h
--- a/Engine/Source/Runtime/Renderer/Private/SceneRendering.h
+++ b/Engine/Source/Runtime/Renderer/Private/SceneRendering.h
@@ -498,6 +498,61 @@ private:
 	uint32 ValidFrameNumber;
 };
 
+class FOcclusionFeedback : public FRenderResource
+{
+public:
+	FOcclusionFeedback();
+	~FOcclusionFeedback();
+
+	// FRenderResource interface
+	virtual void InitDynamicRHI() override;
+	virtual void ReleaseDynamicRHI() override;
+
+	void AddPrimitive(FPrimitiveComponentId PrimitiveId, const FVector& BoundsOrigin, const FVector& BoundsBoxExtent, FGlobalDynamicVertexBuffer& DynamicVertexBuffer);
+
+	void BeginOcclusionScope(FRDGBuilder& GraphBuilder);
+	void EndOcclusionScope(FRDGBuilder& GraphBuilder);
+
+	/** Renders the current batch and resets the batch state. */
+	void SubmitOcclusionDraws(FRHICommandList& RHICmdList, FViewInfo& View);
+	void ReadbackResults(FRHICommandList& RHICmdList);
+
+	inline FRDGBuffer* GetGPUFeedbackBuffer() const
+	{
+		return GPUFeedbackBuffer;
+	}
+
+	inline bool IsOccluded(FPrimitiveComponentId Id) const
+	{
+		return LatestOcclusionResults.Contains(Id);
+	}
+
+private:
+	struct FOcclusionBatch
+	{
+		FGlobalDynamicVertexBuffer::FAllocation VertexAllocation;
+		uint32 NumBatchedPrimitives;
+	};
+
+	/** The pending batches. */
+	TArray<FOcclusionBatch, TInlineAllocator<3>> BatchOcclusionQueries;
+
+	FRDGBuffer* GPUFeedbackBuffer{};
+
+	struct FOcclusionBuffer
+	{
+		TArray<FPrimitiveComponentId> BatchedPrimitives;
+		FRHIGPUBufferReadback* ReadbackBuffer = nullptr;
+	};
+
+	FOcclusionBuffer OcclusionBuffers[3];
+	uint32 CurrentBufferIndex;
+	TSet<FPrimitiveComponentId> LatestOcclusionResults;
+
+	//
+	FVertexDeclarationRHIRef OcclusionVertexDeclarationRHI;
+};
+
 DECLARE_STATS_GROUP(TEXT("Parallel Command List Markers"), STATGROUP_ParallelCommandListMarkers, STATCAT_Advanced);
 
 /** Helper class to marshal data from your RDG pass into the parallel command list set. */
diff --git a/Engine/Source/Runtime/Renderer/Private/SceneVisibility.cpp b/Engine/Source/Runtime/Renderer/Private/SceneVisibility.cpp
--- a/Engine/Source/Runtime/Renderer/Private/SceneVisibility.cpp
+++ b/Engine/Source/Runtime/Renderer/Private/SceneVisibility.cpp
@@ -112,6 +112,14 @@ static FAutoConsoleVariableRef CVarHZBOcclusion(
 	ECVF_RenderThreadSafe
 	);
 
+int32 GOcclusionFeedback_Enable = 0;
+static FAutoConsoleVariableRef CVarOcclusionFeedback_Enable(
+	TEXT("r.OcclusionFeedback.Enable"),
+	GOcclusionFeedback_Enable,
+	TEXT("Whether to enable occlusion system based on a rendering feedback. Currently works only with a mobile rendering\n"),
+	ECVF_RenderThreadSafe
+);
+
 static int32 GVisualizeOccludedPrimitives = 0;
 static FAutoConsoleVariableRef CVarVisualizeOccludedPrimitives(
 	TEXT("r.VisualizeOccludedPrimitives"),
@@ -1163,6 +1171,8 @@ static void FetchVisibilityForPrimitives_Range(FVisForPrimParams& Params, FGloba
 	const float CurrentRealTime = View.Family->Time.GetRealTimeSeconds();
 	uint32 OcclusionFrameCounter = ViewState->OcclusionFrameCounter;
 	FHZBOcclusionTester& HZBOcclusionTests = ViewState->HZBOcclusionTests;
+	FOcclusionFeedback& OcclusionFeedback = ViewState->OcclusionFeedback;
+	const bool bUseOcclusionFeedback = bSingleThreaded && OcclusionFeedback.IsInitialized();
 
 	int32 ReadBackLagTolerance = NumBufferedFrames;
 
@@ -1368,7 +1378,12 @@ static void FetchVisibilityForPrimitives_Range(FVisForPrimParams& Params, FGloba
 					}
 					else if (bCanBeOccluded)
 					{
-						if (bHZBOcclusion)
+						if (bUseOcclusionFeedback)
+						{
+							bIsOccluded = OcclusionFeedback.IsOccluded(PrimitiveId);
+							bOcclusionStateIsDefinite = true;
+						}
+						else if (bHZBOcclusion)
 						{
 							if (HZBOcclusionTests.IsValidFrame(PrimitiveOcclusionHistory->LastTestFrameNumber))
 							{
@@ -1519,7 +1534,14 @@ static void FetchVisibilityForPrimitives_Range(FVisForPrimParams& Params, FGloba
 						if (bAllowBoundsTest)
 						{
 							PrimitiveOcclusionHistory->LastTestFrameNumber = OcclusionFrameCounter;
-							if (bHZBOcclusion)
+
+							if (bUseOcclusionFeedback)
+							{
+								const FVector BoundOrigin = OcclusionBounds.Origin + View.ViewMatrices.GetPreViewTranslation();
+								const FVector BoundExtent = OcclusionBounds.BoxExtent;
+								OcclusionFeedback.AddPrimitive(PrimitiveId, BoundOrigin, BoundExtent, *DynamicVertexBufferIfSingleThreaded);
+							}
+							else if (bHZBOcclusion)
 							{
 								// Always run
 								if (bSingleThreaded)
@@ -2024,7 +2046,12 @@ static int32 OcclusionCull(FRHICommandListImmediate& RHICmdList, const FScene* S
 		bSubmitQueries = bSubmitQueries && !ViewState->HasViewParent() && !ViewState->bIsFrozen;
 #endif
 
-		if( bHZBOcclusion )
+		if (ViewState->OcclusionFeedback.IsInitialized())
+		{
+			QUICK_SCOPE_CYCLE_COUNTER(STAT_OcclusionFeedback_ReadbackResults);
+			ViewState->OcclusionFeedback.ReadbackResults(RHICmdList);
+		}
+		else if( bHZBOcclusion )
 		{
 			QUICK_SCOPE_CYCLE_COUNTER(STAT_MapHZBResults);
 			check(!ViewState->HZBOcclusionTests.IsValidFrame(ViewState->OcclusionFrameCounter));
@@ -2062,6 +2089,19 @@ static int32 OcclusionCull(FRHICommandListImmediate& RHICmdList, const FScene* S
 				ViewState->HZBOcclusionTests.SetValidFrameNumber(ViewState->OcclusionFrameCounter);
 			}
 		}
+
+		if (View.FeatureLevel == ERHIFeatureLevel::ES3_1)
+		{
+			// Initialize/release OcclusionFeedback system on demand
+			if (GOcclusionFeedback_Enable == 0 && ViewState->OcclusionFeedback.IsInitialized())
+			{
+				ViewState->OcclusionFeedback.ReleaseResource();
+			}
+			else if (GOcclusionFeedback_Enable != 0 && !ViewState->OcclusionFeedback.IsInitialized())
+			{
+				ViewState->OcclusionFeedback.InitResource();
+			}
+		}
 	}
 	RHICmdList.SetCurrentStat(GET_STATID(STAT_CLMM_AfterOcclusionReadback));
 	return NumOccludedPrimitives;
