Implemented support for rotation in SparseVolumeTextureViewer.

#rnx

diff --git a/Engine/Shaders/Private/SparseVolumeTexture/VisualizeSparseVolumeTexture.usf b/Engine/Shaders/Private/SparseVolumeTexture/VisualizeSparseVolumeTexture.usf
--- a/Engine/Shaders/Private/SparseVolumeTexture/VisualizeSparseVolumeTexture.usf
+++ b/Engine/Shaders/Private/SparseVolumeTexture/VisualizeSparseVolumeTexture.usf
@@ -42,6 +42,9 @@ bool Slabs(float3 p0, float3 p1, float3 rayOrigin, float3 invRaydir, out float o
 
 float3 VolumeBoundMinWorld;
 float3 VolumeBoundMaxWorld;
+float4 WorldToLocal0;
+float4 WorldToLocal1;
+float4 WorldToLocal2;
 
 
 Texture3D<float4> SparseVolumeTexture;
@@ -65,15 +68,26 @@ void VisualizeSparseVolumeTexturePS(
 	const float Depth = 10000.0f;
 	float4 TranslatedWorldPos = mul(float4(ScreenPosition * Depth, Depth, 1), View.ScreenToTranslatedWorld);
 
-	const float3 RayDir = normalize(TranslatedWorldPos.xyz - View.TranslatedWorldCameraOrigin);
-	const float3 RayOrig = LWCHackToFloat(ResolvedView.WorldCameraOrigin);
+	const float3 RayDirWorld = normalize(TranslatedWorldPos.xyz - View.TranslatedWorldCameraOrigin);
+	const float3 RayOrigWorld = LWCHackToFloat(ResolvedView.WorldCameraOrigin);
+
+	// Transform into local [-1.0f, 1.0f] space
+	float3 RayDir;
+	RayDir.x = dot(WorldToLocal0.xyz, RayDirWorld);
+	RayDir.y = dot(WorldToLocal1.xyz, RayDirWorld);
+	RayDir.z = dot(WorldToLocal2.xyz, RayDirWorld);
+
+	float3 RayOrig;
+	RayOrig.x = dot(WorldToLocal0, float4(RayOrigWorld, 1.0f));
+	RayOrig.y = dot(WorldToLocal1, float4(RayOrigWorld, 1.0f));
+	RayOrig.z = dot(WorldToLocal2, float4(RayOrigWorld, 1.0f));
 
 	OutLuminanceAlpha = float4(0.0, 0.0, 0.0, 1.0);
 
 	float TMin = 0.0f;
 	float TMax = 0.0f;
 	float Transmittance = 1.0f;
-	if (Slabs(VolumeBoundMinWorld, VolumeBoundMaxWorld, RayOrig, 1.0 / RayDir, TMin, TMax) && TMax > 0.0f)
+	if (Slabs(-1.0f, 1.0f, RayOrig, 1.0 / RayDir, TMin, TMax) && TMax > 0.0f)
 	{
 		Transmittance = 0.5; // show the bounding box
 
@@ -84,7 +98,7 @@ void VisualizeSparseVolumeTexturePS(
 
 		float3 VolumeResolution = SparseVolumeTexturePageTableResolution * SPARSE_VOLUME_TILE_RES;
 
-		const float3 StartUVs = (StartPos - VolumeBoundMinWorld) / (VolumeBoundMaxWorld - VolumeBoundMinWorld);
+		const float3 StartUVs = StartPos * 0.5f + 0.5f;
 		float3 P = StartUVs * VolumeResolution;
 
 		// Force round to volume
diff --git a/Engine/Source/Runtime/Renderer/Private/SparseVolumeTexture/SparseVolumeTextureViewerComponent.cpp b/Engine/Source/Runtime/Renderer/Private/SparseVolumeTexture/SparseVolumeTextureViewerComponent.cpp
--- a/Engine/Source/Runtime/Renderer/Private/SparseVolumeTexture/SparseVolumeTextureViewerComponent.cpp
+++ b/Engine/Source/Runtime/Renderer/Private/SparseVolumeTexture/SparseVolumeTextureViewerComponent.cpp
@@ -161,11 +161,15 @@ void USparseVolumeTextureViewerComponent::SendRenderTransformCommand()
 		float MaxScaling = FMath::Max(Scale3D.X, FMath::Max(Scale3D.Y, Scale3D.Z));
 		ToWorldMat.SetScale3D(FVector(MaxScaling, MaxScaling, MaxScaling));		// Keep max scaling
 		FBoxSphereBounds ComponentBounds = CalcBounds(ToWorldMat);
+
+		FMatrix44f ToLocalMat = FMatrix44f(FTranslationMatrix(-ToWorldMat.GetTranslation()) * FRotationMatrix(FRotator(GetComponentToWorld().GetRotation().Inverse())) * FScaleMatrix(ComponentBounds.BoxExtent.Reciprocal()));
+
 		FSparseVolumeTextureViewerSceneProxy* SVTViewerSceneProxy = SparseVolumeTextureViewerSceneProxy;
 		ENQUEUE_RENDER_COMMAND(FUpdateSparseVolumeTextureViewerProxyTransformCommand)(
-			[SVTViewerSceneProxy, ComponentBounds](FRHICommandList& RHICmdList)
+			[SVTViewerSceneProxy, ComponentBounds, ToLocalMat](FRHICommandList& RHICmdList)
 			{
 				SVTViewerSceneProxy->VolumeWorldBounds = ComponentBounds;
+				SVTViewerSceneProxy->WorldToLocal = ToLocalMat;
 			});
 	}
 }
diff --git a/Engine/Source/Runtime/Renderer/Private/SparseVolumeTexture/SparseVolumeTextureViewerRendering.cpp b/Engine/Source/Runtime/Renderer/Private/SparseVolumeTexture/SparseVolumeTextureViewerRendering.cpp
--- a/Engine/Source/Runtime/Renderer/Private/SparseVolumeTexture/SparseVolumeTextureViewerRendering.cpp
+++ b/Engine/Source/Runtime/Renderer/Private/SparseVolumeTexture/SparseVolumeTextureViewerRendering.cpp
@@ -55,6 +55,9 @@ class FVisualizeSparseVolumeTexturePS : public FGlobalShader
 	BEGIN_SHADER_PARAMETER_STRUCT(FParameters, )
 		SHADER_PARAMETER_STRUCT_REF(FViewUniformShaderParameters, ViewUniformBuffer)
 		RENDER_TARGET_BINDING_SLOTS()
+		SHADER_PARAMETER(FVector4f, WorldToLocal0)
+		SHADER_PARAMETER(FVector4f, WorldToLocal1)
+		SHADER_PARAMETER(FVector4f, WorldToLocal2)
 		SHADER_PARAMETER(FVector3f, VolumeBoundMinWorld)
 		SHADER_PARAMETER(FVector3f, VolumeBoundMaxWorld)
 		SHADER_PARAMETER(FVector3f, SparseVolumeTextureResolution)
@@ -108,6 +111,7 @@ void AddSparseVolumeTextureViewerRenderPass(FRDGBuilder& GraphBuilder, FSceneRen
 		for (auto& SVTProxy : Scene->SparseVolumeTextureViewers)
 		{
 			const FBoxSphereBounds& SVTProxyBound = SVTProxy->VolumeWorldBounds;
+			const FMatrix44f& WorldToLocal = SVTProxy->WorldToLocal;
 			
 			FGlobalShaderMap* GlobalShaderMap = GetGlobalShaderMap(GMaxRHIFeatureLevel);
 			FVisualizeSparseVolumeTextureVS::FPermutationDomain VsPermutationVector;
@@ -118,6 +122,9 @@ void AddSparseVolumeTextureViewerRenderPass(FRDGBuilder& GraphBuilder, FSceneRen
 			FVisualizeSparseVolumeTexturePS::FParameters* PsPassParameters = GraphBuilder.AllocParameters<FVisualizeSparseVolumeTexturePS::FParameters>();
 			PsPassParameters->ViewUniformBuffer = View.ViewUniformBuffer;
 			PsPassParameters->RenderTargets[0] = FRenderTargetBinding(SceneTextures.Color.Target, ERenderTargetLoadAction::ELoad);
+			PsPassParameters->WorldToLocal0 = FVector4f(WorldToLocal.M[0][0], WorldToLocal.M[1][0], WorldToLocal.M[2][0], WorldToLocal.M[3][0]);
+			PsPassParameters->WorldToLocal1 = FVector4f(WorldToLocal.M[0][1], WorldToLocal.M[1][1], WorldToLocal.M[2][1], WorldToLocal.M[3][1]);
+			PsPassParameters->WorldToLocal2 = FVector4f(WorldToLocal.M[0][2], WorldToLocal.M[1][2], WorldToLocal.M[2][2], WorldToLocal.M[3][2]);
 			PsPassParameters->VolumeBoundMinWorld = FVector3f(SVTProxyBound.Origin - SVTProxyBound.BoxExtent);
 			PsPassParameters->VolumeBoundMaxWorld = FVector3f(SVTProxyBound.Origin + SVTProxyBound.BoxExtent);
 			PsPassParameters->SparseVolumeTextureResolution = FVector3f::OneVector;
diff --git a/Engine/Source/Runtime/Renderer/Private/SparseVolumeTexture/SparseVolumeTextureViewerSceneProxy.h b/Engine/Source/Runtime/Renderer/Private/SparseVolumeTexture/SparseVolumeTextureViewerSceneProxy.h
--- a/Engine/Source/Runtime/Renderer/Private/SparseVolumeTexture/SparseVolumeTextureViewerSceneProxy.h
+++ b/Engine/Source/Runtime/Renderer/Private/SparseVolumeTexture/SparseVolumeTextureViewerSceneProxy.h
@@ -19,6 +19,7 @@ public:
 
 	const FSparseVolumeTextureSceneProxy* SparseVolumeTextureSceneProxy = nullptr;
 	FBoxSphereBounds VolumeWorldBounds;
+	FMatrix44f WorldToLocal;
 
 protected:
 
