Change SubmitShaderPrintData() to take a FSceneInterface, and only draw the data in views from a matching Scene.

diff --git a/Engine/Plugins/Animation/DeformerGraph/Source/OptimusCore/Private/DataInterfaces/OptimusDataInterfaceDebugDraw.cpp b/Engine/Plugins/Animation/DeformerGraph/Source/OptimusCore/Private/DataInterfaces/OptimusDataInterfaceDebugDraw.cpp
--- a/Engine/Plugins/Animation/DeformerGraph/Source/OptimusCore/Private/DataInterfaces/OptimusDataInterfaceDebugDraw.cpp
+++ b/Engine/Plugins/Animation/DeformerGraph/Source/OptimusCore/Private/DataInterfaces/OptimusDataInterfaceDebugDraw.cpp
@@ -116,6 +116,8 @@ FComputeDataProviderRenderProxy* UOptimusDebugDrawDataProvider::GetRenderProxy()
 FOptimusDebugDrawDataProviderProxy::FOptimusDebugDrawDataProviderProxy(UPrimitiveComponent* InPrimitiveComponent, FOptimusDebugDrawParameters const& InDebugDrawParameters)
 	: Setup(FIntRect(0, 0, 1920, 1080))
 {
+	Scene = InPrimitiveComponent->GetScene();
+
 	// Split LocalToWorld into a pre-translation and transform for large world coordinate support.
 	FMatrix RenderMatrix = InPrimitiveComponent->GetRenderMatrix();
 	FVector PreViewTranslation = -RenderMatrix.GetOrigin();
@@ -151,7 +153,7 @@ void FOptimusDebugDrawDataProviderProxy::AllocateResources(FRDGBuilder& GraphBui
 	{
 		// Enqueue for display at next view render.
 		FFrozenShaderPrintData FrozenShaderPrintData = ShaderPrint::FreezeShaderPrintData(GraphBuilder, ShaderPrintData);
-		ShaderPrint::SubmitShaderPrintData(FrozenShaderPrintData);
+		ShaderPrint::SubmitShaderPrintData(FrozenShaderPrintData, Scene);
 	}
 }
 
diff --git a/Engine/Plugins/Animation/DeformerGraph/Source/OptimusCore/Private/DataInterfaces/OptimusDataInterfaceDebugDraw.h b/Engine/Plugins/Animation/DeformerGraph/Source/OptimusCore/Private/DataInterfaces/OptimusDataInterfaceDebugDraw.h
--- a/Engine/Plugins/Animation/DeformerGraph/Source/OptimusCore/Private/DataInterfaces/OptimusDataInterfaceDebugDraw.h
+++ b/Engine/Plugins/Animation/DeformerGraph/Source/OptimusCore/Private/DataInterfaces/OptimusDataInterfaceDebugDraw.h
@@ -94,6 +94,7 @@ public:
 	//~ End FComputeDataProviderRenderProxy Interface
 
 private:
+	FSceneInterface const* Scene;
 	FMatrix44f LocalToWorld;
 	ShaderPrint::FShaderPrintSetup Setup;
 	ShaderPrint::FShaderPrintCommonParameters ConfigParameters;
diff --git a/Engine/Source/Runtime/Renderer/Private/DeferredShadingRenderer.cpp b/Engine/Source/Runtime/Renderer/Private/DeferredShadingRenderer.cpp
--- a/Engine/Source/Runtime/Renderer/Private/DeferredShadingRenderer.cpp
+++ b/Engine/Source/Runtime/Renderer/Private/DeferredShadingRenderer.cpp
@@ -2206,23 +2206,21 @@ void FDeferredShadingSceneRenderer::Render(FRDGBuilder& GraphBuilder)
 
 	GPU_MESSAGE_SCOPE(GraphBuilder);
 
+	ShaderPrint::BeginViews(GraphBuilder, Views);
+
+	ON_SCOPE_EXIT
+	{
+		ShaderPrint::EndViews(Views);
+	};
+
 	for (int32 ViewIndex = 0; ViewIndex < Views.Num(); ViewIndex++)
 	{
 		FViewInfo& View = Views[ViewIndex];
 		RDG_GPU_MASK_SCOPE(GraphBuilder, View.GPUMask);
 
-		ShaderPrint::BeginView(GraphBuilder, View);
 		ShadingEnergyConservation::Init(GraphBuilder, View);
 	}
 	
-	ON_SCOPE_EXIT
-	{
-		for (FViewInfo& View : Views)
-		{
-			ShaderPrint::EndView(View);
-		}
-	};
-
 	Scene->UpdateAllPrimitiveSceneInfos(GraphBuilder, true);
 
 #if RHI_RAYTRACING
diff --git a/Engine/Source/Runtime/Renderer/Private/SceneHitProxyRendering.cpp b/Engine/Source/Runtime/Renderer/Private/SceneHitProxyRendering.cpp
--- a/Engine/Source/Runtime/Renderer/Private/SceneHitProxyRendering.cpp
+++ b/Engine/Source/Runtime/Renderer/Private/SceneHitProxyRendering.cpp
@@ -647,10 +647,7 @@ void FDeferredShadingSceneRenderer::RenderHitProxies(FRDGBuilder& GraphBuilder)
 		}
 	}
 
-	for (int32 ViewIndex = 0; ViewIndex < Views.Num(); ViewIndex++)
-	{
-		ShaderPrint::BeginView(GraphBuilder, Views[ViewIndex]);
-	}
+	ShaderPrint::BeginViews(GraphBuilder, Views);
 
 	{
 		FRDGExternalAccessQueue ExternalAccessQueue;
@@ -753,10 +750,7 @@ void FDeferredShadingSceneRenderer::RenderHitProxies(FRDGBuilder& GraphBuilder)
 
 	::DoRenderHitProxies(GraphBuilder, this, HitProxyTexture, HitProxyDepthTexture, NaniteRasterResults, InstanceCullingManager);
 
-	for (int32 ViewIndex = 0; ViewIndex < Views.Num(); ViewIndex++)
-	{
-		ShaderPrint::EndView(Views[ViewIndex]);
-	}
+	ShaderPrint::EndViews(Views);
 
 	GEngine->GetPostRenderDelegateEx().Broadcast(GraphBuilder);
 
diff --git a/Engine/Source/Runtime/Renderer/Private/ShaderPrint.cpp b/Engine/Source/Runtime/Renderer/Private/ShaderPrint.cpp
--- a/Engine/Source/Runtime/Renderer/Private/ShaderPrint.cpp
+++ b/Engine/Source/Runtime/Renderer/Private/ShaderPrint.cpp
@@ -83,8 +83,15 @@ namespace ShaderPrint
 	static uint32 GLineRequestCount = 0;
 	static uint32 GTriangleRequestCount = 0;
 	static FViewInfo* GDefaultView = nullptr;
-	static TArray<FFrozenShaderPrintData> GShaderPrintDataToRender;
-
+	
+	struct FQueuedRenderItem
+	{
+		FFrozenShaderPrintData Payload;
+		FSceneInterface const* Scene = nullptr;
+		uint32 FrameForGC = 0;
+	};
+	static TArray<FQueuedRenderItem> GQueuedRenderItems;
+	
 	//////////////////////////////////////////////////////////////////////////////////////////////////
 	// Struct & Functions
 
@@ -249,9 +256,17 @@ namespace ShaderPrint
 		GTriangleRequestCount += InCount;
 	}
 
+	void SubmitShaderPrintData(FFrozenShaderPrintData& InData, FSceneInterface const* InScene)
+	{
+		// Queue with a frame number so that we can garbage collect if no matching view ever renders.
+		GQueuedRenderItems.Add({InData, InScene, GFrameNumberRenderThread + 10});
+	}
+
 	void SubmitShaderPrintData(FFrozenShaderPrintData& InData)
 	{
-		GShaderPrintDataToRender.Add(InData);
+		PRAGMA_DISABLE_DEPRECATION_WARNINGS
+		SubmitShaderPrintData(InData, nullptr);
+		PRAGMA_ENABLE_DEPRECATION_WARNINGS
 	}
 
 	//////////////////////////////////////////////////////////////////////////////////////////////////
@@ -681,18 +696,27 @@ namespace ShaderPrint
 			return;
 		}
 
-		// Invalid to call begin twice for the same view.
-		ensure(GDefaultView != &View);
-		if (GDefaultView == nullptr)
-		{
-			GDefaultView = &View;
-		}
-
 		// Create the render data and store on the view.
 		FShaderPrintSetup ShaderPrintSetup(View);
 		View.ShaderPrintData = CreateShaderPrintData(GraphBuilder, ShaderPrintSetup, View.ViewState);
+	}
 
-		// Reset counter which is read on the next BeginView().
+	void BeginViews(FRDGBuilder& GraphBuilder, TArrayView<FViewInfo> Views)
+	{
+		ensure(GDefaultView == nullptr);
+		if (Views.Num() > 0)
+		{
+			GDefaultView = &Views[0];
+		}
+
+		for (int32 ViewIndex = 0; ViewIndex < Views.Num(); ViewIndex++)
+		{
+			FViewInfo& View = Views[ViewIndex];
+			RDG_GPU_MASK_SCOPE(GraphBuilder, View.GPUMask);
+			BeginView(GraphBuilder, View);
+		}
+
+		// Reset counters which are read on the next BeginViews().
 		GCharacterRequestCount = 0;
 		GWidgetRequestCount = 0;
 		GLineRequestCount = 0;
@@ -976,17 +1000,32 @@ namespace ShaderPrint
 		InternalDrawView(GraphBuilder, View, View.ShaderPrintData, OutputTexture, DepthTexture);
 
 		// Draw any externally enqueued shader print data.
-		for (FFrozenShaderPrintData& ShaderPrintDataToRender : GShaderPrintDataToRender)
+		FSceneInterface const* Scene = View.Family != nullptr ? View.Family->Scene : nullptr;
+		for (FQueuedRenderItem& ShaderPrintDataToRender : GQueuedRenderItems)
 		{
-			FShaderPrintData ShaderPrintData = UnFreezeShaderPrintData(GraphBuilder, ShaderPrintDataToRender);
-			InternalDrawView(GraphBuilder, View, ShaderPrintData, OutputTexture, DepthTexture);
+			if (ShaderPrintDataToRender.Scene == nullptr || ShaderPrintDataToRender.Scene == Scene)
+			{
+				FShaderPrintData ShaderPrintData = UnFreezeShaderPrintData(GraphBuilder, ShaderPrintDataToRender.Payload);
+				InternalDrawView(GraphBuilder, View, ShaderPrintData, OutputTexture, DepthTexture);
+			}
 		}
 	}
 
-	void EndView(FViewInfo& View)
+	void EndViews(TArrayView<FViewInfo> Views)
 	{
-		View.ShaderPrintData = FShaderPrintData();
+		// Clear the shader print data owned by the views.
+		for (int32 ViewIndex = 0; ViewIndex < Views.Num(); ViewIndex++)
+		{
+			Views[ViewIndex].ShaderPrintData = FShaderPrintData();
+		}
+
+		// Remove all externally enqueud shader print data for the matching Scene, since we will have drawn it in the calls to DrawView().
+		if (Views.Num())
+		{
+			FSceneInterface const* Scene = Views[0].Family != nullptr ? Views[0].Family->Scene : nullptr;
+			GQueuedRenderItems.RemoveAll([Scene](FQueuedRenderItem& Data) { return Data.Scene == nullptr || Data.Scene == Scene || Data.FrameForGC < GFrameNumberRenderThread; });
+		}
+
 		GDefaultView = nullptr;
-		GShaderPrintDataToRender.Reset();
 	}
 }
\ No newline at end of file
diff --git a/Engine/Source/Runtime/Renderer/Private/ShaderPrint.h b/Engine/Source/Runtime/Renderer/Private/ShaderPrint.h
--- a/Engine/Source/Runtime/Renderer/Private/ShaderPrint.h
+++ b/Engine/Source/Runtime/Renderer/Private/ShaderPrint.h
@@ -1,9 +1,9 @@
 // Copyright Epic Games, Inc. All Rights Reserved.
 
 // Use ShaderPrint to debug print from any shader type.
-// Call BeginView() at the start of a view to init the debug output buffer.
+// Call BeginViews() at the start of a view group to init the debug output buffer.
 // Call DrawView() to composite the debug data to the final render target.
-// Call EndView() at the end of the view to tidy up.
+// Call EndViews() at the end of the view group to tidy up.
 
 // The shader HLSL code should include the ShaderPrintCommon.ush and use the ShaderPrint*() functions.
 // The shader C++ code needs to bind the buffers required to capture the debug print. See ShaderPrintParameters.h for this. 
@@ -19,9 +19,9 @@ struct FScreenPassTexture;
 namespace ShaderPrint
 {
 	// Allocate the debug print buffer associated with the view
-	void BeginView(FRDGBuilder& GraphBuilder, FViewInfo& View);
+	void BeginViews(FRDGBuilder& GraphBuilder, TArrayView<FViewInfo> Views);
 	// Draw info from the debug print buffer to the given output target
 	void DrawView(FRDGBuilder& GraphBuilder, const FViewInfo& View, const FScreenPassTexture& OutputTexture, const FScreenPassTexture& DepthTexture);
 	// Release the debug print buffer associated with the view
-	void EndView(FViewInfo& View);
+	void EndViews(TArrayView<FViewInfo> Views);
 }
diff --git a/Engine/Source/Runtime/Renderer/Public/ShaderPrintParameters.h b/Engine/Source/Runtime/Renderer/Public/ShaderPrintParameters.h
--- a/Engine/Source/Runtime/Renderer/Public/ShaderPrintParameters.h
+++ b/Engine/Source/Runtime/Renderer/Public/ShaderPrintParameters.h
@@ -128,7 +128,10 @@ namespace ShaderPrint
 	/** Import the shader print buffers into an RDG builder and recreate the FShaderPrintData object. */
 	RENDERER_API FShaderPrintData UnFreezeShaderPrintData(FRDGBuilder& GraphBuilder, FFrozenShaderPrintData& FrozenShaderPrintData);
 
-	/** Submit shader print data for display in the next rendered frame. */
+	/** Submit shader print data for display in the next rendered frame. The data is displayed in views from the scene, or all views if InScene==nullptr. */
+	RENDERER_API void SubmitShaderPrintData(FFrozenShaderPrintData& InData, FSceneInterface const* InScene);
+
+	UE_DEPRECATED(5.2, "Use the version of this function that takes a FSceneInterface")
 	RENDERER_API void SubmitShaderPrintData(FFrozenShaderPrintData& InData);
 
 	/** Fill the FShaderParameters with an explicit FShaderPrintData managed by the calling code. */
