FeedbackOcclusion: Adaptive latency instead of constant 3 frames

diff --git a/Engine/Source/Runtime/Renderer/Private/SceneOcclusion.cpp b/Engine/Source/Runtime/Renderer/Private/SceneOcclusion.cpp
--- a/Engine/Source/Runtime/Renderer/Private/SceneOcclusion.cpp
+++ b/Engine/Source/Runtime/Renderer/Private/SceneOcclusion.cpp
@@ -1607,6 +1607,7 @@ void FSceneRenderer::WaitOcclusionTests(FRHICommandListImmediate& RHICmdList)
 FOcclusionFeedback::FOcclusionFeedback()
 	: FRenderResource()
 	, CurrentBufferIndex(0)
+	, ResultsOcclusionFrameCounter(0)
 {
 	OcclusionVertexDeclarationRHI = nullptr;
 }
@@ -1670,19 +1671,56 @@ void FOcclusionFeedback::EndOcclusionScope(FRDGBuilder& GraphBuilder)
 	GPUFeedbackBuffer = nullptr;
 }
 
+void FOcclusionFeedback::AdvanceFrame(uint32 InOcclusionFrameCounter)
+{
+	constexpr uint32 NumBuffers = UE_ARRAY_COUNT(OcclusionBuffers);
+	CurrentBufferIndex = (CurrentBufferIndex + 1u) % NumBuffers;
+	
+	FOcclusionBuffer& OcclusionBuffer = OcclusionBuffers[CurrentBufferIndex];
+	OcclusionBuffer.BatchedPrimitives.Empty();
+	OcclusionBuffer.OcclusionFrameCounter = InOcclusionFrameCounter;
+}
+
 void FOcclusionFeedback::ReadbackResults(FRHICommandList& RHICmdList)
 {
 	constexpr uint32 NumBuffers = UE_ARRAY_COUNT(OcclusionBuffers);
+	TArray<uint32, TInlineAllocator<NumBuffers>> ReadyBuffers;
 
-	CurrentBufferIndex = (CurrentBufferIndex + 1u) % NumBuffers;
+	// Find the most recent Ready buffer
+	for (uint32 Index = 1u; Index <= NumBuffers; ++Index)
+	{
+		uint32 ReadbackBufferIndex = (CurrentBufferIndex + Index) % NumBuffers;
+		FOcclusionBuffer& OcclusionBuffer = OcclusionBuffers[ReadbackBufferIndex];
 
-	FOcclusionBuffer& OcclusionBuffer = OcclusionBuffers[CurrentBufferIndex];
-	uint32 NumPrimitives = OcclusionBuffer.BatchedPrimitives.Num();
+		if (OcclusionBuffer.BatchedPrimitives.Num() == 0)
+		{
+			// No data to readback
+			break;
+		}
+		
+		if (OcclusionBuffer.OcclusionFrameCounter <= ResultsOcclusionFrameCounter)
+		{
+			// Already read
+			continue;
+		}
+
+		if (!OcclusionBuffer.ReadbackBuffer->IsReady())
+		{
+			break;
+		}
+
+		ReadyBuffers.Add(ReadbackBufferIndex);
+	}
 
-	bool bReady = (NumPrimitives > 0 && OcclusionBuffer.ReadbackBuffer->IsReady());
+	bool bReady = (ReadyBuffers.Num() != 0);
 	if (!bReady)
 	{
-		if (OcclusionBuffer.BatchedPrimitives.Num() > 0)
+		uint32 TailBufferIndex = (CurrentBufferIndex + 1u) % NumBuffers;
+		FOcclusionBuffer& TailOcclusionBuffer = OcclusionBuffers[TailBufferIndex];
+		
+		// Check if we have to wait for result
+		if (TailOcclusionBuffer.BatchedPrimitives.Num() > 0 
+			&& ResultsOcclusionFrameCounter < TailOcclusionBuffer.OcclusionFrameCounter)
 		{
 			// We'll do manual wait
 			double StartTime = FPlatformTime::Seconds();
@@ -1705,15 +1743,25 @@ void FOcclusionFeedback::ReadbackResults(FRHICommandList& RHICmdList)
 					break;
 				}
 
-				bReady = OcclusionBuffer.ReadbackBuffer->IsReady();
+				bReady = TailOcclusionBuffer.ReadbackBuffer->IsReady();
+			}
+
+			if (bReady)
+			{
+				ReadyBuffers.Add(TailBufferIndex);
 			}
 		}
 	}
 
 	if (bReady)
 	{
+		uint32 ReadbackBufferIndex = ReadyBuffers.Last();
+		FOcclusionBuffer& OcclusionBuffer = OcclusionBuffers[ReadbackBufferIndex];
+
+		ResultsOcclusionFrameCounter = OcclusionBuffer.OcclusionFrameCounter;
 		LatestOcclusionResults.Reset();
 
+		uint32 NumPrimitives = OcclusionBuffer.BatchedPrimitives.Num();
 		uint32 NumBytes = NumPrimitives * sizeof(uint32);
 		uint32* Results = (uint32*)OcclusionBuffer.ReadbackBuffer->Lock(NumBytes);
 
@@ -1727,8 +1775,6 @@ void FOcclusionFeedback::ReadbackResults(FRHICommandList& RHICmdList)
 
 		OcclusionBuffer.ReadbackBuffer->Unlock();
 	}
-
-	OcclusionBuffer.BatchedPrimitives.Empty();
 }
 
 void FOcclusionFeedback::AddPrimitive(FPrimitiveComponentId PrimitiveId, const FVector& BoundsOrigin, const FVector& BoundsBoxExtent, FGlobalDynamicVertexBuffer& DynamicVertexBuffer)
@@ -1802,7 +1848,7 @@ IMPLEMENT_GLOBAL_SHADER(FOcclusionFeedbackPS,"/Engine/Private/OcclusionFeedbackS
 
 int32 GOcclusionFeeback_Blending = 1;
 static FAutoConsoleVariableRef CVarOcclusionFeeback_Blending(
-	TEXT("r.OcclusionFeeback.Blending"),
+	TEXT("r.OcclusionFeedback.Blending"),
 	GOcclusionFeeback_Blending,
 	TEXT("0: Opaque \n")
 	TEXT("1: Enable blending  (default) \n")
diff --git a/Engine/Source/Runtime/Renderer/Private/SceneRendering.h b/Engine/Source/Runtime/Renderer/Private/SceneRendering.h
--- a/Engine/Source/Runtime/Renderer/Private/SceneRendering.h
+++ b/Engine/Source/Runtime/Renderer/Private/SceneRendering.h
@@ -515,7 +515,9 @@ public:
 
 	/** Renders the current batch and resets the batch state. */
 	void SubmitOcclusionDraws(FRHICommandList& RHICmdList, FViewInfo& View);
+	
 	void ReadbackResults(FRHICommandList& RHICmdList);
+	void AdvanceFrame(uint32 OcclusionFrameCounter);
 
 	inline FRDGBuffer* GetGPUFeedbackBuffer() const
 	{
@@ -543,11 +545,14 @@ private:
 	{
 		TArray<FPrimitiveComponentId> BatchedPrimitives;
 		FRHIGPUBufferReadback* ReadbackBuffer = nullptr;
+		uint32 OcclusionFrameCounter = 0u;
 	};
 
 	FOcclusionBuffer OcclusionBuffers[3];
 	uint32 CurrentBufferIndex;
+	
 	TSet<FPrimitiveComponentId> LatestOcclusionResults;
+	uint32 ResultsOcclusionFrameCounter;
 
 	//
 	FVertexDeclarationRHIRef OcclusionVertexDeclarationRHI;
diff --git a/Engine/Source/Runtime/Renderer/Private/SceneVisibility.cpp b/Engine/Source/Runtime/Renderer/Private/SceneVisibility.cpp
--- a/Engine/Source/Runtime/Renderer/Private/SceneVisibility.cpp
+++ b/Engine/Source/Runtime/Renderer/Private/SceneVisibility.cpp
@@ -2050,6 +2050,7 @@ static int32 OcclusionCull(FRHICommandListImmediate& RHICmdList, const FScene* S
 		{
 			QUICK_SCOPE_CYCLE_COUNTER(STAT_OcclusionFeedback_ReadbackResults);
 			ViewState->OcclusionFeedback.ReadbackResults(RHICmdList);
+			ViewState->OcclusionFeedback.AdvanceFrame(ViewState->OcclusionFrameCounter);
 		}
 		else if( bHZBOcclusion )
 		{
