PR #9734: Fx up PassName of "SecondaryUpscale" (Contributed by [REDACTED_PERSON])

diff --git a/Engine/Source/Runtime/Renderer/Private/PostProcess/PostProcessing.cpp b/Engine/Source/Runtime/Renderer/Private/PostProcess/PostProcessing.cpp
--- a/Engine/Source/Runtime/Renderer/Private/PostProcess/PostProcessing.cpp
+++ b/Engine/Source/Runtime/Renderer/Private/PostProcess/PostProcessing.cpp
@@ -1836,7 +1836,7 @@ void AddMobilePostProcessingPasses(FRDGBuilder& GraphBuilder, FScene* Scene, con
 		TEXT("SelectionOutline"),
 		TEXT("EditorPrimitive"),
 		TEXT("PrimaryUpscale"),
-		TEXT("Sharpen"),
+		TEXT("SecondaryUpscale"),
 		TEXT("Visualize"),
 		TEXT("HMDDistortion")
 	};
