Multi-view:  Fix for uniform buffer crash, by running one scene renderer at a time in the outer loop, matching 5.0 behavior, in terms of the code executing on the render thread.  Root cause remains unknown, but is assumed to be some sort of state that needs to be reset between scene renderers.

#rnx

diff --git a/Engine/Source/Runtime/Renderer/Private/SceneRendering.cpp b/Engine/Source/Runtime/Renderer/Private/SceneRendering.cpp
--- a/Engine/Source/Runtime/Renderer/Private/SceneRendering.cpp
+++ b/Engine/Source/Runtime/Renderer/Private/SceneRendering.cpp
@@ -4701,8 +4701,20 @@ void FRendererModule::BeginRenderingViewFamilies(FCanvas* Canvas, TArrayView<FSc
 				uint64 SceneRenderStart = FPlatformTime::Cycles64();
 				const float StartDelayMillisec = FPlatformTime::ToMilliseconds64(SceneRenderStart - DrawSceneEnqueue);
 				CSV_CUSTOM_STAT_GLOBAL(DrawSceneCommand_StartDelay, StartDelayMillisec, ECsvCustomStatOp::Set);
-				RenderViewFamilies_RenderThread(RHICmdList, LocalSceneRenderers);
-				FlushPendingDeleteRHIResources_RenderThread();
+
+				// TODO:  There were some random uniform buffer crashes that resulted when rendering multiple view families
+				// in the function below.  Pass one at a time for now as a workaround (matches original Release 5.0 behavior),
+				// while we investigate.  The theory is that something in FSceneRenderer::RenderThreadBegin/End needs to run
+				// between scene renders to reset some state, but we don't know what.
+				TArray<FSceneRenderer*> SingleSceneRenderer;
+				SingleSceneRenderer.AddZeroed(1);
+
+				for (FSceneRenderer* SceneRenderer : LocalSceneRenderers)
+				{
+					SingleSceneRenderer[0] = SceneRenderer;
+					RenderViewFamilies_RenderThread(RHICmdList, SingleSceneRenderer);
+					FlushPendingDeleteRHIResources_RenderThread();
+				}
 			});
 
 		// Force kick the RT if we've got RT polling on.
