Use shader hashes for computing PrecachePSOHash.

- Allows construction of StatePrecachePSOHash in FGraphicsPipelineStateInitializers that are not associated with MDCs.
- Also better reflects the actual PSO compatibility.

diff --git a/Engine/Source/Runtime/Apple/MetalRHI/Private/Shaders/MetalShaderLibrary.h b/Engine/Source/Runtime/Apple/MetalRHI/Private/Shaders/MetalShaderLibrary.h
--- a/Engine/Source/Runtime/Apple/MetalRHI/Private/Shaders/MetalShaderLibrary.h
+++ b/Engine/Source/Runtime/Apple/MetalRHI/Private/Shaders/MetalShaderLibrary.h
@@ -31,6 +31,10 @@ public:
 
 	virtual int32 FindShaderMapIndex(const FSHAHash& Hash) override;
 	virtual int32 FindShaderIndex(const FSHAHash& Hash) override;
+	virtual FSHAHash GetShaderHash(int32 ShaderMapIndex, int32 ShaderIndex) override
+	{ 
+		return SerializedShaders.ShaderHashes[GetShaderIndex(ShaderMapIndex, ShaderIndex)];
+	};
 
 	virtual bool PreloadShader(int32 ShaderIndex, FGraphEventArray& OutCompletionEvents) override { return false; }
 	virtual bool PreloadShaderMap(int32 ShaderMapIndex, FGraphEventArray& OutCompletionEvents) override { return false; }
diff --git a/Engine/Source/Runtime/RHI/Private/PipelineStateCache.cpp b/Engine/Source/Runtime/RHI/Private/PipelineStateCache.cpp
--- a/Engine/Source/Runtime/RHI/Private/PipelineStateCache.cpp
+++ b/Engine/Source/Runtime/RHI/Private/PipelineStateCache.cpp
@@ -979,9 +979,6 @@ public:
 				return nullptr;
 			}
 
-			// create new graphics state
-			NewState = new FGraphicsPipelineState();
-
 			// Add to cache before starting async operation, because the async op will remove the task from active ops when done			
 			PrecachedPSOInitializers.Add(Initializer);
 		}
@@ -992,14 +989,27 @@ public:
 			// Only need the hash of the PSO here
 			uint64 PrecachePSOHash = RHIComputePrecachePSOHash(Initializer);
 
+			FRWScopeLock WriteLock(ActivePrecachingTasksRWLock, SLT_Write);
+
+			if (ActivePrecachingTasks.Contains(PrecachePSOHash))
+			{
+				return nullptr;
+			}
+
+			// create new graphics state
+			NewState = new FGraphicsPipelineState();
+
 			FPrecacheTask PrecacheTask;
 			PrecacheTask.Initializer = Initializer;
 			PrecacheTask.PSO = NewState;
 
-			FRWScopeLock WriteLock(ActivePrecachingTasksRWLock, SLT_Write);
-			check(!ActivePrecachingTasks.Contains(PrecachePSOHash));
 			ActivePrecachingTasks.Add(PrecachePSOHash, PrecacheTask);
 		}
+		else
+		{
+			// create new graphics state
+			NewState = new FGraphicsPipelineState();
+		}
 
 		return NewState;
 	}
diff --git a/Engine/Source/Runtime/RHI/Public/RHIResources.h b/Engine/Source/Runtime/RHI/Public/RHIResources.h
--- a/Engine/Source/Runtime/RHI/Public/RHIResources.h
+++ b/Engine/Source/Runtime/RHI/Public/RHIResources.h
@@ -3651,6 +3651,7 @@ public:
 	virtual int32 GetNumShaders() const = 0;
 	virtual int32 GetNumShadersForShaderMap(int32 ShaderMapIndex) const = 0;
 	virtual int32 GetShaderIndex(int32 ShaderMapIndex, int32 i) const = 0;
+	virtual FSHAHash GetShaderHash(int32 ShaderMapIndex, int32 ShaderIndex) = 0;
 	virtual int32 FindShaderMapIndex(const FSHAHash& Hash) = 0;
 	virtual int32 FindShaderIndex(const FSHAHash& Hash) = 0;
 	virtual bool PreloadShader(int32 ShaderIndex, FGraphEventArray& OutCompletionEvents) { return false; }
diff --git a/Engine/Source/Runtime/RenderCore/Private/ShaderCodeLibrary.cpp b/Engine/Source/Runtime/RenderCore/Private/ShaderCodeLibrary.cpp
--- a/Engine/Source/Runtime/RenderCore/Private/ShaderCodeLibrary.cpp
+++ b/Engine/Source/Runtime/RenderCore/Private/ShaderCodeLibrary.cpp
@@ -240,6 +240,7 @@ public:
 	virtual void ReleaseRHI() override;
 
 	// FShaderMapResource interface
+	virtual FSHAHash GetShaderHash(int32 ShaderIndex) override;
 	virtual FRHIShader* CreateRHIShaderOrCrash(int32 ShaderIndex) override;
 	virtual void ReleasePreloadedShaderCode(int32 ShaderIndex) override;
 	virtual bool TryRelease() override;
@@ -1158,6 +1159,11 @@ FShaderMapResource_SharedCode::~FShaderMapResource_SharedCode()
 	
 }
 
+FSHAHash FShaderMapResource_SharedCode::GetShaderHash(int32 ShaderIndex)
+{
+	return LibraryInstance->Library->GetShaderHash(ShaderMapIndex, ShaderIndex);
+}
+
 FRHIShader* FShaderMapResource_SharedCode::CreateRHIShaderOrCrash(int32 ShaderIndex)
 {
 	SCOPED_LOADTIMER(FShaderMapResource_SharedCode_InitRHI);
diff --git a/Engine/Source/Runtime/RenderCore/Private/ShaderResource.cpp b/Engine/Source/Runtime/RenderCore/Private/ShaderResource.cpp
--- a/Engine/Source/Runtime/RenderCore/Private/ShaderResource.cpp
+++ b/Engine/Source/Runtime/RenderCore/Private/ShaderResource.cpp
@@ -650,6 +650,11 @@ FRHIShader* FShaderMapResource::CreateShaderOrCrash(int32 ShaderIndex)
 	return Shader;
 }
 
+FSHAHash FShaderMapResource_InlineCode::GetShaderHash(int32 ShaderIndex)
+{
+	return Code->ShaderHashes[ShaderIndex];
+}
+
 FRHIShader* FShaderMapResource_InlineCode::CreateRHIShaderOrCrash(int32 ShaderIndex)
 {
 	TRACE_CPUPROFILER_EVENT_SCOPE(FShaderMapResource_InlineCode::CreateRHIShaderOrCrash);
diff --git a/Engine/Source/Runtime/RenderCore/Public/Shader.h b/Engine/Source/Runtime/RenderCore/Public/Shader.h
--- a/Engine/Source/Runtime/RenderCore/Public/Shader.h
+++ b/Engine/Source/Runtime/RenderCore/Public/Shader.h
@@ -349,6 +349,9 @@ public:
 		return Shader;
 	}
 
+	/** Return shader hash for a particular shader without creating it. */
+	virtual FSHAHash GetShaderHash(int32 ShaderIndex) = 0;
+
 	void BeginCreateAllShaders();
 
 #if RHI_RAYTRACING
@@ -524,6 +527,7 @@ public:
 	{}
 
 	// FShaderMapResource interface
+	virtual FSHAHash GetShaderHash(int32 ShaderIndex) override;
 	virtual FRHIShader* CreateRHIShaderOrCrash(int32 ShaderIndex) override;
 	virtual uint32 GetSizeBytes() const override { return sizeof(*this) + GetAllocatedSize(); }
 
diff --git a/Engine/Source/Runtime/RenderCore/Public/ShaderCodeArchive.h b/Engine/Source/Runtime/RenderCore/Public/ShaderCodeArchive.h
--- a/Engine/Source/Runtime/RenderCore/Public/ShaderCodeArchive.h
+++ b/Engine/Source/Runtime/RenderCore/Public/ShaderCodeArchive.h
@@ -319,6 +319,11 @@ public:
 		return SerializedShaders.FindShader(Hash);
 	}
 
+	virtual FSHAHash GetShaderHash(int32 ShaderMapIndex, int32 ShaderIndex) override
+	{
+		return SerializedShaders.ShaderHashes[GetShaderIndex(ShaderMapIndex, ShaderIndex)];
+	};
+
 	virtual bool PreloadShader(int32 ShaderIndex, FGraphEventArray& OutCompletionEvents) override;
 
 	virtual bool PreloadShaderMap(int32 ShaderMapIndex, FGraphEventArray& OutCompletionEvents) override;
@@ -568,6 +573,11 @@ public:
 
 	virtual int32 FindShaderMapIndex(const FSHAHash& Hash) override;
 	virtual int32 FindShaderIndex(const FSHAHash& Hash) override;
+	virtual FSHAHash GetShaderHash(int32 ShaderMapIndex, int32 ShaderIndex) override
+	{
+		return Header.ShaderHashes[GetShaderIndex(ShaderMapIndex, ShaderIndex)];
+	}
+
 	virtual bool PreloadShader(int32 ShaderIndex, FGraphEventArray& OutCompletionEvents) override;
 	virtual bool PreloadShaderMap(int32 ShaderMapIndex, FGraphEventArray& OutCompletionEvents) override;
 	virtual bool PreloadShaderMap(int32 ShaderMapIndex, FCoreDelegates::FAttachShaderReadRequestFunc AttachShaderReadRequestFunc) override;
diff --git a/Engine/Source/Runtime/Renderer/Private/MeshPassProcessor.cpp b/Engine/Source/Runtime/Renderer/Private/MeshPassProcessor.cpp
--- a/Engine/Source/Runtime/Renderer/Private/MeshPassProcessor.cpp
+++ b/Engine/Source/Runtime/Renderer/Private/MeshPassProcessor.cpp
@@ -907,13 +907,13 @@ void FGraphicsMinimalPipelineStateInitializer::ComputePrecachePSOHash()
 	FMemory::Memzero(&HashKey, sizeof(FHashKey));
 
 	HashKey.VertexDeclaration = BoundShaderState.VertexDeclarationRHI ? BoundShaderState.VertexDeclarationRHI->GetPrecachePSOHash() : 0;
-	HashKey.VertexShader = HashCombine(GetTypeHash(BoundShaderState.VertexShaderIndex), GetTypeHash(BoundShaderState.VertexShaderResource));
-	HashKey.PixelShader = HashCombine(GetTypeHash(BoundShaderState.PixelShaderIndex), GetTypeHash(BoundShaderState.PixelShaderResource));
+	HashKey.VertexShader = BoundShaderState.VertexShaderResource.IsValid() ? GetTypeHash(BoundShaderState.VertexShaderResource->GetShaderHash(BoundShaderState.VertexShaderIndex)) : 0;
+	HashKey.PixelShader = BoundShaderState.PixelShaderResource.IsValid() ? GetTypeHash(BoundShaderState.PixelShaderResource->GetShaderHash(BoundShaderState.PixelShaderIndex)) : 0;
 #if PLATFORM_SUPPORTS_GEOMETRY_SHADERS
-	HashKey.GeometryShader = HashCombine(GetTypeHash(BoundShaderState.GeometryShaderIndex), GetTypeHash(BoundShaderState.GeometryShaderResource));
+	HashKey.GeometryShader = BoundShaderState.GeometryShaderResource.IsValid() ? GetTypeHash(BoundShaderState.GeometryShaderResource->GetShaderHash(BoundShaderState.GeometryShaderIndex)) : 0;
 #endif
 #if PLATFORM_SUPPORTS_MESH_SHADERS
-	HashKey.MeshShader = HashCombine(GetTypeHash(BoundShaderState.MeshShaderIndex), GetTypeHash(BoundShaderState.MeshShaderResource));
+	HashKey.MeshShader = BoundShaderState.MeshShaderResource.IsValid() ? GetTypeHash(BoundShaderState.MeshShaderResource->GetShaderHash(BoundShaderState.MeshShaderIndex)) : 0;
 #endif
 
 	FBlendStateInitializerRHI BlendStateInitializerRHI;
