Fixed several RHI validation transition issues.
 - Fixed GPU scene resources not being available for the base pass. The external access queue is flushed only on the graphics pipe for the base pass and then later for both graphics | async compute. This is necessary to avoid
   having to immediately transition back from Graphics | AsyncCompute when updating shadow GPU scene buffers, which introduces an immediate stall on any in-flight async compute work on certain platforms.
 - Fixed issue where RHI transition validation log was being hashed including the stack breadcrumb, which was causing log spam.
 - Merged reflection capture / atmosphere external access queues together with the main scene rendering one to reduce passes.
 - Fixed incorrect  RHI validation assert in hardware ray tracing passes.
 - Fixed IES texture atlas assert due to not using external access mode on the texture.

diff --git a/Engine/Source/Runtime/Engine/Public/SceneInterface.h b/Engine/Source/Runtime/Engine/Public/SceneInterface.h
--- a/Engine/Source/Runtime/Engine/Public/SceneInterface.h
+++ b/Engine/Source/Runtime/Engine/Public/SceneInterface.h
@@ -37,6 +37,7 @@ class FViewInfo;
 class FSceneRenderer;
 class FInstanceCullingManager;
 class FSceneViewStateInterface;
+class FRDGExternalAccessQueue;
 struct FHairStrandsInstance;
 struct FLightRenderParameters;
 class FSparseVolumeTextureViewerSceneProxy;
@@ -197,7 +198,10 @@ public:
 	 */
 	virtual void UpdateSkyCaptureContents(const USkyLightComponent* CaptureComponent, bool bCaptureEmissiveOnly, UTextureCube* SourceCubemap, FTexture* OutProcessedTexture, float& OutAverageBrightness, FSHVectorRGB3& OutIrradianceEnvironmentMap, TArray<FFloat16Color>* OutRadianceMap, FLinearColor* SpecifiedCubemapColorScale) {}
 
+	UE_DEPRECATED(5.2, "AllocateAndCaptureFrameSkyEnvMap now takes an FRDGExternalAccessQueue")
 	virtual void AllocateAndCaptureFrameSkyEnvMap(FRDGBuilder& GraphBuilder, FSceneRenderer& SceneRenderer, FViewInfo& MainView, bool bShouldRenderSkyAtmosphere, bool bShouldRenderVolumetricCloud, FInstanceCullingManager& InstanceCullingManager) {}
+
+	virtual void AllocateAndCaptureFrameSkyEnvMap(FRDGBuilder& GraphBuilder, FSceneRenderer& SceneRenderer, FViewInfo& MainView, bool bShouldRenderSkyAtmosphere, bool bShouldRenderVolumetricCloud, FInstanceCullingManager& InstanceCullingManager, FRDGExternalAccessQueue& ExternalAccessQueue) {}
 	virtual void ValidateSkyLightRealTimeCapture(FRDGBuilder& GraphBuilder, const FViewInfo& View, FRDGTextureRef SceneColorTexture) {}
 
 	virtual void AddPlanarReflection(class UPlanarReflectionComponent* Component) {}
diff --git a/Engine/Source/Runtime/RHI/Private/RHIValidation.cpp b/Engine/Source/Runtime/RHI/Private/RHIValidation.cpp
--- a/Engine/Source/Runtime/RHI/Private/RHIValidation.cpp
+++ b/Engine/Source/Runtime/RHI/Private/RHIValidation.cpp
@@ -534,6 +534,22 @@ thread_local TConstArrayView<const TCHAR*> FRHIValidationBreadcrumbScope::Breadc
 
 void FValidationRHI::ReportValidationFailure(const TCHAR* InMessage)
 {
+	// Report failures only once per session, since many of them will happen repeatedly. This is similar to what ensure() does, but
+	// ensure() looks at the source location to determine if it's seen the error before. We want to look at the actual message, since
+	// all failures of a given kind will come from the same place, but (hopefully) the error message contains the name of the resource
+	// and a description of the state, so it should be unique for each failure.
+	uint32 Hash = FCrc::StrCrc32<TCHAR>(InMessage);
+	
+	SeenFailureHashesMutex.Lock();
+	bool bIsAlreadyInSet;
+	SeenFailureHashes.Add(Hash, &bIsAlreadyInSet);
+	SeenFailureHashesMutex.Unlock();
+
+	if (bIsAlreadyInSet)
+	{
+		return;
+	}
+
 	FString Message;
 
 	if (!FRHIValidationBreadcrumbScope::Breadcrumbs.IsEmpty())
@@ -557,22 +573,6 @@ void FValidationRHI::ReportValidationFailure(const TCHAR* InMessage)
 		Message = InMessage;
 	}
 
-	// Report failures only once per session, since many of them will happen repeatedly. This is similar to what ensure() does, but
-	// ensure() looks at the source location to determine if it's seen the error before. We want to look at the actual message, since
-	// all failures of a given kind will come from the same place, but (hopefully) the error message contains the name of the resource
-	// and a description of the state, so it should be unique for each failure.
-	uint32 Hash = FCrc::StrCrc32<TCHAR>(*Message);
-	
-	SeenFailureHashesMutex.Lock();
-	bool bIsAlreadyInSet;
-	SeenFailureHashes.Add(Hash, &bIsAlreadyInSet);
-	SeenFailureHashesMutex.Unlock();
-
-	if (bIsAlreadyInSet)
-	{
-		return;
-	}
-
 	UE_LOG(LogRHI, Error, TEXT("%s"), *Message);
 
 	if (FPlatformMisc::IsDebuggerPresent() && RHIValidation::GBreakOnTransitionError)
diff --git a/Engine/Source/Runtime/RHI/Public/RHIDefinitions.h b/Engine/Source/Runtime/RHI/Public/RHIDefinitions.h
--- a/Engine/Source/Runtime/RHI/Public/RHIDefinitions.h
+++ b/Engine/Source/Runtime/RHI/Public/RHIDefinitions.h
@@ -106,6 +106,20 @@ enum EShaderFrequency : uint8
 };
 static_assert(SF_NumFrequencies <= (1 << SF_NumBits), "SF_NumFrequencies will not fit on SF_NumBits");
 
+inline bool IsComputeShaderFrequency(EShaderFrequency ShaderFrequency)
+{
+	switch (ShaderFrequency)
+	{
+	case SF_Compute:
+	case SF_RayGen:
+	case SF_RayMiss:
+	case SF_RayHitGroup:
+	case SF_RayCallable:
+		return true;
+	}
+	return false;
+}
+
 enum ERenderQueryType
 {
 	// e.g. WaitForFrameEventCompletion()
diff --git a/Engine/Source/Runtime/RHICore/Public/RHICoreShader.h b/Engine/Source/Runtime/RHICore/Public/RHICoreShader.h
--- a/Engine/Source/Runtime/RHICore/Public/RHICoreShader.h
+++ b/Engine/Source/Runtime/RHICore/Public/RHICoreShader.h
@@ -177,7 +177,7 @@ void SetResourcesFromTables(TBinder&& Binder, FRHIShader const& Shader, FShaderR
 #if ENABLE_RHI_VALIDATION
 				if (Tracker)
 				{
-					ERHIAccess Access = Shader.GetFrequency() == SF_Compute
+					ERHIAccess Access = IsComputeShaderFrequency(Shader.GetFrequency())
 						? ERHIAccess::SRVCompute
 						: ERHIAccess::SRVGraphics;
 
@@ -197,7 +197,7 @@ void SetResourcesFromTables(TBinder&& Binder, FRHIShader const& Shader, FShaderR
 #if ENABLE_RHI_VALIDATION
 				if (Tracker)
 				{
-					ERHIAccess Access = Shader.GetFrequency() == SF_Compute
+					ERHIAccess Access = IsComputeShaderFrequency(Shader.GetFrequency())
 						? ERHIAccess::SRVCompute
 						: ERHIAccess::SRVGraphics;
 
@@ -214,7 +214,7 @@ void SetResourcesFromTables(TBinder&& Binder, FRHIShader const& Shader, FShaderR
 #if ENABLE_RHI_VALIDATION
 				if (Tracker)
 				{
-					ERHIAccess Access = Shader.GetFrequency() == SF_Compute
+					ERHIAccess Access = IsComputeShaderFrequency(Shader.GetFrequency())
 						? ERHIAccess::UAVCompute
 						: ERHIAccess::UAVGraphics;
 
diff --git a/Engine/Source/Runtime/RenderCore/Private/RenderGraphUtils.cpp b/Engine/Source/Runtime/RenderCore/Private/RenderGraphUtils.cpp
--- a/Engine/Source/Runtime/RenderCore/Private/RenderGraphUtils.cpp
+++ b/Engine/Source/Runtime/RenderCore/Private/RenderGraphUtils.cpp
@@ -1078,6 +1078,41 @@ FRDGWaitForTasksScope::~FRDGWaitForTasksScope()
 	}
 }
 
+void FRDGExternalAccessQueue::Submit(FRDGBuilder& GraphBuilder, ERHIPipeline AllowedPipelines)
+{
+	// Force any resources that were previously submitted with different pipelines back to internal access mode first.
+	for (FResource Resource : Resources)
+	{
+		if (Resource.SubmittedPipelines != ERHIPipeline::None)
+		{
+			GraphBuilder.UseInternalAccessMode(Resource.Resource);
+		}
+	}
+
+	TArray<FResource, FRDGArrayAllocator> ResourcesToKeep;
+
+	if (AllowedPipelines != ERHIPipeline::All)
+	{
+		ResourcesToKeep.Reserve(Resources.Num());
+	}
+
+	for (FResource Resource : Resources)
+	{
+		const ERHIPipeline ResourcePipelines = AllowedPipelines & Resource.RequestedPipelines;
+
+		// If the resource is requesting a pipe that we aren't allowing, then continue to queue the resource for a later submit.
+		if (ResourcePipelines != Resource.RequestedPipelines)
+		{
+			Resource.SubmittedPipelines |= AllowedPipelines;
+			ResourcesToKeep.Emplace(Resource);
+		}
+
+		GraphBuilder.UseExternalAccessMode(Resource.Resource, Resource.Access, ResourcePipelines);
+	}
+
+	Swap(ResourcesToKeep, Resources);
+}
+
 bool AllocatePooledBuffer(
 	const FRDGBufferDesc& Desc,
 	TRefCountPtr<FRDGPooledBuffer>& Out,
diff --git a/Engine/Source/Runtime/RenderCore/Public/RenderGraphUtils.h b/Engine/Source/Runtime/RenderCore/Public/RenderGraphUtils.h
--- a/Engine/Source/Runtime/RenderCore/Public/RenderGraphUtils.h
+++ b/Engine/Source/Runtime/RenderCore/Public/RenderGraphUtils.h
@@ -1048,7 +1048,7 @@ void AddReadbackBufferPass(FRDGBuilder& GraphBuilder, FRDGEventName&& Name, FRDG
 }
 
 /** Batches up RDG external resource access mode requests and submits them all at once to RDG. */
-class FRDGExternalAccessQueue
+class RENDERCORE_API FRDGExternalAccessQueue
 {
 public:
 	FRDGExternalAccessQueue() = default;
@@ -1091,14 +1091,7 @@ public:
 		Resources.Emplace(Resource, Access, Pipelines);
 	}
 
-	void Submit(FRDGBuilder& GraphBuilder)
-	{
-		for (FResource Resource : Resources)
-		{
-			GraphBuilder.UseExternalAccessMode(Resource.Resource, Resource.Access, Resource.Pipelines);
-		}
-		Resources.Empty();
-	}
+	void Submit(FRDGBuilder& GraphBuilder, ERHIPipeline AllowedPipes = ERHIPipeline::All);
 
 	bool Contains(FRDGViewableResource* Resource)
 	{
@@ -1148,12 +1141,13 @@ private:
 		FResource(FRDGViewableResource* InResource, ERHIAccess InAccess, ERHIPipeline InPipelines)
 			: Resource(InResource)
 			, Access(InAccess)
-			, Pipelines(InPipelines)
+			, RequestedPipelines(InPipelines)
 		{}
 
 		FRDGViewableResource* Resource;
 		ERHIAccess Access;
-		ERHIPipeline Pipelines;
+		ERHIPipeline RequestedPipelines;
+		ERHIPipeline SubmittedPipelines = ERHIPipeline::None;
 	};
 
 	TArray<FResource, FRDGArrayAllocator> Resources;
diff --git a/Engine/Source/Runtime/Renderer/Private/DeferredShadingRenderer.cpp b/Engine/Source/Runtime/Renderer/Private/DeferredShadingRenderer.cpp
--- a/Engine/Source/Runtime/Renderer/Private/DeferredShadingRenderer.cpp
+++ b/Engine/Source/Runtime/Renderer/Private/DeferredShadingRenderer.cpp
@@ -2739,10 +2739,10 @@ void FDeferredShadingSceneRenderer::Render(FRDGBuilder& GraphBuilder)
 	}
 
 	PrepareDistanceFieldScene(GraphBuilder, ExternalAccessQueue, false);
-	if (!CVarShadowsUseSharedExternalAccessQueue.GetValueOnRenderThread())
-	{
-		ExternalAccessQueue.Submit(GraphBuilder);
-	}
+
+	// Changes are flushed to the graphics pipe only, which avoids having to transition from async compute
+	// back to graphics later on during shadow GPU scene updates, which is serializing async Lumen compute work.
+	ExternalAccessQueue.Submit(GraphBuilder, ERHIPipeline::Graphics);
 
 	const bool bShouldRenderVelocities = ShouldRenderVelocities();
 	const EShaderPlatform Platform = GetViewFamilyInfo(Views).GetShaderPlatform();
@@ -3039,7 +3039,7 @@ void FDeferredShadingSceneRenderer::Render(FRDGBuilder& GraphBuilder)
 	if (SkyAtmospherePassLocation == ESkyAtmospherePassLocation::BeforeOcclusion && bShouldRenderSkyAtmosphere)
 	{
 		// Generate the Sky/Atmosphere look up tables
-		RenderSkyAtmosphereLookUpTables(GraphBuilder);
+		RenderSkyAtmosphereLookUpTables(GraphBuilder, ExternalAccessQueue);
 	}
 
 	const auto RenderOcclusionLambda = [&]()
@@ -3084,15 +3084,19 @@ void FDeferredShadingSceneRenderer::Render(FRDGBuilder& GraphBuilder)
 	if (SkyAtmospherePassLocation == ESkyAtmospherePassLocation::BeforeBasePass && bShouldRenderSkyAtmosphere)
 	{
 		// Generate the Sky/Atmosphere look up tables
-		RenderSkyAtmosphereLookUpTables(GraphBuilder);
+		RenderSkyAtmosphereLookUpTables(GraphBuilder, ExternalAccessQueue);
 	}
 
+	// Changes are flushed to the graphics pipe only, which avoids having to transition from async compute
+	// back to graphics later on during shadow GPU scene updates, which is serializing async Lumen compute work.
+	ExternalAccessQueue.Submit(GraphBuilder, ERHIPipeline::Graphics);
+
 	// Capture the SkyLight using the SkyAtmosphere and VolumetricCloud component if available.
 	const bool bRealTimeSkyCaptureEnabled = Scene->SkyLight && Scene->SkyLight->bRealTimeCaptureEnabled && Views.Num() > 0 && ViewFamily.EngineShowFlags.SkyLighting;
 	if (bRealTimeSkyCaptureEnabled)
 	{
 		FViewInfo& MainView = Views[0];
-		Scene->AllocateAndCaptureFrameSkyEnvMap(GraphBuilder, *this, MainView, bShouldRenderSkyAtmosphere, bShouldRenderVolumetricCloud, InstanceCullingManager);
+		Scene->AllocateAndCaptureFrameSkyEnvMap(GraphBuilder, *this, MainView, bShouldRenderSkyAtmosphere, bShouldRenderVolumetricCloud, InstanceCullingManager, ExternalAccessQueue);
 	}
 
 	const ECustomDepthPassLocation CustomDepthPassLocation = GetCustomDepthPassLocation(ShaderPlatform);
@@ -3402,10 +3406,9 @@ void FDeferredShadingSceneRenderer::Render(FRDGBuilder& GraphBuilder)
 			RenderLumenSceneLighting(GraphBuilder, LumenFrameTemporaries);
 		}
 	}
-	if (CVarShadowsUseSharedExternalAccessQueue.GetValueOnRenderThread())
-	{
-		ExternalAccessQueue.Submit(GraphBuilder);
-	}
+
+	ExternalAccessQueue.Submit(GraphBuilder);
+
 	// End shadow and fog after base pass
 
 	// Trigger a command submit here, to avoid GPU bubbles
diff --git a/Engine/Source/Runtime/Renderer/Private/IESTextureManager.cpp b/Engine/Source/Runtime/Renderer/Private/IESTextureManager.cpp
--- a/Engine/Source/Runtime/Renderer/Private/IESTextureManager.cpp
+++ b/Engine/Source/Runtime/Renderer/Private/IESTextureManager.cpp
@@ -248,6 +248,8 @@ static void AddSlotsPass(
 		const FIntVector DispatchCount = FComputeShaderUtils::GetGroupCount(FIntVector(Parameters->AtlasResolution.X, Parameters->AtlasResolution.Y, SlotCount), FIntVector(8, 8, 1));
 		FComputeShaderUtils::AddPass(GraphBuilder, RDG_EVENT_NAME("IESAtlas::AddTexture"), ComputeShader, Parameters, DispatchCount);
 	}
+
+	GraphBuilder.UseExternalAccessMode(OutAtlas, ERHIAccess::SRVMask);
 }
 
 ///////////////////////////////////////////////////////////////////////////////////////////////////
diff --git a/Engine/Source/Runtime/Renderer/Private/MobileShadingRenderer.cpp b/Engine/Source/Runtime/Renderer/Private/MobileShadingRenderer.cpp
--- a/Engine/Source/Runtime/Renderer/Private/MobileShadingRenderer.cpp
+++ b/Engine/Source/Runtime/Renderer/Private/MobileShadingRenderer.cpp
@@ -951,13 +951,6 @@ void FMobileSceneRenderer::Render(FRDGBuilder& GraphBuilder)
 		}
 	}
 
-	// Generate the Sky/Atmosphere look up tables
-	const bool bShouldRenderSkyAtmosphere = ShouldRenderSkyAtmosphere(Scene, ViewFamily.EngineShowFlags);
-	if (bShouldRenderSkyAtmosphere)
-	{
-		RenderSkyAtmosphereLookUpTables(GraphBuilder);
-	}
-
 	// Notify the FX system that the scene is about to be rendered.
 	if (FXSystem && ViewFamily.EngineShowFlags.Particles)
 	{
@@ -967,9 +960,17 @@ void FMobileSceneRenderer::Render(FRDGBuilder& GraphBuilder)
 			GPUSortManager->OnPreRender(GraphBuilder);
 		}
 	}
-	
+
+	FRDGExternalAccessQueue ExternalAccessQueue;
+
+	// Generate the Sky/Atmosphere look up tables
+	const bool bShouldRenderSkyAtmosphere = ShouldRenderSkyAtmosphere(Scene, ViewFamily.EngineShowFlags);
+	if (bShouldRenderSkyAtmosphere)
+	{
+		RenderSkyAtmosphereLookUpTables(GraphBuilder, ExternalAccessQueue);
+	}
+
 	GraphBuilder.SetCommandListStat(GET_STATID(STAT_CLMM_Shadows));
-	FRDGExternalAccessQueue ExternalAccessQueue; // Ideally we'd use the shared access queue and pass that in, but keeping this minimal for now
 	RenderShadowDepthMaps(GraphBuilder, InstanceCullingManager, ExternalAccessQueue);
 	ExternalAccessQueue.Submit(GraphBuilder);
 	
diff --git a/Engine/Source/Runtime/Renderer/Private/ReflectionEnvironmentRealTimeCapture.cpp b/Engine/Source/Runtime/Renderer/Private/ReflectionEnvironmentRealTimeCapture.cpp
--- a/Engine/Source/Runtime/Renderer/Private/ReflectionEnvironmentRealTimeCapture.cpp
+++ b/Engine/Source/Runtime/Renderer/Private/ReflectionEnvironmentRealTimeCapture.cpp
@@ -294,7 +294,7 @@ END_SHADER_PARAMETER_STRUCT()
 
 void FScene::AllocateAndCaptureFrameSkyEnvMap(
 	FRDGBuilder& GraphBuilder, FSceneRenderer& SceneRenderer, FViewInfo& MainView,
-	bool bShouldRenderSkyAtmosphere, bool bShouldRenderVolumetricCloud, FInstanceCullingManager& InstanceCullingManager)
+	bool bShouldRenderSkyAtmosphere, bool bShouldRenderVolumetricCloud, FInstanceCullingManager& InstanceCullingManager, FRDGExternalAccessQueue& ExternalAccessQueue)
 {
 	check(SkyLight && SkyLight->bRealTimeCaptureEnabled && !SkyLight->bHasStaticLighting);
 
@@ -361,8 +361,6 @@ void FScene::AllocateAndCaptureFrameSkyEnvMap(
 		PF_FloatR11G11B10, FClearValueBinding::Black, TexCreate_TargetArraySlicesIndependently,
 		TexCreate_ShaderResource | TexCreate_UAV | TexCreate_RenderTargetable, false, 1, CubeMipCount, false);
 
-	FRDGExternalAccessQueue ExternalAccessQueue;
-
 	const bool bTimeSlicedRealTimeCapture = CVarRealTimeReflectionCaptureTimeSlicing.GetValueOnRenderThread() > 0 && !MainView.Family->bCurrentlyBeingEdited;
 
 	const bool CubeResolutionInvalidated = ConvolvedSkyRenderTargetReadyIndex < 0 || (ConvolvedSkyRenderTarget[ConvolvedSkyRenderTargetReadyIndex].IsValid() && ConvolvedSkyRenderTarget[ConvolvedSkyRenderTargetReadyIndex]->GetDesc().GetSize().X != CubeWidth);
@@ -1122,8 +1120,6 @@ void FScene::AllocateAndCaptureFrameSkyEnvMap(
 	{
 		ExternalAccessQueue.Add(GraphBuilder.RegisterExternalTexture(ConvolvedSkyRenderTarget[ConvolvedSkyRenderTargetReadyIndex]), ERHIAccess::SRVMask, ERHIPipeline::All);
 	}
-
-	ExternalAccessQueue.Submit(GraphBuilder);
 }
 
 
diff --git a/Engine/Source/Runtime/Renderer/Private/ScenePrivate.h b/Engine/Source/Runtime/Renderer/Private/ScenePrivate.h
--- a/Engine/Source/Runtime/Renderer/Private/ScenePrivate.h
+++ b/Engine/Source/Runtime/Renderer/Private/ScenePrivate.h
@@ -3127,7 +3127,7 @@ public:
 	virtual void AllocateReflectionCaptures(const TArray<UReflectionCaptureComponent*>& NewCaptures, const TCHAR* CaptureReason, bool bVerifyOnlyCapturing, bool bCapturingForMobile, bool bInsideTick) override;
 	virtual void ResetReflectionCaptures(bool bOnlyIfOOM) override;
 	virtual void UpdateSkyCaptureContents(const USkyLightComponent* CaptureComponent, bool bCaptureEmissiveOnly, UTextureCube* SourceCubemap, FTexture* OutProcessedTexture, float& OutAverageBrightness, FSHVectorRGB3& OutIrradianceEnvironmentMap, TArray<FFloat16Color>* OutRadianceMap, FLinearColor* SpecifiedCubemapColorScale) override;
-	virtual void AllocateAndCaptureFrameSkyEnvMap(FRDGBuilder& GraphBuilder, FSceneRenderer& SceneRenderer, FViewInfo& MainView, bool bShouldRenderSkyAtmosphere, bool bShouldRenderVolumetricCloud, FInstanceCullingManager& InstanceCullingManager) override;
+	virtual void AllocateAndCaptureFrameSkyEnvMap(FRDGBuilder& GraphBuilder, FSceneRenderer& SceneRenderer, FViewInfo& MainView, bool bShouldRenderSkyAtmosphere, bool bShouldRenderVolumetricCloud, FInstanceCullingManager& InstanceCullingManager, FRDGExternalAccessQueue& ExternalAccessQueue) override;
 	virtual void ValidateSkyLightRealTimeCapture(FRDGBuilder& GraphBuilder, const FViewInfo& View, FRDGTextureRef SceneColorTexture) override;
 	virtual void AddPrecomputedLightVolume(const class FPrecomputedLightVolume* Volume) override;
 	virtual void RemovePrecomputedLightVolume(const class FPrecomputedLightVolume* Volume) override;
diff --git a/Engine/Source/Runtime/Renderer/Private/SceneRendering.h b/Engine/Source/Runtime/Renderer/Private/SceneRendering.h
--- a/Engine/Source/Runtime/Renderer/Private/SceneRendering.h
+++ b/Engine/Source/Runtime/Renderer/Private/SceneRendering.h
@@ -2401,7 +2401,7 @@ protected:
 	void InitSkyAtmosphereForViews(FRHICommandListImmediate& RHICmdList);
 	
 	/** Render the sky atmosphere look up table needed for this frame.*/
-	void RenderSkyAtmosphereLookUpTables(FRDGBuilder& GraphBuilder);
+	void RenderSkyAtmosphereLookUpTables(FRDGBuilder& GraphBuilder, FRDGExternalAccessQueue& ExternalAccessQueue);
 
 	/** Render the sky atmosphere over the scene.*/
 	void RenderSkyAtmosphere(FRDGBuilder& GraphBuilder, const FMinimalSceneTextures& SceneTextures);
diff --git a/Engine/Source/Runtime/Renderer/Private/SkyAtmosphereRendering.cpp b/Engine/Source/Runtime/Renderer/Private/SkyAtmosphereRendering.cpp
--- a/Engine/Source/Runtime/Renderer/Private/SkyAtmosphereRendering.cpp
+++ b/Engine/Source/Runtime/Renderer/Private/SkyAtmosphereRendering.cpp
@@ -1297,7 +1297,7 @@ static void SetupSkyAtmosphereInternalCommonParameters(
 	InternalCommonParameters.CameraAerialPerspectiveVolumeSizeAndInvSize = GetSizeAndInvSize(CameraAerialPerspectiveVolumeScreenResolution, CameraAerialPerspectiveVolumeScreenResolution);
 }
 
-void FSceneRenderer::RenderSkyAtmosphereLookUpTables(FRDGBuilder& GraphBuilder)
+void FSceneRenderer::RenderSkyAtmosphereLookUpTables(FRDGBuilder& GraphBuilder, FRDGExternalAccessQueue& ExternalAccessQueue)
 {
 	check(ShouldRenderSkyAtmosphere(Scene, ViewFamily.EngineShowFlags)); // This should not be called if we should not render SkyAtmosphere
 
@@ -1323,8 +1323,6 @@ void FSceneRenderer::RenderSkyAtmosphereLookUpTables(FRDGBuilder& GraphBuilder)
 	FRDGTextureRef MultiScatteredLuminanceLut = GraphBuilder.RegisterExternalTexture(SkyInfo.GetMultiScatteredLuminanceLutTexture());
 	FRDGTextureUAVRef MultiScatteredLuminanceLutUAV = GraphBuilder.CreateUAV(FRDGTextureUAVDesc(MultiScatteredLuminanceLut, 0));
 
-	FRDGExternalAccessQueue ExternalAccessQueue;
-
 	ERDGPassFlags PassFlag = CVarSkyAtmosphereASyncCompute.GetValueOnAnyThread() ? ERDGPassFlags::AsyncCompute : ERDGPassFlags::Compute;
 
 	// Transmittance LUT
@@ -1710,8 +1708,6 @@ void FSceneRenderer::RenderSkyAtmosphereLookUpTables(FRDGBuilder& GraphBuilder)
 
 	SkyInfo.GetTransmittanceLutTexture() = ConvertToExternalAccessTexture(GraphBuilder, ExternalAccessQueue, TransmittanceLut);
 
-	ExternalAccessQueue.Submit(GraphBuilder);
-
 	// TODO have RDG execute those above passes with compute overlap similarly to using AutomaticCacheFlushAfterComputeShader(true);
 }
 
