Adding includes to rendering cpp files in preparation for header changes.

diff --git a/Engine/Plugins/Experimental/Water/Source/Runtime/Private/WaterUtils.cpp b/Engine/Plugins/Experimental/Water/Source/Runtime/Private/WaterUtils.cpp
--- a/Engine/Plugins/Experimental/Water/Source/Runtime/Private/WaterUtils.cpp
+++ b/Engine/Plugins/Experimental/Water/Source/Runtime/Private/WaterUtils.cpp
@@ -4,6 +4,7 @@
 #include "Materials/MaterialInstanceDynamic.h"
 #include "Engine/TextureRenderTarget2D.h"
 #include "UObject/Package.h"
+#include "RenderingThread.h"
 
 extern TAutoConsoleVariable<int32> CVarWaterEnabled;
 extern TAutoConsoleVariable<int32> CVarWaterMeshEnabled;
diff --git a/Engine/Plugins/Media/ImgMedia/Source/ImgMediaEditor/Private/Widgets/SImgMediaProcessEXR.cpp b/Engine/Plugins/Media/ImgMedia/Source/ImgMediaEditor/Private/Widgets/SImgMediaProcessEXR.cpp
--- a/Engine/Plugins/Media/ImgMedia/Source/ImgMediaEditor/Private/Widgets/SImgMediaProcessEXR.cpp
+++ b/Engine/Plugins/Media/ImgMedia/Source/ImgMediaEditor/Private/Widgets/SImgMediaProcessEXR.cpp
@@ -34,6 +34,7 @@
 #include "Widgets/Notifications/SNotificationList.h"
 #include "Widgets/Text/STextBlock.h"
 #include "TextureResource.h"
+#include "RenderingThread.h"
 
 #if IMGMEDIAEDITOR_EXR_SUPPORTED_PLATFORM
 #include "OpenExrWrapper.h"
diff --git a/Engine/Plugins/Runtime/Oculus/OculusVR/Source/OculusHMD/Private/OculusSceneCaptureCubemap.cpp b/Engine/Plugins/Runtime/Oculus/OculusVR/Source/OculusHMD/Private/OculusSceneCaptureCubemap.cpp
--- a/Engine/Plugins/Runtime/Oculus/OculusVR/Source/OculusHMD/Private/OculusSceneCaptureCubemap.cpp
+++ b/Engine/Plugins/Runtime/Oculus/OculusVR/Source/OculusHMD/Private/OculusSceneCaptureCubemap.cpp
@@ -14,6 +14,7 @@
 #include "HAL/FileManager.h"
 #include "Misc/FileHelper.h"
 #include "XRThreadUtils.h"
+#include "RenderingThread.h"
 
 //-------------------------------------------------------------------------------------------------
 // UOculusSceneCaptureCubemap
diff --git a/Engine/Plugins/VirtualProduction/CameraCalibrationCore/Source/CameraCalibrationCore/Private/LensFile.cpp b/Engine/Plugins/VirtualProduction/CameraCalibrationCore/Source/CameraCalibrationCore/Private/LensFile.cpp
--- a/Engine/Plugins/VirtualProduction/CameraCalibrationCore/Source/CameraCalibrationCore/Private/LensFile.cpp
+++ b/Engine/Plugins/VirtualProduction/CameraCalibrationCore/Source/CameraCalibrationCore/Private/LensFile.cpp
@@ -20,6 +20,7 @@
 #include "Tables/LensTableUtils.h"
 #include "TextureResource.h"
 #include "UObject/Package.h"
+#include "RenderingThread.h"
 
 namespace LensFileUtils
 {
diff --git a/Engine/Plugins/VirtualProduction/DMX/DMXFixtures/Source/DMXFixtures/Private/DMXFixtureActorMatrix.cpp b/Engine/Plugins/VirtualProduction/DMX/DMXFixtures/Source/DMXFixtures/Private/DMXFixtureActorMatrix.cpp
--- a/Engine/Plugins/VirtualProduction/DMX/DMXFixtures/Source/DMXFixtures/Private/DMXFixtureActorMatrix.cpp
+++ b/Engine/Plugins/VirtualProduction/DMX/DMXFixtures/Source/DMXFixtures/Private/DMXFixtureActorMatrix.cpp
@@ -10,6 +10,7 @@
 #include "Engine/Texture2D.h"
 #include "Rendering/Texture2DResource.h"
 #include "Components/StaticMeshComponent.h"
+#include "RenderingThread.h"
 
 DECLARE_CYCLE_STAT(TEXT("Fixture Actor Matrix Push Fixture Matrix Cell Data"), STAT_FixtureActorMatrixPushFixtureMatrixCellData, STATGROUP_DMX);
 
diff --git a/Engine/Source/Runtime/Apple/MetalRHI/Private/MetalRHIRenderQuery.cpp b/Engine/Source/Runtime/Apple/MetalRHI/Private/MetalRHIRenderQuery.cpp
--- a/Engine/Source/Runtime/Apple/MetalRHI/Private/MetalRHIRenderQuery.cpp
+++ b/Engine/Source/Runtime/Apple/MetalRHI/Private/MetalRHIRenderQuery.cpp
@@ -11,7 +11,7 @@
 #include "MetalLLM.h"
 #include "MetalCommandBuffer.h"
 #include "HAL/PThreadEvent.h"
-
+#include "RenderCore.h"
 
 //------------------------------------------------------------------------------
 
diff --git a/Engine/Source/Runtime/Engine/Private/BatchedElements.cpp b/Engine/Source/Runtime/Engine/Private/BatchedElements.cpp
--- a/Engine/Source/Runtime/Engine/Private/BatchedElements.cpp
+++ b/Engine/Source/Runtime/Engine/Private/BatchedElements.cpp
@@ -9,6 +9,7 @@
 #include "RHIStaticStates.h"
 #include "MeshPassProcessor.h"
 #include "SceneRelativeViewMatrices.h"
+#include "HDRHelper.h"
 
 DEFINE_LOG_CATEGORY_STATIC(LogBatchedElements, Log, All);
 
diff --git a/Engine/Source/Runtime/Engine/Private/CanvasRenderTarget2D.cpp b/Engine/Source/Runtime/Engine/Private/CanvasRenderTarget2D.cpp
--- a/Engine/Source/Runtime/Engine/Private/CanvasRenderTarget2D.cpp
+++ b/Engine/Source/Runtime/Engine/Private/CanvasRenderTarget2D.cpp
@@ -9,6 +9,7 @@
 #include "Engine/Canvas.h"
 #include "UObject/UObjectThreadContext.h"
 #include "TextureResource.h"
+#include "RenderingThread.h"
 
 #include UE_INLINE_GENERATED_CPP_BY_NAME(CanvasRenderTarget2D)
 
diff --git a/Engine/Source/Runtime/Engine/Private/Materials/MaterialShared.cpp b/Engine/Source/Runtime/Engine/Private/Materials/MaterialShared.cpp
--- a/Engine/Source/Runtime/Engine/Private/Materials/MaterialShared.cpp
+++ b/Engine/Source/Runtime/Engine/Private/Materials/MaterialShared.cpp
@@ -62,6 +62,7 @@
 #include "ODSC/ODSCManager.h"
 #endif
 #include "ProfilingDebugging/CountersTrace.h"
+#include "RenderCore.h"
 
 #define LOCTEXT_NAMESPACE "MaterialShared"
 
diff --git a/Engine/Source/Runtime/Engine/Private/Particles/ParticleSystemRender.cpp b/Engine/Source/Runtime/Engine/Private/Particles/ParticleSystemRender.cpp
--- a/Engine/Source/Runtime/Engine/Private/Particles/ParticleSystemRender.cpp
+++ b/Engine/Source/Runtime/Engine/Private/Particles/ParticleSystemRender.cpp
@@ -37,6 +37,7 @@
 #include "Particles/ParticleLODLevel.h"
 #include "Engine/StaticMesh.h"
 #include "UnrealEngine.h"
+#include "RenderCore.h"
 
 DECLARE_CYCLE_STAT(TEXT("ParticleSystemSceneProxy Create GT"), STAT_FParticleSystemSceneProxy_Create, STATGROUP_Particles);
 DECLARE_CYCLE_STAT(TEXT("ParticleSystemSceneProxy GetMeshElements RT"), STAT_FParticleSystemSceneProxy_GetMeshElements, STATGROUP_Particles);
diff --git a/Engine/Source/Runtime/Engine/Private/Rendering/Texture2DResource.cpp b/Engine/Source/Runtime/Engine/Private/Rendering/Texture2DResource.cpp
--- a/Engine/Source/Runtime/Engine/Private/Rendering/Texture2DResource.cpp
+++ b/Engine/Source/Runtime/Engine/Private/Rendering/Texture2DResource.cpp
@@ -9,6 +9,7 @@
 #include "Engine/Texture2D.h"
 #include "RenderUtils.h"
 #include "SceneTypes.h"
+#include "RenderingThread.h"
 
 // TODO Only adding this setting to allow backwards compatibility to be forced.  The default
 // behavior is to NOT do this.  This variable should be removed in the future.  #ADDED 4.13
diff --git a/Engine/Source/Runtime/Engine/Private/SkeletalRenderCPUSkin.cpp b/Engine/Source/Runtime/Engine/Private/SkeletalRenderCPUSkin.cpp
--- a/Engine/Source/Runtime/Engine/Private/SkeletalRenderCPUSkin.cpp
+++ b/Engine/Source/Runtime/Engine/Private/SkeletalRenderCPUSkin.cpp
@@ -15,6 +15,7 @@
 #include "GPUSkinVertexFactory.h"
 #include "Rendering/SkinWeightVertexBuffer.h"
 #include "BoneIndices.h"
+#include "RenderCore.h"
 
 struct FMorphTargetDelta;
 
diff --git a/Engine/Source/Runtime/Engine/Private/SkeletalRenderGPUSkin.cpp b/Engine/Source/Runtime/Engine/Private/SkeletalRenderGPUSkin.cpp
--- a/Engine/Source/Runtime/Engine/Private/SkeletalRenderGPUSkin.cpp
+++ b/Engine/Source/Runtime/Engine/Private/SkeletalRenderGPUSkin.cpp
@@ -19,6 +19,7 @@
 #include "HAL/LowLevelMemTracker.h"
 #include "RenderGraphUtils.h"
 #include "Interfaces/IPluginManager.h"
+#include "RenderCore.h"
 
 DEFINE_LOG_CATEGORY_STATIC(LogSkeletalGPUSkinMesh, Warning, All);
 
diff --git a/Engine/Source/Runtime/Engine/Private/StaticMeshRender.cpp b/Engine/Source/Runtime/Engine/Private/StaticMeshRender.cpp
--- a/Engine/Source/Runtime/Engine/Private/StaticMeshRender.cpp
+++ b/Engine/Source/Runtime/Engine/Private/StaticMeshRender.cpp
@@ -47,6 +47,7 @@
 #include "RayTracingInstance.h"
 #include "PrimitiveSceneInfo.h"
 #include "NaniteSceneProxy.h"
+#include "RenderCore.h"
 
 #if WITH_EDITOR
 #include "Rendering/StaticLightingSystemInterface.h"
diff --git a/Engine/Source/Runtime/Engine/Private/TextureDerivedData.cpp b/Engine/Source/Runtime/Engine/Private/TextureDerivedData.cpp
--- a/Engine/Source/Runtime/Engine/Private/TextureDerivedData.cpp
+++ b/Engine/Source/Runtime/Engine/Private/TextureDerivedData.cpp
@@ -32,6 +32,7 @@
 #include "VT/VirtualTextureBuiltData.h"
 #include "HAL/FileManager.h"
 #include "Misc/ConfigCacheIni.h"
+#include "RenderingThread.h"
 
 #if WITH_EDITOR
 
diff --git a/Engine/Source/Runtime/ImageWriteQueue/Private/ImageWriteBlueprintLibrary.cpp b/Engine/Source/Runtime/ImageWriteQueue/Private/ImageWriteBlueprintLibrary.cpp
--- a/Engine/Source/Runtime/ImageWriteQueue/Private/ImageWriteBlueprintLibrary.cpp
+++ b/Engine/Source/Runtime/ImageWriteQueue/Private/ImageWriteBlueprintLibrary.cpp
@@ -9,6 +9,7 @@
 #include "Engine/Texture2D.h"
 #include "Engine/TextureRenderTarget2D.h"
 #include "TextureResource.h"
+#include "RenderingThread.h"
 
 #include UE_INLINE_GENERATED_CPP_BY_NAME(ImageWriteBlueprintLibrary)
 
diff --git a/Engine/Source/Runtime/Landscape/Private/LandscapeRender.cpp b/Engine/Source/Runtime/Landscape/Private/LandscapeRender.cpp
--- a/Engine/Source/Runtime/Landscape/Private/LandscapeRender.cpp
+++ b/Engine/Source/Runtime/Landscape/Private/LandscapeRender.cpp
@@ -48,6 +48,7 @@ LandscapeRender.cpp: New terrain rendering
 #include "TextureResource.h"
 #include "NaniteSceneProxy.h"
 #include "Rendering/Texture2DResource.h"
+#include "RenderCore.h"
 
 
 IMPLEMENT_GLOBAL_SHADER_PARAMETER_STRUCT(FLandscapeUniformShaderParameters, "LandscapeParameters");
diff --git a/Engine/Source/Runtime/OpenGLDrv/Private/OpenGLQuery.cpp b/Engine/Source/Runtime/OpenGLDrv/Private/OpenGLQuery.cpp
--- a/Engine/Source/Runtime/OpenGLDrv/Private/OpenGLQuery.cpp
+++ b/Engine/Source/Runtime/OpenGLDrv/Private/OpenGLQuery.cpp
@@ -8,6 +8,7 @@
 #include "Stats/Stats.h"
 #include "OpenGLDrv.h"
 #include "OpenGLDrvPrivate.h"
+#include "RenderCore.h"
 
 static int32 GOpenGLPollRenderQueryResult = 1;
 static FAutoConsoleVariableRef CVarOpenGLPollRenderQueryResult(
diff --git a/Engine/Source/Runtime/RenderCore/Private/DumpGPU.cpp b/Engine/Source/Runtime/RenderCore/Private/DumpGPU.cpp
--- a/Engine/Source/Runtime/RenderCore/Private/DumpGPU.cpp
+++ b/Engine/Source/Runtime/RenderCore/Private/DumpGPU.cpp
@@ -19,6 +19,7 @@
 #include "GlobalShader.h"
 #include "RHIValidation.h"
 #include "RenderGraphPrivate.h"
+#include "RenderCore.h"
 
 IDumpGPUUploadServiceProvider* IDumpGPUUploadServiceProvider::GProvider = nullptr;
 
diff --git a/Engine/Source/Runtime/RenderCore/Private/DynamicBufferAllocator.cpp b/Engine/Source/Runtime/RenderCore/Private/DynamicBufferAllocator.cpp
--- a/Engine/Source/Runtime/RenderCore/Private/DynamicBufferAllocator.cpp
+++ b/Engine/Source/Runtime/RenderCore/Private/DynamicBufferAllocator.cpp
@@ -7,6 +7,7 @@ DynamicBufferAllocator.cpp: Classes for allocating transient rendering data.
 #include "DynamicBufferAllocator.h"
 #include "RenderResource.h"
 #include "Misc/ScopeLock.h"
+#include "RenderCore.h"
 
 int32 GMaxReadBufferRenderingBytesAllocatedPerFrame = 32 * 1024 * 1024;
 
diff --git a/Engine/Source/Runtime/RenderCore/Private/RenderTargetPool.cpp b/Engine/Source/Runtime/RenderCore/Private/RenderTargetPool.cpp
--- a/Engine/Source/Runtime/RenderCore/Private/RenderTargetPool.cpp
+++ b/Engine/Source/Runtime/RenderCore/Private/RenderTargetPool.cpp
@@ -6,6 +6,7 @@
 #include "Hash/CityHash.h"
 #include "Trace/Trace.inl"
 #include "ProfilingDebugging/CountersTrace.h"
+#include "RenderCore.h"
 
 /** The global render targets pool. */
 TGlobalResource<FRenderTargetPool> GRenderTargetPool;
diff --git a/Engine/Source/Runtime/RenderCore/Private/RenderUtils.cpp b/Engine/Source/Runtime/RenderCore/Private/RenderUtils.cpp
--- a/Engine/Source/Runtime/RenderCore/Private/RenderUtils.cpp
+++ b/Engine/Source/Runtime/RenderCore/Private/RenderUtils.cpp
@@ -10,6 +10,7 @@
 #include "Misc/ConfigCacheIni.h"
 #include "RenderGraphResourcePool.h"
 #include "Misc/DataDrivenPlatformInfoRegistry.h"
+#include "RenderCore.h"
 
 #if WITH_EDITOR
 #include "Misc/CoreMisc.h"
diff --git a/Engine/Source/Runtime/RenderCore/Private/VisualizeTexture.cpp b/Engine/Source/Runtime/RenderCore/Private/VisualizeTexture.cpp
--- a/Engine/Source/Runtime/RenderCore/Private/VisualizeTexture.cpp
+++ b/Engine/Source/Runtime/RenderCore/Private/VisualizeTexture.cpp
@@ -11,6 +11,7 @@
 #include "CommonRenderResources.h"
 #include "PixelShaderUtils.h"
 #include "Misc/FileHelper.h"
+#include "RenderCore.h"
 
 void FVisualizeTexture::ParseCommands(const TCHAR* Cmd, FOutputDevice &Ar)
 {
diff --git a/Engine/Source/Runtime/Renderer/Private/AnisotropyRendering.cpp b/Engine/Source/Runtime/Renderer/Private/AnisotropyRendering.cpp
--- a/Engine/Source/Runtime/Renderer/Private/AnisotropyRendering.cpp
+++ b/Engine/Source/Runtime/Renderer/Private/AnisotropyRendering.cpp
@@ -5,6 +5,7 @@
 #include "MeshPassProcessor.inl"
 #include "ScenePrivate.h"
 #include "DeferredShadingRenderer.h"
+#include "RenderCore.h"
 
 DECLARE_GPU_STAT_NAMED(RenderAnisotropyPass, TEXT("Render Anisotropy Pass"));
 
diff --git a/Engine/Source/Runtime/Renderer/Private/BasePassRendering.cpp b/Engine/Source/Runtime/Renderer/Private/BasePassRendering.cpp
--- a/Engine/Source/Runtime/Renderer/Private/BasePassRendering.cpp
+++ b/Engine/Source/Runtime/Renderer/Private/BasePassRendering.cpp
@@ -21,6 +21,7 @@
 #include "DebugProbeRendering.h"
 #include "AnisotropyRendering.h"
 #include "Nanite/NaniteVisualize.h"
+#include "RenderCore.h"
 
 // Changing this causes a full shader recompile
 static TAutoConsoleVariable<int32> CVarSelectiveBasePassOutputs(
diff --git a/Engine/Source/Runtime/Renderer/Private/CompositionLighting/PostProcessDeferredDecals.cpp b/Engine/Source/Runtime/Renderer/Private/CompositionLighting/PostProcessDeferredDecals.cpp
--- a/Engine/Source/Runtime/Renderer/Private/CompositionLighting/PostProcessDeferredDecals.cpp
+++ b/Engine/Source/Runtime/Renderer/Private/CompositionLighting/PostProcessDeferredDecals.cpp
@@ -12,6 +12,7 @@
 #include "ScenePrivate.h"
 #include "VelocityRendering.h"
 #include "VisualizeTexture.h"
+#include "RenderCore.h"
 
 static TAutoConsoleVariable<float> CVarStencilSizeThreshold(
 	TEXT("r.Decal.StencilSizeThreshold"),
diff --git a/Engine/Source/Runtime/Renderer/Private/DeferredShadingRenderer.cpp b/Engine/Source/Runtime/Renderer/Private/DeferredShadingRenderer.cpp
--- a/Engine/Source/Runtime/Renderer/Private/DeferredShadingRenderer.cpp
+++ b/Engine/Source/Runtime/Renderer/Private/DeferredShadingRenderer.cpp
@@ -82,6 +82,7 @@
 #include "Shadows/ShadowSceneRenderer.h"
 #include "HeterogeneousVolumes/HeterogeneousVolumes.h"
 #include "ComponentRecreateRenderStateContext.h"
+#include "RenderCore.h"
 
 extern int32 GNaniteShowStats;
 extern int32 GNanitePickingDomain;
diff --git a/Engine/Source/Runtime/Renderer/Private/DepthRendering.cpp b/Engine/Source/Runtime/Renderer/Private/DepthRendering.cpp
--- a/Engine/Source/Runtime/Renderer/Private/DepthRendering.cpp
+++ b/Engine/Source/Runtime/Renderer/Private/DepthRendering.cpp
@@ -31,6 +31,7 @@
 #include "RenderGraphUtils.h"
 #include "SceneRenderingUtils.h"
 #include "DebugProbeRendering.h"
+#include "RenderCore.h"
 
 static TAutoConsoleVariable<int32> CVarParallelPrePass(
 	TEXT("r.ParallelPrePass"),
diff --git a/Engine/Source/Runtime/Renderer/Private/GPUScene.cpp b/Engine/Source/Runtime/Renderer/Private/GPUScene.cpp
--- a/Engine/Source/Runtime/Renderer/Private/GPUScene.cpp
+++ b/Engine/Source/Runtime/Renderer/Private/GPUScene.cpp
@@ -22,6 +22,7 @@
 #include "HAL/LowLevelMemStats.h"
 #include "InstanceUniformShaderParameters.h"
 #include "ShaderPrint.h"
+#include "RenderCore.h"
 #include "LightSceneData.h"
 
 #define LOG_INSTANCE_ALLOCATIONS 0
diff --git a/Engine/Source/Runtime/Renderer/Private/HairStrands/HairStrandsRasterCommon.cpp b/Engine/Source/Runtime/Renderer/Private/HairStrands/HairStrandsRasterCommon.cpp
--- a/Engine/Source/Runtime/Renderer/Private/HairStrands/HairStrandsRasterCommon.cpp
+++ b/Engine/Source/Runtime/Renderer/Private/HairStrands/HairStrandsRasterCommon.cpp
@@ -10,6 +10,7 @@
 #include "MeshPassProcessor.h"
 #include "MeshPassProcessor.inl"
 #include "ScenePrivate.h"
+#include "RenderCore.h"
 
 IMPLEMENT_STATIC_UNIFORM_BUFFER_STRUCT(FHairDeepShadowRasterUniformParameters, "DeepRasterPass", SceneTextures);
 
diff --git a/Engine/Source/Runtime/Renderer/Private/IndirectLightingCache.cpp b/Engine/Source/Runtime/Renderer/Private/IndirectLightingCache.cpp
--- a/Engine/Source/Runtime/Renderer/Private/IndirectLightingCache.cpp
+++ b/Engine/Source/Runtime/Renderer/Private/IndirectLightingCache.cpp
@@ -17,6 +17,7 @@
 #include "DynamicPrimitiveDrawing.h"
 #include "ScenePrivate.h"
 #include "PrecomputedLightVolume.h"
+#include "RenderCore.h"
 
 /** 
  * Primitive bounds size will be rounded up to the next value of Pow(BoundSizeRoundUpBase, N) and N is an integer. 
diff --git a/Engine/Source/Runtime/Renderer/Private/InstanceCulling/InstanceCullingContext.cpp b/Engine/Source/Runtime/Renderer/Private/InstanceCulling/InstanceCullingContext.cpp
--- a/Engine/Source/Runtime/Renderer/Private/InstanceCulling/InstanceCullingContext.cpp
+++ b/Engine/Source/Runtime/Renderer/Private/InstanceCulling/InstanceCullingContext.cpp
@@ -13,6 +13,7 @@
 #include "InstanceCulling/InstanceCullingManager.h"
 #include "InstanceCullingLoadBalancer.h"
 #include "InstanceCullingMergedContext.h"
+#include "RenderCore.h"
 
 static TAutoConsoleVariable<int32> CVarCullInstances(
 	TEXT("r.CullInstances"),
diff --git a/Engine/Source/Runtime/Renderer/Private/LightMapRendering.cpp b/Engine/Source/Runtime/Renderer/Private/LightMapRendering.cpp
--- a/Engine/Source/Runtime/Renderer/Private/LightMapRendering.cpp
+++ b/Engine/Source/Runtime/Renderer/Private/LightMapRendering.cpp
@@ -7,6 +7,7 @@ LightMapRendering.cpp: Light map rendering implementations.
 #include "LightMapRendering.h"
 #include "ScenePrivate.h"
 #include "PrecomputedVolumetricLightmap.h"
+#include "RenderCore.h"
 
 IMPLEMENT_TYPE_LAYOUT(FUniformLightMapPolicyShaderParametersType);
 
diff --git a/Engine/Source/Runtime/Renderer/Private/LightRendering.cpp b/Engine/Source/Runtime/Renderer/Private/LightRendering.cpp
--- a/Engine/Source/Runtime/Renderer/Private/LightRendering.cpp
+++ b/Engine/Source/Runtime/Renderer/Private/LightRendering.cpp
@@ -26,6 +26,7 @@
 #include "AnisotropyRendering.h"
 #include "Engine/SubsurfaceProfile.h"
 #include "Shadows/ShadowSceneRenderer.h"
+#include "RenderCore.h"
 
 // ENABLE_DEBUG_DISCARD_PROP is used to test the lighting code by allowing to discard lights to see how performance scales
 // It ought never to be enabled in a shipping build, and is probably only really useful when woring on the shading code.
diff --git a/Engine/Source/Runtime/Renderer/Private/LightShaftRendering.cpp b/Engine/Source/Runtime/Renderer/Private/LightShaftRendering.cpp
--- a/Engine/Source/Runtime/Renderer/Private/LightShaftRendering.cpp
+++ b/Engine/Source/Runtime/Renderer/Private/LightShaftRendering.cpp
@@ -9,6 +9,7 @@
 #include "SceneTextureParameters.h"
 #include "ScenePrivate.h"
 #include "DeferredShadingRenderer.h"
+#include "RenderCore.h"
 
 int32 GLightShafts = 1;
 static FAutoConsoleVariableRef CVarLightShaftQuality(
diff --git a/Engine/Source/Runtime/Renderer/Private/MeshPassProcessor.cpp b/Engine/Source/Runtime/Renderer/Private/MeshPassProcessor.cpp
--- a/Engine/Source/Runtime/Renderer/Private/MeshPassProcessor.cpp
+++ b/Engine/Source/Runtime/Renderer/Private/MeshPassProcessor.cpp
@@ -17,6 +17,7 @@
 #include "RayTracing/RayTracingMaterialHitShaders.h"
 #include "Hash/CityHash.h"
 #include "ComponentRecreateRenderStateContext.h"
+#include "RenderCore.h"
 
 FRWLock FGraphicsMinimalPipelineStateId::PersistentIdTableLock;
 FGraphicsMinimalPipelineStateId::PersistentTableType FGraphicsMinimalPipelineStateId::PersistentIdTable;
diff --git a/Engine/Source/Runtime/Renderer/Private/MobileBasePassRendering.cpp b/Engine/Source/Runtime/Renderer/Private/MobileBasePassRendering.cpp
--- a/Engine/Source/Runtime/Renderer/Private/MobileBasePassRendering.cpp
+++ b/Engine/Source/Runtime/Renderer/Private/MobileBasePassRendering.cpp
@@ -20,6 +20,7 @@
 #include "Engine/SubsurfaceProfile.h"
 #include "LocalLightSceneProxy.h"
 #include "ReflectionEnvironment.h"
+#include "RenderCore.h"
 
 // Changing this causes a full shader recompile
 static TAutoConsoleVariable<int32> CVarMobileDisableVertexFog(
diff --git a/Engine/Source/Runtime/Renderer/Private/MobileDecalRendering.cpp b/Engine/Source/Runtime/Renderer/Private/MobileDecalRendering.cpp
--- a/Engine/Source/Runtime/Renderer/Private/MobileDecalRendering.cpp
+++ b/Engine/Source/Runtime/Renderer/Private/MobileDecalRendering.cpp
@@ -15,6 +15,7 @@
 #include "ScenePrivate.h"
 #include "DecalRenderingCommon.h"
 #include "DecalRenderingShared.h"
+#include "RenderCore.h"
 
 void RenderMeshDecalsMobile(FRHICommandListImmediate& RHICmdList, const FViewInfo& View, EDecalRenderStage DecalRenderStage, EDecalRenderTargetMode RenderTargetMode);
 extern void RenderDeferredDecalsMobile(FRHICommandListImmediate& RHICmdList, const FScene& Scene, const FViewInfo& View, EDecalRenderStage DecalRenderStage, EDecalRenderTargetMode RenderTargetMode);
diff --git a/Engine/Source/Runtime/Renderer/Private/MobileShadingRenderer.cpp b/Engine/Source/Runtime/Renderer/Private/MobileShadingRenderer.cpp
--- a/Engine/Source/Runtime/Renderer/Private/MobileShadingRenderer.cpp
+++ b/Engine/Source/Runtime/Renderer/Private/MobileShadingRenderer.cpp
@@ -59,6 +59,7 @@
 #include "SceneTextureReductions.h"
 #include "GPUMessaging.h"
 #include "Strata/Strata.h"
+#include "RenderCore.h"
 #include "RectLightTextureManager.h"
 #include "IESTextureManager.h"
 
diff --git a/Engine/Source/Runtime/Renderer/Private/MobileTranslucentRendering.cpp b/Engine/Source/Runtime/Renderer/Private/MobileTranslucentRendering.cpp
--- a/Engine/Source/Runtime/Renderer/Private/MobileTranslucentRendering.cpp
+++ b/Engine/Source/Runtime/Renderer/Private/MobileTranslucentRendering.cpp
@@ -1,35 +1,8 @@
 // Copyright Epic Games, Inc. All Rights Reserved.
 
-/*=============================================================================
-	MobileTranslucentRendering.cpp: translucent rendering implementation.
-=============================================================================*/
-
-#include "CoreMinimal.h"
-#include "HAL/IConsoleManager.h"
-#include "RHI.h"
-#include "HitProxies.h"
-#include "Shader.h"
-#include "StaticBoundShaderState.h"
-#include "SceneUtils.h"
-#include "RHIStaticStates.h"
-#include "PostProcess/SceneRenderTargets.h"
-#include "GlobalShader.h"
-#include "SceneRenderTargetParameters.h"
 #include "SceneRendering.h"
-#include "LightMapRendering.h"
-#include "MaterialShaderType.h"
-#include "MeshMaterialShaderType.h"
-#include "MeshMaterialShader.h"
-#include "BasePassRendering.h"
-#include "DynamicPrimitiveDrawing.h"
+#include "RenderCore.h"
 #include "TranslucentRendering.h"
-#include "MobileBasePassRendering.h"
-#include "ScenePrivate.h"
-#include "ScreenRendering.h"
-#include "PostProcess/SceneFilterRendering.h"
-#include "PipelineStateCache.h"
-#include "MeshPassProcessor.inl"
-#include "ClearQuad.h"
 
 void FMobileSceneRenderer::RenderTranslucency(FRHICommandListImmediate& RHICmdList, const FViewInfo& View)
 {	
diff --git a/Engine/Source/Runtime/Renderer/Private/PrimitiveSceneInfo.cpp b/Engine/Source/Runtime/Renderer/Private/PrimitiveSceneInfo.cpp
--- a/Engine/Source/Runtime/Renderer/Private/PrimitiveSceneInfo.cpp
+++ b/Engine/Source/Runtime/Renderer/Private/PrimitiveSceneInfo.cpp
@@ -25,6 +25,7 @@
 #include "NaniteSceneProxy.h"
 #include "Lumen/LumenSceneRendering.h"
 #include "RayTracingDefinitions.h"
+#include "RenderCore.h"
 
 extern int32 GGPUSceneInstanceClearList;
 extern int32 GGPUSceneInstanceBVH;
diff --git a/Engine/Source/Runtime/Renderer/Private/RayTracing/RayTracingInstanceBufferUtil.cpp b/Engine/Source/Runtime/Renderer/Private/RayTracing/RayTracingInstanceBufferUtil.cpp
--- a/Engine/Source/Runtime/Renderer/Private/RayTracing/RayTracingInstanceBufferUtil.cpp
+++ b/Engine/Source/Runtime/Renderer/Private/RayTracing/RayTracingInstanceBufferUtil.cpp
@@ -10,6 +10,7 @@
 #include "RenderGraphBuilder.h"
 #include "ShaderParameterUtils.h"
 #include "RendererInterface.h"
+#include "RenderCore.h"
 #include "ShaderParameterStruct.h"
 #include "GlobalShader.h"
 
diff --git a/Engine/Source/Runtime/Renderer/Private/RayTracing/RayTracingMaterialHitShaders.cpp b/Engine/Source/Runtime/Renderer/Private/RayTracing/RayTracingMaterialHitShaders.cpp
--- a/Engine/Source/Runtime/Renderer/Private/RayTracing/RayTracingMaterialHitShaders.cpp
+++ b/Engine/Source/Runtime/Renderer/Private/RayTracing/RayTracingMaterialHitShaders.cpp
@@ -1,10 +1,12 @@
 // Copyright Epic Games, Inc. All Rights Reserved.
 
 #include "RayTracing/RayTracingMaterialHitShaders.h"
+
+#if RHI_RAYTRACING
 #include "DeferredShadingRenderer.h"
 #include "PipelineStateCache.h"
+#include "RenderCore.h"
 
-#if RHI_RAYTRACING
 #include "Nanite/NaniteRayTracing.h"
 
 #include "RayTracingDefinitions.h"
diff --git a/Engine/Source/Runtime/Renderer/Private/RendererScene.cpp b/Engine/Source/Runtime/Renderer/Private/RendererScene.cpp
--- a/Engine/Source/Runtime/Renderer/Private/RendererScene.cpp
+++ b/Engine/Source/Runtime/Renderer/Private/RendererScene.cpp
@@ -68,6 +68,7 @@
 #include "VelocityRendering.h"
 #include "RectLightSceneProxy.h"
 #include "RectLightTextureManager.h"
+#include "RenderCore.h"
 #include "IESTextureManager.h"
 
 #if RHI_RAYTRACING
diff --git a/Engine/Source/Runtime/Renderer/Private/SceneOcclusion.cpp b/Engine/Source/Runtime/Renderer/Private/SceneOcclusion.cpp
--- a/Engine/Source/Runtime/Renderer/Private/SceneOcclusion.cpp
+++ b/Engine/Source/Runtime/Renderer/Private/SceneOcclusion.cpp
@@ -20,6 +20,7 @@
 #include "VisualizeTexture.h"
 #include "ProfilingDebugging/CpuProfilerTrace.h"
 #include "PixelShaderUtils.h"
+#include "RenderCore.h"
 
 
 /*-----------------------------------------------------------------------------
diff --git a/Engine/Source/Runtime/Renderer/Private/SceneRendering.cpp b/Engine/Source/Runtime/Renderer/Private/SceneRendering.cpp
--- a/Engine/Source/Runtime/Renderer/Private/SceneRendering.cpp
+++ b/Engine/Source/Runtime/Renderer/Private/SceneRendering.cpp
@@ -81,6 +81,7 @@
 #include "NaniteVisualizationData.h"
 #include "Shadows/ShadowSceneRenderer.h"
 #include "TextureResource.h"
+#include "RenderCore.h"
 
 /*-----------------------------------------------------------------------------
 	Globals
diff --git a/Engine/Source/Runtime/Renderer/Private/SceneVisibility.cpp b/Engine/Source/Runtime/Renderer/Private/SceneVisibility.cpp
--- a/Engine/Source/Runtime/Renderer/Private/SceneVisibility.cpp
+++ b/Engine/Source/Runtime/Renderer/Private/SceneVisibility.cpp
@@ -47,6 +47,7 @@
 #include "RayTracing/RayTracingInstanceCulling.h"
 #include "RendererModule.h"
 #include "SceneViewExtension.h"
+#include "RenderCore.h"
 
 #if !UE_BUILD_SHIPPING
 #include "ViewDebug.h"
diff --git a/Engine/Source/Runtime/Renderer/Private/ShadowDepthRendering.cpp b/Engine/Source/Runtime/Renderer/Private/ShadowDepthRendering.cpp
--- a/Engine/Source/Runtime/Renderer/Private/ShadowDepthRendering.cpp
+++ b/Engine/Source/Runtime/Renderer/Private/ShadowDepthRendering.cpp
@@ -47,6 +47,7 @@
 #include "Rendering/NaniteResources.h"
 #include "Rendering/NaniteStreamingManager.h"
 #include "Shadows/ShadowSceneRenderer.h"
+#include "RenderCore.h"
 
 DEFINE_GPU_DRAWCALL_STAT(ShadowDepths);
 
diff --git a/Engine/Source/Runtime/Renderer/Private/ShadowRendering.cpp b/Engine/Source/Runtime/Renderer/Private/ShadowRendering.cpp
--- a/Engine/Source/Runtime/Renderer/Private/ShadowRendering.cpp
+++ b/Engine/Source/Runtime/Renderer/Private/ShadowRendering.cpp
@@ -15,6 +15,7 @@
 #include "HairStrands/HairStrandsRendering.h"
 #include "VirtualShadowMaps/VirtualShadowMapProjection.h"
 #include "Shadows/ShadowSceneRenderer.h"
+#include "RenderCore.h"
 
 ///////////////////////////////////////////////////////////////////////////////////////////////////
 // Directional light
diff --git a/Engine/Source/Runtime/Renderer/Private/ShadowSetup.cpp b/Engine/Source/Runtime/Renderer/Private/ShadowSetup.cpp
--- a/Engine/Source/Runtime/Renderer/Private/ShadowSetup.cpp
+++ b/Engine/Source/Runtime/Renderer/Private/ShadowSetup.cpp
@@ -34,6 +34,7 @@
 #include "InstanceCulling/InstanceCullingManager.h"
 #include "Shadows/ShadowSceneRenderer.h"
 #include "Lumen/Lumen.h"
+#include "RenderCore.h"
 
 /** Number of cube map shadow depth surfaces that will be created and used for rendering one pass point light shadows. */
 static const int32 NumCubeShadowDepthSurfaces = 5;
diff --git a/Engine/Source/Runtime/Renderer/Private/ShadowSetupMobile.cpp b/Engine/Source/Runtime/Renderer/Private/ShadowSetupMobile.cpp
--- a/Engine/Source/Runtime/Renderer/Private/ShadowSetupMobile.cpp
+++ b/Engine/Source/Runtime/Renderer/Private/ShadowSetupMobile.cpp
@@ -17,6 +17,7 @@ ShadowSetupMobile.cpp: Shadow setup implementation for mobile specific features.
 #include "ScenePrivate.h"
 #include "MeshPassProcessor.h"
 #include "UObject/UObjectIterator.h"
+#include "RenderCore.h"
 
 static TAutoConsoleVariable<int32> CVarCsmShaderCullingDebugGfx(
 	TEXT("r.Mobile.Shadow.CSMShaderCullingDebugGfx"),
diff --git a/Engine/Source/Runtime/Renderer/Private/SingleLayerWaterRendering.cpp b/Engine/Source/Runtime/Renderer/Private/SingleLayerWaterRendering.cpp
--- a/Engine/Source/Runtime/Renderer/Private/SingleLayerWaterRendering.cpp
+++ b/Engine/Source/Runtime/Renderer/Private/SingleLayerWaterRendering.cpp
@@ -21,6 +21,7 @@
 #include "VirtualShadowMaps/VirtualShadowMapProjection.h"
 #include "Lumen/LumenSceneData.h"
 #include "Lumen/LumenTracingUtils.h"
+#include "RenderCore.h"
 
 DECLARE_GPU_STAT_NAMED(RayTracingWaterReflections, TEXT("Ray Tracing Water Reflections"));
 
diff --git a/Engine/Source/Runtime/Renderer/Private/TranslucentLighting.cpp b/Engine/Source/Runtime/Renderer/Private/TranslucentLighting.cpp
--- a/Engine/Source/Runtime/Renderer/Private/TranslucentLighting.cpp
+++ b/Engine/Source/Runtime/Renderer/Private/TranslucentLighting.cpp
@@ -48,6 +48,7 @@
 #include "MeshPassProcessor.inl"
 #include "SkyAtmosphereRendering.h"
 #include "VolumetricCloudRendering.h"
+#include "RenderCore.h"
 
 class FMaterial;
 
diff --git a/Engine/Source/Runtime/Renderer/Private/VelocityRendering.cpp b/Engine/Source/Runtime/Renderer/Private/VelocityRendering.cpp
--- a/Engine/Source/Runtime/Renderer/Private/VelocityRendering.cpp
+++ b/Engine/Source/Runtime/Renderer/Private/VelocityRendering.cpp
@@ -23,6 +23,7 @@
 #include "MeshPassProcessor.inl"
 #include "DebugProbeRendering.h"
 #include "RendererModule.h"
+#include "RenderCore.h"
 
 // Changing this causes a full shader recompile
 static TAutoConsoleVariable<int32> CVarVelocityOutputPass(
diff --git a/Engine/Source/Runtime/SlateRHIRenderer/Private/SlateRHIRenderer.cpp b/Engine/Source/Runtime/SlateRHIRenderer/Private/SlateRHIRenderer.cpp
--- a/Engine/Source/Runtime/SlateRHIRenderer/Private/SlateRHIRenderer.cpp
+++ b/Engine/Source/Runtime/SlateRHIRenderer/Private/SlateRHIRenderer.cpp
@@ -37,6 +37,7 @@
 #include "Rendering/RenderingCommon.h"
 #include "IHeadMountedDisplayModule.h"
 #include "HDRHelper.h"
+#include "RenderCore.h"
 
 DECLARE_CYCLE_STAT(TEXT("Slate RT: Rendering"), STAT_SlateRenderingRTTime, STATGROUP_Slate);
 
diff --git a/Engine/Source/Runtime/UMG/Private/Blueprint/AsyncTaskDownloadImage.cpp b/Engine/Source/Runtime/UMG/Private/Blueprint/AsyncTaskDownloadImage.cpp
--- a/Engine/Source/Runtime/UMG/Private/Blueprint/AsyncTaskDownloadImage.cpp
+++ b/Engine/Source/Runtime/UMG/Private/Blueprint/AsyncTaskDownloadImage.cpp
@@ -9,6 +9,7 @@
 #include "Interfaces/IHttpResponse.h"
 #include "HttpModule.h"
 #include "TextureResource.h"
+#include "RenderingThread.h"
 
 #include UE_INLINE_GENERATED_CPP_BY_NAME(AsyncTaskDownloadImage)
 
diff --git a/Engine/Source/Runtime/VulkanRHI/Private/VulkanQuery.cpp b/Engine/Source/Runtime/VulkanRHI/Private/VulkanQuery.cpp
--- a/Engine/Source/Runtime/VulkanRHI/Private/VulkanQuery.cpp
+++ b/Engine/Source/Runtime/VulkanRHI/Private/VulkanQuery.cpp
@@ -9,6 +9,7 @@
 #include "VulkanContext.h"
 #include "VulkanCommandBuffer.h"
 #include "EngineGlobals.h"
+#include "RenderCore.h"
 
 #if VULKAN_QUERY_CALLSTACK
 #include "HAL/PlatformStackwalk.h"
diff --git a/Engine/Source/Runtime/VulkanRHI/Private/VulkanUtil.cpp b/Engine/Source/Runtime/VulkanRHI/Private/VulkanUtil.cpp
--- a/Engine/Source/Runtime/VulkanRHI/Private/VulkanUtil.cpp
+++ b/Engine/Source/Runtime/VulkanRHI/Private/VulkanUtil.cpp
@@ -12,6 +12,7 @@
 #include "Misc/OutputDeviceRedirector.h"
 #include "RHIValidationContext.h"
 #include "HAL/FileManager.h"
+#include "RenderCore.h"
 
 #if NV_AFTERMATH
 #include "GFSDK_Aftermath_GpuCrashDump.h"
diff --git a/Engine/Source/Runtime/Windows/D3D11RHI/Private/D3D11Query.cpp b/Engine/Source/Runtime/Windows/D3D11RHI/Private/D3D11Query.cpp
--- a/Engine/Source/Runtime/Windows/D3D11RHI/Private/D3D11Query.cpp
+++ b/Engine/Source/Runtime/Windows/D3D11RHI/Private/D3D11Query.cpp
@@ -5,6 +5,7 @@
 =============================================================================*/
 
 #include "D3D11RHIPrivate.h"
+#include "RenderCore.h"
 
 class FD3D11RenderQueryBatcher final
 {
