Removed unused Broadcast/WaitForEyeAdaptationTemporalEffect.

diff --git a/Engine/Source/Runtime/Renderer/Private/PostProcess/PostProcessEyeAdaptation.cpp b/Engine/Source/Runtime/Renderer/Private/PostProcess/PostProcessEyeAdaptation.cpp
--- a/Engine/Source/Runtime/Renderer/Private/PostProcess/PostProcessEyeAdaptation.cpp
+++ b/Engine/Source/Runtime/Renderer/Private/PostProcess/PostProcessEyeAdaptation.cpp
@@ -1171,21 +1171,3 @@ void FSceneViewState::UpdatePreExposure(FViewInfo& View)
 		PrevFrameViewInfo.SceneColorPreExposure = PreExposure;
 	}
 }
-
-#if WITH_MGPU
-static const FName NAME_EyeAdaptation(TEXT("EyeAdaptation"));
-
-void FSceneViewState::BroadcastEyeAdaptationTemporalEffect(FRHICommandList& RHICmdList)
-{
-	FRHITexture* EyeAdaptation = GetCurrentEyeAdaptationTexture(RHICmdList)->GetRHI();
-	RHICmdList.BroadcastTemporalEffect(FName(NAME_EyeAdaptation, UniqueID), { &EyeAdaptation, 1 });
-}
-
-DECLARE_GPU_STAT(AFRWaitForEyeAdaptation);
-
-void FSceneViewState::WaitForEyeAdaptationTemporalEffect(FRHICommandList& RHICmdList)
-{
-	SCOPED_GPU_STAT(RHICmdList, AFRWaitForEyeAdaptation);
-	RHICmdList.WaitForTemporalEffect(FName(NAME_EyeAdaptation, UniqueID));
-}
-#endif // WITH_MGPU
diff --git a/Engine/Source/Runtime/Renderer/Private/ScenePrivate.h b/Engine/Source/Runtime/Renderer/Private/ScenePrivate.h
--- a/Engine/Source/Runtime/Renderer/Private/ScenePrivate.h
+++ b/Engine/Source/Runtime/Renderer/Private/ScenePrivate.h
@@ -1251,11 +1251,6 @@ public:
 		}
 	}
 
-#if WITH_MGPU
-	void BroadcastEyeAdaptationTemporalEffect(FRHICommandList& RHICmdList);
-	void WaitForEyeAdaptationTemporalEffect(FRHICommandList& RHICmdList);
-#endif
-
 	float GetLastEyeAdaptationExposure() const
 	{
 		return EyeAdaptationManager.GetLastExposure();
diff --git a/Engine/Source/Runtime/Renderer/Private/SceneRendering.cpp b/Engine/Source/Runtime/Renderer/Private/SceneRendering.cpp
--- a/Engine/Source/Runtime/Renderer/Private/SceneRendering.cpp
+++ b/Engine/Source/Runtime/Renderer/Private/SceneRendering.cpp
@@ -2282,28 +2282,6 @@ void FViewInfo::EnqueueEyeAdaptationExposureBufferReadback(FRDGBuilder& GraphBui
 	}
 }
 
-#if WITH_MGPU
-void FViewInfo::WaitForEyeAdaptationTemporalEffect(FRHICommandList& RHICmdList)
-{
-	FSceneViewState* EffectiveViewState = GetEyeAdaptationViewState();
-
-	if (EffectiveViewState)
-	{
-		EffectiveViewState->WaitForEyeAdaptationTemporalEffect(RHICmdList);
-	}
-}
-
-void FViewInfo::BroadcastEyeAdaptationTemporalEffect(FRHICommandList& RHICmdList)
-{
-	FSceneViewState* EffectiveViewState = GetEyeAdaptationViewState();
-
-	if (EffectiveViewState)
-	{
-		EffectiveViewState->BroadcastEyeAdaptationTemporalEffect(RHICmdList);
-	}
-}
-#endif // WITH_MGPU
-
 float FViewInfo::GetLastEyeAdaptationExposure() const
 {
 	if (const FSceneViewState* EffectiveViewState = GetEyeAdaptationViewState())
diff --git a/Engine/Source/Runtime/Renderer/Private/SceneRendering.h b/Engine/Source/Runtime/Renderer/Private/SceneRendering.h
--- a/Engine/Source/Runtime/Renderer/Private/SceneRendering.h
+++ b/Engine/Source/Runtime/Renderer/Private/SceneRendering.h
@@ -1699,11 +1699,6 @@ public:
 	using FSceneView::GetEyeAdaptationBuffer;
 	FRDGPooledBuffer* GetEyeAdaptationBuffer(FRDGBuilder& GraphBuilder) const;
 
-#if WITH_MGPU
-	void BroadcastEyeAdaptationTemporalEffect(FRHICommandList& RHICmdList);
-	void WaitForEyeAdaptationTemporalEffect(FRHICommandList& RHICmdList);
-#endif
-
 	/** Get the last valid exposure value for eye adapation. */
 	float GetLastEyeAdaptationExposure() const;
 
