Support EyeAdaptationBuffer on every platform.
- Converted EyeAdaptation shaders to output to buffer.
- Updated most PostFX shaders to read from EyeAdaptationBuffer.
- Temporary: Non postfx shaders, plugins, etc still expect a texture so, for compatibility, implemented a CS to copy data from EyeAdaptationBuffer to EyeAdaptationTexture.

diff --git a/Engine/Shaders/Private/Histogram.usf b/Engine/Shaders/Private/Histogram.usf
--- a/Engine/Shaders/Private/Histogram.usf
+++ b/Engine/Shaders/Private/Histogram.usf
@@ -43,7 +43,7 @@ uint2 ThreadGroupCount;
 #define HISTOGRAMS_PER_THREAD ( (HISTOGRAM_SIZE < THREADGROUP_SIZEX) ? 1: HISTOGRAM_SIZE/THREADGROUP_SIZEX)
 
 Texture2D InputTexture;
-Texture2D EyeAdaptationTexture;
+StructuredBuffer<float4> EyeAdaptationBuffer;
 
 ///////////////////////////////////////////////////////
 // * LOCAL DATA STORE *
@@ -396,7 +396,7 @@ void HistogramConvertCS(uint3 GroupThreadId : SV_GroupThreadID)
 {
 	float2 InvViewSize = Input_ViewportSizeInverse.xy;
 	float NormalizeFactor = InvViewSize.x * InvViewSize.y;
-	float4 OldExposureScale = ToScalarMemory(EyeAdaptationTexture.Load(int3(0, 0, 0)));
+	float4 OldExposureScale = ToScalarMemory(EyeAdaptationBuffer[0]);
 
 	UNROLL_N(HISTOGRAMS_PER_THREAD)
 	for (uint t = 0; t < HISTOGRAMS_PER_THREAD; t++)
diff --git a/Engine/Shaders/Private/PostProcessBloom.usf b/Engine/Shaders/Private/PostProcessBloom.usf
--- a/Engine/Shaders/Private/PostProcessBloom.usf
+++ b/Engine/Shaders/Private/PostProcessBloom.usf
@@ -26,13 +26,7 @@ float BloomThreshold;
 
 float2 GetExposureScaleMiddleGreyLumValue()
 {
-#if SHADING_PATH_DEFERRED
-	float2 ExposureScaleMiddleGreyLumValue = EyeAdaptationLookupTexture2D(EyeAdaptationTexture).xw;
-#elif SHADING_PATH_MOBILE
 	float2 ExposureScaleMiddleGreyLumValue = EyeAdaptationLookupBuffer(EyeAdaptationBuffer).xw;
-#else
-#	error Unknown shading path.
-#endif
 
 	// Middle grey lum value adjusted by exposure compensation
 	ExposureScaleMiddleGreyLumValue.y = log2(0.18f * ExposureScaleMiddleGreyLumValue.y * EyeAdaptation_LocalExposureMiddleGreyExposureCompensation);
diff --git a/Engine/Shaders/Private/PostProcessEyeAdaptation.usf b/Engine/Shaders/Private/PostProcessEyeAdaptation.usf
--- a/Engine/Shaders/Private/PostProcessEyeAdaptation.usf
+++ b/Engine/Shaders/Private/PostProcessEyeAdaptation.usf
@@ -20,11 +20,11 @@
 #include "SHCommon.ush"
 #include "/Engine/Private/ScreenPass.ush"
 
-RWTexture2D<float4> RWEyeAdaptationTexture;
+RWStructuredBuffer<float4> RWEyeAdaptationBuffer;
 
 #endif
 
-Texture2D EyeAdaptationTexture;
+StructuredBuffer<float4> EyeAdaptationBuffer;
 
 Texture2D HistogramTexture;
 
@@ -144,8 +144,15 @@ void EyeAdaptationCS(uint2 DispatchThreadId : SV_DispatchThreadID)
 	OutColor.z = AverageSceneLuminance;
 	OutColor.w = MiddleGreyExposureCompensation;
 	
-	uint2 PixelPos = DispatchThreadId;
-	RWEyeAdaptationTexture[PixelPos] = OutColor;
+	RWEyeAdaptationBuffer[0] = OutColor;
+}
+
+RWTexture2D<float4> RWEyeAdaptationTexture;
+
+[numthreads(1, 1, 1)]
+void CopyEyeAdaptationToTextureCS(uint2 DispatchThreadId : SV_DispatchThreadID)
+{
+	RWEyeAdaptationTexture[DispatchThreadId] = EyeAdaptationBuffer[0];
 }
 #endif
 
@@ -232,7 +239,7 @@ void BasicEyeAdaptationCS(uint GIndex : SV_GroupIndex, uint2 GTId : SV_GroupThre
 	const float ClampedLumAve = clamp(LumAve, EyeAdaptation_MinAverageLuminance, EyeAdaptation_MaxAverageLuminance);
 	
 	// The Exposure Scale (and thus intensity) used in the previous frame
-	const float ExposureScaleOld = EyeAdaptationTexture.Load(int3(0, 0, 0)).x;
+	const float ExposureScaleOld = EyeAdaptationBuffer[0].x;
 	const float LuminanceAveOld = MiddleGreyExposureCompensation / (ExposureScaleOld != 0 ? ExposureScaleOld : 1.0f);
 
 	// Time-based expoential blend of the intensity to allow the eye adaptation to ramp up over a few frames.
@@ -257,7 +264,7 @@ void BasicEyeAdaptationCS(uint GIndex : SV_GroupIndex, uint2 GTId : SV_GroupThre
 	
 	if(GIndex==0)
 	{
-		RWEyeAdaptationTexture[uint2(0,0)] = OutColor;
+		RWEyeAdaptationBuffer[0] = OutColor;
 	}
 }
 
diff --git a/Engine/Shaders/Private/PostProcessHistogramReduce.usf b/Engine/Shaders/Private/PostProcessHistogramReduce.usf
--- a/Engine/Shaders/Private/PostProcessHistogramReduce.usf
+++ b/Engine/Shaders/Private/PostProcessHistogramReduce.usf
@@ -10,7 +10,7 @@
 uint LoopSize;
 
 // The last frames eye adaptation settings
-Texture2D EyeAdaptationTexture;
+StructuredBuffer<float4> EyeAdaptationBuffer;
 
 Texture2D InputTexture;
 
@@ -38,10 +38,10 @@ void MainPS(noperspective float4 UVAndScreenPos : TEXCOORD0, float4 SvPosition:
 	}
 	else
 	{
-		// line 1: store 4 channels of EyeAdaptationTexture (copied over so we can read the value in EyeAdaptation pass which is writing to eye adaptation)
+		// line 1: store 4 channels of EyeAdaptationBuffer (copied over so we can read the value in EyeAdaptation pass which is writing to eye adaptation)
 
 		// second line first pixel in the texture has the ExposureScale from last frame
-		float4 OldExposureScale = EyeAdaptationTexture.Load(int3(0, 0, 0));
+		float4 OldExposureScale = EyeAdaptationBuffer[0];
 
 		OutColor = OldExposureScale;
 	}
diff --git a/Engine/Shaders/Private/PostProcessLocalExposure.usf b/Engine/Shaders/Private/PostProcessLocalExposure.usf
--- a/Engine/Shaders/Private/PostProcessLocalExposure.usf
+++ b/Engine/Shaders/Private/PostProcessLocalExposure.usf
@@ -9,7 +9,7 @@ SCREEN_PASS_TEXTURE_VIEWPORT(Output)
 
 Texture2D InputTexture;
 
-Texture2D EyeAdaptationTexture;
+StructuredBuffer<float4> EyeAdaptationBuffer;
 
 Texture3D LumBilateralGrid;
 Texture2D BlurredLogLum;
@@ -36,7 +36,7 @@ void ApplyLocalExposureCS(uint2 DispatchThreadId : SV_DispatchThreadID)
 	{
 		const float2 UV = (DispatchThreadId + 0.5f) * Output_ExtentInverse;
 
-		float2 ExposureScaleMiddleGrey = EyeAdaptationTexture.Load(int3(0, 0, 0)).xw;
+		float2 ExposureScaleMiddleGrey = EyeAdaptationBuffer[0].xw;
 		float ExposureScale = ExposureScaleMiddleGrey.x;
 		float MiddleGreyLumValue = log2(0.18 * ExposureScaleMiddleGrey.y * EyeAdaptation_LocalExposureMiddleGreyExposureCompensation);
 
diff --git a/Engine/Shaders/Private/PostProcessTonemap.usf b/Engine/Shaders/Private/PostProcessTonemap.usf
--- a/Engine/Shaders/Private/PostProcessTonemap.usf
+++ b/Engine/Shaders/Private/PostProcessTonemap.usf
@@ -422,12 +422,7 @@ float4 TonemapCommonPS(
 	SceneColor.rgb *= View.OneOverPreExposure;
 
 #if EYEADAPTATION_EXPOSURE_FIX == 1
-#if FEATURE_LEVEL >= FEATURE_LEVEL_SM5
-	// texture can be GWhiteTexture which is 1x1. It's important we don't read outside bounds.
-	float2 EyeAdaptationData = EyeAdaptationTexture.Load(int3(0, 0, 0)).xw;
-#else
 	float2 EyeAdaptationData = EyeAdaptationBuffer[0].xw;
-#endif
 #else
 	// If eye adaptation is disabled (e.g. fixed exposure level ) or not available.  
 	float2 EyeAdaptationData = DefaultEyeExposure.xx;
diff --git a/Engine/Source/Runtime/Engine/Private/SceneView.cpp b/Engine/Source/Runtime/Engine/Private/SceneView.cpp
--- a/Engine/Source/Runtime/Engine/Private/SceneView.cpp
+++ b/Engine/Source/Runtime/Engine/Private/SceneView.cpp
@@ -2724,7 +2724,7 @@ bool FSceneView::HasValidEyeAdaptationBuffer() const
 
 IPooledRenderTarget* FSceneView::GetEyeAdaptationTexture() const
 {
-	checkf(FeatureLevel > ERHIFeatureLevel::ES3_1, TEXT("SM5 and above use RenderTarget for read back"));
+	checkf(FeatureLevel > ERHIFeatureLevel::ES3_1, TEXT("EyeAdaptation Texture is only available on SM5 and above."));
 	if (EyeAdaptationViewState)
 	{
 		return EyeAdaptationViewState->GetCurrentEyeAdaptationTexture();
@@ -2734,7 +2734,6 @@ IPooledRenderTarget* FSceneView::GetEyeAdaptationTexture() const
 
 FRDGPooledBuffer* FSceneView::GetEyeAdaptationBuffer() const
 {
-	checkf(FeatureLevel == ERHIFeatureLevel::ES3_1, TEXT("ES3_1 use RWBuffer for read back"));
 	if (EyeAdaptationViewState)
 	{
 		return EyeAdaptationViewState->GetCurrentEyeAdaptationBuffer();
diff --git a/Engine/Source/Runtime/Engine/Public/SceneManagement.h b/Engine/Source/Runtime/Engine/Public/SceneManagement.h
--- a/Engine/Source/Runtime/Engine/Public/SceneManagement.h
+++ b/Engine/Source/Runtime/Engine/Public/SceneManagement.h
@@ -237,7 +237,7 @@ public:
 	/** Returns the eye adaptation texture (SM5+ only). */
 	virtual IPooledRenderTarget* GetCurrentEyeAdaptationTexture() const = 0;
 
-	/** Returns the eye adaptation buffer (mobile only). */
+	/** Returns the eye adaptation buffer. */
 	virtual FRDGPooledBuffer* GetCurrentEyeAdaptationBuffer() const = 0;
 
 	/** Returns the eye adaptation exposure. */
diff --git a/Engine/Source/Runtime/Engine/Public/SceneView.h b/Engine/Source/Runtime/Engine/Public/SceneView.h
--- a/Engine/Source/Runtime/Engine/Public/SceneView.h
+++ b/Engine/Source/Runtime/Engine/Public/SceneView.h
@@ -1614,7 +1614,7 @@ public:
 	/** Returns the eye adaptation texture (SM5+ only) or null if it doesn't exist. */
 	IPooledRenderTarget* GetEyeAdaptationTexture() const;
 
-	/** Returns the eye adaptation buffer (mobile) or null if it doesn't exist. */
+	/** Returns the eye adaptation buffer or null if it doesn't exist. */
 	FRDGPooledBuffer* GetEyeAdaptationBuffer() const;
 
 	/** Returns the eye adaptation exposure or 0.0f if it doesn't exist. */
diff --git a/Engine/Source/Runtime/Renderer/Private/PostProcess/PostProcessBloomSetup.cpp b/Engine/Source/Runtime/Renderer/Private/PostProcess/PostProcessBloomSetup.cpp
--- a/Engine/Source/Runtime/Renderer/Private/PostProcess/PostProcessBloomSetup.cpp
+++ b/Engine/Source/Runtime/Renderer/Private/PostProcess/PostProcessBloomSetup.cpp
@@ -31,12 +31,13 @@ BEGIN_SHADER_PARAMETER_STRUCT(FBloomSetupParameters, )
 	SHADER_PARAMETER_SAMPLER(SamplerState, LumBilateralGridSampler)
 	SHADER_PARAMETER_RDG_TEXTURE(Texture2D, BlurredLogLum)
 	SHADER_PARAMETER_SAMPLER(SamplerState, BlurredLogLumSampler)
-	SHADER_PARAMETER_RDG_TEXTURE(Texture2D, EyeAdaptationTexture)
+	SHADER_PARAMETER_RDG_BUFFER_SRV(StructuredBuffer<float4>, EyeAdaptationBuffer)
 	SHADER_PARAMETER_STRUCT(FEyeAdaptationParameters, EyeAdaptation)
 	SHADER_PARAMETER(float, BloomThreshold)
 END_SHADER_PARAMETER_STRUCT()
 
 FBloomSetupParameters GetBloomSetupParameters(
+	FRDGBuilder& GraphBuilder,
 	const FViewInfo& View,
 	const FScreenPassTextureViewport& InputViewport,
 	const FBloomSetupInputs& Inputs)
@@ -50,7 +51,7 @@ FBloomSetupParameters GetBloomSetupParameters(
 	Parameters.LumBilateralGridSampler = TStaticSamplerState<SF_Bilinear, AM_Clamp, AM_Clamp, AM_Clamp>::GetRHI();
 	Parameters.BlurredLogLum = Inputs.BlurredLogLuminanceTexture;
 	Parameters.BlurredLogLumSampler = TStaticSamplerState<SF_Bilinear, AM_Clamp, AM_Clamp, AM_Clamp>::GetRHI();
-	Parameters.EyeAdaptationTexture = Inputs.EyeAdaptationTexture;
+	Parameters.EyeAdaptationBuffer = GraphBuilder.CreateSRV(Inputs.EyeAdaptationBuffer);
 	Parameters.EyeAdaptation = *Inputs.EyeAdaptationParameters;
 	Parameters.BloomThreshold = Inputs.Threshold;
 	return Parameters;
@@ -112,7 +113,7 @@ IMPLEMENT_GLOBAL_SHADER(FBloomSetupCS, "/Engine/Private/PostProcessBloom.usf", "
 FScreenPassTexture AddBloomSetupPass(FRDGBuilder& GraphBuilder, const FViewInfo& View, const FBloomSetupInputs& Inputs)
 {
 	check(Inputs.SceneColor.IsValid());
-	check(Inputs.EyeAdaptationTexture);
+	check(Inputs.EyeAdaptationBuffer);
 	check(Inputs.Threshold > -1.0f || Inputs.EyeAdaptationParameters != nullptr);
 
 	const bool bIsComputePass = View.bUseComputePasses;
@@ -127,7 +128,7 @@ FScreenPassTexture AddBloomSetupPass(FRDGBuilder& GraphBuilder, const FViewInfo&
 	if (bIsComputePass)
 	{
 		FBloomSetupCS::FParameters* PassParameters = GraphBuilder.AllocParameters<FBloomSetupCS::FParameters>();
-		PassParameters->BloomSetup = GetBloomSetupParameters(View, Viewport, Inputs);
+		PassParameters->BloomSetup = GetBloomSetupParameters(GraphBuilder, View, Viewport, Inputs);
 		PassParameters->RWOutputTexture = GraphBuilder.CreateUAV(Output.Texture);
 
 		FBloomSetupCS::FPermutationDomain PermutationVector;
@@ -145,7 +146,7 @@ FScreenPassTexture AddBloomSetupPass(FRDGBuilder& GraphBuilder, const FViewInfo&
 	else
 	{
 		FBloomSetupPS::FParameters* PassParameters = GraphBuilder.AllocParameters<FBloomSetupPS::FParameters>();
-		PassParameters->BloomSetup = GetBloomSetupParameters(View, Viewport, Inputs);
+		PassParameters->BloomSetup = GetBloomSetupParameters(GraphBuilder, View, Viewport, Inputs);
 		PassParameters->SvPositionToInputTextureUV = (
 			FScreenTransform::ChangeTextureBasisFromTo(FScreenPassTextureViewport(Output), FScreenTransform::ETextureBasis::TexelPosition, FScreenTransform::ETextureBasis::ViewportUV) *
 			FScreenTransform::ChangeTextureBasisFromTo(FScreenPassTextureViewport(Inputs.SceneColor), FScreenTransform::ETextureBasis::ViewportUV, FScreenTransform::ETextureBasis::TextureUV));
diff --git a/Engine/Source/Runtime/Renderer/Private/PostProcess/PostProcessBloomSetup.h b/Engine/Source/Runtime/Renderer/Private/PostProcess/PostProcessBloomSetup.h
--- a/Engine/Source/Runtime/Renderer/Private/PostProcess/PostProcessBloomSetup.h
+++ b/Engine/Source/Runtime/Renderer/Private/PostProcess/PostProcessBloomSetup.h
@@ -27,8 +27,8 @@ struct FBloomSetupInputs
 	// [Required]: The intermediate scene color being processed.
 	FScreenPassTexture SceneColor;
 
-	// [Required]: The scene eye adaptation texture.
-	FRDGTextureRef EyeAdaptationTexture = nullptr;
+	// [Required]: The scene eye adaptation buffer.
+	FRDGBufferRef EyeAdaptationBuffer = nullptr;
 
 	// [Required]: The bloom threshold to apply. Must be >0.
 	float Threshold = 0.0f;
diff --git a/Engine/Source/Runtime/Renderer/Private/PostProcess/PostProcessEyeAdaptation.cpp b/Engine/Source/Runtime/Renderer/Private/PostProcess/PostProcessEyeAdaptation.cpp
--- a/Engine/Source/Runtime/Renderer/Private/PostProcess/PostProcessEyeAdaptation.cpp
+++ b/Engine/Source/Runtime/Renderer/Private/PostProcess/PostProcessEyeAdaptation.cpp
@@ -524,6 +524,45 @@ float GetEyeAdaptationFixedExposure(const FViewInfo& View)
 	return ExposureScale * Parameters.ExposureCompensationSettings;
 }
 
+//////////////////////////////////////////////////////////////////////////
+//! CopyEyeAdaptationToTexture
+//////////////////////////////////////////////////////////////////////////
+
+class FCopyEyeAdaptationToTextureCS : public FGlobalShader
+{
+public:
+	DECLARE_GLOBAL_SHADER(FCopyEyeAdaptationToTextureCS);
+	SHADER_USE_PARAMETER_STRUCT(FCopyEyeAdaptationToTextureCS, FGlobalShader);
+
+	BEGIN_SHADER_PARAMETER_STRUCT(FParameters, )
+		SHADER_PARAMETER_RDG_BUFFER_SRV(StructuredBuffer<float4>, EyeAdaptationBuffer)
+		SHADER_PARAMETER_RDG_TEXTURE_UAV(RWTexture2D<float4>, RWEyeAdaptationTexture)
+		END_SHADER_PARAMETER_STRUCT()
+
+		static bool ShouldCompilePermutation(const FGlobalShaderPermutationParameters& Parameters)
+	{
+		return IsFeatureLevelSupported(Parameters.Platform, ERHIFeatureLevel::SM5);
+	}
+};
+
+IMPLEMENT_GLOBAL_SHADER(FCopyEyeAdaptationToTextureCS, "/Engine/Private/PostProcessEyeAdaptation.usf", "CopyEyeAdaptationToTextureCS", SF_Compute);
+
+void AddCopyEyeAdaptationDataToTexturePass(FRDGBuilder& GraphBuilder, const FGlobalShaderMap* ShaderMap, FRDGBufferRef EyeAdaptationBuffer, FRDGTextureRef OutputTexture)
+{
+	auto* PassParameters = GraphBuilder.AllocParameters<FCopyEyeAdaptationToTextureCS::FParameters>();
+	PassParameters->EyeAdaptationBuffer = GraphBuilder.CreateSRV(EyeAdaptationBuffer);
+	PassParameters->RWEyeAdaptationTexture = GraphBuilder.CreateUAV(OutputTexture);
+
+	auto ComputeShader = ShaderMap->GetShader<FCopyEyeAdaptationToTextureCS>();
+
+	FComputeShaderUtils::AddPass(
+		GraphBuilder,
+		RDG_EVENT_NAME("CopyEyeAdaptationToTexture (CS)"),
+		ComputeShader,
+		PassParameters,
+		FIntVector(1, 1, 1));
+}
+
 //////////////////////////////////////////////////////////////////////////
 //! Setup
 //////////////////////////////////////////////////////////////////////////
@@ -704,7 +743,7 @@ class FEyeAdaptationCS : public FGlobalShader
 	BEGIN_SHADER_PARAMETER_STRUCT(FParameters, )
 		SHADER_PARAMETER_STRUCT(FEyeAdaptationParameters, EyeAdaptation)
 		SHADER_PARAMETER_RDG_TEXTURE(Texture2D, HistogramTexture)
-		SHADER_PARAMETER_RDG_TEXTURE_UAV(RWTexture2D<float4>, RWEyeAdaptationTexture)
+		SHADER_PARAMETER_RDG_BUFFER_UAV(RWStructuredBuffer<float4>, RWEyeAdaptationBuffer)
 	END_SHADER_PARAMETER_STRUCT()
 
 	static const EPixelFormat OutputFormat = PF_A32B32G32R32F;
@@ -717,21 +756,21 @@ class FEyeAdaptationCS : public FGlobalShader
 
 IMPLEMENT_GLOBAL_SHADER(FEyeAdaptationCS, "/Engine/Private/PostProcessEyeAdaptation.usf", "EyeAdaptationCS", SF_Compute);
 
-FRDGTextureRef AddHistogramEyeAdaptationPass(
+FRDGBufferRef AddHistogramEyeAdaptationPass(
 	FRDGBuilder& GraphBuilder,
 	const FViewInfo& View,
 	const FEyeAdaptationParameters& EyeAdaptationParameters,
 	FRDGTextureRef HistogramTexture)
 {
-	View.UpdateEyeAdaptationLastExposureFromTexture();
-	View.SwapEyeAdaptationTextures();
+	View.UpdateEyeAdaptationLastExposureFromBuffer();
+	View.SwapEyeAdaptationBuffers();
 
-	FRDGTextureRef OutputTexture = GraphBuilder.RegisterExternalTexture(View.GetEyeAdaptationTexture(GraphBuilder.RHICmdList), ERDGTextureFlags::MultiFrame);
+	FRDGBufferRef OutputBuffer = GraphBuilder.RegisterExternalBuffer(View.GetEyeAdaptationBuffer(GraphBuilder), ERDGBufferFlags::MultiFrame);
 
 	FEyeAdaptationCS::FParameters* PassParameters = GraphBuilder.AllocParameters<FEyeAdaptationCS::FParameters>();
 	PassParameters->EyeAdaptation = GetEyeAdaptationParameters(View, ERHIFeatureLevel::SM5);
 	PassParameters->HistogramTexture = HistogramTexture;
-	PassParameters->RWEyeAdaptationTexture = GraphBuilder.CreateUAV(OutputTexture);
+	PassParameters->RWEyeAdaptationBuffer = GraphBuilder.CreateUAV(OutputBuffer);
 
 	TShaderMapRef<FEyeAdaptationCS> ComputeShader(View.ShaderMap);
 	FComputeShaderUtils::AddPass(
@@ -741,9 +780,14 @@ FRDGTextureRef AddHistogramEyeAdaptationPass(
 		PassParameters,
 		FIntVector(1, 1, 1));
 
-	View.EnqueueEyeAdaptationExposureTextureReadback(GraphBuilder);
+	{
+		FRDGTextureRef OutputTexture = GraphBuilder.RegisterExternalTexture(View.GetEyeAdaptationTexture(GraphBuilder), ERDGTextureFlags::MultiFrame);
+		AddCopyEyeAdaptationDataToTexturePass(GraphBuilder, View.ShaderMap, OutputBuffer, OutputTexture);
+	}
 
-	return OutputTexture;
+	View.EnqueueEyeAdaptationExposureBufferReadback(GraphBuilder);
+
+	return OutputBuffer;
 }
 
 //////////////////////////////////////////////////////////////////////////
@@ -822,8 +866,8 @@ public:
 		SHADER_PARAMETER_STRUCT(FEyeAdaptationParameters, EyeAdaptation)
 		SHADER_PARAMETER_STRUCT(FScreenPassTextureViewportParameters, Color)
 		SHADER_PARAMETER_RDG_TEXTURE(Texture2D, ColorTexture)
-		SHADER_PARAMETER_RDG_TEXTURE(Texture2D, EyeAdaptationTexture)
-		SHADER_PARAMETER_RDG_TEXTURE_UAV(RWTexture2D<float4>, RWEyeAdaptationTexture)
+		SHADER_PARAMETER_RDG_BUFFER_SRV(StructuredBuffer<float4>, EyeAdaptationBuffer)
+		SHADER_PARAMETER_RDG_BUFFER_UAV(RWStructuredBuffer<float4>, RWEyeAdaptationBuffer)
 	END_SHADER_PARAMETER_STRUCT()
 
 	static bool ShouldCompilePermutation(const FGlobalShaderPermutationParameters& Parameters)
@@ -834,29 +878,29 @@ public:
 
 IMPLEMENT_GLOBAL_SHADER(FBasicEyeAdaptationCS, "/Engine/Private/PostProcessEyeAdaptation.usf", "BasicEyeAdaptationCS", SF_Compute);
 
-FRDGTextureRef AddBasicEyeAdaptationPass(
+FRDGBufferRef AddBasicEyeAdaptationPass(
 	FRDGBuilder& GraphBuilder,
 	const FViewInfo& View,
 	const FEyeAdaptationParameters& EyeAdaptationParameters,
 	FScreenPassTexture SceneColor,
-	FRDGTextureRef EyeAdaptationTexture)
+	FRDGBufferRef EyeAdaptationBuffer)
 {
-	View.UpdateEyeAdaptationLastExposureFromTexture();
-	View.SwapEyeAdaptationTextures();
+	View.UpdateEyeAdaptationLastExposureFromBuffer();
+	View.SwapEyeAdaptationBuffers();
 
 	const FScreenPassTextureViewport SceneColorViewport(SceneColor);
 
-	FRDGTextureRef OutputTexture = GraphBuilder.RegisterExternalTexture(View.GetEyeAdaptationTexture(GraphBuilder.RHICmdList), ERDGTextureFlags::MultiFrame);
+	FRDGBufferRef OutputBuffer = GraphBuilder.RegisterExternalBuffer(View.GetEyeAdaptationBuffer(GraphBuilder), ERDGBufferFlags::MultiFrame);
 
 	FBasicEyeAdaptationCS::FParameters* PassParameters = GraphBuilder.AllocParameters<FBasicEyeAdaptationCS::FParameters>();
 	PassParameters->View = View.ViewUniformBuffer;
 	PassParameters->EyeAdaptation = EyeAdaptationParameters;
 	PassParameters->Color = GetScreenPassTextureViewportParameters(SceneColorViewport);
 	PassParameters->ColorTexture = SceneColor.Texture;
-	PassParameters->EyeAdaptationTexture = EyeAdaptationTexture;
-	PassParameters->RWEyeAdaptationTexture = GraphBuilder.CreateUAV(OutputTexture);
+	PassParameters->EyeAdaptationBuffer = GraphBuilder.CreateSRV(EyeAdaptationBuffer);
+	PassParameters->RWEyeAdaptationBuffer = GraphBuilder.CreateUAV(OutputBuffer);
 
-	auto ComputeShader = View.ShaderMap->GetShader<FBasicEyeAdaptationCS>();
+	TShaderMapRef<FBasicEyeAdaptationCS> ComputeShader(View.ShaderMap);
 
 	FComputeShaderUtils::AddPass(
 		GraphBuilder,
@@ -865,20 +909,24 @@ FRDGTextureRef AddBasicEyeAdaptationPass(
 		PassParameters,
 		FIntVector(1, 1, 1));
 
-	View.EnqueueEyeAdaptationExposureTextureReadback(GraphBuilder);
+	{
+		FRDGTextureRef OutputTexture = GraphBuilder.RegisterExternalTexture(View.GetEyeAdaptationTexture(GraphBuilder), ERDGTextureFlags::MultiFrame);
+		AddCopyEyeAdaptationDataToTexturePass(GraphBuilder, View.ShaderMap, OutputBuffer, OutputTexture);
+	}
+
+	View.EnqueueEyeAdaptationExposureBufferReadback(GraphBuilder);
 
-	return OutputTexture;
+	return OutputBuffer;
 }
 
 FSceneViewState::FEyeAdaptationManager::FEyeAdaptationManager()
 {
-	ExposureReadbackTextures.AddZeroed(MAX_READBACK_BUFFERS);
 	ExposureReadbackBuffers.AddZeroed(MAX_READBACK_BUFFERS);
 }
 
 void FSceneViewState::FEyeAdaptationManager::SafeRelease()
 {
-	CurrentBuffer = 0;
+	CurrentBufferIndex = 0;
 	ReadbackBuffersWriteIndex = 0;
 	ReadbackBuffersNumPending = 0;
 
@@ -887,17 +935,8 @@ void FSceneViewState::FEyeAdaptationManager::SafeRelease()
 
 	for (int32 Index = 0; Index < NUM_BUFFERS; Index++)
 	{
-		PooledRenderTarget[Index].SafeRelease();
 		ExposureBufferData[Index].SafeRelease();
-	}
-
-	for (int32 Index = 0; Index < ExposureReadbackTextures.Num(); ++Index)
-	{
-		if (ExposureReadbackTextures[Index])
-		{
-			delete ExposureReadbackTextures[Index];
-			ExposureReadbackTextures[Index] = nullptr;
-		}
+		PooledRenderTarget[Index].SafeRelease();
 	}
 
 	for (int32 Index = 0; Index < ExposureReadbackBuffers.Num(); ++Index)
@@ -910,102 +949,23 @@ void FSceneViewState::FEyeAdaptationManager::SafeRelease()
 	}
 }
 
-FRHIGPUTextureReadback* FSceneViewState::FEyeAdaptationManager::GetLatestReadbackTexture()
-{
-	FRHIGPUTextureReadback* LatestReadbackTexture = nullptr;
-
-	// Find latest texture that is ready
-	while (ReadbackBuffersNumPending > 0)
-	{
-		uint32 Index = (ReadbackBuffersWriteIndex + MAX_READBACK_BUFFERS - ReadbackBuffersNumPending) % MAX_READBACK_BUFFERS;
-
-		if (ExposureReadbackTextures[Index]->IsReady())
-		{
-			--ReadbackBuffersNumPending;
-			LatestReadbackTexture = ExposureReadbackTextures[Index];
-		}
-		else
-		{
-			break;
-		}
-	}
-
-	return LatestReadbackTexture;
-}
-
-void FSceneViewState::FEyeAdaptationManager::SwapTextures()
-{
-	CurrentBuffer = (CurrentBuffer + 1) % NUM_BUFFERS;
-}
-
-void FSceneViewState::FEyeAdaptationManager::UpdateLastExposureFromTexture()
-{
-	QUICK_SCOPE_CYCLE_COUNTER(STAT_FEyeAdaptationRTManager_ReadbackLastExposure);
-
-	// Get the most up to date value
-	FRHIGPUTextureReadback* ReadbackTexture = GetLatestReadbackTexture();
-
-	if (ReadbackTexture)
-	{
-		// Read the last request results.
-		int32 RowPitch = 0;
-		FVector4f* ReadbackData = (FVector4f*)ReadbackTexture->Lock(RowPitch);
-		if (ReadbackData)
-		{
-			LastExposure = ReadbackData->X;
-			LastAverageSceneLuminance = ReadbackData->Z;
-
-			ReadbackTexture->Unlock();
-		}
-	}
-}
-
-void FSceneViewState::FEyeAdaptationManager::EnqueueExposureTextureReadback(FRDGBuilder& GraphBuilder)
-{
-	check(PooledRenderTarget[CurrentBuffer].IsValid());
-	
-	if (ReadbackBuffersNumPending < MAX_READBACK_BUFFERS)
-	{
-		// limit to single bit here for readback to work
-		RDG_GPU_MASK_SCOPE(GraphBuilder, FRHIGPUMask::FromIndex(GraphBuilder.RHICmdList.GetGPUMask().GetFirstIndex()));
-		FRDGTextureRef CurrentTexture = GraphBuilder.RegisterExternalTexture(PooledRenderTarget[CurrentBuffer], ERDGTextureFlags::MultiFrame);
-
-		FRHIGPUTextureReadback* ExposureReadbackTexture = ExposureReadbackTextures[ReadbackBuffersWriteIndex];
-
-		if (ExposureReadbackTexture == nullptr)
-		{
-			static const FName ExposureValueName(TEXT("Scene view state exposure readback"));
-			ExposureReadbackTexture = new FRHIGPUTextureReadback(ExposureValueName);
-			ExposureReadbackTextures[ReadbackBuffersWriteIndex] = ExposureReadbackTexture;
-		}
-
-		AddEnqueueCopyPass(GraphBuilder, ExposureReadbackTexture, CurrentTexture);
-
-		ReadbackBuffersWriteIndex = (ReadbackBuffersWriteIndex + 1) % MAX_READBACK_BUFFERS;
-		ReadbackBuffersNumPending = FMath::Min(ReadbackBuffersNumPending + 1, MAX_READBACK_BUFFERS);
-	}
-}
-
 const TRefCountPtr<IPooledRenderTarget>& FSceneViewState::FEyeAdaptationManager::GetTexture(uint32 TextureIndex) const
 {
 	check(0 <= TextureIndex && TextureIndex < NUM_BUFFERS);
 	return PooledRenderTarget[TextureIndex];
 }
 
-const TRefCountPtr<IPooledRenderTarget>& FSceneViewState::FEyeAdaptationManager::GetOrCreateTexture(FRHICommandList& RHICmdList, uint32 TextureIndex)
+const TRefCountPtr<IPooledRenderTarget>& FSceneViewState::FEyeAdaptationManager::GetOrCreateTexture(FRDGBuilder& GraphBuilder, uint32 TextureIndex)
 {
 	check(0 <= TextureIndex && TextureIndex < NUM_BUFFERS);
 
-	// Create textures if needed.
+	// Create texture if needed.
 	if (!PooledRenderTarget[TextureIndex].IsValid())
 	{
-		// Create the texture needed for EyeAdaptation
-		FPooledRenderTargetDesc Desc(FPooledRenderTargetDesc::Create2DDesc(FIntPoint(1, 1), PF_A32B32G32R32F, FClearValueBinding::None, TexCreate_None, TexCreate_RenderTargetable | TexCreate_ShaderResource, false));
-		if (GMaxRHIFeatureLevel >= ERHIFeatureLevel::SM5)
-		{
-			Desc.Flags |= TexCreate_UAV;
-		}
-		GRenderTargetPool.FindFreeElement(RHICmdList, Desc, PooledRenderTarget[TextureIndex], TEXT("EyeAdaptation"));
+		const FRDGTextureDesc RDGTextureDesc = FRDGTextureDesc::Create2D(FIntPoint(1, 1), PF_A32B32G32R32F, FClearValueBinding::None, TexCreate_UAV | TexCreate_ShaderResource);
+		FRDGTextureRef RDGTexture = GraphBuilder.CreateTexture(RDGTextureDesc, TEXT("EyeAdaptationTexture"), ERDGTextureFlags::MultiFrame);
+
+		PooledRenderTarget[TextureIndex] = GraphBuilder.ConvertToExternalTexture(RDGTexture);
 	}
 
 	return PooledRenderTarget[TextureIndex];
@@ -1022,7 +982,7 @@ const TRefCountPtr<FRDGPooledBuffer>& FSceneViewState::FEyeAdaptationManager::Ge
 {
 	check(0 <= BufferIndex && BufferIndex < NUM_BUFFERS);
 
-	// Create textures if needed.
+	// Create buffer if needed.
 	if (!ExposureBufferData[BufferIndex].IsValid())
 	{
 		FRDGBufferDesc RDGBufferDesc = FRDGBufferDesc::CreateStructuredDesc(sizeof(FVector4f), 1);
@@ -1064,7 +1024,7 @@ FRHIGPUBufferReadback* FSceneViewState::FEyeAdaptationManager::GetLatestReadback
 
 void FSceneViewState::FEyeAdaptationManager::SwapBuffers()
 {
-	CurrentBuffer = (CurrentBuffer + 1) % NUM_BUFFERS;
+	CurrentBufferIndex = (CurrentBufferIndex + 1) % NUM_BUFFERS;
 }
 
 void FSceneViewState::FEyeAdaptationManager::UpdateLastExposureFromBuffer()
@@ -1090,11 +1050,11 @@ void FSceneViewState::FEyeAdaptationManager::UpdateLastExposureFromBuffer()
 
 void FSceneViewState::FEyeAdaptationManager::EnqueueExposureBufferReadback(FRDGBuilder& GraphBuilder)
 {
-	check(ExposureBufferData[CurrentBuffer].IsValid());
+	check(ExposureBufferData[CurrentBufferIndex].IsValid());
 
 	if (ReadbackBuffersNumPending < MAX_READBACK_BUFFERS)
 	{
-		FRDGBufferRef CurrentRDGBuffer = GraphBuilder.RegisterExternalBuffer(ExposureBufferData[CurrentBuffer], ERDGBufferFlags::MultiFrame);
+		FRDGBufferRef CurrentRDGBuffer = GraphBuilder.RegisterExternalBuffer(ExposureBufferData[CurrentBufferIndex], ERDGBufferFlags::MultiFrame);
 
 		FRHIGPUBufferReadback* ExposureReadbackBuffer = ExposureReadbackBuffers[ReadbackBuffersWriteIndex];
 
diff --git a/Engine/Source/Runtime/Renderer/Private/PostProcess/PostProcessEyeAdaptation.h b/Engine/Source/Runtime/Renderer/Private/PostProcess/PostProcessEyeAdaptation.h
--- a/Engine/Source/Runtime/Renderer/Private/PostProcess/PostProcessEyeAdaptation.h
+++ b/Engine/Source/Runtime/Renderer/Private/PostProcess/PostProcessEyeAdaptation.h
@@ -106,8 +106,8 @@ FRDGTextureRef AddCalculateExposureIlluminancePass(
 	const struct FTranslucencyLightingVolumeTextures& TranslucencyLightingVolumeTextures,
 	FRDGTextureRef ExposureIlluminanceSetup);
 
-// Returns the updated 1x1 eye adaptation texture.
-FRDGTextureRef AddHistogramEyeAdaptationPass(
+// Returns the updated eye adaptation buffer.
+FRDGBufferRef AddHistogramEyeAdaptationPass(
 	FRDGBuilder& GraphBuilder,
 	const FViewInfo& View,
 	const FEyeAdaptationParameters& EyeAdaptationParameters,
@@ -120,10 +120,10 @@ FScreenPassTexture AddBasicEyeAdaptationSetupPass(
 	const FEyeAdaptationParameters& EyeAdaptationParameters,
 	FScreenPassTexture SceneColor);
 
-// Returns the updated 1x1 eye adaptation texture.
-FRDGTextureRef AddBasicEyeAdaptationPass(
+// Returns the updated eye adaptation buffer.
+FRDGBufferRef AddBasicEyeAdaptationPass(
 	FRDGBuilder& GraphBuilder,
 	const FViewInfo& View,
 	const FEyeAdaptationParameters& EyeAdaptationParameters,
 	FScreenPassTexture SceneColor,
-	FRDGTextureRef EyeAdaptationTexture);
+	FRDGBufferRef EyeAdaptationBuffer);
diff --git a/Engine/Source/Runtime/Renderer/Private/PostProcess/PostProcessFFTBloom.cpp b/Engine/Source/Runtime/Renderer/Private/PostProcess/PostProcessFFTBloom.cpp
--- a/Engine/Source/Runtime/Renderer/Private/PostProcess/PostProcessFFTBloom.cpp
+++ b/Engine/Source/Runtime/Renderer/Private/PostProcess/PostProcessFFTBloom.cpp
@@ -637,7 +637,7 @@ FFFTBloomOutput AddFFTBloomPass(
 	const FScreenPassTexture& InputSceneColor,
 	float InputResolutionFraction,
 	const FEyeAdaptationParameters& EyeAdaptationParameters,
-	FRDGTextureRef EyeAdaptationTexture,
+	FRDGBufferRef EyeAdaptationBuffer,
 	FRDGTextureRef LocalExposureTexture,
 	FRDGTextureRef BlurredLogLuminanceTexture)
 {
@@ -748,7 +748,7 @@ FFFTBloomOutput AddFFTBloomPass(
 
 		FFTInputSceneColor.Texture = GraphBuilder.CreateTexture(Desc, TEXT("Bloom.FFT.Input"));
 
-		AddApplyLocalExposurePass(GraphBuilder, View, EyeAdaptationParameters, EyeAdaptationTexture, LocalExposureTexture, BlurredLogLuminanceTexture, Temp, FFTInputSceneColor, Intermediates.ComputePassFlags);
+		AddApplyLocalExposurePass(GraphBuilder, View, EyeAdaptationParameters, EyeAdaptationBuffer, LocalExposureTexture, BlurredLogLuminanceTexture, Temp, FFTInputSceneColor, Intermediates.ComputePassFlags);
 	}
 
 	// Init the domain data update the cached kernel if needed.
diff --git a/Engine/Source/Runtime/Renderer/Private/PostProcess/PostProcessFFTBloom.h b/Engine/Source/Runtime/Renderer/Private/PostProcess/PostProcessFFTBloom.h
--- a/Engine/Source/Runtime/Renderer/Private/PostProcess/PostProcessFFTBloom.h
+++ b/Engine/Source/Runtime/Renderer/Private/PostProcess/PostProcessFFTBloom.h
@@ -25,6 +25,6 @@ FFFTBloomOutput AddFFTBloomPass(
 	const FScreenPassTexture& InputSceneColor,
 	float InputResolutionFraction,
 	const FEyeAdaptationParameters& EyeAdaptationParameters,
-	FRDGTextureRef EyeAdaptationTexture,
+	FRDGBufferRef EyeAdaptationBuffer,
 	FRDGTextureRef LocalExposureTexture,
 	FRDGTextureRef BlurredLogLuminanceTexture);
diff --git a/Engine/Source/Runtime/Renderer/Private/PostProcess/PostProcessHistogram.cpp b/Engine/Source/Runtime/Renderer/Private/PostProcess/PostProcessHistogram.cpp
--- a/Engine/Source/Runtime/Renderer/Private/PostProcess/PostProcessHistogram.cpp
+++ b/Engine/Source/Runtime/Renderer/Private/PostProcess/PostProcessHistogram.cpp
@@ -100,7 +100,7 @@ public:
 		SHADER_PARAMETER_STRUCT(FScreenPassTextureViewportParameters, Input)
 		SHADER_PARAMETER_SAMPLER(SamplerState, InputSampler)
 		SHADER_PARAMETER_RDG_TEXTURE(Texture2D, InputTexture)
-		SHADER_PARAMETER_RDG_TEXTURE(Texture2D, EyeAdaptationTexture)
+		SHADER_PARAMETER_RDG_BUFFER_SRV(StructuredBuffer<float4>, EyeAdaptationBuffer)
 		SHADER_PARAMETER(uint32, LoopSize)
 		RENDER_TARGET_BINDING_SLOTS()
 	END_SHADER_PARAMETER_STRUCT()
@@ -138,7 +138,7 @@ public:
 		SHADER_PARAMETER_STRUCT(FScreenPassTextureViewportParameters, Input)
 		SHADER_PARAMETER_STRUCT(FEyeAdaptationParameters, EyeAdaptation)
 		SHADER_PARAMETER_RDG_TEXTURE(Texture2D, InputTexture)
-		SHADER_PARAMETER_RDG_TEXTURE(Texture2D, EyeAdaptationTexture)
+		SHADER_PARAMETER_RDG_BUFFER_SRV(StructuredBuffer<float4>, EyeAdaptationBuffer)
 		SHADER_PARAMETER(FIntPoint, ThreadGroupCount)
 		SHADER_PARAMETER_RDG_TEXTURE_UAV(RWTexture2D<float4>, HistogramScatter64Output)
 		SHADER_PARAMETER_RDG_TEXTURE_UAV(RWTexture2D<float4>, HistogramScatter32Output)
@@ -186,7 +186,7 @@ public:
 
 	BEGIN_SHADER_PARAMETER_STRUCT(FParameters, )
 		SHADER_PARAMETER_STRUCT(FScreenPassTextureViewportParameters, Input)
-		SHADER_PARAMETER_RDG_TEXTURE(Texture2D, EyeAdaptationTexture)
+		SHADER_PARAMETER_RDG_BUFFER_SRV(StructuredBuffer<float4>, EyeAdaptationBuffer)
 		SHADER_PARAMETER_RDG_TEXTURE(Texture2D, HistogramScatter64Texture)
 		SHADER_PARAMETER_RDG_TEXTURE(Texture2D, HistogramScatter32Texture)
 		SHADER_PARAMETER_RDG_TEXTURE_UAV(RWTexture2D<float4>, HistogramOutput)
@@ -214,10 +214,10 @@ static FRDGTextureRef AddHistogramLegacyPass(
 	const FViewInfo& View,
 	const FEyeAdaptationParameters& EyeAdaptationParameters,
 	const FScreenPassTexture& SceneColor,
-	FRDGTextureRef EyeAdaptationTexture)
+	FRDGBufferRef EyeAdaptationBuffer)
 {
 	check(SceneColor.IsValid());
-	check(EyeAdaptationTexture);
+	check(EyeAdaptationBuffer);
 
 	const FIntPoint HistogramThreadGroupCount = FHistogramCS::GetThreadGroupCount(SceneColor.ViewRect.Size(), false);
 	const uint32 HistogramThreadGroupCountTotal = HistogramThreadGroupCount.X * HistogramThreadGroupCount.Y;
@@ -282,7 +282,7 @@ static FRDGTextureRef AddHistogramLegacyPass(
 		PassParameters->InputTexture = HistogramTexture;
 		PassParameters->InputSampler = TStaticSamplerState<SF_Bilinear, AM_Clamp, AM_Clamp, AM_Clamp>::GetRHI();
 		PassParameters->LoopSize = HistogramThreadGroupCountTotal;
-		PassParameters->EyeAdaptationTexture = EyeAdaptationTexture;
+		PassParameters->EyeAdaptationBuffer = GraphBuilder.CreateSRV(EyeAdaptationBuffer);
 		PassParameters->RenderTargets[0] = FRenderTargetBinding(HistogramReduceTexture, ERenderTargetLoadAction::ENoAction);
 
 		TShaderMapRef<FHistogramReducePS> PixelShader(View.ShaderMap);
@@ -306,10 +306,10 @@ static FRDGTextureRef AddHistogramAtomicPass(
 	const FEyeAdaptationParameters& EyeAdaptationParameters,
 	const FScreenPassTexture& SceneColor,
 	const FSceneTextureParameters& SceneTextures,
-	FRDGTextureRef EyeAdaptationTexture)
+	FRDGBufferRef EyeAdaptationBuffer)
 {
 	check(SceneColor.IsValid());
-	check(EyeAdaptationTexture);
+	check(EyeAdaptationBuffer);
 
 	RDG_EVENT_SCOPE(GraphBuilder, "Histogram");
 
@@ -341,7 +341,7 @@ static FRDGTextureRef AddHistogramAtomicPass(
 		PassParameters->SceneTextures = SceneTextures;
 		PassParameters->Input = GetScreenPassTextureViewportParameters(FScreenPassTextureViewport(SceneColor));
 		PassParameters->InputTexture = SceneColor.Texture;
-		PassParameters->EyeAdaptationTexture = EyeAdaptationTexture;
+		PassParameters->EyeAdaptationBuffer = GraphBuilder.CreateSRV(EyeAdaptationBuffer);
 		PassParameters->ThreadGroupCount = FIntPoint(SceneColor.ViewRect.Size().Y, 1);
 		PassParameters->EyeAdaptation = EyeAdaptationParameters;
 		// PassParameters->HistogramScatter64Output = GraphBuilder->CreateUAV(HistogramScatter64Texture);
@@ -395,7 +395,7 @@ static FRDGTextureRef AddHistogramAtomicPass(
 
 		FHistogramAtomicConvertCS::FParameters* PassParameters = GraphBuilder.AllocParameters<FHistogramAtomicConvertCS::FParameters>();
 		PassParameters->Input = GetScreenPassTextureViewportParameters(FScreenPassTextureViewport(SceneColor));
-		PassParameters->EyeAdaptationTexture = EyeAdaptationTexture;
+		PassParameters->EyeAdaptationBuffer = GraphBuilder.CreateSRV(EyeAdaptationBuffer);
 		// PassParameters->HistogramScatter64Texture = HistogramScatter64Texture;
 		PassParameters->HistogramScatter32Texture = HistogramScatter32Texture;
 		PassParameters->HistogramOutput = GraphBuilder.CreateUAV(HistogramTexture);
@@ -423,7 +423,7 @@ FRDGTextureRef AddHistogramPass(
 	const FEyeAdaptationParameters& EyeAdaptationParameters,
 	FScreenPassTexture SceneColor,
 	const FSceneTextureParameters& SceneTextures,
-	FRDGTextureRef EyeAdaptationTexture)
+	FRDGBufferRef EyeAdaptationBuffer)
 {
 	if (CVarUseAtomicHistogram.GetValueOnRenderThread() == 1)
 	{
@@ -433,7 +433,7 @@ FRDGTextureRef AddHistogramPass(
 			EyeAdaptationParameters,
 			SceneColor,
 			SceneTextures,
-			EyeAdaptationTexture);
+			EyeAdaptationBuffer);
 	}
 	else
 	{
@@ -442,7 +442,7 @@ FRDGTextureRef AddHistogramPass(
 			View,
 			EyeAdaptationParameters,
 			SceneColor,
-			EyeAdaptationTexture);
+			EyeAdaptationBuffer);
 	}
 }
 
diff --git a/Engine/Source/Runtime/Renderer/Private/PostProcess/PostProcessHistogram.h b/Engine/Source/Runtime/Renderer/Private/PostProcess/PostProcessHistogram.h
--- a/Engine/Source/Runtime/Renderer/Private/PostProcess/PostProcessHistogram.h
+++ b/Engine/Source/Runtime/Renderer/Private/PostProcess/PostProcessHistogram.h
@@ -12,7 +12,7 @@ FRDGTextureRef AddHistogramPass(
 	const FEyeAdaptationParameters& EyeAdaptationParameters,
 	FScreenPassTexture SceneColor,
 	const FSceneTextureParameters& SceneTextures,
-	FRDGTextureRef EyeAdaptationTexture);
+	FRDGBufferRef EyeAdaptationBuffer);
 
 FRDGTextureRef AddLocalExposurePass(
 	FRDGBuilder& GraphBuilder,
diff --git a/Engine/Source/Runtime/Renderer/Private/PostProcess/PostProcessLocalExposure.cpp b/Engine/Source/Runtime/Renderer/Private/PostProcess/PostProcessLocalExposure.cpp
--- a/Engine/Source/Runtime/Renderer/Private/PostProcess/PostProcessLocalExposure.cpp
+++ b/Engine/Source/Runtime/Renderer/Private/PostProcess/PostProcessLocalExposure.cpp
@@ -64,7 +64,7 @@ public:
 		SHADER_PARAMETER_RDG_TEXTURE_UAV(RWTexture2D<float4>, OutputFloat4)
 
 		SHADER_PARAMETER_STRUCT(FEyeAdaptationParameters, EyeAdaptation)
-		SHADER_PARAMETER_RDG_TEXTURE(Texture2D, EyeAdaptationTexture)
+		SHADER_PARAMETER_RDG_BUFFER_SRV(StructuredBuffer<float4>, EyeAdaptationBuffer)
 		SHADER_PARAMETER_RDG_TEXTURE(Texture3D, LumBilateralGrid)
 		SHADER_PARAMETER_RDG_TEXTURE(Texture2D, BlurredLogLum)
 
@@ -148,7 +148,7 @@ void AddApplyLocalExposurePass(
 	FRDGBuilder& GraphBuilder,
 	const FViewInfo& View,
 	const FEyeAdaptationParameters& EyeAdaptationParameters,
-	FRDGTextureRef EyeAdaptationTexture,
+	FRDGBufferRef EyeAdaptationBuffer,
 	FRDGTextureRef LocalExposureTexture,
 	FRDGTextureRef BlurredLogLuminanceTexture,
 	FScreenPassTexture Input,
@@ -169,7 +169,7 @@ void AddApplyLocalExposurePass(
 	PassParameters->OutputFloat4 = GraphBuilder.CreateUAV(Output.Texture);
 
 	PassParameters->EyeAdaptation = EyeAdaptationParameters;
-	PassParameters->EyeAdaptationTexture = EyeAdaptationTexture;
+	PassParameters->EyeAdaptationBuffer = GraphBuilder.CreateSRV(EyeAdaptationBuffer);
 	PassParameters->LumBilateralGrid = LocalExposureTexture;
 	PassParameters->BlurredLogLum = BlurredLogLuminanceTexture;
 	PassParameters->TextureSampler = TStaticSamplerState<SF_Bilinear, AM_Clamp, AM_Clamp, AM_Clamp>::GetRHI();;
diff --git a/Engine/Source/Runtime/Renderer/Private/PostProcess/PostProcessLocalExposure.h b/Engine/Source/Runtime/Renderer/Private/PostProcess/PostProcessLocalExposure.h
--- a/Engine/Source/Runtime/Renderer/Private/PostProcess/PostProcessLocalExposure.h
+++ b/Engine/Source/Runtime/Renderer/Private/PostProcess/PostProcessLocalExposure.h
@@ -16,7 +16,7 @@ void AddApplyLocalExposurePass(
 	FRDGBuilder& GraphBuilder,
 	const FViewInfo& View,
 	const FEyeAdaptationParameters& EyeAdaptationParameters,
-	FRDGTextureRef EyeAdaptationTexture,
+	FRDGBufferRef EyeAdaptationBuffer,
 	FRDGTextureRef LocalExposureTexture,
 	FRDGTextureRef BlurredLogLuminanceTexture,
 	FScreenPassTexture Input,
diff --git a/Engine/Source/Runtime/Renderer/Private/PostProcess/PostProcessTonemap.cpp b/Engine/Source/Runtime/Renderer/Private/PostProcess/PostProcessTonemap.cpp
--- a/Engine/Source/Runtime/Renderer/Private/PostProcess/PostProcessTonemap.cpp
+++ b/Engine/Source/Runtime/Renderer/Private/PostProcess/PostProcessTonemap.cpp
@@ -305,8 +305,7 @@ BEGIN_SHADER_PARAMETER_STRUCT(FTonemapParameters, )
 	SHADER_PARAMETER_SAMPLER(SamplerState, LumBilateralGridSampler)
 	SHADER_PARAMETER_SAMPLER(SamplerState, BlurredLogLumSampler)
 
-	// SM5 and above use Texture2D for EyeAdaptationTexture
-	SHADER_PARAMETER_RDG_TEXTURE(Texture2D, EyeAdaptationTexture)
+	SHADER_PARAMETER_RDG_BUFFER_SRV(StructuredBuffer<float4>, EyeAdaptationBuffer)
 	SHADER_PARAMETER_RDG_TEXTURE(Texture2D, ColorGradingLUT)
 	SHADER_PARAMETER_TEXTURE(Texture2D, BloomDirtMaskTexture)
 	SHADER_PARAMETER_SAMPLER(SamplerState, ColorSampler)
@@ -325,8 +324,6 @@ BEGIN_SHADER_PARAMETER_STRUCT(FTonemapParameters, )
 	SHADER_PARAMETER(float, DefaultEyeExposure)
 	SHADER_PARAMETER(float, EditorNITLevel)
 	SHADER_PARAMETER(uint32, bOutputInHDR)
-	// ES3_1 uses EyeAdaptationBuffer
-	SHADER_PARAMETER_RDG_BUFFER_SRV(StructuredBuffer<float4>, EyeAdaptationBuffer)
 END_SHADER_PARAMETER_STRUCT()
 
 class FFilmGrainReduceCS : public FGlobalShader
@@ -534,7 +531,7 @@ FScreenPassTexture AddTonemapPass(FRDGBuilder& GraphBuilder, const FViewInfo& Vi
 	const FSceneViewFamily& ViewFamily = *(View.Family);
 	const FPostProcessSettings& PostProcessSettings = View.FinalPostProcessSettings;
 
-	const bool bIsEyeAdaptationResource = (View.GetFeatureLevel() >= ERHIFeatureLevel::SM5) ? Inputs.EyeAdaptationTexture != nullptr : Inputs.EyeAdaptationBuffer != nullptr;
+	const bool bIsEyeAdaptationResource = Inputs.EyeAdaptationBuffer != nullptr;
 	const bool bEyeAdaptation = ViewFamily.EngineShowFlags.EyeAdaptation && bIsEyeAdaptationResource;
 
 	const FScreenPassTextureViewport SceneColorViewport(Inputs.SceneColor);
@@ -741,7 +738,7 @@ FScreenPassTexture AddTonemapPass(FRDGBuilder& GraphBuilder, const FViewInfo& Vi
 	CommonParameters.BlurredLogLum = Inputs.BlurredLogLuminanceTexture;
 	CommonParameters.LumBilateralGridSampler = BilinearClampSampler;
 	CommonParameters.BlurredLogLumSampler = BilinearClampSampler;
-	CommonParameters.EyeAdaptationTexture = Inputs.EyeAdaptationTexture;
+	CommonParameters.EyeAdaptationBuffer = GraphBuilder.CreateSRV(Inputs.EyeAdaptationBuffer);
 	CommonParameters.EyeAdaptation = *Inputs.EyeAdaptationParameters;
 	CommonParameters.ColorGradingLUT = Inputs.ColorGradingTexture;
 	CommonParameters.ColorSampler = BilinearClampSampler;
@@ -817,10 +814,6 @@ FScreenPassTexture AddTonemapPass(FRDGBuilder& GraphBuilder, const FViewInfo& Vi
 	CommonParameters.LensPrincipalPointOffsetScaleInverse.Y = -View.LensPrincipalPointOffsetScale.Y / View.LensPrincipalPointOffsetScale.W;
 	CommonParameters.LensPrincipalPointOffsetScaleInverse.Z = 1.0f / View.LensPrincipalPointOffsetScale.Z;
 	CommonParameters.LensPrincipalPointOffsetScaleInverse.W = 1.0f / View.LensPrincipalPointOffsetScale.W;
-	
-	FRDGBufferSRVRef EyeAdaptationBufferSRV = Inputs.EyeAdaptationBuffer != nullptr ? GraphBuilder.CreateSRV(Inputs.EyeAdaptationBuffer) : nullptr;
-
-	CommonParameters.EyeAdaptationBuffer = EyeAdaptationBufferSRV;
 
 	// Generate permutation vector for the desktop tonemapper.
 	TonemapperPermutation::FDesktopDomain DesktopPermutationVector;
diff --git a/Engine/Source/Runtime/Renderer/Private/PostProcess/PostProcessTonemap.h b/Engine/Source/Runtime/Renderer/Private/PostProcess/PostProcessTonemap.h
--- a/Engine/Source/Runtime/Renderer/Private/PostProcess/PostProcessTonemap.h
+++ b/Engine/Source/Runtime/Renderer/Private/PostProcess/PostProcessTonemap.h
@@ -52,10 +52,7 @@ struct FTonemapInputs
 	// [Required] Color grading texture used to remap colors.
 	FRDGTextureRef ColorGradingTexture = nullptr;
 
-	// [Optional, SM5+] Eye adaptation texture used to compute exposure. If this is null, a default exposure value is used instead.
-	FRDGTextureRef EyeAdaptationTexture = nullptr;
-
-	// [Optional, ES31] Eye adaptation buffer used to compute exposure. 
+	// Eye adaptation buffer used to compute exposure. 
 	FRDGBufferRef EyeAdaptationBuffer = nullptr;
 
 	// [Raster Only] Controls whether the alpha channel of the scene texture should be written to the output texture.
diff --git a/Engine/Source/Runtime/Renderer/Private/PostProcess/PostProcessing.cpp b/Engine/Source/Runtime/Renderer/Private/PostProcess/PostProcessing.cpp
--- a/Engine/Source/Runtime/Renderer/Private/PostProcess/PostProcessing.cpp
+++ b/Engine/Source/Runtime/Renderer/Private/PostProcess/PostProcessing.cpp
@@ -246,8 +246,8 @@ void AddPostProcessingPasses(
 
 	// Default the new eye adaptation to the last one in case it's not generated this frame.
 	const FEyeAdaptationParameters EyeAdaptationParameters = GetEyeAdaptationParameters(View, ERHIFeatureLevel::SM5);
-	FRDGTextureRef LastEyeAdaptationTexture = GetEyeAdaptationTexture(GraphBuilder, View);
-	FRDGTextureRef EyeAdaptationTexture = LastEyeAdaptationTexture;
+	FRDGBufferRef LastEyeAdaptationBuffer = GetEyeAdaptationBuffer(GraphBuilder, View);
+	FRDGBufferRef EyeAdaptationBuffer = LastEyeAdaptationBuffer;
 
 	// Histogram defaults to black because the histogram eye adaptation pass is used for the manual metering mode.
 	FRDGTextureRef HistogramTexture = BlackDummy.Texture;
@@ -826,7 +826,7 @@ void AddPostProcessingPasses(
 				EyeAdaptationParameters,
 				HistogramSceneColor,
 				SceneTextureParameters,
-				LastEyeAdaptationTexture);
+				LastEyeAdaptationBuffer);
 		}
 
 		FSceneDownsampleChain SceneDownsampleChain;
@@ -852,16 +852,16 @@ void AddPostProcessingPasses(
 		if (bBasicEyeAdaptationEnabled)
 		{
 			// Use the alpha channel in the last downsample (smallest) to compute eye adaptations values.
-			EyeAdaptationTexture = AddBasicEyeAdaptationPass(
+			EyeAdaptationBuffer = AddBasicEyeAdaptationPass(
 				GraphBuilder, View,
 				EyeAdaptationParameters,
 				SceneDownsampleChain.GetLastTexture(),
-				LastEyeAdaptationTexture);
+				LastEyeAdaptationBuffer);
 		}
 		// Add histogram eye adaptation pass even if no histogram exists to support the manual clamping mode.
 		else if (bEyeAdaptationEnabled)
 		{
-			EyeAdaptationTexture = AddHistogramEyeAdaptationPass(
+			EyeAdaptationBuffer = AddHistogramEyeAdaptationPass(
 				GraphBuilder, View,
 				EyeAdaptationParameters,
 				HistogramTexture);
@@ -904,7 +904,7 @@ void AddPostProcessingPasses(
 					InputSceneColor,
 					InputResolutionFraction,
 					EyeAdaptationParameters,
-					EyeAdaptationTexture,
+					EyeAdaptationBuffer,
 					CVarBloomApplyLocalExposure.GetValueOnRenderThread() ? LocalExposureTexture : nullptr,
 					LocalExposureBlurredLogLumTexture);
 
@@ -930,7 +930,7 @@ void AddPostProcessingPasses(
 
 						FBloomSetupInputs SetupPassInputs;
 						SetupPassInputs.SceneColor = DownsampleInput;
-						SetupPassInputs.EyeAdaptationTexture = EyeAdaptationTexture;
+						SetupPassInputs.EyeAdaptationBuffer = EyeAdaptationBuffer;
 						SetupPassInputs.EyeAdaptationParameters = &EyeAdaptationParameters;
 						SetupPassInputs.LocalExposureTexture = CVarBloomApplyLocalExposure.GetValueOnRenderThread() ? LocalExposureTexture : nullptr;
 						SetupPassInputs.BlurredLogLuminanceTexture = LocalExposureBlurredLogLumTexture;
@@ -1001,7 +1001,7 @@ void AddPostProcessingPasses(
 				PassInputs.LocalExposureTexture = LocalExposureTexture;
 				PassInputs.BlurredLogLuminanceTexture = LocalExposureBlurredLogLumTexture;
 				PassInputs.EyeAdaptationParameters = &EyeAdaptationParameters;
-				PassInputs.EyeAdaptationTexture = EyeAdaptationTexture;
+				PassInputs.EyeAdaptationBuffer = EyeAdaptationBuffer;
 				PassInputs.ColorGradingTexture = ColorGradingTexture;
 				PassInputs.bWriteAlphaChannel = AntiAliasingMethod == AAM_FXAA || bProcessSceneColorAlpha;
 				PassInputs.bOutputInHDR = bTonemapOutputInHDR;
@@ -1045,7 +1045,7 @@ void AddPostProcessingPasses(
 			PassInputs.SceneColor = SceneColor;
 			PassInputs.SceneDepth = SceneDepth;
 			PassInputs.ColorGradingTexture = TryRegisterExternalTexture(GraphBuilder, View.GetTonemappingLUT());
-			PassInputs.EyeAdaptationTexture = EyeAdaptationTexture;
+			PassInputs.EyeAdaptationTexture = GetEyeAdaptationTexture(GraphBuilder, View);
 			PassInputs.SceneTextures.SceneTextures = Inputs.SceneTextures;
 
 			SceneColor = AddVisualizeLumenScenePass(GraphBuilder, View, bAnyLumenActive, PassInputs, LumenFrameTemporaries);
@@ -1105,7 +1105,7 @@ void AddPostProcessingPasses(
 			PassSequence.AcceptOverrideIfLastPass(EPass::Tonemap, PassInputs.OverrideOutput);
 			PassInputs.SceneColor = SceneColor;
 			PassInputs.EyeAdaptationParameters = &EyeAdaptationParameters;
-			PassInputs.EyeAdaptationTexture = EyeAdaptationTexture;
+			PassInputs.EyeAdaptationBuffer = EyeAdaptationBuffer;
 			PassInputs.bOutputInHDR = bViewFamilyOutputInHDR;
 			PassInputs.bGammaOnly = true;
 
@@ -1283,7 +1283,7 @@ void AddPostProcessingPasses(
 		PassInputs.SceneColor = SceneColor;
 		PassInputs.SceneDepth = SceneDepth;
 		PassInputs.ColorGradingTexture = TryRegisterExternalTexture(GraphBuilder, View.GetTonemappingLUT());
-		PassInputs.EyeAdaptationTexture = EyeAdaptationTexture;
+		PassInputs.EyeAdaptationTexture = GetEyeAdaptationTexture(GraphBuilder, View);
 		PassInputs.SceneTextures.SceneTextures = Inputs.SceneTextures;
 
 		SceneColor = AddVisualizeLumenScenePass(GraphBuilder, View, bAnyLumenActive, PassInputs, LumenFrameTemporaries);
@@ -1296,7 +1296,7 @@ void AddPostProcessingPasses(
 		PassInputs.SceneColor = SceneColor;
 		PassInputs.SceneColorBeforeTonemap = SceneColorBeforeTonemap;
 		PassInputs.HistogramTexture = HistogramTexture;
-		PassInputs.EyeAdaptationTexture = EyeAdaptationTexture;
+		PassInputs.EyeAdaptationTexture = GetEyeAdaptationTexture(GraphBuilder, View);
 		PassInputs.EyeAdaptationParameters = &EyeAdaptationParameters;
 
 		SceneColor = AddVisualizeHDRPass(GraphBuilder, View, PassInputs);
@@ -1310,7 +1310,7 @@ void AddPostProcessingPasses(
 		PassInputs.HDRSceneColor = SceneColorBeforeTonemap;
 		PassInputs.LumBilateralGridTexture = LocalExposureTexture;
 		PassInputs.BlurredLumTexture = LocalExposureBlurredLogLumTexture;
-		PassInputs.EyeAdaptationTexture = EyeAdaptationTexture;
+		PassInputs.EyeAdaptationTexture = GetEyeAdaptationTexture(GraphBuilder, View);
 		PassInputs.EyeAdaptationParameters = &EyeAdaptationParameters;
 
 		SceneColor = AddVisualizeLocalExposurePass(GraphBuilder, View, PassInputs);
@@ -1624,7 +1624,7 @@ void AddDebugViewPostProcessingPasses(FRDGBuilder& GraphBuilder, const FViewInfo
 		PassInputs.bGammaOnly = true;
 		// Do eye adaptation in ray tracing debug modes to match raster buffer visualization modes
 		PassInputs.EyeAdaptationParameters = &EyeAdaptationParameters;
-		PassInputs.EyeAdaptationTexture = GetEyeAdaptationTexture(GraphBuilder, View);
+		PassInputs.EyeAdaptationBuffer = GetEyeAdaptationBuffer(GraphBuilder, View);
 
 		SceneColor = AddTonemapPass(GraphBuilder, View, PassInputs);
 	}
@@ -2377,7 +2377,7 @@ void AddMobilePostProcessingPasses(FRDGBuilder& GraphBuilder, FScene* Scene, con
 		TonemapperInputs.SceneColor = SceneColor;
 		TonemapperInputs.Bloom = BloomOutput;
 		TonemapperInputs.EyeAdaptationParameters = &EyeAdaptationParameters;
-		TonemapperInputs.EyeAdaptationTexture = nullptr;
+		TonemapperInputs.EyeAdaptationBuffer = nullptr;
 		TonemapperInputs.ColorGradingTexture = ColorGradingTexture;
 		TonemapperInputs.bWriteAlphaChannel = View.AntiAliasingMethod == AAM_FXAA || IsPostProcessingWithAlphaChannelSupported() || bUseMobileDof || IsMobilePropagateAlphaEnabled(View.GetShaderPlatform());
 		TonemapperInputs.bOutputInHDR = bHDRTonemapperOutput;
diff --git a/Engine/Source/Runtime/Renderer/Private/RendererScene.cpp b/Engine/Source/Runtime/Renderer/Private/RendererScene.cpp
--- a/Engine/Source/Runtime/Renderer/Private/RendererScene.cpp
+++ b/Engine/Source/Runtime/Renderer/Private/RendererScene.cpp
@@ -882,10 +882,6 @@ uint64 FSceneViewState::FEyeAdaptationManager::GetGPUSizeBytes(bool bLogSizes) c
 	{
 		TotalSize += GetRenderTargetGPUSizeBytes(PooledRenderTarget[TargetIndex], bLogSizes);
 	}
-	for (FRHIGPUTextureReadback* ReadbackTexture : ExposureReadbackTextures)
-	{
-		TotalSize += GetTextureReadbackGPUSizeBytes(ReadbackTexture, bLogSizes);
-	}
 	for (int32 BufferIndex = 0; BufferIndex < UE_ARRAY_COUNT(ExposureBufferData); BufferIndex++)
 	{
 		TotalSize += GetBufferGPUSizeBytes(ExposureBufferData[BufferIndex], bLogSizes);
diff --git a/Engine/Source/Runtime/Renderer/Private/ScenePrivate.h b/Engine/Source/Runtime/Renderer/Private/ScenePrivate.h
--- a/Engine/Source/Runtime/Renderer/Private/ScenePrivate.h
+++ b/Engine/Source/Runtime/Renderer/Private/ScenePrivate.h
@@ -786,7 +786,6 @@ private:
 	bool bUpdateLastExposure;
 
 	// to implement eye adaptation / auto exposure changes over time
-	// SM5 and above should use RenderTarget and ES3_1 for mobile should use RWBuffer for read back.
 	class FEyeAdaptationManager
 	{
 	public:
@@ -802,32 +801,22 @@ private:
 
 		const TRefCountPtr<IPooledRenderTarget>& GetCurrentTexture() const
 		{
-			return GetTexture(CurrentBuffer);
+			return GetTexture(CurrentBufferIndex);
 		}
 
-		/** Return current Render Target */
-		const TRefCountPtr<IPooledRenderTarget>& GetCurrentTexture(FRHICommandList& RHICmdList)
+		const TRefCountPtr<IPooledRenderTarget>& GetCurrentTexture(FRDGBuilder& GraphBuilder)
 		{
-			return GetOrCreateTexture(RHICmdList, CurrentBuffer);
+			return GetOrCreateTexture(GraphBuilder, CurrentBufferIndex);
 		}
 
-		/** Reverse the current/last order of the targets */
-		void SwapTextures();
-
-		/** Update Last Exposure with the most recent value */
-		void UpdateLastExposureFromTexture();
-
-		/** Enqueue a pass to readback current exposure */
-		void EnqueueExposureTextureReadback(FRDGBuilder& GraphBuilder);
-
 		const TRefCountPtr<FRDGPooledBuffer>& GetCurrentBuffer() const
 		{
-			return GetBuffer(CurrentBuffer);
+			return GetBuffer(CurrentBufferIndex);
 		}
 
 		const TRefCountPtr<FRDGPooledBuffer>& GetCurrentBuffer(FRDGBuilder& GraphBuilder)
 		{
-			return GetOrCreateBuffer(GraphBuilder, CurrentBuffer);
+			return GetOrCreateBuffer(GraphBuilder, CurrentBufferIndex);
 		}
 
 		void SwapBuffers();
@@ -838,26 +827,27 @@ private:
 
 	private:
 		const TRefCountPtr<IPooledRenderTarget>& GetTexture(uint32 TextureIndex) const;
-		const TRefCountPtr<IPooledRenderTarget>& GetOrCreateTexture(FRHICommandList& RHICmdList, uint32 TextureIndex);
+		const TRefCountPtr<IPooledRenderTarget>& GetOrCreateTexture(FRDGBuilder& GraphBuilder, uint32 TextureIndex);
 
 		const TRefCountPtr<FRDGPooledBuffer>& GetBuffer(uint32 BufferIndex) const;
 		const TRefCountPtr<FRDGPooledBuffer>& GetOrCreateBuffer(FRDGBuilder& GraphBuilder, uint32 BufferIndex);
 
 		FRHIGPUBufferReadback* GetLatestReadbackBuffer();
-		FRHIGPUTextureReadback* GetLatestReadbackTexture();
 
+		// TODO: Do we need to double buffer?
+		// - for readback we copy data to readback buffers
+		// - do we ever need to access prev frame exposure AFTER current frame exposure has been calculated?
+		// - should at least make it more explicit/safe by having GetCurrentBuffer() and GetPreviousBuffer()
+		//		and assert if current is accessed too early in frame.
 		static const int32 NUM_BUFFERS = 2;
 
-		int32 CurrentBuffer = 0;
+		int32 CurrentBufferIndex = 0;
 
 		float LastExposure = 0;
 		float LastAverageSceneLuminance = 0; // 0 means invalid. Used for Exposure Compensation Curve.
 
-		// Exposure texture/buffer is double buffered
 		TRefCountPtr<IPooledRenderTarget> PooledRenderTarget[NUM_BUFFERS];
-		TArray<FRHIGPUTextureReadback*> ExposureReadbackTextures;
 
-		// ES3.1 feature level. For efficent readback use buffers instead of textures
 		TRefCountPtr<FRDGPooledBuffer> ExposureBufferData[NUM_BUFFERS];
 		TArray<FRHIGPUBufferReadback*> ExposureReadbackBuffers;
 
@@ -1189,32 +1179,10 @@ public:
 		return Texture;
 	}
 
-	IPooledRenderTarget* GetCurrentEyeAdaptationTexture(FRHICommandList& RHICmdList)
+	IPooledRenderTarget* GetCurrentEyeAdaptationTexture(FRDGBuilder& GraphBuilder)
 	{
 		bValidEyeAdaptationTexture = true;
-		return EyeAdaptationManager.GetCurrentTexture(RHICmdList).GetReference();
-	}
-
-	/** Swaps the double-buffer targets used in eye adaptation */
-	void SwapEyeAdaptationTextures()
-	{
-		EyeAdaptationManager.SwapTextures();
-	}
-
-	void UpdateEyeAdaptationLastExposureFromTexture()
-	{
-		if (bUpdateLastExposure && bValidEyeAdaptationTexture)
-		{
-			EyeAdaptationManager.UpdateLastExposureFromTexture();
-		}
-	}
-
-	void EnqueueEyeAdaptationExposureTextureReadback(FRDGBuilder& GraphBuilder)
-	{
-		if (bUpdateLastExposure && bValidEyeAdaptationTexture)
-		{
-			EyeAdaptationManager.EnqueueExposureTextureReadback(GraphBuilder);
-		}
+		return EyeAdaptationManager.GetCurrentTexture(GraphBuilder).GetReference();
 	}
 
 	FRDGPooledBuffer* GetCurrentEyeAdaptationBuffer() const override final
diff --git a/Engine/Source/Runtime/Renderer/Private/SceneRendering.cpp b/Engine/Source/Runtime/Renderer/Private/SceneRendering.cpp
--- a/Engine/Source/Runtime/Renderer/Private/SceneRendering.cpp
+++ b/Engine/Source/Runtime/Renderer/Private/SceneRendering.cpp
@@ -2200,51 +2200,19 @@ FSceneViewState* FViewInfo::GetEyeAdaptationViewState() const
 	return static_cast<FSceneViewState*>(EyeAdaptationViewState);
 }
 
-IPooledRenderTarget* FViewInfo::GetEyeAdaptationTexture(FRHICommandList& RHICmdList) const
+IPooledRenderTarget* FViewInfo::GetEyeAdaptationTexture(FRDGBuilder& GraphBuilder) const
 {
-	checkf(FeatureLevel > ERHIFeatureLevel::ES3_1, TEXT("SM5 and above use RenderTarget for read back"));
+	checkf(FeatureLevel > ERHIFeatureLevel::ES3_1, TEXT("EyeAdaptation Texture is only available on SM5 and above."));
 
 	if (FSceneViewState* EffectiveViewState = GetEyeAdaptationViewState())
 	{
-		return EffectiveViewState->GetCurrentEyeAdaptationTexture(RHICmdList);
+		return EffectiveViewState->GetCurrentEyeAdaptationTexture(GraphBuilder);
 	}
 	return nullptr;
 }
 
-void FViewInfo::SwapEyeAdaptationTextures() const
-{
-	checkf(FeatureLevel > ERHIFeatureLevel::ES3_1, TEXT("SM5 and above use RenderTarget for read back"));
-
-	if (FSceneViewState* EffectiveViewState = GetEyeAdaptationViewState())
-	{
-		EffectiveViewState->SwapEyeAdaptationTextures();
-	}
-}
-
-void FViewInfo::UpdateEyeAdaptationLastExposureFromTexture() const
-{
-	checkf(FeatureLevel > ERHIFeatureLevel::ES3_1, TEXT("SM5 and above use RenderTarget for read back"));
-
-	if (FSceneViewState* EffectiveViewState = GetEyeAdaptationViewState())
-	{
-		EffectiveViewState->UpdateEyeAdaptationLastExposureFromTexture();
-	}
-}
-
-void FViewInfo::EnqueueEyeAdaptationExposureTextureReadback(FRDGBuilder& GraphBuilder) const
-{
-	checkf(FeatureLevel > ERHIFeatureLevel::ES3_1, TEXT("SM5 and above use RenderTarget for read back"));
-
-	if (FSceneViewState* EffectiveViewState = GetEyeAdaptationViewState())
-	{
-		EffectiveViewState->EnqueueEyeAdaptationExposureTextureReadback(GraphBuilder);
-	}
-}
-
 FRDGPooledBuffer* FViewInfo::GetEyeAdaptationBuffer(FRDGBuilder& GraphBuilder) const
 {
-	checkf(FeatureLevel == ERHIFeatureLevel::ES3_1, TEXT("ES3_1 use RWBuffer for read back"));
-
 	if (FSceneViewState* EffectiveViewState = GetEyeAdaptationViewState())
 	{
 		return EffectiveViewState->GetCurrentEyeAdaptationBuffer(GraphBuilder);
@@ -2254,8 +2222,6 @@ FRDGPooledBuffer* FViewInfo::GetEyeAdaptationBuffer(FRDGBuilder& GraphBuilder) c
 
 void FViewInfo::SwapEyeAdaptationBuffers() const
 {
-	checkf(FeatureLevel == ERHIFeatureLevel::ES3_1, TEXT("ES3_1 use RWBuffer for read back"));
-
 	if (FSceneViewState* EffectiveViewState = GetEyeAdaptationViewState())
 	{
 		EffectiveViewState->SwapEyeAdaptationBuffers();
@@ -2264,8 +2230,6 @@ void FViewInfo::SwapEyeAdaptationBuffers() const
 
 void FViewInfo::UpdateEyeAdaptationLastExposureFromBuffer() const
 {
-	checkf(FeatureLevel == ERHIFeatureLevel::ES3_1, TEXT("ES3_1 use RWBuffer for read back"));
-
 	if (FSceneViewState* EffectiveViewState = GetEyeAdaptationViewState())
 	{
 		EffectiveViewState->UpdateEyeAdaptationLastExposureFromBuffer();
@@ -2274,8 +2238,6 @@ void FViewInfo::UpdateEyeAdaptationLastExposureFromBuffer() const
 
 void FViewInfo::EnqueueEyeAdaptationExposureBufferReadback(FRDGBuilder& GraphBuilder) const
 {
-	checkf(FeatureLevel == ERHIFeatureLevel::ES3_1, TEXT("ES3_1 use RWBuffer for read back"));
-
 	if (FSceneViewState* EffectiveViewState = GetEyeAdaptationViewState())
 	{
 		EffectiveViewState->EnqueueEyeAdaptationExposureBufferReadback(GraphBuilder);
diff --git a/Engine/Source/Runtime/Renderer/Private/SceneRendering.h b/Engine/Source/Runtime/Renderer/Private/SceneRendering.h
--- a/Engine/Source/Runtime/Renderer/Private/SceneRendering.h
+++ b/Engine/Source/Runtime/Renderer/Private/SceneRendering.h
@@ -1693,7 +1693,7 @@ public:
 
 	/** Allocates and returns the current eye adaptation texture. */
 	using FSceneView::GetEyeAdaptationTexture;
-	IPooledRenderTarget* GetEyeAdaptationTexture(FRHICommandList& RHICmdList) const;
+	IPooledRenderTarget* GetEyeAdaptationTexture(FRDGBuilder& GraphBuilder) const;
 
 	/** Allocates and returns the current eye adaptation buffer. */
 	using FSceneView::GetEyeAdaptationBuffer;
@@ -1706,15 +1706,12 @@ public:
 	float GetLastAverageSceneLuminance() const;
 
 	/**Swap the order of the two eye adaptation targets in the double buffer system */
-	void SwapEyeAdaptationTextures() const;
 	void SwapEyeAdaptationBuffers() const;
 
 	/** Update Last Exposure with the most recent available value */
-	void UpdateEyeAdaptationLastExposureFromTexture() const;
 	void UpdateEyeAdaptationLastExposureFromBuffer() const;
 
 	/** Enqueue a pass to readback current exposure */
-	void EnqueueEyeAdaptationExposureTextureReadback(FRDGBuilder& GraphBuilder) const;
 	void EnqueueEyeAdaptationExposureBufferReadback(FRDGBuilder& GraphBuilder) const;
 	
 	/** Returns the load action to use when overwriting all pixels of a target that you intend to read from. Takes into account the HMD hidden area mesh. */
