Implements r.Translucency.AutoBeforeDOF

diff --git a/Engine/Source/Runtime/Renderer/Private/BasePassRendering.cpp b/Engine/Source/Runtime/Renderer/Private/BasePassRendering.cpp
--- a/Engine/Source/Runtime/Renderer/Private/BasePassRendering.cpp
+++ b/Engine/Source/Runtime/Renderer/Private/BasePassRendering.cpp
@@ -1838,6 +1838,12 @@ bool FBasePassMeshProcessor::ShouldDraw(const FMaterial& Material)
 			{
 			case ETranslucencyPass::TPT_TranslucencyStandard:
 				bShouldDraw = !Material.IsTranslucencyAfterDOFEnabled() && !Material.IsTranslucencyAfterMotionBlurEnabled();
+				if (AutoBeforeDOFTranslucencyBoundary > 0.0f)
+				{
+					bShouldDraw = bShouldDraw || (Material.IsTranslucencyAfterDOFEnabled() && StrataBlendMode != SBM_ColoredTransmittanceOnly);
+					bShouldDraw = bShouldDraw || (Material.IsTranslucencyAfterDOFEnabled() && (Material.IsDualBlendingEnabled(GetFeatureLevelShaderPlatform(FeatureLevel)) ||
+						BlendMode == BLEND_Modulate || StrataBlendMode == SBM_ColoredTransmittanceOnly));
+				}
 				break;
 
 			case ETranslucencyPass::TPT_TranslucencyStandardModulate:
@@ -1877,13 +1883,64 @@ bool FBasePassMeshProcessor::TryAddMeshBatch(const FMeshBatch& RESTRICT MeshBatc
 {
 	// Determine the mesh's material and blend mode.
 	const EBlendMode BlendMode = Material.GetBlendMode();
+	const EStrataBlendMode StrataBlendMode = Material.GetStrataBlendMode();
 	const FMaterialShadingModelField ShadingModels = Material.GetShadingModels();
 	const bool bIsTranslucent = IsTranslucentBlendMode(BlendMode);
 	const FMeshDrawingPolicyOverrideSettings OverrideSettings = ComputeMeshOverrideSettings(MeshBatch);
 	const ERasterizerFillMode MeshFillMode = ComputeMeshFillMode(Material, OverrideSettings);
 	const ERasterizerCullMode MeshCullMode = ComputeMeshCullMode(Material, OverrideSettings);
 
-	bool bShouldDraw = ShouldDraw(Material);
+
+	bool bShouldDraw = false;
+	if (AutoBeforeDOFTranslucencyBoundary > 0.0f && PrimitiveSceneProxy && bIsTranslucent)
+	{
+		check(ViewIfDynamicMeshCommand);
+		check(TranslucencyPassType != ETranslucencyPass::TPT_MAX);
+
+		const FVector& BoundsOrigin = PrimitiveSceneProxy->GetBounds().Origin;
+
+		const FVector& ViewOrigin = ViewIfDynamicMeshCommand->ViewMatrices.GetViewOrigin();
+		const FVector ViewForward = ViewIfDynamicMeshCommand->ViewMatrices.GetOverriddenTranslatedViewMatrix().GetColumn(2);
+
+		const FVector CameraToObject = BoundsOrigin - ViewOrigin;
+		float Distance = FVector::DotProduct(CameraToObject, ViewForward);
+		bool bIsInDOFBackground = Distance > AutoBeforeDOFTranslucencyBoundary;
+
+		bool bIsStandardTranslucency = !Material.IsTranslucencyAfterDOFEnabled() && !Material.IsTranslucencyAfterMotionBlurEnabled();
+		bool bIsAfterDOF = Material.IsTranslucencyAfterDOFEnabled() && !Material.IsTranslucencyAfterMotionBlurEnabled() && StrataBlendMode != SBM_ColoredTransmittanceOnly;
+		bool bIsAfterDOFModulate = Material.IsTranslucencyAfterDOFEnabled() && (Material.IsDualBlendingEnabled(GetFeatureLevelShaderPlatform(FeatureLevel)) ||
+			BlendMode == BLEND_Modulate || StrataBlendMode == SBM_ColoredTransmittanceOnly);
+
+		// When AutoBeforeDOFTranslucencyBoundary is valid, we automatically move After DOF translucent meshes (never blurred by DOF)
+		// before DOF if those elements are behind the focus distance.
+		if (TranslucencyPassType == ETranslucencyPass::TPT_StandardTranslucency)
+		{
+			if (bIsStandardTranslucency)
+			{
+				bShouldDraw = true;
+			}
+			else if (bIsInDOFBackground)
+			{
+				bShouldDraw = bIsAfterDOF || bIsAfterDOFModulate;
+			}
+		}
+		else if (TranslucencyPassType == ETranslucencyPass::TPT_TranslucencyAfterDOF)
+		{
+			bShouldDraw = bIsAfterDOF && !bIsInDOFBackground;
+		}
+		else if (TranslucencyPassType == ETranslucencyPass::TPT_TranslucencyAfterDOFModulate)
+		{
+			bShouldDraw = bIsAfterDOFModulate && !bIsInDOFBackground;
+		}
+		else
+		{
+			unimplemented();
+		}
+	}
+	else
+	{
+		bShouldDraw = ShouldDraw(Material);
+	}
 
 	// Only draw opaque materials.
 	bool bResult = true;
@@ -2286,6 +2343,16 @@ FBasePassMeshProcessor::FBasePassMeshProcessor(
 	, EarlyZPassMode(Scene ? Scene->EarlyZPassMode : DDM_None)
 	, bRequiresExplicit128bitRT((Flags & EFlags::bRequires128bitRT) == EFlags::bRequires128bitRT)
 {
+	if (InTranslucencyPassType != ETranslucencyPass::TPT_MAX && InViewIfDynamicMeshCommand && ViewIfDynamicMeshCommand->bIsViewInfo)
+	{
+		const FViewInfo* ViewInfo = (FViewInfo*)ViewIfDynamicMeshCommand;
+		if (InTranslucencyPassType == ETranslucencyPass::TPT_StandardTranslucency ||
+			InTranslucencyPassType == ETranslucencyPass::TPT_TranslucencyAfterDOF ||
+			InTranslucencyPassType == ETranslucencyPass::TPT_TranslucencyAfterDOFModulate)
+		{
+			AutoBeforeDOFTranslucencyBoundary = ViewInfo->AutoBeforeDOFTranslucencyBoundary;
+		}
+	}
 }
 
 FMeshPassProcessor* CreateBasePassProcessor(ERHIFeatureLevel::Type FeatureLevel, const FScene* Scene, const FSceneView* InViewIfDynamicMeshCommand, FMeshPassDrawListContext* InDrawListContext)
diff --git a/Engine/Source/Runtime/Renderer/Private/BasePassRendering.h b/Engine/Source/Runtime/Renderer/Private/BasePassRendering.h
--- a/Engine/Source/Runtime/Renderer/Private/BasePassRendering.h
+++ b/Engine/Source/Runtime/Renderer/Private/BasePassRendering.h
@@ -830,6 +830,7 @@ private:
 	const bool bEnableReceiveDecalOutput;
 	EDepthDrawingMode EarlyZPassMode;
 	bool bRequiresExplicit128bitRT;
+	float AutoBeforeDOFTranslucencyBoundary = 0.0f;
 };
 
 ENUM_CLASS_FLAGS(FBasePassMeshProcessor::EFlags);
diff --git a/Engine/Source/Runtime/Renderer/Private/PostProcess/DiaphragmDOF.cpp b/Engine/Source/Runtime/Renderer/Private/PostProcess/DiaphragmDOF.cpp
--- a/Engine/Source/Runtime/Renderer/Private/PostProcess/DiaphragmDOF.cpp
+++ b/Engine/Source/Runtime/Renderer/Private/PostProcess/DiaphragmDOF.cpp
@@ -1315,7 +1315,7 @@ bool DiaphragmDOF::IsEnabled(const FViewInfo& View)
 	static const auto CVar = IConsoleManager::Get().FindTConsoleVariableDataInt(TEXT("r.DepthOfFieldQuality"));
 	check(CVar);
 
-	const bool bDepthOfFieldRequestedByCVar = CVar->GetValueOnRenderThread() > 0;
+	const bool bDepthOfFieldRequestedByCVar = CVar->GetValueOnAnyThread() > 0;
 
 	return
 		DiaphragmDOF::IsSupported(View.GetShaderPlatform()) &&
diff --git a/Engine/Source/Runtime/Renderer/Private/SceneRendering.cpp b/Engine/Source/Runtime/Renderer/Private/SceneRendering.cpp
--- a/Engine/Source/Runtime/Renderer/Private/SceneRendering.cpp
+++ b/Engine/Source/Runtime/Renderer/Private/SceneRendering.cpp
@@ -19,6 +19,7 @@
 #include "RenderTargetTemp.h"
 #include "RendererModule.h"
 #include "ScenePrivate.h"
+#include "PostProcess/DiaphragmDOF.h"
 #include "PostProcess/SceneFilterRendering.h"
 #include "PostProcess/PostProcessEyeAdaptation.h"
 #include "PostProcess/PostProcessSubsurface.h"
@@ -358,6 +359,12 @@ static TAutoConsoleVariable<int32> CVarNaniteShowUnsupportedError(
 
 #endif
 
+static TAutoConsoleVariable<float> CVarTranslucencyAutoBeforeDOF(
+	TEXT("r.Translucency.AutoBeforeDOF"), 0.5f,
+	TEXT("Automatically bin After DOF translucency before DOF if behind focus distance (Experimental)"),
+	ECVF_Default);
+
+
 static FParallelCommandListSet* GOutstandingParallelCommandListSet = nullptr;
 
 FOcclusionSubmittedFenceState FSceneRenderer::OcclusionSubmittedFence[FOcclusionQueryHelpers::MaxBufferedOcclusionFrames];
@@ -839,6 +846,7 @@ void FViewInfo::Init()
 	bHasSingleLayerWaterMaterial = 0;
 	bHasTranslucencySeparateModulation = 0;
 	bHasStandardTranslucencyModulation = 0;
+	AutoBeforeDOFTranslucencyBoundary = 0.0f;
 
 	NumVisibleStaticMeshElements = 0;
 	PrecomputedVisibilityData = 0;
@@ -2031,10 +2039,10 @@ void FViewInfo::CreateViewUniformBuffers(const FViewUniformShaderParameters& Par
 			InstancedViewParametersUtils::CopyIntoInstancedViewParameters(LocalInstancedViewUniformShaderParameters, Params, 1);
 		}
 
-		InstancedViewUniformBuffer = TUniformBufferRef<FInstancedViewUniformShaderParameters>::CreateUniformBufferImmediate(
+			InstancedViewUniformBuffer = TUniformBufferRef<FInstancedViewUniformShaderParameters>::CreateUniformBufferImmediate(
 			LocalInstancedViewUniformShaderParameters,
-			UniformBuffer_SingleFrame);
-	}
+				UniformBuffer_SingleFrame);
+		}
 }
 
 extern TSet<IPersistentViewUniformBufferExtension*> PersistentViewUniformBufferExtensions;
@@ -2524,6 +2532,11 @@ FSceneRenderer::FSceneRenderer(const FSceneViewFamily* InViewFamily, FHitProxyCo
 				}
 			}
 		}
+
+		if (CVarTranslucencyAutoBeforeDOF.GetValueOnGameThread() >= 0.0f && DiaphragmDOF::IsEnabled(*ViewInfo))
+		{
+			ViewInfo->AutoBeforeDOFTranslucencyBoundary = ViewInfo->FinalPostProcessSettings.DepthOfFieldFocalDistance / FMath::Clamp(1.0f - CVarTranslucencyAutoBeforeDOF.GetValueOnGameThread(), 0.01f, 1.0f);
+		}
 	}
 
 	// Catches inconsistency one engine show flags for screen percentage and whether it is supported or not.
diff --git a/Engine/Source/Runtime/Renderer/Private/SceneRendering.h b/Engine/Source/Runtime/Renderer/Private/SceneRendering.h
--- a/Engine/Source/Runtime/Renderer/Private/SceneRendering.h
+++ b/Engine/Source/Runtime/Renderer/Private/SceneRendering.h
@@ -1475,6 +1475,9 @@ public:
 	 */
 	uint32 bHasStandardTranslucencyModulation : 1;
 
+	/** Whether post DOF translucency should be rendered before DOF if primitive bounds behind DOF's focus distance. */
+	float AutoBeforeDOFTranslucencyBoundary;
+
 	/** Bitmask of all shading models used by primitives in this view */
 	uint16 ShadingModelMaskInView;
 
diff --git a/Engine/Source/Runtime/Renderer/Private/SceneVisibility.cpp b/Engine/Source/Runtime/Renderer/Private/SceneVisibility.cpp
--- a/Engine/Source/Runtime/Renderer/Private/SceneVisibility.cpp
+++ b/Engine/Source/Runtime/Renderer/Private/SceneVisibility.cpp
@@ -2472,7 +2472,7 @@ struct FRelevancePacket : public FSceneRenderingAllocatorObject<FRelevancePacket
 			{
 				if (View.Family->AllowTranslucencyAfterDOF())
 				{
-					if (ViewRelevance.bNormalTranslucency)
+					if (ViewRelevance.bNormalTranslucency || (View.AutoBeforeDOFTranslucencyBoundary > 0.0f && (ViewRelevance.bSeparateTranslucency || ViewRelevance.bSeparateTranslucencyModulate)))
 					{
 						TranslucentPrimCount.Add(ETranslucencyPass::TPT_TranslucencyStandard, ViewRelevance.bUsesSceneColorCopy);
 					}
@@ -2821,7 +2821,7 @@ struct FRelevancePacket : public FSceneRenderingAllocatorObject<FRelevancePacket
 						{
 							if (View.Family->AllowTranslucencyAfterDOF())
 							{
-								if (ViewRelevance.bNormalTranslucency)
+								if (ViewRelevance.bNormalTranslucency || (View.AutoBeforeDOFTranslucencyBoundary > 0.0f && (ViewRelevance.bSeparateTranslucency || ViewRelevance.bSeparateTranslucencyModulate)))
 								{
 									DrawCommandPacket.AddCommandsForMesh(PrimitiveIndex, PrimitiveSceneInfo, StaticMeshRelevance, StaticMesh, Scene, bCanCache, EMeshPass::TranslucencyStandard);
 								}
@@ -3305,7 +3305,7 @@ void ComputeDynamicMeshRelevance(EShadingPath ShadingPath, bool bAddLightmapDens
 	{
 		if (View.Family->AllowTranslucencyAfterDOF())
 		{
-			if (ViewRelevance.bNormalTranslucency)
+			if (ViewRelevance.bNormalTranslucency || (View.AutoBeforeDOFTranslucencyBoundary > 0.0f && (ViewRelevance.bSeparateTranslucency || ViewRelevance.bSeparateTranslucencyModulate)))
 			{
 				PassMask.Set(EMeshPass::TranslucencyStandard);
 				View.NumVisibleDynamicMeshElements[EMeshPass::TranslucencyStandard] += NumElements;
