Fixes merge collision compilation failure in translucency

diff --git a/Engine/Source/Runtime/Renderer/Private/BasePassRendering.cpp b/Engine/Source/Runtime/Renderer/Private/BasePassRendering.cpp
--- a/Engine/Source/Runtime/Renderer/Private/BasePassRendering.cpp
+++ b/Engine/Source/Runtime/Renderer/Private/BasePassRendering.cpp
@@ -1913,7 +1913,7 @@ bool FBasePassMeshProcessor::TryAddMeshBatch(const FMeshBatch& RESTRICT MeshBatc
 
 		// When AutoBeforeDOFTranslucencyBoundary is valid, we automatically move After DOF translucent meshes (never blurred by DOF)
 		// before DOF if those elements are behind the focus distance.
-		if (TranslucencyPassType == ETranslucencyPass::TPT_StandardTranslucency)
+		if (TranslucencyPassType == ETranslucencyPass::TPT_TranslucencyStandard)
 		{
 			if (bIsStandardTranslucency)
 			{
@@ -2346,7 +2346,7 @@ FBasePassMeshProcessor::FBasePassMeshProcessor(
 	if (InTranslucencyPassType != ETranslucencyPass::TPT_MAX && InViewIfDynamicMeshCommand && ViewIfDynamicMeshCommand->bIsViewInfo)
 	{
 		const FViewInfo* ViewInfo = (FViewInfo*)ViewIfDynamicMeshCommand;
-		if (InTranslucencyPassType == ETranslucencyPass::TPT_StandardTranslucency ||
+		if (InTranslucencyPassType == ETranslucencyPass::TPT_TranslucencyStandard ||
 			InTranslucencyPassType == ETranslucencyPass::TPT_TranslucencyAfterDOF ||
 			InTranslucencyPassType == ETranslucencyPass::TPT_TranslucencyAfterDOFModulate)
 		{
diff --git a/Engine/Source/Runtime/Renderer/Private/SceneVisibility.cpp b/Engine/Source/Runtime/Renderer/Private/SceneVisibility.cpp
--- a/Engine/Source/Runtime/Renderer/Private/SceneVisibility.cpp
+++ b/Engine/Source/Runtime/Renderer/Private/SceneVisibility.cpp
@@ -2472,12 +2472,12 @@ struct FRelevancePacket : public FSceneRenderingAllocatorObject<FRelevancePacket
 			{
 				if (View.Family->AllowTranslucencyAfterDOF())
 				{
-					if (ViewRelevance.bNormalTranslucency || (View.AutoBeforeDOFTranslucencyBoundary > 0.0f && (ViewRelevance.bSeparateTranslucency || ViewRelevance.bSeparateTranslucencyModulate)))
+					if ((ViewRelevance.bNormalTranslucency || (View.AutoBeforeDOFTranslucencyBoundary > 0.0f && ViewRelevance.bSeparateTranslucency)))
 					{
 						TranslucentPrimCount.Add(ETranslucencyPass::TPT_TranslucencyStandard, ViewRelevance.bUsesSceneColorCopy);
 					}
 
-					if (ViewRelevance.bNormalTranslucency && ViewRelevance.bTranslucencyModulate && View.Family->AllowStandardTranslucencySeparated())
+					if ((ViewRelevance.bNormalTranslucency || (View.AutoBeforeDOFTranslucencyBoundary > 0.0f && ViewRelevance.bSeparateTranslucency)) && ViewRelevance.bTranslucencyModulate && View.Family->AllowStandardTranslucencySeparated())
 					{
 						TranslucentPrimCount.Add(ETranslucencyPass::TPT_TranslucencyStandardModulate, ViewRelevance.bUsesSceneColorCopy);
 					}
@@ -2821,12 +2821,12 @@ struct FRelevancePacket : public FSceneRenderingAllocatorObject<FRelevancePacket
 						{
 							if (View.Family->AllowTranslucencyAfterDOF())
 							{
-								if (ViewRelevance.bNormalTranslucency || (View.AutoBeforeDOFTranslucencyBoundary > 0.0f && (ViewRelevance.bSeparateTranslucency || ViewRelevance.bSeparateTranslucencyModulate)))
+								if ((ViewRelevance.bNormalTranslucency || (View.AutoBeforeDOFTranslucencyBoundary > 0.0f && ViewRelevance.bSeparateTranslucency)))
 								{
 									DrawCommandPacket.AddCommandsForMesh(PrimitiveIndex, PrimitiveSceneInfo, StaticMeshRelevance, StaticMesh, Scene, bCanCache, EMeshPass::TranslucencyStandard);
 								}
 
-								if (ViewRelevance.bNormalTranslucency && ViewRelevance.bTranslucencyModulate && View.Family->AllowStandardTranslucencySeparated())
+								if ((ViewRelevance.bNormalTranslucency || (View.AutoBeforeDOFTranslucencyBoundary > 0.0f && ViewRelevance.bSeparateTranslucency)) && ViewRelevance.bTranslucencyModulate && View.Family->AllowStandardTranslucencySeparated())
 								{
 									DrawCommandPacket.AddCommandsForMesh(PrimitiveIndex, PrimitiveSceneInfo, StaticMeshRelevance, StaticMesh, Scene, bCanCache, EMeshPass::TranslucencyStandardModulate);
 								}
@@ -3305,13 +3305,13 @@ void ComputeDynamicMeshRelevance(EShadingPath ShadingPath, bool bAddLightmapDens
 	{
 		if (View.Family->AllowTranslucencyAfterDOF())
 		{
-			if (ViewRelevance.bNormalTranslucency || (View.AutoBeforeDOFTranslucencyBoundary > 0.0f && (ViewRelevance.bSeparateTranslucency || ViewRelevance.bSeparateTranslucencyModulate)))
+			if ((ViewRelevance.bNormalTranslucency || (View.AutoBeforeDOFTranslucencyBoundary > 0.0f && ViewRelevance.bSeparateTranslucency)))
 			{
 				PassMask.Set(EMeshPass::TranslucencyStandard);
 				View.NumVisibleDynamicMeshElements[EMeshPass::TranslucencyStandard] += NumElements;
 			}
 
-			if (ViewRelevance.bNormalTranslucency && ViewRelevance.bTranslucencyModulate && View.Family->AllowStandardTranslucencySeparated())
+			if ((ViewRelevance.bNormalTranslucency || (View.AutoBeforeDOFTranslucencyBoundary > 0.0f && ViewRelevance.bSeparateTranslucency)) && ViewRelevance.bTranslucencyModulate && View.Family->AllowStandardTranslucencySeparated())
 			{
 				PassMask.Set(EMeshPass::TranslucencyStandardModulate);
 				View.NumVisibleDynamicMeshElements[EMeshPass::TranslucencyStandardModulate] += NumElements;
