Rewrite of VariableRateShadingImageManager to split actual image generation out into a generic VariableRateShadingImageGenerator class, in preparation for adding Contrast Adaptive Shading (Tier 2 VRS). Currently only fixed foveation is available (as before the rewrite).

Additionally, the foveated VRS shader is renamed to reflect the fact that it is not necessarily the only source of a VRS image.

diff --git a/Engine/Shaders/Private/VariableRateShading.usf b/Engine/Shaders/Private/VariableRateShading.usf
deleted file mode 100644
diff --git a/Engine/Shaders/Private/VariableRateShading/FixedFoveationVariableRateShading.usf b/Engine/Shaders/Private/VariableRateShading/FixedFoveationVariableRateShading.usf
new file mode 100644
--- /dev/null
+++ b/Engine/Shaders/Private/VariableRateShading/FixedFoveationVariableRateShading.usf
@@ -0,0 +1,169 @@
+// Copyright Epic Games, Inc. All Rights Reserved.
+
+/*=============================================================================
+	FixedFoveationVariableRateShading.usf
+=============================================================================*/
+
+#include "../Common.ush"
+
+// Supplied defines:
+// SHADING_RATE_1x1			- Shading rate of 1x1
+// SHADING_RATE_1x2			- Shading rate of 1x2
+// SHADING_RATE_2x1			- Shading rate of 2x1
+// SHADING_RATE_2x2			- Shading rate of 2x2
+// SHADING_RATE_2x4			- Shading rate of 2x4
+// SHADING_RATE_4x2			- Shading rate of 4x2
+// SHADING_RATE_4x4			- Shading rate of 4x4
+
+// MAX_COMBINED_SOURCES_IN	- Maximum 	
+// SHADING_RATE_TILE_WIDTH	- Shading rate tile width
+// SHADING_RATE_TILE_HEIGHT - Shading rate tile height
+// THREADGROUP_SIZEX		- Threadgroup size (X)
+// THREADGROUP_SIZEY		- Threadgroup size (Y)
+
+// Generated output shading rate attachment.
+RWTexture2D<uint> RWOutputTexture;
+
+// Input shading rate sources.
+Texture2D<uint> CombineSourceIn_0;
+Texture2D<uint> CombineSourceIn_1;
+Texture2D<uint> CombineSourceIn_2;
+Texture2D<uint> CombineSourceIn_3;
+
+// Horizontal and vertical field of view for the current view.
+float2	FieldOfView;
+
+// Center pixel coordinate for the left eye.
+float2	LeftEyeCenterPixelXY;
+
+// Center pixel coordinate for the right eye.
+float2	RightEyeCenterPixelXY;
+
+// Diagonal of the view in pixels, squared.
+float	ViewDiagonalSquaredInPixels;
+// Cutoff percentage for full-rate shading, fixed foveation (0.0 to 1.0), squared. Outside this range will be shaded at half-rate in x and y.
+float	FixedFoveationFullRateCutoffSquared;
+// Cutoff percentage for half-rate shading, fixed foveation (0.0 to 1.0), squared. Outside this range will be shaded at quarter-rate in x and y.
+float	FixedFoveationHalfRateCutoffSquared;
+
+// Number of (optional) extra shading rate sources that we'll pull in to generate the final attachment.
+uint	CombineSourceCount;
+
+// Bitfield indicating optional supporting things (defined in the EVRSGenerationFlags enum:)
+// 	EVRSGenerationFlags::None = 0x0,
+//	EVRSGenerationFlags::StereoRendering = 0x1,
+//	EVRSGenerationFlags::SideBySideStereo = 0x2,
+//	EVRSGenerationFlags::HMDFixedFoveation = 0x4,
+//	EVRSGenerationFlags::HMDEyeTrackedFoveation = 0x8,
+uint ShadingRateAttachmentGenerationFlags;
+
+
+////////////////////////////////////////////////////////////////////////////////////////////////////
+// Get the shading rate based on foveation (either fixed (distortion-based) or eye-tracked 
+// FOV based).
+////////////////////////////////////////////////////////////////////////////////////////////////////
+
+uint GetFoveationShadingRate(float FractionalOffset, float FullCutoffSquared, float HalfCutoffSquared)
+{
+	if (FractionalOffset > HalfCutoffSquared)
+	{
+		return SHADING_RATE_4x4;
+	}
+
+	if (FractionalOffset > FullCutoffSquared)
+	{
+		return SHADING_RATE_2x2;
+	}
+
+	return SHADING_RATE_1x1;
+}
+
+////////////////////////////////////////////////////////////////////////////////////////////////////
+// Get the fractional offset of the pixel coordinate from the center of the eye-view target 
+// 0.0 = center of view.
+// 1.0 = diagonal far edge of view.
+////////////////////////////////////////////////////////////////////////////////////////////////////
+
+float GetFractionalOffsetFromEyeOrigin(float2 PixelPosition)
+{
+	if ((ShadingRateAttachmentGenerationFlags & SIDE_BY_SIDE_STEREO) != 0)
+	{
+		const float2 VecToLeft = PixelPosition - LeftEyeCenterPixelXY;
+		const float2 VecToRight = PixelPosition - RightEyeCenterPixelXY;
+		return min(dot(VecToLeft, VecToLeft), dot(VecToRight, VecToRight)) / ViewDiagonalSquaredInPixels;
+	}
+	else
+	{
+		// Single eye, just return this case.
+		const float2 Vec = PixelPosition - LeftEyeCenterPixelXY;
+		return dot(Vec, Vec) / ViewDiagonalSquaredInPixels;
+	}
+}
+
+////////////////////////////////////////////////////////////////////////////////////////////////////
+// Get the shading rate at the specified pixel position, for HMD-based fixed foveation.
+////////////////////////////////////////////////////////////////////////////////////////////////////
+
+uint GetFixedFoveationRate(uint2 PixelPositionIn)
+{
+	const float2 PixelPosition = float2((float)PixelPositionIn.x, (float)PixelPositionIn.y);
+	const float FractionalOffset = GetFractionalOffsetFromEyeOrigin(PixelPosition);
+	return GetFoveationShadingRate(FractionalOffset, FixedFoveationFullRateCutoffSquared, FixedFoveationHalfRateCutoffSquared);
+}
+
+uint GetEyetrackedFoveationRate(uint2 PixelPositionIn)
+{
+	// @todo
+	return SHADING_RATE_1x1;
+}
+
+////////////////////////////////////////////////////////////////////////////////////////////////////
+// Return the ideal combination of two specified shading rate values.
+////////////////////////////////////////////////////////////////////////////////////////////////////
+
+uint CombineShadingRates(uint Rate1, uint Rate2)
+{
+	// @todo: Right now just apply the max, but possibly break out rates in X and Y and do something smarter? 
+	return max(Rate1, Rate2);
+}
+
+[numthreads(THREADGROUP_SIZEX, THREADGROUP_SIZEY, 1)]
+void GenerateShadingRateTexture(uint3 DispatchThreadId : SV_DispatchThreadID)
+{
+	const uint2 TexelCoord = DispatchThreadId.xy;
+	uint ShadingRateOut = 0;
+
+#if 0		// Disable for now since this breaks GLES 3_1 ("sampler arrays indexed with non-constant expressions is forbidden in GLSL 1.30 and later") and is currently unused anyway.
+	Texture2D<uint> CombineSourceIn[] =
+	{
+		CombineSourceIn_0,
+		#if MAX_COMBINED_SOURCES_IN > 1
+		CombineSourceIn_1,
+		#endif
+		#if MAX_COMBINED_SOURCES_IN > 2
+		CombineSourceIn_2,
+		#endif
+		#if MAX_COMBINED_SOURCES_IN > 3
+		CombineSourceIn_3,
+		#endif
+	};
+
+	for (uint Source = 0; Source < CombineSourceCount; ++Source)
+	{
+		ShadingRateOut = CombineShadingRates(ShadingRateOut, CombineSourceIn[Source][TexelCoord]);
+	}
+#endif
+
+	if ((ShadingRateAttachmentGenerationFlags & HMD_FIXED_FOVEATION) != 0)
+	{
+		ShadingRateOut = CombineShadingRates(ShadingRateOut, GetFixedFoveationRate(TexelCoord));
+	}
+
+	if ((ShadingRateAttachmentGenerationFlags & HMD_EYETRACKED_FOVEATION) != 0)
+	{
+		ShadingRateOut = CombineShadingRates(ShadingRateOut, GetEyetrackedFoveationRate(TexelCoord));
+	}
+
+	// Conservative combination, just return the max of the two.
+	RWOutputTexture[TexelCoord] = ShadingRateOut;
+}
diff --git a/Engine/Source/Runtime/Renderer/Private/BasePassRendering.cpp b/Engine/Source/Runtime/Renderer/Private/BasePassRendering.cpp
--- a/Engine/Source/Runtime/Renderer/Private/BasePassRendering.cpp
+++ b/Engine/Source/Runtime/Renderer/Private/BasePassRendering.cpp
@@ -1082,7 +1082,6 @@ void FDeferredShadingSceneRenderer::RenderBasePass(
 	// Render targets bindings should remain constant at this point.
 	FRenderTargetBindingSlots BasePassRenderTargets = GetRenderTargetBindings(ERenderTargetLoadAction::ELoad, BasePassTexturesView);
 	BasePassRenderTargets.DepthStencil = FDepthStencilBinding(BasePassDepthTexture, ERenderTargetLoadAction::ELoad, ERenderTargetLoadAction::ELoad, ExclusiveDepthStencil);
-	BasePassRenderTargets.ShadingRateTexture = GVRSImageManager.GetVariableRateShadingImage(GraphBuilder, ViewFamily, nullptr, EVRSType::None);
 	
 	FForwardBasePassTextures ForwardBasePassTextures{};
 
@@ -1434,6 +1433,7 @@ void FDeferredShadingSceneRenderer::RenderBasePassInternal(
 				PassParameters->ReflectionCapture = View.ReflectionCaptureUniformBuffer;
 				PassParameters->BasePass = CreateOpaqueBasePassUniformBuffer(GraphBuilder, View, ViewIndex, ForwardBasePassTextures, DBufferTextures, bLumenGIEnabled);
 				PassParameters->RenderTargets = BasePassRenderTargets;
+				PassParameters->RenderTargets.ShadingRateTexture = GVRSImageManager.GetVariableRateShadingImage(GraphBuilder, View, FVariableRateShadingImageManager::EVRSPassType::BasePass, nullptr);
 
 				const bool bShouldRenderView = View.ShouldRenderView();
 				if (bShouldRenderView)
@@ -1500,6 +1500,7 @@ void FDeferredShadingSceneRenderer::RenderBasePassInternal(
 				PassParameters->ReflectionCapture = View.ReflectionCaptureUniformBuffer;
 				PassParameters->BasePass = CreateOpaqueBasePassUniformBuffer(GraphBuilder, View, ViewIndex, ForwardBasePassTextures, DBufferTextures, bLumenGIEnabled);
 				PassParameters->RenderTargets = BasePassRenderTargets;
+				PassParameters->RenderTargets.ShadingRateTexture = GVRSImageManager.GetVariableRateShadingImage(GraphBuilder, View, FVariableRateShadingImageManager::EVRSPassType::BasePass, nullptr);
 
 				const bool bShouldRenderView = View.ShouldRenderView();
 				if (bShouldRenderView)
diff --git a/Engine/Source/Runtime/Renderer/Private/DeferredShadingRenderer.cpp b/Engine/Source/Runtime/Renderer/Private/DeferredShadingRenderer.cpp
--- a/Engine/Source/Runtime/Renderer/Private/DeferredShadingRenderer.cpp
+++ b/Engine/Source/Runtime/Renderer/Private/DeferredShadingRenderer.cpp
@@ -83,6 +83,7 @@
 #include "HeterogeneousVolumes/HeterogeneousVolumes.h"
 #include "ComponentRecreateRenderStateContext.h"
 #include "RenderCore.h"
+#include "VariableRateShadingImageManager.h"
 
 extern int32 GNaniteShowStats;
 extern int32 GNanitePickingDomain;
@@ -3010,6 +3011,8 @@ void FDeferredShadingSceneRenderer::Render(FRDGBuilder& GraphBuilder)
 
 	AddResolveSceneDepthPass(GraphBuilder, Views, SceneTextures.Depth);
 
+	GVRSImageManager.PrepareImageBasedVRS(GraphBuilder, ViewFamily, SceneTextures);
+
 	FComputeLightGridOutput ComputeLightGridOutput = {};
 
 	// NOTE: The ordering of the lights is used to select sub-sets for different purposes, e.g., those that support clustered deferred.
diff --git a/Engine/Source/Runtime/Renderer/Private/TranslucentRendering.cpp b/Engine/Source/Runtime/Renderer/Private/TranslucentRendering.cpp
--- a/Engine/Source/Runtime/Renderer/Private/TranslucentRendering.cpp
+++ b/Engine/Source/Runtime/Renderer/Private/TranslucentRendering.cpp
@@ -1230,7 +1230,7 @@ static void RenderTranslucencyViewInner(
 	{
 		PassParameters->RenderTargets.DepthStencil = FDepthStencilBinding(SceneDepthTexture, ERenderTargetLoadAction::ELoad, ERenderTargetLoadAction::ELoad, FExclusiveDepthStencil::DepthRead_StencilWrite);
 	}
-	PassParameters->RenderTargets.ShadingRateTexture = GVRSImageManager.GetVariableRateShadingImage(GraphBuilder, SceneRenderer.ViewFamily, nullptr);
+	PassParameters->RenderTargets.ShadingRateTexture = GVRSImageManager.GetVariableRateShadingImage(GraphBuilder, View, FVariableRateShadingImageManager::EVRSPassType::TranslucencyAll, nullptr);
 	PassParameters->RenderTargets.ResolveRect = FResolveRect(Viewport.Rect);
 
 	const EMeshPass::Type MeshPass = TranslucencyPassToMeshPass(TranslucencyPass);
diff --git a/Engine/Source/Runtime/Renderer/Private/VariableRateShading/FixedFoveationImageGenerator.cpp b/Engine/Source/Runtime/Renderer/Private/VariableRateShading/FixedFoveationImageGenerator.cpp
new file mode 100644
--- /dev/null
+++ b/Engine/Source/Runtime/Renderer/Private/VariableRateShading/FixedFoveationImageGenerator.cpp
@@ -0,0 +1,267 @@
+// Copyright Epic Games, Inc. All Rights Reserved.
+
+#include "FixedFoveationImageGenerator.h"
+#include "SystemTextures.h"
+#include "GlobalShader.h"
+#include "IXRTrackingSystem.h"
+#include "StereoRendering.h"
+#include "SceneView.h"
+#include "UnrealClient.h"
+
+/* CVar values used to control generator behavior */
+
+static TAutoConsoleVariable<int> CVarFixedFoveationLevel(
+	TEXT("vr.VRS.HMDFixedFoveationLevel"),
+	0,
+	TEXT("Level of fixed-foveation VRS to apply (when Variable Rate Shading is available)\n")
+	TEXT(" 0: Disabled (default);\n")
+	TEXT(" 1: Low;\n")
+	TEXT(" 2: Medium;\n")
+	TEXT(" 3: High;\n")
+	TEXT(" 4: High Top;\n"),
+	ECVF_RenderThreadSafe);
+
+static TAutoConsoleVariable<int32> CVarFixedFoveationDynamic(
+	TEXT("vr.VRS.HMDFixedFoveationDynamic"),
+	0,
+	TEXT("Whether fixed-foveation level should adjust based on GPU utilization\n")
+	TEXT(" 0: Disabled (default);\n")
+	TEXT(" 1: Enabled\n"),
+	ECVF_RenderThreadSafe);
+
+
+/* Image generation parameters, set up by Prepare() */
+
+FVector2f HMDFieldOfView = FVector2f(90.0f, 90.0f);
+
+float FixedFoveationFullRateCutoff = 1.0f;
+float FixedFoveationHalfRateCutoff = 1.0f;
+
+float FixedFoveationCenterX = 0.5f;
+float FixedFoveationCenterY = 0.5f;
+
+bool bGenerateFixedFoveation = false;
+
+
+/* Shader parameters and parameter structs */
+
+constexpr int32 kComputeGroupSize = FComputeShaderUtils::kGolden2DGroupSize;
+constexpr int32 kMaxCombinedSources = 4;
+
+enum class EVRSGenerationFlags : uint32
+{
+	None = 0x0,
+	StereoRendering = 0x1,
+	SideBySideStereo = 0x2,
+	HMDFixedFoveation = 0x4,
+	HMDEyeTrackedFoveation = 0x8,
+};
+
+ENUM_CLASS_FLAGS(EVRSGenerationFlags);
+
+
+/* The shader itself, which generates the foveated shading rate image */
+
+class FComputeVariableRateShadingImageGeneration : public FGlobalShader
+{
+public:
+	DECLARE_GLOBAL_SHADER(FComputeVariableRateShadingImageGeneration);
+	SHADER_USE_PARAMETER_STRUCT(FComputeVariableRateShadingImageGeneration, FGlobalShader);
+
+	BEGIN_SHADER_PARAMETER_STRUCT(FParameters, )
+		SHADER_PARAMETER_RDG_TEXTURE_UAV(RWTexture2D<uint>, RWOutputTexture)
+		SHADER_PARAMETER_RDG_TEXTURE_ARRAY(Texture2D<uint>, CombineSourceIn, [kMaxCombinedSources])
+		SHADER_PARAMETER(FVector2f, HMDFieldOfView)
+		SHADER_PARAMETER(FVector2f, LeftEyeCenterPixelXY)
+		SHADER_PARAMETER(FVector2f, RightEyeCenterPixelXY)
+		SHADER_PARAMETER(float, ViewDiagonalSquaredInPixels)
+		SHADER_PARAMETER(float, FixedFoveationFullRateCutoffSquared)
+		SHADER_PARAMETER(float, FixedFoveationHalfRateCutoffSquared)
+		SHADER_PARAMETER(uint32, CombineSourceCount)
+		SHADER_PARAMETER(uint32, ShadingRateAttachmentGenerationFlags)
+	END_SHADER_PARAMETER_STRUCT()
+
+	static bool ShouldCompilePermutation(const FGlobalShaderPermutationParameters& Parameters)
+	{
+		return FDataDrivenShaderPlatformInfo::GetSupportsVariableRateShading(Parameters.Platform);
+	}
+
+	static void ModifyCompilationEnvironment(const FGlobalShaderPermutationParameters& Parameters, FShaderCompilerEnvironment& OutEnvironment)
+	{
+		FGlobalShader::ModifyCompilationEnvironment(Parameters, OutEnvironment);
+
+		// Shading rates:
+		OutEnvironment.SetDefine(TEXT("SHADING_RATE_1x1"), VRSSR_1x1);
+		OutEnvironment.SetDefine(TEXT("SHADING_RATE_1x2"), VRSSR_1x2);
+		OutEnvironment.SetDefine(TEXT("SHADING_RATE_2x1"), VRSSR_2x1);
+		OutEnvironment.SetDefine(TEXT("SHADING_RATE_2x2"), VRSSR_2x2);
+		OutEnvironment.SetDefine(TEXT("SHADING_RATE_2x4"), VRSSR_2x4);
+		OutEnvironment.SetDefine(TEXT("SHADING_RATE_4x2"), VRSSR_4x2);
+		OutEnvironment.SetDefine(TEXT("SHADING_RATE_4x4"), VRSSR_4x4);
+
+		OutEnvironment.SetDefine(TEXT("MAX_COMBINED_SOURCES_IN"), kMaxCombinedSources);
+		OutEnvironment.SetDefine(TEXT("SHADING_RATE_TILE_WIDTH"), GRHIVariableRateShadingImageTileMinWidth);
+		OutEnvironment.SetDefine(TEXT("SHADING_RATE_TILE_HEIGHT"), GRHIVariableRateShadingImageTileMinHeight);
+		OutEnvironment.SetDefine(TEXT("THREADGROUP_SIZEX"), kComputeGroupSize);
+		OutEnvironment.SetDefine(TEXT("THREADGROUP_SIZEY"), kComputeGroupSize);
+
+		OutEnvironment.SetDefine(TEXT("STEREO_RENDERING"), (uint32)EVRSGenerationFlags::StereoRendering);
+		OutEnvironment.SetDefine(TEXT("SIDE_BY_SIDE_STEREO"), (uint32)EVRSGenerationFlags::SideBySideStereo);
+		OutEnvironment.SetDefine(TEXT("HMD_FIXED_FOVEATION"), (uint32)EVRSGenerationFlags::HMDFixedFoveation);
+		OutEnvironment.SetDefine(TEXT("HMD_EYETRACKED_FOVEATION"), (uint32)EVRSGenerationFlags::HMDEyeTrackedFoveation);
+	}
+};
+
+IMPLEMENT_GLOBAL_SHADER(FComputeVariableRateShadingImageGeneration, "/Engine/Private/VariableRateShading/FixedFoveationVariableRateShading.usf", "GenerateShadingRateTexture", SF_Compute);
+
+FRDGTextureRef FFixedFoveationImageGenerator::GetImage(FRDGBuilder& GraphBuilder, const FViewInfo& ViewInfo, FVariableRateShadingImageManager::EVRSPassType PassType)
+{
+	return CachedImage;
+}
+
+void FFixedFoveationImageGenerator::PrepareImages(FRDGBuilder& GraphBuilder, const FSceneViewFamily& ViewFamily, const FMinimalSceneTextures& SceneTextures)
+{
+	// TODO: Skip some portion of this update if the CVars haven't changed at all
+
+	// VRS level parameters - pretty arbitrary right now, later should depend on device characteristics
+	static const float kFixedFoveationFullRateCutoffs[] = { 1.0f, 0.7f, 0.50f, 0.35f, 0.35f };
+	static const float kFixedFoveationHalfRateCutofffs[] = { 1.0f, 0.9f, 0.75f, 0.55f, 0.55f };
+	static const float kFixedFoveationCenterX[] = { 0.5f, 0.5f, 0.5f, 0.5f, 0.5f };
+	static const float kFixedFoveationCenterY[] = { 0.5f, 0.5f, 0.5f, 0.5f, 0.42f };
+
+	const int VRSMaxLevel = FMath::Clamp(CVarFixedFoveationLevel.GetValueOnAnyThread(), 0, 4);
+
+	float VRSAmount = 1.0f;
+
+	if (CVarFixedFoveationDynamic.GetValueOnAnyThread() && VRSMaxLevel > 0)
+	{
+		VRSAmount = GetDynamicVRSAmount();
+	}
+
+	bGenerateFixedFoveation = (VRSMaxLevel > 0 && VRSAmount > 0.0f);
+	FixedFoveationFullRateCutoff = FMath::Lerp(kFixedFoveationFullRateCutoffs[0], kFixedFoveationFullRateCutoffs[VRSMaxLevel], VRSAmount);
+	FixedFoveationHalfRateCutoff = FMath::Lerp(kFixedFoveationHalfRateCutofffs[0], kFixedFoveationHalfRateCutofffs[VRSMaxLevel], VRSAmount);
+	FixedFoveationCenterX = FMath::Lerp(kFixedFoveationCenterX[0], kFixedFoveationCenterX[VRSMaxLevel], VRSAmount);
+	FixedFoveationCenterY = FMath::Lerp(kFixedFoveationCenterY[0], kFixedFoveationCenterY[VRSMaxLevel], VRSAmount);
+
+	FIntPoint Size = ViewFamily.RenderTarget->GetSizeXY();
+	bool bStereoRendering = IStereoRendering::IsStereoEyeView(*ViewFamily.Views[0]) && GEngine->XRSystem.IsValid();
+	bool bInstancedStereo = ViewFamily.Views[0]->IsInstancedStereoPass();
+
+	// Sanity check VRS tile size.
+	check(GRHIVariableRateShadingImageTileMinWidth >= 8 && GRHIVariableRateShadingImageTileMinWidth <= 64 && GRHIVariableRateShadingImageTileMinHeight >= 8 && GRHIVariableRateShadingImageTileMaxHeight <= 64);
+
+	FIntPoint TextureSize(Size.X / GRHIVariableRateShadingImageTileMinWidth, Size.Y / GRHIVariableRateShadingImageTileMinHeight);
+
+	// Create texture to hold shading rate image
+	FRDGTextureDesc Desc = FRDGTextureDesc::Create2D(
+		TextureSize,
+		GRHIVariableRateShadingImageFormat,
+		FClearValueBinding::None,
+		TexCreate_Foveation | TexCreate_UAV);
+
+	FRDGTextureRef ShadingRateTexture = GraphBuilder.CreateTexture(Desc, TEXT("FixedFoveationShadingRateTexture"));
+
+	// Setup shader parameters
+	FComputeVariableRateShadingImageGeneration::FParameters* PassParameters = GraphBuilder.AllocParameters<FComputeVariableRateShadingImageGeneration::FParameters>();
+	PassParameters->RWOutputTexture = GraphBuilder.CreateUAV(ShadingRateTexture);
+
+	PassParameters->HMDFieldOfView = HMDFieldOfView;
+
+	PassParameters->FixedFoveationFullRateCutoffSquared = FixedFoveationFullRateCutoff * FixedFoveationFullRateCutoff;
+	PassParameters->FixedFoveationHalfRateCutoffSquared = FixedFoveationHalfRateCutoff * FixedFoveationHalfRateCutoff;
+
+	PassParameters->LeftEyeCenterPixelXY = FVector2f(TextureSize.X * FixedFoveationCenterX, TextureSize.Y * FixedFoveationCenterY);
+	PassParameters->RightEyeCenterPixelXY = PassParameters->LeftEyeCenterPixelXY;
+
+	// Set up flags - can later allow eye-tracking here since the shader theoretically supports it
+	EVRSGenerationFlags GenFlags = EVRSGenerationFlags::HMDFixedFoveation;
+	if (bStereoRendering)
+	{
+		EnumAddFlags(GenFlags, EVRSGenerationFlags::StereoRendering);
+	}
+
+	// If instanced (side-by-side) stereo, there's two "center" points, so adjust both eyes
+	if (bInstancedStereo)
+	{
+		PassParameters->LeftEyeCenterPixelXY.X /= 2;
+		PassParameters->RightEyeCenterPixelXY.X = PassParameters->LeftEyeCenterPixelXY.X + TextureSize.X / 2;
+		EnumAddFlags(GenFlags, EVRSGenerationFlags::SideBySideStereo);
+	}
+
+	PassParameters->ViewDiagonalSquaredInPixels = FVector2f::DotProduct(PassParameters->LeftEyeCenterPixelXY, PassParameters->LeftEyeCenterPixelXY);
+	PassParameters->CombineSourceCount = 0;
+	PassParameters->ShadingRateAttachmentGenerationFlags = (uint32)GenFlags;
+
+	TShaderMapRef<FComputeVariableRateShadingImageGeneration> ComputeShader(GetGlobalShaderMap(GMaxRHIFeatureLevel));
+
+	FIntVector GroupCount = FComputeShaderUtils::GetGroupCount(TextureSize, FComputeShaderUtils::kGolden2DGroupSize);
+
+	GraphBuilder.AddPass(
+		RDG_EVENT_NAME("GenerateFixedFoveationVRSImage"),
+		PassParameters,
+		ERDGPassFlags::Compute,
+		[PassParameters, ComputeShader, GroupCount](FRHIComputeCommandList& RHICmdList)
+		{
+			FComputeShaderUtils::Dispatch(RHICmdList, ComputeShader, *PassParameters, GroupCount);
+		});
+
+	CachedImage = ShadingRateTexture;
+}
+
+bool FFixedFoveationImageGenerator::IsEnabledForView(const FSceneView& View) const
+{
+	// Only enabled if we are in XR
+	bool bStereoRendering = IStereoRendering::IsStereoEyeView(View) && GEngine->XRSystem.IsValid();
+	return bStereoRendering && CVarFixedFoveationLevel.GetValueOnRenderThread() > 0;
+}
+
+float FFixedFoveationImageGenerator::GetDynamicVRSAmount()
+{
+	const float		kUpdateIncrement = 0.1f;
+	const int		kNumFramesToAverage = 10;
+	const double	kRoundingErrorMargin = 0.5; // Add a small margin (.5 ms) to avoid oscillation
+	const double	kDesktopMaxAllowedFrameTime = 12.50; // 80 fps
+	const double	kMobileMaxAllowedFrameTime = 16.66; // 60 fps
+
+	if (GFrameNumber != DynamicVRSData.LastUpdateFrame)
+	{
+		DynamicVRSData.LastUpdateFrame = GFrameNumber;
+
+		const double GPUFrameTime = FPlatformTime::ToMilliseconds(RHIGetGPUFrameCycles());
+
+		// Update sum for rolling average
+		if (DynamicVRSData.NumFramesStored < kNumFramesToAverage)
+		{
+			DynamicVRSData.SumBusyTime += GPUFrameTime;
+			DynamicVRSData.NumFramesStored++;
+		}
+		else if (DynamicVRSData.NumFramesStored == kNumFramesToAverage)
+		{
+			DynamicVRSData.SumBusyTime -= DynamicVRSData.SumBusyTime / kNumFramesToAverage;
+			DynamicVRSData.SumBusyTime += GPUFrameTime;
+		}
+
+		// If rolling average has sufficient samples, check if update is necessary
+		if (DynamicVRSData.NumFramesStored == kNumFramesToAverage)
+		{
+			const double AverageBusyTime = DynamicVRSData.SumBusyTime / kNumFramesToAverage;
+			const double TargetBusyTime = IsMobilePlatform(GShaderPlatformForFeatureLevel[GMaxRHIFeatureLevel]) ? kMobileMaxAllowedFrameTime : kDesktopMaxAllowedFrameTime;
+
+			if (AverageBusyTime > TargetBusyTime + kRoundingErrorMargin && DynamicVRSData.VRSAmount < 1.0f)
+			{
+				DynamicVRSData.VRSAmount = FMath::Clamp(DynamicVRSData.VRSAmount + kUpdateIncrement, 0.0f, 1.0f);
+				DynamicVRSData.SumBusyTime = 0.0;
+				DynamicVRSData.NumFramesStored = 0;
+			}
+			else if (AverageBusyTime < TargetBusyTime - kRoundingErrorMargin && DynamicVRSData.VRSAmount > 0.0f)
+			{
+				DynamicVRSData.VRSAmount = FMath::Clamp(DynamicVRSData.VRSAmount - kUpdateIncrement, 0.0f, 1.0f);
+				DynamicVRSData.SumBusyTime = 0.0;
+				DynamicVRSData.NumFramesStored = 0;
+			}
+		}
+	}
+
+	return DynamicVRSData.VRSAmount;
+}
\ No newline at end of file
diff --git a/Engine/Source/Runtime/Renderer/Private/VariableRateShading/FixedFoveationImageGenerator.h b/Engine/Source/Runtime/Renderer/Private/VariableRateShading/FixedFoveationImageGenerator.h
new file mode 100644
--- /dev/null
+++ b/Engine/Source/Runtime/Renderer/Private/VariableRateShading/FixedFoveationImageGenerator.h
@@ -0,0 +1,33 @@
+// Copyright Epic Games, Inc. All Rights Reserved.
+
+#pragma once
+
+#include "CoreMinimal.h"
+#include "RHI.h"
+#include "RendererInterface.h"
+#include "RenderGraphDefinitions.h"
+#include "VariableRateShadingImageManager.h"
+#include "Engine/Engine.h"
+
+class FFixedFoveationImageGenerator : public IVariableRateShadingImageGenerator
+{
+public:
+	virtual ~FFixedFoveationImageGenerator() override {};
+	virtual FRDGTextureRef GetImage(FRDGBuilder& GraphBuilder, const FViewInfo& ViewInfo, FVariableRateShadingImageManager::EVRSPassType PassType) override;
+	virtual void PrepareImages(FRDGBuilder& GraphBuilder, const FSceneViewFamily& ViewFamily, const FMinimalSceneTextures& SceneTextures) override;
+	virtual bool IsEnabledForView(const FSceneView& View) const override;
+	virtual FVariableRateShadingImageManager::EVRSSourceType GetType() const override
+	{
+		return FVariableRateShadingImageManager::EVRSSourceType::FixedFoveation;
+	}
+private:
+	FRDGTextureRef CachedImage = nullptr;
+	struct FDynamicVRSData
+	{
+		float	VRSAmount = 1.0f;
+		double	SumBusyTime = 0.0;
+		int		NumFramesStored = 0;
+		uint32	LastUpdateFrame = 0;
+	} DynamicVRSData;
+	float GetDynamicVRSAmount();
+};
\ No newline at end of file
diff --git a/Engine/Source/Runtime/Renderer/Private/VariableRateShading/VariableRateShadingImageManager.cpp b/Engine/Source/Runtime/Renderer/Private/VariableRateShading/VariableRateShadingImageManager.cpp
new file mode 100644
--- /dev/null
+++ b/Engine/Source/Runtime/Renderer/Private/VariableRateShading/VariableRateShadingImageManager.cpp
@@ -0,0 +1,191 @@
+// Copyright Epic Games, Inc. All Rights Reserved.
+
+#include "VariableRateShadingImageManager.h"
+#include "FixedFoveationImageGenerator.h"
+#include "StereoRenderTargetManager.h"
+#include "GlobalShader.h"
+#include "ShaderParameterUtils.h"
+#include "RenderGraphUtils.h"
+#include "RHIDefinitions.h"
+#include "RenderTargetPool.h"
+#include "SystemTextures.h"
+#include "SceneView.h"
+#include "IEyeTracker.h"
+#include "IHeadMountedDisplay.h"
+#include "IXRTrackingSystem.h"
+#include "Engine/Engine.h"
+#include "UnrealClient.h"
+#include "PostProcess/PostProcessTonemap.h"
+
+
+TGlobalResource<FVariableRateShadingImageManager> GVRSImageManager;
+
+FVariableRateShadingImageManager::FVariableRateShadingImageManager()
+	: FRenderResource()
+{
+	ImageGenerators.Add(MakeUnique<FFixedFoveationImageGenerator>());
+
+	// TODO: Add more generators
+}
+
+FVariableRateShadingImageManager::~FVariableRateShadingImageManager() {}
+
+void FVariableRateShadingImageManager::ReleaseDynamicRHI()
+{
+	GRenderTargetPool.FreeUnusedResources();
+}
+
+static EDisplayOutputFormat GetDisplayOutputFormat(const FViewInfo& View)
+{
+	FTonemapperOutputDeviceParameters Parameters = GetTonemapperOutputDeviceParameters(*View.Family);
+	return (EDisplayOutputFormat)Parameters.OutputDevice;
+}
+
+static bool IsHDR10(const EDisplayOutputFormat& OutputFormat)
+{
+	return OutputFormat == EDisplayOutputFormat::HDR_ACES_1000nit_ST2084 ||
+		OutputFormat == EDisplayOutputFormat::HDR_ACES_2000nit_ST2084;
+}
+
+bool FVariableRateShadingImageManager::IsVRSSupportedByRHI()
+{
+	return GRHISupportsAttachmentVariableRateShading && GRHIVariableRateShadingEnabled && GRHIAttachmentVariableRateShadingEnabled && FDataDrivenShaderPlatformInfo::GetSupportsVariableRateShading(GMaxRHIShaderPlatform);
+}
+
+bool FVariableRateShadingImageManager::IsVRSCompatibleWithOutputType(const EDisplayOutputFormat& OutputFormat)
+{
+	return OutputFormat == EDisplayOutputFormat::SDR_sRGB || IsHDR10(OutputFormat);
+}
+
+bool FVariableRateShadingImageManager::IsVRSCompatibleWithView(const FViewInfo& ViewInfo)
+{
+	// The VRS texture generation is currently only compatible with SDR and HDR10
+	
+	// TODO: Investigate if it's worthwhile getting scene captures working. Things that we'll need to take care of
+	// is to associate shading rate texture image with main scene, and scene capture.  But what if there is
+	// more than 1 scene capture?  Is there a unique identifier that connects two frames of scene capture.
+	return IsVRSSupportedByRHI() && !ViewInfo.bIsSceneCapture && IsVRSCompatibleWithOutputType(GetDisplayOutputFormat(ViewInfo));
+}
+
+FRDGTextureRef FVariableRateShadingImageManager::GetVariableRateShadingImage(FRDGBuilder& GraphBuilder, const FViewInfo& ViewInfo, FVariableRateShadingImageManager::EVRSPassType PassType,
+	const TArray<TRefCountPtr<IPooledRenderTarget>>* ExternalVRSSources, FVariableRateShadingImageManager::EVRSSourceType VRSTypesToExclude)
+{
+	// If the view doesn't support VRS, bail immediately
+	if (!IsVRSCompatibleWithView(ViewInfo))
+	{
+		return nullptr;
+	}
+
+	// Otherwise collate all internal sources
+	TArray<FRDGTextureRef> InternalVRSSources;
+
+	for (TUniquePtr<IVariableRateShadingImageGenerator>& Generator : ImageGenerators)
+	{
+		FRDGTextureRef Image = nullptr;
+		if (Generator->IsEnabledForView(ViewInfo) && !EnumHasAnyFlags(VRSTypesToExclude, Generator->GetType()))
+		{
+			Image = Generator->GetImage(GraphBuilder, ViewInfo, PassType);
+		}
+
+		if (Image)
+		{
+			InternalVRSSources.Add(Image);
+			break; // No need to generate more than one until we support combining
+		}
+	}
+
+	// TODO: Combine internal and external images with adjustable logic (min, max, priority, etc.)
+
+	// For now, only one source can be used
+	// Once CAS is added, source will be chosen by CVar, for now it defaults to the only available generator (fixed foveation)
+	// We fall back on the first available external source if no internal source is available
+	if (InternalVRSSources.Num() > 0)
+	{
+		return InternalVRSSources[0];
+	}
+	else if (ExternalVRSSources && ExternalVRSSources->Num() > 0)
+	{
+		return GraphBuilder.RegisterExternalTexture((*ExternalVRSSources)[0]);
+	}
+	else
+	{
+		return nullptr;
+	}
+}
+
+void FVariableRateShadingImageManager::PrepareImageBasedVRS(FRDGBuilder& GraphBuilder, const FSceneViewFamily& ViewFamily, const FMinimalSceneTextures& SceneTextures)
+{
+	// If no views support VRS, bail immediately
+	bool bIsAnyViewVRSCompatible = false;
+	for (const FSceneView* View : ViewFamily.Views)
+	{
+		check(View->bIsViewInfo);
+		auto ViewInfo = static_cast<const FViewInfo*>(View);
+		if (IsVRSCompatibleWithView(*ViewInfo))
+		{
+			bIsAnyViewVRSCompatible = true;
+			break;
+		}
+	}
+
+	if (!bIsAnyViewVRSCompatible)
+	{
+		return;
+	}
+
+	// Also bail if we're given a ViewFamily with no valid RenderTarget
+	if (ViewFamily.RenderTarget == nullptr)
+	{
+		ensureMsgf(0, TEXT("VRS Image Manager does not support ViewFamilies with no valid RenderTarget"));
+		return;
+	}
+
+	// Invoke image generators
+	for (TUniquePtr<IVariableRateShadingImageGenerator>& Generator : ImageGenerators)
+	{
+		if (Generator->IsEnabledForView(*ViewFamily.Views[0]))
+		{
+			Generator->PrepareImages(GraphBuilder, ViewFamily, SceneTextures);
+		}
+	}
+}
+
+bool FVariableRateShadingImageManager::IsTypeEnabledForView(const FSceneView& View, FVariableRateShadingImageManager::EVRSSourceType Type) const
+{
+	for (const TUniquePtr<IVariableRateShadingImageGenerator>& Generator : ImageGenerators)
+	{
+		if (EnumHasAnyFlags(Type, Generator->GetType()) && Generator->IsEnabledForView(View))
+		{
+			return true;
+		}
+	}
+	return false;
+}
+
+TRefCountPtr<IPooledRenderTarget> FVariableRateShadingImageManager::GetMobileVariableRateShadingImage(const FSceneViewFamily& ViewFamily)
+{
+	if (!(IStereoRendering::IsStereoEyeView(*ViewFamily.Views[0]) && GEngine->XRSystem.IsValid()))
+	{
+		return TRefCountPtr<IPooledRenderTarget>();
+	}
+
+	FIntPoint Size(ViewFamily.RenderTarget->GetSizeXY());
+
+	const bool bStereo = GEngine->StereoRenderingDevice.IsValid() && GEngine->StereoRenderingDevice->IsStereoEnabled();
+	IStereoRenderTargetManager* const StereoRenderTargetManager = bStereo ? GEngine->StereoRenderingDevice->GetRenderTargetManager() : nullptr;
+
+	FTexture2DRHIRef Texture;
+	FIntPoint TextureSize(0, 0);
+
+	// Allocate variable resolution texture for VR foveation if supported
+	if (StereoRenderTargetManager && StereoRenderTargetManager->NeedReAllocateShadingRateTexture(MobileHMDFixedFoveationOverrideImage))
+	{
+		bool bAllocatedShadingRateTexture = StereoRenderTargetManager->AllocateShadingRateTexture(0, Size.X, Size.Y, GRHIVariableRateShadingImageFormat, 0, TexCreate_None, TexCreate_None, Texture, TextureSize);
+		if (bAllocatedShadingRateTexture)
+		{
+			MobileHMDFixedFoveationOverrideImage = CreateRenderTarget(Texture, TEXT("ShadingRate"));
+		}
+	}
+
+	return MobileHMDFixedFoveationOverrideImage;
+}
diff --git a/Engine/Source/Runtime/Renderer/Private/VariableRateShadingImageManager.cpp b/Engine/Source/Runtime/Renderer/Private/VariableRateShadingImageManager.cpp
deleted file mode 100644
diff --git a/Engine/Source/Runtime/Renderer/Public/VariableRateShadingImageManager.h b/Engine/Source/Runtime/Renderer/Public/VariableRateShadingImageManager.h
--- a/Engine/Source/Runtime/Renderer/Public/VariableRateShadingImageManager.h
+++ b/Engine/Source/Runtime/Renderer/Public/VariableRateShadingImageManager.h
@@ -7,98 +7,119 @@
 
 #pragma once
 
-#include "CoreMinimal.h"
 #include "RHI.h"
 #include "RendererInterface.h"
 #include "RenderGraphDefinitions.h"
 
-/**
- * Explicit VRS source flags; each type of VRS image generated here will have it's own flag defined in this bitfield.
- * These flags can be used for any particular rendering path to exclude specific types of VRS generation; e.g. 
- * 
- */
-enum class EVRSType : uint32
-{
-	/** VRS Image to be applied during base pass. */
-	FixedFoveation = 0x1,
+struct FMinimalSceneTextures; // Forward declaration from SceneTextures.
 
-	/** VRS Image to be applied during post-processing. */
-	EyeTrackedFoveation = 0x2,
+class IVariableRateShadingImageGenerator;
 
-	// @todo: Add more types here as they're implemented.
+class FVariableRateShadingImageManager : public FRenderResource
+{
+public:
+	/**
+	 * Pass type used by some generators. Must be provided by passes requesting a VRS image.
+	 */
+	enum EVRSPassType
+	{
+		BasePass,
+		TranslucencyAll,
+		NaniteEmitGBufferPass,
+		SSAO,
+		SSR,
+		ReflectionEnvironmentAndSky,
+		LightFunctions,
+		Decals,
+		Num
+	};
 
-	/** Mask of all available VRS source types. */
-	All = FixedFoveation | EyeTrackedFoveation,
+	/**
+	 * Explicit VRS source flags; each type of VRS image generated here will have it's own flag defined in this bitfield.
+	 * These flags can be used for any particular rendering path to exclude specific types of VRS generation; e.g.
+	 * Used to allow exclusion of certain types when getting image.
+	 */
+	enum class EVRSSourceType : uint32
+	{
+		/** Fixed Foveation: Shading rate is decreased near the periphery of the viewport */
+		FixedFoveation = 0x1,
 
-	/** Mask of XR fixed-foveation VRS gen only. */
-	XRFoveation = FixedFoveation | EyeTrackedFoveation,
+		/** Eye-Tracked Foveation: Shading rate is decreased in the periphery of the user's gaze */
+		EyeTrackedFoveation = 0x2,
 
-	None = 0x0,
-};
+		/** Contrast Adaptive Shading: Shading rate is decreased in areas of low contrast (calculated based on the prior frame) */
+		ContrastAdaptiveShading = 0x4,
 
-ENUM_CLASS_FLAGS(EVRSType);
+		// @todo: Add more types here as they're implemented.
 
-enum class EVRSGenerationFlags : uint32;
+		/** Mask of all available VRS source types. */
+		All = FixedFoveation | EyeTrackedFoveation | ContrastAdaptiveShading,
 
-struct FVRSImageGenerationParameters;
+		/** Mask of XR foveated VRS gen only. */
+		XRFoveation = FixedFoveation | EyeTrackedFoveation,
 
-struct FDynamicVRSData
-{
-	float	VRSAmount = 1.0f;
-	double	SumBusyTime = 0.0;
-	int		NumFramesStored = 0;
-	uint32	LastUpdateFrame = 0;
-};
+		None = 0x0,
+	};
 
-class FVariableRateShadingImageManager : public FRenderResource
-{
-public:
 	FVariableRateShadingImageManager();
 	virtual ~FVariableRateShadingImageManager();
 
 	virtual void ReleaseDynamicRHI() override;
 
-	/** Get the VRS image for the specified view setup, target size/configuration, and application flags.
-	 *  If the VRS image does not exist, or is out of date, it will be rendered.
-	 *  Static VRS images will be kept until until no longer referenced. */
-	FRDGTextureRef GetVariableRateShadingImage(FRDGBuilder& GraphBuilder, const FSceneViewFamily& ViewFamily, const TArray<TRefCountPtr<IPooledRenderTarget>>* ExternalVRSSources, EVRSType VRSTypesToExclude = EVRSType::None);
-
-	/** Special case (currently) for mobile devices; right now this will return either no image, or a static fixed-foveation shading rate image
-	 *  provided by the stereo device runtime. */
+	/**
+	 * Get the combined VRS image for the specified view setup, target size/configuration, and application flags.
+	 * May be run multiple times from different passes.
+	 */
+	FRDGTextureRef GetVariableRateShadingImage(FRDGBuilder& GraphBuilder, const FViewInfo& ViewInfo, EVRSPassType PassType,
+		const TArray<TRefCountPtr<IPooledRenderTarget>>* ExternalVRSSources, FVariableRateShadingImageManager::EVRSSourceType VRSTypesToExclude = FVariableRateShadingImageManager::EVRSSourceType::None);
+
+	/**
+	 * Prepare VRS images and store them for later access.
+	 * Should be run exactly once in Render(), before attempting to get any VRS images for that frame.
+	 */
+	void PrepareImageBasedVRS(FRDGBuilder& GraphBuilder, const FSceneViewFamily& ViewFamily, const FMinimalSceneTextures& SceneTextures);
+
+	/**
+	 * Returns true if any generator among the given types is enabled, false otherwise.
+	 */
+	bool IsTypeEnabledForView(const FSceneView& View, EVRSSourceType Type) const;
+
+
+	/** 
+	 *  Special case (currently) for mobile devices; right now this will return either no image, or a static fixed-foveation shading rate image
+	 *  provided by the stereo device runtime.
+	 */
 	TRefCountPtr<IPooledRenderTarget> GetMobileVariableRateShadingImage(const FSceneViewFamily& ViewFamily);
 
-protected:
-	TRefCountPtr<IPooledRenderTarget> RenderShadingRateImage(FRDGBuilder& GraphBuilder, uint64 Key, const FVRSImageGenerationParameters& VRSImageGenParamsIn, EVRSGenerationFlags GenFlags);
-	uint64 CalculateVRSImageHash(const FVRSImageGenerationParameters& VRSImageGenParamsIn, EVRSGenerationFlags ViewFlags) const;
-	void UpdateFixedFoveationParameters(FVRSImageGenerationParameters& VRSImageGenParamsInOut);
-	void UpdateEyeTrackedFoveationParameters(FVRSImageGenerationParameters& VRSImageGenParamsInOut, const FSceneViewFamily& ViewFamily);
-	float GetDynamicVRSAmount();
-
-	/** Per frame tick/update. */
-	void Tick();
+	static bool IsVRSSupportedByRHI();
+	static bool IsVRSCompatibleWithView(const FViewInfo& View);
+	static bool IsVRSCompatibleWithOutputType(const EDisplayOutputFormat& DisplayOutputFormat);
 
 private:
-	struct FActiveTarget
-	{
-		explicit FActiveTarget(const TRefCountPtr<IPooledRenderTarget>& TargetIn)
-			: Target(TargetIn)
-			, LastUsedFrame(GFrameNumber)
-		{}
-
-		FActiveTarget(FActiveTarget&& Other)
-			: Target(MoveTemp(Other.Target))
-			, LastUsedFrame(Other.LastUsedFrame)
-		{}
-
-		TRefCountPtr<IPooledRenderTarget> Target;
-		uint64 LastUsedFrame = 0;
-	};
-
-	TMap<uint64, FActiveTarget> ActiveVRSImages;
 	TRefCountPtr<IPooledRenderTarget> MobileHMDFixedFoveationOverrideImage;
+	TArray<TUniquePtr<IVariableRateShadingImageGenerator>> ImageGenerators;
 
-	uint64 LastFrameTick;
-	FDynamicVRSData DynamicVRSData;
 };
 
+ENUM_CLASS_FLAGS(FVariableRateShadingImageManager::EVRSSourceType);
+
 RENDERER_API extern TGlobalResource<FVariableRateShadingImageManager> GVRSImageManager;
+
+class IVariableRateShadingImageGenerator
+{
+public:
+	virtual ~IVariableRateShadingImageGenerator() {};
+
+	// Returns cached VRS image.
+	virtual FRDGTextureRef GetImage(FRDGBuilder& GraphBuilder, const FViewInfo& ViewInfo, FVariableRateShadingImageManager::EVRSPassType PassType) = 0;
+
+	// Generates image(s) and saves to generator cache. Should only be run once per view per frame, in Render().
+	virtual void PrepareImages(FRDGBuilder& GraphBuilder, const FSceneViewFamily& ViewFamily, const FMinimalSceneTextures& SceneTextures) = 0;
+
+	// Returns whether or not generator is enabled - can change at runtime
+	virtual bool IsEnabledForView(const FSceneView& View) const { return false; };
+
+	// Return bitmask of generator type
+	virtual FVariableRateShadingImageManager::EVRSSourceType GetType() const { return FVariableRateShadingImageManager::EVRSSourceType::None; };
+
+};
\ No newline at end of file
