Fixes TSR ghosting when going higher frame rate than 60hz

#localization none
#tests local

diff --git a/Engine/Source/Runtime/Renderer/Private/PostProcess/TemporalSuperResolution.cpp b/Engine/Source/Runtime/Renderer/Private/PostProcess/TemporalSuperResolution.cpp
--- a/Engine/Source/Runtime/Renderer/Private/PostProcess/TemporalSuperResolution.cpp
+++ b/Engine/Source/Runtime/Renderer/Private/PostProcess/TemporalSuperResolution.cpp
@@ -51,9 +51,14 @@ TAutoConsoleVariable<int32> CVarTSRFlickeringEnable(
 	TEXT("Whether to enable the flickering detection heuristic.\n"),
 	ECVF_Scalability | ECVF_RenderThreadSafe);
 
+TAutoConsoleVariable<float> CVarTSRFlickeringFrameRateCap(
+	TEXT("r.TSR.ShadingRejection.Flickering.FrameRateCap"), 60,
+	TEXT("Framerate cap at which point there is no frame rate adjustment anymore.\n"),
+	ECVF_Scalability | ECVF_RenderThreadSafe);
+
 TAutoConsoleVariable<int32> CVarTSRFlickeringAdjustToFrameRate(
 	TEXT("r.TSR.ShadingRejection.Flickering.AdjustToFrameRate"), 1,
-	TEXT("Whether TSR settings should adjust to frame rate (Enabled by default, meant for exclusive use of EngineTest).\n"),
+	TEXT("Whether TSR settings should adjust to frame rate below r.TSR.ShadingRejection.Flickering.FrameRateCap (Enabled by default, meant for exclusive use of EngineTest).\n"),
 	ECVF_Scalability | ECVF_RenderThreadSafe);
 
 TAutoConsoleVariable<float> CVarTSRFlickeringPeriod(
@@ -63,7 +68,7 @@ TAutoConsoleVariable<float> CVarTSRFlickeringPeriod(
 
 TAutoConsoleVariable<float> CVarTSRFlickeringMaxParralaxVelocity(
 	TEXT("r.TSR.ShadingRejection.Flickering.MaxParralaxVelocity"), 10.0,
-	TEXT("Maximum parralax velocity in 1080p 60hz pixels allowed before diminishing flickering.\n"),
+	TEXT("Maximum parralax velocity in 1080p at frame rate defined by r.TSR.ShadingRejection.Flickering.FrameRateCap pixels allowed before diminishing flickering.\n"),
 	ECVF_RenderThreadSafe);
 
 TAutoConsoleVariable<int32> CVarTSRRejectionAntiAliasingQuality(
@@ -828,13 +833,14 @@ ITemporalUpscaler::FOutputs AddTemporalSuperResolutionPasses(
 
 	const bool bGrandReprojection = !bAccumulateTranslucencySeparately && CVarTSRHistoryGrandReprojection.GetValueOnRenderThread() != 0;
 
-	const float RefreshRateTo60Hz = View.Family->Time.GetDeltaRealTimeSeconds() > 0.0f && CVarTSRFlickeringAdjustToFrameRate.GetValueOnRenderThread() ? View.Family->Time.GetDeltaRealTimeSeconds() * 60.0f : 1.0f;
+	const float RefreshRateToFrameRateCap = (View.Family->Time.GetDeltaRealTimeSeconds() > 0.0f && CVarTSRFlickeringAdjustToFrameRate.GetValueOnRenderThread())
+		? View.Family->Time.GetDeltaRealTimeSeconds() * CVarTSRFlickeringFrameRateCap.GetValueOnRenderThread() : 1.0f;
 
 	// whether TSR passes can run on async compute.
 	int32 AsyncComputePasses = GSupportsEfficientAsyncCompute ? CVarTSRAsyncCompute.GetValueOnRenderThread() : 0;
 
 	// period at which history changes is considered too distracting.
-	const float FlickeringFramePeriod = CVarTSRFlickeringEnable.GetValueOnRenderThread() ? (CVarTSRFlickeringPeriod.GetValueOnRenderThread() / RefreshRateTo60Hz) : 0.0f;
+	const float FlickeringFramePeriod = CVarTSRFlickeringEnable.GetValueOnRenderThread() ? (CVarTSRFlickeringPeriod.GetValueOnRenderThread() / FMath::Max(RefreshRateToFrameRateCap, 1.0f)) : 0.0f;
 
 	ETSRHistoryFormatBits HistoryFormatBits = ETSRHistoryFormatBits::None;
 	{
@@ -1088,7 +1094,7 @@ ITemporalUpscaler::FOutputs AddTemporalSuperResolutionPasses(
 		}
 		PassParameters->VelocityExtrapolationMultiplier = FMath::Clamp(CVarTSRVelocityExtrapolation.GetValueOnRenderThread(), 0.0f, 1.0f);
 		{
-			float FlickeringMaxParralaxVelocity = RefreshRateTo60Hz * CVarTSRFlickeringMaxParralaxVelocity.GetValueOnRenderThread() * float(View.ViewRect.Width()) / 1920.0f;
+			float FlickeringMaxParralaxVelocity = RefreshRateToFrameRateCap * CVarTSRFlickeringMaxParralaxVelocity.GetValueOnRenderThread() * float(View.ViewRect.Width()) / 1920.0f;
 			PassParameters->InvFlickeringMaxParralaxVelocity = 1.0f / FlickeringMaxParralaxVelocity;
 		}
 		PassParameters->bOutputIsMovingTexture = bOutputIsMovingTexture;
