Fix crash updating scene captures with mobile renderer.

diff --git a/Engine/Source/Runtime/Renderer/Private/PostProcess/PostProcessTonemap.cpp b/Engine/Source/Runtime/Renderer/Private/PostProcess/PostProcessTonemap.cpp
--- a/Engine/Source/Runtime/Renderer/Private/PostProcess/PostProcessTonemap.cpp
+++ b/Engine/Source/Runtime/Renderer/Private/PostProcess/PostProcessTonemap.cpp
@@ -738,7 +738,7 @@ FScreenPassTexture AddTonemapPass(FRDGBuilder& GraphBuilder, const FViewInfo& Vi
 	CommonParameters.BlurredLogLum = Inputs.BlurredLogLuminanceTexture;
 	CommonParameters.LumBilateralGridSampler = BilinearClampSampler;
 	CommonParameters.BlurredLogLumSampler = BilinearClampSampler;
-	CommonParameters.EyeAdaptationBuffer = GraphBuilder.CreateSRV(Inputs.EyeAdaptationBuffer);
+	CommonParameters.EyeAdaptationBuffer = Inputs.EyeAdaptationBuffer ? GraphBuilder.CreateSRV(Inputs.EyeAdaptationBuffer) : nullptr;
 	CommonParameters.EyeAdaptation = *Inputs.EyeAdaptationParameters;
 	CommonParameters.ColorGradingLUT = Inputs.ColorGradingTexture;
 	CommonParameters.ColorSampler = BilinearClampSampler;
diff --git a/Engine/Source/Runtime/Renderer/Private/PostProcess/PostProcessing.cpp b/Engine/Source/Runtime/Renderer/Private/PostProcess/PostProcessing.cpp
--- a/Engine/Source/Runtime/Renderer/Private/PostProcess/PostProcessing.cpp
+++ b/Engine/Source/Runtime/Renderer/Private/PostProcess/PostProcessing.cpp
@@ -2377,7 +2377,6 @@ void AddMobilePostProcessingPasses(FRDGBuilder& GraphBuilder, FScene* Scene, con
 		TonemapperInputs.SceneColor = SceneColor;
 		TonemapperInputs.Bloom = BloomOutput;
 		TonemapperInputs.EyeAdaptationParameters = &EyeAdaptationParameters;
-		TonemapperInputs.EyeAdaptationBuffer = nullptr;
 		TonemapperInputs.ColorGradingTexture = ColorGradingTexture;
 		TonemapperInputs.bWriteAlphaChannel = View.AntiAliasingMethod == AAM_FXAA || IsPostProcessingWithAlphaChannelSupported() || bUseMobileDof || IsMobilePropagateAlphaEnabled(View.GetShaderPlatform());
 		TonemapperInputs.bOutputInHDR = bHDRTonemapperOutput;
