VolumetricFog - fixed flickering happening when dynamic resolution sclaing happen.

diff --git a/Engine/Shaders/Private/VolumetricFog.usf b/Engine/Shaders/Private/VolumetricFog.usf
--- a/Engine/Shaders/Private/VolumetricFog.usf
+++ b/Engine/Shaders/Private/VolumetricFog.usf
@@ -291,7 +291,7 @@ void InjectShadowedLocalLightPS(
 
 #if USE_TEMPORAL_REPROJECTION
 
-		float3 HistoryUV = ComputeVolumeUV_DEPRECATED(ComputeCellTranslatedWorldPosition(GridCoordinate, .5f), UnjitteredPrevTranslatedWorldToClip);
+		float3 HistoryUV = ComputeVolumeUV_DEPRECATED(ComputeCellTranslatedWorldPosition(GridCoordinate, .5f), UnjitteredPrevTranslatedWorldToClip) * float3(View.VolumetricFogScreenUVToHistoryVolumeUV.xy, 1.0);
 		float HistoryAlpha = HistoryWeight;
 
 		FLATTEN
@@ -523,7 +523,7 @@ void LightScatteringCS(
 	}
 
 #if USE_TEMPORAL_REPROJECTION
-	float3 HistoryUV = ComputeVolumeUV_DEPRECATED(ComputeCellTranslatedWorldPosition(GridCoordinate, .5f), UnjitteredPrevTranslatedWorldToClip);
+	float3 HistoryUV = ComputeVolumeUV_DEPRECATED(ComputeCellTranslatedWorldPosition(GridCoordinate, .5f), UnjitteredPrevTranslatedWorldToClip) * float3(View.VolumetricFogScreenUVToHistoryVolumeUV.xy, 1.0);
 	float HistoryAlpha = HistoryWeight;
 
 	FLATTEN
diff --git a/Engine/Source/Runtime/Engine/Public/SceneView.h b/Engine/Source/Runtime/Engine/Public/SceneView.h
--- a/Engine/Source/Runtime/Engine/Public/SceneView.h
+++ b/Engine/Source/Runtime/Engine/Public/SceneView.h
@@ -787,6 +787,7 @@ enum ETranslucencyVolumeCascade
 	VIEW_UNIFORM_BUFFER_MEMBER(FVector3f, VolumetricFogInvGridSize) \
 	VIEW_UNIFORM_BUFFER_MEMBER(FVector3f, VolumetricFogGridZParams) \
 	VIEW_UNIFORM_BUFFER_MEMBER(FVector2f, VolumetricFogSVPosToVolumeUV) \
+	VIEW_UNIFORM_BUFFER_MEMBER(FVector2f, VolumetricFogScreenUVToHistoryVolumeUV) \
 	VIEW_UNIFORM_BUFFER_MEMBER(float, VolumetricFogMaxDistance) \
 	VIEW_UNIFORM_BUFFER_MEMBER(FVector3f, VolumetricLightmapWorldToUVScale) \
 	VIEW_UNIFORM_BUFFER_MEMBER(FVector3f, VolumetricLightmapWorldToUVAdd) \
diff --git a/Engine/Source/Runtime/Renderer/Private/ScenePrivate.h b/Engine/Source/Runtime/Renderer/Private/ScenePrivate.h
--- a/Engine/Source/Runtime/Renderer/Private/ScenePrivate.h
+++ b/Engine/Source/Runtime/Renderer/Private/ScenePrivate.h
@@ -999,6 +999,7 @@ public:
 	TUniquePtr<FForwardLightingViewResources> ForwardLightingResources;
 
 	float LightScatteringHistoryPreExposure;
+	FVector2f LightScatteringScreenUVToHistoryVolumeUV;
 	TRefCountPtr<IPooledRenderTarget> LightScatteringHistory;
 	TRefCountPtr<IPooledRenderTarget> PrevLightScatteringConservativeDepthTexture;
 
diff --git a/Engine/Source/Runtime/Renderer/Private/SceneRendering.cpp b/Engine/Source/Runtime/Renderer/Private/SceneRendering.cpp
--- a/Engine/Source/Runtime/Renderer/Private/SceneRendering.cpp
+++ b/Engine/Source/Runtime/Renderer/Private/SceneRendering.cpp
@@ -1452,6 +1452,19 @@ void FViewInfo::SetupUniformBufferParameters(
 	SetupDefaultGlobalDistanceFieldUniformBufferParameters(ViewUniformShaderParameters);
 
 	SetupVolumetricFogUniformBufferParameters(ViewUniformShaderParameters);
+	ViewUniformShaderParameters.VolumetricFogScreenUVToHistoryVolumeUV = FVector2f::One();
+	if (ViewState)
+	{
+		// Compute LightScatteringScreenUVToVolumeUV, for the current frame resolution and volume texture resolutoni according to grid size.
+		FIntVector GetVolumetricFogGridSize(FIntPoint ViewRectSize, int32& OutVolumetricFogGridPixelSize);
+		int32 VolumetricFogGridPixelSize;
+		const FIntVector VolumetricFogGridSize = GetVolumetricFogGridSize(ViewRect.Size(), VolumetricFogGridPixelSize);
+		FVector2f LightScatteringScreenUVToVolumeUV =  FVector2f(ViewRect.Size()) / (FVector2f(VolumetricFogGridSize.X, VolumetricFogGridSize.Y) * VolumetricFogGridPixelSize);
+
+		// Due to dynamic resolution scaling, the previous frame might have had a different screen to volume UV due to padding not being aligned on resolution changes.
+		// This effectively correct history samples to account for the change as a ratio of current volume UV to history volume UV.
+		ViewUniformShaderParameters.VolumetricFogScreenUVToHistoryVolumeUV = ViewState->LightScatteringScreenUVToHistoryVolumeUV / LightScatteringScreenUVToVolumeUV;
+	}
 
 	SetupPrecomputedVolumetricLightmapUniformBufferParameters(Scene, Family->EngineShowFlags, ViewUniformShaderParameters);
 
diff --git a/Engine/Source/Runtime/Renderer/Private/VolumetricFog.cpp b/Engine/Source/Runtime/Renderer/Private/VolumetricFog.cpp
--- a/Engine/Source/Runtime/Renderer/Private/VolumetricFog.cpp
+++ b/Engine/Source/Runtime/Renderer/Private/VolumetricFog.cpp
@@ -887,6 +887,7 @@ void FViewInfo::SetupVolumetricFogUniformBufferParameters(FViewUniformShaderPara
 		ViewUniformShaderParameters.VolumetricFogInvGridSize = FVector3f::ZeroVector;
 		ViewUniformShaderParameters.VolumetricFogGridZParams = FVector3f::ZeroVector;
 		ViewUniformShaderParameters.VolumetricFogSVPosToVolumeUV = FVector2f::ZeroVector;
+		ViewUniformShaderParameters.VolumetricFogScreenUVToHistoryVolumeUV = FVector2f::ZeroVector;
 		ViewUniformShaderParameters.VolumetricFogMaxDistance = 0;
 	}
 }
@@ -924,6 +925,7 @@ void FDeferredShadingSceneRenderer::SetupVolumetricFog()
 			{
 				View.ViewState->LightScatteringHistory = NULL;
 				View.ViewState->LightScatteringHistoryPreExposure = 1.0f;
+				View.ViewState->LightScatteringScreenUVToHistoryVolumeUV = FVector2f::One();
 			}
 		}
 	}
@@ -1312,11 +1314,13 @@ void FDeferredShadingSceneRenderer::ComputeVolumetricFog(FRDGBuilder& GraphBuild
 		{
 			GraphBuilder.QueueTextureExtraction(IntegrationData.LightScattering, &View.ViewState->LightScatteringHistory);
 			View.ViewState->LightScatteringHistoryPreExposure = View.CachedViewUniformShaderParameters->PreExposure;
+			View.ViewState->LightScatteringScreenUVToHistoryVolumeUV = FVector2f(View.ViewRect.Size()) / (FVector2f(VolumetricFogGridSize.X, VolumetricFogGridSize.Y) * VolumetricFogGridPixelSize);
 		}
 		else if (View.ViewState)
 		{
 			View.ViewState->LightScatteringHistory = nullptr;
 			View.ViewState->LightScatteringHistoryPreExposure = 1.0f;
+			View.ViewState->LightScatteringScreenUVToHistoryVolumeUV = FVector2f::One();
 		}
 
 		if (bUseTemporalReprojection && GVolumetricFogConservativeDepth > 0)
