Volumetric cloud - per pixel tracing support for cloud on translucent.
r.VolumetricRenderTarget.Mode 3 is the cinematic mode: tracing is done at full resolution in a volume render target. Upsampling is direct.
This alows cloud to be applied on translucent.

diff --git a/Engine/Source/Runtime/Renderer/Private/TranslucentRendering.cpp b/Engine/Source/Runtime/Renderer/Private/TranslucentRendering.cpp
--- a/Engine/Source/Runtime/Renderer/Private/TranslucentRendering.cpp
+++ b/Engine/Source/Runtime/Renderer/Private/TranslucentRendering.cpp
@@ -888,7 +888,8 @@ TRDGUniformBufferRef<FTranslucentBasePassUniformParameters> CreateTranslucentBas
 		BasePassParameters.VolumetricCloudDepthSampler = TStaticSamplerState<SF_Bilinear, AM_Clamp, AM_Clamp, AM_Clamp>::GetRHI();
 		if (IsVolumetricRenderTargetEnabled() && View.ViewState)
 		{
-			if (View.ViewState->VolumetricCloudRenderTarget.GetMode() == 1)
+			int32 VRTMode = View.ViewState->VolumetricCloudRenderTarget.GetMode();
+			if (VRTMode == 1 || VRTMode == 3)
 			{
 				FRDGTextureRef VolumetricReconstructRT = View.ViewState->VolumetricCloudRenderTarget.GetOrCreateVolumetricTracingRT(GraphBuilder);
 				if (VolumetricReconstructRT)
diff --git a/Engine/Source/Runtime/Renderer/Private/VolumetricCloudRendering.cpp b/Engine/Source/Runtime/Renderer/Private/VolumetricCloudRendering.cpp
--- a/Engine/Source/Runtime/Renderer/Private/VolumetricCloudRendering.cpp
+++ b/Engine/Source/Runtime/Renderer/Private/VolumetricCloudRendering.cpp
@@ -2703,7 +2703,7 @@ bool FSceneRenderer::RenderVolumetricCloud(
 					CloudRC.NoiseFrameIndexModPattern = VRT.GetNoiseFrameIndexModPattern();
 
 					check(VRT.GetMode() != 0 || CloudRC.bCloudDebugViewModeEnabled || HalfResolutionDepthCheckerboardMinMaxTexture);
-					CloudRC.SceneDepthZ = (VRT.GetMode() == 0 || VRT.GetMode() == 1 || VRT.GetMode() == 3) ? HalfResolutionDepthCheckerboardMinMaxTexture : SceneDepthZ;
+					CloudRC.SceneDepthZ = (VRT.GetMode() == 0 || VRT.GetMode() == 1) ? HalfResolutionDepthCheckerboardMinMaxTexture : SceneDepthZ;
 					CloudRC.SceneDepthMinAndMax = bShouldVolumetricCloudTraceWithMinMaxDepth ? QuarterResolutionDepthMinMaxTexture : SceneDepthZ;
 				}
 				else
diff --git a/Engine/Source/Runtime/Renderer/Private/VolumetricRenderTarget.cpp b/Engine/Source/Runtime/Renderer/Private/VolumetricRenderTarget.cpp
--- a/Engine/Source/Runtime/Renderer/Private/VolumetricRenderTarget.cpp
+++ b/Engine/Source/Runtime/Renderer/Private/VolumetricRenderTarget.cpp
@@ -26,7 +26,7 @@ static TAutoConsoleVariable<float> CVarVolumetricRenderTargetUvNoiseSampleAccept
 
 static TAutoConsoleVariable<int32> CVarVolumetricRenderTargetMode(
 	TEXT("r.VolumetricRenderTarget.Mode"), 0,
-	TEXT("[0] trace quarter resolution + reconstruct at half resolution + upsample [1] trace half res + reconstruct full res + upsample [2] trace at quarter resolution + reconstruct full resolution (cannot intersect with opaque meshes and forces UpsamplingMode=2 [3] trace 1/8 resolution + reconstruct at half resolution + upsample)"),
+	TEXT("[0] trace quarter resolution + reconstruct at half resolution + upsample [1] trace half res + reconstruct full res + upsample [2] trace at quarter resolution + reconstruct full resolution (cannot intersect with opaque meshes and forces UpsamplingMode=2 [3] Cinematic mode with tracing done at full reoslution in render target so that clouds can also be applied on translucent.)"),
 	ECVF_RenderThreadSafe | ECVF_Scalability);
 
 static TAutoConsoleVariable<int32> CVarVolumetricRenderTargetUpsamplingMode(
@@ -85,13 +85,15 @@ static uint32 GetMainDownsampleFactor(int32 Mode)
 	switch (Mode)
 	{
 	case 0:
-	case 3:
 		return 2; // Reconstruct at half resolution of view
 		break;
 	case 1:
 	case 2:
 		return 1; // Reconstruct at full resolution of view
 		break;
+	case 3:
+		return 1; // SKip reconstruct, tracing at full resolution.
+		break;
 	}
 	check(false); // unhandled mode
 	return 2;
@@ -111,7 +113,7 @@ static uint32 GetTraceDownsampleFactor(int32 Mode)
 		return 4; // Trace at quarter resolution of the reconstructed buffer (with it being at the same resolution as main view)
 		break;
 	case 3:
-		return 4; // Trace at quarter resolution of the reconstructed buffer (with it being at the half the resolution of the main view)
+		return 1; // Trace at full resolution
 		break;
 	}
 	check(false); // unhandled mode
@@ -170,7 +172,7 @@ void FVolumetricRenderTargetViewStateData::Initialise(
 {
 	// Update internal settings
 	Mode = FMath::Clamp(InMode, 0, 3);
-	UpsamplingMode = Mode == 2 ? 2 : FMath::Clamp(InUpsamplingMode, 0, 4); // if we are using mode 2 then we cannot intersect with depth and upsampling should be 2 (simple on/off intersection)
+	UpsamplingMode = Mode == 2 || Mode == 3 ? 2 : FMath::Clamp(InUpsamplingMode, 0, 4); // if we are using mode 2 then we cannot intersect with depth and upsampling should be 2 (simple on/off intersection)
 
 	if (bFirstTimeUsed)
 	{
@@ -243,7 +245,7 @@ void FVolumetricRenderTargetViewStateData::Initialise(
 			}
 		}
 
-		if (Mode == 1)
+		if (Mode == 1 || Mode == 3)
 		{
 			// No need to jitter in this case. Mode on is tracing half res and then upsample without reconstruction.
 			CurrentPixelOffset = FIntPoint::ZeroValue;
@@ -361,7 +363,7 @@ FRDGTextureRef FVolumetricRenderTargetViewStateData::GetOrCreateSrcVolumetricRec
 
 FUintVector4 FVolumetricRenderTargetViewStateData::GetTracingCoordToZbufferCoordScaleBias() const
 {
-	if (Mode == 2)
+	if (Mode == 2 || Mode == 3)
 	{
 		// In this case, the source depth buffer full resolution depth buffer is the full resolution scene one
 		const uint32 CombinedDownsampleFactor = VolumetricReconstructRTDownsampleFactor * VolumetricTracingRTDownsampleFactor;
@@ -520,7 +522,7 @@ void ReconstructVolumetricRenderTarget(
 
 		FVolumetricRenderTargetViewStateData& VolumetricCloudRT = ViewInfo.ViewState->VolumetricCloudRenderTarget;
 
-		if (VolumetricCloudRT.GetMode() == 1)
+		if (VolumetricCloudRT.GetMode() == 1 || VolumetricCloudRT.GetMode() == 3)
 		{
 			// In this case, we trace at half resolution using checker boarded min max depth.
 			// We will then directly up sample on screen from half resolution to full resolution. 
@@ -655,7 +657,7 @@ static FVector2f GetCompositionFullResolutionToVolumetricBufferResolutionScale(u
 
 static void GetCompositionCloudTextures(uint32 VRTMode, FVolumetricRenderTargetViewStateData& VolumetricCloudRT, FRDGBuilder& GraphBuilder, FRDGTextureRef& VolumetricTexture, FRDGTextureRef& VolumetricDepthTexture)
 {
-	if (VRTMode == 1)
+	if (VRTMode == 1 || VRTMode == 3)
 	{
 		// In this case, we trace at half resolution using checker boarded min max depth.
 		// We will then directly up sample on screen from half resolution to full resolution. 
@@ -672,7 +674,7 @@ static void GetCompositionCloudTextures(uint32 VRTMode, FVolumetricRenderTargetV
 
 static int32 GetCompositionUpsamplingMode(uint32 VRTMode, int32 UpsamplingMode)
 {
-	return UpsamplingMode == 3 && (VRTMode == 1 || VRTMode == 2) ? 2 : UpsamplingMode;
+	return UpsamplingMode == 3 && (VRTMode == 1 || VRTMode == 2 || VRTMode == 3) ? 2 : UpsamplingMode;
 }
 
 void ComposeVolumetricRenderTargetOverScene(
