Deprecate EyeAdaptationTexture.
- still maintaing support until 5.3.
- since at least one external plugin requires exposure data in a texture, implemented a helper function AddCopyEyeAdaptationDataToTexturePass().

diff --git a/Engine/Source/Runtime/Engine/Private/SceneView.cpp b/Engine/Source/Runtime/Engine/Private/SceneView.cpp
--- a/Engine/Source/Runtime/Engine/Private/SceneView.cpp
+++ b/Engine/Source/Runtime/Engine/Private/SceneView.cpp
@@ -2708,7 +2708,9 @@ bool FSceneView::HasValidEyeAdaptationTexture() const
 {
 	if (EyeAdaptationViewState)
 	{
+PRAGMA_DISABLE_DEPRECATION_WARNINGS
 		return EyeAdaptationViewState->HasValidEyeAdaptationTexture();
+PRAGMA_ENABLE_DEPRECATION_WARNINGS
 	}
 	return false;
 }
@@ -2727,7 +2729,9 @@ IPooledRenderTarget* FSceneView::GetEyeAdaptationTexture() const
 	checkf(FeatureLevel > ERHIFeatureLevel::ES3_1, TEXT("EyeAdaptation Texture is only available on SM5 and above."));
 	if (EyeAdaptationViewState)
 	{
+PRAGMA_DISABLE_DEPRECATION_WARNINGS
 		return EyeAdaptationViewState->GetCurrentEyeAdaptationTexture();
+PRAGMA_ENABLE_DEPRECATION_WARNINGS
 	}
 	return nullptr;
 }
diff --git a/Engine/Source/Runtime/Engine/Public/SceneManagement.h b/Engine/Source/Runtime/Engine/Public/SceneManagement.h
--- a/Engine/Source/Runtime/Engine/Public/SceneManagement.h
+++ b/Engine/Source/Runtime/Engine/Public/SceneManagement.h
@@ -231,11 +231,18 @@ public:
 	 */
 	virtual uint32 GetDistanceFieldTemporalSampleIndex() const = 0;
 
-	/** Tells if the eye adaptation texture / buffer exists without attempting to allocate it. */
-	bool HasValidEyeAdaptationTexture() const { return bValidEyeAdaptationTexture; }
+	UE_DEPRECATED(5.2, "Use HasValidEyeAdaptationBuffer() instead.")
+	bool HasValidEyeAdaptationTexture() const
+	{
+PRAGMA_DISABLE_DEPRECATION_WARNINGS
+		return bValidEyeAdaptationTexture;
+PRAGMA_ENABLE_DEPRECATION_WARNINGS
+	}
+
+	/** Tells if the eye adaptation buffer exists without attempting to allocate it. */
 	bool HasValidEyeAdaptationBuffer() const { return bValidEyeAdaptationBuffer; }
 
-	/** Returns the eye adaptation texture (SM5+ only). */
+	UE_DEPRECATED(5.2, "Use GetCurrentEyeAdaptationBuffer() instead.")
 	virtual IPooledRenderTarget* GetCurrentEyeAdaptationTexture() const = 0;
 
 	/** Returns the eye adaptation buffer. */
@@ -285,6 +292,7 @@ protected:
 	// Don't allow direct deletion of the view state, Destroy should be called instead.
 	virtual ~FSceneViewStateInterface() {}
 
+	UE_DEPRECATED(5.2, "Use bValidEyeAdaptationBuffer instead.")
 	uint8 bValidEyeAdaptationTexture : 1;
 	uint8 bValidEyeAdaptationBuffer : 1;
 
diff --git a/Engine/Source/Runtime/Engine/Public/SceneView.h b/Engine/Source/Runtime/Engine/Public/SceneView.h
--- a/Engine/Source/Runtime/Engine/Public/SceneView.h
+++ b/Engine/Source/Runtime/Engine/Public/SceneView.h
@@ -1608,11 +1608,14 @@ public:
 	/** Current ray tracing debug visualization mode */
 	FName CurrentRayTracingDebugVisualizationMode;
 #endif
-	/** Tells if the eye adaptation texture / buffer exists without attempting to allocate it. */
+
+	UE_DEPRECATED(5.2, "Use HasValidEyeAdaptationBuffer() instead.")
 	bool HasValidEyeAdaptationTexture() const;
+
+	/** Tells if the eye adaptation buffer exists without attempting to allocate it. */
 	bool HasValidEyeAdaptationBuffer() const;
 
-	/** Returns the eye adaptation texture (SM5+ only) or null if it doesn't exist. */
+	UE_DEPRECATED(5.2, "Use GetEyeAdaptationBuffer() instead.")
 	IPooledRenderTarget* GetEyeAdaptationTexture() const;
 
 	/** Returns the eye adaptation buffer or null if it doesn't exist. */
diff --git a/Engine/Source/Runtime/Renderer/Private/PostProcess/PostProcessEyeAdaptation.cpp b/Engine/Source/Runtime/Renderer/Private/PostProcess/PostProcessEyeAdaptation.cpp
--- a/Engine/Source/Runtime/Renderer/Private/PostProcess/PostProcessEyeAdaptation.cpp
+++ b/Engine/Source/Runtime/Renderer/Private/PostProcess/PostProcessEyeAdaptation.cpp
@@ -563,6 +563,16 @@ void AddCopyEyeAdaptationDataToTexturePass(FRDGBuilder& GraphBuilder, const FGlo
 		FIntVector(1, 1, 1));
 }
 
+FRDGTextureRef AddCopyEyeAdaptationDataToTexturePass(FRDGBuilder& GraphBuilder, const FViewInfo& View)
+{
+	const FRDGTextureDesc TextureDesc = FRDGTextureDesc::Create2D(FIntPoint(1, 1), PF_A32B32G32R32F, FClearValueBinding::None, TexCreate_UAV | TexCreate_ShaderResource);
+	FRDGTextureRef OutputTexture = GraphBuilder.CreateTexture(TextureDesc, TEXT("EyeAdaptationTexture"));
+
+	AddCopyEyeAdaptationDataToTexturePass(GraphBuilder, View.ShaderMap, GetEyeAdaptationBuffer(GraphBuilder, View), OutputTexture);
+
+	return OutputTexture;
+}
+
 //////////////////////////////////////////////////////////////////////////
 //! Setup
 //////////////////////////////////////////////////////////////////////////
@@ -909,10 +919,12 @@ FRDGBufferRef AddBasicEyeAdaptationPass(
 		PassParameters,
 		FIntVector(1, 1, 1));
 
+PRAGMA_DISABLE_DEPRECATION_WARNINGS
 	{
 		FRDGTextureRef OutputTexture = GraphBuilder.RegisterExternalTexture(View.GetEyeAdaptationTexture(GraphBuilder), ERDGTextureFlags::MultiFrame);
 		AddCopyEyeAdaptationDataToTexturePass(GraphBuilder, View.ShaderMap, OutputBuffer, OutputTexture);
 	}
+PRAGMA_ENABLE_DEPRECATION_WARNINGS
 
 	View.EnqueueEyeAdaptationExposureBufferReadback(GraphBuilder);
 
@@ -936,7 +948,9 @@ void FSceneViewState::FEyeAdaptationManager::SafeRelease()
 	for (int32 Index = 0; Index < NUM_BUFFERS; Index++)
 	{
 		ExposureBufferData[Index].SafeRelease();
+PRAGMA_DISABLE_DEPRECATION_WARNINGS
 		PooledRenderTarget[Index].SafeRelease();
+PRAGMA_ENABLE_DEPRECATION_WARNINGS
 	}
 
 	for (int32 Index = 0; Index < ExposureReadbackBuffers.Num(); ++Index)
@@ -952,13 +966,17 @@ void FSceneViewState::FEyeAdaptationManager::SafeRelease()
 const TRefCountPtr<IPooledRenderTarget>& FSceneViewState::FEyeAdaptationManager::GetTexture(uint32 TextureIndex) const
 {
 	check(0 <= TextureIndex && TextureIndex < NUM_BUFFERS);
+	
+PRAGMA_DISABLE_DEPRECATION_WARNINGS
 	return PooledRenderTarget[TextureIndex];
+PRAGMA_ENABLE_DEPRECATION_WARNINGS
 }
 
 const TRefCountPtr<IPooledRenderTarget>& FSceneViewState::FEyeAdaptationManager::GetOrCreateTexture(FRDGBuilder& GraphBuilder, uint32 TextureIndex)
 {
 	check(0 <= TextureIndex && TextureIndex < NUM_BUFFERS);
 
+PRAGMA_DISABLE_DEPRECATION_WARNINGS
 	// Create texture if needed.
 	if (!PooledRenderTarget[TextureIndex].IsValid())
 	{
@@ -969,6 +987,7 @@ const TRefCountPtr<IPooledRenderTarget>& FSceneViewState::FEyeAdaptationManager:
 	}
 
 	return PooledRenderTarget[TextureIndex];
+PRAGMA_ENABLE_DEPRECATION_WARNINGS
 }
 
 const TRefCountPtr<FRDGPooledBuffer>& FSceneViewState::FEyeAdaptationManager::GetBuffer(uint32 BufferIndex) const
diff --git a/Engine/Source/Runtime/Renderer/Private/PostProcess/PostProcessEyeAdaptation.h b/Engine/Source/Runtime/Renderer/Private/PostProcess/PostProcessEyeAdaptation.h
--- a/Engine/Source/Runtime/Renderer/Private/PostProcess/PostProcessEyeAdaptation.h
+++ b/Engine/Source/Runtime/Renderer/Private/PostProcess/PostProcessEyeAdaptation.h
@@ -127,3 +127,9 @@ FRDGBufferRef AddBasicEyeAdaptationPass(
 	const FEyeAdaptationParameters& EyeAdaptationParameters,
 	FScreenPassTexture SceneColor,
 	FRDGBufferRef EyeAdaptationBuffer);
+
+/**
+* Helper function to get current eye adaptation in a texture.
+* Should only be used by external plugins that require eye adaptation data in texture format.
+*/
+RENDERER_API FRDGTextureRef AddCopyEyeAdaptationDataToTexturePass(FRDGBuilder& GraphBuilder, const FViewInfo& View);
diff --git a/Engine/Source/Runtime/Renderer/Private/RendererScene.cpp b/Engine/Source/Runtime/Renderer/Private/RendererScene.cpp
--- a/Engine/Source/Runtime/Renderer/Private/RendererScene.cpp
+++ b/Engine/Source/Runtime/Renderer/Private/RendererScene.cpp
@@ -878,10 +878,12 @@ uint64 FPersistentSkyAtmosphereData::GetGPUSizeBytes(bool bLogSizes) const
 uint64 FSceneViewState::FEyeAdaptationManager::GetGPUSizeBytes(bool bLogSizes) const
 {
 	uint64 TotalSize = 0;
+PRAGMA_DISABLE_DEPRECATION_WARNINGS
 	for (int32 TargetIndex = 0; TargetIndex < UE_ARRAY_COUNT(PooledRenderTarget); TargetIndex++)
 	{
 		TotalSize += GetRenderTargetGPUSizeBytes(PooledRenderTarget[TargetIndex], bLogSizes);
 	}
+PRAGMA_ENABLE_DEPRECATION_WARNINGS
 	for (int32 BufferIndex = 0; BufferIndex < UE_ARRAY_COUNT(ExposureBufferData); BufferIndex++)
 	{
 		TotalSize += GetBufferGPUSizeBytes(ExposureBufferData[BufferIndex], bLogSizes);
diff --git a/Engine/Source/Runtime/Renderer/Private/ScenePrivate.h b/Engine/Source/Runtime/Renderer/Private/ScenePrivate.h
--- a/Engine/Source/Runtime/Renderer/Private/ScenePrivate.h
+++ b/Engine/Source/Runtime/Renderer/Private/ScenePrivate.h
@@ -799,11 +799,13 @@ private:
 		/** Get the last frame average scene luminance (used for exposure compensation curve) */
 		float GetLastAverageSceneLuminance() const { return LastAverageSceneLuminance; }
 
+		UE_DEPRECATED(5.2, "Use GetCurrentBuffer() instead.")
 		const TRefCountPtr<IPooledRenderTarget>& GetCurrentTexture() const
 		{
 			return GetTexture(CurrentBufferIndex);
 		}
 
+		UE_DEPRECATED(5.2, "Use GetCurrentBuffer() instead.")
 		const TRefCountPtr<IPooledRenderTarget>& GetCurrentTexture(FRDGBuilder& GraphBuilder)
 		{
 			return GetOrCreateTexture(GraphBuilder, CurrentBufferIndex);
@@ -846,6 +848,7 @@ private:
 		float LastExposure = 0;
 		float LastAverageSceneLuminance = 0; // 0 means invalid. Used for Exposure Compensation Curve.
 
+		UE_DEPRECATED(5.2, "Use ExposureBufferData instead.")
 		TRefCountPtr<IPooledRenderTarget> PooledRenderTarget[NUM_BUFFERS];
 
 		TRefCountPtr<FRDGPooledBuffer> ExposureBufferData[NUM_BUFFERS];
@@ -1170,20 +1173,23 @@ public:
 	 */
 	bool IsShadowOccluded(FSceneViewState::FProjectedShadowKey ShadowKey, int32 NumBufferedFrames) const;
 
-	/**
-	* Retrieve a single-pixel render targets with intra-frame state for use in eye adaptation post processing.
-	*/
+	UE_DEPRECATED(5.2, "Use GetCurrentEyeAdaptationBuffer() instead.")
 	IPooledRenderTarget* GetCurrentEyeAdaptationTexture() const override final
 	{
+PRAGMA_DISABLE_DEPRECATION_WARNINGS
 		IPooledRenderTarget* Texture = EyeAdaptationManager.GetCurrentTexture().GetReference();
 		check(bValidEyeAdaptationTexture && Texture);
+PRAGMA_ENABLE_DEPRECATION_WARNINGS
 		return Texture;
 	}
 
+	UE_DEPRECATED(5.2, "Use GetCurrentEyeAdaptationTexture() instead.")
 	IPooledRenderTarget* GetCurrentEyeAdaptationTexture(FRDGBuilder& GraphBuilder)
 	{
+PRAGMA_DISABLE_DEPRECATION_WARNINGS
 		bValidEyeAdaptationTexture = true;
 		return EyeAdaptationManager.GetCurrentTexture(GraphBuilder).GetReference();
+PRAGMA_ENABLE_DEPRECATION_WARNINGS
 	}
 
 	FRDGPooledBuffer* GetCurrentEyeAdaptationBuffer() const override final
@@ -1302,7 +1308,9 @@ public:
 		HZBOcclusionTests.ReleaseDynamicRHI();
 		EyeAdaptationManager.SafeRelease();
 		OcclusionFeedback.ReleaseResource();
+PRAGMA_DISABLE_DEPRECATION_WARNINGS
 		bValidEyeAdaptationTexture = false;
+PRAGMA_ENABLE_DEPRECATION_WARNINGS
 		bValidEyeAdaptationBuffer = false;
 	}
 
diff --git a/Engine/Source/Runtime/Renderer/Private/SceneRendering.cpp b/Engine/Source/Runtime/Renderer/Private/SceneRendering.cpp
--- a/Engine/Source/Runtime/Renderer/Private/SceneRendering.cpp
+++ b/Engine/Source/Runtime/Renderer/Private/SceneRendering.cpp
@@ -2227,7 +2227,9 @@ IPooledRenderTarget* FViewInfo::GetEyeAdaptationTexture(FRDGBuilder& GraphBuilde
 
 	if (FSceneViewState* EffectiveViewState = GetEyeAdaptationViewState())
 	{
+PRAGMA_DISABLE_DEPRECATION_WARNINGS
 		return EffectiveViewState->GetCurrentEyeAdaptationTexture(GraphBuilder);
+PRAGMA_ENABLE_DEPRECATION_WARNINGS
 	}
 	return nullptr;
 }
diff --git a/Engine/Source/Runtime/Renderer/Private/SceneTextureParameters.cpp b/Engine/Source/Runtime/Renderer/Private/SceneTextureParameters.cpp
--- a/Engine/Source/Runtime/Renderer/Private/SceneTextureParameters.cpp
+++ b/Engine/Source/Runtime/Renderer/Private/SceneTextureParameters.cpp
@@ -71,19 +71,6 @@ FSceneLightingChannelParameters GetSceneLightingChannelParameters(FRDGBuilder& G
 	return Parameters;
 }
 
-FRDGTextureRef GetEyeAdaptationTexture(FRDGBuilder& GraphBuilder, const FSceneView& View)
-{
-	if (View.HasValidEyeAdaptationTexture())
-	{
-		return GraphBuilder.RegisterExternalTexture(View.GetEyeAdaptationTexture(), ERDGTextureFlags::MultiFrame);
-	}
-	else
-	{
-		const FRDGSystemTextures& SystemTextures = FRDGSystemTextures::Get(GraphBuilder);
-		return SystemTextures.White;
-	}
-}
-
 FRDGBufferRef GetEyeAdaptationBuffer(FRDGBuilder& GraphBuilder, const FSceneView& View)
 {
 	if (View.HasValidEyeAdaptationBuffer())
diff --git a/Engine/Source/Runtime/Renderer/Private/SceneTextureParameters.h b/Engine/Source/Runtime/Renderer/Private/SceneTextureParameters.h
--- a/Engine/Source/Runtime/Renderer/Private/SceneTextureParameters.h
+++ b/Engine/Source/Runtime/Renderer/Private/SceneTextureParameters.h
@@ -38,8 +38,5 @@ END_SHADER_PARAMETER_STRUCT()
 /** Constructs lighting channel parameters from the scene context. */
 FSceneLightingChannelParameters GetSceneLightingChannelParameters(FRDGBuilder& GraphBuilder, FRDGTextureRef LightingChannelsTexture);
 
-/** Returns a render graph texture resource reference onto the eye adaptation or fallback. */
-RENDERER_API FRDGTextureRef GetEyeAdaptationTexture(FRDGBuilder& GraphBuilder, const FSceneView& View);
-
 /** Returns a render graph buffer resource reference onto the eye adaptation or fallback. */
 RENDERER_API FRDGBufferRef GetEyeAdaptationBuffer(FRDGBuilder& GraphBuilder, const FSceneView& View);
