Stereo mobile preview (UE-168931)

- Implemented for D3D11 and D3D12, as Mobile Multi View (MMV) fallback on both.

  High level summary / notable changes:
  - preview shader platforms will get the necessary properties from the desktop shader platforms used to display them
  - stereo render target manager will get informed which shader platform the engine is currently using
  - regularized somewhat what happens when the editor's preview platform changes (enabling/disabling and activating/deactivating). (Preview system is confusing because it's tri-state: preview platform can be set but not active, which is the same as it not being set). More code can be de-duplicated there.
  - Multiview will be considered disabled (for scene textures at least) if, despite ViewFamily requiring it, no View actually has it enabled.
  - Patched a bunch of shaders to account for texture array RTs used for MMV (and MMV fallback).

diff --git a/Engine/Config/Windows/DataDrivenPlatformInfo.ini b/Engine/Config/Windows/DataDrivenPlatformInfo.ini
--- a/Engine/Config/Windows/DataDrivenPlatformInfo.ini
+++ b/Engine/Config/Windows/DataDrivenPlatformInfo.ini
@@ -150,6 +150,7 @@ bSupportsDistanceFields = true
 bSupportsComputeFramework=true
 bSupportsByteBufferComputeShaders = true
 bSupportsInstancedStereo=true
+bSupportsVertexShaderLayer=true
 bSupportsManualVertexFetch=false
 bSupportsSceneDataCompressedTransforms=true
 bSupportsDebugViewShaders=true
diff --git a/Engine/Plugins/Runtime/OpenXR/Source/OpenXRHMD/Private/OpenXRHMD.cpp b/Engine/Plugins/Runtime/OpenXR/Source/OpenXRHMD/Private/OpenXRHMD.cpp
--- a/Engine/Plugins/Runtime/OpenXR/Source/OpenXRHMD/Private/OpenXRHMD.cpp
+++ b/Engine/Plugins/Runtime/OpenXR/Source/OpenXRHMD/Private/OpenXRHMD.cpp
@@ -34,6 +34,7 @@
 #include "HDRHelper.h"
 #include "Shader.h"
 #include "ScreenRendering.h"
+#include "StereoRenderUtils.h"
 #include "DefaultStereoLayers.h"
 
 #if WITH_EDITOR
@@ -1223,18 +1224,7 @@ FOpenXRHMD::FOpenXRHMD(const FAutoRegister& AutoRegister, XrInstance InInstance,
 
 	bSpaceAccellerationSupported = IsExtensionEnabled(XR_EPIC_SPACE_ACCELERATION_NAME);
 
-	static const auto CVarMobileMultiView = IConsoleManager::Get().FindTConsoleVariableDataInt(TEXT("vr.MobileMultiView"));
-	static const auto CVarMobileHDR = IConsoleManager::Get().FindTConsoleVariableDataInt(TEXT("r.MobileHDR"));
-	const bool bMobileHDR = (CVarMobileHDR && CVarMobileHDR->GetValueOnAnyThread() != 0);
-	const bool bMobileMultiView = !bMobileHDR && (CVarMobileMultiView && CVarMobileMultiView->GetValueOnAnyThread() != 0);
-#if PLATFORM_HOLOLENS
-	bIsMobileMultiViewEnabled = bMobileMultiView && GRHISupportsArrayIndexFromAnyShader;
-#else
-	bIsMobileMultiViewEnabled = bMobileMultiView && RHISupportsMobileMultiView(GMaxRHIShaderPlatform);
-#endif
-
-	static const auto CVarPropagateAlpha = IConsoleManager::Get().FindTConsoleVariableDataInt(TEXT("r.PostProcessing.PropagateAlpha"));
-	bProjectionLayerAlphaEnabled = !IsMobilePlatform(GMaxRHIShaderPlatform) && CVarPropagateAlpha->GetValueOnAnyThread() != 0;
+	ReconfigureForShaderPlatform(GMaxRHIShaderPlatform);
 
 	// Enumerate the viewport configurations
 	uint32 ConfigurationCount;
@@ -1339,6 +1329,17 @@ FOpenXRHMD::~FOpenXRHMD()
 	DestroySession();
 }
 
+bool FOpenXRHMD::ReconfigureForShaderPlatform(EShaderPlatform NewShaderPlatform)
+{
+	UE::StereoRenderUtils::FStereoShaderAspects Aspects(NewShaderPlatform);
+	bIsMobileMultiViewEnabled = Aspects.IsMobileMultiViewEnabled();
+
+	static const auto CVarPropagateAlpha = IConsoleManager::Get().FindTConsoleVariableDataInt(TEXT("r.PostProcessing.PropagateAlpha"));
+	bProjectionLayerAlphaEnabled = !IsMobilePlatform(NewShaderPlatform) && CVarPropagateAlpha->GetValueOnAnyThread() != 0;
+
+	return true;
+}
+
 const FOpenXRHMD::FPipelinedFrameState& FOpenXRHMD::GetPipelinedFrameStateForThread() const
 {
 	// Relying on implicit selection of the RHI struct is hazardous since the RHI thread isn't always present
@@ -2076,7 +2077,8 @@ bool FOpenXRHMD::AllocateRenderTargetTexture(uint32 Index, uint32 SizeX, uint32
 	if (Format == PF_A2B10G10R10 && !RenderBridge->Support10BitSwapchain())
 	{
 		UE_LOG(LogHMD, Warning, TEXT("Requesting 10 bit swapchain, but not supported: fall back to 8bpc"));
-		Format = PF_R8G8B8A8;
+		// Match the default logic in GetDefaultMobileSceneColorLowPrecisionFormat() in SceneTexturesConfig.cpp
+		Format = IsStandaloneStereoOnlyDevice() ? PF_R8G8B8A8 : PF_B8G8R8A8;
 	}
 
 	FClearValueBinding ClearColor = FClearValueBinding::Transparent;
@@ -2101,12 +2103,12 @@ bool FOpenXRHMD::AllocateRenderTargetTexture(uint32 Index, uint32 SizeX, uint32
 	LastActualColorSwapchainFormat = ActualFormat;
 
 	// TODO: Pass in known depth parameters (format + flags)? Do we know that at viewport setup time?
-	AllocateDepthTextureInternal(Index, SizeX, SizeY, NumSamples);
+	AllocateDepthTextureInternal(Index, SizeX, SizeY, NumSamples, OutTargetableTexture->GetDesc().ArraySize);
 
 	return true;
 }
 
-void FOpenXRHMD::AllocateDepthTextureInternal(uint32 Index, uint32 SizeX, uint32 SizeY, uint32 NumSamples)
+void FOpenXRHMD::AllocateDepthTextureInternal(uint32 Index, uint32 SizeX, uint32 SizeY, uint32 NumSamples, uint32 InArraySize)
 {
 	// TODO: Allocate depth texture by checking bNeedReAllocatedDepth
 
@@ -2123,7 +2125,7 @@ void FOpenXRHMD::AllocateDepthTextureInternal(uint32 Index, uint32 SizeX, uint32
 	FXRSwapChainPtr& DepthSwapchain = PipelinedLayerStateRendering.DepthSwapchain;
 	const FRHITexture2D* const DepthSwapchainTexture = DepthSwapchain == nullptr ? nullptr : DepthSwapchain->GetTexture2DArray() ? DepthSwapchain->GetTexture2DArray() : DepthSwapchain->GetTexture2D();
 	if (DepthSwapchain == nullptr || DepthSwapchainTexture == nullptr || 
-		DepthSwapchainTexture->GetSizeX() != SizeX || DepthSwapchainTexture->GetSizeY() != SizeY)
+		DepthSwapchainTexture->GetSizeX() != SizeX || DepthSwapchainTexture->GetSizeY() != SizeY || DepthSwapchainTexture->GetDesc().ArraySize != InArraySize)
 	{
 		// We're only creating a 1x target here, but we don't know whether it'll be the targeted texture
 		// or the resolve texture. Because of this, we unify the input flags.
@@ -2140,7 +2142,7 @@ void FOpenXRHMD::AllocateDepthTextureInternal(uint32 Index, uint32 SizeX, uint32
 		constexpr uint32 NumMipsExpected = 1;
 
 		uint8 UnusedActualFormat = 0;
-		DepthSwapchain = RenderBridge->CreateSwapchain(Session, PF_DepthStencil, UnusedActualFormat, SizeX, SizeY, bIsMobileMultiViewEnabled ? 2 : 1, NumMipsExpected, NumSamplesExpected, UnifiedCreateFlags, FClearValueBinding::DepthFar);
+		DepthSwapchain = RenderBridge->CreateSwapchain(Session, PF_DepthStencil, UnusedActualFormat, SizeX, SizeY, InArraySize, NumMipsExpected, NumSamplesExpected, UnifiedCreateFlags, FClearValueBinding::DepthFar);
 		if (!DepthSwapchain)
 		{
 			return;
@@ -3015,8 +3017,9 @@ FIntPoint FOpenXRHMD::GetIdealRenderTargetSize() const
 FIntRect FOpenXRHMD::GetFullFlatEyeRect_RenderThread(FTexture2DRHIRef EyeTexture) const
 {
 	FVector2D SrcNormRectMin(0.05f, 0.2f);
-	FVector2D SrcNormRectMax(0.45f, 0.8f);
-	if (GetDesiredNumberOfViews(bStereoEnabled) > 2)
+	// with MMV, each eye occupies the whole RT layer, so we don't need to limit the source rect to the left half of the RT.
+	FVector2D SrcNormRectMax(bIsMobileMultiViewEnabled ? 0.95f : 0.45f, 0.8f);
+	if (!bIsMobileMultiViewEnabled && GetDesiredNumberOfViews(bStereoEnabled) > 2)
 	{
 		SrcNormRectMin.X /= 2;
 		SrcNormRectMax.X /= 2;
@@ -3030,6 +3033,9 @@ class FDisplayMappingPS : public FGlobalShader
 	DECLARE_SHADER_TYPE(FDisplayMappingPS, Global);
 public:
 
+	class FArraySource : SHADER_PERMUTATION_BOOL("DISPLAY_MAPPING_PS_FROM_ARRAY");
+	using FPermutationDomain = TShaderPermutationDomain<FArraySource>;
+
 	static bool ShouldCompilePermutation(const FGlobalShaderPermutationParameters& Parameters)
 	{
 		return IsFeatureLevelSupported(Parameters.Platform, ERHIFeatureLevel::SM5);
@@ -3249,17 +3255,34 @@ void FOpenXRHMD::CopyTexture_RenderThread(FRHICommandListImmediate& RHICmdList,
 
 		bNeedsDisplayMapping &= (GMaxRHIFeatureLevel > ERHIFeatureLevel::ES3_1);
 
+		bool bIsArraySource = SrcTexture->GetDesc().IsTextureArray();
+
 		if (bNeedsDisplayMapping)
 		{
-			TShaderMapRef<FDisplayMappingPS> DisplayMappingPSRef(ShaderMap);
+			FDisplayMappingPS::FPermutationDomain PermutationVector;
+			PermutationVector.Set<FDisplayMappingPS::FArraySource>(bIsArraySource);
+
+			TShaderMapRef<FDisplayMappingPS> DisplayMappingPSRef(ShaderMap, PermutationVector);
+
+			ensureMsgf(!bIsArraySource, TEXT("Stub implementation for array source exists, but this path is not fully supported because post-processing not working with MMV"));
+
 			DisplayMappingPS = DisplayMappingPSRef.GetShader();
 			PixelShader = DisplayMappingPSRef;
 		}
 		else
 		{
-			TShaderMapRef<FScreenPS> ScreenPSRef(ShaderMap);
-			ScreenPS = ScreenPSRef.GetShader();
-			PixelShader = ScreenPSRef;
+			if (LIKELY(!bIsArraySource))
+			{
+				TShaderMapRef<FScreenPS> ScreenPSRef(ShaderMap);
+				ScreenPS = ScreenPSRef.GetShader();
+				PixelShader = ScreenPSRef;
+			}
+			else
+			{
+				TShaderMapRef<FScreenFromSlice0PS> ScreenPSRef(ShaderMap);
+				ScreenPS = ScreenPSRef.GetShader();
+				PixelShader = ScreenPSRef;
+			}
 		}
 
 		GraphicsPSOInit.BoundShaderState.VertexDeclarationRHI = GFilterVertexDeclaration.VertexDeclarationRHI;
diff --git a/Engine/Plugins/Runtime/OpenXR/Source/OpenXRHMD/Private/OpenXRHMD.h b/Engine/Plugins/Runtime/OpenXR/Source/OpenXRHMD/Private/OpenXRHMD.h
--- a/Engine/Plugins/Runtime/OpenXR/Source/OpenXRHMD/Private/OpenXRHMD.h
+++ b/Engine/Plugins/Runtime/OpenXR/Source/OpenXRHMD/Private/OpenXRHMD.h
@@ -240,7 +240,7 @@ protected:
 	
 	void CopyTexture_RenderThread(FRHICommandListImmediate& RHICmdList, FRHITexture2D* SrcTexture, FIntRect SrcRect, const FXRSwapChainPtr& DstSwapChain, FIntRect DstRect, bool bClearBlack, bool bNoAlpha) const;
 
-	void AllocateDepthTextureInternal(uint32 Index, uint32 SizeX, uint32 SizeY, uint32 NumSamples);
+	void AllocateDepthTextureInternal(uint32 Index, uint32 SizeX, uint32 SizeY, uint32 NumSamples, uint32 ArraySize);
 
 	// Used with FCoreDelegates
 	void VRHeadsetRecenterDelegate();
@@ -315,6 +315,7 @@ public:
 	virtual bool AllocateRenderTargetTexture(uint32 Index, uint32 SizeX, uint32 SizeY, uint8 Format, uint32 NumMips, ETextureCreateFlags Flags, ETextureCreateFlags TargetableTextureFlags, FTexture2DRHIRef& OutTargetableTexture, FTexture2DRHIRef& OutShaderResourceTexture, uint32 NumSamples = 1) override;
 	virtual bool NeedReAllocateDepthTexture(const TRefCountPtr<IPooledRenderTarget>& DepthTarget) override final { return bNeedReAllocatedDepth; }
 	virtual bool AllocateDepthTexture(uint32 Index, uint32 SizeX, uint32 SizeY, uint8 Format, uint32 NumMips, ETextureCreateFlags InTexFlags, ETextureCreateFlags TargetableTextureFlags, FTexture2DRHIRef& OutTargetableTexture, FTexture2DRHIRef& OutShaderResourceTexture, uint32 NumSamples = 1) override final;
+	virtual bool ReconfigureForShaderPlatform(EShaderPlatform NewShaderPlatform) override;
 	virtual EPixelFormat GetActualColorSwapchainFormat() const override { return static_cast<EPixelFormat>(LastActualColorSwapchainFormat); }
 
 	/** FXRRenderTargetManager */
diff --git a/Engine/Shaders/Private/CompositeUIPixelShader.usf b/Engine/Shaders/Private/CompositeUIPixelShader.usf
--- a/Engine/Shaders/Private/CompositeUIPixelShader.usf
+++ b/Engine/Shaders/Private/CompositeUIPixelShader.usf
@@ -84,10 +84,25 @@ void GenerateLUTPS(
 Texture2D<uint> UIWriteMaskTexture;
 #endif
 
+#ifndef DISPLAY_MAPPING_PS_FROM_ARRAY
+	#define DISPLAY_MAPPING_PS_FROM_ARRAY 0
+#endif // DISPLAY_MAPPING_PS_FROM_ARRAY
+
+#if DISPLAY_MAPPING_PS_FROM_ARRAY
+	// if source is an array, we only need the first slice
+	#define TextureOrArray2D						Texture2DArray
+	#define TexOrArr2DSample(Tex, S, UV)			Tex.Sample(S, float3(UV.x, UV.y, 0))
+	#define TexOrArr2DSampleLevel(Tex, S, UV, L)	Tex.SampleLevel(S, float3(UV.x, UV.y, 0), L)
+#else
+	#define TextureOrArray2D						Texture2D
+	#define TexOrArr2DSample(Tex, S, UV)			Texture2DSample(Tex, S, UV)
+	#define TexOrArr2DSampleLevel(Tex, S, UV, L)	Texture2DSampleLevel(Tex, S, UV, L)
+#endif
+
 Texture2D UITexture;
 SamplerState UISampler;
 
-Texture2D SceneTexture;
+TextureOrArray2D SceneTexture;
 SamplerState SceneSampler;
 
 Texture3D ColorSpaceLUT;
@@ -127,7 +142,7 @@ void HDRBackBufferConvertPS(
 	)
 {
 	OutColor = 0;
-	float3 SceneColor = Texture2DSample(SceneTexture, SceneSampler, Input.UV).xyz ;
+	float3 SceneColor = TexOrArr2DSample(SceneTexture, SceneSampler, Input.UV).xyz ;
 	SceneColor = ST2084ToLinear( SceneColor ) / ScRGBScaleFactor; 
 	OutColor.rgb = SceneColor;
 }
@@ -140,7 +155,7 @@ void HDREditorConvert(
 	)
 {
 	OutColor = 0;
-	float3 SceneColor = Texture2DSample(SceneTexture, SceneSampler, Input.UV).xyz ;
+	float3 SceneColor = TexOrArr2DSample(SceneTexture, SceneSampler, Input.UV).xyz ;
 
 	// Get color into Rec2020 space
 	SceneColor = sRGBToLinear( SceneColor) ;
@@ -252,7 +267,7 @@ void Main(
 	}
 #endif
 
-	float3 SceneColor = Texture2DSample(SceneTexture, SceneSampler, Input.UV).xyz;
+	float3 SceneColor = TexOrArr2DSample(SceneTexture, SceneSampler, Input.UV).xyz;
 	SceneColor = ComputeHDRSceneColor(SceneColor);
 	OutColor.xyz = ComposeUIAndScene(SceneColor.xyz, UIColor, UILevel);
 
@@ -336,7 +351,7 @@ void DisplayMappingPS(
 	out float4 OutColor : SV_Target0
 	)
 {
-	const float4 SceneColor = Texture2DSampleLevel(SceneTexture, SceneSampler, Input.UV, 0) ;
+	const float4 SceneColor = TexOrArr2DSampleLevel(SceneTexture, SceneSampler, Input.UV, 0);
 	OutColor = SceneColor;
 
 	float3 LinearSceneColor = pow(SceneColor.rgb, 2.2);
@@ -374,7 +389,7 @@ void BlitUIToHDRPS(
 	}
 #endif
 
-	float3 SceneColor = Texture2DSample(SceneTexture, UISampler, Input.UV).xyz;
+	float3 SceneColor = TexOrArr2DSample(SceneTexture, UISampler, Input.UV).xyz;
 	SceneColor = ComputeHDRSceneColor(SceneColor);
 	OutColor.xyz = ComposeUIAndScene(SceneColor.xyz, UIColor, UILevel) / ScRGBScaleFactor;
 	OutColor.a = 1;
diff --git a/Engine/Shaders/Private/HdrCustomResolveShaders.usf b/Engine/Shaders/Private/HdrCustomResolveShaders.usf
--- a/Engine/Shaders/Private/HdrCustomResolveShaders.usf
+++ b/Engine/Shaders/Private/HdrCustomResolveShaders.usf
@@ -6,16 +6,41 @@
 
 #include "Common.ush"
 
+#ifndef HDR_CUSTOM_RESOLVE_TEXTUREARRAY
+	#define HDR_CUSTOM_RESOLVE_TEXTUREARRAY 0
+#endif 
+
+#if HDR_CUSTOM_RESOLVE_TEXTUREARRAY
+	#define		MSTextureType(Type, NumSamples)		Texture2DMSArray<Type, NumSamples>
+	#define		MSCoords							uint3
+	#define		MSCoordsFromFloat4(Pos)				uint3(Pos.x, Pos.y, ViewID)
+#else
+	#define		MSTextureType(Type, NumSamples)		Texture2DMS<Type, NumSamples>
+	#define		MSCoords							uint2
+	#define		MSCoordsFromFloat4(Pos)				uint2(Pos.xy);
+#endif
+
 #ifdef OVERRIDE_HDRCUSTOMRESOLVESHADERS_USH
 # include "/Platform/Private/HdrCustomResolveShaders.ush"
 #else
 
-float4 HdrCustomResolveVS(uint Id : SV_VertexID) : SV_POSITION
+void HdrCustomResolveVS(in uint InstanceID : SV_InstanceID, in uint Id : SV_VertexID,
+	out float4 OutPosition				: SV_POSITION
+#if HDR_CUSTOM_RESOLVE_TEXTUREARRAY
+	,out nointerpolation uint OutViewId : VIEW_ID,
+	out uint OutRTIndex					: SV_RenderTargetArrayIndex
+#endif
+	)
 {
 	int x = Id &  1;
 	int y = Id >> 1;
 
-	return float4(x * 4 - 1, y * 4 - 1, 0, 1);
+	OutPosition = float4(x * 4 - 1, y * 4 - 1, 0, 1);
+
+#if HDR_CUSTOM_RESOLVE_TEXTUREARRAY
+	OutViewId = (InstanceID & 1);
+	OutRTIndex = OutViewId;
+#endif
 }
 
 float4 Encode(float4 Color)
@@ -35,16 +60,20 @@ Texture2D<int2> FMaskTex;
 int2 fmask;
 
 # if HDR_RESOLVE_NUM_SAMPLES
-void LoadTexSample(Texture2DMS<float4, HDR_RESOLVE_NUM_SAMPLES> TargetTex, int2 P, int sampleIndex, inout float4 sampleSum, inout float weightSum)
+void LoadTexSample(MSTextureType(float4, HDR_RESOLVE_NUM_SAMPLES) TargetTex, int2 P, int sampleIndex, inout float4 sampleSum, inout float weightSum)
 {
 	sampleSum += Encode(TargetTex.Load(P, sampleIndex));
 	weightSum += 1;
 }
 
-Texture2DMS<float4, HDR_RESOLVE_NUM_SAMPLES> Tex;
-float4 HdrCustomResolveFMaskPS(float4 Pos : SV_POSITION) : SV_Target0
+MSTextureType(float4, HDR_RESOLVE_NUM_SAMPLES) Tex;
+float4 HdrCustomResolveFMaskPS(float4 Pos : SV_POSITION
+							#if HDR_CUSTOM_RESOLVE_TEXTUREARRAY
+								,nointerpolation uint ViewID  : VIEW_ID
+							#endif
+	) : SV_Target0
 {
-	uint2 P = uint2(Pos.xy);
+	MSCoords P = MSCoordsFromFloat4(Pos);
 
 	float4 sampleSum = float4(0, 0, 0, 1);
 	float weightSum = 0;
@@ -60,10 +89,14 @@ float4 HdrCustomResolveFMaskPS(float4 Pos : SV_POSITION) : SV_Target0
 
 # if HDR_CUSTOM_RESOLVE_2X
 
-	Texture2DMS<float4,2> Tex;
-	float4 HdrCustomResolve2xPS(float4 Pos : SV_POSITION) : SV_Target0
+	MSTextureType(float4,2) Tex;
+	float4 HdrCustomResolve2xPS(float4 Pos : SV_POSITION
+							#if HDR_CUSTOM_RESOLVE_TEXTUREARRAY
+								,nointerpolation uint ViewID  : VIEW_ID
+							#endif
+		) : SV_Target0
 	{
-		uint2 P = uint2(Pos.xy);
+		MSCoords P = MSCoordsFromFloat4(Pos);
 		float4 C0 = Encode(Tex.Load(P, 0));
 		float4 C1 = Encode(Tex.Load(P, 1));
 		return Decode(C0*0.5 + C1*0.5);
@@ -73,10 +106,14 @@ float4 HdrCustomResolveFMaskPS(float4 Pos : SV_POSITION) : SV_Target0
 
 # if HDR_CUSTOM_RESOLVE_4X
 
-	Texture2DMS<float4,4> Tex;
-	float4 HdrCustomResolve4xPS(float4 Pos : SV_POSITION) : SV_Target0
+	MSTextureType(float4,4) Tex;
+	float4 HdrCustomResolve4xPS(float4 Pos : SV_POSITION
+							#if HDR_CUSTOM_RESOLVE_TEXTUREARRAY
+								,nointerpolation uint ViewID  : VIEW_ID
+							#endif
+		) : SV_Target0
 	{
-		uint2 P = uint2(Pos.xy);
+		MSCoords P = MSCoordsFromFloat4(Pos);
 		float4 C0 = Encode(Tex.Load(P, 0));
 		float4 C1 = Encode(Tex.Load(P, 1));
 		float4 C2 = Encode(Tex.Load(P, 2));
@@ -88,10 +125,14 @@ float4 HdrCustomResolveFMaskPS(float4 Pos : SV_POSITION) : SV_Target0
 
 # if HDR_CUSTOM_RESOLVE_8X
 
-	Texture2DMS<float4,8> Tex;
-	float4 HdrCustomResolve8xPS(float4 Pos : SV_POSITION) : SV_Target0
+	MSTextureType(float4,8) Tex;
+	float4 HdrCustomResolve8xPS(float4 Pos : SV_POSITION
+							#if HDR_CUSTOM_RESOLVE_TEXTUREARRAY
+								,nointerpolation uint ViewID  : VIEW_ID
+							#endif
+		) : SV_Target0
 	{
-		uint2 P = uint2(Pos.xy);
+		MSCoords P = MSCoordsFromFloat4(Pos);
 		float4 C0 = Encode(Tex.Load(P, 0));
 		float4 C1 = Encode(Tex.Load(P, 1));
 		float4 C2 = Encode(Tex.Load(P, 2));
diff --git a/Engine/Shaders/Private/ResolvePixelShader.usf b/Engine/Shaders/Private/ResolvePixelShader.usf
--- a/Engine/Shaders/Private/ResolvePixelShader.usf
+++ b/Engine/Shaders/Private/ResolvePixelShader.usf
@@ -5,23 +5,45 @@
 =============================================================================*/
 
 #include "Common.ush"
+
+#ifndef DEPTH_RESOLVE_TEXTUREARRAY
+	#define DEPTH_RESOLVE_TEXTUREARRAY 0
+#endif 
+
+#if DEPTH_RESOLVE_TEXTUREARRAY
+Texture2DMSArray<float4> UnresolvedSurface;
+#else
 Texture2DMS<float4> UnresolvedSurface;
+#endif
 
 
 void MainDepth(
 	float2 InUV : TEXCOORD0,
+#if DEPTH_RESOLVE_TEXTUREARRAY
+	nointerpolation uint ViewId : VIEW_ID,
+#endif
 	out float OutDepth : SV_DEPTH
 	)
 {
 #ifndef DEPTH_RESOLVE_NUM_SAMPLES
-	float3 SurfaceDimensions;
-	UnresolvedSurface.GetDimensions(SurfaceDimensions.x, SurfaceDimensions.y, SurfaceDimensions.z);
-	int NumSurfaceSamples = (int)SurfaceDimensions.z;
+	#if DEPTH_RESOLVE_TEXTUREARRAY
+		uint4 SurfaceDimensions;
+		UnresolvedSurface.GetDimensions(SurfaceDimensions.x, SurfaceDimensions.y, SurfaceDimensions.z, SurfaceDimensions.w);
+		int NumSurfaceSamples = (int)SurfaceDimensions.w;
+	#else
+		uint3 SurfaceDimensions;
+		UnresolvedSurface.GetDimensions(SurfaceDimensions.x, SurfaceDimensions.y, SurfaceDimensions.z);
+		int NumSurfaceSamples = (int)SurfaceDimensions.z;
+	#endif
 #else
 	int NumSurfaceSamples = DEPTH_RESOLVE_NUM_SAMPLES;
 #endif
 
+#if DEPTH_RESOLVE_TEXTUREARRAY
+	int3 IntUV = int3(trunc(InUV.x), trunc(InUV.y), ViewId);
+#else
 	int2 IntUV = trunc(InUV);
+#endif
 	float ResolvedDepth = UnresolvedSurface.Load(IntUV, 0).r;
 
 	for (int SampleIndex = 1; SampleIndex < NumSurfaceSamples; ++SampleIndex)
@@ -35,18 +57,17 @@ void MainDepth(
 	OutDepth = ResolvedDepth;
 }
 
+#if !DEPTH_RESOLVE_TEXTUREARRAY
 uint SingleSampleIndex;
 
 void MainSingleSample(
 	float2 InUV : TEXCOORD0,
+	nointerpolation uint ViewId : VIEW_ID,
 	out float4 OutColor : SV_Target0
 	)
 {
-	float3 SurfaceDimensions;
-	UnresolvedSurface.GetDimensions(SurfaceDimensions.x,SurfaceDimensions.y,SurfaceDimensions.z);
-	int NumSurfaceSamples = (int)SurfaceDimensions.z;
-
 	int2 IntUV = trunc(InUV);
 
 	OutColor = UnresolvedSurface.Load(IntUV,SingleSampleIndex);
-}
\ No newline at end of file
+}
+#endif // !DEPTH_RESOLVE_TEXTUREARRAY
\ No newline at end of file
diff --git a/Engine/Shaders/Private/ResolveVertexShader.usf b/Engine/Shaders/Private/ResolveVertexShader.usf
--- a/Engine/Shaders/Private/ResolveVertexShader.usf
+++ b/Engine/Shaders/Private/ResolveVertexShader.usf
@@ -6,17 +6,30 @@
 
 #include "Common.ush"
 
+#ifndef DEPTH_RESOLVE_TEXTUREARRAY
+	#define DEPTH_RESOLVE_TEXTUREARRAY 0
+#endif 
+
 float4 PositionMinMax; // xy = min [x, y], zw = max [x, y]
 float4 UVMinMax;       // xy = min [u, v], zw = max [u, v]
 
 void Main(
-	uint VertexID          : SV_VertexID,
-	out float2 OutUV       : TEXCOORD0,
-	out float4 OutPosition : SV_POSITION
+	in uint InstanceID					: SV_InstanceID,
+	uint VertexID						: SV_VertexID,
+	out float2 OutUV					: TEXCOORD0,
+#if DEPTH_RESOLVE_TEXTUREARRAY
+	out nointerpolation uint OutViewId  : VIEW_ID,
+	out uint OutRTIndex					: SV_RenderTargetArrayIndex,
+#endif
+	out float4 OutPosition				: SV_POSITION
 	)
 {
 	// Generate these positions:
+	// ViewId: 0
+	// [0, 1], [1, 1], [0, 0], [1, 0]
+	// ViewId: 1
 	// [0, 1], [1, 1], [0, 0], [1, 0]
+
 	float2 Pos = float2(
 		float(VertexID & 1),
 		float(1 - (VertexID >> 1))
@@ -24,4 +37,9 @@ void Main(
 
 	OutPosition = float4(lerp(PositionMinMax.xy, PositionMinMax.zw, Pos), 0, 1);
     OutUV = lerp(UVMinMax.xy, UVMinMax.zw, Pos);
+
+#if DEPTH_RESOLVE_TEXTUREARRAY
+	OutViewId = (InstanceID & 1);
+	OutRTIndex = OutViewId;
+#endif
 }
diff --git a/Engine/Shaders/Private/ScreenPixelShader.usf b/Engine/Shaders/Private/ScreenPixelShader.usf
--- a/Engine/Shaders/Private/ScreenPixelShader.usf
+++ b/Engine/Shaders/Private/ScreenPixelShader.usf
@@ -6,7 +6,23 @@
 
 #include "Common.ush"
 
-Texture2D InTexture;
+#ifndef SCREEN_PS_FROM_SLICE0
+	#define SCREEN_PS_FROM_SLICE0   0
+#endif
+
+#if SCREEN_PS_FROM_SLICE0
+	Texture2DArray InTexture;
+
+	// if source is an array, we only need the first slice
+	#define TexOrArr2DSample(Tex, S, UV)			Tex.Sample(S, float3(UV.x, UV.y, 0))
+	#define TexOrArr2DSampleLevel(Tex, S, UV, L)	Tex.SampleLevel(S, float3(UV.x, UV.y, 0), L)
+#else
+	Texture2D InTexture;
+
+	#define TexOrArr2DSample(Tex, S, UV)			Texture2DSample(Tex, S, UV)
+	#define TexOrArr2DSampleLevel(Tex, S, UV, L)	Texture2DSampleLevel(Tex, S, UV, L)
+#endif
+
 SamplerState InTextureSampler;
 int MipLevel;
 
@@ -15,7 +31,7 @@ void Main(
 	out float4 OutColor : SV_Target0
 	)
 {
-	OutColor = Texture2DSample(InTexture, InTextureSampler, Input.UV);
+	OutColor = TexOrArr2DSample(InTexture, InTextureSampler, Input.UV);
 }
 
 void MainInvertAlpha(
@@ -23,7 +39,7 @@ void MainInvertAlpha(
 	out float4 OutColor : SV_Target0
 	)
 {
-	OutColor = Texture2DSample(InTexture, InTextureSampler, Input.UV);
+	OutColor = TexOrArr2DSample(InTexture, InTextureSampler, Input.UV);
 	OutColor.a = 1.f - OutColor.a;
 }
 
@@ -32,7 +48,7 @@ void MainsRGBSource(
 	out float4 OutColor : SV_Target0
 	)
 {
-	OutColor = Texture2DSample(InTexture, InTextureSampler, Input.UV);
+	OutColor = TexOrArr2DSample(InTexture, InTextureSampler, Input.UV);
 	OutColor.rgb = pow( OutColor.rgb, 1.0f / 2.2f );
 }
 
@@ -41,7 +57,7 @@ void MainMipLevel(
 	out float4 OutColor : SV_Target0
 	)
 {
-	OutColor = Texture2DSampleLevel(InTexture, InTextureSampler, Input.UV, MipLevel);
+	OutColor = TexOrArr2DSampleLevel(InTexture, InTextureSampler, Input.UV, MipLevel);
 }
 
 void MainsRGBSourceMipLevel(
@@ -49,6 +65,6 @@ void MainsRGBSourceMipLevel(
 	out float4 OutColor : SV_Target0
 	)
 {
-	OutColor = Texture2DSampleLevel(InTexture, InTextureSampler, Input.UV, MipLevel);
+	OutColor = TexOrArr2DSampleLevel(InTexture, InTextureSampler, Input.UV, MipLevel);
 	OutColor.rgb = pow( OutColor.rgb, 1.0f / 2.2f );
 }
diff --git a/Engine/Source/Editor/UnrealEd/Classes/Editor/EditorEngine.h b/Engine/Source/Editor/UnrealEd/Classes/Editor/EditorEngine.h
--- a/Engine/Source/Editor/UnrealEd/Classes/Editor/EditorEngine.h
+++ b/Engine/Source/Editor/UnrealEd/Classes/Editor/EditorEngine.h
@@ -3181,6 +3181,9 @@ private:
 	/** Gets the init values for worlds opened via Map_Load in the editor */
 	UWorld::InitializationValues GetEditorWorldInitializationValues() const;
 
+	/** Logic that needs to run on change of effective shader platform (can happen via both SetPreviewPlatform and ToggleFeatureLevelPreview). */
+	void OnEffectivePreviewShaderPlatformChange();
+
 public:
 	// Launcher Worker
 	TSharedPtr<class ILauncherWorker> LauncherWorker;
diff --git a/Engine/Source/Editor/UnrealEd/Private/EditorEngine.cpp b/Engine/Source/Editor/UnrealEd/Private/EditorEngine.cpp
--- a/Engine/Source/Editor/UnrealEd/Private/EditorEngine.cpp
+++ b/Engine/Source/Editor/UnrealEd/Private/EditorEngine.cpp
@@ -257,6 +257,7 @@
 #include "LevelEditorDragDropHandler.h"
 #include "IProjectExternalContentInterface.h"
 #include "IDocumentation.h"
+#include "StereoRenderTargetManager.h"
 
 DEFINE_LOG_CATEGORY_STATIC(LogEditor, Log, All);
 
@@ -7751,6 +7752,18 @@ void UEditorEngine::OnSceneMaterialsModified()
 {
 }
 
+void UEditorEngine::OnEffectivePreviewShaderPlatformChange()
+{
+	if (XRSystem.IsValid() && StereoRenderingDevice.IsValid())
+	{
+		IStereoRenderTargetManager* StereoRenderTargetManager = StereoRenderingDevice->GetRenderTargetManager();
+		if (StereoRenderTargetManager)
+		{
+			StereoRenderTargetManager->ReconfigureForShaderPlatform(PreviewPlatform.ShaderPlatform);
+		}
+	}
+}
+
 void UEditorEngine::SetPreviewPlatform(const FPreviewPlatformInfo& NewPreviewPlatform, bool bSaveSettings)
 {
 	// Get the requested preview platform, make sure it is valid.
@@ -7760,6 +7773,7 @@ void UEditorEngine::SetPreviewPlatform(const FPreviewPlatformInfo& NewPreviewPla
 	check(NewPreviewPlatform.PreviewShaderFormatName.IsNone() || MaxFeatureLevel == NewPreviewPlatform.PreviewFeatureLevel);
 
 	const bool bChangedPreviewShaderPlatform = NewPreviewPlatform.ShaderPlatform != PreviewPlatform.ShaderPlatform;
+	const bool bChangedEffectiveShaderPlatform = bChangedPreviewShaderPlatform && (PreviewPlatform.bPreviewFeatureLevelActive || NewPreviewPlatform.bPreviewFeatureLevelActive);
 	const ERHIFeatureLevel::Type EffectiveFeatureLevel = NewPreviewPlatform.GetEffectivePreviewFeatureLevel();
 
 	if (NewPreviewPlatform.PreviewShaderFormatName != NAME_None)
@@ -7778,6 +7792,11 @@ void UEditorEngine::SetPreviewPlatform(const FPreviewPlatformInfo& NewPreviewPla
 		MaterialShaderQualitySettings->SetPreviewPlatform(PreviewPlatform.PreviewShaderFormatName);
 
 		UStaticMesh::OnLodStrippingQualityLevelChanged(nullptr);
+
+		if (bChangedEffectiveShaderPlatform)
+		{
+			OnEffectivePreviewShaderPlatformChange();
+		}
 	}
 
 	constexpr bool bUpdateProgressDialog = true;
@@ -7889,7 +7908,7 @@ void UEditorEngine::ToggleFeatureLevelPreview()
 	}
 
 	Scalability::ApplyCachedQualityLevelForShaderPlatform(GetActiveShaderPlatform());
-
+	OnEffectivePreviewShaderPlatformChange();
 	PreviewPlatformChanged.Broadcast();
 
 	UStaticMesh::OnLodStrippingQualityLevelChanged(nullptr);
diff --git a/Engine/Source/Runtime/Engine/Private/ScreenRendering.cpp b/Engine/Source/Runtime/Engine/Private/ScreenRendering.cpp
--- a/Engine/Source/Runtime/Engine/Private/ScreenRendering.cpp
+++ b/Engine/Source/Runtime/Engine/Private/ScreenRendering.cpp
@@ -11,6 +11,7 @@ TGlobalResource<FScreenVertexDeclaration> GScreenVertexDeclaration;
 
 // Shader implementations.
 IMPLEMENT_SHADER_TYPE(, FScreenPS, TEXT("/Engine/Private/ScreenPixelShader.usf"), TEXT("Main"), SF_Pixel);
+IMPLEMENT_SHADER_TYPE(, FScreenFromSlice0PS, TEXT("/Engine/Private/ScreenPixelShader.usf"), TEXT("Main"), SF_Pixel);
 IMPLEMENT_SHADER_TYPE(,FScreenPSInvertAlpha,TEXT("/Engine/Private/ScreenPixelShader.usf"),TEXT("MainInvertAlpha"),SF_Pixel);
 IMPLEMENT_SHADER_TYPE(,FScreenPSsRGBSource, TEXT("/Engine/Private/ScreenPixelShader.usf"), TEXT("MainsRGBSource"), SF_Pixel);
 IMPLEMENT_SHADER_TYPE(,FScreenPSMipLevel, TEXT("/Engine/Private/ScreenPixelShader.usf"), TEXT("MainMipLevel"), SF_Pixel);
diff --git a/Engine/Source/Runtime/Engine/Public/ScreenRendering.h b/Engine/Source/Runtime/Engine/Public/ScreenRendering.h
--- a/Engine/Source/Runtime/Engine/Public/ScreenRendering.h
+++ b/Engine/Source/Runtime/Engine/Public/ScreenRendering.h
@@ -90,6 +90,39 @@ private:
 	LAYOUT_FIELD(FShaderResourceParameter, InTextureSampler);
 };
 
+/**
+ * A pixel shader for rendering a textured screen element, taking only the first slice of the array
+ */
+class FScreenFromSlice0PS : public FScreenPS
+{
+	DECLARE_EXPORTED_SHADER_TYPE(FScreenFromSlice0PS, Global, ENGINE_API);
+public:
+
+	static bool ShouldCompilePermutation(const FGlobalShaderPermutationParameters& Parameters) { return true; }
+
+	static void ModifyCompilationEnvironment(const FGlobalShaderPermutationParameters& Parameters, FShaderCompilerEnvironment& OutEnvironment)
+	{
+		FScreenPS::ModifyCompilationEnvironment(Parameters, OutEnvironment);
+		OutEnvironment.SetDefine(TEXT("SCREEN_PS_FROM_SLICE0"), 1);
+	}
+
+	FScreenFromSlice0PS(const ShaderMetaType::CompiledShaderInitializerType& Initializer) :
+		FScreenPS(Initializer)
+	{
+	}
+	FScreenFromSlice0PS() {}
+
+	void SetParameters(FRHICommandList& RHICmdList, const FTexture* Texture)
+	{
+		FScreenPS::SetParameters(RHICmdList, Texture);
+	}
+
+	void SetParameters(FRHICommandList& RHICmdList, FRHISamplerState* SamplerStateRHI, FRHITexture* TextureRHI)
+	{
+		FScreenPS::SetParameters(RHICmdList, SamplerStateRHI, TextureRHI);
+	}
+};
+
 
 /**
  * A pixel shader for rendering a textured screen element.
diff --git a/Engine/Source/Runtime/Engine/Public/StereoRenderTargetManager.h b/Engine/Source/Runtime/Engine/Public/StereoRenderTargetManager.h
--- a/Engine/Source/Runtime/Engine/Public/StereoRenderTargetManager.h
+++ b/Engine/Source/Runtime/Engine/Public/StereoRenderTargetManager.h
@@ -11,6 +11,7 @@ StereoRenderTargetManager.h: Abstract interface returned from IStereoRendering t
 
 class FRHITexture;
 enum class ETextureCreateFlags : uint64;
+enum EShaderPlatform : uint16;
 typedef TRefCountPtr<FRHITexture> FTexture2DRHIRef;
 
 /** 
@@ -105,5 +106,12 @@ public:
 	 */
 	virtual bool HDRGetMetaDataForStereo(EDisplayOutputFormat& OutDisplayOutputFormat, EDisplayColorGamut& OutDisplayColorGamut, bool& OutbHDRSupported) { return false; }
 
+	/**
+	 * In the editor, we may switch between preview shader platforms, which support various single-pass rendering methods
+	 * (e.g. InstancedStereo or MobileMultiView). Sometimes RT managers have their own state that depends on the method, so this
+	 * API provides them with the possibility to reconfigure it
+	 */
+	virtual bool ReconfigureForShaderPlatform(EShaderPlatform NewShaderPlatform) { return false; };
+
 	static EPixelFormat GetStereoLayerPixelFormat() { return PF_B8G8R8A8; }
 };
\ No newline at end of file
diff --git a/Engine/Source/Runtime/RHI/Private/DataDrivenShaderPlatformInfo.cpp b/Engine/Source/Runtime/RHI/Private/DataDrivenShaderPlatformInfo.cpp
--- a/Engine/Source/Runtime/RHI/Private/DataDrivenShaderPlatformInfo.cpp
+++ b/Engine/Source/Runtime/RHI/Private/DataDrivenShaderPlatformInfo.cpp
@@ -362,8 +362,19 @@ void FGenericDataDrivenShaderPlatformInfo::UpdatePreviewPlatforms()
 				Infos[ShaderPlatform].bSupportsLumenGI &= Infos[EditorSPForPreviewMaxFeatureLevel].bSupportsLumenGI;
 				Infos[ShaderPlatform].bSupportsUInt64ImageAtomics &= Infos[EditorSPForPreviewMaxFeatureLevel].bSupportsUInt64ImageAtomics;
 				Infos[ShaderPlatform].bSupportsGen5TemporalAA &= Infos[EditorSPForPreviewMaxFeatureLevel].bSupportsGen5TemporalAA;
+				// Support for stereo features requires extra consideration. The editor may not use the same technique as the preview platform,
+				// particularly MobileMultiView may be substituted by a fallback path. In order to avoid inundating real mobile platforms
+				// with the properties needed for the desktop MMV fallback path, override them here with the editor ones to make MMV preview possible
+				if (Infos[ShaderPlatform].bSupportsMobileMultiView && !Infos[EditorSPForPreviewMaxFeatureLevel].bSupportsMobileMultiView)
+				{
+					Infos[ShaderPlatform].bSupportsInstancedStereo = Infos[EditorSPForPreviewMaxFeatureLevel].bSupportsInstancedStereo;
+					Infos[ShaderPlatform].bSupportsVertexShaderLayer = Infos[EditorSPForPreviewMaxFeatureLevel].bSupportsVertexShaderLayer;
+				}
+				else
+				{
+					Infos[ShaderPlatform].bSupportsInstancedStereo &= Infos[EditorSPForPreviewMaxFeatureLevel].bSupportsInstancedStereo;
+				}
 				Infos[ShaderPlatform].bSupportsMobileMultiView &= Infos[EditorSPForPreviewMaxFeatureLevel].bSupportsMobileMultiView;
-				Infos[ShaderPlatform].bSupportsInstancedStereo &= Infos[EditorSPForPreviewMaxFeatureLevel].bSupportsInstancedStereo;
 				Infos[ShaderPlatform].bSupportsManualVertexFetch &= Infos[EditorSPForPreviewMaxFeatureLevel].bSupportsManualVertexFetch;
 				Infos[ShaderPlatform].bSupportsRenderTargetWriteMask = false;
 				Infos[ShaderPlatform].bSupportsIntrinsicWaveOnce = false;
diff --git a/Engine/Source/Runtime/RenderCore/Private/ResolveShader.cpp b/Engine/Source/Runtime/RenderCore/Private/ResolveShader.cpp
--- a/Engine/Source/Runtime/RenderCore/Private/ResolveShader.cpp
+++ b/Engine/Source/Runtime/RenderCore/Private/ResolveShader.cpp
@@ -8,8 +8,13 @@ IMPLEMENT_SHADER_TYPE(, FResolveDepthPS, TEXT("/Engine/Private/ResolvePixelShade
 IMPLEMENT_SHADER_TYPE(, FResolveDepth2XPS, TEXT("/Engine/Private/ResolvePixelShader.usf"), TEXT("MainDepth"), SF_Pixel);
 IMPLEMENT_SHADER_TYPE(, FResolveDepth4XPS, TEXT("/Engine/Private/ResolvePixelShader.usf"), TEXT("MainDepth"), SF_Pixel);
 IMPLEMENT_SHADER_TYPE(, FResolveDepth8XPS, TEXT("/Engine/Private/ResolvePixelShader.usf"), TEXT("MainDepth"), SF_Pixel);
+IMPLEMENT_SHADER_TYPE(, FResolveDepthArrayPS, TEXT("/Engine/Private/ResolvePixelShader.usf"), TEXT("MainDepth"), SF_Pixel);
+IMPLEMENT_SHADER_TYPE(, FResolveDepthArray2XPS, TEXT("/Engine/Private/ResolvePixelShader.usf"), TEXT("MainDepth"), SF_Pixel);
+IMPLEMENT_SHADER_TYPE(, FResolveDepthArray4XPS, TEXT("/Engine/Private/ResolvePixelShader.usf"), TEXT("MainDepth"), SF_Pixel);
+IMPLEMENT_SHADER_TYPE(, FResolveDepthArray8XPS, TEXT("/Engine/Private/ResolvePixelShader.usf"), TEXT("MainDepth"), SF_Pixel);
 IMPLEMENT_SHADER_TYPE(, FResolveSingleSamplePS, TEXT("/Engine/Private/ResolvePixelShader.usf"), TEXT("MainSingleSample"), SF_Pixel);
 IMPLEMENT_SHADER_TYPE(, FResolveVS, TEXT("/Engine/Private/ResolveVertexShader.usf"), TEXT("Main"), SF_Vertex);
+IMPLEMENT_SHADER_TYPE(, FResolveArrayVS, TEXT("/Engine/Private/ResolveVertexShader.usf"), TEXT("Main"), SF_Vertex);
 
 void FResolveSingleSamplePS::SetParameters(FRHICommandList& RHICmdList, uint32 SingleSampleIndexValue)
 {
diff --git a/Engine/Source/Runtime/RenderCore/Public/ResolveShader.h b/Engine/Source/Runtime/RenderCore/Public/ResolveShader.h
--- a/Engine/Source/Runtime/RenderCore/Public/ResolveShader.h
+++ b/Engine/Source/Runtime/RenderCore/Public/ResolveShader.h
@@ -32,7 +32,7 @@ public:
 	}
 		
 	FResolveDepthPS(const ShaderMetaType::CompiledShaderInitializerType& Initializer):
-	FGlobalShader(Initializer)
+		FGlobalShader(Initializer)
 	{
 		UnresolvedSurface.Bind(Initializer.ParameterMap,TEXT("UnresolvedSurface"), SPF_Mandatory);
 	}
@@ -45,7 +45,7 @@ public:
 	LAYOUT_FIELD(FShaderResourceParameter, UnresolvedSurface);
 };
 
-class FResolveDepth2XPS : public FGlobalShader
+class FResolveDepth2XPS : public FResolveDepthPS
 {
 	DECLARE_EXPORTED_SHADER_TYPE(FResolveDepth2XPS, Global, RENDERCORE_API);
 public:
@@ -54,26 +54,18 @@ public:
 
 	static void ModifyCompilationEnvironment(const FGlobalShaderPermutationParameters& Parameters, FShaderCompilerEnvironment& OutEnvironment)
 	{
-		FGlobalShader::ModifyCompilationEnvironment(Parameters, OutEnvironment);
+		FResolveDepthPS::ModifyCompilationEnvironment(Parameters, OutEnvironment);
 		OutEnvironment.SetDefine(TEXT("DEPTH_RESOLVE_NUM_SAMPLES"), 2);
 	}
 
 	FResolveDepth2XPS(const ShaderMetaType::CompiledShaderInitializerType& Initializer) :
-		FGlobalShader(Initializer)
+		FResolveDepthPS(Initializer)
 	{
-		UnresolvedSurface.Bind(Initializer.ParameterMap, TEXT("UnresolvedSurface"), SPF_Mandatory);
 	}
 	FResolveDepth2XPS() {}
-
-	void SetParameters(FRHICommandList& RHICmdList, FParameter)
-	{
-	}
-
-	LAYOUT_FIELD(FShaderResourceParameter, UnresolvedSurface);
 };
 
-
-class FResolveDepth4XPS : public FGlobalShader
+class FResolveDepth4XPS : public FResolveDepthPS
 {
 	DECLARE_EXPORTED_SHADER_TYPE(FResolveDepth4XPS, Global, RENDERCORE_API);
 public:
@@ -82,26 +74,19 @@ public:
 
 	static void ModifyCompilationEnvironment(const FGlobalShaderPermutationParameters& Parameters, FShaderCompilerEnvironment& OutEnvironment)
 	{
-		FGlobalShader::ModifyCompilationEnvironment(Parameters, OutEnvironment);
+		FResolveDepthPS::ModifyCompilationEnvironment(Parameters, OutEnvironment);
 		OutEnvironment.SetDefine(TEXT("DEPTH_RESOLVE_NUM_SAMPLES"), 4);
 	}
 
 	FResolveDepth4XPS(const ShaderMetaType::CompiledShaderInitializerType& Initializer) :
-		FGlobalShader(Initializer)
+		FResolveDepthPS(Initializer)
 	{
-		UnresolvedSurface.Bind(Initializer.ParameterMap, TEXT("UnresolvedSurface"), SPF_Mandatory);
 	}
 	FResolveDepth4XPS() {}
-
-	void SetParameters(FRHICommandList& RHICmdList, FParameter)
-	{
-	}
-
-	LAYOUT_FIELD(FShaderResourceParameter, UnresolvedSurface);
 };
 
 
-class FResolveDepth8XPS : public FGlobalShader
+class FResolveDepth8XPS : public FResolveDepthPS
 {
 	DECLARE_EXPORTED_SHADER_TYPE(FResolveDepth8XPS, Global, RENDERCORE_API);
 public:
@@ -112,22 +97,98 @@ public:
 
 	static void ModifyCompilationEnvironment(const FGlobalShaderPermutationParameters& Parameters, FShaderCompilerEnvironment& OutEnvironment)
 	{
-		FGlobalShader::ModifyCompilationEnvironment(Parameters, OutEnvironment);
+		FResolveDepthPS::ModifyCompilationEnvironment(Parameters, OutEnvironment);
 		OutEnvironment.SetDefine(TEXT("DEPTH_RESOLVE_NUM_SAMPLES"), 8);
 	}
 
 	FResolveDepth8XPS(const ShaderMetaType::CompiledShaderInitializerType& Initializer) :
-		FGlobalShader(Initializer)
+		FResolveDepthPS(Initializer)
 	{
-		UnresolvedSurface.Bind(Initializer.ParameterMap, TEXT("UnresolvedSurface"), SPF_Mandatory);
 	}
 	FResolveDepth8XPS() {}
+};
 
-	void SetParameters(FRHICommandList& RHICmdList, FParameter)
+class FResolveDepthArrayPS : public FResolveDepthPS
+{
+	DECLARE_EXPORTED_SHADER_TYPE(FResolveDepthArrayPS, Global, RENDERCORE_API);
+public:
+
+	typedef FDummyResolveParameter FParameter;
+
+	static void ModifyCompilationEnvironment(const FGlobalShaderPermutationParameters& Parameters, FShaderCompilerEnvironment& OutEnvironment)
 	{
+		FResolveDepthPS::ModifyCompilationEnvironment(Parameters, OutEnvironment);
+		OutEnvironment.SetDefine(TEXT("DEPTH_RESOLVE_TEXTUREARRAY"), 1);
 	}
 
-	LAYOUT_FIELD(FShaderResourceParameter, UnresolvedSurface);
+	FResolveDepthArrayPS(const ShaderMetaType::CompiledShaderInitializerType& Initializer) :
+		FResolveDepthPS(Initializer)
+	{
+	}
+	FResolveDepthArrayPS() {}
+};
+
+class FResolveDepthArray2XPS : public FResolveDepthArrayPS
+{
+	DECLARE_EXPORTED_SHADER_TYPE(FResolveDepthArray2XPS, Global, RENDERCORE_API);
+public:
+
+	typedef FDummyResolveParameter FParameter;
+
+	static void ModifyCompilationEnvironment(const FGlobalShaderPermutationParameters& Parameters, FShaderCompilerEnvironment& OutEnvironment)
+	{
+		FResolveDepthArrayPS::ModifyCompilationEnvironment(Parameters, OutEnvironment);
+		OutEnvironment.SetDefine(TEXT("DEPTH_RESOLVE_NUM_SAMPLES"), 2);
+	}
+
+	FResolveDepthArray2XPS(const ShaderMetaType::CompiledShaderInitializerType& Initializer) :
+		FResolveDepthArrayPS(Initializer)
+	{
+	}
+	FResolveDepthArray2XPS() {}
+};
+
+class FResolveDepthArray4XPS : public FResolveDepthArrayPS
+{
+	DECLARE_EXPORTED_SHADER_TYPE(FResolveDepthArray4XPS, Global, RENDERCORE_API);
+public:
+
+	typedef FDummyResolveParameter FParameter;
+
+	static void ModifyCompilationEnvironment(const FGlobalShaderPermutationParameters& Parameters, FShaderCompilerEnvironment& OutEnvironment)
+	{
+		FResolveDepthArrayPS::ModifyCompilationEnvironment(Parameters, OutEnvironment);
+		OutEnvironment.SetDefine(TEXT("DEPTH_RESOLVE_NUM_SAMPLES"), 4);
+	}
+
+	FResolveDepthArray4XPS(const ShaderMetaType::CompiledShaderInitializerType& Initializer) :
+		FResolveDepthArrayPS(Initializer)
+	{
+	}
+	FResolveDepthArray4XPS() {}
+};
+
+
+class FResolveDepthArray8XPS : public FResolveDepthArrayPS
+{
+	DECLARE_EXPORTED_SHADER_TYPE(FResolveDepthArray8XPS, Global, RENDERCORE_API);
+public:
+
+	typedef FDummyResolveParameter FParameter;
+
+	static bool ShouldCache(EShaderPlatform Platform) { return GetMaxSupportedFeatureLevel(Platform) >= ERHIFeatureLevel::SM5; }
+
+	static void ModifyCompilationEnvironment(const FGlobalShaderPermutationParameters& Parameters, FShaderCompilerEnvironment& OutEnvironment)
+	{
+		FResolveDepthArrayPS::ModifyCompilationEnvironment(Parameters, OutEnvironment);
+		OutEnvironment.SetDefine(TEXT("DEPTH_RESOLVE_NUM_SAMPLES"), 8);
+	}
+
+	FResolveDepthArray8XPS(const ShaderMetaType::CompiledShaderInitializerType& Initializer) :
+		FResolveDepthArrayPS(Initializer)
+	{
+	}
+	FResolveDepthArray8XPS() {}
 };
 
 class FResolveSingleSamplePS : public FGlobalShader
@@ -143,7 +204,7 @@ public:
 	}
 	
 	FResolveSingleSamplePS(const ShaderMetaType::CompiledShaderInitializerType& Initializer):
-	FGlobalShader(Initializer)
+		FGlobalShader(Initializer)
 	{
 		UnresolvedSurface.Bind(Initializer.ParameterMap,TEXT("UnresolvedSurface"), SPF_Mandatory);
 		SingleSampleIndex.Bind(Initializer.ParameterMap,TEXT("SingleSampleIndex"), SPF_Mandatory);
@@ -179,3 +240,28 @@ public:
 	LAYOUT_FIELD(FShaderParameter, PositionMinMax);
 	LAYOUT_FIELD(FShaderParameter, UVMinMax);
 };
+
+class FResolveArrayVS : public FResolveVS
+{
+	DECLARE_EXPORTED_SHADER_TYPE(FResolveArrayVS, Global, RENDERCORE_API);
+public:
+
+	static bool ShouldCompilePermutation(const FGlobalShaderPermutationParameters& Parameters) { return true; }
+
+	static void ModifyCompilationEnvironment(const FGlobalShaderPermutationParameters& Parameters, FShaderCompilerEnvironment& OutEnvironment)
+	{
+		FResolveVS::ModifyCompilationEnvironment(Parameters, OutEnvironment);
+		OutEnvironment.SetDefine(TEXT("DEPTH_RESOLVE_TEXTUREARRAY"), 1);
+	}
+
+	FResolveArrayVS(const ShaderMetaType::CompiledShaderInitializerType& Initializer)
+		: FResolveVS(Initializer)
+	{
+	}
+	FResolveArrayVS() {}
+
+	RENDERCORE_API void SetParameters(FRHICommandList& RHICmdList, const FResolveRect& SrcBounds, const FResolveRect& DstBounds, uint32 DstSurfaceWidth, uint32 DstSurfaceHeight)
+	{
+		return FResolveVS::SetParameters(RHICmdList, SrcBounds, DstBounds, DstSurfaceWidth, DstSurfaceHeight);
+	}
+};
diff --git a/Engine/Source/Runtime/Renderer/Private/HdrCustomResolveShaders.cpp b/Engine/Source/Runtime/Renderer/Private/HdrCustomResolveShaders.cpp
--- a/Engine/Source/Runtime/Renderer/Private/HdrCustomResolveShaders.cpp
+++ b/Engine/Source/Runtime/Renderer/Private/HdrCustomResolveShaders.cpp
@@ -10,3 +10,7 @@ IMPLEMENT_SHADER_TYPE(,FHdrCustomResolveFMask2xPS,TEXT("/Engine/Private/HdrCusto
 IMPLEMENT_SHADER_TYPE(,FHdrCustomResolveFMask4xPS,TEXT("/Engine/Private/HdrCustomResolveShaders.usf"),TEXT("HdrCustomResolveFMaskPS"),SF_Pixel);
 IMPLEMENT_SHADER_TYPE(,FHdrCustomResolveFMask8xPS,TEXT("/Engine/Private/HdrCustomResolveShaders.usf"),TEXT("HdrCustomResolveFMaskPS"),SF_Pixel);
 
+IMPLEMENT_SHADER_TYPE(, FHdrCustomResolveArrayVS, TEXT("/Engine/Private/HdrCustomResolveShaders.usf"), TEXT("HdrCustomResolveVS"), SF_Vertex);
+IMPLEMENT_SHADER_TYPE(, FHdrCustomResolveArray2xPS, TEXT("/Engine/Private/HdrCustomResolveShaders.usf"), TEXT("HdrCustomResolve2xPS"), SF_Pixel);
+IMPLEMENT_SHADER_TYPE(, FHdrCustomResolveArray4xPS, TEXT("/Engine/Private/HdrCustomResolveShaders.usf"), TEXT("HdrCustomResolve4xPS"), SF_Pixel);
+IMPLEMENT_SHADER_TYPE(, FHdrCustomResolveArray8xPS, TEXT("/Engine/Private/HdrCustomResolveShaders.usf"), TEXT("HdrCustomResolve8xPS"), SF_Pixel);
diff --git a/Engine/Source/Runtime/Renderer/Private/SceneRendering.cpp b/Engine/Source/Runtime/Renderer/Private/SceneRendering.cpp
--- a/Engine/Source/Runtime/Renderer/Private/SceneRendering.cpp
+++ b/Engine/Source/Runtime/Renderer/Private/SceneRendering.cpp
@@ -4797,7 +4797,7 @@ void FRendererModule::RenderPostOpaqueExtensions(
 			RenderParameters.ColorTexture = SceneTextures.Color.Target;
 			RenderParameters.DepthTexture = SceneTextures.Depth.Target;
 			RenderParameters.NormalTexture = SceneTextures.GBufferA;
-			RenderParameters.VelocityTexture = SceneTextures.Velocity;
+			RenderParameters.VelocityTexture = SceneTextures.Velocity.Target;
 			RenderParameters.SmallDepthTexture = SceneTextures.SmallDepth;
 			RenderParameters.ViewUniformBuffer = View.ViewUniformBuffer;
 			RenderParameters.SceneTexturesUniformParams = SceneTextures.UniformBuffer;
@@ -5229,6 +5229,33 @@ BEGIN_SHADER_PARAMETER_STRUCT(FResolveSceneColorParameters, )
 	RENDER_TARGET_BINDING_SLOTS()
 END_SHADER_PARAMETER_STRUCT()
 
+namespace
+{
+	template<typename ShaderType>
+	void TSetColorResolveShader(FRHICommandList& RHICmdList, FGraphicsPipelineStateInitializer& GraphicsPSOInit, const FViewInfo& View, bool bArrayResolve, FRHITexture* SceneColorTargetableRHI)
+	{
+		TShaderMapRef<ShaderType> ShaderRef(View.ShaderMap);
+		FRHIPixelShader* Shader = ShaderRef.GetPixelShader();
+		check(Shader);
+		GraphicsPSOInit.BoundShaderState.PixelShaderRHI = Shader;
+		SetGraphicsPipelineState(RHICmdList, GraphicsPSOInit, 0);
+		ShaderRef->SetParameters(RHICmdList, SceneColorTargetableRHI);
+	}
+
+	template<typename ShaderType, typename ShaderArrayType>
+	void TChooseColorResolveShader(FRHICommandList& RHICmdList, FGraphicsPipelineStateInitializer& GraphicsPSOInit, const FViewInfo& View, bool bArrayResolve, FRHITexture* SceneColorTargetableRHI)
+	{
+		if (UNLIKELY(bArrayResolve))
+		{
+			TSetColorResolveShader<ShaderArrayType>(RHICmdList, GraphicsPSOInit, View, bArrayResolve, SceneColorTargetableRHI);
+		}
+		else
+		{
+			TSetColorResolveShader<ShaderType>(RHICmdList, GraphicsPSOInit, View, bArrayResolve, SceneColorTargetableRHI);
+		}
+	}
+}
+
 void AddResolveSceneColorPass(FRDGBuilder& GraphBuilder, const FViewInfo& View, FRDGTextureMSAA SceneColor)
 {
 	check(SceneColor.IsValid());
@@ -5299,13 +5326,18 @@ void AddResolveSceneColorPass(FRDGBuilder& GraphBuilder, const FViewInfo& View,
 		}
 		else
 		{
+			bool bArrayResolve = SceneColorTargetableRHI->GetDesc().IsTextureArray();
+
 			TShaderMapRef<FHdrCustomResolveVS> VertexShader(View.ShaderMap);
+			TShaderMapRef<FHdrCustomResolveArrayVS> VertexArrayShader(View.ShaderMap);
 			GraphicsPSOInit.BoundShaderState.VertexDeclarationRHI = GetVertexDeclarationFVector4();
-			GraphicsPSOInit.BoundShaderState.VertexShaderRHI = VertexShader.GetVertexShader();
+			GraphicsPSOInit.BoundShaderState.VertexShaderRHI = bArrayResolve ? VertexArrayShader.GetVertexShader() : VertexShader.GetVertexShader();
 			GraphicsPSOInit.PrimitiveType = PT_TriangleList;
 
 			if (SceneColorFMaskRHI)
 			{
+				checkf(!bArrayResolve, TEXT("Array MSAA resolve is not supported for the FMask path"));
+
 				if (NumSamples == 2)
 				{
 					TShaderMapRef<FHdrCustomResolveFMask2xPS> PixelShader(View.ShaderMap);
@@ -5340,27 +5372,15 @@ void AddResolveSceneColorPass(FRDGBuilder& GraphBuilder, const FViewInfo& View,
 			{
 				if (NumSamples == 2)
 				{
-					TShaderMapRef<FHdrCustomResolve2xPS> PixelShader(View.ShaderMap);
-					GraphicsPSOInit.BoundShaderState.PixelShaderRHI = PixelShader.GetPixelShader();
-
-					SetGraphicsPipelineState(RHICmdList, GraphicsPSOInit, 0);
-					PixelShader->SetParameters(RHICmdList, SceneColorTargetableRHI);
+					TChooseColorResolveShader<FHdrCustomResolve2xPS, FHdrCustomResolveArray2xPS>(RHICmdList, GraphicsPSOInit, View, bArrayResolve, SceneColorTargetableRHI);
 				}
 				else if (NumSamples == 4)
 				{
-					TShaderMapRef<FHdrCustomResolve4xPS> PixelShader(View.ShaderMap);
-					GraphicsPSOInit.BoundShaderState.PixelShaderRHI = PixelShader.GetPixelShader();
-
-					SetGraphicsPipelineState(RHICmdList, GraphicsPSOInit, 0);
-					PixelShader->SetParameters(RHICmdList, SceneColorTargetableRHI);
+					TChooseColorResolveShader<FHdrCustomResolve4xPS, FHdrCustomResolveArray4xPS>(RHICmdList, GraphicsPSOInit, View, bArrayResolve, SceneColorTargetableRHI);
 				}
 				else if (NumSamples == 8)
 				{
-					TShaderMapRef<FHdrCustomResolve8xPS> PixelShader(View.ShaderMap);
-					GraphicsPSOInit.BoundShaderState.PixelShaderRHI = PixelShader.GetPixelShader();
-
-					SetGraphicsPipelineState(RHICmdList, GraphicsPSOInit, 0);
-					PixelShader->SetParameters(RHICmdList, SceneColorTargetableRHI);
+					TChooseColorResolveShader<FHdrCustomResolve8xPS, FHdrCustomResolveArray8xPS>(RHICmdList, GraphicsPSOInit, View, bArrayResolve, SceneColorTargetableRHI);
 				}
 				else
 				{
@@ -5394,6 +5414,17 @@ BEGIN_SHADER_PARAMETER_STRUCT(FResolveSceneDepthParameters, )
 	RENDER_TARGET_BINDING_SLOTS()
 END_SHADER_PARAMETER_STRUCT()
 
+namespace
+{
+	template<typename ShaderType>
+	FRHIPixelShader* GetDepthResolveShader(const FViewInfo& View, int32& OutTextureIndex)
+	{
+		TShaderMapRef<ShaderType> ShaderRef(View.ShaderMap);
+		OutTextureIndex = ShaderRef->UnresolvedSurface.GetBaseIndex();
+		return ShaderRef.GetPixelShader();
+	};
+}
+
 void AddResolveSceneDepthPass(FRDGBuilder& GraphBuilder, const FViewInfo& View, FRDGTextureMSAA SceneDepth)
 {
 	check(SceneDepth.IsValid());
@@ -5439,37 +5470,53 @@ void AddResolveSceneDepthPass(FRDGBuilder& GraphBuilder, const FViewInfo& View,
 
 		RHICmdList.SetViewport(0.0f, 0.0f, 0.0f, DepthExtent.X, DepthExtent.Y, 1.0f);
 
-		TShaderMapRef<FResolveVS> ResolveVertexShader(View.ShaderMap);
-		TShaderMapRef<FResolveDepthPS> ResolvePixelShaderAny(View.ShaderMap);
-		TShaderMapRef<FResolveDepth2XPS> ResolvePixelShader2X(View.ShaderMap);
-		TShaderMapRef<FResolveDepth4XPS> ResolvePixelShader4X(View.ShaderMap);
-		TShaderMapRef<FResolveDepth8XPS> ResolvePixelShader8X(View.ShaderMap);
+		bool bArrayResolve = SourceTextureRHI->GetDesc().IsTextureArray();
+		ensureMsgf(!bArrayResolve || RHISupportsVertexShaderLayer(View.GetShaderPlatform()) && GRHISupportsArrayIndexFromAnyShader,
+			TEXT("Resolving scene depth array requires support for outputting SV_RenderTargetArrayIndex from any shader."));
 
-		int32 TextureIndex = -1;
-		FRHIPixelShader* ResolvePixelShader;
-		switch (NumSamples)
+		/** Chooses one of many ResolvePS variants */
+		auto ChoosePixelShader = [](const FViewInfo& View, bool bIsArrayResolve, int32 NumSamples, int32& OutTextureIndex) -> FRHIPixelShader*
 		{
-		case 2:
-			TextureIndex = ResolvePixelShader2X->UnresolvedSurface.GetBaseIndex();
-			ResolvePixelShader = ResolvePixelShader2X.GetPixelShader();
-			break;
-		case 4:
-			TextureIndex = ResolvePixelShader4X->UnresolvedSurface.GetBaseIndex();
-			ResolvePixelShader = ResolvePixelShader4X.GetPixelShader();
-			break;
-		case 8:
-			TextureIndex = ResolvePixelShader8X->UnresolvedSurface.GetBaseIndex();
-			ResolvePixelShader = ResolvePixelShader8X.GetPixelShader();
-			break;
-		default:
-			ensureMsgf(false, TEXT("Unsupported depth resolve for samples: %i.  Dynamic loop method isn't supported on all platforms.  Please add specific case."), NumSamples);
-			TextureIndex = ResolvePixelShaderAny->UnresolvedSurface.GetBaseIndex();
-			ResolvePixelShader = ResolvePixelShaderAny.GetPixelShader();
-			break;
-		}
+			if (LIKELY(!bIsArrayResolve))
+			{
+				switch (NumSamples)
+				{
+					case 2:
+						return GetDepthResolveShader<FResolveDepth2XPS>(View, OutTextureIndex);
+					case 4:
+						return GetDepthResolveShader<FResolveDepth4XPS>(View, OutTextureIndex);
+					case 8:
+						return GetDepthResolveShader<FResolveDepth8XPS>(View, OutTextureIndex);
+					default:
+						ensureMsgf(false, TEXT("Unsupported depth resolve for samples: %i.  Dynamic loop method isn't supported on all platforms.  Please add specific case."), NumSamples);
+						return GetDepthResolveShader<FResolveDepthPS>(View, OutTextureIndex);
+				}
+			}
+			else
+			{
+				switch (NumSamples)
+				{
+					case 2:
+						return GetDepthResolveShader<FResolveDepthArray2XPS>(View, OutTextureIndex);
+					case 4:
+						return GetDepthResolveShader<FResolveDepthArray4XPS>(View, OutTextureIndex);
+					case 8:
+						return GetDepthResolveShader<FResolveDepthArray8XPS>(View, OutTextureIndex);
+					default:
+						ensureMsgf(false, TEXT("Unsupported depth resolve for samples: %i (texture array case).  Dynamic loop method isn't supported on all platforms.  Please add specific case."), NumSamples);
+						return GetDepthResolveShader<FResolveDepthPS>(View, OutTextureIndex);
+				}
+			}
+			// unreachable
+		};
+
+		TShaderMapRef<FResolveVS> ResolveVertexShader(View.ShaderMap);
+		TShaderMapRef<FResolveArrayVS> ResolveArrayVertexShader(View.ShaderMap);
+		int32 TextureIndex = -1;
+		FRHIPixelShader* ResolvePixelShader = ChoosePixelShader(View, bArrayResolve, NumSamples, TextureIndex);
 
 		GraphicsPSOInit.BoundShaderState.VertexDeclarationRHI = GEmptyVertexDeclaration.VertexDeclarationRHI;
-		GraphicsPSOInit.BoundShaderState.VertexShaderRHI = ResolveVertexShader.GetVertexShader();
+		GraphicsPSOInit.BoundShaderState.VertexShaderRHI = bArrayResolve ? ResolveArrayVertexShader.GetVertexShader() : ResolveVertexShader.GetVertexShader();
 		GraphicsPSOInit.BoundShaderState.PixelShaderRHI = ResolvePixelShader;
 		GraphicsPSOInit.PrimitiveType = PT_TriangleStrip;
 
@@ -5480,7 +5527,7 @@ void AddResolveSceneDepthPass(FRDGBuilder& GraphBuilder, const FViewInfo& View,
 		ResolveVertexShader->SetParameters(RHICmdList, ResolveRect, ResolveRect, DepthExtent.X, DepthExtent.Y);
 
 		RHICmdList.SetStreamSource(0, nullptr, 0);
-		RHICmdList.DrawPrimitive(0, 2, 1);
+		RHICmdList.DrawPrimitive(0, 2, bArrayResolve ? 2 : 1);
 	});
 }
 
diff --git a/Engine/Source/Runtime/Renderer/Private/SceneTextures.cpp b/Engine/Source/Runtime/Renderer/Private/SceneTextures.cpp
--- a/Engine/Source/Runtime/Renderer/Private/SceneTextures.cpp
+++ b/Engine/Source/Runtime/Renderer/Private/SceneTextures.cpp
@@ -395,6 +395,7 @@ void InitializeSceneTexturesConfig(FSceneTexturesConfig& Config, const FSceneVie
 	EShadingPath ShadingPath = FSceneInterface::GetShadingPath(ViewFamily.GetFeatureLevel());
 
 	bool bRequiresAlphaChannel = ShadingPath == EShadingPath::Mobile ? IsMobilePropagateAlphaEnabled(ViewFamily.GetShaderPlatform()) : false;
+	int32 NumberOfViewsWithMultiviewEnabled = 0;
 	for (int32 ViewIndex = 0; ViewIndex < ViewFamily.Views.Num(); ViewIndex++)
 	{
 		// Planar reflections and scene captures use scene color alpha to keep track of where content has been rendered, for compositing into a different scene later
@@ -402,8 +403,15 @@ void InitializeSceneTexturesConfig(FSceneTexturesConfig& Config, const FSceneVie
 		{
 			bRequiresAlphaChannel = true;
 		}
+
+		NumberOfViewsWithMultiviewEnabled += (ViewFamily.Views[ViewIndex]->bIsMobileMultiViewEnabled) ? 1 : 0;
 	}
 
+	ensureMsgf(NumberOfViewsWithMultiviewEnabled == 0 || NumberOfViewsWithMultiviewEnabled == ViewFamily.Views.Num(),
+		TEXT("Either all or no views in a view family should have multiview enabled. Mixing views with enabled and disabled is not allowed."));
+
+	const bool bAllViewsHaveMultiviewEnabled = NumberOfViewsWithMultiviewEnabled == ViewFamily.Views.Num();
+
 	const bool bNeedsStereoAlloc = ViewFamily.Views.ContainsByPredicate([](const FSceneView* View)
 		{
 			return (IStereoRendering::IsStereoEyeView(*View) && (FindStereoRenderTargetManager() != nullptr));
@@ -412,7 +420,7 @@ void InitializeSceneTexturesConfig(FSceneTexturesConfig& Config, const FSceneVie
 	FSceneTexturesConfigInitSettings SceneTexturesConfigInitSettings;
 	SceneTexturesConfigInitSettings.FeatureLevel = ViewFamily.GetFeatureLevel();
 	SceneTexturesConfigInitSettings.Extent = Extent;
-	SceneTexturesConfigInitSettings.bRequireMultiView = ViewFamily.bRequireMultiView;
+	SceneTexturesConfigInitSettings.bRequireMultiView = ViewFamily.bRequireMultiView && bAllViewsHaveMultiviewEnabled;
 	SceneTexturesConfigInitSettings.bRequiresAlphaChannel = bRequiresAlphaChannel;
 	SceneTexturesConfigInitSettings.bSupportsXRTargetManagerDepthAlloc = bNeedsStereoAlloc ? 1 : 0;
 	SceneTexturesConfigInitSettings.ExtraSceneColorCreateFlags = GFastVRamConfig.SceneColor;
@@ -453,6 +461,8 @@ void FMinimalSceneTextures::InitializeViewFamily(FRDGBuilder& GraphBuilder, FVie
 
 			if ((StereoDepthRHI = FindStereoDepthTexture(Config.bSupportsXRTargetManagerDepthAlloc, Config.Extent, Desc.NumSamples)) != nullptr)
 			{
+				ensureMsgf(Desc.ArraySize == StereoDepthRHI->GetDesc().ArraySize, TEXT("Resolve texture does not agree in dimensionality with Target (Resolve.ArraySize=%d, Target.ArraySize=%d)"),
+					Desc.ArraySize, StereoDepthRHI->GetDesc().ArraySize);
 				SceneTextures.Depth.Resolve = RegisterExternalTexture(GraphBuilder, StereoDepthRHI, TEXT("SceneDepthZ"));
 			}
 			else
diff --git a/Engine/Source/Runtime/Renderer/Public/HdrCustomResolveShaders.h b/Engine/Source/Runtime/Renderer/Public/HdrCustomResolveShaders.h
--- a/Engine/Source/Runtime/Renderer/Public/HdrCustomResolveShaders.h
+++ b/Engine/Source/Runtime/Renderer/Public/HdrCustomResolveShaders.h
@@ -14,6 +14,7 @@
 #include "ShaderCore.h"
 #include "ShaderParameterUtils.h"
 #include "ShaderParameters.h"
+#include "DataDrivenShaderPlatformInfo.h"
 
 class FPointerTableBase;
 
@@ -28,6 +29,25 @@ public:
 	}
 };
 
+class FHdrCustomResolveArrayVS : public FHdrCustomResolveVS
+{
+	DECLARE_SHADER_TYPE(FHdrCustomResolveArrayVS, Global);
+public:
+	FHdrCustomResolveArrayVS() {}
+	FHdrCustomResolveArrayVS(const ShaderMetaType::CompiledShaderInitializerType& Initializer)
+		: FHdrCustomResolveVS(Initializer)
+	{
+	}
+
+	static void ModifyCompilationEnvironment(const FGlobalShaderPermutationParameters& Parameters, FShaderCompilerEnvironment& OutEnvironment)
+	{
+		FHdrCustomResolveVS::ModifyCompilationEnvironment(Parameters, OutEnvironment);
+		OutEnvironment.SetDefine(TEXT("HDR_CUSTOM_RESOLVE_TEXTUREARRAY"), 1);
+	}
+};
+
+// --- regular shaders ---
+
 class FHdrCustomResolve2xPS : public FGlobalShader
 {
 	DECLARE_SHADER_TYPE(FHdrCustomResolve2xPS,Global);
@@ -55,61 +75,153 @@ protected:
 	LAYOUT_FIELD(FShaderResourceParameter, Tex);
 };
 
-class FHdrCustomResolve4xPS : public FGlobalShader
+class FHdrCustomResolve4xPS : public FHdrCustomResolve2xPS
 {
 	DECLARE_SHADER_TYPE(FHdrCustomResolve4xPS,Global);
 public:
 	FHdrCustomResolve4xPS() {}
 	FHdrCustomResolve4xPS(const ShaderMetaType::CompiledShaderInitializerType& Initializer)
-		: FGlobalShader( Initializer )
+		: FHdrCustomResolve2xPS( Initializer )
 	{
-		Tex.Bind(Initializer.ParameterMap, TEXT("Tex"), SPF_Mandatory);
+	}
+
+	static bool ShouldCompilePermutation(const FGlobalShaderPermutationParameters& Parameters)
+	{
+		return FHdrCustomResolve2xPS::ShouldCompilePermutation(Parameters);
 	}
 
 	void SetParameters(FRHICommandList& RHICmdList, FRHITexture* Texture2DMS)
 	{
-		FRHIPixelShader* PixelShaderRHI = RHICmdList.GetBoundPixelShader();
-		SetTextureParameter(RHICmdList, PixelShaderRHI, Tex, Texture2DMS);
+		return FHdrCustomResolve2xPS::SetParameters(RHICmdList, Texture2DMS);
 	}
 
 	static void ModifyCompilationEnvironment(const FGlobalShaderPermutationParameters& Parameters, FShaderCompilerEnvironment& OutEnvironment)
 	{
+		// skip parent because it sets 2X macro
 		FGlobalShader::ModifyCompilationEnvironment(Parameters, OutEnvironment);
 		OutEnvironment.SetDefine(TEXT("HDR_CUSTOM_RESOLVE_4X"), 1);
 	}
-
-protected:
-	LAYOUT_FIELD(FShaderResourceParameter, Tex);
 };
 
 
-class FHdrCustomResolve8xPS : public FGlobalShader
+class FHdrCustomResolve8xPS : public FHdrCustomResolve2xPS
 {
 	DECLARE_SHADER_TYPE(FHdrCustomResolve8xPS,Global);
 public:
 	FHdrCustomResolve8xPS() {}
 	FHdrCustomResolve8xPS(const ShaderMetaType::CompiledShaderInitializerType& Initializer)
-		: FGlobalShader( Initializer )
+		: FHdrCustomResolve2xPS( Initializer )
 	{
-		Tex.Bind(Initializer.ParameterMap, TEXT("Tex"), SPF_Mandatory);
+	}
+
+	static bool ShouldCompilePermutation(const FGlobalShaderPermutationParameters& Parameters)
+	{
+		return FHdrCustomResolve2xPS::ShouldCompilePermutation(Parameters);
 	}
 
 	void SetParameters(FRHICommandList& RHICmdList, FRHITexture* Texture2DMS)
 	{
-		FRHIPixelShader* PixelShaderRHI = RHICmdList.GetBoundPixelShader();
-		SetTextureParameter(RHICmdList, PixelShaderRHI, Tex, Texture2DMS);
+		return FHdrCustomResolve2xPS::SetParameters(RHICmdList, Texture2DMS);
 	}
 
 	static void ModifyCompilationEnvironment(const FGlobalShaderPermutationParameters& Parameters, FShaderCompilerEnvironment& OutEnvironment)
 	{
+		// skip parent because it sets 2X macro
 		FGlobalShader::ModifyCompilationEnvironment(Parameters, OutEnvironment);
 		OutEnvironment.SetDefine(TEXT("HDR_CUSTOM_RESOLVE_8X"), 1);
 	}
+};
 
-protected:
-	LAYOUT_FIELD(FShaderResourceParameter, Tex);
+// --- array shaders ---
+
+class FHdrCustomResolveArray2xPS : public FHdrCustomResolve2xPS
+{
+	DECLARE_SHADER_TYPE(FHdrCustomResolveArray2xPS, Global);
+public:
+	FHdrCustomResolveArray2xPS() {}
+	FHdrCustomResolveArray2xPS(const ShaderMetaType::CompiledShaderInitializerType& Initializer)
+		: FHdrCustomResolve2xPS(Initializer)
+	{
+	}
+
+	static bool ShouldCompilePermutation(const FGlobalShaderPermutationParameters& Parameters)
+	{
+		return FHdrCustomResolve2xPS::ShouldCompilePermutation(Parameters);
+	}
+
+	void SetParameters(FRHICommandList& RHICmdList, FRHITexture* Texture2DMS)
+	{
+		return FHdrCustomResolve2xPS::SetParameters(RHICmdList, Texture2DMS);
+	}
+
+	static void ModifyCompilationEnvironment(const FGlobalShaderPermutationParameters& Parameters, FShaderCompilerEnvironment& OutEnvironment)
+	{
+		FHdrCustomResolve2xPS::ModifyCompilationEnvironment(Parameters, OutEnvironment);
+		OutEnvironment.SetDefine(TEXT("HDR_CUSTOM_RESOLVE_TEXTUREARRAY"), 1);
+	}
+};
+
+class FHdrCustomResolveArray4xPS : public FHdrCustomResolveArray2xPS
+{
+	DECLARE_SHADER_TYPE(FHdrCustomResolveArray4xPS, Global);
+public:
+	FHdrCustomResolveArray4xPS() {}
+	FHdrCustomResolveArray4xPS(const ShaderMetaType::CompiledShaderInitializerType& Initializer)
+		: FHdrCustomResolveArray2xPS(Initializer)
+	{
+	}
+
+	static bool ShouldCompilePermutation(const FGlobalShaderPermutationParameters& Parameters)
+	{
+		return FHdrCustomResolveArray2xPS::ShouldCompilePermutation(Parameters);
+	}
+
+	void SetParameters(FRHICommandList& RHICmdList, FRHITexture* Texture2DMS)
+	{
+		return FHdrCustomResolveArray2xPS::SetParameters(RHICmdList, Texture2DMS);
+	}
+
+	static void ModifyCompilationEnvironment(const FGlobalShaderPermutationParameters& Parameters, FShaderCompilerEnvironment& OutEnvironment)
+	{
+		// skip parent because it sets 2X macro
+		FGlobalShader::ModifyCompilationEnvironment(Parameters, OutEnvironment);
+		OutEnvironment.SetDefine(TEXT("HDR_CUSTOM_RESOLVE_4X"), 1);
+		OutEnvironment.SetDefine(TEXT("HDR_CUSTOM_RESOLVE_TEXTUREARRAY"), 1);
+	}
 };
 
+
+class FHdrCustomResolveArray8xPS : public FHdrCustomResolveArray2xPS
+{
+	DECLARE_SHADER_TYPE(FHdrCustomResolveArray8xPS, Global);
+public:
+	FHdrCustomResolveArray8xPS() {}
+	FHdrCustomResolveArray8xPS(const ShaderMetaType::CompiledShaderInitializerType& Initializer)
+		: FHdrCustomResolveArray2xPS(Initializer)
+	{
+	}
+
+	static bool ShouldCompilePermutation(const FGlobalShaderPermutationParameters& Parameters)
+	{
+		return FHdrCustomResolveArray2xPS::ShouldCompilePermutation(Parameters);
+	}
+
+	void SetParameters(FRHICommandList& RHICmdList, FRHITexture* Texture2DMS)
+	{
+		return FHdrCustomResolveArray2xPS::SetParameters(RHICmdList, Texture2DMS);
+	}
+
+	static void ModifyCompilationEnvironment(const FGlobalShaderPermutationParameters& Parameters, FShaderCompilerEnvironment& OutEnvironment)
+	{
+		// skip parent because it sets 2X macro
+		FGlobalShader::ModifyCompilationEnvironment(Parameters, OutEnvironment);
+		OutEnvironment.SetDefine(TEXT("HDR_CUSTOM_RESOLVE_8X"), 1);
+		OutEnvironment.SetDefine(TEXT("HDR_CUSTOM_RESOLVE_TEXTUREARRAY"), 1);
+	}
+};
+
+// --- FMask shaders ---
+
 class FHdrCustomResolveFMask2xPS : public FGlobalShader
 {
 	DECLARE_SHADER_TYPE(FHdrCustomResolveFMask2xPS, Global);
