Fix CIS (local change contaminated another).

diff --git a/Engine/Source/Runtime/Renderer/Private/SceneRendering.cpp b/Engine/Source/Runtime/Renderer/Private/SceneRendering.cpp
--- a/Engine/Source/Runtime/Renderer/Private/SceneRendering.cpp
+++ b/Engine/Source/Runtime/Renderer/Private/SceneRendering.cpp
@@ -4797,7 +4797,7 @@ void FRendererModule::RenderPostOpaqueExtensions(
 			RenderParameters.ColorTexture = SceneTextures.Color.Target;
 			RenderParameters.DepthTexture = SceneTextures.Depth.Target;
 			RenderParameters.NormalTexture = SceneTextures.GBufferA;
-			RenderParameters.VelocityTexture = SceneTextures.Velocity.Target;
+			RenderParameters.VelocityTexture = SceneTextures.Velocity;
 			RenderParameters.SmallDepthTexture = SceneTextures.SmallDepth;
 			RenderParameters.ViewUniformBuffer = View.ViewUniformBuffer;
 			RenderParameters.SceneTexturesUniformParams = SceneTextures.UniformBuffer;
