Implemented an image generator for The Coalition's Contrast Adaptive Shading Tier 2 VRS in the form of an IVariableRateShadingImageGenerator, and added it to the VRS image manager.

Modified tonemapper to output luminance and store in PrevFrameViewInfo when parameter OUTPUT_LUMINANCE is set.

diff --git a/Engine/Shaders/Private/PostProcessTonemap.usf b/Engine/Shaders/Private/PostProcessTonemap.usf
--- a/Engine/Shaders/Private/PostProcessTonemap.usf
+++ b/Engine/Shaders/Private/PostProcessTonemap.usf
@@ -355,10 +355,13 @@ float4 TonemapCommonPS(
 	float4 GrainUV,
 	float2 ScreenPos, // [-1, 1]x[-1, 1]
 	float2 FullViewUV,
-	float4 SvPosition
+	float4 SvPosition,
+	out float OutLuminance
 	)
 {
 	float4 OutColor = 0;
+	OutLuminance = 0;
+	
 
 #if DEBUG_TONEMAPPER
 	// @param DebugTile 0,0 is center, -1,0 is left, 1,0 is right, ...
@@ -530,6 +533,7 @@ float4 TonemapCommonPS(
 		SceneColor.rgb *= ExposureScale;
 
 		OutColor.rgb = pow(SceneColor.rgb, InverseGamma.x);
+		OutLuminance = Luminance(OutColor.rgb);
 
 	#else 
 	{
@@ -615,6 +619,9 @@ float4 TonemapCommonPS(
 		#endif
 
 		half3 OutDeviceColor = ColorLookupTable(LinearColor);
+		
+		half LuminanceForPostProcessAA = dot(OutDeviceColor, half3(0.299f, 0.587f, 0.114f));
+		OutLuminance = LuminanceForPostProcessAA;
 
 		#if USE_GRAIN_QUANTIZATION
 			half Grain = GrainFromUV(GrainUV.zw);
@@ -701,13 +708,21 @@ void MainPS(
 	in noperspective float2 FullViewUV : TEXCOORD4,
 	float4 SvPosition : SV_POSITION,		// after all interpolators
 	out float4 OutColor : SV_Target0
+#if OUTPUT_LUMINANCE
+	, out float OutLuminance: SV_Target1
+#endif
 	)
 {
-    OutColor = TonemapCommonPS(UV, InVignette, GrainUV, ScreenPos, FullViewUV, SvPosition);
+	float Luminance;
+    OutColor = TonemapCommonPS(UV, InVignette, GrainUV, ScreenPos, FullViewUV, SvPosition, Luminance);
+#if OUTPUT_LUMINANCE
+	OutLuminance = Luminance;
+#endif
 }
 
 #if COMPUTESHADER
 RWTexture2D<float4> RWOutputTexture;
+RWTexture2D<float> RWOutputLuminance;
 
 [numthreads(THREADGROUP_SIZEX, THREADGROUP_SIZEY, 1)]
 void MainCS(uint2 DispatchThreadId : SV_DispatchThreadID)
@@ -726,11 +741,15 @@ void MainCS(uint2 DispatchThreadId : SV_DispatchThreadID)
 	float2 FullViewUV, ScreenPos;	
 	TonemapCommonVS(InScreenPos, UV, Vignette, GrainUV, ScreenPos, FullViewUV);
 
-	float4 OutColor = TonemapCommonPS(UV, Vignette, GrainUV, ScreenPos, FullViewUV, SvPosition);
+	float Luminance;
+	float4 OutColor = TonemapCommonPS(UV, Vignette, GrainUV, ScreenPos, FullViewUV, SvPosition, Luminance);
 
 	uint2 PixelPos = DispatchThreadId + Output_ViewportMin;
 
 	RWOutputTexture[PixelPos] = OutColor;
+#if OUTPUT_LUMINANCE
+	RWOutputLuminance[PixelPos] = Luminance;
+#endif
 }
 #endif
 
diff --git a/Engine/Shaders/Private/VariableRateShading/VRSShadingRateCalculate.usf b/Engine/Shaders/Private/VariableRateShading/VRSShadingRateCalculate.usf
new file mode 100644
--- /dev/null
+++ b/Engine/Shaders/Private/VariableRateShading/VRSShadingRateCalculate.usf
@@ -0,0 +1,279 @@
+// Copyright Epic Games, Inc. All Rights Reserved.
+
+#include "VRSShadingRateCommon.ush"
+#include "../Common.ush"
+#include "../ColorSpace.ush"
+
+float4 ViewRect;
+float EdgeThreshold;
+float ConservativeEdgeThreshold;
+
+Texture2D<float> LuminanceTexture;
+RWTexture2D<uint> VariableRateShadingTexture;
+
+#define USE_WAVE_OP COMPILER_SUPPORTS_WAVE_BIT_ORAND
+
+// Per the D3D spec: https://microsoft.github.io/DirectX-Specs/d3d/VariableRateShading.htm
+// coarse rate shading is guaranteed to always stay within VRS tile boundaries so no
+// need to check edges on tile borders 
+#define SKIP_EDGE_DETECTION_ON_BORDERS 1
+
+#if COMPUTE_SHADER
+
+#if !USE_WAVE_OP
+groupshared uint GroupShadingRate;
+#endif 
+
+groupshared half LumaMatrix[THREADGROUP_SIZEX][THREADGROUP_SIZEY];
+
+#endif /* COMPUTE_SHADER */
+
+#define SOBEL_DIMENSION 3
+
+#if COMPUTE_SHADER
+void CalculateShadingRateImageInner(
+	uint2 DispatchThreadId,
+	uint2 SwizzledGroupThreadId,
+	uint2 GroupThreadId,
+	uint2 GroupId,
+	uint GroupIndex,
+	float LuminanceValue)
+{
+	float2 Dimensions = ViewRect.zw - ViewRect.xy;
+
+	uint VRSBitMask = D3D12_SHADING_RATE_2X2 | (D3D12_SHADING_RATE_2X2 << CONSERVATIVE_SHADING_RATE_SHIFT);
+#if !USE_WAVE_OP
+	if (GroupThreadId.x == 1 && GroupThreadId.y == 1)
+	{
+		GroupShadingRate = VRSBitMask;
+	}
+#endif
+
+	LumaMatrix[SwizzledGroupThreadId.y][SwizzledGroupThreadId.x] = LuminanceValue;
+
+	GroupMemoryBarrierWithGroupSync();
+	uint UnsupportedShadingRate = 0;
+
+	// Don't bother doing edge detection on anything on the texture borders
+	bool InvalidThread = 
+		((float)DispatchThreadId.x >= Dimensions.x - 1 ||
+		(float)DispatchThreadId.y >= Dimensions.y - 1 ||
+#if SKIP_EDGE_DETECTION_ON_BORDERS
+		GroupThreadId.x == THREADGROUP_SIZEX - 1 ||
+		GroupThreadId.y == THREADGROUP_SIZEY - 1 ||
+		GroupThreadId.x == 0 ||
+		GroupThreadId.y == 0);
+#else
+		DispatchThreadId.x == 0 ||
+		DispatchThreadId.y == 0);
+#endif
+
+	int SobelWidth = SOBEL_DIMENSION / 2;
+
+	{
+		float SobelXSum = 0.0;
+		float SobelX[SOBEL_DIMENSION][SOBEL_DIMENSION] =
+		{
+			{-1.0, 0.0, 1.0},
+			{-2.0, 0.0, 2.0},
+			{-1.0, 0.0, 1.0}
+		};
+		for (int x = 0; x < SOBEL_DIMENSION; x++)
+		{
+			for (int y = 0; y < SOBEL_DIMENSION; y++)
+			{
+				int2 LumaCoord = GroupThreadId + int2(x - SobelWidth, y - SobelWidth);
+				SobelXSum += SobelX[y][x] * LumaMatrix[LumaCoord.y][LumaCoord.x];
+			}
+		}
+		if (abs(SobelXSum) > EdgeThreshold)
+		{
+			UnsupportedShadingRate |= D3D12_SHADING_RATE_2X1;
+		}
+
+		if (abs(SobelXSum) > ConservativeEdgeThreshold)
+		{
+			UnsupportedShadingRate |= (D3D12_SHADING_RATE_2X1 << CONSERVATIVE_SHADING_RATE_SHIFT);
+		}
+	}
+
+	{
+		float SobelYSum = 0.0;
+		float SobelY[SOBEL_DIMENSION][SOBEL_DIMENSION] =
+		{
+			{-1.0, -2.0, -1.0},
+			{ 0.0,  0.0,  0.0},
+			{ 1.0,  2.0,  1.0}
+		};
+		for (int x = 0; x < SOBEL_DIMENSION; x++)
+		{
+			for (int y = 0; y < SOBEL_DIMENSION; y++)
+			{
+				int2 LumaCoord = GroupThreadId + int2(x - SobelWidth, y - SobelWidth);
+				SobelYSum += SobelY[y][x] * LumaMatrix[LumaCoord.y][LumaCoord.x];
+			}
+		}
+		if (abs(SobelYSum) > EdgeThreshold)
+		{
+			UnsupportedShadingRate |= D3D12_SHADING_RATE_1X2;
+		}
+		
+		if (abs(SobelYSum) > ConservativeEdgeThreshold)
+		{
+			UnsupportedShadingRate |= (D3D12_SHADING_RATE_1X2 << CONSERVATIVE_SHADING_RATE_SHIFT);
+		}
+	}
+
+	const float2 VRSTileSize = float2(THREADGROUP_SIZEX, THREADGROUP_SIZEY);
+	uint2 GroupOffset = ceil(ViewRect.xy / VRSTileSize);
+#if USE_WAVE_OP
+	if (InvalidThread)
+	{
+		UnsupportedShadingRate = 0;
+	}
+
+	int WaveSupportedShadingRate = WaveActiveBitAnd(~UnsupportedShadingRate);
+	if (GroupThreadId.x == 1 && GroupThreadId.y == 1)
+	{
+		VariableRateShadingTexture[GroupId.xy + GroupOffset] = VRSBitMask & WaveSupportedShadingRate;
+	}
+#else
+	uint unused;
+	if(!InvalidThread)
+	{
+		InterlockedAnd(GroupShadingRate, ~UnsupportedShadingRate, unused);
+	}
+
+	GroupMemoryBarrierWithGroupSync();
+	if (GroupThreadId.x == 1 && GroupThreadId.y == 1)
+	{
+		VariableRateShadingTexture[GroupId.xy + GroupOffset] = GroupShadingRate;
+	}
+#endif
+}
+
+[numthreads(THREADGROUP_SIZEX, THREADGROUP_SIZEY, 1)]
+void CalculateShadingRateImage(
+	uint3 DispatchThreadId : SV_DispatchThreadID,
+	uint3 GroupThreadId : SV_GroupThreadID,
+	uint3 GroupId : SV_GroupID,
+	uint GroupIndex : SV_GroupIndex)
+{
+	int2 TextureCoord = int2(int(ViewRect.x + DispatchThreadId.x), int(ViewRect.y + DispatchThreadId.y));
+#if FORCE_RATE != -1
+	const float2 VRSTileSize = float2(THREADGROUP_SIZEX, THREADGROUP_SIZEY);
+	uint2 MaxSize = ceil(ViewRect.zw / VRSTileSize);
+
+	if (all(DispatchThreadId.xy < MaxSize))
+	{
+		uint DebugRate =
+	#if FORCE_RATE == 0
+			D3D12_SHADING_RATE_1X1;
+	#elif FORCE_RATE == 1
+			D3D12_SHADING_RATE_1X2;
+	#elif FORCE_RATE == 4
+			D3D12_SHADING_RATE_2X1;
+	#else // FORCE_RATE == 5
+			D3D12_SHADING_RATE_2X2;
+	#endif
+
+		VariableRateShadingTexture[TextureCoord] = 
+			(DebugRate << CONSERVATIVE_SHADING_RATE_SHIFT) |
+			DebugRate;
+	}
+#else
+	CalculateShadingRateImageInner(
+		DispatchThreadId.xy, 
+		GroupThreadId.xy, 
+		GroupThreadId.xy, 
+		GroupId, 
+		GroupIndex,
+		LuminanceTexture[TextureCoord]);
+#endif
+}
+
+#endif /* COMPUTE_SHADER */
+
+RWTexture2D<float4> SceneColorOut;
+Texture2D<uint> VariableRateShadingTextureIn;
+float DynamicResolutionScale;
+
+uint SampleVariableRateShadingTexture(float2 UV)
+{
+	uint2 VRSDimension;
+	VariableRateShadingTextureIn.GetDimensions(VRSDimension.x, VRSDimension.y);
+
+	int2 VRSIndex = UV * VRSDimension;
+
+
+#if COMPUTE_SHADER
+	return VariableRateShadingTextureIn[VRSIndex].r;
+#else
+	return VariableRateShadingTextureIn.Load(int3(VRSIndex,0)).r;
+#endif
+}
+
+
+float4 PreviewVariableRateShadingTexture_Impl(float2 UV)
+{
+	uint VSRValue = SampleVariableRateShadingTexture(UV) & D3D12_SHADING_RATE_2X2;
+
+	float4 FilterColor = float4(1, 0, 0, 0);
+
+	if(VSRValue > 0)
+	{
+		if(VSRValue == D3D12_SHADING_RATE_1X2)
+		{
+			FilterColor = float4(1, 1, 0, 0);
+		}
+		else if(VSRValue == D3D12_SHADING_RATE_2X1)
+		{
+			FilterColor = float4(0, 0, 1, 0);
+		}
+		else if(VSRValue == D3D12_SHADING_RATE_2X2)
+		{
+			FilterColor = float4(0, 1, 0, 0);
+		}
+	}
+	FilterColor.a = 0.5f;
+	return FilterColor;
+}
+
+
+void PreviewVariableRateShadingTexturePS(
+	noperspective float4 UVAndScreenPos : TEXCOORD0,
+	out float4 OutColor : SV_Target0)
+{
+	// The vertex shader has already scaled this by the DynamicResolutionScale,
+	// so we don't have to do this in the Pixel Shader as we do in the CS.
+	float2 UVLookup = UVAndScreenPos.xy;
+
+	OutColor = PreviewVariableRateShadingTexture_Impl(UVLookup);
+}
+
+#if COMPUTE_SHADER
+
+[numthreads(THREADGROUP_SIZEX, THREADGROUP_SIZEY, 1)]
+void PreviewVariableRateShadingTextureCS(
+	uint3 DispatchThreadId : SV_DispatchThreadID)
+{
+	uint2 PhysicalDimension;
+	SceneColorOut.GetDimensions(PhysicalDimension.x, PhysicalDimension.y);
+
+	uint2 OutputCoord = DispatchThreadId.xy + ViewRect.xy;
+	if(OutputCoord.x >= PhysicalDimension.x || OutputCoord.y >= PhysicalDimension.y)
+	{
+		return;
+	}
+
+	float2 UV = (ViewRect.xy + DispatchThreadId.xy) * DynamicResolutionScale / float2(PhysicalDimension);
+	
+	float4 FilterColor = PreviewVariableRateShadingTexture_Impl(UV);
+
+	if (all (FilterColor != float4(0, 0, 0, 0)))
+	{
+		SceneColorOut[OutputCoord] = FilterColor;
+	}
+}
+
+#endif
\ No newline at end of file
diff --git a/Engine/Shaders/Private/VariableRateShading/VRSShadingRateCommon.ush b/Engine/Shaders/Private/VariableRateShading/VRSShadingRateCommon.ush
new file mode 100644
--- /dev/null
+++ b/Engine/Shaders/Private/VariableRateShading/VRSShadingRateCommon.ush
@@ -0,0 +1,12 @@
+// Copyright Epic Games, Inc. All Rights Reserved.
+
+#pragma once
+
+// Taken from d3d12.h
+#define D3D12_SHADING_RATE_1X1 0x0
+#define D3D12_SHADING_RATE_1X2 0x1
+#define D3D12_SHADING_RATE_2X1 0x4
+#define D3D12_SHADING_RATE_2X2 0x5
+
+#define MAX_D3D12_SHADING_RATE_BITS 4
+#define CONSERVATIVE_SHADING_RATE_SHIFT MAX_D3D12_SHADING_RATE_BITS
\ No newline at end of file
diff --git a/Engine/Shaders/Private/VariableRateShading/VRSShadingRateReproject.usf b/Engine/Shaders/Private/VariableRateShading/VRSShadingRateReproject.usf
new file mode 100644
--- /dev/null
+++ b/Engine/Shaders/Private/VariableRateShading/VRSShadingRateReproject.usf
@@ -0,0 +1,125 @@
+// Copyright Epic Games, Inc. All Rights Reserved.
+
+#include "VRSShadingRateCommon.ush"
+#include "../Common.ush"
+#include "../SceneTexturesCommon.ush"
+
+float2 ScaledSRIDimensions;
+float2 TextureDimensions;
+float2 InvTextureDimensions;
+
+float2 SRIViewRectMin;
+float2 SRIViewRectMax;
+float2 ScaledUVOffset;
+float  InvDynamicResolutionScale;
+
+Texture2D<uint> InputSRITexture;
+RWTexture2D<uint> ScaledSRITexture;
+RWTexture2D<uint> ScaledConservativeSRITexture;
+
+groupshared float CachedDepth[THREADGROUP_SIZE / 2][THREADGROUP_SIZE / 2];
+
+// This is a large approximation of reprojection in a couple of ways in favor of performance:
+// 1. We only look at reprojection based on the camera, not velocity vectors. I.E. this will be incorrect for dynamic objects. In most cases
+//    motion blur will hide this entirely, though I found no perceptable artifacts even with motion blur off
+// 2. The reprojection only does one depth texture fetch in the center of each 16x16 pixel block.
+//
+// Both of these approximations are fine because VRS doesn't require a whole lot of accuracy since the side-effect is 
+// just lower resolution in areas that are "wrong" rather than the type of ghosting you'd get from bad reprojection in something 
+// like TAA. In practice, I found reprojection is only required in a really small set of scenarios since VRS already preserves 
+// silhouettes anyways. The main artifact I saw was materials that have a strip of metal that's not on the edge of a mesh and can
+// cause noticeble aliasing when it's got a bright specular highlight.
+float2 ReprojectBufferUV(uint3 DispatchThreadId, uint3 GroupThreadId, float2 BufferUV, bool bValidThread)
+{
+	// Only read depth once in the top left thread for each 2x2 block of thread IDs
+	// This shader is primarily bound on reading from the depth buffer on Scarlett. At 4k, 
+	// doing a single depth buffer tap that gets shared by each 2x2 block of threads ends 
+	// up making the shader almost 4x faster
+	uint2 CachedDepthIndex = uint2(GroupThreadId.x / 2, GroupThreadId.y / 2);
+	if( GroupThreadId.x % 2 == 0 && GroupThreadId.y % 2 == 0 && bValidThread)
+	{
+		const float2 pixelUVSize = InvTextureDimensions;
+
+		// Center the UV in the middle of the 2x2 block
+		float2 UVOffset = ScaledUVOffset;
+		float2 DepthUV = UVOffset + (float2(DispatchThreadId.xy) + float2(1, 1)) * pixelUVSize;
+		CachedDepth[CachedDepthIndex.x][CachedDepthIndex.y] = LookupDeviceZ(DepthUV);  
+    }
+	GroupMemoryBarrierWithGroupSync();
+
+	if(bValidThread)
+	{
+		float DeviceZ = CachedDepth[CachedDepthIndex.x][CachedDepthIndex.y];  
+		float2 ScreenPosition = ViewportUVToScreenPos(BufferUVToViewportUV(BufferUV));
+		float4 ThisClip = float4(ScreenPosition, DeviceZ, 1);
+		float4 PrevClip = mul(ThisClip, View.ClipToPrevClip);
+		float2 PrevScreen = PrevClip.xy / PrevClip.w;
+		BufferUV = ViewportUVToBufferUV(ScreenPosToViewportUV(PrevScreen));
+	}
+	return BufferUV;
+}
+
+[numthreads(THREADGROUP_SIZE, THREADGROUP_SIZE, 1)]
+void RescaleVariableRateShading(
+	uint3 DispatchThreadId : SV_DispatchThreadID,
+	uint3 GroupThreadId : SV_GroupThreadID)
+{
+	bool bValidThread = DispatchThreadId.x < ScaledSRIDimensions.x && DispatchThreadId.y < ScaledSRIDimensions.y;
+	
+	uint shadingRate = D3D12_SHADING_RATE_2X2 | (D3D12_SHADING_RATE_2X2 << CONSERVATIVE_SHADING_RATE_SHIFT);
+	float2 pixelUVSize = InvTextureDimensions;
+	const float epsilon = 0.0001;
+
+	const float2 HalfPixelOffset = float2(0.5, 0.5);
+	float2 UVOffset = ScaledUVOffset;
+	float2 CenterUV = UVOffset + (float2(DispatchThreadId.xy) + HalfPixelOffset) * pixelUVSize;
+	CenterUV = ReprojectBufferUV(DispatchThreadId, GroupThreadId, CenterUV, bValidThread); 
+
+	// FXC doesn't allow a thread to terminate before GroupMemoryBarrierWithGroupSync so 
+	// this return isn't allowed until right after ReprojectBufferUV
+	if(!bValidThread)
+	{
+		return;
+	}
+
+	// Map from pre-scaled UV to post-scaled UV.
+	float2 uvStart = CenterUV - pixelUVSize * 0.5 + epsilon;
+	float2 uvEnd = CenterUV + pixelUVSize * 0.5 - epsilon;
+
+	// Convert from UV space to Texel space of the source (pre-scaled) image
+	int2 inputCoordStart = uvStart * InvDynamicResolutionScale * TextureDimensions.xy;
+	int2 inputCoordEnd = uvEnd * InvDynamicResolutionScale * TextureDimensions.xy;
+
+	// Clamp start to top left corner of the view rect.
+	inputCoordStart.x = max(inputCoordStart.x, SRIViewRectMin.x);
+	inputCoordStart.y = max(inputCoordStart.y, SRIViewRectMin.y);
+
+	// Clamp end to bottom right corner of the view rect.
+	inputCoordEnd.x = min(inputCoordEnd.x, SRIViewRectMax.x);
+	inputCoordEnd.y = min(inputCoordEnd.y, SRIViewRectMax.y);
+
+	// start.x or y will be greater than end.x or y iff both corners were out of view on the same side.
+	// start is never clamped against the bottom-right and end is never clamped against the top-left
+	// So the only way start can be greater than end, is if both corners were out of view and only
+	// one of them got clamped passed the other.
+	bool bOnScreen = all(inputCoordStart <= inputCoordEnd);
+
+	if(bOnScreen)
+	{
+		for(uint x = inputCoordStart.x; x <= inputCoordEnd.x; x++)
+		{
+			for(uint y = inputCoordStart.y; y <= inputCoordEnd.y; y++)
+			{
+				shadingRate &= InputSRITexture[uint2(x, y)];
+			}
+		}
+	}
+	else
+	{
+		// Don't apply VRS to anything coming offscreen
+		shadingRate = D3D12_SHADING_RATE_1X1 | (D3D12_SHADING_RATE_1X1 << CONSERVATIVE_SHADING_RATE_SHIFT);
+	}
+
+	ScaledSRITexture[DispatchThreadId.xy + UVOffset * TextureDimensions] = shadingRate & D3D12_SHADING_RATE_2X2;
+	ScaledConservativeSRITexture[DispatchThreadId.xy + UVOffset * TextureDimensions] = (shadingRate >> CONSERVATIVE_SHADING_RATE_SHIFT) & D3D12_SHADING_RATE_2X2;
+}
diff --git a/Engine/Source/Runtime/Renderer/Private/DeferredShadingRenderer.cpp b/Engine/Source/Runtime/Renderer/Private/DeferredShadingRenderer.cpp
--- a/Engine/Source/Runtime/Renderer/Private/DeferredShadingRenderer.cpp
+++ b/Engine/Source/Runtime/Renderer/Private/DeferredShadingRenderer.cpp
@@ -4071,6 +4071,11 @@ void FDeferredShadingSceneRenderer::Render(FRDGBuilder& GraphBuilder)
 	// After AddPostProcessingPasses in case of Lumen Visualizations writing to feedback
 	FinishGatheringLumenSurfaceCacheFeedback(GraphBuilder, Views[0], LumenFrameTemporaries);
 
+	if (ViewFamily.bResolveScene && ViewFamilyTexture)
+	{
+		GVRSImageManager.CASDebugPreview(GraphBuilder, ViewFamily, ViewFamilyTexture);
+	}
+
 	GEngine->GetPostRenderDelegateEx().Broadcast(GraphBuilder);
 
 #if RHI_RAYTRACING
diff --git a/Engine/Source/Runtime/Renderer/Private/PostProcess/PostProcessTonemap.cpp b/Engine/Source/Runtime/Renderer/Private/PostProcess/PostProcessTonemap.cpp
--- a/Engine/Source/Runtime/Renderer/Private/PostProcess/PostProcessTonemap.cpp
+++ b/Engine/Source/Runtime/Renderer/Private/PostProcess/PostProcessTonemap.cpp
@@ -19,6 +19,8 @@
 #include "Math/Halton.h"
 #include "SystemTextures.h"
 #include "HDRHelper.h"
+#include "VariableRateShadingImageManager.h"
+
 
 bool SupportsFilmGrain(EShaderPlatform Platform)
 {
@@ -62,6 +64,11 @@ TAutoConsoleVariable<int32> CVarFilmGrainCacheTextureConstants(
 const int32 GTonemapComputeTileSizeX = 8;
 const int32 GTonemapComputeTileSizeY = 8;
 
+struct FOutputLuminance
+{
+	FRDGTextureRef Texture = nullptr;
+};
+
 namespace TonemapperPermutation
 {
 // Shared permutation dimensions between deferred and mobile renderer.
@@ -137,11 +144,13 @@ FCommonDomain BuildCommonPermutationDomain(const FViewInfo& View, bool bGammaOnl
 class FTonemapperColorFringeDim       : SHADER_PERMUTATION_BOOL("USE_COLOR_FRINGE");
 class FTonemapperGrainQuantizationDim : SHADER_PERMUTATION_BOOL("USE_GRAIN_QUANTIZATION");
 class FTonemapperOutputDeviceDim      : SHADER_PERMUTATION_ENUM_CLASS("DIM_OUTPUT_DEVICE", EDisplayOutputFormat);
+class FTonemapperOutputLuminance	  : SHADER_PERMUTATION_BOOL("OUTPUT_LUMINANCE");
 
 using FDesktopDomain = TShaderPermutationDomain<
 	FCommonDomain,
 	FTonemapperColorFringeDim,
 	FTonemapperGrainQuantizationDim,
+	FTonemapperOutputLuminance,
 	FTonemapperOutputDeviceDim>;
 
 FDesktopDomain RemapPermutation(FDesktopDomain PermutationVector, ERHIFeatureLevel::Type FeatureLevel)
@@ -166,6 +175,13 @@ FDesktopDomain RemapPermutation(FDesktopDomain PermutationVector, ERHIFeatureLev
 		PermutationVector.Set<FTonemapperColorFringeDim>(true);
 	}
 
+	// If we're going down a slow path, don't bother with VRS either to keep shader permutations down
+	if (!FVariableRateShadingImageManager::IsVRSCompatibleWithOutputType(PermutationVector.Get<FTonemapperOutputDeviceDim>()) || bFallbackToSlowest)
+	{
+		PermutationVector.Set<FTonemapperOutputLuminance>(false);
+	}
+
+
 	// You most likely need Bloom anyway.
 	CommonPermutationVector.Set<FTonemapperBloomDim>(true);
 
@@ -217,6 +233,8 @@ bool ShouldCompileDesktopPermutation(const FGlobalShaderPermutationParameters& P
 } // namespace TonemapperPermutation
 } // namespace
 
+RDG_REGISTER_BLACKBOARD_STRUCT(FOutputLuminance);
+
 FTonemapperOutputDeviceParameters GetTonemapperOutputDeviceParameters(const FSceneViewFamily& Family)
 {
 	static TConsoleVariableData<float>* CVarOutputGamma = IConsoleManager::Get().FindTConsoleVariableDataFloat(TEXT("r.TonemapperGamma"));
@@ -417,6 +435,7 @@ public:
 	BEGIN_SHADER_PARAMETER_STRUCT(FParameters, )
 		SHADER_PARAMETER_STRUCT_INCLUDE(FTonemapParameters, Tonemap)
 		SHADER_PARAMETER_RDG_TEXTURE_UAV(RWTexture2D, RWOutputTexture)
+		SHADER_PARAMETER_RDG_TEXTURE_UAV(RWTexture2D, RWOutputLuminance)
 	END_SHADER_PARAMETER_STRUCT()
 
 	static bool ShouldCompilePermutation(const FGlobalShaderPermutationParameters& Parameters)
@@ -832,18 +851,55 @@ FScreenPassTexture AddTonemapPass(FRDGBuilder& GraphBuilder, const FViewInfo& Vi
 
 		DesktopPermutationVector.Set<TonemapperPermutation::FTonemapperOutputDeviceDim>(EDisplayOutputFormat(CommonParameters.OutputDevice.OutputDevice));
 
+		DesktopPermutationVector.Set<TonemapperPermutation::FTonemapperOutputLuminance>(FVariableRateShadingImageManager::IsVRSCompatibleWithView(View));
+
 		DesktopPermutationVector = TonemapperPermutation::RemapPermutation(DesktopPermutationVector, View.GetFeatureLevel());
 	}
 
 	// Override output might not support UAVs.
 	const bool bComputePass = (Output.Texture->Desc.Flags & TexCreate_UAV) == TexCreate_UAV ? View.bUseComputePasses : false;
 
+	FRDGTextureRef OutputLuminance = {};
+	const bool bOutputLuminance = DesktopPermutationVector.Get<TonemapperPermutation::FTonemapperOutputLuminance>();
+	if (bOutputLuminance)
+	{
+		// Due to the way split-screen shares the same render target for all views, make sure
+		// that we use the same texture for saving out the luminance as well to keep the memory
+		// footprint small
+		auto CachedOuputLuminance = GraphBuilder.Blackboard.Get<FOutputLuminance>();
+		if (CachedOuputLuminance == nullptr)
+		{
+			auto NewOutputLuminance = &GraphBuilder.Blackboard.Create<FOutputLuminance>();
+			auto OutputSize = Output.Texture->Desc.GetSize();
+			const FIntPoint SDROutputSize = FIntPoint(OutputSize.X, OutputSize.Y);
+			FRDGTextureDesc Desc = FRDGTextureDesc::Create2D(
+				SDROutputSize,
+				PF_R8,
+				FClearValueBinding::Black,
+				bComputePass ? TexCreate_UAV : TexCreate_RenderTargetable);
+			OutputLuminance = NewOutputLuminance->Texture = GraphBuilder.CreateTexture(Desc, TEXT("Final Luminance"));
+		}
+		else
+		{
+			OutputLuminance = CachedOuputLuminance->Texture;
+		}
+	}
+
 	if (bComputePass)
 	{
 		FTonemapCS::FParameters* PassParameters = GraphBuilder.AllocParameters<FTonemapCS::FParameters>();
 		PassParameters->Tonemap = CommonParameters;
 		PassParameters->RWOutputTexture = GraphBuilder.CreateUAV(Output.Texture);
 
+		if (OutputLuminance)
+		{
+			PassParameters->RWOutputLuminance = GraphBuilder.CreateUAV(OutputLuminance);
+		}
+		else
+		{
+			check(DesktopPermutationVector.Get<TonemapperPermutation::FTonemapperOutputLuminance>() == false);
+		}
+
 		TShaderMapRef<FTonemapCS> ComputeShader(View.ShaderMap, DesktopPermutationVector);
 
 		FComputeShaderUtils::AddPass(
@@ -859,6 +915,11 @@ FScreenPassTexture AddTonemapPass(FRDGBuilder& GraphBuilder, const FViewInfo& Vi
 		PassParameters->Tonemap = CommonParameters;
 		PassParameters->RenderTargets[0] = Output.GetRenderTargetBinding();
 
+		if (OutputLuminance)
+		{
+			PassParameters->RenderTargets[1] = FRenderTargetBinding(OutputLuminance, ERenderTargetLoadAction::ELoad);
+		}
+
 		TShaderMapRef<FTonemapVS> VertexShader(View.ShaderMap);
 		TShaderMapRef<FTonemapPS> PixelShader(View.ShaderMap, DesktopPermutationVector);
 
@@ -885,6 +946,11 @@ FScreenPassTexture AddTonemapPass(FRDGBuilder& GraphBuilder, const FViewInfo& Vi
 			SetShaderParameters(RHICmdList, VertexShader, VertexShader.GetVertexShader(), PassParameters->Tonemap);
 			SetShaderParameters(RHICmdList, PixelShader, PixelShader.GetPixelShader(), *PassParameters);
 		});
+
+		if (OutputLuminance && View.ViewState)
+		{
+			GraphBuilder.QueueTextureExtraction(OutputLuminance, &View.ViewState->PrevFrameViewInfo.LuminanceHistory);
+		}
 	}
 
 	return MoveTemp(Output);
diff --git a/Engine/Source/Runtime/Renderer/Private/SceneRendering.h b/Engine/Source/Runtime/Renderer/Private/SceneRendering.h
--- a/Engine/Source/Runtime/Renderer/Private/SceneRendering.h
+++ b/Engine/Source/Runtime/Renderer/Private/SceneRendering.h
@@ -1157,6 +1157,9 @@ struct FPreviousViewInfo
 	// History for denoising all lights penumbra at once.
 	FScreenSpaceDenoiserHistory PolychromaticPenumbraHarmonicsHistory;
 
+	// History for the final back buffer luminance
+	TRefCountPtr<IPooledRenderTarget> LuminanceHistory;
+
 	// Mobile bloom setup eye adaptation surface.
 	TRefCountPtr<IPooledRenderTarget> MobileBloomSetup_EyeAdaptation;
 
diff --git a/Engine/Source/Runtime/Renderer/Private/VariableRateShading/ContrastAdaptiveImageGenerator.cpp b/Engine/Source/Runtime/Renderer/Private/VariableRateShading/ContrastAdaptiveImageGenerator.cpp
new file mode 100644
--- /dev/null
+++ b/Engine/Source/Runtime/Renderer/Private/VariableRateShading/ContrastAdaptiveImageGenerator.cpp
@@ -0,0 +1,787 @@
+// Copyright Epic Games, Inc. All Rights Reserved.
+
+#include "ContrastAdaptiveImageGenerator.h"
+#include "StereoRenderTargetManager.h"
+#include "GlobalShader.h"
+#include "ShaderParameterUtils.h"
+#include "RenderGraphUtils.h"
+#include "RHI.h"
+#include "RHIDefinitions.h"
+#include "RenderTargetPool.h"
+#include "SystemTextures.h"
+#include "SceneView.h"
+#include "IEyeTracker.h"
+#include "IHeadMountedDisplay.h"
+#include "IXRTrackingSystem.h"
+#include "Engine/Engine.h"
+#include "UnrealClient.h"
+#include "PostProcess/PostProcessTonemap.h"
+
+/**
+ * Contrast Adaptive Shading (CAS) is a Tier 2 Variable Rate Shading method which generates a VRS image by examining the contrast from the previous frame.
+ * An image is generated which designates lower shading rates for areas of lower contrast in which reductions are unlikely to be noticed.
+ * This image is then reprojected and rescaled in accordance with camera movement and dynamic resolution changes before being provided to the manager.
+ */
+
+
+/**
+ * CAS Parameters
+ */
+
+TAutoConsoleVariable<int32> CVarVRSContrastAdaptiveShading(
+	TEXT("r.VRS.ContrastAdaptiveShading"),
+	0,
+	TEXT("Enables using Variable Rate Shading based on the luminance from the previous frame's post process output \n"),
+	ECVF_RenderThreadSafe);
+TAutoConsoleVariable<float> CVarVRSEdgeThreshold(
+	TEXT("r.VRS.ContrastAdaptiveShading.EdgeThreshold"),
+	0.2,
+	TEXT(""),
+	ECVF_RenderThreadSafe);
+TAutoConsoleVariable<float> CVarVRSConservativeEdgeThreshold(
+	TEXT("r.VRS.ContrastAdaptiveShading.ConservativeEdgeThreshold"),
+	0.1,
+	TEXT(""),
+	ECVF_RenderThreadSafe);
+TAutoConsoleVariable<float> CVarVRS_HDR10CorrectionMultiplier(
+	TEXT("r.VRS.ContrastAdaptiveShading.HDR10CorrectionMultiplier"),
+	0.55,
+	TEXT("Approximation multiplier to account for how perceptual values are spread out in SDR vs HDR10\n"),
+	ECVF_RenderThreadSafe);
+
+/**
+ * Pass Settings
+ */
+
+TAutoConsoleVariable<int32> CVarVRSBasePass(
+	TEXT("r.VRS.BasePass"),
+	2,
+	TEXT("Enables Variable Rate Shading for the base pass\n")
+	TEXT("0: Disabled")
+	TEXT("1: Full")
+	TEXT("2: Conservative (default)"),
+	ECVF_RenderThreadSafe);
+TAutoConsoleVariable<int32> CVarVRSTranslucency(
+	TEXT("r.VRS.Translucency"),
+	1,
+	TEXT("Enable VRS with translucency rendering.\n")
+	TEXT("0: Disabled")
+	TEXT("1: Full (default)")
+	TEXT("2: Conservative"),
+	ECVF_RenderThreadSafe);
+TAutoConsoleVariable<int32> CVarVRSNaniteEmitGBuffer(
+	TEXT("r.VRS.NaniteEmitGBuffer"),
+	2,
+	TEXT("Enable VRS with Nanite EmitGBuffer rendering.\n")
+	TEXT("0: Disabled")
+	TEXT("1: Full")
+	TEXT("2: Conservative (default)"),
+	ECVF_RenderThreadSafe);
+TAutoConsoleVariable<int32> CVarVRS_SSAO(
+	TEXT("r.VRS.SSAO"),
+	2,
+	TEXT("Enable VRS with SSAO rendering.\n")
+	TEXT("0: Disabled")
+	TEXT("1: Full")
+	TEXT("2: Conservative (default)"),
+	ECVF_RenderThreadSafe);
+TAutoConsoleVariable<int32> CVarVRS_SSR(
+	TEXT("r.VRS.SSR"),
+	2,
+	TEXT("Enable VRS with SSR (PS) rendering.\n")
+	TEXT("0: Disabled")
+	TEXT("1: Full")
+	TEXT("2: Conservative (default)"),
+	ECVF_RenderThreadSafe);
+TAutoConsoleVariable<int32> CVarVRSReflectionEnvironmentSky(
+	TEXT("r.VRS.ReflectionEnvironmentSky"),
+	2,
+	TEXT("Enable VRS with ReflectionEnvironmentAndSky (PS) rendering.\n")
+	TEXT("0: Disabled")
+	TEXT("1: Full")
+	TEXT("2: Conservative (default)"),
+	ECVF_RenderThreadSafe);
+TAutoConsoleVariable<int32> CVarVRSLightFunctions(
+	TEXT("r.VRS.LightFunctions"),
+	1,
+	TEXT("Enables Variable Rate Shading for light functions\n")
+	TEXT("0: Disabled")
+	TEXT("1: Full (default)")
+	TEXT("2: Conservative"),
+	ECVF_RenderThreadSafe);
+TAutoConsoleVariable<int32> CVarVRSDecals(
+	TEXT("r.VRS.Decals"),
+	2,
+	TEXT("Enables Variable Rate Shading for decals\n")
+	TEXT("0: Disabled")
+	TEXT("1: Full")
+	TEXT("2: Conservative (default)"),
+	ECVF_RenderThreadSafe);
+
+/**
+ * Debug Settings 
+ */
+
+int GVRSDebugForceRate = -1;
+FAutoConsoleVariableRef CVarVRSDebugForceRate(
+	TEXT("r.VRS.ContrastAdaptiveShading.Debug.ForceRate"),
+	GVRSDebugForceRate,
+	TEXT("-1 : None, 0 : Force 1x1, 1 : Force 1x2, 4 : Force 2x1, 5: Force 2x2"));
+
+TAutoConsoleVariable<int32> CVarVRSPreview(
+	TEXT("r.VRS.Preview"),
+	0,
+	TEXT("Show a debug visualiation of the SRI texture.")
+	TEXT("0 - off, 1 - the SRI texture, 2- the conservative SRI texture, 3 - the unscaled SRI texture"),
+	ECVF_RenderThreadSafe);
+
+/**
+ * Shaders
+ */
+
+ //---------------------------------------------------------------------------------------------
+class FCalculateShadingRateImageCS : public FGlobalShader
+{
+	DECLARE_GLOBAL_SHADER(FCalculateShadingRateImageCS);
+	SHADER_USE_PARAMETER_STRUCT(FCalculateShadingRateImageCS, FGlobalShader);
+
+	class FThreadGroupX : SHADER_PERMUTATION_SPARSE_INT("THREADGROUP_SIZEX", 8, 16);
+	class FThreadGroupY : SHADER_PERMUTATION_SPARSE_INT("THREADGROUP_SIZEY", 8, 16);
+	class FForceRate : SHADER_PERMUTATION_SPARSE_INT("FORCE_RATE", -1, 0, 1, 4, 5);
+
+	using FPermutationDomain = TShaderPermutationDomain<FThreadGroupX, FThreadGroupY, FForceRate>;
+
+	BEGIN_SHADER_PARAMETER_STRUCT(FParameters, )
+		SHADER_PARAMETER_RDG_TEXTURE(Texture2D<float>, LuminanceTexture)
+		SHADER_PARAMETER(FVector4f, ViewRect)
+		SHADER_PARAMETER(float, EdgeThreshold)
+		SHADER_PARAMETER(float, ConservativeEdgeThreshold)
+		SHADER_PARAMETER_RDG_TEXTURE_UAV(RWTexture2D<uint>, VariableRateShadingTexture)
+	END_SHADER_PARAMETER_STRUCT()
+
+	static bool ShouldCompilePermutation(const FGlobalShaderPermutationParameters& Parameters)
+	{
+		return IsFeatureLevelSupported(Parameters.Platform, ERHIFeatureLevel::SM5)
+			&& FDataDrivenShaderPlatformInfo::GetSupportsVariableRateShading(Parameters.Platform);
+	}
+
+	static void ModifyCompilationEnvironment(const FGlobalShaderPermutationParameters& Parameters, FShaderCompilerEnvironment& OutEnvironment)
+	{
+		FGlobalShader::ModifyCompilationEnvironment(Parameters, OutEnvironment);
+		OutEnvironment.SetDefine(TEXT("COMPUTE_SHADER"), 1);
+	}
+
+	static void InitParameters(
+		FParameters& Parameters,
+		FRDGTextureRef Luminance,
+		const FIntRect& ViewRect,
+		bool bIsHDR10,
+		FRDGTextureUAV* ShadingRateImage)
+	{
+		Parameters.LuminanceTexture = Luminance;
+		Parameters.ViewRect = FVector4f(ViewRect.Min.X, ViewRect.Min.Y, ViewRect.Max.X, ViewRect.Max.Y);
+		const float cEdgeThresholdCorrectionValue = bIsHDR10 ? CVarVRS_HDR10CorrectionMultiplier.GetValueOnRenderThread() : 1.0;
+		Parameters.EdgeThreshold = cEdgeThresholdCorrectionValue * CVarVRSEdgeThreshold.GetValueOnRenderThread();
+		Parameters.ConservativeEdgeThreshold = cEdgeThresholdCorrectionValue * CVarVRSConservativeEdgeThreshold.GetValueOnRenderThread();
+		Parameters.VariableRateShadingTexture = ShadingRateImage;
+	}
+};
+IMPLEMENT_GLOBAL_SHADER(FCalculateShadingRateImageCS, "/Engine/Private/VariableRateShading/VRSShadingRateCalculate.usf", "CalculateShadingRateImage", SF_Compute);
+
+//---------------------------------------------------------------------------------------------
+class FRescaleVariableRateShadingCS : public FGlobalShader
+{
+	DECLARE_GLOBAL_SHADER(FRescaleVariableRateShadingCS);
+	SHADER_USE_PARAMETER_STRUCT(FRescaleVariableRateShadingCS, FGlobalShader);
+
+	static constexpr uint32 ThreadGroupSize = 8;
+
+	static bool ShouldCompilePermutation(const FGlobalShaderPermutationParameters& Parameters)
+	{
+		return IsFeatureLevelSupported(Parameters.Platform, ERHIFeatureLevel::SM5);
+	}
+
+	BEGIN_SHADER_PARAMETER_STRUCT(FParameters, )
+		SHADER_PARAMETER_RDG_UNIFORM_BUFFER(FSceneTextureUniformParameters, SceneTextures)
+		SHADER_PARAMETER_STRUCT_REF(FViewUniformShaderParameters, View)
+		SHADER_PARAMETER_RDG_TEXTURE(Texture2D<uint>, InputSRITexture)
+		SHADER_PARAMETER(FVector2f, SRIViewRectMin)
+		SHADER_PARAMETER(FVector2f, SRIViewRectMax)
+		SHADER_PARAMETER(FVector2f, TextureDimensions)
+		SHADER_PARAMETER(FVector2f, InvTextureDimensions)
+		SHADER_PARAMETER(FVector2f, ScaledSRIDimensions)
+		SHADER_PARAMETER(FVector2f, ScaledUVOffset)
+		SHADER_PARAMETER(float, InvDynamicResolutionScale)
+		SHADER_PARAMETER_RDG_TEXTURE_UAV(RWTexture2D<uint>, ScaledSRITexture)
+		SHADER_PARAMETER_RDG_TEXTURE_UAV(RWTexture2D<uint>, ScaledConservativeSRITexture)
+	END_SHADER_PARAMETER_STRUCT()
+
+	static void ModifyCompilationEnvironment(const FGlobalShaderPermutationParameters& Parameters, FShaderCompilerEnvironment& OutEnvironment)
+	{
+		FGlobalShader::ModifyCompilationEnvironment(Parameters, OutEnvironment);
+		OutEnvironment.SetDefine(TEXT("THREADGROUP_SIZE"), ThreadGroupSize);
+	}
+
+	static void InitParameters(
+		FParameters& Parameters,
+		const FMinimalSceneTextures& SceneTextures,
+		const FViewInfo& ViewInfo,
+		FRDGTextureRef InputSRITexture,
+		FVector2f ViewRectMin,
+		FVector2f ViewRectMax,
+		FVector2f ScaledSRIDimensions,
+		FVector2f TextureDimensions,
+		FVector2f ScaledUVOffset,
+		float DynamicResolutionScale,
+		FRDGTextureUAVRef ScaledSRIUAV,
+		FRDGTextureUAVRef ScaledConservateSRIUAV)
+	{
+		Parameters.SceneTextures = SceneTextures.UniformBuffer;
+		Parameters.View = ViewInfo.ViewUniformBuffer;
+		Parameters.InputSRITexture = InputSRITexture;
+		Parameters.SRIViewRectMin = ViewRectMin;
+		Parameters.SRIViewRectMax = ViewRectMax;
+		Parameters.TextureDimensions = TextureDimensions;
+		Parameters.InvTextureDimensions = FVector2f(1.0f / TextureDimensions.X, 1.0f / TextureDimensions.Y);
+		Parameters.ScaledSRIDimensions = ScaledSRIDimensions;
+		Parameters.ScaledUVOffset = ScaledUVOffset;
+		Parameters.InvDynamicResolutionScale = 1.0f / DynamicResolutionScale;
+		Parameters.ScaledSRITexture = ScaledSRIUAV;
+		Parameters.ScaledConservativeSRITexture = ScaledConservateSRIUAV;
+	}
+};
+IMPLEMENT_GLOBAL_SHADER(FRescaleVariableRateShadingCS, "/Engine/Private/VariableRateShading/VRSShadingRateReproject.usf", "RescaleVariableRateShading", SF_Compute);
+
+class FDebugVariableRateShadingCS : public FGlobalShader
+{
+	DECLARE_GLOBAL_SHADER(FDebugVariableRateShadingCS);
+	SHADER_USE_PARAMETER_STRUCT(FDebugVariableRateShadingCS, FGlobalShader);
+
+	static constexpr uint32 ThreadGroupSize = 8;
+
+	static bool ShouldCompilePermutation(const FGlobalShaderPermutationParameters& Parameters)
+	{
+		return IsFeatureLevelSupported(Parameters.Platform, ERHIFeatureLevel::SM5);
+	}
+
+	BEGIN_SHADER_PARAMETER_STRUCT(FParameters, )
+		SHADER_PARAMETER_RDG_TEXTURE(Texture2D<uint>, VariableRateShadingTextureIn)
+		SHADER_PARAMETER(FVector4f, ViewRect)
+		SHADER_PARAMETER(float, DynamicResolutionScale)
+		SHADER_PARAMETER_RDG_TEXTURE_UAV(RWTexture2D<float4>, SceneColorOut)
+	END_SHADER_PARAMETER_STRUCT()
+
+	static void ModifyCompilationEnvironment(const FGlobalShaderPermutationParameters& Parameters, FShaderCompilerEnvironment& OutEnvironment)
+	{
+		FGlobalShader::ModifyCompilationEnvironment(Parameters, OutEnvironment);
+		OutEnvironment.SetDefine(TEXT("THREADGROUP_SIZEX"), ThreadGroupSize);
+		OutEnvironment.SetDefine(TEXT("THREADGROUP_SIZEY"), ThreadGroupSize);
+		OutEnvironment.SetDefine(TEXT("COMPUTE_SHADER"), 1);
+	}
+
+	static void InitParameters(
+		FParameters& Parameters,
+		FRDGTextureRef VariableRateShadingTexture,
+		const FIntRect& ViewRect,
+		float DynamicResolutionScale,
+		FRDGTextureUAVRef SceneColorUAV)
+	{
+		Parameters.VariableRateShadingTextureIn = VariableRateShadingTexture;
+		Parameters.ViewRect = FVector4f(ViewRect.Min.X, ViewRect.Min.Y, ViewRect.Max.X, ViewRect.Max.Y);
+		Parameters.DynamicResolutionScale = DynamicResolutionScale;
+		Parameters.SceneColorOut = SceneColorUAV;
+	}
+};
+
+IMPLEMENT_GLOBAL_SHADER(FDebugVariableRateShadingCS, "/Engine/Private/VariableRateShading/VRSShadingRateCalculate.usf", "PreviewVariableRateShadingTextureCS", SF_Compute);
+
+//---------------------------------------------------------------------------------------------
+using FDebugVariableRateShadingVS = FScreenPassVS;
+
+//---------------------------------------------------------------------------------------------
+class FDebugVariableRateShadingPS : public FGlobalShader
+{
+	DECLARE_GLOBAL_SHADER(FDebugVariableRateShadingPS);
+	SHADER_USE_PARAMETER_STRUCT(FDebugVariableRateShadingPS, FGlobalShader);
+
+	static bool ShouldCompilePermutation(const FGlobalShaderPermutationParameters& Parameters)
+	{
+		return IsFeatureLevelSupported(Parameters.Platform, ERHIFeatureLevel::SM5);
+	}
+
+	static void ModifyCompilationEnvironment(const FGlobalShaderPermutationParameters& Parameters, FShaderCompilerEnvironment& OutEnvironment)
+	{
+		FGlobalShader::ModifyCompilationEnvironment(Parameters, OutEnvironment);
+		OutEnvironment.SetDefine(TEXT("COMPUTE_SHADER"), 0);
+	}
+
+	BEGIN_SHADER_PARAMETER_STRUCT(FParameters, )
+		SHADER_PARAMETER_RDG_TEXTURE(Texture2D<uint>, VariableRateShadingTextureIn)
+		RENDER_TARGET_BINDING_SLOTS()
+	END_SHADER_PARAMETER_STRUCT()
+		
+	static void InitParameters(
+			FParameters& Parameters,
+			FRDGTextureRef VariableRateShadingTexture,
+			FRDGTextureRef OutputSceneColor)
+	{
+		Parameters.VariableRateShadingTextureIn = VariableRateShadingTexture;
+		Parameters.RenderTargets[0] = FRenderTargetBinding(OutputSceneColor, ERenderTargetLoadAction::ELoad);
+	}
+};
+
+IMPLEMENT_GLOBAL_SHADER(FDebugVariableRateShadingPS, "/Engine/Private/VariableRateShading/VRSShadingRateCalculate.usf", "PreviewVariableRateShadingTexturePS", SF_Pixel);
+
+/**
+ * Helper Functions and Structures
+ */
+
+namespace ESRITextureType
+{
+	enum Type
+	{
+		None = 0,
+		ScaledSRIForRender,
+		ScaledConservativeSRIForRender,
+		ConstructedSRI,
+		Num
+	};
+	bool IsInBounds(int32 TypeAsInt)
+	{
+		return TypeAsInt >= 0 && TypeAsInt < ESRITextureType::Num;
+	}
+	bool IsInBounds(ESRITextureType::Type TextureType)
+	{
+		return IsInBounds(static_cast<int32>(TextureType));
+	}
+	bool IsValidShadingRateTexture(int32 TextureType)
+	{
+		return IsInBounds(TextureType) && TextureType != ESRITextureType::None && TextureType != ESRITextureType::ConstructedSRI;
+	}
+	bool IsValidShadingRateTexture(ESRITextureType::Type TextureType)
+	{
+		return IsValidShadingRateTexture(static_cast<int32>(TextureType));
+	}
+	ESRITextureType::Type GetTextureType(FVariableRateShadingImageManager::EVRSPassType PassType)
+	{
+		static struct FStaticData
+		{
+			TAutoConsoleVariable<int32>*CVarByPassType[FVariableRateShadingImageManager::EVRSPassType::Num] = {};
+			FStaticData()
+			{
+				CVarByPassType[FVariableRateShadingImageManager::EVRSPassType::BasePass] = &CVarVRSBasePass;
+				CVarByPassType[FVariableRateShadingImageManager::EVRSPassType::TranslucencyAll] = &CVarVRSTranslucency;
+				CVarByPassType[FVariableRateShadingImageManager::EVRSPassType::NaniteEmitGBufferPass] = &CVarVRSNaniteEmitGBuffer;
+				CVarByPassType[FVariableRateShadingImageManager::EVRSPassType::SSAO] = &CVarVRS_SSAO;
+				CVarByPassType[FVariableRateShadingImageManager::EVRSPassType::SSR] = &CVarVRS_SSR;
+				CVarByPassType[FVariableRateShadingImageManager::EVRSPassType::ReflectionEnvironmentAndSky] = &CVarVRSReflectionEnvironmentSky;
+				CVarByPassType[FVariableRateShadingImageManager::EVRSPassType::LightFunctions] = &CVarVRSLightFunctions;
+				CVarByPassType[FVariableRateShadingImageManager::EVRSPassType::Decals] = &CVarVRSDecals;
+			}
+		} StaticData;
+		const int32 PassTypeAsInt = static_cast<int32>(PassType);
+		if (PassTypeAsInt >= 0 && PassTypeAsInt < EMeshPass::Num)
+		{
+			auto* CVar = StaticData.CVarByPassType[PassTypeAsInt];
+			if (CVar)
+			{
+				int32 TextureTypeAsInt = CVar->GetValueOnRenderThread();
+				if (IsValidShadingRateTexture(TextureTypeAsInt))
+				{
+					return static_cast<ESRITextureType::Type>(TextureTypeAsInt);
+				}
+			}
+		}
+		return ESRITextureType::None;
+	}
+};
+
+namespace ESRIPreviewType
+{
+	enum Type
+	{
+		Off,
+		Projected,
+		ProjectedConservative,
+		BeforeReprojection,
+		Num
+	};
+	static const TCHAR* Names[] = {
+		TEXT("Off"),
+		TEXT("Projected"),
+		TEXT("ProjectedConservative"),
+		TEXT("BeforeReprojection"),
+	};
+	static const TCHAR* GetName(Type PreviewType)
+	{
+		if (PreviewType >= Off && PreviewType < Num)
+		{
+			return Names[static_cast<int32>(PreviewType)];
+		}
+
+		return TEXT("Invalid Type");
+	}
+};
+
+static FIntPoint GetSRITileSize()
+{
+	return FIntPoint(GRHIVariableRateShadingImageTileMinWidth, GRHIVariableRateShadingImageTileMinHeight);
+}
+
+static const TCHAR* ShadingRateTextureName = TEXT("ShadingRateTexture");
+static const TCHAR* ScaledShadingRateTextureName = TEXT("ScaledShadingRateTexture");
+static const TCHAR* ScaledConservativeShadingRateTextureName = TEXT("ConservativeScaledShadingRateTexture");
+
+struct RENDERER_API FVRSTextures
+{
+	// Returns an FVRSTextures created immutable instance from the builder blackboard. Asserts if none was created.
+	static const FVRSTextures& Get(FRDGBuilder& GraphBuilder)
+	{
+		const FVRSTextures* VRSTextures = GraphBuilder.Blackboard.Get<FVRSTextures>();
+			checkf(VRSTextures, TEXT("FVRSTextures was unexpectedly not initialized."));
+		return *VRSTextures;
+	}
+	static const bool IsInitialized(FRDGBuilder& GraphBuilder)
+	{
+		const FVRSTextures* VRSTextures = GraphBuilder.Blackboard.Get<FVRSTextures>();
+		return VRSTextures != nullptr;
+	}
+	void Create(FRDGBuilder& GraphBuilder, const FSceneViewFamily& ViewFamily)
+	{
+		FRDGTextureDesc ConstructedSRIDesc = CreateSRIDesc(ViewFamily, false);
+		ConstructedSRI = GraphBuilder.CreateTexture(ConstructedSRIDesc, ShadingRateTextureName);
+		FRDGTextureDesc ScaledSRIDesc = CreateSRIDesc(ViewFamily, true);
+		ScaledSRI = GraphBuilder.CreateTexture(ScaledSRIDesc, ScaledShadingRateTextureName);
+		ScaledConservativeSRI = GraphBuilder.CreateTexture(ScaledSRIDesc, ScaledConservativeShadingRateTextureName);
+	}
+	FRDGTextureRef ConstructedSRI;
+	FRDGTextureRef ScaledSRI;
+	FRDGTextureRef ScaledConservativeSRI;
+private:
+	static FRDGTextureDesc CreateSRIDesc(const FSceneViewFamily& ViewFamily, bool bIsForDynResScaled)
+	{
+		FIntPoint TileSize = GetSRITileSize();
+		FIntPoint ViewTargetExtents = (bIsForDynResScaled)
+			? FSceneTexturesConfig::Get().Extent
+			: ViewFamily.RenderTarget->GetSizeXY();
+		FIntPoint SRIDimensions = FMath::DivideAndRoundUp(ViewTargetExtents, TileSize);
+		return FRDGTextureDesc::Create2D(
+			SRIDimensions,
+			GRHIVariableRateShadingImageFormat,
+			EClearBinding::ENoneBound,
+			ETextureCreateFlags::DisableDCC |
+			ETextureCreateFlags::ShaderResource |
+			ETextureCreateFlags::UAV);
+	}
+};
+RDG_REGISTER_BLACKBOARD_STRUCT(FVRSTextures);
+
+static EDisplayOutputFormat GetDisplayOutputFormat(const FSceneView& View)
+{
+	FTonemapperOutputDeviceParameters Parameters = GetTonemapperOutputDeviceParameters(*View.Family);
+	return (EDisplayOutputFormat)Parameters.OutputDevice;
+}
+
+static bool IsHDR10(const EDisplayOutputFormat& OutputFormat)
+{
+	return OutputFormat == EDisplayOutputFormat::HDR_ACES_1000nit_ST2084 ||
+		OutputFormat == EDisplayOutputFormat::HDR_ACES_2000nit_ST2084;
+}
+
+static bool IsContrastAdaptiveShadingEnabled()
+{
+	return GRHISupportsAttachmentVariableRateShading && GRHIAttachmentVariableRateShadingEnabled && (CVarVRSContrastAdaptiveShading.GetValueOnRenderThread() != 0);
+}
+
+static FIntRect GetPostProcessOutputRect(const FViewInfo& ViewInfo)
+{
+	// If TAA/TSR is enabled, upscaling is done at the start of post-processing so the final output will match UnscaledViewRect. Otherwise use the dynamically rescale view rect since
+	// the secondary upscale will happen after post processing
+	return ViewInfo.PrimaryScreenPercentageMethod == EPrimaryScreenPercentageMethod::TemporalUpscale ? ViewInfo.UnscaledViewRect.Scale(ViewInfo.Family->SecondaryViewFraction) : ViewInfo.ViewRect;
+}
+
+bool AddCreateShadingRateImagePass(
+	FRDGBuilder& GraphBuilder,
+	const FViewInfo& View)
+{
+	//------------------------------------------------------------------------------------------------
+	// Do some sanity checks for early out
+	if (!IsContrastAdaptiveShadingEnabled() || !FVariableRateShadingImageManager::IsVRSCompatibleWithView(View) || !View.PrevViewInfo.LuminanceHistory)
+	{
+		// Shading Rate Image unsupported
+		return false;
+	}
+	FRDGTextureRef Luminance = GraphBuilder.RegisterExternalTexture(View.PrevViewInfo.LuminanceHistory);
+	const FVRSTextures& VRSTextures = FVRSTextures::Get(GraphBuilder);
+	// Complete early out sanity checks
+	//------------------------------------------------------------------------------------------------
+
+	{
+		FCalculateShadingRateImageCS::FPermutationDomain PermutationVector;
+
+		const FIntPoint TileSize = GetSRITileSize();
+		PermutationVector.Set<FCalculateShadingRateImageCS::FThreadGroupX>(TileSize.X);
+		PermutationVector.Set<FCalculateShadingRateImageCS::FThreadGroupY>(TileSize.Y);
+
+		// Set an override rate if we're in a debug mode
+		int32 ForceRate = GVRSDebugForceRate;
+		if (ForceRate != 0 &&
+			ForceRate != 1 &&
+			ForceRate != 4 &&
+			ForceRate != 5)
+		{
+			ForceRate = -1;
+		}
+
+		PermutationVector.Set<FCalculateShadingRateImageCS::FForceRate>(ForceRate);
+
+		TShaderMapRef<FCalculateShadingRateImageCS> ComputeShader(View.ShaderMap, PermutationVector);
+		auto* PassParameters = GraphBuilder.AllocParameters<FCalculateShadingRateImageCS::FParameters>();
+		EDisplayOutputFormat OutputDisplayFormat = GetDisplayOutputFormat(View);
+		FIntRect PostProcessRect = View.UnscaledViewRect.Scale(View.Family->SecondaryViewFraction);
+		FCalculateShadingRateImageCS::InitParameters(
+			*PassParameters,
+			Luminance,
+			PostProcessRect,
+			IsHDR10(OutputDisplayFormat),
+			GraphBuilder.CreateUAV(VRSTextures.ConstructedSRI));
+		FComputeShaderUtils::AddPass(
+			GraphBuilder,
+			RDG_EVENT_NAME("CreateShadingRateImage"),
+			ERDGPassFlags::AsyncCompute | ERDGPassFlags::NeverCull,
+			ComputeShader,
+			PassParameters,
+			FComputeShaderUtils::GetGroupCount(PostProcessRect.Size(), TileSize));
+	}
+
+	return true;
+}
+
+void AddPrepareImageBasedVRSPass(
+	FRDGBuilder& GraphBuilder,
+	const FMinimalSceneTextures& SceneTextures,
+	const FSceneViewFamily& ViewFamily)
+{
+	SCOPED_DRAW_EVENT(GraphBuilder.RHICmdList, ScaleVariableRateShadingTexture);
+
+	const FVRSTextures& VRSTextures = FVRSTextures::Get(GraphBuilder);
+	FRDGTextureRef VariableRateShadingImage = VRSTextures.ConstructedSRI;
+
+	FIntPoint TileSize = GetSRITileSize();
+
+	FIntPoint TextureSize = VRSTextures.ScaledSRI->Desc.Extent;
+	FVector2f TextureDimensions(TextureSize.X, TextureSize.Y);
+
+	for (int32 ViewIndex = 0; ViewIndex < ViewFamily.Views.Num(); ViewIndex++)
+	{
+		const FSceneView* View = ViewFamily.Views[ViewIndex];
+		RDG_EVENT_SCOPE_CONDITIONAL(GraphBuilder, ViewFamily.Views.Num() > 1, "View%d", ViewIndex);
+		check(View->bIsViewInfo);
+		const FViewInfo* ViewInfo = static_cast<const FViewInfo*>(View);
+		if (View->bCameraCut || !FVariableRateShadingImageManager::IsVRSCompatibleWithView(*ViewInfo) || !ViewInfo->PrevViewInfo.LuminanceHistory)
+		{
+			break;
+		}
+		FIntPoint SrcBufferSize = FSceneTexturesConfig::Get().Extent;
+
+		TShaderMapRef<FRescaleVariableRateShadingCS> RescaleVariableRateShadingCS(ViewInfo->ShaderMap);
+
+		int32 ViewportWidth = ViewInfo->ViewRect.Width();
+		int32 ViewportHeight = ViewInfo->ViewRect.Height();
+
+		int32 ScaledTilesWide = FMath::DivideAndRoundUp(ViewportWidth, TileSize.X);
+		int32 ScaledTilesHigh = FMath::DivideAndRoundUp(ViewportHeight, TileSize.Y);
+		FVector2f ScaledSRIDimensions(ScaledTilesWide, ScaledTilesHigh);
+
+		FIntRect PostProcessRect = GetPostProcessOutputRect(*ViewInfo);
+
+		FVector2f SRIViewRectMin(
+			static_cast<float>(FMath::DivideAndRoundDown(PostProcessRect.Min.X, TileSize.X)),
+			static_cast<float>(FMath::DivideAndRoundDown(PostProcessRect.Min.Y, TileSize.Y)));
+
+		FVector2f SRIViewRectMax(
+			static_cast<float>(FMath::DivideAndRoundUp(PostProcessRect.Max.X, TileSize.X)),
+			static_cast<float>(FMath::DivideAndRoundUp(PostProcessRect.Max.Y, TileSize.Y)));
+
+		FVector2f UVOffset(
+			(float)ViewInfo->ViewRect.Min.X / (float)SrcBufferSize.X,
+			(float)ViewInfo->ViewRect.Min.Y / (float)SrcBufferSize.Y);
+
+		float DynamicResolutionScale = (float)ViewportWidth / (float)PostProcessRect.Width();
+
+		auto PassParameters = GraphBuilder.AllocParameters<FRescaleVariableRateShadingCS::FParameters>();
+
+		FRescaleVariableRateShadingCS::InitParameters(
+			*PassParameters,
+			SceneTextures,
+			*ViewInfo,
+			VariableRateShadingImage,
+			SRIViewRectMin,
+			SRIViewRectMax,
+			ScaledSRIDimensions,
+			TextureDimensions,
+			UVOffset,
+			DynamicResolutionScale,
+			GraphBuilder.CreateUAV(VRSTextures.ScaledSRI),
+			GraphBuilder.CreateUAV(VRSTextures.ScaledConservativeSRI));
+
+		FComputeShaderUtils::AddPass(
+			GraphBuilder,
+			RDG_EVENT_NAME("PrepareImageBasedVRS"),
+			ERDGPassFlags::AsyncCompute | ERDGPassFlags::NeverCull,
+			RescaleVariableRateShadingCS,
+			PassParameters,
+			FComputeShaderUtils::GetGroupCount(FIntPoint(ScaledTilesWide, ScaledTilesHigh), FRescaleVariableRateShadingCS::ThreadGroupSize));
+	}
+}
+
+void FContrastAdaptiveImageGenerator::VRSDebugPreview(FRDGBuilder& GraphBuilder, const FSceneViewFamily& ViewFamily, FRDGTextureRef OutputSceneColor)
+{
+	RDG_EVENT_SCOPE(GraphBuilder, "VariableRateShading");
+	for (int32 ViewIndex = 0; ViewIndex < ViewFamily.Views.Num(); ViewIndex++)
+	{
+		check(ViewFamily.Views[ViewIndex]->bIsViewInfo);
+		const FViewInfo& View = *static_cast<const FViewInfo*>(ViewFamily.Views[ViewIndex]);
+		RDG_GPU_MASK_SCOPE(GraphBuilder, View.GPUMask);
+		RDG_EVENT_SCOPE_CONDITIONAL(GraphBuilder, ViewFamily.Views.Num() > 1, "View%d", ViewIndex);
+
+		//------------------------------------------------------------------------------------------------
+		// Do some sanity checks for early out
+		if (!OutputSceneColor)
+		{
+			return;
+		}
+
+		ESRIPreviewType::Type PreviewType = static_cast<ESRIPreviewType::Type>(CVarVRSPreview.GetValueOnRenderThread());
+
+		if (PreviewType == ESRIPreviewType::Off || !FVRSTextures::IsInitialized(GraphBuilder))
+		{
+			return;
+		}
+
+		const FVRSTextures& VRSTextures = FVRSTextures::Get(GraphBuilder);
+		FRDGTextureRef PreviewTexture = nullptr;
+		switch (PreviewType)
+		{
+		case ESRIPreviewType::BeforeReprojection:
+			PreviewTexture = VRSTextures.ConstructedSRI;
+			break;
+		case ESRIPreviewType::Projected:
+			PreviewTexture = VRSTextures.ScaledSRI;
+			break;
+		case ESRIPreviewType::ProjectedConservative:
+			PreviewTexture = VRSTextures.ScaledConservativeSRI;
+			break;
+		};
+
+		if (!PreviewTexture)
+		{
+			// We never rendered to this texture this frame, so aborting
+			return;
+		}
+		// Complete early out sanity checks
+		//------------------------------------------------------------------------------------------------
+
+		auto& RHICmdList = GraphBuilder.RHICmdList;
+
+		SCOPED_DRAW_EVENT(RHICmdList, VRSDebugPreview);
+
+		bool bUseRescaledTexture = PreviewType != ESRIPreviewType::BeforeReprojection;
+
+		FIntRect SrcViewRect = bUseRescaledTexture ? View.ViewRect : GetPostProcessOutputRect(View);
+
+		const FIntRect& DestViewRect = View.UnscaledViewRect;
+
+		TShaderMapRef<FDebugVariableRateShadingVS> VertexShader(View.ShaderMap);
+		TShaderMapRef<FDebugVariableRateShadingPS> PixelShader(View.ShaderMap);
+
+		auto* PassParameters = GraphBuilder.AllocParameters<FDebugVariableRateShadingPS::FParameters>();
+
+		FDebugVariableRateShadingPS::InitParameters(
+			*PassParameters,
+			PreviewTexture,
+			OutputSceneColor);
+
+		FRHIBlendState* BlendState =
+			TStaticBlendState<CW_RGB, BO_Add, BF_SourceAlpha, BF_InverseSource1Alpha>::GetRHI();
+		FRHIDepthStencilState* DepthStencilState = FScreenPassPipelineState::FDefaultDepthStencilState::GetRHI();
+
+		EScreenPassDrawFlags DrawFlags = EScreenPassDrawFlags::AllowHMDHiddenAreaMask;
+
+		FIntRect ScaledSrcRect = FIntRect::DivideAndRoundUp(SrcViewRect, GetSRITileSize());
+
+		const FScreenPassTextureViewport InputViewport = FScreenPassTextureViewport(PreviewTexture, ScaledSrcRect);
+		const FScreenPassTextureViewport OutputViewport(OutputSceneColor, DestViewRect);
+
+		AddDrawScreenPass(
+			GraphBuilder,
+			RDG_EVENT_NAME("Display Debug : %s", ESRIPreviewType::GetName(PreviewType)),
+			View,
+			OutputViewport,
+			InputViewport,
+			VertexShader,
+			PixelShader,
+			BlendState,
+			DepthStencilState,
+			PassParameters,
+			DrawFlags);
+	}
+}
+
+/**
+ * Interface Functions
+ */
+
+FRDGTextureRef FContrastAdaptiveImageGenerator::GetImage(FRDGBuilder& GraphBuilder, const FViewInfo& ViewInfo, FVariableRateShadingImageManager::EVRSPassType PassType)
+{
+	if (FVRSTextures::IsInitialized(GraphBuilder))
+	{
+		auto TextureType = ESRITextureType::GetTextureType(PassType);
+		const FVRSTextures& VRSTextures = FVRSTextures::Get(GraphBuilder);
+		return (TextureType == ESRITextureType::ScaledSRIForRender) ? VRSTextures.ScaledSRI : VRSTextures.ScaledConservativeSRI;
+	}
+
+	return nullptr;
+}
+
+void FContrastAdaptiveImageGenerator::PrepareImages(FRDGBuilder& GraphBuilder, const FSceneViewFamily& ViewFamily, const FMinimalSceneTextures& SceneTextures)
+{
+	RDG_EVENT_SCOPE(GraphBuilder, "VariableRateShading");
+	bool bIsAnyViewVRSCompatible = false;
+	for (const FSceneView* View : ViewFamily.Views)
+	{
+		check(View->bIsViewInfo);
+		auto ViewInfo = static_cast<const FViewInfo*>(View);
+		if (!View->bCameraCut && FVariableRateShadingImageManager::IsVRSCompatibleWithView(*ViewInfo) && ViewInfo->PrevViewInfo.LuminanceHistory)
+		{
+			bIsAnyViewVRSCompatible = true;
+		}
+	}
+	bool bPrepareImageBasedVRS = IsContrastAdaptiveShadingEnabled() && bIsAnyViewVRSCompatible;
+	if (!bPrepareImageBasedVRS)
+	{
+		return;
+	}
+
+	FVRSTextures& VRSTextures = GraphBuilder.Blackboard.Create<FVRSTextures>();
+	VRSTextures.Create(GraphBuilder, ViewFamily);
+
+	for (int32 ViewIndex = 0; ViewIndex < ViewFamily.Views.Num(); ViewIndex++)
+	{
+		check(ViewFamily.Views[ViewIndex]->bIsViewInfo);
+		const FViewInfo& View = *(FViewInfo*)ViewFamily.Views[ViewIndex];
+		RDG_GPU_MASK_SCOPE(GraphBuilder, View.GPUMask);
+		RDG_EVENT_SCOPE_CONDITIONAL(GraphBuilder, ViewFamily.Views.Num() > 1, "View%d", ViewIndex);
+		AddCreateShadingRateImagePass(GraphBuilder, View);
+	}
+	AddPrepareImageBasedVRSPass(GraphBuilder, SceneTextures, ViewFamily);
+}
+
+bool FContrastAdaptiveImageGenerator::IsEnabledForView(const FSceneView& View) const
+{
+	EDisplayOutputFormat DisplayOutputFormat = GetDisplayOutputFormat(View);
+	bool bCompatibleWithOutputType = (DisplayOutputFormat == EDisplayOutputFormat::SDR_sRGB) || IsHDR10(DisplayOutputFormat);
+
+	return IsContrastAdaptiveShadingEnabled() && !View.bIsSceneCapture && bCompatibleWithOutputType;
+}
+
diff --git a/Engine/Source/Runtime/Renderer/Private/VariableRateShading/ContrastAdaptiveImageGenerator.h b/Engine/Source/Runtime/Renderer/Private/VariableRateShading/ContrastAdaptiveImageGenerator.h
new file mode 100644
--- /dev/null
+++ b/Engine/Source/Runtime/Renderer/Private/VariableRateShading/ContrastAdaptiveImageGenerator.h
@@ -0,0 +1,26 @@
+// Copyright Epic Games, Inc. All Rights Reserved.
+
+#pragma once
+
+#include "CoreMinimal.h"
+#include "RHI.h"
+#include "RendererInterface.h"
+#include "RenderGraphDefinitions.h"
+#include "VariableRateShadingImageManager.h"
+#include "Engine/Engine.h"
+
+class FContrastAdaptiveImageGenerator : public IVariableRateShadingImageGenerator
+{
+public:
+	virtual ~FContrastAdaptiveImageGenerator() override {};
+	virtual FRDGTextureRef GetImage(FRDGBuilder& GraphBuilder, const FViewInfo& ViewInfo, FVariableRateShadingImageManager::EVRSPassType PassType) override;
+	virtual void PrepareImages(FRDGBuilder& GraphBuilder, const FSceneViewFamily& ViewFamily, const FMinimalSceneTextures& SceneTextures) override;
+	virtual bool IsEnabledForView(const FSceneView& View) const override;
+	virtual FVariableRateShadingImageManager::EVRSSourceType GetType() const override
+	{
+		return FVariableRateShadingImageManager::EVRSSourceType::ContrastAdaptiveShading;
+	}
+
+	void VRSDebugPreview(FRDGBuilder& GraphBuilder, const FSceneViewFamily& ViewFamily, FRDGTextureRef OutputSceneColor);
+};
+
diff --git a/Engine/Source/Runtime/Renderer/Private/VariableRateShading/VariableRateShadingImageManager.cpp b/Engine/Source/Runtime/Renderer/Private/VariableRateShading/VariableRateShadingImageManager.cpp
--- a/Engine/Source/Runtime/Renderer/Private/VariableRateShading/VariableRateShadingImageManager.cpp
+++ b/Engine/Source/Runtime/Renderer/Private/VariableRateShading/VariableRateShadingImageManager.cpp
@@ -2,6 +2,7 @@
 
 #include "VariableRateShadingImageManager.h"
 #include "FixedFoveationImageGenerator.h"
+#include "ContrastAdaptiveImageGenerator.h"
 #include "StereoRenderTargetManager.h"
 #include "GlobalShader.h"
 #include "ShaderParameterUtils.h"
@@ -24,8 +25,9 @@ FVariableRateShadingImageManager::FVariableRateShadingImageManager()
 	: FRenderResource()
 {
 	ImageGenerators.Add(MakeUnique<FFixedFoveationImageGenerator>());
+	ImageGenerators.Add(MakeUnique<FContrastAdaptiveImageGenerator>());
 
-	// TODO: Add more generators
+	// TODO: Add more generators or allow registration from plugins
 }
 
 FVariableRateShadingImageManager::~FVariableRateShadingImageManager() {}
@@ -41,12 +43,6 @@ static EDisplayOutputFormat GetDisplayOutputFormat(const FViewInfo& View)
 	return (EDisplayOutputFormat)Parameters.OutputDevice;
 }
 
-static bool IsHDR10(const EDisplayOutputFormat& OutputFormat)
-{
-	return OutputFormat == EDisplayOutputFormat::HDR_ACES_1000nit_ST2084 ||
-		OutputFormat == EDisplayOutputFormat::HDR_ACES_2000nit_ST2084;
-}
-
 bool FVariableRateShadingImageManager::IsVRSSupportedByRHI()
 {
 	return GRHISupportsAttachmentVariableRateShading && GRHIVariableRateShadingEnabled && GRHIAttachmentVariableRateShadingEnabled && FDataDrivenShaderPlatformInfo::GetSupportsVariableRateShading(GMaxRHIShaderPlatform);
@@ -54,7 +50,9 @@ bool FVariableRateShadingImageManager::IsVRSSupportedByRHI()
 
 bool FVariableRateShadingImageManager::IsVRSCompatibleWithOutputType(const EDisplayOutputFormat& OutputFormat)
 {
-	return OutputFormat == EDisplayOutputFormat::SDR_sRGB || IsHDR10(OutputFormat);
+	return OutputFormat == EDisplayOutputFormat::SDR_sRGB
+		|| OutputFormat == EDisplayOutputFormat::HDR_ACES_1000nit_ST2084
+		|| OutputFormat == EDisplayOutputFormat::HDR_ACES_2000nit_ST2084;
 }
 
 bool FVariableRateShadingImageManager::IsVRSCompatibleWithView(const FViewInfo& ViewInfo)
@@ -189,3 +187,17 @@ TRefCountPtr<IPooledRenderTarget> FVariableRateShadingImageManager::GetMobileVar
 
 	return MobileHMDFixedFoveationOverrideImage;
 }
+
+// Temporary passthrough for CAS debug overlay, pending that functionality being moved to the manager
+void FVariableRateShadingImageManager::CASDebugPreview(FRDGBuilder& GraphBuilder, const FSceneViewFamily& ViewFamily, FRDGTextureRef OutputSceneColor)
+{
+	// Get the first CAS generator present and call its debug function
+	for (TUniquePtr<IVariableRateShadingImageGenerator>& Generator : ImageGenerators)
+	{
+		if (EnumHasAllFlags(Generator->GetType(), FVariableRateShadingImageManager::EVRSSourceType::ContrastAdaptiveShading))
+		{
+			((FContrastAdaptiveImageGenerator*) Generator.Get())->VRSDebugPreview(GraphBuilder, ViewFamily, OutputSceneColor);
+			break;
+		}
+	}
+}
diff --git a/Engine/Source/Runtime/Renderer/Public/VariableRateShadingImageManager.h b/Engine/Source/Runtime/Renderer/Public/VariableRateShadingImageManager.h
--- a/Engine/Source/Runtime/Renderer/Public/VariableRateShadingImageManager.h
+++ b/Engine/Source/Runtime/Renderer/Public/VariableRateShadingImageManager.h
@@ -95,6 +95,8 @@ public:
 	static bool IsVRSCompatibleWithView(const FViewInfo& View);
 	static bool IsVRSCompatibleWithOutputType(const EDisplayOutputFormat& DisplayOutputFormat);
 
+	void CASDebugPreview(FRDGBuilder& GraphBuilder, const FSceneViewFamily& ViewFamily, FRDGTextureRef OutputSceneColor);
+
 private:
 	TRefCountPtr<IPooledRenderTarget> MobileHMDFixedFoveationOverrideImage;
 	TArray<TUniquePtr<IVariableRateShadingImageGenerator>> ImageGenerators;
