Made r.volumetriccloud scalable. It works and it is also used to load the blue noise texture on demand as a function of the current quality.

diff --git a/Engine/Source/Runtime/Renderer/Private/VolumetricCloudRendering.cpp b/Engine/Source/Runtime/Renderer/Private/VolumetricCloudRendering.cpp
--- a/Engine/Source/Runtime/Renderer/Private/VolumetricCloudRendering.cpp
+++ b/Engine/Source/Runtime/Renderer/Private/VolumetricCloudRendering.cpp
@@ -35,7 +35,7 @@
 static TAutoConsoleVariable<int32> CVarVolumetricCloud(
 	TEXT("r.VolumetricCloud"), 1,
 	TEXT("VolumetricCloud components are rendered when this is not 0, otherwise ignored."),
-	ECVF_RenderThreadSafe);
+	ECVF_RenderThreadSafe | ECVF_Scalability);
 
 static TAutoConsoleVariable<float> CVarVolumetricCloudDistanceToSampleMaxCount(
 	TEXT("r.VolumetricCloud.DistanceToSampleMaxCount"), 15.0f,
