Added CPU trace events to fill in gaps in the frame.

diff --git a/Engine/Source/Runtime/Renderer/Private/DistortionRendering.cpp b/Engine/Source/Runtime/Renderer/Private/DistortionRendering.cpp
--- a/Engine/Source/Runtime/Renderer/Private/DistortionRendering.cpp
+++ b/Engine/Source/Runtime/Renderer/Private/DistortionRendering.cpp
@@ -489,6 +489,7 @@ void FDeferredShadingSceneRenderer::RenderDistortion(
 	}
 
 	QUICK_SCOPE_CYCLE_COUNTER(STAT_FSceneRenderer_RenderDistortion);
+	SCOPED_NAMED_EVENT(RenderDistortion, FColor::Emerald);
 	RDG_EVENT_SCOPE(GraphBuilder, "Distortion");
 	RDG_GPU_STAT_SCOPE(GraphBuilder, Distortion);
 
diff --git a/Engine/Source/Runtime/Renderer/Private/HairStrands/HairStrandsComposition.cpp b/Engine/Source/Runtime/Renderer/Private/HairStrands/HairStrandsComposition.cpp
--- a/Engine/Source/Runtime/Renderer/Private/HairStrands/HairStrandsComposition.cpp
+++ b/Engine/Source/Runtime/Renderer/Private/HairStrands/HairStrandsComposition.cpp
@@ -478,6 +478,7 @@ static void InternalRenderHairComposition(
 	DECLARE_GPU_STAT(HairStrandsComposition);
 	RDG_EVENT_SCOPE(GraphBuilder, "HairStrandsComposition");
 	RDG_GPU_STAT_SCOPE(GraphBuilder, HairStrandsComposition);
+	SCOPED_NAMED_EVENT(HairStrandsComposition, FColor::Emerald);
 
 	{
 		{
diff --git a/Engine/Source/Runtime/Renderer/Private/HeterogeneousVolumes/HeterogeneousVolumes.cpp b/Engine/Source/Runtime/Renderer/Private/HeterogeneousVolumes/HeterogeneousVolumes.cpp
--- a/Engine/Source/Runtime/Renderer/Private/HeterogeneousVolumes/HeterogeneousVolumes.cpp
+++ b/Engine/Source/Runtime/Renderer/Private/HeterogeneousVolumes/HeterogeneousVolumes.cpp
@@ -332,6 +332,7 @@ void FDeferredShadingSceneRenderer::RenderHeterogeneousVolumes(
 {
 	RDG_EVENT_SCOPE(GraphBuilder, "HeterogeneousVolumes");
 	RDG_GPU_STAT_SCOPE(GraphBuilder, HeterogeneousVolumesStat);
+	SCOPED_NAMED_EVENT(HeterogeneousVolumes, FColor::Emerald);
 
 	for (int32 ViewIndex = 0; ViewIndex < Views.Num(); ViewIndex++)
 	{
diff --git a/Engine/Source/Runtime/Renderer/Private/IndirectLightRendering.cpp b/Engine/Source/Runtime/Renderer/Private/IndirectLightRendering.cpp
--- a/Engine/Source/Runtime/Renderer/Private/IndirectLightRendering.cpp
+++ b/Engine/Source/Runtime/Renderer/Private/IndirectLightRendering.cpp
@@ -812,6 +812,7 @@ void FDeferredShadingSceneRenderer::DispatchAsyncLumenIndirectLightingWork(
 
 	LLM_SCOPE_BYTAG(Lumen);
 	RDG_EVENT_SCOPE(GraphBuilder, "DiffuseIndirectAndAO");
+	SCOPED_NAMED_EVENT(DispatchAsyncLumenIndirectLightingWork, FColor::Emerald);
 	
 	Outputs.DoneAsync(GLumenReflectionsAsyncCompute != 0);
 	Outputs.Resize(Views.Num());
diff --git a/Engine/Source/Runtime/Renderer/Private/SingleLayerWaterRendering.cpp b/Engine/Source/Runtime/Renderer/Private/SingleLayerWaterRendering.cpp
--- a/Engine/Source/Runtime/Renderer/Private/SingleLayerWaterRendering.cpp
+++ b/Engine/Source/Runtime/Renderer/Private/SingleLayerWaterRendering.cpp
@@ -1244,6 +1244,7 @@ void FDeferredShadingSceneRenderer::RenderSingleLayerWater(
 {
 	RDG_EVENT_SCOPE(GraphBuilder, "SingleLayerWater");
 	RDG_GPU_STAT_SCOPE(GraphBuilder, SingleLayerWater);
+	SCOPED_NAMED_EVENT(SingleLayerWater, FColor::Emerald);
 
 	// Copy the texture to be available for the water surface to refract
 	SceneWithoutWaterTextures = AddCopySceneWithoutWaterPass(GraphBuilder, ViewFamily, Views, SceneTextures.Color.Resolve, SceneTextures.Depth.Resolve);
diff --git a/Engine/Source/Runtime/Renderer/Private/SkyAtmosphereRendering.cpp b/Engine/Source/Runtime/Renderer/Private/SkyAtmosphereRendering.cpp
--- a/Engine/Source/Runtime/Renderer/Private/SkyAtmosphereRendering.cpp
+++ b/Engine/Source/Runtime/Renderer/Private/SkyAtmosphereRendering.cpp
@@ -1897,6 +1897,7 @@ void FSceneRenderer::RenderSkyAtmosphere(FRDGBuilder& GraphBuilder, const FMinim
 
 	RDG_EVENT_SCOPE(GraphBuilder, "SkyAtmosphere");
 	RDG_GPU_STAT_SCOPE(GraphBuilder, SkyAtmosphere);
+	SCOPED_NAMED_EVENT(SkyAtmosphere, FColor::Emerald);
 
 	FSkyAtmosphereRenderSceneInfo& SkyInfo = *Scene->GetSkyAtmosphereSceneInfo();
 	const FSkyAtmosphereSceneProxy& SkyAtmosphereSceneProxy = SkyInfo.GetSkyAtmosphereSceneProxy();
diff --git a/Engine/Source/Runtime/Renderer/Private/VirtualShadowMaps/VirtualShadowMapArray.cpp b/Engine/Source/Runtime/Renderer/Private/VirtualShadowMaps/VirtualShadowMapArray.cpp
--- a/Engine/Source/Runtime/Renderer/Private/VirtualShadowMaps/VirtualShadowMapArray.cpp
+++ b/Engine/Source/Runtime/Renderer/Private/VirtualShadowMaps/VirtualShadowMapArray.cpp
@@ -1190,6 +1190,7 @@ void FVirtualShadowMapArray::BuildPageAllocations(
 	}
 
 	RDG_EVENT_SCOPE(GraphBuilder, "FVirtualShadowMapArray::BuildPageAllocation");
+	SCOPED_NAMED_EVENT(FVirtualShadowMapArray_BuildPageAllocation, FColor::Emerald);
 
 	{
 		// Buffer for marking static invalidating instances for readback
diff --git a/Engine/Source/Runtime/Renderer/Private/VolumetricCloudRendering.cpp b/Engine/Source/Runtime/Renderer/Private/VolumetricCloudRendering.cpp
--- a/Engine/Source/Runtime/Renderer/Private/VolumetricCloudRendering.cpp
+++ b/Engine/Source/Runtime/Renderer/Private/VolumetricCloudRendering.cpp
@@ -2578,6 +2578,7 @@ bool FSceneRenderer::RenderVolumetricCloud(
 		{
 			RDG_EVENT_SCOPE(GraphBuilder, "VolumetricCloud");
 			RDG_GPU_STAT_SCOPE(GraphBuilder, VolumetricCloud);
+			SCOPED_NAMED_EVENT(VolumetricCloud, FColor::Emerald);
 
 			FRDGTextureRef SceneDepthZ = SceneTextures.Depth.Resolve;
 
diff --git a/Engine/Source/Runtime/Renderer/Public/SimpleMeshDrawCommandPass.h b/Engine/Source/Runtime/Renderer/Public/SimpleMeshDrawCommandPass.h
--- a/Engine/Source/Runtime/Renderer/Public/SimpleMeshDrawCommandPass.h
+++ b/Engine/Source/Runtime/Renderer/Public/SimpleMeshDrawCommandPass.h
@@ -68,6 +68,8 @@ void AddSimpleMeshPass(FRDGBuilder& GraphBuilder, PassParametersType* PassParame
 	AddMeshBatchesCallbackLambdaType AddMeshBatchesCallback,
 	PassPrologueLambdaType PassPrologueCallback)
 {
+	TRACE_CPUPROFILER_EVENT_SCOPE(AddSimpleMeshPass);
+
 	// TODO: don't do this when the parameters are exclusive (optimization)
 	PassParametersType* PassParameters = GraphBuilder.AllocParameters<PassParametersType>();
 
