Support combining multiple VRS images within the VRS image manager.

For now this can be used to run CAS and fixed foveation at the same time.

diff --git a/Engine/Shaders/Private/VariableRateShading/VRSShadingRateCombine.usf b/Engine/Shaders/Private/VariableRateShading/VRSShadingRateCombine.usf
new file mode 100644
--- /dev/null
+++ b/Engine/Shaders/Private/VariableRateShading/VRSShadingRateCombine.usf
@@ -0,0 +1,19 @@
+// Copyright Epic Games, Inc. All Rights Reserved.
+
+#include "../Common.ush"
+
+RWTexture2D<uint> RWOutputTexture;
+Texture2D<uint> SourceTexture_1;
+Texture2D<uint> SourceTexture_2;
+
+
+[numthreads(THREADGROUP_SIZEX, THREADGROUP_SIZEY, 1)]
+void CombineShadingRateTextures(uint3 DispatchThreadId : SV_DispatchThreadID)
+{
+	const uint2 TexelCoord = DispatchThreadId.xy;
+	uint ShadingRateOut = 0;
+	
+	ShadingRateOut = max(SourceTexture_1[TexelCoord], SourceTexture_2[TexelCoord]);
+	
+	RWOutputTexture[TexelCoord] = ShadingRateOut;
+}
\ No newline at end of file
diff --git a/Engine/Source/Runtime/Renderer/Private/VariableRateShading/ContrastAdaptiveImageGenerator.cpp b/Engine/Source/Runtime/Renderer/Private/VariableRateShading/ContrastAdaptiveImageGenerator.cpp
--- a/Engine/Source/Runtime/Renderer/Private/VariableRateShading/ContrastAdaptiveImageGenerator.cpp
+++ b/Engine/Source/Runtime/Renderer/Private/VariableRateShading/ContrastAdaptiveImageGenerator.cpp
@@ -423,11 +423,6 @@ namespace ESRIPreviewType
 	}
 };
 
-static FIntPoint GetSRITileSize()
-{
-	return FIntPoint(GRHIVariableRateShadingImageTileMinWidth, GRHIVariableRateShadingImageTileMinHeight);
-}
-
 static const TCHAR* ShadingRateTextureName = TEXT("ShadingRateTexture");
 static const TCHAR* ScaledShadingRateTextureName = TEXT("ScaledShadingRateTexture");
 static const TCHAR* ScaledConservativeShadingRateTextureName = TEXT("ConservativeScaledShadingRateTexture");
@@ -460,7 +455,7 @@ struct RENDERER_API FVRSTextures
 private:
 	static FRDGTextureDesc CreateSRIDesc(const FSceneViewFamily& ViewFamily, bool bIsForDynResScaled)
 	{
-		FIntPoint TileSize = GetSRITileSize();
+		FIntPoint TileSize = FVariableRateShadingImageManager::GetSRITileSize();
 		FIntPoint ViewTargetExtents = (bIsForDynResScaled)
 			? FSceneTexturesConfig::Get().Extent
 			: ViewFamily.RenderTarget->GetSizeXY();
@@ -519,7 +514,7 @@ bool AddCreateShadingRateImagePass(
 	{
 		FCalculateShadingRateImageCS::FPermutationDomain PermutationVector;
 
-		const FIntPoint TileSize = GetSRITileSize();
+		const FIntPoint TileSize = FVariableRateShadingImageManager::GetSRITileSize();
 		PermutationVector.Set<FCalculateShadingRateImageCS::FThreadGroupX>(TileSize.X);
 		PermutationVector.Set<FCalculateShadingRateImageCS::FThreadGroupY>(TileSize.Y);
 
@@ -567,7 +562,7 @@ void AddPrepareImageBasedVRSPass(
 	const FVRSTextures& VRSTextures = FVRSTextures::Get(GraphBuilder);
 	FRDGTextureRef VariableRateShadingImage = VRSTextures.ConstructedSRI;
 
-	FIntPoint TileSize = GetSRITileSize();
+	FIntPoint TileSize = FVariableRateShadingImageManager::GetSRITileSize();
 
 	FIntPoint TextureSize = VRSTextures.ScaledSRI->Desc.Extent;
 	FVector2f TextureDimensions(TextureSize.X, TextureSize.Y);
@@ -708,7 +703,7 @@ void FContrastAdaptiveImageGenerator::VRSDebugPreview(FRDGBuilder& GraphBuilder,
 
 		EScreenPassDrawFlags DrawFlags = EScreenPassDrawFlags::AllowHMDHiddenAreaMask;
 
-		FIntRect ScaledSrcRect = FIntRect::DivideAndRoundUp(SrcViewRect, GetSRITileSize());
+		FIntRect ScaledSrcRect = FIntRect::DivideAndRoundUp(SrcViewRect, FVariableRateShadingImageManager::GetSRITileSize());
 
 		const FScreenPassTextureViewport InputViewport = FScreenPassTextureViewport(PreviewTexture, ScaledSrcRect);
 		const FScreenPassTextureViewport OutputViewport(OutputSceneColor, DestViewRect);
diff --git a/Engine/Source/Runtime/Renderer/Private/VariableRateShading/VariableRateShadingImageManager.cpp b/Engine/Source/Runtime/Renderer/Private/VariableRateShading/VariableRateShadingImageManager.cpp
--- a/Engine/Source/Runtime/Renderer/Private/VariableRateShading/VariableRateShadingImageManager.cpp
+++ b/Engine/Source/Runtime/Renderer/Private/VariableRateShading/VariableRateShadingImageManager.cpp
@@ -21,6 +21,48 @@
 
 TGlobalResource<FVariableRateShadingImageManager> GVRSImageManager;
 
+/**
+ * Shaders
+ */
+
+constexpr int32 kCombineGroupSize = FComputeShaderUtils::kGolden2DGroupSize;
+
+class FCombineShadingRateTexturesCS : public FGlobalShader
+{
+public:
+	DECLARE_GLOBAL_SHADER(FCombineShadingRateTexturesCS);
+	SHADER_USE_PARAMETER_STRUCT(FCombineShadingRateTexturesCS, FGlobalShader);
+
+	BEGIN_SHADER_PARAMETER_STRUCT(FParameters, )
+		SHADER_PARAMETER_RDG_TEXTURE_UAV(RWTexture2D<uint>, RWOutputTexture)
+		SHADER_PARAMETER_RDG_TEXTURE(Texture2D<uint>, SourceTexture_1)
+		SHADER_PARAMETER_RDG_TEXTURE(Texture2D<uint>, SourceTexture_2)
+		SHADER_PARAMETER_RDG_TEXTURE(Texture2D<uint>, SourceTexture_3)
+		SHADER_PARAMETER_RDG_TEXTURE(Texture2D<uint>, SourceTexture_4)
+	END_SHADER_PARAMETER_STRUCT()
+
+	static bool ShouldCompilePermutation(const FGlobalShaderPermutationParameters& Parameters)
+	{
+		return FDataDrivenShaderPlatformInfo::GetSupportsVariableRateShading(Parameters.Platform);
+	}
+
+	static void ModifyCompilationEnvironment(const FGlobalShaderPermutationParameters& Parameters, FShaderCompilerEnvironment& OutEnvironment)
+	{
+		FGlobalShader::ModifyCompilationEnvironment(Parameters, OutEnvironment);
+
+		OutEnvironment.SetDefine(TEXT("THREADGROUP_SIZEX"), kCombineGroupSize);
+		OutEnvironment.SetDefine(TEXT("THREADGROUP_SIZEY"), kCombineGroupSize);
+	}
+
+};
+
+IMPLEMENT_GLOBAL_SHADER(FCombineShadingRateTexturesCS, "/Engine/Private/VariableRateShading/VRSShadingRateCombine.usf", "CombineShadingRateTextures", SF_Compute);
+
+
+/**
+ * Public functions
+ */
+
 FVariableRateShadingImageManager::FVariableRateShadingImageManager()
 	: FRenderResource()
 {
@@ -65,6 +107,11 @@ bool FVariableRateShadingImageManager::IsVRSCompatibleWithView(const FViewInfo&
 	return IsVRSSupportedByRHI() && !ViewInfo.bIsSceneCapture && IsVRSCompatibleWithOutputType(GetDisplayOutputFormat(ViewInfo));
 }
 
+FIntPoint FVariableRateShadingImageManager::GetSRITileSize()
+{
+	return FIntPoint(GRHIVariableRateShadingImageTileMinWidth, GRHIVariableRateShadingImageTileMinHeight);
+}
+
 FRDGTextureRef FVariableRateShadingImageManager::GetVariableRateShadingImage(FRDGBuilder& GraphBuilder, const FViewInfo& ViewInfo, FVariableRateShadingImageManager::EVRSPassType PassType,
 	const TArray<TRefCountPtr<IPooledRenderTarget>>* ExternalVRSSources, FVariableRateShadingImageManager::EVRSSourceType VRSTypesToExclude)
 {
@@ -88,19 +135,49 @@ FRDGTextureRef FVariableRateShadingImageManager::GetVariableRateShadingImage(FRD
 		if (Image)
 		{
 			InternalVRSSources.Add(Image);
-			break; // No need to generate more than one until we support combining
 		}
 	}
 
-	// TODO: Combine internal and external images with adjustable logic (min, max, priority, etc.)
+	// If we have more than one internal source, combine the first available two
+	// TODO: Support combining more textures
+	if (InternalVRSSources.Num() > 1)
+	{
+		FIntPoint ViewSize = ViewInfo.UnscaledViewRect.Scale(ViewInfo.Family->SecondaryViewFraction).Size();
+		const FIntPoint TileSize = GetSRITileSize();
+
+		// Create texture to hold shading rate image
+		FRDGTextureDesc Desc = FRDGTextureDesc::Create2D(
+			ViewSize / TileSize,
+			GRHIVariableRateShadingImageFormat,
+			FClearValueBinding::None,
+			TexCreate_Foveation | TexCreate_UAV);
 
-	// For now, only one source can be used
-	// Once CAS is added, source will be chosen by CVar, for now it defaults to the only available generator (fixed foveation)
-	// We fall back on the first available external source if no internal source is available
-	if (InternalVRSSources.Num() > 0)
+		FRDGTextureRef CombinedShadingRateTexture = GraphBuilder.CreateTexture(Desc, TEXT("CombinedShadingRateTexture"));
+
+		FCombineShadingRateTexturesCS::FParameters* PassParameters = GraphBuilder.AllocParameters<FCombineShadingRateTexturesCS::FParameters>();
+		PassParameters->SourceTexture_2 = InternalVRSSources[0];
+		PassParameters->SourceTexture_1 = InternalVRSSources[1];
+		PassParameters->RWOutputTexture = GraphBuilder.CreateUAV(CombinedShadingRateTexture);
+
+		TShaderMapRef<FCombineShadingRateTexturesCS> ComputeShader(GetGlobalShaderMap(GMaxRHIFeatureLevel));
+
+		FComputeShaderUtils::AddPass(
+			GraphBuilder,
+			RDG_EVENT_NAME("CombineShadingRateImages"),
+			ERDGPassFlags::AsyncCompute | ERDGPassFlags::NeverCull,
+			ComputeShader,
+			PassParameters,
+			FComputeShaderUtils::GetGroupCount(ViewSize, TileSize));
+
+		return CombinedShadingRateTexture;
+	}
+	else if (InternalVRSSources.Num() == 1)
 	{
 		return InternalVRSSources[0];
 	}
+
+	// Fall back on external sources only if we have no internal ones
+	// TODO: Combine external sources as well
 	else if (ExternalVRSSources && ExternalVRSSources->Num() > 0)
 	{
 		return GraphBuilder.RegisterExternalTexture((*ExternalVRSSources)[0]);
diff --git a/Engine/Source/Runtime/Renderer/Public/VariableRateShadingImageManager.h b/Engine/Source/Runtime/Renderer/Public/VariableRateShadingImageManager.h
--- a/Engine/Source/Runtime/Renderer/Public/VariableRateShadingImageManager.h
+++ b/Engine/Source/Runtime/Renderer/Public/VariableRateShadingImageManager.h
@@ -95,6 +95,8 @@ public:
 	static bool IsVRSCompatibleWithView(const FViewInfo& View);
 	static bool IsVRSCompatibleWithOutputType(const EDisplayOutputFormat& DisplayOutputFormat);
 
+	static FIntPoint GetSRITileSize();
+
 	void CASDebugPreview(FRDGBuilder& GraphBuilder, const FSceneViewFamily& ViewFamily, FRDGTextureRef OutputSceneColor);
 
 private:
