Added kCombineGroupSize to namespace.

diff --git a/Engine/Source/Runtime/Renderer/Private/VariableRateShading/VariableRateShadingImageManager.cpp b/Engine/Source/Runtime/Renderer/Private/VariableRateShading/VariableRateShadingImageManager.cpp
--- a/Engine/Source/Runtime/Renderer/Private/VariableRateShading/VariableRateShadingImageManager.cpp
+++ b/Engine/Source/Runtime/Renderer/Private/VariableRateShading/VariableRateShadingImageManager.cpp
@@ -25,7 +25,10 @@ TGlobalResource<FVariableRateShadingImageManager> GVRSImageManager;
  * Shaders
  */
 
-constexpr int32 kCombineGroupSize = FComputeShaderUtils::kGolden2DGroupSize;
+namespace VRSHelpers
+{
+	constexpr int32 kCombineGroupSize = FComputeShaderUtils::kGolden2DGroupSize;
+}
 
 class FCombineShadingRateTexturesCS : public FGlobalShader
 {
@@ -50,8 +53,8 @@ public:
 	{
 		FGlobalShader::ModifyCompilationEnvironment(Parameters, OutEnvironment);
 
-		OutEnvironment.SetDefine(TEXT("THREADGROUP_SIZEX"), kCombineGroupSize);
-		OutEnvironment.SetDefine(TEXT("THREADGROUP_SIZEY"), kCombineGroupSize);
+		OutEnvironment.SetDefine(TEXT("THREADGROUP_SIZEX"), VRSHelpers::kCombineGroupSize);
+		OutEnvironment.SetDefine(TEXT("THREADGROUP_SIZEY"), VRSHelpers::kCombineGroupSize);
 	}
 
 };
