Moved ray tracing relevancy and BeginInitShadows to run in parallel with GatherDynamicMeshElements.

diff --git a/Engine/Source/Runtime/Renderer/Private/DeferredShadingRenderer.cpp b/Engine/Source/Runtime/Renderer/Private/DeferredShadingRenderer.cpp
--- a/Engine/Source/Runtime/Renderer/Private/DeferredShadingRenderer.cpp
+++ b/Engine/Source/Runtime/Renderer/Private/DeferredShadingRenderer.cpp
@@ -913,6 +913,46 @@ static void GatherRayTracingRelevantPrimitives(const FScene& Scene, const FViewI
 	Result.bValid = true;
 }
 
+#if RHI_RAYTRACING
+struct FRayTracingRelevantPrimitiveTaskData
+{
+	FRayTracingRelevantPrimitiveList List;
+	FGraphEventRef Task;
+};
+#endif // RHI_RAYTRACING
+
+void FDeferredShadingSceneRenderer::PreGatherDynamicMeshElements()
+{
+#if RHI_RAYTRACING
+	if (bAnyRayTracingPassEnabled)
+	{
+		const int32 ReferenceViewIndex = 0;
+		FViewInfo& ReferenceView = Views[ReferenceViewIndex];
+
+		RayTracingRelevantPrimitiveTaskData = Allocator.Create<FRayTracingRelevantPrimitiveTaskData>();
+		RayTracingRelevantPrimitiveTaskData->Task = FFunctionGraphTask::CreateAndDispatchWhenReady(
+			[Scene = this->Scene, &ReferenceView, &RayTracingRelevantPrimitiveList = RayTracingRelevantPrimitiveTaskData->List]()
+		{
+			FTaskTagScope TaskTagScope(ETaskTag::EParallelRenderingThread);
+			GatherRayTracingRelevantPrimitives(*Scene, ReferenceView, RayTracingRelevantPrimitiveList);
+		}, TStatId(), nullptr, ENamedThreads::AnyNormalThreadHiPriTask);
+	}
+#endif // RHI_RAYTRACING
+
+	const bool bHasRayTracedOverlay = HasRayTracedOverlay(ViewFamily);
+
+	extern int32 GEarlyInitDynamicShadows;
+
+	if (GEarlyInitDynamicShadows &&
+		CurrentDynamicShadowsTaskData == nullptr &&
+		ViewFamily.EngineShowFlags.DynamicShadows
+		&& !ViewFamily.EngineShowFlags.HitProxies
+		&& !bHasRayTracedOverlay)
+	{
+		CurrentDynamicShadowsTaskData = BeginInitDynamicShadows(true);
+	}
+}
+
 bool FDeferredShadingSceneRenderer::GatherRayTracingWorldInstancesForView(FRDGBuilder& GraphBuilder, FViewInfo& View, FRayTracingScene& RayTracingScene, FRayTracingRelevantPrimitiveList& RelevantPrimitiveList)
 {
 	checkf(IsRayTracingEnabled() && bAnyRayTracingPassEnabled, TEXT("GatherRayTracingWorldInstancesForView should only be called if ray tracing is used"))
@@ -2538,19 +2578,6 @@ void FDeferredShadingSceneRenderer::Render(FRDGBuilder& GraphBuilder)
 			}
 		}
 	}
-
-	// Asynchronously create a list of primitives relevant to ray tracing scene
-	FRayTracingRelevantPrimitiveList RayTracingRelevantPrimitiveList;
-	FGraphEventRef RayTracingRelevantPrimitiveListTask;
-	if (bAnyRayTracingPassEnabled)
-	{
-		RayTracingRelevantPrimitiveListTask = FFunctionGraphTask::CreateAndDispatchWhenReady(
-			[Scene = this->Scene, &ReferenceView, &RayTracingRelevantPrimitiveList]()
-		{
-			FTaskTagScope TaskTagScope(ETaskTag::EParallelRenderingThread);
-			GatherRayTracingRelevantPrimitives(*Scene, ReferenceView, RayTracingRelevantPrimitiveList);
-		}, TStatId(), nullptr, ENamedThreads::AnyNormalThreadHiPriTask);
-	}
 #endif // RHI_RAYTRACING
 
 	// Notify the FX system that the scene is about to be rendered.
@@ -2613,16 +2640,16 @@ void FDeferredShadingSceneRenderer::Render(FRDGBuilder& GraphBuilder)
 	if (bAnyRayTracingPassEnabled)
 	{
 		// Wait until RayTracingRelevantPrimitiveList is ready
-		if (RayTracingRelevantPrimitiveListTask.IsValid())
+		if (RayTracingRelevantPrimitiveTaskData->Task.IsValid())
 		{
-			RayTracingRelevantPrimitiveListTask->Wait();
-			RayTracingRelevantPrimitiveListTask.SafeRelease();
+			RayTracingRelevantPrimitiveTaskData->Task->Wait();
+			RayTracingRelevantPrimitiveTaskData->Task.SafeRelease();
 		}
 
 		// Prepare ray tracing scene instance list
-		checkf(RayTracingRelevantPrimitiveList.bValid, TEXT("Ray tracing relevant primitive list is expected to have been created before GatherRayTracingWorldInstancesForView() is called."));
+		checkf(RayTracingRelevantPrimitiveTaskData->List.bValid, TEXT("Ray tracing relevant primitive list is expected to have been created before GatherRayTracingWorldInstancesForView() is called."));
 
-		GatherRayTracingWorldInstancesForView(GraphBuilder, ReferenceView, RayTracingScene, RayTracingRelevantPrimitiveList);
+		GatherRayTracingWorldInstancesForView(GraphBuilder, ReferenceView, RayTracingScene, RayTracingRelevantPrimitiveTaskData->List);
 	}
 #endif // RHI_RAYTRACING
 
diff --git a/Engine/Source/Runtime/Renderer/Private/DeferredShadingRenderer.h b/Engine/Source/Runtime/Renderer/Private/DeferredShadingRenderer.h
--- a/Engine/Source/Runtime/Renderer/Private/DeferredShadingRenderer.h
+++ b/Engine/Source/Runtime/Renderer/Private/DeferredShadingRenderer.h
@@ -51,6 +51,10 @@ struct FLumenSceneFrameTemporaries;
 struct FSingleLayerWaterPrePassResult;
 struct FBuildHZBAsyncComputeParams;
 
+#if RHI_RAYTRACING
+struct FRayTracingRelevantPrimitiveTaskData;
+#endif
+
 /**   
  * Data for rendering meshes into Lumen Lighting Cards.
  */
@@ -516,6 +520,8 @@ private:
 		const TArray<FProjectedShadowInfo*, SceneRenderingAllocator>& ViewDependentWholeSceneShadows,
 		TArray<FProjectedShadowInfo*, SceneRenderingAllocator>& OutPreShadows);
 
+	void PreGatherDynamicMeshElements() override final;
+
 	void PreVisibilityFrameSetup(FRDGBuilder& GraphBuilder, const FSceneTexturesConfig& SceneTexturesConfig);
 
 	void ComputeLightVisibility();
@@ -523,7 +529,6 @@ private:
 	/** Determines which primitives are visible for each view. */
 	void InitViews(FRDGBuilder& GraphBuilder, const FSceneTexturesConfig& SceneTexturesConfig, FExclusiveDepthStencil::Type BasePassDepthStencilAccess, struct FILCUpdatePrimTaskData& ILCTaskData, FInstanceCullingManager& InstanceCullingManager);
 
-	void InitViewsBeforePrepass(FRDGBuilder& GraphBuilder, FInstanceCullingManager& InstanceCullingManager);
 	void InitViewsAfterPrepass(FRDGBuilder& GraphBuilder, FLumenSceneFrameTemporaries& FrameTemporaries, struct FILCUpdatePrimTaskData& ILCTaskData, FInstanceCullingManager& InstanceCullingManager);
 	void BeginUpdateLumenSceneTasks(FRDGBuilder& GraphBuilder, FLumenSceneFrameTemporaries& FrameTemporaries);
 	void UpdateLumenScene(FRDGBuilder& GraphBuilder, FLumenSceneFrameTemporaries& FrameTemporaries);
@@ -1167,6 +1172,10 @@ private:
 	bool bAreLightsInLightGrid;
 
 	FDynamicShadowsTaskData* CurrentDynamicShadowsTaskData{};
+
+#if RHI_RAYTRACING
+	FRayTracingRelevantPrimitiveTaskData* RayTracingRelevantPrimitiveTaskData = nullptr;
+#endif
 };
 
 DECLARE_CYCLE_STAT_EXTERN(TEXT("PrePass"), STAT_CLM_PrePass, STATGROUP_CommandListMarkers, );
\ No newline at end of file
diff --git a/Engine/Source/Runtime/Renderer/Private/SceneRendering.h b/Engine/Source/Runtime/Renderer/Private/SceneRendering.h
--- a/Engine/Source/Runtime/Renderer/Private/SceneRendering.h
+++ b/Engine/Source/Runtime/Renderer/Private/SceneRendering.h
@@ -2366,6 +2366,8 @@ protected:
 	/** Performs once per frame setup after to visibility determination. */
 	void PostVisibilityFrameSetup(FILCUpdatePrimTaskData& OutILCTaskData);
 
+	virtual void PreGatherDynamicMeshElements() {}
+
 	void GatherDynamicMeshElements(
 		TArray<FViewInfo>& InViews, 
 		const FScene* InScene, 
diff --git a/Engine/Source/Runtime/Renderer/Private/SceneVisibility.cpp b/Engine/Source/Runtime/Renderer/Private/SceneVisibility.cpp
--- a/Engine/Source/Runtime/Renderer/Private/SceneVisibility.cpp
+++ b/Engine/Source/Runtime/Renderer/Private/SceneVisibility.cpp
@@ -174,7 +174,7 @@ static FAutoConsoleVariableRef CVarILCUpdatePrimitivesTask(
 	ECVF_RenderThreadSafe
 	);
 
-static int32 GEarlyInitDynamicShadows = 1;
+int32 GEarlyInitDynamicShadows = 1;
 static FAutoConsoleVariableRef CVarEarlyInitDynamicShadows(
 	TEXT("r.EarlyInitDynamicShadows"),
 	GEarlyInitDynamicShadows,
@@ -4913,6 +4913,10 @@ void FSceneRenderer::ComputeViewVisibility(
 		}
 	}
 
+	ComputeLightVisibilityTask.Wait();
+
+	PreGatherDynamicMeshElements();
+
 	{
 		SCOPED_NAMED_EVENT(FSceneRenderer_GatherDynamicMeshElements, FColor::Yellow);
 		// Gather FMeshBatches from scene proxies
@@ -4947,8 +4951,6 @@ void FSceneRenderer::ComputeViewVisibility(
 	INC_DWORD_STAT_BY(STAT_ProcessedPrimitives,NumProcessedPrimitives);
 	INC_DWORD_STAT_BY(STAT_CulledPrimitives,NumCulledPrimitives);
 	INC_DWORD_STAT_BY(STAT_OccludedPrimitives,NumOccludedPrimitives);
-	
-	ComputeLightVisibilityTask.Wait();
 }
 
 void FDeferredShadingSceneRenderer::ComputeLightVisibility()
@@ -5257,7 +5259,6 @@ void FDeferredShadingSceneRenderer::InitViews(FRDGBuilder& GraphBuilder, const F
 	}
 
 	PostVisibilityFrameSetup(ILCTaskData);
-	InitViewsBeforePrepass(GraphBuilder, InstanceCullingManager);
 
 	{
 		QUICK_SCOPE_CYCLE_COUNTER(STAT_InitViews_InitRHIResources);
@@ -5359,20 +5360,6 @@ void FSceneRenderer::SetupSceneReflectionCaptureBuffer(FRHICommandListImmediate&
 	}
 }
 
-void FDeferredShadingSceneRenderer::InitViewsBeforePrepass(FRDGBuilder& GraphBuilder, FInstanceCullingManager& InstanceCullingManager)
-{
-	const bool bHasRayTracedOverlay = HasRayTracedOverlay(ViewFamily);
-
-	if (GEarlyInitDynamicShadows &&
-		CurrentDynamicShadowsTaskData == nullptr &&
-		ViewFamily.EngineShowFlags.DynamicShadows
-		&& !ViewFamily.EngineShowFlags.HitProxies
-		&& !bHasRayTracedOverlay)
-	{
-		CurrentDynamicShadowsTaskData = BeginInitDynamicShadows(true);
-	}
-}
-
 void FDeferredShadingSceneRenderer::InitViewsAfterPrepass(FRDGBuilder& GraphBuilder, FLumenSceneFrameTemporaries& FrameTemporaries, struct FILCUpdatePrimTaskData& ILCTaskData, FInstanceCullingManager& InstanceCullingManager)
 {
 	SCOPED_NAMED_EVENT(FDeferredShadingSceneRenderer_InitViewsAfterPrepass, FColor::Emerald);
