Fixing build for non-Raytracing platforms.

diff --git a/Engine/Source/Runtime/Renderer/Private/DeferredShadingRenderer.cpp b/Engine/Source/Runtime/Renderer/Private/DeferredShadingRenderer.cpp
--- a/Engine/Source/Runtime/Renderer/Private/DeferredShadingRenderer.cpp
+++ b/Engine/Source/Runtime/Renderer/Private/DeferredShadingRenderer.cpp
@@ -913,46 +913,6 @@
 	Result.bValid = true;
 }
 
-#if RHI_RAYTRACING
-struct FRayTracingRelevantPrimitiveTaskData
-{
-	FRayTracingRelevantPrimitiveList List;
-	FGraphEventRef Task;
-};
-#endif // RHI_RAYTRACING
-
-void FDeferredShadingSceneRenderer::PreGatherDynamicMeshElements()
-{
-#if RHI_RAYTRACING
-	if (bAnyRayTracingPassEnabled)
-	{
-		const int32 ReferenceViewIndex = 0;
-		FViewInfo& ReferenceView = Views[ReferenceViewIndex];
-
-		RayTracingRelevantPrimitiveTaskData = Allocator.Create<FRayTracingRelevantPrimitiveTaskData>();
-		RayTracingRelevantPrimitiveTaskData->Task = FFunctionGraphTask::CreateAndDispatchWhenReady(
-			[Scene = this->Scene, &ReferenceView, &RayTracingRelevantPrimitiveList = RayTracingRelevantPrimitiveTaskData->List]()
-		{
-			FTaskTagScope TaskTagScope(ETaskTag::EParallelRenderingThread);
-			GatherRayTracingRelevantPrimitives(*Scene, ReferenceView, RayTracingRelevantPrimitiveList);
-		}, TStatId(), nullptr, ENamedThreads::AnyNormalThreadHiPriTask);
-	}
-#endif // RHI_RAYTRACING
-
-	const bool bHasRayTracedOverlay = HasRayTracedOverlay(ViewFamily);
-
-	extern int32 GEarlyInitDynamicShadows;
-
-	if (GEarlyInitDynamicShadows &&
-		CurrentDynamicShadowsTaskData == nullptr &&
-		ViewFamily.EngineShowFlags.DynamicShadows
-		&& !ViewFamily.EngineShowFlags.HitProxies
-		&& !bHasRayTracedOverlay)
-	{
-		CurrentDynamicShadowsTaskData = BeginInitDynamicShadows(true);
-	}
-}
-
 bool FDeferredShadingSceneRenderer::GatherRayTracingWorldInstancesForView(FRDGBuilder& GraphBuilder, FViewInfo& View, FRayTracingScene& RayTracingScene, FRayTracingRelevantPrimitiveList& RelevantPrimitiveList)
 {
 	checkf(IsRayTracingEnabled() && bAnyRayTracingPassEnabled, TEXT("GatherRayTracingWorldInstancesForView should only be called if ray tracing is used"))
@@ -2177,8 +2137,45 @@
 	Scene->RayTracingScene.Transition(GraphBuilder, ERayTracingSceneState::Readable);
 }
 
+struct FRayTracingRelevantPrimitiveTaskData
+{
+	FRayTracingRelevantPrimitiveList List;
+	FGraphEventRef Task;
+};
 #endif // RHI_RAYTRACING
 
+void FDeferredShadingSceneRenderer::PreGatherDynamicMeshElements()
+{
+#if RHI_RAYTRACING
+	if (bAnyRayTracingPassEnabled)
+	{
+		const int32 ReferenceViewIndex = 0;
+		FViewInfo& ReferenceView = Views[ReferenceViewIndex];
+
+		RayTracingRelevantPrimitiveTaskData = Allocator.Create<FRayTracingRelevantPrimitiveTaskData>();
+		RayTracingRelevantPrimitiveTaskData->Task = FFunctionGraphTask::CreateAndDispatchWhenReady(
+			[Scene = this->Scene, &ReferenceView, &RayTracingRelevantPrimitiveList = RayTracingRelevantPrimitiveTaskData->List]()
+			{
+				FTaskTagScope TaskTagScope(ETaskTag::EParallelRenderingThread);
+				GatherRayTracingRelevantPrimitives(*Scene, ReferenceView, RayTracingRelevantPrimitiveList);
+			}, TStatId(), nullptr, ENamedThreads::AnyNormalThreadHiPriTask);
+	}
+#endif // RHI_RAYTRACING
+
+	const bool bHasRayTracedOverlay = HasRayTracedOverlay(ViewFamily);
+
+	extern int32 GEarlyInitDynamicShadows;
+
+	if (GEarlyInitDynamicShadows &&
+		CurrentDynamicShadowsTaskData == nullptr &&
+		ViewFamily.EngineShowFlags.DynamicShadows
+		&& !ViewFamily.EngineShowFlags.HitProxies
+		&& !bHasRayTracedOverlay)
+	{
+		CurrentDynamicShadowsTaskData = BeginInitDynamicShadows(true);
+	}
+}
+
 static TAutoConsoleVariable<float> CVarStallInitViews(
 	TEXT("CriticalPathStall.AfterInitViews"),
 	0.0f,
