* Add groom view mode menu to level editor
* Unify most groom debug rendering to go through view more rather than CVars

diff --git a/Engine/Plugins/Runtime/HairStrands/Source/HairStrandsCore/Private/Deformer/DeformerDataInterfaceGroom.cpp b/Engine/Plugins/Runtime/HairStrands/Source/HairStrandsCore/Private/Deformer/DeformerDataInterfaceGroom.cpp
--- a/Engine/Plugins/Runtime/HairStrands/Source/HairStrandsCore/Private/Deformer/DeformerDataInterfaceGroom.cpp
+++ b/Engine/Plugins/Runtime/HairStrands/Source/HairStrandsCore/Private/Deformer/DeformerDataInterfaceGroom.cpp
@@ -165,7 +165,7 @@ void FOptimusGroomDataProviderProxy::AllocateResources(FRDGBuilder& GraphBuilder
 	}
 }
 
-FHairGroupPublicData::FVertexFactoryInput ComputeHairStrandsVertexInputData(const FHairGroupInstance* Instance);
+FHairGroupPublicData::FVertexFactoryInput ComputeHairStrandsVertexInputData(const FHairGroupInstance* Instance, EGroomViewMode ViewMode);
 void FOptimusGroomDataProviderProxy::GatherDispatchData(FDispatchSetup const& InDispatchSetup, FCollectedDispatchData& InOutDispatchData)
 {
 	if (!ensure(InDispatchSetup.ParameterStructSizeForValidation == sizeof(FGroomDataInterfaceParameters)))
@@ -184,7 +184,7 @@ void FOptimusGroomDataProviderProxy::GatherDispatchData(FDispatchSetup const& In
 			const int32 NumControlPoints = bIsSRVValid ? Instance->Strands.Data->PointCount : 0;
 			const int32 NumCurves = bIsSRVValid ? Instance->Strands.Data->CurveCount : 0;
 
-			const FHairGroupPublicData::FVertexFactoryInput VFInput = ComputeHairStrandsVertexInputData(Instance);
+			const FHairGroupPublicData::FVertexFactoryInput VFInput = ComputeHairStrandsVertexInputData(Instance, EGroomViewMode::None);
 			
 			FGroomDataInterfaceParameters* Parameters = (FGroomDataInterfaceParameters*)(InOutDispatchData.ParameterBuffer + InDispatchSetup.ParameterBufferOffset + InDispatchSetup.ParameterBufferStride * InvocationIndex);
 			Parameters->NumControlPoints = NumControlPoints;
diff --git a/Engine/Plugins/Runtime/HairStrands/Source/HairStrandsCore/Private/GroomComponent.cpp b/Engine/Plugins/Runtime/HairStrands/Source/HairStrandsCore/Private/GroomComponent.cpp
--- a/Engine/Plugins/Runtime/HairStrands/Source/HairStrandsCore/Private/GroomComponent.cpp
+++ b/Engine/Plugins/Runtime/HairStrands/Source/HairStrandsCore/Private/GroomComponent.cpp
@@ -72,8 +72,6 @@ static FAutoConsoleVariableRef CVarHairStrands_UseAttachedSimulationComponents(T
 
 ////////////////////////////////////////////////////////////////////////////////////////////////////////////////////////
 
-EHairStrandsDebugMode GetHairStrandsGeometryDebugMode(const FHairGroupInstance* Instance);
-
 const FLinearColor GetHairGroupDebugColor(int32 GroupIt)
 {
 	static TArray<FLinearColor> IndexedColor;
@@ -635,36 +633,48 @@ public:
 					check(Instances[GroupIt]->GetRefCount() > 0);
 
 					FMaterialRenderProxy* Debug_MaterialProxy = nullptr;
-					EHairStrandsDebugMode DebugMode = GetHairStrandsGeometryDebugMode(Instances[GroupIt]);
-					if (AllowDebugViewmodes() && View->Family->EngineShowFlags.LODColoration)
+					const EGroomViewMode ViewMode = AllowDebugViewmodes() ? GetGroomViewMode(*View) : EGroomViewMode::None;
+					bool bNeedDebugMaterial = false;
+					switch(ViewMode)
 					{
-						DebugMode = EHairStrandsDebugMode::RenderLODColoration;
-					}
-					
-					const bool bNeedDebugMaterial = 
-						DebugMode != EHairStrandsDebugMode::NoneDebug &&
-						DebugMode != EHairStrandsDebugMode::RenderHairControlPoints &&
-						DebugMode != EHairStrandsDebugMode::RenderHairTangent;
+					case EGroomViewMode::SimHairStrands				:
+					case EGroomViewMode::RenderVisCluster			:
+					case EGroomViewMode::RenderVisClusterAABB		:
+					case EGroomViewMode::RenderHairStrands			:
+						bNeedDebugMaterial = Instances[GroupIt]->GeometryType == EHairGeometryType::Strands; break;
+					case EGroomViewMode::RenderHairRootUV			:
+					case EGroomViewMode::RenderHairUV				:
+					case EGroomViewMode::RenderHairSeed				:
+					case EGroomViewMode::RenderHairDimension		:
+					case EGroomViewMode::RenderHairRadiusVariation	:
+					case EGroomViewMode::RenderHairRootUDIM			:
+					case EGroomViewMode::RenderHairBaseColor		:
+					case EGroomViewMode::RenderHairRoughness		:
+					case EGroomViewMode::RenderHairGroup			:
+					case EGroomViewMode::RenderLODColoration		:
+						bNeedDebugMaterial = true; break;
+					};
 
 					if (bNeedDebugMaterial)
 					{
 						float DebugModeScalar = 0;
-						switch(DebugMode)
+						switch(ViewMode)
 						{
-						case EHairStrandsDebugMode::NoneDebug					: DebugModeScalar =99.f; break;
-						case EHairStrandsDebugMode::SimHairStrands				: DebugModeScalar = 0.f; break;
-						case EHairStrandsDebugMode::RenderHairStrands			: DebugModeScalar = 0.f; break;
-						case EHairStrandsDebugMode::RenderHairRootUV			: DebugModeScalar = 1.f; break;
-						case EHairStrandsDebugMode::RenderHairUV				: DebugModeScalar = 2.f; break;
-						case EHairStrandsDebugMode::RenderHairSeed				: DebugModeScalar = 3.f; break;
-						case EHairStrandsDebugMode::RenderHairDimension			: DebugModeScalar = 4.f; break;
-						case EHairStrandsDebugMode::RenderHairRadiusVariation	: DebugModeScalar = 5.f; break;
-						case EHairStrandsDebugMode::RenderHairRootUDIM			: DebugModeScalar = 6.f; break;
-						case EHairStrandsDebugMode::RenderHairBaseColor			: DebugModeScalar = 7.f; break;
-						case EHairStrandsDebugMode::RenderHairRoughness			: DebugModeScalar = 8.f; break;
-						case EHairStrandsDebugMode::RenderVisCluster			: DebugModeScalar = 0.f; break;
-						case EHairStrandsDebugMode::RenderHairGroup				: DebugModeScalar = 9.f; break;
-						case EHairStrandsDebugMode::RenderLODColoration			: DebugModeScalar = 10.f; break;
+						case EGroomViewMode::None						: DebugModeScalar =99.f; break;
+						case EGroomViewMode::SimHairStrands				: DebugModeScalar = 0.f; break;
+						case EGroomViewMode::RenderHairStrands			: DebugModeScalar = 0.f; break;
+						case EGroomViewMode::RenderHairRootUV			: DebugModeScalar = 1.f; break;
+						case EGroomViewMode::RenderHairUV				: DebugModeScalar = 2.f; break;
+						case EGroomViewMode::RenderHairSeed				: DebugModeScalar = 3.f; break;
+						case EGroomViewMode::RenderHairDimension		: DebugModeScalar = 4.f; break;
+						case EGroomViewMode::RenderHairRadiusVariation	: DebugModeScalar = 5.f; break;
+						case EGroomViewMode::RenderHairRootUDIM			: DebugModeScalar = 6.f; break;
+						case EGroomViewMode::RenderHairBaseColor		: DebugModeScalar = 7.f; break;
+						case EGroomViewMode::RenderHairRoughness		: DebugModeScalar = 8.f; break;
+						case EGroomViewMode::RenderVisCluster			: DebugModeScalar = 0.f; break;
+						case EGroomViewMode::RenderVisClusterAABB		: DebugModeScalar = 0.f; break;
+						case EGroomViewMode::RenderHairGroup			: DebugModeScalar = 9.f; break;
+						case EGroomViewMode::RenderLODColoration		: DebugModeScalar = 10.f; break;
 						};
 
 						// TODO: fix this as the radius is incorrect. This code run before the interpolation code, which is where HairRadius is updated.
@@ -675,17 +685,17 @@ public:
 						}
 						
 						// Reuse the HairMaxRadius field to send the LOD index instead of adding yet another variable
-						if (DebugMode == EHairStrandsDebugMode::RenderLODColoration)
+						if (ViewMode == EGroomViewMode::RenderLODColoration)
 						{
 							HairMaxRadius = Instances[GroupIt]->HairGroupPublicData ? Instances[GroupIt]->HairGroupPublicData->LODIndex : 0;
 						}
 
 						FVector HairColor = FVector::ZeroVector;
-						if (DebugMode == EHairStrandsDebugMode::RenderHairGroup)
+						if (ViewMode == EGroomViewMode::RenderHairGroup)
 						{
 							HairColor = FVector(GetHairGroupDebugColor(GroupIt));
 						}
-						else if (DebugMode == EHairStrandsDebugMode::RenderLODColoration)
+						else if (ViewMode == EGroomViewMode::RenderLODColoration)
 						{
 							int32 LODIndex = Instances[GroupIt]->HairGroupPublicData ? Instances[GroupIt]->HairGroupPublicData->LODIndex : 0;
 							LODIndex = FMath::Clamp(LODIndex, 0, GEngine->LODColorationColors.Num() - 1);
@@ -1610,14 +1620,6 @@ void UGroomComponent::SetHairLengthScaleEnable(bool bEnable)
 	InitResources();
 }
 
-#if WITH_EDITOR
-void UGroomComponent::SetDebugMode(EHairStrandsDebugMode InMode)
-{ 
-	DebugMode = InMode; 
-	UpdateHairGroupsDescAndInvalidateRenderState(true);
-}
-#endif
-
 bool UGroomComponent::GetIsHairLengthScaleEnabled()
 {
 	bool IsEnabled = true;
@@ -1815,11 +1817,7 @@ void UGroomComponent::UpdateHairGroupsDescAndInvalidateRenderState(bool bInvalid
 	uint32 GroupIndex = 0;
 	for (FHairGroupInstance* Instance : HairGroupInstances)
 	{
-		Instance->Strands.Modifier  = GetGroomGroupsDesc(GroomAsset, this, GroupIndex);
-#if WITH_EDITORONLY_DATA
-		Instance->Debug.DebugMode = DebugMode;
-		Instance->HairGroupPublicData->DebugMode = DebugMode; // Replicated value for accessing it from the engine side. Refactor this.
-#endif// #if WITH_EDITORONLY_DATA
+		Instance->Strands.Modifier = GetGroomGroupsDesc(GroomAsset, this, GroupIndex);
 		++GroupIndex;
 	}
 	if (bInvalid)
@@ -2672,9 +2670,6 @@ void UGroomComponent::InitResources(bool bIsBindingReloading)
 
 		// LODBias is in the Modifier which is needed for LOD selection regardless if the strands are there or not
 		HairGroupInstance->Strands.Modifier = GetGroomGroupsDesc(GroomAsset, this, GroupIt);
-		#if WITH_EDITORONLY_DATA
-		HairGroupInstance->Debug.DebugMode = DebugMode;
-		#endif// #if WITH_EDITORONLY_DATA
 
 		// Strands data/resources
 		if (bNeedStrandsData && GroupData.Strands.IsValid())
diff --git a/Engine/Plugins/Runtime/HairStrands/Source/HairStrandsCore/Private/GroomDebug.cpp b/Engine/Plugins/Runtime/HairStrands/Source/HairStrandsCore/Private/GroomDebug.cpp
--- a/Engine/Plugins/Runtime/HairStrands/Source/HairStrandsCore/Private/GroomDebug.cpp
+++ b/Engine/Plugins/Runtime/HairStrands/Source/HairStrandsCore/Private/GroomDebug.cpp
@@ -15,6 +15,7 @@
 #include "SystemTextures.h"
 #include "CanvasTypes.h"
 #include "ShaderCompilerCore.h"
+#include "GroomVisualizationData.h"
 
 ///////////////////////////////////////////////////////////////////////////////////////////////////
 
@@ -32,7 +33,7 @@ static int32 GHairDebugMeshProjection_Render_HairRestTriangles = 0;
 static int32 GHairDebugMeshProjection_Render_HairRestFrames = 0;
 static int32 GHairDebugMeshProjection_Render_HairRestSamples = 0;
 static int32 GHairDebugMeshProjection_Render_HairDeformedTriangles = 0;
-static int32 GHairDebugMeshProjection_Render_HairDeformedFrames = 0;
+static int32 GHairDebugMeshProjection_Render_HairDeformedFrames = 1;
 static int32 GHairDebugMeshProjection_Render_HairDeformedSamples = 0;
 
 
@@ -59,14 +60,6 @@ static FAutoConsoleVariableRef CVarHairCardsAtlasDebug(TEXT("r.HairStrands.Cards
 static int32 GHairCardsVoxelDebug = 0;
 static FAutoConsoleVariableRef CVarHairCardsVoxelDebug(TEXT("r.HairStrands.Cards.DebugVoxel"), GHairCardsVoxelDebug, TEXT("Draw debug hair cards voxel datas."));
 
-static int32 GHairCardsGuidesDebug_Ren = 0;
-static int32 GHairCardsGuidesDebug_Sim = 0;
-static FAutoConsoleVariableRef CVarHairCardsGuidesDebug_Ren(TEXT("r.HairStrands.Cards.DebugGuides.Render"), GHairCardsGuidesDebug_Ren, TEXT("Draw debug hair cards guides (1: Rest, 2: Deformed)."));
-static FAutoConsoleVariableRef CVarHairCardsGuidesDebug_Sim(TEXT("r.HairStrands.Cards.DebugGuides.Sim"), GHairCardsGuidesDebug_Sim, TEXT("Draw debug hair sim guides (1: Rest, 2: Deformed)."));
-
-static int32 GHairStrandsControlPointDebug = 0;
-static FAutoConsoleVariableRef CVarHairStrandsControlPointDebug(TEXT("r.HairStrands.Strands.DebugControlPoint"), GHairStrandsControlPointDebug, TEXT("Draw debug hair strands control points)."));
-
 ///////////////////////////////////////////////////////////////////////////////////////////////////
 
 FCachedGeometry GetCacheGeometryForHair(
@@ -1048,14 +1041,14 @@ void RunHairStrandsDebug(
 	FIntRect Viewport,
 	const TUniformBufferRef<FViewUniformShaderParameters>& ViewUniformBuffer)
 {
-	const EHairDebugMode HairDebugMode = GetHairStrandsDebugMode();
+	const EGroomViewMode ViewMode = GetGroomViewMode(View);
 
-	if (HairDebugMode == EHairDebugMode::MacroGroups)
+	if (ViewMode == EGroomViewMode::MacroGroups)
 	{
 		AddHairDebugPrintInstancePass(GraphBuilder, ShaderMap, ShaderPrintData, Instances, InstanceCountPerType);
 	}
 
-	if (HairDebugMode == EHairDebugMode::MeshProjection)
+	if (ViewMode == EGroomViewMode::MeshProjection)
 	{
 		{
 			if (GHairDebugMeshProjection_SkinCacheMesh > 0)
@@ -1162,17 +1155,17 @@ void RunHairStrandsDebug(
 	{
 		FHairGroupInstance* Instance = static_cast<FHairGroupInstance*>(AbstractInstance);
 
-		if (GHairCardsGuidesDebug_Ren > 0 || Instance->Debug.bDrawCardsGuides)
+		if (ViewMode == EGroomViewMode::CardGuides)
 		{
-			AddDrawDebugCardsGuidesPass(GraphBuilder, View, ShaderMap, Instance, ShaderPrintData, Instance->Debug.bDrawCardsGuides ? false : GHairCardsGuidesDebug_Ren == 1, true);
+			AddDrawDebugCardsGuidesPass(GraphBuilder, View, ShaderMap, Instance, ShaderPrintData, true /*bDeformed*/, true);
 		}
 
-		if (GHairCardsGuidesDebug_Sim > 0)
+		if (ViewMode == EGroomViewMode::SimHairStrands)
 		{
-			AddDrawDebugCardsGuidesPass(GraphBuilder, View, ShaderMap, Instance, ShaderPrintData, GHairCardsGuidesDebug_Sim == 1, false);
+			AddDrawDebugCardsGuidesPass(GraphBuilder, View, ShaderMap, Instance, ShaderPrintData, true /*bDeformed*/, false);
 		}
 
-		if (GHairStrandsControlPointDebug || Instance->HairGroupPublicData->DebugMode == EHairStrandsDebugMode::RenderHairControlPoints)
+		if (ViewMode == EGroomViewMode::RenderHairControlPoints)
 		{
 			AddDrawDebugStrandsCVsPass(GraphBuilder, View, ShaderMap, Instance, ShaderPrintData, SceneColorTexture, SceneDepthTexture);
 		}
diff --git a/Engine/Plugins/Runtime/HairStrands/Source/HairStrandsCore/Private/GroomManager.cpp b/Engine/Plugins/Runtime/HairStrands/Source/HairStrandsCore/Private/GroomManager.cpp
--- a/Engine/Plugins/Runtime/HairStrands/Source/HairStrandsCore/Private/GroomManager.cpp
+++ b/Engine/Plugins/Runtime/HairStrands/Source/HairStrandsCore/Private/GroomManager.cpp
@@ -150,6 +150,8 @@ static void RunInternalHairStrandsInterpolation(
 	const uint32 ViewRayTracingMask = 0u;
 	#endif
 
+	const EGroomViewMode ViewMode = GetGroomViewMode(*View);
+
 	// Update dynamic mesh triangles
 	for (FHairStrandsInstance* AbstractInstance : Instances)
 	{
@@ -381,6 +383,7 @@ static void RunInternalHairStrandsInterpolation(
 				ShaderMap,
 				ViewUniqueID,
 				ViewRayTracingMask,
+				ViewMode,
 				TranslatedWorldOffset,
 				ShaderPrintData,
 				Instance,
diff --git a/Engine/Plugins/Runtime/HairStrands/Source/HairStrandsCore/Private/HairStrandsCore.cpp b/Engine/Plugins/Runtime/HairStrands/Source/HairStrandsCore/Private/HairStrandsCore.cpp
--- a/Engine/Plugins/Runtime/HairStrands/Source/HairStrandsCore/Private/HairStrandsCore.cpp
+++ b/Engine/Plugins/Runtime/HairStrands/Source/HairStrandsCore/Private/HairStrandsCore.cpp
@@ -26,12 +26,12 @@ void FHairStrandsCore::StartupModule()
 	// Maps virtual shader source directory /Plugin/FX/Niagara to the plugin's actual Shaders directory.
 	FString PluginShaderDir = FPaths::Combine(IPluginManager::Get().FindPlugin(TEXT("HairStrands"))->GetBaseDir(), TEXT("Shaders"));
 	AddShaderSourceDirectoryMapping(TEXT("/Plugin/Runtime/HairStrands"), PluginShaderDir);
-	SetHairStrandsEnabled(true);
+	SetGroomEnabled(true);
 }
 
 void FHairStrandsCore::ShutdownModule()
 {
-	SetHairStrandsEnabled(false);
+	SetGroomEnabled(false);
 }
 
 void FHairStrandsCore::RegisterAssetHelper(const FHairAssetHelper& Helper)
diff --git a/Engine/Plugins/Runtime/HairStrands/Source/HairStrandsCore/Private/HairStrandsRendering.cpp b/Engine/Plugins/Runtime/HairStrands/Source/HairStrandsCore/Private/HairStrandsRendering.cpp
--- a/Engine/Plugins/Runtime/HairStrands/Source/HairStrandsCore/Private/HairStrandsRendering.cpp
+++ b/Engine/Plugins/Runtime/HairStrands/Source/HairStrandsCore/Private/HairStrandsRendering.cpp
@@ -1612,21 +1612,6 @@ void RegisterClusterData(FHairGroupInstance* Instance, FHairStrandClusterData* I
 	HairGroupCluster.HairGroupPublicPtr->ClusterDataIndex = ClusterDataGroupIndex;
 }
 
-EHairStrandsDebugMode GetHairStrandsGeometryDebugMode(const FHairGroupInstance* Instance)
-{
-	// 1. Use the per-instance debug mode by default, and fallback on the global debug otherwise
-	EHairStrandsDebugMode DebugMode = Instance->Debug.DebugMode != EHairStrandsDebugMode::NoneDebug ? Instance->Debug.DebugMode : GetHairStrandsDebugStrandsMode();
-
-	// 2. If groom does not have guides (e.g., because it does not have simulation/RBF data), then we fallback to default view mode
-	if ((DebugMode == EHairStrandsDebugMode::SimHairStrands ||
-		 DebugMode == EHairStrandsDebugMode::RenderHairStrands)
-		&& Instance->Guides.Data == nullptr)
-	{
-		DebugMode = EHairStrandsDebugMode::NoneDebug;
-	}
-	return DebugMode;
-}
-
 static void ConvertHairStrandsVFParameters(
 	FRDGBuilder* GraphBuilder,
 	FRDGImportedBuffer& OutBuffer,
@@ -1644,17 +1629,17 @@ static void ConvertHairStrandsVFParameters(
 #define CONVERT_HAIRSSTRANDS_VF_PARAMETERS(OutName, ExternalBuffer) ConvertHairStrandsVFParameters(GraphBuilder, OutName, OutName##External, OutName##RHISRV, ExternalBuffer);
 
 // Compute/Update the hair strands description which will be used for rendering (VF) / voxelization & co.
-static FHairGroupPublicData::FVertexFactoryInput InternalComputeHairStrandsVertexInputData(FRDGBuilder* GraphBuilder, const FHairGroupInstance* Instance)
+static FHairGroupPublicData::FVertexFactoryInput InternalComputeHairStrandsVertexInputData(FRDGBuilder* GraphBuilder, const FHairGroupInstance* Instance, EGroomViewMode ViewMode)
 {
 	FHairGroupPublicData::FVertexFactoryInput OutVFInput;
 	if (!Instance || Instance->GeometryType != EHairGeometryType::Strands)
 		return OutVFInput;
 
-	const EHairStrandsDebugMode DebugMode = GetHairStrandsGeometryDebugMode(Instance);
+	const bool bSupportHasGuides = Instance->Guides.IsValid();
 	const bool bSupportDeformation = Instance->Strands.DeformedResource != nullptr;
 
 	// 1. Guides
-	if (DebugMode == EHairStrandsDebugMode::SimHairStrands)
+	if (ViewMode == EGroomViewMode::SimHairStrands && bSupportHasGuides)
 	{
 		CONVERT_HAIRSSTRANDS_VF_PARAMETERS(OutVFInput.Strands.PositionBuffer,			Instance->Guides.DeformedResource->GetBuffer(FHairStrandsDeformedResource::EFrameType::Current));
 		CONVERT_HAIRSSTRANDS_VF_PARAMETERS(OutVFInput.Strands.PrevPositionBuffer,		Instance->Guides.DeformedResource->GetBuffer(FHairStrandsDeformedResource::EFrameType::Previous));
@@ -1696,7 +1681,7 @@ static FHairGroupPublicData::FVertexFactoryInput InternalComputeHairStrandsVerte
 			OutVFInput.Strands.PrevPositionOffset = Instance->Strands.DeformedResource->GetPositionOffset(FHairStrandsDeformedResource::EFrameType::Current);
 		}
 		CONVERT_HAIRSSTRANDS_VF_PARAMETERS(OutVFInput.Strands.TangentBuffer,			Instance->Strands.DeformedResource->TangentBuffer);
-		CONVERT_HAIRSSTRANDS_VF_PARAMETERS(OutVFInput.Strands.Attribute0Buffer,			(DebugMode == EHairStrandsDebugMode::RenderHairStrands || DebugMode == EHairStrandsDebugMode::RenderVisCluster) ? Instance->Strands.DebugAttributeBuffer : Instance->Strands.RestResource->Attribute0Buffer);
+		CONVERT_HAIRSSTRANDS_VF_PARAMETERS(OutVFInput.Strands.Attribute0Buffer,			(ViewMode == EGroomViewMode::RenderHairStrands || ViewMode == EGroomViewMode::RenderVisCluster || ViewMode == EGroomViewMode::RenderVisClusterAABB) ? Instance->Strands.DebugAttributeBuffer : Instance->Strands.RestResource->Attribute0Buffer);
 		CONVERT_HAIRSSTRANDS_VF_PARAMETERS(OutVFInput.Strands.Attribute1Buffer,			Instance->Strands.RestResource->Attribute1Buffer);
 		CONVERT_HAIRSSTRANDS_VF_PARAMETERS(OutVFInput.Strands.MaterialBuffer,			Instance->Strands.RestResource->MaterialBuffer);
 
@@ -1708,7 +1693,7 @@ static FHairGroupPublicData::FVertexFactoryInput InternalComputeHairStrandsVerte
 		CONVERT_HAIRSSTRANDS_VF_PARAMETERS(OutVFInput.Strands.PositionBuffer,			Instance->Strands.RestResource->PositionBuffer);
 		CONVERT_HAIRSSTRANDS_VF_PARAMETERS(OutVFInput.Strands.PrevPositionBuffer,		Instance->Strands.RestResource->PositionBuffer);
 		CONVERT_HAIRSSTRANDS_VF_PARAMETERS(OutVFInput.Strands.TangentBuffer,			Instance->Strands.RestResource->TangentBuffer);
-		CONVERT_HAIRSSTRANDS_VF_PARAMETERS(OutVFInput.Strands.Attribute0Buffer,			(DebugMode == EHairStrandsDebugMode::RenderVisCluster) ? Instance->Strands.DebugAttributeBuffer : Instance->Strands.RestResource->Attribute0Buffer);
+		CONVERT_HAIRSSTRANDS_VF_PARAMETERS(OutVFInput.Strands.Attribute0Buffer,			(ViewMode == EGroomViewMode::RenderVisCluster || ViewMode == EGroomViewMode::RenderVisClusterAABB) ? Instance->Strands.DebugAttributeBuffer : Instance->Strands.RestResource->Attribute0Buffer);
 		CONVERT_HAIRSSTRANDS_VF_PARAMETERS(OutVFInput.Strands.Attribute1Buffer,			Instance->Strands.RestResource->Attribute1Buffer);
 		CONVERT_HAIRSSTRANDS_VF_PARAMETERS(OutVFInput.Strands.MaterialBuffer,			Instance->Strands.RestResource->MaterialBuffer);
 		CONVERT_HAIRSSTRANDS_VF_PARAMETERS(OutVFInput.Strands.PositionOffsetBuffer,		Instance->Strands.RestResource->PositionOffsetBuffer);
@@ -1738,9 +1723,9 @@ static FHairGroupPublicData::FVertexFactoryInput InternalComputeHairStrandsVerte
 	return OutVFInput;
 }
 
-FHairGroupPublicData::FVertexFactoryInput ComputeHairStrandsVertexInputData(const FHairGroupInstance* Instance)
+FHairGroupPublicData::FVertexFactoryInput ComputeHairStrandsVertexInputData(const FHairGroupInstance* Instance, EGroomViewMode ViewMode)
 {
-	return InternalComputeHairStrandsVertexInputData(nullptr, Instance);
+	return InternalComputeHairStrandsVertexInputData(nullptr, Instance, ViewMode);
 }
 
 void CreateHairStrandsDebugAttributeBuffer(FRDGBuilder& GraphBuilder, FRDGExternalBuffer* DebugAttributeBuffer, uint32 VertexCount);
@@ -1750,6 +1735,7 @@ void ComputeHairStrandsInterpolation(
 	FGlobalShaderMap* ShaderMap,
 	const uint32 ViewUniqueID,
 	const uint32 ViewRayTracingMask,
+	const EGroomViewMode ViewMode,
 	const FVector& TranslatedWorldOffset,
 	const FShaderPrintData* ShaderPrintData,
 	FHairGroupInstance* Instance,
@@ -1779,14 +1765,11 @@ void ComputeHairStrandsInterpolation(
 		// * None	: Display hair normally
 		// * Sim	: Show sim strands
 		// * Render : Show rendering strands with sim color influence
-		const EHairStrandsDebugMode DebugMode = GetHairStrandsGeometryDebugMode(Instance);
-
-		if (DebugMode == EHairStrandsDebugMode::SimHairStrands)
+		if (ViewMode == EGroomViewMode::SimHairStrands && Instance->Guides.DeformedResource)
 		{
 			// Disable culling when drawing only guides, as the culling output has been computed for the strands, not for the guides.
 			Instance->HairGroupPublicData->SetCullingResultAvailable(false);
 
-			check(Instance->Guides.DeformedResource);
 			AddHairTangentPass(
 				GraphBuilder,
 				ShaderMap,
@@ -1796,7 +1779,7 @@ void ComputeHairStrandsInterpolation(
 				RegisterAsSRV(GraphBuilder, Instance->Guides.DeformedResource->GetBuffer(FHairStrandsDeformedResource::Current)),
 				Register(GraphBuilder, Instance->Guides.DeformedResource->TangentBuffer, ERDGImportedBufferFlags::CreateUAV));
 
-			Instance->HairGroupPublicData->VFInput = InternalComputeHairStrandsVertexInputData(&GraphBuilder, Instance);
+			Instance->HairGroupPublicData->VFInput = InternalComputeHairStrandsVertexInputData(&GraphBuilder, Instance, ViewMode);
 		}
 		else
 		{
@@ -1818,16 +1801,20 @@ void ComputeHairStrandsInterpolation(
 			// "WITH_EDITOR && PatchMode == EHairPatchAttribute::GuideInflucence" is a special path when visualizing groom within the groom editor, so that we can diplay guides even when there is no simulation or global interpolation enabled
 			EHairPatchAttribute PatchMode = EHairPatchAttribute::None;
 			#if WITH_EDITOR
-			if (DebugMode == EHairStrandsDebugMode::RenderHairStrands)	PatchMode = EHairPatchAttribute::GuideInflucence;
-			if (DebugMode == EHairStrandsDebugMode::RenderVisCluster)	PatchMode = EHairPatchAttribute::ClusterInfluence;
+			if (ViewMode == EGroomViewMode::RenderHairStrands && bNeedDeformation)									{ PatchMode = EHairPatchAttribute::GuideInflucence; }
+			if (ViewMode == EGroomViewMode::RenderVisCluster || ViewMode == EGroomViewMode::RenderVisClusterAABB)	{ PatchMode = EHairPatchAttribute::ClusterInfluence; }
 			#endif
 
 			// 2. Deform hair if needed (e.g.: skinning, simulation, RBF) and recompute tangent
 			FRDGBufferSRVRef Strands_PositionSRV = nullptr;
 			FRDGBufferSRVRef Strands_PositionOffsetSRV = nullptr;
 			FRDGBufferSRVRef Strands_PrevPositionOffsetSRV = nullptr;
-
 			FRDGBufferSRVRef Strands_TangentSRV = nullptr;
+
+			const bool bValidGuide		= bNeedDeformation && (Instance->Guides.bIsSimulationEnable || Instance->Guides.bHasGlobalInterpolation || Instance->Guides.bIsDeformationEnable);// || (WITH_EDITOR && PatchMode == EHairPatchAttribute::GuideInflucence);
+			const bool bUseSingleGuide	= bNeedDeformation && bValidGuide && Instance->Strands.InterpolationResource->UseSingleGuide();
+			const bool bHasSkinning		= bNeedDeformation && Instance->BindingType == EHairBindingType::Skinning;
+
 			if (bNeedDeformation)
 			{
 				FRDGImportedBuffer Strands_DeformedPosition = Register(GraphBuilder, Instance->Strands.DeformedResource->GetBuffer(FHairStrandsDeformedResource::Current), ERDGImportedBufferFlags::CreateViews);
@@ -1842,8 +1829,6 @@ void ComputeHairStrandsInterpolation(
 				Strands_PositionOffsetSRV = RegisterAsSRV(GraphBuilder, Instance->Strands.DeformedResource->GetPositionOffsetBuffer(FHairStrandsDeformedResource::EFrameType::Current));
 				Strands_TangentSRV = Strands_DeformedTangent.SRV;
 
-				const bool bValidGuide = Instance->Guides.bIsSimulationEnable || Instance->Guides.bHasGlobalInterpolation || Instance->Guides.bIsDeformationEnable || (WITH_EDITOR && PatchMode == EHairPatchAttribute::GuideInflucence);
-				const bool bHasSkinning = Instance->BindingType == EHairBindingType::Skinning;
 
 				bool bHasValidGroomCacheBuffers = Instance->Debug.GroomCacheBuffers.IsValid() && Instance->Debug.GroomCacheBuffers->GetInterpolatedFrameBuffer().GroupsData.IsValidIndex(Instance->Debug.GroupIndex);
 				if (bHasValidGroomCacheBuffers && Instance->Debug.GroomCacheType == EGroomCacheType::Guides)
@@ -1923,7 +1908,6 @@ void ComputeHairStrandsInterpolation(
 					VertexCount = Instance->HairGroupPublicData->GetActiveStrandsVertexCount(Instance->Strands.RestResource->GetVertexCount(), Instance->HairGroupPublicData->MaxScreenSize);
 
 					// 2.1.1 Current Position
-					const bool bUseSingleGuide = bValidGuide && Instance->Strands.InterpolationResource->UseSingleGuide();
 					AddHairStrandsInterpolationPass(
 						GraphBuilder,
 						ShaderMap,
@@ -2068,8 +2052,6 @@ void ComputeHairStrandsInterpolation(
 				}
 				FRDGImportedBuffer OutRenAttribute0Buffer = Register(GraphBuilder, Instance->Strands.DebugAttributeBuffer, ERDGImportedBufferFlags::CreateUAV);
 
-				const bool bValidGuide = PatchMode == EHairPatchAttribute::GuideInflucence;
-				const bool bUseSingleGuide = bValidGuide && Instance->Strands.InterpolationResource->UseSingleGuide();
 				AddPatchAttributePass(
 					GraphBuilder,
 					ShaderMap,
@@ -2088,7 +2070,7 @@ void ComputeHairStrandsInterpolation(
 			#endif
 
 			// 2.2 Update the VF input with the update resources
-			Instance->HairGroupPublicData->VFInput = InternalComputeHairStrandsVertexInputData(&GraphBuilder, Instance);
+			Instance->HairGroupPublicData->VFInput = InternalComputeHairStrandsVertexInputData(&GraphBuilder, Instance, ViewMode);
 
 			// 3. Compute cluster AABBs (used for LODing and voxelization)
 			{
@@ -2220,7 +2202,7 @@ void ComputeHairStrandsInterpolation(
 
 		// Sanity check
 		check(Instance->HairGroupPublicData->VFInput.Strands.PositionBuffer.Buffer);
-		Instance->Strands.UniformBuffer.UpdateUniformBufferImmediate(Instance->GetHairStandsUniformShaderParameters());
+		Instance->Strands.UniformBuffer.UpdateUniformBufferImmediate(Instance->GetHairStandsUniformShaderParameters(ViewMode));
 	}
 	else if (InstanceGeometryType == EHairGeometryType::Cards)
 	{	
diff --git a/Engine/Plugins/Runtime/HairStrands/Source/HairStrandsCore/Private/HairStrandsVertexFactory.cpp b/Engine/Plugins/Runtime/HairStrands/Source/HairStrandsCore/Private/HairStrandsVertexFactory.cpp
--- a/Engine/Plugins/Runtime/HairStrands/Source/HairStrandsCore/Private/HairStrandsVertexFactory.cpp
+++ b/Engine/Plugins/Runtime/HairStrands/Source/HairStrandsCore/Private/HairStrandsVertexFactory.cpp
@@ -12,6 +12,7 @@
 #include "MeshMaterialShader.h"
 #include "HairStrandsInterface.h"
 #include "GroomInstance.h"
+#include "GroomVisualizationData.h"
 
 #define VF_STRANDS_SUPPORT_GPU_SCENE 1
 #define VF_STRANDS_PROCEDURAL_INTERSECTOR 1
@@ -62,11 +63,11 @@ TGlobalResource<FDummyCulledDispatchVertexIdsBuffer> GDummyCulledDispatchVertexI
 
 /////////////////////////////////////////////////////////////////////////////////////////
 
-FHairGroupPublicData::FVertexFactoryInput ComputeHairStrandsVertexInputData(const FHairGroupInstance* Instance);
+FHairGroupPublicData::FVertexFactoryInput ComputeHairStrandsVertexInputData(const FHairGroupInstance* Instance, EGroomViewMode ViewMode);
 int GetHairRaytracingProceduralSplits();
-FHairStrandsVertexFactoryUniformShaderParameters FHairGroupInstance::GetHairStandsUniformShaderParameters() const
+FHairStrandsVertexFactoryUniformShaderParameters FHairGroupInstance::GetHairStandsUniformShaderParameters(EGroomViewMode ViewMode) const
 {
-	const FHairGroupPublicData::FVertexFactoryInput VFInput = ComputeHairStrandsVertexInputData(this);
+	const FHairGroupPublicData::FVertexFactoryInput VFInput = ComputeHairStrandsVertexInputData(this, ViewMode);
 
 	FHairStrandsVertexFactoryUniformShaderParameters Out = {};
 	Out.GroupIndex					= Debug.GroupIndex;
@@ -232,7 +233,7 @@ void FHairStrandsVertexFactory::InitResources()
 	check(IsValidRef(GetDeclaration()));
 
 	// create the buffer
-	FHairStrandsVertexFactoryUniformShaderParameters Parameters = Data.Instance->GetHairStandsUniformShaderParameters();
+	FHairStrandsVertexFactoryUniformShaderParameters Parameters = Data.Instance->GetHairStandsUniformShaderParameters(EGroomViewMode::None);
 	Data.Instance->Strands.UniformBuffer = FHairStrandsUniformBuffer::CreateUniformBufferImmediate(Parameters, UniformBuffer_MultiFrame);
 }
 
diff --git a/Engine/Plugins/Runtime/HairStrands/Source/HairStrandsCore/Public/GroomComponent.h b/Engine/Plugins/Runtime/HairStrands/Source/HairStrandsCore/Public/GroomComponent.h
--- a/Engine/Plugins/Runtime/HairStrands/Source/HairStrandsCore/Public/GroomComponent.h
+++ b/Engine/Plugins/Runtime/HairStrands/Source/HairStrandsCore/Public/GroomComponent.h
@@ -176,9 +176,6 @@ public:
 	void ValidateMaterials(bool bMapCheck) const;
 	void Invalidate();
 	void InvalidateAndRecreate();
-
-	void SetDebugMode(EHairStrandsDebugMode InMode);
-	EHairStrandsDebugMode GetDebugMode() const { return DebugMode; }
 #endif
 
 	/* Accessor function for changing Groom asset from blueprint/sequencer */
@@ -352,9 +349,6 @@ private:
 	float LODPredictedIndex = -1.f;
 	float LODForcedIndex = -1.f;
 
-	// Transient property for visualizing the groom in a certain debug mode. This is used by the groom editor
-	EHairStrandsDebugMode DebugMode = EHairStrandsDebugMode::NoneDebug;
-
 	void InitResources(bool bIsBindingReloading=false);
 	void ReleaseResources();
 
diff --git a/Engine/Plugins/Runtime/HairStrands/Source/HairStrandsCore/Public/GroomInstance.h b/Engine/Plugins/Runtime/HairStrands/Source/HairStrandsCore/Public/GroomInstance.h
--- a/Engine/Plugins/Runtime/HairStrands/Source/HairStrandsCore/Public/GroomInstance.h
+++ b/Engine/Plugins/Runtime/HairStrands/Source/HairStrandsCore/Public/GroomInstance.h
@@ -14,6 +14,7 @@ class UGroomComponent;
 class FHairCardsVertexFactory;
 class FHairStrandsVertexFactory;
 enum class EGroomCacheType : uint8;
+enum class EGroomViewMode : uint8;
 
 // @hair_todo: pack card ID + card UV in 32Bits alpha channel's of the position buffer:
 //  * 10/10 bits for UV -> max 1024/1024 rect resolution
@@ -236,7 +237,6 @@ struct HAIRSTRANDSCORE_API FHairGroupInstance : public FHairStrandsInstance
 	struct FDebug
 	{
 		// Data
-		EHairStrandsDebugMode	DebugMode = EHairStrandsDebugMode::NoneDebug;
 		uint32					ComponentId = ~0;
 		uint32					GroupIndex = ~0;
 		uint32					GroupCount = 0;
@@ -258,7 +258,6 @@ struct HAIRSTRANDSCORE_API FHairGroupInstance : public FHairStrandsInstance
 		FTransform				SkinningCurrentLocalToWorld = FTransform::Identity;
 		FTransform				RigidPreviousLocalToWorld = FTransform::Identity;
 		FTransform				SkinningPreviousLocalToWorld = FTransform::Identity;
-		bool					bDrawCardsGuides = false;
 
 		TSharedPtr<class IGroomCacheBuffers, ESPMode::ThreadSafe> GroomCacheBuffers;
 
@@ -309,5 +308,5 @@ struct HAIRSTRANDSCORE_API FHairGroupInstance : public FHairStrandsInstance
 															 Debug.RigidPreviousLocalToWorld;
 	}
 
-	FHairStrandsVertexFactoryUniformShaderParameters GetHairStandsUniformShaderParameters() const;
+	FHairStrandsVertexFactoryUniformShaderParameters GetHairStandsUniformShaderParameters(EGroomViewMode ViewMode) const;
 };
diff --git a/Engine/Plugins/Runtime/HairStrands/Source/HairStrandsCore/Public/HairStrandsRendering.h b/Engine/Plugins/Runtime/HairStrands/Source/HairStrandsCore/Public/HairStrandsRendering.h
--- a/Engine/Plugins/Runtime/HairStrands/Source/HairStrandsCore/Public/HairStrandsRendering.h
+++ b/Engine/Plugins/Runtime/HairStrands/Source/HairStrandsCore/Public/HairStrandsRendering.h
@@ -28,6 +28,7 @@ void ComputeHairStrandsInterpolation(
 	FGlobalShaderMap* ShaderMap,
 	const uint32 ViewUniqueID,
 	const uint32 ViewRayTracingMask,
+	const EGroomViewMode ViewMode,
 	const FVector& TranslatedWorldOffset,
 	const struct FShaderPrintData* ShaderPrintData,
 	struct FHairGroupInstance* Instance,
diff --git a/Engine/Plugins/Runtime/HairStrands/Source/HairStrandsEditor/HairStrandsEditor.Build.cs b/Engine/Plugins/Runtime/HairStrands/Source/HairStrandsEditor/HairStrandsEditor.Build.cs
--- a/Engine/Plugins/Runtime/HairStrands/Source/HairStrandsEditor/HairStrandsEditor.Build.cs
+++ b/Engine/Plugins/Runtime/HairStrands/Source/HairStrandsEditor/HairStrandsEditor.Build.cs
@@ -11,7 +11,7 @@ namespace UnrealBuildTool.Rules
 				{
 					"Core",
 					"CoreUObject",
-					"EditorFramework",					
+					"EditorFramework",
 					"Engine",
 					"GeometryCache",
 					"HairStrandsCore",
@@ -34,7 +34,8 @@ namespace UnrealBuildTool.Rules
 					"MovieScene",
 					"MovieSceneTools",
 					"Sequencer",
-					"HairCardGeneratorFramework"
+					"HairCardGeneratorFramework",
+					"CommonMenuExtensions",
 				});
 			AddEngineThirdPartyPrivateStaticDependencies(Target,
 			 "FBX"
diff --git a/Engine/Plugins/Runtime/HairStrands/Source/HairStrandsEditor/Private/GroomEditorCommands.cpp b/Engine/Plugins/Runtime/HairStrands/Source/HairStrandsEditor/Private/GroomEditorCommands.cpp
--- a/Engine/Plugins/Runtime/HairStrands/Source/HairStrandsEditor/Private/GroomEditorCommands.cpp
+++ b/Engine/Plugins/Runtime/HairStrands/Source/HairStrandsEditor/Private/GroomEditorCommands.cpp
@@ -27,24 +27,6 @@ void FGroomEditorCommands::RegisterCommands()
 
 	UI_COMMAND(PlayAnimation,		"Play",		"Play animation",						EUserInterfaceActionType::Button, FInputChord());
 	UI_COMMAND(StopAnimation,		"Stop",		"Stop animation",						EUserInterfaceActionType::Button, FInputChord());
-
-	UI_COMMAND(ViewMode_Lit,			"Lit",				"Default",						EUserInterfaceActionType::Button, FInputChord());
-	UI_COMMAND(ViewMode_Guide,			"Guide",			"View guide (only available if the groom has simulation enabled)",					EUserInterfaceActionType::Button, FInputChord());
-	UI_COMMAND(ViewMode_GuideInfluence, "GuideInfluence",	"View guide influence (only available if the groom has simulation enabled)",		EUserInterfaceActionType::Button, FInputChord());
-	UI_COMMAND(ViewMode_UV,				"UV",				"View strands UVs ",			EUserInterfaceActionType::Button, FInputChord());
-	UI_COMMAND(ViewMode_RootUV,			"RootUV",			"View roots UVs ",				EUserInterfaceActionType::Button, FInputChord());
-	UI_COMMAND(ViewMode_RootUDIM,		"RootUDIM",			"View root UDIM UVs",			EUserInterfaceActionType::Button, FInputChord());
-	UI_COMMAND(ViewMode_Seed,			"Seed",				"View root seed",				EUserInterfaceActionType::Button, FInputChord());
-	UI_COMMAND(ViewMode_Dimension,		"Dimension",		"View strands dimension",		EUserInterfaceActionType::Button, FInputChord());
-	UI_COMMAND(ViewMode_RadiusVariation,"RadiusVariation",	"View strands radius variation",EUserInterfaceActionType::Button, FInputChord());
-	UI_COMMAND(ViewMode_Tangent,		"Tangent",			"View tangents direction",		EUserInterfaceActionType::Button, FInputChord());
-	UI_COMMAND(ViewMode_BaseColor,		"BaseColor",		"View vertex base color",		EUserInterfaceActionType::Button, FInputChord());
-	UI_COMMAND(ViewMode_Roughness,		"Roughness",		"View vertex roughness",		EUserInterfaceActionType::Button, FInputChord());
-	UI_COMMAND(ViewMode_ControlPoints,	"Hair CVs",			"View hair strand CVs",			EUserInterfaceActionType::Button, FInputChord());
-	UI_COMMAND(ViewMode_VisCluster,		"Vis. Clusters",	"View clusters (only available if the groom has simulation enabled)",				EUserInterfaceActionType::Button, FInputChord());
-	UI_COMMAND(ViewMode_Group,			"Hair Groups",		"View hair groups",				EUserInterfaceActionType::Button, FInputChord());
-
-	UI_COMMAND(ViewMode_CardsGuides,	"CardsGuide",		"Cards Guides (only available if the groom has cards geometry)",					EUserInterfaceActionType::Button, FInputChord());
 }
 
 #undef LOCTEXT_NAMESPACE
diff --git a/Engine/Plugins/Runtime/HairStrands/Source/HairStrandsEditor/Private/GroomEditorViewportToolBar.cpp b/Engine/Plugins/Runtime/HairStrands/Source/HairStrandsEditor/Private/GroomEditorViewportToolBar.cpp
--- a/Engine/Plugins/Runtime/HairStrands/Source/HairStrandsEditor/Private/GroomEditorViewportToolBar.cpp
+++ b/Engine/Plugins/Runtime/HairStrands/Source/HairStrandsEditor/Private/GroomEditorViewportToolBar.cpp
@@ -15,6 +15,8 @@
 #include "Widgets/SBoxPanel.h"
 #include "Widgets/SGroomEditorViewport.h"
 #include "Widgets/Text/STextBlock.h"
+#include "GroomVisualizationMenuCommands.h"
+#include "EditorViewportCommands.h"
 
 #define LOCTEXT_NAMESPACE "GroomEditorViewportToolbar"
 
@@ -38,39 +40,10 @@ TSharedRef<SWidget> SGroomEditorViewportToolbar::GenerateShowMenu() const
 	const bool bInShouldCloseWindowAfterMenuSelection = true;
 	FMenuBuilder ViewMenuBuilder(bInShouldCloseWindowAfterMenuSelection, ViewportRef->GetCommandList());
 	{
-		auto Commands = FGroomEditorCommands::Get();
-
-		
-		ViewMenuBuilder.AddMenuEntry(Commands.ViewMode_Lit);
-		ViewMenuBuilder.AddMenuSeparator();
-		
-		ViewMenuBuilder.AddMenuEntry(Commands.ViewMode_Guide);
-		ViewMenuBuilder.AddMenuEntry(Commands.ViewMode_GuideInfluence);
-
-		ViewMenuBuilder.AddMenuSeparator();
-		
-		ViewMenuBuilder.AddMenuEntry(Commands.ViewMode_UV);
-		ViewMenuBuilder.AddMenuEntry(Commands.ViewMode_RootUV);
-		ViewMenuBuilder.AddMenuEntry(Commands.ViewMode_RootUDIM);
-		ViewMenuBuilder.AddMenuEntry(Commands.ViewMode_Seed);
-		ViewMenuBuilder.AddMenuEntry(Commands.ViewMode_Dimension);
-		ViewMenuBuilder.AddMenuEntry(Commands.ViewMode_RadiusVariation);
-		ViewMenuBuilder.AddMenuEntry(Commands.ViewMode_Tangent);
-		ViewMenuBuilder.AddMenuEntry(Commands.ViewMode_BaseColor);
-		ViewMenuBuilder.AddMenuEntry(Commands.ViewMode_Roughness);
-
-		ViewMenuBuilder.AddMenuSeparator();
-
-		ViewMenuBuilder.AddMenuEntry(Commands.ViewMode_ControlPoints);
-		ViewMenuBuilder.AddMenuEntry(Commands.ViewMode_VisCluster);
-
-		ViewMenuBuilder.AddMenuSeparator();
-
-		ViewMenuBuilder.AddMenuEntry(Commands.ViewMode_Group);
-
+		const FEditorViewportCommands& ViewportCommands = FEditorViewportCommands::Get();
+		ViewMenuBuilder.AddMenuEntry(ViewportCommands.LitMode);
 		ViewMenuBuilder.AddMenuSeparator();
-
-		ViewMenuBuilder.AddMenuEntry(Commands.ViewMode_CardsGuides);
+		FGroomVisualizationMenuCommands::BuildVisualisationSubMenuForGroomEditor(ViewMenuBuilder);
 	}
 
 	return ViewMenuBuilder.MakeWidget();
diff --git a/Engine/Plugins/Runtime/HairStrands/Source/HairStrandsEditor/Private/Widgets/SGroomEditorViewport.cpp b/Engine/Plugins/Runtime/HairStrands/Source/HairStrandsEditor/Private/Widgets/SGroomEditorViewport.cpp
--- a/Engine/Plugins/Runtime/HairStrands/Source/HairStrandsEditor/Private/Widgets/SGroomEditorViewport.cpp
+++ b/Engine/Plugins/Runtime/HairStrands/Source/HairStrandsEditor/Private/Widgets/SGroomEditorViewport.cpp
@@ -19,6 +19,8 @@
 #include "Components/StaticMeshComponent.h"
 #include "GroomEditorCommands.h"
 #include "GroomEditorViewportToolBar.h"
+#include "GroomVisualizationMenuCommands.h"
+#include "EditorViewportCommands.h"
 
 #define LOCTEXT_NAMESPACE "SGroomEditorViewport"
 
@@ -188,79 +190,8 @@ void SGroomEditorViewport::BindCommands()
 		FIsActionChecked::CreateSP(this, &SGroomEditorViewport::IsLODModelSelected, 0));
 	// all other LODs will be added dynamically
 	
-	CommandList->MapAction(FGroomEditorCommands::Get().ViewMode_Lit,				FExecuteAction::CreateSP(this, &SGroomEditorViewport::OnViewMode, EHairStrandsDebugMode::NoneDebug),				FCanExecuteAction::CreateSP(this, &SGroomEditorViewport::CanViewMode, false));
-	CommandList->MapAction(FGroomEditorCommands::Get().ViewMode_Guide,				FExecuteAction::CreateSP(this, &SGroomEditorViewport::OnViewMode, EHairStrandsDebugMode::SimHairStrands),			FCanExecuteAction::CreateSP(this, &SGroomEditorViewport::CanViewMode, true));
-	CommandList->MapAction(FGroomEditorCommands::Get().ViewMode_GuideInfluence,		FExecuteAction::CreateSP(this, &SGroomEditorViewport::OnViewMode, EHairStrandsDebugMode::RenderHairStrands),		FCanExecuteAction::CreateSP(this, &SGroomEditorViewport::CanViewMode, true));
-	CommandList->MapAction(FGroomEditorCommands::Get().ViewMode_UV,					FExecuteAction::CreateSP(this, &SGroomEditorViewport::OnViewMode, EHairStrandsDebugMode::RenderHairUV),				FCanExecuteAction::CreateSP(this, &SGroomEditorViewport::CanViewMode, false));
-	CommandList->MapAction(FGroomEditorCommands::Get().ViewMode_RootUV,				FExecuteAction::CreateSP(this, &SGroomEditorViewport::OnViewMode, EHairStrandsDebugMode::RenderHairRootUV),			FCanExecuteAction::CreateSP(this, &SGroomEditorViewport::CanViewMode, false));
-	CommandList->MapAction(FGroomEditorCommands::Get().ViewMode_RootUDIM,			FExecuteAction::CreateSP(this, &SGroomEditorViewport::OnViewMode, EHairStrandsDebugMode::RenderHairRootUDIM),		FCanExecuteAction::CreateSP(this, &SGroomEditorViewport::CanViewMode, false));
-	CommandList->MapAction(FGroomEditorCommands::Get().ViewMode_Seed,				FExecuteAction::CreateSP(this, &SGroomEditorViewport::OnViewMode, EHairStrandsDebugMode::RenderHairSeed),			FCanExecuteAction::CreateSP(this, &SGroomEditorViewport::CanViewMode, false));
-	CommandList->MapAction(FGroomEditorCommands::Get().ViewMode_Dimension,			FExecuteAction::CreateSP(this, &SGroomEditorViewport::OnViewMode, EHairStrandsDebugMode::RenderHairDimension),		FCanExecuteAction::CreateSP(this, &SGroomEditorViewport::CanViewMode, false));
-	CommandList->MapAction(FGroomEditorCommands::Get().ViewMode_RadiusVariation,	FExecuteAction::CreateSP(this, &SGroomEditorViewport::OnViewMode, EHairStrandsDebugMode::RenderHairRadiusVariation),FCanExecuteAction::CreateSP(this, &SGroomEditorViewport::CanViewMode, false));
-	CommandList->MapAction(FGroomEditorCommands::Get().ViewMode_Tangent,			FExecuteAction::CreateSP(this, &SGroomEditorViewport::OnViewMode, EHairStrandsDebugMode::RenderHairTangent),		FCanExecuteAction::CreateSP(this, &SGroomEditorViewport::CanViewMode, false));
-	CommandList->MapAction(FGroomEditorCommands::Get().ViewMode_BaseColor,			FExecuteAction::CreateSP(this, &SGroomEditorViewport::OnViewMode, EHairStrandsDebugMode::RenderHairBaseColor),		FCanExecuteAction::CreateSP(this, &SGroomEditorViewport::CanViewMode, false));
-	CommandList->MapAction(FGroomEditorCommands::Get().ViewMode_Roughness,			FExecuteAction::CreateSP(this, &SGroomEditorViewport::OnViewMode, EHairStrandsDebugMode::RenderHairRoughness),		FCanExecuteAction::CreateSP(this, &SGroomEditorViewport::CanViewMode, false));
-	CommandList->MapAction(FGroomEditorCommands::Get().ViewMode_ControlPoints,		FExecuteAction::CreateSP(this, &SGroomEditorViewport::OnViewMode, EHairStrandsDebugMode::RenderHairControlPoints),	FCanExecuteAction::CreateSP(this, &SGroomEditorViewport::CanViewMode, false));
-	CommandList->MapAction(FGroomEditorCommands::Get().ViewMode_VisCluster,			FExecuteAction::CreateSP(this, &SGroomEditorViewport::OnViewMode, EHairStrandsDebugMode::RenderVisCluster),			FCanExecuteAction::CreateSP(this, &SGroomEditorViewport::CanViewMode, true)); // While Clusters LOD/Culling doesn't require simulation, cluster visualization requires HairInterpolation to run, which is only available when simulation is enabled in the Groom editor.
-	CommandList->MapAction(FGroomEditorCommands::Get().ViewMode_Group,				FExecuteAction::CreateSP(this, &SGroomEditorViewport::OnViewMode, EHairStrandsDebugMode::RenderHairGroup),			FCanExecuteAction::CreateSP(this, &SGroomEditorViewport::CanViewMode, false));
-	CommandList->MapAction(FGroomEditorCommands::Get().ViewMode_CardsGuides,		FExecuteAction::CreateSP(this, &SGroomEditorViewport::OnCardsGuides),												FCanExecuteAction::CreateSP(this, &SGroomEditorViewport::CanCardsGuides));
-}
-
-void SGroomEditorViewport::OnViewMode(EHairStrandsDebugMode Mode)
-{
-	if (GroomComponent)
-	{
-		const EHairStrandsDebugMode OldMode = GroomComponent->GetDebugMode();
-		if (OldMode != Mode)
-		{
-			GroomComponent->SetDebugMode(Mode);
-		}
-		else
-		{
-			GroomComponent->SetDebugMode(EHairStrandsDebugMode::NoneDebug);
-		}
-	}
-}
-
-bool SGroomEditorViewport::CanViewMode(bool bRequireSimulation) const
-{
-	if (bRequireSimulation)
-	{
-		return GroomComponent && GroomComponent->GroomAsset != nullptr && GroomComponent->GroomAsset->IsSimulationEnable();
-	}
-	else
-	{
-		return GroomComponent && GroomComponent->GroomAsset != nullptr;
-	}
-}
-
-void SGroomEditorViewport::OnCardsGuides()
-{
-	bool bEnabled = true;
-	if (GroomComponent && GroomComponent->GroomAsset != nullptr)
-	{
-		for (uint32 GroupIt=0, GroupCount=GroomComponent->GetGroupCount(); GroupIt < GroupCount; ++GroupIt)
-		{
-			FHairGroupInstance* Instance = GroomComponent->GetGroupInstance(GroupIt);
-			Instance->Debug.bDrawCardsGuides = !Instance->Debug.bDrawCardsGuides;
-		}
-	}
-}
-
-bool SGroomEditorViewport::CanCardsGuides() const
-{
-	if (GroomComponent && GroomComponent->GroomAsset != nullptr)
-	{
-		for (uint32 GroupIt = 0, GroupCount = GroomComponent->GetGroupCount(); GroupIt < GroupCount; ++GroupIt)
-		{
-			const FHairGroupInstance* Instance = GroomComponent->GetGroupInstance(GroupIt); 
-			if (Instance->GeometryType == EHairGeometryType::Cards)
-			{
-				return GroomComponent->GroomAsset->IsSimulationEnable();
-			}
-		}
-	}
-	return false;
+	const FGroomVisualizationMenuCommands& GroomCommands = FGroomVisualizationMenuCommands::Get();
+	GroomCommands.BindCommands(*CommandList, SystemViewportClient);
 }
 
 void SGroomEditorViewport::RefreshViewport()
diff --git a/Engine/Plugins/Runtime/HairStrands/Source/HairStrandsEditor/Public/GroomEditorCommands.h b/Engine/Plugins/Runtime/HairStrands/Source/HairStrandsEditor/Public/GroomEditorCommands.h
--- a/Engine/Plugins/Runtime/HairStrands/Source/HairStrandsEditor/Public/GroomEditorCommands.h
+++ b/Engine/Plugins/Runtime/HairStrands/Source/HairStrandsEditor/Public/GroomEditorCommands.h
@@ -23,24 +23,6 @@ public:
 
 	TSharedPtr<FUICommandInfo> Simulate;
 
-	TSharedPtr<FUICommandInfo> ViewMode_Lit;
-	TSharedPtr<FUICommandInfo> ViewMode_Guide;
-	TSharedPtr<FUICommandInfo> ViewMode_GuideInfluence;
-	TSharedPtr<FUICommandInfo> ViewMode_UV;
-	TSharedPtr<FUICommandInfo> ViewMode_RootUV;
-	TSharedPtr<FUICommandInfo> ViewMode_RootUDIM;
-	TSharedPtr<FUICommandInfo> ViewMode_Seed;
-	TSharedPtr<FUICommandInfo> ViewMode_Dimension;
-	TSharedPtr<FUICommandInfo> ViewMode_RadiusVariation;
-	TSharedPtr<FUICommandInfo> ViewMode_Tangent;
-	TSharedPtr<FUICommandInfo> ViewMode_BaseColor;
-	TSharedPtr<FUICommandInfo> ViewMode_Roughness;
-	TSharedPtr<FUICommandInfo> ViewMode_ControlPoints;
-	TSharedPtr<FUICommandInfo> ViewMode_VisCluster;
-	TSharedPtr<FUICommandInfo> ViewMode_Group;
-
-	TSharedPtr<FUICommandInfo> ViewMode_CardsGuides;
-
 	virtual void RegisterCommands() override;
 };
 
diff --git a/Engine/Plugins/Runtime/HairStrands/Source/HairStrandsEditor/Public/Widgets/SGroomEditorViewport.h b/Engine/Plugins/Runtime/HairStrands/Source/HairStrandsEditor/Public/Widgets/SGroomEditorViewport.h
--- a/Engine/Plugins/Runtime/HairStrands/Source/HairStrandsEditor/Public/Widgets/SGroomEditorViewport.h
+++ b/Engine/Plugins/Runtime/HairStrands/Source/HairStrandsEditor/Public/Widgets/SGroomEditorViewport.h
@@ -58,15 +58,8 @@ public:
 	int32 GetLODSelection() const;
 	int32 GetLODModelCount() const;
 	bool IsLODModelSelected(int32 InLODSelection) const;
-
 	void OnSetLODModel(int32 InLODSelection);
 
-	void OnViewMode(EHairStrandsDebugMode Mode);
-	bool CanViewMode(bool bGuideMode) const;
-
-	void OnCardsGuides();
-	bool CanCardsGuides() const;
-
 protected:
 
 	/** SEditorViewport interface */
diff --git a/Engine/Source/Editor/CommonMenuExtensions/Private/CommonMenuExtensionsModule.cpp b/Engine/Source/Editor/CommonMenuExtensions/Private/CommonMenuExtensionsModule.cpp
--- a/Engine/Source/Editor/CommonMenuExtensions/Private/CommonMenuExtensionsModule.cpp
+++ b/Engine/Source/Editor/CommonMenuExtensions/Private/CommonMenuExtensionsModule.cpp
@@ -5,6 +5,7 @@
 #include "NaniteVisualizationMenuCommands.h"
 #include "LumenVisualizationMenuCommands.h"
 #include "StrataVisualizationMenuCommands.h"
+#include "GroomVisualizationMenuCommands.h"
 #include "VirtualShadowMapVisualizationMenuCommands.h"
 #include "ShowFlagMenuCommands.h"
 #include "Modules/ModuleManager.h"
@@ -17,6 +18,7 @@ void FCommonMenuExtensionsModule::StartupModule()
 	FNaniteVisualizationMenuCommands::Register();
 	FLumenVisualizationMenuCommands::Register();
 	FStrataVisualizationMenuCommands::Register();
+	FGroomVisualizationMenuCommands::Register();
 	FVirtualShadowMapVisualizationMenuCommands::Register();
 	FShowFlagMenuCommands::Register();
 }
@@ -26,6 +28,7 @@ void FCommonMenuExtensionsModule::ShutdownModule()
 	FShowFlagMenuCommands::Unregister();
 	FVirtualShadowMapVisualizationMenuCommands::Unregister();
 	FNaniteVisualizationMenuCommands::Unregister();
+	FGroomVisualizationMenuCommands::Unregister();
 	FStrataVisualizationMenuCommands::Unregister();
 	FLumenVisualizationMenuCommands::Unregister();
 	FBufferVisualizationMenuCommands::Unregister();
diff --git a/Engine/Source/Editor/CommonMenuExtensions/Private/GroomVisualizationMenuCommands.cpp b/Engine/Source/Editor/CommonMenuExtensions/Private/GroomVisualizationMenuCommands.cpp
new file mode 100644
--- /dev/null
+++ b/Engine/Source/Editor/CommonMenuExtensions/Private/GroomVisualizationMenuCommands.cpp
@@ -0,0 +1,206 @@
+// Copyright Epic Games, Inc. All Rights Reserved.
+
+#include "GroomVisualizationMenuCommands.h"
+
+#include "Delegates/Delegate.h"
+#include "EditorViewportClient.h"
+#include "Framework/Commands/InputChord.h"
+#include "Framework/Commands/UIAction.h"
+#include "Framework/Commands/UICommandInfo.h"
+#include "Framework/Commands/UICommandList.h"
+#include "Framework/MultiBox/MultiBoxBuilder.h"
+#include "Internationalization/Internationalization.h"
+#include "Internationalization/Text.h"
+#include "GroomVisualizationData.h"
+#include "Misc/AssertionMacros.h"
+#include "Styling/AppStyle.h"
+#include "UObject/UnrealNames.h"
+#include "GroomVisualizationData.h"
+
+#define LOCTEXT_NAMESPACE "GroomVisualizationMenuCommands"
+
+FGroomVisualizationMenuCommands::FGroomVisualizationMenuCommands()
+	: TCommands<FGroomVisualizationMenuCommands>
+	(
+		TEXT("GroomVisualizationMenu"), // Context name for fast lookup
+		NSLOCTEXT("Contexts", "GroomVisualizationMenu", "Groom"), // Localized context name for displaying
+		NAME_None, // Parent context name.  
+		FAppStyle::GetAppStyleSetName() // Icon Style Set
+	),
+	CommandMap()
+{
+}
+
+void FGroomVisualizationMenuCommands::BuildCommandMap()
+{
+	const FGroomVisualizationData& VisualizationData = GetGroomVisualizationData();
+	const FGroomVisualizationData::TModeMap& ModeMap = VisualizationData.GetModeMap();
+
+	CommandMap.Empty();
+	for (FGroomVisualizationData::TModeMap::TConstIterator It = ModeMap.CreateConstIterator(); It; ++It)
+	{
+		const FGroomVisualizationData::FModeRecord& Entry = It.Value();
+		FGroomVisualizationRecord& Record = CommandMap.Add(Entry.ModeName, FGroomVisualizationRecord());
+		Record.Name = Entry.ModeName;
+		Record.Mode = Entry.Mode;
+		Record.Command = FUICommandInfoDecl(
+			this->AsShared(),
+			Entry.ModeName,
+			Entry.ModeText,
+			Entry.ModeDesc)
+			.UserInterfaceType(EUserInterfaceActionType::RadioButton)
+			.DefaultChord(FInputChord()
+		);
+
+	}
+}
+
+void FGroomVisualizationMenuCommands::InternalBuildVisualisationSubMenu(FMenuBuilder& Menu, bool bIsGroomEditor)
+{
+	const FGroomVisualizationMenuCommands& Commands = FGroomVisualizationMenuCommands::Get();
+	if (Commands.IsPopulated())
+	{
+		// General
+		{
+			Menu.BeginSection("LevelViewportGroomVisualizationGeneral", LOCTEXT("GroomVisualizationGeneral", "Groom General"));
+			if (!bIsGroomEditor)
+			{
+				Commands.AddCommandTypeToMenu(Menu, EGroomViewMode::MacroGroups);
+			}
+			Commands.AddCommandTypeToMenu(Menu, EGroomViewMode::RenderHairGroup);
+			Commands.AddCommandTypeToMenu(Menu, EGroomViewMode::RenderLODColoration);
+			Menu.EndSection();
+		}
+
+		// Attributes
+		{
+			Menu.BeginSection("LevelViewportGroomVisualizationGeneral", LOCTEXT("GroomVisualizationGeneral", "Groom Attributes"));
+			Commands.AddCommandTypeToMenu(Menu, EGroomViewMode::RenderHairSeed);
+			Commands.AddCommandTypeToMenu(Menu, EGroomViewMode::RenderHairRootUV);
+			Commands.AddCommandTypeToMenu(Menu, EGroomViewMode::RenderHairRootUDIM);
+			Commands.AddCommandTypeToMenu(Menu, EGroomViewMode::RenderHairUV);
+			Commands.AddCommandTypeToMenu(Menu, EGroomViewMode::RenderHairDimension);
+			Commands.AddCommandTypeToMenu(Menu, EGroomViewMode::RenderHairRadiusVariation);
+			Commands.AddCommandTypeToMenu(Menu, EGroomViewMode::RenderHairTangent);
+			Menu.EndSection();
+
+		}
+		// Guides
+		{
+			Menu.BeginSection("LevelViewportGroomVisualizationGeneral", LOCTEXT("GroomVisualizationGuides", "Groom Guides"));
+			Commands.AddCommandTypeToMenu(Menu, EGroomViewMode::SimHairStrands);
+			Commands.AddCommandTypeToMenu(Menu, EGroomViewMode::RenderHairStrands);
+			Commands.AddCommandTypeToMenu(Menu, EGroomViewMode::CardGuides);
+			Menu.EndSection();
+		}
+
+		// Strands
+		{
+			Menu.BeginSection("LevelViewportGroomVisualizationGeneral", LOCTEXT("GroomVisualizationMaterial", "Groom Strands"));
+			Commands.AddCommandTypeToMenu(Menu, EGroomViewMode::RenderHairControlPoints);
+			Commands.AddCommandTypeToMenu(Menu, EGroomViewMode::MaterialDepth);
+			Commands.AddCommandTypeToMenu(Menu, EGroomViewMode::MaterialBaseColor);
+			Commands.AddCommandTypeToMenu(Menu, EGroomViewMode::MaterialRoughness);
+			Commands.AddCommandTypeToMenu(Menu, EGroomViewMode::MaterialSpecular);
+			Commands.AddCommandTypeToMenu(Menu, EGroomViewMode::MaterialTangent);
+			Commands.AddCommandTypeToMenu(Menu, EGroomViewMode::RenderHairBaseColor);
+			Commands.AddCommandTypeToMenu(Menu, EGroomViewMode::RenderHairRoughness);
+			Menu.EndSection();
+		}
+
+		// Advanced
+		{
+			Menu.BeginSection("LevelViewportGroomVisualizationAdvanced", LOCTEXT("GroomVisualizationAdvanced", "Groom Strands Advanced"));
+			Commands.AddCommandTypeToMenu(Menu, EGroomViewMode::VoxelsDensity);
+			Commands.AddCommandTypeToMenu(Menu, EGroomViewMode::RenderVisCluster);
+			Commands.AddCommandTypeToMenu(Menu, EGroomViewMode::RenderVisClusterAABB);
+			if (!bIsGroomEditor)
+			{
+				Commands.AddCommandTypeToMenu(Menu, EGroomViewMode::LightBounds);
+				Commands.AddCommandTypeToMenu(Menu, EGroomViewMode::DeepOpacityMaps);
+				Commands.AddCommandTypeToMenu(Menu, EGroomViewMode::MacroGroupScreenRect);
+				Commands.AddCommandTypeToMenu(Menu, EGroomViewMode::SamplePerPixel);
+				Commands.AddCommandTypeToMenu(Menu, EGroomViewMode::CoverageType);
+				Commands.AddCommandTypeToMenu(Menu, EGroomViewMode::TAAResolveType);
+				Commands.AddCommandTypeToMenu(Menu, EGroomViewMode::MeshProjection); // Rename RootBinding
+				Commands.AddCommandTypeToMenu(Menu, EGroomViewMode::Coverage);
+				Commands.AddCommandTypeToMenu(Menu, EGroomViewMode::Tile);
+			}
+			Menu.EndSection();
+		}
+	}
+}
+
+void FGroomVisualizationMenuCommands::BuildVisualisationSubMenu(FMenuBuilder& Menu)
+{
+	InternalBuildVisualisationSubMenu(Menu, false);
+}
+
+void FGroomVisualizationMenuCommands::BuildVisualisationSubMenuForGroomEditor(FMenuBuilder& Menu)
+{
+	InternalBuildVisualisationSubMenu(Menu, true);
+}
+
+
+bool FGroomVisualizationMenuCommands::AddCommandTypeToMenu(FMenuBuilder& Menu, const EGroomViewMode Mode) const
+{
+	bool bAddedCommands = false;
+
+	const TGroomVisualizationModeCommandMap& Commands = CommandMap;
+	for (TCommandConstIterator It = CreateCommandConstIterator(); It; ++It)
+	{
+		const FGroomVisualizationRecord& Record = It.Value();
+		if (Record.Mode == Mode)
+		{
+			Menu.AddMenuEntry(Record.Command, NAME_None, Record.Command->GetLabel());
+			bAddedCommands = true;
+		}
+	}
+
+	return bAddedCommands;
+}
+
+FGroomVisualizationMenuCommands::TCommandConstIterator FGroomVisualizationMenuCommands::CreateCommandConstIterator() const
+{
+	return CommandMap.CreateConstIterator();
+}
+
+void FGroomVisualizationMenuCommands::RegisterCommands()
+{
+	BuildCommandMap();
+}
+
+void FGroomVisualizationMenuCommands::BindCommands(FUICommandList& CommandList, const TSharedPtr<FEditorViewportClient>& Client) const
+{
+	// Map Groom visualization mode actions
+	for (FGroomVisualizationMenuCommands::TCommandConstIterator It = FGroomVisualizationMenuCommands::Get().CreateCommandConstIterator(); It; ++It)
+	{
+		const FGroomVisualizationMenuCommands::FGroomVisualizationRecord& Record = It.Value();
+		CommandList.MapAction(
+			Record.Command,
+			FExecuteAction::CreateStatic(&FGroomVisualizationMenuCommands::ChangeGroomVisualizationMode, Client.ToWeakPtr(), Record.Name),
+			FCanExecuteAction(),
+			FIsActionChecked::CreateStatic(&FGroomVisualizationMenuCommands::IsGroomVisualizationModeSelected, Client.ToWeakPtr(), Record.Name)
+		);
+	}
+}
+
+void FGroomVisualizationMenuCommands::ChangeGroomVisualizationMode(TWeakPtr<FEditorViewportClient> WeakClient, FName InName)
+{
+	if (TSharedPtr<FEditorViewportClient> Client = WeakClient.Pin())
+	{
+		Client->ChangeGroomVisualizationMode(InName);
+	}
+}
+
+bool FGroomVisualizationMenuCommands::IsGroomVisualizationModeSelected(TWeakPtr<FEditorViewportClient> WeakClient, FName InName)
+{
+	if (TSharedPtr<FEditorViewportClient> Client = WeakClient.Pin())
+	{
+		return Client->IsGroomVisualizationModeSelected(InName);
+	}
+	
+	return false;
+}
+
+#undef LOCTEXT_NAMESPACE
diff --git a/Engine/Source/Editor/CommonMenuExtensions/Public/GroomVisualizationMenuCommands.h b/Engine/Source/Editor/CommonMenuExtensions/Public/GroomVisualizationMenuCommands.h
new file mode 100644
--- /dev/null
+++ b/Engine/Source/Editor/CommonMenuExtensions/Public/GroomVisualizationMenuCommands.h
@@ -0,0 +1,58 @@
+// Copyright Epic Games, Inc. All Rights Reserved.
+
+#pragma once
+
+#include "Containers/Map.h"
+#include "Framework/Commands/Commands.h"
+#include "Framework/Commands/UICommandInfo.h"
+#include "Framework/MultiBox/MultiBoxBuilder.h"
+#include "HAL/Platform.h"
+#include "Templates/SharedPointer.h"
+#include "UObject/NameTypes.h"
+
+class FEditorViewportClient;
+class FMenuBuilder;
+class FUICommandInfo;
+class FUICommandList;
+enum class EGroomViewMode : uint8;
+
+class COMMONMENUEXTENSIONS_API FGroomVisualizationMenuCommands : public TCommands<FGroomVisualizationMenuCommands>
+{
+public:
+	struct FGroomVisualizationRecord
+	{
+		FName Name;
+		TSharedPtr<FUICommandInfo> Command;
+		EGroomViewMode Mode;
+	};
+
+	typedef TMultiMap<FName, FGroomVisualizationRecord> TGroomVisualizationModeCommandMap;
+	typedef TGroomVisualizationModeCommandMap::TConstIterator TCommandConstIterator;
+
+	FGroomVisualizationMenuCommands();
+
+	TCommandConstIterator CreateCommandConstIterator() const;
+
+	static void BuildVisualisationSubMenu(FMenuBuilder& Menu);
+	static void BuildVisualisationSubMenuForGroomEditor(FMenuBuilder& Menu);
+
+	virtual void RegisterCommands() override;
+
+	void BindCommands(FUICommandList& CommandList, const TSharedPtr<FEditorViewportClient>& Client) const;
+
+	inline bool IsPopulated() const
+	{
+		return CommandMap.Num() > 0;
+	}
+
+private:
+	void BuildCommandMap();
+	bool AddCommandTypeToMenu(FMenuBuilder& Menu, const EGroomViewMode Type) const;
+
+	static void InternalBuildVisualisationSubMenu(FMenuBuilder& Menu, bool bIsGroomEditor);
+	static void ChangeGroomVisualizationMode(TWeakPtr<FEditorViewportClient> WeakClient, FName InName);
+	static bool IsGroomVisualizationModeSelected(TWeakPtr<FEditorViewportClient> WeakClient, FName InName);
+
+private:
+	TGroomVisualizationModeCommandMap CommandMap;
+};
\ No newline at end of file
diff --git a/Engine/Source/Editor/DetailCustomizations/Private/DebugCameraControllerSettingsCustomization.cpp b/Engine/Source/Editor/DetailCustomizations/Private/DebugCameraControllerSettingsCustomization.cpp
--- a/Engine/Source/Editor/DetailCustomizations/Private/DebugCameraControllerSettingsCustomization.cpp
+++ b/Engine/Source/Editor/DetailCustomizations/Private/DebugCameraControllerSettingsCustomization.cpp
@@ -46,6 +46,7 @@ void FDebugCameraControllerSettingsViewModeIndexCustomization::CustomizeHeader(T
 	EnumRestriction->AddHiddenValue(ViewModeIndexEnum->GetNameStringByValue((uint8)EViewModeIndex::VMI_VisualizeNanite));
 	EnumRestriction->AddHiddenValue(ViewModeIndexEnum->GetNameStringByValue((uint8)EViewModeIndex::VMI_VisualizeLumen));
 	EnumRestriction->AddHiddenValue(ViewModeIndexEnum->GetNameStringByValue((uint8)EViewModeIndex::VMI_VisualizeStrata));
+	EnumRestriction->AddHiddenValue(ViewModeIndexEnum->GetNameStringByValue((uint8)EViewModeIndex::VMI_VisualizeGroom));
 	EnumRestriction->AddHiddenValue(ViewModeIndexEnum->GetNameStringByValue((uint8)EViewModeIndex::VMI_VisualizeVirtualShadowMap));
 	EnumRestriction->AddHiddenValue(ViewModeIndexEnum->GetNameStringByValue((uint8)EViewModeIndex::VMI_VisualizeGPUSkinCache));
 	EnumRestriction->AddHiddenValue(ViewModeIndexEnum->GetNameStringByValue((uint8)EViewModeIndex::VMI_StationaryLightOverlap));
diff --git a/Engine/Source/Editor/EditorStyle/Private/StarshipStyle.cpp b/Engine/Source/Editor/EditorStyle/Private/StarshipStyle.cpp
--- a/Engine/Source/Editor/EditorStyle/Private/StarshipStyle.cpp
+++ b/Engine/Source/Editor/EditorStyle/Private/StarshipStyle.cpp
@@ -4091,8 +4091,9 @@ void FStarshipEditorStyle::FStyle::SetupLevelEditorStyle()
 		Set( "EditorViewport.ReflectionOverrideMode", new IMAGE_BRUSH_SVG("Starship/Common/Reflections", Icon16x16 ) );
 		Set( "EditorViewport.VisualizeBufferMode",    new IMAGE_BRUSH_SVG("Starship/Common/BufferVisualization", Icon16x16 ) );
 		Set( "EditorViewport.VisualizeNaniteMode",    new IMAGE_BRUSH_SVG("Starship/Common/BufferVisualization", Icon16x16 ) );
-		Set( "EditorViewport.VisualizeLumenMode",    new IMAGE_BRUSH_SVG("Starship/Common/BufferVisualization", Icon16x16 ) );
+		Set( "EditorViewport.VisualizeLumenMode",     new IMAGE_BRUSH_SVG("Starship/Common/BufferVisualization", Icon16x16 ) );
 		Set( "EditorViewport.VisualizeStrataMode",    new IMAGE_BRUSH_SVG("Starship/Common/BufferVisualization", Icon16x16 ) );
+		Set( "EditorViewport.VisualizeGroomMode",     new IMAGE_BRUSH_SVG("Starship/Common/BufferVisualization", Icon16x16 ) );
 		Set( "EditorViewport.VisualizeVirtualShadowMapMode", new IMAGE_BRUSH_SVG("Starship/Common/BufferVisualization", Icon16x16 ) );
 		Set( "EditorViewport.CollisionPawn",          new IMAGE_BRUSH_SVG("Starship/Common/PlayerCollision", Icon16x16 ) );
 		Set( "EditorViewport.CollisionVisibility",    new IMAGE_BRUSH_SVG("Starship/Common/VisibilityCollision", Icon16x16 ) );
diff --git a/Engine/Source/Editor/LevelEditor/Private/LevelViewportActions.cpp b/Engine/Source/Editor/LevelEditor/Private/LevelViewportActions.cpp
--- a/Engine/Source/Editor/LevelEditor/Private/LevelViewportActions.cpp
+++ b/Engine/Source/Editor/LevelEditor/Private/LevelViewportActions.cpp
@@ -12,6 +12,7 @@
 #include "NaniteVisualizationData.h"
 #include "LumenVisualizationData.h"
 #include "StrataVisualizationData.h"
+#include "GroomVisualizationData.h"
 #include "VirtualShadowMapVisualizationData.h"
 #include "Bookmarks/BookmarkUI.h"
 #include "UnrealEdGlobals.h"
diff --git a/Engine/Source/Editor/LevelEditor/Private/SLevelViewport.cpp b/Engine/Source/Editor/LevelEditor/Private/SLevelViewport.cpp
--- a/Engine/Source/Editor/LevelEditor/Private/SLevelViewport.cpp
+++ b/Engine/Source/Editor/LevelEditor/Private/SLevelViewport.cpp
@@ -63,6 +63,7 @@
 #include "NaniteVisualizationData.h"
 #include "LumenVisualizationData.h"
 #include "StrataVisualizationData.h"
+#include "GroomVisualizationData.h"
 #include "VirtualShadowMapVisualizationData.h"
 #include "Framework/Notifications/NotificationManager.h"
 #include "Widgets/Notifications/SNotificationList.h"
@@ -500,6 +501,7 @@ void SLevelViewport::ConstructLevelEditorViewportClient(FLevelEditorViewportInst
 	LevelViewportClient->CurrentNaniteVisualizationMode = ViewportInstanceSettings.NaniteVisualizationMode;
 	LevelViewportClient->CurrentLumenVisualizationMode = ViewportInstanceSettings.LumenVisualizationMode;
 	LevelViewportClient->CurrentStrataVisualizationMode = ViewportInstanceSettings.StrataVisualizationMode;
+	LevelViewportClient->CurrentGroomVisualizationMode = ViewportInstanceSettings.GroomVisualizationMode;
 	LevelViewportClient->CurrentVirtualShadowMapVisualizationMode = ViewportInstanceSettings.VirtualShadowMapVisualizationMode;
 	LevelViewportClient->CurrentRayTracingDebugVisualizationMode = ViewportInstanceSettings.RayTracingDebugVisualizationMode;
 	LevelViewportClient->CurrentGPUSkinCacheVisualizationMode = ViewportInstanceSettings.GPUSkinCacheVisualizationMode;
@@ -1553,6 +1555,7 @@ void SLevelViewport::BindViewCommands( FUICommandList& OutCommandList )
 	{
 		FStrataVisualizationMenuCommands::Get().BindCommands(OutCommandList, Client);
 	}
+	FGroomVisualizationMenuCommands::Get().BindCommands(OutCommandList, Client);
 	FVirtualShadowMapVisualizationMenuCommands::Get().BindCommands(OutCommandList, Client);
 }
 
@@ -2253,6 +2256,7 @@ void SLevelViewport::SaveConfig(const FString& ConfigName) const
 		ViewportInstanceSettings.NaniteVisualizationMode = LevelViewportClient->CurrentNaniteVisualizationMode;
 		ViewportInstanceSettings.LumenVisualizationMode = LevelViewportClient->CurrentLumenVisualizationMode;
 		ViewportInstanceSettings.StrataVisualizationMode = LevelViewportClient->CurrentStrataVisualizationMode;
+		ViewportInstanceSettings.GroomVisualizationMode = LevelViewportClient->CurrentGroomVisualizationMode;
 		ViewportInstanceSettings.VirtualShadowMapVisualizationMode = LevelViewportClient->CurrentVirtualShadowMapVisualizationMode;
 		ViewportInstanceSettings.RayTracingDebugVisualizationMode = LevelViewportClient->CurrentRayTracingDebugVisualizationMode;
 		ViewportInstanceSettings.GPUSkinCacheVisualizationMode = LevelViewportClient->CurrentGPUSkinCacheVisualizationMode;
diff --git a/Engine/Source/Editor/LevelEditor/Private/SLevelViewportToolBar.cpp b/Engine/Source/Editor/LevelEditor/Private/SLevelViewportToolBar.cpp
--- a/Engine/Source/Editor/LevelEditor/Private/SLevelViewportToolBar.cpp
+++ b/Engine/Source/Editor/LevelEditor/Private/SLevelViewportToolBar.cpp
@@ -35,6 +35,7 @@
 #include "NaniteVisualizationData.h"
 #include "LumenVisualizationData.h"
 #include "StrataVisualizationData.h"
+#include "GroomVisualizationData.h"
 #include "VirtualShadowMapVisualizationData.h"
 #include "FoliageType.h"
 #include "ShowFlagMenuCommands.h"
@@ -1642,6 +1643,31 @@ void SLevelViewportToolBar::FillViewMenu(UToolMenu* Menu)
 						);
 	}
 
+	if (IsGroomEnabled())
+	{
+		FToolMenuSection& Section = Menu->FindOrAddSection("ViewMode");
+		Section.AddSubMenu(
+			"VisualizeGroomViewMode",
+			LOCTEXT("VisualizeGroomViewModeDisplayName", "Groom"),
+			LOCTEXT("GroomVisualizationMenu_ToolTip", "Select a mode for Groom visualization"),
+			FNewMenuDelegate::CreateStatic(&FGroomVisualizationMenuCommands::BuildVisualisationSubMenu),
+			FUIAction(
+				FExecuteAction(),
+				FCanExecuteAction(),
+				FIsActionChecked::CreateLambda([this]()
+					{
+						const TSharedRef<SEditorViewport> ViewportRef = Viewport.Pin().ToSharedRef();
+						const TSharedPtr<FEditorViewportClient> ViewportClient = ViewportRef->GetViewportClient();
+						check(ViewportClient.IsValid());
+						return ViewportClient->IsViewModeEnabled(VMI_VisualizeGroom);
+					})
+			),
+			EUserInterfaceActionType::RadioButton,
+						/* bInOpenSubMenuOnClick = */ false,
+						FSlateIcon(FAppStyle::GetAppStyleSetName(), "EditorViewport.VisualizeGroomMode")
+						);
+	}
+
 	{
 		FToolMenuSection& Section = Menu->FindOrAddSection("ViewMode");
 		Section.AddSubMenu(
diff --git a/Engine/Source/Editor/LevelEditor/Public/LevelViewportActions.h b/Engine/Source/Editor/LevelEditor/Public/LevelViewportActions.h
--- a/Engine/Source/Editor/LevelEditor/Public/LevelViewportActions.h
+++ b/Engine/Source/Editor/LevelEditor/Public/LevelViewportActions.h
@@ -11,6 +11,7 @@
 #include "NaniteVisualizationMenuCommands.h"
 #include "LumenVisualizationMenuCommands.h"
 #include "StrataVisualizationMenuCommands.h"
+#include "GroomVisualizationMenuCommands.h"
 #include "VirtualShadowMapVisualizationMenuCommands.h"
 
 /**
diff --git a/Engine/Source/Editor/UnrealEd/Classes/Settings/LevelEditorViewportSettings.h b/Engine/Source/Editor/UnrealEd/Classes/Settings/LevelEditorViewportSettings.h
--- a/Engine/Source/Editor/UnrealEd/Classes/Settings/LevelEditorViewportSettings.h
+++ b/Engine/Source/Editor/UnrealEd/Classes/Settings/LevelEditorViewportSettings.h
@@ -99,6 +99,7 @@ struct UNREALED_API FLevelEditorViewportInstanceSettings
 		, NaniteVisualizationMode()
 		, LumenVisualizationMode()
 		, StrataVisualizationMode()
+		, GroomVisualizationMode()
 		, VirtualShadowMapVisualizationMode()
 		, GPUSkinCacheVisualizationMode()
 		, ExposureSettings()
@@ -153,6 +154,10 @@ struct UNREALED_API FLevelEditorViewportInstanceSettings
 	UPROPERTY(config)
 	FName StrataVisualizationMode;
 
+	/** The Groom visualization mode for the viewport. */
+	UPROPERTY(config)
+	FName GroomVisualizationMode;
+
 	/** The virtual shadow map visualization mode for the viewport. */
 	UPROPERTY(config)
 	FName VirtualShadowMapVisualizationMode;
diff --git a/Engine/Source/Editor/UnrealEd/Private/EditorViewportClient.cpp b/Engine/Source/Editor/UnrealEd/Private/EditorViewportClient.cpp
--- a/Engine/Source/Editor/UnrealEd/Private/EditorViewportClient.cpp
+++ b/Engine/Source/Editor/UnrealEd/Private/EditorViewportClient.cpp
@@ -56,6 +56,7 @@
 #include "NaniteVisualizationData.h"
 #include "LumenVisualizationData.h"
 #include "StrataVisualizationData.h"
+#include "GroomVisualizationData.h"
 #include "VirtualShadowMapVisualizationData.h"
 #include "GPUSkinCacheVisualizationData.h"
 #include "UnrealWidget.h"
@@ -394,6 +395,7 @@ FEditorViewportClient::FEditorViewportClient(FEditorModeTools* InModeTools, FPre
 	, CurrentNaniteVisualizationMode(NAME_None)
 	, CurrentLumenVisualizationMode(NAME_None)
 	, CurrentStrataVisualizationMode(NAME_None)
+	, CurrentGroomVisualizationMode(NAME_None)
 	, CurrentVirtualShadowMapVisualizationMode(NAME_None)
 	, CurrentRayTracingDebugVisualizationMode(NAME_None)
 	, CurrentGPUSkinCacheVisualizationMode(NAME_None)
@@ -2699,6 +2701,22 @@ FText FEditorViewportClient::GetCurrentStrataVisualizationModeDisplayName() cons
 	return GetStrataVisualizationData().GetModeDisplayName(CurrentStrataVisualizationMode);
 }
 
+void FEditorViewportClient::ChangeGroomVisualizationMode(FName InName)
+{
+	SetViewMode(VMI_VisualizeGroom);
+	CurrentGroomVisualizationMode = InName;
+}
+
+bool FEditorViewportClient::IsGroomVisualizationModeSelected(FName InName) const
+{
+	return IsViewModeEnabled(VMI_VisualizeGroom) && CurrentGroomVisualizationMode == InName;
+}
+
+FText FEditorViewportClient::GetCurrentGroomVisualizationModeDisplayName() const
+{
+	checkf(IsViewModeEnabled(VMI_VisualizeGroom), TEXT("In order to call GetCurrentGroomVisualizationMode(), first you must set ViewMode to VMI_VisualizeGroom."));
+	return GetGroomVisualizationData().GetModeDisplayName(CurrentGroomVisualizationMode);
+}
 bool FEditorViewportClient::IsVisualizeCalibrationMaterialEnabled() const
 {
 	// Get the list of requested buffers from the console
@@ -3864,6 +3882,7 @@ void FEditorViewportClient::SetupViewForRendering(FSceneViewFamily& ViewFamily,
 	View.CurrentNaniteVisualizationMode = CurrentNaniteVisualizationMode;
 	View.CurrentLumenVisualizationMode = CurrentLumenVisualizationMode;
 	View.CurrentStrataVisualizationMode = CurrentStrataVisualizationMode;
+	View.CurrentGroomVisualizationMode = CurrentGroomVisualizationMode;
 	View.CurrentVirtualShadowMapVisualizationMode = CurrentVirtualShadowMapVisualizationMode;
 	View.CurrentGPUSkinCacheVisualizationMode = CurrentGPUSkinCacheVisualizationMode;
 #if RHI_RAYTRACING
diff --git a/Engine/Source/Editor/UnrealEd/Private/EditorViewportCommands.cpp b/Engine/Source/Editor/UnrealEd/Private/EditorViewportCommands.cpp
--- a/Engine/Source/Editor/UnrealEd/Private/EditorViewportCommands.cpp
+++ b/Engine/Source/Editor/UnrealEd/Private/EditorViewportCommands.cpp
@@ -121,6 +121,7 @@ void FEditorViewportCommands::RegisterCommands()
 	UI_COMMAND( VisualizeNaniteMode, "Nanite Visualization View Mode", "Visualizes various rendering aspects of Nanite.", EUserInterfaceActionType::RadioButton, FInputChord());
 	UI_COMMAND( VisualizeLumenMode, "Lumen Visualization View Mode", "Visualizes Lumen debug views.", EUserInterfaceActionType::RadioButton, FInputChord());
 	UI_COMMAND( VisualizeStrataMode, "Strata Visualization View Mode", "Visualizes Strata debug views.", EUserInterfaceActionType::RadioButton, FInputChord());
+	UI_COMMAND( VisualizeGroomMode, "Groom Visualization View Mode", "Visualizes Groom debug views.", EUserInterfaceActionType::RadioButton, FInputChord());
 	UI_COMMAND( VisualizeVirtualShadowMapMode, "Virtual Shadow Map Visualization View Mode", "Visualizes various rendering aspects of virtual shadow maps.", EUserInterfaceActionType::RadioButton, FInputChord());
 
 	UI_COMMAND( ReflectionOverrideMode, "Reflections View Mode", "Renders the scene with reflections only", EUserInterfaceActionType::RadioButton, FInputChord() );
diff --git a/Engine/Source/Editor/UnrealEd/Private/SEditorViewport.cpp b/Engine/Source/Editor/UnrealEd/Private/SEditorViewport.cpp
--- a/Engine/Source/Editor/UnrealEd/Private/SEditorViewport.cpp
+++ b/Engine/Source/Editor/UnrealEd/Private/SEditorViewport.cpp
@@ -430,6 +430,7 @@ void SEditorViewport::BindCommands()
 	MAP_VIEWMODE_ACTION( Commands.VisualizeNaniteMode, VMI_VisualizeNanite );
 	MAP_VIEWMODE_ACTION( Commands.VisualizeLumenMode, VMI_VisualizeLumen );
 	MAP_VIEWMODE_ACTION( Commands.VisualizeStrataMode, VMI_VisualizeStrata);
+	MAP_VIEWMODE_ACTION( Commands.VisualizeGroomMode, VMI_VisualizeGroom);
 	MAP_VIEWMODE_ACTION( Commands.VisualizeVirtualShadowMapMode, VMI_VisualizeVirtualShadowMap );
 	MAP_VIEWMODE_ACTION( Commands.CollisionPawn, VMI_CollisionPawn);
 	MAP_VIEWMODE_ACTION( Commands.CollisionVisibility, VMI_CollisionVisibility);
diff --git a/Engine/Source/Editor/UnrealEd/Private/SEditorViewportViewMenu.cpp b/Engine/Source/Editor/UnrealEd/Private/SEditorViewportViewMenu.cpp
--- a/Engine/Source/Editor/UnrealEd/Private/SEditorViewportViewMenu.cpp
+++ b/Engine/Source/Editor/UnrealEd/Private/SEditorViewportViewMenu.cpp
@@ -59,6 +59,10 @@ FText SEditorViewportViewMenu::GetViewMenuLabel() const
 		{
 			Label = ViewportClient->GetCurrentStrataVisualizationModeDisplayName();
 		}
+		else if (ViewMode == VMI_VisualizeGroom)
+		{
+			Label = ViewportClient->GetCurrentGroomVisualizationModeDisplayName();
+		}
 		else if (ViewMode == VMI_VisualizeVirtualShadowMap)
 		{
 			Label = ViewportClient->GetCurrentVirtualShadowMapVisualizationModeDisplayName();
diff --git a/Engine/Source/Editor/UnrealEd/Public/EditorViewportClient.h b/Engine/Source/Editor/UnrealEd/Public/EditorViewportClient.h
--- a/Engine/Source/Editor/UnrealEd/Public/EditorViewportClient.h
+++ b/Engine/Source/Editor/UnrealEd/Public/EditorViewportClient.h
@@ -1308,6 +1308,26 @@ public:
 	 */
 	FText GetCurrentStrataVisualizationModeDisplayName() const;
 
+	/**
+	 * Changes the Groom visualization mode for this viewport.
+	 *
+	 * @param InName	The ID of the required visualization mode
+	 */
+	void ChangeGroomVisualizationMode(FName InName);
+
+	/**
+	 * Checks if a Groom visualization mode is selected.
+	 *
+	 * @param InName	The ID of the required visualization mode
+	 * @return	true if the supplied Groom visualization mode is checked
+	 */
+	bool IsGroomVisualizationModeSelected(FName InName) const;
+
+	/**
+	 * Returns the FText display name associated with CurrentGroomVisualizationMode.
+	 */
+	FText GetCurrentGroomVisualizationModeDisplayName() const;
+
 	/**
 	* Changes the virtual shadow map visualization mode for this viewport.
 	*
@@ -1675,6 +1695,7 @@ public:
 	FName CurrentNaniteVisualizationMode;
 	FName CurrentLumenVisualizationMode;
 	FName CurrentStrataVisualizationMode;
+	FName CurrentGroomVisualizationMode;
 	FName CurrentVirtualShadowMapVisualizationMode;
 
 	FName CurrentRayTracingDebugVisualizationMode;
diff --git a/Engine/Source/Editor/UnrealEd/Public/EditorViewportCommands.h b/Engine/Source/Editor/UnrealEd/Public/EditorViewportCommands.h
--- a/Engine/Source/Editor/UnrealEd/Public/EditorViewportCommands.h
+++ b/Engine/Source/Editor/UnrealEd/Public/EditorViewportCommands.h
@@ -139,6 +139,9 @@ public:
 	/** Changes the viewport to visualize Strata debug views */
 	TSharedPtr< FUICommandInfo > VisualizeStrataMode;
 
+	/** Changes the viewport to visualize Groom debug views */
+	TSharedPtr< FUICommandInfo > VisualizeGroomMode;
+
 	/** Changes the viewport to visualize virtual shadow map content */
 	TSharedPtr< FUICommandInfo > VisualizeVirtualShadowMapMode;
 
diff --git a/Engine/Source/Runtime/Engine/Classes/Engine/EngineBaseTypes.h b/Engine/Source/Runtime/Engine/Classes/Engine/EngineBaseTypes.h
--- a/Engine/Source/Runtime/Engine/Classes/Engine/EngineBaseTypes.h
+++ b/Engine/Source/Runtime/Engine/Classes/Engine/EngineBaseTypes.h
@@ -929,6 +929,9 @@ enum EViewModeIndex : int
 	/** Visualize Strata debug views */
 	VMI_VisualizeStrata = 34 UMETA(DisplayName = "Strata Visualization"),
 
+	/** Visualize Groom debug views */
+	VMI_VisualizeGroom = 35 UMETA(DisplayName = "Groom Visualization"),
+
 	VMI_Max UMETA(Hidden),
 
 	// VMI_Unknown - The value assigned to VMI_Unknown must be the highest possible of any member of EViewModeIndex, or GetViewModeName might seg-fault
diff --git a/Engine/Source/Runtime/Engine/Classes/Engine/GameViewportClient.h b/Engine/Source/Runtime/Engine/Classes/Engine/GameViewportClient.h
--- a/Engine/Source/Runtime/Engine/Classes/Engine/GameViewportClient.h
+++ b/Engine/Source/Runtime/Engine/Classes/Engine/GameViewportClient.h
@@ -661,6 +661,9 @@ protected:
 	void SetCurrentStrataVisualizationMode(FName NewStrataVisualizationMode) { CurrentStrataVisualizationMode = NewStrataVisualizationMode; }
 	FName GetCurrentStrataVisualizationMode() const { return CurrentStrataVisualizationMode; }
 
+	void SetCurrentGroomVisualizationMode(FName NewGroomVisualizationMode) { CurrentGroomVisualizationMode = NewGroomVisualizationMode; }
+	FName GetCurrentGroomVisualizationMode() const { return CurrentGroomVisualizationMode; }
+
 	void SetCurrentVirtualShadowMapVisualizationMode(FName NewVirtualShadowMapVisualizationMode) { CurrentVirtualShadowMapVisualizationMode = NewVirtualShadowMapVisualizationMode; }
 	FName GetCurrentVirtualShadowMapVisualizationMode() const { return CurrentVirtualShadowMapVisualizationMode; }
 
@@ -971,6 +974,9 @@ private:
 	/** Current Strata visualization mode for this game viewport */
 	FName CurrentStrataVisualizationMode;
 
+	/** Current Groom visualization mode for this game viewport */
+	FName CurrentGroomVisualizationMode;
+
 	/** Current virtual shadow map visualization mode for this game viewport */
 	FName CurrentVirtualShadowMapVisualizationMode;
 
diff --git a/Engine/Source/Runtime/Engine/Private/GameViewportClient.cpp b/Engine/Source/Runtime/Engine/Private/GameViewportClient.cpp
--- a/Engine/Source/Runtime/Engine/Private/GameViewportClient.cpp
+++ b/Engine/Source/Runtime/Engine/Private/GameViewportClient.cpp
@@ -241,6 +241,7 @@ UGameViewportClient::UGameViewportClient(const FObjectInitializer& ObjectInitial
 	, CurrentNaniteVisualizationMode(NAME_None)
 	, CurrentLumenVisualizationMode(NAME_None)
 	, CurrentStrataVisualizationMode(NAME_None)
+	, CurrentGroomVisualizationMode(NAME_None)
 	, CurrentVirtualShadowMapVisualizationMode(NAME_None)
 	, HighResScreenshotDialog(nullptr)
 	, bUseSoftwareCursorWidgets(true)
@@ -336,6 +337,7 @@ UGameViewportClient::UGameViewportClient(FVTableHelper& Helper)
 	, CurrentNaniteVisualizationMode(NAME_None)
 	, CurrentLumenVisualizationMode(NAME_None)
 	, CurrentStrataVisualizationMode(NAME_None)
+	, CurrentGroomVisualizationMode(NAME_None)
 	, CurrentVirtualShadowMapVisualizationMode(NAME_None)
 	, HighResScreenshotDialog(nullptr)
 	, bIgnoreInput(false)
@@ -1501,6 +1503,7 @@ void UGameViewportClient::Draw(FViewport* InViewport, FCanvas* SceneCanvas)
 					View->CurrentNaniteVisualizationMode = CurrentNaniteVisualizationMode;
 					View->CurrentLumenVisualizationMode = CurrentLumenVisualizationMode;
 					View->CurrentStrataVisualizationMode = CurrentStrataVisualizationMode;
+					View->CurrentGroomVisualizationMode = CurrentGroomVisualizationMode;
 					View->CurrentVirtualShadowMapVisualizationMode = CurrentVirtualShadowMapVisualizationMode;
 
 					View->CameraConstrainedViewRect = View->UnscaledViewRect;
diff --git a/Engine/Source/Runtime/Engine/Private/GroomVisualizationData.cpp b/Engine/Source/Runtime/Engine/Private/GroomVisualizationData.cpp
new file mode 100644
--- /dev/null
+++ b/Engine/Source/Runtime/Engine/Private/GroomVisualizationData.cpp
@@ -0,0 +1,266 @@
+// Copyright Epic Games, Inc. All Rights Reserved.
+
+#include "GroomVisualizationData.h"
+#include "HAL/IConsoleManager.h"
+#include "Materials/Material.h"
+#include "Misc/ConfigCacheIni.h"
+#include "SceneView.h"
+
+#define LOCTEXT_NAMESPACE "FGroomVisualizationData"
+
+static int32 GHairStrandsPluginEnable = 0;
+
+static TAutoConsoleVariable<int32> CVarHairStrandsGlobalEnable(
+	TEXT("r.HairStrands.Enable"), 1,
+	TEXT("Enable/Disable the entire hair strands system. This affects all geometric representations (i.e., strands, cards, and meshes)."),
+	ECVF_RenderThreadSafe | ECVF_Scalability);
+
+bool IsGroomEnabled()
+{
+	return GHairStrandsPluginEnable > 0 && CVarHairStrandsGlobalEnable.GetValueOnAnyThread() > 0;
+}
+
+void SetGroomEnabled(bool In)
+{
+	GHairStrandsPluginEnable = In ? 1 : 0;
+}
+
+EGroomViewMode GetGroomViewMode(const FSceneView& View)
+{
+	EGroomViewMode Out = EGroomViewMode::None;
+	if (IsGroomEnabled())
+	{
+		static const auto CVarGroomViewMode = IConsoleManager::Get().FindConsoleVariable(FGroomVisualizationData::GetVisualizeConsoleCommandName());
+		const uint32 ViewMode = CVarGroomViewMode && CVarGroomViewMode->AsVariableInt() ? CVarGroomViewMode->AsVariableInt()->GetValueOnRenderThread() : 0;
+		switch (ViewMode)
+		{
+		case 1:  return EGroomViewMode::MacroGroups;
+		case 2:  return EGroomViewMode::LightBounds;
+		case 3:  return EGroomViewMode::MacroGroupScreenRect;
+		case 4:  return EGroomViewMode::DeepOpacityMaps;
+		case 5:  return EGroomViewMode::SamplePerPixel;
+		case 6:  return EGroomViewMode::TAAResolveType;
+		case 7:  return EGroomViewMode::CoverageType;
+		case 8:  return EGroomViewMode::VoxelsDensity;
+		case 9:	 return EGroomViewMode::None;
+		case 10: return EGroomViewMode::None;
+		case 11: return EGroomViewMode::None;
+		case 12: return EGroomViewMode::MeshProjection;
+		case 13: return EGroomViewMode::Coverage;
+		case 14: return EGroomViewMode::MaterialDepth;
+		case 15: return EGroomViewMode::MaterialBaseColor;
+		case 16: return EGroomViewMode::MaterialRoughness;
+		case 17: return EGroomViewMode::MaterialSpecular;
+		case 18: return EGroomViewMode::MaterialTangent;
+		case 19: return EGroomViewMode::Tile;
+		case 20: return EGroomViewMode::None;
+		case 21: return EGroomViewMode::SimHairStrands;
+		case 22: return EGroomViewMode::RenderHairStrands;
+		case 23: return EGroomViewMode::RenderHairRootUV;
+		case 24: return EGroomViewMode::RenderHairRootUDIM;
+		case 25: return EGroomViewMode::RenderHairUV;
+		case 26: return EGroomViewMode::RenderHairSeed;
+		case 27: return EGroomViewMode::RenderHairDimension;
+		case 28: return EGroomViewMode::RenderHairRadiusVariation;
+		case 29: return EGroomViewMode::RenderHairBaseColor;
+		case 30: return EGroomViewMode::RenderHairRoughness;
+		case 31: return EGroomViewMode::RenderVisCluster;
+		case 32: return EGroomViewMode::RenderVisClusterAABB;
+		case 33: return EGroomViewMode::RenderHairTangent;
+		case 34: return EGroomViewMode::RenderHairControlPoints;
+		case 35: return EGroomViewMode::RenderHairGroup;
+		case 36: return EGroomViewMode::RenderLODColoration;
+		case 37: return EGroomViewMode::CardGuides;
+		default: break;
+		}
+
+		const FGroomVisualizationData& VisualizationData = GetGroomVisualizationData();
+		if (View.Family && View.Family->EngineShowFlags.VisualizeGroom)
+		{
+			Out = VisualizationData.GetViewMode(View.CurrentGroomVisualizationMode);
+		}
+		else if (View.Family && View.Family->EngineShowFlags.LODColoration)
+		{
+			Out = EGroomViewMode::RenderLODColoration;
+		}
+	}
+	return Out;
+}
+
+const TCHAR* GetGroomViewModeName(EGroomViewMode In)
+{
+	switch (In)
+	{
+	case EGroomViewMode::None:						return TEXT("NoneDebug");
+	case EGroomViewMode::MacroGroups:				return TEXT("MacroGroups");
+	case EGroomViewMode::LightBounds:				return TEXT("LightBounds");
+	case EGroomViewMode::MacroGroupScreenRect:		return TEXT("MacroGroupScreenRect");
+	case EGroomViewMode::DeepOpacityMaps:			return TEXT("DeepOpacityMaps");
+	case EGroomViewMode::SamplePerPixel:			return TEXT("SamplePerPixel");
+	case EGroomViewMode::TAAResolveType:			return TEXT("TAAResolveType");
+	case EGroomViewMode::CoverageType:				return TEXT("CoverageType");
+	case EGroomViewMode::VoxelsDensity:				return TEXT("VoxelsDensity");
+	case EGroomViewMode::MeshProjection:			return TEXT("MeshProjection");
+	case EGroomViewMode::Coverage:					return TEXT("Coverage");
+	case EGroomViewMode::MaterialDepth:				return TEXT("MaterialDepth");
+	case EGroomViewMode::MaterialBaseColor:			return TEXT("MaterialBaseColor");
+	case EGroomViewMode::MaterialRoughness:			return TEXT("MaterialRoughness");
+	case EGroomViewMode::MaterialSpecular:			return TEXT("MaterialSpecular");
+	case EGroomViewMode::MaterialTangent:			return TEXT("MaterialTangent");
+	case EGroomViewMode::Tile:						return TEXT("Tile");
+	case EGroomViewMode::SimHairStrands:			return TEXT("SimHairStrands");
+	case EGroomViewMode::RenderHairStrands:			return TEXT("RenderHairStrands");
+	case EGroomViewMode::RenderHairRootUV:			return TEXT("RenderHairRootUV");
+	case EGroomViewMode::RenderHairRootUDIM:		return TEXT("RenderHairRootUDIM");
+	case EGroomViewMode::RenderHairUV:				return TEXT("RenderHairUV");
+	case EGroomViewMode::RenderHairSeed:			return TEXT("RenderHairSeed");
+	case EGroomViewMode::RenderHairDimension:		return TEXT("RenderHairDimension");
+	case EGroomViewMode::RenderHairRadiusVariation:	return TEXT("RenderHairRadiusVariation");
+	case EGroomViewMode::RenderHairBaseColor:		return TEXT("RenderHairBaseColor");
+	case EGroomViewMode::RenderHairRoughness:		return TEXT("RenderHairRoughness");
+	case EGroomViewMode::RenderVisCluster:			return TEXT("RenderVisCluster");
+	case EGroomViewMode::RenderVisClusterAABB:		return TEXT("RenderVisClusterAABB");
+	case EGroomViewMode::RenderHairTangent:			return TEXT("RenderHairTangent");
+	case EGroomViewMode::RenderHairControlPoints:	return TEXT("RenderHairControlPoints");
+	case EGroomViewMode::RenderHairGroup:			return TEXT("RenderHairGroup");
+	case EGroomViewMode::RenderLODColoration:		return TEXT("RenderLODColoration");
+	case EGroomViewMode::CardGuides:				return TEXT("CardGuides");
+	}
+	return TEXT("None");
+}
+
+static FGroomVisualizationData GGroomVisualizationData;
+
+static FString ConfigureConsoleCommand(FGroomVisualizationData::TModeMap& ModeMap)
+{
+	FString AvailableVisualizationModes;
+	for (FGroomVisualizationData::TModeMap::TConstIterator It = ModeMap.CreateConstIterator(); It; ++It)
+	{
+		const FGroomVisualizationData::FModeRecord& Record = It.Value();
+		AvailableVisualizationModes += FString(TEXT("\n  "));
+		AvailableVisualizationModes += Record.ModeString;
+	}
+
+	FString Out;
+	Out = TEXT("When the viewport view-mode is set to 'Groom Visualization', this command specifies which of the various channels to display. Values entered other than the allowed values shown below will be ignored.");
+	Out += AvailableVisualizationModes;
+
+	IConsoleManager::Get().RegisterConsoleVariable(
+		FGroomVisualizationData::GetVisualizeConsoleCommandName(),
+		0,
+		*Out,
+		ECVF_Cheat);
+
+	return Out;
+}
+
+static void AddVisualizationMode(
+	FGroomVisualizationData::TModeMap& ModeMap,
+	bool DefaultComposited,
+	const EGroomViewMode Mode,
+	const FText& ModeText,
+	const FText& ModeDesc)
+{
+	const TCHAR* ModeString = GetGroomViewModeName(Mode);
+	const FName ModeName = FName(ModeString);
+
+	FGroomVisualizationData::FModeRecord& Record = ModeMap.Emplace(ModeName);
+	Record.ModeString			= FString(ModeString);
+	Record.ModeName				= ModeName;
+	Record.ModeText				= ModeText;
+	Record.ModeDesc				= ModeDesc;
+	Record.Mode					= Mode;
+	Record.DefaultComposited	= DefaultComposited;
+}
+
+void FGroomVisualizationData::Initialize()
+{
+	if (!bIsInitialized && IsGroomEnabled())
+	{
+		AddVisualizationMode(ModeMap, true, EGroomViewMode::None,						LOCTEXT("NoneDebug", "None"),								LOCTEXT("NoneDebugDesc", "No debug mode"));
+		AddVisualizationMode(ModeMap, true, EGroomViewMode::MacroGroups,				LOCTEXT("MacroGroups", "Instances"),						LOCTEXT("MacroGroupsDesc", "Instances info"));
+		AddVisualizationMode(ModeMap, true, EGroomViewMode::LightBounds,				LOCTEXT("LightBounds", "Light Bound"),						LOCTEXT("LightBoundsDesc", "All DOMs light bounds"));
+		AddVisualizationMode(ModeMap, true, EGroomViewMode::MacroGroupScreenRect,		LOCTEXT("MacroGroupScreenRect", "screen Bounds"),			LOCTEXT("MacroGroupScreenRectDesc", "Screen projected instances"));
+		AddVisualizationMode(ModeMap, true, EGroomViewMode::DeepOpacityMaps,			LOCTEXT("DeepOpacityMaps", "Deep Shadows"),					LOCTEXT("DeepOpacityMapsDesc", "Deep opacity maps"));
+		AddVisualizationMode(ModeMap, true, EGroomViewMode::SamplePerPixel,				LOCTEXT("SamplePerPixel", "Sample Per Pixel"),				LOCTEXT("SamplePerPixelDesc", "Sub-pixel sample count"));
+		AddVisualizationMode(ModeMap, true, EGroomViewMode::TAAResolveType,				LOCTEXT("TAAResolveType", "AA Type"),						LOCTEXT("TAAResolveTypeDesc", "TAA resolve type (regular/responsive)"));
+		AddVisualizationMode(ModeMap, true, EGroomViewMode::CoverageType,				LOCTEXT("CoverageType", "Coverage Type"),					LOCTEXT("CoverageTypeDesc", "Type of hair coverage - Fully covered : Green / Partially covered : Red"));
+		AddVisualizationMode(ModeMap, true, EGroomViewMode::VoxelsDensity,				LOCTEXT("VoxelsDensity", "Voxels"),							LOCTEXT("VoxelsDensityDesc", "Hair density volume"));
+		AddVisualizationMode(ModeMap, true, EGroomViewMode::MeshProjection,				LOCTEXT("MeshProjection", "Root Bindings"),					LOCTEXT("MeshProjectionDesc", "Hair mesh projection"));
+		AddVisualizationMode(ModeMap, true, EGroomViewMode::Coverage,					LOCTEXT("Coverage", "Coverage"),							LOCTEXT("CoverageDesc", "Hair coverage"));
+		AddVisualizationMode(ModeMap, true, EGroomViewMode::MaterialDepth,				LOCTEXT("MaterialDepth", "Depth"),							LOCTEXT("MaterialDepthDesc", "Hair material depth"));
+		AddVisualizationMode(ModeMap, true, EGroomViewMode::MaterialBaseColor,			LOCTEXT("MaterialBaseColor", "BaseColor"),					LOCTEXT("MaterialBaseColorDesc", "Hair material base color"));
+		AddVisualizationMode(ModeMap, true, EGroomViewMode::MaterialRoughness,			LOCTEXT("MaterialRoughness", "Roughness"),					LOCTEXT("MaterialRoughnessDesc", "Hair material roughness"));
+		AddVisualizationMode(ModeMap, true, EGroomViewMode::MaterialSpecular,			LOCTEXT("MaterialSpecular", "Specular"),					LOCTEXT("MaterialSpecularDesc", "Hair material specular"));
+		AddVisualizationMode(ModeMap, true, EGroomViewMode::MaterialTangent,			LOCTEXT("MaterialTangent", "Tangent"),						LOCTEXT("MaterialTangentDesc", "Hair material tangent"));
+		AddVisualizationMode(ModeMap, true, EGroomViewMode::Tile,						LOCTEXT("Tile", "Tile"),									LOCTEXT("TileDesc", "Hair tile cotegorization"));
+		AddVisualizationMode(ModeMap, true, EGroomViewMode::SimHairStrands,				LOCTEXT("SimHairStrands", "Guides"),						LOCTEXT("SimHairStrandsDesc", "Simulation strands"));
+		AddVisualizationMode(ModeMap, true, EGroomViewMode::RenderHairStrands,			LOCTEXT("RenderHairStrands", "Stands Guides Influences"),	LOCTEXT("RenderHairStrandsDesc", "Rendering strands influences"));
+		AddVisualizationMode(ModeMap, true, EGroomViewMode::CardGuides, 				LOCTEXT("CardGuides", "Cards Guides"),						LOCTEXT("CardGuidesDesc", "Cards Guides"));
+		AddVisualizationMode(ModeMap, true, EGroomViewMode::RenderHairRootUV,			LOCTEXT("RenderHairRootUV", "Root UV"),						LOCTEXT("RenderHairRootUVDesc", "Roots UV"));
+		AddVisualizationMode(ModeMap, true, EGroomViewMode::RenderHairRootUDIM,			LOCTEXT("RenderHairRootUDIM", "Root UDIM"),					LOCTEXT("RenderHairRootUDIMDesc", "Roots UV UDIM texture index"));
+		AddVisualizationMode(ModeMap, true, EGroomViewMode::RenderHairUV,				LOCTEXT("RenderHairUV", "UV"),								LOCTEXT("RenderHairUVDesc", "Hair UV"));
+		AddVisualizationMode(ModeMap, true, EGroomViewMode::RenderHairSeed,				LOCTEXT("RenderHairSeed", "Seed"),							LOCTEXT("RenderHairSeedDesc", "Hair seed"));
+		AddVisualizationMode(ModeMap, true, EGroomViewMode::RenderHairDimension,		LOCTEXT("RenderHairDimension", "Dimension"),				LOCTEXT("RenderHairDimensionDesc", "Hair dimensions"));
+		AddVisualizationMode(ModeMap, true, EGroomViewMode::RenderHairRadiusVariation,	LOCTEXT("RenderHairRadiusVariation", "Radius Variation"),	LOCTEXT("RenderHairRadiusVariationDesc", "Hair radius variation"));
+		AddVisualizationMode(ModeMap, true, EGroomViewMode::RenderHairTangent,			LOCTEXT("RenderHairTangent", "Tangent"),					LOCTEXT("RenderHairTangentDesc", "Hair tangent"));
+		AddVisualizationMode(ModeMap, true, EGroomViewMode::RenderHairControlPoints,	LOCTEXT("RenderHairControlPoints", "Control Points"),		LOCTEXT("RenderHairControlPointsDesc", "Hair control points"));
+		AddVisualizationMode(ModeMap, true, EGroomViewMode::RenderHairBaseColor,		LOCTEXT("RenderHairBaseColor", "Per-CV Color"),				LOCTEXT("RenderHairBaseColorDesc", "CV color"));
+		AddVisualizationMode(ModeMap, true, EGroomViewMode::RenderHairRoughness,		LOCTEXT("RenderHairRoughness", "Per-CV Roughness"),			LOCTEXT("RenderHairRoughnessDesc", "CV roughness"));
+		AddVisualizationMode(ModeMap, true, EGroomViewMode::RenderVisCluster,			LOCTEXT("RenderVisCluster", "Clusters"),					LOCTEXT("RenderVisClusterDesc", "Hair visility clusters"));
+		AddVisualizationMode(ModeMap, true, EGroomViewMode::RenderVisClusterAABB,		LOCTEXT("RenderVisClusterAABB", "Clusters Bounds"),			LOCTEXT("RenderVisClusterAABBDesc", "Hair visility clusters AABBs"));
+		AddVisualizationMode(ModeMap, true, EGroomViewMode::RenderHairGroup,			LOCTEXT("RenderHairGroup", "Groups"),						LOCTEXT("RenderHairGroupDesc", "Hair hair groups"));
+		AddVisualizationMode(ModeMap, true, EGroomViewMode::RenderLODColoration,		LOCTEXT("RenderLODColoration", "LOD Color"),				LOCTEXT("RenderLODColorationDesc", "Hair LOD coloring"));
+
+		ConsoleDocumentationVisualizationMode = ConfigureConsoleCommand(ModeMap);
+	}
+	bIsInitialized = true;
+}
+
+FText FGroomVisualizationData::GetModeDisplayName(const FName& InModeName) const
+{
+	if (const FModeRecord* Record = ModeMap.Find(InModeName))
+	{
+		return Record->ModeText;
+	}
+	else
+	{
+		return FText::GetEmpty();
+	}
+}
+
+EGroomViewMode FGroomVisualizationData::GetViewMode(const FName& InModeName) const
+{
+	if (const FModeRecord* Record = ModeMap.Find(InModeName))
+	{
+		return Record->Mode;
+	}
+	else
+	{
+		return EGroomViewMode::None;
+	}
+}
+
+bool FGroomVisualizationData::GetModeDefaultComposited(const FName& InModeName) const
+{
+	if (const FModeRecord* Record = ModeMap.Find(InModeName))
+	{
+		return Record->DefaultComposited;
+	}
+	else
+	{
+		return false;
+	}
+}
+
+FGroomVisualizationData& GetGroomVisualizationData()
+{
+	if (!GGroomVisualizationData.IsInitialized())
+	{
+		GGroomVisualizationData.Initialize();
+	}
+
+	return GGroomVisualizationData;
+}
+
+#undef LOCTEXT_NAMESPACE
diff --git a/Engine/Source/Runtime/Engine/Private/SceneView.cpp b/Engine/Source/Runtime/Engine/Private/SceneView.cpp
--- a/Engine/Source/Runtime/Engine/Private/SceneView.cpp
+++ b/Engine/Source/Runtime/Engine/Private/SceneView.cpp
@@ -20,6 +20,7 @@
 #include "NaniteVisualizationData.h"
 #include "LumenVisualizationData.h"
 #include "StrataVisualizationData.h"
+#include "GroomVisualizationData.h"
 #include "VirtualShadowMapVisualizationData.h"
 #include "Interfaces/Interface_PostProcessVolume.h"
 #include "Engine/TextureCube.h"
diff --git a/Engine/Source/Runtime/Engine/Private/ShowFlags.cpp b/Engine/Source/Runtime/Engine/Private/ShowFlags.cpp
--- a/Engine/Source/Runtime/Engine/Private/ShowFlags.cpp
+++ b/Engine/Source/Runtime/Engine/Private/ShowFlags.cpp
@@ -347,6 +347,9 @@ void ApplyViewMode(EViewModeIndex ViewModeIndex, bool bPerspective, FEngineShowF
 		case VMI_VisualizeStrata:
 			bPostProcessing = true;
 			break;
+		case VMI_VisualizeGroom:
+			bPostProcessing = true;
+			break;
 		case VMI_VisualizeVirtualShadowMap:
 			bPostProcessing = true;
 			break;
@@ -385,6 +388,7 @@ void ApplyViewMode(EViewModeIndex ViewModeIndex, bool bPerspective, FEngineShowF
 	EngineShowFlags.SetVisualizeNanite(ViewModeIndex == VMI_VisualizeNanite);
 	EngineShowFlags.SetVisualizeLumen(ViewModeIndex == VMI_VisualizeLumen);
 	EngineShowFlags.SetVisualizeStrata(ViewModeIndex == VMI_VisualizeStrata);
+	EngineShowFlags.SetVisualizeGroom(ViewModeIndex == VMI_VisualizeGroom);
 	EngineShowFlags.SetVisualizeVirtualShadowMap(ViewModeIndex == VMI_VisualizeVirtualShadowMap);
 	EngineShowFlags.SetVisualizeLightCulling(ViewModeIndex == VMI_LightComplexity);
 	EngineShowFlags.SetShaderComplexity(ViewModeIndex == VMI_ShaderComplexity || ViewModeIndex == VMI_QuadOverdraw || ViewModeIndex == VMI_ShaderComplexityWithQuadOverdraw);
@@ -711,6 +715,10 @@ EViewModeIndex FindViewMode(const FEngineShowFlags& EngineShowFlags)
 	{
 		return VMI_VisualizeStrata;
 	}
+	else if (EngineShowFlags.VisualizeGroom)
+	{
+		return VMI_VisualizeGroom;
+	}
 	else if (EngineShowFlags.VisualizeVirtualShadowMap)
 	{
 		return VMI_VisualizeVirtualShadowMap;
@@ -854,6 +862,7 @@ const TCHAR* GetViewModeName(EViewModeIndex ViewModeIndex)
 		case VMI_VisualizeNanite:			return TEXT("VisualizeNanite");
 		case VMI_VisualizeLumen:			return TEXT("VisualizeLumen");
 		case VMI_VisualizeStrata:			return TEXT("VisualizeStrata");
+		case VMI_VisualizeGroom:			return TEXT("VisualizeGroom");
 		case VMI_VisualizeVirtualShadowMap:	return TEXT("VisualizeVirtualShadowMap");
 		case VMI_RayTracingDebug:			return TEXT("RayTracingDebug");
 		case VMI_PathTracing:				return TEXT("PathTracing");
diff --git a/Engine/Source/Runtime/Engine/Private/ViewModeNames.cpp b/Engine/Source/Runtime/Engine/Private/ViewModeNames.cpp
--- a/Engine/Source/Runtime/Engine/Private/ViewModeNames.cpp
+++ b/Engine/Source/Runtime/Engine/Private/ViewModeNames.cpp
@@ -285,6 +285,10 @@ TArray<const FSlateBrush*> FillViewModeDisplayIcons()
 		{
 			ViewModeDisplayIcons.Emplace(FAppStyle::Get().GetBrush("EditorViewport.VisualizeStrataMode"));
 		}
+		else if (ViewModeIndex == VMI_VisualizeGroom)
+		{
+			ViewModeDisplayIcons.Emplace(FAppStyle::Get().GetBrush("EditorViewport.VisualizeGroomMode"));
+		}
 		else if (ViewModeIndex == VMI_VisualizeVirtualShadowMap)
 		{
 			ViewModeDisplayIcons.Emplace(FAppStyle::Get().GetBrush("EditorViewport.VisualizeVirtualShadowMapMode"));
diff --git a/Engine/Source/Runtime/Engine/Public/GroomVisualizationData.h b/Engine/Source/Runtime/Engine/Public/GroomVisualizationData.h
new file mode 100644
--- /dev/null
+++ b/Engine/Source/Runtime/Engine/Public/GroomVisualizationData.h
@@ -0,0 +1,114 @@
+// Copyright Epic Games, Inc. All Rights Reserved.
+
+#pragma once
+
+#include "CoreMinimal.h"
+
+class FSceneView;
+
+enum class EGroomViewMode : uint8
+{
+	None,
+	SimHairStrands,    // Guide
+	RenderHairStrands, // Guide influence + add cluster
+	RenderHairUV,
+	RenderHairRootUV,
+	RenderHairRootUDIM,
+	RenderHairSeed,
+	RenderHairDimension,
+	RenderHairRadiusVariation,
+	RenderHairTangent,
+	RenderHairBaseColor,
+	RenderHairRoughness,
+	RenderVisCluster,
+	RenderVisClusterAABB,
+	RenderHairGroup,
+	RenderLODColoration,
+	RenderHairControlPoints,	
+	MacroGroups,
+	LightBounds,
+	DeepOpacityMaps,
+	MacroGroupScreenRect,
+	SamplePerPixel,
+	CoverageType,
+	TAAResolveType,
+	VoxelsDensity,
+	MeshProjection, // Rename RootBinding
+	Coverage,
+	MaterialDepth,
+	MaterialBaseColor,
+	MaterialRoughness,
+	MaterialSpecular,
+	MaterialTangent,
+	CardGuides,
+	Tile,
+	Count
+};
+
+class FGroomVisualizationData
+{
+public:
+	
+	/** Describes a single available visualization mode. */
+	struct FModeRecord
+	{
+		FString        ModeString;
+		FName          ModeName;
+		FText          ModeText;
+		FText          ModeDesc;
+		EGroomViewMode Mode;
+
+		// Whether or not this mode (by default) composites with regular scene depth.
+		bool           DefaultComposited;
+	};
+
+	/** Mapping of FName to a visualization mode record. */
+	typedef TMultiMap<FName, FModeRecord> TModeMap;
+
+public:
+	FGroomVisualizationData()
+	: bIsInitialized(false)
+	{
+	}
+
+	/** Initialize the system. */
+	void Initialize();
+
+	/** Check if system was initialized. */
+	inline bool IsInitialized() const { return bIsInitialized; }
+
+	/** Get the display name of a named mode from the available mode map. **/
+	ENGINE_API FText GetModeDisplayName(const FName& InModeName) const;
+
+	ENGINE_API EGroomViewMode GetViewMode(const FName& InModeName) const;
+
+	ENGINE_API bool GetModeDefaultComposited(const FName& InModeName) const;
+
+	inline const TModeMap& GetModeMap() const
+	{
+		return ModeMap;
+	}
+
+	/** Return the console command name for enabling single mode visualization. */
+	static const TCHAR* GetVisualizeConsoleCommandName()
+	{
+		return TEXT("r.HairStrands.ViewMode");
+	}
+
+private:
+
+	/** The name->mode mapping table */
+	TModeMap ModeMap;
+
+	/** Storage for console variable documentation strings. **/
+	FString ConsoleDocumentationVisualizationMode;
+
+	/** Flag indicating if system is initialized. **/
+	bool bIsInitialized;
+};
+
+ENGINE_API FGroomVisualizationData& GetGroomVisualizationData();
+ENGINE_API EGroomViewMode GetGroomViewMode(const FSceneView& View);
+ENGINE_API const TCHAR* GetGroomViewModeName(EGroomViewMode In);
+ENGINE_API bool IsGroomEnabled();
+ENGINE_API void SetGroomEnabled(bool In);
\ No newline at end of file
diff --git a/Engine/Source/Runtime/Engine/Public/SceneView.h b/Engine/Source/Runtime/Engine/Public/SceneView.h
--- a/Engine/Source/Runtime/Engine/Public/SceneView.h
+++ b/Engine/Source/Runtime/Engine/Public/SceneView.h
@@ -1181,6 +1181,9 @@ public:
 	/** Current Strata visualization mode */
 	FName CurrentStrataVisualizationMode;
 
+	/** Current Groom visualization mode */
+	FName CurrentGroomVisualizationMode;
+
 	/** Current Virtual Shadow Map visualization mode */
 	FName CurrentVirtualShadowMapVisualizationMode;
 
diff --git a/Engine/Source/Runtime/Engine/Public/ShowFlags.h b/Engine/Source/Runtime/Engine/Public/ShowFlags.h
--- a/Engine/Source/Runtime/Engine/Public/ShowFlags.h
+++ b/Engine/Source/Runtime/Engine/Public/ShowFlags.h
@@ -377,6 +377,7 @@ private:
 		SetVisualizeNanite(false);
 		SetVisualizeLumen(false);
 		SetVisualizeStrata(false);
+		SetVisualizeGroom(false);
 		SetVisualizeVirtualShadowMap(false);
 		SetVectorFields(false);
 		SetGBufferHints(false);
diff --git a/Engine/Source/Runtime/Engine/Public/ShowFlagsValues.inl b/Engine/Source/Runtime/Engine/Public/ShowFlagsValues.inl
--- a/Engine/Source/Runtime/Engine/Public/ShowFlagsValues.inl
+++ b/Engine/Source/Runtime/Engine/Public/ShowFlagsValues.inl
@@ -68,6 +68,8 @@ SHOWFLAG_ALWAYS_ACCESSIBLE(VisualizeNanite, SFG_Hidden, NSLOCTEXT("UnrealEd", "V
 SHOWFLAG_ALWAYS_ACCESSIBLE(VisualizeLumen, SFG_Hidden, NSLOCTEXT("UnrealEd", "VisualizeLumenSF", "Lumen Visualization"))
 /** Needed for VMI_VisualizeStrata, Whether to enable the Strata visualization mode. */
 SHOWFLAG_ALWAYS_ACCESSIBLE(VisualizeStrata, SFG_Hidden, NSLOCTEXT("UnrealEd", "VisualizeStrataSF", "Strata Visualization"))
+/** Needed for VMI_VisualizeGroom, Whether to enable the Strata visualization mode. */
+SHOWFLAG_ALWAYS_ACCESSIBLE(VisualizeGroom, SFG_Hidden, NSLOCTEXT("UnrealEd", "VisualizeGroomSF", "Groom Visualization"))
 /** Needed for VMI_VisualizeVirtualShadowMap, Whether to enable the virtual shadow map visualization mode. */
 SHOWFLAG_ALWAYS_ACCESSIBLE(VisualizeVirtualShadowMap, SFG_Hidden, NSLOCTEXT("UnrealEd", "VisualizeVirtualShadowMapSF", "Virtual Shadow Map Visualization"))
 /** Allows to disable all direct lighting (does not affect indirect light) */
diff --git a/Engine/Source/Runtime/Renderer/Private/HairStrands/HairStrandsClusters.cpp b/Engine/Source/Runtime/Renderer/Private/HairStrands/HairStrandsClusters.cpp
--- a/Engine/Source/Runtime/Renderer/Private/HairStrands/HairStrandsClusters.cpp
+++ b/Engine/Source/Runtime/Renderer/Private/HairStrands/HairStrandsClusters.cpp
@@ -239,9 +239,6 @@ static FVector CapturedCameraWorldPos;
 static FMatrix CapturedWorldToClipMatrix;
 static FMatrix CapturedProjMatrix;
 
-bool IsHairStrandsClusterDebugEnable();
-bool IsHairStrandsClusterDebugAABBEnable();
-
 static void AddClusterCullingPass(
 	FRDGBuilder& GraphBuilder,
 	FGlobalShaderMap* ShaderMap,
@@ -284,8 +281,8 @@ static void AddClusterCullingPass(
 		};
 
 		static IConsoleVariable* CVarClusterDebug = IConsoleManager::Get().FindConsoleVariable(TEXT("r.HairStrands.Cluster.Debug"));
-		bClusterDebugAABBBuffer = IsHairStrandsClusterDebugAABBEnable();
-		bClusterDebug = IsHairStrandsClusterDebugEnable();
+		bClusterDebugAABBBuffer = GetGroomViewMode(View) == EGroomViewMode::RenderVisClusterAABB;
+		bClusterDebug = GetGroomViewMode(View) == EGroomViewMode::RenderVisCluster;
 		ClusterDebugInfoBuffer = GraphBuilder.CreateBuffer(FRDGBufferDesc::CreateStructuredDesc(sizeof(FHairClusterDebugInfo), ClusterData.ClusterCount), TEXT("Hair.CulledCompactedIndexBuffer"));
 	}
 #endif
diff --git a/Engine/Source/Runtime/Renderer/Private/HairStrands/HairStrandsDebug.cpp b/Engine/Source/Runtime/Renderer/Private/HairStrands/HairStrandsDebug.cpp
--- a/Engine/Source/Runtime/Renderer/Private/HairStrands/HairStrandsDebug.cpp
+++ b/Engine/Source/Runtime/Renderer/Private/HairStrands/HairStrandsDebug.cpp
@@ -28,18 +28,14 @@
 #include "ShaderPrint.h"
 #include "ScreenPass.h"
 
+#include "GroomVisualizationData.h"
+
 static int32 GDeepShadowDebugIndex = 0;
 static float GDeepShadowDebugScale = 20;
 
 static FAutoConsoleVariableRef CVarDeepShadowDebugDomIndex(TEXT("r.HairStrands.DeepShadow.DebugDOMIndex"), GDeepShadowDebugIndex, TEXT("Index of the DOM texture to draw"));
 static FAutoConsoleVariableRef CVarDeepShadowDebugDomScale(TEXT("r.HairStrands.DeepShadow.DebugDOMScale"), GDeepShadowDebugScale, TEXT("Scaling value for the DeepOpacityMap when drawing the deep shadow stats"));
 
-static int32 GHairStrandsDebugMode = 0;
-static FAutoConsoleVariableRef CVarDeepShadowStats(TEXT("r.HairStrands.DebugMode"), GHairStrandsDebugMode, TEXT("Draw various stats/debug mode about hair rendering"));
-
-static int32 GHairStrandsDebugStrandsMode = 0;
-static FAutoConsoleVariableRef CVarDebugPhysicsStrand(TEXT("r.HairStrands.StrandsMode"), GHairStrandsDebugStrandsMode, TEXT("Render debug mode for hair strands. 0:off, 1:simulation strands, 2:render strands with colored simulation strands influence, 3:hair UV, 4:hair root UV, 5: hair seed, 6: dimensions"));
-
 static int32 GHairStrandsDebugPlotBsdf = 0;
 static FAutoConsoleVariableRef CVarHairStrandsDebugBSDF(TEXT("r.HairStrands.PlotBsdf"), GHairStrandsDebugPlotBsdf, TEXT("Debug view for visualizing hair BSDF."));
 
@@ -52,17 +48,6 @@ static FAutoConsoleVariableRef CVarHairStrandsDebugBSDFAbsorption(TEXT("r.HairSt
 static float GHairStrandsDebugPlotBsdfExposure = 1.1f;
 static FAutoConsoleVariableRef CVarHairStrandsDebugBSDFExposure(TEXT("r.HairStrands.PlotBsdf.Exposure"), GHairStrandsDebugPlotBsdfExposure, TEXT("Change the exposure of the plot."));
 
-static int GHairStrandsDebugSampleIndex = -1;
-static FAutoConsoleVariableRef CVarHairStrandsDebugMaterialSampleIndex(TEXT("r.HairStrands.DebugMode.SampleIndex"), GHairStrandsDebugSampleIndex, TEXT("Debug value for a given sample index (default:-1, i.e., average sample information)."));
-
-static int32 GHairStrandsClusterDebug = 0;
-static FAutoConsoleVariableRef CVarHairStrandsDebugClusterAABB(TEXT("r.HairStrands.Cluster.Debug"), GHairStrandsClusterDebug, TEXT("Draw debug the world bounding box of hair clusters used for culling optimisation (0:off, 1:visible cluster, 2:culled cluster, 3:colored LOD, 4:LOD info)."));
-
-static int32 GHairTangentDebug = 0;
-static int32 GHairTangentDebug_TileSize = 8;
-static FAutoConsoleVariableRef CVarHairTangentDebug(TEXT("r.HairStrands.DebugMode.Tangent"), GHairTangentDebug, TEXT("Draw debug tangent for hair strands and hair cards."));
-static FAutoConsoleVariableRef CVarHairTangentDebug_TileSize(TEXT("r.HairStrands.DebugMode.Tangent.TileSize"), GHairTangentDebug_TileSize, TEXT("Draw debug tangent - Grid size for drawing debug tangent"));
-
 static int32 GHairVirtualVoxel_DrawDebugPage = 0;
 static FAutoConsoleVariableRef CVarHairVirtualVoxel_DrawDebugPage(TEXT("r.HairStrands.Voxelization.Virtual.DrawDebugPage"), GHairVirtualVoxel_DrawDebugPage, TEXT("When voxel debug rendering is enable 1: render the page bounds, instead of the voxel 2: the occupancy within the page (i.e., 8x8x8 brick)"));
 static int32 GHairVirtualVoxel_ForceMipLevel = -1;
@@ -81,16 +66,6 @@ static bool TryEnableShaderDrawAndShaderPrint(const FViewInfo& View, uint32 Resq
 	return ShaderPrint::IsEnabled(View.ShaderPrintData);
 }
 
-bool IsHairStrandsClusterDebugEnable()
-{
-	return GHairStrandsClusterDebug > 0;
-}
-
-bool IsHairStrandsClusterDebugAABBEnable()
-{
-	return GHairStrandsClusterDebug > 1;
-}
-
 FHairStrandsDebugData::FPlotData FHairStrandsDebugData::CreatePlotData(FRDGBuilder& GraphBuilder)
 {
 	FHairStrandsDebugData::FPlotData Out;
@@ -123,108 +98,6 @@ void FHairStrandsDebugData::SetParameters(FRDGBuilder& GraphBuilder, const FHair
 	Out.Debug_SampleCounter = GraphBuilder.CreateSRV(In.SampleCounter, PF_R32_UINT);
 }
 
-EHairDebugMode GetHairStrandsDebugMode()
-{
-	switch (GHairStrandsDebugMode)
-	{
-	case 0:  return EHairDebugMode::None;
-	case 1:  return EHairDebugMode::MacroGroups;
-	case 2:  return EHairDebugMode::LightBounds;
-	case 3:  return EHairDebugMode::MacroGroupScreenRect;
-	case 4:  return EHairDebugMode::DeepOpacityMaps;
-	case 5:  return EHairDebugMode::SamplePerPixel;
-	case 6:  return EHairDebugMode::TAAResolveType;
-	case 7:  return EHairDebugMode::CoverageType;
-	case 8:  return EHairDebugMode::VoxelsDensity;
-	case 9:  return EHairDebugMode::VoxelsTangent;
-	case 10: return EHairDebugMode::VoxelsBaseColor;
-	case 11: return EHairDebugMode::VoxelsRoughness;
-	case 12: return EHairDebugMode::MeshProjection;
-	case 13: return EHairDebugMode::Coverage;
-	case 14: return EHairDebugMode::MaterialDepth;
-	case 15: return EHairDebugMode::MaterialBaseColor;
-	case 16: return EHairDebugMode::MaterialRoughness;
-	case 17: return EHairDebugMode::MaterialSpecular;
-	case 18: return EHairDebugMode::MaterialTangent;
-	case 19: return EHairDebugMode::Tile;
-	default: return EHairDebugMode::None;
-	};
-}
-
-static const TCHAR* ToString(EHairDebugMode DebugMode)
-{
-	switch (DebugMode)
-	{
-	case EHairDebugMode::None: return TEXT("None");
-	case EHairDebugMode::MacroGroups: return TEXT("Macro groups info");
-	case EHairDebugMode::LightBounds: return TEXT("All DOMs light bounds");
-	case EHairDebugMode::MacroGroupScreenRect: return TEXT("Screen projected macro groups");
-	case EHairDebugMode::DeepOpacityMaps: return TEXT("Deep opacity maps");
-	case EHairDebugMode::SamplePerPixel: return TEXT("Sub-pixel sample count");
-	case EHairDebugMode::TAAResolveType: return TEXT("TAA resolve type (regular/responsive)");
-	case EHairDebugMode::CoverageType: return TEXT("Type of hair coverage - Fully covered : Green / Partially covered : Red");
-	case EHairDebugMode::VoxelsDensity: return TEXT("Hair density volume");
-	case EHairDebugMode::VoxelsTangent: return TEXT("Hair tangent volume");
-	case EHairDebugMode::VoxelsBaseColor: return TEXT("Hair base color volume");
-	case EHairDebugMode::VoxelsRoughness: return TEXT("Hair roughness volume");
-	case EHairDebugMode::MeshProjection: return TEXT("Hair mesh projection");
-	case EHairDebugMode::Coverage: return TEXT("Hair coverage");
-	case EHairDebugMode::MaterialDepth: return TEXT("Hair material depth");
-	case EHairDebugMode::MaterialBaseColor: return TEXT("Hair material base color");
-	case EHairDebugMode::MaterialRoughness: return TEXT("Hair material roughness");
-	case EHairDebugMode::MaterialSpecular: return TEXT("Hair material specular");
-	case EHairDebugMode::MaterialTangent: return TEXT("Hair material tangent");
-	case EHairDebugMode::Tile: return TEXT("Hair tile cotegorization");
-	default: return TEXT("None");
-	};
-}
-
-
-EHairStrandsDebugMode GetHairStrandsDebugStrandsMode()
-{
-	switch (GHairStrandsDebugStrandsMode)
-	{
-	case  0:  return EHairStrandsDebugMode::NoneDebug;
-	case  1:  return EHairStrandsDebugMode::SimHairStrands;
-	case  2:  return EHairStrandsDebugMode::RenderHairStrands;
-	case  3:  return EHairStrandsDebugMode::RenderHairRootUV;
-	case  4:  return EHairStrandsDebugMode::RenderHairRootUDIM;
-	case  5:  return EHairStrandsDebugMode::RenderHairUV;
-	case  6:  return EHairStrandsDebugMode::RenderHairSeed;
-	case  7:  return EHairStrandsDebugMode::RenderHairDimension;
-	case  8:  return EHairStrandsDebugMode::RenderHairRadiusVariation;
-	case  9:  return EHairStrandsDebugMode::RenderHairBaseColor;
-	case 10:  return EHairStrandsDebugMode::RenderHairRoughness;
-	case 11:  return EHairStrandsDebugMode::RenderVisCluster;
-	case 12:  return EHairStrandsDebugMode::RenderHairTangent;
-	case 13:  return EHairStrandsDebugMode::RenderHairControlPoints;
-	case 14:  return EHairStrandsDebugMode::RenderHairGroup;
-	default:  return EHairStrandsDebugMode::NoneDebug;
-	};
-}
-
-static const TCHAR* ToString(EHairStrandsDebugMode DebugMode)
-{
-	switch (DebugMode)
-	{
-	case EHairStrandsDebugMode::NoneDebug						: return TEXT("None");
-	case EHairStrandsDebugMode::SimHairStrands				: return TEXT("Simulation strands");
-	case EHairStrandsDebugMode::RenderHairStrands			: return TEXT("Rendering strands influences");
-	case EHairStrandsDebugMode::RenderHairRootUV			: return TEXT("Roots UV");
-	case EHairStrandsDebugMode::RenderHairRootUDIM			: return TEXT("Roots UV UDIM texture index");
-	case EHairStrandsDebugMode::RenderHairUV				: return TEXT("Hair UV");
-	case EHairStrandsDebugMode::RenderHairSeed				: return TEXT("Hair seed");
-	case EHairStrandsDebugMode::RenderHairDimension			: return TEXT("Hair dimensions");
-	case EHairStrandsDebugMode::RenderHairRadiusVariation	: return TEXT("Hair radius variation");
-	case EHairStrandsDebugMode::RenderHairTangent			: return TEXT("Hair tangent");
-	case EHairStrandsDebugMode::RenderHairControlPoints		: return TEXT("Hair control points");
-	case EHairStrandsDebugMode::RenderHairBaseColor			: return TEXT("Hair vertices color");
-	case EHairStrandsDebugMode::RenderHairRoughness			: return TEXT("Hair vertices roughness");
-	case EHairStrandsDebugMode::RenderVisCluster			: return TEXT("Hair visility clusters");
-	default													: return TEXT("None");
-	};
-}
-
 ///////////////////////////////////////////////////////////////////////////////////////////////////
 class FHairPrintLODInfoCS : public FGlobalShader
 {
@@ -348,7 +221,7 @@ static void AddDebugHairPrintPass(
 	FRDGBuilder& GraphBuilder,
 	const FScene* Scene,
 	const FViewInfo* View,
-	const EHairDebugMode InDebugMode,
+	const EGroomViewMode ViewMode,
 	const FHairStrandsVisibilityData& VisibilityData,
 	const FHairStrandsMacroGroupDatas& MacroGroupDatas,
 	const FHairStrandsMacroGroupResources& MacroGroupResources,
@@ -568,24 +441,24 @@ IMPLEMENT_GLOBAL_SHADER(FHairDebugPS, "/Engine/Private/HairStrands/HairStrandsDe
 static void AddDebugHairPass(
 	FRDGBuilder& GraphBuilder,
 	const FViewInfo* View,
-	const EHairDebugMode InDebugMode,
+	const EGroomViewMode InDebugMode,
 	const FHairStrandsVisibilityData& VisibilityData,
 	FRDGTextureSRVRef InDepthStencilTexture,
 	FRDGTextureRef& OutTarget)
 {
 	check(OutTarget);
-	check(InDebugMode == EHairDebugMode::TAAResolveType || 
-		InDebugMode == EHairDebugMode::SamplePerPixel || 
-		InDebugMode == EHairDebugMode::CoverageType || 
-		InDebugMode == EHairDebugMode::Coverage ||
-		InDebugMode == EHairDebugMode::MaterialDepth ||
-		InDebugMode == EHairDebugMode::MaterialBaseColor ||
-		InDebugMode == EHairDebugMode::MaterialRoughness ||
-		InDebugMode == EHairDebugMode::MaterialSpecular ||
-		InDebugMode == EHairDebugMode::MaterialTangent);
+	check(InDebugMode == EGroomViewMode::TAAResolveType || 
+		InDebugMode == EGroomViewMode::SamplePerPixel || 
+		InDebugMode == EGroomViewMode::CoverageType || 
+		InDebugMode == EGroomViewMode::Coverage ||
+		InDebugMode == EGroomViewMode::MaterialDepth ||
+		InDebugMode == EGroomViewMode::MaterialBaseColor ||
+		InDebugMode == EGroomViewMode::MaterialRoughness ||
+		InDebugMode == EGroomViewMode::MaterialSpecular ||
+		InDebugMode == EGroomViewMode::MaterialTangent);
 
 	if (!VisibilityData.CoverageTexture || !VisibilityData.NodeIndex || !VisibilityData.NodeData) return;
-	if (InDebugMode == EHairDebugMode::TAAResolveType && !InDepthStencilTexture) return;
+	if (InDebugMode == EGroomViewMode::TAAResolveType && !InDepthStencilTexture) return;
 
 	const FIntRect Viewport = View->ViewRect;
 	const FIntPoint Resolution(Viewport.Width(), Viewport.Height());
@@ -593,15 +466,15 @@ static void AddDebugHairPass(
 	uint32 InternalDebugMode = 0;
 	switch (InDebugMode)
 	{
-		case EHairDebugMode::SamplePerPixel:	InternalDebugMode = 0; break;
-		case EHairDebugMode::CoverageType:		InternalDebugMode = 1; break;
-		case EHairDebugMode::TAAResolveType:	InternalDebugMode = 2; break;
-		case EHairDebugMode::Coverage:			InternalDebugMode = 3; break;
-		case EHairDebugMode::MaterialDepth:		InternalDebugMode = 4; break;
-		case EHairDebugMode::MaterialBaseColor:	InternalDebugMode = 5; break;
-		case EHairDebugMode::MaterialRoughness:	InternalDebugMode = 6; break;
-		case EHairDebugMode::MaterialSpecular:	InternalDebugMode = 7; break;
-		case EHairDebugMode::MaterialTangent:	InternalDebugMode = 8; break;
+		case EGroomViewMode::SamplePerPixel:	InternalDebugMode = 0; break;
+		case EGroomViewMode::CoverageType:		InternalDebugMode = 1; break;
+		case EGroomViewMode::TAAResolveType:	InternalDebugMode = 2; break;
+		case EGroomViewMode::Coverage:			InternalDebugMode = 3; break;
+		case EGroomViewMode::MaterialDepth:		InternalDebugMode = 4; break;
+		case EGroomViewMode::MaterialBaseColor:	InternalDebugMode = 5; break;
+		case EGroomViewMode::MaterialRoughness:	InternalDebugMode = 6; break;
+		case EGroomViewMode::MaterialSpecular:	InternalDebugMode = 7; break;
+		case EGroomViewMode::MaterialTangent:	InternalDebugMode = 8; break;
 	};
 
 	FHairDebugPS::FParameters* Parameters = GraphBuilder.AllocParameters<FHairDebugPS::FParameters>();
@@ -612,7 +485,7 @@ static void AddDebugHairPass(
 	Parameters->DepthStencilTexture = InDepthStencilTexture;
 	Parameters->LinearSampler = TStaticSamplerState<SF_Bilinear, AM_Clamp, AM_Clamp, AM_Clamp>::GetRHI();
 	Parameters->DebugMode = InternalDebugMode;
-	Parameters->SampleIndex = GHairStrandsDebugSampleIndex;
+	Parameters->SampleIndex = 0;
 	Parameters->RenderTargets[0] = FRenderTargetBinding(OutTarget, ERenderTargetLoadAction::ELoad, 0);
 	TShaderMapRef<FPostProcessVS> VertexShader(View->ShaderMap);
 
@@ -621,7 +494,7 @@ static void AddDebugHairPass(
 	ClearUnusedGraphResources(PixelShader, Parameters);
 
 	GraphBuilder.AddPass(
-		RDG_EVENT_NAME("HairStrands::DebugMode(%s)", ToString(InDebugMode)),
+		RDG_EVENT_NAME("HairStrands::DebugMode(%s)", GetGroomViewModeName(InDebugMode)),
 		Parameters,
 		ERDGPassFlags::Raster,
 		[Parameters, VertexShader, PixelShader, Viewport, Resolution](FRHICommandList& RHICmdList)
@@ -946,7 +819,7 @@ static void AddDebugHairTangentPass(
 	const FSceneTextures& SceneTextures,
 	FRDGTextureRef& OutputTexture)
 {
-	if (!TryEnableShaderDrawAndShaderPrint(View, 8192u, 1000u))
+	if (!TryEnableShaderDrawAndShaderPrint(View, 64000u, 1000u))
 	{
 		return;
 	}
@@ -955,7 +828,7 @@ static void AddDebugHairTangentPass(
 	Parameters->ViewUniformBuffer		= View.ViewUniformBuffer;
 	Parameters->HairStrands				= View.HairStrandsViewData.UniformBuffer;
 	Parameters->OutputResolution		= OutputTexture->Desc.Extent;
-	Parameters->TileSize				= FMath::Clamp(GHairTangentDebug_TileSize, 4, 32);
+	Parameters->TileSize				= 8;
 	Parameters->TileCount				= FIntPoint(FMath::FloorToInt(Parameters->OutputResolution.X / Parameters->TileSize), FMath::FloorToInt(Parameters->OutputResolution.X / Parameters->TileSize));
 	Parameters->SceneTextures			= SceneTextures.UniformBuffer;
 	Parameters->BilinearTextureSampler	= TStaticSamplerState<SF_Bilinear>::GetRHI();
@@ -1234,26 +1107,29 @@ static void AddDrawDebugClusterPass(
 	const FHairStrandClusterData& HairClusterData,
 	const FViewInfo& View)
 {
-	const bool bDebugEnable = IsHairStrandsClusterDebugAABBEnable();
+	const EGroomViewMode ViewMode = GetGroomViewMode(View);
+	const bool bDebugEnable = ViewMode == EGroomViewMode::RenderVisClusterAABB || ViewMode == EGroomViewMode::RenderVisCluster;
 	const bool bCullingEnable = IsHairStrandsClusterCullingEnable();
 	if (!bDebugEnable || !bCullingEnable)
 	{
 		return;
 	}
 	
-	if (!TryEnableShaderDrawAndShaderPrint(View, 5000, 2000))
+	if (!TryEnableShaderDrawAndShaderPrint(View, 64000u, 2000))
 	{
 		return;
 	}
 
-
 	for (const FHairStrandsMacroGroupData& MacroGroupData : View.HairStrandsViewData.MacroGroupDatas)
 	{
-		const bool bDebugAABB = IsHairStrandsClusterDebugAABBEnable();
+		const bool bDebugAABB = GetGroomViewMode(View) == EGroomViewMode::RenderVisClusterAABB;
 
 		for (const FHairStrandsMacroGroupData::PrimitiveInfo& PrimitiveInfo : MacroGroupData.PrimitivesInfos)
 		{
-			check(PrimitiveInfo.Mesh && PrimitiveInfo.Mesh->Elements.Num() > 0);
+			if (!PrimitiveInfo.Mesh || PrimitiveInfo.Mesh->Elements.Num() == 0)
+			{
+				continue;
+			}
 
 			for (int DataIndex = 0; DataIndex < HairClusterData.HairGroups.Num(); ++DataIndex)
 			{
@@ -1261,7 +1137,9 @@ static void AddDrawDebugClusterPass(
 
 				// Find a better/less hacky way
 				if (PrimitiveInfo.PublicDataPtr != HairGroupClusters.HairGroupPublicPtr)
+				{
 					continue;
+				}
 
 				if (ShaderPrint::IsEnabled(View.ShaderPrintData) && HairGroupClusters.CulledClusterCountBuffer)
 				{
@@ -1276,7 +1154,7 @@ static void AddDrawDebugClusterPass(
 					Parameters->ClusterCount = HairGroupClusters.ClusterCount;
 					Parameters->TriangleCount = HairGroupClusters.VertexCount * 2; // VertexCount is actually the number of control points
 					Parameters->HairGroupId = DataIndex;
-					Parameters->ClusterDebugMode = GHairStrandsClusterDebug;
+					Parameters->ClusterDebugMode = ViewMode == EGroomViewMode::RenderVisCluster ? 1 : (ViewMode == EGroomViewMode::RenderVisClusterAABB ? 2 : 0);
 					Parameters->ClusterAABBBuffer = HairGroupClusters.ClusterAABBBuffer->SRV;
 					Parameters->CulledDispatchIndirectParametersClusterCountBuffer = GraphBuilder.CreateSRV(HairGroupClusters.CulledClusterCountBuffer, EPixelFormat::PF_R32_UINT);
 					Parameters->CulledDrawIndirectParameters = DrawIndirectBuffer.SRV;
@@ -1331,38 +1209,12 @@ static void InternalRenderHairStrandsDebugInfo(
 		RunHairStrandsBookmark(GraphBuilder, EHairStrandsBookmark::ProcessDebug, Params);
 	}
 
+	const EGroomViewMode ViewMode = GetGroomViewMode(View);
+
 	// Display tangent vector for strands/cards/meshes
+	if (ViewMode == EGroomViewMode::RenderHairTangent)
 	{
-		// Check among the hair instances, if hair tangent debug mode is requested
-		bool bTangentEnabled = GHairTangentDebug > 0;
-		if (!bTangentEnabled)
-		{
-			for (const FMeshBatchAndRelevance& Mesh : View.HairStrandsMeshElements)
-			{
-				const FHairGroupPublicData* GroupData = HairStrands::GetHairData(Mesh.Mesh);
-				if (GroupData->DebugMode == EHairStrandsDebugMode::RenderHairTangent)
-				{
-					bTangentEnabled = true;
-					break;
-				}
-			}
-		}
-		if (!bTangentEnabled)
-		{
-			for (const FMeshBatchAndRelevance& Mesh : View.HairCardsMeshElements)
-			{
-				const FHairGroupPublicData* GroupData = HairStrands::GetHairData(Mesh.Mesh);
-				if (GroupData->DebugMode == EHairStrandsDebugMode::RenderHairTangent)
-				{
-					bTangentEnabled = true;
-					break;
-				}
-			}
-		}
-		if (bTangentEnabled)
-		{
-			AddDebugHairTangentPass(GraphBuilder, View, SceneTextures, SceneColorTexture);
-		}
+		AddDebugHairTangentPass(GraphBuilder, View, SceneTextures, SceneColorTexture);
 	}
 
 	// Draw LOD info 
@@ -1391,11 +1243,6 @@ static void InternalRenderHairStrandsDebugInfo(
 	}
 
 	const FScreenPassRenderTarget SceneColor(SceneColorTexture, View.ViewRect, ERenderTargetLoadAction::ELoad);
-
-	// Debug mode name only
-	const EHairStrandsDebugMode StrandsDebugMode = GetHairStrandsDebugStrandsMode();
-	const EHairDebugMode HairDebugMode = GetHairStrandsDebugMode();
-
 	
 	const FHairStrandsViewData& HairData = View.HairStrandsViewData;
 
@@ -1418,12 +1265,12 @@ static void InternalRenderHairStrandsDebugInfo(
 		AddDebugHairShadowCullingPass(GraphBuilder, Scene, &View);
 	}
 
-	if (HairDebugMode == EHairDebugMode::MacroGroups)
+	if (ViewMode == EGroomViewMode::MacroGroups)
 	{
-		AddDebugHairPrintPass(GraphBuilder, Scene, &View, HairDebugMode, HairData.VisibilityData, HairData.MacroGroupDatas, HairData.MacroGroupResources, SceneTextures.Stencil);
+		AddDebugHairPrintPass(GraphBuilder, Scene, &View, ViewMode, HairData.VisibilityData, HairData.MacroGroupDatas, HairData.MacroGroupResources, SceneTextures.Stencil);
 	}
 
-	if (HairDebugMode == EHairDebugMode::DeepOpacityMaps)
+	if (ViewMode == EGroomViewMode::DeepOpacityMaps)
 	{
 		for (const FHairStrandsMacroGroupData& MacroGroup : HairData.MacroGroupDatas)
 		{
@@ -1444,7 +1291,7 @@ static void InternalRenderHairStrandsDebugInfo(
 	}
 
 	// View Rect
-	if (IsHairStrandsViewRectOptimEnable() && HairDebugMode == EHairDebugMode::MacroGroupScreenRect)
+	if (IsHairStrandsViewRectOptimEnable() && ViewMode == EGroomViewMode::MacroGroupScreenRect)
 	{
 		for (const FHairStrandsMacroGroupData& MacroGroupData : HairData.MacroGroupDatas)
 		{
@@ -1455,11 +1302,10 @@ static void InternalRenderHairStrandsDebugInfo(
 		AddDebugDeepShadowTexturePass(GraphBuilder, &View, TotalRect, nullptr, nullptr, SceneColorTexture);
 	}
 	
-	const bool bIsVoxelMode = HairDebugMode == EHairDebugMode::VoxelsDensity || HairDebugMode == EHairDebugMode::VoxelsTangent || HairDebugMode == EHairDebugMode::VoxelsBaseColor || HairDebugMode == EHairDebugMode::VoxelsRoughness;
 
 	// Render Frustum for all lights & macro groups 
 	{
-		if ((HairDebugMode == EHairDebugMode::LightBounds || HairDebugMode == EHairDebugMode::DeepOpacityMaps))
+		if ((ViewMode == EGroomViewMode::LightBounds || ViewMode == EGroomViewMode::DeepOpacityMaps))
 		{
 			if (HairData.DeepShadowResources.bIsGPUDriven)
 			{
@@ -1469,26 +1315,27 @@ static void InternalRenderHairStrandsDebugInfo(
 	}
 	
 	const bool bRunDebugPass =
-		HairDebugMode == EHairDebugMode::TAAResolveType ||
-		HairDebugMode == EHairDebugMode::SamplePerPixel ||
-		HairDebugMode == EHairDebugMode::CoverageType ||
-		HairDebugMode == EHairDebugMode::Coverage ||
-		HairDebugMode == EHairDebugMode::MaterialDepth ||
-		HairDebugMode == EHairDebugMode::MaterialBaseColor ||
-		HairDebugMode == EHairDebugMode::MaterialRoughness ||
-		HairDebugMode == EHairDebugMode::MaterialSpecular ||
-		HairDebugMode == EHairDebugMode::MaterialTangent;
+		ViewMode == EGroomViewMode::TAAResolveType ||
+		ViewMode == EGroomViewMode::SamplePerPixel ||
+		ViewMode == EGroomViewMode::CoverageType ||
+		ViewMode == EGroomViewMode::Coverage ||
+		ViewMode == EGroomViewMode::MaterialDepth ||
+		ViewMode == EGroomViewMode::MaterialBaseColor ||
+		ViewMode == EGroomViewMode::MaterialRoughness ||
+		ViewMode == EGroomViewMode::MaterialSpecular ||
+		ViewMode == EGroomViewMode::MaterialTangent;
 	if (bRunDebugPass)
 	{
-		AddDebugHairPass(GraphBuilder, &View, HairDebugMode, HairData.VisibilityData, SceneTextures.Stencil, SceneColorTexture);
-		AddDebugHairPrintPass(GraphBuilder, Scene, &View, HairDebugMode, HairData.VisibilityData, HairData.MacroGroupDatas, HairData.MacroGroupResources, SceneTextures.Stencil);
+		AddDebugHairPass(GraphBuilder, &View, ViewMode, HairData.VisibilityData, SceneTextures.Stencil, SceneColorTexture);
+		AddDebugHairPrintPass(GraphBuilder, Scene, &View, ViewMode, HairData.VisibilityData, HairData.MacroGroupDatas, HairData.MacroGroupResources, SceneTextures.Stencil);
 	}
-	else if (HairDebugMode == EHairDebugMode::Tile)
+	else if (ViewMode == EGroomViewMode::Tile)
 	{
 		check(HairData.VisibilityData.TileData.IsValid());
 		AddHairStrandsDebugTilePass(GraphBuilder, View, SceneColorTexture, HairData.VisibilityData.TileData);
 	}
 
+	const bool bIsVoxelMode = ViewMode == EGroomViewMode::VoxelsDensity;
 	if (bIsVoxelMode)
 	{
 		if (HairData.VirtualVoxelResources.IsValid())
@@ -1520,22 +1367,22 @@ static void InternalRenderHairStrandsDebugInfo(
 		FComputeShaderUtils::AddPass(GraphBuilder, RDG_EVENT_NAME("HairStrands::PPLLDebug"), ComputeShader, PassParameters, FIntVector::DivideAndRoundUp(TextureSize, FIntVector(8, 8, 1)));
 	}
 
-	if (GHairStrandsClusterDebug > 0)
+	if (ViewMode == EGroomViewMode::RenderVisCluster || ViewMode == EGroomViewMode::RenderVisClusterAABB)
 	{
 		AddDrawDebugClusterPass(GraphBuilder, HairClusterData, View);
 	}
 
-	if (StrandsDebugMode != EHairStrandsDebugMode::NoneDebug || HairDebugMode != EHairDebugMode::None)
+	if (ViewMode != EGroomViewMode::None)
 	{
-		AddDrawCanvasPass(GraphBuilder, {}, View, SceneColor, [&View, &StrandsDebugMode, YStep, HairDebugMode](FCanvas& Canvas)
+		AddDrawCanvasPass(GraphBuilder, {}, View, SceneColor, [&View, YStep, ViewMode](FCanvas& Canvas)
 		{
 			float X = 40;
 			float Y = View.ViewRect.Height() - YStep * 3.f;
 			FString Line;
-			if (StrandsDebugMode != EHairStrandsDebugMode::NoneDebug)
-				Line = FString::Printf(TEXT("Hair Debug mode - %s"), ToString(StrandsDebugMode));
-			else if (HairDebugMode != EHairDebugMode::None)
-				Line = FString::Printf(TEXT("Hair Debug mode - %s"), ToString(HairDebugMode));
+			if (ViewMode != EGroomViewMode::None)
+			{
+				Line = FString::Printf(TEXT("Hair Debug mode - %s"), GetGroomViewModeName(ViewMode));
+			}
 
 			Canvas.DrawShadowedString(X, Y += YStep, *Line, GetStatsFont(), FLinearColor(1, 1, 0));
 		});
diff --git a/Engine/Source/Runtime/Renderer/Private/HairStrands/HairStrandsInterface.cpp b/Engine/Source/Runtime/Renderer/Private/HairStrands/HairStrandsInterface.cpp
--- a/Engine/Source/Runtime/Renderer/Private/HairStrands/HairStrandsInterface.cpp
+++ b/Engine/Source/Runtime/Renderer/Private/HairStrands/HairStrandsInterface.cpp
@@ -26,13 +26,6 @@ static TAutoConsoleVariable<int32> CVarHairStrandsRaytracingEnable(
 	TEXT("Enable/Disable hair strands raytracing geometry. This is anopt-in option per groom asset/groom instance."),
 	ECVF_RenderThreadSafe | ECVF_Scalability);
 
-static int32 GHairStrandsPluginEnable = 0;
-
-static TAutoConsoleVariable<int32> CVarHairStrandsGlobalEnable(
-	TEXT("r.HairStrands.Enable"), 1,
-	TEXT("Enable/Disable the entire hair strands system. This affects all geometric representations (i.e., strands, cards, and meshes)."),
-	ECVF_RenderThreadSafe | ECVF_Scalability);
-
 static TAutoConsoleVariable<int32> CVarHairStrandsEnable(
 	TEXT("r.HairStrands.Strands"), 1,
 	TEXT("Enable/Disable hair strands rendering"),
@@ -207,7 +200,7 @@ bool IsHairRayTracingEnabled()
 
 bool IsHairStrandsSupported(EHairStrandsShaderType Type, EShaderPlatform Platform)
 {
-	if (GHairStrandsPluginEnable <= 0 || CVarHairStrandsGlobalEnable.GetValueOnAnyThread() <= 0) return false;
+	if (!IsGroomEnabled()) return false;
 
 	// Important:
 	// EHairStrandsShaderType::All: Mobile is excluded as we don't need any interpolation/simulation code for this. It only do rigid transformation. 
@@ -228,7 +221,7 @@ bool IsHairStrandsSupported(EHairStrandsShaderType Type, EShaderPlatform Platfor
 
 bool IsHairStrandsEnabled(EHairStrandsShaderType Type, EShaderPlatform Platform)
 {
-	const bool HairStrandsGlobalEnable = CVarHairStrandsGlobalEnable.GetValueOnAnyThread() > 0 && GHairStrandsPluginEnable > 0;
+	const bool HairStrandsGlobalEnable = IsGroomEnabled();
 	if (!HairStrandsGlobalEnable) return false;
 
 	// Important:
@@ -253,11 +246,6 @@ bool IsHairStrandsEnabled(EHairStrandsShaderType Type, EShaderPlatform Platform)
 	return false;
 }
 
-void SetHairStrandsEnabled(bool In)
-{
-	GHairStrandsPluginEnable = In ? 1 : 0;
-}
-
 bool IsHairStrandsBindingEnable()
 {
 	return CVarHairStrandsBinding.GetValueOnAnyThread() > 0;
diff --git a/Engine/Source/Runtime/Renderer/Private/PostProcess/TemporalSuperResolution.cpp b/Engine/Source/Runtime/Renderer/Private/PostProcess/TemporalSuperResolution.cpp
--- a/Engine/Source/Runtime/Renderer/Private/PostProcess/TemporalSuperResolution.cpp
+++ b/Engine/Source/Runtime/Renderer/Private/PostProcess/TemporalSuperResolution.cpp
@@ -950,7 +950,8 @@ ITemporalUpscaler::FOutputs AddTemporalSuperResolutionPasses(
 		   (!View.Family->EngineShowFlags.VisualizeBuffer || (View.Family->EngineShowFlags.VisualizeBuffer && View.CurrentBufferVisualizationMode == OverviewName))
 		&& (!View.Family->EngineShowFlags.VisualizeNanite || (View.Family->EngineShowFlags.VisualizeNanite && View.CurrentNaniteVisualizationMode == OverviewName))
 		&& (!View.Family->EngineShowFlags.VisualizeLumen  || (View.Family->EngineShowFlags.VisualizeLumen  && View.CurrentLumenVisualizationMode  == OverviewName))
-		&& (!View.Family->EngineShowFlags.VisualizeStrata || (View.Family->EngineShowFlags.VisualizeStrata && View.CurrentStrataVisualizationMode == OverviewName));
+		&& (!View.Family->EngineShowFlags.VisualizeStrata || (View.Family->EngineShowFlags.VisualizeStrata && View.CurrentStrataVisualizationMode == OverviewName))
+		&& (!View.Family->EngineShowFlags.VisualizeGroom  || (View.Family->EngineShowFlags.VisualizeGroom  && View.CurrentGroomVisualizationMode  == OverviewName));
 #endif
 	if (bHasSeparateTranslucency)
 	{
diff --git a/Engine/Source/Runtime/Renderer/Public/HairStrandsInterface.h b/Engine/Source/Runtime/Renderer/Public/HairStrandsInterface.h
--- a/Engine/Source/Runtime/Renderer/Public/HairStrandsInterface.h
+++ b/Engine/Source/Runtime/Renderer/Public/HairStrandsInterface.h
@@ -14,6 +14,7 @@
 #include "RenderResource.h"
 #include "RenderGraphResources.h"
 #include "ShaderPrintParameters.h"
+#include "GroomVisualizationData.h"
 
 class UTexture2D;
 
@@ -52,55 +53,6 @@ RENDERER_API void				ConvertToExternalBufferWithViews(FRDGBuilder& GraphBuilder,
 ////////////////////////////////////////////////////////////////////////////////////////////////////
 // Misc/Helpers
 
-enum class EHairStrandsDebugMode : uint8
-{
-	NoneDebug,
-	SimHairStrands,
-	RenderHairStrands,
-	RenderHairUV,
-	RenderHairRootUV,
-	RenderHairRootUDIM,
-	RenderHairSeed,
-	RenderHairDimension,
-	RenderHairRadiusVariation,
-	RenderHairTangent,
-	RenderHairBaseColor,
-	RenderHairRoughness,
-	RenderVisCluster,
-	RenderHairGroup,
-	RenderLODColoration,
-	RenderHairControlPoints,
-	Count
-};
-
-enum class EHairDebugMode : uint8
-{
-	None,
-	MacroGroups,
-	LightBounds,
-	DeepOpacityMaps,
-	MacroGroupScreenRect,
-	SamplePerPixel,
-	CoverageType,
-	TAAResolveType,
-	VoxelsDensity,
-	VoxelsTangent,
-	VoxelsBaseColor,
-	VoxelsRoughness,
-	MeshProjection,
-	Coverage,
-	MaterialDepth,
-	MaterialBaseColor,
-	MaterialRoughness,
-	MaterialSpecular,
-	MaterialTangent,
-	Tile
-};
-
-/// Return the active debug view mode
-RENDERER_API EHairStrandsDebugMode GetHairStrandsDebugStrandsMode();
-RENDERER_API EHairDebugMode GetHairStrandsDebugMode();
-
 struct FHairStrandClusterCullingData;
 struct IPooledRenderTarget;
 struct FRWBuffer;
@@ -368,7 +320,6 @@ public:
 	bool  bDebugDrawLODInfo = false; // Enable/disable hair LOD info
 	float DebugScreenSize = 0.f;
 	FLinearColor DebugGroupColor;
-	EHairStrandsDebugMode DebugMode = EHairStrandsDebugMode::NoneDebug;
 };
 
 ////////////////////////////////////////////////////////////////////////////////////////////////////
@@ -426,9 +377,9 @@ enum class EHairStrandsShaderType
 	Tool,
 	All
 };
+
 RENDERER_API bool IsHairStrandsSupported(EHairStrandsShaderType Type, EShaderPlatform Platform);
 RENDERER_API bool IsHairStrandsEnabled(EHairStrandsShaderType Type, EShaderPlatform Platform = EShaderPlatform::SP_NumPlatforms);
-RENDERER_API void SetHairStrandsEnabled(bool In);
 
 RENDERER_API bool IsHairRayTracingEnabled();
 
