Replaced UploadDynamicPrimitiveShaderDataForView use with FillViewShaderParameters in cases where no data is actually uploaded.

diff --git a/Engine/Source/Runtime/Renderer/Private/GPUScene.cpp b/Engine/Source/Runtime/Renderer/Private/GPUScene.cpp
--- a/Engine/Source/Runtime/Renderer/Private/GPUScene.cpp
+++ b/Engine/Source/Runtime/Renderer/Private/GPUScene.cpp
@@ -1567,14 +1567,10 @@ void FGPUScene::UploadDynamicPrimitiveShaderDataForViewInternal(FRDGBuilder& Gra
 		UploadGeneral<FUploadDataSourceAdapterDynamicPrimitives>(GraphBuilder, Scene, ExternalAccessQueue, UploadAdapter);
 	}
 
-	// Update view uniform buffer
-	View.CachedViewUniformShaderParameters->PrimitiveSceneData = PrimitiveBuffer->GetSRV();
-	View.CachedViewUniformShaderParameters->LightmapSceneData = LightmapDataBuffer->GetSRV();
-	View.CachedViewUniformShaderParameters->InstancePayloadData = InstancePayloadDataBuffer->GetSRV();
-	View.CachedViewUniformShaderParameters->InstanceSceneData = InstanceSceneDataBuffer->GetSRV();
-	View.CachedViewUniformShaderParameters->InstanceSceneDataSOAStride = InstanceSceneDataSOAStride;
-
-	View.ViewUniformBuffer.UpdateUniformBufferImmediate(*View.CachedViewUniformShaderParameters);
+	if (FillViewShaderParameters(*View.CachedViewUniformShaderParameters))
+	{
+		View.ViewUniformBuffer.UpdateUniformBufferImmediate(*View.CachedViewUniformShaderParameters);
+	}
 
 	// Execute any instance data GPU writer callbacks. (Note: Done after the UB update, in case the user requires it)
 	if (bNeedsUpload) 
@@ -1633,6 +1629,25 @@ void FGPUScene::UploadDynamicPrimitiveShaderDataForViewInternal(FRDGBuilder& Gra
 	}
 }
 
+bool FGPUScene::FillViewShaderParameters(FViewUniformShaderParameters& OutParameters)
+{
+	bool bParametersChanged = false;
+
+	const auto SetParameter = [&](auto& OutParameter, const auto& InParameter)
+	{
+		bParametersChanged |= OutParameter != InParameter;
+		OutParameter = InParameter;
+	};
+
+	SetParameter(OutParameters.PrimitiveSceneData, PrimitiveBuffer->GetSRV());
+	SetParameter(OutParameters.LightmapSceneData, LightmapDataBuffer->GetSRV());
+	SetParameter(OutParameters.InstancePayloadData, InstancePayloadDataBuffer->GetSRV());
+	SetParameter(OutParameters.InstanceSceneData, InstanceSceneDataBuffer->GetSRV());
+	SetParameter(OutParameters.InstanceSceneDataSOAStride, InstanceSceneDataSOAStride);
+
+	return bParametersChanged;
+}
+
 void FGPUScene::AddPrimitiveToUpdate(int32 PrimitiveId, EPrimitiveDirtyState DirtyState)
 {
 	LLM_SCOPE_BYTAG(GPUScene);
diff --git a/Engine/Source/Runtime/Renderer/Private/GPUScene.h b/Engine/Source/Runtime/Renderer/Private/GPUScene.h
--- a/Engine/Source/Runtime/Renderer/Private/GPUScene.h
+++ b/Engine/Source/Runtime/Renderer/Private/GPUScene.h
@@ -220,6 +220,11 @@ public:
 	 */
 	void UploadDynamicPrimitiveShaderDataForView(FRDGBuilder& GraphBuilder, FScene& Scene, FViewInfo& View, FRDGExternalAccessQueue& ExternalAccessQueue, bool bIsShadowView = false);
 
+	/**
+	 * Modifies the GPU scene specific view shader parameters to the current versions. Returns true if any of the parameters changed.
+	 */
+	bool FillViewShaderParameters(FViewUniformShaderParameters& View);
+
 	/**
 	 * Pull all pending updates from Scene and upload primitive & instance data.
 	 */
diff --git a/Engine/Source/Runtime/Renderer/Private/ReflectionEnvironmentRealTimeCapture.cpp b/Engine/Source/Runtime/Renderer/Private/ReflectionEnvironmentRealTimeCapture.cpp
--- a/Engine/Source/Runtime/Renderer/Private/ReflectionEnvironmentRealTimeCapture.cpp
+++ b/Engine/Source/Runtime/Renderer/Private/ReflectionEnvironmentRealTimeCapture.cpp
@@ -563,12 +563,13 @@ void FScene::AllocateAndCaptureFrameSkyEnvMap(
 					CubeView.CachedViewUniformShaderParameters->CameraAerialPerspectiveVolume = GSystemTextures.VolumetricBlackDummy->GetRHI();
 				}
 
-				CubeView.CreateViewUniformBuffers(*CubeView.CachedViewUniformShaderParameters);
 				if (CubeView.bSceneHasSkyMaterial)
 				{
-					GPUScene.UploadDynamicPrimitiveShaderDataForView(GraphBuilder, *this, CubeView, ExternalAccessQueue);
+					GPUScene.FillViewShaderParameters(*CubeView.CachedViewUniformShaderParameters);
 				}
 
+				CubeView.CreateViewUniformBuffers(*CubeView.CachedViewUniformShaderParameters);
+
 				SkyRC.ViewUniformBuffer = CubeView.ViewUniformBuffer;
 				SkyRC.ViewMatrices = &CubeViewMatrices;
 
diff --git a/Engine/Source/Runtime/Renderer/Private/VT/RuntimeVirtualTextureRender.cpp b/Engine/Source/Runtime/Renderer/Private/VT/RuntimeVirtualTextureRender.cpp
--- a/Engine/Source/Runtime/Renderer/Private/VT/RuntimeVirtualTextureRender.cpp
+++ b/Engine/Source/Runtime/Renderer/Private/VT/RuntimeVirtualTextureRender.cpp
@@ -1382,12 +1382,9 @@ namespace RuntimeVirtualTexture
 		View->CachedViewUniformShaderParameters->RuntimeVirtualTextureMipLevel = MipLevelParameter;
 		View->CachedViewUniformShaderParameters->RuntimeVirtualTexturePackHeight = FVector2f(WorldHeightPackParameter);	// LWC_TODO: Precision loss
 		View->CachedViewUniformShaderParameters->RuntimeVirtualTextureDebugParams = FVector4f(DebugType == ERuntimeVirtualTextureDebugType::Debug ? 1.f : 0.f, 0.f, 0.f, 0.f);
+		Scene->GPUScene.FillViewShaderParameters(*View->CachedViewUniformShaderParameters);
 		View->ViewUniformBuffer = TUniformBufferRef<FViewUniformShaderParameters>::CreateUniformBufferImmediate(*View->CachedViewUniformShaderParameters, UniformBuffer_SingleFrame);
 
-		FRDGExternalAccessQueue ExternalAccessQueue;
-		Scene->GPUScene.UploadDynamicPrimitiveShaderDataForView(GraphBuilder, *Scene, *View, ExternalAccessQueue);
-		ExternalAccessQueue.Submit(GraphBuilder);
-
 		// Build graph
 		FRenderGraphSetup::FInitDesc GraphSetupDesc(Scene->GetFeatureLevel(), MaterialType, TextureSize);
 		GraphSetupDesc.OutputFormat0 = OutputTexture0 != nullptr ? OutputTexture0->GetFormat() : PF_Unknown;
