Moved virtual texture uploads earlier in the frame so that they can run async and moved the entire gather phase off the render thread. Saves ~1ms of RT time on CitySample.

diff --git a/Engine/Source/Runtime/Engine/Private/VT/VirtualTextureChunkDDCCache.cpp b/Engine/Source/Runtime/Engine/Private/VT/VirtualTextureChunkDDCCache.cpp
--- a/Engine/Source/Runtime/Engine/Private/VT/VirtualTextureChunkDDCCache.cpp
+++ b/Engine/Source/Runtime/Engine/Private/VT/VirtualTextureChunkDDCCache.cpp
@@ -284,7 +284,7 @@ bool FVirtualTextureChunkDDCCache::MakeChunkAvailable(struct FVirtualTextureData
 {
 	TRACE_CPUPROFILER_EVENT_SCOPE(FVirtualTextureChunkDDCCache::MakeChunkAvailable);
 
-	check(IsInRenderingThread());
+	check(IsInParallelRenderingThread());
 
 	const FString CachedFilePath = AbsoluteCachePath / Chunk->ShortDerivedDataKey;
 	const FString TempFilePath = AbsoluteCachePath / FGuid::NewGuid().ToString() + ".tmp";
diff --git a/Engine/Source/Runtime/Renderer/Private/DeferredShadingRenderer.cpp b/Engine/Source/Runtime/Renderer/Private/DeferredShadingRenderer.cpp
--- a/Engine/Source/Runtime/Renderer/Private/DeferredShadingRenderer.cpp
+++ b/Engine/Source/Runtime/Renderer/Private/DeferredShadingRenderer.cpp
@@ -2411,14 +2411,15 @@ void FDeferredShadingSceneRenderer::Render(FRDGBuilder& GraphBuilder)
 	const bool bHasRayTracedOverlay = HasRayTracedOverlay(ViewFamily);
 	const bool bAllowStaticLighting = !bHasRayTracedOverlay && IsStaticLightingAllowed();
 
+	TUniquePtr<FVirtualTextureUpdater> VirtualTextureUpdater;
+
 	const bool bUseVirtualTexturing = UseVirtualTexturing(FeatureLevel);
 	if (bUseVirtualTexturing)
 	{
-		RDG_CSV_STAT_EXCLUSIVE_SCOPE(GraphBuilder, VirtualTextureUpdate);
-		RDG_GPU_STAT_SCOPE(GraphBuilder, VirtualTextureUpdate);
-		// AllocateResources needs to be called before RHIBeginScene
-		FVirtualTextureSystem::Get().AllocateResources(GraphBuilder, FeatureLevel);
-		FVirtualTextureSystem::Get().CallPendingCallbacks();
+		FVirtualTextureUpdateSettings Settings;
+		Settings.EnableThrottling(!ViewFamily.bOverrideVirtualTextureThrottle);
+
+		VirtualTextureUpdater = FVirtualTextureSystem::Get().BeginUpdate(GraphBuilder, FeatureLevel, Scene, Settings);
 		VirtualTextureFeedbackBegin(GraphBuilder, Views, SceneTexturesConfig.Extent);
 	}
 
@@ -2624,13 +2625,10 @@ void FDeferredShadingSceneRenderer::Render(FRDGBuilder& GraphBuilder)
 	FSceneTextures::InitializeViewFamily(GraphBuilder, ViewFamily);
 	FSceneTextures& SceneTextures = GetActiveSceneTextures();
 
-	// Note, should happen after the GPU-Scene update to ensure rendering to runtime virtual textures is using the correctly updated scene
 	if (bUseVirtualTexturing)
 	{
-		RDG_GPU_STAT_SCOPE(GraphBuilder, VirtualTextureUpdate);
-		FVirtualTextureUpdateSettings Settings;
-		Settings.DisableThrottling(ViewFamily.bOverrideVirtualTextureThrottle);
-		FVirtualTextureSystem::Get().Update(GraphBuilder, FeatureLevel, Scene, Settings);
+		// Note, should happen after the GPU-Scene update to ensure rendering to runtime virtual textures is using the correctly updated scene
+		FVirtualTextureSystem::Get().EndUpdate(GraphBuilder, MoveTemp(VirtualTextureUpdater), FeatureLevel);
 	}
 
 #if RHI_RAYTRACING
@@ -3543,6 +3541,7 @@ void FDeferredShadingSceneRenderer::Render(FRDGBuilder& GraphBuilder)
 		RDG_GPU_STAT_SCOPE(GraphBuilder, RenderDeferredLighting);
 		RDG_CSV_STAT_EXCLUSIVE_SCOPE(GraphBuilder, RenderLighting);
 		SCOPE_CYCLE_COUNTER(STAT_FDeferredShadingSceneRenderer_Lighting);
+		SCOPED_NAMED_EVENT(RenderLighting, FColor::Emerald);
 
 		FRDGTextureRef DynamicBentNormalAOTexture = nullptr;
 
@@ -3678,6 +3677,7 @@ void FDeferredShadingSceneRenderer::Render(FRDGBuilder& GraphBuilder)
 		if (EnumHasAnyFlags(TranslucencyViewsToRender, ETranslucencyView::UnderWater))
 		{
 			RDG_CSV_STAT_EXCLUSIVE_SCOPE(GraphBuilder, RenderTranslucency);
+			SCOPED_NAMED_EVENT(RenderTranslucency, FColor::Emerald);
 			SCOPE_CYCLE_COUNTER(STAT_TranslucencyDrawTime);
 			GraphBuilder.SetCommandListStat(GET_STATID(STAT_CLM_Translucency));
 			const bool bStandardTranslucentCanRenderSeparate = false;
@@ -3720,6 +3720,7 @@ void FDeferredShadingSceneRenderer::Render(FRDGBuilder& GraphBuilder)
 	if (!bHasRayTracedOverlay && ShouldRenderFog(ViewFamily))
 	{
 		RDG_CSV_STAT_EXCLUSIVE_SCOPE(GraphBuilder, RenderFog);
+		SCOPED_NAMED_EVENT(RenderFog, FColor::Emerald);
 		SCOPE_CYCLE_COUNTER(STAT_FDeferredShadingSceneRenderer_RenderFog);
 		RenderFog(GraphBuilder, SceneTextures, LightShaftOcclusionTexture);
 	}
@@ -3762,6 +3763,7 @@ void FDeferredShadingSceneRenderer::Render(FRDGBuilder& GraphBuilder)
 	if (!bHasRayTracedOverlay && TranslucencyViewsToRender != ETranslucencyView::None)
 	{
 		RDG_CSV_STAT_EXCLUSIVE_SCOPE(GraphBuilder, RenderTranslucency);
+		SCOPED_NAMED_EVENT(RenderTranslucency, FColor::Emerald);
 		SCOPE_CYCLE_COUNTER(STAT_TranslucencyDrawTime);
 
 		RDG_EVENT_SCOPE(GraphBuilder, "Translucency");
@@ -4013,7 +4015,7 @@ void FDeferredShadingSceneRenderer::Render(FRDGBuilder& GraphBuilder)
 	{
 		RDG_EVENT_SCOPE(GraphBuilder, "PostProcessing");
 		RDG_GPU_STAT_SCOPE(GraphBuilder, Postprocessing);
-		SCOPE_CYCLE_COUNTER(STAT_FinishRenderViewTargetTime);
+		SCOPED_NAMED_EVENT(PostProcessing, FColor::Emerald);
 
 		GraphBuilder.SetCommandListStat(GET_STATID(STAT_CLM_PostProcessing));
 
diff --git a/Engine/Source/Runtime/Renderer/Private/MobileShadingRenderer.cpp b/Engine/Source/Runtime/Renderer/Private/MobileShadingRenderer.cpp
--- a/Engine/Source/Runtime/Renderer/Private/MobileShadingRenderer.cpp
+++ b/Engine/Source/Runtime/Renderer/Private/MobileShadingRenderer.cpp
@@ -448,15 +448,6 @@ void FMobileSceneRenderer::InitViews(FRDGBuilder& GraphBuilder, FSceneTexturesCo
 		Views[ViewIndex].GPUSceneViewId = InstanceCullingManager.RegisterView(Views[ViewIndex]);
 	}
 
-	if (bUseVirtualTexturing)
-	{
-		RDG_GPU_STAT_SCOPE(GraphBuilder, VirtualTextureUpdate);
-		// AllocateResources needs to be called before RHIBeginScene
-		FVirtualTextureSystem::Get().AllocateResources(GraphBuilder, FeatureLevel);
-		FVirtualTextureSystem::Get().CallPendingCallbacks();
-		VirtualTextureFeedbackBegin(GraphBuilder, Views, SceneTexturesConfig.Extent);
-	}
-
 	FILCUpdatePrimTaskData ILCTaskData;
 	FViewVisibleCommandsPerView ViewCommandsPerView;
 	ViewCommandsPerView.SetNum(Views.Num());
@@ -876,6 +867,17 @@ void FMobileSceneRenderer::Render(FRDGBuilder& GraphBuilder)
 
 	FRDGSystemTextures::Create(GraphBuilder);
 
+	TUniquePtr<FVirtualTextureUpdater> VirtualTextureUpdater;
+
+	if (bUseVirtualTexturing)
+	{
+		FVirtualTextureUpdateSettings Settings;
+		Settings.EnableThrottling(!ViewFamily.bOverrideVirtualTextureThrottle);
+
+		VirtualTextureUpdater = FVirtualTextureSystem::Get().BeginUpdate(GraphBuilder, FeatureLevel, Scene, Settings);
+		VirtualTextureFeedbackBegin(GraphBuilder, Views, SceneTexturesConfig.Extent);
+	}
+
 	// Strata initialization is always run even when not enabled.
 	if (Strata::IsStrataEnabled())
 	{
@@ -928,10 +930,7 @@ void FMobileSceneRenderer::Render(FRDGBuilder& GraphBuilder)
 
 	if (bUseVirtualTexturing)
 	{
-		RDG_GPU_STAT_SCOPE(GraphBuilder, VirtualTextureUpdate);
-		FVirtualTextureUpdateSettings Settings;
-		Settings.DisableThrottling(ViewFamily.bOverrideVirtualTextureThrottle);
-		FVirtualTextureSystem::Get().Update(GraphBuilder, FeatureLevel, Scene, Settings);
+		FVirtualTextureSystem::Get().EndUpdate(GraphBuilder, MoveTemp(VirtualTextureUpdater), FeatureLevel);
 	}
 
 	FSortedLightSetSceneInfo& SortedLightSet = *GraphBuilder.AllocObject<FSortedLightSetSceneInfo>();
diff --git a/Engine/Source/Runtime/Renderer/Private/Renderer.cpp b/Engine/Source/Runtime/Renderer/Private/Renderer.cpp
--- a/Engine/Source/Runtime/Renderer/Private/Renderer.cpp
+++ b/Engine/Source/Runtime/Renderer/Private/Renderer.cpp
@@ -222,12 +222,8 @@ void FRendererModule::DrawTileMesh(FCanvasRenderContext& RenderContext, FMeshPas
 		const bool bUseVirtualTexturing = UseVirtualTexturing(FeatureLevel) && !MeshMaterial.GetUniformVirtualTextureExpressions().IsEmpty();
 		if (bUseVirtualTexturing)
 		{
-			RDG_GPU_STAT_SCOPE(GraphBuilder, VirtualTextureUpdate);
-			FVirtualTextureSystem::Get().AllocateResources(GraphBuilder, FeatureLevel);
-			FVirtualTextureSystem::Get().CallPendingCallbacks();
-
 			FVirtualTextureUpdateSettings Settings;
-			Settings.DisableThrottling(true);
+			Settings.EnableThrottling(false);
 			FVirtualTextureSystem::Get().Update(GraphBuilder, FeatureLevel, Scene, Settings);
 
 			VirtualTextureFeedbackBegin(GraphBuilder, TArrayView<const FViewInfo>(&View, 1), RenderContext.GetViewportRect().Size());
diff --git a/Engine/Source/Runtime/Renderer/Private/SceneHitProxyRendering.cpp b/Engine/Source/Runtime/Renderer/Private/SceneHitProxyRendering.cpp
--- a/Engine/Source/Runtime/Renderer/Private/SceneHitProxyRendering.cpp
+++ b/Engine/Source/Runtime/Renderer/Private/SceneHitProxyRendering.cpp
@@ -187,10 +187,10 @@ void InitHitProxyRender(FRDGBuilder& GraphBuilder, FSceneRenderer* SceneRenderer
 	// Ensure VirtualTexture resources are allocated
 	if (UseVirtualTexturing(FeatureLevel))
 	{
-		FVirtualTextureSystem::Get().AllocateResources(GraphBuilder, FeatureLevel);
-		FVirtualTextureSystem::Get().CallPendingCallbacks();
-		// Because there is no Update(), we need to manually finalize the resources
-		FVirtualTextureSystem::Get().FinalizeResources(GraphBuilder, FeatureLevel);
+		FVirtualTextureUpdateSettings Settings;
+		Settings.EnablePageRequests(false);
+
+		FVirtualTextureSystem::Get().Update(GraphBuilder, FeatureLevel, nullptr, Settings);
 	}
 
 	// Initialize global system textures (pass-through if already initialized).
diff --git a/Engine/Source/Runtime/Renderer/Private/VT/VirtualTextureSystem.cpp b/Engine/Source/Runtime/Renderer/Private/VT/VirtualTextureSystem.cpp
--- a/Engine/Source/Runtime/Renderer/Private/VT/VirtualTextureSystem.cpp
+++ b/Engine/Source/Runtime/Renderer/Private/VT/VirtualTextureSystem.cpp
@@ -157,7 +157,12 @@ static TAutoConsoleVariable<int32> CVarVTCsvStats(
 	TEXT("0=off, 1=on, 2=verbose"),
 	ECVF_Default
 );
-
+static TAutoConsoleVariable<int32> CVarVTAsyncPageRequestTask(
+	TEXT("r.VT.AsyncPageRequestTask"),
+	1,
+	TEXT("Performs VT page requests on an async task."),
+	ECVF_RenderThreadSafe
+);
 
 FVirtualTextureUpdateSettings::FVirtualTextureUpdateSettings()
 {
@@ -213,6 +218,7 @@ public:
 
 	static void DoTask(FFeedbackAnalysisParameters& InParams)
 	{
+		TRACE_CPUPROFILER_EVENT_SCOPE(FVirtualTextureSystem::FeedbackAnalysisTask);
 		InParams.UniquePageList->Initialize();
 		InParams.System->FeedbackAnalysisTask(InParams);
 	}
@@ -284,6 +290,7 @@ public:
 	void DoTask(ENamedThreads::Type CurrentThread, const FGraphEventRef& MyCompletionGraphEvent)
 	{
 		FTaskTagScope TaskTagScope(ETaskTag::EParallelRenderingThread);
+		TRACE_CPUPROFILER_EVENT_SCOPE(FVirtualTextureSystem::GatherRequestsTask);
 		Parameters.RequestList->Initialize();
 		Parameters.System->GatherRequestsTask(Parameters);
 	}
@@ -388,6 +395,7 @@ void FVirtualTextureSystem::FlushCache()
 
 void FVirtualTextureSystem::FlushCache(FVirtualTextureProducerHandle const& ProducerHandle, int32 SpaceID, FIntRect const& TextureRegion, uint32 MaxLevel)
 {
+	check(!bUpdating);
 	checkSlow(IsInRenderingThread());
 
 	SCOPE_CYCLE_COUNTER(STAT_FlushCache);
@@ -766,6 +774,7 @@ void FVirtualTextureSystem::DestroyPendingVirtualTextures(bool bForceDestroyAll)
 
 IAdaptiveVirtualTexture* FVirtualTextureSystem::AllocateAdaptiveVirtualTexture(const FAdaptiveVTDescription& AdaptiveVTDesc, const FAllocatedVTDescription& AllocatedVTDesc)
 {
+	check(!bUpdating);
 	check(IsInRenderingThread());
 	FAdaptiveVirtualTexture* AdaptiveVT = new FAdaptiveVirtualTexture(AdaptiveVTDesc, AllocatedVTDesc);
 	AdaptiveVT->Init(this);
@@ -776,6 +785,7 @@ IAdaptiveVirtualTexture* FVirtualTextureSystem::AllocateAdaptiveVirtualTexture(c
 
 void FVirtualTextureSystem::DestroyAdaptiveVirtualTexture(IAdaptiveVirtualTexture* AdaptiveVT)
 {
+	check(!bUpdating);
 	check(IsInRenderingThread());
 	check(AdaptiveVTs[AdaptiveVT->GetSpaceID()] == AdaptiveVT);
 	AdaptiveVTs[AdaptiveVT->GetSpaceID()] = nullptr;
@@ -784,21 +794,25 @@ void FVirtualTextureSystem::DestroyAdaptiveVirtualTexture(IAdaptiveVirtualTextur
 
 FVirtualTextureProducerHandle FVirtualTextureSystem::RegisterProducer(const FVTProducerDescription& InDesc, IVirtualTexture* InProducer)
 {
+	check(!bUpdating);
 	return Producers.RegisterProducer(this, InDesc, InProducer);
 }
 
 void FVirtualTextureSystem::ReleaseProducer(const FVirtualTextureProducerHandle& Handle)
 {
+	check(!bUpdating);
 	Producers.ReleaseProducer(this, Handle);
 }
 
 void FVirtualTextureSystem::AddProducerDestroyedCallback(const FVirtualTextureProducerHandle& Handle, FVTProducerDestroyedFunction* Function, void* Baton)
 {
+	check(!bUpdating);
 	Producers.AddDestroyedCallback(Handle, Function, Baton);
 }
 
 uint32 FVirtualTextureSystem::RemoveAllProducerDestroyedCallbacks(const void* Baton)
 {
+	check(!bUpdating);
 	return Producers.RemoveAllCallbacks(Baton);
 }
 
@@ -809,6 +823,7 @@ FVirtualTextureProducer* FVirtualTextureSystem::FindProducer(const FVirtualTextu
 
 FVirtualTextureSpace* FVirtualTextureSystem::AcquireSpace(const FVTSpaceDescription& InDesc, uint8 InForceSpaceID, FAllocatedVirtualTexture* AllocatedVT)
 {
+	check(!bUpdating);
 	LLM_SCOPE(ELLMTag::VirtualTextureSystem);
 
 	uint32 NumFailedAllocations = 0u;
@@ -874,6 +889,7 @@ FVirtualTextureSpace* FVirtualTextureSystem::AcquireSpace(const FVTSpaceDescript
 
 void FVirtualTextureSystem::ReleaseSpace(FVirtualTextureSpace* Space)
 {
+	check(!bUpdating);
 	check(IsInRenderingThread());
 	const uint32 NumRefs = Space->Release();
 	if (NumRefs == 0u && Space->GetDescription().bPrivateSpace)
@@ -1064,6 +1080,7 @@ void FVirtualTextureSystem::ReleasePendingSpaces()
 
 void FVirtualTextureSystem::LockTile(const FVirtualTextureLocalTile& Tile)
 {
+	check(!bUpdating);
 	check(IsInRenderingThread());
 
 	if (TileLocks.Lock(Tile))
@@ -1089,6 +1106,7 @@ static void UnlockTileInternal(const FVirtualTextureProducerHandle& ProducerHand
 
 void FVirtualTextureSystem::UnlockTile(const FVirtualTextureLocalTile& Tile, const FVirtualTextureProducer* Producer)
 {
+	check(!bUpdating);
 	check(IsInRenderingThread());
 
 	if (TileLocks.Unlock(Tile))
@@ -1106,6 +1124,7 @@ void FVirtualTextureSystem::UnlockTile(const FVirtualTextureLocalTile& Tile, con
 
 void FVirtualTextureSystem::ForceUnlockAllTiles(const FVirtualTextureProducerHandle& ProducerHandle, const FVirtualTextureProducer* Producer)
 {
+	check(!bUpdating);
 	check(IsInRenderingThread());
 
 	TArray<FVirtualTextureLocalTile> TilesToUnlock;
@@ -1135,6 +1154,7 @@ static float ComputeMipLevel(const IAllocatedVirtualTexture* AllocatedVT, const
 
 void FVirtualTextureSystem::RequestTiles(const FVector2D& InScreenSpaceSize, int32 InMipLevel)
 {
+	check(!bUpdating);
 	check(IsInRenderingThread());
 
 	FScopeLock Lock(&RequestedTilesLock);
@@ -1147,6 +1167,7 @@ void FVirtualTextureSystem::RequestTiles(const FVector2D& InScreenSpaceSize, int
 
 void FVirtualTextureSystem::RequestTiles(const FMaterialRenderProxy* InMaterialRenderProxy, const FVector2D& InScreenSpaceSize, ERHIFeatureLevel::Type InFeatureLevel)
 {
+	check(!bUpdating);
 	check(IsInRenderingThread());
 
 	FScopeLock Lock(&RequestedTilesLock);
@@ -1206,6 +1227,7 @@ void FVirtualTextureSystem::RequestTilesForRegion(IAllocatedVirtualTexture* Allo
 
 void FVirtualTextureSystem::LoadPendingTiles(FRDGBuilder& GraphBuilder, ERHIFeatureLevel::Type FeatureLevel)
 {
+	check(!bUpdating);
 	check(IsInRenderingThread());
 
 	TArray<uint32> PackedTiles;
@@ -1232,8 +1254,9 @@ void FVirtualTextureSystem::LoadPendingTiles(FRDGBuilder& GraphBuilder, ERHIFeat
 		RequestList->Initialize();
 		GatherRequests(RequestList, UniquePageList, Frame, Allocator, Settings);
 		// No need to sort requests, since we're submitting all of them here (no throttling)
-		AllocateResources(GraphBuilder, FeatureLevel);
-		SubmitRequests(GraphBuilder, FeatureLevel, Allocator, Settings, RequestList, false);
+		AllocateResources(GraphBuilder);
+		BeginSubmitRequests(GraphBuilder.RHICmdList, FeatureLevel, Allocator, Settings, RequestList, false);
+		EndSubmitRequests(GraphBuilder);
 	}
 }
 
@@ -1340,310 +1363,6 @@ void FVirtualTextureSystem::FeedbackAnalysisTask(const FFeedbackAnalysisParamete
 	}
 }
 
-void FVirtualTextureSystem::Update(FRDGBuilder& GraphBuilder, ERHIFeatureLevel::Type FeatureLevel, FScene* Scene, FVirtualTextureUpdateSettings const& Settings)
-{
-	check(IsInRenderingThread());
-
-	RDG_CSV_STAT_EXCLUSIVE_SCOPE(GraphBuilder, VirtualTextureSystem_Update);
-	TRACE_CPUPROFILER_EVENT_SCOPE(FVirtualTextureSystem::Update);
-	SCOPE_CYCLE_COUNTER(STAT_VirtualTextureSystem_Update);
-	RDG_GPU_STAT_SCOPE(GraphBuilder, VirtualTexture);
-	
-	// Update Adaptive VTs. This can trigger allocation/destruction of VTs and must happen before the flush below.
-	{
-		SCOPE_CYCLE_COUNTER(STAT_ProcessRequests_UpdateAdaptiveAllocations);
-		for (uint32 ID = 0; ID < MaxSpaces; ID++)
-		{
-			if (AdaptiveVTs[ID])
-			{
-				AdaptiveVTs[ID]->UpdateAllocations(this, GraphBuilder.RHICmdList, Frame);
-			}
-		}
-	}
-
-	if (bFlushCaches)
-	{
-		SCOPE_CYCLE_COUNTER(STAT_FlushCache);
-		INC_DWORD_STAT_BY(STAT_NumFlushCache, 1);
-
-		for (int32 i = 0; i < PhysicalSpaces.Num(); ++i)
-		{
-			FVirtualTexturePhysicalSpace* PhysicalSpace = PhysicalSpaces[i];
-			if (PhysicalSpace)
-			{
-				if (CVarVTProduceLockedTilesOnFlush.GetValueOnRenderThread())
-				{
-					// Collect locked pages to be produced again
-					PhysicalSpace->GetPagePool().GetAllLockedPages(this, MappedTilesToProduce);
-				}
-				// Flush unlocked pages
-				PhysicalSpace->GetPagePool().EvictAllPages(this);
-			}
-		}
-
-		bFlushCaches = false;
-	}
-
-	DestroyPendingVirtualTextures(false);
-
-	// Early out when no allocated VTs
-	if (AllocatedVTs.Num() == 0)
-	{
-		MappedTilesToProduce.Reset();
-		return;
-	}
-
-	// Flush any dirty runtime virtual textures for the current scene
-	if (Scene != nullptr)
-	{
-		// Only flush if we know that there is GPU feedback available to refill the visible data this frame
-		// This prevents bugs when low frame rate causes feedback buffer to stall so that the physical cache isn't filled immediately which causes visible glitching
-		if (GVirtualTextureFeedback.CanMap(GraphBuilder.RHICmdList))
-		{
-			// Each RVT will call FVirtualTextureSystem::FlushCache()
-			Scene->FlushDirtyRuntimeVirtualTextures();
-		}
-	}
-
-	FConcurrentLinearBulkObjectAllocator Allocator;
-
-	FUniquePageList* MergedUniquePageList = Allocator.Create<FUniquePageList>();
-	MergedUniquePageList->Initialize();
-	
-	if (Settings.bEnableFeedback)
-	{
-		// Fetch feedback for analysis
-		FVirtualTextureFeedback::FMapResult FeedbackResult;
-
-		{
-			SCOPE_CYCLE_COUNTER(STAT_FeedbackMap);
-			FeedbackResult = GVirtualTextureFeedback.Map(GraphBuilder.RHICmdList);
-		}
-
-		// Create tasks to read the feedback data
-		// Give each task a section of the feedback buffer to analyze
-		FFeedbackAnalysisParameters FeedbackAnalysisParameters[MaxNumTasks];
-
-		const uint32 MaxNumFeedbackTasks = FMath::Clamp((uint32)Settings.NumFeedbackTasks, 1u, MaxNumTasks);
-		const uint32 FeedbackSizePerTask = FMath::DivideAndRoundUp(FeedbackResult.Size, MaxNumFeedbackTasks);
-
-		uint32 NumFeedbackTasks = 0;
-		uint32 CurrentOffset = 0;
-		while (CurrentOffset < FeedbackResult.Size)
-		{
-			const uint32 TaskIndex = NumFeedbackTasks++;
-			FFeedbackAnalysisParameters& Params = FeedbackAnalysisParameters[TaskIndex];
-			Params.System = this;
-			if (TaskIndex == 0u)
-			{
-				Params.UniquePageList = MergedUniquePageList;
-			}
-			else
-			{
-				Params.UniquePageList = Allocator.Create<FUniquePageList>();
-			}
-			Params.FeedbackBuffer = FeedbackResult.Data + CurrentOffset;
-
-			const uint32 Size = FMath::Min(FeedbackSizePerTask, FeedbackResult.Size - CurrentOffset);
-			Params.FeedbackSize = Size;
-			CurrentOffset += Size;
-		}
-
-		// Kick the tasks
-		const int32 LocalFeedbackTaskCount = Settings.bParallelFeedbackTasks ? 1 : NumFeedbackTasks;
-		const int32 WorkerFeedbackTaskCount = NumFeedbackTasks - LocalFeedbackTaskCount;
-
-		FGraphEventArray Tasks;
-		if (WorkerFeedbackTaskCount > 0)
-		{
-			SCOPE_CYCLE_COUNTER(STAT_ProcessRequests_SubmitTasks);
-			Tasks.Reserve(WorkerFeedbackTaskCount);
-			for (uint32 TaskIndex = LocalFeedbackTaskCount; TaskIndex < NumFeedbackTasks; ++TaskIndex)
-			{
-				Tasks.Add(TGraphTask<FFeedbackAnalysisTask>::CreateTask().ConstructAndDispatchWhenReady(FeedbackAnalysisParameters[TaskIndex]));
-			}
-		}
-
-		if (NumFeedbackTasks > 0u)
-		{
-			SCOPE_CYCLE_COUNTER(STAT_FeedbackAnalysis);
-
-			for (int32 TaskIndex = 0; TaskIndex < LocalFeedbackTaskCount; ++TaskIndex)
-			{
-				FFeedbackAnalysisTask::DoTask(FeedbackAnalysisParameters[TaskIndex]);
-			}
-			if (WorkerFeedbackTaskCount > 0)
-			{
-				SCOPE_CYCLE_COUNTER(STAT_ProcessRequests_WaitTasks);
-
-				FTaskGraphInterface::Get().WaitUntilTasksComplete(Tasks, ENamedThreads::GetRenderThread_Local());
-			}
-		}
-
-		if (NumFeedbackTasks > 1u)
-		{
-			SCOPE_CYCLE_COUNTER(STAT_ProcessRequests_MergePages);
-			for (uint32 TaskIndex = 1u; TaskIndex < NumFeedbackTasks; ++TaskIndex)
-			{
-				MergedUniquePageList->MergePages(FeedbackAnalysisParameters[TaskIndex].UniquePageList);
-			}
-		}
-
-		GVirtualTextureFeedback.Unmap(GraphBuilder.RHICmdList, FeedbackResult.MapHandle);
-	}
-
-#if WITH_EDITOR
-	// If we're are recording page requests, then copy off pages to the recording buffer.
-	if (PageRequestRecordHandle != ~0ull)
-	{
-		RecordPageRequests(MergedUniquePageList, PageRequestRecordBuffer);
-	}
-#endif
-
-	// Add any page requests from recording playback.
-	if (PageRequestPlaybackBuffer.Num() > 0)
-	{
-		if (Settings.bEnablePlayback)
-		{
-			// todo: We can split this into concurrent tasks. 
-			FAddRequestedTilesParameters Parameters;
-			Parameters.System = this;
-			Parameters.LevelBias = FMath::FloorToInt(CVarVTPlaybackMipBias.GetValueOnRenderThread() + GetGlobalMipBias() + 0.5f);
-			Parameters.RequestBuffer = PageRequestPlaybackBuffer.GetData();
-			Parameters.NumRequests = PageRequestPlaybackBuffer.Num();
-			Parameters.UniquePageList = Allocator.Create<FUniquePageList>();
-
-			FAddRequestedTilesTask::DoTask(Parameters);
-			MergedUniquePageList->MergePages(Parameters.UniquePageList);
-		}
-
-		PageRequestPlaybackBuffer.Reset(0);
-	}
-
-	FUniqueRequestList* MergedRequestList = Allocator.Create<FUniqueRequestList>(Allocator);
-	MergedRequestList->Initialize();
-
-	// Collect tiles to lock
-	{
-		TArray<FVirtualTextureLocalTile> RemainingTilesToLock;
-		for (const FVirtualTextureLocalTile& Tile : TilesToLock)
-		{
-			const FVirtualTextureProducerHandle ProducerHandle = Tile.GetProducerHandle();
-			const FVirtualTextureProducer* Producer = Producers.FindProducer(ProducerHandle);
-			checkSlow(TileLocks.IsLocked(Tile));
-			if (Producer)
-			{
-				uint8 ProducerLayerMaskToLoad = 0u;
-				for (uint32 ProducerLayerIndex = 0u; ProducerLayerIndex < Producer->GetNumTextureLayers(); ++ProducerLayerIndex)
-				{
-					uint32 GroupIndex = Producer->GetPhysicalGroupIndexForTextureLayer(ProducerLayerIndex);
-					FVirtualTexturePhysicalSpace* PhysicalSpace = Producer->GetPhysicalSpaceForPhysicalGroup(GroupIndex);
-					FTexturePagePool& PagePool = PhysicalSpace->GetPagePool();
-					const uint32 pAddress = PagePool.FindPageAddress(ProducerHandle, GroupIndex, Tile.Local_vAddress, Tile.Local_vLevel);
-					if (pAddress == ~0u)
-					{
-						ProducerLayerMaskToLoad |= (1u << ProducerLayerIndex);
-					}
-					else
-					{
-						PagePool.Lock(pAddress);
-					}
-				}
-
-				if (ProducerLayerMaskToLoad != 0u)
-				{
-					const uint16 LoadRequestIndex = MergedRequestList->LockLoadRequest(FVirtualTextureLocalTile(Tile.GetProducerHandle(), Tile.Local_vAddress, Tile.Local_vLevel), ProducerLayerMaskToLoad);
-					if (LoadRequestIndex == 0xffff)
-					{
-						// Overflowed the request list...try to lock the tile again next frame
-						RemainingTilesToLock.Add(Tile);
-					}
-				}
-			}
-		}
-
-		TilesToLock = MoveTemp(RemainingTilesToLock);
-	}
-
-	TArray<uint32> PackedTiles;
-	if(RequestedPackedTiles.Num() > 0)
-	{
-		FScopeLock Lock(&RequestedTilesLock);
-		PackedTiles = MoveTemp(RequestedPackedTiles);
-		RequestedPackedTiles.Reset();
-	}
-
-	if (PackedTiles.Num() > 0)
-	{
-		// Collect explicitly requested tiles
-		// These tiles are generated on the current frame, so they are collected/processed in a separate list
-		FUniquePageList* RequestedPageList = Allocator.Create<FUniquePageList>();
-		RequestedPageList->Initialize();
-		for (uint32 Tile : PackedTiles)
-		{
-			RequestedPageList->Add(Tile, 0xffff);
-		}
-		GatherRequests(MergedRequestList, RequestedPageList, Frame, Allocator, Settings);
-	}
-
-	// Pages from feedback buffer were generated several frames ago, so they may no longer be valid for newly allocated VTs
-	static uint32 PendingFrameDelay = 3u;
-	if (Frame >= PendingFrameDelay)
-	{
-		GatherRequests(MergedRequestList, MergedUniquePageList, Frame - PendingFrameDelay, Allocator, Settings);
-	}
-
-	if (MergedRequestList->GetNumAdaptiveAllocationRequests() > 0)
-	{
-		SCOPE_CYCLE_COUNTER(STAT_ProcessRequests_QueueAdaptiveRequests);
-		FAdaptiveVirtualTexture::QueuePackedAllocationRequests(this, &MergedRequestList->GetAdaptiveAllocationRequest(0), MergedRequestList->GetNumAdaptiveAllocationRequests(), Frame);
-	}
-
-	{
-		SCOPE_CYCLE_COUNTER(STAT_ProcessRequests_Sort);
-
-		// Limit the number of uploads (account for MappedTilesToProduce this frame)
-		// Are all pages equal? Should there be different limits on different types of pages?
-		const int32 MaxNumUploads = Settings.MaxPageUploads;
-		const int32 MaxRequestUploads = FMath::Max(MaxNumUploads - MappedTilesToProduce.Num(), 1);
-
-		if (MaxRequestUploads < (int32)MergedRequestList->GetNumLoadRequests())
-		{
-			// Dropping requests is normal but track to log here if we want to tune settings.
-			if (CVarVTVerbose.GetValueOnRenderThread())
-			{
-				UE_LOG(LogConsoleResponse, Display, TEXT("VT dropped %d load requests."), MergedRequestList->GetNumLoadRequests() - MaxRequestUploads);
-			}
-		}
-
-		MergedRequestList->SortRequests(Producers, Allocator, MaxRequestUploads);
-	}
-
-	{
-		// After sorting and clamping the load requests, if we still have unused upload bandwidth then use it to add some continous updates
-		const int32 MaxNumUploads = Settings.MaxPageUploads;
-		const int32 MaxTilesToProduce = FMath::Max(MaxNumUploads - MappedTilesToProduce.Num() - (int32)MergedRequestList->GetNumLoadRequests(), 0);
-		const int32 MaxContinuousUpdates = Settings.MaxContinuousUpdates;
-
-		GetContinuousUpdatesToProduce(MergedRequestList, MaxTilesToProduce, MaxContinuousUpdates);
-	}
-
-	// Submit the requests to produce pages that are already mapped
-	SubmitPreMappedRequests(GraphBuilder, FeatureLevel);
-	// Submit the merged requests
-	SubmitRequests(GraphBuilder, FeatureLevel, Allocator, Settings, MergedRequestList, true);
-
-	Producers.NotifyRequestsCompleted();
-
-	UpdateResidencyTracking();
-
-#if !UE_BUILD_SHIPPING
-	UpdateCsvStats();
-#endif
-
-	ReleasePendingSpaces();
-}
-
 void FVirtualTextureSystem::GatherRequests(FUniqueRequestList* MergedRequestList, const FUniquePageList* UniquePageList, uint32 FrameRequested, FConcurrentLinearBulkObjectAllocator& Allocator, FVirtualTextureUpdateSettings const& Settings)
 {
 	const uint32 MaxNumGatherTasks = FMath::Clamp((uint32)Settings.NumGatherTasks, 1u, MaxNumTasks);
@@ -1708,7 +1427,7 @@ void FVirtualTextureSystem::GatherRequests(FUniqueRequestList* MergedRequestList
 		{
 			SCOPE_CYCLE_COUNTER(STAT_ProcessRequests_WaitTasks);
 
-			FTaskGraphInterface::Get().WaitUntilTasksComplete(Tasks, ENamedThreads::GetRenderThread_Local());
+			FTaskGraphInterface::Get().WaitUntilTasksComplete(Tasks);
 		}
 	}
 
@@ -2204,7 +1923,7 @@ void FVirtualTextureSystem::UpdateResidencyTracking() const
 	}
 }
 
-void FVirtualTextureSystem::SubmitRequestsFromLocalTileList(TArray<FVirtualTextureLocalTile>& OutDeferredTiles, const TSet<FVirtualTextureLocalTile>& LocalTileList, EVTProducePageFlags Flags, FRDGBuilder& GraphBuilder, ERHIFeatureLevel::Type FeatureLevel)
+void FVirtualTextureSystem::SubmitRequestsFromLocalTileList(FRHICommandList& RHICmdList, TArray<FVirtualTextureLocalTile>& OutDeferredTiles, const TSet<FVirtualTextureLocalTile>& LocalTileList, EVTProducePageFlags Flags, ERHIFeatureLevel::Type FeatureLevel)
 {
 	LLM_SCOPE(ELLMTag::VirtualTextureSystem);
 
@@ -2263,7 +1982,7 @@ void FVirtualTextureSystem::SubmitRequestsFromLocalTileList(TArray<FVirtualTextu
 		}
 
 		IVirtualTextureFinalizer* VTFinalizer = Producer->GetVirtualTexture()->ProducePageData(
-			GraphBuilder.RHICmdList, FeatureLevel,
+			RHICmdList, FeatureLevel,
 			Flags,
 			ProducerHandle, LayerMask, Tile.Local_vLevel, Tile.Local_vAddress,
 			RequestPageResult.Handle,
@@ -2277,13 +1996,13 @@ void FVirtualTextureSystem::SubmitRequestsFromLocalTileList(TArray<FVirtualTextu
 	}
 }
 
-void FVirtualTextureSystem::SubmitPreMappedRequests(FRDGBuilder& GraphBuilder, ERHIFeatureLevel::Type FeatureLevel)
+void FVirtualTextureSystem::SubmitPreMappedRequests(FRHICommandList& RHICmdList, ERHIFeatureLevel::Type FeatureLevel)
 {
 	check(TransientCollectedPages.Num() == 0);
 
 	{
 		INC_DWORD_STAT_BY(STAT_NumMappedPageUpdate, MappedTilesToProduce.Num());
-		SubmitRequestsFromLocalTileList(TransientCollectedPages, MappedTilesToProduce, EVTProducePageFlags::None, GraphBuilder, FeatureLevel);
+		SubmitRequestsFromLocalTileList(RHICmdList, TransientCollectedPages, MappedTilesToProduce, EVTProducePageFlags::None, FeatureLevel);
 		MappedTilesToProduce.Reset();
 		MappedTilesToProduce.Append(TransientCollectedPages);
 		TransientCollectedPages.Reset();
@@ -2291,15 +2010,15 @@ void FVirtualTextureSystem::SubmitPreMappedRequests(FRDGBuilder& GraphBuilder, E
 
 	{
 		INC_DWORD_STAT_BY(STAT_NumContinuousPageUpdate, ContinuousUpdateTilesToProduce.Num());
-		SubmitRequestsFromLocalTileList(TransientCollectedPages, ContinuousUpdateTilesToProduce, EVTProducePageFlags::ContinuousUpdate, GraphBuilder, FeatureLevel);
+		SubmitRequestsFromLocalTileList(RHICmdList, TransientCollectedPages, ContinuousUpdateTilesToProduce, EVTProducePageFlags::ContinuousUpdate, FeatureLevel);
 		ContinuousUpdateTilesToProduce.Reset();
 		TransientCollectedPages.Reset();
 	}
 }
 
-void FVirtualTextureSystem::SubmitRequests(FRDGBuilder& GraphBuilder, ERHIFeatureLevel::Type FeatureLevel, FConcurrentLinearBulkObjectAllocator& Allocator, FVirtualTextureUpdateSettings const& Settings, FUniqueRequestList* RequestList, bool bAsync)
+void FVirtualTextureSystem::BeginSubmitRequests(FRHICommandList& RHICmdList, ERHIFeatureLevel::Type FeatureLevel, FConcurrentLinearBulkObjectAllocator& Allocator, FVirtualTextureUpdateSettings const& Settings, FUniqueRequestList* RequestList, bool bAsync)
 {
-	TRACE_CPUPROFILER_EVENT_SCOPE(FVirtualTextureSystem::SubmitRequests);
+	TRACE_CPUPROFILER_EVENT_SCOPE(FVirtualTextureSystem::BeginSubmitRequests);
 	LLM_SCOPE(ELLMTag::VirtualTextureSystem);
 
 	// Allocate space to hold the physical address we allocate for each page load (1 page per layer per request)
@@ -2382,7 +2101,7 @@ void FVirtualTextureSystem::SubmitRequests(FRDGBuilder& GraphBuilder, ERHIFeatur
 				//checkf(!bLockTile, TEXT("Tried to lock an invalid VT tile"));
 
 				bTileInvalid = true;
-				if (CVarVTVerbose.GetValueOnRenderThread())
+				if (CVarVTVerbose.GetValueOnAnyThread())
 				{
 					UE_LOG(LogConsoleResponse, Display, TEXT("vAddr %i@%i is not a valid request for AllocatedVT but is still requested."), TileToLoad.Local_vAddress, TileToLoad.Local_vLevel);
 				}
@@ -2527,13 +2246,13 @@ void FVirtualTextureSystem::SubmitRequests(FRDGBuilder& GraphBuilder, ERHIFeatur
 				{
 					TRACE_CPUPROFILER_EVENT_SCOPE(FVirtualTextureSystem::ProcessRequests_Wait);
 					QUICK_SCOPE_CYCLE_COUNTER(ProcessRequests_Wait);
-					FTaskGraphInterface::Get().WaitUntilTasksComplete(ProducePageTasks, ENamedThreads::GetRenderThread_Local());
+					FTaskGraphInterface::Get().WaitUntilTasksComplete(ProducePageTasks);
 				}
 			}
 
 			for (FProducePageDataPrepareTask& Task : PrepareTasks)
 			{
-				IVirtualTextureFinalizer* VTFinalizer = Task.VirtualTexture->ProducePageData(GraphBuilder.RHICmdList, FeatureLevel,
+				IVirtualTextureFinalizer* VTFinalizer = Task.VirtualTexture->ProducePageData(RHICmdList, FeatureLevel,
 					Task.Flags,
 					Task.ProducerHandle, Task.LayerMask, Task.vLevel, Task.vAddress,
 					Task.RequestHandle,
@@ -2604,11 +2323,17 @@ void FVirtualTextureSystem::SubmitRequests(FRDGBuilder& GraphBuilder, ERHIFeatur
 
 		AllocatedVTsToMap.Shrink();
 	}
+}
+
+void FVirtualTextureSystem::EndSubmitRequests(FRDGBuilder& GraphBuilder)
+{
+	TRACE_CPUPROFILER_EVENT_SCOPE(FVirtualTextureSystem::EndSubmitRequests);
 
 	// Finalize requests
 	{
 		SCOPE_CYCLE_COUNTER(STAT_ProcessRequests_Finalize);
 		RDG_EVENT_SCOPE(GraphBuilder, "VirtualTextureFinalizeRequests");
+		TRACE_CPUPROFILER_EVENT_SCOPE(FVirtualTextureSystem::Finalize);
 		
 		for (IVirtualTextureFinalizer* VTFinalizer : Finalizers)
 		{
@@ -2621,6 +2346,7 @@ void FVirtualTextureSystem::SubmitRequests(FRDGBuilder& GraphBuilder, ERHIFeatur
 	{
 		SCOPE_CYCLE_COUNTER(STAT_PageTableUpdates);
 		RDG_EVENT_SCOPE(GraphBuilder, "VirtualTexturePageTableUpdates");
+		TRACE_CPUPROFILER_EVENT_SCOPE(FVirtualTextureSystem::ApplyUpdates);
 
 		for (uint32 ID = 0; ID < MaxSpaces; ID++)
 		{
@@ -2634,10 +2360,11 @@ void FVirtualTextureSystem::SubmitRequests(FRDGBuilder& GraphBuilder, ERHIFeatur
 	Frame++;
 }
 
-void FVirtualTextureSystem::AllocateResources(FRDGBuilder& GraphBuilder, ERHIFeatureLevel::Type FeatureLevel)
+void FVirtualTextureSystem::AllocateResources(FRDGBuilder& GraphBuilder)
 {
 	LLM_SCOPE(ELLMTag::VirtualTextureSystem);
 	RDG_GPU_STAT_SCOPE(GraphBuilder, VirtualTextureAllocate);
+	RDG_CSV_STAT_EXCLUSIVE_SCOPE(GraphBuilder, VirtualTextureSystem_Update);
 
 	for (uint32 ID = 0; ID < MaxSpaces; ID++)
 	{
@@ -2648,26 +2375,377 @@ void FVirtualTextureSystem::AllocateResources(FRDGBuilder& GraphBuilder, ERHIFea
 	}
 }
 
-void FVirtualTextureSystem::FinalizeResources(FRDGBuilder& GraphBuilder, ERHIFeatureLevel::Type FeatureLevel)
+TUniquePtr<FVirtualTextureUpdater> FVirtualTextureSystem::BeginUpdate(FRDGBuilder& GraphBuilder, ERHIFeatureLevel::Type FeatureLevel, FScene* Scene, const FVirtualTextureUpdateSettings& Settings)
 {
-	LLM_SCOPE(ELLMTag::VirtualTextureSystem);
+	check(IsInRenderingThread());
+	check(!bUpdating);
 
-	for (uint32 ID = 0; ID < MaxSpaces; ID++)
+	AllocateResources(GraphBuilder);
+
+	RDG_CSV_STAT_EXCLUSIVE_SCOPE(GraphBuilder, VirtualTextureSystem_Update);
+	TRACE_CPUPROFILER_EVENT_SCOPE(FVirtualTextureSystem::BeginUpdate);
+	SCOPE_CYCLE_COUNTER(STAT_VirtualTextureSystem_Update);
+	RDG_GPU_STAT_SCOPE(GraphBuilder, VirtualTexture);
+
+	Producers.CallPendingCallbacks();
+
+	if (!Settings.bEnablePageRequests)
 	{
-		if (Spaces[ID])
+		for (uint32 ID = 0; ID < MaxSpaces; ID++)
+		{
+			if (Spaces[ID])
+			{
+				Spaces[ID]->FinalizeTextures(GraphBuilder);
+			}
+		}
+
+		return {};
+	}
+
+	// Update Adaptive VTs. This can trigger allocation/destruction of VTs and must happen before the flush below.
+	{
+		SCOPE_CYCLE_COUNTER(STAT_ProcessRequests_UpdateAdaptiveAllocations);
+		for (uint32 ID = 0; ID < MaxSpaces; ID++)
+		{
+			if (AdaptiveVTs[ID])
+			{
+				AdaptiveVTs[ID]->UpdateAllocations(this, GraphBuilder.RHICmdList, Frame);
+			}
+		}
+	}
+
+	if (bFlushCaches)
+	{
+		SCOPE_CYCLE_COUNTER(STAT_FlushCache);
+		INC_DWORD_STAT_BY(STAT_NumFlushCache, 1);
+
+		for (int32 i = 0; i < PhysicalSpaces.Num(); ++i)
+		{
+			FVirtualTexturePhysicalSpace* PhysicalSpace = PhysicalSpaces[i];
+			if (PhysicalSpace)
+			{
+				if (CVarVTProduceLockedTilesOnFlush.GetValueOnRenderThread())
+				{
+					// Collect locked pages to be produced again
+					PhysicalSpace->GetPagePool().GetAllLockedPages(this, MappedTilesToProduce);
+				}
+				// Flush unlocked pages
+				PhysicalSpace->GetPagePool().EvictAllPages(this);
+			}
+		}
+
+		bFlushCaches = false;
+	}
+
+	DestroyPendingVirtualTextures(false);
+
+	// Early out when no allocated VTs
+	if (AllocatedVTs.Num() == 0)
+	{
+		MappedTilesToProduce.Reset();
+		return {};
+	}
+
+	// Flush any dirty runtime virtual textures for the current scene
+	if (Scene != nullptr)
+	{
+		// Only flush if we know that there is GPU feedback available to refill the visible data this frame
+		// This prevents bugs when low frame rate causes feedback buffer to stall so that the physical cache isn't filled immediately which causes visible glitching
+		if (GVirtualTextureFeedback.CanMap(GraphBuilder.RHICmdList))
 		{
-			Spaces[ID]->FinalizeTextures(GraphBuilder);
+			// Each RVT will call FVirtualTextureSystem::FlushCache()
+			Scene->FlushDirtyRuntimeVirtualTextures();
 		}
 	}
+
+	// Mark updating to true now that we are potentially launching async tasks.
+	bUpdating = true;
+
+	TUniquePtr<FVirtualTextureUpdater> Updater(new FVirtualTextureUpdater());
+
+	if (Settings.bEnableFeedback)
+	{
+		SCOPE_CYCLE_COUNTER(STAT_FeedbackMap);
+		Updater->FeedbackMapResult = GVirtualTextureFeedback.Map(GraphBuilder.RHICmdList);
+	}
+
+	Updater->AsyncTask = GraphBuilder.AddCommandListSetupTask([this, Settings, FeedbackResult = Updater->FeedbackMapResult, FeatureLevel] (FRHICommandList& RHICmdList)
+	{
+		TRACE_CPUPROFILER_EVENT_SCOPE(FVirtualTextureSystem::ProcessPageRequests);
+
+		FConcurrentLinearBulkObjectAllocator Allocator;
+		FUniqueRequestList* MergedRequestList = Allocator.Create<FUniqueRequestList>(Allocator);
+
+		MergedRequestList->Initialize();
+
+		FUniquePageList* MergedUniquePageList = Allocator.Create<FUniquePageList>();
+		MergedUniquePageList->Initialize();
+
+		if (Settings.bEnableFeedback)
+		{
+			// Create tasks to read the feedback data
+			// Give each task a section of the feedback buffer to analyze
+			FFeedbackAnalysisParameters FeedbackAnalysisParameters[MaxNumTasks];
+
+			const uint32 MaxNumFeedbackTasks = FMath::Clamp((uint32)Settings.NumFeedbackTasks, 1u, MaxNumTasks);
+			const uint32 FeedbackSizePerTask = FMath::DivideAndRoundUp(FeedbackResult.Size, MaxNumFeedbackTasks);
+
+			uint32 NumFeedbackTasks = 0;
+			uint32 CurrentOffset = 0;
+			while (CurrentOffset < FeedbackResult.Size)
+			{
+				const uint32 TaskIndex = NumFeedbackTasks++;
+				FFeedbackAnalysisParameters& Params = FeedbackAnalysisParameters[TaskIndex];
+				Params.System = this;
+				if (TaskIndex == 0u)
+				{
+					Params.UniquePageList = MergedUniquePageList;
+				}
+				else
+				{
+					Params.UniquePageList = Allocator.Create<FUniquePageList>();
+				}
+				Params.FeedbackBuffer = FeedbackResult.Data + CurrentOffset;
+
+				const uint32 Size = FMath::Min(FeedbackSizePerTask, FeedbackResult.Size - CurrentOffset);
+				Params.FeedbackSize = Size;
+				CurrentOffset += Size;
+			}
+
+			// Kick the tasks
+			const int32 LocalFeedbackTaskCount = Settings.bParallelFeedbackTasks ? 1 : NumFeedbackTasks;
+			const int32 WorkerFeedbackTaskCount = NumFeedbackTasks - LocalFeedbackTaskCount;
+
+			FGraphEventArray Tasks;
+			if (WorkerFeedbackTaskCount > 0)
+			{
+				SCOPE_CYCLE_COUNTER(STAT_ProcessRequests_SubmitTasks);
+				Tasks.Reserve(WorkerFeedbackTaskCount);
+				for (uint32 TaskIndex = LocalFeedbackTaskCount; TaskIndex < NumFeedbackTasks; ++TaskIndex)
+				{
+					Tasks.Add(TGraphTask<FFeedbackAnalysisTask>::CreateTask().ConstructAndDispatchWhenReady(FeedbackAnalysisParameters[TaskIndex]));
+				}
+			}
+
+			if (NumFeedbackTasks > 0u)
+			{
+				SCOPE_CYCLE_COUNTER(STAT_FeedbackAnalysis);
+
+				for (int32 TaskIndex = 0; TaskIndex < LocalFeedbackTaskCount; ++TaskIndex)
+				{
+					FFeedbackAnalysisTask::DoTask(FeedbackAnalysisParameters[TaskIndex]);
+				}
+				if (WorkerFeedbackTaskCount > 0)
+				{
+					SCOPE_CYCLE_COUNTER(STAT_ProcessRequests_WaitTasks);
+
+					FTaskGraphInterface::Get().WaitUntilTasksComplete(Tasks, ENamedThreads::GetRenderThread_Local());
+				}
+			}
+
+			if (NumFeedbackTasks > 1u)
+			{
+				SCOPE_CYCLE_COUNTER(STAT_ProcessRequests_MergePages);
+				for (uint32 TaskIndex = 1u; TaskIndex < NumFeedbackTasks; ++TaskIndex)
+				{
+					MergedUniquePageList->MergePages(FeedbackAnalysisParameters[TaskIndex].UniquePageList);
+				}
+			}
+		}
+
+#if WITH_EDITOR
+		// If we're are recording page requests, then copy off pages to the recording buffer.
+		if (PageRequestRecordHandle != ~0ull)
+		{
+			RecordPageRequests(MergedUniquePageList, PageRequestRecordBuffer);
+		}
+#endif
+
+		// Add any page requests from recording playback.
+		if (PageRequestPlaybackBuffer.Num() > 0)
+		{
+			if (Settings.bEnablePlayback)
+			{
+				// todo: We can split this into concurrent tasks. 
+				FAddRequestedTilesParameters Parameters;
+				Parameters.System = this;
+				Parameters.LevelBias = FMath::FloorToInt(CVarVTPlaybackMipBias.GetValueOnRenderThread() + GetGlobalMipBias() + 0.5f);
+				Parameters.RequestBuffer = PageRequestPlaybackBuffer.GetData();
+				Parameters.NumRequests = PageRequestPlaybackBuffer.Num();
+				Parameters.UniquePageList = Allocator.Create<FUniquePageList>();
+
+				FAddRequestedTilesTask::DoTask(Parameters);
+				MergedUniquePageList->MergePages(Parameters.UniquePageList);
+			}
+
+			PageRequestPlaybackBuffer.Reset(0);
+		}
+
+		// Collect tiles to lock
+		{
+			TArray<FVirtualTextureLocalTile> RemainingTilesToLock;
+			for (const FVirtualTextureLocalTile& Tile : TilesToLock)
+			{
+				const FVirtualTextureProducerHandle ProducerHandle = Tile.GetProducerHandle();
+				const FVirtualTextureProducer* Producer = Producers.FindProducer(ProducerHandle);
+				checkSlow(TileLocks.IsLocked(Tile));
+				if (Producer)
+				{
+					uint8 ProducerLayerMaskToLoad = 0u;
+					for (uint32 ProducerLayerIndex = 0u; ProducerLayerIndex < Producer->GetNumTextureLayers(); ++ProducerLayerIndex)
+					{
+						uint32 GroupIndex = Producer->GetPhysicalGroupIndexForTextureLayer(ProducerLayerIndex);
+						FVirtualTexturePhysicalSpace* PhysicalSpace = Producer->GetPhysicalSpaceForPhysicalGroup(GroupIndex);
+						FTexturePagePool& PagePool = PhysicalSpace->GetPagePool();
+						const uint32 pAddress = PagePool.FindPageAddress(ProducerHandle, GroupIndex, Tile.Local_vAddress, Tile.Local_vLevel);
+						if (pAddress == ~0u)
+						{
+							ProducerLayerMaskToLoad |= (1u << ProducerLayerIndex);
+						}
+						else
+						{
+							PagePool.Lock(pAddress);
+						}
+					}
+
+					if (ProducerLayerMaskToLoad != 0u)
+					{
+						const uint16 LoadRequestIndex = MergedRequestList->LockLoadRequest(FVirtualTextureLocalTile(Tile.GetProducerHandle(), Tile.Local_vAddress, Tile.Local_vLevel), ProducerLayerMaskToLoad);
+						if (LoadRequestIndex == 0xffff)
+						{
+							// Overflowed the request list...try to lock the tile again next frame
+							RemainingTilesToLock.Add(Tile);
+						}
+					}
+				}
+			}
+
+			TilesToLock = MoveTemp(RemainingTilesToLock);
+		}
+
+		TArray<uint32> PackedTiles;
+		if (RequestedPackedTiles.Num() > 0)
+		{
+			FScopeLock Lock(&RequestedTilesLock);
+			PackedTiles = MoveTemp(RequestedPackedTiles);
+			RequestedPackedTiles.Reset();
+		}
+
+		if (PackedTiles.Num() > 0)
+		{
+			// Collect explicitly requested tiles
+			// These tiles are generated on the current frame, so they are collected/processed in a separate list
+			FUniquePageList* RequestedPageList = Allocator.Create<FUniquePageList>();
+			RequestedPageList->Initialize();
+			for (uint32 Tile : PackedTiles)
+			{
+				RequestedPageList->Add(Tile, 0xffff);
+			}
+			GatherRequests(MergedRequestList, RequestedPageList, Frame, Allocator, Settings);
+		}
+
+		// Pages from feedback buffer were generated several frames ago, so they may no longer be valid for newly allocated VTs
+		static uint32 PendingFrameDelay = 3u;
+		if (Frame >= PendingFrameDelay)
+		{
+			GatherRequests(MergedRequestList, MergedUniquePageList, Frame - PendingFrameDelay, Allocator, Settings);
+		}
+
+		if (MergedRequestList->GetNumAdaptiveAllocationRequests() > 0)
+		{
+			SCOPE_CYCLE_COUNTER(STAT_ProcessRequests_QueueAdaptiveRequests);
+			FAdaptiveVirtualTexture::QueuePackedAllocationRequests(this, &MergedRequestList->GetAdaptiveAllocationRequest(0), MergedRequestList->GetNumAdaptiveAllocationRequests(), Frame);
+		}
+
+		{
+			SCOPE_CYCLE_COUNTER(STAT_ProcessRequests_Sort);
+
+			// Limit the number of uploads (account for MappedTilesToProduce this frame)
+			// Are all pages equal? Should there be different limits on different types of pages?
+			const int32 MaxNumUploads = Settings.MaxPageUploads;
+			const int32 MaxRequestUploads = FMath::Max(MaxNumUploads - MappedTilesToProduce.Num(), 1);
+
+			if (MaxRequestUploads < (int32)MergedRequestList->GetNumLoadRequests())
+			{
+				// Dropping requests is normal but track to log here if we want to tune settings.
+				if (CVarVTVerbose.GetValueOnRenderThread())
+				{
+					UE_LOG(LogConsoleResponse, Display, TEXT("VT dropped %d load requests."), MergedRequestList->GetNumLoadRequests() - MaxRequestUploads);
+				}
+			}
+
+			MergedRequestList->SortRequests(Producers, Allocator, MaxRequestUploads);
+		}
+
+		{
+			// After sorting and clamping the load requests, if we still have unused upload bandwidth then use it to add some continous updates
+			const int32 MaxNumUploads = Settings.MaxPageUploads;
+			const int32 MaxTilesToProduce = FMath::Max(MaxNumUploads - MappedTilesToProduce.Num() - (int32)MergedRequestList->GetNumLoadRequests(), 0);
+			const int32 MaxContinuousUpdates = Settings.MaxContinuousUpdates;
+
+			GetContinuousUpdatesToProduce(MergedRequestList, MaxTilesToProduce, MaxContinuousUpdates);
+		}
+
+		// Submit the requests to produce pages that are already mapped
+		SubmitPreMappedRequests(RHICmdList, FeatureLevel);
+
+		// Submit the merged requests
+		BeginSubmitRequests(RHICmdList, FeatureLevel, Allocator, Settings, MergedRequestList, true);
+
+	}, Settings.bEnableAsyncTasks && CVarVTAsyncPageRequestTask.GetValueOnRenderThread());
+
+	return Updater;
 }
 
-void FVirtualTextureSystem::CallPendingCallbacks()
+void FVirtualTextureSystem::EndUpdate(FRDGBuilder& GraphBuilder, TUniquePtr<FVirtualTextureUpdater>&& Updater, ERHIFeatureLevel::Type FeatureLevel)
 {
-	Producers.CallPendingCallbacks();
+	if (!Updater)
+	{
+		return;
+	}
+
+	check(bUpdating);
+	check(IsInRenderingThread());
+
+	RDG_CSV_STAT_EXCLUSIVE_SCOPE(GraphBuilder, VirtualTextureSystem_Update);
+	TRACE_CPUPROFILER_EVENT_SCOPE(FVirtualTextureSystem::EndUpdate);
+	SCOPE_CYCLE_COUNTER(STAT_VirtualTextureSystem_Update);
+	RDG_GPU_STAT_SCOPE(GraphBuilder, VirtualTexture);
+
+	Updater->AsyncTask.Wait();
+
+	if (Updater->FeedbackMapResult.Data)
+	{
+		GVirtualTextureFeedback.Unmap(GraphBuilder.RHICmdList, Updater->FeedbackMapResult.MapHandle);
+	}
+
+	EndSubmitRequests(GraphBuilder);
+
+	Producers.NotifyRequestsCompleted();
+
+	UpdateResidencyTracking();
+
+#if !UE_BUILD_SHIPPING
+	UpdateCsvStats();
+#endif
+
+	ReleasePendingSpaces();
+
+	bUpdating = false;
+}
+
+void FVirtualTextureSystem::Update(FRDGBuilder& GraphBuilder, ERHIFeatureLevel::Type FeatureLevel, FScene* Scene, const FVirtualTextureUpdateSettings& InSettings)
+{
+	FVirtualTextureUpdateSettings Settings = InSettings;
+	Settings.EnableAsyncTasks(false);
+
+	TUniquePtr<FVirtualTextureUpdater> Updater = BeginUpdate(GraphBuilder, FeatureLevel, Scene, Settings);
+	EndUpdate(GraphBuilder, MoveTemp(Updater), FeatureLevel);
 }
 
 void FVirtualTextureSystem::ReleasePendingResources()
 {
+	check(!bUpdating);
 	DestroyPendingVirtualTextures(true);
 	ReleasePendingSpaces();
 }
diff --git a/Engine/Source/Runtime/Renderer/Private/VT/VirtualTextureSystem.h b/Engine/Source/Runtime/Renderer/Private/VT/VirtualTextureSystem.h
--- a/Engine/Source/Runtime/Renderer/Private/VT/VirtualTextureSystem.h
+++ b/Engine/Source/Runtime/Renderer/Private/VT/VirtualTextureSystem.h
@@ -12,6 +12,8 @@
 #include "VT/VirtualTextureProducer.h"
 #include "VT/TexturePageLocks.h"
 #include "VirtualTexturing.h"
+#include "VT/VirtualTextureFeedback.h"
+#include "Tasks/Task.h"
 
 class FAdaptiveVirtualTexture;
 class FAllocatedVirtualTexture;
@@ -21,6 +23,7 @@ class FUniqueRequestList;
 class FVirtualTexturePhysicalSpace;
 class FVirtualTextureProducer;
 class FVirtualTextureSpace;
+class FVirtualTextureSystem;
 struct FVTSpaceDescription;
 struct FVTPhysicalSpaceDescription;
 union FPhysicalSpaceIDAndAddress;
@@ -36,9 +39,9 @@ struct FVirtualTextureUpdateSettings
 	FVirtualTextureUpdateSettings();
 
 	/** Force settings so that throttling page uploads is effectively disabled. */
-	FVirtualTextureUpdateSettings& DisableThrottling(bool bDisable)
+	FVirtualTextureUpdateSettings& EnableThrottling(bool bEnable)
 	{
-		if (bDisable)
+		if (!bEnable)
 		{
 			MaxPageUploads = 99999;
 			MaxPagesProduced = 99999;
@@ -46,6 +49,22 @@ struct FVirtualTextureUpdateSettings
 		return *this;
 	}
 
+	/** Force virtual texture updates to be done synchronously. */
+	FVirtualTextureUpdateSettings& EnableAsyncTasks(bool bEnable = true)
+	{
+		bEnableAsyncTasks = bEnable;
+		return *this;
+	}
+
+	/** Do not perform any updates related to read backs. */
+	FVirtualTextureUpdateSettings& EnablePageRequests(bool bEnable = true)
+	{
+		bEnablePageRequests = bEnable;
+		return *this;
+	}
+
+	bool bEnableAsyncTasks = true;
+	bool bEnablePageRequests = true;
 	bool bEnableFeedback;
 	bool bEnablePlayback;
 	bool bForceContinuousUpdate;
@@ -58,6 +77,17 @@ struct FVirtualTextureUpdateSettings
 	int32 MaxContinuousUpdates;
 };
 
+class FVirtualTextureUpdater
+{
+private:
+	FVirtualTextureUpdater() = default;
+
+	FVirtualTextureFeedback::FMapResult FeedbackMapResult;
+	UE::Tasks::FTask AsyncTask;
+
+	friend FVirtualTextureSystem;
+};
+
 class FVirtualTextureSystem
 {
 public:
@@ -67,10 +97,11 @@ public:
 
 	uint32 GetFrame() const { return Frame; }
 
-	void AllocateResources(FRDGBuilder& GraphBuilder, ERHIFeatureLevel::Type FeatureLevel);
-	void FinalizeResources(FRDGBuilder& GraphBuilder, ERHIFeatureLevel::Type FeatureLevel);
-	void CallPendingCallbacks();
-	void Update(FRDGBuilder& GraphBuilder, ERHIFeatureLevel::Type FeatureLevel, FScene* Scene, FVirtualTextureUpdateSettings const& Settings);
+	void Update(FRDGBuilder& GraphBuilder, ERHIFeatureLevel::Type FeatureLevel, FScene* Scene, const FVirtualTextureUpdateSettings& Settings);
+
+	TUniquePtr<FVirtualTextureUpdater> BeginUpdate(FRDGBuilder& GraphBuilder, ERHIFeatureLevel::Type FeatureLevel, FScene* Scene, const FVirtualTextureUpdateSettings& Settings);
+	void EndUpdate(FRDGBuilder& GraphBuilder, TUniquePtr<FVirtualTextureUpdater>&& Updater, ERHIFeatureLevel::Type FeatureLevel);
+
 	void ReleasePendingResources();
 
 	IAllocatedVirtualTexture* AllocateVirtualTexture(const FAllocatedVTDescription& Desc);
@@ -123,6 +154,7 @@ private:
 	FVirtualTextureSystem();
 	~FVirtualTextureSystem();
 
+	void AllocateResources(FRDGBuilder& GraphBuilder);
 	void DestroyPendingVirtualTextures(bool bForceDestroyAll);
 	void ReleasePendingSpaces();
 
@@ -130,11 +162,12 @@ private:
 	void RequestTilesInternal(const IAllocatedVirtualTexture* AllocatedVT, int32 InMipLevel);
 	void RequestTilesInternal(const IAllocatedVirtualTexture* AllocatedVT, const FVector2D& InScreenSpaceSize, int32 InMipLevel);
 	
-	void SubmitRequestsFromLocalTileList(TArray<FVirtualTextureLocalTile>& OutDeferredTiles, const TSet<FVirtualTextureLocalTile>& LocalTileList, EVTProducePageFlags Flags, FRDGBuilder& GraphBuilder, ERHIFeatureLevel::Type FeatureLevel);
+	void SubmitRequestsFromLocalTileList(FRHICommandList& RHICmdList, TArray<FVirtualTextureLocalTile>& OutDeferredTiles, const TSet<FVirtualTextureLocalTile>& LocalTileList, EVTProducePageFlags Flags, ERHIFeatureLevel::Type FeatureLevel);
 
-	void SubmitPreMappedRequests(FRDGBuilder& GraphBuilder, ERHIFeatureLevel::Type FeatureLevel);
+	void SubmitPreMappedRequests(FRHICommandList& RHICmdList, ERHIFeatureLevel::Type FeatureLevel);
 
-	void SubmitRequests(FRDGBuilder& GraphBuilder, ERHIFeatureLevel::Type FeatureLevel, FConcurrentLinearBulkObjectAllocator& Allocator, FVirtualTextureUpdateSettings const& Settings, FUniqueRequestList* RequestList, bool bAsync);
+	void BeginSubmitRequests(FRHICommandList& RHICmdList, ERHIFeatureLevel::Type FeatureLevel, FConcurrentLinearBulkObjectAllocator& Allocator, FVirtualTextureUpdateSettings const& Settings, FUniqueRequestList* RequestList, bool bAsync);
+	void EndSubmitRequests(FRDGBuilder& GraphBuilder);
 
 	void GatherRequests(FUniqueRequestList* MergedRequestList, const FUniquePageList* UniquePageList, uint32 FrameRequested, FConcurrentLinearBulkObjectAllocator& Allocator, FVirtualTextureUpdateSettings const& Settings);
 
@@ -170,6 +203,7 @@ private:
 
 	FAdaptiveVirtualTexture* AdaptiveVTs[MaxSpaces] = { nullptr };
 
+	bool bUpdating = false;
 	bool bFlushCaches;
 	void FlushCachesFromConsole();
 	FAutoConsoleCommand FlushCachesCommand;
