Strata rough refraction account for depth

diff --git a/Engine/Shaders/Private/DistortAccumulatePS.usf b/Engine/Shaders/Private/DistortAccumulatePS.usf
--- a/Engine/Shaders/Private/DistortAccumulatePS.usf
+++ b/Engine/Shaders/Private/DistortAccumulatePS.usf
@@ -49,7 +49,9 @@ void Main(
 	out float4 OutColor		: SV_Target0
 #if STRATA_ENABLED
 	,
-	out float4 OutVariance	: SV_Target1
+	out float2 OutVarianceCoverage : SV_Target1
+	,
+	out float  OutClosestDepthMeter : SV_Target2
 #endif
 	)
 {
@@ -129,11 +131,17 @@ void Main(
 	}
 	TopF0 = SanitizeF0(TopF0);	// Sanitize, in case a Add node has been used
 
-	// Write out variance (as this is more linear to accumulate than roughness)
-	OutVariance = float4(Sigma, 0.0f, 0.0f, 0.0f);
+	const FStrataOperator RootOperator	= StrataTree.Operators[StrataData.OperatorIndex];
+	const float RefractionCoverage		= RootOperator.Coverage;														// Coverage of the material
+	const float RefractionDepthMeter	= length(MaterialParameters.WorldPosition_CamRelative) * CENTIMETER_TO_METER;	// Distance for camera in meters
+	const float RefractionLobeVariance	= Sigma;																		// The lobe variance that can be added to in the en recover a single lobe.
+
+	// Write out variance (as this is more linear to accumulate than roughness).
+	OutVarianceCoverage = float2(RefractionLobeVariance, RefractionCoverage);
+	OutClosestDepthMeter = RefractionDepthMeter;
 
 	// Only allow clipping if roughness is close to 0.
-	const bool TryToClip = OutVariance.x < 0.01f;
+	const bool TryToClip = OutVarianceCoverage.x < 0.01f; // STRATA_TODO || RefractionCoverage <= 0.0f;	// Have an option to account for coverage or not?, ignored by default
 
 	// Compute UV distortion
 	float2 BufferUVDistortion = ComputeBufferUVDistortion(
diff --git a/Engine/Shaders/Private/DistortApplyScreenPS.usf b/Engine/Shaders/Private/DistortApplyScreenPS.usf
--- a/Engine/Shaders/Private/DistortApplyScreenPS.usf
+++ b/Engine/Shaders/Private/DistortApplyScreenPS.usf
@@ -6,8 +6,9 @@
 #include "Common.ush"
 #include "PostProcessCommon.ush"
 
+#include "/Engine/Private/Strata/Strata.ush"
 
-#if USE_ROUGH_REFRACTION
+#if USE_ROUGH_REFRACTION && STRATA_ENABLED
 #include "/Engine/Private/Strata/StrataStatisticalOperators.ush"
 #endif
 
@@ -24,29 +25,36 @@ static const half InvDistortionScaleBias = 1 / 4.0f;
 
 Texture2D		SceneColorTexture;
 SamplerState	SceneColorTextureSampler;
+Texture2D		SceneDepthMeterTexture;
+SamplerState	SceneDepthMeterTextureSampler;
   
  #if USE_MSAA
 Texture2DMS<float4> DistortionMSAATexture;
 Texture2DMS<float4> SceneColorMSAATexture;
-Texture2DMS<float4> RoughnessScatterMSAATexture;
+Texture2DMS<float4> SceneDepthMeterMSAATexture;
+Texture2DMS<float4> RefractionVarianceCoverageMSAATexture;
+Texture2DMS<float4> RefractionClosestDepthMeterMSAATexture;
 #if COMPOSITE_STANDARD_TRANSLUCENT
 Texture2DMS<float4> StandardTranslucentColorAlphaMSAATexture;
 Texture2DMS<float4> StandardTranslucentTransmittanceMSAATexture;
 #endif
 #else
-Texture2D		DistortionTexture;
-SamplerState	DistortionTextureSampler;
-Texture2D<float4> RoughnessScatterTexture;
-SamplerState      RoughnessScatterSampler;
+Texture2D			DistortionTexture;
+SamplerState		DistortionTextureSampler;
+Texture2D<float4>	RefractionVarianceCoverageTexture;
+SamplerState		RefractionVarianceCoverageSampler;
+Texture2D<float4>	RefractionClosestDepthMeterTexture;
+SamplerState		RefractionClosestDepthMeterSampler;
 #if COMPOSITE_STANDARD_TRANSLUCENT
-Texture2D		StandardTranslucentColorAlphaTexture;
-SamplerState	StandardTranslucentColorAlphaTextureSampler;
-Texture2D		StandardTranslucentTransmittanceTexture;
-SamplerState	StandardTranslucentTransmittanceTextureSampler;
+Texture2D			StandardTranslucentColorAlphaTexture;
+SamplerState		StandardTranslucentColorAlphaTextureSampler;
+Texture2D			StandardTranslucentTransmittanceTexture;
+SamplerState		StandardTranslucentTransmittanceTextureSampler;
 #endif
 #endif
 
-float RefractionRoughnessToMipLevelFactor;
+float RefractionMaxMipLevel;
+float StandardDevInPixelForMip0;
 
 /** distorts screen texture using accumulated distortion offsets */
 void Main(
@@ -71,25 +79,77 @@ void Main(
 	//Scale by the screen size and a fudge factor to come close to the offset values we would have had under normal circumstances before my changes. Also flip Y and invert the precision bias scale.
 	DistBufferUVOffset *= InvDistortionScaleBias;
 	
-#if USE_ROUGH_REFRACTION
+#if USE_ROUGH_REFRACTION && STRATA_ENABLED
 
 #if USE_MSAA
-	half RefractionLobeVariance = RoughnessScatterMSAATexture.Load(int2(TexCoord.xy * View.BufferSizeAndInvSize.xy), SampleIndex).x;
+	half2 RefractionVarianceCoverage= RefractionVarianceCoverageMSAATexture.Load(int2(TexCoord.xy * View.BufferSizeAndInvSize.xy), SampleIndex).xy;
+	half  RefractionDepthMeter		= RefractionClosestDepthMeterMSAATexture.Load(int2(TexCoord.xy * View.BufferSizeAndInvSize.xy), SampleIndex).x;
 #else
-	half RefractionLobeVariance = Texture2DSample(RoughnessScatterTexture, RoughnessScatterSampler, TexCoord.xy).x;
+	half2 RefractionVarianceCoverage= Texture2DSample(RefractionVarianceCoverageTexture, RefractionVarianceCoverageSampler, TexCoord.xy).xy;
+	half  RefractionDepthMeter		= Texture2DSample(RefractionClosestDepthMeterTexture, RefractionClosestDepthMeterSampler, TexCoord.xy).x;
 #endif
-	half RefractionRoughness = saturate(StrataLobeVarianceToRoughness(RefractionLobeVariance));
-	half RefractionMipLevel = RefractionRoughness * RefractionRoughnessToMipLevelFactor;
 
-	const float2 NoDistortionCoord = (TexCoord.xy * View.BufferSizeAndInvSize.xy - View.ViewRectMin.xy);
-	const float2 NoDistortionUVs = NoDistortionCoord * View.ViewSizeAndInvSize.zw;
-	float2 NewBufferUV = NoDistortionUVs + DistBufferUVOffset;
+	const float2 SVPos				=  TexCoord.xy * View.BufferSizeAndInvSize.xy + View.ViewRectMin.xy;
+	const float2 NoDistortionCoord	= (TexCoord.xy * View.BufferSizeAndInvSize.xy - View.ViewRectMin.xy);
+	const float2 NoDistortionUVs	= NoDistortionCoord * View.ViewSizeAndInvSize.zw;
 
+	float2 NewBufferUV = NoDistortionUVs + DistBufferUVOffset;
 	FLATTEN if (any(NewBufferUV < 0.0f) || any(NewBufferUV > 1.0f))
 	{
 		NewBufferUV = NoDistortionUVs;
 	}
 
+	half SceneDepthMeter			= 0.0f;
+#if USE_MSAA
+	SceneDepthMeter = SceneDepthMeterMSAATexture.Load(int2(NewBufferUV * View.ViewSizeAndInvSize.xy), SampleIndex).x;
+#else
+	const half SceneDepthMeterMip0 = Texture2DSampleLevel(SceneDepthMeterTexture, SceneDepthMeterTextureSampler, NewBufferUV, 0).x;
+	const half SceneDepthMeterMipN = Texture2DSampleLevel(SceneDepthMeterTexture, SceneDepthMeterTextureSampler, NewBufferUV, 4).x;
+	SceneDepthMeter = min(SceneDepthMeterMip0, SceneDepthMeterMipN);
+#endif
+
+	// We evaluate the gaussian blur needed to match the refraction of a surface at a distance behind a surface exhibiting rough refraction.
+	// For that we used a process and blurred depth buffer storing linear depth in meter. This allows to have soft transitions around edges.
+	// It is not perfect however with large depth difference.
+	// The gaussian blur standard deviation is estiamted from the pixel footprint of the projected radiance at a certain depth from a rough surfaces using GetVarianceFromRoughness and the Thales theorem.
+	// The correct blured scene mip level is then evaluated to make sure the standard devianttion in pixel matches the one of the projected rough refraction gaussian kernel.
+
+	half FrontSurfaceRefractionDepthCm			= max(0.0f, SceneDepthMeter - RefractionDepthMeter) * METER_TO_CENTIMETER;
+	half RefractionRoughness					= saturate(StrataLobeVarianceToRoughness(RefractionVarianceCoverage.x));
+	const float StandardDeviationCm				= sqrt(GetVarianceFromRoughness(RefractionRoughness));
+	const float StandardDeviationCmForThickness = StandardDeviationCm * FrontSurfaceRefractionDepthCm;
+	const float VarianceAtDepth					= StandardDeviationCmForThickness * StandardDeviationCmForThickness;
+
+	half RefractionMipLevel = 0.0;
+	{
+		const float2 PixelRadiusCm2 = METER_TO_CENTIMETER * SceneDepthMeter * GetTanHalfFieldOfView() * View.ViewSizeAndInvSize.zw;
+		const float PixelRadiusCm = min(PixelRadiusCm2.x, PixelRadiusCm2.y);
+
+		const float PixelStandardDevInPixel = StandardDeviationCmForThickness / PixelRadiusCm;
+
+		// Jitter UV according to MIP level to hide texel.
+		if(true)
+		{
+			const float BlurRandomness = InterleavedGradientNoise(SVPos.xy, float(View.StateFrameIndexMod8)) - 0.5; // In [-0.5, 0.5]
+			const float RandAngle = BlurRandomness * 2.0f * PI;
+			float CosA, SinA;
+			sincos(RandAngle, SinA, CosA);
+			NewBufferUV += 0.5 * float2(CosA, SinA) * PixelStandardDevInPixel * View.ViewSizeAndInvSize.zw;
+		}
+
+		// The standard deviation in pixel for mip0 = StandardDevInPixelForMip0.
+		// It means that for each mips N, with the same gaussian with same variance is applied, the standard deviation in pixel of image at source resolution is = StandardDevInPixelForMip0 * 2^M.
+		// We now want to compute the mip to select for the current pixel PixelStandardDevInPixel.
+		// So PixelStandardDevInPixel = StandardDevInPixelForMip0 * 2^M
+		//	  PixelStandardDevInPixel / StandardDevInPixelForMip0 = 2^M
+		// and because y = z ^ x => x = log(y) / log(z)
+		// we obtain  M = log(PixelStandardDevInPixel / StandardDevInPixelForMip0) / log(2)
+
+		// Compute mip level according to VarianceAtDepth at a depth of FrontSurfaceRefractionDepthCm
+		RefractionMipLevel = log(PixelStandardDevInPixel / StandardDevInPixelForMip0) / log(2);
+	}
+	RefractionMipLevel = min(RefractionMipLevel, RefractionMaxMipLevel);
+
 #if USE_MSAA
 	if (RefractionMipLevel < 0.001f)
 	{
diff --git a/Engine/Shaders/Private/DistortFiltering.usf b/Engine/Shaders/Private/DistortFiltering.usf
--- a/Engine/Shaders/Private/DistortFiltering.usf
+++ b/Engine/Shaders/Private/DistortFiltering.usf
@@ -8,16 +8,30 @@
 
 
 
-#ifdef COPYSCENECOLORTEXTUREPS
+#ifdef COPYSCENECOLORDEPTHTEXTUREPS
 
 Texture2D SceneColorTexture;
 SamplerState SceneColorSampler;
 
-void CopySceneColorTexturePS(FScreenVertexOutput Input, out float4 OutColor : SV_Target0)
+Texture2D SceneDepthTexture;
+SamplerState SceneDepthSampler;
+
+void CopySceneColorDepthTexturePS(FScreenVertexOutput Input, out float4 OutColor : SV_Target0, out float OutDepthMeter : SV_Target1)
 {
-	float2 TextureUV = (Input.UV * View.ViewSizeAndInvSize.xy + View.ViewRectMin.xy) * View.BufferSizeAndInvSize.zw;
+	float2 SVPos = Input.UV * View.ViewSizeAndInvSize.xy + View.ViewRectMin.xy;
+
+	float2 TextureUV = SVPos * View.BufferSizeAndInvSize.zw;
 	float4 LinearColor = Texture2DSample(SceneColorTexture, SceneColorSampler, TextureUV);
-	OutColor = float4(LinearColor.rgb, 0);
+
+	float DeviceZ = Texture2DSample(SceneDepthTexture, SceneDepthSampler, TextureUV).r;
+#if HAS_INVERTED_Z_BUFFER
+	DeviceZ = max(0.000000000001, DeviceZ);	// TODO: avoid bad values when DeviceZ is far=0 when using inverted z
+#endif
+	float4 SVPosition = float4(SVPos, DeviceZ, 1);
+	float3 TranslatedWorldPosition = SvPositionToTranslatedWorld(SVPosition);
+
+	OutColor = float4(LinearColor.rgb, 1.0f);
+	OutDepthMeter = length(TranslatedWorldPosition) * CENTIMETER_TO_METER;
 }
 
 #endif // COPYSCENECOLORPS
@@ -60,7 +74,10 @@ float4 BlurDirection;
 SamplerState SourceSampler;
 Texture2D<float4> SourceTexture;
 RWTexture2D<float4> OutTextureMipColor;
-float4 SampleOffsetsWeights[128];
+
+#define MAX_FILTER_SAMPLE_COUNT_RADIUS	32
+#define MAX_FILTER_SAMPLE_COUNT	(MAX_FILTER_SAMPLE_COUNT_RADIUS*2+1)
+float4 SampleOffsetsWeights[MAX_FILTER_SAMPLE_COUNT];
 
 [numthreads(THREADGROUP_SIZE, THREADGROUP_SIZE, 1)]
 void FilterColorCS(uint3 ThreadId : SV_DispatchThreadID)
@@ -77,7 +94,9 @@ void FilterColorCS(uint3 ThreadId : SV_DispatchThreadID)
 	{
 		float2 SampleOffsetWeight = SampleOffsetsWeights[i].xy;
 
-		OutColor += SourceTexture.SampleLevel(SourceSampler, ScaledUVs + SampleOffsetWeight.xx * BlurDirection.xy, SrcMipIndex) * SampleOffsetWeight.y;
+		float4 Sample = SourceTexture.SampleLevel(SourceSampler, ScaledUVs + SampleOffsetWeight.xx * BlurDirection.xy, SrcMipIndex);
+
+		OutColor += Sample * SampleOffsetWeight.y;
 	}
 
 	OutTextureMipColor[uint2(ThreadId.xy)] = OutColor;
diff --git a/Engine/Source/Runtime/Renderer/Private/DistortionRendering.cpp b/Engine/Source/Runtime/Renderer/Private/DistortionRendering.cpp
--- a/Engine/Source/Runtime/Renderer/Private/DistortionRendering.cpp
+++ b/Engine/Source/Runtime/Renderer/Private/DistortionRendering.cpp
@@ -27,33 +27,19 @@ static TAutoConsoleVariable<int32> CVarDisableDistortion(
 static TAutoConsoleVariable<int32> CVarRefractionBlur(
 	TEXT("r.Refraction.Blur"),
 	1,
-	TEXT("Prevent rough refraction from happening, i.e. blurring of the background."),
+	TEXT("Enables rough refractions, i.e. blurring of the background."),
 	ECVF_Scalability | ECVF_RenderThreadSafe);
 
-static TAutoConsoleVariable<float> CVarRefractionBlurScale(
-	TEXT("r.Refraction.BlurScale"),
-	0.2f,
-	TEXT("Global scale the background blur amount after it is mapped form the surface back roughness/scattering amount."),
-	ECVF_RenderThreadSafe);
-
-static TAutoConsoleVariable<float> CVarRefractionBlurCenterWeight(
-	TEXT("r.Refraction.BlurCenterWeight"),
-	0.0f,
-	TEXT("The weight of the center sample. Value greater than 0 means the sharp image is more and more visible."),
-	ECVF_RenderThreadSafe);
-
 static TAutoConsoleVariable<int32> CVarRefractionOffsetQuality(
 	TEXT("r.Refraction.OffsetQuality"),
 	0,
 	TEXT("When enabled, the offset buffer is made float for higher quality. This is important to maintain the softness of the blurred scene buffer."),
 	ECVF_Scalability | ECVF_RenderThreadSafe);
 
-// Using a simple scale to map roughness to mip level.
-// Later we could remove that and use depth+variance to have a better lob roughness to mip match.
-static TAutoConsoleVariable<float> CVarRefractionRoughnessToMipLevelFactor(
-	TEXT("r.Refraction.RoughnessToMipLevelFactor"),
+static TAutoConsoleVariable<float> CVarRefractionBlurMaxStandardDeviationInScreenPercent(
+	TEXT("r.Refraction.Blur.MaxStandardDeviationInScreenRatio"),
 	5.0f,
-	TEXT("Factor to translate the roughness factor into a mip level."),
+	TEXT("This will clamp the maximum refraction blur radius on screen."),
 	ECVF_Scalability | ECVF_RenderThreadSafe);
 
 BEGIN_GLOBAL_SHADER_PARAMETER_STRUCT(FDistortionPassUniformParameters, RENDERER_API)
@@ -119,22 +105,29 @@ public:
 
 	BEGIN_SHADER_PARAMETER_STRUCT(FParameters, )
 		SHADER_PARAMETER_STRUCT_REF(FViewUniformShaderParameters, View)
-		SHADER_PARAMETER_RDG_TEXTURE(Texture2DMS<float4>, RoughnessScatterMSAATexture)
+		SHADER_PARAMETER_RDG_TEXTURE(Texture2DMS<float4>, RefractionVarianceCoverageMSAATexture)
+		SHADER_PARAMETER_RDG_TEXTURE(Texture2DMS<float4>, RefractionClosestDepthMeterMSAATexture)
 		SHADER_PARAMETER_RDG_TEXTURE(Texture2DMS<float4>, DistortionMSAATexture)
 		SHADER_PARAMETER_RDG_TEXTURE(Texture2DMS<float4>, SceneColorMSAATexture)
+		SHADER_PARAMETER_RDG_TEXTURE(Texture2DMS<float4>, SceneDepthMeterMSAATexture)
 		SHADER_PARAMETER_RDG_TEXTURE(Texture2DMS<float4>, StandardTranslucentColorAlphaMSAATexture)
 		SHADER_PARAMETER_RDG_TEXTURE(Texture2DMS<float4>, StandardTranslucentTransmittanceMSAATexture)
-		SHADER_PARAMETER_RDG_TEXTURE(Texture2D, RoughnessScatterTexture)
+		SHADER_PARAMETER_RDG_TEXTURE(Texture2D, RefractionVarianceCoverageTexture)
+		SHADER_PARAMETER_RDG_TEXTURE(Texture2D, RefractionClosestDepthMeterTexture)
 		SHADER_PARAMETER_RDG_TEXTURE(Texture2D, DistortionTexture)
 		SHADER_PARAMETER_RDG_TEXTURE(Texture2D, SceneColorTexture)
+		SHADER_PARAMETER_RDG_TEXTURE(Texture2D, SceneDepthMeterTexture)
 		SHADER_PARAMETER_RDG_TEXTURE(Texture2D, StandardTranslucentColorAlphaTexture)
 		SHADER_PARAMETER_RDG_TEXTURE(Texture2D, StandardTranslucentTransmittanceTexture)
-		SHADER_PARAMETER_SAMPLER(SamplerState, RoughnessScatterSampler)
+		SHADER_PARAMETER_SAMPLER(SamplerState, RefractionVarianceCoverageSampler)
+		SHADER_PARAMETER_SAMPLER(SamplerState, RefractionClosestDepthMeterSampler)
 		SHADER_PARAMETER_SAMPLER(SamplerState, DistortionTextureSampler)
 		SHADER_PARAMETER_SAMPLER(SamplerState, SceneColorTextureSampler)
+		SHADER_PARAMETER_SAMPLER(SamplerState, SceneDepthMeterTextureSampler)
 		SHADER_PARAMETER_SAMPLER(SamplerState, StandardTranslucentColorAlphaTextureSampler)
 		SHADER_PARAMETER_SAMPLER(SamplerState, StandardTranslucentTransmittanceTextureSampler)
-		SHADER_PARAMETER(float, RefractionRoughnessToMipLevelFactor)
+		SHADER_PARAMETER(float, RefractionMaxMipLevel)
+		SHADER_PARAMETER(float, StandardDevInPixelForMip0)
 		RENDER_TARGET_BINDING_SLOTS()
 	END_SHADER_PARAMETER_STRUCT()
 
@@ -257,7 +250,9 @@ public:
 	BEGIN_SHADER_PARAMETER_STRUCT(FParameters, )
 		SHADER_PARAMETER_STRUCT_REF(FViewUniformShaderParameters, View)
 		SHADER_PARAMETER_RDG_TEXTURE(Texture2D, SceneColorTexture)
+		SHADER_PARAMETER_RDG_TEXTURE(Texture2D, SceneDepthTexture)
 		SHADER_PARAMETER_SAMPLER(SamplerState, SceneColorSampler)
+		SHADER_PARAMETER_SAMPLER(SamplerState, SceneDepthSampler)
 		RENDER_TARGET_BINDING_SLOTS()
 	END_SHADER_PARAMETER_STRUCT()
 
@@ -268,13 +263,13 @@ public:
 	static void ModifyCompilationEnvironment(const FGlobalShaderPermutationParameters& Parameters, FShaderCompilerEnvironment& OutEnvironment)
 	{
 		FGlobalShader::ModifyCompilationEnvironment(Parameters, OutEnvironment);
-		OutEnvironment.SetDefine(TEXT("COPYSCENECOLORTEXTUREPS"), 1);
+		OutEnvironment.SetDefine(TEXT("COPYSCENECOLORDEPTHTEXTUREPS"), 1);
 	}
 };
 
-IMPLEMENT_GLOBAL_SHADER(FCopySceneColorTexturePS, "/Engine/Private/DistortFiltering.usf", "CopySceneColorTexturePS", SF_Pixel);
+IMPLEMENT_GLOBAL_SHADER(FCopySceneColorTexturePS, "/Engine/Private/DistortFiltering.usf", "CopySceneColorDepthTexturePS", SF_Pixel);
 
-static void AddCopySceneColorPass(FRDGBuilder& GraphBuilder, const FViewInfo& View, FRDGTextureRef SceneColorTexture, FRDGTextureRef SceneColorCopyTexture)
+static void AddCopySceneColorDepthPass(FRDGBuilder& GraphBuilder, const FViewInfo& View, FRDGTextureRef SceneColorTexture, FRDGTextureRef SceneDepthTexture, FRDGTextureRef SceneColorCopyTexture, FRDGTextureRef SceneDepthCopyTexture)
 {
 	const FScreenPassTextureViewport Viewport(SceneColorCopyTexture, View.ViewRect);
 
@@ -284,8 +279,11 @@ static void AddCopySceneColorPass(FRDGBuilder& GraphBuilder, const FViewInfo& Vi
 	auto* PassParameters = GraphBuilder.AllocParameters<FCopySceneColorTexturePS::FParameters>();
 	PassParameters->View = View.ViewUniformBuffer;
 	PassParameters->SceneColorTexture = SceneColorTexture;
+	PassParameters->SceneDepthTexture = SceneDepthTexture;
 	PassParameters->SceneColorSampler = TStaticSamplerState<SF_Point>::GetRHI();
+	PassParameters->SceneDepthSampler = TStaticSamplerState<SF_Point>::GetRHI();
 	PassParameters->RenderTargets[0] = FRenderTargetBinding(SceneColorCopyTexture, ERenderTargetLoadAction::ENoAction);
+	PassParameters->RenderTargets[1] = FRenderTargetBinding(SceneDepthCopyTexture, ERenderTargetLoadAction::ENoAction);
 
 	// The scene color is copied into from multi-view-rect layout into a single-rect layout.
 	FIntRect ViewRect;
@@ -345,65 +343,38 @@ static void AddDownsampleSceneColorPass(FRDGBuilder& GraphBuilder, FDownsampleSc
 
 //////////////////////////////////////////////////////////////////////////
 
-// STRATA_TODO make common with the function in PostProcessWeightedSampleSum.cpp
 class FGaussianFiltering
 {
 public:
 	// Evaluates an unnormalized normal distribution PDF around 0 at given X with Variance.
-	static float NormalDistributionUnscaled(float X, float Sigma, float CrossCenterWeight)
+	static float NormalDistributionUnscaled(float X, float Sigma)
 	{
 		const float DX = FMath::Abs(X);
 
-		const float ClampedOneMinusDX = FMath::Max(0.0f, 1.0f - DX);
-
-		// Tweak the gaussian shape e.g. "r.Bloom.Cross 3.5"
-		if (CrossCenterWeight > 1.0f)
-		{
-			return FMath::Pow(ClampedOneMinusDX, CrossCenterWeight);
-		}
-		else
-		{
-			// Constant is tweaked give a similar look to UE before we fixed the scale bug (Some content tweaking might be needed).
-			// The value defines how much of the Gaussian clipped by the sample window.
-			// r.Filter.SizeScale allows to tweak that for performance/quality.
-			const float LegacyCompatibilityConstant = -16.7f;
-
-			const float Gaussian = FMath::Exp(LegacyCompatibilityConstant * FMath::Square(DX / Sigma));
-
-			return FMath::Lerp(Gaussian, ClampedOneMinusDX, CrossCenterWeight);
-		}
-	}
-
-	static float GetClampedKernelRadius(uint32 SampleCountMax, float KernelRadius)
-	{
-		return FMath::Clamp<float>(KernelRadius, DELTA, SampleCountMax - 1);
-	}
-
-	static int GetIntegerKernelRadius(uint32 SampleCountMax, float KernelRadius)
-	{
-		// Smallest radius will be 1.
-		return FMath::Min<int32>(FMath::CeilToInt(GetClampedKernelRadius(SampleCountMax, KernelRadius)), SampleCountMax - 1);
+		// Unscaled so no front factor, because Compute1DGaussianFilterKernel normalizes weights after
+		const float Gaussian = FMath::Exp(-0.5f * FMath::Square(DX / Sigma)); 
+		return Gaussian;
 	}
 
-	static uint32 Compute1DGaussianFilterKernel(FVector2D* OutOffsetAndWeight, uint32 SampleCountMax, float KernelRadius, float CrossCenterWeight, float FilterSizeScale)
+	// Sigma is Standard Deviation in pixel
+	// Variance = sigma x sigma
+	static uint32 Compute1DGaussianFilterKernel(FVector2D* OutOffsetAndWeight, uint32 SampleCountRadius, uint32 MaxSampleCountRadius, float Sigma)
 	{
-		const float ClampedKernelRadius = GetClampedKernelRadius(SampleCountMax, KernelRadius);
-
-		const int32 IntegerKernelRadius = GetIntegerKernelRadius(SampleCountMax, KernelRadius * FilterSizeScale);
+		const float IntegerSampleCountRadius = FMath::Clamp<float>(SampleCountRadius, 1, MaxSampleCountRadius);
 
 		uint32 SampleCount = 0;
 		float WeightSum = 0.0f;
 
-		for (int32 SampleIndex = -IntegerKernelRadius; SampleIndex <= IntegerKernelRadius; SampleIndex += 2)
+		for (int32 SampleIndex = -IntegerSampleCountRadius; SampleIndex <= IntegerSampleCountRadius; SampleIndex += 2)
 		{
-			float Weight0 = NormalDistributionUnscaled(SampleIndex, ClampedKernelRadius, CrossCenterWeight);
+			float Weight0 = NormalDistributionUnscaled(SampleIndex, Sigma);
 			float Weight1 = 0.0f;
 
 			// We use the bilinear filter optimization for gaussian blur. However, we don't want to bias the
 			// last sample off the edge of the filter kernel, so the very last tap just is on the pixel center.
-			if (SampleIndex != IntegerKernelRadius)
+			if (SampleIndex != IntegerSampleCountRadius)
 			{
-				Weight1 = NormalDistributionUnscaled(SampleIndex + 1, ClampedKernelRadius, CrossCenterWeight);
+				Weight1 = NormalDistributionUnscaled(SampleIndex + 1, Sigma);
 			}
 
 			const float TotalWeight = Weight0 + Weight1;
@@ -425,7 +396,8 @@ public:
 };
 
 // If this is update, please update DistortFiltering.usf
-#define MAX_FILTER_SAMPLE_COUNT	128
+#define MAX_FILTER_SAMPLE_COUNT_RADIUS	32
+#define MAX_FILTER_SAMPLE_COUNT	(MAX_FILTER_SAMPLE_COUNT_RADIUS*2+1)
 
 class FFilterSceneColorCS : public FGlobalShader
 {
@@ -497,17 +469,27 @@ void FDeferredShadingSceneRenderer::RenderDistortion(
 	FDepthStencilBinding StencilWriteBinding(SceneDepthTexture, ERenderTargetLoadAction::ELoad, ERenderTargetLoadAction::ELoad, FExclusiveDepthStencil::DepthRead_StencilWrite);
 
 	FRDGTextureRef DistortionTexture = nullptr;
-	FRDGTextureRef RoughnessScatterTexture = nullptr;
+	FRDGTextureRef RefractionVarianceCoverage = nullptr;
+	FRDGTextureRef RefractionClosestDepthMeter = nullptr;
 	FRDGTextureRef TempSceneColorMipchainTexture = nullptr;
+	FRDGTextureRef TempSceneDepthMeterMipchainTexture = nullptr;
 	FGlobalShaderMap* ShaderMap = GetGlobalShaderMap(FeatureLevel);
 
+	const float FilterSampleCountRadius = 5.0f;	// The maximum sample count we will run. This is scaled down by FilterSizeScale.
+	const float FilterSigmaPixel = 2.0f;		// Variance in pixel, even though it is not really exact since weights are renormalizes, but it does control the spread
+
 	// Create a texture of the scene color with MIP levels, unfiltered
 	// NOTE: we cannot use mips on the scene color render target it self because multiple views are continuous in there and they would leak color one onto another.
 	TArray<FRDGTextureRef> ViewsSceneColorMipchain;
+	TArray<FRDGTextureRef> ViewsSceneDepthMeterMipchain;
 	ViewsSceneColorMipchain.SetNum(Views.Num());
+	ViewsSceneDepthMeterMipchain.SetNum(Views.Num());
 	const bool bUseRoughRefraction = GetUseRoughRefraction();
 	if(bUseRoughRefraction)
 	{
+		FVector2D OffsetAndWeight[MAX_FILTER_SAMPLE_COUNT];
+		uint32 SampleCount = FGaussianFiltering::Compute1DGaussianFilterKernel(OffsetAndWeight, FilterSampleCountRadius, MAX_FILTER_SAMPLE_COUNT_RADIUS, FilterSigmaPixel);
+
 		for (int32 ViewIndex = 0, Num = Views.Num(); ViewIndex < Num; ++ViewIndex)
 		{
 			RDG_EVENT_SCOPE(GraphBuilder, "Rough Refraction View%d", ViewIndex);
@@ -520,21 +502,26 @@ void FDeferredShadingSceneRenderer::RenderDistortion(
 
 			FRDGTextureDesc SceneColorMipchainDesc = FRDGTextureDesc::Create2D(SceneColorMip0Resolution, PF_FloatR11G11B10, FClearValueBinding::None,
 				TexCreate_TargetArraySlicesIndependently | TexCreate_ShaderResource | TexCreate_UAV | TexCreate_RenderTargetable, MipCount);
-
 			FRDGTextureRef SceneColorMipchainTexture = GraphBuilder.CreateTexture(SceneColorMipchainDesc, TEXT("SceneColorMipchain"));
 			ViewsSceneColorMipchain[ViewIndex] = SceneColorMipchainTexture;
 
+			FRDGTextureDesc SceneDepthMeterMipchainDesc = SceneColorMipchainDesc;
+			SceneDepthMeterMipchainDesc.Format = PF_R16F; // Store depth form camera in meters, depth up to 64Km.
+			FRDGTextureRef SceneDepthMeterMipchainTexture = GraphBuilder.CreateTexture(SceneDepthMeterMipchainDesc, TEXT("SceneDepthMeterMipchain"));
+			ViewsSceneDepthMeterMipchain[ViewIndex] = SceneDepthMeterMipchainTexture;
+
 			if (!TempSceneColorMipchainTexture || TempSceneColorMipchainTexture->Desc.Extent != SceneColorMipchainTexture->Desc.Extent)
 			{
 				// Only allocate a new temporary if it needs to have another resolution.
 				TempSceneColorMipchainTexture = GraphBuilder.CreateTexture(SceneColorMipchainDesc, TEXT("TempSceneColorMipchain"));
+				TempSceneDepthMeterMipchainTexture = GraphBuilder.CreateTexture(SceneDepthMeterMipchainDesc, TEXT("TempSceneDepthMeterMipchain"));
 			}
 
 			// Copy scene color into the first mip level
 
 			{
-				RDG_EVENT_SCOPE(GraphBuilder, "CopySceneColor");
-				AddCopySceneColorPass(GraphBuilder, View, SceneColorTexture, SceneColorMipchainTexture);
+				RDG_EVENT_SCOPE(GraphBuilder, "CopySceneColorDepth");
+				AddCopySceneColorDepthPass(GraphBuilder, View, SceneColorTexture, SceneDepthTexture, SceneColorMipchainTexture, SceneDepthMeterMipchainTexture);
 			}
 
 			// Now render the mip chain
@@ -555,54 +542,56 @@ void FDeferredShadingSceneRenderer::RenderDistortion(
 
 					PassParameters->SourceTexture = GraphBuilder.CreateSRV(FRDGTextureSRVDesc::CreateForMipLevel(SceneColorMipchainTexture, DstMipIndex - 1));
 					PassParameters->OutTextureMipColor = GraphBuilder.CreateUAV(FRDGTextureUAVDesc(SceneColorMipchainTexture, DstMipIndex));
-
 					AddDownsampleSceneColorPass(GraphBuilder, PassParameters, View);
 
+					FDownsampleSceneColorCS::FParameters* PassParametersDepth = GraphBuilder.AllocParameters<FDownsampleSceneColorCS::FParameters>();
+					*PassParametersDepth = *PassParameters;
+					PassParametersDepth->SourceTexture = GraphBuilder.CreateSRV(FRDGTextureSRVDesc::CreateForMipLevel(SceneDepthMeterMipchainTexture, DstMipIndex - 1));
+					PassParametersDepth->OutTextureMipColor = GraphBuilder.CreateUAV(FRDGTextureUAVDesc(SceneDepthMeterMipchainTexture, DstMipIndex));
+					AddDownsampleSceneColorPass(GraphBuilder, PassParametersDepth, View);
+
 					SrcMipResolution = DstMipResolution;
 					DstMipResolution = DstMipResolution / 2;
 				}
 			}
 
 			// Now the horizontal blur
-			// STATA_TODO: check that compute overlap is working as all horizontal filtering steps can happen in parallel
-
-			const float KernelRadius = 16.0f; // The maximum sample count we will run. This is scaled down by FilterSizeScale.
-			const float FilterSizeScale = CVarRefractionBlurScale.GetValueOnRenderThread();
-			const float CrossCenterWeight = CVarRefractionBlurCenterWeight.GetValueOnRenderThread(); // This must be 0 for the sharp image to not affect the output, the blurring dominates.
 
 			{
-				RDG_EVENT_SCOPE(GraphBuilder, "SceneColorMipHBlur");
+				RDG_EVENT_SCOPE(GraphBuilder, "SceneColorDepthMipHBlur");
 
 				FIntPoint MipResolution = SceneColorMip0Resolution / 2;
 				for (int32 MipIndex = 1; MipIndex < MipCount; MipIndex++)
 				{
-					FVector2D OffsetAndWeight[MAX_FILTER_SAMPLE_COUNT];
-					uint32 SampleCount = FGaussianFiltering::Compute1DGaussianFilterKernel(OffsetAndWeight, MAX_FILTER_SAMPLE_COUNT, KernelRadius, CrossCenterWeight, FilterSizeScale);
 					const FVector2D InverseFilterTextureExtent(1.0f / static_cast<float>(MipResolution.X), 1.0f / static_cast<float>(MipResolution.Y));
 
 					FFilterSceneColorCS::FParameters* PassParameters = GraphBuilder.AllocParameters<FFilterSceneColorCS::FParameters>();
 					PassParameters->SrcMipIndex = MipIndex;
 					PassParameters->MipResolution = MipResolution;
 					PassParameters->SourceSampler = TStaticSamplerState<SF_Bilinear>::GetRHI();
-
-					PassParameters->SourceTexture = GraphBuilder.CreateSRV(FRDGTextureSRVDesc::CreateForMipLevel(SceneColorMipchainTexture, MipIndex));
-					PassParameters->OutTextureMipColor = GraphBuilder.CreateUAV(FRDGTextureUAVDesc(TempSceneColorMipchainTexture, MipIndex));
-
 					PassParameters->SampleCount = SampleCount;
 					PassParameters->BlurDirection = FVector4f(1.0f, 0.0f, 0.0f, 0.0f);
+
 					for (uint32 i = 0; i < SampleCount; ++i)
 					{
 						PassParameters->SampleOffsetsWeights[i] = FVector4f(InverseFilterTextureExtent.X * OffsetAndWeight[i].X, OffsetAndWeight[i].Y);
 					}
 
+					PassParameters->SourceTexture = GraphBuilder.CreateSRV(FRDGTextureSRVDesc::CreateForMipLevel(SceneColorMipchainTexture, MipIndex));
+					PassParameters->OutTextureMipColor = GraphBuilder.CreateUAV(FRDGTextureUAVDesc(TempSceneColorMipchainTexture, MipIndex));
 					AddFilterSceneColorPass(GraphBuilder, PassParameters, View);
 
+					FFilterSceneColorCS::FParameters* PassParametersDepth = GraphBuilder.AllocParameters<FFilterSceneColorCS::FParameters>();
+					*PassParametersDepth = *PassParameters;
+					PassParametersDepth->SourceTexture = GraphBuilder.CreateSRV(FRDGTextureSRVDesc::CreateForMipLevel(SceneDepthMeterMipchainTexture, MipIndex));
+					PassParametersDepth->OutTextureMipColor = GraphBuilder.CreateUAV(FRDGTextureUAVDesc(TempSceneDepthMeterMipchainTexture, MipIndex));
+					AddFilterSceneColorPass(GraphBuilder, PassParametersDepth, View);
+
 					MipResolution = MipResolution / 2;
 				}
 			}
 
 			// Now the vertical blur
-			// STATA_TODO: check that compute overlap is working as all vertical filtering steps can happen in parallel
 
 			{
 				RDG_EVENT_SCOPE(GraphBuilder, "SceneColorMipVBlur");
@@ -610,27 +599,30 @@ void FDeferredShadingSceneRenderer::RenderDistortion(
 				FIntPoint MipResolution = SceneColorMip0Resolution / 2;
 				for (int32 MipIndex = 1; MipIndex < MipCount; MipIndex++)
 				{
-					FVector2D OffsetAndWeight[MAX_FILTER_SAMPLE_COUNT];
-					uint32 SampleCount = FGaussianFiltering::Compute1DGaussianFilterKernel(OffsetAndWeight, MAX_FILTER_SAMPLE_COUNT, KernelRadius, CrossCenterWeight, FilterSizeScale);
 					const FVector2D InverseFilterTextureExtent(1.0f / static_cast<float>(MipResolution.X), 1.0f / static_cast<float>(MipResolution.Y));
 
 					FFilterSceneColorCS::FParameters* PassParameters = GraphBuilder.AllocParameters<FFilterSceneColorCS::FParameters>();
 					PassParameters->SrcMipIndex = MipIndex;
 					PassParameters->MipResolution = MipResolution;
 					PassParameters->SourceSampler = TStaticSamplerState<SF_Bilinear>::GetRHI();
-
-					PassParameters->SourceTexture = GraphBuilder.CreateSRV(FRDGTextureSRVDesc::CreateForMipLevel(TempSceneColorMipchainTexture, MipIndex));
-					PassParameters->OutTextureMipColor = GraphBuilder.CreateUAV(FRDGTextureUAVDesc(SceneColorMipchainTexture, MipIndex));
-
 					PassParameters->SampleCount = SampleCount;
 					PassParameters->BlurDirection = FVector4f(0.0f, 1.0f, 0.0f, 0.0f);
+
 					for (uint32 i = 0; i < SampleCount; ++i)
 					{
 						PassParameters->SampleOffsetsWeights[i] = FVector4f(InverseFilterTextureExtent.Y * OffsetAndWeight[i].X, OffsetAndWeight[i].Y);
 					}
 
+					PassParameters->SourceTexture = GraphBuilder.CreateSRV(FRDGTextureSRVDesc::CreateForMipLevel(TempSceneColorMipchainTexture, MipIndex));
+					PassParameters->OutTextureMipColor = GraphBuilder.CreateUAV(FRDGTextureUAVDesc(SceneColorMipchainTexture, MipIndex));
 					AddFilterSceneColorPass(GraphBuilder, PassParameters, View);
 
+					FFilterSceneColorCS::FParameters* PassParametersDepth = GraphBuilder.AllocParameters<FFilterSceneColorCS::FParameters>();
+					*PassParametersDepth = *PassParameters;
+					PassParametersDepth->SourceTexture = GraphBuilder.CreateSRV(FRDGTextureSRVDesc::CreateForMipLevel(TempSceneDepthMeterMipchainTexture, MipIndex));
+					PassParametersDepth->OutTextureMipColor = GraphBuilder.CreateUAV(FRDGTextureUAVDesc(SceneDepthMeterMipchainTexture, MipIndex));
+					AddFilterSceneColorPass(GraphBuilder, PassParametersDepth, View);
+
 					MipResolution = MipResolution / 2;
 				}
 			}
@@ -665,15 +657,24 @@ void FDeferredShadingSceneRenderer::RenderDistortion(
 		if (bUseRoughRefraction)
 		{
 			// This is the texture containing information about the surface back scattering process
-			RoughnessScatterTexture = GraphBuilder.CreateTexture(
+			RefractionVarianceCoverage = GraphBuilder.CreateTexture(
+				FRDGTextureDesc::Create2D(
+					SceneDepthTexture->Desc.Extent,
+					PF_G16R16F,
+					FClearValueBinding::Black,
+					GFastVRamConfig.Distortion | TexCreate_RenderTargetable | TexCreate_ShaderResource,
+					1,
+					SceneDepthTexture->Desc.NumSamples),
+				TEXT("RefractionVarianceCoverage"));
+			RefractionClosestDepthMeter = GraphBuilder.CreateTexture(
 				FRDGTextureDesc::Create2D(
 					SceneDepthTexture->Desc.Extent,
 					PF_R16F,
-					FClearValueBinding::Transparent,
+					FClearValueBinding(FLinearColor(65504.0f, 1.0f, 1.0f, 1.0f)),	// Max 16bits float depth
 					GFastVRamConfig.Distortion | TexCreate_RenderTargetable | TexCreate_ShaderResource,
 					1,
 					SceneDepthTexture->Desc.NumSamples),
-				TEXT("DistortionRoughnessScatter"));
+				TEXT("RefractionClosestDepthMeter"));
 		}
 
 		ERenderTargetLoadAction LoadAction = ERenderTargetLoadAction::EClear;
@@ -698,7 +699,8 @@ void FDeferredShadingSceneRenderer::RenderDistortion(
 			PassParameters->RenderTargets[0] = FRenderTargetBinding(DistortionTexture, LoadAction);
 			if (bUseRoughRefraction)
 			{
-				PassParameters->RenderTargets[1] = FRenderTargetBinding(RoughnessScatterTexture, LoadAction);
+				PassParameters->RenderTargets[1] = FRenderTargetBinding(RefractionVarianceCoverage, LoadAction);
+				PassParameters->RenderTargets[2] = FRenderTargetBinding(RefractionClosestDepthMeter, LoadAction);
 			}
 			PassParameters->RenderTargets.DepthStencil = StencilWriteBinding;
 
@@ -732,15 +734,20 @@ void FDeferredShadingSceneRenderer::RenderDistortion(
 	CommonParameters.DistortionTexture = DistortionTexture;
 	if (bUseRoughRefraction)
 	{
-		CommonParameters.RoughnessScatterMSAATexture = RoughnessScatterTexture;
-		CommonParameters.RoughnessScatterTexture = RoughnessScatterTexture;
+		CommonParameters.RefractionVarianceCoverageMSAATexture = RefractionVarianceCoverage;
+		CommonParameters.RefractionVarianceCoverageTexture = RefractionVarianceCoverage;
+		CommonParameters.RefractionClosestDepthMeterMSAATexture = RefractionClosestDepthMeter;
+		CommonParameters.RefractionClosestDepthMeterTexture = RefractionClosestDepthMeter;
 	}
 	CommonParameters.SceneColorTextureSampler = bUseRoughRefraction ? TStaticSamplerState<SF_Trilinear>::GetRHI() : TStaticSamplerState<>::GetRHI();
+	CommonParameters.SceneDepthMeterTextureSampler = bUseRoughRefraction ? TStaticSamplerState<SF_Trilinear>::GetRHI() : TStaticSamplerState<>::GetRHI();
 	CommonParameters.DistortionTextureSampler = TStaticSamplerState<>::GetRHI();
-	CommonParameters.RoughnessScatterSampler = TStaticSamplerState<>::GetRHI();
+	CommonParameters.RefractionVarianceCoverageSampler = TStaticSamplerState<>::GetRHI();
+	CommonParameters.RefractionClosestDepthMeterSampler = TStaticSamplerState<>::GetRHI();
 	CommonParameters.StandardTranslucentColorAlphaTextureSampler = TStaticSamplerState<>::GetRHI();
 	CommonParameters.StandardTranslucentTransmittanceTextureSampler = TStaticSamplerState<>::GetRHI();
-	CommonParameters.RefractionRoughnessToMipLevelFactor = FMath::Max(0.0f, CVarRefractionRoughnessToMipLevelFactor.GetValueOnRenderThread());
+	CommonParameters.StandardDevInPixelForMip0 = FilterSigmaPixel;
+
 
 	const bool bAllowStandardTranslucencySeparated = ViewFamily.AllowStandardTranslucencySeparated();
 	FDistortionScreenPS::FPermutationDomain PermutationVector;
@@ -755,6 +762,14 @@ void FDeferredShadingSceneRenderer::RenderDistortion(
 	FScreenPassPipelineState PipelineState(VertexShader, {});
 	FScreenPassTextureViewport Viewport(SceneColorTexture);
 
+	auto GetViewRefractionMaxMipLevel = [](const FViewInfo& View, float StandardDevInPixelForMip0)
+	{
+		float MaxStandardDeviationInScreenRatio = FMath::Clamp(CVarRefractionBlurMaxStandardDeviationInScreenPercent.GetValueOnRenderThread() / 100.0f, 0.0f, 1.0f);
+		float MaxStandardDevInPixel = MaxStandardDeviationInScreenRatio * View.ViewRect.Width();
+		float MaxMipLevel = FMath::Loge(MaxStandardDevInPixel / StandardDevInPixelForMip0) / FMath::Loge(2.0f); // This formula is explained in DistortApplyScreenPS.usf.
+		return MaxMipLevel;
+	};
+
 	// Apply distortion and store off-screen.
 	{
 		QUICK_SCOPE_CYCLE_COUNTER(STAT_FSceneRenderer_RenderDistortion_Apply);
@@ -786,6 +801,9 @@ void FDeferredShadingSceneRenderer::RenderDistortion(
 			{
 				PassParameters->SceneColorMSAATexture = ViewsSceneColorMipchain[ViewIndex];
 				PassParameters->SceneColorTexture = ViewsSceneColorMipchain[ViewIndex];
+				PassParameters->SceneDepthMeterMSAATexture = ViewsSceneDepthMeterMipchain[ViewIndex];
+				PassParameters->SceneDepthMeterTexture = ViewsSceneDepthMeterMipchain[ViewIndex];
+				PassParameters->RefractionMaxMipLevel = GetViewRefractionMaxMipLevel(View, PassParameters->StandardDevInPixelForMip0);
 			}
 			PassParameters->View = View.ViewUniformBuffer;
 			PassParameters->RenderTargets[0] = FRenderTargetBinding(DistortionSceneColorTexture, LoadAction);
@@ -1049,7 +1067,8 @@ FMeshPassProcessor* CreateDistortionPassProcessor(ERHIFeatureLevel::Type Feature
 	if (GetUseRoughRefraction())
 	{
 		DistortionPassState.SetBlendState(TStaticBlendState<CW_RGBA, BO_Add, BF_One, BF_One, BO_Add, BF_One, BF_One,
-															CW_RGBA, BO_Add, BF_One, BF_One, BO_Add, BF_One, BF_One>::GetRHI());
+															CW_RGBA, BO_Add, BF_One, BF_One, BO_Add, BF_One, BF_One,
+															CW_RED,  BO_Min, BF_One, BF_One, BO_Min, BF_One, BF_One>::GetRHI());
 	}
 	else
 	{
@@ -1078,7 +1097,8 @@ FMeshPassProcessor* CreateMobileDistortionPassProcessor(ERHIFeatureLevel::Type F
 	if (GetUseRoughRefraction())
 	{
 		DistortionPassState.SetBlendState(TStaticBlendState<CW_RGBA, BO_Add, BF_One, BF_One, BO_Add, BF_One, BF_One,
-															CW_RGBA, BO_Add, BF_One, BF_One, BO_Add, BF_One, BF_One>::GetRHI());
+															CW_RGBA, BO_Add, BF_One, BF_One, BO_Add, BF_One, BF_One,
+															CW_RED,  BO_Min, BF_One, BF_One, BO_Min, BF_One, BF_One>::GetRHI());
 	}
 	else
 	{
