Fix out of bounds primitive ID overwrite when the instance data buffer size shrinks.

diff --git a/Engine/Source/Runtime/Renderer/Private/GPUScene.cpp b/Engine/Source/Runtime/Renderer/Private/GPUScene.cpp
--- a/Engine/Source/Runtime/Renderer/Private/GPUScene.cpp
+++ b/Engine/Source/Runtime/Renderer/Private/GPUScene.cpp
@@ -2164,10 +2164,18 @@ void FGPUScene::AddClearInstancesPass(FRDGBuilder& GraphBuilder)
 #endif
 	for (FInstanceRange Range : InstanceRangesToClear)
 	{
-		ClearIdData.Add(Range.InstanceSceneDataOffset, Range.NumInstanceSceneDataEntries, INVALID_PRIMITIVE_ID);
+		// Clamp the instance range to the used instances.
+		int32 RangeEnd = FMath::Min(int32(Range.InstanceSceneDataOffset + Range.NumInstanceSceneDataEntries), GetNumInstances());
+		// Clamp to zero to avoid overflow since the start of the range may also be outside the valid instances
+		Range.NumInstanceSceneDataEntries = uint32(FMath::Max(0, RangeEnd - int32(Range.InstanceSceneDataOffset)));
+
+		if (Range.NumInstanceSceneDataEntries > 0u)
+		{			
+			ClearIdData.Add(Range.InstanceSceneDataOffset, Range.NumInstanceSceneDataEntries, INVALID_PRIMITIVE_ID);
 #if LOG_INSTANCE_ALLOCATIONS
-		RangesStr.Appendf(TEXT("[%6d, %6d), "), Range.InstanceSceneDataOffset, Range.InstanceSceneDataOffset + Range.NumInstanceSceneDataEntries);
+			RangesStr.Appendf(TEXT("[%6d, %6d), "), Range.InstanceSceneDataOffset, Range.InstanceSceneDataOffset + Range.NumInstanceSceneDataEntries);
 #endif
+		}
 	}
 #if LOG_INSTANCE_ALLOCATIONS
 	UE_LOG(LogTemp, Warning, TEXT("AddClearInstancesPass: \n%s"), *RangesStr);
