Removed unused code from DistanceFieldAOShared.ush

diff --git a/Engine/Shaders/Private/DistanceFieldAOShared.ush b/Engine/Shaders/Private/DistanceFieldAOShared.ush
--- a/Engine/Shaders/Private/DistanceFieldAOShared.ush
+++ b/Engine/Shaders/Private/DistanceFieldAOShared.ush
@@ -1,21 +1,9 @@
 // Copyright Epic Games, Inc. All Rights Reserved.
 
 /*=============================================================================
-	DistanceFieldAOShared.usf
+	DistanceFieldAOShared.ush
 =============================================================================*/
 
-#ifndef SORTTHREADGROUP_SIZEX
-#define SORTTHREADGROUP_SIZEX 1
-#endif
-
-#ifndef SORTTHREADGROUP_SIZEY
-#define SORTTHREADGROUP_SIZEY 1
-#endif
-
-#ifndef CONE_TRACE_PHASE
-#define CONE_TRACE_PHASE -1
-#endif
-
 // Must match C++
 #define NUM_CONE_STEPS 10
 
@@ -39,9 +27,6 @@ float GetAOMaxDistance()
 	return max(AOObjectMaxDistance, AOGlobalMaxOcclusionDistance);
 }
 
-// One for near, one for far
-#define NUM_CULLED_OBJECT_LISTS 2
-
 float GetStepOffset(float StepIndex)
 {
 	// Original heuristic
@@ -77,20 +62,10 @@ float GetDownsampledDepth(float2 ScreenUV)
 	return abs(Texture2DSampleLevel(DistanceFieldNormalTexture, DistanceFieldNormalSampler, ScreenUV, 0).w);
 }
 
-uint CurrentLevelDownsampleFactor;
-float2 AOBufferSize;
-
-uint DownsampleFactorToBaseLevel;
 float2 BaseLevelTexelSize;
 
 Texture2D BentNormalAOTexture;
 SamplerState BentNormalAOSampler;
-#if !SUPPORTS_INDEPENDENT_SAMPLERS
-SamplerState ConfidenceSampler;
-#endif
-
-Texture2D IrradianceTexture;
-SamplerState IrradianceSampler;
 
 bool SphereIntersectCone(float4 SphereCenterAndRadius, float3 ConeVertex, float3 ConeAxis, float ConeAngleCos, float ConeAngleSin)
 {
diff --git a/Engine/Shaders/Private/DistanceFieldLightingPost.usf b/Engine/Shaders/Private/DistanceFieldLightingPost.usf
--- a/Engine/Shaders/Private/DistanceFieldLightingPost.usf
+++ b/Engine/Shaders/Private/DistanceFieldLightingPost.usf
@@ -37,7 +37,6 @@ float ComputeHistoryWeightBasedOnPosition(float2 ScreenPosition, float SceneDept
 float2 DistanceFieldGBufferTexelSize;
 float2 DistanceFieldGBufferJitterOffset;
 float4 BentNormalBufferAndTexelSize;
-float MinDownsampleFactorToBaseLevel;
 
 float4 GetNormalWeights(float2 Corner00UV, float2 LowResTexelSize, float3 WorldNormal)
 {
diff --git a/Engine/Source/Runtime/Renderer/Private/DistanceFieldLightingPost.cpp b/Engine/Source/Runtime/Renderer/Private/DistanceFieldLightingPost.cpp
--- a/Engine/Source/Runtime/Renderer/Private/DistanceFieldLightingPost.cpp
+++ b/Engine/Source/Runtime/Renderer/Private/DistanceFieldLightingPost.cpp
@@ -79,7 +79,6 @@ BEGIN_SHADER_PARAMETER_STRUCT(FGeometryAwareUpsampleParameters, )
 	SHADER_PARAMETER(FVector2f, DistanceFieldGBufferTexelSize)
 	SHADER_PARAMETER(FVector2f, DistanceFieldGBufferJitterOffset)
 	SHADER_PARAMETER(FVector2f, JitterOffset)
-	SHADER_PARAMETER(float, MinDownsampleFactorToBaseLevel)
 	SHADER_PARAMETER(float, DistanceFadeScale)
 END_SHADER_PARAMETER_STRUCT()
 
@@ -95,9 +94,6 @@ FGeometryAwareUpsampleParameters SetupGeometryAwareUpsampleParameters(const FVie
 	const FIntPoint ConeTracingBufferSize = GetBufferSizeForConeTracing(View);
 	const FVector4f BentNormalBufferAndTexelSizeValue(ConeTracingBufferSize.X, ConeTracingBufferSize.Y, 1.0f / ConeTracingBufferSize.X, 1.0f / ConeTracingBufferSize.Y);
 
-	extern int32 GConeTraceDownsampleFactor;
-	const float MinDownsampleFactor = GConeTraceDownsampleFactor;
-
 	extern float GAOViewFadeDistanceScale;
 	const float DistanceFadeScaleValue = 1.0f / ((1.0f - GAOViewFadeDistanceScale) * GetMaxAOViewDistance());
 
@@ -110,7 +106,6 @@ FGeometryAwareUpsampleParameters SetupGeometryAwareUpsampleParameters(const FVie
 	ShaderParameters.DistanceFieldGBufferTexelSize = BaseLevelTexelSizeValue;
 	ShaderParameters.DistanceFieldGBufferJitterOffset = BaseLevelTexelSizeValue * JitterOffsetValue;
 	ShaderParameters.JitterOffset = JitterOffsetValue;
-	ShaderParameters.MinDownsampleFactorToBaseLevel = MinDownsampleFactor;
 	ShaderParameters.DistanceFadeScale = DistanceFadeScaleValue;
 
 	return ShaderParameters;
