Fix a bug of missing LocalLights if overlapping lights are more than 4.

diff --git a/Engine/Shaders/Private/LightData.ush b/Engine/Shaders/Private/LightData.ush
--- a/Engine/Shaders/Private/LightData.ush
+++ b/Engine/Shaders/Private/LightData.ush
@@ -334,7 +334,6 @@ FDeferredLightData ConvertToDeferredLight_Mobile(const FLocalLightData In)
 	//Render RectLight approximately as SpotLight
 	Out.bSpotLight		= Out.bSpotLight || Out.bRectLight;
 	Out.bRectLight		= false;
-	Out.ShadowedBits	= 1;
 
 	return Out;
 }
diff --git a/Engine/Source/Runtime/Engine/Private/Components/LightComponent.cpp b/Engine/Source/Runtime/Engine/Private/Components/LightComponent.cpp
--- a/Engine/Source/Runtime/Engine/Private/Components/LightComponent.cpp
+++ b/Engine/Source/Runtime/Engine/Private/Components/LightComponent.cpp
@@ -470,6 +470,16 @@ FLightSceneProxy::FLightSceneProxy(const ULightComponent* InLightComponent)
 	LightFunctionDisabledBrightness = LightComponent->DisabledBrightness;
 
 	SamplesPerPixel = LightComponent->SamplesPerPixel;
+
+	if (bCastDynamicShadow && IsMobilePlatform(SceneInterface->GetShaderPlatform()))
+	{
+		if (GetLightType() == LightType_Point
+			|| GetLightType() == LightType_Rect
+			|| (GetLightType() == LightType_Spot && !IsMobileMovableSpotlightShadowsEnabled(SceneInterface->GetShaderPlatform())))
+		{
+			bCastDynamicShadow = false;
+		}
+	}
 }
 
 bool FLightSceneProxy::ShouldCreatePerObjectShadowsForDynamicObjects() const
diff --git a/Engine/Source/Runtime/Engine/Public/SceneManagement.h b/Engine/Source/Runtime/Engine/Public/SceneManagement.h
--- a/Engine/Source/Runtime/Engine/Public/SceneManagement.h
+++ b/Engine/Source/Runtime/Engine/Public/SceneManagement.h
@@ -1762,7 +1762,7 @@ protected:
 	uint8 bStaticShadowing : 1;
 
 	/** True if the light casts dynamic shadows. */
-	const uint8 bCastDynamicShadow : 1;
+	uint8 bCastDynamicShadow : 1;
 
 	/** True if the light casts static shadows. */
 	const uint8 bCastStaticShadow : 1;
diff --git a/Engine/Source/Runtime/Renderer/Private/DynamicShadowMapChannelBindingHelper.cpp b/Engine/Source/Runtime/Renderer/Private/DynamicShadowMapChannelBindingHelper.cpp
--- a/Engine/Source/Runtime/Renderer/Private/DynamicShadowMapChannelBindingHelper.cpp
+++ b/Engine/Source/Runtime/Renderer/Private/DynamicShadowMapChannelBindingHelper.cpp
@@ -35,7 +35,7 @@ void FDynamicShadowMapChannelBindingHelper::UpdateAvailableChannels(const TSpars
 		for (auto It = Lights.CreateConstIterator(); It; ++It)
 		{
 			FLightSceneInfo* OtherLightInfo = It->LightSceneInfo;
-			if (OtherLightInfo && OtherLightInfo != LightInfo)
+			if (OtherLightInfo && (OtherLightInfo->Proxy->CastsDynamicShadow() || OtherLightInfo->Proxy->GetLightFunctionMaterial()) && OtherLightInfo != LightInfo)
 			{
 				const int32 OtherLightChannel = OtherLightInfo->GetDynamicShadowMapChannel();
 				if (OtherLightChannel != INDEX_NONE && IsChannelEnabled(OtherLightChannel))
diff --git a/Engine/Source/Runtime/Renderer/Private/RendererScene.cpp b/Engine/Source/Runtime/Renderer/Private/RendererScene.cpp
--- a/Engine/Source/Runtime/Renderer/Private/RendererScene.cpp
+++ b/Engine/Source/Runtime/Renderer/Private/RendererScene.cpp
@@ -2565,7 +2565,7 @@ void FScene::AssignAvailableShadowMapChannelForLight(FLightSceneInfo* LightScene
 			return;
 		}
 	}
-	else if (LightSceneInfo->Proxy->GetLightType() == LightType_Directional && !IsMobilePlatform(GetShaderPlatform()))
+	else if (LightSceneInfo->Proxy->GetLightType() == LightType_Directional)
 	{
 		// The implementation of forward lighting in ShadowProjectionPixelShader.usf does not support binding the directional light to channel 3.
 		// This is related to the USE_FADE_PLANE feature that encodes the CSM blend factor the alpha channel.
diff --git a/Engine/Source/Runtime/Renderer/Private/ShadowSetup.cpp b/Engine/Source/Runtime/Renderer/Private/ShadowSetup.cpp
--- a/Engine/Source/Runtime/Renderer/Private/ShadowSetup.cpp
+++ b/Engine/Source/Runtime/Renderer/Private/ShadowSetup.cpp
@@ -5019,7 +5019,7 @@ void FSceneRenderer::AllocateShadowDepthTargets(FRHICommandListImmediate& RHICmd
 				bShadowIsVisible = false;
 			}
 
-			if (IsForwardShadingEnabled(ShaderPlatform)
+			if ((IsForwardShadingEnabled(ShaderPlatform) || (IsMobilePlatform(ShaderPlatform) && MobileUsesShadowMaskTexture(ShaderPlatform)))
 				&& ProjectedShadowInfo->GetLightSceneInfo().GetDynamicShadowMapChannel() == -1)
 			{
 				// With forward shading, dynamic shadows are projected into channels of the light attenuation texture based on their assigned DynamicShadowMapChannel
