Potential fix for issue with cloud and volumetric render target access for cloud on translucent.

diff --git a/Engine/Source/Runtime/Renderer/Private/TranslucentRendering.cpp b/Engine/Source/Runtime/Renderer/Private/TranslucentRendering.cpp
--- a/Engine/Source/Runtime/Renderer/Private/TranslucentRendering.cpp
+++ b/Engine/Source/Runtime/Renderer/Private/TranslucentRendering.cpp
@@ -28,6 +28,9 @@ DECLARE_CYCLE_STAT(TEXT("Translucency"), STAT_CLP_Translucency, STATGROUP_Parall
 DECLARE_FLOAT_COUNTER_STAT(TEXT("Translucency GPU Time (MS)"), STAT_TranslucencyGPU, STATGROUP_SceneRendering);
 DEFINE_GPU_DRAWCALL_STAT(Translucency);
 
+// Forward declarations
+bool ShouldRenderVolumetricCloud(const FScene* Scene, const FEngineShowFlags& EngineShowFlags);
+
 static TAutoConsoleVariable<float> CVarSeparateTranslucencyScreenPercentage(
 	TEXT("r.SeparateTranslucencyScreenPercentage"),
 	100.0f,
@@ -886,7 +889,7 @@ TRDGUniformBufferRef<FTranslucentBasePassUniformParameters> CreateTranslucentBas
 		BasePassParameters.VolumetricCloudDepth = nullptr;
 		BasePassParameters.VolumetricCloudColorSampler = TStaticSamplerState<SF_Bilinear, AM_Clamp, AM_Clamp, AM_Clamp>::GetRHI();
 		BasePassParameters.VolumetricCloudDepthSampler = TStaticSamplerState<SF_Bilinear, AM_Clamp, AM_Clamp, AM_Clamp>::GetRHI();
-		if (IsVolumetricRenderTargetEnabled() && View.ViewState)
+		if (IsVolumetricRenderTargetEnabled() && View.ViewState && ShouldRenderVolumetricCloud(Scene, View.Family->EngineShowFlags))
 		{
 			int32 VRTMode = View.ViewState->VolumetricCloudRenderTarget.GetMode();
 			if (VRTMode == 1 || VRTMode == 3)
