Sample DistanceFieldNormalTexture using point sampling in GeometryAwareUpsample.

diff --git a/Engine/Source/Runtime/Renderer/Private/DistanceFieldLightingPost.cpp b/Engine/Source/Runtime/Renderer/Private/DistanceFieldLightingPost.cpp
--- a/Engine/Source/Runtime/Renderer/Private/DistanceFieldLightingPost.cpp
+++ b/Engine/Source/Runtime/Renderer/Private/DistanceFieldLightingPost.cpp
@@ -99,7 +99,7 @@ FGeometryAwareUpsampleParameters SetupGeometryAwareUpsampleParameters(const FVie
 
 	FGeometryAwareUpsampleParameters ShaderParameters;
 	ShaderParameters.DistanceFieldNormalTexture = DistanceFieldNormal;
-	ShaderParameters.DistanceFieldNormalSampler = TStaticSamplerState<SF_Bilinear, AM_Clamp, AM_Clamp, AM_Clamp>::GetRHI();
+	ShaderParameters.DistanceFieldNormalSampler = TStaticSamplerState<SF_Point, AM_Clamp, AM_Clamp, AM_Clamp>::GetRHI();
 	ShaderParameters.BentNormalAOTexture = DistanceFieldAOBentNormal;
 	ShaderParameters.BentNormalAOSampler = TStaticSamplerState<SF_Bilinear, AM_Clamp, AM_Clamp, AM_Clamp>::GetRHI();
 	ShaderParameters.BentNormalBufferAndTexelSize = BentNormalBufferAndTexelSizeValue;
@@ -157,9 +157,7 @@ public:
 		SHADER_PARAMETER_STRUCT_REF(FViewUniformShaderParameters, View)
 		SHADER_PARAMETER_RDG_UNIFORM_BUFFER(FSceneTextureUniformParameters, SceneTextures)
 		SHADER_PARAMETER_RDG_TEXTURE(Texture2D, BentNormalAOTexture)
-		SHADER_PARAMETER_RDG_TEXTURE(Texture2D, DistanceFieldNormalTexture)
 		SHADER_PARAMETER_SAMPLER(SamplerState, BentNormalAOSampler)
-		SHADER_PARAMETER_SAMPLER(SamplerState, DistanceFieldNormalSampler)
 		SHADER_PARAMETER(FVector2f, BentNormalAOTexelSize)
 		SHADER_PARAMETER(FVector2f, MaxSampleBufferUV)
 		SHADER_PARAMETER(float, HistoryWeight)
@@ -454,9 +452,7 @@ void UpdateHistory(
 				auto* PassParameters = GraphBuilder.AllocParameters<FFilterHistoryPS::FParameters>();
 				PassParameters->View = View.ViewUniformBuffer;
 				PassParameters->SceneTextures = SceneTexturesUniformBuffer;
-				PassParameters->DistanceFieldNormalTexture = DistanceFieldNormal;
 				PassParameters->BentNormalAOTexture = NewBentNormalHistory;
-				PassParameters->DistanceFieldNormalSampler = TStaticSamplerState<SF_Bilinear, AM_Clamp, AM_Clamp, AM_Clamp>::GetRHI();
 				PassParameters->BentNormalAOSampler = TStaticSamplerState<SF_Bilinear, AM_Clamp, AM_Clamp, AM_Clamp>::GetRHI();
 				PassParameters->BentNormalAOTexelSize = FVector2f(1.0f / DownsampledBufferSize.X, 1.0f / DownsampledBufferSize.Y);
 				PassParameters->MaxSampleBufferUV = MaxSampleBufferUV;
