Moved shadow GPU scene updates to happen before async Lumen, which removes the need transition resources from async compute back to graphics.

diff --git a/Engine/Source/Runtime/RenderCore/Private/RenderGraphUtils.cpp b/Engine/Source/Runtime/RenderCore/Private/RenderGraphUtils.cpp
--- a/Engine/Source/Runtime/RenderCore/Private/RenderGraphUtils.cpp
+++ b/Engine/Source/Runtime/RenderCore/Private/RenderGraphUtils.cpp
@@ -1078,39 +1078,13 @@ FRDGWaitForTasksScope::~FRDGWaitForTasksScope()
 	}
 }
 
-void FRDGExternalAccessQueue::Submit(FRDGBuilder& GraphBuilder, ERHIPipeline AllowedPipelines)
+void FRDGExternalAccessQueue::Submit(FRDGBuilder& GraphBuilder)
 {
-	// Force any resources that were previously submitted with different pipelines back to internal access mode first.
 	for (FResource Resource : Resources)
 	{
-		if (Resource.SubmittedPipelines != ERHIPipeline::None)
-		{
-			GraphBuilder.UseInternalAccessMode(Resource.Resource);
-		}
-	}
-
-	TArray<FResource, FRDGArrayAllocator> ResourcesToKeep;
-
-	if (AllowedPipelines != ERHIPipeline::All)
-	{
-		ResourcesToKeep.Reserve(Resources.Num());
+		GraphBuilder.UseExternalAccessMode(Resource.Resource, Resource.Access, Resource.Pipelines);
 	}
-
-	for (FResource Resource : Resources)
-	{
-		const ERHIPipeline ResourcePipelines = AllowedPipelines & Resource.RequestedPipelines;
-
-		// If the resource is requesting a pipe that we aren't allowing, then continue to queue the resource for a later submit.
-		if (ResourcePipelines != Resource.RequestedPipelines)
-		{
-			Resource.SubmittedPipelines |= AllowedPipelines;
-			ResourcesToKeep.Emplace(Resource);
-		}
-
-		GraphBuilder.UseExternalAccessMode(Resource.Resource, Resource.Access, ResourcePipelines);
-	}
-
-	Swap(ResourcesToKeep, Resources);
+	Resources.Empty();
 }
 
 bool AllocatePooledBuffer(
diff --git a/Engine/Source/Runtime/RenderCore/Public/RenderGraphUtils.h b/Engine/Source/Runtime/RenderCore/Public/RenderGraphUtils.h
--- a/Engine/Source/Runtime/RenderCore/Public/RenderGraphUtils.h
+++ b/Engine/Source/Runtime/RenderCore/Public/RenderGraphUtils.h
@@ -1091,7 +1091,7 @@ public:
 		Resources.Emplace(Resource, Access, Pipelines);
 	}
 
-	void Submit(FRDGBuilder& GraphBuilder, ERHIPipeline AllowedPipes = ERHIPipeline::All);
+	void Submit(FRDGBuilder& GraphBuilder);
 
 	bool Contains(FRDGViewableResource* Resource)
 	{
@@ -1141,13 +1141,12 @@ private:
 		FResource(FRDGViewableResource* InResource, ERHIAccess InAccess, ERHIPipeline InPipelines)
 			: Resource(InResource)
 			, Access(InAccess)
-			, RequestedPipelines(InPipelines)
+			, Pipelines(InPipelines)
 		{}
 
 		FRDGViewableResource* Resource;
 		ERHIAccess Access;
-		ERHIPipeline RequestedPipelines;
-		ERHIPipeline SubmittedPipelines = ERHIPipeline::None;
+		ERHIPipeline Pipelines = ERHIPipeline::None;
 	};
 
 	TArray<FResource, FRDGArrayAllocator> Resources;
diff --git a/Engine/Source/Runtime/Renderer/Private/DeferredShadingRenderer.cpp b/Engine/Source/Runtime/Renderer/Private/DeferredShadingRenderer.cpp
--- a/Engine/Source/Runtime/Renderer/Private/DeferredShadingRenderer.cpp
+++ b/Engine/Source/Runtime/Renderer/Private/DeferredShadingRenderer.cpp
@@ -2455,7 +2455,7 @@ void FDeferredShadingSceneRenderer::Render(FRDGBuilder& GraphBuilder)
 
 	{
 		RDG_GPU_STAT_SCOPE(GraphBuilder, VisibilityCommands);
-		InitViews(GraphBuilder, SceneTexturesConfig, BasePassDepthStencilAccess, ILCTaskData, InstanceCullingManager);
+		BeginInitViews(GraphBuilder, SceneTexturesConfig, BasePassDepthStencilAccess, ILCTaskData, InstanceCullingManager);
 	}
 
 	// GetBinIndexTranslator cannot be called before UpdateAllPrimitiveSceneInfos which can change the number of raster bins
@@ -2648,14 +2648,6 @@ void FDeferredShadingSceneRenderer::Render(FRDGBuilder& GraphBuilder)
 	}
 #endif // RHI_RAYTRACING
 
-	// Dynamic vertex and index buffers need to be committed before rendering.
-	{
-		SCOPE_CYCLE_COUNTER(STAT_FDeferredShadingSceneRenderer_FGlobalDynamicVertexBuffer_Commit);
-		DynamicIndexBufferForInitViews.Commit();
-		DynamicVertexBufferForInitViews.Commit();
-		DynamicReadBufferForInitViews.Commit();
-	}
-
 	const bool bUseGBuffer = IsUsingGBuffers(ShaderPlatform);
 	
 	const bool bRenderDeferredLighting = ViewFamily.EngineShowFlags.Lighting
@@ -2767,9 +2759,30 @@ void FDeferredShadingSceneRenderer::Render(FRDGBuilder& GraphBuilder)
 
 	PrepareDistanceFieldScene(GraphBuilder, ExternalAccessQueue, false);
 
-	// Changes are flushed to the graphics pipe only, which avoids having to transition from async compute
-	// back to graphics later on during shadow GPU scene updates, which is serializing async Lumen compute work.
-	ExternalAccessQueue.Submit(GraphBuilder, ERHIPipeline::Graphics);
+	FLumenSceneFrameTemporaries LumenFrameTemporaries;
+	{
+		{
+			RDG_RHI_GPU_STAT_SCOPE(GraphBuilder, VisibilityCommands);
+			EndInitViews(GraphBuilder, LumenFrameTemporaries, ILCTaskData, InstanceCullingManager, ExternalAccessQueue);
+		}
+
+		// Dynamic vertex and index buffers need to be committed before rendering.
+		{
+			SCOPE_CYCLE_COUNTER(STAT_FDeferredShadingSceneRenderer_FGlobalDynamicVertexBuffer_Commit);
+			DynamicIndexBufferForInitViews.Commit();
+			DynamicVertexBufferForInitViews.Commit();
+			DynamicReadBufferForInitViews.Commit();
+		}
+
+		{
+			SCOPE_CYCLE_COUNTER(STAT_FDeferredShadingSceneRenderer_FGlobalDynamicVertexBuffer_Commit);
+			DynamicVertexBufferForInitShadows.Commit();
+			DynamicIndexBufferForInitShadows.Commit();
+			DynamicReadBufferForInitShadows.Commit();
+		}
+	}
+
+	ExternalAccessQueue.Submit(GraphBuilder);
 
 	const bool bShouldRenderVelocities = ShouldRenderVelocities();
 	const EShaderPlatform Platform = GetViewFamilyInfo(Views).GetShaderPlatform();
@@ -2810,21 +2823,6 @@ void FDeferredShadingSceneRenderer::Render(FRDGBuilder& GraphBuilder)
 		}
 	}
 
-	FLumenSceneFrameTemporaries LumenFrameTemporaries;
-	{
-		{
-			RDG_RHI_GPU_STAT_SCOPE(GraphBuilder, VisibilityCommands);
-			InitViewsAfterPrepass(GraphBuilder, LumenFrameTemporaries, ILCTaskData, InstanceCullingManager);
-		}
-
-		{
-			SCOPE_CYCLE_COUNTER(STAT_FDeferredShadingSceneRenderer_FGlobalDynamicVertexBuffer_Commit);
-			DynamicVertexBufferForInitShadows.Commit();
-			DynamicIndexBufferForInitShadows.Commit();
-			DynamicReadBufferForInitShadows.Commit();
-		}
-	}
-
 	TArray<Nanite::FRasterResults, TInlineAllocator<2>> NaniteRasterResults;
 	{
 		if (bNaniteEnabled && Views.Num() > 0)
@@ -3117,10 +3115,6 @@ void FDeferredShadingSceneRenderer::Render(FRDGBuilder& GraphBuilder)
 		RenderSkyAtmosphereLookUpTables(GraphBuilder, ExternalAccessQueue);
 	}
 
-	// Changes are flushed to the graphics pipe only, which avoids having to transition from async compute
-	// back to graphics later on during shadow GPU scene updates, which is serializing async Lumen compute work.
-	ExternalAccessQueue.Submit(GraphBuilder, ERHIPipeline::Graphics);
-
 	// Capture the SkyLight using the SkyAtmosphere and VolumetricCloud component if available.
 	const bool bRealTimeSkyCaptureEnabled = Scene->SkyLight && Scene->SkyLight->bRealTimeCaptureEnabled && Views.Num() > 0 && ViewFamily.EngineShowFlags.SkyLighting;
 	if (bRealTimeSkyCaptureEnabled)
@@ -3140,6 +3134,9 @@ void FDeferredShadingSceneRenderer::Render(FRDGBuilder& GraphBuilder)
 		}
 	}
 
+	// Lumen updates need access to sky atmosphere LUT.
+	ExternalAccessQueue.Submit(GraphBuilder);
+
 	UpdateLumenScene(GraphBuilder, LumenFrameTemporaries);
 
 	FRDGTextureRef HalfResolutionDepthCheckerboardMinMaxTexture = nullptr;
@@ -3189,6 +3186,8 @@ void FDeferredShadingSceneRenderer::Render(FRDGBuilder& GraphBuilder)
 		ComputeVolumetricFog(GraphBuilder, SceneTextures);
 	}
 
+	ExternalAccessQueue.Submit(GraphBuilder);
+
 	FDBufferTextures DBufferTextures = CreateDBufferTextures(GraphBuilder, SceneTextures.Config.Extent, ShaderPlatform);
 
 	{
diff --git a/Engine/Source/Runtime/Renderer/Private/DeferredShadingRenderer.h b/Engine/Source/Runtime/Renderer/Private/DeferredShadingRenderer.h
--- a/Engine/Source/Runtime/Renderer/Private/DeferredShadingRenderer.h
+++ b/Engine/Source/Runtime/Renderer/Private/DeferredShadingRenderer.h
@@ -527,9 +527,9 @@ private:
 	void ComputeLightVisibility();
 
 	/** Determines which primitives are visible for each view. */
-	void InitViews(FRDGBuilder& GraphBuilder, const FSceneTexturesConfig& SceneTexturesConfig, FExclusiveDepthStencil::Type BasePassDepthStencilAccess, struct FILCUpdatePrimTaskData& ILCTaskData, FInstanceCullingManager& InstanceCullingManager);
+	void BeginInitViews(FRDGBuilder& GraphBuilder, const FSceneTexturesConfig& SceneTexturesConfig, FExclusiveDepthStencil::Type BasePassDepthStencilAccess, struct FILCUpdatePrimTaskData& ILCTaskData, FInstanceCullingManager& InstanceCullingManager);
+	void EndInitViews(FRDGBuilder& GraphBuilder, FLumenSceneFrameTemporaries& FrameTemporaries, struct FILCUpdatePrimTaskData& ILCTaskData, FInstanceCullingManager& InstanceCullingManager, FRDGExternalAccessQueue& ExternalAccessQueue);
 
-	void InitViewsAfterPrepass(FRDGBuilder& GraphBuilder, FLumenSceneFrameTemporaries& FrameTemporaries, struct FILCUpdatePrimTaskData& ILCTaskData, FInstanceCullingManager& InstanceCullingManager);
 	void BeginUpdateLumenSceneTasks(FRDGBuilder& GraphBuilder, FLumenSceneFrameTemporaries& FrameTemporaries);
 	void UpdateLumenScene(FRDGBuilder& GraphBuilder, FLumenSceneFrameTemporaries& FrameTemporaries);
 	void RenderLumenSceneLighting(FRDGBuilder& GraphBuilder, const FLumenSceneFrameTemporaries& FrameTemporaries);
diff --git a/Engine/Source/Runtime/Renderer/Private/MobileShadingRenderer.cpp b/Engine/Source/Runtime/Renderer/Private/MobileShadingRenderer.cpp
--- a/Engine/Source/Runtime/Renderer/Private/MobileShadingRenderer.cpp
+++ b/Engine/Source/Runtime/Renderer/Private/MobileShadingRenderer.cpp
@@ -626,11 +626,13 @@ void FMobileSceneRenderer::InitViews(FRDGBuilder& GraphBuilder, FSceneTexturesCo
 	}
 	
 	const bool bDynamicShadows = ViewFamily.EngineShowFlags.DynamicShadows;
-	
+
+	FRDGExternalAccessQueue ExternalAccessQueue;
+
 	if (bDynamicShadows)
 	{
 		// Setup dynamic shadows.
-		InitDynamicShadows(RHICmdList, InstanceCullingManager);
+		InitDynamicShadows(GraphBuilder, InstanceCullingManager, ExternalAccessQueue);
 	}
 	else
 	{
@@ -660,8 +662,6 @@ void FMobileSceneRenderer::InitViews(FRDGBuilder& GraphBuilder, FSceneTexturesCo
 		View.InitRHIResources();
 	}
 
-	FRDGExternalAccessQueue ExternalAccessQueue;
-
 	{
 		RDG_CSV_STAT_EXCLUSIVE_SCOPE(GraphBuilder, UpdateGPUScene);
 
diff --git a/Engine/Source/Runtime/Renderer/Private/SceneHitProxyRendering.cpp b/Engine/Source/Runtime/Renderer/Private/SceneHitProxyRendering.cpp
--- a/Engine/Source/Runtime/Renderer/Private/SceneHitProxyRendering.cpp
+++ b/Engine/Source/Runtime/Renderer/Private/SceneHitProxyRendering.cpp
@@ -627,12 +627,10 @@ void FDeferredShadingSceneRenderer::RenderHitProxies(FRDGBuilder& GraphBuilder)
 	FInstanceCullingManager& InstanceCullingManager = *GraphBuilder.AllocObject<FInstanceCullingManager>(Scene->GPUScene.IsEnabled(), GraphBuilder);
 
 	// Find the visible primitives.
-	{
-		FLumenSceneFrameTemporaries LumenFrameTemporaries;
-		FILCUpdatePrimTaskData ILCTaskData;
-		InitViews(GraphBuilder, SceneTexturesConfig, FExclusiveDepthStencil::DepthWrite_StencilWrite, ILCTaskData, InstanceCullingManager);
-		InitViewsAfterPrepass(GraphBuilder, LumenFrameTemporaries, ILCTaskData, InstanceCullingManager);
-	}
+	FLumenSceneFrameTemporaries LumenFrameTemporaries;
+	FILCUpdatePrimTaskData ILCTaskData;
+	FRDGExternalAccessQueue ExternalAccessQueue;
+	BeginInitViews(GraphBuilder, SceneTexturesConfig, FExclusiveDepthStencil::DepthWrite_StencilWrite, ILCTaskData, InstanceCullingManager);
 
 	extern TSet<IPersistentViewUniformBufferExtension*> PersistentViewUniformBufferExtensions;
 
@@ -649,18 +647,16 @@ void FDeferredShadingSceneRenderer::RenderHitProxies(FRDGBuilder& GraphBuilder)
 
 	ShaderPrint::BeginViews(GraphBuilder, Views);
 
-	{
-		FRDGExternalAccessQueue ExternalAccessQueue;
+	Scene->GPUScene.Update(GraphBuilder, *Scene, ExternalAccessQueue);
 
-		Scene->GPUScene.Update(GraphBuilder, *Scene, ExternalAccessQueue);
+	for (int32 ViewIndex = 0; ViewIndex < Views.Num(); ViewIndex++)
+	{
+		Scene->GPUScene.UploadDynamicPrimitiveShaderDataForView(GraphBuilder, *Scene, Views[ViewIndex], ExternalAccessQueue);
+	}
 
-		for (int32 ViewIndex = 0; ViewIndex < Views.Num(); ViewIndex++)
-		{
-			Scene->GPUScene.UploadDynamicPrimitiveShaderDataForView(GraphBuilder, *Scene, Views[ViewIndex], ExternalAccessQueue);
-		}
+	EndInitViews(GraphBuilder, LumenFrameTemporaries, ILCTaskData, InstanceCullingManager, ExternalAccessQueue);
 
-		ExternalAccessQueue.Submit(GraphBuilder);
-	}
+	ExternalAccessQueue.Submit(GraphBuilder);
 
 	InstanceCullingManager.FlushRegisteredViews(GraphBuilder);
 
diff --git a/Engine/Source/Runtime/Renderer/Private/SceneRendering.h b/Engine/Source/Runtime/Renderer/Private/SceneRendering.h
--- a/Engine/Source/Runtime/Renderer/Private/SceneRendering.h
+++ b/Engine/Source/Runtime/Renderer/Private/SceneRendering.h
@@ -2245,8 +2245,8 @@ protected:
 	// Shared functionality between all scene renderers
 
 	FDynamicShadowsTaskData* BeginInitDynamicShadows(bool bRunningEarly);
-	void FinishInitDynamicShadows(FRHICommandListImmediate& RHICmdList, FDynamicShadowsTaskData* TaskData, FGlobalDynamicIndexBuffer& DynamicIndexBuffer, FGlobalDynamicVertexBuffer& DynamicVertexBuffer, FGlobalDynamicReadBuffer& DynamicReadBuffer, FInstanceCullingManager& InstanceCullingManager);
-	void InitDynamicShadows(FRHICommandListImmediate& RHICmdList, FGlobalDynamicIndexBuffer& DynamicIndexBuffer, FGlobalDynamicVertexBuffer& DynamicVertexBuffer, FGlobalDynamicReadBuffer& DynamicReadBuffer, FInstanceCullingManager& InstanceCullingManager);
+	void FinishInitDynamicShadows(FRDGBuilder& GraphBuilder, FDynamicShadowsTaskData* TaskData, FGlobalDynamicIndexBuffer& DynamicIndexBuffer, FGlobalDynamicVertexBuffer& DynamicVertexBuffer, FGlobalDynamicReadBuffer& DynamicReadBuffer, FInstanceCullingManager& InstanceCullingManager, FRDGExternalAccessQueue& ExternalAccessQueue);
+	void InitDynamicShadows(FRDGBuilder& GraphBuilder, FGlobalDynamicIndexBuffer& DynamicIndexBuffer, FGlobalDynamicVertexBuffer& DynamicVertexBuffer, FGlobalDynamicReadBuffer& DynamicReadBuffer, FInstanceCullingManager& InstanceCullingManager, FRDGExternalAccessQueue& ExternalAccessQueue);
 
 	void SetupMeshPass(FViewInfo& View, FExclusiveDepthStencil::Type BasePassDepthStencilAccess, FViewCommands& ViewCommands, FInstanceCullingManager& InstanceCullingManager);
 
@@ -2532,7 +2532,7 @@ public:
 
 protected:
 	/** Finds the visible dynamic shadows for each view. */
-	void InitDynamicShadows(FRHICommandListImmediate& RHICmdList, FInstanceCullingManager& FInstanceCullingManager);
+	void InitDynamicShadows(FRDGBuilder& GraphBuilder, FInstanceCullingManager& FInstanceCullingManager, FRDGExternalAccessQueue& ExternalAccessQueue);
 
 	void PrepareViewVisibilityLists();
 
diff --git a/Engine/Source/Runtime/Renderer/Private/SceneVisibility.cpp b/Engine/Source/Runtime/Renderer/Private/SceneVisibility.cpp
--- a/Engine/Source/Runtime/Renderer/Private/SceneVisibility.cpp
+++ b/Engine/Source/Runtime/Renderer/Private/SceneVisibility.cpp
@@ -5208,7 +5208,7 @@ void FDeferredShadingSceneRenderer::PreVisibilityFrameSetup(FRDGBuilder& GraphBu
  * Initialize scene's views.
  * Check visibility, build visible mesh commands, etc.
  */
-void FDeferredShadingSceneRenderer::InitViews(FRDGBuilder& GraphBuilder, const FSceneTexturesConfig& SceneTexturesConfig, FExclusiveDepthStencil::Type BasePassDepthStencilAccess, struct FILCUpdatePrimTaskData& ILCTaskData, FInstanceCullingManager& InstanceCullingManager)
+void FDeferredShadingSceneRenderer::BeginInitViews(FRDGBuilder& GraphBuilder, const FSceneTexturesConfig& SceneTexturesConfig, FExclusiveDepthStencil::Type BasePassDepthStencilAccess, struct FILCUpdatePrimTaskData& ILCTaskData, FInstanceCullingManager& InstanceCullingManager)
 {
 	SCOPED_NAMED_EVENT(FDeferredShadingSceneRenderer_InitViews, FColor::Emerald);
 	SCOPE_CYCLE_COUNTER(STAT_InitViewsTime);
@@ -5360,7 +5360,7 @@ void FSceneRenderer::SetupSceneReflectionCaptureBuffer(FRHICommandListImmediate&
 	}
 }
 
-void FDeferredShadingSceneRenderer::InitViewsAfterPrepass(FRDGBuilder& GraphBuilder, FLumenSceneFrameTemporaries& FrameTemporaries, struct FILCUpdatePrimTaskData& ILCTaskData, FInstanceCullingManager& InstanceCullingManager)
+void FDeferredShadingSceneRenderer::EndInitViews(FRDGBuilder& GraphBuilder, FLumenSceneFrameTemporaries& FrameTemporaries, struct FILCUpdatePrimTaskData& ILCTaskData, FInstanceCullingManager& InstanceCullingManager, FRDGExternalAccessQueue& ExternalAccessQueue)
 {
 	SCOPED_NAMED_EVENT(FDeferredShadingSceneRenderer_InitViewsAfterPrepass, FColor::Emerald);
 	SCOPE_CYCLE_COUNTER(STAT_InitViewsPossiblyAfterPrepass);
@@ -5376,12 +5376,12 @@ void FDeferredShadingSceneRenderer::InitViewsAfterPrepass(FRDGBuilder& GraphBuil
 		// Setup dynamic shadows.
 		if (CurrentDynamicShadowsTaskData)
 		{
-			FinishInitDynamicShadows(RHICmdList, CurrentDynamicShadowsTaskData, DynamicIndexBufferForInitShadows, DynamicVertexBufferForInitShadows, DynamicReadBufferForInitShadows, InstanceCullingManager);
+			FinishInitDynamicShadows(GraphBuilder, CurrentDynamicShadowsTaskData, DynamicIndexBufferForInitShadows, DynamicVertexBufferForInitShadows, DynamicReadBufferForInitShadows, InstanceCullingManager, ExternalAccessQueue);
 			CurrentDynamicShadowsTaskData = nullptr;
 		}
 		else
 		{
-			InitDynamicShadows(RHICmdList, DynamicIndexBufferForInitShadows, DynamicVertexBufferForInitShadows, DynamicReadBufferForInitShadows, InstanceCullingManager);
+			InitDynamicShadows(GraphBuilder, DynamicIndexBufferForInitShadows, DynamicVertexBufferForInitShadows, DynamicReadBufferForInitShadows, InstanceCullingManager, ExternalAccessQueue);
 		}
 
 		if (GDynamicRHI->RHIIncludeOptionalFlushes())
diff --git a/Engine/Source/Runtime/Renderer/Private/ShadowDepthRendering.cpp b/Engine/Source/Runtime/Renderer/Private/ShadowDepthRendering.cpp
--- a/Engine/Source/Runtime/Renderer/Private/ShadowDepthRendering.cpp
+++ b/Engine/Source/Runtime/Renderer/Private/ShadowDepthRendering.cpp
@@ -1642,46 +1642,6 @@ void FSceneRenderer::RenderShadowDepthMaps(FRDGBuilder& GraphBuilder, FInstanceC
 	RDG_EVENT_SCOPE(GraphBuilder, "ShadowDepths");
 	RDG_GPU_STAT_SCOPE(GraphBuilder, ShadowDepths);
 
-	// Ensure all shadow view dynamic primitives are uploaded before shadow-culling batching pass.
-	// TODO: automate this such that:
-	//  1. we only process views that need it (have dynamic primitives)
-	//  2. it is integrated in the GPU-scene (it already collects the dynamic primives and know about them...)
-	//  3. BUT: we need to touch the views to update the GPUScene buffer references in the FViewInfo
-	//          so need to refactor that into its own binding point, probably. Or something.
-	for (FSortedShadowMapAtlas& ShadowMapAtlas : SortedShadowsForShadowDepthPass.ShadowMapAtlases)
-	{
-		for (FProjectedShadowInfo* ProjectedShadowInfo : ShadowMapAtlas.Shadows)
-		{
-			Scene->GPUScene.UploadDynamicPrimitiveShaderDataForView(GraphBuilder, *Scene, *ProjectedShadowInfo->ShadowDepthView, ExternalAccessQueue, true);
-		}
-	}
-	for (FSortedShadowMapAtlas& ShadowMap : SortedShadowsForShadowDepthPass.ShadowMapCubemaps)
-	{
-		check(ShadowMap.Shadows.Num() == 1);
-		FProjectedShadowInfo* ProjectedShadowInfo = ShadowMap.Shadows[0];
-		Scene->GPUScene.UploadDynamicPrimitiveShaderDataForView(GraphBuilder, *Scene, *ProjectedShadowInfo->ShadowDepthView, ExternalAccessQueue, true);
-	}
-	for (FProjectedShadowInfo* ProjectedShadowInfo : SortedShadowsForShadowDepthPass.PreshadowCache.Shadows)
-	{
-		if (!ProjectedShadowInfo->bDepthsCached)
-		{
-			Scene->GPUScene.UploadDynamicPrimitiveShaderDataForView(GraphBuilder, *Scene, *ProjectedShadowInfo->ShadowDepthView, ExternalAccessQueue, true);
-		}
-	}
-	for (const FSortedShadowMapAtlas& ShadowMapAtlas : SortedShadowsForShadowDepthPass.TranslucencyShadowMapAtlases)
-	{
-		for (FProjectedShadowInfo* ProjectedShadowInfo : ShadowMapAtlas.Shadows)
-		{
-			Scene->GPUScene.UploadDynamicPrimitiveShaderDataForView(GraphBuilder, *Scene, *ProjectedShadowInfo->ShadowDepthView, ExternalAccessQueue, true);
-		}
-	}
-	for (FProjectedShadowInfo* ProjectedShadowInfo : SortedShadowsForShadowDepthPass.VirtualShadowMapShadows)
-	{
-		Scene->GPUScene.UploadDynamicPrimitiveShaderDataForView(GraphBuilder, *Scene, *ProjectedShadowInfo->ShadowDepthView, ExternalAccessQueue, true);
-	}
-
-	ExternalAccessQueue.Submit(GraphBuilder);
-
 	// Begin new deferred culling batching scope to catch shadow render passes, as there can use dynamic primitives that have not been uploaded before 
 	// the previous batching scope. Also flushes the culling views registered during the setup (in InitViewsAfterPrepass) that are referenced in the shadow view
 	// culling.
diff --git a/Engine/Source/Runtime/Renderer/Private/ShadowSetup.cpp b/Engine/Source/Runtime/Renderer/Private/ShadowSetup.cpp
--- a/Engine/Source/Runtime/Renderer/Private/ShadowSetup.cpp
+++ b/Engine/Source/Runtime/Renderer/Private/ShadowSetup.cpp
@@ -5824,7 +5824,7 @@ FDynamicShadowsTaskData* FSceneRenderer::BeginInitDynamicShadows(bool bRunningEa
 }
 
 
-void FSceneRenderer::FinishInitDynamicShadows(FRHICommandListImmediate& RHICmdList, FDynamicShadowsTaskData* TaskData, FGlobalDynamicIndexBuffer& DynamicIndexBuffer, FGlobalDynamicVertexBuffer& DynamicVertexBuffer, FGlobalDynamicReadBuffer& DynamicReadBuffer, FInstanceCullingManager& InstanceCullingManager)
+void FSceneRenderer::FinishInitDynamicShadows(FRDGBuilder& GraphBuilder, FDynamicShadowsTaskData* TaskData, FGlobalDynamicIndexBuffer& DynamicIndexBuffer, FGlobalDynamicVertexBuffer& DynamicVertexBuffer, FGlobalDynamicReadBuffer& DynamicReadBuffer, FInstanceCullingManager& InstanceCullingManager, FRDGExternalAccessQueue& ExternalAccessQueue)
 {
 	SCOPED_NAMED_EVENT_TEXT("FSceneRenderer::FinishInitDynamicShadows", FColor::Magenta);
 	SCOPE_CYCLE_COUNTER(STAT_InitDynamicShadowsTime);
@@ -5838,16 +5838,54 @@ void FSceneRenderer::FinishInitDynamicShadows(FRHICommandListImmediate& RHICmdLi
 	{
 		ShadowSceneRenderer->PostInitDynamicShadowsSetup();
 	}
-	AllocateShadowDepthTargets(RHICmdList);
+	AllocateShadowDepthTargets(GraphBuilder.RHICmdList);
 
 	// Generate mesh element arrays from shadow primitive arrays
 	GatherShadowDynamicMeshElements(DynamicIndexBuffer, DynamicVertexBuffer, DynamicReadBuffer, InstanceCullingManager);
+
+	// Ensure all shadow view dynamic primitives are uploaded before shadow-culling batching pass.
+	// TODO: automate this such that:
+	//  1. we only process views that need it (have dynamic primitives)
+	//  2. it is integrated in the GPU-scene (it already collects the dynamic primives and know about them...)
+	//  3. BUT: we need to touch the views to update the GPUScene buffer references in the FViewInfo
+	//          so need to refactor that into its own binding point, probably. Or something.
+	for (FSortedShadowMapAtlas& ShadowMapAtlas : SortedShadowsForShadowDepthPass.ShadowMapAtlases)
+	{
+		for (FProjectedShadowInfo* ProjectedShadowInfo : ShadowMapAtlas.Shadows)
+		{
+			Scene->GPUScene.UploadDynamicPrimitiveShaderDataForView(GraphBuilder, *Scene, *ProjectedShadowInfo->ShadowDepthView, ExternalAccessQueue, true);
+		}
+	}
+	for (FSortedShadowMapAtlas& ShadowMap : SortedShadowsForShadowDepthPass.ShadowMapCubemaps)
+	{
+		check(ShadowMap.Shadows.Num() == 1);
+		FProjectedShadowInfo* ProjectedShadowInfo = ShadowMap.Shadows[0];
+		Scene->GPUScene.UploadDynamicPrimitiveShaderDataForView(GraphBuilder, *Scene, *ProjectedShadowInfo->ShadowDepthView, ExternalAccessQueue, true);
+	}
+	for (FProjectedShadowInfo* ProjectedShadowInfo : SortedShadowsForShadowDepthPass.PreshadowCache.Shadows)
+	{
+		if (!ProjectedShadowInfo->bDepthsCached)
+		{
+			Scene->GPUScene.UploadDynamicPrimitiveShaderDataForView(GraphBuilder, *Scene, *ProjectedShadowInfo->ShadowDepthView, ExternalAccessQueue, true);
+		}
+	}
+	for (const FSortedShadowMapAtlas& ShadowMapAtlas : SortedShadowsForShadowDepthPass.TranslucencyShadowMapAtlases)
+	{
+		for (FProjectedShadowInfo* ProjectedShadowInfo : ShadowMapAtlas.Shadows)
+		{
+			Scene->GPUScene.UploadDynamicPrimitiveShaderDataForView(GraphBuilder, *Scene, *ProjectedShadowInfo->ShadowDepthView, ExternalAccessQueue, true);
+		}
+	}
+	for (FProjectedShadowInfo* ProjectedShadowInfo : SortedShadowsForShadowDepthPass.VirtualShadowMapShadows)
+	{
+		Scene->GPUScene.UploadDynamicPrimitiveShaderDataForView(GraphBuilder, *Scene, *ProjectedShadowInfo->ShadowDepthView, ExternalAccessQueue, true);
+	}
 }
 
-void FSceneRenderer::InitDynamicShadows(FRHICommandListImmediate& RHICmdList, FGlobalDynamicIndexBuffer& DynamicIndexBuffer, FGlobalDynamicVertexBuffer& DynamicVertexBuffer, FGlobalDynamicReadBuffer& DynamicReadBuffer, FInstanceCullingManager& InstanceCullingManager)
+void FSceneRenderer::InitDynamicShadows(FRDGBuilder& GraphBuilder, FGlobalDynamicIndexBuffer& DynamicIndexBuffer, FGlobalDynamicVertexBuffer& DynamicVertexBuffer, FGlobalDynamicReadBuffer& DynamicReadBuffer, FInstanceCullingManager& InstanceCullingManager, FRDGExternalAccessQueue& ExternalAccessQueue)
 {
 	FDynamicShadowsTaskData* TaskData = BeginInitDynamicShadows(false);
-	FinishInitDynamicShadows(RHICmdList, TaskData, DynamicIndexBuffer, DynamicVertexBuffer, DynamicReadBuffer, InstanceCullingManager);
+	FinishInitDynamicShadows(GraphBuilder, TaskData, DynamicIndexBuffer, DynamicVertexBuffer, DynamicReadBuffer, InstanceCullingManager, ExternalAccessQueue);
 }
 
 void FSceneRenderer::AllocateMobileCSMAndSpotLightShadowDepthTargets(FRHICommandListImmediate& RHICmdList, const TArray<FProjectedShadowInfo*, SceneRenderingAllocator>& MobileCSMAndSpotLightShadows)
diff --git a/Engine/Source/Runtime/Renderer/Private/ShadowSetupMobile.cpp b/Engine/Source/Runtime/Renderer/Private/ShadowSetupMobile.cpp
--- a/Engine/Source/Runtime/Renderer/Private/ShadowSetupMobile.cpp
+++ b/Engine/Source/Runtime/Renderer/Private/ShadowSetupMobile.cpp
@@ -280,7 +280,7 @@ static void VisualizeMobileDynamicCSMSubjectCapsules(FViewInfo& View, FLightScen
 }
 
 /** Finds the visible dynamic shadows for each view. */
-void FMobileSceneRenderer::InitDynamicShadows(FRHICommandListImmediate& RHICmdList, FInstanceCullingManager& InstanceCullingManager)
+void FMobileSceneRenderer::InitDynamicShadows(FRDGBuilder& GraphBuilder, FInstanceCullingManager& InstanceCullingManager, FRDGExternalAccessQueue& ExternalAccessQueue)
 {
 	static auto* MyCVarMobileEnableStaticAndCSMShadowReceivers = IConsoleManager::Get().FindTConsoleVariableDataInt(TEXT("r.Mobile.EnableStaticAndCSMShadowReceivers"));
 	const bool bCombinedStaticAndCSMEnabled = MyCVarMobileEnableStaticAndCSMShadowReceivers->GetValueOnRenderThread()!=0;
@@ -304,7 +304,7 @@ void FMobileSceneRenderer::InitDynamicShadows(FRHICommandListImmediate& RHICmdLi
 		}
 	}
 
-	FSceneRenderer::InitDynamicShadows(RHICmdList, DynamicIndexBuffer, DynamicVertexBuffer, DynamicReadBuffer, InstanceCullingManager);
+	FSceneRenderer::InitDynamicShadows(GraphBuilder, DynamicIndexBuffer, DynamicVertexBuffer, DynamicReadBuffer, InstanceCullingManager, ExternalAccessQueue);
 
 	bool bAlwaysUseCSM = false;
 	const bool bSkipCSMShaderCulling = MobileBasePassAlwaysUsesCSM(Scene->GetShaderPlatform());
