Dependency cleanup around DataDrivenShaderPlatformInfo and a few other headers.

diff --git a/Engine/Plugins/2D/Paper2D/Source/Paper2D/Private/PaperSpriteVertexBuffer.cpp b/Engine/Plugins/2D/Paper2D/Source/Paper2D/Private/PaperSpriteVertexBuffer.cpp
--- a/Engine/Plugins/2D/Paper2D/Source/Paper2D/Private/PaperSpriteVertexBuffer.cpp
+++ b/Engine/Plugins/2D/Paper2D/Source/Paper2D/Private/PaperSpriteVertexBuffer.cpp
@@ -7,6 +7,7 @@
 #include "EngineGlobals.h"
 #include "Engine/Engine.h"
 #include "PaperSpriteComponent.h"
+#include "DataDrivenShaderPlatformInfo.h"
 
 //////////////////////////////////////////////////////////////////////////
 // FPaperSpriteVertexBuffer
diff --git a/Engine/Plugins/Animation/DeformerGraph/Source/OptimusCore/Private/DataInterfaces/OptimusDataInterfaceSkinnedMeshWrite.cpp b/Engine/Plugins/Animation/DeformerGraph/Source/OptimusCore/Private/DataInterfaces/OptimusDataInterfaceSkinnedMeshWrite.cpp
--- a/Engine/Plugins/Animation/DeformerGraph/Source/OptimusCore/Private/DataInterfaces/OptimusDataInterfaceSkinnedMeshWrite.cpp
+++ b/Engine/Plugins/Animation/DeformerGraph/Source/OptimusCore/Private/DataInterfaces/OptimusDataInterfaceSkinnedMeshWrite.cpp
@@ -11,6 +11,7 @@
 #include "ShaderParameterMetadataBuilder.h"
 #include "SkeletalMeshDeformerHelpers.h"
 #include "SkeletalRenderPublic.h"
+#include "DataDrivenShaderPlatformInfo.h"
 
 #include UE_INLINE_GENERATED_CPP_BY_NAME(OptimusDataInterfaceSkinnedMeshWrite)
 
diff --git a/Engine/Plugins/Animation/RigLogic/Source/RigLogicModule/Private/SkelMeshDNAUtils.cpp b/Engine/Plugins/Animation/RigLogic/Source/RigLogicModule/Private/SkelMeshDNAUtils.cpp
--- a/Engine/Plugins/Animation/RigLogic/Source/RigLogicModule/Private/SkelMeshDNAUtils.cpp
+++ b/Engine/Plugins/Animation/RigLogic/Source/RigLogicModule/Private/SkelMeshDNAUtils.cpp
@@ -3,7 +3,7 @@
 #include "DNAToSkelMeshMap.h"
 
 #include "RenderResource.h"
-#include "RHICommandList.h"
+#include "RHI.h"
 #include "Async/ParallelFor.h"
 #include "Components/SkeletalMeshComponent.h"
 #include "Engine/StaticMesh.h"
diff --git a/Engine/Plugins/Experimental/GPULightmass/Source/GPULightmassEditor/Private/GPULightmassEditorModule.cpp b/Engine/Plugins/Experimental/GPULightmass/Source/GPULightmassEditor/Private/GPULightmassEditorModule.cpp
--- a/Engine/Plugins/Experimental/GPULightmass/Source/GPULightmassEditor/Private/GPULightmassEditorModule.cpp
+++ b/Engine/Plugins/Experimental/GPULightmass/Source/GPULightmassEditor/Private/GPULightmassEditorModule.cpp
@@ -7,6 +7,7 @@
 #include "HAL/IConsoleManager.h"
 #include "PropertyEditorModule.h"
 #include "RenderingThread.h"
+#include "DataDrivenShaderPlatformInfo.h"
 #include "Interfaces/IPluginManager.h"
 #include "ShaderCore.h"
 #include "SceneInterface.h"
@@ -73,6 +74,9 @@ enum class ERayTracingDisabledReason : int32
 	INCAPABLE_HARDWARE, // The video card isn't capable of hardware ray tracing (too old), or supports only software emulated ray tracing (GTX 10 series)
 };
 
+typedef TBitArray<TInlineAllocator<EShaderPlatform::SP_NumPlatforms / 8>> ShaderPlatformMaskType;
+extern RENDERCORE_API ShaderPlatformMaskType GRayTracingPlatformMask;
+
 ERayTracingDisabledReason GetRayTracingDisabledReason()
 {
 	if (IsRayTracingEnabled())
@@ -86,7 +90,6 @@ ERayTracingDisabledReason GetRayTracingDisabledReason()
 		return ERayTracingDisabledReason::DISABLED_BY_PROJECT_SETTINGS;
 	}
 
-	extern RENDERCORE_API ShaderPlatformMaskType GRayTracingPlatformMask;
 	if (!RHISupportsRayTracing(GMaxRHIShaderPlatform))
 	{
 		return ERayTracingDisabledReason::INCOMPATIBLE_SHADER_PLATFORM;
diff --git a/Engine/Plugins/Experimental/Water/Source/Runtime/Private/GerstnerWaterWaveViewExtension.cpp b/Engine/Plugins/Experimental/Water/Source/Runtime/Private/GerstnerWaterWaveViewExtension.cpp
--- a/Engine/Plugins/Experimental/Water/Source/Runtime/Private/GerstnerWaterWaveViewExtension.cpp
+++ b/Engine/Plugins/Experimental/Water/Source/Runtime/Private/GerstnerWaterWaveViewExtension.cpp
@@ -4,6 +4,7 @@
 #include "GerstnerWaterWaveSubsystem.h"
 #include "WaterBodyActor.h"
 #include "GerstnerWaterWaves.h"
+#include "Containers/DynamicRHIResourceArray.h"
 #include "Engine/Engine.h"
 #include "SceneView.h"
 #include "WaterBodyManager.h"
diff --git a/Engine/Plugins/FX/Niagara/Source/Niagara/Private/Experimental/NiagaraMeshUvMapping.h b/Engine/Plugins/FX/Niagara/Source/Niagara/Private/Experimental/NiagaraMeshUvMapping.h
--- a/Engine/Plugins/FX/Niagara/Source/Niagara/Private/Experimental/NiagaraMeshUvMapping.h
+++ b/Engine/Plugins/FX/Niagara/Source/Niagara/Private/Experimental/NiagaraMeshUvMapping.h
@@ -7,6 +7,8 @@
 #include "Experimental/NiagaraMeshUvMappingHandle.h"
 #include "NiagaraUvQuadTree.h"
 #include "RenderResource.h"
+#include "RHI.h"
+#include "Containers/DynamicRHIResourceArray.h"
 
 struct FMeshUvMapping;
 class FNiagaraUvQuadTree;
diff --git a/Engine/Plugins/FX/Niagara/Source/Niagara/Public/NiagaraSystemStaticBuffers.h b/Engine/Plugins/FX/Niagara/Source/Niagara/Public/NiagaraSystemStaticBuffers.h
--- a/Engine/Plugins/FX/Niagara/Source/Niagara/Public/NiagaraSystemStaticBuffers.h
+++ b/Engine/Plugins/FX/Niagara/Source/Niagara/Public/NiagaraSystemStaticBuffers.h
@@ -7,6 +7,7 @@
 #include "RenderResource.h"
 #include "RenderingThread.h"
 #include "RHIUtilities.h"
+#include "Containers/DynamicRHIResourceArray.h"
 
 // Static buffers shared between all instances built once on load
 struct FNiagaraSystemStaticBuffers
diff --git a/Engine/Plugins/Runtime/nDisplay/Source/DisplayClusterShaders/Private/Render/DisplayClusterRenderTexture.cpp b/Engine/Plugins/Runtime/nDisplay/Source/DisplayClusterShaders/Private/Render/DisplayClusterRenderTexture.cpp
--- a/Engine/Plugins/Runtime/nDisplay/Source/DisplayClusterShaders/Private/Render/DisplayClusterRenderTexture.cpp
+++ b/Engine/Plugins/Runtime/nDisplay/Source/DisplayClusterShaders/Private/Render/DisplayClusterRenderTexture.cpp
@@ -3,6 +3,7 @@
 #include "Render/DisplayClusterRenderTexture.h"
 
 #include "RenderingThread.h"
+#include "Containers/ResourceArray.h"
 
 //---------------------------------------------------------------------------------------
 // FDisplayClusterResourceBulkData
diff --git a/Engine/Source/Developer/ShaderFormatOpenGL/Private/OpenGLShaderCompiler.cpp b/Engine/Source/Developer/ShaderFormatOpenGL/Private/OpenGLShaderCompiler.cpp
--- a/Engine/Source/Developer/ShaderFormatOpenGL/Private/OpenGLShaderCompiler.cpp
+++ b/Engine/Source/Developer/ShaderFormatOpenGL/Private/OpenGLShaderCompiler.cpp
@@ -3525,22 +3525,6 @@ static bool MoveHashLines(FString& Destination, FString &Source)
 	return bFound;
 }
 
-inline bool OpenGLShaderPlatformSeparable(const GLSLVersion InShaderPlatform)
-{
-	switch (InShaderPlatform)
-	{
-		case GLSL_150_ES3_1:
-			return false;
-
-		case GLSL_ES3_1_ANDROID:
-			return false;
-
-		default:
-			return true;
-		break;
-	}
-}
-
 TSharedPtr<ANSICHAR> FOpenGLFrontend::PrepareCodeForOfflineCompilation(const GLSLVersion ShaderVersion, EShaderFrequency Frequency, const ANSICHAR* InShaderSource) const
 {
 	FString OriginalShaderSource(ANSI_TO_TCHAR(InShaderSource));
diff --git a/Engine/Source/Editor/MaterialEditor/Private/MaterialStats.cpp b/Engine/Source/Editor/MaterialEditor/Private/MaterialStats.cpp
--- a/Engine/Source/Editor/MaterialEditor/Private/MaterialStats.cpp
+++ b/Engine/Source/Editor/MaterialEditor/Private/MaterialStats.cpp
@@ -16,7 +16,7 @@
 #include "Preferences/MaterialStatsOptions.h"
 #include "MaterialEditorSettings.h"
 #include "ShaderCompiler.h"
-
+#include "DataDrivenShaderPlatformInfo.h"
 #include "Modules/ModuleManager.h"
 #include "MessageLogModule.h"
 
diff --git a/Engine/Source/Editor/UnrealEd/Private/MeshPaintRendering.cpp b/Engine/Source/Editor/UnrealEd/Private/MeshPaintRendering.cpp
--- a/Engine/Source/Editor/UnrealEd/Private/MeshPaintRendering.cpp
+++ b/Engine/Source/Editor/UnrealEd/Private/MeshPaintRendering.cpp
@@ -6,7 +6,7 @@
 
 #include "MeshPaintRendering.h"
 #include "ShaderParameters.h"
-#include "RenderResource.h"
+#include "DataDrivenShaderPlatformInfo.h"
 #include "Shader.h"
 #include "StaticBoundShaderState.h"
 #include "RHIStaticStates.h"
diff --git a/Engine/Source/Runtime/Engine/Private/CubemapUnwrapUtils.cpp b/Engine/Source/Runtime/Engine/Private/CubemapUnwrapUtils.cpp
--- a/Engine/Source/Runtime/Engine/Private/CubemapUnwrapUtils.cpp
+++ b/Engine/Source/Runtime/Engine/Private/CubemapUnwrapUtils.cpp
@@ -15,6 +15,7 @@
 #include "Engine/TextureRenderTargetCube.h"
 #include "PipelineStateCache.h"
 #include "RHIStaticStates.h"
+#include "DataDrivenShaderPlatformInfo.h"
 
 IMPLEMENT_SHADER_TYPE(,FCubemapTexturePropertiesVS,TEXT("/Engine/Private/SimpleElementVertexShader.usf"),TEXT("Main"),SF_Vertex);
 IMPLEMENT_SHADER_TYPE(,FCubemapTexturePropertiesPS,TEXT("/Engine/Private/SimpleElementPixelShader.usf"),TEXT("CubemapTextureProperties"),SF_Pixel);
@@ -115,11 +116,26 @@ namespace CubemapHelpers
 	}
 }
 
+FCubemapTexturePropertiesVS::FCubemapTexturePropertiesVS() {}
+
+FCubemapTexturePropertiesVS::FCubemapTexturePropertiesVS(const ShaderMetaType::CompiledShaderInitializerType& Initializer) :
+	FGlobalShader(Initializer)
+{
+	Transform.Bind(Initializer.ParameterMap, TEXT("Transform"), SPF_Mandatory);
+}
+
+bool FCubemapTexturePropertiesVS::ShouldCompilePermutation(const FGlobalShaderPermutationParameters& Parameters)
+{
+	return IsPCPlatform(Parameters.Platform);
+}
+
 void FCubemapTexturePropertiesVS::SetParameters( FRHICommandList& RHICmdList, const FMatrix& TransformValue )
 {
 	SetShaderValue(RHICmdList, RHICmdList.GetBoundVertexShader(), Transform, (FMatrix44f)TransformValue);
 }
 
+FCubemapTexturePropertiesPS::FCubemapTexturePropertiesPS() = default;
+
 FCubemapTexturePropertiesPS::FCubemapTexturePropertiesPS(const ShaderMetaType::CompiledShaderInitializerType& Initializer)
 	: FGlobalShader(Initializer)
 {
@@ -199,6 +215,21 @@ void FMipLevelBatchedElementParameters::BindShaders(FRHICommandList& RHICmdList,
 	PixelShader->SetParameters(RHICmdList, Texture, ColorWeights, MipLevel, SliceIndex, bIsTextureCubeArray, ViewMatrix, bShowLongLatUnwrap, InGamma, bUsePointSampling);
 }
 
+FIESLightProfilePS::FIESLightProfilePS() = default;
+
+FIESLightProfilePS::FIESLightProfilePS(const ShaderMetaType::CompiledShaderInitializerType& Initializer)
+	: FGlobalShader(Initializer)
+{
+	IESTexture.Bind(Initializer.ParameterMap, TEXT("IESTexture"));
+	IESTextureSampler.Bind(Initializer.ParameterMap, TEXT("IESTextureSampler"));
+	BrightnessInLumens.Bind(Initializer.ParameterMap, TEXT("BrightnessInLumens"));
+}
+
+bool FIESLightProfilePS::ShouldCompilePermutation(const FGlobalShaderPermutationParameters& Parameters)
+{
+	return IsFeatureLevelSupported(Parameters.Platform, ERHIFeatureLevel::SM5) && !IsConsolePlatform(Parameters.Platform);
+}
+
 void FIESLightProfilePS::SetParameters( FRHICommandList& RHICmdList, const FTexture* Texture, float InBrightnessInLumens )
 {
 	FRHIPixelShader* ShaderRHI = RHICmdList.GetBoundPixelShader();
diff --git a/Engine/Source/Runtime/Engine/Private/Materials/HLSLMaterialTranslator.cpp b/Engine/Source/Runtime/Engine/Private/Materials/HLSLMaterialTranslator.cpp
--- a/Engine/Source/Runtime/Engine/Private/Materials/HLSLMaterialTranslator.cpp
+++ b/Engine/Source/Runtime/Engine/Private/Materials/HLSLMaterialTranslator.cpp
@@ -1,7 +1,9 @@
 // Copyright Epic Games, Inc. All Rights Reserved.
+
 /*=============================================================================
 	HLSLMaterialTranslator.cpp: Translates material expressions into HLSL code.
 =============================================================================*/
+
 #include "HLSLMaterialTranslator.h"
 #include "VT/VirtualTextureScalability.h"
 #include "Materials/MaterialExpressionCustom.h"
@@ -9,6 +11,7 @@
 #include "MaterialExpressionSettings.h"
 #include "Engine/RendererSettings.h"
 #include "ShaderPlatformCachedIniValue.h"
+#include "DataDrivenShaderPlatformInfo.h"
 #include <functional>
 
 static int32 GetLWCTruncateMode()
diff --git a/Engine/Source/Runtime/Engine/Private/Particles/FXSystem.cpp b/Engine/Source/Runtime/Engine/Private/Particles/FXSystem.cpp
--- a/Engine/Source/Runtime/Engine/Private/Particles/FXSystem.cpp
+++ b/Engine/Source/Runtime/Engine/Private/Particles/FXSystem.cpp
@@ -16,10 +16,26 @@
 #include "SceneUtils.h"
 #include "SceneRendering.h" // needed for STATGROUP_CommandListMarkers
 #include "GPUSortManager.h"
+#include "DataDrivenShaderPlatformInfo.h"
 
 
 TMap<FName, FCreateCustomFXSystemDelegate> FFXSystemInterface::CreateCustomFXDelegates;
 
+bool IsParticleCollisionModeSupported(EShaderPlatform InPlatform, EParticleCollisionShaderMode InCollisionShaderMode)
+{
+	switch (InCollisionShaderMode)
+	{
+	case PCM_None:
+		return true;
+	case PCM_DepthBuffer:
+		return IsFeatureLevelSupported(InPlatform, ERHIFeatureLevel::SM5);
+	case PCM_DistanceField:
+		return IsFeatureLevelSupported(InPlatform, ERHIFeatureLevel::SM5);
+	}
+	check(0);
+	return IsFeatureLevelSupported(InPlatform, ERHIFeatureLevel::SM5);
+}
+
 /*-----------------------------------------------------------------------------
 	External FX system interface.
 -----------------------------------------------------------------------------*/
diff --git a/Engine/Source/Runtime/Engine/Private/Particles/FXSystemPrivate.h b/Engine/Source/Runtime/Engine/Private/Particles/FXSystemPrivate.h
--- a/Engine/Source/Runtime/Engine/Private/Particles/FXSystemPrivate.h
+++ b/Engine/Source/Runtime/Engine/Private/Particles/FXSystemPrivate.h
@@ -57,21 +57,7 @@ enum EParticleCollisionShaderMode
 };
 
 /** Helper function to determine whether the given particle collision shader mode is supported on the given shader platform */
-inline bool IsParticleCollisionModeSupported(EShaderPlatform InPlatform, EParticleCollisionShaderMode InCollisionShaderMode)
-{
-	switch (InCollisionShaderMode)
-	{
-	case PCM_None:
-		return true;
-	case PCM_DepthBuffer:
-		return IsFeatureLevelSupported(InPlatform, ERHIFeatureLevel::SM5);
-	case PCM_DistanceField:
-		return IsFeatureLevelSupported(InPlatform, ERHIFeatureLevel::SM5);
-	}
-	check(0);
-	return IsFeatureLevelSupported(InPlatform, ERHIFeatureLevel::SM5);
-}
-
+extern bool IsParticleCollisionModeSupported(EShaderPlatform InPlatform, EParticleCollisionShaderMode InCollisionShaderMode);
 
 inline EParticleSimulatePhase::Type GetLastParticleSimulationPhase(EShaderPlatform InPlatform)
 {
diff --git a/Engine/Source/Runtime/Engine/Private/PrecomputedVolumetricLightmap.cpp b/Engine/Source/Runtime/Engine/Private/PrecomputedVolumetricLightmap.cpp
--- a/Engine/Source/Runtime/Engine/Private/PrecomputedVolumetricLightmap.cpp
+++ b/Engine/Source/Runtime/Engine/Private/PrecomputedVolumetricLightmap.cpp
@@ -17,6 +17,7 @@
 #include "RenderGraphUtils.h"
 // FIXME: temp fix for ordering issue between WorldContext.World()->InitWorld(); and GShaderCompilingManager->ProcessAsyncResults(false, true); in UnrealEngine.cpp
 #include "ShaderCompiler.h"
+#include "DataDrivenShaderPlatformInfo.h"
 
 DECLARE_MEMORY_STAT(TEXT("Volumetric Lightmap"),STAT_VolumetricLightmapBuildData,STATGROUP_MapBuildData);
 
@@ -817,6 +818,41 @@ FVector ComputeBrickTextureCoordinate(
 	return BrickTextureCoordinate;
 }
 
+bool FRemoveSubLevelBricksCS::ShouldCompilePermutation(const FGlobalShaderPermutationParameters& Parameters)
+{
+	static const auto AllowStaticLightingVar = IConsoleManager::Get().FindTConsoleVariableDataInt(TEXT("r.AllowStaticLighting"));
+
+	return IsFeatureLevelSupported(Parameters.Platform, ERHIFeatureLevel::SM5) && AllowStaticLightingVar->GetValueOnAnyThread() != 0;
+}
+
+bool FCopyResidentBricksCS::ShouldCompilePermutation(const FGlobalShaderPermutationParameters& Parameters)
+{
+	static const auto AllowStaticLightingVar = IConsoleManager::Get().FindTConsoleVariableDataInt(TEXT("r.AllowStaticLighting"));
+
+	return IsFeatureLevelSupported(Parameters.Platform, ERHIFeatureLevel::SM5) && AllowStaticLightingVar->GetValueOnAnyThread() != 0;
+}
+
+bool FCopyResidentBrickSHCoefficientsCS::ShouldCompilePermutation(const FGlobalShaderPermutationParameters& Parameters)
+{
+	static const auto AllowStaticLightingVar = IConsoleManager::Get().FindTConsoleVariableDataInt(TEXT("r.AllowStaticLighting"));
+
+	return IsFeatureLevelSupported(Parameters.Platform, ERHIFeatureLevel::SM5) && AllowStaticLightingVar->GetValueOnAnyThread() != 0;
+}
+
+bool FPatchIndirectionTextureCS::ShouldCompilePermutation(const FGlobalShaderPermutationParameters& Parameters)
+{
+	static const auto AllowStaticLightingVar = IConsoleManager::Get().FindTConsoleVariableDataInt(TEXT("r.AllowStaticLighting"));
+
+	return IsFeatureLevelSupported(Parameters.Platform, ERHIFeatureLevel::SM5) && AllowStaticLightingVar->GetValueOnAnyThread() != 0;
+}
+
+bool FMoveWholeIndirectionTextureCS::ShouldCompilePermutation(const FGlobalShaderPermutationParameters& Parameters)
+{
+	static const auto AllowStaticLightingVar = IConsoleManager::Get().FindTConsoleVariableDataInt(TEXT("r.AllowStaticLighting"));
+
+	return IsFeatureLevelSupported(Parameters.Platform, ERHIFeatureLevel::SM5) && AllowStaticLightingVar->GetValueOnAnyThread() != 0;
+}
+
 IMPLEMENT_GLOBAL_SHADER(FRemoveSubLevelBricksCS, "/Engine/Private/VolumetricLightmapStreaming.usf", "RemoveSubLevelBricksCS", SF_Compute);
 IMPLEMENT_GLOBAL_SHADER(FCopyResidentBricksCS, "/Engine/Private/VolumetricLightmapStreaming.usf", "CopyResidentBricksCS", SF_Compute);
 IMPLEMENT_GLOBAL_SHADER(FCopyResidentBrickSHCoefficientsCS, "/Engine/Private/VolumetricLightmapStreaming.usf", "CopyResidentBrickSHCoefficientsCS", SF_Compute);
diff --git a/Engine/Source/Runtime/Engine/Private/Rendering/SubsurfaceProfile.cpp b/Engine/Source/Runtime/Engine/Private/Rendering/SubsurfaceProfile.cpp
--- a/Engine/Source/Runtime/Engine/Private/Rendering/SubsurfaceProfile.cpp
+++ b/Engine/Source/Runtime/Engine/Private/Rendering/SubsurfaceProfile.cpp
@@ -1,7 +1,7 @@
 // Copyright Epic Games, Inc. All Rights Reserved.
 
 #include "Engine/SubsurfaceProfile.h"
-#include "RenderingThread.h"
+#include "DataDrivenShaderPlatformInfo.h"
 #include "RendererInterface.h"
 #include "Rendering/BurleyNormalizedSSS.h"
 #include "EngineModule.h"
diff --git a/Engine/Source/Runtime/Engine/Private/Rendering/Texture2DResource.cpp b/Engine/Source/Runtime/Engine/Private/Rendering/Texture2DResource.cpp
--- a/Engine/Source/Runtime/Engine/Private/Rendering/Texture2DResource.cpp
+++ b/Engine/Source/Runtime/Engine/Private/Rendering/Texture2DResource.cpp
@@ -10,6 +10,7 @@
 #include "RenderUtils.h"
 #include "SceneTypes.h"
 #include "RenderingThread.h"
+#include "Containers/ResourceArray.h"
 
 // TODO Only adding this setting to allow backwards compatibility to be forced.  The default
 // behavior is to NOT do this.  This variable should be removed in the future.  #ADDED 4.13
diff --git a/Engine/Source/Runtime/Engine/Private/SkeletalRenderGPUSkin.cpp b/Engine/Source/Runtime/Engine/Private/SkeletalRenderGPUSkin.cpp
--- a/Engine/Source/Runtime/Engine/Private/SkeletalRenderGPUSkin.cpp
+++ b/Engine/Source/Runtime/Engine/Private/SkeletalRenderGPUSkin.cpp
@@ -1100,6 +1100,22 @@ void FSkeletalMeshObjectGPUSkin::UpdateMorphVertexBuffer(FRHICommandListImmediat
 
 TArray<float> FSkeletalMeshObjectGPUSkin::FSkeletalMeshObjectLOD::MorphAccumulatedWeightArray;
 
+FGPUMorphUpdateCS::FGPUMorphUpdateCS() = default;
+
+FGPUMorphUpdateCS::FGPUMorphUpdateCS(const ShaderMetaType::CompiledShaderInitializerType& Initializer)
+	: FGlobalShader(Initializer)
+{
+	MorphVertexBufferParameter.Bind(Initializer.ParameterMap, TEXT("MorphVertexBuffer"));
+
+	MorphTargetWeightsParameter.Bind(Initializer.ParameterMap, TEXT("MorphTargetWeights"));
+	MorphTargetBatchOffsetsParameter.Bind(Initializer.ParameterMap, TEXT("MorphTargetBatchOffsets"));
+	MorphTargetGroupOffsetsParameter.Bind(Initializer.ParameterMap, TEXT("MorphTargetGroupOffsets"));
+	PositionScaleParameter.Bind(Initializer.ParameterMap, TEXT("PositionScale"));
+	PrecisionParameter.Bind(Initializer.ParameterMap, TEXT("Precision"));
+
+	MorphDataBufferParameter.Bind(Initializer.ParameterMap, TEXT("MorphDataBuffer"));
+}
+
 void FGPUMorphUpdateCS::SetParameters(FRHICommandList& RHICmdList, const FVector4& LocalScale, const FMorphTargetVertexInfoBuffers& MorphTargetVertexInfoBuffers, FMorphVertexBuffer& MorphVertexBuffer)
 {
 	FRHIComputeShader* CS = RHICmdList.GetBoundComputeShader();
@@ -1133,8 +1149,27 @@ void FGPUMorphUpdateCS::EndAllDispatches(FRHICommandList& RHICmdList)
 	SetUAVParameter(RHICmdList, CS, MorphVertexBufferParameter, nullptr);
 }
 
+bool FGPUMorphUpdateCS::ShouldCompilePermutation(const FGlobalShaderPermutationParameters& Parameters)
+{
+	return IsFeatureLevelSupported(Parameters.Platform, ERHIFeatureLevel::SM5);
+}
+
 IMPLEMENT_SHADER_TYPE(, FGPUMorphUpdateCS, TEXT("/Engine/Private/MorphTargets.usf"), TEXT("GPUMorphUpdateCS"), SF_Compute);
 
+FGPUMorphNormalizeCS::FGPUMorphNormalizeCS() = default;
+
+FGPUMorphNormalizeCS::FGPUMorphNormalizeCS(const ShaderMetaType::CompiledShaderInitializerType& Initializer)
+	: FGlobalShader(Initializer)
+{
+	MorphVertexBufferParameter.Bind(Initializer.ParameterMap, TEXT("MorphVertexBuffer"));
+	PositionScaleParameter.Bind(Initializer.ParameterMap, TEXT("PositionScale"));
+}
+
+bool FGPUMorphNormalizeCS::ShouldCompilePermutation(const FGlobalShaderPermutationParameters& Parameters)
+{
+	return IsFeatureLevelSupported(Parameters.Platform, ERHIFeatureLevel::SM5);
+}
+
 void FGPUMorphNormalizeCS::SetParameters(FRHICommandList& RHICmdList, const FVector4& InvLocalScale, const FMorphTargetVertexInfoBuffers& MorphTargetVertexInfoBuffers, FMorphVertexBuffer& MorphVertexBuffer)
 {
 	FRHIComputeShader* CS = RHICmdList.GetBoundComputeShader();
diff --git a/Engine/Source/Runtime/Engine/Private/SkeletalRenderGPUSkin.h b/Engine/Source/Runtime/Engine/Private/SkeletalRenderGPUSkin.h
--- a/Engine/Source/Runtime/Engine/Private/SkeletalRenderGPUSkin.h
+++ b/Engine/Source/Runtime/Engine/Private/SkeletalRenderGPUSkin.h
@@ -696,21 +696,8 @@ class FGPUMorphUpdateCS : public FGlobalShader
 public:
 	DECLARE_SHADER_TYPE(FGPUMorphUpdateCS, Global);
 
-	FGPUMorphUpdateCS() {}
-
-	FGPUMorphUpdateCS(const ShaderMetaType::CompiledShaderInitializerType& Initializer)
-		: FGlobalShader(Initializer)
-	{
-		MorphVertexBufferParameter.Bind(Initializer.ParameterMap, TEXT("MorphVertexBuffer"));
-
-		MorphTargetWeightsParameter.Bind(Initializer.ParameterMap, TEXT("MorphTargetWeights"));
-		MorphTargetBatchOffsetsParameter.Bind(Initializer.ParameterMap, TEXT("MorphTargetBatchOffsets"));
-		MorphTargetGroupOffsetsParameter.Bind(Initializer.ParameterMap, TEXT("MorphTargetGroupOffsets"));
-		PositionScaleParameter.Bind(Initializer.ParameterMap, TEXT("PositionScale"));
-		PrecisionParameter.Bind(Initializer.ParameterMap, TEXT("Precision"));
-
-		MorphDataBufferParameter.Bind(Initializer.ParameterMap, TEXT("MorphDataBuffer"));
-	}
+	FGPUMorphUpdateCS();
+	FGPUMorphUpdateCS(const ShaderMetaType::CompiledShaderInitializerType& Initializer);
 
 	static const uint32 MorphTargetDispatchBatchSize = 128;
 
@@ -720,10 +707,7 @@ public:
 	void Dispatch(FRHICommandList& RHICmdList, uint32 Size);
 	void EndAllDispatches(FRHICommandList& RHICmdList);
 
-	static bool ShouldCompilePermutation(const FGlobalShaderPermutationParameters& Parameters)
-	{
-		return IsFeatureLevelSupported(Parameters.Platform, ERHIFeatureLevel::SM5);
-	}
+	static bool ShouldCompilePermutation(const FGlobalShaderPermutationParameters& Parameters);
 
 protected:
 	LAYOUT_FIELD(FShaderResourceParameter, MorphVertexBufferParameter);
@@ -743,19 +727,10 @@ class FGPUMorphNormalizeCS : public FGlobalShader
 public:
 	DECLARE_SHADER_TYPE(FGPUMorphNormalizeCS, Global);
 
-	FGPUMorphNormalizeCS() {}
-
-	FGPUMorphNormalizeCS(const ShaderMetaType::CompiledShaderInitializerType& Initializer)
-		: FGlobalShader(Initializer)
-	{
-		MorphVertexBufferParameter.Bind(Initializer.ParameterMap, TEXT("MorphVertexBuffer"));
-		PositionScaleParameter.Bind(Initializer.ParameterMap, TEXT("PositionScale"));
-	}
+	FGPUMorphNormalizeCS();
+	FGPUMorphNormalizeCS(const ShaderMetaType::CompiledShaderInitializerType& Initializer);
 
-	static bool ShouldCompilePermutation(const FGlobalShaderPermutationParameters& Parameters)
-	{
-		return IsFeatureLevelSupported(Parameters.Platform, ERHIFeatureLevel::SM5);
-	}
+	static bool ShouldCompilePermutation(const FGlobalShaderPermutationParameters& Parameters);
 
 	void SetParameters(FRHICommandList& RHICmdList, const FVector4& LocalScale, const FMorphTargetVertexInfoBuffers& MorphTargetVertexInfoBuffers, FMorphVertexBuffer& MorphVertexBuffer);
 
diff --git a/Engine/Source/Runtime/Engine/Private/VolumeRendering.cpp b/Engine/Source/Runtime/Engine/Private/VolumeRendering.cpp
--- a/Engine/Source/Runtime/Engine/Private/VolumeRendering.cpp
+++ b/Engine/Source/Runtime/Engine/Private/VolumeRendering.cpp
@@ -7,10 +7,68 @@
 #include "VolumeRendering.h"
 #include "ScreenRendering.h"
 #include "RHIStaticStates.h"
+#include "DataDrivenShaderPlatformInfo.h"
+
+FWriteToSliceVS::FWriteToSliceVS() = default;
+
+FWriteToSliceVS::FWriteToSliceVS(const ShaderMetaType::CompiledShaderInitializerType& Initializer)
+	: FGlobalShader(Initializer)
+{
+	UVScaleBias.Bind(Initializer.ParameterMap, TEXT("UVScaleBias"));
+	MinZ.Bind(Initializer.ParameterMap, TEXT("MinZ"));
+}
+
+bool FWriteToSliceVS::ShouldCompilePermutation(const FGlobalShaderPermutationParameters& Parameters)
+{
+	return IsFeatureLevelSupported(Parameters.Platform, ERHIFeatureLevel::SM5);
+}
+
+void FWriteToSliceVS::ModifyCompilationEnvironment(const FGlobalShaderPermutationParameters& Parameters, FShaderCompilerEnvironment& OutEnvironment)
+{
+	FGlobalShader::ModifyCompilationEnvironment(Parameters, OutEnvironment);
+	OutEnvironment.CompilerFlags.Add(CFLAG_VertexToGeometryShader);
+}
 
-IMPLEMENT_SHADER_TYPE(,FWriteToSliceGS,TEXT("/Engine/Private/TranslucentLightingShaders.usf"),TEXT("WriteToSliceMainGS"),SF_Geometry);
 IMPLEMENT_SHADER_TYPE(,FWriteToSliceVS,TEXT("/Engine/Private/TranslucentLightingShaders.usf"),TEXT("WriteToSliceMainVS"),SF_Vertex);
 
+
+FWriteToSliceGS::FWriteToSliceGS() = default;
+
+FWriteToSliceGS::FWriteToSliceGS(const ShaderMetaType::CompiledShaderInitializerType& Initializer)
+	: FGlobalShader(Initializer)
+{
+	MinZ.Bind(Initializer.ParameterMap, TEXT("MinZ"));
+}
+
+bool FWriteToSliceGS::ShouldCompilePermutation(const FGlobalShaderPermutationParameters& Parameters)
+{
+	return IsFeatureLevelSupported(Parameters.Platform, ERHIFeatureLevel::SM5) && RHISupportsGeometryShaders(Parameters.Platform);
+}
+
+IMPLEMENT_SHADER_TYPE(, FWriteToSliceGS, TEXT("/Engine/Private/TranslucentLightingShaders.usf"), TEXT("WriteToSliceMainGS"), SF_Geometry);
+
+void FVolumeRasterizeVertexBuffer::InitRHI()
+{
+	// Used as a non-indexed triangle strip, so 4 vertices per quad
+	const uint32 Size = 4 * sizeof(FScreenVertex);
+	FRHIResourceCreateInfo CreateInfo(TEXT("FVolumeRasterizeVertexBuffer"));
+	VertexBufferRHI = RHICreateBuffer(Size, BUF_Static | BUF_VertexBuffer, 0, ERHIAccess::VertexOrIndexBuffer, CreateInfo);
+	FScreenVertex* DestVertex = (FScreenVertex*)RHILockBuffer(VertexBufferRHI, 0, Size, RLM_WriteOnly);
+
+	// Setup a full - render target quad
+	// A viewport and UVScaleBias will be used to implement rendering to a sub region
+	DestVertex[0].Position = FVector2f(1, -GProjectionSignY);
+	DestVertex[0].UV = FVector2f(1, 1);
+	DestVertex[1].Position = FVector2f(1, GProjectionSignY);
+	DestVertex[1].UV = FVector2f(1, 0);
+	DestVertex[2].Position = FVector2f(-1, -GProjectionSignY);
+	DestVertex[2].UV = FVector2f(0, 1);
+	DestVertex[3].Position = FVector2f(-1, GProjectionSignY);
+	DestVertex[3].UV = FVector2f(0, 0);
+
+	RHIUnlockBuffer(VertexBufferRHI);
+}
+
 TGlobalResource<FVolumeRasterizeVertexBuffer> GVolumeRasterizeVertexBuffer;
 
 /** Draws a quad per volume texture slice to the subregion of the volume texture specified. */
diff --git a/Engine/Source/Runtime/Engine/Public/CubemapUnwrapUtils.h b/Engine/Source/Runtime/Engine/Public/CubemapUnwrapUtils.h
--- a/Engine/Source/Runtime/Engine/Public/CubemapUnwrapUtils.h
+++ b/Engine/Source/Runtime/Engine/Public/CubemapUnwrapUtils.h
@@ -47,15 +47,10 @@ class FCubemapTexturePropertiesVS : public FGlobalShader
 {
 	DECLARE_SHADER_TYPE(FCubemapTexturePropertiesVS,Global);
 public:
+	FCubemapTexturePropertiesVS();
+	FCubemapTexturePropertiesVS(const ShaderMetaType::CompiledShaderInitializerType& Initializer);
 
-	static bool ShouldCompilePermutation(const FGlobalShaderPermutationParameters& Parameters) { return IsPCPlatform(Parameters.Platform);}
-
-	FCubemapTexturePropertiesVS(const ShaderMetaType::CompiledShaderInitializerType& Initializer):
-		FGlobalShader(Initializer)
-	{
-		Transform.Bind(Initializer.ParameterMap,TEXT("Transform"), SPF_Mandatory);
-	}
-	FCubemapTexturePropertiesVS() {}
+	static bool ShouldCompilePermutation(const FGlobalShaderPermutationParameters& Parameters);
 
 	void SetParameters(FRHICommandList& RHICmdList, const FMatrix& TransformValue);
 
@@ -75,7 +70,7 @@ class FCubemapTexturePropertiesPS : public FGlobalShader
 	using FPermutationDomain = TShaderPermutationDomain<FHDROutput, FCubeArray>;
 
 public:
-	FCubemapTexturePropertiesPS() = default;
+	FCubemapTexturePropertiesPS();
 	FCubemapTexturePropertiesPS(const ShaderMetaType::CompiledShaderInitializerType& Initializer);
 
 	static bool ShouldCompilePermutation(const FGlobalShaderPermutationParameters& Parameters);
@@ -134,11 +129,10 @@ class FIESLightProfilePS : public FGlobalShader
 {
 	DECLARE_SHADER_TYPE(FIESLightProfilePS,Global);
 public:
+	FIESLightProfilePS();
+	FIESLightProfilePS(const ShaderMetaType::CompiledShaderInitializerType& Initializer);
 
-	static bool ShouldCompilePermutation(const FGlobalShaderPermutationParameters& Parameters)
-	{
-		return IsFeatureLevelSupported(Parameters.Platform, ERHIFeatureLevel::SM5) && !IsConsolePlatform(Parameters.Platform);
-	}
+	static bool ShouldCompilePermutation(const FGlobalShaderPermutationParameters& Parameters);
 
 	static void ModifyCompilationEnvironment(const FGlobalShaderPermutationParameters& Parameters, FShaderCompilerEnvironment& OutEnvironment)
 	{
@@ -147,16 +141,6 @@ public:
 		OutEnvironment.SetDefine(TEXT("USE_IES_STANDALONE_TEXTURE"), 1);
 	}
 
-	FIESLightProfilePS() {}
-
-	FIESLightProfilePS(const ShaderMetaType::CompiledShaderInitializerType& Initializer)
-		: FGlobalShader(Initializer)
-	{
-		IESTexture.Bind(Initializer.ParameterMap,TEXT("IESTexture"));
-		IESTextureSampler.Bind(Initializer.ParameterMap,TEXT("IESTextureSampler"));
-		BrightnessInLumens.Bind(Initializer.ParameterMap,TEXT("BrightnessInLumens"));
-	}
-
 	void SetParameters(FRHICommandList& RHICmdList, const FTexture* Texture, float InBrightnessInLumens);
 
 private:
diff --git a/Engine/Source/Runtime/Engine/Public/PrecomputedVolumetricLightmap.h b/Engine/Source/Runtime/Engine/Public/PrecomputedVolumetricLightmap.h
--- a/Engine/Source/Runtime/Engine/Public/PrecomputedVolumetricLightmap.h
+++ b/Engine/Source/Runtime/Engine/Public/PrecomputedVolumetricLightmap.h
@@ -410,12 +410,7 @@ class FRemoveSubLevelBricksCS : public FGlobalShader
 	DECLARE_GLOBAL_SHADER(FRemoveSubLevelBricksCS);
 	SHADER_USE_PARAMETER_STRUCT(FRemoveSubLevelBricksCS, FGlobalShader);
 
-	static bool ShouldCompilePermutation(const FGlobalShaderPermutationParameters& Parameters)
-	{
-		static const auto AllowStaticLightingVar = IConsoleManager::Get().FindTConsoleVariableDataInt(TEXT("r.AllowStaticLighting"));
-
-		return IsFeatureLevelSupported(Parameters.Platform, ERHIFeatureLevel::SM5) && AllowStaticLightingVar->GetValueOnAnyThread() != 0;
-	}
+	static bool ShouldCompilePermutation(const FGlobalShaderPermutationParameters& Parameters);
 
 	BEGIN_SHADER_PARAMETER_STRUCT(FParameters, )
 		SHADER_PARAMETER(uint32, NumBricks)
@@ -434,12 +429,7 @@ class FCopyResidentBricksCS : public FGlobalShader
 	class FHasSkyBentNormal : SHADER_PERMUTATION_BOOL("HAS_SKY_BENT_NORMAL");
 	using FPermutationDomain = TShaderPermutationDomain<FHasSkyBentNormal>;
 
-	static bool ShouldCompilePermutation(const FGlobalShaderPermutationParameters& Parameters)
-	{
-		static const auto AllowStaticLightingVar = IConsoleManager::Get().FindTConsoleVariableDataInt(TEXT("r.AllowStaticLighting"));
-
-		return IsFeatureLevelSupported(Parameters.Platform, ERHIFeatureLevel::SM5) && AllowStaticLightingVar->GetValueOnAnyThread() != 0;
-	}
+	static bool ShouldCompilePermutation(const FGlobalShaderPermutationParameters& Parameters);
 
 	BEGIN_SHADER_PARAMETER_STRUCT(FParameters, )
 		SHADER_PARAMETER(uint32, StartPosInOldVolume)
@@ -459,12 +449,7 @@ class FCopyResidentBrickSHCoefficientsCS : public FGlobalShader
 	DECLARE_GLOBAL_SHADER(FCopyResidentBrickSHCoefficientsCS);
 	SHADER_USE_PARAMETER_STRUCT(FCopyResidentBrickSHCoefficientsCS, FGlobalShader);
 
-	static bool ShouldCompilePermutation(const FGlobalShaderPermutationParameters& Parameters)
-	{
-		static const auto AllowStaticLightingVar = IConsoleManager::Get().FindTConsoleVariableDataInt(TEXT("r.AllowStaticLighting"));
-
-		return IsFeatureLevelSupported(Parameters.Platform, ERHIFeatureLevel::SM5) && AllowStaticLightingVar->GetValueOnAnyThread() != 0;
-	}
+	static bool ShouldCompilePermutation(const FGlobalShaderPermutationParameters& Parameters);
 
 	BEGIN_SHADER_PARAMETER_STRUCT(FParameters, )
 		SHADER_PARAMETER(uint32, StartPosInOldVolume)
@@ -479,12 +464,7 @@ class FPatchIndirectionTextureCS : public FGlobalShader
 	DECLARE_GLOBAL_SHADER(FPatchIndirectionTextureCS);
 	SHADER_USE_PARAMETER_STRUCT(FPatchIndirectionTextureCS, FGlobalShader);
 
-	static bool ShouldCompilePermutation(const FGlobalShaderPermutationParameters& Parameters)
-	{
-		static const auto AllowStaticLightingVar = IConsoleManager::Get().FindTConsoleVariableDataInt(TEXT("r.AllowStaticLighting"));
-
-		return IsFeatureLevelSupported(Parameters.Platform, ERHIFeatureLevel::SM5) && AllowStaticLightingVar->GetValueOnAnyThread() != 0;
-	}
+	static bool ShouldCompilePermutation(const FGlobalShaderPermutationParameters& Parameters);
 
 	BEGIN_SHADER_PARAMETER_STRUCT(FParameters, )
 		SHADER_PARAMETER(uint32, NumBricks)
@@ -499,12 +479,7 @@ class FMoveWholeIndirectionTextureCS : public FGlobalShader
 	DECLARE_GLOBAL_SHADER(FMoveWholeIndirectionTextureCS);
 	SHADER_USE_PARAMETER_STRUCT(FMoveWholeIndirectionTextureCS, FGlobalShader);
 
-	static bool ShouldCompilePermutation(const FGlobalShaderPermutationParameters& Parameters)
-	{
-		static const auto AllowStaticLightingVar = IConsoleManager::Get().FindTConsoleVariableDataInt(TEXT("r.AllowStaticLighting"));
-
-		return IsFeatureLevelSupported(Parameters.Platform, ERHIFeatureLevel::SM5) && AllowStaticLightingVar->GetValueOnAnyThread() != 0;
-	}
+	static bool ShouldCompilePermutation(const FGlobalShaderPermutationParameters& Parameters);
 
 	BEGIN_SHADER_PARAMETER_STRUCT(FParameters, )
 		SHADER_PARAMETER(uint32, NumBricks)
diff --git a/Engine/Source/Runtime/Launch/Private/LaunchEngineLoop.cpp b/Engine/Source/Runtime/Launch/Private/LaunchEngineLoop.cpp
--- a/Engine/Source/Runtime/Launch/Private/LaunchEngineLoop.cpp
+++ b/Engine/Source/Runtime/Launch/Private/LaunchEngineLoop.cpp
@@ -45,6 +45,7 @@
 #if WITH_ENGINE
 #include "HAL/PlatformSplash.h"
 #include "StereoRenderUtils.h"
+#include "DataDrivenShaderPlatformInfo.h"
 #endif
 #if WITH_APPLICATION_CORE
 #include "HAL/PlatformApplicationMisc.h"
@@ -2831,9 +2832,6 @@ PRAGMA_ENABLE_DEPRECATION_WARNINGS
 		SCOPED_BOOT_TIMING("RHIInit");
 		// Initialize the RHI.
 		RHIInit(bHasEditorToken);
-#if WITH_EDITOR
-		FGenericDataDrivenShaderPlatformInfo::UpdatePreviewPlatforms();
-#endif
 	}
 
 	{
diff --git a/Engine/Source/Runtime/NullDrv/Private/NullRHI.cpp b/Engine/Source/Runtime/NullDrv/Private/NullRHI.cpp
--- a/Engine/Source/Runtime/NullDrv/Private/NullRHI.cpp
+++ b/Engine/Source/Runtime/NullDrv/Private/NullRHI.cpp
@@ -5,7 +5,7 @@
 #include "Containers/List.h"
 #include "RenderResource.h"
 #include "RenderUtils.h"
-
+#include "DataDrivenShaderPlatformInfo.h"
 
 FNullDynamicRHI::FNullDynamicRHI()
 {
diff --git a/Engine/Source/Runtime/OpenGLDrv/Private/OpenGLDrvPrivate.h b/Engine/Source/Runtime/OpenGLDrv/Private/OpenGLDrvPrivate.h
--- a/Engine/Source/Runtime/OpenGLDrv/Private/OpenGLDrvPrivate.h
+++ b/Engine/Source/Runtime/OpenGLDrv/Private/OpenGLDrvPrivate.h
@@ -446,21 +446,3 @@ inline uint32 CalcDynamicBufferSize(uint32 Size)
 void InitDefaultGLContextState(void);
 
 extern bool GUseEmulatedUniformBuffers;
-
-inline bool OpenGLShaderPlatformSeparable(const EShaderPlatform InShaderPlatform)
-{
-	switch (InShaderPlatform)
-	{
-		case SP_OPENGL_PCES3_1:
-			return true;
-
-		case SP_OPENGL_ES3_1_ANDROID:
-			return false;
-		default:
-			check(IsOpenGLPlatform(InShaderPlatform));
-			checkf(false, TEXT("invalid shader platform (%d)"), int(InShaderPlatform));
-			return true;
-			break;
-	}
-}
-
diff --git a/Engine/Source/Runtime/RHI/Private/DynamicRHI.cpp b/Engine/Source/Runtime/RHI/Private/DynamicRHI.cpp
--- a/Engine/Source/Runtime/RHI/Private/DynamicRHI.cpp
+++ b/Engine/Source/Runtime/RHI/Private/DynamicRHI.cpp
@@ -389,6 +389,10 @@ void RHIInit(bool bHasEditorToken)
 
 		check(GDynamicRHI);
 	}
+
+#if WITH_EDITOR
+	FGenericDataDrivenShaderPlatformInfo::UpdatePreviewPlatforms();
+#endif
 }
 
 void RHIPostInit(const TArray<uint32>& InPixelFormatByteWidth)
diff --git a/Engine/Source/Runtime/RenderCore/Private/GlobalRenderResources.cpp b/Engine/Source/Runtime/RenderCore/Private/GlobalRenderResources.cpp
--- a/Engine/Source/Runtime/RenderCore/Private/GlobalRenderResources.cpp
+++ b/Engine/Source/Runtime/RenderCore/Private/GlobalRenderResources.cpp
@@ -1,8 +1,8 @@
 // Copyright Epic Games, Inc. All Rights Reserved.
 
 #include "GlobalRenderResources.h"
-//#include "RenderUtils.h"
 #include "RenderGraphUtils.h"
+#include "Containers/ResourceArray.h"
 
 // The maximum number of transient vertex buffer bytes to allocate before we start panic logging who is doing the allocations
 int32 GMaxVertexBytesAllocatedPerFrame = 32 * 1024 * 1024;
diff --git a/Engine/Source/Runtime/RenderCore/Private/PixelShaderUtils.cpp b/Engine/Source/Runtime/RenderCore/Private/PixelShaderUtils.cpp
--- a/Engine/Source/Runtime/RenderCore/Private/PixelShaderUtils.cpp
+++ b/Engine/Source/Runtime/RenderCore/Private/PixelShaderUtils.cpp
@@ -5,7 +5,9 @@
 =============================================================================*/
 
 #include "PixelShaderUtils.h"
+
 #include "CommonRenderResources.h"
+#include "DataDrivenShaderPlatformInfo.h"
 #include "RenderGraph.h"
 
 IMPLEMENT_SHADER_TYPE(, FPixelShaderUtils::FRasterizeToRectsVS, TEXT("/Engine/Private/RenderGraphUtilities.usf"), TEXT("RasterizeToRectsVS"), SF_Vertex);
diff --git a/Engine/Source/Runtime/RenderCore/Private/RenderGraphPrivate.cpp b/Engine/Source/Runtime/RenderCore/Private/RenderGraphPrivate.cpp
--- a/Engine/Source/Runtime/RenderCore/Private/RenderGraphPrivate.cpp
+++ b/Engine/Source/Runtime/RenderCore/Private/RenderGraphPrivate.cpp
@@ -1,6 +1,7 @@
 // Copyright Epic Games, Inc. All Rights Reserved.
 
 #include "RenderGraphPrivate.h"
+#include "DataDrivenShaderPlatformInfo.h"
 
 #if RDG_ENABLE_DEBUG
 
@@ -608,3 +609,37 @@ void InitRenderGraph()
 	}
 #endif
 }
+
+bool IsParallelExecuteEnabled()
+{
+	return GRDGParallelExecute > 0
+		&& !GRHICommandList.Bypass()
+		&& !IsImmediateMode()
+		&& !GRDGDebug
+		&& !GRDGTransitionLog
+		&& !IsMobilePlatform(GMaxRHIShaderPlatform)
+		&& GRHISupportsMultithreadedShaderCreation
+#if WITH_DUMPGPU
+		&& !IsDumpingRDGResources()
+#endif
+		// Only run parallel RDG if we have a rendering thread.
+		&& IsInActualRenderingThread()
+		;
+}
+
+bool IsParallelSetupEnabled()
+{
+	return GRDGParallelSetup > 0
+		&& !GRHICommandList.Bypass()
+		&& !IsImmediateMode()
+		&& !GRDGDebug
+		&& !GRDGTransitionLog
+		&& !IsMobilePlatform(GMaxRHIShaderPlatform)
+		&& GRHISupportsMultithreadedShaderCreation
+#if WITH_DUMPGPU
+		&& !IsDumpingRDGResources()
+#endif
+		// Only run parallel RDG if we have a rendering thread.
+		&& IsInActualRenderingThread()
+		;
+}
diff --git a/Engine/Source/Runtime/RenderCore/Private/RenderGraphPrivate.h b/Engine/Source/Runtime/RenderCore/Private/RenderGraphPrivate.h
--- a/Engine/Source/Runtime/RenderCore/Private/RenderGraphPrivate.h
+++ b/Engine/Source/Runtime/RenderCore/Private/RenderGraphPrivate.h
@@ -259,39 +259,8 @@ FORCEINLINE bool IsAsyncComputeSupported()
 
 bool IsDumpingRDGResources();
 
-FORCEINLINE bool IsParallelExecuteEnabled()
-{
-	return GRDGParallelExecute > 0
-		&& !GRHICommandList.Bypass()
-		&& !IsImmediateMode()
-		&& !GRDGDebug
-		&& !GRDGTransitionLog
-		&& !IsMobilePlatform(GMaxRHIShaderPlatform)
-		&& GRHISupportsMultithreadedShaderCreation
-#if WITH_DUMPGPU
-		&& !IsDumpingRDGResources()
-#endif
-		// Only run parallel RDG if we have a rendering thread.
-		&& IsInActualRenderingThread()
-		;
-}
-
-FORCEINLINE bool IsParallelSetupEnabled()
-{
-	return GRDGParallelSetup > 0
-		&& !GRHICommandList.Bypass()
-		&& !IsImmediateMode()
-		&& !GRDGDebug
-		&& !GRDGTransitionLog
-		&& !IsMobilePlatform(GMaxRHIShaderPlatform)
-		&& GRHISupportsMultithreadedShaderCreation
-#if WITH_DUMPGPU
-		&& !IsDumpingRDGResources()
-#endif
-		// Only run parallel RDG if we have a rendering thread.
-		&& IsInActualRenderingThread()
-		;
-}
+extern bool IsParallelExecuteEnabled();
+extern bool IsParallelSetupEnabled();
 
 template <typename ResourceRegistryType, typename FunctionType>
 inline void EnumerateExtendedLifetimeResources(ResourceRegistryType& Registry, FunctionType Function)
diff --git a/Engine/Source/Runtime/RenderCore/Private/RenderGraphUtils.cpp b/Engine/Source/Runtime/RenderCore/Private/RenderGraphUtils.cpp
--- a/Engine/Source/Runtime/RenderCore/Private/RenderGraphUtils.cpp
+++ b/Engine/Source/Runtime/RenderCore/Private/RenderGraphUtils.cpp
@@ -1,15 +1,18 @@
 // Copyright Epic Games, Inc. All Rights Reserved.
 
 #include "RenderGraphUtils.h"
+
 #include "ClearQuad.h"
 #include "ClearReplacementShaders.h"
-#include "ShaderParameterUtils.h"
-#include "RenderTargetPool.h"
-#include <initializer_list>
+#include "DataDrivenShaderPlatformInfo.h"
 #include "GlobalShader.h"
 #include "PixelShaderUtils.h"
 #include "RenderGraphResourcePool.h"
+#include "RenderTargetPool.h"
 #include "RHIGPUReadback.h"
+#include "ShaderParameterUtils.h"
+
+#include <initializer_list>
 
 void ClearUnusedGraphResourcesImpl(
 	const FShaderParameterBindings& ShaderBindings,
diff --git a/Engine/Source/Runtime/RenderCore/Private/RenderUtils.cpp b/Engine/Source/Runtime/RenderCore/Private/RenderUtils.cpp
--- a/Engine/Source/Runtime/RenderCore/Private/RenderUtils.cpp
+++ b/Engine/Source/Runtime/RenderCore/Private/RenderUtils.cpp
@@ -1,21 +1,21 @@
 // Copyright Epic Games, Inc. All Rights Reserved.
 
 #include "RenderUtils.h"
-#include "Containers/ResourceArray.h"
+
 #include "Containers/DynamicRHIResourceArray.h"
-#include "RenderResource.h"
-#include "RHIStaticStates.h"
-#include "PipelineStateCache.h"
+#include "DataDrivenShaderPlatformInfo.h"
+#include "HAL/ConsoleManager.h"
 #include "Misc/ConfigCacheIni.h"
-#include "RenderGraphResourcePool.h"
 #include "Misc/DataDrivenPlatformInfoRegistry.h"
-#include "RenderCore.h"
 #include "PackedNormal.h"
+#include "PipelineStateCache.h"
+#include "RenderResource.h"
+#include "RHI.h"
 #include "Shader.h"
 #include "ShaderPlatformCachedIniValue.h"
+#include "DataDrivenShaderPlatformInfo.h"
 
 #if WITH_EDITOR
-#include "Misc/CoreMisc.h"
 #include "Interfaces/ITargetPlatform.h"
 #include "Interfaces/ITargetPlatformManagerModule.h"
 #include "RHIShaderFormatDefinitions.inl"
@@ -562,6 +562,8 @@ RENDERCORE_API uint32 GetPlatformShadingModelsMask(const FStaticShaderPlatform P
 	return 0xFFFFFFFF;
 }
 
+typedef TBitArray<TInlineAllocator<EShaderPlatform::SP_NumPlatforms / 8>> ShaderPlatformMaskType;
+
 RENDERCORE_API ShaderPlatformMaskType GMobileAmbientOcclusionPlatformMask;
 
 RENDERCORE_API bool IsMobileAmbientOcclusionEnabled(const FStaticShaderPlatform Platform)
diff --git a/Engine/Source/Runtime/Renderer/Private/DebugViewModeRendering.cpp b/Engine/Source/Runtime/Renderer/Private/DebugViewModeRendering.cpp
--- a/Engine/Source/Runtime/Renderer/Private/DebugViewModeRendering.cpp
+++ b/Engine/Source/Runtime/Renderer/Private/DebugViewModeRendering.cpp
@@ -149,13 +149,52 @@ TRDGUniformBufferRef<FDebugViewModePassUniformParameters> CreateDebugViewModePas
 	return GraphBuilder.CreateUniformBuffer(UniformBufferParameters);
 }
 
-IMPLEMENT_MATERIAL_SHADER_TYPE(,FDebugViewModeVS,TEXT("/Engine/Private/DebugViewModeVertexShader.usf"),TEXT("Main"),SF_Vertex);	
+FDebugViewModeVS::FDebugViewModeVS() = default;
+FDebugViewModeVS::FDebugViewModeVS(const FMeshMaterialShaderType::CompiledShaderInitializerType& Initializer)
+	: FMeshMaterialShader(Initializer)
+{}
 
 bool FDebugViewModeVS::ShouldCompilePermutation(const FMeshMaterialShaderPermutationParameters& Parameters)
 {
 	return AllowDebugViewVSDSHS(Parameters.Platform) && EnumHasAllFlags(Parameters.Flags, EShaderPermutationFlags::HasEditorOnlyData);
 }
 
+void FDebugViewModeVS::GetShaderBindings(
+	const FScene* Scene,
+	ERHIFeatureLevel::Type FeatureLevel,
+	const FPrimitiveSceneProxy* PrimitiveSceneProxy,
+	const FMaterialRenderProxy& MaterialRenderProxy,
+	const FMaterial& Material,
+	const FMeshPassProcessorRenderState& DrawRenderState,
+	const FDebugViewModeShaderElementData& ShaderElementData,
+	FMeshDrawSingleShaderBindings& ShaderBindings) const
+{
+	FMeshMaterialShader::GetShaderBindings(Scene, FeatureLevel, PrimitiveSceneProxy, MaterialRenderProxy, Material, DrawRenderState, ShaderElementData, ShaderBindings);
+}
+
+void FDebugViewModeVS::SetCommonDefinitions(const FMaterialShaderPermutationParameters& Parameters, FShaderCompilerEnvironment& OutEnvironment)
+{
+	// SM4 has less input interpolants. Also instanced meshes use more interpolants.
+	if (Parameters.MaterialParameters.bIsDefaultMaterial || (IsFeatureLevelSupported(Parameters.Platform, ERHIFeatureLevel::SM5) && !Parameters.MaterialParameters.bIsUsedWithInstancedStaticMeshes))
+	{	// Force the default material to pass enough texcoords to the pixel shaders (even though not using them).
+		// This is required to allow material shaders to have access to the sampled coords.
+		OutEnvironment.SetDefine(TEXT("MIN_MATERIAL_TEXCOORDS"), (uint32)4);
+	}
+	else // Otherwise still pass at minimum amount to have debug shader using a texcoord to work (material might not use any).
+	{
+		OutEnvironment.SetDefine(TEXT("MIN_MATERIAL_TEXCOORDS"), (uint32)2);
+	}
+
+}
+
+void FDebugViewModeVS::ModifyCompilationEnvironment(const FMaterialShaderPermutationParameters& Parameters, FShaderCompilerEnvironment& OutEnvironment)
+{
+	SetCommonDefinitions(Parameters, OutEnvironment);
+	FMeshMaterialShader::ModifyCompilationEnvironment(Parameters, OutEnvironment);
+}
+
+IMPLEMENT_MATERIAL_SHADER_TYPE(,FDebugViewModeVS,TEXT("/Engine/Private/DebugViewModeVertexShader.usf"),TEXT("Main"),SF_Vertex);
+
 BEGIN_SHADER_PARAMETER_STRUCT(FDebugViewModePassParameters, )
 	SHADER_PARAMETER_STRUCT_INCLUDE(FViewShaderParameters, View)
 	SHADER_PARAMETER_RDG_UNIFORM_BUFFER(FDebugViewModePassUniformParameters, Pass)
@@ -195,11 +234,54 @@ void RenderDebugViewMode(FRDGBuilder& GraphBuilder, TArrayView<FViewInfo> Views,
 	}
 }
 
+FDebugViewModePS::FDebugViewModePS() = default;
+FDebugViewModePS::FDebugViewModePS(const FMeshMaterialShaderType::CompiledShaderInitializerType& Initializer)
+	: FMeshMaterialShader(Initializer)
+{
+	OneOverCPUTexCoordScalesParameter.Bind(Initializer.ParameterMap, TEXT("OneOverCPUTexCoordScales"));
+	TexCoordIndicesParameter.Bind(Initializer.ParameterMap, TEXT("TexCoordIndices"));
+	CPUTexelFactorParameter.Bind(Initializer.ParameterMap, TEXT("CPUTexelFactor"));
+	NormalizedComplexity.Bind(Initializer.ParameterMap, TEXT("NormalizedComplexity"));
+	AnalysisParamsParameter.Bind(Initializer.ParameterMap, TEXT("AnalysisParams"));
+	PrimitiveAlphaParameter.Bind(Initializer.ParameterMap, TEXT("PrimitiveAlpha"));
+	TexCoordAnalysisIndexParameter.Bind(Initializer.ParameterMap, TEXT("TexCoordAnalysisIndex"));
+	CPULogDistanceParameter.Bind(Initializer.ParameterMap, TEXT("CPULogDistance"));
+	ShowQuadOverdraw.Bind(Initializer.ParameterMap, TEXT("bShowQuadOverdraw"));
+	OutputQuadOverdrawParameter.Bind(Initializer.ParameterMap, TEXT("bOutputQuadOverdraw"));
+	LODIndexParameter.Bind(Initializer.ParameterMap, TEXT("LODIndex"));
+	SkinCacheDebugColorParameter.Bind(Initializer.ParameterMap, TEXT("SkinCacheDebugColor"));
+	VisualizeModeParameter.Bind(Initializer.ParameterMap, TEXT("VisualizeMode"));
+	QuadBufferUAV.Bind(Initializer.ParameterMap, TEXT("RWQuadBuffer"));
+}
+
 bool FDebugViewModePS::ShouldCompilePermutation(const FMeshMaterialShaderPermutationParameters& Parameters)
 {
 	return ShouldCompileDebugViewModeShader(Parameters);
 }
 
+void FDebugViewModePS::ModifyCompilationEnvironment(const FMaterialShaderPermutationParameters& Parameters, FShaderCompilerEnvironment& OutEnvironment)
+{
+	OutEnvironment.SetDefine(TEXT("UNDEFINED_ACCURACY"), UndefinedStreamingAccuracyIntensity);
+	OutEnvironment.SetDefine(TEXT("MAX_NUM_TEX_COORD"), (uint32)TEXSTREAM_MAX_NUM_UVCHANNELS);
+	OutEnvironment.SetDefine(TEXT("INITIAL_GPU_SCALE"), (uint32)TEXSTREAM_INITIAL_GPU_SCALE);
+	OutEnvironment.SetDefine(TEXT("TILE_RESOLUTION"), (uint32)TEXSTREAM_TILE_RESOLUTION);
+	OutEnvironment.SetDefine(TEXT("MAX_NUM_TEXTURE_REGISTER"), (uint32)TEXSTREAM_MAX_NUM_TEXTURES_PER_MATERIAL);
+	OutEnvironment.SetDefine(TEXT("SCENE_TEXTURES_DISABLED"), 1u);
+	TCHAR BufferRegister[] = { 'u', '0', 0 };
+	BufferRegister[1] += GetQuadOverdrawUAVIndex(Parameters.Platform);
+	OutEnvironment.SetDefine(TEXT("QUAD_BUFFER_REGISTER"), BufferRegister);
+	const bool bUsingMobileRenderer = FSceneInterface::GetShadingPath(GetMaxSupportedFeatureLevel(Parameters.Platform)) == EShadingPath::Mobile;
+	OutEnvironment.SetDefine(TEXT("OUTPUT_QUAD_OVERDRAW"), !bUsingMobileRenderer);
+
+
+	for (int i = 0; i < DVSM_MAX; ++i)
+	{
+		OutEnvironment.SetDefine(DebugViewShaderModeToString(static_cast<EDebugViewShaderMode>(i)), i);
+	}
+
+	FMeshMaterialShader::ModifyCompilationEnvironment(Parameters, OutEnvironment);
+}
+
 void FDebugViewModePS::GetElementShaderBindings(
 	const FShaderMapPointerTable& PointerTable,
 	const FScene* Scene,
diff --git a/Engine/Source/Runtime/Renderer/Private/DebugViewModeRendering.h b/Engine/Source/Runtime/Renderer/Private/DebugViewModeRendering.h
--- a/Engine/Source/Runtime/Renderer/Private/DebugViewModeRendering.h
+++ b/Engine/Source/Runtime/Renderer/Private/DebugViewModeRendering.h
@@ -94,10 +94,8 @@ class FDebugViewModeVS : public FMeshMaterialShader
 	DECLARE_SHADER_TYPE(FDebugViewModeVS,MeshMaterial);
 protected:
 
-	FDebugViewModeVS() = default;
-	FDebugViewModeVS(const FMeshMaterialShaderType::CompiledShaderInitializerType& Initializer) 
-		: FMeshMaterialShader(Initializer)
-	{}
+	FDebugViewModeVS();
+	FDebugViewModeVS(const FMeshMaterialShaderType::CompiledShaderInitializerType& Initializer);
 
 public:
 
@@ -111,31 +109,10 @@ public:
 		const FMaterial& Material,
 		const FMeshPassProcessorRenderState& DrawRenderState,
 		const FDebugViewModeShaderElementData& ShaderElementData,
-		FMeshDrawSingleShaderBindings& ShaderBindings) const
-	{
-		FMeshMaterialShader::GetShaderBindings(Scene, FeatureLevel, PrimitiveSceneProxy, MaterialRenderProxy, Material, DrawRenderState, ShaderElementData, ShaderBindings);
-	}
-
-	static void SetCommonDefinitions(const FMaterialShaderPermutationParameters& Parameters, FShaderCompilerEnvironment& OutEnvironment)
-	{
-		// SM4 has less input interpolants. Also instanced meshes use more interpolants.
-		if (Parameters.MaterialParameters.bIsDefaultMaterial || (IsFeatureLevelSupported(Parameters.Platform, ERHIFeatureLevel::SM5) && !Parameters.MaterialParameters.bIsUsedWithInstancedStaticMeshes))
-		{	// Force the default material to pass enough texcoords to the pixel shaders (even though not using them).
-			// This is required to allow material shaders to have access to the sampled coords.
-			OutEnvironment.SetDefine(TEXT("MIN_MATERIAL_TEXCOORDS"), (uint32)4);
-		}
-		else // Otherwise still pass at minimum amount to have debug shader using a texcoord to work (material might not use any).
-		{
-			OutEnvironment.SetDefine(TEXT("MIN_MATERIAL_TEXCOORDS"), (uint32)2);
-		}
-
-	}
-
-	static void ModifyCompilationEnvironment(const FMaterialShaderPermutationParameters& Parameters, FShaderCompilerEnvironment& OutEnvironment)
-	{
-		SetCommonDefinitions(Parameters, OutEnvironment);
-		FMeshMaterialShader::ModifyCompilationEnvironment(Parameters, OutEnvironment);
-	}
+		FMeshDrawSingleShaderBindings& ShaderBindings) const;
+
+	static void SetCommonDefinitions(const FMaterialShaderPermutationParameters& Parameters, FShaderCompilerEnvironment& OutEnvironment);
+	static void ModifyCompilationEnvironment(const FMaterialShaderPermutationParameters& Parameters, FShaderCompilerEnvironment& OutEnvironment);
 };
 
 class FDebugViewModePS : public FMeshMaterialShader
@@ -143,50 +120,11 @@ class FDebugViewModePS : public FMeshMaterialShader
 public:
 	DECLARE_SHADER_TYPE(FDebugViewModePS, MeshMaterial);
 
-	FDebugViewModePS() = default;
-	FDebugViewModePS(const FMeshMaterialShaderType::CompiledShaderInitializerType& Initializer)
-		: FMeshMaterialShader(Initializer)
-	{
-		OneOverCPUTexCoordScalesParameter.Bind(Initializer.ParameterMap, TEXT("OneOverCPUTexCoordScales"));
-		TexCoordIndicesParameter.Bind(Initializer.ParameterMap, TEXT("TexCoordIndices"));
-		CPUTexelFactorParameter.Bind(Initializer.ParameterMap, TEXT("CPUTexelFactor"));
-		NormalizedComplexity.Bind(Initializer.ParameterMap, TEXT("NormalizedComplexity"));
-		AnalysisParamsParameter.Bind(Initializer.ParameterMap, TEXT("AnalysisParams"));
-		PrimitiveAlphaParameter.Bind(Initializer.ParameterMap, TEXT("PrimitiveAlpha"));
-		TexCoordAnalysisIndexParameter.Bind(Initializer.ParameterMap, TEXT("TexCoordAnalysisIndex"));
-		CPULogDistanceParameter.Bind(Initializer.ParameterMap, TEXT("CPULogDistance"));
-		ShowQuadOverdraw.Bind(Initializer.ParameterMap, TEXT("bShowQuadOverdraw"));
-		OutputQuadOverdrawParameter.Bind(Initializer.ParameterMap, TEXT("bOutputQuadOverdraw"));
-		LODIndexParameter.Bind(Initializer.ParameterMap, TEXT("LODIndex"));
-		SkinCacheDebugColorParameter.Bind(Initializer.ParameterMap, TEXT("SkinCacheDebugColor"));
-		VisualizeModeParameter.Bind(Initializer.ParameterMap, TEXT("VisualizeMode"));
-		QuadBufferUAV.Bind(Initializer.ParameterMap, TEXT("RWQuadBuffer"));
-	}
+	FDebugViewModePS();
+	FDebugViewModePS(const FMeshMaterialShaderType::CompiledShaderInitializerType& Initializer);
 
 	static bool ShouldCompilePermutation(const FMeshMaterialShaderPermutationParameters& Parameters);
-
-	static void ModifyCompilationEnvironment(const FMaterialShaderPermutationParameters& Parameters, FShaderCompilerEnvironment& OutEnvironment)
-	{
-		OutEnvironment.SetDefine(TEXT("UNDEFINED_ACCURACY"), UndefinedStreamingAccuracyIntensity);
-		OutEnvironment.SetDefine(TEXT("MAX_NUM_TEX_COORD"), (uint32)TEXSTREAM_MAX_NUM_UVCHANNELS);
-		OutEnvironment.SetDefine(TEXT("INITIAL_GPU_SCALE"), (uint32)TEXSTREAM_INITIAL_GPU_SCALE);
-		OutEnvironment.SetDefine(TEXT("TILE_RESOLUTION"), (uint32)TEXSTREAM_TILE_RESOLUTION);
-		OutEnvironment.SetDefine(TEXT("MAX_NUM_TEXTURE_REGISTER"), (uint32)TEXSTREAM_MAX_NUM_TEXTURES_PER_MATERIAL);
-		OutEnvironment.SetDefine(TEXT("SCENE_TEXTURES_DISABLED"), 1u);
-		TCHAR BufferRegister[] = { 'u', '0', 0 };
-		BufferRegister[1] += GetQuadOverdrawUAVIndex(Parameters.Platform);
-		OutEnvironment.SetDefine(TEXT("QUAD_BUFFER_REGISTER"), BufferRegister);
-		const bool bUsingMobileRenderer = FSceneInterface::GetShadingPath(GetMaxSupportedFeatureLevel(Parameters.Platform)) == EShadingPath::Mobile;
-		OutEnvironment.SetDefine(TEXT("OUTPUT_QUAD_OVERDRAW"), !bUsingMobileRenderer);
-
-
-		for (int i = 0; i < DVSM_MAX; ++i)
-		{
-			OutEnvironment.SetDefine(DebugViewShaderModeToString(static_cast<EDebugViewShaderMode>(i)), i);
-		}
-
-		FMeshMaterialShader::ModifyCompilationEnvironment(Parameters, OutEnvironment);
-	}
+	static void ModifyCompilationEnvironment(const FMaterialShaderPermutationParameters& Parameters, FShaderCompilerEnvironment& OutEnvironment);
 
 	void GetElementShaderBindings(
 		const FShaderMapPointerTable& PointerTable,
diff --git a/Engine/Source/Runtime/Renderer/Private/GenerateConservativeDepthBuffer.cpp b/Engine/Source/Runtime/Renderer/Private/GenerateConservativeDepthBuffer.cpp
--- a/Engine/Source/Runtime/Renderer/Private/GenerateConservativeDepthBuffer.cpp
+++ b/Engine/Source/Runtime/Renderer/Private/GenerateConservativeDepthBuffer.cpp
@@ -6,8 +6,7 @@
 
 #include "GenerateConservativeDepthBuffer.h"
 #include "SceneRendering.h"
-
-
+#include "DataDrivenShaderPlatformInfo.h"
 
 class FGenerateConservativeDepthBufferCS : public FGlobalShader
 {
diff --git a/Engine/Source/Runtime/Renderer/Private/LightRendering.cpp b/Engine/Source/Runtime/Renderer/Private/LightRendering.cpp
--- a/Engine/Source/Runtime/Renderer/Private/LightRendering.cpp
+++ b/Engine/Source/Runtime/Renderer/Private/LightRendering.cpp
@@ -191,6 +191,22 @@ bool ShouldRenderRayTracingShadowsForLight(const FLightSceneInfoCompact& LightIn
 }
 #endif // RHI_RAYTRACING
 
+bool FDeferredLightVS::ShouldCompilePermutation(const FGlobalShaderPermutationParameters& Parameters)
+{
+	FPermutationDomain PermutationVector(Parameters.PermutationId);
+	if (PermutationVector.Get<FRadialLight>())
+	{
+		return IsFeatureLevelSupported(Parameters.Platform, ERHIFeatureLevel::SM5) || IsMobileDeferredShadingEnabled(Parameters.Platform);
+	}
+	// used with FPrefilterPlanarReflectionPS on mobile
+	return true;
+}
+
+void FDeferredLightVS::ModifyCompilationEnvironment(const FGlobalShaderPermutationParameters& Parameters, FShaderCompilerEnvironment& OutEnvironment)
+{
+	FGlobalShader::ModifyCompilationEnvironment(Parameters, OutEnvironment);
+}
+
 FDrawFullScreenRectangleParameters FDeferredLightVS::GetFullScreenRectParameters(
 	float X, float Y,
 	float SizeX, float SizeY,
diff --git a/Engine/Source/Runtime/Renderer/Private/LightRendering.h b/Engine/Source/Runtime/Renderer/Private/LightRendering.h
--- a/Engine/Source/Runtime/Renderer/Private/LightRendering.h
+++ b/Engine/Source/Runtime/Renderer/Private/LightRendering.h
@@ -493,21 +493,8 @@ class FDeferredLightVS : public FGlobalShader
 	END_SHADER_PARAMETER_STRUCT()
 
 public:
-	static bool ShouldCompilePermutation(const FGlobalShaderPermutationParameters& Parameters)
-	{
-		FPermutationDomain PermutationVector(Parameters.PermutationId);
-		if (PermutationVector.Get<FRadialLight>())
-		{
-			return IsFeatureLevelSupported(Parameters.Platform, ERHIFeatureLevel::SM5) || IsMobileDeferredShadingEnabled(Parameters.Platform);
-		}
-		// used with FPrefilterPlanarReflectionPS on mobile
-		return true;
-	}
-
-	static void ModifyCompilationEnvironment(const FGlobalShaderPermutationParameters& Parameters, FShaderCompilerEnvironment& OutEnvironment)
-	{
-		FGlobalShader::ModifyCompilationEnvironment(Parameters, OutEnvironment);
-	}
+	static bool ShouldCompilePermutation(const FGlobalShaderPermutationParameters& Parameters);
+	static void ModifyCompilationEnvironment(const FGlobalShaderPermutationParameters& Parameters, FShaderCompilerEnvironment& OutEnvironment);
 
 	static FDrawFullScreenRectangleParameters GetFullScreenRectParameters(
 		float X, float Y,
diff --git a/Engine/Source/Runtime/Renderer/Private/PostProcess/PostProcessMobile.cpp b/Engine/Source/Runtime/Renderer/Private/PostProcess/PostProcessMobile.cpp
--- a/Engine/Source/Runtime/Renderer/Private/PostProcess/PostProcessMobile.cpp
+++ b/Engine/Source/Runtime/Renderer/Private/PostProcess/PostProcessMobile.cpp
@@ -32,6 +32,11 @@ static TAutoConsoleVariable<int32> CVarMobileEyeAdaptation(
 	TEXT(" 1: Enabled (Default)"),
 	ECVF_RenderThreadSafe);
 
+bool FMSAADecodeAndCopyRectPS_Mobile::ShouldCompilePermutation(const FGlobalShaderPermutationParameters& Parameters)
+{
+	return IsMetalMobilePlatform(Parameters.Platform);
+}
+
 IMPLEMENT_GLOBAL_SHADER(FMSAADecodeAndCopyRectPS_Mobile, "/Engine/Private/PostProcessMobile.usf", "MSAADecodeAndCopyRectPS", SF_Pixel);
 
 static EPixelFormat GetHDRPixelFormat()
diff --git a/Engine/Source/Runtime/Renderer/Private/PostProcess/PostProcessMobile.h b/Engine/Source/Runtime/Renderer/Private/PostProcess/PostProcessMobile.h
--- a/Engine/Source/Runtime/Renderer/Private/PostProcess/PostProcessMobile.h
+++ b/Engine/Source/Runtime/Renderer/Private/PostProcess/PostProcessMobile.h
@@ -183,10 +183,7 @@ class FMSAADecodeAndCopyRectPS_Mobile : public FGlobalShader
 	DECLARE_GLOBAL_SHADER(FMSAADecodeAndCopyRectPS_Mobile);
 	SHADER_USE_PARAMETER_STRUCT(FMSAADecodeAndCopyRectPS_Mobile, FGlobalShader);
 
-	static bool ShouldCompilePermutation(const FGlobalShaderPermutationParameters& Parameters)
-	{
-		return IsMetalMobilePlatform(Parameters.Platform);
-	}
+	static bool ShouldCompilePermutation(const FGlobalShaderPermutationParameters& Parameters);
 
 	BEGIN_SHADER_PARAMETER_STRUCT(FParameters, )
 		SHADER_PARAMETER_RDG_TEXTURE(Texture2D, InputTexture)
diff --git a/Engine/Source/Runtime/Renderer/Private/Strata/Strata.cpp b/Engine/Source/Runtime/Renderer/Private/Strata/Strata.cpp
--- a/Engine/Source/Runtime/Renderer/Private/Strata/Strata.cpp
+++ b/Engine/Source/Runtime/Renderer/Private/Strata/Strata.cpp
@@ -859,6 +859,17 @@ IMPLEMENT_GLOBAL_SHADER(FStrataBSDFTilePrepareArgsPassCS, "/Engine/Private/Strat
 
 ////////////////////////////////////////////////////////////////////////////////////////////////////////////////////////////////////////////////
 
+bool FStrataTilePassVS::ShouldCompilePermutation(const FGlobalShaderPermutationParameters& Parameters)
+{
+	return GetMaxSupportedFeatureLevel(Parameters.Platform) >= ERHIFeatureLevel::SM5; // We do not skip the compilation because we have some conditional when tiling a pass and the shader must be fetch once before hand.
+}
+
+void FStrataTilePassVS::ModifyCompilationEnvironment(const FGlobalShaderPermutationParameters& Parameters, FShaderCompilerEnvironment& OutEnvironment)
+{
+	FGlobalShader::ModifyCompilationEnvironment(Parameters, OutEnvironment);
+	OutEnvironment.SetDefine(TEXT("SHADER_TILE_VS"), 1);
+}
+
 class FStrataMaterialStencilTaggingPassPS : public FGlobalShader
 {
 	DECLARE_GLOBAL_SHADER(FStrataMaterialStencilTaggingPassPS);
diff --git a/Engine/Source/Runtime/Renderer/Private/Strata/Strata.h b/Engine/Source/Runtime/Renderer/Private/Strata/Strata.h
--- a/Engine/Source/Runtime/Renderer/Private/Strata/Strata.h
+++ b/Engine/Source/Runtime/Renderer/Private/Strata/Strata.h
@@ -208,16 +208,8 @@ class FStrataTilePassVS : public FGlobalShader
 		RDG_BUFFER_ACCESS(TileIndirectBuffer, ERHIAccess::IndirectArgs)
 	END_SHADER_PARAMETER_STRUCT()
 
-	static bool ShouldCompilePermutation(const FGlobalShaderPermutationParameters& Parameters)
-	{
-		return GetMaxSupportedFeatureLevel(Parameters.Platform) >= ERHIFeatureLevel::SM5; // We do not skip the compilation because we have some conditional when tiling a pass and the shader must be fetch once before hand.
-	}
-
-	static void ModifyCompilationEnvironment(const FGlobalShaderPermutationParameters& Parameters, FShaderCompilerEnvironment& OutEnvironment)
-	{
-		FGlobalShader::ModifyCompilationEnvironment(Parameters, OutEnvironment);
-		OutEnvironment.SetDefine(TEXT("SHADER_TILE_VS"), 1);
-	}
+	static bool ShouldCompilePermutation(const FGlobalShaderPermutationParameters& Parameters);
+	static void ModifyCompilationEnvironment(const FGlobalShaderPermutationParameters& Parameters, FShaderCompilerEnvironment& OutEnvironment);
 };
 
 FStrataTileParameter SetTileParameters(FRDGBuilder& GraphBuilder, const FViewInfo& View, const EStrataTileType Type);
diff --git a/Engine/Source/Runtime/Renderer/Private/SubsurfaceTiles.cpp b/Engine/Source/Runtime/Renderer/Private/SubsurfaceTiles.cpp
--- a/Engine/Source/Runtime/Renderer/Private/SubsurfaceTiles.cpp
+++ b/Engine/Source/Runtime/Renderer/Private/SubsurfaceTiles.cpp
@@ -2,8 +2,7 @@
 
 #include "SubsurfaceTiles.h"
 #include "SceneRendering.h"
-
-
+#include "DataDrivenShaderPlatformInfo.h"
 
 bool FSubsurfaceTilePassVS::ShouldCompilePermutation(const FGlobalShaderPermutationParameters& Parameters)
 {
diff --git a/Engine/Source/Runtime/Renderer/Public/VolumeRendering.h b/Engine/Source/Runtime/Renderer/Public/VolumeRendering.h
--- a/Engine/Source/Runtime/Renderer/Public/VolumeRendering.h
+++ b/Engine/Source/Runtime/Renderer/Public/VolumeRendering.h
@@ -44,30 +44,15 @@ struct FVolumeBounds
 };
 
 /** Vertex shader used to write to a range of slices of a 3d volume texture. */
-class FWriteToSliceVS : public FGlobalShader
+class ENGINE_API FWriteToSliceVS : public FGlobalShader
 {
-	DECLARE_EXPORTED_SHADER_TYPE(FWriteToSliceVS,Global,ENGINE_API);
+	DECLARE_GLOBAL_SHADER(FWriteToSliceVS);
 public:
+	FWriteToSliceVS();
+	FWriteToSliceVS(const ShaderMetaType::CompiledShaderInitializerType& Initializer);
 
-	static bool ShouldCompilePermutation(const FGlobalShaderPermutationParameters& Parameters) 
-	{ 
-		return IsFeatureLevelSupported(Parameters.Platform, ERHIFeatureLevel::SM5); 
-	}
-
-	static void ModifyCompilationEnvironment(const FGlobalShaderPermutationParameters& Parameters, FShaderCompilerEnvironment& OutEnvironment )
-	{
-		FGlobalShader::ModifyCompilationEnvironment( Parameters, OutEnvironment );
-		OutEnvironment.CompilerFlags.Add( CFLAG_VertexToGeometryShader );
-	}
-
-	FWriteToSliceVS(const ShaderMetaType::CompiledShaderInitializerType& Initializer):
-		FGlobalShader(Initializer)
-	{
-		UVScaleBias.Bind(Initializer.ParameterMap, TEXT("UVScaleBias"));
-		MinZ.Bind(Initializer.ParameterMap, TEXT("MinZ"));
-	}
-
-	FWriteToSliceVS() {}
+	static bool ShouldCompilePermutation(const FGlobalShaderPermutationParameters& Parameters);
+	static void ModifyCompilationEnvironment(const FGlobalShaderPermutationParameters& Parameters, FShaderCompilerEnvironment& OutEnvironment);
 
 	template <typename TRHICommandList>
 	void SetParameters(TRHICommandList& RHICmdList, const FVolumeBounds& VolumeBounds, FIntVector VolumeResolution)
@@ -88,22 +73,14 @@ private:
 };
 
 /** Geometry shader used to write to a range of slices of a 3d volume texture. */
-class FWriteToSliceGS : public FGlobalShader
+class ENGINE_API FWriteToSliceGS : public FGlobalShader
 {
-	DECLARE_EXPORTED_SHADER_TYPE(FWriteToSliceGS,Global,ENGINE_API);
+	DECLARE_GLOBAL_SHADER(FWriteToSliceGS);
 public:
+	FWriteToSliceGS();
+	FWriteToSliceGS(const ShaderMetaType::CompiledShaderInitializerType& Initializer);
 
-	static bool ShouldCompilePermutation(const FGlobalShaderPermutationParameters& Parameters) 
-	{ 
-		return IsFeatureLevelSupported(Parameters.Platform, ERHIFeatureLevel::SM5) && RHISupportsGeometryShaders(Parameters.Platform); 
-	}
-
-	FWriteToSliceGS(const ShaderMetaType::CompiledShaderInitializerType& Initializer):
-		FGlobalShader(Initializer)
-	{
-		MinZ.Bind(Initializer.ParameterMap, TEXT("MinZ"));
-	}
-	FWriteToSliceGS() {}
+	static bool ShouldCompilePermutation(const FGlobalShaderPermutationParameters& Parameters);
 
 	template <typename TRHICommandList>
 	void SetParameters(TRHICommandList& RHICmdList, int32 MinZValue)
@@ -119,31 +96,10 @@ private:
 extern ENGINE_API void RasterizeToVolumeTexture(FRHICommandList& RHICmdList, FVolumeBounds VolumeBounds);
 
 /** Vertex buffer used for rendering into a volume texture. */
-class FVolumeRasterizeVertexBuffer : public FVertexBuffer
+class ENGINE_API FVolumeRasterizeVertexBuffer : public FVertexBuffer
 {
 public:
-
-	virtual void InitRHI() override
-	{
-		// Used as a non-indexed triangle strip, so 4 vertices per quad
-		const uint32 Size = 4 * sizeof(FScreenVertex);
-		FRHIResourceCreateInfo CreateInfo(TEXT("FVolumeRasterizeVertexBuffer"));
-		VertexBufferRHI = RHICreateBuffer(Size, BUF_Static | BUF_VertexBuffer, 0, ERHIAccess::VertexOrIndexBuffer, CreateInfo);
-		FScreenVertex* DestVertex = (FScreenVertex*)RHILockBuffer(VertexBufferRHI, 0, Size, RLM_WriteOnly);
-
-		// Setup a full - render target quad
-		// A viewport and UVScaleBias will be used to implement rendering to a sub region
-		DestVertex[0].Position = FVector2f(1, -GProjectionSignY);
-		DestVertex[0].UV = FVector2f(1, 1);
-		DestVertex[1].Position = FVector2f(1, GProjectionSignY);
-		DestVertex[1].UV = FVector2f(1, 0);
-		DestVertex[2].Position = FVector2f(-1, -GProjectionSignY);
-		DestVertex[2].UV = FVector2f(0, 1);
-		DestVertex[3].Position = FVector2f(-1, GProjectionSignY);
-		DestVertex[3].UV = FVector2f(0, 0);
-
-		RHIUnlockBuffer(VertexBufferRHI);      
-	}
+	virtual void InitRHI() override;
 };
 
 extern ENGINE_API TGlobalResource<FVolumeRasterizeVertexBuffer> GVolumeRasterizeVertexBuffer;
\ No newline at end of file
diff --git a/Engine/Source/Runtime/SlateRHIRenderer/Private/SlateShaders.cpp b/Engine/Source/Runtime/SlateRHIRenderer/Private/SlateShaders.cpp
--- a/Engine/Source/Runtime/SlateRHIRenderer/Private/SlateShaders.cpp
+++ b/Engine/Source/Runtime/SlateRHIRenderer/Private/SlateShaders.cpp
@@ -3,6 +3,7 @@
 #include "SlateShaders.h"
 #include "Rendering/RenderingCommon.h"
 #include "PipelineStateCache.h"
+#include "DataDrivenShaderPlatformInfo.h"
 
 /** Flag to determine if we are running with a color vision deficiency shader on */
 EColorVisionDeficiency GSlateColorDeficiencyType = EColorVisionDeficiency::NormalVision;
@@ -33,6 +34,11 @@ IMPLEMENT_SHADER_TYPE(, FSlateDebugBatchingPS, TEXT("/Engine/Private/SlateElemen
 	IMPLEMENT_SHADER_TYPE(template<>,TSlateElementPS##ShaderType##bDrawDisabledEffect##bUseTextureAlpha##bUseTextureGrayscale##bIsVirtualTexture##A,TEXT("/Engine/Private/SlateElementPixelShader.usf"),TEXT("Main"),SF_Pixel);
 
 #if WITH_EDITOR
+bool FHDREditorConvertPS::ShouldCompilePermutation(const FGlobalShaderPermutationParameters& Parameters)
+{
+	return IsFeatureLevelSupported(Parameters.Platform, ERHIFeatureLevel::SM5);
+}
+
 IMPLEMENT_SHADER_TYPE(, FHDREditorConvertPS, TEXT("/Engine/Private/CompositeUIPixelShader.usf"), TEXT("HDREditorConvert"), SF_Pixel);
 #endif
 /**
