Fix usage of a shader type that may not be compiled.

diff --git a/Engine/Source/Runtime/Renderer/Private/SceneRendering.cpp b/Engine/Source/Runtime/Renderer/Private/SceneRendering.cpp
--- a/Engine/Source/Runtime/Renderer/Private/SceneRendering.cpp
+++ b/Engine/Source/Runtime/Renderer/Private/SceneRendering.cpp
@@ -5254,6 +5254,13 @@ namespace
 			TSetColorResolveShader<ShaderType>(RHICmdList, GraphicsPSOInit, View, bArrayResolve, SceneColorTargetableRHI);
 		}
 	}
+
+	template<typename ShaderType>
+	FRHIVertexShader* GetTypedVS(const FViewInfo& View)
+	{
+		TShaderMapRef<ShaderType> ShaderRef(View.ShaderMap);
+		return ShaderRef.GetVertexShader();
+	}
 }
 
 void AddResolveSceneColorPass(FRDGBuilder& GraphBuilder, const FViewInfo& View, FRDGTextureMSAA SceneColor)
@@ -5328,10 +5335,8 @@ void AddResolveSceneColorPass(FRDGBuilder& GraphBuilder, const FViewInfo& View,
 		{
 			bool bArrayResolve = SceneColorTargetableRHI->GetDesc().IsTextureArray();
 
-			TShaderMapRef<FHdrCustomResolveVS> VertexShader(View.ShaderMap);
-			TShaderMapRef<FHdrCustomResolveArrayVS> VertexArrayShader(View.ShaderMap);
 			GraphicsPSOInit.BoundShaderState.VertexDeclarationRHI = GetVertexDeclarationFVector4();
-			GraphicsPSOInit.BoundShaderState.VertexShaderRHI = bArrayResolve ? VertexArrayShader.GetVertexShader() : VertexShader.GetVertexShader();
+			GraphicsPSOInit.BoundShaderState.VertexShaderRHI = bArrayResolve ? GetTypedVS<FHdrCustomResolveArrayVS>(View) : GetTypedVS<FHdrCustomResolveVS>(View);
 			GraphicsPSOInit.PrimitiveType = PT_TriangleList;
 
 			if (SceneColorFMaskRHI)
@@ -5423,6 +5428,14 @@ namespace
 		OutTextureIndex = ShaderRef->UnresolvedSurface.GetBaseIndex();
 		return ShaderRef.GetPixelShader();
 	};
+
+	template<typename ShaderType>
+	FRHIVertexShader* GetDepthResolveVS(const FViewInfo& View, TShaderRef<FResolveVS>& OutShaderMapRef)
+	{
+		TShaderMapRef<ShaderType> ShaderRef(View.ShaderMap);
+		OutShaderMapRef = ShaderRef;
+		return ShaderRef.GetVertexShader();
+	}
 }
 
 void AddResolveSceneDepthPass(FRDGBuilder& GraphBuilder, const FViewInfo& View, FRDGTextureMSAA SceneDepth)
@@ -5510,13 +5523,12 @@ void AddResolveSceneDepthPass(FRDGBuilder& GraphBuilder, const FViewInfo& View,
 			// unreachable
 		};
 
-		TShaderMapRef<FResolveVS> ResolveVertexShader(View.ShaderMap);
-		TShaderMapRef<FResolveArrayVS> ResolveArrayVertexShader(View.ShaderMap);
 		int32 TextureIndex = -1;
 		FRHIPixelShader* ResolvePixelShader = ChoosePixelShader(View, bArrayResolve, NumSamples, TextureIndex);
 
+		TShaderRef<FResolveVS> ResolveVertexShader;
 		GraphicsPSOInit.BoundShaderState.VertexDeclarationRHI = GEmptyVertexDeclaration.VertexDeclarationRHI;
-		GraphicsPSOInit.BoundShaderState.VertexShaderRHI = bArrayResolve ? ResolveArrayVertexShader.GetVertexShader() : ResolveVertexShader.GetVertexShader();
+		GraphicsPSOInit.BoundShaderState.VertexShaderRHI = bArrayResolve ? GetDepthResolveVS<FResolveArrayVS>(View, ResolveVertexShader) : GetDepthResolveVS<FResolveVS>(View, ResolveVertexShader);
 		GraphicsPSOInit.BoundShaderState.PixelShaderRHI = ResolvePixelShader;
 		GraphicsPSOInit.PrimitiveType = PT_TriangleStrip;
 
