Fix issues in DFAO when primary ViewRect is not top-left.

diff --git a/Engine/Shaders/Private/DistanceFieldLightingPost.usf b/Engine/Shaders/Private/DistanceFieldLightingPost.usf
--- a/Engine/Shaders/Private/DistanceFieldLightingPost.usf
+++ b/Engine/Shaders/Private/DistanceFieldLightingPost.usf
@@ -8,6 +8,7 @@
 #include "DeferredShadingCommon.ush"
 #include "DistanceFieldLightingShared.ush"
 #include "DistanceFieldAOShared.ush"
+#include "ScreenPass.ush"
 
 Texture2D BentNormalHistoryTexture;
 SamplerState BentNormalHistorySampler;
@@ -38,6 +39,8 @@ float2 DistanceFieldGBufferTexelSize;
 float2 DistanceFieldGBufferJitterOffset;
 float4 BentNormalBufferAndTexelSize;
 
+FScreenTransform UVToScreenPos;
+
 float4 GetNormalWeights(float2 Corner00UV, float2 LowResTexelSize, float3 WorldNormal)
 {
 	float4 NormalWeights;
@@ -75,7 +78,7 @@ float4 GetNormalWeights(float2 Corner00UV, float2 LowResTexelSize, float3 WorldN
 
 float ComputeSampleWeightBasedOnPosition(float4 ReferencePlane, float2 SampleScreenUV, float SampleDepth)
 {
-	float2 SampleScreenPosition = (SampleScreenUV.xy - View.ScreenPositionScaleBias.wz) / View.ScreenPositionScaleBias.xy;
+	float2 SampleScreenPosition = ApplyScreenTransform(SampleScreenUV, UVToScreenPos);
 	float3 SampleTranslatedWorldPosition = mul(float4(SampleScreenPosition * SampleDepth, SampleDepth, 1), PrimaryView.ScreenToTranslatedWorld).xyz;
 
 	float PlaneDistance = dot(ReferencePlane, float4(SampleTranslatedWorldPosition, 1));
diff --git a/Engine/Source/Runtime/Renderer/Private/DistanceFieldLightingPost.cpp b/Engine/Source/Runtime/Renderer/Private/DistanceFieldLightingPost.cpp
--- a/Engine/Source/Runtime/Renderer/Private/DistanceFieldLightingPost.cpp
+++ b/Engine/Source/Runtime/Renderer/Private/DistanceFieldLightingPost.cpp
@@ -75,6 +75,7 @@ BEGIN_SHADER_PARAMETER_STRUCT(FGeometryAwareUpsampleParameters, )
 	SHADER_PARAMETER_SAMPLER(SamplerState, DistanceFieldNormalSampler)
 	SHADER_PARAMETER_RDG_TEXTURE(Texture2D, BentNormalAOTexture)
 	SHADER_PARAMETER_SAMPLER(SamplerState, BentNormalAOSampler)
+	SHADER_PARAMETER(FScreenTransform, UVToScreenPos)
 	SHADER_PARAMETER(FVector4f, BentNormalBufferAndTexelSize)
 	SHADER_PARAMETER(FVector2f, DistanceFieldGBufferTexelSize)
 	SHADER_PARAMETER(FVector2f, DistanceFieldGBufferJitterOffset)
@@ -97,11 +98,14 @@ FGeometryAwareUpsampleParameters SetupGeometryAwareUpsampleParameters(const FVie
 	extern float GAOViewFadeDistanceScale;
 	const float DistanceFadeScaleValue = 1.0f / ((1.0f - GAOViewFadeDistanceScale) * GetMaxAOViewDistance());
 
+	const FIntRect AOViewRect = FIntRect(FIntPoint::ZeroValue, FIntPoint::DivideAndRoundDown(View.ViewRect.Size(), GAODownsampleFactor));
+
 	FGeometryAwareUpsampleParameters ShaderParameters;
 	ShaderParameters.DistanceFieldNormalTexture = DistanceFieldNormal;
 	ShaderParameters.DistanceFieldNormalSampler = TStaticSamplerState<SF_Point, AM_Clamp, AM_Clamp, AM_Clamp>::GetRHI();
 	ShaderParameters.BentNormalAOTexture = DistanceFieldAOBentNormal;
 	ShaderParameters.BentNormalAOSampler = TStaticSamplerState<SF_Bilinear, AM_Clamp, AM_Clamp, AM_Clamp>::GetRHI();
+	ShaderParameters.UVToScreenPos = FScreenTransform::ChangeTextureBasisFromTo(FScreenPassTextureViewport(DownsampledBufferSize, AOViewRect), FScreenTransform::ETextureBasis::TextureUV, FScreenTransform::ETextureBasis::ScreenPosition);
 	ShaderParameters.BentNormalBufferAndTexelSize = BentNormalBufferAndTexelSizeValue;
 	ShaderParameters.DistanceFieldGBufferTexelSize = BaseLevelTexelSizeValue;
 	ShaderParameters.DistanceFieldGBufferJitterOffset = BaseLevelTexelSizeValue * JitterOffsetValue;
@@ -252,6 +256,8 @@ void GeometryAwareUpsample(FRDGBuilder& GraphBuilder, const FViewInfo& View, FRD
 
 		SetShaderParameters(RHICmdList, PixelShader, PixelShader.GetPixelShader(), *PassParameters);
 
+		FIntPoint BufferSize = GetBufferSizeForAO(View);
+
 		DrawRectangle(
 			RHICmdList,
 			0, 0,
@@ -259,7 +265,7 @@ void GeometryAwareUpsample(FRDGBuilder& GraphBuilder, const FViewInfo& View, FRD
 			0, 0,
 			View.ViewRect.Width() / GAODownsampleFactor, View.ViewRect.Height() / GAODownsampleFactor,
 			FIntPoint(View.ViewRect.Width() / GAODownsampleFactor, View.ViewRect.Height() / GAODownsampleFactor),
-			View.GetSceneTexturesConfig().Extent / FIntPoint(GAODownsampleFactor, GAODownsampleFactor),
+			BufferSize,
 			VertexShader);
 	});
 }
@@ -415,7 +421,7 @@ void UpdateHistory(
 						RHICmdList,
 						0, 0,
 						View.ViewRect.Width() / GAODownsampleFactor, View.ViewRect.Height() / GAODownsampleFactor,
-						View.ViewRect.Min.X / GAODownsampleFactor, View.ViewRect.Min.Y / GAODownsampleFactor,
+						0, 0,
 						View.ViewRect.Width() / GAODownsampleFactor, View.ViewRect.Height() / GAODownsampleFactor,
 						FIntPoint(View.ViewRect.Width() / GAODownsampleFactor, View.ViewRect.Height() / GAODownsampleFactor),
 						SceneTextureExtent / FIntPoint(GAODownsampleFactor, GAODownsampleFactor),
