Illuminance meter

diff --git a/Engine/Shaders/Private/ReflectionEnvironmentShaders.usf b/Engine/Shaders/Private/ReflectionEnvironmentShaders.usf
--- a/Engine/Shaders/Private/ReflectionEnvironmentShaders.usf
+++ b/Engine/Shaders/Private/ReflectionEnvironmentShaders.usf
@@ -213,6 +213,27 @@ void CopyCubemapToCubeFaceColorPS(
 	OutColor = ClampToFP16 == 0.0f ? OutColor : clamp(OutColor, float4(0.0f, 0.0f, 0.0f, 0.0f), float4(MaxHalfFloat, MaxHalfFloat, MaxHalfFloat, 1.0f));
 }
 
+void CopyCubemapToIlluminanceMeterCubemapPS(
+	float4 SvPosition : SV_POSITION,
+	out float4 OutColor : SV_Target0
+	)
+{
+	float2 UV = SvPosition.xy * SvPositionToUVScale;
+	float2 ScaledUVs = UV * 2 - 1;
+
+	float3 CubeCoordinates = GetCubemapVector(ScaledUVs, CubeFace);
+
+	float NoL = saturate(normalize(CubeCoordinates).z);
+	//float NoL = CubeCoordinates.z > 0.0 ? 1 : 0;
+
+	// See https://www.ppsloan.org/publications/StupidSH36.pdf, p.9. 
+	float Tmp= 1.0f + ScaledUVs.x * ScaledUVs.x + ScaledUVs.y * ScaledUVs.y; 
+	float dA = 4.0f / (sqrt(Tmp) * Tmp); // Not really dA, but using same variable name
+
+	const float3 L = TextureCubeSampleLevel(SourceCubemapTexture, SourceCubemapSampler, CubeCoordinates, 0).rgb;
+	OutColor = float4(L * NoL* dA, dA);
+}
+
 uint MipIndex;
 uint NumMips;
 
@@ -327,6 +348,49 @@ void DownsampleMaxPS(
 	}
 }
 
+void DownsampleIntegrateCubemapIlluminancePS(
+	float4 SvPosition : SV_POSITION,
+	out float4 OutColor : SV_Target0
+)
+{
+	OutColor = 0;
+	float2 UV = SvPosition.xy * SvPositionToUVScale;
+	float2 ScaledUVs = UV * 2 - 1;
+	const int SelectedCubeFace = CubeFace;
+
+	float3 CubeCoordinates = GetCubemapVector(ScaledUVs, SelectedCubeFace);
+
+	uint MipSize = 1u << (NumMips - MipIndex - 1);
+
+	float3 TangentX = 0.0f;
+	float3 TangentY = 0.0f;
+	GetCubemapTangent(SelectedCubeFace, TangentX, TangentY);
+
+	const float SampleOffset = 1.0f / MipSize;
+
+	// Offset from the connect of 4 texels towards their center.
+	float2 Offsets[] =
+	{
+		float2(-0.5,-0.5),
+		float2( 0.5,-0.5),
+		float2( 0.5, 0.5),
+		float2(-0.5, 0.5),
+	};
+
+	OutColor = 0;
+
+	UNROLL
+	for (uint i = 0; i < 4; i++)
+	{
+		float3 SampleDir = CubeCoordinates;
+		SampleDir += TangentX * (Offsets[i].x * SampleOffset);
+		SampleDir += TangentY * (Offsets[i].y * SampleOffset);
+		OutColor += SourceCubemapTexture.SampleLevel(SourceCubemapSampler, SampleDir, 0);
+	}
+}
+
+
+
 #ifdef USE_COMPUTE
 float4 LowerHemisphereSolidColor;
 [numthreads(THREADGROUP_SIZE, THREADGROUP_SIZE, 1)]
@@ -818,4 +882,148 @@ void RenderRealTimeReflectionHeightFogPS(
 
 #endif
 
+#ifdef VISUALIZE_ILLUMINANCE_METER
+
+#include "/Engine/Private/ShaderPrint.ush"
+
+int CubeMipIndexToSampleIlluminance;
+
+TextureCube SkyLightCubeTexture;
+SamplerState SkyLightCubeSampler;
+
+float4 GetIlluminance()
+{
+	float4 IlluminanceAcc = 0;
+	IlluminanceAcc += TextureCubeSampleLevel(SourceCubemapTexture, SourceCubemapSampler, float3(-1, 0, 0), CubeMipIndexToSampleIlluminance);
+	IlluminanceAcc += TextureCubeSampleLevel(SourceCubemapTexture, SourceCubemapSampler, float3( 1, 0, 0), CubeMipIndexToSampleIlluminance);
+	IlluminanceAcc += TextureCubeSampleLevel(SourceCubemapTexture, SourceCubemapSampler, float3( 0,-1, 0), CubeMipIndexToSampleIlluminance);
+	IlluminanceAcc += TextureCubeSampleLevel(SourceCubemapTexture, SourceCubemapSampler, float3( 0, 1, 0), CubeMipIndexToSampleIlluminance);
+	IlluminanceAcc += TextureCubeSampleLevel(SourceCubemapTexture, SourceCubemapSampler, float3( 0, 0,-1), CubeMipIndexToSampleIlluminance);
+	IlluminanceAcc += TextureCubeSampleLevel(SourceCubemapTexture, SourceCubemapSampler, float3( 0, 0, 1), CubeMipIndexToSampleIlluminance);
+
+	// See https://www.ppsloan.org/publications/StupidSH36.pdf, p.9.
+	IlluminanceAcc.rgb *= 4.0 * PI / IlluminanceAcc.a;
+
+	return IlluminanceAcc;
+}
+
+float GetSimpleAutoExposure()
+{
+	float4 IlluminanceAcc = GetIlluminance();
+	float MaxIlluminance = max3(IlluminanceAcc.r, IlluminanceAcc.g, IlluminanceAcc.b);
+	float SimpleAutoExposure = 1.0f / (MaxIlluminance <= 0.0 ? 0.01f : MaxIlluminance);
+	return SimpleAutoExposure;
+}
+
+float4 SampleCubeAsLatLon(TextureCube CubeTexture, SamplerState CubeSampler, float2 UV)
+{
+	float3 DirFromLatLong = float3(cos(UV.x * 2.0 * PI), sin(UV.x * 2.0 * PI), -sin(UV.y * PI - PI / 2));
+
+	float4 DebugVis = TextureCubeSampleLevel(CubeTexture, CubeSampler, DirFromLatLong, 0);
+	DebugVis.rgb *= GetSimpleAutoExposure();
+	return DebugVis;
+}
+
+float4 SampleCubeAsSphere(TextureCube CubeTexture, SamplerState CubeSampler, float2 UV)
+{
+	float2 Dir2D = (UV - 0.5) * 2.0;
+
+	float4 DebugVis = float4(0.0, 0.0, 0.0, 1.0);
+	if (length(Dir2D) <= 1.0)
+	{
+		float3 N = 0;
+		const float3 RayO = float3(UV, -1);
+		const float3 RayD = float3(0.0f, 0.0f, 1.0f);
+		const float4 Sphere = float4(0.5f, 0.5f, 0.f, 0.5);
+		const float2 Hit = RayIntersectSphere(RayO, RayD, Sphere);
+		if (Hit.x >= 0)
+		{
+			const float3 P = RayO + RayD * Hit.x;
+			N = normalize(P - Sphere.xyz);
+			N *= float3(1.0, -1.0, 1.0);
+		}
+
+		DebugVis = TextureCubeSampleLevel(CubeTexture, CubeSampler, mul(N, (float3x3)View.ViewToTranslatedWorld), 0);
+		DebugVis.rgb *= GetSimpleAutoExposure();
+	}
+
+	return DebugVis;
+}
+
+void VisualizeIlluminanceMeterPS(
+	float4 SVPos : SV_POSITION,
+	out float4 OutColor : SV_Target0)
+{
+	const uint2 PixelPosDynRes = uint2(float2(SVPos.xy) * View.ViewResolutionFraction);
+	float2 BufferUV = SvPositionToBufferUV(float4(PixelPosDynRes, SVPos.zw));
+
+	OutColor = 0.0f;
+
+	const uint2 CubeMapPreviewOrigin	= uint2(100, 110);
+	const uint2 CubeMapPreviewSize		= uint2(500, 250);
+
+	const uint2 SkyMapPreviewOrigin		= uint2(100, 370);
+	const uint2 SkyMapPreviewSize		= uint2(500, 250);	
+
+	const uint2 SphereCubeMapPreviewOrigin	= CubeMapPreviewOrigin + uint2(CubeMapPreviewSize.x + 10, 0);
+	const uint2 SphereCubeMapPreviewSize	= uint2(250, 250);
+
+	const uint2 SphereSkyMapPreviewOrigin	= SkyMapPreviewOrigin + uint2(SkyMapPreviewSize.x + 10, 0);
+	const uint2 SphereSkyMapPreviewSize		= uint2(250, 250);
+
+	if (all(uint2(SVPos.xy) == uint2(1, 1)))
+	{
+		float4 IlluminanceAcc = GetIlluminance();
+
+		FShaderPrintContext Context = InitShaderPrintContext(true, uint2(110, 70));
+		Print(Context, TEXT("Sky light illuminance = "), InitFontColor(0.8, 0.8, 0.8));
+		Print(Context, IlluminanceAcc.r, InitFontColor(1.0, 0.8, 0.8));
+		Print(Context, IlluminanceAcc.g, InitFontColor(0.8, 1.0, 0.8));
+		Print(Context, IlluminanceAcc.b, InitFontColor(0.8, 0.8, 1.0));
+
+		Newline(Context);
+
+		const float3 SunLightIlluminance = View.AtmosphereLightIlluminanceOnGroundPostTransmittance[0];
+
+		Print(Context, TEXT("Sun light illuminance = "), InitFontColor(0.8, 0.8, 0.8));
+		Print(Context, SunLightIlluminance.r, InitFontColor(1.0, 0.8, 0.8));
+		Print(Context, SunLightIlluminance.g, InitFontColor(0.8, 1.0, 0.8));
+		Print(Context, SunLightIlluminance.b, InitFontColor(0.8, 0.8, 1.0));
+	//	Print(Context, TEXT("(On ground, post SkyAtmosphere transmittance)"), InitFontColor(0.8, 0.8, 0.8));
+	}
+	else if (all(uint2(SVPos.xy) >= CubeMapPreviewOrigin) && all(uint2(SVPos.xy) <= (CubeMapPreviewOrigin + CubeMapPreviewSize)))
+	{
+		float2 UV = (SVPos.xy - float2(CubeMapPreviewOrigin)) / float2(CubeMapPreviewSize);
+		float4 DebugVis = SampleCubeAsLatLon(SourceCubemapTexture, SourceCubemapSampler, UV);
+		DebugVis.rgb *= 1.0 / DebugVis.a;	// This is a way to remove the texel dArea factor to show what luminance is integrated as a unction of NoL.
+		OutColor = float4(DebugVis.rgb, 0.0);
+	}
+	else if (all(uint2(SVPos.xy) >= SkyMapPreviewOrigin) && all(uint2(SVPos.xy) <= (SkyMapPreviewOrigin + SkyMapPreviewSize)))
+	{
+		float2 UV = (SVPos.xy - float2(SkyMapPreviewOrigin)) / float2(SkyMapPreviewSize);
+		float4 DebugVis = SampleCubeAsLatLon(SkyLightCubeTexture, SkyLightCubeSampler, UV);
+		OutColor = float4(DebugVis.rgb, 0.0);
+	}
+	else if (all(uint2(SVPos.xy) >= SphereCubeMapPreviewOrigin) && all(uint2(SVPos.xy) <= (SphereCubeMapPreviewOrigin + SphereCubeMapPreviewSize)))
+	{
+		float2 UV = (SVPos.xy - float2(SphereCubeMapPreviewOrigin)) / float2(SphereCubeMapPreviewSize);
+		float4 DebugVis = SampleCubeAsSphere(SourceCubemapTexture, SourceCubemapSampler, UV);
+		DebugVis.rgb *= 1.0 / DebugVis.a;	// This is a way to remove the texel dArea factor to show what luminance is integrated as a unction of NoL.
+		OutColor = float4(DebugVis.rgb, 0.0);
+	}
+	else if (all(uint2(SVPos.xy) >= SphereSkyMapPreviewOrigin) && all(uint2(SVPos.xy) <= (SphereSkyMapPreviewOrigin + SphereSkyMapPreviewSize)))
+	{
+		float2 UV = (SVPos.xy - float2(SphereSkyMapPreviewOrigin)) / float2(SphereSkyMapPreviewSize);
+		float4 DebugVis = SampleCubeAsSphere(SkyLightCubeTexture, SkyLightCubeSampler, UV);
+		OutColor = float4(DebugVis.rgb, 0.0);
+	}
+	else
+	{
+		clip(-1.0);
+	}
+
+	// TODO print sky as a sphere
+}
+
+#endif // DEBUG_ILLUMINANCE_METER
 
diff --git a/Engine/Source/Runtime/Engine/Classes/Components/SkyLightComponent.h b/Engine/Source/Runtime/Engine/Classes/Components/SkyLightComponent.h
--- a/Engine/Source/Runtime/Engine/Classes/Components/SkyLightComponent.h
+++ b/Engine/Source/Runtime/Engine/Classes/Components/SkyLightComponent.h
@@ -202,6 +202,12 @@ class ENGINE_API USkyLightComponent : public ULightComponentBase
 	/** Controls how occlusion from Distance Field Ambient Occlusion is combined with Screen Space Ambient Occlusion. */
 	UPROPERTY(EditAnywhere, BlueprintReadOnly, Category=DistanceFieldAmbientOcclusion)
 	TEnumAsByte<enum EOcclusionCombineMode> OcclusionCombineMode;
+	
+	/**
+	 * When enabled, draw on screen the bull cube map as well as the integrate illuminance half hemisphere, together with the sky illuminance contribution. It also show the atmospheric sun illuminance post SkyAtmosphere transmittance on ground if one is enabled.
+	 */
+	UPROPERTY(EditAnywhere, BlueprintReadOnly, Category = Visualization)
+	uint32 bShowIlluminanceMeter : 1;
 		
 	class FSkyLightSceneProxy* CreateSceneProxy() const;
 
diff --git a/Engine/Source/Runtime/Engine/Private/Components/SkyLightComponent.cpp b/Engine/Source/Runtime/Engine/Private/Components/SkyLightComponent.cpp
--- a/Engine/Source/Runtime/Engine/Private/Components/SkyLightComponent.cpp
+++ b/Engine/Source/Runtime/Engine/Private/Components/SkyLightComponent.cpp
@@ -249,6 +249,7 @@ FSkyLightSceneProxy::FSkyLightSceneProxy(const USkyLightComponent* InLightCompon
 	, bCaptureSkyLightWaitingForShaders(false)
 	, bCaptureSkyLightWaitingForMeshesOrTextures(false)
 #endif
+	, bShowIlluminanceMeter(InLightComponent->bShowIlluminanceMeter)
 	, LightColor(FLinearColor(InLightComponent->LightColor) * InLightComponent->Intensity)
 	, bMovable(InLightComponent->IsMovable())
 {
@@ -312,6 +313,7 @@ USkyLightComponent::USkyLightComponent(const FObjectInitializer& ObjectInitializ
 	CloudAmbientOcclusionStrength = 1.0f;
 	CloudAmbientOcclusionMapResolutionScale = 1.0f;
 	CloudAmbientOcclusionApertureScale = 0.05f;
+	bShowIlluminanceMeter = false;
 
 #if WITH_EDITOR
 	CaptureStatus = ESkyLightCaptureStatus::SLCS_Uninitialized;
diff --git a/Engine/Source/Runtime/Engine/Public/SceneInterface.h b/Engine/Source/Runtime/Engine/Public/SceneInterface.h
--- a/Engine/Source/Runtime/Engine/Public/SceneInterface.h
+++ b/Engine/Source/Runtime/Engine/Public/SceneInterface.h
@@ -204,6 +204,8 @@ public:
 	virtual void AllocateAndCaptureFrameSkyEnvMap(FRDGBuilder& GraphBuilder, FSceneRenderer& SceneRenderer, FViewInfo& MainView, bool bShouldRenderSkyAtmosphere, bool bShouldRenderVolumetricCloud, FInstanceCullingManager& InstanceCullingManager, FRDGExternalAccessQueue& ExternalAccessQueue) {}
 	virtual void ValidateSkyLightRealTimeCapture(FRDGBuilder& GraphBuilder, const FViewInfo& View, FRDGTextureRef SceneColorTexture) {}
 
+	virtual void ProcessAndRenderIlluminanceMeter(FRDGBuilder& GraphBuilder, TArrayView<FViewInfo> Views, FRDGTextureRef SceneColorTexture) {}
+
 	virtual void AddPlanarReflection(class UPlanarReflectionComponent* Component) {}
 	virtual void RemovePlanarReflection(class UPlanarReflectionComponent* Component) {}
 	virtual void UpdatePlanarReflectionTransform(UPlanarReflectionComponent* Component) {}
diff --git a/Engine/Source/Runtime/Engine/Public/SceneManagement.h b/Engine/Source/Runtime/Engine/Public/SceneManagement.h
--- a/Engine/Source/Runtime/Engine/Public/SceneManagement.h
+++ b/Engine/Source/Runtime/Engine/Public/SceneManagement.h
@@ -1223,6 +1223,8 @@ public:
 	bool bCaptureSkyLightWaitingForMeshesOrTextures;
 #endif
 
+	uint8 bShowIlluminanceMeter : 1;
+
 private:
 	FLinearColor LightColor;
 	const uint8 bMovable : 1;
diff --git a/Engine/Source/Runtime/Renderer/Private/DeferredShadingRenderer.cpp b/Engine/Source/Runtime/Renderer/Private/DeferredShadingRenderer.cpp
--- a/Engine/Source/Runtime/Renderer/Private/DeferredShadingRenderer.cpp
+++ b/Engine/Source/Runtime/Renderer/Private/DeferredShadingRenderer.cpp
@@ -3881,6 +3881,8 @@ void FDeferredShadingSceneRenderer::Render(FRDGBuilder& GraphBuilder)
 		ShadingEnergyConservation::Debug(GraphBuilder, View, SceneTextures);
 	}
 
+	Scene->ProcessAndRenderIlluminanceMeter(GraphBuilder, Views, SceneTextures.Color.Target);
+
 	if (!bHasRayTracedOverlay && ViewFamily.EngineShowFlags.LightShafts)
 	{
 		SCOPE_CYCLE_COUNTER(STAT_FDeferredShadingSceneRenderer_RenderLightShaftBloom);
diff --git a/Engine/Source/Runtime/Renderer/Private/IlluminanceMeter.cpp b/Engine/Source/Runtime/Renderer/Private/IlluminanceMeter.cpp
new file mode 100644
--- /dev/null
+++ b/Engine/Source/Runtime/Renderer/Private/IlluminanceMeter.cpp
@@ -0,0 +1,213 @@
+// Copyright Epic Games, Inc. All Rights Reserved.
+
+/*=============================================================================
+	Functionality for a real time illuminance meter reference for the sky light and sun.
+=============================================================================*/
+
+#include "ReflectionEnvironmentCapture.h"
+#include "RHIResources.h"
+#include "RHIShaderPlatform.h"
+#include "PixelShaderUtils.h"
+#include "SceneRendering.h"
+#include "ScenePrivate.h"
+#include "ShaderCompiler.h"
+
+DECLARE_GPU_STAT(IlluminanceMeter);
+
+class FCopyCubemapToIlluminanceMeterCubemapPS : public FGlobalShader
+{
+public:
+	DECLARE_GLOBAL_SHADER(FCopyCubemapToIlluminanceMeterCubemapPS);
+	SHADER_USE_PARAMETER_STRUCT(FCopyCubemapToIlluminanceMeterCubemapPS, FGlobalShader);
+
+	BEGIN_SHADER_PARAMETER_STRUCT(FParameters, )
+		SHADER_PARAMETER_TEXTURE(TextureCube, SourceCubemapTexture)
+		SHADER_PARAMETER_SAMPLER(SamplerState, SourceCubemapSampler)
+		SHADER_PARAMETER(FVector2f, SvPositionToUVScale)
+		SHADER_PARAMETER(int32, CubeFace)
+		RENDER_TARGET_BINDING_SLOTS()
+	END_SHADER_PARAMETER_STRUCT()
+};
+
+IMPLEMENT_GLOBAL_SHADER(FCopyCubemapToIlluminanceMeterCubemapPS, "/Engine/Private/ReflectionEnvironmentShaders.usf", "CopyCubemapToIlluminanceMeterCubemapPS", SF_Pixel);
+
+class FDownsampleIntegrateCubemapIlluminancePS : public FGlobalShader
+{
+	DECLARE_GLOBAL_SHADER(FDownsampleIntegrateCubemapIlluminancePS);
+	SHADER_USE_PARAMETER_STRUCT(FDownsampleIntegrateCubemapIlluminancePS, FGlobalShader);
+	
+	BEGIN_SHADER_PARAMETER_STRUCT(FParameters, )
+		SHADER_PARAMETER_RDG_TEXTURE_SRV(TextureCube, SourceCubemapTexture)
+		SHADER_PARAMETER_SAMPLER(SamplerState, SourceCubemapSampler)
+		SHADER_PARAMETER(FVector2f, SvPositionToUVScale)
+		SHADER_PARAMETER(int32, CubeFace)
+		SHADER_PARAMETER(int32, MipIndex)
+		SHADER_PARAMETER(int32, NumMips)
+		RENDER_TARGET_BINDING_SLOTS()
+	END_SHADER_PARAMETER_STRUCT()
+};
+
+IMPLEMENT_GLOBAL_SHADER(FDownsampleIntegrateCubemapIlluminancePS, "/Engine/Private/ReflectionEnvironmentShaders.usf", "DownsampleIntegrateCubemapIlluminancePS", SF_Pixel);
+
+static bool CanIlluminanceMeterDisplayOnPlatform(EShaderPlatform Platform)
+{
+	// On some consoles, this ALU heavy shader (and with optimisation disables for the sake of low compilation time) would spill registers. So only keep it for the editor.
+	return GetMaxSupportedFeatureLevel(Platform) >= ERHIFeatureLevel::SM5 && IsPCPlatform(Platform);
+}
+
+class FPrintIlluminanceMeterPS : public FGlobalShader
+{
+	DECLARE_GLOBAL_SHADER(FPrintIlluminanceMeterPS);
+	SHADER_USE_PARAMETER_STRUCT(FPrintIlluminanceMeterPS, FGlobalShader);
+
+	BEGIN_SHADER_PARAMETER_STRUCT(FParameters, )
+		SHADER_PARAMETER_STRUCT_REF(FViewUniformShaderParameters, ViewUniformBuffer)
+		SHADER_PARAMETER_RDG_TEXTURE(TextureCube, SourceCubemapTexture)
+		SHADER_PARAMETER_SAMPLER(SamplerState, SourceCubemapSampler)
+		SHADER_PARAMETER_TEXTURE(TextureCube, SkyLightCubeTexture)
+		SHADER_PARAMETER_SAMPLER(SamplerState, SkyLightCubeSampler)
+		SHADER_PARAMETER(int32, CubeMipIndexToSampleIlluminance)
+		SHADER_PARAMETER_STRUCT_INCLUDE(ShaderPrint::FShaderParameters, ShaderPrintParameters)
+		RENDER_TARGET_BINDING_SLOTS()
+	END_SHADER_PARAMETER_STRUCT()
+
+	static FPermutationDomain RemapPermutation(FPermutationDomain PermutationVector)
+	{
+		return PermutationVector;
+	}
+
+	static bool ShouldCompilePermutation(const FGlobalShaderPermutationParameters& Parameters)
+	{
+		return CanIlluminanceMeterDisplayOnPlatform(Parameters.Platform) && EnumHasAllFlags(Parameters.Flags, EShaderPermutationFlags::HasEditorOnlyData);
+	}
+
+	static void ModifyCompilationEnvironment(const FGlobalShaderPermutationParameters& Parameters, FShaderCompilerEnvironment& OutEnvironment)
+	{
+		FGlobalShader::ModifyCompilationEnvironment(Parameters, OutEnvironment);
+
+		// Stay debug and skip optimizations to reduce compilation time on this long shader.
+		OutEnvironment.CompilerFlags.Add(CFLAG_Debug);
+		OutEnvironment.SetDefine(TEXT("VISUALIZE_ILLUMINANCE_METER"), 1);
+	}
+};
+
+IMPLEMENT_GLOBAL_SHADER(FPrintIlluminanceMeterPS, "/Engine/Private/ReflectionEnvironmentShaders.usf", "VisualizeIlluminanceMeterPS", SF_Pixel);
+
+void FScene::ProcessAndRenderIlluminanceMeter(FRDGBuilder& GraphBuilder, TArrayView<FViewInfo> Views, FRDGTextureRef SceneColorTexture)
+{
+	FViewInfo& MainView = Views[0];
+	if (!SkyLight || !SkyLight->bShowIlluminanceMeter || MainView.bIsSceneCapture || MainView.bIsReflectionCapture || MainView.bIsPlanarReflection ||
+		!CanIlluminanceMeterDisplayOnPlatform(MainView.GetShaderPlatform()))
+	{
+		return;
+	}
+
+	RDG_EVENT_SCOPE(GraphBuilder, "IlluminanceMeter");
+	RDG_GPU_STAT_SCOPE(GraphBuilder, IlluminanceMeter);
+
+	// 0- Get source cubemap and allocate cube map transient resources
+
+	FTextureRHIRef SkyLightTextureRHI = GBlackTextureCube->TextureRHI;;
+	if (SkyLight->bRealTimeCaptureEnabled && ConvolvedSkyRenderTargetReadyIndex >= 0)
+	{
+		// Cannot blend with this capture mode as of today.
+		SkyLightTextureRHI = ConvolvedSkyRenderTarget[ConvolvedSkyRenderTargetReadyIndex]->GetRHI();
+
+		//
+	}
+	else if (SkyLight->ProcessedTexture)
+	{
+		SkyLightTextureRHI = SkyLight->ProcessedTexture->TextureRHI;
+	}
+
+	const int32 CubemapSize = SkyLightTextureRHI->GetDesc().GetSize().X;
+	const int32 NumReflectionCaptureMips = GetNumMips(CubemapSize);
+	const FRDGTextureDesc TextureDesc = FRDGTextureDesc::CreateCube(CubemapSize, PF_A32B32G32R32F, FClearValueBinding::Black, 
+		TexCreate_RenderTargetable | TexCreate_ShaderResource | TexCreate_TargetArraySlicesIndependently | TexCreate_DisableDCC, NumReflectionCaptureMips);
+	FRDGTextureRef IlluminanceMeterCubeTexture = GraphBuilder.CreateTexture(TextureDesc, TEXT("IlluminanceMeterTexture"));
+
+	// 1- Compute illuminance contribution for each texel luminance as a function of their solid angle in the source cube map
+	{
+		RDG_EVENT_SCOPE(GraphBuilder, "CopyCubemapToIlluminanceMeterCubemap");
+
+		TShaderMapRef<FCopyCubemapToIlluminanceMeterCubemapPS> PixelShader(MainView.ShaderMap);
+
+		for (uint32 CubeFace = 0; CubeFace < CubeFace_MAX; CubeFace++)
+		{
+			auto* PassParameters = GraphBuilder.AllocParameters<FCopyCubemapToIlluminanceMeterCubemapPS::FParameters>();
+			PassParameters->SourceCubemapSampler = TStaticSamplerState<SF_Point, AM_Clamp, AM_Clamp, AM_Clamp>::GetRHI();
+			PassParameters->SourceCubemapTexture = SkyLightTextureRHI;
+			PassParameters->CubeFace = CubeFace;
+			PassParameters->SvPositionToUVScale = FVector2f(1.0f / CubemapSize, 1.0f / CubemapSize);
+			PassParameters->RenderTargets[0] = FRenderTargetBinding(IlluminanceMeterCubeTexture, ERenderTargetLoadAction::ENoAction, 0, CubeFace);
+
+			const FIntRect ViewRect(0, 0, CubemapSize, CubemapSize);
+
+			FPixelShaderUtils::AddFullscreenPass(
+				GraphBuilder,
+				MainView.ShaderMap,
+				RDG_EVENT_NAME("FCopyCubemapToIlluminanceMeterCubemap"),
+				PixelShader,
+				PassParameters,
+				ViewRect);
+		}
+	}
+
+	// 2- Generate all the mips of the cubemap, each time storing the sum the contributions to illuminance.
+	{
+		RDG_EVENT_SCOPE(GraphBuilder, "IntegrateCubemapIlluminance");
+
+		TShaderMapRef<FDownsampleIntegrateCubemapIlluminancePS> PixelShader(MainView.ShaderMap);
+		const int32 NumMips = IlluminanceMeterCubeTexture->Desc.NumMips;
+
+		for (int32 MipIndex = 1; MipIndex < NumMips; MipIndex++)
+		{
+			const int32 MipSize = 1 << (NumMips - MipIndex - 1);
+			const FIntRect ViewRect(0, 0, MipSize, MipSize);
+
+			for (int32 CubeFace = 0; CubeFace < CubeFace_MAX; CubeFace++)
+			{
+				auto* PassParameters = GraphBuilder.AllocParameters<FDownsampleIntegrateCubemapIlluminancePS::FParameters>();
+				PassParameters->SourceCubemapTexture = GraphBuilder.CreateSRV(FRDGTextureSRVDesc::CreateForMipLevel(IlluminanceMeterCubeTexture, MipIndex - 1));
+				PassParameters->SourceCubemapSampler = TStaticSamplerState<SF_Point, AM_Clamp, AM_Clamp, AM_Clamp>::GetRHI();
+				PassParameters->CubeFace = CubeFace;
+				PassParameters->MipIndex = MipIndex;
+				PassParameters->NumMips = NumMips;
+				PassParameters->SvPositionToUVScale = FVector2f(1.0f / MipSize, 1.0f / MipSize);
+				PassParameters->RenderTargets[0] = FRenderTargetBinding(IlluminanceMeterCubeTexture, ERenderTargetLoadAction::ENoAction, MipIndex, CubeFace);
+
+				FPixelShaderUtils::AddFullscreenPass(
+					GraphBuilder,
+					MainView.ShaderMap,
+					RDG_EVENT_NAME("CreateCubeMips (Mip: %d, Face: %d)", MipIndex, CubeFace),
+					PixelShader,
+					PassParameters,
+					ViewRect);
+			}
+		}
+	}
+
+	// 3- Sum up each faces and display illuminance on screen
+	{
+		ShaderPrint::SetEnabled(true);
+		ShaderPrint::RequestSpaceForLines(1024);
+		ShaderPrint::RequestSpaceForCharacters(1024);
+
+		FRHIBlendState* PreMultipliedColorTransmittanceBlend = TStaticBlendState<CW_RGB, BO_Add, BF_One, BF_SourceAlpha, BO_Add, BF_Zero, BF_One>::GetRHI();
+
+		FPrintIlluminanceMeterPS::FParameters* PassParameters = GraphBuilder.AllocParameters<FPrintIlluminanceMeterPS::FParameters>();
+		PassParameters->ViewUniformBuffer = MainView.ViewUniformBuffer;
+		PassParameters->SourceCubemapTexture = IlluminanceMeterCubeTexture;
+		PassParameters->SourceCubemapSampler = TStaticSamplerState<SF_Trilinear, AM_Clamp, AM_Clamp, AM_Clamp>::GetRHI();
+		PassParameters->SkyLightCubeTexture = SkyLightTextureRHI;
+		PassParameters->SkyLightCubeSampler = TStaticSamplerState<SF_Trilinear, AM_Clamp, AM_Clamp, AM_Clamp>::GetRHI();
+		PassParameters->CubeMipIndexToSampleIlluminance = IlluminanceMeterCubeTexture->Desc.NumMips - 1;
+		ShaderPrint::SetParameters(GraphBuilder, MainView.ShaderPrintData, PassParameters->ShaderPrintParameters);
+		PassParameters->RenderTargets[0] = FRenderTargetBinding(SceneColorTexture, ERenderTargetLoadAction::ELoad);
+
+		FPrintIlluminanceMeterPS::FPermutationDomain PermutationVector;
+		TShaderMapRef<FPrintIlluminanceMeterPS> PixelShader(MainView.ShaderMap, PermutationVector);
+
+		FPixelShaderUtils::AddFullscreenPass<FPrintIlluminanceMeterPS>(GraphBuilder, MainView.ShaderMap, RDG_EVENT_NAME("Strata::VisualizeMaterial(Draw)"), PixelShader, PassParameters, MainView.ViewRect, PreMultipliedColorTransmittanceBlend);
+	}
+}
diff --git a/Engine/Source/Runtime/Renderer/Private/ScenePrivate.h b/Engine/Source/Runtime/Renderer/Private/ScenePrivate.h
--- a/Engine/Source/Runtime/Renderer/Private/ScenePrivate.h
+++ b/Engine/Source/Runtime/Renderer/Private/ScenePrivate.h
@@ -3101,6 +3101,7 @@ public:
 	virtual void UpdateSkyCaptureContents(const USkyLightComponent* CaptureComponent, bool bCaptureEmissiveOnly, UTextureCube* SourceCubemap, FTexture* OutProcessedTexture, float& OutAverageBrightness, FSHVectorRGB3& OutIrradianceEnvironmentMap, TArray<FFloat16Color>* OutRadianceMap, FLinearColor* SpecifiedCubemapColorScale) override;
 	virtual void AllocateAndCaptureFrameSkyEnvMap(FRDGBuilder& GraphBuilder, FSceneRenderer& SceneRenderer, FViewInfo& MainView, bool bShouldRenderSkyAtmosphere, bool bShouldRenderVolumetricCloud, FInstanceCullingManager& InstanceCullingManager, FRDGExternalAccessQueue& ExternalAccessQueue) override;
 	virtual void ValidateSkyLightRealTimeCapture(FRDGBuilder& GraphBuilder, const FViewInfo& View, FRDGTextureRef SceneColorTexture) override;
+	virtual void ProcessAndRenderIlluminanceMeter(FRDGBuilder& GraphBuilder, TArrayView<FViewInfo> Views, FRDGTextureRef SceneColorTexture);
 	virtual void AddPrecomputedLightVolume(const class FPrecomputedLightVolume* Volume) override;
 	virtual void RemovePrecomputedLightVolume(const class FPrecomputedLightVolume* Volume) override;
 	virtual bool HasPrecomputedVolumetricLightmap_RenderThread() const override;
