Fix curve buffer binding.

diff --git a/Engine/Source/Runtime/Renderer/Private/HairStrands/HairStrandsVisibility.cpp b/Engine/Source/Runtime/Renderer/Private/HairStrands/HairStrandsVisibility.cpp
--- a/Engine/Source/Runtime/Renderer/Private/HairStrands/HairStrandsVisibility.cpp
+++ b/Engine/Source/Runtime/Renderer/Private/HairStrands/HairStrandsVisibility.cpp
@@ -2807,16 +2807,18 @@ FHairStrandsInstanceParameters GetHairStrandsInstanceParameters(FRDGBuilder& Gra
 	{
 		Out.HairStrandsVF_PositionBuffer		= Register(GraphBuilder, VFInput.Strands.PositionBufferExternal, ERDGImportedBufferFlags::CreateSRV).SRV;;
 		Out.HairStrandsVF_PositionOffsetBuffer	= Register(GraphBuilder, VFInput.Strands.PositionOffsetBufferExternal, ERDGImportedBufferFlags::CreateSRV).SRV;;
-
+		Out.HairStrandsVF_CurveBuffer = Register(GraphBuilder, VFInput.Strands.CurveBufferExternal, ERDGImportedBufferFlags::CreateSRV).SRV;
 	}
 	else
 	{
 		Out.HairStrandsVF_PositionBuffer		= VFInput.Strands.PositionBuffer.SRV;
 		Out.HairStrandsVF_PositionOffsetBuffer	= VFInput.Strands.PositionOffsetBuffer.SRV;
+		Out.HairStrandsVF_CurveBuffer = VFInput.Strands.CurveBuffer.SRV;
 	}
 
 
 	Out.HairStrandsVF_VertexCount						= VFInput.Strands.VertexCount;
+	Out.HairStrandsVF_CurveCount						= VFInput.Strands.CurveCount;
 	Out.HairStrandsVF_Radius							= VFInput.Strands.HairRadius;
 	Out.HairStrandsVF_RootScale							= VFInput.Strands.HairRootScale;
 	Out.HairStrandsVF_TipScale							= VFInput.Strands.HairTipScale;
