Recreate SceneTextures uniform buffer if we only create the velocity texture late in the frame for translucency.

diff --git a/Engine/Source/Runtime/Renderer/Private/DeferredShadingRenderer.cpp b/Engine/Source/Runtime/Renderer/Private/DeferredShadingRenderer.cpp
--- a/Engine/Source/Runtime/Renderer/Private/DeferredShadingRenderer.cpp
+++ b/Engine/Source/Runtime/Renderer/Private/DeferredShadingRenderer.cpp
@@ -3834,7 +3834,7 @@ void FDeferredShadingSceneRenderer::Render(FRDGBuilder& GraphBuilder)
 
 		if (bShouldRenderVelocities)
 		{
-			const bool bRecreateSceneTextures = !SceneTextures.Velocity;
+			const bool bRecreateSceneTextures = !HasBeenProduced(SceneTextures.Velocity);
 
 			GraphBuilder.SetCommandListStat(GET_STATID(STAT_CLM_TranslucentVelocity));
 			RenderVelocities(GraphBuilder, SceneTextures, EVelocityPass::Translucent, false);
