Strata - Rough refraction - added a TAA pass to stabilize rough refractions.

diff --git a/Engine/Source/Runtime/Renderer/Private/DeferredShadingRenderer.cpp b/Engine/Source/Runtime/Renderer/Private/DeferredShadingRenderer.cpp
--- a/Engine/Source/Runtime/Renderer/Private/DeferredShadingRenderer.cpp
+++ b/Engine/Source/Runtime/Renderer/Private/DeferredShadingRenderer.cpp
@@ -3829,7 +3829,7 @@ void FDeferredShadingSceneRenderer::Render(FRDGBuilder& GraphBuilder)
 		if (bShouldRenderDistortion)
 		{
 			GraphBuilder.SetCommandListStat(GET_STATID(STAT_CLM_Distortion));
-			RenderDistortion(GraphBuilder, SceneTextures.Color.Target, SceneTextures.Depth.Target, TranslucencyResourceMap);
+			RenderDistortion(GraphBuilder, SceneTextures.Color.Target, SceneTextures.Depth.Target, SceneTextures.Velocity, TranslucencyResourceMap);
 		}
 
 		if (bShouldRenderVelocities)
diff --git a/Engine/Source/Runtime/Renderer/Private/DeferredShadingRenderer.h b/Engine/Source/Runtime/Renderer/Private/DeferredShadingRenderer.h
--- a/Engine/Source/Runtime/Renderer/Private/DeferredShadingRenderer.h
+++ b/Engine/Source/Runtime/Renderer/Private/DeferredShadingRenderer.h
@@ -804,6 +804,7 @@ private:
 	void RenderDistortion(FRDGBuilder& GraphBuilder, 
 		FRDGTextureRef SceneColorTexture, 
 		FRDGTextureRef SceneDepthTexture,
+		FRDGTextureRef SceneVelocityTexture,
 		FTranslucencyPassResourcesMap& TranslucencyResourceMap);
 
 	void CollectLightForTranslucencyLightingVolumeInjection(
diff --git a/Engine/Source/Runtime/Renderer/Private/DistortionRendering.cpp b/Engine/Source/Runtime/Renderer/Private/DistortionRendering.cpp
--- a/Engine/Source/Runtime/Renderer/Private/DistortionRendering.cpp
+++ b/Engine/Source/Runtime/Renderer/Private/DistortionRendering.cpp
@@ -15,6 +15,7 @@
 #include "MeshPassProcessor.inl"
 #include "Strata/Strata.h"
 #include "ScreenRendering.h"
+#include "PostProcess/TemporalAA.h"
 
 DECLARE_GPU_DRAWCALL_STAT(Distortion);
 
@@ -42,6 +43,12 @@ static TAutoConsoleVariable<float> CVarRefractionBlurMaxStandardDeviationInScree
 	TEXT("This will clamp the maximum refraction blur radius on screen."),
 	ECVF_Scalability | ECVF_RenderThreadSafe);
 
+static TAutoConsoleVariable<int32> CVarRefractionBlurTemporalAA(
+	TEXT("r.Refraction.Blur.TemporalAA"),
+	1,
+	TEXT("Enables temporal AA of the scene color buffer in order to avoid flickering in rough refractions."),
+	ECVF_Scalability | ECVF_RenderThreadSafe);
+
 BEGIN_GLOBAL_SHADER_PARAMETER_STRUCT(FDistortionPassUniformParameters, RENDERER_API)
 	SHADER_PARAMETER_STRUCT(FSceneTextureUniformParameters, SceneTextures)
 	SHADER_PARAMETER(FVector4f, DistortionParams)
@@ -450,6 +457,7 @@ void FDeferredShadingSceneRenderer::RenderDistortion(
 	FRDGBuilder& GraphBuilder, 
 	FRDGTextureRef SceneColorTexture, 
 	FRDGTextureRef SceneDepthTexture,
+	FRDGTextureRef SceneVelocityTexture,
 	FTranslucencyPassResourcesMap& TranslucencyResourceMap)
 {
 	check(SceneDepthTexture);
@@ -517,11 +525,33 @@ void FDeferredShadingSceneRenderer::RenderDistortion(
 				TempSceneDepthMeterMipchainTexture = GraphBuilder.CreateTexture(SceneDepthMeterMipchainDesc, TEXT("TempSceneDepthMeterMipchain"));
 			}
 
+			// Apply temporal AA on the current scene buffer to avoid large flickering in rough refraction
+
+			FRDGTextureRef TAASceneColorTexture = SceneColorTexture;
+			if(View.ViewState && CVarRefractionBlurTemporalAA.GetValueOnRenderThread() > 0)
+			{
+				FTAAPassParameters TAASettings(View);
+				TAASettings.SceneDepthTexture = SceneDepthTexture;
+				TAASettings.SceneVelocityTexture = SceneVelocityTexture;
+				TAASettings.Pass = ETAAPassConfig::Main;		// Reusing main config for now. We could add a ReoughRefraction config forcing 111110 format.
+				TAASettings.SceneColorInput = TAASceneColorTexture;
+				TAASettings.bOutputRenderTargetable = true;
+
+				FTAAOutputs TAAOutputs = AddTemporalAAPass(
+					GraphBuilder,
+					View,
+					TAASettings,
+					View.PrevViewInfo.RoughRefractionHistory,
+					&View.ViewState->PrevFrameViewInfo.RoughRefractionHistory);
+
+				TAASceneColorTexture = TAAOutputs.SceneColor;
+			}
+
 			// Copy scene color into the first mip level
 
 			{
 				RDG_EVENT_SCOPE(GraphBuilder, "CopySceneColorDepth");
-				AddCopySceneColorDepthPass(GraphBuilder, View, SceneColorTexture, SceneDepthTexture, SceneColorMipchainTexture, SceneDepthMeterMipchainTexture);
+				AddCopySceneColorDepthPass(GraphBuilder, View, TAASceneColorTexture, SceneDepthTexture, SceneColorMipchainTexture, SceneDepthMeterMipchainTexture);
 			}
 
 			// Now render the mip chain
diff --git a/Engine/Source/Runtime/Renderer/Private/RendererScene.cpp b/Engine/Source/Runtime/Renderer/Private/RendererScene.cpp
--- a/Engine/Source/Runtime/Renderer/Private/RendererScene.cpp
+++ b/Engine/Source/Runtime/Renderer/Private/RendererScene.cpp
@@ -949,6 +949,7 @@ uint64 FPreviousViewInfo::GetGPUSizeBytes(bool bLogSizes) const
 		DOFSetupHistory.GetGPUSizeBytes(bLogSizes) +
 		SSRHistory.GetGPUSizeBytes(bLogSizes) +
 		WaterSSRHistory.GetGPUSizeBytes(bLogSizes) +
+		RoughRefractionHistory.GetGPUSizeBytes(bLogSizes) +
 		HairHistory.GetGPUSizeBytes(bLogSizes) +
 		EditorPrimtiveDepthHistory.GetGPUSizeBytes(bLogSizes) +
 		CustomSSRInput.GetGPUSizeBytes(bLogSizes) +
diff --git a/Engine/Source/Runtime/Renderer/Private/SceneRendering.h b/Engine/Source/Runtime/Renderer/Private/SceneRendering.h
--- a/Engine/Source/Runtime/Renderer/Private/SceneRendering.h
+++ b/Engine/Source/Runtime/Renderer/Private/SceneRendering.h
@@ -1122,6 +1122,9 @@ struct FPreviousViewInfo
 	FTemporalAAHistory SSRHistory;
 	FTemporalAAHistory WaterSSRHistory;
 
+	// Temporal AA history for Rough refraction
+	FTemporalAAHistory RoughRefractionHistory;
+
 	// Temporal AA history for Hair
 	FTemporalAAHistory HairHistory;
 
