Fixed crash on using RVT on platfrom with a GPUScene disabled

diff --git a/Engine/Source/Runtime/Renderer/Private/GPUScene.cpp b/Engine/Source/Runtime/Renderer/Private/GPUScene.cpp
--- a/Engine/Source/Runtime/Renderer/Private/GPUScene.cpp
+++ b/Engine/Source/Runtime/Renderer/Private/GPUScene.cpp
@@ -1631,8 +1631,12 @@ void FGPUScene::UploadDynamicPrimitiveShaderDataForViewInternal(FRDGBuilder& Gra
 
 bool FGPUScene::FillViewShaderParameters(FViewUniformShaderParameters& OutParameters)
 {
+	if (!bIsEnabled)
+	{
+		return false;
+	}
+	
 	bool bParametersChanged = false;
-
 	const auto SetParameter = [&](auto& OutParameter, const auto& InParameter)
 	{
 		bParametersChanged |= OutParameter != InParameter;
