Strata - Rough refraction
Added max pre exposed illuminance cvar to reduce flickring even more.
Removed rough refraction noise for since it seems to not be resolved well.

diff --git a/Engine/Shaders/Private/DistortApplyScreenPS.usf b/Engine/Shaders/Private/DistortApplyScreenPS.usf
--- a/Engine/Shaders/Private/DistortApplyScreenPS.usf
+++ b/Engine/Shaders/Private/DistortApplyScreenPS.usf
@@ -127,8 +127,9 @@ void Main(
 
 		const float PixelStandardDevInPixel = StandardDeviationCmForThickness / PixelRadiusCm;
 
-		// Jitter UV according to MIP level to hide texel.
-		if(true)
+		// Jitter UV according to MIP level to hide texel
+		// => disabled because it is a bit unpleasant over bright area (TSR seems to have trouble resolving it)
+		if(false)
 		{
 			const float BlurRandomness = InterleavedGradientNoise(SVPos.xy, float(View.StateFrameIndexMod8)) - 0.5; // In [-0.5, 0.5]
 			const float RandAngle = BlurRandomness * 2.0f * PI;
diff --git a/Engine/Shaders/Private/DistortFiltering.usf b/Engine/Shaders/Private/DistortFiltering.usf
--- a/Engine/Shaders/Private/DistortFiltering.usf
+++ b/Engine/Shaders/Private/DistortFiltering.usf
@@ -16,6 +16,8 @@ SamplerState SceneColorSampler;
 Texture2D SceneDepthTexture;
 SamplerState SceneDepthSampler;
 
+float MaxExposedLuminance;
+
 void CopySceneColorDepthTexturePS(FScreenVertexOutput Input, out float4 OutColor : SV_Target0, out float OutDepthMeter : SV_Target1)
 {
 	float2 SVPos = Input.UV * View.ViewSizeAndInvSize.xy + View.ViewRectMin.xy;
@@ -23,6 +25,9 @@ void CopySceneColorDepthTexturePS(FScreenVertexOutput Input, out float4 OutColor
 	float2 TextureUV = SVPos * View.BufferSizeAndInvSize.zw;
 	float4 LinearColor = Texture2DSample(SceneColorTexture, SceneColorSampler, TextureUV);
 
+	// Clamping to a certain max luminance help to reduce bright specular highlight flickering, even with temporal AA stabilization.
+	LinearColor.rgb = min(LinearColor, MaxExposedLuminance.xxx);
+
 	float DeviceZ = Texture2DSample(SceneDepthTexture, SceneDepthSampler, TextureUV).r;
 #if HAS_INVERTED_Z_BUFFER
 	DeviceZ = max(0.000000000001, DeviceZ);	// TODO: avoid bad values when DeviceZ is far=0 when using inverted z
diff --git a/Engine/Source/Runtime/Renderer/Private/DistortionRendering.cpp b/Engine/Source/Runtime/Renderer/Private/DistortionRendering.cpp
--- a/Engine/Source/Runtime/Renderer/Private/DistortionRendering.cpp
+++ b/Engine/Source/Runtime/Renderer/Private/DistortionRendering.cpp
@@ -49,6 +49,12 @@ static TAutoConsoleVariable<int32> CVarRefractionBlurTemporalAA(
 	TEXT("Enables temporal AA of the scene color buffer in order to avoid flickering in rough refractions."),
 	ECVF_Scalability | ECVF_RenderThreadSafe);
 
+static TAutoConsoleVariable<float> CVarRefractionBlurMaxExposedLuminance(
+	TEXT("r.Refraction.Blur.MaxExposedLuminance"),
+	10.0f,
+	TEXT("Clamp scene pre-exposed luminance to this maximum value. It helps to reduce bright specular highlight flickering, even when r.Refraction.Blur.TemporalAA=1."),
+	ECVF_Scalability | ECVF_RenderThreadSafe);
+
 BEGIN_GLOBAL_SHADER_PARAMETER_STRUCT(FDistortionPassUniformParameters, RENDERER_API)
 	SHADER_PARAMETER_STRUCT(FSceneTextureUniformParameters, SceneTextures)
 	SHADER_PARAMETER(FVector4f, DistortionParams)
@@ -260,6 +266,7 @@ public:
 		SHADER_PARAMETER_RDG_TEXTURE(Texture2D, SceneDepthTexture)
 		SHADER_PARAMETER_SAMPLER(SamplerState, SceneColorSampler)
 		SHADER_PARAMETER_SAMPLER(SamplerState, SceneDepthSampler)
+		SHADER_PARAMETER(float, MaxExposedLuminance)
 		RENDER_TARGET_BINDING_SLOTS()
 	END_SHADER_PARAMETER_STRUCT()
 
@@ -291,6 +298,7 @@ static void AddCopySceneColorDepthPass(FRDGBuilder& GraphBuilder, const FViewInf
 	PassParameters->SceneDepthSampler = TStaticSamplerState<SF_Point>::GetRHI();
 	PassParameters->RenderTargets[0] = FRenderTargetBinding(SceneColorCopyTexture, ERenderTargetLoadAction::ENoAction);
 	PassParameters->RenderTargets[1] = FRenderTargetBinding(SceneDepthCopyTexture, ERenderTargetLoadAction::ENoAction);
+	PassParameters->MaxExposedLuminance = FMath::Max(0.0f, CVarRefractionBlurMaxExposedLuminance.GetValueOnRenderThread());
 
 	// The scene color is copied into from multi-view-rect layout into a single-rect layout.
 	FIntRect ViewRect;
