Move FReflectionUniformParameters to RDG so that RDG resources are properly marked as used in a given pass

diff --git a/Engine/Source/Runtime/Renderer/Private/HairStrands/HairStrandsEnvironment.cpp b/Engine/Source/Runtime/Renderer/Private/HairStrands/HairStrandsEnvironment.cpp
--- a/Engine/Source/Runtime/Renderer/Private/HairStrands/HairStrandsEnvironment.cpp
+++ b/Engine/Source/Runtime/Renderer/Private/HairStrands/HairStrandsEnvironment.cpp
@@ -255,7 +255,7 @@ class FHairEnvironmentLightingPS : public FGlobalShader
 		SHADER_PARAMETER_RDG_TEXTURE(Texture2D, HairEnergyLUTTexture)
 
 		SHADER_PARAMETER_STRUCT_REF(FViewUniformShaderParameters, ViewUniformBuffer)
-		SHADER_PARAMETER_STRUCT_REF(FReflectionUniformParameters, ReflectionsParameters)
+		SHADER_PARAMETER_RDG_UNIFORM_BUFFER(FReflectionUniformParameters, ReflectionsParameters)
 		SHADER_PARAMETER_STRUCT_REF(FReflectionCaptureShaderData, ReflectionCaptureData)
 		SHADER_PARAMETER_RDG_UNIFORM_BUFFER(FForwardLightData, ForwardLightData)
 		SHADER_PARAMETER_RDG_UNIFORM_BUFFER(FHairStrandsViewUniformParameters, HairStrands)
@@ -384,12 +384,8 @@ static void AddHairStrandsEnvironmentLightingPassPS(
 	ParametersPS->HairStrands = HairStrands::BindHairStrandsViewUniformParameters(View);
 	ParametersPS->ViewUniformBuffer = View.ViewUniformBuffer;
 	ParametersPS->ReflectionCaptureData = View.ReflectionCaptureUniformBuffer;
-	{
-		FReflectionUniformParameters ReflectionUniformParameters;
-		SetupReflectionUniformParameters(GraphBuilder, View, ReflectionUniformParameters);
-		ParametersPS->ReflectionsParameters = CreateUniformBufferImmediate(ReflectionUniformParameters, UniformBuffer_SingleDraw);
-	}
-
+	ParametersPS->ReflectionsParameters = CreateReflectionUniformBuffer(GraphBuilder, View);
+	
 	if (LightingType == EHairLightingSourceType::Lumen)
 	{
 		const FRadianceCacheState& RadianceCacheState = View.ViewState->Lumen.RadianceCacheState;
diff --git a/Engine/Source/Runtime/Renderer/Private/IndirectLightRendering.cpp b/Engine/Source/Runtime/Renderer/Private/IndirectLightRendering.cpp
--- a/Engine/Source/Runtime/Renderer/Private/IndirectLightRendering.cpp
+++ b/Engine/Source/Runtime/Renderer/Private/IndirectLightRendering.cpp
@@ -356,7 +356,7 @@ class FReflectionEnvironmentSkyLightingPS : public FGlobalShader
 		SHADER_PARAMETER_STRUCT_INCLUDE(FSceneTextureParameters, SceneTextures)
 
 		SHADER_PARAMETER_STRUCT_REF(FViewUniformShaderParameters, ViewUniformBuffer)
-		SHADER_PARAMETER_STRUCT_REF(FReflectionUniformParameters, ReflectionsParameters)
+		SHADER_PARAMETER_RDG_UNIFORM_BUFFER(FReflectionUniformParameters, ReflectionsParameters)
 		SHADER_PARAMETER_STRUCT_REF(FReflectionCaptureShaderData, ReflectionCaptureData)
 		SHADER_PARAMETER_RDG_UNIFORM_BUFFER(FForwardLightData, ForwardLightData)
 
@@ -1593,11 +1593,7 @@ static void AddSkyReflectionPass(
 
 		PassParameters->PS.ViewUniformBuffer = View.ViewUniformBuffer;
 		PassParameters->PS.ReflectionCaptureData = View.ReflectionCaptureUniformBuffer;
-		{
-			FReflectionUniformParameters ReflectionUniformParameters;
-			SetupReflectionUniformParameters(GraphBuilder, View, ReflectionUniformParameters);
-			PassParameters->PS.ReflectionsParameters = CreateUniformBufferImmediate(ReflectionUniformParameters, UniformBuffer_SingleDraw);
-		}
+		PassParameters->PS.ReflectionsParameters = CreateReflectionUniformBuffer(GraphBuilder, View);
 		PassParameters->PS.ForwardLightData = View.ForwardLightingResources.ForwardLightUniformBuffer;
 
 		PassParameters->PS.Strata = Strata::BindStrataGlobalUniformParameters(View);
diff --git a/Engine/Source/Runtime/Renderer/Private/SingleLayerWaterRendering.cpp b/Engine/Source/Runtime/Renderer/Private/SingleLayerWaterRendering.cpp
--- a/Engine/Source/Runtime/Renderer/Private/SingleLayerWaterRendering.cpp
+++ b/Engine/Source/Runtime/Renderer/Private/SingleLayerWaterRendering.cpp
@@ -275,7 +275,7 @@ BEGIN_SHADER_PARAMETER_STRUCT(FSingleLayerWaterCommonShaderParameters, )
 	SHADER_PARAMETER_STRUCT_INCLUDE(FSceneTextureParameters, SceneTextures)	// Water scene texture
 	SHADER_PARAMETER_STRUCT_INCLUDE(FViewShaderParameters, View)
 	SHADER_PARAMETER_STRUCT_REF(FReflectionCaptureShaderData, ReflectionCaptureData)
-	SHADER_PARAMETER_STRUCT_REF(FReflectionUniformParameters, ReflectionsParameters)
+	SHADER_PARAMETER_RDG_UNIFORM_BUFFER(FReflectionUniformParameters, ReflectionsParameters)
 	SHADER_PARAMETER_RDG_UNIFORM_BUFFER(FForwardLightData, ForwardLightData)
 	SHADER_PARAMETER_RDG_UNIFORM_BUFFER(FStrataGlobalUniformParameters, Strata)
 END_SHADER_PARAMETER_STRUCT()
@@ -872,11 +872,7 @@ void FDeferredShadingSceneRenderer::RenderSingleLayerWaterReflections(
 			Parameters.SceneTextures = SceneTextureParameters;
 			Parameters.View = View.GetShaderParameters();
 			Parameters.ReflectionCaptureData = View.ReflectionCaptureUniformBuffer;
-			{
-				FReflectionUniformParameters ReflectionUniformParameters;
-				SetupReflectionUniformParameters(GraphBuilder, View, ReflectionUniformParameters);
-				Parameters.ReflectionsParameters = CreateUniformBufferImmediate(ReflectionUniformParameters, UniformBuffer_SingleDraw);
-			}
+			Parameters.ReflectionsParameters = CreateReflectionUniformBuffer(GraphBuilder, View);
 			Parameters.ForwardLightData = View.ForwardLightingResources.ForwardLightUniformBuffer;
 			Parameters.Strata = Strata::BindStrataGlobalUniformParameters(View);
 		};
