Update VRS combiner to use a bi-dimensional max.

diff --git a/Engine/Shaders/Private/VariableRateShading/VRSShadingRateCombine.usf b/Engine/Shaders/Private/VariableRateShading/VRSShadingRateCombine.usf
--- a/Engine/Shaders/Private/VariableRateShading/VRSShadingRateCombine.usf
+++ b/Engine/Shaders/Private/VariableRateShading/VRSShadingRateCombine.usf
@@ -3,17 +3,34 @@
 #include "../Common.ush"
 
 RWTexture2D<uint> RWOutputTexture;
-Texture2D<uint> SourceTexture_1;
-Texture2D<uint> SourceTexture_2;
+Texture2D<uint> SourceTexture0;
+Texture2D<uint> SourceTexture1;
 
+uint2 DecodeShadingRate(uint EncodedRate)
+{
+	uint2 DecodedRate;
+	DecodedRate.x = ((1U << SHADING_RATE_DIMENSION_BITS) - 1U) & (EncodedRate >> SHADING_RATE_DIMENSION_BITS);
+	DecodedRate.y = ((1U << SHADING_RATE_DIMENSION_BITS) - 1U) & EncodedRate;
+	return DecodedRate;
+}
+
+uint2 EncodeShadingRate(uint2 DecodedRate)
+{
+	return (DecodedRate.x << SHADING_RATE_DIMENSION_BITS) + DecodedRate.y;
+}
+
+uint ShadingRateMax(uint InputRate0, uint InputRate1)
+{
+	uint2 DecodedRate0 = DecodeShadingRate(InputRate0);
+	uint2 DecodedRate1 = DecodeShadingRate(InputRate1);
+	uint2 DecodedMaxRate = max(DecodedRate0, DecodedRate1);
+	return EncodeShadingRate(DecodedMaxRate);
+}
 
 [numthreads(THREADGROUP_SIZEX, THREADGROUP_SIZEY, 1)]
 void CombineShadingRateTextures(uint3 DispatchThreadId : SV_DispatchThreadID)
 {
 	const uint2 TexelCoord = DispatchThreadId.xy;
-	uint ShadingRateOut = 0;
-	
-	ShadingRateOut = max(SourceTexture_1[TexelCoord], SourceTexture_2[TexelCoord]);
-	
-	RWOutputTexture[TexelCoord] = ShadingRateOut;
-}
\ No newline at end of file
+	RWOutputTexture[TexelCoord] = ShadingRateMax(SourceTexture0[TexelCoord], SourceTexture1[TexelCoord]);
+}
+
diff --git a/Engine/Source/Runtime/Renderer/Private/VariableRateShading/VariableRateShadingImageManager.cpp b/Engine/Source/Runtime/Renderer/Private/VariableRateShading/VariableRateShadingImageManager.cpp
--- a/Engine/Source/Runtime/Renderer/Private/VariableRateShading/VariableRateShadingImageManager.cpp
+++ b/Engine/Source/Runtime/Renderer/Private/VariableRateShading/VariableRateShadingImageManager.cpp
@@ -29,6 +29,7 @@ TGlobalResource<FVariableRateShadingImageManager> GVRSImageManager;
 namespace VRSHelpers
 {
 	constexpr int32 kCombineGroupSize = FComputeShaderUtils::kGolden2DGroupSize;
+	constexpr uint32 kShadingRateDimensionBits = 2;
 }
 
 class FCombineShadingRateTexturesCS : public FGlobalShader
@@ -39,10 +40,8 @@ public:
 
 	BEGIN_SHADER_PARAMETER_STRUCT(FParameters, )
 		SHADER_PARAMETER_RDG_TEXTURE_UAV(RWTexture2D<uint>, RWOutputTexture)
-		SHADER_PARAMETER_RDG_TEXTURE(Texture2D<uint>, SourceTexture_1)
-		SHADER_PARAMETER_RDG_TEXTURE(Texture2D<uint>, SourceTexture_2)
-		SHADER_PARAMETER_RDG_TEXTURE(Texture2D<uint>, SourceTexture_3)
-		SHADER_PARAMETER_RDG_TEXTURE(Texture2D<uint>, SourceTexture_4)
+		SHADER_PARAMETER_RDG_TEXTURE(Texture2D<uint>, SourceTexture0)
+		SHADER_PARAMETER_RDG_TEXTURE(Texture2D<uint>, SourceTexture1)
 	END_SHADER_PARAMETER_STRUCT()
 
 	static bool ShouldCompilePermutation(const FGlobalShaderPermutationParameters& Parameters)
@@ -56,6 +55,8 @@ public:
 
 		OutEnvironment.SetDefine(TEXT("THREADGROUP_SIZEX"), VRSHelpers::kCombineGroupSize);
 		OutEnvironment.SetDefine(TEXT("THREADGROUP_SIZEY"), VRSHelpers::kCombineGroupSize);
+
+		OutEnvironment.SetDefine(TEXT("SHADING_RATE_DIMENSION_BITS"), VRSHelpers::kShadingRateDimensionBits);
 	}
 
 };
@@ -159,8 +160,8 @@ FRDGTextureRef FVariableRateShadingImageManager::GetVariableRateShadingImage(FRD
 		FRDGTextureRef CombinedShadingRateTexture = GraphBuilder.CreateTexture(Desc, TEXT("CombinedShadingRateTexture"));
 
 		FCombineShadingRateTexturesCS::FParameters* PassParameters = GraphBuilder.AllocParameters<FCombineShadingRateTexturesCS::FParameters>();
-		PassParameters->SourceTexture_2 = InternalVRSSources[0];
-		PassParameters->SourceTexture_1 = InternalVRSSources[1];
+		PassParameters->SourceTexture0 = InternalVRSSources[0];
+		PassParameters->SourceTexture1 = InternalVRSSources[1];
 		PassParameters->RWOutputTexture = GraphBuilder.CreateUAV(CombinedShadingRateTexture);
 
 		TShaderMapRef<FCombineShadingRateTexturesCS> ComputeShader(GetGlobalShaderMap(GMaxRHIFeatureLevel));
