Fix issue where movable mesh initializes so that it always writes velocity.
This happens randomly on scene start for meshs, and can be seen by taking a simple scene and going in and out of PIE after applying the "vis scenevelocity" command.
When the issue happens there are no visual issues because we still write zero velocity. But the debug view modes are less clear, because the meshes are shown as writing velocity when they shouldn't be.

diff --git a/Engine/Source/Runtime/Renderer/Private/RendererScene.cpp b/Engine/Source/Runtime/Renderer/Private/RendererScene.cpp
--- a/Engine/Source/Runtime/Renderer/Private/RendererScene.cpp
+++ b/Engine/Source/Runtime/Renderer/Private/RendererScene.cpp
@@ -3714,15 +3714,18 @@ void FSceneVelocityData::StartFrame(FScene* Scene)
 
 void FScene::GetPrimitiveUniformShaderParameters_RenderThread(const FPrimitiveSceneInfo* PrimitiveSceneInfo, bool& bHasPrecomputedVolumetricLightmap, FMatrix& PreviousLocalToWorld, int32& SingleCaptureIndex, bool& bOutputVelocity) const 
 {
-	bHasPrecomputedVolumetricLightmap = VolumetricLightmapSceneData.HasData();
-	
-	bOutputVelocity = VelocityData.GetComponentPreviousLocalToWorld(PrimitiveSceneInfo->PrimitiveComponentId, PreviousLocalToWorld);
+	const FMatrix LocalToWorld = PrimitiveSceneInfo->Proxy->GetLocalToWorld();
+	PreviousLocalToWorld = LocalToWorld;
+	bOutputVelocity = false;
 
-	if (!bOutputVelocity)
+	const bool bHasPreviousLocalToWorld = VelocityData.GetComponentPreviousLocalToWorld(PrimitiveSceneInfo->PrimitiveComponentId, PreviousLocalToWorld);
+	if (bHasPreviousLocalToWorld)
 	{
-		PreviousLocalToWorld = PrimitiveSceneInfo->Proxy->GetLocalToWorld();
+		bOutputVelocity = !LocalToWorld.Equals(PreviousLocalToWorld, 0.0001f);
 	}
 
+	bHasPrecomputedVolumetricLightmap = VolumetricLightmapSceneData.HasData();
+
 	// Get index if proxy exists, otherwise fall back to index 0 which will contain the default black cubemap
 	SingleCaptureIndex = PrimitiveSceneInfo->CachedReflectionCaptureProxy ? PrimitiveSceneInfo->CachedReflectionCaptureProxy->SortedCaptureIndex : 0;
 }
