Multi-GPU:  Fix for scene captures not being transferred to second GPU.  Need to set output RenderTargetGPUMask in the scene renderer to all GPUs the render target exists on, rather than just GPU0.  This effectively reverts the change from CL 20540730 (which made cross GPU transfers conditional on AFR), just for scene captures.

#rnx

diff --git a/Engine/Source/Runtime/Renderer/Private/SceneRendering.cpp b/Engine/Source/Runtime/Renderer/Private/SceneRendering.cpp
--- a/Engine/Source/Runtime/Renderer/Private/SceneRendering.cpp
+++ b/Engine/Source/Runtime/Renderer/Private/SceneRendering.cpp
@@ -2944,7 +2944,12 @@ void FSceneRenderer::ComputeGPUMasks(FRHICommandListImmediate* RHICmdList)
 
 	RenderTargetGPUMask = FRHIGPUMask::GPU0();
 	
-	if ((GNumExplicitGPUsForRendering > 1) && (GNumAlternateFrameRenderingGroups > 1) && ViewFamily.RenderTarget)
+	// Scene capture render targets should be propagated to all GPUs the render target exists on.  For other render targets
+	// (like nDisplay outputs), we default them to only be copied to GPU0, for performance.
+	//
+	// TODO:  we should remove this conditional, and set the GPU mask for the source render targets, but the goal is to have
+	// a minimal scope CL for the 5.1.1 hot fix.  This effectively reverts the change from CL 20540730, just for scene captures.
+	if ((GNumExplicitGPUsForRendering > 1) && ViewFamily.RenderTarget && ((GNumAlternateFrameRenderingGroups > 1) || Views[0].bIsSceneCapture))
 	{
 		// Get render target GPU mask, taking into account AFR
 		check(RHICmdList);
