Fixed non-unity/pch compile errors and also moved some forward declarations/includes to the top of the file instead of in the middle of the file

diff --git a/Engine/Source/Runtime/Renderer/Private/RendererPrivateUtils.h b/Engine/Source/Runtime/Renderer/Private/RendererPrivateUtils.h
--- a/Engine/Source/Runtime/Renderer/Private/RendererPrivateUtils.h
+++ b/Engine/Source/Runtime/Renderer/Private/RendererPrivateUtils.h
@@ -3,6 +3,7 @@
 #pragma once
 
 #include "RenderResource.h"
+#include "RHI.h"
 
 class FTileTexCoordVertexBuffer : public FVertexBuffer
 {
diff --git a/Engine/Source/Runtime/Renderer/Private/SpriteIndexBuffer.h b/Engine/Source/Runtime/Renderer/Private/SpriteIndexBuffer.h
--- a/Engine/Source/Runtime/Renderer/Private/SpriteIndexBuffer.h
+++ b/Engine/Source/Runtime/Renderer/Private/SpriteIndexBuffer.h
@@ -3,6 +3,7 @@
 #pragma once
 
 #include "RenderResource.h"
+#include "RHI.h"
 
 template< uint32 NumSprites >
 class FSpriteIndexBuffer : public FIndexBuffer
