Ensure hair forward raster is only enabled on SM6.

diff --git a/Engine/Shaders/Private/HairStrands/HairStrandsForwardRaster.usf b/Engine/Shaders/Private/HairStrands/HairStrandsForwardRaster.usf
--- a/Engine/Shaders/Private/HairStrands/HairStrandsForwardRaster.usf
+++ b/Engine/Shaders/Private/HairStrands/HairStrandsForwardRaster.usf
@@ -1024,7 +1024,6 @@ void RasterCS(
 	uint2 GroupThread2D    : SV_GroupThreadID,  /* 8x8 */
 	uint  GroupID          : SV_GroupID)		/* Rasterizer ID */
 {
-	#if 0
 	ResolvedView = ResolveView();
 	const uint BinTileNum = VisTileArgs[0];
 
@@ -1666,7 +1665,6 @@ void RasterCS(
 		} // 4. Raster
 
 	}// Work item
-	#endif
 }
 
 #endif //SHADER_RASTERCOMPUTE_RASTER
diff --git a/Engine/Source/Runtime/Renderer/Private/HairStrands/HairStrandsDebug.cpp b/Engine/Source/Runtime/Renderer/Private/HairStrands/HairStrandsDebug.cpp
--- a/Engine/Source/Runtime/Renderer/Private/HairStrands/HairStrandsDebug.cpp
+++ b/Engine/Source/Runtime/Renderer/Private/HairStrands/HairStrandsDebug.cpp
@@ -1179,7 +1179,7 @@ static void AddDrawDebugClusterPass(
 }
 
 ///////////////////////////////////////////////////////////////////////////////////////////////////
-uint32 GetHairStrandsMeanSamplePerPixel();
+uint32 GetHairStrandsMeanSamplePerPixel(EShaderPlatform In);
 static void InternalRenderHairStrandsDebugInfo(
 	FRDGBuilder& GraphBuilder,
 	FScene* Scene,
@@ -1352,7 +1352,7 @@ static void InternalRenderHairStrandsDebugInfo(
 
 		const FIntPoint PPLLResolution = HairData.DebugData.PPLLData.NodeIndexTexture->Desc.Extent;
 		FHairVisibilityDebugPPLLCS::FParameters* PassParameters = GraphBuilder.AllocParameters<FHairVisibilityDebugPPLLCS::FParameters>();
-		PassParameters->PPLLMeanListElementCountPerPixel = GetHairStrandsMeanSamplePerPixel();
+		PassParameters->PPLLMeanListElementCountPerPixel = GetHairStrandsMeanSamplePerPixel(View.GetShaderPlatform());
 		PassParameters->PPLLMaxTotalListElementCount = HairData.DebugData.PPLLData.NodeDataBuffer->Desc.NumElements;
 		PassParameters->PPLLCounter = HairData.DebugData.PPLLData.NodeCounterTexture;
 		PassParameters->PPLLNodeIndex = HairData.DebugData.PPLLData.NodeIndexTexture;
diff --git a/Engine/Source/Runtime/Renderer/Private/HairStrands/HairStrandsForwardRaster.cpp b/Engine/Source/Runtime/Renderer/Private/HairStrands/HairStrandsForwardRaster.cpp
--- a/Engine/Source/Runtime/Renderer/Private/HairStrands/HairStrandsForwardRaster.cpp
+++ b/Engine/Source/Runtime/Renderer/Private/HairStrands/HairStrandsForwardRaster.cpp
@@ -10,6 +10,7 @@
 #include "RenderGraphUtils.h"
 #include "HairStrandsInterface.h"
 #include "SceneRendering.h"
+#include "DataDrivenShaderPlatformInfo.h"
 
 extern int32 GHairVisibilityComputeRaster_Culling;
 extern int32 GHairVisibilityComputeRaster_MaxTiles;
@@ -23,6 +24,10 @@ extern int32 GHairVisibilityComputeRaster_Debug;
 void SetUpViewHairRenderInfo(const FViewInfo& ViewInfo, FVector4f& OutHairRenderInfo, uint32& OutHairRenderInfoBits, uint32& OutHairComponents);
 void SetUpViewHairRenderInfo(const FViewInfo& ViewInfo, bool bEnableMSAA, FVector4f& OutHairRenderInfo, uint32& OutHairRenderInfoBits, uint32& OutHairComponents);
 
+inline bool IsHairStrandsForwardRasterSupported(EShaderPlatform In) 
+{ 
+	return IsFeatureLevelSupported(In, ERHIFeatureLevel::SM6) && IsHairStrandsSupported(EHairStrandsShaderType::Strands, In); 
+}
 
 /////////////////////////////////////////////////////////////////////////////////////////
 // Culling pass
@@ -48,7 +53,7 @@ class FHairCullSegmentCS : public FGlobalShader
 		END_SHADER_PARAMETER_STRUCT()
 
 public:
-	static bool ShouldCompilePermutation(const FGlobalShaderPermutationParameters& Parameters) { return IsHairStrandsSupported(EHairStrandsShaderType::Strands, Parameters.Platform); }
+	static bool ShouldCompilePermutation(const FGlobalShaderPermutationParameters& Parameters) { return IsHairStrandsForwardRasterSupported(Parameters.Platform); }
 	static void ModifyCompilationEnvironment(const FGlobalShaderPermutationParameters& Parameters, FShaderCompilerEnvironment& OutEnvironment)
 	{
 		FGlobalShader::ModifyCompilationEnvironment(Parameters, OutEnvironment);
@@ -82,10 +87,10 @@ END_SHADER_PARAMETER_STRUCT()
 ///////////////////////////////////////////////////////////////////////////////////////////////////
 // Compute depth tile data based on scene data
 
-class FVisiblityRasterComputeForwardPrepareDepthGridCS : public FGlobalShader
+class FHairStrandsForwardRasterPrepareDepthGridCS : public FGlobalShader
 {
-	DECLARE_GLOBAL_SHADER(FVisiblityRasterComputeForwardPrepareDepthGridCS);
-	SHADER_USE_PARAMETER_STRUCT(FVisiblityRasterComputeForwardPrepareDepthGridCS, FGlobalShader);
+	DECLARE_GLOBAL_SHADER(FHairStrandsForwardRasterPrepareDepthGridCS);
+	SHADER_USE_PARAMETER_STRUCT(FHairStrandsForwardRasterPrepareDepthGridCS, FGlobalShader);
 
 	BEGIN_SHADER_PARAMETER_STRUCT(FParameters, )
 		SHADER_PARAMETER_STRUCT_INCLUDE(FRasterComputeForwardCommonParameters, Common)
@@ -94,7 +99,7 @@ class FVisiblityRasterComputeForwardPrepareDepthGridCS : public FGlobalShader
 	END_SHADER_PARAMETER_STRUCT()
 
 public:
-	static bool ShouldCompilePermutation(const FGlobalShaderPermutationParameters& Parameters) { return IsHairStrandsSupported(EHairStrandsShaderType::Strands, Parameters.Platform); }
+	static bool ShouldCompilePermutation(const FGlobalShaderPermutationParameters& Parameters) { return IsHairStrandsForwardRasterSupported(Parameters.Platform); }
 	static void ModifyCompilationEnvironment(const FGlobalShaderPermutationParameters& Parameters, FShaderCompilerEnvironment& OutEnvironment)
 	{
 		FGlobalShader::ModifyCompilationEnvironment(Parameters, OutEnvironment);
@@ -102,15 +107,15 @@ public:
 	}
 };
 
-IMPLEMENT_GLOBAL_SHADER(FVisiblityRasterComputeForwardPrepareDepthGridCS, "/Engine/Private/HairStrands/HairStrandsForwardRaster.usf", "PrepareDepthGridCS", SF_Compute);
+IMPLEMENT_GLOBAL_SHADER(FHairStrandsForwardRasterPrepareDepthGridCS, "/Engine/Private/HairStrands/HairStrandsForwardRaster.usf", "PrepareDepthGridCS", SF_Compute);
 
 /////////////////////////////////////////////////////////////////////////////////////////
 // Bin strands segments
 
-class FVisiblityRasterComputeForwardBinningCS : public FGlobalShader
+class FHairStrandsForwardRasterBinningCS : public FGlobalShader
 {
-	DECLARE_GLOBAL_SHADER(FVisiblityRasterComputeForwardBinningCS);
-	SHADER_USE_PARAMETER_STRUCT(FVisiblityRasterComputeForwardBinningCS, FGlobalShader);
+	DECLARE_GLOBAL_SHADER(FHairStrandsForwardRasterBinningCS);
+	SHADER_USE_PARAMETER_STRUCT(FHairStrandsForwardRasterBinningCS, FGlobalShader);
 
 	class FDebug : SHADER_PERMUTATION_BOOL("PERMUTATION_DEBUG");
 	using FPermutationDomain = TShaderPermutationDomain<FDebug>;
@@ -129,7 +134,7 @@ class FVisiblityRasterComputeForwardBinningCS : public FGlobalShader
 	END_SHADER_PARAMETER_STRUCT()
 
 public:
-	static bool ShouldCompilePermutation(const FGlobalShaderPermutationParameters& Parameters) { return IsHairStrandsSupported(EHairStrandsShaderType::Strands, Parameters.Platform); }
+	static bool ShouldCompilePermutation(const FGlobalShaderPermutationParameters& Parameters) { return IsHairStrandsForwardRasterSupported(Parameters.Platform); }
 	static void ModifyCompilationEnvironment(const FGlobalShaderPermutationParameters& Parameters, FShaderCompilerEnvironment& OutEnvironment)
 	{
 		FGlobalShader::ModifyCompilationEnvironment(Parameters, OutEnvironment);
@@ -137,15 +142,15 @@ public:
 	}
 };
 
-IMPLEMENT_GLOBAL_SHADER(FVisiblityRasterComputeForwardBinningCS, "/Engine/Private/HairStrands/HairStrandsForwardRaster.usf", "BinningCS", SF_Compute);
+IMPLEMENT_GLOBAL_SHADER(FHairStrandsForwardRasterBinningCS, "/Engine/Private/HairStrands/HairStrandsForwardRaster.usf", "BinningCS", SF_Compute);
 
 /////////////////////////////////////////////////////////////////////////////////////////
 // Compacted binned segments into contiguous list 
 
-class FVisiblityRasterComputeForwardCompactionCS : public FGlobalShader
+class FHairStrandsForwardRasterCompactionCS : public FGlobalShader
 {
-	DECLARE_GLOBAL_SHADER(FVisiblityRasterComputeForwardCompactionCS);
-	SHADER_USE_PARAMETER_STRUCT(FVisiblityRasterComputeForwardCompactionCS, FGlobalShader);
+	DECLARE_GLOBAL_SHADER(FHairStrandsForwardRasterCompactionCS);
+	SHADER_USE_PARAMETER_STRUCT(FHairStrandsForwardRasterCompactionCS, FGlobalShader);
 
 	class FDebug : SHADER_PERMUTATION_BOOL("PERMUTATION_DEBUG");
 	using FPermutationDomain = TShaderPermutationDomain<FDebug>;
@@ -163,7 +168,7 @@ class FVisiblityRasterComputeForwardCompactionCS : public FGlobalShader
 	END_SHADER_PARAMETER_STRUCT()
 
 public:
-	static bool ShouldCompilePermutation(const FGlobalShaderPermutationParameters &Parameters) { return IsHairStrandsSupported(EHairStrandsShaderType::Strands, Parameters.Platform); }
+	static bool ShouldCompilePermutation(const FGlobalShaderPermutationParameters& Parameters) { return IsHairStrandsForwardRasterSupported(Parameters.Platform); }
 	static void ModifyCompilationEnvironment(const FGlobalShaderPermutationParameters &Parameters, FShaderCompilerEnvironment &OutEnvironment)
 	{
 		FGlobalShader::ModifyCompilationEnvironment(Parameters, OutEnvironment);
@@ -171,15 +176,15 @@ public:
 	}
 };
 
-IMPLEMENT_GLOBAL_SHADER(FVisiblityRasterComputeForwardCompactionCS, "/Engine/Private/HairStrands/HairStrandsForwardRaster.usf", "CompactionCS", SF_Compute);
+IMPLEMENT_GLOBAL_SHADER(FHairStrandsForwardRasterCompactionCS, "/Engine/Private/HairStrands/HairStrandsForwardRaster.usf", "CompactionCS", SF_Compute);
 
 /////////////////////////////////////////////////////////////////////////////////////////
 // Rasterized binned & shaded segements
 
-class FVisiblityRasterComputeForwardRasterizeCS : public FGlobalShader
+class FHairStrandsForwardRasterRasterizeCS : public FGlobalShader
 {
-	DECLARE_GLOBAL_SHADER(FVisiblityRasterComputeForwardRasterizeCS);
-	SHADER_USE_PARAMETER_STRUCT(FVisiblityRasterComputeForwardRasterizeCS, FGlobalShader);
+	DECLARE_GLOBAL_SHADER(FHairStrandsForwardRasterRasterizeCS);
+	SHADER_USE_PARAMETER_STRUCT(FHairStrandsForwardRasterRasterizeCS, FGlobalShader);
 
 	class FDebug : SHADER_PERMUTATION_BOOL("PERMUTATION_DEBUG");
 	using FPermutationDomain = TShaderPermutationDomain<FDebug>;
@@ -202,7 +207,7 @@ class FVisiblityRasterComputeForwardRasterizeCS : public FGlobalShader
 	END_SHADER_PARAMETER_STRUCT()
 
 public:
-	static bool ShouldCompilePermutation(const FGlobalShaderPermutationParameters& Parameters) { return IsHairStrandsSupported(EHairStrandsShaderType::Strands, Parameters.Platform); }
+	static bool ShouldCompilePermutation(const FGlobalShaderPermutationParameters& Parameters) { return IsHairStrandsForwardRasterSupported(Parameters.Platform); }
 	static void ModifyCompilationEnvironment(const FGlobalShaderPermutationParameters& Parameters, FShaderCompilerEnvironment& OutEnvironment)
 	{
 		FGlobalShader::ModifyCompilationEnvironment(Parameters, OutEnvironment);
@@ -212,15 +217,15 @@ public:
 	}
 };
 
-IMPLEMENT_GLOBAL_SHADER(FVisiblityRasterComputeForwardRasterizeCS, "/Engine/Private/HairStrands/HairStrandsForwardRaster.usf", "RasterCS", SF_Compute);
+IMPLEMENT_GLOBAL_SHADER(FHairStrandsForwardRasterRasterizeCS, "/Engine/Private/HairStrands/HairStrandsForwardRaster.usf", "RasterCS", SF_Compute);
 
 /////////////////////////////////////////////////////////////////////////////////////////
 // Debug pass
 
-class FVisiblityRasterComputeForwardDebugCS : public FGlobalShader
+class FHairStrandsForwardRasterDebugCS : public FGlobalShader
 {
-	DECLARE_GLOBAL_SHADER(FVisiblityRasterComputeForwardDebugCS);
-	SHADER_USE_PARAMETER_STRUCT(FVisiblityRasterComputeForwardDebugCS, FGlobalShader);
+	DECLARE_GLOBAL_SHADER(FHairStrandsForwardRasterDebugCS);
+	SHADER_USE_PARAMETER_STRUCT(FHairStrandsForwardRasterDebugCS, FGlobalShader);
 
 	using FPermutationDomain = TShaderPermutationDomain<>;
 
@@ -237,7 +242,7 @@ class FVisiblityRasterComputeForwardDebugCS : public FGlobalShader
 	END_SHADER_PARAMETER_STRUCT()
 
 public:
-	static bool ShouldCompilePermutation(const FGlobalShaderPermutationParameters& Parameters) { return IsHairStrandsSupported(EHairStrandsShaderType::Strands, Parameters.Platform); }
+	static bool ShouldCompilePermutation(const FGlobalShaderPermutationParameters& Parameters) { return IsHairStrandsForwardRasterSupported(Parameters.Platform); }
 	static void ModifyCompilationEnvironment(const FGlobalShaderPermutationParameters& Parameters, FShaderCompilerEnvironment& OutEnvironment)
 	{
 		FGlobalShader::ModifyCompilationEnvironment(Parameters, OutEnvironment);
@@ -245,7 +250,7 @@ public:
 	}
 };
 
-IMPLEMENT_GLOBAL_SHADER(FVisiblityRasterComputeForwardDebugCS, "/Engine/Private/HairStrands/HairStrandsForwardRaster.usf", "MainCS", SF_Compute);
+IMPLEMENT_GLOBAL_SHADER(FHairStrandsForwardRasterDebugCS, "/Engine/Private/HairStrands/HairStrandsForwardRaster.usf", "MainCS", SF_Compute);
 
 
 /////////////////////////////////////////////////////////////////////////////////////////
@@ -260,6 +265,11 @@ FRasterForwardCullingOutput AddHairStrandsForwardCullingPass(
 	bool bSupportCulling,
 	bool bForceRegister)
 {	
+	if (!IsHairStrandsForwardRasterSupported(ViewInfo.GetShaderPlatform()))
+	{
+		return FRasterForwardCullingOutput();
+	}
+
 	// Compute maximum number of PrimIDs
 	uint32 MaxNumPrimIDs = 0;
 	for (const FHairStrandsMacroGroupData &MacroGroup : MacroGroupDatas)
@@ -414,6 +424,11 @@ void AddHairStrandsForwardRasterPass(
 	const FRDGTextureRef SceneDepthTexture,
 	const FRDGTextureRef SceneColorTexture)
 {	
+	if (!IsHairStrandsForwardRasterSupported(ViewInfo.GetShaderPlatform()))
+	{
+		return;
+	}
+
 	const FIntPoint OutResolution = InResolution;
 
 	// See the comment on the GHairVisibilityComputeRaster_MaxTiles declaration for an explanation for the large upper bound.
@@ -506,8 +521,8 @@ void AddHairStrandsForwardRasterPass(
 	{
 		VisTileDepthGrid = GraphBuilder.CreateTexture(FRDGTextureDesc::Create2D(BinTileRes, PF_R32_UINT, FClearValueBinding::None, TexCreate_UAV | TexCreate_ShaderResource), TEXT("Hair.VisTileDepthGrid"));
 
-		TShaderMapRef<FVisiblityRasterComputeForwardPrepareDepthGridCS> ComputeShaderPrepareDepthGrid(ViewInfo.ShaderMap);
-		FVisiblityRasterComputeForwardPrepareDepthGridCS::FParameters* Parameters = GraphBuilder.AllocParameters<FVisiblityRasterComputeForwardPrepareDepthGridCS::FParameters>();		
+		TShaderMapRef<FHairStrandsForwardRasterPrepareDepthGridCS> ComputeShaderPrepareDepthGrid(ViewInfo.ShaderMap);
+		FHairStrandsForwardRasterPrepareDepthGridCS::FParameters* Parameters = GraphBuilder.AllocParameters<FHairStrandsForwardRasterPrepareDepthGridCS::FParameters>();		
 		Parameters->Common = Common;
 		Parameters->SceneDepthTexture = SceneDepthTexture;
 		Parameters->OutVisTileDepthGrid = GraphBuilder.CreateUAV(VisTileDepthGrid);
@@ -534,7 +549,7 @@ void AddHairStrandsForwardRasterPass(
 
 	// Binning pass
 	{
-		FVisiblityRasterComputeForwardBinningCS::FParameters* Parameters = GraphBuilder.AllocParameters<FVisiblityRasterComputeForwardBinningCS::FParameters>();
+		FHairStrandsForwardRasterBinningCS::FParameters* Parameters = GraphBuilder.AllocParameters<FHairStrandsForwardRasterBinningCS::FParameters>();
 		Parameters->Common = Common;
 		Parameters->ViewUniformBuffer = ViewUniformShaderParameters;
 		Parameters->OutVisTilePrims = GraphBuilder.CreateUAV(BinData.VisTilePrims, PF_R32_UINT);
@@ -545,9 +560,9 @@ void AddHairStrandsForwardRasterPass(
 		Parameters->VisTileBinningGridTex = BinData.VisTileBinningGrid;
 		ShaderPrint::SetParameters(GraphBuilder, ViewInfo.ShaderPrintData, Parameters->ShaderPrintParameters);
 
-		FVisiblityRasterComputeForwardBinningCS::FPermutationDomain PermutationVector;
-		PermutationVector.Set<FVisiblityRasterComputeForwardBinningCS::FDebug>(bDebugEnabled);
-		TShaderMapRef<FVisiblityRasterComputeForwardBinningCS> ComputeShaderBinning(ViewInfo.ShaderMap, PermutationVector);
+		FHairStrandsForwardRasterBinningCS::FPermutationDomain PermutationVector;
+		PermutationVector.Set<FHairStrandsForwardRasterBinningCS::FDebug>(bDebugEnabled);
+		TShaderMapRef<FHairStrandsForwardRasterBinningCS> ComputeShaderBinning(ViewInfo.ShaderMap, PermutationVector);
 		FComputeShaderUtils::AddPass(
 			GraphBuilder, 
 			RDG_EVENT_NAME("HairStrands::VisibilityComputeRasterBinning"),
@@ -558,7 +573,7 @@ void AddHairStrandsForwardRasterPass(
 
 	// Compaction
 	{
-		FVisiblityRasterComputeForwardCompactionCS::FParameters* Parameters = GraphBuilder.AllocParameters<FVisiblityRasterComputeForwardCompactionCS::FParameters>();
+		FHairStrandsForwardRasterCompactionCS::FParameters* Parameters = GraphBuilder.AllocParameters<FHairStrandsForwardRasterCompactionCS::FParameters>();
 		Parameters->Common = Common;
 		Parameters->ViewUniformBuffer = ViewUniformShaderParameters;
 		Parameters->InData  = BinData.VisTileDataSRV;
@@ -569,9 +584,9 @@ void AddHairStrandsForwardRasterPass(
 		Parameters->OutArgs = GraphBuilder.CreateUAV(BinData.CompactedVisTileArgs, PF_R32_UINT);
 		ShaderPrint::SetParameters(GraphBuilder, ViewInfo.ShaderPrintData, Parameters->ShaderPrintParameters);
 
-		FVisiblityRasterComputeForwardCompactionCS::FPermutationDomain PermutationVector;
-		PermutationVector.Set<FVisiblityRasterComputeForwardCompactionCS::FDebug>(bDebugEnabled);
-		TShaderMapRef<FVisiblityRasterComputeForwardCompactionCS> ComputeShaderCompaction(ViewInfo.ShaderMap, PermutationVector);
+		FHairStrandsForwardRasterCompactionCS::FPermutationDomain PermutationVector;
+		PermutationVector.Set<FHairStrandsForwardRasterCompactionCS::FDebug>(bDebugEnabled);
+		TShaderMapRef<FHairStrandsForwardRasterCompactionCS> ComputeShaderCompaction(ViewInfo.ShaderMap, PermutationVector);
 		FComputeShaderUtils::AddPass(GraphBuilder, RDG_EVENT_NAME("HairStrands::VisibilityComputeRasterCompaction"), ComputeShaderCompaction, Parameters, FIntVector(BinTileRes.X, BinTileRes.Y, 1));
 	}
 
@@ -594,7 +609,7 @@ void AddHairStrandsForwardRasterPass(
 		FRDGBufferUAVRef WorkCounterUAV = GraphBuilder.CreateUAV(WorkCounter, PF_R32_UINT);
 		AddClearUAVPass(GraphBuilder, WorkCounterUAV, 0);
 
-		FVisiblityRasterComputeForwardRasterizeCS::FParameters* Parameters = GraphBuilder.AllocParameters<FVisiblityRasterComputeForwardRasterizeCS::FParameters>();
+		FHairStrandsForwardRasterRasterizeCS::FParameters* Parameters = GraphBuilder.AllocParameters<FHairStrandsForwardRasterRasterizeCS::FParameters>();
 		Parameters->Common = Common;
 		Parameters->ViewUniformBuffer = ViewUniformShaderParameters;
 		Parameters->SceneDepthTexture = SceneDepthTexture;
@@ -609,15 +624,15 @@ void AddHairStrandsForwardRasterPass(
 		Parameters->RWWorkCounter = WorkCounterUAV;
 		ShaderPrint::SetParameters(GraphBuilder, ViewInfo.ShaderPrintData, Parameters->ShaderPrintParameters);
 
-		FVisiblityRasterComputeForwardRasterizeCS::FPermutationDomain PermutationVector;
-		PermutationVector.Set<FVisiblityRasterComputeForwardRasterizeCS::FDebug>(bDebugEnabled);
-		TShaderMapRef<FVisiblityRasterComputeForwardRasterizeCS> ComputeShaderRaster(ViewInfo.ShaderMap, PermutationVector);
+		FHairStrandsForwardRasterRasterizeCS::FPermutationDomain PermutationVector;
+		PermutationVector.Set<FHairStrandsForwardRasterRasterizeCS::FDebug>(bDebugEnabled);
+		TShaderMapRef<FHairStrandsForwardRasterRasterizeCS> ComputeShaderRaster(ViewInfo.ShaderMap, PermutationVector);
 		FComputeShaderUtils::AddPass(GraphBuilder, RDG_EVENT_NAME("HairStrands::VisibilityComputeRasterRaster(tiled)"), ComputeShaderRaster, Parameters, FIntVector(NumRasterizers, 1, 1));
 	}
 
 	if (bDebugEnabled)
 	{
-		FVisiblityRasterComputeForwardDebugCS::FParameters* Parameters = GraphBuilder.AllocParameters<FVisiblityRasterComputeForwardDebugCS::FParameters>();
+		FHairStrandsForwardRasterDebugCS::FParameters* Parameters = GraphBuilder.AllocParameters<FHairStrandsForwardRasterDebugCS::FParameters>();
 		Parameters->Common = Common;
 		Parameters->InstanceCount = InData.RasterizedInstanceCount;
 		Parameters->VisTileData = GraphBuilder.CreateSRV(BinData.CompactedVisTileData);
@@ -626,7 +641,7 @@ void AddHairStrandsForwardRasterPass(
 		Parameters->HairPixelCountPerTile_ForDebug = HairPixelCountPerTile_ForDebug;
 		ShaderPrint::SetParameters(GraphBuilder, ViewInfo.ShaderPrintData, Parameters->ShaderPrintParameters);
 
-		TShaderMapRef<FVisiblityRasterComputeForwardDebugCS> DebugComputeShader(ViewInfo.ShaderMap);
+		TShaderMapRef<FHairStrandsForwardRasterDebugCS> DebugComputeShader(ViewInfo.ShaderMap);
 		FComputeShaderUtils::AddPass(GraphBuilder, RDG_EVENT_NAME("HairStrands::VisibilityComputeRaster(Debug)"), DebugComputeShader, Parameters, FIntVector(FMath::DivideAndRoundUp(Common.BinTileRes.X, 8), FMath::DivideAndRoundUp(Common.BinTileRes.Y, 8), 1));
 	}
 }
diff --git a/Engine/Source/Runtime/Renderer/Private/HairStrands/HairStrandsInterface.cpp b/Engine/Source/Runtime/Renderer/Private/HairStrands/HairStrandsInterface.cpp
--- a/Engine/Source/Runtime/Renderer/Private/HairStrands/HairStrandsInterface.cpp
+++ b/Engine/Source/Runtime/Renderer/Private/HairStrands/HairStrandsInterface.cpp
@@ -495,9 +495,9 @@ bool IsHairVisibilityComputeRasterEnabled()
 	return CVarHairStrandsVisibilityComputeRaster.GetValueOnAnyThread() == 1;
 }
 
-bool IsHairVisibilityComputeRasterForwardEnabled()
+bool IsHairVisibilityComputeRasterForwardEnabled(EShaderPlatform InPlatform)
 {
-	return CVarHairStrandsVisibilityComputeRaster.GetValueOnAnyThread() == 2;
+	return IsFeatureLevelSupported(InPlatform, ERHIFeatureLevel::SM6) && CVarHairStrandsVisibilityComputeRaster.GetValueOnAnyThread() == 2;
 }
 
 bool IsHairVisibilityComputeRasterContinuousLODEnabled()
diff --git a/Engine/Source/Runtime/Renderer/Private/HairStrands/HairStrandsVisibility.cpp b/Engine/Source/Runtime/Renderer/Private/HairStrands/HairStrandsVisibility.cpp
--- a/Engine/Source/Runtime/Renderer/Private/HairStrands/HairStrandsVisibility.cpp
+++ b/Engine/Source/Runtime/Renderer/Private/HairStrands/HairStrandsVisibility.cpp
@@ -154,9 +154,9 @@ enum EHairVisibilityRenderMode
 };
 
 bool IsHairVisibilityComputeRasterEnabled();
-bool IsHairVisibilityComputeRasterForwardEnabled();
+bool IsHairVisibilityComputeRasterForwardEnabled(EShaderPlatform InPlatform);
 
-inline EHairVisibilityRenderMode GetHairVisibilityRenderMode()
+inline EHairVisibilityRenderMode GetHairVisibilityRenderMode(EShaderPlatform InPlatform)
 {
 	if (GHairVisibilityPPLL > 0)
 	{
@@ -166,7 +166,7 @@ inline EHairVisibilityRenderMode GetHairVisibilityRenderMode()
 	{
 		return HairVisibilityRenderMode_ComputeRaster;
 	}
-	else if (IsHairVisibilityComputeRasterForwardEnabled())
+	else if (IsHairVisibilityComputeRasterForwardEnabled(InPlatform))
 	{
 		return HairVisibilityRenderMode_ComputeRasterForward;
 	}
@@ -176,15 +176,15 @@ inline EHairVisibilityRenderMode GetHairVisibilityRenderMode()
 	}
 }
 
-inline bool IsMsaaEnabled()
+inline bool IsMsaaEnabled(EShaderPlatform InPlatform)
 {
-	const EHairVisibilityRenderMode Mode = GetHairVisibilityRenderMode();
+	const EHairVisibilityRenderMode Mode = GetHairVisibilityRenderMode(InPlatform);
 	return Mode == HairVisibilityRenderMode_MSAA_Visibility;
 }
 
-static uint32 GetMaxSamplePerPixel()
+static uint32 GetMaxSamplePerPixel(EShaderPlatform InPlatform)
 {
-	switch (GetHairVisibilityRenderMode())
+	switch (GetHairVisibilityRenderMode(InPlatform))
 	{
 		case HairVisibilityRenderMode_ComputeRaster:
 		{
@@ -249,10 +249,10 @@ static uint32 GetMaxSamplePerPixel()
 	return 1;
 }
 
-inline uint32 GetMeanSamplePerPixel()
+inline uint32 GetMeanSamplePerPixel(EShaderPlatform InPlatform)
 {
-	const uint32 SamplePerPixel = GetMaxSamplePerPixel();
-	switch (GetHairVisibilityRenderMode())
+	const uint32 SamplePerPixel = GetMaxSamplePerPixel(InPlatform);
+	switch (GetHairVisibilityRenderMode(InPlatform))
 	{
 	case HairVisibilityRenderMode_ComputeRasterForward:
 	case HairVisibilityRenderMode_ComputeRaster:
@@ -268,9 +268,9 @@ inline uint32 GetMeanSamplePerPixel()
 	return 1;
 }
 
-uint32 GetHairStrandsMeanSamplePerPixel()
+uint32 GetHairStrandsMeanSamplePerPixel(EShaderPlatform InPlatform)
 {
-	return GetMeanSamplePerPixel();
+	return GetMeanSamplePerPixel(InPlatform);
 }
 
 struct FRasterComputeOutput
@@ -284,9 +284,9 @@ struct FRasterComputeOutput
 	FRDGTextureRef PrimMatTexture = nullptr;
 };
 
-static uint32 GetTotalSampleCountForAllocation(FIntPoint Resolution)
+static uint32 GetTotalSampleCountForAllocation(FIntPoint Resolution, EShaderPlatform InPlatform)
 {
-	return Resolution.X * Resolution.Y * GetMeanSamplePerPixel();
+	return Resolution.X * Resolution.Y * GetMeanSamplePerPixel(InPlatform);
 }
 
 void SetUpViewHairRenderInfo(const FViewInfo& ViewInfo, bool bEnableMSAA, FVector4f& OutHairRenderInfo, uint32& OutHairRenderInfoBits, uint32& OutHairComponents)
@@ -295,7 +295,7 @@ void SetUpViewHairRenderInfo(const FViewInfo& ViewInfo, bool bEnableMSAA, FVecto
 	const float VelocityMagnitudeScale = FMath::Clamp(CVarHairVelocityMagnitudeScale.GetValueOnAnyThread(), 0, 512) * FMath::Min(PixelVelocity.X, PixelVelocity.Y);
 
 	// In the case we render coverage, we need to override some view uniform shader parameters to account for the change in MSAA sample count.
-	const uint32 HairVisibilitySampleCount = bEnableMSAA ? GetMaxSamplePerPixel() : 1;	// The coverage pass does not use MSAA
+	const uint32 HairVisibilitySampleCount = bEnableMSAA ? GetMaxSamplePerPixel(ViewInfo.GetShaderPlatform()) : 1;	// The coverage pass does not use MSAA
 	const float RasterizationScaleOverride = 0.0f;	// no override
 	FMinHairRadiusAtDepth1 MinHairRadius = ComputeMinStrandRadiusAtDepth1(
 		FIntPoint(ViewInfo.UnconstrainedViewRect.Width(), ViewInfo.UnconstrainedViewRect.Height()), ViewInfo.FOV, HairVisibilitySampleCount, RasterizationScaleOverride);
@@ -307,7 +307,7 @@ void SetUpViewHairRenderInfo(const FViewInfo& ViewInfo, bool bEnableMSAA, FVecto
 
 void SetUpViewHairRenderInfo(const FViewInfo& ViewInfo, FVector4f& OutHairRenderInfo, uint32& OutHairRenderInfoBits, uint32& OutHairComponents)
 {
-	SetUpViewHairRenderInfo(ViewInfo, IsMsaaEnabled(), OutHairRenderInfo, OutHairRenderInfoBits, OutHairComponents);
+	SetUpViewHairRenderInfo(ViewInfo, IsMsaaEnabled(ViewInfo.GetShaderPlatform()), OutHairRenderInfo, OutHairRenderInfoBits, OutHairComponents);
 }
 
 static bool IsCompatibleWithHairVisibility(const FMeshMaterialShaderPermutationParameters& Parameters)
@@ -1679,6 +1679,7 @@ static void AddHairVisibilityPrimitiveIdCompactionPass(
 	FRDGBufferRef& OutIndirectArgsBuffer,
 	uint32& OutMaxRenderNodeCount)
 {
+	const uint32 MaxSamplePerPixel = GetMaxSamplePerPixel(View.GetShaderPlatform());
 	FIntPoint Resolution;
 	if (bUsePPLL)
 	{
@@ -1689,7 +1690,7 @@ static void AddHairVisibilityPrimitiveIdCompactionPass(
 	}
 	else
 	{
-		check(PassParameters->MSAA_DepthTexture->Desc.NumSamples == GetMaxSamplePerPixel());
+		check(PassParameters->MSAA_DepthTexture->Desc.NumSamples == MaxSamplePerPixel);
 		check(PassParameters->MSAA_DepthTexture);
 		check(PassParameters->MSAA_IDTexture);
 		Resolution = PassParameters->MSAA_DepthTexture->Desc.Extent;
@@ -1720,9 +1721,9 @@ static void AddHairVisibilityPrimitiveIdCompactionPass(
 	const FIntPoint EffectiveResolution = bUsePPLL ? FIntPoint(View.ViewRect.Width(), View.ViewRect.Height()) : FIntPoint(HairRect.Width(), HairRect.Height());
 
 	// Select render node count according to current mode
-	const uint32 MSAASampleCount = GetHairVisibilityRenderMode() == HairVisibilityRenderMode_MSAA_Visibility ? GetMaxSamplePerPixel() : 1;
-	const uint32 PPLLMaxRenderNodePerPixel = GetMaxSamplePerPixel();
-	const uint32 MaxRenderNodeCount = GetTotalSampleCountForAllocation(EffectiveResolution);
+	const uint32 MSAASampleCount = GetHairVisibilityRenderMode(View.GetShaderPlatform()) == HairVisibilityRenderMode_MSAA_Visibility ? MaxSamplePerPixel : 1;
+	const uint32 PPLLMaxRenderNodePerPixel = MaxSamplePerPixel;
+	const uint32 MaxRenderNodeCount = GetTotalSampleCountForAllocation(EffectiveResolution, View.GetShaderPlatform());
 	check(TileData.IsValid());
 
 	OutCompactNodeVis = GraphBuilder.CreateBuffer(FRDGBufferDesc::CreateStructuredDesc(sizeof(HairStrandsVisibilityInternal::NodeVis), MaxRenderNodeCount), TEXT("Hair.VisibilityNodeVis"));
@@ -1861,11 +1862,12 @@ static void AddHairVisibilityCompactionComputeRasterPass(
 	AddClearUAVPass(GraphBuilder, GraphBuilder.CreateUAV(OutCoverageTexture), 0.f);
 
 	// Select render node count according to current mode
+	const uint32 MaxSamplePerPixel = GetMaxSamplePerPixel(View.GetShaderPlatform());
 	check(TileData.IsValid());
 	const FHairStrandsTiles::ETileType TileType = FHairStrandsTiles::ETileType::HairAll;
-	const uint32 MSAASampleCount = GetHairVisibilityRenderMode() == HairVisibilityRenderMode_MSAA_Visibility ? GetMaxSamplePerPixel() : 1;
-	const uint32 PPLLMaxRenderNodePerPixel = GetMaxSamplePerPixel();
-	const uint32 MaxRenderNodeCount = GetTotalSampleCountForAllocation(Resolution);
+	const uint32 MSAASampleCount = GetHairVisibilityRenderMode(View.GetShaderPlatform()) == HairVisibilityRenderMode_MSAA_Visibility ? MaxSamplePerPixel : 1;
+	const uint32 PPLLMaxRenderNodePerPixel = MaxSamplePerPixel;
+	const uint32 MaxRenderNodeCount = GetTotalSampleCountForAllocation(Resolution, View.GetShaderPlatform());
 	OutCompactNodeVis = GraphBuilder.CreateBuffer(FRDGBufferDesc::CreateStructuredDesc(sizeof(HairStrandsVisibilityInternal::NodeVis), MaxRenderNodeCount), TEXT("Hair.VisibilityPrimitiveIdCompactNodeData"));
 	OutCompactNodeCoord = GraphBuilder.CreateBuffer(FRDGBufferDesc::CreateBufferDesc(sizeof(uint32), MaxRenderNodeCount), TEXT("Hair.VisibilityPrimitiveIdCompactNodeCoord"));
 
@@ -1942,11 +1944,11 @@ static FRDGTextureRef AddHairVisibilityFillOpaqueDepth(
 	const FHairStrandsTiles& TileData,
 	const FRDGTextureRef& SceneDepthTexture)
 {
-	check(GetHairVisibilityRenderMode() == HairVisibilityRenderMode_MSAA_Visibility || GetHairVisibilityRenderMode() == HairVisibilityRenderMode_ComputeRaster);
+	check(GetHairVisibilityRenderMode(View.GetShaderPlatform()) == HairVisibilityRenderMode_MSAA_Visibility || GetHairVisibilityRenderMode(View.GetShaderPlatform()) == HairVisibilityRenderMode_ComputeRaster);
 	check(TileData.IsValid());
 
 	const FHairStrandsTiles::ETileType TileType = FHairStrandsTiles::ETileType::HairAll;
-	FRDGTextureRef OutVisibilityDepthTexture = GraphBuilder.CreateTexture(FRDGTextureDesc::Create2D(Resolution, PF_D24, FClearValueBinding::DepthFar, TexCreate_DepthStencilTargetable | TexCreate_ShaderResource, 1, GetMaxSamplePerPixel()), TEXT("Hair.VisibilityDepthTexture"));
+	FRDGTextureRef OutVisibilityDepthTexture = GraphBuilder.CreateTexture(FRDGTextureDesc::Create2D(Resolution, PF_D24, FClearValueBinding::DepthFar, TexCreate_DepthStencilTargetable | TexCreate_ShaderResource, 1, GetMaxSamplePerPixel(View.GetShaderPlatform())), TEXT("Hair.VisibilityDepthTexture"));
 
 	FHairVisibilityFillOpaqueDepthPS::FParameters* Parameters = GraphBuilder.AllocParameters<FHairVisibilityFillOpaqueDepthPS::FParameters>();
 	Parameters->TileData = GetHairStrandsTileParameters(View, TileData, TileType);
@@ -2099,7 +2101,7 @@ static void AddHairVisibilityMSAAPass(
 	FRDGTextureRef& OutVisibilityIdTexture,
 	FRDGTextureRef& OutVisibilityDepthTexture)
 {
-	const uint32 MSAASampleCount = GetMaxSamplePerPixel();
+	const uint32 MSAASampleCount = GetMaxSamplePerPixel(ViewInfo->GetShaderPlatform());
 	{
 		{
 			FRDGTextureDesc Desc = FRDGTextureDesc::Create2D(Resolution, PF_R32_UINT, FClearValueBinding(EClearBinding::ENoneBound), TexCreate_NoFastClear | TexCreate_RenderTargetable | TexCreate_ShaderResource, 1, MSAASampleCount);
@@ -2147,7 +2149,7 @@ static void AddHairVisibilityPPLLPass(
 	const FIntRect HairRect = ViewInfo->ViewRect;
 	const FIntPoint EffectiveResolution(HairRect.Width(), HairRect.Height());
 
-	const uint32 PPLLMaxTotalListElementCount = GetTotalSampleCountForAllocation(EffectiveResolution);
+	const uint32 PPLLMaxTotalListElementCount = GetTotalSampleCountForAllocation(EffectiveResolution, ViewInfo->GetShaderPlatform());
 	{
 		OutVisibilityPPLLNodeData = GraphBuilder.CreateBuffer(FRDGBufferDesc::CreateStructuredDesc(sizeof(FPackedHairVisPPLL), PPLLMaxTotalListElementCount), TEXT("Hair.VisibilityPPLLNodeData"));
 	}
@@ -2730,7 +2732,7 @@ void AddMeshDrawTransitionPass(
 			ExternalAccessQueue.Add(CulledVertexRadiusScaleBuffer);
 			ExternalAccessQueue.Add(DrawIndirectBuffer, ERHIAccess::IndirectArgs);
 
-			const EHairVisibilityRenderMode RasterMode = GetHairVisibilityRenderMode();
+			const EHairVisibilityRenderMode RasterMode = GetHairVisibilityRenderMode(ViewInfo.GetShaderPlatform());
 			if (RasterMode != HairVisibilityRenderMode_ComputeRaster && RasterMode != HairVisibilityRenderMode_ComputeRasterForward)
 			{
 				VFInput.Strands.PositionBuffer				= FRDGImportedBuffer();
@@ -3964,7 +3966,7 @@ void DrawHitProxies(
 {
 	// Proxy rendering is only supported/compatible with MSAA-visibility rendering. 
 	// PPLL is not supported, but it is supposed to be used only for final render.
-	const EHairVisibilityRenderMode RenderMode = GetHairVisibilityRenderMode();
+	const EHairVisibilityRenderMode RenderMode = GetHairVisibilityRenderMode(View.GetShaderPlatform());
 	if (RenderMode != HairVisibilityRenderMode_MSAA_Visibility && RenderMode != HairVisibilityRenderMode_ComputeRaster && RenderMode != HairVisibilityRenderMode_ComputeRasterForward)
 	{
 		return;
@@ -4193,7 +4195,7 @@ void RenderHairStrandsVisibilityBuffer(
 		{
 			FHairStrandsVisibilityData& VisibilityData = View.HairStrandsViewData.VisibilityData;
 			VisibilityData.NodeGroupSize = GetVendorOptimalGroupSize1D();
-			VisibilityData.MaxSampleCount = GetMaxSamplePerPixel();
+			VisibilityData.MaxSampleCount = GetMaxSamplePerPixel(View.GetShaderPlatform());
 
 			// Use the scene color for computing target resolution as the View.ViewRect, 
 			// doesn't include the actual resolution padding which make buffer size 
@@ -4201,7 +4203,7 @@ void RenderHairStrandsVisibilityBuffer(
 			check(SceneDepthTexture);
 			const FIntPoint Resolution = SceneDepthTexture->Desc.Extent;
 
-			const EHairVisibilityRenderMode RenderMode = GetHairVisibilityRenderMode();
+			const EHairVisibilityRenderMode RenderMode = GetHairVisibilityRenderMode(View.GetShaderPlatform());
 			check(
 				RenderMode == HairVisibilityRenderMode_MSAA_Visibility || 
 				RenderMode == HairVisibilityRenderMode_PPLL || 
@@ -4727,7 +4729,7 @@ void RenderHairStrandsVisibilityBuffer(
 						VisibilityData.SampleLightingTexture			= PassOutput.SampleLightingTexture;
 					}
 
-					VisibilityData.MaxSampleCount = GetMaxSamplePerPixel();
+					VisibilityData.MaxSampleCount = GetMaxSamplePerPixel(View.GetShaderPlatform());
 					VisibilityData.NodeIndex = CompactNodeIndex;
 					VisibilityData.CoverageTexture = CoverageTexture;
 					VisibilityData.HairOnlyDepthTexture = HairOnlyDepthTexture;
