Reduce long frame time during scene updates in editor by better balancing on more threads
  - Longest CacheMeshDrawCommands in CitySample went from 426ms to 76ms

#rnx

diff --git a/Engine/Source/Runtime/Renderer/Private/PrimitiveSceneInfo.cpp b/Engine/Source/Runtime/Renderer/Private/PrimitiveSceneInfo.cpp
--- a/Engine/Source/Runtime/Renderer/Private/PrimitiveSceneInfo.cpp
+++ b/Engine/Source/Runtime/Renderer/Private/PrimitiveSceneInfo.cpp
@@ -297,7 +297,9 @@ void FPrimitiveSceneInfo::CacheMeshDrawCommands(FRHICommandListImmediate& RHICmd
 
 	QUICK_SCOPE_CYCLE_COUNTER(STAT_CacheMeshDrawCommands);
 
-	static constexpr int BATCH_SIZE = 64;
+	// This reduce stuttering in editor by improving balancing of all the 
+	// shadermap processing. Keep it as it is for runtime as the requirements are different.
+	static constexpr int BATCH_SIZE = WITH_EDITOR ? 1 : 64;
 	const int NumBatches = (SceneInfos.Num() + BATCH_SIZE - 1) / BATCH_SIZE;
 
 	auto DoWorkLambda = [Scene, SceneInfos](FCachedPassMeshDrawListContext& DrawListContext, int32 Index)
