Disable all-primitive validation of persistent primitive ID in anything but debug build.

diff --git a/Engine/Source/Runtime/Renderer/Private/RendererScene.cpp b/Engine/Source/Runtime/Renderer/Private/RendererScene.cpp
--- a/Engine/Source/Runtime/Renderer/Private/RendererScene.cpp
+++ b/Engine/Source/Runtime/Renderer/Private/RendererScene.cpp
@@ -6057,13 +6057,15 @@ void FScene::UpdateAllPrimitiveSceneInfos(FRDGBuilder& GraphBuilder, bool bAsync
 	DistanceFieldSceneDataUpdates.Empty();
 	AddedPrimitiveSceneInfos.Empty();
 
+#if DO_GUARD_SLOW
 	for (const FPrimitiveSceneInfo* PrimitiveSceneInfo : Primitives)
 	{
-		check(PrimitiveSceneInfo->PackedIndex != INDEX_NONE);
-		check(PrimitiveSceneInfo->PackedIndex < Primitives.Num());
-		check(PrimitiveSceneInfo->PersistentIndex.Index != INDEX_NONE);
-		check(PersistentPrimitiveIdToIndexMap[PrimitiveSceneInfo->PersistentIndex.Index] == PrimitiveSceneInfo->PackedIndex);
+		checkSlow(PrimitiveSceneInfo->PackedIndex != INDEX_NONE);
+		checkSlow(PrimitiveSceneInfo->PackedIndex < Primitives.Num());
+		checkSlow(PrimitiveSceneInfo->PersistentIndex.Index != INDEX_NONE);
+		checkSlow(PersistentPrimitiveIdToIndexMap[PrimitiveSceneInfo->PersistentIndex.Index] == PrimitiveSceneInfo->PackedIndex);
 	}
+#endif
 }
 
 void FScene::CreateLightPrimitiveInteractionsForPrimitive(FPrimitiveSceneInfo* PrimitiveInfo, bool bAsyncCreateLPIs)
