Make ComputeBlendModeMask compatible with Strata blend mode

diff --git a/Engine/Plugins/Experimental/GPULightmass/Source/GPULightmass/Private/LightmapRenderer.cpp b/Engine/Plugins/Experimental/GPULightmass/Source/GPULightmass/Private/LightmapRenderer.cpp
--- a/Engine/Plugins/Experimental/GPULightmass/Source/GPULightmass/Private/LightmapRenderer.cpp
+++ b/Engine/Plugins/Experimental/GPULightmass/Source/GPULightmass/Private/LightmapRenderer.cpp
@@ -658,7 +658,7 @@ void FCachedRayTracingSceneData::SetupFromSceneRenderState(FSceneRenderState& Sc
 				bAllSegmentsUnlit &= Material.GetShadingModels().HasOnlyShadingModel(MSM_Unlit) || !MeshBatches[SegmentIndex].CastShadow;
 				bAllSegmentsOpaque &= IsOpaqueBlendMode(Material);
 				bAnySegmentsCastShadow |= MeshBatches[SegmentIndex].CastRayTracedShadow && Material.CastsRayTracedShadows();
-				InstanceMask |= ComputeBlendModeMask(Material.GetBlendMode());
+				InstanceMask |= ComputeBlendModeMask(Material);
 			}
 
 			InstanceMask |= bAnySegmentsCastShadow ? RAY_TRACING_MASK_SHADOW : 0;
@@ -710,7 +710,7 @@ void FCachedRayTracingSceneData::SetupFromSceneRenderState(FSceneRenderState& Sc
 				bAllSegmentsUnlit &= Material.GetShadingModels().HasOnlyShadingModel(MSM_Unlit) || !MeshBatches[SegmentIndex].CastShadow;
 				bAllSegmentsOpaque &= IsOpaqueBlendMode(Material);
 				bAnySegmentsCastShadow |= MeshBatches[SegmentIndex].CastRayTracedShadow && Material.CastsRayTracedShadows();
-				InstanceMask |= ComputeBlendModeMask(Material.GetBlendMode());
+				InstanceMask |= ComputeBlendModeMask(Material);
 			}
 
 			InstanceMask |= bAnySegmentsCastShadow ? RAY_TRACING_MASK_SHADOW : 0;
@@ -1056,7 +1056,7 @@ bool FSceneRenderState::SetupRayTracingScene(int32 LODIndex)
 							bAllSegmentsUnlit &= Material.GetShadingModels().HasOnlyShadingModel(MSM_Unlit) || !MeshBatches[SegmentIndex].CastShadow;
 							bAllSegmentsOpaque &= IsOpaqueBlendMode(Material);
 							bAnySegmentsCastShadow |= MeshBatches[SegmentIndex].CastRayTracedShadow && Material.CastsRayTracedShadows();
-							InstanceMask |= ComputeBlendModeMask(Material.GetBlendMode());
+							InstanceMask |= ComputeBlendModeMask(Material);
 						}
 
 						InstanceMask |= bAnySegmentsCastShadow ? RAY_TRACING_MASK_SHADOW : 0;
diff --git a/Engine/Source/Runtime/Engine/Private/RayTracingInstance.cpp b/Engine/Source/Runtime/Engine/Private/RayTracingInstance.cpp
--- a/Engine/Source/Runtime/Engine/Private/RayTracingInstance.cpp
+++ b/Engine/Source/Runtime/Engine/Private/RayTracingInstance.cpp
@@ -41,8 +41,7 @@ FRayTracingMaskAndFlags BuildRayTracingInstanceMaskAndFlags(TArrayView<const FMe
 		if (MeshBatch.bUseForMaterial && MeshBatch.MaterialRenderProxy)
 		{
 			const FMaterial& Material = MeshBatch.MaterialRenderProxy->GetIncompleteMaterialWithFallback(FeatureLevel);
-			const EBlendMode BlendMode = Material.GetBlendMode(); // STRATA_TODO_BLENDMODE
-			Result.Mask |= ComputeBlendModeMask(BlendMode);
+			Result.Mask |= ComputeBlendModeMask(Material);
 			bAllSegmentsOpaque &= IsOpaqueBlendMode(Material);
 			bAnySegmentsCastShadow |= MeshBatch.CastRayTracedShadow && Material.CastsRayTracedShadows();
 			bAllSegmentsCastShadow &= MeshBatch.CastRayTracedShadow && Material.CastsRayTracedShadows();
@@ -77,9 +76,9 @@ FRayTracingMaskAndFlags BuildRayTracingInstanceMaskAndFlags(TArrayView<const FMe
 	return Result;
 }
 
-uint8 ComputeBlendModeMask(const EBlendMode BlendMode)
+uint8 ComputeBlendModeMask(const FMaterial& Material)
 {
-	return (BlendMode == BLEND_Opaque || BlendMode == BLEND_Masked) ? RAY_TRACING_MASK_OPAQUE : RAY_TRACING_MASK_TRANSLUCENT;
+	return IsOpaqueOrMaskedBlendMode(Material) ? RAY_TRACING_MASK_OPAQUE : RAY_TRACING_MASK_TRANSLUCENT;
 }
 
 #endif
diff --git a/Engine/Source/Runtime/Engine/Public/RayTracingInstance.h b/Engine/Source/Runtime/Engine/Public/RayTracingInstance.h
--- a/Engine/Source/Runtime/Engine/Public/RayTracingInstance.h
+++ b/Engine/Source/Runtime/Engine/Public/RayTracingInstance.h
@@ -113,5 +113,5 @@ struct FRayTracingInstance
 
 /** Build mask and flags based on materials specified in Materials. You can still override Mask after calling this function. */
 ENGINE_API FRayTracingMaskAndFlags BuildRayTracingInstanceMaskAndFlags(TArrayView<const FMeshBatch> MeshBatches, ERHIFeatureLevel::Type FeatureLevel, ERayTracingInstanceLayer InstanceLayer = ERayTracingInstanceLayer::NearField, uint8 ExtraMask = 0);
-ENGINE_API uint8 ComputeBlendModeMask(const EBlendMode BlendMode);
+ENGINE_API uint8 ComputeBlendModeMask(const FMaterial& Material);
 #endif
diff --git a/Engine/Source/Runtime/Renderer/Private/RayTracing/RayTracingMaterialHitShaders.h b/Engine/Source/Runtime/Renderer/Private/RayTracing/RayTracingMaterialHitShaders.h
--- a/Engine/Source/Runtime/Renderer/Private/RayTracing/RayTracingMaterialHitShaders.h
+++ b/Engine/Source/Runtime/Renderer/Private/RayTracing/RayTracingMaterialHitShaders.h
@@ -10,7 +10,7 @@
 
 #include "RayTracingPayloadType.h"
 
-ENGINE_API uint8 ComputeBlendModeMask(const EBlendMode BlendMode);
+ENGINE_API uint8 ComputeBlendModeMask(const FMaterial& Material);
 FRHIRayTracingShader* GetRayTracingDefaultMissShader(const FGlobalShaderMap* ShaderMap);
 FRHIRayTracingShader* GetRayTracingDefaultOpaqueShader(const FGlobalShaderMap* ShaderMap);
 FRHIRayTracingShader* GetRayTracingDefaultHiddenShader(const FGlobalShaderMap* ShaderMap);
@@ -74,16 +74,13 @@ protected:
 			SharedCommand.SetShaders(PassShaders.GetUntypedShaders());
 		}
 
-		EBlendMode BlendMode = MaterialResource.GetBlendMode();
-		EStrataBlendMode StrataBlendMode = MaterialResource.GetStrataBlendMode();
-
-		SharedCommand.InstanceMask = ComputeBlendModeMask(BlendMode);
+		SharedCommand.InstanceMask = ComputeBlendModeMask(MaterialResource);
 		SharedCommand.bCastRayTracedShadows = MeshBatch.CastRayTracedShadow && MaterialResource.CastsRayTracedShadows();
-		SharedCommand.bOpaque = IsOpaqueBlendMode(BlendMode, StrataBlendMode) && !(VertexFactory->GetType()->SupportsRayTracingProceduralPrimitive() && FDataDrivenShaderPlatformInfo::GetSupportsRayTracingProceduralPrimitive(GMaxRHIShaderPlatform));
+		SharedCommand.bOpaque = IsOpaqueBlendMode(MaterialResource) && !(VertexFactory->GetType()->SupportsRayTracingProceduralPrimitive() && FDataDrivenShaderPlatformInfo::GetSupportsRayTracingProceduralPrimitive(GMaxRHIShaderPlatform));
 		SharedCommand.bDecal = MaterialResource.GetMaterialDomain() == EMaterialDomain::MD_DeferredDecal;
 		SharedCommand.bIsSky = MaterialResource.IsSky();
 		SharedCommand.bTwoSided = MaterialResource.IsTwoSided();
-		SharedCommand.bIsTranslucent = IsTranslucentOnlyBlendMode(BlendMode, StrataBlendMode);
+		SharedCommand.bIsTranslucent = IsTranslucentOnlyBlendMode(MaterialResource);
 
 		FVertexInputStreamArray VertexStreams;
 		VertexFactory->GetStreams(ERHIFeatureLevel::SM5, EVertexInputStreamType::Default, VertexStreams);
