Adding missing includes after recent dependency cleanups.

diff --git a/Engine/Plugins/Interchange/Runtime/Source/Import/Public/Material/InterchangeMaterialFactory.h b/Engine/Plugins/Interchange/Runtime/Source/Import/Public/Material/InterchangeMaterialFactory.h
--- a/Engine/Plugins/Interchange/Runtime/Source/Import/Public/Material/InterchangeMaterialFactory.h
+++ b/Engine/Plugins/Interchange/Runtime/Source/Import/Public/Material/InterchangeMaterialFactory.h
@@ -14,6 +14,7 @@ class UInterchangeBaseMaterialFactoryNode;
 class UInterchangeMaterialExpressionFactoryNode;
 class UMaterial;
 class UMaterialExpression;
+class UMaterialExpressionMaterialFunctionCall;
 class UMaterialInstance;
 
 UCLASS(BlueprintType)
diff --git a/Engine/Source/Runtime/Engine/Classes/Materials/MaterialExpressionRerouteBase.h b/Engine/Source/Runtime/Engine/Classes/Materials/MaterialExpressionRerouteBase.h
--- a/Engine/Source/Runtime/Engine/Classes/Materials/MaterialExpressionRerouteBase.h
+++ b/Engine/Source/Runtime/Engine/Classes/Materials/MaterialExpressionRerouteBase.h
@@ -5,6 +5,8 @@
 #include "Materials/MaterialExpression.h"
 #include "MaterialExpressionRerouteBase.generated.h"
 
+class FMaterialExpressionKey;
+
 UCLASS(abstract)
 class ENGINE_API UMaterialExpressionRerouteBase : public UMaterialExpression
 {
diff --git a/Engine/Source/Runtime/Engine/Classes/Materials/MaterialExpressionShadingModel.h b/Engine/Source/Runtime/Engine/Classes/Materials/MaterialExpressionShadingModel.h
--- a/Engine/Source/Runtime/Engine/Classes/Materials/MaterialExpressionShadingModel.h
+++ b/Engine/Source/Runtime/Engine/Classes/Materials/MaterialExpressionShadingModel.h
@@ -7,6 +7,8 @@
 #include "Materials/MaterialExpression.h"
 #include "MaterialExpressionShadingModel.generated.h"
 
+class FMaterialCompiler;
+
 /**
  * Compile a select "blend" between ShadingModels
  *
diff --git a/Engine/Source/Runtime/Engine/Classes/Materials/MaterialExpressionStrata.h b/Engine/Source/Runtime/Engine/Classes/Materials/MaterialExpressionStrata.h
--- a/Engine/Source/Runtime/Engine/Classes/Materials/MaterialExpressionStrata.h
+++ b/Engine/Source/Runtime/Engine/Classes/Materials/MaterialExpressionStrata.h
@@ -7,6 +7,7 @@
 #include "Materials/MaterialExpression.h"
 #include "MaterialExpressionStrata.generated.h"
 
+class FMaterialCompiler;
 
 /**
  * Compile a special blend function for strata when blending material attribute
diff --git a/Engine/Source/Runtime/Engine/Private/LightMap.cpp b/Engine/Source/Runtime/Engine/Private/LightMap.cpp
--- a/Engine/Source/Runtime/Engine/Private/LightMap.cpp
+++ b/Engine/Source/Runtime/Engine/Private/LightMap.cpp
@@ -26,6 +26,7 @@
 #include "Misc/PackageName.h"
 #include "ProfilingDebugging/LoadTimeTracker.h"
 #include "TextureResource.h"
+#include "DataDrivenShaderPlatformInfo.h"
 
 #define VISUALIZE_PACKING 0
 
diff --git a/Engine/Source/Runtime/Renderer/Private/VariableRateShading/ContrastAdaptiveImageGenerator.cpp b/Engine/Source/Runtime/Renderer/Private/VariableRateShading/ContrastAdaptiveImageGenerator.cpp
--- a/Engine/Source/Runtime/Renderer/Private/VariableRateShading/ContrastAdaptiveImageGenerator.cpp
+++ b/Engine/Source/Runtime/Renderer/Private/VariableRateShading/ContrastAdaptiveImageGenerator.cpp
@@ -16,6 +16,7 @@
 #include "Engine/Engine.h"
 #include "UnrealClient.h"
 #include "PostProcess/PostProcessTonemap.h"
+#include "DataDrivenShaderPlatformInfo.h"
 
 /**
  * Contrast Adaptive Shading (CAS) is a Tier 2 Variable Rate Shading method which generates a VRS image by examining the contrast from the previous frame.
diff --git a/Engine/Source/Runtime/Renderer/Private/VariableRateShading/FixedFoveationImageGenerator.cpp b/Engine/Source/Runtime/Renderer/Private/VariableRateShading/FixedFoveationImageGenerator.cpp
--- a/Engine/Source/Runtime/Renderer/Private/VariableRateShading/FixedFoveationImageGenerator.cpp
+++ b/Engine/Source/Runtime/Renderer/Private/VariableRateShading/FixedFoveationImageGenerator.cpp
@@ -7,6 +7,7 @@
 #include "StereoRendering.h"
 #include "SceneView.h"
 #include "UnrealClient.h"
+#include "DataDrivenShaderPlatformInfo.h"
 
 /* CVar values used to control generator behavior */
 
