Change CalculateDepthPassMeshStaticSortKey to not use blend mode as input.

diff --git a/Engine/Source/Runtime/Renderer/Private/DepthRendering.cpp b/Engine/Source/Runtime/Renderer/Private/DepthRendering.cpp
--- a/Engine/Source/Runtime/Renderer/Private/DepthRendering.cpp
+++ b/Engine/Source/Runtime/Renderer/Private/DepthRendering.cpp
@@ -676,14 +676,14 @@ void FDeferredShadingSceneRenderer::RenderPrePassHMD(FRDGBuilder& GraphBuilder,
 	}
 }
 
-FMeshDrawCommandSortKey CalculateDepthPassMeshStaticSortKey(EBlendMode BlendMode, const FMeshMaterialShader* VertexShader, const FMeshMaterialShader* PixelShader)
+FMeshDrawCommandSortKey CalculateDepthPassMeshStaticSortKey(bool bIsMasked, const FMeshMaterialShader* VertexShader, const FMeshMaterialShader* PixelShader)
 {
 	FMeshDrawCommandSortKey SortKey;
 	if (GEarlyZSortMasked)
 	{
 		SortKey.BasePass.VertexShaderHash = (VertexShader ? VertexShader->GetSortKey() : 0) & 0xFFFF;
 		SortKey.BasePass.PixelShaderHash = PixelShader ? PixelShader->GetSortKey() : 0;
-		SortKey.BasePass.Masked = BlendMode == EBlendMode::BLEND_Masked ? 1 : 0;
+		SortKey.BasePass.Masked = bIsMasked ? 1 : 0;
 	}
 	else
 	{
@@ -727,7 +727,6 @@ bool FDepthPassMeshProcessor::Process(
 	const FMeshBatch& RESTRICT MeshBatch,
 	uint64 BatchElementMask,
 	int32 StaticMeshId,
-	EBlendMode BlendMode,
 	const FPrimitiveSceneProxy* RESTRICT PrimitiveSceneProxy,
 	const FMaterialRenderProxy& RESTRICT MaterialRenderProxy,
 	const FMaterial& RESTRICT MaterialResource,
@@ -770,7 +769,8 @@ bool FDepthPassMeshProcessor::Process(
 	FMeshMaterialShaderElementData ShaderElementData;
 	ShaderElementData.InitializeMeshMaterialData(ViewIfDynamicMeshCommand, PrimitiveSceneProxy, MeshBatch, StaticMeshId, true);
 
-	const FMeshDrawCommandSortKey SortKey = CalculateDepthPassMeshStaticSortKey(BlendMode, DepthPassShaders.VertexShader.GetShader(), DepthPassShaders.PixelShader.GetShader());
+	const bool bIsMasked = IsMaskedBlendMode(MaterialResource);
+	const FMeshDrawCommandSortKey SortKey = CalculateDepthPassMeshStaticSortKey(bIsMasked, DepthPassShaders.VertexShader.GetShader(), DepthPassShaders.PixelShader.GetShader());
 
 	BuildMeshDrawCommands(
 		MeshBatch,
@@ -878,7 +878,6 @@ bool FDepthPassMeshProcessor::UseDefaultMaterial(const FMaterial& Material, bool
 
 bool FDepthPassMeshProcessor::TryAddMeshBatch(const FMeshBatch& RESTRICT MeshBatch, uint64 BatchElementMask, const FPrimitiveSceneProxy* RESTRICT PrimitiveSceneProxy, int32 StaticMeshId, const FMaterialRenderProxy& MaterialRenderProxy, const FMaterial& Material)
 {
-	const EBlendMode BlendMode = Material.GetBlendMode();
 	const bool bIsTranslucent = IsTranslucentBlendMode(Material);
 	bool ShouldRenderInDepthPass = (!PrimitiveSceneProxy || PrimitiveSceneProxy->ShouldRenderInDepthPass());
 
@@ -907,11 +906,11 @@ bool FDepthPassMeshProcessor::TryAddMeshBatch(const FMeshBatch& RESTRICT MeshBat
 
 		if (bPositionOnly)
 		{
-			bResult = Process<true>(MeshBatch, BatchElementMask, StaticMeshId, BlendMode /*STRATA_TODO_BLENDMODE StrataBlendMode */, PrimitiveSceneProxy, *EffectiveMaterialRenderProxy, *EffectiveMaterial, MeshFillMode, MeshCullMode);
+			bResult = Process<true>(MeshBatch, BatchElementMask, StaticMeshId, PrimitiveSceneProxy, *EffectiveMaterialRenderProxy, *EffectiveMaterial, MeshFillMode, MeshCullMode);
 		}
 		else
 		{
-			bResult = Process<false>(MeshBatch, BatchElementMask, StaticMeshId, BlendMode /*STRATA_TODO_BLENDMODE StrataBlendMode */, PrimitiveSceneProxy, *EffectiveMaterialRenderProxy, *EffectiveMaterial, MeshFillMode, MeshCullMode);
+			bResult = Process<false>(MeshBatch, BatchElementMask, StaticMeshId, PrimitiveSceneProxy, *EffectiveMaterialRenderProxy, *EffectiveMaterial, MeshFillMode, MeshCullMode);
 		}
 	}
 
diff --git a/Engine/Source/Runtime/Renderer/Private/DepthRendering.h b/Engine/Source/Runtime/Renderer/Private/DepthRendering.h
--- a/Engine/Source/Runtime/Renderer/Private/DepthRendering.h
+++ b/Engine/Source/Runtime/Renderer/Private/DepthRendering.h
@@ -57,7 +57,7 @@ extern FDepthPassInfo GetDepthPassInfo(const FScene* Scene);
 
 void AddDitheredStencilFillPass(FRDGBuilder& GraphBuilder, TConstArrayView<FViewInfo> Views, FRDGTextureRef DepthTexture, const FDepthPassInfo& DepthPass);
 
-FMeshDrawCommandSortKey CalculateDepthPassMeshStaticSortKey(EBlendMode BlendMode, const FMeshMaterialShader* VertexShader, const FMeshMaterialShader* PixelShader);
+FMeshDrawCommandSortKey CalculateDepthPassMeshStaticSortKey(const bool bIsMasked, const FMeshMaterialShader* VertexShader, const FMeshMaterialShader* PixelShader);
 
 /**
  * A vertex shader for rendering the depth of a mesh.
@@ -198,7 +198,6 @@ private:
 		const FMeshBatch& MeshBatch,
 		uint64 BatchElementMask,
 		int32 StaticMeshId,
-		EBlendMode BlendMode,
 		const FPrimitiveSceneProxy* RESTRICT PrimitiveSceneProxy,
 		const FMaterialRenderProxy& RESTRICT MaterialRenderProxy,
 		const FMaterial& RESTRICT MaterialResource,
diff --git a/Engine/Source/Runtime/Renderer/Private/SingleLayerWaterRendering.cpp b/Engine/Source/Runtime/Renderer/Private/SingleLayerWaterRendering.cpp
--- a/Engine/Source/Runtime/Renderer/Private/SingleLayerWaterRendering.cpp
+++ b/Engine/Source/Runtime/Renderer/Private/SingleLayerWaterRendering.cpp
@@ -1669,8 +1669,7 @@ private:
 		const FMaterialRenderProxy& RESTRICT MaterialRenderProxy,
 		const FMaterial& RESTRICT MaterialResource,
 		ERasterizerFillMode MeshFillMode,
-		ERasterizerCullMode MeshCullMode,
-		EBlendMode BlendMode);
+		ERasterizerCullMode MeshCullMode);
 
 	template<bool bPositionOnly>
 	void CollectPSOInitializersInternal(
@@ -1727,20 +1726,18 @@ bool FSingleLayerWaterDepthPrepassMeshProcessor::TryAddMeshBatch(
 	if (Material.GetShadingModels().HasShadingModel(MSM_SingleLayerWater))
 	{
 		// Determine the mesh's material and blend mode.
-		const EBlendMode BlendMode = Material.GetBlendMode();
-		const EStrataBlendMode StrataBlendMode = Material.GetStrataBlendMode();
 		const FMeshDrawingPolicyOverrideSettings OverrideSettings = ComputeMeshOverrideSettings(MeshBatch);
 		const ERasterizerFillMode MeshFillMode = ComputeMeshFillMode(Material, OverrideSettings);
 		const ERasterizerCullMode MeshCullMode = ComputeMeshCullMode(Material, OverrideSettings);
 
-		if (IsOpaqueBlendMode(BlendMode, StrataBlendMode)
+		if (IsOpaqueBlendMode(Material)
 			&& MeshBatch.VertexFactory->SupportsPositionOnlyStream()
 			&& !Material.MaterialModifiesMeshPosition_RenderThread()
 			&& Material.WritesEveryPixel())
 		{
 			const FMaterialRenderProxy& DefaultProxy = *UMaterial::GetDefaultMaterial(MD_Surface)->GetRenderProxy();
 			const FMaterial& DefaultMaterial = *DefaultProxy.GetMaterialNoFallback(FeatureLevel);
-			return Process<true>(MeshBatch, BatchElementMask, PrimitiveSceneProxy, StaticMeshId, DefaultProxy, DefaultMaterial, MeshFillMode, MeshCullMode, BlendMode);
+			return Process<true>(MeshBatch, BatchElementMask, PrimitiveSceneProxy, StaticMeshId, DefaultProxy, DefaultMaterial, MeshFillMode, MeshCullMode);
 		}
 		else
 		{
@@ -1756,7 +1753,7 @@ bool FSingleLayerWaterDepthPrepassMeshProcessor::TryAddMeshBatch(
 				check(EffectiveMaterial);
 			}
 
-			return Process<false>(MeshBatch, BatchElementMask, PrimitiveSceneProxy, StaticMeshId, *EffectiveMaterialRenderProxy, *EffectiveMaterial, MeshFillMode, MeshCullMode, BlendMode);
+			return Process<false>(MeshBatch, BatchElementMask, PrimitiveSceneProxy, StaticMeshId, *EffectiveMaterialRenderProxy, *EffectiveMaterial, MeshFillMode, MeshCullMode);
 		}
 	}
 
@@ -1772,8 +1769,7 @@ bool FSingleLayerWaterDepthPrepassMeshProcessor::Process(
 	const FMaterialRenderProxy& RESTRICT MaterialRenderProxy,
 	const FMaterial& RESTRICT MaterialResource,
 	ERasterizerFillMode MeshFillMode,
-	ERasterizerCullMode MeshCullMode,
-	EBlendMode BlendMode)
+	ERasterizerCullMode MeshCullMode)
 {
 	TMeshProcessorShaders<TDepthOnlyVS<bPositionOnly>, FDepthOnlyPS> DepthPassShaders;
 	FShaderPipelineRef ShaderPipeline;
@@ -1793,7 +1789,8 @@ bool FSingleLayerWaterDepthPrepassMeshProcessor::Process(
 	FMeshMaterialShaderElementData ShaderElementData;
 	ShaderElementData.InitializeMeshMaterialData(ViewIfDynamicMeshCommand, PrimitiveSceneProxy, MeshBatch, StaticMeshId, true);
 
-	FMeshDrawCommandSortKey SortKey = CalculateDepthPassMeshStaticSortKey(BlendMode, DepthPassShaders.VertexShader.GetShader(), DepthPassShaders.PixelShader.GetShader());
+	const bool bIsMasked = IsMaskedBlendMode(MaterialResource);
+	FMeshDrawCommandSortKey SortKey = CalculateDepthPassMeshStaticSortKey(bIsMasked, DepthPassShaders.VertexShader.GetShader(), DepthPassShaders.PixelShader.GetShader());
 
 	BuildMeshDrawCommands(
 		MeshBatch,
