PR #9905: Option to disable reflections and sky contribution in reflection captures (Contributed by [REDACTED_PERSON])

diff --git a/Engine/Source/Runtime/Renderer/Private/IndirectLightRendering.cpp b/Engine/Source/Runtime/Renderer/Private/IndirectLightRendering.cpp
--- a/Engine/Source/Runtime/Renderer/Private/IndirectLightRendering.cpp
+++ b/Engine/Source/Runtime/Renderer/Private/IndirectLightRendering.cpp
@@ -76,6 +76,12 @@ static TAutoConsoleVariable<float> CVarSkySpecularOcclusionStrength(
 	TEXT("Strength of skylight specular occlusion from DFAO (default is 1.0)"),
 	ECVF_RenderThreadSafe);
 
+static TAutoConsoleVariable<int32> CVarReflectionCaptureEnableDeferredReflectionsAndSkyLighting(
+	TEXT("r.ReflectionCapture.EnableDefferedReflectionsAndSkyLighting"),
+	1,
+	TEXT("Whether to evaluate deferred reflections and sky contribution when rendering reflection captures."),
+	ECVF_RenderThreadSafe);
+
 static TAutoConsoleVariable<int32> CVarProbeSamplePerPixel(
 	TEXT("r.Lumen.ProbeHierarchy.SamplePerPixel"), 8,
 	TEXT("Number of sample to do per full res pixel."),
@@ -94,6 +100,11 @@ DECLARE_GPU_STAT(SkyLightDiffuse);
 int GetReflectionEnvironmentCVar();
 bool IsAmbientCubemapPassRequired(const FSceneView& View);
 
+bool IsDeferredReflectionsAndSkyLightingDisabledForReflectionCapture()
+{
+	return CVarReflectionCaptureEnableDeferredReflectionsAndSkyLighting.GetValueOnRenderThread() == 0;
+}
+
 class FDiffuseIndirectCompositePS : public FGlobalShader
 {
 	DECLARE_GLOBAL_SHADER(FDiffuseIndirectCompositePS)
@@ -1732,6 +1743,12 @@ void FDeferredShadingSceneRenderer::RenderDeferredReflectionsAndSkyLighting(
 		bReflectionCapture = bReflectionCapture || View.bIsReflectionCapture;
 	}
 
+	if (bReflectionCapture 
+		&& IsDeferredReflectionsAndSkyLightingDisabledForReflectionCapture())
+	{
+		return;
+	}
+
 	// The specular sky light contribution is also needed by RT Reflections as a fallback.
 	const bool bSkyLight = Scene->SkyLight
 		&& (Scene->SkyLight->ProcessedTexture || Scene->SkyLight->bRealTimeCaptureEnabled)
