Fix typo in the CVar to disable deferred reflections and skylighting for reflection capture.

diff --git a/Engine/Source/Runtime/Renderer/Private/IndirectLightRendering.cpp b/Engine/Source/Runtime/Renderer/Private/IndirectLightRendering.cpp
--- a/Engine/Source/Runtime/Renderer/Private/IndirectLightRendering.cpp
+++ b/Engine/Source/Runtime/Renderer/Private/IndirectLightRendering.cpp
@@ -77,7 +77,7 @@ static TAutoConsoleVariable<float> CVarSkySpecularOcclusionStrength(
 	ECVF_RenderThreadSafe);
 
 static TAutoConsoleVariable<int32> CVarReflectionCaptureEnableDeferredReflectionsAndSkyLighting(
-	TEXT("r.ReflectionCapture.EnableDefferedReflectionsAndSkyLighting"),
+	TEXT("r.ReflectionCapture.EnableDeferredReflectionsAndSkyLighting"),
 	1,
 	TEXT("Whether to evaluate deferred reflections and sky contribution when rendering reflection captures."),
 	ECVF_RenderThreadSafe);
