Fix decal blend mode.

diff --git a/Engine/Source/Runtime/Engine/Classes/Engine/EngineTypes.h b/Engine/Source/Runtime/Engine/Classes/Engine/EngineTypes.h
--- a/Engine/Source/Runtime/Engine/Classes/Engine/EngineTypes.h
+++ b/Engine/Source/Runtime/Engine/Classes/Engine/EngineTypes.h
@@ -318,6 +318,18 @@ ENGINE_API bool IsModulateBlendMode(EBlendMode LegacyBlendMode, EStrataBlendMode
 ENGINE_API bool IsModulateBlendMode(const FMaterial& In);
 ENGINE_API bool IsModulateBlendMode(const UMaterialInterface& In);
 
+ENGINE_API bool IsAdditiveBlendMode(EStrataBlendMode BlendMode);
+ENGINE_API bool IsAdditiveBlendMode(EBlendMode BlendMode);
+ENGINE_API bool IsAdditiveBlendMode(EBlendMode LegacyBlendMode, EStrataBlendMode StrataBlendMode);
+ENGINE_API bool IsAdditiveBlendMode(const FMaterial& In);
+ENGINE_API bool IsAdditiveBlendMode(const UMaterialInterface& In);
+
+ENGINE_API bool IsAlphaCompositeBlendMode(EStrataBlendMode BlendMode);
+ENGINE_API bool IsAlphaCompositeBlendMode(EBlendMode BlendMode);
+ENGINE_API bool IsAlphaCompositeBlendMode(EBlendMode LegacyBlendMode, EStrataBlendMode StrataBlendMode);
+ENGINE_API bool IsAlphaCompositeBlendMode(const FMaterial& In);
+ENGINE_API bool IsAlphaCompositeBlendMode(const UMaterialInterface& In);
+
 /** The default float precision for material's pixel shaders on mobile devices*/
 UENUM()
 enum EMaterialFloatPrecisionMode : int
diff --git a/Engine/Source/Runtime/Engine/Private/Materials/StrataMaterial.cpp b/Engine/Source/Runtime/Engine/Private/Materials/StrataMaterial.cpp
--- a/Engine/Source/Runtime/Engine/Private/Materials/StrataMaterial.cpp
+++ b/Engine/Source/Runtime/Engine/Private/Materials/StrataMaterial.cpp
@@ -67,4 +67,6 @@ IsGenericBlendMode(OpaqueOrMasked,	BlendMode == BLEND_Opaque || BlendMode == BLE
 IsGenericBlendMode(Translucent,		BlendMode != BLEND_Opaque && BlendMode != BLEND_Masked, BlendMode != SBM_Opaque && BlendMode != SBM_Masked)													// General translucency (i.e., blend mode is something else than Opaque/Masked)
 IsGenericBlendMode(TranslucentOnly, BlendMode == BLEND_Translucent,							BlendMode == SBM_TranslucentColoredTransmittance || BlendMode == SBM_TranslucentGreyTransmittance) 	// Explicit translucency blend mode
 IsGenericBlendMode(AlphaHoldout,	BlendMode == BLEND_AlphaHoldout,						BlendMode == SBM_AlphaHoldout)																		// AlphaHoldout blend mode
-IsGenericBlendMode(Modulate,		BlendMode == BLEND_Modulate,							BlendMode == SBM_ColoredTransmittanceOnly)															// Modulate blend mode
\ No newline at end of file
+IsGenericBlendMode(Modulate,		BlendMode == BLEND_Modulate,							BlendMode == SBM_ColoredTransmittanceOnly)															// Modulate blend mode
+IsGenericBlendMode(Additive,		BlendMode == BLEND_Additive,							BlendMode == SBM_TranslucentGreyTransmittance)														// Additive blend mode
+IsGenericBlendMode(AlphaComposite,	BlendMode == BLEND_AlphaComposite,						BlendMode == SBM_TranslucentGreyTransmittance)														// AlphaComposite blend mode
\ No newline at end of file
diff --git a/Engine/Source/Runtime/Renderer/Private/DecalRenderingCommon.cpp b/Engine/Source/Runtime/Renderer/Private/DecalRenderingCommon.cpp
--- a/Engine/Source/Runtime/Renderer/Private/DecalRenderingCommon.cpp
+++ b/Engine/Source/Runtime/Renderer/Private/DecalRenderingCommon.cpp
@@ -8,9 +8,19 @@
 
 namespace DecalRendering
 {
+	inline bool IsOpaqueBlendMode(const FDecalBlendDesc& In)			{ return IsOpaqueBlendMode((EBlendMode)In.BlendMode, (EStrataBlendMode)In.StrataBlendMode); }
+	inline bool IsOpaqueOrMaskedBlendMode(const FDecalBlendDesc& In)	{ return IsOpaqueOrMaskedBlendMode((EBlendMode)In.BlendMode, (EStrataBlendMode)In.StrataBlendMode); }
+	inline bool IsMaskedBlendMode(const FDecalBlendDesc& In)			{ return IsMaskedBlendMode((EBlendMode)In.BlendMode, (EStrataBlendMode)In.StrataBlendMode); }
+	inline bool IsTranslucentOnlyBlendMode(const FDecalBlendDesc& In)	{ return IsTranslucentOnlyBlendMode((EBlendMode)In.BlendMode, (EStrataBlendMode)In.StrataBlendMode); }
+	inline bool IsTranslucentBlendMode(const FDecalBlendDesc& In)		{ return IsTranslucentBlendMode((EBlendMode)In.BlendMode, (EStrataBlendMode)In.StrataBlendMode); }
+	inline bool IsAlphaHoldoutBlendMode(const FDecalBlendDesc& In)		{ return IsAlphaHoldoutBlendMode((EBlendMode)In.BlendMode, (EStrataBlendMode)In.StrataBlendMode); }
+	inline bool IsModulateBlendMode(const FDecalBlendDesc& In)			{ return IsModulateBlendMode((EBlendMode)In.BlendMode, (EStrataBlendMode)In.StrataBlendMode); }
+	inline bool IsAlphaCompositeBlendMode(const FDecalBlendDesc& In)	{ return IsAlphaCompositeBlendMode((EBlendMode)In.BlendMode, (EStrataBlendMode)In.StrataBlendMode); }
+
 	/** Finalize the initialization of FDecalBlendDesc after BlendMode and bWrite flags have all been set. */
 	void FinalizeBlendDesc(EShaderPlatform Platform, FDecalBlendDesc& Desc)
 	{
+		const bool bIsStrataEnabled = Strata::IsStrataEnabled();
 		const bool bIsMobilePlatform = IsMobilePlatform(Platform);
 		const bool bIsMobileDeferredPlatform = bIsMobilePlatform && IsMobileDeferredShadingEnabled(Platform);
 		const bool bIsDBufferPlatform = !bIsMobilePlatform && IsUsingDBuffers(Platform);
@@ -19,13 +29,15 @@ namespace DecalRendering
 		Desc.bWriteDBufferMask = bIsDBufferMaskPlatform;
 
 		// Enforce platform blend mode limitations.
-		if (Desc.BlendMode != BLEND_Translucent && Desc.BlendMode != BLEND_AlphaComposite && Desc.BlendMode != BLEND_Modulate)
+		if (!IsTranslucentOnlyBlendMode(Desc) && !IsAlphaCompositeBlendMode(Desc) && !IsModulateBlendMode(Desc))
 		{
 			Desc.BlendMode = BLEND_Translucent;
+			Desc.StrataBlendMode = SBM_TranslucentGreyTransmittance;
 		}
-		if (bIsDBufferPlatform && Desc.BlendMode == BLEND_Modulate)
+		if (bIsDBufferPlatform && IsModulateBlendMode(Desc))
 		{
 			Desc.BlendMode = BLEND_Translucent;
+			Desc.StrataBlendMode = SBM_TranslucentGreyTransmittance;
 		}
 
 		// Enforce platform output limitations.
@@ -40,7 +52,7 @@ namespace DecalRendering
 		}
 
 		// Enforce blend modes output limitations.
-		if (Desc.BlendMode == BLEND_AlphaComposite)
+		if (!bIsStrataEnabled && IsAlphaCompositeBlendMode(Desc))
 		{
 			Desc.bWriteNormal = false;
 		}
@@ -86,6 +98,7 @@ namespace DecalRendering
 				Material.HasMaterialPropertyConnected(EMaterialProperty::MP_SpecularColor);	// This is used for Strata Slab using (DiffuseAlbedo | F0) parameterization
 
 			Desc.BlendMode = Material.GetBlendMode();
+			Desc.StrataBlendMode = Material.GetStrataBlendMode();
 			Desc.bWriteBaseColor = Material.HasMaterialPropertyConnected(EMaterialProperty::MP_BaseColor) || bUseDiffuseAlbedoAndF0;
 			Desc.bWriteNormal = Material.HasMaterialPropertyConnected(EMaterialProperty::MP_Normal);
 			Desc.bWriteRoughnessSpecularMetallic =
@@ -101,6 +114,7 @@ namespace DecalRendering
 		else
 		{
 			Desc.BlendMode = Material.GetBlendMode();
+			Desc.StrataBlendMode = Material.GetStrataBlendMode();
 			Desc.bWriteBaseColor = Material.HasBaseColorConnected();
 			Desc.bWriteNormal = Material.HasNormalConnected();
 			Desc.bWriteRoughnessSpecularMetallic = Material.HasRoughnessConnected() || Material.HasSpecularConnected() || Material.HasMetallicConnected();
@@ -120,7 +134,8 @@ namespace DecalRendering
 				MaterialShaderParameters.bHasDiffuseAlbedoConnected || 
 				MaterialShaderParameters.bHasF0Connected;
 			
-			Desc.BlendMode = MaterialShaderParameters.BlendMode; // STRATA_TODO_BLENDMODE
+			Desc.BlendMode = MaterialShaderParameters.BlendMode;
+			Desc.StrataBlendMode = MaterialShaderParameters.StrataBlendMode;
 			Desc.bWriteBaseColor = MaterialShaderParameters.bHasBaseColorConnected || bUseDiffuseAlbedoAndF0;
 			Desc.bWriteNormal = MaterialShaderParameters.bHasNormalConnected;
 			Desc.bWriteRoughnessSpecularMetallic = 
@@ -133,7 +148,8 @@ namespace DecalRendering
 		}
 		else
 		{
-			Desc.BlendMode = MaterialShaderParameters.BlendMode; // STRATA_TODO_BLENDMODE
+			Desc.BlendMode = MaterialShaderParameters.BlendMode;
+			Desc.StrataBlendMode = MaterialShaderParameters.StrataBlendMode;
 			Desc.bWriteBaseColor = MaterialShaderParameters.bHasBaseColorConnected;
 			Desc.bWriteNormal = MaterialShaderParameters.bHasNormalConnected;
 			Desc.bWriteRoughnessSpecularMetallic = MaterialShaderParameters.bHasRoughnessConnected || MaterialShaderParameters.bHasSpecularConnected || MaterialShaderParameters.bHasMetallicConnected;
@@ -247,7 +263,7 @@ namespace DecalRendering
 		// Ignore DBuffer mask bit and always set that MRT active.
 		const uint32 Mask = GetRenderTargetWriteMask(DecalBlendDesc, DecalRenderStage, EDecalRenderTargetMode::DBuffer) & 0x7;
 
-		if (DecalBlendDesc.BlendMode == BLEND_Translucent)
+		if (IsTranslucentOnlyBlendMode(DecalBlendDesc))
 		{
 			switch (Mask)
 			{
@@ -302,7 +318,7 @@ namespace DecalRendering
 				>::GetRHI();
 			}
 		}
-		else if (DecalBlendDesc.BlendMode == BLEND_AlphaComposite)
+		else if (IsAlphaCompositeBlendMode(DecalBlendDesc))
 		{
 			ensure((Mask & 2) == 0); // AlphaComposite shouldn't write normal.
 
@@ -340,7 +356,7 @@ namespace DecalRendering
 	{
 		const uint32 Mask = GetRenderTargetWriteMask(DecalBlendDesc, DecalRenderStage, EDecalRenderTargetMode::SceneColorAndGBuffer);
 
-		if (DecalBlendDesc.BlendMode == BLEND_Translucent)
+		if (IsTranslucentOnlyBlendMode(DecalBlendDesc))
 		{
 			switch (Mask)
 			{
@@ -451,7 +467,7 @@ namespace DecalRendering
 				>::GetRHI();
 			}
 		}
-		else if (DecalBlendDesc.BlendMode == BLEND_AlphaComposite)
+		else if (IsAlphaCompositeBlendMode(DecalBlendDesc))
 		{
 			ensure((Mask & 2) == 0); // AlphaComposite shouldn't write normal.
 
@@ -508,7 +524,7 @@ namespace DecalRendering
 				>::GetRHI();
 			}
 		}
-		else if (DecalBlendDesc.BlendMode == BLEND_Modulate)
+		else if (IsModulateBlendMode(DecalBlendDesc))
 		{
 			switch (Mask)
 			{
@@ -628,7 +644,7 @@ namespace DecalRendering
 	{
 		const uint32 Mask = GetRenderTargetWriteMask(DecalBlendDesc, DecalRenderStage, EDecalRenderTargetMode::SceneColorAndGBufferNoNormal);
 
-		if (DecalBlendDesc.BlendMode == BLEND_Translucent)
+		if (IsTranslucentOnlyBlendMode(DecalBlendDesc))
 		{
 			switch (Mask)
 			{
@@ -676,7 +692,7 @@ namespace DecalRendering
 				>::GetRHI();
 			}
 		}
-		else if (DecalBlendDesc.BlendMode == BLEND_AlphaComposite)
+		else if (IsAlphaCompositeBlendMode(DecalBlendDesc))
 		{
 			switch (Mask)
 			{
@@ -724,7 +740,7 @@ namespace DecalRendering
 				>::GetRHI();
 			}
 		}
-		else if (DecalBlendDesc.BlendMode == BLEND_Modulate)
+		else if (IsModulateBlendMode(DecalBlendDesc))
 		{
 			switch (Mask)
 			{
@@ -781,7 +797,7 @@ namespace DecalRendering
 	{
 		if (DecalRenderStage == EDecalRenderStage::Mobile)
 		{
-			if (DecalBlendDesc.BlendMode == BLEND_Translucent)
+			if (IsTranslucentOnlyBlendMode(DecalBlendDesc))
 			{
 				if (DecalBlendDesc.bWriteEmissive)
 				{
@@ -794,11 +810,11 @@ namespace DecalRendering
 					return TStaticBlendState<CW_RGB, BO_Add, BF_SourceAlpha, BF_InverseSourceAlpha, BO_Add, BF_Zero, BF_One, CW_NONE>::GetRHI();
 				}
 			}
-			else if (DecalBlendDesc.BlendMode == BLEND_AlphaComposite)
+			else if (IsAlphaCompositeBlendMode(DecalBlendDesc))
 			{
 				return TStaticBlendState<CW_RGB, BO_Add, BF_One, BF_InverseSourceAlpha, BO_Add, BF_Zero, BF_One, CW_NONE>::GetRHI();
 			}
-			else if (DecalBlendDesc.BlendMode == BLEND_Modulate)
+			else if (IsModulateBlendMode(DecalBlendDesc))
 			{
 				return TStaticBlendState<CW_RGB, BO_Add, BF_DestColor, BF_InverseSourceAlpha, BO_Add, BF_Zero, BF_One, CW_NONE>::GetRHI();
 			}
diff --git a/Engine/Source/Runtime/Renderer/Private/DecalRenderingCommon.h b/Engine/Source/Runtime/Renderer/Private/DecalRenderingCommon.h
--- a/Engine/Source/Runtime/Renderer/Private/DecalRenderingCommon.h
+++ b/Engine/Source/Runtime/Renderer/Private/DecalRenderingCommon.h
@@ -18,6 +18,7 @@ union FDecalBlendDesc
 	struct
 	{
 		uint32 BlendMode : 8;
+		uint32 StrataBlendMode : 8;
 		uint32 RenderStageMask : 8;
 		uint32 bWriteBaseColor : 1;
 		uint32 bWriteNormal : 1;
