Fix merge error for Change 23505760

diff --git a/Engine/Source/Runtime/Renderer/Private/PrimitiveSceneInfo.cpp b/Engine/Source/Runtime/Renderer/Private/PrimitiveSceneInfo.cpp
--- a/Engine/Source/Runtime/Renderer/Private/PrimitiveSceneInfo.cpp
+++ b/Engine/Source/Runtime/Renderer/Private/PrimitiveSceneInfo.cpp
@@ -876,8 +876,8 @@ void CacheRayTracingPrimitive(
 	}
 
 	// Write flags
-	Flags = SceneInfo->Proxy->GetCachedRayTracingInstance(CachedRayTracingInstance);
-	UpdateRayTracingInstanceMaskAndFlagsIfNeeded(CachedRayTracingInstance, *(SceneInfo->Proxy));
+	OutFlags = SceneInfo->Proxy->GetCachedRayTracingInstance(OutCachedRayTracingInstance);
+	UpdateRayTracingInstanceMaskAndFlagsIfNeeded(OutCachedRayTracingInstance, *(SceneInfo->Proxy));
 
 	// Cache the coarse mesh streaming handle
 	SceneInfo->CoarseMeshStreamingHandle = SceneInfo->Proxy->GetCoarseMeshStreamingHandle();
