Fix Strata blend mode for mobile.

diff --git a/Engine/Shaders/Private/MobileBasePassPixelShader.usf b/Engine/Shaders/Private/MobileBasePassPixelShader.usf
--- a/Engine/Shaders/Private/MobileBasePassPixelShader.usf
+++ b/Engine/Shaders/Private/MobileBasePassPixelShader.usf
@@ -680,11 +680,13 @@ void Main(
 
 	// Compute final value
 	{
-	#if STRATA_OPTIMIZED_UNLIT
+	#if STRATA_OPAQUE_MATERIAL && STRATA_OPTIMIZED_UNLIT
 		OutColor = float4(DualBlendSurfaceLuminancePostCoverage, 1);
 	#elif STRATA_OPAQUE_MATERIAL
 		OutColor = float4(Color, 1);
-	#elif STRATA_TRANSLUCENT_MATERIAL
+	#endif
+
+	#if STRATA_TRANSLUCENT_MATERIAL
 			
 	#if STRATA_LEGACY_PREMULT_ALPHA_OVERRIDE && !STRATA_OPTIMIZED_UNLIT
 		// This data patching is working because DualBlendSurfaceCoverage is always 1 when converting material from legacy to Strata with premultiplied alpha blending.
diff --git a/Engine/Source/Runtime/Renderer/Private/MobileBasePass.cpp b/Engine/Source/Runtime/Renderer/Private/MobileBasePass.cpp
--- a/Engine/Source/Runtime/Renderer/Private/MobileBasePass.cpp
+++ b/Engine/Source/Runtime/Renderer/Private/MobileBasePass.cpp
@@ -216,8 +216,8 @@ ELightMapPolicyType MobileBasePass::SelectMeshLightmapPolicy(
 	FMaterialShadingModelField ShadingModels, 
 	bool bPrimReceivesCSM,
 	bool bUsesDeferredShading,
-	ERHIFeatureLevel::Type FeatureLevel,
-	EBlendMode BlendMode) // STRATA_TODO_BLENDMODE
+	bool bIsTranslucent,
+	ERHIFeatureLevel::Type FeatureLevel)
 {
 	// Unlit uses NoLightmapPolicy with 0 point lights
 	ELightMapPolicyType SelectedLightmapPolicy = LMP_NO_LIGHTMAP;
@@ -229,8 +229,7 @@ ELightMapPolicyType MobileBasePass::SelectMeshLightmapPolicy(
 
 		if (!ReadOnlyCVARCache.bAllowStaticLighting || (ReadOnlyCVARCache.bMobileEnableNoPrecomputedLightingCSMShader && Scene && Scene->GetForceNoPrecomputedLighting()))
 		{
-			if (!IsTranslucentBlendMode(BlendMode) &&
-				!ShadingModels.HasShadingModel(MSM_SingleLayerWater))
+			if (!bIsTranslucent && !ShadingModels.HasShadingModel(MSM_SingleLayerWater))
 			{
 				// Whether to use a single CSM permutation with a branch in the shader
 				bPrimReceivesCSM |= MobileUseCSMShaderBranch();
@@ -393,11 +392,39 @@ void MobileBasePass::SetTranslucentRenderState(FMeshPassProcessorRenderState& Dr
 	EShaderPlatform ShaderPlatform = GetFeatureLevelShaderPlatform(Material.GetFeatureLevel());
 
 	const bool bIsUsingMobilePixelProjectedReflection = Material.IsUsingPlanarForwardReflections() && IsUsingMobilePixelProjectedReflection(ShaderPlatform);
+	const bool bIsDualSourceBlending = RHISupportsDualSourceBlending(ShaderPlatform);
+	const EShaderPlatform Platform = GetFeatureLevelShaderPlatform(Material.GetFeatureLevel());
 
-	// STRATA_TODO_BLENDMODE
-	if (ShadingModels.HasShadingModel(MSM_ThinTranslucent))
+	if (Strata::IsStrataEnabled())
+	{
+		if (Material.IsDualBlendingEnabled(Platform))
+		{
+			// If requested and available, we do standard dual blending, and the alpha gets ignored
+			// Blend by putting add in target 0 and multiply by background in target 1.
+			DrawRenderState.SetBlendState(TStaticBlendState<CW_RGBA, BO_Add, BF_One, BF_Source1Color, BO_Add, BF_One, BF_Source1Alpha>::GetRHI());
+		}
+		else
+		{
+			if (Material.GetStrataBlendMode() == SBM_ColoredTransmittanceOnly)
+			{
+				// Modulate with the existing scene color, preserve destination alpha.
+				DrawRenderState.SetBlendState(TStaticBlendState<CW_RGB, BO_Add, BF_DestColor, BF_Zero>::GetRHI());
+			}
+			else if (Material.GetStrataBlendMode() == SBM_AlphaHoldout)
+			{
+				// Blend by holding out the matte shape of the source alpha
+				DrawRenderState.SetBlendState(TStaticBlendState<CW_RGBA, BO_Add, BF_Zero, BF_InverseSourceAlpha, BO_Add, BF_One, BF_InverseSourceAlpha>::GetRHI());
+			}
+			else
+			{
+				// We always use premultiplied alpha for translucent rendering.
+				// If a material was requesting dual source blending, the shader will use static platform knowledge to convert colored transmittance to a grey scale transmittance.
+				DrawRenderState.SetBlendState(TStaticBlendState<CW_RGBA, BO_Add, BF_One, BF_InverseSourceAlpha, BO_Add, BF_Zero, BF_InverseSourceAlpha>::GetRHI());
+			}
+		}
+	}
+	else if (ShadingModels.HasShadingModel(MSM_ThinTranslucent))
 	{
-		const bool bIsDualSourceBlending = RHISupportsDualSourceBlending(ShaderPlatform);
 		if (bIsDualSourceBlending)
 		{
 			// Blend by putting add in target 0 and multiply by background in target 1.
@@ -521,10 +548,10 @@ void MobileBasePass::SetTranslucentRenderState(FMeshPassProcessorRenderState& Dr
 	}
 }
 
-static FMeshDrawCommandSortKey GetBasePassStaticSortKey(EBlendMode BlendMode, bool bBackground)
+static FMeshDrawCommandSortKey GetBasePassStaticSortKey(const bool bIsMasked, bool bBackground)
 {
 	FMeshDrawCommandSortKey SortKey;
-	SortKey.PackedData = (BlendMode == EBlendMode::BLEND_Masked ? 1 : 0);
+	SortKey.PackedData = (bIsMasked ? 1 : 0);
 	SortKey.PackedData|= (bBackground ? 2 : 0); // background flag in second bit
 	return SortKey;
 }
@@ -604,10 +631,9 @@ FMobileBasePassMeshProcessor::FMobileBasePassMeshProcessor(
 
 bool FMobileBasePassMeshProcessor::TryAddMeshBatch(const FMeshBatch& RESTRICT MeshBatch, uint64 BatchElementMask, const FPrimitiveSceneProxy* RESTRICT PrimitiveSceneProxy, int32 StaticMeshId, const FMaterialRenderProxy& MaterialRenderProxy, const FMaterial& Material)
 {
-	const EBlendMode BlendMode = Material.GetBlendMode();
-	const EStrataBlendMode StrataBlendMode = Material.GetStrataBlendMode();
 	const FMaterialShadingModelField ShadingModels = Material.GetShadingModels();
-	const bool bIsTranslucent = IsTranslucentBlendMode(BlendMode, StrataBlendMode);
+	const bool bIsMasked = IsMaskedBlendMode(Material);
+	const bool bIsTranslucent = IsTranslucentBlendMode(Material);
 	const bool bUsesWaterMaterial = ShadingModels.HasShadingModel(MSM_SingleLayerWater); // Water goes into the translucent pass
 	const bool bCanReceiveCSM = ((Flags & EFlags::CanReceiveCSM) == EFlags::CanReceiveCSM);
 
@@ -625,8 +651,8 @@ bool FMobileBasePassMeshProcessor::TryAddMeshBatch(const FMeshBatch& RESTRICT Me
 			check(bCanReceiveCSM == false);
 			const FLightSceneInfo* MobileDirectionalLight = MobileBasePass::GetDirectionalLightInfo(Scene, PrimitiveSceneProxy);
 			// Opaque meshes used for mobile pixel projected reflection could receive CSM in translucent pass.
-			ELightMapPolicyType LightmapPolicyType = MobileBasePass::SelectMeshLightmapPolicy(Scene, MeshBatch, PrimitiveSceneProxy, MobileDirectionalLight, ShadingModels, bCanReceiveCSM, false, FeatureLevel, BlendMode);
-			bResult = Process(MeshBatch, BatchElementMask, StaticMeshId, PrimitiveSceneProxy, MaterialRenderProxy, Material, BlendMode, ShadingModels, LightmapPolicyType, bCanReceiveCSM, MeshBatch.LCI);
+			ELightMapPolicyType LightmapPolicyType = MobileBasePass::SelectMeshLightmapPolicy(Scene, MeshBatch, PrimitiveSceneProxy, MobileDirectionalLight, ShadingModels, bCanReceiveCSM, false, bIsTranslucent, FeatureLevel);
+			bResult = Process(MeshBatch, BatchElementMask, StaticMeshId, PrimitiveSceneProxy, MaterialRenderProxy, Material, bIsMasked, bIsTranslucent, ShadingModels, LightmapPolicyType, bCanReceiveCSM, MeshBatch.LCI);
 		}
 	}
 	else
@@ -635,8 +661,8 @@ bool FMobileBasePassMeshProcessor::TryAddMeshBatch(const FMeshBatch& RESTRICT Me
 		if (!bIsTranslucent && !bUsesWaterMaterial)
 		{
 			const FLightSceneInfo* MobileDirectionalLight = MobileBasePass::GetDirectionalLightInfo(Scene, PrimitiveSceneProxy);
-			ELightMapPolicyType LightmapPolicyType = MobileBasePass::SelectMeshLightmapPolicy(Scene, MeshBatch, PrimitiveSceneProxy, MobileDirectionalLight, ShadingModels, bCanReceiveCSM, bUsesDeferredShading, FeatureLevel, BlendMode);
-			bResult = Process(MeshBatch, BatchElementMask, StaticMeshId, PrimitiveSceneProxy, MaterialRenderProxy, Material, BlendMode, ShadingModels, LightmapPolicyType, bCanReceiveCSM, MeshBatch.LCI);
+			ELightMapPolicyType LightmapPolicyType = MobileBasePass::SelectMeshLightmapPolicy(Scene, MeshBatch, PrimitiveSceneProxy, MobileDirectionalLight, ShadingModels, bCanReceiveCSM, bUsesDeferredShading, bIsTranslucent, FeatureLevel);
+			bResult = Process(MeshBatch, BatchElementMask, StaticMeshId, PrimitiveSceneProxy, MaterialRenderProxy, Material, bIsMasked, bIsTranslucent, ShadingModels, LightmapPolicyType, bCanReceiveCSM, MeshBatch.LCI);
 		}
 	}
 
@@ -675,7 +701,8 @@ bool FMobileBasePassMeshProcessor::Process(
 		const FPrimitiveSceneProxy* RESTRICT PrimitiveSceneProxy,
 		const FMaterialRenderProxy& RESTRICT MaterialRenderProxy,
 		const FMaterial& RESTRICT MaterialResource,
-		EBlendMode BlendMode,
+		const bool bIsMasked,
+		const bool bIsTranslucent,
 		FMaterialShadingModelField ShadingModels,
 		const ELightMapPolicyType LightMapPolicyType,
 		const bool bCanReceiveCSM,
@@ -745,7 +772,6 @@ bool FMobileBasePassMeshProcessor::Process(
 		SortKey = CalculateTranslucentMeshStaticSortKey(PrimitiveSceneProxy, MeshBatch.MeshIdInPrimitive);
 		// We always want water to be rendered first on mobile in order to mimic other renderers where it is opaque. We shift the other priorities by 1.
 		// And we also want to render the meshes used for mobile pixel projected reflection first if it is opaque.
-		const bool bIsTranslucent = IsTranslucentBlendMode(BlendMode); // STRATA_TODO_BLENDMODE
 		SortKey.Translucent.Priority = ShadingModels.HasShadingModel(MSM_SingleLayerWater) || (!bIsTranslucent && bIsUsingMobilePixelProjectedReflection) ? uint16(0) : uint16(FMath::Clamp(uint32(SortKey.Translucent.Priority) + 1, 0u, uint32(USHRT_MAX)));
 	}
 	else
@@ -754,7 +780,7 @@ bool FMobileBasePassMeshProcessor::Process(
 		bool bBackground = PrimitiveSceneProxy ? PrimitiveSceneProxy->TreatAsBackgroundForOcclusion() : false;
 		// Default static sort key separates masked and non-masked geometry, generic mesh sorting will also sort by PSO
 		// if platform wants front to back sorting, this key will be recomputed in InitViews
-		SortKey = GetBasePassStaticSortKey(BlendMode, bBackground); // STRATA_TODO_BLENDMODE
+		SortKey = GetBasePassStaticSortKey(bIsMasked, bBackground);
 	}
 	
 	const FMeshDrawingPolicyOverrideSettings OverrideSettings = ComputeMeshOverrideSettings(MeshBatch);
diff --git a/Engine/Source/Runtime/Renderer/Private/MobileBasePassRendering.h b/Engine/Source/Runtime/Renderer/Private/MobileBasePassRendering.h
--- a/Engine/Source/Runtime/Renderer/Private/MobileBasePassRendering.h
+++ b/Engine/Source/Runtime/Renderer/Private/MobileBasePassRendering.h
@@ -299,8 +299,8 @@ namespace MobileBasePass
 		FMaterialShadingModelField ShadingModels, 
 		bool bPrimReceivesCSM, 
 		bool bUsedDeferredShading,
-		ERHIFeatureLevel::Type FeatureLevel,
-		EBlendMode BlendMode);
+		bool bIsTranslucent,
+		ERHIFeatureLevel::Type FeatureLevel);
 
 	bool GetShaders(
 		ELightMapPolicyType LightMapPolicyType,
@@ -456,7 +456,8 @@ private:
 		const FPrimitiveSceneProxy* RESTRICT PrimitiveSceneProxy,
 		const FMaterialRenderProxy& RESTRICT MaterialRenderProxy,
 		const FMaterial& RESTRICT MaterialResource,
-		EBlendMode BlendMode,
+		const bool bIsMasked,
+		const bool bIsTranslucent,
 		FMaterialShadingModelField ShadingModels,
 		const ELightMapPolicyType LightMapPolicyType,
 		const bool bCanReceiveCSM,
