Change hair voxel to have 0.5 pixel size by default.

diff --git a/Engine/Source/Runtime/Renderer/Private/HairStrands/HairStrandsCluster.cpp b/Engine/Source/Runtime/Renderer/Private/HairStrands/HairStrandsCluster.cpp
--- a/Engine/Source/Runtime/Renderer/Private/HairStrands/HairStrandsCluster.cpp
+++ b/Engine/Source/Runtime/Renderer/Private/HairStrands/HairStrandsCluster.cpp
@@ -13,7 +13,7 @@
 #include "ShaderParameterStruct.h"
 #include "ScenePrivate.h"
 
-static int32 GHairVirtualVoxel_NumPixelPerVoxel = 2;
+static int32 GHairVirtualVoxel_NumPixelPerVoxel = 0.5f;
 static FAutoConsoleVariableRef CVarHairVirtualVoxel_NumPixelPerVoxel(TEXT("r.HairStrands.Voxelization.VoxelSizeInPixel"), GHairVirtualVoxel_NumPixelPerVoxel, TEXT("Target size of voxel size in pixels"), ECVF_RenderThreadSafe);
 
 class FHairMacroGroupAABBCS : public FGlobalShader
