Fix CIS issue.

diff --git a/Engine/Source/Runtime/Renderer/Private/HairStrands/HairStrandsVisibility.cpp b/Engine/Source/Runtime/Renderer/Private/HairStrands/HairStrandsVisibility.cpp
--- a/Engine/Source/Runtime/Renderer/Private/HairStrands/HairStrandsVisibility.cpp
+++ b/Engine/Source/Runtime/Renderer/Private/HairStrands/HairStrandsVisibility.cpp
@@ -4210,9 +4210,7 @@ void RenderHairStrandsVisibilityBuffer(
 				RenderMode == HairVisibilityRenderMode_ComputeRaster ||
 				RenderMode == HairVisibilityRenderMode_ComputeRasterForward);
 
-			const bool bRunColorAndDepthPatching = 
-				(SceneMaterial0 && SceneMaterial1) &&
-				SceneColorTexture && (RenderMode != HairVisibilityRenderMode_ComputeRaster || RenderMode != HairVisibilityRenderMode_ComputeRasterForward);
+			const bool bRunColorAndDepthPatching = SceneMaterial0 != nullptr && SceneMaterial1 != nullptr && SceneColorTexture != nullptr && (RenderMode != HairVisibilityRenderMode_ComputeRaster || RenderMode != HairVisibilityRenderMode_ComputeRasterForward);
 
 			FRDGTextureRef HairOnlyDepthTexture = GraphBuilder.CreateTexture(SceneDepthTexture->Desc, TEXT("Hair.HairOnlyDepthTexture"));
 			FRDGTextureRef CoverageTexture = nullptr;
