Fix compile error caused by merge.

diff --git a/Engine/Source/Runtime/Renderer/Private/CustomDepthRendering.cpp b/Engine/Source/Runtime/Renderer/Private/CustomDepthRendering.cpp
--- a/Engine/Source/Runtime/Renderer/Private/CustomDepthRendering.cpp
+++ b/Engine/Source/Runtime/Renderer/Private/CustomDepthRendering.cpp
@@ -285,6 +285,7 @@ bool FSceneRenderer::RenderCustomDepthPass(
 		Nanite::FRasterContext RasterContext = Nanite::InitRasterContext(
 			GraphBuilder,
 			SharedContext,
+			ViewFamily,
 			RasterTextureSize,
 			RasterTextureRect,
 			false, // bVisualize
