PR #9908: Volumetric Fog is now more stable under Dynamic Resolution (Contributed by [REDACTED_PERSON])

Enginetest were looking undisrupted by this change.

diff --git a/Engine/Shaders/Private/VolumetricFog.usf b/Engine/Shaders/Private/VolumetricFog.usf
--- a/Engine/Shaders/Private/VolumetricFog.usf
+++ b/Engine/Shaders/Private/VolumetricFog.usf
@@ -171,7 +171,8 @@ struct FWriteToSliceVertexOutput
 /** Z index of the minimum slice in the range. */
 int MinZ;
 float4 ViewSpaceBoundingSphere;
-float4x4 ViewToVolumeClip; 
+float4x4 ViewToVolumeClip;
+float2 ClipRatio = float2(1.0f, 1.0f);
 
 /** Vertex shader that writes to a range of slices of a volume texture. */
 void WriteToBoundingSphereVS(
@@ -191,6 +192,9 @@ void WriteToBoundingSphereVS(
 		// Place the quad vertex to tightly bound the circle
 		float3 ViewSpaceVertexPosition = float3(ViewSpaceBoundingSphere.xy + (InUV * 2 - 1) * SliceRadius, SliceDepth);
 		Output.Vertex.Position = mul(float4(ViewSpaceVertexPosition, 1), ViewToVolumeClip);
+
+		// Scale rendered position to account for frustum difference between fog voxelization and rendered scene
+		Output.Vertex.Position.xy *= ClipRatio;
 	}
 	else
 	{
diff --git a/Engine/Shaders/Private/VolumetricFogVoxelization.usf b/Engine/Shaders/Private/VolumetricFogVoxelization.usf
--- a/Engine/Shaders/Private/VolumetricFogVoxelization.usf
+++ b/Engine/Shaders/Private/VolumetricFogVoxelization.usf
@@ -314,8 +314,11 @@ void VoxelizePS(
 	UpdateMaterialCloudParam(MaterialParameters, AbsoluteWorldPosition, ResolvedView, CloudLayerParams, TRACING_SHADOW_DISTANCE_OFF, SPACE_SKIPPING_SLICE_DEPTH_OFF);
 #endif
 
+	// Scale rendered position to account for frustum difference between fog voxelization and rendered scene
+	float4 AdjustedSvPosition = SvPosition * float4(VoxelizeVolumePass.ClipRatio, 1, 1);
+
 	FPixelMaterialInputs PixelMaterialInputs;
-	CalcMaterialParameters(MaterialParameters, PixelMaterialInputs, SvPosition, true);
+	CalcMaterialParameters(MaterialParameters, PixelMaterialInputs, AdjustedSvPosition, true);
 
 	float3 EmissiveColor = SampleEmissive(PixelMaterialInputs);
 	float3 Albedo = SampleAlbedo(PixelMaterialInputs, ResolvedView);
diff --git a/Engine/Source/Runtime/Renderer/Private/SceneRendering.cpp b/Engine/Source/Runtime/Renderer/Private/SceneRendering.cpp
--- a/Engine/Source/Runtime/Renderer/Private/SceneRendering.cpp
+++ b/Engine/Source/Runtime/Renderer/Private/SceneRendering.cpp
@@ -1456,7 +1456,6 @@ void FViewInfo::SetupUniformBufferParameters(
 	if (ViewState)
 	{
 		// Compute LightScatteringScreenUVToVolumeUV, for the current frame resolution and volume texture resolutoni according to grid size.
-		FIntVector GetVolumetricFogGridSize(FIntPoint ViewRectSize, int32& OutVolumetricFogGridPixelSize);
 		int32 VolumetricFogGridPixelSize;
 		const FIntVector VolumetricFogGridSize = GetVolumetricFogGridSize(ViewRect.Size(), VolumetricFogGridPixelSize);
 		FVector2f LightScatteringScreenUVToVolumeUV =  FVector2f(ViewRect.Size()) / (FVector2f(VolumetricFogGridSize.X, VolumetricFogGridSize.Y) * VolumetricFogGridPixelSize);
diff --git a/Engine/Source/Runtime/Renderer/Private/VolumetricFog.cpp b/Engine/Source/Runtime/Renderer/Private/VolumetricFog.cpp
--- a/Engine/Source/Runtime/Renderer/Private/VolumetricFog.cpp
+++ b/Engine/Source/Runtime/Renderer/Private/VolumetricFog.cpp
@@ -249,6 +249,7 @@ class FWriteToBoundingSphereVS : public FGlobalShader
 	BEGIN_SHADER_PARAMETER_STRUCT(FParameters, )
 		SHADER_PARAMETER_STRUCT_INCLUDE(FVolumetricFogIntegrationParameters, VolumetricFogParameters)
 		SHADER_PARAMETER(FMatrix44f, ViewToVolumeClip)
+		SHADER_PARAMETER(FVector2f, ClipRatio)
 		SHADER_PARAMETER(FVector4f, ViewSpaceBoundingSphere)
 		SHADER_PARAMETER(int32, MinZ)
 	END_SHADER_PARAMETER_STRUCT()
@@ -608,6 +609,15 @@ void FDeferredShadingSceneRenderer::RenderLocalLightsForVolumetricFog(
 						VSPassParameters.MinZ = VolumeZBounds.X;
 						VSPassParameters.ViewSpaceBoundingSphere = FVector4f(FVector4f(View.ViewMatrices.GetViewMatrix().TransformPosition(LightBounds.Center)), LightBounds.W); // LWC_TODO: precision loss
 						VSPassParameters.ViewToVolumeClip = FMatrix44f(View.ViewMatrices.ComputeProjectionNoAAMatrix());	// LWC_TODO: Precision loss?
+
+						// Calculate how much the Fog froxel volume "overhangs" the actual view frustum to the right and bottom
+						const FIntPoint FogRect = View.ViewRect.Size();
+						int32 VolumetricFogGridPixelSize;
+						const FIntVector VolumetricFogGridSize = GetVolumetricFogGridSize(FogRect, VolumetricFogGridPixelSize);
+						const FVector2f FogPhysicalSize = (FVector2f(VolumetricFogGridSize.X, VolumetricFogGridSize.Y) * VolumetricFogGridPixelSize);
+						const FVector2f ClipRatio = FogPhysicalSize / FVector2f(FogRect);
+						VSPassParameters.ClipRatio = ClipRatio;
+
 						VSPassParameters.VolumetricFogParameters = PassParameters->VolumetricFogParameters;
 						SetShaderParameters(RHICmdList, VertexShader, VertexShader.GetVertexShader(), VSPassParameters);
 
diff --git a/Engine/Source/Runtime/Renderer/Private/VolumetricFog.h b/Engine/Source/Runtime/Renderer/Private/VolumetricFog.h
--- a/Engine/Source/Runtime/Renderer/Private/VolumetricFog.h
+++ b/Engine/Source/Runtime/Renderer/Private/VolumetricFog.h
@@ -47,6 +47,8 @@ private:
 	FRDGTextureRef DefaultLightFunctionAtlasItemTexture;
 };
 
+extern FIntVector GetVolumetricFogGridSize(FIntPoint ViewRectSize, int32& OutVolumetricFogGridPixelSize);
+
 extern bool DoesPlatformSupportVolumetricFog(const FStaticShaderPlatform Platform);
 extern bool DoesPlatformSupportVolumetricFogVoxelization(const FStaticShaderPlatform Platform);
 
diff --git a/Engine/Source/Runtime/Renderer/Private/VolumetricFogShared.h b/Engine/Source/Runtime/Renderer/Private/VolumetricFogShared.h
--- a/Engine/Source/Runtime/Renderer/Private/VolumetricFogShared.h
+++ b/Engine/Source/Runtime/Renderer/Private/VolumetricFogShared.h
@@ -9,6 +9,7 @@
 BEGIN_GLOBAL_SHADER_PARAMETER_STRUCT(FVoxelizeVolumePassUniformParameters, )
 	SHADER_PARAMETER_STRUCT(FSceneTextureUniformParameters, SceneTextures)
 	SHADER_PARAMETER(FMatrix44f, ViewToVolumeClip)
+	SHADER_PARAMETER(FVector2f, ClipRatio)
 	SHADER_PARAMETER(FVector4f, FrameJitterOffset0)
 	SHADER_PARAMETER_STRUCT(FVolumetricFogGlobalData, VolumetricFog)
 	SHADER_PARAMETER(FVector3f, RenderVolumetricCloudParametersCloudLayerCenterKm)
diff --git a/Engine/Source/Runtime/Renderer/Private/VolumetricFogVoxelization.cpp b/Engine/Source/Runtime/Renderer/Private/VolumetricFogVoxelization.cpp
--- a/Engine/Source/Runtime/Renderer/Private/VolumetricFogVoxelization.cpp
+++ b/Engine/Source/Runtime/Renderer/Private/VolumetricFogVoxelization.cpp
@@ -50,6 +50,14 @@ TRDGUniformBufferRef<FVoxelizeVolumePassUniformParameters> CreateVoxelizeVolumeP
 	auto* Parameters = GraphBuilder.AllocParameters<FVoxelizeVolumePassUniformParameters>();
 	SetupSceneTextureUniformParameters(GraphBuilder, &View.GetSceneTextures(), View.FeatureLevel, ESceneTextureSetupMode::None, Parameters->SceneTextures);
 
+	// Account for "overhang" of the fog volume off the bottom right of the render
+	const FVector2f FogRenderRect = FVector2f(View.ViewRect.Size());
+	int32 VolumetricFogGridPixelSize;
+	const FIntVector VolumetricFogGridSize = GetVolumetricFogGridSize(View.ViewRect.Size(), VolumetricFogGridPixelSize);
+	const FVector2f FogPhysicalRect = (FVector2f(VolumetricFogGridSize.X, VolumetricFogGridSize.Y) * VolumetricFogGridPixelSize);
+	const FVector2f ClipRatio = FogPhysicalRect / FogRenderRect;
+	Parameters->ClipRatio = ClipRatio;
+
 	Parameters->ViewToVolumeClip = FMatrix44f(View.ViewMatrices.ComputeProjectionNoAAMatrix());		// LWC_TODO: Precision loss?
 	Parameters->ViewToVolumeClip.M[2][0] += Jitter.X;
 	Parameters->ViewToVolumeClip.M[2][1] += Jitter.Y;
