Followup to 23537214 to make functions common.

diff --git a/Engine/Source/Runtime/Renderer/Private/VolumetricFog.cpp b/Engine/Source/Runtime/Renderer/Private/VolumetricFog.cpp
--- a/Engine/Source/Runtime/Renderer/Private/VolumetricFog.cpp
+++ b/Engine/Source/Runtime/Renderer/Private/VolumetricFog.cpp
@@ -610,13 +610,7 @@ void FDeferredShadingSceneRenderer::RenderLocalLightsForVolumetricFog(
 						VSPassParameters.ViewSpaceBoundingSphere = FVector4f(FVector4f(View.ViewMatrices.GetViewMatrix().TransformPosition(LightBounds.Center)), LightBounds.W); // LWC_TODO: precision loss
 						VSPassParameters.ViewToVolumeClip = FMatrix44f(View.ViewMatrices.ComputeProjectionNoAAMatrix());	// LWC_TODO: Precision loss?
 
-						// Calculate how much the Fog froxel volume "overhangs" the actual view frustum to the right and bottom
-						const FIntPoint FogRect = View.ViewRect.Size();
-						int32 VolumetricFogGridPixelSize;
-						const FIntVector VolumetricFogGridSize = GetVolumetricFogGridSize(FogRect, VolumetricFogGridPixelSize);
-						const FVector2f FogPhysicalSize = (FVector2f(VolumetricFogGridSize.X, VolumetricFogGridSize.Y) * VolumetricFogGridPixelSize);
-						const FVector2f ClipRatio = FogPhysicalSize / FVector2f(FogRect);
-						VSPassParameters.ClipRatio = ClipRatio;
+						VSPassParameters.ClipRatio = GetVolumetricFogFroxelToScreenSVPosRatio(View.ViewRect.Size());
 
 						VSPassParameters.VolumetricFogParameters = PassParameters->VolumetricFogParameters;
 						SetShaderParameters(RHICmdList, VertexShader, VertexShader.GetVertexShader(), VSPassParameters);
@@ -857,6 +851,17 @@ FIntVector GetVolumetricFogGridSize(FIntPoint ViewRectSize, int32& OutVolumetric
 	return FIntVector(VolumetricFogGridSizeXY.X, VolumetricFogGridSizeXY.Y, GetVolumetricFogGridSizeZ());
 }
 
+FVector2f GetVolumetricFogFroxelToScreenSVPosRatio(FIntPoint ViewRectSize)
+{
+	// Calculate how much the Fog froxel volume "overhangs" the actual view frustum to the right and bottom.
+	// This needs to be applied on SVPos because froxel pixel size (see r.VolumetricFog.GridPixelSize) does not align perfectly with view rect.
+	int32 VolumetricFogGridPixelSize;
+	const FIntVector VolumetricFogGridSize = GetVolumetricFogGridSize(ViewRectSize, VolumetricFogGridPixelSize);
+	const FVector2f FogPhysicalSize = FVector2f(VolumetricFogGridSize.X, VolumetricFogGridSize.Y) * VolumetricFogGridPixelSize;
+	const FVector2f ClipRatio = FogPhysicalSize / FVector2f(ViewRectSize);
+	return ClipRatio;
+}
+
 void SetupVolumetricFogGlobalData(const FViewInfo& View, FVolumetricFogGlobalData& Parameters)
 {
 	const FScene* Scene = (FScene*)View.Family->Scene;
diff --git a/Engine/Source/Runtime/Renderer/Private/VolumetricFog.h b/Engine/Source/Runtime/Renderer/Private/VolumetricFog.h
--- a/Engine/Source/Runtime/Renderer/Private/VolumetricFog.h
+++ b/Engine/Source/Runtime/Renderer/Private/VolumetricFog.h
@@ -48,6 +48,7 @@ private:
 };
 
 extern FIntVector GetVolumetricFogGridSize(FIntPoint ViewRectSize, int32& OutVolumetricFogGridPixelSize);
+extern FVector2f GetVolumetricFogFroxelToScreenSVPosRatio(FIntPoint ViewRectSize);
 
 extern bool DoesPlatformSupportVolumetricFog(const FStaticShaderPlatform Platform);
 extern bool DoesPlatformSupportVolumetricFogVoxelization(const FStaticShaderPlatform Platform);
diff --git a/Engine/Source/Runtime/Renderer/Private/VolumetricFogVoxelization.cpp b/Engine/Source/Runtime/Renderer/Private/VolumetricFogVoxelization.cpp
--- a/Engine/Source/Runtime/Renderer/Private/VolumetricFogVoxelization.cpp
+++ b/Engine/Source/Runtime/Renderer/Private/VolumetricFogVoxelization.cpp
@@ -50,13 +50,7 @@ TRDGUniformBufferRef<FVoxelizeVolumePassUniformParameters> CreateVoxelizeVolumeP
 	auto* Parameters = GraphBuilder.AllocParameters<FVoxelizeVolumePassUniformParameters>();
 	SetupSceneTextureUniformParameters(GraphBuilder, &View.GetSceneTextures(), View.FeatureLevel, ESceneTextureSetupMode::None, Parameters->SceneTextures);
 
-	// Account for "overhang" of the fog volume off the bottom right of the render
-	const FVector2f FogRenderRect = FVector2f(View.ViewRect.Size());
-	int32 VolumetricFogGridPixelSize;
-	const FIntVector VolumetricFogGridSize = GetVolumetricFogGridSize(View.ViewRect.Size(), VolumetricFogGridPixelSize);
-	const FVector2f FogPhysicalRect = (FVector2f(VolumetricFogGridSize.X, VolumetricFogGridSize.Y) * VolumetricFogGridPixelSize);
-	const FVector2f ClipRatio = FogPhysicalRect / FogRenderRect;
-	Parameters->ClipRatio = ClipRatio;
+	Parameters->ClipRatio = GetVolumetricFogFroxelToScreenSVPosRatio(View.ViewRect.Size());
 
 	Parameters->ViewToVolumeClip = FMatrix44f(View.ViewMatrices.ComputeProjectionNoAAMatrix());		// LWC_TODO: Precision loss?
 	Parameters->ViewToVolumeClip.M[2][0] += Jitter.X;
