Add `Affects Dynamice Indirect Lighting` support to path tracing.

diff --git a/Engine/Source/Runtime/Renderer/Private/RayTracing/RayTracingInstance.cpp b/Engine/Source/Runtime/Renderer/Private/RayTracing/RayTracingInstance.cpp
--- a/Engine/Source/Runtime/Renderer/Private/RayTracing/RayTracingInstance.cpp
+++ b/Engine/Source/Runtime/Renderer/Private/RayTracing/RayTracingInstance.cpp
@@ -15,13 +15,14 @@
 
 void FRayTracingInstance::BuildInstanceMaskAndFlags(ERHIFeatureLevel::Type FeatureLevel)
 {
-	ERayTracingViewMaskMode MaskMode = ERayTracingViewMaskMode::RayTracing; // Deprecated function only supports RayTracing
-	
-	uint8 ExtraMask = bThinGeometry ? ComputeRayTracingInstanceMask(ERayTracingInstanceMaskType::HairStrands, MaskMode) : 0;
+	FSceneProxyRayTracingMaskInfo MaskInfo;
+	MaskInfo.MaskMode = ERayTracingViewMaskMode::RayTracing; // Deprecated function only supports RayTracing
+
+	uint8 ExtraMask = bThinGeometry ? ComputeRayTracingInstanceMask(ERayTracingInstanceMaskType::HairStrands, MaskInfo.MaskMode) : 0;
 
 	TArrayView<const FMeshBatch> MeshBatches = GetMaterials();
 	FRayTracingMaskAndFlags MaskAndFlags = BuildRayTracingInstanceMaskAndFlags(
-		MeshBatches, FeatureLevel, MaskMode, false, InstanceLayer, false, ExtraMask);
+		MeshBatches, FeatureLevel, MaskInfo, InstanceLayer, ExtraMask);
 
 	Mask = MaskAndFlags.Mask;
 	bForceOpaque = bForceOpaque || MaskAndFlags.bForceOpaque;
@@ -30,15 +31,16 @@ void FRayTracingInstance::BuildInstanceMaskAndFlags(ERHIFeatureLevel::Type Featu
 
 FRayTracingMaskAndFlags BuildRayTracingInstanceMaskAndFlags(TArrayView<const FMeshBatch> MeshBatches, ERHIFeatureLevel::Type FeatureLevel, ERayTracingInstanceLayer InstanceLayer, uint8 ExtraMask)
 {
+	FSceneProxyRayTracingMaskInfo MaskInfo;
+	MaskInfo.MaskMode = ERayTracingViewMaskMode::RayTracing; // Deprecated function only supports RayTracing
+
 	return BuildRayTracingInstanceMaskAndFlags(
 		MeshBatches,
 		FeatureLevel,
-		ERayTracingViewMaskMode::RayTracing,	/* Deprecated function only supports RayTracing	*/
-		false,									/*	bAffectsIndirectLightingOnly	*/
+		MaskInfo,
 		InstanceLayer,
-		false,									/*	bCastHiddenShadow	*/
 		ExtraMask);
-}
+}	
 
 uint8 ComputeBlendModeMask(const EBlendMode BlendMode)
 {
diff --git a/Engine/Source/Runtime/Renderer/Private/RayTracing/RayTracingInstanceMask.cpp b/Engine/Source/Runtime/Renderer/Private/RayTracing/RayTracingInstanceMask.cpp
--- a/Engine/Source/Runtime/Renderer/Private/RayTracing/RayTracingInstanceMask.cpp
+++ b/Engine/Source/Runtime/Renderer/Private/RayTracing/RayTracingInstanceMask.cpp
@@ -42,6 +42,8 @@ uint8 ComputeRayTracingInstanceMask(ERayTracingInstanceMaskType MaskType, ERayTr
 		case ERayTracingInstanceMaskType::VisibleInPrimaryRay:
 			Mask = 0;
 			break;
+		case ERayTracingInstanceMaskType::VisibleInIndirectRay:
+			Mask = 0;
 		default:
 			checkNoEntry();
 			break;
@@ -76,6 +78,9 @@ uint8 ComputeRayTracingInstanceMask(ERayTracingInstanceMaskType MaskType, ERayTr
 		case ERayTracingInstanceMaskType::VisibleInPrimaryRay:
 			Mask = PATHTRACER_MASK_CAMERA | PATHTRACER_MASK_HAIR_CAMERA;
 			break;
+		case ERayTracingInstanceMaskType::VisibleInIndirectRay:
+			Mask = PATHTRACER_MASK_INDIRECT | PATHTRACER_MASK_HAIR_INDIRECT;
+			break;
 		default:
 			checkNoEntry();
 			break;
@@ -102,6 +107,7 @@ FSceneProxyRayTracingMaskInfo GetSceneProxyRayTracingMaskInfo(const FPrimitiveSc
 
 	bool bAffectsIndirectLightingOnly = PrimitiveSceneProxy.AffectsIndirectLightingWhileHidden() && !PrimitiveSceneProxy.IsDrawnInGame();
 	bool bCastHiddenShadow = PrimitiveSceneProxy.CastsHiddenShadow() && !PrimitiveSceneProxy.IsDrawnInGame();
+	bool bAffectsDynamicIndirectLighting = PrimitiveSceneProxy.AffectsDynamicIndirectLighting();
 
 	ERayTracingViewMaskMode MaskMode = ERayTracingViewMaskMode::RayTracing;
 
@@ -122,13 +128,12 @@ FSceneProxyRayTracingMaskInfo GetSceneProxyRayTracingMaskInfo(const FPrimitiveSc
 		}
 	}
 
-	return {bAffectsIndirectLightingOnly, bCastHiddenShadow, MaskMode};
+	return {bAffectsIndirectLightingOnly, bCastHiddenShadow, bAffectsDynamicIndirectLighting, MaskMode};
 }
 
 
 FRayTracingMaskAndFlags BuildRayTracingInstanceMaskAndFlags(TArrayView<const FMeshBatch> MeshBatches, ERHIFeatureLevel::Type FeatureLevel,
-	ERayTracingViewMaskMode MaskMode, bool bAffectIndirectLightingOnly, ERayTracingInstanceLayer InstanceLayer, bool bCastHiddenShadow,
-	uint8 ExtraMask)
+	const FSceneProxyRayTracingMaskInfo& SceneProxyRayTracingMaskInfo, ERayTracingInstanceLayer InstanceLayer, uint8 ExtraMask)
 {
 	FRayTracingMaskAndFlags Result;
 
@@ -138,6 +143,7 @@ FRayTracingMaskAndFlags BuildRayTracingInstanceMaskAndFlags(TArrayView<const FMe
 	bool bAnySegmentsCastShadow = false;
 	bool bAllSegmentsCastShadow = true;
 	bool bDoubleSided = false;
+	ERayTracingViewMaskMode MaskMode = SceneProxyRayTracingMaskInfo.MaskMode;
 	Result.Mask = ExtraMask;
 
 	for (int32 SegmentIndex = 0; SegmentIndex < MeshBatches.Num(); SegmentIndex++)
@@ -181,14 +187,20 @@ FRayTracingMaskAndFlags BuildRayTracingInstanceMaskAndFlags(TArrayView<const FMe
 		}
 	}
 
-	if (bAffectIndirectLightingOnly)
+	if (!SceneProxyRayTracingMaskInfo.bAffectsDynamicIndirectLighting)
+	{
+		// If the object does not affect indirect lighting, remove all indirect bits.
+		Result.Mask &= ~ComputeRayTracingInstanceMask(ERayTracingInstanceMaskType::VisibleInIndirectRay, MaskMode);
+	}
+
+	if (SceneProxyRayTracingMaskInfo.bAffectsIndirectLightingOnly)
 	{
 		Result.Mask &= ~ComputeRayTracingInstanceMask(ERayTracingInstanceMaskType::VisibleInPrimaryRay, MaskMode);
 	}
 
-	if (bCastHiddenShadow && bAnySegmentsCastShadow)
+	if (SceneProxyRayTracingMaskInfo.bCastHiddenShadow && bAnySegmentsCastShadow)
 	{
-		if (!bAffectIndirectLightingOnly)
+		if (!SceneProxyRayTracingMaskInfo.bAffectsIndirectLightingOnly)
 		{
 			// objects should not be in any visible group if any segments cast shadow and the caster is hidden
 			// and not affecting indirect.
@@ -219,7 +231,7 @@ FRayTracingMaskAndFlags BuildRayTracingInstanceMaskAndFlags(const FRayTracingIns
 	FRayTracingMaskAndFlags MaskAndFlags =
 		BuildRayTracingInstanceMaskAndFlags(
 			MeshBatches, PrimitiveSceneProxy.GetScene().GetFeatureLevel(), 
-			MaskInfo.MaskMode, MaskInfo.bAffectsIndirectLightingOnly, Instance.InstanceLayer, MaskInfo.bCastHiddenShadow, ExtraMask);
+			MaskInfo, Instance.InstanceLayer, ExtraMask);
 
 	MaskAndFlags.bForceOpaque = MaskAndFlags.bForceOpaque || Instance.bForceOpaque;
 	MaskAndFlags.bDoubleSided = MaskAndFlags.bDoubleSided || Instance.bDoubleSided;
@@ -245,6 +257,12 @@ void SetupRayTracingMeshCommandMaskAndStatus(FRayTracingMeshCommand& MeshCommand
 
 	if (MaskMode == ERayTracingViewMaskMode::PathTracing || MaskMode == ERayTracingViewMaskMode::LightMapTracing)
 	{
+		if (!MaskInfo.bAffectsDynamicIndirectLighting)
+		{
+			// If the object does not affect indirect lighting, remove all indirect bits.
+			MeshCommand.InstanceMask &= ~ComputeRayTracingInstanceMask(ERayTracingInstanceMaskType::VisibleInIndirectRay, MaskMode);
+		}
+
 		if (MaskInfo.bAffectsIndirectLightingOnly)
 		{
 			MeshCommand.InstanceMask &= ~ComputeRayTracingInstanceMask(ERayTracingInstanceMaskType::VisibleInPrimaryRay, MaskMode);
diff --git a/Engine/Source/Runtime/Renderer/Private/RayTracing/RayTracingInstanceMask.h b/Engine/Source/Runtime/Renderer/Private/RayTracing/RayTracingInstanceMask.h
--- a/Engine/Source/Runtime/Renderer/Private/RayTracing/RayTracingInstanceMask.h
+++ b/Engine/Source/Runtime/Renderer/Private/RayTracing/RayTracingInstanceMask.h
@@ -24,6 +24,7 @@ enum class ERayTracingInstanceMaskType : uint8
 
 	// path tracing specific mask type
 	VisibleInPrimaryRay,
+	VisibleInIndirectRay
 };
 
 /** MeshCommands mode shares the same status for ray tracing view mask mode*/
@@ -47,6 +48,7 @@ struct FSceneProxyRayTracingMaskInfo
 {
 	bool bAffectsIndirectLightingOnly = false;
 	bool bCastHiddenShadow = false;
+	bool bAffectsDynamicIndirectLighting = true;
 	ERayTracingViewMaskMode MaskMode = ERayTracingViewMaskMode::RayTracing;
 };
 
@@ -58,7 +60,7 @@ FSceneProxyRayTracingMaskInfo GetSceneProxyRayTracingMaskInfo(const FPrimitiveSc
 //-------------------------------------------------------
 
 FRayTracingMaskAndFlags BuildRayTracingInstanceMaskAndFlags(TArrayView<const FMeshBatch> MeshBatches, ERHIFeatureLevel::Type FeatureLevel,
-	ERayTracingViewMaskMode MaskMode, bool bAffectIndirectLightingOnly, ERayTracingInstanceLayer InstanceLayer, bool bCastHiddenShadow, uint8 ExtraMask = 0);
+	const FSceneProxyRayTracingMaskInfo& SceneProxyRayTracingMaskInfo, ERayTracingInstanceLayer InstanceLayer, uint8 ExtraMask = 0);
 
 // Build mask and flags without modification of RayTracingInstance
 FRayTracingMaskAndFlags BuildRayTracingInstanceMaskAndFlags(const FRayTracingInstance& Instance, const FPrimitiveSceneProxy& PrimitiveSceneProxy, const FSceneViewFamily* SceneViewFamily);
diff --git a/Engine/Source/Runtime/Renderer/Public/RayTracingInstance.h b/Engine/Source/Runtime/Renderer/Public/RayTracingInstance.h
--- a/Engine/Source/Runtime/Renderer/Public/RayTracingInstance.h
+++ b/Engine/Source/Runtime/Renderer/Public/RayTracingInstance.h
@@ -129,7 +129,7 @@ struct FRayTracingInstance
 
 
 /** Build mask and flags based on materials specified in Materials. You can still override Mask after calling this function. */
-UE_DEPRECATED(5.2, "Use BuildRayTracingInstanceMaskAndFlags() with ERayTracingViewMaskMode instead. Calling this function leads to incorrect path tracing result")
+UE_DEPRECATED(5.2, "Use BuildRayTracingInstanceMaskAndFlags() with FSceneProxyRayTracingMaskInfo instead. Calling this function leads to incorrect path tracing result")
 FRayTracingMaskAndFlags BuildRayTracingInstanceMaskAndFlags(TArrayView<const FMeshBatch> MeshBatches, ERHIFeatureLevel::Type FeatureLevel, ERayTracingInstanceLayer InstanceLayer = ERayTracingInstanceLayer::NearField, uint8 ExtraMask = 0);
 
 UE_DEPRECATED(5.2, "Use BlendModeToRayTracingInstanceMask() instead. Calling this function leads to incorrect path tracing result.")
