PR #9920: Fix missing break causing editor to crash when loading specific scenes (ComputeRayTracingInstanceMask) (Contributed by [REDACTED_PERSON])

diff --git a/Engine/Source/Runtime/Renderer/Private/RayTracing/RayTracingInstanceMask.cpp b/Engine/Source/Runtime/Renderer/Private/RayTracing/RayTracingInstanceMask.cpp
--- a/Engine/Source/Runtime/Renderer/Private/RayTracing/RayTracingInstanceMask.cpp
+++ b/Engine/Source/Runtime/Renderer/Private/RayTracing/RayTracingInstanceMask.cpp
@@ -44,6 +44,7 @@ uint8 ComputeRayTracingInstanceMask(ERayTracingInstanceMaskType MaskType, ERayTr
 			break;
 		case ERayTracingInstanceMaskType::VisibleInIndirectRay:
 			Mask = 0;
+			break;
 		default:
 			checkNoEntry();
 			break;
