Workaround for corruption in RVT render texture when we choose the fast path to not clear them.
Issue is related to the RDG transient alloctor and transitions on DX12, so skip use of transient allocator for RVT render targets for now.
Memory overhead of the workaround is small since the render target allocations are likely less than total 1MB in most use cases and are still reused between passes.

diff --git a/Engine/Source/Runtime/Renderer/Private/VT/RuntimeVirtualTextureRender.cpp b/Engine/Source/Runtime/Renderer/Private/VT/RuntimeVirtualTextureRender.cpp
--- a/Engine/Source/Runtime/Renderer/Private/VT/RuntimeVirtualTextureRender.cpp
+++ b/Engine/Source/Runtime/Renderer/Private/VT/RuntimeVirtualTextureRender.cpp
@@ -491,9 +491,14 @@ namespace RuntimeVirtualTexture
 			// Use direct aliasing for compression pass on platforms that support it.
 			bDirectAliasing = bCompressedFormat && GRHISupportsUAVFormatAliasing && CVarVTDirectCompress.GetValueOnRenderThread() != 0;
 
-			// Some RHI need explicit flag to avoid a fast clear, even when we don't use ERenderTargetLoadAction::EClear
-			const ETextureCreateFlags RTClearFlags = Desc.bClearTextures ? TexCreate_None : TexCreate_NoFastClear;
+			// Some problems happen when we don't use ERenderTargetLoadAction::EClear:
+			// * Some RHI need explicit flag to avoid a fast clear (TexCreate_NoFastClear).
+			// * DX12 RHI has a bug with RDG transient allocator (UE-173023) so we use TexCreate_Shared to avoid that.
+			const ETextureCreateFlags RTNoClearHackFlags = TexCreate_NoFastClear | TexCreate_Shared;
+
+			const ETextureCreateFlags RTClearFlags = Desc.bClearTextures ? TexCreate_None : RTNoClearHackFlags;
 			const ETextureCreateFlags RTCreateFlags = TexCreate_RenderTargetable | TexCreate_ShaderResource | RTClearFlags;
+			
 			// Not all mobile RHIs support sRGB texture views/aliasing, use only linear targets on mobile
 			const ETextureCreateFlags RTSrgbFlags = Desc.FeatureLevel > ERHIFeatureLevel::ES3_1 ? TexCreate_SRGB : TexCreate_None;
 
