Fix a bug of customdepth on mobile opengl.

diff --git a/Engine/Source/Runtime/Renderer/Private/MobileShadingRenderer.cpp b/Engine/Source/Runtime/Renderer/Private/MobileShadingRenderer.cpp
--- a/Engine/Source/Runtime/Renderer/Private/MobileShadingRenderer.cpp
+++ b/Engine/Source/Runtime/Renderer/Private/MobileShadingRenderer.cpp
@@ -128,7 +128,7 @@ static FMobileCustomDepthStencilUsage GetCustomDepthStencilUsage(const FViewInfo
 		if (CVarMobileCustomDepthForTranslucency.GetValueOnAnyThread() != 0)
 		{
 			CustomDepthStencilUsage.bUsesCustomDepthStencil = View.bUsesCustomDepth || View.bUsesCustomStencil;
-			CustomDepthStencilUsage.bSamplesCustomDepthAndStencil = View.bUsesCustomDepth && View.bUsesCustomStencil;
+			CustomDepthStencilUsage.bSamplesCustomDepthAndStencil = View.bUsesCustomStencil;
 		}
 
 		if (!CustomDepthStencilUsage.bSamplesCustomDepthAndStencil)
@@ -152,7 +152,7 @@ static FMobileCustomDepthStencilUsage GetCustomDepthStencilUsage(const FViewInfo
 						CustomDepthStencilUsage.bUsesCustomDepthStencil |= true;
 					}
 
-					if (bUsesCustomDepth && bUsesCustomStencil)
+					if (bUsesCustomStencil)
 					{
 						CustomDepthStencilUsage.bSamplesCustomDepthAndStencil |= true;
 						break;
