Path Tracer: Only apply decals below a certain path roughness threshold

On the scifi train station benchmark (camera 15) this reduces frame time from 66ms down to 41ms (compared to 36ms with decals disabled entirely)

In the future the logic could be fine-tuned to gradually fadeoff the decals around the cutoff, but keeping the logic simple for now.

diff --git a/Engine/Shaders/Private/PathTracing/PathTracingCore.ush b/Engine/Shaders/Private/PathTracing/PathTracingCore.ush
--- a/Engine/Shaders/Private/PathTracing/PathTracingCore.ush
+++ b/Engine/Shaders/Private/PathTracing/PathTracingCore.ush
@@ -53,7 +53,7 @@ uint EnableEmissive;
 uint SamplerType;
 uint VisualizeLightGrid;
 uint VisualizeDecalGrid;
-uint EnableDecals;
+float DecalRoughnessCutoff;
 float MaxPathIntensity;
 float MaxNormalBias;
 float FilterWidth;
@@ -370,7 +370,7 @@ FPathTracingPayload TraceTransparentRay(inout FPathState PathState, bool bIsCame
 		FPathTracingPayload HitPayload = UnpackPathTracingPayloadMaterial(PackedPayload, PathState.Ray);
 
 #if PLATFORM_SUPPORTS_CALLABLE_SHADERS
-		if (EnableDecals && HitPayload.IsDecalReceiver())
+		if (PathState.PathRoughness < DecalRoughnessCutoff && HitPayload.IsDecalReceiver())
 		{
 			const float3 TranslatedWorldPos = HitPayload.TranslatedWorldPos;
 
diff --git a/Engine/Source/Runtime/Renderer/Private/PathTracing.cpp b/Engine/Source/Runtime/Renderer/Private/PathTracing.cpp
--- a/Engine/Source/Runtime/Renderer/Private/PathTracing.cpp
+++ b/Engine/Source/Runtime/Renderer/Private/PathTracing.cpp
@@ -277,6 +277,13 @@ TAutoConsoleVariable<int32> CVarPathTracingDecalGridVisualize(
 	ECVF_RenderThreadSafe
 );
 
+TAutoConsoleVariable<float> CVarPathTracingDecalRoughnessCutoff(
+	TEXT("r.PathTracing.DecalRoughnessCutoff"),
+	0.15f,
+	TEXT("Do not evaluate decals beyond this roughness level to improve performance (default=0.15)\n"),
+	ECVF_RenderThreadSafe
+);
+
 TAutoConsoleVariable<int32> CVarPathTracingLightFunctionColor(
 	TEXT("r.PathTracing.LightFunctionColor"),
 	0,
@@ -305,12 +312,12 @@ BEGIN_SHADER_PARAMETER_STRUCT(FPathTracingData, )
 	SHADER_PARAMETER(uint32, VisualizeDecalGrid)
 	SHADER_PARAMETER(uint32, EnableAtmosphere)
 	SHADER_PARAMETER(uint32, EnableFog)
-	SHADER_PARAMETER(uint32, EnableDecals)
 	SHADER_PARAMETER(int32, MaxRaymarchSteps)
 	SHADER_PARAMETER(float, MaxPathIntensity)
 	SHADER_PARAMETER(float, MaxNormalBias)
 	SHADER_PARAMETER(float, FilterWidth)
 	SHADER_PARAMETER(float, AbsorptionScale)
+	SHADER_PARAMETER(float, DecalRoughnessCutoff)
 	SHADER_PARAMETER(float, CameraFocusDistance)
 	SHADER_PARAMETER(FVector2f, CameraLensRadius)
 END_SHADER_PARAMETER_STRUCT()
@@ -351,7 +358,7 @@ struct FPathTracingConfig
 			PathTracingData.AbsorptionScale != Other.PathTracingData.AbsorptionScale ||
 			PathTracingData.EnableAtmosphere != Other.PathTracingData.EnableAtmosphere ||
 			PathTracingData.EnableFog != Other.PathTracingData.EnableFog ||
-			PathTracingData.EnableDecals != Other.PathTracingData.EnableDecals ||
+			PathTracingData.DecalRoughnessCutoff != Other.PathTracingData.DecalRoughnessCutoff ||
 			PathTracingData.MaxRaymarchSteps != Other.PathTracingData.MaxRaymarchSteps ||
 			ViewRect != Other.ViewRect ||
 			LightShowFlags != Other.LightShowFlags ||
@@ -525,7 +532,7 @@ static void PreparePathTracingData(const FScene* Scene, const FViewInfo& View, F
 		&& (Scene->ExponentialFogs[0].FogData[0].Density > 0 ||
 			Scene->ExponentialFogs[0].FogData[1].Density > 0);
 
-	PathTracingData.EnableDecals = PathTracing::UsesDecals(*View.Family) && View.bHasRayTracingDecals;
+	PathTracingData.DecalRoughnessCutoff = PathTracing::UsesDecals(*View.Family) && View.bHasRayTracingDecals ? CVarPathTracingDecalRoughnessCutoff.GetValueOnRenderThread() : -1.0f;
 
 	PathTracingData.MaxRaymarchSteps = CVarPathTracingMaxRaymarchSteps.GetValueOnRenderThread();
 }
