Track if any/all segments of cached and dynamic raytracing instances are decals.
- exclude instances with mixed mesh decals from raytracing scene for now since they're not properly supported yet.

diff --git a/Engine/Source/Runtime/Renderer/Private/DeferredShadingRenderer.cpp b/Engine/Source/Runtime/Renderer/Private/DeferredShadingRenderer.cpp
--- a/Engine/Source/Runtime/Renderer/Private/DeferredShadingRenderer.cpp
+++ b/Engine/Source/Runtime/Renderer/Private/DeferredShadingRenderer.cpp
@@ -682,6 +682,7 @@ struct FRayTracingRelevantPrimitive
 	bool bAllSegmentsOpaque = true;
 	bool bAnySegmentsCastShadow = false;
 	bool bAnySegmentsDecal = false;
+	bool bAllSegmentsDecal = true;
 	bool bTwoSided = false;
 	bool bIsSky = false;
 	bool bAllSegmentsTranslucent = true;
@@ -693,9 +694,10 @@ struct FRayTracingRelevantPrimitive
 		Key ^= bAllSegmentsOpaque ? 0x1ull << 40 : 0x0;
 		Key ^= bAnySegmentsCastShadow ? 0x1ull << 41 : 0x0;
 		Key ^= bAnySegmentsDecal ? 0x1ull << 42 : 0x0;
-		Key ^= bTwoSided ? 0x1ull << 43 : 0x0;
-		Key ^= bIsSky ? 0x1ull << 44 : 0x0;
-		Key ^= bAllSegmentsTranslucent ? 0x1ull << 45 : 0x0;
+		Key ^= bAllSegmentsDecal ? 0x1ull << 43 : 0x0;
+		Key ^= bTwoSided ? 0x1ull << 44 : 0x0;
+		Key ^= bIsSky ? 0x1ull << 45 : 0x0;
+		Key ^= bAllSegmentsTranslucent ? 0x1ull << 46 : 0x0;
 		return Key ^ reinterpret_cast<uint64>(RayTracingGeometryRHI);
 	}
 
@@ -713,12 +715,6 @@ struct FRayTracingRelevantPrimitive
 		UpdateRayTracingMeshCommandMasks(Command, Flags, MaskMode);
 
 		InstanceMask = Command.InstanceMask;
-		bAllSegmentsOpaque = Command.bOpaque;
-		bAnySegmentsCastShadow = Command.bCastRayTracedShadows;
-		bAnySegmentsDecal = Command.bDecal;
-		bTwoSided = Command.bTwoSided;
-		bIsSky = Command.bIsSky;
-		bAllSegmentsTranslucent = Command.bIsTranslucent;
 	}
 };
 
@@ -900,9 +896,8 @@ static void GatherRayTracingRelevantPrimitives(const FScene& Scene, const FViewI
 
 					RelevantPrimitive.CachedRayTracingMeshCommandIndices = SceneInfo->CachedRayTracingMeshCommandIndicesPerLOD[LODIndex];
 					RelevantPrimitive.StateHash = SceneInfo->CachedRayTracingMeshCommandsHashPerLOD[LODIndex];
-
-					ERayTracingViewMaskMode MaskMode = static_cast<ERayTracingViewMaskMode>(Scene.CachedRayTracingMeshCommandsMode);
 					
+					// TODO: Cache these flags to avoid having to loop over the RayTracingMeshCommands
 					for (int32 CommandIndex : RelevantPrimitive.CachedRayTracingMeshCommandIndices)
 					{
 						if (CommandIndex >= 0)
@@ -913,6 +908,7 @@ static void GatherRayTracingRelevantPrimitives(const FScene& Scene, const FViewI
 							RelevantPrimitive.bAllSegmentsOpaque &= RayTracingMeshCommand.bOpaque;
 							RelevantPrimitive.bAnySegmentsCastShadow |= RayTracingMeshCommand.bCastRayTracedShadows;
 							RelevantPrimitive.bAnySegmentsDecal |= RayTracingMeshCommand.bDecal;
+							RelevantPrimitive.bAllSegmentsDecal &= RayTracingMeshCommand.bDecal;
 							RelevantPrimitive.bTwoSided |= RayTracingMeshCommand.bTwoSided;
 							RelevantPrimitive.bIsSky |= RayTracingMeshCommand.bIsSky;
 							RelevantPrimitive.bAllSegmentsTranslucent &= RayTracingMeshCommand.bIsTranslucent;
@@ -924,6 +920,8 @@ static void GatherRayTracingRelevantPrimitives(const FScene& Scene, const FViewI
 						}
 					}
 
+					ERayTracingViewMaskMode MaskMode = static_cast<ERayTracingViewMaskMode>(Scene.CachedRayTracingMeshCommandsMode);
+
 					RelevantPrimitive.UpdateMasks(Flags, MaskMode);
 				}
 			}
@@ -1176,6 +1174,17 @@ bool FDeferredShadingSceneRenderer::GatherRayTracingWorldInstancesForView(FRDGBu
 					// Autobuild of InstanceMaskAndFlags if the mask and flags are not built
 					UpdateRayTracingInstanceMaskAndFlagsIfNeeded(Instance, *SceneProxy, &ViewFamily);
 
+					if (Instance.MaskAndFlags.bAnySegmentsDecal && !Instance.MaskAndFlags.bAllSegmentsDecal)
+					{
+						// TODO: Warn user meshes with mixed decal segments are not currently supported in ray tracing
+						continue;
+					}
+
+					if (GRayTracingExcludeDecals && Instance.MaskAndFlags.bAnySegmentsDecal)
+					{
+						continue;
+					}
+
 					const uint32 InstanceIndex = RayTracingScene.Instances.Num();
 
 					FRayTracingGeometryInstance& RayTracingInstance = RayTracingScene.Instances.AddDefaulted_GetRef();
@@ -1439,8 +1448,18 @@ bool FDeferredShadingSceneRenderer::GatherRayTracingWorldInstancesForView(FRDGBu
 						continue;
 					}
 					
-					// TODO: support GRayTracingExcludeDecals, but not in the form of RayTracingMeshCommand.bDecal as that requires looping over all cached MDCs
-					// Instead, either make r.RayTracing.ExcludeDecals read only or request a recache of all ray tracing commands during which decals are excluded
+					// TODO: Consider requesting a recache of all ray tracing commands during which decals are excluded
+
+					if (SceneInfo->bCachedRayTracingInstancebAnySegmentsDecal && !SceneInfo->bCachedRayTracingInstancebAllSegmentsDecal)
+					{
+						// TODO: Warn user meshes with mixed decal segments are not currently supported in ray tracing
+						continue;
+					}
+
+					if (GRayTracingExcludeDecals && SceneInfo->bCachedRayTracingInstancebAnySegmentsDecal)
+					{
+						continue;
+					}
 
 					const int32 NewInstanceIndex = RayTracingScene.Instances.Num();
 
@@ -1512,6 +1531,12 @@ bool FDeferredShadingSceneRenderer::GatherRayTracingWorldInstancesForView(FRDGBu
 						continue; // skip dynamic primitives and other 
 					}
 
+					if (RelevantPrimitive.bAnySegmentsDecal && !RelevantPrimitive.bAllSegmentsDecal)
+					{
+						// TODO: Warn user meshes with mixed decal segments are not currently supported in ray tracing
+						continue;
+					}
+
 					if ((GRayTracingExcludeDecals && RelevantPrimitive.bAnySegmentsDecal)
 						|| (GRayTracingExcludeTranslucent && RelevantPrimitive.bAllSegmentsTranslucent)
 						|| (GRayTracingExcludeSky && RelevantPrimitive.bIsSky))
diff --git a/Engine/Source/Runtime/Renderer/Private/PrimitiveSceneInfo.cpp b/Engine/Source/Runtime/Renderer/Private/PrimitiveSceneInfo.cpp
--- a/Engine/Source/Runtime/Renderer/Private/PrimitiveSceneInfo.cpp
+++ b/Engine/Source/Runtime/Renderer/Private/PrimitiveSceneInfo.cpp
@@ -1049,6 +1049,9 @@ void FPrimitiveSceneInfo::UpdateCachedRayTracingInstance(FPrimitiveSceneInfo* Sc
 		{
 			SceneInfo->CachedRayTracingInstance.Flags |= ERayTracingInstanceFlags::TriangleCullDisable;
 		}
+
+		SceneInfo->bCachedRayTracingInstancebAnySegmentsDecal = CachedRayTracingInstance.MaskAndFlags.bAnySegmentsDecal;
+		SceneInfo->bCachedRayTracingInstancebAllSegmentsDecal = CachedRayTracingInstance.MaskAndFlags.bAllSegmentsDecal;
 	}
 }
 
diff --git a/Engine/Source/Runtime/Renderer/Private/RayTracing/RayTracingInstanceMask.cpp b/Engine/Source/Runtime/Renderer/Private/RayTracing/RayTracingInstanceMask.cpp
--- a/Engine/Source/Runtime/Renderer/Private/RayTracing/RayTracingInstanceMask.cpp
+++ b/Engine/Source/Runtime/Renderer/Private/RayTracing/RayTracingInstanceMask.cpp
@@ -144,6 +144,8 @@ FRayTracingMaskAndFlags BuildRayTracingInstanceMaskAndFlags(TArrayView<const FMe
 	bool bAllSegmentsOpaque = true;
 	bool bAnySegmentsCastShadow = false;
 	bool bAllSegmentsCastShadow = true;
+	bool bAnySegmentsDecal = false;
+	bool bAllSegmentsDecal = true;
 	bool bDoubleSided = false;
 	ERayTracingViewMaskMode MaskMode = SceneProxyRayTracingMaskInfo.MaskMode;
 	Result.Mask = ExtraMask;
@@ -161,12 +163,16 @@ FRayTracingMaskAndFlags BuildRayTracingInstanceMaskAndFlags(TArrayView<const FMe
 			bAllSegmentsOpaque &= BlendMode == BLEND_Opaque;
 			bAnySegmentsCastShadow |= MeshBatch.CastRayTracedShadow && Material.CastsRayTracedShadows();
 			bAllSegmentsCastShadow &= MeshBatch.CastRayTracedShadow && Material.CastsRayTracedShadows();
+			bAnySegmentsDecal |= Material.IsDeferredDecal();
+			bAllSegmentsDecal &= Material.IsDeferredDecal();
 			bDoubleSided |= MeshBatch.bDisableBackfaceCulling || Material.IsTwoSided();
 		}
 	}
 
 	Result.bForceOpaque = bAllSegmentsOpaque && bAllSegmentsCastShadow;
 	Result.bDoubleSided = bDoubleSided;
+	Result.bAnySegmentsDecal = bAnySegmentsDecal;
+	Result.bAllSegmentsDecal = bAllSegmentsDecal;
 
 	Result.Mask |= bAnySegmentsCastShadow ? ComputeRayTracingInstanceMask(ERayTracingInstanceMaskType::Shadow, MaskMode) : 0;
 
@@ -228,11 +234,9 @@ FRayTracingMaskAndFlags BuildRayTracingInstanceMaskAndFlags(const FRayTracingIns
 	const TArrayView<const FMeshBatch> MeshBatches = Instance.GetMaterials();
 
 	// add extra mask bit for hair.
-	uint8 ExtraMask = Instance.bThinGeometry ? ComputeRayTracingInstanceMask(ERayTracingInstanceMaskType::HairStrands, MaskInfo.MaskMode) : 0;
+	const uint8 ExtraMask = Instance.bThinGeometry ? ComputeRayTracingInstanceMask(ERayTracingInstanceMaskType::HairStrands, MaskInfo.MaskMode) : 0;
 
-	return BuildRayTracingInstanceMaskAndFlags(
-			MeshBatches, PrimitiveSceneProxy.GetScene().GetFeatureLevel(), 
-			MaskInfo, Instance.InstanceLayer, ExtraMask);
+	return BuildRayTracingInstanceMaskAndFlags(MeshBatches, PrimitiveSceneProxy.GetScene().GetFeatureLevel(), MaskInfo, Instance.InstanceLayer, ExtraMask);
 }
 
 void SetupRayTracingMeshCommandMaskAndStatus(FRayTracingMeshCommand& MeshCommand, const FMeshBatch& MeshBatch, const FPrimitiveSceneProxy& PrimitiveSceneProxy,
@@ -242,7 +246,7 @@ void SetupRayTracingMeshCommandMaskAndStatus(FRayTracingMeshCommand& MeshCommand
 
 	MeshCommand.bCastRayTracedShadows = MeshBatch.CastRayTracedShadow && MaterialResource.CastsRayTracedShadows();
 	MeshCommand.bOpaque = MaterialResource.GetBlendMode() == EBlendMode::BLEND_Opaque && !(VertexFactory->GetType()->SupportsRayTracingProceduralPrimitive() && FDataDrivenShaderPlatformInfo::GetSupportsRayTracingProceduralPrimitive(GMaxRHIShaderPlatform));
-	MeshCommand.bDecal = MaterialResource.GetMaterialDomain() == EMaterialDomain::MD_DeferredDecal;
+	MeshCommand.bDecal = MaterialResource.IsDeferredDecal();
 	MeshCommand.bIsSky = MaterialResource.IsSky();
 	MeshCommand.bTwoSided = MaterialResource.IsTwoSided();
 	MeshCommand.bIsTranslucent = MaterialResource.GetBlendMode() == EBlendMode::BLEND_Translucent;
diff --git a/Engine/Source/Runtime/Renderer/Public/PrimitiveSceneInfo.h b/Engine/Source/Runtime/Renderer/Public/PrimitiveSceneInfo.h
--- a/Engine/Source/Runtime/Renderer/Public/PrimitiveSceneInfo.h
+++ b/Engine/Source/Runtime/Renderer/Public/PrimitiveSceneInfo.h
@@ -408,6 +408,8 @@ public:
 	TArray<uint64> CachedRayTracingMeshCommandsHashPerLOD;
 	// TODO: this should be placed in FRayTracingScene and we have a pointer/handle here. It's here for now for PoC
 	FRayTracingGeometryInstance CachedRayTracingInstance;
+	bool bCachedRayTracingInstancebAnySegmentsDecal : 1;
+	bool bCachedRayTracingInstancebAllSegmentsDecal : 1;
 	TArray<FBoxSphereBounds> CachedRayTracingInstanceWorldBounds;
 	int32 SmallestRayTracingInstanceWorldBoundsIndex;
 #endif
diff --git a/Engine/Source/Runtime/Renderer/Public/RayTracingInstance.h b/Engine/Source/Runtime/Renderer/Public/RayTracingInstance.h
--- a/Engine/Source/Runtime/Renderer/Public/RayTracingInstance.h
+++ b/Engine/Source/Runtime/Renderer/Public/RayTracingInstance.h
@@ -15,6 +15,8 @@ struct FRayTracingMaskAndFlags
 		: Mask(0xFF)
 		, bForceOpaque(false)
 		, bDoubleSided(false)
+		, bAnySegmentsDecal(false)
+		, bAllSegmentsDecal(false)
 	{}
 
 	/** Instance mask that can be used to exclude the instance from specific effects (eg. ray traced shadows). */
@@ -25,6 +27,10 @@ struct FRayTracingMaskAndFlags
 
 	/** Whether ray hits should be registered for front and back faces. */
 	uint8 bDoubleSided : 1;
+
+	/** Whether any or all of the segments in the instance are decals. */
+	uint8 bAnySegmentsDecal : 1;
+	uint8 bAllSegmentsDecal : 1;
 };
 
 enum class ERayTracingInstanceLayer : uint8
