Fix a bug of CustomDepth on mobile OpenGL.

diff --git a/Engine/Source/Runtime/Engine/Public/SceneTexturesConfig.h b/Engine/Source/Runtime/Engine/Public/SceneTexturesConfig.h
--- a/Engine/Source/Runtime/Engine/Public/SceneTexturesConfig.h
+++ b/Engine/Source/Runtime/Engine/Public/SceneTexturesConfig.h
@@ -116,7 +116,7 @@ struct ENGINE_API FSceneTexturesConfig
 		, bIsUsingGBuffers{}
 		, bKeepDepthContent{ 1 }
 		, bPreciseDepthAux{}
-		, bSamplesCustomDepthAndStencil{}
+		, bSamplesCustomStencil{}
 		, bMemorylessMSAA{}
 		, bSupportsXRTargetManagerDepthAlloc{}
 	{}
@@ -184,8 +184,8 @@ struct ENGINE_API FSceneTexturesConfig
 	// (Mobile) True if SceneDepthAux should use a precise pixel format
 	uint32 bPreciseDepthAux : 1;
 
-	// (Mobile) True if both CustomDepth and CustomStencil are sampled in a shader
-	uint32 bSamplesCustomDepthAndStencil : 1;
+	// (Mobile) True if CustomStencil are sampled in a shader
+	uint32 bSamplesCustomStencil : 1;
 	
 	// (Mobile) True if MSAA targets can be memoryless
 	uint32 bMemorylessMSAA : 1;
diff --git a/Engine/Source/Runtime/Renderer/Private/CustomDepthRendering.cpp b/Engine/Source/Runtime/Renderer/Private/CustomDepthRendering.cpp
--- a/Engine/Source/Runtime/Renderer/Private/CustomDepthRendering.cpp
+++ b/Engine/Source/Runtime/Renderer/Private/CustomDepthRendering.cpp
@@ -362,8 +362,8 @@ bool FSceneRenderer::RenderCustomDepthPass(
 		FRDGTextureRef CustomDepth = CustomDepthTextures.Depth;
 
 		// TextureView is not supported in GLES, so we can't lookup CustomDepth and CustomStencil from a single texture
-		// Do a copy of the CustomDepthStencil texture if both CustomDepth and CustomStencil are sampled in a shader.
-		if (IsOpenGLPlatform(ShaderPlatform) && Config.bSamplesCustomDepthAndStencil)
+		// Do a copy of the CustomDepthStencil texture if CustomStencil is sampled in a shader.
+		if (IsOpenGLPlatform(ShaderPlatform) && Config.bSamplesCustomStencil)
 		{
 			CustomDepth = GraphBuilder.CreateTexture(CustomDepthTextures.Depth->Desc, TEXT("CustomDepthCopy"));
 			AddCopyTexturePass(GraphBuilder, CustomDepthTextures.Depth, CustomDepth);
diff --git a/Engine/Source/Runtime/Renderer/Private/MobileShadingRenderer.cpp b/Engine/Source/Runtime/Renderer/Private/MobileShadingRenderer.cpp
--- a/Engine/Source/Runtime/Renderer/Private/MobileShadingRenderer.cpp
+++ b/Engine/Source/Runtime/Renderer/Private/MobileShadingRenderer.cpp
@@ -113,8 +113,8 @@ extern bool IsMobileEyeAdaptationEnabled(const FViewInfo& View);
 struct FMobileCustomDepthStencilUsage
 {
 	bool bUsesCustomDepthStencil = false;
-	// whether both CustomDepth and CustomStencil are sampled as a textures
-	bool bSamplesCustomDepthAndStencil = false;
+	// whether CustomStencil is sampled as a textures
+	bool bSamplesCustomStencil = false;
 };
 
 static FMobileCustomDepthStencilUsage GetCustomDepthStencilUsage(const FViewInfo& View)
@@ -128,10 +128,10 @@ static FMobileCustomDepthStencilUsage GetCustomDepthStencilUsage(const FViewInfo
 		if (CVarMobileCustomDepthForTranslucency.GetValueOnAnyThread() != 0)
 		{
 			CustomDepthStencilUsage.bUsesCustomDepthStencil = View.bUsesCustomDepth || View.bUsesCustomStencil;
-			CustomDepthStencilUsage.bSamplesCustomDepthAndStencil = View.bUsesCustomStencil;
+			CustomDepthStencilUsage.bSamplesCustomStencil = View.bUsesCustomStencil;
 		}
 
-		if (!CustomDepthStencilUsage.bSamplesCustomDepthAndStencil)
+		if (!CustomDepthStencilUsage.bSamplesCustomStencil)
 		{
 			// Find out whether post-process materials use CustomDepth/Stencil lookups
 			const FBlendableManager& BlendableManager = View.FinalPostProcessSettings.BlendableManager;
@@ -154,7 +154,7 @@ static FMobileCustomDepthStencilUsage GetCustomDepthStencilUsage(const FViewInfo
 
 					if (bUsesCustomStencil)
 					{
-						CustomDepthStencilUsage.bSamplesCustomDepthAndStencil |= true;
+						CustomDepthStencilUsage.bSamplesCustomStencil |= true;
 						break;
 					}
 				}
@@ -604,7 +604,7 @@ void FMobileSceneRenderer::InitViews(FRDGBuilder& GraphBuilder, FSceneTexturesCo
 			CustomDepthStencilUsage = GetCustomDepthStencilUsage(Views[ViewIndex]);
 			Views[ViewIndex].bCustomDepthStencilValid = bCouldUseCustomDepthStencil && CustomDepthStencilUsage.bUsesCustomDepthStencil;
 			bShouldRenderCustomDepth |= Views[ViewIndex].bCustomDepthStencilValid;
-			SceneTexturesConfig.bSamplesCustomDepthAndStencil |= bShouldRenderCustomDepth && CustomDepthStencilUsage.bSamplesCustomDepthAndStencil;
+			SceneTexturesConfig.bSamplesCustomStencil |= bShouldRenderCustomDepth && CustomDepthStencilUsage.bSamplesCustomStencil;
 		}
 	}
 	
