Add ShaderPrint zoom.

diff --git a/Engine/Shaders/Private/ShaderPrintDraw.usf b/Engine/Shaders/Private/ShaderPrintDraw.usf
--- a/Engine/Shaders/Private/ShaderPrintDraw.usf
+++ b/Engine/Shaders/Private/ShaderPrintDraw.usf
@@ -356,4 +356,59 @@ void DrawSymbolsPS(
 
 	float3 Color = InColor * (InShadow == 0 ? Alpha : 0);
 	OutColor = float4(Color, Alpha);
+}
+
+
+// --------------------------------------------------------------------------
+// Zoomed output
+
+Texture2D<float4> InTexture;
+RWTexture2D<float4> OutTexture;
+
+uint PixelExtent;
+uint ZoomFactor;
+uint Corner;
+int2 Resolution;
+
+bool IsValid(int2 InCoord)
+{
+	return all(InCoord >= 0) && all(InCoord < Resolution);
+}
+
+[numthreads(8, 8, 1)]
+void DrawZoomCS(int2 DispatchThreadId : SV_DispatchThreadID)
+{	
+	if (any(DispatchThreadId > PixelExtent * 2))
+	{
+		return;
+	}
+
+	const int2 InCoord = DispatchThreadId + ShaderPrintData.CursorCoord - PixelExtent.xx;
+	const int2 Offset = int2(10, 10);
+	const int2 OutMaxPixel = (PixelExtent * 2 + 1) * ZoomFactor;
+
+	int2 OutOffset = Offset;
+	switch(Corner)
+	{
+	case 0: OutOffset = Offset; break;
+	case 1: OutOffset = int2(Resolution.x - Offset.x - OutMaxPixel.x, Offset.y); break;
+	case 2: OutOffset = Resolution - Offset - OutMaxPixel; break;
+	case 3: OutOffset = int2(Offset.x, Resolution.y - Offset.y - OutMaxPixel.y); break;
+	};
+
+	if (IsValid(ShaderPrintData.CursorCoord) && IsValid(InCoord))
+	{
+		const float4 Color = InTexture.Load(uint3(InCoord, 0));
+		for (uint Y=0; Y<ZoomFactor; ++Y)
+		for (uint X=0; X<ZoomFactor; ++X)
+		{
+			const int2 OutCoord = DispatchThreadId * ZoomFactor + int2(X, Y) + OutOffset;
+			const bool bIsBorder = any(OutCoord == OutOffset) || any(OutCoord == OutOffset + OutMaxPixel - 1);
+			if (IsValid(OutCoord))
+			{
+				OutTexture[OutCoord] = bIsBorder ? float4(1,1,0,1) : Color;
+			}
+
+		}
+	}
 }
\ No newline at end of file
diff --git a/Engine/Source/Runtime/Renderer/Private/ShaderPrint.cpp b/Engine/Source/Runtime/Renderer/Private/ShaderPrint.cpp
--- a/Engine/Source/Runtime/Renderer/Private/ShaderPrint.cpp
+++ b/Engine/Source/Runtime/Renderer/Private/ShaderPrint.cpp
@@ -75,6 +75,31 @@ namespace ShaderPrint
 		TEXT("Lock the line drawing.\n"),
 		ECVF_Cheat | ECVF_RenderThreadSafe);
 
+
+	static TAutoConsoleVariable<int32> CVarDrawZoomEnable(
+		TEXT("r.ShaderPrint.Zoom"),
+		0,
+		TEXT("Enable zoom magnification around the mouse cursor.\n"),
+		ECVF_Cheat | ECVF_RenderThreadSafe);
+
+	static TAutoConsoleVariable<int32> CVarDrawZoomPixel(
+		TEXT("r.ShaderPrint.Zoom.Pixel"),
+		16,
+		TEXT("Number of pixels magnified around the mouse cursor.\n"),
+		ECVF_Cheat | ECVF_RenderThreadSafe);
+
+	static TAutoConsoleVariable<int32> CVarDrawZoomFactor(
+		TEXT("r.ShaderPrint.Zoom.Factor"),
+		8,
+		TEXT("Zoom factor for magnification around the mouse cursor.\n"),
+		ECVF_Cheat | ECVF_RenderThreadSafe);
+
+	static TAutoConsoleVariable<int32> CVarDrawZoomCorner(
+		TEXT("r.ShaderPrint.Zoom.Corner"),
+		3,
+		TEXT("Select in which corner the zoom magnifer is displayed (0:top-left, 1:top-right, 2:bottom-right, 3:bottom-left).\n"),
+		ECVF_Cheat | ECVF_RenderThreadSafe);
+
 	//////////////////////////////////////////////////////////////////////////////////////////////////
 	// Global states
 
@@ -530,6 +555,32 @@ namespace ShaderPrint
 		SHADER_PARAMETER_STRUCT_INCLUDE(FShaderDrawDebugPS::FParameters, PS)
 	END_SHADER_PARAMETER_STRUCT()
 
+
+	// Shader to zoom the final output
+	class FShaderZoomCS : public FGlobalShader
+	{
+	public:
+		DECLARE_GLOBAL_SHADER(FShaderZoomCS);
+		SHADER_USE_PARAMETER_STRUCT(FShaderZoomCS, FGlobalShader);
+
+		BEGIN_SHADER_PARAMETER_STRUCT(FParameters, )
+			SHADER_PARAMETER(uint32, PixelExtent)
+			SHADER_PARAMETER(uint32, ZoomFactor)
+			SHADER_PARAMETER(uint32, Corner)
+			SHADER_PARAMETER(FIntPoint, Resolution)
+			SHADER_PARAMETER_STRUCT_REF(FShaderPrintCommonParameters, Common)
+			SHADER_PARAMETER_RDG_TEXTURE(Texture2D<float4>, InTexture)
+			SHADER_PARAMETER_RDG_TEXTURE_UAV(RWTexture2D<float4>, OutTexture)
+		END_SHADER_PARAMETER_STRUCT()
+
+		static bool ShouldCompilePermutation(FGlobalShaderPermutationParameters const& Parameters)
+		{
+			return IsSupported(Parameters.Platform);
+		}
+	};
+
+	IMPLEMENT_GLOBAL_SHADER(FShaderZoomCS, "/Engine/Private/ShaderPrintDraw.usf", "DrawZoomCS", SF_Compute);
+
 	//////////////////////////////////////////////////////////////////////////////////////////////////
 	// Setup render data
 
@@ -953,6 +1004,36 @@ namespace ShaderPrint
 			});
 	}
 
+	void InternalDrawZoom(FRDGBuilder& GraphBuilder, const FShaderPrintData& ShaderPrintData, const FScreenPassTexture& OutputTexture)
+	{
+		FGlobalShaderMap* GlobalShaderMap = GetGlobalShaderMap(GMaxRHIFeatureLevel);
+		FRDGTextureDesc Desc = OutputTexture.Texture->Desc;
+		Desc.Flags |= ETextureCreateFlags::UAV;
+		FRDGTextureRef OutZoomTexture = GraphBuilder.CreateTexture(Desc, TEXT("ShaderPrint.OutZoomTexture"));
+		AddCopyTexturePass(GraphBuilder, OutputTexture.Texture, OutZoomTexture);
+
+		TShaderMapRef<FShaderZoomCS> ComputeShader(GlobalShaderMap);
+
+		FShaderZoomCS::FParameters* PassParameters = GraphBuilder.AllocParameters<FShaderZoomCS::FParameters>();
+		PassParameters->PixelExtent = FMath::Clamp(CVarDrawZoomPixel.GetValueOnRenderThread(), 2, 128);
+		PassParameters->ZoomFactor  = FMath::Clamp(CVarDrawZoomFactor.GetValueOnRenderThread(), 1, 10);
+		PassParameters->Resolution = OutputTexture.Texture->Desc.Extent;
+		PassParameters->Corner = FMath::Clamp(CVarDrawZoomCorner.GetValueOnRenderThread(), 0, 3);
+		PassParameters->Common = ShaderPrintData.UniformBuffer;
+		PassParameters->InTexture = OutputTexture.Texture;
+		PassParameters->OutTexture = GraphBuilder.CreateUAV(OutZoomTexture);
+
+		const uint32 SrcPixelCount = PassParameters->PixelExtent * 2 + 1;
+		const uint32 OutPixelCount = PassParameters->ZoomFactor * SrcPixelCount;
+		const FIntVector GroupCount = FComputeShaderUtils::GetGroupCount(FIntPoint(PassParameters->PixelExtent * 2 + 1), FIntPoint(8));
+		FComputeShaderUtils::AddPass(
+			GraphBuilder,
+			RDG_EVENT_NAME("ShaderPrint::DrawZoom"),
+			ComputeShader, PassParameters, GroupCount);
+
+		AddCopyTexturePass(GraphBuilder, OutZoomTexture, OutputTexture.Texture);
+	}
+
 	void InternalDrawView(FRDGBuilder& GraphBuilder, const FViewInfo& View, const FShaderPrintData& ShaderPrintData, const FScreenPassTexture& OutputTexture, const FScreenPassTexture& DepthTexture)
 	{
 		if (!ensure(OutputTexture.IsValid()))
@@ -992,6 +1073,13 @@ namespace ShaderPrint
 			const int32 FrameNumber = View.Family ? View.Family->FrameNumber : 0u;
 			InternalDrawView_Characters(GraphBuilder, ShaderPrintData, OutputViewRect, FrameNumber, OutputTexture);
 		}
+
+		// Zoom
+		const bool bZoom = CVarDrawZoomEnable.GetValueOnRenderThread() > 0;
+		if (bZoom)
+		{
+			InternalDrawZoom(GraphBuilder, ShaderPrintData, OutputTexture);
+		}
 	}
 
 	void DrawView(FRDGBuilder& GraphBuilder, const FViewInfo& View, const FScreenPassTexture& OutputTexture, const FScreenPassTexture& DepthTexture)
