Moving more types to RenderGraphFwd.h and using that in a couple new places.

diff --git a/Engine/Source/Runtime/RenderCore/Public/RenderGraphDefinitions.h b/Engine/Source/Runtime/RenderCore/Public/RenderGraphDefinitions.h
--- a/Engine/Source/Runtime/RenderCore/Public/RenderGraphDefinitions.h
+++ b/Engine/Source/Runtime/RenderCore/Public/RenderGraphDefinitions.h
@@ -743,8 +743,6 @@ using FRDGBufferHandle = TRDGHandle<FRDGBuffer, uint16>;
 using FRDGBufferRegistry = TRDGHandleRegistry<FRDGBufferHandle, ERDGHandleRegistryDestructPolicy::Registry>;
 using FRDGBufferBitArray = TRDGHandleBitArray<FRDGBufferHandle>;
 
-class FRDGPooledTexture;
-class FRDGPooledBuffer;
 class FRDGBufferPool;
 class FRDGTransientRenderTarget;
 
diff --git a/Engine/Source/Runtime/RenderCore/Public/RenderGraphFwd.h b/Engine/Source/Runtime/RenderCore/Public/RenderGraphFwd.h
--- a/Engine/Source/Runtime/RenderCore/Public/RenderGraphFwd.h
+++ b/Engine/Source/Runtime/RenderCore/Public/RenderGraphFwd.h
@@ -25,6 +25,10 @@ class FRDGBuilder;
 class FRDGPass;
 using FRDGPassRef = FRDGPass*;
 
+class FRDGPooledBuffer;
+
+class FRDGPooledTexture;
+
 class FRDGResource;
 using FRDGResourceRef = FRDGResource*;
 
diff --git a/Engine/Source/Runtime/Renderer/Private/ReflectionEnvironment.h b/Engine/Source/Runtime/Renderer/Private/ReflectionEnvironment.h
--- a/Engine/Source/Runtime/Renderer/Private/ReflectionEnvironment.h
+++ b/Engine/Source/Runtime/Renderer/Private/ReflectionEnvironment.h
@@ -8,11 +8,10 @@
 
 #include "CoreTypes.h"
 #include "Math/Vector4.h"
+#include "RenderGraphFwd.h"
 #include "RHIFwd.h"
 #include "ShaderParameterMacros.h"
 
-class FRDGBuilder;
-class FRDGPooledBuffer;
 class FSkyLightSceneProxy;
 class FViewInfo;
 struct FEngineShowFlags;
diff --git a/Engine/Source/Runtime/Renderer/Private/VirtualShadowMaps/VirtualShadowMapProjection.h b/Engine/Source/Runtime/Renderer/Private/VirtualShadowMaps/VirtualShadowMapProjection.h
--- a/Engine/Source/Runtime/Renderer/Private/VirtualShadowMaps/VirtualShadowMapProjection.h
+++ b/Engine/Source/Runtime/Renderer/Private/VirtualShadowMaps/VirtualShadowMapProjection.h
@@ -8,16 +8,14 @@
 
 #include "CoreTypes.h"
 #include "Math/MathFwd.h"
+#include "RenderGraphFwd.h"
 #include "Templates/SharedPointerFwd.h"
 
 class FLightSceneInfo;
-class FRDGBuilder;
-class FRDGTexture;
 class FViewInfo;
 class FVirtualShadowMapArray;
 class FVirtualShadowMapClipmap;
 struct FMinimalSceneTextures;
-using FRDGTextureRef = FRDGTexture*;
 
 struct FTiledVSMProjection
 {
