Move VRS debug overlay functionaltiy over to the VRS image manager and display debug info for final combined shading rate image.

The generator interface provides a new function to get the current debug image in order to support generators (like CAS) that output different images depending on the pass. Output can be configured via generator specific CVars.

Added support for 2x4, 4x2, and 4x4 shading rates in the preview overlay shader.

diff --git a/Engine/Shaders/Private/VariableRateShading/VRSShadingRateCalculate.usf b/Engine/Shaders/Private/VariableRateShading/VRSShadingRateCalculate.usf
--- a/Engine/Shaders/Private/VariableRateShading/VRSShadingRateCalculate.usf
+++ b/Engine/Shaders/Private/VariableRateShading/VRSShadingRateCalculate.usf
@@ -192,88 +192,4 @@ void CalculateShadingRateImage(
 #endif
 }
 
-#endif /* COMPUTE_SHADER */
-
-RWTexture2D<float4> SceneColorOut;
-Texture2D<uint> VariableRateShadingTextureIn;
-float DynamicResolutionScale;
-
-uint SampleVariableRateShadingTexture(float2 UV)
-{
-	uint2 VRSDimension;
-	VariableRateShadingTextureIn.GetDimensions(VRSDimension.x, VRSDimension.y);
-
-	int2 VRSIndex = UV * VRSDimension;
-
-
-#if COMPUTE_SHADER
-	return VariableRateShadingTextureIn[VRSIndex].r;
-#else
-	return VariableRateShadingTextureIn.Load(int3(VRSIndex,0)).r;
-#endif
-}
-
-
-float4 PreviewVariableRateShadingTexture_Impl(float2 UV)
-{
-	uint VSRValue = SampleVariableRateShadingTexture(UV) & D3D12_SHADING_RATE_2X2;
-
-	float4 FilterColor = float4(1, 0, 0, 0);
-
-	if(VSRValue > 0)
-	{
-		if(VSRValue == D3D12_SHADING_RATE_1X2)
-		{
-			FilterColor = float4(1, 1, 0, 0);
-		}
-		else if(VSRValue == D3D12_SHADING_RATE_2X1)
-		{
-			FilterColor = float4(0, 0, 1, 0);
-		}
-		else if(VSRValue == D3D12_SHADING_RATE_2X2)
-		{
-			FilterColor = float4(0, 1, 0, 0);
-		}
-	}
-	FilterColor.a = 0.5f;
-	return FilterColor;
-}
-
-
-void PreviewVariableRateShadingTexturePS(
-	noperspective float4 UVAndScreenPos : TEXCOORD0,
-	out float4 OutColor : SV_Target0)
-{
-	// The vertex shader has already scaled this by the DynamicResolutionScale,
-	// so we don't have to do this in the Pixel Shader as we do in the CS.
-	float2 UVLookup = UVAndScreenPos.xy;
-
-	OutColor = PreviewVariableRateShadingTexture_Impl(UVLookup);
-}
-
-#if COMPUTE_SHADER
-
-[numthreads(THREADGROUP_SIZEX, THREADGROUP_SIZEY, 1)]
-void PreviewVariableRateShadingTextureCS(
-	uint3 DispatchThreadId : SV_DispatchThreadID)
-{
-	uint2 PhysicalDimension;
-	SceneColorOut.GetDimensions(PhysicalDimension.x, PhysicalDimension.y);
-
-	uint2 OutputCoord = DispatchThreadId.xy + ViewRect.xy;
-	if(OutputCoord.x >= PhysicalDimension.x || OutputCoord.y >= PhysicalDimension.y)
-	{
-		return;
-	}
-
-	float2 UV = (ViewRect.xy + DispatchThreadId.xy) * DynamicResolutionScale / float2(PhysicalDimension);
-	
-	float4 FilterColor = PreviewVariableRateShadingTexture_Impl(UV);
-
-	if (all (FilterColor != float4(0, 0, 0, 0)))
-	{
-		SceneColorOut[OutputCoord] = FilterColor;
-	}
-}
-
-#endif
\ No newline at end of file
+#endif /* COMPUTE_SHADER */
\ No newline at end of file
diff --git a/Engine/Shaders/Private/VariableRateShading/VRSShadingRateCommon.ush b/Engine/Shaders/Private/VariableRateShading/VRSShadingRateCommon.ush
--- a/Engine/Shaders/Private/VariableRateShading/VRSShadingRateCommon.ush
+++ b/Engine/Shaders/Private/VariableRateShading/VRSShadingRateCommon.ush
@@ -7,6 +7,10 @@
 #define D3D12_SHADING_RATE_1X2 0x1
 #define D3D12_SHADING_RATE_2X1 0x4
 #define D3D12_SHADING_RATE_2X2 0x5
+#define D3D12_SHADING_RATE_2X4 0x6
+#define D3D12_SHADING_RATE_4X2 0x9
+#define D3D12_SHADING_RATE_4X4 0xa
+
 
 #define MAX_D3D12_SHADING_RATE_BITS 4
 #define CONSERVATIVE_SHADING_RATE_SHIFT MAX_D3D12_SHADING_RATE_BITS
\ No newline at end of file
diff --git a/Engine/Shaders/Private/VariableRateShading/VRSShadingRatePreview.usf b/Engine/Shaders/Private/VariableRateShading/VRSShadingRatePreview.usf
new file mode 100644
--- /dev/null
+++ b/Engine/Shaders/Private/VariableRateShading/VRSShadingRatePreview.usf
@@ -0,0 +1,104 @@
+// Copyright Epic Games, Inc. All Rights Reserved.
+
+#include "VRSShadingRateCommon.ush"
+#include "../Common.ush"
+#include "../ColorSpace.ush"
+
+float4 ViewRect;
+
+RWTexture2D<float4> SceneColorOut;
+Texture2D<uint> VariableRateShadingTextureIn;
+float DynamicResolutionScale;
+
+uint SampleVariableRateShadingTexture(float2 UV)
+{
+	uint2 VRSDimension;
+	VariableRateShadingTextureIn.GetDimensions(VRSDimension.x, VRSDimension.y);
+
+	int2 VRSIndex = UV * VRSDimension;
+
+
+#if COMPUTE_SHADER
+	return VariableRateShadingTextureIn[VRSIndex].r;
+#else
+	return VariableRateShadingTextureIn.Load(int3(VRSIndex,0)).r;
+#endif
+}
+
+
+float4 PreviewVariableRateShadingTexture_Impl(float2 UV)
+{
+	uint VRSValue = SampleVariableRateShadingTexture(UV) & (15U);
+
+	float4 FilterColor = float4(1.0, 0.0, 0.0, 0); // Red (1x1 shading rate)
+
+	if (VRSValue > 0)
+	{
+		if (VRSValue == D3D12_SHADING_RATE_1X2)
+		{
+			FilterColor = float4(1.0, 0.5, 0.0, 0); // Orange
+		}
+		else if (VRSValue == D3D12_SHADING_RATE_2X1)
+		{
+			FilterColor = float4(1.0, 1.0, 0.0, 0); // Yellow
+		}
+		else if (VRSValue == D3D12_SHADING_RATE_2X2)
+		{
+			FilterColor = float4(0.0, 1.0, 0.0, 0); // Green
+		}
+		else if (VRSValue == D3D12_SHADING_RATE_2X4)
+		{
+			FilterColor = float4(0.0, 0.0, 1.0, 0); // Blue
+		}
+		else if (VRSValue == D3D12_SHADING_RATE_4X2)
+		{
+			FilterColor = float4(1.0, 0.0, 1.0, 0); // Purple
+		}
+		else // if (VRSValue == D3D12_SHADING_RATE_4X4)
+		{
+			FilterColor = float4(1.0, 0.5, 0.8, 0); // Pink
+		}
+	}
+	
+	FilterColor.a = 0.5f;
+	return FilterColor;
+}
+
+
+void PreviewVariableRateShadingTexturePS(
+	noperspective float4 UVAndScreenPos : TEXCOORD0,
+	out float4 OutColor : SV_Target0)
+{
+	// The vertex shader has already scaled this by the DynamicResolutionScale,
+	// so we don't have to do this in the Pixel Shader as we do in the CS.
+	float2 UVLookup = UVAndScreenPos.xy;
+
+	OutColor = PreviewVariableRateShadingTexture_Impl(UVLookup);
+}
+
+#if COMPUTE_SHADER
+
+[numthreads(THREADGROUP_SIZEX, THREADGROUP_SIZEY, 1)]
+void PreviewVariableRateShadingTextureCS(
+	uint3 DispatchThreadId : SV_DispatchThreadID)
+{
+	uint2 PhysicalDimension;
+	SceneColorOut.GetDimensions(PhysicalDimension.x, PhysicalDimension.y);
+
+	uint2 OutputCoord = DispatchThreadId.xy + ViewRect.xy;
+	if(OutputCoord.x >= PhysicalDimension.x || OutputCoord.y >= PhysicalDimension.y)
+	{
+		return;
+	}
+
+	float2 UV = (ViewRect.xy + DispatchThreadId.xy) * DynamicResolutionScale / float2(PhysicalDimension);
+	
+	float4 FilterColor = PreviewVariableRateShadingTexture_Impl(UV);
+
+	if (all (FilterColor != float4(0, 0, 0, 0)))
+	{
+		SceneColorOut[OutputCoord] = FilterColor;
+	}
+}
+
+#endif
\ No newline at end of file
diff --git a/Engine/Source/Runtime/Renderer/Private/DeferredShadingRenderer.cpp b/Engine/Source/Runtime/Renderer/Private/DeferredShadingRenderer.cpp
--- a/Engine/Source/Runtime/Renderer/Private/DeferredShadingRenderer.cpp
+++ b/Engine/Source/Runtime/Renderer/Private/DeferredShadingRenderer.cpp
@@ -4172,7 +4172,7 @@ void FDeferredShadingSceneRenderer::Render(FRDGBuilder& GraphBuilder)
 
 	if (ViewFamily.bResolveScene && ViewFamilyTexture)
 	{
-		GVRSImageManager.CASDebugPreview(GraphBuilder, ViewFamily, ViewFamilyTexture);
+		GVRSImageManager.DrawDebugPreview(GraphBuilder, ViewFamily, ViewFamilyTexture);
 	}
 
 	GEngine->GetPostRenderDelegateEx().Broadcast(GraphBuilder);
diff --git a/Engine/Source/Runtime/Renderer/Private/VariableRateShading/ContrastAdaptiveImageGenerator.cpp b/Engine/Source/Runtime/Renderer/Private/VariableRateShading/ContrastAdaptiveImageGenerator.cpp
--- a/Engine/Source/Runtime/Renderer/Private/VariableRateShading/ContrastAdaptiveImageGenerator.cpp
+++ b/Engine/Source/Runtime/Renderer/Private/VariableRateShading/ContrastAdaptiveImageGenerator.cpp
@@ -129,13 +129,15 @@ FAutoConsoleVariableRef CVarVRSDebugForceRate(
 	GVRSDebugForceRate,
 	TEXT("-1 : None, 0 : Force 1x1, 1 : Force 1x2, 4 : Force 2x1, 5: Force 2x2"));
 
-TAutoConsoleVariable<int32> CVarVRSPreview(
-	TEXT("r.VRS.Preview"),
-	0,
-	TEXT("Show a debug visualiation of the SRI texture.")
-	TEXT("0 - off, 1 - the SRI texture, 2- the conservative SRI texture, 3 - the unscaled SRI texture"),
+TAutoConsoleVariable<int32> CVarCASPreviewType(
+	TEXT("r.VRS.ContrastAdaptiveShading.PreviewType"),
+	1,
+	TEXT("Include CAS in the VRS debug overlay.")
+	TEXT("0 - off, 1 - the SRI texture, 2- the conservative SRI texture"),
 	ECVF_RenderThreadSafe);
 
+
+
 /**
  * Shaders
  */
@@ -253,86 +255,6 @@ class FRescaleVariableRateShadingCS : public FGlobalShader
 };
 IMPLEMENT_GLOBAL_SHADER(FRescaleVariableRateShadingCS, "/Engine/Private/VariableRateShading/VRSShadingRateReproject.usf", "RescaleVariableRateShading", SF_Compute);
 
-class FDebugVariableRateShadingCS : public FGlobalShader
-{
-	DECLARE_GLOBAL_SHADER(FDebugVariableRateShadingCS);
-	SHADER_USE_PARAMETER_STRUCT(FDebugVariableRateShadingCS, FGlobalShader);
-
-	static constexpr uint32 ThreadGroupSize = 8;
-
-	static bool ShouldCompilePermutation(const FGlobalShaderPermutationParameters& Parameters)
-	{
-		return IsFeatureLevelSupported(Parameters.Platform, ERHIFeatureLevel::SM5);
-	}
-
-	BEGIN_SHADER_PARAMETER_STRUCT(FParameters, )
-		SHADER_PARAMETER_RDG_TEXTURE(Texture2D<uint>, VariableRateShadingTextureIn)
-		SHADER_PARAMETER(FVector4f, ViewRect)
-		SHADER_PARAMETER(float, DynamicResolutionScale)
-		SHADER_PARAMETER_RDG_TEXTURE_UAV(RWTexture2D<float4>, SceneColorOut)
-	END_SHADER_PARAMETER_STRUCT()
-
-	static void ModifyCompilationEnvironment(const FGlobalShaderPermutationParameters& Parameters, FShaderCompilerEnvironment& OutEnvironment)
-	{
-		FGlobalShader::ModifyCompilationEnvironment(Parameters, OutEnvironment);
-		OutEnvironment.SetDefine(TEXT("THREADGROUP_SIZEX"), ThreadGroupSize);
-		OutEnvironment.SetDefine(TEXT("THREADGROUP_SIZEY"), ThreadGroupSize);
-		OutEnvironment.SetDefine(TEXT("COMPUTE_SHADER"), 1);
-	}
-
-	static void InitParameters(
-		FParameters& Parameters,
-		FRDGTextureRef VariableRateShadingTexture,
-		const FIntRect& ViewRect,
-		float DynamicResolutionScale,
-		FRDGTextureUAVRef SceneColorUAV)
-	{
-		Parameters.VariableRateShadingTextureIn = VariableRateShadingTexture;
-		Parameters.ViewRect = FVector4f(ViewRect.Min.X, ViewRect.Min.Y, ViewRect.Max.X, ViewRect.Max.Y);
-		Parameters.DynamicResolutionScale = DynamicResolutionScale;
-		Parameters.SceneColorOut = SceneColorUAV;
-	}
-};
-
-IMPLEMENT_GLOBAL_SHADER(FDebugVariableRateShadingCS, "/Engine/Private/VariableRateShading/VRSShadingRateCalculate.usf", "PreviewVariableRateShadingTextureCS", SF_Compute);
-
-//---------------------------------------------------------------------------------------------
-using FDebugVariableRateShadingVS = FScreenPassVS;
-
-//---------------------------------------------------------------------------------------------
-class FDebugVariableRateShadingPS : public FGlobalShader
-{
-	DECLARE_GLOBAL_SHADER(FDebugVariableRateShadingPS);
-	SHADER_USE_PARAMETER_STRUCT(FDebugVariableRateShadingPS, FGlobalShader);
-
-	static bool ShouldCompilePermutation(const FGlobalShaderPermutationParameters& Parameters)
-	{
-		return IsFeatureLevelSupported(Parameters.Platform, ERHIFeatureLevel::SM5);
-	}
-
-	static void ModifyCompilationEnvironment(const FGlobalShaderPermutationParameters& Parameters, FShaderCompilerEnvironment& OutEnvironment)
-	{
-		FGlobalShader::ModifyCompilationEnvironment(Parameters, OutEnvironment);
-		OutEnvironment.SetDefine(TEXT("COMPUTE_SHADER"), 0);
-	}
-
-	BEGIN_SHADER_PARAMETER_STRUCT(FParameters, )
-		SHADER_PARAMETER_RDG_TEXTURE(Texture2D<uint>, VariableRateShadingTextureIn)
-		RENDER_TARGET_BINDING_SLOTS()
-	END_SHADER_PARAMETER_STRUCT()
-		
-	static void InitParameters(
-			FParameters& Parameters,
-			FRDGTextureRef VariableRateShadingTexture,
-			FRDGTextureRef OutputSceneColor)
-	{
-		Parameters.VariableRateShadingTextureIn = VariableRateShadingTexture;
-		Parameters.RenderTargets[0] = FRenderTargetBinding(OutputSceneColor, ERenderTargetLoadAction::ELoad);
-	}
-};
-
-IMPLEMENT_GLOBAL_SHADER(FDebugVariableRateShadingPS, "/Engine/Private/VariableRateShading/VRSShadingRateCalculate.usf", "PreviewVariableRateShadingTexturePS", SF_Pixel);
-
 /**
  * Helper Functions and Structures
  */
@@ -631,98 +553,6 @@ void AddPrepareImageBasedVRSPass(
 	}
 }
 
-void FContrastAdaptiveImageGenerator::VRSDebugPreview(FRDGBuilder& GraphBuilder, const FSceneViewFamily& ViewFamily, FRDGTextureRef OutputSceneColor)
-{
-	RDG_EVENT_SCOPE(GraphBuilder, "VariableRateShading");
-	for (int32 ViewIndex = 0; ViewIndex < ViewFamily.Views.Num(); ViewIndex++)
-	{
-		check(ViewFamily.Views[ViewIndex]->bIsViewInfo);
-		const FViewInfo& View = *static_cast<const FViewInfo*>(ViewFamily.Views[ViewIndex]);
-		RDG_GPU_MASK_SCOPE(GraphBuilder, View.GPUMask);
-		RDG_EVENT_SCOPE_CONDITIONAL(GraphBuilder, ViewFamily.Views.Num() > 1, "View%d", ViewIndex);
-
-		//------------------------------------------------------------------------------------------------
-		// Do some sanity checks for early out
-		if (!OutputSceneColor)
-		{
-			return;
-		}
-
-		ESRIPreviewType::Type PreviewType = static_cast<ESRIPreviewType::Type>(CVarVRSPreview.GetValueOnRenderThread());
-
-		if (PreviewType == ESRIPreviewType::Off || !FVRSTextures::IsInitialized(GraphBuilder))
-		{
-			return;
-		}
-
-		const FVRSTextures& VRSTextures = FVRSTextures::Get(GraphBuilder);
-		FRDGTextureRef PreviewTexture = nullptr;
-		switch (PreviewType)
-		{
-		case ESRIPreviewType::BeforeReprojection:
-			PreviewTexture = VRSTextures.ConstructedSRI;
-			break;
-		case ESRIPreviewType::Projected:
-			PreviewTexture = VRSTextures.ScaledSRI;
-			break;
-		case ESRIPreviewType::ProjectedConservative:
-			PreviewTexture = VRSTextures.ScaledConservativeSRI;
-			break;
-		};
-
-		if (!PreviewTexture)
-		{
-			// We never rendered to this texture this frame, so aborting
-			return;
-		}
-		// Complete early out sanity checks
-		//------------------------------------------------------------------------------------------------
-
-		auto& RHICmdList = GraphBuilder.RHICmdList;
-
-		SCOPED_DRAW_EVENT(RHICmdList, VRSDebugPreview);
-
-		bool bUseRescaledTexture = PreviewType != ESRIPreviewType::BeforeReprojection;
-
-		FIntRect SrcViewRect = bUseRescaledTexture ? View.ViewRect : GetPostProcessOutputRect(View);
-
-		const FIntRect& DestViewRect = View.UnscaledViewRect;
-
-		TShaderMapRef<FDebugVariableRateShadingVS> VertexShader(View.ShaderMap);
-		TShaderMapRef<FDebugVariableRateShadingPS> PixelShader(View.ShaderMap);
-
-		auto* PassParameters = GraphBuilder.AllocParameters<FDebugVariableRateShadingPS::FParameters>();
-
-		FDebugVariableRateShadingPS::InitParameters(
-			*PassParameters,
-			PreviewTexture,
-			OutputSceneColor);
-
-		FRHIBlendState* BlendState =
-			TStaticBlendState<CW_RGB, BO_Add, BF_SourceAlpha, BF_InverseSource1Alpha>::GetRHI();
-		FRHIDepthStencilState* DepthStencilState = FScreenPassPipelineState::FDefaultDepthStencilState::GetRHI();
-
-		EScreenPassDrawFlags DrawFlags = EScreenPassDrawFlags::AllowHMDHiddenAreaMask;
-
-		FIntRect ScaledSrcRect = FIntRect::DivideAndRoundUp(SrcViewRect, FVariableRateShadingImageManager::GetSRITileSize());
-
-		const FScreenPassTextureViewport InputViewport = FScreenPassTextureViewport(PreviewTexture, ScaledSrcRect);
-		const FScreenPassTextureViewport OutputViewport(OutputSceneColor, DestViewRect);
-
-		AddDrawScreenPass(
-			GraphBuilder,
-			RDG_EVENT_NAME("Display Debug : %s", ESRIPreviewType::GetName(PreviewType)),
-			View,
-			OutputViewport,
-			InputViewport,
-			VertexShader,
-			PixelShader,
-			BlendState,
-			DepthStencilState,
-			PassParameters,
-			DrawFlags);
-	}
-}
 
 /**
  * Interface Functions
@@ -781,3 +611,27 @@ bool FContrastAdaptiveImageGenerator::IsEnabledForView(const FSceneView& View) c
 	return IsContrastAdaptiveShadingEnabled() && !View.bIsSceneCapture && bCompatibleWithOutputType;
 }
 
+FRDGTextureRef FContrastAdaptiveImageGenerator::GetDebugImage(FRDGBuilder& GraphBuilder, const FViewInfo& ViewInfo)
+{
+	if (!FVRSTextures::IsInitialized(GraphBuilder))
+	{
+		return nullptr;
+	}
+
+	const FVRSTextures& VRSTextures = FVRSTextures::Get(GraphBuilder);
+
+	ESRIPreviewType::Type PreviewType = static_cast<ESRIPreviewType::Type>(CVarCASPreviewType.GetValueOnRenderThread());
+
+	switch (PreviewType)
+	{
+	case ESRIPreviewType::BeforeReprojection:
+		return VRSTextures.ConstructedSRI;
+		break;
+	case ESRIPreviewType::Projected:
+		return VRSTextures.ScaledSRI;
+		break;
+	}
+
+	return nullptr;
+}
+
diff --git a/Engine/Source/Runtime/Renderer/Private/VariableRateShading/ContrastAdaptiveImageGenerator.h b/Engine/Source/Runtime/Renderer/Private/VariableRateShading/ContrastAdaptiveImageGenerator.h
--- a/Engine/Source/Runtime/Renderer/Private/VariableRateShading/ContrastAdaptiveImageGenerator.h
+++ b/Engine/Source/Runtime/Renderer/Private/VariableRateShading/ContrastAdaptiveImageGenerator.h
@@ -20,7 +20,6 @@ public:
 	{
 		return FVariableRateShadingImageManager::EVRSSourceType::ContrastAdaptiveShading;
 	}
-
-	void VRSDebugPreview(FRDGBuilder& GraphBuilder, const FSceneViewFamily& ViewFamily, FRDGTextureRef OutputSceneColor);
+	virtual FRDGTextureRef GetDebugImage(FRDGBuilder& GraphBuilder, const FViewInfo& ViewInfo) override;
 };
 
diff --git a/Engine/Source/Runtime/Renderer/Private/VariableRateShading/FixedFoveationImageGenerator.cpp b/Engine/Source/Runtime/Renderer/Private/VariableRateShading/FixedFoveationImageGenerator.cpp
--- a/Engine/Source/Runtime/Renderer/Private/VariableRateShading/FixedFoveationImageGenerator.cpp
+++ b/Engine/Source/Runtime/Renderer/Private/VariableRateShading/FixedFoveationImageGenerator.cpp
@@ -30,6 +30,14 @@ static TAutoConsoleVariable<int32> CVarFixedFoveationDynamic(
 	TEXT(" 1: Enabled\n"),
 	ECVF_RenderThreadSafe);
 
+TAutoConsoleVariable<int32> CVarFoveatedPreview(
+	TEXT("vr.VRS.HMDFixedFoveationPreview"),
+	1,
+	TEXT("Include foveated VRS in the VRS debug overlay.")
+	TEXT(" 0: Disabled;\n")
+	TEXT(" 1: Enabled (default)\n"),
+	ECVF_RenderThreadSafe);
+
 
 /* Image generation parameters, set up by Prepare() */
 
@@ -217,6 +225,11 @@ bool FFixedFoveationImageGenerator::IsEnabledForView(const FSceneView& View) con
 	return bStereoRendering && CVarFixedFoveationLevel.GetValueOnRenderThread() > 0;
 }
 
+FRDGTextureRef FFixedFoveationImageGenerator::GetDebugImage(FRDGBuilder& GraphBuilder, const FViewInfo& ViewInfo)
+{
+	return CachedImage;
+}
+
 float FFixedFoveationImageGenerator::GetDynamicVRSAmount()
 {
 	const float		kUpdateIncrement = 0.1f;
diff --git a/Engine/Source/Runtime/Renderer/Private/VariableRateShading/FixedFoveationImageGenerator.h b/Engine/Source/Runtime/Renderer/Private/VariableRateShading/FixedFoveationImageGenerator.h
--- a/Engine/Source/Runtime/Renderer/Private/VariableRateShading/FixedFoveationImageGenerator.h
+++ b/Engine/Source/Runtime/Renderer/Private/VariableRateShading/FixedFoveationImageGenerator.h
@@ -20,6 +20,7 @@ public:
 	{
 		return FVariableRateShadingImageManager::EVRSSourceType::FixedFoveation;
 	}
+	virtual FRDGTextureRef GetDebugImage(FRDGBuilder& GraphBuilder, const FViewInfo& ViewInfo) override;
 private:
 	FRDGTextureRef CachedImage = nullptr;
 	struct FDynamicVRSData
diff --git a/Engine/Source/Runtime/Renderer/Private/VariableRateShading/VariableRateShadingImageManager.cpp b/Engine/Source/Runtime/Renderer/Private/VariableRateShading/VariableRateShadingImageManager.cpp
--- a/Engine/Source/Runtime/Renderer/Private/VariableRateShading/VariableRateShadingImageManager.cpp
+++ b/Engine/Source/Runtime/Renderer/Private/VariableRateShading/VariableRateShadingImageManager.cpp
@@ -22,6 +22,17 @@
 
 TGlobalResource<FVariableRateShadingImageManager> GVRSImageManager;
 
+/**
+ * CVars
+ */
+
+TAutoConsoleVariable<int32> CVarVRSPreview(
+	TEXT("r.VRS.Preview"),
+	0,
+	TEXT("Show a debug visualiation of the VRS shading rate image texture.")
+	TEXT("0 - off, 1 - on"),
+	ECVF_RenderThreadSafe);
+
 /**
  * Shaders
  */
@@ -61,7 +72,87 @@ public:
 
 };
 
-IMPLEMENT_GLOBAL_SHADER(FCombineShadingRateTexturesCS, "/Engine/Private/VariableRateShading/VRSShadingRateCombine.usf", "CombineShadingRateTextures", SF_Compute);
+IMPLEMENT_GLOBAL_SHADER(FCombineShadingRateTexturesCS, "/Engine/Private/VariableRateShading/VRSShadingRateCombine.usf", "CombineShadingRateTextures", SF_Compute);\
+
+class FDebugVariableRateShadingCS : public FGlobalShader
+{
+	DECLARE_GLOBAL_SHADER(FDebugVariableRateShadingCS);
+	SHADER_USE_PARAMETER_STRUCT(FDebugVariableRateShadingCS, FGlobalShader);
+
+	static constexpr uint32 ThreadGroupSize = 8;
+
+	static bool ShouldCompilePermutation(const FGlobalShaderPermutationParameters& Parameters)
+	{
+		return IsFeatureLevelSupported(Parameters.Platform, ERHIFeatureLevel::SM5);
+	}
+
+	BEGIN_SHADER_PARAMETER_STRUCT(FParameters, )
+		SHADER_PARAMETER_RDG_TEXTURE(Texture2D<uint>, VariableRateShadingTextureIn)
+		SHADER_PARAMETER(FVector4f, ViewRect)
+		SHADER_PARAMETER(float, DynamicResolutionScale)
+		SHADER_PARAMETER_RDG_TEXTURE_UAV(RWTexture2D<float4>, SceneColorOut)
+		END_SHADER_PARAMETER_STRUCT()
+
+		static void ModifyCompilationEnvironment(const FGlobalShaderPermutationParameters& Parameters, FShaderCompilerEnvironment& OutEnvironment)
+	{
+		FGlobalShader::ModifyCompilationEnvironment(Parameters, OutEnvironment);
+		OutEnvironment.SetDefine(TEXT("THREADGROUP_SIZEX"), ThreadGroupSize);
+		OutEnvironment.SetDefine(TEXT("THREADGROUP_SIZEY"), ThreadGroupSize);
+		OutEnvironment.SetDefine(TEXT("COMPUTE_SHADER"), 1);
+	}
+
+	static void InitParameters(
+		FParameters& Parameters,
+		FRDGTextureRef VariableRateShadingTexture,
+		const FIntRect& ViewRect,
+		float DynamicResolutionScale,
+		FRDGTextureUAVRef SceneColorUAV)
+	{
+		Parameters.VariableRateShadingTextureIn = VariableRateShadingTexture;
+		Parameters.ViewRect = FVector4f(ViewRect.Min.X, ViewRect.Min.Y, ViewRect.Max.X, ViewRect.Max.Y);
+		Parameters.DynamicResolutionScale = DynamicResolutionScale;
+		Parameters.SceneColorOut = SceneColorUAV;
+	}
+};
+
+IMPLEMENT_GLOBAL_SHADER(FDebugVariableRateShadingCS, "/Engine/Private/VariableRateShading/VRSShadingRatePreview.usf", "PreviewVariableRateShadingTextureCS", SF_Compute);
+
+//---------------------------------------------------------------------------------------------
+using FDebugVariableRateShadingVS = FScreenPassVS;
+
+//---------------------------------------------------------------------------------------------
+class FDebugVariableRateShadingPS : public FGlobalShader
+{
+	DECLARE_GLOBAL_SHADER(FDebugVariableRateShadingPS);
+	SHADER_USE_PARAMETER_STRUCT(FDebugVariableRateShadingPS, FGlobalShader);
+
+	static bool ShouldCompilePermutation(const FGlobalShaderPermutationParameters& Parameters)
+	{
+		return IsFeatureLevelSupported(Parameters.Platform, ERHIFeatureLevel::SM5);
+	}
+
+	static void ModifyCompilationEnvironment(const FGlobalShaderPermutationParameters& Parameters, FShaderCompilerEnvironment& OutEnvironment)
+	{
+		FGlobalShader::ModifyCompilationEnvironment(Parameters, OutEnvironment);
+		OutEnvironment.SetDefine(TEXT("COMPUTE_SHADER"), 0);
+	}
+
+	BEGIN_SHADER_PARAMETER_STRUCT(FParameters, )
+		SHADER_PARAMETER_RDG_TEXTURE(Texture2D<uint>, VariableRateShadingTextureIn)
+		RENDER_TARGET_BINDING_SLOTS()
+		END_SHADER_PARAMETER_STRUCT()
+
+		static void InitParameters(
+			FParameters& Parameters,
+			FRDGTextureRef VariableRateShadingTexture,
+			FRDGTextureRef OutputSceneColor)
+	{
+		Parameters.VariableRateShadingTextureIn = VariableRateShadingTexture;
+		Parameters.RenderTargets[0] = FRenderTargetBinding(OutputSceneColor, ERenderTargetLoadAction::ELoad);
+	}
+};
+
+IMPLEMENT_GLOBAL_SHADER(FDebugVariableRateShadingPS, "/Engine/Private/VariableRateShading/VRSShadingRatePreview.usf", "PreviewVariableRateShadingTexturePS", SF_Pixel);
 
 
 /**
@@ -117,35 +208,19 @@ FIntPoint FVariableRateShadingImageManager::GetSRITileSize()
 	return FIntPoint(GRHIVariableRateShadingImageTileMinWidth, GRHIVariableRateShadingImageTileMinHeight);
 }
 
-FRDGTextureRef FVariableRateShadingImageManager::GetVariableRateShadingImage(FRDGBuilder& GraphBuilder, const FViewInfo& ViewInfo, FVariableRateShadingImageManager::EVRSPassType PassType,
-	const TArray<TRefCountPtr<IPooledRenderTarget>>* ExternalVRSSources, FVariableRateShadingImageManager::EVRSSourceType VRSTypesToExclude)
+FRDGTextureRef FVariableRateShadingImageManager::CombineShadingRateImages(FRDGBuilder& GraphBuilder, const FViewInfo& ViewInfo, TArray<FRDGTextureRef> Sources)
 {
-	// If the view doesn't support VRS, bail immediately
-	if (!IsVRSCompatibleWithView(ViewInfo))
+	// If we have more than one source, combine the first available two
+	// TODO: Support combining more textures
+	if (Sources.Num() < 1)
 	{
 		return nullptr;
 	}
-
-	// Otherwise collate all internal sources
-	TArray<FRDGTextureRef> InternalVRSSources;
-
-	for (TUniquePtr<IVariableRateShadingImageGenerator>& Generator : ImageGenerators)
+	else if (Sources.Num() == 1)
 	{
-		FRDGTextureRef Image = nullptr;
-		if (Generator->IsEnabledForView(ViewInfo) && !EnumHasAnyFlags(VRSTypesToExclude, Generator->GetType()))
-		{
-			Image = Generator->GetImage(GraphBuilder, ViewInfo, PassType);
-		}
-
-		if (Image)
-		{
-			InternalVRSSources.Add(Image);
-		}
+		return Sources[0];
 	}
-
-	// If we have more than one internal source, combine the first available two
-	// TODO: Support combining more textures
-	if (InternalVRSSources.Num() > 1)
+	else
 	{
 		FIntPoint ViewSize = ViewInfo.UnscaledViewRect.Scale(ViewInfo.Family->SecondaryViewFraction).Size();
 		const FIntPoint TileSize = GetSRITileSize();
@@ -160,8 +235,8 @@ FRDGTextureRef FVariableRateShadingImageManager::GetVariableRateShadingImage(FRD
 		FRDGTextureRef CombinedShadingRateTexture = GraphBuilder.CreateTexture(Desc, TEXT("CombinedShadingRateTexture"));
 
 		FCombineShadingRateTexturesCS::FParameters* PassParameters = GraphBuilder.AllocParameters<FCombineShadingRateTexturesCS::FParameters>();
-		PassParameters->SourceTexture0 = InternalVRSSources[0];
-		PassParameters->SourceTexture1 = InternalVRSSources[1];
+		PassParameters->SourceTexture0 = Sources[0];
+		PassParameters->SourceTexture1 = Sources[1];
 		PassParameters->RWOutputTexture = GraphBuilder.CreateUAV(CombinedShadingRateTexture);
 
 		TShaderMapRef<FCombineShadingRateTexturesCS> ComputeShader(GetGlobalShaderMap(GMaxRHIFeatureLevel));
@@ -176,9 +251,39 @@ FRDGTextureRef FVariableRateShadingImageManager::GetVariableRateShadingImage(FRD
 
 		return CombinedShadingRateTexture;
 	}
-	else if (InternalVRSSources.Num() == 1)
+
+}
+
+FRDGTextureRef FVariableRateShadingImageManager::GetVariableRateShadingImage(FRDGBuilder& GraphBuilder, const FViewInfo& ViewInfo, FVariableRateShadingImageManager::EVRSPassType PassType,
+	const TArray<TRefCountPtr<IPooledRenderTarget>>* ExternalVRSSources, FVariableRateShadingImageManager::EVRSSourceType VRSTypesToExclude)
+{
+	// If the view doesn't support VRS, bail immediately
+	if (!IsVRSCompatibleWithView(ViewInfo))
 	{
-		return InternalVRSSources[0];
+		return nullptr;
+	}
+
+	// Otherwise collate all internal sources
+	TArray<FRDGTextureRef> InternalVRSSources;
+
+	for (TUniquePtr<IVariableRateShadingImageGenerator>& Generator : ImageGenerators)
+	{
+		FRDGTextureRef Image = nullptr;
+		if (Generator->IsEnabledForView(ViewInfo) && !EnumHasAnyFlags(VRSTypesToExclude, Generator->GetType()))
+		{
+			Image = Generator->GetImage(GraphBuilder, ViewInfo, PassType);
+		}
+
+		if (Image)
+		{
+			InternalVRSSources.Add(Image);
+		}
+	}
+
+	// If we have more than one internal source, combine the first available two
+	if (InternalVRSSources.Num())
+	{
+		return CombineShadingRateImages(GraphBuilder, ViewInfo, InternalVRSSources);
 	}
 
 	// Fall back on external sources only if we have no internal ones
@@ -270,16 +375,89 @@ TRefCountPtr<IPooledRenderTarget> FVariableRateShadingImageManager::GetMobileVar
 	return MobileHMDFixedFoveationOverrideImage;
 }
 
-// Temporary passthrough for CAS debug overlay, pending that functionality being moved to the manager
-void FVariableRateShadingImageManager::CASDebugPreview(FRDGBuilder& GraphBuilder, const FSceneViewFamily& ViewFamily, FRDGTextureRef OutputSceneColor)
+void FVariableRateShadingImageManager::DrawDebugPreview(FRDGBuilder& GraphBuilder, const FSceneViewFamily& ViewFamily, FRDGTextureRef OutputSceneColor)
 {
-	// Get the first CAS generator present and call its debug function
-	for (TUniquePtr<IVariableRateShadingImageGenerator>& Generator : ImageGenerators)
+
+	if (CVarVRSPreview.GetValueOnRenderThread() != 1 || !OutputSceneColor)
 	{
-		if (EnumHasAllFlags(Generator->GetType(), FVariableRateShadingImageManager::EVRSSourceType::ContrastAdaptiveShading))
+		return;
+	}
+
+	for (const FSceneView* View : ViewFamily.Views)
+	{
+		check(View->bIsViewInfo);
+		auto ViewInfo = static_cast<const FViewInfo*>(View);
+		if (IsVRSCompatibleWithView(*ViewInfo))
 		{
-			((FContrastAdaptiveImageGenerator*) Generator.Get())->VRSDebugPreview(GraphBuilder, ViewFamily, OutputSceneColor);
-			break;
+			// Collate debug images
+			TArray<FRDGTextureRef> InternalVRSSources;
+
+			for (TUniquePtr<IVariableRateShadingImageGenerator>& Generator : ImageGenerators)
+			{
+				FRDGTextureRef Image = nullptr;
+				if (Generator->IsEnabledForView(*View))
+				{
+					Image = Generator->GetDebugImage(GraphBuilder, *ViewInfo);
+				}
+
+				if (Image)
+				{
+					InternalVRSSources.Add(Image);
+				}
+			}
+
+			FRDGTextureRef PreviewTexture = CombineShadingRateImages(GraphBuilder, *ViewInfo, InternalVRSSources);
+			if (!PreviewTexture)
+			{
+				return;
+			}
+
+			// If we have an active debug image, render it as a preview overlay
+			auto& RHICmdList = GraphBuilder.RHICmdList;
+
+			SCOPED_DRAW_EVENT(RHICmdList, VRSDebugPreview);
+
+			FIntRect SrcViewRect = ViewInfo->ViewRect;
+
+			const FIntRect& DestViewRect = ViewInfo->UnscaledViewRect;
+
+			TShaderMapRef<FDebugVariableRateShadingVS> VertexShader(ViewInfo->ShaderMap);
+			TShaderMapRef<FDebugVariableRateShadingPS> PixelShader(ViewInfo->ShaderMap);
+
+			auto* PassParameters = GraphBuilder.AllocParameters<FDebugVariableRateShadingPS::FParameters>();
+
+			FDebugVariableRateShadingPS::InitParameters(
+				*PassParameters,
+				PreviewTexture,
+				OutputSceneColor);
+
+			FRHIBlendState* BlendState =
+				TStaticBlendState<CW_RGB, BO_Add, BF_SourceAlpha, BF_InverseSource1Alpha>::GetRHI();
+			FRHIDepthStencilState* DepthStencilState = FScreenPassPipelineState::FDefaultDepthStencilState::GetRHI();
+
+			EScreenPassDrawFlags DrawFlags = EScreenPassDrawFlags::AllowHMDHiddenAreaMask;
+
+			FIntRect ScaledSrcRect = FIntRect::DivideAndRoundUp(SrcViewRect, FVariableRateShadingImageManager::GetSRITileSize());
+
+			const FScreenPassTextureViewport InputViewport = FScreenPassTextureViewport(PreviewTexture, ScaledSrcRect);
+			const FScreenPassTextureViewport OutputViewport(OutputSceneColor, DestViewRect);
+
+			AddDrawScreenPass(
+				GraphBuilder,
+				RDG_EVENT_NAME("Display VRS Debug Preview"),
+				*ViewInfo,
+				OutputViewport,
+				InputViewport,
+				VertexShader,
+				PixelShader,
+				BlendState,
+				DepthStencilState,
+				PassParameters,
+				DrawFlags);
 		}
 	}
-}
+
+	
+
+	
+}
\ No newline at end of file
diff --git a/Engine/Source/Runtime/Renderer/Public/VariableRateShadingImageManager.h b/Engine/Source/Runtime/Renderer/Public/VariableRateShadingImageManager.h
--- a/Engine/Source/Runtime/Renderer/Public/VariableRateShadingImageManager.h
+++ b/Engine/Source/Runtime/Renderer/Public/VariableRateShadingImageManager.h
@@ -97,12 +97,13 @@ public:
 
 	static FIntPoint GetSRITileSize();
 
-	void CASDebugPreview(FRDGBuilder& GraphBuilder, const FSceneViewFamily& ViewFamily, FRDGTextureRef OutputSceneColor);
+	void DrawDebugPreview(FRDGBuilder& GraphBuilder, const FSceneViewFamily& ViewFamily, FRDGTextureRef OutputSceneColor);
 
 private:
 	TRefCountPtr<IPooledRenderTarget> MobileHMDFixedFoveationOverrideImage;
 	TArray<TUniquePtr<IVariableRateShadingImageGenerator>> ImageGenerators;
 
+	FRDGTextureRef CombineShadingRateImages(FRDGBuilder& GraphBuilder, const FViewInfo& ViewInfo, TArray<FRDGTextureRef> Sources);
 };
 
 ENUM_CLASS_FLAGS(FVariableRateShadingImageManager::EVRSSourceType);
@@ -126,4 +127,6 @@ public:
 	// Return bitmask of generator type
 	virtual FVariableRateShadingImageManager::EVRSSourceType GetType() const { return FVariableRateShadingImageManager::EVRSSourceType::None; };
 
+	// Get VRS image to be used w/ debug overlay
+	virtual FRDGTextureRef GetDebugImage(FRDGBuilder& GraphBuilder, const FViewInfo& ViewInfo) = 0;
 };
\ No newline at end of file
