Fix a bug of missing MobileBassPass UB binding in fog rendering.

diff --git a/Engine/Source/Runtime/Renderer/Private/MobileFogRendering.cpp b/Engine/Source/Runtime/Renderer/Private/MobileFogRendering.cpp
--- a/Engine/Source/Runtime/Renderer/Private/MobileFogRendering.cpp
+++ b/Engine/Source/Runtime/Renderer/Private/MobileFogRendering.cpp
@@ -14,6 +14,7 @@
 #include "ScenePrivate.h"
 #include "FogRendering.h"
 #include "ProfilingDebugging/RealtimeGPUProfiler.h"
+#include "MobileBasePassRendering.h"
 
 static TAutoConsoleVariable<int32> CVarPixelFogQuality(
 	TEXT("r.Mobile.PixelFogQuality"),
@@ -104,6 +105,7 @@ class FMobileFogPS : public FGlobalShader
 
 	BEGIN_SHADER_PARAMETER_STRUCT(FParameters, )
 		SHADER_PARAMETER_STRUCT_INCLUDE(FViewShaderParameters, View)
+		SHADER_PARAMETER_RDG_UNIFORM_BUFFER(FMobileBasePassUniformParameters, MobileBasePass)
 	END_SHADER_PARAMETER_STRUCT()
 };
 
