Fix crash due to unbound uniform buffer when using Per Instance Custom Data in deferred decals in Path Tracer.

diff --git a/Engine/Source/Runtime/Renderer/Private/RayTracing/RayTracingDecals.cpp b/Engine/Source/Runtime/Renderer/Private/RayTracing/RayTracingDecals.cpp
--- a/Engine/Source/Runtime/Renderer/Private/RayTracing/RayTracingDecals.cpp
+++ b/Engine/Source/Runtime/Renderer/Private/RayTracing/RayTracingDecals.cpp
@@ -197,6 +197,9 @@ public:
 		FMaterialShader::GetShaderBindings(Scene, FeatureLevel, MaterialRenderProxy, Material, ShaderBindings);
 
 		ShaderBindings.Add(GetUniformBufferParameter<FViewUniformShaderParameters>(), View.ViewUniformBuffer);
+		// Similar to DecalRendering::SetShader(...) uses GIdentityPrimitiveUniformBuffer
+		// We could potentially bind the actual primitive uniform buffer
+		ShaderBindings.Add(GetUniformBufferParameter<FPrimitiveUniformShaderParameters>(), GIdentityPrimitiveUniformBuffer);
 		ShaderBindings.Add(DecalParameter, DecalParameters);
 	}
 
