Support mesh decals in path tracer
- This initial version doesn't support blending of mesh decals. Instead it performs a simple alpha masking to provide limited overlap support.
- Mesh decals are assigned to a new layer of RayTracingScene.
- At each surface hit, a short ray is traced against the Decal TLAS, using the geonormal as direction, and the result is applied to the surface material properties.
    - The origin of the decal ray is biased into the surface since mesh decals might are sometimes coplanar or even inside the underlying mesh. The bias can be controlled using r.PathTracing.MeshDecalBias.
    - The length of the decal ray is 2 x MeshDecalBias.
- Mesh decals use the appropriate permutation of TRayTracingDecalMaterial based on whether the material requires alpha test.

diff --git a/Engine/Shaders/Private/PathTracing/PathTracingCore.ush b/Engine/Shaders/Private/PathTracing/PathTracingCore.ush
--- a/Engine/Shaders/Private/PathTracing/PathTracingCore.ush
+++ b/Engine/Shaders/Private/PathTracing/PathTracingCore.ush
@@ -54,6 +54,8 @@ uint SamplerType;
 uint VisualizeLightGrid;
 uint VisualizeDecalGrid;
 float DecalRoughnessCutoff;
+float MeshDecalRoughnessCutoff;
+float MeshDecalBias;
 float MaxPathIntensity;
 float MaxNormalBias;
 float FilterWidth;
@@ -65,6 +67,7 @@ RWTexture2D<float4> RadianceTexture;
 RWTexture2D<float4> AlbedoTexture;
 RWTexture2D<float4> NormalTexture;
 RaytracingAccelerationStructure TLAS;
+RaytracingAccelerationStructure DecalTLAS;
 uint SceneVisibleLightCount;
 
 Buffer<float> StartingExtinctionCoefficient;
@@ -369,6 +372,39 @@ FPathTracingPayload TraceTransparentRay(inout FPathState PathState, bool bIsCame
 		// Unpack the payload
 		FPathTracingPayload HitPayload = UnpackPathTracingPayloadMaterial(PackedPayload, PathState.Ray);
 
+		if (PathState.PathRoughness < MeshDecalRoughnessCutoff && HitPayload.IsDecalReceiver())
+		{
+			FRayDesc DecalRay;
+			DecalRay.Direction = HitPayload.WorldGeoNormal;
+			DecalRay.Origin = HitPayload.TranslatedWorldPos - DecalRay.Direction * MeshDecalBias;
+			DecalRay.TMin = 0.0f;
+			DecalRay.TMax = MeshDecalBias * 2;
+
+			FRayTracingDecalPayload DecalPayload = (FRayTracingDecalPayload)0;
+			DecalPayload.Params.SetInputTranslatedWorldPosition(HitPayload.TranslatedWorldPos);
+
+			TraceRay(
+				DecalTLAS,
+				RAY_FLAG_CULL_FRONT_FACING_TRIANGLES,
+				InstanceInclusionMask,
+				RAY_TRACING_SHADER_SLOT_MATERIAL,
+				RAY_TRACING_NUM_SHADER_SLOTS,
+				MissShaderIndex,
+				DecalRay.GetNativeDesc(),
+				DecalPayload);
+
+			if (!DecalPayload.IsMiss())
+			{
+				ApplyDecal(HitPayload, DecalPayload.Params);
+
+				if (ApproximateCaustics)
+				{
+					Payload.Roughness = max(Payload.Roughness, PathState.PathRoughness);
+				}
+				HitPayload.WorldNormal = AdjustShadingNormal(HitPayload.WorldNormal, HitPayload.WorldGeoNormal, PathState.Ray.Direction);
+			}
+		}
+
 #if PLATFORM_SUPPORTS_CALLABLE_SHADERS
 		if (PathState.PathRoughness < DecalRoughnessCutoff && HitPayload.IsDecalReceiver())
 		{
diff --git a/Engine/Shaders/Private/RayTracing/RayTracingCommon.ush b/Engine/Shaders/Private/RayTracing/RayTracingCommon.ush
--- a/Engine/Shaders/Private/RayTracing/RayTracingCommon.ush
+++ b/Engine/Shaders/Private/RayTracing/RayTracingCommon.ush
@@ -918,7 +918,7 @@ CHECK_RT_PAYLOAD_SIZE(FPackedMaterialClosestHitPayload)
 #define DECAL_WRITE_EMISSIVE_FLAG						(1 << 3)
 #define DECAL_WRITE_AMBIENT_OCCLUSION_FLAG				(1 << 4)
 
-#if IS_PAYLOAD_ENABLED(RT_PAYLOAD_TYPE_DECALS)
+#if IS_PAYLOAD_ENABLED(RT_PAYLOAD_TYPE_DECALS | RT_PAYLOAD_TYPE_MESH_DECALS)
 struct FDecalShaderParams
 {
 	float Data[15];
@@ -970,9 +970,20 @@ struct FDecalShaderParams
 	}
 	float3 GetEmissive() { return float3(Data[12], Data[13], Data[14]); }
 };
+#endif
+
+#if IS_PAYLOAD_ENABLED(RT_PAYLOAD_TYPE_DECALS)
 CHECK_RT_PAYLOAD_SIZE(FDecalShaderParams)
 #endif
 
+#if IS_PAYLOAD_ENABLED(RT_PAYLOAD_TYPE_MESH_DECALS)
+struct FRayTracingDecalPayload : FMinimalPayload
+{
+	FDecalShaderParams Params;
+};
+CHECK_RT_PAYLOAD_SIZE(FRayTracingDecalPayload)
+#endif
+
 ///////////////////////////////////////////////////////////////////////////////////////////////////
 // Tracing util functions
 
diff --git a/Engine/Shaders/Private/RayTracing/RayTracingDecalMaterialShader.usf b/Engine/Shaders/Private/RayTracing/RayTracingDecalMaterialShader.usf
--- a/Engine/Shaders/Private/RayTracing/RayTracingDecalMaterialShader.usf
+++ b/Engine/Shaders/Private/RayTracing/RayTracingDecalMaterialShader.usf
@@ -15,8 +15,6 @@
 #define DecalParams RayTracingDecalParameters.DecalParams
 #define DecalWriteFlags RayTracingDecalParameters.DecalWriteFlags
 
-// DECAL_PRIMITIVE informs material templates which functions to expose when rendering decals.
-#define DECAL_PRIMITIVE 1
 #include "/Engine/Generated/Material.ush"
 
 RAY_TRACING_ENTRY_CALLABLE(RayTracingDecalMaterialShader, FDecalShaderParams, Params)
@@ -96,3 +94,131 @@ RAY_TRACING_ENTRY_CALLABLE(RayTracingDecalMaterialShader, FDecalShaderParams, Pa
 
 	Params.SetFlags(DecalWriteFlags);
 }
+
+#if RAYHITGROUPSHADER
+
+#ifdef NEEDS_VERTEX_FACTORY_INTERPOLATION
+#undef NEEDS_VERTEX_FACTORY_INTERPOLATION
+#endif
+// Needed for VertexFactoryInterpolate to interpolate attributes from vertices to hit point
+#define NEEDS_VERTEX_FACTORY_INTERPOLATION 1
+
+#include "/Engine/Private/RayTracing/RayTracingHitGroupCommon.ush"
+
+#include "/Engine/Generated/VertexFactory.ush"
+
+#include "/Engine/Private/RayTracing/RayTracingCalcInterpolants.ush"
+
+RAY_TRACING_ENTRY_CLOSEST_HIT(RayTracingDecalMaterialCHS,
+	FRayTracingDecalPayload, Payload,
+	FRayTracingIntersectionAttributes, Attributes)
+{
+	Payload.HitT = RayTCurrent();
+
+	ResolvedView = ResolveView();
+
+	const float3 TranslatedWorldPosition = TranslatedWorldRayOrigin() + WorldRayDirection() * RayTCurrent();
+	const float4 SvPosition = mul(float4(TranslatedWorldPosition, 1.0f), ResolvedView.TranslatedWorldToClip);
+
+#if VF_SUPPORTS_RAYTRACING_PREPARE_MATERIAL_PIXEL_PARAMETERS
+	// this is a newer codepath that is both more flexible and allows for more direct calculation compared to the other codepath
+	// TODO: implement such a method for all vertex factories
+	float3 GeoNormal = 0;
+	FMaterialPixelParameters MaterialParameters = GetMaterialPixelParameters(TranslatedWorldRayOrigin(), WorldRayDirection(), RayTCurrent(), PrimitiveIndex(), Attributes, HitKind(), GeoNormal);
+#else
+	FVertexFactoryInterpolantsVSToPS Interpolants;
+	float3 GeoNormal = 0;
+	CalcInterpolants((FRayCone)0, Attributes, Interpolants, GeoNormal);
+
+	FMaterialPixelParameters MaterialParameters = GetMaterialPixelParameters(Interpolants, SvPosition);
+#endif
+	FPixelMaterialInputs PixelMaterialInputs;
+
+	{
+		const float4 ScreenPosition = SvPositionToResolvedScreenPosition(SvPosition);
+		const bool bIsFrontFace = HitKind() == HIT_KIND_TRIANGLE_FRONT_FACE;
+
+		// #dxr_todo: UE-72130 support world position offset 
+		// #if USE_WORLD_POSITION_EXCLUDING_SHADER_OFFSETS
+		// CalcMaterialParametersEx(MaterialParameters, PixelMaterialInputs, SvPosition, ScreenPosition, bIsFrontFace, TranslatedWorldPosition, BasePassInterpolants.PixelPositionExcludingWPO);
+		// #else
+		CalcMaterialParametersEx(MaterialParameters, PixelMaterialInputs, SvPosition, ScreenPosition, bIsFrontFace, TranslatedWorldPosition, TranslatedWorldPosition);
+		// #endif
+	}
+	
+	Payload.Params.SetOpacity(GetMaterialOpacity(PixelMaterialInputs));
+
+	if (DECAL_PAYLOAD_FLAGS & DECAL_WRITE_BASE_COLOR_FLAG)
+	{
+		Payload.Params.SetBaseColor(GetMaterialBaseColor(PixelMaterialInputs));
+	}
+
+	if (DECAL_PAYLOAD_FLAGS & DECAL_WRITE_NORMAL_FLAG)
+	{
+		Payload.Params.SetWorldNormal(-MaterialParameters.WorldNormal);
+	}
+
+	if (DECAL_PAYLOAD_FLAGS & DECAL_WRITE_ROUGHNESS_SPECULAR_METALLIC_FLAG)
+	{
+		Payload.Params.SetMetallic(GetMaterialMetallic(PixelMaterialInputs));
+		Payload.Params.SetSpecular(GetMaterialSpecular(PixelMaterialInputs));
+		Payload.Params.SetRoughness(GetMaterialRoughness(PixelMaterialInputs));
+	}
+
+	if (DECAL_PAYLOAD_FLAGS & DECAL_WRITE_AMBIENT_OCCLUSION_FLAG)
+	{
+		Payload.Params.SetAO(GetMaterialAmbientOcclusion(PixelMaterialInputs));
+	}
+
+	if (DECAL_PAYLOAD_FLAGS & DECAL_WRITE_EMISSIVE_FLAG)
+	{
+		Payload.Params.SetEmissive(GetMaterialEmissive(PixelMaterialInputs));
+	}
+
+	Payload.Params.SetFlags(DECAL_PAYLOAD_FLAGS);
+
+	return;
+}
+
+RAY_TRACING_ENTRY_ANY_HIT(RayTracingDecalMaterialAHS,
+	FRayTracingDecalPayload, Payload,
+	FRayTracingIntersectionAttributes, Attributes)
+{
+	ResolvedView = ResolveView();
+
+	const float3 TranslatedWorldPosition = TranslatedWorldRayOrigin() + WorldRayDirection() * RayTCurrent();
+	const float4 SvPosition = mul(float4(TranslatedWorldPosition, 1.0f), ResolvedView.TranslatedWorldToClip);
+
+	FVertexFactoryInterpolantsVSToPS Interpolants;
+
+	CalcInterpolants((FRayCone)0, Attributes, Interpolants);
+
+	FMaterialPixelParameters MaterialParameters = GetMaterialPixelParameters(Interpolants, SvPosition);
+	FPixelMaterialInputs PixelMaterialInputs;
+
+	const bool bIsFrontFace = HitKind() == HIT_KIND_TRIANGLE_FRONT_FACE;
+
+	{
+		const float4 ScreenPosition = SvPositionToResolvedScreenPosition(SvPosition);
+
+		MaterialParameters.CameraVector = -WorldRayDirection();
+
+		// #dxr_todo: UE-72130 support world position offset 
+		// #if USE_WORLD_POSITION_EXCLUDING_SHADER_OFFSETS
+		// CalcMaterialParametersEx(MaterialParameters, PixelMaterialInputs, SvPosition, ScreenPosition, bIsFrontFace, TranslatedWorldPosition, BasePassInterpolants.PixelPositionExcludingWPO);
+		// #else
+		CalcMaterialParametersEx(MaterialParameters, PixelMaterialInputs, SvPosition, ScreenPosition, bIsFrontFace, TranslatedWorldPosition, TranslatedWorldPosition);
+		// #endif
+	}
+
+	// We don't support blending of mesh decals in ray tracing at the moment.
+	// Instead we only use alpha masking to provide some limited overlap support, especially useful when large meshes are used to render sparse decals.
+	float Opacity = GetMaterialOpacity(PixelMaterialInputs);
+	if (Opacity <= 0.01f)
+	{
+		IgnoreHit();
+		return;
+	}
+}
+
+#endif // RAYHITGROUPSHADER
diff --git a/Engine/Shaders/Private/RayTracing/RayTracingDefaultDecalHitShader.usf b/Engine/Shaders/Private/RayTracing/RayTracingDefaultDecalHitShader.usf
new file mode 100644
--- /dev/null
+++ b/Engine/Shaders/Private/RayTracing/RayTracingDefaultDecalHitShader.usf
@@ -0,0 +1,25 @@
+#include "/Engine/Private/RayTracing/RayTracingCommon.ush"
+
+RAY_TRACING_ENTRY_CLOSEST_HIT(RayTracingDefaultOpaqueDecalMaterialCHS,
+	FRayTracingDecalPayload, Payload,
+	FRayTracingIntersectionAttributes, Attributes)
+{
+	Payload.HitT = RayTCurrent();
+
+	Payload.Params.SetOpacity(1.0f);
+	Payload.Params.SetBaseColor(float3(0, 0, 0));
+	Payload.Params.SetFlags(DECAL_WRITE_BASE_COLOR_FLAG);
+}
+
+RAY_TRACING_ENTRY_CLOSEST_HIT(RayTracingDefaultHiddenDecalMaterialCHS,
+	FRayTracingDecalPayload, Payload,
+	FRayTracingIntersectionAttributes, Attributes)
+{
+}
+
+RAY_TRACING_ENTRY_ANY_HIT(RayTracingDefaultHiddenDecalMaterialAHS,
+	FRayTracingDecalPayload, Payload,
+	FRayTracingIntersectionAttributes, Attributes)
+{
+	IgnoreHit();
+}
diff --git a/Engine/Shaders/Shared/RayTracingPayloadType.h b/Engine/Shaders/Shared/RayTracingPayloadType.h
--- a/Engine/Shaders/Shared/RayTracingPayloadType.h
+++ b/Engine/Shaders/Shared/RayTracingPayloadType.h
@@ -13,7 +13,8 @@
 #define RT_PAYLOAD_TYPE_LUMEN_MINIMAL			(1 << 6)	// FLumenMinimalPayload
 #define RT_PAYLOAD_TYPE_VFX   					(1 << 7)	// FVFXTracePayload
 #define RT_PAYLOAD_TYPE_DECALS					(1 << 8)	// FDecalShaderParams 
-#define RT_PAYLOAD_TYPE_SPARSE_VOXEL			(1 << 9)	// FSparseVoxelPayload
+#define RT_PAYLOAD_TYPE_MESH_DECALS				(1 << 9)	// FRayTracingDecalPayload 
+#define RT_PAYLOAD_TYPE_SPARSE_VOXEL			(1 << 10)	// FSparseVoxelPayload
 
 #ifdef __cplusplus
 /**
@@ -38,6 +39,7 @@ enum ERayTracingPayloadType : uint // HLSL
 	LumenMinimal = RT_PAYLOAD_TYPE_LUMEN_MINIMAL,
 	VFX = RT_PAYLOAD_TYPE_VFX,
 	Decals = RT_PAYLOAD_TYPE_DECALS,
+	MeshDecals = RT_PAYLOAD_TYPE_MESH_DECALS,
 	SparseVoxel = RT_PAYLOAD_TYPE_SPARSE_VOXEL,
 };
 
diff --git a/Engine/Source/Runtime/Renderer/Private/DeferredShadingRenderer.cpp b/Engine/Source/Runtime/Renderer/Private/DeferredShadingRenderer.cpp
--- a/Engine/Source/Runtime/Renderer/Private/DeferredShadingRenderer.cpp
+++ b/Engine/Source/Runtime/Renderer/Private/DeferredShadingRenderer.cpp
@@ -1198,6 +1198,7 @@ bool FDeferredShadingSceneRenderer::GatherRayTracingWorldInstancesForView(FRDGBu
 
 					RayTracingInstance.DefaultUserData = PrimitiveIndex;
 					RayTracingInstance.bApplyLocalBoundsTransform = Instance.bApplyLocalBoundsTransform;
+					RayTracingInstance.LayerIndex = (uint8)(Instance.MaskAndFlags.bAnySegmentsDecal ? ERayTracingSceneLayer::Decals : ERayTracingSceneLayer::Base);
 					RayTracingInstance.Mask = Instance.MaskAndFlags.Mask;
 
 					if (Instance.MaskAndFlags.bForceOpaque)
@@ -1461,7 +1462,11 @@ bool FDeferredShadingSceneRenderer::GatherRayTracingWorldInstancesForView(FRDGBu
 						continue;
 					}
 
-					const int32 NewInstanceIndex = RayTracingScene.Instances.Num();
+					checkf(SceneInfo->CachedRayTracingInstance.GeometryRHI, TEXT("Ray tracing instance must have a valid geometry."));
+					const int32 NewInstanceIndex = RayTracingScene.Instances.Add(SceneInfo->CachedRayTracingInstance);
+					FRayTracingGeometryInstance& NewInstance = RayTracingScene.Instances[NewInstanceIndex];
+
+					NewInstance.LayerIndex = (uint8)(RelevantPrimitive.bAnySegmentsDecal ? ERayTracingSceneLayer::Decals : ERayTracingSceneLayer::Base);
 
 					// At the moment we only support SM & ISMs on this path
 					check(EnumHasAnyFlags(Flags, ERayTracingPrimitiveFlags::CacheMeshCommands));
@@ -1477,9 +1482,6 @@ bool FDeferredShadingSceneRenderer::GatherRayTracingWorldInstancesForView(FRDGBu
 						}
 					}
 
-					checkf(SceneInfo->CachedRayTracingInstance.GeometryRHI, TEXT("Ray tracing instance must have a valid geometry."));
-					RayTracingScene.Instances.Add(SceneInfo->CachedRayTracingInstance);
-
 					if (bEnableInstanceDebugData)
 					{
 						RayTracingScene.AddInstanceDebugData(SceneInfo->CachedRayTracingInstance.GeometryRHI, SceneInfo->Proxy, false);
@@ -1490,8 +1492,6 @@ bool FDeferredShadingSceneRenderer::GatherRayTracingWorldInstancesForView(FRDGBu
 
 					if (CullingParameters.CullInRayTracing > 0 && GetRayTracingCullingPerInstance() && SceneInfo->CachedRayTracingInstance.NumTransforms > 1 && !bUseGroupBounds)
 					{
-						FRayTracingGeometryInstance& NewInstance = RayTracingScene.Instances.Last();
-
 						const bool bIsFarFieldPrimitive = EnumHasAnyFlags(Flags, ERayTracingPrimitiveFlags::FarField);
 
 						TArrayView<uint32> InstanceActivationMask = RayTracingScene.Allocate<uint32>(FMath::DivideAndRoundUp(NewInstance.NumTransforms, 32u));
@@ -1520,7 +1520,7 @@ bool FDeferredShadingSceneRenderer::GatherRayTracingWorldInstancesForView(FRDGBu
 						}
 					}
 
-					AddDebugRayTracingInstanceFlags(RayTracingScene.Instances.Last().Flags);
+					AddDebugRayTracingInstanceFlags(NewInstance.Flags);
 				}
 				else
 				{
@@ -1613,6 +1613,8 @@ bool FDeferredShadingSceneRenderer::GatherRayTracingWorldInstancesForView(FRDGBu
 							RayTracingInstance.Flags |= ERayTracingInstanceFlags::TriangleCullDisable;
 						}
 						AddDebugRayTracingInstanceFlags(RayTracingInstance.Flags);
+
+						RayTracingInstance.LayerIndex = (uint8)(RelevantPrimitive.bAnySegmentsDecal ? ERayTracingSceneLayer::Decals : ERayTracingSceneLayer::Base);
 					}
 				}
 			}
diff --git a/Engine/Source/Runtime/Renderer/Private/PathTracing.cpp b/Engine/Source/Runtime/Renderer/Private/PathTracing.cpp
--- a/Engine/Source/Runtime/Renderer/Private/PathTracing.cpp
+++ b/Engine/Source/Runtime/Renderer/Private/PathTracing.cpp
@@ -23,6 +23,7 @@ TAutoConsoleVariable<int32> CVarPathTracing(
 #include "PathTracingDefinitions.h"
 #include "RayTracing/RayTracingMaterialHitShaders.h"
 #include "RayTracing/RayTracingDecals.h"
+#include "DecalRenderingCommon.h"
 #include "GenerateMips.h"
 #include "HairStrands/HairStrandsData.h"
 #include "Modules/ModuleManager.h"
@@ -280,7 +281,21 @@ TAutoConsoleVariable<int32> CVarPathTracingDecalGridVisualize(
 TAutoConsoleVariable<float> CVarPathTracingDecalRoughnessCutoff(
 	TEXT("r.PathTracing.DecalRoughnessCutoff"),
 	0.15f,
-	TEXT("Do not evaluate decals beyond this roughness level to improve performance (default=0.15)\n"),
+	TEXT("Do not evaluate decals beyond this roughness level to improve performance (default=0.15)"),
+	ECVF_RenderThreadSafe
+);
+
+TAutoConsoleVariable<float> CVarPathTracingMeshDecalRoughnessCutoff(
+	TEXT("r.PathTracing.MeshDecalRoughnessCutoff"),
+	0.15f,
+	TEXT("Do not evaluate mesh decals beyond this roughness level to improve performance (default=0.15)"),
+	ECVF_RenderThreadSafe
+);
+
+TAutoConsoleVariable<float> CVarPathTracingMeshDecalBias(
+	TEXT("r.PathTracing.MeshDecalBias"),
+	1.0f,
+	TEXT("Bias applied to mesh decal rays to avoid intersection with geometry (default = 1.0f)"),
 	ECVF_RenderThreadSafe
 );
 
@@ -318,6 +333,8 @@ BEGIN_SHADER_PARAMETER_STRUCT(FPathTracingData, )
 	SHADER_PARAMETER(float, FilterWidth)
 	SHADER_PARAMETER(float, AbsorptionScale)
 	SHADER_PARAMETER(float, DecalRoughnessCutoff)
+	SHADER_PARAMETER(float, MeshDecalRoughnessCutoff)
+	SHADER_PARAMETER(float, MeshDecalBias)
 	SHADER_PARAMETER(float, CameraFocusDistance)
 	SHADER_PARAMETER(FVector2f, CameraLensRadius)
 END_SHADER_PARAMETER_STRUCT()
@@ -359,6 +376,8 @@ struct FPathTracingConfig
 			PathTracingData.EnableAtmosphere != Other.PathTracingData.EnableAtmosphere ||
 			PathTracingData.EnableFog != Other.PathTracingData.EnableFog ||
 			PathTracingData.DecalRoughnessCutoff != Other.PathTracingData.DecalRoughnessCutoff ||
+			PathTracingData.MeshDecalRoughnessCutoff != Other.PathTracingData.MeshDecalRoughnessCutoff ||
+			PathTracingData.MeshDecalBias != Other.PathTracingData.MeshDecalBias ||
 			PathTracingData.MaxRaymarchSteps != Other.PathTracingData.MaxRaymarchSteps ||
 			ViewRect != Other.ViewRect ||
 			LightShowFlags != Other.LightShowFlags ||
@@ -534,6 +553,9 @@ static void PreparePathTracingData(const FScene* Scene, const FViewInfo& View, F
 
 	PathTracingData.DecalRoughnessCutoff = PathTracing::UsesDecals(*View.Family) && View.bHasRayTracingDecals ? CVarPathTracingDecalRoughnessCutoff.GetValueOnRenderThread() : -1.0f;
 
+	PathTracingData.MeshDecalRoughnessCutoff = PathTracing::UsesDecals(*View.Family) && Scene->RayTracingScene.GetRHIRayTracingScene()->GetInitializer().NumNativeInstancesPerLayer[(uint32)ERayTracingSceneLayer::Decals] > 0 ? CVarPathTracingMeshDecalRoughnessCutoff.GetValueOnRenderThread() : -1.0f;
+	PathTracingData.MeshDecalBias = CVarPathTracingMeshDecalBias.GetValueOnRenderThread();
+
 	PathTracingData.MaxRaymarchSteps = CVarPathTracingMaxRaymarchSteps.GetValueOnRenderThread();
 }
 
@@ -721,7 +743,7 @@ class FPathTracingRG : public FGlobalShader
 
 	static ERayTracingPayloadType GetRayTracingPayloadType(const int32 PermutationId)
 	{
-		return ERayTracingPayloadType::PathTracingMaterial | ERayTracingPayloadType::Decals;
+		return ERayTracingPayloadType::PathTracingMaterial | ERayTracingPayloadType::Decals | ERayTracingPayloadType::MeshDecals;
 	}
 
 	BEGIN_SHADER_PARAMETER_STRUCT(FParameters, )
@@ -729,6 +751,7 @@ class FPathTracingRG : public FGlobalShader
 		SHADER_PARAMETER_RDG_TEXTURE_UAV(RWTexture2D<float4>, AlbedoTexture)
 		SHADER_PARAMETER_RDG_TEXTURE_UAV(RWTexture2D<float4>, NormalTexture)
 		SHADER_PARAMETER_SRV(RaytracingAccelerationStructure, TLAS)
+		SHADER_PARAMETER_SRV(RaytracingAccelerationStructure, DecalTLAS)
 
 		SHADER_PARAMETER_STRUCT_REF(FViewUniformShaderParameters, ViewUniformBuffer)
 		SHADER_PARAMETER_STRUCT_INCLUDE(FPathTracingData, PathTracingData)
@@ -1281,20 +1304,29 @@ bool FRayTracingMeshProcessor::ProcessPathTracing(
 
 	const bool bUseProceduralPrimitive = MeshBatch.VertexFactory->GetType()->SupportsRayTracingProceduralPrimitive() &&
 		FDataDrivenShaderPlatformInfo::GetSupportsRayTracingProceduralPrimitive(GMaxRHIShaderPlatform);
+
+	const bool bIsDecal = MaterialResource.GetMaterialDomain() == MD_DeferredDecal;
+	FShaderType* DecalShaderType_CHS = GetRayTracingDecalMaterialShaderType(false);
+	FShaderType* DecalShaderType_CHS_AHS = GetRayTracingDecalMaterialShaderType(true);
+
 	switch (RayTracingMeshCommandsMode)
 	{
 		case ERayTracingMeshCommandsMode::PATH_TRACING:
 		{
 			if (NeedsAnyHitShader(MaterialResource))
 			{
-				if (bUseProceduralPrimitive)
+				if (bIsDecal)
+					ShaderTypes.AddShaderType(DecalShaderType_CHS_AHS);
+				else if (bUseProceduralPrimitive)
 					ShaderTypes.AddShaderType<FPathTracingMaterialCHS_AHS_IS>();
 				else
 					ShaderTypes.AddShaderType<FPathTracingMaterialCHS_AHS>();
 			}
 			else
 			{
-				if (bUseProceduralPrimitive)
+				if (bIsDecal)
+					ShaderTypes.AddShaderType(DecalShaderType_CHS);
+				else if (bUseProceduralPrimitive)
 					ShaderTypes.AddShaderType<FPathTracingMaterialCHS_IS>();
 				else
 					ShaderTypes.AddShaderType<FPathTracingMaterialCHS>();
@@ -2308,6 +2340,7 @@ void FDeferredShadingSceneRenderer::RenderPathTracing(
 						{
 							FPathTracingRG::FParameters* PassParameters = GraphBuilder.AllocParameters<FPathTracingRG::FParameters>();
 							PassParameters->TLAS = Scene->RayTracingScene.GetLayerSRVChecked(ERayTracingSceneLayer::Base);
+							PassParameters->DecalTLAS = Scene->RayTracingScene.GetLayerSRVChecked(ERayTracingSceneLayer::Decals);
 							PassParameters->ViewUniformBuffer = View.ViewUniformBuffer;
 							PassParameters->PathTracingData = Config.PathTracingData;
 							PassParameters->StartingExtinctionCoefficient = GraphBuilder.CreateSRV(StartingExtinctionCoefficient, PF_R32_FLOAT);
diff --git a/Engine/Source/Runtime/Renderer/Private/RayTracing/RayTracingDecals.cpp b/Engine/Source/Runtime/Renderer/Private/RayTracing/RayTracingDecals.cpp
--- a/Engine/Source/Runtime/Renderer/Private/RayTracing/RayTracingDecals.cpp
+++ b/Engine/Source/Runtime/Renderer/Private/RayTracing/RayTracingDecals.cpp
@@ -36,6 +36,20 @@ enum class EDecalsWriteFlags : uint32
 	AmbientOcclusion			= (1 << 4),
 };
 
+static uint32 CalculateDecalWriteFlags(FDecalBlendDesc DecalBlendDesc)
+{
+	uint32 Flags = (uint32)EDecalsWriteFlags::None;
+#define SET_FLAG(BOOL_VAL, WRITE_FLAG) Flags |= uint32(DecalBlendDesc.BOOL_VAL ? EDecalsWriteFlags::WRITE_FLAG : EDecalsWriteFlags::None)
+	SET_FLAG(bWriteBaseColor, BaseColor);
+	SET_FLAG(bWriteNormal, Normal);
+	SET_FLAG(bWriteRoughnessSpecularMetallic, RoughnessSpecularMetallic);
+	SET_FLAG(bWriteEmissive, Emissive);
+	SET_FLAG(bWriteAmbientOcclusion, AmbientOcclusion);
+#undef SET_FLAG
+
+	return Flags;
+}
+
 IMPLEMENT_GLOBAL_SHADER_PARAMETER_STRUCT(FRayTracingDecals, "RayTracingDecals");
 
 BEGIN_GLOBAL_SHADER_PARAMETER_STRUCT(FDecalParametersRayTracing, )
@@ -81,14 +95,7 @@ static TUniformBufferRef<FDecalParametersRayTracing> CreateDecalParametersBuffer
 
 	Parameters.DecalParams = FVector2f(DecalData.FadeAlpha, LifetimeAlpha);
 
-	Parameters.DecalWriteFlags = (uint32)EDecalsWriteFlags::None;
-#define SET_FLAG(BOOL_VAL, WRITE_FLAG) Parameters.DecalWriteFlags |= uint32(DecalData.BlendDesc.BOOL_VAL ? EDecalsWriteFlags::WRITE_FLAG : EDecalsWriteFlags::None)
-	SET_FLAG(bWriteBaseColor, BaseColor);
-	SET_FLAG(bWriteNormal, Normal);
-	SET_FLAG(bWriteRoughnessSpecularMetallic, RoughnessSpecularMetallic);
-	SET_FLAG(bWriteEmissive, Emissive);
-	SET_FLAG(bWriteAmbientOcclusion, AmbientOcclusion);
-#undef SET_FLAG
+	Parameters.DecalWriteFlags = CalculateDecalWriteFlags(DecalData.BlendDesc);
 
 	return CreateUniformBufferImmediate(Parameters, Usage);
 }
@@ -157,6 +164,9 @@ public:
 	static void ModifyCompilationEnvironment(const FMaterialShaderPermutationParameters& Parameters, FShaderCompilerEnvironment& OutEnvironment)
 	{
 		FMaterialShader::ModifyCompilationEnvironment(Parameters, OutEnvironment);
+
+		// DECAL_PRIMITIVE informs material templates which functions to expose when rendering decals.
+		OutEnvironment.SetDefine(TEXT("DECAL_PRIMITIVE"), 1);
 	}
 
 	static bool ValidateCompiledResult(EShaderPlatform Platform, const FShaderParameterMap& ParameterMap, TArray<FString>& OutError)
@@ -211,6 +221,151 @@ IMPLEMENT_RT_PAYLOAD_TYPE(ERayTracingPayloadType::Decals, 60);
 
 IMPLEMENT_SHADER_TYPE(, FRayTracingDecalMaterialShader, TEXT("/Engine/Private/RayTracing/RayTracingDecalMaterialShader.usf"), TEXT("RayTracingDecalMaterialShader"), SF_RayCallable);
 
+template<bool bUseAnyHitShader>
+class TRayTracingDecalMaterial : public FMeshMaterialShader
+{
+	DECLARE_SHADER_TYPE(TRayTracingDecalMaterial, MeshMaterial);
+public:
+	TRayTracingDecalMaterial() = default;
+
+	TRayTracingDecalMaterial(const FMeshMaterialShaderType::CompiledShaderInitializerType& Initializer)
+		: FMeshMaterialShader(Initializer)
+	{}
+
+	static bool ShouldCompilePermutation(const FMeshMaterialShaderPermutationParameters& Parameters)
+	{
+		if (!ShouldCompileRayTracingShadersForProject(Parameters.Platform))
+		{
+			// is raytracing enabled at all?
+			return false;
+		}
+		if (!Parameters.VertexFactoryType->SupportsRayTracing())
+		{
+			// does the VF support ray tracing at all?
+			return false;
+		}
+		if (Parameters.MaterialParameters.MaterialDomain != MD_DeferredDecal)
+		{
+			return false;
+		}
+		if ((Parameters.MaterialParameters.bIsMasked || Parameters.MaterialParameters.BlendMode != BLEND_Opaque) != bUseAnyHitShader)
+		{
+			// the anyhit permutation is only required if the material is masked or has a non-opaque blend mode
+			return false;
+		}
+
+		return true;
+	}
+
+	static void ModifyCompilationEnvironment(const FMaterialShaderPermutationParameters& Parameters, FShaderCompilerEnvironment& OutEnvironment)
+	{
+		FMeshMaterialShader::ModifyCompilationEnvironment(Parameters, OutEnvironment);
+
+		check(Parameters.MaterialParameters.MaterialDomain == MD_DeferredDecal);
+
+		FDecalBlendDesc DecalBlendDesc = DecalRendering::ComputeDecalBlendDesc(Parameters.Platform, Parameters.MaterialParameters);
+
+		OutEnvironment.SetDefine(TEXT("DECAL_PAYLOAD_FLAGS"), CalculateDecalWriteFlags(DecalBlendDesc));
+	}
+
+	static bool ValidateCompiledResult(EShaderPlatform Platform, const FShaderParameterMap& ParameterMap, TArray<FString>& OutError)
+	{
+		if (ParameterMap.ContainsParameterAllocation(FSceneTextureUniformParameters::StaticStructMetadata.GetShaderVariableName()))
+		{
+			OutError.Add(TEXT("Ray tracing closest hit shaders cannot read from the SceneTexturesStruct."));
+			return false;
+		}
+
+		for (const auto& It : ParameterMap.GetParameterMap())
+		{
+			const FParameterAllocation& ParamAllocation = It.Value;
+			if (ParamAllocation.Type != EShaderParameterType::UniformBuffer
+				&& ParamAllocation.Type != EShaderParameterType::LooseData)
+			{
+				OutError.Add(FString::Printf(TEXT("Invalid ray tracing shader parameter '%s'. Only uniform buffers and loose data parameters are supported."), *(It.Key)));
+				return false;
+			}
+		}
+
+		return true;
+	}
+
+	static ERayTracingPayloadType GetRayTracingPayloadType(const int32 PermutationId)
+	{
+		return ERayTracingPayloadType::MeshDecals;
+	}
+};
+
+using FRayTracingDecalMaterialCHS = TRayTracingDecalMaterial<false>;
+using FRayTracingDecalMaterialCHS_AHS = TRayTracingDecalMaterial<true>;
+
+IMPLEMENT_MATERIAL_SHADER_TYPE(template<>, FRayTracingDecalMaterialCHS, TEXT("/Engine/Private/RayTracing/RayTracingDecalMaterialShader.usf"), TEXT("closesthit=RayTracingDecalMaterialCHS"), SF_RayHitGroup);
+IMPLEMENT_MATERIAL_SHADER_TYPE(template<>, FRayTracingDecalMaterialCHS_AHS, TEXT("/Engine/Private/RayTracing/RayTracingDecalMaterialShader.usf"), TEXT("closesthit=RayTracingDecalMaterialCHS anyhit=RayTracingDecalMaterialAHS"), SF_RayHitGroup);
+
+IMPLEMENT_RT_PAYLOAD_TYPE(ERayTracingPayloadType::MeshDecals, 64);
+
+FShaderType* GetRayTracingDecalMaterialShaderType(bool bUseAnyHitShader)
+{
+	if (bUseAnyHitShader)
+	{
+		return &FRayTracingDecalMaterialCHS_AHS::StaticType;
+	}
+	else
+	{
+		return &FRayTracingDecalMaterialCHS::StaticType;
+	}
+}
+
+class FDefaultOpaqueMeshDecalHitGroup : public FGlobalShader
+{
+	DECLARE_GLOBAL_SHADER(FDefaultOpaqueMeshDecalHitGroup)
+	SHADER_USE_ROOT_PARAMETER_STRUCT(FDefaultOpaqueMeshDecalHitGroup, FGlobalShader)
+
+	static bool ShouldCompilePermutation(const FGlobalShaderPermutationParameters& Parameters)
+	{
+		return ShouldCompileRayTracingShadersForProject(Parameters.Platform);
+	}
+
+	static ERayTracingPayloadType GetRayTracingPayloadType(const int32 PermutationId)
+	{
+		return ERayTracingPayloadType::MeshDecals;
+	}
+
+	using FParameters = FEmptyShaderParameters;
+};
+
+IMPLEMENT_SHADER_TYPE(, FDefaultOpaqueMeshDecalHitGroup, TEXT("/Engine/Private/RayTracing/RayTracingDefaultDecalHitShader.usf"), TEXT("closesthit=RayTracingDefaultOpaqueDecalMaterialCHS"), SF_RayHitGroup);
+
+class FDefaultHiddenMeshDecalHitGroup : public FGlobalShader
+{
+	DECLARE_GLOBAL_SHADER(FDefaultHiddenMeshDecalHitGroup)
+	SHADER_USE_ROOT_PARAMETER_STRUCT(FDefaultHiddenMeshDecalHitGroup, FGlobalShader)
+
+	static bool ShouldCompilePermutation(const FGlobalShaderPermutationParameters& Parameters)
+	{
+		return ShouldCompileRayTracingShadersForProject(Parameters.Platform);
+	}
+
+	static ERayTracingPayloadType GetRayTracingPayloadType(const int32 PermutationId)
+	{
+		return ERayTracingPayloadType::MeshDecals;
+	}
+
+	using FParameters = FEmptyShaderParameters;
+};
+
+IMPLEMENT_SHADER_TYPE(, FDefaultHiddenMeshDecalHitGroup, TEXT("/Engine/Private/RayTracing/RayTracingDefaultDecalHitShader.usf"), TEXT("closesthit=RayTracingDefaultHiddenDecalMaterialCHS anyhit=RayTracingDefaultHiddenDecalMaterialAHS"), SF_RayHitGroup);
+
+FRHIRayTracingShader* GetDefaultOpaqueMeshDecalHitShader(const FGlobalShaderMap* ShaderMap)
+{
+	return ShaderMap->GetShader<FDefaultOpaqueMeshDecalHitGroup>().GetRayTracingShader();
+}
+
+FRHIRayTracingShader* GetDefaultHiddenMeshDecalHitShader(const FGlobalShaderMap* ShaderMap)
+{
+	return ShaderMap->GetShader<FDefaultHiddenMeshDecalHitGroup>().GetRayTracingShader();
+}
+
 void BuildDecalGrid(FRDGBuilder& GraphBuilder, uint32 NumDecals, FRDGBufferSRVRef DecalsSRV, FRayTracingDecals& OutParameters)
 {
 	// if we have some decals -- build a decal grid
diff --git a/Engine/Source/Runtime/Renderer/Private/RayTracing/RayTracingDecals.h b/Engine/Source/Runtime/Renderer/Private/RayTracing/RayTracingDecals.h
--- a/Engine/Source/Runtime/Renderer/Private/RayTracing/RayTracingDecals.h
+++ b/Engine/Source/Runtime/Renderer/Private/RayTracing/RayTracingDecals.h
@@ -9,6 +9,9 @@
 #if RHI_RAYTRACING
 
 class FScene;
+class FShaderType;
+class FRHIRayTracingShader;
+class FGlobalShaderMap;
 
 BEGIN_GLOBAL_SHADER_PARAMETER_STRUCT(FRayTracingDecals, RENDERER_API)
 	SHADER_PARAMETER_RDG_TEXTURE(Texture2D<uint>, Grid)
@@ -24,4 +27,9 @@ END_GLOBAL_SHADER_PARAMETER_STRUCT()
 TRDGUniformBufferRef<FRayTracingDecals> CreateNullRayTracingDecalsUniformBuffer(FRDGBuilder& GraphBuilder);
 TRDGUniformBufferRef<FRayTracingDecals> CreateRayTracingDecalData(FRDGBuilder& GraphBuilder, FScene& Scene, const FViewInfo& View, uint32 BaseCallableSlotIndex);
 
+FShaderType* GetRayTracingDecalMaterialShaderType(bool bUseAnyHitShader);
+
+FRHIRayTracingShader* GetDefaultOpaqueMeshDecalHitShader(const FGlobalShaderMap* ShaderMap);
+FRHIRayTracingShader* GetDefaultHiddenMeshDecalHitShader(const FGlobalShaderMap* ShaderMap);
+
 #endif
diff --git a/Engine/Source/Runtime/Renderer/Private/RayTracing/RayTracingMaterialHitShaders.cpp b/Engine/Source/Runtime/Renderer/Private/RayTracing/RayTracingMaterialHitShaders.cpp
--- a/Engine/Source/Runtime/Renderer/Private/RayTracing/RayTracingMaterialHitShaders.cpp
+++ b/Engine/Source/Runtime/Renderer/Private/RayTracing/RayTracingMaterialHitShaders.cpp
@@ -15,6 +15,7 @@
 #include "BuiltInRayTracingShaders.h"
 #include "RaytracingOptions.h"
 #include "RayTracingLighting.h"
+#include "RayTracingDecals.h"
 #include "PathTracing.h"
 
 int32 GEnableRayTracingMaterials = 1;
@@ -639,13 +640,20 @@ FRayTracingPipelineState* FDeferredShadingSceneRenderer::CreateRayTracingMateria
 	SCOPE_CYCLE_COUNTER(STAT_BindRayTracingPipeline);
 
 	const bool bIsPathTracing = ViewFamily.EngineShowFlags.PathTracing;
+	const bool bSupportMeshDecals = bIsPathTracing;
+
+	ERayTracingPayloadType PayloadType = bIsPathTracing
+		? (ERayTracingPayloadType::PathTracingMaterial | ERayTracingPayloadType::Decals)
+		: ERayTracingPayloadType::RayTracingMaterial;
+
+	if (bSupportMeshDecals)
+	{
+		PayloadType |= ERayTracingPayloadType::MeshDecals;
+	}
 
 	FRayTracingPipelineStateInitializer Initializer;
 
-	Initializer.MaxPayloadSizeInBytes = GetRayTracingPayloadTypeMaxSize(
-		bIsPathTracing
-			? (ERayTracingPayloadType::PathTracingMaterial | ERayTracingPayloadType::Decals)
-			: ERayTracingPayloadType::RayTracingMaterial);
+	Initializer.MaxPayloadSizeInBytes = GetRayTracingPayloadTypeMaxSize(PayloadType);
 	Initializer.bAllowHitGroupIndexing = true;
 
 	FRHIRayTracingShader* DefaultMissShader = bIsPathTracing ? GetPathTracingDefaultMissShader(View.ShaderMap) : GetRayTracingDefaultMissShader(View.ShaderMap);
@@ -667,11 +675,19 @@ FRayTracingPipelineState* FDeferredShadingSceneRenderer::CreateRayTracingMateria
 
 	FRHIRayTracingShader* OpaqueShadowShader   = bIsPathTracing ? GetPathTracingDefaultOpaqueHitShader(View.ShaderMap) : GetRayTracingDefaultOpaqueShader(View.ShaderMap);
 	FRHIRayTracingShader* HiddenMaterialShader = bIsPathTracing ? GetPathTracingDefaultHiddenHitShader(View.ShaderMap) : GetRayTracingDefaultHiddenShader(View.ShaderMap);
+
+	FRHIRayTracingShader* OpaqueMeshDecalHitShader = GetDefaultOpaqueMeshDecalHitShader(View.ShaderMap);
+	FRHIRayTracingShader* HiddenMeshDecalHitShader = GetDefaultHiddenMeshDecalHitShader(View.ShaderMap);
 	
 	TArray<FRHIRayTracingShader*> RayTracingHitGroupLibrary;
 	if (bEnableMaterials)
 	{
 		FShaderMapResource::GetRayTracingHitGroupLibrary(RayTracingHitGroupLibrary, OpaqueShadowShader);
+
+		if (bSupportMeshDecals)
+		{
+			FShaderMapResource::GetRayTracingHitGroupLibrary(RayTracingHitGroupLibrary, OpaqueMeshDecalHitShader);
+		}
 	}
 
 	FRHIRayTracingShader* RequiredHitShaders[] =
@@ -747,6 +763,9 @@ FRayTracingPipelineState* FDeferredShadingSceneRenderer::CreateRayTracingMateria
 	const int32 OpaqueShadowMaterialIndex = FindRayTracingHitGroupIndex(PipelineState, OpaqueShadowShader, true);
 	const int32 HiddenMaterialIndex = FindRayTracingHitGroupIndex(PipelineState, HiddenMaterialShader, true);
 
+	const int32 OpaqueMeshDecalHitGroupIndex = FindRayTracingHitGroupIndex(PipelineState, OpaqueShadowShader, true);
+	const int32 HiddenMeshDecalHitGroupIndex = FindRayTracingHitGroupIndex(PipelineState, HiddenMaterialShader, true);
+
 	FViewInfo& ReferenceView = Views[0];
 
 	// material hit groups
@@ -770,8 +789,8 @@ FRayTracingPipelineState* FDeferredShadingSceneRenderer::CreateRayTracingMateria
 			View.RayTracingMaterialBindings[TaskIndex] = BindingWriter;
 
 			TaskList.Add(FFunctionGraphTask::CreateAndDispatchWhenReady(
-				[&View, PipelineState, BindingWriter, MeshCommands, NumCommands, bEnableMaterials, bEnableShadowMaterials,
-				OpaqueShadowMaterialIndex, HiddenMaterialIndex, TaskIndex]()
+				[&View, bIsPathTracing, PipelineState, BindingWriter, MeshCommands, NumCommands, bEnableMaterials, bEnableShadowMaterials, bSupportMeshDecals,
+				OpaqueShadowMaterialIndex, HiddenMaterialIndex, OpaqueMeshDecalHitGroupIndex, HiddenMeshDecalHitGroupIndex, TaskIndex]()
 				{
 					TRACE_CPUPROFILER_EVENT_SCOPE(BindRayTracingMaterialPipelineTask);
 
@@ -780,7 +799,21 @@ FRayTracingPipelineState* FDeferredShadingSceneRenderer::CreateRayTracingMateria
 						const FVisibleRayTracingMeshCommand VisibleMeshCommand = MeshCommands[CommandIndex];
 						const FRayTracingMeshCommand& MeshCommand = *VisibleMeshCommand.RayTracingMeshCommand;
 
-						int32 HitGroupIndex = OpaqueShadowMaterialIndex; // Force the same shader to be used on all geometry unless materials are enabled
+						// Force the same shader to be used on all geometry unless materials are enabled
+						int32 HitGroupIndex;
+
+						if (MeshCommand.MaterialShader->RayTracingPayloadType == (uint32)ERayTracingPayloadType::MeshDecals)
+						{
+							checkf(bSupportMeshDecals, TEXT("Unexpected ray tracing mesh command using Mesh Decal payload. Fix logic adding the command or update bSupportMeshDecals as appropriate."));
+							HitGroupIndex = OpaqueMeshDecalHitGroupIndex;
+						}
+						else
+						{
+							checkf((!bIsPathTracing && MeshCommand.MaterialShader->RayTracingPayloadType == (uint32)ERayTracingPayloadType::RayTracingMaterial)
+								|| (bIsPathTracing && MeshCommand.MaterialShader->RayTracingPayloadType == (uint32)ERayTracingPayloadType::PathTracingMaterial),
+								TEXT("Incorrectly using RayTracingMaterial when path tracer is enabled or vice-versa."));
+							HitGroupIndex = OpaqueShadowMaterialIndex;
+						}
 
 						if (bEnableMaterials)
 						{
diff --git a/Engine/Source/Runtime/Renderer/Private/RayTracing/RayTracingScene.h b/Engine/Source/Runtime/Renderer/Private/RayTracing/RayTracingScene.h
--- a/Engine/Source/Runtime/Renderer/Private/RayTracing/RayTracingScene.h
+++ b/Engine/Source/Runtime/Renderer/Private/RayTracing/RayTracingScene.h
@@ -26,6 +26,7 @@ class FRDGBuilder;
 enum class ERayTracingSceneLayer : uint8
 {
 	Base,
+	Decals,
 
 	NUM
 };
