Render Translucency for Mobile when Debug View is enabled

diff --git a/Engine/Source/Runtime/Renderer/Private/MobileTranslucentRendering.cpp b/Engine/Source/Runtime/Renderer/Private/MobileTranslucentRendering.cpp
--- a/Engine/Source/Runtime/Renderer/Private/MobileTranslucentRendering.cpp
+++ b/Engine/Source/Runtime/Renderer/Private/MobileTranslucentRendering.cpp
@@ -6,7 +6,7 @@
 
 void FMobileSceneRenderer::RenderTranslucency(FRHICommandListImmediate& RHICmdList, const FViewInfo& View)
 {	
-	const bool bShouldRenderTranslucency = ShouldRenderTranslucency(StandardTranslucencyPass) && ViewFamily.EngineShowFlags.Translucency && !ViewFamily.UseDebugViewPS();
+	const bool bShouldRenderTranslucency = ShouldRenderTranslucency(StandardTranslucencyPass) && ViewFamily.EngineShowFlags.Translucency;
 	if (bShouldRenderTranslucency)
 	{
 		CSV_SCOPED_TIMING_STAT_EXCLUSIVE(RenderTranslucency);
