Remove WaitForRayTracingScene dependency on scratch buffer parameters:
* The dependency was added to prevent RDG from removing those buffers when TLAS build was running async. After TLAS build pass was moved to RDG it's no longer needed.
* Because WaitForRayTracingScene can run on async queue right after RayTracingUpdate the dependency resulted in the stall on graphics queue.

diff --git a/Engine/Source/Runtime/Renderer/Private/DeferredShadingRenderer.cpp b/Engine/Source/Runtime/Renderer/Private/DeferredShadingRenderer.cpp
--- a/Engine/Source/Runtime/Renderer/Private/DeferredShadingRenderer.cpp
+++ b/Engine/Source/Runtime/Renderer/Private/DeferredShadingRenderer.cpp
@@ -2093,10 +2093,9 @@ void FDeferredShadingSceneRenderer::WaitForRayTracingScene(FRDGBuilder& GraphBui
 
 	const bool bIsPathTracing = ViewFamily.EngineShowFlags.PathTracing;
 
-	// Scratch buffer must be referenced in this pass, as it must live until the BVH build is complete.
 	FBuildAccelerationStructurePassParams* PassParams = GraphBuilder.AllocParameters<FBuildAccelerationStructurePassParams>();
-	PassParams->RayTracingSceneScratchBuffer = Scene->RayTracingScene.BuildScratchBuffer;
-	PassParams->DynamicGeometryScratchBuffer = DynamicGeometryScratchBuffer;
+	PassParams->RayTracingSceneScratchBuffer = nullptr;
+	PassParams->DynamicGeometryScratchBuffer = nullptr;
 	PassParams->LightDataPacked = bIsPathTracing ? nullptr : ReferenceView.RayTracingLightDataUniformBuffer; // accessed by FRayTracingLightingMS
 
 	if (IsNaniteEnabled())
