Add PrimitiveSceneProxy null check while setting ray tracing mask. It could be nullptr in GPU Lightmass.

diff --git a/Engine/Source/Runtime/Renderer/Private/RayTracing/RayTracingInstanceMask.cpp b/Engine/Source/Runtime/Renderer/Private/RayTracing/RayTracingInstanceMask.cpp
--- a/Engine/Source/Runtime/Renderer/Private/RayTracing/RayTracingInstanceMask.cpp
+++ b/Engine/Source/Runtime/Renderer/Private/RayTracing/RayTracingInstanceMask.cpp
@@ -239,7 +239,7 @@ FRayTracingMaskAndFlags BuildRayTracingInstanceMaskAndFlags(const FRayTracingIns
 	return BuildRayTracingInstanceMaskAndFlags(MeshBatches, PrimitiveSceneProxy.GetScene().GetFeatureLevel(), MaskInfo, Instance.InstanceLayer, ExtraMask);
 }
 
-void SetupRayTracingMeshCommandMaskAndStatus(FRayTracingMeshCommand& MeshCommand, const FMeshBatch& MeshBatch, const FPrimitiveSceneProxy& PrimitiveSceneProxy,
+void SetupRayTracingMeshCommandMaskAndStatus(FRayTracingMeshCommand& MeshCommand, const FMeshBatch& MeshBatch, const FPrimitiveSceneProxy* PrimitiveSceneProxy,
 	const FMaterial& MaterialResource, ERayTracingViewMaskMode MaskMode)
 {
 	const FVertexFactory* VertexFactory = MeshBatch.VertexFactory;
@@ -253,7 +253,12 @@ void SetupRayTracingMeshCommandMaskAndStatus(FRayTracingMeshCommand& MeshCommand
 
 	MeshCommand.InstanceMask = BlendModeToRayTracingInstanceMask(MaterialResource.GetBlendMode(), MaskMode);
 
-	FSceneProxyRayTracingMaskInfo MaskInfo = GetSceneProxyRayTracingMaskInfo(PrimitiveSceneProxy, nullptr);
+	if (!PrimitiveSceneProxy)
+	{
+		return;
+	}
+
+	FSceneProxyRayTracingMaskInfo MaskInfo = GetSceneProxyRayTracingMaskInfo(*PrimitiveSceneProxy, nullptr);
 
 	if (MaskMode == ERayTracingViewMaskMode::PathTracing || MaskMode == ERayTracingViewMaskMode::LightMapTracing)
 	{
diff --git a/Engine/Source/Runtime/Renderer/Private/RayTracing/RayTracingInstanceMask.h b/Engine/Source/Runtime/Renderer/Private/RayTracing/RayTracingInstanceMask.h
--- a/Engine/Source/Runtime/Renderer/Private/RayTracing/RayTracingInstanceMask.h
+++ b/Engine/Source/Runtime/Renderer/Private/RayTracing/RayTracingInstanceMask.h
@@ -91,7 +91,7 @@ class FRayTracingMeshCommand;
 enum class ERayTracingPrimitiveFlags : uint8;
 class FMaterial;
 
-void SetupRayTracingMeshCommandMaskAndStatus(FRayTracingMeshCommand& MeshCommand, const FMeshBatch& MeshBatch, const FPrimitiveSceneProxy& PrimitiveSceneProxy,
+void SetupRayTracingMeshCommandMaskAndStatus(FRayTracingMeshCommand& MeshCommand, const FMeshBatch& MeshBatch, const FPrimitiveSceneProxy* PrimitiveSceneProxy,
 	const FMaterial& MaterialResource, ERayTracingViewMaskMode MaskMode);
 
 void UpdateRayTracingMeshCommandMasks(FRayTracingMeshCommand& RayTracingMeshCommands, const ERayTracingPrimitiveFlags Flags, ERayTracingViewMaskMode MaskMode);
diff --git a/Engine/Source/Runtime/Renderer/Private/RayTracing/RayTracingMaterialHitShaders.h b/Engine/Source/Runtime/Renderer/Private/RayTracing/RayTracingMaterialHitShaders.h
--- a/Engine/Source/Runtime/Renderer/Private/RayTracing/RayTracingMaterialHitShaders.h
+++ b/Engine/Source/Runtime/Renderer/Private/RayTracing/RayTracingMaterialHitShaders.h
@@ -60,7 +60,7 @@ protected:
 
 		FRayTracingMeshCommand SharedCommand;
 
-		SetupRayTracingMeshCommandMaskAndStatus(SharedCommand, MeshBatch, *PrimitiveSceneProxy, MaterialResource, MaskMode);
+		SetupRayTracingMeshCommandMaskAndStatus(SharedCommand, MeshBatch, PrimitiveSceneProxy, MaterialResource, MaskMode);
 
 		if (GRHISupportsRayTracingShaders)
 		{
