Small cleanup to auto exposure shader code.
- Renamed a couple of shader permutations
- Renamed some local variables that were misleading

diff --git a/Engine/Shaders/Private/Histogram.usf b/Engine/Shaders/Private/Histogram.usf
--- a/Engine/Shaders/Private/Histogram.usf
+++ b/Engine/Shaders/Private/Histogram.usf
@@ -319,10 +319,9 @@ void MainAtomicCS(
 
 			float LuminanceVal = CalculateEyeAdaptationLuminance(SceneColor);
 
-#if USE_APPROX_ILLUMINANCE
-#	if USE_PRECALCULATED_ILLUMINANCE
+#if USE_PRECALCULATED_LUMINANCE
 			LuminanceVal = SceneColor.r;
-#	else
+#elif USE_APPROX_ILLUMINANCE
 			const uint2 GBufferTexelPos = View.ViewRectMinAndSize.xy + ScreenUV * View.ViewRectMinAndSize.zw;
 			const float2 GBufferUV = (GBufferTexelPos + 0.5f) * View.BufferSizeAndInvSize.zw;
 
@@ -358,10 +357,9 @@ void MainAtomicCS(
 				LuminanceVal /= max(BaseColorLuminanceVal, EyeAdaptation_IgnoreMaterialsMinBaseColorLuminance);
 				LuminanceVal *= lerp(EyeAdaptation_IgnoreMaterialsLuminanceScale, 1.0f, EyeAdaptation_IgnoreMaterialsMinBaseColorLuminance);
 			}
-#	endif
 #endif
 
-#if USE_APPROX_ILLUMINANCE_DEBUG
+#if USE_DEBUG_OUTPUT
 			DebugOutput[TexelPos - Input_ViewportMin] = LuminanceVal;
 #endif
 
diff --git a/Engine/Shaders/Private/PostProcessEyeAdaptation.usf b/Engine/Shaders/Private/PostProcessEyeAdaptation.usf
--- a/Engine/Shaders/Private/PostProcessEyeAdaptation.usf
+++ b/Engine/Shaders/Private/PostProcessEyeAdaptation.usf
@@ -374,7 +374,7 @@ void CalculateExposureIlluminance(uint2 ThreadId : SV_DispatchThreadID)
 
 	const float3 SceneColor = ColorTexture.Load(int3(TexelPos, 0)).rgb * View.OneOverPreExposure;
 
-	float Illuminance = CalculateEyeAdaptationLuminance(SceneColor);
+	float Luminance = CalculateEyeAdaptationLuminance(SceneColor);
 
 	const float2 GBufferUV = (TexelPos + 0.5f) * View.BufferSizeAndInvSize.zw;
 	const FGBufferData GBufferData = GetGBufferDataFromSceneTextures(GBufferUV);
@@ -391,10 +391,12 @@ void CalculateExposureIlluminance(uint2 ThreadId : SV_DispatchThreadID)
 			bUseTranslucencyVolume = true;
 		}
 
+		float Illuminance;
+
 		if (!bUseTranslucencyVolume)
 		{
 			// Calculate approximate illuminance from SceneColor and GBufferData
-			// Illuminance ~= (SceneColor - Emissive) / (DiffuseColor + SubsurfaceColor + EnvBRDF)
+			// Illuminance ~= (SceneColorLuminance - Emissive) / (DiffuseColor + SubsurfaceColor + EnvBRDF)
 
 			float3 Denominator = 0.0f;
 
@@ -428,10 +430,10 @@ void CalculateExposureIlluminance(uint2 ThreadId : SV_DispatchThreadID)
 
 			const float DenominatorLuminance = CalculateEyeAdaptationLuminanceWithoutMin(Denominator);
 
+			// Emissive is copied to IlluminanceTexture after base pass and before lighting pass
 			const float Emissive = RWIlluminanceTexture[OutputTexelPos] * View.OneOverPreExposure;
 
-			Illuminance = max(0.0f, Illuminance - Emissive);
-			Illuminance /= max(DenominatorLuminance, EyeAdaptation_IgnoreMaterialsMinBaseColorLuminance);
+			Illuminance = max(0.0f, Luminance - Emissive) / max(DenominatorLuminance, EyeAdaptation_IgnoreMaterialsMinBaseColorLuminance);
 		}
 		else
 		{
@@ -482,10 +484,10 @@ void CalculateExposureIlluminance(uint2 ThreadId : SV_DispatchThreadID)
 			Illuminance = CalculateEyeAdaptationLuminance(VolumeLighting.rgb);
 		}
 
-		Illuminance *= lerp(EyeAdaptation_IgnoreMaterialsLuminanceScale, 1.0f, EyeAdaptation_IgnoreMaterialsMinBaseColorLuminance);
+		Luminance = Illuminance * lerp(EyeAdaptation_IgnoreMaterialsLuminanceScale, 1.0f, EyeAdaptation_IgnoreMaterialsMinBaseColorLuminance);
 	}
 	
-	RWIlluminanceTexture[OutputTexelPos] = Illuminance * View.PreExposure;
+	RWIlluminanceTexture[OutputTexelPos] = Luminance * View.PreExposure;
 }
 
 #endif // COMPUTESHADER
diff --git a/Engine/Shaders/Private/PostProcessHistogramCommon.ush b/Engine/Shaders/Private/PostProcessHistogramCommon.ush
--- a/Engine/Shaders/Private/PostProcessHistogramCommon.ush
+++ b/Engine/Shaders/Private/PostProcessHistogramCommon.ush
@@ -1,7 +1,7 @@
 // Copyright Epic Games, Inc. All Rights Reserved.
 
 /*=============================================================================
-	PostProcessHistogramCommon.usf: PostProcessing histogram shared functions and structures.
+	PostProcessHistogramCommon.ush: PostProcessing histogram shared functions and structures.
 =============================================================================*/
 
 float EyeAdaptation_ExposureLowPercent;
diff --git a/Engine/Source/Runtime/Renderer/Private/PostProcess/PostProcessHistogram.cpp b/Engine/Source/Runtime/Renderer/Private/PostProcess/PostProcessHistogram.cpp
--- a/Engine/Source/Runtime/Renderer/Private/PostProcess/PostProcessHistogram.cpp
+++ b/Engine/Source/Runtime/Renderer/Private/PostProcess/PostProcessHistogram.cpp
@@ -149,13 +149,20 @@ public:
 		SHADER_PARAMETER_SAMPLER(SamplerState, PreIntegratedGFSampler)
 	END_SHADER_PARAMETER_STRUCT()
 
+	class FUsePrecalculatedLuminanceDim : SHADER_PERMUTATION_BOOL("USE_PRECALCULATED_LUMINANCE");
 	class FUseApproxIlluminanceDim : SHADER_PERMUTATION_BOOL("USE_APPROX_ILLUMINANCE");
-	class FUsePrecalculatedIlluminanceDim : SHADER_PERMUTATION_BOOL("USE_PRECALCULATED_ILLUMINANCE");
-	class FUseApproxIlluminanceDebugDim : SHADER_PERMUTATION_BOOL("USE_APPROX_ILLUMINANCE_DEBUG");
-	using FPermutationDomain = TShaderPermutationDomain<FUseApproxIlluminanceDim, FUsePrecalculatedIlluminanceDim, FUseApproxIlluminanceDebugDim>;
+	class FUseDebugOutputDim : SHADER_PERMUTATION_BOOL("USE_DEBUG_OUTPUT");
+	using FPermutationDomain = TShaderPermutationDomain<FUsePrecalculatedLuminanceDim, FUseApproxIlluminanceDim, FUseDebugOutputDim>;
 
 	static bool ShouldCompilePermutation(const FGlobalShaderPermutationParameters& Parameters)
 	{
+		const FPermutationDomain PermutationVector(Parameters.PermutationId);
+
+		if (PermutationVector.Get<FUsePrecalculatedLuminanceDim>() && PermutationVector.Get<FUseApproxIlluminanceDim>())
+		{
+			return false;
+		}
+
 		return IsFeatureLevelSupported(Parameters.Platform, ERHIFeatureLevel::SM5);
 	}
 
@@ -368,9 +375,18 @@ static FRDGTextureRef AddHistogramAtomicPass(
 		static const auto CVarUsePrecalculatedIlluminance = IConsoleManager::Get().FindTConsoleVariableDataBool(TEXT("r.AutoExposure.IgnoreMaterials.UsePrecalculatedIlluminance"));
 
 		FHistogramAtomicCS::FPermutationDomain PermutationVector;
-		PermutationVector.Set<FHistogramAtomicCS::FUseApproxIlluminanceDim>(CVarIgnoreMaterials->GetValueOnRenderThread());
-		PermutationVector.Set<FHistogramAtomicCS::FUsePrecalculatedIlluminanceDim>(CVarUsePrecalculatedIlluminance->GetValueOnRenderThread());
-		PermutationVector.Set<FHistogramAtomicCS::FUseApproxIlluminanceDebugDim>(CVarIgnoreMaterialsDebug->GetValueOnRenderThread());
+		if (CVarIgnoreMaterials->GetValueOnRenderThread())
+		{
+			if (CVarUsePrecalculatedIlluminance->GetValueOnRenderThread())
+			{
+				PermutationVector.Set<FHistogramAtomicCS::FUsePrecalculatedLuminanceDim>(true);
+			}
+			else
+			{
+				PermutationVector.Set<FHistogramAtomicCS::FUseApproxIlluminanceDim>(true);
+			}
+		}
+		PermutationVector.Set<FHistogramAtomicCS::FUseDebugOutputDim>(CVarIgnoreMaterialsDebug->GetValueOnRenderThread());
 
 		auto ComputeShader = View.ShaderMap->GetShader<FHistogramAtomicCS>(PermutationVector);
 
