Fix for the rebuild of many reflection captures.

PR #9707: Fix for crash in editor when baking new reflection captures (Contributed by [REDACTED_PERSON])

diff --git a/Engine/Source/Runtime/Renderer/Private/ReflectionEnvironment.cpp b/Engine/Source/Runtime/Renderer/Private/ReflectionEnvironment.cpp
--- a/Engine/Source/Runtime/Renderer/Private/ReflectionEnvironment.cpp
+++ b/Engine/Source/Runtime/Renderer/Private/ReflectionEnvironment.cpp
@@ -266,7 +266,9 @@ FCaptureComponentSceneState& FReflectionCaptureCache::FindChecked(const UReflect
 
 FCaptureComponentSceneState& FReflectionCaptureCache::Add(const UReflectionCaptureComponent* Component, const FCaptureComponentSceneState& Value)
 {
-	check(IsValid(Component))
+	// During Reflection Capture Placement in editor, this is potentially not IsValid
+	//  So just check to make sure that the pointer is non-null
+	check(Component)
 
 	FCaptureComponentSceneState* Existing = AddReference(Component);
 	if (Existing != nullptr)
@@ -283,7 +285,9 @@ FCaptureComponentSceneState& FReflectionCaptureCache::Add(const UReflectionCaptu
 
 FCaptureComponentSceneState* FReflectionCaptureCache::AddReference(const UReflectionCaptureComponent* Component)
 {
-	check(IsValid(Component))
+	// During Reflection Capture Placement in editor, this is potentially not IsValid
+	//  So just check to make sure that the pointer is non-null
+	check(Component)
 
 	bool Remap = RemapRegisteredComponentMapBuildDataId(Component);
 	FReflectionCaptureCacheEntry* Found = CaptureData.Find(Component->MapBuildDataId);
@@ -315,6 +319,31 @@ void FReflectionCaptureCache::Empty()
 	RegisteredComponentMapBuildDataIds.Empty();
 }
 
+int32 FReflectionCaptureCache::Prune(const TSet<FGuid> KeysToKeep, TArray<int32>& ReleasedIndices)
+{
+	TSet<FGuid> ExistingKeys;
+	CaptureData.GetKeys(ExistingKeys);
+
+	TSet<FGuid> KeysToRemove = ExistingKeys.Difference(KeysToKeep);
+	ReleasedIndices.Empty();
+	ReleasedIndices.Reserve(KeysToRemove.Num());
+
+	for (const FGuid& Key : KeysToRemove)
+	{
+		FReflectionCaptureCacheEntry* Found = CaptureData.Find(Key);
+		if (Found == nullptr)
+			continue;
+
+		int32 CubemapIndex = Found->SceneState.CubemapIndex;
+		if (CubemapIndex != -1)
+			ReleasedIndices.Add(CubemapIndex);
+
+		CaptureData.Remove(Key);
+	}
+
+	return ReleasedIndices.Num();
+}
+
 bool FReflectionCaptureCache::Remove(const UReflectionCaptureComponent* Component)
 {
 	if (!Component)	// Intentionally not IsValid(Component), as this often occurs when Component is explicitly PendingKill.
@@ -353,6 +382,9 @@ bool FReflectionCaptureCache::RemapRegisteredComponentMapBuildDataId(const URefl
 		*OldBuildId != Component->MapBuildDataId)
 	{
 		FGuid OldBuildIdCopy = *OldBuildId;
+		const FReflectionCaptureCacheEntry* Current = CaptureData.Find(OldBuildIdCopy);
+		if (Current == nullptr)
+			return false; // No current entry to remap to, so no remap to perform
 
 		// Remap all pointers that point to the old guid to the new one.
 		int32 ReferenceCount = 0;
@@ -365,7 +397,7 @@ bool FReflectionCaptureCache::RemapRegisteredComponentMapBuildDataId(const URefl
 			}
 		}
 
-		FReflectionCaptureCacheEntry Entry = *CaptureData.Find(OldBuildIdCopy);
+		FReflectionCaptureCacheEntry Entry = *Current;
 		Entry.RefCount = ReferenceCount;
 
 		CaptureData.Remove(OldBuildIdCopy);
diff --git a/Engine/Source/Runtime/Renderer/Private/ReflectionEnvironmentCapture.cpp b/Engine/Source/Runtime/Renderer/Private/ReflectionEnvironmentCapture.cpp
--- a/Engine/Source/Runtime/Renderer/Private/ReflectionEnvironmentCapture.cpp
+++ b/Engine/Source/Runtime/Renderer/Private/ReflectionEnvironmentCapture.cpp
@@ -863,6 +863,35 @@ void FScene::AllocateReflectionCaptures(const TArray<UReflectionCaptureComponent
 			}
 #endif
 
+			if (DesiredMaxCubemaps > 0)
+			{
+				// In the editor, some captures might have been marked dirty and requested new guids
+				//   So remove any probes that can no longer be referenced from the cache
+
+				TSet<FGuid> ToKeep;
+				ToKeep.Reserve(DesiredMaxCubemaps);
+				for (TSparseArray<UReflectionCaptureComponent*>::TIterator It(ReflectionSceneData.AllocatedReflectionCapturesGameThread); It; ++It)
+				{
+					if (IsValid(*It) && (*It)->MapBuildDataId.IsValid())
+					{
+						ToKeep.Add((*It)->MapBuildDataId);
+					}
+				}
+
+				FScene* Scene = this;
+				ENQUEUE_RENDER_COMMAND(PruneReflectionCaptures)(
+					[Scene, ToKeep = MoveTemp(ToKeep)](FRHICommandListImmediate& RHICmdList)
+				{
+					TArray<int32> ReleasedIndices;
+					Scene->ReflectionSceneData.AllocatedReflectionCaptureState.Prune(ToKeep, ReleasedIndices);
+
+					for (int32 Index : ReleasedIndices)
+					{
+						Scene->ReflectionSceneData.CubemapArraySlotsUsed[Index] = false;
+					}
+				});
+			}
+
 			// If this is not the first time the scene has allocated the cubemap array, include slack to reduce reallocations
 			const float MaxCubemapsRoundUpBase = 1.5f;
 			if (ReflectionSceneData.MaxAllocatedReflectionCubemapsGameThread > 0)
diff --git a/Engine/Source/Runtime/Renderer/Private/ScenePrivate.h b/Engine/Source/Runtime/Renderer/Private/ScenePrivate.h
--- a/Engine/Source/Runtime/Renderer/Private/ScenePrivate.h
+++ b/Engine/Source/Runtime/Renderer/Private/ScenePrivate.h
@@ -1612,6 +1612,7 @@ public:
 	FCaptureComponentSceneState& Add(const UReflectionCaptureComponent* Component, const FCaptureComponentSceneState& Value);
 	FCaptureComponentSceneState* AddReference(const UReflectionCaptureComponent* Component);
 	bool Remove(const UReflectionCaptureComponent* Component);
+	int32 Prune(const TSet<FGuid> KeysToKeep, TArray<int32>& ReleasedIndices);
 
 	int32 GetKeys(TArray<FGuid>& OutKeys) const;
 	int32 GetKeys(TSet<FGuid>& OutKeys) const;
