Fix over allocation of tile data buffer.

diff --git a/Engine/Shaders/Private/HairStrands/HairStrandsForwardRaster.usf b/Engine/Shaders/Private/HairStrands/HairStrandsForwardRaster.usf
--- a/Engine/Shaders/Private/HairStrands/HairStrandsForwardRaster.usf
+++ b/Engine/Shaders/Private/HairStrands/HairStrandsForwardRaster.usf
@@ -1385,8 +1385,8 @@ void RasterCS(
 
 			// For debug only
 			#if PERMUTATION_DEBUG
-			const bool bDebugEnableAll	= all(RasterTileCoord == uint2(ShaderPrintData.CursorCoord / float(RASTER_TILE_SIZE)));
-			const bool bDebugEnable		= all(RasterTileCoord == uint2(ShaderPrintData.CursorCoord / float(RASTER_TILE_SIZE))) && GroupThread1D == 0;
+			const bool bDebugEnableAll	= FrontToBackIndex == 0 && all(RasterTileCoord == uint2(ShaderPrintData.CursorCoord / float(RASTER_TILE_SIZE)));
+			const bool bDebugEnable		= FrontToBackIndex == 0 && all(RasterTileCoord == uint2(ShaderPrintData.CursorCoord / float(RASTER_TILE_SIZE))) && GroupThread1D == 0;
 			FShaderPrintContext Ctx		= InitShaderPrintContext(bDebugEnable, uint2(350, 50));
 			const FFontColor FontLegend = FontWhite;
 			const FFontColor FontValue  = FontOrange;
@@ -1958,9 +1958,14 @@ void RasterCS(
 Buffer<uint>						VisTileArgs;
 ByteAddressBuffer					VisTileData;
 
+Buffer<uint>						CompactedVisTileArgs;
+ByteAddressBuffer					CompactedVisTileData;
+
 Texture2D<uint>						HairCountTexture_ForDebug;
 Texture2D<uint>						HairPixelCountPerTile_ForDebug;
 uint								InstanceCount;
+uint								CPUAllocatedTileCount;
+uint								CPUAllocatedCompactedTileCount;
 
 uint GetTileTotalSegment(uint2 TileCoord, bool bPrintDetails, uint InTileSize)
 {
@@ -2040,7 +2045,11 @@ void MainCS(uint3 ThreadId : SV_DispatchThreadID)
 		Print(Context, TEXT("Num Rasterizers      : "), FontLegend); Print(Context, NumRasterizers, FontValue); Newline(Context);
 		Newline(Context);
 
-		Print(Context, TEXT("Alloc. Tile Count    : "), FontLegend); Print(Context, VisTileArgs[0], FontValue); Newline(Context);
+		const FFontColor AllocColor			 = InitFontColor(ColorMapTurbo(VisTileArgs[0] / float(CPUAllocatedTileCount)));
+		const FFontColor AllocCompactedColor = InitFontColor(ColorMapTurbo(CompactedVisTileArgs[0] / float(CPUAllocatedCompactedTileCount)));
+
+		Print(Context, TEXT("Alloc. Tile          : "), FontLegend); Print(Context, VisTileArgs[0], AllocColor); Print(Context, TEXT(" / "), FontLegend); Print(Context, CPUAllocatedTileCount, FontValue); Newline(Context);
+		Print(Context, TEXT("Alloc. Compacted Tile: "), FontLegend); Print(Context, CompactedVisTileArgs[0], AllocColor); Print(Context, TEXT(" / "), FontLegend); Print(Context, CPUAllocatedCompactedTileCount, FontValue); Newline(Context);
 		Print(Context, TEXT("Rasterizer Max Work  : "), FontLegend); Print(Context, NumRasterizers * MAX_WORK_COUNT, FontValue); Newline(Context);
 		Newline(Context);
 		
diff --git a/Engine/Source/Runtime/Renderer/Private/HairStrands/HairStrandsForwardRaster.cpp b/Engine/Source/Runtime/Renderer/Private/HairStrands/HairStrandsForwardRaster.cpp
--- a/Engine/Source/Runtime/Renderer/Private/HairStrands/HairStrandsForwardRaster.cpp
+++ b/Engine/Source/Runtime/Renderer/Private/HairStrands/HairStrandsForwardRaster.cpp
@@ -246,10 +246,14 @@ class FHairStrandsForwardRasterDebugCS : public FGlobalShader
 	BEGIN_SHADER_PARAMETER_STRUCT(FParameters, )
 		SHADER_PARAMETER_STRUCT_INCLUDE(FRasterComputeForwardCommonParameters, Common)
 		SHADER_PARAMETER(uint32, InstanceCount)
+		SHADER_PARAMETER(uint32, CPUAllocatedTileCount)
+		SHADER_PARAMETER(uint32, CPUAllocatedCompactedTileCount)
 		SHADER_PARAMETER_RDG_TEXTURE(Texture2D, HairCountTexture_ForDebug)
 		SHADER_PARAMETER_RDG_TEXTURE(Texture2D, HairPixelCountPerTile_ForDebug)
 		SHADER_PARAMETER_RDG_BUFFER_SRV(Buffer, VisTileArgs)
 		SHADER_PARAMETER_RDG_BUFFER_SRV(ByteAddressBuffer, VisTileData)
+		SHADER_PARAMETER_RDG_BUFFER_SRV(Buffer, CompactedVisTileArgs)
+		SHADER_PARAMETER_RDG_BUFFER_SRV(ByteAddressBuffer, CompactedVisTileData)
 
 		SHADER_PARAMETER_STRUCT_REF(FViewUniformShaderParameters, ViewUniformBuffer)
 		SHADER_PARAMETER_STRUCT_INCLUDE(ShaderPrint::FShaderParameters, ShaderPrintParameters)
@@ -490,23 +494,26 @@ void AddHairStrandsForwardRasterPass(
 		FRDGBufferRef CompactedVisTileWorkCount;
 	};
 
+	const uint32 EntryCount = 5; // See VT_XXX in .usf: PrimOffset / PrimCount / Coord / MaxIndex / MinMaxDepth
+
 	FBinningData BinData;
 	{
-		const uint32 BinMaxTiles = FMath::Min(FMath::Max(GHairVisibilityComputeRaster_MaxTiles, 1024), 262144); //?
+		const uint32 MaxPrimPerTile = 1024u;
+		const uint32 BinMaxTiles = FMath::Clamp(GHairVisibilityComputeRaster_MaxTiles, MaxPrimPerTile, 262144); //?
+		const uint32 MaxTotalPrims = BinMaxTiles * MaxPrimPerTile; // Number of 'binned' primitives. This number can/needs to be higher than the actual number of primitives, since a primitive can be binned into several bins
 
-		BinData.VisTilePrims = GraphBuilder.CreateBuffer(FRDGBufferDesc::CreateBufferDesc(sizeof(uint32), BinMaxTiles * 1024), TEXT("Hair.VisTilePrims"));
+		BinData.VisTilePrims = GraphBuilder.CreateBuffer(FRDGBufferDesc::CreateBufferDesc(sizeof(uint32), MaxTotalPrims), TEXT("Hair.VisTilePrims"));
 		BinData.VisTilePrimsSRV = GraphBuilder.CreateSRV(BinData.VisTilePrims, PF_R32_UINT);
 
-		BinData.VisTilePrimDepths = GraphBuilder.CreateBuffer(FRDGBufferDesc::CreateBufferDesc(sizeof(uint32), BinMaxTiles * 1024), TEXT("Hair.VisTilePrimDepths"));
+		BinData.VisTilePrimDepths = GraphBuilder.CreateBuffer(FRDGBufferDesc::CreateBufferDesc(sizeof(uint32), MaxTotalPrims), TEXT("Hair.VisTilePrimDepths"));
 		BinData.VisTilePrimDepthsSRV = GraphBuilder.CreateSRV(BinData.VisTilePrimDepths, PF_R32_UINT);
 
-		const uint32 EntryCount = 5; // See VT_XXX in .usf: PrimOffset / PrimCount / Coord / MaxIndex / MinMaxDepth
-		FRDGBufferDesc DescTileData = FRDGBufferDesc::CreateStructuredDesc(sizeof(uint32), BinMaxTiles * EntryCount * 2);
+		FRDGBufferDesc DescTileData = FRDGBufferDesc::CreateStructuredDesc(sizeof(uint32), BinMaxTiles * EntryCount);
 		DescTileData.Usage = EBufferUsageFlags(DescTileData.Usage | BUF_ByteAddressBuffer);
 		BinData.VisTileData = GraphBuilder.CreateBuffer(DescTileData, TEXT("Hair.VisTileData"));
 		BinData.VisTileDataSRV = GraphBuilder.CreateSRV(BinData.VisTileData);
 
-		FRDGBufferDesc VisTileArgsDesc = FRDGBufferDesc::CreateBufferDesc(sizeof(uint32), 4 + 16 + (((BinTileRes.X * BinTileRes.Y) + 31) / 32));
+		FRDGBufferDesc VisTileArgsDesc = FRDGBufferDesc::CreateBufferDesc(sizeof(uint32), 1);
 		BinData.VisTileArgs = GraphBuilder.CreateBuffer(VisTileArgsDesc, TEXT("Hair.VisTileArgs"));
 		BinData.VisTileArgsUAV = GraphBuilder.CreateUAV(BinData.VisTileArgs, PF_R32_UINT);
 		BinData.VisTileArgsSRV = GraphBuilder.CreateSRV(BinData.VisTileArgs, PF_R32_UINT);
@@ -520,10 +527,10 @@ void AddHairStrandsForwardRasterPass(
 		BinData.VisTileBinningGridMinZUAV = GraphBuilder.CreateUAV(BinData.VisTileBinningGridMinZ);
 		BinData.VisTileBinningGridMaxZUAV = GraphBuilder.CreateUAV(BinData.VisTileBinningGridMaxZ);
 
-		BinData.CompactedVisTilePrims = GraphBuilder.CreateBuffer(FRDGBufferDesc::CreateBufferDesc(sizeof(uint32), BinMaxTiles * 1024), TEXT("Hair.CompactedVisTilePrims"));
+		BinData.CompactedVisTilePrims = GraphBuilder.CreateBuffer(FRDGBufferDesc::CreateBufferDesc(sizeof(uint32), MaxTotalPrims), TEXT("Hair.CompactedVisTilePrims"));
 		BinData.CompactedVisTileData = GraphBuilder.CreateBuffer(DescTileData, TEXT("Hair.CompactedVisTileData"));
 		BinData.CompactedVisTileArgs = GraphBuilder.CreateBuffer(VisTileArgsDesc, TEXT("Hair.CompactedVisTileArgs"));
-		BinData.CompactedVisTileWork		= GraphBuilder.CreateBuffer(FRDGBufferDesc::CreateStructuredDesc(sizeof(uint32), BinMaxTiles * 1024), TEXT("Hair.CompactedVisTileWork"));
+		BinData.CompactedVisTileWork		= GraphBuilder.CreateBuffer(FRDGBufferDesc::CreateStructuredDesc(sizeof(uint32), MaxTotalPrims), TEXT("Hair.CompactedVisTileWork"));
 		BinData.CompactedVisTileDataCount	= GraphBuilder.CreateBuffer(FRDGBufferDesc::CreateStructuredDesc(sizeof(uint32), 1), TEXT("Hair.CompactedVisTileDataCount"));
 		BinData.CompactedVisTileWorkCount	= GraphBuilder.CreateBuffer(FRDGBufferDesc::CreateStructuredDesc(sizeof(uint32), 1), TEXT("Hair.CompactedVisTileWorkCount"));
 	}
@@ -696,8 +703,12 @@ void AddHairStrandsForwardRasterPass(
 		FHairStrandsForwardRasterDebugCS::FParameters* Parameters = GraphBuilder.AllocParameters<FHairStrandsForwardRasterDebugCS::FParameters>();
 		Parameters->Common = Common;
 		Parameters->InstanceCount = InData.RasterizedInstanceCount;
+		Parameters->CPUAllocatedTileCount = BinData.VisTileData->Desc.NumElements / EntryCount;
+		Parameters->CPUAllocatedCompactedTileCount = BinData.CompactedVisTileData->Desc.NumElements / EntryCount;
 		Parameters->VisTileData = GraphBuilder.CreateSRV(BinData.CompactedVisTileData);
 		Parameters->VisTileArgs = BinData.VisTileArgsSRV;
+		Parameters->CompactedVisTileData = GraphBuilder.CreateSRV(BinData.CompactedVisTileData);
+		Parameters->CompactedVisTileArgs = GraphBuilder.CreateSRV(BinData.CompactedVisTileArgs, PF_R32_UINT);
 		Parameters->HairCountTexture_ForDebug = HairCountTexture_ForDebug;
 		Parameters->HairPixelCountPerTile_ForDebug = HairPixelCountPerTile_ForDebug;
 		ShaderPrint::SetParameters(GraphBuilder, ViewInfo.ShaderPrintData, Parameters->ShaderPrintParameters);
