Moving LightMapRendering implementations to CPP file to reduce header dependencies.

diff --git a/Engine/Source/Runtime/Renderer/Private/LightMapRendering.cpp b/Engine/Source/Runtime/Renderer/Private/LightMapRendering.cpp
--- a/Engine/Source/Runtime/Renderer/Private/LightMapRendering.cpp
+++ b/Engine/Source/Runtime/Renderer/Private/LightMapRendering.cpp
@@ -5,6 +5,7 @@ LightMapRendering.cpp: Light map rendering implementations.
 =============================================================================*/
 
 #include "LightMapRendering.h"
+#include "LightMap.h"
 #include "ScenePrivate.h"
 #include "PrecomputedVolumetricLightmap.h"
 #include "RenderCore.h"
@@ -13,6 +14,14 @@ IMPLEMENT_TYPE_LAYOUT(FUniformLightMapPolicyShaderParametersType);
 
 IMPLEMENT_GLOBAL_SHADER_PARAMETER_STRUCT(FIndirectLightingCacheUniformParameters, "IndirectLightingCache");
 
+bool MobileUseCSMShaderBranch()
+{
+	static const auto* CVar = IConsoleManager::Get().FindTConsoleVariableDataInt(TEXT("r.Mobile.UseCSMShaderBranch"));
+	return (CVar && CVar->GetValueOnAnyThread() != 0);
+}
+
+// One of these per lightmap quality
+
 const TCHAR* GLightmapDefineName[2] =
 {
 	TEXT("LQ_TEXTURE_LIGHTMAP"),
@@ -25,11 +34,281 @@ int32 GNumLightmapCoefficients[2] =
 	NUM_HQ_LIGHTMAP_COEF
 };
 
+
+void LightMapPolicyImpl::ModifyCompilationEnvironment(ELightmapQuality LightmapQuality, const FMaterialShaderPermutationParameters& Parameters, FShaderCompilerEnvironment& OutEnvironment)
+{
+	OutEnvironment.SetDefine(GLightmapDefineName[LightmapQuality], TEXT("1"));
+	OutEnvironment.SetDefine(TEXT("NUM_LIGHTMAP_COEFFICIENTS"), GNumLightmapCoefficients[LightmapQuality]);
+
+	static const auto CVar = IConsoleManager::Get().FindTConsoleVariableDataInt(TEXT("r.VirtualTexturedLightmaps"));
+	const bool VirtualTextureLightmaps = (CVar->GetValueOnAnyThread() != 0) && UseVirtualTexturing(GMaxRHIFeatureLevel, OutEnvironment.TargetPlatform);
+	OutEnvironment.SetDefine(TEXT("LIGHTMAP_VT_ENABLED"), VirtualTextureLightmaps);
+}
+
+bool LightMapPolicyImpl::ShouldCompilePermutation(ELightmapQuality LightmapQuality, const FMeshMaterialShaderPermutationParameters& Parameters)
+{
+	static const auto AllowStaticLightingVar = IConsoleManager::Get().FindTConsoleVariableDataInt(TEXT("r.AllowStaticLighting"));
+	static const auto CVarProjectCanHaveLowQualityLightmaps = IConsoleManager::Get().FindTConsoleVariableDataInt(TEXT("r.SupportLowQualityLightmaps"));
+	static const auto CVarSupportAllShadersPermutations = IConsoleManager::Get().FindTConsoleVariableDataInt(TEXT("r.SupportAllShaderPermutations"));
+	const bool bForceAllPermutations = CVarSupportAllShadersPermutations && CVarSupportAllShadersPermutations->GetValueOnAnyThread() != 0;
+
+	// if GEngine doesn't exist yet to have the project flag then we should be conservative and cache the LQ lightmap policy
+	const bool bProjectCanHaveLowQualityLightmaps = bForceAllPermutations || (!CVarProjectCanHaveLowQualityLightmaps) || (CVarProjectCanHaveLowQualityLightmaps->GetValueOnAnyThread() != 0);
+
+	const bool bShouldCacheQuality = (LightmapQuality != ELightmapQuality::LQ_LIGHTMAP) || bProjectCanHaveLowQualityLightmaps;
+
+	// GetValueOnAnyThread() as it's possible that ShouldCache is called from rendering thread. That is to output some error message.
+	return (Parameters.MaterialParameters.ShadingModels.IsLit())
+		&& bShouldCacheQuality
+		&& Parameters.VertexFactoryType->SupportsStaticLighting()
+		&& (!AllowStaticLightingVar || AllowStaticLightingVar->GetValueOnAnyThread() != 0)
+		&& (Parameters.MaterialParameters.bIsUsedWithStaticLighting || Parameters.MaterialParameters.bIsSpecialEngineMaterial);
+}
+
+void DistanceFieldShadowsAndLightMapPolicyImpl::ModifyCompilationEnvironment(ELightmapQuality LightmapQuality, const FMaterialShaderPermutationParameters& Parameters, FShaderCompilerEnvironment& OutEnvironment)
+{
+	OutEnvironment.SetDefine(TEXT("STATICLIGHTING_TEXTUREMASK"), 1);
+	OutEnvironment.SetDefine(TEXT("STATICLIGHTING_SIGNEDDISTANCEFIELD"), 1);
+	LightMapPolicyImpl::ModifyCompilationEnvironment(LightmapQuality, Parameters, OutEnvironment);
+}
+
+bool FDummyLightMapPolicy::ShouldCompilePermutation(const FMeshMaterialShaderPermutationParameters& Parameters)
+{
+	return Parameters.MaterialParameters.ShadingModels.IsLit() && Parameters.VertexFactoryType->SupportsStaticLighting();
+}
+
+bool FSelfShadowedTranslucencyPolicy::ShouldCompilePermutation(const FMeshMaterialShaderPermutationParameters& Parameters)
+{
+	return Parameters.MaterialParameters.ShadingModels.IsLit() &&
+		IsTranslucentBlendMode(Parameters.MaterialParameters) &&
+		IsFeatureLevelSupported(Parameters.Platform, ERHIFeatureLevel::SM5);
+}
+
+void FSelfShadowedTranslucencyPolicy::ModifyCompilationEnvironment(const FMaterialShaderPermutationParameters& Parameters, FShaderCompilerEnvironment& OutEnvironment)
+{
+	OutEnvironment.SetDefine(TEXT("TRANSLUCENT_SELF_SHADOWING"), TEXT("1"));
+}
+
+FSelfShadowedTranslucencyPolicy::FSelfShadowedTranslucencyPolicy()
+{
+}
+
+void FSelfShadowedTranslucencyPolicy::GetVertexShaderBindings(
+	const FPrimitiveSceneProxy* PrimitiveSceneProxy,
+	const ElementDataType& ShaderElementData,
+	const VertexParametersType* VertexShaderParameters,
+	FMeshDrawSingleShaderBindings& ShaderBindings)
+{
+}
+
+void FSelfShadowedTranslucencyPolicy::GetPixelShaderBindings(
+	const FPrimitiveSceneProxy* PrimitiveSceneProxy,
+	const ElementDataType& ShaderElementData,
+	const PixelParametersType* PixelShaderParameters,
+	FMeshDrawSingleShaderBindings& ShaderBindings)
+{
+	ShaderBindings.Add(PixelShaderParameters->TranslucentSelfShadowBufferParameter, ShaderElementData);
+}
+
+void FSelfShadowedTranslucencyPolicy::GetComputeShaderBindings(
+	const FPrimitiveSceneProxy* PrimitiveSceneProxy,
+	const ElementDataType& ShaderElementData,
+	const ComputeParametersType* ComputeShaderParameters,
+	FMeshDrawSingleShaderBindings& ShaderBindings)
+{
+	ShaderBindings.Add(ComputeShaderParameters->TranslucentSelfShadowBufferParameter, ShaderElementData);
+}
+
+bool FPrecomputedVolumetricLightmapLightingPolicy::ShouldCompilePermutation(const FMeshMaterialShaderPermutationParameters& Parameters)
+{
+	static const auto AllowStaticLightingVar = IConsoleManager::Get().FindTConsoleVariableDataInt(TEXT("r.AllowStaticLighting"));
+
+	return Parameters.MaterialParameters.ShadingModels.IsLit() &&
+		(!AllowStaticLightingVar || AllowStaticLightingVar->GetValueOnAnyThread() != 0);
+}
+
+void FPrecomputedVolumetricLightmapLightingPolicy::ModifyCompilationEnvironment(const FMaterialShaderPermutationParameters& Parameters, FShaderCompilerEnvironment& OutEnvironment)
+{
+	OutEnvironment.SetDefine(TEXT("PRECOMPUTED_IRRADIANCE_VOLUME_LIGHTING"), TEXT("1"));
+}
+
+bool FCachedVolumeIndirectLightingPolicy::ShouldCompilePermutation(const FMeshMaterialShaderPermutationParameters& Parameters)
+{
+	static const auto AllowStaticLightingVar = IConsoleManager::Get().FindTConsoleVariableDataInt(TEXT("r.AllowStaticLighting"));
+
+	return Parameters.MaterialParameters.ShadingModels.IsLit()
+		&& !IsTranslucentBlendMode(Parameters.MaterialParameters)
+		&& (!AllowStaticLightingVar || AllowStaticLightingVar->GetValueOnAnyThread() != 0)
+		&& IsFeatureLevelSupported(Parameters.Platform, ERHIFeatureLevel::SM5);
+}
+
+void FCachedVolumeIndirectLightingPolicy::ModifyCompilationEnvironment(const FMaterialShaderPermutationParameters& Parameters, FShaderCompilerEnvironment& OutEnvironment)
+{
+	OutEnvironment.SetDefine(TEXT("CACHED_VOLUME_INDIRECT_LIGHTING"), TEXT("1"));
+}
+
+bool FCachedPointIndirectLightingPolicy::ShouldCompilePermutation(const FMeshMaterialShaderPermutationParameters& Parameters)
+{
+	static const auto AllowStaticLightingVar = IConsoleManager::Get().FindTConsoleVariableDataInt(TEXT("r.AllowStaticLighting"));
+
+	return Parameters.MaterialParameters.ShadingModels.IsLit()
+		&& (!AllowStaticLightingVar || AllowStaticLightingVar->GetValueOnAnyThread() != 0);
+}
+
 void FCachedPointIndirectLightingPolicy::ModifyCompilationEnvironment(const FMaterialShaderPermutationParameters& Parameters, FShaderCompilerEnvironment& OutEnvironment)
 {
 	OutEnvironment.SetDefine(TEXT("CACHED_POINT_INDIRECT_LIGHTING"),TEXT("1"));	
 }
 
+void FMobileDirectionalLightAndCSMPolicy::ModifyCompilationEnvironment(const FMaterialShaderPermutationParameters& Parameters, FShaderCompilerEnvironment& OutEnvironment)
+{
+	OutEnvironment.SetDefine(TEXT("DIRECTIONAL_LIGHT_CSM"), TEXT("1"));
+	OutEnvironment.SetDefine(TEXT("MOBILE_USE_CSM_BRANCH"), MobileUseCSMShaderBranch());
+}
+
+bool FMobileDirectionalLightAndCSMPolicy::ShouldCompilePermutation(const FMeshMaterialShaderPermutationParameters& Parameters)
+{
+	if (IsMobileDeferredShadingEnabled(Parameters.Platform))
+	{
+		return false;
+	}
+
+	static auto* CVarAllowStaticLighting = IConsoleManager::Get().FindTConsoleVariableDataInt(TEXT("r.AllowStaticLighting"));
+	static auto* CVarEnableNoPrecomputedLightingCSMShader = IConsoleManager::Get().FindTConsoleVariableDataInt(TEXT("r.Mobile.EnableNoPrecomputedLightingCSMShader"));
+	const bool bAllowStaticLighting = CVarAllowStaticLighting->GetValueOnAnyThread() != 0;
+	const bool bEnableNoPrecomputedLightingCSMShader = CVarEnableNoPrecomputedLightingCSMShader && CVarEnableNoPrecomputedLightingCSMShader->GetValueOnAnyThread() != 0;
+
+	return (!bAllowStaticLighting || bEnableNoPrecomputedLightingCSMShader) &&
+		Parameters.MaterialParameters.ShadingModels.IsLit() &&
+		!IsTranslucentBlendMode(Parameters.MaterialParameters);
+}
+
+bool FMobileDistanceFieldShadowsAndLQLightMapPolicy::ShouldCompilePermutation(const FMeshMaterialShaderPermutationParameters& Parameters)
+{
+	static auto* CVarMobileAllowDistanceFieldShadows = IConsoleManager::Get().FindTConsoleVariableDataInt(TEXT("r.Mobile.AllowDistanceFieldShadows"));
+	const bool bMobileAllowDistanceFieldShadows = CVarMobileAllowDistanceFieldShadows->GetValueOnAnyThread() == 1;
+	return bMobileAllowDistanceFieldShadows && Super::ShouldCompilePermutation(Parameters);
+}
+
+void FMobileDistanceFieldShadowsAndLQLightMapPolicy::ModifyCompilationEnvironment(const FMaterialShaderPermutationParameters& Parameters, FShaderCompilerEnvironment& OutEnvironment)
+{
+	Super::ModifyCompilationEnvironment(Parameters, OutEnvironment);
+}
+
+bool FMobileDistanceFieldShadowsLightMapAndCSMLightingPolicy::ShouldCompilePermutation(const FMeshMaterialShaderPermutationParameters& Parameters)
+{
+	if (IsMobileDeferredShadingEnabled(Parameters.Platform))
+	{
+		return false;
+	}
+
+	static auto* CVarMobileEnableStaticAndCSMShadowReceivers = IConsoleManager::Get().FindTConsoleVariableDataInt(TEXT("r.Mobile.EnableStaticAndCSMShadowReceivers"));
+	const bool bMobileEnableStaticAndCSMShadowReceivers = CVarMobileEnableStaticAndCSMShadowReceivers->GetValueOnAnyThread() == 1;
+	return bMobileEnableStaticAndCSMShadowReceivers &&
+		Parameters.MaterialParameters.ShadingModels.IsLit() &&
+		!IsTranslucentBlendMode(Parameters.MaterialParameters) &&
+		Super::ShouldCompilePermutation(Parameters);
+}
+
+void FMobileDistanceFieldShadowsLightMapAndCSMLightingPolicy::ModifyCompilationEnvironment(const FMaterialShaderPermutationParameters& Parameters, FShaderCompilerEnvironment& OutEnvironment)
+{
+	OutEnvironment.SetDefine(TEXT("DIRECTIONAL_LIGHT_CSM"), TEXT("1"));
+
+	Super::ModifyCompilationEnvironment(Parameters, OutEnvironment);
+}
+
+bool FMobileDirectionalLightCSMAndLightMapPolicy::ShouldCompilePermutation(const FMeshMaterialShaderPermutationParameters& Parameters)
+{
+	if (IsMobileDeferredShadingEnabled(Parameters.Platform))
+	{
+		return false;
+	}
+
+	return !IsTranslucentBlendMode(Parameters.MaterialParameters) && Super::ShouldCompilePermutation(Parameters);
+}
+
+void FMobileDirectionalLightCSMAndLightMapPolicy::ModifyCompilationEnvironment(const FMaterialShaderPermutationParameters& Parameters, FShaderCompilerEnvironment& OutEnvironment)
+{
+	OutEnvironment.SetDefine(TEXT("DIRECTIONAL_LIGHT_CSM"), TEXT("1"));
+
+	Super::ModifyCompilationEnvironment(Parameters, OutEnvironment);
+}
+
+bool FMobileDirectionalLightAndSHIndirectPolicy::ShouldCompilePermutation(const FMeshMaterialShaderPermutationParameters& Parameters)
+{
+	static auto* CVarAllowStaticLighting = IConsoleManager::Get().FindTConsoleVariableDataInt(TEXT("r.AllowStaticLighting"));
+	const bool bAllowStaticLighting = CVarAllowStaticLighting->GetValueOnAnyThread() != 0;
+
+	return bAllowStaticLighting && Parameters.MaterialParameters.ShadingModels.IsLit() && FCachedPointIndirectLightingPolicy::ShouldCompilePermutation(Parameters);
+}
+
+void FMobileDirectionalLightAndSHIndirectPolicy::ModifyCompilationEnvironment(const FMaterialShaderPermutationParameters& Parameters, FShaderCompilerEnvironment& OutEnvironment)
+{
+	FCachedPointIndirectLightingPolicy::ModifyCompilationEnvironment(Parameters, OutEnvironment);
+}
+
+bool FMobileDirectionalLightCSMAndSHIndirectPolicy::ShouldCompilePermutation(const FMeshMaterialShaderPermutationParameters& Parameters)
+{
+	if (IsMobileDeferredShadingEnabled(Parameters.Platform))
+	{
+		return false;
+	}
+
+	static auto* CVarMobileEnableStaticAndCSMShadowReceivers = IConsoleManager::Get().FindTConsoleVariableDataInt(TEXT("r.Mobile.EnableStaticAndCSMShadowReceivers"));
+	const bool bMobileEnableStaticAndCSMShadowReceivers = CVarMobileEnableStaticAndCSMShadowReceivers->GetValueOnAnyThread() == 1;
+	return bMobileEnableStaticAndCSMShadowReceivers &&
+		!IsTranslucentBlendMode(Parameters.MaterialParameters) &&
+		Super::ShouldCompilePermutation(Parameters);
+}
+
+void FMobileDirectionalLightCSMAndSHIndirectPolicy::ModifyCompilationEnvironment(const FMaterialShaderPermutationParameters& Parameters, FShaderCompilerEnvironment& OutEnvironment)
+{
+	OutEnvironment.SetDefine(TEXT("DIRECTIONAL_LIGHT_CSM"), TEXT("1"));
+
+	Super::ModifyCompilationEnvironment(Parameters, OutEnvironment);
+}
+
+bool FMobileMovableDirectionalLightWithLightmapPolicy::ShouldCompilePermutation(const FMeshMaterialShaderPermutationParameters& Parameters)
+{
+	static auto* CVarMobileAllowMovableDirectionalLights = IConsoleManager::Get().FindTConsoleVariableDataInt(TEXT("r.Mobile.AllowMovableDirectionalLights"));
+	const bool bMobileAllowMovableDirectionalLights = CVarMobileAllowMovableDirectionalLights->GetValueOnAnyThread() != 0;
+
+	return bMobileAllowMovableDirectionalLights && Super::ShouldCompilePermutation(Parameters);
+}
+
+void FMobileMovableDirectionalLightWithLightmapPolicy::ModifyCompilationEnvironment(const FMaterialShaderPermutationParameters& Parameters, FShaderCompilerEnvironment& OutEnvironment)
+{
+	OutEnvironment.SetDefine(TEXT("MOVABLE_DIRECTIONAL_LIGHT"), TEXT("1"));
+	Super::ModifyCompilationEnvironment(Parameters, OutEnvironment);
+}
+
+bool FMobileMovableDirectionalLightCSMWithLightmapPolicy::ShouldCompilePermutation(const FMeshMaterialShaderPermutationParameters& Parameters)
+{
+	if (IsMobileDeferredShadingEnabled(Parameters.Platform))
+	{
+		return false;
+	}
+
+	return Super::ShouldCompilePermutation(Parameters);
+}
+
+void FMobileMovableDirectionalLightCSMWithLightmapPolicy::ModifyCompilationEnvironment(const FMaterialShaderPermutationParameters& Parameters, FShaderCompilerEnvironment& OutEnvironment)
+{
+	OutEnvironment.SetDefine(TEXT("DIRECTIONAL_LIGHT_CSM"), TEXT("1"));
+
+	Super::ModifyCompilationEnvironment(Parameters, OutEnvironment);
+}
+
+bool FSelfShadowedCachedPointIndirectLightingPolicy::ShouldCompilePermutation(const FMeshMaterialShaderPermutationParameters& Parameters)
+{
+	static IConsoleVariable* AllowStaticLightingVar = IConsoleManager::Get().FindConsoleVariable(TEXT("r.AllowStaticLighting"));
+
+	return Parameters.MaterialParameters.ShadingModels.IsLit()
+		&& IsTranslucentBlendMode(Parameters.MaterialParameters)
+		&& (!AllowStaticLightingVar || AllowStaticLightingVar->GetInt() != 0)
+		&& FSelfShadowedTranslucencyPolicy::ShouldCompilePermutation(Parameters);
+}
+
 void FSelfShadowedCachedPointIndirectLightingPolicy::ModifyCompilationEnvironment(const FMaterialShaderPermutationParameters& Parameters, FShaderCompilerEnvironment& OutEnvironment)
 {
 	OutEnvironment.SetDefine(TEXT("CACHED_POINT_INDIRECT_LIGHTING"),TEXT("1"));	
@@ -108,6 +387,25 @@ void FSelfShadowedCachedPointIndirectLightingPolicy::GetComputeShaderBindings(
 	FSelfShadowedTranslucencyPolicy::GetComputeShaderBindings(PrimitiveSceneProxy, ShaderElementData.SelfShadowTranslucencyUniformBuffer, ComputeShaderParameters, ShaderBindings);
 }
 
+bool FSelfShadowedVolumetricLightmapPolicy::ShouldCompilePermutation(const FMeshMaterialShaderPermutationParameters& Parameters)
+{
+	static IConsoleVariable* AllowStaticLightingVar = IConsoleManager::Get().FindConsoleVariable(TEXT("r.AllowStaticLighting"));
+
+	return Parameters.MaterialParameters.ShadingModels.IsLit()
+		&& IsTranslucentBlendMode(Parameters.MaterialParameters)
+		&& (!AllowStaticLightingVar || AllowStaticLightingVar->GetInt() != 0)
+		&& FSelfShadowedTranslucencyPolicy::ShouldCompilePermutation(Parameters);
+}
+
+void FSelfShadowedVolumetricLightmapPolicy::ModifyCompilationEnvironment(const FMaterialShaderPermutationParameters& Parameters, FShaderCompilerEnvironment& OutEnvironment)
+{
+	OutEnvironment.SetDefine(TEXT("PRECOMPUTED_IRRADIANCE_VOLUME_LIGHTING"), TEXT("1"));
+	FSelfShadowedTranslucencyPolicy::ModifyCompilationEnvironment(Parameters, OutEnvironment);
+}
+
+/** Initialization constructor. */
+FSelfShadowedVolumetricLightmapPolicy::FSelfShadowedVolumetricLightmapPolicy() {}
+
 void FSelfShadowedVolumetricLightmapPolicy::GetPixelShaderBindings(
 	const FPrimitiveSceneProxy* PrimitiveSceneProxy,
 	const ElementDataType& ShaderElementData,
@@ -146,6 +444,122 @@ void FSelfShadowedVolumetricLightmapPolicy::GetComputeShaderBindings(
 	FSelfShadowedTranslucencyPolicy::GetComputeShaderBindings(PrimitiveSceneProxy, ShaderElementData.SelfShadowTranslucencyUniformBuffer, ComputeShaderParameters, ShaderBindings);
 }
 
+bool FUniformLightMapPolicy::ShouldCompilePermutation(ELightMapPolicyType Policy, const FMeshMaterialShaderPermutationParameters& Parameters)
+{
+	CA_SUPPRESS(6326);
+	switch (Policy)
+	{
+	case LMP_NO_LIGHTMAP:
+		return FNoLightMapPolicy::ShouldCompilePermutation(Parameters);
+	case LMP_PRECOMPUTED_IRRADIANCE_VOLUME_INDIRECT_LIGHTING:
+		return FPrecomputedVolumetricLightmapLightingPolicy::ShouldCompilePermutation(Parameters);
+	case LMP_CACHED_VOLUME_INDIRECT_LIGHTING:
+		return FCachedVolumeIndirectLightingPolicy::ShouldCompilePermutation(Parameters);
+	case LMP_CACHED_POINT_INDIRECT_LIGHTING:
+		return FCachedPointIndirectLightingPolicy::ShouldCompilePermutation(Parameters);
+	case LMP_LQ_LIGHTMAP:
+		return TLightMapPolicy<LQ_LIGHTMAP>::ShouldCompilePermutation(Parameters);
+	case LMP_HQ_LIGHTMAP:
+		return TLightMapPolicy<HQ_LIGHTMAP>::ShouldCompilePermutation(Parameters);
+	case LMP_DISTANCE_FIELD_SHADOWS_AND_HQ_LIGHTMAP:
+		return TDistanceFieldShadowsAndLightMapPolicy<HQ_LIGHTMAP>::ShouldCompilePermutation(Parameters);
+
+		// Mobile specific
+	case LMP_MOBILE_DISTANCE_FIELD_SHADOWS_AND_LQ_LIGHTMAP:
+		return FMobileDistanceFieldShadowsAndLQLightMapPolicy::ShouldCompilePermutation(Parameters);
+	case LMP_MOBILE_DISTANCE_FIELD_SHADOWS_LIGHTMAP_AND_CSM:
+		return FMobileDistanceFieldShadowsLightMapAndCSMLightingPolicy::ShouldCompilePermutation(Parameters);
+	case LMP_MOBILE_DIRECTIONAL_LIGHT_CSM_AND_LIGHTMAP:
+		return FMobileDirectionalLightCSMAndLightMapPolicy::ShouldCompilePermutation(Parameters);
+	case LMP_MOBILE_DIRECTIONAL_LIGHT_AND_SH_INDIRECT:
+		return FMobileDirectionalLightAndSHIndirectPolicy::ShouldCompilePermutation(Parameters);
+	case LMP_MOBILE_DIRECTIONAL_LIGHT_CSM_AND_SH_INDIRECT:
+		return FMobileDirectionalLightCSMAndSHIndirectPolicy::ShouldCompilePermutation(Parameters);
+	case LMP_MOBILE_MOVABLE_DIRECTIONAL_LIGHT_WITH_LIGHTMAP:
+		return FMobileMovableDirectionalLightWithLightmapPolicy::ShouldCompilePermutation(Parameters);
+	case LMP_MOBILE_MOVABLE_DIRECTIONAL_LIGHT_CSM_WITH_LIGHTMAP:
+		return FMobileMovableDirectionalLightCSMWithLightmapPolicy::ShouldCompilePermutation(Parameters);
+	case LMP_MOBILE_DIRECTIONAL_LIGHT_CSM:
+		return FMobileDirectionalLightAndCSMPolicy::ShouldCompilePermutation(Parameters);
+
+		// LightMapDensity
+
+	case LMP_DUMMY:
+		return FDummyLightMapPolicy::ShouldCompilePermutation(Parameters);
+
+	default:
+		check(false);
+		return false;
+	};
+}
+
+void FUniformLightMapPolicy::ModifyCompilationEnvironment(ELightMapPolicyType Policy, const FMaterialShaderPermutationParameters& Parameters, FShaderCompilerEnvironment& OutEnvironment)
+{
+	OutEnvironment.SetDefine(TEXT("MAX_NUM_LIGHTMAP_COEF"), MAX_NUM_LIGHTMAP_COEF);
+
+	CA_SUPPRESS(6326);
+	switch (Policy)
+	{
+	case LMP_NO_LIGHTMAP:
+		FNoLightMapPolicy::ModifyCompilationEnvironment(Parameters, OutEnvironment);
+		break;
+	case LMP_PRECOMPUTED_IRRADIANCE_VOLUME_INDIRECT_LIGHTING:
+		FPrecomputedVolumetricLightmapLightingPolicy::ModifyCompilationEnvironment(Parameters, OutEnvironment);
+		break;
+	case LMP_CACHED_VOLUME_INDIRECT_LIGHTING:
+		FCachedVolumeIndirectLightingPolicy::ModifyCompilationEnvironment(Parameters, OutEnvironment);
+		break;
+	case LMP_CACHED_POINT_INDIRECT_LIGHTING:
+		FCachedPointIndirectLightingPolicy::ModifyCompilationEnvironment(Parameters, OutEnvironment);
+		break;
+	case LMP_LQ_LIGHTMAP:
+		TLightMapPolicy<LQ_LIGHTMAP>::ModifyCompilationEnvironment(Parameters, OutEnvironment);
+		break;
+	case LMP_HQ_LIGHTMAP:
+		TLightMapPolicy<HQ_LIGHTMAP>::ModifyCompilationEnvironment(Parameters, OutEnvironment);
+		break;
+	case LMP_DISTANCE_FIELD_SHADOWS_AND_HQ_LIGHTMAP:
+		TDistanceFieldShadowsAndLightMapPolicy<HQ_LIGHTMAP>::ModifyCompilationEnvironment(Parameters, OutEnvironment);
+		break;
+
+		// Mobile specific
+	case LMP_MOBILE_DISTANCE_FIELD_SHADOWS_AND_LQ_LIGHTMAP:
+		FMobileDistanceFieldShadowsAndLQLightMapPolicy::ModifyCompilationEnvironment(Parameters, OutEnvironment);
+		break;
+	case LMP_MOBILE_DISTANCE_FIELD_SHADOWS_LIGHTMAP_AND_CSM:
+		FMobileDistanceFieldShadowsLightMapAndCSMLightingPolicy::ModifyCompilationEnvironment(Parameters, OutEnvironment);
+		break;
+	case LMP_MOBILE_DIRECTIONAL_LIGHT_CSM_AND_LIGHTMAP:
+		FMobileDirectionalLightCSMAndLightMapPolicy::ModifyCompilationEnvironment(Parameters, OutEnvironment);
+		break;
+	case LMP_MOBILE_DIRECTIONAL_LIGHT_AND_SH_INDIRECT:
+		FMobileDirectionalLightAndSHIndirectPolicy::ModifyCompilationEnvironment(Parameters, OutEnvironment);
+		break;
+	case LMP_MOBILE_DIRECTIONAL_LIGHT_CSM_AND_SH_INDIRECT:
+		FMobileDirectionalLightCSMAndSHIndirectPolicy::ModifyCompilationEnvironment(Parameters, OutEnvironment);
+		break;
+	case LMP_MOBILE_MOVABLE_DIRECTIONAL_LIGHT_WITH_LIGHTMAP:
+		FMobileMovableDirectionalLightWithLightmapPolicy::ModifyCompilationEnvironment(Parameters, OutEnvironment);
+		break;
+	case LMP_MOBILE_MOVABLE_DIRECTIONAL_LIGHT_CSM_WITH_LIGHTMAP:
+		FMobileMovableDirectionalLightCSMWithLightmapPolicy::ModifyCompilationEnvironment(Parameters, OutEnvironment);
+		break;
+	case LMP_MOBILE_DIRECTIONAL_LIGHT_CSM:
+		FMobileDirectionalLightAndCSMPolicy::ModifyCompilationEnvironment(Parameters, OutEnvironment);
+		break;
+
+		// LightMapDensity
+
+	case LMP_DUMMY:
+		FDummyLightMapPolicy::ModifyCompilationEnvironment(Parameters, OutEnvironment);
+		break;
+
+	default:
+		check(false);
+		break;
+	}
+}
+
 void FUniformLightMapPolicy::GetVertexShaderBindings(
 	const FPrimitiveSceneProxy* PrimitiveSceneProxy,
 	const ElementDataType& ShaderElementData,
@@ -440,4 +854,17 @@ void FEmptyIndirectLightingCacheUniformBuffer::InitDynamicRHI()
 }
 
 /** */
-TGlobalResource< FEmptyIndirectLightingCacheUniformBuffer > GEmptyIndirectLightingCacheUniformBuffer;
\ No newline at end of file
+TGlobalResource< FEmptyIndirectLightingCacheUniformBuffer > GEmptyIndirectLightingCacheUniformBuffer;
+
+bool FNoLightMapPolicy::ShouldCompilePermutation(const FMeshMaterialShaderPermutationParameters& Parameters)
+{
+	if (IsMobilePlatform(Parameters.Platform))
+	{
+		return MobileUsesNoLightMapPermutation(Parameters);
+	}
+	return true;
+}
+
+void FNoLightMapPolicy::ModifyCompilationEnvironment(const FMaterialShaderPermutationParameters& Parameters, FShaderCompilerEnvironment& OutEnvironment)
+{
+}
diff --git a/Engine/Source/Runtime/Renderer/Private/LightMapRendering.h b/Engine/Source/Runtime/Renderer/Private/LightMapRendering.h
--- a/Engine/Source/Runtime/Renderer/Private/LightMapRendering.h
+++ b/Engine/Source/Runtime/Renderer/Private/LightMapRendering.h
@@ -6,15 +6,20 @@
 
 #pragma once
 
-#include "CoreMinimal.h"
-#include "HAL/IConsoleManager.h"
-#include "RenderResource.h"
+#include "HAL/Platform.h"
+#include "LightmapUniformShaderParameters.h"
+#include "Math/Vector.h"
+#include "ShaderParameterMacros.h"
 #include "UniformBuffer.h"
-#include "ShaderParameters.h"
-#include "ShadowRendering.h"
-#include "LightMap.h"
 
+class FMeshDrawSingleShaderBindings;
 class FPrimitiveSceneProxy;
+struct FMaterialShaderPermutationParameters;
+struct FMeshMaterialShaderPermutationParameters;
+
+RENDERER_API bool MobileUsesNoLightMapPermutation(const FMeshMaterialShaderPermutationParameters& Parameters);
+
+extern bool MobileUseCSMShaderBranch();
 
 BEGIN_GLOBAL_SHADER_PARAMETER_STRUCT(FIndirectLightingCacheUniformParameters, )
 	SHADER_PARAMETER(FVector3f, IndirectLightingCachePrimitiveAdd) // FCachedVolumeIndirectLightingPolicy
@@ -70,25 +75,13 @@ public:
 /** Global uniform buffer containing the default precomputed lighting data. */
 extern TGlobalResource< FEmptyIndirectLightingCacheUniformBuffer > GEmptyIndirectLightingCacheUniformBuffer;
 
-RENDERER_API bool MobileUsesNoLightMapPermutation(const FMeshMaterialShaderPermutationParameters& Parameters);
-
 /**
  * A policy for shaders without a lightmap.
  */
 struct FNoLightMapPolicy
 {
-	static bool ShouldCompilePermutation(const FMeshMaterialShaderPermutationParameters& Parameters)
-	{
-		if (IsMobilePlatform(Parameters.Platform))
-		{
-			return MobileUsesNoLightMapPermutation(Parameters);
-		}
-		return true;
-	}
-
-	static void ModifyCompilationEnvironment(const FMaterialShaderPermutationParameters& Parameters, FShaderCompilerEnvironment& OutEnvironment)
-	{
-	}
+	static bool ShouldCompilePermutation(const FMeshMaterialShaderPermutationParameters& Parameters);
+	static void ModifyCompilationEnvironment(const FMaterialShaderPermutationParameters& Parameters, FShaderCompilerEnvironment& OutEnvironment);
 };
 
 enum ELightmapQuality
@@ -97,58 +90,41 @@ enum ELightmapQuality
 	HQ_LIGHTMAP,
 };
 
-// One of these per lightmap quality
-extern const TCHAR* GLightmapDefineName[2];
-extern int32 GNumLightmapCoefficients[2];
+namespace LightMapPolicyImpl
+{
+	void ModifyCompilationEnvironment(ELightmapQuality LightmapQuality, const FMaterialShaderPermutationParameters& Parameters, FShaderCompilerEnvironment& OutEnvironment);
+	bool ShouldCompilePermutation(ELightmapQuality LightmapQuality, const FMeshMaterialShaderPermutationParameters& Parameters);
+};
 
 /**
  * Base policy for shaders with lightmaps.
  */
-template< ELightmapQuality LightmapQuality >
+template<ELightmapQuality LightmapQuality>
 struct TLightMapPolicy
 {
 	static void ModifyCompilationEnvironment(const FMaterialShaderPermutationParameters& Parameters, FShaderCompilerEnvironment& OutEnvironment)
 	{
-		OutEnvironment.SetDefine(GLightmapDefineName[LightmapQuality],TEXT("1"));
-		OutEnvironment.SetDefine(TEXT("NUM_LIGHTMAP_COEFFICIENTS"), GNumLightmapCoefficients[LightmapQuality]);
-
-		static const auto CVar = IConsoleManager::Get().FindTConsoleVariableDataInt(TEXT("r.VirtualTexturedLightmaps"));
-		const bool VirtualTextureLightmaps = (CVar->GetValueOnAnyThread() != 0) && UseVirtualTexturing(GMaxRHIFeatureLevel, OutEnvironment.TargetPlatform);
-		OutEnvironment.SetDefine(TEXT("LIGHTMAP_VT_ENABLED"), VirtualTextureLightmaps);
+		LightMapPolicyImpl::ModifyCompilationEnvironment(LightmapQuality, Parameters, OutEnvironment);
 	}
 
 	static bool ShouldCompilePermutation(const FMeshMaterialShaderPermutationParameters& Parameters)
 	{
-		static const auto AllowStaticLightingVar = IConsoleManager::Get().FindTConsoleVariableDataInt(TEXT("r.AllowStaticLighting"));
-		static const auto CVarProjectCanHaveLowQualityLightmaps = IConsoleManager::Get().FindTConsoleVariableDataInt(TEXT("r.SupportLowQualityLightmaps"));
-		static const auto CVarSupportAllShadersPermutations = IConsoleManager::Get().FindTConsoleVariableDataInt(TEXT("r.SupportAllShaderPermutations"));
-		const bool bForceAllPermutations = CVarSupportAllShadersPermutations && CVarSupportAllShadersPermutations->GetValueOnAnyThread() != 0;
-
-		// if GEngine doesn't exist yet to have the project flag then we should be conservative and cache the LQ lightmap policy
-		const bool bProjectCanHaveLowQualityLightmaps = bForceAllPermutations || (!CVarProjectCanHaveLowQualityLightmaps) || (CVarProjectCanHaveLowQualityLightmaps->GetValueOnAnyThread() != 0);
-
-		const bool bShouldCacheQuality = (LightmapQuality != ELightmapQuality::LQ_LIGHTMAP) || bProjectCanHaveLowQualityLightmaps;
-
-		// GetValueOnAnyThread() as it's possible that ShouldCache is called from rendering thread. That is to output some error message.
-		return (Parameters.MaterialParameters.ShadingModels.IsLit())
-			&& bShouldCacheQuality
-			&& Parameters.VertexFactoryType->SupportsStaticLighting()
-			&& (!AllowStaticLightingVar || AllowStaticLightingVar->GetValueOnAnyThread() != 0)
-			&& (Parameters.MaterialParameters.bIsUsedWithStaticLighting || Parameters.MaterialParameters.bIsSpecialEngineMaterial);
+		return LightMapPolicyImpl::ShouldCompilePermutation(LightmapQuality, Parameters);
 	}
 };
 
+namespace DistanceFieldShadowsAndLightMapPolicyImpl
+{
+	void ModifyCompilationEnvironment(ELightmapQuality LightmapQuality, const FMaterialShaderPermutationParameters& Parameters, FShaderCompilerEnvironment& OutEnvironment);
+};
+
 // A lightmap policy for computing up to 4 signed distance field shadow factors in the base pass.
 template< ELightmapQuality LightmapQuality >
 struct TDistanceFieldShadowsAndLightMapPolicy : public TLightMapPolicy< LightmapQuality >
 {
-	typedef TLightMapPolicy< LightmapQuality >	Super;
-
 	static void ModifyCompilationEnvironment(const FMaterialShaderPermutationParameters& Parameters, FShaderCompilerEnvironment& OutEnvironment)
 	{
-		OutEnvironment.SetDefine(TEXT("STATICLIGHTING_TEXTUREMASK"), 1);
-		OutEnvironment.SetDefine(TEXT("STATICLIGHTING_SIGNEDDISTANCEFIELD"), 1);
-		Super::ModifyCompilationEnvironment(Parameters, OutEnvironment);
+		DistanceFieldShadowsAndLightMapPolicyImpl::ModifyCompilationEnvironment(LightmapQuality, Parameters, OutEnvironment);
 	}
 };
 
@@ -157,12 +133,7 @@ struct TDistanceFieldShadowsAndLightMapPolicy : public TLightMapPolicy< Lightmap
  */
 struct FDummyLightMapPolicy : public TLightMapPolicy< HQ_LIGHTMAP >
 {
-public:
-
-	static bool ShouldCompilePermutation(const FMeshMaterialShaderPermutationParameters& Parameters)
-	{
-		return Parameters.MaterialParameters.ShadingModels.IsLit() && Parameters.VertexFactoryType->SupportsStaticLighting();
-	}
+	static bool ShouldCompilePermutation(const FMeshMaterialShaderPermutationParameters& Parameters);
 };
 
 /**
@@ -216,44 +187,29 @@ public:
 		LAYOUT_FIELD(FShaderUniformBufferParameter, TranslucentSelfShadowBufferParameter);
 	};
 
-	static bool ShouldCompilePermutation(const FMeshMaterialShaderPermutationParameters& Parameters)
-	{
-		return Parameters.MaterialParameters.ShadingModels.IsLit() &&
-			IsTranslucentBlendMode(Parameters.MaterialParameters) &&
-			IsFeatureLevelSupported(Parameters.Platform, ERHIFeatureLevel::SM5);
-	}
-
-	static void ModifyCompilationEnvironment(const FMaterialShaderPermutationParameters& Parameters, FShaderCompilerEnvironment& OutEnvironment)
-	{
-		OutEnvironment.SetDefine(TEXT("TRANSLUCENT_SELF_SHADOWING"),TEXT("1"));
-	}
+	static bool ShouldCompilePermutation(const FMeshMaterialShaderPermutationParameters& Parameters);
+	static void ModifyCompilationEnvironment(const FMaterialShaderPermutationParameters& Parameters, FShaderCompilerEnvironment& OutEnvironment);
 
 	/** Initialization constructor. */
-	FSelfShadowedTranslucencyPolicy() {}
+	FSelfShadowedTranslucencyPolicy();
 	
 	static void GetVertexShaderBindings(
 		const FPrimitiveSceneProxy* PrimitiveSceneProxy,
 		const ElementDataType& ShaderElementData,
 		const VertexParametersType* VertexShaderParameters,
-		FMeshDrawSingleShaderBindings& ShaderBindings) {}
+		FMeshDrawSingleShaderBindings& ShaderBindings);
 
 	static void GetPixelShaderBindings(
 		const FPrimitiveSceneProxy* PrimitiveSceneProxy,
 		const ElementDataType& ShaderElementData,
 		const PixelParametersType* PixelShaderParameters,
-		FMeshDrawSingleShaderBindings& ShaderBindings)
-	{
-		ShaderBindings.Add(PixelShaderParameters->TranslucentSelfShadowBufferParameter, ShaderElementData);
-	}
+		FMeshDrawSingleShaderBindings& ShaderBindings);
 
 	static void GetComputeShaderBindings(
 		const FPrimitiveSceneProxy* PrimitiveSceneProxy,
 		const ElementDataType& ShaderElementData,
 		const ComputeParametersType* ComputeShaderParameters,
-		FMeshDrawSingleShaderBindings& ShaderBindings)
-	{
-		ShaderBindings.Add(ComputeShaderParameters->TranslucentSelfShadowBufferParameter, ShaderElementData);
-	}
+		FMeshDrawSingleShaderBindings& ShaderBindings);
 
 	friend bool operator==(const FSelfShadowedTranslucencyPolicy A,const FSelfShadowedTranslucencyPolicy B)
 	{
@@ -266,18 +222,8 @@ public:
  */
 struct FPrecomputedVolumetricLightmapLightingPolicy
 {
-	static bool ShouldCompilePermutation(const FMeshMaterialShaderPermutationParameters& Parameters)
-	{
-		static const auto AllowStaticLightingVar = IConsoleManager::Get().FindTConsoleVariableDataInt(TEXT("r.AllowStaticLighting"));
-	
-		return Parameters.MaterialParameters.ShadingModels.IsLit() &&
-			(!AllowStaticLightingVar || AllowStaticLightingVar->GetValueOnAnyThread() != 0);
-	}
-
-	static void ModifyCompilationEnvironment(const FMaterialShaderPermutationParameters& Parameters, FShaderCompilerEnvironment& OutEnvironment)
-	{
-		OutEnvironment.SetDefine(TEXT("PRECOMPUTED_IRRADIANCE_VOLUME_LIGHTING"),TEXT("1"));
-	}
+	static bool ShouldCompilePermutation(const FMeshMaterialShaderPermutationParameters& Parameters);
+	static void ModifyCompilationEnvironment(const FMaterialShaderPermutationParameters& Parameters, FShaderCompilerEnvironment& OutEnvironment);
 };
 
 /**
@@ -285,70 +231,24 @@ struct FPrecomputedVolumetricLightmapLightingPolicy
  */
 struct FCachedVolumeIndirectLightingPolicy
 {
-	static bool ShouldCompilePermutation(const FMeshMaterialShaderPermutationParameters& Parameters)
-	{
-		static const auto AllowStaticLightingVar = IConsoleManager::Get().FindTConsoleVariableDataInt(TEXT("r.AllowStaticLighting"));
-
-		return Parameters.MaterialParameters.ShadingModels.IsLit()
-			&& !IsTranslucentBlendMode(Parameters.MaterialParameters)
-			&& (!AllowStaticLightingVar || AllowStaticLightingVar->GetValueOnAnyThread() != 0)
-			&& IsFeatureLevelSupported(Parameters.Platform, ERHIFeatureLevel::SM5);
-	}
-
-	static void ModifyCompilationEnvironment(const FMaterialShaderPermutationParameters& Parameters, FShaderCompilerEnvironment& OutEnvironment)
-	{
-		OutEnvironment.SetDefine(TEXT("CACHED_VOLUME_INDIRECT_LIGHTING"),TEXT("1"));
-	}
+	static bool ShouldCompilePermutation(const FMeshMaterialShaderPermutationParameters& Parameters);
+	static void ModifyCompilationEnvironment(const FMaterialShaderPermutationParameters& Parameters, FShaderCompilerEnvironment& OutEnvironment);
 };
 
-
 /**
  * Allows a dynamic object to access indirect lighting through a per-object lighting sample
  */
 struct FCachedPointIndirectLightingPolicy
 {
-	static bool ShouldCompilePermutation(const FMeshMaterialShaderPermutationParameters& Parameters)
-	{
-		static const auto AllowStaticLightingVar = IConsoleManager::Get().FindTConsoleVariableDataInt(TEXT("r.AllowStaticLighting"));
-	
-		return Parameters.MaterialParameters.ShadingModels.IsLit()
-			&& (!AllowStaticLightingVar || AllowStaticLightingVar->GetValueOnAnyThread() != 0);
-	}
-
+	static bool ShouldCompilePermutation(const FMeshMaterialShaderPermutationParameters& Parameters);
 	static void ModifyCompilationEnvironment(const FMaterialShaderPermutationParameters& Parameters, FShaderCompilerEnvironment& OutEnvironment);
 };
 
-FORCEINLINE_DEBUGGABLE bool MobileUseCSMShaderBranch()
-{
-	static const auto* CVar = IConsoleManager::Get().FindTConsoleVariableDataInt(TEXT("r.Mobile.UseCSMShaderBranch"));
-	return (CVar && CVar->GetValueOnAnyThread() != 0);
-}
-
 /** Mobile Specific: Combines a directional light and CSM. Not used with pre-computed lighting */
 struct FMobileDirectionalLightAndCSMPolicy
 {
-	static void ModifyCompilationEnvironment(const FMaterialShaderPermutationParameters& Parameters, FShaderCompilerEnvironment& OutEnvironment)
-	{
-		OutEnvironment.SetDefine(TEXT("DIRECTIONAL_LIGHT_CSM"), TEXT("1"));
-		OutEnvironment.SetDefine(TEXT("MOBILE_USE_CSM_BRANCH"), MobileUseCSMShaderBranch());
-	}
-	
-	static bool ShouldCompilePermutation(const FMeshMaterialShaderPermutationParameters& Parameters)
-	{
-		if (IsMobileDeferredShadingEnabled(Parameters.Platform))
-		{
-			return false;
-		}
-		
-		static auto* CVarAllowStaticLighting = IConsoleManager::Get().FindTConsoleVariableDataInt(TEXT("r.AllowStaticLighting"));
-		static auto* CVarEnableNoPrecomputedLightingCSMShader = IConsoleManager::Get().FindTConsoleVariableDataInt(TEXT("r.Mobile.EnableNoPrecomputedLightingCSMShader"));
-		const bool bAllowStaticLighting = CVarAllowStaticLighting->GetValueOnAnyThread() != 0;
-		const bool bEnableNoPrecomputedLightingCSMShader = CVarEnableNoPrecomputedLightingCSMShader && CVarEnableNoPrecomputedLightingCSMShader->GetValueOnAnyThread() != 0;
-
-		return (!bAllowStaticLighting || bEnableNoPrecomputedLightingCSMShader) &&
-			Parameters.MaterialParameters.ShadingModels.IsLit() && 
-			!IsTranslucentBlendMode(Parameters.MaterialParameters);
-	}
+	static void ModifyCompilationEnvironment(const FMaterialShaderPermutationParameters& Parameters, FShaderCompilerEnvironment& OutEnvironment);
+	static bool ShouldCompilePermutation(const FMeshMaterialShaderPermutationParameters& Parameters);
 };
 
 /** Mobile Specific: Combines a distance field shadow with LQ lightmaps. */
@@ -357,17 +257,8 @@ class FMobileDistanceFieldShadowsAndLQLightMapPolicy : public TDistanceFieldShad
 	typedef TDistanceFieldShadowsAndLightMapPolicy<LQ_LIGHTMAP>	Super;
 
 public:
-	static bool ShouldCompilePermutation(const FMeshMaterialShaderPermutationParameters& Parameters)
-	{
-		static auto* CVarMobileAllowDistanceFieldShadows = IConsoleManager::Get().FindTConsoleVariableDataInt(TEXT("r.Mobile.AllowDistanceFieldShadows"));
-		const bool bMobileAllowDistanceFieldShadows = CVarMobileAllowDistanceFieldShadows->GetValueOnAnyThread() == 1;
-		return bMobileAllowDistanceFieldShadows && Super::ShouldCompilePermutation(Parameters);
-	}
-
-	static void ModifyCompilationEnvironment(const FMaterialShaderPermutationParameters& Parameters, FShaderCompilerEnvironment& OutEnvironment)
-	{
-		Super::ModifyCompilationEnvironment(Parameters, OutEnvironment);
-	}
+	static bool ShouldCompilePermutation(const FMeshMaterialShaderPermutationParameters& Parameters);
+	static void ModifyCompilationEnvironment(const FMaterialShaderPermutationParameters& Parameters, FShaderCompilerEnvironment& OutEnvironment);
 };
 
 /** Mobile Specific: Combines an distance field shadow with LQ lightmaps and CSM. */
@@ -376,27 +267,8 @@ class FMobileDistanceFieldShadowsLightMapAndCSMLightingPolicy : public FMobileDi
 	typedef FMobileDistanceFieldShadowsAndLQLightMapPolicy Super;
 
 public:
-	static bool ShouldCompilePermutation(const FMeshMaterialShaderPermutationParameters& Parameters)
-	{
-		if (IsMobileDeferredShadingEnabled(Parameters.Platform))
-		{
-			return false;
-		}
-		
-		static auto* CVarMobileEnableStaticAndCSMShadowReceivers = IConsoleManager::Get().FindTConsoleVariableDataInt(TEXT("r.Mobile.EnableStaticAndCSMShadowReceivers"));
-		const bool bMobileEnableStaticAndCSMShadowReceivers = CVarMobileEnableStaticAndCSMShadowReceivers->GetValueOnAnyThread() == 1;
-		return bMobileEnableStaticAndCSMShadowReceivers && 
-			Parameters.MaterialParameters.ShadingModels.IsLit() && 
-			!IsTranslucentBlendMode(Parameters.MaterialParameters) &&
-			Super::ShouldCompilePermutation(Parameters);
-	}
-
-	static void ModifyCompilationEnvironment(const FMaterialShaderPermutationParameters& Parameters, FShaderCompilerEnvironment& OutEnvironment)
-	{
-		OutEnvironment.SetDefine(TEXT("DIRECTIONAL_LIGHT_CSM"), TEXT("1"));
-
-		Super::ModifyCompilationEnvironment(Parameters, OutEnvironment);
-	}
+	static bool ShouldCompilePermutation(const FMeshMaterialShaderPermutationParameters& Parameters);
+	static void ModifyCompilationEnvironment(const FMaterialShaderPermutationParameters& Parameters, FShaderCompilerEnvironment& OutEnvironment);
 };
 
 /** Mobile Specific: Combines a directional light with LQ lightmaps and CSM  */
@@ -405,40 +277,15 @@ struct FMobileDirectionalLightCSMAndLightMapPolicy : public TLightMapPolicy< LQ_
 	typedef TLightMapPolicy<LQ_LIGHTMAP> Super;
 
 public:
-
-	static bool ShouldCompilePermutation(const FMeshMaterialShaderPermutationParameters& Parameters)
-	{
-		if (IsMobileDeferredShadingEnabled(Parameters.Platform))
-		{
-			return false;
-		}
-		
-		return !IsTranslucentBlendMode(Parameters.MaterialParameters) && Super::ShouldCompilePermutation(Parameters);
-	}
-	
-	static void ModifyCompilationEnvironment(const FMaterialShaderPermutationParameters& Parameters, FShaderCompilerEnvironment& OutEnvironment)
-	{
-		OutEnvironment.SetDefine(TEXT("DIRECTIONAL_LIGHT_CSM"), TEXT("1"));
-
-		Super::ModifyCompilationEnvironment(Parameters, OutEnvironment);
-	}
+	static bool ShouldCompilePermutation(const FMeshMaterialShaderPermutationParameters& Parameters);
+	static void ModifyCompilationEnvironment(const FMaterialShaderPermutationParameters& Parameters, FShaderCompilerEnvironment& OutEnvironment);
 };
 
 /** Mobile Specific: Combines an unshadowed directional light with indirect lighting from a single SH sample. */
 struct FMobileDirectionalLightAndSHIndirectPolicy
 {
-	static bool ShouldCompilePermutation(const FMeshMaterialShaderPermutationParameters& Parameters)
-	{
-		static auto* CVarAllowStaticLighting = IConsoleManager::Get().FindTConsoleVariableDataInt(TEXT("r.AllowStaticLighting"));
-		const bool bAllowStaticLighting = CVarAllowStaticLighting->GetValueOnAnyThread() != 0;
-
-		return bAllowStaticLighting && Parameters.MaterialParameters.ShadingModels.IsLit() && FCachedPointIndirectLightingPolicy::ShouldCompilePermutation(Parameters);
-	}
-
-	static void ModifyCompilationEnvironment(const FMaterialShaderPermutationParameters& Parameters, FShaderCompilerEnvironment& OutEnvironment)
-	{
-		FCachedPointIndirectLightingPolicy::ModifyCompilationEnvironment(Parameters, OutEnvironment);
-	}
+	static bool ShouldCompilePermutation(const FMeshMaterialShaderPermutationParameters& Parameters);
+	static void ModifyCompilationEnvironment(const FMaterialShaderPermutationParameters& Parameters, FShaderCompilerEnvironment& OutEnvironment);
 };
 
 /** Mobile Specific: Combines a directional light with CSM with indirect lighting from a single SH sample. */
@@ -447,26 +294,8 @@ class FMobileDirectionalLightCSMAndSHIndirectPolicy : public FMobileDirectionalL
 	typedef FMobileDirectionalLightAndSHIndirectPolicy Super;
 
 public:
-	static bool ShouldCompilePermutation(const FMeshMaterialShaderPermutationParameters& Parameters)
-	{
-		if (IsMobileDeferredShadingEnabled(Parameters.Platform))
-		{
-			return false;
-		}
-		
-		static auto* CVarMobileEnableStaticAndCSMShadowReceivers = IConsoleManager::Get().FindTConsoleVariableDataInt(TEXT("r.Mobile.EnableStaticAndCSMShadowReceivers"));
-		const bool bMobileEnableStaticAndCSMShadowReceivers = CVarMobileEnableStaticAndCSMShadowReceivers->GetValueOnAnyThread() == 1;
-		return bMobileEnableStaticAndCSMShadowReceivers &&
-			!IsTranslucentBlendMode(Parameters.MaterialParameters) &&
-			Super::ShouldCompilePermutation(Parameters);
-	}
-
-	static void ModifyCompilationEnvironment(const FMaterialShaderPermutationParameters& Parameters, FShaderCompilerEnvironment& OutEnvironment)
-	{
-		OutEnvironment.SetDefine(TEXT("DIRECTIONAL_LIGHT_CSM"), TEXT("1"));
-
-		Super::ModifyCompilationEnvironment(Parameters, OutEnvironment);
-	}
+	static bool ShouldCompilePermutation(const FMeshMaterialShaderPermutationParameters& Parameters);
+	static void ModifyCompilationEnvironment(const FMaterialShaderPermutationParameters& Parameters, FShaderCompilerEnvironment& OutEnvironment);
 };
 
 /** Mobile Specific: Combines a movable directional light with LQ lightmaps  */
@@ -474,20 +303,8 @@ struct FMobileMovableDirectionalLightWithLightmapPolicy : public TLightMapPolicy
 {
 	typedef TLightMapPolicy<LQ_LIGHTMAP> Super;
 
-	static bool ShouldCompilePermutation(const FMeshMaterialShaderPermutationParameters& Parameters)
-	{
-	
-		static auto* CVarMobileAllowMovableDirectionalLights = IConsoleManager::Get().FindTConsoleVariableDataInt(TEXT("r.Mobile.AllowMovableDirectionalLights"));
-		const bool bMobileAllowMovableDirectionalLights = CVarMobileAllowMovableDirectionalLights->GetValueOnAnyThread() != 0;
-			
-		return bMobileAllowMovableDirectionalLights && Super::ShouldCompilePermutation(Parameters);
-	}
-
-	static void ModifyCompilationEnvironment(const FMaterialShaderPermutationParameters& Parameters, FShaderCompilerEnvironment& OutEnvironment)
-	{
-		OutEnvironment.SetDefine(TEXT("MOVABLE_DIRECTIONAL_LIGHT"), TEXT("1"));
-		Super::ModifyCompilationEnvironment(Parameters, OutEnvironment);
-	}
+	static bool ShouldCompilePermutation(const FMeshMaterialShaderPermutationParameters& Parameters);
+	static void ModifyCompilationEnvironment(const FMaterialShaderPermutationParameters& Parameters, FShaderCompilerEnvironment& OutEnvironment);
 };
 
 /** Mobile Specific: Combines a movable directional light with LQ lightmaps and CSM */
@@ -495,22 +312,8 @@ struct FMobileMovableDirectionalLightCSMWithLightmapPolicy : public FMobileMovab
 {
 	typedef FMobileMovableDirectionalLightWithLightmapPolicy Super;
 	
-	static bool ShouldCompilePermutation(const FMeshMaterialShaderPermutationParameters& Parameters)
-	{
-		if (IsMobileDeferredShadingEnabled(Parameters.Platform))
-		{
-			return false;
-		}
-
-		return Super::ShouldCompilePermutation(Parameters);
-	}
-	
-	static void ModifyCompilationEnvironment(const FMaterialShaderPermutationParameters& Parameters, FShaderCompilerEnvironment& OutEnvironment)
-	{
-		OutEnvironment.SetDefine(TEXT("DIRECTIONAL_LIGHT_CSM"), TEXT("1"));
-
-		Super::ModifyCompilationEnvironment(Parameters, OutEnvironment);
-	}
+	static bool ShouldCompilePermutation(const FMeshMaterialShaderPermutationParameters& Parameters);
+	static void ModifyCompilationEnvironment(const FMaterialShaderPermutationParameters& Parameters, FShaderCompilerEnvironment& OutEnvironment);
 };
 
 
@@ -573,16 +376,11 @@ public:
 #endif
 	typedef FUniformLightMapPolicyShaderParametersType ComputeParametersType;
 
-	static bool ShouldCompilePermutation(const FMeshMaterialShaderPermutationParameters& Parameters)
-	{
-		return false; // This one does not compile shaders since we can't tell which policy to use.
-	}
-
-	static void ModifyCompilationEnvironment(const FMaterialShaderPermutationParameters& Parameters, FShaderCompilerEnvironment& OutEnvironment)
-	{}
-
 	FUniformLightMapPolicy(ELightMapPolicyType InIndirectPolicy) : IndirectPolicy(InIndirectPolicy) {}
 
+	static bool ShouldCompilePermutation(ELightMapPolicyType Policy, const FMeshMaterialShaderPermutationParameters& Parameters);
+	static void ModifyCompilationEnvironment(ELightMapPolicyType Policy, const FMaterialShaderPermutationParameters& Parameters, FShaderCompilerEnvironment& OutEnvironment);
+
 	static void GetVertexShaderBindings(
 		const FPrimitiveSceneProxy* PrimitiveSceneProxy,
 		const ElementDataType& ShaderElementData,
@@ -618,131 +416,23 @@ public:
 	ELightMapPolicyType GetIndirectPolicy() const { return IndirectPolicy; }
 
 private:
-
-	ELightMapPolicyType IndirectPolicy;
+	const ELightMapPolicyType IndirectPolicy;
 };
 
 template <ELightMapPolicyType Policy>
 class TUniformLightMapPolicy : public FUniformLightMapPolicy
 {
 public:
-
 	TUniformLightMapPolicy() : FUniformLightMapPolicy(Policy) {}
 
 	static bool ShouldCompilePermutation(const FMeshMaterialShaderPermutationParameters& Parameters)
 	{
-		CA_SUPPRESS(6326);
-		switch (Policy)
-		{
-		case LMP_NO_LIGHTMAP:
-			return FNoLightMapPolicy::ShouldCompilePermutation(Parameters);
-		case LMP_PRECOMPUTED_IRRADIANCE_VOLUME_INDIRECT_LIGHTING:
-			return FPrecomputedVolumetricLightmapLightingPolicy::ShouldCompilePermutation(Parameters);
-		case LMP_CACHED_VOLUME_INDIRECT_LIGHTING:
-			return FCachedVolumeIndirectLightingPolicy::ShouldCompilePermutation(Parameters);
-		case LMP_CACHED_POINT_INDIRECT_LIGHTING:
-			return FCachedPointIndirectLightingPolicy::ShouldCompilePermutation(Parameters);
-		case LMP_LQ_LIGHTMAP:
-			return TLightMapPolicy<LQ_LIGHTMAP>::ShouldCompilePermutation(Parameters);
-		case LMP_HQ_LIGHTMAP:
-			return TLightMapPolicy<HQ_LIGHTMAP>::ShouldCompilePermutation(Parameters);
-		case LMP_DISTANCE_FIELD_SHADOWS_AND_HQ_LIGHTMAP:
-			return TDistanceFieldShadowsAndLightMapPolicy<HQ_LIGHTMAP>::ShouldCompilePermutation(Parameters);
-
-		// Mobile specific
-		case LMP_MOBILE_DISTANCE_FIELD_SHADOWS_AND_LQ_LIGHTMAP:
-			return FMobileDistanceFieldShadowsAndLQLightMapPolicy::ShouldCompilePermutation(Parameters);
-		case LMP_MOBILE_DISTANCE_FIELD_SHADOWS_LIGHTMAP_AND_CSM:
-			return FMobileDistanceFieldShadowsLightMapAndCSMLightingPolicy::ShouldCompilePermutation(Parameters);
-		case LMP_MOBILE_DIRECTIONAL_LIGHT_CSM_AND_LIGHTMAP:
-			return FMobileDirectionalLightCSMAndLightMapPolicy::ShouldCompilePermutation(Parameters);
-		case LMP_MOBILE_DIRECTIONAL_LIGHT_AND_SH_INDIRECT:
-			return FMobileDirectionalLightAndSHIndirectPolicy::ShouldCompilePermutation(Parameters);
-		case LMP_MOBILE_DIRECTIONAL_LIGHT_CSM_AND_SH_INDIRECT:
-			return FMobileDirectionalLightCSMAndSHIndirectPolicy::ShouldCompilePermutation(Parameters);
-		case LMP_MOBILE_MOVABLE_DIRECTIONAL_LIGHT_WITH_LIGHTMAP:
-			return FMobileMovableDirectionalLightWithLightmapPolicy::ShouldCompilePermutation(Parameters);
-		case LMP_MOBILE_MOVABLE_DIRECTIONAL_LIGHT_CSM_WITH_LIGHTMAP:
-			return FMobileMovableDirectionalLightCSMWithLightmapPolicy::ShouldCompilePermutation(Parameters);
-		case LMP_MOBILE_DIRECTIONAL_LIGHT_CSM:
-			return FMobileDirectionalLightAndCSMPolicy::ShouldCompilePermutation(Parameters);
-
-		// LightMapDensity
-	
-		case LMP_DUMMY:
-			return FDummyLightMapPolicy::ShouldCompilePermutation(Parameters);
-
-		default:
-			check(false);
-			return false;
-		};
+		return FUniformLightMapPolicy::ShouldCompilePermutation(Policy, Parameters);
 	}
 
 	static void ModifyCompilationEnvironment(const FMaterialShaderPermutationParameters& Parameters, FShaderCompilerEnvironment& OutEnvironment)
 	{
-		OutEnvironment.SetDefine(TEXT("MAX_NUM_LIGHTMAP_COEF"), MAX_NUM_LIGHTMAP_COEF);
-
-		CA_SUPPRESS(6326);
-		switch (Policy)
-		{
-		case LMP_NO_LIGHTMAP:							
-			FNoLightMapPolicy::ModifyCompilationEnvironment(Parameters, OutEnvironment);
-			break;
-		case LMP_PRECOMPUTED_IRRADIANCE_VOLUME_INDIRECT_LIGHTING:
-			FPrecomputedVolumetricLightmapLightingPolicy::ModifyCompilationEnvironment(Parameters, OutEnvironment);
-			break;
-		case LMP_CACHED_VOLUME_INDIRECT_LIGHTING:
-			FCachedVolumeIndirectLightingPolicy::ModifyCompilationEnvironment(Parameters, OutEnvironment);
-			break;
-		case LMP_CACHED_POINT_INDIRECT_LIGHTING:
-			FCachedPointIndirectLightingPolicy::ModifyCompilationEnvironment(Parameters, OutEnvironment);
-			break;
-		case LMP_LQ_LIGHTMAP:
-			TLightMapPolicy<LQ_LIGHTMAP>::ModifyCompilationEnvironment(Parameters, OutEnvironment);
-			break;
-		case LMP_HQ_LIGHTMAP:
-			TLightMapPolicy<HQ_LIGHTMAP>::ModifyCompilationEnvironment(Parameters, OutEnvironment);
-			break;
-		case LMP_DISTANCE_FIELD_SHADOWS_AND_HQ_LIGHTMAP:
-			TDistanceFieldShadowsAndLightMapPolicy<HQ_LIGHTMAP>::ModifyCompilationEnvironment(Parameters, OutEnvironment);
-			break;
-
-		// Mobile specific
-		case LMP_MOBILE_DISTANCE_FIELD_SHADOWS_AND_LQ_LIGHTMAP:
-			FMobileDistanceFieldShadowsAndLQLightMapPolicy::ModifyCompilationEnvironment(Parameters, OutEnvironment);
-			break;
-		case LMP_MOBILE_DISTANCE_FIELD_SHADOWS_LIGHTMAP_AND_CSM:
-			FMobileDistanceFieldShadowsLightMapAndCSMLightingPolicy::ModifyCompilationEnvironment(Parameters, OutEnvironment);
-			break;
-		case LMP_MOBILE_DIRECTIONAL_LIGHT_CSM_AND_LIGHTMAP:
-			FMobileDirectionalLightCSMAndLightMapPolicy::ModifyCompilationEnvironment(Parameters, OutEnvironment);
-			break;
-		case LMP_MOBILE_DIRECTIONAL_LIGHT_AND_SH_INDIRECT:
-			FMobileDirectionalLightAndSHIndirectPolicy::ModifyCompilationEnvironment(Parameters, OutEnvironment);
-			break;
-		case LMP_MOBILE_DIRECTIONAL_LIGHT_CSM_AND_SH_INDIRECT:
-			FMobileDirectionalLightCSMAndSHIndirectPolicy::ModifyCompilationEnvironment(Parameters, OutEnvironment);
-			break;
-		case LMP_MOBILE_MOVABLE_DIRECTIONAL_LIGHT_WITH_LIGHTMAP:
-			FMobileMovableDirectionalLightWithLightmapPolicy::ModifyCompilationEnvironment(Parameters, OutEnvironment);
-			break;
-		case LMP_MOBILE_MOVABLE_DIRECTIONAL_LIGHT_CSM_WITH_LIGHTMAP:
-			FMobileMovableDirectionalLightCSMWithLightmapPolicy::ModifyCompilationEnvironment(Parameters, OutEnvironment);
-			break;
-		case LMP_MOBILE_DIRECTIONAL_LIGHT_CSM:
-			FMobileDirectionalLightAndCSMPolicy::ModifyCompilationEnvironment(Parameters, OutEnvironment);
-			break;
-
-		// LightMapDensity
-	
-		case LMP_DUMMY:
-			FDummyLightMapPolicy::ModifyCompilationEnvironment(Parameters, OutEnvironment);
-			break;
-
-		default:
-			check(false);
-			break;
-		}
+		FUniformLightMapPolicy::ModifyCompilationEnvironment(Policy, Parameters, OutEnvironment);
 	}
 };
 
@@ -794,16 +484,7 @@ public:
 		}
 	};
 
-	static bool ShouldCompilePermutation(const FMeshMaterialShaderPermutationParameters& Parameters)
-	{
-		static IConsoleVariable* AllowStaticLightingVar = IConsoleManager::Get().FindConsoleVariable(TEXT("r.AllowStaticLighting"));
-
-		return Parameters.MaterialParameters.ShadingModels.IsLit()
-			&& IsTranslucentBlendMode(Parameters.MaterialParameters)
-			&& (!AllowStaticLightingVar || AllowStaticLightingVar->GetInt() != 0)
-			&& FSelfShadowedTranslucencyPolicy::ShouldCompilePermutation(Parameters);
-	}
-
+	static bool ShouldCompilePermutation(const FMeshMaterialShaderPermutationParameters& Parameters);
 	static void ModifyCompilationEnvironment(const FMaterialShaderPermutationParameters& Parameters, FShaderCompilerEnvironment& OutEnvironment);
 
 	/** Initialization constructor. */
@@ -867,24 +548,11 @@ public:
 		}
 	};
 
-	static bool ShouldCompilePermutation(const FMeshMaterialShaderPermutationParameters& Parameters)
-	{
-		static IConsoleVariable* AllowStaticLightingVar = IConsoleManager::Get().FindConsoleVariable(TEXT("r.AllowStaticLighting"));
-
-		return Parameters.MaterialParameters.ShadingModels.IsLit()
-			&& IsTranslucentBlendMode(Parameters.MaterialParameters)
-			&& (!AllowStaticLightingVar || AllowStaticLightingVar->GetInt() != 0)
-			&& FSelfShadowedTranslucencyPolicy::ShouldCompilePermutation(Parameters);
-	}
-
-	static void ModifyCompilationEnvironment(const FMaterialShaderPermutationParameters& Parameters, FShaderCompilerEnvironment& OutEnvironment)
-	{
-		OutEnvironment.SetDefine(TEXT("PRECOMPUTED_IRRADIANCE_VOLUME_LIGHTING"),TEXT("1"));	
-		FSelfShadowedTranslucencyPolicy::ModifyCompilationEnvironment(Parameters, OutEnvironment);
-	}
+	static bool ShouldCompilePermutation(const FMeshMaterialShaderPermutationParameters& Parameters);
+	static void ModifyCompilationEnvironment(const FMaterialShaderPermutationParameters& Parameters, FShaderCompilerEnvironment& OutEnvironment);
 
 	/** Initialization constructor. */
-	FSelfShadowedVolumetricLightmapPolicy() {}
+	FSelfShadowedVolumetricLightmapPolicy();
 
 	static void GetVertexShaderBindings(
 		const FPrimitiveSceneProxy* PrimitiveSceneProxy,
