Removing more headers from ScenePrivate.h

diff --git a/Engine/Source/Runtime/Renderer/Private/BasePassRendering.cpp b/Engine/Source/Runtime/Renderer/Private/BasePassRendering.cpp
--- a/Engine/Source/Runtime/Renderer/Private/BasePassRendering.cpp
+++ b/Engine/Source/Runtime/Renderer/Private/BasePassRendering.cpp
@@ -23,6 +23,13 @@
 #include "Nanite/NaniteVisualize.h"
 #include "RenderCore.h"
 #include "DataDrivenShaderPlatformInfo.h"
+#include "VolumetricFog.h"
+
+#include "BasePassRendering.inl"
+
+// Instantiate the common policies
+template class TBasePassVertexShaderPolicyParamType<FUniformLightMapPolicy>;
+template class TBasePassPixelShaderPolicyParamType<FUniformLightMapPolicy>;
 
 // Changing this causes a full shader recompile
 static TAutoConsoleVariable<int32> CVarSelectiveBasePassOutputs(
diff --git a/Engine/Source/Runtime/Renderer/Private/BasePassRendering.inl b/Engine/Source/Runtime/Renderer/Private/BasePassRendering.inl
--- a/Engine/Source/Runtime/Renderer/Private/BasePassRendering.inl
+++ b/Engine/Source/Runtime/Renderer/Private/BasePassRendering.inl
@@ -8,6 +8,7 @@
 #pragma once
 
 #include "CoreFwd.h"
+#include "BasePassRendering.h"
 
 class FMeshMaterialShader;
 class FPrimitiveSceneProxy;
diff --git a/Engine/Source/Runtime/Renderer/Private/DistanceFieldStreaming.cpp b/Engine/Source/Runtime/Renderer/Private/DistanceFieldStreaming.cpp
--- a/Engine/Source/Runtime/Renderer/Private/DistanceFieldStreaming.cpp
+++ b/Engine/Source/Runtime/Renderer/Private/DistanceFieldStreaming.cpp
@@ -7,6 +7,7 @@
 #include "CoreMinimal.h"
 #include "Stats/Stats.h"
 #include "HAL/IConsoleManager.h"
+#include "IO/IoDispatcher.h"
 #include "Async/ParallelFor.h"
 #include "RHI.h"
 #include "RenderResource.h"
@@ -436,6 +437,68 @@ void FDistanceFieldBlockAllocator::Free(const TArray<int32, TInlineAllocator<4>>
 	FreeBlocks.Append(ElementRange);
 }
 
+struct FDistanceFieldReadRequest
+{
+	// SDF scene context
+	FSetElementId AssetSetId;
+	int32 ReversedMipIndex = 0;
+	int32 NumDistanceFieldBricks = 0;
+	uint64 BuiltDataId = 0;
+
+	// Used when BulkData is nullptr
+	const uint8* AlwaysLoadedDataPtr = nullptr;
+
+	// Inputs of read request
+	const FBulkData* BulkData = nullptr;
+	uint32 BulkOffset = 0;
+	uint32 BulkSize = 0;
+
+#if !WITH_EDITOR
+	// Outputs of read request
+	FBulkDataBatchReadRequest RequestHandle;
+	FIoBuffer RequestBuffer;
+#endif
+};
+
+
+struct FDistanceFieldAsyncUpdateParameters
+{
+	FDistanceFieldSceneData* DistanceFieldSceneData = nullptr;
+	FIntVector4* BrickUploadCoordinatesPtr = nullptr;
+	uint8* BrickUploadDataPtr = nullptr;
+
+	uint32* IndirectionIndicesUploadPtr = nullptr;
+	FVector4f* IndirectionDataUploadPtr = nullptr;
+
+	TArray<FDistanceFieldReadRequest> NewReadRequests;
+	TArray<FDistanceFieldReadRequest> ReadRequestsToUpload;
+	TArray<FDistanceFieldReadRequest> ReadRequestsToCleanUp;
+};
+
+FDistanceFieldSceneData::FDistanceFieldSceneData(FDistanceFieldSceneData&&) = default;
+
+FDistanceFieldSceneData::FDistanceFieldSceneData(EShaderPlatform ShaderPlatform)
+	: NumObjectsInBuffer(0)
+	, IndirectionAtlasLayout(8, 8, 8, 512, 512, 512, false, true, false)
+	, HeightFieldAtlasGeneration(0)
+	, HFVisibilityAtlasGenerattion(0)
+{
+	ObjectBuffers = nullptr;
+	HeightFieldObjectBuffers = nullptr;
+
+	bTrackAllPrimitives = ShouldAllPrimitivesHaveDistanceField(ShaderPlatform);
+
+	bCanUse16BitObjectIndices = RHISupportsBufferLoadTypeConversion(ShaderPlatform);
+
+	StreamingRequestReadbackBuffers.AddZeroed(MaxStreamingReadbackBuffers);
+}
+
+FDistanceFieldSceneData::~FDistanceFieldSceneData()
+{
+	delete ObjectBuffers;
+	delete HeightFieldObjectBuffers;
+}
+
 class FDistanceFieldStreamingUpdateTask
 {
 public:
diff --git a/Engine/Source/Runtime/Renderer/Private/FogRendering.h b/Engine/Source/Runtime/Renderer/Private/FogRendering.h
--- a/Engine/Source/Runtime/Renderer/Private/FogRendering.h
+++ b/Engine/Source/Runtime/Renderer/Private/FogRendering.h
@@ -6,10 +6,11 @@
 
 #pragma once
 
-#include "CoreMinimal.h"
-#include "ShaderParameters.h"
-#include "SceneRendering.h"
-#include "VolumetricFog.h"
+#include "ShaderParameterMacros.h"
+#include "RenderGraphFwd.h"
+
+class FSceneViewFamily;
+class FViewInfo;
 
 BEGIN_GLOBAL_SHADER_PARAMETER_STRUCT(FFogUniformParameters,)
 	SHADER_PARAMETER(FVector4f, ExponentialFogParameters)
diff --git a/Engine/Source/Runtime/Renderer/Private/HairStrands/HairStrandsLUT.h b/Engine/Source/Runtime/Renderer/Private/HairStrands/HairStrandsLUT.h
--- a/Engine/Source/Runtime/Renderer/Private/HairStrands/HairStrandsLUT.h
+++ b/Engine/Source/Runtime/Renderer/Private/HairStrands/HairStrandsLUT.h
@@ -6,9 +6,7 @@
 
 #pragma once
 
-#include "CoreMinimal.h"
-#include "RendererInterface.h"
-#include "RenderGraphDefinitions.h"
+#include "RenderGraphFwd.h"
 
 class FViewInfo;
 
diff --git a/Engine/Source/Runtime/Renderer/Private/RendererScene.cpp b/Engine/Source/Runtime/Renderer/Private/RendererScene.cpp
--- a/Engine/Source/Runtime/Renderer/Private/RendererScene.cpp
+++ b/Engine/Source/Runtime/Renderer/Private/RendererScene.cpp
@@ -606,28 +606,6 @@ bool FPixelInspectorData::AddPixelInspectorRequest(FPixelInspectorRequest *Pixel
 
 #endif //WITH_EDITOR
 
-FDistanceFieldSceneData::FDistanceFieldSceneData(EShaderPlatform ShaderPlatform) 
-	: NumObjectsInBuffer(0)
-	, IndirectionAtlasLayout(8, 8, 8, 512, 512, 512, false, true, false)
-	, HeightFieldAtlasGeneration(0)
-	, HFVisibilityAtlasGenerattion(0)
-{
-	ObjectBuffers = nullptr;
-	HeightFieldObjectBuffers = nullptr;
-
-	bTrackAllPrimitives = ShouldAllPrimitivesHaveDistanceField(ShaderPlatform);
-
-	bCanUse16BitObjectIndices = RHISupportsBufferLoadTypeConversion(ShaderPlatform);
-
-	StreamingRequestReadbackBuffers.AddZeroed(MaxStreamingReadbackBuffers);
-}
-
-FDistanceFieldSceneData::~FDistanceFieldSceneData() 
-{
-	delete ObjectBuffers;
-	delete HeightFieldObjectBuffers;
-}
-
 bool IncludePrimitiveInDistanceFieldSceneData(bool bTrackAllPrimitives, const FPrimitiveSceneProxy* Proxy)
 {
 	return PrimitiveNeedsDistanceFieldSceneData(
diff --git a/Engine/Source/Runtime/Renderer/Private/ScenePrivate.h b/Engine/Source/Runtime/Renderer/Private/ScenePrivate.h
--- a/Engine/Source/Runtime/Renderer/Private/ScenePrivate.h
+++ b/Engine/Source/Runtime/Renderer/Private/ScenePrivate.h
@@ -8,8 +8,6 @@
 
 // Dependencies.
 
-#include "CoreMinimal.h"
-#include "IO/IoDispatcher.h"
 #include "Misc/Guid.h"
 #include "Math/RandomStream.h"
 #include "Templates/PimplPtr.h"
@@ -31,13 +29,10 @@
 #include "LightSceneInfo.h"
 #include "DepthRendering.h"
 #include "SceneHitProxyRendering.h"
-//#include "ShadowRendering.h"
 #include "TextureLayout.h"
 #include "SceneRendering.h"
 #include "LightMapRendering.h"
 #include "VelocityRendering.h"
-#include "BasePassRendering.h"
-#include "MobileBasePassRendering.h"
 #include "VolumeRendering.h"
 #include "CommonRenderResources.h"
 #include "VisualizeTexture.h"
@@ -667,13 +662,7 @@ public:
 			return (PrimitiveId == Other.PrimitiveId && Light == Other.Light && ShadowSplitIndex == Other.ShadowSplitIndex && bTranslucentShadow == Other.bTranslucentShadow);
 		}
 
-		FProjectedShadowKey(const FProjectedShadowInfo& ProjectedShadowInfo)
-			: PrimitiveId(ProjectedShadowInfo.GetParentSceneInfo() ? ProjectedShadowInfo.GetParentSceneInfo()->PrimitiveComponentId : FPrimitiveComponentId())
-			, Light(ProjectedShadowInfo.GetLightSceneInfo().Proxy->GetLightComponent())
-			, ShadowSplitIndex(ProjectedShadowInfo.CascadeSettings.ShadowSplitIndex)
-			, bTranslucentShadow(ProjectedShadowInfo.bTranslucentShadow)
-		{
-		}
+		FProjectedShadowKey(const FProjectedShadowInfo& ProjectedShadowInfo);
 
 		FProjectedShadowKey(FPrimitiveComponentId InPrimitiveId, const ULightComponent* InLight, int32 InSplitIndex, bool bInTranslucentShadow)
 			: PrimitiveId(InPrimitiveId)
@@ -1899,47 +1888,14 @@ private:
 	TArray<int32, TInlineAllocator<4>> FreeBlocks;
 };
 
-struct FDistanceFieldReadRequest
-{
-	// SDF scene context
-	FSetElementId AssetSetId;
-	int32 ReversedMipIndex = 0;
-	int32 NumDistanceFieldBricks = 0;
-	uint64 BuiltDataId = 0;
-
-	// Used when BulkData is nullptr
-	const uint8* AlwaysLoadedDataPtr = nullptr;
-
-	// Inputs of read request
-	const FBulkData* BulkData = nullptr;
-	uint32 BulkOffset = 0;
-	uint32 BulkSize = 0;
-
-#if !WITH_EDITOR
-	// Outputs of read request
-	FBulkDataBatchReadRequest RequestHandle;
-	FIoBuffer RequestBuffer;
-#endif
-};
-
-struct FDistanceFieldAsyncUpdateParameters
-{
-	FDistanceFieldSceneData* DistanceFieldSceneData = nullptr;
-	FIntVector4* BrickUploadCoordinatesPtr = nullptr;
-	uint8* BrickUploadDataPtr = nullptr;
-
-	uint32* IndirectionIndicesUploadPtr = nullptr;
-	FVector4f* IndirectionDataUploadPtr = nullptr;
-
-	TArray<FDistanceFieldReadRequest> NewReadRequests;
-	TArray<FDistanceFieldReadRequest> ReadRequestsToUpload;
-};
+struct FDistanceFieldReadRequest;
+struct FDistanceFieldAsyncUpdateParameters;
 
 /** Scene data used to manage distance field object buffers on the GPU. */
 class FDistanceFieldSceneData
 {
 public:
-
+	FDistanceFieldSceneData(FDistanceFieldSceneData&&);
 	FDistanceFieldSceneData(EShaderPlatform ShaderPlatform);
 	~FDistanceFieldSceneData();
 
@@ -3623,5 +3579,3 @@ inline bool ShouldIncludeDomainInMeshPass(EMaterialDomain Domain)
 	return Domain != MD_Volume;
 }
 
-#include "BasePassRendering.inl" // IWYU pragma: export
-
diff --git a/Engine/Source/Runtime/Renderer/Private/ShadowSetup.cpp b/Engine/Source/Runtime/Renderer/Private/ShadowSetup.cpp
--- a/Engine/Source/Runtime/Renderer/Private/ShadowSetup.cpp
+++ b/Engine/Source/Runtime/Renderer/Private/ShadowSetup.cpp
@@ -2492,6 +2492,14 @@ void FProjectedShadowInfo::ClearTransientArrays()
 	GraphicsMinimalPipelineStateSet.Empty();
 }
 
+FSceneViewState::FProjectedShadowKey::FProjectedShadowKey(const FProjectedShadowInfo& ProjectedShadowInfo)
+	: PrimitiveId(ProjectedShadowInfo.GetParentSceneInfo() ? ProjectedShadowInfo.GetParentSceneInfo()->PrimitiveComponentId : FPrimitiveComponentId())
+	, Light(ProjectedShadowInfo.GetLightSceneInfo().Proxy->GetLightComponent())
+	, ShadowSplitIndex(ProjectedShadowInfo.CascadeSettings.ShadowSplitIndex)
+	, bTranslucentShadow(ProjectedShadowInfo.bTranslucentShadow)
+{
+}
+
 /** Returns a cached preshadow matching the input criteria if one exists. */
 TRefCountPtr<FProjectedShadowInfo> FSceneRenderer::GetCachedPreshadow(
 	const FLightPrimitiveInteraction* InParentInteraction, 
