Adding missing include after header reorg.

diff --git a/Engine/Source/Runtime/Renderer/Private/SceneRendering.cpp b/Engine/Source/Runtime/Renderer/Private/SceneRendering.cpp
--- a/Engine/Source/Runtime/Renderer/Private/SceneRendering.cpp
+++ b/Engine/Source/Runtime/Renderer/Private/SceneRendering.cpp
@@ -85,6 +85,7 @@
 #include "TextureResource.h"
 #include "RenderCore.h"
 #include "VolumetricCloudRendering.h"
+#include "VolumetricFog.h"
 
 /*-----------------------------------------------------------------------------
 	Globals
