Path Tracer: Add an option to probe mesh decals along the primary ray to match raster behavior better

Keep the previous logic for experimentation, but probe along the negative of the geometric normal so the logic can stay unified

diff --git a/Engine/Shaders/Private/PathTracing/PathTracingCore.ush b/Engine/Shaders/Private/PathTracing/PathTracingCore.ush
--- a/Engine/Shaders/Private/PathTracing/PathTracingCore.ush
+++ b/Engine/Shaders/Private/PathTracing/PathTracingCore.ush
@@ -56,6 +56,7 @@ uint VisualizeDecalGrid;
 float DecalRoughnessCutoff;
 float MeshDecalRoughnessCutoff;
 float MeshDecalBias;
+int MeshDecalsUseViewRay;
 float MaxPathIntensity;
 float MaxNormalBias;
 float FilterWidth;
@@ -374,18 +375,28 @@ FPathTracingPayload TraceTransparentRay(inout FPathState PathState, bool bIsCame
 
 		if (PathState.PathRoughness < MeshDecalRoughnessCutoff && HitPayload.IsDecalReceiver())
 		{
-			FRayDesc DecalRay;
-			DecalRay.Direction = HitPayload.WorldGeoNormal;
-			DecalRay.Origin = HitPayload.TranslatedWorldPos - DecalRay.Direction * MeshDecalBias;
-			DecalRay.TMin = 0.0f;
-			DecalRay.TMax = MeshDecalBias * 2;
+			FRayDesc DecalRay = (FRayDesc)0;
+			if (MeshDecalsUseViewRay)
+			{
+				DecalRay = PathState.Ray;
+				// go a bit further than the current hit, to handle decals partially inside the model
+				// NOTE: TMin will be inherited from current ray
+				DecalRay.TMax = HitPayload.HitT + MeshDecalBias;
+			}
+			else
+			{
+				DecalRay.Direction = -HitPayload.WorldGeoNormal;
+				DecalRay.Origin = HitPayload.TranslatedWorldPos - DecalRay.Direction * MeshDecalBias;
+				DecalRay.TMin = 0.0f;
+				DecalRay.TMax = MeshDecalBias * 2;
+			}
 
 			FRayTracingDecalPayload DecalPayload = (FRayTracingDecalPayload)0;
 			DecalPayload.Params.SetInputTranslatedWorldPosition(HitPayload.TranslatedWorldPos);
 
 			TraceRay(
 				DecalTLAS,
-				RAY_FLAG_CULL_FRONT_FACING_TRIANGLES,
+				RAY_FLAG_CULL_BACK_FACING_TRIANGLES,
 				InstanceInclusionMask,
 				RAY_TRACING_SHADER_SLOT_MATERIAL,
 				RAY_TRACING_NUM_SHADER_SLOTS,
diff --git a/Engine/Shaders/Private/RayTracing/RayTracingDecalMaterialShader.usf b/Engine/Shaders/Private/RayTracing/RayTracingDecalMaterialShader.usf
--- a/Engine/Shaders/Private/RayTracing/RayTracingDecalMaterialShader.usf
+++ b/Engine/Shaders/Private/RayTracing/RayTracingDecalMaterialShader.usf
@@ -155,7 +155,7 @@ RAY_TRACING_ENTRY_CLOSEST_HIT(RayTracingDecalMaterialCHS,
 
 	if (DECAL_PAYLOAD_FLAGS & DECAL_WRITE_NORMAL_FLAG)
 	{
-		Payload.Params.SetWorldNormal(-MaterialParameters.WorldNormal);
+		Payload.Params.SetWorldNormal(MaterialParameters.WorldNormal);
 	}
 
 	if (DECAL_PAYLOAD_FLAGS & DECAL_WRITE_ROUGHNESS_SPECULAR_METALLIC_FLAG)
diff --git a/Engine/Source/Runtime/Renderer/Private/PathTracing.cpp b/Engine/Source/Runtime/Renderer/Private/PathTracing.cpp
--- a/Engine/Source/Runtime/Renderer/Private/PathTracing.cpp
+++ b/Engine/Source/Runtime/Renderer/Private/PathTracing.cpp
@@ -301,6 +301,13 @@ TAutoConsoleVariable<float> CVarPathTracingMeshDecalBias(
 	ECVF_RenderThreadSafe
 );
 
+TAutoConsoleVariable<int32> CVarPathTracingMeshDecalsUseViewRay(
+	TEXT("r.PathTracing.MeshDecalUsesViewRay"),
+	1,
+	TEXT("Controls if mesh decals should be probed by the viewing ray (default) or by a short ray perpendicular to the surface (experimental)"),
+	ECVF_RenderThreadSafe
+);
+
 TAutoConsoleVariable<int32> CVarPathTracingLightFunctionColor(
 	TEXT("r.PathTracing.LightFunctionColor"),
 	0,
@@ -330,6 +337,7 @@ BEGIN_SHADER_PARAMETER_STRUCT(FPathTracingData, )
 	SHADER_PARAMETER(uint32, EnableAtmosphere)
 	SHADER_PARAMETER(uint32, EnableFog)
 	SHADER_PARAMETER(int32, MaxRaymarchSteps)
+	SHADER_PARAMETER(int32, MeshDecalsUseViewRay)
 	SHADER_PARAMETER(float, MaxPathIntensity)
 	SHADER_PARAMETER(float, MaxNormalBias)
 	SHADER_PARAMETER(float, FilterWidth)
@@ -380,6 +388,7 @@ struct FPathTracingConfig
 			PathTracingData.DecalRoughnessCutoff != Other.PathTracingData.DecalRoughnessCutoff ||
 			PathTracingData.MeshDecalRoughnessCutoff != Other.PathTracingData.MeshDecalRoughnessCutoff ||
 			PathTracingData.MeshDecalBias != Other.PathTracingData.MeshDecalBias ||
+			PathTracingData.MeshDecalsUseViewRay != Other.PathTracingData.MeshDecalsUseViewRay ||
 			PathTracingData.MaxRaymarchSteps != Other.PathTracingData.MaxRaymarchSteps ||
 			ViewRect != Other.ViewRect ||
 			LightShowFlags != Other.LightShowFlags ||
@@ -557,6 +566,7 @@ static void PreparePathTracingData(const FScene* Scene, const FViewInfo& View, F
 
 	PathTracingData.MeshDecalRoughnessCutoff = PathTracing::UsesDecals(*View.Family) && Scene->RayTracingScene.GetRHIRayTracingScene()->GetInitializer().NumNativeInstancesPerLayer[(uint32)ERayTracingSceneLayer::Decals] > 0 ? CVarPathTracingMeshDecalRoughnessCutoff.GetValueOnRenderThread() : -1.0f;
 	PathTracingData.MeshDecalBias = CVarPathTracingMeshDecalBias.GetValueOnRenderThread();
+	PathTracingData.MeshDecalsUseViewRay = CVarPathTracingMeshDecalsUseViewRay.GetValueOnRenderThread();
 
 	PathTracingData.MaxRaymarchSteps = CVarPathTracingMaxRaymarchSteps.GetValueOnRenderThread();
 }
