Check that FRayTracingMeshCommand marked as bDecal uses a shader with ERayTracingPayloadType::MeshDecals.

diff --git a/Engine/Source/Runtime/Renderer/Private/RayTracing/RayTracingMaterialHitShaders.cpp b/Engine/Source/Runtime/Renderer/Private/RayTracing/RayTracingMaterialHitShaders.cpp
--- a/Engine/Source/Runtime/Renderer/Private/RayTracing/RayTracingMaterialHitShaders.cpp
+++ b/Engine/Source/Runtime/Renderer/Private/RayTracing/RayTracingMaterialHitShaders.cpp
@@ -801,10 +801,13 @@ FRayTracingPipelineState* FDeferredShadingSceneRenderer::CreateRayTracingMateria
 						const FVisibleRayTracingMeshCommand VisibleMeshCommand = MeshCommands[CommandIndex];
 						const FRayTracingMeshCommand& MeshCommand = *VisibleMeshCommand.RayTracingMeshCommand;
 
+						const bool bIsMeshDecalShader = MeshCommand.MaterialShader->RayTracingPayloadType == (uint32)ERayTracingPayloadType::MeshDecals;
+						check(bIsMeshDecalShader == MeshCommand.bDecal);
+
 						// Force the same shader to be used on all geometry unless materials are enabled
 						int32 HitGroupIndex;
 
-						if (MeshCommand.MaterialShader->RayTracingPayloadType == (uint32)ERayTracingPayloadType::MeshDecals)
+						if (bIsMeshDecalShader)
 						{
 							checkf(bSupportMeshDecals, TEXT("Unexpected ray tracing mesh command using Mesh Decal payload. Fix logic adding the command or update bSupportMeshDecals as appropriate."));
 							HitGroupIndex = OpaqueMeshDecalHitGroupIndex;
