Adding missing includes.

diff --git a/Engine/Source/Runtime/Renderer/Private/GPUScene.h b/Engine/Source/Runtime/Renderer/Private/GPUScene.h
--- a/Engine/Source/Runtime/Renderer/Private/GPUScene.h
+++ b/Engine/Source/Runtime/Renderer/Private/GPUScene.h
@@ -17,6 +17,7 @@ class FViewInfo;
 class FLightSceneInfoCompact;
 class FGPUScene;
 class FGPUSceneDynamicContext;
+class FViewUniformShaderParameters;
 
 BEGIN_SHADER_PARAMETER_STRUCT(FGPUSceneResourceParameters, )
 	SHADER_PARAMETER_RDG_BUFFER_SRV(StructuredBuffer<float4>, GPUSceneInstanceSceneData)
diff --git a/Engine/Source/Runtime/Renderer/Private/IlluminanceMeter.cpp b/Engine/Source/Runtime/Renderer/Private/IlluminanceMeter.cpp
--- a/Engine/Source/Runtime/Renderer/Private/IlluminanceMeter.cpp
+++ b/Engine/Source/Runtime/Renderer/Private/IlluminanceMeter.cpp
@@ -4,6 +4,7 @@
 	Functionality for a real time illuminance meter reference for the sky light and sun.
 =============================================================================*/
 
+#include "DataDrivenShaderPlatformInfo.h"
 #include "ReflectionEnvironmentCapture.h"
 #include "RHIResources.h"
 #include "RHIShaderPlatform.h"
diff --git a/Engine/Source/Runtime/Renderer/Private/RayTracing/RayTracingInstanceMask.cpp b/Engine/Source/Runtime/Renderer/Private/RayTracing/RayTracingInstanceMask.cpp
--- a/Engine/Source/Runtime/Renderer/Private/RayTracing/RayTracingInstanceMask.cpp
+++ b/Engine/Source/Runtime/Renderer/Private/RayTracing/RayTracingInstanceMask.cpp
@@ -4,6 +4,7 @@
 
 #if RHI_RAYTRACING
 
+#include "DataDrivenShaderPlatformInfo.h"
 #include "RayTracingDefinitions.h"
 #include "PathTracingDefinitions.h"
 #include "PrimitiveSceneProxy.h"
diff --git a/Engine/Source/Runtime/Renderer/Private/SparseVolumeTexture/SparseVolumeTextureViewerRendering.cpp b/Engine/Source/Runtime/Renderer/Private/SparseVolumeTexture/SparseVolumeTextureViewerRendering.cpp
--- a/Engine/Source/Runtime/Renderer/Private/SparseVolumeTexture/SparseVolumeTextureViewerRendering.cpp
+++ b/Engine/Source/Runtime/Renderer/Private/SparseVolumeTexture/SparseVolumeTextureViewerRendering.cpp
@@ -4,6 +4,7 @@
 #include "SparseVolumeTexture/SparseVolumeTextureViewerSceneProxy.h"
 #include "SparseVolumeTexture/SparseVolumeTexture.h"
 
+#include "DataDrivenShaderPlatformInfo.h"
 #include "SceneView.h"
 #include "ScenePrivate.h"
 #include "SceneRendering.h"
diff --git a/Engine/Source/Runtime/Renderer/Private/Strata/Strata.cpp b/Engine/Source/Runtime/Renderer/Private/Strata/Strata.cpp
--- a/Engine/Source/Runtime/Renderer/Private/Strata/Strata.cpp
+++ b/Engine/Source/Runtime/Renderer/Private/Strata/Strata.cpp
@@ -9,6 +9,7 @@
 #include "SceneRendering.h"
 #include "RendererInterface.h"
 #include "UniformBuffer.h"
+#include "PostProcess/SceneRenderTargets.h"
 #include "SceneTextureParameters.h"
 #include "ShaderCompiler.h"
 #include "Lumen/Lumen.h"
