Change checks regading mesh decals in raytracing since it's not valid in code paths at the moment.

diff --git a/Engine/Source/Runtime/Renderer/Private/RayTracing/RayTracingMaterialHitShaders.cpp b/Engine/Source/Runtime/Renderer/Private/RayTracing/RayTracingMaterialHitShaders.cpp
--- a/Engine/Source/Runtime/Renderer/Private/RayTracing/RayTracingMaterialHitShaders.cpp
+++ b/Engine/Source/Runtime/Renderer/Private/RayTracing/RayTracingMaterialHitShaders.cpp
@@ -802,14 +802,17 @@ FRayTracingPipelineState* FDeferredShadingSceneRenderer::CreateRayTracingMateria
 						const FRayTracingMeshCommand& MeshCommand = *VisibleMeshCommand.RayTracingMeshCommand;
 
 						const bool bIsMeshDecalShader = MeshCommand.MaterialShader->RayTracingPayloadType == (uint32)ERayTracingPayloadType::MeshDecals;
-						check(bIsMeshDecalShader == MeshCommand.bDecal);
+
+						// TODO: Following check is disabled since FRayTracingMeshProcessor non-path-tracing code paths still don't assign the appropriate shader to decal mesh commands.
+						// We could also potentially use regular materials to approximate decals in ray tracing in some situations.
+						// check(bIsMeshDecalShader == MeshCommand.bDecal);
 
 						// Force the same shader to be used on all geometry unless materials are enabled
 						int32 HitGroupIndex;
 
 						if (bIsMeshDecalShader)
 						{
-							checkf(bSupportMeshDecals, TEXT("Unexpected ray tracing mesh command using Mesh Decal payload. Fix logic adding the command or update bSupportMeshDecals as appropriate."));
+							checkf(bSupportMeshDecals && MeshCommand.bDecal, TEXT("Unexpected ray tracing mesh command using Mesh Decal payload. Fix logic adding the command or update bSupportMeshDecals as appropriate."));
 							HitGroupIndex = OpaqueMeshDecalHitGroupIndex;
 						}
 						else
