Fixed: Uninitialized value that was used to select a post-process material permutation

diff --git a/Engine/Source/Runtime/Renderer/Private/PostProcess/PostProcessMaterial.h b/Engine/Source/Runtime/Renderer/Private/PostProcess/PostProcessMaterial.h
--- a/Engine/Source/Runtime/Renderer/Private/PostProcess/PostProcessMaterial.h
+++ b/Engine/Source/Runtime/Renderer/Private/PostProcess/PostProcessMaterial.h
@@ -105,7 +105,7 @@ struct FPostProcessMaterialInputs
 	EPixelFormat OutputFormat = PF_Unknown;
 
 	/** Whether or not the stencil test must be done in the pixel shader rather than rasterizer state. */
-	bool bManualStencilTest;
+	bool bManualStencilTest = false;
 
 	/** Custom depth/stencil used for stencil operations. */
 	FRDGTextureRef CustomDepthTexture = nullptr;
