Fix directional light CSM instance culling bug.

diff --git a/Engine/Shaders/Private/InstanceCulling/BuildInstanceDrawCommands.usf b/Engine/Shaders/Private/InstanceCulling/BuildInstanceDrawCommands.usf
--- a/Engine/Shaders/Private/InstanceCulling/BuildInstanceDrawCommands.usf
+++ b/Engine/Shaders/Private/InstanceCulling/BuildInstanceDrawCommands.usf
@@ -207,9 +207,6 @@ bool IsInstanceVisible(FInstanceSceneData InstanceData, uint ViewIdIndex)
 		return true;
 	}
 
-	// TODO: impact? make permutation for shadow views? Or use a dynamic branch?
-	bool bNearClip = true;
-
 	// TODO: remove this indirection and go straight to data index
 	checkStructuredBufferAccessSlow(ViewIds, ViewIdIndex);
 
@@ -222,6 +219,7 @@ bool IsInstanceVisible(FInstanceSceneData InstanceData, uint ViewIdIndex)
 		float4x4 LocalToTranslatedWorld = LWCMultiplyTranslation(InstanceData.LocalToWorld, NaniteView.PreViewTranslation);
 		float4x4 LocalToClip = mul(LocalToTranslatedWorld, NaniteView.TranslatedWorldToClip);
 		const bool bIsOrtho = IsOrthoProjection(NaniteView.ViewToClip);
+		const bool bNearClip = (NaniteView.Flags & NANITE_VIEW_FLAG_NEAR_CLIP) != 0;
 
 		FFrustumCullData Cull = BoxCullFrustum(InstanceData.LocalBoundsCenter, InstanceData.LocalBoundsExtent, LocalToClip, bIsOrtho, bNearClip, false);
 
diff --git a/Engine/Source/Runtime/Renderer/Private/ShadowSetup.cpp b/Engine/Source/Runtime/Renderer/Private/ShadowSetup.cpp
--- a/Engine/Source/Runtime/Renderer/Private/ShadowSetup.cpp
+++ b/Engine/Source/Runtime/Renderer/Private/ShadowSetup.cpp
@@ -2151,7 +2151,7 @@ void FProjectedShadowInfo::SetupMeshDrawCommandsForShadowDepth(FSceneRenderer& R
 			// Setup packed view
 			TArray<Nanite::FPackedView, SceneRenderingAllocator> PackedViews;
 			{
-				Nanite::FPackedViewParams Params;
+				Nanite::FPackedViewParams Params{};
 				Params.ViewMatrices = GetShadowDepthRenderingViewMatrices(CubemapFaceIndex);
 				// TODO: Real prev frame matrices
 				Params.PrevViewMatrices = Params.ViewMatrices;
@@ -2165,7 +2165,7 @@ void FProjectedShadowInfo::SetupMeshDrawCommandsForShadowDepth(FSceneRenderer& R
 	else if (VirtualShadowMapClipmap.IsValid())
 	{
 		// TODO: Register view per clip level such that they are culled early (?)
-		Nanite::FPackedViewParams Params;
+		Nanite::FPackedViewParams Params{};
 		// Note: To ensure conservative culling, we get the coarsest clip-level view since it covers the finer ones.
 		Params.ViewMatrices = VirtualShadowMapClipmap->GetViewMatrices(VirtualShadowMapClipmap->GetLevelCount() - 1);
 		// TODO: Real prev frame matrices
@@ -2177,12 +2177,16 @@ void FProjectedShadowInfo::SetupMeshDrawCommandsForShadowDepth(FSceneRenderer& R
 	}
 	else
 	{
-		Nanite::FPackedViewParams Params;
+		Nanite::FPackedViewParams Params{};
 		Params.ViewMatrices = GetShadowDepthRenderingViewMatrices();
 		// TODO: Real prev frame matrices
 		Params.PrevViewMatrices = Params.ViewMatrices;
 		Params.ViewRect = GetInnerViewRect();
 		Params.RasterContextSize = FIntPoint(ResolutionX, ResolutionY);
+		if (IsWholeSceneDirectionalShadow())
+		{
+			Params.Flags &= ~NANITE_VIEW_FLAG_NEAR_CLIP;
+		}
 		Params.MaxPixelsPerEdgeMultipler = 1.0f;
 		ViewIds.Add(InstanceCullingManager.RegisterView(Params));
 	}
