Clamp depth bound to fix RHI validation warning

diff --git a/Engine/Source/Runtime/Renderer/Private/IndirectLightRendering.cpp b/Engine/Source/Runtime/Renderer/Private/IndirectLightRendering.cpp
--- a/Engine/Source/Runtime/Renderer/Private/IndirectLightRendering.cpp
+++ b/Engine/Source/Runtime/Renderer/Private/IndirectLightRendering.cpp
@@ -1327,7 +1327,7 @@ void FDeferredShadingSceneRenderer::RenderDiffuseIndirectAndAmbientOcclusion(
 							const FFinalPostProcessSettings& Settings = View.FinalPostProcessSettings;
 							const FMatrix& ProjectionMatrix = View.ViewMatrices.GetProjectionMatrix();
 							const FVector4f Far = (FVector4f)ProjectionMatrix.TransformFVector4(FVector4(0, 0, Settings.AmbientOcclusionFadeDistance));
-							float DepthFar = FMath::Min(1.0f, Far.Z / Far.W);
+							float DepthFar = FMath::Clamp(Far.Z / Far.W, 0.0f, 1.0f);
 							RHICmdList.SetDepthBounds(DepthFar, 1.0f);
 						}
 
