Moved illuminance shader permutation logic to a common AutoExposurePermutation namespace.

diff --git a/Engine/Source/Runtime/Renderer/Private/PostProcess/PostProcessEyeAdaptation.cpp b/Engine/Source/Runtime/Renderer/Private/PostProcess/PostProcessEyeAdaptation.cpp
--- a/Engine/Source/Runtime/Renderer/Private/PostProcess/PostProcessEyeAdaptation.cpp
+++ b/Engine/Source/Runtime/Renderer/Private/PostProcess/PostProcessEyeAdaptation.cpp
@@ -117,6 +117,40 @@ namespace
 		ECVF_RenderThreadSafe);
 }
 
+namespace AutoExposurePermutation
+{
+	bool ShouldCompileCommonPermutation(const FCommonDomain& PermutationVector)
+	{
+		if (PermutationVector.Get<FUseApproxIlluminanceDim>() && PermutationVector.Get<FUsePrecalculatedLuminanceDim>())
+		{
+			return false;
+		}
+
+		return true;
+	}
+
+	FCommonDomain BuildCommonPermutationDomain()
+	{
+		FCommonDomain PermutationVector;
+
+		if (CVarAutoExposureIgnoreMaterials.GetValueOnRenderThread())
+		{
+			if (CVarAutoExposureIgnoreMaterialsUsePrecalculatedIlluminance.GetValueOnRenderThread())
+			{
+				PermutationVector.Set<FUsePrecalculatedLuminanceDim>(true);
+			}
+			else
+			{
+				PermutationVector.Set<FUseApproxIlluminanceDim>(true);
+			}
+		}
+
+		PermutationVector.Set<FUseDebugOutputDim>(CVarAutoExposureIgnoreMaterialsDebug.GetValueOnRenderThread());
+
+		return PermutationVector;
+	}
+} // namespace TonemapperPermutation
+
 // Basic eye adaptation is supported everywhere except mobile when MobileHDR is disabled
 static ERHIFeatureLevel::Type GetBasicEyeAdaptationMinFeatureLevel()
 {
diff --git a/Engine/Source/Runtime/Renderer/Private/PostProcess/PostProcessEyeAdaptation.h b/Engine/Source/Runtime/Renderer/Private/PostProcess/PostProcessEyeAdaptation.h
--- a/Engine/Source/Runtime/Renderer/Private/PostProcess/PostProcessEyeAdaptation.h
+++ b/Engine/Source/Runtime/Renderer/Private/PostProcess/PostProcessEyeAdaptation.h
@@ -26,6 +26,19 @@ FORCEINLINE float Log2ToEV100(float LuminanceMax, float Log2)
 	return Log2 - FMath::Log2(LuminanceMax);
 }
 
+namespace AutoExposurePermutation
+{
+	class FUsePrecalculatedLuminanceDim : SHADER_PERMUTATION_BOOL("USE_PRECALCULATED_LUMINANCE");
+	class FUseApproxIlluminanceDim : SHADER_PERMUTATION_BOOL("USE_APPROX_ILLUMINANCE");
+	class FUseDebugOutputDim : SHADER_PERMUTATION_BOOL("USE_DEBUG_OUTPUT");
+
+	using FCommonDomain = TShaderPermutationDomain<FUsePrecalculatedLuminanceDim, FUseApproxIlluminanceDim, FUseDebugOutputDim>;
+
+	bool ShouldCompileCommonPermutation(const FCommonDomain& PermutationVector);
+
+	FCommonDomain BuildCommonPermutationDomain();
+} // namespace AutoExposurePermutation
+
 // For converting the auto exposure to new values from 4.24 to 4.25
 float CalculateEyeAdaptationParameterExposureConversion(const FPostProcessSettings& Settings, const bool bExtendedLuminanceRange);
 
diff --git a/Engine/Source/Runtime/Renderer/Private/PostProcess/PostProcessHistogram.cpp b/Engine/Source/Runtime/Renderer/Private/PostProcess/PostProcessHistogram.cpp
--- a/Engine/Source/Runtime/Renderer/Private/PostProcess/PostProcessHistogram.cpp
+++ b/Engine/Source/Runtime/Renderer/Private/PostProcess/PostProcessHistogram.cpp
@@ -149,21 +149,14 @@ public:
 		SHADER_PARAMETER_SAMPLER(SamplerState, PreIntegratedGFSampler)
 	END_SHADER_PARAMETER_STRUCT()
 
-	class FUsePrecalculatedLuminanceDim : SHADER_PERMUTATION_BOOL("USE_PRECALCULATED_LUMINANCE");
-	class FUseApproxIlluminanceDim : SHADER_PERMUTATION_BOOL("USE_APPROX_ILLUMINANCE");
-	class FUseDebugOutputDim : SHADER_PERMUTATION_BOOL("USE_DEBUG_OUTPUT");
-	using FPermutationDomain = TShaderPermutationDomain<FUsePrecalculatedLuminanceDim, FUseApproxIlluminanceDim, FUseDebugOutputDim>;
+	using FPermutationDomain = TShaderPermutationDomain<AutoExposurePermutation::FCommonDomain>;
 
 	static bool ShouldCompilePermutation(const FGlobalShaderPermutationParameters& Parameters)
 	{
 		const FPermutationDomain PermutationVector(Parameters.PermutationId);
 
-		if (PermutationVector.Get<FUsePrecalculatedLuminanceDim>() && PermutationVector.Get<FUseApproxIlluminanceDim>())
-		{
-			return false;
-		}
-
-		return IsFeatureLevelSupported(Parameters.Platform, ERHIFeatureLevel::SM5);
+		return IsFeatureLevelSupported(Parameters.Platform, ERHIFeatureLevel::SM5)
+			&& AutoExposurePermutation::ShouldCompileCommonPermutation(PermutationVector.Get<AutoExposurePermutation::FCommonDomain>());
 	}
 
 	static void ModifyCompilationEnvironment(const FGlobalShaderPermutationParameters& Parameters, FShaderCompilerEnvironment& OutEnvironment)
@@ -361,7 +354,6 @@ static FRDGTextureRef AddHistogramAtomicPass(
 			PassParameters->DebugOutput = GraphBuilder.CreateUAV(DebugOutputTexture);
 		}
 
-
 		PassParameters->PreIntegratedGF = GSystemTextures.PreintegratedGF->GetRHI();
 		PassParameters->PreIntegratedGFSampler = TStaticSamplerState<SF_Bilinear, AM_Clamp, AM_Clamp, AM_Clamp>::GetRHI();
 
@@ -370,23 +362,8 @@ static FRDGTextureRef AddHistogramAtomicPass(
 		// AddClearUAVPass(GraphBuilder, PassParameters->HistogramScatter64Output, ClearValues);
 		AddClearUAVPass(GraphBuilder, PassParameters->HistogramScatter32Output, ClearValues);
 
-		static const auto CVarIgnoreMaterials = IConsoleManager::Get().FindTConsoleVariableDataBool(TEXT("r.AutoExposure.IgnoreMaterials"));
-		static const auto CVarIgnoreMaterialsDebug = IConsoleManager::Get().FindTConsoleVariableDataBool(TEXT("r.AutoExposure.IgnoreMaterials.Debug"));
-		static const auto CVarUsePrecalculatedIlluminance = IConsoleManager::Get().FindTConsoleVariableDataBool(TEXT("r.AutoExposure.IgnoreMaterials.UsePrecalculatedIlluminance"));
-
 		FHistogramAtomicCS::FPermutationDomain PermutationVector;
-		if (CVarIgnoreMaterials->GetValueOnRenderThread())
-		{
-			if (CVarUsePrecalculatedIlluminance->GetValueOnRenderThread())
-			{
-				PermutationVector.Set<FHistogramAtomicCS::FUsePrecalculatedLuminanceDim>(true);
-			}
-			else
-			{
-				PermutationVector.Set<FHistogramAtomicCS::FUseApproxIlluminanceDim>(true);
-			}
-		}
-		PermutationVector.Set<FHistogramAtomicCS::FUseDebugOutputDim>(CVarIgnoreMaterialsDebug->GetValueOnRenderThread());
+		PermutationVector.Set<AutoExposurePermutation::FCommonDomain>(AutoExposurePermutation::BuildCommonPermutationDomain());
 
 		auto ComputeShader = View.ShaderMap->GetShader<FHistogramAtomicCS>(PermutationVector);
 
