Multi-GPU:  Scope TLAS portion of ray tracing scene update to the view's GPUs (BLAS related updates still run on all GPUs).  Saves around 0.5 ms per GPU and 3.5 ms RHI thread CPU (the bottleneck thread) with two views on the Virtual Production Valley Lumen map.  Can be live toggled for debugging.

For Niagara GPU sim ray traced particle collisions, we disable ray traced collision when the sim is ticked outside the renderer with multi-GPU, because the TLAS data won't necessarily be up to date on the GPU the Niagara sim is running on.  We could solve this later, if necessary, but it's low priority for now.

#rnx

diff --git a/Engine/Plugins/FX/Niagara/Source/Niagara/Private/NiagaraGpuComputeDispatch.cpp b/Engine/Plugins/FX/Niagara/Source/Niagara/Private/NiagaraGpuComputeDispatch.cpp
--- a/Engine/Plugins/FX/Niagara/Source/Niagara/Private/NiagaraGpuComputeDispatch.cpp
+++ b/Engine/Plugins/FX/Niagara/Source/Niagara/Private/NiagaraGpuComputeDispatch.cpp
@@ -431,6 +431,9 @@ void FNiagaraGpuComputeDispatch::ProcessPendingTicksFlush(FRHICommandListImmedia
 			// in the event of submitting commands, i.e. when we write a fence, or indeed perform a manual flush.
 			RHICmdList.BeginScene();
 
+			// Allow downstream logic to detect we are running pending ticks outside the scene renderer
+			bIsOutsideSceneRenderer = true;
+
 			// Execute all ticks that we can support without invalid simulations
 			MaxTicksToFlush = TickFlushMaxPendingTicks;
 			for (int32 iTickBatch = 0; iTickBatch < MaxPendingTicks; iTickBatch+=MaxTicksToFlush)
@@ -456,6 +459,8 @@ void FNiagaraGpuComputeDispatch::ProcessPendingTicksFlush(FRHICommandListImmedia
 			}
 			MaxTicksToFlush = TNumericLimits<int32>::Max();
 
+			bIsOutsideSceneRenderer = false;
+
 			// We have completed flushing the commands
 			RHICmdList.EndScene();
 			break;
diff --git a/Engine/Plugins/FX/Niagara/Source/Niagara/Public/NiagaraGpuComputeDispatchInterface.h b/Engine/Plugins/FX/Niagara/Source/Niagara/Public/NiagaraGpuComputeDispatchInterface.h
--- a/Engine/Plugins/FX/Niagara/Source/Niagara/Public/NiagaraGpuComputeDispatchInterface.h
+++ b/Engine/Plugins/FX/Niagara/Source/Niagara/Public/NiagaraGpuComputeDispatchInterface.h
@@ -142,6 +142,8 @@ public:
 
 	virtual FNiagaraAsyncGpuTraceHelper& GetAsyncGpuTraceHelper() const = 0;
 
+	bool IsOutsideSceneRenderer() const { return bIsOutsideSceneRenderer; }
+
 #if WITH_MGPU
 	/**
 	Notify that a GPU resource was modified that will impact MultiGPU rendering.
@@ -180,4 +182,6 @@ protected:
 	TArray<TPair<FName, TUniquePtr<FNiagaraGpuComputeDataManager>>> GpuDataManagers;
 
 	TConstArrayView<FViewInfo>				SimulationViewInfos;
+
+	bool									bIsOutsideSceneRenderer = false;
 };
diff --git a/Engine/Plugins/FX/Niagara/Source/NiagaraShader/Private/NiagaraAsyncGpuTraceProviderHwrt.cpp b/Engine/Plugins/FX/Niagara/Source/NiagaraShader/Private/NiagaraAsyncGpuTraceProviderHwrt.cpp
--- a/Engine/Plugins/FX/Niagara/Source/NiagaraShader/Private/NiagaraAsyncGpuTraceProviderHwrt.cpp
+++ b/Engine/Plugins/FX/Niagara/Source/NiagaraShader/Private/NiagaraAsyncGpuTraceProviderHwrt.cpp
@@ -280,7 +280,18 @@ bool FNiagaraAsyncGpuTraceProviderHwrt::IsSupported()
 
 bool FNiagaraAsyncGpuTraceProviderHwrt::IsAvailable() const
 {
-	if (!GNiagaraAsyncGpuTraceHwrtEnabled)
+	// Disable HWRT collision if there are multiple GPUs and this is a pending tick.  The ray tracing scene may be in an indeterminate state
+	// on a given GPU outside the context of the scene renderer, where pending ticks are handled.  To handle this properly, we would either
+	// need to provide a mechanism to force Niagara to run on the GPU where the ray tracing scene is valid, or refresh the ray tracing scene
+	// on GPU 0.  Either solution is a bit complicated, so we'll wait to implement one until it becomes important.  Virtual Production is the
+	// main client of multi-GPU, and they generally don't de-focus the application to create a scenario where pending ticks accumulate,
+	// so they are unlikely to encounter this issue.
+	//
+	// Note that pending ticks may run into inconsistencies with multiple views, even on a single GPU, because collision will be run against the
+	// last view, rather than the first view, which is normally where Niagara runs.  Longer term, it might be interesting to support running
+	// Niagara against all view's ray tracing scenes, and if you wanted to go that route, the solution would be to preserve the otherwise
+	// transient TLAS data for all views.  That could solve this issue as well.
+	if (!GNiagaraAsyncGpuTraceHwrtEnabled || (GNumExplicitGPUsForRendering > 1 && Dispatcher->IsOutsideSceneRenderer()))
 	{
 		return false;
 	}
diff --git a/Engine/Source/Runtime/Renderer/Private/DeferredShadingRenderer.cpp b/Engine/Source/Runtime/Renderer/Private/DeferredShadingRenderer.cpp
--- a/Engine/Source/Runtime/Renderer/Private/DeferredShadingRenderer.cpp
+++ b/Engine/Source/Runtime/Renderer/Private/DeferredShadingRenderer.cpp
@@ -261,6 +261,15 @@ static FAutoConsoleVariableRef CVarRayTracingDebugForceOpaque(
 	TEXT("Forces all ray tracing geometry instances to be opaque, effectively disabling any-hit shaders. This is useful for debugging and profiling. (default = 0)")
 );
 
+static int32 GRayTracingMultiGpuTLASMask = 1;
+static FAutoConsoleVariableRef CVarRayTracingMultiGpuTLASMask(
+	TEXT("r.RayTracing.MultiGpuMaskTLAS"),
+	GRayTracingMultiGpuTLASMask,
+	TEXT("For Multi-GPU, controls which GPUs TLAS and material pipeline updates run on.  (default = 1)\n")
+	TEXT(" 0: Run TLAS and material pipeline updates on all GPUs.  Original behavior, which may be useful for debugging.\n")
+	TEXT(" 1: Run TLAS and material pipeline updates masked to the active view's GPUs to improve performance.  BLAS updates still run on all GPUs.")
+);
+
 static TAutoConsoleVariable<int32> CVarSceneDepthHZBAsyncCompute(
 	TEXT("r.SceneDepthHZBAsyncCompute"), 0,
 	TEXT("Selects whether HZB for scene depth buffer should be built with async compute.\n")
@@ -1952,7 +1961,8 @@ bool FDeferredShadingSceneRenderer::DispatchRayTracingWorldUpdates(FRDGBuilder&
 		}
 	}
 
-	RDG_GPU_MASK_SCOPE(GraphBuilder, FRHIGPUMask::All());
+	// Keep mask the same as what's already set (which will be the view mask) if TLAS updates should be masked to the view
+	RDG_GPU_MASK_SCOPE(GraphBuilder, GRayTracingMultiGpuTLASMask ? GraphBuilder.RHICmdList.GetGPUMask() : FRHIGPUMask::All());
 
 	RayTracingScene.CreateWithInitializationData(GraphBuilder, &Scene->GPUScene, ReferenceView.ViewMatrices, MoveTemp(ReferenceView.RayTracingSceneInitData));
 
@@ -1977,44 +1987,65 @@ bool FDeferredShadingSceneRenderer::DispatchRayTracingWorldUpdates(FRDGBuilder&
 	{
 		RDG_GPU_STAT_SCOPE(GraphBuilder, RayTracingUpdate);
 
-		FBuildAccelerationStructurePassParams* PassParams = GraphBuilder.AllocParameters<FBuildAccelerationStructurePassParams>();
-		PassParams->RayTracingSceneScratchBuffer = Scene->RayTracingScene.BuildScratchBuffer;
-		PassParams->RayTracingSceneInstanceBuffer = Scene->RayTracingScene.InstanceBuffer;
-		PassParams->View = ReferenceView.ViewUniformBuffer;
-		PassParams->DynamicGeometryScratchBuffer = OutDynamicGeometryScratchBuffer;
-		PassParams->LightDataPacked = nullptr;
-		PassParams->ClusterPageData = nullptr;
-		PassParams->HierarchyBuffer = nullptr;
-		PassParams->RayTracingDataBuffer = nullptr;
-
-		// Use ERDGPassFlags::NeverParallel so the pass never runs off the render thread and we always get the following order of execution on the CPU:
-		// BuildTLASInstanceBuffer, RayTracingUpdate, RayTracingEndUpdate, ..., ReleaseRayTracingResources
-		GraphBuilder.AddPass(RDG_EVENT_NAME("RayTracingUpdate"), PassParams, ComputePassFlags | ERDGPassFlags::NeverCull | ERDGPassFlags::NeverParallel,
-			[this, PassParams, bRayTracingAsyncBuild](FRHICommandListImmediate& RHICmdList)
 		{
+			// Dynamic geometry (BLAS) updates must always run on all GPUs.  Other passes may either run on all GPUs or be scoped to the view's GPUs.
+			// See GRayTracingMultiGpuTLASMask.
+			RDG_GPU_MASK_SCOPE(GraphBuilder, FRHIGPUMask::All());
+
+			FBuildAccelerationStructurePassParams* PassParams = GraphBuilder.AllocParameters<FBuildAccelerationStructurePassParams>();
+			PassParams->RayTracingSceneScratchBuffer = nullptr;
+			PassParams->RayTracingSceneInstanceBuffer = nullptr;
+			PassParams->View = ReferenceView.ViewUniformBuffer;
+			PassParams->DynamicGeometryScratchBuffer = OutDynamicGeometryScratchBuffer;
+			PassParams->LightDataPacked = nullptr;
+			PassParams->ClusterPageData = nullptr;
+			PassParams->HierarchyBuffer = nullptr;
+			PassParams->RayTracingDataBuffer = nullptr;
+
+			// Use ERDGPassFlags::NeverParallel so the pass never runs off the render thread and we always get the following order of execution on the CPU:
+			// BuildTLASInstanceBuffer, RayTracingDynamicUpdate, RayTracingUpdate, RayTracingEndUpdate, ..., ReleaseRayTracingResources
+			GraphBuilder.AddPass(RDG_EVENT_NAME("RayTracingDynamicUpdate"), PassParams, ComputePassFlags | ERDGPassFlags::NeverCull | ERDGPassFlags::NeverParallel,
+				[this, PassParams, bRayTracingAsyncBuild](FRHICommandListImmediate& RHICmdList)
 			{
 				SCOPED_GPU_STAT(RHICmdList, RayTracingGeometry);
 				FRHIBuffer* DynamicGeometryScratchBuffer = PassParams->DynamicGeometryScratchBuffer ? PassParams->DynamicGeometryScratchBuffer->GetRHI() : nullptr;
 				Scene->GetRayTracingDynamicGeometryCollection()->DispatchUpdates(RHICmdList, DynamicGeometryScratchBuffer);
-			}
-
-			SCOPED_GPU_STAT(RHICmdList, RayTracingScene);
-
-			FRHIRayTracingScene* RayTracingSceneRHI = Scene->RayTracingScene.GetRHIRayTracingSceneChecked();
-			FRHIBuffer* AccelerationStructureBuffer = Scene->RayTracingScene.GetBufferChecked();
-			FRHIBuffer* ScratchBuffer = PassParams->RayTracingSceneScratchBuffer->GetRHI();
-			FRHIBuffer* InstanceBuffer = PassParams->RayTracingSceneInstanceBuffer->GetRHI();
-
-			FRayTracingSceneBuildParams BuildParams;
-			BuildParams.Scene = RayTracingSceneRHI;
-			BuildParams.ScratchBuffer = ScratchBuffer;
-			BuildParams.ScratchBufferOffset = 0;
-			BuildParams.InstanceBuffer = InstanceBuffer;
-			BuildParams.InstanceBufferOffset = 0;
+			});
+		}
 
-			RHICmdList.BindAccelerationStructureMemory(RayTracingSceneRHI, AccelerationStructureBuffer, 0);
-			RHICmdList.BuildAccelerationStructure(BuildParams);
-		});
+		{
+			FBuildAccelerationStructurePassParams* PassParams = GraphBuilder.AllocParameters<FBuildAccelerationStructurePassParams>();
+			PassParams->RayTracingSceneScratchBuffer = Scene->RayTracingScene.BuildScratchBuffer;
+			PassParams->RayTracingSceneInstanceBuffer = Scene->RayTracingScene.InstanceBuffer;
+			PassParams->View = ReferenceView.ViewUniformBuffer;
+			PassParams->DynamicGeometryScratchBuffer = OutDynamicGeometryScratchBuffer;
+			PassParams->LightDataPacked = nullptr;
+			PassParams->ClusterPageData = nullptr;
+			PassParams->HierarchyBuffer = nullptr;
+			PassParams->RayTracingDataBuffer = nullptr;
+
+			// Use ERDGPassFlags::NeverParallel here too -- see comment above on the previous pass
+			GraphBuilder.AddPass(RDG_EVENT_NAME("RayTracingUpdate"), PassParams, ComputePassFlags | ERDGPassFlags::NeverCull | ERDGPassFlags::NeverParallel,
+				[this, PassParams, bRayTracingAsyncBuild](FRHICommandListImmediate& RHICmdList)
+			{
+				SCOPED_GPU_STAT(RHICmdList, RayTracingScene);
+
+				FRHIRayTracingScene* RayTracingSceneRHI = Scene->RayTracingScene.GetRHIRayTracingSceneChecked();
+				FRHIBuffer* AccelerationStructureBuffer = Scene->RayTracingScene.GetBufferChecked();
+				FRHIBuffer* ScratchBuffer = PassParams->RayTracingSceneScratchBuffer->GetRHI();
+				FRHIBuffer* InstanceBuffer = PassParams->RayTracingSceneInstanceBuffer->GetRHI();
+
+				FRayTracingSceneBuildParams BuildParams;
+				BuildParams.Scene = RayTracingSceneRHI;
+				BuildParams.ScratchBuffer = ScratchBuffer;
+				BuildParams.ScratchBufferOffset = 0;
+				BuildParams.InstanceBuffer = InstanceBuffer;
+				BuildParams.InstanceBufferOffset = 0;
+
+				RHICmdList.BindAccelerationStructureMemory(RayTracingSceneRHI, AccelerationStructureBuffer, 0);
+				RHICmdList.BuildAccelerationStructure(BuildParams);
+			});
+		}
 	}
 
 	AddPass(GraphBuilder, RDG_EVENT_NAME("RayTracingEndUpdate"), [this, bRayTracingAsyncBuild](FRHICommandListImmediate& RHICmdList)
@@ -2034,9 +2065,8 @@ bool FDeferredShadingSceneRenderer::DispatchRayTracingWorldUpdates(FRDGBuilder&
 
 static void ReleaseRaytracingResources(FRDGBuilder& GraphBuilder, TArrayView<FViewInfo> Views, FRayTracingScene &RayTracingScene, bool bIsLastRenderer)
 {
-	// Ray tracing is always set up for all GPUs (see WaitForRayTracingScene), so ClearRayTracingBindings must run on all GPUs as well
-	RDG_GPU_MASK_SCOPE(GraphBuilder, FRHIGPUMask::All());
-
+	// Keep mask the same as what's already set (which will be the view mask) if TLAS updates should be masked to the view
+	RDG_GPU_MASK_SCOPE(GraphBuilder, GRayTracingMultiGpuTLASMask ? GraphBuilder.RHICmdList.GetGPUMask() : FRHIGPUMask::All());
 	AddPass(GraphBuilder, RDG_EVENT_NAME("ReleaseRayTracingResources"), [Views, &RayTracingScene, bIsLastRenderer](FRHICommandListImmediate& RHICmdList)
 	{
 		if (RayTracingScene.IsCreated())
@@ -2074,7 +2104,8 @@ void FDeferredShadingSceneRenderer::WaitForRayTracingScene(FRDGBuilder& GraphBui
 
 	TRACE_CPUPROFILER_EVENT_SCOPE(FDeferredShadingSceneRenderer::WaitForRayTracingScene);
 
-	RDG_GPU_MASK_SCOPE(GraphBuilder, FRHIGPUMask::All());
+	// Keep mask the same as what's already set (which will be the view mask) if TLAS updates should be masked to the view
+	RDG_GPU_MASK_SCOPE(GraphBuilder, GRayTracingMultiGpuTLASMask ? GraphBuilder.RHICmdList.GetGPUMask() : FRHIGPUMask::All());
 
 	SetupRayTracingPipelineStates(GraphBuilder);
 
