Support precalculated luminance permutation in AddHistogramLegacyPass.
- required to support illuminance based auto exposure in this code path.

diff --git a/Engine/Shaders/Private/PostProcessHistogram.usf b/Engine/Shaders/Private/PostProcessHistogram.usf
--- a/Engine/Shaders/Private/PostProcessHistogram.usf
+++ b/Engine/Shaders/Private/PostProcessHistogram.usf
@@ -61,6 +61,8 @@ RWTexture3D<float2> BilateralGridRWTexture;
 RWTexture2D<float4> HistogramRWTexture;
 #endif
 
+RWTexture2D<float> DebugOutput;
+
 // Number of thread groups in the dispatch
 uint2 ThreadGroupCount;
 
@@ -121,6 +123,19 @@ void MainCS(
 				LuminanceVal.z = CalculateEyeAdaptationLuminance(float3(SceneColorR.z, SceneColorG.z, SceneColorB.z));
 				LuminanceVal.w = CalculateEyeAdaptationLuminance(float3(SceneColorR.w, SceneColorG.w, SceneColorB.w));
 
+#if USE_PRECALCULATED_LUMINANCE
+				LuminanceVal = SceneColorR;
+#elif USE_APPROX_ILLUMINANCE
+				// TODO
+#endif
+
+#if USE_DEBUG_OUTPUT
+				DebugOutput[TexelPos - Input_ViewportMin + uint2(0, 1)] = LuminanceVal.r;
+				DebugOutput[TexelPos - Input_ViewportMin + uint2(1, 1)] = LuminanceVal.g;
+				DebugOutput[TexelPos - Input_ViewportMin + uint2(1, 0)] = LuminanceVal.b;
+				DebugOutput[TexelPos - Input_ViewportMin + uint2(0, 0)] = LuminanceVal.a;
+#endif
+
 				// only sample screen weight once per quad
 				float2 ScreenUV = (TexelPos.xy - Input_ViewportMin) * InvViewSize;
 				float ScreenWeight = AdaptationWeightTexture(ScreenUV);
diff --git a/Engine/Source/Runtime/Renderer/Private/PostProcess/PostProcessHistogram.cpp b/Engine/Source/Runtime/Renderer/Private/PostProcess/PostProcessHistogram.cpp
--- a/Engine/Source/Runtime/Renderer/Private/PostProcess/PostProcessHistogram.cpp
+++ b/Engine/Source/Runtime/Renderer/Private/PostProcess/PostProcessHistogram.cpp
@@ -40,15 +40,35 @@ public:
 		SHADER_PARAMETER_RDG_TEXTURE(Texture2D, InputTexture)
 		SHADER_PARAMETER_RDG_TEXTURE_UAV(RWTexture2D<float4>, HistogramRWTexture)
 		SHADER_PARAMETER_RDG_TEXTURE_UAV(RWTexture3D<float2>, BilateralGridRWTexture)
+		SHADER_PARAMETER_RDG_TEXTURE_UAV(RWTexture2D<float>, DebugOutput)
 		SHADER_PARAMETER(FIntPoint, ThreadGroupCount)
 	END_SHADER_PARAMETER_STRUCT()
 
 	class FBilateralGrid : SHADER_PERMUTATION_BOOL("BILATERAL_GRID");
-	using FPermutationDomain = TShaderPermutationDomain<FBilateralGrid>;
+	using FPermutationDomain = TShaderPermutationDomain<FBilateralGrid, AutoExposurePermutation::FCommonDomain>;
 
 	static bool ShouldCompilePermutation(const FGlobalShaderPermutationParameters& Parameters)
 	{
-		return IsFeatureLevelSupported(Parameters.Platform, ERHIFeatureLevel::SM5);
+		if (!IsFeatureLevelSupported(Parameters.Platform, ERHIFeatureLevel::SM5))
+		{
+			return false;
+		}
+
+		const FPermutationDomain PermutationVector(Parameters.PermutationId);
+		const AutoExposurePermutation::FCommonDomain& AEPermutationCommon = PermutationVector.Get<AutoExposurePermutation::FCommonDomain>();
+
+		if (!AutoExposurePermutation::ShouldCompileCommonPermutation(AEPermutationCommon))
+		{
+			return false;
+		}
+
+		if (PermutationVector.Get<FBilateralGrid>() && AEPermutationCommon != AutoExposurePermutation::FCommonDomain())
+		{
+			// bilateral grid permutation doesn't use AE permutations
+			return false;
+		}
+
+		return true;
 	}
 
 	static void ModifyCompilationEnvironment(const FGlobalShaderPermutationParameters& Parameters, FShaderCompilerEnvironment& OutEnvironment)
@@ -249,7 +269,14 @@ static FRDGTextureRef AddHistogramLegacyPass(
 		PassParameters->ThreadGroupCount = HistogramThreadGroupCount;
 		PassParameters->EyeAdaptation = EyeAdaptationParameters;
 
+		{
+			const FRDGTextureDesc Desc = FRDGTextureDesc::Create2D(SceneColor.ViewRect.Size(), PF_R16F, FClearValueBinding::Black, TexCreate_UAV | TexCreate_ShaderResource);
+			FRDGTextureRef DebugOutputTexture = GraphBuilder.CreateTexture(Desc, TEXT("EyeAdaptation_DebugOutput"));
+			PassParameters->DebugOutput = GraphBuilder.CreateUAV(DebugOutputTexture);
+		}
+
 		FHistogramCS::FPermutationDomain PermutationVector;
+		PermutationVector.Set<AutoExposurePermutation::FCommonDomain>(AutoExposurePermutation::BuildCommonPermutationDomain());
 
 		auto ComputeShader = View.ShaderMap->GetShader<FHistogramCS>(PermutationVector);
 
