Refactor how instances are added to FRayTracingScene.
- Context: Instances and InstancesDebugData order needs to match. Currently they're managed separately which is error prone.
- Implemented FRayTracingScene::AddInstance(...) to handle this logic.
- Whether InstancesDebugData is enabled is specificed in FRayTracingScene::Reset(...)
- Made some members of FRayTracingScene private.

diff --git a/Engine/Source/Runtime/Renderer/Private/DeferredShadingRenderer.cpp b/Engine/Source/Runtime/Renderer/Private/DeferredShadingRenderer.cpp
--- a/Engine/Source/Runtime/Renderer/Private/DeferredShadingRenderer.cpp
+++ b/Engine/Source/Runtime/Renderer/Private/DeferredShadingRenderer.cpp
@@ -1017,7 +1017,6 @@ bool FDeferredShadingSceneRenderer::GatherRayTracingWorldInstancesForView(FRDGBu
 	);
 
 	const float CurrentWorldTime = View.Family->Time.GetWorldTimeSeconds();
-	const bool bEnableInstanceDebugData = IsRayTracingInstanceDebugDataEnabled(View);
 
 	// Consume output of the relevant primitive gathering task
 	RayTracingScene.UsedCoarseMeshStreamingHandles = MoveTemp(RelevantPrimitiveList.UsedCoarseMeshStreamingHandles);
@@ -1195,17 +1194,9 @@ bool FDeferredShadingSceneRenderer::GatherRayTracingWorldInstancesForView(FRDGBu
 						continue;
 					}
 
-					const uint32 InstanceIndex = RayTracingScene.Instances.Num();
-
-					FRayTracingGeometryInstance& RayTracingInstance = RayTracingScene.Instances.AddDefaulted_GetRef();
+					FRayTracingGeometryInstance RayTracingInstance;
 					RayTracingInstance.GeometryRHI = Geometry->RayTracingGeometryRHI;
 					checkf(RayTracingInstance.GeometryRHI, TEXT("Ray tracing instance must have a valid geometry."));
-					
-					if (bEnableInstanceDebugData)
-					{
-						RayTracingScene.AddInstanceDebugData(RayTracingInstance.GeometryRHI, SceneProxy, true);
-					}					
-
 					RayTracingInstance.DefaultUserData = PrimitiveIndex;
 					RayTracingInstance.bApplyLocalBoundsTransform = Instance.bApplyLocalBoundsTransform;
 					RayTracingInstance.LayerIndex = (uint8)(Instance.MaskAndFlags.bAnySegmentsDecal ? ERayTracingSceneLayer::Decals : ERayTracingSceneLayer::Base);
@@ -1248,6 +1239,8 @@ bool FDeferredShadingSceneRenderer::GatherRayTracingWorldInstancesForView(FRDGBu
 						}
 					}
 
+					const uint32 InstanceIndex = RayTracingScene.AddInstance(MoveTemp(RayTracingInstance), SceneProxy, true);
+
 					if (bParallelMeshBatchSetup)
 					{
 						if (NumPendingMeshBatches >= RayTracingParallelMeshBatchSize)
@@ -1332,7 +1325,6 @@ bool FDeferredShadingSceneRenderer::GatherRayTracingWorldInstancesForView(FRDGBu
 		const FScene& Scene;
 		TChunkedArray<FRayTracingRelevantPrimitive>& RelevantStaticPrimitives;
 		const FRayTracingCullingParameters& CullingParameters;
-		const bool bEnableInstanceDebugData;
 
 		// Outputs
 
@@ -1342,13 +1334,12 @@ bool FDeferredShadingSceneRenderer::GatherRayTracingWorldInstancesForView(FRDGBu
 
 		FRayTracingSceneAddInstancesTask(const FScene& InScene,
 											TChunkedArray<FRayTracingRelevantPrimitive>& InRelevantStaticPrimitives,
-											const FRayTracingCullingParameters& InCullingParameters, bool InEnableInstanceDebugData,
+											const FRayTracingCullingParameters& InCullingParameters,
 											FRayTracingScene& InRayTracingScene, TArray<FVisibleRayTracingMeshCommand>& InVisibleRayTracingMeshCommands,
 											TArray<FPrimitiveSceneProxy*>& InProxiesWithDirtyCachedInstance)
 			: Scene(InScene)
 			, RelevantStaticPrimitives(InRelevantStaticPrimitives)
 			, CullingParameters(InCullingParameters)
-			, bEnableInstanceDebugData(InEnableInstanceDebugData)
 			, RayTracingScene(InRayTracingScene)
 			, VisibleRayTracingMeshCommands(InVisibleRayTracingMeshCommands)
 			, ProxiesWithDirtyCachedInstance(InProxiesWithDirtyCachedInstance)
@@ -1356,6 +1347,8 @@ bool FDeferredShadingSceneRenderer::GatherRayTracingWorldInstancesForView(FRDGBu
 			VisibleRayTracingMeshCommands.Reserve(RelevantStaticPrimitives.Num());
 		}
 
+		// TODO: Consider moving auto instance batching logic into FRayTracingScene
+
 		struct FAutoInstanceBatch
 		{
 			int32 Index = INDEX_NONE;
@@ -1473,30 +1466,10 @@ bool FDeferredShadingSceneRenderer::GatherRayTracingWorldInstancesForView(FRDGBu
 					}
 
 					checkf(SceneInfo->CachedRayTracingInstance.GeometryRHI, TEXT("Ray tracing instance must have a valid geometry."));
-					const int32 NewInstanceIndex = RayTracingScene.Instances.Add(SceneInfo->CachedRayTracingInstance);
-					FRayTracingGeometryInstance& NewInstance = RayTracingScene.Instances[NewInstanceIndex];
+					FRayTracingGeometryInstance NewInstance = SceneInfo->CachedRayTracingInstance;
 
 					NewInstance.LayerIndex = (uint8)(RelevantPrimitive.bAnySegmentsDecal ? ERayTracingSceneLayer::Decals : ERayTracingSceneLayer::Base);
 
-					// At the moment we only support SM & ISMs on this path
-					check(EnumHasAnyFlags(Flags, ERayTracingPrimitiveFlags::CacheMeshCommands));
-					if (SceneInfo->CachedRayTracingMeshCommandIndicesPerLOD.Num() > 0 && SceneInfo->CachedRayTracingMeshCommandIndicesPerLOD[LODIndex].Num() > 0)
-					{
-						for (int32 CommandIndex : SceneInfo->CachedRayTracingMeshCommandIndicesPerLOD[LODIndex])
-						{
-							FVisibleRayTracingMeshCommand NewVisibleMeshCommand;
-
-							NewVisibleMeshCommand.RayTracingMeshCommand = &Scene.CachedRayTracingMeshCommands[CommandIndex];
-							NewVisibleMeshCommand.InstanceIndex = NewInstanceIndex;
-							VisibleRayTracingMeshCommands.Add(NewVisibleMeshCommand);
-						}
-					}
-
-					if (bEnableInstanceDebugData)
-					{
-						RayTracingScene.AddInstanceDebugData(SceneInfo->CachedRayTracingInstance.GeometryRHI, SceneInfo->Proxy, false);
-					}					
-
 					const Experimental::FHashElementId GroupId = Scene.PrimitiveRayTracingGroupIds[PrimitiveIndex];
 					const bool bUseGroupBounds = CullingParameters.bCullUsingGroupIds && GroupId.IsValid();
 
@@ -1531,6 +1504,21 @@ bool FDeferredShadingSceneRenderer::GatherRayTracingWorldInstancesForView(FRDGBu
 					}
 
 					AddDebugRayTracingInstanceFlags(NewInstance.Flags);
+
+					const int32 NewInstanceIndex = RayTracingScene.AddInstance(MoveTemp(NewInstance), SceneInfo->Proxy, false);
+
+					// At the moment we only support SM & ISMs on this path
+					check(EnumHasAnyFlags(Flags, ERayTracingPrimitiveFlags::CacheMeshCommands));
+					if (SceneInfo->CachedRayTracingMeshCommandIndicesPerLOD.Num() > 0 && SceneInfo->CachedRayTracingMeshCommandIndicesPerLOD[LODIndex].Num() > 0)
+					{
+						for (int32 CommandIndex : SceneInfo->CachedRayTracingMeshCommandIndicesPerLOD[LODIndex])
+						{
+							FVisibleRayTracingMeshCommand NewVisibleMeshCommand;
+							NewVisibleMeshCommand.RayTracingMeshCommand = &Scene.CachedRayTracingMeshCommands[CommandIndex];
+							NewVisibleMeshCommand.InstanceIndex = NewInstanceIndex;
+							VisibleRayTracingMeshCommands.Add(NewVisibleMeshCommand);
+						}
+					}
 				}
 				else
 				{
@@ -1555,7 +1543,7 @@ bool FDeferredShadingSceneRenderer::GatherRayTracingWorldInstancesForView(FRDGBu
 					}
 
 					// location if this is a new entry
-					const int32 NewInstanceIndex = RayTracingScene.Instances.Num();
+					const int32 NewInstanceIndex = RayTracingScene.GetInstances().Num();
 					const uint64 InstanceKey = RelevantPrimitive.InstancingKey();
 
 					FAutoInstanceBatch DummyInstanceBatch = { NewInstanceIndex };
@@ -1565,7 +1553,7 @@ bool FDeferredShadingSceneRenderer::GatherRayTracingWorldInstancesForView(FRDGBu
 					{
 						// Reusing a previous entry, just append to the instance list.
 
-						FRayTracingGeometryInstance& RayTracingInstance = RayTracingScene.Instances[InstanceBatch.Index];
+						FRayTracingGeometryInstance& RayTracingInstance = RayTracingScene.GetInstance(InstanceBatch.Index);
 						bool bReallocated = InstanceBatch.Add(RayTracingScene, SceneInfo->GetInstanceSceneDataOffset(), (uint32)PrimitiveIndex);
 
 						++RayTracingInstance.NumTransforms;
@@ -1598,16 +1586,11 @@ bool FDeferredShadingSceneRenderer::GatherRayTracingWorldInstancesForView(FRDGBu
 							}
 						}
 
-						FRayTracingGeometryInstance& RayTracingInstance = RayTracingScene.Instances.AddDefaulted_GetRef();
+						InstanceBatch.Add(RayTracingScene, SceneInfo->GetInstanceSceneDataOffset(), (uint32)PrimitiveIndex);
+
+						FRayTracingGeometryInstance RayTracingInstance;
 						RayTracingInstance.GeometryRHI = RelevantPrimitive.RayTracingGeometryRHI;
 						checkf(RayTracingInstance.GeometryRHI, TEXT("Ray tracing instance must have a valid geometry."));
-
-						if (bEnableInstanceDebugData)
-						{
-							RayTracingScene.AddInstanceDebugData(RayTracingInstance.GeometryRHI, SceneInfo->Proxy, false);
-						}						
-
-						InstanceBatch.Add(RayTracingScene, SceneInfo->GetInstanceSceneDataOffset(), (uint32)PrimitiveIndex);
 						RayTracingInstance.InstanceSceneDataOffsets = InstanceBatch.InstanceSceneDataOffsets;
 						RayTracingInstance.UserData = InstanceBatch.UserData;
 						RayTracingInstance.NumTransforms = 1;
@@ -1625,6 +1608,10 @@ bool FDeferredShadingSceneRenderer::GatherRayTracingWorldInstancesForView(FRDGBu
 						AddDebugRayTracingInstanceFlags(RayTracingInstance.Flags);
 
 						RayTracingInstance.LayerIndex = (uint8)(RelevantPrimitive.bAnySegmentsDecal ? ERayTracingSceneLayer::Decals : ERayTracingSceneLayer::Base);
+
+						const uint32 ActualInstanceIndex = RayTracingScene.AddInstance(MoveTemp(RayTracingInstance), SceneInfo->Proxy, false);
+
+						check(NewInstanceIndex == ActualInstanceIndex);
 					}
 				}
 			}
@@ -1648,7 +1635,7 @@ bool FDeferredShadingSceneRenderer::GatherRayTracingWorldInstancesForView(FRDGBu
 	AddInstancesTaskPrerequisites.Add(RelevantPrimitiveList.StaticPrimitiveLODTask);
 
 	FGraphEventRef AddInstancesTask = TGraphTask<FRayTracingSceneAddInstancesTask>::CreateTask(&AddInstancesTaskPrerequisites).ConstructAndDispatchWhenReady(
-		*Scene, RelevantPrimitiveList.StaticPrimitives, View.RayTracingCullingParameters, bEnableInstanceDebugData, // inputs 
+		*Scene, RelevantPrimitiveList.StaticPrimitives, View.RayTracingCullingParameters, // inputs 
 		RayTracingScene, View.VisibleRayTracingMeshCommands, View.ProxiesWithDirtyCachedInstance // outputs
 	);
 
@@ -2079,7 +2066,7 @@ static void ReleaseRaytracingResources(FRDGBuilder& GraphBuilder, TArrayView<FVi
 
 			// Track if we ended up rendering anything this frame.  After rendering all view families, we'll release the
 			// ray tracing scene resources if nothing used ray tracing.
-			if (RayTracingScene.Instances.Num() > 0)
+			if (RayTracingScene.GetInstances().Num() > 0)
 			{
 				RayTracingScene.bUsedThisFrame = true;
 			}
@@ -2613,9 +2600,12 @@ void FDeferredShadingSceneRenderer::Render(FRDGBuilder& GraphBuilder)
 
 #if RHI_RAYTRACING
 
+	const int32 ReferenceViewIndex = 0;
+	FViewInfo& ReferenceView = Views[ReferenceViewIndex];
+
 	// Prepare the scene for rendering this frame.
 	FRayTracingScene& RayTracingScene = Scene->RayTracingScene;
-	RayTracingScene.Reset(); // Resets the internal arrays, but does not release any resources.
+	RayTracingScene.Reset(IsRayTracingInstanceDebugDataEnabled(ReferenceView)); // Resets the internal arrays, but does not release any resources.
 
 	if (ShouldPrepareRayTracingDecals(*Scene, ViewFamily))
 	{
@@ -2642,9 +2632,6 @@ void FDeferredShadingSceneRenderer::Render(FRDGBuilder& GraphBuilder)
 		}
 	}
 
-	const int32 ReferenceViewIndex = 0;
-	FViewInfo& ReferenceView = Views[ReferenceViewIndex];
-
 	if (IsRayTracingEnabled() && RHISupportsRayTracingShaders(ViewFamily.GetShaderPlatform()))
 	{
 		// Nanite raytracing manager update must run before GPUScene update since it can modify primitive data
diff --git a/Engine/Source/Runtime/Renderer/Private/HeterogeneousVolumes/HeterogeneousVolumesHardwareRayTracing.cpp b/Engine/Source/Runtime/Renderer/Private/HeterogeneousVolumes/HeterogeneousVolumesHardwareRayTracing.cpp
--- a/Engine/Source/Runtime/Renderer/Private/HeterogeneousVolumes/HeterogeneousVolumesHardwareRayTracing.cpp
+++ b/Engine/Source/Runtime/Renderer/Private/HeterogeneousVolumes/HeterogeneousVolumesHardwareRayTracing.cpp
@@ -197,19 +197,18 @@ void GenerateRayTracingScene(
 	FRayTracingScene& RayTracingScene
 )
 {
-	RayTracingScene.Reset();
+	RayTracingScene.Reset(false);
 
 	// Collect instances
-	TArray<FRayTracingGeometryInstance> RayTracingInstances;
 	for (int32 GeometryIndex = 0; GeometryIndex < RayTracingGeometries.Num(); ++GeometryIndex)
 	{
-		checkf(RayTracingGeometries[GeometryIndex], TEXT("RayTracingGeometryInstance not created."))
-			FRayTracingGeometryInstance RayTracingGeometryInstance = {};
+		checkf(RayTracingGeometries[GeometryIndex], TEXT("RayTracingGeometryInstance not created."));
+		FRayTracingGeometryInstance RayTracingGeometryInstance = {};
 		RayTracingGeometryInstance.GeometryRHI = RayTracingGeometries[GeometryIndex];
 		RayTracingGeometryInstance.NumTransforms = 1;
 		RayTracingGeometryInstance.Transforms = MakeArrayView(&RayTracingTransforms[GeometryIndex], 1);
 
-		RayTracingInstances.Add(RayTracingGeometryInstance);
+		RayTracingScene.AddInstance(RayTracingGeometryInstance);
 	}
 
 	// Build instance BLAS
@@ -238,7 +237,6 @@ void GenerateRayTracingScene(
 	// Create RayTracingScene
 	const FGPUScene* EmptyGPUScene = nullptr;
 	FViewMatrices EmptyViewMatrices;
-	RayTracingScene.Instances = RayTracingInstances;
 	RayTracingScene.Create(GraphBuilder, EmptyGPUScene, EmptyViewMatrices);
 
 	// Build TLAS
diff --git a/Engine/Source/Runtime/Renderer/Private/RayTracing/RayTracingDebug.cpp b/Engine/Source/Runtime/Renderer/Private/RayTracing/RayTracingDebug.cpp
--- a/Engine/Source/Runtime/Renderer/Private/RayTracing/RayTracingDebug.cpp
+++ b/Engine/Source/Runtime/Renderer/Private/RayTracing/RayTracingDebug.cpp
@@ -1045,7 +1045,7 @@ void FDeferredShadingSceneRenderer::RayTracingDisplayPicking(const FRayTracingPi
 	Writer.EmptyLine();
 
 	FRHIRayTracingGeometry* Geometry = nullptr;
-	for (const FRayTracingGeometryInstance& Instance : Scene->RayTracingScene.Instances)
+	for (const FRayTracingGeometryInstance& Instance : Scene->RayTracingScene.GetInstances())
 	{
 		if (Instance.GeometryRHI)
 		{
diff --git a/Engine/Source/Runtime/Renderer/Private/RayTracing/RayTracingScene.cpp b/Engine/Source/Runtime/Renderer/Private/RayTracing/RayTracingScene.cpp
--- a/Engine/Source/Runtime/Renderer/Private/RayTracing/RayTracingScene.cpp
+++ b/Engine/Source/Runtime/Renderer/Private/RayTracing/RayTracingScene.cpp
@@ -10,6 +10,7 @@
 #include "RenderGraphBuilder.h"
 #include "RenderGraphUtils.h"
 #include "RaytracingOptions.h"
+#include "PrimitiveSceneProxy.h"
 
 BEGIN_SHADER_PARAMETER_STRUCT(FBuildInstanceBufferPassParams, )
 	SHADER_PARAMETER_RDG_BUFFER_UAV(RWStructuredBuffer, InstanceBuffer)
@@ -314,19 +315,28 @@ FRHIShaderResourceView* FRayTracingScene::GetLayerSRVChecked(ERayTracingSceneLay
 	return LayerSRVs[uint8(Layer)].GetReference();
 }
 
-void FRayTracingScene::AddInstanceDebugData(const FRHIRayTracingGeometry* GeometryRHI, const FPrimitiveSceneProxy* Proxy, bool bDynamic)
+uint32 FRayTracingScene::AddInstance(FRayTracingGeometryInstance Instance, const FPrimitiveSceneProxy* Proxy, bool bDynamic)
 {
-	FRayTracingInstanceDebugData& InstanceDebugData = InstancesDebugData.AddDefaulted_GetRef();
-	InstanceDebugData.Flags = bDynamic ? 1 : 0;
-	InstanceDebugData.GeometryAddress = uint64(GeometryRHI);
+	FRHIRayTracingGeometry* GeometryRHI = Instance.GeometryRHI;
 
-	if (Proxy)
+	const uint32 InstanceIndex = Instances.Add(MoveTemp(Instance));
+
+	if (bInstanceDebugDataEnabled)
 	{
-		InstanceDebugData.ProxyHash = Proxy->GetTypeHash();
+		FRayTracingInstanceDebugData& InstanceDebugData = InstancesDebugData.AddDefaulted_GetRef();
+		InstanceDebugData.Flags = bDynamic ? 1 : 0;
+		InstanceDebugData.GeometryAddress = uint64(GeometryRHI);
+
+		if (Proxy)
+		{
+			InstanceDebugData.ProxyHash = Proxy->GetTypeHash();
+		}
 	}
+
+	return InstanceIndex;
 }
 
-void FRayTracingScene::Reset()
+void FRayTracingScene::Reset(bool bInInstanceDebugDataEnabled)
 {
 	WaitForTasks();
 
@@ -344,11 +354,13 @@ void FRayTracingScene::Reset()
 	BuildScratchBuffer = nullptr;
 	InstanceDebugBuffer = nullptr;
 	DebugInstanceGPUSceneIndexBuffer = nullptr;
+
+	bInstanceDebugDataEnabled = bInInstanceDebugDataEnabled;
 }
 
 void FRayTracingScene::ResetAndReleaseResources()
 {
-	Reset();
+	Reset(false);
 
 	Instances.Empty();
 	InstancesDebugData.Empty();
diff --git a/Engine/Source/Runtime/Renderer/Private/RayTracing/RayTracingScene.h b/Engine/Source/Runtime/Renderer/Private/RayTracing/RayTracingScene.h
--- a/Engine/Source/Runtime/Renderer/Private/RayTracing/RayTracingScene.h
+++ b/Engine/Source/Runtime/Renderer/Private/RayTracing/RayTracingScene.h
@@ -12,7 +12,6 @@
 #include "RenderGraphResources.h"
 #include "Misc/MemStack.h"
 #include "Containers/ArrayView.h"
-#include "PrimitiveSceneProxy.h"
 #include "MeshPassProcessor.h"
 #include "RayTracingInstanceBufferUtil.h"
 #include "RayTracingDebugTypes.h"
@@ -22,6 +21,12 @@ class FRHIRayTracingScene;
 class FRHIShaderResourceView;
 class FRayTracingGeometry;
 class FRDGBuilder;
+class FPrimitiveSceneProxy;
+
+namespace Nanite
+{
+	using CoarseMeshStreamingHandle = int16;
+}
 
 enum class ERayTracingSceneLayer : uint8
 {
@@ -48,6 +53,8 @@ public:
 	FRayTracingScene();
 	~FRayTracingScene();
 
+	uint32 AddInstance(FRayTracingGeometryInstance Instance, const FPrimitiveSceneProxy* Proxy = nullptr, bool bDynamic = false);
+
 	// Allocates RayTracingSceneRHI and builds various metadata required to create the final scene.
 	FRayTracingSceneWithGeometryInstances BuildInitializationData() const;
 
@@ -60,7 +67,7 @@ public:
 	void Create(FRDGBuilder& GraphBuilder, const FGPUScene* GPUScene, const FViewMatrices& ViewMatrices);
 
 	// Resets the instance list and reserves memory for this frame.
-	void Reset();
+	void Reset(bool bInstanceDebugDataEnabled);
 
 	// Similar to Reset(), but also releases any persistent CPU and GPU memory allocations.
 	void ResetAndReleaseResources();
@@ -90,19 +97,17 @@ public:
 	RENDERER_API FRHIBuffer* GetBufferChecked() const;
 	
 	RENDERER_API FRHIShaderResourceView* GetLayerSRVChecked(ERayTracingSceneLayer Layer) const;
-	
-	void AddInstanceDebugData(const FRHIRayTracingGeometry* GeometryRHI, const FPrimitiveSceneProxy* Proxy, bool bDynamic);
+
+	TArrayView<const FRayTracingGeometryInstance> GetInstances() const { return MakeArrayView(Instances); }
+
+	FRayTracingGeometryInstance& GetInstance(uint32 InstanceIndex) { return Instances[InstanceIndex]; }
 
 	void InitPreViewTranslation(const FViewMatrices& ViewMatrices);
+
 public:
 
 	// Public members for initial refactoring step (previously were public members of FViewInfo).
 
-	// Persistent storage for ray tracing instance descriptors.
-	// Cleared every frame without releasing memory to avoid large heap allocations.
-	// This must be filled before calling CreateRayTracingSceneWithGeometryInstances() and Create().
-	TArray<FRayTracingGeometryInstance> Instances;
-
 	bool bUsedThisFrame = false;
 	uint32 NumMissShaderSlots = 1; // we must have a default miss shader, so always include it from the start
 	uint32 NumCallableShaderSlots = 0;
@@ -122,7 +127,6 @@ public:
 
 	// Special data for debugging purposes
 	FRDGBufferRef InstanceDebugBuffer = nullptr;
-	TArray<FRayTracingInstanceDebugData> InstancesDebugData;
 
 	FRDGBufferRef DebugInstanceGPUSceneIndexBuffer = nullptr;
 	bool bNeedsDebugInstanceGPUSceneIndexBuffer = false;
@@ -155,7 +159,16 @@ private:
 
 	mutable FGraphEventRef FillInstanceUploadBufferTask;
 
+	// Persistent storage for ray tracing instance descriptors.
+	// Cleared every frame without releasing memory to avoid large heap allocations.
+	// This must be filled before calling CreateRayTracingSceneWithGeometryInstances() and Create().
+	TArray<FRayTracingGeometryInstance> Instances;
+
+	TArray<FRayTracingInstanceDebugData> InstancesDebugData;
+
 	ERayTracingSceneState State = ERayTracingSceneState::Writable;
+
+	bool bInstanceDebugDataEnabled = false;
 };
 
 #endif // RHI_RAYTRACING
