Rename bCachedRayTracingInstanceAnySegmentsDecal and bCachedRayTracingInstanceAllSegmentsDecal

diff --git a/Engine/Source/Runtime/Renderer/Private/DeferredShadingRenderer.cpp b/Engine/Source/Runtime/Renderer/Private/DeferredShadingRenderer.cpp
--- a/Engine/Source/Runtime/Renderer/Private/DeferredShadingRenderer.cpp
+++ b/Engine/Source/Runtime/Renderer/Private/DeferredShadingRenderer.cpp
@@ -1454,13 +1454,13 @@ bool FDeferredShadingSceneRenderer::GatherRayTracingWorldInstancesForView(FRDGBu
 					
 					// TODO: Consider requesting a recache of all ray tracing commands during which decals are excluded
 
-					if (SceneInfo->bCachedRayTracingInstancebAnySegmentsDecal && !SceneInfo->bCachedRayTracingInstancebAllSegmentsDecal)
+					if (SceneInfo->bCachedRayTracingInstanceAnySegmentsDecal && !SceneInfo->bCachedRayTracingInstanceAllSegmentsDecal)
 					{
 						// TODO: Warn user meshes with mixed decal segments are not currently supported in ray tracing
 						continue;
 					}
 
-					if (GRayTracingExcludeDecals && SceneInfo->bCachedRayTracingInstancebAnySegmentsDecal)
+					if (GRayTracingExcludeDecals && SceneInfo->bCachedRayTracingInstanceAnySegmentsDecal)
 					{
 						continue;
 					}
diff --git a/Engine/Source/Runtime/Renderer/Private/PrimitiveSceneInfo.cpp b/Engine/Source/Runtime/Renderer/Private/PrimitiveSceneInfo.cpp
--- a/Engine/Source/Runtime/Renderer/Private/PrimitiveSceneInfo.cpp
+++ b/Engine/Source/Runtime/Renderer/Private/PrimitiveSceneInfo.cpp
@@ -1054,8 +1054,8 @@ void FPrimitiveSceneInfo::UpdateCachedRayTracingInstance(FPrimitiveSceneInfo* Sc
 			SceneInfo->CachedRayTracingInstance.Flags |= ERayTracingInstanceFlags::TriangleCullDisable;
 		}
 
-		SceneInfo->bCachedRayTracingInstancebAnySegmentsDecal = CachedRayTracingInstance.MaskAndFlags.bAnySegmentsDecal;
-		SceneInfo->bCachedRayTracingInstancebAllSegmentsDecal = CachedRayTracingInstance.MaskAndFlags.bAllSegmentsDecal;
+		SceneInfo->bCachedRayTracingInstanceAnySegmentsDecal = CachedRayTracingInstance.MaskAndFlags.bAnySegmentsDecal;
+		SceneInfo->bCachedRayTracingInstanceAllSegmentsDecal = CachedRayTracingInstance.MaskAndFlags.bAllSegmentsDecal;
 	}
 }
 
diff --git a/Engine/Source/Runtime/Renderer/Public/PrimitiveSceneInfo.h b/Engine/Source/Runtime/Renderer/Public/PrimitiveSceneInfo.h
--- a/Engine/Source/Runtime/Renderer/Public/PrimitiveSceneInfo.h
+++ b/Engine/Source/Runtime/Renderer/Public/PrimitiveSceneInfo.h
@@ -415,8 +415,8 @@ public:
 	TArray<uint64> CachedRayTracingMeshCommandsHashPerLOD;
 	// TODO: this should be placed in FRayTracingScene and we have a pointer/handle here. It's here for now for PoC
 	FRayTracingGeometryInstance CachedRayTracingInstance;
-	bool bCachedRayTracingInstancebAnySegmentsDecal : 1;
-	bool bCachedRayTracingInstancebAllSegmentsDecal : 1;
+	bool bCachedRayTracingInstanceAnySegmentsDecal : 1;
+	bool bCachedRayTracingInstanceAllSegmentsDecal : 1;
 	TArray<FBoxSphereBounds> CachedRayTracingInstanceWorldBounds;
 	int32 SmallestRayTracingInstanceWorldBoundsIndex;
 #endif
