Fix incorrect flag being checked in ERayTracingPrimitiveFlags::CacheInstances code path.
- Flags in FRayTracingRelevantPrimitive are only valid for non-cached instances.

diff --git a/Engine/Source/Runtime/Renderer/Private/DeferredShadingRenderer.cpp b/Engine/Source/Runtime/Renderer/Private/DeferredShadingRenderer.cpp
--- a/Engine/Source/Runtime/Renderer/Private/DeferredShadingRenderer.cpp
+++ b/Engine/Source/Runtime/Renderer/Private/DeferredShadingRenderer.cpp
@@ -1468,7 +1468,7 @@ bool FDeferredShadingSceneRenderer::GatherRayTracingWorldInstancesForView(FRDGBu
 					checkf(SceneInfo->CachedRayTracingInstance.GeometryRHI, TEXT("Ray tracing instance must have a valid geometry."));
 					FRayTracingGeometryInstance NewInstance = SceneInfo->CachedRayTracingInstance;
 
-					NewInstance.LayerIndex = (uint8)(RelevantPrimitive.bAnySegmentsDecal ? ERayTracingSceneLayer::Decals : ERayTracingSceneLayer::Base);
+					NewInstance.LayerIndex = (uint8)(SceneInfo->bCachedRayTracingInstanceAnySegmentsDecal ? ERayTracingSceneLayer::Decals : ERayTracingSceneLayer::Base);
 
 					const Experimental::FHashElementId GroupId = Scene.PrimitiveRayTracingGroupIds[PrimitiveIndex];
 					const bool bUseGroupBounds = CullingParameters.bCullUsingGroupIds && GroupId.IsValid();
