Adding missing include.

diff --git a/Engine/Source/Runtime/Renderer/Private/GPUScene.cpp b/Engine/Source/Runtime/Renderer/Private/GPUScene.cpp
--- a/Engine/Source/Runtime/Renderer/Private/GPUScene.cpp
+++ b/Engine/Source/Runtime/Renderer/Private/GPUScene.cpp
@@ -25,6 +25,7 @@
 #include "ShaderPrint.h"
 #include "RenderCore.h"
 #include "LightSceneData.h"
+#include "LightSceneProxy.h"
 
 #define LOG_INSTANCE_ALLOCATIONS 0
 
