Remove unnecessary IsGPUSkinCacheAvailable() call since it is always true if Scene has a SkinCache object.
Remove function deprecated two versions ago.

diff --git a/Engine/Source/Runtime/Engine/Private/GPUSkinCache.cpp b/Engine/Source/Runtime/Engine/Private/GPUSkinCache.cpp
--- a/Engine/Source/Runtime/Engine/Private/GPUSkinCache.cpp
+++ b/Engine/Source/Runtime/Engine/Private/GPUSkinCache.cpp
@@ -263,14 +263,6 @@ ENGINE_API bool GPUSkinCacheNeedsDuplicatedVertices()
 #endif
 }
 
-UE_DEPRECATED(5.0, "This function is no longer in use and will be removed.")
-ENGINE_API bool DoRecomputeSkinTangentsOnGPU_RT()
-{
-	// currently only implemented and tested on Window SM5 (needs Compute, Atomics, SRV for index buffers, UAV for VertexBuffers)
-	//#todo-gpuskin: Enable on PS4 when SRVs for IB exist
-	return IsGPUSkinCacheAllowed(GMaxRHIShaderPlatform) && GEnableGPUSkinCacheShaders != 0 && (GEnableGPUSkinCache && GSkinCacheRecomputeTangents != 0);
-}
-
 // determine if during DispatchUpdateSkinning caching should occur
 enum class EGPUSkinCacheDispatchFlags
 {
diff --git a/Engine/Source/Runtime/Renderer/Private/SceneRendering.cpp b/Engine/Source/Runtime/Renderer/Private/SceneRendering.cpp
--- a/Engine/Source/Runtime/Renderer/Private/SceneRendering.cpp
+++ b/Engine/Source/Runtime/Renderer/Private/SceneRendering.cpp
@@ -3279,12 +3279,9 @@ void FSceneRenderer::RenderFinish(FRDGBuilder& GraphBuilder, FRDGTextureRef View
 		static auto* CVarSkinCacheOOM = IConsoleManager::Get().FindTConsoleVariableDataFloat(TEXT("r.SkinCache.SceneMemoryLimitInMB"));
 
 		uint64 GPUSkinCacheExtraRequiredMemory = 0;
-		if (IsGPUSkinCacheAvailable(ShaderPlatform))
+		if (FGPUSkinCache* SkinCache = Scene->GetGPUSkinCache())
 		{
-			if (FGPUSkinCache* SkinCache = Scene->GetGPUSkinCache())
-			{
-				GPUSkinCacheExtraRequiredMemory = SkinCache->GetExtraRequiredMemoryAndReset();
-			}
+			GPUSkinCacheExtraRequiredMemory = SkinCache->GetExtraRequiredMemoryAndReset();
 		}
 		const bool bShowSkinCacheOOM = CVarSkinCacheOOM != nullptr && GPUSkinCacheExtraRequiredMemory > 0;
 
