Fix a bug of missing MobileUniformBuffer in mobile modulated shadow rendering.

diff --git a/Engine/Shaders/Private/ShadowProjectionPixelShader.usf b/Engine/Shaders/Private/ShadowProjectionPixelShader.usf
--- a/Engine/Shaders/Private/ShadowProjectionPixelShader.usf
+++ b/Engine/Shaders/Private/ShadowProjectionPixelShader.usf
@@ -27,6 +27,7 @@
 #include "HairStrands/HairStrandsVisibilityCommonStruct.ush"
 #include "Common.ush"
 
+// Needs to bind MobileBasePass UB since the modulated shadow is in mobile base pass. 
 #if MODULATED_SHADOWS
 	#define MobileSceneTextures MobileBasePass.SceneTextures
 #endif
diff --git a/Engine/Source/Runtime/Renderer/Private/DecalRenderingShared.cpp b/Engine/Source/Runtime/Renderer/Private/DecalRenderingShared.cpp
--- a/Engine/Source/Runtime/Renderer/Private/DecalRenderingShared.cpp
+++ b/Engine/Source/Runtime/Renderer/Private/DecalRenderingShared.cpp
@@ -14,6 +14,7 @@
 #include "DebugViewModeRendering.h"
 #include "ScenePrivate.h"
 #include "PipelineStateCache.h"
+#include "MobileBasePassRendering.h"
 
 static TAutoConsoleVariable<float> CVarDecalFadeScreenSizeMultiplier(
 	TEXT("r.Decal.FadeScreenSizeMult"),
@@ -80,6 +81,7 @@ public:
 		DecalToWorldInvScale.Bind(Initializer.ParameterMap, TEXT("DecalToWorldInvScale"));
 		DecalOrientation.Bind(Initializer.ParameterMap,TEXT("DecalOrientation"));
 		DecalParams.Bind(Initializer.ParameterMap, TEXT("DecalParams"));
+		MobileBasePassUniformBuffer.Bind(Initializer.ParameterMap, FMobileBasePassUniformParameters::StaticStructMetadata.GetShaderVariableName());
 	}
 
 	void SetParameters(FRHICommandList& RHICmdList, const FViewInfo& View, const FDeferredDecalProxy& DecalProxy, const FMaterialRenderProxy* MaterialProxy, const FMaterial* MaterialResource, const float FadeAlphaValue = 1.0f)
@@ -154,6 +156,7 @@ private:
 	LAYOUT_FIELD(FShaderParameter, DecalToWorldInvScale);
 	LAYOUT_FIELD(FShaderParameter, DecalOrientation);
 	LAYOUT_FIELD(FShaderParameter, DecalParams);
+	LAYOUT_FIELD(FShaderUniformBufferParameter, MobileBasePassUniformBuffer);
 };
 
 IMPLEMENT_MATERIAL_SHADER_TYPE(,FDeferredDecalPS,TEXT("/Engine/Private/DeferredDecal.usf"),TEXT("MainPS"),SF_Pixel);
diff --git a/Engine/Source/Runtime/Renderer/Private/ShadowRendering.cpp b/Engine/Source/Runtime/Renderer/Private/ShadowRendering.cpp
--- a/Engine/Source/Runtime/Renderer/Private/ShadowRendering.cpp
+++ b/Engine/Source/Runtime/Renderer/Private/ShadowRendering.cpp
@@ -17,6 +17,7 @@
 #include "Shadows/ShadowSceneRenderer.h"
 #include "RenderCore.h"
 #include "TranslucentLighting.h"
+#include "MobileBasePassRendering.h"
 
 ///////////////////////////////////////////////////////////////////////////////////////////////////
 // Directional light
@@ -337,6 +338,14 @@ IMPLEMENT_SHADER_TYPE(,FShadowProjectionNoTransformVS,TEXT("/Engine/Private/Shad
 
 IMPLEMENT_SHADER_TYPE(,FShadowVolumeBoundProjectionVS,TEXT("/Engine/Private/ShadowProjectionVertexShader.usf"),TEXT("ShadowVolumeBoundProjectionVS"),SF_Vertex);
 
+template<uint32 T>
+TModulatedShadowProjection<T>::TModulatedShadowProjection(const ShaderMetaType::CompiledShaderInitializerType& Initializer) :
+TShadowProjectionPS<T, false, true>(Initializer)
+{
+	ModulatedShadowColorParameter.Bind(Initializer.ParameterMap, TEXT("ModulatedShadowColor"));
+	MobileBasePassUniformBuffer.Bind(Initializer.ParameterMap, FMobileBasePassUniformParameters::StaticStructMetadata.GetShaderVariableName());
+} 
+
 /**
  * Implementations for TShadowProjectionPS.  
  */
@@ -381,6 +390,7 @@ IMPLEMENT_SHADOW_PROJECTION_PIXEL_SHADER(5, false, false, true);
 #undef IMPLEMENT_SHADOW_PROJECTION_PIXEL_SHADER
 
 #define IMPLEMENT_MODULATED_SHADOW_PROJECTION_PIXEL_SHADER(Quality) \
+	template class TModulatedShadowProjection<Quality>; \
 	using FShadowModulatedProjectionPS##Quality = TShadowProjectionPS<Quality, false, true>; \
 	IMPLEMENT_TEMPLATE_TYPE_LAYOUT(template<>, FShadowModulatedProjectionPS##Quality); \
 	IMPLEMENT_SHADER_TYPE(template<>, TModulatedShadowProjection<Quality>, TEXT("/Engine/Private/ShadowProjectionPixelShader.usf"), TEXT("Main"), SF_Pixel);
diff --git a/Engine/Source/Runtime/Renderer/Private/ShadowRendering.h b/Engine/Source/Runtime/Renderer/Private/ShadowRendering.h
--- a/Engine/Source/Runtime/Renderer/Private/ShadowRendering.h
+++ b/Engine/Source/Runtime/Renderer/Private/ShadowRendering.h
@@ -1289,11 +1289,7 @@ public:
 
 	TModulatedShadowProjection() {}
 
-	TModulatedShadowProjection(const ShaderMetaType::CompiledShaderInitializerType& Initializer) :
-		TShadowProjectionPS<Quality, false, true>(Initializer)
-	{
-		ModulatedShadowColorParameter.Bind(Initializer.ParameterMap, TEXT("ModulatedShadowColor"));
-	}
+	TModulatedShadowProjection(const ShaderMetaType::CompiledShaderInitializerType& Initializer);
 
 	void SetParameters(
 		FRHICommandList& RHICmdList,
@@ -1308,6 +1304,7 @@ public:
 
 protected:
 	LAYOUT_FIELD(FShaderParameter, ModulatedShadowColorParameter);
+	LAYOUT_FIELD(FShaderUniformBufferParameter, MobileBasePassUniformBuffer);
 };
 
 /** Translucency shadow projection uniform buffer containing data needed for Fourier opacity maps. */
