Support primitives with mixed mesh decal segments in raytracing scene.
- Primitives with mixed decal and non-decal segments generate two raytracing scene instances one in each layer.
- Segments are masked accordingly in each layer using hidden hit groups (AHS always ignores hit).

diff --git a/Engine/Plugins/Experimental/GPULightmass/Source/GPULightmass/Private/Scene/Scene.h b/Engine/Plugins/Experimental/GPULightmass/Source/GPULightmass/Private/Scene/Scene.h
--- a/Engine/Plugins/Experimental/GPULightmass/Source/GPULightmass/Private/Scene/Scene.h
+++ b/Engine/Plugins/Experimental/GPULightmass/Source/GPULightmass/Private/Scene/Scene.h
@@ -42,9 +42,7 @@ public:
 
 	virtual void FinalizeCommand(FRayTracingMeshCommand& RayTracingMeshCommand) override final 
 	{
-		FVisibleRayTracingMeshCommand NewVisibleMeshCommand;
-		NewVisibleMeshCommand.RayTracingMeshCommand = &RayTracingMeshCommand;
-		NewVisibleMeshCommand.InstanceIndex = RayTracingInstanceIndex;
+		FVisibleRayTracingMeshCommand NewVisibleMeshCommand(&RayTracingMeshCommand, RayTracingInstanceIndex);
 		VisibleCommandStorage.Add(NewVisibleMeshCommand);
 		check(NewVisibleMeshCommand.RayTracingMeshCommand);
 	}
diff --git a/Engine/Source/Runtime/Renderer/Private/DeferredShadingRenderer.cpp b/Engine/Source/Runtime/Renderer/Private/DeferredShadingRenderer.cpp
--- a/Engine/Source/Runtime/Renderer/Private/DeferredShadingRenderer.cpp
+++ b/Engine/Source/Runtime/Renderer/Private/DeferredShadingRenderer.cpp
@@ -1044,7 +1044,8 @@ bool FDeferredShadingSceneRenderer::GatherRayTracingWorldInstancesForView(FRDGBu
 			const FPrimitiveSceneProxy* SceneProxy = nullptr;
 			TArray<FMeshBatch> MeshBatchesOwned;
 			TArrayView<const FMeshBatch> MeshBatchesView;
-			uint32 InstanceIndex = 0;
+			uint32 InstanceIndex;
+			uint32 DecalInstanceIndex;
 
 			TArrayView<const FMeshBatch> GetMeshBatches() const
 			{
@@ -1103,7 +1104,7 @@ bool FDeferredShadingSceneRenderer::GatherRayTracingWorldInstancesForView(FRDGBu
 								const FMeshBatch& MeshBatch = MeshBatches[SegmentIndex];
 								FDynamicRayTracingMeshCommandContext CommandContext(
 									*TaskDynamicCommandStorage, *TaskVisibleCommands,
-									SegmentIndex, WorkItem.InstanceIndex);
+									SegmentIndex, WorkItem.InstanceIndex, WorkItem.DecalInstanceIndex);
 								FMeshPassProcessorRenderState PassDrawRenderState;
 								FRayTracingMeshProcessor RayTracingMeshProcessor(&CommandContext, Scene, &View, PassDrawRenderState, Scene->CachedRayTracingMeshCommandsMode);
 								RayTracingMeshProcessor.AddMeshBatch(MeshBatch, 1, WorkItem.SceneProxy);
@@ -1183,13 +1184,13 @@ bool FDeferredShadingSceneRenderer::GatherRayTracingWorldInstancesForView(FRDGBu
 					// Autobuild of InstanceMaskAndFlags if the mask and flags are not built
 					UpdateRayTracingInstanceMaskAndFlagsIfNeeded(Instance, *SceneProxy, &ViewFamily);
 
-					if (Instance.MaskAndFlags.bAnySegmentsDecal && !Instance.MaskAndFlags.bAllSegmentsDecal)
-					{
-						// TODO: Warn user meshes with mixed decal segments are not currently supported in ray tracing
-						continue;
-					}
+					// if primitive has mixed decal and non-decal segments we need to have two ray tracing instances
+					// one containing non-decal segments and the other with decal segments
+					// masking of segments is done using "hidden" hitgroups
+					// TODO: Debug Visualization to highlight primitives using this?
+					const bool bNeedSeparateDecalInstance = Instance.MaskAndFlags.bAnySegmentsDecal && !Instance.MaskAndFlags.bAllSegmentsDecal;
 
-					if (GRayTracingExcludeDecals && Instance.MaskAndFlags.bAnySegmentsDecal)
+					if (GRayTracingExcludeDecals && Instance.MaskAndFlags.bAnySegmentsDecal && !bNeedSeparateDecalInstance)
 					{
 						continue;
 					}
@@ -1199,7 +1200,7 @@ bool FDeferredShadingSceneRenderer::GatherRayTracingWorldInstancesForView(FRDGBu
 					checkf(RayTracingInstance.GeometryRHI, TEXT("Ray tracing instance must have a valid geometry."));
 					RayTracingInstance.DefaultUserData = PrimitiveIndex;
 					RayTracingInstance.bApplyLocalBoundsTransform = Instance.bApplyLocalBoundsTransform;
-					RayTracingInstance.LayerIndex = (uint8)(Instance.MaskAndFlags.bAnySegmentsDecal ? ERayTracingSceneLayer::Decals : ERayTracingSceneLayer::Base);
+					RayTracingInstance.LayerIndex = (uint8)(Instance.MaskAndFlags.bAnySegmentsDecal && !bNeedSeparateDecalInstance ? ERayTracingSceneLayer::Decals : ERayTracingSceneLayer::Base);
 					RayTracingInstance.Mask = Instance.MaskAndFlags.Mask;
 
 					if (Instance.MaskAndFlags.bForceOpaque)
@@ -1239,7 +1240,16 @@ bool FDeferredShadingSceneRenderer::GatherRayTracingWorldInstancesForView(FRDGBu
 						}
 					}
 
-					const uint32 InstanceIndex = RayTracingScene.AddInstance(MoveTemp(RayTracingInstance), SceneProxy, true);
+					const uint32 InstanceIndex = RayTracingScene.AddInstance(RayTracingInstance, SceneProxy, true);
+
+					uint32 DecalInstanceIndex = INDEX_NONE;
+					if (bNeedSeparateDecalInstance && !GRayTracingExcludeDecals)
+					{
+						FRayTracingGeometryInstance DecalRayTracingInstance = RayTracingInstance;
+						DecalRayTracingInstance.LayerIndex = (uint8)ERayTracingSceneLayer::Decals;
+
+						DecalInstanceIndex = RayTracingScene.AddInstance(MoveTemp(DecalRayTracingInstance), SceneProxy, true);
+					}
 
 					if (bParallelMeshBatchSetup)
 					{
@@ -1278,6 +1288,7 @@ bool FDeferredShadingSceneRenderer::GatherRayTracingWorldInstancesForView(FRDGBu
 
 						WorkItem.SceneProxy = SceneProxy;
 						WorkItem.InstanceIndex = InstanceIndex;
+						WorkItem.DecalInstanceIndex = DecalInstanceIndex;
 					}
 					else
 					{
@@ -1285,7 +1296,7 @@ bool FDeferredShadingSceneRenderer::GatherRayTracingWorldInstancesForView(FRDGBu
 						for (int32 SegmentIndex = 0; SegmentIndex < InstanceMaterials.Num(); SegmentIndex++)
 						{
 							const FMeshBatch& MeshBatch = InstanceMaterials[SegmentIndex];
-							FDynamicRayTracingMeshCommandContext CommandContext(View.DynamicRayTracingMeshCommandStorage, View.VisibleRayTracingMeshCommands, SegmentIndex, InstanceIndex);
+							FDynamicRayTracingMeshCommandContext CommandContext(View.DynamicRayTracingMeshCommandStorage, View.VisibleRayTracingMeshCommands, SegmentIndex, InstanceIndex, DecalInstanceIndex);
 							FMeshPassProcessorRenderState PassDrawRenderState;
 							FRayTracingMeshProcessor RayTracingMeshProcessor(&CommandContext, Scene, &View, PassDrawRenderState, Scene->CachedRayTracingMeshCommandsMode);
 							RayTracingMeshProcessor.AddMeshBatch(MeshBatch, 1, SceneProxy);
@@ -1352,6 +1363,7 @@ bool FDeferredShadingSceneRenderer::GatherRayTracingWorldInstancesForView(FRDGBu
 		struct FAutoInstanceBatch
 		{
 			int32 Index = INDEX_NONE;
+			int32 DecalIndex = INDEX_NONE;
 
 			// Copies the next InstanceSceneDataOffset and user data into the current batch, returns true if arrays were re-allocated.
 			bool Add(FRayTracingScene& InRayTracingScene, uint32 InInstanceSceneDataOffset, uint32 InUserData)
@@ -1453,22 +1465,22 @@ bool FDeferredShadingSceneRenderer::GatherRayTracingWorldInstancesForView(FRDGBu
 					}
 					
 					// TODO: Consider requesting a recache of all ray tracing commands during which decals are excluded
+					
+					// if primitive has mixed decal and non-decal segments we need to have two ray tracing instances
+					// one containing non-decal segments and the other with decal segments
+					// masking of segments is done using "hidden" hitgroups
+					// TODO: Debug Visualization to highlight primitives using this?
+					const bool bNeedSeparateDecalInstance = SceneInfo->bCachedRayTracingInstanceAnySegmentsDecal && !SceneInfo->bCachedRayTracingInstanceAllSegmentsDecal;
 
-					if (SceneInfo->bCachedRayTracingInstanceAnySegmentsDecal && !SceneInfo->bCachedRayTracingInstanceAllSegmentsDecal)
-					{
-						// TODO: Warn user meshes with mixed decal segments are not currently supported in ray tracing
-						continue;
-					}
-
-					if (GRayTracingExcludeDecals && SceneInfo->bCachedRayTracingInstanceAnySegmentsDecal)
+					if (GRayTracingExcludeDecals && SceneInfo->bCachedRayTracingInstanceAnySegmentsDecal && !bNeedSeparateDecalInstance)
 					{
 						continue;
 					}
 
 					checkf(SceneInfo->CachedRayTracingInstance.GeometryRHI, TEXT("Ray tracing instance must have a valid geometry."));
-					FRayTracingGeometryInstance NewInstance = SceneInfo->CachedRayTracingInstance;
 
-					NewInstance.LayerIndex = (uint8)(SceneInfo->bCachedRayTracingInstanceAnySegmentsDecal ? ERayTracingSceneLayer::Decals : ERayTracingSceneLayer::Base);
+					FRayTracingGeometryInstance NewInstance = SceneInfo->CachedRayTracingInstance;
+					NewInstance.LayerIndex = (uint8)(SceneInfo->bCachedRayTracingInstanceAnySegmentsDecal && !bNeedSeparateDecalInstance ? ERayTracingSceneLayer::Decals : ERayTracingSceneLayer::Base);
 
 					const Experimental::FHashElementId GroupId = Scene.PrimitiveRayTracingGroupIds[PrimitiveIndex];
 					const bool bUseGroupBounds = CullingParameters.bCullUsingGroupIds && GroupId.IsValid();
@@ -1505,18 +1517,39 @@ bool FDeferredShadingSceneRenderer::GatherRayTracingWorldInstancesForView(FRDGBu
 
 					AddDebugRayTracingInstanceFlags(NewInstance.Flags);
 
-					const int32 NewInstanceIndex = RayTracingScene.AddInstance(MoveTemp(NewInstance), SceneInfo->Proxy, false);
+					const int32 NewInstanceIndex = RayTracingScene.AddInstance(NewInstance, SceneInfo->Proxy, false);
+
+					uint32 DecalInstanceIndex = INDEX_NONE;
+					if (bNeedSeparateDecalInstance && !GRayTracingExcludeDecals)
+					{
+						FRayTracingGeometryInstance DecalRayTracingInstance = NewInstance;
+						DecalRayTracingInstance.LayerIndex = (uint8)ERayTracingSceneLayer::Decals;
+
+						DecalInstanceIndex = RayTracingScene.AddInstance(MoveTemp(DecalRayTracingInstance), SceneInfo->Proxy, false);
+					}
 
 					// At the moment we only support SM & ISMs on this path
 					check(EnumHasAnyFlags(Flags, ERayTracingPrimitiveFlags::CacheMeshCommands));
 					if (SceneInfo->CachedRayTracingMeshCommandIndicesPerLOD.Num() > 0 && SceneInfo->CachedRayTracingMeshCommandIndicesPerLOD[LODIndex].Num() > 0)
 					{
+						const bool bHasDecalInstanceIndex = DecalInstanceIndex != INDEX_NONE;
+
 						for (int32 CommandIndex : SceneInfo->CachedRayTracingMeshCommandIndicesPerLOD[LODIndex])
 						{
-							FVisibleRayTracingMeshCommand NewVisibleMeshCommand;
-							NewVisibleMeshCommand.RayTracingMeshCommand = &Scene.CachedRayTracingMeshCommands[CommandIndex];
-							NewVisibleMeshCommand.InstanceIndex = NewInstanceIndex;
-							VisibleRayTracingMeshCommands.Add(NewVisibleMeshCommand);
+							const FRayTracingMeshCommand& MeshCommand = Scene.CachedRayTracingMeshCommands[CommandIndex];
+
+							{
+								const bool bHidden = bHasDecalInstanceIndex && MeshCommand.bDecal;
+								FVisibleRayTracingMeshCommand NewVisibleMeshCommand(&MeshCommand, NewInstanceIndex, bHidden);
+								VisibleRayTracingMeshCommands.Add(NewVisibleMeshCommand);
+							}
+
+							if(bHasDecalInstanceIndex)
+							{
+								const bool bHidden = !MeshCommand.bDecal;
+								FVisibleRayTracingMeshCommand NewVisibleMeshCommand(&MeshCommand, DecalInstanceIndex, bHidden);
+								VisibleRayTracingMeshCommands.Add(NewVisibleMeshCommand);
+							}
 						}
 					}
 				}
@@ -1529,9 +1562,14 @@ bool FDeferredShadingSceneRenderer::GatherRayTracingWorldInstancesForView(FRDGBu
 						continue; // skip dynamic primitives and other 
 					}
 
-					if (RelevantPrimitive.bAnySegmentsDecal && !RelevantPrimitive.bAllSegmentsDecal)
+					// if primitive has mixed decal and non-decal segments we need to have two ray tracing instances
+					// one containing non-decal segments and the other with decal segments
+					// masking of segments is done using "hidden" hitgroups
+					// TODO: Debug Visualization to highlight primitives using this?
+					const bool bNeedSeparateDecalInstance = RelevantPrimitive.bAnySegmentsDecal && !RelevantPrimitive.bAllSegmentsDecal;
+
+					if (GRayTracingExcludeDecals && RelevantPrimitive.bAnySegmentsDecal && !bNeedSeparateDecalInstance)
 					{
-						// TODO: Warn user meshes with mixed decal segments are not currently supported in ray tracing
 						continue;
 					}
 
@@ -1543,48 +1581,46 @@ bool FDeferredShadingSceneRenderer::GatherRayTracingWorldInstancesForView(FRDGBu
 					}
 
 					// location if this is a new entry
-					const int32 NewInstanceIndex = RayTracingScene.GetInstances().Num();
 					const uint64 InstanceKey = RelevantPrimitive.InstancingKey();
 
-					FAutoInstanceBatch DummyInstanceBatch = { NewInstanceIndex };
+					FAutoInstanceBatch DummyInstanceBatch = { };
 					FAutoInstanceBatch& InstanceBatch = bAutoInstance ? InstanceBatches.FindOrAdd(InstanceKey, DummyInstanceBatch) : DummyInstanceBatch;
 
-					if (InstanceBatch.Index != NewInstanceIndex)
+					if (InstanceBatch.IsValid())
 					{
 						// Reusing a previous entry, just append to the instance list.
 
-						FRayTracingGeometryInstance& RayTracingInstance = RayTracingScene.GetInstance(InstanceBatch.Index);
 						bool bReallocated = InstanceBatch.Add(RayTracingScene, SceneInfo->GetInstanceSceneDataOffset(), (uint32)PrimitiveIndex);
 
-						++RayTracingInstance.NumTransforms;
-						check(RayTracingInstance.NumTransforms == InstanceBatch.Cursor); // sanity check
-
-						if (bReallocated)
+						check(InstanceBatch.Index != INDEX_NONE);
 						{
-							RayTracingInstance.InstanceSceneDataOffsets = InstanceBatch.InstanceSceneDataOffsets;
-							RayTracingInstance.UserData = InstanceBatch.UserData;
+							FRayTracingGeometryInstance& RayTracingInstance = RayTracingScene.GetInstance(InstanceBatch.Index);
+							++RayTracingInstance.NumTransforms;
+							check(RayTracingInstance.NumTransforms == InstanceBatch.Cursor); // sanity check
+
+							if (bReallocated)
+							{
+								RayTracingInstance.InstanceSceneDataOffsets = InstanceBatch.InstanceSceneDataOffsets;
+								RayTracingInstance.UserData = InstanceBatch.UserData;
+							}
 						}
-					}
-					else
-					{
-						// Starting new instance batch
 
-						for (int32 CommandIndex : RelevantPrimitive.CachedRayTracingMeshCommandIndices)
+						if(InstanceBatch.DecalIndex != INDEX_NONE)
 						{
-							if (CommandIndex >= 0)
-							{
-								FVisibleRayTracingMeshCommand NewVisibleMeshCommand;
+							FRayTracingGeometryInstance& RayTracingInstance = RayTracingScene.GetInstance(InstanceBatch.DecalIndex);
+							++RayTracingInstance.NumTransforms;
+							check(RayTracingInstance.NumTransforms == InstanceBatch.Cursor); // sanity check
 
-								NewVisibleMeshCommand.RayTracingMeshCommand = &Scene.CachedRayTracingMeshCommands[CommandIndex];
-								NewVisibleMeshCommand.InstanceIndex = NewInstanceIndex;
-								VisibleRayTracingMeshCommands.Add(NewVisibleMeshCommand);
-							}
-							else
+							if (bReallocated)
 							{
-								// CommandIndex == -1 indicates that the mesh batch has been filtered by FRayTracingMeshProcessor (like the shadow depth pass batch)
-								// Do nothing in this case
+								RayTracingInstance.InstanceSceneDataOffsets = InstanceBatch.InstanceSceneDataOffsets;
+								RayTracingInstance.UserData = InstanceBatch.UserData;
 							}
 						}
+					}
+					else
+					{
+						// Starting new instance batch
 
 						InstanceBatch.Add(RayTracingScene, SceneInfo->GetInstanceSceneDataOffset(), (uint32)PrimitiveIndex);
 
@@ -1607,11 +1643,45 @@ bool FDeferredShadingSceneRenderer::GatherRayTracingWorldInstancesForView(FRDGBu
 						}
 						AddDebugRayTracingInstanceFlags(RayTracingInstance.Flags);
 
-						RayTracingInstance.LayerIndex = (uint8)(RelevantPrimitive.bAnySegmentsDecal ? ERayTracingSceneLayer::Decals : ERayTracingSceneLayer::Base);
+						RayTracingInstance.LayerIndex = (uint8)(RelevantPrimitive.bAnySegmentsDecal && !bNeedSeparateDecalInstance ? ERayTracingSceneLayer::Decals : ERayTracingSceneLayer::Base);
+
+						InstanceBatch.Index = RayTracingScene.AddInstance(RayTracingInstance, SceneInfo->Proxy, false);
+
+						if (bNeedSeparateDecalInstance && !GRayTracingExcludeDecals)
+						{
+							FRayTracingGeometryInstance DecalRayTracingInstance = RayTracingInstance;
+							DecalRayTracingInstance.LayerIndex = (uint8)ERayTracingSceneLayer::Decals;
+
+							InstanceBatch.DecalIndex = RayTracingScene.AddInstance(MoveTemp(DecalRayTracingInstance), SceneInfo->Proxy, false);
+						}
 
-						const uint32 ActualInstanceIndex = RayTracingScene.AddInstance(MoveTemp(RayTracingInstance), SceneInfo->Proxy, false);
+						const bool bHasDecalInstanceIndex = InstanceBatch.DecalIndex != INDEX_NONE;
 
-						check(NewInstanceIndex == ActualInstanceIndex);
+						for (int32 CommandIndex : RelevantPrimitive.CachedRayTracingMeshCommandIndices)
+						{
+							if (CommandIndex >= 0)
+							{
+								const FRayTracingMeshCommand& MeshCommand = Scene.CachedRayTracingMeshCommands[CommandIndex];
+
+								{
+									const bool bHidden = bHasDecalInstanceIndex && MeshCommand.bDecal;
+									FVisibleRayTracingMeshCommand NewVisibleMeshCommand(&MeshCommand, InstanceBatch.Index, bHidden);
+									VisibleRayTracingMeshCommands.Add(NewVisibleMeshCommand);
+								}
+
+								if (bHasDecalInstanceIndex)
+								{
+									const bool bHidden = !MeshCommand.bDecal;
+									FVisibleRayTracingMeshCommand NewVisibleMeshCommand(&MeshCommand, InstanceBatch.DecalIndex, bHidden);
+									VisibleRayTracingMeshCommands.Add(NewVisibleMeshCommand);
+								}
+							}
+							else
+							{
+								// CommandIndex == -1 indicates that the mesh batch has been filtered by FRayTracingMeshProcessor (like the shadow depth pass batch)
+								// Do nothing in this case
+							}
+						}
 					}
 				}
 			}
diff --git a/Engine/Source/Runtime/Renderer/Private/RayTracing/RayTracingMaterialHitShaders.cpp b/Engine/Source/Runtime/Renderer/Private/RayTracing/RayTracingMaterialHitShaders.cpp
--- a/Engine/Source/Runtime/Renderer/Private/RayTracing/RayTracingMaterialHitShaders.cpp
+++ b/Engine/Source/Runtime/Renderer/Private/RayTracing/RayTracingMaterialHitShaders.cpp
@@ -813,17 +813,17 @@ FRayTracingPipelineState* FDeferredShadingSceneRenderer::CreateRayTracingMateria
 						if (bIsMeshDecalShader)
 						{
 							checkf(bSupportMeshDecals && MeshCommand.bDecal, TEXT("Unexpected ray tracing mesh command using Mesh Decal payload. Fix logic adding the command or update bSupportMeshDecals as appropriate."));
-							HitGroupIndex = OpaqueMeshDecalHitGroupIndex;
+							HitGroupIndex = VisibleMeshCommand.bHidden ? HiddenMeshDecalHitGroupIndex : OpaqueMeshDecalHitGroupIndex;
 						}
 						else
 						{
 							checkf((!bIsPathTracing && MeshCommand.MaterialShader->RayTracingPayloadType == (uint32)ERayTracingPayloadType::RayTracingMaterial)
 								|| (bIsPathTracing && MeshCommand.MaterialShader->RayTracingPayloadType == (uint32)ERayTracingPayloadType::PathTracingMaterial),
 								TEXT("Incorrectly using RayTracingMaterial when path tracer is enabled or vice-versa."));
-							HitGroupIndex = OpaqueShadowMaterialIndex;
+							HitGroupIndex = VisibleMeshCommand.bHidden ? HiddenMaterialIndex : OpaqueShadowMaterialIndex;
 						}
 
-						if (bEnableMaterials)
+						if (bEnableMaterials && !VisibleMeshCommand.bHidden)
 						{
 							const int32 FoundIndex = FindRayTracingHitGroupIndex(PipelineState, MeshCommand.MaterialShader, false);
 							if (FoundIndex != INDEX_NONE)
@@ -846,7 +846,7 @@ FRayTracingPipelineState* FDeferredShadingSceneRenderer::CreateRayTracingMateria
 
 						// Bind shadow shader
 
-						if (MeshCommand.bCastRayTracedShadows)
+						if (MeshCommand.bCastRayTracedShadows && !VisibleMeshCommand.bHidden)
 						{
 							if (MeshCommand.bOpaque || !bEnableShadowMaterials)
 							{
diff --git a/Engine/Source/Runtime/Renderer/Public/MeshPassProcessor.h b/Engine/Source/Runtime/Renderer/Public/MeshPassProcessor.h
--- a/Engine/Source/Runtime/Renderer/Public/MeshPassProcessor.h
+++ b/Engine/Source/Runtime/Renderer/Public/MeshPassProcessor.h
@@ -2222,9 +2222,17 @@ private:
 class FVisibleRayTracingMeshCommand
 {
 public:
-	const FRayTracingMeshCommand* RayTracingMeshCommand;
+	FVisibleRayTracingMeshCommand(const FRayTracingMeshCommand* InRayTracingMeshCommand, uint32 InInstanceIndex, bool bInHidden = false)
+		: RayTracingMeshCommand(InRayTracingMeshCommand)
+		, InstanceIndex(InInstanceIndex)
+		, bHidden(bInHidden)
+	{
+		check(InstanceIndex != INDEX_NONE);
+	}
 
+	const FRayTracingMeshCommand* RayTracingMeshCommand;
 	uint32 InstanceIndex;
+	bool bHidden;
 };
 
 template <>
@@ -2280,13 +2288,15 @@ public:
 	(
 		FDynamicRayTracingMeshCommandStorage& InDynamicCommandStorage,
 		FRayTracingMeshCommandOneFrameArray& InVisibleCommands,
-		uint32 InGeometrySegmentIndex = ~0u,
-		uint32 InRayTracingInstanceIndex = ~0u
+		uint32 InGeometrySegmentIndex,
+		uint32 InRayTracingInstanceIndex,
+		uint32 InRayTracingDecalInstanceIndex = INDEX_NONE
 	) :
 		DynamicCommandStorage(InDynamicCommandStorage),
 		VisibleCommands(InVisibleCommands),
 		GeometrySegmentIndex(InGeometrySegmentIndex),
-		RayTracingInstanceIndex(InRayTracingInstanceIndex)
+		RayTracingInstanceIndex(InRayTracingInstanceIndex),
+		RayTracingDecalInstanceIndex(InRayTracingDecalInstanceIndex)
 	{}
 
 	virtual FRayTracingMeshCommand& AddCommand(const FRayTracingMeshCommand& Initializer) override final
@@ -2299,10 +2309,20 @@ public:
 
 	virtual void FinalizeCommand(FRayTracingMeshCommand& RayTracingMeshCommand) override final
 	{
-		FVisibleRayTracingMeshCommand NewVisibleMeshCommand;
-		NewVisibleMeshCommand.RayTracingMeshCommand = &RayTracingMeshCommand;
-		NewVisibleMeshCommand.InstanceIndex = RayTracingInstanceIndex;
-		VisibleCommands.Add(NewVisibleMeshCommand);
+		const bool bHasDecalInstanceIndex = RayTracingDecalInstanceIndex != INDEX_NONE;
+
+		{
+			const bool bHidden = bHasDecalInstanceIndex && RayTracingMeshCommand.bDecal;
+			FVisibleRayTracingMeshCommand NewVisibleMeshCommand(&RayTracingMeshCommand, RayTracingInstanceIndex, bHidden);
+			VisibleCommands.Add(NewVisibleMeshCommand);
+		}
+
+		if (bHasDecalInstanceIndex)
+		{
+			const bool bHidden = !RayTracingMeshCommand.bDecal;
+			FVisibleRayTracingMeshCommand NewVisibleMeshCommand(&RayTracingMeshCommand, RayTracingDecalInstanceIndex, bHidden);
+			VisibleCommands.Add(NewVisibleMeshCommand);
+		}
 	}
 
 private:
@@ -2310,6 +2330,7 @@ private:
 	FRayTracingMeshCommandOneFrameArray& VisibleCommands;
 	uint32 GeometrySegmentIndex;
 	uint32 RayTracingInstanceIndex;
+	uint32 RayTracingDecalInstanceIndex;
 };
 
 class FRayTracingShaderCommand
