Fix a bug of missing modulated shadow in shadow preview mode.

diff --git a/Engine/Source/Runtime/Renderer/Private/SceneCore.cpp b/Engine/Source/Runtime/Renderer/Private/SceneCore.cpp
--- a/Engine/Source/Runtime/Renderer/Private/SceneCore.cpp
+++ b/Engine/Source/Runtime/Renderer/Private/SceneCore.cpp
@@ -256,6 +256,12 @@ FLightPrimitiveInteraction::FLightPrimitiveInteraction(
 					PrimitiveSceneInfo->BeginDeferredUpdateStaticMeshes();
 				}
 			} 
+
+			if (LightSceneInfo->Proxy->CastsModulatedShadows() && !LightSceneInfo->Proxy->UseCSMForDynamicObjects())
+			{
+				// Force bCastInsetShadow to be enabled to cast modulated shadow on mobile
+				PrimitiveSceneInfo->Proxy->bCastInsetShadow = true;
+			}
 		}
 	}
 
