Remove a wait inside FVirtualTextureProducer::Release that can stall the render thread and cause a hitch

diff --git a/Engine/Source/Runtime/Renderer/Private/VT/VirtualTextureProducer.cpp b/Engine/Source/Runtime/Renderer/Private/VT/VirtualTextureProducer.cpp
--- a/Engine/Source/Runtime/Renderer/Private/VT/VirtualTextureProducer.cpp
+++ b/Engine/Source/Runtime/Renderer/Private/VT/VirtualTextureProducer.cpp
@@ -8,6 +8,36 @@ FVirtualTextureProducer::~FVirtualTextureProducer()
 {
 }
 
+struct FReleaseVirtualTextureTask
+{
+	IVirtualTexture* VirtualTexture;
+
+	FReleaseVirtualTextureTask(IVirtualTexture* InVirtualTexture)
+		: VirtualTexture(InVirtualTexture)
+	{}
+
+	void DoTask(ENamedThreads::Type CurrentThread, const FGraphEventRef& MyCompletionGraphEvent)
+	{
+		delete VirtualTexture;
+		VirtualTexture = nullptr;
+	}
+
+	static ESubsequentsMode::Type GetSubsequentsMode()
+	{
+		return ESubsequentsMode::FireAndForget;
+	}
+
+	ENamedThreads::Type GetDesiredThread()
+	{
+		return ENamedThreads::GetRenderThread();
+	}
+
+	FORCEINLINE TStatId GetStatId() const
+	{
+		RETURN_QUICK_DECLARE_CYCLE_STAT(DeferredReleaseVirtualTexture, STATGROUP_VirtualTexturing);
+	}
+};
+
 void FVirtualTextureProducer::Release(FVirtualTextureSystem* System, const FVirtualTextureProducerHandle& HandleToSelf)
 {
 	TRACE_CPUPROFILER_EVENT_SCOPE(FVirtualTextureProducer::Release);
@@ -30,11 +60,13 @@ void FVirtualTextureProducer::Release(FVirtualTextureSystem* System, const FVirt
 	VirtualTexture->GatherProducePageDataTasks(HandleToSelf, ProducePageTasks);
 	if (ProducePageTasks.Num())
 	{
-		TRACE_CPUPROFILER_EVENT_SCOPE(FVirtualTextureProducer::Release_Wait);
-		FTaskGraphInterface::Get().WaitUntilTasksComplete(ProducePageTasks, ENamedThreads::GetRenderThread_Local());
+		TGraphTask<FReleaseVirtualTextureTask>::CreateTask(&ProducePageTasks).ConstructAndDispatchWhenReady(VirtualTexture);
+	}
+	else
+	{
+		delete VirtualTexture;
 	}
 
-	delete VirtualTexture;
 	VirtualTexture = nullptr;
 	
 	Description = FVTProducerDescription();
