Fixed the fixed foveated VRS image generator's behaviour when not using instanced stereo. Now correctly generates two foveated patterns (one per eye), rather than a single pattern for both eyes.

diff --git a/Engine/Source/Runtime/Renderer/Private/VariableRateShading/FixedFoveationImageGenerator.cpp b/Engine/Source/Runtime/Renderer/Private/VariableRateShading/FixedFoveationImageGenerator.cpp
--- a/Engine/Source/Runtime/Renderer/Private/VariableRateShading/FixedFoveationImageGenerator.cpp
+++ b/Engine/Source/Runtime/Renderer/Private/VariableRateShading/FixedFoveationImageGenerator.cpp
@@ -155,7 +155,6 @@ void FFixedFoveationImageGenerator::PrepareImages(FRDGBuilder& GraphBuilder, con
 
 	FIntPoint Size = ViewFamily.RenderTarget->GetSizeXY();
 	bool bStereoRendering = IStereoRendering::IsStereoEyeView(*ViewFamily.Views[0]) && GEngine->XRSystem.IsValid();
-	bool bInstancedStereo = ViewFamily.Views[0]->IsInstancedStereoPass();
 
 	// Sanity check VRS tile size.
 	check(GRHIVariableRateShadingImageTileMinWidth >= 8 && GRHIVariableRateShadingImageTileMinWidth <= 64 && GRHIVariableRateShadingImageTileMinHeight >= 8 && GRHIVariableRateShadingImageTileMaxHeight <= 64);
@@ -188,14 +187,11 @@ void FFixedFoveationImageGenerator::PrepareImages(FRDGBuilder& GraphBuilder, con
 	if (bStereoRendering)
 	{
 		EnumAddFlags(GenFlags, EVRSGenerationFlags::StereoRendering);
-	}
+		EnumAddFlags(GenFlags, EVRSGenerationFlags::SideBySideStereo); // May need to change this if we add support for mobile multi-view
 
-	// If instanced (side-by-side) stereo, there's two "center" points, so adjust both eyes
-	if (bInstancedStereo)
-	{
+		// Adjust eyes for side-by-side stereo
 		PassParameters->LeftEyeCenterPixelXY.X /= 2;
 		PassParameters->RightEyeCenterPixelXY.X = PassParameters->LeftEyeCenterPixelXY.X + TextureSize.X / 2;
-		EnumAddFlags(GenFlags, EVRSGenerationFlags::SideBySideStereo);
 	}
 
 	PassParameters->ViewDiagonalSquaredInPixels = FVector2f::DotProduct(PassParameters->LeftEyeCenterPixelXY, PassParameters->LeftEyeCenterPixelXY);
