Fixed issue where VRS image combiner was using the dimensions of a single eye rather than the entire render target for side-by-side stereo.

diff --git a/Engine/Source/Runtime/Renderer/Private/VariableRateShading/VariableRateShadingImageManager.cpp b/Engine/Source/Runtime/Renderer/Private/VariableRateShading/VariableRateShadingImageManager.cpp
--- a/Engine/Source/Runtime/Renderer/Private/VariableRateShading/VariableRateShadingImageManager.cpp
+++ b/Engine/Source/Runtime/Renderer/Private/VariableRateShading/VariableRateShadingImageManager.cpp
@@ -222,7 +222,7 @@ FRDGTextureRef FVariableRateShadingImageManager::CombineShadingRateImages(FRDGBu
 	}
 	else
 	{
-		FIntPoint ViewSize = ViewInfo.UnscaledViewRect.Scale(ViewInfo.Family->SecondaryViewFraction).Size();
+		FIntPoint ViewSize = ViewInfo.Family->RenderTarget->GetSizeXY();
 		const FIntPoint TileSize = GetSRITileSize();
 
 		// Create texture to hold shading rate image
