PR #9946: Fix the High Res Screenshot and LDR buffers (Contributed by [REDACTED_PERSON])

diff --git a/Engine/Source/Runtime/Renderer/Private/PostProcess/PostProcessVisualizeBuffer.cpp b/Engine/Source/Runtime/Renderer/Private/PostProcess/PostProcessVisualizeBuffer.cpp
--- a/Engine/Source/Runtime/Renderer/Private/PostProcess/PostProcessVisualizeBuffer.cpp
+++ b/Engine/Source/Runtime/Renderer/Private/PostProcess/PostProcessVisualizeBuffer.cpp
@@ -285,6 +285,11 @@ void AddDumpToFilePass(FRDGBuilder& GraphBuilder, FScreenPassTexture Input, cons
 			// ImageTask->PixelData should be sRGB
 			//  it will gamma correct automatically if written to EXR
 		}
+		else if(ImageTask->Format == EImageFormat::PNG)
+		{
+			// PNGs can't have 0 alpha or RGB data is destroyed.
+			ImageTask->PixelPreProcessors.Add(TAsyncAlphaWrite<FLinearColor>(1.0f));
+		}
 
 		HighResScreenshotConfig.ImageWriteQueue->Enqueue(MoveTemp(ImageTask));
 	});
