Exempt moving distant lights from force-caching and enable 'distant lights' by default.

diff --git a/Engine/Source/Runtime/Renderer/Private/Shadows/ShadowSceneRenderer.cpp b/Engine/Source/Runtime/Renderer/Private/Shadows/ShadowSceneRenderer.cpp
--- a/Engine/Source/Runtime/Renderer/Private/Shadows/ShadowSceneRenderer.cpp
+++ b/Engine/Source/Runtime/Renderer/Private/Shadows/ShadowSceneRenderer.cpp
@@ -1,7 +1,5 @@
 // Copyright Epic Games, Inc. All Rights Reserved.
-/*=============================================================================
-	ShadowSceneRenderer.cpp:
-=============================================================================*/
+
 #include "ShadowSceneRenderer.h"
 #include "DeferredShadingRenderer.h"
 #include "ScenePrivate.h"
@@ -33,11 +31,20 @@ TAutoConsoleVariable<int32> CVarMaxDistantLightsPerFrame(
 
 static TAutoConsoleVariable<int32> CVarDistantLightMode(
 	TEXT("r.Shadow.Virtual.DistantLightMode"),
-	0,
-	TEXT("Control whether distant light mode is enabled for local lights.\n0 == Off (default), \n1 == On, \n2 == Force All."),
+	1,
+	TEXT("Control whether distant light mode is enabled for local lights.\n0 == Off, \n1 == On (default), \n2 == Force All."),
 	ECVF_RenderThreadSafe
 );
 
+static TAutoConsoleVariable<float> CVarDistantLightForceCacheFootprintFraction(
+	TEXT("r.Shadow.Virtual.DistantLightForceCacheFootprintFraction"),
+	0.0f,
+	TEXT("Fraction of footprint size below which start force-caching lights that are invalidated (i.e., are moving or re-added)\n")
+	TEXT("  The base footprint is based on the page size.\n")
+	TEXT("  0.0 == Never force-cache (default), 1.0 == Always force-cache."),
+	ECVF_Scalability | ECVF_RenderThreadSafe
+);
+
 static TAutoConsoleVariable<float> CVarNaniteShadowsLODBias(
 	TEXT("r.Shadow.NaniteLODBias"),
 	1.0f,
@@ -51,13 +58,15 @@ TAutoConsoleVariable<int32> CVarVirtualShadowOnePassProjection(
 	ECVF_Scalability | ECVF_RenderThreadSafe
 );
 
-
 DECLARE_DWORD_COUNTER_STAT(TEXT("VSM Total Raster Bins"), STAT_VSMNaniteBasePassTotalRasterBins, STATGROUP_ShadowRendering);
 DECLARE_DWORD_COUNTER_STAT(TEXT("VSM Total Shading Draws"), STAT_VSMNaniteBasePassTotalShadingDraws, STATGROUP_ShadowRendering);
 
 DECLARE_DWORD_COUNTER_STAT(TEXT("VSM Visible Raster Bins"), STAT_VSMNaniteBasePassVisibleRasterBins, STATGROUP_ShadowRendering);
 DECLARE_DWORD_COUNTER_STAT(TEXT("VSM Visible Shading Draws"), STAT_VSMNaniteBassPassVisibleShadingDraws, STATGROUP_ShadowRendering);
 
+DECLARE_DWORD_COUNTER_STAT(TEXT("Distant Light Count"), STAT_DistantLightCount, STATGROUP_ShadowRendering);
+DECLARE_DWORD_COUNTER_STAT(TEXT("Distant Cached Count"), STAT_DistantCachedCount, STATGROUP_ShadowRendering);
+
 FShadowSceneRenderer::FShadowSceneRenderer(FDeferredShadingSceneRenderer& InSceneRenderer)
 	: SceneRenderer(InSceneRenderer)
 	, Scene(*InSceneRenderer.Scene)
@@ -84,8 +93,9 @@ TSharedPtr<FVirtualShadowMapPerLightCacheEntry> FShadowSceneRenderer::AddLocalLi
 	//       we can absolutely mirror the page marking calc better, just unclear how much it helps. 
 	//       Also possible to feed back from gpu - which would be more accurate wrt partially visible lights (e.g., a spot going through the ground).
 	//       Of course this creates jumps if visibility changes, which may or may not create unsolvable artifacts.
+	const float BiasedFootprintThreshold = float(FVirtualShadowMap::PageSize) * FMath::Exp2(VirtualShadowMapArray.GetResolutionLODBiasLocal());
 	const bool bIsDistantLight = CVarDistantLightMode.GetValueOnRenderThread() != 0
-		&& (MaxScreenRadius <= float(FVirtualShadowMap::PageSize) * FMath::Exp2(VirtualShadowMapArray.GetResolutionLODBiasLocal()) || CVarDistantLightMode.GetValueOnRenderThread() == 2);
+		&& (MaxScreenRadius <= BiasedFootprintThreshold || CVarDistantLightMode.GetValueOnRenderThread() == 2);
 
 
 	const int32 NumMaps = ProjectedShadowInitializer.bOnePassPointLightShadow ? 6 : 1;
@@ -100,8 +110,10 @@ TSharedPtr<FVirtualShadowMapPerLightCacheEntry> FShadowSceneRenderer::AddLocalLi
 	TSharedPtr<FVirtualShadowMapPerLightCacheEntry> PerLightCacheEntry = CacheManager->FindCreateLightCacheEntry(LightSceneInfo->Id);
 	if (PerLightCacheEntry.IsValid())
 	{
+		const float DistantLightForceCacheFootprintFraction = FMath::Clamp(CVarDistantLightForceCacheFootprintFraction.GetValueOnRenderThread(), 0.0f, 1.0f);
+		bool bShouldForceTimeSliceDistantUpdate = bIsDistantLight && MaxScreenRadius <= BiasedFootprintThreshold * DistantLightForceCacheFootprintFraction;
 		LocalLightShadowFrameSetup.PerLightCacheEntry = PerLightCacheEntry;
-		PerLightCacheEntry->UpdateLocal(ProjectedShadowInitializer, bIsDistantLight);
+		bool bIsCached = PerLightCacheEntry->UpdateLocal(ProjectedShadowInitializer, bIsDistantLight, !bShouldForceTimeSliceDistantUpdate);
 
 		for (int32 Index = 0; Index < NumMaps; ++Index)
 		{
@@ -111,8 +123,8 @@ TSharedPtr<FVirtualShadowMapPerLightCacheEntry> FShadowSceneRenderer::AddLocalLi
 			VirtualSmCacheEntry->UpdateLocal(VirtualShadowMap->ID, *PerLightCacheEntry);
 			VirtualShadowMap->VirtualShadowMapCacheEntry = VirtualSmCacheEntry;
 		}
-
-		if (bIsDistantLight)
+		// Only round-robin those that were not invalidated.
+		if (bIsDistantLight && bIsCached)
 		{
 			// This priority could be calculated based also on whether the light has actually been invalidated or not (currently not tracked on host).
 			// E.g., all things being equal update those with an animated mesh in, for example. Plus don't update those the don't need it at all.
@@ -255,13 +267,12 @@ void FShadowSceneRenderer::UpdateDistantLightPriorityRender()
 	}
 }
 
-
 void FShadowSceneRenderer::PostSetupDebugRender()
 {
 #if !(UE_BUILD_SHIPPING || UE_BUILD_TEST)
-	// TODO: Move to debug rendering function in FShadowSceneRenderer
-	if ((SceneRenderer.ViewFamily.EngineShowFlags.DebugDrawDistantVirtualSMLights))
+	if ((SceneRenderer.ViewFamily.EngineShowFlags.DebugDrawDistantVirtualSMLights) && VirtualShadowMapArray.IsEnabled())
 	{
+		int32 NumFullyCached = 0;
 		int32 NumDistant = 0;
 		for (FViewInfo& View : SceneRenderer.Views)
 		{
@@ -275,7 +286,15 @@ void FShadowSceneRenderer::PostSetupDebugRender()
 					++NumDistant;
 					int32 FramesSinceLastRender = int32(Scene.GetFrameNumber()) - int32(LightSetup.PerLightCacheEntry->GetLastScheduledFrameNumber());
 					float Fade = FMath::Min(0.8f, float(FramesSinceLastRender) / float(LocalLights.Num()));
-					Color = LightSetup.PerLightCacheEntry->IsFullyCached() ? FMath::Lerp(FLinearColor(FColor::Green), FLinearColor(FColor::Red), Fade) : FLinearColor(FColor::Red);
+					if (LightSetup.PerLightCacheEntry->IsFullyCached())
+					{
+						++NumFullyCached;
+						Color = FMath::Lerp(FLinearColor(FColor::Green), FLinearColor(FColor::Red), Fade);
+					}
+					else
+					{
+						Color = FLinearColor(FColor::Purple);
+					}
 				}
 
 				Color.A = 1.0f;
@@ -292,12 +311,8 @@ void FShadowSceneRenderer::PostSetupDebugRender()
 				}
 			}
 		}
-		SceneRenderer.OnGetOnScreenMessages.AddLambda([this, NumDistant](FScreenMessageWriter& ScreenMessageWriter)->void
-		{
-			ScreenMessageWriter.DrawLine(FText::FromString(FString::Printf(TEXT("Distant Light Count: %d"), NumDistant)), 10, FColor::Yellow);
-			ScreenMessageWriter.DrawLine(FText::FromString(FString::Printf(TEXT("Active Local Light Count: %d"), LocalLights.Num())), 10, FColor::Yellow);
-			ScreenMessageWriter.DrawLine(FText::FromString(FString::Printf(TEXT("Scene Light Count: %d"), Scene.Lights.Num())), 10, FColor::Yellow);
-		});
+		SET_DWORD_STAT(STAT_DistantLightCount, NumDistant);
+		SET_DWORD_STAT(STAT_DistantCachedCount, NumFullyCached);
 	}
 #endif
 }
diff --git a/Engine/Source/Runtime/Renderer/Private/VirtualShadowMaps/VirtualShadowMapCacheManager.cpp b/Engine/Source/Runtime/Renderer/Private/VirtualShadowMaps/VirtualShadowMapCacheManager.cpp
--- a/Engine/Source/Runtime/Renderer/Private/VirtualShadowMaps/VirtualShadowMapCacheManager.cpp
+++ b/Engine/Source/Runtime/Renderer/Private/VirtualShadowMaps/VirtualShadowMapCacheManager.cpp
@@ -1,8 +1,5 @@
 // Copyright Epic Games, Inc. All Rights Reserved.
 
-/*=============================================================================
-	VirtualShadowMap.h:
-=============================================================================*/
 #include "VirtualShadowMapCacheManager.h"
 #include "VirtualShadowMapClipmap.h"
 #include "RendererModule.h"
@@ -184,7 +181,7 @@ void FVirtualShadowMapPerLightCacheEntry::UpdateClipmap()
 	CurrentRenderedFrameNumber = -1;
 }
 
-void FVirtualShadowMapPerLightCacheEntry::UpdateLocal(const FProjectedShadowInitializer& InCacheKey, bool bIsDistantLight)
+bool FVirtualShadowMapPerLightCacheEntry::UpdateLocal(const FProjectedShadowInitializer& InCacheKey, bool bIsDistantLight, bool bAllowInvalidation)
 {
 	bPrevIsDistantLight = bCurrentIsDistantLight;
 	PrevRenderedFrameNumber = FMath::Max(PrevRenderedFrameNumber, CurrentRenderedFrameNumber);
@@ -195,7 +192,7 @@ void FVirtualShadowMapPerLightCacheEntry::UpdateLocal(const FProjectedShadowInit
 	{
 		// If it is a distant light, we want to let the time-share perform the invalidation.
 		// TODO: track invalidation state somehow for later.
-		if (!bIsDistantLight)
+		if (bAllowInvalidation)
 		{
 			PrevRenderedFrameNumber = -1;
 		}
@@ -206,6 +203,8 @@ void FVirtualShadowMapPerLightCacheEntry::UpdateLocal(const FProjectedShadowInit
 	bCurrentIsDistantLight = bIsDistantLight;
 	CurrentRenderedFrameNumber = -1;
 	CurrenScheduledFrameNumber = -1;
+
+	return PrevRenderedFrameNumber >= 0;
 }
 
 void FVirtualShadowMapPerLightCacheEntry::Invalidate()
diff --git a/Engine/Source/Runtime/Renderer/Private/VirtualShadowMaps/VirtualShadowMapCacheManager.h b/Engine/Source/Runtime/Renderer/Private/VirtualShadowMaps/VirtualShadowMapCacheManager.h
--- a/Engine/Source/Runtime/Renderer/Private/VirtualShadowMaps/VirtualShadowMapCacheManager.h
+++ b/Engine/Source/Runtime/Renderer/Private/VirtualShadowMaps/VirtualShadowMapCacheManager.h
@@ -1,8 +1,5 @@
 // Copyright Epic Games, Inc. All Rights Reserved.
 
-/*=============================================================================
-	VirtualShadowMapArray.h:
-=============================================================================*/
 #pragma once
 
 #include "CoreMinimal.h"
@@ -78,11 +75,14 @@ public:
 	/**
 	 * The (local) VSM is fully cached if the previous frame if was distant and is distant this frame also.
 	 */
-	inline bool IsFullyCached() const { return bCurrentIsDistantLight && bPrevIsDistantLight; }
+	inline bool IsFullyCached() const { return bCurrentIsDistantLight && bPrevIsDistantLight && PrevRenderedFrameNumber >= 0; }
 	void MarkRendered(int32 FrameIndex) { CurrentRenderedFrameNumber = FrameIndex; }
 	int32 GetLastScheduledFrameNumber() const { return PrevScheduledFrameNumber; }
 	void UpdateClipmap();
-	void UpdateLocal(const FProjectedShadowInitializer &InCacheKey, bool bIsDistantLight = false);
+	/**
+	 * Returns true if the cache entry is valid (has previous state).
+	 */
+	bool UpdateLocal(const FProjectedShadowInitializer &InCacheKey, bool bIsDistantLight, bool bAllowInvalidation);
 
 	/**
 	 * Mark as invalid, i.e., needing rendering.
