Implement batched add/remove/update of lights in the scene, similar to primitives. Stepping stone to handling light scene changes efficiently.

diff --git a/Engine/Source/Runtime/Renderer/Private/RendererScene.cpp b/Engine/Source/Runtime/Renderer/Private/RendererScene.cpp
--- a/Engine/Source/Runtime/Renderer/Private/RendererScene.cpp
+++ b/Engine/Source/Runtime/Renderer/Private/RendererScene.cpp
@@ -149,6 +149,93 @@ TGlobalResource< FGlobalDitherUniformBuffer > GDitherFadedInUniformBuffer;
 
 static FThreadSafeCounter FSceneViewState_UniqueID;
 
+struct FUpdateLightTransformParameters
+{
+	FMatrix LightToWorld;
+	FVector4 Position;
+};
+
+struct FUpdateLightCommand
+{
+	enum class EAddOrRemove
+	{
+		Add,
+		Remove,
+		None,
+	};
+
+	struct FColorParameters
+	{
+		FLinearColor NewColor;
+		float NewIndirectLightingScale;
+		float NewVolumetricScatteringIntensity;
+	};
+
+	using FTransformParameters = FUpdateLightTransformParameters;
+
+	inline bool IsAdd() const { return AddOrRemove == EAddOrRemove::Add; }
+	inline bool IsRemove() const { return AddOrRemove == EAddOrRemove::Remove; }
+
+	void Set(const FColorParameters& InColorParameters)
+	{
+		bHasColor = true;
+		ColorParameters = InColorParameters;
+	}
+	void Set(const FTransformParameters& InTransformParameters)
+	{
+		bHasTransform = true;
+		TransformParameters = InTransformParameters;
+	}
+	void Set(const EAddOrRemove& InAddOrRemove)
+	{
+		AddOrRemove = InAddOrRemove;
+	}
+
+	FUpdateLightCommand(const FColorParameters& InColorParameters, FLightSceneInfo* InLightSceneInfo) : LightSceneInfo(InLightSceneInfo), AddOrRemove(EAddOrRemove::None), bHasTransform(false) { Set(InColorParameters); }
+	FUpdateLightCommand(const FTransformParameters& InTransformParameters, FLightSceneInfo* InLightSceneInfo) : LightSceneInfo(InLightSceneInfo), AddOrRemove(EAddOrRemove::None), bHasColor(false) { Set(InTransformParameters); }
+	// Crtor for add/re,pve
+	FUpdateLightCommand(EAddOrRemove InAddOrRemove, FLightSceneInfo* InLightSceneInfo) : LightSceneInfo(InLightSceneInfo), AddOrRemove(InAddOrRemove), bHasTransform(false), bHasColor(false) {}
+
+	FTransformParameters TransformParameters;
+	FColorParameters ColorParameters;
+	FLightSceneInfo* LightSceneInfo;
+	EAddOrRemove AddOrRemove;
+	uint32 bHasTransform : 1;
+	uint32 bHasColor : 1;
+};
+
+class FSceneLightInfoUpdates
+{
+public:
+	Experimental::TRobinHoodHashMap<FLightSceneInfo*, FUpdateLightCommand> Commands;
+
+	int32 NumAdds = 0;
+	int32 NumRemoves = 0;
+	int32 NumUpdates = 0;
+
+	template <typename PayloadT>
+	void Enqueue(const PayloadT& Payload, FLightSceneInfo* LightSceneInfo)
+	{
+		// Allocate a new slot for update data
+		bool bWasAlreadyInMap = false;
+		FUpdateLightCommand* Command = Commands.FindOrAdd(LightSceneInfo, FUpdateLightCommand(Payload, LightSceneInfo), bWasAlreadyInMap);
+		if (bWasAlreadyInMap)
+		{
+			check(LightSceneInfo == Command->LightSceneInfo);
+			Command->Set(Payload);
+		}
+	}
+
+	void Reset()
+	{
+		Commands.Empty();
+		NumAdds = 0;
+		NumRemoves = 0;
+		NumUpdates = 0;
+	}
+};
+
+
 #define ENABLE_LOG_PRIMITIVE_INSTANCE_ID_STATS_TO_CSV 0
 
 #if ENABLE_LOG_PRIMITIVE_INSTANCE_ID_STATS_TO_CSV
@@ -1938,6 +2025,8 @@ FScene::FScene(UWorld* InWorld, bool bInRequiresHitProxies, bool bInIsEditorScen
 	// We use a default Virtual Shadow Map cache, if one hasn't been allocated for a specific view.  GPU resources for
 	// a cache aren't allocated until the cache is actually used, so this shouldn't waste any GPU memory when not in use.
 	DefaultVirtualShadowMapCache = new FVirtualShadowMapArrayCacheManager(this);
+
+	SceneLightInfoUpdates = new FSceneLightInfoUpdates;
 }
 
 FScene::~FScene()
@@ -2007,6 +2096,8 @@ FScene::~FScene()
 #endif // RHI_RAYTRACING
 
 	checkf(RemovedPrimitiveSceneInfos.Num() == 0, TEXT("Leaking %d FPrimitiveSceneInfo instances."), RemovedPrimitiveSceneInfos.Num()); // Ensure UpdateAllPrimitiveSceneInfos() is called before destruction.
+
+	delete SceneLightInfoUpdates;
 }
 
 void FScene::AddPrimitive(UPrimitiveComponent* Primitive)
@@ -2697,14 +2788,16 @@ void FScene::AddLight(ULightComponent* Light)
 		++NumVisibleLights_GameThread;
 
 		// Send a command to the rendering thread to add the light to the scene.
-		FScene* Scene = this;
 		FLightSceneInfo* LightSceneInfo = Proxy->LightSceneInfo;
 		ENQUEUE_RENDER_COMMAND(FAddLightCommand)(
-			[Scene, LightSceneInfo](FRHICommandListImmediate& RHICmdList)
+			[this, LightSceneInfo](FRHICommandListImmediate& RHICmdList)
 			{
 				CSV_SCOPED_TIMING_STAT_EXCLUSIVE(Scene_AddLight);
 				FScopeCycleCounter Context(LightSceneInfo->Proxy->GetStatId());
-				Scene->AddLightSceneInfo_RenderThread(LightSceneInfo);
+				bool bWasAlreadyInMap = false;
+				SceneLightInfoUpdates->Commands.FindOrAdd(LightSceneInfo, FUpdateLightCommand(FUpdateLightCommand::EAddOrRemove::Add, LightSceneInfo), bWasAlreadyInMap);
+				check(!bWasAlreadyInMap);
+				++SceneLightInfoUpdates->NumAdds;
 			});
 	}
 }
@@ -3743,50 +3836,45 @@ void FScene::GetPrimitiveUniformShaderParameters_RenderThread(const FPrimitiveSc
 	SingleCaptureIndex = PrimitiveSceneInfo->CachedReflectionCaptureProxy ? PrimitiveSceneInfo->CachedReflectionCaptureProxy->SortedCaptureIndex : 0;
 }
 
-struct FUpdateLightTransformParameters
+void FScene::UpdateLightTransform_RenderThread(int32 LightId, FLightSceneInfo* LightSceneInfo, const FUpdateLightCommand::FTransformParameters& Parameters)
 {
-	FMatrix LightToWorld;
-	FVector4 Position;
-};
-
-void FScene::UpdateLightTransform_RenderThread(FLightSceneInfo* LightSceneInfo, const FUpdateLightTransformParameters& Parameters)
-{
-	SCOPE_CYCLE_COUNTER(STAT_UpdateSceneLightTime);
 	SCOPED_NAMED_EVENT(FScene_UpdateLightTransform_RenderThread, FColor::Yellow);
 
-	if( LightSceneInfo && LightSceneInfo->bVisible )
-	{
-		// Don't remove directional lights when their transform changes as nothing in RemoveFromScene() depends on their transform
-		bool bRemove = !(LightSceneInfo->Proxy->GetLightType() == LightType_Directional);
-		bool bHasId = LightSceneInfo->Id != INDEX_NONE;
+	// This is called without a valid ID when the update is fused with an 'add' command (saves redundant scene updates to do the update first)
+	const bool bHasId = LightId != INDEX_NONE;
+	// Don't remove directional lights when their transform changes as nothing in RemoveFromScene() depends on their transform
+	const bool bRemove = bHasId && (Lights[LightId].LightType != LightType_Directional);
 
-		if( bRemove )
+	if (bHasId)
+	{
+		if (bRemove)
 		{
 			// Remove the light from the scene.
 			LightSceneInfo->RemoveFromScene();
 		}
+	}
 
-		// Invalidate the path tracer if the transform actually changed
-		// NOTE: Position is derived from the Matrix, so there is no need to check it separately
-		if( !Parameters.LightToWorld.Equals(LightSceneInfo->Proxy->LightToWorld, SMALL_NUMBER) )
-		{
-			InvalidatePathTracedOutput();
-		}
+	// Invalidate the path tracer if the transform actually changed
+	// NOTE: Position is derived from the Matrix, so there is no need to check it separately
+	if (!Parameters.LightToWorld.Equals(LightSceneInfo->Proxy->LightToWorld, SMALL_NUMBER))
+	{
+		InvalidatePathTracedOutput();
+	}
 
-		// Update the light's transform and position.
-		LightSceneInfo->Proxy->SetTransform(Parameters.LightToWorld,Parameters.Position);
+	// Update the light's transform and position.
+	LightSceneInfo->Proxy->SetTransform(Parameters.LightToWorld, Parameters.Position);
 
-		// Also update the LightSceneInfoCompact
-		if( bHasId )
-		{
-			LightSceneInfo->Scene->Lights[LightSceneInfo->Id].Init(LightSceneInfo);
+	// Also update the LightSceneInfoCompact (if one exists)
+	if (bHasId)
+	{
+		checkSlow(Lights[LightId].LightSceneInfo == LightSceneInfo);
+		Lights[LightId].Init(LightSceneInfo);
 
-			// Don't re-add directional lights when their transform changes as nothing in AddToScene() depends on their transform
-			if( bRemove )
-			{
-				// Add the light to the scene at its new location.
-				LightSceneInfo->AddToScene();
-			}
+		// Don't re-add directional lights when their transform changes as nothing in AddToScene() depends on their transform
+		if (bRemove)
+		{
+			// Add the light to the scene at its new location.
+			LightSceneInfo->AddToScene();
 		}
 	}
 }
@@ -3795,17 +3883,21 @@ void FScene::UpdateLightTransform(ULightComponent* Light)
 {
 	if(Light->SceneProxy)
 	{
-		FUpdateLightTransformParameters Parameters;
+		FUpdateLightCommand::FTransformParameters Parameters;
 		Parameters.LightToWorld = Light->GetComponentTransform().ToMatrixNoScale();
 		Parameters.Position = Light->GetLightPosition();
-		FScene* Scene = this;
+		
 		FLightSceneInfo* LightSceneInfo = Light->SceneProxy->GetLightSceneInfo();
-		ENQUEUE_RENDER_COMMAND(UpdateLightTransform)(
-			[Scene, LightSceneInfo, Parameters](FRHICommandListImmediate& RHICmdList)
-			{
-				FScopeCycleCounter Context(LightSceneInfo->Proxy->GetStatId());
-				Scene->UpdateLightTransform_RenderThread(LightSceneInfo, Parameters);
-			});
+		if (LightSceneInfo->bVisible)
+		{
+			ENQUEUE_RENDER_COMMAND(UpdateLightTransform)(
+				[this, LightSceneInfo, Parameters](FRHICommandListImmediate& RHICmdList)
+				{
+					FScopeCycleCounter Context(LightSceneInfo->Proxy->GetStatId());
+					SceneLightInfoUpdates->Enqueue(Parameters, LightSceneInfo);
+					++SceneLightInfoUpdates->NumUpdates;
+				});
+		}
 	}
 }
 
@@ -3818,47 +3910,22 @@ void FScene::UpdateLightColorAndBrightness(ULightComponent* Light)
 {
 	if(Light->SceneProxy)
 	{
-		struct FUpdateLightColorParameters
-		{
-			FLinearColor NewColor;
-			float NewIndirectLightingScale;
-			float NewVolumetricScatteringIntensity;
-		};
-
-		FUpdateLightColorParameters NewParameters;
+		FUpdateLightCommand::FColorParameters NewParameters;
 		NewParameters.NewColor = Light->GetColoredLightBrightness();
 		NewParameters.NewIndirectLightingScale = Light->IndirectLightingIntensity;
 		NewParameters.NewVolumetricScatteringIntensity = Light->VolumetricScatteringIntensity;
 
 		FScene* Scene = this;
 		FLightSceneInfo* LightSceneInfo = Light->SceneProxy->GetLightSceneInfo();
-		ENQUEUE_RENDER_COMMAND(UpdateLightColorAndBrightness)(
-			[LightSceneInfo, Scene, NewParameters](FRHICommandListImmediate& RHICmdList)
-			{
-				if( LightSceneInfo && LightSceneInfo->bVisible )
+		if (LightSceneInfo->bVisible)
+		{
+			ENQUEUE_RENDER_COMMAND(UpdateLightColorAndBrightness)(
+				[this, LightSceneInfo, NewParameters](FRHICommandListImmediate& RHICmdList)
 				{
-					// Mobile renderer:
-					// a light with no color/intensity can cause the light to be ignored when rendering.
-					// thus, lights that change state in this way must update the draw lists.
-					Scene->bScenesPrimitivesNeedStaticMeshElementUpdate =
-						Scene->bScenesPrimitivesNeedStaticMeshElementUpdate ||
-						( Scene->GetShadingPath() == EShadingPath::Mobile 
-						&& NewParameters.NewColor.IsAlmostBlack() != LightSceneInfo->Proxy->GetColor().IsAlmostBlack() );
-
-					// Path Tracing: something about the light has changed, restart path traced accumulation
-					Scene->InvalidatePathTracedOutput();
-
-					LightSceneInfo->Proxy->SetColor(NewParameters.NewColor);
-					LightSceneInfo->Proxy->IndirectLightingScale = NewParameters.NewIndirectLightingScale;
-					LightSceneInfo->Proxy->VolumetricScatteringIntensity = NewParameters.NewVolumetricScatteringIntensity;
-
-					// Also update the LightSceneInfoCompact
-					if( LightSceneInfo->Id != INDEX_NONE )
-					{
-						Scene->Lights[ LightSceneInfo->Id ].Color = NewParameters.NewColor;
-					}
-				}
-			});
+					SceneLightInfoUpdates->Enqueue(NewParameters, LightSceneInfo);
+					++SceneLightInfoUpdates->NumUpdates;
+				});
+		}
 	}
 }
 
@@ -3867,104 +3934,99 @@ void FScene::RemoveLightSceneInfo_RenderThread(FLightSceneInfo* LightSceneInfo)
 	SCOPE_CYCLE_COUNTER(STAT_RemoveSceneLightTime);
 	SCOPED_NAMED_EVENT(FScene_RemoveLightSceneInfo_RenderThread, FColor::Red);
 
-	if (LightSceneInfo->bVisible)
-	{
-		const bool bDirectionalLight = LightSceneInfo->Proxy->GetLightType() == LightType_Directional;
+	check(LightSceneInfo->bVisible);
 
-		if (bDirectionalLight)
-		{
-			DirectionalLights.Remove(LightSceneInfo);
-		}
+	const bool bDirectionalLight = LightSceneInfo->Proxy->GetLightType() == LightType_Directional;
 
-		// check SimpleDirectionalLight
-		if (LightSceneInfo == SimpleDirectionalLight)
-		{
-			SimpleDirectionalLight = nullptr;
-		}
+	if (bDirectionalLight)
+	{
+		DirectionalLights.Remove(LightSceneInfo);
+	}
 
-		if(GetShadingPath() == EShadingPath::Mobile)
-		{
-			const bool bUseCSMForDynamicObjects = LightSceneInfo->Proxy->UseCSMForDynamicObjects();
+	// check SimpleDirectionalLight
+	if (LightSceneInfo == SimpleDirectionalLight)
+	{
+		SimpleDirectionalLight = nullptr;
+	}
+
+	if(GetShadingPath() == EShadingPath::Mobile)
+	{
+		const bool bUseCSMForDynamicObjects = LightSceneInfo->Proxy->UseCSMForDynamicObjects();
 
 #if !(UE_BUILD_SHIPPING || UE_BUILD_TEST)
-			// Tracked for disabled shader permutation warnings.
-			// Condition must match that in AddLightSceneInfo_RenderThread
-			if (LightSceneInfo->Proxy->GetLightType() == LightType_Directional && !LightSceneInfo->Proxy->HasStaticLighting())
+		// Tracked for disabled shader permutation warnings.
+		// Condition must match that in AddLightSceneInfo_RenderThread
+		if (LightSceneInfo->Proxy->GetLightType() == LightType_Directional && !LightSceneInfo->Proxy->HasStaticLighting())
+		{
+			if (LightSceneInfo->Proxy->IsMovable())
 			{
-				if (LightSceneInfo->Proxy->IsMovable())
-				{
-					NumMobileMovableDirectionalLights_RenderThread--;
-				}
-				if (bUseCSMForDynamicObjects)
-				{
-					NumMobileStaticAndCSMLights_RenderThread--;
-				}
+				NumMobileMovableDirectionalLights_RenderThread--;
 			}
+			if (bUseCSMForDynamicObjects)
+			{
+				NumMobileStaticAndCSMLights_RenderThread--;
+			}
+		}
 #endif
 
-		    // check MobileDirectionalLights
-		    for (int32 LightChannelIdx = 0; LightChannelIdx < UE_ARRAY_COUNT(MobileDirectionalLights); LightChannelIdx++)
-		    {
-			    if (LightSceneInfo == MobileDirectionalLights[LightChannelIdx])
-			    {
-				    MobileDirectionalLights[LightChannelIdx] = nullptr;
+		// check MobileDirectionalLights
+		for (int32 LightChannelIdx = 0; LightChannelIdx < UE_ARRAY_COUNT(MobileDirectionalLights); LightChannelIdx++)
+		{
+			if (LightSceneInfo == MobileDirectionalLights[LightChannelIdx])
+			{
+				MobileDirectionalLights[LightChannelIdx] = nullptr;
 
-					// find another light that could be the new MobileDirectionalLight for this channel
-					for (const FLightSceneInfoCompact& OtherLight : Lights)
+				// find another light that could be the new MobileDirectionalLight for this channel
+				for (const FLightSceneInfoCompact& OtherLight : Lights)
+				{
+					if (OtherLight.LightSceneInfo != LightSceneInfo &&
+						OtherLight.LightType == LightType_Directional &&
+						!OtherLight.bStaticLighting &&
+						GetFirstLightingChannelFromMask(OtherLight.LightSceneInfo->Proxy->GetLightingChannelMask()) == LightChannelIdx)
 					{
-						if (OtherLight.LightSceneInfo != LightSceneInfo &&
-							OtherLight.LightType == LightType_Directional &&
-							!OtherLight.bStaticLighting &&
-							GetFirstLightingChannelFromMask(OtherLight.LightSceneInfo->Proxy->GetLightingChannelMask()) == LightChannelIdx)
-						{
-							MobileDirectionalLights[LightChannelIdx] = OtherLight.LightSceneInfo;
-							break;
-						}
+						MobileDirectionalLights[LightChannelIdx] = OtherLight.LightSceneInfo;
+						break;
 					}
+				}
 
-				    // if this light is a dynamic shadowcast then we need to update the static draw lists to pick a new lightingpolicy
-					if (!LightSceneInfo->Proxy->HasStaticShadowing() || bUseCSMForDynamicObjects)
-					{
-						bScenesPrimitivesNeedStaticMeshElementUpdate = true;
-					}
-				    break;
-			    }
-		    }
+				// if this light is a dynamic shadowcast then we need to update the static draw lists to pick a new lightingpolicy
+				if (!LightSceneInfo->Proxy->HasStaticShadowing() || bUseCSMForDynamicObjects)
+				{
+					bScenesPrimitivesNeedStaticMeshElementUpdate = true;
+				}
+				break;
+			}
 		}
+	}
 
-		ProcessAtmosphereLightRemoval_RenderThread(LightSceneInfo);
-
-		// Remove the light from the scene.
-		LightSceneInfo->RemoveFromScene();
+	ProcessAtmosphereLightRemoval_RenderThread(LightSceneInfo);
 
-		// Remove the light from the lights list.
-		Lights.RemoveAt(LightSceneInfo->Id);
+	// Remove the light from the scene.
+	LightSceneInfo->RemoveFromScene();
 
-		// TODO: move this work to FShadowScene & batch the light removals
-		{
-			TArray<FVirtualShadowMapArrayCacheManager*, SceneRenderingAllocator> VirtualShadowCacheManagers;
-			GetAllVirtualShadowMapCacheManagers(VirtualShadowCacheManagers);
+	// Remove the light from the lights list.
+	Lights.RemoveAt(LightSceneInfo->Id);
 
-			for (FVirtualShadowMapArrayCacheManager* CacheManager : VirtualShadowCacheManagers)
-			{
-				CacheManager->OnLightRemoved(LightSceneInfo->Id);
-			}
-		}
+	// TODO: move this work to FShadowScene & batch the light removals
+	{
+		TArray<FVirtualShadowMapArrayCacheManager*, SceneRenderingAllocator> VirtualShadowCacheManagers;
+		GetAllVirtualShadowMapCacheManagers(VirtualShadowCacheManagers);
 
-		if (!LightSceneInfo->Proxy->HasStaticShadowing()
-			&& LightSceneInfo->Proxy->CastsDynamicShadow()
-			&& LightSceneInfo->GetDynamicShadowMapChannel() == -1)
+		for (FVirtualShadowMapArrayCacheManager* CacheManager : VirtualShadowCacheManagers)
 		{
-			OverflowingDynamicShadowedLights.Remove(LightSceneInfo->Proxy->GetOwnerNameOrLabel());
+			CacheManager->OnLightRemoved(LightSceneInfo->Id);
 		}
-
-		InvalidatePathTracedOutput();
 	}
-	else
+
+	if (!LightSceneInfo->Proxy->HasStaticShadowing()
+		&& LightSceneInfo->Proxy->CastsDynamicShadow()
+		&& LightSceneInfo->GetDynamicShadowMapChannel() == -1)
 	{
-		InvisibleLights.RemoveAt(LightSceneInfo->Id);
+		OverflowingDynamicShadowedLights.Remove(LightSceneInfo->Proxy->GetOwnerNameOrLabel());
 	}
 
+	InvalidatePathTracedOutput();
+
 	if (LightSceneInfo->Proxy->GetLightType() == LightType_Rect)
 	{
 		const FRectLightSceneProxy* RectProxy = (const FRectLightSceneProxy*)LightSceneInfo->Proxy;
@@ -3993,15 +4055,42 @@ void FScene::RemoveLight(ULightComponent* Light)
 		--NumVisibleLights_GameThread;
 
 		// Disassociate the primitive's render info.
-		Light->SceneProxy = NULL;
+		Light->SceneProxy = nullptr;
 
-		// Send a command to the rendering thread to remove the light from the scene.
-		FScene* Scene = this;
-		ENQUEUE_RENDER_COMMAND(FRemoveLightCommand)(
-			[Scene, LightSceneInfo](FRHICommandListImmediate& RHICmdList)
+		// Send a command to the rendering thread to queue the light for removal from the scene.
+		ENQUEUE_RENDER_COMMAND(FQueueRemoveLightCommand)(
+			[this, LightSceneInfo](FRHICommandListImmediate& RHICmdList)
 			{
 				FScopeCycleCounter Context(LightSceneInfo->Proxy->GetStatId());
-				Scene->RemoveLightSceneInfo_RenderThread(LightSceneInfo);
+
+				if (LightSceneInfo->bVisible)
+				{
+					auto HashElementId = SceneLightInfoUpdates->Commands.FindId(LightSceneInfo);
+					// If an Add is pending, just remove the update and fall through to delete the data.
+					if (HashElementId.IsValid() && SceneLightInfoUpdates->Commands.GetByElementId(HashElementId).Value.IsAdd())
+					{
+						SceneLightInfoUpdates->Commands.RemoveByElementId(HashElementId);
+						--SceneLightInfoUpdates->NumAdds;
+					}
+					else
+					{
+						SceneLightInfoUpdates->Enqueue(FUpdateLightCommand::EAddOrRemove::Remove, LightSceneInfo);
+						++SceneLightInfoUpdates->NumRemoves;
+						// early out to defer the deletion
+						return;
+					}
+				}
+				else
+				{
+					// There should never be updates queued for lights that are not visible
+					check(SceneLightInfoUpdates->Commands.Find(LightSceneInfo) == nullptr);
+					// The "invisible lights" are removed at once.
+					InvisibleLights.RemoveAt(LightSceneInfo->Id);
+				}
+
+				// Free the light scene info and proxy.
+				delete LightSceneInfo->Proxy;
+				delete LightSceneInfo;
 			});
 	}
 }
@@ -5140,6 +5229,124 @@ static inline bool IsPrimitiveRelevantToPathTracing(FPrimitiveSceneInfo* Primiti
 #endif
 }
 
+
+FLightSceneChangeSet FScene::UpdateAllLightSceneInfos(FRDGBuilder& GraphBuilder)
+{
+	struct FFLightSceneChangeSetAllocation
+	{
+		TArray<int32, SceneRenderingAllocator> RemovedLightIds;
+		TArray<int32, SceneRenderingAllocator> TransformUpdatedLightIds;
+		TArray<int32, SceneRenderingAllocator> ColorUpdatedLightIds;
+		TArray<int32, SceneRenderingAllocator> AddedLightIds;
+	};
+	// Allocate change set storage with graph builder lifetime such that we can safely pass it to async tasks.
+	FFLightSceneChangeSetAllocation& ChangeSet = *GraphBuilder.AllocObject<FFLightSceneChangeSetAllocation>();
+
+	// Conservative pre-size the arrays (some updates are covered by Adds).
+	ChangeSet.RemovedLightIds.Reserve(SceneLightInfoUpdates->NumRemoves);
+	ChangeSet.TransformUpdatedLightIds.Reserve(SceneLightInfoUpdates->NumUpdates);
+	ChangeSet.ColorUpdatedLightIds.Reserve(SceneLightInfoUpdates->NumUpdates);
+	ChangeSet.AddedLightIds.Reserve(SceneLightInfoUpdates->NumAdds);
+
+	// Filter out removes & updates:
+	for (const auto& Element : SceneLightInfoUpdates->Commands)
+	{
+		const FUpdateLightCommand& UpdateLightCommand = Element.Value;
+		int32 Id = UpdateLightCommand.LightSceneInfo->Id;
+		if (UpdateLightCommand.IsRemove())
+		{
+			ChangeSet.RemovedLightIds.Add(Id);
+		}
+		else if (!UpdateLightCommand.IsAdd())
+		{
+			check(Id != INDEX_NONE);
+			if (UpdateLightCommand.bHasTransform)
+			{
+				ChangeSet.TransformUpdatedLightIds.Add(Id);
+			}
+			if (UpdateLightCommand.bHasColor)
+			{
+				ChangeSet.ColorUpdatedLightIds.Add(Id);
+			}
+		}
+	}
+	// This can't access the scene light data if done async since it happens before the actual removals.
+	OnPreLigtSceneInfoUpdate.Broadcast(FLightSceneChangeSet{ ChangeSet.RemovedLightIds, TConstArrayView<int32>(), ChangeSet.TransformUpdatedLightIds, ChangeSet.ColorUpdatedLightIds });
+
+	// Batch process all light removes
+	for (int32 LightId : ChangeSet.RemovedLightIds)
+	{
+		FLightSceneInfo* LightSceneInfo = Lights[LightId].LightSceneInfo;
+		FScopeCycleCounter Context(LightSceneInfo->Proxy->GetStatId());
+		RemoveLightSceneInfo_RenderThread(LightSceneInfo);
+	}
+
+	// Process all light adds & updates
+	for (const auto &Element : SceneLightInfoUpdates->Commands)
+	{
+		const FUpdateLightCommand& UpdateLightCommand = Element.Value;
+		if (UpdateLightCommand.IsRemove())
+		{
+			continue;
+		}
+
+		FLightSceneInfo* LightSceneInfo = UpdateLightCommand.LightSceneInfo;
+		FScopeCycleCounter Context(LightSceneInfo->Proxy->GetStatId());
+
+		const int32 Id = LightSceneInfo->Id;
+		const bool bHasId = Id != INDEX_NONE;
+		check(bHasId == !UpdateLightCommand.IsAdd());
+		// Directly process updates.
+		if (UpdateLightCommand.bHasTransform)
+		{
+			SCOPE_CYCLE_COUNTER(STAT_UpdateSceneLightTime);
+
+			UpdateLightTransform_RenderThread(Id, LightSceneInfo, UpdateLightCommand.TransformParameters);
+		}
+		if (UpdateLightCommand.bHasColor)
+		{
+			SCOPE_CYCLE_COUNTER(STAT_UpdateSceneLightTime);
+
+			const FUpdateLightCommand::FColorParameters& NewParameters = UpdateLightCommand.ColorParameters;
+
+			// Mobile renderer:
+			// a light with no color/intensity can cause the light to be ignored when rendering.
+			// thus, lights that change state in this way must update the draw lists.
+			bScenesPrimitivesNeedStaticMeshElementUpdate =
+				bScenesPrimitivesNeedStaticMeshElementUpdate ||
+				(GetShadingPath() == EShadingPath::Mobile
+					&& NewParameters.NewColor.IsAlmostBlack() != LightSceneInfo->Proxy->GetColor().IsAlmostBlack());
+
+			// Path Tracing: something about the light has changed, restart path traced accumulation
+			InvalidatePathTracedOutput();
+
+			LightSceneInfo->Proxy->SetColor(NewParameters.NewColor);
+			LightSceneInfo->Proxy->IndirectLightingScale = NewParameters.NewIndirectLightingScale;
+			LightSceneInfo->Proxy->VolumetricScatteringIntensity = NewParameters.NewVolumetricScatteringIntensity;
+
+			// Also update the LightSceneInfoCompact (if it does not have an ID, it is being added)
+			if (bHasId)
+			{
+				Lights[Id].Color = NewParameters.NewColor;
+			}
+		}
+
+		// Perform Add after update, since that reduces redundant processing (e.g., Add + Move)
+		if (UpdateLightCommand.IsAdd())
+		{
+			AddLightSceneInfo_RenderThread(LightSceneInfo);
+			// Note: Id is set in AddLightSceneInfo_RenderThread so we must fetch it again
+			ChangeSet.AddedLightIds.Add(LightSceneInfo->Id);
+		}
+	}
+
+	OnPostLigtSceneInfoUpdate.Broadcast(FLightSceneChangeSet{ ChangeSet.RemovedLightIds, ChangeSet.AddedLightIds, ChangeSet.TransformUpdatedLightIds, ChangeSet.ColorUpdatedLightIds });
+
+	SceneLightInfoUpdates->Reset();
+
+	return FLightSceneChangeSet{ ChangeSet.RemovedLightIds, ChangeSet.AddedLightIds, ChangeSet.TransformUpdatedLightIds, ChangeSet.ColorUpdatedLightIds };
+}
+
 void FScene::UpdateAllPrimitiveSceneInfos(FRDGBuilder& GraphBuilder, bool bAsyncCreateLPIs)
 {
 	TRACE_CPUPROFILER_EVENT_SCOPE(Scene::UpdateAllPrimitiveSceneInfos);
@@ -5152,6 +5359,8 @@ void FScene::UpdateAllPrimitiveSceneInfos(FRDGBuilder& GraphBuilder, bool bAsync
 
 	RDG_EVENT_SCOPE(GraphBuilder, "UpdateAllPrimitiveSceneInfos");
 
+	UpdateAllLightSceneInfos(GraphBuilder);
+
 #if RHI_RAYTRACING
 	UpdateRayTracingGroupBounds_RemovePrimitives(RemovedPrimitiveSceneInfos);
 	UpdateRayTracingGroupBounds_AddPrimitives(AddedPrimitiveSceneInfos);
diff --git a/Engine/Source/Runtime/Renderer/Private/ScenePrivate.h b/Engine/Source/Runtime/Renderer/Private/ScenePrivate.h
--- a/Engine/Source/Runtime/Renderer/Private/ScenePrivate.h
+++ b/Engine/Source/Runtime/Renderer/Private/ScenePrivate.h
@@ -87,6 +87,23 @@ class FSparseVolumeTextureViewerSceneProxy;
 class FLandscapeRayTracingStateList;
 class FExponentialHeightFogSceneInfo;
 class FStaticMeshBatch;
+class FSceneLightInfoUpdates;
+/**
+ * Describes all light modifications to the scene by recording the light scene IDs.
+ * TODO: If needed, we could add a reference to the FLightUpdates (which contains the commands) since this would enable systems to consume out the delta updates as they come in.
+ *       If this is useful we must ensure FLightUpdates are kept alive as long as any async tasks might require.
+ */
+struct FLightSceneChangeSet
+{
+	// IDs of all lights before they were removed, IDs in this array may not be valid at all times when the change-set is used (depends on whether the callback site is before or after the given lights are removed from the scene).
+	TConstArrayView<int32> RemovedLightIds;
+	// IDs of all lights added to the scene, only available after all lights are added to the scene, may contain the same ID's as removed, as they may be reused.
+	TConstArrayView<int32> AddedLightIds;
+	// IDs of updated lights, does not contain any from the above, since 'add' implies the update of all aspects and 'remove' implies cancellation of all updates. 
+	// The updated arrays are not disjoint as a light may have both types of update applied.
+	TConstArrayView<int32> TransformUpdatedLightIds;
+	TConstArrayView<int32> ColorUpdatedLightIds;
+};
 
 /** Holds information about a single primitive's occlusion. */
 class FPrimitiveOcclusionHistory
@@ -3409,6 +3426,23 @@ public:
 	void LumenInvalidateSurfaceCacheForPrimitive(FPrimitiveSceneInfo* InPrimitive);
 	void LumenRemovePrimitive(FPrimitiveSceneInfo* InPrimitive, int32 PrimitiveIndex);
 
+	/**
+	 * Light scene change delegates, may be used to hook in subsystems that need to respond to light scene changes.
+	 * Note, all the light scene changes are applied _before_ all the primitive scene infos are updated.
+	 */
+	DECLARE_MULTICAST_DELEGATE_OneParam(FSceneLightSceneInfoUpdateDelegate, const FLightSceneChangeSet&);
+	/**
+	 * This delegate is invoked during the scene update phase _before_ the scene has had any light changes applied.
+	 * Thus, AddedLightIds is not valid in the change set as the added lights do not have assigned IDs yet.
+	 * IF using this to drive an async task, care must be taken as the (light) scene will be modified directly after.
+	 */
+	FSceneLightSceneInfoUpdateDelegate OnPreLigtSceneInfoUpdate;
+	/**
+	 * This delegate is invoked during the scene update phase _after_ all light changes are applied.
+	 * Thus, RemovedLightIds may contain ID's that are no longer valid or are now referencing newly added lights.
+	 * IF using this to drive an async task, the core light scene info may be used, but primitive scene updates will still be ongoing (e.g., light/primitive interactions may change).
+	 */
+	FSceneLightSceneInfoUpdateDelegate OnPostLigtSceneInfoUpdate;
 protected:
 
 private:
@@ -3469,7 +3503,7 @@ private:
 	 */
 	void RemoveLightSceneInfo_RenderThread(FLightSceneInfo* LightSceneInfo);
 
-	void UpdateLightTransform_RenderThread(FLightSceneInfo* LightSceneInfo, const struct FUpdateLightTransformParameters& Parameters);
+	void UpdateLightTransform_RenderThread(int32 LightId, FLightSceneInfo* LightSceneInfo, const struct FUpdateLightTransformParameters& Parameters);
 
 	/** 
 	* Updates the contents of the given reflection capture by rendering the scene. 
@@ -3517,6 +3551,11 @@ private:
 	void ProcessAtmosphereLightRemoval_RenderThread(FLightSceneInfo* LightSceneInfo);
 	void ProcessAtmosphereLightAddition_RenderThread(FLightSceneInfo* LightSceneInfo);
 
+	/**
+	 * Process all scene updates for lights, returns the change-set, which references arrays allocated with a RDG builder life-time.
+	 */
+	FLightSceneChangeSet UpdateAllLightSceneInfos(FRDGBuilder& GraphBuilder);
+
 private:
 #if RHI_RAYTRACING
 	void UpdateRayTracingGroupBounds_AddPrimitives(const Experimental::TRobinHoodHashSet<FPrimitiveSceneInfo*>& PrimitiveSceneInfos);
@@ -3552,6 +3591,8 @@ private:
 	Experimental::TRobinHoodHashSet<FPrimitiveSceneInfo*> RemovedPrimitiveSceneInfos;
 	Experimental::TRobinHoodHashSet<FPrimitiveSceneInfo*> DistanceFieldSceneDataUpdates;
 
+	FSceneLightInfoUpdates *SceneLightInfoUpdates;
+
 	/** 
 	 * The number of visible lights in the scene
 	 * Note: This is tracked on the game thread!
