Restore WaitForRayTracingScene dependency on scratch buffer parameters
* without this RayTracingDebug was overlapping with TLAS build leading to a crash.

diff --git a/Engine/Source/Runtime/Renderer/Private/DeferredShadingRenderer.cpp b/Engine/Source/Runtime/Renderer/Private/DeferredShadingRenderer.cpp
--- a/Engine/Source/Runtime/Renderer/Private/DeferredShadingRenderer.cpp
+++ b/Engine/Source/Runtime/Renderer/Private/DeferredShadingRenderer.cpp
@@ -2183,9 +2183,10 @@ void FDeferredShadingSceneRenderer::WaitForRayTracingScene(FRDGBuilder& GraphBui
 
 	const bool bIsPathTracing = ViewFamily.EngineShowFlags.PathTracing;
 
+	// Scratch buffer must be referenced in this pass, as it must live until the BVH build is complete.
 	FBuildAccelerationStructurePassParams* PassParams = GraphBuilder.AllocParameters<FBuildAccelerationStructurePassParams>();
-	PassParams->RayTracingSceneScratchBuffer = nullptr;
-	PassParams->DynamicGeometryScratchBuffer = nullptr;
+	PassParams->RayTracingSceneScratchBuffer = Scene->RayTracingScene.BuildScratchBuffer;
+	PassParams->DynamicGeometryScratchBuffer = DynamicGeometryScratchBuffer;
 	PassParams->LightDataPacked = bIsPathTracing ? nullptr : ReferenceView.RayTracingLightDataUniformBuffer; // accessed by FRayTracingLightingMS
 
 	if (IsNaniteEnabled())
