Updated CAS generator to use the size of the luminance texture from last frame as its initial unscaled shading rate image size, and changed combiner and fixed foveated generator to use FSceneTexturesConfig::Get().Extent rather than ViewFamily.RenderTarget.GetSizeXY(). This avoids an issue where the HMD ViewFamily has a render target matching the dimensions of the desktop preview window rather than the HMD itself when launching directly into VR with no editor.

diff --git a/Engine/Source/Runtime/Renderer/Private/VariableRateShading/ContrastAdaptiveImageGenerator.cpp b/Engine/Source/Runtime/Renderer/Private/VariableRateShading/ContrastAdaptiveImageGenerator.cpp
--- a/Engine/Source/Runtime/Renderer/Private/VariableRateShading/ContrastAdaptiveImageGenerator.cpp
+++ b/Engine/Source/Runtime/Renderer/Private/VariableRateShading/ContrastAdaptiveImageGenerator.cpp
@@ -379,9 +379,20 @@ private:
 	static FRDGTextureDesc CreateSRIDesc(const FSceneViewFamily& ViewFamily, bool bIsForDynResScaled)
 	{
 		FIntPoint TileSize = FVariableRateShadingImageManager::GetSRITileSize();
-		FIntPoint ViewTargetExtents = (bIsForDynResScaled)
-			? FSceneTexturesConfig::Get().Extent
-			: ViewFamily.RenderTarget->GetSizeXY();
+
+		FIntPoint ViewTargetExtents = FIntPoint::ZeroValue;
+		if (bIsForDynResScaled)
+		{
+			ViewTargetExtents = FSceneTexturesConfig::Get().Extent;
+		}
+		else
+		{
+			// Get initial size based on luminance texture from previous frame
+			check(ViewFamily.Views[0]->bIsViewInfo);
+			const FViewInfo* ViewInfo = static_cast<const FViewInfo*>(ViewFamily.Views[0]);
+			ViewTargetExtents = ViewInfo->PrevViewInfo.LuminanceHistory->GetDesc().Extent;
+		}
+
 		FIntPoint SRIDimensions = FMath::DivideAndRoundUp(ViewTargetExtents, TileSize);
 		return FRDGTextureDesc::Create2D(
 			SRIDimensions,
diff --git a/Engine/Source/Runtime/Renderer/Private/VariableRateShading/FixedFoveationImageGenerator.cpp b/Engine/Source/Runtime/Renderer/Private/VariableRateShading/FixedFoveationImageGenerator.cpp
--- a/Engine/Source/Runtime/Renderer/Private/VariableRateShading/FixedFoveationImageGenerator.cpp
+++ b/Engine/Source/Runtime/Renderer/Private/VariableRateShading/FixedFoveationImageGenerator.cpp
@@ -8,6 +8,7 @@
 #include "SceneView.h"
 #include "UnrealClient.h"
 #include "DataDrivenShaderPlatformInfo.h"
+#include "SceneTexturesConfig.h"
 
 /* CVar values used to control generator behavior */
 
@@ -153,7 +154,7 @@ void FFixedFoveationImageGenerator::PrepareImages(FRDGBuilder& GraphBuilder, con
 	FixedFoveationCenterX = FMath::Lerp(kFixedFoveationCenterX[0], kFixedFoveationCenterX[VRSMaxLevel], VRSAmount);
 	FixedFoveationCenterY = FMath::Lerp(kFixedFoveationCenterY[0], kFixedFoveationCenterY[VRSMaxLevel], VRSAmount);
 
-	FIntPoint Size = ViewFamily.RenderTarget->GetSizeXY();
+	FIntPoint Size = FSceneTexturesConfig::Get().Extent;
 	bool bStereoRendering = IStereoRendering::IsStereoEyeView(*ViewFamily.Views[0]) && GEngine->XRSystem.IsValid();
 
 	// Sanity check VRS tile size.
diff --git a/Engine/Source/Runtime/Renderer/Private/VariableRateShading/VariableRateShadingImageManager.cpp b/Engine/Source/Runtime/Renderer/Private/VariableRateShading/VariableRateShadingImageManager.cpp
--- a/Engine/Source/Runtime/Renderer/Private/VariableRateShading/VariableRateShadingImageManager.cpp
+++ b/Engine/Source/Runtime/Renderer/Private/VariableRateShading/VariableRateShadingImageManager.cpp
@@ -222,7 +222,7 @@ FRDGTextureRef FVariableRateShadingImageManager::CombineShadingRateImages(FRDGBu
 	}
 	else
 	{
-		FIntPoint ViewSize = ViewInfo.Family->RenderTarget->GetSizeXY();
+		FIntPoint ViewSize = FSceneTexturesConfig::Get().Extent;
 		const FIntPoint TileSize = GetSRITileSize();
 
 		// Create texture to hold shading rate image
