Moved  GPU scene light uploads off the render thread.

diff --git a/Engine/Source/Runtime/Renderer/Private/GPUScene.cpp b/Engine/Source/Runtime/Renderer/Private/GPUScene.cpp
--- a/Engine/Source/Runtime/Renderer/Private/GPUScene.cpp
+++ b/Engine/Source/Runtime/Renderer/Private/GPUScene.cpp
@@ -601,21 +601,24 @@ void FGPUScene::EndRender()
 
 void FGPUScene::UpdateGPULights(FRDGBuilder& GraphBuilder, FScene& Scene)
 {
-	TArray<FLightSceneData, SceneRenderingAllocator> LightData{};
-	LightData.SetNum(Scene.Lights.Num(), false);
+	FRDGUploadData<FLightSceneData> LightData(GraphBuilder, Scene.Lights.Num());
 
-	static const auto AllowStaticLightingVar = IConsoleManager::Get().FindTConsoleVariableDataInt(TEXT("r.AllowStaticLighting"));
-	const bool bAllowStaticLighting = (!AllowStaticLightingVar || AllowStaticLightingVar->GetValueOnRenderThread() != 0);
-
-	for (int32 Index = 0; Index < Scene.Lights.Num(); ++Index)
+	GraphBuilder.AddSetupTask([this, LightData, &Scene]
 	{
-		if (Scene.Lights.IsAllocated(Index))
+		SCOPED_NAMED_EVENT(UpdateGPUScene_Lights, FColor::Green);
+		static const auto AllowStaticLightingVar = IConsoleManager::Get().FindTConsoleVariableDataInt(TEXT("r.AllowStaticLighting"));
+		const bool bAllowStaticLighting = (!AllowStaticLightingVar || AllowStaticLightingVar->GetValueOnRenderThread() != 0);
+
+		for (int32 Index = 0; Index < Scene.Lights.Num(); ++Index)
 		{
-			InitLightData(Scene.Lights[Index], bAllowStaticLighting, LightData[Index]);
+			if (Scene.Lights.IsAllocated(Index))
+			{
+				InitLightData(Scene.Lights[Index], bAllowStaticLighting, LightData[Index]);
+			}
 		}
-	}
+	});
 
-	GraphBuilder.QueueBufferUpload<FLightSceneData>(BufferState.LightDataBuffer, LightData);
+	GraphBuilder.QueueBufferUpload<FLightSceneData>(BufferState.LightDataBuffer, LightData, ERDGInitialDataFlags::NoCopy);
 }
 
 void FGPUScene::InitLightData(const FLightSceneInfoCompact& LightInfoCompact, bool bAllowStaticLighting, FLightSceneData& DataOut)
@@ -1018,7 +1021,7 @@ void FGPUScene::UploadGeneral(FRDGBuilder& GraphBuilder, FScene& Scene, FRDGExte
 
 	GraphBuilder.AddCommandListSetupTask([&TaskContext, &UploadDataSourceAdapter, &Scene, bNaniteEnabled, bExecuteInParallel, FeatureLevel = FeatureLevel](FRHICommandListBase& RHICmdList)
 	{
-		SCOPED_NAMED_EVENT(UpdateGPUScene, FColor::Green);
+		SCOPED_NAMED_EVENT(UpdateGPUScene_Primitives, FColor::Green);
 
 		LockIfValid(RHICmdList, TaskContext.PrimitiveUploader);
 		LockIfValid(RHICmdList, TaskContext.InstancePayloadUploader);
