Added details to on screen message for the illuminance meter part of HDR visualize.

diff --git a/Engine/Source/Runtime/Renderer/Private/PostProcess/PostProcessVisualizeHDR.cpp b/Engine/Source/Runtime/Renderer/Private/PostProcess/PostProcessVisualizeHDR.cpp
--- a/Engine/Source/Runtime/Renderer/Private/PostProcess/PostProcessVisualizeHDR.cpp
+++ b/Engine/Source/Runtime/Renderer/Private/PostProcess/PostProcessVisualizeHDR.cpp
@@ -207,7 +207,8 @@ FScreenPassTexture AddVisualizeHDRPass(FRDGBuilder& GraphBuilder, const FViewInf
 			const float IlluminanceMeterTextY = (ScreenCenterY - 140.0f) / DPIScale;
 			if(IsIlluminanceMeterSupportedByView(View))
 			{
-				Canvas.DrawShadowedString(IlluminanceMeterTextX, IlluminanceMeterTextY, TEXT("Illuminance meter - over the hemisphere of the surface patch"), GetStatsFont(), FLinearColor(1, 1, 1));
+				Canvas.DrawShadowedString(IlluminanceMeterTextX, IlluminanceMeterTextY   , TEXT("Illuminance meter - over the hemisphere of the surface patch"), GetStatsFont(), FLinearColor(1, 1, 1));
+				Canvas.DrawShadowedString(IlluminanceMeterTextX, IlluminanceMeterTextY+15, TEXT("                    (Forced to be a perfect white diffuse only surface)"), GetStatsFont(), FLinearColor(1, 1, 1));
 			}
 			else
 			{
