Fixed sky atmosphere transition for dx12 and vulkan.

diff --git a/Engine/Source/Runtime/Renderer/Private/SkyAtmosphereRendering.cpp b/Engine/Source/Runtime/Renderer/Private/SkyAtmosphereRendering.cpp
--- a/Engine/Source/Runtime/Renderer/Private/SkyAtmosphereRendering.cpp
+++ b/Engine/Source/Runtime/Renderer/Private/SkyAtmosphereRendering.cpp
@@ -1173,6 +1173,7 @@ void InitSkyAtmosphereForScene(FRHICommandListImmediate& RHICmdList, FScene* Sce
 				FIntPoint(1, 1),
 				TextureLUTFormat, FClearValueBinding::None, TexCreate_None, TexCreate_ShaderResource | TexCreate_UAV, false);
 			GRenderTargetPool.FindFreeElement(RHICmdList, Desc, DistantSkyLightLutTexture, TEXT("SkyAtmosphere.DistantSkyLightLut"));
+		//	RHICmdList.Transition(FRHITransitionInfo(DistantSkyLightLutTexture->GetRHI(), ERHIAccess::Unknown, ERHIAccess::SRVMask)); // ERHIPipeline::All
 		}
 	}
 }
@@ -1512,7 +1513,7 @@ void FSceneRenderer::RenderSkyAtmosphereLookUpTables(FRDGBuilder& GraphBuilder,
 			const FIntVector NumGroups = FIntVector::DivideAndRoundUp(TextureSize, FRenderSkyViewLutCS::GroupSize);
 			FComputeShaderUtils::AddPass(GraphBuilder, RDG_EVENT_NAME("RealTimeCaptureSkyViewLut"), PassFlag, ComputeShader, PassParameters, NumGroups);
 
-			Scene->RealTimeReflectionCaptureSkyAtmosphereViewLutTexture = ConvertToExternalAccessTexture(GraphBuilder, ExternalAccessQueue, RealTimeReflectionCaptureSkyAtmosphereViewLutTexture);
+			Scene->RealTimeReflectionCaptureSkyAtmosphereViewLutTexture = ConvertToExternalAccessTexture(GraphBuilder, ExternalAccessQueue, RealTimeReflectionCaptureSkyAtmosphereViewLutTexture, ERHIAccess::SRVMask, ERHIPipeline::All);
 		}
 		else
 		{
@@ -1566,7 +1567,7 @@ void FSceneRenderer::RenderSkyAtmosphereLookUpTables(FRDGBuilder& GraphBuilder,
 			const FIntVector NumGroups = FIntVector::DivideAndRoundUp(TextureSize, FRenderCameraAerialPerspectiveVolumeCS::GroupSize);
 			FComputeShaderUtils::AddPass(GraphBuilder, RDG_EVENT_NAME("RealTimeCaptureCamera360VolumeLut"), PassFlag, ComputeShader, PassParameters, NumGroups);
 
-			Scene->RealTimeReflectionCaptureCamera360APLutTexture = ConvertToExternalAccessTexture(GraphBuilder, ExternalAccessQueue, RealTimeReflectionCaptureCamera360APLutTexture);
+			Scene->RealTimeReflectionCaptureCamera360APLutTexture = ConvertToExternalAccessTexture(GraphBuilder, ExternalAccessQueue, RealTimeReflectionCaptureCamera360APLutTexture, ERHIAccess::SRVMask, ERHIPipeline::All);
 		}
 		else
 		{
