Moved shadow initialization to an async task.

diff --git a/Engine/Source/Runtime/Renderer/Private/DynamicPrimitiveDrawing.inl b/Engine/Source/Runtime/Renderer/Private/DynamicPrimitiveDrawing.inl
--- a/Engine/Source/Runtime/Renderer/Private/DynamicPrimitiveDrawing.inl
+++ b/Engine/Source/Runtime/Renderer/Private/DynamicPrimitiveDrawing.inl
@@ -200,6 +200,7 @@ inline bool MeshBatchHasPrimitives(const FMeshBatch& Mesh)
 inline int32 FViewElementPDI::DrawMesh(const FMeshBatch& Mesh)
 {
 	// Warning: can be called from Game Thread or Rendering Thread.  Be careful what you access.
+	check(DynamicPrimitiveCollector);
 
 	if (ensure(MeshBatchHasPrimitives(Mesh)))
 	{
diff --git a/Engine/Source/Runtime/Renderer/Private/SceneRendering.h b/Engine/Source/Runtime/Renderer/Private/SceneRendering.h
--- a/Engine/Source/Runtime/Renderer/Private/SceneRendering.h
+++ b/Engine/Source/Runtime/Renderer/Private/SceneRendering.h
@@ -2258,6 +2258,10 @@ protected:
 	void FinishInitDynamicShadows(FRDGBuilder& GraphBuilder, FDynamicShadowsTaskData* TaskData, FGlobalDynamicIndexBuffer& DynamicIndexBuffer, FGlobalDynamicVertexBuffer& DynamicVertexBuffer, FGlobalDynamicReadBuffer& DynamicReadBuffer, FInstanceCullingManager& InstanceCullingManager, FRDGExternalAccessQueue& ExternalAccessQueue);
 	void InitDynamicShadows(FRDGBuilder& GraphBuilder, FGlobalDynamicIndexBuffer& DynamicIndexBuffer, FGlobalDynamicVertexBuffer& DynamicVertexBuffer, FGlobalDynamicReadBuffer& DynamicReadBuffer, FInstanceCullingManager& InstanceCullingManager, FRDGExternalAccessQueue& ExternalAccessQueue);
 
+	static void WaitForPrepareDynamicShadowsTask(FDynamicShadowsTaskData* TaskData);
+	void PrepareDynamicShadows(FDynamicShadowsTaskData& TaskData);
+	friend struct FGatherShadowPrimitivesPrepareTask;
+
 	void SetupMeshPass(FViewInfo& View, FExclusiveDepthStencil::Type BasePassDepthStencilAccess, FViewCommands& ViewCommands, FInstanceCullingManager& InstanceCullingManager);
 
 	void RenderShadowProjections(
@@ -2292,6 +2296,7 @@ protected:
 
 	/** Creates a per object projected shadow for the given interaction. */
 	void CreatePerObjectProjectedShadow(
+		FDynamicShadowsTaskData& TaskData,
 		FLightPrimitiveInteraction* Interaction,
 		bool bCreateTranslucentObjectShadow,
 		bool bCreateInsetObjectShadow,
@@ -2300,6 +2305,7 @@ protected:
 
 	/** Creates shadows for the given interaction. */
 	void SetupInteractionShadows(
+		FDynamicShadowsTaskData& TaskData,
 		FLightPrimitiveInteraction* Interaction,
 		FVisibleLightInfo& VisibleLightInfo,
 		bool bStaticSceneOnly,
@@ -2348,7 +2354,7 @@ protected:
 	* Creates a projected shadow for all primitives affected by a light.
 	* @param LightSceneInfo - The light to create a shadow for.
 	*/
-	void CreateWholeSceneProjectedShadow(FLightSceneInfo* LightSceneInfo, uint32& NumPointShadowCachesUpdatedThisFrame, uint32& NumSpotShadowCachesUpdatedThisFrame);
+	void CreateWholeSceneProjectedShadow(FDynamicShadowsTaskData& TaskData, FLightSceneInfo* LightSceneInfo, uint32& NumPointShadowCachesUpdatedThisFrame, uint32& NumSpotShadowCachesUpdatedThisFrame);
 
 	/** Updates the preshadow cache, allocating new preshadows that can fit and evicting old ones. */
 	void UpdatePreshadowCache();
@@ -2357,7 +2363,10 @@ protected:
 	static void GatherSimpleLights(const FSceneViewFamily& ViewFamily, const TArray<FViewInfo>& Views, FSimpleLightArray& SimpleLights);
 
 	/** Calculates projected shadow visibility. */
-	void InitProjectedShadowVisibility();
+	void InitProjectedShadowVisibility(FDynamicShadowsTaskData& TaskData);
+
+	/** Adds a debug shadow frustum to the views primitive draw interface. */
+	void DrawDebugShadowFrustum(FViewInfo& View, FProjectedShadowInfo& ProjectedShadowInfo);
 
 	/** Gathers dynamic mesh elements for all shadows. */
 	void GatherShadowDynamicMeshElements(FGlobalDynamicIndexBuffer& DynamicIndexBuffer, FGlobalDynamicVertexBuffer& DynamicVertexBuffer, FGlobalDynamicReadBuffer& DynamicReadBuffer, FInstanceCullingManager& InstanceCullingManager);
diff --git a/Engine/Source/Runtime/Renderer/Private/SceneVisibility.cpp b/Engine/Source/Runtime/Renderer/Private/SceneVisibility.cpp
--- a/Engine/Source/Runtime/Renderer/Private/SceneVisibility.cpp
+++ b/Engine/Source/Runtime/Renderer/Private/SceneVisibility.cpp
@@ -5276,6 +5276,8 @@ void FDeferredShadingSceneRenderer::BeginInitViews(FRDGBuilder& GraphBuilder, co
 	{
 		RHICmdList.ImmediateFlush(EImmediateFlushType::DispatchToRHIThread);
 	}
+
+	WaitForPrepareDynamicShadowsTask(CurrentDynamicShadowsTaskData);
 }
 
 template<class T>
diff --git a/Engine/Source/Runtime/Renderer/Private/ShadowRendering.h b/Engine/Source/Runtime/Renderer/Private/ShadowRendering.h
--- a/Engine/Source/Runtime/Renderer/Private/ShadowRendering.h
+++ b/Engine/Source/Runtime/Renderer/Private/ShadowRendering.h
@@ -658,7 +658,7 @@ public:
 	/**
 	 * Adds a primitive to the shadow's subject list.
 	 */
-	bool AddSubjectPrimitive(FPrimitiveSceneInfo* PrimitiveSceneInfo, TArrayView<FViewInfo> ViewArray, bool bRecordShadowSubjectForMobileShading);
+	bool AddSubjectPrimitive(FDynamicShadowsTaskData& TaskData, FPrimitiveSceneInfo* PrimitiveSceneInfo, TArrayView<FViewInfo> ViewArray, bool bRecordShadowSubjectForMobileShading);
 
 	uint64 AddSubjectPrimitive_AnyThread(
 		const FPrimitiveSceneInfoCompact& PrimitiveSceneInfoCompact,
@@ -670,6 +670,7 @@ public:
 	void PresizeSubjectPrimitiveArrays(struct FAddSubjectPrimitiveStats const& Stats);
 
 	void FinalizeAddSubjectPrimitive(
+		FDynamicShadowsTaskData& TaskData,
 		struct FAddSubjectPrimitiveOp const& Op,
 		TArrayView<FViewInfo> ViewArray,
 		struct FFinalizeAddSubjectPrimitiveContext& Context);
diff --git a/Engine/Source/Runtime/Renderer/Private/ShadowSetup.cpp b/Engine/Source/Runtime/Renderer/Private/ShadowSetup.cpp
--- a/Engine/Source/Runtime/Renderer/Private/ShadowSetup.cpp
+++ b/Engine/Source/Runtime/Renderer/Private/ShadowSetup.cpp
@@ -1146,6 +1146,82 @@ void FProjectedShadowInfo::AddCachedMeshDrawCommandsForPass(
 	}
 }
 
+inline bool AreAnyViewsStaticSceneOnly(TArrayView<FViewInfo> Views)
+{
+	for (const FViewInfo& View : Views)
+	{
+		if (View.bStaticSceneOnly)
+		{
+			return true;
+		}
+	}
+	return false;
+}
+
+typedef TArray<FAddSubjectPrimitiveOp> FShadowSubjectPrimitives;
+typedef TArray<FAddSubjectPrimitiveStats> FPerShadowGatherStats;
+
+struct FDrawDebugShadowFrustumOp
+{
+	FDrawDebugShadowFrustumOp() = default;
+	FDrawDebugShadowFrustumOp(FViewInfo& InView, FProjectedShadowInfo& InProjectedShadowInfo)
+		: View(&InView)
+		, ProjectedShadowInfo(&InProjectedShadowInfo)
+	{}
+
+	FViewInfo* View = nullptr;
+	FProjectedShadowInfo* ProjectedShadowInfo = nullptr;
+};
+
+// Common setup and working data for all GatherShadowPrimitives Tasks
+struct FDynamicShadowsTaskData : public FSceneRenderingAllocatorObject<FDynamicShadowsTaskData>
+{
+	// Common data read from all points
+	FSceneRenderer* SceneRenderer;
+	const FScene* Scene;
+	TArrayView<FViewInfo> Views;
+	const ERHIFeatureLevel::Type FeatureLevel;
+	const EShaderPlatform ShaderPlatform;
+	const bool bStaticSceneOnly;
+	const bool bRunningEarly;
+	const bool bMultithreaded;
+
+	// Generated from prepare task
+	TArray<FProjectedShadowInfo*, SceneRenderingAllocator> PreShadows;
+	TArray<FProjectedShadowInfo*, SceneRenderingAllocator> ViewDependentWholeSceneShadows;
+	TArray<FDrawDebugShadowFrustumOp, SceneRenderingAllocator> DrawDebugShadowFrustumOps;
+
+	// Generated from finalize task.
+	TArray<FPrimitiveSceneInfo*, SceneRenderingAllocator> NeedsUniformBufferUpdatePrimitives;
+	TArray<FPrimitiveSceneInfo*, SceneRenderingAllocator> NeedsUpdateStaticMeshPrimitives;
+
+	// Written from task
+	TArray<struct FGatherShadowPrimitivesPacket*, SceneRenderingAllocator> Packets;
+	FPerShadowGatherStats GatherStats;
+
+	// Used by RenderThread
+	FGraphEventRef TaskEvent;
+	FGraphEventRef PrepareTaskEvent;
+
+	FDynamicShadowsTaskData(const FDynamicShadowsTaskData&) = delete;
+
+	FDynamicShadowsTaskData(FSceneRenderer* InSceneRenderer, bool bInRunningEarly)
+		: SceneRenderer(InSceneRenderer)
+		, Scene(InSceneRenderer->Scene)
+		, Views(InSceneRenderer->Views)
+		, FeatureLevel(Scene->GetFeatureLevel())
+		, ShaderPlatform(GShaderPlatformForFeatureLevel[FeatureLevel])
+		, bStaticSceneOnly(AreAnyViewsStaticSceneOnly(Views))
+		, bRunningEarly(bInRunningEarly)
+		, bMultithreaded((FApp::ShouldUseThreadingForPerformance() || FForkProcessHelper::IsForkedMultithreadInstance()) && CVarParallelGatherShadowPrimitives.GetValueOnRenderThread() > 0)
+	{
+		if (bMultithreaded)
+		{
+			PrepareTaskEvent = FGraphEvent::CreateGraphEvent();
+		}
+	}
+};
+
 struct FAddSubjectPrimitiveOverflowedIndices
 {
 	TArray<uint16> MDCIndices;
@@ -1584,7 +1660,7 @@ bool FProjectedShadowInfo::ShouldDrawStaticMeshes_AnyThread(
 	return bDrawingStaticMeshes;
 }
 
-bool FProjectedShadowInfo::AddSubjectPrimitive(FPrimitiveSceneInfo* PrimitiveSceneInfo, TArrayView<FViewInfo> ViewArray, bool bShouldRecordShadowSubjectsForMobileShading)
+bool FProjectedShadowInfo::AddSubjectPrimitive(FDynamicShadowsTaskData& TaskData, FPrimitiveSceneInfo* PrimitiveSceneInfo, TArrayView<FViewInfo> ViewArray, bool bShouldRecordShadowSubjectsForMobileShading)
 {
 	// Ray traced shadows use the GPU managed distance field object buffers, no CPU culling should be used
 	check(!bRayTracedDistanceField);
@@ -1677,20 +1753,12 @@ bool FProjectedShadowInfo::AddSubjectPrimitive(FPrimitiveSceneInfo* PrimitiveSce
 
 			if (PrimitiveSceneInfo->NeedsUniformBufferUpdate())
 			{
-				for (const FViewInfo& CurrentView : Views)
-				{
-					// Main view visible primitives are processed on parallel tasks, updating uniform buffer them here will cause a race condition.
-					check(!CurrentView.PrimitiveVisibilityMap[PrimitiveSceneInfo->GetIndex()]);
-				}
-
-				PrimitiveSceneInfo->ConditionalUpdateUniformBuffer(FRHICommandListExecutor::GetImmediateCommandList());
+				TaskData.NeedsUniformBufferUpdatePrimitives.Emplace(PrimitiveSceneInfo);
 			}
 
 			if (PrimitiveSceneInfo->NeedsUpdateStaticMeshes())
 			{
-				// Need to defer to next InitViews, as main view visible primitives are processed on parallel tasks and calling 
-				// CacheMeshDrawCommands may resize CachedDrawLists/CachedMeshDrawCommandStateBuckets causing a crash.
-				PrimitiveSceneInfo->BeginDeferredUpdateStaticMeshesWithoutVisibilityCheck();
+				TaskData.NeedsUpdateStaticMeshPrimitives.Emplace(PrimitiveSceneInfo);
 			}
 		}
 
@@ -1988,6 +2056,7 @@ void FProjectedShadowInfo::PresizeSubjectPrimitiveArrays(const FAddSubjectPrimit
 }
 
 void FProjectedShadowInfo::FinalizeAddSubjectPrimitive(
+	FDynamicShadowsTaskData& TaskData,
 	const FAddSubjectPrimitiveOp& Op,
 	TArrayView<FViewInfo> ViewArray,
 	FFinalizeAddSubjectPrimitiveContext& Context)
@@ -2015,7 +2084,7 @@ void FProjectedShadowInfo::FinalizeAddSubjectPrimitive(
 
 	if (Result.bAddOnRenderThread)
 	{
-		if (AddSubjectPrimitive(PrimitiveSceneInfo, ViewArray, Result.bShouldRecordShadowSubjectsForMobile) && VirtualShadowMapClipmap.IsValid())
+		if (AddSubjectPrimitive(TaskData, PrimitiveSceneInfo, ViewArray, Result.bShouldRecordShadowSubjectsForMobile) && VirtualShadowMapClipmap.IsValid())
 		{
 			VirtualShadowMapClipmap->OnPrimitiveRendered(PrimitiveSceneInfo);
 		}
@@ -2096,12 +2165,12 @@ void FProjectedShadowInfo::FinalizeAddSubjectPrimitive(
 
 	if (Result.bNeedUniformBufferUpdate)
 	{
-		PrimitiveSceneInfo->ConditionalUpdateUniformBuffer(FRHICommandListExecutor::GetImmediateCommandList());
+		TaskData.NeedsUniformBufferUpdatePrimitives.Add(PrimitiveSceneInfo);
 	}
 
 	if (Result.bNeedUpdateStaticMeshes)
 	{
-		PrimitiveSceneInfo->BeginDeferredUpdateStaticMeshesWithoutVisibilityCheck();
+		TaskData.NeedsUpdateStaticMeshPrimitives.Add(PrimitiveSceneInfo);
 	}
 }
 
@@ -2648,6 +2717,7 @@ bool ShouldCreateObjectShadowForStationaryLight(const FLightSceneInfo* LightScen
 }
 
 void FSceneRenderer::SetupInteractionShadows(
+	FDynamicShadowsTaskData& TaskData,
 	FLightPrimitiveInteraction* Interaction, 
 	FVisibleLightInfo& VisibleLightInfo, 
 	bool bStaticSceneOnly,
@@ -2683,12 +2753,13 @@ void FSceneRenderer::SetupInteractionShadows(
 			&& (bCreateTranslucentObjectShadow || bCreateInsetObjectShadow || bCreateObjectShadowForStationaryLight))
 		{
 			// Create projected shadow infos
-			CreatePerObjectProjectedShadow(Interaction, bCreateTranslucentObjectShadow, bCreateInsetObjectShadow || bCreateObjectShadowForStationaryLight, ViewDependentWholeSceneShadows, PreShadows);
+			CreatePerObjectProjectedShadow(TaskData, Interaction, bCreateTranslucentObjectShadow, bCreateInsetObjectShadow || bCreateObjectShadowForStationaryLight, ViewDependentWholeSceneShadows, PreShadows);
 		}
 	}
 }
 
 void FSceneRenderer::CreatePerObjectProjectedShadow(
+	FDynamicShadowsTaskData& TaskData,
 	FLightPrimitiveInteraction* Interaction, 
 	bool bCreateTranslucentObjectShadow, 
 	bool bCreateOpaqueObjectShadow,
@@ -2958,7 +3029,7 @@ void FSceneRenderer::CreatePerObjectProjectedShadow(
 						for (int32 ChildIndex = 0, ChildCount = ShadowGroupPrimitives.Num(); ChildIndex < ChildCount; ChildIndex++)
 						{
 							FPrimitiveSceneInfo* ShadowChild = ShadowGroupPrimitives[ChildIndex];
-							ProjectedShadowInfo->AddSubjectPrimitive(ShadowChild, Views, false);
+							ProjectedShadowInfo->AddSubjectPrimitive(TaskData, ShadowChild, Views, false);
 						}
 					}
 					else if (bShadowIsPotentiallyVisibleNextFrame)
@@ -2999,7 +3070,7 @@ void FSceneRenderer::CreatePerObjectProjectedShadow(
 						for (int32 ChildIndex = 0, ChildCount = ShadowGroupPrimitives.Num(); ChildIndex < ChildCount; ChildIndex++)
 						{
 							FPrimitiveSceneInfo* ShadowChild = ShadowGroupPrimitives[ChildIndex];
-							ProjectedShadowInfo->AddSubjectPrimitive(ShadowChild, Views, false);
+							ProjectedShadowInfo->AddSubjectPrimitive(TaskData, ShadowChild, Views, false);
 						}
 					}
 					else if (bShadowIsPotentiallyVisibleNextFrame)
@@ -3528,6 +3599,7 @@ bool IntersectsConvexHulls(FLightViewFrustumConvexHulls const& ConvexHulls, FBox
 
 
 void FSceneRenderer::CreateWholeSceneProjectedShadow(
+	FDynamicShadowsTaskData& TaskData,
 	FLightSceneInfo* LightSceneInfo,
 	uint32& InOutNumPointShadowCachesUpdatedThisFrame,
 	uint32& InOutNumSpotShadowCachesUpdatedThisFrame)
@@ -3684,7 +3756,7 @@ void FSceneRenderer::CreateWholeSceneProjectedShadow(
 				}
 
 				// Consolidate repeated operation into a single site.
-				auto AddInteractingPrimitives = [bStaticSceneOnly, &Views = this->Views, FeatureLevel = this->FeatureLevel]
+				auto AddInteractingPrimitives = [bStaticSceneOnly, &Views = this->Views, FeatureLevel = this->FeatureLevel, &TaskData]
 				(FLightPrimitiveInteraction* InteractionList, 
 					FProjectedShadowInfo* ProjectedShadowInfo, 
 					bool& bContainsNaniteSubjectsOut, 
@@ -3712,7 +3784,7 @@ void FSceneRenderer::CreateWholeSceneProjectedShadow(
 								FBoxSphereBounds const& Bounds = PrimitiveSceneInfo->Proxy->GetBounds();
 								if (IntersectsConvexHulls(LightViewFrustumConvexHulls, Bounds))
 								{
-									if (ProjectedShadowInfo->AddSubjectPrimitive(Interaction->GetPrimitiveSceneInfo(), Views, false) 
+									if (ProjectedShadowInfo->AddSubjectPrimitive(TaskData, Interaction->GetPrimitiveSceneInfo(), Views, false) 
 										&& VirtualSmCacheEntry.IsValid())
 									{
 										VirtualSmCacheEntry->OnPrimitiveRendered(PrimitiveSceneInfo);
@@ -3856,7 +3928,7 @@ void FSceneRenderer::CreateWholeSceneProjectedShadow(
 	}
 }
 
-void FSceneRenderer::InitProjectedShadowVisibility()
+void FSceneRenderer::InitProjectedShadowVisibility(FDynamicShadowsTaskData& TaskData)
 {
 	SCOPE_CYCLE_COUNTER(STAT_InitProjectedShadowVisibility);
 	int32 NumBufferedFrames = FOcclusionQueryHelpers::GetNumBufferedFrames(FeatureLevel);
@@ -3948,48 +4020,9 @@ void FSceneRenderer::InitProjectedShadowVisibility()
 					{
 						bool bDrawPreshadowFrustum = CVarDrawPreshadowFrustum.GetValueOnRenderThread() != 0;
 
-						if ((ViewFamily.EngineShowFlags.ShadowFrustums)
-							&& ((bDrawPreshadowFrustum && ProjectedShadowInfo.bPreShadow) || (!bDrawPreshadowFrustum && !ProjectedShadowInfo.bPreShadow)))
+						if (ViewFamily.EngineShowFlags.ShadowFrustums && ((bDrawPreshadowFrustum && ProjectedShadowInfo.bPreShadow) || (!bDrawPreshadowFrustum && !ProjectedShadowInfo.bPreShadow)))
 						{
-							FViewElementPDI ShadowFrustumPDI(&Views[ViewIndex], nullptr, &Views[ViewIndex].DynamicPrimitiveCollector);
-							
-							if(ProjectedShadowInfo.IsWholeSceneDirectionalShadow())
-							{
-								// Get split color
-								FColor Color = FColor::White;
-								switch(ProjectedShadowInfo.CascadeSettings.ShadowSplitIndex)
-								{
-									case 0: Color = FColor::Red; break;
-									case 1: Color = FColor::Yellow; break;
-									case 2: Color = FColor::Green; break;
-									case 3: Color = FColor::Blue; break;
-								}
-
-								const FMatrix ViewMatrix = View.ViewMatrices.GetViewMatrix();
-								const FMatrix ProjectionMatrix = View.ViewMatrices.GetProjectionMatrix();
-								const FVector4 ViewOrigin = View.ViewMatrices.GetViewOrigin();
-
-								float AspectRatio = ProjectionMatrix.M[1][1] / ProjectionMatrix.M[0][0];
-								float ActualFOV = (ViewOrigin.W > 0.0f) ? FMath::Atan(1.0f / ProjectionMatrix.M[0][0]) : PI/4.0f;
-
-								float Near = ProjectedShadowInfo.CascadeSettings.SplitNear;
-								float Mid = ProjectedShadowInfo.CascadeSettings.FadePlaneOffset;
-								float Far = ProjectedShadowInfo.CascadeSettings.SplitFar;
-
-								// Camera Subfrustum
-								DrawFrustumWireframe(&ShadowFrustumPDI, (ViewMatrix * FPerspectiveMatrix(ActualFOV, AspectRatio, 1.0f, Near, Mid)).Inverse(), Color, 0);
-								DrawFrustumWireframe(&ShadowFrustumPDI, (ViewMatrix * FPerspectiveMatrix(ActualFOV, AspectRatio, 1.0f, Mid, Far)).Inverse(), FColor::White, 0);
-
-								// Subfrustum Sphere Bounds
-								//DrawWireSphere(&ShadowFrustumPDI, FTransform(ProjectedShadowInfo.ShadowBounds.Center), Color, ProjectedShadowInfo.ShadowBounds.W, 40, 0);
-
-								// Shadow Map Projection Bounds
-								DrawFrustumWireframe(&ShadowFrustumPDI, FMatrix(ProjectedShadowInfo.TranslatedWorldToClipInnerMatrix.Inverse()) * FTranslationMatrix(-ProjectedShadowInfo.PreShadowTranslation), Color, 0);
-							}
-							else
-							{
-								ProjectedShadowInfo.RenderFrustumWireframe(&ShadowFrustumPDI);
-							}
+							TaskData.DrawDebugShadowFrustumOps.Emplace(Views[ViewIndex], ProjectedShadowInfo);
 						}
 					}
 				}
@@ -4056,6 +4089,46 @@ void FSceneRenderer::InitProjectedShadowVisibility()
 #endif // !UE_BUILD_SHIPPING
 }
 
+void FSceneRenderer::DrawDebugShadowFrustum(FViewInfo& View, FProjectedShadowInfo& ProjectedShadowInfo)
+{
+	FViewElementPDI ShadowFrustumPDI(&View, nullptr, nullptr);
+
+	if (ProjectedShadowInfo.IsWholeSceneDirectionalShadow())
+	{
+		// Get split color
+		FColor Color = FColor::White;
+		switch (ProjectedShadowInfo.CascadeSettings.ShadowSplitIndex)
+		{
+		case 0: Color = FColor::Red; break;
+		case 1: Color = FColor::Yellow; break;
+		case 2: Color = FColor::Green; break;
+		case 3: Color = FColor::Blue; break;
+		}
+
+		const FMatrix ViewMatrix = View.ViewMatrices.GetViewMatrix();
+		const FMatrix ProjectionMatrix = View.ViewMatrices.GetProjectionMatrix();
+		const FVector4 ViewOrigin = View.ViewMatrices.GetViewOrigin();
+
+		float AspectRatio = ProjectionMatrix.M[1][1] / ProjectionMatrix.M[0][0];
+		float ActualFOV = (ViewOrigin.W > 0.0f) ? FMath::Atan(1.0f / ProjectionMatrix.M[0][0]) : PI / 4.0f;
+
+		float Near = ProjectedShadowInfo.CascadeSettings.SplitNear;
+		float Mid = ProjectedShadowInfo.CascadeSettings.FadePlaneOffset;
+		float Far = ProjectedShadowInfo.CascadeSettings.SplitFar;
+
+		// Camera Subfrustum
+		DrawFrustumWireframe(&ShadowFrustumPDI, (ViewMatrix * FPerspectiveMatrix(ActualFOV, AspectRatio, 1.0f, Near, Mid)).Inverse(), Color, 0);
+		DrawFrustumWireframe(&ShadowFrustumPDI, (ViewMatrix * FPerspectiveMatrix(ActualFOV, AspectRatio, 1.0f, Mid, Far)).Inverse(), FColor::White, 0);
+
+		// Shadow Map Projection Bounds
+		DrawFrustumWireframe(&ShadowFrustumPDI, FMatrix(ProjectedShadowInfo.TranslatedWorldToClipInnerMatrix.Inverse()) * FTranslationMatrix(-ProjectedShadowInfo.PreShadowTranslation), Color, 0);
+	}
+	else
+	{
+		ProjectedShadowInfo.RenderFrustumWireframe(&ShadowFrustumPDI);
+	}
+}
+
 void FSceneRenderer::GatherShadowDynamicMeshElements(FGlobalDynamicIndexBuffer& DynamicIndexBuffer, FGlobalDynamicVertexBuffer& DynamicVertexBuffer, FGlobalDynamicReadBuffer& DynamicReadBuffer, FInstanceCullingManager& InstanceCullingManager)
 {
 	TArray<const FSceneView*> ReusedViewsArray;
@@ -4118,56 +4191,6 @@ void FSceneRenderer::GatherShadowDynamicMeshElements(FGlobalDynamicIndexBuffer&
 	}
 }
 
-inline bool AreAnyViewsStaticSceneOnly(TArrayView<FViewInfo> Views)
-{
-	for (const FViewInfo& View : Views)
-	{
-		if (View.bStaticSceneOnly)
-		{
-			return true;
-		}
-	}
-	return false;
-}
-
-typedef TArray<FAddSubjectPrimitiveOp> FShadowSubjectPrimitives;
-typedef TArray<FAddSubjectPrimitiveStats> FPerShadowGatherStats;
-
-// Common setup and working data for all GatherShadowPrimitives Tasks
-struct FDynamicShadowsTaskData : public FSceneRenderingAllocatorObject<FDynamicShadowsTaskData>
-{
-	// Common data read from all points
-	const FScene* Scene;
-	TArrayView<FViewInfo> Views;
-	const ERHIFeatureLevel::Type FeatureLevel;
-	const EShaderPlatform ShaderPlatform;
-	const bool bStaticSceneOnly;
-	const bool bRunningEarly;
-	const bool bMultithreaded;
-
-	// Generated from RenderThread
-	TArray<FProjectedShadowInfo*, SceneRenderingAllocator> PreShadows;
-	TArray<FProjectedShadowInfo*, SceneRenderingAllocator> ViewDependentWholeSceneShadows;
-
-	// Written from task
-	TArray<struct FGatherShadowPrimitivesPacket*, SceneRenderingAllocator> Packets;
-	FPerShadowGatherStats GatherStats;
-
-	// Used by RenderThread
-	FGraphEventRef TaskEvent;
-
-	FDynamicShadowsTaskData(const FScene* InScene, TArray<FViewInfo>& InViews, bool bInRunningEarly)
-		: Scene(InScene)
-		, Views(InViews)
-		, FeatureLevel(InScene->GetFeatureLevel())
-		, ShaderPlatform(GShaderPlatformForFeatureLevel[FeatureLevel])
-		, bStaticSceneOnly(AreAnyViewsStaticSceneOnly(InViews))
-		, bRunningEarly(bInRunningEarly)
-		, bMultithreaded(FApp::ShouldUseThreadingForPerformance() || FForkProcessHelper::IsForkedMultithreadInstance())
-	{
-	}
-};
-
 FAutoConsoleTaskPriority CPrio_GatherShadowPrimitives(
 	TEXT("TaskGraph.TaskPriorities.GatherShadowPrimitives"),
 	TEXT("Task and thread priority for GatherShadowPrimitives tasks."),
@@ -4426,8 +4449,13 @@ struct FGatherShadowPrimitivesPacket
 		}
 	}
 
-	void RenderThreadFinalize(FDynamicShadowsTaskData& TaskData)
+	void AnyThreadFinalize(FDynamicShadowsTaskData& TaskData)
 	{
+		for (int32 ShadowIdx = 0, Num = TaskData.ViewDependentWholeSceneShadows.Num(); ShadowIdx < Num; ++ShadowIdx)
+		{
+			TaskData.ViewDependentWholeSceneShadows[ShadowIdx]->PresizeSubjectPrimitiveArrays(TaskData.GatherStats[ShadowIdx]);
+		}
+
 		for (int32 ShadowIndex = 0; ShadowIndex < PreShadowSubjectPrimitives.Num(); ShadowIndex++)
 		{
 			FProjectedShadowInfo* ProjectedShadowInfo = TaskData.PreShadows[ShadowIndex];
@@ -4438,7 +4466,7 @@ struct FGatherShadowPrimitivesPacket
 
 			for (const FAddSubjectPrimitiveOp& PrimitiveOp : PreShadowSubjectPrimitives[ShadowIndex])
 			{
-				ProjectedShadowInfo->FinalizeAddSubjectPrimitive(PrimitiveOp, TaskData.Views, Context);
+				ProjectedShadowInfo->FinalizeAddSubjectPrimitive(TaskData, PrimitiveOp, TaskData.Views, Context);
 			}
 		}
 
@@ -4464,7 +4492,7 @@ struct FGatherShadowPrimitivesPacket
 
 			for (const FAddSubjectPrimitiveOp& PrimitiveOp : ViewDependentWholeSceneShadowSubjectPrimitives[ShadowIndex])
 			{
-				ProjectedShadowInfo->FinalizeAddSubjectPrimitive(PrimitiveOp, TArrayView<FViewInfo>(), Context);
+				ProjectedShadowInfo->FinalizeAddSubjectPrimitive(TaskData, PrimitiveOp, TArrayView<FViewInfo>(), Context);
 			}
 
 			if (ProjectedShadowInfo->CacheMode == SDCM_StaticPrimitivesOnly && ProjectedShadowInfo->IsWholeSceneDirectionalShadow())
@@ -4487,7 +4515,7 @@ struct FGatherShadowPrimitivesPrepareTask
 
 	static FORCEINLINE TStatId GetStatId()
 	{
-		return GET_STATID(STAT_ParallelForTask);
+		return TStatId();
 	}
 
 	static ENamedThreads::Type GetDesiredThread()
@@ -4503,11 +4531,18 @@ struct FGatherShadowPrimitivesPrepareTask
 	void DoTask(ENamedThreads::Type CurrentThread, const FGraphEventRef& MyCompletionGraphEvent)
 	{
 		SCOPED_NAMED_EVENT_TEXT("FGatherShadowPrimitivesPrepareTask", FColor::Green);
+		FOptionalTaskTagScope Scope(ETaskTag::EParallelRenderingThread);
+
+		TaskData.SceneRenderer->PrepareDynamicShadows(TaskData);
+		TaskData.PrepareTaskEvent->DispatchSubsequents();
 
 		AnyThreadTask();
 
 		// Kick off any child tasks if we added them
 
+		FGraphEventArray Prereqs;
+		Prereqs.Reserve(TaskData.Packets.Num());
+
 		if (const int32 NumPackets = TaskData.Packets.Num(); NumPackets > 0)
 		{
 			const int32 NumTasks = FMath::Min<int32>(FTaskGraphInterface::Get().GetNumWorkerThreads(), NumPackets);
@@ -4515,31 +4550,46 @@ struct FGatherShadowPrimitivesPrepareTask
 
 			for (int32 TaskIndex = 0; TaskIndex < NumTasks; TaskIndex++)
 			{
-				FDynamicShadowsTaskData* pTaskData = &TaskData;
 				const int32 FirstPacketToProcess = TaskIndex * PacketsPerTask;
 				const int32 NumPacketsToProcess = FMath::Min(PacketsPerTask, FMath::Max(0, NumPackets - FirstPacketToProcess));
 
 				if (NumPacketsToProcess > 0)
 				{
-					FGraphEventRef NewTask = FFunctionGraphTask::CreateAndDispatchWhenReady([pTaskData, FirstPacketToProcess, NumPacketsToProcess]()
+					Prereqs.Emplace(FFunctionGraphTask::CreateAndDispatchWhenReady([TaskData = &TaskData, FirstPacketToProcess, NumPacketsToProcess]()
 						{
 							SCOPED_NAMED_EVENT_TEXT("FGatherShadowPrimitivesTask", FColor::Green);
 							for (int32 PacketIndex = 0; PacketIndex < NumPacketsToProcess; PacketIndex++)
 							{
-								if (int32 PacketToProcess = FirstPacketToProcess + PacketIndex; PacketToProcess < pTaskData->Packets.Num())
+								if (int32 PacketToProcess = FirstPacketToProcess + PacketIndex; PacketToProcess < TaskData->Packets.Num())
 								{
-									pTaskData->Packets[PacketToProcess]->AnyThreadTask(*pTaskData);
+									TaskData->Packets[PacketToProcess]->AnyThreadTask(*TaskData);
 								}
 							}
 						},
-						GET_STATID(STAT_ParallelForTask),
-							nullptr,
-							CPrio_GatherShadowPrimitives.Get());
-
-					MyCompletionGraphEvent->DontCompleteUntil(NewTask);
+						TStatId(),
+						nullptr,
+						CPrio_GatherShadowPrimitives.Get()));
 				}
 			}
 		}
+
+		if (Prereqs.Num() > 0)
+		{
+			FGraphEventRef FinalizeTask = FFunctionGraphTask::CreateAndDispatchWhenReady([TaskData = &TaskData]()
+				{
+					SCOPED_NAMED_EVENT_TEXT("FGatherShadowPrimitivesFinalizeTask", FColor::Green);
+					for (FGatherShadowPrimitivesPacket* Packet : TaskData->Packets)
+					{
+						Packet->AnyThreadFinalize(*TaskData);
+						delete Packet;
+					}
+				},
+				TStatId(),
+				&Prereqs,
+				CPrio_GatherShadowPrimitives.Get());
+
+			MyCompletionGraphEvent->DontCompleteUntil(FinalizeTask);
+		}
 	}
 
 	void AddSubTask(FScenePrimitiveOctree::FNodeIndex NodeIndex, int32 StartPrimitiveIndex, int32 NumPrimitives)
@@ -4622,58 +4672,77 @@ void FSceneRenderer::BeginGatherShadowPrimitives(FDynamicShadowsTaskData* TaskDa
 {
 	SCOPE_CYCLE_COUNTER(STAT_GatherShadowPrimitivesTime);
 
-	if (TaskData && (TaskData->PreShadows.Num() || TaskData->ViewDependentWholeSceneShadows.Num()))
+	check(TaskData);
+
+	if (!TaskData->bMultithreaded || !TaskData->bRunningEarly)
 	{
-		if (!TaskData->bMultithreaded)
+		PrepareDynamicShadows(*TaskData);
+
+		if (TaskData->PreShadows.Num() || TaskData->ViewDependentWholeSceneShadows.Num())
 		{
 			FGatherShadowPrimitivesPrepareTask(*TaskData).AnyThreadTask();
-			for (FGatherShadowPrimitivesPacket* Packet : TaskData->Packets)
+
+			if (!TaskData->bMultithreaded)
 			{
-				Packet->AnyThreadTask(*TaskData);
-			}
-		}
-		else if (!TaskData->bRunningEarly)
-		{
-			FGatherShadowPrimitivesPrepareTask(*TaskData).AnyThreadTask();
-			ParallelFor(TaskData->Packets.Num(),
-				[TaskData](int32 Index)
+				for (FGatherShadowPrimitivesPacket* Packet : TaskData->Packets)
 				{
-					TaskData->Packets[Index]->AnyThreadTask(*TaskData);
-				},
-				!(FApp::ShouldUseThreadingForPerformance() && CVarParallelGatherShadowPrimitives.GetValueOnRenderThread() > 0));
-		}
-		else
-		{
-			TaskData->TaskEvent = TGraphTask<FGatherShadowPrimitivesPrepareTask>::CreateTask().ConstructAndDispatchWhenReady(*TaskData);
+					Packet->AnyThreadTask(*TaskData);
+				}
+			}
+			else // !TaskData->bRunningEarly
+			{
+				ParallelFor(TaskData->Packets.Num(),
+					[TaskData](int32 Index)
+					{
+						TaskData->Packets[Index]->AnyThreadTask(*TaskData);
+					},
+					!(FApp::ShouldUseThreadingForPerformance() && CVarParallelGatherShadowPrimitives.GetValueOnRenderThread() > 0));
+			}
 		}
 	}
+	else
+	{
+		TaskData->TaskEvent = TGraphTask<FGatherShadowPrimitivesPrepareTask>::CreateTask().ConstructAndDispatchWhenReady(*TaskData);
+	}
 }
 
 void FSceneRenderer::FinishGatherShadowPrimitives(FDynamicShadowsTaskData* TaskData)
 {
-	if (TaskData)
-	{
-		SCOPED_NAMED_EVENT_TEXT("FSceneRenderer::FinishGatherShadowPrimitives", FColor::Green);
-		QUICK_SCOPE_CYCLE_COUNTER(STAT_RenderThreadFinalize);
-
-		if (TaskData->TaskEvent)
-		{
-			FTaskGraphInterface::Get().WaitUntilTaskCompletes(TaskData->TaskEvent, ENamedThreads::GetRenderThread_Local());
-		}
+	SCOPED_NAMED_EVENT_TEXT("FSceneRenderer::FinishGatherShadowPrimitives", FColor::Green);
+	QUICK_SCOPE_CYCLE_COUNTER(STAT_RenderThreadFinalize);
 
-		for (int32 ShadowIdx = 0, Num = TaskData->ViewDependentWholeSceneShadows.Num(); ShadowIdx < Num; ++ShadowIdx)
-		{
-			TaskData->ViewDependentWholeSceneShadows[ShadowIdx]->PresizeSubjectPrimitiveArrays(TaskData->GatherStats[ShadowIdx]);
-		}
+	check(TaskData);
 
+	if (TaskData->TaskEvent)
+	{
+		FTaskGraphInterface::Get().WaitUntilTaskCompletes(TaskData->TaskEvent, ENamedThreads::GetRenderThread_Local());
+	}
+	else
+	{
+		SCOPED_NAMED_EVENT_TEXT("FGatherShadowPrimitivesTask", FColor::Green);
 		for (FGatherShadowPrimitivesPacket* Packet : TaskData->Packets)
 		{
-			Packet->RenderThreadFinalize(*TaskData);
+			Packet->AnyThreadFinalize(*TaskData);
 			delete Packet;
 		}
+	}
 
-		delete TaskData;
+	for (FPrimitiveSceneInfo* PrimitiveSceneInfo : TaskData->NeedsUniformBufferUpdatePrimitives)
+	{
+		PrimitiveSceneInfo->ConditionalUpdateUniformBuffer(FRHICommandListExecutor::GetImmediateCommandList());
+	}
+
+	for (FPrimitiveSceneInfo* PrimitiveSceneInfo : TaskData->NeedsUpdateStaticMeshPrimitives)
+	{
+		PrimitiveSceneInfo->BeginDeferredUpdateStaticMeshesWithoutVisibilityCheck();
 	}
+
+	for (FDrawDebugShadowFrustumOp Op : TaskData->DrawDebugShadowFrustumOps)
+	{
+		DrawDebugShadowFrustum(*Op.View, *Op.ProjectedShadowInfo);
+	}
+
+	delete TaskData;
 }
 
 static bool NeedsUnatlasedCSMDepthsWorkaround(ERHIFeatureLevel::Type FeatureLevel)
@@ -5638,12 +5707,8 @@ void FSceneRenderer::AllocateTranslucentShadowDepthTargets(FRHICommandListImmedi
 	}
 }
 
-FDynamicShadowsTaskData* FSceneRenderer::BeginInitDynamicShadows(bool bRunningEarly)
+void FSceneRenderer::PrepareDynamicShadows(FDynamicShadowsTaskData& TaskData)
 {
-	SCOPE_CYCLE_COUNTER(STAT_DynamicShadowSetupTime);
-	CSV_SCOPED_TIMING_STAT_EXCLUSIVE(InitViews_Shadows);
-	SCOPED_NAMED_EVENT_TEXT("FSceneRenderer::BeginInitDynamicShadows", FColor::Magenta);
-
 	const bool bMobile = FeatureLevel < ERHIFeatureLevel::SM5;
 	const bool bHairStrands = HairStrands::HasHairInstanceInScene(*Scene);
 
@@ -5656,12 +5721,11 @@ FDynamicShadowsTaskData* FSceneRenderer::BeginInitDynamicShadows(bool bRunningEa
 	uint32 NumSpotShadowCachesUpdatedThisFrame = 0;
 	uint32 NumCSMCachesUpdatedThisFrame = 0;
 
-	FDynamicShadowsTaskData* DynamicShadowTaskData = new FDynamicShadowsTaskData(Scene, Views, bRunningEarly);
-
-	const bool bStaticSceneOnly = DynamicShadowTaskData->bStaticSceneOnly;
-	TArray<FProjectedShadowInfo*, SceneRenderingAllocator>& PreShadows = DynamicShadowTaskData->PreShadows;
+	const bool bStaticSceneOnly = TaskData.bStaticSceneOnly;
+	TArray<FProjectedShadowInfo*, SceneRenderingAllocator>& PreShadows = TaskData.PreShadows;
 	TArray<FProjectedShadowInfo*, SceneRenderingAllocator> ViewDependentWholeSceneShadows;
-	TArray<FProjectedShadowInfo*, SceneRenderingAllocator>& ViewDependentWholeSceneShadowsThatNeedCulling = DynamicShadowTaskData->ViewDependentWholeSceneShadows;
+	TArray<FProjectedShadowInfo*, SceneRenderingAllocator>& ViewDependentWholeSceneShadowsThatNeedCulling = TaskData.ViewDependentWholeSceneShadows;
+
 	{
 		SCOPE_CYCLE_COUNTER(STAT_InitDynamicShadowsTime);
 		CSV_SCOPED_TIMING_STAT_EXCLUSIVE(ShadowInitDynamic);
@@ -5761,7 +5825,7 @@ FDynamicShadowsTaskData* FSceneRenderer::BeginInitDynamicShadows(bool bRunningEa
 						if (bCreateShadowForMovableLight || bCreateShadowToPreviewStaticLight || bCreateShadowForOverflowStaticShadowing)
 						{
 							// Try to create a whole scene projected shadow.
-							CreateWholeSceneProjectedShadow(LightSceneInfo, NumPointShadowCachesUpdatedThisFrame, NumSpotShadowCachesUpdatedThisFrame);
+							CreateWholeSceneProjectedShadow(TaskData, LightSceneInfo, NumPointShadowCachesUpdatedThisFrame, NumSpotShadowCachesUpdatedThisFrame);
 						}
 
 						// Register visible lights for allowing hair strands to cast shadow (non-directional light)
@@ -5791,7 +5855,7 @@ FDynamicShadowsTaskData* FSceneRenderer::BeginInitDynamicShadows(bool bRunningEa
 									const int32 NumPrims = InteractionShadowPrimitives->Num();
 									for (int32 Idx = 0; Idx < NumPrims; ++Idx)
 									{
-										SetupInteractionShadows((*InteractionShadowPrimitives)[Idx], VisibleLightInfo, bStaticSceneOnly, ViewDependentWholeSceneShadows, PreShadows);
+										SetupInteractionShadows(TaskData, (*InteractionShadowPrimitives)[Idx], VisibleLightInfo, bStaticSceneOnly, ViewDependentWholeSceneShadows, PreShadows);
 									}
 								}
 								else
@@ -5802,7 +5866,7 @@ FDynamicShadowsTaskData* FSceneRenderer::BeginInitDynamicShadows(bool bRunningEa
 										Interaction = Interaction->GetNextPrimitive()
 										)
 									{
-										SetupInteractionShadows(Interaction, VisibleLightInfo, bStaticSceneOnly, ViewDependentWholeSceneShadows, PreShadows);
+										SetupInteractionShadows(TaskData, Interaction, VisibleLightInfo, bStaticSceneOnly, ViewDependentWholeSceneShadows, PreShadows);
 									}
 
 									for (FLightPrimitiveInteraction* Interaction = LightSceneInfo->GetDynamicInteractionStaticPrimitiveList();
@@ -5810,7 +5874,7 @@ FDynamicShadowsTaskData* FSceneRenderer::BeginInitDynamicShadows(bool bRunningEa
 										Interaction = Interaction->GetNextPrimitive()
 										)
 									{
-										SetupInteractionShadows(Interaction, VisibleLightInfo, bStaticSceneOnly, ViewDependentWholeSceneShadows, PreShadows);
+										SetupInteractionShadows(TaskData, Interaction, VisibleLightInfo, bStaticSceneOnly, ViewDependentWholeSceneShadows, PreShadows);
 									}
 								}
 							}
@@ -5824,14 +5888,34 @@ FDynamicShadowsTaskData* FSceneRenderer::BeginInitDynamicShadows(bool bRunningEa
 		CSV_CUSTOM_STAT_GLOBAL(ShadowCacheUsageMB, (float(Scene->GetCachedWholeSceneShadowMapsSize()) / 1024) / 1024, ECsvCustomStatOp::Set);
 
 		// Calculate visibility of the projected shadows.
-		InitProjectedShadowVisibility();
+		InitProjectedShadowVisibility(TaskData);
 	}
 
-	//SetupFallbackShadowMaps();
-
 	// Clear old preshadows and attempt to add new ones to the cache
 	UpdatePreshadowCache();
 
+	if (ShadowSceneRenderer)
+	{
+		ShadowSceneRenderer->PostInitDynamicShadowsSetup();
+	}
+}
+
+void FSceneRenderer::WaitForPrepareDynamicShadowsTask(FDynamicShadowsTaskData* DynamicShadowTaskData)
+{
+	if (DynamicShadowTaskData && DynamicShadowTaskData->PrepareTaskEvent)
+	{
+		DynamicShadowTaskData->PrepareTaskEvent->Wait(ENamedThreads::GetRenderThread_Local());
+	}
+}
+
+FDynamicShadowsTaskData* FSceneRenderer::BeginInitDynamicShadows(bool bRunningEarly)
+{
+	SCOPE_CYCLE_COUNTER(STAT_DynamicShadowSetupTime);
+	CSV_SCOPED_TIMING_STAT_EXCLUSIVE(InitViews_Shadows);
+	SCOPED_NAMED_EVENT_TEXT("FSceneRenderer::BeginInitDynamicShadows", FColor::Magenta);
+
+	FDynamicShadowsTaskData* DynamicShadowTaskData = new FDynamicShadowsTaskData(this, bRunningEarly);
+
 	// Gathers the list of primitives used to draw various shadow types
 	BeginGatherShadowPrimitives(DynamicShadowTaskData);
 
@@ -5848,11 +5932,11 @@ void FSceneRenderer::FinishInitDynamicShadows(FRDGBuilder& GraphBuilder, FDynami
 
 	FinishGatherShadowPrimitives(TaskData);
 
-	// TODO: Move this to the deferred scene renderer
 	if (ShadowSceneRenderer)
 	{
-		ShadowSceneRenderer->PostInitDynamicShadowsSetup();
+		ShadowSceneRenderer->PostSetupDebugRender();
 	}
+
 	AllocateShadowDepthTargets(GraphBuilder.RHICmdList);
 
 	// Generate mesh element arrays from shadow primitive arrays
diff --git a/Engine/Source/Runtime/Renderer/Private/Shadows/ShadowSceneRenderer.cpp b/Engine/Source/Runtime/Renderer/Private/Shadows/ShadowSceneRenderer.cpp
--- a/Engine/Source/Runtime/Renderer/Private/Shadows/ShadowSceneRenderer.cpp
+++ b/Engine/Source/Runtime/Renderer/Private/Shadows/ShadowSceneRenderer.cpp
@@ -197,7 +197,6 @@ void FShadowSceneRenderer::PostInitDynamicShadowsSetup()
 			}
 		}
 	}
-	PostSetupDebugRender();
 }
 
 void FShadowSceneRenderer::RenderVirtualShadowMaps(FRDGBuilder& GraphBuilder, bool bNaniteEnabled, bool bUpdateNaniteStreaming, bool bNaniteProgrammableRaster)
@@ -276,7 +275,7 @@ void FShadowSceneRenderer::PostSetupDebugRender()
 		int32 NumDistant = 0;
 		for (FViewInfo& View : SceneRenderer.Views)
 		{
-			FViewElementPDI DebugPDI(&View, nullptr, &View.DynamicPrimitiveCollector);
+			FViewElementPDI DebugPDI(&View, nullptr, nullptr);
 
 			for (const FLocalLightShadowFrameSetup& LightSetup : LocalLights)
 			{			
