Proper fix for Nanite uniform buffer crash, allowing multi-view rendering optimizations to be re-enabled (reversing CL 23287608, which disabled those optimizations).  Saves mulitple ms CPU potentially.

Need to ensure pipeline state cache references (FGraphicsMinimalPipelineStateId) aren't released by another thread while still in use by parallel task threads (FSubmitNaniteMaterialPassCommandsAnyThreadTask or FDrawVisibleMeshCommandsAnyThreadTask).  This can happen during calls to FPrimitiveSceneInfo::UpdateStaticMeshes, which calls RemoveCachedMeshDrawCommands followed by CacheMeshDrawCommands (or the Nanite equivalent) for a set of proxies to rebuild their mesh commands.  If any persistent pipeline state ID goes to zero reference count during that process, it will be released and the cached ID will change when it gets reallocated.

A utility function freezes the pipeline state ID table during FSceneRenderer::ComputeViewVisibility (which is what calls UpdateStaticMeshes) for scene renderers besides the first, making it so persistent pipeline state IDs aren't freed even if their reference count goes to zero.  The same pipeline states end up being used in the draw commands for a render proxy following an update, given that the scene renderers are run sequentially, and nothing can run between them to change any state.  If any zero reference items remain when the table is unfrozen, they will be cleaned up anyway, just in case that ever happens.

#rnx

diff --git a/Engine/Source/Runtime/Renderer/Private/MeshPassProcessor.cpp b/Engine/Source/Runtime/Renderer/Private/MeshPassProcessor.cpp
--- a/Engine/Source/Runtime/Renderer/Private/MeshPassProcessor.cpp
+++ b/Engine/Source/Runtime/Renderer/Private/MeshPassProcessor.cpp
@@ -24,11 +24,14 @@ FRWLock FGraphicsMinimalPipelineStateId::PersistentIdTableLock;
 FGraphicsMinimalPipelineStateId::PersistentTableType FGraphicsMinimalPipelineStateId::PersistentIdTable;
 
 #if MESH_DRAW_COMMAND_DEBUG_DATA
+std::atomic<int32> FGraphicsMinimalPipelineStateId::DebugSaltAllocationIndex;
 std::atomic<int32> FGraphicsMinimalPipelineStateId::LocalPipelineIdTableSize(0);
 std::atomic<int32> FGraphicsMinimalPipelineStateId::CurrentLocalPipelineIdTableSize(0);
 #endif //MESH_DRAW_COMMAND_DEBUG_DATA
 
 bool FGraphicsMinimalPipelineStateId::NeedsShaderInitialisation = true;
+bool FGraphicsMinimalPipelineStateId::bIsIdTableFrozen = false;
+std::atomic<int32> FGraphicsMinimalPipelineStateId::ReffedItemCount(0);
 
 const FMeshDrawCommandSortKey FMeshDrawCommandSortKey::Default = { {0} };
 
@@ -474,6 +477,9 @@ void FMeshDrawShaderBindings::SetRayTracingShaderBindingsForMissShader(
 FGraphicsMinimalPipelineStateId FGraphicsMinimalPipelineStateId::GetPersistentId(const FGraphicsMinimalPipelineStateInitializer& InPipelineState)
 {
 	Experimental::FHashElementId TableId;
+#if MESH_DRAW_COMMAND_DEBUG_DATA
+	int32 DebugSalt;
+#endif
 	auto hash = PersistentIdTable.ComputeHash(InPipelineState);
 	{
 		FRWScopeLock Lock(PersistentIdTableLock, SLT_ReadOnly);
@@ -496,11 +502,19 @@ FGraphicsMinimalPipelineStateId FGraphicsMinimalPipelineStateId::GetPersistentId
 		
 		FRefCountedGraphicsMinimalPipelineState& Value = PersistentIdTable.GetByElementId(TableId).Value;
 
-		if (Value.RefNum == 0 && !NeedsShaderInitialisation)
+		if (Value.RefNum == 0)
 		{
-			NeedsShaderInitialisation = true;
+			if (!NeedsShaderInitialisation)
+			{
+				NeedsShaderInitialisation = true;
+			}
+			ReffedItemCount++;
 		}
 		Value.RefNum++;
+
+#if MESH_DRAW_COMMAND_DEBUG_DATA
+		DebugSalt = Value.DebugSalt;
+#endif
 	}
 
 	checkf(TableId.GetIndex() < (MAX_uint32 >> 2), TEXT("Persistent FGraphicsMinimalPipelineStateId table overflow!"));
@@ -509,6 +523,9 @@ FGraphicsMinimalPipelineStateId FGraphicsMinimalPipelineStateId::GetPersistentId
 	Ret.bValid = 1;
 	Ret.bComesFromLocalPipelineStateSet = 0;
 	Ret.SetElementIndex = TableId.GetIndex();
+#if MESH_DRAW_COMMAND_DEBUG_DATA
+	Ret.DebugSalt = DebugSalt;
+#endif
 	return Ret;
 }
 
@@ -540,7 +557,46 @@ void FGraphicsMinimalPipelineStateId::RemovePersistentId(FGraphicsMinimalPipelin
 		--RefCountedStateInitializer.RefNum;
 		if (RefCountedStateInitializer.RefNum == 0)
 		{
-			PersistentIdTable.RemoveByElementId(Id.SetElementIndex);
+			ReffedItemCount--;
+			if (bIsIdTableFrozen == false)
+			{
+				PersistentIdTable.RemoveByElementId(Id.SetElementIndex);
+			}
+		}
+	}
+}
+
+void FGraphicsMinimalPipelineStateId::FreezeIdTable(bool bFreeze)
+{
+	if (bFreeze != bIsIdTableFrozen)
+	{
+		FRWScopeLock WriteLock(PersistentIdTableLock, SLT_Write);
+		bIsIdTableFrozen = bFreeze;
+
+		// When set back to false, do a pass through the table and clean up any zero ref items.  Zero ref items usually don't happen
+		// in practice given the context of how this function is used during calls to FPrimitiveSceneInfo::UpdateStaticMeshes where
+		// multiple scene renderers are running (i.e. nDisplay multi-view).  That function removes and re-creates draw commands for a
+		// set of render proxies, which we expect to produce the same set of pipeline cache IDs, given that no game logic can run
+		// to change any state in the render proxies between renders of the same scene.  We could assert, but cleaning up the zero
+		// reference items isn't an unreasonable option either (in case someone tries to use this function for other purposes)...
+		if (bFreeze == false)
+		{
+			// We can detect whether we need to scan the table for zero ref items, by checking whether the number of items with
+			// non-zero ref count equals the number of items in the table.
+			if (PersistentIdTable.Num() != ReffedItemCount)
+			{
+				// There's a bug in our reference counting if the reffed item count is more than the items actually in the table
+				check(PersistentIdTable.Num() > ReffedItemCount);
+
+				// Remove zero reffed items
+				for (auto PersistentIdIterator = PersistentIdTable.begin(); PersistentIdIterator != PersistentIdTable.end(); ++PersistentIdIterator)
+				{
+					if (PersistentIdTable.GetByElementId(PersistentIdIterator.GetElementId()).Value.RefNum == 0)
+					{
+						PersistentIdTable.RemoveByElementId(PersistentIdIterator.GetElementId());
+					}
+				}
+			}
 		}
 	}
 }
@@ -564,6 +620,9 @@ FGraphicsMinimalPipelineStateId FGraphicsMinimalPipelineStateId::GetPipelineStat
 	checkf(TableIndex.GetIndex() < (MAX_uint32 >> 2), TEXT("One frame FGraphicsMinimalPipelineStateId table overflow!"));
 
 	Ret.SetElementIndex = TableIndex.GetIndex();
+#if MESH_DRAW_COMMAND_DEBUG_DATA
+	Ret.DebugSalt = 0;		// Salt is ignored for pipelines from local pipeline state set, just initialize to zero
+#endif
 	return Ret;
 }
 
diff --git a/Engine/Source/Runtime/Renderer/Private/SceneRendering.cpp b/Engine/Source/Runtime/Renderer/Private/SceneRendering.cpp
--- a/Engine/Source/Runtime/Renderer/Private/SceneRendering.cpp
+++ b/Engine/Source/Runtime/Renderer/Private/SceneRendering.cpp
@@ -4679,20 +4679,8 @@ void FRendererModule::BeginRenderingViewFamilies(FCanvas* Canvas, TArrayView<FSc
 				uint64 SceneRenderStart = FPlatformTime::Cycles64();
 				const float StartDelayMillisec = FPlatformTime::ToMilliseconds64(SceneRenderStart - DrawSceneEnqueue);
 				CSV_CUSTOM_STAT_GLOBAL(DrawSceneCommand_StartDelay, StartDelayMillisec, ECsvCustomStatOp::Set);
-
-				// TODO:  There were some random uniform buffer crashes that resulted when rendering multiple view families
-				// in the function below.  Pass one at a time for now as a workaround (matches original Release 5.0 behavior),
-				// while we investigate.  The theory is that something in FSceneRenderer::RenderThreadBegin/End needs to run
-				// between scene renders to reset some state, but we don't know what.
-				TArray<FSceneRenderer*> SingleSceneRenderer;
-				SingleSceneRenderer.AddZeroed(1);
-
-				for (FSceneRenderer* SceneRenderer : LocalSceneRenderers)
-				{
-					SingleSceneRenderer[0] = SceneRenderer;
-					RenderViewFamilies_RenderThread(RHICmdList, SingleSceneRenderer);
-					FlushPendingDeleteRHIResources_RenderThread();
-				}
+				RenderViewFamilies_RenderThread(RHICmdList, LocalSceneRenderers);
+				FlushPendingDeleteRHIResources_RenderThread();
 			});
 
 		// Force kick the RT if we've got RT polling on.
diff --git a/Engine/Source/Runtime/Renderer/Private/SceneVisibility.cpp b/Engine/Source/Runtime/Renderer/Private/SceneVisibility.cpp
--- a/Engine/Source/Runtime/Renderer/Private/SceneVisibility.cpp
+++ b/Engine/Source/Runtime/Renderer/Private/SceneVisibility.cpp
@@ -4615,6 +4615,22 @@ void FSceneRenderer::ComputeViewVisibility(
 	STAT(int32 NumCulledPrimitives = 0);
 	STAT(int32 NumOccludedPrimitives = 0);
 
+	/**
+	  * UpdateStaticMeshes removes and re-creates cached FMeshDrawCommands.  If there are multiple scene renderers being run together,
+	  * we need allocated pipeline state IDs not to change, in case async tasks related to prior scene renderers are still in flight
+	  * (FSubmitNaniteMaterialPassCommandsAnyThreadTask or FDrawVisibleMeshCommandsAnyThreadTask).  So we freeze pipeline state IDs,
+	  * preventing them from being de-allocated even if their reference count temporarily goes to zero during calls to
+	  * RemoveCachedMeshDrawCommands followed by CacheMeshDrawCommands (or the Nanite equivalent).
+	  *
+	  * Note that on the first scene renderer, we do want to de-allocate items, so they can be permanently released if no longer in use
+	  * (for example, if there was an impactful change to a render proxy by game logic), but the assumption is that sequential renders
+	  * of the same scene from different views can't make such changes.
+	  */
+	if (!bIsFirstSceneRenderer)
+	{
+		FGraphicsMinimalPipelineStateId::FreezeIdTable(true);
+	}
+
 	UE::Tasks::FTask ComputeLightVisibilityTask = LaunchSceneRenderTask(UE_SOURCE_LOCATION, [this]
 	{
 		ComputeLightVisibility();
@@ -4938,6 +4954,12 @@ void FSceneRenderer::ComputeViewVisibility(
 	INC_DWORD_STAT_BY(STAT_ProcessedPrimitives,NumProcessedPrimitives);
 	INC_DWORD_STAT_BY(STAT_CulledPrimitives,NumCulledPrimitives);
 	INC_DWORD_STAT_BY(STAT_OccludedPrimitives,NumOccludedPrimitives);
+
+	// See comment where this is called above
+	if (!bIsFirstSceneRenderer)
+	{
+		FGraphicsMinimalPipelineStateId::FreezeIdTable(false);
+	}
 }
 
 void FDeferredShadingSceneRenderer::ComputeLightVisibility()
diff --git a/Engine/Source/Runtime/Renderer/Public/MeshPassProcessor.h b/Engine/Source/Runtime/Renderer/Public/MeshPassProcessor.h
--- a/Engine/Source/Runtime/Renderer/Public/MeshPassProcessor.h
+++ b/Engine/Source/Runtime/Renderer/Public/MeshPassProcessor.h
@@ -534,7 +534,11 @@ public:
 
 		{
 			FRWScopeLock ReadLock(PersistentIdTableLock, SLT_ReadOnly);
-			return PersistentIdTable.GetByElementId(SetElementIndex).Key;
+			const auto& Result = PersistentIdTable.GetByElementId(SetElementIndex);
+#if MESH_DRAW_COMMAND_DEBUG_DATA
+			checkf(Result.Value.RefNum && Result.Value.DebugSalt == DebugSalt, TEXT("Pipeline state ID used after release.  Call 'AddRefPersistentId' to ensure pipeline state doesn't get released while in use."));
+#endif
+			return Result.Key;
 		}
 	}
 
@@ -576,6 +580,14 @@ public:
 		return PersistentIdTable.Num();
 	}
 
+	/**
+	  * When ID table is frozen, items with a zero ref count aren't actually released, so if the same item gets used again,
+	  * its ID doesn't change.  Useful when re-creating draw commands in the middle of the frame, where you want to prevent
+	  * existing pipeline state IDs from changing, which may be referenced by in flight tasks.  Unfreezing the table will
+	  * clean up anything with a zero ref count.
+	  */
+	static void FreezeIdTable(bool bEnable);
+
 private:
 	union
 	{
@@ -593,15 +605,23 @@ private:
 	{
 		FRefCountedGraphicsMinimalPipelineState() : RefNum(0)
 		{
+#if MESH_DRAW_COMMAND_DEBUG_DATA
+			DebugSalt = DebugSaltAllocationIndex++;
+#endif
 		}
 
 		FRefCountedGraphicsMinimalPipelineState(const FRefCountedGraphicsMinimalPipelineState&& Other) :
 			RefNum(Other.RefNum.load())
 		{
-
+#if MESH_DRAW_COMMAND_DEBUG_DATA
+			DebugSalt = Other.DebugSalt;
+#endif
 		}
 
 		std::atomic<uint32> RefNum;
+#if MESH_DRAW_COMMAND_DEBUG_DATA
+		int32 DebugSalt;
+#endif
 	};
 
 	static FRWLock PersistentIdTableLock;
@@ -609,7 +629,12 @@ private:
 	static PersistentTableType PersistentIdTable;
 	static bool NeedsShaderInitialisation;
 
+	static bool bIsIdTableFrozen;
+	static std::atomic<int32> ReffedItemCount;	// Number of items with a non-zero reference count
+
 #if MESH_DRAW_COMMAND_DEBUG_DATA
+	int32 DebugSalt;
+	static std::atomic<int32> DebugSaltAllocationIndex;
 	static std::atomic<int32> LocalPipelineIdTableSize;
 	static std::atomic<int32> CurrentLocalPipelineIdTableSize;
 #endif //MESH_DRAW_COMMAND_DEBUG_DATA
