Add some perf counters and stats to shadow rendering paths.

diff --git a/Engine/Source/Runtime/Renderer/Private/LightRendering.cpp b/Engine/Source/Runtime/Renderer/Private/LightRendering.cpp
--- a/Engine/Source/Runtime/Renderer/Private/LightRendering.cpp
+++ b/Engine/Source/Runtime/Renderer/Private/LightRendering.cpp
@@ -47,6 +47,8 @@ IMPLEMENT_TYPE_LAYOUT(FShadowProjectionShaderParameters);
 
 IMPLEMENT_GLOBAL_SHADER_PARAMETER_STRUCT(FDeferredLightUniformStruct, "DeferredLightUniforms");
 
+DECLARE_DWORD_COUNTER_STAT(TEXT("VSM Light Projections (Local One Pass Fast)"), STAT_VSMLocalProjectionOnePassFast, STATGROUP_ShadowRendering);
+
 extern int32 GUseTranslucentLightingVolumes;
 
 static int32 GAllowDepthBoundsTest = 1;
@@ -1946,7 +1948,12 @@ void FDeferredShadingSceneRenderer::RenderLights(
 						const FViewInfo& View = Views[ViewIndex];
 
 						// If the light elided the screen space shadow mask, sample directly from the packed shadow mask
-						const int32 VirtualShadowMapId = bElideScreenShadowMask ? VisibleLightInfo.GetVirtualShadowMapId(&View) : INDEX_NONE;
+						int32 VirtualShadowMapId = INDEX_NONE;
+						if (bElideScreenShadowMask)
+						{
+							INC_DWORD_STAT(STAT_VSMLocalProjectionOnePassFast);
+							VirtualShadowMapId = VisibleLightInfo.GetVirtualShadowMapId(&View);
+						}
 
 						RDG_EVENT_SCOPE_CONDITIONAL(GraphBuilder, ViewCount > 1, "View%d", ViewIndex);
 						SCOPED_GPU_MASK(GraphBuilder.RHICmdList, View.GPUMask);
diff --git a/Engine/Source/Runtime/Renderer/Private/ShadowDepthRendering.cpp b/Engine/Source/Runtime/Renderer/Private/ShadowDepthRendering.cpp
--- a/Engine/Source/Runtime/Renderer/Private/ShadowDepthRendering.cpp
+++ b/Engine/Source/Runtime/Renderer/Private/ShadowDepthRendering.cpp
@@ -1499,6 +1499,8 @@ bool IsParallelDispatchEnabled(const FProjectedShadowInfo* ProjectedShadowInfo,
 
 void FSceneRenderer::RenderShadowDepthMapAtlases(FRDGBuilder& GraphBuilder)
 {
+	TRACE_CPUPROFILER_EVENT_SCOPE(FSceneRenderer::RenderShadowDepthMapAtlases);
+
 	const bool bNaniteEnabled = 
 		UseNanite(ShaderPlatform) &&
 		ViewFamily.EngineShowFlags.NaniteMeshes &&
@@ -1869,11 +1871,14 @@ void FSceneRenderer::RenderShadowDepthMaps(FRDGBuilder& GraphBuilder, FInstanceC
 	}
 
 	// Move current persistent shadow state to previous and clear current.
-	// TODO: This could be very slow.
-	for (const FLightSceneInfoCompact& Light : Scene->Lights)
 	{
-		Light.LightSceneInfo->PrevPersistentShadows = Light.LightSceneInfo->PersistentShadows;
-		Light.LightSceneInfo->PersistentShadows.Empty();
+		TRACE_CPUPROFILER_EVENT_SCOPE(CopyPersistentLightState);
+		// TODO: This could be very slow.
+		for (const FLightSceneInfoCompact& Light : Scene->Lights)
+		{
+			Light.LightSceneInfo->PrevPersistentShadows = Light.LightSceneInfo->PersistentShadows;
+			Light.LightSceneInfo->PersistentShadows.Empty();
+		}
 	}
 
 	bShadowDepthRenderCompleted = true;
diff --git a/Engine/Source/Runtime/Renderer/Private/Shadows/ShadowSceneRenderer.cpp b/Engine/Source/Runtime/Renderer/Private/Shadows/ShadowSceneRenderer.cpp
--- a/Engine/Source/Runtime/Renderer/Private/Shadows/ShadowSceneRenderer.cpp
+++ b/Engine/Source/Runtime/Renderer/Private/Shadows/ShadowSceneRenderer.cpp
@@ -67,6 +67,10 @@ DECLARE_DWORD_COUNTER_STAT(TEXT("VSM Visible Shading Draws"), STAT_VSMNaniteBass
 DECLARE_DWORD_COUNTER_STAT(TEXT("Distant Light Count"), STAT_DistantLightCount, STATGROUP_ShadowRendering);
 DECLARE_DWORD_COUNTER_STAT(TEXT("Distant Cached Count"), STAT_DistantCachedCount, STATGROUP_ShadowRendering);
 
+DECLARE_DWORD_COUNTER_STAT(TEXT("VSM Light Projections (Directional)"), STAT_VSMDirectionalProjectionFull, STATGROUP_ShadowRendering);
+DECLARE_DWORD_COUNTER_STAT(TEXT("VSM Light Projections (Local Full)"), STAT_VSMLocalProjectionFull, STATGROUP_ShadowRendering);
+DECLARE_DWORD_COUNTER_STAT(TEXT("VSM Light Projections (Local One Pass Copy)"), STAT_VSMLocalProjectionOnePassCopy, STATGROUP_ShadowRendering);
+
 FShadowSceneRenderer::FShadowSceneRenderer(FDeferredShadingSceneRenderer& InSceneRenderer)
 	: SceneRenderer(InSceneRenderer)
 	, Scene(*InSceneRenderer.Scene)
@@ -201,6 +205,8 @@ void FShadowSceneRenderer::PostInitDynamicShadowsSetup()
 
 void FShadowSceneRenderer::RenderVirtualShadowMaps(FRDGBuilder& GraphBuilder, bool bNaniteEnabled, bool bUpdateNaniteStreaming, bool bNaniteProgrammableRaster)
 {
+	TRACE_CPUPROFILER_EVENT_SCOPE(FShadowSceneRenderer::RenderVirtualShadowMaps);
+
 	// Always process an existing query if it exists
 	FNaniteVisibilityResults VisibilityResults;
 	if (NaniteVisibilityQuery != nullptr)
@@ -406,6 +412,8 @@ void FShadowSceneRenderer::ApplyVirtualShadowMapProjectionForLight(
 			{
 				if (VisibleLightInfo.VirtualShadowMapClipmaps.Num() > 0)
 				{
+					INC_DWORD_STAT(STAT_VSMDirectionalProjectionFull);
+
 					// Project directional light virtual shadow map
 					RenderVirtualShadowMapProjection(
 						GraphBuilder,
@@ -421,6 +429,8 @@ void FShadowSceneRenderer::ApplyVirtualShadowMapProjectionForLight(
 				}
 				else if (bShouldUseVirtualShadowMapOnePassProjection)
 				{
+					INC_DWORD_STAT(STAT_VSMLocalProjectionOnePassCopy);
+
 					// Copy local light from one pass projection output
 					CompositeVirtualShadowMapFromMaskBits(
 						GraphBuilder,
@@ -435,6 +445,8 @@ void FShadowSceneRenderer::ApplyVirtualShadowMapProjectionForLight(
 				}
 				else
 				{
+					INC_DWORD_STAT(STAT_VSMLocalProjectionFull);
+
 					// Project local light virtual shadow map
 					RenderVirtualShadowMapProjection(
 						GraphBuilder,
diff --git a/Engine/Source/Runtime/Renderer/Private/VirtualShadowMaps/VirtualShadowMapArray.cpp b/Engine/Source/Runtime/Renderer/Private/VirtualShadowMaps/VirtualShadowMapArray.cpp
--- a/Engine/Source/Runtime/Renderer/Private/VirtualShadowMaps/VirtualShadowMapArray.cpp
+++ b/Engine/Source/Runtime/Renderer/Private/VirtualShadowMaps/VirtualShadowMapArray.cpp
@@ -28,6 +28,8 @@ IMPLEMENT_STATIC_UNIFORM_BUFFER_STRUCT(FVirtualShadowMapUniformParameters, "Virt
 
 CSV_DEFINE_CATEGORY(VSM, false);
 
+DECLARE_DWORD_COUNTER_STAT(TEXT("VSM Nanite Views (Primary)"), STAT_VSMNaniteViewsPrimary, STATGROUP_ShadowRendering);
+
 struct FShadowMapCacheData
 {
 	int32 PrevVirtualShadowMapId = INDEX_NONE;
@@ -1970,6 +1972,8 @@ void FVirtualShadowMapArray::LogStats(FRDGBuilder& GraphBuilder, const FViewInfo
 
 void FVirtualShadowMapArray::CreateMipViews( TArray<Nanite::FPackedView, SceneRenderingAllocator>& Views ) const
 {
+	TRACE_CPUPROFILER_EVENT_SCOPE(CreateMipViews);
+
 	// strategy: 
 	// 1. Use the cull pass to generate copies of every node for every view needed.
 	// [2. Fabricate a HZB array?]
@@ -2529,6 +2533,9 @@ void FVirtualShadowMapArray::RenderVirtualShadowMapsNanite(FRDGBuilder& GraphBui
 	bCsvLogEnabled = FCsvProfiler::Get()->IsCapturing_Renderthread() && FCsvProfiler::Get()->IsCategoryEnabled(CSV_CATEGORY_INDEX(VSM));
 #endif
 
+	TRACE_CPUPROFILER_EVENT_SCOPE(FVirtualShadowMapArray::RenderVirtualShadowMapsNanite);
+	RDG_EVENT_SCOPE(GraphBuilder, "RenderVirtualShadowMaps(Nanite)");
+
 	const TArray<FProjectedShadowInfo*, SceneRenderingAllocator>& VirtualShadowMapShadows = SceneRenderer.SortedShadowsForShadowDepthPass.VirtualShadowMapShadows;
 	TArray<TSharedPtr<FVirtualShadowMapClipmap>, SceneRenderingAllocator>& VirtualShadowMapClipmaps = SceneRenderer.SortedShadowsForShadowDepthPass.VirtualShadowMapClipmaps;
 
@@ -2545,8 +2552,6 @@ void FVirtualShadowMapArray::RenderVirtualShadowMapsNanite(FRDGBuilder& GraphBui
 
 	const TRefCountPtr<IPooledRenderTarget> PrevHZBPhysical = bVSMUseHZB ? CacheManager->PrevBuffers.HZBPhysical : nullptr;
 
-	RDG_EVENT_SCOPE(GraphBuilder, "RenderVirtualShadowMaps(Nanite)");
-
 	check(PhysicalPagePoolRDG != nullptr);
 
 	Nanite::FRasterContext RasterContext = Nanite::InitRasterContext(
@@ -2567,18 +2572,22 @@ void FVirtualShadowMapArray::RenderVirtualShadowMapsNanite(FRDGBuilder& GraphBui
 
 	TArray<Nanite::FPackedView, SceneRenderingAllocator> VirtualShadowViews;
 
-	for (FProjectedShadowInfo* ProjectedShadowInfo : VirtualShadowMapShadows)
 	{
-		if (ProjectedShadowInfo->bShouldRenderVSM)
+		TRACE_CPUPROFILER_EVENT_SCOPE(AddRenderViews);
+
+		for (FProjectedShadowInfo* ProjectedShadowInfo : VirtualShadowMapShadows)
 		{
-			AddRenderViews(
-				ProjectedShadowInfo,
-				SceneRenderer.Views,
-				ShadowsLODScaleFactor,
-				PrevHZBPhysical.IsValid(),
-				bVSMUseHZB,
-				ProjectedShadowInfo->ShouldClampToNearPlane(),
-				VirtualShadowViews);
+			if (ProjectedShadowInfo->bShouldRenderVSM)
+			{
+				AddRenderViews(
+					ProjectedShadowInfo,
+					SceneRenderer.Views,
+					ShadowsLODScaleFactor,
+					PrevHZBPhysical.IsValid(),
+					bVSMUseHZB,
+					ProjectedShadowInfo->ShouldClampToNearPlane(),
+					VirtualShadowViews);
+			}
 		}
 	}
 
@@ -2587,6 +2596,8 @@ void FVirtualShadowMapArray::RenderVirtualShadowMapsNanite(FRDGBuilder& GraphBui
 		int32 NumPrimaryViews = VirtualShadowViews.Num();
 		CreateMipViews(VirtualShadowViews);
 
+		SET_DWORD_STAT(STAT_VSMNaniteViewsPrimary, NumPrimaryViews);
+
 		Nanite::FCullingContext::FConfiguration CullingConfig = { 0 };
 		CullingConfig.bUpdateStreaming = bUpdateNaniteStreaming;
 		CullingConfig.bTwoPassOcclusion = UseTwoPassHzbOcclusion();
@@ -2644,6 +2655,7 @@ void FVirtualShadowMapArray::RenderVirtualShadowMapsNonNanite(FRDGBuilder& Graph
 		return;
 	}
 
+	TRACE_CPUPROFILER_EVENT_SCOPE(FVirtualShadowMapArray::RenderVirtualShadowMapsNonNanite);
 	RDG_EVENT_SCOPE(GraphBuilder, "RenderVirtualShadowMaps(Non-Nanite)");
 
 	FGPUScene& GPUScene = Scene.GPUScene;
