Disable async prepare dynamic shadows on RHI's which don't support async render queries.

diff --git a/Engine/Source/Runtime/D3D12RHI/Private/D3D12RHI.cpp b/Engine/Source/Runtime/D3D12RHI/Private/D3D12RHI.cpp
--- a/Engine/Source/Runtime/D3D12RHI/Private/D3D12RHI.cpp
+++ b/Engine/Source/Runtime/D3D12RHI/Private/D3D12RHI.cpp
@@ -75,6 +75,7 @@ FD3D12DynamicRHI::FD3D12DynamicRHI(const TArray<TSharedPtr<FD3D12Adapter>>& Chos
 
 	GRHISupportsMultithreading = true;
 	GRHISupportsMultithreadedResources = true;
+	GRHISupportsAsyncGetRenderQueryResult = true;
 	GRHIMultiPipelineMergeableAccessMask = GRHIMergeableAccessMask;
 	EnumRemoveFlags(GRHIMultiPipelineMergeableAccessMask, ERHIAccess::UAVMask);
 
diff --git a/Engine/Source/Runtime/RHI/Private/RHI.cpp b/Engine/Source/Runtime/RHI/Private/RHI.cpp
--- a/Engine/Source/Runtime/RHI/Private/RHI.cpp
+++ b/Engine/Source/Runtime/RHI/Private/RHI.cpp
@@ -1271,6 +1271,7 @@ bool GSupportsMobileMultiView = false;
 bool GSupportsImageExternal = false;
 bool GRHISupportsDrawIndirect = true;
 bool GRHISupportsMultithreading = false;
+bool GRHISupportsAsyncGetRenderQueryResult = false;
 bool GRHISupportsUpdateFromBufferTexture = false;
 bool GSupportsWideMRT = true;
 bool GRHINeedsExtraDeletionLatency = false;
diff --git a/Engine/Source/Runtime/RHI/Public/RHI.h b/Engine/Source/Runtime/RHI/Public/RHI.h
--- a/Engine/Source/Runtime/RHI/Public/RHI.h
+++ b/Engine/Source/Runtime/RHI/Public/RHI.h
@@ -75,6 +75,9 @@ extern RHI_API bool GRHISupportsDrawIndirect;
 /** Whether the RHI can send commands to the device context from multiple threads. Used in the GPU readback to avoid stalling the RHI threads. */
 extern RHI_API bool GRHISupportsMultithreading;
 
+/** Whether RHIGetRenderQueryResult can be safely called off the render thread. */
+extern RHI_API bool GRHISupportsAsyncGetRenderQueryResult;
+
 /** 
  * only set if RHI has the information (after init of the RHI and only if RHI has that information, never changes after that)
  * e.g. "NVIDIA GeForce GTX 670"
diff --git a/Engine/Source/Runtime/Renderer/Private/ShadowSetup.cpp b/Engine/Source/Runtime/Renderer/Private/ShadowSetup.cpp
--- a/Engine/Source/Runtime/Renderer/Private/ShadowSetup.cpp
+++ b/Engine/Source/Runtime/Renderer/Private/ShadowSetup.cpp
@@ -254,6 +254,14 @@ static TAutoConsoleVariable<int32> CVarParallelGatherShadowPrimitives(
 	ECVF_RenderThreadSafe
 	);
 
+int32 GParallelPrepareDynamicShadows = 1;
+static FAutoConsoleVariableRef CVarParallelPrepareDynamicShadows(
+	TEXT("r.ParallelPrepareDynamicShadows"),
+	GParallelPrepareDynamicShadows,
+	TEXT("Toggles parallel prepare dynamic shadows. 0 = off; 1 = on"),
+	ECVF_RenderThreadSafe
+);
+
 static TAutoConsoleVariable<int32> CVarParallelGatherNumPrimitivesPerPacket(
 	TEXT("r.ParallelGatherNumPrimitivesPerPacket"),
 	256,  
@@ -1185,6 +1193,7 @@ struct FDynamicShadowsTaskData : public FSceneRenderingAllocatorObject<FDynamicS
 	const bool bStaticSceneOnly;
 	const bool bRunningEarly;
 	const bool bMultithreaded;
+	const bool bMultithreadedPrepareAndFinalize;
 
 	// Generated from prepare task
 	TArray<FProjectedShadowInfo*, SceneRenderingAllocator> PreShadows;
@@ -1214,8 +1223,9 @@ struct FDynamicShadowsTaskData : public FSceneRenderingAllocatorObject<FDynamicS
 		, bStaticSceneOnly(AreAnyViewsStaticSceneOnly(Views))
 		, bRunningEarly(bInRunningEarly)
 		, bMultithreaded((FApp::ShouldUseThreadingForPerformance() || FForkProcessHelper::IsForkedMultithreadInstance()) && CVarParallelGatherShadowPrimitives.GetValueOnRenderThread() > 0)
+		, bMultithreadedPrepareAndFinalize(bRunningEarly && bMultithreaded && GRHISupportsAsyncGetRenderQueryResult && GParallelPrepareDynamicShadows > 0)
 	{
-		if (bMultithreaded)
+		if (bMultithreadedPrepareAndFinalize)
 		{
 			PrepareTaskEvent = FGraphEvent::CreateGraphEvent();
 		}
@@ -4533,8 +4543,11 @@ struct FGatherShadowPrimitivesPrepareTask
 		SCOPED_NAMED_EVENT_TEXT("FGatherShadowPrimitivesPrepareTask", FColor::Green);
 		FOptionalTaskTagScope Scope(ETaskTag::EParallelRenderingThread);
 
-		TaskData.SceneRenderer->PrepareDynamicShadows(TaskData);
-		TaskData.PrepareTaskEvent->DispatchSubsequents();
+		if (TaskData.bMultithreadedPrepareAndFinalize)
+		{
+			TaskData.SceneRenderer->PrepareDynamicShadows(TaskData);
+			TaskData.PrepareTaskEvent->DispatchSubsequents();
+		}
 
 		AnyThreadTask();
 
@@ -4569,11 +4582,13 @@ struct FGatherShadowPrimitivesPrepareTask
 						TStatId(),
 						nullptr,
 						CPrio_GatherShadowPrimitives.Get()));
+
+					MyCompletionGraphEvent->DontCompleteUntil(Prereqs.Last());
 				}
 			}
 		}
 
-		if (Prereqs.Num() > 0)
+		if (TaskData.bMultithreadedPrepareAndFinalize && Prereqs.Num() > 0)
 		{
 			FGraphEventRef FinalizeTask = FFunctionGraphTask::CreateAndDispatchWhenReady([TaskData = &TaskData]()
 				{
@@ -4674,29 +4689,32 @@ void FSceneRenderer::BeginGatherShadowPrimitives(FDynamicShadowsTaskData* TaskDa
 
 	check(TaskData);
 
-	if (!TaskData->bMultithreaded || !TaskData->bRunningEarly)
+	if (!TaskData->bMultithreadedPrepareAndFinalize)
 	{
 		PrepareDynamicShadows(*TaskData);
+	}
 
+	if (!TaskData->bMultithreaded || !TaskData->bRunningEarly)
+	{
 		if (TaskData->PreShadows.Num() || TaskData->ViewDependentWholeSceneShadows.Num())
 		{
 			FGatherShadowPrimitivesPrepareTask(*TaskData).AnyThreadTask();
 
-			if (!TaskData->bMultithreaded)
-			{
-				for (FGatherShadowPrimitivesPacket* Packet : TaskData->Packets)
-				{
-					Packet->AnyThreadTask(*TaskData);
-				}
-			}
-			else // !TaskData->bRunningEarly
+			if (!TaskData->bRunningEarly)
 			{
 				ParallelFor(TaskData->Packets.Num(),
 					[TaskData](int32 Index)
 					{
 						TaskData->Packets[Index]->AnyThreadTask(*TaskData);
 					},
-					!(FApp::ShouldUseThreadingForPerformance() && CVarParallelGatherShadowPrimitives.GetValueOnRenderThread() > 0));
+					!TaskData->bMultithreaded);
+			}
+			else
+			{
+				for (FGatherShadowPrimitivesPacket* Packet : TaskData->Packets)
+				{
+					Packet->AnyThreadTask(*TaskData);
+				}
 			}
 		}
 	}
@@ -4717,7 +4735,8 @@ void FSceneRenderer::FinishGatherShadowPrimitives(FDynamicShadowsTaskData* TaskD
 	{
 		FTaskGraphInterface::Get().WaitUntilTaskCompletes(TaskData->TaskEvent, ENamedThreads::GetRenderThread_Local());
 	}
-	else
+
+	if (!TaskData->bMultithreadedPrepareAndFinalize)
 	{
 		SCOPED_NAMED_EVENT_TEXT("FGatherShadowPrimitivesTask", FColor::Green);
 		for (FGatherShadowPrimitivesPacket* Packet : TaskData->Packets)
