Strata - fixed access to scene texture and compilation bug on translucent materials.

diff --git a/Engine/Shaders/Private/MaterialTemplate.ush b/Engine/Shaders/Private/MaterialTemplate.ush
--- a/Engine/Shaders/Private/MaterialTemplate.ush
+++ b/Engine/Shaders/Private/MaterialTemplate.ush
@@ -2444,7 +2444,7 @@ float4 SceneTextureLookup(float2 UV, int SceneTextureIndex, bool bFiltered)
 	#if STRATA_ENABLED
 	{
 		#ifdef StrataStruct
-			#define StrataSceneTexture StrataStructure
+			#define StrataSceneTexture StrataStruct
 		#else
 			#define StrataSceneTexture Strata
 		#endif
diff --git a/Engine/Source/Runtime/Renderer/Private/Strata/Strata.cpp b/Engine/Source/Runtime/Renderer/Private/Strata/Strata.cpp
--- a/Engine/Source/Runtime/Renderer/Private/Strata/Strata.cpp
+++ b/Engine/Source/Runtime/Renderer/Private/Strata/Strata.cpp
@@ -604,6 +604,7 @@ void BindStrataForwardPasslUniformParameters(FRDGBuilder& GraphBuilder, const FV
 		OutStrataUniformParameters.MaxBytesPerPixel = StrataSceneData->MaxBytesPerPixel;
 		OutStrataUniformParameters.bRoughDiffuse = StrataSceneData->bRoughDiffuse ? 1u : 0u;
 		OutStrataUniformParameters.PeelLayersAboveDepth = StrataSceneData->PeelLayersAboveDepth;
+		OutStrataUniformParameters.FirstSliceStoringStrataSSSData = StrataSceneData->FirstSliceStoringStrataSSSData;
 		OutStrataUniformParameters.MaterialTextureArray = StrataSceneData->MaterialTextureArray;
 		OutStrataUniformParameters.TopLayerTexture = StrataSceneData->TopLayerTexture;
 	}
@@ -614,6 +615,7 @@ void BindStrataForwardPasslUniformParameters(FRDGBuilder& GraphBuilder, const FV
 		OutStrataUniformParameters.MaxBytesPerPixel = 0;
 		OutStrataUniformParameters.bRoughDiffuse = 0;
 		OutStrataUniformParameters.PeelLayersAboveDepth = 0;
+		OutStrataUniformParameters.FirstSliceStoringStrataSSSData = -1;
 		OutStrataUniformParameters.MaterialTextureArray = DefaultTextureArray;
 		OutStrataUniformParameters.TopLayerTexture = SystemTextures.DefaultNormal8Bit;
 	}
diff --git a/Engine/Source/Runtime/Renderer/Private/Strata/Strata.h b/Engine/Source/Runtime/Renderer/Private/Strata/Strata.h
--- a/Engine/Source/Runtime/Renderer/Private/Strata/Strata.h
+++ b/Engine/Source/Runtime/Renderer/Private/Strata/Strata.h
@@ -34,6 +34,7 @@ BEGIN_SHADER_PARAMETER_STRUCT(FStrataForwardPassUniformParameters, )
 	SHADER_PARAMETER(uint32, MaxBytesPerPixel)
 	SHADER_PARAMETER(uint32, bRoughDiffuse)
 	SHADER_PARAMETER(uint32, PeelLayersAboveDepth)
+	SHADER_PARAMETER(int32, FirstSliceStoringStrataSSSData)
 	SHADER_PARAMETER_RDG_TEXTURE(Texture2DArray<uint>, MaterialTextureArray)
 	SHADER_PARAMETER_RDG_TEXTURE(Texture2D<uint2>, TopLayerTexture)
 END_SHADER_PARAMETER_STRUCT()
