Fix for Persistent ID ReffedItem count assert -- conditional needs to be atomic so ReffedItem can only be incremented once, if multiple threads happen to request a given item at the same time.

#rnx

diff --git a/Engine/Source/Runtime/Renderer/Private/MeshPassProcessor.cpp b/Engine/Source/Runtime/Renderer/Private/MeshPassProcessor.cpp
--- a/Engine/Source/Runtime/Renderer/Private/MeshPassProcessor.cpp
+++ b/Engine/Source/Runtime/Renderer/Private/MeshPassProcessor.cpp
@@ -502,7 +502,8 @@ FGraphicsMinimalPipelineStateId FGraphicsMinimalPipelineStateId::GetPersistentId
 		
 		FRefCountedGraphicsMinimalPipelineState& Value = PersistentIdTable.GetByElementId(TableId).Value;
 
-		if (Value.RefNum == 0)
+		// Needs to be atomic, so condition is only true on first increment
+		if (Value.RefNum.fetch_add(1) == 0)
 		{
 			if (!NeedsShaderInitialisation)
 			{
@@ -510,7 +511,6 @@ FGraphicsMinimalPipelineStateId FGraphicsMinimalPipelineStateId::GetPersistentId
 			}
 			ReffedItemCount++;
 		}
-		Value.RefNum++;
 
 #if MESH_DRAW_COMMAND_DEBUG_DATA
 		DebugSalt = Value.DebugSalt;
