Fix compilation issue in Unity Builds

The proper fix is to move the function to a common header, but will do this in a second pass after this unbreaks the build.

diff --git a/Engine/Source/Runtime/Renderer/Private/RayTracing/RayTracingGlobalIllumination.cpp b/Engine/Source/Runtime/Renderer/Private/RayTracing/RayTracingGlobalIllumination.cpp
--- a/Engine/Source/Runtime/Renderer/Private/RayTracing/RayTracingGlobalIllumination.cpp
+++ b/Engine/Source/Runtime/Renderer/Private/RayTracing/RayTracingGlobalIllumination.cpp
@@ -246,7 +246,8 @@ DECLARE_GPU_STAT_NAMED(RayTracingGIBruteForce, TEXT("Ray Tracing GI: Brute Force
 DECLARE_GPU_STAT_NAMED(RayTracingGIFinalGather, TEXT("Ray Tracing GI: Final Gather"));
 DECLARE_GPU_STAT_NAMED(RayTracingGICreateGatherPoints, TEXT("Ray Tracing GI: Create Gather Points"));
 
-static uint32 EncodeToF16x2(const FVector2f& In)
+// TODO: Move this to a common header
+static uint32 RTEncodeToF16x2(const FVector2f& In)
 {
 	return FFloat16(In.X).Encoded | (FFloat16(In.Y).Encoded << 16);
 }
@@ -330,7 +331,7 @@ static void SetupLightParameters(
 			DestLight.dPdv = LightShaderParameters.Tangent;
 			DestLight.Color = FVector3f(LightShaderParameters.Color);
 			DestLight.Dimensions = FVector2f(2.0f * LightShaderParameters.SourceRadius, 2.0f * LightShaderParameters.SourceLength);
-			DestLight.Shaping = EncodeToF16x2(FVector2f(LightShaderParameters.RectLightBarnCosAngle, LightShaderParameters.RectLightBarnLength));
+			DestLight.Shaping = RTEncodeToF16x2(FVector2f(LightShaderParameters.RectLightBarnCosAngle, LightShaderParameters.RectLightBarnLength));
 			DestLight.Flags |= PATHTRACING_LIGHT_RECT;
 			break;
 		}
@@ -357,7 +358,7 @@ static void SetupLightParameters(
 			DestLight.Color = FVector3f(LightShaderParameters.Color);
 			float SourceRadius = 0.0; // LightShaderParameters.SourceRadius causes too much noise for little pay off at this time
 			DestLight.Dimensions = FVector2f(SourceRadius, 0.0);
-			DestLight.Shaping = EncodeToF16x2(LightShaderParameters.SpotAngles);
+			DestLight.Shaping = RTEncodeToF16x2(LightShaderParameters.SpotAngles);
 			DestLight.Flags |= PATHTRACING_LIGHT_SPOT;
 			break;
 		}
