Change checks to improve error reporting when a dynamic primitive is not processed.

diff --git a/Engine/Source/Runtime/Renderer/Private/GPUScene.cpp b/Engine/Source/Runtime/Renderer/Private/GPUScene.cpp
--- a/Engine/Source/Runtime/Renderer/Private/GPUScene.cpp
+++ b/Engine/Source/Runtime/Renderer/Private/GPUScene.cpp
@@ -209,30 +209,21 @@ void FGPUScenePrimitiveCollector::Add(
 
 #if DO_CHECK
 
-bool FGPUScenePrimitiveCollector::IsPrimitiveProcessed(uint32 PrimitiveIndex, const FGPUScene& GPUScene) const
+void FGPUScenePrimitiveCollector::CheckPrimitiveProcessed(uint32 PrimitiveIndex, const FGPUScene& GPUScene) const
 {
-	if (UploadData == nullptr || !bCommitted)
-	{
-		// The collector hasn't collected anything or hasn't been uploaded
-		return false;
-	}
+	checkf(UploadData != nullptr && bCommitted, TEXT("Dynamic Primitive index %u has not been fully processed. The collector hasn't collected anything or hasn't been uploaded."), PrimitiveIndex);
 
-	if (PrimitiveIndex >= uint32(UploadData->PrimitiveData.Num()))
+	if (UploadData != nullptr)
 	{
-		// The specified index is out of range
-		return false;
-	}
+		checkf(PrimitiveIndex < uint32(UploadData->PrimitiveData.Num()), TEXT("Dynamic Primitive index %u has not been fully processed. The specified index is out of range [0,%d)."), PrimitiveIndex, UploadData->PrimitiveData.Num());
 
-	const FMeshBatchDynamicPrimitiveData& SourceData = UploadData->PrimitiveData[PrimitiveIndex].SourceData;
-	if (!SourceData.DataWriterGPU.IsBound() || SourceData.DataWriterGPUPass == EGPUSceneGPUWritePass::None)
-	{
-		// The primitive doesn't have a pending GPU write and has been uploaded or written to by the GPU already
-		return true;
+		const FMeshBatchDynamicPrimitiveData& SourceData = UploadData->PrimitiveData[PrimitiveIndex].SourceData;
+		checkf(!SourceData.DataWriterGPU.IsBound() || SourceData.DataWriterGPUPass == EGPUSceneGPUWritePass::None, TEXT("Dynamic Primitive index %u has not been fully processed. The primitive doesn't have a pending GPU write and has been uploaded or written to by the GPU already."), PrimitiveIndex);
 	}
 
 	// If the GPU scene still has a pending deferred write for the primitive, then it has not been fully processed yet
 	const uint32 PrimitiveId = GetPrimitiveIdRange().GetLowerBoundValue() + PrimitiveIndex;
-	return !GPUScene.HasPendingGPUWrite(PrimitiveId);
+	checkf(!GPUScene.HasPendingGPUWrite(PrimitiveId), TEXT("Dynamic Primitive index %u has not been fully processed. The GPU scene still has a pending deferred write for the primitive, it has not been fully processed yet."), PrimitiveIndex);
 }
 
 #endif // DO_CHECK
diff --git a/Engine/Source/Runtime/Renderer/Private/GPUScene.h b/Engine/Source/Runtime/Renderer/Private/GPUScene.h
--- a/Engine/Source/Runtime/Renderer/Private/GPUScene.h
+++ b/Engine/Source/Runtime/Renderer/Private/GPUScene.h
@@ -92,7 +92,7 @@ public:
 	/**
 	 * Determines if the specified primitive has been sufficiently processed and its data can be read
 	 */
-	bool IsPrimitiveProcessed(uint32 PrimitiveIndex, const FGPUScene& GPUScene) const;
+	void CheckPrimitiveProcessed(uint32 PrimitiveIndex, const FGPUScene& GPUScene) const;
 #endif // DO_CHECK
 
 private:
diff --git a/Engine/Source/Runtime/Renderer/Private/MeshDrawCommands.cpp b/Engine/Source/Runtime/Renderer/Private/MeshDrawCommands.cpp
--- a/Engine/Source/Runtime/Renderer/Private/MeshDrawCommands.cpp
+++ b/Engine/Source/Runtime/Renderer/Private/MeshDrawCommands.cpp
@@ -1429,9 +1429,7 @@ void FParallelMeshDrawCommandPass::BuildRenderingCommands(
 			if (VisibleMeshDrawCommand.PrimitiveIdInfo.bIsDynamicPrimitive)
 			{
 				uint32 PrimitiveIndex = VisibleMeshDrawCommand.PrimitiveIdInfo.DrawPrimitiveId & ~GPrimIDDynamicFlag;
-				checkf(TaskContext.View->DynamicPrimitiveCollector.IsPrimitiveProcessed(PrimitiveIndex, GPUScene),
-					TEXT("Dynamic Primitive index %u has not been fully processed. It may be an invalid index for the collector or it has a pending GPU write."),
-					PrimitiveIndex);
+				TaskContext.View->DynamicPrimitiveCollector.CheckPrimitiveProcessed(PrimitiveIndex, GPUScene);
 			}
 		}
 #endif
