Remove invalid thread check in rect texture manager.

diff --git a/Engine/Source/Runtime/Renderer/Private/RendererScene.cpp b/Engine/Source/Runtime/Renderer/Private/RendererScene.cpp
--- a/Engine/Source/Runtime/Renderer/Private/RendererScene.cpp
+++ b/Engine/Source/Runtime/Renderer/Private/RendererScene.cpp
@@ -3064,7 +3064,7 @@ void FScene::GetRectLightAtlasSlot(const FRectLightSceneProxy* Proxy, FLightRend
 {
 	if (Proxy)
 	{
-		check(IsInRenderingThread());
+		//check(IsInRenderingThread());
 		const RectLightAtlas::FAtlasSlotDesc Slot = RectLightAtlas::GetAtlasSlot(Proxy->RectAtlasId);
 		Out->RectLightAtlasUVOffset = Slot.UVOffset;
 		Out->RectLightAtlasUVScale = Slot.UVScale;
