Fix translucent selection toggle.

During the last blend mode reactoring, blend mode moved to be fetched later. This cause the blend mode to be fetched from the default material used for selection rather than the original material, causing the toggling to no longer be turned off.

diff --git a/Engine/Source/Runtime/Renderer/Private/SceneHitProxyRendering.cpp b/Engine/Source/Runtime/Renderer/Private/SceneHitProxyRendering.cpp
--- a/Engine/Source/Runtime/Renderer/Private/SceneHitProxyRendering.cpp
+++ b/Engine/Source/Runtime/Renderer/Private/SceneHitProxyRendering.cpp
@@ -758,6 +758,7 @@ void FDeferredShadingSceneRenderer::RenderHitProxies(FRDGBuilder& GraphBuilder)
 
 bool FHitProxyMeshProcessor::TryAddMeshBatch(const FMeshBatch& RESTRICT MeshBatch, uint64 BatchElementMask, const FPrimitiveSceneProxy* RESTRICT PrimitiveSceneProxy, int32 StaticMeshId, const FMaterialRenderProxy* MaterialRenderProxy, const FMaterial* Material)
 {
+	const bool bIsTranslucent = IsTranslucentBlendMode(*Material);
 	const FMeshDrawingPolicyOverrideSettings OverrideSettings = ComputeMeshOverrideSettings(MeshBatch);
 	const ERasterizerFillMode MeshFillMode = ComputeMeshFillMode(*Material, OverrideSettings);
 	const ERasterizerCullMode MeshCullMode = ComputeMeshCullMode(*Material, OverrideSettings);
@@ -795,7 +796,7 @@ bool FHitProxyMeshProcessor::TryAddMeshBatch(const FMeshBatch& RESTRICT MeshBatc
 	}
 
 	bool bResult = true;
-	if (bAddTranslucentPrimitive || !IsTranslucentBlendMode(*Material))
+	if (bAddTranslucentPrimitive || !bIsTranslucent)
 	{
 		bResult = Process(MeshBatch, BatchElementMask, StaticMeshId, PrimitiveSceneProxy, *MaterialRenderProxy, *Material, MeshFillMode, MeshCullMode);
 	}
