Path Tracer: Fix random path tracer invalidations

The use of a timer to detect changes for the path tracer leads to occasional random refreshes which is problematic with long running lookdev sessions.

diff --git a/Engine/Source/Runtime/Renderer/Private/SceneVisibility.cpp b/Engine/Source/Runtime/Renderer/Private/SceneVisibility.cpp
--- a/Engine/Source/Runtime/Renderer/Private/SceneVisibility.cpp
+++ b/Engine/Source/Runtime/Renderer/Private/SceneVisibility.cpp
@@ -3938,14 +3938,6 @@ void FSceneRenderer::PreVisibilityFrameSetup(FRDGBuilder& GraphBuilder, const FS
 			const bool bResetCamera = (bFirstFrameOrTimeWasReset || View.bCameraCut || bIsLargeCameraMovement || View.bForceCameraVisibilityReset);
 			
 #if RHI_RAYTRACING
-			// Note: 0.18 deg is the minimum angle for avoiding numerical precision issue (which would cause constant invalidation)
-			const bool bIsCameraMove = IsLargeCameraMovement(
-				View,
-				ViewState->PrevFrameViewInfo.ViewMatrices.GetViewMatrix(),
-				ViewState->PrevFrameViewInfo.ViewMatrices.GetViewOrigin(),
-				0.18f /*degree*/, 0.1f /*cm*/);
-			const bool bIsProjMatrixDifferent = View.ViewMatrices.GetProjectionNoAAMatrix() != View.ViewState->PrevFrameViewInfo.ViewMatrices.GetProjectionNoAAMatrix();
-			
 			static const auto CVarTemporalDenoiser = IConsoleManager::Get().FindTConsoleVariableDataInt(TEXT("r.PathTracing.TemporalDenoiser.mode"));
 			const int TemporalDenoiserMode = CVarTemporalDenoiser ? CVarTemporalDenoiser->GetValueOnAnyThread() : 0;
 
@@ -3964,16 +3956,25 @@ void FSceneRenderer::PreVisibilityFrameSetup(FRDGBuilder& GraphBuilder, const FS
 			{
 				// for interactive usage - any movement or scene change should restart the path tracer
 
+				// Note: 0.18 deg is the minimum angle for avoiding numerical precision issue (which would cause constant invalidation)
+				const bool bIsCameraMove = IsLargeCameraMovement(
+					View,
+					ViewState->PrevFrameViewInfo.ViewMatrices.GetViewMatrix(),
+					ViewState->PrevFrameViewInfo.ViewMatrices.GetViewOrigin(),
+					0.18f /*degree*/, 0.1f /*cm*/);
+				const bool bIsProjMatrixDifferent = View.ViewMatrices.GetProjectionNoAAMatrix() != View.ViewState->PrevFrameViewInfo.ViewMatrices.GetProjectionNoAAMatrix();
+
 				// For each view, we remember what the invalidation counter was set to last time we were here so we can catch all changes
-				bool bNeedsInvalidation = ViewState->PathTracingInvalidationCounter != CurrentPathTracingInvalidationCounter;
+				const bool bNeedsInvalidation = ViewState->PathTracingInvalidationCounter != CurrentPathTracingInvalidationCounter;
 				ViewState->PathTracingInvalidationCounter = CurrentPathTracingInvalidationCounter;
 				if (bNeedsInvalidation ||
-					bResetCamera ||
 					bIsProjMatrixDifferent ||
 					bIsCameraMove ||
+					View.bCameraCut ||
+					View.bForceCameraVisibilityReset ||
 					View.bForcePathTracerReset)
 				{
-					const bool bClearTemporalDenoisingHistory = (TemporalDenoiserMode == 2) ? (View.bCameraCut || bResetCamera) : true;
+					const bool bClearTemporalDenoisingHistory = (TemporalDenoiserMode == 2) ? View.bCameraCut : true;
 					ViewState->PathTracingInvalidate(bClearTemporalDenoisingHistory);
 				}
 			}
