Fix crash when SSR with HalfrResScene color is enabled. HalfresScene texture wasnt computed.

#rnx

diff --git a/Engine/Source/Runtime/Renderer/Private/PostProcess/PostProcessing.cpp b/Engine/Source/Runtime/Renderer/Private/PostProcess/PostProcessing.cpp
--- a/Engine/Source/Runtime/Renderer/Private/PostProcess/PostProcessing.cpp
+++ b/Engine/Source/Runtime/Renderer/Private/PostProcess/PostProcessing.cpp
@@ -540,7 +540,8 @@ void AddPostProcessingPasses(
 			(bBloomEnabled && !bFFTBloomEnabled) ||
 			(bLensFlareEnabled && bFFTBloomEnabled) ||
 			bLocalExposureBlurredLum);
-		const bool bNeedPostMotionBlurHalfRes = !bProcessQuarterResolution || (bFFTBloomEnabled && FFTBloomResolutionFraction > 0.25f && FFTBloomResolutionFraction <= 0.5f);
+		extern int32 GSSRHalfResSceneColor;
+		const bool bNeedPostMotionBlurHalfRes = !bProcessQuarterResolution || (bFFTBloomEnabled && FFTBloomResolutionFraction > 0.25f && FFTBloomResolutionFraction <= 0.5f) || (ReflectionsMethod == EReflectionsMethod::SSR && !View.bStatePrevViewInfoIsReadOnly && GSSRHalfResSceneColor);
 		const bool bNeedPostMotionBlurQuarterRes = bProcessQuarterResolution || (bFFTBloomEnabled && FFTBloomResolutionFraction <= 0.25f);
 
 		// Post Process Material Chain - Before Translucency
@@ -793,7 +794,6 @@ void AddPostProcessingPasses(
 
 
 		// Store half res scene color in the history
-		extern int32 GSSRHalfResSceneColor;
 		if (ReflectionsMethod == EReflectionsMethod::SSR && !View.bStatePrevViewInfoIsReadOnly && GSSRHalfResSceneColor)
 		{
 			check(View.ViewState);
