Fix invalid indirect arg reading/writing in groom system on platform requiring 16bytes alignment.

diff --git a/Engine/Shaders/Private/HairStrands/HairStrandsDeepTransmittanceMask.usf b/Engine/Shaders/Private/HairStrands/HairStrandsDeepTransmittanceMask.usf
--- a/Engine/Shaders/Private/HairStrands/HairStrandsDeepTransmittanceMask.usf
+++ b/Engine/Shaders/Private/HairStrands/HairStrandsDeepTransmittanceMask.usf
@@ -271,7 +271,7 @@ Texture2D<uint> 	ShadowMaskBitsTexture;
 Texture2D<float4> 	RayMarchMaskTexture;
 float4				ShadowChannelMask;
 #endif
-Buffer<uint4>		IndirectArgsBuffer;
+Buffer<uint>		IndirectArgsBuffer;
 RWBuffer<uint>		OutTransmittanceMask;
 
 FHairTransmittanceMask ComputeTransmittanceVoxel(
diff --git a/Engine/Shaders/Private/HairStrands/HairStrandsEnvironmentLighting.usf b/Engine/Shaders/Private/HairStrands/HairStrandsEnvironmentLighting.usf
--- a/Engine/Shaders/Private/HairStrands/HairStrandsEnvironmentLighting.usf
+++ b/Engine/Shaders/Private/HairStrands/HairStrandsEnvironmentLighting.usf
@@ -829,7 +829,7 @@ void MainVS(
 
 #if LIGHTING_PS
 
-Buffer<uint4> IndirectArgsBuffer;
+Buffer<uint> IndirectArgsBuffer;
 
 void MainPS(
 	float4 SVPos : SV_POSITION,
diff --git a/Engine/Shaders/Private/HairStrands/HairStrandsVoxelPageAllocation.usf b/Engine/Shaders/Private/HairStrands/HairStrandsVoxelPageAllocation.usf
--- a/Engine/Shaders/Private/HairStrands/HairStrandsVoxelPageAllocation.usf
+++ b/Engine/Shaders/Private/HairStrands/HairStrandsVoxelPageAllocation.usf
@@ -740,9 +740,7 @@ void VoxelIndPageClearBufferGenCS(uint GroupIndex : SV_GroupIndex, uint3 Dispatc
 	const uint DispatchX = TotalAllocatedPageCount % MaxGroupCountPerDim;
 	const uint DispatchY = ceil(TotalAllocatedPageCount / float(MaxGroupCountPerDim));
 
-	OutIndirectArgsBuffer[0] = DispatchX;
-	OutIndirectArgsBuffer[1] = DispatchY;
-	OutIndirectArgsBuffer[2] = ceil(VoxelCountPerPage / float(IND_CLEAR_GROUP_SIZE_Z));
+	WriteDispatchIndirectArgs(OutIndirectArgsBuffer, 0, DispatchX, DispatchY, ceil(VoxelCountPerPage / float(IND_CLEAR_GROUP_SIZE_Z)));
 }
 
 #endif
diff --git a/Engine/Source/Runtime/Renderer/Private/HairStrands/HairStrandsVoxelization.cpp b/Engine/Source/Runtime/Renderer/Private/HairStrands/HairStrandsVoxelization.cpp
--- a/Engine/Source/Runtime/Renderer/Private/HairStrands/HairStrandsVoxelization.cpp
+++ b/Engine/Source/Runtime/Renderer/Private/HairStrands/HairStrandsVoxelization.cpp
@@ -203,7 +203,7 @@ static void AddVirtualVoxelInjectOpaquePass(
 	const FGlobalShaderMap* GlobalShaderMap = View.ShaderMap;
 
 	check(VoxelResources.Parameters.Common.IndirectDispatchGroupSize == 64);
-	const uint32 ArgsOffset = sizeof(uint32) * 3 * Parameters->MacroGroupId;
+	const uint32 ArgsOffset = sizeof(FRHIDispatchIndirectParameters) * Parameters->MacroGroupId;
 
 	FComputeShaderUtils::AddPass(
 		GraphBuilder, 
@@ -903,7 +903,7 @@ static void AddAllocateVoxelPagesPass(
 				Parameters->PageIndexResolutionAndOffsetBuffer = PageIndexResolutionAndOffsetBufferSRV;
 				Parameters->IndirectBufferArgs = PageIndexAllocationIndirectBufferArgs;
 
-				const uint32 ArgsOffset = sizeof(uint32) * 3 * MacroGroup.MacroGroupId;
+				const uint32 ArgsOffset = sizeof(FRHIDispatchIndirectParameters) * MacroGroup.MacroGroupId;
 
 				FComputeShaderUtils::AddPass(GraphBuilder, RDG_EVENT_NAME("HairStrands::AllocateVoxelPage"), ComputeShader, Parameters, 
 					PageIndexAllocationIndirectBufferArgs,
