Fix for assert in AddChildNode.

diff --git a/Engine/Source/Runtime/Renderer/Private/RendererScene.cpp b/Engine/Source/Runtime/Renderer/Private/RendererScene.cpp
--- a/Engine/Source/Runtime/Renderer/Private/RendererScene.cpp
+++ b/Engine/Source/Runtime/Renderer/Private/RendererScene.cpp
@@ -5868,12 +5868,6 @@ void FScene::UpdateAllPrimitiveSceneInfos(FRDGBuilder& GraphBuilder, EUpdateAllP
 				// Flush virtual textures touched by primitive
 				QueueSceneInfoToFlushVirtualTexture(PrimitiveSceneInfo);
 
-				// Set LOD parent information if valid
-				PrimitiveSceneInfo->LinkLODParentComponent();
-
-				// Update scene LOD tree
-				SceneLODHierarchy.UpdateNodeSceneInfo(PrimitiveSceneInfo->PrimitiveComponentId, PrimitiveSceneInfo);
-
 				bNeedPathTracedInvalidation = bNeedPathTracedInvalidation || IsPrimitiveRelevantToPathTracing(PrimitiveSceneInfo);
 			}
 			AddedLocalPrimitiveSceneInfos.RemoveAt(StartIndex, AddedLocalPrimitiveSceneInfos.Num() - StartIndex, false);
@@ -6128,6 +6122,15 @@ void FScene::UpdateAllPrimitiveSceneInfos(FRDGBuilder& GraphBuilder, EUpdateAllP
 		});
 	}
 
+	for (FPrimitiveSceneInfo* PrimitiveSceneInfo : AddedPrimitiveSceneInfos)
+	{
+		// Set LOD parent information if valid
+		PrimitiveSceneInfo->LinkLODParentComponent();
+
+		// Update scene LOD tree
+		SceneLODHierarchy.UpdateNodeSceneInfo(PrimitiveSceneInfo->PrimitiveComponentId, PrimitiveSceneInfo);
+	}
+
 	for (FPrimitiveSceneInfo* PrimitiveSceneInfo : SceneInfosToFlushVirtualTexture)
 	{
 		PrimitiveSceneInfo->FlushRuntimeVirtualTexture();
