Keeps r.TSR.History.GrandReprojection setable on scalability settings

diff --git a/Engine/Source/Runtime/Renderer/Private/PostProcess/TemporalSuperResolution.cpp b/Engine/Source/Runtime/Renderer/Private/PostProcess/TemporalSuperResolution.cpp
--- a/Engine/Source/Runtime/Renderer/Private/PostProcess/TemporalSuperResolution.cpp
+++ b/Engine/Source/Runtime/Renderer/Private/PostProcess/TemporalSuperResolution.cpp
@@ -42,9 +42,9 @@ TAutoConsoleVariable<int32> CVarTSRHistorySeparateTranslucency(
 	ECVF_RenderThreadSafe);
 
 TAutoConsoleVariable<int32> CVarTSRHistoryGrandReprojection(
-	TEXT("r.TSR.History.GrandReprojection"), 0,
+	TEXT("r.TSR.History.GrandReprojection"), 1,
 	TEXT("Experimental functionality to keep higher sharpness in TSR's history in motion."),
-	ECVF_RenderThreadSafe);
+	ECVF_Scalability | ECVF_RenderThreadSafe);
 
 TAutoConsoleVariable<int32> CVarTSRWaveOps(
 	TEXT("r.TSR.WaveOps"), 1,
