RayTracing: Objects casting hidden shadows or affecting indirect while hidden were not visible to ray tracing

This meant that simply hitting 'g' in the editor would make those ray visibility controls for the path tracer work or not.

diff --git a/Engine/Source/Runtime/Renderer/Private/DeferredShadingRenderer.cpp b/Engine/Source/Runtime/Renderer/Private/DeferredShadingRenderer.cpp
--- a/Engine/Source/Runtime/Renderer/Private/DeferredShadingRenderer.cpp
+++ b/Engine/Source/Runtime/Renderer/Private/DeferredShadingRenderer.cpp
@@ -826,7 +826,12 @@ static void GatherRayTracingRelevantPrimitives(const FScene& Scene, const FViewI
 			// and hidden shadow casters must still always be added to the RT scene.
 			if (bGameView && !SceneInfo->bDrawInGame && !SceneInfo->bRayTracingFarField)
 			{
-				continue;
+				// Make sure this isn't an object that wants to be hidden to camera but still wants to cast shadows or be visible to indirect
+				check(SceneInfo->Proxy != nullptr);
+				if (!SceneInfo->Proxy->CastsHiddenShadow() && !SceneInfo->Proxy->AffectsIndirectLightingWhileHidden())
+				{
+					continue;
+				}
 			}
 
 			// Marked visible and used after point, check if streaming then mark as used in the TLAS (so it can be streamed in)
