Fix CAS VRS not being correctly disabled for certain passes based on their CVar setting.

diff --git a/Engine/Source/Runtime/Renderer/Private/VariableRateShading/ContrastAdaptiveImageGenerator.cpp b/Engine/Source/Runtime/Renderer/Private/VariableRateShading/ContrastAdaptiveImageGenerator.cpp
--- a/Engine/Source/Runtime/Renderer/Private/VariableRateShading/ContrastAdaptiveImageGenerator.cpp
+++ b/Engine/Source/Runtime/Renderer/Private/VariableRateShading/ContrastAdaptiveImageGenerator.cpp
@@ -538,7 +538,7 @@ void AddPrepareImageBasedVRSPass(
 
 		float DynamicResolutionScale = (float)ViewportWidth / (float)PostProcessRect.Width();
 
-		auto PassParameters = GraphBuilder.AllocParameters<FRescaleVariableRateShadingCS::FParameters>();
+		FRescaleVariableRateShadingCS::FParameters* PassParameters = GraphBuilder.AllocParameters<FRescaleVariableRateShadingCS::FParameters>();
 
 		FRescaleVariableRateShadingCS::InitParameters(
 			*PassParameters,
@@ -573,9 +573,12 @@ FRDGTextureRef FContrastAdaptiveImageGenerator::GetImage(FRDGBuilder& GraphBuild
 {
 	if (FVRSTextures::IsInitialized(GraphBuilder))
 	{
-		auto TextureType = ESRITextureType::GetTextureType(PassType);
-		const FVRSTextures& VRSTextures = FVRSTextures::Get(GraphBuilder);
-		return (TextureType == ESRITextureType::ScaledSRIForRender) ? VRSTextures.ScaledSRI : VRSTextures.ScaledConservativeSRI;
+		ESRITextureType::Type TextureType = ESRITextureType::GetTextureType(PassType);
+		if (TextureType != ESRITextureType::None)
+		{
+			const FVRSTextures& VRSTextures = FVRSTextures::Get(GraphBuilder);
+			return (TextureType == ESRITextureType::ScaledSRIForRender) ? VRSTextures.ScaledSRI : VRSTextures.ScaledConservativeSRI;
+		}
 	}
 
 	return nullptr;
@@ -588,7 +591,7 @@ void FContrastAdaptiveImageGenerator::PrepareImages(FRDGBuilder& GraphBuilder, c
 	for (const FSceneView* View : ViewFamily.Views)
 	{
 		check(View->bIsViewInfo);
-		auto ViewInfo = static_cast<const FViewInfo*>(View);
+		const FViewInfo* ViewInfo = static_cast<const FViewInfo*>(View);
 		if (!View->bCameraCut && FVariableRateShadingImageManager::IsVRSCompatibleWithView(*ViewInfo) && ViewInfo->PrevViewInfo.LuminanceHistory)
 		{
 			bIsAnyViewVRSCompatible = true;
