Remove the now unused EYEADAPTATION_EXPOSURE_FIX parameter from the tonemapper, in order to reduce permutations.

diff --git a/Engine/Shaders/Private/PostProcessTonemap.usf b/Engine/Shaders/Private/PostProcessTonemap.usf
--- a/Engine/Shaders/Private/PostProcessTonemap.usf
+++ b/Engine/Shaders/Private/PostProcessTonemap.usf
@@ -79,11 +79,6 @@ float EditorNITLevel;
 uint bOutputInHDR;
 float OutputMaxLuminance;
 
-#if EYEADAPTATION_EXPOSURE_FIX != 1
-// Default value used when eye adaptation is disabled (e.g fixed exposure) or not available.
-float DefaultEyeExposure;
-#endif
-
 float4 BloomDirtMaskTint;
 Texture2D BloomDirtMaskTexture;
 SamplerState BloomDirtMaskSampler;
@@ -424,12 +419,7 @@ float4 TonemapCommonPS(
 
 	SceneColor.rgb *= View.OneOverPreExposure;
 
-#if EYEADAPTATION_EXPOSURE_FIX == 1
 	float2 EyeAdaptationData = EyeAdaptationBuffer[0].xw;
-#else
-	// If eye adaptation is disabled (e.g. fixed exposure level ) or not available.  
-	float2 EyeAdaptationData = DefaultEyeExposure.xx;
-#endif
 
 	float ExposureScale = EyeAdaptationData.x;
 
diff --git a/Engine/Source/Runtime/Renderer/Private/PostProcess/PostProcessTonemap.cpp b/Engine/Source/Runtime/Renderer/Private/PostProcess/PostProcessTonemap.cpp
--- a/Engine/Source/Runtime/Renderer/Private/PostProcess/PostProcessTonemap.cpp
+++ b/Engine/Source/Runtime/Renderer/Private/PostProcess/PostProcessTonemap.cpp
@@ -21,6 +21,7 @@
 #include "SystemTextures.h"
 #include "HDRHelper.h"
 #include "VariableRateShadingImageManager.h"
+#include "DataDrivenShaderPlatformInfo.h"
 
 
 bool SupportsFilmGrain(EShaderPlatform Platform)
@@ -80,8 +81,6 @@ class FTonemapperVignetteDim       : SHADER_PERMUTATION_BOOL("USE_VIGNETTE");
 class FTonemapperSharpenDim        : SHADER_PERMUTATION_BOOL("USE_SHARPEN");
 class FTonemapperFilmGrainDim      : SHADER_PERMUTATION_BOOL("USE_FILM_GRAIN");
 class FTonemapperMsaaDim           : SHADER_PERMUTATION_BOOL("METAL_MSAA_HDR_DECODE");
-class FTonemapperEyeAdaptationDim  : SHADER_PERMUTATION_BOOL("EYEADAPTATION_EXPOSURE_FIX");
-
 using FCommonDomain = TShaderPermutationDomain<
 	FTonemapperBloomDim,
 	FTonemapperGammaOnlyDim,
@@ -89,8 +88,7 @@ using FCommonDomain = TShaderPermutationDomain<
 	FTonemapperVignetteDim,
 	FTonemapperSharpenDim,
 	FTonemapperFilmGrainDim,
-	FTonemapperMsaaDim,
-	FTonemapperEyeAdaptationDim>;
+	FTonemapperMsaaDim>;
 
 bool ShouldCompileCommonPermutation(const FGlobalShaderPermutationParameters& Parameters, const FCommonDomain& PermutationVector)
 {
@@ -108,14 +106,13 @@ bool ShouldCompileCommonPermutation(const FGlobalShaderPermutationParameters& Pa
 			!PermutationVector.Get<FTonemapperVignetteDim>() &&
 			!PermutationVector.Get<FTonemapperSharpenDim>() &&
 			!PermutationVector.Get<FTonemapperFilmGrainDim>() &&
-			!PermutationVector.Get<FTonemapperMsaaDim>() &&
-			!PermutationVector.Get<FTonemapperEyeAdaptationDim>();
+			!PermutationVector.Get<FTonemapperMsaaDim>();
 	}
 	return true;
 }
 
 // Common conversion of engine settings into.
-FCommonDomain BuildCommonPermutationDomain(const FViewInfo& View, bool bGammaOnly, bool bEyeAdaptation, bool bLocalExposure, bool bMetalMSAAHDRDecode)
+FCommonDomain BuildCommonPermutationDomain(const FViewInfo& View, bool bGammaOnly, bool bLocalExposure, bool bMetalMSAAHDRDecode)
 {
 	const FSceneViewFamily* Family = View.Family;
 
@@ -137,7 +134,6 @@ FCommonDomain BuildCommonPermutationDomain(const FViewInfo& View, bool bGammaOnl
 	PermutationVector.Set<FTonemapperFilmGrainDim>(View.FilmGrainTexture != nullptr);
 	PermutationVector.Set<FTonemapperSharpenDim>(CVarTonemapperSharpen.GetValueOnRenderThread() > 0.0f);
 	PermutationVector.Set<FTonemapperMsaaDim>(bMetalMSAAHDRDecode);
-	PermutationVector.Set<FTonemapperEyeAdaptationDim>(bEyeAdaptation);
 	return PermutationVector;
 }
 
@@ -340,7 +336,6 @@ BEGIN_SHADER_PARAMETER_STRUCT(FTonemapParameters, )
 	SHADER_PARAMETER(float, InvLUTSize)
 	SHADER_PARAMETER(float, LUTScale)
 	SHADER_PARAMETER(float, LUTOffset)
-	SHADER_PARAMETER(float, DefaultEyeExposure)
 	SHADER_PARAMETER(float, EditorNITLevel)
 	SHADER_PARAMETER(uint32, bOutputInHDR)
 END_SHADER_PARAMETER_STRUCT()
@@ -412,6 +407,10 @@ public:
 
 	static bool ShouldCompilePermutation(const FGlobalShaderPermutationParameters& Parameters)
 	{
+		if (!IsFeatureLevelSupported(Parameters.Platform, ERHIFeatureLevel::ES3_1))
+		{
+			return false;
+		}
 		return TonemapperPermutation::ShouldCompileDesktopPermutation(Parameters, FPermutationDomain(Parameters.PermutationId));
 	}
 
@@ -549,6 +548,13 @@ FScreenPassTexture AddTonemapPass(FRDGBuilder& GraphBuilder, const FViewInfo& Vi
 
 	const bool bIsEyeAdaptationResource = Inputs.EyeAdaptationBuffer != nullptr;
 	const bool bEyeAdaptation = ViewFamily.EngineShowFlags.EyeAdaptation && bIsEyeAdaptationResource;
+	FRDGBufferRef EyeAdaptationBuffer = Inputs.EyeAdaptationBuffer;
+
+	if (!bEyeAdaptation)
+	{
+		float DefaultEyeAdaptation = GetEyeAdaptationFixedExposure(View);
+		EyeAdaptationBuffer = CreateStructuredBuffer(GraphBuilder, TEXT("DefaultEyeAdaptationBuffer"), sizeof(float), 1, &DefaultEyeAdaptation, sizeof(float), ERDGInitialDataFlags::None);
+	}
 
 	const FScreenPassTextureViewport SceneColorViewport(Inputs.SceneColor);
 
@@ -605,8 +611,6 @@ FScreenPassTexture AddTonemapPass(FRDGBuilder& GraphBuilder, const FViewInfo& Vi
 	FRHISamplerState* BilinearClampSampler = TStaticSamplerState<SF_Bilinear, AM_Clamp, AM_Clamp, AM_Clamp>::GetRHI();
 	FRHISamplerState* PointClampSampler = TStaticSamplerState<SF_Point, AM_Clamp, AM_Clamp, AM_Clamp>::GetRHI();
 
-	const float DefaultEyeExposure = bEyeAdaptation ? 1.0f : GetEyeAdaptationFixedExposure(View);
-
 	const float SharpenDiv6 = FMath::Clamp(CVarTonemapperSharpen.GetValueOnRenderThread(), 0.0f, 10.0f) / 6.0f;
 
 	FVector4f ChromaticAberrationParams;
@@ -754,7 +758,7 @@ FScreenPassTexture AddTonemapPass(FRDGBuilder& GraphBuilder, const FViewInfo& Vi
 	CommonParameters.BlurredLogLum = Inputs.BlurredLogLuminanceTexture;
 	CommonParameters.LumBilateralGridSampler = BilinearClampSampler;
 	CommonParameters.BlurredLogLumSampler = BilinearClampSampler;
-	CommonParameters.EyeAdaptationBuffer = Inputs.EyeAdaptationBuffer ? GraphBuilder.CreateSRV(Inputs.EyeAdaptationBuffer) : nullptr;
+	CommonParameters.EyeAdaptationBuffer = GraphBuilder.CreateSRV(Inputs.EyeAdaptationBuffer);
 	CommonParameters.EyeAdaptation = *Inputs.EyeAdaptationParameters;
 	CommonParameters.ColorGradingLUT = Inputs.ColorGradingTexture;
 	CommonParameters.ColorSampler = BilinearClampSampler;
@@ -762,7 +766,6 @@ FScreenPassTexture AddTonemapPass(FRDGBuilder& GraphBuilder, const FViewInfo& Vi
 	CommonParameters.ColorScale0 = PostProcessSettings.SceneColorTint;
 	CommonParameters.ChromaticAberrationParams = ChromaticAberrationParams;
 	CommonParameters.TonemapperParams = FVector4f(PostProcessSettings.VignetteIntensity, SharpenDiv6, 0.0f, 0.0f);
-	CommonParameters.DefaultEyeExposure = DefaultEyeExposure;
 	CommonParameters.EditorNITLevel = EditorNITLevel;
 	CommonParameters.bOutputInHDR = ViewFamily.bIsHDR;
 	CommonParameters.LUTSize = LUTSize;
@@ -835,7 +838,7 @@ FScreenPassTexture AddTonemapPass(FRDGBuilder& GraphBuilder, const FViewInfo& Vi
 	TonemapperPermutation::FDesktopDomain DesktopPermutationVector;
 
 	{
-		TonemapperPermutation::FCommonDomain CommonDomain = TonemapperPermutation::BuildCommonPermutationDomain(View, Inputs.bGammaOnly, bEyeAdaptation, Inputs.LocalExposureTexture != nullptr, Inputs.bMetalMSAAHDRDecode);
+		TonemapperPermutation::FCommonDomain CommonDomain = TonemapperPermutation::BuildCommonPermutationDomain(View, Inputs.bGammaOnly, Inputs.LocalExposureTexture != nullptr, Inputs.bMetalMSAAHDRDecode);
 		DesktopPermutationVector.Set<TonemapperPermutation::FCommonDomain>(CommonDomain);
 
 		if (!CommonDomain.Get<TonemapperPermutation::FTonemapperGammaOnlyDim>())
