RayTracing: Add more explanation to debug check()

diff --git a/Engine/Source/Runtime/Renderer/Private/DeferredShadingRenderer.cpp b/Engine/Source/Runtime/Renderer/Private/DeferredShadingRenderer.cpp
--- a/Engine/Source/Runtime/Renderer/Private/DeferredShadingRenderer.cpp
+++ b/Engine/Source/Runtime/Renderer/Private/DeferredShadingRenderer.cpp
@@ -827,7 +827,7 @@ static void GatherRayTracingRelevantPrimitives(const FScene& Scene, const FViewI
 			if (bGameView && !SceneInfo->bDrawInGame && !SceneInfo->bRayTracingFarField)
 			{
 				// Make sure this isn't an object that wants to be hidden to camera but still wants to cast shadows or be visible to indirect
-				check(SceneInfo->Proxy != nullptr);
+				checkf(SceneInfo->Proxy != nullptr, TEXT("SceneInfo does not have a valid Proxy object. If this occurs, this object should probably have been filtered out before being added to Scene.Primitives"));
 				if (!SceneInfo->Proxy->CastsHiddenShadow() && !SceneInfo->Proxy->AffectsIndirectLightingWhileHidden())
 				{
 					continue;
