Clarified that the energy conservation render setting is only for the legacy material system. Strata always has that forced to true.

diff --git a/Engine/Shaders/Private/ShadingEnergyConservation.ush b/Engine/Shaders/Private/ShadingEnergyConservation.ush
--- a/Engine/Shaders/Private/ShadingEnergyConservation.ush
+++ b/Engine/Shaders/Private/ShadingEnergyConservation.ush
@@ -4,8 +4,8 @@
 
 // If EnergyConservation flag is not enabled explicitly then:
 // * If strata is enabled: ON by default
-// * If strata is disnabled: OFF by default
-#if STRATA_ENABLED || MATERIAL_ENERGYCONSERVATION
+// * If strata is disabled: OFF by default
+#if STRATA_ENABLED || LEGACY_MATERIAL_ENERGYCONSERVATION
 	#undef USE_ENERGY_CONSERVATION
 	#define USE_ENERGY_CONSERVATION 1
 #else
diff --git a/Engine/Source/Runtime/Engine/Classes/Engine/RendererSettings.h b/Engine/Source/Runtime/Engine/Classes/Engine/RendererSettings.h
--- a/Engine/Source/Runtime/Engine/Classes/Engine/RendererSettings.h
+++ b/Engine/Source/Runtime/Engine/Classes/Engine/RendererSettings.h
@@ -959,7 +959,7 @@ class ENGINE_API URendererSettings : public UDeveloperSettings
 	*/
 	UPROPERTY(config, EditAnywhere, Category = Materials, meta = (
 		ConsoleVariable = "r.Material.EnergyConservation", DisplayName = "Enable Energy Conservation on Material",
-		ToolTip = "Enable Energy Conservation on Material.",
+		ToolTip = "Enable Energy Conservation on Material. Please note that when Strata is enabled, energy conservation is forced to enabled.",
 		ConfigRestartRequired = true))
 		uint32 bMaterialEnergyConservation : 1;
 
diff --git a/Engine/Source/Runtime/Engine/Private/ShaderCompiler/ShaderCompiler.cpp b/Engine/Source/Runtime/Engine/Private/ShaderCompiler/ShaderCompiler.cpp
--- a/Engine/Source/Runtime/Engine/Private/ShaderCompiler/ShaderCompiler.cpp
+++ b/Engine/Source/Runtime/Engine/Private/ShaderCompiler/ShaderCompiler.cpp
@@ -6424,7 +6424,7 @@ void GlobalBeginCompileShader(
 	{
 		static IConsoleVariable* CVar = IConsoleManager::Get().FindConsoleVariable(TEXT("r.Material.EnergyConservation"));
 		const bool bMaterialEnergyConservation = CVar && CVar->GetInt() != 0;
-		Input.Environment.SetDefine(TEXT("MATERIAL_ENERGYCONSERVATION"), bMaterialEnergyConservation ? 1 : 0);
+		Input.Environment.SetDefine(TEXT("LEGACY_MATERIAL_ENERGYCONSERVATION"), bMaterialEnergyConservation ? 1 : 0);
 	}
 
 	{
diff --git a/Engine/Source/Runtime/Renderer/Private/ShadingEnergyConservation.cpp b/Engine/Source/Runtime/Renderer/Private/ShadingEnergyConservation.cpp
--- a/Engine/Source/Runtime/Renderer/Private/ShadingEnergyConservation.cpp
+++ b/Engine/Source/Runtime/Renderer/Private/ShadingEnergyConservation.cpp
@@ -51,7 +51,7 @@ static TAutoConsoleVariable<int32> CVarShadingFurnaceTest_TableFormat(
 static TAutoConsoleVariable<int32> CVarMaterialEnergyConservation(
 	TEXT("r.Material.EnergyConservation"),
 	0,
-	TEXT("Enable energy conservation for material (project settings, read only)."),
+	TEXT("Enable energy conservation for legacy materials (project settings, read only). Please note that when Strata is enabled, energy conservation is forced to enabled."),
 	ECVF_ReadOnly | ECVF_RenderThreadSafe);
 
 #define SHADING_ENERGY_CONSERVATION_TABLE_RESOLUTION 32
