Fix missing interpolation pass of cards/meshes instance being outside of view. This was causing invalida shadow when grom were visible in shadow, but in primary view.

diff --git a/Engine/Plugins/Runtime/HairStrands/Source/HairStrandsCore/Private/GroomDebug.cpp b/Engine/Plugins/Runtime/HairStrands/Source/HairStrandsCore/Private/GroomDebug.cpp
--- a/Engine/Plugins/Runtime/HairStrands/Source/HairStrandsCore/Private/GroomDebug.cpp
+++ b/Engine/Plugins/Runtime/HairStrands/Source/HairStrandsCore/Private/GroomDebug.cpp
@@ -811,7 +811,8 @@ class FHairDebugPrintInstanceCS : public FGlobalShader
 		SHADER_PARAMETER(uint32, NameCharacterCount)
 		SHADER_PARAMETER(uint32, InstanceCount_StrandsPrimaryView)
 		SHADER_PARAMETER(uint32, InstanceCount_StrandsShadowView)
-		SHADER_PARAMETER(uint32, InstanceCount_CardsOrMeshes)
+		SHADER_PARAMETER(uint32, InstanceCount_CardsOrMeshesPrimaryView)
+		SHADER_PARAMETER(uint32, InstanceCount_CardsOrMeshesShadowView)
 		SHADER_PARAMETER_RDG_BUFFER_SRV(StructuredBuffer<uint2>, NameInfos)
 		SHADER_PARAMETER_RDG_BUFFER_SRV(Buffer<uint8>, Names)
 		SHADER_PARAMETER_RDG_BUFFER_SRV(StructuredBuffer<uint4>, Infos)
@@ -931,10 +932,9 @@ static void AddHairDebugPrintInstancePass(
 				}
 
 				{
-					Data2.X = InstanceIndex < InstanceCountPerType[HairInstanceCount_StrandsPrimaryView] ? 3 : 2; // Visible in primary & shadow (3) or only in shadow (2)
-					Data2.Y = 0;
-					Data2.Z = 0;
-					Data2.W = 0;
+					Data2 = FUintVector4(0);
+					Data2.X |= InstanceIndex < InstanceCountPerType[HairInstanceCount_StrandsPrimaryView] ? 0x1u : 0u;
+					Data2.X |= InstanceIndex < InstanceCountPerType[HairInstanceCount_StrandsShadowView]  ? 0x2u : 0u;
 				}
 			}
 			break;
@@ -943,7 +943,10 @@ static void AddHairDebugPrintInstancePass(
 			{
 				Data0.Z = Instance->Cards.LODs[IntLODIndex].Guides.IsValid() ? Instance->Cards.LODs[IntLODIndex].Guides.Data->GetNumCurves() : 0;
 				Data0.W = Instance->Cards.LODs[IntLODIndex].Data->GetNumVertices();
-				Data2.X = 3; // Visible in primary & shadow
+
+				Data2 = FUintVector4(0);
+				Data2.X |= InstanceIndex < InstanceCountPerType[HairInstanceCount_CardsOrMeshesPrimaryView] ? 0x1u : 0u;
+				Data2.X |= InstanceIndex < InstanceCountPerType[HairInstanceCount_CardsOrMeshesShadowView]  ? 0x2u : 0u;
 			}
 			break;
 		case EHairGeometryType::Meshes:
@@ -951,7 +954,10 @@ static void AddHairDebugPrintInstancePass(
 			{
 				Data0.Z = 0;
 				Data0.W = Instance->Meshes.LODs[IntLODIndex].Data->GetNumVertices();
-				Data2.X = 3; // Visible in primary & shadow
+
+				Data2 = FUintVector4(0);
+				Data2.X |= InstanceIndex < InstanceCountPerType[HairInstanceCount_CardsOrMeshesPrimaryView] ? 0x1u : 0u;
+				Data2.X |= InstanceIndex < InstanceCountPerType[HairInstanceCount_CardsOrMeshesShadowView]  ? 0x2u : 0u;
 			}
 			break;
 		}
@@ -982,7 +988,8 @@ static void AddHairDebugPrintInstancePass(
 		Parameters->InstanceCount = InstanceCount;
 		Parameters->InstanceCount_StrandsPrimaryView = InstanceCountPerType[HairInstanceCount_StrandsPrimaryView];
 		Parameters->InstanceCount_StrandsShadowView = InstanceCountPerType[HairInstanceCount_StrandsShadowView];
-		Parameters->InstanceCount_CardsOrMeshes = InstanceCountPerType[HairInstanceCount_CardsOrMeshes];
+		Parameters->InstanceCount_CardsOrMeshesPrimaryView = InstanceCountPerType[HairInstanceCount_CardsOrMeshesPrimaryView];
+		Parameters->InstanceCount_CardsOrMeshesShadowView = InstanceCountPerType[HairInstanceCount_CardsOrMeshesShadowView];
 		Parameters->NameInfoCount = NameInfos.Num();
 		Parameters->NameCharacterCount = Names.Num();
 		Parameters->Names = GraphBuilder.CreateSRV(NameBuffer, PF_R8_UINT);
diff --git a/Engine/Shaders/Private/HairStrands/HairStrandsDebug.usf b/Engine/Shaders/Private/HairStrands/HairStrandsDebug.usf
--- a/Engine/Shaders/Private/HairStrands/HairStrandsDebug.usf
+++ b/Engine/Shaders/Private/HairStrands/HairStrandsDebug.usf
@@ -1066,7 +1066,8 @@ int2					MaxResolution;
 uint					InstanceCount;
 uint					InstanceCount_StrandsPrimaryView;
 uint					InstanceCount_StrandsShadowView;
-uint					InstanceCount_CardsOrMeshes;
+uint					InstanceCount_CardsOrMeshesPrimaryView;
+uint					InstanceCount_CardsOrMeshesShadowView;
 
 Buffer<uint>			Infos;
 
@@ -1239,16 +1240,20 @@ void MainCS(uint GroupIndex : SV_GroupIndex, uint3 DispatchThreadId : SV_Dispatc
 	Newline(Context);
 
 	{
-		Print(Context, TEXT("Total Strand primary: "), FontWhite);
+		Print(Context, TEXT("Total Strand (primary)        : "), FontWhite);
 		Print(Context,InstanceCount_StrandsPrimaryView, FontYellow, 8);
 		Newline(Context);
 
-		Print(Context, TEXT("Total Strand shadow : "), FontWhite);
+		Print(Context, TEXT("Total Strand (shadow)         : "), FontWhite);
 		Print(Context,InstanceCount_StrandsShadowView, FontYellow, 8);
 		Newline(Context);
 
-		Print(Context, TEXT("Total Cards / Meshes: "), FontWhite);
-		Print(Context,InstanceCount_CardsOrMeshes, FontYellow, 8);
+		Print(Context, TEXT("Total Cards / Meshes (primary): "), FontWhite);
+		Print(Context,InstanceCount_CardsOrMeshesPrimaryView, FontYellow, 8);
+		Newline(Context);
+
+		Print(Context, TEXT("Total Cards / Meshes (shadow) : "), FontWhite);
+		Print(Context, InstanceCount_CardsOrMeshesShadowView, FontYellow, 8);
 		Newline(Context);
 	}
 	Newline(Context);
diff --git a/Engine/Source/Runtime/Renderer/Private/HairStrands/HairStrandsInterface.cpp b/Engine/Source/Runtime/Renderer/Private/HairStrands/HairStrandsInterface.cpp
--- a/Engine/Source/Runtime/Renderer/Private/HairStrands/HairStrandsInterface.cpp
+++ b/Engine/Source/Runtime/Renderer/Private/HairStrands/HairStrandsInterface.cpp
@@ -428,8 +428,14 @@ bool IsHairStrandsVisibleInShadows(const FViewInfo& View, const FHairStrandsInst
 	bool bIsVisibleInShadow = false;
 	if (const FHairGroupPublicData* HairData = Instance.GetHairData())
 	{
-		const bool bIsStrands = HairData->LODIndex >= 0 && Instance.GetHairGeometry() == EHairGeometryType::Strands;
-		if (!bIsStrands)
+		// Run simulation if the instance is either Strands geometry, have simulation, or has global interpolation enabled.
+		// This ensures that the groom is correctly updated if visible in shadows
+		const bool bNeedUpdate = 
+			HairData->LODIndex >= 0 && 
+			(Instance.GetHairGeometry() == EHairGeometryType::Strands || 
+			 HairData->IsSimulationEnable(HairData->LODIndex) || 
+			 HairData->IsGlobalInterpolationEnable(HairData->LODIndex));
+		if (!bNeedUpdate)
 		{
 			return false;
 		}
@@ -632,10 +638,16 @@ FHairStrandsBookmarkParameters CreateHairStrandsBookmarkParameters(FScene* Scene
 	const int32 ActiveInstanceCount = Scene->HairStrandsSceneData.RegisteredProxies.Num();
 	TBitArray InstancesVisibility(false, ActiveInstanceCount);
 
+	static bool bDebug = false;
+	if (Scene->HairStrandsSceneData.RegisteredProxies.Num() > 0)
+	{
+		bDebug = true;
+	}
+
 	FHairStrandsBookmarkParameters Out;
 	Out.VisibleInstances.Reserve(View.HairStrandsMeshElements.Num());
 
-	// 1. Add all visible strands instances
+	// 1. Strands - Add all visible strands instances
 	for (const FMeshBatchAndRelevance& MeshBatch : View.HairStrandsMeshElements)
 	{
 		check(MeshBatch.PrimitiveSceneProxy && MeshBatch.PrimitiveSceneProxy->ShouldRenderInMainPass());
@@ -651,12 +663,13 @@ FHairStrandsBookmarkParameters CreateHairStrandsBookmarkParameters(FScene* Scene
 	}
 	Out.InstanceCountPerType[HairInstanceCount_StrandsPrimaryView] = Out.VisibleInstances.Num();
 
-	// 2. Add all instances non-visible primary view(s) but visible in shadow view(s)
+	// 2. Strands - Add all instances non-visible primary view(s) but visible in shadow view(s)
 	if (IsHairStrandsNonVisibleShadowCastingEnable())
 	{
 		for (FHairStrandsInstance* Instance : Scene->HairStrandsSceneData.RegisteredProxies)
 		{
-			if (Instance->RegisteredIndex >= 0 && Instance->RegisteredIndex < ActiveInstanceCount && !InstancesVisibility[Instance->RegisteredIndex])
+			const bool bStrands = Instance->GetHairGeometry() == EHairGeometryType::Strands;
+			if (Instance->RegisteredIndex >= 0 && Instance->RegisteredIndex < ActiveInstanceCount && !InstancesVisibility[Instance->RegisteredIndex] && bStrands)
 			{
 				if (IsHairStrandsVisibleInShadows(View, *Instance))
 				{
@@ -667,7 +680,7 @@ FHairStrandsBookmarkParameters CreateHairStrandsBookmarkParameters(FScene* Scene
 	}
 	Out.InstanceCountPerType[HairInstanceCount_StrandsShadowView] = Out.VisibleInstances.Num() - Out.InstanceCountPerType[HairInstanceCount_StrandsPrimaryView];
 
-	// 3. Add all visible cards instances
+	// 3. Cards/Meshes - Add all visible cards instances
 	for (const FMeshBatchAndRelevance& MeshBatch : View.HairCardsMeshElements)
 	{
 		check(MeshBatch.PrimitiveSceneProxy && MeshBatch.PrimitiveSceneProxy->ShouldRenderInMainPass());
@@ -681,7 +694,24 @@ FHairStrandsBookmarkParameters CreateHairStrandsBookmarkParameters(FScene* Scene
 			}
 		}
 	}
-	Out.InstanceCountPerType[HairInstanceCount_CardsOrMeshes] = Out.VisibleInstances.Num() - Out.InstanceCountPerType[HairInstanceCount_StrandsShadowView];
+	Out.InstanceCountPerType[HairInstanceCount_CardsOrMeshesPrimaryView] = Out.VisibleInstances.Num() - Out.InstanceCountPerType[HairInstanceCount_StrandsShadowView];
+
+	// 4. Cards/Meshes - Add all instances non-visible primary view(s) but visible in shadow view(s)
+	if (IsHairStrandsNonVisibleShadowCastingEnable())
+	{
+		for (FHairStrandsInstance* Instance : Scene->HairStrandsSceneData.RegisteredProxies)
+		{
+			const bool bCardsOrMeshes = Instance->GetHairGeometry() == EHairGeometryType::Cards || Instance->GetHairGeometry() == EHairGeometryType::Meshes;
+			if (Instance->RegisteredIndex >= 0 && Instance->RegisteredIndex < ActiveInstanceCount && !InstancesVisibility[Instance->RegisteredIndex] && bCardsOrMeshes)
+			{
+				if (IsHairStrandsVisibleInShadows(View, *Instance))
+				{
+					Out.VisibleInstances.Add(Instance);
+				}
+			}
+		}
+	}
+	Out.InstanceCountPerType[HairInstanceCount_CardsOrMeshesShadowView] = Out.VisibleInstances.Num() - Out.InstanceCountPerType[HairInstanceCount_CardsOrMeshesPrimaryView];
 
 	Out.ShaderPrintData			= ShaderPrint::IsEnabled(View.ShaderPrintData) ? &View.ShaderPrintData : nullptr;
 	Out.ShaderMap				= View.ShaderMap;
diff --git a/Engine/Source/Runtime/Renderer/Public/HairStrandsInterface.h b/Engine/Source/Runtime/Renderer/Public/HairStrandsInterface.h
--- a/Engine/Source/Runtime/Renderer/Public/HairStrandsInterface.h
+++ b/Engine/Source/Runtime/Renderer/Public/HairStrandsInterface.h
@@ -429,7 +429,8 @@ enum EHairInstanceCount : uint8
 {
 	HairInstanceCount_StrandsPrimaryView = 0,
 	HairInstanceCount_StrandsShadowView = 1,
-	HairInstanceCount_CardsOrMeshes = 2,
+	HairInstanceCount_CardsOrMeshesPrimaryView = 2,
+	HairInstanceCount_CardsOrMeshesShadowView = 3
 };
 
 struct FHairStrandsBookmarkParameters
