Fix FScene::IncrementFrameNumber to be thread safe (as it is called from the game thread).

diff --git a/Engine/Source/Runtime/Renderer/Private/RendererScene.cpp b/Engine/Source/Runtime/Renderer/Private/RendererScene.cpp
--- a/Engine/Source/Runtime/Renderer/Private/RendererScene.cpp
+++ b/Engine/Source/Runtime/Renderer/Private/RendererScene.cpp
@@ -1861,6 +1861,7 @@ FScene::FScene(UWorld* InWorld, bool bInRequiresHitProxies, bool bInIsEditorScen
 ,	NumVisibleLights_GameThread(0)
 ,	NumEnabledSkylights_GameThread(0)
 ,	SceneFrameNumber(0)
+,	SceneFrameNumberRenderThread(0)
 ,	bForceNoPrecomputedLighting(InWorld->GetWorldSettings()->bForceNoPrecomputedLighting)
 {
 	FMemory::Memzero(MobileDirectionalLights);
@@ -6595,3 +6596,25 @@ void FRendererModule::InvalidatePathTracedOutput()
 		Scene->InvalidatePathTracedOutput();
 	}
 }
+
+uint32 FScene::GetFrameNumber() const
+{
+	if (IsInGameThread())
+	{
+		return SceneFrameNumber;
+	}
+	else
+	{
+		return SceneFrameNumberRenderThread;
+	}
+}
+
+void FScene::IncrementFrameNumber()
+{
+	// Increment game-tread version
+	++SceneFrameNumber;
+	ENQUEUE_RENDER_COMMAND(SceneStartFrame)([this,NewNumber = SceneFrameNumber](FRHICommandListImmediate& RHICmdList)
+	{
+		SceneFrameNumberRenderThread = NewNumber;
+	});
+}
\ No newline at end of file
diff --git a/Engine/Source/Runtime/Renderer/Private/ScenePrivate.h b/Engine/Source/Runtime/Renderer/Private/ScenePrivate.h
--- a/Engine/Source/Runtime/Renderer/Private/ScenePrivate.h
+++ b/Engine/Source/Runtime/Renderer/Private/ScenePrivate.h
@@ -3348,15 +3348,14 @@ public:
 		VelocityData.StartFrame(this);
 	}
 
-	virtual uint32 GetFrameNumber() const override
-	{
-		return SceneFrameNumber;
-	}
+	/**
+	 * Returns the current "FrameNumber" where frame corresponds to how many times FRendererModule::BeginRenderingViewFamilies has been called.
+	 * Thread safe, and returns a different copy for game/render thread. GetFrameNumberRenderThread can only be called from the render thread. 
+	 */
+	virtual uint32 GetFrameNumber() const override;
+	inline uint32 GetFrameNumberRenderThread() const { return SceneFrameNumberRenderThread; }
 
-	virtual void IncrementFrameNumber() override
-	{
-		++SceneFrameNumber;
-	}
+	virtual void IncrementFrameNumber() override;
 
 	void DumpMeshDrawCommandMemoryStats();
 
@@ -3618,8 +3617,9 @@ private:
 	 */
 	int32 NumEnabledSkylights_GameThread;
 
-	/** Frame number incremented per-family viewing this scene. */
+	/** Frame number incremented per-family (except if there are multiple view families in one render call) viewing this scene. */
 	uint32 SceneFrameNumber;
+	uint32 SceneFrameNumberRenderThread;
 
 	/** Whether world settings has bForceNoPrecomputedLighting set */
 	bool bForceNoPrecomputedLighting;
