Enable DBuffer pass during classification only on console for now.

diff --git a/Engine/Source/Runtime/Renderer/Private/Strata/Strata.cpp b/Engine/Source/Runtime/Renderer/Private/Strata/Strata.cpp
--- a/Engine/Source/Runtime/Renderer/Private/Strata/Strata.cpp
+++ b/Engine/Source/Runtime/Renderer/Private/Strata/Strata.cpp
@@ -743,6 +743,10 @@ class FStrataMaterialTileClassificationPassCS : public FGlobalShader
 		{
 			return false;
 		}
+		if (PermutationVector.Get<FDecal>() && !IsConsolePlatform(Parameters.Platform))
+		{
+			return false;
+		}		
 		return GetMaxSupportedFeatureLevel(Parameters.Platform) >= ERHIFeatureLevel::SM5 && Strata::IsStrataEnabled();
 	}
 
@@ -1236,7 +1240,7 @@ void AddStrataMaterialClassificationPass(FRDGBuilder& GraphBuilder, const FMinim
 			}
 
 			// If Dbuffer pass (i.e. apply DBuffer data after the base-pass) is enabled, run special classification for outputing tile with/without tiles
-			const bool bDBufferTiles = CVarStrataDBufferPass.GetValueOnRenderThread() > 0 && CVarStrataDBufferPassDedicatedTiles.GetValueOnRenderThread() > 0 && DBufferTextures.IsValid();
+			const bool bDBufferTiles = CVarStrataDBufferPass.GetValueOnRenderThread() > 0 && CVarStrataDBufferPassDedicatedTiles.GetValueOnRenderThread() > 0 && DBufferTextures.IsValid() && IsConsolePlatform(View.GetShaderPlatform());
 
 			FStrataMaterialTileClassificationPassCS::FPermutationDomain PermutationVector;
 			PermutationVector.Set< FStrataMaterialTileClassificationPassCS::FCmask >(bSupportCMask);
