Renderer: Exposed FXAAPass to external modules

diff --git a/Engine/Source/Runtime/Renderer/Private/PostProcess/PostProcessAA.h b/Engine/Source/Runtime/Renderer/Private/PostProcess/PostProcessAA.h
--- a/Engine/Source/Runtime/Renderer/Private/PostProcess/PostProcessAA.h
+++ b/Engine/Source/Runtime/Renderer/Private/PostProcess/PostProcessAA.h
@@ -32,4 +32,4 @@ struct FFXAAInputs
 	EFXAAQuality Quality = EFXAAQuality::MAX;
 };
 
-FScreenPassTexture AddFXAAPass(FRDGBuilder& GraphBuilder, const FViewInfo& View, const FFXAAInputs& Inputs);
\ No newline at end of file
+FScreenPassTexture RENDERER_API AddFXAAPass(FRDGBuilder& GraphBuilder, const FViewInfo& View, const FFXAAInputs& Inputs);
\ No newline at end of file
