Volumetric cloud now account for camera cut.

diff --git a/Engine/Source/Runtime/Renderer/Private/DeferredShadingRenderer.cpp b/Engine/Source/Runtime/Renderer/Private/DeferredShadingRenderer.cpp
--- a/Engine/Source/Runtime/Renderer/Private/DeferredShadingRenderer.cpp
+++ b/Engine/Source/Runtime/Renderer/Private/DeferredShadingRenderer.cpp
@@ -3290,6 +3290,10 @@ void FDeferredShadingSceneRenderer::Render(FRDGBuilder& GraphBuilder)
 	{
 		InitVolumetricRenderTargetForViews(GraphBuilder, Views);
 	}
+	else
+	{
+		ResetVolumetricRenderTargetForViews(GraphBuilder, Views);
+	}
 
 	InitVolumetricCloudsForViews(GraphBuilder, bShouldRenderVolumetricCloudBase, InstanceCullingManager);
 
diff --git a/Engine/Source/Runtime/Renderer/Private/VolumetricRenderTarget.cpp b/Engine/Source/Runtime/Renderer/Private/VolumetricRenderTarget.cpp
--- a/Engine/Source/Runtime/Renderer/Private/VolumetricRenderTarget.cpp
+++ b/Engine/Source/Runtime/Renderer/Private/VolumetricRenderTarget.cpp
@@ -168,13 +168,14 @@ FVolumetricRenderTargetViewStateData::~FVolumetricRenderTargetViewStateData()
 void FVolumetricRenderTargetViewStateData::Initialise(
 	FIntPoint& ViewRectResolutionIn,
 	int32 InMode,
-	int32 InUpsamplingMode)
+	int32 InUpsamplingMode,
+	bool bCameraCut)
 {
 	// Update internal settings
 	Mode = FMath::Clamp(InMode, 0, 3);
 	UpsamplingMode = Mode == 2 || Mode == 3 ? 2 : FMath::Clamp(InUpsamplingMode, 0, 4); // if we are using mode 2 then we cannot intersect with depth and upsampling should be 2 (simple on/off intersection)
 
-	if (bFirstTimeUsed)
+	if (bFirstTimeUsed || bCameraCut)
 	{
 		bFirstTimeUsed = false;
 		bHistoryValid = false;
@@ -218,7 +219,7 @@ void FVolumetricRenderTargetViewStateData::Initialise(
 		{
 			// Do not mark history as valid if the half resolution buffer is not valid. That means nothing has been rendered last frame.
 			// That can happen when cloud is used to render into that buffer
-			bHistoryValid = VolumetricReconstructRT[PreviousRT].IsValid();
+			bHistoryValid = !bCameraCut && VolumetricReconstructRT[PreviousRT].IsValid();
 
 			NoiseFrameIndex += FrameId == 0 ? 1 : 0;
 			NoiseFrameIndexModPattern = NoiseFrameIndex % (VolumetricTracingRTDownsampleFactor * VolumetricTracingRTDownsampleFactor);
@@ -253,6 +254,27 @@ void FVolumetricRenderTargetViewStateData::Initialise(
 	}
 }
 
+void FVolumetricRenderTargetViewStateData::Reset()
+{
+	bFirstTimeUsed = false;
+	bHistoryValid = false;
+	FrameId = 0;
+	NoiseFrameIndex = 0;
+	NoiseFrameIndexModPattern = 0;
+	CurrentPixelOffset = FIntPoint::ZeroValue;
+	CurrentRT = 0;
+
+	// Release GPU resources
+	VolumetricTracingRT.SafeRelease();
+	VolumetricSecondaryTracingRT.SafeRelease();
+	VolumetricTracingRTDepth.SafeRelease();
+	for(uint32 i = 0 ; i < kRenderTargetCount; ++i)
+	{
+		VolumetricReconstructRT[i].SafeRelease();
+		VolumetricReconstructRTDepth[i].SafeRelease();
+	}
+}
+
 FRDGTextureRef FVolumetricRenderTargetViewStateData::GetOrCreateVolumetricTracingRT(FRDGBuilder& GraphBuilder)
 {
 	check(FullResolution != FIntPoint::ZeroValue); // check that initialization has been done at least once
@@ -403,11 +425,17 @@ void InitVolumetricRenderTargetForViews(FRDGBuilder& GraphBuilder, TArrayView<FV
 		}
 		FVolumetricRenderTargetViewStateData& VolumetricCloudRT = ViewInfo.ViewState->VolumetricCloudRenderTarget;
 
+		// Determine if we are initializing or we should reset the persistent state
+		const float DeltaTime = ViewInfo.Family->Time.GetRealTimeSeconds() - ViewInfo.ViewState->LastRenderTime;
+		const bool bFirstFrameOrTimeWasReset = DeltaTime < -0.0001f || ViewInfo.ViewState->LastRenderTime < 0.0001f;
+		const bool bCameraCut = (bFirstFrameOrTimeWasReset || ViewInfo.bCameraCut || ViewInfo.bForceCameraVisibilityReset);
+
 		FIntPoint ViewRect = ViewInfo.ViewRect.Size();
 		VolumetricCloudRT.Initialise(	// TODO this is going to reallocate a buffer each time dynamic resolution scaling is applied 
 			ViewRect,
 			CVarVolumetricRenderTargetMode.GetValueOnRenderThread(),
-			CVarVolumetricRenderTargetUpsamplingMode.GetValueOnAnyThread());
+			CVarVolumetricRenderTargetUpsamplingMode.GetValueOnAnyThread(),
+			bCameraCut);
 
 		FViewUniformShaderParameters ViewVolumetricCloudRTParameters = *ViewInfo.CachedViewUniformShaderParameters;
 		{
@@ -442,6 +470,19 @@ void InitVolumetricRenderTargetForViews(FRDGBuilder& GraphBuilder, TArrayView<FV
 	}
 }
 
+void ResetVolumetricRenderTargetForViews(FRDGBuilder& GraphBuilder, TArrayView<FViewInfo> Views)
+{
+	for (int32 ViewIndex = 0; ViewIndex < Views.Num(); ViewIndex++)
+	{
+		FViewInfo& ViewInfo = Views[ViewIndex];
+		if (ViewInfo.ViewState)
+		{
+			FVolumetricRenderTargetViewStateData& VolumetricCloudRT = ViewInfo.ViewState->VolumetricCloudRenderTarget;
+			VolumetricCloudRT.Reset();
+		}
+	}
+}
+
 //////////////////////////////////////////////////////////////////////////
 
 class FReconstructVolumetricRenderTargetPS : public FGlobalShader
diff --git a/Engine/Source/Runtime/Renderer/Private/VolumetricRenderTarget.h b/Engine/Source/Runtime/Renderer/Private/VolumetricRenderTarget.h
--- a/Engine/Source/Runtime/Renderer/Private/VolumetricRenderTarget.h
+++ b/Engine/Source/Runtime/Renderer/Private/VolumetricRenderTarget.h
@@ -18,6 +18,7 @@ struct FSceneWithoutWaterTextures;
 struct FMinimalSceneTextures;
 
 void InitVolumetricRenderTargetForViews(FRDGBuilder& GraphBuilder, TArrayView<FViewInfo> Views);
+void ResetVolumetricRenderTargetForViews(FRDGBuilder& GraphBuilder, TArrayView<FViewInfo> Views);
 
 void ReconstructVolumetricRenderTarget(
 	FRDGBuilder& GraphBuilder,
diff --git a/Engine/Source/Runtime/Renderer/Private/VolumetricRenderTargetViewStateData.h b/Engine/Source/Runtime/Renderer/Private/VolumetricRenderTargetViewStateData.h
--- a/Engine/Source/Runtime/Renderer/Private/VolumetricRenderTargetViewStateData.h
+++ b/Engine/Source/Runtime/Renderer/Private/VolumetricRenderTargetViewStateData.h
@@ -24,7 +24,10 @@ public:
 	void Initialise(
 		FIntPoint& ViewRectResolutionIn,
 		int32 Mode,
-		int32 UpsamplingMode);
+		int32 UpsamplingMode,
+		bool bCameraCut);
+
+	void Reset();
 
 	FRDGTextureRef GetOrCreateVolumetricTracingRT(FRDGBuilder& GraphBuilder);
 	FRDGTextureRef GetOrCreateVolumetricSecondaryTracingRT(FRDGBuilder& GraphBuilder);
