Remove sRGB/Linear workaround for virtual textures on mobile.
OpenGL ES which does not support texture views will allocate a separate physical spaces for sRGB and Linear data

diff --git a/Engine/Shaders/Private/BCCompressionCommon.ush b/Engine/Shaders/Private/BCCompressionCommon.ush
--- a/Engine/Shaders/Private/BCCompressionCommon.ush
+++ b/Engine/Shaders/Private/BCCompressionCommon.ush
@@ -188,11 +188,6 @@ uint2 CompressBC1Block( in float3 Block[16] )
 // We expect linear colors as input and convert internally
 uint2 CompressBC1Block_SRGB( in float3 Block[16] )
 {
-#if (FEATURE_LEVEL == FEATURE_LEVEL_ES3_1)
-	// mobile uses linear color for RVT
-	return CompressBC1Block(Block);
-#endif
-	
 	float3 MinColor, MaxColor;
 	GetMinMax( Block, MinColor, MaxColor );
 
@@ -242,11 +237,6 @@ uint4 CompressBC3Block( in float3 BlockRGB[16], in float BlockA[16] )
 // We expect linear colors as input and convert internally
 uint4 CompressBC3Block_SRGB( in float3 BlockRGB[16], in float BlockA[16] )
 {
-#if (FEATURE_LEVEL == FEATURE_LEVEL_ES3_1)
-	// mobile uses linear color for RVT
-	return CompressBC3Block(BlockRGB, BlockA);
-#endif
-		
 	float3 MinColor, MaxColor;
 	GetMinMax( BlockRGB, MinColor, MaxColor );
 
diff --git a/Engine/Shaders/Private/ETCCompressionCommon.ush b/Engine/Shaders/Private/ETCCompressionCommon.ush
--- a/Engine/Shaders/Private/ETCCompressionCommon.ush
+++ b/Engine/Shaders/Private/ETCCompressionCommon.ush
@@ -177,7 +177,7 @@ uint2 CompressBlock_ETC1S(in float3 Block[16], bool bSRGB, bool bCoCg)
 
 uint2 CompressBlock_ETC2_SRGB(in float3 Block[16])
 {
-	const bool bSRGB = false;// true;
+	const bool bSRGB = true;
 	const bool bCoCg = false;
 	return CompressBlock_ETC1S(Block, bSRGB, bCoCg);
 }
diff --git a/Engine/Shaders/Private/MaterialTexture.ush b/Engine/Shaders/Private/MaterialTexture.ush
--- a/Engine/Shaders/Private/MaterialTexture.ush
+++ b/Engine/Shaders/Private/MaterialTexture.ush
@@ -148,12 +148,7 @@ MaterialFloat4 ProcessMaterialColorTextureLookup(MaterialFloat4 TextureValue)
 
 MaterialFloat4 ProcessMaterialVirtualColorTextureLookup(MaterialFloat4 TextureValue)
 {
-	TextureValue = ProcessMaterialColorTextureLookup(TextureValue);
-#if FEATURE_LEVEL == FEATURE_LEVEL_ES3_1
-	// on mobile all VT physical spaces use linear color formats, do sRGB to Linear conversion in the shader
-	TextureValue = MaterialFloat4(TextureValue.rgb*TextureValue.rgb, TextureValue.a);
-#endif
-	return TextureValue;
+	return ProcessMaterialColorTextureLookup(TextureValue);
 }
 
 MaterialFloat4 ProcessMaterialExternalTextureLookup(MaterialFloat4 TextureValue)
diff --git a/Engine/Source/Runtime/Engine/Private/Texture2D.cpp b/Engine/Source/Runtime/Engine/Private/Texture2D.cpp
--- a/Engine/Source/Runtime/Engine/Private/Texture2D.cpp
+++ b/Engine/Source/Runtime/Engine/Private/Texture2D.cpp
@@ -1338,6 +1338,7 @@ void FVirtualTexture2DResource::InitRHI()
 		ProducerDesc.LayerFormat[LayerIndex] = VTData->LayerTypes[LayerIndex];
 		ProducerDesc.LayerFallbackColor[LayerIndex] = VTData->LayerFallbackColors[LayerIndex];
 		ProducerDesc.PhysicalGroupIndex[LayerIndex] = bSinglePhysicalSpace ? 0 : LayerIndex;
+		ProducerDesc.bIsLayerSRGB[LayerIndex] = TextureOwner->SRGB;
 	}
 
 	FUploadingVirtualTexture* VirtualTexture = new FUploadingVirtualTexture(ProducerDesc.Name, VTData, FirstMipToUse);
diff --git a/Engine/Source/Runtime/Engine/Private/VT/RuntimeVirtualTexture.cpp b/Engine/Source/Runtime/Engine/Private/VT/RuntimeVirtualTexture.cpp
--- a/Engine/Source/Runtime/Engine/Private/VT/RuntimeVirtualTexture.cpp
+++ b/Engine/Source/Runtime/Engine/Private/VT/RuntimeVirtualTexture.cpp
@@ -341,6 +341,7 @@ void URuntimeVirtualTexture::GetProducerDescription(FVTProducerDescription& OutD
 	{
 		OutDesc.LayerFormat[Layer] = GetLayerFormat(Layer);
 		OutDesc.PhysicalGroupIndex[Layer] = bSinglePhysicalSpace ? 0 : Layer;
+		OutDesc.bIsLayerSRGB[Layer] = IsLayerSRGB(Layer);
 	}
 }
 
diff --git a/Engine/Source/Runtime/OpenGLDrv/Private/OpenGLDevice.cpp b/Engine/Source/Runtime/OpenGLDrv/Private/OpenGLDevice.cpp
--- a/Engine/Source/Runtime/OpenGLDrv/Private/OpenGLDevice.cpp
+++ b/Engine/Source/Runtime/OpenGLDrv/Private/OpenGLDevice.cpp
@@ -1004,6 +1004,9 @@ static void InitRHICapabilitiesForGL()
 		;
 
 	GRHISupportsMultithreadedResources = GRHISupportsRHIThread;
+
+	// OpenGL ES does not support glTextureView
+	GRHISupportsTextureViews = false;
 	
 	// By default use emulated UBs on mobile
 	GUseEmulatedUniformBuffers = IsUsingEmulatedUniformBuffers(GMaxRHIShaderPlatform);
diff --git a/Engine/Source/Runtime/OpenGLDrv/Private/OpenGLTexture.cpp b/Engine/Source/Runtime/OpenGLDrv/Private/OpenGLTexture.cpp
--- a/Engine/Source/Runtime/OpenGLDrv/Private/OpenGLTexture.cpp
+++ b/Engine/Source/Runtime/OpenGLDrv/Private/OpenGLTexture.cpp
@@ -2087,10 +2087,6 @@ void FOpenGLDynamicRHI::RHICopyTexture(FRHITexture* SourceTextureRHI, FRHITextur
 		DestTexture->Target, DestTextureRHI->GetFormat(), DestTextureRHI->GetFlags()
 	);
 	
-	checkf((SourceTextureRHI->GetFlags() & TexCreate_SRGB) == (DestTextureRHI->GetFlags() & TexCreate_SRGB), TEXT("Cannot copy between sRGB and linear, SourceTexture Format=%d, Flags=%x; DestTexture Format=%d, Flags=%x"),
-		SourceTextureRHI->GetFormat(), SourceTextureRHI->GetFlags(),
-		DestTextureRHI->GetFormat(), DestTextureRHI->GetFlags()
-	);
 
 	GLsizei Width, Height, Depth;
 
diff --git a/Engine/Source/Runtime/RHI/Private/RHI.cpp b/Engine/Source/Runtime/RHI/Private/RHI.cpp
--- a/Engine/Source/Runtime/RHI/Private/RHI.cpp
+++ b/Engine/Source/Runtime/RHI/Private/RHI.cpp
@@ -1326,6 +1326,7 @@ bool GRHISupportsParallelRHIExecute = false;
 bool GSupportsParallelOcclusionQueries = false;
 bool GRHIRequiresRenderTargetForPixelShaderUAVs = false;
 bool GRHISupportsUAVFormatAliasing = false;
+bool GRHISupportsTextureViews = true;
 bool GRHISupportsDirectGPUMemoryLock = false;
 bool GRHISupportsMultithreadedShaderCreation = true;
 bool GRHISupportsMultithreadedResources = false;
diff --git a/Engine/Source/Runtime/RHI/Public/RHI.h b/Engine/Source/Runtime/RHI/Public/RHI.h
--- a/Engine/Source/Runtime/RHI/Public/RHI.h
+++ b/Engine/Source/Runtime/RHI/Public/RHI.h
@@ -323,6 +323,9 @@ extern RHI_API bool GRHIRequiresRenderTargetForPixelShaderUAVs;
 /** true if the RHI supports unordered access view format aliasing */
 extern RHI_API bool GRHISupportsUAVFormatAliasing;
 
+/** true if the RHI supports texture views (data aliasing) */
+extern RHI_API bool GRHISupportsTextureViews;
+
 /** true if the pointer returned by Lock is a persistent direct pointer to gpu memory */
 extern RHI_API bool GRHISupportsDirectGPUMemoryLock;
 
diff --git a/Engine/Source/Runtime/RenderCore/Public/VirtualTexturing.h b/Engine/Source/Runtime/RenderCore/Public/VirtualTexturing.h
--- a/Engine/Source/Runtime/RenderCore/Public/VirtualTexturing.h
+++ b/Engine/Source/Runtime/RenderCore/Public/VirtualTexturing.h
@@ -176,6 +176,7 @@ struct FVTProducerDescription
 	uint8 NumTextureLayers = 0u;
 	TEnumAsByte<EPixelFormat> LayerFormat[VIRTUALTEXTURE_SPACE_MAXLAYERS] = { PF_Unknown };
 	FLinearColor LayerFallbackColor[VIRTUALTEXTURE_SPACE_MAXLAYERS] = { FLinearColor::Black };
+	bool bIsLayerSRGB[VIRTUALTEXTURE_SPACE_MAXLAYERS] = { false };
 	
 	uint8 NumPhysicalGroups = 0u;
 	uint8 PhysicalGroupIndex[VIRTUALTEXTURE_SPACE_MAXLAYERS] = { 0 };
diff --git a/Engine/Source/Runtime/Renderer/Private/VT/RuntimeVirtualTextureRender.cpp b/Engine/Source/Runtime/Renderer/Private/VT/RuntimeVirtualTextureRender.cpp
--- a/Engine/Source/Runtime/Renderer/Private/VT/RuntimeVirtualTextureRender.cpp
+++ b/Engine/Source/Runtime/Renderer/Private/VT/RuntimeVirtualTextureRender.cpp
@@ -433,7 +433,7 @@ namespace RuntimeVirtualTexture
 		static void SetupRenderTargetsInfo(ERuntimeVirtualTextureMaterialType MaterialType, ERHIFeatureLevel::Type FeatureLevel, bool bLQFormat, FGraphicsPipelineRenderTargetsInfo& RenderTargetsInfo)
 		{
 			const ETextureCreateFlags RTCreateFlags = TexCreate_RenderTargetable | TexCreate_ShaderResource;
-			const ETextureCreateFlags RTSrgbFlags = FeatureLevel > ERHIFeatureLevel::ES3_1 ? TexCreate_SRGB : TexCreate_None;
+			const ETextureCreateFlags RTSrgbFlags = TexCreate_SRGB;
 
 			switch (MaterialType)
 			{
@@ -501,9 +501,7 @@ namespace RuntimeVirtualTexture
 
 			const ETextureCreateFlags RTClearFlags = Desc.bClearTextures ? TexCreate_None : RTNoClearHackFlags;
 			const ETextureCreateFlags RTCreateFlags = TexCreate_RenderTargetable | TexCreate_ShaderResource | RTClearFlags;
-			
-			// Not all mobile RHIs support sRGB texture views/aliasing, use only linear targets on mobile
-			const ETextureCreateFlags RTSrgbFlags = Desc.FeatureLevel > ERHIFeatureLevel::ES3_1 ? TexCreate_SRGB : TexCreate_None;
+			const ETextureCreateFlags RTSrgbFlags = TexCreate_SRGB;
 
 			const EPixelFormat Compressed64BitFormat = UseRGBA16(GMaxRHIShaderPlatform) ? PF_R16G16B16A16_UINT : PF_R32G32_UINT;
 			const EPixelFormat Compressed128BitFormat = PF_R32G32B32A32_UINT;
diff --git a/Engine/Source/Runtime/Renderer/Private/VT/VirtualTexturePhysicalSpace.cpp b/Engine/Source/Runtime/Renderer/Private/VT/VirtualTexturePhysicalSpace.cpp
--- a/Engine/Source/Runtime/Renderer/Private/VT/VirtualTexturePhysicalSpace.cpp
+++ b/Engine/Source/Runtime/Renderer/Private/VT/VirtualTexturePhysicalSpace.cpp
@@ -75,6 +75,13 @@ FVirtualTexturePhysicalSpace::FVirtualTexturePhysicalSpace(const FVTPhysicalSpac
 	for (uint32 Layer = 0u; Layer < Description.NumLayers; ++Layer)
 	{
 		FormatString += GPixelFormats[Description.Format[Layer]].Name;
+		
+		// sRGB flag is only relevant on platforms that do not support texture views
+		if (!GRHISupportsTextureViews)
+		{
+			FormatString += (Description.bHasLayerSrgbView[Layer] ? TEXT(" (sRGB)") : TEXT(" (Linear)"));
+		}
+
 		if (Layer + 1u < Description.NumLayers)
 		{
 			FormatString += TEXT(", ");
@@ -121,11 +128,11 @@ void FVirtualTexturePhysicalSpace::InitRHI()
 		const EPixelFormat FormatUAV = GetUnorderedAccessViewFormat(FormatSRV);
 		const bool bCreateAliasedUAV = (FormatUAV != PF_Unknown) && (FormatUAV != FormatSRV);
 		
-		// Not all mobile RHIs support sRGB views/aliasing, use only linear targets on mobile.
-		const bool bFeatureLevelSupportsSRGB = GetFeatureLevel() > ERHIFeatureLevel::ES3_1;
+		// Not all RHIs support sRGB views/aliasing. On those platforms create texture in an expected storage format 
+		const bool bDefaultToSRGB = Description.bHasLayerSrgbView[Layer];
 		// Not all formats support sRGB.
 		const bool bFormatSupportsSRGB = FormatSRV != EPixelFormat::PF_R5G6B5_UNORM && FormatSRV != EPixelFormat::PF_B5G5R5A1_UNORM && FormatSRV != EPixelFormat::PF_G16;
-		ETextureCreateFlags VT_SRGB = (bFeatureLevelSupportsSRGB && bFormatSupportsSRGB) ? TexCreate_SRGB : TexCreate_None;
+		ETextureCreateFlags VT_SRGB = (bDefaultToSRGB && bFormatSupportsSRGB) ? TexCreate_SRGB : TexCreate_None;
 
 		// Allocate physical texture from the render target pool
 		const uint32 TextureSize = GetTextureSize();
diff --git a/Engine/Source/Runtime/Renderer/Private/VT/VirtualTexturePhysicalSpace.h b/Engine/Source/Runtime/Renderer/Private/VT/VirtualTexturePhysicalSpace.h
--- a/Engine/Source/Runtime/Renderer/Private/VT/VirtualTexturePhysicalSpace.h
+++ b/Engine/Source/Runtime/Renderer/Private/VT/VirtualTexturePhysicalSpace.h
@@ -16,6 +16,7 @@ struct FVTPhysicalSpaceDescription
 	uint8 Dimensions;
 	uint8 NumLayers;
 	TEnumAsByte<EPixelFormat> Format[VIRTUALTEXTURE_SPACE_MAXLAYERS];
+	bool bHasLayerSrgbView[VIRTUALTEXTURE_SPACE_MAXLAYERS];
 	bool bContinuousUpdate;
 };
 
@@ -34,6 +35,12 @@ inline bool operator==(const FVTPhysicalSpaceDescription& Lhs, const FVTPhysical
 		{
 			return false;
 		}
+
+		if (Lhs.bHasLayerSrgbView[Layer] != Rhs.bHasLayerSrgbView[Layer])
+		{
+			return false;
+		}
+
 	}
 	return true;
 }
@@ -51,6 +58,7 @@ inline uint32 GetTypeHash(const FVTPhysicalSpaceDescription& Desc)
 	for (int32 Layer = 0; Layer < Desc.NumLayers; ++Layer)
 	{
 		Hash = HashCombine(Hash, GetTypeHash(Desc.Format[Layer]));
+		Hash = HashCombine(Hash, GetTypeHash(Desc.bHasLayerSrgbView[Layer]));
 	}
 	return Hash;
 }
diff --git a/Engine/Source/Runtime/Renderer/Private/VT/VirtualTextureProducer.cpp b/Engine/Source/Runtime/Renderer/Private/VT/VirtualTextureProducer.cpp
--- a/Engine/Source/Runtime/Renderer/Private/VT/VirtualTextureProducer.cpp
+++ b/Engine/Source/Runtime/Renderer/Private/VT/VirtualTextureProducer.cpp
@@ -85,6 +85,8 @@ FVirtualTextureProducerHandle FVirtualTextureProducerCollection::RegisterProduce
 			if (InDesc.PhysicalGroupIndex[LayerIndex] == PhysicalGroupIndex)
 			{
 				PhysicalSpaceDesc.Format[PhysicalSpaceDesc.NumLayers] = InDesc.LayerFormat[LayerIndex];
+				// sRGB on/off flag is ignored on platforms that has support for texture aliasing
+				PhysicalSpaceDesc.bHasLayerSrgbView[PhysicalSpaceDesc.NumLayers] = (GRHISupportsTextureViews ? true : InDesc.bIsLayerSRGB[LayerIndex]);
 				PhysicalSpaceDesc.NumLayers++;
 			}
 		}
