Improve FRayTracingSceneAddInstancesTask performance by caching some of the data from FPrimitiveSceneInfo in FRayTracingRelevantPrimitive during LOD calculation task.

* This avoids some of the cache misses on the critical path (serial) task when processing a large number of scene primitives that use cached HWRT instance data (ERayTracingPrimitiveFlags::CacheInstances)

diff --git a/Engine/Source/Runtime/Renderer/Private/DeferredShadingRenderer.cpp b/Engine/Source/Runtime/Renderer/Private/DeferredShadingRenderer.cpp
--- a/Engine/Source/Runtime/Renderer/Private/DeferredShadingRenderer.cpp
+++ b/Engine/Source/Runtime/Renderer/Private/DeferredShadingRenderer.cpp
@@ -697,7 +697,6 @@ static void AddDebugRayTracingInstanceFlags(ERayTracingInstanceFlags& InOutFlags
 struct FRayTracingRelevantPrimitive
 {
 	FRHIRayTracingGeometry* RayTracingGeometryRHI = nullptr;
-	TArrayView<const int32> CachedRayTracingMeshCommandIndices; // Pointer to FPrimitiveSceneInfo::CachedRayTracingMeshCommandIndicesPerLOD data
 	uint64 StateHash = 0;
 	int32 PrimitiveIndex = -1;
 	int8 LODIndex = -1;
@@ -711,6 +710,10 @@ struct FRayTracingRelevantPrimitive
 	bool bIsSky = false;
 	bool bAllSegmentsTranslucent = true;
 
+	bool bCachedRayTracingGeometryValid = true;
+	const FRayTracingGeometryInstance* CachedRayTracingInstance = nullptr;
+	TArrayView<const int32> CachedRayTracingMeshCommandIndices; // Pointer to FPrimitiveSceneInfo::CachedRayTracingMeshCommandIndicesPerLOD data
+
 	uint64 InstancingKey() const
 	{
 		uint64 Key = StateHash;
@@ -907,7 +910,26 @@ static void GatherRayTracingRelevantPrimitives(const FScene& Scene, const FViewI
 				LODIndex = LODToRender.GetRayTracedLOD();
 			}
 
-			if (!EnumHasAllFlags(Flags, ERayTracingPrimitiveFlags::CacheInstances))
+
+			if (EnumHasAnyFlags(Flags, ERayTracingPrimitiveFlags::CacheInstances))
+			{
+				RelevantPrimitive.bCachedRayTracingGeometryValid = SceneInfo->IsCachedRayTracingGeometryValid();
+				RelevantPrimitive.bAnySegmentsDecal = SceneInfo->bCachedRayTracingInstanceAnySegmentsDecal;
+				RelevantPrimitive.bAllSegmentsDecal = SceneInfo->bCachedRayTracingInstanceAllSegmentsDecal;
+				RelevantPrimitive.CachedRayTracingInstance = &SceneInfo->CachedRayTracingInstance;
+
+				checkf(SceneInfo->CachedRayTracingInstance.GeometryRHI, TEXT("Ray tracing instance must have a valid geometry."));
+
+				// For primitives with ERayTracingPrimitiveFlags::CacheInstances flag we only cache the instance/mesh commands of the current LOD
+				// (see FPrimitiveSceneInfo::UpdateCachedRayTracingInstance(...) and CacheRayTracingPrimitive(...))
+				LODIndex = 0;
+
+				if (SceneInfo->CachedRayTracingMeshCommandIndicesPerLOD.IsValidIndex(LODIndex))
+				{
+					RelevantPrimitive.CachedRayTracingMeshCommandIndices = SceneInfo->CachedRayTracingMeshCommandIndicesPerLOD[LODIndex];
+				}
+			}
+			else
 			{
 				FRHIRayTracingGeometry* RayTracingGeometryInstance = SceneInfo->GetStaticRayTracingGeometryInstance(LODIndex);
 				if (RayTracingGeometryInstance == nullptr)
@@ -1460,23 +1482,19 @@ bool FDeferredShadingSceneRenderer::GatherRayTracingWorldInstancesForView(FRDGBu
 
 				if (EnumHasAnyFlags(Flags, ERayTracingPrimitiveFlags::CacheInstances))
 				{
-					// For primitives with ERayTracingPrimitiveFlags::CacheInstances flag we only cache the instance/mesh commands of the current LOD
-					// (see FPrimitiveSceneInfo::UpdateCachedRayTracingInstance(...) and CacheRayTracingPrimitive(...))
-					const int32 LODIndex = 0;
-
 					const bool bUsingNaniteRayTracing = (Nanite::GetRayTracingMode() != Nanite::ERayTracingMode::Fallback) && SceneProxy->IsNaniteMesh();
 
 					if (bUsingNaniteRayTracing)
 					{
 						Nanite::GRayTracingManager.AddVisiblePrimitive(SceneInfo);
 
-						if (SceneInfo->CachedRayTracingInstance.GeometryRHI == nullptr)
+						if (RelevantPrimitive.CachedRayTracingInstance->GeometryRHI == nullptr)
 						{
 							// Nanite ray tracing geometry not ready yet, doesn't include primitive in ray tracing scene
 							continue;
 						}
 					}
-					else if (!SceneInfo->IsCachedRayTracingGeometryValid())
+					else if (!RelevantPrimitive.bCachedRayTracingGeometryValid)
 					{
 						// cached instance is not valid (eg: was streamed out) need to invalidate for next frame
 						ProxiesWithDirtyCachedInstance.Add(Scene.PrimitiveSceneProxies[PrimitiveIndex]);
@@ -1489,86 +1507,86 @@ bool FDeferredShadingSceneRenderer::GatherRayTracingWorldInstancesForView(FRDGBu
 					// one containing non-decal segments and the other with decal segments
 					// masking of segments is done using "hidden" hitgroups
 					// TODO: Debug Visualization to highlight primitives using this?
-					const bool bNeedSeparateDecalInstance = SceneInfo->bCachedRayTracingInstanceAnySegmentsDecal && !SceneInfo->bCachedRayTracingInstanceAllSegmentsDecal;
+					const bool bNeedSeparateDecalInstance = RelevantPrimitive.bAnySegmentsDecal && !RelevantPrimitive.bAllSegmentsDecal;
 
-					if (GRayTracingExcludeDecals && SceneInfo->bCachedRayTracingInstanceAnySegmentsDecal && !bNeedSeparateDecalInstance)
+					if (GRayTracingExcludeDecals && RelevantPrimitive.bAnySegmentsDecal && !bNeedSeparateDecalInstance)
 					{
 						continue;
 					}
 
-					checkf(SceneInfo->CachedRayTracingInstance.GeometryRHI, TEXT("Ray tracing instance must have a valid geometry."));
-
-					FRayTracingGeometryInstance NewInstance = SceneInfo->CachedRayTracingInstance;
-					NewInstance.LayerIndex = (uint8)(SceneInfo->bCachedRayTracingInstanceAnySegmentsDecal && !bNeedSeparateDecalInstance ? ERayTracingSceneLayer::Decals : ERayTracingSceneLayer::Base);
+					check(RelevantPrimitive.CachedRayTracingInstance);
 
-					const Experimental::FHashElementId GroupId = Scene.PrimitiveRayTracingGroupIds[PrimitiveIndex];
-					const bool bUseGroupBounds = CullingParameters.bCullUsingGroupIds && GroupId.IsValid();
+					const int32 NewInstanceIndex = RayTracingScene.AddInstance(*RelevantPrimitive.CachedRayTracingInstance, SceneProxy, false);
+					uint32 DecalInstanceIndex = INDEX_NONE;
 
-					if (CullingParameters.CullInRayTracing > 0 && GetRayTracingCullingPerInstance() && SceneInfo->CachedRayTracingInstance.NumTransforms > 1 && !bUseGroupBounds)
 					{
-						const bool bIsFarFieldPrimitive = EnumHasAnyFlags(Flags, ERayTracingPrimitiveFlags::FarField);
+						FRayTracingGeometryInstance& NewInstance = RayTracingScene.GetInstance(NewInstanceIndex);
+						AddDebugRayTracingInstanceFlags(NewInstance.Flags);
 
-						TArrayView<uint32> InstanceActivationMask = RayTracingScene.Allocate<uint32>(FMath::DivideAndRoundUp(NewInstance.NumTransforms, 32u));
+						NewInstance.LayerIndex = (uint8)(RelevantPrimitive.bAnySegmentsDecal && !bNeedSeparateDecalInstance ? ERayTracingSceneLayer::Decals : ERayTracingSceneLayer::Base);
 
-						NewInstance.ActivationMask = InstanceActivationMask;
+						const Experimental::FHashElementId GroupId = Scene.PrimitiveRayTracingGroupIds[PrimitiveIndex];
+						const bool bUseGroupBounds = CullingParameters.bCullUsingGroupIds && GroupId.IsValid();
 
-						FRayTracingCullPrimitiveInstancesClosure Closure;
-						Closure.Scene = &Scene;
-						Closure.SceneInfo = SceneInfo;
-						Closure.PrimitiveIndex = PrimitiveIndex;
-						Closure.bIsFarFieldPrimitive = bIsFarFieldPrimitive;
-						Closure.CullingParameters = &CullingParameters;
-						Closure.OutInstanceActivationMask = InstanceActivationMask;
+						if (CullingParameters.CullInRayTracing > 0 && GetRayTracingCullingPerInstance() && RelevantPrimitive.CachedRayTracingInstance->NumTransforms > 1 && !bUseGroupBounds)
+						{
+							const bool bIsFarFieldPrimitive = EnumHasAnyFlags(Flags, ERayTracingPrimitiveFlags::FarField);
 
-						CullInstancesClosures.Add(MoveTemp(Closure));
+							TArrayView<uint32> InstanceActivationMask = RayTracingScene.Allocate<uint32>(FMath::DivideAndRoundUp(NewInstance.NumTransforms, 32u));
 
-						if (CullInstancesClosures.Num() >= 256)
-						{
-							CullingTasks.Add(FFunctionGraphTask::CreateAndDispatchWhenReady([CullInstancesClosures = MoveTemp(CullInstancesClosures)]()
-							{
-								for (auto& Closure : CullInstancesClosures)
-								{
-									Closure();
-								}
-							}, TStatId(), nullptr, ENamedThreads::AnyThread));
-						}
-					}
+							NewInstance.ActivationMask = InstanceActivationMask;
 
-					AddDebugRayTracingInstanceFlags(NewInstance.Flags);
+							FRayTracingCullPrimitiveInstancesClosure Closure;
+							Closure.Scene = &Scene;
+							Closure.SceneInfo = SceneInfo;
+							Closure.PrimitiveIndex = PrimitiveIndex;
+							Closure.bIsFarFieldPrimitive = bIsFarFieldPrimitive;
+							Closure.CullingParameters = &CullingParameters;
+							Closure.OutInstanceActivationMask = InstanceActivationMask;
 
-					const int32 NewInstanceIndex = RayTracingScene.AddInstance(NewInstance, SceneInfo->Proxy, false);
+							CullInstancesClosures.Add(MoveTemp(Closure));
 
-					uint32 DecalInstanceIndex = INDEX_NONE;
-					if (bNeedSeparateDecalInstance && !GRayTracingExcludeDecals)
-					{
-						FRayTracingGeometryInstance DecalRayTracingInstance = NewInstance;
-						DecalRayTracingInstance.LayerIndex = (uint8)ERayTracingSceneLayer::Decals;
+							if (CullInstancesClosures.Num() >= 256)
+							{
+								CullingTasks.Add(FFunctionGraphTask::CreateAndDispatchWhenReady([CullInstancesClosures = MoveTemp(CullInstancesClosures)]()
+									{
+										for (auto& Closure : CullInstancesClosures)
+										{
+											Closure();
+										}
+									}, TStatId(), nullptr, ENamedThreads::AnyThread));
+							}
+						}
+
+						if (bNeedSeparateDecalInstance && !GRayTracingExcludeDecals)
+						{
+							FRayTracingGeometryInstance DecalRayTracingInstance = NewInstance;
+							DecalRayTracingInstance.LayerIndex = (uint8)ERayTracingSceneLayer::Decals;
 
-						DecalInstanceIndex = RayTracingScene.AddInstance(MoveTemp(DecalRayTracingInstance), SceneInfo->Proxy, false);
+							DecalInstanceIndex = RayTracingScene.AddInstance(MoveTemp(DecalRayTracingInstance), SceneProxy, false);
+						}
 					}
 
 					// At the moment we only support SM & ISMs on this path
 					check(EnumHasAnyFlags(Flags, ERayTracingPrimitiveFlags::CacheMeshCommands));
-					if (SceneInfo->CachedRayTracingMeshCommandIndicesPerLOD.Num() > 0 && SceneInfo->CachedRayTracingMeshCommandIndicesPerLOD[LODIndex].Num() > 0)
+					
+					const bool bHasDecalInstanceIndex = DecalInstanceIndex != INDEX_NONE;
+
+					for (int32 CommandIndex : RelevantPrimitive.CachedRayTracingMeshCommandIndices)
 					{
-						const bool bHasDecalInstanceIndex = DecalInstanceIndex != INDEX_NONE;
+						const FRayTracingMeshCommand& MeshCommand = Scene.CachedRayTracingMeshCommands[CommandIndex];
 
-						for (int32 CommandIndex : SceneInfo->CachedRayTracingMeshCommandIndicesPerLOD[LODIndex])
 						{
-							const FRayTracingMeshCommand& MeshCommand = Scene.CachedRayTracingMeshCommands[CommandIndex];
-
-							{
-								const bool bHidden = bHasDecalInstanceIndex && MeshCommand.bDecal;
-								FVisibleRayTracingMeshCommand NewVisibleMeshCommand(&MeshCommand, NewInstanceIndex, bHidden);
-								VisibleRayTracingMeshCommands.Add(NewVisibleMeshCommand);
-							}
+							const bool bHidden = bHasDecalInstanceIndex && MeshCommand.bDecal;
+							FVisibleRayTracingMeshCommand NewVisibleMeshCommand(&MeshCommand, NewInstanceIndex, bHidden);
+							VisibleRayTracingMeshCommands.Add(NewVisibleMeshCommand);
+						}
 
-							if(bHasDecalInstanceIndex)
-							{
-								const bool bHidden = !MeshCommand.bDecal;
-								FVisibleRayTracingMeshCommand NewVisibleMeshCommand(&MeshCommand, DecalInstanceIndex, bHidden);
-								VisibleRayTracingMeshCommands.Add(NewVisibleMeshCommand);
-							}
+						if(bHasDecalInstanceIndex)
+						{
+							const bool bHidden = !MeshCommand.bDecal;
+							FVisibleRayTracingMeshCommand NewVisibleMeshCommand(&MeshCommand, DecalInstanceIndex, bHidden);
+							VisibleRayTracingMeshCommands.Add(NewVisibleMeshCommand);
 						}
 					}
 				}
@@ -1664,14 +1682,14 @@ bool FDeferredShadingSceneRenderer::GatherRayTracingWorldInstancesForView(FRDGBu
 
 						RayTracingInstance.LayerIndex = (uint8)(RelevantPrimitive.bAnySegmentsDecal && !bNeedSeparateDecalInstance ? ERayTracingSceneLayer::Decals : ERayTracingSceneLayer::Base);
 
-						InstanceBatch.Index = RayTracingScene.AddInstance(RayTracingInstance, SceneInfo->Proxy, false);
+						InstanceBatch.Index = RayTracingScene.AddInstance(RayTracingInstance, SceneProxy, false);
 
 						if (bNeedSeparateDecalInstance && !GRayTracingExcludeDecals)
 						{
 							FRayTracingGeometryInstance DecalRayTracingInstance = RayTracingInstance;
 							DecalRayTracingInstance.LayerIndex = (uint8)ERayTracingSceneLayer::Decals;
 
-							InstanceBatch.DecalIndex = RayTracingScene.AddInstance(MoveTemp(DecalRayTracingInstance), SceneInfo->Proxy, false);
+							InstanceBatch.DecalIndex = RayTracingScene.AddInstance(MoveTemp(DecalRayTracingInstance), SceneProxy, false);
 						}
 
 						const bool bHasDecalInstanceIndex = InstanceBatch.DecalIndex != INDEX_NONE;
