Fix spotlight angle precision issue, compress local light color to move light id

diff --git a/Engine/Shaders/Private/ClusteredDeferredShadingPixelShader.usf b/Engine/Shaders/Private/ClusteredDeferredShadingPixelShader.usf
--- a/Engine/Shaders/Private/ClusteredDeferredShadingPixelShader.usf
+++ b/Engine/Shaders/Private/ClusteredDeferredShadingPixelShader.usf
@@ -182,8 +182,8 @@ float4 GetSimpleLightLighting(
 		FSimpleDeferredLightData SimpleLightData = (FSimpleDeferredLightData)0;
 		SimpleLightData.TranslatedWorldPosition = LocalLight.LightPositionAndInvRadius.xyz;
 		SimpleLightData.InvRadius = LocalLight.LightPositionAndInvRadius.w;
-		SimpleLightData.Color = LocalLight.LightColorAndFalloffExponent.xyz;
-		SimpleLightData.FalloffExponent = LocalLight.LightColorAndFalloffExponent.w;
+		SimpleLightData.Color = UnpackLightColor(LocalLight.LightColorAndIdAndFalloffExponent.xy);
+		SimpleLightData.FalloffExponent = LocalLight.LightColorAndIdAndFalloffExponent.w;
 		SimpleLightData.bInverseSquared = SimpleLightData.FalloffExponent == 0;
 		FStrataDeferredLighting StrataLighting = GetSimpleDynamicLighting(TranslatedWorldPosition, CameraVector, ScreenSpaceData.GBuffer.WorldNormal, ScreenSpaceData.AmbientOcclusion, ScreenSpaceData.GBuffer.DiffuseColor, ScreenSpaceData.GBuffer.SpecularColor, ScreenSpaceData.GBuffer.Roughness, SimpleLightData);
 
diff --git a/Engine/Shaders/Private/LightData.ush b/Engine/Shaders/Private/LightData.ush
--- a/Engine/Shaders/Private/LightData.ush
+++ b/Engine/Shaders/Private/LightData.ush
@@ -72,8 +72,8 @@ struct FLocalLightData
 {
 	// Raw value form light buffers
 	float4 LightPositionAndInvRadius;
-	float4 LightColorAndFalloffExponent;
-	float4 SpotAnglesAndIdAndSourceRadiusPacked;
+	float4 LightColorAndIdAndFalloffExponent;
+	float4 SpotAnglesAndSourceRadiusPacked;
 	float4 LightDirectionAndShadowMask;
 	float4 LightTangentAndIESDataAndSpecularScale;
 	float4 RectDataAndVirtualShadowMapId;
@@ -183,6 +183,16 @@ struct FDeferredLightData
 /**
  * Helper to unpack the shadow mask part of the FLocalLightData::LightDirectionAndShadowMask.w
  */
+
+float3 UnpackLightColor(float2 In)
+{
+	float3 Dir = float3(
+		((asuint(In.y) >> 0) & 0x3FF),
+		((asuint(In.y) >> 10) & 0x3FF),
+		((asuint(In.y) >> 20) & 0x3FF));
+	return In.x * Dir;
+}
+
 float4 UnpackShadowMapChannelMask(uint In)
 {
 	// 8 ALU
@@ -212,22 +222,22 @@ bool UnpackCastShadow(uint InShadowMapChannelMaskPacked)
 
 float UnpackVolumetricScatteringIntensity(FLocalLightData In)
 {
-	return f16tof32(asuint(In.SpotAnglesAndIdAndSourceRadiusPacked.w) >> 16);
+	return f16tof32(asuint(In.SpotAnglesAndSourceRadiusPacked.w) >> 16);
 }
 
 float UnpackLightSourceRadius(FLocalLightData In)
 {
-	return f16tof32(asuint(In.SpotAnglesAndIdAndSourceRadiusPacked.z) & 0xFFFF);
+	return f16tof32(asuint(In.SpotAnglesAndSourceRadiusPacked.z) & 0xFFFF);
 }
 
 float UnpackLightSoftSourceRadius(FLocalLightData In)
 {
-	return f16tof32(asuint(In.SpotAnglesAndIdAndSourceRadiusPacked.z) >> 16);
+	return f16tof32(asuint(In.SpotAnglesAndSourceRadiusPacked.z) >> 16);
 }
 
 float UnpackLightSourceLength(FLocalLightData In)
 {
-	return f16tof32(asuint(In.SpotAnglesAndIdAndSourceRadiusPacked.w));
+	return f16tof32(asuint(In.SpotAnglesAndSourceRadiusPacked.w));
 }
 
 float GetLightSourceRadius(FLocalLightData In)
@@ -237,9 +247,7 @@ float GetLightSourceRadius(FLocalLightData In)
 
 float2 GetLightSpotAngles(FLocalLightData In)
 {
-	uint SpotAnglesPacked = asuint(In.SpotAnglesAndIdAndSourceRadiusPacked.x);
-	float2 SpotAngles = float2(f16tof32(SpotAnglesPacked & 0xFFFF), f16tof32(SpotAnglesPacked >> 16));
-	return SpotAngles;
+	return In.SpotAnglesAndSourceRadiusPacked.xy;
 }
 
 float UnpackLightSpecularScale(FLocalLightData In)
@@ -278,8 +286,8 @@ FDeferredLightData ConvertToDeferredLight(
 
 	Out.TranslatedWorldPosition = In.LightPositionAndInvRadius.xyz;
 	Out.InvRadius				= In.LightPositionAndInvRadius.w;
-	Out.Color					= In.LightColorAndFalloffExponent.xyz;
-	Out.FalloffExponent			= In.LightColorAndFalloffExponent.w;
+	Out.Color					= UnpackLightColor(In.LightColorAndIdAndFalloffExponent.xy);
+	Out.FalloffExponent			= In.LightColorAndIdAndFalloffExponent.w;
 	Out.Direction				= In.LightDirectionAndShadowMask.xyz;
 	Out.SpotAngles				= GetLightSpotAngles(In);
 	Out.SourceRadius			= UnpackLightSourceRadius(In);
@@ -365,8 +373,8 @@ FLightShaderParameters ConvertFromLocal(const FLocalLightData In)
 	FLightShaderParameters Out	= (FLightShaderParameters)0;
 	Out.TranslatedWorldPosition	= In.LightPositionAndInvRadius.xyz;
 	Out.InvRadius				= In.LightPositionAndInvRadius.w;
-	Out.Color					= In.LightColorAndFalloffExponent.xyz;
-	Out.FalloffExponent			= In.LightColorAndFalloffExponent.w;
+	Out.Color					= UnpackLightColor(In.LightColorAndIdAndFalloffExponent.xy);
+	Out.FalloffExponent			= In.LightColorAndIdAndFalloffExponent.w;
 	Out.Direction				= In.LightDirectionAndShadowMask.xyz;
 	Out.Tangent					= In.LightTangentAndIESDataAndSpecularScale.xyz;
 	Out.SpotAngles				= GetLightSpotAngles(In);
@@ -426,8 +434,8 @@ FSimpleDeferredLightData ConvertToSimpleLight(FLocalLightData In)
 	FSimpleDeferredLightData Out = (FSimpleDeferredLightData)0;
 	Out.TranslatedWorldPosition = In.LightPositionAndInvRadius.xyz;
 	Out.InvRadius				= In.LightPositionAndInvRadius.w;
-	Out.Color					= In.LightColorAndFalloffExponent.xyz;
-	Out.FalloffExponent			= In.LightColorAndFalloffExponent.w;
+	Out.Color					= UnpackLightColor(In.LightColorAndIdAndFalloffExponent.xy);
+	Out.FalloffExponent			= In.LightColorAndIdAndFalloffExponent.w;
 	Out.bInverseSquared			= Out.FalloffExponent == 0;
 	return Out;
 }
diff --git a/Engine/Shaders/Private/LightGridCommon.ush b/Engine/Shaders/Private/LightGridCommon.ush
--- a/Engine/Shaders/Private/LightGridCommon.ush
+++ b/Engine/Shaders/Private/LightGridCommon.ush
@@ -166,13 +166,13 @@ FLocalLightData GetLocalLightData_Internal(
 	Out.bLumenLightSupported					= LocalLightIndex >= LumenSupportedStartIndex;
 	Out.bIsSimpleLight							= LocalLightIndex <  SimpleLightsEndIndex;
 	Out.LightPositionAndInvRadius				= InData0;
-	Out.LightColorAndFalloffExponent			= InData1;
+	Out.LightColorAndIdAndFalloffExponent		= InData1;
 	Out.LightDirectionAndShadowMask				= InData2;
-	Out.SpotAnglesAndIdAndSourceRadiusPacked	= InData3;
+	Out.SpotAnglesAndSourceRadiusPacked			= InData3;
 	Out.LightTangentAndIESDataAndSpecularScale	= InData4;
 	Out.RectDataAndVirtualShadowMapId			= InData5;
 	Out.VirtualShadowMapId						= int(Out.RectDataAndVirtualShadowMapId.z);
-	Out.LightSceneId							= int(Out.SpotAnglesAndIdAndSourceRadiusPacked.y);
+	Out.LightSceneId							= int(Out.LightColorAndIdAndFalloffExponent.z);
 
 	return Out;
 }
diff --git a/Engine/Shaders/Private/VirtualShadowMaps/VirtualShadowMapProjection.usf b/Engine/Shaders/Private/VirtualShadowMaps/VirtualShadowMapProjection.usf
--- a/Engine/Shaders/Private/VirtualShadowMaps/VirtualShadowMapProjection.usf
+++ b/Engine/Shaders/Private/VirtualShadowMaps/VirtualShadowMapProjection.usf
@@ -223,13 +223,6 @@ FVirtualShadowMapSampleResult ProjectLight(
 	bool bInLightRadius = InvDist >= Light.InvRadius;
 	bool bInLightCone = SpotAttenuationMask(L, -Light.Direction, Light.SpotAngles) > 0.0f;
 
-	#if ONE_PASS_PROJECTION
-		// Hacky workaround for spot angles being quantized to 16-bit when they are loaded from the light grid!
-		// This avoids potential artifacts at cone edges, and in this path we have the tile culling to lessen the
-		// overhead. Still needs to be fixed though as it is not trivial overhead!
-		bInLightCone = true;
-	#endif
-
 	bool bInLightRegion = bInLightRadius && bInLightCone;
 #endif // DIRECTIONAL_LIGHT
 
diff --git a/Engine/Source/Runtime/Renderer/Private/BasePassRendering.h b/Engine/Source/Runtime/Renderer/Private/BasePassRendering.h
--- a/Engine/Source/Runtime/Renderer/Private/BasePassRendering.h
+++ b/Engine/Source/Runtime/Renderer/Private/BasePassRendering.h
@@ -49,9 +49,9 @@ class FForwardLocalLightData
 {
 public:
 	FVector4f LightPositionAndInvRadius;
-	FVector4f LightColorAndFalloffExponent;
+	FVector4f LightColorAndIdAndFalloffExponent;
 	FVector4f LightDirectionAndShadowMapChannelMask;
-	FVector4f SpotAnglesAndIdAndSourceRadiusPacked;
+	FVector4f SpotAnglesAndSourceRadiusPacked;
 	FVector4f LightTangentAndIESDataAndSpecularScale;
 	FVector4f RectDataAndVirtualShadowMapId;
 };
diff --git a/Engine/Source/Runtime/Renderer/Private/LightGridInjection.cpp b/Engine/Source/Runtime/Renderer/Private/LightGridInjection.cpp
--- a/Engine/Source/Runtime/Renderer/Private/LightGridInjection.cpp
+++ b/Engine/Source/Runtime/Renderer/Private/LightGridInjection.cpp
@@ -283,6 +283,24 @@ static uint32 PackRG16(float In0, float In1)
 	return uint32(FFloat16(In0).Encoded) | (uint32(FFloat16(In1).Encoded) << 16);
 }
 
+static FVector2f PackLightColor(const FVector3f& LightColor)
+{
+	FVector3f LightColorDir;
+	float LightColorLength;
+	LightColor.ToDirectionAndLength(LightColorDir, LightColorLength);
+
+	FVector2f LightColorPacked;
+	uint32 LightColorDirPacked = 
+		((static_cast<uint32>(LightColorDir.X * 0x3FF) & 0x3FF) <<  0) |
+		((static_cast<uint32>(LightColorDir.Y * 0x3FF) & 0x3FF) << 10) |
+		((static_cast<uint32>(LightColorDir.Z * 0x3FF) & 0x3FF) << 20);
+
+	LightColorPacked.X = LightColorLength / 0x3FF;
+	*(uint32*)(&LightColorPacked.Y) = LightColorDirPacked;
+
+	return LightColorPacked;
+}
+
 static void PackLocalLightData(
 	FForwardLocalLightData& Out,
 	const FViewInfo& View,
@@ -314,18 +332,18 @@ static void PackLocalLightData(
 	const uint32 RectPackedY = 0;
 	const uint32 RectPackedW = FFloat16(-2.f).Encoded;
 
-	// Pack spot angles in 2x f16
-	uint32 SpotAnglesPacked = PackRG16(-2, 1);
-
 	// Pack specular scale and IES profile index
 	const float SpecularScale = 1.f;
 	const float IESAtlasIndex = INDEX_NONE;
 	const uint32 SpecularScaleAndIESData = PackRG16(SpecularScale, IESAtlasIndex);
 
+	const FVector3f LightColor = (FVector3f)SimpleLight.Color * FLightRenderParameters::GetLightExposureScale(View.GetLastEyeAdaptationExposure(), SimpleLight.InverseExposureBlend);
+	const FVector2f LightColorPacked = PackLightColor(LightColor);
+
 	Out.LightPositionAndInvRadius				= FVector4f(LightTranslatedWorldPosition, 1.0f / FMath::Max(SimpleLight.Radius, KINDA_SMALL_NUMBER));
-	Out.LightColorAndFalloffExponent			= FVector4f((FVector3f)SimpleLight.Color * FLightRenderParameters::GetLightExposureScale(View.GetLastEyeAdaptationExposure(), SimpleLight.InverseExposureBlend), SimpleLight.Exponent);
+	Out.LightColorAndIdAndFalloffExponent		= FVector4f(LightColorPacked.X, LightColorPacked.Y, INDEX_NONE, SimpleLight.Exponent);
 	Out.LightDirectionAndShadowMapChannelMask	= FVector4f(FVector3f(1, 0, 0), FMath::AsFloat(ShadowMapChannelMask));
-	Out.SpotAnglesAndIdAndSourceRadiusPacked	= FVector4f(FMath::AsFloat(SpotAnglesPacked), INDEX_NONE, FMath::AsFloat(PackedZ), FMath::AsFloat(PackedW));
+	Out.SpotAnglesAndSourceRadiusPacked			= FVector4f(-2, 1, FMath::AsFloat(PackedZ), FMath::AsFloat(PackedW));
 	Out.LightTangentAndIESDataAndSpecularScale	= FVector4f(1.0f, 0.0f, 0.0f, FMath::AsFloat(SpecularScaleAndIESData));
 	Out.RectDataAndVirtualShadowMapId			= FVector4f(FMath::AsFloat(RectPackedX), FMath::AsFloat(RectPackedY), -1, FMath::AsFloat(RectPackedW));
 }
@@ -355,18 +373,18 @@ static void PackLocalLightData(
 	RectPackedW |= uint32(FMath::Clamp(LightParameters.RectLightBarnCosAngle,  0.f, 1.0f) * 0x3FF) << 16;	// 10 bits
 	RectPackedW |= uint32(FMath::Clamp(LightParameters.RectLightAtlasMaxLevel, 0.f, 63.f)) << 26;			//  6 bits
 
-	// Pack spot angles in 2x f16
-	uint32 SpotAnglesPacked = PackRG16(LightParameters.SpotAngles.X, LightParameters.SpotAngles.Y);
-
 	// Pack specular scale and IES profile index
 	const uint32 SpecularScaleAndIESData = PackRG16(LightParameters.SpecularScale, LightParameters.IESAtlasIndex);
 
+	const FVector2f LightColorPacked = PackLightColor(FVector3f(LightParameters.Color));
+
 	// NOTE: This cast of VirtualShadowMapId to float is not ideal, but bitcast has issues here with INDEX_NONE -> NaN
 	// and 32-bit floats have enough mantissa to cover all reasonable numbers here for now.
+	// NOTE: SpotAngles needs full-precision for VSM one pass projection
 	Out.LightPositionAndInvRadius				= FVector4f(LightTranslatedWorldPosition, LightParameters.InvRadius);
-	Out.LightColorAndFalloffExponent			= FVector4f(LightParameters.Color, LightParameters.FalloffExponent);
+	Out.LightColorAndIdAndFalloffExponent		= FVector4f(LightColorPacked.X, LightColorPacked.Y, LightSceneId, LightParameters.FalloffExponent);
 	Out.LightDirectionAndShadowMapChannelMask	= FVector4f(LightParameters.Direction, FMath::AsFloat(LightTypeAndShadowMapChannelMaskPacked));
-	Out.SpotAnglesAndIdAndSourceRadiusPacked	= FVector4f(FMath::AsFloat(SpotAnglesPacked), LightSceneId, FMath::AsFloat(PackedZ), FMath::AsFloat(PackedW));
+	Out.SpotAnglesAndSourceRadiusPacked			= FVector4f(LightParameters.SpotAngles.X, LightParameters.SpotAngles.Y, FMath::AsFloat(PackedZ), FMath::AsFloat(PackedW));
 	Out.LightTangentAndIESDataAndSpecularScale	= FVector4f(LightParameters.Tangent, FMath::AsFloat(SpecularScaleAndIESData));
 	Out.RectDataAndVirtualShadowMapId			= FVector4f(FMath::AsFloat(RectPackedX), FMath::AsFloat(RectPackedY), float(VirtualShadowMapId), FMath::AsFloat(RectPackedW));
 
