Add second VSM resolution bias for moving lights along with tracking of what lights are moving. When a light stops is does a gradual transition to the second bias.
 - r.Shadow.Virtual.ResolutionLodBias[Local|Directiona]lMoving configures the LOD bias for a light that is currently moving.
 - Added FShadowScene that will contain persistent, event driven, shadow aspects of the scene.

diff --git a/Engine/Config/BaseScalability.ini b/Engine/Config/BaseScalability.ini
--- a/Engine/Config/BaseScalability.ini
+++ b/Engine/Config/BaseScalability.ini
@@ -132,7 +132,9 @@ r.LightMaxDrawDistanceScale=0
 r.CapsuleShadows=0
 r.Shadow.Virtual.MaxPhysicalPages=512
 r.Shadow.Virtual.ResolutionLodBiasDirectional=0.0
+r.Shadow.Virtual.ResolutionLodBiasDirectionalMoving=1.0
 r.Shadow.Virtual.ResolutionLodBiasLocal=1.0
+r.Shadow.Virtual.ResolutionLodBiasLocalMoving=2.0
 r.Shadow.Virtual.SMRT.RayCountDirectional=0
 r.Shadow.Virtual.SMRT.SamplesPerRayDirectional=2
 r.Shadow.Virtual.SMRT.RayCountLocal=0
@@ -154,7 +156,9 @@ r.LightMaxDrawDistanceScale=.5
 r.CapsuleShadows=1
 r.Shadow.Virtual.MaxPhysicalPages=512
 r.Shadow.Virtual.ResolutionLodBiasDirectional=0.0
+r.Shadow.Virtual.ResolutionLodBiasDirectionalMoving=1.0
 r.Shadow.Virtual.ResolutionLodBiasLocal=1.0
+r.Shadow.Virtual.ResolutionLodBiasLocalMoving=2.0
 r.Shadow.Virtual.SMRT.RayCountDirectional=4
 r.Shadow.Virtual.SMRT.SamplesPerRayDirectional=2
 r.Shadow.Virtual.SMRT.RayCountLocal=4
@@ -179,7 +183,9 @@ r.LightMaxDrawDistanceScale=1
 r.CapsuleShadows=1
 r.Shadow.Virtual.MaxPhysicalPages=2048
 r.Shadow.Virtual.ResolutionLodBiasDirectional=-0.5
+r.Shadow.Virtual.ResolutionLodBiasDirectionalMoving=0.5
 r.Shadow.Virtual.ResolutionLodBiasLocal=0.0
+r.Shadow.Virtual.ResolutionLodBiasLocalMoving=1.0
 r.Shadow.Virtual.SMRT.RayCountDirectional=8
 r.Shadow.Virtual.SMRT.SamplesPerRayDirectional=4
 r.Shadow.Virtual.SMRT.RayCountLocal=8
@@ -204,7 +210,9 @@ r.LightMaxDrawDistanceScale=1
 r.CapsuleShadows=1
 r.Shadow.Virtual.MaxPhysicalPages=4096
 r.Shadow.Virtual.ResolutionLodBiasDirectional=-1.5
+r.Shadow.Virtual.ResolutionLodBiasDirectionalMoving=0.0
 r.Shadow.Virtual.ResolutionLodBiasLocal=0.0
+r.Shadow.Virtual.ResolutionLodBiasLocalMoving=1.0
 r.Shadow.Virtual.SMRT.RayCountDirectional=8
 r.Shadow.Virtual.SMRT.SamplesPerRayDirectional=4
 r.Shadow.Virtual.SMRT.RayCountLocal=8
@@ -229,7 +237,9 @@ r.LightMaxDrawDistanceScale=1
 r.CapsuleShadows=1
 r.Shadow.Virtual.MaxPhysicalPages=8192
 r.Shadow.Virtual.ResolutionLodBiasDirectional=-1.5
+r.Shadow.Virtual.ResolutionLodBiasDirectionalMoving=-1.5
 r.Shadow.Virtual.ResolutionLodBiasLocal=0.0
+r.Shadow.Virtual.ResolutionLodBiasLocalMoving=0.0
 r.Shadow.Virtual.SMRT.RayCountDirectional=16
 r.Shadow.Virtual.SMRT.SamplesPerRayDirectional=8
 r.Shadow.Virtual.SMRT.RayCountLocal=16
diff --git a/Engine/Shaders/Private/VirtualShadowMaps/VirtualShadowMapPageCacheCommon.ush b/Engine/Shaders/Private/VirtualShadowMaps/VirtualShadowMapPageCacheCommon.ush
--- a/Engine/Shaders/Private/VirtualShadowMaps/VirtualShadowMapPageCacheCommon.ush
+++ b/Engine/Shaders/Private/VirtualShadowMaps/VirtualShadowMapPageCacheCommon.ush
@@ -19,7 +19,6 @@ StructuredBuffer<FShadowMapCacheData> ShadowMapCacheData;
 StructuredBuffer<uint> PrevPageFlags;
 StructuredBuffer<uint> PrevPageTable;
 StructuredBuffer<FPhysicalPageMetaData> PrevPhysicalPageMetaData;
-ByteAddressBuffer PrevProjectionData;
 
 // Flags buffer (one bit per PersistentPrimitiveIndex) to track statically cached primitives that are invalidating
 // Currently we use GPU messaging to send these primitives back for the CPU to handle as desired
diff --git a/Engine/Source/Runtime/Renderer/Private/DeferredShadingRenderer.cpp b/Engine/Source/Runtime/Renderer/Private/DeferredShadingRenderer.cpp
--- a/Engine/Source/Runtime/Renderer/Private/DeferredShadingRenderer.cpp
+++ b/Engine/Source/Runtime/Renderer/Private/DeferredShadingRenderer.cpp
@@ -86,6 +86,7 @@
 #include "ComponentRecreateRenderStateContext.h"
 #include "RenderCore.h"
 #include "VariableRateShadingImageManager.h"
+#include "Shadows/ShadowScene.h"
 
 extern int32 GNaniteShowStats;
 extern int32 GNanitePickingDomain;
@@ -2465,6 +2466,8 @@ void FDeferredShadingSceneRenderer::Render(FRDGBuilder& GraphBuilder)
 
 		Scene->UpdateAllPrimitiveSceneInfos(GraphBuilder, AsyncOps);
 	}
+	// kick off dependent scene updates 
+	Scene->ShadowScene->UpdateForRenderedFrame(GraphBuilder);
 
 #if RHI_RAYTRACING
 	// Initialize ray tracing flags, in case they weren't initialized in the CreateSceneRenderers code path
@@ -2795,6 +2798,10 @@ void FDeferredShadingSceneRenderer::Render(FRDGBuilder& GraphBuilder)
 		}
 	}
 
+#if !(UE_BUILD_SHIPPING || UE_BUILD_TEST)
+	Scene->DebugRender(Views);
+#endif
+
 	FRDGExternalAccessQueue ExternalAccessQueue;
 
 	{
@@ -3610,7 +3617,10 @@ void FDeferredShadingSceneRenderer::Render(FRDGBuilder& GraphBuilder)
 				RDG_GPU_STAT_SCOPE(GraphBuilder, ShadowDepths);
 
 				ensureMsgf(AreLightsInLightGrid(), TEXT("Virtual shadow map setup requires local lights to be injected into the light grid (this may be caused by 'r.LightCulling.Quality=0')."));
-				VirtualShadowMapArray.BuildPageAllocations(GraphBuilder, SceneTextures, Views, ViewFamily.EngineShowFlags, SortedLightSet, VisibleLightInfos, NaniteRasterResults, SingleLayerWaterPrePassResult);
+
+				// TODO: propagate config in some other way (along with shadowing lights flowing from persistent setup to the VSM setup)
+				auto GetLightMobilityFactor = [this](int32 LightId)->float { return Scene->ShadowScene->GetLightMobilityFactor(LightId); };
+				VirtualShadowMapArray.BuildPageAllocations(GraphBuilder, GetActiveSceneTextures(), Views, ViewFamily.EngineShowFlags, SortedLightSet, VisibleLightInfos, GetLightMobilityFactor, SingleLayerWaterPrePassResult);
 			}
 
 			RenderShadowDepthMaps(GraphBuilder, InstanceCullingManager, ExternalAccessQueue);
diff --git a/Engine/Source/Runtime/Renderer/Private/LightSceneInfo.cpp b/Engine/Source/Runtime/Renderer/Private/LightSceneInfo.cpp
--- a/Engine/Source/Runtime/Renderer/Private/LightSceneInfo.cpp
+++ b/Engine/Source/Runtime/Renderer/Private/LightSceneInfo.cpp
@@ -42,6 +42,7 @@ void FLightSceneInfoCompact::Init(FLightSceneInfo* InLightSceneInfo)
 	bStaticLighting = InLightSceneInfo->Proxy->HasStaticLighting();
 	bAffectReflection = InLightSceneInfo->Proxy->AffectReflection();
 	bAffectGlobalIllumination = InLightSceneInfo->Proxy->AffectGlobalIllumination();
+	bIsMovable = InLightSceneInfo->Proxy->IsMovable();
 	CastRaytracedShadow = InLightSceneInfo->Proxy->CastsRaytracedShadow();
 }
 
diff --git a/Engine/Source/Runtime/Renderer/Private/LightSceneInfo.h b/Engine/Source/Runtime/Renderer/Private/LightSceneInfo.h
--- a/Engine/Source/Runtime/Renderer/Private/LightSceneInfo.h
+++ b/Engine/Source/Runtime/Renderer/Private/LightSceneInfo.h
@@ -46,6 +46,7 @@ public:
 	uint32 bStaticLighting : 1;
 	uint32 bAffectReflection : 1;
 	uint32 bAffectGlobalIllumination : 1;
+	uint32 bIsMovable : 1;
 	TEnumAsByte<ECastRayTracedShadow::Type> CastRaytracedShadow;
 
 	/** Initializes the compact scene info from the light's full scene info. */
diff --git a/Engine/Source/Runtime/Renderer/Private/RendererScene.cpp b/Engine/Source/Runtime/Renderer/Private/RendererScene.cpp
--- a/Engine/Source/Runtime/Renderer/Private/RendererScene.cpp
+++ b/Engine/Source/Runtime/Renderer/Private/RendererScene.cpp
@@ -84,6 +84,7 @@
 #include "ShaderPrint.h"
 
 #include "VirtualShadowMaps/VirtualShadowMapCacheManager.h"
+#include "Shadows/ShadowScene.h"
 
 #if WITH_EDITOR
 #include "Rendering/StaticLightingSystemInterface.h"
@@ -1923,6 +1924,9 @@ FScene::FScene(UWorld* InWorld, bool bInRequiresHitProxies, bool bInIsEditorScen
 	DefaultVirtualShadowMapCache = new FVirtualShadowMapArrayCacheManager(this);
 
 	SceneLightInfoUpdates = new FSceneLightInfoUpdates;
+
+	// Allocate the shadow scene, it is always present but we use a pointer such that it can be forward declared.
+	ShadowScene = new FShadowScene(*this);
 }
 
 FScene::~FScene()
@@ -1963,6 +1967,12 @@ FScene::~FScene()
 		DefaultVirtualShadowMapCache = nullptr;
 	}
 
+	if (ShadowScene)
+	{
+		delete ShadowScene;
+		ShadowScene = nullptr;
+	}
+
 	if (DefaultLumenSceneData)
 	{
 		delete DefaultLumenSceneData;
@@ -5167,7 +5177,7 @@ FLightSceneChangeSet FScene::UpdateAllLightSceneInfos(FRDGBuilder& GraphBuilder)
 		}
 	}
 	// This can't access the scene light data if done async since it happens before the actual removals.
-	OnPreLigtSceneInfoUpdate.Broadcast(FLightSceneChangeSet{ ChangeSet.RemovedLightIds, TConstArrayView<int32>(), ChangeSet.TransformUpdatedLightIds, ChangeSet.ColorUpdatedLightIds });
+	OnPreLigtSceneInfoUpdate.Broadcast(GraphBuilder, FLightSceneChangeSet{ ChangeSet.RemovedLightIds, TConstArrayView<int32>(), ChangeSet.TransformUpdatedLightIds, ChangeSet.ColorUpdatedLightIds });
 
 	// Batch process all light removes
 	for (int32 LightId : ChangeSet.RemovedLightIds)
@@ -5236,7 +5246,7 @@ FLightSceneChangeSet FScene::UpdateAllLightSceneInfos(FRDGBuilder& GraphBuilder)
 		}
 	}
 
-	OnPostLigtSceneInfoUpdate.Broadcast(FLightSceneChangeSet{ ChangeSet.RemovedLightIds, ChangeSet.AddedLightIds, ChangeSet.TransformUpdatedLightIds, ChangeSet.ColorUpdatedLightIds });
+	OnPostLigtSceneInfoUpdate.Broadcast(GraphBuilder, FLightSceneChangeSet{ ChangeSet.RemovedLightIds, ChangeSet.AddedLightIds, ChangeSet.TransformUpdatedLightIds, ChangeSet.ColorUpdatedLightIds });
 
 	SceneLightInfoUpdates->Reset();
 
@@ -6617,4 +6627,11 @@ void FScene::IncrementFrameNumber()
 	{
 		SceneFrameNumberRenderThread = NewNumber;
 	});
-}
\ No newline at end of file
+}
+
+#if !(UE_BUILD_SHIPPING || UE_BUILD_TEST)
+void FScene::DebugRender(TArrayView<FViewInfo> Views)
+{
+	ShadowScene->DebugRender(Views);
+}
+#endif
diff --git a/Engine/Source/Runtime/Renderer/Private/ScenePrivate.h b/Engine/Source/Runtime/Renderer/Private/ScenePrivate.h
--- a/Engine/Source/Runtime/Renderer/Private/ScenePrivate.h
+++ b/Engine/Source/Runtime/Renderer/Private/ScenePrivate.h
@@ -87,6 +87,7 @@ class FSparseVolumeTextureViewerSceneProxy;
 class FLandscapeRayTracingStateList;
 class FExponentialHeightFogSceneInfo;
 class FStaticMeshBatch;
+class FShadowScene;
 class FSceneLightInfoUpdates;
 /**
  * Describes all light modifications to the scene by recording the light scene IDs.
@@ -2921,6 +2922,11 @@ public:
 	 */
 	FVirtualShadowMapArrayCacheManager* DefaultVirtualShadowMapCache;
 
+	/**
+	 * Stores scene-aspects needed for shadow rendering.
+	 */
+	FShadowScene* ShadowScene;
+
 	/** Texture layout that tracks current allocations in the PreshadowCache render target. */
 	FTextureLayout PreshadowCacheLayout;
 
@@ -3435,11 +3441,15 @@ public:
 	void LumenInvalidateSurfaceCacheForPrimitive(FPrimitiveSceneInfo* InPrimitive);
 	void LumenRemovePrimitive(FPrimitiveSceneInfo* InPrimitive, int32 PrimitiveIndex);
 
+#if !(UE_BUILD_SHIPPING || UE_BUILD_TEST)
+	void DebugRender(TArrayView<FViewInfo> Views);
+#endif
+
 	/**
 	 * Light scene change delegates, may be used to hook in subsystems that need to respond to light scene changes.
 	 * Note, all the light scene changes are applied _before_ all the primitive scene infos are updated.
 	 */
-	DECLARE_MULTICAST_DELEGATE_OneParam(FSceneLightSceneInfoUpdateDelegate, const FLightSceneChangeSet&);
+	DECLARE_MULTICAST_DELEGATE_TwoParams(FSceneLightSceneInfoUpdateDelegate, FRDGBuilder& , const FLightSceneChangeSet&);
 	/**
 	 * This delegate is invoked during the scene update phase _before_ the scene has had any light changes applied.
 	 * Thus, AddedLightIds is not valid in the change set as the added lights do not have assigned IDs yet.
diff --git a/Engine/Source/Runtime/Renderer/Private/ShadowSetup.cpp b/Engine/Source/Runtime/Renderer/Private/ShadowSetup.cpp
--- a/Engine/Source/Runtime/Renderer/Private/ShadowSetup.cpp
+++ b/Engine/Source/Runtime/Renderer/Private/ShadowSetup.cpp
@@ -38,6 +38,7 @@
 #include "StaticMeshBatch.h"
 #include "UnrealEngine.h"
 #include "CanvasTypes.h"
+#include "Shadows/ShadowScene.h"
 
 /** Number of cube map shadow depth surfaces that will be created and used for rendering one pass point light shadows. */
 static const int32 NumCubeShadowDepthSurfaces = 5;
@@ -5000,7 +5001,8 @@ void FSceneRenderer::AddViewDependentWholeSceneShadowsForView(
 					WorldToLight,
 					View.ViewMatrices,
 					View.ViewRect.Size(),
-					&View
+					&View,
+					Scene->ShadowScene->GetLightMobilityFactor(LightSceneInfo.Id)
 				));
 
 				// TODO: Not clear we need both of these in this path, but keep it consistent for now
diff --git a/Engine/Source/Runtime/Renderer/Private/Shadows/ShadowScene.cpp b/Engine/Source/Runtime/Renderer/Private/Shadows/ShadowScene.cpp
new file mode 100644
--- /dev/null
+++ b/Engine/Source/Runtime/Renderer/Private/Shadows/ShadowScene.cpp
@@ -0,0 +1,182 @@
+// Copyright Epic Games, Inc. All Rights Reserved.
+
+#include "ShadowScene.h"
+#include "ScenePrivate.h"
+#include "Algo/ForEach.h"
+
+#if !(UE_BUILD_SHIPPING || UE_BUILD_TEST)
+#include "DynamicPrimitiveDrawing.h"
+#include "LightSceneProxy.h"
+#endif
+
+TAutoConsoleVariable<int32> CVarDebugDrawLightActiveStateTracking(
+	TEXT("r.Shadow.Scene.DebugDrawLightActiveStateTracking"),
+	0,
+	TEXT("."),
+	ECVF_RenderThreadSafe
+);
+
+TAutoConsoleVariable<int32> CVarShadowSceneLightActiveFrameCount(
+	TEXT("r.Shadow.Scene.LightActiveFrameCount"),
+	10,
+	TEXT("Number of frames before a light that has been moving (updated or transform changed) goes to inactive state.\n")
+	TEXT("  This determines the number of frames that the MobilityFactor goes to zero over, and thus a higher number spreads invalidations out over a longer time."),
+	ECVF_RenderThreadSafe
+);
+
+DECLARE_DWORD_COUNTER_STAT(TEXT("Active Light Count"), STAT_ActiveLightCount, STATGROUP_ShadowRendering);
+
+
+FShadowScene::FShadowScene(FScene & InScene)
+	: Scene(InScene)
+{
+	Scene.OnPostLigtSceneInfoUpdate.AddLambda([this](FRDGBuilder& GraphBuilder, const FLightSceneChangeSet& LightSceneChangeSet) { PostLightsUpdate(GraphBuilder, LightSceneChangeSet); });
+}
+
+void FShadowScene::PostLightsUpdate(FRDGBuilder& GraphBuilder, const FLightSceneChangeSet& LightSceneChangeSet)
+{
+	// don't spawn async work for no good reason.
+	constexpr int32 kMinWorkSizeForAsync = 64;
+
+	int32 WorkSize = LightSceneChangeSet.AddedLightIds.Num() + LightSceneChangeSet.RemovedLightIds.Num() + LightSceneChangeSet.TransformUpdatedLightIds.Num();
+
+	// Don't sync, or kick off a new job if there is no work to do.
+	if (WorkSize > 0)
+	{
+		// Need to wait in case the update is performed several times in a row for some reason.
+		SceneChangeUpdateTask.Wait();
+		SceneChangeUpdateTask = GraphBuilder.AddSetupTask([this, LightSceneChangeSet]
+		{
+			// Oust all removed Ids.
+			for (int32 Id : LightSceneChangeSet.RemovedLightIds)
+			{
+				check(Lights.IsValidIndex(Id));
+				ActiveLights[Id] = false;
+				Lights.RemoveAt(Id);
+			}
+	
+			// Track active lights (those that are or were moving, and thus need updating)
+			ActiveLights.SetNum(FMath::Max(Lights.GetMaxIndex(), Lights.Num() + LightSceneChangeSet.AddedLightIds.Num()), false);
+
+			auto UpdateLight = [&](int32 Id)
+			{
+				FLightData& Light = GetOrAddLight(Id);
+				// Mark as not rendered
+				Light.FirstActiveFrameNumber = INDEX_NONE;
+				// Only movable lights can become "active" (i.e., having moved recently and thus needing active update)
+				if (Scene.Lights[Id].bIsMovable)
+				{
+					ActiveLights[Id] = true;
+					Light.MobilityFactor = 1.0f;
+				}
+				else
+				{
+					ActiveLights[Id] = false;
+					Light.MobilityFactor = 0.0f;
+					// Go straight to not active
+				}
+			};
+
+			Algo::ForEach(LightSceneChangeSet.AddedLightIds, UpdateLight);
+			Algo::ForEach(LightSceneChangeSet.TransformUpdatedLightIds, UpdateLight);
+
+		}, WorkSize > kMinWorkSizeForAsync);
+	}
+}
+
+
+void FShadowScene::UpdateForRenderedFrame(FRDGBuilder& GraphBuilder)
+{
+	SceneChangeUpdateTask.Wait();
+
+	const int32 ActiveFrameCount = FMath::Max(1, CVarShadowSceneLightActiveFrameCount.GetValueOnRenderThread());
+	
+	// 1. FScene::FrameNumber is incremented before a call to render is being dispatched to the RT.
+	int32 SceneFrameNumber = Scene.GetFrameNumberRenderThread();
+	// Iterate the previously active lights and update
+	for (TConstSetBitIterator<> BitIt(ActiveLights); BitIt; ++BitIt)
+	{
+		int32 Id = BitIt.GetIndex();
+		// No need to process if we already decided it is active (i.e., it was moved again this frame)
+		FLightData& Light = Lights[Id];
+		
+		// If it was not rendered before, we record the current scene frame number
+		if (Light.FirstActiveFrameNumber == INDEX_NONE)
+		{
+			Light.FirstActiveFrameNumber = SceneFrameNumber;
+			Light.MobilityFactor = 1.0f;
+		}
+		else if(SceneFrameNumber - Light.FirstActiveFrameNumber  < ActiveFrameCount)
+		{
+			// If it was updated before, but not this frame, check how many rendered frames have elapsed and transition to the inactive state
+			Light.MobilityFactor = 1.0f - FMath::Clamp(float(SceneFrameNumber - Light.FirstActiveFrameNumber) / float(ActiveFrameCount), 0.0f, 1.0f);
+		}
+		else
+		{
+			// It's not been updated for more than K frames transition to non-active state
+			ActiveLights[Id] = false;
+			Light.MobilityFactor = 0.0f;
+		}
+	}
+
+	SET_DWORD_STAT(STAT_ActiveLightCount, ActiveLights.CountSetBits());
+}
+
+
+
+float FShadowScene::GetLightMobilityFactor(int32 LightId) const 
+{ 
+	SceneChangeUpdateTask.Wait();
+
+	if (IsActive(LightId))
+	{
+		return Lights[LightId].MobilityFactor;
+	}
+	return 0.0f;
+}
+
+
+void FShadowScene::DebugRender(TArrayView<FViewInfo> Views)
+{
+#if !(UE_BUILD_SHIPPING || UE_BUILD_TEST)
+
+	if (CVarDebugDrawLightActiveStateTracking.GetValueOnRenderThread() != 0)
+	{
+		SceneChangeUpdateTask.Wait();
+
+		for (FViewInfo& View : Views)
+		{
+			FViewElementPDI DebugPDI(&View, nullptr, &View.DynamicPrimitiveCollector);
+
+			for (TConstSetBitIterator<> BitIt(ActiveLights); BitIt; ++BitIt)
+			{
+				int32 Id = BitIt.GetIndex();
+				FLightSceneInfoCompact& LightSceneInfoCompact = Scene.Lights[Id];
+				FLightSceneProxy* Proxy = LightSceneInfoCompact.LightSceneInfo->Proxy;
+				FLinearColor Color = FMath::Lerp(FLinearColor(FColor::Yellow), FLinearColor(FColor::Blue), Scene.ShadowScene->GetLightMobilityFactor(Id));
+				switch (LightSceneInfoCompact.LightType)
+				{
+				case LightType_Directional:
+					DrawWireSphereAutoSides(&DebugPDI, Proxy->GetLightToWorld().GetOrigin(), Color, FMath::Min(100.0f, Proxy->GetRadius()), SDPG_World);
+					DrawWireSphereAutoSides(&DebugPDI, Proxy->GetLightToWorld().GetOrigin(), Color, FMath::Min(200.0f, Proxy->GetRadius()), SDPG_World);
+					break;
+				case LightType_Spot:
+				{
+					FTransform TransformNoScale = FTransform(Proxy->GetLightToWorld());
+					TransformNoScale.RemoveScaling();
+
+					DrawWireSphereCappedCone(&DebugPDI, TransformNoScale, Proxy->GetRadius(), FMath::RadiansToDegrees(Proxy->GetOuterConeAngle()), 16, 4, 8, Color, SDPG_World);
+				}
+				break;
+				default:
+				{
+					DrawWireSphereAutoSides(&DebugPDI, Proxy->GetPosition(), Color, Proxy->GetRadius(), SDPG_World);
+				}
+				break;
+				};
+			}
+		}
+	}
+
+#endif
+}
\ No newline at end of file
diff --git a/Engine/Source/Runtime/Renderer/Private/Shadows/ShadowScene.h b/Engine/Source/Runtime/Renderer/Private/Shadows/ShadowScene.h
new file mode 100644
--- /dev/null
+++ b/Engine/Source/Runtime/Renderer/Private/Shadows/ShadowScene.h
@@ -0,0 +1,70 @@
+// Copyright Epic Games, Inc. All Rights Reserved.
+#pragma once
+
+#include "CoreMinimal.h"
+#include "Containers/SparseArray.h"
+#include "Containers/ArrayView.h"
+#include "Containers/Array.h"
+#include "Tasks/Task.h"
+
+class FProjectedShadowInfo;
+class FScene;
+class FWholeSceneProjectedShadowInitializer;
+class FRDGBuilder;
+class FVirtualShadowMapPerLightCacheEntry;
+struct FLightSceneChangeSet;
+class FViewInfo;
+
+/**
+ * Persistent scene-representation of for shadow rendering.
+ */
+class FShadowScene
+{
+public:
+	FShadowScene(FScene &InScene);
+
+	// Handle all scene changes wrt lights, attached to the FScene::OnPostLigtSceneInfoUpdate delegate.
+	void PostLightsUpdate(FRDGBuilder& GraphBuilder, const FLightSceneChangeSet &LightSceneChangeSet);
+
+	/**
+	 * Fetch the "mobility factor" for the light, [0,1] where 0.0 means not moving, and 1.0 means was updated this frame.
+	 * Does a smooth transition from 1 to 0 over N frames, defined by the cvar.
+	 */
+	float GetLightMobilityFactor(int32 LightId) const;
+
+	/**
+	 * Call once per rendered frame to update state that depend on number of rendered frames.
+	 */
+	void UpdateForRenderedFrame(FRDGBuilder& GraphBuilder);
+
+	void DebugRender(TArrayView<FViewInfo> Views);
+private:
+	//friend class FShadowSceneRenderer;
+
+	inline bool IsActive(int32 LightId) const { return ActiveLights[LightId]; }
+
+	struct FLightData
+	{
+		/**
+		 * Scene rendering frame number of the first frame that the scene was rendered after being modified (moved/added).
+		 */
+		int32 FirstActiveFrameNumber;
+		float MobilityFactor;
+	};
+	TSparseArray<FLightData> Lights;
+
+	inline FLightData& GetOrAddLight(int32 LightId)
+	{
+		if (!Lights.IsValidIndex(LightId))
+		{
+			Lights.EmplaceAt(LightId, FLightData{});
+		}
+		return Lights[LightId];
+	}
+	// Bit-array marking active lights, those we deem active are ones that have been modified in a recent frame and thus need some kind of active update.
+	TBitArray<> ActiveLights;
+
+	// Links to other systems etc.
+	FScene& Scene;
+	mutable UE::Tasks::FTask SceneChangeUpdateTask;
+};
diff --git a/Engine/Source/Runtime/Renderer/Private/Shadows/ShadowSceneRenderer.cpp b/Engine/Source/Runtime/Renderer/Private/Shadows/ShadowSceneRenderer.cpp
--- a/Engine/Source/Runtime/Renderer/Private/Shadows/ShadowSceneRenderer.cpp
+++ b/Engine/Source/Runtime/Renderer/Private/Shadows/ShadowSceneRenderer.cpp
@@ -1,6 +1,7 @@
 // Copyright Epic Games, Inc. All Rights Reserved.
 
 #include "ShadowSceneRenderer.h"
+#include "ShadowScene.h"
 #include "DeferredShadingRenderer.h"
 #include "ScenePrivate.h"
 #include "ShadowRendering.h"
@@ -74,6 +75,7 @@ DECLARE_DWORD_COUNTER_STAT(TEXT("VSM Light Projections (Local One Pass Copy)"),
 FShadowSceneRenderer::FShadowSceneRenderer(FDeferredShadingSceneRenderer& InSceneRenderer)
 	: SceneRenderer(InSceneRenderer)
 	, Scene(*InSceneRenderer.Scene)
+	, ShadowScene(*Scene.ShadowScene)
 	, VirtualShadowMapArray(InSceneRenderer.VirtualShadowMapArray)
 {
 }
@@ -97,7 +99,7 @@ TSharedPtr<FVirtualShadowMapPerLightCacheEntry> FShadowSceneRenderer::AddLocalLi
 	//       we can absolutely mirror the page marking calc better, just unclear how much it helps. 
 	//       Also possible to feed back from gpu - which would be more accurate wrt partially visible lights (e.g., a spot going through the ground).
 	//       Of course this creates jumps if visibility changes, which may or may not create unsolvable artifacts.
-	const float BiasedFootprintThreshold = float(FVirtualShadowMap::PageSize) * FMath::Exp2(VirtualShadowMapArray.GetResolutionLODBiasLocal());
+	const float BiasedFootprintThreshold = float(FVirtualShadowMap::PageSize) * FMath::Exp2(VirtualShadowMapArray.GetResolutionLODBiasLocal(ShadowScene.GetLightMobilityFactor(LightSceneInfo->Id)));
 	const bool bIsDistantLight = CVarDistantLightMode.GetValueOnRenderThread() != 0
 		&& (MaxScreenRadius <= BiasedFootprintThreshold || CVarDistantLightMode.GetValueOnRenderThread() == 2);
 
@@ -266,7 +268,7 @@ void FShadowSceneRenderer::UpdateDistantLightPriorityRender()
 		
 		// Force fully cached to be off.
 		LocalLightShadowFrameSetup.ProjectedShadowInfo->bShouldRenderVSM = true;
-		LocalLightShadowFrameSetup.PerLightCacheEntry->CurrenScheduledFrameNumber = Scene.GetFrameNumber();
+		LocalLightShadowFrameSetup.PerLightCacheEntry->Current.ScheduledFrameNumber = Scene.GetFrameNumber();
 		// Should trigger invalidations also.
 		LocalLightShadowFrameSetup.PerLightCacheEntry->Invalidate();
 	}
@@ -286,7 +288,7 @@ void FShadowSceneRenderer::PostSetupDebugRender()
 			for (const FLocalLightShadowFrameSetup& LightSetup : LocalLights)
 			{			
 				FLinearColor Color = FLinearColor(FColor::Blue);
-				if (LightSetup.PerLightCacheEntry && LightSetup.PerLightCacheEntry->bCurrentIsDistantLight)
+				if (LightSetup.PerLightCacheEntry && LightSetup.PerLightCacheEntry->Current.bIsDistantLight)
 				{
 					++NumDistant;
 					int32 FramesSinceLastRender = int32(Scene.GetFrameNumber()) - int32(LightSetup.PerLightCacheEntry->GetLastScheduledFrameNumber());
diff --git a/Engine/Source/Runtime/Renderer/Private/Shadows/ShadowSceneRenderer.h b/Engine/Source/Runtime/Renderer/Private/Shadows/ShadowSceneRenderer.h
--- a/Engine/Source/Runtime/Renderer/Private/Shadows/ShadowSceneRenderer.h
+++ b/Engine/Source/Runtime/Renderer/Private/Shadows/ShadowSceneRenderer.h
@@ -1,7 +1,4 @@
 // Copyright Epic Games, Inc. All Rights Reserved.
-/*=============================================================================
-	ShadowSceneRenderer.h:
-=============================================================================*/
 #pragma once
 
 #include "CoreMinimal.h"
@@ -18,6 +15,7 @@ class FWholeSceneProjectedShadowInitializer;
 class FRDGBuilder;
 class FVirtualShadowMapPerLightCacheEntry;
 struct FNaniteVisibilityQuery;
+class FShadowScene;
 
 /**
  * Transient scope for per-frame rendering resources for the shadow rendering.
@@ -89,8 +87,6 @@ private:
 	 */
 	void UpdateDistantLightPriorityRender();
 
-	// TODO: void RenderVirtualShadowMaps(FRDGBuilder& GraphBuilder, bool bNaniteEnabled);
-
 	struct FLocalLightShadowFrameSetup
 	{
 		TArray<FVirtualShadowMap*, TInlineAllocator<6>> VirtualShadowMaps;
@@ -104,7 +100,6 @@ private:
 	//       tradeoff is easy to look up (given light ID) but not compact, but OTOH can keep compact lists of indices for various purposes
 	TArray<FLocalLightShadowFrameSetup, SceneRenderingAllocator> LocalLights;
 
-
 	struct FDirectionalLightShadowFrameSetup
 	{
 		FProjectedShadowInfo* ProjectedShadowInfo = nullptr;
@@ -120,6 +115,7 @@ private:
 	// Links to other systems etc.
 	FDeferredShadingSceneRenderer& SceneRenderer;
 	FScene& Scene;
+	FShadowScene& ShadowScene;
 	FVirtualShadowMapArray& VirtualShadowMapArray;
 
 	FNaniteVisibilityQuery* NaniteVisibilityQuery = nullptr;
diff --git a/Engine/Source/Runtime/Renderer/Private/VirtualShadowMaps/VirtualShadowMapArray.cpp b/Engine/Source/Runtime/Renderer/Private/VirtualShadowMaps/VirtualShadowMapArray.cpp
--- a/Engine/Source/Runtime/Renderer/Private/VirtualShadowMaps/VirtualShadowMapArray.cpp
+++ b/Engine/Source/Runtime/Renderer/Private/VirtualShadowMaps/VirtualShadowMapArray.cpp
@@ -113,6 +113,12 @@ static TAutoConsoleVariable<float> CVarResolutionLodBiasLocal(
 	ECVF_Scalability | ECVF_RenderThreadSafe
 );
 
+static TAutoConsoleVariable<float> CVarResolutionLodBiasLocalMoving(
+	TEXT("r.Shadow.Virtual.ResolutionLodBiasLocalMoving"),
+	1.0f,
+	TEXT("Bias applied to LOD calculations for moving local lights. -1.0 doubles resolution, 1.0 halves it and so on."),
+	ECVF_Scalability | ECVF_RenderThreadSafe
+);
 static TAutoConsoleVariable<float> CVarPageDilationBorderSizeDirectional(
 	TEXT("r.Shadow.Virtual.PageDilationBorderSizeDirectional"),
 	0.05f,
@@ -395,7 +401,6 @@ BEGIN_SHADER_PARAMETER_STRUCT(FCacheDataParameters, )
 	SHADER_PARAMETER_RDG_BUFFER_SRV( StructuredBuffer< uint >,					PrevPageTable )
 	SHADER_PARAMETER_RDG_BUFFER_SRV( StructuredBuffer< FPhysicalPageMetaData >,	PrevPhysicalPageMetaData)
 	SHADER_PARAMETER_RDG_BUFFER_SRV( StructuredBuffer< uint >,					PrevPhysicalPageMetaDataOut)
-	SHADER_PARAMETER_RDG_BUFFER_SRV( StructuredBuffer< uint >,					PrevProjectionData)
 END_SHADER_PARAMETER_STRUCT()
 
 static void SetCacheDataShaderParameters(FRDGBuilder& GraphBuilder, const TArray<TUniquePtr<FVirtualShadowMap>, SceneRenderingAllocator>& ShadowMaps, FVirtualShadowMapArrayCacheManager* CacheManager, FCacheDataParameters &CacheDataParameters)
@@ -424,7 +429,6 @@ static void SetCacheDataShaderParameters(FRDGBuilder& GraphBuilder, const TArray
 	CacheDataParameters.PrevPageFlags = GraphBuilder.CreateSRV(GraphBuilder.RegisterExternalBuffer(CacheManager->PrevBuffers.PageFlags, TEXT("Shadow.Virtual.PrevPageFlags")));
 	CacheDataParameters.PrevPageTable = GraphBuilder.CreateSRV(GraphBuilder.RegisterExternalBuffer(CacheManager->PrevBuffers.PageTable, TEXT("Shadow.Virtual.PrevPageTable")));
 	CacheDataParameters.PrevPhysicalPageMetaData = GraphBuilder.CreateSRV(GraphBuilder.RegisterExternalBuffer(CacheManager->PrevBuffers.PhysicalPageMetaData, TEXT("Shadow.Virtual.PrevPhysicalPageMetaData")));
-	CacheDataParameters.PrevProjectionData = GraphBuilder.CreateSRV(GraphBuilder.RegisterExternalBuffer(CacheManager->PrevBuffers.ProjectionData, TEXT("Shadow.Virtual.PrevProjectionData")));
 }
 
 
@@ -553,9 +557,9 @@ FVirtualShadowMap* FVirtualShadowMapArray::Allocate(bool bSinglePageShadowMap)
 	return SM;
 }
 
-float FVirtualShadowMapArray::GetResolutionLODBiasLocal() const
+float FVirtualShadowMapArray::GetResolutionLODBiasLocal(float LightMobilityFactor) const
 {
-	return CVarResolutionLodBiasLocal.GetValueOnRenderThread();
+	return InterpolateResolutionBias(CVarResolutionLodBiasLocal.GetValueOnRenderThread(), CVarResolutionLodBiasLocalMoving.GetValueOnRenderThread(), LightMobilityFactor);
 }
 
 FVirtualShadowMapArray::~FVirtualShadowMapArray()
@@ -1180,11 +1184,11 @@ static uint32 GetShadowMapsToAllocate(uint32 NumShadowMaps)
 void FVirtualShadowMapArray::BuildPageAllocations(
 	FRDGBuilder& GraphBuilder,
 	const FMinimalSceneTextures& SceneTextures,
-	const TArray<FViewInfo>& Views,
+	const TConstArrayView<FViewInfo>& Views,
 	const FEngineShowFlags& EngineShowFlags,
 	const FSortedLightSetSceneInfo& SortedLightsInfo,
-	const TArray<FVisibleLightInfo, SceneRenderingAllocator>& VisibleLightInfos,
-	const TArray<Nanite::FRasterResults, TInlineAllocator<2>>& NaniteRasterResults,
+	const TConstArrayView<FVisibleLightInfo>& VisibleLightInfos,
+	const TFunctionRef<float(int32 LightId)>& GetLightMobilityFactor,
 	const FSingleLayerWaterPrePassResult* SingleLayerWaterPrePassResult)
 {
 	check(IsEnabled());
@@ -1246,8 +1250,9 @@ void FVirtualShadowMapArray::BuildPageAllocations(
 	FRDGScatterUploadBuffer ProjectionDataUploader;
 	ProjectionDataUploader.Init(GraphBuilder, GetNumShadowMaps(), sizeof(FVirtualShadowMapProjectionShaderData), false, TEXT("Shadow.Virtual.ProjectionData.UploadBuffer"));
 
-	for (const FVisibleLightInfo& VisibleLightInfo : VisibleLightInfos)
+	for (int32 LightId = 0; LightId < VisibleLightInfos.Num(); ++LightId)
 	{
+		const FVisibleLightInfo& VisibleLightInfo = VisibleLightInfos[LightId];
 		for (const TSharedPtr<FVirtualShadowMapClipmap>& Clipmap : VisibleLightInfo.VirtualShadowMapClipmaps)
 		{
 			// NOTE: Shader assumes all levels from a given clipmap are contiguous
@@ -1269,12 +1274,12 @@ void FVirtualShadowMapArray::BuildPageAllocations(
 			}
 		}
 
-		const float ResolutionLodBiasLocal = GetResolutionLODBiasLocal();
 
 		for (FProjectedShadowInfo* ProjectedShadowInfo : VisibleLightInfo.AllProjectedShadows)
 		{
 			if (ProjectedShadowInfo->HasVirtualShadowMap())
 			{
+				const float ResolutionLodBiasLocal = GetResolutionLODBiasLocal(GetLightMobilityFactor(LightId));
 				check(ProjectedShadowInfo->CascadeSettings.ShadowSplitIndex == INDEX_NONE);		// We use clipmaps for virtual shadow maps, not cascades
 
 				// NOTE: Virtual shadow maps are never atlased, but verify our assumptions
@@ -1294,9 +1299,10 @@ void FVirtualShadowMapArray::BuildPageAllocations(
 				int32 NumMaps = ProjectedShadowInfo->bOnePassPointLightShadow ? 6 : 1;
 				for(int32 Index = 0; Index < NumMaps; ++Index)
 				{
-					check(!CacheEntry.IsValid() || CacheEntry->bCurrentIsDistantLight == ProjectedShadowInfo->VirtualShadowMaps[Index]->bIsSinglePageSM);
+					check(!CacheEntry.IsValid() || CacheEntry->Current.bIsDistantLight == ProjectedShadowInfo->VirtualShadowMaps[Index]->bIsSinglePageSM);
 					
 					uint32 Flags = ProjectedShadowInfo->VirtualShadowMaps[Index]->bIsSinglePageSM ? VSM_PROJ_FLAG_CURRENT_DISTANT_LIGHT : 0U;
+					Flags |= CacheEntry.IsValid() && CacheEntry->IsUncached() ? VSM_PROJ_FLAG_UNCACHED : 0U;
 					int32 ID = ProjectedShadowInfo->VirtualShadowMaps[Index]->ID;
 
 					const FViewMatrices ViewMatrices = ProjectedShadowInfo->GetShadowDepthRenderingViewMatrices(Index, true );
@@ -3235,13 +3241,17 @@ uint32 FVirtualShadowMapArray::AddRenderViews(const TSharedPtr<FVirtualShadowMap
 	BaseParams.PrevTargetLayerIndex = INDEX_NONE;
 	BaseParams.TargetMipLevel = 0;
 	BaseParams.TargetMipCount = 1;	// No mips for clipmaps
-	BaseParams.Flags = GForceInvalidateDirectionalVSM != 0 ? NANITE_VIEW_FLAG_UNCACHED : 0u;
+	BaseParams.Flags = 0u;
 	BaseParams.bOverrideDrawDistanceOrigin = true;
 	BaseParams.DrawDistanceOrigin = CullingViewOrigin;
 
-	if (Clipmap->GetCacheEntry())
+	const TSharedPtr<FVirtualShadowMapPerLightCacheEntry>& CacheEntry = Clipmap->GetCacheEntry();
+	if (CacheEntry.IsValid())
 	{
-		Clipmap->GetCacheEntry()->MarkRendered(Scene.GetFrameNumber());
+		CacheEntry->MarkRendered(Scene.GetFrameNumber());
+
+		// Enable the force-uncaching to remove needless caching overhead if there is nothing to cache (the light is constantly invalidated).
+		BaseParams.Flags |= CacheEntry->IsUncached() ? NANITE_VIEW_FLAG_UNCACHED : 0u;
 	}
 
 	for (int32 ClipmapLevelIndex = 0; ClipmapLevelIndex < Clipmap->GetLevelCount(); ++ClipmapLevelIndex)
@@ -3322,9 +3332,12 @@ uint32 FVirtualShadowMapArray::AddRenderViews(const FProjectedShadowInfo* Projec
 	BaseParams.bOverrideDrawDistanceOrigin = true;
 	BaseParams.DrawDistanceOrigin = ClosestCullingViewOrigin;
 
-	if (ProjectedShadowInfo->VirtualShadowMapPerLightCacheEntry)
+	TSharedPtr<FVirtualShadowMapPerLightCacheEntry> CacheEntry = ProjectedShadowInfo->VirtualShadowMapPerLightCacheEntry;
+	if (CacheEntry.IsValid())
 	{
-		ProjectedShadowInfo->VirtualShadowMapPerLightCacheEntry->MarkRendered(Scene.GetFrameNumber());
+		CacheEntry->MarkRendered(Scene.GetFrameNumber());
+
+		BaseParams.Flags |= CacheEntry->IsUncached() ? NANITE_VIEW_FLAG_UNCACHED : 0U;
 	}
 
 	int32 NumMaps = ProjectedShadowInfo->bOnePassPointLightShadow ? 6 : 1;
@@ -3418,3 +3431,8 @@ void FVirtualShadowMapArray::AddVisualizePass(FRDGBuilder& GraphBuilder, const F
 	}
 #endif
 }
+
+float FVirtualShadowMapArray::InterpolateResolutionBias(float BiasNonMoving, float BiasMoving, float LightMobilityFactor)
+{
+	return FMath::Lerp(BiasNonMoving, FMath::Max(BiasNonMoving, BiasMoving), LightMobilityFactor);
+}
diff --git a/Engine/Source/Runtime/Renderer/Private/VirtualShadowMaps/VirtualShadowMapArray.h b/Engine/Source/Runtime/Renderer/Private/VirtualShadowMaps/VirtualShadowMapArray.h
--- a/Engine/Source/Runtime/Renderer/Private/VirtualShadowMaps/VirtualShadowMapArray.h
+++ b/Engine/Source/Runtime/Renderer/Private/VirtualShadowMaps/VirtualShadowMapArray.h
@@ -241,9 +241,9 @@ public:
 	}
 
 	/**
-	 * Get configured LOD bias for the local lights (maps to cvar, 
+	 * Get configured LOD bias for the local lights, LightMobilityFactor [0,1] selects between (by interpolation) the two LOD biases for static (non-moving) and moving lights.
 	 */
-	float GetResolutionLODBiasLocal() const;
+	float GetResolutionLODBiasLocal(float LightMobilityFactor) const;
 
 	// Raw size of the physical pool, including both static and dynamic pages (if enabled)
 	FIntPoint GetPhysicalPoolSize() const;
@@ -265,11 +265,11 @@ public:
 	void BuildPageAllocations(
 		FRDGBuilder& GraphBuilder,
 		const FMinimalSceneTextures& SceneTextures,
-		const TArray<FViewInfo> &Views, 
+		const TConstArrayView<FViewInfo> &Views,
 		const FEngineShowFlags& EngineShowFlags,
 		const FSortedLightSetSceneInfo& SortedLights, 
-		const TArray<FVisibleLightInfo, SceneRenderingAllocator> &VisibleLightInfos, 
-		const TArray<Nanite::FRasterResults, TInlineAllocator<2>> &NaniteRasterResults,
+		const TConstArrayView<FVisibleLightInfo>& VisibleLightInfos,
+		const TFunctionRef<float(int32 LightId)>& GetLightMobilityFactor, // TODO: propagate setup in a better(tm) way
 		const FSingleLayerWaterPrePassResult* SingleLayerWaterPrePassResult);
 
 	bool IsAllocated() const
@@ -323,6 +323,12 @@ public:
 	bool UseHzbOcclusion() const { return bUseHzbOcclusion; }
 	bool UseTwoPassHzbOcclusion() const { return bUseTwoPassHzbOcclusion; }
 
+	/**
+	 * Helper function to add clamping when interpolating the LOD resolution biases to ensure the bias for moving lights can never be lower than the one for not.
+	 * This could occur fairly easily since it is possible to both set the values through console as well as scalability.
+	 */
+	static float InterpolateResolutionBias(float BiasNonMoving, float BiasMoving, float LightMobilityFactor);
+
 	// We keep a reference to the cache manager that was used to initialize this frame as it owns some of the buffers
 	FVirtualShadowMapArrayCacheManager* CacheManager = nullptr;
 
diff --git a/Engine/Source/Runtime/Renderer/Private/VirtualShadowMaps/VirtualShadowMapCacheManager.cpp b/Engine/Source/Runtime/Renderer/Private/VirtualShadowMaps/VirtualShadowMapCacheManager.cpp
--- a/Engine/Source/Runtime/Renderer/Private/VirtualShadowMaps/VirtualShadowMapCacheManager.cpp
+++ b/Engine/Source/Runtime/Renderer/Private/VirtualShadowMaps/VirtualShadowMapCacheManager.cpp
@@ -73,7 +73,6 @@ static FAutoConsoleVariableRef  CVarForceInvalidateDirectionalVSM(
 
 void FVirtualShadowMapCacheEntry::UpdateClipmap(
 	int32 VirtualShadowMapId,
-	const FMatrix &WorldToLight,
 	FInt64Point PageSpaceLocation,
 	FInt64Point ClipmapCornerOffset,
 	double LevelRadius,
@@ -82,14 +81,8 @@ void FVirtualShadowMapCacheEntry::UpdateClipmap(
 	double ViewRadiusZ,
 	const FVirtualShadowMapPerLightCacheEntry& PerLightEntry)
 {
-	bool bCacheValid = (CurrentVirtualShadowMapId != INDEX_NONE) && GForceInvalidateDirectionalVSM == 0;
+	bool bCacheValid = CurrentVirtualShadowMapId != INDEX_NONE;
 	
-	if (bCacheValid && WorldToLight != Clipmap.WorldToLight)
-	{
-		bCacheValid = false;
-		//UE_LOG(LogRenderer, Display, TEXT("Invalidated clipmap level (VSM %d) due to light movement"), VirtualShadowMapId);
-	}
-
 	if (bCacheValid && GClipmapPanning == 0)
 	{
 		if (PageSpaceLocation.X != PrevPageSpaceLocation.X ||
@@ -113,7 +106,7 @@ void FVirtualShadowMapCacheEntry::UpdateClipmap(
 	}
 
 	// Not valid if it was never rendered
-	if (PerLightEntry.PrevRenderedFrameNumber < 0)
+	if (PerLightEntry.Prev.RenderedFrameNumber < 0)
 	{
 		bCacheValid = false;
 	}
@@ -134,7 +127,6 @@ void FVirtualShadowMapCacheEntry::UpdateClipmap(
 
 		// New cached level
 		PrevVirtualShadowMapId = INDEX_NONE;
-		Clipmap.WorldToLight = WorldToLight;
 		Clipmap.ViewCenterZ = ViewCenterZ;
 		Clipmap.ViewRadiusZ = ViewRadiusZ;
 	}
@@ -155,13 +147,13 @@ void FVirtualShadowMapCacheEntry::UpdateLocal(int32 VirtualShadowMapId, const FV
 	PrevVirtualShadowMapId = CurrentVirtualShadowMapId;
 
 	// Not valid if it was never rendered
-	if (PerLightEntry.PrevRenderedFrameNumber < 0)
+	if (PerLightEntry.Prev.RenderedFrameNumber < 0)
 	{
 		PrevVirtualShadowMapId = INDEX_NONE;
 	}
 
 	// Invalidate on transition from distant to full
-	if (!PerLightEntry.bCurrentIsDistantLight && PerLightEntry.bPrevIsDistantLight)
+	if (!PerLightEntry.Current.bIsDistantLight && PerLightEntry.Prev.bIsDistantLight)
 	{
 		PrevVirtualShadowMapId = INDEX_NONE;
 	}
@@ -175,17 +167,33 @@ void FVirtualShadowMapCacheEntry::Invalidate()
 	PrevVirtualShadowMapId = INDEX_NONE;
 }
 
-void FVirtualShadowMapPerLightCacheEntry::UpdateClipmap()
+void FVirtualShadowMapPerLightCacheEntry::UpdateClipmap(const FMatrix& WorldToLight)
 {
-	PrevRenderedFrameNumber = FMath::Max(PrevRenderedFrameNumber, CurrentRenderedFrameNumber);
-	CurrentRenderedFrameNumber = -1;
+	Prev.bIsUncached = Current.bIsUncached;
+	Prev.RenderedFrameNumber = FMath::Max(Prev.RenderedFrameNumber, Current.RenderedFrameNumber);
+	Current.RenderedFrameNumber = -1;
+	if (GForceInvalidateDirectionalVSM != 0 || WorldToLight != ClipmapCacheKey.WorldToLight)
+	{
+		Prev.RenderedFrameNumber = -1;
+		//UE_LOG(LogRenderer, Display, TEXT("Invalidated clipmap level (VSM %d) due to light movement"), VirtualShadowMapId);
+	}
+	ClipmapCacheKey.WorldToLight = WorldToLight;
+
+	Current.bIsUncached = GForceInvalidateDirectionalVSM != 0 || Prev.RenderedFrameNumber < 0;
+
+	// On transition from uncached -> cached we must invalidate since the static pages are not initialized at this point.
+	if (!Current.bIsUncached && Prev.bIsUncached)
+	{
+		Prev.RenderedFrameNumber = -1;
+	}
 }
 
 bool FVirtualShadowMapPerLightCacheEntry::UpdateLocal(const FProjectedShadowInitializer& InCacheKey, bool bIsDistantLight, bool bAllowInvalidation)
 {
-	bPrevIsDistantLight = bCurrentIsDistantLight;
-	PrevRenderedFrameNumber = FMath::Max(PrevRenderedFrameNumber, CurrentRenderedFrameNumber);
-	PrevScheduledFrameNumber = FMath::Max(PrevScheduledFrameNumber, CurrenScheduledFrameNumber);
+	Prev.bIsUncached = Current.bIsUncached;
+	Prev.bIsDistantLight = Current.bIsDistantLight;
+	Prev.RenderedFrameNumber = FMath::Max(Prev.RenderedFrameNumber, Current.RenderedFrameNumber);
+	Prev.ScheduledFrameNumber = FMath::Max(Prev.ScheduledFrameNumber, Current.ScheduledFrameNumber);
 
 	// Check cache validity based of shadow setup
 	if (!LocalCacheKey.IsCachedShadowValid(InCacheKey))
@@ -194,22 +202,29 @@ bool FVirtualShadowMapPerLightCacheEntry::UpdateLocal(const FProjectedShadowInit
 		// TODO: track invalidation state somehow for later.
 		if (bAllowInvalidation)
 		{
-			PrevRenderedFrameNumber = -1;
+			Prev.RenderedFrameNumber = -1;
 		}
 		//UE_LOG(LogRenderer, Display, TEXT("Invalidated!"));
 	}
 	LocalCacheKey = InCacheKey;
 
-	bCurrentIsDistantLight = bIsDistantLight;
-	CurrentRenderedFrameNumber = -1;
-	CurrenScheduledFrameNumber = -1;
+	Current.bIsDistantLight = bIsDistantLight;
+	Current.RenderedFrameNumber = -1;
+	Current.ScheduledFrameNumber = -1;
+	Current.bIsUncached = Prev.RenderedFrameNumber < 0;
+
+	// On transition from uncached -> cached we must invalidate since the static pages are not initialized at this point.
+	if (!Current.bIsUncached && Prev.bIsUncached)
+	{
+		Prev.RenderedFrameNumber = -1;
+	}
 
-	return PrevRenderedFrameNumber >= 0;
+	return Prev.RenderedFrameNumber >= 0;
 }
 
 void FVirtualShadowMapPerLightCacheEntry::Invalidate()
 {
-	PrevRenderedFrameNumber = -1;
+	Prev.RenderedFrameNumber = -1;
 
 	for (TSharedPtr<FVirtualShadowMapCacheEntry>& Entry : ShadowMapEntries)
 	{
@@ -1025,19 +1040,6 @@ void FVirtualShadowMapArrayCacheManager::ProcessRemovedOrUpdatedPrimitives(FRDGB
 	}
 }
 
-static void ResizeFlagArray(TBitArray<>& BitArray, int32 NewMax, bool DefaultValue = false)
-{
-	if (BitArray.Num() > NewMax)
-	{
-		// Trim off excess items
-		BitArray.SetNumUninitialized(NewMax);
-	}
-	else if (BitArray.Num() < NewMax)
-	{
-		BitArray.Add(DefaultValue, NewMax - BitArray.Num());
-	}
-}
-
 void FVirtualShadowMapArrayCacheManager::OnSceneChange()
 {
 	if (CVarCacheVirtualSMs.GetValueOnRenderThread() != 0)
@@ -1046,13 +1048,13 @@ void FVirtualShadowMapArrayCacheManager::OnSceneChange()
 
 		for (auto& CacheEntry : PrevCacheEntries)
 		{
-			ResizeFlagArray(CacheEntry.Value->CachedPrimitives, MaxPersistentPrimitiveIndex);
-			ResizeFlagArray(CacheEntry.Value->RenderedPrimitives, MaxPersistentPrimitiveIndex);
+			CacheEntry.Value->CachedPrimitives.SetNum(MaxPersistentPrimitiveIndex, false);
+			CacheEntry.Value->RenderedPrimitives.SetNum(MaxPersistentPrimitiveIndex, false);
 		}
 		for (auto& CacheEntry : CacheEntries)
 		{
-			ResizeFlagArray(CacheEntry.Value->CachedPrimitives, MaxPersistentPrimitiveIndex);
-			ResizeFlagArray(CacheEntry.Value->RenderedPrimitives, MaxPersistentPrimitiveIndex);
+			CacheEntry.Value->CachedPrimitives.SetNum(MaxPersistentPrimitiveIndex, false);
+			CacheEntry.Value->RenderedPrimitives.SetNum(MaxPersistentPrimitiveIndex, false);
 		}
 	}
 }
diff --git a/Engine/Source/Runtime/Renderer/Private/VirtualShadowMaps/VirtualShadowMapCacheManager.h b/Engine/Source/Runtime/Renderer/Private/VirtualShadowMaps/VirtualShadowMapCacheManager.h
--- a/Engine/Source/Runtime/Renderer/Private/VirtualShadowMaps/VirtualShadowMapCacheManager.h
+++ b/Engine/Source/Runtime/Renderer/Private/VirtualShadowMaps/VirtualShadowMapCacheManager.h
@@ -27,7 +27,6 @@ public:
 	void UpdateLocal(int32 VirtualShadowMapId, const FVirtualShadowMapPerLightCacheEntry &PerLightEntry);
 
 	void UpdateClipmap(int32 VirtualShadowMapId,
-		const FMatrix &WorldToLight,
 		FInt64Point PageSpaceLocation,
 		FInt64Point ClipmapCornerOffset,
 		double LevelRadius,
@@ -47,7 +46,6 @@ public:
 
 	struct FClipmapInfo
 	{
-		FMatrix	WorldToLight;
 		double ViewCenterZ;
 		double ViewRadiusZ;
 
@@ -74,11 +72,17 @@ public:
 	void OnPrimitiveRendered(const FPrimitiveSceneInfo* PrimitiveSceneInfo);
 	/**
 	 * The (local) VSM is fully cached if the previous frame if was distant and is distant this frame also.
+	 * "Fully" implies that we know all pages are mapped as well as rendered to (ignoring potential CPU-side object culling).
+	 */
+	inline bool IsFullyCached() const { return Current.bIsDistantLight && Prev.bIsDistantLight && Prev.RenderedFrameNumber >= 0; }
+
+	/**
 	 */
-	inline bool IsFullyCached() const { return bCurrentIsDistantLight && bPrevIsDistantLight && PrevRenderedFrameNumber >= 0; }
-	void MarkRendered(int32 FrameIndex) { CurrentRenderedFrameNumber = FrameIndex; }
-	int32 GetLastScheduledFrameNumber() const { return PrevScheduledFrameNumber; }
-	void UpdateClipmap();
+	inline bool IsUncached() const { return Current.bIsUncached; }
+
+	void MarkRendered(int32 FrameIndex) { Current.RenderedFrameNumber = FrameIndex; }
+	int32 GetLastScheduledFrameNumber() const { return Prev.ScheduledFrameNumber; }
+	void UpdateClipmap(const FMatrix& WorldToLight);
 	/**
 	 * Returns true if the cache entry is valid (has previous state).
 	 */
@@ -89,13 +93,17 @@ public:
 	 */
 	void Invalidate();
 
-	bool bPrevIsDistantLight = false;
-	int32 PrevRenderedFrameNumber = -1;
-	int32 PrevScheduledFrameNumber = -1;
+	struct FFrameState
+	{
+		bool bIsUncached = false;
+		bool bIsDistantLight = false;
+		int32 RenderedFrameNumber = -1;
+		int32 ScheduledFrameNumber = -1;
+
+	};
+	FFrameState Prev;
+	FFrameState Current;
 
-	bool bCurrentIsDistantLight = false;
-	int32 CurrentRenderedFrameNumber = -1;
-	int32 CurrenScheduledFrameNumber = -1;
 	// Primitives that have been rendered (not culled) the previous frame, when a primitive transitions from being culled to not it must be rendered into the VSM
 	// Key culling reasons are small size or distance cutoff.
 	TBitArray<> RenderedPrimitives;
@@ -118,6 +126,11 @@ public:
 
 private:
 	FProjectedShadowInitializer LocalCacheKey;
+	struct FClipmapCacheKey
+	{
+		FMatrix	WorldToLight;
+	};
+	FClipmapCacheKey ClipmapCacheKey;
 };
 
 class FVirtualShadowMapFeedback
diff --git a/Engine/Source/Runtime/Renderer/Private/VirtualShadowMaps/VirtualShadowMapClipmap.cpp b/Engine/Source/Runtime/Renderer/Private/VirtualShadowMaps/VirtualShadowMapClipmap.cpp
--- a/Engine/Source/Runtime/Renderer/Private/VirtualShadowMaps/VirtualShadowMapClipmap.cpp
+++ b/Engine/Source/Runtime/Renderer/Private/VirtualShadowMaps/VirtualShadowMapClipmap.cpp
@@ -24,6 +24,14 @@ static TAutoConsoleVariable<float> CVarVirtualShadowMapResolutionLodBiasDirectio
 	TEXT( "Bias applied to LOD calculations for directional lights. -1.0 doubles resolution, 1.0 halves it and so on." ),
 	ECVF_Scalability | ECVF_RenderThreadSafe
 );
+
+static TAutoConsoleVariable<float> CVarVirtualShadowMapResolutionLodBiasDirectionalMoving(
+	TEXT( "r.Shadow.Virtual.ResolutionLodBiasDirectionalMoving" ),
+	0.5f,
+	TEXT( "Bias applied to LOD calculations for directional lights that are moving. -1.0 doubles resolution, 1.0 halves it and so on." ),
+	ECVF_Scalability | ECVF_RenderThreadSafe
+);
+
 static TAutoConsoleVariable<int32> CVarVirtualShadowMapClipmapFirstLevel(
 	TEXT( "r.Shadow.Virtual.Clipmap.FirstLevel" ),
 	6,
@@ -82,7 +90,8 @@ FVirtualShadowMapClipmap::FVirtualShadowMapClipmap(
 	const FMatrix& WorldToLightRotationMatrix,
 	const FViewMatrices& CameraViewMatrices,
 	FIntPoint CameraViewRectSize,
-	const FViewInfo* InDependentView)
+	const FViewInfo* InDependentView,
+	float LightMobilityFactor)
 	: LightSceneInfo(InLightSceneInfo),
 	  DependentView(InDependentView)
 {
@@ -114,7 +123,7 @@ FVirtualShadowMapClipmap::FVirtualShadowMapClipmap(
 	// For now we adjust resolution by just biasing the page we look up in. This is wasteful in terms of page table vs.
 	// just resizing the virtual shadow maps for each clipmap, but convenient for now. This means we need to additionally bias
 	// which levels are present.
-	ResolutionLodBias = CVarVirtualShadowMapResolutionLodBiasDirectional.GetValueOnRenderThread() + FMath::Log2(LodScale);
+	ResolutionLodBias = FVirtualShadowMapArray::InterpolateResolutionBias(CVarVirtualShadowMapResolutionLodBiasDirectional.GetValueOnRenderThread(), CVarVirtualShadowMapResolutionLodBiasDirectionalMoving.GetValueOnRenderThread(), LightMobilityFactor) + FMath::Log2(LodScale);
 	// Clamp negative absolute resolution biases as they would exceed the maximum resolution/ranges allocated
 	ResolutionLodBias = FMath::Max(0.0f, ResolutionLodBias);
 
@@ -134,7 +143,7 @@ FVirtualShadowMapClipmap::FVirtualShadowMapClipmap(
 		PerLightCacheEntry = VirtualShadowMapArrayCacheManager->FindCreateLightCacheEntry(LightSceneInfo.Id, InDependentView->ViewState->GetViewKey());
 		if (PerLightCacheEntry.IsValid())
 		{
-			PerLightCacheEntry->UpdateClipmap();
+			PerLightCacheEntry->UpdateClipmap(WorldToLightRotationMatrix);
 		}
 	}
 
@@ -214,7 +223,6 @@ FVirtualShadowMapClipmap::FVirtualShadowMapClipmap(
 			FInt64Point PageOffset(CornerOffset * (FVirtualShadowMap::Level0DimPagesXY >> 2));
 
 			CacheEntry->UpdateClipmap(Level.VirtualShadowMap->ID,
-				WorldToLightRotationMatrix,
 				PageOffset,
 				CornerOffset,
 				RawLevelRadius,
@@ -317,7 +325,11 @@ FVirtualShadowMapProjectionShaderData FVirtualShadowMapClipmap::GetProjectionSha
 	Data.ResolutionLodBias = ResolutionLodBias;
 	Data.ClipmapCornerRelativeOffset = Level.RelativeCornerOffset;
 	Data.LightSourceRadius = GetLightSceneInfo().Proxy->GetSourceRadius();
-	Data.Flags = GForceInvalidateDirectionalVSM ? VSM_PROJ_FLAG_UNCACHED : 0U;
+	Data.Flags = 0U;
+	if (PerLightCacheEntry.IsValid() && PerLightCacheEntry->IsUncached())
+	{
+		Data.Flags |= VSM_PROJ_FLAG_UNCACHED;
+	}
 
 	return Data;
 }
@@ -386,4 +398,5 @@ void FVirtualShadowMapClipmap::UpdateCachedFrameData()
 	{
 		PerLightCacheEntry->RenderedPrimitives = MoveTemp(RenderedPrimitives);
 	}
-}
\ No newline at end of file
+}
+
diff --git a/Engine/Source/Runtime/Renderer/Private/VirtualShadowMaps/VirtualShadowMapClipmap.h b/Engine/Source/Runtime/Renderer/Private/VirtualShadowMaps/VirtualShadowMapClipmap.h
--- a/Engine/Source/Runtime/Renderer/Private/VirtualShadowMaps/VirtualShadowMapClipmap.h
+++ b/Engine/Source/Runtime/Renderer/Private/VirtualShadowMaps/VirtualShadowMapClipmap.h
@@ -28,7 +28,8 @@ public:
 		const FMatrix& WorldToLightRotationMatrix,
 		const FViewMatrices& CameraViewMatrices,
 		FIntPoint CameraViewRectSize,
-		const FViewInfo* InDependentView
+		const FViewInfo* InDependentView,
+		float LightMobilityFactor
 	);
 
 	FViewMatrices GetViewMatrices(int32 ClipmapIndex) const;
