Fix race condition in GatherRayTracingRelevantPrimitives_ComputeLOD when computing bCachedRayTracingGeometryValid flag

* FPrimitiveSceneInfo::UpdateCachedRaytracingData must be called before the ComputeLOD task runs to update all dirty cached RT data

diff --git a/Engine/Source/Runtime/Renderer/Private/DeferredShadingRenderer.cpp b/Engine/Source/Runtime/Renderer/Private/DeferredShadingRenderer.cpp
--- a/Engine/Source/Runtime/Renderer/Private/DeferredShadingRenderer.cpp
+++ b/Engine/Source/Runtime/Renderer/Private/DeferredShadingRenderer.cpp
@@ -768,7 +768,7 @@ struct FRayTracingRelevantPrimitiveList
 
 // Iterates over Scene's PrimitiveSceneProxies and extracts ones that are relevant for ray tracing.
 // This function can run on any thread.
-static void GatherRayTracingRelevantPrimitives(const FScene& Scene, const FViewInfo& View, FRayTracingRelevantPrimitiveList& Result)
+static void GatherRayTracingRelevantPrimitives(FScene& Scene, const FViewInfo& View, FRayTracingRelevantPrimitiveList& Result)
 {
 	Result.DirtyCachedRayTracingPrimitives.Reserve(Scene.PrimitiveSceneProxies.Num());
 
@@ -874,6 +874,8 @@ static void GatherRayTracingRelevantPrimitives(const FScene& Scene, const FViewI
 		}
 	}
 
+	FPrimitiveSceneInfo::UpdateCachedRaytracingData(&Scene, Result.DirtyCachedRayTracingPrimitives);
+
 	static const auto ICVarStaticMeshLODDistanceScale = IConsoleManager::Get().FindConsoleVariable(TEXT("r.StaticMeshLODDistanceScale"));
 	const float LODScaleCVarValue = ICVarStaticMeshLODDistanceScale->GetFloat();
 	const int32 ForcedLODLevel = GetCVarForceLOD();
@@ -914,13 +916,20 @@ static void GatherRayTracingRelevantPrimitives(const FScene& Scene, const FViewI
 
 			if (EnumHasAnyFlags(Flags, ERayTracingPrimitiveFlags::CacheInstances))
 			{
-				RelevantPrimitive.bCachedRayTracingGeometryValid = SceneInfo->IsCachedRayTracingGeometryValid();
+				FPrimitiveSceneProxy* SceneProxy = Scene.PrimitiveSceneProxies[PrimitiveIndex];
+				const bool bUsingNaniteRayTracing = (Nanite::GetRayTracingMode() != Nanite::ERayTracingMode::Fallback) && SceneProxy->IsNaniteMesh();
+
+				if (!bUsingNaniteRayTracing)
+				{
+					// Currently IsCachedRayTracingGeometryValid() can only be called for non-nanite geometries
+					RelevantPrimitive.bCachedRayTracingGeometryValid = SceneInfo->IsCachedRayTracingGeometryValid();
+					checkf(SceneInfo->CachedRayTracingInstance.GeometryRHI, TEXT("Ray tracing instance must have a valid geometry."));
+				}
+
 				RelevantPrimitive.bAnySegmentsDecal = SceneInfo->bCachedRayTracingInstanceAnySegmentsDecal;
 				RelevantPrimitive.bAllSegmentsDecal = SceneInfo->bCachedRayTracingInstanceAllSegmentsDecal;
 				RelevantPrimitive.CachedRayTracingInstance = &SceneInfo->CachedRayTracingInstance;
 
-				checkf(SceneInfo->CachedRayTracingInstance.GeometryRHI, TEXT("Ray tracing instance must have a valid geometry."));
-
 				// For primitives with ERayTracingPrimitiveFlags::CacheInstances flag we only cache the instance/mesh commands of the current LOD
 				// (see FPrimitiveSceneInfo::UpdateCachedRayTracingInstance(...) and CacheRayTracingPrimitive(...))
 				LODIndex = 0;
@@ -1072,8 +1081,6 @@ bool FDeferredShadingSceneRenderer::GatherRayTracingWorldInstancesForView(FRDGBu
 
 	INC_DWORD_STAT_BY(STAT_VisibleRayTracingPrimitives, RelevantPrimitiveList.DynamicPrimitives.Num() + RelevantPrimitiveList.StaticPrimitives.Num());
 
-	FPrimitiveSceneInfo::UpdateCachedRaytracingData(Scene, RelevantPrimitiveList.DirtyCachedRayTracingPrimitives);
-
 	{
 		TRACE_CPUPROFILER_EVENT_SCOPE(GatherRayTracingWorldInstances_DynamicElements);
 
