Fixed IsInParallelRenderingThread check in GetRuntimeVirtualTextureSceneIndex.

diff --git a/Engine/Source/Runtime/Renderer/Private/RendererScene.cpp b/Engine/Source/Runtime/Renderer/Private/RendererScene.cpp
--- a/Engine/Source/Runtime/Renderer/Private/RendererScene.cpp
+++ b/Engine/Source/Runtime/Renderer/Private/RendererScene.cpp
@@ -3625,7 +3625,7 @@ void FScene::UpdateRuntimeVirtualTextureForAllPrimitives_RenderThread()
 
 uint32 FScene::GetRuntimeVirtualTextureSceneIndex(uint32 ProducerId)
 {
-	checkSlow(IsInRenderingThread());
+	checkSlow(IsInParallelRenderingThread());
 	for (FRuntimeVirtualTextureSceneProxy const* Proxy : RuntimeVirtualTextures)
 	{
 		if (Proxy->ProducerId == ProducerId)
