Fix incorrect eye adaptation data in tonemap shader when eye adaptation is disabled.
- Incorrect buffer being bound
- Also fixing the data in "DefaultEyeAdaptationBuffer" to be a float4

diff --git a/Engine/Source/Runtime/Renderer/Private/PostProcess/PostProcessTonemap.cpp b/Engine/Source/Runtime/Renderer/Private/PostProcess/PostProcessTonemap.cpp
--- a/Engine/Source/Runtime/Renderer/Private/PostProcess/PostProcessTonemap.cpp
+++ b/Engine/Source/Runtime/Renderer/Private/PostProcess/PostProcessTonemap.cpp
@@ -552,8 +552,9 @@ FScreenPassTexture AddTonemapPass(FRDGBuilder& GraphBuilder, const FViewInfo& Vi
 
 	if (!bEyeAdaptation)
 	{
-		float DefaultEyeAdaptation = GetEyeAdaptationFixedExposure(View);
-		EyeAdaptationBuffer = CreateStructuredBuffer(GraphBuilder, TEXT("DefaultEyeAdaptationBuffer"), sizeof(float), 1, &DefaultEyeAdaptation, sizeof(float), ERDGInitialDataFlags::None);
+		const float DefaultEyeAdaptation = GetEyeAdaptationFixedExposure(View);
+		const FVector4f DefaultEyeAdaptationData(DefaultEyeAdaptation, 0, 0, DefaultEyeAdaptation);
+		EyeAdaptationBuffer = CreateStructuredBuffer(GraphBuilder, TEXT("DefaultEyeAdaptationBuffer"), sizeof(DefaultEyeAdaptationData), 1, &DefaultEyeAdaptationData, sizeof(DefaultEyeAdaptationData), ERDGInitialDataFlags::None);
 	}
 
 	const FScreenPassTextureViewport SceneColorViewport(Inputs.SceneColor);
@@ -758,7 +759,7 @@ FScreenPassTexture AddTonemapPass(FRDGBuilder& GraphBuilder, const FViewInfo& Vi
 	CommonParameters.BlurredLogLum = Inputs.BlurredLogLuminanceTexture;
 	CommonParameters.LumBilateralGridSampler = BilinearClampSampler;
 	CommonParameters.BlurredLogLumSampler = BilinearClampSampler;
-	CommonParameters.EyeAdaptationBuffer = GraphBuilder.CreateSRV(Inputs.EyeAdaptationBuffer);
+	CommonParameters.EyeAdaptationBuffer = GraphBuilder.CreateSRV(EyeAdaptationBuffer);
 	CommonParameters.EyeAdaptation = *Inputs.EyeAdaptationParameters;
 	CommonParameters.ColorGradingLUT = Inputs.ColorGradingTexture;
 	CommonParameters.ColorSampler = BilinearClampSampler;
