Fixed NonUnity/NonPch for server config

diff --git a/Engine/Source/Runtime/Engine/Private/WorldPartition/WorldPartitionStreamingPolicy.cpp b/Engine/Source/Runtime/Engine/Private/WorldPartition/WorldPartitionStreamingPolicy.cpp
--- a/Engine/Source/Runtime/Engine/Private/WorldPartition/WorldPartitionStreamingPolicy.cpp
+++ b/Engine/Source/Runtime/Engine/Private/WorldPartition/WorldPartitionStreamingPolicy.cpp
@@ -25,6 +25,8 @@
 
 #if WITH_EDITOR
 #include "LevelEditorViewport.h"
+#else
+#include "Misc/ScopeExit.h"
 #endif
 
 #define LOCTEXT_NAMESPACE "WorldPartitionStreamingPolicy"
diff --git a/Engine/Source/Runtime/Renderer/Private/Strata/StrataRoughRefraction.cpp b/Engine/Source/Runtime/Renderer/Private/Strata/StrataRoughRefraction.cpp
--- a/Engine/Source/Runtime/Renderer/Private/Strata/StrataRoughRefraction.cpp
+++ b/Engine/Source/Runtime/Renderer/Private/Strata/StrataRoughRefraction.cpp
@@ -275,10 +275,6 @@ void AddStrataOpaqueRoughRefractionPasses(
 // Keeping it simple: this should always be checked in with a value of 0
 #define STRATA_ROUGH_REFRACTION_RND 0
 
-#include "HAL/PlatformFileManager.h"
-#include "HAL/PlatformMisc.h"
-#include "Misc/FileHelper.h"
-
 #if STRATA_ROUGH_REFRACTION_RND
 
 static TAutoConsoleVariable<int32> CVarStrataRoughRefractionShadersShowRoughRefractionRnD(
