Fixing VALIDATE_UNIFORM_BUFFER_LIFETIME compile errors and getting it to run correctly in a multithreaded environment.

diff --git a/Engine/Source/Runtime/RHI/Public/RHIResources.h b/Engine/Source/Runtime/RHI/Public/RHIResources.h
--- a/Engine/Source/Runtime/RHI/Public/RHIResources.h
+++ b/Engine/Source/Runtime/RHI/Public/RHIResources.h
@@ -1267,7 +1267,7 @@ public:
 	const FRHIUniformBufferLayout* GetLayoutPtr() const { return Layout; }
 
 #if VALIDATE_UNIFORM_BUFFER_LIFETIME
-	mutable int32 NumMeshCommandReferencesForDebugging = 0;
+	mutable std::atomic<int32> NumMeshCommandReferencesForDebugging = 0;
 #endif
 
 	const TArray<TRefCountPtr<FRHIResource>>& GetResourceTable() const { return ResourceTable; }
diff --git a/Engine/Source/Runtime/Renderer/Private/MeshPassProcessor.cpp b/Engine/Source/Runtime/Renderer/Private/MeshPassProcessor.cpp
--- a/Engine/Source/Runtime/Renderer/Private/MeshPassProcessor.cpp
+++ b/Engine/Source/Runtime/Renderer/Private/MeshPassProcessor.cpp
@@ -855,12 +855,12 @@ void FMeshDrawShaderBindings::CopyFrom(const FMeshDrawShaderBindings& Other)
 	for (int32 ShaderBindingsIndex = 0; ShaderBindingsIndex < ShaderLayouts.Num(); ShaderBindingsIndex++)
 	{
 		FMeshDrawSingleShaderBindings SingleShaderBindings(ShaderLayouts[ShaderBindingsIndex], ShaderBindingDataPtr);
-		FRHIUniformBuffer** RESTRICT UniformBufferBindings = SingleShaderBindings.GetUniformBufferStart();
+		const FRHIUniformBuffer** RESTRICT UniformBufferBindings = SingleShaderBindings.GetUniformBufferStart();
 		const int32 NumUniformBuffers = SingleShaderBindings.ParameterMapInfo.UniformBuffers.Num();
 
 		for (int32 UniformBufferIndex = 0; UniformBufferIndex < NumUniformBuffers; UniformBufferIndex++)
 		{
-			FRHIUniformBuffer* UniformBuffer = UniformBufferBindings[UniformBufferIndex];
+			const FRHIUniformBuffer* UniformBuffer = UniformBufferBindings[UniformBufferIndex];
 
 			if (UniformBuffer)
 			{
@@ -881,17 +881,15 @@ void FMeshDrawShaderBindings::Release()
 	for (int32 ShaderBindingsIndex = 0; ShaderBindingsIndex < ShaderLayouts.Num(); ShaderBindingsIndex++)
 	{
 		FMeshDrawSingleShaderBindings SingleShaderBindings(ShaderLayouts[ShaderBindingsIndex], ShaderBindingDataPtr);
-		FRHIUniformBuffer** RESTRICT UniformBufferBindings = SingleShaderBindings.GetUniformBufferStart();
+		const FRHIUniformBuffer** RESTRICT UniformBufferBindings = SingleShaderBindings.GetUniformBufferStart();
 		const int32 NumUniformBuffers = SingleShaderBindings.ParameterMapInfo.UniformBuffers.Num();
 
 		for (int32 UniformBufferIndex = 0; UniformBufferIndex < NumUniformBuffers; UniformBufferIndex++)
 		{
-			FRHIUniformBuffer* UniformBuffer = UniformBufferBindings[UniformBufferIndex];
-
-			if (UniformBuffer)
+			if (const FRHIUniformBuffer* UniformBuffer = UniformBufferBindings[UniformBufferIndex])
 			{
-				UniformBuffer->NumMeshCommandReferencesForDebugging--;
-				check(UniformBuffer->NumMeshCommandReferencesForDebugging >= 0);
+				const int32 NumMeshCommandReferencesForDebugging = --UniformBuffer->NumMeshCommandReferencesForDebugging;
+				check(NumMeshCommandReferencesForDebugging >= 0);
 			}
 		}
 
diff --git a/Engine/Source/Runtime/Renderer/Public/MeshDrawShaderBindings.h b/Engine/Source/Runtime/Renderer/Public/MeshDrawShaderBindings.h
--- a/Engine/Source/Runtime/Renderer/Public/MeshDrawShaderBindings.h
+++ b/Engine/Source/Runtime/Renderer/Public/MeshDrawShaderBindings.h
@@ -301,10 +301,10 @@ private:
 		if (FoundIndex >= 0)
 		{
 #if VALIDATE_UNIFORM_BUFFER_LIFETIME
-			if (GetUniformBufferStart()[FoundIndex])
+			if (const FRHIUniformBuffer* Previous = GetUniformBufferStart()[FoundIndex])
 			{
-				GetUniformBufferStart()[FoundIndex]->NumMeshCommandReferencesForDebugging--;
-				check(GetUniformBufferStart()[FoundIndex]->NumMeshCommandReferencesForDebugging >= 0);
+				const int32 NumMeshCommandReferencesForDebugging = --Previous->NumMeshCommandReferencesForDebugging;
+				check(NumMeshCommandReferencesForDebugging >= 0);
 			}
 			Value->NumMeshCommandReferencesForDebugging++;
 #endif
