Disable async virtual texture task.
InitViews can allocate VTs from dynamic mesh batches which causes a race condition.
#swarm https://p4-swarm.epicgames.net/reviews/23855220

diff --git a/Engine/Source/Runtime/Renderer/Private/DeferredShadingRenderer.cpp b/Engine/Source/Runtime/Renderer/Private/DeferredShadingRenderer.cpp
--- a/Engine/Source/Runtime/Renderer/Private/DeferredShadingRenderer.cpp
+++ b/Engine/Source/Runtime/Renderer/Private/DeferredShadingRenderer.cpp
@@ -2626,15 +2626,9 @@ void FDeferredShadingSceneRenderer::Render(FRDGBuilder& GraphBuilder)
 	const bool bHasRayTracedOverlay = HasRayTracedOverlay(ViewFamily);
 	const bool bAllowStaticLighting = !bHasRayTracedOverlay && IsStaticLightingAllowed();
 
-	TUniquePtr<FVirtualTextureUpdater> VirtualTextureUpdater;
-
 	const bool bUseVirtualTexturing = UseVirtualTexturing(FeatureLevel);
 	if (bUseVirtualTexturing)
 	{
-		FVirtualTextureUpdateSettings Settings;
-		Settings.EnableThrottling(!ViewFamily.bOverrideVirtualTextureThrottle);
-
-		VirtualTextureUpdater = FVirtualTextureSystem::Get().BeginUpdate(GraphBuilder, FeatureLevel, Scene, Settings);
 		VirtualTextureFeedbackBegin(GraphBuilder, Views, SceneTexturesConfig.Extent);
 	}
 
@@ -2846,6 +2840,13 @@ void FDeferredShadingSceneRenderer::Render(FRDGBuilder& GraphBuilder)
 
 	if (bUseVirtualTexturing)
 	{
+		FVirtualTextureUpdateSettings Settings;
+		Settings.EnableThrottling(!ViewFamily.bOverrideVirtualTextureThrottle);
+
+		// We can move this call earlier and enable r.VT.AsyncPageRequestTask after fixing race conditions with InitViews
+		TUniquePtr<FVirtualTextureUpdater> VirtualTextureUpdater;
+		VirtualTextureUpdater = FVirtualTextureSystem::Get().BeginUpdate(GraphBuilder, FeatureLevel, Scene, Settings);
+
 		// Note, should happen after the GPU-Scene update to ensure rendering to runtime virtual textures is using the correctly updated scene
 		FVirtualTextureSystem::Get().EndUpdate(GraphBuilder, MoveTemp(VirtualTextureUpdater), FeatureLevel);
 	}
diff --git a/Engine/Source/Runtime/Renderer/Private/MobileShadingRenderer.cpp b/Engine/Source/Runtime/Renderer/Private/MobileShadingRenderer.cpp
--- a/Engine/Source/Runtime/Renderer/Private/MobileShadingRenderer.cpp
+++ b/Engine/Source/Runtime/Renderer/Private/MobileShadingRenderer.cpp
@@ -868,14 +868,8 @@ void FMobileSceneRenderer::Render(FRDGBuilder& GraphBuilder)
 
 	FRDGSystemTextures::Create(GraphBuilder);
 
-	TUniquePtr<FVirtualTextureUpdater> VirtualTextureUpdater;
-
 	if (bUseVirtualTexturing)
 	{
-		FVirtualTextureUpdateSettings Settings;
-		Settings.EnableThrottling(!ViewFamily.bOverrideVirtualTextureThrottle);
-
-		VirtualTextureUpdater = FVirtualTextureSystem::Get().BeginUpdate(GraphBuilder, FeatureLevel, Scene, Settings);
 		VirtualTextureFeedbackBegin(GraphBuilder, Views, SceneTexturesConfig.Extent);
 	}
 
@@ -931,6 +925,12 @@ void FMobileSceneRenderer::Render(FRDGBuilder& GraphBuilder)
 
 	if (bUseVirtualTexturing)
 	{
+		FVirtualTextureUpdateSettings Settings;
+		Settings.EnableThrottling(!ViewFamily.bOverrideVirtualTextureThrottle);
+
+		TUniquePtr<FVirtualTextureUpdater> VirtualTextureUpdater;
+		VirtualTextureUpdater = FVirtualTextureSystem::Get().BeginUpdate(GraphBuilder, FeatureLevel, Scene, Settings);
+
 		FVirtualTextureSystem::Get().EndUpdate(GraphBuilder, MoveTemp(VirtualTextureUpdater), FeatureLevel);
 	}
 
diff --git a/Engine/Source/Runtime/Renderer/Private/VT/VirtualTextureSystem.cpp b/Engine/Source/Runtime/Renderer/Private/VT/VirtualTextureSystem.cpp
--- a/Engine/Source/Runtime/Renderer/Private/VT/VirtualTextureSystem.cpp
+++ b/Engine/Source/Runtime/Renderer/Private/VT/VirtualTextureSystem.cpp
@@ -161,7 +161,7 @@ static TAutoConsoleVariable<int32> CVarVTCsvStats(
 );
 static TAutoConsoleVariable<int32> CVarVTAsyncPageRequestTask(
 	TEXT("r.VT.AsyncPageRequestTask"),
-	1,
+	0,
 	TEXT("Performs VT page requests on an async task."),
 	ECVF_RenderThreadSafe
 );
