Move VSM HZB metadata into the cached light structure (and make sure that structure exists even when physical page caching is disabled).
Avoids a few indirections and moves towards more persistent light state.
Potentially also fixes some theoretical cases around multi-view directional light clipmaps.

diff --git a/Engine/Source/Runtime/Renderer/Private/Shadows/ShadowSceneRenderer.cpp b/Engine/Source/Runtime/Renderer/Private/Shadows/ShadowSceneRenderer.cpp
--- a/Engine/Source/Runtime/Renderer/Private/Shadows/ShadowSceneRenderer.cpp
+++ b/Engine/Source/Runtime/Renderer/Private/Shadows/ShadowSceneRenderer.cpp
@@ -117,9 +117,9 @@ TSharedPtr<FVirtualShadowMapPerLightCacheEntry> FShadowSceneRenderer::AddLocalLi
 	if (PerLightCacheEntry.IsValid())
 	{
 		const float DistantLightForceCacheFootprintFraction = FMath::Clamp(CVarDistantLightForceCacheFootprintFraction.GetValueOnRenderThread(), 0.0f, 1.0f);
-		bool bShouldForceTimeSliceDistantUpdate = bIsDistantLight && MaxScreenRadius <= BiasedFootprintThreshold * DistantLightForceCacheFootprintFraction;
+		bool bShouldForceTimeSliceDistantUpdate = (bIsDistantLight && MaxScreenRadius <= BiasedFootprintThreshold * DistantLightForceCacheFootprintFraction);
 		LocalLightShadowFrameSetup.PerLightCacheEntry = PerLightCacheEntry;
-		bool bIsCached = PerLightCacheEntry->UpdateLocal(ProjectedShadowInitializer, bIsDistantLight, !bShouldForceTimeSliceDistantUpdate);
+		bool bIsCached = PerLightCacheEntry->UpdateLocal(ProjectedShadowInitializer, bIsDistantLight, CacheManager->IsCacheEnabled(), !bShouldForceTimeSliceDistantUpdate);
 
 		for (int32 Index = 0; Index < NumMaps; ++Index)
 		{
@@ -163,7 +163,7 @@ void FShadowSceneRenderer::PostInitDynamicShadowsSetup()
 			FProjectedShadowInfo* ProjectedShadowInfo = DirectionalLightShadowFrameSetup.ProjectedShadowInfo;
 			if (!bUnboundedClipmap && ProjectedShadowInfo->bShouldRenderVSM)
 			{
-				const bool bIsCached = VirtualShadowMapArray.CacheManager->IsValid() && GForceInvalidateDirectionalVSM == 0;
+				const bool bIsCached = VirtualShadowMapArray.CacheManager->IsCacheDataAvailable() && GForceInvalidateDirectionalVSM == 0;
 
 				// We can only do this culling if the light is both uncached & it is using the accurate bounds (i.e., r.Shadow.Virtual.Clipmap.UseConservativeCulling is turned off).
 				if (!bIsCached && !ProjectedShadowInfo->CascadeSettings.ShadowBoundsAccurate.Planes.IsEmpty())
diff --git a/Engine/Source/Runtime/Renderer/Private/VirtualShadowMaps/VirtualShadowMapArray.cpp b/Engine/Source/Runtime/Renderer/Private/VirtualShadowMaps/VirtualShadowMapArray.cpp
--- a/Engine/Source/Runtime/Renderer/Private/VirtualShadowMaps/VirtualShadowMapArray.cpp
+++ b/Engine/Source/Runtime/Renderer/Private/VirtualShadowMaps/VirtualShadowMapArray.cpp
@@ -410,7 +410,7 @@ static void SetCacheDataShaderParameters(FRDGBuilder& GraphBuilder, const TArray
 	for (int32 SmIndex = 0; SmIndex < ShadowMaps.Num(); ++SmIndex)
 	{
 		TSharedPtr<FVirtualShadowMapCacheEntry> VirtualShadowMapCacheEntry = ShadowMaps[SmIndex].IsValid() ? ShadowMaps[SmIndex]->VirtualShadowMapCacheEntry : nullptr;
-		if (VirtualShadowMapCacheEntry != nullptr && VirtualShadowMapCacheEntry->IsValid())
+		if (VirtualShadowMapCacheEntry != nullptr && VirtualShadowMapCacheEntry->PrevVirtualShadowMapId != INDEX_NONE)
 		{
 			ShadowMapCacheData[SmIndex].PrevVirtualShadowMapId = VirtualShadowMapCacheEntry->PrevVirtualShadowMapId;
 			
@@ -483,7 +483,7 @@ void FVirtualShadowMapArray::Initialize(FRDGBuilder& GraphBuilder, FVirtualShado
 		if (CVarCacheStaticSeparate.GetValueOnRenderThread() != 0)
 		{
 			#if !DEBUG_ALLOW_STATIC_SEPARATE_WITHOUT_CACHING
-			if (CacheManager->IsValid())
+			if (CacheManager->IsCacheEnabled())
 			#endif
 			{
 				// Enable separate static caching in the second texture array element
@@ -1610,7 +1610,7 @@ void FVirtualShadowMapArray::BuildPageAllocations(
 		PassParameters->bDynamicPageInvalidation = CVarDebugSkipDynamicPageInvalidation.GetValueOnRenderThread() == 0 ? 1 : 0;
 #endif
 
-		bool bCacheEnabled = CacheManager->IsValid();
+		bool bCacheEnabled = CacheManager->IsCacheDataAvailable();
 		if (bCacheEnabled)
 		{
 			SetCacheDataShaderParameters(GraphBuilder, ShadowMaps, CacheManager, PassParameters->CacheDataParameters);
@@ -2666,9 +2666,10 @@ void FVirtualShadowMapArray::RenderVirtualShadowMapsNonNanite(FRDGBuilder& Graph
 
 	FGPUScene& GPUScene = Scene.GPUScene;
 
-	FRDGBufferSRVRef PrevPageTableRDGSRV = CacheManager->IsValid() ? GraphBuilder.CreateSRV(GraphBuilder.RegisterExternalBuffer(CacheManager->PrevBuffers.PageTable, TEXT("Shadow.Virtual.PrevPageTable"))) : nullptr;
-	FRDGBufferSRVRef PrevPageFlagsRDGSRV = CacheManager->IsValid() ? GraphBuilder.CreateSRV(GraphBuilder.RegisterExternalBuffer(CacheManager->PrevBuffers.PageFlags, TEXT("Shadow.Virtual.PrevPageFlags"))) : nullptr;
-	FRDGBufferSRVRef PrevPageRectBoundsRDGSRV = CacheManager->IsValid() ? GraphBuilder.CreateSRV(GraphBuilder.RegisterExternalBuffer(CacheManager->PrevBuffers.PageRectBounds, TEXT("Shadow.Virtual.PrevPageRectBounds"))) : nullptr;
+	bool bIsCachingDataAvailable = CacheManager->IsCacheDataAvailable();
+	FRDGBufferSRVRef PrevPageTableRDGSRV = bIsCachingDataAvailable ? GraphBuilder.CreateSRV(GraphBuilder.RegisterExternalBuffer(CacheManager->PrevBuffers.PageTable, TEXT("Shadow.Virtual.PrevPageTable"))) : nullptr;
+	FRDGBufferSRVRef PrevPageFlagsRDGSRV = bIsCachingDataAvailable ? GraphBuilder.CreateSRV(GraphBuilder.RegisterExternalBuffer(CacheManager->PrevBuffers.PageFlags, TEXT("Shadow.Virtual.PrevPageFlags"))) : nullptr;
+	FRDGBufferSRVRef PrevPageRectBoundsRDGSRV = bIsCachingDataAvailable ? GraphBuilder.CreateSRV(GraphBuilder.RegisterExternalBuffer(CacheManager->PrevBuffers.PageRectBounds, TEXT("Shadow.Virtual.PrevPageRectBounds"))) : nullptr;
 
 	int32 HZBMode = CVarNonNaniteVsmUseHzb.GetValueOnRenderThread();
 	// When disabling Nanite, there may be stale data in the Nanite-HZB causing incorrect culling.
@@ -2679,7 +2680,7 @@ void FVirtualShadowMapArray::RenderVirtualShadowMapsNonNanite(FRDGBuilder& Graph
 
 	auto InitHZB = [&]()->FRDGTextureRef
 	{
-		if (HZBMode == 1 && CacheManager->IsValid())
+		if (HZBMode == 1 && bIsCachingDataAvailable)
 		{
 			return GraphBuilder.RegisterExternalTexture(CacheManager->PrevBuffers.HZBPhysical);
 		}
@@ -3265,21 +3266,24 @@ uint32 FVirtualShadowMapArray::AddRenderViews(const TSharedPtr<FVirtualShadowMap
 		Params.PrevViewMatrices = Params.ViewMatrices;
 		// TODO: MaxPixelsPerEdgeMultipler
 
-		// TODO: Clean this up - could be stored in a single structure for the whole clipmap
-		int32 HZBKey = Clipmap->GetHZBKey(ClipmapLevelIndex);
-
-		if (bSetHZBParams)
+		if (CacheEntry)
 		{
-			CacheManager->SetHZBViewParams(HZBKey, Params);
-		}
+			TSharedPtr<FVirtualShadowMapCacheEntry> LevelEntry = CacheEntry->ShadowMapEntries[Clipmap->GetClipmapLevel(ClipmapLevelIndex)];
 
-		// If we're going to generate a new HZB this frame, save the associated metadata
-		if (bUpdateHZBMetaData)
-		{
-			FVirtualShadowMapHZBMetadata& HZBMeta = HZBMetadata.FindOrAdd(HZBKey);
-			HZBMeta.TargetLayerIndex = Params.TargetLayerIndex;
-			HZBMeta.ViewMatrices = Params.ViewMatrices;
-			HZBMeta.ViewRect = Params.ViewRect;
+			if (bSetHZBParams)
+			{
+				LevelEntry->SetHZBViewParams(Params);
+			}
+
+			// If we're going to generate a new HZB this frame, save the associated metadata
+			if (bUpdateHZBMetaData)
+			{
+				FVirtualShadowMapHZBMetadata HZBMeta;
+				HZBMeta.TargetLayerIndex = Params.TargetLayerIndex;
+				HZBMeta.ViewMatrices = Params.ViewMatrices;
+				HZBMeta.ViewRect = Params.ViewRect;
+				LevelEntry->CurrentHZBMetadata = HZBMeta;
+			}
 		}
 
 		Nanite::FPackedView View = Nanite::CreatePackedView(Params);
@@ -3351,20 +3355,24 @@ uint32 FVirtualShadowMapArray::AddRenderViews(const FProjectedShadowInfo* Projec
 		Params.RangeBasedCullingDistance = ProjectedShadowInfo->GetLightSceneInfo().Proxy->GetRadius();
 		// TODO: MaxPixelsPerEdgeMultipler
 
-		int32 HZBKey = ProjectedShadowInfo->GetLightSceneInfo().Id + (Index << 24);
-
-		if (bSetHZBParams)
+		if (CacheEntry)
 		{
-			CacheManager->SetHZBViewParams(HZBKey, Params);
-		}
+			TSharedPtr<FVirtualShadowMapCacheEntry> LevelEntry = CacheEntry->ShadowMapEntries[Index];
+			
+			if (bSetHZBParams)
+			{
+				LevelEntry->SetHZBViewParams(Params);
+			}
 
-		// If we're going to generate a new HZB this frame, save the associated metadata
-		if (bUpdateHZBMetaData)
-		{
-			FVirtualShadowMapHZBMetadata& HZBMeta = HZBMetadata.FindOrAdd(HZBKey);
-			HZBMeta.TargetLayerIndex = Params.TargetLayerIndex;
-			HZBMeta.ViewMatrices = Params.ViewMatrices;
-			HZBMeta.ViewRect = Params.ViewRect;
+			// If we're going to generate a new HZB this frame, save the associated metadata
+			if (bUpdateHZBMetaData)
+			{
+				FVirtualShadowMapHZBMetadata HZBMeta;
+				HZBMeta.TargetLayerIndex = Params.TargetLayerIndex;
+				HZBMeta.ViewMatrices = Params.ViewMatrices;
+				HZBMeta.ViewRect = Params.ViewRect;
+				LevelEntry->CurrentHZBMetadata = HZBMeta;
+			}
 		}
 
 		OutVirtualShadowViews.Add(Nanite::CreatePackedView(Params));
diff --git a/Engine/Source/Runtime/Renderer/Private/VirtualShadowMaps/VirtualShadowMapArray.h b/Engine/Source/Runtime/Renderer/Private/VirtualShadowMaps/VirtualShadowMapArray.h
--- a/Engine/Source/Runtime/Renderer/Private/VirtualShadowMaps/VirtualShadowMapArray.h
+++ b/Engine/Source/Runtime/Renderer/Private/VirtualShadowMaps/VirtualShadowMapArray.h
@@ -362,7 +362,6 @@ public:
 	FRDGBufferRef StaticInvalidatingPrimitivesRDG = nullptr;
 
 	FRDGTextureRef HZBPhysical = nullptr;
-	TMap<int32, FVirtualShadowMapHZBMetadata> HZBMetadata;
 
 	// See Engine\Shaders\Private\VirtualShadowMaps\VirtualShadowMapStats.ush for definitions of the different stat indexes
 	static constexpr uint32 NumStats = 16;
diff --git a/Engine/Source/Runtime/Renderer/Private/VirtualShadowMaps/VirtualShadowMapCacheManager.cpp b/Engine/Source/Runtime/Renderer/Private/VirtualShadowMaps/VirtualShadowMapCacheManager.cpp
--- a/Engine/Source/Runtime/Renderer/Private/VirtualShadowMaps/VirtualShadowMapCacheManager.cpp
+++ b/Engine/Source/Runtime/Renderer/Private/VirtualShadowMaps/VirtualShadowMapCacheManager.cpp
@@ -81,7 +81,10 @@ void FVirtualShadowMapCacheEntry::UpdateClipmap(
 	double ViewRadiusZ,
 	const FVirtualShadowMapPerLightCacheEntry& PerLightEntry)
 {
-	bool bCacheValid = CurrentVirtualShadowMapId != INDEX_NONE;
+	PrevPageSpaceLocation = CurrentPageSpaceLocation;
+	PrevHZBMetadata = CurrentHZBMetadata;
+
+	bool bCacheValid = (CurrentVirtualShadowMapId != INDEX_NONE);
 	
 	if (bCacheValid && GClipmapPanning == 0)
 	{
@@ -130,11 +133,10 @@ void FVirtualShadowMapCacheEntry::UpdateClipmap(
 		Clipmap.ViewCenterZ = ViewCenterZ;
 		Clipmap.ViewRadiusZ = ViewRadiusZ;
 	}
-		
-	PrevPageSpaceLocation = CurrentPageSpaceLocation;
 	
 	CurrentVirtualShadowMapId = VirtualShadowMapId;
 	CurrentPageSpaceLocation = PageSpaceLocation;
+	// Current HZB metadata gets updated during rendering
 
 	Clipmap.PrevClipmapCornerOffset = Clipmap.CurrentClipmapCornerOffset;
 	Clipmap.CurrentClipmapCornerOffset = ClipmapCornerOffset;
@@ -145,6 +147,7 @@ void FVirtualShadowMapCacheEntry::UpdateLocal(int32 VirtualShadowMapId, const FV
 	// Swap previous frame data over.
 	PrevPageSpaceLocation = FInt64Point(0, 0);		// Not used for local lights
 	PrevVirtualShadowMapId = CurrentVirtualShadowMapId;
+	PrevHZBMetadata = CurrentHZBMetadata;
 
 	// Not valid if it was never rendered
 	if (PerLightEntry.Prev.RenderedFrameNumber < 0)
@@ -160,8 +163,17 @@ void FVirtualShadowMapCacheEntry::UpdateLocal(int32 VirtualShadowMapId, const FV
 
 	CurrentVirtualShadowMapId = VirtualShadowMapId;
 	CurrentPageSpaceLocation = FInt64Point(0, 0);		// Not used for local lights
+	// Current HZB metadata gets updated during rendering
+}
+
+void FVirtualShadowMapCacheEntry::SetHZBViewParams(Nanite::FPackedViewParams& OutParams)
+{
+	OutParams.PrevTargetLayerIndex = PrevHZBMetadata.TargetLayerIndex;
+	OutParams.PrevViewMatrices = PrevHZBMetadata.ViewMatrices;
+	OutParams.Flags |= NANITE_VIEW_FLAG_HZBTEST;
 }
 
+
 void FVirtualShadowMapCacheEntry::Invalidate()
 {
 	PrevVirtualShadowMapId = INDEX_NONE;
@@ -188,7 +200,7 @@ void FVirtualShadowMapPerLightCacheEntry::UpdateClipmap(const FMatrix& WorldToLi
 	}
 }
 
-bool FVirtualShadowMapPerLightCacheEntry::UpdateLocal(const FProjectedShadowInitializer& InCacheKey, bool bIsDistantLight, bool bAllowInvalidation)
+bool FVirtualShadowMapPerLightCacheEntry::UpdateLocal(const FProjectedShadowInitializer& InCacheKey, bool bIsDistantLight, bool bCacheEnabled, bool bAllowInvalidation)
 {
 	Prev.bIsUncached = Current.bIsUncached;
 	Prev.bIsDistantLight = Current.bIsDistantLight;
@@ -196,14 +208,12 @@ bool FVirtualShadowMapPerLightCacheEntry::UpdateLocal(const FProjectedShadowInit
 	Prev.ScheduledFrameNumber = FMath::Max(Prev.ScheduledFrameNumber, Current.ScheduledFrameNumber);
 
 	// Check cache validity based of shadow setup
-	if (!LocalCacheKey.IsCachedShadowValid(InCacheKey))
+	// If it is a distant light, we want to let the time-share perform the invalidation.
+	if (!bCacheEnabled
+		|| (bAllowInvalidation && !LocalCacheKey.IsCachedShadowValid(InCacheKey)))
 	{
-		// If it is a distant light, we want to let the time-share perform the invalidation.
 		// TODO: track invalidation state somehow for later.
-		if (bAllowInvalidation)
-		{
-			Prev.RenderedFrameNumber = -1;
-		}
+		Prev.RenderedFrameNumber = -1;
 		//UE_LOG(LogRenderer, Display, TEXT("Invalidated!"));
 	}
 	LocalCacheKey = InCacheKey;
@@ -678,11 +688,6 @@ TSharedPtr<FVirtualShadowMapCacheEntry> FVirtualShadowMapPerLightCacheEntry::Fin
 
 TSharedPtr<FVirtualShadowMapPerLightCacheEntry> FVirtualShadowMapArrayCacheManager::FindCreateLightCacheEntry(int32 LightSceneId, uint32 ViewUniqueID)
 {
-	if (CVarCacheVirtualSMs.GetValueOnRenderThread() == 0)
-	{
-		return nullptr;
-	}
-
 	const uint64 CacheKey = (uint64(ViewUniqueID) << 32U) | uint64(LightSceneId);
 
 	if (TSharedPtr<FVirtualShadowMapPerLightCacheEntry> *LightEntry = CacheEntries.Find(CacheKey))
@@ -790,10 +795,9 @@ void FVirtualShadowMapArrayCacheManager::ExtractFrameData(
 		{
 			bExtractHzbData = true;
 			GraphBuilder.QueueTextureExtraction(VirtualShadowMapArray.HZBPhysical, &PrevBuffers.HZBPhysical);
-			PrevBuffers.HZBMetadata = VirtualShadowMapArray.HZBMetadata;
 		}
 
-		if (CVarCacheVirtualSMs.GetValueOnRenderThread() != 0)
+		if (IsCacheEnabled())
 		{
 			bExtractHzbData = true;
 
@@ -836,14 +840,6 @@ void FVirtualShadowMapArrayCacheManager::ExtractFrameData(
 	}
 	else
 	{
-		// Caching is disabled
-		if (CVarCacheVirtualSMs.GetValueOnRenderThread() == 0)
-		{
-			// Make sure we empty out any resources associated with caching
-			PrevCacheEntries.Reset();
-			CacheEntries.Reset();
-		}
-
 		// Do nothing; maintain the data that we had
 		// This allows us to work around some cases where the renderer gets called multiple times in a given frame
 		// - such as scene captures - but does no shadow-related work in all but one of them. We do not want to drop
@@ -1001,16 +997,19 @@ void FVirtualShadowMapArrayCacheManager::ExtractStats(FRDGBuilder& GraphBuilder,
 	}
 }
 
+bool FVirtualShadowMapArrayCacheManager::IsCacheEnabled()
+{
+	return CVarCacheVirtualSMs.GetValueOnRenderThread() != 0;
+}
 
-bool FVirtualShadowMapArrayCacheManager::IsValid()
+bool FVirtualShadowMapArrayCacheManager::IsCacheDataAvailable()
 {
-	return CVarCacheVirtualSMs.GetValueOnRenderThread() != 0
+	return IsCacheEnabled()
 		&& PrevBuffers.PageTable
 		&& PrevBuffers.PageFlags
 		&& PrevBuffers.PhysicalPageMetaData;
 }
 
-
 bool FVirtualShadowMapArrayCacheManager::IsAccumulatingStats()
 {
 	return CVarAccumulateStats.GetValueOnRenderThread() != 0;
@@ -1026,7 +1025,7 @@ void FVirtualShadowMapArrayCacheManager::ProcessRemovedOrUpdatedPrimitives(FRDGB
 	// Always incorporate any scene removals into the "recently removed" list
 	UpdateRecentlyRemoved(InvalidatingPrimitiveCollector.GetRemovedPrimitives());
 
-	if (CVarCacheVirtualSMs.GetValueOnRenderThread() != 0 && PrevBuffers.PhysicalPageMetaData.IsValid())
+	if (IsCacheDataAvailable())
 	{
 		RDG_EVENT_SCOPE(GraphBuilder, "Shadow.Virtual.ProcessRemovedOrUpdatedPrimitives");
 
@@ -1042,7 +1041,6 @@ void FVirtualShadowMapArrayCacheManager::ProcessRemovedOrUpdatedPrimitives(FRDGB
 
 void FVirtualShadowMapArrayCacheManager::OnSceneChange()
 {
-	if (CVarCacheVirtualSMs.GetValueOnRenderThread() != 0)
 	{
 		const int32 MaxPersistentPrimitiveIndex = FMath::Max(1, Scene->GetMaxPersistentPrimitiveIndex());
 
@@ -1137,17 +1135,6 @@ TRDGUniformBufferRef<FVirtualShadowMapUniformParameters> FVirtualShadowMapArrayC
 	return GraphBuilder.CreateUniformBuffer(VersionedParameters);
 }
 
-void FVirtualShadowMapArrayCacheManager::SetHZBViewParams(int32 HZBKey, Nanite::FPackedViewParams& OutParams)
-{
-	FVirtualShadowMapHZBMetadata* PrevHZBMeta = PrevBuffers.HZBMetadata.Find(HZBKey);
-	if (PrevHZBMeta)
-	{
-		OutParams.PrevTargetLayerIndex = PrevHZBMeta->TargetLayerIndex;
-		OutParams.PrevViewMatrices = PrevHZBMeta->ViewMatrices;
-		OutParams.Flags |= NANITE_VIEW_FLAG_HZBTEST;
-	}
-}
-
 #if WITH_MGPU
 void FVirtualShadowMapArrayCacheManager::UpdateGPUMask(FRHIGPUMask GPUMask)
 {
diff --git a/Engine/Source/Runtime/Renderer/Private/VirtualShadowMaps/VirtualShadowMapCacheManager.h b/Engine/Source/Runtime/Renderer/Private/VirtualShadowMaps/VirtualShadowMapCacheManager.h
--- a/Engine/Source/Runtime/Renderer/Private/VirtualShadowMaps/VirtualShadowMapCacheManager.h
+++ b/Engine/Source/Runtime/Renderer/Private/VirtualShadowMaps/VirtualShadowMapCacheManager.h
@@ -20,10 +20,6 @@ namespace Nanite { struct FPackedViewParams; }
 class FVirtualShadowMapCacheEntry
 {
 public:
-	// True if the cache has been (re)populated, set to false on init and set to true once the cache update process has happened.
-	// Also set to false whenever key data was not valid and all cached data is invalidated.
-	bool IsValid() { return PrevVirtualShadowMapId != INDEX_NONE; }
-
 	void UpdateLocal(int32 VirtualShadowMapId, const FVirtualShadowMapPerLightCacheEntry &PerLightEntry);
 
 	void UpdateClipmap(int32 VirtualShadowMapId,
@@ -36,13 +32,17 @@ public:
 
 	void Invalidate();
 
+	void SetHZBViewParams(Nanite::FPackedViewParams& OutParams);
+
 	// Previous frame data
 	FInt64Point PrevPageSpaceLocation = FInt64Point(0, 0);
 	int32 PrevVirtualShadowMapId = INDEX_NONE;
+	FVirtualShadowMapHZBMetadata PrevHZBMetadata;
 
 	// Current frame data
 	FInt64Point CurrentPageSpaceLocation = FInt64Point(0, 0);
 	int32 CurrentVirtualShadowMapId = INDEX_NONE;
+	FVirtualShadowMapHZBMetadata CurrentHZBMetadata;
 
 	struct FClipmapInfo
 	{
@@ -86,7 +86,7 @@ public:
 	/**
 	 * Returns true if the cache entry is valid (has previous state).
 	 */
-	bool UpdateLocal(const FProjectedShadowInitializer &InCacheKey, bool bIsDistantLight, bool bAllowInvalidation);
+	bool UpdateLocal(const FProjectedShadowInitializer &InCacheKey, bool bIsDistantLight, bool bCacheEnabled, bool bAllowInvalidation);
 
 	/**
 	 * Mark as invalid, i.e., needing rendering.
@@ -167,7 +167,6 @@ struct FVirtualShadowMapArrayFrameData
 	TRefCountPtr<FRDGPooledBuffer>				PhysicalPageMetaData;
 
 	TRefCountPtr<IPooledRenderTarget>			HZBPhysical;
-	TMap<int32, FVirtualShadowMapHZBMetadata>	HZBMetadata;
 
 	uint64 GetGPUSizeBytes(bool bLogSizes) const;
 };
@@ -207,10 +206,8 @@ public:
 	 */
 	TSharedPtr<FVirtualShadowMapPerLightCacheEntry> FindCreateLightCacheEntry(int32 LightSceneId, uint32 ViewUniqueID = 0U);
 
-	/*
-	 * Returns true if cached data is available.
-	 */
-	bool IsValid();
+	bool IsCacheEnabled();
+	bool IsCacheDataAvailable();
 
 	bool IsAccumulatingStats();
 
@@ -292,8 +289,6 @@ public:
 
 	FVirtualShadowMapArrayFrameData PrevBuffers;
 	FVirtualShadowMapUniformParameters PrevUniformParameters;
-			
-	void SetHZBViewParams(int32 HZBKey, Nanite::FPackedViewParams& OutParams);
 
 #if WITH_MGPU
 	void UpdateGPUMask(FRHIGPUMask GPUMask);
diff --git a/Engine/Source/Runtime/Renderer/Private/VirtualShadowMaps/VirtualShadowMapClipmap.cpp b/Engine/Source/Runtime/Renderer/Private/VirtualShadowMaps/VirtualShadowMapClipmap.cpp
--- a/Engine/Source/Runtime/Renderer/Private/VirtualShadowMaps/VirtualShadowMapClipmap.cpp
+++ b/Engine/Source/Runtime/Renderer/Private/VirtualShadowMaps/VirtualShadowMapClipmap.cpp
@@ -99,8 +99,6 @@ FVirtualShadowMapClipmap::FVirtualShadowMapClipmap(
 
 	FVirtualShadowMapArrayCacheManager* VirtualShadowMapArrayCacheManager = VirtualShadowMapArray.CacheManager;
 
-	const bool bCacheValid = VirtualShadowMapArrayCacheManager && VirtualShadowMapArrayCacheManager->IsValid();
-
 	const FMatrix FaceMatrix(
 		FPlane( 0, 0, 1, 0 ),
 		FPlane( 0, 1, 0, 0 ),
@@ -289,14 +287,6 @@ FViewMatrices FVirtualShadowMapClipmap::GetViewMatrices(int32 ClipmapIndex) cons
 	return FViewMatrices(Initializer);
 }
 
-int32 FVirtualShadowMapClipmap::GetHZBKey(int32 ClipmapLevelIndex) const
-{
-	int32 AbsoluteClipmapLevel = GetClipmapLevel(ClipmapLevelIndex);		// NOTE: Can be negative!
-	int32 ClipmapLevelKey = AbsoluteClipmapLevel + 128;
-	check(ClipmapLevelKey > 0 && ClipmapLevelKey < 256);
-	return GetLightSceneInfo().Id + (ClipmapLevelKey << 24);
-}
-
 FVirtualShadowMapProjectionShaderData FVirtualShadowMapClipmap::GetProjectionShaderData(int32 ClipmapIndex) const
 {
 	check(ClipmapIndex >= 0 && ClipmapIndex < LevelData.Num());
