Fix incorrect material appearance on raytracing material with Strata enabled

* Add adaptive RT material payload for Strata.
* Centralize all Strata settings in RenderCore/RenderUtils.h to ensure consistency

diff --git a/Engine/Shaders/Private/Strata/Strata.ush b/Engine/Shaders/Private/Strata/Strata.ush
--- a/Engine/Shaders/Private/Strata/Strata.ush
+++ b/Engine/Shaders/Private/Strata/Strata.ush
@@ -1638,18 +1638,14 @@ FStrataData GetStrataSingleLayerWaterBSDF(
 ///////////////////////////////////////////////////////////////////////////////
 // Raytracing material buffer
 
+// Data payload size is determined based on Strata.BytesPerPixel settings
 #ifndef STRATA_RT_PAYLOAD_NUM_UINTS
- // When STRATA_NORMAL_QUALITY == 1, PackedTopLayerData is a uint2. To keep FPackedMaterialClosestHitPayload 
- // at 64 bytes (see declaration in BuiltInRayTracingShaders.cpp), we make data payload 4 bytes shorter)
- #define STRATA_RT_PAYLOAD_NUM_UINTS  (10 - STRATA_NORMAL_QUALITY)
+#error STRATA_RT_PAYLOAD_NUM_UINTS needs to be defined
 #endif
 
 struct FStrataRaytracingPayload
 {
-	// Total size: 40 bytes. (assuming the default value of STRATA_RT_PAYLOAD_NUM_UINTS)
-	// If this changes, update FPackedMaterialClosestHitPayload in RayTracingCommon.ush 
-	// and ensure the payload size is smaller than the maximum size
-	STRATA_TOP_LAYER_TYPE PackedTopLayerData; //
+	STRATA_TOP_LAYER_TYPE PackedTopLayerData;
 	uint Data[STRATA_RT_PAYLOAD_NUM_UINTS];
 };
 
diff --git a/Engine/Source/Editor/LevelEditor/Private/SLevelViewport.cpp b/Engine/Source/Editor/LevelEditor/Private/SLevelViewport.cpp
--- a/Engine/Source/Editor/LevelEditor/Private/SLevelViewport.cpp
+++ b/Engine/Source/Editor/LevelEditor/Private/SLevelViewport.cpp
@@ -107,12 +107,6 @@ static const FName LevelEditorName("LevelEditor");
 static FAutoConsoleCommand EnableInViewportMenu(TEXT("Editor.EnableInViewportMenu"), TEXT("Enables the new in-viewport property menu"), FConsoleCommandDelegate::CreateStatic(&SLevelViewport::EnableInViewportMenu));
 bool SLevelViewport::bInViewportMenuEnabled = false;
 
-bool LevelEditor_IsStrataEnabled()
-{
-	static const auto CVar = IConsoleManager::Get().FindTConsoleVariableDataInt(TEXT("r.Strata"));
-	return CVar && CVar->GetValueOnAnyThread() > 0;
-}
-
 #define LOCTEXT_NAMESPACE "LevelViewport"
 
 // @todo Slate Hack: Disallow game UI to be used in play in viewport until GWorld problem is fixed
@@ -1553,7 +1547,7 @@ void SLevelViewport::BindViewCommands( FUICommandList& OutCommandList )
 	FBufferVisualizationMenuCommands::Get().BindCommands(OutCommandList, Client);
 	FNaniteVisualizationMenuCommands::Get().BindCommands(OutCommandList, Client);
 	FLumenVisualizationMenuCommands::Get().BindCommands(OutCommandList, Client);
-	if (LevelEditor_IsStrataEnabled())
+	if (Strata::IsStrataEnabled())
 	{
 		FStrataVisualizationMenuCommands::Get().BindCommands(OutCommandList, Client);
 	}
diff --git a/Engine/Source/Editor/LevelEditor/Private/SLevelViewportToolBar.cpp b/Engine/Source/Editor/LevelEditor/Private/SLevelViewportToolBar.cpp
--- a/Engine/Source/Editor/LevelEditor/Private/SLevelViewportToolBar.cpp
+++ b/Engine/Source/Editor/LevelEditor/Private/SLevelViewportToolBar.cpp
@@ -52,8 +52,6 @@
 
 #define LOCTEXT_NAMESPACE "LevelViewportToolBar"
 
-bool LevelEditor_IsStrataEnabled();
-
 /** Override the view menu, just so we can specify the level viewport as active when the button is clicked */
 class SLevelEditorViewportViewMenu : public SEditorViewportViewMenu
 {
@@ -1619,7 +1617,7 @@ void SLevelViewportToolBar::FillViewMenu(UToolMenu* Menu)
 						);
 	}
 
-	if (LevelEditor_IsStrataEnabled())
+	if (Strata::IsStrataEnabled())
 	{
 		FToolMenuSection& Section = Menu->FindOrAddSection("ViewMode");
 		Section.AddSubMenu(
diff --git a/Engine/Source/Editor/MaterialEditor/Private/MaterialEditor.cpp b/Engine/Source/Editor/MaterialEditor/Private/MaterialEditor.cpp
--- a/Engine/Source/Editor/MaterialEditor/Private/MaterialEditor.cpp
+++ b/Engine/Source/Editor/MaterialEditor/Private/MaterialEditor.cpp
@@ -178,12 +178,6 @@ static TAutoConsoleVariable<int32> CVarMaterialEdUseDevShaders(
 	TEXT("Toggles whether the material editor will use shaders that include extra overhead incurred by the editor. Material editor must be re-opened if changed at runtime."),
 	ECVF_RenderThreadSafe);
 
-bool Editor_IsStrataEnabled()
-{
-	static const auto CVar = IConsoleManager::Get().FindTConsoleVariableDataInt(TEXT("r.Strata"));
-	return CVar && CVar->GetValueOnAnyThread() > 0;
-}
-
 ///////////////////////////
 // FMatExpressionPreview //
 ///////////////////////////
@@ -282,7 +276,7 @@ int32 FMatExpressionPreview::CompilePropertyAndSetMaterialProperty(EMaterialProp
 	// needs to be called in this function!!
 	Compiler->SetMaterialProperty(Property, OverrideShaderFrequency, bUsePreviousFrameTime);
 
-	if(Editor_IsStrataEnabled())
+	if(Strata::IsStrataEnabled())
 	{
 		// Set the Strata export mode to material preview
 		Compiler->SetStrataMaterialExportType(SME_MaterialPreview, EStrataMaterialExportContext::SMEC_Opaque, 0);
diff --git a/Engine/Source/Editor/MaterialEditor/Private/MaterialEditorDetailCustomization.cpp b/Engine/Source/Editor/MaterialEditor/Private/MaterialEditorDetailCustomization.cpp
--- a/Engine/Source/Editor/MaterialEditor/Private/MaterialEditorDetailCustomization.cpp
+++ b/Engine/Source/Editor/MaterialEditor/Private/MaterialEditorDetailCustomization.cpp
@@ -42,18 +42,15 @@
 #include "Engine/Texture.h"
 #include "Materials/MaterialExpressionCurveAtlasRowParameter.h"
 #include "Curves/CurveLinearColorAtlas.h"
-
+#include "RenderUtils.h"
 
 #define LOCTEXT_NAMESPACE "MaterialEditor"
 
-bool Editor_IsStrataEnabled();
-
 // Update the blend mode names based on what is supported in legacy mode or Strata mode
 UEnum* GetBlendModeEnum()
 {
-	const bool bStrataEnabled = Editor_IsStrataEnabled();
 	UEnum* BlendModeEnum = StaticEnum<EBlendMode>();
-	if (bStrataEnabled)
+	if (Strata::IsStrataEnabled())
 	{
 		// BLEND_Translucent & BLEND_TranslucentGreyTransmittance are mapped onto the same enum index
 		BlendModeEnum->SetMetaData(TEXT("DisplayName"), TEXT("TranslucentGreyTransmittance"), BLEND_Translucent);
@@ -831,7 +828,7 @@ void FMaterialDetailCustomization::CustomizeDetails( IDetailLayoutBuilder& Detai
 				}
 			}
 
-			if (!Editor_IsStrataEnabled())
+			if (!Strata::IsStrataEnabled())
 			{
 				if (PropertyName == GET_MEMBER_NAME_CHECKED(UMaterial, bIsThinSurface))
 				{
diff --git a/Engine/Source/Editor/MaterialEditor/Private/MaterialExpressionClasses.cpp b/Engine/Source/Editor/MaterialEditor/Private/MaterialExpressionClasses.cpp
--- a/Engine/Source/Editor/MaterialEditor/Private/MaterialExpressionClasses.cpp
+++ b/Engine/Source/Editor/MaterialEditor/Private/MaterialExpressionClasses.cpp
@@ -81,12 +81,6 @@ FCategorizedMaterialExpressionNode* MaterialExpressionClasses::GetCategoryNode(c
 	return NULL;
 }
 
-bool MaterialEditor_IsStrataEnabled()
-{
-	static const auto CVar = IConsoleManager::Get().FindTConsoleVariableDataInt(TEXT("r.Strata"));
-	return CVar && CVar->GetValueOnAnyThread() > 0;
-}
-
 void MaterialExpressionClasses::InitMaterialExpressionClasses()
 {
 	if(!bInitialized)
@@ -182,7 +176,7 @@ void MaterialExpressionClasses::InitMaterialExpressionClasses()
 
 							// Skip adding Strata node to the context menu if Strata is disabled
 							// STRATA_TODO: remove this when Strata becomes the only shading path
-							if ((Class->IsChildOf(UMaterialExpressionStrataBSDF::StaticClass()) || Class->IsChildOf(UMaterialExpressionStrataUtilityBase::StaticClass())) && !MaterialEditor_IsStrataEnabled())
+							if ((Class->IsChildOf(UMaterialExpressionStrataBSDF::StaticClass()) || Class->IsChildOf(UMaterialExpressionStrataUtilityBase::StaticClass())) && !Strata::IsStrataEnabled())
 							{
 								continue;
 							}
diff --git a/Engine/Source/Editor/MaterialEditor/Private/SMaterialEditorStrataWidget.cpp b/Engine/Source/Editor/MaterialEditor/Private/SMaterialEditorStrataWidget.cpp
--- a/Engine/Source/Editor/MaterialEditor/Private/SMaterialEditorStrataWidget.cpp
+++ b/Engine/Source/Editor/MaterialEditor/Private/SMaterialEditorStrataWidget.cpp
@@ -15,8 +15,6 @@
 
 DEFINE_LOG_CATEGORY_STATIC(LogMaterialEditorStrataWidget, Log, All);
 
-bool Editor_IsStrataEnabled();
-
 void SMaterialEditorStrataWidget::Construct(const FArguments& InArgs, TWeakPtr<FMaterialEditor> InMaterialEditorPtr)
 {
 	MaterialEditorPtr = InMaterialEditorPtr;
@@ -30,7 +28,7 @@ void SMaterialEditorStrataWidget::Construct(const FArguments& InArgs, TWeakPtr<F
 		.Padding(5.0f)
 		.ToolTipText(LOCTEXT("CheckBoxForceFullSimplificationToolTip", "This will force full simplification of the material."));	// Just a test, needs to be more explicit
 
-	if (Editor_IsStrataEnabled())
+	if (Strata::IsStrataEnabled())
 	{
 		this->ChildSlot
 		[
diff --git a/Engine/Source/Editor/UnrealEd/Private/MaterialGraphSchema.cpp b/Engine/Source/Editor/UnrealEd/Private/MaterialGraphSchema.cpp
--- a/Engine/Source/Editor/UnrealEd/Private/MaterialGraphSchema.cpp
+++ b/Engine/Source/Editor/UnrealEd/Private/MaterialGraphSchema.cpp
@@ -43,6 +43,7 @@
 #include "AssetRegistry/AssetRegistryModule.h"
 #include "MaterialEditorActions.h"
 #include "MaterialGraphNode_Knot.h"
+#include "RenderUtils.h"
 
 #define LOCTEXT_NAMESPACE "MaterialGraphSchema"
 
@@ -914,12 +915,6 @@ TSharedPtr<FEdGraphSchemaAction> UMaterialGraphSchema::GetCreateCommentAction()
 	return TSharedPtr<FEdGraphSchemaAction>(static_cast<FEdGraphSchemaAction*>(new FMaterialGraphSchemaAction_NewComment));
 }
 
-bool UnrealEd_IsStrataEnabled()
-{
-	static const auto CVar = IConsoleManager::Get().FindTConsoleVariableDataInt(TEXT("r.Strata"));
-	return CVar && CVar->GetValueOnAnyThread() > 0;
-}
-
 void UMaterialGraphSchema::GetMaterialFunctionActions(FGraphActionMenuBuilder& ActionMenuBuilder) const
 {
 	// Get type of dragged pin
@@ -1003,7 +998,7 @@ void UMaterialGraphSchema::GetMaterialFunctionActions(FGraphActionMenuBuilder& A
 					{
 						if (Category.ToString().Contains("Strata"))
 						{
-							bSkipMaterialFunction = !UnrealEd_IsStrataEnabled();
+							bSkipMaterialFunction = !Strata::IsStrataEnabled();
 							break;
 						}
 					}
diff --git a/Engine/Source/Runtime/Engine/Private/Materials/HLSLMaterialTranslator.cpp b/Engine/Source/Runtime/Engine/Private/Materials/HLSLMaterialTranslator.cpp
--- a/Engine/Source/Runtime/Engine/Private/Materials/HLSLMaterialTranslator.cpp
+++ b/Engine/Source/Runtime/Engine/Private/Materials/HLSLMaterialTranslator.cpp
@@ -126,19 +126,8 @@ static inline bool IsDebugTextureSampleEnabled()
 	return IsAnalyticDerivEnabled() && (GDebugTextureSampleEnabled != 0);
 }
 
-bool Engine_IsStrataEnabled();
-
 #define DEBUG_STRATA_TREE_STACK 0
 
-static uint32 GetStrataBytePerPixel(EShaderPlatform InPlatform)
-{
-	// We enforce at least 20 bytes per pixel because this is the minimal Strata GBuffer footprint of the simplest material.
-	const uint32 MinStrataBytePerPixel = 20u;
-	const uint32 MaxStrataBytePerPixel = IsMobilePlatform(InPlatform) ? 24u : 256u;
-	static FShaderPlatformCachedIniValue<int32> CVarBudget(TEXT("r.Strata.BytesPerPixel"));	
-	return FMath::Clamp(uint32(CVarBudget.Get(InPlatform)), MinStrataBytePerPixel, MaxStrataBytePerPixel);
-}
-
 /** @return the vector type containing a given number of components. */
 static inline EMaterialValueType GetVectorType(uint32 NumComponents)
 {
@@ -793,7 +782,7 @@ bool FHLSLMaterialTranslator::Translate()
 		//
 		// Process the strata tree representing the material topology.
 		//
-		const bool bStrataEnabled = Engine_IsStrataEnabled();
+		const bool bStrataEnabled = Strata::IsStrataEnabled();
 		FStrataMaterialInput* FrontMaterialInput = Material->GetMaterialInterface() ? &Material->GetMaterialInterface()->GetMaterial()->GetEditorOnlyData()->FrontMaterial : nullptr;
 		UMaterialExpression* FrontMaterialExpr = FrontMaterialInput ? FrontMaterialInput->GetTracedInput().Expression : nullptr;
 		if (bStrataEnabled && FrontMaterialExpr)
@@ -2052,7 +2041,7 @@ void FHLSLMaterialTranslator::GetMaterialEnvironment(EShaderPlatform InPlatform,
 		// Now write some feedback to the user
 		{
 			// Output some debug info as comment in code and in the material stat window
-			const uint32 StrataBytePerPixel = GetStrataBytePerPixel(InPlatform);
+			const uint32 StrataBytePerPixel = Strata::GetBytePerPixel(InPlatform);
 
 			OutEnvironment.SetDefine(TEXT("STRATA_CLAMPED_BSDF_COUNT"), StrataMaterialBSDFCount);
 
@@ -10461,7 +10450,7 @@ bool FHLSLMaterialTranslator::StrataGenerateDerivedMaterialOperatorData()
 
 	StrataSimplificationStatus.bRunFullSimplification = StrataCompilationConfig.bFullSimplify && !bStrataUsesConversionFromLegacy;
 
-	const uint32 StrataBytePerPixel = GetStrataBytePerPixel(Platform);
+	const uint32 StrataBytePerPixel = Strata::GetBytePerPixel(Platform);
 	do 
 	{
 		// Reset some data for simplifiucation iteration
diff --git a/Engine/Source/Runtime/Engine/Private/Materials/Material.cpp b/Engine/Source/Runtime/Engine/Private/Materials/Material.cpp
--- a/Engine/Source/Runtime/Engine/Private/Materials/Material.cpp
+++ b/Engine/Source/Runtime/Engine/Private/Materials/Material.cpp
@@ -123,12 +123,6 @@ static TAutoConsoleVariable<int32> CVarMaterialParameterLegacyChecks(
 	TEXT("Note that this can be slow"),
 	ECVF_RenderThreadSafe);
 
-bool Engine_IsStrataEnabled()
-{
-	static const auto CVar = IConsoleManager::Get().FindTConsoleVariableDataInt(TEXT("r.Strata"));
-	return CVar && CVar->GetValueOnAnyThread() > 0;
-}
-
 namespace MaterialImpl
 {
 	// Filtered list of properties that we follow on mobile platforms.
@@ -2175,7 +2169,7 @@ EMaterialShadingModel UMaterial::GetMaterialShadingModelFromString(const TCHAR*
 
 const TCHAR* UMaterial::GetBlendModeString(EBlendMode InBlendMode)
 {
-	const bool bStrataEnabled = Engine_IsStrataEnabled();
+	const bool bStrataEnabled = Strata::IsStrataEnabled();
 	switch (InBlendMode)
 	{
 	case BLEND_Opaque:							return TEXT("BLEND_Opaque");
@@ -3073,8 +3067,7 @@ EBlendMode ConvertLegacyBlendMode(EBlendMode InBlendMode, FMaterialShadingModelF
 void UMaterial::ConvertMaterialToStrataMaterial()
 {
 #if WITH_EDITOR
-	const bool bStrataEnabled = Engine_IsStrataEnabled();
-	if (!bStrataEnabled)
+	if (!Strata::IsStrataEnabled())
 	{
 		return;
 	}
@@ -4284,7 +4277,7 @@ bool UMaterial::CanEditChange(const FProperty* InProperty) const
 	{
 		const UMaterialEditorOnlyData* EditorOnly = GetEditorOnlyData();
 		FString PropertyName = InProperty->GetName();
-		const bool bStrataEnabled = Engine_IsStrataEnabled();
+		const bool bStrataEnabled = Strata::IsStrataEnabled();
 
 		if (PropertyName == GET_MEMBER_NAME_STRING_CHECKED(UMaterial, PhysMaterial) || PropertyName == GET_MEMBER_NAME_STRING_CHECKED(UMaterial, PhysMaterialMask))
 		{
@@ -4503,7 +4496,7 @@ void UMaterial::PostEditChangePropertyInternal(FPropertyChangedEvent& PropertyCh
 	bCanMaskedBeAssumedOpaque = !EditorOnly->OpacityMask.Expression && !(EditorOnly->OpacityMask.UseConstant && EditorOnly->OpacityMask.Constant < 0.999f) && !bUseMaterialAttributes;
 
 	// If BLEND_TranslucentColoredTransmittance is selected while Strata is not enabled, force BLEND_Translucent blend mode
-	if (!Engine_IsStrataEnabled() && BlendMode == BLEND_TranslucentColoredTransmittance)
+	if (!Strata::IsStrataEnabled() && BlendMode == BLEND_TranslucentColoredTransmittance)
 	{
 		BlendMode = BLEND_Translucent;
 	}
@@ -4726,10 +4719,9 @@ void UMaterial::UpdateExpressionParameterName(UMaterialExpression* Expression)
 void UMaterial::RebuildShadingModelField()
 {
 	ShadingModels.ClearShadingModels();
-	const bool bStrataEnabled = Engine_IsStrataEnabled();
 	UMaterialEditorOnlyData* EditorOnly = GetEditorOnlyData();
 
-	if (bStrataEnabled && EditorOnly->FrontMaterial.IsConnected())
+	if (Strata::IsStrataEnabled() && EditorOnly->FrontMaterial.IsConnected())
 	{
 		FStrataMaterialInfo StrataMaterialInfo;
 		check(EditorOnly->FrontMaterial.Expression);
@@ -6451,7 +6443,7 @@ void UMaterial::SetShadingModel(EMaterialShadingModel NewModel)
 bool UMaterial::IsPropertySupported(EMaterialProperty InProperty) const
 {
 	bool bSupported = true;
-	if (Engine_IsStrataEnabled())
+	if (Strata::IsStrataEnabled())
 	{
 		bSupported = false;
 		switch (InProperty)
@@ -6505,7 +6497,7 @@ static bool IsPropertyActive_Internal(EMaterialProperty InProperty,
 	bool bIsThinSurface,
 	bool bIsSupported)
 {
-	const bool bStrataEnabled = Engine_IsStrataEnabled();
+	const bool bStrataEnabled = Strata::IsStrataEnabled();
 
 	if (Domain == MD_PostProcess)
 	{
diff --git a/Engine/Source/Runtime/Engine/Private/Materials/MaterialExpressions.cpp b/Engine/Source/Runtime/Engine/Private/Materials/MaterialExpressions.cpp
--- a/Engine/Source/Runtime/Engine/Private/Materials/MaterialExpressions.cpp
+++ b/Engine/Source/Runtime/Engine/Private/Materials/MaterialExpressions.cpp
@@ -302,8 +302,6 @@ FUObjectAnnotationSparseBool GMaterialFunctionsThatNeedSamplerFixup;
 FUObjectAnnotationSparseBool GMaterialFunctionsThatNeedFeatureLevelSM6Fix;
 #endif // #if WITH_EDITOR
 
-bool Engine_IsStrataEnabled();
-
 /** Returns whether the given expression class is allowed. */
 bool IsAllowedExpressionType(const UClass* const Class, const bool bMaterialFunction)
 {
@@ -2034,7 +2032,7 @@ void UMaterialExpression::ConnectToPreviewMaterial(UMaterial* InMaterial, int32
 {
 	if (InMaterial && OutputIndex >= 0 && OutputIndex < Outputs.Num())
 	{
-		if (Engine_IsStrataEnabled())
+		if (Strata::IsStrataEnabled())
 		{
 			if (IsResultStrataMaterial(0))
 			{
@@ -20758,8 +20756,7 @@ int32  UMaterialExpressionClearCoatNormalCustomOutput::Compile(class FMaterialCo
 	}
 	else
 	{
-		const bool bStrata = Engine_IsStrataEnabled();
-		if (!bStrata)
+		if (!Strata::IsStrataEnabled())
 		{
 			return CompilerError(Compiler, TEXT("Input missing"));
 		}
@@ -21807,8 +21804,7 @@ int32 UMaterialExpressionSingleLayerWaterMaterialOutput::Compile(class FMaterial
 {
 	int32 CodeInput = INDEX_NONE;
 
-	static const auto CVarStrata = IConsoleManager::Get().FindTConsoleVariableDataInt(TEXT("r.Strata"));
-	const bool bStrata = CVarStrata ? CVarStrata->GetValueOnAnyThread() > 0 : false;
+	const bool bStrata = Strata::IsStrataEnabled();
 
 	if (!ScatteringCoefficients.IsConnected() && !AbsorptionCoefficients.IsConnected() && !PhaseG.IsConnected()
 		&& !bStrata)
diff --git a/Engine/Source/Runtime/Engine/Private/Materials/MaterialInstance.cpp b/Engine/Source/Runtime/Engine/Private/Materials/MaterialInstance.cpp
--- a/Engine/Source/Runtime/Engine/Private/Materials/MaterialInstance.cpp
+++ b/Engine/Source/Runtime/Engine/Private/Materials/MaterialInstance.cpp
@@ -2013,7 +2013,6 @@ void UMaterialInstance::InitStaticPermutation(EMaterialShaderPrecompileMode Prec
 	FMaterial::DeferredDeleteArray(ResourcesToFree);
 }
 
-bool Engine_IsStrataEnabled();
 EBlendMode ConvertLegacyBlendMode(EBlendMode InBlendMode, FMaterialShadingModelField InShadingModels);
 
 static void SanitizeBlendMode(TEnumAsByte<EBlendMode>& InBlendMode)
@@ -2104,8 +2103,7 @@ void UMaterialInstance::UpdateOverridableBaseProperties()
 		}
 	}
 
-	const bool bIsStratEnabled = Engine_IsStrataEnabled();
-	if (bIsStratEnabled)
+	if (Strata::IsStrataEnabled())
 	{
 		BasePropertyOverrides.BlendMode = ConvertLegacyBlendMode(BasePropertyOverrides.BlendMode, ShadingModels);
 		BlendMode = ConvertLegacyBlendMode(Parent->GetBlendMode(), ShadingModels);
@@ -3830,7 +3828,7 @@ void UMaterialInstance::PostEditChangeProperty(FPropertyChangedEvent& PropertyCh
 	}
 
 	// If BLEND_TranslucentColoredTransmittance is selected while Strata is not enabled, force BLEND_Translucent blend mode
-	if (!Engine_IsStrataEnabled())
+	if (!Strata::IsStrataEnabled())
 	{
 		SanitizeBlendMode(BlendMode);
 		SanitizeBlendMode(BasePropertyOverrides.BlendMode);
diff --git a/Engine/Source/Runtime/Engine/Private/Materials/MaterialShared.cpp b/Engine/Source/Runtime/Engine/Private/Materials/MaterialShared.cpp
--- a/Engine/Source/Runtime/Engine/Private/Materials/MaterialShared.cpp
+++ b/Engine/Source/Runtime/Engine/Private/Materials/MaterialShared.cpp
@@ -71,8 +71,6 @@ IMPLEMENT_TYPE_LAYOUT(FMaterialTextureParameterInfo);
 IMPLEMENT_TYPE_LAYOUT(FMaterialExternalTextureParameterInfo);
 IMPLEMENT_TYPE_LAYOUT(FMaterialVirtualTextureStack);
 
-bool Engine_IsStrataEnabled();
-
 struct FAllowCachingStaticParameterValues
 {
 	FAllowCachingStaticParameterValues(FMaterial& InMaterial)
@@ -1877,7 +1875,7 @@ bool FMaterialResource::IsStrataMaterial() const
 {
 	// We no longer support both types of material (Strata and non strata) so no need to check if FrontMaterial is plugged in.
 	// We simply consider all material as Strata when Strata is enabled.
-	return Engine_IsStrataEnabled();
+	return Strata::IsStrataEnabled();
 }
 
 bool FMaterialResource::HasMaterialPropertyConnected(EMaterialProperty In) const
@@ -1890,7 +1888,7 @@ bool FMaterialResource::HasMaterialPropertyConnected(EMaterialProperty In) const
 		Material->MaterialDomain == MD_Surface || 
 		Material->MaterialDomain == MD_Volume;
 
-	if (Engine_IsStrataEnabled() && bIsStrataSupportedDomain)
+	if (Strata::IsStrataEnabled() && bIsStrataSupportedDomain)
 	{
 		if (In == MP_AmbientOcclusion)
 		{
@@ -2257,7 +2255,7 @@ void FMaterial::SetupMaterialEnvironment(
 		OutEnvironment.CompilerFlags.Add(CFLAG_UsesExternalTexture);
 	}
 
-	if (!Engine_IsStrataEnabled())
+	if (!Strata::IsStrataEnabled())
 	{
 		switch(GetBlendMode())
 		{
@@ -4263,7 +4261,7 @@ int32 UMaterialInterface::CompileProperty(FMaterialCompiler* Compiler, EMaterial
 		Result = FMaterialAttributeDefinitionMap::CompileDefaultExpression(Compiler, Property);
 	}
 
-	if (Result == INDEX_NONE && Property == MP_FrontMaterial && Engine_IsStrataEnabled())
+	if (Result == INDEX_NONE && Property == MP_FrontMaterial && Strata::IsStrataEnabled())
 	{
 		Result = Compiler->StrataCreateAndRegisterNullMaterial();
 	}
@@ -4956,7 +4954,7 @@ FMaterialShaderParameters::FMaterialShaderParameters(const FMaterial* InMaterial
 	bShouldCastDynamicShadows = InMaterial->ShouldCastDynamicShadows();
 	bWritesEveryPixel = InMaterial->WritesEveryPixel(false);
 	bWritesEveryPixelShadowPass = InMaterial->WritesEveryPixel(true);
-	if (Engine_IsStrataEnabled())
+	if (Strata::IsStrataEnabled())
 	{
 		bHasDiffuseAlbedoConnected  = InMaterial->HasMaterialPropertyConnected(MP_DiffuseColor);
 		bHasF0Connected = InMaterial->HasMaterialPropertyConnected(MP_SpecularColor);
diff --git a/Engine/Source/Runtime/Engine/Private/Materials/StrataMaterial.cpp b/Engine/Source/Runtime/Engine/Private/Materials/StrataMaterial.cpp
--- a/Engine/Source/Runtime/Engine/Private/Materials/StrataMaterial.cpp
+++ b/Engine/Source/Runtime/Engine/Private/Materials/StrataMaterial.cpp
@@ -4,8 +4,7 @@
 #include "MaterialCompiler.h"
 #include "Materials/MaterialInterface.h"
 #include "StrataDefinitions.h"
-
-
+#include "RenderUtils.h"
 
 FString GetStrataBSDFName(uint8 BSDFType)
 {
@@ -48,11 +47,9 @@ FStrataRegisteredSharedLocalBasis StrataCompilationInfoCreateSharedLocalBasis(FM
 	return Compiler->StrataCompilationInfoRegisterSharedLocalBasis(NormalCodeChunk, TangentCodeChunk);
 }
 
-bool Engine_IsStrataEnabled();
-
 inline bool IsStrataEnabled()
 {
-	static bool bStrataEnabled = Engine_IsStrataEnabled();
+	static bool bStrataEnabled = Strata::IsStrataEnabled();
 	return bStrataEnabled;
 }
 
diff --git a/Engine/Source/Runtime/Engine/Private/ShaderCompiler/ShaderCompiler.cpp b/Engine/Source/Runtime/Engine/Private/ShaderCompiler/ShaderCompiler.cpp
--- a/Engine/Source/Runtime/Engine/Private/ShaderCompiler/ShaderCompiler.cpp
+++ b/Engine/Source/Runtime/Engine/Private/ShaderCompiler/ShaderCompiler.cpp
@@ -6296,58 +6296,47 @@ void GlobalBeginCompileShader(
 	}
 
 	{
-		static IConsoleVariable* CVarStrata = IConsoleManager::Get().FindConsoleVariable(TEXT("r.Strata"));
-		const bool bStrata = CVarStrata && CVarStrata->GetInt() != 0;
+		const bool bStrata = Strata::IsStrataEnabled();
 		Input.Environment.SetDefine(TEXT("STRATA_ENABLED"), bStrata ? 1 : 0);
 
 		// Force rough diffuse to be disable on platform which explicitly disable it from their settings
 		if (bStrata && IsConsolePlatform(Target.GetPlatform()))
 		{
-			static FShaderPlatformCachedIniValue<int32> CVarStrataRoughDiffuse(TEXT("r.Strata.RoughDiffuse"));
-			const bool bStrataRoughDiffuse = CVarStrataRoughDiffuse.Get(Target.GetPlatform()) != 0;
+			const bool bStrataRoughDiffuse = Strata::IsRoughDiffuseEnabled();
 			Input.Environment.SetDefine(TEXT("STRATA_DIFFUSE_CHAN"), bStrataRoughDiffuse ? 1 : 0);
 		}
 
 		// Force rough diffuse to be disable on platform which explicitly disable it from their settings
 		if (bStrata)
 		{
-			static FShaderPlatformCachedIniValue<int32> CVarStrataShadingQuality(TEXT("r.Strata.ShadingQuality"));
-			const uint32 StrataShadingQuality = CVarStrataShadingQuality.Get(Target.GetPlatform());
-			Input.Environment.SetDefine(TEXT("STRATA_SHADING_QUALITY"), FMath::Max(StrataShadingQuality, 1u));
+			const uint32 StrataShadingQuality = Strata::GetShadingQuality(Target.GetPlatform());
+			Input.Environment.SetDefine(TEXT("STRATA_SHADING_QUALITY"), StrataShadingQuality);
 			if (StrataShadingQuality > 1)
 			{
 				Input.Environment.SetDefine(TEXT("USE_ACHROMATIC_BXDF_ENERGY"), 1u);
 			}
 
 			static const auto CVarStrataGBufferFormat = IConsoleManager::Get().FindTConsoleVariableDataInt(TEXT("r.GBufferFormat"));
-			const uint32 StrataNormalQuality = CVarStrataGBufferFormat && CVarStrataGBufferFormat->GetValueOnAnyThread() > 1 ? 1 : 0;
+			const uint32 StrataNormalQuality = Strata::GetNormalQuality();
 			Input.Environment.SetDefine(TEXT("STRATA_NORMAL_QUALITY"), StrataNormalQuality);
 
-			static FShaderPlatformCachedIniValue<int32> CVarStrataTileCoord8Bits(TEXT("r.Strata.TileCoord8bits"));
-			const uint32 StrataTileCoord8Bits = CVarStrataTileCoord8Bits.Get(Target.GetPlatform());
-			Input.Environment.SetDefine(TEXT("USE_8BIT_TILE_COORD"), StrataTileCoord8Bits > 0 ? 1 : 0);
+			const uint32 StrataUintPerPixel = Strata::GetBytePerPixel(Target.GetPlatform()) / 4u;
+			Input.Environment.SetDefine(TEXT("STRATA_RT_PAYLOAD_NUM_UINTS"), StrataUintPerPixel);
 
-			// DBuffer pass is only available if high quality normal is disabled, or if we are on a console.
-			// That is because in high quality normals requires a uint2 UAV which is only supported on some graphic cards on PC 
-			// and this is an unknown when compiling shaders at this stage.
-			// !!! If this is changed, please update all sites reading r.Strata.DBufferPass !!!
-			const bool bDBufferPassSupported = StrataNormalQuality == 0 || (StrataNormalQuality > 0 && IsConsolePlatform(Target.GetPlatform()));
+			const bool bTileCoord8Bits = Strata::Is8bitTileCoordEnabled();
+			Input.Environment.SetDefine(TEXT("USE_8BIT_TILE_COORD"), bTileCoord8Bits ? 1 : 0);
 
-			static FShaderPlatformCachedIniValue<int32> CVarStrataDBufferPass(TEXT("r.Strata.DBufferPass"));
-			const uint32 StrataDBufferPass = bDBufferPassSupported ? CVarStrataDBufferPass.Get(Target.GetPlatform()) : 0;
-			Input.Environment.SetDefine(TEXT("STRATA_USE_DBUFFER_PASS"), StrataDBufferPass > 0 ? 1 : 0);
+			const bool bStrataDBufferPass = Strata::IsDBufferPassEnabled(Target.GetPlatform());
+			Input.Environment.SetDefine(TEXT("STRATA_USE_DBUFFER_PASS"), bStrataDBufferPass ? 1 : 0);
 		}
 
-		static IConsoleVariable* CVarBackCompatibility = IConsoleManager::Get().FindConsoleVariable(TEXT("r.StrataBackCompatibility"));
-		const bool bStrataBackCompatibility = CVarBackCompatibility && CVarBackCompatibility->GetInt() > 0;
+		const bool bStrataBackCompatibility = bStrata && Strata::IsBackCompatibilityEnabled();
 		Input.Environment.SetDefine(TEXT("PROJECT_STRATA_BACKCOMPATIBILITY"), bStrataBackCompatibility ? 1 : 0);
 
-		static IConsoleVariable* CVarOpaqueRoughRefrac = IConsoleManager::Get().FindConsoleVariable(TEXT("r.Strata.OpaqueMaterialRoughRefraction"));
-		const bool bStrataOpaqueRoughRefrac = bStrata && CVarOpaqueRoughRefrac && CVarOpaqueRoughRefrac->GetInt() != 0;
+		const bool bStrataOpaqueRoughRefrac = bStrata && Strata::IsOpaqueRoughRefractionEnabled();
 		Input.Environment.SetDefine(TEXT("STRATA_OPAQUE_ROUGH_REFRACTION_ENABLED"), bStrataOpaqueRoughRefrac ? 1 : 0);
 
-		static IConsoleVariable* CVarAdvDebug = IConsoleManager::Get().FindConsoleVariable(TEXT("r.Strata.Debug.AdvancedVisualizationShaders"));
-		const bool bStrataAdvDebug = bStrata && CVarAdvDebug && CVarAdvDebug->GetInt() != 0;
+		const bool bStrataAdvDebug = bStrata && Strata::IsAdvancedVisualizationEnabled();
 		Input.Environment.SetDefine(TEXT("STRATA_ADVANCED_DEBUG_ENABLED"), bStrataAdvDebug ? 1 : 0);
 	}
 
diff --git a/Engine/Source/Runtime/Engine/Private/ShaderCompiler/ShaderGenerationUtil.cpp b/Engine/Source/Runtime/Engine/Private/ShaderCompiler/ShaderGenerationUtil.cpp
--- a/Engine/Source/Runtime/Engine/Private/ShaderCompiler/ShaderGenerationUtil.cpp
+++ b/Engine/Source/Runtime/Engine/Private/ShaderCompiler/ShaderGenerationUtil.cpp
@@ -39,9 +39,6 @@ bool NeedsVelocityDepth(EShaderPlatform TargetPlatform)
 		|| FDataDrivenShaderPlatformInfo::GetSupportsRayTracing(TargetPlatform);
 }
 
-// Strata::IsEnabled is only accessible in the Renderer module
-bool Engine_IsStrataEnabled();
-
 #if WITH_EDITOR
 
 static int FetchCompileInt(const FShaderCompilerEnvironment& OutEnvironment, const char* SrcName)
@@ -1122,8 +1119,7 @@ static FString CreateGBufferDecodeFunctionDirect(const FGBufferInfo& BufferInfo)
 	// Default initialization in case no gbuffer data are generated when Strata is enabled to 
 	// prevent shader compiler error (division by zero, variable not-initialized) with passes 
 	// not converted to Strata and still using Gbuffer data
-	const bool bStrata = Engine_IsStrataEnabled();
-	if (bStrata)
+	if (Strata::IsStrataEnabled())
 	{
 		FullStr += TEXT("\tRet.WorldNormal = float3(0,0,1);\n");
 		FullStr += TEXT("\tRet.Depth = 0.f;\n");
diff --git a/Engine/Source/Runtime/Engine/Private/StrataVisualizationData.cpp b/Engine/Source/Runtime/Engine/Private/StrataVisualizationData.cpp
--- a/Engine/Source/Runtime/Engine/Private/StrataVisualizationData.cpp
+++ b/Engine/Source/Runtime/Engine/Private/StrataVisualizationData.cpp
@@ -4,40 +4,12 @@
 #include "HAL/IConsoleManager.h"
 #include "Materials/Material.h"
 #include "Misc/ConfigCacheIni.h"
+#include "RenderUtils.h"
 
 #define LOCTEXT_NAMESPACE "FStrataVisualizationData"
 
 static FStrataVisualizationData GStrataVisualizationData;
 
-bool Engine_IsStrataEnabled();
-
-static bool Engine_IsStrataDBufferPassEnabled()
-{
-	static const auto CVarStrataGBufferFormat = IConsoleManager::Get().FindTConsoleVariableDataInt(TEXT("r.GBufferFormat")); 
-	const uint32 StrataNormalQuality = CVarStrataGBufferFormat && CVarStrataGBufferFormat->GetValueOnAnyThread() > 1 ? 1 : 0;
-
-	// DBuffer pass is only available if high quality normal is disabled, or if we are on a console.
-	// That is because in high quality normals requires a uint2 UAV which is only supported on some graphic cards on PC 
-	// and this is an unknown when compiling shaders at this stage.
-	// !!! If this is changed, please update all sites reading r.Strata.DBufferPass !!!
-	const bool bDBufferPassSupported = StrataNormalQuality == 0;	// We do not have access to platform to check if we are on console here. Visualize is likely used on PC editor so we can only visualize DBuffer if NormalQuality is 0.
-
-	static const auto CVar = IConsoleManager::Get().FindTConsoleVariableDataInt(TEXT("r.Strata.DBufferPass"));
-	return CVar && CVar->GetValueOnAnyThread() > 0 && bDBufferPassSupported;
-}
-
-static bool Engine_IsStrataRoughRefractionEnabled()
-{
-	static const auto CVar = IConsoleManager::Get().FindTConsoleVariableDataInt(TEXT("r.Strata.OpaqueMaterialRoughRefraction"));
-	return CVar && CVar->GetValueOnAnyThread() > 0;
-}
-
-static bool Engine_IsStrataAdvancedDebugShaderEnabled()
-{
-	static const auto CVar = IConsoleManager::Get().FindTConsoleVariableDataInt(TEXT("r.Strata.Debug.AdvancedVisualizationShaders"));
-	return CVar && CVar->GetValueOnAnyThread() > 0;
-}
-
 static FString ConfigureConsoleCommand(FStrataVisualizationData::TModeMap& ModeMap)
 {
 	FString AvailableVisualizationModes;
@@ -92,7 +64,7 @@ static void AddVisualizationMode(
 
 void FStrataVisualizationData::Initialize()
 {
-	if (!bIsInitialized && Engine_IsStrataEnabled())
+	if (!bIsInitialized && Strata::IsStrataEnabled())
 	{
 		TModeMap AllModeMap;
 
@@ -123,7 +95,7 @@ void FStrataVisualizationData::Initialize()
 			LOCTEXT("AdvancedMaterialPropertiesDesc", "Visualizes Strata advanced material properties"),
 			FViewMode::AdvancedMaterialProperties,
 			true,
-			Engine_IsStrataAdvancedDebugShaderEnabled(),
+			Strata::IsAdvancedVisualizationEnabled(),
 			LOCTEXT("IsStrataAdvancedDebugShaderEnabled", "Strata advanced debugging r.Strata.Debug.AdvancedVisualizationShaders is disabled"));
 
 		AddVisualizationMode(
@@ -143,7 +115,7 @@ void FStrataVisualizationData::Initialize()
 			LOCTEXT("DecalClassificationDesc", "Visualizes Strata decal classification"),
 			FViewMode::DecalClassification,
 			true,
-			Engine_IsStrataDBufferPassEnabled(),
+			false, // Disable for now, as it is not important, and is mainly used for debugging
 			LOCTEXT("IsStrataDBufferPassEnabled", "Strata tiled DBuffer pass (r.Strata.DBufferPass and r.Strata.DBufferPass.DedicatedTiles) is disabled"));
 
 		AddVisualizationMode(
@@ -153,7 +125,7 @@ void FStrataVisualizationData::Initialize()
 			LOCTEXT("RoughRefractionClassificationDesc", "Visualizes Strata rough refraction classification"),
 			FViewMode::RoughRefractionClassification,
 			true,
-			Engine_IsStrataRoughRefractionEnabled(),
+			Strata::IsOpaqueRoughRefractionEnabled(),
 			LOCTEXT("IsStrataRoughRefractionEnabled", "Strata rough refraction r.Strata.OpaqueMaterialRoughRefraction is disabled"));
 
 		AddVisualizationMode(
diff --git a/Engine/Source/Runtime/RenderCore/Private/BuiltInRayTracingShaders.cpp b/Engine/Source/Runtime/RenderCore/Private/BuiltInRayTracingShaders.cpp
--- a/Engine/Source/Runtime/RenderCore/Private/BuiltInRayTracingShaders.cpp
+++ b/Engine/Source/Runtime/RenderCore/Private/BuiltInRayTracingShaders.cpp
@@ -3,13 +3,19 @@
 #include "BuiltInRayTracingShaders.h"
 #include "ShaderParameterUtils.h"
 #include "PipelineStateCache.h"
+#include "RendererInterface.h"
 
 #if RHI_RAYTRACING
 
 #include "RayTracingPayloadType.h"
 
+uint32 GetRaytracingMaterialPayloadSize()
+{
+	return Strata::IsStrataEnabled() ? Strata::GetRayTracingMaterialPayloadSizeInBytes() : 64u;
+}
+
 IMPLEMENT_RT_PAYLOAD_TYPE(ERayTracingPayloadType::Default, 24);
-IMPLEMENT_RT_PAYLOAD_TYPE(ERayTracingPayloadType::RayTracingMaterial, 64);
+IMPLEMENT_RT_PAYLOAD_TYPE_FUNCTION(ERayTracingPayloadType::RayTracingMaterial, GetRaytracingMaterialPayloadSize);
 
 IMPLEMENT_GLOBAL_SHADER( FDefaultPayloadMS, "/Engine/Private/RayTracing/RayTracingBuiltInShaders.usf", "DefaultPayloadMS", SF_RayMiss);
 IMPLEMENT_GLOBAL_SHADER( FPackedMaterialClosestHitPayloadMS, "/Engine/Private/RayTracing/RayTracingBuiltInShaders.usf", "PackedMaterialClosestHitPayloadMS", SF_RayMiss);
diff --git a/Engine/Source/Runtime/RenderCore/Private/GBufferInfo.cpp b/Engine/Source/Runtime/RenderCore/Private/GBufferInfo.cpp
--- a/Engine/Source/Runtime/RenderCore/Private/GBufferInfo.cpp
+++ b/Engine/Source/Runtime/RenderCore/Private/GBufferInfo.cpp
@@ -19,14 +19,6 @@ static const int32 EGBufferFormat_HighPrecisionNormals = 3;
 /** Forces all GBuffers to 16 bits per channel. Intended as profiling for best quality. */
 static const int32 EGBufferFormat_Force16BitsPerChannel = 5;
 
-
-// Strata::IsEnabled is only accessible in the Renderer module
-bool RenderCore_IsStrataEnabled()
-{
-	static IConsoleVariable* CVar = IConsoleManager::Get().FindConsoleVariable(TEXT("r.Strata"));
-	return CVar->GetInt() > 0;
-}
-
 static bool IsGBufferPackingEqual(const FGBufferPacking& Lhs, const FGBufferPacking& Rhs)
 {
 	return
@@ -257,8 +249,7 @@ FGBufferInfo RENDERCORE_API FetchLegacyGBufferInfo(const FGBufferParams& Params)
 	int32 TargetSeparatedMainDirLight = -1;
 
 	// Strata outputs material data through UAV. Only SceneColor, PrecalcShadow & Velocity data are still emitted through RenderTargets
-	const bool bStrata = RenderCore_IsStrataEnabled();
-	if (bStrata)
+	if (Strata::IsStrataEnabled())
 	{
 		TargetGBufferA = -1;
 		TargetGBufferB = -1;
diff --git a/Engine/Source/Runtime/RenderCore/Private/RenderUtils.cpp b/Engine/Source/Runtime/RenderCore/Private/RenderUtils.cpp
--- a/Engine/Source/Runtime/RenderCore/Private/RenderUtils.cpp
+++ b/Engine/Source/Runtime/RenderCore/Private/RenderUtils.cpp
@@ -1087,8 +1087,6 @@ RENDERCORE_API FBufferRHIRef& GetUnitCubeAABBVertexBuffer()
 }
 #endif // RHI_RAYTRACING
 
-bool RenderCore_IsStrataEnabled();
-
 RENDERCORE_API void QuantizeSceneBufferSize(const FIntPoint& InBufferSize, FIntPoint& OutBufferSize)
 {
 	// Ensure sizes are dividable by STRATA_TILE_SIZE (==8) 2d tiles to make it more convenient.
@@ -1099,7 +1097,7 @@ RENDERCORE_API void QuantizeSceneBufferSize(const FIntPoint& InBufferSize, FIntP
 	const uint32 LegacyDividableBy = 4;
 	static_assert(LegacyDividableBy % 4 == 0, "A lot of graphic algorithms where previously assuming DividableBy == 4");
 
-	const uint32 DividableBy = RenderCore_IsStrataEnabled() ? StrataDividableBy : LegacyDividableBy;
+	const uint32 DividableBy = Strata::IsStrataEnabled() ? StrataDividableBy : LegacyDividableBy;
 
 	const uint32 Mask = ~(DividableBy - 1);
 	OutBufferSize.X = (InBufferSize.X + DividableBy - 1) & Mask;
@@ -1505,3 +1503,170 @@ ERayTracingMode GetRayTracingMode()
 {
 	return IsRayTracingAllowed() ? GRayTracingMode : ERayTracingMode::Disabled;
 }
+
+///////////////////////////////////////////////////////////////////////////////////////////////////
+// Strata settings interface
+
+static TAutoConsoleVariable<int32> CVarStrata(
+	TEXT("r.Strata"),
+	0,
+	TEXT("Enable Strata materials (Beta)."),
+	ECVF_ReadOnly | ECVF_RenderThreadSafe);
+
+static TAutoConsoleVariable<int32> CVarStrataBytePerPixel(
+	TEXT("r.Strata.BytesPerPixel"),
+	80,
+	TEXT("Strata allocated byte per pixel to store materials data. Higher value means more complex material can be represented."),
+	ECVF_ReadOnly | ECVF_RenderThreadSafe);
+
+static TAutoConsoleVariable<int32> CVarStrataBackCompatibility(
+	TEXT("r.StrataBackCompatibility"),
+	0,
+	TEXT("Disables Strata multiple scattering and replaces Chan diffuse by Lambert."),
+	ECVF_ReadOnly | ECVF_RenderThreadSafe);
+
+static TAutoConsoleVariable<int32> CVarStrataOpaqueMaterialRoughRefraction(
+	TEXT("r.Strata.OpaqueMaterialRoughRefraction"),
+	0,
+	TEXT("Enable Strata opaque material rough refractions effect from top layers over layers below."),
+	ECVF_ReadOnly | ECVF_RenderThreadSafe);
+
+static TAutoConsoleVariable<int32> CVarStrataShadingQuality(
+	TEXT("r.Strata.ShadingQuality"),
+	1,
+	TEXT("Define Strata shading quality (1: accurate lighting, 2: approximate lighting). This variable is read-only."),
+	ECVF_ReadOnly | ECVF_RenderThreadSafe);
+
+static TAutoConsoleVariable<int32> CVarStrataDBufferPass(
+	TEXT("r.Strata.DBufferPass"),
+	0,
+	TEXT("Apply DBuffer after the base-pass as a separate pass. Read only because when this is changed, it will require the recompilation of all shaders."),
+	ECVF_ReadOnly | ECVF_RenderThreadSafe);
+
+// Transition render settings that will disapear when strata gets enabled
+static TAutoConsoleVariable<int32> CVarMaterialRoughDiffuse(
+	TEXT("r.Material.RoughDiffuse"),
+	0,
+	TEXT("Enable rough diffuse material."),
+	ECVF_ReadOnly | ECVF_RenderThreadSafe);
+
+static TAutoConsoleVariable<int32> CVarStrataRoughDiffuse(
+	TEXT("r.Strata.RoughDiffuse"),
+	1,
+	TEXT("Enable Strata rough diffuse model (works only if r.Material.RoughDiffuse is enabled in the project settings). Togglable at runtime"),
+	ECVF_RenderThreadSafe);
+
+static TAutoConsoleVariable<int32> CVarStrataDebugAdvancedVisualizationShaders(
+	TEXT("r.Strata.Debug.AdvancedVisualizationShaders"),
+	0,
+	TEXT("Enable advanced strata material debug visualization shaders. Base pass shaders can output such advanced data."),
+	ECVF_ReadOnly | ECVF_RenderThreadSafe);
+
+static TAutoConsoleVariable<int32> CVarStrataTileCoord8Bits(
+	TEXT("r.Strata.TileCoord8bits"),
+	0,
+	TEXT("Format of tile coord. This variable is read-only."),
+	ECVF_ReadOnly | ECVF_RenderThreadSafe);
+
+namespace Strata
+{
+	bool IsStrataEnabled()
+	{
+		return CVarStrata.GetValueOnAnyThread() > 0;
+	}
+
+	static uint32 InternalGetBytePerPixel(uint32 InBytePerPixel)
+	{	
+		// We enforce at least 20 bytes per pixel because this is the minimal Strata GBuffer footprint of the simplest material.
+		const uint32 MinStrataBytePerPixel = 20u;
+		const uint32 MaxStrataBytePerPixel = /*IsMobilePlatform(GMaxRHI InPlatform) ? 24u : */256u; // STRATA_TODO
+		return FMath::Clamp(InBytePerPixel, MinStrataBytePerPixel, MaxStrataBytePerPixel);
+	}
+
+	uint32 GetBytePerPixel()
+	{
+		return InternalGetBytePerPixel(CVarStrataBytePerPixel.GetValueOnAnyThread());
+	}
+
+	uint32 GetBytePerPixel(EShaderPlatform InPlatform)
+	{
+		// Variant for shader compilation per platform
+		static FShaderPlatformCachedIniValue<int32> CVarBudget(TEXT("r.Strata.BytesPerPixel"));
+		return InternalGetBytePerPixel(CVarBudget.Get(InPlatform));
+	}
+
+	static uint32 InternalGetRayTracingMaterialPayloadSizeInBytes(uint32 InBytePerPixels, uint32 InNormalQuality)
+	{
+		// If Strata is enabled, resize the payload accordingly to the byte per pixel request
+		const uint32 HeaderPayload = sizeof(uint32) * (5u); // See FPackedMaterialClosestHitPayload in RayTracingCommon.ush
+		const uint32 TopNormalPayload = sizeof(uint32) * (InNormalQuality == 1 ? 2u : 1u);
+
+		return HeaderPayload + TopNormalPayload + InBytePerPixels;
+	}
+
+	uint32 GetNormalQuality()
+	{
+		static const auto CVar = IConsoleManager::Get().FindTConsoleVariableDataInt(TEXT("r.GBufferFormat"));
+		return CVar && CVar->GetValueOnAnyThread() > 1 ? 1 : 0;
+	}
+
+	uint32 GetRayTracingMaterialPayloadSizeInBytes()
+	{
+		return InternalGetRayTracingMaterialPayloadSizeInBytes(GetBytePerPixel(), GetNormalQuality());
+	}
+
+	uint32 GetRayTracingMaterialPayloadSizeInBytes(EShaderPlatform InPlatform)
+	{
+		return InternalGetRayTracingMaterialPayloadSizeInBytes(GetBytePerPixel(InPlatform), GetNormalQuality());
+	}
+
+	bool IsBackCompatibilityEnabled()
+	{
+		return CVarStrataBackCompatibility.GetValueOnAnyThread() > 0 ? 1 : 0;
+	}
+	
+	bool IsRoughDiffuseEnabled()
+	{
+		return CVarStrataRoughDiffuse.GetValueOnAnyThread() > 0;
+	}
+
+	bool Is8bitTileCoordEnabled()
+	{
+		return CVarStrataTileCoord8Bits.GetValueOnAnyThread() > 0 ? 1 : 0;
+	}
+
+	uint32 GetShadingQuality()
+	{
+		return CVarStrataShadingQuality.GetValueOnAnyThread();
+	}
+
+	uint32 GetShadingQuality(EShaderPlatform InPlatform)
+
+	{
+		static FShaderPlatformCachedIniValue<int32> CVarStrataShadingQualityPlatform(TEXT("r.Strata.ShadingQuality"));
+		return CVarStrataShadingQualityPlatform.Get(InPlatform);
+	}
+	
+	bool IsDBufferPassEnabled(EShaderPlatform InPlatform)
+	{
+		// DBuffer pass is only available if high quality normal is disabled, or if we are on a console.
+		// That is because in high quality normals requires a uint2 UAV which is only supported on some graphic cards on PC 
+		// and this is an unknown when compiling shaders at this stage.
+		// !!! If this is changed, please update all sites reading r.Strata.DBufferPass !!!
+		const uint32 StrataNormalQuality = GetNormalQuality();
+		const bool bDBufferPassSupported = StrataNormalQuality == 0 || (StrataNormalQuality > 0 && IsConsolePlatform(InPlatform));
+		
+		static FShaderPlatformCachedIniValue<int32> CVarStrataDBufferPassPlatform(TEXT("r.Strata.DBufferPass"));
+		return IsStrataEnabled() && IsUsingDBuffers(InPlatform) && bDBufferPassSupported && CVarStrataDBufferPassPlatform.Get(InPlatform) > 0;
+	}
+
+	bool IsOpaqueRoughRefractionEnabled()
+	{
+		return IsStrataEnabled() && CVarStrataOpaqueMaterialRoughRefraction.GetValueOnAnyThread() > 0;
+	}
+
+	bool IsAdvancedVisualizationEnabled()
+	{
+		return CVarStrataDebugAdvancedVisualizationShaders.GetValueOnAnyThread() > 0;
+	}
+}
\ No newline at end of file
diff --git a/Engine/Source/Runtime/RenderCore/Private/Shader.cpp b/Engine/Source/Runtime/RenderCore/Private/Shader.cpp
--- a/Engine/Source/Runtime/RenderCore/Private/Shader.cpp
+++ b/Engine/Source/Runtime/RenderCore/Private/Shader.cpp
@@ -109,8 +109,6 @@ static TLinkedList<FShaderPipelineType*>*	GShaderPipelineList = nullptr;
 
 static FSHAHash ShaderSourceDefaultHash; //will only be read (never written) for the cooking case
 
-bool RenderCore_IsStrataEnabled();
-
 /**
  * Find the shader pipeline type with the given name.
  * @return NULL if no type matched.
@@ -209,6 +207,7 @@ namespace {
 
 uint32 RegisteredRayTracingPayloads = 0;
 uint32 RayTracingPayloadSizes[32] = {};
+TRaytracingPayloadSizeFunction RayTracingPayloadSizeFunctions[32] = {};
 
 bool IsRayTracingPayloadRegistered(ERayTracingPayloadType PayloadType)
 {
@@ -1919,63 +1918,37 @@ void ShaderMapAppendKeyString(EShaderPlatform Platform, FString& KeyString)
 		}
 	}
 
+	if (Strata::IsStrataEnabled())
 	{
-		const bool bStrataEnabled = RenderCore_IsStrataEnabled();
-		if (bStrataEnabled)
 		{
 			KeyString += TEXT("_STRATA");
 		}
 
-		if (bStrataEnabled)
 		{
-			// We enforce at least 20 bytes per pixel because this is the minimal Strata GBuffer footprint of the simplest material.
-			const uint32 MinStrataBytePerPixel = 20u;
-			const uint32 MaxStrataBytePerPixel = IsMobilePlatform(Platform) ? 24u : 256u;
-			static FShaderPlatformCachedIniValue<int32> CVarBudget(TEXT("r.Strata.BytesPerPixel"));
-			const uint32 BytesPerPixel = FMath::Clamp(uint32(CVarBudget.Get(Platform)), MinStrataBytePerPixel, MaxStrataBytePerPixel);
-			KeyString += FString::Printf(TEXT("_BUDGET%u"), BytesPerPixel);
+			KeyString += FString::Printf(TEXT("_BUDGET%u"), Strata::GetBytePerPixel(Platform));
 		}
 
-		if (bStrataEnabled)
+		if (Strata::IsDBufferPassEnabled(Platform))
 		{
-			static const auto CVarStrataGBufferFormat = IConsoleManager::Get().FindTConsoleVariableDataInt(TEXT("r.GBufferFormat"));
-			const uint32 StrataNormalQuality = CVarStrataGBufferFormat && CVarStrataGBufferFormat->GetValueOnAnyThread() > 1 ? 1 : 0;
-
-			// DBuffer pass is only available if high quality normal is disabled, or if we are on a console.
-			// That is because in high quality normals requires a uint2 UAV which is only supported on some graphic cards on PC 
-			// and this is an unknown when compiling shaders at this stage.
-			// !!! If this is changed, please update all sites reading r.Strata.DBufferPass !!!
-			const bool bDBufferPassSupported = StrataNormalQuality == 0 || (StrataNormalQuality > 0 && IsConsolePlatform(Platform));
-
-			static FShaderPlatformCachedIniValue<int32> CVarDBufferPass(TEXT("r.Strata.DBufferPass"));
-			const int32 StrataDBufferPass = CVarDBufferPass.Get(Platform);
-			if (StrataDBufferPass>0)
-			{
-				KeyString += FString::Printf(TEXT("_DBUFFERPASS"));
-			}
+			KeyString += FString::Printf(TEXT("_DBUFFERPASS"));
 		}
 
-		static const auto CVarBackCompatibility = IConsoleManager::Get().FindTConsoleVariableDataInt(TEXT("r.StrataBackCompatibility"));
-		if (bStrataEnabled && CVarBackCompatibility && CVarBackCompatibility->GetValueOnAnyThread()>0)
+		if (Strata::IsBackCompatibilityEnabled())
 		{
 			KeyString += FString::Printf(TEXT("_BACKCOMPAT"));
 		}
 
-		static const auto CVarOpaqueRoughRefrac = IConsoleManager::Get().FindTConsoleVariableDataInt(TEXT("r.Strata.OpaqueMaterialRoughRefraction"));
-		if (bStrataEnabled && CVarOpaqueRoughRefrac && CVarOpaqueRoughRefrac->GetValueOnAnyThread() > 0)
+		if (Strata::IsOpaqueRoughRefractionEnabled())
 		{
 			KeyString += FString::Printf(TEXT("_ROUGHDIFF"));
 		}
 
-		static const auto CVarStrataGBufferFormat = IConsoleManager::Get().FindTConsoleVariableDataInt(TEXT("r.GBufferFormat"));
-		const uint32 StrataNormalQuality = CVarStrataGBufferFormat && CVarStrataGBufferFormat->GetValueOnAnyThread() > 1 ? 1 : 0;
-		if (bStrataEnabled && StrataNormalQuality > 0)
+		if (Strata::GetNormalQuality() > 0)
 		{
 			KeyString += FString::Printf(TEXT("_STRTNRMQ"));
 		}
 
-		static const auto CVarAdvancedDebug = IConsoleManager::Get().FindTConsoleVariableDataInt(TEXT("r.Strata.Debug.AdvancedVisualizationShaders"));
-		if (bStrataEnabled && CVarAdvancedDebug && CVarAdvancedDebug->GetValueOnAnyThread() > 0)
+		if (Strata::IsAdvancedVisualizationEnabled())
 		{
 			KeyString += FString::Printf(TEXT("_ADVDEBUG"));
 		}
@@ -2191,14 +2164,15 @@ EShaderPermutationFlags GetShaderPermutationFlags(const FPlatformTypeLayoutParam
 	return Result;
 }
 
-void RegisterRayTracingPayloadType(ERayTracingPayloadType PayloadType, uint32 PayloadSize)
+void RegisterRayTracingPayloadType(ERayTracingPayloadType PayloadType, uint32 PayloadSize, TRaytracingPayloadSizeFunction PayloadSizeFunction)
 {
 	// Make sure we haven't registered this payload type yet
 	uint32 PayloadTypeInt = static_cast<uint32>(PayloadType);
 	checkf(FMath::CountBits(PayloadTypeInt) == 1, TEXT("PayloadType should have only 1 bit set -- got %u"), PayloadTypeInt);
 	checkf(!IsRayTracingPayloadRegistered(PayloadType), TEXT("Payload type %u has already been registered"), PayloadTypeInt);
 	int32 PayloadIndex = FPlatformMath::CountTrailingZeros(PayloadTypeInt);
-	RayTracingPayloadSizes[PayloadIndex] = PayloadSize;
+	RayTracingPayloadSizeFunctions[PayloadIndex] = PayloadSizeFunction;
+	RayTracingPayloadSizes[PayloadIndex] = PayloadSizeFunction ? 0u : PayloadSize;
 	RegisteredRayTracingPayloads |= PayloadTypeInt;
 }
 
@@ -2210,7 +2184,15 @@ uint32 GetRayTracingPayloadTypeMaxSize(ERayTracingPayloadType PayloadType)
 	for (uint32 PayloadTypeInt = static_cast<uint32>(PayloadType); PayloadTypeInt;)
 	{
 		const int32 PayloadIndex = FPlatformMath::CountTrailingZeros(PayloadTypeInt);
-		Result = FMath::Max(Result, RayTracingPayloadSizes[PayloadIndex]);
+		if (RayTracingPayloadSizeFunctions[PayloadIndex] != nullptr)
+		{
+			Result = FMath::Max(Result, RayTracingPayloadSizeFunctions[PayloadIndex]());
+		}
+		else
+		{
+			Result = FMath::Max(Result, RayTracingPayloadSizes[PayloadIndex]);
+		}
+
 		// remove bit we just processed
 		PayloadTypeInt &= ~(1u << PayloadIndex);
 	}
diff --git a/Engine/Source/Runtime/RenderCore/Public/RenderUtils.h b/Engine/Source/Runtime/RenderCore/Public/RenderUtils.h
--- a/Engine/Source/Runtime/RenderCore/Public/RenderUtils.h
+++ b/Engine/Source/Runtime/RenderCore/Public/RenderUtils.h
@@ -477,3 +477,25 @@ extern RENDERCORE_API bool IsRayTracingAllowed();
 // Returns the ray tracing mode if ray tracing is allowed.
 // This function may only be called at runtime, never during cooking.
 extern RENDERCORE_API ERayTracingMode GetRayTracingMode();
+
+namespace Strata
+{
+	RENDERCORE_API bool IsStrataEnabled();
+	RENDERCORE_API bool IsRoughDiffuseEnabled();
+	RENDERCORE_API bool IsBackCompatibilityEnabled();
+	RENDERCORE_API bool IsDBufferPassEnabled(EShaderPlatform InPlatform);
+	RENDERCORE_API bool IsOpaqueRoughRefractionEnabled();
+	RENDERCORE_API bool IsAdvancedVisualizationEnabled();
+	RENDERCORE_API bool Is8bitTileCoordEnabled();
+
+	RENDERCORE_API uint32 GetRayTracingMaterialPayloadSizeInBytes();
+	RENDERCORE_API uint32 GetRayTracingMaterialPayloadSizeInBytes(EShaderPlatform InPlatform);
+
+	RENDERCORE_API uint32 GetBytePerPixel();
+	RENDERCORE_API uint32 GetBytePerPixel(EShaderPlatform InPlatform);
+
+	RENDERCORE_API uint32 GetNormalQuality();
+
+	RENDERCORE_API uint32 GetShadingQuality();
+	RENDERCORE_API uint32 GetShadingQuality(EShaderPlatform InPlatform);
+}
\ No newline at end of file
diff --git a/Engine/Source/Runtime/RenderCore/Public/Shader.h b/Engine/Source/Runtime/RenderCore/Public/Shader.h
--- a/Engine/Source/Runtime/RenderCore/Public/Shader.h
+++ b/Engine/Source/Runtime/RenderCore/Public/Shader.h
@@ -307,22 +307,24 @@ public:
 };
 
 enum class ERayTracingPayloadType : uint32; // actual enum is defined in /Shaders/Shared/RayTracingPayloadType.h
+typedef uint32(*TRaytracingPayloadSizeFunction)();
 ENUM_CLASS_FLAGS(ERayTracingPayloadType);
 
 RENDERCORE_API uint32 GetRayTracingPayloadTypeMaxSize(ERayTracingPayloadType PayloadType);
-RENDERCORE_API void RegisterRayTracingPayloadType(ERayTracingPayloadType PayloadType, uint32 PayloadSize);
+RENDERCORE_API void RegisterRayTracingPayloadType(ERayTracingPayloadType PayloadType, uint32 PayloadSize, TRaytracingPayloadSizeFunction PayloadSizeFunction);
 
 struct FRegisterRayTracingPayloadTypeHelper
 {
-	FRegisterRayTracingPayloadTypeHelper(ERayTracingPayloadType PayloadType, uint32 PayloadSize)
+	FRegisterRayTracingPayloadTypeHelper(ERayTracingPayloadType PayloadType, uint32 PayloadSize, TRaytracingPayloadSizeFunction PayloadSizeFunction)
 	{
-		RegisterRayTracingPayloadType(PayloadType, PayloadSize);
+		RegisterRayTracingPayloadType(PayloadType, PayloadSize, PayloadSizeFunction);
 	}
 };
 
 #define IMPLEMENT_RT_PAYLOAD_TYPE_CONCAT2( x, y ) x##y
 #define IMPLEMENT_RT_PAYLOAD_TYPE_CONCAT( x, y ) IMPLEMENT_RT_PAYLOAD_TYPE_CONCAT2( x, y )
-#define IMPLEMENT_RT_PAYLOAD_TYPE(PayloadType, PayloadSize)  static FRegisterRayTracingPayloadTypeHelper IMPLEMENT_RT_PAYLOAD_TYPE_CONCAT(PayloadHelper, __COUNTER__) = FRegisterRayTracingPayloadTypeHelper(PayloadType, PayloadSize);
+#define IMPLEMENT_RT_PAYLOAD_TYPE(PayloadType, PayloadSize)  static FRegisterRayTracingPayloadTypeHelper IMPLEMENT_RT_PAYLOAD_TYPE_CONCAT(PayloadHelper, __COUNTER__) = FRegisterRayTracingPayloadTypeHelper(PayloadType, PayloadSize, nullptr);
+#define IMPLEMENT_RT_PAYLOAD_TYPE_FUNCTION(PayloadType, PayloadSizeFunction)  static FRegisterRayTracingPayloadTypeHelper IMPLEMENT_RT_PAYLOAD_TYPE_CONCAT(PayloadHelper, __COUNTER__) = FRegisterRayTracingPayloadTypeHelper(PayloadType, 0u, PayloadSizeFunction);
 
 class RENDERCORE_API FShaderMapResource : public FRenderResource, public FDeferredCleanupInterface
 {
diff --git a/Engine/Source/Runtime/Renderer/Private/BasePassRendering.cpp b/Engine/Source/Runtime/Renderer/Private/BasePassRendering.cpp
--- a/Engine/Source/Runtime/Renderer/Private/BasePassRendering.cpp
+++ b/Engine/Source/Runtime/Renderer/Private/BasePassRendering.cpp
@@ -427,7 +427,7 @@ void SetDepthStencilStateForBasePass_Internal(FMeshPassProcessorRenderState& InD
 template<bool bDepthTest, ECompareFunction CompareFunction>
 void SetDepthStencilStateForBasePass_Internal(FMeshPassProcessorRenderState& InDrawRenderState, ERHIFeatureLevel::Type FeatureLevel)
 {	
-	const static bool bStrataDufferPassEnabled = Strata::IsStrataEnabled() && Strata::IsStrataDbufferPassEnabled(GShaderPlatformForFeatureLevel[FeatureLevel]);
+	const static bool bStrataDufferPassEnabled = Strata::IsStrataEnabled() && Strata::IsDBufferPassEnabled(GShaderPlatformForFeatureLevel[FeatureLevel]);
 	if (bStrataDufferPassEnabled)
 	{
 		SetDepthStencilStateForBasePass_Internal<bDepthTest, CompareFunction, GET_STENCIL_BIT_MASK(STRATA_RECEIVE_DBUFFER_NORMAL, 1) | GET_STENCIL_BIT_MASK(STRATA_RECEIVE_DBUFFER_DIFFUSE, 1) | GET_STENCIL_BIT_MASK(STRATA_RECEIVE_DBUFFER_ROUGHNESS, 1) | GET_STENCIL_BIT_MASK(DISTANCE_FIELD_REPRESENTATION, 1) | STENCIL_LIGHTING_CHANNELS_MASK(0x7)>(InDrawRenderState);
@@ -1759,7 +1759,7 @@ bool FBasePassMeshProcessor::Process(
 	
 	if (bEnableReceiveDecalOutput)
 	{
-		static const bool bStrataDufferPassEnabled = Strata::IsStrataEnabled() && Strata::IsStrataDbufferPassEnabled(GShaderPlatformForFeatureLevel[FeatureLevel]);
+		static const bool bStrataDufferPassEnabled = Strata::IsStrataEnabled() && Strata::IsDBufferPassEnabled(GShaderPlatformForFeatureLevel[FeatureLevel]);
 
 		// Set stencil value for this draw call
 		// This is effectively extending the GBuffer using the stencil bits
diff --git a/Engine/Source/Runtime/Renderer/Private/ClusteredDeferredShadingPass.cpp b/Engine/Source/Runtime/Renderer/Private/ClusteredDeferredShadingPass.cpp
--- a/Engine/Source/Runtime/Renderer/Private/ClusteredDeferredShadingPass.cpp
+++ b/Engine/Source/Runtime/Renderer/Private/ClusteredDeferredShadingPass.cpp
@@ -176,7 +176,7 @@ static void InternalAddClusteredDeferredShadingPass(
 	{
 		PassParameters->RenderTargets[0] = FRenderTargetBinding(View.HairStrandsViewData.VisibilityData.SampleLightingTexture, ERenderTargetLoadAction::ELoad);
 	}
-	if (Strata::IsStrataOpaqueMaterialRoughRefractionEnabled())
+	if (Strata::IsOpaqueRoughRefractionEnabled())
 	{
 		check(StrataSceneData);
 		PassParameters->RenderTargets[1] = FRenderTargetBinding(StrataSceneData->SeparatedOpaqueRoughRefractionSceneColor, ERenderTargetLoadAction::ELoad);
@@ -212,7 +212,7 @@ static void InternalAddClusteredDeferredShadingPass(
 			InRHICmdList.ApplyCachedRenderTargets(GraphicsPSOInit);
 
 			// Additive blend to accumulate lighting contributions.
-			if (Strata::IsStrataOpaqueMaterialRoughRefractionEnabled())
+			if (Strata::IsOpaqueRoughRefractionEnabled())
 			{
 				GraphicsPSOInit.BlendState = TStaticBlendState<
 					CW_RGBA, BO_Add, BF_One, BF_One, BO_Add, BF_One, BF_One,
diff --git a/Engine/Source/Runtime/Renderer/Private/IndirectLightRendering.cpp b/Engine/Source/Runtime/Renderer/Private/IndirectLightRendering.cpp
--- a/Engine/Source/Runtime/Renderer/Private/IndirectLightRendering.cpp
+++ b/Engine/Source/Runtime/Renderer/Private/IndirectLightRendering.cpp
@@ -1352,7 +1352,7 @@ void FDeferredShadingSceneRenderer::RenderDiffuseIndirectAndAmbientOcclusion(
 				check(Strata::IsStrataEnabled());
 
 				// Rough refraction targets
-				if (Strata::IsStrataOpaqueMaterialRoughRefractionEnabled())
+				if (Strata::IsOpaqueRoughRefractionEnabled())
 				{
 					PassParameters->OutOpaqueRoughRefractionSceneColor = GraphBuilder.CreateUAV(Scene->StrataSceneData.SeparatedOpaqueRoughRefractionSceneColor);
 					PassParameters->OutSubSurfaceSceneColor = GraphBuilder.CreateUAV(Scene->StrataSceneData.SeparatedSubSurfaceSceneColor);
@@ -1438,7 +1438,7 @@ void FDeferredShadingSceneRenderer::RenderDiffuseIndirectAndAmbientOcclusion(
 				PassParameters->ViewUniformBuffer = View.ViewUniformBuffer;
 
 				PassParameters->RenderTargets[0] = FRenderTargetBinding(SceneColorTexture, ERenderTargetLoadAction::ELoad);
-				if (Strata::IsStrataOpaqueMaterialRoughRefractionEnabled())
+				if (Strata::IsOpaqueRoughRefractionEnabled())
 				{
 					PassParameters->RenderTargets[1] = FRenderTargetBinding(Scene->StrataSceneData.SeparatedOpaqueRoughRefractionSceneColor, ERenderTargetLoadAction::ELoad);
 					PassParameters->RenderTargets[2] = FRenderTargetBinding(Scene->StrataSceneData.SeparatedSubSurfaceSceneColor, ERenderTargetLoadAction::ELoad);
@@ -1618,7 +1618,7 @@ static void AddSkyReflectionPass(
 	}
 
 	PassParameters->RenderTargets[0] = FRenderTargetBinding(SceneColorTexture.Target, ERenderTargetLoadAction::ELoad);
-	if (Strata::IsStrataOpaqueMaterialRoughRefractionEnabled())
+	if (Strata::IsOpaqueRoughRefractionEnabled())
 	{
 		PassParameters->RenderTargets[1] = FRenderTargetBinding(Scene->StrataSceneData.SeparatedOpaqueRoughRefractionSceneColor, ERenderTargetLoadAction::ELoad);
 		PassParameters->RenderTargets[2] = FRenderTargetBinding(Scene->StrataSceneData.SeparatedSubSurfaceSceneColor, ERenderTargetLoadAction::ELoad);
@@ -1684,7 +1684,7 @@ static void AddSkyReflectionPass(
 			}
 			else
 			{
-				if (Strata::IsStrataOpaqueMaterialRoughRefractionEnabled())
+				if (Strata::IsOpaqueRoughRefractionEnabled())
 				{
 					GraphicsPSOInit.BlendState = TStaticBlendState<
 						CW_RGBA, BO_Add, BF_One, BF_One, BO_Add, BF_One, BF_One,
diff --git a/Engine/Source/Runtime/Renderer/Private/LightRendering.cpp b/Engine/Source/Runtime/Renderer/Private/LightRendering.cpp
--- a/Engine/Source/Runtime/Renderer/Private/LightRendering.cpp
+++ b/Engine/Source/Runtime/Renderer/Private/LightRendering.cpp
@@ -2197,7 +2197,7 @@ static FDeferredLightPS::FParameters GetDeferredLightPSParameters(
 	Out.ShadowMaskBits = ShadowMaskBits ? ShadowMaskBits : GSystemTextures.GetZeroUIntDummy(GraphBuilder);
 	// PS - Render Targets
 	Out.RenderTargets[0] = FRenderTargetBinding(SceneColorTexture, ERenderTargetLoadAction::ELoad);
-	if (Strata::IsStrataOpaqueMaterialRoughRefractionEnabled())
+	if (Strata::IsOpaqueRoughRefractionEnabled())
 	{
 		Out.RenderTargets[1] = FRenderTargetBinding(Scene->StrataSceneData.SeparatedOpaqueRoughRefractionSceneColor, ERenderTargetLoadAction::ELoad);
 		Out.RenderTargets[2] = FRenderTargetBinding(Scene->StrataSceneData.SeparatedSubSurfaceSceneColor, ERenderTargetLoadAction::ELoad);
@@ -2243,7 +2243,7 @@ static void InternalRenderLight(
 		// Set the device viewport for the view.
 		RHICmdList.SetViewport(View.ViewRect.Min.X, View.ViewRect.Min.Y, 0.0f, View.ViewRect.Max.X, View.ViewRect.Max.Y, 1.0f);
 
-		if (Strata::IsStrataOpaqueMaterialRoughRefractionEnabled())
+		if (Strata::IsOpaqueRoughRefractionEnabled())
 		{
 			GraphicsPSOInit.BlendState = TStaticBlendState<
 				CW_RGBA, BO_Add, BF_One, BF_One, BO_Add, BF_One, BF_One,
diff --git a/Engine/Source/Runtime/Renderer/Private/PathTracing.cpp b/Engine/Source/Runtime/Renderer/Private/PathTracing.cpp
--- a/Engine/Source/Runtime/Renderer/Private/PathTracing.cpp
+++ b/Engine/Source/Runtime/Renderer/Private/PathTracing.cpp
@@ -736,7 +736,7 @@ static FPathTracingFogParameters PrepareFogParameters(const FViewInfo& View, con
 	return Parameters;
 }
 
-IMPLEMENT_RT_PAYLOAD_TYPE(ERayTracingPayloadType::PathTracingMaterial, 64); // TODO: Expand this when strata is enabled
+IMPLEMENT_RT_PAYLOAD_TYPE(ERayTracingPayloadType::PathTracingMaterial, 64); // TODO: Expand this when strata is enabled (use GetStrataRaytracingMaterialPayload())
 
 class FPathTracingRG : public FGlobalShader
 {
diff --git a/Engine/Source/Runtime/Renderer/Private/PostProcess/PostProcessSubsurface.cpp b/Engine/Source/Runtime/Renderer/Private/PostProcess/PostProcessSubsurface.cpp
--- a/Engine/Source/Runtime/Renderer/Private/PostProcess/PostProcessSubsurface.cpp
+++ b/Engine/Source/Runtime/Renderer/Private/PostProcess/PostProcessSubsurface.cpp
@@ -326,7 +326,7 @@ uint32 GetSubsurfaceRequiredViewMask(TArrayView<const FViewInfo> Views)
 
 bool IsSubsurfaceCheckerboardFormat(EPixelFormat SceneColorFormat)
 {
-	if (Strata::IsStrataOpaqueMaterialRoughRefractionEnabled())
+	if (Strata::IsOpaqueRoughRefractionEnabled())
 	{
 		// With this mode, specular and subsurface colors are correctly separated so checkboard is not required.
 		return false;
@@ -1016,8 +1016,8 @@ void AddSubsurfaceViewPass(
 	const FScreenPassTextureViewportParameters SubsurfaceViewportParameters = GetScreenPassTextureViewportParameters(SubsurfaceViewport);
 	const FScreenPassTextureViewportParameters SceneViewportParameters = GetScreenPassTextureViewportParameters(SceneViewport);
 
-	const bool bReadSeparatedSubSurfaceSceneColor = Strata::IsStrataOpaqueMaterialRoughRefractionEnabled();
-	const bool bWriteSeparatedOpaqueRoughRefractionSceneColor = Strata::IsStrataOpaqueMaterialRoughRefractionEnabled();
+	const bool bReadSeparatedSubSurfaceSceneColor = Strata::IsOpaqueRoughRefractionEnabled();
+	const bool bWriteSeparatedOpaqueRoughRefractionSceneColor = Strata::IsOpaqueRoughRefractionEnabled();
 	FRDGTextureRef SeparatedSubSurfaceSceneColor = View.StrataViewData.SceneData->SeparatedSubSurfaceSceneColor;
 	FRDGTextureRef SeparatedOpaqueRoughRefractionSceneColor = View.StrataViewData.SceneData->SeparatedOpaqueRoughRefractionSceneColor;
 
diff --git a/Engine/Source/Runtime/Renderer/Private/Strata/Strata.cpp b/Engine/Source/Runtime/Renderer/Private/Strata/Strata.cpp
--- a/Engine/Source/Runtime/Renderer/Private/Strata/Strata.cpp
+++ b/Engine/Source/Runtime/Renderer/Private/Strata/Strata.cpp
@@ -18,54 +18,9 @@
 #include "SystemTextures.h"
 #include "DBufferTextures.h"
 
-//PRAGMA_DISABLE_OPTIMIZATION
+//UE_DISABLE_OPTIMIZATION
 
 // The project setting for Strata
-static TAutoConsoleVariable<int32> CVarStrata(
-	TEXT("r.Strata"),
-	0,
-	TEXT("Enable Strata materials (Beta)."),
-	ECVF_ReadOnly | ECVF_RenderThreadSafe);
-
-static TAutoConsoleVariable<int32> CVarStrataBackCompatibility(
-	TEXT("r.StrataBackCompatibility"),
-	0,
-	TEXT("Disables Strata multiple scattering and replaces Chan diffuse by Lambert."),
-	ECVF_ReadOnly | ECVF_RenderThreadSafe);
-
-static TAutoConsoleVariable<int32> CVarStrataBytePerPixel(
-	TEXT("r.Strata.BytesPerPixel"),
-	80,
-	TEXT("Strata allocated byte per pixel to store materials data. Higher value means more complex material can be represented."),
-	ECVF_ReadOnly | ECVF_RenderThreadSafe);
-
-static TAutoConsoleVariable<int32> CVarStrataShadingQuality(
-	TEXT("r.Strata.ShadingQuality"),
-	1,
-	TEXT("Define Strata shading quality (1: accurate lighting, 2: approximate lighting). This variable is read-only."),
-	ECVF_ReadOnly | ECVF_RenderThreadSafe);
-
-static TAutoConsoleVariable<int32> CVarStrataTileCoord8Bits(
-	TEXT("r.Strata.TileCoord8bits"),
-	0,
-	TEXT("Format of tile coord. This variable is read-only."),
-	ECVF_ReadOnly | ECVF_RenderThreadSafe);
-
-static TAutoConsoleVariable<int32> CVarStrataRoughDiffuse(
-	TEXT("r.Strata.RoughDiffuse"),
-	1,
-	TEXT("Enable Strata rough diffuse model (works only if r.Material.RoughDiffuse is enabled in the project settings). Togglable at runtime"),
-	ECVF_RenderThreadSafe);
-
-// Transition render settings that will disapear when strata gets enabled
-
-static TAutoConsoleVariable<int32> CVarMaterialRoughDiffuse(
-	TEXT("r.Material.RoughDiffuse"),
-	0,
-	TEXT("Enable rough diffuse material."),
-	ECVF_ReadOnly | ECVF_RenderThreadSafe);
-
-// STRATA_TODO we keep this for now and can remove it once battletested.
 static TAutoConsoleVariable<int32> CVarUseCmaskClear(
 	TEXT("r.Strata.UseCmaskClear"),
 	0,
@@ -96,12 +51,6 @@ static TAutoConsoleVariable<int32> CVarStrataAsyncClassification(
 	TEXT("Run Strata material classification in async (with shadow)."),
 	ECVF_RenderThreadSafe);
 
-static TAutoConsoleVariable<int32> CVarStrataDBufferPass(
-	TEXT("r.Strata.DBufferPass"),
-	0,
-	TEXT("Apply DBuffer after the base-pass as a separate pass. Read only because when this is changed, it will require the recompilation of all shaders."),
-	ECVF_ReadOnly | ECVF_RenderThreadSafe);
-
 static TAutoConsoleVariable<int32> CVarStrataDBufferPassDedicatedTiles(
 	TEXT("r.Strata.DBufferPass.DedicatedTiles"),
 	0,
@@ -110,8 +59,6 @@ static TAutoConsoleVariable<int32> CVarStrataDBufferPassDedicatedTiles(
 
 IMPLEMENT_GLOBAL_SHADER_PARAMETER_STRUCT(FStrataGlobalUniformParameters, "Strata");
 
-extern bool IsStrataAdvancedVisualizationShadersEnabled();
-
 void FStrataViewData::Reset()
 {
 	// Preserve old max BSDF count which is set prior to the reset operation
@@ -147,16 +94,6 @@ const TCHAR* ToString(EStrataTileType Type)
 namespace Strata
 {
 
-bool IsStrataEnabled()
-{
-	return CVarStrata.GetValueOnAnyThread() > 0;
-}
-
-bool IsStrataDbufferPassEnabled(const EShaderPlatform InPlatform)
-{
-	return IsStrataEnabled() && IsUsingDBuffers(InPlatform) && CVarStrataDBufferPass.GetValueOnAnyThread() > 0;
-}
-
 enum EStrataTileSpace
 {
 	StrataTileSpace_Primary = 1u,
@@ -216,11 +153,6 @@ bool SupportsCMask(const FStaticShaderPlatform InPlatform)
 	return CVarUseCmaskClear.GetValueOnRenderThread() > 0 && FDataDrivenShaderPlatformInfo::GetSupportsRenderTargetWriteMask(InPlatform);
 }
 
-bool IsClassificationCoord8bits()
-{
-	return CVarStrataTileCoord8Bits.GetValueOnRenderThread() == 1;
-}
-
 bool IsClassificationAsync()
 {
 	return CVarStrataAsyncClassification.GetValueOnRenderThread() > 0;
@@ -229,7 +161,7 @@ bool IsClassificationAsync()
 static EPixelFormat GetClassificationTileFormat(const FIntPoint& InResolution)
 {
 	// For platform which whose resolution is never above 1080p, use 8bit tile format for performance.
-	const bool bRequest8bit = IsClassificationCoord8bits();
+	const bool bRequest8bit = Is8bitTileCoordEnabled();
 	if (bRequest8bit)
 	{
 		check(InResolution.X <= 2048 && InResolution.Y <= 2048);
@@ -276,8 +208,8 @@ static void InitialiseStrataViewData(FRDGBuilder& GraphBuilder, FViewInfo& View,
 			AddClearUAVPass(GraphBuilder, Out.ClassificationTileDispatchIndirectBufferUAV, 0);
 
 			// Separated subsurface & rough refraction textures (tile data)
-			const uint32 RoughTileCount = IsStrataOpaqueMaterialRoughRefractionEnabled() ? TileResolution.X * TileResolution.Y : 4;
-			const uint32 DecalTileCount = IsStrataDbufferPassEnabled(View.GetShaderPlatform()) ? TileResolution.X * TileResolution.Y : 4;
+			const uint32 RoughTileCount = IsOpaqueRoughRefractionEnabled() ? TileResolution.X * TileResolution.Y : 4;
+			const uint32 DecalTileCount = IsDBufferPassEnabled(View.GetShaderPlatform()) ? TileResolution.X * TileResolution.Y : 4;
 			const uint32 RegularTileCount = TileResolution.X * TileResolution.Y;
 
 			const EPixelFormat ClassificationTileFormat = GetClassificationTileFormat(ViewResolution);
@@ -357,21 +289,13 @@ static bool NeedBSDFOffsets(const FScene* Scene, const FViewInfo& View)
 	return  ShouldRenderLumenDiffuseGI(Scene, View) || ShouldRenderLumenReflections(View) || Strata::ShouldRenderStrataDebugPasses(View);
 }
 
-static uint32 StrataGetBytePerPixel(EShaderPlatform InPlatform)
-{
-	// We enforce at least 20 bytes per pixel because this is the minimal Strata GBuffer footprint of the simplest material.
-	const uint32 MinStrataBytePerPixel = 20u;
-	const uint32 MaxStrataBytePerPixel = IsMobilePlatform(InPlatform) ? 24u : 256u;
-	return FMath::Clamp(uint32(CVarStrataBytePerPixel.GetValueOnAnyThread()), MinStrataBytePerPixel, MaxStrataBytePerPixel);
-}
-
 static void RecordStrataAnalytics(EShaderPlatform InPlatform)
 {
 	if (FEngineAnalytics::IsAvailable())
 	{
 		TArray<FAnalyticsEventAttribute> EventAttributes;
 		EventAttributes.Add(FAnalyticsEventAttribute(TEXT("Enabled"), 1));
-		EventAttributes.Add(FAnalyticsEventAttribute(TEXT("BytesPerPixel"), StrataGetBytePerPixel(InPlatform)));
+		EventAttributes.Add(FAnalyticsEventAttribute(TEXT("BytesPerPixel"), GetBytePerPixel(InPlatform)));
 
 		FString OutStr(TEXT("Strata"));
 		FEngineAnalytics::GetProvider().RecordEvent(OutStr, EventAttributes);
@@ -380,9 +304,7 @@ static void RecordStrataAnalytics(EShaderPlatform InPlatform)
 
 static EPixelFormat GetTopLayerTextureFormat()
 {
-	static const auto CVarStrataGBufferFormat = IConsoleManager::Get().FindTConsoleVariableDataInt(TEXT("r.GBufferFormat")); 
-
-	const bool bStrataHighQualityNormal = CVarStrataGBufferFormat && CVarStrataGBufferFormat->GetValueOnAnyThread() > 1;
+	const bool bStrataHighQualityNormal = GetNormalQuality() > 1;
 
 	// High quality normal is not supported on platforms that do not support R32G32 UAV load.
 	// This is dues to the way Strata account for decals. See FStrataDBufferPassCS, updating TopLayerTexture this way.
@@ -405,9 +327,11 @@ void InitialiseStrataFrameSceneData(FRDGBuilder& GraphBuilder, FSceneRenderer& S
 	};
 
 	bool bNeedBSDFOffsets = false;
+	bool bNeedUAV = false;
 	for (const FViewInfo& View : SceneRenderer.Views)
 	{
 		bNeedBSDFOffsets = bNeedBSDFOffsets || NeedBSDFOffsets(SceneRenderer.Scene, View);
+		bNeedUAV = bNeedUAV || IsDBufferPassEnabled(View.GetShaderPlatform());
 	}
 
 	FIntPoint MaterialBufferSizeXY;
@@ -428,18 +352,17 @@ void InitialiseStrataFrameSceneData(FRDGBuilder& GraphBuilder, FSceneRenderer& S
 		UpdateMaterialBufferToTiledResolution(SceneTextureExtent, MaterialBufferSizeXY);
 
 		const uint32 RoundToValue = 4u;
-		Out.MaxBytesPerPixel = FMath::DivideAndRoundUp(StrataGetBytePerPixel(SceneRenderer.ShaderPlatform), RoundToValue) * RoundToValue;
+		Out.MaxBytesPerPixel = FMath::DivideAndRoundUp(GetBytePerPixel(SceneRenderer.ShaderPlatform), RoundToValue) * RoundToValue;
 
 		// Top layer texture
 		{
-			const bool bNeedUAV = CVarStrataDBufferPass.GetValueOnRenderThread() > 0;
 			Out.TopLayerTexture = GraphBuilder.CreateTexture(FRDGTextureDesc::Create2D(SceneTextureExtent, GetTopLayerTextureFormat(), FClearValueBinding::Black, TexCreate_RenderTargetable | TexCreate_ShaderResource | TexCreate_FastVRAM | (bNeedUAV ? TexCreate_UAV : TexCreate_None)), TEXT("Strata.TopLayerTexture"));
 		}
 
 		// Separated subsurface and rough refraction textures
 		{
-			const bool bIsStrataOpaqueMaterialRoughRefractionEnabled = IsStrataOpaqueMaterialRoughRefractionEnabled();
-			const FIntPoint OpaqueRoughRefractionSceneExtent		= bIsStrataOpaqueMaterialRoughRefractionEnabled ? SceneTextureExtent : FIntPoint(4, 4);
+			const bool bIsStrataOpaqueMaterialRoughRefractionEnabled = IsOpaqueRoughRefractionEnabled();
+			const FIntPoint OpaqueRoughRefractionSceneExtent		 = bIsStrataOpaqueMaterialRoughRefractionEnabled ? SceneTextureExtent : FIntPoint(4, 4);
 			
 			Out.OpaqueRoughRefractionTexture = GraphBuilder.CreateTexture(
 				FRDGTextureDesc::Create2D(OpaqueRoughRefractionSceneExtent, PF_FloatR11G11B10, FClearValueBinding::Black, TexCreate_ShaderResource | TexCreate_UAV | TexCreate_RenderTargetable), TEXT("Strata.OpaqueRoughRefractionTexture"));
@@ -475,7 +398,7 @@ void InitialiseStrataFrameSceneData(FRDGBuilder& GraphBuilder, FSceneRenderer& S
 	FIntPoint SceneTextureExtent = IsStrataEnabled() ? SceneRenderer.GetActiveSceneTexturesConfig().Extent : FIntPoint(2, 2);
 
 	const uint32 SliceCountSSS = STRATA_SSS_DATA_UINT_COUNT;
-	const uint32 SliceCountAdvDebug = IsStrataAdvancedVisualizationShadersEnabled() ? 1 : 0;
+	const uint32 SliceCountAdvDebug = IsAdvancedVisualizationEnabled() ? 1 : 0;
 	const uint32 SliceCount = FMath::DivideAndRoundUp(Out.MaxBytesPerPixel, 4u) + SliceCountSSS + SliceCountAdvDebug;
 	FRDGTextureDesc MaterialTextureDesc = FRDGTextureDesc::Create2DArray(SceneTextureExtent, PF_R32_UINT, FClearValueBinding::Transparent, TexCreate_TargetArraySlicesIndependently | TexCreate_DisableDCC | TexCreate_NoFastClear | TexCreate_RenderTargetable | TexCreate_ShaderResource | TexCreate_UAV | TexCreate_FastVRAM, SliceCount, 1, 1);
 	MaterialTextureDesc.FastVRAMPercentage = (1.0f / SliceCount) * 0xFF; // Only allocate the first slice into ESRAM
@@ -489,7 +412,7 @@ void InitialiseStrataFrameSceneData(FRDGBuilder& GraphBuilder, FSceneRenderer& S
 	Out.MaterialTextureArrayUAVWithoutRTs = GraphBuilder.CreateUAV(FRDGTextureUAVDesc(Out.MaterialTextureArray, 0, PF_Unknown, STRATA_BASE_PASS_MRT_OUTPUT_COUNT, SliceCount - STRATA_BASE_PASS_MRT_OUTPUT_COUNT));
 
 	// Rough diffuse model
-	Out.bRoughDiffuse = CVarStrataRoughDiffuse.GetValueOnRenderThread() > 0 ? 1u : 0u;
+	Out.bRoughDiffuse = IsRoughDiffuseEnabled();
 
 	Out.PeelLayersAboveDepth = FMath::Max(CVarStrataDebugPeelLayersAboveDepth.GetValueOnRenderThread(), 0);
 
@@ -1203,8 +1126,6 @@ void SetBasePassRenderTargetOutputFormat(const EShaderPlatform Platform, const F
 
 ////////////////////////////////////////////////////////////////////////////////////////////////////////////////////////////////////////////////
 
-//#include "RendererUtils.h"
-
 void AddStrataMaterialClassificationPass(FRDGBuilder& GraphBuilder, const FMinimalSceneTextures& SceneTextures, const FDBufferTextures& DBufferTextures, const TArray<FViewInfo>& Views)
 {
 	RDG_EVENT_SCOPE_CONDITIONAL(GraphBuilder, IsStrataEnabled() && Views.Num() > 0, "Strata::MaterialClassification");
@@ -1240,7 +1161,7 @@ void AddStrataMaterialClassificationPass(FRDGBuilder& GraphBuilder, const FMinim
 			}
 
 			// If Dbuffer pass (i.e. apply DBuffer data after the base-pass) is enabled, run special classification for outputing tile with/without tiles
-			const bool bDBufferTiles = CVarStrataDBufferPass.GetValueOnRenderThread() > 0 && CVarStrataDBufferPassDedicatedTiles.GetValueOnRenderThread() > 0 && DBufferTextures.IsValid() && IsConsolePlatform(View.GetShaderPlatform());
+			const bool bDBufferTiles = IsDBufferPassEnabled(View.GetShaderPlatform()) && CVarStrataDBufferPassDedicatedTiles.GetValueOnRenderThread() > 0 && DBufferTextures.IsValid() && IsConsolePlatform(View.GetShaderPlatform());
 
 			FStrataMaterialTileClassificationPassCS::FPermutationDomain PermutationVector;
 			PermutationVector.Set< FStrataMaterialTileClassificationPassCS::FCmask >(bSupportCMask);
@@ -1363,7 +1284,7 @@ void AddStrataMaterialClassificationPass(FRDGBuilder& GraphBuilder, const FMinim
 void AddStrataDBufferPass(FRDGBuilder& GraphBuilder, const FMinimalSceneTextures& SceneTextures, const FDBufferTextures& DBufferTextures, const TArray<FViewInfo>& Views)
 {
 	RDG_EVENT_SCOPE_CONDITIONAL(GraphBuilder, IsStrataEnabled() && Views.Num() > 0, "Strata::DBuffer");
-	if (!IsStrataEnabled() || !DBufferTextures.IsValid() || CVarStrataDBufferPass.GetValueOnRenderThread() <= 0)
+	if (!IsStrataEnabled() || !DBufferTextures.IsValid())
 	{
 		return;
 	}
@@ -1373,7 +1294,7 @@ void AddStrataDBufferPass(FRDGBuilder& GraphBuilder, const FMinimalSceneTextures
 		RDG_EVENT_SCOPE_CONDITIONAL(GraphBuilder, Views.Num() > 1, "View%d", i);
 
 		const FViewInfo& View = Views[i];
-		if (!IsUsingDBuffers(View.GetShaderPlatform()) || View.Family->EngineShowFlags.Decals == 0)
+		if (!IsUsingDBuffers(View.GetShaderPlatform()) || View.Family->EngineShowFlags.Decals == 0 || !IsDBufferPassEnabled(View.GetShaderPlatform()))
 		{
 			continue;
 		}
diff --git a/Engine/Source/Runtime/Renderer/Private/Strata/Strata.h b/Engine/Source/Runtime/Renderer/Private/Strata/Strata.h
--- a/Engine/Source/Runtime/Renderer/Private/Strata/Strata.h
+++ b/Engine/Source/Runtime/Renderer/Private/Strata/Strata.h
@@ -159,9 +159,6 @@ constexpr uint32 StencilBit_Fast   = 0x08; // In sync with SceneRenderTargets.h
 constexpr uint32 StencilBit_Single = 0x10; // In sync with SceneRenderTargets.h - GET_STENCIL_BIT_MASK(STENCIL_STRATA_SINGLEPATH)
 constexpr uint32 StencilBit_Complex= 0x20; // In sync with SceneRenderTargets.h - GET_STENCIL_BIT_MASK(STENCIL_STRATA_COMPLEX)
 
-bool IsStrataEnabled();
-bool IsStrataDbufferPassEnabled(const EShaderPlatform Platform);
-
 FIntPoint GetStrataTextureResolution(const FViewInfo& View, const FIntPoint& InResolution);
 
 void InitialiseStrataFrameSceneData(FRDGBuilder& GraphBuilder, FSceneRenderer& SceneRenderer);
@@ -180,7 +177,6 @@ void AddStrataDBufferPass(FRDGBuilder& GraphBuilder, const FMinimalSceneTextures
 
 void AddStrataStencilPass(FRDGBuilder& GraphBuilder, const TArray<FViewInfo>& Views, const FMinimalSceneTextures& SceneTextures);
 
-bool IsStrataOpaqueMaterialRoughRefractionEnabled();
 void AddStrataOpaqueRoughRefractionPasses(
 	FRDGBuilder& GraphBuilder,
 	FSceneTextures& SceneTextures,
diff --git a/Engine/Source/Runtime/Renderer/Private/Strata/StrataRoughRefraction.cpp b/Engine/Source/Runtime/Renderer/Private/Strata/StrataRoughRefraction.cpp
--- a/Engine/Source/Runtime/Renderer/Private/Strata/StrataRoughRefraction.cpp
+++ b/Engine/Source/Runtime/Renderer/Private/Strata/StrataRoughRefraction.cpp
@@ -14,27 +14,15 @@
 #include "ShaderCompiler.h"
 
 
-static TAutoConsoleVariable<int32> CVarStrataOpaqueMaterialRoughRefraction(
-	TEXT("r.Strata.OpaqueMaterialRoughRefraction"),
-	0,
-	TEXT("Enable Strata opaque material rough refractions effect from top layers over layers below."),
-	ECVF_ReadOnly | ECVF_RenderThreadSafe);
-
 static TAutoConsoleVariable<float> CVarStrataOpaqueMaterialRoughRefractionBlurScale(
 	TEXT("r.Strata.OpaqueMaterialRoughRefraction.BlurScale"),
 	1.0f,
 	TEXT("Scale opaque rough refraction strengh for debug purposes."),
 	ECVF_RenderThreadSafe);
 
-
 namespace Strata
 {
 
-bool IsStrataOpaqueMaterialRoughRefractionEnabled()
-{
-	return IsStrataEnabled() && CVarStrataOpaqueMaterialRoughRefraction.GetValueOnAnyThread() > 0;
-}
-
 class FOpaqueRoughRefractionPS : public FGlobalShader
 {
 	DECLARE_GLOBAL_SHADER(FOpaqueRoughRefractionPS);
@@ -77,7 +65,7 @@ void AddStrataOpaqueRoughRefractionPasses(
 	TArrayView<const FViewInfo> Views)
 {
 	const uint32 ViewCount = Views.Num();
-	const bool bOpaqueRoughRefractionEnabled = IsStrataOpaqueMaterialRoughRefractionEnabled() && ViewCount > 0;
+	const bool bOpaqueRoughRefractionEnabled = IsOpaqueRoughRefractionEnabled() && ViewCount > 0;
 	RDG_EVENT_SCOPE_CONDITIONAL(GraphBuilder, bOpaqueRoughRefractionEnabled, "Strata::OpaqueRoughRefraction");
 	if (!bOpaqueRoughRefractionEnabled)
 	{
diff --git a/Engine/Source/Runtime/Renderer/Private/Strata/StrataVisualize.cpp b/Engine/Source/Runtime/Renderer/Private/Strata/StrataVisualize.cpp
--- a/Engine/Source/Runtime/Renderer/Private/Strata/StrataVisualize.cpp
+++ b/Engine/Source/Runtime/Renderer/Private/Strata/StrataVisualize.cpp
@@ -11,18 +11,6 @@
 #include "IndirectLightRendering.h"
 #include "StrataVisualizationData.h"
 
-
-static TAutoConsoleVariable<int32> CVarStrataDebugAdvancedVisualizationShaders(
-	TEXT("r.Strata.Debug.AdvancedVisualizationShaders"),
-	0,
-	TEXT("Enable advanced strata material debug visualization shaders. Base pass shaders can output such advanced data."),
-	ECVF_ReadOnly | ECVF_RenderThreadSafe);
-
-bool IsStrataAdvancedVisualizationShadersEnabled()
-{
-	return CVarStrataDebugAdvancedVisualizationShaders.GetValueOnRenderThread() > 0;
-}
-
 namespace Strata
 {
 // Forward declarations
@@ -285,7 +273,6 @@ static void AddVisualizeMaterialCountPasses(FRDGBuilder & GraphBuilder, const FV
 }
 
 float GetStrataTileOverflowRatio(const FViewInfo& View);
-bool IsClassificationCoord8bits();
 bool IsClassificationAsync();
 bool SupportsCMask(const FStaticShaderPlatform InPlatform);
 bool DoesStrataTileOverflowUseMaterialData();
@@ -298,16 +285,16 @@ static void AddVisualizeSystemInfoPasses(FRDGBuilder& GraphBuilder, const FViewI
 	ShaderPrint::RequestSpaceForCharacters(1024);
 
 	FStrataSystemInfoCS::FParameters* PassParameters = GraphBuilder.AllocParameters<FStrataSystemInfoCS::FParameters>();
-	PassParameters->bAdvancedDebugEnabled = IsStrataAdvancedVisualizationShadersEnabled() ? 1u : 0u;
+	PassParameters->bAdvancedDebugEnabled = IsAdvancedVisualizationEnabled() ? 1u : 0u;
 	PassParameters->bEnergyConservation = View.ViewState ? View.ViewState->ShadingEnergyConservationData.bEnergyConservation : false;;
 	PassParameters->bEnergyPreservation = View.ViewState ? View.ViewState->ShadingEnergyConservationData.bEnergyPreservation : false;;
-	PassParameters->bDbufferPass = IsStrataDbufferPassEnabled(View.GetShaderPlatform()) ? 1 : 0;
+	PassParameters->bDbufferPass = IsDBufferPassEnabled(View.GetShaderPlatform()) ? 1 : 0;
 	PassParameters->ClassificationCMask = SupportsCMask(View.GetShaderPlatform()) ? 1 : 0;
 	PassParameters->ClassificationAsync = IsClassificationAsync() ? 1 : 0;
-	PassParameters->Classification8bits = IsClassificationCoord8bits() ? 1 : 0;
+	PassParameters->Classification8bits = Is8bitTileCoordEnabled() ? 1 : 0;
 	PassParameters->TileOverflowRatio = GetStrataTileOverflowRatio(View);
 	PassParameters->bTileOverflowUseMaterialData = DoesStrataTileOverflowUseMaterialData() ? 1 : 0;
-	PassParameters->bRoughRefraction = IsStrataOpaqueMaterialRoughRefractionEnabled() ? 1 : 0;
+	PassParameters->bRoughRefraction = IsOpaqueRoughRefractionEnabled() ? 1 : 0;
 	PassParameters->ClassificationTileDrawIndirectBuffer = GraphBuilder.CreateSRV(View.StrataViewData.ClassificationTileDrawIndirectBuffer, PF_R32_UINT);
 	PassParameters->ViewUniformBuffer = View.ViewUniformBuffer;
 	PassParameters->Strata = Strata::BindStrataGlobalUniformParameters(View);
@@ -321,7 +308,7 @@ static void AddVisualizeSystemInfoPasses(FRDGBuilder& GraphBuilder, const FViewI
 // Draw each material layer independently
 static void AddVisualizeAdvancedMaterialPasses(FRDGBuilder& GraphBuilder, const FViewInfo& View, FScreenPassTexture& ScreenPassSceneColor, EShaderPlatform Platform)
 {
-	if (!IsStrataAdvancedVisualizationShadersEnabled())
+	if (!IsAdvancedVisualizationEnabled())
 	{
 		return;
 	}
@@ -426,7 +413,7 @@ FScreenPassTexture AddStrataDebugPasses(FRDGBuilder& GraphBuilder, const FViewIn
 		}
 		else if (DebugMode == FStrataVisualizationData::FViewMode::DecalClassification)
 		{
-			if (IsStrataDbufferPassEnabled(View.GetShaderPlatform()))
+			if (IsDBufferPassEnabled(View.GetShaderPlatform()))
 			{
 				AddStrataInternalClassificationTilePass(GraphBuilder, View, nullptr, &ScreenPassSceneColor.Texture, EStrataTileType::EDecalSimple, bDebugPass);
 				AddStrataInternalClassificationTilePass(GraphBuilder, View, nullptr, &ScreenPassSceneColor.Texture, EStrataTileType::EDecalSingle, bDebugPass);
@@ -435,7 +422,7 @@ FScreenPassTexture AddStrataDebugPasses(FRDGBuilder& GraphBuilder, const FViewIn
 		}
 		else if (DebugMode == FStrataVisualizationData::FViewMode::RoughRefractionClassification)
 		{
-			if (IsStrataOpaqueMaterialRoughRefractionEnabled())
+			if (IsOpaqueRoughRefractionEnabled())
 			{
 				AddStrataInternalClassificationTilePass(GraphBuilder, View, nullptr, &ScreenPassSceneColor.Texture, EStrataTileType::EOpaqueRoughRefraction, bDebugPass);
 				AddStrataInternalClassificationTilePass(GraphBuilder, View, nullptr, &ScreenPassSceneColor.Texture, EStrataTileType::EOpaqueRoughRefractionSSSWithout, bDebugPass);
