Fix sky LUT textures not being transitioned to readable before they're used for environment capture.

#rnx

diff --git a/Engine/Source/Runtime/Renderer/Private/DeferredShadingRenderer.cpp b/Engine/Source/Runtime/Renderer/Private/DeferredShadingRenderer.cpp
--- a/Engine/Source/Runtime/Renderer/Private/DeferredShadingRenderer.cpp
+++ b/Engine/Source/Runtime/Renderer/Private/DeferredShadingRenderer.cpp
@@ -3334,6 +3334,9 @@ void FDeferredShadingSceneRenderer::Render(FRDGBuilder& GraphBuilder)
 	{
 		// Generate the Sky/Atmosphere look up tables
 		RenderSkyAtmosphereLookUpTables(GraphBuilder, ExternalAccessQueue);
+
+		// Sky env map capture uses the view UB, which contains the LUTs computed above. We need to transition them to readable now.
+		ExternalAccessQueue.Submit(GraphBuilder);
 	}
 
 	// Capture the SkyLight using the SkyAtmosphere and VolumetricCloud component if available.
