PR #9998: Avoid FmodReportError (Contributed by [REDACTED_PERSON])

diff --git a/Engine/Source/Runtime/Renderer/Private/TranslucentLighting.cpp b/Engine/Source/Runtime/Renderer/Private/TranslucentLighting.cpp
--- a/Engine/Source/Runtime/Renderer/Private/TranslucentLighting.cpp
+++ b/Engine/Source/Runtime/Renderer/Private/TranslucentLighting.cpp
@@ -184,15 +184,22 @@ void FViewInfo::CalcTranslucencyLightingVolumeBounds(FBox* InOutCascadeBoundsArr
 			RadiusSquared = FMath::Max(RadiusSquared, (Center - SplitVertices[VertexIndex]).SizeSquared());
 		}
 
-		FSphere SphereBounds(Center, FMath::Sqrt(RadiusSquared));
+		if (RadiusSquared > 0) // Avoid issues with bad cvar usage, e.g. r.TranslucencyLightingVolumeInnerDistance.
+		{
+			FSphere SphereBounds(Center, FMath::Sqrt(RadiusSquared));
 
-		// Snap the center to a multiple of the volume dimension for stability
-		const int32 TranslucencyLightingVolumeDim = GetTranslucencyLightingVolumeDim();
-		SphereBounds.Center.X = SphereBounds.Center.X - FMath::Fmod(SphereBounds.Center.X, SphereBounds.W * 2 / TranslucencyLightingVolumeDim);
-		SphereBounds.Center.Y = SphereBounds.Center.Y - FMath::Fmod(SphereBounds.Center.Y, SphereBounds.W * 2 / TranslucencyLightingVolumeDim);
-		SphereBounds.Center.Z = SphereBounds.Center.Z - FMath::Fmod(SphereBounds.Center.Z, SphereBounds.W * 2 / TranslucencyLightingVolumeDim);
+			// Snap the center to a multiple of the volume dimension for stability
+			const int32 TranslucencyLightingVolumeDim = GetTranslucencyLightingVolumeDim();
+			SphereBounds.Center.X = SphereBounds.Center.X - FMath::Fmod(SphereBounds.Center.X, SphereBounds.W * 2 / TranslucencyLightingVolumeDim);
+			SphereBounds.Center.Y = SphereBounds.Center.Y - FMath::Fmod(SphereBounds.Center.Y, SphereBounds.W * 2 / TranslucencyLightingVolumeDim);
+			SphereBounds.Center.Z = SphereBounds.Center.Z - FMath::Fmod(SphereBounds.Center.Z, SphereBounds.W * 2 / TranslucencyLightingVolumeDim);
 
-		InOutCascadeBoundsArray[CascadeIndex] = FBox(SphereBounds.Center - SphereBounds.W, SphereBounds.Center + SphereBounds.W);
+			InOutCascadeBoundsArray[CascadeIndex] = FBox(SphereBounds.Center - SphereBounds.W, SphereBounds.Center + SphereBounds.W);
+		}
+		else
+		{
+			InOutCascadeBoundsArray[CascadeIndex] = FBox(Center, Center);
+		}
 	}
 }
 
