Fix Lumen.Visualize.IndirectDiffuse with Strata enabled

diff --git a/Engine/Shaders/Private/DiffuseIndirectComposite.usf b/Engine/Shaders/Private/DiffuseIndirectComposite.usf
--- a/Engine/Shaders/Private/DiffuseIndirectComposite.usf
+++ b/Engine/Shaders/Private/DiffuseIndirectComposite.usf
@@ -244,8 +244,11 @@ FStrataDeferredLighting StrataIndirectLighting(
 		const float4 InSpecularLighting			= InputSpecularLighting;
 	#endif
 
-		// STRATA_TODO: Override the BSDF F0/DiffuseColor when bVisualizeDiffuseIndirect is set
 		FStrataBSDF BSDF = UnpackStrataBSDF(MaterialBuffer, StrataAddressing, StrataPixelHeader);
+		if (bVisualizeDiffuseIndirect)
+		{
+			StrataSetBSDFDiffuseColor(BSDF, float3(0.18f, 0.18f, 0.18f));
+		}
 
 		// Create the BSDF context
 		FStrataBSDFContext StrataBSDFContext = StrataCreateBSDFContext(StrataPixelHeader, BSDF, StrataAddressing, V);
diff --git a/Engine/Shaders/Private/Strata/Strata.ush b/Engine/Shaders/Private/Strata/Strata.ush
--- a/Engine/Shaders/Private/Strata/Strata.ush
+++ b/Engine/Shaders/Private/Strata/Strata.ush
@@ -2794,6 +2794,17 @@ float3 StrataGetBSDFDiffuseColor(in FStrataBSDF BSDF)
 	return 0.f;
 }
 
+void StrataSetBSDFDiffuseColor(inout FStrataBSDF BSDF, float3 InColor)
+{
+	switch (BSDF_GETTYPE(BSDF))
+	{
+	case STRATA_BSDF_TYPE_SLAB: 			SLAB_DIFFUSEALBEDO(BSDF) = InColor; break;
+	case STRATA_BSDF_TYPE_HAIR: 			HAIR_BASECOLOR(BSDF) = InColor; break;
+	case STRATA_BSDF_TYPE_EYE: 				EYE_DIFFUSEALBEDO(BSDF) = InColor; break;
+	case STRATA_BSDF_TYPE_SINGLELAYERWATER: SLW_BASECOLOR(BSDF) = InColor; break;
+	}
+}
+
 float3 StrataGetBSDFSpecularF0(in FStrataBSDF BSDF)
 {
 	switch (BSDF_GETTYPE(BSDF))
diff --git a/Engine/Source/Runtime/Renderer/Private/IndirectLightRendering.cpp b/Engine/Source/Runtime/Renderer/Private/IndirectLightRendering.cpp
--- a/Engine/Source/Runtime/Renderer/Private/IndirectLightRendering.cpp
+++ b/Engine/Source/Runtime/Renderer/Private/IndirectLightRendering.cpp
@@ -1178,6 +1178,12 @@ void FDeferredShadingSceneRenderer::RenderDiffuseIndirectAndAmbientOcclusion(
 		const bool bApplyDiffuseIndirect = (DenoiserOutputs.Textures[0] || AmbientOcclusionMask) && (!bIsVisualizePass || ViewPipelineState.DiffuseIndirectDenoiser != IScreenSpaceDenoiser::EMode::Disabled || ViewPipelineState.DiffuseIndirectMethod == EDiffuseIndirectMethod::Lumen)
 			&& !(IsMetalPlatform(ShaderPlatform) && !IsMetalSM5Platform(ShaderPlatform));
 
+		// When Lumen visualization is enabled, clear the scene color to ensure sky pixels are cleared correctly.
+		if (Strata::IsStrataEnabled() && bIsVisualizePass)
+		{
+			AddClearRenderTargetPass(GraphBuilder, SceneColorTexture, FLinearColor::Black);
+		}
+
 		auto ApplyDiffuseIndirect = [&](EStrataTileType TileType)
 		{
 			FDiffuseIndirectCompositePS::FParameters* PassParameters = GraphBuilder.AllocParameters<FDiffuseIndirectCompositePS::FParameters>();
