Remove unused parameter from FMobileDirectionalLightFunctionPS and FMobileRadialLightFunctionPS

diff --git a/Engine/Shaders/Private/MobileDeferredShading.usf b/Engine/Shaders/Private/MobileDeferredShading.usf
--- a/Engine/Shaders/Private/MobileDeferredShading.usf
+++ b/Engine/Shaders/Private/MobileDeferredShading.usf
@@ -33,7 +33,7 @@
 #endif 
 
 float4x4 TranslatedWorldToLight;
-float3 LightFunctionParameters2;
+float2 LightFunctionParameters2;
 
 half ComputeLightFunctionMultiplier(float3 TranslatedWorldPosition)
 {
diff --git a/Engine/Source/Runtime/Renderer/Private/MobileDeferredShadingPass.cpp b/Engine/Source/Runtime/Renderer/Private/MobileDeferredShadingPass.cpp
--- a/Engine/Source/Runtime/Renderer/Private/MobileDeferredShadingPass.cpp
+++ b/Engine/Source/Runtime/Renderer/Private/MobileDeferredShadingPass.cpp
@@ -68,7 +68,7 @@ class FMobileDirectionalLightFunctionPS : public FMaterialShader
 		SHADER_PARAMETER_STRUCT_REF(FMobileReflectionCaptureShaderData, MobileReflectionCaptureData)
 		SHADER_PARAMETER(FMatrix44f, TranslatedWorldToLight)
 		SHADER_PARAMETER(FVector4f, LightFunctionParameters)
-		SHADER_PARAMETER(FVector3f, LightFunctionParameters2)
+		SHADER_PARAMETER(FVector2f, LightFunctionParameters2)
 		SHADER_PARAMETER_TEXTURE(Texture2D, ScreenSpaceShadowMaskTexture)
 		SHADER_PARAMETER_SAMPLER(SamplerState, ScreenSpaceShadowMaskSampler)
 	END_SHADER_PARAMETER_STRUCT()
@@ -177,7 +177,7 @@ public:
 	BEGIN_SHADER_PARAMETER_STRUCT(FParameters, )
 		SHADER_PARAMETER(FMatrix44f, TranslatedWorldToLight)
 		SHADER_PARAMETER(FVector4f, LightFunctionParameters)
-		SHADER_PARAMETER(FVector3f, LightFunctionParameters2)
+		SHADER_PARAMETER(FVector2f, LightFunctionParameters2)
 		SHADER_PARAMETER_STRUCT_REF(FDeferredLightUniformStruct, DeferredLightUniforms)
 		SHADER_PARAMETER_STRUCT_INCLUDE(FMobileMovableLocalLightShadowParameters, MobileMovableLocalLightShadow)
 	END_SHADER_PARAMETER_STRUCT()
@@ -493,8 +493,7 @@ static void RenderDirectionalLight(FRHICommandListImmediate& RHICmdList, const F
 	}
 
 	{
-		const bool bUseMovableLight = !DirectionalLight.Proxy->HasStaticShadowing();
-		PassParameters.LightFunctionParameters2 = FVector3f(DirectionalLight.Proxy->GetLightFunctionFadeDistance(), DirectionalLight.Proxy->GetLightFunctionDisabledBrightness(), bUseMovableLight ? 1.0f : 0.0f);
+		PassParameters.LightFunctionParameters2 = FVector2f(DirectionalLight.Proxy->GetLightFunctionFadeDistance(), DirectionalLight.Proxy->GetLightFunctionDisabledBrightness());
 		const FVector Scale = DirectionalLight.Proxy->GetLightFunctionScale();
 		// Switch x and z so that z of the user specified scale affects the distance along the light direction
 		const FVector InverseScale = FVector(1.f / Scale.Z, 1.f / Scale.Y, 1.f / Scale.X);
@@ -758,7 +757,7 @@ static void RenderLocalLight(
 	PassParameters.DeferredLightUniforms = TUniformBufferRef<FDeferredLightUniformStruct>::CreateUniformBufferImmediate(GetDeferredLightParameters(View, LightSceneInfo), EUniformBufferUsage::UniformBuffer_SingleFrame);
 	const float TanOuterAngle = bIsSpotLight ? FMath::Tan(LightSceneInfo.Proxy->GetOuterConeAngle()) : 1.0f;
 	PassParameters.LightFunctionParameters = FVector4f(TanOuterAngle, 1.0f /*ShadowFadeFraction*/, bIsSpotLight ? 1.0f : 0.0f, bIsPointLight ? 1.0f : 0.0f);
-	PassParameters.LightFunctionParameters2 = FVector3f(LightSceneInfo.Proxy->GetLightFunctionFadeDistance(), LightSceneInfo.Proxy->GetLightFunctionDisabledBrightness(), IsMobileMovableSpotlightShadowsEnabled(Scene.GetShaderPlatform()) ? 1.0f : 0.0f);
+	PassParameters.LightFunctionParameters2 = FVector2f(LightSceneInfo.Proxy->GetLightFunctionFadeDistance(), LightSceneInfo.Proxy->GetLightFunctionDisabledBrightness());
 	const FVector Scale = LightSceneInfo.Proxy->GetLightFunctionScale();
 	// Switch x and z so that z of the user specified scale affects the distance along the light direction
 	const FVector InverseScale = FVector(1.f / Scale.Z, 1.f / Scale.Y, 1.f / Scale.X);
