The Global SDF will recapture the first two clipmaps if they cache Mesh SDFs which haven't streamed in yet, fixing the Global SDF after a teleport or new map load
 * CullObjectsToClipmapCS writes to a buffer when any incomplete Mesh SDFs are used, which is read back by the CPU and a new clipmap update gets queued for when streaming is complete
 * Controlled by r.AOGlobalDistanceField.RecacheClipmapsWithPendingStreaming (enabled by default)
 * New debug option to profile the cost of heightfield / Mesh SDF streaming recaches: r.AOGlobalDistanceField.ForceRecacheForStreaming
 * New information added to the draw event to debug full recapture reasons: "Clipmap:0 CacheType:MostlyStatic FULL UPDATE: HeightfieldStreaming"

diff --git a/Engine/Shaders/Private/DistanceField/GlobalDistanceField.usf b/Engine/Shaders/Private/DistanceField/GlobalDistanceField.usf
--- a/Engine/Shaders/Private/DistanceField/GlobalDistanceField.usf
+++ b/Engine/Shaders/Private/DistanceField/GlobalDistanceField.usf
@@ -38,6 +38,9 @@ StructuredBuffer<uint> CullGridObjectArray;
 
 RWStructuredBuffer<uint> RWObjectIndexBuffer;
 RWStructuredBuffer<uint> RWObjectIndexNumBuffer;
+
+RWStructuredBuffer<uint> RWHasPendingStreaming;
+
 float3 ClipmapTranslatedWorldCenter;
 float3 ClipmapWorldExtent;
 float InfluenceRadius;
@@ -74,6 +77,16 @@ void CullObjectsToClipmapCS(
 				uint DestIndex;
 				InterlockedAdd(RWObjectIndexNumBuffer[0], 1U, DestIndex);
 				RWObjectIndexBuffer[DestIndex] = ObjectIndex;
+
+#if READBACK_HAS_PENDING_STREAMING
+				FDFObjectData DFObjectData = LoadDFObjectData(ObjectIndex);
+				uint NumMips = LoadDFAssetData(DFObjectData.AssetIndex, 0).NumMips;
+
+				if (NumMips == 1)
+				{
+					RWHasPendingStreaming[0] = 1;
+				}
+#endif
 			}
 		}
 	}
diff --git a/Engine/Source/Runtime/Renderer/Private/GlobalDistanceField.cpp b/Engine/Source/Runtime/Renderer/Private/GlobalDistanceField.cpp
--- a/Engine/Source/Runtime/Renderer/Private/GlobalDistanceField.cpp
+++ b/Engine/Source/Runtime/Renderer/Private/GlobalDistanceField.cpp
@@ -269,6 +269,22 @@ TAutoConsoleVariable<int32> CVarGlobalDistanceFieldDebug(
 	ECVF_Scalability | ECVF_RenderThreadSafe
 );
 
+int32 GAOGlobalDistanceFieldRecacheClipmapsWithPendingStreaming = 1;
+FAutoConsoleVariableRef CVarAOGlobalDistanceFieldRecacheClipmapsWithPendingStreaming(
+	TEXT("r.AOGlobalDistanceField.RecacheClipmapsWithPendingStreaming"),
+	GAOGlobalDistanceFieldRecacheClipmapsWithPendingStreaming,
+	TEXT("Whether to readback clipmaps cached with incomplete Mesh SDFs due to streaming and reache them on subsequent frames.  Fixes innaccurate Global SDF around the camera after teleporting or loading a new level."),
+	ECVF_Scalability | ECVF_RenderThreadSafe
+);
+
+int32 GAOGlobalDistanceFieldForceRecacheForStreaming = 0;
+FAutoConsoleVariableRef CVarAOGlobalDistanceFieldForceRecacheForStreaming(
+	TEXT("r.AOGlobalDistanceField.ForceRecacheForStreaming"),
+	GAOGlobalDistanceFieldForceRecacheForStreaming,
+	TEXT("Useful for debugging or profiling full clipmap updates that happen when a clipmap is detected to have pending streaming."),
+	ECVF_Scalability | ECVF_RenderThreadSafe
+);
+
 bool UseGlobalDistanceField()
 {
 	return GAOGlobalDistanceField != 0;
@@ -315,6 +331,24 @@ FGlobalDistanceFieldParameters2 SetupGlobalDistanceFieldParameters(const FGlobal
 	return ShaderParameters;
 }
 
+const TCHAR* GetRecaptureReasonString(EGlobalSDFFullRecaptureReason Reason)
+{
+	switch (Reason)
+	{
+	case EGlobalSDFFullRecaptureReason::None:
+		return TEXT("");
+	case EGlobalSDFFullRecaptureReason::TooManyUpdateBounds:
+		return TEXT("FULL UPDATE: TooManyUpdateBounds");
+	case EGlobalSDFFullRecaptureReason::HeightfieldStreaming:
+		return TEXT("FULL UPDATE: HeightfieldStreaming");
+	case EGlobalSDFFullRecaptureReason::MeshSDFStreaming:
+		return TEXT("FULL UPDATE: MeshSDFStreaming");
+	case EGlobalSDFFullRecaptureReason::NoViewState:
+		return TEXT("FULL UPDATE: NoViewState");
+	default:
+		return nullptr;
+	}
+}
 
 float GetMinMeshSDFRadius(float VoxelWorldSize)
 {
@@ -725,6 +759,60 @@ FVector GetGlobalDistanceFieldViewOrigin(const FViewInfo& View, int32 ClipmapInd
 	return CameraOrigin;
 }
 
+void RecaptureClipmapForMeshSDFStreamingIfNeeded(
+	const FViewInfo& View,
+	const FScene* Scene,
+	FPersistentGlobalDistanceFieldData& GlobalDistanceFieldData, 
+	FGlobalDistanceFieldClipmapState& ClipmapViewState,
+	FGlobalDistanceFieldClipmap& Clipmap,
+	int32 ClipmapIndex,
+	uint32 CacheType,
+	const FBox& ClipmapBounds)
+{
+	FRHIGPUBufferReadback* LatestReadbackBuffer = nullptr;
+
+	// Find latest buffer that is ready
+	while (ClipmapViewState.ReadbackBuffersNumPending > 0)
+	{
+		uint32 Index = (ClipmapViewState.ReadbackBuffersWriteIndex + ClipmapViewState.MaxPendingStreamingReadbackBuffers - ClipmapViewState.ReadbackBuffersNumPending) % ClipmapViewState.MaxPendingStreamingReadbackBuffers;
+		if (ClipmapViewState.HasPendingStreamingReadbackBuffers[Index]->IsReady())
+		{
+			ClipmapViewState.ReadbackBuffersNumPending--;
+			LatestReadbackBuffer = ClipmapViewState.HasPendingStreamingReadbackBuffers[Index];
+		}
+		else
+		{
+			break;
+		}
+	}
+
+	bool bPreviousCaptureHadPendingStreaming = false;
+
+	// Readback whether the last CullObjectsToClipmap for this clipmap detected Mesh SDFs that have not streamed in (NumMips == 1)
+	if (LatestReadbackBuffer)
+	{
+		const uint32* LatestReadbackBufferPtr = (const uint32*)LatestReadbackBuffer->Lock(1 * sizeof(uint32));
+		bPreviousCaptureHadPendingStreaming = LatestReadbackBufferPtr[0] != 0;
+		LatestReadbackBuffer->Unlock();
+	}
+
+	TArray<int32>& DeferredUpdatesForMeshSDFStreaming = GlobalDistanceFieldData.DeferredUpdatesForMeshSDFStreaming[CacheType];
+
+	if (bPreviousCaptureHadPendingStreaming)
+	{
+		// Add a new deferred update for when the streaming is finished
+		DeferredUpdatesForMeshSDFStreaming.AddUnique(ClipmapIndex);
+	}
+	// Mesh SDFs are done streaming, recapture the clipmaps
+	else if (!Scene->DistanceFieldSceneData.HasPendingStreaming() && DeferredUpdatesForMeshSDFStreaming.Remove(ClipmapIndex))
+	{
+		// Push full update
+		Clipmap.UpdateBounds.Reset();
+		Clipmap.UpdateBounds.Add(FClipmapUpdateBounds(ClipmapBounds.GetCenter(), ClipmapBounds.GetExtent(), false));
+		Clipmap.FullRecaptureReason = EGlobalSDFFullRecaptureReason::MeshSDFStreaming;
+	}
+}
+
 static void ComputeUpdateRegionsAndUpdateViewState(
 	FRHICommandListImmediate& RHICmdList, 
 	FViewInfo& View, 
@@ -1109,6 +1197,7 @@ static void ComputeUpdateRegionsAndUpdateViewState(
 					{
 						Clipmap.UpdateBounds.Reset();
 						Clipmap.UpdateBounds.Add(FClipmapUpdateBounds(ClipmapBounds.GetCenter(), ClipmapBounds.GetExtent(), false));
+						Clipmap.FullRecaptureReason = EGlobalSDFFullRecaptureReason::TooManyUpdateBounds;
 					}
 
 					// Check if the clipmap intersects with a pending update region
@@ -1123,18 +1212,21 @@ static void ComputeUpdateRegionsAndUpdateViewState(
 					}
 
 					// If some of the height fields has pending streaming regions, postpone a full update.
-					if (bHasPendingStreaming)
+					if (bHasPendingStreaming && GAOGlobalDistanceFieldForceRecacheForStreaming == 0)
 					{
 						// Mark a pending update for this height field. It will get processed when all pending texture streaming affecting it will be completed.
 						GlobalDistanceFieldData.DeferredUpdates[CacheType].AddUnique(ClipmapIndex);
 					}
-					else if (GlobalDistanceFieldData.DeferredUpdates[CacheType].Remove(ClipmapIndex) > 0)
+					else if (GlobalDistanceFieldData.DeferredUpdates[CacheType].Remove(ClipmapIndex) > 0 || GAOGlobalDistanceFieldForceRecacheForStreaming != 0)
 					{
 						// Push full update
 						Clipmap.UpdateBounds.Reset();
 						Clipmap.UpdateBounds.Add(FClipmapUpdateBounds(ClipmapBounds.GetCenter(), ClipmapBounds.GetExtent(), false));
+						Clipmap.FullRecaptureReason = EGlobalSDFFullRecaptureReason::HeightfieldStreaming;
 					}
 
+					RecaptureClipmapForMeshSDFStreamingIfNeeded(View, Scene, GlobalDistanceFieldData, ClipmapViewState, Clipmap, ClipmapIndex, CacheType, ClipmapBounds);
+
 					ClipmapViewState.Cache[CacheType].PrimitiveModifiedBounds.Empty(DistanceField::MinPrimitiveModifiedBoundsAllocation);
 				}
 
@@ -1200,6 +1292,7 @@ static void ComputeUpdateRegionsAndUpdateViewState(
 
 				Clipmap.UpdateBounds.Reset();
 				Clipmap.UpdateBounds.Add(FClipmapUpdateBounds(ClipmapBounds.GetCenter(), ClipmapBounds.GetExtent(), false));
+				Clipmap.FullRecaptureReason = EGlobalSDFFullRecaptureReason::NoViewState;
 			}
 		}
 	}
@@ -1302,7 +1395,9 @@ class FCullObjectsToClipmapCS : public FGlobalShader
 	BEGIN_SHADER_PARAMETER_STRUCT(FParameters, )
 		SHADER_PARAMETER_RDG_BUFFER_UAV(RWStructuredBuffer<uint>, RWObjectIndexBuffer)
 		SHADER_PARAMETER_RDG_BUFFER_UAV(RWStructuredBuffer<uint>, RWObjectIndexNumBuffer)
+		SHADER_PARAMETER_RDG_BUFFER_UAV(RWStructuredBuffer<uint>, RWHasPendingStreaming)
 		SHADER_PARAMETER_STRUCT_INCLUDE(FDistanceFieldObjectBufferParameters, DistanceFieldObjectBuffers)
+		SHADER_PARAMETER_STRUCT_INCLUDE(FDistanceFieldAtlasParameters, DistanceFieldAtlasParameters)
 		SHADER_PARAMETER(FVector3f, ClipmapTranslatedWorldCenter)
 		SHADER_PARAMETER(FVector3f, ClipmapWorldExtent)
 		SHADER_PARAMETER(uint32, AcceptOftenMovingObjectsOnly)
@@ -1312,6 +1407,9 @@ class FCullObjectsToClipmapCS : public FGlobalShader
 		SHADER_PARAMETER(FVector3f, RelativePreViewTranslation)
 	END_SHADER_PARAMETER_STRUCT()
 
+	class FReadbackHasPendingStreaming : SHADER_PERMUTATION_BOOL("READBACK_HAS_PENDING_STREAMING");
+	using FPermutationDomain = TShaderPermutationDomain<FReadbackHasPendingStreaming>;
+
 	static bool ShouldCompilePermutation(const FGlobalShaderPermutationParameters& Parameters)
 	{
 		return ShouldCompileDistanceFieldShaders(Parameters.Platform);
@@ -2090,9 +2188,8 @@ void UpdateGlobalDistanceFieldVolume(
 
 				for (int32 ClipmapIndex = 0; ClipmapIndex < Clipmaps.Num(); ClipmapIndex++)
 				{
-					RDG_EVENT_SCOPE(GraphBuilder, "Clipmap:%d CacheType:%s", ClipmapIndex, CacheType == GDF_MostlyStatic ? TEXT("MostlyStatic") : TEXT("Movable"));
-
 					const FGlobalDistanceFieldClipmap& Clipmap = Clipmaps[ClipmapIndex];
+					RDG_EVENT_SCOPE(GraphBuilder, "Clipmap:%d CacheType:%s %s", ClipmapIndex, CacheType == GDF_MostlyStatic ? TEXT("MostlyStatic") : TEXT("Movable"), GetRecaptureReasonString(Clipmap.FullRecaptureReason));
 
 					const FVector ClipmapWorldCenter = Clipmap.Bounds.GetCenter();
 
@@ -2211,29 +2308,84 @@ void UpdateGlobalDistanceFieldVolume(
 
 							AddClearUAVPass(GraphBuilder, GraphBuilder.CreateUAV(ObjectIndexNumBuffer, PF_R32_UINT), 0);
 
-							FCullObjectsToClipmapCS::FParameters* PassParameters = GraphBuilder.AllocParameters<FCullObjectsToClipmapCS::FParameters>();
-							PassParameters->RWObjectIndexBuffer = GraphBuilder.CreateUAV(ObjectIndexBuffer, PF_R32_UINT);
-							PassParameters->RWObjectIndexNumBuffer = GraphBuilder.CreateUAV(ObjectIndexNumBuffer, PF_R32_UINT);
-							PassParameters->DistanceFieldObjectBuffers = DistanceField::SetupObjectBufferParameters(GraphBuilder, DistanceFieldSceneData);
-							PassParameters->ClipmapTranslatedWorldCenter = (FVector3f)TranslatedBounds.GetCenter();
-							PassParameters->ClipmapWorldExtent = (FVector3f)TranslatedBounds.GetExtent();
-							PassParameters->AcceptOftenMovingObjectsOnly = AcceptOftenMovingObjectsOnlyValue;
-							const float RadiusThresholdScale = bLumenEnabled ? 1.0f / FMath::Clamp(View.FinalPostProcessSettings.LumenSceneDetail, .01f, 100.0f) : 1.0f;
-							PassParameters->MeshSDFRadiusThreshold = GetMinMeshSDFRadius(ClipmapVoxelSize.X) * RadiusThresholdScale;
-							PassParameters->InfluenceRadiusSq = ClipmapInfluenceRadius * ClipmapInfluenceRadius;
+							bool bRecacheClipmapsWithPendingStreaming = GAOGlobalDistanceFieldRecacheClipmapsWithPendingStreaming != 0 &&
+								View.ViewState &&
+								CacheType == GDF_MostlyStatic &&
+								ClipmapIndex < 2;
 
-							PassParameters->ViewTilePosition = ViewTilePosition;
-							PassParameters->RelativePreViewTranslation = RelativePreViewTranslation;
+							FRDGBufferRef HasPendingStreamingReadbackBuffer = nullptr;
+
+							if (bRecacheClipmapsWithPendingStreaming)
+							{
+								const FGlobalDistanceFieldClipmapState& ClipmapViewState = View.ViewState->GlobalDistanceFieldData->ClipmapState[ClipmapIndex];
 
-							auto ComputeShader = View.ShaderMap->GetShader<FCullObjectsToClipmapCS>();
-							const FIntVector GroupSize = FComputeShaderUtils::GetGroupCount(DistanceFieldSceneData.NumObjectsInBuffer, FCullObjectsToClipmapCS::GetGroupSize());
+								// It is not safe to EnqueueCopy on a buffer that already has a pending copy
+								bRecacheClipmapsWithPendingStreaming = bRecacheClipmapsWithPendingStreaming && ClipmapViewState.ReadbackBuffersNumPending < ClipmapViewState.MaxPendingStreamingReadbackBuffers;
 
-							FComputeShaderUtils::AddPass(
-								GraphBuilder,
-								RDG_EVENT_NAME("CullToClipmap"),
-								ComputeShader,
-								PassParameters,
-								GroupSize);
+								if (bRecacheClipmapsWithPendingStreaming)
+								{
+									FRDGBufferDesc HasPendingStreamingReadbackDesc = FRDGBufferDesc::CreateStructuredDesc(sizeof(uint32), 1);
+									HasPendingStreamingReadbackDesc.Usage = EBufferUsageFlags(HasPendingStreamingReadbackDesc.Usage | BUF_SourceCopy);
+									HasPendingStreamingReadbackBuffer = GraphBuilder.CreateBuffer(HasPendingStreamingReadbackDesc, TEXT("GlobalDistanceField.HasPendingStreamingReadback"));
+									AddClearUAVPass(GraphBuilder, GraphBuilder.CreateUAV(HasPendingStreamingReadbackBuffer, PF_R32_UINT), 0);
+								}
+							}
+
+							{
+								FCullObjectsToClipmapCS::FParameters* PassParameters = GraphBuilder.AllocParameters<FCullObjectsToClipmapCS::FParameters>();
+								PassParameters->RWObjectIndexBuffer = GraphBuilder.CreateUAV(ObjectIndexBuffer, PF_R32_UINT);
+								PassParameters->RWObjectIndexNumBuffer = GraphBuilder.CreateUAV(ObjectIndexNumBuffer, PF_R32_UINT);
+
+								PassParameters->RWHasPendingStreaming = HasPendingStreamingReadbackBuffer ? GraphBuilder.CreateUAV(HasPendingStreamingReadbackBuffer, PF_R32_UINT) : nullptr;
+
+								PassParameters->DistanceFieldObjectBuffers = DistanceField::SetupObjectBufferParameters(GraphBuilder, DistanceFieldSceneData);
+								PassParameters->DistanceFieldAtlasParameters = DistanceField::SetupAtlasParameters(GraphBuilder, DistanceFieldSceneData);
+
+								PassParameters->ClipmapTranslatedWorldCenter = (FVector3f)TranslatedBounds.GetCenter();
+								PassParameters->ClipmapWorldExtent = (FVector3f)TranslatedBounds.GetExtent();
+								PassParameters->AcceptOftenMovingObjectsOnly = AcceptOftenMovingObjectsOnlyValue;
+								const float RadiusThresholdScale = bLumenEnabled ? 1.0f / FMath::Clamp(View.FinalPostProcessSettings.LumenSceneDetail, .01f, 100.0f) : 1.0f;
+								PassParameters->MeshSDFRadiusThreshold = GetMinMeshSDFRadius(ClipmapVoxelSize.X) * RadiusThresholdScale;
+								PassParameters->InfluenceRadiusSq = ClipmapInfluenceRadius * ClipmapInfluenceRadius;
+
+								PassParameters->ViewTilePosition = ViewTilePosition;
+								PassParameters->RelativePreViewTranslation = RelativePreViewTranslation;
+
+								FCullObjectsToClipmapCS::FPermutationDomain PermutationVector;
+								PermutationVector.Set<FCullObjectsToClipmapCS::FReadbackHasPendingStreaming>(bRecacheClipmapsWithPendingStreaming);
+								auto ComputeShader = View.ShaderMap->GetShader<FCullObjectsToClipmapCS>(PermutationVector);
+
+								const FIntVector GroupSize = FComputeShaderUtils::GetGroupCount(DistanceFieldSceneData.NumObjectsInBuffer, FCullObjectsToClipmapCS::GetGroupSize());
+
+								FComputeShaderUtils::AddPass(
+									GraphBuilder,
+									RDG_EVENT_NAME("CullToClipmap"),
+									ComputeShader,
+									PassParameters,
+									GroupSize);
+							}
+
+							if (bRecacheClipmapsWithPendingStreaming)
+							{
+								FGlobalDistanceFieldClipmapState& ClipmapViewState = View.ViewState->GlobalDistanceFieldData->ClipmapState[ClipmapIndex];
+
+								if (!ClipmapViewState.HasPendingStreamingReadbackBuffers[ClipmapViewState.ReadbackBuffersWriteIndex])
+								{
+									FRHIGPUBufferReadback* GPUBufferReadback = new FRHIGPUBufferReadback(TEXT("GlobalDistanceField.HasPendingStreamingReadback"));
+									ClipmapViewState.HasPendingStreamingReadbackBuffers[ClipmapViewState.ReadbackBuffersWriteIndex] = GPUBufferReadback;
+								}
+
+								FRHIGPUBufferReadback* ReadbackBuffer = ClipmapViewState.HasPendingStreamingReadbackBuffers[ClipmapViewState.ReadbackBuffersWriteIndex];
+
+								AddReadbackBufferPass(GraphBuilder, RDG_EVENT_NAME("GlobalDistanceField.HasPendingStreamingReadback"), HasPendingStreamingReadbackBuffer,
+									[ReadbackBuffer, HasPendingStreamingReadbackBuffer](FRHICommandList& RHICmdList)
+									{
+										ReadbackBuffer->EnqueueCopy(RHICmdList, HasPendingStreamingReadbackBuffer->GetRHI(), 0u);
+									});
+
+								ClipmapViewState.ReadbackBuffersWriteIndex = (ClipmapViewState.ReadbackBuffersWriteIndex + 1u) % ClipmapViewState.MaxPendingStreamingReadbackBuffers;
+								ClipmapViewState.ReadbackBuffersNumPending = FMath::Min(ClipmapViewState.ReadbackBuffersNumPending + 1u, ClipmapViewState.MaxPendingStreamingReadbackBuffers);
+							}
 						}
 
 						const uint32 GGlobalDistanceFieldMaxPageNum = GlobalDistanceField::GetMaxPageNum(bLumenEnabled, View.FinalPostProcessSettings.LumenSceneViewDistance);
@@ -2532,7 +2684,7 @@ void UpdateGlobalDistanceFieldVolume(
 
 								FComputeShaderUtils::AddPass(
 									GraphBuilder,
-									RDG_EVENT_NAME("SetupPageFreeListRetunIndirectArgs"),
+									RDG_EVENT_NAME("SetupPageFreeListReturnIndirectArgs"),
 									ComputeShader,
 									PassParameters,
 									FIntVector(1, 1, 1));
diff --git a/Engine/Source/Runtime/Renderer/Private/ScenePrivate.h b/Engine/Source/Runtime/Renderer/Private/ScenePrivate.h
--- a/Engine/Source/Runtime/Renderer/Private/ScenePrivate.h
+++ b/Engine/Source/Runtime/Renderer/Private/ScenePrivate.h
@@ -460,6 +460,8 @@ public:
 		CacheClipmapInfluenceRadius = 0.0f;
 		CacheMostlyStaticSeparately = 1;
 		LastUsedSceneDataForFullUpdate = nullptr;
+
+		HasPendingStreamingReadbackBuffers.AddZeroed(MaxPendingStreamingReadbackBuffers);
 	}
 
 	FInt64Vector FullUpdateOriginInPages;
@@ -470,6 +472,11 @@ public:
 	float CachedClipmapExtent;
 	float CacheClipmapInfluenceRadius;
 
+	TArray<FRHIGPUBufferReadback*> HasPendingStreamingReadbackBuffers;
+	uint32 MaxPendingStreamingReadbackBuffers = 4;
+	uint32 ReadbackBuffersWriteIndex = 0;
+	uint32 ReadbackBuffersNumPending = 0;
+
 	FGlobalDistanceFieldCacheTypeState Cache[GDF_Num];
 
 	// Used to perform a full update of the clip map when the scene data changes
@@ -634,6 +641,7 @@ struct FPersistentGlobalDistanceFieldData : public FThreadSafeRefCountedObject
 {
 	// Array of ClipmapIndex
 	TArray<int32> DeferredUpdates[GDF_Num];
+	TArray<int32> DeferredUpdatesForMeshSDFStreaming[GDF_Num];
 
 	int32	UpdateFrame = 0;
 	bool	bFirstFrame = true;
@@ -1963,6 +1971,11 @@ public:
 		return false;
 	}
 
+	bool HasPendingStreaming() const
+	{
+		return ReadRequests.Num() > 0;
+	}
+
 	inline bool CanUse16BitObjectIndices() const
 	{
 		return bCanUse16BitObjectIndices && (NumObjectsInBuffer < (1 << 16));
diff --git a/Engine/Source/Runtime/Renderer/Private/SceneRendering.h b/Engine/Source/Runtime/Renderer/Private/SceneRendering.h
--- a/Engine/Source/Runtime/Renderer/Private/SceneRendering.h
+++ b/Engine/Source/Runtime/Renderer/Private/SceneRendering.h
@@ -715,6 +715,15 @@ public:
 	FVector Extent;
 };
 
+enum class EGlobalSDFFullRecaptureReason
+{
+	None,
+	TooManyUpdateBounds,
+	HeightfieldStreaming,
+	MeshSDFStreaming,
+	NoViewState
+};
+
 class FGlobalDistanceFieldClipmap
 {
 public:
@@ -724,7 +733,7 @@ public:
 	/** Offset applied to UVs so that only new or dirty areas of the volume texture have to be updated. */
 	FVector ScrollOffset;
 
-	/** Legacy regions in the volume texture to update. Used only by heighfield composition. */
+	EGlobalSDFFullRecaptureReason FullRecaptureReason = EGlobalSDFFullRecaptureReason::None;
 
 	// Bounds in the volume texture to update.
 	TArray<FClipmapUpdateBounds, TInlineAllocator<64>> UpdateBounds;
