Move Instanced View UB up from FViewInfo to FSceneView and fix FMaterialShader::SetViewParameters().

- After CL 23341764 / UE-163717 Instanced View UB cannot be substituted by instanced view's View UB since it is of a different size and holds data for both views.

diff --git a/Engine/Source/Runtime/Engine/Public/SceneView.h b/Engine/Source/Runtime/Engine/Public/SceneView.h
--- a/Engine/Source/Runtime/Engine/Public/SceneView.h
+++ b/Engine/Source/Runtime/Engine/Public/SceneView.h
@@ -1656,8 +1656,17 @@ public:
 	/** Get all secondary views associated with the primary view. */
 	TArray<const FSceneView*> GetSecondaryViews() const;
 
+	/** Returns uniform buffer that is used to access (first) instanced view's properties. Note: it is not the same as that view's ViewUniformBuffer, albeit it contains all relevant data from it. */
+	const TUniformBufferRef<FInstancedViewUniformShaderParameters>& GetInstancedViewUniformBuffer() const
+	{
+		return InstancedViewUniformBuffer;
+	}
+
 protected:
 	FSceneViewStateInterface* EyeAdaptationViewState = nullptr;
+
+	/** Instanced view uniform buffer held on the primary view. */
+	TUniformBufferRef<FInstancedViewUniformShaderParameters> InstancedViewUniformBuffer;
 };
 
 //////////////////////////////////////////////////////////////////////////
diff --git a/Engine/Source/Runtime/Renderer/Private/SceneRendering.h b/Engine/Source/Runtime/Renderer/Private/SceneRendering.h
--- a/Engine/Source/Runtime/Renderer/Private/SceneRendering.h
+++ b/Engine/Source/Runtime/Renderer/Private/SceneRendering.h
@@ -1850,9 +1850,6 @@ private:
 
 	/** Calculates bounding boxes for the translucency lighting volume cascades. */
 	void CalcTranslucencyLightingVolumeBounds(FBox* InOutCascadeBoundsArray, int32 NumCascades) const;
-
-	/** Instanced view uniform buffer held on the primary view. */
-	TUniformBufferRef<FInstancedViewUniformShaderParameters> InstancedViewUniformBuffer;
 };
 
 
diff --git a/Engine/Source/Runtime/Renderer/Public/MaterialShader.h b/Engine/Source/Runtime/Renderer/Public/MaterialShader.h
--- a/Engine/Source/Runtime/Renderer/Public/MaterialShader.h
+++ b/Engine/Source/Runtime/Renderer/Public/MaterialShader.h
@@ -88,13 +88,9 @@ public:
 
 		if (View.bShouldBindInstancedViewUB)
 		{
-			// When drawing the left eye in a stereo scene, copy the right eye view values into the instanced view uniform buffer.
-			const FSceneView* InstancedView = View.GetInstancedSceneView();
-			if (InstancedView)
-			{
-				const auto& InstancedViewUniformBufferParameter = GetUniformBufferParameter<FInstancedViewUniformShaderParameters>();
-				SetUniformBufferParameter(RHICmdList, ShaderRHI, InstancedViewUniformBufferParameter, InstancedView->ViewUniformBuffer);
-			}
+			// When drawing an instanced stereo scene, the instanced view UB should be taken from the same view where it will contains a copy of both left and eye values (see FViewInfo::CreateViewUniformBuffers).
+			const auto& InstancedViewUniformBufferParameter = GetUniformBufferParameter<FInstancedViewUniformShaderParameters>();
+			SetUniformBufferParameter(RHICmdList, ShaderRHI, InstancedViewUniformBufferParameter, View.GetInstancedViewUniformBuffer());
 		}
 	}
 
