Refactored virtual texture update to support additional requests after the initial async task is launched. The gather task is waited on prior to GatherDynamicMeshElements, which allows new virtual textures to be allocated. In EndUpdate, new pages are gathered and processed in EndUpdate. Additionally, page requests are now deferred until EndUpdate if r.VT.AsyncPageRequestTask is 0.

diff --git a/Engine/Source/Runtime/Renderer/Private/DeferredShadingRenderer.cpp b/Engine/Source/Runtime/Renderer/Private/DeferredShadingRenderer.cpp
--- a/Engine/Source/Runtime/Renderer/Private/DeferredShadingRenderer.cpp
+++ b/Engine/Source/Runtime/Renderer/Private/DeferredShadingRenderer.cpp
@@ -2626,9 +2626,15 @@ void FDeferredShadingSceneRenderer::Render(FRDGBuilder& GraphBuilder)
 	const bool bHasRayTracedOverlay = HasRayTracedOverlay(ViewFamily);
 	const bool bAllowStaticLighting = !bHasRayTracedOverlay && IsStaticLightingAllowed();
 
+	TUniquePtr<FVirtualTextureUpdater> VirtualTextureUpdater;
+
 	const bool bUseVirtualTexturing = UseVirtualTexturing(FeatureLevel);
 	if (bUseVirtualTexturing)
 	{
+		FVirtualTextureUpdateSettings Settings;
+		Settings.EnableThrottling(!ViewFamily.bOverrideVirtualTextureThrottle);
+
+		VirtualTextureUpdater = FVirtualTextureSystem::Get().BeginUpdate(GraphBuilder, FeatureLevel, Scene, Settings);
 		VirtualTextureFeedbackBegin(GraphBuilder, Views, SceneTexturesConfig.Extent);
 	}
 
@@ -2664,7 +2670,7 @@ void FDeferredShadingSceneRenderer::Render(FRDGBuilder& GraphBuilder)
 
 	{
 		RDG_GPU_STAT_SCOPE(GraphBuilder, VisibilityCommands);
-		BeginInitViews(GraphBuilder, SceneTexturesConfig, BasePassDepthStencilAccess, ILCTaskData, InstanceCullingManager);
+		BeginInitViews(GraphBuilder, SceneTexturesConfig, BasePassDepthStencilAccess, ILCTaskData, InstanceCullingManager, VirtualTextureUpdater.Get());
 	}
 
 	// GetBinIndexTranslator cannot be called before UpdateAllPrimitiveSceneInfos which can change the number of raster bins
@@ -2840,13 +2846,6 @@ void FDeferredShadingSceneRenderer::Render(FRDGBuilder& GraphBuilder)
 
 	if (bUseVirtualTexturing)
 	{
-		FVirtualTextureUpdateSettings Settings;
-		Settings.EnableThrottling(!ViewFamily.bOverrideVirtualTextureThrottle);
-
-		// We can move this call earlier and enable r.VT.AsyncPageRequestTask after fixing race conditions with InitViews
-		TUniquePtr<FVirtualTextureUpdater> VirtualTextureUpdater;
-		VirtualTextureUpdater = FVirtualTextureSystem::Get().BeginUpdate(GraphBuilder, FeatureLevel, Scene, Settings);
-
 		// Note, should happen after the GPU-Scene update to ensure rendering to runtime virtual textures is using the correctly updated scene
 		FVirtualTextureSystem::Get().EndUpdate(GraphBuilder, MoveTemp(VirtualTextureUpdater), FeatureLevel);
 	}
diff --git a/Engine/Source/Runtime/Renderer/Private/DeferredShadingRenderer.h b/Engine/Source/Runtime/Renderer/Private/DeferredShadingRenderer.h
--- a/Engine/Source/Runtime/Renderer/Private/DeferredShadingRenderer.h
+++ b/Engine/Source/Runtime/Renderer/Private/DeferredShadingRenderer.h
@@ -532,7 +532,7 @@ private:
 	void ComputeLightVisibility();
 
 	/** Determines which primitives are visible for each view. */
-	void BeginInitViews(FRDGBuilder& GraphBuilder, const FSceneTexturesConfig& SceneTexturesConfig, FExclusiveDepthStencil::Type BasePassDepthStencilAccess, struct FILCUpdatePrimTaskData& ILCTaskData, FInstanceCullingManager& InstanceCullingManager);
+	void BeginInitViews(FRDGBuilder& GraphBuilder, const FSceneTexturesConfig& SceneTexturesConfig, FExclusiveDepthStencil::Type BasePassDepthStencilAccess, struct FILCUpdatePrimTaskData& ILCTaskData, FInstanceCullingManager& InstanceCullingManager, FVirtualTextureUpdater* VirtualTextureUpdater);
 	void EndInitViews(FRDGBuilder& GraphBuilder, FLumenSceneFrameTemporaries& FrameTemporaries, struct FILCUpdatePrimTaskData& ILCTaskData, FInstanceCullingManager& InstanceCullingManager, FRDGExternalAccessQueue& ExternalAccessQueue);
 
 	void BeginUpdateLumenSceneTasks(FRDGBuilder& GraphBuilder, FLumenSceneFrameTemporaries& FrameTemporaries);
diff --git a/Engine/Source/Runtime/Renderer/Private/MobileShadingRenderer.cpp b/Engine/Source/Runtime/Renderer/Private/MobileShadingRenderer.cpp
--- a/Engine/Source/Runtime/Renderer/Private/MobileShadingRenderer.cpp
+++ b/Engine/Source/Runtime/Renderer/Private/MobileShadingRenderer.cpp
@@ -433,7 +433,7 @@ void FMobileSceneRenderer::SetupMobileBasePassAfterShadowInit(FExclusiveDepthSte
  * Initialize scene's views.
  * Check visibility, sort translucent items, etc.
  */
-void FMobileSceneRenderer::InitViews(FRDGBuilder& GraphBuilder, FSceneTexturesConfig& SceneTexturesConfig, FInstanceCullingManager& InstanceCullingManager)
+void FMobileSceneRenderer::InitViews(FRDGBuilder& GraphBuilder, FSceneTexturesConfig& SceneTexturesConfig, FInstanceCullingManager& InstanceCullingManager, FVirtualTextureUpdater* VirtualTextureUpdater)
 {
 	FRHICommandListImmediate& RHICmdList = GraphBuilder.RHICmdList;
 	
@@ -464,7 +464,7 @@ void FMobileSceneRenderer::InitViews(FRDGBuilder& GraphBuilder, FSceneTexturesCo
 		Views[0].InitRHIResources();
 		FXSystem->PostInitViews(GraphBuilder, Views, !ViewFamily.EngineShowFlags.HitProxies);
 	}
-	ComputeViewVisibility(RHICmdList, BasePassDepthStencilAccess, ViewCommandsPerView, DynamicIndexBuffer, DynamicVertexBuffer, DynamicReadBuffer, InstanceCullingManager);
+	ComputeViewVisibility(RHICmdList, BasePassDepthStencilAccess, ViewCommandsPerView, DynamicIndexBuffer, DynamicVertexBuffer, DynamicReadBuffer, InstanceCullingManager, VirtualTextureUpdater);
 	PostVisibilityFrameSetup(ILCTaskData);
 
 	FIntPoint RenderTargetSize = ViewFamily.RenderTarget->GetSizeXY();
@@ -868,8 +868,14 @@ void FMobileSceneRenderer::Render(FRDGBuilder& GraphBuilder)
 
 	FRDGSystemTextures::Create(GraphBuilder);
 
+	TUniquePtr<FVirtualTextureUpdater> VirtualTextureUpdater;
+
 	if (bUseVirtualTexturing)
 	{
+		FVirtualTextureUpdateSettings Settings;
+		Settings.EnableThrottling(!ViewFamily.bOverrideVirtualTextureThrottle);
+
+		VirtualTextureUpdater = FVirtualTextureSystem::Get().BeginUpdate(GraphBuilder, FeatureLevel, Scene, Settings);
 		VirtualTextureFeedbackBegin(GraphBuilder, Views, SceneTexturesConfig.Extent);
 	}
 
@@ -902,7 +908,7 @@ void FMobileSceneRenderer::Render(FRDGBuilder& GraphBuilder)
 	GraphBuilder.SetCommandListStat(GET_STATID(STAT_CLMM_InitViews));
 
 	// Find the visible primitives and prepare targets and buffers for rendering
-	InitViews(GraphBuilder, SceneTexturesConfig, InstanceCullingManager);
+	InitViews(GraphBuilder, SceneTexturesConfig, InstanceCullingManager, VirtualTextureUpdater.Get());
 
 	GraphBuilder.SetCommandListStat(GET_STATID(STAT_CLMM_AfterInitViews));
 
@@ -925,12 +931,6 @@ void FMobileSceneRenderer::Render(FRDGBuilder& GraphBuilder)
 
 	if (bUseVirtualTexturing)
 	{
-		FVirtualTextureUpdateSettings Settings;
-		Settings.EnableThrottling(!ViewFamily.bOverrideVirtualTextureThrottle);
-
-		TUniquePtr<FVirtualTextureUpdater> VirtualTextureUpdater;
-		VirtualTextureUpdater = FVirtualTextureSystem::Get().BeginUpdate(GraphBuilder, FeatureLevel, Scene, Settings);
-
 		FVirtualTextureSystem::Get().EndUpdate(GraphBuilder, MoveTemp(VirtualTextureUpdater), FeatureLevel);
 	}
 
diff --git a/Engine/Source/Runtime/Renderer/Private/SceneHitProxyRendering.cpp b/Engine/Source/Runtime/Renderer/Private/SceneHitProxyRendering.cpp
--- a/Engine/Source/Runtime/Renderer/Private/SceneHitProxyRendering.cpp
+++ b/Engine/Source/Runtime/Renderer/Private/SceneHitProxyRendering.cpp
@@ -585,7 +585,7 @@ void FMobileSceneRenderer::RenderHitProxies(FRDGBuilder& GraphBuilder)
 	FInstanceCullingManager& InstanceCullingManager = *GraphBuilder.AllocObject<FInstanceCullingManager>(Scene->GPUScene.IsEnabled(), GraphBuilder);
 
 	// Find the visible primitives.
-	InitViews(GraphBuilder, SceneTexturesConfig, InstanceCullingManager);
+	InitViews(GraphBuilder, SceneTexturesConfig, InstanceCullingManager, nullptr);
 
 	GEngine->GetPreRenderDelegateEx().Broadcast(GraphBuilder);
 
@@ -633,7 +633,7 @@ void FDeferredShadingSceneRenderer::RenderHitProxies(FRDGBuilder& GraphBuilder)
 	FLumenSceneFrameTemporaries LumenFrameTemporaries;
 	FILCUpdatePrimTaskData ILCTaskData;
 	FRDGExternalAccessQueue ExternalAccessQueue;
-	BeginInitViews(GraphBuilder, SceneTexturesConfig, FExclusiveDepthStencil::DepthWrite_StencilWrite, ILCTaskData, InstanceCullingManager);
+	BeginInitViews(GraphBuilder, SceneTexturesConfig, FExclusiveDepthStencil::DepthWrite_StencilWrite, ILCTaskData, InstanceCullingManager, nullptr);
 
 	extern TSet<IPersistentViewUniformBufferExtension*> PersistentViewUniformBufferExtensions;
 
diff --git a/Engine/Source/Runtime/Renderer/Private/SceneRendering.h b/Engine/Source/Runtime/Renderer/Private/SceneRendering.h
--- a/Engine/Source/Runtime/Renderer/Private/SceneRendering.h
+++ b/Engine/Source/Runtime/Renderer/Private/SceneRendering.h
@@ -58,6 +58,7 @@ class FSceneTextureShaderParameters;
 class FLumenSceneData;
 class FShadowSceneRenderer;
 class FGlobalShaderMap;
+class FVirtualTextureUpdater;
 
 struct FCloudRenderContext;
 struct FSceneWithoutWaterTextures;
@@ -2392,7 +2393,8 @@ protected:
 
 	/** Computes which primitives are visible and relevant for each view. */
 	void ComputeViewVisibility(FRHICommandListImmediate& RHICmdList, FExclusiveDepthStencil::Type BasePassDepthStencilAccess, FViewVisibleCommandsPerView& ViewCommandsPerView, 
-		FGlobalDynamicIndexBuffer& DynamicIndexBuffer, FGlobalDynamicVertexBuffer& DynamicVertexBuffer, FGlobalDynamicReadBuffer& DynamicReadBuffer, FInstanceCullingManager& InstanceCullingManager);
+		FGlobalDynamicIndexBuffer& DynamicIndexBuffer, FGlobalDynamicVertexBuffer& DynamicVertexBuffer, FGlobalDynamicReadBuffer& DynamicReadBuffer, FInstanceCullingManager& InstanceCullingManager,
+		FVirtualTextureUpdater* VirtualTextureUpdater);
 
 	virtual void ComputeLightVisibility();
 
@@ -2580,7 +2582,7 @@ protected:
 	/** Build visibility lists on CSM receivers and non-csm receivers. */
 	void BuildCSMVisibilityState(FLightSceneInfo* LightSceneInfo);
 
-	void InitViews(FRDGBuilder& GraphBuilder, FSceneTexturesConfig& SceneTexturesConfig, FInstanceCullingManager& InstanceCullingManager);
+	void InitViews(FRDGBuilder& GraphBuilder, FSceneTexturesConfig& SceneTexturesConfig, FInstanceCullingManager& InstanceCullingManager, FVirtualTextureUpdater* VirtualTextureUpdater);
 
 	void RenderPrePass(FRHICommandListImmediate& RHICmdList, const FViewInfo& View);
 	void RenderMaskedPrePass(FRHICommandListImmediate& RHICmdList, const FViewInfo& View);
diff --git a/Engine/Source/Runtime/Renderer/Private/SceneVisibility.cpp b/Engine/Source/Runtime/Renderer/Private/SceneVisibility.cpp
--- a/Engine/Source/Runtime/Renderer/Private/SceneVisibility.cpp
+++ b/Engine/Source/Runtime/Renderer/Private/SceneVisibility.cpp
@@ -54,6 +54,7 @@
 #include "RenderCore.h"
 #include "StaticMeshBatch.h"
 #include "UnrealEngine.h"
+#include "VT/VirtualTextureSystem.h"
 
 #if !UE_BUILD_SHIPPING
 #include "ViewDebug.h"
@@ -4606,7 +4607,8 @@ void FSceneRenderer::ComputeViewVisibility(
 	FGlobalDynamicIndexBuffer& DynamicIndexBuffer,
 	FGlobalDynamicVertexBuffer& DynamicVertexBuffer,
 	FGlobalDynamicReadBuffer& DynamicReadBuffer, 
-	FInstanceCullingManager& InstanceCullingManager)
+	FInstanceCullingManager& InstanceCullingManager,
+	FVirtualTextureUpdater* VirtualTextureUpdater)
 {
 	SCOPE_CYCLE_COUNTER(STAT_ViewVisibilityTime);
 	SCOPED_NAMED_EVENT(FSceneRenderer_ComputeViewVisibility, FColor::Magenta);
@@ -4920,6 +4922,11 @@ void FSceneRenderer::ComputeViewVisibility(
 
 	PreGatherDynamicMeshElements();
 
+	// Sync the virtual texture update task before gathering dynamic mesh elements. Render proxies can register new
+	// materials which require evaluating uniform expression caches, which can contain virtual textures. Newly allocated
+	// virtual textures are processed later.
+	FVirtualTextureSystem::Get().WaitForTasks(VirtualTextureUpdater);
+
 	{
 		SCOPED_NAMED_EVENT(FSceneRenderer_GatherDynamicMeshElements, FColor::Yellow);
 		// Gather FMeshBatches from scene proxies
@@ -5217,7 +5224,7 @@ void FDeferredShadingSceneRenderer::PreVisibilityFrameSetup(FRDGBuilder& GraphBu
  * Initialize scene's views.
  * Check visibility, build visible mesh commands, etc.
  */
-void FDeferredShadingSceneRenderer::BeginInitViews(FRDGBuilder& GraphBuilder, const FSceneTexturesConfig& SceneTexturesConfig, FExclusiveDepthStencil::Type BasePassDepthStencilAccess, struct FILCUpdatePrimTaskData& ILCTaskData, FInstanceCullingManager& InstanceCullingManager)
+void FDeferredShadingSceneRenderer::BeginInitViews(FRDGBuilder& GraphBuilder, const FSceneTexturesConfig& SceneTexturesConfig, FExclusiveDepthStencil::Type BasePassDepthStencilAccess, struct FILCUpdatePrimTaskData& ILCTaskData, FInstanceCullingManager& InstanceCullingManager, FVirtualTextureUpdater* VirtualTextureUpdater)
 {
 	SCOPED_NAMED_EVENT(FDeferredShadingSceneRenderer_InitViews, FColor::Emerald);
 	SCOPE_CYCLE_COUNTER(STAT_InitViewsTime);
@@ -5253,7 +5260,7 @@ void FDeferredShadingSceneRenderer::BeginInitViews(FRDGBuilder& GraphBuilder, co
 	FViewVisibleCommandsPerView ViewCommandsPerView;
 	ViewCommandsPerView.SetNum(Views.Num());
 
-	ComputeViewVisibility(RHICmdList, BasePassDepthStencilAccess, ViewCommandsPerView, DynamicIndexBufferForInitViews, DynamicVertexBufferForInitViews, DynamicReadBufferForInitViews, InstanceCullingManager);
+	ComputeViewVisibility(RHICmdList, BasePassDepthStencilAccess, ViewCommandsPerView, DynamicIndexBufferForInitViews, DynamicVertexBufferForInitViews, DynamicReadBufferForInitViews, InstanceCullingManager, VirtualTextureUpdater);
 
 	// This must happen before we start initialising and using views.
 	if (Scene)
diff --git a/Engine/Source/Runtime/Renderer/Private/VT/UniqueRequestList.h b/Engine/Source/Runtime/Renderer/Private/VT/UniqueRequestList.h
--- a/Engine/Source/Runtime/Renderer/Private/VT/UniqueRequestList.h
+++ b/Engine/Source/Runtime/Renderer/Private/VT/UniqueRequestList.h
@@ -91,6 +91,17 @@ public:
 		ContinuousUpdateRequestHash.Clear();
 	}
 
+	inline void Reset()
+	{
+		Initialize();
+		NumLoadRequests = 0;
+		NumLockRequests = 0;
+		NumMappingRequests = 0;
+		NumDirectMappingRequests = 0;
+		NumContinuousUpdateRequests = 0;
+		NumAdaptiveAllocationRequests = 0;
+	}
+
 	inline uint32 GetNumLoadRequests() const { return NumLoadRequests; }
 	inline uint32 GetNumMappingRequests() const { return NumMappingRequests; }
 	inline uint32 GetNumDirectMappingRequests() const { return NumDirectMappingRequests; }
diff --git a/Engine/Source/Runtime/Renderer/Private/VT/VirtualTextureProducer.cpp b/Engine/Source/Runtime/Renderer/Private/VT/VirtualTextureProducer.cpp
--- a/Engine/Source/Runtime/Renderer/Private/VT/VirtualTextureProducer.cpp
+++ b/Engine/Source/Runtime/Renderer/Private/VT/VirtualTextureProducer.cpp
@@ -229,6 +229,8 @@ uint32 FVirtualTextureProducerCollection::RemoveAllCallbacks(const void* Baton)
 
 void FVirtualTextureProducerCollection::NotifyRequestsCompleted()
 {
+	TRACE_CPUPROFILER_EVENT_SCOPE(FVirtualTextureSystem::NotifyRequestsComplete);
+
 	for (FProducerEntry& Entry : Producers)
 	{
 		if (Entry.Producer.Description.bNotifyCompleted)
diff --git a/Engine/Source/Runtime/Renderer/Private/VT/VirtualTextureSystem.cpp b/Engine/Source/Runtime/Renderer/Private/VT/VirtualTextureSystem.cpp
--- a/Engine/Source/Runtime/Renderer/Private/VT/VirtualTextureSystem.cpp
+++ b/Engine/Source/Runtime/Renderer/Private/VT/VirtualTextureSystem.cpp
@@ -161,7 +161,7 @@ static TAutoConsoleVariable<int32> CVarVTCsvStats(
 );
 static TAutoConsoleVariable<int32> CVarVTAsyncPageRequestTask(
 	TEXT("r.VT.AsyncPageRequestTask"),
-	0,
+	1,
 	TEXT("Performs VT page requests on an async task."),
 	ECVF_RenderThreadSafe
 );
@@ -1257,8 +1257,8 @@ void FVirtualTextureSystem::LoadPendingTiles(FRDGBuilder& GraphBuilder, ERHIFeat
 		GatherRequests(RequestList, UniquePageList, Frame, Allocator, Settings);
 		// No need to sort requests, since we're submitting all of them here (no throttling)
 		AllocateResources(GraphBuilder);
-		BeginSubmitRequests(GraphBuilder.RHICmdList, FeatureLevel, Allocator, Settings, RequestList, false);
-		EndSubmitRequests(GraphBuilder);
+		SubmitRequests(GraphBuilder.RHICmdList, FeatureLevel, Allocator, Settings, RequestList, false);
+		FinalizeRequests(GraphBuilder);
 	}
 }
 
@@ -1913,6 +1913,7 @@ void FVirtualTextureSystem::GetContinuousUpdatesToProduce(FUniqueRequestList con
 
 void FVirtualTextureSystem::UpdateResidencyTracking() const
 {
+	TRACE_CPUPROFILER_EVENT_SCOPE(FVirtualTextureSystem::UpdateResidencyTracking);
 	SCOPE_CYCLE_COUNTER(STAT_ResidencyTracking);
 
 	for (int32 i = 0; i < PhysicalSpaces.Num(); ++i)
@@ -2018,9 +2019,62 @@ void FVirtualTextureSystem::SubmitPreMappedRequests(FRHICommandList& RHICmdList,
 	}
 }
 
-void FVirtualTextureSystem::BeginSubmitRequests(FRHICommandList& RHICmdList, ERHIFeatureLevel::Type FeatureLevel, FConcurrentLinearBulkObjectAllocator& Allocator, FVirtualTextureUpdateSettings const& Settings, FUniqueRequestList* RequestList, bool bAsync)
+void FVirtualTextureSystem::SubmitThrottledRequests(FRHICommandList& RHICmdList, FVirtualTextureUpdater* Updater, bool bContinousUpdates)
+{
+	const FVirtualTextureUpdateSettings& Settings   = Updater->Settings;
+	const ERHIFeatureLevel::Type  FeatureLevel      = Updater->FeatureLevel;
+	FConcurrentLinearBulkObjectAllocator& Allocator = Updater->Allocator;
+	FUniqueRequestList* MergedRequestList           = Updater->MergedRequestList;
+
+	if (MergedRequestList->GetNumAdaptiveAllocationRequests() > 0)
+	{
+		SCOPE_CYCLE_COUNTER(STAT_ProcessRequests_QueueAdaptiveRequests);
+		FAdaptiveVirtualTexture::QueuePackedAllocationRequests(this, &MergedRequestList->GetAdaptiveAllocationRequest(0), MergedRequestList->GetNumAdaptiveAllocationRequests(), Frame);
+	}
+
+	{
+		SCOPE_CYCLE_COUNTER(STAT_ProcessRequests_Sort);
+
+		// Limit the number of uploads (account for MappedTilesToProduce this frame)
+		// Are all pages equal? Should there be different limits on different types of pages?
+		const int32 MaxNumUploads = Settings.MaxPageUploads;
+		const int32 MaxRequestUploads = FMath::Max(MaxNumUploads - MappedTilesToProduce.Num() - (int32)Updater->NumProcessedLoadRequests, 1);
+
+		if (MaxRequestUploads < (int32)MergedRequestList->GetNumLoadRequests())
+		{
+			// Dropping requests is normal but track to log here if we want to tune settings.
+			if (CVarVTVerbose.GetValueOnRenderThread())
+			{
+				UE_LOG(LogConsoleResponse, Display, TEXT("VT dropped %d load requests."), MergedRequestList->GetNumLoadRequests() - MaxRequestUploads);
+			}
+		}
+
+		MergedRequestList->SortRequests(Producers, Allocator, MaxRequestUploads);
+	}
+
+	if (bContinousUpdates)
+	{
+		// After sorting and clamping the load requests, if we still have unused upload bandwidth then use it to add some continous updates
+		const int32 MaxNumUploads = Settings.MaxPageUploads;
+		const int32 MaxTilesToProduce = FMath::Max(MaxNumUploads - MappedTilesToProduce.Num() - (int32)MergedRequestList->GetNumLoadRequests() - (int32)Updater->NumProcessedLoadRequests, 0);
+		const int32 MaxContinuousUpdates = Settings.MaxContinuousUpdates;
+
+		GetContinuousUpdatesToProduce(MergedRequestList, MaxTilesToProduce, MaxContinuousUpdates);
+	}
+
+	// Track total number of requests made this frame.
+	Updater->NumProcessedLoadRequests += MergedRequestList->GetNumLoadRequests();
+
+	// Submit the requests to produce pages that are already mapped
+	SubmitPreMappedRequests(RHICmdList, FeatureLevel);
+
+	// Submit the merged requests
+	SubmitRequests(RHICmdList, FeatureLevel, Allocator, Settings, MergedRequestList, true);
+}
+
+void FVirtualTextureSystem::SubmitRequests(FRHICommandList& RHICmdList, ERHIFeatureLevel::Type FeatureLevel, FConcurrentLinearBulkObjectAllocator& Allocator, FVirtualTextureUpdateSettings const& Settings, FUniqueRequestList* RequestList, bool bAsync)
 {
-	TRACE_CPUPROFILER_EVENT_SCOPE(FVirtualTextureSystem::BeginSubmitRequests);
+	TRACE_CPUPROFILER_EVENT_SCOPE(FVirtualTextureSystem::SubmitRequests);
 	LLM_SCOPE(ELLMTag::VirtualTextureSystem);
 
 	// Allocate space to hold the physical address we allocate for each page load (1 page per layer per request)
@@ -2224,7 +2278,7 @@ void FVirtualTextureSystem::BeginSubmitRequests(FRHICommandList& RHICmdList, ERH
 			if (bLockTile && !bTileLoaded && !bTileInvalid)
 			{
 				// Want to lock this tile, but didn't manage to load it this frame, add it back to the list to try the lock again next frame
-				TilesToLock.Add(TileToLoad);
+				TilesToLockForNextFrame.Add(TileToLoad);
 			}
 		}
 
@@ -2327,9 +2381,9 @@ void FVirtualTextureSystem::BeginSubmitRequests(FRHICommandList& RHICmdList, ERH
 	}
 }
 
-void FVirtualTextureSystem::EndSubmitRequests(FRDGBuilder& GraphBuilder)
+void FVirtualTextureSystem::FinalizeRequests(FRDGBuilder& GraphBuilder)
 {
-	TRACE_CPUPROFILER_EVENT_SCOPE(FVirtualTextureSystem::EndSubmitRequests);
+	TRACE_CPUPROFILER_EVENT_SCOPE(FVirtualTextureSystem::FinalizeRequests);
 
 	// Finalize requests
 	{
@@ -2377,6 +2431,214 @@ void FVirtualTextureSystem::AllocateResources(FRDGBuilder& GraphBuilder)
 	}
 }
 
+void FVirtualTextureSystem::GatherFeedbackRequests(FConcurrentLinearBulkObjectAllocator& Allocator, const FVirtualTextureUpdateSettings& Settings, const FVirtualTextureFeedback::FMapResult& FeedbackResult, FUniqueRequestList* MergedRequestList)
+{
+	FUniquePageList* MergedUniquePageList = Allocator.Create<FUniquePageList>();
+	MergedUniquePageList->Initialize();
+
+	if (Settings.bEnableFeedback)
+	{
+		// Create tasks to read the feedback data
+		// Give each task a section of the feedback buffer to analyze
+		FFeedbackAnalysisParameters FeedbackAnalysisParameters[MaxNumTasks];
+
+		const uint32 MaxNumFeedbackTasks = FMath::Clamp((uint32)Settings.NumFeedbackTasks, 1u, MaxNumTasks);
+		const uint32 FeedbackSizePerTask = FMath::DivideAndRoundUp(FeedbackResult.Size, MaxNumFeedbackTasks);
+
+		uint32 NumFeedbackTasks = 0;
+		uint32 CurrentOffset = 0;
+		while (CurrentOffset < FeedbackResult.Size)
+		{
+			const uint32 TaskIndex = NumFeedbackTasks++;
+			FFeedbackAnalysisParameters& Params = FeedbackAnalysisParameters[TaskIndex];
+			Params.System = this;
+			if (TaskIndex == 0u)
+			{
+				Params.UniquePageList = MergedUniquePageList;
+			}
+			else
+			{
+				Params.UniquePageList = Allocator.Create<FUniquePageList>();
+			}
+			Params.FeedbackBuffer = FeedbackResult.Data + CurrentOffset;
+
+			const uint32 Size = FMath::Min(FeedbackSizePerTask, FeedbackResult.Size - CurrentOffset);
+			Params.FeedbackSize = Size;
+			CurrentOffset += Size;
+		}
+
+		// Kick the tasks
+		const int32 LocalFeedbackTaskCount = Settings.bParallelFeedbackTasks ? 1 : NumFeedbackTasks;
+		const int32 WorkerFeedbackTaskCount = NumFeedbackTasks - LocalFeedbackTaskCount;
+
+		FGraphEventArray Tasks;
+		if (WorkerFeedbackTaskCount > 0)
+		{
+			SCOPE_CYCLE_COUNTER(STAT_ProcessRequests_SubmitTasks);
+			Tasks.Reserve(WorkerFeedbackTaskCount);
+			for (uint32 TaskIndex = LocalFeedbackTaskCount; TaskIndex < NumFeedbackTasks; ++TaskIndex)
+			{
+				Tasks.Add(TGraphTask<FFeedbackAnalysisTask>::CreateTask().ConstructAndDispatchWhenReady(FeedbackAnalysisParameters[TaskIndex]));
+			}
+		}
+
+		if (NumFeedbackTasks > 0u)
+		{
+			SCOPE_CYCLE_COUNTER(STAT_FeedbackAnalysis);
+
+			for (int32 TaskIndex = 0; TaskIndex < LocalFeedbackTaskCount; ++TaskIndex)
+			{
+				FFeedbackAnalysisTask::DoTask(FeedbackAnalysisParameters[TaskIndex]);
+			}
+			if (WorkerFeedbackTaskCount > 0)
+			{
+				SCOPE_CYCLE_COUNTER(STAT_ProcessRequests_WaitTasks);
+
+				FTaskGraphInterface::Get().WaitUntilTasksComplete(Tasks, ENamedThreads::GetRenderThread_Local());
+			}
+		}
+
+		if (NumFeedbackTasks > 1u)
+		{
+			SCOPE_CYCLE_COUNTER(STAT_ProcessRequests_MergePages);
+			for (uint32 TaskIndex = 1u; TaskIndex < NumFeedbackTasks; ++TaskIndex)
+			{
+				MergedUniquePageList->MergePages(FeedbackAnalysisParameters[TaskIndex].UniquePageList);
+			}
+		}
+	}
+
+#if WITH_EDITOR
+	// If we're are recording page requests, then copy off pages to the recording buffer.
+	if (PageRequestRecordHandle != ~0ull)
+	{
+		RecordPageRequests(MergedUniquePageList, PageRequestRecordBuffer);
+	}
+#endif
+
+	// Add any page requests from recording playback.
+	if (PageRequestPlaybackBuffer.Num() > 0)
+	{
+		if (Settings.bEnablePlayback)
+		{
+			// todo: We can split this into concurrent tasks. 
+			FAddRequestedTilesParameters Parameters;
+			Parameters.System = this;
+			Parameters.LevelBias = FMath::FloorToInt(CVarVTPlaybackMipBias.GetValueOnRenderThread() + GetGlobalMipBias() + 0.5f);
+			Parameters.RequestBuffer = PageRequestPlaybackBuffer.GetData();
+			Parameters.NumRequests = PageRequestPlaybackBuffer.Num();
+			Parameters.UniquePageList = Allocator.Create<FUniquePageList>();
+
+			FAddRequestedTilesTask::DoTask(Parameters);
+			MergedUniquePageList->MergePages(Parameters.UniquePageList);
+		}
+
+		PageRequestPlaybackBuffer.Reset(0);
+	}
+
+	// Pages from feedback buffer were generated several frames ago, so they may no longer be valid for newly allocated VTs
+	static uint32 PendingFrameDelay = 3u;
+	if (Frame >= PendingFrameDelay)
+	{
+		GatherRequests(MergedRequestList, MergedUniquePageList, Frame - PendingFrameDelay, Allocator, Settings);
+	}
+}
+
+void FVirtualTextureSystem::GatherLockedTileRequests(FUniqueRequestList* MergedRequestList)
+{
+	for (const FVirtualTextureLocalTile& Tile : TilesToLock)
+	{
+		const FVirtualTextureProducerHandle ProducerHandle = Tile.GetProducerHandle();
+		const FVirtualTextureProducer* Producer = Producers.FindProducer(ProducerHandle);
+		checkSlow(TileLocks.IsLocked(Tile));
+		if (Producer)
+		{
+			uint8 ProducerLayerMaskToLoad = 0u;
+			for (uint32 ProducerLayerIndex = 0u; ProducerLayerIndex < Producer->GetNumTextureLayers(); ++ProducerLayerIndex)
+			{
+				uint32 GroupIndex = Producer->GetPhysicalGroupIndexForTextureLayer(ProducerLayerIndex);
+				FVirtualTexturePhysicalSpace* PhysicalSpace = Producer->GetPhysicalSpaceForPhysicalGroup(GroupIndex);
+				FTexturePagePool& PagePool = PhysicalSpace->GetPagePool();
+				const uint32 pAddress = PagePool.FindPageAddress(ProducerHandle, GroupIndex, Tile.Local_vAddress, Tile.Local_vLevel);
+				if (pAddress == ~0u)
+				{
+					ProducerLayerMaskToLoad |= (1u << ProducerLayerIndex);
+				}
+				else
+				{
+					PagePool.Lock(pAddress);
+				}
+			}
+
+			if (ProducerLayerMaskToLoad != 0u)
+			{
+				const uint16 LoadRequestIndex = MergedRequestList->LockLoadRequest(FVirtualTextureLocalTile(Tile.GetProducerHandle(), Tile.Local_vAddress, Tile.Local_vLevel), ProducerLayerMaskToLoad);
+				if (LoadRequestIndex == 0xffff)
+				{
+					// Overflowed the request list...try to lock the tile again next frame
+					TilesToLockForNextFrame.Add(Tile);
+				}
+			}
+		}
+	}
+	TilesToLock.Reset();
+}
+
+void FVirtualTextureSystem::GatherPackedTileRequests(FConcurrentLinearBulkObjectAllocator& Allocator, const FVirtualTextureUpdateSettings& Settings, FUniqueRequestList* MergedRequestList)
+{
+	TArray<uint32> PackedTiles;
+	if (RequestedPackedTiles.Num() > 0)
+	{
+		FScopeLock Lock(&RequestedTilesLock);
+		PackedTiles = MoveTemp(RequestedPackedTiles);
+		RequestedPackedTiles.Reset();
+	}
+
+	if (PackedTiles.Num() > 0)
+	{
+		// Collect explicitly requested tiles
+		// These tiles are generated on the current frame, so they are collected/processed in a separate list
+		FUniquePageList* RequestedPageList = Allocator.Create<FUniquePageList>();
+		RequestedPageList->Initialize();
+		for (uint32 Tile : PackedTiles)
+		{
+			RequestedPageList->Add(Tile, 0xffff);
+		}
+		GatherRequests(MergedRequestList, RequestedPageList, Frame, Allocator, Settings);
+	}
+}
+
+void FVirtualTextureSystem::BeginUpdate(FRDGBuilder& GraphBuilder, FVirtualTextureUpdater* Updater)
+{
+	// Mark updating to true now that we are potentially launching async tasks.
+	bUpdating = true;
+
+	if (Updater->Settings.bEnableFeedback)
+	{
+		SCOPE_CYCLE_COUNTER(STAT_FeedbackMap);
+		Updater->FeedbackMapResult = GVirtualTextureFeedback.Map(GraphBuilder.RHICmdList);
+	}
+
+	Updater->AsyncTask = GraphBuilder.AddCommandListSetupTask([this, Updater, &Allocator = Updater->Allocator, Settings = Updater->Settings, FeedbackResult = Updater->FeedbackMapResult](FRHICommandList& RHICmdList)
+	{
+		TRACE_CPUPROFILER_EVENT_SCOPE(FVirtualTextureSystem::GatherAndSubmitRequests);
+
+		const bool bContinousUpdates = false;
+
+		Updater->MergedRequestList = Allocator.Create<FUniqueRequestList>(Allocator);
+		Updater->MergedRequestList->Initialize();
+
+		GatherFeedbackRequests(Allocator, Settings, FeedbackResult, Updater->MergedRequestList);
+		GatherLockedTileRequests(Updater->MergedRequestList);
+		GatherPackedTileRequests(Allocator, Settings, Updater->MergedRequestList);
+		SubmitThrottledRequests(RHICmdList, Updater, bContinousUpdates);
+
+		// Reset the request list for the gather in EndUpdate.
+		Updater->MergedRequestList->Reset();
+
+	}, Updater->bAsyncTaskAllowed);
+}
+
 TUniquePtr<FVirtualTextureUpdater> FVirtualTextureSystem::BeginUpdate(FRDGBuilder& GraphBuilder, ERHIFeatureLevel::Type FeatureLevel, FScene* Scene, const FVirtualTextureUpdateSettings& Settings)
 {
 	check(IsInRenderingThread());
@@ -2460,268 +2722,74 @@ TUniquePtr<FVirtualTextureUpdater> FVirtualTextureSystem::BeginUpdate(FRDGBuilde
 		}
 	}
 
-	// Mark updating to true now that we are potentially launching async tasks.
-	bUpdating = true;
-
 	TUniquePtr<FVirtualTextureUpdater> Updater(new FVirtualTextureUpdater());
+	Updater->Settings = Settings;
+	Updater->FeatureLevel = FeatureLevel;
+	Updater->bAsyncTaskAllowed = Settings.bEnableAsyncTasks && CVarVTAsyncPageRequestTask.GetValueOnRenderThread();
 
-	if (Settings.bEnableFeedback)
+	if (Updater->bAsyncTaskAllowed)
 	{
-		SCOPE_CYCLE_COUNTER(STAT_FeedbackMap);
-		Updater->FeedbackMapResult = GVirtualTextureFeedback.Map(GraphBuilder.RHICmdList);
+		BeginUpdate(GraphBuilder, Updater.Get());
 	}
 
-	Updater->AsyncTask = GraphBuilder.AddCommandListSetupTask([this, Settings, FeedbackResult = Updater->FeedbackMapResult, FeatureLevel] (FRHICommandList& RHICmdList)
-	{
-		TRACE_CPUPROFILER_EVENT_SCOPE(FVirtualTextureSystem::ProcessPageRequests);
-
-		FConcurrentLinearBulkObjectAllocator Allocator;
-		FUniqueRequestList* MergedRequestList = Allocator.Create<FUniqueRequestList>(Allocator);
-
-		MergedRequestList->Initialize();
-
-		FUniquePageList* MergedUniquePageList = Allocator.Create<FUniquePageList>();
-		MergedUniquePageList->Initialize();
-
-		if (Settings.bEnableFeedback)
-		{
-			// Create tasks to read the feedback data
-			// Give each task a section of the feedback buffer to analyze
-			FFeedbackAnalysisParameters FeedbackAnalysisParameters[MaxNumTasks];
-
-			const uint32 MaxNumFeedbackTasks = FMath::Clamp((uint32)Settings.NumFeedbackTasks, 1u, MaxNumTasks);
-			const uint32 FeedbackSizePerTask = FMath::DivideAndRoundUp(FeedbackResult.Size, MaxNumFeedbackTasks);
-
-			uint32 NumFeedbackTasks = 0;
-			uint32 CurrentOffset = 0;
-			while (CurrentOffset < FeedbackResult.Size)
-			{
-				const uint32 TaskIndex = NumFeedbackTasks++;
-				FFeedbackAnalysisParameters& Params = FeedbackAnalysisParameters[TaskIndex];
-				Params.System = this;
-				if (TaskIndex == 0u)
-				{
-					Params.UniquePageList = MergedUniquePageList;
-				}
-				else
-				{
-					Params.UniquePageList = Allocator.Create<FUniquePageList>();
-				}
-				Params.FeedbackBuffer = FeedbackResult.Data + CurrentOffset;
-
-				const uint32 Size = FMath::Min(FeedbackSizePerTask, FeedbackResult.Size - CurrentOffset);
-				Params.FeedbackSize = Size;
-				CurrentOffset += Size;
-			}
-
-			// Kick the tasks
-			const int32 LocalFeedbackTaskCount = Settings.bParallelFeedbackTasks ? 1 : NumFeedbackTasks;
-			const int32 WorkerFeedbackTaskCount = NumFeedbackTasks - LocalFeedbackTaskCount;
-
-			FGraphEventArray Tasks;
-			if (WorkerFeedbackTaskCount > 0)
-			{
-				SCOPE_CYCLE_COUNTER(STAT_ProcessRequests_SubmitTasks);
-				Tasks.Reserve(WorkerFeedbackTaskCount);
-				for (uint32 TaskIndex = LocalFeedbackTaskCount; TaskIndex < NumFeedbackTasks; ++TaskIndex)
-				{
-					Tasks.Add(TGraphTask<FFeedbackAnalysisTask>::CreateTask().ConstructAndDispatchWhenReady(FeedbackAnalysisParameters[TaskIndex]));
-				}
-			}
-
-			if (NumFeedbackTasks > 0u)
-			{
-				SCOPE_CYCLE_COUNTER(STAT_FeedbackAnalysis);
-
-				for (int32 TaskIndex = 0; TaskIndex < LocalFeedbackTaskCount; ++TaskIndex)
-				{
-					FFeedbackAnalysisTask::DoTask(FeedbackAnalysisParameters[TaskIndex]);
-				}
-				if (WorkerFeedbackTaskCount > 0)
-				{
-					SCOPE_CYCLE_COUNTER(STAT_ProcessRequests_WaitTasks);
-
-					FTaskGraphInterface::Get().WaitUntilTasksComplete(Tasks, ENamedThreads::GetRenderThread_Local());
-				}
-			}
-
-			if (NumFeedbackTasks > 1u)
-			{
-				SCOPE_CYCLE_COUNTER(STAT_ProcessRequests_MergePages);
-				for (uint32 TaskIndex = 1u; TaskIndex < NumFeedbackTasks; ++TaskIndex)
-				{
-					MergedUniquePageList->MergePages(FeedbackAnalysisParameters[TaskIndex].UniquePageList);
-				}
-			}
-		}
-
-#if WITH_EDITOR
-		// If we're are recording page requests, then copy off pages to the recording buffer.
-		if (PageRequestRecordHandle != ~0ull)
-		{
-			RecordPageRequests(MergedUniquePageList, PageRequestRecordBuffer);
-		}
-#endif
-
-		// Add any page requests from recording playback.
-		if (PageRequestPlaybackBuffer.Num() > 0)
-		{
-			if (Settings.bEnablePlayback)
-			{
-				// todo: We can split this into concurrent tasks. 
-				FAddRequestedTilesParameters Parameters;
-				Parameters.System = this;
-				Parameters.LevelBias = FMath::FloorToInt(CVarVTPlaybackMipBias.GetValueOnRenderThread() + GetGlobalMipBias() + 0.5f);
-				Parameters.RequestBuffer = PageRequestPlaybackBuffer.GetData();
-				Parameters.NumRequests = PageRequestPlaybackBuffer.Num();
-				Parameters.UniquePageList = Allocator.Create<FUniquePageList>();
-
-				FAddRequestedTilesTask::DoTask(Parameters);
-				MergedUniquePageList->MergePages(Parameters.UniquePageList);
-			}
-
-			PageRequestPlaybackBuffer.Reset(0);
-		}
-
-		// Collect tiles to lock
-		{
-			TArray<FVirtualTextureLocalTile> RemainingTilesToLock;
-			for (const FVirtualTextureLocalTile& Tile : TilesToLock)
-			{
-				const FVirtualTextureProducerHandle ProducerHandle = Tile.GetProducerHandle();
-				const FVirtualTextureProducer* Producer = Producers.FindProducer(ProducerHandle);
-				checkSlow(TileLocks.IsLocked(Tile));
-				if (Producer)
-				{
-					uint8 ProducerLayerMaskToLoad = 0u;
-					for (uint32 ProducerLayerIndex = 0u; ProducerLayerIndex < Producer->GetNumTextureLayers(); ++ProducerLayerIndex)
-					{
-						uint32 GroupIndex = Producer->GetPhysicalGroupIndexForTextureLayer(ProducerLayerIndex);
-						FVirtualTexturePhysicalSpace* PhysicalSpace = Producer->GetPhysicalSpaceForPhysicalGroup(GroupIndex);
-						FTexturePagePool& PagePool = PhysicalSpace->GetPagePool();
-						const uint32 pAddress = PagePool.FindPageAddress(ProducerHandle, GroupIndex, Tile.Local_vAddress, Tile.Local_vLevel);
-						if (pAddress == ~0u)
-						{
-							ProducerLayerMaskToLoad |= (1u << ProducerLayerIndex);
-						}
-						else
-						{
-							PagePool.Lock(pAddress);
-						}
-					}
-
-					if (ProducerLayerMaskToLoad != 0u)
-					{
-						const uint16 LoadRequestIndex = MergedRequestList->LockLoadRequest(FVirtualTextureLocalTile(Tile.GetProducerHandle(), Tile.Local_vAddress, Tile.Local_vLevel), ProducerLayerMaskToLoad);
-						if (LoadRequestIndex == 0xffff)
-						{
-							// Overflowed the request list...try to lock the tile again next frame
-							RemainingTilesToLock.Add(Tile);
-						}
-					}
-				}
-			}
-
-			TilesToLock = MoveTemp(RemainingTilesToLock);
-		}
-
-		TArray<uint32> PackedTiles;
-		if (RequestedPackedTiles.Num() > 0)
-		{
-			FScopeLock Lock(&RequestedTilesLock);
-			PackedTiles = MoveTemp(RequestedPackedTiles);
-			RequestedPackedTiles.Reset();
-		}
-
-		if (PackedTiles.Num() > 0)
-		{
-			// Collect explicitly requested tiles
-			// These tiles are generated on the current frame, so they are collected/processed in a separate list
-			FUniquePageList* RequestedPageList = Allocator.Create<FUniquePageList>();
-			RequestedPageList->Initialize();
-			for (uint32 Tile : PackedTiles)
-			{
-				RequestedPageList->Add(Tile, 0xffff);
-			}
-			GatherRequests(MergedRequestList, RequestedPageList, Frame, Allocator, Settings);
-		}
-
-		// Pages from feedback buffer were generated several frames ago, so they may no longer be valid for newly allocated VTs
-		static uint32 PendingFrameDelay = 3u;
-		if (Frame >= PendingFrameDelay)
-		{
-			GatherRequests(MergedRequestList, MergedUniquePageList, Frame - PendingFrameDelay, Allocator, Settings);
-		}
-
-		if (MergedRequestList->GetNumAdaptiveAllocationRequests() > 0)
-		{
-			SCOPE_CYCLE_COUNTER(STAT_ProcessRequests_QueueAdaptiveRequests);
-			FAdaptiveVirtualTexture::QueuePackedAllocationRequests(this, &MergedRequestList->GetAdaptiveAllocationRequest(0), MergedRequestList->GetNumAdaptiveAllocationRequests(), Frame);
-		}
-
-		{
-			SCOPE_CYCLE_COUNTER(STAT_ProcessRequests_Sort);
-
-			// Limit the number of uploads (account for MappedTilesToProduce this frame)
-			// Are all pages equal? Should there be different limits on different types of pages?
-			const int32 MaxNumUploads = Settings.MaxPageUploads;
-			const int32 MaxRequestUploads = FMath::Max(MaxNumUploads - MappedTilesToProduce.Num(), 1);
-
-			if (MaxRequestUploads < (int32)MergedRequestList->GetNumLoadRequests())
-			{
-				// Dropping requests is normal but track to log here if we want to tune settings.
-				if (CVarVTVerbose.GetValueOnRenderThread())
-				{
-					UE_LOG(LogConsoleResponse, Display, TEXT("VT dropped %d load requests."), MergedRequestList->GetNumLoadRequests() - MaxRequestUploads);
-				}
-			}
-
-			MergedRequestList->SortRequests(Producers, Allocator, MaxRequestUploads);
-		}
-
-		{
-			// After sorting and clamping the load requests, if we still have unused upload bandwidth then use it to add some continous updates
-			const int32 MaxNumUploads = Settings.MaxPageUploads;
-			const int32 MaxTilesToProduce = FMath::Max(MaxNumUploads - MappedTilesToProduce.Num() - (int32)MergedRequestList->GetNumLoadRequests(), 0);
-			const int32 MaxContinuousUpdates = Settings.MaxContinuousUpdates;
-
-			GetContinuousUpdatesToProduce(MergedRequestList, MaxTilesToProduce, MaxContinuousUpdates);
-		}
-
-		// Submit the requests to produce pages that are already mapped
-		SubmitPreMappedRequests(RHICmdList, FeatureLevel);
-
-		// Submit the merged requests
-		BeginSubmitRequests(RHICmdList, FeatureLevel, Allocator, Settings, MergedRequestList, true);
+	return MoveTemp(Updater);
+}
 
-	}, Settings.bEnableAsyncTasks && CVarVTAsyncPageRequestTask.GetValueOnRenderThread());
+void FVirtualTextureSystem::WaitForTasks(FVirtualTextureUpdater* Updater)
+{
+	if (!Updater)
+	{
+		return;
+	}
 
-	return Updater;
+	Updater->AsyncTask.Wait();
+	bUpdating = false;
 }
 
 void FVirtualTextureSystem::EndUpdate(FRDGBuilder& GraphBuilder, TUniquePtr<FVirtualTextureUpdater>&& Updater, ERHIFeatureLevel::Type FeatureLevel)
 {
+	check(IsInRenderingThread());
+
 	if (!Updater)
 	{
 		return;
 	}
 
-	check(bUpdating);
-	check(IsInRenderingThread());
-
 	RDG_CSV_STAT_EXCLUSIVE_SCOPE(GraphBuilder, VirtualTextureSystem_Update);
 	TRACE_CPUPROFILER_EVENT_SCOPE(FVirtualTextureSystem::EndUpdate);
 	SCOPE_CYCLE_COUNTER(STAT_VirtualTextureSystem_Update);
 	RDG_GPU_STAT_SCOPE(GraphBuilder, VirtualTexture);
 
-	Updater->AsyncTask.Wait();
+	if (Updater->bAsyncTaskAllowed)
+	{
+		Updater->AsyncTask.Wait();
+		bUpdating = true;
+	}
+	else
+	{
+		BeginUpdate(GraphBuilder, Updater.Get());
+	}
 
 	if (Updater->FeedbackMapResult.Data)
 	{
 		GVirtualTextureFeedback.Unmap(GraphBuilder.RHICmdList, Updater->FeedbackMapResult.MapHandle);
 	}
 
-	EndSubmitRequests(GraphBuilder);
+	{
+		TRACE_CPUPROFILER_EVENT_SCOPE(FVirtualTextureSystem::GatherAndSubmitRequests);
+		const bool bContinousUpdates = true;
+
+		// We only need to gather remaining requests if an async task was used early in the frame.
+		if (Updater->bAsyncTaskAllowed)
+		{
+			GatherLockedTileRequests(Updater->MergedRequestList);
+			GatherPackedTileRequests(Updater->Allocator, Updater->Settings, Updater->MergedRequestList);
+		}
+
+		SubmitThrottledRequests(GraphBuilder.RHICmdList, Updater.Get(), bContinousUpdates);
+	}
+
+	FinalizeRequests(GraphBuilder);
 
 	Producers.NotifyRequestsCompleted();
 
@@ -2733,6 +2801,7 @@ void FVirtualTextureSystem::EndUpdate(FRDGBuilder& GraphBuilder, TUniquePtr<FVir
 
 	ReleasePendingSpaces();
 
+	TilesToLock = MoveTemp(TilesToLockForNextFrame);
 	bUpdating = false;
 }
 
diff --git a/Engine/Source/Runtime/Renderer/Private/VT/VirtualTextureSystem.h b/Engine/Source/Runtime/Renderer/Private/VT/VirtualTextureSystem.h
--- a/Engine/Source/Runtime/Renderer/Private/VT/VirtualTextureSystem.h
+++ b/Engine/Source/Runtime/Renderer/Private/VT/VirtualTextureSystem.h
@@ -80,10 +80,14 @@ struct FVirtualTextureUpdateSettings
 class FVirtualTextureUpdater
 {
 private:
-	FVirtualTextureUpdater() = default;
-
+	FVirtualTextureUpdateSettings Settings;
+	FConcurrentLinearBulkObjectAllocator Allocator;
+	FUniqueRequestList* MergedRequestList = nullptr;
 	FVirtualTextureFeedback::FMapResult FeedbackMapResult;
 	UE::Tasks::FTask AsyncTask;
+	ERHIFeatureLevel::Type FeatureLevel = ERHIFeatureLevel::Num;
+	uint32 NumProcessedLoadRequests = 0;
+	bool bAsyncTaskAllowed = false;
 
 	friend FVirtualTextureSystem;
 };
@@ -100,6 +104,7 @@ public:
 	void Update(FRDGBuilder& GraphBuilder, ERHIFeatureLevel::Type FeatureLevel, FScene* Scene, const FVirtualTextureUpdateSettings& Settings);
 
 	TUniquePtr<FVirtualTextureUpdater> BeginUpdate(FRDGBuilder& GraphBuilder, ERHIFeatureLevel::Type FeatureLevel, FScene* Scene, const FVirtualTextureUpdateSettings& Settings);
+	void WaitForTasks(FVirtualTextureUpdater* Updater);
 	void EndUpdate(FRDGBuilder& GraphBuilder, TUniquePtr<FVirtualTextureUpdater>&& Updater, ERHIFeatureLevel::Type FeatureLevel);
 
 	void ReleasePendingResources();
@@ -154,6 +159,8 @@ private:
 	FVirtualTextureSystem();
 	~FVirtualTextureSystem();
 
+	void BeginUpdate(FRDGBuilder& GraphBuilder, FVirtualTextureUpdater* Updater);
+
 	void AllocateResources(FRDGBuilder& GraphBuilder);
 	void DestroyPendingVirtualTextures(bool bForceDestroyAll);
 	void ReleasePendingSpaces();
@@ -164,10 +171,15 @@ private:
 	
 	void SubmitRequestsFromLocalTileList(FRHICommandList& RHICmdList, TArray<FVirtualTextureLocalTile>& OutDeferredTiles, const TSet<FVirtualTextureLocalTile>& LocalTileList, EVTProducePageFlags Flags, ERHIFeatureLevel::Type FeatureLevel);
 
+	void GatherFeedbackRequests(FConcurrentLinearBulkObjectAllocator& Allocator, const FVirtualTextureUpdateSettings& Settings, const FVirtualTextureFeedback::FMapResult& FeedbackResult, FUniqueRequestList* MergedRequestList);
+	void GatherLockedTileRequests(FUniqueRequestList* MergedRequestList);
+	void GatherPackedTileRequests(FConcurrentLinearBulkObjectAllocator& Allocator, const FVirtualTextureUpdateSettings& Settings, FUniqueRequestList* MergedRequestList);
+
 	void SubmitPreMappedRequests(FRHICommandList& RHICmdList, ERHIFeatureLevel::Type FeatureLevel);
 
-	void BeginSubmitRequests(FRHICommandList& RHICmdList, ERHIFeatureLevel::Type FeatureLevel, FConcurrentLinearBulkObjectAllocator& Allocator, FVirtualTextureUpdateSettings const& Settings, FUniqueRequestList* RequestList, bool bAsync);
-	void EndSubmitRequests(FRDGBuilder& GraphBuilder);
+	void SubmitThrottledRequests(FRHICommandList& RHICmdList, FVirtualTextureUpdater* Updater, bool bContinuousUpdates);
+	void SubmitRequests(FRHICommandList& RHICmdList, ERHIFeatureLevel::Type FeatureLevel, FConcurrentLinearBulkObjectAllocator& Allocator, FVirtualTextureUpdateSettings const& Settings, FUniqueRequestList* RequestList, bool bAsync);
+	void FinalizeRequests(FRDGBuilder& GraphBuilder);
 
 	void GatherRequests(FUniqueRequestList* MergedRequestList, const FUniquePageList* UniquePageList, uint32 FrameRequested, FConcurrentLinearBulkObjectAllocator& Allocator, FVirtualTextureUpdateSettings const& Settings);
 
@@ -226,6 +238,7 @@ private:
 	TArray<uint32> RequestedPackedTiles;
 
 	TArray<FVirtualTextureLocalTile> TilesToLock;
+	TArray<FVirtualTextureLocalTile> TilesToLockForNextFrame;
 	FTexturePageLocks TileLocks;
 
 	TSet<FVirtualTextureLocalTile> ContinuousUpdateTilesToProduce;
