Change InstancedStereoWidth to a rect.

diff --git a/Engine/Source/Runtime/Renderer/Private/SceneRendering.cpp b/Engine/Source/Runtime/Renderer/Private/SceneRendering.cpp
--- a/Engine/Source/Runtime/Renderer/Private/SceneRendering.cpp
+++ b/Engine/Source/Runtime/Renderer/Private/SceneRendering.cpp
@@ -2063,10 +2063,8 @@ void FViewInfo::CreateViewUniformBuffers(const FViewUniformShaderParameters& Par
 			InstancedViewParametersUtils::CopyIntoInstancedViewParameters(LocalInstancedViewUniformShaderParameters, Params, 1);
 		}
 
-			InstancedViewUniformBuffer = TUniformBufferRef<FInstancedViewUniformShaderParameters>::CreateUniformBufferImmediate(
-			LocalInstancedViewUniformShaderParameters,
-				UniformBuffer_SingleFrame);
-		}
+		InstancedViewUniformBuffer = TUniformBufferRef<FInstancedViewUniformShaderParameters>::CreateUniformBufferImmediate(LocalInstancedViewUniformShaderParameters, UniformBuffer_SingleFrame);
+	}
 }
 
 extern TSet<IPersistentViewUniformBufferExtension*> PersistentViewUniformBufferExtensions;
@@ -3216,8 +3214,12 @@ void FSceneRenderer::ComputeFamilySize()
 		MaxFamilyX = FMath::Max(MaxFamilyX, FinalViewMaxX);
 		MaxFamilyY = FMath::Max(MaxFamilyY, FinalViewMaxY);
 
-		const FViewInfo* InstancedView = View.GetInstancedView();
-		View.InstancedStereoWidth = InstancedView ? InstancedView->ViewRect.Max.X : View.ViewRect.Max.X;
+		View.ViewRectWithSecondaryViews = View.ViewRect;
+		for (const FSceneView* SecondaryView : View.GetSecondaryViews())
+		{
+			const FViewInfo& InstancedView = static_cast<const FViewInfo&>(*SecondaryView);
+			View.ViewRectWithSecondaryViews.Union(InstancedView.ViewRect);
+		}
 	}
 
 	// We render to the actual position of the viewports so with black borders we need the max.
@@ -5189,7 +5191,7 @@ void FSceneRenderer::SetStereoViewport(FRHICommandList& RHICmdList, const FViewI
 		}
 		else
 		{
-			RHICmdList.SetViewport(View.ViewRect.Min.X * ViewportScale, View.ViewRect.Min.Y * ViewportScale, 0.0f, View.InstancedStereoWidth * ViewportScale, View.ViewRect.Max.Y * ViewportScale, 1.0f);
+			RHICmdList.SetViewport(View.ViewRectWithSecondaryViews.Min.X * ViewportScale, View.ViewRectWithSecondaryViews.Min.Y * ViewportScale, 0.0f, View.ViewRectWithSecondaryViews.Max.X * ViewportScale, View.ViewRectWithSecondaryViews.Max.Y * ViewportScale, 1.0f);
 		}
 	}
 	else
@@ -5313,8 +5315,7 @@ void AddResolveSceneColorPass(FRDGBuilder& GraphBuilder, const FViewInfo& View,
 
 		// Resolve views individually. In the case of adaptive resolution, the view family will be much larger than the views individually.
 		RHICmdList.SetViewport(0.0f, 0.0f, 0.0f, SceneColorExtent.X, SceneColorExtent.Y, 1.0f);
-		RHICmdList.SetScissorRect(true, View.IsInstancedStereoPass() ? 0 : View.ViewRect.Min.X, View.ViewRect.Min.Y,
-			View.IsInstancedStereoPass() ? View.InstancedStereoWidth : View.ViewRect.Max.X, View.ViewRect.Max.Y);
+		RHICmdList.SetScissorRect(true, View.ViewRectWithSecondaryViews.Min.X, View.ViewRectWithSecondaryViews.Min.Y, View.ViewRectWithSecondaryViews.Max.X, View.ViewRectWithSecondaryViews.Max.Y);
 
 		int32 ResolveWidth = CVarWideCustomResolve.GetValueOnRenderThread();
 
@@ -5447,13 +5448,7 @@ void AddResolveSceneDepthPass(FRDGBuilder& GraphBuilder, const FViewInfo& View,
 		return;
 	}
 
-	FResolveRect ResolveRect(View.ViewRect);
-	if (View.IsInstancedStereoPass())
-	{
-		ResolveRect.X1 = 0;
-		ResolveRect.X2 = View.InstancedStereoWidth;
-	}
-
+	FResolveRect ResolveRect(View.ViewRectWithSecondaryViews);
 	const FIntPoint DepthExtent = SceneDepth.Resolve->Desc.Extent;
 
 	FResolveSceneDepthParameters* PassParameters = GraphBuilder.AllocParameters<FResolveSceneDepthParameters>();
diff --git a/Engine/Source/Runtime/Renderer/Private/SceneRendering.h b/Engine/Source/Runtime/Renderer/Private/SceneRendering.h
--- a/Engine/Source/Runtime/Renderer/Private/SceneRendering.h
+++ b/Engine/Source/Runtime/Renderer/Private/SceneRendering.h
@@ -1642,12 +1642,12 @@ public:
 	 */
 	int32 GPUSceneViewId;
 
-	/**
-	 * The width in pixels of the ISR pair being rendered. This may be different than FamilySizeX if
-	 * we're using adaptive resolution stereo rendering. In that case, FamilySizeX represents the maximum size of
+	/* View rect for all instanced views laid out side-by-side. Only primary view will have it populated. 
+	 *
+	 * This may be different than FamilySize if we're using adaptive resolution stereo rendering. In that case, FamilySize represents the maximum size of
 	 * the family to ensure the backing render targets don't change between frames as the view size varies.
 	 */
-	uint32 InstancedStereoWidth = 0;
+	FIntRect ViewRectWithSecondaryViews;
 
 #if WITH_EDITOR
 	TArray<uint32> EditorVisualizeLevelInstanceIds;
diff --git a/Engine/Source/Runtime/Renderer/Private/TranslucentRendering.cpp b/Engine/Source/Runtime/Renderer/Private/TranslucentRendering.cpp
--- a/Engine/Source/Runtime/Renderer/Private/TranslucentRendering.cpp
+++ b/Engine/Source/Runtime/Renderer/Private/TranslucentRendering.cpp
@@ -711,11 +711,7 @@ bool FSceneRenderer::ShouldRenderTranslucency(ETranslucencyPass::Type Translucen
 
 FScreenPassTextureViewport FSeparateTranslucencyDimensions::GetInstancedStereoViewport(const FViewInfo& View) const
 {
-	FIntRect ViewRect = View.ViewRect;
-	if (View.IsInstancedStereoPass() && !View.bIsMultiViewportEnabled)
-	{
-		ViewRect.Max.X = ViewRect.Min.X + View.InstancedStereoWidth;
-	}
+	FIntRect ViewRect = View.ViewRectWithSecondaryViews;
 	ViewRect = GetScaledRect(ViewRect, Scale);
 	return FScreenPassTextureViewport(Extent, ViewRect);
 }
