Fixed cloud flickering after some time. The extra camera cut logic is not necessary and not safe for large time.

diff --git a/Engine/Source/Runtime/Renderer/Private/VolumetricRenderTarget.cpp b/Engine/Source/Runtime/Renderer/Private/VolumetricRenderTarget.cpp
--- a/Engine/Source/Runtime/Renderer/Private/VolumetricRenderTarget.cpp
+++ b/Engine/Source/Runtime/Renderer/Private/VolumetricRenderTarget.cpp
@@ -426,9 +426,7 @@ void InitVolumetricRenderTargetForViews(FRDGBuilder& GraphBuilder, TArrayView<FV
 		FVolumetricRenderTargetViewStateData& VolumetricCloudRT = ViewInfo.ViewState->VolumetricCloudRenderTarget;
 
 		// Determine if we are initializing or we should reset the persistent state
-		const float DeltaTime = ViewInfo.Family->Time.GetRealTimeSeconds() - ViewInfo.ViewState->LastRenderTime;
-		const bool bFirstFrameOrTimeWasReset = DeltaTime < -0.0001f || ViewInfo.ViewState->LastRenderTime < 0.0001f;
-		const bool bCameraCut = (bFirstFrameOrTimeWasReset || ViewInfo.bCameraCut || ViewInfo.bForceCameraVisibilityReset);
+		const bool bCameraCut = ViewInfo.bCameraCut || ViewInfo.bForceCameraVisibilityReset;
 
 		FIntPoint ViewRect = ViewInfo.ViewRect.Size();
 		VolumetricCloudRT.Initialise(	// TODO this is going to reallocate a buffer each time dynamic resolution scaling is applied 
