Add heterogeneous volumes property to FPrimitiveSceneProxy.
- Expose heterogeneous volumes property to Niagara MeshRenderer Properties
- Conditionalize use of homogeneous volumes rendering based on existence of this property.

diff --git a/Engine/Plugins/FX/Niagara/Source/Niagara/Private/NiagaraComponent.cpp b/Engine/Plugins/FX/Niagara/Source/Niagara/Private/NiagaraComponent.cpp
--- a/Engine/Plugins/FX/Niagara/Source/Niagara/Private/NiagaraComponent.cpp
+++ b/Engine/Plugins/FX/Niagara/Source/Niagara/Private/NiagaraComponent.cpp
@@ -189,6 +189,7 @@ FNiagaraSceneProxy::FNiagaraSceneProxy(UNiagaraComponent* InComponent)
 		ComputeDispatchInterface = SystemInstanceController->GetComputeDispatchInterface();
 
 		bAlwaysHasVelocity = RenderData->HasAnyMotionBlurEnabled();
+		bIsHeterogeneousVolume = RenderData->HasAnyHeterogeneousVolumesEnabled();
 
 		SystemStatID = InComponent->GetAsset()->GetStatID(false, false);
 	}
diff --git a/Engine/Plugins/FX/Niagara/Source/Niagara/Private/NiagaraMeshRendererProperties.cpp b/Engine/Plugins/FX/Niagara/Source/Niagara/Private/NiagaraMeshRendererProperties.cpp
--- a/Engine/Plugins/FX/Niagara/Source/Niagara/Private/NiagaraMeshRendererProperties.cpp
+++ b/Engine/Plugins/FX/Niagara/Source/Niagara/Private/NiagaraMeshRendererProperties.cpp
@@ -234,6 +234,7 @@ UNiagaraMeshRendererProperties::UNiagaraMeshRendererProperties()
 	: SourceMode(ENiagaraRendererSourceDataMode::Particles)
 	, SortMode(ENiagaraSortMode::None)
 	, bOverrideMaterials(false)
+	, bUseHeterogeneousVolumes(false)
 	, bSortOnlyWhenTranslucent(true)
 	, bSubImageBlend(false)
 	, SubImageSize(1.0f, 1.0f)
diff --git a/Engine/Plugins/FX/Niagara/Source/Niagara/Private/NiagaraSystemRenderData.cpp b/Engine/Plugins/FX/Niagara/Source/Niagara/Private/NiagaraSystemRenderData.cpp
--- a/Engine/Plugins/FX/Niagara/Source/Niagara/Private/NiagaraSystemRenderData.cpp
+++ b/Engine/Plugins/FX/Niagara/Source/Niagara/Private/NiagaraSystemRenderData.cpp
@@ -266,6 +266,7 @@ void FNiagaraSystemRenderData::RecacheRenderers(const FNiagaraSystemInstance& Sy
 	EmitterRenderers_GT.Reserve(RendererDrawOrder.Num());
 
 	bAnyMotionBlurEnabled = false;
+	bAnyHeterogeneousVolumesEnabled = false;
 	for (TSharedRef<const FNiagaraEmitterInstance, ESPMode::ThreadSafe> EmitterInst : SystemInstance.GetEmitters())
 	{
 		if (FVersionedNiagaraEmitterData* EmitterData = EmitterInst->GetCachedEmitterData())
@@ -282,6 +283,10 @@ void FNiagaraSystemRenderData::RecacheRenderers(const FNiagaraSystemInstance& Sy
 					{
 						bAnyMotionBlurEnabled = true;
 					}
+					if (Properties->UseHeterogeneousVolumes())
+					{
+						bAnyHeterogeneousVolumesEnabled = true;
+					}
 				}
 				EmitterRenderers_GT.Add(NewRenderer);
 			}
diff --git a/Engine/Plugins/FX/Niagara/Source/Niagara/Public/NiagaraMeshRendererProperties.h b/Engine/Plugins/FX/Niagara/Source/Niagara/Public/NiagaraMeshRendererProperties.h
--- a/Engine/Plugins/FX/Niagara/Source/Niagara/Public/NiagaraMeshRendererProperties.h
+++ b/Engine/Plugins/FX/Niagara/Source/Niagara/Public/NiagaraMeshRendererProperties.h
@@ -219,6 +219,7 @@ public:
 	virtual TArray<FNiagaraVariable> GetBoundAttributes() const override;
 #endif // WITH_EDITORONLY_DATA
 	virtual void CacheFromCompiledData(const FNiagaraDataSetCompiledData* CompiledData) override;
+	virtual bool UseHeterogeneousVolumes() const override { return bUseHeterogeneousVolumes; }
 
 	void UpdateMICs();
 
@@ -249,6 +250,10 @@ public:
 	UPROPERTY(EditAnywhere, Category = "Mesh Rendering", DisplayName="Enable Material Overrides")
 	uint32 bOverrideMaterials : 1;
 
+	/** Whether or not to render with heterogeneous volumes.*/
+	UPROPERTY(EditAnywhere, Category = "Mesh Rendering")
+	uint32 bUseHeterogeneousVolumes : 1;
+
 	/** If true, the particles are only sorted when using a translucent material. */
 	UPROPERTY(EditAnywhere, Category = "Sorting", meta = (EditCondition = "SortMode != ENiagaraSortMode::None", EditConditionHides))
 	uint32 bSortOnlyWhenTranslucent : 1;
diff --git a/Engine/Plugins/FX/Niagara/Source/Niagara/Public/NiagaraRendererProperties.h b/Engine/Plugins/FX/Niagara/Source/Niagara/Public/NiagaraRendererProperties.h
--- a/Engine/Plugins/FX/Niagara/Source/Niagara/Public/NiagaraRendererProperties.h
+++ b/Engine/Plugins/FX/Niagara/Source/Niagara/Public/NiagaraRendererProperties.h
@@ -369,6 +369,7 @@ public:
 	virtual bool NeedsSystemCompletion() const { return false; }
 
 	bool NeedsPreciseMotionVectors() const;
+	virtual bool UseHeterogeneousVolumes() const { return false; }
 
 	static bool IsSortHighPrecision(ENiagaraRendererSortPrecision SortPrecision);
 
diff --git a/Engine/Plugins/FX/Niagara/Source/Niagara/Public/NiagaraSystemRenderData.h b/Engine/Plugins/FX/Niagara/Source/Niagara/Public/NiagaraSystemRenderData.h
--- a/Engine/Plugins/FX/Niagara/Source/Niagara/Public/NiagaraSystemRenderData.h
+++ b/Engine/Plugins/FX/Niagara/Source/Niagara/Public/NiagaraSystemRenderData.h
@@ -79,6 +79,7 @@ public:
 	FORCEINLINE bool IsRenderingEnabled() const { return bRenderingEnabled && (IsInRenderingThread() ? EmitterRenderers_RT.Num() > 0 : EmitterRenderers_GT.Num() > 0); }
 	FORCEINLINE void SetRenderingEnabled(bool bInEnabled) { bRenderingEnabled = bInEnabled; }
 	FORCEINLINE bool HasAnyMotionBlurEnabled() const { return bAnyMotionBlurEnabled; }
+	FORCEINLINE bool HasAnyHeterogeneousVolumesEnabled() const { return bAnyHeterogeneousVolumesEnabled; }
 
 	FORCEINLINE int32 GetNumRenderers() const { return RendererDrawOrder.Num(); }
 
@@ -94,6 +95,7 @@ private:
 
 	bool bRenderingEnabled = true;
 	bool bAnyMotionBlurEnabled = false;
+	bool bAnyHeterogeneousVolumesEnabled = false;
 
 	ERHIFeatureLevel::Type FeatureLevel;
 };
diff --git a/Engine/Source/Runtime/Engine/Private/PrimitiveSceneProxy.cpp b/Engine/Source/Runtime/Engine/Private/PrimitiveSceneProxy.cpp
--- a/Engine/Source/Runtime/Engine/Private/PrimitiveSceneProxy.cpp
+++ b/Engine/Source/Runtime/Engine/Private/PrimitiveSceneProxy.cpp
@@ -247,6 +247,7 @@ FPrimitiveSceneProxy::FPrimitiveSceneProxy(const UPrimitiveComponent* InComponen
 ,	bStaticElementsAlwaysUseProxyPrimitiveUniformBuffer(false)
 ,	bVFRequiresPrimitiveUniformBuffer(true)
 ,	bIsNaniteMesh(false)
+,	bIsHeterogeneousVolume(false)
 ,	bIsHierarchicalInstancedStaticMesh(false)
 ,	bIsLandscapeGrass(false)
 ,	bSupportsGPUScene(false)
diff --git a/Engine/Source/Runtime/Engine/Public/PrimitiveSceneProxy.h b/Engine/Source/Runtime/Engine/Public/PrimitiveSceneProxy.h
--- a/Engine/Source/Runtime/Engine/Public/PrimitiveSceneProxy.h
+++ b/Engine/Source/Runtime/Engine/Public/PrimitiveSceneProxy.h
@@ -778,6 +778,14 @@ public:
 		return bIsNaniteMesh;
 	}
 
+	/**
+	* Returns whether this proxy is a heterogeneous volume.
+	*/
+	inline bool IsHeterogeneousVolume() const
+	{
+		return bIsHeterogeneousVolume;
+	}
+
 	/**
 	 * Returns true if all meshes drawn by this proxy support GPU scene. 
 	 */
@@ -1243,6 +1251,9 @@ protected:
 	/** Whether this proxy is a Nanite mesh. */
 	uint8 bIsNaniteMesh : 1;
 
+	/** Whether this proxy is a heterogeneous volume. */
+	uint8 bIsHeterogeneousVolume : 1;
+
 	/** Whether the primitive is a HierarchicalInstancedStaticMesh. */
 	uint8 bIsHierarchicalInstancedStaticMesh : 1;
 
diff --git a/Engine/Source/Runtime/Renderer/Private/HeterogeneousVolumes/HeterogeneousVolumes.cpp b/Engine/Source/Runtime/Renderer/Private/HeterogeneousVolumes/HeterogeneousVolumes.cpp
--- a/Engine/Source/Runtime/Renderer/Private/HeterogeneousVolumes/HeterogeneousVolumes.cpp
+++ b/Engine/Source/Runtime/Renderer/Private/HeterogeneousVolumes/HeterogeneousVolumes.cpp
@@ -195,6 +195,20 @@ bool DoesPlatformSupportHeterogeneousVolumes(EShaderPlatform Platform)
 		&& !IsForwardShadingEnabled(Platform);
 }
 
+bool DoesMaterialShaderSupportHeterogeneousVolumes(const FMaterialShaderParameters& MaterialShaderParameters)
+{
+	return (MaterialShaderParameters.MaterialDomain == MD_Volume)
+		// Restricting compilation to materials bound to Niagara meshes
+		&& MaterialShaderParameters.bIsUsedWithNiagaraMeshParticles;
+}
+
+bool DoesMaterialShaderSupportHeterogeneousVolumes(const FMaterial& Material)
+{
+	return (Material.GetMaterialDomain() == MD_Volume)
+		// Restricting compilation to materials bound to Niagara meshes
+		&& Material.IsUsedWithNiagaraMeshParticles();
+}
+
 namespace HeterogeneousVolumes
 {
 	// CVars
@@ -352,13 +366,12 @@ void FDeferredShadingSceneRenderer::RenderHeterogeneousVolumes(
 				const FMeshBatch* Mesh = View.VolumetricMeshBatches[MeshBatchIndex].Mesh;
 				const FMaterialRenderProxy* MaterialRenderProxy = Mesh->MaterialRenderProxy;
 				const FMaterial& Material = MaterialRenderProxy->GetMaterialWithFallback(View.GetFeatureLevel(), MaterialRenderProxy);
-				// Only Niagara mesh particles bound to volume materials
-				if (Material.GetMaterialDomain() != MD_Volume || !Material.IsUsedWithNiagaraMeshParticles())
+				const FPrimitiveSceneProxy* PrimitiveSceneProxy = View.VolumetricMeshBatches[MeshBatchIndex].Proxy;
+				if (!PrimitiveSceneProxy->IsHeterogeneousVolume() || !DoesMaterialShaderSupportHeterogeneousVolumes(Material))
 				{
 					continue;
 				}
 
-				const FPrimitiveSceneProxy* PrimitiveSceneProxy = View.VolumetricMeshBatches[MeshBatchIndex].Proxy;
 				const FPrimitiveSceneInfo* PrimitiveSceneInfo = PrimitiveSceneProxy->GetPrimitiveSceneInfo();
 				const int32 PrimitiveId = PrimitiveSceneInfo->GetIndex();
 				const FBoxSphereBounds LocalBoxSphereBounds = PrimitiveSceneProxy->GetLocalBounds();
diff --git a/Engine/Source/Runtime/Renderer/Private/HeterogeneousVolumes/HeterogeneousVolumes.h b/Engine/Source/Runtime/Renderer/Private/HeterogeneousVolumes/HeterogeneousVolumes.h
--- a/Engine/Source/Runtime/Renderer/Private/HeterogeneousVolumes/HeterogeneousVolumes.h
+++ b/Engine/Source/Runtime/Renderer/Private/HeterogeneousVolumes/HeterogeneousVolumes.h
@@ -16,6 +16,7 @@ class FViewInfo;
 class FVirtualShadowMapArray;
 class FVisibleLightInfo;
 
+struct FMaterialShaderParameters;
 struct FRDGTextureDesc;
 struct FSceneTextures;
 
@@ -26,6 +27,7 @@ struct FSceneTextures;
 bool ShouldRenderHeterogeneousVolumes(const FScene* Scene);
 bool ShouldRenderHeterogeneousVolumesForView(const FSceneView& View);
 bool DoesPlatformSupportHeterogeneousVolumes(EShaderPlatform Platform);
+bool DoesMaterialShaderSupportHeterogeneousVolumes(const FMaterialShaderParameters& Parameters);
 
 //
 // Internal API
diff --git a/Engine/Source/Runtime/Renderer/Private/HeterogeneousVolumes/HeterogeneousVolumesLiveShadingPipeline.cpp b/Engine/Source/Runtime/Renderer/Private/HeterogeneousVolumes/HeterogeneousVolumesLiveShadingPipeline.cpp
--- a/Engine/Source/Runtime/Renderer/Private/HeterogeneousVolumes/HeterogeneousVolumesLiveShadingPipeline.cpp
+++ b/Engine/Source/Runtime/Renderer/Private/HeterogeneousVolumes/HeterogeneousVolumesLiveShadingPipeline.cpp
@@ -84,9 +84,7 @@ class FRenderLightingCacheWithLiveShadingCS : public FMeshMaterialShader
 	)
 	{
 		return DoesPlatformSupportHeterogeneousVolumes(Parameters.Platform)
-			&& (Parameters.MaterialParameters.MaterialDomain == MD_Volume)
-			// Restricting compilation to materials bound to Niagara meshes
-			&& Parameters.MaterialParameters.bIsUsedWithNiagaraMeshParticles;
+			&& DoesMaterialShaderSupportHeterogeneousVolumes(Parameters.MaterialParameters);
 	}
 
 	static FPermutationDomain RemapPermutation(FPermutationDomain PermutationVector)
@@ -215,8 +213,7 @@ class FRenderSingleScatteringWithLiveShadingCS : public FMeshMaterialShader
 	)
 	{
 		return DoesPlatformSupportHeterogeneousVolumes(Parameters.Platform)
-			&& (Parameters.MaterialParameters.MaterialDomain == MD_Volume)
-			&& Parameters.MaterialParameters.bIsUsedWithNiagaraMeshParticles;
+			&& DoesMaterialShaderSupportHeterogeneousVolumes(Parameters.MaterialParameters);
 	}
 
 	static FPermutationDomain RemapPermutation(FPermutationDomain PermutationVector)
diff --git a/Engine/Source/Runtime/Renderer/Private/HeterogeneousVolumes/HeterogeneousVolumesMaterialBakingPipeline.cpp b/Engine/Source/Runtime/Renderer/Private/HeterogeneousVolumes/HeterogeneousVolumesMaterialBakingPipeline.cpp
--- a/Engine/Source/Runtime/Renderer/Private/HeterogeneousVolumes/HeterogeneousVolumesMaterialBakingPipeline.cpp
+++ b/Engine/Source/Runtime/Renderer/Private/HeterogeneousVolumes/HeterogeneousVolumesMaterialBakingPipeline.cpp
@@ -64,9 +64,7 @@ class FHeterogeneousVolumesBakeMaterialCS : public FMeshMaterialShader
 	)
 	{
 		return DoesPlatformSupportHeterogeneousVolumes(Parameters.Platform)
-			&& (Parameters.MaterialParameters.MaterialDomain == MD_Volume)
-			// Restricting compilation to materials bound to Niagara meshes
-			&& Parameters.MaterialParameters.bIsUsedWithNiagaraMeshParticles;
+			&& DoesMaterialShaderSupportHeterogeneousVolumes(Parameters.MaterialParameters);
 	}
 
 	static FPermutationDomain RemapPermutation(FPermutationDomain PermutationVector)
