Enforce full Z prepass for Nanite (removed !prepass code that would never be hit + cleaned up all the code around this).

diff --git a/Engine/Source/Runtime/Renderer/Private/BasePassRendering.cpp b/Engine/Source/Runtime/Renderer/Private/BasePassRendering.cpp
--- a/Engine/Source/Runtime/Renderer/Private/BasePassRendering.cpp
+++ b/Engine/Source/Runtime/Renderer/Private/BasePassRendering.cpp
@@ -1345,34 +1345,6 @@ void FDeferredShadingSceneRenderer::RenderBasePassInternal(
 	}
 #endif
 
-	auto RenderNaniteDepthPass = [&](FViewInfo& View, int32 ViewIndex)
-	{
-		// Emit Nanite depth if there was not an earlier depth pre-pass
-		if (!ShouldRenderPrePass())
-		{
-			RDG_GPU_STAT_SCOPE(GraphBuilder, NaniteBasePass);
-
-			Nanite::FRasterResults& RasterResults = NaniteRasterResults[ViewIndex];
-
-			// Emit velocity with depth if not writing it in base pass.
-			FRDGTexture* VelocityBuffer = !IsUsingBasePassVelocity(ShaderPlatform) ? SceneTextures.Velocity : nullptr;
-
-			Nanite::EmitDepthTargets(
-				GraphBuilder,
-				*Scene,
-				View,
-				RasterResults.PageConstants,
-				RasterResults.VisibleClustersSWHW,
-				RasterResults.ViewsBuffer,
-				SceneTextures.Depth.Target,
-				RasterResults.VisBuffer64,
-				VelocityBuffer,
-				RasterResults.MaterialDepth,
-				RasterResults.ShadingMask
-			);
-		}
-	};
-
 	auto RenderNaniteBasePass = [&](FViewInfo& View, int32 ViewIndex)
 	{
 		Nanite::FRasterResults& RasterResults = NaniteRasterResults[ViewIndex];
@@ -1415,13 +1387,15 @@ void FDeferredShadingSceneRenderer::RenderBasePassInternal(
 		// Debug view support for Nanite
 		if (bNaniteEnabled)
 		{
+			// Should always have a full Z prepass with Nanite
+			check(ShouldRenderPrePass());
+
 			for (int32 ViewIndex = 0; ViewIndex < Views.Num(); ++ViewIndex)
 			{
 				FViewInfo& View = Views[ViewIndex];
 				RDG_GPU_MASK_SCOPE(GraphBuilder, View.GPUMask);
 				RDG_EVENT_SCOPE_CONDITIONAL(GraphBuilder, Views.Num() > 1, "View%d", ViewIndex);
 
-				RenderNaniteDepthPass(View, ViewIndex);
 				RenderNaniteBasePass(View, ViewIndex);
 			}
 		}
@@ -1443,8 +1417,6 @@ void FDeferredShadingSceneRenderer::RenderBasePassInternal(
 		RDG_EVENT_SCOPE(GraphBuilder, "BasePass");
 		RDG_GPU_STAT_SCOPE(GraphBuilder, Basepass);
 
-		const bool bNeedsPrePass = ShouldRenderPrePass();
-
 		if (bParallelBasePass)
 		{
 			RDG_WAIT_FOR_TASKS_CONDITIONAL(GraphBuilder, IsBasePassWaitForTasksEnabled());
@@ -1486,7 +1458,9 @@ void FDeferredShadingSceneRenderer::RenderBasePassInternal(
 
 				if (bNaniteEnabled)
 				{
-					RenderNaniteDepthPass(View, ViewIndex);
+					// Should always have a full Z prepass with Nanite
+					check(ShouldRenderPrePass());
+
 					RenderNaniteBasePass(View, ViewIndex);
 				}
 
@@ -1554,7 +1528,9 @@ void FDeferredShadingSceneRenderer::RenderBasePassInternal(
 
 				if (bNaniteEnabled)
 				{
-					RenderNaniteDepthPass(View, ViewIndex);
+					// Should always have a full Z prepass with Nanite
+					check(ShouldRenderPrePass());
+
 					RenderNaniteBasePass(View, ViewIndex);
 				}
 
diff --git a/Engine/Source/Runtime/Renderer/Private/DeferredShadingRenderer.cpp b/Engine/Source/Runtime/Renderer/Private/DeferredShadingRenderer.cpp
--- a/Engine/Source/Runtime/Renderer/Private/DeferredShadingRenderer.cpp
+++ b/Engine/Source/Runtime/Renderer/Private/DeferredShadingRenderer.cpp
@@ -2920,6 +2920,9 @@ void FDeferredShadingSceneRenderer::Render(FRDGBuilder& GraphBuilder)
 		&& (!ViewFamily.EngineShowFlags.Wireframe || bIsViewFrozen || bHasViewParent);
 	const bool bNeedsPrePass = ShouldRenderPrePass();
 
+	// Sanity check - Note: Nanite forces a Z prepass in ShouldForceFullDepthPass()
+	check(!UseNanite(ShaderPlatform) || bNeedsPrePass);
+
 	GEngine->GetPreRenderDelegateEx().Broadcast(GraphBuilder);
 
 	// Strata initialisation is always run even when not enabled.
@@ -3199,45 +3202,22 @@ void FDeferredShadingSceneRenderer::Render(FRDGBuilder& GraphBuilder)
 						// Emit velocity with depth if not writing it in base pass.
 						FRDGTexture* VelocityBuffer = !IsUsingBasePassVelocity(ShaderPlatform) ? SceneTextures.Velocity : nullptr;
 
-						if (bNeedsPrePass)
-						{
-							Nanite::EmitDepthTargets(
-								GraphBuilder,
-								*Scene,
-								Views[ViewIndex],
-								CullingContext.PageConstants,
-								CullingContext.VisibleClustersSWHW,
-								CullingContext.ViewsBuffer,
-								SceneTextures.Depth.Target,
-								RasterContext.VisBuffer64,
-								VelocityBuffer,
-								RasterResults.MaterialDepth,
-								RasterResults.ShadingMask
-							);
-						}
-
-						if (!bIsEarlyDepthComplete && CullingConfig.bTwoPassOcclusion && View.ViewState)
-						{
-							// Won't have a complete SceneDepth for post pass so can't use complete HZB for main pass or it will poke holes in the post pass HZB killing occlusion culling.
-							RDG_EVENT_SCOPE(GraphBuilder, "Nanite::BuildHZB");
-
-							FRDGTextureRef SceneDepth = SystemTextures.Black;
-							FRDGTextureRef GraphHZB = nullptr;
-
-							const FIntRect PrimaryViewRect = View.GetPrimaryView()->ViewRect;
-
-							BuildHZBFurthest(
-								GraphBuilder,
-								SceneDepth,
-								RasterContext.VisBuffer64,
-								PrimaryViewRect,
-								FeatureLevel,
-								ShaderPlatform,
-								TEXT("Nanite.HZB"),
-								/* OutFurthestHZBTexture = */ &GraphHZB);
-
-							GraphBuilder.QueueTextureExtraction(GraphHZB, &View.ViewState->PrevFrameViewInfo.NaniteHZB);
-						}
+						Nanite::EmitDepthTargets(
+							GraphBuilder,
+							*Scene,
+							Views[ViewIndex],
+							CullingContext.PageConstants,
+							CullingContext.VisibleClustersSWHW,
+							CullingContext.ViewsBuffer,
+							SceneTextures.Depth.Target,
+							RasterContext.VisBuffer64,
+							VelocityBuffer,
+							RasterResults.MaterialDepth,
+							RasterResults.ShadingMask
+						);
+						
+						// Sanity check (always force Z prepass)
+						check(bIsEarlyDepthComplete);
 
 						Nanite::ExtractResults(GraphBuilder, CullingContext, RasterContext, RasterResults);
 					}
