Apply plugin GI regardless of ray tracing capabilities

Contributed by: Marcus Svensson (marcussvensson92)
https://github.com/EpicGames/UnrealEngine/pull/9779

#9779

diff --git a/Engine/Source/Runtime/Renderer/Private/IndirectLightRendering.cpp b/Engine/Source/Runtime/Renderer/Private/IndirectLightRendering.cpp
--- a/Engine/Source/Runtime/Renderer/Private/IndirectLightRendering.cpp
+++ b/Engine/Source/Runtime/Renderer/Private/IndirectLightRendering.cpp
@@ -395,6 +395,11 @@ IMPLEMENT_GLOBAL_SHADER(FReflectionEnvironmentSkyLightingPS, "/Engine/Private/Re
 
 IMPLEMENT_GLOBAL_SHADER_PARAMETER_STRUCT(FReflectionUniformParameters, "ReflectionStruct");
 
+static bool ShouldRenderPluginGlobalIllumination(const FViewInfo& View)
+{
+	return View.FinalPostProcessSettings.DynamicGlobalIlluminationMethod == EDynamicGlobalIlluminationMethod::Plugin;
+}
+
 void FDeferredShadingSceneRenderer::CommitIndirectLightingState()
 {
 	for (int32 ViewIndex = 0; ViewIndex < Views.Num(); ViewIndex++)
@@ -422,7 +427,7 @@ void FDeferredShadingSceneRenderer::CommitIndirectLightingState()
 			DiffuseIndirectMethod = EDiffuseIndirectMethod::RTGI;
 			DiffuseIndirectDenoiser = IScreenSpaceDenoiser::GetDenoiserMode(CVarDiffuseIndirectDenoiser);
 		}
-		else if (ShouldRenderPluginRayTracingGlobalIllumination(View))
+		else if (ShouldRenderPluginGlobalIllumination(View))
 		{
 			DiffuseIndirectMethod = EDiffuseIndirectMethod::Plugin;
 			DiffuseIndirectDenoiser = IScreenSpaceDenoiser::GetDenoiserMode(CVarDiffuseIndirectDenoiser);
@@ -585,7 +590,7 @@ TRDGUniformBufferRef<FReflectionUniformParameters> CreateReflectionUniformBuffer
 #if RHI_RAYTRACING
 bool ShouldRenderPluginRayTracingGlobalIllumination(const FViewInfo& View)
 {
-	if (View.FinalPostProcessSettings.DynamicGlobalIlluminationMethod != EDynamicGlobalIlluminationMethod::Plugin)
+	if (!ShouldRenderPluginGlobalIllumination(View))
 	{
 		return false;
 	}
@@ -603,7 +608,7 @@ void FDeferredShadingSceneRenderer::PrepareRayTracingGlobalIlluminationPlugin(co
 	FGlobalIlluminationPluginDelegates::FPrepareRayTracing& Delegate = FGlobalIlluminationPluginDelegates::PrepareRayTracing();
 	Delegate.Broadcast(View, OutRayGenShaders);
 }
-#endif
+#endif // RHI_RAYTRACING
 
 static const FVector SampleArray4x4x6[96] = {
 	FVector(0.72084325551986694, -0.44043412804603577, -0.53516626358032227),
