* Increase bucket size to improve OIT triangle sorting precision.
* Fix dispatch count for the writing pass, not issuing the correct count
* Removing experimental label.

diff --git a/Engine/Shaders/Private/OIT/OITSorting.usf b/Engine/Shaders/Private/OIT/OITSorting.usf
--- a/Engine/Shaders/Private/OIT/OITSorting.usf
+++ b/Engine/Shaders/Private/OIT/OITSorting.usf
@@ -137,6 +137,7 @@ float3 ViewBound_Max;
 
 float4x4 LocalToWorld; // GPU scene buffer instead? How?
 float4x4 WorldToView;
+float4x4 LocalToView;
 
 uint SourcePrimitiveType;
 uint NumPrimitives;
@@ -160,11 +161,7 @@ RWBuffer<uint> OutDebugData;
 groupshared uint s_SliceIndices[SORTING_SLICE_COUNT];
 groupshared uint s_SliceOffsets[SORTING_SLICE_COUNT];
 
-#if SORTING_SLICE_COUNT > PERMUTATION_GROUP_SIZE
-	#error GroupSize (PERMUTATION_GROUP_SIZE) needs to be larger than the number of slices (SORTING_SLICE_COUNT) for having correct counting/allocation
-#endif
-
-[numthreads(PERMUTATION_GROUP_SIZE, 1, 1)]
+[numthreads(SORTING_SLICE_COUNT, 1, 1)]
 void MainCS(uint3 DispatchThreadId : SV_DispatchThreadID, uint GroupThreadId : SV_GroupIndex)
 {
 	uint PrimitiveId = DispatchThreadId.x;
@@ -195,8 +192,7 @@ void MainCS(uint3 DispatchThreadId : SV_DispatchThreadID, uint GroupThreadId : S
 
 		// Compute triangle center in view space
 		const float3 LocalP = (P0 + P1 + P2) / 3;
-		const float3 WorldP = mul(float4(LocalP, 1), LocalToWorld).xyz;
-		const float3 ViewP  = mul(float4(WorldP, 1), WorldToView).xyz;
+		const float3 ViewP = mul(float4(LocalP, 1), LocalToView).xyz;
 
 		const float fSliceIndex = saturate((ViewP.z - ViewBoundMinZ) / (ViewBoundMaxZ - ViewBoundMinZ));
 		SliceIndex = SORTING_SLICE_COUNT * (SortType == BACK_TO_FRONT ? (1.f - fSliceIndex) : fSliceIndex);
diff --git a/Engine/Source/Runtime/Engine/Classes/Components/StaticMeshComponent.h b/Engine/Source/Runtime/Engine/Classes/Components/StaticMeshComponent.h
--- a/Engine/Source/Runtime/Engine/Classes/Components/StaticMeshComponent.h
+++ b/Engine/Source/Runtime/Engine/Classes/Components/StaticMeshComponent.h
@@ -280,7 +280,7 @@ public:
 #endif
 
 	/** Enable dynamic sort mesh's triangles to remove ordering issue when rendered with a translucent material */
-	UPROPERTY(EditAnywhere, AdvancedDisplay, BlueprintReadOnly, Category = Lighting, meta = (UIMin = "0", UIMax = "1", DisplayName = "Sort Triangles (Experimental)"))
+	UPROPERTY(EditAnywhere, AdvancedDisplay, BlueprintReadOnly, Category = Lighting, meta = (UIMin = "0", UIMax = "1", DisplayName = "Sort Triangles"))
 	uint8 bSortTriangles : 1;
 
 	/**
diff --git a/Engine/Source/Runtime/Renderer/Private/OIT/OIT.cpp b/Engine/Source/Runtime/Renderer/Private/OIT/OIT.cpp
--- a/Engine/Source/Runtime/Renderer/Private/OIT/OIT.cpp
+++ b/Engine/Source/Runtime/Renderer/Private/OIT/OIT.cpp
@@ -22,7 +22,7 @@
 static TAutoConsoleVariable<int32> CVarOIT_SortedTriangles_Enable(
 	TEXT("r.OIT.SortedTriangles"), 
 	1, 
-	TEXT("Enable per-instance triangle sorting to avoid invalid triangle ordering (experimental)."),
+	TEXT("Enable per-instance triangle sorting to avoid invalid triangle ordering."),
 	ECVF_ReadOnly | ECVF_RenderThreadSafe);
 
 static TAutoConsoleVariable<int32> CVarOIT_SortedTriangles_Debug(
@@ -256,7 +256,7 @@ static void AddInternalOITComposePass(
 class FSortedIndexBuffer : public FIndexBuffer
 {
 public:
-	static const uint32 SliceCount = 32;
+	static const uint32 SliceCount = 256u;
 
 	FSortedIndexBuffer(uint32 InId, const FBufferRHIRef& InSourceIndexBuffer, uint32 InNumIndices, const TCHAR* InDebugName)
 	: SourceIndexBuffer(InSourceIndexBuffer)
@@ -421,18 +421,6 @@ void FOITSceneData::Deallocate(FIndexBuffer* InIndexBuffer)
 	}
 }
 
-static uint32 GetGroupSize()
-{
-	if (IsRHIDeviceNVIDIA())
-	{
-		return 32u;
-	}
-	else  // IsRHIDeviceAMD() and others
-	{
-		return 64u;
-	}
-}
-
 ///////////////////////////////////////////////////////////////////////////////////////////////////
 // Sort triangle indices to order them front-to-back or back-to-front
 
@@ -441,9 +429,8 @@ class FOITSortTriangleIndex_ScanCS : public FGlobalShader
 	DECLARE_GLOBAL_SHADER(FOITSortTriangleIndex_ScanCS);
 	SHADER_USE_PARAMETER_STRUCT(FOITSortTriangleIndex_ScanCS, FGlobalShader);
 
-	class FGroupSize : SHADER_PERMUTATION_SPARSE_INT("PERMUTATION_GROUP_SIZE", 32, 64); 
 	class FDebug : SHADER_PERMUTATION_BOOL("PERMUTATION_DEBUG");
-	using FPermutationDomain = TShaderPermutationDomain<FGroupSize, FDebug>;
+	using FPermutationDomain = TShaderPermutationDomain<FDebug>;
 
 	BEGIN_SHADER_PARAMETER_STRUCT(FParameters, )
 
@@ -457,6 +444,7 @@ class FOITSortTriangleIndex_ScanCS : public FGlobalShader
 
 		SHADER_PARAMETER(FMatrix44f, LocalToWorld)
 		SHADER_PARAMETER(FMatrix44f, WorldToView)
+		SHADER_PARAMETER(FMatrix44f, LocalToView)
 
 		SHADER_PARAMETER(uint32, SourcePrimitiveType)
 		SHADER_PARAMETER(uint32, NumPrimitives)
@@ -668,8 +656,7 @@ static void AddOITSortTriangleIndexPass(
 		const float ViewBoundMaxZ = ViewBounds.GetBox().Max.Z;
 
 		FOITSortTriangleIndex_ScanCS::FParameters* Parameters = GraphBuilder.AllocParameters<FOITSortTriangleIndex_ScanCS::FParameters>();
-		Parameters->LocalToWorld			= FMatrix44f(MeshBatch.Proxy->GetLocalToWorld());	// LWC_TODO: Precision loss?
-		Parameters->WorldToView				= FMatrix44f(View.ViewMatrices.GetViewMatrix());
+		Parameters->LocalToView				= FMatrix44f(MeshBatch.Proxy->GetLocalToWorld() * View.ViewMatrices.GetViewMatrix());
 		Parameters->SourcePrimitiveType		= Allocation.SourcePrimitiveType == PT_TriangleStrip ? 1u : 0u;
 		Parameters->NumPrimitives			= Allocation.NumPrimitives;
 		Parameters->NumIndices				= Allocation.NumIndices;
@@ -706,14 +693,12 @@ static void AddOITSortTriangleIndexPass(
 			Parameters->OutDebugData = GraphBuilder.CreateUAV(DebugData.Buffer, PF_R32_UINT);
 		}
 
-		const uint32 GroupSize = GetGroupSize();
-
 		FOITSortTriangleIndex_ScanCS::FPermutationDomain PermutationVector;
-		PermutationVector.Set<FOITSortTriangleIndex_ScanCS::FGroupSize>(GroupSize);
 		PermutationVector.Set<FOITSortTriangleIndex_ScanCS::FDebug>(bDebugEnable ? 1 : 0);
 		TShaderMapRef<FOITSortTriangleIndex_ScanCS> ComputeShader(View.ShaderMap, PermutationVector);
 		
-		const FIntVector DispatchCount = FIntVector(FMath::CeilToInt(float(Parameters->NumPrimitives) / float(GroupSize)), 1u, 1u);
+		const uint32 GroupSize = FSortedIndexBuffer::SliceCount;
+		const FIntVector DispatchCount = FIntVector(FMath::DivideAndRoundUp(Parameters->NumPrimitives, GroupSize), 1u, 1u);
 		check(DispatchCount.X < GRHIMaxDispatchThreadGroupsPerDimension.X);
 		ClearUnusedGraphResources(ComputeShader, Parameters);
 		GraphBuilder.AddPass(
@@ -759,7 +744,7 @@ static void AddOITSortTriangleIndexPass(
 		TShaderMapRef<FOITSortTriangleIndex_WriteOutCS> ComputeShader(View.ShaderMap);
 
 		const uint32 GroupSize = 256;
-		const FIntVector DispatchCount = FIntVector(GroupSize, 1u, 1u);
+		const FIntVector DispatchCount = FIntVector(FMath::DivideAndRoundUp(Parameters->NumPrimitives, GroupSize), 1u, 1u);
 		ClearUnusedGraphResources(ComputeShader, Parameters);
 		GraphBuilder.AddPass(
 			RDG_EVENT_NAME("OIT::SortTriangleIndices(Write)"),
