Cleaning up includes of RenderGraph

diff --git a/Engine/Plugins/Animation/MLDeformer/MLDeformerFramework/Source/MLDeformerFramework/Private/MLDeformerModelInstance.cpp b/Engine/Plugins/Animation/MLDeformer/MLDeformerFramework/Source/MLDeformerFramework/Private/MLDeformerModelInstance.cpp
--- a/Engine/Plugins/Animation/MLDeformer/MLDeformerFramework/Source/MLDeformerFramework/Private/MLDeformerModelInstance.cpp
+++ b/Engine/Plugins/Animation/MLDeformer/MLDeformerFramework/Source/MLDeformerFramework/Private/MLDeformerModelInstance.cpp
@@ -9,6 +9,8 @@
 #include "Components/SkeletalMeshComponent.h"
 #include "Engine/SkeletalMesh.h"
 #include "Animation/AnimInstance.h"
+#include "RenderingThread.h"
+#include "RenderGraphBuilder.h"
 
 #include UE_INLINE_GENERATED_CPP_BY_NAME(MLDeformerModelInstance)
 
diff --git a/Engine/Plugins/Experimental/NNI/Source/NeuralNetworkInference/Public/NeuralTensor.h b/Engine/Plugins/Experimental/NNI/Source/NeuralNetworkInference/Public/NeuralTensor.h
--- a/Engine/Plugins/Experimental/NNI/Source/NeuralNetworkInference/Public/NeuralTensor.h
+++ b/Engine/Plugins/Experimental/NNI/Source/NeuralNetworkInference/Public/NeuralTensor.h
@@ -2,14 +2,20 @@
 
 #pragma once
 
-#include "CoreMinimal.h"
 #include "NeuralEnumClasses.h"
+#include "RenderGraphFwd.h"
+#include "RHIFwd.h"
+#include "Templates/RefCounting.h" // TRefCountPtr
+#include <algorithm> // std::fill
+
+#if UE_ENABLE_INCLUDE_ORDER_DEPRECATED_IN_5_3
+#include "CoreMinimal.h"
 #include "RenderGraphBuilder.h" // FRDGBuilder
 #include "RenderGraphDefinitions.h" // FRDGBufferSRVRef, FRDGBufferUAVRef
 #include "RenderGraphResources.h" // FRDGPooledBuffer
 #include "RHIDefinitions.h" // EBufferUsageFlags
-#include "Templates/RefCounting.h" // TRefCountPtr
-#include <algorithm> // std::fill
+#endif
+
 #include "NeuralTensor.generated.h"
 
 /**
diff --git a/Engine/Plugins/FX/Niagara/Source/Niagara/Classes/NiagaraDataInterfaceArrayImpl.h b/Engine/Plugins/FX/Niagara/Source/Niagara/Classes/NiagaraDataInterfaceArrayImpl.h
--- a/Engine/Plugins/FX/Niagara/Source/Niagara/Classes/NiagaraDataInterfaceArrayImpl.h
+++ b/Engine/Plugins/FX/Niagara/Source/Niagara/Classes/NiagaraDataInterfaceArrayImpl.h
@@ -13,6 +13,7 @@
 
 #include "Async/Async.h"
 #include "Containers/ArrayView.h"
+#include "RenderGraphUtils.h"
 #include "ShaderParameterUtils.h"
 #include "ShaderCompilerCore.h"
 
diff --git a/Engine/Plugins/FX/Niagara/Source/Niagara/Private/DataInterface/NiagaraDataInterfaceSimpleCounter.cpp b/Engine/Plugins/FX/Niagara/Source/Niagara/Private/DataInterface/NiagaraDataInterfaceSimpleCounter.cpp
--- a/Engine/Plugins/FX/Niagara/Source/Niagara/Private/DataInterface/NiagaraDataInterfaceSimpleCounter.cpp
+++ b/Engine/Plugins/FX/Niagara/Source/Niagara/Private/DataInterface/NiagaraDataInterfaceSimpleCounter.cpp
@@ -9,6 +9,7 @@
 
 #include "Internationalization/Internationalization.h"
 #include "Async/Async.h"
+#include "RenderGraphUtils.h"
 #include "ShaderParameterUtils.h"
 #include "ShaderCompilerCore.h"
 
diff --git a/Engine/Plugins/FX/Niagara/Source/Niagara/Private/NiagaraDataInterface2DArrayTexture.cpp b/Engine/Plugins/FX/Niagara/Source/Niagara/Private/NiagaraDataInterface2DArrayTexture.cpp
--- a/Engine/Plugins/FX/Niagara/Source/Niagara/Private/NiagaraDataInterface2DArrayTexture.cpp
+++ b/Engine/Plugins/FX/Niagara/Source/Niagara/Private/NiagaraDataInterface2DArrayTexture.cpp
@@ -9,6 +9,7 @@
 
 #include "Engine/Texture2DArray.h"
 #include "Engine/TextureRenderTarget2DArray.h"
+#include "RenderGraphBuilder.h"
 #include "RHIStaticStates.h"
 #include "ShaderCompilerCore.h"
 #include "TextureResource.h"
diff --git a/Engine/Plugins/FX/Niagara/Source/Niagara/Private/NiagaraDataInterfaceCubeTexture.cpp b/Engine/Plugins/FX/Niagara/Source/Niagara/Private/NiagaraDataInterfaceCubeTexture.cpp
--- a/Engine/Plugins/FX/Niagara/Source/Niagara/Private/NiagaraDataInterfaceCubeTexture.cpp
+++ b/Engine/Plugins/FX/Niagara/Source/Niagara/Private/NiagaraDataInterfaceCubeTexture.cpp
@@ -9,6 +9,7 @@
 
 #include "Engine/TextureCube.h"
 #include "Engine/TextureRenderTargetCube.h"
+#include "RenderGraphBuilder.h"
 #include "RHIStaticStates.h"
 #include "ShaderCompilerCore.h"
 #include "TextureResource.h"
diff --git a/Engine/Plugins/FX/Niagara/Source/Niagara/Private/NiagaraDataInterfaceExport.cpp b/Engine/Plugins/FX/Niagara/Source/Niagara/Private/NiagaraDataInterfaceExport.cpp
--- a/Engine/Plugins/FX/Niagara/Source/Niagara/Private/NiagaraDataInterfaceExport.cpp
+++ b/Engine/Plugins/FX/Niagara/Source/Niagara/Private/NiagaraDataInterfaceExport.cpp
@@ -11,6 +11,8 @@
 #include "NiagaraSystemInstance.h"
 
 #include "Internationalization/Internationalization.h"
+#include "RenderGraphBuilder.h"
+#include "RenderGraphUtils.h"
 #include "ShaderCompilerCore.h"
 #include "ShaderParameterUtils.h"
 #include "Misc/LargeWorldRenderPosition.h"
diff --git a/Engine/Plugins/FX/Niagara/Source/Niagara/Private/NiagaraDataInterfaceGrid2DCollection.cpp b/Engine/Plugins/FX/Niagara/Source/Niagara/Private/NiagaraDataInterfaceGrid2DCollection.cpp
--- a/Engine/Plugins/FX/Niagara/Source/Niagara/Private/NiagaraDataInterfaceGrid2DCollection.cpp
+++ b/Engine/Plugins/FX/Niagara/Source/Niagara/Private/NiagaraDataInterfaceGrid2DCollection.cpp
@@ -17,6 +17,8 @@
 #include "Engine/TextureRenderTarget2DArray.h"
 #include "ClearQuad.h"
 #include "CanvasItem.h"
+#include "RenderGraphBuilder.h"
+#include "RenderGraphUtils.h"
 #include "ShaderParameterMetadataBuilder.h"
 #include "ShaderParameterUtils.h"
 #include "TextureResource.h"
diff --git a/Engine/Plugins/FX/Niagara/Source/Niagara/Private/NiagaraDataInterfaceGrid3DCollection.cpp b/Engine/Plugins/FX/Niagara/Source/Niagara/Private/NiagaraDataInterfaceGrid3DCollection.cpp
--- a/Engine/Plugins/FX/Niagara/Source/Niagara/Private/NiagaraDataInterfaceGrid3DCollection.cpp
+++ b/Engine/Plugins/FX/Niagara/Source/Niagara/Private/NiagaraDataInterfaceGrid3DCollection.cpp
@@ -17,6 +17,7 @@
 #include "Engine/TextureRenderTargetVolume.h"
 #include "Engine/VolumeTexture.h"
 #include "ClearQuad.h"
+#include "RenderGraphUtils.h"
 #include "ShaderParameterUtils.h"
 #include "TextureResource.h"
 
diff --git a/Engine/Plugins/FX/Niagara/Source/Niagara/Private/NiagaraDataInterfaceIntRenderTarget2D.cpp b/Engine/Plugins/FX/Niagara/Source/Niagara/Private/NiagaraDataInterfaceIntRenderTarget2D.cpp
--- a/Engine/Plugins/FX/Niagara/Source/Niagara/Private/NiagaraDataInterfaceIntRenderTarget2D.cpp
+++ b/Engine/Plugins/FX/Niagara/Source/Niagara/Private/NiagaraDataInterfaceIntRenderTarget2D.cpp
@@ -20,6 +20,8 @@
 #include "CanvasTypes.h"
 #include "Engine/TextureRenderTarget2D.h"
 #include "Engine/Texture2D.h"
+#include "RenderGraphBuilder.h"
+#include "RenderGraphUtils.h"
 
 #include UE_INLINE_GENERATED_CPP_BY_NAME(NiagaraDataInterfaceIntRenderTarget2D)
 
diff --git a/Engine/Plugins/FX/Niagara/Source/Niagara/Private/NiagaraDataInterfaceNeighborGrid3D.cpp b/Engine/Plugins/FX/Niagara/Source/Niagara/Private/NiagaraDataInterfaceNeighborGrid3D.cpp
--- a/Engine/Plugins/FX/Niagara/Source/Niagara/Private/NiagaraDataInterfaceNeighborGrid3D.cpp
+++ b/Engine/Plugins/FX/Niagara/Source/Niagara/Private/NiagaraDataInterfaceNeighborGrid3D.cpp
@@ -8,6 +8,7 @@
 #include "NiagaraRenderer.h"
 #include "NiagaraShaderParticleID.h"
 #include "DataDrivenShaderPlatformInfo.h"
+#include "RenderGraphUtils.h"
 
 #include UE_INLINE_GENERATED_CPP_BY_NAME(NiagaraDataInterfaceNeighborGrid3D)
 
diff --git a/Engine/Plugins/FX/Niagara/Source/Niagara/Private/NiagaraDataInterfaceRasterizationGrid3D.cpp b/Engine/Plugins/FX/Niagara/Source/Niagara/Private/NiagaraDataInterfaceRasterizationGrid3D.cpp
--- a/Engine/Plugins/FX/Niagara/Source/Niagara/Private/NiagaraDataInterfaceRasterizationGrid3D.cpp
+++ b/Engine/Plugins/FX/Niagara/Source/Niagara/Private/NiagaraDataInterfaceRasterizationGrid3D.cpp
@@ -8,6 +8,7 @@
 #include "NiagaraRenderer.h"
 #include "NiagaraShaderParticleID.h"
 #include "DataDrivenShaderPlatformInfo.h"
+#include "RenderGraphUtils.h"
 
 #include UE_INLINE_GENERATED_CPP_BY_NAME(NiagaraDataInterfaceRasterizationGrid3D)
 
diff --git a/Engine/Plugins/FX/Niagara/Source/Niagara/Private/NiagaraDataInterfaceRenderTarget2D.cpp b/Engine/Plugins/FX/Niagara/Source/Niagara/Private/NiagaraDataInterfaceRenderTarget2D.cpp
--- a/Engine/Plugins/FX/Niagara/Source/Niagara/Private/NiagaraDataInterfaceRenderTarget2D.cpp
+++ b/Engine/Plugins/FX/Niagara/Source/Niagara/Private/NiagaraDataInterfaceRenderTarget2D.cpp
@@ -18,6 +18,8 @@
 #include "NiagaraShader.h"
 #include "NiagaraShaderParametersBuilder.h"
 #include "NiagaraStats.h"
+#include "RenderGraphBuilder.h"
+#include "RenderGraphUtils.h"
 #include "RHIStaticStates.h"
 
 #include UE_INLINE_GENERATED_CPP_BY_NAME(NiagaraDataInterfaceRenderTarget2D)
diff --git a/Engine/Plugins/FX/Niagara/Source/Niagara/Private/NiagaraDataInterfaceRenderTarget2DArray.cpp b/Engine/Plugins/FX/Niagara/Source/Niagara/Private/NiagaraDataInterfaceRenderTarget2DArray.cpp
--- a/Engine/Plugins/FX/Niagara/Source/Niagara/Private/NiagaraDataInterfaceRenderTarget2DArray.cpp
+++ b/Engine/Plugins/FX/Niagara/Source/Niagara/Private/NiagaraDataInterfaceRenderTarget2DArray.cpp
@@ -16,6 +16,8 @@
 #include "NiagaraShaderParametersBuilder.h"
 #include "NiagaraGpuComputeDebugInterface.h"
 #include "NiagaraGpuComputeDispatchInterface.h"
+#include "RenderGraphBuilder.h"
+#include "RenderGraphUtils.h"
 #include "RHIStaticStates.h"
 
 #include UE_INLINE_GENERATED_CPP_BY_NAME(NiagaraDataInterfaceRenderTarget2DArray)
diff --git a/Engine/Plugins/FX/Niagara/Source/Niagara/Private/NiagaraDataInterfaceRenderTargetCube.cpp b/Engine/Plugins/FX/Niagara/Source/Niagara/Private/NiagaraDataInterfaceRenderTargetCube.cpp
--- a/Engine/Plugins/FX/Niagara/Source/Niagara/Private/NiagaraDataInterfaceRenderTargetCube.cpp
+++ b/Engine/Plugins/FX/Niagara/Source/Niagara/Private/NiagaraDataInterfaceRenderTargetCube.cpp
@@ -14,6 +14,8 @@
 #include "NiagaraShaderParametersBuilder.h"
 #include "NiagaraGpuComputeDebugInterface.h"
 #include "NiagaraGpuComputeDispatchInterface.h"
+#include "RenderGraphBuilder.h"
+#include "RenderGraphUtils.h"
 #include "RHIStaticStates.h"
 
 #include UE_INLINE_GENERATED_CPP_BY_NAME(NiagaraDataInterfaceRenderTargetCube)
diff --git a/Engine/Plugins/FX/Niagara/Source/Niagara/Private/NiagaraDataInterfaceRenderTargetVolume.cpp b/Engine/Plugins/FX/Niagara/Source/Niagara/Private/NiagaraDataInterfaceRenderTargetVolume.cpp
--- a/Engine/Plugins/FX/Niagara/Source/Niagara/Private/NiagaraDataInterfaceRenderTargetVolume.cpp
+++ b/Engine/Plugins/FX/Niagara/Source/Niagara/Private/NiagaraDataInterfaceRenderTargetVolume.cpp
@@ -18,6 +18,8 @@
 #include "NiagaraShaderParametersBuilder.h"
 #include "NiagaraGpuComputeDebugInterface.h"
 #include "VolumeCache.h"
+#include "RenderGraphBuilder.h"
+#include "RenderGraphUtils.h"
 #include "RHIStaticStates.h"
 #include "Misc/PathViews.h"
 #include "Misc/Paths.h"
diff --git a/Engine/Plugins/FX/Niagara/Source/Niagara/Private/NiagaraDataInterfaceTexture.cpp b/Engine/Plugins/FX/Niagara/Source/Niagara/Private/NiagaraDataInterfaceTexture.cpp
--- a/Engine/Plugins/FX/Niagara/Source/Niagara/Private/NiagaraDataInterfaceTexture.cpp
+++ b/Engine/Plugins/FX/Niagara/Source/Niagara/Private/NiagaraDataInterfaceTexture.cpp
@@ -10,6 +10,7 @@
 #include "NiagaraShader.h"
 #include "NiagaraShaderParametersBuilder.h"
 #include "NiagaraSystemInstance.h"
+#include "RenderGraphBuilder.h"
 #include "RHIStaticStates.h"
 #include "ShaderCompilerCore.h"
 #include "TextureResource.h"
diff --git a/Engine/Plugins/FX/Niagara/Source/Niagara/Private/NiagaraDataInterfaceVolumeTexture.cpp b/Engine/Plugins/FX/Niagara/Source/Niagara/Private/NiagaraDataInterfaceVolumeTexture.cpp
--- a/Engine/Plugins/FX/Niagara/Source/Niagara/Private/NiagaraDataInterfaceVolumeTexture.cpp
+++ b/Engine/Plugins/FX/Niagara/Source/Niagara/Private/NiagaraDataInterfaceVolumeTexture.cpp
@@ -8,6 +8,7 @@
 #include "NiagaraSystemInstance.h"
 #include "Engine/TextureRenderTargetVolume.h"
 #include "Engine/VolumeTexture.h"
+#include "RenderGraphBuilder.h"
 #include "RHIStaticStates.h"
 #include "ShaderCompilerCore.h"
 #include "TextureResource.h"
diff --git a/Engine/Plugins/FX/Niagara/Source/Niagara/Private/NiagaraGpuComputeDispatch.h b/Engine/Plugins/FX/Niagara/Source/Niagara/Private/NiagaraGpuComputeDispatch.h
--- a/Engine/Plugins/FX/Niagara/Source/Niagara/Private/NiagaraGpuComputeDispatch.h
+++ b/Engine/Plugins/FX/Niagara/Source/Niagara/Private/NiagaraGpuComputeDispatch.h
@@ -22,6 +22,7 @@ the same VectorVM byte code / compute shader code
 #include "NiagaraSystemGpuComputeProxy.h"
 #include "ParticleResources.h"
 #include "Particles/ParticleSortingGPU.h"
+#include "RenderGraphUtils.h"
 #include "RHIResources.h"
 #include "RendererInterface.h"
 #include "Tickable.h"
diff --git a/Engine/Plugins/FX/Niagara/Source/Niagara/Private/NiagaraRendererRibbons.cpp b/Engine/Plugins/FX/Niagara/Source/Niagara/Private/NiagaraRendererRibbons.cpp
--- a/Engine/Plugins/FX/Niagara/Source/Niagara/Private/NiagaraRendererRibbons.cpp
+++ b/Engine/Plugins/FX/Niagara/Source/Niagara/Private/NiagaraRendererRibbons.cpp
@@ -22,6 +22,7 @@
 #include "NiagaraGpuComputeDispatchInterface.h"
 #include "NiagaraRibbonCompute.h"
 #include "Misc/LazySingleton.h"
+#include "RenderGraphUtils.h"
 
 DECLARE_CYCLE_STAT(TEXT("Generate Ribbon Vertex Data [GT]"), STAT_NiagaraGenRibbonVertexData, STATGROUP_Niagara);
 DECLARE_CYCLE_STAT(TEXT("Render Ribbons [RT]"), STAT_NiagaraRenderRibbons, STATGROUP_Niagara);
diff --git a/Engine/Plugins/FX/Niagara/Source/NiagaraShader/Public/NiagaraDistanceFieldHelper.h b/Engine/Plugins/FX/Niagara/Source/NiagaraShader/Public/NiagaraDistanceFieldHelper.h
--- a/Engine/Plugins/FX/Niagara/Source/NiagaraShader/Public/NiagaraDistanceFieldHelper.h
+++ b/Engine/Plugins/FX/Niagara/Source/NiagaraShader/Public/NiagaraDistanceFieldHelper.h
@@ -3,11 +3,15 @@
 #pragma once
 
 #include "GlobalDistanceFieldParameters.h"
+
+#if UE_ENABLE_INCLUDE_ORDER_DEPRECATED_IN_5_3
 #include "RenderGraphBuilder.h"
+#endif
 
 class FDistanceFieldSceneData;
 class FDistanceFieldObjectBufferParameters;
 class FDistanceFieldAtlasParameters;
+class FRDGBuilder;
 
 namespace FNiagaraDistanceFieldHelper
 {
diff --git a/Engine/Plugins/FX/Niagara/Source/NiagaraShader/Public/NiagaraShader.h b/Engine/Plugins/FX/Niagara/Source/NiagaraShader/Public/NiagaraShader.h
--- a/Engine/Plugins/FX/Niagara/Source/NiagaraShader/Public/NiagaraShader.h
+++ b/Engine/Plugins/FX/Niagara/Source/NiagaraShader/Public/NiagaraShader.h
@@ -15,10 +15,12 @@
 #include "NiagaraScriptBase.h"
 #include "NiagaraShared.h"
 #include "NiagaraShaderType.h"
-//#if UE_ENABLE_INCLUDE_ORDER_DEPRECATED_IN_5_1
+#include "SceneTexturesConfig.h"
+
+#if UE_ENABLE_INCLUDE_ORDER_DEPRECATED_IN_5_3
 #include "RenderGraph.h"
 #include "SceneRenderTargetParameters.h"
-//#endif
+#endif
 
 class UClass;
 
diff --git a/Engine/Plugins/Runtime/HairStrands/Source/HairStrandsCore/Private/GroomAsset.cpp b/Engine/Plugins/Runtime/HairStrands/Source/HairStrandsCore/Private/GroomAsset.cpp
--- a/Engine/Plugins/Runtime/HairStrands/Source/HairStrandsCore/Private/GroomAsset.cpp
+++ b/Engine/Plugins/Runtime/HairStrands/Source/HairStrandsCore/Private/GroomAsset.cpp
@@ -39,6 +39,7 @@
 #include "HairCardsVertexFactory.h"
 #include "PSOPrecache.h"
 #include "UObject/UObjectIterator.h"
+#include "RenderGraphBuilder.h"
 #include "UObject/DevObjectVersion.h"
 
 #if WITH_EDITORONLY_DATA
diff --git a/Engine/Plugins/Runtime/HairStrands/Source/HairStrandsCore/Public/GroomInstance.h b/Engine/Plugins/Runtime/HairStrands/Source/HairStrandsCore/Public/GroomInstance.h
--- a/Engine/Plugins/Runtime/HairStrands/Source/HairStrandsCore/Public/GroomInstance.h
+++ b/Engine/Plugins/Runtime/HairStrands/Source/HairStrandsCore/Public/GroomInstance.h
@@ -7,6 +7,7 @@
 #include "UObject/ObjectMacros.h"
 #include "GroomAsset.h"
 #include "GroomBindingAsset.h"
+#include "GroomDesc.h"
 #include "HairStrandsInterface.h"
 #include "HairStrandsDefinitions.h"
 
diff --git a/Engine/Plugins/Runtime/HairStrands/Source/HairStrandsCore/Public/HairCardsBuilder.h b/Engine/Plugins/Runtime/HairStrands/Source/HairStrandsCore/Public/HairCardsBuilder.h
--- a/Engine/Plugins/Runtime/HairStrands/Source/HairStrandsCore/Public/HairCardsBuilder.h
+++ b/Engine/Plugins/Runtime/HairStrands/Source/HairStrandsCore/Public/HairCardsBuilder.h
@@ -2,6 +2,10 @@
 
 #pragma once
 
+#include "HAL/Platform.h"
+#include "Math/MathFwd.h"
+
+#if UE_ENABLE_INCLUDE_ORDER_DEPRECATED_IN_5_3
 #include "CoreMinimal.h"
 #include "UObject/ObjectMacros.h"
 #include "Engine/EngineTypes.h"
@@ -11,10 +15,23 @@
 #include "HairStrandsDatas.h"
 #include "GroomAssetCards.h"
 #include "GroomResources.h"
+#endif
 
+struct FHairCardsBulkData;
+struct FHairCardsDatas;
+struct FHairCardsInterpolationBulkData;
+struct FHairCardsProceduralDatas;
+struct FHairCardsProceduralResource;
+struct FHairCardsRestResource;
 struct FHairGroupCardsTextures;
+struct FHairGroupsProceduralCards;
+struct FHairMeshesBulkData;
+struct FHairStrandsDatas;
 struct FHairStrandsVoxelData;
 
+class FString;
+class UStaticMesh;
+
 #if WITH_EDITOR
 namespace FHairCardsBuilder
 {
@@ -37,9 +54,9 @@ namespace FHairCardsBuilder
 
 	void BuildGeometry(
 		const FString& LODName,
-		const struct FHairStrandsDatas& InRen,
-		const struct FHairStrandsDatas& InSim,
-		const struct FHairGroupsProceduralCards& Settings,
+		const FHairStrandsDatas& InRen,
+		const FHairStrandsDatas& InSim,
+		const FHairGroupsProceduralCards& Settings,
 		FHairCardsProceduralDatas& Out,
 		FHairCardsBulkData& OutBulk,
 		FHairStrandsDatas& OutGuides,
diff --git a/Engine/Plugins/Runtime/HairStrands/Source/HairStrandsCore/Public/HairStrandsRendering.h b/Engine/Plugins/Runtime/HairStrands/Source/HairStrandsCore/Public/HairStrandsRendering.h
--- a/Engine/Plugins/Runtime/HairStrands/Source/HairStrandsCore/Public/HairStrandsRendering.h
+++ b/Engine/Plugins/Runtime/HairStrands/Source/HairStrandsCore/Public/HairStrandsRendering.h
@@ -2,6 +2,11 @@
 
 #pragma once
 
+#include "Containers/ContainersFwd.h"
+#include "HAL/Platform.h"
+#include "Math/MathFwd.h"
+
+#if UE_ENABLE_INCLUDE_ORDER_DEPRECATED_IN_5_3
 #include "CoreMinimal.h"
 #include "UObject/ObjectMacros.h"
 #include "Engine/EngineTypes.h"
@@ -9,8 +14,17 @@
 #include "HairStrandsInterface.h"
 #include "Shader.h"
 #include "GroomDesc.h"
+#endif
 
 class FGlobalShaderMap;
+class FRDGBuilder;
+class FSkeletalMeshRenderData;
+class UGroomBindingAsset;
+
+struct FHairGroupInstance;
+struct FHairStrandClusterData;
+struct FRWBuffer;
+struct FShaderPrintData;
 
 // Reset the interpolation data. This needs to be called prior to ComputeHairStrandsInterpolation 
 // and prior to the actual hair simulation in order to insure that:
@@ -20,7 +34,7 @@ class FGlobalShaderMap;
 void ResetHairStrandsInterpolation(
 	FRDGBuilder& GraphBuilder,
 	FGlobalShaderMap* ShaderMap,
-	struct FHairGroupInstance* Instance,
+	FHairGroupInstance* Instance,
 	int32 LODIndex);
 
 void ComputeHairStrandsInterpolation(
@@ -30,13 +44,13 @@ void ComputeHairStrandsInterpolation(
 	const uint32 ViewRayTracingMask,
 	const EGroomViewMode ViewMode,
 	const FVector& TranslatedWorldOffset,
-	const struct FShaderPrintData* ShaderPrintData,
-	struct FHairGroupInstance* Instance,
+	const FShaderPrintData* ShaderPrintData,
+	FHairGroupInstance* Instance,
 	int32 LODIndex,
-	struct FHairStrandClusterData* ClusterData);
+	FHairStrandClusterData* ClusterData);
 
 void RegisterClusterData(
-	struct FHairGroupInstance* Instance,
+	FHairGroupInstance* Instance,
 	FHairStrandClusterData* InClusterData);
 
-HAIRSTRANDSCORE_API void ComputeInterpolationWeights(class UGroomBindingAsset* BindingAsset, class FSkeletalMeshRenderData* TargetRenderData, TArray<FRWBuffer>& TransferedPositions);
+HAIRSTRANDSCORE_API void ComputeInterpolationWeights(UGroomBindingAsset* BindingAsset, FSkeletalMeshRenderData* TargetRenderData, TArray<FRWBuffer>& TransferedPositions);
diff --git a/Engine/Plugins/Runtime/nDisplay/Source/DisplayClusterLightCardEditorShaders/Public/DisplayClusterMeshProjectionRenderer.h b/Engine/Plugins/Runtime/nDisplay/Source/DisplayClusterLightCardEditorShaders/Public/DisplayClusterMeshProjectionRenderer.h
--- a/Engine/Plugins/Runtime/nDisplay/Source/DisplayClusterLightCardEditorShaders/Public/DisplayClusterMeshProjectionRenderer.h
+++ b/Engine/Plugins/Runtime/nDisplay/Source/DisplayClusterLightCardEditorShaders/Public/DisplayClusterMeshProjectionRenderer.h
@@ -2,9 +2,12 @@
 
 #pragma once
 
+#include "SceneView.h"
+
+#if UE_ENABLE_INCLUDE_ORDER_DEPRECATED_IN_5_3
 #include "CoreMinimal.h"
 #include "RenderGraphBuilder.h"
-#include "SceneView.h"
+#endif
 
 class AActor;
 class UPrimitiveComponent;
@@ -13,6 +16,7 @@ class FSceneView;
 class FViewInfo;
 class FPrimitiveDrawInterface;
 class FPrimitiveSceneProxy;
+class FRDGBuilder;
 class FSimpleElementCollector;
 class FMeshProjectionPassParameters;
 struct FEngineShowFlags;
diff --git a/Engine/Source/Runtime/Engine/Private/DistanceFieldAtlas.cpp b/Engine/Source/Runtime/Engine/Private/DistanceFieldAtlas.cpp
--- a/Engine/Source/Runtime/Engine/Private/DistanceFieldAtlas.cpp
+++ b/Engine/Source/Runtime/Engine/Private/DistanceFieldAtlas.cpp
@@ -21,6 +21,7 @@
 #include "MeshCardRepresentation.h"
 #include "Misc/QueuedThreadPoolWrapper.h"
 #include "ObjectCacheContext.h"
+#include "RenderGraphUtils.h"
 
 #if WITH_EDITOR
 #include "DerivedDataCacheInterface.h"
@@ -1228,6 +1229,11 @@ FVector4f FLandscapeTextureAtlas::GetAllocationScaleBias(uint32 Handle) const
 	return AddrSpaceAllocator.GetScaleBias(Handle);
 }
 
+FRDGTexture* FLandscapeTextureAtlas::GetAtlasTexture(FRDGBuilder& GraphBuilder) const
+{
+	return TryRegisterExternalTexture(GraphBuilder, AtlasTextureRHI);
+}
+
 void FLandscapeTextureAtlas::FSubAllocator::Init(uint32 InTileSize, uint32 InBorderSize, uint32 InDimInTiles)
 {
 	check(InDimInTiles && !(InDimInTiles & (InDimInTiles - 1)));
diff --git a/Engine/Source/Runtime/Engine/Public/DistanceFieldAtlas.h b/Engine/Source/Runtime/Engine/Public/DistanceFieldAtlas.h
--- a/Engine/Source/Runtime/Engine/Public/DistanceFieldAtlas.h
+++ b/Engine/Source/Runtime/Engine/Public/DistanceFieldAtlas.h
@@ -6,7 +6,9 @@
 
 #pragma once
 
+#if UE_ENABLE_INCLUDE_ORDER_DEPRECATED_IN_5_3
 #include "CoreMinimal.h"
+#endif
 #include "Containers/LockFreeList.h"
 #include "ProfilingDebugging/ResourceSize.h"
 #include "Engine/EngineTypes.h"
@@ -22,7 +24,9 @@
 #include "Templates/UniquePtr.h"
 #include "DerivedMeshDataTaskUtils.h"
 #include "Async/AsyncWork.h"
+#if UE_ENABLE_INCLUDE_ORDER_DEPRECATED_IN_5_3
 #include "RenderGraphUtils.h"
+#endif
 
 #if WITH_EDITOR
 #include "MeshUtilities.h"
@@ -82,10 +86,7 @@ public:
 
 	bool HasAtlasTexture() const { return AtlasTextureRHI.IsValid(); }
 
-	FRDGTexture* GetAtlasTexture(FRDGBuilder& GraphBuilder) const
-	{
-		return TryRegisterExternalTexture(GraphBuilder, AtlasTextureRHI);
-	}
+	FRDGTexture* GetAtlasTexture(FRDGBuilder& GraphBuilder) const;
 
 	uint32 GetSizeX() const
 	{
diff --git a/Engine/Source/Runtime/RenderCore/Public/GPUMessaging.h b/Engine/Source/Runtime/RenderCore/Public/GPUMessaging.h
--- a/Engine/Source/Runtime/RenderCore/Public/GPUMessaging.h
+++ b/Engine/Source/Runtime/RenderCore/Public/GPUMessaging.h
@@ -4,7 +4,9 @@
 
 #include "HAL/Platform.h"
 #include "Misc/AssertionMacros.h"
+#if UE_ENABLE_INCLUDE_ORDER_DEPRECATED_IN_5_3
 #include "RenderGraphBuilder.h"
+#endif
 #include "RenderGraphDefinitions.h"
 #include "ShaderParameterMacros.h"
 #include "Templates/SharedPointer.h"
diff --git a/Engine/Source/Runtime/RenderCore/Public/VisualizeTexture.h b/Engine/Source/Runtime/RenderCore/Public/VisualizeTexture.h
--- a/Engine/Source/Runtime/RenderCore/Public/VisualizeTexture.h
+++ b/Engine/Source/Runtime/RenderCore/Public/VisualizeTexture.h
@@ -10,13 +10,16 @@
 #include "Misc/EnumClassFlags.h"
 #include "Misc/Optional.h"
 #include "Misc/WildcardString.h"
+#include "RenderResource.h"
+#include "RendererInterface.h"
+#include "Templates/RefCounting.h"
+
+#if UE_ENABLE_INCLUDE_ORDER_DEPRECATED_IN_5_3
 #include "RHIDefinitions.h"
 #include "RenderGraph.h"
 #include "RenderGraphDefinitions.h"
 #include "RenderGraphResources.h"
-#include "RenderResource.h"
-#include "RendererInterface.h"
-#include "Templates/RefCounting.h"
+#endif
 
 class FOutputDevice;
 class FRDGBuilder;
diff --git a/Engine/Source/Runtime/Renderer/Private/CustomDepthRendering.h b/Engine/Source/Runtime/Renderer/Private/CustomDepthRendering.h
--- a/Engine/Source/Runtime/Renderer/Private/CustomDepthRendering.h
+++ b/Engine/Source/Runtime/Renderer/Private/CustomDepthRendering.h
@@ -2,7 +2,9 @@
 
 #pragma once
 
-#include "RenderGraph.h"
+#include "RenderGraphFwd.h"
+#include "RHIDefinitions.h"
+#include "RHIFwd.h"
 
 enum class ECustomDepthPassLocation : uint32
 {
diff --git a/Engine/Source/Runtime/Renderer/Private/GPUScene.h b/Engine/Source/Runtime/Renderer/Private/GPUScene.h
--- a/Engine/Source/Runtime/Renderer/Private/GPUScene.h
+++ b/Engine/Source/Runtime/Renderer/Private/GPUScene.h
@@ -11,6 +11,7 @@
 #include "MeshBatch.h"
 #include "LightSceneData.h"
 
+class FRDGExternalAccessQueue;
 class FRHICommandList;
 class FScene;
 class FViewInfo;
diff --git a/Engine/Source/Runtime/Renderer/Private/HairStrands/HairStrandsComposition.h b/Engine/Source/Runtime/Renderer/Private/HairStrands/HairStrandsComposition.h
--- a/Engine/Source/Runtime/Renderer/Private/HairStrands/HairStrandsComposition.h
+++ b/Engine/Source/Runtime/Renderer/Private/HairStrands/HairStrandsComposition.h
@@ -6,8 +6,10 @@
 
 #pragma once
 
-#include "RenderGraph.h"
-#include "SceneRendering.h"
+#include "Containers/ContainersFwd.h"
+#include "RenderGraphFwd.h"
+
+class FViewInfo;
 
 void RenderHairComposition(
 	FRDGBuilder& GraphBuilder,
diff --git a/Engine/Source/Runtime/Renderer/Private/InstanceCulling/InstanceCullingLoadBalancer.cpp b/Engine/Source/Runtime/Renderer/Private/InstanceCulling/InstanceCullingLoadBalancer.cpp
--- a/Engine/Source/Runtime/Renderer/Private/InstanceCulling/InstanceCullingLoadBalancer.cpp
+++ b/Engine/Source/Runtime/Renderer/Private/InstanceCulling/InstanceCullingLoadBalancer.cpp
@@ -1,3 +1,39 @@
 // Copyright Epic Games, Inc. All Rights Reserved.
 
 #include "InstanceCullingLoadBalancer.h"
+#include "RenderGraphUtils.h"
+#include "ShaderCore.h"
+
+void FInstanceCullingLoadBalancerBase::SetShaderDefines(FShaderCompilerEnvironment& OutEnvironment)
+{
+	OutEnvironment.SetDefine(TEXT("NUM_THREADS_PER_GROUP"), ThreadGroupSize);
+	OutEnvironment.SetDefine(TEXT("NUM_INSTANCES_ITEM_BITS"), NumInstancesItemBits);
+	OutEnvironment.SetDefine(TEXT("NUM_INSTANCES_ITEM_MASK"), NumInstancesItemMask);
+	OutEnvironment.SetDefine(TEXT("PREFIX_BITS"), PrefixBits);
+	OutEnvironment.SetDefine(TEXT("PREFIX_BIT_MASK"), PrefixBitMask);
+}
+
+void FInstanceCullingLoadBalancerBase::FGPUData::GetShaderParameters(FRDGBuilder& GraphBuilder, FShaderParameters& ShaderParameters)
+{
+	ShaderParameters.BatchBuffer = GraphBuilder.CreateSRV(BatchBuffer);
+	ShaderParameters.ItemBuffer = GraphBuilder.CreateSRV(ItemBuffer);
+	ShaderParameters.NumBatches = NumBatches;
+	ShaderParameters.NumItems = NumItems;
+}
+
+FInstanceCullingLoadBalancerBase::FGPUData FInstanceCullingLoadBalancerBase::Upload(FRDGBuilder& GraphBuilder, TConstArrayView<FPackedBatch> Batches, TConstArrayView<FPackedItem> Items, ERDGInitialDataFlags RDGInitialDataFlags)
+{
+	FGPUData Result;
+	// TODO: Several of these load balancers are being created on the stack and the memory is going out of scope before RDG execution. Always making a copy for now.
+	Result.BatchBuffer = CreateStructuredBuffer(GraphBuilder, TEXT("InstanceCullingLoadBalancer.Batches"), Batches, /* RDGInitialDataFlags */ ERDGInitialDataFlags::None);
+	Result.ItemBuffer = CreateStructuredBuffer(GraphBuilder, TEXT("InstanceCullingLoadBalancer.Items"), Items, /* RDGInitialDataFlags */ ERDGInitialDataFlags::None);
+	Result.NumBatches = Batches.Num();
+	Result.NumItems = Items.Num();
+
+	return Result;
+}
+
+FIntVector FInstanceCullingLoadBalancerBase::GetWrappedCsGroupCount(TConstArrayView<FPackedBatch> Batches) const
+{
+	return FComputeShaderUtils::GetGroupCountWrapped(Batches.Num());
+}
diff --git a/Engine/Source/Runtime/Renderer/Private/InstanceCulling/InstanceCullingLoadBalancer.h b/Engine/Source/Runtime/Renderer/Private/InstanceCulling/InstanceCullingLoadBalancer.h
--- a/Engine/Source/Runtime/Renderer/Private/InstanceCulling/InstanceCullingLoadBalancer.h
+++ b/Engine/Source/Runtime/Renderer/Private/InstanceCulling/InstanceCullingLoadBalancer.h
@@ -2,12 +2,12 @@
 
 #pragma once
 
-#include "CoreMinimal.h"
-#include "RHI.h"
-#include "RenderGraphBuilder.h"
-#include "RenderGraphResources.h"
-#include "RenderGraphUtils.h"
+#include "HAL/Platform.h"
+#include "RenderGraphDefinitions.h"
+#include "RendererInterface.h"
+#include "ShaderParameterMacros.h"
 
+struct FShaderCompilerEnvironment;
 
 namespace InstanceCullingImplementationDetails
 {
@@ -83,6 +83,12 @@ public:
 		SHADER_PARAMETER(uint32, NumItems)
 	END_SHADER_PARAMETER_STRUCT()
 
+	/*
+	* Publish constants to a shader implementing a kernel using the load balancer.
+	* Call from ModifyCompilationEnvironment
+	*/
+	static void SetShaderDefines(FShaderCompilerEnvironment& OutEnvironment);
+
 	struct FGPUData
 	{
 		int32 NumBatches = 0;
@@ -90,14 +96,12 @@ public:
 		FRDGBufferRef BatchBuffer = nullptr;
 		FRDGBufferRef ItemBuffer = nullptr;
 
-		void GetShaderParameters(FRDGBuilder& GraphBuilder, FShaderParameters& ShaderParameters)
-		{
-			ShaderParameters.BatchBuffer = GraphBuilder.CreateSRV(BatchBuffer);
-			ShaderParameters.ItemBuffer = GraphBuilder.CreateSRV(ItemBuffer);
-			ShaderParameters.NumBatches = NumBatches;
-			ShaderParameters.NumItems = NumItems;
-		}
+		void GetShaderParameters(FRDGBuilder& GraphBuilder, FShaderParameters& ShaderParameters);
 	};
+
+	FGPUData Upload(FRDGBuilder& GraphBuilder, TConstArrayView<FPackedBatch> Batches, TConstArrayView<FPackedItem> Items, ERDGInitialDataFlags RDGInitialDataFlags);
+
+	FIntVector GetWrappedCsGroupCount(TConstArrayView<FPackedBatch> Batches) const;
 };
 
 /*
@@ -110,20 +114,6 @@ public:
 	using AllocatorType = InAllocatorType;
 	static constexpr ERDGInitialDataFlags DefaultRDGInitialDataFlags = InstanceCullingImplementationDetails::AllocatorTypeRDGInitialDataFlags<AllocatorType>::Flags;
 
-
-	/*
-	 * Publish constants to a shader implementing a kernel using the load balancer.
-	 * Call from ModifyCompilationEnvironment
-	 */
-	static void SetShaderDefines(FShaderCompilerEnvironment& OutEnvironment)
-	{
-		OutEnvironment.SetDefine(TEXT("NUM_THREADS_PER_GROUP"), ThreadGroupSize);
-		OutEnvironment.SetDefine(TEXT("NUM_INSTANCES_ITEM_BITS"), NumInstancesItemBits);
-		OutEnvironment.SetDefine(TEXT("NUM_INSTANCES_ITEM_MASK"), NumInstancesItemMask);
-		OutEnvironment.SetDefine(TEXT("PREFIX_BITS"), PrefixBits);
-		OutEnvironment.SetDefine(TEXT("PREFIX_BIT_MASK"), PrefixBitMask);
-	}
-
 	void ReserveStorage(int32 NumBatches, int32 NumItems)
 	{
 		Batches.Empty(NumBatches);
@@ -177,14 +167,7 @@ public:
 	{
 		FinalizeBatches();
 
-		FGPUData Result;
-		// TODO: Several of these load balancers are being created on the stack and the memory is going out of scope before RDG execution. Always making a copy for now.
-		Result.BatchBuffer = CreateStructuredBuffer(GraphBuilder, TEXT("InstanceCullingLoadBalancer.Batches"), Batches, /* RDGInitialDataFlags */ ERDGInitialDataFlags::None);
-		Result.ItemBuffer = CreateStructuredBuffer(GraphBuilder, TEXT("InstanceCullingLoadBalancer.Items"), Items, /* RDGInitialDataFlags */ ERDGInitialDataFlags::None);
-		Result.NumBatches = Batches.Num();
-		Result.NumItems = Items.Num();
-
-		return Result;
+		return FInstanceCullingLoadBalancerBase::Upload(GraphBuilder, Batches, Items, RDGInitialDataFlags);
 	}
 
 	/**
@@ -205,7 +188,7 @@ public:
 	 */
 	FIntVector GetWrappedCsGroupCount() const
 	{
-		return FComputeShaderUtils::GetGroupCountWrapped(Batches.Num());
+		return FInstanceCullingLoadBalancerBase::GetWrappedCsGroupCount(Batches);
 	}
 
 	const TArray<FPackedBatch, AllocatorType> &GetBatches() const
diff --git a/Engine/Source/Runtime/Renderer/Private/InstanceCulling/InstanceCullingMergedContext.cpp b/Engine/Source/Runtime/Renderer/Private/InstanceCulling/InstanceCullingMergedContext.cpp
--- a/Engine/Source/Runtime/Renderer/Private/InstanceCulling/InstanceCullingMergedContext.cpp
+++ b/Engine/Source/Runtime/Renderer/Private/InstanceCulling/InstanceCullingMergedContext.cpp
@@ -2,6 +2,7 @@
 
 #include "InstanceCulling/InstanceCullingMergedContext.h"
 #include "InstanceCulling/InstanceCullingManager.h"
+#include "RenderGraphBuilder.h"
 
 void FInstanceCullingMergedContext::MergeBatches()
 {
diff --git a/Engine/Source/Runtime/Renderer/Private/RayTracing/RayTracingInstanceBufferUtil.cpp b/Engine/Source/Runtime/Renderer/Private/RayTracing/RayTracingInstanceBufferUtil.cpp
--- a/Engine/Source/Runtime/Renderer/Private/RayTracing/RayTracingInstanceBufferUtil.cpp
+++ b/Engine/Source/Runtime/Renderer/Private/RayTracing/RayTracingInstanceBufferUtil.cpp
@@ -13,6 +13,7 @@
 #include "RenderCore.h"
 #include "ShaderParameterStruct.h"
 #include "GlobalShader.h"
+#include "PipelineStateCache.h"
 #include "DataDrivenShaderPlatformInfo.h"
 
 #include "Async/ParallelFor.h"
