Fixed VT not processing continous updates when async VT task is enabled.

diff --git a/Engine/Source/Runtime/Renderer/Private/VT/UniqueRequestList.h b/Engine/Source/Runtime/Renderer/Private/VT/UniqueRequestList.h
--- a/Engine/Source/Runtime/Renderer/Private/VT/UniqueRequestList.h
+++ b/Engine/Source/Runtime/Renderer/Private/VT/UniqueRequestList.h
@@ -91,15 +91,22 @@ public:
 		ContinuousUpdateRequestHash.Clear();
 	}
 
-	inline void Reset()
+	inline void Reset(bool bResetContinousUpdates)
 	{
-		Initialize();
+		LoadRequestHash.Clear();
+		MappingRequestHash.Clear();
+		DirectMappingRequestHash.Clear();
 		NumLoadRequests = 0;
 		NumLockRequests = 0;
 		NumMappingRequests = 0;
 		NumDirectMappingRequests = 0;
-		NumContinuousUpdateRequests = 0;
 		NumAdaptiveAllocationRequests = 0;
+
+		if (bResetContinousUpdates)
+		{
+			NumContinuousUpdateRequests = 0;
+			ContinuousUpdateRequestHash.Clear();
+		}
 	}
 
 	inline uint32 GetNumLoadRequests() const { return NumLoadRequests; }
diff --git a/Engine/Source/Runtime/Renderer/Private/VT/VirtualTextureSystem.cpp b/Engine/Source/Runtime/Renderer/Private/VT/VirtualTextureSystem.cpp
--- a/Engine/Source/Runtime/Renderer/Private/VT/VirtualTextureSystem.cpp
+++ b/Engine/Source/Runtime/Renderer/Private/VT/VirtualTextureSystem.cpp
@@ -2634,7 +2634,7 @@ void FVirtualTextureSystem::BeginUpdate(FRDGBuilder& GraphBuilder, FVirtualTextu
 		SubmitThrottledRequests(RHICmdList, Updater, bContinousUpdates);
 
 		// Reset the request list for the gather in EndUpdate.
-		Updater->MergedRequestList->Reset();
+		Updater->MergedRequestList->Reset(bContinousUpdates);
 
 	}, Updater->bAsyncTaskAllowed);
 }
