Enable heterogeneous volumes by default.
- Disable indirect lighting by default.

diff --git a/Engine/Source/Runtime/Renderer/Private/HeterogeneousVolumes/HeterogeneousVolumes.cpp b/Engine/Source/Runtime/Renderer/Private/HeterogeneousVolumes/HeterogeneousVolumes.cpp
--- a/Engine/Source/Runtime/Renderer/Private/HeterogeneousVolumes/HeterogeneousVolumes.cpp
+++ b/Engine/Source/Runtime/Renderer/Private/HeterogeneousVolumes/HeterogeneousVolumes.cpp
@@ -12,8 +12,8 @@
 
 static TAutoConsoleVariable<int32> CVarHeterogeneousVolumes(
 	TEXT("r.HeterogeneousVolumes"),
-	0,
-	TEXT("Enables the Heterogeneous volume integrator (Default = 0)"),
+	1,
+	TEXT("Enables the Heterogeneous volume integrator (Default = 1)"),
 	ECVF_RenderThreadSafe
 );
 
@@ -33,8 +33,8 @@ static TAutoConsoleVariable<int32> CVarHeterogeneousVolumesHardwareRayTracing(
 
 static TAutoConsoleVariable<int32> CVarHeterogeneousVolumesIndirectLighting(
 	TEXT("r.HeterogeneousVolumes.IndirectLighting"),
-	1,
-	TEXT("Enables indirect lighting (Default = 1)"),
+	0,
+	TEXT("Enables indirect lighting (Default = 0)"),
 	ECVF_RenderThreadSafe
 );
 
