Added gpu timer and stat for volumetric cloud reconstruction and upsampling.

diff --git a/Engine/Source/Runtime/Renderer/Private/VolumetricRenderTarget.cpp b/Engine/Source/Runtime/Renderer/Private/VolumetricRenderTarget.cpp
--- a/Engine/Source/Runtime/Renderer/Private/VolumetricRenderTarget.cpp
+++ b/Engine/Source/Runtime/Renderer/Private/VolumetricRenderTarget.cpp
@@ -144,6 +144,10 @@ static bool AnyViewRequiresProcessing(TArrayView<FViewInfo> Views)
 	return bAnyViewRequiresProcessing;
 }
 
+DECLARE_GPU_STAT(VolCloudReconstruction);
+DECLARE_GPU_STAT(VolCloudComposeOverScene);
+DECLARE_GPU_STAT(VolCloudComposeUnderSLW);
+DECLARE_GPU_STAT(VolCloudComposeForVis);
 
 /*=============================================================================
 	UVolumetricCloudComponent implementation.
@@ -549,6 +553,10 @@ void ReconstructVolumetricRenderTarget(
 		return;
 	}
 
+	RDG_EVENT_SCOPE(GraphBuilder, "VolCloudReconstruction");
+	RDG_GPU_STAT_SCOPE(GraphBuilder, VolCloudReconstruction);
+	SCOPED_NAMED_EVENT(VolCloudReconstruction, FColor::Emerald);
+
 	const FRDGSystemTextures& SystemTextures = FRDGSystemTextures::Get(GraphBuilder);
 
 	for (int32 ViewIndex = 0; ViewIndex < Views.Num(); ViewIndex++)
@@ -730,6 +738,10 @@ void ComposeVolumetricRenderTargetOverScene(
 		return;
 	}
 
+	RDG_EVENT_SCOPE(GraphBuilder, "VolCloudComposeOverScene");
+	RDG_GPU_STAT_SCOPE(GraphBuilder, VolCloudComposeOverScene);
+	SCOPED_NAMED_EVENT(VolCloudComposeOverScene, FColor::Emerald);
+
 	FRHIBlendState* PreMultipliedColorTransmittanceBlend = TStaticBlendState<CW_RGB, BO_Add, BF_One, BF_SourceAlpha, BO_Add, BF_Zero, BF_One>::GetRHI();
 
 	for (int32 ViewIndex = 0; ViewIndex < Views.Num(); ViewIndex++)
@@ -806,6 +818,10 @@ void ComposeVolumetricRenderTargetOverSceneUnderWater(
 		return;
 	}
 
+	RDG_EVENT_SCOPE(GraphBuilder, "VolCloudComposeUnderSLW");
+	RDG_GPU_STAT_SCOPE(GraphBuilder, VolCloudComposeUnderSLW);
+	SCOPED_NAMED_EVENT(VolCloudComposeUnderSLW, FColor::Emerald);
+
 	FRHIBlendState* PreMultipliedColorTransmittanceBlend = TStaticBlendState<CW_RGB, BO_Add, BF_One, BF_SourceAlpha, BO_Add, BF_Zero, BF_One>::GetRHI();
 
 	for (int32 ViewIndex = 0; ViewIndex < Views.Num(); ViewIndex++)
@@ -872,6 +888,10 @@ void ComposeVolumetricRenderTargetOverSceneForVisualization(
 		return;
 	}
 
+	RDG_EVENT_SCOPE(GraphBuilder, "VolCloudComposeForVis");
+	RDG_GPU_STAT_SCOPE(GraphBuilder, VolCloudComposeForVis);
+	SCOPED_NAMED_EVENT(VolCloudComposeForVis, FColor::Emerald);
+
 	for (int32 ViewIndex = 0; ViewIndex < Views.Num(); ViewIndex++)
 	{
 		FViewInfo& ViewInfo = Views[ViewIndex];
