Prevents crash when heterogeneous volumes and forward shading is enabled.

diff --git a/Engine/Source/Runtime/Renderer/Private/HeterogeneousVolumes/HeterogeneousVolumes.cpp b/Engine/Source/Runtime/Renderer/Private/HeterogeneousVolumes/HeterogeneousVolumes.cpp
--- a/Engine/Source/Runtime/Renderer/Private/HeterogeneousVolumes/HeterogeneousVolumes.cpp
+++ b/Engine/Source/Runtime/Renderer/Private/HeterogeneousVolumes/HeterogeneousVolumes.cpp
@@ -172,7 +172,8 @@ bool ShouldRenderHeterogeneousVolumes(
 )
 {
 	return IsHeterogeneousVolumesEnabled()
-		&& Scene != nullptr;
+		&& Scene != nullptr
+		&& DoesPlatformSupportHeterogeneousVolumes(Scene->GetShaderPlatform());
 }
 
 bool ShouldRenderHeterogeneousVolumesForView(
