Remove unused declaration

diff --git a/Engine/Source/Runtime/Renderer/Private/VirtualShadowMaps/VirtualShadowMapClipmap.h b/Engine/Source/Runtime/Renderer/Private/VirtualShadowMaps/VirtualShadowMapClipmap.h
--- a/Engine/Source/Runtime/Renderer/Private/VirtualShadowMaps/VirtualShadowMapClipmap.h
+++ b/Engine/Source/Runtime/Renderer/Private/VirtualShadowMaps/VirtualShadowMapClipmap.h
@@ -60,8 +60,6 @@ public:
 		return LightSceneInfo;
 	}
 
-	int32 GetHZBKey(int32 ClipmapLevelIndex) const;
-
 	FVirtualShadowMapProjectionShaderData GetProjectionShaderData(int32 ClipmapIndex) const;
 
 	FVector GetWorldOrigin() const
