Remove vertex start from groom logic.

This is part of a general clean up which will remove 'temporal' support.

diff --git a/Engine/Plugins/Runtime/HairStrands/Source/HairStrandsCore/Private/GroomResources.cpp b/Engine/Plugins/Runtime/HairStrands/Source/HairStrandsCore/Private/GroomResources.cpp
--- a/Engine/Plugins/Runtime/HairStrands/Source/HairStrandsCore/Private/GroomResources.cpp
+++ b/Engine/Plugins/Runtime/HairStrands/Source/HairStrandsCore/Private/GroomResources.cpp
@@ -914,7 +914,6 @@ void FHairStrandsRestResource::InternalAllocate(FRDGBuilder& GraphBuilder)
 void AddHairTangentPass(
 	FRDGBuilder& GraphBuilder,
 	FGlobalShaderMap* ShaderMap,
-	uint32 VertexStart,
 	uint32 VertexCount,
 	FHairGroupPublicData* HairGroupPublicData,
 	FRDGBufferSRVRef PositionBuffer,
@@ -930,7 +929,6 @@ FRDGExternalBuffer FHairStrandsRestResource::GetTangentBuffer(FRDGBuilder& Graph
 		AddHairTangentPass(
 			GraphBuilder,
 			ShaderMap,
-			0,
 			BulkData.PointCount,
 			nullptr,
 			RegisterAsSRV(GraphBuilder, PositionBuffer),
diff --git a/Engine/Plugins/Runtime/HairStrands/Source/HairStrandsCore/Private/HairStrandsRendering.cpp b/Engine/Plugins/Runtime/HairStrands/Source/HairStrandsCore/Private/HairStrandsRendering.cpp
--- a/Engine/Plugins/Runtime/HairStrands/Source/HairStrandsCore/Private/HairStrandsRendering.cpp
+++ b/Engine/Plugins/Runtime/HairStrands/Source/HairStrandsCore/Private/HairStrandsRendering.cpp
@@ -488,7 +488,6 @@ class FHairInterpolationCS : public FGlobalShader
 
 	BEGIN_SHADER_PARAMETER_STRUCT(FParameters, )
 		SHADER_PARAMETER_STRUCT_INCLUDE(ShaderPrint::FShaderParameters, ShaderDrawParameters)
-		SHADER_PARAMETER(uint32, VertexStart)
 		SHADER_PARAMETER(uint32, VertexCount)
 		SHADER_PARAMETER(uint32, DispatchCountX)
 		SHADER_PARAMETER(float, HairLengthScale)
@@ -568,7 +567,6 @@ static void AddHairStrandsInterpolationPass(
 	FGlobalShaderMap* ShaderMap,
 	const FShaderPrintData* ShaderPrintData,
 	const FHairGroupInstance* Instance,
-	const uint32 VertexStart,
 	const uint32 VertexCount,
 	const int32 MeshLODIndex,
 	const float HairLengthScale,
@@ -616,7 +614,6 @@ static void AddHairStrandsInterpolationPass(
 	}
 	Parameters->OutRenDeformedPositionBuffer = OutRenPositionBuffer.UAV;
 
-	Parameters->VertexStart = VertexStart;
 	Parameters->VertexCount = VertexCount;
 	Parameters->InRenHairPositionOffset = (FVector3f)InRenHairWorldOffset;
 	Parameters->InSimHairPositionOffset = (FVector3f)InSimHairWorldOffset;
@@ -865,7 +862,6 @@ class FHairClusterAABBCS : public FGlobalShader
 	BEGIN_SHADER_PARAMETER_STRUCT(FParameters, )
 		SHADER_PARAMETER_STRUCT_INCLUDE(ShaderPrint::FShaderParameters, ShaderDrawParameters)
 		SHADER_PARAMETER(uint32, TotalClusterCount)
-		SHADER_PARAMETER(uint32, VertexStart)
 		SHADER_PARAMETER(uint32, VertexCount)
 		SHADER_PARAMETER(float, RcpSampleWeight)
 
@@ -937,7 +933,6 @@ static void AddHairClusterAABBPass(
 	Parameters->RenderDeformedOffsetBuffer = RenderDeformedOffsetBuffer;
 	Parameters->TotalClusterCount = 1;
 
-	Parameters->VertexStart = Instance->HairGroupPublicData->GetActiveStrandsVertexStart(Instance->Strands.RestResource->GetVertexCount());
 	Parameters->VertexCount = Instance->HairGroupPublicData->GetActiveStrandsVertexCount(Instance->Strands.RestResource->GetVertexCount(), Instance->HairGroupPublicData->MaxScreenSize);
 	Parameters->RcpSampleWeight = 1.0 / Instance->HairGroupPublicData->GetActiveStrandsSampleWeight(true, Instance->HairGroupPublicData->MaxScreenSize);
 
@@ -1176,7 +1171,6 @@ class FHairTangentCS : public FGlobalShader
 	using FPermutationDomain = TShaderPermutationDomain<FGroupSize, FCulling>;
 
 	BEGIN_SHADER_PARAMETER_STRUCT(FParameters, )
-		SHADER_PARAMETER(uint32, VertexStart)
 		SHADER_PARAMETER(uint32, VertexCount)
 		SHADER_PARAMETER(uint32, DispatchCountX)
 		SHADER_PARAMETER(uint32, HairStrandsVF_bIsCullingEnable)
@@ -1196,7 +1190,6 @@ IMPLEMENT_GLOBAL_SHADER(FHairTangentCS, "/Engine/Private/HairStrands/HairStrands
 void AddHairTangentPass(
 	FRDGBuilder& GraphBuilder,
 	FGlobalShaderMap* ShaderMap,
-	uint32 VertexStart,
 	uint32 VertexCount,
 	FHairGroupPublicData* HairGroupPublicData,
 	FRDGBufferSRVRef PositionBuffer,
@@ -1209,7 +1202,6 @@ void AddHairTangentPass(
 	FHairTangentCS::FParameters* Parameters = GraphBuilder.AllocParameters<FHairTangentCS::FParameters>();
 	Parameters->PositionBuffer = PositionBuffer;
 	Parameters->OutputTangentBuffer = OutTangentBuffer.UAV;
-	Parameters->VertexStart = VertexStart;
 	Parameters->VertexCount = VertexCount;
 	Parameters->DispatchCountX = DispatchCount.X;
 	Parameters->HairStrandsVF_bIsCullingEnable = bCullingEnable ? 1 : 0;
@@ -1792,7 +1784,6 @@ void ComputeHairStrandsInterpolation(
 			AddHairTangentPass(
 				GraphBuilder,
 				ShaderMap,
-				0,
 				Instance->Guides.RestResource->GetVertexCount(),
 				Instance->HairGroupPublicData,
 				RegisterAsSRV(GraphBuilder, Instance->Guides.DeformedResource->GetBuffer(FHairStrandsDeformedResource::Current)),
@@ -1918,13 +1909,11 @@ void ComputeHairStrandsInterpolation(
 					}
 				}
 
-				uint32 VertexStart = 0;
 				uint32 VertexCount = Instance->Strands.RestResource->GetVertexCount();
 
 				// 2.1 Compute deformation position based on simulation/skinning/RBF
 				if (Instance->Debug.GroomCacheType != EGroomCacheType::Strands)
 				{
-					VertexStart = Instance->HairGroupPublicData->GetActiveStrandsVertexStart(Instance->Strands.RestResource->GetVertexCount());
 					VertexCount = Instance->HairGroupPublicData->GetActiveStrandsVertexCount(Instance->Strands.RestResource->GetVertexCount(), Instance->HairGroupPublicData->MaxScreenSize);
 
 					// 2.1.1 Current Position
@@ -1933,7 +1922,6 @@ void ComputeHairStrandsInterpolation(
 						ShaderMap,
 						ShaderPrintData,
 						Instance,
-						VertexStart,
 						VertexCount,
 						MeshLODIndex,
 						Instance->Strands.Modifier.HairLengthScale,
@@ -1984,7 +1972,6 @@ void ComputeHairStrandsInterpolation(
 							ShaderMap,
 							ShaderPrintData,
 							Instance,
-							VertexStart,
 							VertexCount,
 							MeshLODIndex,
 							Instance->Strands.Modifier.HairLengthScale,
@@ -2049,7 +2036,6 @@ void ComputeHairStrandsInterpolation(
 				AddHairTangentPass(
 					GraphBuilder,
 					ShaderMap,
-					VertexStart,
 					VertexCount,
 					Instance->HairGroupPublicData,
 					Strands_DeformedPosition.SRV,
@@ -2268,7 +2254,6 @@ void ComputeHairStrandsInterpolation(
 						ShaderMap,
 						ShaderPrintData,
 						Instance,
-						0,
 						LOD.Guides.RestResource->GetVertexCount(),
 						MeshLODIndex,
 						1.0f,
diff --git a/Engine/Shaders/Private/HairStrands/HairStrandsInterpolation.usf b/Engine/Shaders/Private/HairStrands/HairStrandsInterpolation.usf
--- a/Engine/Shaders/Private/HairStrands/HairStrandsInterpolation.usf
+++ b/Engine/Shaders/Private/HairStrands/HairStrandsInterpolation.usf
@@ -71,7 +71,6 @@ FGuideData GetGuideData(uint InVertexIndex, uint4 InInterpolationData0, uint4 In
 
 #if SHADER_HAIRINTERPOLATION
 
-uint VertexStart;
 uint VertexCount;
 uint DispatchCountX;
 uint HairDebugMode;
@@ -460,7 +459,7 @@ float3 ApplyDeformer(const FRootTriangleData RenTri, float3 InRestPostion, float
 void MainCS(uint2 DispatchThreadId : SV_DispatchThreadID)
 {
 	// 1. Find vertex index
-	uint VertexIndex = VertexStart + GetHairStrandsVertexFetchIndex(DispatchThreadId, GROUP_SIZE, DispatchCountX);
+	uint VertexIndex = GetHairStrandsVertexFetchIndex(DispatchThreadId, GROUP_SIZE, DispatchCountX);
 #if PERMUTATION_CULLING
 	if (HairStrandsVF_bIsCullingEnable)
 	{
@@ -476,7 +475,7 @@ void MainCS(uint2 DispatchThreadId : SV_DispatchThreadID)
 
 
 	// 2. Appply deformation
-	if (VertexIndex < (VertexStart + VertexCount))
+	if (VertexIndex < VertexCount)
 	{
 		const FGuideData GuideData = GetGuideData(VertexIndex);
 		const FRootTriangleData RenTri = GetRootTriangleData(VertexIndex, false);
@@ -650,7 +649,6 @@ uint TotalClusterCount;
 float3 CPUBoundMin;
 float3 CPUBoundMax;
 
-uint VertexStart;
 uint VertexCount;
 float RcpSampleWeight;
 
@@ -699,7 +697,7 @@ void ClusterAABBEvaluationCS(uint3 DispatchThreadId : SV_DispatchThreadID, uint3
 	}
 	const uint ClusterId   = ClusterIdBuffer[ClusterFetchIndex];
 	const uint IndexOffset = ClusterIndexOffsetBuffer[ClusterFetchIndex];
-	const uint IndexCount = max(min(ClusterIndexCountBuffer[ClusterFetchIndex], GROUP_SIZE),ClusterIndexCountBuffer[ClusterFetchIndex] * RcpSampleWeight);
+	const uint IndexCount = max(min(ClusterIndexCountBuffer[ClusterFetchIndex], GROUP_SIZE),ClusterIndexCountBuffer[ClusterFetchIndex] * RcpSampleWeight); // HAIR_TODO:
 	
 	const uint VertexCountAfterCulling = HairStrandsVF_CullingIndirectBuffer[3];
 
@@ -716,7 +714,7 @@ void ClusterAABBEvaluationCS(uint3 DispatchThreadId : SV_DispatchThreadID, uint3
 		{
 			const float3 OutHairPositionOffset = RenderDeformedOffsetBuffer[0].xyz;
 			const uint VertexIndex = ClusterVertexIdBuffer[IndexOffset + ThreadReadOffset];
-			if ((VertexIndex >= VertexStart) && (VertexIndex < (VertexStart + VertexCount)))
+			if (VertexIndex < VertexCount)
 			{
 				const FHairControlPoint CP = UnpackHairControlPoint(RenderDeformedPositionBuffer[VertexIndex], OutHairPositionOffset, 1, 1, 1);
 
diff --git a/Engine/Shaders/Private/HairStrands/HairStrandsTangent.usf b/Engine/Shaders/Private/HairStrands/HairStrandsTangent.usf
--- a/Engine/Shaders/Private/HairStrands/HairStrandsTangent.usf
+++ b/Engine/Shaders/Private/HairStrands/HairStrandsTangent.usf
@@ -4,7 +4,6 @@
 #include "../MonteCarlo.ush"
 #include "HairStrandsClusterCommon.ush"
 
-uint VertexStart;
 uint VertexCount;
 uint DispatchCountX;
 
@@ -35,8 +34,8 @@ float4 PackNormal(float3 V, uint A)
 [numthreads(GROUP_SIZE, 1, 1)]
 void MainCS(uint2 DispatchThreadId : SV_DispatchThreadID)
 {
-	const int MaxVertexCount = VertexStart + int(VertexCount);
-	int IndexCurr = VertexStart + GetHairStrandsVertexFetchIndex(DispatchThreadId, GROUP_SIZE, DispatchCountX);
+	const int MaxVertexCount = int(VertexCount);
+	int IndexCurr = GetHairStrandsVertexFetchIndex(DispatchThreadId, GROUP_SIZE, DispatchCountX);
 	int IndexPrev = clamp(IndexCurr-1, 0, MaxVertexCount-1);
 	int IndexNext = clamp(IndexCurr+1, 0, MaxVertexCount-1);
 	bool bValid = IndexCurr < MaxVertexCount;
diff --git a/Engine/Shaders/Private/HairStrands/HairStrandsVisibilityRasterClassification.usf b/Engine/Shaders/Private/HairStrands/HairStrandsVisibilityRasterClassification.usf
--- a/Engine/Shaders/Private/HairStrands/HairStrandsVisibilityRasterClassification.usf
+++ b/Engine/Shaders/Private/HairStrands/HairStrandsVisibilityRasterClassification.usf
@@ -13,7 +13,6 @@
 RWBuffer<uint>		OutPrimIDs;
 RWByteAddressBuffer OutDrawIndexedArgs;
 RWByteAddressBuffer OutSWRasterPrimitiveCount;
-uint				ControlPointStart;
 uint				ControlPointCount;
 uint				PrimIDsBufferSize;
 uint				NumWorkGroups;
@@ -101,13 +100,9 @@ void CSMain(uint DispatchThreadID : SV_DispatchThreadID, uint GroupThreadID : SV
 				PrimID = VertexIndex0;
 			}
 		}
-		else if (bSegValid)
-		{
-			PrimID += ControlPointStart;
-		}
 #else
-		const uint PrimID = BatchIndex * 1024 + GroupThreadID + ControlPointStart;
-		bool bSegValid = (BatchIndex < NumBatches) && (PrimID < (ControlPointStart + ControlPointCount));
+		const uint PrimID = BatchIndex * 1024 + GroupThreadID;
+		bool bSegValid = (BatchIndex < NumBatches) && (PrimID < ControlPointCount);
 #endif // PERMUTATION_CULLING
 		
 		bool bShouldSoftwareRasterize = false;
diff --git a/Engine/Shaders/Private/HairStrands/HairStrandsVisibilityRasterCompute.usf b/Engine/Shaders/Private/HairStrands/HairStrandsVisibilityRasterCompute.usf
--- a/Engine/Shaders/Private/HairStrands/HairStrandsVisibilityRasterCompute.usf
+++ b/Engine/Shaders/Private/HairStrands/HairStrandsVisibilityRasterCompute.usf
@@ -73,7 +73,6 @@ uint			HairMaterialId;
 Texture2D<float> SceneDepthTexture;
 
 uint 			VertexCount;
-uint 			VertexStart;
 float			SampleWeight;
 
 float3 NDCToPixelCoord(float4 InDC)
@@ -450,13 +449,9 @@ void BinningCS(uint DispatchThreadID : SV_DispatchThreadID, uint GroupThreadID :
 				PrimID = VertexIndex0;
 			}
 		}
-		else if (bSegValid)
-		{
-			PrimID += VertexStart;
-		}
 	#else // PERMUTATION_CULLING
-		const uint PrimID = BatchIndex * 1024 + GroupThreadID + VertexStart;
-		bSegValid = bSegValid && (PrimID < (VertexStart + VertexCount));
+		const uint PrimID = BatchIndex * 1024 + GroupThreadID;
+		bSegValid = bSegValid && (PrimID < VertexCount);
 	#endif // PERMUTATION_CULLING
 #endif // PERMUTATION_INDIRECT_PRIM_IDS
 
diff --git a/Engine/Shaders/Private/HairStrands/HairStrandsVisibilityRasterComputeNaive.usf b/Engine/Shaders/Private/HairStrands/HairStrandsVisibilityRasterComputeNaive.usf
--- a/Engine/Shaders/Private/HairStrands/HairStrandsVisibilityRasterComputeNaive.usf
+++ b/Engine/Shaders/Private/HairStrands/HairStrandsVisibilityRasterComputeNaive.usf
@@ -22,7 +22,6 @@ ByteAddressBuffer	IndirectPrimIDCount;
 Buffer<uint>		IndirectPrimIDs;
 float2				OutputResolution;
 uint				HairMaterialId;
-uint				ControlPointStart;
 uint				ControlPointCount;
 float				SampleWeight;
 uint				NumWorkGroups;
@@ -171,13 +170,9 @@ void CSMain(uint DispatchThreadID : SV_DispatchThreadID, uint GroupThreadID : SV
 				PrimID = VertexIndex0;
 			}
 		}
-		else if (bSegValid)
-		{
-			PrimID += ControlPointStart;
-		}
 	#else // PERMUTATION_CULLING
-		const uint PrimID = BatchIndex * 64 + GroupThreadID + ControlPointStart;
-		bSegValid = bSegValid && (PrimID < (ControlPointStart + ControlPointCount));
+		const uint PrimID = BatchIndex * 64 + GroupThreadID;
+		bSegValid = bSegValid && (PrimID < ControlPointCount);
 	#endif // PERMUTATION_CULLING
 #endif // PERMUTATION_INDIRECT_PRIM_IDS
 
diff --git a/Engine/Shaders/Private/HairStrands/HairStrandsVisibilityRasterHW.usf b/Engine/Shaders/Private/HairStrands/HairStrandsVisibilityRasterHW.usf
--- a/Engine/Shaders/Private/HairStrands/HairStrandsVisibilityRasterHW.usf
+++ b/Engine/Shaders/Private/HairStrands/HairStrandsVisibilityRasterHW.usf
@@ -21,7 +21,6 @@ float3				ViewDir;
 uint				HairMaterialId;
 float3				CameraOrigin;
 float				RadiusAtDepth1;
-uint				ControlPointStart;
 uint				PrimIDsBufferSize;
 uint				bIsOrthoView;
 
@@ -78,12 +77,8 @@ VSOutput VSMain(VSInput Input)
 			PrimID = VertexIndex0;
 		}
 	}
-	else
-	{
-		PrimID += ControlPointStart;
-	}
 	#else
-	const uint PrimID = ControlPointStart + Input.InstanceID;
+	const uint PrimID = Input.InstanceID;
 	#endif // PERMUTATION_CULLING
 #endif // PERMUTATION_DRAW_INDIRECT
 
diff --git a/Engine/Shaders/Private/HairStrands/HairStrandsVoxelRasterCompute.usf b/Engine/Shaders/Private/HairStrands/HairStrandsVoxelRasterCompute.usf
--- a/Engine/Shaders/Private/HairStrands/HairStrandsVoxelRasterCompute.usf
+++ b/Engine/Shaders/Private/HairStrands/HairStrandsVoxelRasterCompute.usf
@@ -13,6 +13,7 @@ uint			DispatchCountX;
 uint			MaxRasterCount;
 uint			FrameIdMod8;
 uint			MacroGroupId;
+uint			VertexCount;
 
 uint			VirtualVoxelParams_PageIndexCount;
 uint			VirtualVoxelParams_PageResolution;
@@ -24,8 +25,6 @@ StructuredBuffer<FPackedVirtualVoxelNodeDesc> VirtualVoxelParams_NodeDescBuffer;
 
 RWTexture3D<uint> OutPageTexture;
 
-uint			VertexCount;
-uint			VertexStart;
 float SampleWeight;
 
 ///////////////////////////////////////////////////////////////////////////
@@ -59,10 +58,10 @@ void MainCS(uint2 DispatchThreadID : SV_DispatchThreadID)
 #if PERMUTATION_CULLING == 1
 	uint VertexIndex0 = GetHairStrandsVertexFetchIndex_Culled(DispatchThreadID, GROUP_SIZE);
 #else
-	uint VertexIndex0 = VertexStart + GetHairStrandsVertexFetchIndex(DispatchThreadID, GROUP_SIZE, DispatchCountX);
+	uint VertexIndex0 = GetHairStrandsVertexFetchIndex(DispatchThreadID, GROUP_SIZE, DispatchCountX);
 #endif
-	uint VertexIndex1 = min(VertexIndex0 + 1, VertexStart + VertexCount - 1);
-	if (VertexIndex0 >= VertexStart + VertexCount)
+	uint VertexIndex1 = min(VertexIndex0 + 1, VertexCount - 1);
+	if (VertexIndex0 >= VertexCount)
 		return;
 
 	float LodRadiusScale0 = 1;
diff --git a/Engine/Source/Runtime/Renderer/Private/HairStrands/HairStrandsInterface.cpp b/Engine/Source/Runtime/Renderer/Private/HairStrands/HairStrandsInterface.cpp
--- a/Engine/Source/Runtime/Renderer/Private/HairStrands/HairStrandsInterface.cpp
+++ b/Engine/Source/Runtime/Renderer/Private/HairStrands/HairStrandsInterface.cpp
@@ -586,11 +586,6 @@ float GetHairVisibilityComputeRasterSampleWeight(float ScreenSize, bool bUseTemp
 	return SampleWeight;
 }
 
-uint32 FHairGroupPublicData::GetActiveStrandsVertexStart(uint32 InVertexCount) const
-{
-	return GetHairVisibilityComputeRasterVertexStart(TemporalIndex, InVertexCount);
-}
-
 uint32 FHairGroupPublicData::GetActiveStrandsVertexCount(uint32 InVertexCount, float ScreenSize) const
 {
 	return GetHairVisibilityComputeRasterVertexCount(FMath::Min(MaxScreenSize, ScreenSize), InVertexCount);
diff --git a/Engine/Source/Runtime/Renderer/Private/HairStrands/HairStrandsVisibility.cpp b/Engine/Source/Runtime/Renderer/Private/HairStrands/HairStrandsVisibility.cpp
--- a/Engine/Source/Runtime/Renderer/Private/HairStrands/HairStrandsVisibility.cpp
+++ b/Engine/Source/Runtime/Renderer/Private/HairStrands/HairStrandsVisibility.cpp
@@ -2851,9 +2851,6 @@ FHairStrandsInstanceParameters GetHairStrandsInstanceParameters(FRDGBuilder& Gra
 	LocalToTranslatedWorldTransform.AddToTranslation(TranslatedWorldOffset);
 	Out.HairStrandsVF_LocalToTranslatedWorldPrimitiveTransform = FMatrix44f(LocalToTranslatedWorldTransform.ToMatrixWithScale());
 
-	//Out.VertexCount = HairGroupPublicData->GetActiveStrandsVertexStart(VFInput.Strands.VertexCount);
-	//Out.VertexStart = HairGroupPublicData->GetActiveStrandsVertexCount(VFInput.Strands.VertexCount, ScreenSize);
-
 	if (bCullingEnable)
 	{
 		FRDGImportedBuffer CullingIndirectBuffer = Register(GraphBuilder, HairGroupPublicData->GetDrawIndirectRasterComputeBuffer(), ERDGImportedBufferFlags::CreateSRV);
@@ -2876,7 +2873,6 @@ class FVisiblityRasterClassificationCS : public FGlobalShader
 
 	BEGIN_SHADER_PARAMETER_STRUCT(FParameters, )
 		SHADER_PARAMETER_STRUCT_INCLUDE(FRasterComputeCommonParameters, Common)
-		SHADER_PARAMETER(uint32, ControlPointStart)
 		SHADER_PARAMETER(uint32, ControlPointCount)
 		SHADER_PARAMETER(uint32, PrimIDsBufferSize)
 		SHADER_PARAMETER(uint32, NumWorkGroups)
@@ -2924,7 +2920,6 @@ class FVisiblityRasterComputeBinningCS : public FGlobalShader
 		SHADER_PARAMETER_STRUCT_REF(FViewUniformShaderParameters, ViewUniformBuffer)
 		RDG_TEXTURE_ACCESS(VisTileBinningGridTex, ERHIAccess::UAVCompute)
 		SHADER_PARAMETER(uint32, VertexCount)
-		SHADER_PARAMETER(uint32, VertexStart)
 		SHADER_PARAMETER_RDG_BUFFER_SRV(ByteAddressBuffer, IndirectPrimIDCount)
 		SHADER_PARAMETER_RDG_BUFFER_SRV(Buffer, IndirectPrimIDs)
 	END_SHADER_PARAMETER_STRUCT()
@@ -3098,7 +3093,6 @@ class FVisiblityRasterHWVS : public FGlobalShader
 		SHADER_PARAMETER(uint32, HairMaterialId)
 		SHADER_PARAMETER(FVector3f, CameraOrigin)
 		SHADER_PARAMETER(float, RadiusAtDepth1)
-		SHADER_PARAMETER(uint32, ControlPointStart)
 		SHADER_PARAMETER(uint32, PrimIDsBufferSize)
 		SHADER_PARAMETER(uint32, bIsOrthoView)
 		SHADER_PARAMETER_STRUCT_INCLUDE(FHairStrandsInstanceParameters, HairInstance)
@@ -3164,7 +3158,6 @@ class FVisiblityRasterComputeNaiveCS : public FGlobalShader
 	BEGIN_SHADER_PARAMETER_STRUCT(FParameters, )
 		SHADER_PARAMETER(FVector2f, OutputResolution)
 		SHADER_PARAMETER(uint32, HairMaterialId)
-		SHADER_PARAMETER(uint32, ControlPointStart)
 		SHADER_PARAMETER(uint32, ControlPointCount)
 		SHADER_PARAMETER(float, SampleWeight)
 		SHADER_PARAMETER(uint32, NumWorkGroups)
@@ -3444,7 +3437,6 @@ static FRasterComputeOutput AddVisibilityComputeRasterPass(
 			const FSphere BoundsSphere = HairGroupPublicData->ContinuousLODBounds.GetSphere();
 			const float ScreenSize = ComputeBoundsScreenSize(FVector4(BoundsSphere.Center, 1), BoundsSphere.W, ViewInfo);
 
-			const uint32 VertexStart = HairGroupPublicData->GetActiveStrandsVertexStart(VFInput.Strands.VertexCount);
 			const uint32 VertexCount = HairGroupPublicData->GetActiveStrandsVertexCount(VFInput.Strands.VertexCount, ScreenSize);
 			const float SampleWeight = HairGroupPublicData->GetActiveStrandsSampleWeight(false, ScreenSize);
 
@@ -3461,7 +3453,6 @@ static FRasterComputeOutput AddVisibilityComputeRasterPass(
 
 				FVisiblityRasterClassificationCS::FParameters* ClassificationParameters = GraphBuilder.AllocParameters<FVisiblityRasterClassificationCS::FParameters>();
 				ClassificationParameters->Common = Common;
-				ClassificationParameters->ControlPointStart = VertexStart;
 				ClassificationParameters->ControlPointCount = VertexCount;
 				ClassificationParameters->PrimIDsBufferSize = MaxNumPrimIDs;
 				ClassificationParameters->NumWorkGroups = NumWorkGroups;
@@ -3492,7 +3483,6 @@ static FRasterComputeOutput AddVisibilityComputeRasterPass(
 				Parameters->VS.HairMaterialId = PrimitiveInfo.MaterialId;
 				Parameters->VS.CameraOrigin = FVector3f(CameraOrigin.X, CameraOrigin.Y, CameraOrigin.Z);
 				Parameters->VS.RadiusAtDepth1 = ComputeMinStrandRadiusAtDepth1(Resolution, ViewInfo.FOV, 1, -1.0f).Primary;
-				Parameters->VS.ControlPointStart = VertexStart;
 				Parameters->VS.PrimIDsBufferSize = MaxNumPrimIDs;
 				Parameters->VS.bIsOrthoView = !ViewInfo.IsPerspectiveProjection();
 				Parameters->VS.HairInstance = GetHairStrandsInstanceParameters(GraphBuilder, ViewInfo, HairGroupPublicData, bCullingEnable, bForceRegister);
@@ -3553,7 +3543,6 @@ static FRasterComputeOutput AddVisibilityComputeRasterPass(
 					FVisiblityRasterComputeNaiveCS::FParameters* RasterParameters = GraphBuilder.AllocParameters<FVisiblityRasterComputeNaiveCS::FParameters>();
 					RasterParameters->OutputResolution = Common.OutputResolutionf;
 					RasterParameters->HairMaterialId = PrimitiveInfo.MaterialId;
-					RasterParameters->ControlPointStart = VertexStart;
 					RasterParameters->ControlPointCount = VertexCount;
 					RasterParameters->SampleWeight = SampleWeight;
 					RasterParameters->NumWorkGroups = NumWorkGroups;
@@ -3601,7 +3590,6 @@ static FRasterComputeOutput AddVisibilityComputeRasterPass(
 						BinningParameters->OutVisTileData = GraphBuilder.CreateUAV(VisTileData);
 						BinningParameters->HairInstance = GetHairStrandsInstanceParameters(GraphBuilder, ViewInfo, HairGroupPublicData, bCullingEnableInRasterizers, bForceRegister);
 						BinningParameters->VertexCount = VertexCount;
-						BinningParameters->VertexStart = VertexStart;
 						BinningParameters->VisTileBinningGridTex = VisTileBinningGrid;
 						BinningParameters->IndirectPrimIDCount = bClassification ? SWRasterPrimIDCountSRV : nullptr;
 						BinningParameters->IndirectPrimIDs = bClassification ? RasterizerPrimIDsSRV : nullptr;
diff --git a/Engine/Source/Runtime/Renderer/Private/HairStrands/HairStrandsVoxelization.cpp b/Engine/Source/Runtime/Renderer/Private/HairStrands/HairStrandsVoxelization.cpp
--- a/Engine/Source/Runtime/Renderer/Private/HairStrands/HairStrandsVoxelization.cpp
+++ b/Engine/Source/Runtime/Renderer/Private/HairStrands/HairStrandsVoxelization.cpp
@@ -1237,7 +1237,6 @@ class FVoxelRasterComputeCS : public FGlobalShader
 		SHADER_PARAMETER_RDG_BUFFER_SRV(Buffer, VoxelizationViewInfoBuffer)
 		SHADER_PARAMETER_RDG_TEXTURE_UAV(RWTexture3D, OutPageTexture)
 		SHADER_PARAMETER(uint32, VertexCount)
-		SHADER_PARAMETER(uint32, VertexStart)
 		SHADER_PARAMETER(float, SampleWeight)
 	END_SHADER_PARAMETER_STRUCT()
 
@@ -1299,7 +1298,6 @@ static void AddVirtualVoxelizationComputeRasterPass(
 					continue;
 				}
 
-				const uint32 VertexStart = HairGroupPublicData->GetActiveStrandsVertexStart(VFInput.Strands.VertexCount);
 				const uint32 VertexCount = HairGroupPublicData->GetActiveStrandsVertexCount(VFInput.Strands.VertexCount, HairGroupPublicData->MaxScreenSize);
 				const float SampleWeight = HairGroupPublicData->GetActiveStrandsSampleWeight(false, HairGroupPublicData->MaxScreenSize);
 
@@ -1316,7 +1314,6 @@ static void AddVirtualVoxelizationComputeRasterPass(
 				PassParameters->FrameIdMod8 = FrameIdMode8;
 
 				PassParameters->VertexCount = VertexCount;
-				PassParameters->VertexStart = VertexStart;
 				PassParameters->SampleWeight = SampleWeight;
 			
 				const bool bCullingEnable = !IsHairStrandContinuousDecimationReorderingEnabled() && HairGroupPublicData->GetCullingResultAvailable();
diff --git a/Engine/Source/Runtime/Renderer/Public/HairStrandsInterface.h b/Engine/Source/Runtime/Renderer/Public/HairStrandsInterface.h
--- a/Engine/Source/Runtime/Renderer/Public/HairStrandsInterface.h
+++ b/Engine/Source/Runtime/Renderer/Public/HairStrandsInterface.h
@@ -196,7 +196,6 @@ public:
 
 	uint32 GetClusterCount() const { return ClusterCount;  }
 
-	uint32 GetActiveStrandsVertexStart(uint32 InVertexCount) const;
 	uint32 GetActiveStrandsVertexCount(uint32 InVertexCount, float ScreenSize) const;
 	float GetActiveStrandsSampleWeight(bool bUseTemporalWeight, float ScreenSize) const;
 
