Remove voxel indirect scatter allocation as it is no longer used.

diff --git a/Engine/Shaders/Private/HairStrands/HairStrandsVoxelRasterCompute.usf b/Engine/Shaders/Private/HairStrands/HairStrandsVoxelRasterCompute.usf
--- a/Engine/Shaders/Private/HairStrands/HairStrandsVoxelRasterCompute.usf
+++ b/Engine/Shaders/Private/HairStrands/HairStrandsVoxelRasterCompute.usf
@@ -222,72 +222,11 @@ void AllocatePageIndex(uint2 DispatchThreadId : SV_DispatchThreadID)
 }
 #endif
 
-///////////////////////////////////////////////////////////////////////////////////////////////////
-// Common struct/functions
-
-#if SHADER_MARKVALID_PREPARE || SHADER_MARKVALID_INDIRECTARG || SHADER_MARKVALID_SCATTER
-struct FScatterCoord
-{
-	uint3 MinCoord;
-	uint  PartIndex;
-	uint3 Resolution;
-	uint  MacroGroupId;
-};
-
-uint2 PackScatterCoord(FScatterCoord In)
-{
-	uint2 Out;
-	Out.x =
-		((In.MinCoord.x & 0xFF)) |
-		((In.MinCoord.y & 0xFF) << 8) |
-		((In.MinCoord.z & 0xFF) << 16) |
-		((In.PartIndex & 0xFF) << 24);
-
-	Out.y =
-		((In.Resolution.x & 0xFF)) |
-		((In.Resolution.y & 0xFF) << 8) |
-		((In.Resolution.z & 0xFF) << 16) |
-		((In.MacroGroupId & 0xFF) << 24);
-
-	return Out;
-}
-
-FScatterCoord UnpackScatterCoord(uint2 In)
-{
-	FScatterCoord Out;
-	Out.MinCoord = uint3(
-		In.x & 0xFF,
-		(In.x >> 8) & 0xFF,
-		(In.x >> 16) & 0xFF);
-	Out.PartIndex = (In.x >> 24) & 0xFF;
-
-	Out.Resolution = uint3(
-		In.y & 0xFF,
-		(In.y >> 8) & 0xFF,
-		(In.y >> 16) & 0xFF);
-	Out.MacroGroupId = (In.y >> 24) & 0xFF;
-
-	return Out;
-}
-
-void UpdatePackScatterCoord(inout uint2 Out, uint PartIndex)
-{
-	Out.x = (Out.x & 0x00FFFFFF) | ((PartIndex & 0xFF) << 24);
-}
-
-#define SCATTER_WORK_ITEM_COUNT 8
-#define SCATTER_GROUP_SIZE 32
-#define SCATTER_GROUP_LINE 4
-
-#endif // SHADER_MARKVALID_PREPARE || SHADER_MARKVALID_INDIRECTARG || SHADER_MARKVALID_SCATTER
-
 ///////////////////////////////////////////////////////////////////////////////////////////////////
 
 #if SHADER_MARKVALID_PREPARE
 uint			MaxClusterCount;
 uint			MacroGroupId;
-uint			MaxScatterAllocationCount;
-uint			bForceDirectPageAllocation;
 
 Buffer<int>		GroupAABBsBuffer;
 Buffer<int>		ClusterAABBsBuffer;
@@ -295,15 +234,12 @@ Buffer<int>		MacroGroupVoxelAlignedAABBBuffer;
 Buffer<uint4>	PageIndexResolutionAndOffsetBuffer;
 
 RWBuffer<uint>	OutValidPageIndexBuffer;
-RWBuffer<uint>  OutDeferredScatterCounter;
-RWBuffer<uint2> OutDeferredScatterBuffer;
-
 
 // PageIndexBuffer      is sampled with linear coordinate computed from the 3d page coordinate. VALID NODE ARE NOT COMPACTED. It contains the LINEAR PAGE INDEX (to map to the 3d volume).
 // PageIndexCoordBuffer is sampled with linear coordinate for allocated nodes. VALID NODE ARE COMPACTED. It contains the 3d page coordinate and ClustedId. Only used for opaque voxel injection.
 
 #if PERMUTATION_USE_CLUSTER
-[numthreads(SCATTER_GROUP_SIZE, 1, 1)]
+[numthreads(32, 1, 1)]
 void MarkValid_PrepareCS(uint2 DispatchThreadId : SV_DispatchThreadID)
 {
 	const uint ClusterIndex = DispatchThreadId.x;
@@ -346,13 +282,8 @@ void MarkValid_PrepareCS(uint2 DispatchThreadId : SV_DispatchThreadID)
 	MinCoord = clamp(MinCoord, uint3(0, 0, 0), PageIndexResolutionMinusOne);
 	MaxCoord = clamp(MaxCoord, uint3(0, 0, 0), PageIndexResolutionMinusOne);
 
-	FScatterCoord Coord;
-	Coord.MacroGroupId = MacroGroupId;
-	Coord.PartIndex = 0;
-	Coord.MinCoord = MinCoord;
-	Coord.Resolution = (MaxCoord - MinCoord) + 1;
-
-	const uint ScatterCount = Coord.Resolution.x * Coord.Resolution.y * Coord.Resolution.z;
+	const uint3 CoordResolution = (MaxCoord - MinCoord) + 1;
+	const uint ScatterCount = CoordResolution.x * CoordResolution.y * CoordResolution.z;
 
 	// Arbitrary large number (e.g., 100x10x10 pages covered)
 	// This acts as guards against degenerated case, where the sim would deformed strands large position making the cluster arbitratry large.
@@ -366,30 +297,15 @@ void MarkValid_PrepareCS(uint2 DispatchThreadId : SV_DispatchThreadID)
 		return;
 
 	// Find a good sweet spot
-	if (ScatterCount <= 4 || bForceDirectPageAllocation > 0)
-	{
-		for (uint z = MinCoord.z; z <= MaxCoord.z; ++z)
-			for (uint y = MinCoord.y; y <= MaxCoord.y; ++y)
-				for (uint x = MinCoord.x; x <= MaxCoord.x; ++x)
-				{
-					const uint3 PageIndexCoord = uint3(x, y, z);
-					const uint LinearPageIndexCoord = CoordToIndex(PageIndexCoord, PageIndexResolution, PageIndexOffset);
-					InterlockedOr(OutValidPageIndexBuffer[LinearPageIndexCoord], 1u);
-				}
-	}
-	else
+	for (uint z = MinCoord.z; z <= MaxCoord.z; ++z)
 	{
-		const uint WorkItemCount = (ScatterCount + SCATTER_WORK_ITEM_COUNT - 1) / SCATTER_WORK_ITEM_COUNT;
-
-		uint AllocationOffset = 0;
-		InterlockedAdd(OutDeferredScatterCounter[0], WorkItemCount, AllocationOffset);
-		if (WorkItemCount + AllocationOffset < MaxScatterAllocationCount)
+		for (uint y = MinCoord.y; y <= MaxCoord.y; ++y)
 		{
-			uint2 PackedCoord = PackScatterCoord(Coord);
-			for (uint Item = 0; Item < WorkItemCount; ++Item)
+			for (uint x = MinCoord.x; x <= MaxCoord.x; ++x)
 			{
-				UpdatePackScatterCoord(PackedCoord, Item);
-				OutDeferredScatterBuffer[AllocationOffset + Item] = PackedCoord;
+				const uint3 PageIndexCoord = uint3(x, y, z);
+				const uint LinearPageIndexCoord = CoordToIndex(PageIndexCoord, PageIndexResolution, PageIndexOffset);
+				InterlockedOr(OutValidPageIndexBuffer[LinearPageIndexCoord], 1u);
 			}
 		}
 	}
@@ -423,74 +339,12 @@ void MarkValid_PrepareCS(uint3 DispatchThreadId : SV_DispatchThreadID)
 		const uint LinearPageIndexCoord = CoordToIndex(PageIndexCoord, PageIndexResolution, PageIndexOffset);
 		InterlockedOr(OutValidPageIndexBuffer[LinearPageIndexCoord], 1u);
 	}
-
-	// Fake writting to insure RDG allocate buffers even there is only Group-based allocated elements (vs. cluster-based allocated elements)
-	if (Coord.x > 65000)
-	{
-		OutDeferredScatterCounter[0] = 0;
-		OutDeferredScatterBuffer[0] = uint2(0, 0);
-	}
 }
 #endif // PERMUTATION_USE_CLUSTER
 #endif // SHADER_MARKVALID_PREPARE
 
 ///////////////////////////////////////////////////////////////////////////////////////////////////
 
-#if SHADER_MARKVALID_INDIRECTARG
-Buffer<uint> DeferredScatterCounter;
-RWBuffer<uint> OutIndirectArgsBuffer;
-
-[numthreads(1, 1, 1)]
-void MarkValid_BuildIndirectArgCS(uint2 DispatchThreadId : SV_DispatchThreadID)
-{
-	const uint ScatterItemCount = DeferredScatterCounter[0] * SCATTER_WORK_ITEM_COUNT;
-	const uint DispatchCount = (ScatterItemCount + SCATTER_GROUP_SIZE - 1) / SCATTER_GROUP_SIZE;
-	const uint DispatchX = min(DispatchCount, uint(SCATTER_GROUP_LINE));
-	const uint DispatchY = (DispatchCount + SCATTER_GROUP_LINE - 1) / SCATTER_GROUP_LINE;
-
-	WriteDispatchIndirectArgs(OutIndirectArgsBuffer, 0, DispatchX, DispatchY, 1);
-}
-#endif
-
-///////////////////////////////////////////////////////////////////////////////////////////////////
-
-#if SHADER_MARKVALID_SCATTER
-Buffer<int>		IndirectBufferArgs;
-Buffer<uint4>	PageIndexResolutionAndOffsetBuffer;
-Buffer<uint>	DeferredScatterCounter;
-Buffer<uint2>	DeferredScatterBuffer;
-RWBuffer<uint>	OutValidPageIndexBuffer;
-
-[numthreads(SCATTER_GROUP_SIZE, 1, 1)]
-void MarkValid_ScatterCS(uint2 DispatchThreadId : SV_DispatchThreadID)
-{
-	const uint LinearThreadId = DispatchThreadId.x + DispatchThreadId.y * SCATTER_GROUP_SIZE * SCATTER_GROUP_LINE;
-	const uint ScatterItem = LinearThreadId;
-	const uint WorkItemIndex = LinearThreadId / SCATTER_WORK_ITEM_COUNT;
-	const uint ItemIndex = LinearThreadId % SCATTER_WORK_ITEM_COUNT;
-	const uint MaxScatterItemCount = DeferredScatterCounter[0] * SCATTER_WORK_ITEM_COUNT;
-	if (ScatterItem < MaxScatterItemCount)
-	{
-		const FScatterCoord Coord = UnpackScatterCoord(DeferredScatterBuffer[WorkItemIndex]);
-
-		const uint4 PageIndexResolutionAndOffset = PageIndexResolutionAndOffsetBuffer.Load(Coord.MacroGroupId);
-		const int3  PageIndexResolution = PageIndexResolutionAndOffset.xyz;
-		const uint  PageIndexOffset = PageIndexResolutionAndOffset.w;
-
-		const uint  ClusterPageCount = Coord.Resolution.x * Coord.Resolution.y * Coord.Resolution.z;
-		const uint  LinearVoxelIt = (Coord.PartIndex * SCATTER_WORK_ITEM_COUNT) + ItemIndex;
-		if (LinearVoxelIt < ClusterPageCount)
-		{
-			const uint3 PageIndexCoord = IndexToCoord(LinearVoxelIt, Coord.Resolution) + Coord.MinCoord;
-			const uint LinearPageIndexCoord = CoordToIndex(PageIndexCoord, PageIndexResolution, PageIndexOffset);
-			InterlockedOr(OutValidPageIndexBuffer[LinearPageIndexCoord], 1u);
-		}
-	}
-}
-#endif
-
-///////////////////////////////////////////////////////////////////////////////////////////////////
-
 #if SHADER_MARKVALID
 uint			MaxClusterCount;
 int3			CPU_PageIndexResolution;
@@ -554,13 +408,17 @@ void MarkValidCS(uint2 DispatchThreadId : SV_DispatchThreadID)
 	// Splat cluster contribution
 	// Change this to be a bit smarter (use graphics pipe?)
 	for (uint z = MinCoord.z; z <= MaxCoord.z; ++z)
+	{
 		for (uint y = MinCoord.y; y <= MaxCoord.y; ++y)
+		{
 			for (uint x = MinCoord.x; x <= MaxCoord.x; ++x)
 			{
 				const uint3 PageIndexCoord = uint3(x, y, z);
 				const uint LinearPageIndexCoord = CoordToIndex(PageIndexCoord, PageIndexResolution, PageIndexOffset);
 				InterlockedOr(OutValidPageIndexBuffer[LinearPageIndexCoord], 1u);
 			}
+		}
+	}
 }
 #endif
 
diff --git a/Engine/Source/Runtime/Renderer/Private/HairStrands/HairStrandsVoxelization.cpp b/Engine/Source/Runtime/Renderer/Private/HairStrands/HairStrandsVoxelization.cpp
--- a/Engine/Source/Runtime/Renderer/Private/HairStrands/HairStrandsVoxelization.cpp
+++ b/Engine/Source/Runtime/Renderer/Private/HairStrands/HairStrandsVoxelization.cpp
@@ -102,9 +102,6 @@ static FAutoConsoleVariableRef CVarHairVirtualVoxelInvalidEmptyPageIndex(TEXT("r
 static int32 GHairStrandsVoxelComputeRasterMaxVoxelCount = 32;
 static FAutoConsoleVariableRef CVarHairStrandsVoxelComputeRasterMaxVoxelCount(TEXT("r.HairStrands.Voxelization.Virtual.ComputeRasterMaxVoxelCount"), GHairStrandsVoxelComputeRasterMaxVoxelCount, TEXT("Max number of voxel which are rasterized for a given hair segment. This is for debug purpose only."));
 
-static int32 GHairVirtualVoxelUseImmediatePageAllocation = 1;
-static FAutoConsoleVariableRef CVarHairVirtualVoxelUseImmediatePageAllocation(TEXT("r.HairStrands.Voxelization.Virtual.UseDirectPageAllocation"), GHairVirtualVoxelUseImmediatePageAllocation, TEXT("Use the indirect page allocation code path, but force internally direct page allocation (for debugging purpose only)."));
-
 static float GHairVirtualVoxelRaytracing_ShadowOcclusionThreshold= 1;
 static float GHairVirtualVoxelRaytracing_SkyOcclusionThreshold = 1;
 static FAutoConsoleVariableRef CVarHairVirtualVoxelRaytracing_ShadowOcclusionThreshold(TEXT("r.RayTracing.Shadows.HairOcclusionThreshold"), GHairVirtualVoxelRaytracing_ShadowOcclusionThreshold, TEXT("Define the number of hair that need to be crossed, before casting occlusion (default = 1)"), ECVF_RenderThreadSafe);
@@ -285,47 +282,6 @@ public:
 	}
 };
 
-class FVoxelMarkValidPageIndex_IndirectArgsCS : public FGlobalShader
-{
-	DECLARE_GLOBAL_SHADER(FVoxelMarkValidPageIndex_IndirectArgsCS);
-	SHADER_USE_PARAMETER_STRUCT(FVoxelMarkValidPageIndex_IndirectArgsCS, FGlobalShader);
-
-	BEGIN_SHADER_PARAMETER_STRUCT(FParameters, )
-		SHADER_PARAMETER_RDG_BUFFER_SRV(Buffer<uint>, DeferredScatterCounter)
-		SHADER_PARAMETER_RDG_BUFFER_UAV(RWBuffer<uint>, OutIndirectArgsBuffer)
-	END_SHADER_PARAMETER_STRUCT()
-
-public:
-	static bool ShouldCompilePermutation(const FGlobalShaderPermutationParameters& Parameters) { return IsHairStrandsSupported(EHairStrandsShaderType::Strands, Parameters.Platform); }
-	static void ModifyCompilationEnvironment(const FGlobalShaderPermutationParameters& Parameters, FShaderCompilerEnvironment& OutEnvironment)
-	{
-		FGlobalShader::ModifyCompilationEnvironment(Parameters, OutEnvironment);
-		OutEnvironment.SetDefine(TEXT("SHADER_MARKVALID_INDIRECTARG"), 1);
-	}
-};
-
-class FVoxelMarkValidPageIndex_ScatterCS : public FGlobalShader
-{
-	DECLARE_GLOBAL_SHADER(FVoxelMarkValidPageIndex_ScatterCS);
-	SHADER_USE_PARAMETER_STRUCT(FVoxelMarkValidPageIndex_ScatterCS, FGlobalShader);
-
-	BEGIN_SHADER_PARAMETER_STRUCT(FParameters, )
-		RDG_BUFFER_ACCESS(IndirectBufferArgs, ERHIAccess::IndirectArgs)
-		SHADER_PARAMETER_RDG_BUFFER_SRV(Buffer<uint4>,	PageIndexResolutionAndOffsetBuffer)
-		SHADER_PARAMETER_RDG_BUFFER_SRV(Buffer<uint>,	DeferredScatterCounter)
-		SHADER_PARAMETER_RDG_BUFFER_SRV(Buffer<uint2>,  DeferredScatterBuffer)
-		SHADER_PARAMETER_RDG_BUFFER_UAV(RWBuffer<uint>, OutValidPageIndexBuffer)
-	END_SHADER_PARAMETER_STRUCT()
-
-public:
-	static bool ShouldCompilePermutation(const FGlobalShaderPermutationParameters& Parameters) { return IsHairStrandsSupported(EHairStrandsShaderType::Strands, Parameters.Platform); }
-	static void ModifyCompilationEnvironment(const FGlobalShaderPermutationParameters& Parameters, FShaderCompilerEnvironment& OutEnvironment)
-	{
-		FGlobalShader::ModifyCompilationEnvironment(Parameters, OutEnvironment);
-		OutEnvironment.SetDefine(TEXT("SHADER_MARKVALID_SCATTER"), 1);
-	}
-};
-
 class FVoxelMarkValidPageIndexCS : public FGlobalShader
 {
 	DECLARE_GLOBAL_SHADER(FVoxelMarkValidPageIndexCS);
@@ -511,11 +467,8 @@ public:
 };
 
 IMPLEMENT_GLOBAL_SHADER(FVoxelMarkValidPageIndex_PrepareCS, "/Engine/Private/HairStrands/HairStrandsVoxelRasterCompute.usf", "MarkValid_PrepareCS", SF_Compute);
-IMPLEMENT_GLOBAL_SHADER(FVoxelMarkValidPageIndex_IndirectArgsCS, "/Engine/Private/HairStrands/HairStrandsVoxelRasterCompute.usf", "MarkValid_BuildIndirectArgCS", SF_Compute);
-IMPLEMENT_GLOBAL_SHADER(FVoxelMarkValidPageIndex_ScatterCS, "/Engine/Private/HairStrands/HairStrandsVoxelRasterCompute.usf", "MarkValid_ScatterCS", SF_Compute);
-
-IMPLEMENT_GLOBAL_SHADER(FVoxelAllocatePageIndexCS, "/Engine/Private/HairStrands/HairStrandsVoxelRasterCompute.usf", "AllocatePageIndex", SF_Compute);
 IMPLEMENT_GLOBAL_SHADER(FVoxelMarkValidPageIndexCS, "/Engine/Private/HairStrands/HairStrandsVoxelRasterCompute.usf", "MarkValidCS", SF_Compute);
+IMPLEMENT_GLOBAL_SHADER(FVoxelAllocatePageIndexCS, "/Engine/Private/HairStrands/HairStrandsVoxelRasterCompute.usf", "AllocatePageIndex", SF_Compute);
 IMPLEMENT_GLOBAL_SHADER(FVoxelAllocateVoxelPageCS, "/Engine/Private/HairStrands/HairStrandsVoxelRasterCompute.usf", "AllocateCS", SF_Compute);
 IMPLEMENT_GLOBAL_SHADER(FVoxelAddNodeDescCS, "/Engine/Private/HairStrands/HairStrandsVoxelRasterCompute.usf", "AddDescCS", SF_Compute);
 IMPLEMENT_GLOBAL_SHADER(FVoxelAddIndirectBufferCS, "/Engine/Private/HairStrands/HairStrandsVoxelRasterCompute.usf", "AddIndirectBufferCS", SF_Compute);
@@ -715,18 +668,6 @@ static void AddAllocateVoxelPagesPass(
 		const bool bUseIndirectScatter = GHairVirtualVoxelUseIndirectScatterPageAllocation>0 && bIsGPUDriven;
 		if (bUseIndirectScatter)
 		{			
-			const uint32 AverageWorkItemPerCluster = 8; // Arbitrary/Guess number
-			const uint32 MaxAllocationCount = TotalClusterCount * AverageWorkItemPerCluster;
-			FRDGBufferRef ScatterCounter = GraphBuilder.CreateBuffer(FRDGBufferDesc::CreateBufferDesc(sizeof(uint32), 1), TEXT("Hair.PageScatterCounter"));
-			FRDGBufferRef ScatterBuffer = GraphBuilder.CreateBuffer(FRDGBufferDesc::CreateBufferDesc(2 * sizeof(uint32), MaxAllocationCount), TEXT("Hair.PageScatterBuffer"));
-
-			FRDGBufferUAVRef ScatterCounterUAV = GraphBuilder.CreateUAV(ScatterCounter, PF_R32_UINT);
-			FRDGBufferUAVRef ScatterBufferUAV  = GraphBuilder.CreateUAV(ScatterBuffer, PF_R32G32_UINT);
-
-			AddClearUAVPass(GraphBuilder, ScatterCounterUAV, 0);
-
-			const uint32 bForceDirectPageAllocation = GHairVirtualVoxelUseImmediatePageAllocation > 0 ? 1 : 0;
-			// Prepare
 			for (const FHairStrandsMacroGroupData::PrimitiveInfo& PrimitiveInfo : MacroGroup.PrimitivesInfos)
 			{
 				FHairGroupPublicData* HairGroupData = PrimitiveInfo.PublicDataPtr;
@@ -744,16 +685,12 @@ static void AddAllocateVoxelPagesPass(
 				FVoxelMarkValidPageIndex_PrepareCS::FParameters* Parameters = GraphBuilder.AllocParameters<FVoxelMarkValidPageIndex_PrepareCS::FParameters>();
 				Parameters->MaxClusterCount				= HairGroupData->GetClusterCount();
 				Parameters->MacroGroupId				= MacroGroup.MacroGroupId;
-				Parameters->MaxScatterAllocationCount	= MaxAllocationCount;
-				Parameters->bForceDirectPageAllocation	= bForceDirectPageAllocation;
 
 				Parameters->GroupAABBsBuffer			= RegisterAsSRV(GraphBuilder, HairGroupData->GetGroupAABBBuffer());
 				Parameters->ClusterAABBsBuffer			= RegisterAsSRV(GraphBuilder, HairGroupData->GetClusterAABBBuffer());
 				Parameters->MacroGroupVoxelAlignedAABBBuffer = GraphBuilder.CreateSRV(MacroGroupResources.MacroGroupVoxelAlignedAABBsBuffer, PF_R32_SINT);
 				Parameters->PageIndexResolutionAndOffsetBuffer = PageIndexResolutionAndOffsetBufferSRV;
 
-				Parameters->OutDeferredScatterCounter	= ScatterCounterUAV;
-				Parameters->OutDeferredScatterBuffer	= ScatterBufferUAV;
 				Parameters->OutValidPageIndexBuffer		= PageIndexBufferUAV;
 
 				FIntVector DispatchCount = bUseClusterAABB ?
@@ -769,48 +706,13 @@ static void AddAllocateVoxelPagesPass(
 				TShaderMapRef<FVoxelMarkValidPageIndex_PrepareCS> ComputeShader(View.ShaderMap, PermutationVector);
 				FComputeShaderUtils::AddPass(
 					GraphBuilder,
-					RDG_EVENT_NAME("HairStrands::MarkValidPageIndex_Prepare(%s,%s)", bUseClusterAABB ? TEXT("Cluster") : TEXT("Group"), bForceDirectPageAllocation ? TEXT("Direct") : TEXT("Indirect")),
+					RDG_EVENT_NAME("HairStrands::MarkValidPageIndex_Prepare(%s)", bUseClusterAABB ? TEXT("Cluster") : TEXT("Group")),
 					ComputeShader,
 					Parameters,
 					DispatchCount);
 			}
-
-			if (!bForceDirectPageAllocation)
-			{
-				FRDGBufferSRVRef ScatterCounterSRV = GraphBuilder.CreateSRV(ScatterCounter, PF_R32_UINT);
-
-				// Build indirect buffer args
-				FRDGBufferRef ScatterIndirectArgsBuffer = GraphBuilder.CreateBuffer(FRDGBufferDesc::CreateIndirectDesc<FRHIDispatchIndirectParameters>(1), TEXT("Hair.PageScatterIndirectArgs"));
-				{
-					check(MacroGroup.MacroGroupId < MacroGroupCount);
-
-					FVoxelMarkValidPageIndex_IndirectArgsCS::FParameters* Parameters = GraphBuilder.AllocParameters<FVoxelMarkValidPageIndex_IndirectArgsCS::FParameters>();
-					Parameters->DeferredScatterCounter = ScatterCounterSRV;
-					Parameters->OutIndirectArgsBuffer = GraphBuilder.CreateUAV(ScatterIndirectArgsBuffer);
-
-					const FIntVector DispatchCount(1, 1, 1);
-					TShaderMapRef<FVoxelMarkValidPageIndex_IndirectArgsCS> ComputeShader(View.ShaderMap);
-					FComputeShaderUtils::AddPass(GraphBuilder, RDG_EVENT_NAME("HairStrands::MarkValidPageIndex_IndirectArgs"), ComputeShader, Parameters, DispatchCount);
-				}
-
-				// Scatter
-				{
-					FRDGBufferSRVRef ScatterBufferSRV = GraphBuilder.CreateSRV(ScatterBuffer, PF_R32G32_UINT);
-					check(MacroGroup.MacroGroupId < MacroGroupCount);
-
-					FVoxelMarkValidPageIndex_ScatterCS::FParameters* Parameters = GraphBuilder.AllocParameters<FVoxelMarkValidPageIndex_ScatterCS::FParameters>();
-					Parameters->IndirectBufferArgs = ScatterIndirectArgsBuffer;
-					Parameters->PageIndexResolutionAndOffsetBuffer = PageIndexResolutionAndOffsetBufferSRV;
-					Parameters->DeferredScatterCounter = ScatterCounterSRV;
-					Parameters->DeferredScatterBuffer = ScatterBufferSRV;
-					Parameters->OutValidPageIndexBuffer = PageIndexBufferUAV;
-
-					TShaderMapRef<FVoxelMarkValidPageIndex_ScatterCS> ComputeShader(View.ShaderMap);
-					FComputeShaderUtils::AddPass(GraphBuilder, RDG_EVENT_NAME("HairStrands::MarkValidPageIndex_Scatter"), ComputeShader, Parameters, ScatterIndirectArgsBuffer, 0);
-				}
-			}
 		}
-		else
+		else // TODO remove this path, as the FVoxelMarkValidPageIndex_PrepareCS(FUseCluster=false) handles the same path)
 		{		
 			for (const FHairStrandsMacroGroupData::PrimitiveInfo& PrimitiveInfo : MacroGroup.PrimitivesInfos)
 			{
