#10073

diff --git a/Engine/Source/Runtime/Renderer/Private/VT/VirtualTextureSystem.cpp b/Engine/Source/Runtime/Renderer/Private/VT/VirtualTextureSystem.cpp
--- a/Engine/Source/Runtime/Renderer/Private/VT/VirtualTextureSystem.cpp
+++ b/Engine/Source/Runtime/Renderer/Private/VT/VirtualTextureSystem.cpp
@@ -477,7 +477,7 @@ void FVirtualTextureSystem::ListPhysicalPoolsFromConsole()
 			const FTexturePagePool& PagePool = PhysicalSpace.GetPagePool();
 			const uint32 TotalSizeInBytes = PhysicalSpace.GetSizeInBytes();
 
-			UE_LOG(LogConsoleResponse, Display, TEXT("PhysicaPool: [%i] %s (%ix%i):"), i, *PhysicalSpace.GetFormatString(), Desc.TileSize, Desc.TileSize);
+			UE_LOG(LogConsoleResponse, Display, TEXT("PhysicalPool: [%i] %s (%ix%i):"), i, *PhysicalSpace.GetFormatString(), Desc.TileSize, Desc.TileSize);
 			
 			const int32 AllocatedTiles = PagePool.GetNumAllocatedPages();
 			const float AllocatedLoad = (float)AllocatedTiles / (float)PhysicalSpace.GetNumTiles();
@@ -549,7 +549,7 @@ void FVirtualTextureSystem::DumpPoolUsageFromConsole()
 			const FVTPhysicalSpaceDescription& Desc = PhysicalSpace.GetDescription();
 			const FTexturePagePool& PagePool = PhysicalSpace.GetPagePool();
 
-			UE_LOG(LogConsoleResponse, Display, TEXT("PhysicaPool: [%i] %s (%ix%i):"), i, *PhysicalSpace.GetFormatString(), Desc.TileSize, Desc.TileSize);
+			UE_LOG(LogConsoleResponse, Display, TEXT("PhysicalPool: [%i] %s (%ix%i):"), i, *PhysicalSpace.GetFormatString(), Desc.TileSize, Desc.TileSize);
 
 			TMap<uint32, uint32> ProducerCountMap;
 			PagePool.CollectProducerCounts(ProducerCountMap);
