Fix for virtual texture flush crashing in CitySample

diff --git a/Engine/Source/Runtime/Renderer/Private/RendererScene.cpp b/Engine/Source/Runtime/Renderer/Private/RendererScene.cpp
--- a/Engine/Source/Runtime/Renderer/Private/RendererScene.cpp
+++ b/Engine/Source/Runtime/Renderer/Private/RendererScene.cpp
@@ -6094,7 +6094,7 @@ void FScene::UpdateAllPrimitiveSceneInfos(FRDGBuilder& GraphBuilder, EUpdateAllP
 			FScopeCycleCounter Context(PrimitiveSceneProxy->GetStatId());
 			FPrimitiveSceneInfo* PrimitiveSceneInfo = PrimitiveSceneProxy->GetPrimitiveSceneInfo();
 			
-			PrimitiveSceneInfo->FlushRuntimeVirtualTexture();
+			QueueSceneInfoToFlushVirtualTexture(PrimitiveSceneInfo);
 
 			// If we recorded no adds or removes the instance count has stayed the same.  Therefore the cached mesh draw commands do not
 			// need to be updated.  In situations where the instance count changes the mesh draw command stores the instance count which
