Change hair voxel dispatching.

This changes change dispatching to be 1D along allocated page index, adds logic to ensure we never run over maximum dispatch dimension, and ensure group size are checked if changed in the future. Performanc remains identical.

diff --git a/Engine/Shaders/Private/HairStrands/HairStrandsVisibilityCommon.ush b/Engine/Shaders/Private/HairStrands/HairStrandsVisibilityCommon.ush
--- a/Engine/Shaders/Private/HairStrands/HairStrandsVisibilityCommon.ush
+++ b/Engine/Shaders/Private/HairStrands/HairStrandsVisibilityCommon.ush
@@ -13,8 +13,7 @@
 
 #define HAIR_VISIBILITY_GROUP_COUNT_WIDTH 64
 
-#define MAX_HAIR_MACROGROUP_COUNT 16
-
+#include "/Engine/Shared/HairStrandsDefinitions.h"
 #include "HairStrandsVisibilityCommonStruct.ush"
 
  ////////////////////////////////////////////////////////////////////////////////////////////////////////////////////////////////
diff --git a/Engine/Shaders/Private/HairStrands/HairStrandsVoxelRasterCompute.usf b/Engine/Shaders/Private/HairStrands/HairStrandsVoxelRasterCompute.usf
--- a/Engine/Shaders/Private/HairStrands/HairStrandsVoxelRasterCompute.usf
+++ b/Engine/Shaders/Private/HairStrands/HairStrandsVoxelRasterCompute.usf
@@ -44,8 +44,9 @@ RWBuffer<uint4>								OutPageIndexResolutionAndOffsetBuffer;
 RWBuffer<uint>								OutPageIndexAllocationIndirectBufferArgs;
 RWStructuredBuffer<FVoxelizationViewInfo>	OutVoxelizationViewInfoBuffer;
 
-#ifndef MAX_HAIR_MACROGROUP_COUNT
-#error MAX_HAIR_MACROGROUP_COUNT needs to be defined
+
+#if GROUP_SIZE != MAX_HAIR_MACROGROUP_COUNT
+#error MAX_HAIR_MACROGROUP_COUNT needs to match MAX_HAIR_MACROGROUP_COUNT
 #endif
 
 #define INVALID_OFFSET 0xFFFFFFFF
@@ -53,7 +54,7 @@ RWStructuredBuffer<FVoxelizationViewInfo>	OutVoxelizationViewInfoBuffer;
 groupshared uint PageIndexOffsets[MAX_HAIR_MACROGROUP_COUNT];
 
 // This code assume we have less than 32 macro group (which fit into a single CU/SM)
-[numthreads(MAX_HAIR_MACROGROUP_COUNT, 1, 1)]
+[numthreads(GROUP_SIZE, 1, 1)]
 void AllocatePageIndex(uint2 DispatchThreadId : SV_DispatchThreadID)
 {
 	const uint MacroGroupId = DispatchThreadId.x;
@@ -154,7 +155,7 @@ void AllocatePageIndex(uint2 DispatchThreadId : SV_DispatchThreadID)
 	// Prepare indirect buffer for doing the actual page index allocation and filling the page index
 	{
 		const uint AllocatedPageIndexCount = PageIndexResolution.x * PageIndexResolution.y * PageIndexResolution.z;
-		WriteDispatchIndirectArgs(OutPageIndexAllocationIndirectBufferArgs, MacroGroupId, ceil(AllocatedPageIndexCount / float(IndirectDispatchGroupSize)), 1, 1);
+		WriteDispatchIndirectArgs(OutPageIndexAllocationIndirectBufferArgs, MacroGroupId, DivideAndRoundUp(AllocatedPageIndexCount, IndirectDispatchGroupSize), 1, 1);
 	}
 
 	// Matrix generation
@@ -239,7 +240,7 @@ RWBuffer<uint>	OutValidPageIndexBuffer;
 // PageIndexCoordBuffer is sampled with linear coordinate for allocated nodes. VALID NODE ARE COMPACTED. It contains the 3d page coordinate and ClustedId. Only used for opaque voxel injection.
 
 #if PERMUTATION_USE_CLUSTER
-[numthreads(32, 1, 1)]
+[numthreads(GROUP_SIZE, 1, 1)]
 void MarkValid_PrepareCS(uint2 DispatchThreadId : SV_DispatchThreadID)
 {
 	const uint ClusterIndex = DispatchThreadId.x;
@@ -311,7 +312,7 @@ void MarkValid_PrepareCS(uint2 DispatchThreadId : SV_DispatchThreadID)
 	}
 }
 #else // PERMUTATION_USE_CLUSTER
-[numthreads(8, 8, 8)]
+[numthreads(GROUP_SIZE, GROUP_SIZE, GROUP_SIZE)]
 void MarkValid_PrepareCS(uint3 DispatchThreadId : SV_DispatchThreadID)
 {
 	const uint4 PageIndexResolutionAndOffset = PageIndexResolutionAndOffsetBuffer.Load(MacroGroupId);
@@ -358,8 +359,6 @@ Buffer<int>		MacroGroupVoxelAlignedAABBBuffer;
 Buffer<uint4>	PageIndexResolutionAndOffsetBuffer;
 RWBuffer<uint>	OutValidPageIndexBuffer;
 
-#define GROUP_SIZE 32
-
 // PageIndexBuffer      is sampled with linear coordinate computed from the 3d page coordinate. VALID NODE ARE NOT COMPACTED. It contains the LINEAR PAGE INDEX (to map to the 3d volume).
 // PageIndexCoordBuffer is sampled with linear coordinate for allocated nodes. VALID NODE ARE COMPACTED. It contains the 3d page coordinate and ClustedId. Only used for opaque voxel injection.
 
@@ -442,8 +441,6 @@ RWBuffer<uint4> PageIndexCoordBuffer;
 groupshared uint LocalCounter;
 groupshared uint GroupBase;
 
-#define GROUP_SIZE 32
-
 [numthreads(GROUP_SIZE, 1, 1)]
 void AllocateCS(uint GroupIndex : SV_GroupIndex, uint3 DispatchThreadId : SV_DispatchThreadID)
 {
@@ -559,12 +556,13 @@ void AddDescCS(uint GroupIndex : SV_GroupIndex, uint3 DispatchThreadId : SV_Disp
 #endif
 
 ///////////////////////////////////////////////////////////////////////////////////////////////////
+// Preapare indirect buffer
 
 #if SHADER_ADDINDIRECTBUFFER
 
 uint PageResolution;
 uint MacroGroupId;
-uint IndirectGroupSize;
+int3 IndirectGroupSize;
 
 RWBuffer<uint> OutPageIndexGlobalCounter;
 RWBuffer<uint> OutIndirectArgsBuffer;
@@ -578,10 +576,9 @@ void AddIndirectBufferCS(uint GroupIndex : SV_GroupIndex, uint3 DispatchThreadId
 	const uint AllocatedPageCount = TotalAllocatedPageCount - PrevTotalAllocatedPageCount;
 	const uint VoxelCountPerPage = PageResolution * PageResolution * PageResolution;
 
-	const uint MaxGroupCountPerDim = 65535;
-	const uint DispatchX = AllocatedPageCount % MaxGroupCountPerDim;
-	const uint DispatchY = ceil(AllocatedPageCount / float(MaxGroupCountPerDim));
-	WriteDispatchIndirectArgs(OutIndirectArgsBuffer, MacroGroupId, DispatchX, DispatchY, ceil(VoxelCountPerPage / float(IndirectGroupSize)));
+	const uint DispatchCountX = DivideAndRoundUp(VoxelCountPerPage, IndirectGroupSize.x);
+	const uint DispatchCountZ = DivideAndRoundUp(AllocatedPageCount, IndirectGroupSize.z);
+	WriteDispatchIndirectArgs(OutIndirectArgsBuffer, MacroGroupId, DispatchCountX, 1, DispatchCountZ);
 
 	// Copy current total to previous total
 	OutPageIndexGlobalCounter[1] = TotalAllocatedPageCount;
@@ -593,12 +590,12 @@ void AddIndirectBufferCS(uint GroupIndex : SV_GroupIndex, uint3 DispatchThreadId
 #endif
 
 ///////////////////////////////////////////////////////////////////////////////////////////////////
-
-#define IND_CLEAR_GROUP_SIZE_Z 64
+// Prepare indirect args for indirect clear
 
 #if SHADER_INDPAGECLEARBUFFERGEN
 
 uint PageResolution;
+int3 GroupSize;
 Buffer<uint> PageIndexGlobalCounter;
 RWBuffer<uint> OutIndirectArgsBuffer;
 
@@ -608,36 +605,36 @@ void VoxelIndPageClearBufferGenCS(uint GroupIndex : SV_GroupIndex, uint3 Dispatc
 	const uint TotalAllocatedPageCount = PageIndexGlobalCounter[0];
 	const uint VoxelCountPerPage = PageResolution * PageResolution * PageResolution;
 
-	const uint MaxGroupCountPerDim = 65535;
-	const uint DispatchX = TotalAllocatedPageCount % MaxGroupCountPerDim;
-	const uint DispatchY = ceil(TotalAllocatedPageCount / float(MaxGroupCountPerDim));
+	const uint DispatchX = DivideAndRoundUp(VoxelCountPerPage, GroupSize.x);
+	const uint DispatchZ = DivideAndRoundUp(TotalAllocatedPageCount, GroupSize.z);
 
-	WriteDispatchIndirectArgs(OutIndirectArgsBuffer, 0, DispatchX, DispatchY, ceil(VoxelCountPerPage / float(IND_CLEAR_GROUP_SIZE_Z)));
+	WriteDispatchIndirectArgs(OutIndirectArgsBuffer, 0, DispatchX, 1, DispatchZ);
 }
 
 #endif
 
 ///////////////////////////////////////////////////////////////////////////////////////////////////
+// Indirect clear
 
 #if SHADER_INDPAGECLEAR
 
+Buffer<uint>		PageIndexGlobalCounter;
 uint				VirtualVoxelParams_PageResolution;
 int3				VirtualVoxelParams_PageCountResolution;
 Buffer<uint4>		VirtualVoxelParams_PageIndexCoordBuffer;
 
 RWTexture3D<uint>	OutPageTexture;
 
-[numthreads(1, 1, IND_CLEAR_GROUP_SIZE_Z)]
+[numthreads(GROUP_SIZE_X, 1, GROUP_SIZE_Z)]
 void VoxelIndPageClearCS(uint3 DispatchThreadId : SV_DispatchThreadID, uint3 GroupId : SV_GroupID, uint3 GroupThreadId : SV_GroupThreadID)
 {
-	const uint  LinearVoxelCoord = DispatchThreadId.z;
-
+	const uint TotalAllocatedPageCount = PageIndexGlobalCounter[0];
 	const uint VoxelCountPerPage = VirtualVoxelParams_PageResolution * VirtualVoxelParams_PageResolution * VirtualVoxelParams_PageResolution;
-	//	if (AllocatedPageIndex < TotalAllocatedPageCount)		// Not needed today because the XY size exactly matches the range (page count) we need to clear.
-	if (LinearVoxelCoord < VoxelCountPerPage)
+
+	const uint LinearVoxelCoord = DispatchThreadId.x;
+	const uint AllocatedPageIndex = DispatchThreadId.z;
+	if (AllocatedPageIndex < TotalAllocatedPageCount && LinearVoxelCoord < VoxelCountPerPage)
 	{
-		const uint  MaxDispatchCountPerDim = 65535u;
-		const uint  AllocatedPageIndex = DispatchThreadId.x + DispatchThreadId.y * MaxDispatchCountPerDim;
 		const uint3 VoxelCoordOffset = IndexToCoord(LinearVoxelCoord, VirtualVoxelParams_PageResolution.xxx);
 
 		const uint PageIndex = AllocatedPageIndex; // PageIndexBuffer is not needed, we already know those tiles are allocated linearly in 3D within OutPageTexture.
@@ -780,7 +777,6 @@ void FeedbackCS(uint3 DispatchThreadId : SV_DispatchThreadID)
 
 #if SHADER_RASTERCOMPUTE
 
-uint			DispatchCountX;
 uint			MaxRasterCount;
 uint			FrameIdMod8;
 uint			MacroGroupId;
@@ -800,14 +796,6 @@ float SampleWeight;
 
 #define JITTER_ENABLE 0
 
-#if PERMUTATION_GROUP_SIZE == 64
-	#define GROUP_SIZE 64
-#elif PERMUTATION_GROUP_SIZE == 32
-	#define GROUP_SIZE 32
-#else
-	#error Unknown group size	
-#endif
-
 float3 GetHairVoxelJitter(uint2 PixelCoord, uint Seed)
 {
 	return float3(
@@ -822,7 +810,7 @@ void MainCS(uint2 DispatchThreadID : SV_DispatchThreadID)
 #if PERMUTATION_CULLING == 1
 	uint VertexIndex0 = GetHairStrandsVertexFetchIndex_Culled(DispatchThreadID, GROUP_SIZE);
 #else
-	uint VertexIndex0 = GetHairStrandsVertexFetchIndex(DispatchThreadID, GROUP_SIZE, DispatchCountX);
+	uint VertexIndex0 = DispatchThreadID.x;
 #endif
 	uint VertexIndex1 = min(VertexIndex0 + 1, VertexCount - 1);
 	if (VertexIndex0 >= VertexCount)
@@ -985,20 +973,26 @@ uint VoxelBiasCount;
 uint VoxelMarkCount;
 RWTexture3D<uint> OutPageTexture;
 
+uint		 VirtualVoxelParams_PageCount;
 uint		 VirtualVoxelParams_PageIndexCount;
 uint		 VirtualVoxelParams_PageResolution;
 uint3		 VirtualVoxelParams_PageCountResolution;
 Buffer<uint> VirtualVoxelParams_PageIndexBuffer;
+Buffer<uint> VirtualVoxelParams_AllocatedPageCountBuffer;
 Buffer<uint4>VirtualVoxelParams_PageIndexCoordBuffer;
 StructuredBuffer<FPackedVirtualVoxelNodeDesc> VirtualVoxelParams_NodeDescBuffer;
 
-#define GROUP_SIZE 64
-[numthreads(1, 1, GROUP_SIZE)]
+[numthreads(GROUP_SIZE_X, 1, GROUP_SIZE_Z)]
 void MainCS(uint3 DispatchThreadId : SV_DispatchThreadID)
 {
-	const uint  MaxDispatchCountPerDim = 65535u;
-	const uint  AllocatedPageIndex = DispatchThreadId.x + DispatchThreadId.y * MaxDispatchCountPerDim;
-	const uint  LinearVoxelCoord = DispatchThreadId.z;
+	const uint LinearVoxelCoord = DispatchThreadId.x;
+	const uint AllocatedPageIndex = DispatchThreadId.z;
+	const bool bValidVoxel = AllocatedPageIndex < VirtualVoxelParams_AllocatedPageCountBuffer[0] && LinearVoxelCoord < VirtualVoxelParams_PageCount;
+	if (!bValidVoxel)
+	{
+		return;
+	}
+
 	const uint3 VoxelCoordOffset = IndexToCoord(LinearVoxelCoord, VirtualVoxelParams_PageResolution.xxx);
 
 	const FVirtualVoxelNodeDesc NodeDesc = UnpackVoxelNode(VirtualVoxelParams_NodeDescBuffer[MacroGroupId], VirtualVoxelParams_PageResolution);
@@ -1291,20 +1285,25 @@ uint PageResolution;
 uint SourceMip;
 uint TargetMip;
 
+Buffer<uint> 	  AllocatedPageCountBuffer;
 Texture3D<uint>   InDensityTexture;
 RWTexture3D<uint> OutDensityTexture;
 
-[numthreads(1, 1, 64)]
+[numthreads(GROUP_SIZE_X, 1, GROUP_SIZE_Z)]
 void MainCS(uint3 DispatchThreadId : SV_DispatchThreadID)
 {
-	const uint  MaxDispatchCountPerDim = 65535;
-	const uint  PageIndex = DispatchThreadId.x + DispatchThreadId.y * MaxDispatchCountPerDim;
+	const uint  PageIndex = DispatchThreadId.z;
+	if (PageIndex >= AllocatedPageCountBuffer[0])
+	{
+		return;
+	}
+
 	const uint3 PageCoord = IndexToCoord(PageIndex, PageCountResolution);
 
 	const uint  InPageResolution = PageResolution >> SourceMip;
 	const uint OutPageResolution = PageResolution >> TargetMip;
 	const uint TotalVoxelPerOutPageResolution = OutPageResolution * OutPageResolution * OutPageResolution;
-	const uint VoxelIndex = DispatchThreadId.z;
+	const uint VoxelIndex = DispatchThreadId.x;
 	if (VoxelIndex >= TotalVoxelPerOutPageResolution)
 		return;
 
@@ -1353,7 +1352,7 @@ void MainCS(uint3 DispatchThreadId : SV_DispatchThreadID)
 
 uint PageResolution;
 uint TargetMipIndex;
-uint DispatchGroupSize;
+int3 DispatchGroupSize;
 
 Buffer<uint>   InIndirectArgs;
 RWBuffer<uint> OutIndirectArgs;
@@ -1365,9 +1364,9 @@ void MainCS(uint3 DispatchThreadId : SV_DispatchThreadID)
 
 	const uint TargetPageResolution = PageResolution >> TargetMipIndex;
 	const uint TotalVoxelCount = TargetPageResolution * TargetPageResolution * TargetPageResolution;
-	const uint DispatchZ = ceil(TotalVoxelCount / float(DispatchGroupSize));
+	const uint DispatchX = DivideAndRoundUp(TotalVoxelCount, DispatchGroupSize.x);
 
-	WriteDispatchIndirectArgs(OutIndirectArgs, 0, InIndirectArgs[0], InIndirectArgs[1], DispatchZ);
+	WriteDispatchIndirectArgs(OutIndirectArgs, 0, DispatchX, InIndirectArgs[1], InIndirectArgs[2]);
 }
 
 #endif
@@ -1385,11 +1384,10 @@ Buffer<uint>	  PageToPageIndexBuffer;
 RWBuffer<uint>	  OutPageIndexBuffer;
 RWBuffer<uint2>	  OutPageIndexOccupancyBuffer;
 
-[numthreads(1, 1, 64)]
+[numthreads(GROUP_SIZE_X, 1, GROUP_SIZE_Z)]
 void MainCS(uint3 DispatchThreadId : SV_DispatchThreadID)
 {
-	const uint  MaxDispatchCountPerDim = 65535;
-	const uint  PageIndex = DispatchThreadId.x + DispatchThreadId.y * MaxDispatchCountPerDim;
+	const uint  PageIndex = DispatchThreadId.z;
 	const uint3 PageCoord = IndexToCoord(PageIndex, PageCountResolution);
 
 	const uint AllocatedPageCount = PageIndexGlobalCounter[0];
@@ -1398,7 +1396,7 @@ void MainCS(uint3 DispatchThreadId : SV_DispatchThreadID)
 
 	const uint OutPageResolution = PageResolution >> MipIt;
 	const uint TotalVoxelPerOutPageResolution = OutPageResolution * OutPageResolution * OutPageResolution;
-	const uint VoxelIndex = DispatchThreadId.z;
+	const uint VoxelIndex = DispatchThreadId.x;
 	if (VoxelIndex >= TotalVoxelPerOutPageResolution)
 		return;
 
diff --git a/Engine/Shaders/Shared/HairStrandsDefinitions.h b/Engine/Shaders/Shared/HairStrandsDefinitions.h
--- a/Engine/Shaders/Shared/HairStrandsDefinitions.h
+++ b/Engine/Shaders/Shared/HairStrandsDefinitions.h
@@ -24,11 +24,14 @@
 #define MAX_HAIR_LOD 8
 
 // Group size for dispatching based on a groom vertex/curve/cluster count. 
-// This defines ensures the group size is consistent across the hair pipeline
+// This defines ensures the group size is consistent across the hair pipeline, 
+// and ensures dispatch count are smaller than 65k limits
 #define HAIR_VERTEXCOUNT_GROUP_SIZE  1024u
 #define HAIR_CURVECOUNT_GROUP_SIZE   1024u
 #define HAIR_CLUSTERCOUNT_GROUP_SIZE 1024u
 
+#define MAX_HAIR_MACROGROUP_COUNT 16
+
 // HAIR_ATTRIBUTE_MAX rounded to 4
 #define HAIR_ATTRIBUTE_OFFSET_COUNT ((HAIR_ATTRIBUTE_COUNT + 3) / 4)
 
diff --git a/Engine/Source/Runtime/Renderer/Private/HairStrands/HairStrandsData.h b/Engine/Source/Runtime/Renderer/Private/HairStrands/HairStrandsData.h
--- a/Engine/Source/Runtime/Renderer/Private/HairStrands/HairStrandsData.h
+++ b/Engine/Source/Runtime/Renderer/Private/HairStrands/HairStrandsData.h
@@ -218,6 +218,7 @@ BEGIN_SHADER_PARAMETER_STRUCT(FHairStrandsVoxelCommonParameters, )
 
 	SHADER_PARAMETER(uint32, AllocationFeedbackEnable)
 
+	SHADER_PARAMETER_RDG_BUFFER_SRV(Buffer<uint>, AllocatedPageCountBuffer)
 	SHADER_PARAMETER_RDG_BUFFER_SRV(Buffer<uint>, PageIndexBuffer)
 	SHADER_PARAMETER_RDG_BUFFER_SRV(Buffer<uint2>, PageIndexOccupancyBuffer)
 	SHADER_PARAMETER_RDG_BUFFER_SRV(Buffer<uint>, PageIndexCoordBuffer)
diff --git a/Engine/Source/Runtime/Renderer/Private/HairStrands/HairStrandsVoxelization.cpp b/Engine/Source/Runtime/Renderer/Private/HairStrands/HairStrandsVoxelization.cpp
--- a/Engine/Source/Runtime/Renderer/Private/HairStrands/HairStrandsVoxelization.cpp
+++ b/Engine/Source/Runtime/Renderer/Private/HairStrands/HairStrandsVoxelization.cpp
@@ -16,6 +16,16 @@
 #include "ScenePrivate.h"
 #include "DataDrivenShaderPlatformInfo.h"
 
+// Common threading group layout used in voxelization code by the following pass
+// * FVoxelIndPageClearCS
+// * FVirtualVoxelInjectOpaqueCS
+// * FVirtualVoxelPatchPageIndexWithMipDataCS
+// * FVirtualVoxelGenerateMipCS
+// 
+// This threading layout allows to process up to 1M allocated pages (16 * 65k), and process 4x4x4 voxels within the same page.
+// Notes: 1024 threads is the limit per work-group 
+static const FIntVector GHairVoxel_GroupSize_NumVoxelPerPage_1_NumAllocatedPage(64u, 0u, 16u);
+
 static int32 GHairVoxelizationEnable = 1;
 static FAutoConsoleVariableRef CVarGHairVoxelizationEnable(TEXT("r.HairStrands.Voxelization"), GHairVoxelizationEnable, TEXT("Enable hair voxelization for transmittance evaluation"));
 
@@ -162,11 +172,15 @@ class FVirtualVoxelInjectOpaqueCS : public FGlobalShader
 	END_SHADER_PARAMETER_STRUCT()
 
 public:
+	static FIntVector GetGroupSize() { return GHairVoxel_GroupSize_NumVoxelPerPage_1_NumAllocatedPage; }
 	static bool ShouldCompilePermutation(const FGlobalShaderPermutationParameters& Parameters) { return IsHairStrandsSupported(EHairStrandsShaderType::Strands, Parameters.Platform); }
 	static void ModifyCompilationEnvironment(const FGlobalShaderPermutationParameters& Parameters, FShaderCompilerEnvironment& OutEnvironment)
 	{
 		FGlobalShader::ModifyCompilationEnvironment(Parameters, OutEnvironment);
 		OutEnvironment.SetDefine(TEXT("SHADER_INJECTOPAQUE_VIRTUALVOXEL"), 1);
+		OutEnvironment.SetDefine(TEXT("GROUP_SIZE_X"), GetGroupSize().X);
+		OutEnvironment.SetDefine(TEXT("GROUP_SIZE_Y"), GetGroupSize().Y);
+		OutEnvironment.SetDefine(TEXT("GROUP_SIZE_Z"), GetGroupSize().Z);
 	}
 };
 
@@ -199,7 +213,6 @@ static void AddVirtualVoxelInjectOpaquePass(
 	TShaderMapRef<FVirtualVoxelInjectOpaqueCS> ComputeShader(View.ShaderMap);
 	const FGlobalShaderMap* GlobalShaderMap = View.ShaderMap;
 
-	check(VoxelResources.Parameters.Common.IndirectDispatchGroupSize == 64);
 	const uint32 ArgsOffset = sizeof(FRHIDispatchIndirectParameters) * Parameters->MacroGroupId;
 
 	FComputeShaderUtils::AddPass(
@@ -242,11 +255,13 @@ class FVoxelAllocatePageIndexCS : public FGlobalShader
 	END_SHADER_PARAMETER_STRUCT()
 
 public:
+	static uint32 GetGroupSize() { return MAX_HAIR_MACROGROUP_COUNT; }
 	static bool ShouldCompilePermutation(const FGlobalShaderPermutationParameters& Parameters) { return IsHairStrandsSupported(EHairStrandsShaderType::Strands, Parameters.Platform); }
 	static void ModifyCompilationEnvironment(const FGlobalShaderPermutationParameters& Parameters, FShaderCompilerEnvironment& OutEnvironment)
 	{
 		FGlobalShader::ModifyCompilationEnvironment(Parameters, OutEnvironment);
 		OutEnvironment.SetDefine(TEXT("SHADER_ALLOCATEPAGEINDEX"), 1);
+		OutEnvironment.SetDefine(TEXT("GROUP_SIZE"), GetGroupSize());
 	}
 };
 
@@ -261,24 +276,25 @@ class FVoxelMarkValidPageIndex_PrepareCS : public FGlobalShader
 	BEGIN_SHADER_PARAMETER_STRUCT(FParameters, )
 		SHADER_PARAMETER(uint32, MaxClusterCount)
 		SHADER_PARAMETER(uint32, MacroGroupId)
-		SHADER_PARAMETER(uint32, MaxScatterAllocationCount)
-		SHADER_PARAMETER(uint32, bForceDirectPageAllocation)
 
 		SHADER_PARAMETER_RDG_BUFFER_SRV(Buffer, GroupAABBsBuffer)
 		SHADER_PARAMETER_RDG_BUFFER_SRV(Buffer, ClusterAABBsBuffer)
 		SHADER_PARAMETER_RDG_BUFFER_SRV(Buffer, MacroGroupVoxelAlignedAABBBuffer)
 		SHADER_PARAMETER_RDG_BUFFER_SRV(Buffer, PageIndexResolutionAndOffsetBuffer)
 		SHADER_PARAMETER_RDG_BUFFER_UAV(RWBuffer<uint>, OutValidPageIndexBuffer)
-		SHADER_PARAMETER_RDG_BUFFER_UAV(RWBuffer<uint>, OutDeferredScatterCounter)
-		SHADER_PARAMETER_RDG_BUFFER_UAV(RWBuffer<uint2>, OutDeferredScatterBuffer)
 	END_SHADER_PARAMETER_STRUCT()
 
 public:
+	static uint32 GetGroupSize(bool bUseCluster) { return bUseCluster ? 256u : 8u; }
 	static bool ShouldCompilePermutation(const FGlobalShaderPermutationParameters& Parameters) { return IsHairStrandsSupported(EHairStrandsShaderType::Strands, Parameters.Platform); }
 	static void ModifyCompilationEnvironment(const FGlobalShaderPermutationParameters& Parameters, FShaderCompilerEnvironment& OutEnvironment)
 	{
+		FPermutationDomain PermutationVector(Parameters.PermutationId);
+		const bool bUseCluster = PermutationVector.Get<FUseCluster>() != 0;
+
 		FGlobalShader::ModifyCompilationEnvironment(Parameters, OutEnvironment);
 		OutEnvironment.SetDefine(TEXT("SHADER_MARKVALID_PREPARE"), 1);
+		OutEnvironment.SetDefine(TEXT("GROUP_SIZE"), GetGroupSize(bUseCluster));
 	}
 };
 
@@ -304,11 +320,13 @@ class FVoxelMarkValidPageIndexCS : public FGlobalShader
 	END_SHADER_PARAMETER_STRUCT()
 
 public:
+	static uint32 GetGroupSize() { return 256u; /* Allow to process up to 16M clusters (= 65k * 256) */ }
 	static bool ShouldCompilePermutation(const FGlobalShaderPermutationParameters& Parameters) { return IsHairStrandsSupported(EHairStrandsShaderType::Strands, Parameters.Platform); }
 	static void ModifyCompilationEnvironment(const FGlobalShaderPermutationParameters& Parameters, FShaderCompilerEnvironment& OutEnvironment)
 	{
 		FGlobalShader::ModifyCompilationEnvironment(Parameters, OutEnvironment);
 		OutEnvironment.SetDefine(TEXT("SHADER_MARKVALID"), 1);
+		OutEnvironment.SetDefine(TEXT("GROUP_SIZE"), GetGroupSize());
 	}
 };
 
@@ -335,11 +353,13 @@ class FVoxelAllocateVoxelPageCS : public FGlobalShader
 	END_SHADER_PARAMETER_STRUCT()
 
 public:
+	static uint32 GetGroupSize() { return 256u; /* Allow to handle up to 16M pages (= 65k * 256) */ }
 	static bool ShouldCompilePermutation(const FGlobalShaderPermutationParameters& Parameters) { return IsHairStrandsSupported(EHairStrandsShaderType::Strands, Parameters.Platform); }
 	static void ModifyCompilationEnvironment(const FGlobalShaderPermutationParameters& Parameters, FShaderCompilerEnvironment& OutEnvironment)
 	{
 		FGlobalShader::ModifyCompilationEnvironment(Parameters, OutEnvironment);
 		OutEnvironment.SetDefine(TEXT("SHADER_ALLOCATE"), 1);
+		OutEnvironment.SetDefine(TEXT("GROUP_SIZE"), GetGroupSize());
 	}
 };
 
@@ -367,6 +387,7 @@ class FVoxelAddNodeDescCS : public FGlobalShader
 	END_SHADER_PARAMETER_STRUCT()
 
 public:
+	static uint32 GetGroupSize() { return 1; }
 	static bool ShouldCompilePermutation(const FGlobalShaderPermutationParameters& Parameters) { return IsHairStrandsSupported(EHairStrandsShaderType::Strands, Parameters.Platform); }
 	static void ModifyCompilationEnvironment(const FGlobalShaderPermutationParameters& Parameters, FShaderCompilerEnvironment& OutEnvironment)
 	{
@@ -382,14 +403,15 @@ class FVoxelAddIndirectBufferCS : public FGlobalShader
 
 	BEGIN_SHADER_PARAMETER_STRUCT(FParameters, )
 		SHADER_PARAMETER(uint32, MacroGroupId)
-		SHADER_PARAMETER(uint32, IndirectGroupSize)
 		SHADER_PARAMETER(uint32, PageResolution)
+		SHADER_PARAMETER(FIntVector, IndirectGroupSize)
 		SHADER_PARAMETER_RDG_BUFFER_UAV(RWBuffer, OutPageIndexGlobalCounter)
 		SHADER_PARAMETER_RDG_BUFFER_UAV(RWBuffer, OutIndirectArgsBuffer)
 		SHADER_PARAMETER_RDG_BUFFER_UAV(RWBuffer, OutTotalRequestedPageAllocationBuffer)
 	END_SHADER_PARAMETER_STRUCT()
 
 public:
+	static uint32 GetGroupSize() { return 1; }
 	static bool ShouldCompilePermutation(const FGlobalShaderPermutationParameters& Parameters) { return IsHairStrandsSupported(EHairStrandsShaderType::Strands, Parameters.Platform); }
 	static void ModifyCompilationEnvironment(const FGlobalShaderPermutationParameters& Parameters, FShaderCompilerEnvironment& OutEnvironment)
 	{
@@ -405,12 +427,14 @@ class FVoxelIndPageClearBufferGenCS : public FGlobalShader
 	SHADER_USE_PARAMETER_STRUCT(FVoxelIndPageClearBufferGenCS, FGlobalShader);
 
 	BEGIN_SHADER_PARAMETER_STRUCT(FParameters, )
+		SHADER_PARAMETER(FIntVector, GroupSize)
+		SHADER_PARAMETER(uint32, PageResolution)
 		SHADER_PARAMETER_RDG_BUFFER_SRV(Buffer, PageIndexGlobalCounter)
 		SHADER_PARAMETER_RDG_BUFFER_UAV(RWBuffer, OutIndirectArgsBuffer)
-		SHADER_PARAMETER(uint32, PageResolution)
 	END_SHADER_PARAMETER_STRUCT()
 
 public:
+	static uint32 GetGroupSize() { return 1; }
 	static bool ShouldCompilePermutation(const FGlobalShaderPermutationParameters& Parameters) { return IsHairStrandsSupported(EHairStrandsShaderType::Strands, Parameters.Platform); }
 	static void ModifyCompilationEnvironment(const FGlobalShaderPermutationParameters& Parameters, FShaderCompilerEnvironment& OutEnvironment)
 	{
@@ -426,16 +450,21 @@ class FVoxelIndPageClearCS : public FGlobalShader
 
 	BEGIN_SHADER_PARAMETER_STRUCT(FParameters, )
 		SHADER_PARAMETER_STRUCT(FHairStrandsVoxelCommonParameters, VirtualVoxelParams)
+		SHADER_PARAMETER_RDG_BUFFER_SRV(Buffer, PageIndexGlobalCounter)
 		SHADER_PARAMETER_RDG_TEXTURE_UAV(RWTexture3D, OutPageTexture)
 		RDG_BUFFER_ACCESS(IndirectDispatchBuffer, ERHIAccess::IndirectArgs)
 	END_SHADER_PARAMETER_STRUCT()
 
 public:
+	static FIntVector GetGroupSize() { return GHairVoxel_GroupSize_NumVoxelPerPage_1_NumAllocatedPage; }
 	static bool ShouldCompilePermutation(const FGlobalShaderPermutationParameters& Parameters) { return IsHairStrandsSupported(EHairStrandsShaderType::Strands, Parameters.Platform); }
 	static void ModifyCompilationEnvironment(const FGlobalShaderPermutationParameters& Parameters, FShaderCompilerEnvironment& OutEnvironment)
 	{
 		FGlobalShader::ModifyCompilationEnvironment(Parameters, OutEnvironment);
 		OutEnvironment.SetDefine(TEXT("SHADER_INDPAGECLEAR"), 1);
+		OutEnvironment.SetDefine(TEXT("GROUP_SIZE_X"), GetGroupSize().X);
+		OutEnvironment.SetDefine(TEXT("GROUP_SIZE_Y"), GetGroupSize().Y);
+		OutEnvironment.SetDefine(TEXT("GROUP_SIZE_Z"), GetGroupSize().Z);
 	}
 };
 
@@ -458,6 +487,7 @@ class FVoxelAdaptiveFeedbackCS : public FGlobalShader
 	END_SHADER_PARAMETER_STRUCT()
 
 public:
+	static uint32 GetGroupSize() { return 1; }
 	static bool ShouldCompilePermutation(const FGlobalShaderPermutationParameters& Parameters) { return IsHairStrandsSupported(EHairStrandsShaderType::Strands, Parameters.Platform); }
 	static void ModifyCompilationEnvironment(const FGlobalShaderPermutationParameters& Parameters, FShaderCompilerEnvironment& OutEnvironment)
 	{
@@ -491,7 +521,6 @@ static void AddAllocateVoxelPagesPass(
 	const float CPUMinVoxelWorldSize,
 	const uint32 PageResolution,
 	const FIntVector PageTextureResolution,
-	const uint32 IndirectDispatchGroupSize,
 	uint32& OutTotalPageIndexCount,
 	FRDGBufferRef& OutPageIndexBuffer,
 	FRDGBufferRef& OutPageIndexOccupancyBuffer,
@@ -504,7 +533,6 @@ static void AddAllocateVoxelPagesPass(
 	FRDGBufferRef& CurrGPUMinVoxelSize,
 	FRDGBufferRef& NextGPUMinVoxelSize)
 {
-	const uint32 GroupSize = 32;
 	const bool bIsGPUDriven = GHairVirtualVoxelGPUDriven > 0;
 	const uint32 MacroGroupCount = MacroGroupDatas.Num();
 	if (MacroGroupCount == 0)
@@ -616,7 +644,7 @@ static void AddAllocateVoxelPagesPass(
 		Parameters->MacroGroupVoxelSizeBuffer = GraphBuilder.CreateSRV(MacroGroupResources.MacroGroupVoxelSizeBuffer, PF_R16F);
 		Parameters->MacroGroupAABBBuffer = GraphBuilder.CreateSRV(MacroGroupResources.MacroGroupAABBsBuffer, PF_R32_SINT);
 		Parameters->MacroGroupVoxelAlignedAABBBuffer = GraphBuilder.CreateUAV(MacroGroupResources.MacroGroupVoxelAlignedAABBsBuffer, PF_R32_SINT);
-		Parameters->IndirectDispatchGroupSize = GroupSize; // This is the GroupSize used for FVoxelAllocateVoxelPageCS
+		Parameters->IndirectDispatchGroupSize = FVoxelAllocateVoxelPageCS::GetGroupSize();
 		Parameters->OutPageIndexResolutionAndOffsetBuffer = GraphBuilder.CreateUAV(PageIndexResolutionBuffer, PF_R32G32B32A32_UINT);
 		Parameters->OutVoxelizationViewInfoBuffer = GraphBuilder.CreateUAV(VoxelizationViewInfoBuffer);
 		Parameters->OutPageIndexAllocationIndirectBufferArgs = GraphBuilder.CreateUAV(PageIndexAllocationIndirectBufferArgs);
@@ -694,19 +722,25 @@ static void AddAllocateVoxelPagesPass(
 				Parameters->OutValidPageIndexBuffer		= PageIndexBufferUAV;
 
 				FIntVector DispatchCount = bUseClusterAABB ?
-					FIntVector((Parameters->MaxClusterCount + GroupSize - 1) / GroupSize, 1, 1) :
 					FIntVector(
-						FMath::DivideAndRoundUp(CPUAllocationDesc.PageIndexResolution.X, 8),
-						FMath::DivideAndRoundUp(CPUAllocationDesc.PageIndexResolution.Y, 8),
-						FMath::DivideAndRoundUp(CPUAllocationDesc.PageIndexResolution.Z, 8));
+						FMath::DivideAndRoundUp(Parameters->MaxClusterCount, FVoxelMarkValidPageIndex_PrepareCS::GetGroupSize(true)), 
+						1, 
+						1) :
+					FIntVector(
+						FMath::DivideAndRoundUp(uint32(CPUAllocationDesc.PageIndexResolution.X), FVoxelMarkValidPageIndex_PrepareCS::GetGroupSize(false)),
+						FMath::DivideAndRoundUp(uint32(CPUAllocationDesc.PageIndexResolution.Y), FVoxelMarkValidPageIndex_PrepareCS::GetGroupSize(false)),
+						FMath::DivideAndRoundUp(uint32(CPUAllocationDesc.PageIndexResolution.Z), FVoxelMarkValidPageIndex_PrepareCS::GetGroupSize(false)));
 
 				check(DispatchCount.X <= GRHIMaxDispatchThreadGroupsPerDimension.X);
+				check(DispatchCount.Y <= GRHIMaxDispatchThreadGroupsPerDimension.Y);
+				check(DispatchCount.Z <= GRHIMaxDispatchThreadGroupsPerDimension.Z);
+
 				FVoxelMarkValidPageIndex_PrepareCS::FPermutationDomain PermutationVector;
 				PermutationVector.Set<FVoxelMarkValidPageIndex_PrepareCS::FUseCluster>(bUseClusterAABB ? 1 : 0);
 				TShaderMapRef<FVoxelMarkValidPageIndex_PrepareCS> ComputeShader(View.ShaderMap, PermutationVector);
 				FComputeShaderUtils::AddPass(
 					GraphBuilder,
-					RDG_EVENT_NAME("HairStrands::MarkValidPageIndex_Prepare(%s)", bUseClusterAABB ? TEXT("Cluster") : TEXT("Group")),
+					RDG_EVENT_NAME("HairStrands::MarkValidPageIndex(%s)", bUseClusterAABB ? TEXT("Cluster") : TEXT("Group")),
 					ComputeShader,
 					Parameters,
 					DispatchCount);
@@ -739,7 +773,7 @@ static void AddAllocateVoxelPagesPass(
 				FVoxelMarkValidPageIndexCS::FPermutationDomain PermutationVector;
 				PermutationVector.Set<FVoxelMarkValidPageIndexCS::FGPUDriven>(bIsGPUDriven ? 1 : 0);
 
-				FIntVector DispatchCount((Parameters->MaxClusterCount + GroupSize - 1) / GroupSize, 1, 1);
+				const FIntVector DispatchCount(FMath::DivideAndRoundUp(Parameters->MaxClusterCount, FVoxelMarkValidPageIndexCS::GetGroupSize()), 1, 1);
 				check(DispatchCount.X <= GRHIMaxDispatchThreadGroupsPerDimension.X);
 				TShaderMapRef<FVoxelMarkValidPageIndexCS> ComputeShader(View.ShaderMap, PermutationVector);
 				FComputeShaderUtils::AddPass(
@@ -813,7 +847,7 @@ static void AddAllocateVoxelPagesPass(
 			}
 			else
 			{
-				const FIntVector DispatchCount((CPUAllocationDesc.PageIndexCount + GroupSize - 1) / GroupSize, 1, 1);
+				const FIntVector DispatchCount(FMath::DivideAndRoundUp(CPUAllocationDesc.PageIndexCount, FVoxelAllocateVoxelPageCS::GetGroupSize()), 1, 1);
 				check(DispatchCount.X <= GRHIMaxDispatchThreadGroupsPerDimension.X);
 				FComputeShaderUtils::AddPass(GraphBuilder, RDG_EVENT_NAME("HairStrands::AllocateVoxelPage"), ComputeShader, Parameters, DispatchCount);
 			}
@@ -826,7 +860,7 @@ static void AddAllocateVoxelPagesPass(
 			FVoxelAddIndirectBufferCS::FParameters* Parameters = GraphBuilder.AllocParameters<FVoxelAddIndirectBufferCS::FParameters>();
 			Parameters->MacroGroupId = MacroGroup.MacroGroupId;
 			Parameters->PageResolution = PageResolution;
-			Parameters->IndirectGroupSize = IndirectDispatchGroupSize;
+			Parameters->IndirectGroupSize = FVirtualVoxelInjectOpaqueCS::GetGroupSize();
 			Parameters->OutPageIndexGlobalCounter = GraphBuilder.CreateUAV(PageIndexGlobalCounter, PF_R32_UINT);
 			Parameters->OutIndirectArgsBuffer = GraphBuilder.CreateUAV(IndirectArgsBuffer);
 			Parameters->OutTotalRequestedPageAllocationBuffer = TotalRequestedPageAllocationBufferUAV;
@@ -921,7 +955,6 @@ static FHairStrandsVoxelResources AllocateVirtualVoxelResources(
 	Out.Parameters.Common.SteppingScale_Raytracing	  = FMath::Clamp(GHairStransVoxelRaymarchingSteppingScale_Raytracing		>= 0 ? GHairStransVoxelRaymarchingSteppingScale_Raytracing		: GHairStransVoxelRaymarchingSteppingScale, 1.f, 10.f);
 
 	Out.Parameters.Common.NodeDescCount							= MacroGroupDatas.Num();
-	Out.Parameters.Common.IndirectDispatchGroupSize				= 64;	
 	Out.Parameters.Common.Raytracing_ShadowOcclusionThreshold	= FMath::Max(0.f, GHairVirtualVoxelRaytracing_ShadowOcclusionThreshold);
 	Out.Parameters.Common.Raytracing_SkyOcclusionThreshold		= FMath::Max(0.f, GHairVirtualVoxelRaytracing_SkyOcclusionThreshold);
 	Out.Parameters.Common.HairCoveragePixelRadiusAtDepth1		= ComputeMinStrandRadiusAtDepth1(FIntPoint(View.ViewRect.Width(), View.ViewRect.Height()), View.FOV, 1/*SampleCount*/, 1/*RasterizationScale*/).Primary;
@@ -948,7 +981,6 @@ static FHairStrandsVoxelResources AllocateVirtualVoxelResources(
 		Out.Parameters.Common.CPUMinVoxelWorldSize,
 		Out.Parameters.Common.PageResolution,
 		Out.Parameters.Common.PageTextureResolution,
-		Out.Parameters.Common.IndirectDispatchGroupSize,
 		Out.Parameters.Common.PageIndexCount,
 		Out.PageIndexBuffer,
 		Out.PageIndexOccupancyBuffer,
@@ -985,6 +1017,7 @@ static FHairStrandsVoxelResources AllocateVirtualVoxelResources(
 	Out.Parameters.Common.PageIndexOccupancyBuffer	= GraphBuilder.CreateSRV(Out.PageIndexOccupancyBuffer, PF_R32G32_UINT);
 	Out.Parameters.Common.PageIndexCoordBuffer		= GraphBuilder.CreateSRV(Out.PageIndexCoordBuffer, PF_R8G8B8A8_UINT);
 	Out.Parameters.Common.NodeDescBuffer			= GraphBuilder.CreateSRV(Out.NodeDescBuffer); 
+	Out.Parameters.Common.AllocatedPageCountBuffer  = GraphBuilder.CreateSRV(Out.PageIndexGlobalCounter, PF_R32_UINT);
 
 	Out.Parameters.Common.CurrGPUMinVoxelSize		= GraphBuilder.CreateSRV(CurrGPUMinVoxelSize, PF_R16F);
 	Out.Parameters.Common.NextGPUMinVoxelSize		= GraphBuilder.CreateSRV(NextGPUMinVoxelSize ? NextGPUMinVoxelSize : CurrGPUMinVoxelSize, PF_R16F);
@@ -1031,7 +1064,6 @@ static FHairStrandsVoxelResources AllocateDummyVirtualVoxelResources(
 	Out.Parameters.Common.SteppingScale_Raytracing	  = 1;
 
 	Out.Parameters.Common.NodeDescCount				= 0;
-	Out.Parameters.Common.IndirectDispatchGroupSize = 64;	
 	Out.Parameters.Common.Raytracing_ShadowOcclusionThreshold	= 1;
 	Out.Parameters.Common.Raytracing_SkyOcclusionThreshold		= 1;
 
@@ -1057,6 +1089,7 @@ static FHairStrandsVoxelResources AllocateDummyVirtualVoxelResources(
 	Out.Parameters.Common.PageIndexOccupancyBuffer	= GraphBuilder.CreateSRV(Out.PageIndexOccupancyBuffer, PF_R32G32_UINT);
 	Out.Parameters.Common.PageIndexCoordBuffer		= GraphBuilder.CreateSRV(Out.PageIndexCoordBuffer, PF_R8G8B8A8_UINT);
 	Out.Parameters.Common.NodeDescBuffer			= GraphBuilder.CreateSRV(Out.NodeDescBuffer); 
+	Out.Parameters.Common.AllocatedPageCountBuffer  = GraphBuilder.CreateSRV(Out.PageIndexGlobalCounter);
 	Out.Parameters.Common.CurrGPUMinVoxelSize		= GraphBuilder.CreateSRV(Dummy2BytesBuffer, PF_R16F);
 	Out.Parameters.Common.NextGPUMinVoxelSize		= Out.Parameters.Common.CurrGPUMinVoxelSize;
 	Out.Parameters.PageTexture						= Out.PageTexture;
@@ -1081,13 +1114,15 @@ static FRDGBufferRef IndirectVoxelPageClear(
 	SCOPED_GPU_STAT(GraphBuilder.RHICmdList, HairStrandsIndVoxelPageClear);
 
 	FRDGBufferRef ClearIndArgsBuffer = GraphBuilder.CreateBuffer(FRDGBufferDesc::CreateIndirectDesc<FRHIDispatchIndirectParameters>(1), TEXT("Hair.VirtualVoxelClearIndirectArgsBuffer"));
+	FRDGBufferSRVRef PageIndexGlobalCounter = GraphBuilder.CreateSRV(VoxelResources.PageIndexGlobalCounter, PF_R32_UINT);
 
 	// Generate the indirect buffer required to clear all voxel allocated linearly in the page volume texture, using the global counter.
 	{
 		FVoxelIndPageClearBufferGenCS::FParameters* Parameters = GraphBuilder.AllocParameters<FVoxelIndPageClearBufferGenCS::FParameters>();
+		Parameters->GroupSize = FVoxelIndPageClearCS::GetGroupSize();
 		Parameters->PageResolution = VoxelResources.Parameters.Common.PageResolution;
 		Parameters->OutIndirectArgsBuffer = GraphBuilder.CreateUAV(ClearIndArgsBuffer);
-		Parameters->PageIndexGlobalCounter = GraphBuilder.CreateSRV(VoxelResources.PageIndexGlobalCounter, PF_R32_UINT);
+		Parameters->PageIndexGlobalCounter = PageIndexGlobalCounter;
 
 		TShaderMapRef<FVoxelIndPageClearBufferGenCS> ComputeShader(ViewInfo.ShaderMap);
 		FComputeShaderUtils::AddPass(
@@ -1104,6 +1139,7 @@ static FRDGBufferRef IndirectVoxelPageClear(
 		Parameters->VirtualVoxelParams = VoxelResources.Parameters.Common;
 		Parameters->OutPageTexture = GraphBuilder.CreateUAV(FRDGTextureUAVDesc(VoxelResources.PageTexture));
 		Parameters->IndirectDispatchBuffer = ClearIndArgsBuffer;
+		Parameters->PageIndexGlobalCounter = PageIndexGlobalCounter;
 
 		TShaderMapRef<FVoxelIndPageClearCS> ComputeShader(ViewInfo.ShaderMap);
 		FComputeShaderUtils::AddPass(
@@ -1132,7 +1168,6 @@ class FVoxelRasterComputeCS : public FGlobalShader
 	BEGIN_SHADER_PARAMETER_STRUCT(FParameters, )
 		SHADER_PARAMETER_STRUCT(FHairStrandsVoxelCommonParameters, VirtualVoxelParams)
 		SHADER_PARAMETER(uint32, MacroGroupId)
-		SHADER_PARAMETER(uint32, DispatchCountX)
 		SHADER_PARAMETER(uint32, MaxRasterCount)
 		SHADER_PARAMETER(uint32, FrameIdMod8)
 		SHADER_PARAMETER_STRUCT_INCLUDE(FHairStrandsInstanceParameters, HairInstance)
@@ -1143,11 +1178,17 @@ class FVoxelRasterComputeCS : public FGlobalShader
 	END_SHADER_PARAMETER_STRUCT()
 
 public:
+	static uint32 GetGroupSize() { return HAIR_VERTEXCOUNT_GROUP_SIZE; }
 	static bool ShouldCompilePermutation(const FGlobalShaderPermutationParameters& Parameters) { return IsHairStrandsSupported(EHairStrandsShaderType::Strands, Parameters.Platform); }
 	static void ModifyCompilationEnvironment(const FGlobalShaderPermutationParameters& Parameters, FShaderCompilerEnvironment& OutEnvironment)
 	{
 		FGlobalShader::ModifyCompilationEnvironment(Parameters, OutEnvironment);
 		OutEnvironment.SetDefine(TEXT("SHADER_RASTERCOMPUTE"), 1);
+
+		FPermutationDomain PermutationVector(Parameters.PermutationId);
+		const bool bUseCulling = PermutationVector.Get<FCulling>() != 0;
+		const uint32 GroupSize = PermutationVector.Get<FGroupSize>();		
+		OutEnvironment.SetDefine(TEXT("GROUP_SIZE"), bUseCulling ? GroupSize : GetGroupSize());
 	}
 };
 
@@ -1175,7 +1216,7 @@ static void AddVirtualVoxelizationComputeRasterPass(
 		FRDGTextureUAVRef PageTextureUAV = GraphBuilder.CreateUAV(VoxelResources.PageTexture);
 
 		const uint32 FrameIdMode8 = ViewInfo && ViewInfo->ViewState ? (ViewInfo->ViewState->GetFrameIndex() % 8) : 0;
-		const uint32 GroupSize = GetVendorOptimalGroupSize1D();
+		const uint32 GroupSize = GetVendorOptimalGroupSize1D(); // TODO when updating culling
 		const FVector& TranslatedWorldOffset = ViewInfo->ViewMatrices.GetPreViewTranslation();
 
 		FVoxelRasterComputeCS::FPermutationDomain PermutationVector_Off;
@@ -1211,7 +1252,6 @@ static void AddVirtualVoxelizationComputeRasterPass(
 				PassParameters->VirtualVoxelParams = VoxelResources.Parameters.Common;
 				PassParameters->MacroGroupId = MacroGroup.MacroGroupId;
 				PassParameters->VoxelizationViewInfoBuffer = VoxelizationViewInfoBufferSRV;
-				PassParameters->DispatchCountX = 1;
 				PassParameters->OutPageTexture = PageTextureUAV;
 				PassParameters->FrameIdMod8 = FrameIdMode8;
 
@@ -1228,8 +1268,7 @@ static void AddVirtualVoxelizationComputeRasterPass(
 				}
 				else
 				{
-					const FIntVector DispatchCount = ComputeDispatchCount(PassParameters->VertexCount, GroupSize);
-					PassParameters->DispatchCountX = DispatchCount.X;
+					const FIntVector DispatchCount(FMath::DivideAndRoundUp(PassParameters->VertexCount, FVoxelRasterComputeCS::GetGroupSize()), 1, 1);
 					FComputeShaderUtils::AddPass(GraphBuilder, RDG_EVENT_NAME("HairStrands::VoxelComputeRaster(culling=off)"), ComputeShader_CullingOff, PassParameters, DispatchCount);
 				}
 			}
@@ -1262,6 +1301,7 @@ class FVirtualVoxelGenerateMipCS : public FGlobalShader
 
 	BEGIN_SHADER_PARAMETER_STRUCT(FParameters, )
 		SHADER_PARAMETER_STRUCT_INCLUDE(FSceneTextureParameters, SceneTextures)
+		SHADER_PARAMETER_RDG_BUFFER_SRV(Buffer, AllocatedPageCountBuffer)
 		SHADER_PARAMETER(FIntVector, PageCountResolution)
 		SHADER_PARAMETER(uint32, PageResolution)
 		SHADER_PARAMETER(uint32, SourceMip)
@@ -1274,11 +1314,15 @@ class FVirtualVoxelGenerateMipCS : public FGlobalShader
 	END_SHADER_PARAMETER_STRUCT()
 
 public:
+	static FIntVector GetGroupSize() { return GHairVoxel_GroupSize_NumVoxelPerPage_1_NumAllocatedPage; }
 	static bool ShouldCompilePermutation(const FGlobalShaderPermutationParameters& Parameters) { return IsHairStrandsSupported(EHairStrandsShaderType::Strands, Parameters.Platform); }
 	static void ModifyCompilationEnvironment(const FGlobalShaderPermutationParameters& Parameters, FShaderCompilerEnvironment& OutEnvironment)
 	{
 		FGlobalShader::ModifyCompilationEnvironment(Parameters, OutEnvironment);
 		OutEnvironment.SetDefine(TEXT("SHADER_MIP_VIRTUALVOXEL"), 1);
+		OutEnvironment.SetDefine(TEXT("GROUP_SIZE_X"), GetGroupSize().X);
+		OutEnvironment.SetDefine(TEXT("GROUP_SIZE_Y"), GetGroupSize().Y);
+		OutEnvironment.SetDefine(TEXT("GROUP_SIZE_Z"), GetGroupSize().Z);
 	}
 };
 
@@ -1290,7 +1334,7 @@ class FVirtualVoxelIndirectArgMipCS : public FGlobalShader
 	BEGIN_SHADER_PARAMETER_STRUCT(FParameters, )
 		SHADER_PARAMETER(uint32, PageResolution)
 		SHADER_PARAMETER(uint32, TargetMipIndex)
-		SHADER_PARAMETER(uint32, DispatchGroupSize)
+		SHADER_PARAMETER(FIntVector, DispatchGroupSize)
 		SHADER_PARAMETER_RDG_BUFFER_SRV(Buffer, InIndirectArgs)
 		SHADER_PARAMETER_RDG_BUFFER_UAV(RWBuffer, OutIndirectArgs)
 	END_SHADER_PARAMETER_STRUCT()
@@ -1325,11 +1369,15 @@ class FVirtualVoxelPatchPageIndexWithMipDataCS : public FGlobalShader
 	END_SHADER_PARAMETER_STRUCT()
 
 public:
+	static FIntVector GetGroupSize() { return GHairVoxel_GroupSize_NumVoxelPerPage_1_NumAllocatedPage; }
 	static bool ShouldCompilePermutation(const FGlobalShaderPermutationParameters& Parameters) { return IsHairStrandsSupported(EHairStrandsShaderType::Strands, Parameters.Platform); }
 	static void ModifyCompilationEnvironment(const FGlobalShaderPermutationParameters& Parameters, FShaderCompilerEnvironment& OutEnvironment)
 	{
 		FGlobalShader::ModifyCompilationEnvironment(Parameters, OutEnvironment);
 		OutEnvironment.SetDefine(TEXT("SHADER_UPDATE_PAGEINDEX"), 1);
+		OutEnvironment.SetDefine(TEXT("GROUP_SIZE_X"), GetGroupSize().X);
+		OutEnvironment.SetDefine(TEXT("GROUP_SIZE_Y"), GetGroupSize().Y);
+		OutEnvironment.SetDefine(TEXT("GROUP_SIZE_Z"), GetGroupSize().Z);
 	}
 };
 
@@ -1360,14 +1408,13 @@ static void AddVirtualVoxelGenerateMipPass(
 	for (uint32 MipIt = 0; MipIt < MipCount - 1; ++MipIt)
 	{
 		const uint32 TargetMipIndex = MipIt + 1;
-		const uint32 DispatchGroupSize = 64;
 		FRDGBufferRef MipIndirectArgs = GraphBuilder.CreateBuffer(FRDGBufferDesc::CreateIndirectDesc<FRHIDispatchIndirectParameters>(), TEXT("Hair.VirtualVoxelMipIndirectArgsBuffer"));
 		MipIndirectArgsBuffers.Add(MipIndirectArgs);
 
 		FVirtualVoxelIndirectArgMipCS::FParameters* Parameters = GraphBuilder.AllocParameters<FVirtualVoxelIndirectArgMipCS::FParameters>();
 		Parameters->PageResolution		= VoxelResources.Parameters.Common.PageResolution;
 		Parameters->TargetMipIndex		= TargetMipIndex;
-		Parameters->DispatchGroupSize	= DispatchGroupSize;
+		Parameters->DispatchGroupSize	= FVirtualVoxelGenerateMipCS::GetGroupSize();
 		Parameters->InIndirectArgs		= GraphBuilder.CreateSRV(IndirectArgsBuffer);
 		Parameters->OutIndirectArgs		= GraphBuilder.CreateUAV(MipIndirectArgs);
 
@@ -1388,6 +1435,7 @@ static void AddVirtualVoxelGenerateMipPass(
 		Parameters->PageCountResolution = VoxelResources.Parameters.Common.PageCountResolution;
 		Parameters->SourceMip = SourceMipIndex;
 		Parameters->TargetMip = TargetMipIndex;
+		Parameters->AllocatedPageCountBuffer = VoxelResources.Parameters.Common.AllocatedPageCountBuffer;
 		Parameters->IndirectDispatchArgs = MipIndirectArgsBuffers[MipIt];
 
 		TShaderMapRef<FVirtualVoxelGenerateMipCS> ComputeShader(View.ShaderMap);
