RayTracing: Fix flickering of objects from the RT scene during movement in the editor

This was tied to r.AsyncCacheMeshDrawCommands which includes raytracing caches. This in turn impacted the logic in PrimitiveCull which uses ShouldCullForRayTracing which relies on flags set during the async portion.

diff --git a/Engine/Source/Runtime/Renderer/Private/DeferredShadingRenderer.cpp b/Engine/Source/Runtime/Renderer/Private/DeferredShadingRenderer.cpp
--- a/Engine/Source/Runtime/Renderer/Private/DeferredShadingRenderer.cpp
+++ b/Engine/Source/Runtime/Renderer/Private/DeferredShadingRenderer.cpp
@@ -505,12 +505,6 @@ FDeferredShadingSceneRenderer::FDeferredShadingSceneRenderer(const FSceneViewFam
 {
 	ViewPipelineStates.SetNum(Views.Num());
 
-#if RHI_RAYTRACING
-	// Flags initialized later in render thread
-	bAnyRayTracingPassEnabled = false;
-	bShouldUpdateRayTracingScene = false;
-#endif
-
 	ShadowSceneRenderer = MakeUnique<FShadowSceneRenderer>(*this);
 }
 
diff --git a/Engine/Source/Runtime/Renderer/Private/DeferredShadingRenderer.h b/Engine/Source/Runtime/Renderer/Private/DeferredShadingRenderer.h
--- a/Engine/Source/Runtime/Renderer/Private/DeferredShadingRenderer.h
+++ b/Engine/Source/Runtime/Renderer/Private/DeferredShadingRenderer.h
@@ -308,8 +308,7 @@ public:
 	FLumenCardRenderer LumenCardRenderer;
 
 #if RHI_RAYTRACING
-	bool bAnyRayTracingPassEnabled;
-	bool bShouldUpdateRayTracingScene;
+	bool bShouldUpdateRayTracingScene =  false;
 
 	void InitializeRayTracingFlags_RenderThread();
 #endif
diff --git a/Engine/Source/Runtime/Renderer/Private/SceneRendering.h b/Engine/Source/Runtime/Renderer/Private/SceneRendering.h
--- a/Engine/Source/Runtime/Renderer/Private/SceneRendering.h
+++ b/Engine/Source/Runtime/Renderer/Private/SceneRendering.h
@@ -2092,6 +2092,11 @@ public:
 	bool bIsFirstSceneRenderer;
 	bool bIsLastSceneRenderer;
 
+#if RHI_RAYTRACING
+	/** Wheter any ray tracing pass has been enabled on the current frame. */
+	bool bAnyRayTracingPassEnabled = false;
+#endif
+
 public:
 
 	FSceneRenderer(const FSceneViewFamily* InViewFamily, FHitProxyConsumer* HitProxyConsumer);
diff --git a/Engine/Source/Runtime/Renderer/Private/SceneVisibility.cpp b/Engine/Source/Runtime/Renderer/Private/SceneVisibility.cpp
--- a/Engine/Source/Runtime/Renderer/Private/SceneVisibility.cpp
+++ b/Engine/Source/Runtime/Renderer/Private/SceneVisibility.cpp
@@ -4769,6 +4769,14 @@ void FSceneRenderer::ComputeViewVisibility(
 
 			{
 				TRACE_CPUPROFILER_EVENT_SCOPE(FSceneRenderer_Cull);
+#if RHI_RAYTRACING
+				if (bAnyRayTracingPassEnabled)
+				{
+					// The logic inside PrimitiveCull makes use of ShouldCullForRayTracing to decide if the primitive should be considered for raytracing
+					// Therefore we must be sure we are done with caching the mesh draw commands (which include raytracing caches) before it runs if raytracing is being used.
+					Scene->WaitForCacheMeshDrawCommandsTask();
+				}
+#endif
 				int32 NumCulledPrimitivesForView = PrimitiveCull(Scene, View, bNeedsFrustumCulling);
 				STAT(NumCulledPrimitives += NumCulledPrimitivesForView);
 			}
