Fixed: NonUnity compilation errors in VT

diff --git a/Engine/Source/Runtime/Renderer/Private/VT/VirtualTextureSpace.cpp b/Engine/Source/Runtime/Renderer/Private/VT/VirtualTextureSpace.cpp
--- a/Engine/Source/Runtime/Renderer/Private/VT/VirtualTextureSpace.cpp
+++ b/Engine/Source/Runtime/Renderer/Private/VT/VirtualTextureSpace.cpp
@@ -14,6 +14,7 @@
 #include "RHIStaticStates.h"
 #include "HAL/IConsoleManager.h"
 #include "SceneUtils.h"
+#include "RenderGraph.h"
 
 #include "VT/AllocatedVirtualTexture.h"
 
