Fix some IWYU compile errors

diff --git a/Engine/Source/Runtime/Renderer/Private/RayTracing/RaytracingSkylight.cpp b/Engine/Source/Runtime/Renderer/Private/RayTracing/RaytracingSkylight.cpp
--- a/Engine/Source/Runtime/Renderer/Private/RayTracing/RaytracingSkylight.cpp
+++ b/Engine/Source/Runtime/Renderer/Private/RayTracing/RaytracingSkylight.cpp
@@ -7,6 +7,7 @@
 #include "DeferredShadingRenderer.h"
 #include "DataDrivenShaderPlatformInfo.h"
 #include "SceneTextureParameters.h"
+#include "DataDrivenShaderPlatformInfo.h"
 
 static int32 GRayTracingSkyLight = 0;
 
