Fix flickering lumen reflections when dynamic res changes. When async compute lumen diffuse indirect was enabled but not async lumen reflections, we were doing the depth copy before the reflection pass. This caused reflections to read from this frame's depth buffer instead of the Lumen.DepthHistory, resulting in graphical artifacts.

There is no performance hit from moving the depth copy to the graphics pipe, in fact it seems slightly faster on consoles.

diff --git a/Engine/Source/Runtime/Renderer/Private/DeferredShadingRenderer.h b/Engine/Source/Runtime/Renderer/Private/DeferredShadingRenderer.h
--- a/Engine/Source/Runtime/Renderer/Private/DeferredShadingRenderer.h
+++ b/Engine/Source/Runtime/Renderer/Private/DeferredShadingRenderer.h
@@ -122,10 +122,10 @@ struct FAsyncLumenIndirectLightingOutputs
 	{
 		check(StepsLeft == ELumenIndirectLightingSteps::All);
 
-		EnumRemoveFlags(StepsLeft, ELumenIndirectLightingSteps::ScreenProbeGather | ELumenIndirectLightingSteps::StoreDepthHistory);
+		EnumRemoveFlags(StepsLeft, ELumenIndirectLightingSteps::ScreenProbeGather);
 		if (bAsyncReflections)
 		{
-			EnumRemoveFlags(StepsLeft, ELumenIndirectLightingSteps::Reflections);
+			EnumRemoveFlags(StepsLeft, ELumenIndirectLightingSteps::Reflections | ELumenIndirectLightingSteps::StoreDepthHistory);
 		}
 	}
 
