Fix a bug that decal is missing skylight on mobile deferred with static lighting disabled.

diff --git a/Engine/Source/Runtime/Renderer/Private/MobileBasePassRendering.h b/Engine/Source/Runtime/Renderer/Private/MobileBasePassRendering.h
--- a/Engine/Source/Runtime/Renderer/Private/MobileBasePassRendering.h
+++ b/Engine/Source/Runtime/Renderer/Private/MobileBasePassRendering.h
@@ -375,19 +375,12 @@ public:
 	{		
 		MobileBasePassModifyCompilationEnvironment(Parameters, OutEnvironment, OutputFormat);
 
-		static auto AllowStaticLightingVar = IConsoleManager::Get().FindTConsoleVariableDataInt(TEXT("r.AllowStaticLighting"));
-		const bool bAllowStaticLighting = (!AllowStaticLightingVar || AllowStaticLightingVar->GetValueOnAnyThread() != 0);
-
 		const bool bMobileUsesShadowMaskTexture = MobileUsesShadowMaskTexture(Parameters.Platform);
 		const bool bEnableClusteredReflections = MobileForwardEnableClusteredReflections(Parameters.Platform);
-		const bool bIsLit = Parameters.MaterialParameters.ShadingModels.IsLit();
 		const bool bTranslucentMaterial = IsTranslucentBlendMode(Parameters.MaterialParameters) || Parameters.MaterialParameters.ShadingModels.HasShadingModel(MSM_SingleLayerWater);
-		const bool bDeferredShadingEnabled = IsMobileDeferredShadingEnabled(Parameters.Platform);
-		const bool bForwardShading = !bDeferredShadingEnabled || bTranslucentMaterial;
-		const bool bRenderSkyLightInBasePass = bEnableSkyLight && bIsLit && (bForwardShading || bAllowStaticLighting);
 
 		TMobileBasePassPSBaseType<LightMapPolicyType>::ModifyCompilationEnvironment(Parameters, OutEnvironment);
-		OutEnvironment.SetDefine(TEXT("ENABLE_SKY_LIGHT"), bRenderSkyLightInBasePass);
+		OutEnvironment.SetDefine(TEXT("ENABLE_SKY_LIGHT"), bEnableSkyLight);
 		OutEnvironment.SetDefine(TEXT("ENABLE_AMBIENT_OCCLUSION"), IsMobileAmbientOcclusionEnabled(Parameters.Platform) ? 1u : 0u);
 		
 		FForwardLightingParameters::ModifyCompilationEnvironment(Parameters.Platform, OutEnvironment);
diff --git a/Engine/Source/Runtime/Renderer/Private/MobileDeferredShadingPass.cpp b/Engine/Source/Runtime/Renderer/Private/MobileDeferredShadingPass.cpp
--- a/Engine/Source/Runtime/Renderer/Private/MobileDeferredShadingPass.cpp
+++ b/Engine/Source/Runtime/Renderer/Private/MobileDeferredShadingPass.cpp
@@ -357,9 +357,8 @@ static void GetLightMaterial(const FCachedLightMaterial& DefaultLightMaterial, c
 void RenderReflectionEnvironmentSkyLighting(FRHICommandListImmediate& RHICmdList, const FScene& Scene, const FViewInfo& View)
 {
 	// Skylights with static lighting already had their diffuse contribution baked into lightmaps
-	const bool bIsDynamicLighting = !FReadOnlyCVARCache::Get().bAllowStaticLighting || Scene.GetForceNoPrecomputedLighting();
 	const bool bSkyLight = Scene.SkyLight && !Scene.SkyLight->bHasStaticLighting && View.Family->EngineShowFlags.SkyLighting;
-	const bool bDynamicSkyLight = bSkyLight && (!Scene.SkyLight->bWantsStaticShadowing || bIsDynamicLighting);
+	const bool bDynamicSkyLight = bSkyLight && !Scene.SkyLight->bWantsStaticShadowing;
 	const bool bClustredReflection = (View.NumBoxReflectionCaptures + View.NumSphereReflectionCaptures) > 0;
 	const bool bPlanarReflection = Scene.GetForwardPassGlobalPlanarReflection() != nullptr;
 	if (!(bSkyLight || bClustredReflection || bPlanarReflection))
@@ -502,9 +501,8 @@ static void RenderDirectionalLight(FRHICommandListImmediate& RHICmdList, const F
 	}
 
 	// Skylights with static lighting already had their diffuse contribution baked into lightmaps
-	const bool bIsDynamicLighting = !FReadOnlyCVARCache::Get().bAllowStaticLighting || Scene.GetForceNoPrecomputedLighting();
 	const bool bSkyLight = Scene.SkyLight && !Scene.SkyLight->bHasStaticLighting && View.Family->EngineShowFlags.SkyLighting;
-	const bool bDynamicSkyLight = bSkyLight && (!Scene.SkyLight->bWantsStaticShadowing || bIsDynamicLighting);
+	const bool bDynamicSkyLight = bSkyLight && !Scene.SkyLight->bWantsStaticShadowing;
 	const bool bDynamicShadows = DirectionalLight.Proxy->CastsDynamicShadow() && (LightingChannel == 0u) && View.Family->EngineShowFlags.DynamicShadows;
 	const bool bPlanarReflection = Scene.GetForwardPassGlobalPlanarReflection() != nullptr;
 
diff --git a/Engine/Source/Runtime/Renderer/Private/RendererScene.cpp b/Engine/Source/Runtime/Renderer/Private/RendererScene.cpp
--- a/Engine/Source/Runtime/Renderer/Private/RendererScene.cpp
+++ b/Engine/Source/Runtime/Renderer/Private/RendererScene.cpp
@@ -6383,11 +6383,19 @@ bool FScene::ShouldRenderSkylightInBasePass(bool bIsTranslucent) const
 	{
 		bool bRenderSkyLight = SkyLight && !SkyLight->bHasStaticLighting;
 		const bool bIsForwardShading = bIsTranslucent || !IsMobileDeferredShadingEnabled(GetShaderPlatform());
-		const bool bIsDynamicLighting = !ReadOnlyCVARCache.bAllowStaticLighting || GetForceNoPrecomputedLighting();
 
-		// Both stationary and movable skylights are applied in base pass for forward shading
-		// Only stationary skylights are applied in base pass for deferred with static lighting enabled
-		return bRenderSkyLight && (bIsForwardShading || (!bIsDynamicLighting && SkyLight->bWantsStaticShadowing));
+		if (bIsForwardShading)
+		{
+			// Both stationary and movable skylights are applied in base pass for forward shading
+			bRenderSkyLight = bRenderSkyLight && (ReadOnlyCVARCache.bEnableStationarySkylight || !SkyLight->bWantsStaticShadowing);
+		}
+		else
+		{
+			// Only stationary skylights are applied in base pass for deferred
+			bRenderSkyLight = bRenderSkyLight && (ReadOnlyCVARCache.bEnableStationarySkylight && SkyLight->bWantsStaticShadowing);
+		}
+
+		return bRenderSkyLight;
 	}
 	else
 	{
