Changed PreGatherDynamicMeshElements virtual method to a callback.

diff --git a/Engine/Source/Runtime/Renderer/Private/DeferredShadingRenderer.cpp b/Engine/Source/Runtime/Renderer/Private/DeferredShadingRenderer.cpp
--- a/Engine/Source/Runtime/Renderer/Private/DeferredShadingRenderer.cpp
+++ b/Engine/Source/Runtime/Renderer/Private/DeferredShadingRenderer.cpp
@@ -1005,18 +1005,41 @@ struct FDeferredShadingRayTracingMaterialGatheringContext : public FRayTracingMa
 	}
 };
 
-bool FDeferredShadingSceneRenderer::GatherRayTracingWorldInstancesForView(FRDGBuilder& GraphBuilder, FViewInfo& View, FRayTracingScene& RayTracingScene, FRayTracingRelevantPrimitiveList& RelevantPrimitiveList)
+struct FRayTracingRelevantPrimitiveTaskData
 {
-	checkf(IsRayTracingEnabled() && bAnyRayTracingPassEnabled, TEXT("GatherRayTracingWorldInstancesForView should only be called if ray tracing is used"))
+	FRayTracingRelevantPrimitiveList List;
+	FGraphEventRef Task;
+};
 
+bool FDeferredShadingSceneRenderer::GatherRayTracingWorldInstancesForView(FRDGBuilder& GraphBuilder, FViewInfo& View, FRayTracingScene& RayTracingScene, FRayTracingRelevantPrimitiveTaskData* RayTracingRelevantPrimitiveTaskData)
+{
 	TRACE_CPUPROFILER_EVENT_SCOPE(FDeferredShadingSceneRenderer::GatherRayTracingWorldInstances);
 	SCOPE_CYCLE_COUNTER(STAT_GatherRayTracingWorldInstances);
 
-	// Ensure that any invalidated cached uniform expressions have been updated on the rendering thread.
+	if (!bAnyRayTracingPassEnabled)
+	{
+		return false;
+	}
+
+	check(RayTracingRelevantPrimitiveTaskData)
+
+	// Wait until RayTracingRelevantPrimitiveList is ready
+	if (RayTracingRelevantPrimitiveTaskData->Task.IsValid())
+	{
+		RayTracingRelevantPrimitiveTaskData->Task->Wait();
+		RayTracingRelevantPrimitiveTaskData->Task.SafeRelease();
+	}
+
+	FRayTracingRelevantPrimitiveList& RelevantPrimitiveList = RayTracingRelevantPrimitiveTaskData->List;
+
+	// Prepare ray tracing scene instance list
+	checkf(RelevantPrimitiveList.bValid, TEXT("Ray tracing relevant primitive list is expected to have been created before GatherRayTracingWorldInstancesForView() is called."));
+
+	// Check that any invalidated cached uniform expressions have been updated on the rendering thread.
 	// Normally this work is done through FMaterialRenderProxy::UpdateUniformExpressionCacheIfNeeded,
 	// however ray tracing material processing (FMaterialShader::GetShaderBindings, which accesses UniformExpressionCache)
 	// is done on task threads, therefore all work must be done here up-front as UpdateUniformExpressionCacheIfNeeded is not free-threaded.
-	FMaterialRenderProxy::UpdateDeferredCachedUniformExpressions();
+	check(!FMaterialRenderProxy::HasDeferredUniformExpressionCacheRequests());
 
 	RayTracingCollector.ClearViewMeshArrays();
 
@@ -2323,23 +2346,12 @@ void FDeferredShadingSceneRenderer::WaitForRayTracingScene(FRDGBuilder& GraphBui
 				View->LumenHardwareRayTracingMaterialPipeline = ReferenceView.LumenHardwareRayTracingMaterialPipeline;
 			}
 		}
-
-		if (RayTracingDynamicGeometryUpdateEndTransition)
-		{
-			RHICmdList.EndTransition(RayTracingDynamicGeometryUpdateEndTransition);
-			RayTracingDynamicGeometryUpdateEndTransition = nullptr;
-		}
 	});
 }
 
-struct FRayTracingRelevantPrimitiveTaskData
-{
-	FRayTracingRelevantPrimitiveList List;
-	FGraphEventRef Task;
-};
 #endif // RHI_RAYTRACING
 
-void FDeferredShadingSceneRenderer::PreGatherDynamicMeshElements()
+void FDeferredShadingSceneRenderer::PreGatherDynamicMeshElements(FInitViewTaskDatas& TaskDatas)
 {
 #if RHI_RAYTRACING
 	if (bAnyRayTracingPassEnabled)
@@ -2347,27 +2359,24 @@ void FDeferredShadingSceneRenderer::PreGatherDynamicMeshElements()
 		const int32 ReferenceViewIndex = 0;
 		FViewInfo& ReferenceView = Views[ReferenceViewIndex];
 
-		RayTracingRelevantPrimitiveTaskData = Allocator.Create<FRayTracingRelevantPrimitiveTaskData>();
-		RayTracingRelevantPrimitiveTaskData->Task = FFunctionGraphTask::CreateAndDispatchWhenReady(
-			[Scene = this->Scene, &ReferenceView, &RayTracingRelevantPrimitiveList = RayTracingRelevantPrimitiveTaskData->List]()
-			{
-				FTaskTagScope TaskTagScope(ETaskTag::EParallelRenderingThread);
-				GatherRayTracingRelevantPrimitives(*Scene, ReferenceView, RayTracingRelevantPrimitiveList);
-			}, TStatId(), nullptr, ENamedThreads::AnyNormalThreadHiPriTask);
+		TaskDatas.RayTracingRelevantPrimitives = Allocator.Create<FRayTracingRelevantPrimitiveTaskData>();
+		TaskDatas.RayTracingRelevantPrimitives->Task = FFunctionGraphTask::CreateAndDispatchWhenReady(
+			[Scene = Scene, &ReferenceView, &RayTracingRelevantPrimitiveList = TaskDatas.RayTracingRelevantPrimitives->List]()
+		{
+			FTaskTagScope TaskTagScope(ETaskTag::EParallelRenderingThread);
+			GatherRayTracingRelevantPrimitives(*Scene, ReferenceView, RayTracingRelevantPrimitiveList);
+		}, TStatId(), nullptr, ENamedThreads::AnyNormalThreadHiPriTask);
 	}
 #endif // RHI_RAYTRACING
 
-	const bool bHasRayTracedOverlay = HasRayTracedOverlay(ViewFamily);
-
 	extern int32 GEarlyInitDynamicShadows;
 
 	if (GEarlyInitDynamicShadows &&
-		CurrentDynamicShadowsTaskData == nullptr &&
 		ViewFamily.EngineShowFlags.DynamicShadows
 		&& !ViewFamily.EngineShowFlags.HitProxies
-		&& !bHasRayTracedOverlay)
+		&& !HasRayTracedOverlay(ViewFamily))
 	{
-		CurrentDynamicShadowsTaskData = BeginInitDynamicShadows(true);
+		TaskDatas.DynamicShadows = BeginInitDynamicShadows(true);
 	}
 }
 
@@ -2651,8 +2660,6 @@ void FDeferredShadingSceneRenderer::Render(FRDGBuilder& GraphBuilder)
 		? FExclusiveDepthStencil::DepthRead_StencilWrite
 		: FExclusiveDepthStencil::DepthWrite_StencilWrite;
 
-	FILCUpdatePrimTaskData ILCTaskData;
-
 	// Find the visible primitives.
 	if (GDynamicRHI->RHIIncludeOptionalFlushes())
 	{
@@ -2661,9 +2668,11 @@ void FDeferredShadingSceneRenderer::Render(FRDGBuilder& GraphBuilder)
 
 	FInstanceCullingManager& InstanceCullingManager = *GraphBuilder.AllocObject<FInstanceCullingManager>(GetSceneUniforms(), Scene->GPUScene.IsEnabled(), GraphBuilder);
 
+	FInitViewTaskDatas InitViewTaskDatas;
+
 	{
 		RDG_GPU_STAT_SCOPE(GraphBuilder, VisibilityCommands);
-		BeginInitViews(GraphBuilder, SceneTexturesConfig, BasePassDepthStencilAccess, ILCTaskData, InstanceCullingManager, VirtualTextureUpdater.Get());
+		BeginInitViews(GraphBuilder, SceneTexturesConfig, BasePassDepthStencilAccess, InstanceCullingManager, VirtualTextureUpdater.Get(), InitViewTaskDatas);
 	}
 
 	// GetBinIndexTranslator cannot be called before UpdateAllPrimitiveSceneInfos which can change the number of raster bins
@@ -2844,20 +2853,7 @@ void FDeferredShadingSceneRenderer::Render(FRDGBuilder& GraphBuilder)
 	}
 
 #if RHI_RAYTRACING
-	if (bAnyRayTracingPassEnabled)
-	{
-		// Wait until RayTracingRelevantPrimitiveList is ready
-		if (RayTracingRelevantPrimitiveTaskData->Task.IsValid())
-		{
-			RayTracingRelevantPrimitiveTaskData->Task->Wait();
-			RayTracingRelevantPrimitiveTaskData->Task.SafeRelease();
-		}
-
-		// Prepare ray tracing scene instance list
-		checkf(RayTracingRelevantPrimitiveTaskData->List.bValid, TEXT("Ray tracing relevant primitive list is expected to have been created before GatherRayTracingWorldInstancesForView() is called."));
-
-		GatherRayTracingWorldInstancesForView(GraphBuilder, ReferenceView, RayTracingScene, RayTracingRelevantPrimitiveTaskData->List);
-	}
+	GatherRayTracingWorldInstancesForView(GraphBuilder, ReferenceView, RayTracingScene, InitViewTaskDatas.RayTracingRelevantPrimitives);
 #endif // RHI_RAYTRACING
 
 	const bool bUseGBuffer = IsUsingGBuffers(ShaderPlatform);
@@ -2978,7 +2974,7 @@ void FDeferredShadingSceneRenderer::Render(FRDGBuilder& GraphBuilder)
 	{
 		{
 			RDG_RHI_GPU_STAT_SCOPE(GraphBuilder, VisibilityCommands);
-			EndInitViews(GraphBuilder, LumenFrameTemporaries, ILCTaskData, InstanceCullingManager, ExternalAccessQueue);
+			EndInitViews(GraphBuilder, LumenFrameTemporaries, InstanceCullingManager, ExternalAccessQueue, InitViewTaskDatas);
 		}
 
 		// Dynamic vertex and index buffers need to be committed before rendering.
diff --git a/Engine/Source/Runtime/Renderer/Private/DeferredShadingRenderer.h b/Engine/Source/Runtime/Renderer/Private/DeferredShadingRenderer.h
--- a/Engine/Source/Runtime/Renderer/Private/DeferredShadingRenderer.h
+++ b/Engine/Source/Runtime/Renderer/Private/DeferredShadingRenderer.h
@@ -55,6 +55,7 @@ struct FForwardBasePassTextures;
 struct FTranslucentLightInjectionCollector;
 struct FRayTracingPickingFeedback;
 struct FDBufferTextures;
+struct FILCUpdatePrimTaskData;
 
 #if RHI_RAYTRACING
 struct FRayTracingRelevantPrimitiveTaskData;
@@ -524,15 +525,36 @@ private:
 		const TArray<FProjectedShadowInfo*, SceneRenderingAllocator>& ViewDependentWholeSceneShadows,
 		TArray<FProjectedShadowInfo*, SceneRenderingAllocator>& OutPreShadows);
 
-	void PreGatherDynamicMeshElements() override final;
+	struct FInitViewTaskDatas
+	{
+		FILCUpdatePrimTaskData* ILCUpdatePrim = nullptr;
+	#if RHI_RAYTRACING
+		FRayTracingRelevantPrimitiveTaskData* RayTracingRelevantPrimitives = nullptr;
+	#endif
+		FDynamicShadowsTaskData* DynamicShadows = nullptr;
+	};
+
+	void PreGatherDynamicMeshElements(FInitViewTaskDatas& TaskDatas);
 
 	void PreVisibilityFrameSetup(FRDGBuilder& GraphBuilder, const FSceneTexturesConfig& SceneTexturesConfig);
 
 	void ComputeLightVisibility();
 
 	/** Determines which primitives are visible for each view. */
-	void BeginInitViews(FRDGBuilder& GraphBuilder, const FSceneTexturesConfig& SceneTexturesConfig, FExclusiveDepthStencil::Type BasePassDepthStencilAccess, struct FILCUpdatePrimTaskData& ILCTaskData, FInstanceCullingManager& InstanceCullingManager, FVirtualTextureUpdater* VirtualTextureUpdater);
-	void EndInitViews(FRDGBuilder& GraphBuilder, FLumenSceneFrameTemporaries& FrameTemporaries, struct FILCUpdatePrimTaskData& ILCTaskData, FInstanceCullingManager& InstanceCullingManager, FRDGExternalAccessQueue& ExternalAccessQueue);
+	void BeginInitViews(
+		FRDGBuilder& GraphBuilder,
+		const FSceneTexturesConfig& SceneTexturesConfig,
+		FExclusiveDepthStencil::Type BasePassDepthStencilAccess,
+		FInstanceCullingManager& InstanceCullingManager,
+		FVirtualTextureUpdater* VirtualTextureUpdater,
+		FInitViewTaskDatas& TaskDatas);
+
+	void EndInitViews(
+		FRDGBuilder& GraphBuilder,
+		FLumenSceneFrameTemporaries& FrameTemporaries,
+		FInstanceCullingManager& InstanceCullingManager,
+		FRDGExternalAccessQueue& ExternalAccessQueue,
+		const FInitViewTaskDatas& TaskDatas);
 
 	void BeginUpdateLumenSceneTasks(FRDGBuilder& GraphBuilder, FLumenSceneFrameTemporaries& FrameTemporaries);
 	void UpdateLumenScene(FRDGBuilder& GraphBuilder, FLumenSceneFrameTemporaries& FrameTemporaries);
@@ -1117,7 +1139,7 @@ private:
 	void RayTracingDisplayPicking(const FRayTracingPickingFeedback& PickingFeedback, FScreenMessageWriter& Writer);
 
 	/** Fills RayTracingScene instance list for the given View and adds relevant ray tracing data to the view. Does not reset previous scene contents. */
-	bool GatherRayTracingWorldInstancesForView(FRDGBuilder& GraphBuilder, FViewInfo& View, FRayTracingScene& RayTracingScene, struct FRayTracingRelevantPrimitiveList& RelevantPrimitiveList);
+	bool GatherRayTracingWorldInstancesForView(FRDGBuilder& GraphBuilder, FViewInfo& View, FRayTracingScene& RayTracingScene, FRayTracingRelevantPrimitiveTaskData* RayTracingRelevantPrimitiveTaskData);
 
 	bool SetupRayTracingPipelineStates(FRDGBuilder& GraphBuilder);
 	void SetupRayTracingLightDataForViews(FRDGBuilder& GraphBuilder);
@@ -1171,18 +1193,10 @@ private:
 	static void PrepareLumenHardwareRayTracingRadianceCacheLumenMaterial(const FViewInfo& View, TArray<FRHIRayTracingShader*>& OutRayGenShaders);
 	static void PrepareLumenHardwareRayTracingRadiosityLumenMaterial(const FViewInfo& View, TArray<FRHIRayTracingShader*>& OutRayGenShaders);
 	static void PrepareLumenHardwareRayTracingDirectLightingLumenMaterial(const FViewInfo& View, TArray<FRHIRayTracingShader*>& OutRayGenShaders);
-
-	const FRHITransition* RayTracingDynamicGeometryUpdateEndTransition = nullptr; // Signaled when all AS for this frame are built
 #endif // RHI_RAYTRACING
 
 	/** Set to true if lights were injected into the light grid (this controlled by somewhat complex logic, this flag is used to cross-check). */
 	bool bAreLightsInLightGrid;
-
-	FDynamicShadowsTaskData* CurrentDynamicShadowsTaskData{};
-
-#if RHI_RAYTRACING
-	FRayTracingRelevantPrimitiveTaskData* RayTracingRelevantPrimitiveTaskData = nullptr;
-#endif
 };
 
 DECLARE_CYCLE_STAT_EXTERN(TEXT("PrePass"), STAT_CLM_PrePass, STATGROUP_CommandListMarkers, );
\ No newline at end of file
diff --git a/Engine/Source/Runtime/Renderer/Private/MobileShadingRenderer.cpp b/Engine/Source/Runtime/Renderer/Private/MobileShadingRenderer.cpp
--- a/Engine/Source/Runtime/Renderer/Private/MobileShadingRenderer.cpp
+++ b/Engine/Source/Runtime/Renderer/Private/MobileShadingRenderer.cpp
@@ -452,7 +452,8 @@ void FMobileSceneRenderer::InitViews(FRDGBuilder& GraphBuilder, FSceneTexturesCo
 		Views[ViewIndex].GPUSceneViewId = InstanceCullingManager.RegisterView(Views[ViewIndex]);
 	}
 
-	FILCUpdatePrimTaskData ILCTaskData;
+	FILCUpdatePrimTaskData* ILCTaskData = nullptr;
+	FComputeViewVisibilityCallbacks ComputeViewVisibilityCallbacks;
 	FViewVisibleCommandsPerView ViewCommandsPerView;
 	ViewCommandsPerView.SetNum(Views.Num());
 
@@ -466,7 +467,7 @@ void FMobileSceneRenderer::InitViews(FRDGBuilder& GraphBuilder, FSceneTexturesCo
 		Views[0].InitRHIResources();
 		FXSystem->PostInitViews(GraphBuilder, Views, !ViewFamily.EngineShowFlags.HitProxies);
 	}
-	ComputeViewVisibility(RHICmdList, BasePassDepthStencilAccess, ViewCommandsPerView, DynamicIndexBuffer, DynamicVertexBuffer, DynamicReadBuffer, InstanceCullingManager, VirtualTextureUpdater);
+	ComputeViewVisibility(RHICmdList, BasePassDepthStencilAccess, ViewCommandsPerView, DynamicIndexBuffer, DynamicVertexBuffer, DynamicReadBuffer, InstanceCullingManager, VirtualTextureUpdater, ComputeViewVisibilityCallbacks);
 	PostVisibilityFrameSetup(ILCTaskData);
 
 	FIntPoint RenderTargetSize = ViewFamily.RenderTarget->GetSizeXY();
@@ -670,9 +671,9 @@ void FMobileSceneRenderer::InitViews(FRDGBuilder& GraphBuilder, FSceneTexturesCo
 	SetupMobileBasePassAfterShadowInit(BasePassDepthStencilAccess, ViewCommandsPerView, InstanceCullingManager);
 
 	// if we kicked off ILC update via task, wait and finalize.
-	if (ILCTaskData.TaskRef.IsValid())
+	if (ILCTaskData)
 	{
-		Scene->IndirectLightingCache.FinalizeCacheUpdates(Scene, *this, ILCTaskData);
+		Scene->IndirectLightingCache.FinalizeCacheUpdates(Scene, *this, *ILCTaskData);
 	}
 
 	if (bRequiresDistanceField)
diff --git a/Engine/Source/Runtime/Renderer/Private/SceneHitProxyRendering.cpp b/Engine/Source/Runtime/Renderer/Private/SceneHitProxyRendering.cpp
--- a/Engine/Source/Runtime/Renderer/Private/SceneHitProxyRendering.cpp
+++ b/Engine/Source/Runtime/Renderer/Private/SceneHitProxyRendering.cpp
@@ -631,9 +631,9 @@ void FDeferredShadingSceneRenderer::RenderHitProxies(FRDGBuilder& GraphBuilder)
 
 	// Find the visible primitives.
 	FLumenSceneFrameTemporaries LumenFrameTemporaries;
-	FILCUpdatePrimTaskData ILCTaskData;
+	FInitViewTaskDatas InitViewTaskDatas;
 	FRDGExternalAccessQueue ExternalAccessQueue;
-	BeginInitViews(GraphBuilder, SceneTexturesConfig, FExclusiveDepthStencil::DepthWrite_StencilWrite, ILCTaskData, InstanceCullingManager, nullptr);
+	BeginInitViews(GraphBuilder, SceneTexturesConfig, FExclusiveDepthStencil::DepthWrite_StencilWrite, InstanceCullingManager, nullptr, InitViewTaskDatas);
 
 	extern TSet<IPersistentViewUniformBufferExtension*> PersistentViewUniformBufferExtensions;
 
@@ -657,7 +657,7 @@ void FDeferredShadingSceneRenderer::RenderHitProxies(FRDGBuilder& GraphBuilder)
 		Scene->GPUScene.UploadDynamicPrimitiveShaderDataForView(GraphBuilder, *Scene, Views[ViewIndex], ExternalAccessQueue);
 	}
 
-	EndInitViews(GraphBuilder, LumenFrameTemporaries, ILCTaskData, InstanceCullingManager, ExternalAccessQueue);
+	EndInitViews(GraphBuilder, LumenFrameTemporaries, InstanceCullingManager, ExternalAccessQueue, InitViewTaskDatas);
 
 	ExternalAccessQueue.Submit(GraphBuilder);
 
diff --git a/Engine/Source/Runtime/Renderer/Private/SceneRendering.h b/Engine/Source/Runtime/Renderer/Private/SceneRendering.h
--- a/Engine/Source/Runtime/Renderer/Private/SceneRendering.h
+++ b/Engine/Source/Runtime/Renderer/Private/SceneRendering.h
@@ -50,7 +50,6 @@ class FSceneViewState;
 class FViewInfo;
 struct FILCUpdatePrimTaskData;
 class FPostprocessContext;
-struct FILCUpdatePrimTaskData;
 class FRaytracingLightDataPacked;
 class FRayTracingDecals;
 class FRayTracingLocalShaderBindingWriter;
@@ -2401,19 +2400,29 @@ protected:
 	/** Performs once per frame setup prior to visibility determination. */
 	void PreVisibilityFrameSetup(FRDGBuilder& GraphBuilder, const FSceneTexturesConfig& SceneTexturesConfig);
 
+	struct FComputeViewVisibilityCallbacks
+	{
+		TFunction<void()> PreGatherDynamicMeshElements;
+	};
+
 	/** Computes which primitives are visible and relevant for each view. */
-	void ComputeViewVisibility(FRHICommandListImmediate& RHICmdList, FExclusiveDepthStencil::Type BasePassDepthStencilAccess, FViewVisibleCommandsPerView& ViewCommandsPerView, 
-		FGlobalDynamicIndexBuffer& DynamicIndexBuffer, FGlobalDynamicVertexBuffer& DynamicVertexBuffer, FGlobalDynamicReadBuffer& DynamicReadBuffer, FInstanceCullingManager& InstanceCullingManager,
-		FVirtualTextureUpdater* VirtualTextureUpdater);
+	void ComputeViewVisibility(
+		FRHICommandListImmediate& RHICmdList,
+		FExclusiveDepthStencil::Type BasePassDepthStencilAccess,
+		FViewVisibleCommandsPerView& ViewCommandsPerView, 
+		FGlobalDynamicIndexBuffer& DynamicIndexBuffer,
+		FGlobalDynamicVertexBuffer& DynamicVertexBuffer,
+		FGlobalDynamicReadBuffer& DynamicReadBuffer,
+		FInstanceCullingManager& InstanceCullingManager,
+		FVirtualTextureUpdater* VirtualTextureUpdater,
+		const FComputeViewVisibilityCallbacks& Callbacks);
 
 	virtual void ComputeLightVisibility();
 
 	void GatherReflectionCaptureLightMeshElements();
 
 	/** Performs once per frame setup after to visibility determination. */
-	void PostVisibilityFrameSetup(FILCUpdatePrimTaskData& OutILCTaskData);
-
-	virtual void PreGatherDynamicMeshElements() {}
+	void PostVisibilityFrameSetup(FILCUpdatePrimTaskData*& OutILCTaskData);
 
 	void GatherDynamicMeshElements(
 		TArray<FViewInfo>& InViews, 
diff --git a/Engine/Source/Runtime/Renderer/Private/SceneVisibility.cpp b/Engine/Source/Runtime/Renderer/Private/SceneVisibility.cpp
--- a/Engine/Source/Runtime/Renderer/Private/SceneVisibility.cpp
+++ b/Engine/Source/Runtime/Renderer/Private/SceneVisibility.cpp
@@ -4608,7 +4608,8 @@ void FSceneRenderer::ComputeViewVisibility(
 	FGlobalDynamicVertexBuffer& DynamicVertexBuffer,
 	FGlobalDynamicReadBuffer& DynamicReadBuffer, 
 	FInstanceCullingManager& InstanceCullingManager,
-	FVirtualTextureUpdater* VirtualTextureUpdater)
+	FVirtualTextureUpdater* VirtualTextureUpdater,
+	const FComputeViewVisibilityCallbacks& Callbacks)
 {
 	SCOPE_CYCLE_COUNTER(STAT_ViewVisibilityTime);
 	SCOPED_NAMED_EVENT(FSceneRenderer_ComputeViewVisibility, FColor::Magenta);
@@ -4928,7 +4929,10 @@ void FSceneRenderer::ComputeViewVisibility(
 
 	ComputeLightVisibilityTask.Wait();
 
-	PreGatherDynamicMeshElements();
+	if (Callbacks.PreGatherDynamicMeshElements)
+	{
+		Callbacks.PreGatherDynamicMeshElements();
+	}
 
 	// Sync the virtual texture update task before gathering dynamic mesh elements. Render proxies can register new
 	// materials which require evaluating uniform expression caches, which can contain virtual textures. Newly allocated
@@ -5182,7 +5186,7 @@ void FSceneRenderer::GatherReflectionCaptureLightMeshElements()
 	}
 }
 
-void FSceneRenderer::PostVisibilityFrameSetup(FILCUpdatePrimTaskData& OutILCTaskData)
+void FSceneRenderer::PostVisibilityFrameSetup(FILCUpdatePrimTaskData*& OutILCTaskData)
 {
 	QUICK_SCOPE_CYCLE_COUNTER(STAT_PostVisibilityFrameSetup);
 
@@ -5206,7 +5210,9 @@ void FSceneRenderer::PostVisibilityFrameSetup(FILCUpdatePrimTaskData& OutILCTask
 	if (ViewFamily.EngineShowFlags.HitProxies == 0 && Scene->PrecomputedLightVolumes.Num() > 0
 		&& GILCUpdatePrimTaskEnabled && FPlatformProcess::SupportsMultithreading())
 	{
-		Scene->IndirectLightingCache.StartUpdateCachePrimitivesTask(Scene, *this, true, OutILCTaskData);
+		OutILCTaskData = Allocator.Create<FILCUpdatePrimTaskData>();
+		Scene->IndirectLightingCache.StartUpdateCachePrimitivesTask(Scene, *this, true, *OutILCTaskData);
+		check(OutILCTaskData->TaskRef.IsValid());
 	}
 }
 
@@ -5232,7 +5238,13 @@ void FDeferredShadingSceneRenderer::PreVisibilityFrameSetup(FRDGBuilder& GraphBu
  * Initialize scene's views.
  * Check visibility, build visible mesh commands, etc.
  */
-void FDeferredShadingSceneRenderer::BeginInitViews(FRDGBuilder& GraphBuilder, const FSceneTexturesConfig& SceneTexturesConfig, FExclusiveDepthStencil::Type BasePassDepthStencilAccess, struct FILCUpdatePrimTaskData& ILCTaskData, FInstanceCullingManager& InstanceCullingManager, FVirtualTextureUpdater* VirtualTextureUpdater)
+void FDeferredShadingSceneRenderer::BeginInitViews(
+	FRDGBuilder& GraphBuilder,
+	const FSceneTexturesConfig& SceneTexturesConfig,
+	FExclusiveDepthStencil::Type BasePassDepthStencilAccess,
+	FInstanceCullingManager& InstanceCullingManager,
+	FVirtualTextureUpdater* VirtualTextureUpdater,
+	FInitViewTaskDatas& TaskDatas)
 {
 	SCOPED_NAMED_EVENT(FDeferredShadingSceneRenderer_InitViews, FColor::Emerald);
 	SCOPE_CYCLE_COUNTER(STAT_InitViewsTime);
@@ -5268,7 +5280,16 @@ void FDeferredShadingSceneRenderer::BeginInitViews(FRDGBuilder& GraphBuilder, co
 	FViewVisibleCommandsPerView ViewCommandsPerView;
 	ViewCommandsPerView.SetNum(Views.Num());
 
-	ComputeViewVisibility(RHICmdList, BasePassDepthStencilAccess, ViewCommandsPerView, DynamicIndexBufferForInitViews, DynamicVertexBufferForInitViews, DynamicReadBufferForInitViews, InstanceCullingManager, VirtualTextureUpdater);
+	{
+		FComputeViewVisibilityCallbacks ComputeViewVisibilityCallbacks;
+
+		ComputeViewVisibilityCallbacks.PreGatherDynamicMeshElements = [this, &TaskDatas]
+		{
+			PreGatherDynamicMeshElements(TaskDatas);
+		};
+
+		ComputeViewVisibility(RHICmdList, BasePassDepthStencilAccess, ViewCommandsPerView, DynamicIndexBufferForInitViews, DynamicVertexBufferForInitViews, DynamicReadBufferForInitViews, InstanceCullingManager, VirtualTextureUpdater, ComputeViewVisibilityCallbacks);
+	}
 
 	// This must happen before we start initialising and using views.
 	if (Scene)
@@ -5284,7 +5305,7 @@ void FDeferredShadingSceneRenderer::BeginInitViews(FRDGBuilder& GraphBuilder, co
 
 	Scene->WaitForCreateLightPrimitiveInteractionsTask();
 
-	PostVisibilityFrameSetup(ILCTaskData);
+	PostVisibilityFrameSetup(TaskDatas.ILCUpdatePrim);
 
 	{
 		QUICK_SCOPE_CYCLE_COUNTER(STAT_InitViews_InitRHIResources);
@@ -5314,7 +5335,7 @@ void FDeferredShadingSceneRenderer::BeginInitViews(FRDGBuilder& GraphBuilder, co
 		RHICmdList.ImmediateFlush(EImmediateFlushType::DispatchToRHIThread);
 	}
 
-	WaitForPrepareDynamicShadowsTask(CurrentDynamicShadowsTaskData);
+	WaitForPrepareDynamicShadowsTask(TaskDatas.DynamicShadows);
 }
 
 template<class T>
@@ -5388,7 +5409,12 @@ void FSceneRenderer::SetupSceneReflectionCaptureBuffer(FRHICommandListImmediate&
 	}
 }
 
-void FDeferredShadingSceneRenderer::EndInitViews(FRDGBuilder& GraphBuilder, FLumenSceneFrameTemporaries& FrameTemporaries, struct FILCUpdatePrimTaskData& ILCTaskData, FInstanceCullingManager& InstanceCullingManager, FRDGExternalAccessQueue& ExternalAccessQueue)
+void FDeferredShadingSceneRenderer::EndInitViews(
+	FRDGBuilder& GraphBuilder,
+	FLumenSceneFrameTemporaries& FrameTemporaries,
+	FInstanceCullingManager& InstanceCullingManager,
+	FRDGExternalAccessQueue& ExternalAccessQueue,
+	const FInitViewTaskDatas& TaskDatas)
 {
 	SCOPED_NAMED_EVENT(FDeferredShadingSceneRenderer_InitViewsAfterPrepass, FColor::Emerald);
 	SCOPE_CYCLE_COUNTER(STAT_InitViewsPossiblyAfterPrepass);
@@ -5402,10 +5428,9 @@ void FDeferredShadingSceneRenderer::EndInitViews(FRDGBuilder& GraphBuilder, FLum
 		&& !bHasRayTracedOverlay)
 	{
 		// Setup dynamic shadows.
-		if (CurrentDynamicShadowsTaskData)
+		if (TaskDatas.DynamicShadows)
 		{
-			FinishInitDynamicShadows(GraphBuilder, CurrentDynamicShadowsTaskData, DynamicIndexBufferForInitShadows, DynamicVertexBufferForInitShadows, DynamicReadBufferForInitShadows, InstanceCullingManager, ExternalAccessQueue);
-			CurrentDynamicShadowsTaskData = nullptr;
+			FinishInitDynamicShadows(GraphBuilder, TaskDatas.DynamicShadows, DynamicIndexBufferForInitShadows, DynamicVertexBufferForInitShadows, DynamicReadBufferForInitShadows, InstanceCullingManager, ExternalAccessQueue);
 		}
 		else
 		{
@@ -5424,14 +5449,14 @@ void FDeferredShadingSceneRenderer::EndInitViews(FRDGBuilder& GraphBuilder, FLum
 		&& !(GILCUpdatePrimTaskEnabled && FPlatformProcess::SupportsMultithreading()))
 	{
 		QUICK_SCOPE_CYCLE_COUNTER(STAT_PostVisibilityFrameSetup_IndirectLightingCache_Update);
-		check(!ILCTaskData.TaskRef.IsValid());
+		check(!TaskDatas.ILCUpdatePrim);
 		Scene->IndirectLightingCache.UpdateCache(Scene, *this, true);
 	}
 
 	// If we kicked off ILC update via task, wait and finalize.
-	if (ILCTaskData.TaskRef.IsValid())
+	if (TaskDatas.ILCUpdatePrim)
 	{
-		Scene->IndirectLightingCache.FinalizeCacheUpdates(Scene, *this, ILCTaskData);
+		Scene->IndirectLightingCache.FinalizeCacheUpdates(Scene, *this, *TaskDatas.ILCUpdatePrim);
 	}
 
 	{
