Removed GPU scene buffers from external access queue.

diff --git a/Engine/Source/Runtime/Renderer/Private/GPUScene.cpp b/Engine/Source/Runtime/Renderer/Private/GPUScene.cpp
--- a/Engine/Source/Runtime/Renderer/Private/GPUScene.cpp
+++ b/Engine/Source/Runtime/Renderer/Private/GPUScene.cpp
@@ -581,7 +581,6 @@ void FGPUScene::BeginRender(const FScene* Scene, FGPUSceneDynamicContext &GPUSce
 	CurrentDynamicContext = &GPUSceneDynamicContext;
 	DynamicPrimitivesOffset = NumScenePrimitives;
 	bInBeginEndBlock = true;
-	bInExternalAccessMode = false;
 }
 
 void FGPUScene::EndRender()
@@ -735,8 +734,6 @@ void FGPUScene::UpdateInternal(FRDGBuilder& GraphBuilder, FSceneUniformBuffer& S
 
 		UpdateGPULights(GraphBuilder, Scene);
 	}
-
-	UseExternalAccessMode(ExternalAccessQueue, ERHIAccess::SRVMask, ERHIPipeline::All);
 }
 
 template<typename FUploadDataSourceAdapter>
@@ -880,48 +877,6 @@ struct FInstanceBatcher
 	}
 };
 
-void FGPUScene::UseInternalAccessMode(FRDGBuilder& GraphBuilder)
-{
-	if (!bInExternalAccessMode)
-	{
-		return;
-	}
-
-	GraphBuilder.UseInternalAccessMode({
-		BufferState.InstanceSceneDataBuffer,
-		BufferState.InstancePayloadDataBuffer,
-		BufferState.PrimitiveBuffer,
-		BufferState.LightmapDataBuffer
-	});
-
-	if (BufferState.InstanceBVHBuffer)
-	{
-		GraphBuilder.UseInternalAccessMode(BufferState.InstanceBVHBuffer);
-	}
-
-	bInExternalAccessMode = false;
-}
-
-void FGPUScene::UseExternalAccessMode(FRDGExternalAccessQueue& ExternalAccessQueue, ERHIAccess Access, ERHIPipeline Pipelines)
-{
-	if (bInExternalAccessMode)
-	{
-		return;
-	}
-
-	ExternalAccessQueue.AddUnique(BufferState.InstanceSceneDataBuffer, Access, Pipelines);
-	ExternalAccessQueue.AddUnique(BufferState.InstancePayloadDataBuffer, Access, Pipelines);
-	ExternalAccessQueue.AddUnique(BufferState.PrimitiveBuffer, Access, Pipelines);
-	ExternalAccessQueue.AddUnique(BufferState.LightmapDataBuffer, Access, Pipelines);
-
-	if (BufferState.InstanceBVHBuffer)
-	{
-		ExternalAccessQueue.AddUnique(BufferState.InstanceBVHBuffer, Access, Pipelines);
-	}
-
-	bInExternalAccessMode = true;
-}
-
 template<typename FUploadDataSourceAdapter>
 void FGPUScene::UploadGeneral(FRDGBuilder& GraphBuilder, FScene& Scene, FRDGExternalAccessQueue& ExternalAccessQueue, const FUploadDataSourceAdapter& UploadDataSourceAdapter)
 {
@@ -1560,8 +1515,6 @@ void FGPUScene::UploadDynamicPrimitiveShaderDataForViewInternal(FRDGBuilder& Gra
 
 		UpdateBufferState(GraphBuilder, View.GetSceneUniforms(), Scene, UploadAdapter, false);
 
-		UseInternalAccessMode(GraphBuilder);
-
 		// Run a pass that clears (Sets ID to invalid) any instances that need it.
 		AddClearInstancesPass(GraphBuilder);
 
@@ -1623,8 +1576,6 @@ void FGPUScene::UploadDynamicPrimitiveShaderDataForViewInternal(FRDGBuilder& Gra
 				PrimData.SourceData.DataWriterGPU.Execute(GraphBuilder, Params);
 			}
 		}
-
-		UseExternalAccessMode(ExternalAccessQueue, ERHIAccess::SRVMask, ERHIPipeline::All);
 	}
 }
 
@@ -1963,8 +1914,6 @@ bool FGPUScene::ExecuteDeferredGPUWritePass(FRDGBuilder& GraphBuilder, TArray<FV
 
 	RDG_EVENT_SCOPE(GraphBuilder, "GPUScene.DeferredGPUWrites - Pass %u", uint32(GPUWritePass));
 
-	UseInternalAccessMode(GraphBuilder);
-
 	FGPUSceneWriteDelegateParams Params;
 	Params.GPUWritePass = GPUWritePass;
 	GetWriteParameters(GraphBuilder, Params.GPUWriteParams);
@@ -1981,10 +1930,6 @@ bool FGPUScene::ExecuteDeferredGPUWritePass(FRDGBuilder& GraphBuilder, TArray<FV
 		DeferredWrite.DataWriterGPU.Execute(GraphBuilder, Params);
 	}
 
-	FRDGExternalAccessQueue ExternalAccessQueue;
-	UseExternalAccessMode(ExternalAccessQueue, ERHIAccess::SRVMask, ERHIPipeline::All);
-	ExternalAccessQueue.Submit(GraphBuilder);
-
 	DeferredGPUWritePassDelegates[PassIndex].Reset();
 	return true;
 }
diff --git a/Engine/Source/Runtime/Renderer/Private/GPUScene.h b/Engine/Source/Runtime/Renderer/Private/GPUScene.h
--- a/Engine/Source/Runtime/Renderer/Private/GPUScene.h
+++ b/Engine/Source/Runtime/Renderer/Private/GPUScene.h
@@ -383,7 +383,6 @@ private:
 
 	bool bIsEnabled = false;
 	bool bInBeginEndBlock = false;
-	bool bInExternalAccessMode = false;
 	FGPUSceneDynamicContext* CurrentDynamicContext = nullptr;
 	int32 NumScenePrimitives = 0;
 
@@ -413,9 +412,6 @@ private:
 	void AddUpdatePrimitiveIdsPass(FRDGBuilder& GraphBuilder, FInstanceGPULoadBalancer& IdOnlyUpdateItems);
 
 	void AddClearInstancesPass(FRDGBuilder& GraphBuilder);
-
-	void UseInternalAccessMode(FRDGBuilder& GraphBuilder);
-	void UseExternalAccessMode(FRDGExternalAccessQueue& ExternalAccessQueue, ERHIAccess Access, ERHIPipeline Pipelines);
 };
 
 class FGPUSceneScopeBeginEndHelper
