Adding Heterogeneous Volumes to scene capture and planar reflection.

diff --git a/Engine/Source/Runtime/Renderer/Private/HeterogeneousVolumes/HeterogeneousVolumes.cpp b/Engine/Source/Runtime/Renderer/Private/HeterogeneousVolumes/HeterogeneousVolumes.cpp
--- a/Engine/Source/Runtime/Renderer/Private/HeterogeneousVolumes/HeterogeneousVolumes.cpp
+++ b/Engine/Source/Runtime/Renderer/Private/HeterogeneousVolumes/HeterogeneousVolumes.cpp
@@ -182,10 +182,7 @@ bool ShouldRenderHeterogeneousVolumesForView(
 {
 	return IsHeterogeneousVolumesEnabled()
 		&& View.Family
-		&& !View.bIsPlanarReflection
-		&& !View.bIsSceneCapture
-		&& !View.bIsReflectionCapture
-		&& View.State;
+		&& !View.bIsReflectionCapture;
 }
 
 bool DoesPlatformSupportHeterogeneousVolumes(EShaderPlatform Platform)
