Moved CommitFinalPipelineState to happen at the start of the frame.

diff --git a/Engine/Source/Runtime/Renderer/Private/SceneHitProxyRendering.cpp b/Engine/Source/Runtime/Renderer/Private/SceneHitProxyRendering.cpp
--- a/Engine/Source/Runtime/Renderer/Private/SceneHitProxyRendering.cpp
+++ b/Engine/Source/Runtime/Renderer/Private/SceneHitProxyRendering.cpp
@@ -607,6 +607,8 @@ void FDeferredShadingSceneRenderer::RenderHitProxies(FRDGBuilder& GraphBuilder)
 {
 	const bool bNaniteEnabled = UseNanite(ShaderPlatform);
 
+	CommitFinalPipelineState();
+
 	Scene->UpdateAllPrimitiveSceneInfos(GraphBuilder);
 
 	GPU_MESSAGE_SCOPE(GraphBuilder);
