Renderer interface enhancements.
- Fixing the return type of GetScenePrimitiveComponentIds to be an ConstArrayView instead of a full Array.
- Adding GetPrimitiveSceneInfo that takes a PrimitiveComponentId
- Moving OrWhite2DIfNull and friends to their own Public header file.

diff --git a/Engine/Source/Runtime/Engine/Private/SceneInterface.cpp b/Engine/Source/Runtime/Engine/Private/SceneInterface.cpp
--- a/Engine/Source/Runtime/Engine/Private/SceneInterface.cpp
+++ b/Engine/Source/Runtime/Engine/Private/SceneInterface.cpp
@@ -10,11 +10,6 @@ FSceneInterface::FSceneInterface(ERHIFeatureLevel::Type InFeatureLevel)
 {
 }
 
-TArray<FPrimitiveComponentId> FSceneInterface::GetScenePrimitiveComponentIds() const
-{ 
-	return TArray<FPrimitiveComponentId>(); 
-}
-
 void FSceneInterface::UpdateAllPrimitiveSceneInfos(FRHICommandListImmediate& RHICmdList)
 {
 	FRDGBuilder GraphBuilder(RHICmdList, FRDGEventName(TEXT("UpdateAllPrimitiveSceneInfos")));
diff --git a/Engine/Source/Runtime/Engine/Public/SceneInterface.h b/Engine/Source/Runtime/Engine/Public/SceneInterface.h
--- a/Engine/Source/Runtime/Engine/Public/SceneInterface.h
+++ b/Engine/Source/Runtime/Engine/Public/SceneInterface.h
@@ -150,8 +150,13 @@ public:
 	 * Updates distance field scene data (transforms, uv scale, self-shadow bias, etc.) but doesn't change allocation in the atlas
 	 */
 	virtual void UpdatePrimitiveDistanceFieldSceneData_GameThread(UPrimitiveComponent* Primitive) {}
+
+	/** Finds the  primitive with the associated index into the primitive array. */
+	virtual FPrimitiveSceneInfo* GetPrimitiveSceneInfo(int32 PrimitiveIndex) const = 0;
+
 	/** Finds the  primitive with the associated component id. */
-	virtual FPrimitiveSceneInfo* GetPrimitiveSceneInfo(int32 PrimitiveIndex) = 0;
+	virtual FPrimitiveSceneInfo* GetPrimitiveSceneInfo(FPrimitiveComponentId PrimitiveId) const = 0;
+
 	/** Get the primitive previous local to world (used for motion blur). Returns true if the matrix was set. */
 	virtual bool GetPreviousLocalToWorld(const FPrimitiveSceneInfo* PrimitiveSceneInfo, FMatrix& OutPreviousLocalToWorld) const { return false; }
 	/** 
@@ -619,7 +624,7 @@ public:
 	/**
 	 * Returns the FPrimitiveComponentId for all primitives in the scene
 	 */
-	ENGINE_API virtual TArray<FPrimitiveComponentId> GetScenePrimitiveComponentIds() const;
+	virtual TConstArrayView<FPrimitiveComponentId> GetScenePrimitiveComponentIds() const = 0;
 
 	virtual void StartFrame() {}
 	virtual uint32 GetFrameNumber() const { return 0; }
diff --git a/Engine/Source/Runtime/Engine/Public/SceneView.h b/Engine/Source/Runtime/Engine/Public/SceneView.h
--- a/Engine/Source/Runtime/Engine/Public/SceneView.h
+++ b/Engine/Source/Runtime/Engine/Public/SceneView.h
@@ -1167,6 +1167,9 @@ public:
 	// Whether to use FOV when computing mesh LOD.
 	bool bUseFieldOfViewForLOD;
 
+	// Whether this view should use an HMD hidden area mask where appropriate.
+	bool bHMDHiddenAreaMaskActive = false;
+
 	/** Actual field of view and that desired by the camera originally */
 	float FOV;
 	float DesiredFOV;
diff --git a/Engine/Source/Runtime/Renderer/Private/RendererScene.cpp b/Engine/Source/Runtime/Renderer/Private/RendererScene.cpp
--- a/Engine/Source/Runtime/Renderer/Private/RendererScene.cpp
+++ b/Engine/Source/Runtime/Renderer/Private/RendererScene.cpp
@@ -2491,13 +2491,18 @@ void FScene::UpdatePrimitiveDistanceFieldSceneData_GameThread(UPrimitiveComponen
 	}
 }
 
-FPrimitiveSceneInfo* FScene::GetPrimitiveSceneInfo(int32 PrimitiveIndex)
+FPrimitiveSceneInfo* FScene::GetPrimitiveSceneInfo(int32 PrimitiveIndex) const
 {
-	if(Primitives.IsValidIndex(PrimitiveIndex))
+	if (Primitives.IsValidIndex(PrimitiveIndex))
 	{
 		return Primitives[PrimitiveIndex];
 	}
-	return NULL;
+	return nullptr;
+}
+
+FPrimitiveSceneInfo* FScene::GetPrimitiveSceneInfo(FPrimitiveComponentId PrimitiveId) const
+{
+	return GetPrimitiveSceneInfo(PrimitiveComponentIds.Find(PrimitiveId));
 }
 
 void FScene::RemovePrimitiveSceneInfo_RenderThread(FPrimitiveSceneInfo* PrimitiveSceneInfo)
@@ -6452,7 +6457,8 @@ public:
 	virtual void BatchRemovePrimitives(TArrayView<UPrimitiveComponent*> InPrimitives) override {}
 	virtual void BatchReleasePrimitives(TArrayView<UPrimitiveComponent*> InPrimitives) override {}
 	virtual void UpdateAllPrimitiveSceneInfos(FRDGBuilder& GraphBuilder, EUpdateAllPrimitiveSceneInfosAsyncOps = EUpdateAllPrimitiveSceneInfosAsyncOps::None) override {}
-	virtual FPrimitiveSceneInfo* GetPrimitiveSceneInfo(int32 PrimiteIndex) override { return NULL; }
+	virtual FPrimitiveSceneInfo* GetPrimitiveSceneInfo(int32 PrimiteIndex) const final { return nullptr; }
+	virtual FPrimitiveSceneInfo* GetPrimitiveSceneInfo(FPrimitiveComponentId PrimitiveId) const final { return nullptr; }
 
 	/** Updates the transform of a primitive which has already been added to the scene. */
 	virtual void UpdatePrimitiveTransform(UPrimitiveComponent* Primitive) override {}
@@ -6573,6 +6579,11 @@ public:
 
 	virtual bool HasAnyLights() const override { return false; }
 
+	virtual TConstArrayView<FPrimitiveComponentId> GetScenePrimitiveComponentIds() const override
+	{
+		return TConstArrayView<FPrimitiveComponentId>();
+	}
+
 private:
 	UWorld* World;
 	class FFXSystemInterface* FXSystem;
diff --git a/Engine/Source/Runtime/Renderer/Private/ScenePrivate.h b/Engine/Source/Runtime/Renderer/Private/ScenePrivate.h
--- a/Engine/Source/Runtime/Renderer/Private/ScenePrivate.h
+++ b/Engine/Source/Runtime/Renderer/Private/ScenePrivate.h
@@ -3072,7 +3072,8 @@ public:
 	virtual void UpdatePrimitiveAttachment(UPrimitiveComponent* Primitive) override;
 	virtual void UpdateCustomPrimitiveData(UPrimitiveComponent* Primitive) override;
 	virtual void UpdatePrimitiveDistanceFieldSceneData_GameThread(UPrimitiveComponent* Primitive) override;
-	virtual FPrimitiveSceneInfo* GetPrimitiveSceneInfo(int32 PrimitiveIndex) override;
+	virtual FPrimitiveSceneInfo* GetPrimitiveSceneInfo(int32 PrimitiveIndex) const final;
+	virtual FPrimitiveSceneInfo* GetPrimitiveSceneInfo(FPrimitiveComponentId PrimitiveId) const final;
 	virtual bool GetPreviousLocalToWorld(const FPrimitiveSceneInfo* PrimitiveSceneInfo, FMatrix& OutPreviousLocalToWorld) const override;
 	virtual void AddLight(ULightComponent* Light) override;
 	virtual void RemoveLight(ULightComponent* Light) override;
@@ -3348,7 +3349,7 @@ public:
 
 	bool ShouldRenderSkylightInBasePass(bool bIsTranslucent) const;
 
-	virtual TArray<FPrimitiveComponentId> GetScenePrimitiveComponentIds() const override
+	virtual TConstArrayView<FPrimitiveComponentId> GetScenePrimitiveComponentIds() const override
 	{
 		return PrimitiveComponentIds;
 	}
diff --git a/Engine/Source/Runtime/Renderer/Private/SceneRendering.h b/Engine/Source/Runtime/Renderer/Private/SceneRendering.h
--- a/Engine/Source/Runtime/Renderer/Private/SceneRendering.h
+++ b/Engine/Source/Runtime/Renderer/Private/SceneRendering.h
@@ -39,6 +39,7 @@
 #include "TranslucencyPass.h"
 #include "SceneTexturesConfig.h"
 #include "SceneUniformBuffer.h"
+#include "TextureFallbacks.h"
 
 #if RHI_RAYTRACING
 #include "RayTracingInstanceBufferUtil.h"
@@ -1573,9 +1574,6 @@ public:
 
 	bool bIsSnapshot;
 
-	// Whether this view should use an HMD hidden area mask where appropriate.
-	bool bHMDHiddenAreaMaskActive = false;
-
 	// Whether this view should use compute passes where appropriate.
 	bool bUseComputePasses = false;
 
@@ -2707,74 +2705,6 @@ private:
 // The noise textures need to be set in Slate too.
 RENDERER_API void UpdateNoiseTextureParameters(FViewUniformShaderParameters& ViewUniformShaderParameters);
 
-inline FRHITexture* OrWhite2DIfNull(FRHITexture* Tex)
-{
-	FRHITexture* Result = Tex ? Tex : GWhiteTexture->TextureRHI.GetReference();
-	check(Result);
-	return Result;
-}
-
-inline FRHITexture* OrBlack2DIfNull(FRHITexture* Tex)
-{
-	FRHITexture* Result = Tex ? Tex : GBlackTexture->TextureRHI.GetReference();
-	check(Result);
-	return Result;
-}
-
-inline FRHITexture* OrBlack2DArrayIfNull(FRHITexture* Tex)
-{
-	FRHITexture* Result = Tex ? Tex : GBlackArrayTexture->TextureRHI.GetReference();
-	check(Result);
-	return Result;
-}
-
-inline FRHITexture* OrBlack3DIfNull(FRHITexture* Tex)
-{
-	// we fall back to 2D which are unbound mobile parameters
-	return OrBlack2DIfNull(Tex ? Tex : GBlackVolumeTexture->TextureRHI.GetReference());
-}
-
-inline FRHITexture* OrBlack3DAlpha1IfNull(FRHITexture* Tex)
-{
-	// we fall back to 2D which are unbound mobile parameters
-	return OrBlack2DIfNull(Tex ? Tex : GBlackAlpha1VolumeTexture->TextureRHI.GetReference());
-}
-
-inline FRHITexture* OrBlack3DUintIfNull(FRHITexture* Tex)
-{
-	// we fall back to 2D which are unbound mobile parameters
-	return OrBlack2DIfNull(Tex ? Tex : GBlackUintVolumeTexture->TextureRHI.GetReference());
-}
-
-inline void SetBlack2DIfNull(FRHITexture*& Tex)
-{
-	if (!Tex)
-	{
-		Tex = GBlackTexture->TextureRHI.GetReference();
-		check(Tex);
-	}
-}
-
-inline void SetBlack3DIfNull(FRHITexture*& Tex)
-{
-	if (!Tex)
-	{
-		Tex = GBlackVolumeTexture->TextureRHI.GetReference();
-		// we fall back to 2D which are unbound mobile parameters
-		SetBlack2DIfNull(Tex);
-	}
-}
-
-inline void SetBlackAlpha13DIfNull(FRHITexture*& Tex)
-{
-	if (!Tex)
-	{
-		Tex = GBlackAlpha1VolumeTexture->TextureRHI.GetReference();
-		// we fall back to 2D which are unbound mobile parameters
-		SetBlack2DIfNull(Tex); // This is actually a rgb=0, a=1 texture
-	}
-}
-
 struct FFastVramConfig
 {
 	FFastVramConfig();
diff --git a/Engine/Source/Runtime/Renderer/Public/TextureFallbacks.h b/Engine/Source/Runtime/Renderer/Public/TextureFallbacks.h
new file mode 100644
--- /dev/null
+++ b/Engine/Source/Runtime/Renderer/Public/TextureFallbacks.h
@@ -0,0 +1,66 @@
+// Copyright Epic Games, Inc. All Rights Reserved.
+
+#pragma once
+
+#include "GlobalRenderResources.h"
+
+inline FRHITexture* OrWhite2DIfNull(FRHITexture* Tex)
+{
+	FRHITexture* Result = Tex ? Tex : GWhiteTexture->TextureRHI.GetReference();
+	check(Result);
+	return Result;
+}
+
+inline FRHITexture* OrBlack2DIfNull(FRHITexture* Tex)
+{
+	FRHITexture* Result = Tex ? Tex : GBlackTexture->TextureRHI.GetReference();
+	check(Result);
+	return Result;
+}
+
+inline FRHITexture* OrBlack3DIfNull(FRHITexture* Tex)
+{
+	// we fall back to 2D which are unbound mobile parameters
+	return OrBlack2DIfNull(Tex ? Tex : GBlackVolumeTexture->TextureRHI.GetReference());
+}
+
+inline FRHITexture* OrBlack3DAlpha1IfNull(FRHITexture* Tex)
+{
+	// we fall back to 2D which are unbound mobile parameters
+	return OrBlack2DIfNull(Tex ? Tex : GBlackAlpha1VolumeTexture->TextureRHI.GetReference());
+}
+
+inline FRHITexture* OrBlack3DUintIfNull(FRHITexture* Tex)
+{
+	// we fall back to 2D which are unbound mobile parameters
+	return OrBlack2DIfNull(Tex ? Tex : GBlackUintVolumeTexture->TextureRHI.GetReference());
+}
+
+inline void SetBlack2DIfNull(FRHITexture*& Tex)
+{
+	if (!Tex)
+	{
+		Tex = GBlackTexture->TextureRHI.GetReference();
+		check(Tex);
+	}
+}
+
+inline void SetBlack3DIfNull(FRHITexture*& Tex)
+{
+	if (!Tex)
+	{
+		Tex = GBlackVolumeTexture->TextureRHI.GetReference();
+		// we fall back to 2D which are unbound mobile parameters
+		SetBlack2DIfNull(Tex);
+	}
+}
+
+inline void SetBlackAlpha13DIfNull(FRHITexture*& Tex)
+{
+	if (!Tex)
+	{
+		Tex = GBlackAlpha1VolumeTexture->TextureRHI.GetReference();
+		// we fall back to 2D which are unbound mobile parameters
+		SetBlack2DIfNull(Tex); // This is actually a rgb=0, a=1 texture
+	}
+}
