Restoring function lost after move.

diff --git a/Engine/Source/Runtime/Renderer/Public/TextureFallbacks.h b/Engine/Source/Runtime/Renderer/Public/TextureFallbacks.h
--- a/Engine/Source/Runtime/Renderer/Public/TextureFallbacks.h
+++ b/Engine/Source/Runtime/Renderer/Public/TextureFallbacks.h
@@ -18,6 +18,13 @@ inline FRHITexture* OrBlack2DIfNull(FRHITexture* Tex)
 	return Result;
 }
 
+inline FRHITexture* OrBlack2DArrayIfNull(FRHITexture* Tex)
+{
+	FRHITexture* Result = Tex ? Tex : GBlackArrayTexture->TextureRHI.GetReference();
+	check(Result);
+	return Result;
+}
+
 inline FRHITexture* OrBlack3DIfNull(FRHITexture* Tex)
 {
 	// we fall back to 2D which are unbound mobile parameters
