Enable a directional light with none 0 light channel to cast shadow on mobile deferred.

diff --git a/Engine/Source/Runtime/Renderer/Private/MobileDeferredShadingPass.cpp b/Engine/Source/Runtime/Renderer/Private/MobileDeferredShadingPass.cpp
--- a/Engine/Source/Runtime/Renderer/Private/MobileDeferredShadingPass.cpp
+++ b/Engine/Source/Runtime/Renderer/Private/MobileDeferredShadingPass.cpp
@@ -503,7 +503,7 @@ static void RenderDirectionalLight(FRHICommandListImmediate& RHICmdList, const F
 	// Skylights with static lighting already had their diffuse contribution baked into lightmaps
 	const bool bSkyLight = Scene.SkyLight && !Scene.SkyLight->bHasStaticLighting && View.Family->EngineShowFlags.SkyLighting;
 	const bool bDynamicSkyLight = bSkyLight && !Scene.SkyLight->bWantsStaticShadowing;
-	const bool bDynamicShadows = DirectionalLight.Proxy->CastsDynamicShadow() && (LightingChannel == 0u) && View.Family->EngineShowFlags.DynamicShadows;
+	const bool bDynamicShadows = DirectionalLight.Proxy->CastsDynamicShadow() && View.Family->EngineShowFlags.DynamicShadows;
 	const bool bPlanarReflection = Scene.GetForwardPassGlobalPlanarReflection() != nullptr;
 
 	// Do two passes, first masking DefautLit, second masking all other shading models
