Prevent VolumetricFog from incorrectly skipping primitives not marked with IsHeterogeneousVolume

Unify check in shared ShoulderRenderPrimitiveWithHeterogeneousVolumes() function.

diff --git a/Engine/Source/Runtime/Renderer/Private/HeterogeneousVolumes/HeterogeneousVolumes.cpp b/Engine/Source/Runtime/Renderer/Private/HeterogeneousVolumes/HeterogeneousVolumes.cpp
--- a/Engine/Source/Runtime/Renderer/Private/HeterogeneousVolumes/HeterogeneousVolumes.cpp
+++ b/Engine/Source/Runtime/Renderer/Private/HeterogeneousVolumes/HeterogeneousVolumes.cpp
@@ -207,6 +207,15 @@ bool DoesMaterialShaderSupportHeterogeneousVolumes(const FMaterial& Material)
 		&& Material.IsUsedWithNiagaraMeshParticles();
 }
 
+bool ShouldRenderPrimitiveWithHeterogeneousVolumes(
+	const FPrimitiveSceneProxy* PrimitiveSceneProxy,
+	const FMaterial& Material
+)
+{
+	check(PrimitiveSceneProxy);
+	return PrimitiveSceneProxy->IsHeterogeneousVolume() && DoesMaterialShaderSupportHeterogeneousVolumes(Material);
+}
+
 namespace HeterogeneousVolumes
 {
 	// CVars
@@ -365,7 +374,7 @@ void FDeferredShadingSceneRenderer::RenderHeterogeneousVolumes(
 				const FMaterialRenderProxy* MaterialRenderProxy = Mesh->MaterialRenderProxy;
 				const FMaterial& Material = MaterialRenderProxy->GetMaterialWithFallback(View.GetFeatureLevel(), MaterialRenderProxy);
 				const FPrimitiveSceneProxy* PrimitiveSceneProxy = View.VolumetricMeshBatches[MeshBatchIndex].Proxy;
-				if (!PrimitiveSceneProxy->IsHeterogeneousVolume() || !DoesMaterialShaderSupportHeterogeneousVolumes(Material))
+				if (!ShouldRenderPrimitiveWithHeterogeneousVolumes(PrimitiveSceneProxy, Material))
 				{
 					continue;
 				}
diff --git a/Engine/Source/Runtime/Renderer/Private/HeterogeneousVolumes/HeterogeneousVolumes.h b/Engine/Source/Runtime/Renderer/Private/HeterogeneousVolumes/HeterogeneousVolumes.h
--- a/Engine/Source/Runtime/Renderer/Private/HeterogeneousVolumes/HeterogeneousVolumes.h
+++ b/Engine/Source/Runtime/Renderer/Private/HeterogeneousVolumes/HeterogeneousVolumes.h
@@ -28,6 +28,10 @@ bool ShouldRenderHeterogeneousVolumes(const FScene* Scene);
 bool ShouldRenderHeterogeneousVolumesForView(const FSceneView& View);
 bool DoesPlatformSupportHeterogeneousVolumes(EShaderPlatform Platform);
 bool DoesMaterialShaderSupportHeterogeneousVolumes(const FMaterialShaderParameters& Parameters);
+bool ShouldRenderPrimitiveWithHeterogeneousVolumes(
+	const FPrimitiveSceneProxy* PrimitiveSceneProxy,
+	const FMaterial& Material
+);
 
 //
 // Internal API
diff --git a/Engine/Source/Runtime/Renderer/Private/VolumetricFogVoxelization.cpp b/Engine/Source/Runtime/Renderer/Private/VolumetricFogVoxelization.cpp
--- a/Engine/Source/Runtime/Renderer/Private/VolumetricFogVoxelization.cpp
+++ b/Engine/Source/Runtime/Renderer/Private/VolumetricFogVoxelization.cpp
@@ -677,12 +677,12 @@ void FDeferredShadingSceneRenderer::VoxelizeFogVolumePrimitives(
 						const FMaterial& Material = MaterialRenderProxy->GetMaterialWithFallback(View.GetFeatureLevel(), MaterialRenderProxy);
 
 						// Skip volumes flagged as rendered with HeterogenousVolumes
-						if (bShouldRenderHeterogeneousVolumes && Material.IsUsedWithNiagaraMeshParticles())
+						const FPrimitiveSceneProxy* PrimitiveSceneProxy = View.VolumetricMeshBatches[MeshBatchIndex].Proxy;
+						if (bShouldRenderHeterogeneousVolumes && ShouldRenderPrimitiveWithHeterogeneousVolumes(PrimitiveSceneProxy, Material))
 						{
 							continue;
 						}
 
-						const FPrimitiveSceneProxy* PrimitiveSceneProxy = View.VolumetricMeshBatches[MeshBatchIndex].Proxy;
 						const FPrimitiveSceneInfo* PrimitiveSceneInfo = PrimitiveSceneProxy->GetPrimitiveSceneInfo();
 						const FBoxSphereBounds Bounds = PrimitiveSceneProxy->GetBounds();
 
