Fix hair crash in D3D12.

AllocateDummyVirtualVoxelResources was making the wrong kind of view for PageIndexGlobalCounter.

#rnx

diff --git a/Engine/Source/Runtime/Renderer/Private/HairStrands/HairStrandsVoxelization.cpp b/Engine/Source/Runtime/Renderer/Private/HairStrands/HairStrandsVoxelization.cpp
--- a/Engine/Source/Runtime/Renderer/Private/HairStrands/HairStrandsVoxelization.cpp
+++ b/Engine/Source/Runtime/Renderer/Private/HairStrands/HairStrandsVoxelization.cpp
@@ -1089,7 +1089,7 @@ static FHairStrandsVoxelResources AllocateDummyVirtualVoxelResources(
 	Out.Parameters.Common.PageIndexOccupancyBuffer	= GraphBuilder.CreateSRV(Out.PageIndexOccupancyBuffer, PF_R32G32_UINT);
 	Out.Parameters.Common.PageIndexCoordBuffer		= GraphBuilder.CreateSRV(Out.PageIndexCoordBuffer, PF_R8G8B8A8_UINT);
 	Out.Parameters.Common.NodeDescBuffer			= GraphBuilder.CreateSRV(Out.NodeDescBuffer); 
-	Out.Parameters.Common.AllocatedPageCountBuffer  = GraphBuilder.CreateSRV(Out.PageIndexGlobalCounter);
+	Out.Parameters.Common.AllocatedPageCountBuffer  = GraphBuilder.CreateSRV(Out.PageIndexGlobalCounter, PF_R32_UINT);
 	Out.Parameters.Common.CurrGPUMinVoxelSize		= GraphBuilder.CreateSRV(Dummy2BytesBuffer, PF_R16F);
 	Out.Parameters.Common.NextGPUMinVoxelSize		= Out.Parameters.Common.CurrGPUMinVoxelSize;
 	Out.Parameters.PageTexture						= Out.PageTexture;
