Add missing forward declaration for non-unity build

diff --git a/Engine/Source/Runtime/Renderer/Private/Renderer.cpp b/Engine/Source/Runtime/Renderer/Private/Renderer.cpp
--- a/Engine/Source/Runtime/Renderer/Private/Renderer.cpp
+++ b/Engine/Source/Runtime/Renderer/Private/Renderer.cpp
@@ -38,6 +38,7 @@
 #include "RendererOnScreenNotification.h"
 #include "Lumen/Lumen.h"
 #include "ScenePrivate.h"
+#include "SceneUniformBuffer.h"
 
 DEFINE_LOG_CATEGORY(LogRenderer);
 
diff --git a/Engine/Source/Runtime/Renderer/Private/RendererModule.h b/Engine/Source/Runtime/Renderer/Private/RendererModule.h
--- a/Engine/Source/Runtime/Renderer/Private/RendererModule.h
+++ b/Engine/Source/Runtime/Renderer/Private/RendererModule.h
@@ -19,6 +19,7 @@ class FSceneView;
 class FSceneViewFamily;
 class FSceneViewStateInterface;
 class FViewInfo;
+class FSceneUniformBuffer;
 struct FMeshBatch;
 struct FSynthBenchmarkResults;
 struct FSceneTextures;
