Abstracting away the inner types of DistanceField rendering for Niagara to use its data when rendering without relying on Private Renderer data.

diff --git a/Engine/Plugins/FX/Niagara/Source/Niagara/Classes/NiagaraDataInterface.h b/Engine/Plugins/FX/Niagara/Source/Niagara/Classes/NiagaraDataInterface.h
--- a/Engine/Plugins/FX/Niagara/Source/Niagara/Classes/NiagaraDataInterface.h
+++ b/Engine/Plugins/FX/Niagara/Source/Niagara/Classes/NiagaraDataInterface.h
@@ -382,10 +382,19 @@ struct NIAGARA_API FNiagaraDataInterfaceSetShaderParametersContext
 	FVector3f GetSystemLWCTile() const;
 	bool IsResourceBound(const void* ResourceAddress) const;
 	bool IsParameterBound(const void* ParameterAddress) const;
-	template<typename T> bool IsStructBound(const T* StructAddress) const { return IsStructBoundInternal(StructAddress, sizeof(T)); }
 	bool IsOutputStage() const;
 	bool IsIterationStage() const;
 
+	template<typename T> bool IsStructBound(const T* StructAddress) const
+	{
+		return IsStructBoundInternal(StructAddress, sizeof(T));
+	}
+
+	bool IsStructBound(const uint8* StructAddress, const FShaderParametersMetadata* StructMetadata) const
+	{
+		return IsStructBoundInternal(StructAddress, StructMetadata->GetSize());
+	}
+
 	template<typename T> T* GetParameterNestedStruct() const
 	{
 		const uint32 StructOffset = Align(ParametersOffset, TShaderParameterStructTypeInfo<T>::Alignment);
@@ -393,10 +402,14 @@ struct NIAGARA_API FNiagaraDataInterfaceSetShaderParametersContext
 
 		return reinterpret_cast<T*>(BaseParameters + StructOffset);
 	}
+	inline uint8* GetParameterIncludedStruct(const FShaderParametersMetadata* StructMetadata) const
+	{
+		const uint16 StructOffset = GetParameterIncludedStructInternal(StructMetadata);
+		return BaseParameters + StructOffset;
+	}
 	template<typename T> T* GetParameterIncludedStruct() const
 	{
-		const uint16 StructOffset = GetParameterIncludedStructInternal(TShaderParameterStructTypeInfo<T>::GetStructMetadata());
-		return reinterpret_cast<T*>(BaseParameters + StructOffset);
+		return reinterpret_cast<T*>(GetParameterIncludedStruct(TShaderParameterStructTypeInfo<T>::GetStructMetadata()));
 	}
 	template<typename T> TArrayView<T> GetParameterLooseArray(int32 NumElements) const
 	{
diff --git a/Engine/Plugins/FX/Niagara/Source/Niagara/Private/DataInterface/NiagaraDataInterfaceStaticMesh.cpp b/Engine/Plugins/FX/Niagara/Source/Niagara/Private/DataInterface/NiagaraDataInterfaceStaticMesh.cpp
--- a/Engine/Plugins/FX/Niagara/Source/Niagara/Private/DataInterface/NiagaraDataInterfaceStaticMesh.cpp
+++ b/Engine/Plugins/FX/Niagara/Source/Niagara/Private/DataInterface/NiagaraDataInterfaceStaticMesh.cpp
@@ -23,7 +23,6 @@
 #include "Subsystems/ImportSubsystem.h"
 #endif
 
-#include "DistanceFieldLightingShared.h"
 #include "Engine/StaticMeshActor.h"
 #include "Engine/StaticMeshSocket.h"
 #include "NiagaraDataInterfaceStaticMeshUvMapping.h"
@@ -33,6 +32,7 @@
 
 #include "NiagaraGpuComputeDispatchInterface.h"
 #include "NiagaraGpuComputeDispatch.h"
+#include "FXRenderingUtils.h"
 
 #include UE_INLINE_GENERATED_CPP_BY_NAME(NiagaraDataInterfaceStaticMesh)
 
@@ -2979,8 +2979,8 @@ bool UNiagaraDataInterfaceStaticMesh::GetFunctionHLSL(const FNiagaraDataInterfac
 void UNiagaraDataInterfaceStaticMesh::BuildShaderParameters(FNiagaraShaderParametersBuilder& ShaderParametersBuilder) const
 {
 	ShaderParametersBuilder.AddNestedStruct<NDIStaticMeshLocal::FShaderParameters>();
-	ShaderParametersBuilder.AddIncludedStruct<FDistanceFieldObjectBufferParameters>();
-	ShaderParametersBuilder.AddIncludedStruct<FDistanceFieldAtlasParameters>();
+	ShaderParametersBuilder.AddIncludedStruct(UE::FXRenderingUtils::DistanceFields::GetObjectBufferParametersMetadata());
+	ShaderParametersBuilder.AddIncludedStruct(UE::FXRenderingUtils::DistanceFields::GetAtlasParametersMetadata());
 }
 
 void UNiagaraDataInterfaceStaticMesh::SetShaderParameters(const FNiagaraDataInterfaceSetShaderParametersContext& Context) const
@@ -3053,17 +3053,15 @@ void UNiagaraDataInterfaceStaticMesh::SetShaderParameters(const FNiagaraDataInte
 	ShaderParameters->InstancePreviousRotation						= InstanceData.PrevRotation;
 	ShaderParameters->InstanceWorldVelocity							= DeltaPosition;
 
-	const FDistanceFieldSceneData* DistanceFieldSceneData = nullptr;
-	TConstArrayView<FViewInfo> SimulationViewInfos = Context.GetComputeDispatchInterface().GetSimulationViewInfos();
-	if (SimulationViewInfos.Num() > 0 && SimulationViewInfos[0].Family && SimulationViewInfos[0].Family->Scene && SimulationViewInfos[0].Family->Scene->GetRenderScene())
-	{
-		DistanceFieldSceneData = &SimulationViewInfos[0].Family->Scene->GetRenderScene()->DistanceFieldSceneData;
-	}
+	TConstStridedView<FSceneView> SimulationSceneViews = Context.GetComputeDispatchInterface().GetSimulationSceneViews();
+	const FSceneView* PrimaryView = SimulationSceneViews.Num() > 0 ? &SimulationSceneViews[0] : nullptr;
+
+	const bool bDistanceFieldParametersExist = PrimaryView ? UE::FXRenderingUtils::DistanceFields::HasDataToBind(*PrimaryView) : false;
 
 	if (Context.IsParameterBound(&ShaderParameters->InstanceDistanceFieldIndex))
 	{
 		int32 DistanceFieldIndex = INDEX_NONE;
-		if (DistanceFieldSceneData != nullptr && InstanceData.DistanceFieldPrimitiveId.IsValid())
+		if (bDistanceFieldParametersExist && InstanceData.DistanceFieldPrimitiveId.IsValid())
 		{
 			const FSceneInterface* Scene = Context.GetComputeDispatchInterface().GetScene();
 			if (const FPrimitiveSceneInfo* PrimitiveSceneInfo = Scene ? Scene->GetPrimitiveSceneInfo(InstanceData.DistanceFieldPrimitiveId) : nullptr)
@@ -3075,12 +3073,20 @@ void UNiagaraDataInterfaceStaticMesh::SetShaderParameters(const FNiagaraDataInte
 	}
 
 	// Bind Mesh Distance Field Data
-	FDistanceFieldObjectBufferParameters* ShaderDistanceFieldObjectParameters = Context.GetParameterIncludedStruct<FDistanceFieldObjectBufferParameters>();
-	FDistanceFieldAtlasParameters* ShaderDistanceFieldAtlasParameters = Context.GetParameterIncludedStruct<FDistanceFieldAtlasParameters>();
-	const bool bDistanceFieldDataBound = Context.IsStructBound<FDistanceFieldObjectBufferParameters>(ShaderDistanceFieldObjectParameters) || Context.IsStructBound<FDistanceFieldAtlasParameters>(ShaderDistanceFieldAtlasParameters);
+	const FShaderParametersMetadata* ObjectBufferParametersMetadata = UE::FXRenderingUtils::DistanceFields::GetObjectBufferParametersMetadata();
+	const FShaderParametersMetadata* AtlasParametersMetadata = UE::FXRenderingUtils::DistanceFields::GetAtlasParametersMetadata();
+
+	uint8* ShaderDistanceFieldObjectParameters = Context.GetParameterIncludedStruct(ObjectBufferParametersMetadata);
+	uint8* ShaderDistanceFieldAtlasParameters = Context.GetParameterIncludedStruct(AtlasParametersMetadata);
+
+	const bool bDistanceFieldDataBound =
+		Context.IsStructBound(ShaderDistanceFieldObjectParameters, ObjectBufferParametersMetadata) ||
+		Context.IsStructBound(ShaderDistanceFieldAtlasParameters, AtlasParametersMetadata);
+
 	if (bDistanceFieldDataBound)
 	{
-		FNiagaraDistanceFieldHelper::SetMeshDistanceFieldParameters(Context.GetGraphBuilder(), DistanceFieldSceneData, *ShaderDistanceFieldObjectParameters, *ShaderDistanceFieldAtlasParameters, FNiagaraRenderer::GetDummyFloat4Buffer());
+		UE::FXRenderingUtils::DistanceFields::SetupObjectBufferParameters(Context.GetGraphBuilder(), ShaderDistanceFieldObjectParameters, PrimaryView);
+		UE::FXRenderingUtils::DistanceFields::SetupAtlasParameters(Context.GetGraphBuilder(), ShaderDistanceFieldAtlasParameters, PrimaryView);
 	}
 
 	const FBox3f MeshBounds(InstanceData.PreSkinnedLocalBoundsCenter - InstanceData.PreSkinnedLocalBoundsExtents, InstanceData.PreSkinnedLocalBoundsCenter + InstanceData.PreSkinnedLocalBoundsExtents);
diff --git a/Engine/Plugins/FX/Niagara/Source/Niagara/Private/NiagaraDataInterfaceCollisionQuery.cpp b/Engine/Plugins/FX/Niagara/Source/Niagara/Private/NiagaraDataInterfaceCollisionQuery.cpp
--- a/Engine/Plugins/FX/Niagara/Source/Niagara/Private/NiagaraDataInterfaceCollisionQuery.cpp
+++ b/Engine/Plugins/FX/Niagara/Source/Niagara/Private/NiagaraDataInterfaceCollisionQuery.cpp
@@ -2,7 +2,6 @@
 
 #include "NiagaraDataInterfaceCollisionQuery.h"
 
-#include "DistanceFieldLightingShared.h"
 #include "GlobalDistanceFieldParameters.h"
 #include "NiagaraAsyncGpuTraceHelper.h"
 #include "NiagaraDistanceFieldHelper.h"
diff --git a/Engine/Plugins/FX/Niagara/Source/Niagara/Private/NiagaraDataInterfaceRigidMeshCollisionQuery.cpp b/Engine/Plugins/FX/Niagara/Source/Niagara/Private/NiagaraDataInterfaceRigidMeshCollisionQuery.cpp
--- a/Engine/Plugins/FX/Niagara/Source/Niagara/Private/NiagaraDataInterfaceRigidMeshCollisionQuery.cpp
+++ b/Engine/Plugins/FX/Niagara/Source/Niagara/Private/NiagaraDataInterfaceRigidMeshCollisionQuery.cpp
@@ -8,7 +8,7 @@
 #include "AnimationRuntime.h"
 #include "Components/SkeletalMeshComponent.h"
 #include "Components/StaticMeshComponent.h"
-#include "DistanceFieldLightingShared.h"
+#include "FXRenderingUtils.h"
 #include "Engine/Canvas.h"
 #include "Engine/Engine.h"
 #include "EngineUtils.h"
@@ -1643,8 +1643,8 @@ bool UNiagaraDataInterfaceRigidMeshCollisionQuery::AppendCompileHash(FNiagaraCom
 void UNiagaraDataInterfaceRigidMeshCollisionQuery::BuildShaderParameters(FNiagaraShaderParametersBuilder& ShaderParametersBuilder) const
 {
 	ShaderParametersBuilder.AddNestedStruct<NDIRigidMeshCollisionLocal::FShaderParameters>();
-	ShaderParametersBuilder.AddIncludedStruct<FDistanceFieldObjectBufferParameters>();
-	ShaderParametersBuilder.AddIncludedStruct<FDistanceFieldAtlasParameters>();
+	ShaderParametersBuilder.AddIncludedStruct(UE::FXRenderingUtils::DistanceFields::GetObjectBufferParametersMetadata());
+	ShaderParametersBuilder.AddIncludedStruct(UE::FXRenderingUtils::DistanceFields::GetAtlasParametersMetadata());
 }
 
 void UNiagaraDataInterfaceRigidMeshCollisionQuery::SetShaderParameters(const FNiagaraDataInterfaceSetShaderParametersContext& Context) const
@@ -1654,11 +1654,18 @@ void UNiagaraDataInterfaceRigidMeshCollisionQuery::SetShaderParameters(const FNi
 	const FNDIRigidMeshCollisionProxy::FRenderThreadData* ProxyData = InterfaceProxy.SystemInstancesToProxyData_RT.Find(Context.GetSystemInstanceID());
 
 	NDIRigidMeshCollisionLocal::FShaderParameters* ShaderParameters = Context.GetParameterNestedStruct<NDIRigidMeshCollisionLocal::FShaderParameters>();
-	FDistanceFieldObjectBufferParameters* ShaderDistanceFieldObjectParameters = Context.GetParameterIncludedStruct<FDistanceFieldObjectBufferParameters>();
-	FDistanceFieldAtlasParameters* ShaderDistanceFieldAtlasParameters = Context.GetParameterIncludedStruct<FDistanceFieldAtlasParameters>();
 
-	const bool bDistanceFieldDataBound = Context.IsStructBound<FDistanceFieldObjectBufferParameters>(ShaderDistanceFieldObjectParameters) || Context.IsStructBound<FDistanceFieldAtlasParameters>(ShaderDistanceFieldAtlasParameters);
-	const FDistanceFieldSceneData* DistanceFieldSceneData = nullptr;
+	const FShaderParametersMetadata* ObjectBufferParametersMetadata = UE::FXRenderingUtils::DistanceFields::GetObjectBufferParametersMetadata();
+	const FShaderParametersMetadata* AtlasParametersMetadata = UE::FXRenderingUtils::DistanceFields::GetAtlasParametersMetadata();
+
+	uint8* ShaderDistanceFieldObjectParameters = Context.GetParameterIncludedStruct(ObjectBufferParametersMetadata);
+	uint8* ShaderDistanceFieldAtlasParameters = Context.GetParameterIncludedStruct(AtlasParametersMetadata);
+
+	const bool bDistanceFieldDataBound =
+		Context.IsStructBound(ShaderDistanceFieldObjectParameters, ObjectBufferParametersMetadata) ||
+		Context.IsStructBound(ShaderDistanceFieldAtlasParameters, AtlasParametersMetadata);
+
+	bool bBindDistanceFieldData = false;
 
 	if (ProxyData != nullptr && ProxyData->AssetBuffer != nullptr && ProxyData->AssetBuffer->IsInitialized())
 	{
@@ -1680,14 +1687,7 @@ void UNiagaraDataInterfaceRigidMeshCollisionQuery::SetShaderParameters(const FNi
 		ShaderParameters->ElementOffsets.Z = ProxyData->ElementOffsets.CapsuleOffset;
 		ShaderParameters->ElementOffsets.W = ProxyData->ElementOffsets.NumElements;
 
-		if (bDistanceFieldDataBound)
-		{
-			TConstArrayView<FViewInfo> SimulationViewInfos = Context.GetComputeDispatchInterface().GetSimulationViewInfos();
-			if (SimulationViewInfos.Num() > 0 && SimulationViewInfos[0].Family && SimulationViewInfos[0].Family->Scene && SimulationViewInfos[0].Family->Scene->GetRenderScene())
-			{
-				DistanceFieldSceneData = &SimulationViewInfos[0].Family->Scene->GetRenderScene()->DistanceFieldSceneData;
-			}
-		}
+		bBindDistanceFieldData = true;
 	}
 	else
 	{
@@ -1706,7 +1706,11 @@ void UNiagaraDataInterfaceRigidMeshCollisionQuery::SetShaderParameters(const FNi
 
 	if (bDistanceFieldDataBound)
 	{
-		FNiagaraDistanceFieldHelper::SetMeshDistanceFieldParameters(Context.GetGraphBuilder(), DistanceFieldSceneData, *ShaderDistanceFieldObjectParameters, *ShaderDistanceFieldAtlasParameters, FNiagaraRenderer::GetDummyFloat4Buffer());
+		TConstStridedView<FSceneView> SimulationSceneViews = Context.GetComputeDispatchInterface().GetSimulationSceneViews();
+		const FSceneView* PrimaryView = bBindDistanceFieldData && SimulationSceneViews.Num() > 0 ? &SimulationSceneViews[0] : nullptr;
+
+		UE::FXRenderingUtils::DistanceFields::SetupObjectBufferParameters(Context.GetGraphBuilder(), ShaderDistanceFieldObjectParameters, PrimaryView);
+		UE::FXRenderingUtils::DistanceFields::SetupAtlasParameters(Context.GetGraphBuilder(), ShaderDistanceFieldAtlasParameters, PrimaryView);
 	}
 
 	ShaderParameters->SystemLWCTile = Context.GetSystemLWCTile();
diff --git a/Engine/Plugins/FX/Niagara/Source/Niagara/Private/NiagaraGpuComputeDispatch.h b/Engine/Plugins/FX/Niagara/Source/Niagara/Private/NiagaraGpuComputeDispatch.h
--- a/Engine/Plugins/FX/Niagara/Source/Niagara/Private/NiagaraGpuComputeDispatch.h
+++ b/Engine/Plugins/FX/Niagara/Source/Niagara/Private/NiagaraGpuComputeDispatch.h
@@ -28,7 +28,6 @@ the same VectorVM byte code / compute shader code
 #include "Tickable.h"
 
 class FGPUSortManager;
-class FDistanceFieldSceneData;
 class FNiagaraAsyncGpuTraceHelper;
 
 //////////////////////////////////////////////////////////////////////////
diff --git a/Engine/Plugins/FX/Niagara/Source/Niagara/Private/NiagaraRendererMeshes.cpp b/Engine/Plugins/FX/Niagara/Source/Niagara/Private/NiagaraRendererMeshes.cpp
--- a/Engine/Plugins/FX/Niagara/Source/Niagara/Private/NiagaraRendererMeshes.cpp
+++ b/Engine/Plugins/FX/Niagara/Source/Niagara/Private/NiagaraRendererMeshes.cpp
@@ -285,7 +285,7 @@ void FNiagaraRendererMeshes::PrepareParticleMeshRenderData(FParticleMeshRenderDa
 		// If even one material can cause the mesh to render before FFXSystemInterface::PostRenderOpaque, we cannot use low latency data
 		ParticleMeshRenderData.bIsGpuLowLatencyTranslucency =
 			ParticleMeshRenderData.bIsGpuLowLatencyTranslucency &&
-			!FFXRenderingUtils::CanMaterialRenderBeforeFXPostOpaque(ViewFamily, *SceneProxy, Material);
+			!UE::FXRenderingUtils::CanMaterialRenderBeforeFXPostOpaque(ViewFamily, *SceneProxy, Material);
 	}
 	
 	ParticleMeshRenderData.SourceParticleData = ParticleMeshRenderData.DynamicDataMesh->GetParticleDataToRender(ParticleMeshRenderData.bIsGpuLowLatencyTranslucency);
diff --git a/Engine/Plugins/FX/Niagara/Source/NiagaraShader/Private/NiagaraDistanceFieldHelper.cpp b/Engine/Plugins/FX/Niagara/Source/NiagaraShader/Private/NiagaraDistanceFieldHelper.cpp
--- a/Engine/Plugins/FX/Niagara/Source/NiagaraShader/Private/NiagaraDistanceFieldHelper.cpp
+++ b/Engine/Plugins/FX/Niagara/Source/NiagaraShader/Private/NiagaraDistanceFieldHelper.cpp
@@ -1,10 +1,8 @@
 // Copyright Epic Games, Inc. All Rights Reserved.
 
 #include "NiagaraDistanceFieldHelper.h"
-#include "DistanceFieldLightingShared.h"
-#include "SystemTextures.h"
-#include "ScenePrivate.h"
-#include "SceneRendering.h"
+#include "GlobalDistanceFieldParameters.h"
+#include "TextureFallbacks.h"
 
 // todo - currently duplicated from SetupGlobalDistanceFieldParameters (GlobalDistanceField.cpp) because of problems getting it properly exported from the dll
 void FNiagaraDistanceFieldHelper::SetGlobalDistanceFieldParameters(const FGlobalDistanceFieldParameterData* OptionalParameterData, FGlobalDistanceFieldParameters2& ShaderParameters)
@@ -66,71 +64,3 @@ void FNiagaraDistanceFieldHelper::SetGlobalDistanceFieldParameters(const FGlobal
 	ShaderParameters.DitheredTransparencyStepThreshold = 0.0f;
 	ShaderParameters.DitheredTransparencyTraceThreshold = 0.0f;
 }
-
-void FNiagaraDistanceFieldHelper::SetMeshDistanceFieldParameters(FRDGBuilder& GraphBuilder, const FDistanceFieldSceneData* OptionalDistanceFieldData, FDistanceFieldObjectBufferParameters& ObjectShaderParameters, FDistanceFieldAtlasParameters& AtlasShaderParameters, FRHIShaderResourceView* DummyFloat4Buffer)
-{
-	if (OptionalDistanceFieldData != nullptr && OptionalDistanceFieldData->NumObjectsInBuffer > 0)
-	{
-		ObjectShaderParameters = DistanceField::SetupObjectBufferParameters(GraphBuilder, *OptionalDistanceFieldData);
-		AtlasShaderParameters = DistanceField::SetupAtlasParameters(GraphBuilder, *OptionalDistanceFieldData);
-	}
-	else
-	{
-		class FDFDummyByteAddress : public FRenderResource
-		{
-		public:
-			TRefCountPtr<FRDGPooledBuffer> PooledBuffer;
-
-			virtual void InitRHI() override
-			{
-				FRDGBufferDesc BufferDesc = FRDGBufferDesc::CreateStructuredDesc(sizeof(uint32), 1);
-				BufferDesc.Usage = BUF_Static | BUF_ShaderResource | BUF_ByteAddressBuffer;
-
-				FRHIResourceCreateInfo CreateInfo(TEXT("FDFDummyByteAddress"));
-				FBufferRHIRef RHIBuffer = RHICreateStructuredBuffer(BufferDesc.BytesPerElement, BufferDesc.GetSize(), BufferDesc.Usage, CreateInfo);
-				{
-					void* Data = RHILockBuffer(RHIBuffer, 0, BufferDesc.GetSize(), RLM_WriteOnly);
-					FMemory::Memset(Data, 0, BufferDesc.GetSize());
-					RHIUnlockBuffer(RHIBuffer);
-				}
-
-				PooledBuffer = new FRDGPooledBuffer(RHIBuffer, BufferDesc, BufferDesc.NumElements, CreateInfo.DebugName);
-			}
-
-			virtual void ReleaseRHI() override
-			{
-				PooledBuffer.SafeRelease();
-			}
-		};
-
-		static TGlobalResource<FDFDummyByteAddress> GDFDummyByteAddress;
-
-		FRDGBufferSRVRef DefaultVector4 = GraphBuilder.CreateSRV(FRDGBufferSRVDesc(GSystemTextures.GetDefaultStructuredBuffer(GraphBuilder, sizeof(FVector4f))));
-		FRDGBufferSRVRef DefaultUInt32 = GraphBuilder.CreateSRV(FRDGBufferSRVDesc(GraphBuilder.RegisterExternalBuffer(GDFDummyByteAddress.PooledBuffer, ERDGBufferFlags::SkipTracking)));
-
-		ObjectShaderParameters.SceneObjectBounds = DefaultVector4;
-		ObjectShaderParameters.SceneObjectData = DefaultVector4;
-		ObjectShaderParameters.NumSceneObjects = 0;
-		ObjectShaderParameters.SceneHeightfieldObjectBounds = DefaultVector4;
-		ObjectShaderParameters.SceneHeightfieldObjectData = DefaultVector4;
-		ObjectShaderParameters.NumSceneHeightfieldObjects = 0;
-
-		const FRDGSystemTextures& SystemTextures = FRDGSystemTextures::Get(GraphBuilder);
-
-		AtlasShaderParameters.SceneDistanceFieldAssetData = DefaultVector4;
-		AtlasShaderParameters.DistanceFieldIndirectionTable = DefaultUInt32;
-		AtlasShaderParameters.DistanceFieldIndirection2Table = DefaultVector4;
-		AtlasShaderParameters.DistanceFieldIndirectionAtlas = SystemTextures.VolumetricBlack;
-		AtlasShaderParameters.DistanceFieldBrickTexture = SystemTextures.VolumetricBlack;
-		AtlasShaderParameters.DistanceFieldSampler = TStaticSamplerState<SF_Bilinear, AM_Wrap, AM_Wrap, AM_Wrap>::GetRHI();
-		AtlasShaderParameters.DistanceFieldBrickSize = FVector3f::ZeroVector;
-		AtlasShaderParameters.DistanceFieldUniqueDataBrickSize = FVector3f::ZeroVector;
-		AtlasShaderParameters.DistanceFieldBrickAtlasSizeInBricks = FIntVector::ZeroValue;
-		AtlasShaderParameters.DistanceFieldBrickAtlasMask = FIntVector::ZeroValue;
-		AtlasShaderParameters.DistanceFieldBrickAtlasSizeLog2 = FIntVector::ZeroValue;
-		AtlasShaderParameters.DistanceFieldBrickAtlasTexelSize = FVector3f::ZeroVector;
-		AtlasShaderParameters.DistanceFieldBrickAtlasHalfTexelSize = FVector3f::ZeroVector;
-		AtlasShaderParameters.DistanceFieldBrickOffsetToAtlasUVScale = FVector3f::ZeroVector;
-		AtlasShaderParameters.DistanceFieldUniqueDataBrickSizeInAtlasTexels = FVector3f::ZeroVector;
-	}
-}
diff --git a/Engine/Plugins/FX/Niagara/Source/NiagaraShader/Public/NiagaraDistanceFieldHelper.h b/Engine/Plugins/FX/Niagara/Source/NiagaraShader/Public/NiagaraDistanceFieldHelper.h
--- a/Engine/Plugins/FX/Niagara/Source/NiagaraShader/Public/NiagaraDistanceFieldHelper.h
+++ b/Engine/Plugins/FX/Niagara/Source/NiagaraShader/Public/NiagaraDistanceFieldHelper.h
@@ -2,19 +2,18 @@
 
 #pragma once
 
-#include "GlobalDistanceFieldParameters.h"
+#include "HAL/Platform.h"
 
 #if UE_ENABLE_INCLUDE_ORDER_DEPRECATED_IN_5_3
+#include "GlobalDistanceFieldParameters.h"
 #include "RenderGraphBuilder.h"
 #endif
 
-class FDistanceFieldSceneData;
-class FDistanceFieldObjectBufferParameters;
-class FDistanceFieldAtlasParameters;
 class FRDGBuilder;
+class FGlobalDistanceFieldParameterData;
+class FGlobalDistanceFieldParameters2;
 
 namespace FNiagaraDistanceFieldHelper
 {
 	NIAGARASHADER_API void SetGlobalDistanceFieldParameters(const FGlobalDistanceFieldParameterData* OptionalParameterData, FGlobalDistanceFieldParameters2& ShaderParameters);
-	NIAGARASHADER_API void SetMeshDistanceFieldParameters(FRDGBuilder& GraphBuilder, const FDistanceFieldSceneData* OptionalDistanceFieldData, FDistanceFieldObjectBufferParameters& ObjectShaderParameters, FDistanceFieldAtlasParameters& AtlasShaderParameters, FRHIShaderResourceView* DummyFloat4Buffer);
 }
diff --git a/Engine/Plugins/FX/Niagara/Source/NiagaraShader/Public/NiagaraShaderParametersBuilder.h b/Engine/Plugins/FX/Niagara/Source/NiagaraShader/Public/NiagaraShaderParametersBuilder.h
--- a/Engine/Plugins/FX/Niagara/Source/NiagaraShader/Public/NiagaraShaderParametersBuilder.h
+++ b/Engine/Plugins/FX/Niagara/Source/NiagaraShader/Public/NiagaraShaderParametersBuilder.h
@@ -41,9 +41,8 @@ public:
 
 	// Adds a shader parameters structure that is global in scope
 	// i.e. if the structure contained "MyFloat" the shader variable would be named "MyFloat"
-	template<typename T> void AddIncludedStruct()
+	inline void AddIncludedStruct(const FShaderParametersMetadata* StructMetadata)
 	{
-		const FShaderParametersMetadata* StructMetadata = TShaderParameterStructTypeInfo<T>::GetStructMetadata();
 		for (const FNiagaraDataInterfaceStructIncludeInfo& Existing : StructIncludeInfos)
 		{
 			if (Existing.StructMetadata == StructMetadata)
@@ -56,7 +55,14 @@ public:
 		NewInfo.StructMetadata = StructMetadata;
 		NewInfo.ParamterOffset = Align(MetadataBuilder.GetNextMemberOffset(), SHADER_PARAMETER_STRUCT_ALIGNMENT);
 
-		MetadataBuilder.AddIncludedStruct<T>();
+		MetadataBuilder.AddIncludedStruct(StructMetadata);
+	}
+
+	// Adds a shader parameters structure that is global in scope
+	// i.e. if the structure contained "MyFloat" the shader variable would be named "MyFloat"
+	template<typename T> void AddIncludedStruct()
+	{
+		AddIncludedStruct(TShaderParameterStructTypeInfo<T>::GetStructMetadata());
 	}
 
 	TConstArrayView<FNiagaraDataInterfaceGeneratedFunction> GetGeneratedFunctions() const { return GPUParamInfo.GeneratedFunctions; }
diff --git a/Engine/Source/Runtime/Renderer/Private/DistanceFieldLightingShared.h b/Engine/Source/Runtime/Renderer/Private/DistanceFieldLightingShared.h
--- a/Engine/Source/Runtime/Renderer/Private/DistanceFieldLightingShared.h
+++ b/Engine/Source/Runtime/Renderer/Private/DistanceFieldLightingShared.h
@@ -97,8 +97,8 @@ namespace DistanceField
 {
 	constexpr uint32 MinPrimitiveModifiedBoundsAllocation = 16 * 1024;
 
-	RENDERER_API FDistanceFieldObjectBufferParameters SetupObjectBufferParameters(FRDGBuilder& GraphBuilder, const FDistanceFieldSceneData& DistanceFieldSceneData);
-	RENDERER_API FDistanceFieldAtlasParameters SetupAtlasParameters(FRDGBuilder& GraphBuilder, const FDistanceFieldSceneData& DistanceFieldSceneData);
+	FDistanceFieldObjectBufferParameters SetupObjectBufferParameters(FRDGBuilder& GraphBuilder, const FDistanceFieldSceneData& DistanceFieldSceneData);
+	FDistanceFieldAtlasParameters SetupAtlasParameters(FRDGBuilder& GraphBuilder, const FDistanceFieldSceneData& DistanceFieldSceneData);
 };
 
 BEGIN_SHADER_PARAMETER_STRUCT(FDistanceFieldCulledObjectBufferParameters, )
diff --git a/Engine/Source/Runtime/Renderer/Private/FXRenderingUtils.cpp b/Engine/Source/Runtime/Renderer/Private/FXRenderingUtils.cpp
--- a/Engine/Source/Runtime/Renderer/Private/FXRenderingUtils.cpp
+++ b/Engine/Source/Runtime/Renderer/Private/FXRenderingUtils.cpp
@@ -1,12 +1,16 @@
 // Copyright Epic Games, Inc. All Rights Reserved.
 
 #include "FXRenderingUtils.h"
-#include "MaterialShared.h"
+
+#include "DistanceFieldLightingShared.h"
 #include "Lumen/LumenScreenProbeGather.h"
+#include "MaterialShared.h"
 #include "PrimitiveSceneProxy.h"
 #include "SceneInterface.h"
+#include "ScenePrivate.h"
+#include "SystemTextures.h"
 
-bool FFXRenderingUtils::CanMaterialRenderBeforeFXPostOpaque(
+bool UE::FXRenderingUtils::CanMaterialRenderBeforeFXPostOpaque(
 	const FSceneViewFamily& ViewFamily,
 	const FPrimitiveSceneProxy& SceneProxy,
 	const FMaterial& Material)
@@ -30,4 +34,115 @@ bool FFXRenderingUtils::CanMaterialRenderBeforeFXPostOpaque(
 	}
 		
 	return false;
-}
\ No newline at end of file
+}
+
+const FShaderParametersMetadata* UE::FXRenderingUtils::DistanceFields::GetObjectBufferParametersMetadata()
+{
+	return TShaderParameterStructTypeInfo<FDistanceFieldObjectBufferParameters>::GetStructMetadata();
+}
+
+const FShaderParametersMetadata* UE::FXRenderingUtils::DistanceFields::GetAtlasParametersMetadata()
+{
+	return TShaderParameterStructTypeInfo<FDistanceFieldAtlasParameters>::GetStructMetadata();
+}
+
+inline const FDistanceFieldSceneData* GetDistanceFieldSceneData(const FSceneView& View)
+{
+	FSceneInterface* SceneInterface = View.Family && View.Family->Scene ? View.Family->Scene : nullptr;
+	const FScene* Scene = SceneInterface->GetRenderScene();
+	return  Scene ? &Scene->DistanceFieldSceneData : nullptr;
+}
+
+bool UE::FXRenderingUtils::DistanceFields::HasDataToBind(const FSceneView& View)
+{
+	return GetDistanceFieldSceneData(View) != nullptr;
+}
+
+void UE::FXRenderingUtils::DistanceFields::SetupObjectBufferParameters(FRDGBuilder& GraphBuilder, uint8* DestinationData, const FSceneView* View)
+{
+	if (FDistanceFieldObjectBufferParameters* ObjectBufferParameters = reinterpret_cast<FDistanceFieldObjectBufferParameters*>(DestinationData))
+	{
+		const FDistanceFieldSceneData* DistanceFieldSceneData = View ? GetDistanceFieldSceneData(*View) : nullptr;
+
+		if (DistanceFieldSceneData && DistanceFieldSceneData->NumObjectsInBuffer > 0)
+		{
+			*ObjectBufferParameters = DistanceField::SetupObjectBufferParameters(GraphBuilder, *DistanceFieldSceneData);
+		}
+		else
+		{
+			FRDGBufferSRVRef DefaultVector4 = GraphBuilder.CreateSRV(FRDGBufferSRVDesc(GSystemTextures.GetDefaultStructuredBuffer(GraphBuilder, sizeof(FVector4f))));
+
+			ObjectBufferParameters->SceneObjectBounds = DefaultVector4;
+			ObjectBufferParameters->SceneObjectData = DefaultVector4;
+			ObjectBufferParameters->NumSceneObjects = 0;
+			ObjectBufferParameters->SceneHeightfieldObjectBounds = DefaultVector4;
+			ObjectBufferParameters->SceneHeightfieldObjectData = DefaultVector4;
+			ObjectBufferParameters->NumSceneHeightfieldObjects = 0;
+		}
+	}
+}
+
+class FDFDummyByteAddress : public FRenderResource
+{
+public:
+	TRefCountPtr<FRDGPooledBuffer> PooledBuffer;
+
+	virtual void InitRHI() override
+	{
+		FRDGBufferDesc BufferDesc = FRDGBufferDesc::CreateStructuredDesc(sizeof(uint32), 1);
+		BufferDesc.Usage = BUF_Static | BUF_ShaderResource | BUF_ByteAddressBuffer;
+
+		FRHIResourceCreateInfo CreateInfo(TEXT("FDFDummyByteAddress"));
+		FBufferRHIRef RHIBuffer = RHICreateStructuredBuffer(BufferDesc.BytesPerElement, BufferDesc.GetSize(), BufferDesc.Usage, CreateInfo);
+		{
+			void* Data = RHILockBuffer(RHIBuffer, 0, BufferDesc.GetSize(), RLM_WriteOnly);
+			FMemory::Memset(Data, 0, BufferDesc.GetSize());
+			RHIUnlockBuffer(RHIBuffer);
+		}
+
+		PooledBuffer = new FRDGPooledBuffer(RHIBuffer, BufferDesc, BufferDesc.NumElements, CreateInfo.DebugName);
+	}
+
+	virtual void ReleaseRHI() override
+	{
+		PooledBuffer.SafeRelease();
+	}
+};
+
+static TGlobalResource<FDFDummyByteAddress> GDFDummyByteAddress;
+
+void UE::FXRenderingUtils::DistanceFields::SetupAtlasParameters(FRDGBuilder& GraphBuilder, uint8* DestinationData, const FSceneView* View)
+{
+	if (FDistanceFieldAtlasParameters* AtlasParameters = reinterpret_cast<FDistanceFieldAtlasParameters*>(DestinationData))
+	{
+		const FDistanceFieldSceneData* DistanceFieldSceneData = View ? GetDistanceFieldSceneData(*View) : nullptr;
+
+		if (DistanceFieldSceneData && DistanceFieldSceneData->NumObjectsInBuffer > 0)
+		{
+			*AtlasParameters = DistanceField::SetupAtlasParameters(GraphBuilder, *DistanceFieldSceneData);
+		}
+		else
+		{
+			FRDGBufferSRVRef DefaultVector4 = GraphBuilder.CreateSRV(FRDGBufferSRVDesc(GSystemTextures.GetDefaultStructuredBuffer(GraphBuilder, sizeof(FVector4f))));
+			FRDGBufferSRVRef DefaultUInt32 = GraphBuilder.CreateSRV(FRDGBufferSRVDesc(GraphBuilder.RegisterExternalBuffer(GDFDummyByteAddress.PooledBuffer, ERDGBufferFlags::SkipTracking)));
+
+			const FRDGSystemTextures& SystemTextures = FRDGSystemTextures::Get(GraphBuilder);
+
+			AtlasParameters->SceneDistanceFieldAssetData = DefaultVector4;
+			AtlasParameters->DistanceFieldIndirectionTable = DefaultUInt32;
+			AtlasParameters->DistanceFieldIndirection2Table = DefaultVector4;
+			AtlasParameters->DistanceFieldIndirectionAtlas = SystemTextures.VolumetricBlack;
+			AtlasParameters->DistanceFieldBrickTexture = SystemTextures.VolumetricBlack;
+			AtlasParameters->DistanceFieldSampler = TStaticSamplerState<SF_Bilinear, AM_Wrap, AM_Wrap, AM_Wrap>::GetRHI();
+			AtlasParameters->DistanceFieldBrickSize = FVector3f::ZeroVector;
+			AtlasParameters->DistanceFieldUniqueDataBrickSize = FVector3f::ZeroVector;
+			AtlasParameters->DistanceFieldBrickAtlasSizeInBricks = FIntVector::ZeroValue;
+			AtlasParameters->DistanceFieldBrickAtlasMask = FIntVector::ZeroValue;
+			AtlasParameters->DistanceFieldBrickAtlasSizeLog2 = FIntVector::ZeroValue;
+			AtlasParameters->DistanceFieldBrickAtlasTexelSize = FVector3f::ZeroVector;
+			AtlasParameters->DistanceFieldBrickAtlasHalfTexelSize = FVector3f::ZeroVector;
+			AtlasParameters->DistanceFieldBrickOffsetToAtlasUVScale = FVector3f::ZeroVector;
+			AtlasParameters->DistanceFieldUniqueDataBrickSizeInAtlasTexels = FVector3f::ZeroVector;
+		}
+	}
+}
diff --git a/Engine/Source/Runtime/Renderer/Public/FXRenderingUtils.h b/Engine/Source/Runtime/Renderer/Public/FXRenderingUtils.h
--- a/Engine/Source/Runtime/Renderer/Public/FXRenderingUtils.h
+++ b/Engine/Source/Runtime/Renderer/Public/FXRenderingUtils.h
@@ -2,11 +2,35 @@
 
 #pragma once
 
+#include "HAL/Platform.h"
+#include "UObject/ObjectMacros.h"
+
+#if UE_ENABLE_INCLUDE_ORDER_DEPRECATED_IN_5_3
 #include "CoreMinimal.h"
+#endif
 
-class FSceneViewFamily;
-class FPrimitiveSceneProxy;
 class FMaterial;
+class FPrimitiveSceneProxy;
+class FRDGBuilder;
+class FSceneView;
+class FSceneViewFamily;
+class FShaderParametersMetadata;
+
+namespace UE::FXRenderingUtils
+{
+	RENDERER_API bool CanMaterialRenderBeforeFXPostOpaque(const FSceneViewFamily& ViewFamily, const FPrimitiveSceneProxy& SceneProxy, const FMaterial& Material);
+
+	namespace DistanceFields
+	{
+		RENDERER_API const FShaderParametersMetadata* GetObjectBufferParametersMetadata();
+		RENDERER_API const FShaderParametersMetadata* GetAtlasParametersMetadata();
+
+		RENDERER_API bool HasDataToBind(const FSceneView& View);
+
+		RENDERER_API void SetupObjectBufferParameters(FRDGBuilder& GraphBuilder, uint8* DestinationData, const FSceneView* View);
+		RENDERER_API void SetupAtlasParameters(FRDGBuilder& GraphBuilder, uint8* DestinationData, const FSceneView* View);
+	}
+}
 
 /**
  * This class exposes methods required by FX rendering that must access rendering internals.
@@ -19,8 +43,12 @@ public:
 	FFXRenderingUtils& operator=(const FFXRenderingUtils&) = delete;
 
 	/** Utility to determine if a material might render before the FXSystem's PostRenderOpaque is called for the view family */
+	UE_DEPRECATED(5.3, "Use UE::FXRenderingUtils::CanMaterialRenderBeforeFXPostOpaque")
 	static bool CanMaterialRenderBeforeFXPostOpaque(
 		const FSceneViewFamily& ViewFamily,
 		const FPrimitiveSceneProxy& SceneProxy,
-		const FMaterial& Material);
+		const FMaterial& Material)
+	{
+		return UE::FXRenderingUtils::CanMaterialRenderBeforeFXPostOpaque(ViewFamily, SceneProxy, Material);
+	}
 };
