Use !HasStaticShadowing() instead of IsMovable() on mobile code paths of FSceneRenderer::BeginInitDynamicShadows.
- required to support Stationary lights falling back to Movable when project has r.AllowStaticLighting=0

diff --git a/Engine/Source/Runtime/Renderer/Private/ShadowSetup.cpp b/Engine/Source/Runtime/Renderer/Private/ShadowSetup.cpp
--- a/Engine/Source/Runtime/Renderer/Private/ShadowSetup.cpp
+++ b/Engine/Source/Runtime/Renderer/Private/ShadowSetup.cpp
@@ -5860,7 +5860,7 @@ void FSceneRenderer::PrepareDynamicShadows(FDynamicShadowsTaskData& TaskData)
 						{
 							static_assert(UE_ARRAY_COUNT(Scene->MobileDirectionalLights) == 3, "All array entries for MobileDirectionalLights must be checked");
 							if (!bMobile ||
-								((LightSceneInfo->Proxy->UseCSMForDynamicObjects() || LightSceneInfo->Proxy->IsMovable())
+								((LightSceneInfo->Proxy->UseCSMForDynamicObjects() || !LightSceneInfo->Proxy->HasStaticShadowing())
 									// Mobile uses the scene's MobileDirectionalLights only for whole scene shadows.
 									&& (LightSceneInfo == Scene->MobileDirectionalLights[0] || LightSceneInfo == Scene->MobileDirectionalLights[1] || LightSceneInfo == Scene->MobileDirectionalLights[2])))
 							{
