Remove hair temporal weight.

diff --git a/Engine/Plugins/Runtime/HairStrands/Source/HairStrandsCore/Private/GroomManager.cpp b/Engine/Plugins/Runtime/HairStrands/Source/HairStrandsCore/Private/GroomManager.cpp
--- a/Engine/Plugins/Runtime/HairStrands/Source/HairStrandsCore/Private/GroomManager.cpp
+++ b/Engine/Plugins/Runtime/HairStrands/Source/HairStrandsCore/Private/GroomManager.cpp
@@ -773,8 +773,6 @@ static void RunHairLODSelection(
 					Instance->HairGroupPublicData->ContinuousLODBounds = SphereBound;
 					Instance->HairGroupPublicData->MaxScreenSize = MaxScreenSize;
 				}
-
-				Instance->HairGroupPublicData->UpdateTemporalIndex();
 			}
 			else
 			{
diff --git a/Engine/Plugins/Runtime/HairStrands/Source/HairStrandsCore/Private/HairStrandsRendering.cpp b/Engine/Plugins/Runtime/HairStrands/Source/HairStrandsCore/Private/HairStrandsRendering.cpp
--- a/Engine/Plugins/Runtime/HairStrands/Source/HairStrandsCore/Private/HairStrandsRendering.cpp
+++ b/Engine/Plugins/Runtime/HairStrands/Source/HairStrandsCore/Private/HairStrandsRendering.cpp
@@ -827,8 +827,7 @@ class FHairClusterAABBCS : public FGlobalShader
 	BEGIN_SHADER_PARAMETER_STRUCT(FParameters, )
 		SHADER_PARAMETER_STRUCT_INCLUDE(ShaderPrint::FShaderParameters, ShaderDrawParameters)
 		SHADER_PARAMETER(uint32, TotalClusterCount)
-		SHADER_PARAMETER(uint32, VertexCount)
-		SHADER_PARAMETER(float, RcpSampleWeight)
+		SHADER_PARAMETER(uint32, PointCount)
 
 		SHADER_PARAMETER(FVector3f, CPUBoundMin)
 		SHADER_PARAMETER(FVector3f, CPUBoundMax)
@@ -898,8 +897,7 @@ static void AddHairClusterAABBPass(
 	Parameters->RenderDeformedOffsetBuffer = RenderDeformedOffsetBuffer;
 	Parameters->TotalClusterCount = 1;
 
-	Parameters->VertexCount = Instance->HairGroupPublicData->GetActiveStrandsPointCount(Instance->Strands.RestResource->GetPointCount(), Instance->HairGroupPublicData->MaxScreenSize);
-	Parameters->RcpSampleWeight = 1.0 / Instance->HairGroupPublicData->GetActiveStrandsSampleWeight(true, Instance->HairGroupPublicData->MaxScreenSize);
+	Parameters->PointCount = Instance->HairGroupPublicData->GetActiveStrandsPointCount(Instance->Strands.RestResource->GetPointCount(), Instance->HairGroupPublicData->MaxScreenSize);
 
 	if (ShaderPrintData)
 	{
@@ -1920,7 +1918,7 @@ void ComputeHairStrandsInterpolation(
 					// 2.1.2 Previous Position
 					// * Transfer prev. position
 					// * Or recompute interpolation pass for previous positions if needed
-					const bool bRecomputePrevPosition = IsHairVisibilityComputeRasterContinuousLODEnabled() || IsHairVisibilityComputeRasterTemporalLayeringEnabled();
+					const bool bRecomputePrevPosition = IsHairVisibilityComputeRasterContinuousLODEnabled();
 					const bool bTransferPrevPosition = 
 						!bRecomputePrevPosition &&
 						Instance->HairGroupPublicData->VFInput.bHasLODSwitch && 
diff --git a/Engine/Shaders/Private/HairStrands/HairStrandsInterpolation.usf b/Engine/Shaders/Private/HairStrands/HairStrandsInterpolation.usf
--- a/Engine/Shaders/Private/HairStrands/HairStrandsInterpolation.usf
+++ b/Engine/Shaders/Private/HairStrands/HairStrandsInterpolation.usf
@@ -650,8 +650,7 @@ uint TotalClusterCount;
 float3 CPUBoundMin;
 float3 CPUBoundMax;
 
-uint VertexCount;
-float RcpSampleWeight;
+uint PointCount;
 
 float4x4 LocalToTranslatedWorldMatrix;
 
@@ -698,7 +697,7 @@ void ClusterAABBEvaluationCS(uint3 DispatchThreadId : SV_DispatchThreadID, uint3
 	}
 	const uint ClusterId   = ClusterIdBuffer[ClusterFetchIndex];
 	const uint IndexOffset = ClusterIndexOffsetBuffer[ClusterFetchIndex];
-	const uint IndexCount = max(min(ClusterIndexCountBuffer[ClusterFetchIndex], GROUP_SIZE),ClusterIndexCountBuffer[ClusterFetchIndex] * RcpSampleWeight); // HAIR_TODO:
+	const uint IndexCount = min(ClusterIndexCountBuffer[ClusterFetchIndex], GROUP_SIZE);
 	
 	const uint VertexCountAfterCulling = HairStrandsVF_CullingIndirectBuffer[3];
 
@@ -714,10 +713,10 @@ void ClusterAABBEvaluationCS(uint3 DispatchThreadId : SV_DispatchThreadID, uint3
 		if (ThreadReadOffset < IndexCount)
 		{
 			const float3 OutHairPositionOffset = RenderDeformedOffsetBuffer[0].xyz;
-			const uint VertexIndex = ClusterVertexIdBuffer[IndexOffset + ThreadReadOffset];
-			if (VertexIndex < VertexCount)
+			const uint PointIndex = ClusterVertexIdBuffer[IndexOffset + ThreadReadOffset];
+			if (PointIndex < PointCount)
 			{
-				const FHairControlPoint CP = UnpackHairControlPoint(RenderDeformedPositionBuffer[VertexIndex], OutHairPositionOffset, 1, 1, 1);
+				const FHairControlPoint CP = UnpackHairControlPoint(RenderDeformedPositionBuffer[PointIndex], OutHairPositionOffset, 1, 1, 1);
 
 				if (all(IsFinite(CP.Position)))
 				{
diff --git a/Engine/Shaders/Private/HairStrands/HairStrandsVisibilityRasterCompute.usf b/Engine/Shaders/Private/HairStrands/HairStrandsVisibilityRasterCompute.usf
--- a/Engine/Shaders/Private/HairStrands/HairStrandsVisibilityRasterCompute.usf
+++ b/Engine/Shaders/Private/HairStrands/HairStrandsVisibilityRasterCompute.usf
@@ -73,7 +73,7 @@ uint			HairMaterialId;
 Texture2D<float> SceneDepthTexture;
 
 uint 			VertexCount;
-float			SampleWeight;
+float			CoverageScale;
 
 float3 NDCToPixelCoord(float4 InDC)
 {
@@ -851,7 +851,7 @@ void PlotInternal(int2 Coords, float AntiAliasingFactor, float4 P0, float4 P1, f
 		if (PackedDepthCov > group_SubTile[LinearIndex].w)
 		{
 			const float Rad = ComputePerspectiveCorrectRadius(Rad0, Rad1, Alpha, P0.w, P1.w);
-			InterlockedAdd(group_SubTile[LinearIndex].z, min(Rad, 0.5f) * 2.0f * 1000.0f * SampleWeight * AntiAliasingFactor);
+			InterlockedAdd(group_SubTile[LinearIndex].z, min(Rad, 0.5f) * 2.0f * 1000.0f * CoverageScale * AntiAliasingFactor);
 		}
 	}
 }
@@ -1260,7 +1260,7 @@ void Plot(int2 Coord, float4 P0, float4 P1, float Rad0, float Rad1, float Segmen
 		if (PackedDepthCov > group_SubTileSceneDepth[LinearIndex])
 		{
 			const float Rad = ComputePerspectiveCorrectRadius(Rad0, Rad1, Alpha, P0.w, P1.w);
-			const uint HairCount = min(Rad, 0.5f) * 2.0f * 1000.0f * SampleWeight;
+			const uint HairCount = min(Rad, 0.5f) * 2.0f * 1000.0f * CoverageScale;
 
 			const uint CoverageMask = GetCoverageMask(Coord, P0.xy, P1.xy);
 
diff --git a/Engine/Shaders/Private/HairStrands/HairStrandsVisibilityRasterComputeNaive.usf b/Engine/Shaders/Private/HairStrands/HairStrandsVisibilityRasterComputeNaive.usf
--- a/Engine/Shaders/Private/HairStrands/HairStrandsVisibilityRasterComputeNaive.usf
+++ b/Engine/Shaders/Private/HairStrands/HairStrandsVisibilityRasterComputeNaive.usf
@@ -23,7 +23,7 @@ Buffer<uint>		IndirectPrimIDs;
 float2				OutputResolution;
 uint				HairMaterialId;
 uint				ControlPointCount;
-float				SampleWeight;
+float				CoverageScale;
 uint				NumWorkGroups;
 
 // Line clipping based on "CLIPPING USING HOMOGENEOUS COORDINATES" by Blinn et al.
@@ -110,7 +110,7 @@ void Plot(int2 Coord, float AntiAliasingFactor, float4 P0, float4 P1, float Rad0
 		// Divide by W to make thickness dependent on screen space depth? This division was kept from the previous line rasterization algorithm.
 		const float Rad = lerp(Rad0, Rad1, PerspectiveAlpha) * LerpedRcpW;
 
-		InterlockedAdd(OutHairCountTexture[Coord], min(Rad, 0.5f) * 2.0f * 1000.0f * SampleWeight * AntiAliasingFactor);
+		InterlockedAdd(OutHairCountTexture[Coord], min(Rad, 0.5f) * 2.0f * 1000.0f * CoverageScale * AntiAliasingFactor);
 	}
 }
 
diff --git a/Engine/Shaders/Private/HairStrands/HairStrandsVisibilityRasterHW.usf b/Engine/Shaders/Private/HairStrands/HairStrandsVisibilityRasterHW.usf
--- a/Engine/Shaders/Private/HairStrands/HairStrandsVisibilityRasterHW.usf
+++ b/Engine/Shaders/Private/HairStrands/HairStrandsVisibilityRasterHW.usf
@@ -139,7 +139,7 @@ VSOutput VSMain(VSInput Input)
 RWTexture2D<uint>	OutHairCountTexture;
 RWTexture2D<uint>	OutDepthCovTexture;
 RWTexture2D<uint>	OutPrimMatTexture;
-float				SampleWeight;
+float				CoverageScale;
 
 [earlydepthstencil]
 void PSMain(
@@ -151,7 +151,7 @@ void PSMain(
 	const uint2 Coord = Position.xy;
 	const uint PackedDepthCov = PackHairVisDepthCoverage(Position.z, 1.0f);
 
-	InterlockedAdd(OutHairCountTexture[Coord], min(Radius / Position.w, 0.5f) * 2.0f * 1000.0f * SampleWeight);
+	InterlockedAdd(OutHairCountTexture[Coord], min(Radius / Position.w, 0.5f) * 2.0f * 1000.0f * CoverageScale);
 
 	// Write Depth + PrimMatID if depth test against hair depths is passed
 	uint OldValue;
diff --git a/Engine/Shaders/Private/HairStrands/HairStrandsVoxelRasterCompute.usf b/Engine/Shaders/Private/HairStrands/HairStrandsVoxelRasterCompute.usf
--- a/Engine/Shaders/Private/HairStrands/HairStrandsVoxelRasterCompute.usf
+++ b/Engine/Shaders/Private/HairStrands/HairStrandsVoxelRasterCompute.usf
@@ -792,7 +792,7 @@ StructuredBuffer<FPackedVirtualVoxelNodeDesc> VirtualVoxelParams_NodeDescBuffer;
 
 RWTexture3D<uint> OutPageTexture;
 
-float SampleWeight;
+float CoverageScale;
 
 #define JITTER_ENABLE 0
 
@@ -944,7 +944,7 @@ void MainCS(uint2 DispatchThreadID : SV_DispatchThreadID)
 				const float VoxelFixPointScale = GetVoxelDensityFixPointScale();
 				
 				const float HairCoverage = lerp(HairCoverage0, HairCoverage1, U);
-				uint RawData = HairCoverage * VoxelFixPointScale * HairStrandsVF_Density * SampleWeight;
+				uint RawData = HairCoverage * VoxelFixPointScale * HairStrandsVF_Density * CoverageScale;
 				InterlockedAdd(OutPageTexture[VoxelPageCoord], RawData);
 
 				// Groom having raytraced geometry will cast-shadow on opaque geomtry with their RT-geometry, not with their 
diff --git a/Engine/Source/Runtime/Renderer/Private/HairStrands/HairStrandsComposition.cpp b/Engine/Source/Runtime/Renderer/Private/HairStrands/HairStrandsComposition.cpp
--- a/Engine/Source/Runtime/Renderer/Private/HairStrands/HairStrandsComposition.cpp
+++ b/Engine/Source/Runtime/Renderer/Private/HairStrands/HairStrandsComposition.cpp
@@ -182,8 +182,6 @@ public:
 
 IMPLEMENT_GLOBAL_SHADER(FHairVisibilityComposeSamplePS, "/Engine/Private/HairStrands/HairStrandsVisibilityComposeSubPixelPS.usf", "ComposeSamplePS", SF_Pixel);
 
-int32 GetHairVisibilityComputeRasterTemporalLayerCount();
-
 static void AddHairVisibilityComposeSamplePass(
 	FRDGBuilder& GraphBuilder,
 	const FViewInfo& View,
@@ -198,7 +196,7 @@ static void AddHairVisibilityComposeSamplePass(
 	check(VisibilityData.SampleLightingTexture);
 	const bool bDOFEnable = HairDOFDepthTexture != nullptr ? 1 : 0;
 
-	const bool bTemporal = HairTemporalAccumulationTexture != nullptr ? 1 : 0;
+	const bool bTemporal = false; //HAIR_TODO: remove
 
 	TRDGUniformBufferRef<FFogUniformParameters> FogBuffer = CreateFogUniformBuffer(GraphBuilder, View);
 
@@ -210,7 +208,7 @@ static void AddHairVisibilityComposeSamplePass(
 	Parameters->ViewUniformBuffer = View.ViewUniformBuffer;
 	Parameters->HairStrands = View.HairStrandsViewData.UniformBuffer;
 	Parameters->FogStruct = FogBuffer;
-	Parameters->TemporalLayerCount = bTemporal ? GetHairVisibilityComputeRasterTemporalLayerCount() : 1;
+	Parameters->TemporalLayerCount = 1;
 	Parameters->HairTemporalAccumulationTexture = bTemporal ? HairTemporalAccumulationTexture : GSystemTextures.GetBlackDummy(GraphBuilder);
 	Parameters->RenderTargets[0] = FRenderTargetBinding(OutColorTexture, bTemporal ? ERenderTargetLoadAction::ELoad : (bTemporalLayeringEnabled ? ERenderTargetLoadAction::EClear : ERenderTargetLoadAction::ELoad) );
 	Parameters->RenderTargets.DepthStencil = FDepthStencilBinding(OutDepthTexture, ERenderTargetLoadAction::ELoad, ERenderTargetLoadAction::ENoAction, FExclusiveDepthStencil::DepthWrite_StencilNop);	
@@ -516,7 +514,7 @@ static void InternalRenderHairComposition(
 
 		FRDGTextureRef AccumulatedColor = nullptr;
 
-		const bool bTemporalLayeringEnabled = IsHairVisibilityComputeRasterTemporalLayeringEnabled() && View.ViewState;
+		const bool bTemporalLayeringEnabled = false; // HAIR_TODO: remove
 
 		if (bTemporalLayeringEnabled)
 		{
diff --git a/Engine/Source/Runtime/Renderer/Private/HairStrands/HairStrandsInterface.cpp b/Engine/Source/Runtime/Renderer/Private/HairStrands/HairStrandsInterface.cpp
--- a/Engine/Source/Runtime/Renderer/Private/HairStrands/HairStrandsInterface.cpp
+++ b/Engine/Source/Runtime/Renderer/Private/HairStrands/HairStrandsInterface.cpp
@@ -83,21 +83,6 @@ static TAutoConsoleVariable<int32> CVarHairStrandsVisibilityComputeRaster_Contin
 	TEXT("Enable Continuos LOD when using compute rasterization. Experimental"),
 	ECVF_RenderThreadSafe);
 
-static TAutoConsoleVariable<int32> CVarHairStrandsVisibilityComputeRaster_TemporalLayering(
-	TEXT("r.HairStrands.Visibility.ComputeRaster.TemporalLayering"), 1,
-	TEXT("Enable Experimental WIP Temporal Layering (requires TAA changes to work well)"),
-	ECVF_RenderThreadSafe);
-
-static TAutoConsoleVariable<int32> CVarHairStrandsVisibilityComputeRaster_TemporalLayering_LayerCount(
-	TEXT("r.HairStrands.Visibility.ComputeRaster.TemporalLayering.LayerCount"), 2,
-	TEXT("Temporal Layering Layer Count (default: 2)"),
-	ECVF_RenderThreadSafe);
-
-static TAutoConsoleVariable<int32> CVarHairStrandsVisibilityComputeRaster_TemporalLayering_OverrideIndex(
-	TEXT("r.HairStrands.Visibility.ComputeRaster.TemporalLayering.OverrideIndex"), -1,
-	TEXT("Enable Temporal Layering Override Index (default: -1 = no override)"),
-	ECVF_RenderThreadSafe);
-
 ////////////////////////////////////////////////////////////////////////////////////////////////////////////////////////
 // Hair strands instance ref. counting for debug purpose only
 uint32 FHairStrandsInstance::GetRefCount() const
@@ -512,43 +497,10 @@ float GetHairVisibilityComputeRasterContinuousLODScale(float ScreenSize)
 	return FMath::Pow(FMath::Clamp(ScreenSize, 1.0f / 16.0f, 1.0f), 1.0f);
 }
 
-bool IsHairVisibilityComputeRasterTemporalLayeringEnabled()
-{
-	return IsHairStrandContinuousDecimationReorderingEnabled() && IsHairVisibilityComputeRasterEnabled() && (CVarHairStrandsVisibilityComputeRaster_TemporalLayering.GetValueOnAnyThread() > 0);
-}
-
-int32 GetHairVisibilityComputeRasterTemporalLayerCount()
-{
-	const uint32 LayerCount = FMath::Clamp(CVarHairStrandsVisibilityComputeRaster_TemporalLayering_LayerCount.GetValueOnAnyThread(), 1, 32);
-	return LayerCount;
-}
-
-uint32 GetHairVisibilityComputeRasterVertexStart(uint32 TemporalIndex, uint32 InVertexCount)
-{
-	uint32 VertexStart = 0;
-
-	if (IsHairVisibilityComputeRasterTemporalLayeringEnabled())
-	{
-		const uint32 LayerCount = GetHairVisibilityComputeRasterTemporalLayerCount();
-		const int32 OverrideIndex = CVarHairStrandsVisibilityComputeRaster_TemporalLayering_OverrideIndex.GetValueOnAnyThread();
-
-		const uint32 VertexCount = InVertexCount / LayerCount;
-
-		VertexStart = (((OverrideIndex >= 0) ? OverrideIndex : TemporalIndex) % LayerCount) * VertexCount;
-	}
-
-	return VertexStart;
-}
-
-uint32 GetHairVisibilityComputeRasterPointCount(float ScreenSize, uint32 InPointCount)
+static uint32 GetHairVisibilityComputeRasterPointCount(float ScreenSize, uint32 InPointCount)
 {
 	uint32 PointCount = InPointCount;
 
-	if (IsHairVisibilityComputeRasterTemporalLayeringEnabled())
-	{
-		PointCount /= GetHairVisibilityComputeRasterTemporalLayerCount();
-	}
-
 	if (IsHairVisibilityComputeRasterContinuousLODEnabled())
 	{
 		PointCount *= GetHairVisibilityComputeRasterContinuousLODScale(ScreenSize);
@@ -557,38 +509,14 @@ uint32 GetHairVisibilityComputeRasterPointCount(float ScreenSize, uint32 InPoint
 	return PointCount;
 }
 
-float GetHairVisibilityComputeRasterSampleWeight(float ScreenSize, bool bUseTemporalWeight)
-{
-	float SampleWeight = 1.0;
-
-	if (IsHairVisibilityComputeRasterTemporalLayeringEnabled() && bUseTemporalWeight)
-	{
-		// sample weight should increase if layer count increases
-		SampleWeight *= GetHairVisibilityComputeRasterTemporalLayerCount();
-	}
-
-	if (IsHairVisibilityComputeRasterContinuousLODEnabled())
-	{
-		//sample weight should increase if decimation scale is lower
-		SampleWeight /= GetHairVisibilityComputeRasterContinuousLODScale(ScreenSize);
-	}
-
-	return SampleWeight;
-}
-
 uint32 FHairGroupPublicData::GetActiveStrandsPointCount(uint32 InPointCount, float ScreenSize) const
 {
 	return GetHairVisibilityComputeRasterPointCount(FMath::Min(MaxScreenSize, ScreenSize), InPointCount);
 }
 
-float FHairGroupPublicData::GetActiveStrandsSampleWeight(bool bUseTemporalWeight, float ScreenSize) const
-{
-	return GetHairVisibilityComputeRasterSampleWeight(FMath::Min(MaxScreenSize, ScreenSize), bUseTemporalWeight);
-}
-
-void FHairGroupPublicData::UpdateTemporalIndex()
+float FHairGroupPublicData::GetActiveStrandsCoverageScale() const
 {
-	TemporalIndex = IsHairVisibilityComputeRasterTemporalLayeringEnabled() ? ((TemporalIndex + 1) % GetHairVisibilityComputeRasterTemporalLayerCount()) : 0;
+	return 1.f;
 }
 
 ///////////////////////////////////////////////////////////////////////////////////////////////////
diff --git a/Engine/Source/Runtime/Renderer/Private/HairStrands/HairStrandsVisibility.cpp b/Engine/Source/Runtime/Renderer/Private/HairStrands/HairStrandsVisibility.cpp
--- a/Engine/Source/Runtime/Renderer/Private/HairStrands/HairStrandsVisibility.cpp
+++ b/Engine/Source/Runtime/Renderer/Private/HairStrands/HairStrandsVisibility.cpp
@@ -2982,7 +2982,7 @@ class FVisiblityRasterComputeRasterizeCS : public FGlobalShader
 		SHADER_PARAMETER(uint32, HairMaterialId)
 		SHADER_PARAMETER_STRUCT_INCLUDE(FHairStrandsInstanceParameters, HairInstance)
 		SHADER_PARAMETER_RDG_TEXTURE(Texture2D, SceneDepthTexture)
-		SHADER_PARAMETER(float, SampleWeight)
+		SHADER_PARAMETER(float, CoverageScale)
 		SHADER_PARAMETER_RDG_BUFFER_SRV(Buffer, VisTilePrims)
 		SHADER_PARAMETER_RDG_BUFFER_SRV(Buffer, VisTileArgs)
 		SHADER_PARAMETER_RDG_BUFFER_SRV(ByteAddressBuffer, VisTileData)
@@ -3020,7 +3020,7 @@ class FVisiblityRasterComputeRasterizeMultiSampleCS : public FGlobalShader
 			SHADER_PARAMETER(uint32, HairMaterialId)
 		SHADER_PARAMETER_STRUCT_INCLUDE(FHairStrandsInstanceParameters, HairInstance)
 		SHADER_PARAMETER_RDG_TEXTURE(Texture2D, SceneDepthTexture)
-		SHADER_PARAMETER(float, SampleWeight)
+		SHADER_PARAMETER(float, CoverageScale)
 		SHADER_PARAMETER_RDG_BUFFER_SRV(Buffer, VisTilePrims)
 		SHADER_PARAMETER_RDG_BUFFER_SRV(Buffer, VisTileArgs)
 		SHADER_PARAMETER_RDG_BUFFER_SRV(ByteAddressBuffer, VisTileData)
@@ -3121,7 +3121,7 @@ class FVisiblityRasterHWPS : public FGlobalShader
 		SHADER_PARAMETER_RDG_TEXTURE_UAV(RWTexture2D, OutHairCountTexture)
 		SHADER_PARAMETER_RDG_TEXTURE_UAV(RWTexture2D, OutDepthCovTexture)
 		SHADER_PARAMETER_RDG_TEXTURE_UAV(RWTexture2D, OutPrimMatTexture)
-		SHADER_PARAMETER(float, SampleWeight)
+		SHADER_PARAMETER(float, CoverageScale)
 		RENDER_TARGET_BINDING_SLOTS()
 	END_SHADER_PARAMETER_STRUCT()
 
@@ -3156,7 +3156,7 @@ class FVisiblityRasterComputeNaiveCS : public FGlobalShader
 		SHADER_PARAMETER(FVector2f, OutputResolution)
 		SHADER_PARAMETER(uint32, HairMaterialId)
 		SHADER_PARAMETER(uint32, ControlPointCount)
-		SHADER_PARAMETER(float, SampleWeight)
+		SHADER_PARAMETER(float, CoverageScale)
 		SHADER_PARAMETER(uint32, NumWorkGroups)
 		SHADER_PARAMETER_STRUCT_INCLUDE(FHairStrandsInstanceParameters, HairInstance)
 		SHADER_PARAMETER_STRUCT_REF(FViewUniformShaderParameters, ViewUniformBuffer)
@@ -3435,7 +3435,7 @@ static FRasterComputeOutput AddVisibilityComputeRasterPass(
 			const float ScreenSize = ComputeBoundsScreenSize(FVector4(BoundsSphere.Center, 1), BoundsSphere.W, ViewInfo);
 
 			const uint32 PointCount = HairGroupPublicData->GetActiveStrandsPointCount(VFInput.Strands.PointCount, ScreenSize);
-			const float SampleWeight = HairGroupPublicData->GetActiveStrandsSampleWeight(false, ScreenSize);
+			const float CoverageScale = HairGroupPublicData->GetActiveStrandsCoverageScale();
 
 			// HW/SW classification
 			if (bClassification)
@@ -3488,7 +3488,7 @@ static FRasterComputeOutput AddVisibilityComputeRasterPass(
 				Parameters->PS.OutHairCountTexture = HairCountTextureUAV;
 				Parameters->PS.OutDepthCovTexture = DepthCovTextureUAV;
 				Parameters->PS.OutPrimMatTexture = PrimMatTextureUAV;
-				Parameters->PS.SampleWeight = SampleWeight;
+				Parameters->PS.CoverageScale = CoverageScale;
 				Parameters->PS.RenderTargets.DepthStencil = FDepthStencilBinding(
 					SceneDepthTexture,
 					ERenderTargetLoadAction::ELoad,
@@ -3541,7 +3541,7 @@ static FRasterComputeOutput AddVisibilityComputeRasterPass(
 					RasterParameters->OutputResolution = Common.OutputResolutionf;
 					RasterParameters->HairMaterialId = PrimitiveInfo.MaterialId;
 					RasterParameters->ControlPointCount = PointCount;
-					RasterParameters->SampleWeight = SampleWeight;
+					RasterParameters->CoverageScale = CoverageScale;
 					RasterParameters->NumWorkGroups = NumWorkGroups;
 					RasterParameters->HairInstance = GetHairStrandsInstanceParameters(GraphBuilder, ViewInfo, HairGroupPublicData, bCullingEnableInRasterizers, bForceRegister);
 					RasterParameters->ViewUniformBuffer = ViewUniformShaderParameters;
@@ -3627,7 +3627,7 @@ static FRasterComputeOutput AddVisibilityComputeRasterPass(
 						RasterParameters->OutHairCountTexture = HairCountTextureUAV;
 						RasterParameters->OutDepthCovTexture = DepthCovTextureUAV;
 						RasterParameters->OutPrimMatTexture = PrimMatTextureUAV;
-						RasterParameters->SampleWeight = SampleWeight;
+						RasterParameters->CoverageScale = CoverageScale;
 						RasterParameters->HairInstance = GetHairStrandsInstanceParameters(GraphBuilder, ViewInfo, HairGroupPublicData, bCullingEnableInRasterizers, bForceRegister);
 						RasterParameters->VisTilePrims = bTileCompactionPass ? GraphBuilder.CreateSRV(CompactedVisTilePrims, PF_R32_UINT) : VisTilePrimsSRV;
 						RasterParameters->VisTileArgs = bTileCompactionPass ? GraphBuilder.CreateSRV(CompactedVisTileArgs, PF_R32_UINT) : VisTileArgsSRV;
@@ -3646,7 +3646,7 @@ static FRasterComputeOutput AddVisibilityComputeRasterPass(
 						RasterParameters->OutHairCountTexture = HairCountTextureUAV;
 						RasterParameters->OutDepthCovTexture = DepthCovTextureUAV;
 						RasterParameters->OutPrimMatTexture = PrimMatTextureUAV;
-						RasterParameters->SampleWeight = SampleWeight;
+						RasterParameters->CoverageScale = CoverageScale;
 						RasterParameters->HairInstance = GetHairStrandsInstanceParameters(GraphBuilder, ViewInfo, HairGroupPublicData, bCullingEnableInRasterizers, bForceRegister);
 						RasterParameters->VisTilePrims = bTileCompactionPass ? GraphBuilder.CreateSRV(CompactedVisTilePrims, PF_R32_UINT) : VisTilePrimsSRV;
 						RasterParameters->VisTileArgs = bTileCompactionPass ? GraphBuilder.CreateSRV(CompactedVisTileArgs, PF_R32_UINT) : VisTileArgsSRV;
diff --git a/Engine/Source/Runtime/Renderer/Private/HairStrands/HairStrandsVoxelization.cpp b/Engine/Source/Runtime/Renderer/Private/HairStrands/HairStrandsVoxelization.cpp
--- a/Engine/Source/Runtime/Renderer/Private/HairStrands/HairStrandsVoxelization.cpp
+++ b/Engine/Source/Runtime/Renderer/Private/HairStrands/HairStrandsVoxelization.cpp
@@ -1173,7 +1173,7 @@ class FVoxelRasterComputeCS : public FGlobalShader
 		SHADER_PARAMETER_RDG_BUFFER_SRV(Buffer, VoxelizationViewInfoBuffer)
 		SHADER_PARAMETER_RDG_TEXTURE_UAV(RWTexture3D, OutPageTexture)
 		SHADER_PARAMETER(uint32, VertexCount)
-		SHADER_PARAMETER(float, SampleWeight)
+		SHADER_PARAMETER(float, CoverageScale)
 	END_SHADER_PARAMETER_STRUCT()
 
 public:
@@ -1234,7 +1234,7 @@ static void AddVirtualVoxelizationComputeRasterPass(
 				}
 
 				const uint32 PointCount = HairGroupPublicData->GetActiveStrandsPointCount(VFInput.Strands.PointCount, HairGroupPublicData->MaxScreenSize);
-				const float SampleWeight = HairGroupPublicData->GetActiveStrandsSampleWeight(false, HairGroupPublicData->MaxScreenSize);
+				const float CoverageScale = HairGroupPublicData->GetActiveStrandsCoverageScale();
 
 				FTransform LocalToTranslatedWorldTransform = VFInput.LocalToWorldTransform;
 				LocalToTranslatedWorldTransform.AddToTranslation(TranslatedWorldOffset);
@@ -1248,7 +1248,7 @@ static void AddVirtualVoxelizationComputeRasterPass(
 				PassParameters->FrameIdMod8 = FrameIdMode8;
 
 				PassParameters->VertexCount = PointCount;
-				PassParameters->SampleWeight = SampleWeight;
+				PassParameters->CoverageScale = CoverageScale;
 			
 				const bool bCullingEnable = !IsHairStrandContinuousDecimationReorderingEnabled() && HairGroupPublicData->GetCullingResultAvailable();
 
diff --git a/Engine/Source/Runtime/Renderer/Public/HairStrandsInterface.h b/Engine/Source/Runtime/Renderer/Public/HairStrandsInterface.h
--- a/Engine/Source/Runtime/Renderer/Public/HairStrandsInterface.h
+++ b/Engine/Source/Runtime/Renderer/Public/HairStrandsInterface.h
@@ -193,8 +193,7 @@ public:
 	uint32 GetClusterCount() const { return ClusterCount;  }
 
 	uint32 GetActiveStrandsPointCount(uint32 InPointCount, float ScreenSize) const;
-	float GetActiveStrandsSampleWeight(bool bUseTemporalWeight, float ScreenSize) const;
-	void UpdateTemporalIndex();
+	float  GetActiveStrandsCoverageScale() const;
 
 	struct FVertexFactoryInput 
 	{
@@ -391,9 +390,6 @@ RENDERER_API bool IsHairStrandContinuousDecimationReorderingEnabled();
 // Return true if continuous LOD is enabled - implies computer raster and continuous decimation reordering is true
 RENDERER_API bool IsHairVisibilityComputeRasterContinuousLODEnabled();
 
-// Return true if temporal layering is enabled - implies computer raster and continuous decimation reordering is true
-RENDERER_API bool IsHairVisibilityComputeRasterTemporalLayeringEnabled();
-
 ////////////////////////////////////////////////////////////////////////////////////////////////////////////////////////
 
 typedef TArray<FRHIUnorderedAccessView*> FBufferTransitionQueue;
