Pipe hair coverage scaling to legacy HW raster path.

diff --git a/Engine/Shaders/Private/HairStrands/HairStrandsVisibilityPS.usf b/Engine/Shaders/Private/HairStrands/HairStrandsVisibilityPS.usf
--- a/Engine/Shaders/Private/HairStrands/HairStrandsVisibilityPS.usf
+++ b/Engine/Shaders/Private/HairStrands/HairStrandsVisibilityPS.usf
@@ -33,6 +33,7 @@
 uint HairVisibilityPass_HairMacroGroupIndex;
 uint HairVisibilityPass_HairMaterialId;
 uint HairVisibilityPass_LightChannelMask;
+float HairVisibilityPass_HairCoverageScale;
 
 #if HAIR_RENDER_MODE == RENDER_MODE_PPLL
 EARLYDEPTHSTENCIL	// Required for UAV operations to not happen.
@@ -81,7 +82,7 @@ void MainVisibility(
 		const float SceneDepth = ConvertFromDeviceZ(SvPosition.z); // Linear depth in world unit
 		const float PixelRadius = HairRenderInfo.bIsOrthoView ? HairRenderInfo.RadiusAtDepth1Primary : SceneDepth * HairRenderInfo.RadiusAtDepth1Primary; // Not correct but the coverage is not used (we count instead the number of sub-sample covered)
 		const float StrandRealRadius = WorldStrandRadius;
-		Coverage = StrandRealRadius / max(StrandRealRadius, PixelRadius);
+		Coverage = saturate(StrandRealRadius / max(StrandRealRadius, PixelRadius) * HairVisibilityPass_HairCoverageScale);
 	}
 #endif
 
diff --git a/Engine/Source/Runtime/Renderer/Private/HairStrands/HairStrandsVisibility.cpp b/Engine/Source/Runtime/Renderer/Private/HairStrands/HairStrandsVisibility.cpp
--- a/Engine/Source/Runtime/Renderer/Private/HairStrands/HairStrandsVisibility.cpp
+++ b/Engine/Source/Runtime/Renderer/Private/HairStrands/HairStrandsVisibility.cpp
@@ -489,11 +489,12 @@ IMPLEMENT_MATERIAL_SHADER_TYPE(, FHairMaterialVS, TEXT("/Engine/Private/HairStra
 class FHairMaterialShaderElementData : public FMeshMaterialShaderElementData
 {
 public:
-	FHairMaterialShaderElementData(int32 MacroGroupId, int32 MaterialId, int32 PrimitiveId, uint32 LightChannelMask) : MaterialPass_MacroGroupId(MacroGroupId), MaterialPass_MaterialId(MaterialId), MaterialPass_PrimitiveId(PrimitiveId), MaterialPass_LightChannelMask(LightChannelMask){ }
+	FHairMaterialShaderElementData(int32 MacroGroupId, int32 MaterialId, int32 PrimitiveId, uint32 LightChannelMask, float HairCoverageScale) : MaterialPass_MacroGroupId(MacroGroupId), MaterialPass_MaterialId(MaterialId), MaterialPass_PrimitiveId(PrimitiveId), MaterialPass_LightChannelMask(LightChannelMask), MaterialPass_HairCoverageScale(HairCoverageScale) { }
 	uint32 MaterialPass_MacroGroupId;
 	uint32 MaterialPass_MaterialId;
 	uint32 MaterialPass_PrimitiveId;
 	uint32 MaterialPass_LightChannelMask;
+	uint32 MaterialPass_HairCoverageScale;
 };
 
 #define HAIR_MATERIAL_DEBUG_OUTPUT 0
@@ -516,6 +517,7 @@ public:
 		MaterialPass_MaterialId.Bind(Initializer.ParameterMap, TEXT("MaterialPass_MaterialId"));
 		MaterialPass_PrimitiveId.Bind(Initializer.ParameterMap, TEXT("MaterialPass_PrimitiveId"));
 		MaterialPass_LightChannelMask.Bind(Initializer.ParameterMap, TEXT("MaterialPass_LightChannelMask"));
+		MaterialPass_HairCoverageScale.Bind(Initializer.ParameterMap, TEXT("MaterialPass_HairCoverageScale"));
 	}
 
 	FHairMaterialPS() {}
@@ -553,6 +555,7 @@ public:
 		ShaderBindings.Add(MaterialPass_MaterialId, ShaderElementData.MaterialPass_MaterialId);
 		ShaderBindings.Add(MaterialPass_PrimitiveId, ShaderElementData.MaterialPass_PrimitiveId);
 		ShaderBindings.Add(MaterialPass_LightChannelMask, ShaderElementData.MaterialPass_LightChannelMask);
+		ShaderBindings.Add(MaterialPass_HairCoverageScale, ShaderElementData.MaterialPass_HairCoverageScale);
 	}
 
 private:
@@ -560,6 +563,7 @@ private:
 	LAYOUT_FIELD(FShaderParameter, MaterialPass_MaterialId);
 	LAYOUT_FIELD(FShaderParameter, MaterialPass_PrimitiveId);
 	LAYOUT_FIELD(FShaderParameter, MaterialPass_LightChannelMask);
+	LAYOUT_FIELD(FShaderParameter, MaterialPass_HairCoverageScale);
 };
 IMPLEMENT_MATERIAL_SHADER_TYPE(, FHairMaterialPS, TEXT("/Engine/Private/HairStrands/HairStrandsMaterialPS.usf"), TEXT("Main"), SF_Pixel);
 
@@ -575,7 +579,7 @@ public:
 		FDynamicPassMeshDrawListContext* InDrawListContext);
 
 	virtual void AddMeshBatch(const FMeshBatch& RESTRICT MeshBatch, uint64 BatchElementMask, const FPrimitiveSceneProxy* RESTRICT PrimitiveSceneProxy, int32 StaticMeshId = -1) override final;
-	void AddMeshBatch(const FMeshBatch& RESTRICT MeshBatch, uint64 BatchElementMask, const FPrimitiveSceneProxy* RESTRICT PrimitiveSceneProxy, int32 StaticMeshId, int32 MacroGroupId, int32 HairMaterialId);
+	void AddMeshBatch(const FMeshBatch& RESTRICT MeshBatch, uint64 BatchElementMask, const FPrimitiveSceneProxy* RESTRICT PrimitiveSceneProxy, int32 StaticMeshId, int32 MacroGroupId, int32 HairMaterialId, float HairCoverageScale);
 
 private:
 	bool TryAddMeshBatch(
@@ -585,6 +589,7 @@ private:
 		int32 StaticMeshId,
 		uint32 MacroGroupId,
 		uint32 HairMaterialId,
+		float HairCoverageScale,
 		const FMaterialRenderProxy& MaterialRenderProxy,
 		const FMaterial& Material);
 
@@ -598,17 +603,18 @@ private:
 		const int32 MacroGroupId,
 		const int32 HairMaterialId,
 		const int32 HairPrimitiveId,
-		const uint32 HairPrimitiveLightChannelMask);
+		const uint32 HairPrimitiveLightChannelMask,
+		const float HairCoverageScale);
 
 	FMeshPassProcessorRenderState PassDrawRenderState;
 };
 
 void FHairMaterialProcessor::AddMeshBatch(const FMeshBatch& RESTRICT MeshBatch, uint64 BatchElementMask, const FPrimitiveSceneProxy* RESTRICT PrimitiveSceneProxy, int32 StaticMeshId)
 {
-	AddMeshBatch(MeshBatch, BatchElementMask, PrimitiveSceneProxy, StaticMeshId, 0, 0);
+	AddMeshBatch(MeshBatch, BatchElementMask, PrimitiveSceneProxy, StaticMeshId, 0, 0, 1.f);
 }
 
-void FHairMaterialProcessor::AddMeshBatch(const FMeshBatch& RESTRICT MeshBatch, uint64 BatchElementMask, const FPrimitiveSceneProxy* RESTRICT PrimitiveSceneProxy, int32 StaticMeshId, int32 MacroGroupId, int32 HairMaterialId)
+void FHairMaterialProcessor::AddMeshBatch(const FMeshBatch& RESTRICT MeshBatch, uint64 BatchElementMask, const FPrimitiveSceneProxy* RESTRICT PrimitiveSceneProxy, int32 StaticMeshId, int32 MacroGroupId, int32 HairMaterialId, float HairCoverageScale)
 {
 	const FMaterialRenderProxy* MaterialRenderProxy = MeshBatch.MaterialRenderProxy;
 	while (MaterialRenderProxy)
@@ -616,7 +622,7 @@ void FHairMaterialProcessor::AddMeshBatch(const FMeshBatch& RESTRICT MeshBatch,
 		const FMaterial* Material = MaterialRenderProxy->GetMaterialNoFallback(FeatureLevel);
 		if (Material)
 		{
-			if (TryAddMeshBatch(MeshBatch, BatchElementMask, PrimitiveSceneProxy, StaticMeshId, MacroGroupId, HairMaterialId, *MaterialRenderProxy, *Material))
+			if (TryAddMeshBatch(MeshBatch, BatchElementMask, PrimitiveSceneProxy, StaticMeshId, MacroGroupId, HairMaterialId, HairCoverageScale, *MaterialRenderProxy, *Material))
 			{
 				break;
 			}
@@ -633,6 +639,7 @@ bool FHairMaterialProcessor::TryAddMeshBatch(
 	int32 StaticMeshId,
 	uint32 MacroGroupId,
 	uint32 HairMaterialId,
+	float HairCoverageScale,
 	const FMaterialRenderProxy& MaterialRenderProxy,
 	const FMaterial& Material)
 {
@@ -663,7 +670,7 @@ bool FHairMaterialProcessor::TryAddMeshBatch(
 		FMeshDrawCommandPrimitiveIdInfo IdInfo = GetDrawCommandPrimitiveId(SceneInfo, MeshBatch.Elements[0]);
 		uint32 LightChannelMask = PrimitiveSceneProxy ? PrimitiveSceneProxy->GetLightingChannelMask() : 0;
 
-		return Process(MeshBatchCopy, BatchElementMask, PrimitiveSceneProxy, StaticMeshId, MaterialRenderProxy, Material, MacroGroupId, HairMaterialId, IdInfo.DrawPrimitiveId, LightChannelMask);
+		return Process(MeshBatchCopy, BatchElementMask, PrimitiveSceneProxy, StaticMeshId, MaterialRenderProxy, Material, MacroGroupId, HairMaterialId, IdInfo.DrawPrimitiveId, LightChannelMask, HairCoverageScale);
 	}
 
 	return true;
@@ -679,7 +686,8 @@ bool FHairMaterialProcessor::Process(
 	const int32 MacroGroupId,
 	const int32 HairMaterialId,
 	const int32 HairPrimitiveId,
-	const uint32 HairPrimitiveLightChannelMask)
+	const uint32 HairPrimitiveLightChannelMask,
+	const float HairCoverageScale)
 {
 	const FVertexFactory* VertexFactory = MeshBatch.VertexFactory;
 
@@ -704,7 +712,7 @@ bool FHairMaterialProcessor::Process(
 	}
 
 	FMeshPassProcessorRenderState DrawRenderState(PassDrawRenderState);
-	FHairMaterialShaderElementData ShaderElementData(MacroGroupId, HairMaterialId, HairPrimitiveId, HairPrimitiveLightChannelMask);
+	FHairMaterialShaderElementData ShaderElementData(MacroGroupId, HairMaterialId, HairPrimitiveId, HairPrimitiveLightChannelMask, HairCoverageScale);
 	ShaderElementData.InitializeMeshMaterialData(ViewIfDynamicMeshCommand, PrimitiveSceneProxy, MeshBatch, StaticMeshId, false);
 
 	const bool bReverseCulling = ViewIfDynamicMeshCommand ? ViewIfDynamicMeshCommand->bIsPlanarReflection : false;
@@ -953,7 +961,7 @@ static FMaterialPassOutput AddHairMaterialPass(
 
 						if (bIsCompatible)
 						{
-							MeshProcessor.AddMeshBatch(*MeshBatch, BatchElementMask, PrimitiveInfo.PrimitiveSceneProxy, -1, MacroGroupData.MacroGroupId, PrimitiveInfo.MaterialId);
+							MeshProcessor.AddMeshBatch(*MeshBatch, BatchElementMask, PrimitiveInfo.PrimitiveSceneProxy, -1, MacroGroupData.MacroGroupId, PrimitiveInfo.MaterialId, PrimitiveInfo.PublicDataPtr->GetActiveStrandsCoverageScale());
 						}
 					}
 				}
@@ -1180,10 +1188,11 @@ IMPLEMENT_MATERIAL_SHADER_TYPE(template<>, THairVisiblityVS_PPLL,						TEXT("/En
 class FHairVisibilityShaderElementData : public FMeshMaterialShaderElementData
 {
 public:
-	FHairVisibilityShaderElementData(uint32 InHairMacroGroupId, uint32 InHairMaterialId, uint32 InLightChannelMask) : HairMacroGroupId(InHairMacroGroupId), HairMaterialId(InHairMaterialId), LightChannelMask(InLightChannelMask) { }
+	FHairVisibilityShaderElementData(uint32 InHairMacroGroupId, uint32 InHairMaterialId, uint32 InLightChannelMask, float InHairCoverageScale) : HairMacroGroupId(InHairMacroGroupId), HairMaterialId(InHairMaterialId), LightChannelMask(InLightChannelMask), HairCoverageScale(InHairCoverageScale) { }
 	uint32 HairMacroGroupId;
 	uint32 HairMaterialId;
 	uint32 LightChannelMask;
+	float HairCoverageScale;
 };
 
 template<EHairVisibilityRenderMode RenderMode>
@@ -1201,6 +1210,7 @@ public:
 		HairVisibilityPass_HairMacroGroupIndex.Bind(Initializer.ParameterMap, TEXT("HairVisibilityPass_HairMacroGroupIndex"));
 		HairVisibilityPass_HairMaterialId.Bind(Initializer.ParameterMap, TEXT("HairVisibilityPass_HairMaterialId"));
 		HairVisibilityPass_LightChannelMask.Bind(Initializer.ParameterMap, TEXT("HairVisibilityPass_LightChannelMask"));
+		HairVisibilityPass_HairCoverageScale.Bind(Initializer.ParameterMap, TEXT("HairVisibilityPass_HairCoverageScale"));
 	}
 
 	FHairVisibilityPS() {}
@@ -1258,11 +1268,13 @@ public:
 		ShaderBindings.Add(HairVisibilityPass_HairMacroGroupIndex, ShaderElementData.HairMacroGroupId);
 		ShaderBindings.Add(HairVisibilityPass_HairMaterialId, ShaderElementData.HairMaterialId);
 		ShaderBindings.Add(HairVisibilityPass_LightChannelMask, ShaderElementData.LightChannelMask);
+		ShaderBindings.Add(HairVisibilityPass_HairCoverageScale, ShaderElementData.HairCoverageScale);
 	}
 
 	LAYOUT_FIELD(FShaderParameter, HairVisibilityPass_HairMacroGroupIndex);
 	LAYOUT_FIELD(FShaderParameter, HairVisibilityPass_HairMaterialId);
 	LAYOUT_FIELD(FShaderParameter, HairVisibilityPass_LightChannelMask);
+	LAYOUT_FIELD(FShaderParameter, HairVisibilityPass_HairCoverageScale);
 };
 IMPLEMENT_MATERIAL_SHADER_TYPE(template<>, FHairVisibilityPS<HairVisibilityRenderMode_MSAA_Visibility>, TEXT("/Engine/Private/HairStrands/HairStrandsVisibilityPS.usf"), TEXT("MainVisibility"), SF_Pixel);
 IMPLEMENT_MATERIAL_SHADER_TYPE(template<>, FHairVisibilityPS<HairVisibilityRenderMode_Transmittance>, TEXT("/Engine/Private/HairStrands/HairStrandsVisibilityPS.usf"), TEXT("MainVisibility"), SF_Pixel);
@@ -1282,7 +1294,7 @@ public:
 		FDynamicPassMeshDrawListContext* InDrawListContext);
 
 	virtual void AddMeshBatch(const FMeshBatch& RESTRICT MeshBatch, uint64 BatchElementMask, const FPrimitiveSceneProxy* RESTRICT PrimitiveSceneProxy, int32 StaticMeshId = -1) override final;
-	void AddMeshBatch(const FMeshBatch& RESTRICT MeshBatch, uint64 BatchElementMask, const FPrimitiveSceneProxy* RESTRICT PrimitiveSceneProxy, int32 StaticMeshId, uint32 HairMacroGroupId, uint32 HairMaterialId, bool bCullingEnable);
+	void AddMeshBatch(const FMeshBatch& RESTRICT MeshBatch, uint64 BatchElementMask, const FPrimitiveSceneProxy* RESTRICT PrimitiveSceneProxy, int32 StaticMeshId, uint32 HairMacroGroupId, uint32 HairMaterialId, float HairCoverageScale, bool bCullingEnable);
 
 private:
 	bool TryAddMeshBatch(
@@ -1292,6 +1304,7 @@ private:
 		int32 StaticMeshId,
 		uint32 HairMacroGroupId,
 		uint32 HairMaterialId,
+		float HairCoverageScale,
 		bool bCullingEnable,
 		const FMaterialRenderProxy& MaterialRenderProxy,
 		const FMaterial& Material);
@@ -1307,6 +1320,7 @@ private:
 		const uint32 HairMacroGroupId,
 		const uint32 HairMaterialId,
 		const uint32 LightChannelMask,
+		const float HairCoverageScale,
 		ERasterizerFillMode MeshFillMode,
 		ERasterizerCullMode MeshCullMode);
 
@@ -1316,10 +1330,10 @@ private:
 
 void FHairVisibilityProcessor::AddMeshBatch(const FMeshBatch& RESTRICT MeshBatch, uint64 BatchElementMask, const FPrimitiveSceneProxy* RESTRICT PrimitiveSceneProxy, int32 StaticMeshId)
 {
-	AddMeshBatch(MeshBatch, BatchElementMask, PrimitiveSceneProxy, StaticMeshId, 0, 0, false);
+	AddMeshBatch(MeshBatch, BatchElementMask, PrimitiveSceneProxy, StaticMeshId, 0, 0, 1.f, false);
 }
 
-void FHairVisibilityProcessor::AddMeshBatch(const FMeshBatch& RESTRICT MeshBatch, uint64 BatchElementMask, const FPrimitiveSceneProxy* RESTRICT PrimitiveSceneProxy, int32 StaticMeshId, uint32 HairMacroGroupId, uint32 HairMaterialId, bool bCullingEnable)
+void FHairVisibilityProcessor::AddMeshBatch(const FMeshBatch& RESTRICT MeshBatch, uint64 BatchElementMask, const FPrimitiveSceneProxy* RESTRICT PrimitiveSceneProxy, int32 StaticMeshId, uint32 HairMacroGroupId, uint32 HairMaterialId, float HairCoverageScale, bool bCullingEnable)
 {
 	const FMaterialRenderProxy* MaterialRenderProxy = MeshBatch.MaterialRenderProxy;
 	while (MaterialRenderProxy)
@@ -1327,7 +1341,7 @@ void FHairVisibilityProcessor::AddMeshBatch(const FMeshBatch& RESTRICT MeshBatch
 		const FMaterial* Material = MaterialRenderProxy->GetMaterialNoFallback(FeatureLevel);
 		if (Material)
 		{
-			if (TryAddMeshBatch(MeshBatch, BatchElementMask, PrimitiveSceneProxy, StaticMeshId, HairMacroGroupId, HairMaterialId, bCullingEnable, *MaterialRenderProxy, *Material))
+			if (TryAddMeshBatch(MeshBatch, BatchElementMask, PrimitiveSceneProxy, StaticMeshId, HairMacroGroupId, HairMaterialId, HairCoverageScale, bCullingEnable, *MaterialRenderProxy, *Material))
 			{
 				break;
 			}
@@ -1344,6 +1358,7 @@ bool FHairVisibilityProcessor::TryAddMeshBatch(
 	int32 StaticMeshId,
 	uint32 HairMacroGroupId,
 	uint32 HairMaterialId,
+	float HairCoverageScale,
 	bool bCullingEnable,
 	const FMaterialRenderProxy& MaterialRenderProxy,
 	const FMaterial& Material)
@@ -1365,15 +1380,15 @@ bool FHairVisibilityProcessor::TryAddMeshBatch(
 		const ERasterizerFillMode MeshFillMode = ComputeMeshFillMode(Material, OverrideSettings);
 		const ERasterizerCullMode MeshCullMode = ComputeMeshCullMode(Material, OverrideSettings);
 		if (RenderMode == HairVisibilityRenderMode_MSAA_Visibility && bCullingEnable)
-			return Process<HairVisibilityRenderMode_MSAA_Visibility, true>(MeshBatch, BatchElementMask, PrimitiveSceneProxy, StaticMeshId, MaterialRenderProxy, Material, HairMacroGroupId, HairMaterialId, LightChannelMask, MeshFillMode, MeshCullMode);
+			return Process<HairVisibilityRenderMode_MSAA_Visibility, true>(MeshBatch, BatchElementMask, PrimitiveSceneProxy, StaticMeshId, MaterialRenderProxy, Material, HairMacroGroupId, HairMaterialId, LightChannelMask, HairCoverageScale, MeshFillMode, MeshCullMode);
 		else if (RenderMode == HairVisibilityRenderMode_MSAA_Visibility && !bCullingEnable)
-			return Process<HairVisibilityRenderMode_MSAA_Visibility, false>(MeshBatch, BatchElementMask, PrimitiveSceneProxy, StaticMeshId, MaterialRenderProxy, Material, HairMacroGroupId, HairMaterialId, LightChannelMask, MeshFillMode, MeshCullMode);
+			return Process<HairVisibilityRenderMode_MSAA_Visibility, false>(MeshBatch, BatchElementMask, PrimitiveSceneProxy, StaticMeshId, MaterialRenderProxy, Material, HairMacroGroupId, HairMaterialId, LightChannelMask, HairCoverageScale, MeshFillMode, MeshCullMode);
 		else if (RenderMode == HairVisibilityRenderMode_Transmittance)
-			return Process<HairVisibilityRenderMode_Transmittance>(MeshBatch, BatchElementMask, PrimitiveSceneProxy, StaticMeshId, MaterialRenderProxy, Material, HairMacroGroupId, HairMaterialId, LightChannelMask, MeshFillMode, MeshCullMode);
+			return Process<HairVisibilityRenderMode_Transmittance>(MeshBatch, BatchElementMask, PrimitiveSceneProxy, StaticMeshId, MaterialRenderProxy, Material, HairMacroGroupId, HairMaterialId, LightChannelMask, HairCoverageScale, MeshFillMode, MeshCullMode);
 		else if (RenderMode == HairVisibilityRenderMode_TransmittanceAndHairCount)
-			return Process<HairVisibilityRenderMode_TransmittanceAndHairCount>(MeshBatch, BatchElementMask, PrimitiveSceneProxy, StaticMeshId, MaterialRenderProxy, Material, HairMacroGroupId, HairMaterialId, LightChannelMask, MeshFillMode, MeshCullMode);
+			return Process<HairVisibilityRenderMode_TransmittanceAndHairCount>(MeshBatch, BatchElementMask, PrimitiveSceneProxy, StaticMeshId, MaterialRenderProxy, Material, HairMacroGroupId, HairMaterialId, LightChannelMask, HairCoverageScale, MeshFillMode, MeshCullMode);
 		else if (RenderMode == HairVisibilityRenderMode_PPLL)
-			return Process<HairVisibilityRenderMode_PPLL>(MeshBatch, BatchElementMask, PrimitiveSceneProxy, StaticMeshId, MaterialRenderProxy, Material, HairMacroGroupId, HairMaterialId, LightChannelMask, MeshFillMode, MeshCullMode);
+			return Process<HairVisibilityRenderMode_PPLL>(MeshBatch, BatchElementMask, PrimitiveSceneProxy, StaticMeshId, MaterialRenderProxy, Material, HairMacroGroupId, HairMaterialId, LightChannelMask, HairCoverageScale, MeshFillMode, MeshCullMode);
 	}
 
 	return true;
@@ -1390,6 +1405,7 @@ bool FHairVisibilityProcessor::Process(
 	const uint32 HairMacroGroupId,
 	const uint32 HairMaterialId,
 	const uint32 LightChannelMask,
+	const float HairCoverageScale,
 	ERasterizerFillMode MeshFillMode,
 	ERasterizerCullMode MeshCullMode)
 {
@@ -1416,7 +1432,7 @@ bool FHairVisibilityProcessor::Process(
 	}
 
 	FMeshPassProcessorRenderState DrawRenderState(PassDrawRenderState);
-	FHairVisibilityShaderElementData ShaderElementData(HairMacroGroupId, HairMaterialId, LightChannelMask);
+	FHairVisibilityShaderElementData ShaderElementData(HairMacroGroupId, HairMaterialId, LightChannelMask, HairCoverageScale);
 	ShaderElementData.InitializeMeshMaterialData(ViewIfDynamicMeshCommand, PrimitiveSceneProxy, MeshBatch, StaticMeshId, false);
 
 	BuildMeshDrawCommands(
@@ -2090,7 +2106,8 @@ static void AddHairVisibilityCommonPass(
 					if (const FMeshBatch* MeshBatch = PrimitiveInfo.Mesh)
 					{
 						const uint64 BatchElementMask = ~0ull;
-						MeshProcessor.AddMeshBatch(*MeshBatch, BatchElementMask, PrimitiveInfo.PrimitiveSceneProxy, -1, MacroGroupData.MacroGroupId, PrimitiveInfo.MaterialId, PrimitiveInfo.IsCullingEnable());
+						const float HairCoverageScale = PrimitiveInfo.PublicDataPtr->GetActiveStrandsCoverageScale();
+						MeshProcessor.AddMeshBatch(*MeshBatch, BatchElementMask, PrimitiveInfo.PrimitiveSceneProxy, -1, MacroGroupData.MacroGroupId, PrimitiveInfo.MaterialId, HairCoverageScale, PrimitiveInfo.IsCullingEnable());
 					}
 				}
 			}
