Fixed assert due to missing ParallelRenderingThread tag scope.

#rnx

diff --git a/Engine/Source/Runtime/Renderer/Private/ShadowSetup.cpp b/Engine/Source/Runtime/Renderer/Private/ShadowSetup.cpp
--- a/Engine/Source/Runtime/Renderer/Private/ShadowSetup.cpp
+++ b/Engine/Source/Runtime/Renderer/Private/ShadowSetup.cpp
@@ -4594,6 +4594,8 @@ struct FGatherShadowPrimitivesPrepareTask
 			FGraphEventRef FinalizeTask = FFunctionGraphTask::CreateAndDispatchWhenReady([TaskData = &TaskData]()
 				{
 					SCOPED_NAMED_EVENT_TEXT("FGatherShadowPrimitivesFinalizeTask", FColor::Green);
+					FOptionalTaskTagScope Scope(ETaskTag::EParallelRenderingThread);
+
 					for (FGatherShadowPrimitivesPacket* Packet : TaskData->Packets)
 					{
 						Packet->AnyThreadFinalize(*TaskData);
