Remove HeterogeneousVolumes setup cost when no volumes are to be rendered.

- Split VolumetricMeshBatches into Volumetric Fog and Heterogeneous variants
- Do not perform invoke setup if a view has no heterogeneous volumes

diff --git a/Engine/Source/Runtime/Renderer/Private/HeterogeneousVolumes/HeterogeneousVolumes.cpp b/Engine/Source/Runtime/Renderer/Private/HeterogeneousVolumes/HeterogeneousVolumes.cpp
--- a/Engine/Source/Runtime/Renderer/Private/HeterogeneousVolumes/HeterogeneousVolumes.cpp
+++ b/Engine/Source/Runtime/Renderer/Private/HeterogeneousVolumes/HeterogeneousVolumes.cpp
@@ -177,10 +177,11 @@ bool ShouldRenderHeterogeneousVolumes(
 }
 
 bool ShouldRenderHeterogeneousVolumesForView(
-	const FSceneView& View
+	const FViewInfo& View
 )
 {
 	return IsHeterogeneousVolumesEnabled()
+		&& !View.HeterogeneousVolumesMeshBatches.IsEmpty()
 		&& View.Family
 		&& !View.bIsReflectionCapture;
 }
@@ -207,13 +208,21 @@ bool DoesMaterialShaderSupportHeterogeneousVolumes(const FMaterial& Material)
 		&& Material.IsUsedWithNiagaraMeshParticles();
 }
 
-bool ShouldRenderPrimitiveWithHeterogeneousVolumes(
-	const FPrimitiveSceneProxy* PrimitiveSceneProxy,
-	const FMaterial& Material
+bool ShouldRenderMeshBatchWithHeterogeneousVolumes(
+	const FMeshBatch* Mesh,
+	const FPrimitiveSceneProxy* Proxy,
+	ERHIFeatureLevel::Type FeatureLevel
 )
 {
-	check(PrimitiveSceneProxy);
-	return PrimitiveSceneProxy->IsHeterogeneousVolume() && DoesMaterialShaderSupportHeterogeneousVolumes(Material);
+	check(Mesh);
+	check(Proxy);
+	check(Mesh->MaterialRenderProxy);
+
+	const FMaterialRenderProxy* MaterialRenderProxy = Mesh->MaterialRenderProxy;
+	const FMaterial& Material = MaterialRenderProxy->GetMaterialWithFallback(FeatureLevel, MaterialRenderProxy);
+	return IsHeterogeneousVolumesEnabled()
+		&& Proxy->IsHeterogeneousVolume()
+		&& DoesMaterialShaderSupportHeterogeneousVolumes(Material);
 }
 
 namespace HeterogeneousVolumes
@@ -359,26 +368,24 @@ void FDeferredShadingSceneRenderer::RenderHeterogeneousVolumes(
 	{
 		FViewInfo& View = Views[ViewIndex];
 
-		// Per-view??
-		FRDGTextureDesc Desc = SceneTextures.Color.Target->Desc;
-		Desc.Format = PF_FloatRGBA;
-		Desc.Flags &= ~(TexCreate_FastVRAM);
-		FRDGTextureRef HeterogeneousVolumeRadiance = GraphBuilder.CreateTexture(Desc, TEXT("HeterogeneousVolumes"));
-		AddClearUAVPass(GraphBuilder, GraphBuilder.CreateUAV(HeterogeneousVolumeRadiance), FLinearColor::Transparent);
-
 		if (ShouldRenderHeterogeneousVolumesForView(View))
 		{
-			for (int32 MeshBatchIndex = 0; MeshBatchIndex < View.VolumetricMeshBatches.Num(); ++MeshBatchIndex)
+			FRDGTextureDesc Desc = SceneTextures.Color.Target->Desc;
+			Desc.Format = PF_FloatRGBA;
+			Desc.Flags &= ~(TexCreate_FastVRAM);
+			FRDGTextureRef HeterogeneousVolumeRadiance = GraphBuilder.CreateTexture(Desc, TEXT("HeterogeneousVolumes"));
+			AddClearUAVPass(GraphBuilder, GraphBuilder.CreateUAV(HeterogeneousVolumeRadiance), FLinearColor::Transparent);
+
+			for (int32 MeshBatchIndex = 0; MeshBatchIndex < View.HeterogeneousVolumesMeshBatches.Num(); ++MeshBatchIndex)
 			{
-				const FMeshBatch* Mesh = View.VolumetricMeshBatches[MeshBatchIndex].Mesh;
-				const FMaterialRenderProxy* MaterialRenderProxy = Mesh->MaterialRenderProxy;
-				const FMaterial& Material = MaterialRenderProxy->GetMaterialWithFallback(View.GetFeatureLevel(), MaterialRenderProxy);
-				const FPrimitiveSceneProxy* PrimitiveSceneProxy = View.VolumetricMeshBatches[MeshBatchIndex].Proxy;
-				if (!ShouldRenderPrimitiveWithHeterogeneousVolumes(PrimitiveSceneProxy, Material))
+				const FMeshBatch* Mesh = View.HeterogeneousVolumesMeshBatches[MeshBatchIndex].Mesh;
+				const FPrimitiveSceneProxy* PrimitiveSceneProxy = View.HeterogeneousVolumesMeshBatches[MeshBatchIndex].Proxy;
+				if (!ShouldRenderMeshBatchWithHeterogeneousVolumes(Mesh, PrimitiveSceneProxy, View.GetFeatureLevel()))
 				{
 					continue;
 				}
 
+				const FMaterialRenderProxy* MaterialRenderProxy = Mesh->MaterialRenderProxy;
 				const FPrimitiveSceneInfo* PrimitiveSceneInfo = PrimitiveSceneProxy->GetPrimitiveSceneInfo();
 				const int32 PrimitiveId = PrimitiveSceneInfo->GetIndex();
 				const FBoxSphereBounds LocalBoxSphereBounds = PrimitiveSceneProxy->GetLocalBounds();
@@ -445,9 +452,9 @@ void FDeferredShadingSceneRenderer::RenderHeterogeneousVolumes(
 					);
 				}
 			}
-		}
 
-		View.HeterogeneousVolumeRadiance = HeterogeneousVolumeRadiance;
+			View.HeterogeneousVolumeRadiance = HeterogeneousVolumeRadiance;
+		}
 	}
 }
 
diff --git a/Engine/Source/Runtime/Renderer/Private/HeterogeneousVolumes/HeterogeneousVolumes.h b/Engine/Source/Runtime/Renderer/Private/HeterogeneousVolumes/HeterogeneousVolumes.h
--- a/Engine/Source/Runtime/Renderer/Private/HeterogeneousVolumes/HeterogeneousVolumes.h
+++ b/Engine/Source/Runtime/Renderer/Private/HeterogeneousVolumes/HeterogeneousVolumes.h
@@ -25,12 +25,13 @@ struct FSceneTextures;
 //
 
 bool ShouldRenderHeterogeneousVolumes(const FScene* Scene);
-bool ShouldRenderHeterogeneousVolumesForView(const FSceneView& View);
+bool ShouldRenderHeterogeneousVolumesForView(const FViewInfo& View);
 bool DoesPlatformSupportHeterogeneousVolumes(EShaderPlatform Platform);
 bool DoesMaterialShaderSupportHeterogeneousVolumes(const FMaterialShaderParameters& Parameters);
-bool ShouldRenderPrimitiveWithHeterogeneousVolumes(
-	const FPrimitiveSceneProxy* PrimitiveSceneProxy,
-	const FMaterial& Material
+bool ShouldRenderMeshBatchWithHeterogeneousVolumes(
+	const FMeshBatch* Mesh,
+	const FPrimitiveSceneProxy* Proxy,
+	ERHIFeatureLevel::Type FeatureLevel
 );
 
 //
diff --git a/Engine/Source/Runtime/Renderer/Private/SceneRendering.h b/Engine/Source/Runtime/Renderer/Private/SceneRendering.h
--- a/Engine/Source/Runtime/Renderer/Private/SceneRendering.h
+++ b/Engine/Source/Runtime/Renderer/Private/SceneRendering.h
@@ -1342,6 +1342,9 @@ public:
 	/** Mesh batches with a volumetric material. */
 	TArray<FVolumetricMeshBatch, SceneRenderingAllocator> VolumetricMeshBatches;
 
+	/** Mesh batches for heterogeneous volumes rendering. */
+	TArray<FVolumetricMeshBatch, SceneRenderingAllocator> HeterogeneousVolumesMeshBatches;
+
 	/** Mesh batches with a sky material. */
 	TArray<FSkyMeshBatch, SceneRenderingAllocator> SkyMeshBatches;
 
diff --git a/Engine/Source/Runtime/Renderer/Private/SceneVisibility.cpp b/Engine/Source/Runtime/Renderer/Private/SceneVisibility.cpp
--- a/Engine/Source/Runtime/Renderer/Private/SceneVisibility.cpp
+++ b/Engine/Source/Runtime/Renderer/Private/SceneVisibility.cpp
@@ -49,6 +49,7 @@
 #include "InstanceCulling/InstanceCullingManager.h"
 #include "PostProcess/TemporalAA.h"
 #include "RayTracing/RayTracingInstanceCulling.h"
+#include "HeterogeneousVolumes/HeterogeneousVolumes.h"
 #include "RendererModule.h"
 #include "SceneViewExtension.h"
 #include "RenderCore.h"
@@ -2301,6 +2302,7 @@ struct FRelevancePacket : public FSceneRenderingAllocatorObject<FRelevancePacket
 
 	TArray<FMeshDecalBatch> MeshDecalBatches;
 	TArray<FVolumetricMeshBatch> VolumetricMeshBatches;
+	TArray<FVolumetricMeshBatch> HeterogeneousVolumesMeshBatches;
 	TArray<FSkyMeshBatch> SkyMeshBatches;
 	TArray<FSortedTrianglesMeshBatch> SortedTrianglesMeshBatches;
 	FDrawCommandRelevancePacket DrawCommandPacket;
@@ -2898,10 +2900,20 @@ struct FRelevancePacket : public FSceneRenderingAllocatorObject<FRelevancePacket
 
 						if (ViewRelevance.bHasVolumeMaterialDomain)
 						{
-							VolumetricMeshBatches.AddUninitialized(1);
-							FVolumetricMeshBatch& BatchAndProxy = VolumetricMeshBatches.Last();
-							BatchAndProxy.Mesh = &StaticMesh;
-							BatchAndProxy.Proxy = PrimitiveSceneInfo->Proxy;
+							if (ShouldRenderMeshBatchWithHeterogeneousVolumes(&StaticMesh, PrimitiveSceneInfo->Proxy, View.FeatureLevel))
+							{
+								HeterogeneousVolumesMeshBatches.AddUninitialized(1);
+								FVolumetricMeshBatch& BatchAndProxy = HeterogeneousVolumesMeshBatches.Last();
+								BatchAndProxy.Mesh = &StaticMesh;
+								BatchAndProxy.Proxy = PrimitiveSceneInfo->Proxy;
+							}
+							else
+							{
+								VolumetricMeshBatches.AddUninitialized(1);
+								FVolumetricMeshBatch& BatchAndProxy = VolumetricMeshBatches.Last();
+								BatchAndProxy.Mesh = &StaticMesh;
+								BatchAndProxy.Proxy = PrimitiveSceneInfo->Proxy;
+							}
 						}
 
 						if (ViewRelevance.bUsesSkyMaterial)
@@ -3012,6 +3024,7 @@ struct FRelevancePacket : public FSceneRenderingAllocatorObject<FRelevancePacket
 
 		WriteView.MeshDecalBatches.Append(MeshDecalBatches);
 		WriteView.VolumetricMeshBatches.Append(VolumetricMeshBatches);
+		WriteView.HeterogeneousVolumesMeshBatches.Append(HeterogeneousVolumesMeshBatches);
 		WriteView.SkyMeshBatches.Append(SkyMeshBatches);
 		WriteView.SortedTrianglesMeshBatches.Append(SortedTrianglesMeshBatches);
 
@@ -3399,10 +3412,20 @@ void ComputeDynamicMeshRelevance(EShadingPath ShadingPath, bool bAddLightmapDens
 
 	if (ViewRelevance.bHasVolumeMaterialDomain)
 	{
-		View.VolumetricMeshBatches.AddUninitialized(1);
-		FVolumetricMeshBatch& BatchAndProxy = View.VolumetricMeshBatches.Last();
-		BatchAndProxy.Mesh = MeshBatch.Mesh;
-		BatchAndProxy.Proxy = MeshBatch.PrimitiveSceneProxy;
+		if (ShouldRenderMeshBatchWithHeterogeneousVolumes(MeshBatch.Mesh, MeshBatch.PrimitiveSceneProxy, View.FeatureLevel))
+		{
+			View.HeterogeneousVolumesMeshBatches.AddUninitialized(1);
+			FVolumetricMeshBatch& BatchAndProxy = View.HeterogeneousVolumesMeshBatches.Last();
+			BatchAndProxy.Mesh = MeshBatch.Mesh;
+			BatchAndProxy.Proxy = MeshBatch.PrimitiveSceneProxy;
+		}
+		else
+		{
+			View.VolumetricMeshBatches.AddUninitialized(1);
+			FVolumetricMeshBatch& BatchAndProxy = View.VolumetricMeshBatches.Last();
+			BatchAndProxy.Mesh = MeshBatch.Mesh;
+			BatchAndProxy.Proxy = MeshBatch.PrimitiveSceneProxy;
+		}
 	}
 
 	if (ViewRelevance.bUsesSkyMaterial)
diff --git a/Engine/Source/Runtime/Renderer/Private/VolumetricFogVoxelization.cpp b/Engine/Source/Runtime/Renderer/Private/VolumetricFogVoxelization.cpp
--- a/Engine/Source/Runtime/Renderer/Private/VolumetricFogVoxelization.cpp
+++ b/Engine/Source/Runtime/Renderer/Private/VolumetricFogVoxelization.cpp
@@ -672,13 +672,10 @@ void FDeferredShadingSceneRenderer::VoxelizeFogVolumePrimitives(
 
 					for (int32 MeshBatchIndex = 0; MeshBatchIndex < View.VolumetricMeshBatches.Num(); ++MeshBatchIndex)
 					{
-						const FMeshBatch* Mesh = View.VolumetricMeshBatches[MeshBatchIndex].Mesh;
-						const FMaterialRenderProxy* MaterialRenderProxy = Mesh->MaterialRenderProxy;
-						const FMaterial& Material = MaterialRenderProxy->GetMaterialWithFallback(View.GetFeatureLevel(), MaterialRenderProxy);
-
 						// Skip volumes flagged as rendered with HeterogenousVolumes
+						const FMeshBatch* Mesh = View.VolumetricMeshBatches[MeshBatchIndex].Mesh;
 						const FPrimitiveSceneProxy* PrimitiveSceneProxy = View.VolumetricMeshBatches[MeshBatchIndex].Proxy;
-						if (bShouldRenderHeterogeneousVolumes && ShouldRenderPrimitiveWithHeterogeneousVolumes(PrimitiveSceneProxy, Material))
+						if (ShouldRenderMeshBatchWithHeterogeneousVolumes(Mesh, PrimitiveSceneProxy, View.GetFeatureLevel()))
 						{
 							continue;
 						}
