Fixing NonUnity builds.

diff --git a/Engine/Plugins/Online/OnlineSubsystem/Source/Public/Interfaces/OnlineGameActivityInterface.h b/Engine/Plugins/Online/OnlineSubsystem/Source/Public/Interfaces/OnlineGameActivityInterface.h
--- a/Engine/Plugins/Online/OnlineSubsystem/Source/Public/Interfaces/OnlineGameActivityInterface.h
+++ b/Engine/Plugins/Online/OnlineSubsystem/Source/Public/Interfaces/OnlineGameActivityInterface.h
@@ -4,6 +4,7 @@
 
 #include "Math/Vector.h"
 #include "OnlineDelegateMacros.h"
+#include "Online/CoreOnlineFwd.h"
 
 class FOnlineEventParms;
 template <typename OptionalType> struct TOptional;
diff --git a/Engine/Source/Runtime/Renderer/Private/RayTracing/RayTracingDecals.h b/Engine/Source/Runtime/Renderer/Private/RayTracing/RayTracingDecals.h
--- a/Engine/Source/Runtime/Renderer/Private/RayTracing/RayTracingDecals.h
+++ b/Engine/Source/Runtime/Renderer/Private/RayTracing/RayTracingDecals.h
@@ -13,6 +13,8 @@ class FShaderType;
 class FRHIRayTracingShader;
 class FGlobalShaderMap;
 
+enum EBlendMode : int;
+
 BEGIN_GLOBAL_SHADER_PARAMETER_STRUCT(FRayTracingDecals, RENDERER_API)
 	SHADER_PARAMETER_RDG_TEXTURE(Texture2D<uint>, Grid)
 	SHADER_PARAMETER_RDG_BUFFER_SRV(Buffer<uint>, GridData)
