Removing more dependencies of Renderer Private from Niagara.

diff --git a/Engine/Plugins/FX/Niagara/Source/Niagara/Private/DataInterface/NiagaraDataInterfaceGBuffer.cpp b/Engine/Plugins/FX/Niagara/Source/Niagara/Private/DataInterface/NiagaraDataInterfaceGBuffer.cpp
--- a/Engine/Plugins/FX/Niagara/Source/Niagara/Private/DataInterface/NiagaraDataInterfaceGBuffer.cpp
+++ b/Engine/Plugins/FX/Niagara/Source/Niagara/Private/DataInterface/NiagaraDataInterfaceGBuffer.cpp
@@ -10,8 +10,8 @@
 
 #include "Internationalization/Internationalization.h"
 #include "ShaderParameterUtils.h"
-#include "SceneRendering.h"
 #include "SceneTextures.h"
+#include "FXRenderingUtils.h"
 
 #include UE_INLINE_GENERATED_CPP_BY_NAME(NiagaraDataInterfaceGBuffer)
 
@@ -180,13 +180,10 @@ void UNiagaraDataInterfaceGBuffer::SetShaderParameters(const FNiagaraDataInterfa
 	NiagaraDataInterfaceGBufferLocal::FShaderParameters* Parameters = Context.GetParameterNestedStruct<NiagaraDataInterfaceGBufferLocal::FShaderParameters>();
 	if (Context.IsResourceBound(&Parameters->VelocityTexture))
 	{
-		TConstArrayView<FViewInfo> ViewInfos = Context.GetComputeDispatchInterface().GetSimulationViewInfos();
-		if (ViewInfos.Num() > 0)
+		TConstStridedView<FSceneView> SimulationSceneViews = Context.GetComputeDispatchInterface().GetSimulationSceneViews();
+		if (SimulationSceneViews.Num() > 0)
 		{
-			if ( const FSceneTextures* SceneTextures = GetViewFamilyInfo(ViewInfos).GetSceneTexturesChecked() )
-			{
-				VelocityTexture = SceneTextures->Velocity;
-			}
+			VelocityTexture = UE::FXRenderingUtils::GetSceneVelocityTexture(SimulationSceneViews[0]);
 		}
 		
 		if (VelocityTexture == nullptr)
diff --git a/Engine/Plugins/FX/Niagara/Source/Niagara/Private/NiagaraDataInterfaceCollisionQuery.cpp b/Engine/Plugins/FX/Niagara/Source/Niagara/Private/NiagaraDataInterfaceCollisionQuery.cpp
--- a/Engine/Plugins/FX/Niagara/Source/Niagara/Private/NiagaraDataInterfaceCollisionQuery.cpp
+++ b/Engine/Plugins/FX/Niagara/Source/Niagara/Private/NiagaraDataInterfaceCollisionQuery.cpp
@@ -2,6 +2,7 @@
 
 #include "NiagaraDataInterfaceCollisionQuery.h"
 
+#include "FXRenderingUtils.h"
 #include "GlobalDistanceFieldParameters.h"
 #include "NiagaraAsyncGpuTraceHelper.h"
 #include "NiagaraDistanceFieldHelper.h"
@@ -12,7 +13,6 @@
 #include "NiagaraStats.h"
 #include "NiagaraTypes.h"
 #include "NiagaraWorldManager.h"
-#include "SceneRendering.h"
 #include "Shader.h"
 #include "ShaderCore.h"
 #include "ShaderCompilerCore.h"
@@ -546,9 +546,9 @@ void UNiagaraDataInterfaceCollisionQuery::SetShaderParameters(const FNiagaraData
 	FGlobalDistanceFieldParameters2* ShaderGDFParameters = Context.GetParameterIncludedStruct<FGlobalDistanceFieldParameters2>();
 	if (Context.IsStructBound(ShaderGDFParameters))
 	{
-		TConstArrayView<FViewInfo> SimulationViewInfos = Context.GetComputeDispatchInterface().GetSimulationViewInfos();
+		TConstStridedView<FSceneView> SimulationSceneViews = Context.GetComputeDispatchInterface().GetSimulationSceneViews();
 		FNiagaraDistanceFieldHelper::SetGlobalDistanceFieldParameters(
-			SimulationViewInfos.Num() > 0 ? &SimulationViewInfos[0].GlobalDistanceFieldInfo.ParameterData : nullptr,
+			UE::FXRenderingUtils::GetGlobalDistanceFieldParameterData(SimulationSceneViews),
 			*ShaderGDFParameters
 		);
 	}
diff --git a/Engine/Plugins/FX/Niagara/Source/Niagara/Private/NiagaraDataInterfaceRigidMeshCollisionQuery.cpp b/Engine/Plugins/FX/Niagara/Source/Niagara/Private/NiagaraDataInterfaceRigidMeshCollisionQuery.cpp
--- a/Engine/Plugins/FX/Niagara/Source/Niagara/Private/NiagaraDataInterfaceRigidMeshCollisionQuery.cpp
+++ b/Engine/Plugins/FX/Niagara/Source/Niagara/Private/NiagaraDataInterfaceRigidMeshCollisionQuery.cpp
@@ -28,7 +28,8 @@
 #include "NiagaraSystemInstance.h"
 #include "NiagaraSystemSimulation.h"
 #include "PhysicsEngine/PhysicsAsset.h"
-#include "ScenePrivate.h"
+#include "PrimitiveSceneInfo.h"
+#include "SceneInterface.h"
 #include "ShaderParameterUtils.h"
 
 #include UE_INLINE_GENERATED_CPP_BY_NAME(NiagaraDataInterfaceRigidMeshCollisionQuery)
@@ -936,7 +937,7 @@ struct FNDIRigidMeshCollisionProxy : public FNiagaraDataInterfaceProxy
 					const uint32 BufferBytes = sizeof(uint32) * ElementCount;
 					void* BufferData = RHILockBuffer(ProxyData->AssetBuffer->DFIndexBuffer.Buffer, 0, BufferBytes, RLM_WriteOnly);
 
-					const FSceneInterface* Scene = Context.GetComputeDispatchInterface().GetScene();
+					const FSceneInterface* Scene = Context.GetComputeDispatchInterface().GetSceneInterface();
 					if (Scene && !ProxyData->UniqueComponentIds.IsEmpty())
 					{
 						TArray<uint32> UniqueDistanceFieldIndices;
diff --git a/Engine/Plugins/FX/Niagara/Source/Niagara/Private/NiagaraGpuComputeDispatch.cpp b/Engine/Plugins/FX/Niagara/Source/Niagara/Private/NiagaraGpuComputeDispatch.cpp
--- a/Engine/Plugins/FX/Niagara/Source/Niagara/Private/NiagaraGpuComputeDispatch.cpp
+++ b/Engine/Plugins/FX/Niagara/Source/Niagara/Private/NiagaraGpuComputeDispatch.cpp
@@ -32,10 +32,12 @@
 #include "PipelineStateCache.h"
 #include "RHI.h"
 #include "SceneInterface.h"
-#include "SceneRendering.h"
 #include "SceneRenderTargetParameters.h"
 #include "ShaderParameterUtils.h"
 #include "TextureResource.h"
+#include "FXRenderingUtils.h"
+
+#include "ScenePrivate.h"
 
 DECLARE_CYCLE_STAT(TEXT("GPU Dispatch Setup [RT]"), STAT_NiagaraGPUDispatchSetup_RT, STATGROUP_Niagara);
 DECLARE_CYCLE_STAT(TEXT("GPU Emitter Dispatch [RT]"), STAT_NiagaraGPUSimTick_RT, STATGROUP_Niagara);
@@ -1047,24 +1049,13 @@ void FNiagaraGpuComputeDispatch::ExecuteTicks(FRDGBuilder& GraphBuilder, TConstA
 	SimulationSceneViews = MakeStridedViewOfBase<const FSceneView>(Views);
 	SimulationViewInfos = Views;
 
+	if (FSceneInterface::GetShadingPath(FeatureLevel) == EShadingPath::Deferred)
 	{
-		const FSceneTextures* SceneTextures = Views.Num() > 0 ? GetViewFamilyInfo(Views).GetSceneTexturesChecked() : nullptr;
-		if (FSceneInterface::GetShadingPath(FeatureLevel) == EShadingPath::Deferred)
-		{
-			SceneTexturesUniformParams = SceneTextures ? SceneTextures->UniformBuffer : nullptr;
-			if (SceneTexturesUniformParams == nullptr)
-			{
-				SceneTexturesUniformParams = CreateSceneTextureUniformBuffer(GraphBuilder, nullptr, FeatureLevel, ESceneTextureSetupMode::SceneVelocity);
-			}
-		}
-		else if (FSceneInterface::GetShadingPath(FeatureLevel) == EShadingPath::Mobile)
-		{
-			MobileSceneTexturesUniformParams = SceneTextures ? SceneTextures->MobileUniformBuffer : nullptr;
-			if (MobileSceneTexturesUniformParams == nullptr)
-			{
-				MobileSceneTexturesUniformParams = CreateMobileSceneTextureUniformBuffer(GraphBuilder, SceneTextures, EMobileSceneTextureSetupMode::None);
-			}
-		}
+		SceneTexturesUniformParams = UE::FXRenderingUtils::GetOrCreateSceneTextureUniformBuffer(GraphBuilder, SimulationSceneViews, FeatureLevel, ESceneTextureSetupMode::SceneVelocity);
+	}
+	else if (FSceneInterface::GetShadingPath(FeatureLevel) == EShadingPath::Mobile)
+	{
+		MobileSceneTexturesUniformParams = UE::FXRenderingUtils::GetOrCreateMobileSceneTextureUniformBuffer(GraphBuilder, SimulationSceneViews, EMobileSceneTextureSetupMode::None);
 	}
 
 	// Loop over dispatches
diff --git a/Engine/Plugins/FX/Niagara/Source/NiagaraShader/Private/NiagaraAsyncGpuTraceProvider.cpp b/Engine/Plugins/FX/Niagara/Source/NiagaraShader/Private/NiagaraAsyncGpuTraceProvider.cpp
--- a/Engine/Plugins/FX/Niagara/Source/NiagaraShader/Private/NiagaraAsyncGpuTraceProvider.cpp
+++ b/Engine/Plugins/FX/Niagara/Source/NiagaraShader/Private/NiagaraAsyncGpuTraceProvider.cpp
@@ -7,7 +7,6 @@
 #include "NiagaraAsyncGpuTraceProviderHwrt.h"
 #include "NiagaraSettings.h"
 #include "NiagaraShaderParticleID.h"
-#include "ScenePrivate.h"
 
 #include "PrimitiveSceneInfo.h"
 #include "RenderGraphUtils.h"
diff --git a/Engine/Plugins/FX/Niagara/Source/NiagaraShader/Private/NiagaraAsyncGpuTraceProviderGsdf.cpp b/Engine/Plugins/FX/Niagara/Source/NiagaraShader/Private/NiagaraAsyncGpuTraceProviderGsdf.cpp
--- a/Engine/Plugins/FX/Niagara/Source/NiagaraShader/Private/NiagaraAsyncGpuTraceProviderGsdf.cpp
+++ b/Engine/Plugins/FX/Niagara/Source/NiagaraShader/Private/NiagaraAsyncGpuTraceProviderGsdf.cpp
@@ -3,6 +3,7 @@
 #include "NiagaraAsyncGpuTraceProviderGsdf.h"
 
 #include "DataDrivenShaderPlatformInfo.h"
+#include "FXRenderingUtils.h"
 #include "GlobalShader.h"
 #include "NiagaraDistanceFieldHelper.h"
 #include "NiagaraGpuComputeDispatchInterface.h"
@@ -109,9 +110,15 @@ bool FNiagaraAsyncGpuTraceProviderGsdf::IsAvailable() const
 
 void FNiagaraAsyncGpuTraceProviderGsdf::PostRenderOpaque(FRHICommandList& RHICmdList, TConstArrayView<FViewInfo> Views, FCollisionGroupHashMap* CollisionGroupHash)
 {
-	const FViewInfo& ReferenceView = Views[0];
-	m_DistanceFieldData = ReferenceView.GlobalDistanceFieldInfo.ParameterData;
-	m_ViewUniformBuffer = ReferenceView.ViewUniformBuffer;
+	if (const FGlobalDistanceFieldParameterData* DistanceFieldData = UE::FXRenderingUtils::GetGlobalDistanceFieldParameterData(MakeStridedViewOfBase<const FSceneView>(Views)))
+	{
+		m_DistanceFieldData = *DistanceFieldData;
+	}
+	else
+	{
+		m_DistanceFieldData = {};
+	}
+	m_ViewUniformBuffer = Views[0].ViewUniformBuffer;
 }
 
 void FNiagaraAsyncGpuTraceProviderGsdf::IssueTraces(FRHICommandList& RHICmdList, const FDispatchRequest& Request, FCollisionGroupHashMap* CollisionHashMap)
diff --git a/Engine/Plugins/FX/Niagara/Source/NiagaraShader/Private/NiagaraAsyncGpuTraceProviderHwrt.cpp b/Engine/Plugins/FX/Niagara/Source/NiagaraShader/Private/NiagaraAsyncGpuTraceProviderHwrt.cpp
--- a/Engine/Plugins/FX/Niagara/Source/NiagaraShader/Private/NiagaraAsyncGpuTraceProviderHwrt.cpp
+++ b/Engine/Plugins/FX/Niagara/Source/NiagaraShader/Private/NiagaraAsyncGpuTraceProviderHwrt.cpp
@@ -14,8 +14,6 @@
 #include "ShaderParameterMacros.h"
 #include "ShaderParameterStruct.h"
 #include "Templates/Function.h"
-
-#include "ScenePrivate.h"
 #include "RayTracingPayloadType.h"
 
 static int GNiagaraAsyncGpuTraceHwrtEnabled = 1;
diff --git a/Engine/Source/Runtime/Renderer/Private/FXRenderingUtils.cpp b/Engine/Source/Runtime/Renderer/Private/FXRenderingUtils.cpp
--- a/Engine/Source/Runtime/Renderer/Private/FXRenderingUtils.cpp
+++ b/Engine/Source/Runtime/Renderer/Private/FXRenderingUtils.cpp
@@ -36,6 +36,61 @@ bool UE::FXRenderingUtils::CanMaterialRenderBeforeFXPostOpaque(
 	return false;
 }
 
+const FGlobalDistanceFieldParameterData* UE::FXRenderingUtils::GetGlobalDistanceFieldParameterData(TConstStridedView<FSceneView> Views)
+{
+	return Views.Num() > 0 ? &static_cast<const FViewInfo&>(Views[0]).GlobalDistanceFieldInfo.ParameterData : nullptr;
+}
+
+FRDGTextureRef UE::FXRenderingUtils::GetSceneVelocityTexture(const FSceneView& View)
+{
+	const FViewFamilyInfo* ViewFamily = static_cast<const FViewFamilyInfo*>(View.Family);
+	const FSceneTextures* SceneTextures = ViewFamily ? ViewFamily->GetSceneTexturesChecked() : nullptr;
+	return SceneTextures ? SceneTextures->Velocity : nullptr;
+}
+
+TRDGUniformBufferRef<FSceneTextureUniformParameters> UE::FXRenderingUtils::GetOrCreateSceneTextureUniformBuffer(
+	FRDGBuilder& GraphBuilder,
+	TConstStridedView<FSceneView> Views,
+	ERHIFeatureLevel::Type FeatureLevel,
+	ESceneTextureSetupMode SetupMode)
+{
+	TRDGUniformBufferRef<FSceneTextureUniformParameters> SceneTexturesUniformParams = nullptr;
+
+	const FViewInfo* View = Views.Num() > 0 ? static_cast<const FViewInfo*>(&Views[0]) : nullptr;
+	if (const FSceneTextures* SceneTextures = View ? static_cast<const FViewFamilyInfo*>(View->Family)->GetSceneTexturesChecked() : nullptr)
+	{
+		SceneTexturesUniformParams = SceneTextures->UniformBuffer;
+	}
+
+	if (SceneTexturesUniformParams == nullptr)
+	{
+		SceneTexturesUniformParams = CreateSceneTextureUniformBuffer(GraphBuilder, nullptr, FeatureLevel, SetupMode);
+	}
+
+	return SceneTexturesUniformParams;
+}
+
+TRDGUniformBufferRef<FMobileSceneTextureUniformParameters> UE::FXRenderingUtils::GetOrCreateMobileSceneTextureUniformBuffer(
+	FRDGBuilder& GraphBuilder,
+	TConstStridedView<FSceneView> Views,
+	EMobileSceneTextureSetupMode SetupMode)
+{
+	TRDGUniformBufferRef<FMobileSceneTextureUniformParameters> MobileSceneTexturesUniformParams = nullptr;
+
+	const FViewInfo* View = Views.Num() > 0 ? static_cast<const FViewInfo*>(&Views[0]) : nullptr;
+	if (const FSceneTextures* SceneTextures = View ? static_cast<const FViewFamilyInfo*>(View->Family)->GetSceneTexturesChecked() : nullptr)
+	{
+		MobileSceneTexturesUniformParams = SceneTextures->MobileUniformBuffer;
+	}
+
+	if (MobileSceneTexturesUniformParams == nullptr)
+	{
+		MobileSceneTexturesUniformParams = CreateMobileSceneTextureUniformBuffer(GraphBuilder, nullptr, SetupMode);
+	}
+
+	return MobileSceneTexturesUniformParams;
+}
+
 const FShaderParametersMetadata* UE::FXRenderingUtils::DistanceFields::GetObjectBufferParametersMetadata()
 {
 	return TShaderParameterStructTypeInfo<FDistanceFieldObjectBufferParameters>::GetStructMetadata();
diff --git a/Engine/Source/Runtime/Renderer/Public/FXRenderingUtils.h b/Engine/Source/Runtime/Renderer/Public/FXRenderingUtils.h
--- a/Engine/Source/Runtime/Renderer/Public/FXRenderingUtils.h
+++ b/Engine/Source/Runtime/Renderer/Public/FXRenderingUtils.h
@@ -3,7 +3,9 @@
 #pragma once
 
 #include "HAL/Platform.h"
+#include "RenderGraphFwd.h"
 #include "RHIDefinitions.h"
+#include "SceneRenderTargetParameters.h"
 #include "UObject/ObjectMacros.h"
 
 #if UE_ENABLE_INCLUDE_ORDER_DEPRECATED_IN_5_3
@@ -28,6 +30,20 @@ namespace UE::FXRenderingUtils
 {
 	RENDERER_API bool CanMaterialRenderBeforeFXPostOpaque(const FSceneViewFamily& ViewFamily, const FPrimitiveSceneProxy& SceneProxy, const FMaterial& Material);
 
+	RENDERER_API const FGlobalDistanceFieldParameterData* GetGlobalDistanceFieldParameterData(TConstStridedView<FSceneView> Views);
+	RENDERER_API FRDGTextureRef GetSceneVelocityTexture(const FSceneView& View);
+
+	RENDERER_API TRDGUniformBufferRef<FSceneTextureUniformParameters> GetOrCreateSceneTextureUniformBuffer(
+		FRDGBuilder& GraphBuilder,
+		TConstStridedView<FSceneView> Views,
+		ERHIFeatureLevel::Type FeatureLevel,
+		ESceneTextureSetupMode SetupMode = ESceneTextureSetupMode::All);
+
+	RENDERER_API TRDGUniformBufferRef<FMobileSceneTextureUniformParameters> GetOrCreateMobileSceneTextureUniformBuffer(
+		FRDGBuilder& GraphBuilder,
+		TConstStridedView<FSceneView> Views,
+		EMobileSceneTextureSetupMode SetupMode = EMobileSceneTextureSetupMode::All);
+
 	namespace DistanceFields
 	{
 		RENDERER_API const FShaderParametersMetadata* GetObjectBufferParametersMetadata();
