HeterogeneousVolumes: Clamp lighting cache downsample factor.

diff --git a/Engine/Source/Runtime/Renderer/Private/HeterogeneousVolumes/HeterogeneousVolumes.cpp b/Engine/Source/Runtime/Renderer/Private/HeterogeneousVolumes/HeterogeneousVolumes.cpp
--- a/Engine/Source/Runtime/Renderer/Private/HeterogeneousVolumes/HeterogeneousVolumes.cpp
+++ b/Engine/Source/Runtime/Renderer/Private/HeterogeneousVolumes/HeterogeneousVolumes.cpp
@@ -348,7 +348,7 @@ namespace HeterogeneousVolumes
 
 	FIntVector GetLightingCacheResolution()
 	{
-		float DownsampleFactor = CVarHeterogeneousVolumesLightingCacheDownsampleFactor.GetValueOnRenderThread();
+		float DownsampleFactor = FMath::Max(CVarHeterogeneousVolumesLightingCacheDownsampleFactor.GetValueOnRenderThread(), 1);
 		FIntVector LightingCacheResolution = GetVolumeResolution() / DownsampleFactor;
 
 		return LightingCacheResolution;
