Add widget to control hair voxel debug visualization.

diff --git a/Engine/Shaders/Private/HairStrands/HairStrandsVoxelDebug.usf b/Engine/Shaders/Private/HairStrands/HairStrandsVoxelDebug.usf
--- a/Engine/Shaders/Private/HairStrands/HairStrandsVoxelDebug.usf
+++ b/Engine/Shaders/Private/HairStrands/HairStrandsVoxelDebug.usf
@@ -22,8 +22,6 @@
 #include "HairStrandsVoxelPageTraversal.ush"
 
 float2 OutputResolution; 
-uint bDrawPage;
-int ForcedMipLevel;
 uint MacroGroupId;
 uint MacroGroupCount;
 uint MaxTotalPageIndexCount;
@@ -62,7 +60,6 @@ void MainCS(uint3 DispatchThreadId : SV_DispatchThreadID)
 	const FVirtualVoxelNodeDesc NodeDesc = UnpackVoxelNode(VirtualVoxel.NodeDescBuffer[MacroGroupId], VirtualVoxel.PageResolution);
 	const bool bIsNodeValid = all(NodeDesc.PageIndexResolution != 0);
 
-
 	const bool bPrint = all(uint2(PixelCoord) == 0);
 
 	// Draw macro group AABB & debug infos
@@ -254,9 +251,23 @@ void MainCS(uint3 DispatchThreadId : SV_DispatchThreadID)
 		}
 	}
 
+
+	// Draw widgets
+	bool bDrawPage = false;
+	int ForcedMipLevel = -1;
+	{
+		FShaderPrintContext WidgetContext = InitShaderPrintContext(MacroGroupId == 0 && bPrint, uint2(500, 100));
+		bDrawPage = AddCheckbox(WidgetContext, TEXT("Draw Pages"), false, FontWhite); Newline(WidgetContext);
+		bool  bForceMipLevel = AddCheckbox(WidgetContext, TEXT("Force MIP"), false, FontWhite); Newline(WidgetContext);
+		if (bForceMipLevel)
+		{
+			ForcedMipLevel = 8u * AddSlider(WidgetContext, TEXT("MIP"), 0, FontWhite); Newline(WidgetContext);
+		}
+	}
+
 	bool bIsValid = false;
 	float3 HitColor = 0;
-	if (bDrawPage > 0)
+	if (bDrawPage)
 	{
 		// Reverse-Z
 		const float3 O = PrimaryView.TranslatedWorldCameraOrigin;
diff --git a/Engine/Source/Runtime/Renderer/Private/HairStrands/HairStrandsDebug.cpp b/Engine/Source/Runtime/Renderer/Private/HairStrands/HairStrandsDebug.cpp
--- a/Engine/Source/Runtime/Renderer/Private/HairStrands/HairStrandsDebug.cpp
+++ b/Engine/Source/Runtime/Renderer/Private/HairStrands/HairStrandsDebug.cpp
@@ -50,10 +50,6 @@ static FAutoConsoleVariableRef CVarHairStrandsDebugBSDFAbsorption(TEXT("r.HairSt
 static float GHairStrandsDebugPlotBsdfExposure = 1.1f;
 static FAutoConsoleVariableRef CVarHairStrandsDebugBSDFExposure(TEXT("r.HairStrands.PlotBsdf.Exposure"), GHairStrandsDebugPlotBsdfExposure, TEXT("Change the exposure of the plot."));
 
-static int32 GHairVirtualVoxel_DrawDebugPage = 0;
-static FAutoConsoleVariableRef CVarHairVirtualVoxel_DrawDebugPage(TEXT("r.HairStrands.Voxelization.Virtual.DrawDebugPage"), GHairVirtualVoxel_DrawDebugPage, TEXT("When voxel debug rendering is enable 1: render the page bounds, instead of the voxel 2: the occupancy within the page (i.e., 8x8x8 brick)"));
-static int32 GHairVirtualVoxel_ForceMipLevel = -1;
-static FAutoConsoleVariableRef CVarHairVirtualVoxel_ForceMipLevel(TEXT("r.HairStrands.Voxelization.Virtual.ForceMipLevel"), GHairVirtualVoxel_ForceMipLevel, TEXT("Force a particular mip-level"));
 static int32 GHairVirtualVoxel_DebugTraversalType = 0;
 static FAutoConsoleVariableRef CVarHairVirtualVoxel_DebugTraversalType(TEXT("r.HairStrands.Voxelization.Virtual.DebugTraversalType"), GHairVirtualVoxel_DebugTraversalType, TEXT("Traversal mode (0:linear, 1:mip) for debug voxel visualization."));
 
@@ -724,8 +720,6 @@ class FVoxelVirtualRaymarchingCS : public FGlobalShader
 		SHADER_PARAMETER_STRUCT_INCLUDE(FSceneTextureParameters, SceneTextures)
 		SHADER_PARAMETER_STRUCT_INCLUDE(ShaderPrint::FShaderParameters, ShaderPrintParameters)
 		SHADER_PARAMETER(FVector2f, OutputResolution)
-		SHADER_PARAMETER( int32, ForcedMipLevel)
-		SHADER_PARAMETER(uint32, bDrawPage)
 		SHADER_PARAMETER(uint32, MacroGroupId)
 		SHADER_PARAMETER(uint32, MacroGroupCount)
 		SHADER_PARAMETER(uint32, MaxTotalPageIndexCount)
@@ -769,8 +763,6 @@ static void AddVoxelPageRaymarchingPass(
 		Parameters->ViewUniformBuffer		= View.ViewUniformBuffer;
 		Parameters->OutputResolution		= Resolution;
 		Parameters->SceneTextures			= SceneTextures;
-		Parameters->bDrawPage				= FMath::Clamp(GHairVirtualVoxel_DrawDebugPage, 0, 2);
-		Parameters->ForcedMipLevel			= FMath::Clamp(GHairVirtualVoxel_ForceMipLevel, -1, 5);
 		Parameters->MacroGroupId			= MacroGroupData.MacroGroupId;
 		Parameters->MacroGroupCount			= MacroGroupDatas.Num();
 		Parameters->MaxTotalPageIndexCount  = VoxelResources.Parameters.Common.PageIndexCount;
