Fix CPU driven hair voxelization.

diff --git a/Engine/Shaders/Private/HairStrands/HairStrandsVoxelRasterCompute.usf b/Engine/Shaders/Private/HairStrands/HairStrandsVoxelRasterCompute.usf
--- a/Engine/Shaders/Private/HairStrands/HairStrandsVoxelRasterCompute.usf
+++ b/Engine/Shaders/Private/HairStrands/HairStrandsVoxelRasterCompute.usf
@@ -31,9 +31,9 @@ uint				IndirectDispatchGroupSize;
 uint				bDoesMacroGroupSupportVoxelization;
 
 // For testing parity with CPU version
-float3				CPU_TranslatedWorldMinAABB;
-float3				CPU_TranslatedWorldMaxAABB;
-int3				CPU_PageIndexResolution;
+float4				CPU_TranslatedWorldMinAABB[MAX_HAIR_MACROGROUP_COUNT];
+float4				CPU_TranslatedWorldMaxAABB[MAX_HAIR_MACROGROUP_COUNT];
+int4				CPU_PageIndexResolution[MAX_HAIR_MACROGROUP_COUNT];
 uint				CPU_bUseCPUData;
 
 Buffer<float>								GPUVoxelWorldSize;
@@ -65,7 +65,16 @@ void AllocatePageIndex(uint2 DispatchThreadId : SV_DispatchThreadID)
 	if (bIsValid)
 	{
 		const bool bSupportVoxelization = (bDoesMacroGroupSupportVoxelization >> MacroGroupId) & 0x1;
-		Bound = ReadHairAABB(MacroGroupId, MacroGroupAABBBuffer);
+		if (CPU_bUseCPUData > 0)
+		{
+			Bound.Min = CPU_TranslatedWorldMinAABB[MacroGroupId];
+			Bound.Max = CPU_TranslatedWorldMaxAABB[MacroGroupId];
+		}
+		else
+		{
+			Bound = ReadHairAABB(MacroGroupId, MacroGroupAABBBuffer);
+		}
+
 		const float VoxelWorldSize = bUseCPUVoxelWorldSize ? CPUVoxelWorldSize : max(GPUVoxelWorldSize[0], MacroGroupVoxelSizeBuffer[MacroGroupId]);
 		PageWorldSize = VoxelWorldSize * PageResolution;
 
@@ -86,11 +95,17 @@ void AllocatePageIndex(uint2 DispatchThreadId : SV_DispatchThreadID)
 		// because the macro group AABBs for voxelization need to be snapped to the voxel page boundary.
 
 		// Allocate enough pages to cover the AABB, where page (0,0,0) origin sit on MinAABB.
-		float3 MacroGroupSize = 0;
 		if (bIsValid)
 		{
-			MacroGroupSize = Bound.Max - Bound.Min;
-			PageIndexResolution = ceil(MacroGroupSize / PageWorldSize);
+			float3 MacroGroupSize = Bound.Max - Bound.Min;
+			if (CPU_bUseCPUData > 0) 
+			{ 
+				PageIndexResolution = CPU_PageIndexResolution[MacroGroupId]; 
+			}
+			else
+			{
+				PageIndexResolution = ceil(MacroGroupSize / PageWorldSize);
+			}
 			MacroGroupSize = PageIndexResolution * PageWorldSize;
 			Bound.Max = MacroGroupSize + Bound.Min;
 		}
@@ -142,16 +157,6 @@ void AllocatePageIndex(uint2 DispatchThreadId : SV_DispatchThreadID)
 		return;
 	}
 
-	// Use CPU data, which allow to compare the GPU driven part with CPU data
-	// This code path is only for debug purpose and work only when there is a single instance group
-	if (CPU_bUseCPUData > 0)
-	{
-		Bound.Min = CPU_TranslatedWorldMinAABB;
-		Bound.Max = CPU_TranslatedWorldMaxAABB;
-		PageIndexResolution = CPU_PageIndexResolution;
-		WriteHairAABB(MacroGroupId, Bound, MacroGroupVoxelAlignedAABBBuffer);
-	}
-
 	// Prepare indirect buffer for doing the actual page index allocation and filling the page index
 	{
 		const uint AllocatedPageIndexCount = PageIndexResolution.x * PageIndexResolution.y * PageIndexResolution.z;
@@ -228,6 +233,9 @@ void AllocatePageIndex(uint2 DispatchThreadId : SV_DispatchThreadID)
 #if SHADER_MARKVALID_PREPARE
 uint			MaxClusterCount;
 uint			MacroGroupId;
+uint			bUseMacroGroupBoundCPU;
+float3			MacroGroupBoundCPU_TranslatedWorldMinAABB;
+float3			MacroGroupBoundCPU_TranslatedWorldMaxAABB;
 
 Buffer<int>		GroupAABBsBuffer;
 Buffer<int>		ClusterAABBsBuffer;
@@ -316,7 +324,6 @@ void MarkValid_PrepareCS(uint2 DispatchThreadId : SV_DispatchThreadID)
 void MarkValid_PrepareCS(uint3 DispatchThreadId : SV_DispatchThreadID)
 {
 	const uint4 PageIndexResolutionAndOffset = PageIndexResolutionAndOffsetBuffer.Load(MacroGroupId);
-	const FHairAABB MacroGroupBound = ReadHairAABB(MacroGroupId, MacroGroupVoxelAlignedAABBBuffer);
 	const int3 PageIndexResolution = PageIndexResolutionAndOffset.xyz;
 	const uint PageIndexOffset = PageIndexResolutionAndOffset.w;
 
@@ -324,7 +331,23 @@ void MarkValid_PrepareCS(uint3 DispatchThreadId : SV_DispatchThreadID)
 	if (any(Coord >= PageIndexResolution))
 		return;
 
-	const FHairAABB GroupBound = ReadHairAABB(0, GroupAABBsBuffer);
+	FHairAABB MacroGroupBound;
+	FHairAABB GroupBound;
+	if (bUseMacroGroupBoundCPU)
+	{
+		MacroGroupBound.Min = MacroGroupBoundCPU_TranslatedWorldMinAABB;
+		MacroGroupBound.Max = MacroGroupBoundCPU_TranslatedWorldMaxAABB;
+
+		// HAIR_TODO: Can we have reliable primitive AABB to have tigher bound?
+		GroupBound.Min = MacroGroupBoundCPU_TranslatedWorldMinAABB;
+		GroupBound.Max = MacroGroupBoundCPU_TranslatedWorldMaxAABB;
+	}
+	else
+	{
+		MacroGroupBound = ReadHairAABB(MacroGroupId, MacroGroupVoxelAlignedAABBBuffer);
+		GroupBound = ReadHairAABB(0, GroupAABBsBuffer);
+	}
+
 	const uint3 MinCoord = PositionToCoord(GroupBound.Min, MacroGroupBound.Min, MacroGroupBound.Max, PageIndexResolution);
 	const uint3 MaxCoord = PositionToCoord(GroupBound.Max, MacroGroupBound.Min, MacroGroupBound.Max, PageIndexResolution);
 
@@ -346,83 +369,6 @@ void MarkValid_PrepareCS(uint3 DispatchThreadId : SV_DispatchThreadID)
 
 ///////////////////////////////////////////////////////////////////////////////////////////////////
 
-#if SHADER_MARKVALID
-uint			MaxClusterCount;
-int3			CPU_PageIndexResolution;
-float3			CPU_TranslatedWorldMinAABB;
-uint			CPU_PageIndexOffset;
-float3			CPU_TranslatedWorldMaxAABB;
-uint			MacroGroupId;
-
-Buffer<int>		ClusterAABBsBuffer;
-Buffer<int>		MacroGroupVoxelAlignedAABBBuffer;
-Buffer<uint4>	PageIndexResolutionAndOffsetBuffer;
-RWBuffer<uint>	OutValidPageIndexBuffer;
-
-// PageIndexBuffer      is sampled with linear coordinate computed from the 3d page coordinate. VALID NODE ARE NOT COMPACTED. It contains the LINEAR PAGE INDEX (to map to the 3d volume).
-// PageIndexCoordBuffer is sampled with linear coordinate for allocated nodes. VALID NODE ARE COMPACTED. It contains the 3d page coordinate and ClustedId. Only used for opaque voxel injection.
-
-[numthreads(GROUP_SIZE, 1, 1)]
-void MarkValidCS(uint2 DispatchThreadId : SV_DispatchThreadID)
-{
-	const uint ClusterIndex = DispatchThreadId.x;
-	if (ClusterIndex >= MaxClusterCount)
-		return;
-
-	const uint BaseClusterIndex = 6 * ClusterIndex;
-
-	FHairAABB ClusterBound;
-	ClusterBound.Min.x = float(ClusterAABBsBuffer[BaseClusterIndex + 0]);
-	ClusterBound.Min.y = float(ClusterAABBsBuffer[BaseClusterIndex + 1]);
-	ClusterBound.Min.z = float(ClusterAABBsBuffer[BaseClusterIndex + 2]);
-
-	ClusterBound.Max.x = float(ClusterAABBsBuffer[BaseClusterIndex + 3]);
-	ClusterBound.Max.y = float(ClusterAABBsBuffer[BaseClusterIndex + 4]);
-	ClusterBound.Max.z = float(ClusterAABBsBuffer[BaseClusterIndex + 5]);
-
-	FHairAABB MacroGroupBound;
-	int3 PageIndexResolution;
-	uint PageIndexOffset;
-
-#if PERMUTATION_GPU_DRIVEN == 1
-	const uint4 PageIndexResolutionAndOffset = PageIndexResolutionAndOffsetBuffer.Load(MacroGroupId);
-
-	MacroGroupBound = ReadHairAABB(MacroGroupId, MacroGroupVoxelAlignedAABBBuffer);
-	PageIndexResolution = PageIndexResolutionAndOffset.xyz;
-	PageIndexOffset = PageIndexResolutionAndOffset.w;
-#else
-	MacroGroupBound.Min = CPU_TranslatedWorldMinAABB;
-	MacroGroupBound.Max = CPU_TranslatedWorldMaxAABB;
-	PageIndexResolution = CPU_PageIndexResolution;
-	PageIndexOffset = CPU_PageIndexOffset;
-#endif
-
-	uint3 MinCoord = PositionToCoord(ClusterBound.Min, MacroGroupBound.Min, MacroGroupBound.Max, PageIndexResolution);
-	uint3 MaxCoord = PositionToCoord(ClusterBound.Max, MacroGroupBound.Min, MacroGroupBound.Max, PageIndexResolution);
-
-	uint3 PageIndexResolutionMinusOne = PageIndexResolution - 1;
-	MinCoord = clamp(MinCoord, uint3(0, 0, 0), PageIndexResolutionMinusOne);
-	MaxCoord = clamp(MaxCoord, uint3(0, 0, 0), PageIndexResolutionMinusOne);
-
-	// Splat cluster contribution
-	// Change this to be a bit smarter (use graphics pipe?)
-	for (uint z = MinCoord.z; z <= MaxCoord.z; ++z)
-	{
-		for (uint y = MinCoord.y; y <= MaxCoord.y; ++y)
-		{
-			for (uint x = MinCoord.x; x <= MaxCoord.x; ++x)
-			{
-				const uint3 PageIndexCoord = uint3(x, y, z);
-				const uint LinearPageIndexCoord = CoordToIndex(PageIndexCoord, PageIndexResolution, PageIndexOffset);
-				InterlockedOr(OutValidPageIndexBuffer[LinearPageIndexCoord], 1u);
-			}
-		}
-	}
-}
-#endif
-
-///////////////////////////////////////////////////////////////////////////////////////////////////
-
 #if SHADER_ALLOCATE
 uint MacroGroupId;
 uint PageCount;
diff --git a/Engine/Source/Runtime/Renderer/Private/HairStrands/HairStrandsVoxelization.cpp b/Engine/Source/Runtime/Renderer/Private/HairStrands/HairStrandsVoxelization.cpp
--- a/Engine/Source/Runtime/Renderer/Private/HairStrands/HairStrandsVoxelization.cpp
+++ b/Engine/Source/Runtime/Renderer/Private/HairStrands/HairStrandsVoxelization.cpp
@@ -101,9 +101,6 @@ static int32 GHairVirtualVoxelGPUDrivenClampPageIndexRes = 256;
 static FAutoConsoleVariableRef CVarHairVirtualVoxelGPUDriven(TEXT("r.HairStrands.Voxelization.GPUDriven"), GHairVirtualVoxelGPUDriven, TEXT("Enable GPU driven voxelization."));
 static FAutoConsoleVariableRef CVarHairVirtualVoxelGPUDrivenMaxPageIndexRes(TEXT("r.HairStrands.Voxelization.GPUDriven.MaxPageIndexResolution"), GHairVirtualVoxelGPUDrivenMaxPageIndexRes, TEXT("Max resolution of the page index. This is used for allocating a conservative page index buffer when GPU driven allocation is enabled."));
 
-static int32 GHairVirtualVoxelUseIndirectScatterPageAllocation = 1;
-static FAutoConsoleVariableRef CVarHairVirtualVoxelUseIndirectScatterPageAllocation(TEXT("r.HairStrands.Voxelization.UseIndiretScatterPageAllocate"), GHairVirtualVoxelUseIndirectScatterPageAllocation, TEXT("Enable indirect scatter page allocation (faster)."));
-
 static const FIntPoint GPUDrivenViewportResolution = FIntPoint(4096, 4096);
 
 static int32 GHairVirtualVoxelInvalidEmptyPageIndex = 1;
@@ -127,6 +124,12 @@ static FAutoConsoleVariableRef CVarHairVirtualVoxelAdaptive_CorrectionThreshold(
 static int32 GHairVirtualVoxel_JitterMode = 1;
 static FAutoConsoleVariableRef CVarHairVirtualVoxel_JitterMode(TEXT("r.HairStrands.Voxelization.Virtual.Jitter"), GHairVirtualVoxel_JitterMode, TEXT("Change jittered for voxelization/traversal. 0: No jitter 1: Regular randomized jitter: 2: Constant Jitter (default = 1)"), ECVF_RenderThreadSafe);
 
+bool IsHairStrandsContinousLODEnabled()
+{
+	static IConsoleVariable* CVar = IConsoleManager::Get().FindConsoleVariable(TEXT("r.HairStrands.ContinuousLOD"));
+	return CVar && CVar->GetInt() > 0;
+}
+
 void GetVoxelPageResolution(uint32& OutPageResolution, uint32& OutPageResolutionLog2)
 {
 	OutPageResolutionLog2 = FMath::CeilLogTwo(FMath::Clamp(GHairVirtualVoxel_PageResolution, 2, 256));
@@ -135,7 +138,8 @@ void GetVoxelPageResolution(uint32& OutPageResolution, uint32& OutPageResolution
 
 static bool IsHairStrandsAdaptiveVoxelAllocationEnable()
 {
-	return GHairVirtualVoxelAdaptive_Enable > 0;
+	// HAIR_TODO: GPU driven should be compatible with CLOD
+	return !IsHairStrandsContinousLODEnabled() && GHairVirtualVoxelAdaptive_Enable > 0;
 }
 
 bool IsHairStrandsVoxelizationEnable()
@@ -148,6 +152,13 @@ bool IsHairStrandsForVoxelTransmittanceAndShadowEnable()
 	return IsHairStrandsVoxelizationEnable() && GHairForVoxelTransmittanceAndShadow > 0;
 }
 
+bool IsHairStrandsGPUDrivenVoxelizationEnabled()
+{
+	// HAIR_TODO: GPU driven should be compatible with CLOD
+	return !IsHairStrandsContinousLODEnabled() && GHairVirtualVoxelGPUDriven > 0; 
+}
+
+
 ///////////////////////////////////////////////////////////////////////////////////////////////////
 
 IMPLEMENT_GLOBAL_SHADER_PARAMETER_STRUCT(FVirtualVoxelParameters, "VirtualVoxel");
@@ -240,9 +251,9 @@ class FVoxelAllocatePageIndexCS : public FGlobalShader
 		SHADER_PARAMETER(uint32, IndirectDispatchGroupSize)
 		SHADER_PARAMETER(uint32, bDoesMacroGroupSupportVoxelization)
 
-		SHADER_PARAMETER(FVector3f, CPU_TranslatedWorldMinAABB)
-		SHADER_PARAMETER(FVector3f, CPU_TranslatedWorldMaxAABB)
-		SHADER_PARAMETER(FIntVector, CPU_PageIndexResolution)
+		SHADER_PARAMETER_ARRAY(FVector4f, CPU_TranslatedWorldMinAABB, [MAX_HAIR_MACROGROUP_COUNT])
+		SHADER_PARAMETER_ARRAY(FVector4f, CPU_TranslatedWorldMaxAABB, [MAX_HAIR_MACROGROUP_COUNT])
+		SHADER_PARAMETER_ARRAY(FIntVector4, CPU_PageIndexResolution, [MAX_HAIR_MACROGROUP_COUNT])
 		SHADER_PARAMETER(uint32, CPU_bUseCPUData)
 
 		SHADER_PARAMETER_RDG_BUFFER_SRV(Buffer<float>, GPUVoxelWorldSize)
@@ -276,6 +287,9 @@ class FVoxelMarkValidPageIndex_PrepareCS : public FGlobalShader
 	BEGIN_SHADER_PARAMETER_STRUCT(FParameters, )
 		SHADER_PARAMETER(uint32, MaxClusterCount)
 		SHADER_PARAMETER(uint32, MacroGroupId)
+		SHADER_PARAMETER(uint32, bUseMacroGroupBoundCPU)
+		SHADER_PARAMETER(FVector3f, MacroGroupBoundCPU_TranslatedWorldMinAABB)
+		SHADER_PARAMETER(FVector3f, MacroGroupBoundCPU_TranslatedWorldMaxAABB)
 
 		SHADER_PARAMETER_RDG_BUFFER_SRV(Buffer, GroupAABBsBuffer)
 		SHADER_PARAMETER_RDG_BUFFER_SRV(Buffer, ClusterAABBsBuffer)
@@ -298,38 +312,6 @@ public:
 	}
 };
 
-class FVoxelMarkValidPageIndexCS : public FGlobalShader
-{
-	DECLARE_GLOBAL_SHADER(FVoxelMarkValidPageIndexCS);
-	SHADER_USE_PARAMETER_STRUCT(FVoxelMarkValidPageIndexCS, FGlobalShader);
-
-	class FGPUDriven : SHADER_PERMUTATION_INT("PERMUTATION_GPU_DRIVEN", 2);
-	using FPermutationDomain = TShaderPermutationDomain<FGPUDriven>;
-
-	BEGIN_SHADER_PARAMETER_STRUCT(FParameters, )
-		SHADER_PARAMETER(FIntVector, CPU_PageIndexResolution)
-		SHADER_PARAMETER(FVector3f, CPU_TranslatedWorldMinAABB)
-		SHADER_PARAMETER(uint32, MaxClusterCount)
-		SHADER_PARAMETER(FVector3f, CPU_TranslatedWorldMaxAABB)
-		SHADER_PARAMETER(uint32, CPU_PageIndexOffset)
-		SHADER_PARAMETER(uint32, MacroGroupId)
-		SHADER_PARAMETER_RDG_BUFFER_SRV(Buffer, ClusterAABBsBuffer)
-		SHADER_PARAMETER_RDG_BUFFER_SRV(Buffer, MacroGroupVoxelAlignedAABBBuffer)
-		SHADER_PARAMETER_RDG_BUFFER_SRV(Buffer, PageIndexResolutionAndOffsetBuffer)
-		SHADER_PARAMETER_RDG_BUFFER_UAV(RWBuffer<uint>, OutValidPageIndexBuffer)
-	END_SHADER_PARAMETER_STRUCT()
-
-public:
-	static uint32 GetGroupSize() { return 256u; /* Allow to process up to 16M clusters (= 65k * 256) */ }
-	static bool ShouldCompilePermutation(const FGlobalShaderPermutationParameters& Parameters) { return IsHairStrandsSupported(EHairStrandsShaderType::Strands, Parameters.Platform); }
-	static void ModifyCompilationEnvironment(const FGlobalShaderPermutationParameters& Parameters, FShaderCompilerEnvironment& OutEnvironment)
-	{
-		FGlobalShader::ModifyCompilationEnvironment(Parameters, OutEnvironment);
-		OutEnvironment.SetDefine(TEXT("SHADER_MARKVALID"), 1);
-		OutEnvironment.SetDefine(TEXT("GROUP_SIZE"), GetGroupSize());
-	}
-};
-
 class FVoxelAllocateVoxelPageCS : public FGlobalShader
 {
 	DECLARE_GLOBAL_SHADER(FVoxelAllocateVoxelPageCS);
@@ -497,7 +479,6 @@ public:
 };
 
 IMPLEMENT_GLOBAL_SHADER(FVoxelMarkValidPageIndex_PrepareCS, "/Engine/Private/HairStrands/HairStrandsVoxelRasterCompute.usf", "MarkValid_PrepareCS", SF_Compute);
-IMPLEMENT_GLOBAL_SHADER(FVoxelMarkValidPageIndexCS, "/Engine/Private/HairStrands/HairStrandsVoxelRasterCompute.usf", "MarkValidCS", SF_Compute);
 IMPLEMENT_GLOBAL_SHADER(FVoxelAllocatePageIndexCS, "/Engine/Private/HairStrands/HairStrandsVoxelRasterCompute.usf", "AllocatePageIndex", SF_Compute);
 IMPLEMENT_GLOBAL_SHADER(FVoxelAllocateVoxelPageCS, "/Engine/Private/HairStrands/HairStrandsVoxelRasterCompute.usf", "AllocateCS", SF_Compute);
 IMPLEMENT_GLOBAL_SHADER(FVoxelAddNodeDescCS, "/Engine/Private/HairStrands/HairStrandsVoxelRasterCompute.usf", "AddDescCS", SF_Compute);
@@ -533,7 +514,7 @@ static void AddAllocateVoxelPagesPass(
 	FRDGBufferRef& CurrGPUMinVoxelSize,
 	FRDGBufferRef& NextGPUMinVoxelSize)
 {
-	const bool bIsGPUDriven = GHairVirtualVoxelGPUDriven > 0;
+	const bool bIsGPUDriven = IsHairStrandsGPUDrivenVoxelizationEnabled();
 	const uint32 MacroGroupCount = MacroGroupDatas.Num();
 	if (MacroGroupCount == 0)
 		return;
@@ -557,8 +538,8 @@ static void AddAllocateVoxelPagesPass(
 		// Scale the bounding box in place of proper GPU driven AABB for now
 		const float Scale = FMath::Clamp(GHairVoxelizationAABBScale, 0.01f, 10.0f);
 		const FVector BoxCenter = MacroGroup.Bounds.GetBox().GetCenter();
-		FVector TranslatedMinAABB = (MacroGroup.Bounds.GetBox().Min - BoxCenter) * Scale + BoxCenter;
-		FVector TranslatedMaxAABB = (MacroGroup.Bounds.GetBox().Max - BoxCenter) * Scale + BoxCenter;
+		FVector TranslatedMinAABB = (MacroGroup.Bounds.GetBox().Min - BoxCenter) * Scale + BoxCenter + View.ViewMatrices.GetPreViewTranslation();
+		FVector TranslatedMaxAABB = (MacroGroup.Bounds.GetBox().Max - BoxCenter) * Scale + BoxCenter + View.ViewMatrices.GetPreViewTranslation();
 
 		// Allocate enough pages to cover the AABB, where page (0,0,0) origin sit on MinAABB.
 		FVector MacroGroupSize = TranslatedMaxAABB - TranslatedMinAABB;
@@ -580,7 +561,7 @@ static void AddAllocateVoxelPagesPass(
 	// Over-allocation (upper bound)
 	if (bIsGPUDriven)
 	{
-		// Use the max between the estimated size on CPU and a pseudo-conservative side driven by settings. The CPU estimation is no necessarely correct as the bounds are not reliable on skel. mesh.
+		// Use the max between the estimated size on CPU and a pseudo-conservative side driven by settings. The CPU estimation is no necessarily correct as the bounds are not reliable on skel. mesh.
 		const uint32 MinPageIndexCount = GHairVirtualVoxelGPUDrivenMaxPageIndexRes * GHairVirtualVoxelGPUDrivenMaxPageIndexRes * GHairVirtualVoxelGPUDrivenMaxPageIndexRes;
 		const uint32 MaxPageIndexCount = GHairVirtualVoxelGPUDrivenClampPageIndexRes * GHairVirtualVoxelGPUDrivenClampPageIndexRes * GHairVirtualVoxelGPUDrivenClampPageIndexRes;
 		OutTotalPageIndexCount = FMath::Clamp(OutTotalPageIndexCount, MinPageIndexCount, MaxPageIndexCount);
@@ -622,7 +603,6 @@ static void AddAllocateVoxelPagesPass(
 	const bool bAdaptiveVoxelEnable = IsHairStrandsAdaptiveVoxelAllocationEnable();
 
 	// Allocate page index for all instance group
-//	if (bIsGPUDriven)
 	{
 		uint32 bDoesMacroGroupSupportVoxelization = 0;
 		for (const FHairStrandsMacroGroupData& MacroGroup : MacroGroupDatas)
@@ -636,7 +616,7 @@ static void AddAllocateVoxelPagesPass(
 		FVoxelAllocatePageIndexCS::FParameters* Parameters = GraphBuilder.AllocParameters<FVoxelAllocatePageIndexCS::FParameters>();
 		Parameters->CPUPageWorldSize = CPUPageWorldSize;
 		Parameters->CPUVoxelWorldSize = CPUMinVoxelWorldSize;
-		Parameters->bUseCPUVoxelWorldSize = bAdaptiveVoxelEnable ? 0u : 1u;
+		Parameters->bUseCPUVoxelWorldSize = bIsGPUDriven && bAdaptiveVoxelEnable ? 0u : 1u;
 		Parameters->TotalPageIndexCount = OutTotalPageIndexCount;
 		Parameters->PageResolution = PageResolution;
 		Parameters->MacroGroupCount = MacroGroupCount;
@@ -648,13 +628,24 @@ static void AddAllocateVoxelPagesPass(
 		Parameters->OutPageIndexResolutionAndOffsetBuffer = GraphBuilder.CreateUAV(PageIndexResolutionBuffer, PF_R32G32B32A32_UINT);
 		Parameters->OutVoxelizationViewInfoBuffer = GraphBuilder.CreateUAV(VoxelizationViewInfoBuffer);
 		Parameters->OutPageIndexAllocationIndirectBufferArgs = GraphBuilder.CreateUAV(PageIndexAllocationIndirectBufferArgs);
-		Parameters->CPU_bUseCPUData = GHairVirtualVoxelGPUDriven == 2 ? 1 : 0;
+		Parameters->CPU_bUseCPUData = bIsGPUDriven ? 0 : 1;
 		Parameters->GPUVoxelWorldSize = CurrGPUMinVoxelSizeSRV;
-		if (Parameters->CPU_bUseCPUData)
+
+		for (int32 It=0; It < MAX_HAIR_MACROGROUP_COUNT; ++It)
 		{
-			Parameters->CPU_TranslatedWorldMinAABB = CPUAllocationDescs.Num() > 0 ? (FVector3f)CPUAllocationDescs[0].TranslatedMinAABB : FVector3f::ZeroVector;
-			Parameters->CPU_TranslatedWorldMaxAABB = CPUAllocationDescs.Num() > 0 ? (FVector3f)CPUAllocationDescs[0].TranslatedMaxAABB : FVector3f::ZeroVector;
-			Parameters->CPU_PageIndexResolution = CPUAllocationDescs.Num() > 0 ? CPUAllocationDescs[0].PageIndexResolution : FIntVector(0, 0, 0);
+			if (It < CPUAllocationDescs.Num())
+			{
+				const FCPUMacroGroupAllocation& Desc = CPUAllocationDescs[It];
+				Parameters->CPU_TranslatedWorldMinAABB[It] = FVector4f(Desc.TranslatedMinAABB.X, Desc.TranslatedMinAABB.Y, Desc.TranslatedMinAABB.Z, 0);
+				Parameters->CPU_TranslatedWorldMaxAABB[It] = FVector4f(Desc.TranslatedMaxAABB.X, Desc.TranslatedMaxAABB.Y, Desc.TranslatedMaxAABB.Z, 0);
+				Parameters->CPU_PageIndexResolution[It]    = FIntVector4(Desc.PageIndexResolution, 0.f);
+			}
+			else
+			{
+				Parameters->CPU_TranslatedWorldMinAABB[It] = FVector4f(0, 0, 0, 0);
+				Parameters->CPU_TranslatedWorldMaxAABB[It] = FVector4f(0, 0, 0, 0);
+				Parameters->CPU_PageIndexResolution[It]    = FIntVector4(0, 0, 0, 0);
+			}
 		}
 
 		// Currently support only 32 instance group at max
@@ -669,20 +660,6 @@ static void AddAllocateVoxelPagesPass(
 	}
 	FRDGBufferSRVRef PageIndexResolutionAndOffsetBufferSRV = GraphBuilder.CreateSRV(PageIndexResolutionBuffer, PF_R32G32B32A32_UINT);
 
-	uint32 TotalClusterCount = 0;
-	for (uint32 MacroGroupIt = 0; MacroGroupIt < MacroGroupCount; ++MacroGroupIt)
-	{
-		const FHairStrandsMacroGroupData& MacroGroup = MacroGroupDatas[MacroGroupIt];
-		for (const FHairStrandsMacroGroupData::PrimitiveInfo& PrimitiveInfo : MacroGroup.PrimitivesInfos)
-		{
-			FHairGroupPublicData* HairGroupData = PrimitiveInfo.PublicDataPtr;
-			if (HairGroupData->DoesSupportVoxelization())
-			{
-				TotalClusterCount += HairGroupData->GetClusterCount();
-			}
-		}
-	}
-
 	// Mark valid page index
 	for (uint32 MacroGroupIt=0; MacroGroupIt <MacroGroupCount;++MacroGroupIt)
 	{
@@ -692,9 +669,6 @@ static void AddAllocateVoxelPagesPass(
 
 		const FHairStrandsMacroGroupData& MacroGroup = MacroGroupDatas[MacroGroupIt];
 		FCPUMacroGroupAllocation& CPUAllocationDesc = CPUAllocationDescs[MacroGroupIt];
-
-		const bool bUseIndirectScatter = GHairVirtualVoxelUseIndirectScatterPageAllocation>0 && bIsGPUDriven;
-		if (bUseIndirectScatter)
 		{			
 			for (const FHairStrandsMacroGroupData::PrimitiveInfo& PrimitiveInfo : MacroGroup.PrimitivesInfos)
 			{
@@ -705,21 +679,23 @@ static void AddAllocateVoxelPagesPass(
 					continue;
 
 				// For allocating/marking the used pages, a primitive/group can:
-				// * Either use its group AABB (for static groom)
+				// * Either use its group AABB + cluster AABB (for static groom)
 				// * Or use its clusters AABBs (for dynamic groom)
+				// * Or use its CPU AABB (for CLOD)
 				// Even if the bGroupAABBValid is invalid, it has been reset to invalid AABB, so the code below won't mark any page as used.
-				const bool bUseClusterAABB = HairGroupData->bClusterAABBValid;
+				const bool bUseClusterAABB = HairGroupData->bClusterAABBValid && bIsGPUDriven;
 				
 				FVoxelMarkValidPageIndex_PrepareCS::FParameters* Parameters = GraphBuilder.AllocParameters<FVoxelMarkValidPageIndex_PrepareCS::FParameters>();
-				Parameters->MaxClusterCount				= HairGroupData->GetClusterCount();
-				Parameters->MacroGroupId				= MacroGroup.MacroGroupId;
-
-				Parameters->GroupAABBsBuffer			= RegisterAsSRV(GraphBuilder, HairGroupData->GetGroupAABBBuffer());
-				Parameters->ClusterAABBsBuffer			= RegisterAsSRV(GraphBuilder, HairGroupData->GetClusterAABBBuffer());
-				Parameters->MacroGroupVoxelAlignedAABBBuffer = GraphBuilder.CreateSRV(MacroGroupResources.MacroGroupVoxelAlignedAABBsBuffer, PF_R32_SINT);
-				Parameters->PageIndexResolutionAndOffsetBuffer = PageIndexResolutionAndOffsetBufferSRV;
-
-				Parameters->OutValidPageIndexBuffer		= PageIndexBufferUAV;
+				Parameters->MaxClusterCount								= HairGroupData->GetClusterCount();
+				Parameters->MacroGroupId								= MacroGroup.MacroGroupId;
+				Parameters->GroupAABBsBuffer							= RegisterAsSRV(GraphBuilder, HairGroupData->GetGroupAABBBuffer());
+				Parameters->ClusterAABBsBuffer							= RegisterAsSRV(GraphBuilder, HairGroupData->GetClusterAABBBuffer());
+				Parameters->MacroGroupVoxelAlignedAABBBuffer			= GraphBuilder.CreateSRV(MacroGroupResources.MacroGroupVoxelAlignedAABBsBuffer, PF_R32_SINT);
+				Parameters->PageIndexResolutionAndOffsetBuffer			= PageIndexResolutionAndOffsetBufferSRV;
+				Parameters->bUseMacroGroupBoundCPU						= bIsGPUDriven ? 0 : 1;
+				Parameters->MacroGroupBoundCPU_TranslatedWorldMinAABB	= (FVector3f)CPUAllocationDesc.TranslatedMinAABB;
+				Parameters->MacroGroupBoundCPU_TranslatedWorldMaxAABB	= (FVector3f)CPUAllocationDesc.TranslatedMaxAABB;
+				Parameters->OutValidPageIndexBuffer						= PageIndexBufferUAV;
 
 				FIntVector DispatchCount = bUseClusterAABB ?
 					FIntVector(
@@ -740,45 +716,7 @@ static void AddAllocateVoxelPagesPass(
 				TShaderMapRef<FVoxelMarkValidPageIndex_PrepareCS> ComputeShader(View.ShaderMap, PermutationVector);
 				FComputeShaderUtils::AddPass(
 					GraphBuilder,
-					RDG_EVENT_NAME("HairStrands::MarkValidPageIndex(%s)", bUseClusterAABB ? TEXT("Cluster") : TEXT("Group")),
-					ComputeShader,
-					Parameters,
-					DispatchCount);
-			}
-		}
-		else // TODO remove this path, as the FVoxelMarkValidPageIndex_PrepareCS(FUseCluster=false) handles the same path)
-		{		
-			for (const FHairStrandsMacroGroupData::PrimitiveInfo& PrimitiveInfo : MacroGroup.PrimitivesInfos)
-			{
-				FHairGroupPublicData* HairGroupData = PrimitiveInfo.PublicDataPtr;
-				if (!HairGroupData->DoesSupportVoxelization())
-					continue;
-
-				FVoxelMarkValidPageIndexCS::FParameters* Parameters = GraphBuilder.AllocParameters<FVoxelMarkValidPageIndexCS::FParameters>();
-				Parameters->MacroGroupId = MacroGroup.MacroGroupId;
-				Parameters->MaxClusterCount = HairGroupData->GetClusterCount();
-				Parameters->CPU_PageIndexResolution = CPUAllocationDesc.PageIndexResolution;
-				Parameters->CPU_PageIndexOffset = CPUAllocationDesc.PageIndexOffset;
-				Parameters->CPU_TranslatedWorldMinAABB = (FVector3f)CPUAllocationDesc.TranslatedMinAABB;
-				Parameters->CPU_TranslatedWorldMaxAABB = (FVector3f)CPUAllocationDesc.TranslatedMaxAABB;
-				Parameters->ClusterAABBsBuffer = RegisterAsSRV(GraphBuilder, HairGroupData->GetClusterAABBBuffer());
-				Parameters->OutValidPageIndexBuffer = PageIndexBufferUAV;
-
-				if (bIsGPUDriven)
-				{
-					Parameters->MacroGroupVoxelAlignedAABBBuffer = GraphBuilder.CreateSRV(MacroGroupResources.MacroGroupVoxelAlignedAABBsBuffer, PF_R32_SINT);
-					Parameters->PageIndexResolutionAndOffsetBuffer = PageIndexResolutionAndOffsetBufferSRV;
-				}
-
-				FVoxelMarkValidPageIndexCS::FPermutationDomain PermutationVector;
-				PermutationVector.Set<FVoxelMarkValidPageIndexCS::FGPUDriven>(bIsGPUDriven ? 1 : 0);
-
-				const FIntVector DispatchCount(FMath::DivideAndRoundUp(Parameters->MaxClusterCount, FVoxelMarkValidPageIndexCS::GetGroupSize()), 1, 1);
-				check(DispatchCount.X <= GRHIMaxDispatchThreadGroupsPerDimension.X);
-				TShaderMapRef<FVoxelMarkValidPageIndexCS> ComputeShader(View.ShaderMap, PermutationVector);
-				FComputeShaderUtils::AddPass(
-					GraphBuilder,
-					RDG_EVENT_NAME("HairStrands::MarkValidPageIndex"),
+					RDG_EVENT_NAME("HairStrands::MarkValidPageIndex(%s)", bUseClusterAABB ? TEXT("Cluster") : (bIsGPUDriven ? TEXT("Group") : TEXT("CPUGroup"))),
 					ComputeShader,
 					Parameters,
 					DispatchCount);
@@ -1197,10 +1135,6 @@ static void AddVirtualVoxelizationComputeRasterPass(
 	const FHairStrandsVoxelResources& VoxelResources,
 	const FHairStrandsMacroGroupData& MacroGroup)
 {
-	const bool bIsGPUDriven = GHairVirtualVoxelGPUDriven > 0;
-	if (!bIsGPUDriven)
-		return;
-
 	ensure(MacroGroup.bSupportVoxelization);
 
 	if (ViewInfo)
@@ -1250,7 +1184,7 @@ static void AddVirtualVoxelizationComputeRasterPass(
 				PassParameters->VertexCount = PointCount;
 				PassParameters->CoverageScale = CoverageScale;
 			
-				const bool bCullingEnable = !IsHairStrandContinuousDecimationReorderingEnabled() && HairGroupPublicData->GetCullingResultAvailable();
+				const bool bCullingEnable = !IsHairStrandsContinousLODEnabled() && !IsHairStrandContinuousDecimationReorderingEnabled() && HairGroupPublicData->GetCullingResultAvailable();
 
 				PassParameters->HairInstance = GetHairStrandsInstanceParameters(GraphBuilder, *ViewInfo, HairGroupPublicData, bCullingEnable, true);
 
