Remove (no longer) used voxelization view info.

This was used by the old HW raster path, and is no longer used now.

diff --git a/Engine/Shaders/Private/HairStrands/HairStrandsVoxelPageCommonStruct.ush b/Engine/Shaders/Private/HairStrands/HairStrandsVoxelPageCommonStruct.ush
--- a/Engine/Shaders/Private/HairStrands/HairStrandsVoxelPageCommonStruct.ush
+++ b/Engine/Shaders/Private/HairStrands/HairStrandsVoxelPageCommonStruct.ush
@@ -8,13 +8,4 @@ struct FPackedVirtualVoxelNodeDesc
 	uint	PackedPageIndexResolution; // 8bits encoding   
 	float3	TranslatedWorldMaxAABB;
 	uint	PageIndexOffset_VoxelWorldSize;
-}; 
-
-struct FVoxelizationViewInfo
-{
-	float4x4 TranslatedWorldToClip;
-	float3   ViewForward;
-	float    Pad0;
-	uint2    RasterResolution;
-	float2   Pad1;
-};  
\ No newline at end of file
+};
\ No newline at end of file
diff --git a/Engine/Shaders/Private/HairStrands/HairStrandsVoxelRasterCompute.usf b/Engine/Shaders/Private/HairStrands/HairStrandsVoxelRasterCompute.usf
--- a/Engine/Shaders/Private/HairStrands/HairStrandsVoxelRasterCompute.usf
+++ b/Engine/Shaders/Private/HairStrands/HairStrandsVoxelRasterCompute.usf
@@ -42,7 +42,6 @@ Buffer<int>									MacroGroupAABBBuffer;
 RWBuffer<int>								MacroGroupVoxelAlignedAABBBuffer;
 RWBuffer<uint4>								OutPageIndexResolutionAndOffsetBuffer;
 RWBuffer<uint>								OutPageIndexAllocationIndirectBufferArgs;
-RWStructuredBuffer<FVoxelizationViewInfo>	OutVoxelizationViewInfoBuffer;
 
 
 #if GROUP_SIZE != MAX_HAIR_MACROGROUP_COUNT
@@ -141,14 +140,6 @@ void AllocatePageIndex(uint2 DispatchThreadId : SV_DispatchThreadID)
 			// Clear all output if the allocation is not valid
 			OutPageIndexResolutionAndOffsetBuffer[MacroGroupId] = uint4(0, 0, 0, 0);
 			WriteDispatchIndirectArgs(OutPageIndexAllocationIndirectBufferArgs, MacroGroupId, 0, 1, 1);
-
-			FVoxelizationViewInfo ViewInfo;
-			ViewInfo.TranslatedWorldToClip = 0;
-			ViewInfo.ViewForward = 0;
-			ViewInfo.RasterResolution = 0;
-			ViewInfo.Pad0 = 0;
-			ViewInfo.Pad1 = 0;
-			OutVoxelizationViewInfoBuffer[MacroGroupId] = ViewInfo;
 		}
 	}
 
@@ -162,69 +153,6 @@ void AllocatePageIndex(uint2 DispatchThreadId : SV_DispatchThreadID)
 		const uint AllocatedPageIndexCount = PageIndexResolution.x * PageIndexResolution.y * PageIndexResolution.z;
 		WriteDispatchIndirectArgs(OutPageIndexAllocationIndirectBufferArgs, MacroGroupId, DivideAndRoundUp(AllocatedPageIndexCount, IndirectDispatchGroupSize), 1, 1);
 	}
-
-	// Matrix generation
-	if (bIsValid)
-	{
-		// Find the largest resolution and its dominant axis
-		uint2 RasterResolution = uint2(0, 0);
-		float3 RasterProjectionSize = float3(0, 0, 0);
-		float3 RasterDirection = float3(0, 0, 0);
-		float3 RasterUp = float3(0, 0, 0);
-		const int3 TotalVoxelResolution = PageIndexResolution * PageResolution;
-		{
-			const uint ResolutionXY = TotalVoxelResolution.x * TotalVoxelResolution.y;
-			const uint ResolutionXZ = TotalVoxelResolution.x * TotalVoxelResolution.z;
-			const uint ResolutionYZ = TotalVoxelResolution.y * TotalVoxelResolution.y;
-			if (ResolutionXY >= ResolutionXZ && ResolutionXY >= ResolutionYZ)
-			{
-				RasterResolution = int2(TotalVoxelResolution.x, TotalVoxelResolution.y);
-				RasterDirection = float3(0, 0, 1);
-				RasterUp = float3(0, 1, 0);
-
-				FHairAABB ProjRasterAABB;
-				ProjRasterAABB.Min = float3(Bound.Min.x, Bound.Min.y, Bound.Min.z);
-				ProjRasterAABB.Max = float3(Bound.Max.x, Bound.Max.y, Bound.Max.z);
-				RasterProjectionSize = ProjRasterAABB.Max - ProjRasterAABB.Min;
-			}
-			else if (ResolutionXZ >= ResolutionXY && ResolutionXZ >= ResolutionYZ)
-			{
-				RasterResolution = int2(TotalVoxelResolution.x, TotalVoxelResolution.z);
-				RasterDirection = float3(0, -1, 0);
-				RasterUp = float3(0, 0, 1);
-
-				FHairAABB ProjRasterAABB;
-				ProjRasterAABB.Min = float3(Bound.Min.x, Bound.Min.z, Bound.Min.y);
-				ProjRasterAABB.Max = float3(Bound.Max.x, Bound.Max.z, Bound.Max.y);
-				RasterProjectionSize = ProjRasterAABB.Max - ProjRasterAABB.Min;
-			}
-			else
-			{
-				RasterResolution = int2(TotalVoxelResolution.y, TotalVoxelResolution.z);
-				RasterDirection = float3(1, 0, 0);
-				RasterUp = float3(0, 0, 1);
-
-				FHairAABB ProjRasterAABB;
-				ProjRasterAABB.Min = float3(Bound.Min.y, Bound.Min.z, Bound.Min.x);
-				ProjRasterAABB.Max = float3(Bound.Max.y, Bound.Max.z, Bound.Max.x);
-				RasterProjectionSize = ProjRasterAABB.Max - ProjRasterAABB.Min;
-			}
-		}
-
-		const float3 RasterAABBSize = Bound.Max - Bound.Min;
-		const float3 RasterAABBCenter = (Bound.Max + Bound.Min) * 0.5f;
-
-		const float4x4 ProjMatrix = ReversedZOrthoMatrix(0.5f * RasterProjectionSize.x, 0.5f * RasterProjectionSize.y, 1.0f / RasterProjectionSize.z, 0);
-		const float4x4 TranslatedViewMatrix = LookAtMatrix(RasterAABBCenter - RasterDirection * RasterProjectionSize.z * 0.5f, RasterAABBCenter, RasterUp);
-
-		FVoxelizationViewInfo ViewInfo;
-		ViewInfo.TranslatedWorldToClip = mul(TranslatedViewMatrix, ProjMatrix);
-		ViewInfo.ViewForward = RasterDirection;
-		ViewInfo.RasterResolution = RasterResolution;
-		ViewInfo.Pad0 = 0;
-		ViewInfo.Pad1 = 0;
-		OutVoxelizationViewInfoBuffer[MacroGroupId] = ViewInfo;
-	}
 }
 #endif
 
diff --git a/Engine/Source/Runtime/Renderer/Private/HairStrands/HairStrandsData.h b/Engine/Source/Runtime/Renderer/Private/HairStrands/HairStrandsData.h
--- a/Engine/Source/Runtime/Renderer/Private/HairStrands/HairStrandsData.h
+++ b/Engine/Source/Runtime/Renderer/Private/HairStrands/HairStrandsData.h
@@ -245,7 +245,6 @@ struct FHairStrandsVoxelResources
 	FRDGBufferRef PageIndexCoordBuffer = nullptr;
 	FRDGBufferRef IndirectArgsBuffer = nullptr;
 	FRDGBufferRef PageIndexGlobalCounter = nullptr;
-	FRDGBufferRef VoxelizationViewInfoBuffer = nullptr;
 
 	const bool IsValid() const { return UniformBuffer != nullptr && PageTexture != nullptr && NodeDescBuffer != nullptr; }
 };
diff --git a/Engine/Source/Runtime/Renderer/Private/HairStrands/HairStrandsVoxelization.cpp b/Engine/Source/Runtime/Renderer/Private/HairStrands/HairStrandsVoxelization.cpp
--- a/Engine/Source/Runtime/Renderer/Private/HairStrands/HairStrandsVoxelization.cpp
+++ b/Engine/Source/Runtime/Renderer/Private/HairStrands/HairStrandsVoxelization.cpp
@@ -261,7 +261,6 @@ class FVoxelAllocatePageIndexCS : public FGlobalShader
 		SHADER_PARAMETER_RDG_BUFFER_SRV(Buffer<int>, MacroGroupAABBBuffer)
 		SHADER_PARAMETER_RDG_BUFFER_UAV(RWBuffer<int>, MacroGroupVoxelAlignedAABBBuffer)
 		SHADER_PARAMETER_RDG_BUFFER_UAV(RWBuffer<uint4>, OutPageIndexResolutionAndOffsetBuffer)
-		SHADER_PARAMETER_RDG_BUFFER_UAV(RWBuffer, OutVoxelizationViewInfoBuffer)
 		SHADER_PARAMETER_RDG_BUFFER_UAV(RWBuffer, OutPageIndexAllocationIndirectBufferArgs)
 	END_SHADER_PARAMETER_STRUCT()
 
@@ -510,7 +509,6 @@ static void AddAllocateVoxelPagesPass(
 	FRDGBufferRef& OutNodeDescBuffer,
 	FRDGBufferRef& OutIndirectArgsBuffer,
 	FRDGBufferRef& OutPageIndexGlobalCounter,
-	FRDGBufferRef& OutVoxelizationViewInfoBuffer,
 	FRDGBufferRef& CurrGPUMinVoxelSize,
 	FRDGBufferRef& NextGPUMinVoxelSize)
 {
@@ -584,9 +582,6 @@ static void AddAllocateVoxelPagesPass(
 	FRDGBufferUAVRef PageIndexOccupancyBufferUAV = GraphBuilder.CreateUAV(PageIndexOccupancyBuffer, PF_R32G32_UINT);
 	FRDGBufferUAVRef PageIndexGlobalCounterUAV = GraphBuilder.CreateUAV(PageIndexGlobalCounter, PF_R32_UINT);
 	
-	// Stored FVoxelizationViewInfo structs
-	// See HairStrandsVoxelPageCommonStruct.ush for more details
-	FRDGBufferRef VoxelizationViewInfoBuffer = GraphBuilder.CreateBuffer(FRDGBufferDesc::CreateStructuredDesc(24 * sizeof(float), MacroGroupCount), TEXT("Hair.VoxelizationViewInfo"));
 	FRDGBufferRef PageIndexResolutionBuffer = GraphBuilder.CreateBuffer(FRDGBufferDesc::CreateBufferDesc(MacroGroupCount * 4 * sizeof(uint32), OutTotalPageIndexCount), TEXT("Hair.PageIndexResolutionBuffer"));
 	FRDGBufferRef PageIndexAllocationIndirectBufferArgs = GraphBuilder.CreateBuffer(FRDGBufferDesc::CreateIndirectDesc<FRHIDispatchIndirectParameters>(MacroGroupCount), TEXT("Hair.PageIndexAllocationIndirectBufferArgs"));
 
@@ -626,7 +621,6 @@ static void AddAllocateVoxelPagesPass(
 		Parameters->MacroGroupVoxelAlignedAABBBuffer = GraphBuilder.CreateUAV(MacroGroupResources.MacroGroupVoxelAlignedAABBsBuffer, PF_R32_SINT);
 		Parameters->IndirectDispatchGroupSize = FVoxelAllocateVoxelPageCS::GetGroupSize();
 		Parameters->OutPageIndexResolutionAndOffsetBuffer = GraphBuilder.CreateUAV(PageIndexResolutionBuffer, PF_R32G32B32A32_UINT);
-		Parameters->OutVoxelizationViewInfoBuffer = GraphBuilder.CreateUAV(VoxelizationViewInfoBuffer);
 		Parameters->OutPageIndexAllocationIndirectBufferArgs = GraphBuilder.CreateUAV(PageIndexAllocationIndirectBufferArgs);
 		Parameters->CPU_bUseCPUData = bIsGPUDriven ? 0 : 1;
 		Parameters->GPUVoxelWorldSize = CurrGPUMinVoxelSizeSRV;
@@ -836,7 +830,6 @@ static void AddAllocateVoxelPagesPass(
 	OutNodeDescBuffer = NodeDescBuffer;
 	OutIndirectArgsBuffer = IndirectArgsBuffer;
 	OutPageIndexGlobalCounter = PageIndexGlobalCounter;
-	OutVoxelizationViewInfoBuffer = VoxelizationViewInfoBuffer;
 }
 
 static float RoundHairVoxeliSize(float In)
@@ -927,7 +920,6 @@ static FHairStrandsVoxelResources AllocateVirtualVoxelResources(
 		Out.NodeDescBuffer,
 		Out.IndirectArgsBuffer,
 		Out.PageIndexGlobalCounter,
-		Out.VoxelizationViewInfoBuffer,
 		CurrGPUMinVoxelSize,
 		NextGPUMinVoxelSize);
 
@@ -1020,7 +1012,6 @@ static FHairStrandsVoxelResources AllocateDummyVirtualVoxelResources(
 	Out.NodeDescBuffer				= DummyStructBuffer;
 	Out.IndirectArgsBuffer			= Dummy4BytesBuffer;
 	Out.PageIndexGlobalCounter		= Dummy4BytesBuffer;
-	Out.VoxelizationViewInfoBuffer	= DummyStructBuffer;
 	Out.PageTexture					= GSystemTextures.GetDefaultTexture(GraphBuilder, ETextureDimension::Texture3D, PF_R32_UINT, 0u);
 
 	Out.Parameters.Common.PageIndexBuffer			= GraphBuilder.CreateSRV(Out.PageIndexBuffer, PF_R32_UINT);
@@ -1108,7 +1099,6 @@ class FVoxelRasterComputeCS : public FGlobalShader
 		SHADER_PARAMETER(uint32, MaxRasterCount)
 		SHADER_PARAMETER(uint32, FrameIdMod8)
 		SHADER_PARAMETER_STRUCT_INCLUDE(FHairStrandsInstanceParameters, HairInstance)
-		SHADER_PARAMETER_RDG_BUFFER_SRV(Buffer, VoxelizationViewInfoBuffer)
 		SHADER_PARAMETER_RDG_TEXTURE_UAV(RWTexture3D, OutPageTexture)
 		SHADER_PARAMETER(uint32, VertexCount)
 		SHADER_PARAMETER(float, CoverageScale)
@@ -1141,7 +1131,6 @@ static void AddVirtualVoxelizationComputeRasterPass(
 	{
 		const FHairStrandsMacroGroupData::TPrimitiveInfos& PrimitiveSceneInfos = MacroGroup.PrimitivesInfos;
 
-		FRDGBufferSRVRef VoxelizationViewInfoBufferSRV = GraphBuilder.CreateSRV(VoxelResources.VoxelizationViewInfoBuffer);
 		FRDGTextureUAVRef PageTextureUAV = GraphBuilder.CreateUAV(VoxelResources.PageTexture);
 
 		const uint32 FrameIdMode8 = ViewInfo && ViewInfo->ViewState ? (ViewInfo->ViewState->GetFrameIndex() % 8) : 0;
@@ -1177,7 +1166,6 @@ static void AddVirtualVoxelizationComputeRasterPass(
 				PassParameters->MaxRasterCount = FMath::Clamp(GHairStrandsVoxelComputeRasterMaxVoxelCount, 1, 256);
 				PassParameters->VirtualVoxelParams = VoxelResources.Parameters.Common;
 				PassParameters->MacroGroupId = MacroGroup.MacroGroupId;
-				PassParameters->VoxelizationViewInfoBuffer = VoxelizationViewInfoBufferSRV;
 				PassParameters->OutPageTexture = PageTextureUAV;
 				PassParameters->FrameIdMod8 = FrameIdMode8;
 
