Strata Public API for Niagara and other external users

diff --git a/Engine/Plugins/FX/Niagara/Shaders/Private/NiagaraDataInterfaceCollisionQuery.ush b/Engine/Plugins/FX/Niagara/Shaders/Private/NiagaraDataInterfaceCollisionQuery.ush
--- a/Engine/Plugins/FX/Niagara/Shaders/Private/NiagaraDataInterfaceCollisionQuery.ush
+++ b/Engine/Plugins/FX/Niagara/Shaders/Private/NiagaraDataInterfaceCollisionQuery.ush
@@ -3,6 +3,7 @@
 #pragma once
 
 #include "/Engine/Private/Common.ush"
+#include "/Engine/Public/StrataPublic.ush"
 
 #define NDICOLLISIONQUERY_USE_GBUFFER_NORMAL    (FEATURE_LEVEL >= FEATURE_LEVEL_SM5)
 
@@ -39,7 +40,12 @@ void NDICollisionQuery_QuerySceneDepthGPU(in float3 In_SamplePos, in float3 In_L
 		FLWCVector3 WorldPosition = WorldPositionFromSceneDepth(ScreenPosition.xy, SceneDepth);
 
 #if NDICOLLISIONQUERY_USE_GBUFFER_NORMAL
+#if STRATA_ENABLED
+		uint2 ScreenPos = uint2(ScreenUV * View.BufferSizeAndInvSize.xy);
+		float3 WorldNormal = StrataPublic_GetWorldNormal(ScreenPos);
+#else
 		float3 WorldNormal = Texture2DSampleLevel(SceneTexturesStruct.GBufferATexture, SceneTexturesStruct_GBufferATextureSampler, ScreenUV, 0).xyz * 2.0 - 1.0;
+#endif
 #else
 		float CollisionDepthBounds = 500.0f;
 		float SceneDepth0 = CalcSceneDepth(ScreenUV + float2(View.BufferSizeAndInvSize.z, 0.0));
@@ -96,7 +102,13 @@ void NDICollisionQuery_QueryCustomDepthGPU(in float3 In_SamplePos, in float3 In_
 		FLWCVector3 WorldPosition = WorldPositionFromSceneDepth(ScreenPosition.xy, SceneDepth);
 
 #if NDICOLLISIONQUERY_USE_GBUFFER_NORMAL
+#if STRATA_ENABLED
+		uint2 ScreenPos = uint2(ScreenUV * View.BufferSizeAndInvSize.xy);
+		float3 WorldNormal = StrataPublic_GetWorldNormal(ScreenPos);
+#else
 		float3 WorldNormal = Texture2DSampleLevel(SceneTexturesStruct.GBufferATexture, SceneTexturesStruct_GBufferATextureSampler, ScreenUV, 0).xyz * 2.0 - 1.0;
+#endif
+
 #else
 		float CollisionDepthBounds = 500.0f;
 		float SceneDepth0 = NDICollisionQuery_GetCustomSceneDepth(ScreenUV + float2(View.BufferSizeAndInvSize.z, 0.0));
diff --git a/Engine/Plugins/FX/Niagara/Shaders/Private/NiagaraDataInterfaceGBufferTemplate.ush b/Engine/Plugins/FX/Niagara/Shaders/Private/NiagaraDataInterfaceGBufferTemplate.ush
--- a/Engine/Plugins/FX/Niagara/Shaders/Private/NiagaraDataInterfaceGBufferTemplate.ush
+++ b/Engine/Plugins/FX/Niagara/Shaders/Private/NiagaraDataInterfaceGBufferTemplate.ush
@@ -1,6 +1,7 @@
 // Copyright Epic Games, Inc. All Rights Reserved.
 
 #include "/Engine/Private/DeferredShadingCommon.ush"
+#include "/Engine/Public/StrataPublic.ush"
 
 #define NDIGBUFFER_HAS_GBUFFER				(FEATURE_LEVEL >= FEATURE_LEVEL_SM5)
 #define NDIGBUFFER_USE_GBUFFER_VELOCITY		(false && NDIGBUFFER_HAS_GBUFFER)
@@ -37,9 +38,15 @@ void DecodeWorldNormal_{ParameterName}(in float2 ScreenUV, in bool bApplyViewpor
 {
 #if NDIGBUFFER_HAS_GBUFFER
 	ScreenUV = ApplyViewportOffset_{ParameterName}(ScreenUV, bApplyViewportOffset);
-	FGBufferData GBufferData = GetGBufferData(ScreenUV);
 	IsValid = true;
-	OutValue = GBufferData.WorldNormal;
+
+	#if STRATA_ENABLED
+		uint2 ScreenPos = uint2(ScreenUV * View.BufferSizeAndInvSize.xy);
+		OutValue = StrataPublic_GetWorldNormal(ScreenPos);
+	#else
+		FGBufferData GBufferData = GetGBufferData(ScreenUV);
+		OutValue = GBufferData.WorldNormal;
+	#endif
 #else
 	IsValid = false;
 	OutValue = 0;
diff --git a/Engine/Plugins/FX/Niagara/Shaders/Private/NiagaraEmitterInstanceShader.usf b/Engine/Plugins/FX/Niagara/Shaders/Private/NiagaraEmitterInstanceShader.usf
--- a/Engine/Plugins/FX/Niagara/Shaders/Private/NiagaraEmitterInstanceShader.usf
+++ b/Engine/Plugins/FX/Niagara/Shaders/Private/NiagaraEmitterInstanceShader.usf
@@ -18,6 +18,7 @@ NiagaraSimulationShader.usf:
 #if GPU_SIMULATION
 	#include "/Engine/Private/Common.ush"
 	#include "/Engine/Private/DistanceField/GlobalDistanceFieldShared.ush"
+	#include "/Engine/Generated/UniformBuffers/StrataPublic.ush"
 #else
 	const static float PI = 3.1415926535897932f;
 #endif
diff --git a/Engine/Plugins/FX/Niagara/Source/Niagara/Private/NiagaraGpuComputeDispatch.cpp b/Engine/Plugins/FX/Niagara/Source/Niagara/Private/NiagaraGpuComputeDispatch.cpp
--- a/Engine/Plugins/FX/Niagara/Source/Niagara/Private/NiagaraGpuComputeDispatch.cpp
+++ b/Engine/Plugins/FX/Niagara/Source/Niagara/Private/NiagaraGpuComputeDispatch.cpp
@@ -1055,6 +1055,8 @@ void FNiagaraGpuComputeDispatch::ExecuteTicks(FRDGBuilder& GraphBuilder, TConstS
 		MobileSceneTexturesUniformParams = UE::FXRenderingUtils::GetOrCreateMobileSceneTextureUniformBuffer(GraphBuilder, SimulationSceneViews, EMobileSceneTextureSetupMode::None);
 	}
 
+	StrataPublicGlobalUniformParams = ::Strata::GetPublicGlobalUniformBuffer(GraphBuilder, *GetScene());
+
 	// Loop over dispatches
 	for ( const FNiagaraGpuDispatchGroup& DispatchGroup : DispatchList.DispatchGroups )
 	{
@@ -1332,6 +1334,7 @@ void FNiagaraGpuComputeDispatch::ExecuteTicks(FRDGBuilder& GraphBuilder, TConstS
 	SimulationSceneViews = TConstStridedView<FSceneView>();
 	SceneTexturesUniformParams = nullptr;
 	MobileSceneTexturesUniformParams = nullptr;
+	StrataPublicGlobalUniformParams = nullptr;
 
 	CurrentPassExternalAccessQueue.Submit(GraphBuilder);
 }
@@ -1591,6 +1594,7 @@ void FNiagaraGpuComputeDispatch::DispatchStage(FRDGBuilder& GraphBuilder, const
 	}
 	DispatchParameters->SceneTextures.SceneTextures			= SceneTexturesUniformParams;
 	DispatchParameters->SceneTextures.MobileSceneTextures	= MobileSceneTexturesUniformParams;
+	DispatchParameters->StrataPublic						= StrataPublicGlobalUniformParams;
 
 	// Indirect Gpu Dispatch
 	if (SimStageData.StageMetaData->bGpuIndirectDispatch)
diff --git a/Engine/Plugins/FX/Niagara/Source/Niagara/Private/NiagaraGpuComputeDispatch.h b/Engine/Plugins/FX/Niagara/Source/Niagara/Private/NiagaraGpuComputeDispatch.h
--- a/Engine/Plugins/FX/Niagara/Source/Niagara/Private/NiagaraGpuComputeDispatch.h
+++ b/Engine/Plugins/FX/Niagara/Source/Niagara/Private/NiagaraGpuComputeDispatch.h
@@ -284,4 +284,5 @@ private:
 
 	TRDGUniformBufferRef<FSceneTextureUniformParameters> SceneTexturesUniformParams = nullptr;
 	TRDGUniformBufferRef<FMobileSceneTextureUniformParameters> MobileSceneTexturesUniformParams = nullptr;
+	TRDGUniformBufferRef<FStrataPublicGlobalUniformParameters> StrataPublicGlobalUniformParams = nullptr;
 };
diff --git a/Engine/Plugins/FX/Niagara/Source/NiagaraShader/Public/NiagaraShader.h b/Engine/Plugins/FX/Niagara/Source/NiagaraShader/Public/NiagaraShader.h
--- a/Engine/Plugins/FX/Niagara/Source/NiagaraShader/Public/NiagaraShader.h
+++ b/Engine/Plugins/FX/Niagara/Source/NiagaraShader/Public/NiagaraShader.h
@@ -20,6 +20,8 @@
 #if UE_ENABLE_INCLUDE_ORDER_DEPRECATED_IN_5_3
 #include "RenderGraph.h"
 #include "SceneRenderTargetParameters.h"
+
+#include "RendererUtils.h"
 #endif
 
 class UClass;
@@ -201,6 +203,8 @@ public:
 
 		SHADER_PARAMETER_STRUCT_REF(FViewUniformShaderParameters,	View)
 		SHADER_PARAMETER_STRUCT_INCLUDE(FSceneTextureShaderParameters,	SceneTextures)
+
+		SHADER_PARAMETER_RDG_UNIFORM_BUFFER(FStrataPublicGlobalUniformParameters, StrataPublic)
 	END_SHADER_PARAMETER_STRUCT()
 
 	using FPermutationParameters = FNiagaraShaderPermutationParameters;
diff --git a/Engine/Shaders/Public/StrataPublic.ush b/Engine/Shaders/Public/StrataPublic.ush
new file mode 100644
--- /dev/null
+++ b/Engine/Shaders/Public/StrataPublic.ush
@@ -0,0 +1,17 @@
+// Copyright Epic Games, Inc. All Rights Reserved.
+
+#pragma once
+
+/*=============================================================================
+	StrataPublic.ush: Public facing API for external users of Strata.
+=============================================================================*/
+
+#include "/Engine/Private/Strata/Strata.ush"
+
+/** Public facing functions allowing external users access to some Strata data. */
+float3 StrataPublic_GetWorldNormal(uint2 PixelCoord)
+{
+  return StrataUnpackTopLayerData(StrataPublic.TopLayerTexture.Load(uint3(PixelCoord, 0))).WorldNormal;
+}
+
+//TODO: Implement other accessors that were previously available via SceneTextures.
diff --git a/Engine/Source/Runtime/Core/Private/UObject/UE5ReleaseStreamObjectVersion.cpp b/Engine/Source/Runtime/Core/Private/UObject/UE5ReleaseStreamObjectVersion.cpp
--- a/Engine/Source/Runtime/Core/Private/UObject/UE5ReleaseStreamObjectVersion.cpp
+++ b/Engine/Source/Runtime/Core/Private/UObject/UE5ReleaseStreamObjectVersion.cpp
@@ -11,8 +11,8 @@ TMap<FGuid, FGuid> FUE5ReleaseStreamObjectVersion::GetSystemGuids()
 	SystemGuids.Add(DevGuids.GLOBALSHADERMAP_DERIVEDDATA_VER, FGuid("D0BF3452816D46908073DFDD4B855AE5"));
 	SystemGuids.Add(DevGuids.MATERIALSHADERMAP_DERIVEDDATA_VER, FGuid("79090A66F9D94B5BB285D49E5D39468E"));
 	SystemGuids.Add(DevGuids.NANITE_DERIVEDDATA_VER, FGuid("AC88CFBDDA614A9C8488F64D2DAF699D"));
-	SystemGuids.Add(DevGuids.NIAGARASHADERMAP_DERIVEDDATA_VER, FGuid("D204CCC166ED482685DAEE8AB72895AB"));
-	SystemGuids.Add(DevGuids.Niagara_LatestScriptCompileVersion, FGuid("29F4804724014091BC9B1AF11BA53330"));
+	SystemGuids.Add(DevGuids.NIAGARASHADERMAP_DERIVEDDATA_VER, FGuid("6BE35B6FACB34568970120F9BC8DAB80"));
+	SystemGuids.Add(DevGuids.Niagara_LatestScriptCompileVersion, FGuid("4882CE3BA4DF4D6F8D8B7828CED1C1A4"));
 	SystemGuids.Add(DevGuids.SkeletalMeshDerivedDataVersion, FGuid("25C49E579B3142DDA2A8C14037267679"));
 
 	return SystemGuids;
diff --git a/Engine/Source/Runtime/Renderer/Private/DeferredShadingRenderer.cpp b/Engine/Source/Runtime/Renderer/Private/DeferredShadingRenderer.cpp
--- a/Engine/Source/Runtime/Renderer/Private/DeferredShadingRenderer.cpp
+++ b/Engine/Source/Runtime/Renderer/Private/DeferredShadingRenderer.cpp
@@ -2675,6 +2675,8 @@ void FDeferredShadingSceneRenderer::Render(FRDGBuilder& GraphBuilder)
 
 	FInitViewTaskDatas InitViewTaskDatas;
 
+	::Strata::PreInitViews(*Scene);
+
 	{
 		RDG_GPU_STAT_SCOPE(GraphBuilder, VisibilityCommands);
 		BeginInitViews(GraphBuilder, SceneTexturesConfig, BasePassDepthStencilAccess, InstanceCullingManager, VirtualTextureUpdater.Get(), InitViewTaskDatas);
@@ -4313,6 +4315,8 @@ void FDeferredShadingSceneRenderer::Render(FRDGBuilder& GraphBuilder)
 
 	QueueSceneTextureExtractions(GraphBuilder, SceneTextures);
 
+	::Strata::PostRender(*Scene);
+
 	// Release the view's previous frame histories so that their memory can be reused at the graph's execution.
 	for (int32 ViewIndex = 0; ViewIndex < Views.Num(); ViewIndex++)
 	{
diff --git a/Engine/Source/Runtime/Renderer/Private/RendererUtils.cpp b/Engine/Source/Runtime/Renderer/Private/RendererUtils.cpp
--- a/Engine/Source/Runtime/Renderer/Private/RendererUtils.cpp
+++ b/Engine/Source/Runtime/Renderer/Private/RendererUtils.cpp
@@ -5,6 +5,8 @@
 #include "RHIDefinitions.h"
 #include "DataDrivenShaderPlatformInfo.h"
 #include "VisualizeTexture.h"
+#include "ScenePrivate.h"
+#include "SystemTextures.h"
 
 class FRTWriteMaskDecodeCS : public FGlobalShader
 {
@@ -189,3 +191,36 @@ FDepthBounds::FDepthBoundsValues FDepthBounds::CalculateNearFarDepthExcludingSky
 
 	return Values;
 }
+
+IMPLEMENT_GLOBAL_SHADER_PARAMETER_STRUCT(FStrataPublicGlobalUniformParameters, "StrataPublic");
+
+namespace Strata
+{
+	void PreInitViews(FScene& Scene)
+	{
+		FStrataSceneData& StrataScene = Scene.StrataSceneData;
+		StrataScene.StrataPublicGlobalUniformParameters = nullptr;
+	}
+
+	void PostRender(FScene& Scene)
+	{
+		FStrataSceneData& StrataScene = Scene.StrataSceneData;
+		StrataScene.StrataPublicGlobalUniformParameters = nullptr;
+	}
+
+	TRDGUniformBufferRef<FStrataPublicGlobalUniformParameters> GetPublicGlobalUniformBuffer(FRDGBuilder& GraphBuilder, FScene& Scene)
+	{		
+		if(::Strata::IsStrataEnabled())
+		{
+			FStrataSceneData& StrataScene = Scene.StrataSceneData;
+
+			if(StrataScene.StrataPublicGlobalUniformParameters == nullptr)
+			{
+				return CreatePublicGlobalUniformBuffer(GraphBuilder, nullptr);//We are creating a dummy here so pass in null for the scene data.
+			}
+			return StrataScene.StrataPublicGlobalUniformParameters;
+		}
+
+		return nullptr;
+	}
+}
diff --git a/Engine/Source/Runtime/Renderer/Private/Strata/Strata.cpp b/Engine/Source/Runtime/Renderer/Private/Strata/Strata.cpp
--- a/Engine/Source/Runtime/Renderer/Private/Strata/Strata.cpp
+++ b/Engine/Source/Runtime/Renderer/Private/Strata/Strata.cpp
@@ -428,6 +428,11 @@ void InitialiseStrataFrameSceneData(FRDGBuilder& GraphBuilder, FSceneRenderer& S
 	{
 		Strata::InitialiseStrataViewData(GraphBuilder, SceneRenderer.Views[ViewIndex], SceneRenderer.GetActiveSceneTexturesConfig(), bNeedBSDFOffsets, Out);
 	}
+
+	if (IsStrataEnabled())
+	{
+		Out.StrataPublicGlobalUniformParameters = ::Strata::CreatePublicGlobalUniformBuffer(GraphBuilder, &Out);
+	}
 }
 
 void BindStrataBasePassUniformParameters(FRDGBuilder& GraphBuilder, const FViewInfo& View, FStrataBasePassUniformParameters& OutStrataUniformParameters)
@@ -567,6 +572,29 @@ TRDGUniformBufferRef<FStrataGlobalUniformParameters> BindStrataGlobalUniformPara
 	return View.StrataViewData.StrataGlobalUniformParameters;
 }
 
+static void BindStrataPublicGlobalUniformParameters(FRDGBuilder& GraphBuilder, FStrataSceneData* StrataSceneData, FStrataPublicGlobalUniformParameters& OutStrataUniformParameters)
+{
+	const FRDGSystemTextures& SystemTextures = FRDGSystemTextures::Get(GraphBuilder);
+	if (StrataSceneData && StrataSceneData->TopLayerTexture)
+	{
+		OutStrataUniformParameters.TopLayerTexture = StrataSceneData->TopLayerTexture;
+	}
+	else
+	{
+		OutStrataUniformParameters.TopLayerTexture = SystemTextures.Black;
+	}
+
+	//TODO: Other Strata scene textures or other globals.
+}
+
+TRDGUniformBufferRef<FStrataPublicGlobalUniformParameters> CreatePublicGlobalUniformBuffer(FRDGBuilder& GraphBuilder, FStrataSceneData* StrataScene)
+{
+	FStrataPublicGlobalUniformParameters* StrataPublicUniformParameters = GraphBuilder.AllocParameters<FStrataPublicGlobalUniformParameters>();
+	check(StrataPublicUniformParameters);
+	BindStrataPublicGlobalUniformParameters(GraphBuilder, StrataScene, *StrataPublicUniformParameters);
+	return GraphBuilder.CreateUniformBuffer(StrataPublicUniformParameters);
+}
+
 ////////////////////////////////////////////////////////////////////////////////////////////////////////////////////////////////////////////////
 
 class FStrataBSDFTilePassCS : public FGlobalShader
diff --git a/Engine/Source/Runtime/Renderer/Private/Strata/Strata.h b/Engine/Source/Runtime/Renderer/Private/Strata/Strata.h
--- a/Engine/Source/Runtime/Renderer/Private/Strata/Strata.h
+++ b/Engine/Source/Runtime/Renderer/Private/Strata/Strata.h
@@ -11,6 +11,7 @@
 #include "StrataDefinitions.h"
 #include "RenderGraphUtils.h"
 #include "GBufferInfo.h"
+#include "RendererUtils.h"
 
 // Forward declarations.
 class FScene;
@@ -119,6 +120,9 @@ struct FStrataSceneData
 
 	// Used for Luminance that should go through opaque rough refraction (when under a top layer interface)
 	FRDGTextureRef SeparatedOpaqueRoughRefractionSceneColor = nullptr;
+
+	//Public facing minimal uniform data.
+	TRDGUniformBufferRef<FStrataPublicGlobalUniformParameters> StrataPublicGlobalUniformParameters{};
 };
 
 struct FStrataViewData
@@ -171,6 +175,7 @@ TRDGUniformBufferRef<FStrataGlobalUniformParameters> BindStrataGlobalUniformPara
 void AppendStrataMRTs(const FSceneRenderer& SceneRenderer, uint32& BasePassTextureCount, TStaticArray<FTextureRenderTargetBinding, MaxSimultaneousRenderTargets>& BasePassTextures);
 void SetBasePassRenderTargetOutputFormat(const EShaderPlatform Platform, const FMaterialShaderParameters& MaterialParameters, FShaderCompilerEnvironment& OutEnvironment, EGBufferLayout GBufferLayout);
 
+TRDGUniformBufferRef<FStrataPublicGlobalUniformParameters> CreatePublicGlobalUniformBuffer(FRDGBuilder& GraphBuilder, FStrataSceneData* StrataScene);
 
 void AddStrataMaterialClassificationPass(FRDGBuilder& GraphBuilder, const FMinimalSceneTextures& SceneTextures, const FDBufferTextures& DBufferTextures, const TArray<FViewInfo>& Views);
 void AddStrataDBufferPass(FRDGBuilder& GraphBuilder, const FMinimalSceneTextures& SceneTextures, const FDBufferTextures& DBufferTextures, const TArray<FViewInfo>& Views);
diff --git a/Engine/Source/Runtime/Renderer/Public/RendererUtils.h b/Engine/Source/Runtime/Renderer/Public/RendererUtils.h
--- a/Engine/Source/Runtime/Renderer/Public/RendererUtils.h
+++ b/Engine/Source/Runtime/Renderer/Public/RendererUtils.h
@@ -8,6 +8,8 @@
 #include "RHIDefinitions.h"
 #include "RenderGraph.h"
 #include "RenderGraphDefinitions.h"
+#include "ShaderParameterMacros.h"
+#include "RenderGraphFwd.h"
 
 class FGlobalShaderMap;
 class FRDGBuilder;
@@ -47,3 +49,18 @@ public:
 
 	static FDepthBoundsValues CalculateNearFarDepthExcludingSky();
 };
+
+// A minimal uniform struct providing necessary access for external systems to Strata parameters.
+BEGIN_GLOBAL_SHADER_PARAMETER_STRUCT(FStrataPublicGlobalUniformParameters, RENDERER_API)
+	SHADER_PARAMETER_RDG_TEXTURE(Texture2D<uint2>, TopLayerTexture)
+END_GLOBAL_SHADER_PARAMETER_STRUCT()
+
+struct FStrataSceneData;
+
+namespace Strata
+{
+	void PreInitViews(FScene& Scene);
+	void PostRender(FScene& Scene);
+
+	RENDERER_API TRDGUniformBufferRef<FStrataPublicGlobalUniformParameters> GetPublicGlobalUniformBuffer(FRDGBuilder& GraphBuilder, FScene& Scene);
+}
