- Forward declare FScene to fix no unity issues

#rnx

diff --git a/Engine/Source/Runtime/Renderer/Public/RendererUtils.h b/Engine/Source/Runtime/Renderer/Public/RendererUtils.h
--- a/Engine/Source/Runtime/Renderer/Public/RendererUtils.h
+++ b/Engine/Source/Runtime/Renderer/Public/RendererUtils.h
@@ -14,6 +14,7 @@
 class FGlobalShaderMap;
 class FRDGBuilder;
 class FRHICommandListImmediate;
+class FScene;
 struct IPooledRenderTarget;
 template <typename ReferencedType> class TRefCountPtr;
 
