- Handle nullptr for SceneInterface

#rnx

diff --git a/Engine/Source/Runtime/Renderer/Private/FXRenderingUtils.cpp b/Engine/Source/Runtime/Renderer/Private/FXRenderingUtils.cpp
--- a/Engine/Source/Runtime/Renderer/Private/FXRenderingUtils.cpp
+++ b/Engine/Source/Runtime/Renderer/Private/FXRenderingUtils.cpp
@@ -115,7 +115,7 @@ const FShaderParametersMetadata* UE::FXRenderingUtils::DistanceFields::GetAtlasP
 inline const FDistanceFieldSceneData* GetDistanceFieldSceneData(const FSceneView& View)
 {
 	FSceneInterface* SceneInterface = View.Family && View.Family->Scene ? View.Family->Scene : nullptr;
-	const FScene* Scene = SceneInterface->GetRenderScene();
+	const FScene* Scene = SceneInterface ? SceneInterface->GetRenderScene() : nullptr;
 	return  Scene ? &Scene->DistanceFieldSceneData : nullptr;
 }
 
