Fix CIS issue.

diff --git a/Engine/Source/Runtime/Renderer/Private/HairStrands/HairStrandsVoxelization.cpp b/Engine/Source/Runtime/Renderer/Private/HairStrands/HairStrandsVoxelization.cpp
--- a/Engine/Source/Runtime/Renderer/Private/HairStrands/HairStrandsVoxelization.cpp
+++ b/Engine/Source/Runtime/Renderer/Private/HairStrands/HairStrandsVoxelization.cpp
@@ -124,7 +124,7 @@ static FAutoConsoleVariableRef CVarHairVirtualVoxelAdaptive_CorrectionThreshold(
 static int32 GHairVirtualVoxel_JitterMode = 1;
 static FAutoConsoleVariableRef CVarHairVirtualVoxel_JitterMode(TEXT("r.HairStrands.Voxelization.Virtual.Jitter"), GHairVirtualVoxel_JitterMode, TEXT("Change jittered for voxelization/traversal. 0: No jitter 1: Regular randomized jitter: 2: Constant Jitter (default = 1)"), ECVF_RenderThreadSafe);
 
-bool IsHairStrandsContinousLODEnabled()
+static bool IsHairStrandsContinousLODEnabled()
 {
 	static IConsoleVariable* CVar = IConsoleManager::Get().FindConsoleVariable(TEXT("r.HairStrands.ContinuousLOD"));
 	return CVar && CVar->GetInt() > 0;
