Change Light Grid Data buffers to be StructuredBuffers.

diff --git a/Engine/Shaders/Private/LightGridCommon.ush b/Engine/Shaders/Private/LightGridCommon.ush
--- a/Engine/Shaders/Private/LightGridCommon.ush
+++ b/Engine/Shaders/Private/LightGridCommon.ush
@@ -16,6 +16,16 @@
 	#endif
 #endif
 
+// If this is changed, the LightGridUses16BitBuffers function from LightGridInjection.cpp should also be updated.
+#define LIGHT_GRID_USES_16BIT_BUFFERS (PLATFORM_SUPPORTS_BUFFER_LOAD_TYPE_CONVERSION && !SHADING_PATH_MOBILE)
+
+// ForwardLightData.CulledLightDataGrid could be StructuredBuffer or 16bit Buffer, make sure to access it from this MACRO
+#if LIGHT_GRID_USES_16BIT_BUFFERS
+	#define GetCulledLightDataGrid(UB, GridIndex) UB.CulledLightDataGrid16Bit[GridIndex]
+#else
+	#define GetCulledLightDataGrid(UB, GridIndex) UB.CulledLightDataGrid32Bit[GridIndex]
+#endif
+
 struct FLightGridData
 {
 	uint LightGridPixelSizeShift;
@@ -205,7 +215,7 @@ FLocalLightData GetLocalLightData(uint GridIndex, uint EyeIndex)
 	if (EyeIndex == 0)
 	{
 #endif
-		uint LocalLightIndex = ForwardLightData.CulledLightDataGrid[GridIndex];
+		uint LocalLightIndex = GetCulledLightDataGrid(ForwardLightData, GridIndex);
 		uint LocalLightBaseIndex = LocalLightIndex * LOCAL_LIGHT_DATA_STRIDE;
 
 		Out = GetLocalLightData_Internal(
@@ -223,7 +233,7 @@ FLocalLightData GetLocalLightData(uint GridIndex, uint EyeIndex)
 	}
 	else
 	{
-		uint LocalLightIndex = ForwardLightDataISR.CulledLightDataGrid[GridIndex];
+		uint LocalLightIndex = GetCulledLightDataGrid(ForwardLightDataISR, GridIndex);
 		uint LocalLightBaseIndex = LocalLightIndex * LOCAL_LIGHT_DATA_STRIDE;
 
 		Out = GetLocalLightData_Internal(
diff --git a/Engine/Shaders/Private/LightGridInjection.usf b/Engine/Shaders/Private/LightGridInjection.usf
--- a/Engine/Shaders/Private/LightGridInjection.usf
+++ b/Engine/Shaders/Private/LightGridInjection.usf
@@ -6,15 +6,23 @@
 #include "Definitions.usf"
 #include "LargeWorldCoordinates.ush"
 #include "ReflectionEnvironmentShared.ush"
+#include "LightGridCommon.ush"
 
-RWBuffer<uint> RWNumCulledLightsGrid;
-RWBuffer<uint> RWCulledLightDataGrid;
+RWStructuredBuffer<uint> RWNumCulledLightsGrid;
 
-RWBuffer<uint> RWNextCulledLightLink;
-RWBuffer<uint> RWStartOffsetGrid;
-RWBuffer<uint> RWCulledLightLinks;
+#if LIGHT_GRID_USES_16BIT_BUFFERS
+	RWBuffer<uint>					RWCulledLightDataGrid16Bit;
+	#define RWCulledLightDataGrid	RWCulledLightDataGrid16Bit
+#else
+	RWStructuredBuffer<uint>		RWCulledLightDataGrid32Bit;
+	#define RWCulledLightDataGrid	RWCulledLightDataGrid32Bit
+#endif
 
-Buffer<float4> ForwardLocalLightBuffer;
+RWStructuredBuffer<uint> RWNextCulledLightLink;
+RWStructuredBuffer<uint> RWStartOffsetGrid;
+RWStructuredBuffer<uint> RWCulledLightLinks;
+
+StructuredBuffer<float4> ForwardLocalLightBuffer;
 
 uint NumReflectionCaptures;
 uint NumLocalLights;
@@ -25,8 +33,8 @@ float3 LightGridZParams;
 uint LightGridPixelSizeShift;
 
 #if ENABLE_LIGHT_CULLING_VIEW_SPACE_BUILD_DATA
-Buffer<float4> LightViewSpacePositionAndRadius;
-Buffer<float4> LightViewSpaceDirAndPreprocAngle;
+StructuredBuffer<float4> LightViewSpacePositionAndRadius;
+StructuredBuffer<float4> LightViewSpaceDirAndPreprocAngle;
 #endif // ENABLE_LIGHT_CULLING_VIEW_SPACE_BUILD_DATA
 
 #define REFINE_SPOTLIGHT_BOUNDS 1
@@ -320,9 +328,9 @@ void LightGridInjectionCS(
 
 #if SHADER_LIGHT_GRID_COMPACT_CS
 
-RWBuffer<uint> RWNextCulledLightData;
-Buffer<uint> StartOffsetGrid;
-Buffer<uint> CulledLightLinks;
+RWStructuredBuffer<uint> RWNextCulledLightData;
+StructuredBuffer<uint> StartOffsetGrid;
+StructuredBuffer<uint> CulledLightLinks;
 
 #if USE_LOCALSTORE_FOR_ATOMICS
 groupshared uint LDSNextCulledLightData;
@@ -421,7 +429,6 @@ void LightGridCompactCS(
 #ifdef SHADER_DEBUG_LIGHT_GRID_PS
 
 #define SUPPORT_CONTACT_SHADOWS 0
-#include "LightGridCommon.ush"
 #include "/Engine/Private/ShaderPrint.ush"
 #include "ColorMap.ush"
 
diff --git a/Engine/Shaders/Private/ReflectionEnvironmentComposite.ush b/Engine/Shaders/Private/ReflectionEnvironmentComposite.ush
--- a/Engine/Shaders/Private/ReflectionEnvironmentComposite.ush
+++ b/Engine/Shaders/Private/ReflectionEnvironmentComposite.ush
@@ -4,7 +4,10 @@
 	ReflectionEnvironmentComposite
 =============================================================================*/
 
+#pragma once
+
 #include "LargeWorldCoordinates.ush"
+#include "LightGridCommon.ush"
 
 float3 CompositeReflectionCapturesAndSkylightTWS(
 	float CompositeAlpha, 
@@ -45,13 +48,13 @@ float3 CompositeReflectionCapturesAndSkylightTWS(
 		{
 	#endif
 	
-			CaptureIndex = ForwardLightData.CulledLightDataGrid[CaptureDataStartIndex + TileCaptureIndex];
+			CaptureIndex = GetCulledLightDataGrid(ForwardLightData, CaptureDataStartIndex + TileCaptureIndex);
 	
 	#if INSTANCED_STEREO
 		}
 		else
 		{
-			CaptureIndex = ForwardLightDataISR.CulledLightDataGrid[CaptureDataStartIndex + TileCaptureIndex];
+			CaptureIndex = GetCulledLightDataGrid(ForwardLightDataISR, CaptureDataStartIndex + TileCaptureIndex);
 		}
 	#endif
 #endif
diff --git a/Engine/Shaders/Private/VirtualShadowMaps/VirtualShadowMapPageManagement.usf b/Engine/Shaders/Private/VirtualShadowMaps/VirtualShadowMapPageManagement.usf
--- a/Engine/Shaders/Private/VirtualShadowMaps/VirtualShadowMapPageManagement.usf
+++ b/Engine/Shaders/Private/VirtualShadowMaps/VirtualShadowMapPageManagement.usf
@@ -173,7 +173,7 @@ void PruneLightGridCS(uint GridLinearIndex : SV_DispatchThreadID)
 
 				// Copy light grid data index
 				checkStructuredBufferAccessSlow(OutPrunedLightGridData, CulledLightGridData.DataStartIndex + NumRetainedLights);
-				OutPrunedLightGridData[CulledLightGridData.DataStartIndex + NumRetainedLights++] = ForwardLightData.CulledLightDataGrid[CulledLightGridData.DataStartIndex + Index];
+				OutPrunedLightGridData[CulledLightGridData.DataStartIndex + NumRetainedLights++] = GetCulledLightDataGrid(ForwardLightData, CulledLightGridData.DataStartIndex + Index);
 			}
 		}
 	}
diff --git a/Engine/Source/Runtime/Engine/Private/ShaderCompiler/ShaderCompiler.cpp b/Engine/Source/Runtime/Engine/Private/ShaderCompiler/ShaderCompiler.cpp
--- a/Engine/Source/Runtime/Engine/Private/ShaderCompiler/ShaderCompiler.cpp
+++ b/Engine/Source/Runtime/Engine/Private/ShaderCompiler/ShaderCompiler.cpp
@@ -6500,6 +6500,7 @@ void GlobalBeginCompileShader(
 	Input.Environment.SetDefine(TEXT("PLATFORM_SUPPORTS_MESH_SHADERS_TIER0"), RHISupportsMeshShadersTier0(EShaderPlatform(Target.Platform)) ? 1 : 0);
 	Input.Environment.SetDefine(TEXT("PLATFORM_SUPPORTS_MESH_SHADERS_TIER1"), RHISupportsMeshShadersTier1(EShaderPlatform(Target.Platform)) ? 1 : 0);
 	Input.Environment.SetDefine(TEXT("PLATFORM_ALLOW_SCENE_DATA_COMPRESSED_TRANSFORMS"), FDataDrivenShaderPlatformInfo::GetSupportSceneDataCompressedTransforms(EShaderPlatform(Target.Platform)) ? 1 : 0);
+	Input.Environment.SetDefine(TEXT("PLATFORM_SUPPORTS_BUFFER_LOAD_TYPE_CONVERSION"), RHISupportsBufferLoadTypeConversion(ShaderPlatform) ? 1 : 0);
 
 	bool bEnableBindlessMacro = false;
 	if (RHIGetBindlessSupport(EShaderPlatform(Target.Platform)) != ERHIBindlessSupport::Unsupported)
diff --git a/Engine/Source/Runtime/Renderer/Private/LightGridInjection.cpp b/Engine/Source/Runtime/Renderer/Private/LightGridInjection.cpp
--- a/Engine/Source/Runtime/Renderer/Private/LightGridInjection.cpp
+++ b/Engine/Source/Runtime/Renderer/Private/LightGridInjection.cpp
@@ -31,6 +31,7 @@ LightGridInjection.cpp
 #include "PixelShaderUtils.h"
 #include "ShaderPrint.h"
 #include "ShaderPrintParameters.h"
+#include "RenderUtils.h"
 
 int32 GLightGridPixelSize = 64;
 FAutoConsoleVariableRef CVarLightGridPixelSize(
@@ -103,35 +104,46 @@ bool ShouldVisualizeLightGrid()
 	return GForwardLightGridDebug > 0;
 }
 
-void SetupDummyForwardLightUniformParameters(FRDGBuilder& GraphBuilder, FForwardLightData& ForwardLightData)
+// If this is changed, the LIGHT_GRID_USES_16BIT_BUFFERS define from LightGridCommon.ush should also be updated.
+bool LightGridUses16BitBuffers(EShaderPlatform Platform)
+{
+	// CulledLightDataGrid, is typically 16bit elements to save on memory and bandwidth. So to not introduce any regressions it will stay as texel buffer on all platforms, except mobile and Metal (which does not support type conversions).
+	return RHISupportsBufferLoadTypeConversion(Platform) && !IsMobilePlatform(Platform);
+}
+
+void SetupDummyForwardLightUniformParameters(FRDGBuilder& GraphBuilder, FForwardLightData& ForwardLightData, EShaderPlatform ShaderPlatform)
 {
 	const FRDGSystemTextures& SystemTextures = FRDGSystemTextures::Get(GraphBuilder);
 	
 	ForwardLightData.DirectionalLightShadowmapAtlas = SystemTextures.Black;
 	ForwardLightData.DirectionalLightStaticShadowmap = GBlackTexture->TextureRHI;
 
-	FRDGBufferRef ForwardLocalLightBuffer = GSystemTextures.GetDefaultBuffer(GraphBuilder, sizeof(FVector4f));
-	ForwardLightData.ForwardLocalLightBuffer = GraphBuilder.CreateSRV(ForwardLocalLightBuffer, PF_A32B32G32R32F);
+	FRDGBufferRef ForwardLocalLightBuffer = GSystemTextures.GetDefaultStructuredBuffer(GraphBuilder, sizeof(FVector4f));
+	ForwardLightData.ForwardLocalLightBuffer = GraphBuilder.CreateSRV(ForwardLocalLightBuffer);
 
-	FRDGBufferRef NumCulledLightsGrid = GSystemTextures.GetDefaultBuffer(GraphBuilder, sizeof(uint32));
-	ForwardLightData.NumCulledLightsGrid = GraphBuilder.CreateSRV(NumCulledLightsGrid, PF_R32_UINT);
+	FRDGBufferRef NumCulledLightsGrid = GSystemTextures.GetDefaultStructuredBuffer(GraphBuilder, sizeof(uint32));
+	ForwardLightData.NumCulledLightsGrid = GraphBuilder.CreateSRV(NumCulledLightsGrid);
 
-	if (RHISupportsBufferLoadTypeConversion(GMaxRHIShaderPlatform))
+	const bool bLightGridUses16BitBuffers = LightGridUses16BitBuffers(ShaderPlatform);
+	FRDGBufferSRVRef CulledLightDataGridSRV = nullptr;
+	if (bLightGridUses16BitBuffers)
 	{
 		FRDGBufferRef CulledLightDataGrid = GSystemTextures.GetDefaultBuffer(GraphBuilder, sizeof(uint16));
-		ForwardLightData.CulledLightDataGrid = GraphBuilder.CreateSRV(CulledLightDataGrid, PF_R16_UINT);
+		CulledLightDataGridSRV = GraphBuilder.CreateSRV(CulledLightDataGrid, PF_R16_UINT);
 	}
 	else
 	{
-		FRDGBufferRef CulledLightDataGrid = GSystemTextures.GetDefaultBuffer(GraphBuilder, sizeof(uint32));
-		ForwardLightData.CulledLightDataGrid = GraphBuilder.CreateSRV(CulledLightDataGrid, PF_R32_UINT);
+		FRDGBufferRef CulledLightDataGrid = GSystemTextures.GetDefaultStructuredBuffer(GraphBuilder, sizeof(uint32));
+		CulledLightDataGridSRV = GraphBuilder.CreateSRV(CulledLightDataGrid);
 	}
+	ForwardLightData.CulledLightDataGrid32Bit = CulledLightDataGridSRV;
+	ForwardLightData.CulledLightDataGrid16Bit = CulledLightDataGridSRV;
 }
 
-TRDGUniformBufferRef<FForwardLightData> CreateDummyForwardLightUniformBuffer(FRDGBuilder& GraphBuilder)
+TRDGUniformBufferRef<FForwardLightData> CreateDummyForwardLightUniformBuffer(FRDGBuilder& GraphBuilder, EShaderPlatform ShaderPlatform)
 {
 	FForwardLightData* ForwardLightData = GraphBuilder.AllocParameters<FForwardLightData>();
-	SetupDummyForwardLightUniformParameters(GraphBuilder, *ForwardLightData);
+	SetupDummyForwardLightUniformParameters(GraphBuilder, *ForwardLightData, ShaderPlatform);
 	return GraphBuilder.CreateUniformBuffer(ForwardLightData);
 }
 
@@ -169,14 +181,15 @@ public:
 		SHADER_PARAMETER_STRUCT_REF(FReflectionCaptureShaderData, ReflectionCapture)
 		SHADER_PARAMETER_STRUCT_REF(FMobileReflectionCaptureShaderData, MobileReflectionCaptureData)
 		SHADER_PARAMETER_STRUCT_REF(FViewUniformShaderParameters, View)
-		SHADER_PARAMETER_RDG_BUFFER_UAV(RWBuffer<uint>, RWNumCulledLightsGrid)
-		SHADER_PARAMETER_RDG_BUFFER_UAV(RWBuffer<uint>, RWCulledLightDataGrid)
-		SHADER_PARAMETER_RDG_BUFFER_UAV(RWBuffer<uint>, RWNextCulledLightLink)
-		SHADER_PARAMETER_RDG_BUFFER_UAV(RWBuffer<uint>, RWStartOffsetGrid)
-		SHADER_PARAMETER_RDG_BUFFER_UAV(RWBuffer<uint>, RWCulledLightLinks)
-		SHADER_PARAMETER_RDG_BUFFER_SRV(Buffer<float4>, LightViewSpacePositionAndRadius)
-		SHADER_PARAMETER_RDG_BUFFER_SRV(Buffer<float4>, LightViewSpaceDirAndPreprocAngle)
-		SHADER_PARAMETER_RDG_BUFFER_SRV(Buffer<float4>, ForwardLocalLightBuffer)
+		SHADER_PARAMETER_RDG_BUFFER_UAV(RWStructuredBuffer<uint>, RWNumCulledLightsGrid)
+		SHADER_PARAMETER_RDG_BUFFER_UAV(RWStructuredBuffer<uint>, RWCulledLightDataGrid32Bit)
+		SHADER_PARAMETER_RDG_BUFFER_UAV(RWBuffer<uint>, RWCulledLightDataGrid16Bit)
+		SHADER_PARAMETER_RDG_BUFFER_UAV(RWStructuredBuffer<uint>, RWNextCulledLightLink)
+		SHADER_PARAMETER_RDG_BUFFER_UAV(RWStructuredBuffer<uint>, RWStartOffsetGrid)
+		SHADER_PARAMETER_RDG_BUFFER_UAV(RWStructuredBuffer<uint>, RWCulledLightLinks)
+		SHADER_PARAMETER_RDG_BUFFER_SRV(StructuredBuffer<float4>, LightViewSpacePositionAndRadius)
+		SHADER_PARAMETER_RDG_BUFFER_SRV(StructuredBuffer<float4>, LightViewSpaceDirAndPreprocAngle)
+		SHADER_PARAMETER_RDG_BUFFER_SRV(StructuredBuffer<float4>, ForwardLocalLightBuffer)
 
 		SHADER_PARAMETER(FIntVector, CulledGridSize)
 		SHADER_PARAMETER(uint32, NumReflectionCaptures)
@@ -213,11 +226,12 @@ class FLightGridCompactCS : public FGlobalShader
 public:
 	BEGIN_SHADER_PARAMETER_STRUCT(FParameters, )
 		SHADER_PARAMETER_STRUCT_REF(FViewUniformShaderParameters, View)
-		SHADER_PARAMETER_RDG_BUFFER_UAV(RWBuffer<uint>, RWNumCulledLightsGrid)
-		SHADER_PARAMETER_RDG_BUFFER_UAV(RWBuffer<uint>, RWCulledLightDataGrid)
-		SHADER_PARAMETER_RDG_BUFFER_UAV(RWBuffer<uint>, RWNextCulledLightData)
-		SHADER_PARAMETER_RDG_BUFFER_SRV(Buffer<uint>, StartOffsetGrid)
-		SHADER_PARAMETER_RDG_BUFFER_SRV(Buffer<uint>, CulledLightLinks)
+		SHADER_PARAMETER_RDG_BUFFER_UAV(RWStructuredBuffer<uint>, RWNumCulledLightsGrid)
+		SHADER_PARAMETER_RDG_BUFFER_UAV(RWStructuredBuffer<uint>, RWCulledLightDataGrid32Bit)
+		SHADER_PARAMETER_RDG_BUFFER_UAV(RWBuffer<uint>, RWCulledLightDataGrid16Bit)
+		SHADER_PARAMETER_RDG_BUFFER_UAV(RWStructuredBuffer<uint>, RWNextCulledLightData)
+		SHADER_PARAMETER_RDG_BUFFER_SRV(StructuredBuffer<uint>, StartOffsetGrid)
+		SHADER_PARAMETER_RDG_BUFFER_SRV(StructuredBuffer<uint>, CulledLightLinks)
 
 		SHADER_PARAMETER(FIntVector, CulledGridSize)
 		SHADER_PARAMETER(uint32, NumReflectionCaptures)
@@ -401,7 +415,7 @@ FComputeLightGridOutput FSceneRenderer::ComputeLightGrid(FRDGBuilder& GraphBuild
 
 	static const auto AllowStaticLightingVar = IConsoleManager::Get().FindTConsoleVariableDataInt(TEXT("r.AllowStaticLighting"));
 	const bool bAllowStaticLighting = (!AllowStaticLightingVar || AllowStaticLightingVar->GetValueOnRenderThread() != 0);
-	const bool bAllowFormatConversion = RHISupportsBufferLoadTypeConversion(GMaxRHIShaderPlatform);
+	const bool bLightGridUses16BitBuffers = LightGridUses16BitBuffers(ShaderPlatform);
 
 	const FRDGSystemTextures& SystemTextures = FRDGSystemTextures::Get(GraphBuilder);
 
@@ -686,11 +700,11 @@ FComputeLightGridOutput FSceneRenderer::ComputeLightGrid(FRDGBuilder& GraphBuild
 		// Store off the number of lights before we add a fake entry
 		const int32 NumLocalLightsFinal = ForwardLocalLightData.Num();
 
-		FRDGBufferRef ForwardLocalLightBuffer = CreateUploadBuffer(GraphBuilder, TEXT("ForwardLocalLightBuffer"), TConstArrayView<FForwardLocalLightData>(ForwardLocalLightData));
+		FRDGBufferRef ForwardLocalLightBuffer = CreateStructuredBuffer(GraphBuilder, TEXT("ForwardLocalLightBuffer"), TConstArrayView<FForwardLocalLightData>(ForwardLocalLightData));
 		View.ForwardLightingResources.LocalLightVisibleLightInfosIndex = LocalLightVisibleLightInfosIndex;
 
 		const FIntPoint LightGridSizeXY = FIntPoint::DivideAndRoundUp(View.ViewRect.Size(), GLightGridPixelSize);
-		ForwardLightData->ForwardLocalLightBuffer = GraphBuilder.CreateSRV(FRDGBufferSRVDesc(ForwardLocalLightBuffer, PF_A32B32G32R32F));
+		ForwardLightData->ForwardLocalLightBuffer = GraphBuilder.CreateSRV(FRDGBufferSRVDesc(ForwardLocalLightBuffer));
 		ForwardLightData->NumLocalLights = NumLocalLightsFinal;
 		ForwardLightData->NumReflectionCaptures = View.NumBoxReflectionCaptures + View.NumSphereReflectionCaptures;
 		ForwardLightData->NumGridCells = LightGridSizeXY.X * LightGridSizeXY.Y * GLightGridSizeZ;
@@ -709,7 +723,7 @@ FComputeLightGridOutput FSceneRenderer::ComputeLightGrid(FRDGBuilder& GraphBuild
 		FVector ZParams = GetLightGridZParams(View.NearClippingDistance, FarPlane + 10.f);
 		ForwardLightData->LightGridZParams = (FVector3f)ZParams;
 
-		const uint64 NumIndexableLights = !bAllowFormatConversion ? (1llu << (sizeof(FLightIndexType32) * 8llu)) : (1llu << (sizeof(FLightIndexType) * 8llu));
+		const uint64 NumIndexableLights = !bLightGridUses16BitBuffers ? (1llu << (sizeof(FLightIndexType32) * 8llu)) : (1llu << (sizeof(FLightIndexType) * 8llu));
 
 		if ((uint64)ForwardLocalLightData.Num() > NumIndexableLights)
 		{
@@ -723,7 +737,6 @@ FComputeLightGridOutput FSceneRenderer::ComputeLightGrid(FRDGBuilder& GraphBuild
 		}
 
 #if ENABLE_LIGHT_CULLING_VIEW_SPACE_BUILD_DATA
-		const SIZE_T LightIndexTypeSize = !bAllowFormatConversion ? sizeof(FLightIndexType32) : sizeof(FLightIndexType);
 		// Fuse these loops as I see no reason why not and we build some temporary data that is needed in the build pass and is 
 		// not needed to be stored permanently.
 #else // !ENABLE_LIGHT_CULLING_VIEW_SPACE_BUILD_DATA
@@ -731,8 +744,6 @@ FComputeLightGridOutput FSceneRenderer::ComputeLightGrid(FRDGBuilder& GraphBuild
 		ForwardLightDataPerView.Emplace(ForwardLightData);
 	}
 
-	const SIZE_T LightIndexTypeSize = !bAllowFormatConversion ? sizeof(FLightIndexType32) : sizeof(FLightIndexType);
-
 	for (int32 ViewIndex = 0; ViewIndex < Views.Num(); ViewIndex++)
 	{
 		const FViewInfo& View = Views[ViewIndex];
@@ -758,17 +769,34 @@ FComputeLightGridOutput FSceneRenderer::ComputeLightGrid(FRDGBuilder& GraphBuild
 				ForwardLightData->NumReflectionCaptures);
 
 			const uint32 CulledLightLinksElements = MaxNumCells * GMaxCulledLightsPerCell * LightLinkStride;
-			const EPixelFormat CulledLightDataGridFormat = LightIndexTypeSize == sizeof(uint16) ? PF_R16_UINT : PF_R32_UINT;
-
-			FRDGBufferRef CulledLightLinksBuffer = GraphBuilder.CreateBuffer(FRDGBufferDesc::CreateBufferDesc(sizeof(uint32), CulledLightLinksElements), TEXT("CulledLightLinks"));
-			FRDGBufferRef StartOffsetGridBuffer = GraphBuilder.CreateBuffer(FRDGBufferDesc::CreateBufferDesc(sizeof(uint32), MaxNumCells), TEXT("StartOffsetGrid"));
-			FRDGBufferRef NextCulledLightLinkBuffer = GraphBuilder.CreateBuffer(FRDGBufferDesc::CreateBufferDesc(sizeof(uint32), 1), TEXT("NextCulledLightLink"));
-			FRDGBufferRef NextCulledLightDataBuffer = GraphBuilder.CreateBuffer(FRDGBufferDesc::CreateBufferDesc(sizeof(uint32), 1), TEXT("NextCulledLightData"));
-			FRDGBufferUAVRef NextCulledLightDataUAV = GraphBuilder.CreateUAV(NextCulledLightDataBuffer, PF_R32_UINT);
-			FRDGBufferRef CulledLightDataGrid       = GraphBuilder.CreateBuffer(FRDGBufferDesc::CreateBufferDesc(LightIndexTypeSize, MaxNumCells * GMaxCulledLightsPerCell), TEXT("CulledLightDataGrid"));
-			FRDGBufferUAVRef CulledLightDataGridUAV = GraphBuilder.CreateUAV(CulledLightDataGrid, CulledLightDataGridFormat);
-			FRDGBufferRef NumCulledLightsGrid       = GraphBuilder.CreateBuffer(FRDGBufferDesc::CreateBufferDesc(sizeof(uint32), MaxNumCells * NumCulledLightsGridStride), TEXT("NumCulledLightsGrid"));
-			FRDGBufferUAVRef NumCulledLightsGridUAV = GraphBuilder.CreateUAV(NumCulledLightsGrid, PF_R32_UINT);
+
+
+			FRDGBufferRef CulledLightLinksBuffer = GraphBuilder.CreateBuffer(FRDGBufferDesc::CreateStructuredDesc(sizeof(uint32), CulledLightLinksElements), TEXT("CulledLightLinks"));
+			FRDGBufferRef StartOffsetGridBuffer = GraphBuilder.CreateBuffer(FRDGBufferDesc::CreateStructuredDesc(sizeof(uint32), MaxNumCells), TEXT("StartOffsetGrid"));
+			FRDGBufferRef NextCulledLightLinkBuffer = GraphBuilder.CreateBuffer(FRDGBufferDesc::CreateStructuredDesc(sizeof(uint32), 1), TEXT("NextCulledLightLink"));
+			FRDGBufferRef NextCulledLightDataBuffer = GraphBuilder.CreateBuffer(FRDGBufferDesc::CreateStructuredDesc(sizeof(uint32), 1), TEXT("NextCulledLightData"));
+			FRDGBufferUAVRef NextCulledLightDataUAV = GraphBuilder.CreateUAV(NextCulledLightDataBuffer);
+			FRDGBufferSRVRef CulledLightDataGridSRV = nullptr;
+			FRDGBufferUAVRef CulledLightDataGridUAV = nullptr;
+			FRDGBufferRef NumCulledLightsGrid       = GraphBuilder.CreateBuffer(FRDGBufferDesc::CreateStructuredDesc(sizeof(uint32), MaxNumCells * NumCulledLightsGridStride), TEXT("NumCulledLightsGrid"));
+			FRDGBufferUAVRef NumCulledLightsGridUAV = GraphBuilder.CreateUAV(NumCulledLightsGrid);
+
+			if (bLightGridUses16BitBuffers)
+			{
+				const SIZE_T LightIndexTypeSize = sizeof(FLightIndexType);
+				const EPixelFormat CulledLightDataGridFormat = PF_R16_UINT;
+				FRDGBufferRef CulledLightDataGrid = GraphBuilder.CreateBuffer(FRDGBufferDesc::CreateBufferDesc(LightIndexTypeSize, MaxNumCells * GMaxCulledLightsPerCell), TEXT("CulledLightDataGrid"));
+				CulledLightDataGridSRV = GraphBuilder.CreateSRV(CulledLightDataGrid, CulledLightDataGridFormat);
+				CulledLightDataGridUAV = GraphBuilder.CreateUAV(CulledLightDataGrid, CulledLightDataGridFormat);
+			}
+			else
+			{
+				const SIZE_T LightIndexTypeSize = sizeof(FLightIndexType32);
+				const EPixelFormat CulledLightDataGridFormat = PF_R32_UINT;
+				FRDGBufferRef CulledLightDataGrid = GraphBuilder.CreateBuffer(FRDGBufferDesc::CreateStructuredDesc(LightIndexTypeSize, MaxNumCells * GMaxCulledLightsPerCell), TEXT("CulledLightDataGrid"));
+				CulledLightDataGridSRV = GraphBuilder.CreateSRV(CulledLightDataGrid);
+				CulledLightDataGridUAV = GraphBuilder.CreateUAV(CulledLightDataGrid);
+			}
 
 			FLightGridInjectionCS::FParameters *PassParameters = GraphBuilder.AllocParameters<FLightGridInjectionCS::FParameters>();
 
@@ -784,10 +812,11 @@ FComputeLightGridOutput FSceneRenderer::ComputeLightGrid(FRDGBuilder& GraphBuild
 			}
 
 			PassParameters->RWNumCulledLightsGrid   = NumCulledLightsGridUAV;
-			PassParameters->RWCulledLightDataGrid   = CulledLightDataGridUAV;
-			PassParameters->RWNextCulledLightLink   = GraphBuilder.CreateUAV(NextCulledLightLinkBuffer, PF_R32_UINT);
-			PassParameters->RWStartOffsetGrid       = GraphBuilder.CreateUAV(StartOffsetGridBuffer, PF_R32_UINT);
-			PassParameters->RWCulledLightLinks      = GraphBuilder.CreateUAV(CulledLightLinksBuffer, PF_R32_UINT);
+			PassParameters->RWCulledLightDataGrid32Bit   = CulledLightDataGridUAV;
+			PassParameters->RWCulledLightDataGrid16Bit = CulledLightDataGridUAV;
+			PassParameters->RWNextCulledLightLink   = GraphBuilder.CreateUAV(NextCulledLightLinkBuffer);
+			PassParameters->RWStartOffsetGrid       = GraphBuilder.CreateUAV(StartOffsetGridBuffer);
+			PassParameters->RWCulledLightLinks      = GraphBuilder.CreateUAV(CulledLightLinksBuffer);
 			PassParameters->ForwardLocalLightBuffer = ForwardLightData->ForwardLocalLightBuffer;
 			PassParameters->CulledGridSize          = ForwardLightData->CulledGridSize;
 			PassParameters->LightGridZParams        = ForwardLightData->LightGridZParams;
@@ -801,11 +830,11 @@ FComputeLightGridOutput FSceneRenderer::ComputeLightGrid(FRDGBuilder& GraphBuild
 			check(ViewSpacePosAndRadiusData.Num() == ForwardLocalLightData.Num());
 			check(ViewSpaceDirAndPreprocAngleData.Num() == ForwardLocalLightData.Num());
 
-			FRDGBufferRef LightViewSpacePositionAndRadius  = CreateUploadBuffer(GraphBuilder, TEXT("ViewSpacePosAndRadiusData"), TConstArrayView<FVector4f>(ViewSpacePosAndRadiusData));
-			FRDGBufferRef LightViewSpaceDirAndPreprocAngle = CreateUploadBuffer(GraphBuilder, TEXT("ViewSpacePosAndRadiusData"), TConstArrayView<FVector4f>(ViewSpaceDirAndPreprocAngleData));
+			FRDGBufferRef LightViewSpacePositionAndRadius  = CreateStructuredBuffer(GraphBuilder, TEXT("ViewSpacePosAndRadiusData"), TConstArrayView<FVector4f>(ViewSpacePosAndRadiusData));
+			FRDGBufferRef LightViewSpaceDirAndPreprocAngle = CreateStructuredBuffer(GraphBuilder, TEXT("ViewSpacePosAndRadiusData"), TConstArrayView<FVector4f>(ViewSpaceDirAndPreprocAngleData));
 
-			PassParameters->LightViewSpacePositionAndRadius  = GraphBuilder.CreateSRV(FRDGBufferSRVDesc(LightViewSpacePositionAndRadius, PF_A32B32G32R32F));
-			PassParameters->LightViewSpaceDirAndPreprocAngle = GraphBuilder.CreateSRV(FRDGBufferSRVDesc(LightViewSpaceDirAndPreprocAngle, PF_A32B32G32R32F));
+			PassParameters->LightViewSpacePositionAndRadius  = GraphBuilder.CreateSRV(FRDGBufferSRVDesc(LightViewSpacePositionAndRadius));
+			PassParameters->LightViewSpaceDirAndPreprocAngle = GraphBuilder.CreateSRV(FRDGBufferSRVDesc(LightViewSpaceDirAndPreprocAngle));
 #endif // ENABLE_LIGHT_CULLING_VIEW_SPACE_BUILD_DATA
 
 			FLightGridInjectionCS::FPermutationDomain PermutationVector;
@@ -825,11 +854,12 @@ FComputeLightGridOutput FSceneRenderer::ComputeLightGrid(FRDGBuilder& GraphBuild
 					FLightGridCompactCS::FParameters *PassParametersCompact = GraphBuilder.AllocParameters<FLightGridCompactCS::FParameters>();
 					PassParametersCompact->View = View.ViewUniformBuffer;
 
-					PassParametersCompact->CulledLightLinks = GraphBuilder.CreateSRV(CulledLightLinksBuffer, PF_R32_UINT);
+					PassParametersCompact->CulledLightLinks = GraphBuilder.CreateSRV(CulledLightLinksBuffer);
 					PassParametersCompact->RWNumCulledLightsGrid = NumCulledLightsGridUAV;
-					PassParametersCompact->RWCulledLightDataGrid = CulledLightDataGridUAV;
+					PassParametersCompact->RWCulledLightDataGrid32Bit = CulledLightDataGridUAV;
+					PassParametersCompact->RWCulledLightDataGrid16Bit = CulledLightDataGridUAV;
 					PassParametersCompact->RWNextCulledLightData = NextCulledLightDataUAV;
-					PassParametersCompact->StartOffsetGrid = GraphBuilder.CreateSRV(StartOffsetGridBuffer, PF_R32_UINT);
+					PassParametersCompact->StartOffsetGrid = GraphBuilder.CreateSRV(StartOffsetGridBuffer);
 
 					PassParametersCompact->CulledGridSize = ForwardLightData->CulledGridSize;
 					PassParametersCompact->NumReflectionCaptures = ForwardLightData->NumReflectionCaptures;
@@ -844,9 +874,9 @@ FComputeLightGridOutput FSceneRenderer::ComputeLightGrid(FRDGBuilder& GraphBuild
 				AddClearUAVPass(GraphBuilder, NumCulledLightsGridUAV, 0);
 				FComputeShaderUtils::AddPass(GraphBuilder, RDG_EVENT_NAME("LightGridInject:NotLinkedList"), ComputeShader, PassParameters, NumGroups);
 			}
-
-			ForwardLightData->CulledLightDataGrid = GraphBuilder.CreateSRV(FRDGBufferSRVDesc(CulledLightDataGrid, CulledLightDataGridFormat));
-			ForwardLightData->NumCulledLightsGrid = GraphBuilder.CreateSRV(FRDGBufferSRVDesc(NumCulledLightsGrid, PF_R32_UINT));
+			ForwardLightData->CulledLightDataGrid32Bit = CulledLightDataGridSRV;
+			ForwardLightData->CulledLightDataGrid16Bit = CulledLightDataGridSRV;
+			ForwardLightData->NumCulledLightsGrid = GraphBuilder.CreateSRV(FRDGBufferSRVDesc(NumCulledLightsGrid));
 			View.ForwardLightingResources.SetUniformBuffer(GraphBuilder.CreateUniformBuffer(ForwardLightData));
 		}
 	}
@@ -879,7 +909,7 @@ FComputeLightGridOutput FDeferredShadingSceneRenderer::GatherLightsAndComputeLig
 	
 	if (!bNeedLightGrid)
 	{
-		TRDGUniformBufferRef<FForwardLightData> ForwardLightUniformBuffer = CreateDummyForwardLightUniformBuffer(GraphBuilder);
+		TRDGUniformBufferRef<FForwardLightData> ForwardLightUniformBuffer = CreateDummyForwardLightUniformBuffer(GraphBuilder, ShaderPlatform);
 		for (auto& View : Views)
 		{
 			View.ForwardLightingResources.SetUniformBuffer(ForwardLightUniformBuffer);
diff --git a/Engine/Source/Runtime/Renderer/Private/MobileBasePassRendering.cpp b/Engine/Source/Runtime/Renderer/Private/MobileBasePassRendering.cpp
--- a/Engine/Source/Runtime/Renderer/Private/MobileBasePassRendering.cpp
+++ b/Engine/Source/Runtime/Renderer/Private/MobileBasePassRendering.cpp
@@ -161,7 +161,7 @@ bool TMobileBasePassPSPolicyParamType<LightMapPolicyType>::ModifyCompilationEnvi
 	return true;
 }
 
-extern void SetupDummyForwardLightUniformParameters(FRDGBuilder& GraphBuilder, FForwardLightData& ForwardLightData);
+extern void SetupDummyForwardLightUniformParameters(FRDGBuilder& GraphBuilder, FForwardLightData& ForwardLightData, EShaderPlatform ShaderPlatform);
 
 void SetupMobileBasePassUniformParameters(
 	FRDGBuilder& GraphBuilder,
@@ -179,7 +179,7 @@ void SetupMobileBasePassUniformParameters(
 	}
 	else
 	{
-		SetupDummyForwardLightUniformParameters(GraphBuilder, BasePassParameters.Forward);
+		SetupDummyForwardLightUniformParameters(GraphBuilder, BasePassParameters.Forward, View.GetShaderPlatform());
 	}
 
 	const FScene* Scene = View.Family->Scene ? View.Family->Scene->GetRenderScene() : nullptr;
diff --git a/Engine/Source/Runtime/Renderer/Private/Renderer.cpp b/Engine/Source/Runtime/Renderer/Private/Renderer.cpp
--- a/Engine/Source/Runtime/Renderer/Private/Renderer.cpp
+++ b/Engine/Source/Runtime/Renderer/Private/Renderer.cpp
@@ -251,7 +251,7 @@ void FRendererModule::DrawTileMesh(FCanvasRenderContext& RenderContext, FMeshPas
 		}
 
 		View.InitRHIResources();
-		View.ForwardLightingResources.SetUniformBuffer(CreateDummyForwardLightUniformBuffer(GraphBuilder));
+		View.ForwardLightingResources.SetUniformBuffer(CreateDummyForwardLightUniformBuffer(GraphBuilder, View.GetShaderPlatform()));
 
 		TUniformBufferRef<FReflectionCaptureShaderData> EmptyReflectionCaptureUniformBuffer;
 
diff --git a/Engine/Source/Runtime/Renderer/Private/SceneRendering.h b/Engine/Source/Runtime/Renderer/Private/SceneRendering.h
--- a/Engine/Source/Runtime/Renderer/Private/SceneRendering.h
+++ b/Engine/Source/Runtime/Renderer/Private/SceneRendering.h
@@ -798,12 +798,13 @@ BEGIN_GLOBAL_SHADER_PARAMETER_STRUCT_WITH_CONSTRUCTOR(FForwardLightData, )
 	SHADER_PARAMETER_SAMPLER(SamplerState, ShadowmapSampler)
 	SHADER_PARAMETER_TEXTURE(Texture2D, DirectionalLightStaticShadowmap)
 	SHADER_PARAMETER_SAMPLER(SamplerState, StaticShadowmapSampler)
-	SHADER_PARAMETER_RDG_BUFFER_SRV(Buffer<float4>, ForwardLocalLightBuffer)
-	SHADER_PARAMETER_RDG_BUFFER_SRV(Buffer<uint>, NumCulledLightsGrid)
-	SHADER_PARAMETER_RDG_BUFFER_SRV(Buffer<uint>, CulledLightDataGrid)
+	SHADER_PARAMETER_RDG_BUFFER_SRV(StructuredBuffer<float4>, ForwardLocalLightBuffer)
+	SHADER_PARAMETER_RDG_BUFFER_SRV(StructuredBuffer<uint>, NumCulledLightsGrid)
+	SHADER_PARAMETER_RDG_BUFFER_SRV(StructuredBuffer<uint>, CulledLightDataGrid32Bit)
+	SHADER_PARAMETER_RDG_BUFFER_SRV(Buffer<uint>, CulledLightDataGrid16Bit)
 END_GLOBAL_SHADER_PARAMETER_STRUCT()
 
-extern TRDGUniformBufferRef<FForwardLightData> CreateDummyForwardLightUniformBuffer(FRDGBuilder& GraphBuilder);
+extern TRDGUniformBufferRef<FForwardLightData> CreateDummyForwardLightUniformBuffer(FRDGBuilder& GraphBuilder, EShaderPlatform ShaderPlatform);
 
 class FForwardLightingViewResources
 {
diff --git a/Engine/Source/Runtime/Renderer/Private/VirtualShadowMaps/VirtualShadowMapArray.cpp b/Engine/Source/Runtime/Renderer/Private/VirtualShadowMaps/VirtualShadowMapArray.cpp
--- a/Engine/Source/Runtime/Renderer/Private/VirtualShadowMaps/VirtualShadowMapArray.cpp
+++ b/Engine/Source/Runtime/Renderer/Private/VirtualShadowMaps/VirtualShadowMapArray.cpp
@@ -19,6 +19,7 @@
 #include "VirtualShadowMapClipmap.h"
 #include "VirtualShadowMapVisualizationData.h"
 #include "SingleLayerWaterRendering.h"
+#include "RenderUtils.h"
 
 #define DEBUG_ALLOW_STATIC_SEPARATE_WITHOUT_CACHING 0
 
@@ -358,6 +359,8 @@ namespace Nanite
 	extern bool IsStatFilterActive(const FString& FilterName);
 }
 
+extern bool LightGridUses16BitBuffers(EShaderPlatform Platform);
+
 FMatrix CalcTranslatedWorldToShadowUVMatrix(
 	const FMatrix& TranslatedWorldToShadowView,
 	const FMatrix& ViewToClip)
@@ -1472,7 +1475,9 @@ void FVirtualShadowMapArray::BuildPageAllocations(
 				// It's currently safe to overlap these passes that all write to same page request flags
 				FRDGBufferUAVRef PageRequestFlagsUAV = GraphBuilder.CreateUAV(PageRequestFlagsRDG, ERDGUnorderedAccessViewFlags::SkipBarrier);
 
-				uint32 LightGridIndexBufferSize = View.ForwardLightingResources.ForwardLightData->CulledLightDataGrid->Desc.Buffer->Desc.GetSize();
+				const bool bLightGridUses16BitBuffers = LightGridUses16BitBuffers(View.GetShaderPlatform());
+				FRDGBufferSRVRef CulledLightDataGrid = bLightGridUses16BitBuffers ? View.ForwardLightingResources.ForwardLightData->CulledLightDataGrid16Bit : View.ForwardLightingResources.ForwardLightData->CulledLightDataGrid32Bit;
+				uint32 LightGridIndexBufferSize = CulledLightDataGrid->Desc.Buffer->Desc.GetSize();
 				FRDGBufferRef PrunedLightGridDataRDG = GraphBuilder.CreateBuffer(FRDGBufferDesc::CreateStructuredDesc(sizeof(uint32), LightGridIndexBufferSize), TEXT("Shadow.Virtual.PrunedLightGridData"));
 				
 				const uint32 NumLightGridCells = View.ForwardLightingResources.ForwardLightData->NumGridCells;
