Non-unity fix. Forward decl.

diff --git a/Engine/Source/Runtime/Renderer/Public/RendererUtils.h b/Engine/Source/Runtime/Renderer/Public/RendererUtils.h
--- a/Engine/Source/Runtime/Renderer/Public/RendererUtils.h
+++ b/Engine/Source/Runtime/Renderer/Public/RendererUtils.h
@@ -11,6 +11,8 @@
 #include "ShaderParameterMacros.h"
 #include "RenderGraphFwd.h"
 
+struct FStrataSceneData;
+class FScene;
 class FGlobalShaderMap;
 class FRDGBuilder;
 class FRHICommandListImmediate;
@@ -56,8 +58,6 @@ BEGIN_GLOBAL_SHADER_PARAMETER_STRUCT(FStrataPublicGlobalUniformParameters, RENDE
 	SHADER_PARAMETER_RDG_TEXTURE(Texture2D<uint2>, TopLayerTexture)
 END_GLOBAL_SHADER_PARAMETER_STRUCT()
 
-struct FStrataSceneData;
-
 namespace Strata
 {
 	void PreInitViews(FScene& Scene);
