Use non GameThread specific function to access the default material during water PSO collection as fallback

#rnx

diff --git a/Engine/Source/Runtime/Renderer/Private/SingleLayerWaterRendering.cpp b/Engine/Source/Runtime/Renderer/Private/SingleLayerWaterRendering.cpp
--- a/Engine/Source/Runtime/Renderer/Private/SingleLayerWaterRendering.cpp
+++ b/Engine/Source/Runtime/Renderer/Private/SingleLayerWaterRendering.cpp
@@ -23,6 +23,7 @@
 #include "Lumen/LumenSceneData.h"
 #include "Lumen/LumenTracingUtils.h"
 #include "RenderCore.h"
+#include "UnrealEngine.h"
 
 DECLARE_GPU_STAT_NAMED(RayTracingWaterReflections, TEXT("Ray Tracing Water Reflections"));
 
@@ -1839,8 +1840,8 @@ void FSingleLayerWaterDepthPrepassMeshProcessor::CollectPSOInitializers(const FS
 			if (!bMaterialMasked && !Material.MaterialModifiesMeshPosition_GameThread())
 			{
 				// Override with the default material for opaque materials that are not two sided
-				FMaterialRenderProxy* EffectiveMaterialRenderProxy = UMaterial::GetDefaultMaterial(MD_Surface)->GetRenderProxy();
-				EffectiveMaterial = EffectiveMaterialRenderProxy->GetMaterialNoFallback(FeatureLevel);
+				EMaterialQualityLevel::Type ActiveQualityLevel = GetCachedScalabilityCVars().MaterialQualityLevel;
+				EffectiveMaterial = UMaterial::GetDefaultMaterial(MD_Surface)->GetMaterialResource(FeatureLevel, ActiveQualityLevel);
 				check(EffectiveMaterial);
 			}
 
