Path Tracer: Breakup the lighting component controls into direct vs. indirect.

This gives greater control over how the image is decomposed. Also tweak the implementation of these controls so that they behave as expected behind refraction.

Removed the EnableEmissive control as it is redundant with IncludeIndirectEmissive. Also removed the old EnableDirectLighting which did not work as expected behind glass. The new controls provide the same functionality in a better way. Remapped the ShowFlags.DirectLighting and ShowFlags.GlobalIllumination to the new mechanism.

diff --git a/Engine/Plugins/Experimental/GPULightmass/Shaders/Private/LightmapPathTracing.usf b/Engine/Plugins/Experimental/GPULightmass/Shaders/Private/LightmapPathTracing.usf
--- a/Engine/Plugins/Experimental/GPULightmass/Shaders/Private/LightmapPathTracing.usf
+++ b/Engine/Plugins/Experimental/GPULightmass/Shaders/Private/LightmapPathTracing.usf
@@ -129,7 +129,6 @@ static const uint MISMode = 2;
 
 static const uint ApproximateCaustics = 1;
 static const uint EnableCameraBackfaceCulling = 0;
-static const uint EnableDirectLighting = 1;
 static const uint EnableEmissive = 1;
 static const float MaxNormalBias = 0.1;
 
@@ -395,7 +394,7 @@ void PathTracingKernel(
 	{		
 		const bool bIsCameraRay = Bounce == 0;
 		const bool bIsLastBounce = Bounce == MaxBounces;
-		const bool bIncludeEmissive = (EnableDirectLighting != 0 || Bounce > 1) &&
+		const bool bIncludeEmissive = (Bounce > 1) &&
 									  (EnableEmissive       != 0 || bIsCameraRay);
 		FPathTracingPayload Payload = (FPathTracingPayload)0;
 		
@@ -563,12 +562,10 @@ void PathTracingKernel(
 		// Choose a random number for both Light sampling and BxDF sampling
 		float4 RandSample = RandomSequence_GenerateSample4D(RandSequence);
 
-		const bool bDoLightLoop = EnableDirectLighting != 0 || Bounce > 0;
-
 		float LightPickingCdfSum = 0;
 
 		// If we are using Light sampling and the material can use it ...
-		if (MISMode != 0 && SceneLightCount > 0 && bDoLightLoop)
+		if (MISMode != 0 && SceneLightCount > 0)
 		{
 			// Choose a light and sample it
 			float3 TranslatedWorldPos = Payload.TranslatedWorldPos;
@@ -764,7 +761,7 @@ void PathTracingKernel(
 		}
 		
 		// If we are using Material sampling for lights
-		if (MISMode != 1 && bDoLightLoop)
+		if (MISMode != 1)
 		{
 			// Check which lights can be seen by the material ray and trace a dedicated shadow ray
 			// While it would be possible to just loop around and use the indirect ray to do this, it would prevent the application
diff --git a/Engine/Shaders/Private/PathTracing/PathTracingCore.ush b/Engine/Shaders/Private/PathTracing/PathTracingCore.ush
--- a/Engine/Shaders/Private/PathTracing/PathTracingCore.ush
+++ b/Engine/Shaders/Private/PathTracing/PathTracingCore.ush
@@ -47,13 +47,12 @@ uint VolumeMISMode;
 
 uint ApproximateCaustics;
 uint EnableCameraBackfaceCulling;
-uint EnableDirectLighting;
-uint EnableEmissive;
 uint SamplerType;
 uint VisualizeLightGrid;
 uint VisualizeDecalGrid;
 uint ApplyDiffuseSpecularOverrides;
-uint EnabledLightingContributions;
+uint EnabledDirectLightingContributions;
+uint EnabledIndirectLightingContributions;
 float DecalRoughnessCutoff;
 float MeshDecalRoughnessCutoff;
 float MeshDecalBias;
@@ -109,33 +108,38 @@ struct FPathState
 		switch (FirstScatterType)
 		{
 			// In these cases, only keep tracing if the corresponding lighting component is enabled
-			case PATHTRACER_SCATTER_DIFFUSE:  return (EnabledLightingContributions & PATHTRACER_CONTRIBUTION_DIFFUSE ) != 0;
-			case PATHTRACER_SCATTER_SPECULAR: return (EnabledLightingContributions & PATHTRACER_CONTRIBUTION_SPECULAR) != 0;
-			case PATHTRACER_SCATTER_VOLUME:   return (EnabledLightingContributions & PATHTRACER_CONTRIBUTION_VOLUME  ) != 0;
+			case PATHTRACER_SCATTER_DIFFUSE:  return (GetEnabledLightingContributions() & PATHTRACER_CONTRIBUTION_DIFFUSE ) != 0;
+			case PATHTRACER_SCATTER_SPECULAR: return (GetEnabledLightingContributions() & PATHTRACER_CONTRIBUTION_SPECULAR) != 0;
+			case PATHTRACER_SCATTER_VOLUME:   return (GetEnabledLightingContributions() & PATHTRACER_CONTRIBUTION_VOLUME  ) != 0;
 		}
 		// In all other cases, keep tracing, we haven't made a final decision on what this path is yet
 		return true;
 	}
 
+	uint GetEnabledLightingContributions()
+	{
+		return HasMadeContributionScatter() ? EnabledIndirectLightingContributions : EnabledDirectLightingContributions;
+	}
+
 	// The following functions allow us to scale down the contributions for lighting components before we have scattered
 	bool ShouldAccumulateEmissive()
 	{
-		return (HasMadeContributionScatter() || (EnabledLightingContributions & PATHTRACER_CONTRIBUTION_EMISSIVE) != 0);
+		return (GetEnabledLightingContributions() & PATHTRACER_CONTRIBUTION_EMISSIVE) != 0;
 	}
 
 	bool ShouldAccumulateDiffuse()
 	{
-		return (HasMadeContributionScatter() || (EnabledLightingContributions & PATHTRACER_CONTRIBUTION_DIFFUSE) != 0);
+		return (GetEnabledLightingContributions() & PATHTRACER_CONTRIBUTION_DIFFUSE) != 0;
 	}
 
 	bool ShouldAccumulateSpecular()
 	{
-		return (HasMadeContributionScatter() || (EnabledLightingContributions & PATHTRACER_CONTRIBUTION_SPECULAR) != 0);
+		return (GetEnabledLightingContributions() & PATHTRACER_CONTRIBUTION_SPECULAR) != 0;
 	}
 
 	bool ShouldAccumulateVolume()
 	{
-		return (HasMadeContributionScatter() || (EnabledLightingContributions & PATHTRACER_CONTRIBUTION_VOLUME) != 0);
+		return (GetEnabledLightingContributions() & PATHTRACER_CONTRIBUTION_VOLUME) != 0;
 	}
 
 	float2 GetDiffuseSpecularScale(bool bIsVolumeSample)
@@ -1002,10 +1006,7 @@ bool PathTracingKernel(inout FPathState PathState, int Bounce)
 #endif
 	const bool bIsCameraRay = Bounce == 0;
 	const bool bIsLastBounce = Bounce == MaxBounces;
-	const bool bIncludeEmissive =
-		(EnableDirectLighting != 0 || Bounce > 1) &&
-		(EnableEmissive != 0 || bIsCameraRay) &&
-		PathState.ShouldAccumulateEmissive();
+	const bool bIncludeEmissive = PathState.ShouldAccumulateEmissive();
 
 	// number of directly visible lights for the first bounce
 	// after the first bounce, we don't need to include lights in the trace call
@@ -1193,7 +1194,7 @@ bool PathTracingKernel(inout FPathState PathState, int Bounce)
 
 	float2 DiffuseSpecularScale = PathState.GetDiffuseSpecularScale(bIsVolumeSample);
 
-	const bool bDoLightLoop = (EnableDirectLighting != 0 || Bounce > 0) && any(DiffuseSpecularScale > 0);
+	const bool bDoLightLoop = any(DiffuseSpecularScale > 0);
 
 	float LightPickingCdfSum = 0;
 
diff --git a/Engine/Source/Runtime/Engine/Classes/Engine/Scene.h b/Engine/Source/Runtime/Engine/Classes/Engine/Scene.h
--- a/Engine/Source/Runtime/Engine/Classes/Engine/Scene.h
+++ b/Engine/Source/Runtime/Engine/Classes/Engine/Scene.h
@@ -1272,9 +1272,6 @@ struct FPostProcessSettings
 	UPROPERTY(EditAnywhere, BlueprintReadWrite, Category = Overrides, meta = (PinHiddenByDefault, InlineEditConditionToggle))
 	uint32 bOverride_PathTracingFilterWidth : 1;
 
-	UPROPERTY(EditAnywhere, BlueprintReadWrite, Category = Overrides, meta = (PinHiddenByDefault, InlineEditConditionToggle))
-	uint32 bOverride_PathTracingEnableEmissive : 1;
-
 	UPROPERTY(EditAnywhere, BlueprintReadWrite, Category = Overrides, meta = (PinHiddenByDefault, InlineEditConditionToggle))
 	uint32 bOverride_PathTracingMaxPathExposure : 1;
 
@@ -1290,15 +1287,26 @@ struct FPostProcessSettings
 	UPROPERTY(EditAnywhere, BlueprintReadWrite, Category = Overrides, meta = (PinHiddenByDefault, InlineEditConditionToggle))
 	uint32 bOverride_PathTracingIncludeEmissive : 1;
 
+	UPROPERTY(EditAnywhere, BlueprintReadWrite, Category = Overrides, meta = (PinHiddenByDefault, InlineEditConditionToggle))
+	uint32 bOverride_PathTracingIncludeIndirectEmissive : 1;
+
 	UPROPERTY(EditAnywhere, BlueprintReadWrite, Category = Overrides, meta = (PinHiddenByDefault, InlineEditConditionToggle))
 	uint32 bOverride_PathTracingIncludeDiffuse : 1;
 
+	UPROPERTY(EditAnywhere, BlueprintReadWrite, Category = Overrides, meta = (PinHiddenByDefault, InlineEditConditionToggle))
+	uint32 bOverride_PathTracingIncludeIndirectDiffuse : 1;
+
 	UPROPERTY(EditAnywhere, BlueprintReadWrite, Category = Overrides, meta = (PinHiddenByDefault, InlineEditConditionToggle))
 	uint32 bOverride_PathTracingIncludeSpecular : 1;
 
+	UPROPERTY(EditAnywhere, BlueprintReadWrite, Category = Overrides, meta = (PinHiddenByDefault, InlineEditConditionToggle))
+	uint32 bOverride_PathTracingIncludeIndirectSpecular : 1;
+
 	UPROPERTY(EditAnywhere, BlueprintReadWrite, Category = Overrides, meta = (PinHiddenByDefault, InlineEditConditionToggle))
 	uint32 bOverride_PathTracingIncludeVolume : 1;
 
+	UPROPERTY(EditAnywhere, BlueprintReadWrite, Category = Overrides, meta = (PinHiddenByDefault, InlineEditConditionToggle))
+	uint32 bOverride_PathTracingIncludeIndirectVolume : 1;
 
 	// -----------------------------------------------------------------------
 
@@ -2151,17 +2159,13 @@ struct FPostProcessSettings
 	int32 PathTracingMaxBounces;
 
 	/** Sets the samples per pixel for the path tracer. */
-	UPROPERTY(interp, EditAnywhere, BlueprintReadWrite, Category = "Path Tracing", meta = (ClampMin = "1", UIMax = "65536", editcondition = "bOverride_PathTracingSamplesPerPixel", DisplayName = "Samples Per Pixel"))
+	UPROPERTY(EditAnywhere, BlueprintReadWrite, Category = "Path Tracing", meta = (ClampMin = "1", UIMax = "65536", editcondition = "bOverride_PathTracingSamplesPerPixel", DisplayName = "Samples Per Pixel"))
 	int32 PathTracingSamplesPerPixel;
 
 	/** Sets anti-aliasing filter width for the path tracer. Lower values are sharper (and more aliased), larger values are softer (and blurrier). */
 	UPROPERTY(interp, EditAnywhere, BlueprintReadWrite, Category = "Path Tracing", meta = (ClampMin = "1.0", ClampMax = "6.0", editcondition = "bOverride_PathTracingFilterWidth", DisplayName = "Filter Width"))
 	float PathTracingFilterWidth;
 
-	/** Enables emissive materials for the path tracer. This can prevent double-counting of illumination from surfaces that are also represented by light sources, and noise from small emitters. */
-	UPROPERTY(EditAnywhere, BlueprintReadWrite, Category = "Path Tracing", meta = (editcondition = "bOverride_PathTracingEnableEmissive", DisplayName = "Emissive Materials"))
-	uint32 PathTracingEnableEmissive : 1;
-
 	/** Sets the maximum exposure allowed in the path tracer to reduce fireflies. This should be set a few stops higher than the scene exposure. */
 	UPROPERTY(interp, EditAnywhere, BlueprintReadWrite, Category = "Path Tracing", meta = (ClampMin = "-10.0", ClampMax = "30.0", editcondition = "bOverride_PathTracingMaxPathExposure", DisplayName = "Max Path Exposure"))
 	float PathTracingMaxPathExposure;
@@ -2182,18 +2186,33 @@ struct FPostProcessSettings
 	UPROPERTY(EditAnywhere, BlueprintReadWrite, Category = "Path Tracing|Lighting Components", meta = (editcondition = "bOverride_PathTracingIncludeEmissive", DisplayName = "Emissive"))
 	uint32 PathTracingIncludeEmissive : 1;
 
+	/** Should the render include indirectly visible emissive elements? */
+	UPROPERTY(EditAnywhere, BlueprintReadWrite, Category = "Path Tracing|Lighting Components", meta = (editcondition = "bOverride_PathTracingIncludeIndirectEmissive", DisplayName = "Indirect Emissive"))
+	uint32 PathTracingIncludeIndirectEmissive : 1;
+
 	/** Should the render include diffuse lighting contributions? */
 	UPROPERTY(EditAnywhere, BlueprintReadWrite, Category = "Path Tracing|Lighting Components", meta = (editcondition = "bOverride_PathTracingIncludeDiffuse", DisplayName = "Diffuse"))
 	uint32 PathTracingIncludeDiffuse : 1;
 
+	/** Should the render include indirect diffuse lighting contributions? */
+	UPROPERTY(EditAnywhere, BlueprintReadWrite, Category = "Path Tracing|Lighting Components", meta = (editcondition = "bOverride_PathTracingIncludeIndirectDiffuse", DisplayName = "Indirect Diffuse"))
+	uint32 PathTracingIncludeIndirectDiffuse : 1;
+
 	/** Should the render include specular lighting contributions? */
 	UPROPERTY(EditAnywhere, BlueprintReadWrite, Category = "Path Tracing|Lighting Components", meta = (editcondition = "bOverride_PathTracingIncludeSpecular", DisplayName = "Specular"))
 	uint32 PathTracingIncludeSpecular : 1;
 
+	/** Should the render include indirect specular lighting contributions? */
+	UPROPERTY(EditAnywhere, BlueprintReadWrite, Category = "Path Tracing|Lighting Components", meta = (editcondition = "bOverride_PathTracingIncludeIndirectSpecular", DisplayName = "Indirect Specular"))
+	uint32 PathTracingIncludeIndirectSpecular : 1;
+
 	/** Should the render include volume lighting contributions? */
 	UPROPERTY(EditAnywhere, BlueprintReadWrite, Category = "Path Tracing|Lighting Components", meta = (editcondition = "bOverride_PathTracingIncludeVolume", DisplayName = "Volume"))
 	uint32 PathTracingIncludeVolume : 1;
 
+	/** Should the render include volume lighting contributions? */
+	UPROPERTY(EditAnywhere, BlueprintReadWrite, Category = "Path Tracing|Lighting Components", meta = (editcondition = "bOverride_PathTracingIncludeIndirectVolume", DisplayName = "Indirect Volume"))
+	uint32 PathTracingIncludeIndirectVolume : 1;
 
 	UPROPERTY()
 	float LPVFadeRange_DEPRECATED;
diff --git a/Engine/Source/Runtime/Engine/Private/Scene.cpp b/Engine/Source/Runtime/Engine/Private/Scene.cpp
--- a/Engine/Source/Runtime/Engine/Private/Scene.cpp
+++ b/Engine/Source/Runtime/Engine/Private/Scene.cpp
@@ -631,17 +631,20 @@ FPostProcessSettings::FPostProcessSettings()
 	PathTracingMaxBounces = 32;
 	PathTracingSamplesPerPixel = 2048;
 	PathTracingFilterWidth = 3.0f;
-	PathTracingEnableEmissive = 1;
 	PathTracingMaxPathExposure = 30.0f;
 	PathTracingEnableReferenceDOF = 0;
 	PathTracingEnableReferenceAtmosphere = 0;
 	PathTracingEnableDenoiser = 1;
 
 	PathTracingIncludeEmissive = 1;
+	PathTracingIncludeIndirectEmissive = 1;
 	PathTracingIncludeDiffuse = 1;
+	PathTracingIncludeIndirectDiffuse = 1;
 	PathTracingIncludeSpecular = 1;
+	PathTracingIncludeIndirectSpecular = 1;
 	PathTracingIncludeVolume = 1;
-	
+	PathTracingIncludeIndirectVolume = 1;
+
 	bMobileHQGaussian = false;
 
 #if DO_CHECK && WITH_EDITOR
@@ -844,15 +847,18 @@ FPostProcessSettings::FPostProcessSettings(const FPostProcessSettings& Settings)
 	, bOverride_PathTracingMaxBounces(Settings.bOverride_PathTracingMaxBounces)
 	, bOverride_PathTracingSamplesPerPixel(Settings.bOverride_PathTracingSamplesPerPixel)
 	, bOverride_PathTracingFilterWidth(Settings.bOverride_PathTracingFilterWidth)
-	, bOverride_PathTracingEnableEmissive(Settings.bOverride_PathTracingEnableEmissive)
 	, bOverride_PathTracingMaxPathExposure(Settings.bOverride_PathTracingMaxPathExposure)
 	, bOverride_PathTracingEnableReferenceDOF(Settings.bOverride_PathTracingEnableReferenceDOF)
 	, bOverride_PathTracingEnableReferenceAtmosphere(Settings.bOverride_PathTracingEnableReferenceAtmosphere)
 	, bOverride_PathTracingEnableDenoiser(Settings.bOverride_PathTracingEnableDenoiser)
 	, bOverride_PathTracingIncludeEmissive(Settings.bOverride_PathTracingIncludeEmissive)
+	, bOverride_PathTracingIncludeIndirectEmissive(Settings.bOverride_PathTracingIncludeIndirectEmissive)
 	, bOverride_PathTracingIncludeDiffuse(Settings.bOverride_PathTracingIncludeDiffuse)
+	, bOverride_PathTracingIncludeIndirectDiffuse(Settings.bOverride_PathTracingIncludeIndirectDiffuse)
 	, bOverride_PathTracingIncludeSpecular(Settings.bOverride_PathTracingIncludeSpecular)
+	, bOverride_PathTracingIncludeIndirectSpecular(Settings.bOverride_PathTracingIncludeIndirectSpecular)
 	, bOverride_PathTracingIncludeVolume(Settings.bOverride_PathTracingIncludeVolume)
+	, bOverride_PathTracingIncludeIndirectVolume(Settings.bOverride_PathTracingIncludeIndirectVolume)
 
 	, bMobileHQGaussian(Settings.bMobileHQGaussian)
 	, BloomMethod(Settings.BloomMethod)
@@ -1043,15 +1049,18 @@ FPostProcessSettings::FPostProcessSettings(const FPostProcessSettings& Settings)
 	, PathTracingMaxBounces(Settings.PathTracingMaxBounces)
 	, PathTracingSamplesPerPixel(Settings.PathTracingSamplesPerPixel)
 	, PathTracingFilterWidth(Settings.PathTracingFilterWidth)
-	, PathTracingEnableEmissive(Settings.PathTracingEnableEmissive)
 	, PathTracingMaxPathExposure(Settings.PathTracingMaxPathExposure)
 	, PathTracingEnableReferenceDOF(Settings.PathTracingEnableReferenceDOF)
 	, PathTracingEnableReferenceAtmosphere(Settings.PathTracingEnableReferenceAtmosphere)
 	, PathTracingEnableDenoiser(Settings.PathTracingEnableDenoiser)
 	, PathTracingIncludeEmissive(Settings.PathTracingIncludeEmissive)
+	, PathTracingIncludeIndirectEmissive(Settings.PathTracingIncludeIndirectEmissive)
 	, PathTracingIncludeDiffuse(Settings.PathTracingIncludeDiffuse)
+	, PathTracingIncludeIndirectDiffuse(Settings.PathTracingIncludeIndirectDiffuse)
 	, PathTracingIncludeSpecular(Settings.PathTracingIncludeSpecular)
+	, PathTracingIncludeIndirectSpecular(Settings.PathTracingIncludeIndirectSpecular)
 	, PathTracingIncludeVolume(Settings.PathTracingIncludeVolume)
+	, PathTracingIncludeIndirectVolume(Settings.PathTracingIncludeIndirectVolume)
 
 	, ScreenPercentage_DEPRECATED(Settings.ScreenPercentage_DEPRECATED)
 
diff --git a/Engine/Source/Runtime/Engine/Private/SceneView.cpp b/Engine/Source/Runtime/Engine/Private/SceneView.cpp
--- a/Engine/Source/Runtime/Engine/Private/SceneView.cpp
+++ b/Engine/Source/Runtime/Engine/Private/SceneView.cpp
@@ -1720,11 +1720,6 @@ void FSceneView::OverridePostProcessSettings(const FPostProcessSettings& Src, fl
 			Dest.PathTracingFilterWidth = Src.PathTracingFilterWidth;
 		}
 
-		if (Src.bOverride_PathTracingEnableEmissive)
-		{
-			Dest.PathTracingEnableEmissive = Src.PathTracingEnableEmissive;
-		}
-
 		if (Src.bOverride_PathTracingMaxPathExposure)
 		{
 			Dest.PathTracingMaxPathExposure = Src.PathTracingMaxPathExposure;
@@ -1750,21 +1745,41 @@ void FSceneView::OverridePostProcessSettings(const FPostProcessSettings& Src, fl
 			Dest.PathTracingIncludeEmissive = Src.PathTracingIncludeEmissive;
 		}
 
+		if (Src.bOverride_PathTracingIncludeIndirectEmissive)
+		{
+			Dest.PathTracingIncludeIndirectEmissive = Src.PathTracingIncludeIndirectEmissive;
+		}
+
 		if (Src.bOverride_PathTracingIncludeDiffuse)
 		{
 			Dest.PathTracingIncludeDiffuse = Src.PathTracingIncludeDiffuse;
 		}
 
+		if (Src.bOverride_PathTracingIncludeIndirectDiffuse)
+		{
+			Dest.PathTracingIncludeIndirectDiffuse = Src.PathTracingIncludeIndirectDiffuse;
+		}
+
 		if (Src.bOverride_PathTracingIncludeSpecular)
 		{
 			Dest.PathTracingIncludeSpecular = Src.PathTracingIncludeSpecular;
 		}
 
+		if (Src.bOverride_PathTracingIncludeIndirectSpecular)
+		{
+			Dest.PathTracingIncludeIndirectSpecular = Src.PathTracingIncludeIndirectSpecular;
+		}
+
 		if (Src.bOverride_PathTracingIncludeVolume)
 		{
 			Dest.PathTracingIncludeVolume = Src.PathTracingIncludeVolume;
 		}
 
+		if (Src.bOverride_PathTracingIncludeIndirectVolume)
+		{
+			Dest.PathTracingIncludeIndirectVolume = Src.PathTracingIncludeIndirectVolume;
+		}
+
 		if (Src.bOverride_DepthOfFieldBladeCount)
 		{
 			Dest.DepthOfFieldBladeCount = Src.DepthOfFieldBladeCount;
diff --git a/Engine/Source/Runtime/Renderer/Private/PathTracing.cpp b/Engine/Source/Runtime/Renderer/Private/PathTracing.cpp
--- a/Engine/Source/Runtime/Renderer/Private/PathTracing.cpp
+++ b/Engine/Source/Runtime/Renderer/Private/PathTracing.cpp
@@ -324,14 +324,13 @@ BEGIN_SHADER_PARAMETER_STRUCT(FPathTracingData, )
 	SHADER_PARAMETER(uint32, VolumeMISMode)
 	SHADER_PARAMETER(uint32, ApproximateCaustics)
 	SHADER_PARAMETER(uint32, EnableCameraBackfaceCulling)
-	SHADER_PARAMETER(uint32, EnableDirectLighting)
-	SHADER_PARAMETER(uint32, EnableEmissive)
 	SHADER_PARAMETER(uint32, SamplerType)
 	SHADER_PARAMETER(uint32, VisualizeLightGrid)
 	SHADER_PARAMETER(uint32, VisualizeDecalGrid)
 	SHADER_PARAMETER(uint32, EnableAtmosphere)
 	SHADER_PARAMETER(uint32, EnableFog)
-	SHADER_PARAMETER(uint32, EnabledLightingContributions) // PATHTRACER_CONTRIBUTION_*
+	SHADER_PARAMETER(uint32, EnabledDirectLightingContributions)   // PATHTRACER_CONTRIBUTION_*
+	SHADER_PARAMETER(uint32, EnabledIndirectLightingContributions) // PATHTRACER_CONTRIBUTION_*
 	SHADER_PARAMETER(uint32, ApplyDiffuseSpecularOverrides)
 	SHADER_PARAMETER(int32, MaxRaymarchSteps)
 	SHADER_PARAMETER(float, MaxPathIntensity)
@@ -372,8 +371,6 @@ struct FPathTracingConfig
 			PathTracingData.SamplerType != Other.PathTracingData.SamplerType ||
 			PathTracingData.ApproximateCaustics != Other.PathTracingData.ApproximateCaustics ||
 			PathTracingData.EnableCameraBackfaceCulling != Other.PathTracingData.EnableCameraBackfaceCulling ||
-			PathTracingData.EnableDirectLighting != Other.PathTracingData.EnableDirectLighting ||
-			PathTracingData.EnableEmissive != Other.PathTracingData.EnableEmissive ||
 			PathTracingData.VisualizeLightGrid != Other.PathTracingData.VisualizeLightGrid ||
 			PathTracingData.VisualizeDecalGrid != Other.PathTracingData.VisualizeDecalGrid ||
 			PathTracingData.MaxPathIntensity != Other.PathTracingData.MaxPathIntensity ||
@@ -382,7 +379,8 @@ struct FPathTracingConfig
 			PathTracingData.EnableAtmosphere != Other.PathTracingData.EnableAtmosphere ||
 			PathTracingData.EnableFog != Other.PathTracingData.EnableFog ||
 			PathTracingData.ApplyDiffuseSpecularOverrides != Other.PathTracingData.ApplyDiffuseSpecularOverrides ||
-			PathTracingData.EnabledLightingContributions != Other.PathTracingData.EnabledLightingContributions ||
+			PathTracingData.EnabledDirectLightingContributions != Other.PathTracingData.EnabledDirectLightingContributions ||
+			PathTracingData.EnabledIndirectLightingContributions != Other.PathTracingData.EnabledIndirectLightingContributions ||
 			PathTracingData.DecalRoughnessCutoff != Other.PathTracingData.DecalRoughnessCutoff ||
 			PathTracingData.MeshDecalRoughnessCutoff != Other.PathTracingData.MeshDecalRoughnessCutoff ||
 			PathTracingData.MeshDecalBias != Other.PathTracingData.MeshDecalBias ||
@@ -481,36 +479,17 @@ namespace PathTracing
 // This function prepares the portion of shader arguments that may involve invalidating the path traced state
 static void PreparePathTracingData(const FScene* Scene, const FViewInfo& View, FPathTracingData& PathTracingData)
 {
-	PathTracingData.EnableDirectLighting = true;
+	const FFinalPostProcessSettings& PPV = View.FinalPostProcessSettings;
+	const FEngineShowFlags& ShowFlags = View.Family->EngineShowFlags;
+
 	int32 MaxBounces = CVarPathTracingMaxBounces.GetValueOnRenderThread();
 	if (MaxBounces < 0)
 	{
-		MaxBounces = View.FinalPostProcessSettings.PathTracingMaxBounces;
-	}
-	if (View.Family->EngineShowFlags.DirectLighting)
-	{
-		if (!View.Family->EngineShowFlags.GlobalIllumination)
-		{
-			// direct lighting, but no GI
-			MaxBounces = 1;
-		}
-	}
-	else
-	{
-		PathTracingData.EnableDirectLighting = false;
-		if (View.Family->EngineShowFlags.GlobalIllumination)
-		{
-			// skip direct lighting, but still do the full bounces
-		}
-		else
-		{
-			// neither direct, nor GI is on
-			MaxBounces = 0;
-		}
+		MaxBounces = PPV.PathTracingMaxBounces;
 	}
 
 	PathTracingData.MaxBounces = MaxBounces;
-	PathTracingData.MaxSSSBounces = View.Family->EngineShowFlags.SubsurfaceScattering ? CVarPathTracingMaxSSSBounces.GetValueOnRenderThread() : 0;
+	PathTracingData.MaxSSSBounces = ShowFlags.SubsurfaceScattering ? CVarPathTracingMaxSSSBounces.GetValueOnRenderThread() : 0;
 	PathTracingData.MaxNormalBias = GetRaytracingMaxNormalBias();
 	PathTracingData.MISMode = CVarPathTracingMISMode.GetValueOnRenderThread();
 	PathTracingData.VolumeMISMode = CVarPathTracingVolumeMISMode.GetValueOnRenderThread();
@@ -518,38 +497,36 @@ static void PreparePathTracingData(const FScene* Scene, const FViewInfo& View, F
 	if (PathTracingData.MaxPathIntensity <= 0)
 	{
 		// cvar clamp disabled, use PPV exposure value instad
-		PathTracingData.MaxPathIntensity = FMath::Pow(2.0f, View.FinalPostProcessSettings.PathTracingMaxPathExposure);
+		PathTracingData.MaxPathIntensity = FMath::Pow(2.0f, PPV.PathTracingMaxPathExposure);
 	}
 	PathTracingData.ApproximateCaustics = CVarPathTracingApproximateCaustics.GetValueOnRenderThread();
 	PathTracingData.EnableCameraBackfaceCulling = CVarPathTracingEnableCameraBackfaceCulling.GetValueOnRenderThread();
 	PathTracingData.SamplerType = CVarPathTracingSamplerType.GetValueOnRenderThread();
-	int32 EnableEmissive = CVarPathTracingEnableEmissive.GetValueOnRenderThread();
-	PathTracingData.EnableEmissive = EnableEmissive < 0 ? View.FinalPostProcessSettings.PathTracingEnableEmissive : EnableEmissive;
 	PathTracingData.VisualizeLightGrid = CVarPathTracingLightGridVisualize.GetValueOnRenderThread();
 	PathTracingData.VisualizeDecalGrid = CVarPathTracingDecalGridVisualize.GetValueOnRenderThread();
 	float FilterWidth = CVarPathTracingFilterWidth.GetValueOnRenderThread();
 	if (FilterWidth < 0)
 	{
-		FilterWidth = View.FinalPostProcessSettings.PathTracingFilterWidth;
+		FilterWidth = PPV.PathTracingFilterWidth;
 	}
 	PathTracingData.FilterWidth = FilterWidth;
 	PathTracingData.AbsorptionScale = CVarPathTracingAbsorptionScale.GetValueOnRenderThread();
 	PathTracingData.CameraFocusDistance = 0;
 	PathTracingData.CameraLensRadius = FVector2f::ZeroVector;
-	if (View.Family->EngineShowFlags.DepthOfField &&
-		View.FinalPostProcessSettings.PathTracingEnableReferenceDOF &&
-		View.FinalPostProcessSettings.DepthOfFieldFocalDistance > 0 &&
-		View.FinalPostProcessSettings.DepthOfFieldFstop > 0)
+	if (ShowFlags.DepthOfField &&
+		PPV.PathTracingEnableReferenceDOF &&
+		PPV.DepthOfFieldFocalDistance > 0 &&
+		PPV.DepthOfFieldFstop > 0)
 	{
-		const float FocalLengthInCM = 0.05f * View.FinalPostProcessSettings.DepthOfFieldSensorWidth * View.ViewMatrices.GetProjectionMatrix().M[0][0];
-		PathTracingData.CameraFocusDistance = View.FinalPostProcessSettings.DepthOfFieldFocalDistance;
-		PathTracingData.CameraLensRadius.Y = 0.5f * FocalLengthInCM / View.FinalPostProcessSettings.DepthOfFieldFstop;
-		PathTracingData.CameraLensRadius.X = PathTracingData.CameraLensRadius.Y / FMath::Clamp(View.FinalPostProcessSettings.DepthOfFieldSqueezeFactor, 1.0f, 2.0f);
+		const float FocalLengthInCM = 0.05f * PPV.DepthOfFieldSensorWidth * View.ViewMatrices.GetProjectionMatrix().M[0][0];
+		PathTracingData.CameraFocusDistance = PPV.DepthOfFieldFocalDistance;
+		PathTracingData.CameraLensRadius.Y = 0.5f * FocalLengthInCM / PPV.DepthOfFieldFstop;
+		PathTracingData.CameraLensRadius.X = PathTracingData.CameraLensRadius.Y / FMath::Clamp(PPV.DepthOfFieldSqueezeFactor, 1.0f, 2.0f);
 	}
 	PathTracingData.EnableAtmosphere =
-		ShouldRenderSkyAtmosphere(Scene, View.Family->EngineShowFlags) && 
-		View.SkyAtmosphereUniformShaderParameters &&
-		View.FinalPostProcessSettings.PathTracingEnableReferenceAtmosphere;
+		ShouldRenderSkyAtmosphere(Scene, ShowFlags) && 
+		View.SkyAtmosphereUniformShaderParameters != nullptr &&
+		PPV.PathTracingEnableReferenceAtmosphere != 0;
 
 	PathTracingData.EnableFog = ShouldRenderFog(*View.Family)
 		&& Scene->ExponentialFogs.Num() > 0
@@ -568,15 +545,27 @@ static void PreparePathTracingData(const FScene* Scene, const FViewInfo& View, F
 
 	// NOTE: Diffuse and Specular show flags also modify the override colors, but we prefer to tie those to the lighting contribution mechanism below which is more principled
 	PathTracingData.ApplyDiffuseSpecularOverrides =
-		View.Family->EngineShowFlags.LightingOnlyOverride       != 0 ||
-		View.Family->EngineShowFlags.OverrideDiffuseAndSpecular != 0 ||
-		View.Family->EngineShowFlags.ReflectionOverride         != 0;
-
-	PathTracingData.EnabledLightingContributions = 0;
-	PathTracingData.EnabledLightingContributions |= (View.FinalPostProcessSettings.PathTracingIncludeEmissive != 0                                              ) ? PATHTRACER_CONTRIBUTION_EMISSIVE : 0;
-	PathTracingData.EnabledLightingContributions |= (View.FinalPostProcessSettings.PathTracingIncludeDiffuse  != 0 && View.Family->EngineShowFlags.Diffuse  != 0) ? PATHTRACER_CONTRIBUTION_DIFFUSE  : 0;
-	PathTracingData.EnabledLightingContributions |= (View.FinalPostProcessSettings.PathTracingIncludeSpecular != 0 && View.Family->EngineShowFlags.Specular != 0) ? PATHTRACER_CONTRIBUTION_SPECULAR : 0;
-	PathTracingData.EnabledLightingContributions |= (View.FinalPostProcessSettings.PathTracingIncludeVolume   != 0                                              ) ? PATHTRACER_CONTRIBUTION_VOLUME   : 0;
+		ShowFlags.LightingOnlyOverride       != 0 ||
+		ShowFlags.OverrideDiffuseAndSpecular != 0 ||
+		ShowFlags.ReflectionOverride         != 0;
+
+	PathTracingData.EnabledDirectLightingContributions = 0;
+	if (ShowFlags.DirectLighting != 0)
+	{
+		PathTracingData.EnabledDirectLightingContributions |= (PPV.PathTracingIncludeEmissive != 0                           ) ? PATHTRACER_CONTRIBUTION_EMISSIVE : 0;
+		PathTracingData.EnabledDirectLightingContributions |= (PPV.PathTracingIncludeDiffuse  != 0 && ShowFlags.Diffuse  != 0) ? PATHTRACER_CONTRIBUTION_DIFFUSE  : 0;
+		PathTracingData.EnabledDirectLightingContributions |= (PPV.PathTracingIncludeSpecular != 0 && ShowFlags.Specular != 0) ? PATHTRACER_CONTRIBUTION_SPECULAR : 0;
+		PathTracingData.EnabledDirectLightingContributions |= (PPV.PathTracingIncludeVolume   != 0                           ) ? PATHTRACER_CONTRIBUTION_VOLUME   : 0;
+	}
+	PathTracingData.EnabledIndirectLightingContributions = 0;
+	if (ShowFlags.GlobalIllumination != 0)
+	{
+		const bool bEnableEmissive = CVarPathTracingEnableEmissive.GetValueOnRenderThread() != 0;
+		PathTracingData.EnabledIndirectLightingContributions |= (PPV.PathTracingIncludeIndirectEmissive != 0 && bEnableEmissive        ) ? PATHTRACER_CONTRIBUTION_EMISSIVE : 0;
+		PathTracingData.EnabledIndirectLightingContributions |= (PPV.PathTracingIncludeIndirectDiffuse  != 0 && ShowFlags.Diffuse  != 0) ? PATHTRACER_CONTRIBUTION_DIFFUSE  : 0;
+		PathTracingData.EnabledIndirectLightingContributions |= (PPV.PathTracingIncludeIndirectSpecular != 0 && ShowFlags.Specular != 0) ? PATHTRACER_CONTRIBUTION_SPECULAR : 0;
+		PathTracingData.EnabledIndirectLightingContributions |= (PPV.PathTracingIncludeIndirectVolume   != 0                           ) ? PATHTRACER_CONTRIBUTION_VOLUME   : 0;
+	}
 }
 
 static bool ShouldCompilePathTracingShadersForProject(EShaderPlatform ShaderPlatform)
@@ -2384,11 +2373,6 @@ void FDeferredShadingSceneRenderer::RenderPathTracing(
 								PassParameters->SceneLights = PreviousPassParameters->SceneLights;
 								PassParameters->SkylightParameters = PreviousPassParameters->SkylightParameters;
 							}
-							if (Config.PathTracingData.EnableDirectLighting == 0)
-							{
-								PassParameters->SceneVisibleLightCount = 0;
-							}
-
 							PassParameters->DecalParameters = View.RayTracingDecalUniformBuffer;
 
 							PassParameters->RadianceTexture = GraphBuilder.CreateUAV(RadianceTexture);
