Changing ScreenPass to use FSceneView instead of FViewInfo.
Also removing all the Renderer/Private includes from ScreenPass.h

diff --git a/Engine/Source/Runtime/Renderer/Private/ScreenPass.cpp b/Engine/Source/Runtime/Renderer/Private/ScreenPass.cpp
--- a/Engine/Source/Runtime/Renderer/Private/ScreenPass.cpp
+++ b/Engine/Source/Runtime/Renderer/Private/ScreenPass.cpp
@@ -1,12 +1,15 @@
 // Copyright Epic Games, Inc. All Rights Reserved.
 
 #include "ScreenPass.h"
+
 #include "DataDrivenShaderPlatformInfo.h"
 #include "EngineGlobals.h"
-#include "ScenePrivate.h"
+#include "PostProcess/SceneFilterRendering.h"
 #include "RendererModule.h"
 #include "RenderGraphUtils.h"
+#include "ScenePrivate.h"
 #include "SystemTextures.h"
+#include "UnrealClient.h"
 
 IMPLEMENT_GLOBAL_SHADER(FScreenPassVS, "/Engine/Private/ScreenPass.usf", "ScreenPassVS", SF_Vertex);
 
@@ -39,6 +42,43 @@ FRDGTextureRef TryCreateViewFamilyTexture(FRDGBuilder& GraphBuilder, const FScen
 	return Texture;
 }
 
+FScreenPassRenderTarget FScreenPassRenderTarget::CreateFromInput(
+	FRDGBuilder& GraphBuilder,
+	FScreenPassTexture Input,
+	ERenderTargetLoadAction OutputLoadAction,
+	const TCHAR* OutputName)
+{
+	check(Input.IsValid());
+
+	FRDGTextureDesc OutputDesc = Input.Texture->Desc;
+	OutputDesc.Reset();
+
+	return FScreenPassRenderTarget(GraphBuilder.CreateTexture(OutputDesc, OutputName), Input.ViewRect, OutputLoadAction);
+}
+
+FScreenPassRenderTarget FScreenPassRenderTarget::CreateViewFamilyOutput(FRDGTextureRef ViewFamilyTexture, const FViewInfo& View)
+{
+	const FIntRect ViewRect = View.PrimaryScreenPercentageMethod == EPrimaryScreenPercentageMethod::RawOutput ? View.ViewRect : View.UnscaledViewRect;
+
+	ERenderTargetLoadAction LoadAction = ERenderTargetLoadAction::ENoAction;
+
+	if (!View.IsFirstInFamily() || View.Family->bAdditionalViewFamily)
+	{
+		LoadAction = ERenderTargetLoadAction::ELoad;
+	}
+	else if (ViewRect.Min != FIntPoint::ZeroValue || ViewRect.Size() != ViewFamilyTexture->Desc.Extent)
+	{
+		LoadAction = ERenderTargetLoadAction::EClear;
+	}
+
+	return FScreenPassRenderTarget(
+		ViewFamilyTexture,
+		// Raw output mode uses the original view rect. Otherwise the final unscaled rect is used.
+		ViewRect,
+		// First view clears the view family texture; all remaining views load.
+		LoadAction);
+}
+
 FScreenPassTextureViewportParameters GetScreenPassTextureViewportParameters(const FScreenPassTextureViewport& InViewport)
 {
 	const FVector2f Extent(InViewport.Extent);
@@ -112,9 +152,40 @@ void SetScreenPassPipelineState(FRHICommandList& RHICmdList, const FScreenPassPi
 	SetGraphicsPipelineState(RHICmdList, GraphicsPSOInit, ScreenPassDraw.StencilRef);
 }
 
+void DrawScreenPass_PostSetup(
+	FRHICommandList& RHICmdList,
+	const FSceneView& View,
+	const FScreenPassTextureViewport& OutputViewport,
+	const FScreenPassTextureViewport& InputViewport,
+	const FScreenPassPipelineState& PipelineState,
+	EScreenPassDrawFlags Flags)
+{
+	const FIntRect InputRect = InputViewport.Rect;
+	const FIntPoint InputSize = InputViewport.Extent;
+	const FIntRect OutputRect = OutputViewport.Rect;
+	const FIntPoint OutputSize = OutputViewport.Rect.Size();
+
+	FIntPoint LocalOutputPos(FIntPoint::ZeroValue);
+	FIntPoint LocalOutputSize(OutputSize);
+	EDrawRectangleFlags DrawRectangleFlags = EDRF_UseTriangleOptimization;
+
+	const bool bUseHMDHiddenAreaMask = EnumHasAllFlags(Flags, EScreenPassDrawFlags::AllowHMDHiddenAreaMask) && View.bHMDHiddenAreaMaskActive;
+
+	DrawPostProcessPass(
+		RHICmdList,
+		LocalOutputPos.X, LocalOutputPos.Y, LocalOutputSize.X, LocalOutputSize.Y,
+		InputRect.Min.X, InputRect.Min.Y, InputRect.Width(), InputRect.Height(),
+		OutputSize,
+		InputSize,
+		PipelineState.VertexShader,
+		View.StereoViewIndex,
+		bUseHMDHiddenAreaMask,
+		DrawRectangleFlags);
+}
+
 void AddDrawTexturePass(
 	FRDGBuilder& GraphBuilder,
-	const FViewInfo& View,
+	const FSceneView& View,
 	FRDGTextureRef InputTexture,
 	FRDGTextureRef OutputTexture,
 	FIntPoint InputPosition,
@@ -144,7 +215,7 @@ void AddDrawTexturePass(
 	const FScreenPassTextureViewport InputViewport(InputDesc.Extent, FIntRect(InputPosition, InputPosition + Size));
 	const FScreenPassTextureViewport OutputViewport(OutputDesc.Extent, FIntRect(OutputPosition, OutputPosition + Size));
 
-	TShaderMapRef<FCopyRectPS> PixelShader(View.ShaderMap);
+	TShaderMapRef<FCopyRectPS> PixelShader(static_cast<const FViewInfo&>(View).ShaderMap);
 
 	FCopyRectPS::FParameters* Parameters = GraphBuilder.AllocParameters<FCopyRectPS::FParameters>();
 	Parameters->InputTexture = InputTexture;
@@ -156,14 +227,14 @@ void AddDrawTexturePass(
 
 void AddDrawTexturePass(
 	FRDGBuilder& GraphBuilder,
-	const FViewInfo& View,
+	const FSceneView& View,
 	FScreenPassTexture Input,
 	FScreenPassRenderTarget Output)
 {
 	const FScreenPassTextureViewport InputViewport(Input);
 	const FScreenPassTextureViewport OutputViewport(Output);
 
-	TShaderMapRef<FCopyRectPS> PixelShader(View.ShaderMap);
+	TShaderMapRef<FCopyRectPS> PixelShader(static_cast<const FViewInfo&>(View).ShaderMap);
 
 	FCopyRectPS::FParameters* Parameters = GraphBuilder.AllocParameters<FCopyRectPS::FParameters>();
 	Parameters->InputTexture = Input.Texture;
diff --git a/Engine/Source/Runtime/Renderer/Private/ScreenPass.h b/Engine/Source/Runtime/Renderer/Private/ScreenPass.h
--- a/Engine/Source/Runtime/Renderer/Private/ScreenPass.h
+++ b/Engine/Source/Runtime/Renderer/Private/ScreenPass.h
@@ -2,14 +2,14 @@
 
 #pragma once
 
-#include "CoreMinimal.h"
+#include "CanvasTypes.h"
 #include "CommonRenderResources.h"
 #include "GlobalShader.h"
 #include "PipelineStateCache.h"
-#include "RenderGraph.h"
-#include "PostProcess/SceneFilterRendering.h"
-#include "RenderTargetTemp.h"
-#include "CanvasTypes.h"
+#include "RenderGraphBuilder.h"
+#include "RenderGraphUtils.h"
+#include "RHIStaticStates.h"
+#include "SceneView.h"
 
 // Returns whether a HMD hidden area mask is being used for VR.
 bool IsHMDHiddenAreaMaskActive();
@@ -61,7 +61,7 @@ struct FScreenPassTexture
 // Describes a texture with a load action for usage as a render target.
 struct FScreenPassRenderTarget : public FScreenPassTexture
 {
-	static FScreenPassRenderTarget CreateFromInput(
+	RENDERER_API static FScreenPassRenderTarget CreateFromInput(
 		FRDGBuilder& GraphBuilder,
 		FScreenPassTexture Input,
 		ERenderTargetLoadAction OutputLoadAction,
@@ -392,6 +392,14 @@ enum class EScreenPassDrawFlags : uint8
 };
 ENUM_CLASS_FLAGS(EScreenPassDrawFlags);
 
+RENDERER_API void DrawScreenPass_PostSetup(
+	FRHICommandList& RHICmdList,
+	const FSceneView& View,
+	const FScreenPassTextureViewport& OutputViewport,
+	const FScreenPassTextureViewport& InputViewport,
+	const FScreenPassPipelineState& PipelineState,
+	EScreenPassDrawFlags Flags);
+
 /** More advanced variant of screen pass drawing. Supports overriding blend / depth stencil
  *  pipeline state, and providing a custom vertex shader. Shader parameters are not bound by
  *  this method, instead the user provides a setup function that is called prior to draw, but
@@ -400,7 +408,7 @@ ENUM_CLASS_FLAGS(EScreenPassDrawFlags);
 template<typename TSetupFunction>
 void DrawScreenPass(
 	FRHICommandList& RHICmdList,
-	const FViewInfo& View,
+	const FSceneView& View,
 	const FScreenPassTextureViewport& OutputViewport,
 	const FScreenPassTextureViewport& InputViewport,
 	const FScreenPassPipelineState& PipelineState,
@@ -409,10 +417,7 @@ void DrawScreenPass(
 {
 	PipelineState.Validate();
 
-	const FIntRect InputRect = InputViewport.Rect;
-	const FIntPoint InputSize = InputViewport.Extent;
 	const FIntRect OutputRect = OutputViewport.Rect;
-	const FIntPoint OutputSize = OutputRect.Size();
 
 	RHICmdList.SetViewport(OutputRect.Min.X, OutputRect.Min.Y, 0.0f, OutputRect.Max.X, OutputRect.Max.Y, 1.0f);
 
@@ -420,24 +425,7 @@ void DrawScreenPass(
 
 	SetupFunction(RHICmdList);
 
-	FIntPoint LocalOutputPos(FIntPoint::ZeroValue);
-	FIntPoint LocalOutputSize(OutputSize);
-	EDrawRectangleFlags DrawRectangleFlags = EDRF_UseTriangleOptimization;
-
-	const bool bUseHMDHiddenAreaMask = (Flags & EScreenPassDrawFlags::AllowHMDHiddenAreaMask) == EScreenPassDrawFlags::AllowHMDHiddenAreaMask
-		? View.bHMDHiddenAreaMaskActive
-		: false;
-
-	DrawPostProcessPass(
-		RHICmdList,
-		LocalOutputPos.X, LocalOutputPos.Y, LocalOutputSize.X, LocalOutputSize.Y,
-		InputRect.Min.X, InputRect.Min.Y, InputRect.Width(), InputRect.Height(),
-		OutputSize,
-		InputSize,
-		PipelineState.VertexShader,
-		View.StereoViewIndex,
-		bUseHMDHiddenAreaMask,
-		DrawRectangleFlags);
+	DrawScreenPass_PostSetup(RHICmdList, View, OutputViewport, InputViewport, PipelineState, Flags);
 }
 
 /** Render graph variant of simpler DrawScreenPass function. Clears graph resources unused by the
@@ -447,7 +435,7 @@ template <typename PixelShaderType>
 FORCEINLINE void AddDrawScreenPass(
 	FRDGBuilder& GraphBuilder,
 	FRDGEventName&& PassName,
-	const FViewInfo& View,
+	const FSceneView& View,
 	const FScreenPassTextureViewport& OutputViewport,
 	const FScreenPassTextureViewport& InputViewport,
 	const TShaderRef<FShader>& VertexShader,
@@ -482,7 +470,7 @@ template <typename PixelShaderType>
 FORCEINLINE void AddDrawScreenPass(
 	FRDGBuilder& GraphBuilder,
 	FRDGEventName&& PassName,
-	const FViewInfo& View,
+	const FSceneView& View,
 	const FScreenPassTextureViewport& OutputViewport,
 	const FScreenPassTextureViewport& InputViewport,
 	const TShaderRef<FShader>& VertexShader,
@@ -499,7 +487,7 @@ template <typename PixelShaderType>
 FORCEINLINE void AddDrawScreenPass(
 	FRDGBuilder& GraphBuilder,
 	FRDGEventName&& PassName,
-	const FViewInfo& View,
+	const FSceneView& View,
 	const FScreenPassTextureViewport& OutputViewport,
 	const FScreenPassTextureViewport& InputViewport,
 	const TShaderRef<FShader>& VertexShader,
@@ -516,7 +504,7 @@ template <typename PixelShaderType>
 FORCEINLINE void AddDrawScreenPass(
 	FRDGBuilder& GraphBuilder,
 	FRDGEventName&& PassName,
-	const FViewInfo& View,
+	const FSceneView& View,
 	const FScreenPassTextureViewport& OutputViewport,
 	const FScreenPassTextureViewport& InputViewport,
 	const TShaderRef<FShader>& VertexShader,
@@ -533,14 +521,15 @@ template <typename PixelShaderType>
 FORCEINLINE void AddDrawScreenPass(
 	FRDGBuilder& GraphBuilder,
 	FRDGEventName&& PassName,
-	const FViewInfo& View,
+	const FSceneView& View,
 	const FScreenPassTextureViewport& OutputViewport,
 	const FScreenPassTextureViewport& InputViewport,
 	const TShaderRef<PixelShaderType>& PixelShader,
 	typename PixelShaderType::FParameters* PixelShaderParameters,
 	EScreenPassDrawFlags Flags = EScreenPassDrawFlags::None)
 {
-	TShaderMapRef<FScreenPassVS> VertexShader(View.ShaderMap);
+	FGlobalShaderMap* ShaderMap = GetGlobalShaderMap(View.GetFeatureLevel());
+	TShaderMapRef<FScreenPassVS> VertexShader(ShaderMap);
 	FRHIBlendState* BlendState = FScreenPassPipelineState::FDefaultBlendState::GetRHI();
 	FRHIDepthStencilState* DepthStencilState = FScreenPassPipelineState::FDefaultDepthStencilState::GetRHI();
 	AddDrawScreenPass(GraphBuilder, Forward<FRDGEventName&&>(PassName), View, OutputViewport, InputViewport, VertexShader, PixelShader, BlendState, DepthStencilState, PixelShaderParameters, Flags);
@@ -553,7 +542,7 @@ template <typename TSetupFunction, typename TPassParameterStruct>
 FORCEINLINE void AddDrawScreenPass(
 	FRDGBuilder& GraphBuilder,
 	FRDGEventName&& PassName,
-	const FViewInfo& View,
+	const FSceneView& View,
 	const FScreenPassTextureViewport& OutputViewport,
 	const FScreenPassTextureViewport& InputViewport,
 	const FScreenPassPipelineState& PipelineState,
@@ -578,7 +567,7 @@ template <typename TSetupFunction, typename TPassParameterStruct>
 FORCEINLINE void AddDrawScreenPass(
 	FRDGBuilder& GraphBuilder,
 	FRDGEventName&& PassName,
-	const FViewInfo& View,
+	const FSceneView& View,
 	const FScreenPassTextureViewport& OutputViewport,
 	const FScreenPassTextureViewport& InputViewport,
 	const FScreenPassPipelineState& PipelineState,
@@ -595,7 +584,7 @@ FORCEINLINE void AddDrawScreenPass(
  */
 void RENDERER_API AddDrawTexturePass(
 	FRDGBuilder& GraphBuilder,
-	const FViewInfo& View,
+	const FSceneView& View,
 	FRDGTextureRef InputTexture,
 	FRDGTextureRef OutputTexture,
 	FIntPoint InputPosition = FIntPoint::ZeroValue,
@@ -605,7 +594,7 @@ void RENDERER_API AddDrawTexturePass(
 /** Helper variant which takes a shared viewport instead of unique input / output positions. */
 FORCEINLINE void AddDrawTexturePass(
 	FRDGBuilder& GraphBuilder,
-	const FViewInfo& View,
+	const FSceneView& View,
 	FRDGTextureRef InputTexture,
 	FRDGTextureRef OutputTexture,
 	FIntRect ViewportRect)
@@ -615,7 +604,7 @@ FORCEINLINE void AddDrawTexturePass(
 
 void RENDERER_API AddDrawTexturePass(
 	FRDGBuilder& GraphBuilder,
-	const FViewInfo& View,
+	const FSceneView& View,
 	FScreenPassTexture Input,
 	FScreenPassRenderTarget Output);
 
diff --git a/Engine/Source/Runtime/Renderer/Private/ScreenPass.inl b/Engine/Source/Runtime/Renderer/Private/ScreenPass.inl
--- a/Engine/Source/Runtime/Renderer/Private/ScreenPass.inl
+++ b/Engine/Source/Runtime/Renderer/Private/ScreenPass.inl
@@ -63,43 +63,6 @@ inline FIntRect GetRectFromExtent(FIntPoint Extent)
 	return FIntRect(FIntPoint::ZeroValue, Extent);
 }
 
-inline FScreenPassRenderTarget FScreenPassRenderTarget::CreateFromInput(
-	FRDGBuilder& GraphBuilder,
-	FScreenPassTexture Input,
-	ERenderTargetLoadAction OutputLoadAction,
-	const TCHAR* OutputName)
-{
-	check(Input.IsValid());
-
-	FRDGTextureDesc OutputDesc = Input.Texture->Desc;
-	OutputDesc.Reset();
-
-	return FScreenPassRenderTarget(GraphBuilder.CreateTexture(OutputDesc, OutputName), Input.ViewRect, OutputLoadAction);
-}
-
-inline FScreenPassRenderTarget FScreenPassRenderTarget::CreateViewFamilyOutput(FRDGTextureRef ViewFamilyTexture, const FViewInfo& View)
-{
-	const FIntRect ViewRect = View.PrimaryScreenPercentageMethod == EPrimaryScreenPercentageMethod::RawOutput ? View.ViewRect : View.UnscaledViewRect;
-
-	ERenderTargetLoadAction LoadAction = ERenderTargetLoadAction::ENoAction;
-
-	if (!View.IsFirstInFamily() || View.Family->bAdditionalViewFamily)
-	{
-		LoadAction = ERenderTargetLoadAction::ELoad;
-	}
-	else if (ViewRect.Min != FIntPoint::ZeroValue || ViewRect.Size() != ViewFamilyTexture->Desc.Extent)
-	{
-		LoadAction = ERenderTargetLoadAction::EClear;
-	}
-
-	return FScreenPassRenderTarget(
-		ViewFamilyTexture,
-		// Raw output mode uses the original view rect. Otherwise the final unscaled rect is used.
-		ViewRect,
-		// First view clears the view family texture; all remaining views load.
-		LoadAction);
-}
-
 inline FScreenPassTexture::FScreenPassTexture(FRDGTextureRef InTexture)
 	: Texture(InTexture)
 {
