Updated GPU scene parallel for tasks with names.

diff --git a/Engine/Source/Runtime/Renderer/Private/GPUScene.cpp b/Engine/Source/Runtime/Renderer/Private/GPUScene.cpp
--- a/Engine/Source/Runtime/Renderer/Private/GPUScene.cpp
+++ b/Engine/Source/Runtime/Renderer/Private/GPUScene.cpp
@@ -1055,7 +1055,7 @@ void FGPUScene::UploadGeneral(FRDGBuilder& GraphBuilder, FScene& Scene, FRDGExte
 
 			TaskContext.PrimitiveUploader->Add(UploadDataSourceAdapter.GetItemPrimitiveIds());
 
-			ParallelFor(TaskContext.NumPrimitiveDataUploads, [&TaskContext, &UploadDataSourceAdapter](int32 ItemIndex)
+			ParallelForTemplate(TEXT("GPUScene Upload Primitives Task"), TaskContext.NumPrimitiveDataUploads, 1, [&TaskContext, &UploadDataSourceAdapter](int32 ItemIndex)
 			{
 				FTaskTagScope TaskTagScope(ETaskTag::EParallelRenderingThread);
 
@@ -1068,14 +1068,14 @@ void FGPUScene::UploadGeneral(FRDGBuilder& GraphBuilder, FScene& Scene, FRDGExte
 					DstData[VectorIndex] = UploadInfo.PrimitiveSceneData.Data[VectorIndex];
 				}
 
-			}, !bExecuteInParallel);
+			}, bExecuteInParallel ? EParallelForFlags::None : EParallelForFlags::ForceSingleThread);
 		}
 
 		if (TaskContext.NumInstanceSceneDataUploads > 0 && !InstanceUpdates.UpdateBatches.IsEmpty())
 		{
 			SCOPED_NAMED_EVENT(Instances, FColor::Green);
 
-			ParallelFor(InstanceUpdates.UpdateBatches.Num(), [&TaskContext, &InstanceUpdates, &UploadDataSourceAdapter](int32 BatchIndex)
+			ParallelForTemplate(TEXT("GPUScene Upload Instances Task"), InstanceUpdates.UpdateBatches.Num(), 1, [&TaskContext, &InstanceUpdates, &UploadDataSourceAdapter](int32 BatchIndex)
 			{
 				const FInstanceUploadBatch Batch = InstanceUpdates.UpdateBatches[BatchIndex];
 
@@ -1200,7 +1200,7 @@ void FGPUScene::UploadGeneral(FRDGBuilder& GraphBuilder, FScene& Scene, FRDGExte
 					}
 				}
 
-			}, !bExecuteInParallel);
+			}, bExecuteInParallel ? EParallelForFlags::None : EParallelForFlags::ForceSingleThread);
 		}
 
 		if (TaskContext.InstanceBVHUploader)
@@ -1863,6 +1863,7 @@ void FGPUScene::EndDeferAllocatorMerges()
 {
 	if (GGPUSceneAllowDeferredAllocatorMerges != 0)
 	{
+		SCOPED_NAMED_EVENT(FGPUScene_EndDeferAllocatorMerges, FColor::Green);
 		InstanceSceneDataAllocator.EndDeferMerges();
 		InstancePayloadDataAllocator.EndDeferMerges();
 		LightmapDataAllocator.EndDeferMerges();
