Only use BF_Source1Color for ThinTranlucent blend state when a single render target is set (when bRequiresSceneDepthAux is set depth is also bound in RT 1)

#rnx

diff --git a/Engine/Source/Runtime/Renderer/Private/MobileBasePass.cpp b/Engine/Source/Runtime/Renderer/Private/MobileBasePass.cpp
--- a/Engine/Source/Runtime/Renderer/Private/MobileBasePass.cpp
+++ b/Engine/Source/Runtime/Renderer/Private/MobileBasePass.cpp
@@ -425,7 +425,8 @@ void MobileBasePass::SetTranslucentRenderState(FMeshPassProcessorRenderState& Dr
 	}
 	else if (ShadingModels.HasShadingModel(MSM_ThinTranslucent))
 	{
-		if (bIsDualSourceBlending)
+		const bool bRequiresSceneDepthAux = MobileRequiresSceneDepthAux(ShaderPlatform);
+		if (bIsDualSourceBlending && !bRequiresSceneDepthAux)
 		{
 			// Blend by putting add in target 0 and multiply by background in target 1.
 			DrawRenderState.SetBlendState(TStaticBlendState<CW_RGBA, BO_Add, BF_One, BF_Source1Color, BO_Add, BF_One, BF_Source1Alpha,
