Moving ScreenPass.h, ScreenPass.inl and SystemTextures.h to Public.

diff --git a/Engine/Plugins/Experimental/ColorCorrectRegions/Source/ColorCorrectRegions/Public/ColorCorrectRegionsPostProcessMaterial.h b/Engine/Plugins/Experimental/ColorCorrectRegions/Source/ColorCorrectRegions/Public/ColorCorrectRegionsPostProcessMaterial.h
--- a/Engine/Plugins/Experimental/ColorCorrectRegions/Source/ColorCorrectRegions/Public/ColorCorrectRegionsPostProcessMaterial.h
+++ b/Engine/Plugins/Experimental/ColorCorrectRegions/Source/ColorCorrectRegions/Public/ColorCorrectRegionsPostProcessMaterial.h
@@ -19,8 +19,8 @@
 #include "UObject/ConstructorHelpers.h"
 
 // FScreenPassTextureViewportParameters and FScreenPassTextureInput
-#include "Runtime/Renderer/Private/ScreenPass.h"
-#include "Runtime/Renderer/Private/SceneTextureParameters.h"
+#include "ScreenPass.h"
+#include "SceneTexturesConfig.h"
 
 
 
diff --git a/Engine/Source/Runtime/Renderer/Private/ScreenPass.h b/Engine/Source/Runtime/Renderer/Private/ScreenPass.h
deleted file mode 100644
diff --git a/Engine/Source/Runtime/Renderer/Private/ScreenPass.inl b/Engine/Source/Runtime/Renderer/Private/ScreenPass.inl
deleted file mode 100644
diff --git a/Engine/Source/Runtime/Renderer/Private/SystemTextures.h b/Engine/Source/Runtime/Renderer/Private/SystemTextures.h
deleted file mode 100644
diff --git a/Engine/Source/Runtime/Renderer/Public/ScreenPass.h b/Engine/Source/Runtime/Renderer/Public/ScreenPass.h
new file mode 100644
--- /dev/null
+++ b/Engine/Source/Runtime/Renderer/Public/ScreenPass.h
@@ -0,0 +1,665 @@
+// Copyright Epic Games, Inc. All Rights Reserved.
+
+#pragma once
+
+#include "CanvasTypes.h"
+#include "CommonRenderResources.h"
+#include "GlobalShader.h"
+#include "PipelineStateCache.h"
+#include "RenderGraphBuilder.h"
+#include "RenderGraphUtils.h"
+#include "RHIStaticStates.h"
+#include "SceneView.h"
+
+// Returns whether a HMD hidden area mask is being used for VR.
+bool IsHMDHiddenAreaMaskActive();
+
+// Returns the global engine mini font texture.
+FRHITexture* GetMiniFontTexture();
+
+// Creates and returns an RDG texture for the view family output. Returns null if no RHI texture exists.
+FRDGTextureRef RENDERER_API TryCreateViewFamilyTexture(FRDGBuilder& GraphBuilder, const FSceneViewFamily& ViewFamily);
+
+// The vertex shader used by DrawScreenPass to draw a rectangle.
+class RENDERER_API FScreenPassVS : public FGlobalShader
+{
+public:
+	DECLARE_GLOBAL_SHADER(FScreenPassVS);
+
+	static bool ShouldCompilePermutation(const FGlobalShaderPermutationParameters&)
+	{
+		return true;
+	}
+
+	FScreenPassVS() = default;
+	FScreenPassVS(const ShaderMetaType::CompiledShaderInitializerType& Initializer)
+		: FGlobalShader(Initializer)
+	{}
+};
+
+// Describes a texture with a paired viewport rect.
+struct FScreenPassTexture
+{
+	FScreenPassTexture() = default;
+
+	explicit FScreenPassTexture(FRDGTextureRef InTexture);
+
+	FScreenPassTexture(FRDGTextureRef InTexture, FIntRect InViewRect)
+		: Texture(InTexture)
+		, ViewRect(InViewRect)
+	{}
+
+	bool IsValid() const;
+
+	bool operator==(FScreenPassTexture Other) const;
+	bool operator!=(FScreenPassTexture Other) const;
+
+	FRDGTextureRef Texture = nullptr;
+	FIntRect ViewRect;
+};
+
+// Describes a texture with a load action for usage as a render target.
+struct FScreenPassRenderTarget : public FScreenPassTexture
+{
+	RENDERER_API static FScreenPassRenderTarget CreateFromInput(
+		FRDGBuilder& GraphBuilder,
+		FScreenPassTexture Input,
+		ERenderTargetLoadAction OutputLoadAction,
+		const TCHAR* OutputName);
+
+	static FScreenPassRenderTarget CreateViewFamilyOutput(FRDGTextureRef ViewFamilyTexture, const FViewInfo& View);
+
+	FScreenPassRenderTarget() = default;
+
+	FScreenPassRenderTarget(FScreenPassTexture InTexture, ERenderTargetLoadAction InLoadAction)
+		: FScreenPassTexture(InTexture)
+		, LoadAction(InLoadAction)
+	{}
+
+	FScreenPassRenderTarget(FRDGTextureRef InTexture, ERenderTargetLoadAction InLoadAction)
+		: FScreenPassTexture(InTexture)
+		, LoadAction(InLoadAction)
+	{}
+
+	FScreenPassRenderTarget(FRDGTextureRef InTexture, FIntRect InViewRect, ERenderTargetLoadAction InLoadAction)
+		: FScreenPassTexture(InTexture, InViewRect)
+		, LoadAction(InLoadAction)
+	{}
+
+	bool operator==(FScreenPassRenderTarget Other) const;
+	bool operator!=(FScreenPassRenderTarget Other) const;
+
+	FRenderTargetBinding GetRenderTargetBinding() const;
+
+	ERenderTargetLoadAction LoadAction = ERenderTargetLoadAction::ENoAction;
+};
+
+// Describes a view rect contained within the extent of a texture. Used to derive texture coordinate transformations.
+class FScreenPassTextureViewport
+{
+public:
+	// Creates a viewport that is downscaled by an integer multiple.
+	static FScreenPassTextureViewport CreateDownscaled(const FScreenPassTextureViewport& Other, FIntPoint ScaleFactor);
+
+	// Creates a viewport scaled by a floating point multiplier.
+	static FScreenPassTextureViewport CreateScaled(const FScreenPassTextureViewport& Other, FVector2D Scale);
+
+	FScreenPassTextureViewport() = default;
+
+	explicit FScreenPassTextureViewport(FIntRect InRect)
+		: Extent(InRect.Max)
+		, Rect(InRect)
+	{}
+
+	explicit FScreenPassTextureViewport(FScreenPassTexture InTexture);
+
+	explicit FScreenPassTextureViewport(FRDGTextureRef InTexture)
+		: FScreenPassTextureViewport(FScreenPassTexture(InTexture))
+	{}
+
+	explicit FScreenPassTextureViewport(FIntPoint InExtent)
+		: Extent(InExtent)
+		, Rect(FIntPoint::ZeroValue, InExtent)
+	{}
+
+	FScreenPassTextureViewport(FIntPoint InExtent, FIntRect InRect)
+		: Extent(InExtent)
+		, Rect(InRect)
+	{}
+
+	FScreenPassTextureViewport(FRDGTextureRef InTexture, FIntRect InRect)
+		: FScreenPassTextureViewport(FScreenPassTexture(InTexture, InRect))
+	{}
+
+	FScreenPassTextureViewport(const FScreenPassTextureViewport&) = default;
+
+	bool operator==(const FScreenPassTextureViewport& Other) const;
+	bool operator!=(const FScreenPassTextureViewport& Other) const;
+
+	// Returns whether the viewport contains an empty viewport or extent.
+	bool IsEmpty() const;
+
+	// Returns whether the viewport covers the full extent of the texture.
+	bool IsFullscreen() const;
+
+	// Returns the ratio of rect size to extent along each axis.
+	FVector2D GetRectToExtentRatio() const;
+
+	// The texture extent, in pixels; defines a super-set [0, 0]x(Extent, Extent).
+	FIntPoint Extent = FIntPoint::ZeroValue;
+
+	// The viewport rect, in pixels; defines a sub-set within [0, 0]x(Extent, Extent).
+	FIntRect Rect;
+};
+
+// Returns an extent downscaled by a multiple of the integer divisor (and clamped to 1).
+FIntPoint GetDownscaledExtent(FIntPoint Extent, FIntPoint Divisor);
+
+// Returns an extent scaled by the floating point scale factor.
+FIntPoint GetScaledExtent(FIntPoint Extent, FVector2D Multiplier);
+FIntPoint GetScaledExtent(FIntPoint Extent, float Multiplier);
+
+// Returns a rect downscaled by a multiple of the integer divisor.
+FIntRect GetDownscaledRect(FIntRect Rect, FIntPoint Divisor);
+
+// Returns a rect scaled by the floating point scale factor.
+FIntRect GetScaledRect(FIntRect Rect, FVector2D Multiplier);
+FIntRect GetScaledRect(FIntRect Rect, float Multiplier);
+
+// Returns the texture viewport downscaled by an integer divisor.
+FScreenPassTextureViewport GetDownscaledViewport(FScreenPassTextureViewport Viewport, FIntPoint Divisor);
+
+// Returns the texture viewport scaled by a float multiplier.
+FScreenPassTextureViewport GetScaledViewport(FScreenPassTextureViewport Viewport, FVector2D Multiplier);
+
+// Returns a rect with the min point at the origin and the max point at Extent.
+FIntRect GetRectFromExtent(FIntPoint Extent);
+
+// Describes the set of shader parameters for a screen pass texture viewport.
+BEGIN_SHADER_PARAMETER_STRUCT(FScreenPassTextureViewportParameters, )
+	// Texture extent in pixels.
+	SHADER_PARAMETER(FVector2f, Extent)
+	SHADER_PARAMETER(FVector2f, ExtentInverse)
+
+	// Scale / Bias factor to convert from [-1, 1] to [ViewportMin, ViewportMax]
+	SHADER_PARAMETER(FVector2f, ScreenPosToViewportScale)
+	SHADER_PARAMETER(FVector2f, ScreenPosToViewportBias)
+
+	// Texture viewport min / max in pixels.
+	SHADER_PARAMETER(FIntPoint, ViewportMin)
+	SHADER_PARAMETER(FIntPoint, ViewportMax)
+
+	// Texture viewport size in pixels.
+	SHADER_PARAMETER(FVector2f, ViewportSize)
+	SHADER_PARAMETER(FVector2f, ViewportSizeInverse)
+
+	// Texture viewport min / max in normalized UV coordinates, with respect to the texture extent.
+	SHADER_PARAMETER(FVector2f, UVViewportMin)
+	SHADER_PARAMETER(FVector2f, UVViewportMax)
+
+	// Texture viewport size in normalized UV coordinates, with respect to the texture extent.
+	SHADER_PARAMETER(FVector2f, UVViewportSize)
+	SHADER_PARAMETER(FVector2f, UVViewportSizeInverse)
+
+	// Texture viewport min / max in normalized UV coordinates, with respect to the texture extent,
+	// adjusted by a half pixel offset for bilinear filtering. Useful for clamping to avoid sampling
+	// pixels on viewport edges; e.g. clamp(UV, UVViewportBilinearMin, UVViewportBilinearMax);
+	SHADER_PARAMETER(FVector2f, UVViewportBilinearMin)
+	SHADER_PARAMETER(FVector2f, UVViewportBilinearMax)
+END_SHADER_PARAMETER_STRUCT()
+
+FScreenPassTextureViewportParameters RENDERER_API GetScreenPassTextureViewportParameters(const FScreenPassTextureViewport& InViewport);
+
+/** Generic affine 2D texture coordinate transformation x * S + B.
+ *
+ * Construct:
+ *		FVector2f PointInA = ...;
+ *		FVector2f PointInB = PointInA * Scale0 + Bias0;
+ * 
+ *		FScreenTransform AToB(Scale0, Bias0);
+ *		FVector2f PointInB = PointInA * AToB;
+ * 
+ * Associativity:
+ *		FVector2f PointInA = ...;
+ *		FScreenTransform AToB = ...;
+ *		FScreenTransform BToC = ...;
+ *		FVector2f PointInC = (PointInA * AToB) * BToC;
+ *
+ *		FScreenTransform AToC = AToB * BToC;
+ *		FVector2f PointInC = PointInA * AToC;
+ * 
+ * Explicit construction by factorization:
+ *		FVector2f PointInA = ...;
+ *		FVector2f PointInB = PointInA * Scale0 + Bias0;
+ *		FVector2f PointInC = PointInB * Scale1 + Bias1;
+ * 
+ *		FScreenTransform AToC = (FScreenTransform::Identity * Scale0 + Bias0) * Scale1 + Bias1;
+ *		FVector2f PointInC = PointInA * AToC;
+ *
+ * Shader code:
+ *		#include "/Engine/Private/ScreenPass.ush"
+ * 
+ *		FScreenTransform AToC; // shader parameter in global scope
+ * 
+ *		{
+ *			float2 PointInA = ...;
+ *			float2 PointInC = ApplyScreenTransform(PointInA, AToC);
+ *		}
+ * 
+ */
+struct FScreenTransform
+{
+	FVector2f Scale;
+	FVector2f Bias;
+
+	inline FScreenTransform()
+	{ }
+
+	inline FScreenTransform(const FVector2f& InScale, const FVector2f& InBias)
+		: Scale(InScale)
+		, Bias(InBias)
+	{ }
+
+	inline FScreenTransform(const FScreenTransform& AToB)
+		: Scale(AToB.Scale)
+		, Bias(AToB.Bias)
+	{ }
+
+
+	// A * FScreenTransform::Identity = A
+	static RENDERER_API const FScreenTransform Identity;
+
+	// Transforms ScreenPos to/from ViewportUV
+	static RENDERER_API const FScreenTransform ScreenPosToViewportUV;
+	static RENDERER_API const FScreenTransform ViewportUVToScreenPos;
+
+
+	/** Invert a transformation AToB to BToA. */
+	static inline FScreenTransform Invert(const FScreenTransform& AToB);
+
+	/** Change of coordinate to map from a rectangle to another. */
+	static FScreenTransform ChangeRectFromTo(
+		FVector2f SourceOffset, FVector2f SourceExtent,
+		FVector2f DestinationOffset, FVector2f DestinationExtent);
+	static FScreenTransform ChangeRectFromTo(const FIntRect& SrcViewport, const FIntRect& DestViewport);
+
+	/** Different texture coordinate basis. */
+	enum class ETextureBasis
+	{
+		// Viewport maps [-1.0,1.0] on X, ]1.0, -1.0[ on Y.
+		ScreenPosition,
+
+		// Viewport maps [0.0,1.0]
+		ViewportUV,
+
+		// Viewport maps [Viewport.Min,Viewport.Max] in pixel coordinate in the texture
+		// Used for instance for MyTexture[uint(TexelPosition)];
+		TexelPosition,
+
+		// Viewport maps [Viewport.Min / TextureExtent,Viewport.Max / TextureExtent]
+		// Used for MyTexture.SampleLevel(MySampler, TextureUV, 0);
+		TextureUV,
+	};
+
+	/** Change of basis for texture coordinate. */
+	static inline FScreenTransform ChangeTextureBasisFromTo(
+		const FIntPoint& TextureExtent, const FIntRect& TextureViewport,
+		ETextureBasis SrcBasis, ETextureBasis DestBasis);
+
+	static inline FScreenTransform ChangeTextureBasisFromTo(
+		const FScreenPassTextureViewport& TextureViewport,
+		ETextureBasis SrcBasis, ETextureBasis DestBasis)
+	{
+		return ChangeTextureBasisFromTo(TextureViewport.Extent, TextureViewport.Rect, SrcBasis, DestBasis);
+	}
+
+	static inline FScreenTransform ChangeTextureBasisFromTo(
+		const FScreenPassTexture& Texture,
+		ETextureBasis SrcBasis, ETextureBasis DestBasis)
+	{
+		return ChangeTextureBasisFromTo(FScreenPassTextureViewport(Texture.Texture, Texture.ViewRect), SrcBasis, DestBasis);
+	}
+
+	/** Change TextureUV coordinate from one texture to another, taking into account change in texture extent too. */
+	static FScreenTransform ChangeTextureUVCoordinateFromTo(
+		const FScreenPassTextureViewport& SrcViewport,
+		const FScreenPassTextureViewport& DestViewport);
+
+	static FScreenTransform SvPositionToViewportUV(const FIntRect& SrcViewport);
+	static FScreenTransform DispatchThreadIdToViewportUV(const FIntRect& SrcViewport);
+}; // FScreenTransform
+
+// A utility shader parameter struct containing the viewport, texture, and sampler for a unique texture input to a shader.
+BEGIN_SHADER_PARAMETER_STRUCT(FScreenPassTextureInput, )
+	SHADER_PARAMETER_STRUCT_INCLUDE(FScreenPassTextureViewportParameters, Viewport)
+	SHADER_PARAMETER_RDG_TEXTURE(Texture2D, Texture)
+	SHADER_PARAMETER_SAMPLER(SamplerState, Sampler)
+END_SHADER_PARAMETER_STRUCT()
+
+FScreenPassTextureInput GetScreenPassTextureInput(FScreenPassTexture Input, FRHISamplerState* Sampler);
+
+/** Draw information for the more advanced DrawScreenPass variant. Allows customizing the blend / depth stencil state,
+ *  providing a custom vertex shader, and more fine-grained control of the underlying draw call.
+ */
+struct FScreenPassPipelineState
+{
+	using FDefaultBlendState = TStaticBlendState<>;
+	using FDefaultDepthStencilState = TStaticDepthStencilState<false, CF_Always>;
+
+	FScreenPassPipelineState() = default;
+
+	FScreenPassPipelineState(
+		const TShaderRef<FShader>& InVertexShader,
+		const TShaderRef<FShader>& InPixelShader,
+		FRHIBlendState* InBlendState = FDefaultBlendState::GetRHI(),
+		FRHIDepthStencilState* InDepthStencilState = FDefaultDepthStencilState::GetRHI(),
+		uint32 InStencilRef = 0,
+		FRHIVertexDeclaration* InVertexDeclaration = GFilterVertexDeclaration.VertexDeclarationRHI)
+		: VertexShader(InVertexShader)
+		, PixelShader(InPixelShader)
+		, BlendState(InBlendState)
+		, DepthStencilState(InDepthStencilState)
+		, VertexDeclaration(InVertexDeclaration)
+		, StencilRef(InStencilRef)
+	{}
+
+	void Validate() const
+	{
+		check(VertexShader.IsValid());
+		check(PixelShader.IsValid());
+		check(BlendState);
+		check(DepthStencilState);
+		check(VertexDeclaration);
+	}
+
+	TShaderRef<FShader> VertexShader;
+	TShaderRef<FShader> PixelShader;
+	FRHIBlendState* BlendState = nullptr;
+	FRHIDepthStencilState* DepthStencilState = nullptr;
+	FRHIVertexDeclaration* VertexDeclaration = nullptr;
+	uint32 StencilRef{};
+};
+
+// Helper function which sets the pipeline state object on the command list prior to invoking a screen pass.
+void RENDERER_API SetScreenPassPipelineState(FRHICommandList& RHICmdList, const FScreenPassPipelineState& ScreenPassDraw);
+
+enum class EScreenPassDrawFlags : uint8
+{
+	None,
+
+	// Allows the screen pass to use a HMD hidden area mask if one is available. Used for VR.
+	AllowHMDHiddenAreaMask = 0x2
+};
+ENUM_CLASS_FLAGS(EScreenPassDrawFlags);
+
+RENDERER_API void DrawScreenPass_PostSetup(
+	FRHICommandList& RHICmdList,
+	const FSceneView& View,
+	const FScreenPassTextureViewport& OutputViewport,
+	const FScreenPassTextureViewport& InputViewport,
+	const FScreenPassPipelineState& PipelineState,
+	EScreenPassDrawFlags Flags);
+
+/** More advanced variant of screen pass drawing. Supports overriding blend / depth stencil
+ *  pipeline state, and providing a custom vertex shader. Shader parameters are not bound by
+ *  this method, instead the user provides a setup function that is called prior to draw, but
+ *  after setting the PSO. This setup function should assign shader parameters.
+ */
+template<typename TSetupFunction>
+void DrawScreenPass(
+	FRHICommandList& RHICmdList,
+	const FSceneView& View,
+	const FScreenPassTextureViewport& OutputViewport,
+	const FScreenPassTextureViewport& InputViewport,
+	const FScreenPassPipelineState& PipelineState,
+	EScreenPassDrawFlags Flags,
+	TSetupFunction SetupFunction)
+{
+	PipelineState.Validate();
+
+	const FIntRect OutputRect = OutputViewport.Rect;
+
+	RHICmdList.SetViewport(OutputRect.Min.X, OutputRect.Min.Y, 0.0f, OutputRect.Max.X, OutputRect.Max.Y, 1.0f);
+
+	SetScreenPassPipelineState(RHICmdList, PipelineState);
+
+	SetupFunction(RHICmdList);
+
+	DrawScreenPass_PostSetup(RHICmdList, View, OutputViewport, InputViewport, PipelineState, Flags);
+}
+
+/** Render graph variant of simpler DrawScreenPass function. Clears graph resources unused by the
+ *  pixel shader prior to adding the pass.
+ */
+template <typename PixelShaderType>
+FORCEINLINE void AddDrawScreenPass(
+	FRDGBuilder& GraphBuilder,
+	FRDGEventName&& PassName,
+	const FSceneView& View,
+	const FScreenPassTextureViewport& OutputViewport,
+	const FScreenPassTextureViewport& InputViewport,
+	const TShaderRef<FShader>& VertexShader,
+	const TShaderRef<PixelShaderType>& PixelShader,
+	FRHIBlendState* BlendState,
+	FRHIDepthStencilState* DepthStencilState,
+	typename PixelShaderType::FParameters* PixelShaderParameters,
+	EScreenPassDrawFlags Flags = EScreenPassDrawFlags::None)
+{
+	check(VertexShader.IsValid());
+	check(PixelShader.IsValid());
+	check(PixelShaderParameters);
+
+	ClearUnusedGraphResources(PixelShader, PixelShaderParameters);
+
+	const FScreenPassPipelineState PipelineState(VertexShader, PixelShader, BlendState, DepthStencilState);
+
+	GraphBuilder.AddPass(
+		Forward<FRDGEventName&&>(PassName),
+		PixelShaderParameters,
+		ERDGPassFlags::Raster,
+		[&View, OutputViewport, InputViewport, PipelineState, PixelShader, PixelShaderParameters, Flags](FRHICommandList& RHICmdList)
+	{
+		DrawScreenPass(RHICmdList, View, OutputViewport, InputViewport, PipelineState, Flags, [&](FRHICommandList&)
+		{
+			SetShaderParameters(RHICmdList, PixelShader, PixelShader.GetPixelShader(), *PixelShaderParameters);
+		});
+	});
+}
+
+template <typename PixelShaderType>
+FORCEINLINE void AddDrawScreenPass(
+	FRDGBuilder& GraphBuilder,
+	FRDGEventName&& PassName,
+	const FSceneView& View,
+	const FScreenPassTextureViewport& OutputViewport,
+	const FScreenPassTextureViewport& InputViewport,
+	const TShaderRef<FShader>& VertexShader,
+	const TShaderRef<PixelShaderType>& PixelShader,
+	typename PixelShaderType::FParameters* PixelShaderParameters,
+	EScreenPassDrawFlags Flags = EScreenPassDrawFlags::None)
+{
+	FRHIBlendState* BlendState = FScreenPassPipelineState::FDefaultBlendState::GetRHI();
+	FRHIDepthStencilState* DepthStencilState = FScreenPassPipelineState::FDefaultDepthStencilState::GetRHI();
+	AddDrawScreenPass(GraphBuilder, Forward<FRDGEventName&&>(PassName), View, OutputViewport, InputViewport, VertexShader, PixelShader, BlendState, DepthStencilState, PixelShaderParameters, Flags);
+}
+
+template <typename PixelShaderType>
+FORCEINLINE void AddDrawScreenPass(
+	FRDGBuilder& GraphBuilder,
+	FRDGEventName&& PassName,
+	const FSceneView& View,
+	const FScreenPassTextureViewport& OutputViewport,
+	const FScreenPassTextureViewport& InputViewport,
+	const TShaderRef<FShader>& VertexShader,
+	const TShaderRef<PixelShaderType>& PixelShader,
+	FRHIBlendState* BlendState,
+	typename PixelShaderType::FParameters* PixelShaderParameters,
+	EScreenPassDrawFlags Flags = EScreenPassDrawFlags::None)
+{
+	FRHIDepthStencilState* DepthStencilState = FScreenPassPipelineState::FDefaultDepthStencilState::GetRHI();
+	AddDrawScreenPass(GraphBuilder, Forward<FRDGEventName&&>(PassName), View, OutputViewport, InputViewport, VertexShader, PixelShader, BlendState, DepthStencilState, PixelShaderParameters, Flags);
+}
+
+template <typename PixelShaderType>
+FORCEINLINE void AddDrawScreenPass(
+	FRDGBuilder& GraphBuilder,
+	FRDGEventName&& PassName,
+	const FSceneView& View,
+	const FScreenPassTextureViewport& OutputViewport,
+	const FScreenPassTextureViewport& InputViewport,
+	const TShaderRef<FShader>& VertexShader,
+	const TShaderRef<PixelShaderType>& PixelShader,
+	FRHIDepthStencilState* DepthStencilState,
+	typename PixelShaderType::FParameters* PixelShaderParameters,
+	EScreenPassDrawFlags Flags = EScreenPassDrawFlags::None)
+{
+	FRHIBlendState* BlendState = FScreenPassPipelineState::FDefaultBlendState::GetRHI();
+	AddDrawScreenPass(GraphBuilder, Forward<FRDGEventName&&>(PassName), View, OutputViewport, InputViewport, VertexShader, PixelShader, BlendState, DepthStencilState, PixelShaderParameters, Flags);
+}
+
+template <typename PixelShaderType>
+FORCEINLINE void AddDrawScreenPass(
+	FRDGBuilder& GraphBuilder,
+	FRDGEventName&& PassName,
+	const FSceneView& View,
+	const FScreenPassTextureViewport& OutputViewport,
+	const FScreenPassTextureViewport& InputViewport,
+	const TShaderRef<PixelShaderType>& PixelShader,
+	typename PixelShaderType::FParameters* PixelShaderParameters,
+	EScreenPassDrawFlags Flags = EScreenPassDrawFlags::None)
+{
+	FGlobalShaderMap* ShaderMap = GetGlobalShaderMap(View.GetFeatureLevel());
+	TShaderMapRef<FScreenPassVS> VertexShader(ShaderMap);
+	FRHIBlendState* BlendState = FScreenPassPipelineState::FDefaultBlendState::GetRHI();
+	FRHIDepthStencilState* DepthStencilState = FScreenPassPipelineState::FDefaultDepthStencilState::GetRHI();
+	AddDrawScreenPass(GraphBuilder, Forward<FRDGEventName&&>(PassName), View, OutputViewport, InputViewport, VertexShader, PixelShader, BlendState, DepthStencilState, PixelShaderParameters, Flags);
+}
+
+/** Render graph variant of more advanced DrawScreenPass function. Does *not* clear unused graph
+ *  resources, since the parameters might be shared between the vertex and pixel shaders.
+ */
+template <typename TSetupFunction, typename TPassParameterStruct>
+FORCEINLINE void AddDrawScreenPass(
+	FRDGBuilder& GraphBuilder,
+	FRDGEventName&& PassName,
+	const FSceneView& View,
+	const FScreenPassTextureViewport& OutputViewport,
+	const FScreenPassTextureViewport& InputViewport,
+	const FScreenPassPipelineState& PipelineState,
+	TPassParameterStruct* PassParameterStruct,
+	EScreenPassDrawFlags Flags,
+	TSetupFunction SetupFunction)
+{
+	PipelineState.Validate();
+	check(PassParameterStruct);
+
+	GraphBuilder.AddPass(
+		Forward<FRDGEventName&&>(PassName),
+		PassParameterStruct,
+		ERDGPassFlags::Raster,
+		[&View, OutputViewport, InputViewport, PipelineState, SetupFunction, Flags] (FRHICommandList& RHICmdList)
+	{
+		DrawScreenPass(RHICmdList, View, OutputViewport, InputViewport, PipelineState, Flags, SetupFunction);
+	});
+}
+
+template <typename TSetupFunction, typename TPassParameterStruct>
+FORCEINLINE void AddDrawScreenPass(
+	FRDGBuilder& GraphBuilder,
+	FRDGEventName&& PassName,
+	const FSceneView& View,
+	const FScreenPassTextureViewport& OutputViewport,
+	const FScreenPassTextureViewport& InputViewport,
+	const FScreenPassPipelineState& PipelineState,
+	TPassParameterStruct* PassParameterStruct,
+	TSetupFunction SetupFunction)
+{
+	AddDrawScreenPass(GraphBuilder, Forward<FRDGEventName&&>(PassName), View, OutputViewport, InputViewport, PipelineState, PassParameterStruct, EScreenPassDrawFlags::None, SetupFunction);
+}
+
+/** Helper function which copies a region of an input texture to a region of the output texture,
+ *  with support for format conversion. If formats match, the method falls back to a simple DMA
+ *  (CopyTexture); otherwise, it rasterizes using a pixel shader. Use this method if the two
+ *  textures may have different formats.
+ */
+void RENDERER_API AddDrawTexturePass(
+	FRDGBuilder& GraphBuilder,
+	const FSceneView& View,
+	FRDGTextureRef InputTexture,
+	FRDGTextureRef OutputTexture,
+	FIntPoint InputPosition = FIntPoint::ZeroValue,
+	FIntPoint OutputPosition = FIntPoint::ZeroValue,
+	FIntPoint Size = FIntPoint::ZeroValue);
+
+/** Helper variant which takes a shared viewport instead of unique input / output positions. */
+FORCEINLINE void AddDrawTexturePass(
+	FRDGBuilder& GraphBuilder,
+	const FSceneView& View,
+	FRDGTextureRef InputTexture,
+	FRDGTextureRef OutputTexture,
+	FIntRect ViewportRect)
+{
+	AddDrawTexturePass(GraphBuilder, View, InputTexture, OutputTexture, ViewportRect.Min, ViewportRect.Min, ViewportRect.Size());
+}
+
+void RENDERER_API AddDrawTexturePass(
+	FRDGBuilder& GraphBuilder,
+	const FSceneView& View,
+	FScreenPassTexture Input,
+	FScreenPassRenderTarget Output);
+
+template <typename TFunction>
+FORCEINLINE void AddRenderTargetPass(
+	FRDGBuilder& GraphBuilder,
+	FRDGEventName&& PassName,
+	FScreenPassRenderTarget Output,
+	TFunction&& Function)
+{
+	FRenderTargetParameters* PassParameters = GraphBuilder.AllocParameters<FRenderTargetParameters>();
+	PassParameters->RenderTargets[0] = Output.GetRenderTargetBinding();
+	GraphBuilder.AddPass(MoveTemp(PassName), PassParameters, ERDGPassFlags::Raster, MoveTemp(Function));
+}
+
+template <typename TFunction>
+FORCEINLINE void AddDrawCanvasPass(
+	FRDGBuilder& GraphBuilder,
+	FRDGEventName&& PassName,
+	const FSceneView& View,
+	FScreenPassRenderTarget Output,
+	TFunction Function)
+{
+	check(Output.IsValid());
+
+	const FSceneViewFamily& ViewFamily = *View.Family;
+	FCanvas& Canvas = *FCanvas::Create(GraphBuilder, Output.Texture, nullptr, ViewFamily.Time, View.GetFeatureLevel(), ViewFamily.DebugDPIScale);
+	Canvas.SetRenderTargetRect(Output.ViewRect);
+
+	Function(Canvas);
+
+	const bool bFlush = false;
+	Canvas.Flush_RenderThread(GraphBuilder, bFlush);
+}
+
+enum class EDownsampleDepthFilter
+{
+	// Produces a depth value that is not conservative but has consistent error (i.e. picks the sample).
+	Point,
+
+	// Produces a conservative max depth value.
+	Max,
+
+	// Produces a checkerboarded selection of min and max depth values
+	Checkerboard,
+
+	// Procuce a color texture where R=min and G=max DeviceZ
+	MinAndMaxDepth
+};
+
+void AddDownsampleDepthPass(
+	FRDGBuilder& GraphBuilder,
+	const FViewInfo& View,
+	FScreenPassTexture Input,
+	FScreenPassRenderTarget Output,
+	EDownsampleDepthFilter DownsampleDepthFilter);
+
+#include "ScreenPass.inl" // IWYU pragma: export
diff --git a/Engine/Source/Runtime/Renderer/Public/ScreenPass.inl b/Engine/Source/Runtime/Renderer/Public/ScreenPass.inl
new file mode 100644
--- /dev/null
+++ b/Engine/Source/Runtime/Renderer/Public/ScreenPass.inl
@@ -0,0 +1,345 @@
+// Copyright Epic Games, Inc. All Rights Reserved.
+
+#pragma once
+
+inline FIntPoint GetDownscaledExtent(FIntPoint Extent, FIntPoint Divisor)
+{
+	Extent = FIntPoint::DivideAndRoundUp(Extent, Divisor);
+	Extent.X = FMath::Max(1, Extent.X);
+	Extent.Y = FMath::Max(1, Extent.Y);
+	return Extent;
+}
+
+inline FIntPoint GetScaledExtent(FIntPoint Extent, FVector2D Multiplier)
+{
+	Extent.X *= Multiplier.X;
+	Extent.Y *= Multiplier.Y;
+	Extent.X = FMath::Max(1, Extent.X);
+	Extent.Y = FMath::Max(1, Extent.Y);
+	return Extent;
+}
+
+inline FIntPoint GetScaledExtent(FIntPoint Extent, float Multiplier)
+{
+	return GetScaledExtent(Extent, FVector2D(Multiplier));
+}
+
+inline FIntRect GetDownscaledRect(FIntRect Rect, FIntPoint Divisor)
+{
+	Rect.Min /= Divisor;
+	Rect.Max = GetDownscaledExtent(Rect.Max, Divisor);
+	return Rect;
+}
+
+inline FIntRect GetScaledRect(FIntRect Rect, FVector2D Multiplier)
+{
+	Rect.Min.X *= Multiplier.X;
+	Rect.Min.Y *= Multiplier.Y;
+	Rect.Max = GetScaledExtent(Rect.Max, Multiplier);
+	return Rect;
+}
+
+FORCEINLINE FIntRect GetScaledRect(FIntRect Rect, float Multiplier)
+{
+	return GetScaledRect(Rect, FVector2D(Multiplier));
+}
+
+inline FScreenPassTextureViewport GetDownscaledViewport(FScreenPassTextureViewport Viewport, FIntPoint Divisor)
+{
+	Viewport.Rect = GetDownscaledRect(Viewport.Rect, Divisor);
+	Viewport.Extent = GetDownscaledExtent(Viewport.Extent, Divisor);
+	return Viewport;
+}
+
+inline FScreenPassTextureViewport GetScaledViewport(FScreenPassTextureViewport Viewport, FVector2D Multiplier)
+{
+	Viewport.Rect = GetScaledRect(Viewport.Rect, Multiplier);
+	Viewport.Extent = GetScaledExtent(Viewport.Extent, Multiplier);
+	return Viewport;
+}
+
+inline FIntRect GetRectFromExtent(FIntPoint Extent)
+{
+	return FIntRect(FIntPoint::ZeroValue, Extent);
+}
+
+inline FScreenPassTexture::FScreenPassTexture(FRDGTextureRef InTexture)
+	: Texture(InTexture)
+{
+	if (InTexture)
+	{
+		ViewRect.Max = InTexture->Desc.Extent;
+	}
+}
+
+inline bool FScreenPassTexture::IsValid() const
+{
+	return Texture != nullptr && !ViewRect.IsEmpty();
+}
+
+inline bool FScreenPassTexture::operator==(FScreenPassTexture Other) const
+{
+	return Texture == Other.Texture && ViewRect == Other.ViewRect;
+}
+
+inline bool FScreenPassTexture::operator!=(FScreenPassTexture Other) const
+{
+	return !(*this == Other);
+}
+
+inline bool FScreenPassRenderTarget::operator==(FScreenPassRenderTarget Other) const
+{
+	return FScreenPassTexture::operator==(Other) && LoadAction == Other.LoadAction;
+}
+
+inline bool FScreenPassRenderTarget::operator!=(FScreenPassRenderTarget Other) const
+{
+	return !(*this == Other);
+}
+
+inline FRenderTargetBinding FScreenPassRenderTarget::GetRenderTargetBinding() const
+{
+	return FRenderTargetBinding(Texture, LoadAction);
+}
+
+inline FScreenPassTextureViewport::FScreenPassTextureViewport(FScreenPassTexture InTexture)
+{
+	check(InTexture.IsValid());
+	Extent = InTexture.Texture->Desc.Extent;
+	Rect = InTexture.ViewRect;
+}
+
+inline bool FScreenPassTextureViewport::operator==(const FScreenPassTextureViewport& Other) const
+{
+	return Extent == Other.Extent && Rect == Other.Rect;
+}
+
+inline bool FScreenPassTextureViewport::operator!=(const FScreenPassTextureViewport& Other) const
+{
+	return Extent != Other.Extent || Rect != Other.Rect;
+}
+
+inline bool FScreenPassTextureViewport::IsEmpty() const
+{
+	return Extent == FIntPoint::ZeroValue || Rect.IsEmpty();
+}
+
+inline bool FScreenPassTextureViewport::IsFullscreen() const
+{
+	return Rect.Min == FIntPoint::ZeroValue && Rect.Max == Extent;
+}
+
+inline FVector2D FScreenPassTextureViewport::GetRectToExtentRatio() const
+{
+	return FVector2D((float)Rect.Width() / Extent.X, (float)Rect.Height() / Extent.Y);
+}
+
+inline FScreenTransform FScreenTransform::Invert(const FScreenTransform& AToB)
+{
+	ensure(!FMath::IsNearlyZero(AToB.Scale.X));
+	ensure(!FMath::IsNearlyZero(AToB.Scale.Y));
+
+	float InvX = 1.0f / AToB.Scale.X;
+	float InvY = 1.0f / AToB.Scale.Y;
+
+	return FScreenTransform(
+		FVector2f(InvX, InvY),
+		FVector2f(-InvX * AToB.Bias.X, -InvY * AToB.Bias.Y));
+}
+
+
+inline FVector2f operator * (const FVector2f& PInA, const FScreenTransform& AToB)
+{
+	return PInA * AToB.Scale + AToB.Bias;
+}
+
+inline FScreenTransform operator * (const FScreenTransform& AToB, const FVector2f& Scale)
+{
+	return FScreenTransform(AToB.Scale * Scale, AToB.Bias * Scale);
+}
+
+inline FScreenTransform operator * (const FScreenTransform& AToB, const float& Scale)
+{
+	return FScreenTransform(AToB.Scale * Scale, AToB.Bias * Scale);
+}
+
+inline FScreenTransform operator * (const FScreenTransform& AToB, const FIntPoint& Scale)
+{
+	return AToB * FVector2f(Scale.X, Scale.Y);
+}
+
+inline FScreenTransform operator * (const FScreenTransform& AToB, const FScreenTransform& BToC)
+{
+	return FScreenTransform(AToB.Scale * BToC.Scale, AToB.Bias * BToC.Scale + BToC.Bias);
+}
+
+
+inline FScreenTransform operator + (const FScreenTransform& AToB, const FVector2f& Bias)
+{
+	return FScreenTransform(AToB.Scale, AToB.Bias + Bias);
+}
+
+inline FScreenTransform operator + (const FScreenTransform& AToB, const float& Bias)
+{
+	return AToB + FVector2f(Bias, Bias);
+}
+
+inline FScreenTransform operator + (const FScreenTransform& AToB, const FIntPoint& Bias)
+{
+	return AToB + FVector2f(Bias.X, Bias.Y);
+}
+
+
+inline FScreenTransform operator - (const FScreenTransform& AToB, const FVector2f& Bias)
+{
+	return FScreenTransform(AToB.Scale, AToB.Bias - Bias);
+}
+
+inline FScreenTransform operator - (const FScreenTransform& AToB, const float& Bias)
+{
+	return AToB - FVector2f(Bias, Bias);
+}
+
+inline FScreenTransform operator - (const FScreenTransform& AToB, const FIntPoint& Bias)
+{
+	return AToB - FVector2f(Bias.X, Bias.Y);
+}
+
+
+inline FScreenTransform operator / (const FScreenTransform& AToB, const FVector2f& InvertedScale)
+{
+	ensure(!FMath::IsNearlyZero(InvertedScale.X));
+	ensure(!FMath::IsNearlyZero(InvertedScale.Y));
+	return AToB * FVector2f(1.0f / InvertedScale.X, 1.0f / InvertedScale.Y);
+}
+
+inline FScreenTransform operator / (const FScreenTransform& AToB, const float& InvertedScale)
+{
+	return AToB / FVector2f(InvertedScale, InvertedScale);
+}
+
+inline FScreenTransform operator / (const FScreenTransform& AToB, const FIntPoint& InvertedScale)
+{
+	return AToB / FVector2f(InvertedScale.X, InvertedScale.Y);
+}
+
+inline FScreenTransform operator / (const FScreenTransform& AToB, const FScreenTransform& CToB)
+{
+	FScreenTransform BToC = FScreenTransform::Invert(CToB);
+	return AToB * BToC;
+}
+
+
+inline FScreenTransform FScreenTransform::ChangeTextureBasisFromTo(
+	const FIntPoint& TextureExtent,
+	const FIntRect& TextureViewport,
+	FScreenTransform::ETextureBasis SrcBasis,
+	FScreenTransform::ETextureBasis DestBasis)
+{
+	if (int32(SrcBasis) < int32(DestBasis))
+	{
+		if (int32(SrcBasis) + 1 != int32(DestBasis))
+		{
+			ETextureBasis IntermediaryBasis = ETextureBasis((int32(SrcBasis) + int32(DestBasis)) / 2);
+
+			return (
+				FScreenTransform::ChangeTextureBasisFromTo(TextureExtent, TextureViewport, SrcBasis, IntermediaryBasis) *
+				FScreenTransform::ChangeTextureBasisFromTo(TextureExtent, TextureViewport, IntermediaryBasis, DestBasis));
+		}
+		else if (DestBasis == ETextureBasis::ViewportUV)
+		{
+			return FScreenTransform::ScreenPosToViewportUV;
+		}
+		else if (DestBasis == ETextureBasis::TexelPosition)
+		{
+			return FScreenTransform(
+				FVector2f(TextureViewport.Width(), TextureViewport.Height()),
+				FVector2f(TextureViewport.Min.X, TextureViewport.Min.Y));
+		}
+		else if (DestBasis == ETextureBasis::TextureUV)
+		{
+			return FScreenTransform(FVector2f(1.0f / float(TextureExtent.X), 1.0f / float(TextureExtent.Y)), FVector2f(0.0f, 0.0f));
+		}
+		else
+		{
+			check(0);
+		}
+	}
+	else if (int32(SrcBasis) > int32(DestBasis))
+	{
+		if (int32(SrcBasis) != int32(DestBasis) + 1)
+		{
+			ETextureBasis IntermediaryBasis = ETextureBasis((int32(SrcBasis) + int32(DestBasis)) / 2);
+
+			return (
+				FScreenTransform::ChangeTextureBasisFromTo(TextureExtent, TextureViewport, SrcBasis, IntermediaryBasis) *
+				FScreenTransform::ChangeTextureBasisFromTo(TextureExtent, TextureViewport, IntermediaryBasis, DestBasis));
+		}
+		else if (DestBasis == ETextureBasis::ScreenPosition)
+		{
+			return FScreenTransform::ViewportUVToScreenPos;
+		}
+		else if (DestBasis == ETextureBasis::ViewportUV)
+		{
+			float InvWidth = 1.0f / float(TextureViewport.Width());
+			float InvHeight = 1.0f / float(TextureViewport.Height());
+
+			return FScreenTransform(
+				FVector2f(InvWidth, InvHeight),
+				FVector2f(-InvWidth * TextureViewport.Min.X, -InvHeight * TextureViewport.Min.Y));
+		}
+		else if (DestBasis == ETextureBasis::TexelPosition)
+		{
+			return FScreenTransform(FVector2f(TextureExtent.X, TextureExtent.Y), FVector2f(0.0f, 0.0f));
+		}
+		else
+		{
+			check(0);
+		}
+	}
+
+	return FScreenTransform::Identity;
+}
+
+inline FScreenTransform FScreenTransform::ChangeRectFromTo(
+	FVector2f SourceOffset,
+	FVector2f SourceExtent,
+	FVector2f DestinationOffset,
+	FVector2f DestinationExtent)
+{
+	FScreenTransform Transform;
+	Transform.Scale = DestinationExtent / SourceExtent;
+	Transform.Bias = DestinationOffset - Transform.Scale * SourceOffset;
+	return Transform;
+}
+
+inline FScreenTransform FScreenTransform::ChangeRectFromTo(const FIntRect& SrcViewport, const FIntRect& DestViewport)
+{
+	return FScreenTransform::ChangeRectFromTo(
+		FVector2f(SrcViewport.Min), FVector2f(SrcViewport.Size()),
+		FVector2f(DestViewport.Min), FVector2f(DestViewport.Size()));
+}
+
+template<>
+struct TShaderParameterTypeInfo<FScreenTransform>
+{
+	static constexpr EUniformBufferBaseType BaseType = UBMT_FLOAT32;
+	static constexpr int32 NumRows = 1;
+	static constexpr int32 NumColumns = 4;
+	static constexpr int32 NumElements = 0;
+	static constexpr int32 Alignment = 16;
+	static constexpr bool bIsStoredInConstantBuffer = true;
+
+	using TAlignedType = TAlignedTypedef<FScreenTransform, Alignment>::Type;
+
+	static const FShaderParametersMetadata* GetStructMetadata() { return nullptr; }
+};
+
+inline FScreenPassTextureInput GetScreenPassTextureInput(FScreenPassTexture TexturePair, FRHISamplerState* Sampler)
+{
+	FScreenPassTextureInput Input;
+	Input.Viewport = GetScreenPassTextureViewportParameters(FScreenPassTextureViewport(TexturePair));
+	Input.Texture = TexturePair.Texture;
+	Input.Sampler = Sampler;
+	return Input;
+}
\ No newline at end of file
diff --git a/Engine/Source/Runtime/Renderer/Public/SystemTextures.h b/Engine/Source/Runtime/Renderer/Public/SystemTextures.h
new file mode 100644
--- /dev/null
+++ b/Engine/Source/Runtime/Renderer/Public/SystemTextures.h
@@ -0,0 +1,243 @@
+// Copyright Epic Games, Inc. All Rights Reserved.
+
+/*=============================================================================
+	SystemTextures.h: System textures definitions.
+=============================================================================*/
+
+#pragma once
+
+#include "RenderGraph.h"
+
+/** Contains system textures which can be registered for read-only access on an RDG pass. */
+struct FRDGSystemTextures
+{
+	/** Call to initialize for the requested builder instance. */
+	RENDERER_API static const FRDGSystemTextures& Create(FRDGBuilder& GraphBuilder);
+
+	/** Returns the instance for the builder. Must be called after Create. */
+	RENDERER_API static const FRDGSystemTextures& Get(FRDGBuilder& GraphBuilder);
+
+	/** Returns whether the system textures have been created. */
+	RENDERER_API static bool IsValid(FRDGBuilder& GraphBuilder);
+
+	FRDGTextureRef White{};
+	FRDGTextureRef Black{};
+	FRDGTextureRef BlackAlphaOne{};
+	FRDGTextureRef BlackArray{};
+	FRDGTextureRef MaxFP16Depth{};
+	FRDGTextureRef DepthDummy{};
+	FRDGTextureRef StencilDummy{};
+	FRDGTextureRef BlackDepthCube{};
+	FRDGTextureRef Green{};
+	FRDGTextureRef DefaultNormal8Bit{};
+	FRDGTextureRef MidGrey{};
+	FRDGTextureRef VolumetricBlack{};
+	FRDGTextureRef VolumetricBlackAlphaOne{};
+	FRDGTextureRef VolumetricBlackUint{};
+	FRDGTextureRef CubeBlack{};
+	FRDGTextureRef CubeArrayBlack{};
+
+	FRDGTextureSRVRef StencilDummySRV{};
+};
+
+
+struct FDefaultTextureKey
+{
+	uint32 ValueAsUInt[4] = { 0u,0u,0u,0u };
+	EPixelFormat Format = PF_Unknown;
+	ETextureDimension Dimension = ETextureDimension::Texture2D;
+};
+struct FDefaultTexture
+{
+	uint32 Hash = 0;
+	FDefaultTextureKey Key;
+	TRefCountPtr<IPooledRenderTarget> Texture;
+};
+
+struct FDefaultBufferKey
+{
+	uint32 ValueAsUInt[4] = { 0u,0u,0u,0u };
+	uint32 NumBytePerElement = 0;
+	bool bIsStructuredBuffer = false;
+};
+struct FDefaultBuffer
+{
+	uint32 Hash = 0;
+	FDefaultBufferKey Key;
+	TRefCountPtr<FRDGPooledBuffer> Buffer;
+};
+
+/**
+ * Encapsulates the system textures used for scene rendering.
+ */
+class FSystemTextures : public FRenderResource
+{
+public:
+	FSystemTextures()
+		: FRenderResource()
+		, FeatureLevelInitializedTo(ERHIFeatureLevel::Num)
+	{}
+
+	/**
+	 * Initialize/allocate textures if not already.
+	 */
+	void InitializeTextures(FRHICommandListImmediate& RHICmdList, const ERHIFeatureLevel::Type InFeatureLevel);
+
+	// FRenderResource interface.
+	/**
+	 * Release textures when device is lost/destroyed.
+	 */
+	virtual void ReleaseDynamicRHI();
+
+	// -----------
+
+	/**
+		Any Textures added here MUST be explicitly released on ReleaseDynamicRHI()!
+		Some RHIs need all their references released during destruction!
+	*/
+
+	// float4(1,1,1,1) can be used in case a light is not shadow casting
+	TRefCountPtr<IPooledRenderTarget> WhiteDummy;
+	// float4(0,0,0,0) can be used in additive postprocessing to avoid a shader combination
+	TRefCountPtr<IPooledRenderTarget> BlackDummy;
+	// float4(0,0,0,1)
+	TRefCountPtr<IPooledRenderTarget> BlackAlphaOneDummy;
+	// Dummy texture array with 1 black slice 
+	TRefCountPtr<IPooledRenderTarget> BlackArrayDummy;
+	// used by the material expression Noise
+	TRefCountPtr<IPooledRenderTarget> PerlinNoiseGradient;
+	// used by the material expression Noise (faster version, should replace old version), todo: move out of SceneRenderTargets
+	TRefCountPtr<IPooledRenderTarget> PerlinNoise3D;
+	// Sobol sampling texture, the first sample points for four sobol dimensions in RGBA
+	TRefCountPtr<IPooledRenderTarget> SobolSampling;
+	/** SSAO randomization */
+	TRefCountPtr<IPooledRenderTarget> SSAORandomization;
+	/** GTAO PreIntegrated */
+	TRefCountPtr<IPooledRenderTarget> GTAOPreIntegrated;
+
+	/** Preintegrated GF for single sample IBL */
+	TRefCountPtr<IPooledRenderTarget> PreintegratedGF;
+	/** Hair BSDF LUT texture */
+	TRefCountPtr<IPooledRenderTarget> HairLUT0;
+	TRefCountPtr<IPooledRenderTarget> HairLUT1;
+	TRefCountPtr<IPooledRenderTarget> HairLUT2;
+	/** Linearly Transformed Cosines LUTs */
+	TRefCountPtr<IPooledRenderTarget> LTCMat;
+	TRefCountPtr<IPooledRenderTarget> LTCAmp;
+	/** Texture that holds a single value containing the maximum depth that can be stored as FP16. */
+	TRefCountPtr<IPooledRenderTarget> MaxFP16Depth;
+	/** Depth texture that holds a single depth value */
+	TRefCountPtr<IPooledRenderTarget> DepthDummy;
+	/** Stencil texture that holds a single stencil value. */
+	TRefCountPtr<IPooledRenderTarget> StencilDummy;
+	TRefCountPtr<IPooledRenderTarget> BlackDepthCube;
+	// float4(0,1,0,1)
+	TRefCountPtr<IPooledRenderTarget> GreenDummy;
+	// float4(0.5,0.5,0.5,1)
+	TRefCountPtr<IPooledRenderTarget> DefaultNormal8Bit;
+	// float4(0.5,0.5,0.5,0.5)
+	TRefCountPtr<IPooledRenderTarget> MidGreyDummy;
+
+	/** float4(0,0,0,0) volumetric texture. */
+	TRefCountPtr<IPooledRenderTarget> VolumetricBlackDummy;
+	TRefCountPtr<IPooledRenderTarget> VolumetricBlackAlphaOneDummy;
+	TRefCountPtr<IPooledRenderTarget> VolumetricBlackUintDummy;
+
+	/** float4(0,0,0,0) cube textures. */
+	TRefCountPtr<IPooledRenderTarget> CubeBlackDummy;
+	TRefCountPtr<IPooledRenderTarget> CubeArrayBlackDummy;
+
+	// Dummy 0 Uint texture for RHIs that need explicit overloads
+	TRefCountPtr<IPooledRenderTarget> ZeroUIntDummy;
+	// Dummy 0 Uint texture for RHIs that need explicit overloads
+	TRefCountPtr<IPooledRenderTarget> ZeroUIntArrayDummy;
+
+	// SRV for WhiteDummy Texture.
+	TRefCountPtr<FRHIShaderResourceView> WhiteDummySRV;
+	// SRV for StencilDummy Texture.
+	TRefCountPtr<FRHIShaderResourceView> StencilDummySRV;
+
+	// ASCII Standard character set - IBM code page 437 (character 32-127)
+	TRefCountPtr<IPooledRenderTarget> AsciiTexture;
+
+	// Create simple default texture
+	FRDGTextureRef RENDERER_API GetWhiteDummy(FRDGBuilder& GraphBuilder) const;
+	FRDGTextureRef RENDERER_API GetBlackDummy(FRDGBuilder& GraphBuilder) const;
+	FRDGTextureRef RENDERER_API GetBlackArrayDummy(FRDGBuilder& GraphBuilder) const;
+	FRDGTextureRef RENDERER_API GetZeroUIntDummy(FRDGBuilder& GraphBuilder) const;
+	FRDGTextureRef RENDERER_API GetZeroUIntArrayDummy(FRDGBuilder& GraphBuilder) const;
+	FRDGTextureRef RENDERER_API GetZeroUShort4Dummy(FRDGBuilder& GraphBuilder) const;
+	FRDGTextureRef RENDERER_API GetBlackAlphaOneDummy(FRDGBuilder& GraphBuilder) const;
+	FRDGTextureRef RENDERER_API GetMaxFP16Depth(FRDGBuilder& GraphBuilder) const;
+	FRDGTextureRef RENDERER_API GetDepthDummy(FRDGBuilder& GraphBuilder) const;
+	FRDGTextureRef RENDERER_API GetStencilDummy(FRDGBuilder& GraphBuilder) const;
+	FRDGTextureRef RENDERER_API GetGreenDummy(FRDGBuilder& GraphBuilder) const;
+	FRDGTextureRef RENDERER_API GetDefaultNormal8Bit(FRDGBuilder& GraphBuilder) const;
+	FRDGTextureRef RENDERER_API GetMidGreyDummy(FRDGBuilder& GraphBuilder) const;
+	FRDGTextureRef RENDERER_API GetVolumetricBlackDummy(FRDGBuilder& GraphBuilder) const;
+	FRDGTextureRef RENDERER_API GetVolumetricBlackUintDummy(FRDGBuilder& GraphBuilder) const;
+	FRDGTextureRef RENDERER_API GetCubeBlackDummy(FRDGBuilder& GraphBuilder) const;
+	FRDGTextureRef RENDERER_API GetCubeArrayBlackDummy(FRDGBuilder& GraphBuilder) const;
+
+	// Create default 2D texture (1x1) with specific format and initialize value 
+	FRDGTextureRef RENDERER_API GetDefaultTexture2D(FRDGBuilder& GraphBuilder, EPixelFormat Format, float Value);
+	FRDGTextureRef RENDERER_API GetDefaultTexture2D(FRDGBuilder& GraphBuilder, EPixelFormat Format, uint32 Value);
+	FRDGTextureRef RENDERER_API GetDefaultTexture2D(FRDGBuilder& GraphBuilder, EPixelFormat Format, const FVector3f& Value);
+	FRDGTextureRef RENDERER_API GetDefaultTexture2D(FRDGBuilder& GraphBuilder, EPixelFormat Format, const FVector4f& Value);
+	FRDGTextureRef RENDERER_API GetDefaultTexture2D(FRDGBuilder& GraphBuilder, EPixelFormat Format, const FUintVector4& Value);
+	FRDGTextureRef RENDERER_API GetDefaultTexture2D(FRDGBuilder& GraphBuilder, EPixelFormat Format, const FClearValueBinding& Value);
+
+	// Create default 2D/3D/Cube/Array texture (1x1) with specific format and initialize value 
+	FRDGTextureRef RENDERER_API GetDefaultTexture(FRDGBuilder& GraphBuilder, ETextureDimension Dimension, EPixelFormat Format, float Value);
+	FRDGTextureRef RENDERER_API GetDefaultTexture(FRDGBuilder& GraphBuilder, ETextureDimension Dimension, EPixelFormat Format, uint32 Value);
+	FRDGTextureRef RENDERER_API GetDefaultTexture(FRDGBuilder& GraphBuilder, ETextureDimension Dimension, EPixelFormat Format, const FVector2D& Value);
+	FRDGTextureRef RENDERER_API GetDefaultTexture(FRDGBuilder& GraphBuilder, ETextureDimension Dimension, EPixelFormat Format, const FIntPoint& Value);
+	FRDGTextureRef RENDERER_API GetDefaultTexture(FRDGBuilder& GraphBuilder, ETextureDimension Dimension, EPixelFormat Format, const FVector3f& Value);
+	FRDGTextureRef RENDERER_API GetDefaultTexture(FRDGBuilder& GraphBuilder, ETextureDimension Dimension, EPixelFormat Format, const FVector4f& Value);
+	FRDGTextureRef RENDERER_API GetDefaultTexture(FRDGBuilder& GraphBuilder, ETextureDimension Dimension, EPixelFormat Format, const FUintVector4& Value);
+	FRDGTextureRef RENDERER_API GetDefaultTexture(FRDGBuilder& GraphBuilder, ETextureDimension Dimension, EPixelFormat Format, const FClearValueBinding& Value);
+
+	// Create default buffer initialize to zero.
+	FRDGBufferRef RENDERER_API GetDefaultBuffer(FRDGBuilder& GraphBuilder, uint32 NumBytePerElement);
+	FRDGBufferRef RENDERER_API GetDefaultStructuredBuffer(FRDGBuilder& GraphBuilder, uint32 NumBytePerElement);
+
+	template <typename T>
+	FRDGBufferRef GetDefaultBuffer(FRDGBuilder& GraphBuilder)
+	{
+		return GetDefaultBuffer(GraphBuilder, sizeof(T));
+	}
+
+	template <typename T>
+	FRDGBufferRef GetDefaultStructuredBuffer(FRDGBuilder& GraphBuilder)
+	{
+		return GetDefaultStructuredBuffer(GraphBuilder, sizeof(T));
+	}
+
+	// Create a default buffer initialized with a reference element.
+	FRDGBufferRef RENDERER_API GetDefaultBuffer(FRDGBuilder& GraphBuilder, uint32 NumBytePerElement, float Value);
+	FRDGBufferRef RENDERER_API GetDefaultBuffer(FRDGBuilder& GraphBuilder, uint32 NumBytePerElement, uint32 Value);
+	FRDGBufferRef RENDERER_API GetDefaultBuffer(FRDGBuilder& GraphBuilder, uint32 NumBytePerElement, const FVector3f& Value);
+	FRDGBufferRef RENDERER_API GetDefaultBuffer(FRDGBuilder& GraphBuilder, uint32 NumBytePerElement, const FVector4f& Value);
+	FRDGBufferRef RENDERER_API GetDefaultBuffer(FRDGBuilder& GraphBuilder, uint32 NumBytePerElement, const FUintVector4& Value);
+	FRDGBufferRef RENDERER_API GetDefaultStructuredBuffer(FRDGBuilder& GraphBuilder, uint32 NumBytePerElement, float Value);
+	FRDGBufferRef RENDERER_API GetDefaultStructuredBuffer(FRDGBuilder& GraphBuilder, uint32 NumBytePerElement, uint32 Value);
+	FRDGBufferRef RENDERER_API GetDefaultStructuredBuffer(FRDGBuilder& GraphBuilder, uint32 NumBytePerElement, const FVector3f& Value);
+	FRDGBufferRef RENDERER_API GetDefaultStructuredBuffer(FRDGBuilder& GraphBuilder, uint32 NumBytePerElement, const FVector4f& Value);
+	FRDGBufferRef RENDERER_API GetDefaultStructuredBuffer(FRDGBuilder& GraphBuilder, uint32 NumBytePerElement, const FUintVector4& Value);
+
+protected:
+	/** Maximum feature level that the textures have been initialized up to */
+	ERHIFeatureLevel::Type FeatureLevelInitializedTo;
+
+	/** Default textures allocated on-demand */
+	TArray<FDefaultTexture> DefaultTextures;
+	TArray<FDefaultBuffer> DefaultBuffers;
+	FHashTable HashDefaultTextures;
+	FHashTable HashDefaultBuffers;
+
+	void InitializeCommonTextures(FRHICommandListImmediate& RHICmdList);
+	void InitializeFeatureLevelDependentTextures(FRHICommandListImmediate& RHICmdList, const ERHIFeatureLevel::Type InFeatureLevel);
+};
+
+/** The global system textures used for scene rendering. */
+RENDERER_API extern TGlobalResource<FSystemTextures> GSystemTextures;
