Fixing nonunity errors.

diff --git a/Engine/Plugins/Media/MediaIOFramework/Source/MediaIOCore/Public/MediaIOCoreAudioOutput.h b/Engine/Plugins/Media/MediaIOFramework/Source/MediaIOCore/Public/MediaIOCoreAudioOutput.h
--- a/Engine/Plugins/Media/MediaIOFramework/Source/MediaIOCore/Public/MediaIOCoreAudioOutput.h
+++ b/Engine/Plugins/Media/MediaIOFramework/Source/MediaIOCore/Public/MediaIOCoreAudioOutput.h
@@ -2,6 +2,7 @@
 
 #pragma once
 
+#include "AudioDeviceHandle.h"
 #include "DSP/BufferVectorOperations.h"
 #include "DSP/MultithreadedPatching.h"
 #include "ISubmixBufferListener.h"
diff --git a/Engine/Source/Runtime/Renderer/Private/VariableRateShading/ContrastAdaptiveImageGenerator.cpp b/Engine/Source/Runtime/Renderer/Private/VariableRateShading/ContrastAdaptiveImageGenerator.cpp
--- a/Engine/Source/Runtime/Renderer/Private/VariableRateShading/ContrastAdaptiveImageGenerator.cpp
+++ b/Engine/Source/Runtime/Renderer/Private/VariableRateShading/ContrastAdaptiveImageGenerator.cpp
@@ -9,6 +9,8 @@
 #include "RHIDefinitions.h"
 #include "RenderTargetPool.h"
 #include "SystemTextures.h"
+#include "SceneRendering.h"
+#include "SceneTextures.h"
 #include "SceneView.h"
 #include "IEyeTracker.h"
 #include "IHeadMountedDisplay.h"
diff --git a/Engine/Source/Runtime/Renderer/Private/VariableRateShading/VariableRateShadingImageManager.cpp b/Engine/Source/Runtime/Renderer/Private/VariableRateShading/VariableRateShadingImageManager.cpp
--- a/Engine/Source/Runtime/Renderer/Private/VariableRateShading/VariableRateShadingImageManager.cpp
+++ b/Engine/Source/Runtime/Renderer/Private/VariableRateShading/VariableRateShadingImageManager.cpp
@@ -10,6 +10,7 @@
 #include "RHIDefinitions.h"
 #include "DataDrivenShaderPlatformInfo.h"
 #include "RenderTargetPool.h"
+#include "SceneRendering.h"
 #include "SystemTextures.h"
 #include "SceneView.h"
 #include "IEyeTracker.h"
