Fix for strata interacting with responssive AA.

diff --git a/Engine/Source/Runtime/Renderer/Private/DeferredShadingRenderer.cpp b/Engine/Source/Runtime/Renderer/Private/DeferredShadingRenderer.cpp
--- a/Engine/Source/Runtime/Renderer/Private/DeferredShadingRenderer.cpp
+++ b/Engine/Source/Runtime/Renderer/Private/DeferredShadingRenderer.cpp
@@ -3810,6 +3810,12 @@ void FDeferredShadingSceneRenderer::Render(FRDGBuilder& GraphBuilder)
 			CompositeRayTracingSkyLight(GraphBuilder, SceneTextures, SkyLightTexture, SkyLightHitDistanceTexture);
 		}
 	#endif
+
+		if (Strata::IsStrataEnabled())
+		{
+			// Now remove all the Strata tile stencil tags used by deferred tiled light passes. Make later marks such as responssive AA works.
+			AddClearStencilPass(GraphBuilder, SceneTextures.Depth.Target);
+		}
 	}
 	else if (HairStrands::HasViewHairStrandsData(Views) && ViewFamily.EngineShowFlags.Lighting)
 	{
