Fix hair strands rendering with stereo rendering.

* Fix hair tile rendering not taking into account View.MinRect
* Fix hair voxelization having incorrect traversing for the second view. This was due to voxelization being expressed relative to the view origin (for precision purpopse), and the offset correct was incorrect for the secondary view.

diff --git a/Engine/Shaders/Private/HairStrands/HairStrandsDeepShadowMask.usf b/Engine/Shaders/Private/HairStrands/HairStrandsDeepShadowMask.usf
--- a/Engine/Shaders/Private/HairStrands/HairStrandsDeepShadowMask.usf
+++ b/Engine/Shaders/Private/HairStrands/HairStrandsDeepShadowMask.usf
@@ -449,6 +449,10 @@ void MainPS(
 	const bool bIsValidOpaquePixel = HairPixelCoverage < CoverageThreshold;
 	
 	float3 TranslatedWorldPosition = GetTranslatedWorldPosition(UV, SceneDepth);
+	if (View.StereoPassIndex == 1)
+	{
+		TranslatedWorldPosition += VirtualVoxel.TranslatedWorldOffsetStereoCorrection;
+	}
 
 	// PCF or simple point sampler
 	float Visibility = 1;
diff --git a/Engine/Shaders/Private/HairStrands/HairStrandsDeepTransmittanceMask.usf b/Engine/Shaders/Private/HairStrands/HairStrandsDeepTransmittanceMask.usf
--- a/Engine/Shaders/Private/HairStrands/HairStrandsDeepTransmittanceMask.usf
+++ b/Engine/Shaders/Private/HairStrands/HairStrandsDeepTransmittanceMask.usf
@@ -366,8 +366,12 @@ void MainCS(uint2 DispatchThreadId : SV_DispatchThreadID)
 	const float SceneDepth		= ConvertFromDeviceZ(Sample.Depth);
 	const uint2 PixelCoord		= HairStrands.HairSampleCoords[SampleIndex];
 	const float2 UV				= (PixelCoord.xy + float2(0.5f, 0.5f)) / float2(View.BufferSizeAndInvSize.xy);
-	const float3 TranslatedWorldPosition = GetTranslatedWorldPosition(UV, SceneDepth);
 	const float PixelRadius		= SceneDepth * VirtualVoxel.HairCoveragePixelRadiusAtDepth1;
+	float3 TranslatedWorldPosition = GetTranslatedWorldPosition(UV, SceneDepth);
+	if (View.StereoPassIndex == 1)
+	{
+		TranslatedWorldPosition += VirtualVoxel.TranslatedWorldOffsetStereoCorrection;
+	}
 
 	const uint4 PackedShadowMask = ~ShadowMaskBitsTexture[PixelCoord];
 	const float3 Random = GetHairVoxelJitter(PixelCoord.xy, View.StateFrameIndexMod8, VirtualVoxel.JitterMode);
@@ -464,7 +468,11 @@ void MainCS(uint2 DispatchThreadId : SV_DispatchThreadID)
 		}
 
 		const float SceneDepth = ConvertFromDeviceZ(Sample.Depth);
-		const float3 TranslatedWorldPosition = GetTranslatedWorldPosition(UV, SceneDepth);
+		float3 TranslatedWorldPosition = GetTranslatedWorldPosition(UV, SceneDepth);
+		if (View.StereoPassIndex == 1)
+		{
+			TranslatedWorldPosition += VirtualVoxel.TranslatedWorldOffsetStereoCorrection;
+		}
 
 		const uint MacroGroupIndex = clamp(Sample.MacroGroupId, 0, MAX_HAIR_MACROGROUP_COUNT - 1);
 
diff --git a/Engine/Shaders/Private/HairStrands/HairStrandsVisibilityCompaction.usf b/Engine/Shaders/Private/HairStrands/HairStrandsVisibilityCompaction.usf
--- a/Engine/Shaders/Private/HairStrands/HairStrandsVisibilityCompaction.usf
+++ b/Engine/Shaders/Private/HairStrands/HairStrandsVisibilityCompaction.usf
@@ -373,8 +373,12 @@ void MainCS(
 	const uint2 TileCoord = TileDataBuffer[LinearIndex];
 	uint2 PixelCoord = TileCoord * TileSize + GroupThreadId.xy;
 
-	if (PixelCoord.x >= uint(OutputResolution.x) || PixelCoord.y >= uint(OutputResolution.y))
+	const bool bIsValid = all(PixelCoord < View.ViewRectMinAndSize.zw);
+	if (!bIsValid)
+	{
 		PixelCoord = uint2(0, 0);
+	}
+	PixelCoord += View.ViewRectMinAndSize.xy;
 
 	FSampleSetDesc SampleDesc;
 
diff --git a/Engine/Shaders/Private/HairStrands/HairStrandsVisibilityTile.usf b/Engine/Shaders/Private/HairStrands/HairStrandsVisibilityTile.usf
--- a/Engine/Shaders/Private/HairStrands/HairStrandsVisibilityTile.usf
+++ b/Engine/Shaders/Private/HairStrands/HairStrandsVisibilityTile.usf
@@ -31,11 +31,12 @@ groupshared uint2 s_pMask[GROUP_THREAD_COUNT];
 [numthreads(HAIR_TILE_SIZE, HAIR_TILE_SIZE, 1)]
 void TileMainCS(uint2 DispatchThreadId : SV_DispatchThreadID, uint LinearIndex : SV_GroupIndex, uint3 GroupId : SV_GroupID)
 {
-	const uint2 PixelCoord = DispatchThreadId;
+	const uint2 PixelCoord = DispatchThreadId + View.ViewRectMin.xy;
+	const bool bIsValid    = all(DispatchThreadId < uint2(View.ViewRectMinAndSize.zw));
 
 	bool fHair = false;
 	bool pHair = false;
-	if (all(PixelCoord < uint2(BufferResolution)))
+	if (bIsValid)
 	{
 		if (bForceOutputAllTiles)
 		{
@@ -337,7 +338,7 @@ void MainVS(
 	uint2 TileVertex = TileCoord * HAIR_TILE_SIZE;
 	TileVertex.x += InVertexId == 1 || InVertexId == 2 || InVertexId == 4 ? HAIR_TILE_SIZE : 0;
 	TileVertex.y += InVertexId == 2 || InVertexId == 4 || InVertexId == 5 ? HAIR_TILE_SIZE : 0;
-	Out.UV = float2(TileVertex - ViewMin) * ViewInvSize;
+	Out.UV = float2(TileVertex) * ViewInvSize; // No need to take ViewMin into account, as the viewport rect is center on the actual view (i.e., xy:ViewMin)
 	Out.Position = float4(Out.UV * float2(2.0f, -2.0f) + float2(-1.0, 1.0f), 0.5f, 1.0f);
 }
 
diff --git a/Engine/Source/Runtime/Renderer/Private/HairStrands/HairStrandsData.h b/Engine/Source/Runtime/Renderer/Private/HairStrands/HairStrandsData.h
--- a/Engine/Source/Runtime/Renderer/Private/HairStrands/HairStrandsData.h
+++ b/Engine/Source/Runtime/Renderer/Private/HairStrands/HairStrandsData.h
@@ -215,6 +215,7 @@ BEGIN_SHADER_PARAMETER_STRUCT(FHairStrandsVoxelCommonParameters, )
 	SHADER_PARAMETER(float, Raytracing_SkyOcclusionThreshold)
 
 	SHADER_PARAMETER(FVector3f, TranslatedWorldOffset) // For debug purpose
+	SHADER_PARAMETER(FVector3f, TranslatedWorldOffsetStereoCorrection) // PreViewTranslation correction between View0 & View1 when rendering stereo
 
 	SHADER_PARAMETER(uint32, AllocationFeedbackEnable)
 
diff --git a/Engine/Source/Runtime/Renderer/Private/HairStrands/HairStrandsRendering.cpp b/Engine/Source/Runtime/Renderer/Private/HairStrands/HairStrandsRendering.cpp
--- a/Engine/Source/Runtime/Renderer/Private/HairStrands/HairStrandsRendering.cpp
+++ b/Engine/Source/Runtime/Renderer/Private/HairStrands/HairStrandsRendering.cpp
@@ -97,8 +97,10 @@ void RenderHairPrePass(
 			continue;
 
 		// For stereo rendering, hair groups/voxelization/deep-shadow are only produced once
+		FVector PreViewStereoCorrection = FVector::ZeroVector;
 		if (IStereoRendering::IsStereoEyeView(View))
 		{
+			PreViewStereoCorrection = Views[0].ViewMatrices.GetPreViewTranslation() - Views[1].ViewMatrices.GetPreViewTranslation();
 			if (IStereoRendering::IsASecondaryView(View))
 			{
 				// No need to copy the view state (i.e., HairStrandsViewStateData) as it is only used for 
@@ -122,7 +124,7 @@ void RenderHairPrePass(
 		GraphBuilder.AddDispatchHint();
 
 		// Voxelization and Deep Opacity Maps
-		VoxelizeHairStrands(GraphBuilder, Scene, View, InstanceCullingManager);
+		VoxelizeHairStrands(GraphBuilder, Scene, View, InstanceCullingManager, PreViewStereoCorrection);
 		if (View.HairStrandsViewData.MacroGroupDatas.Num() > 0)
 		{
 			AddMeshDrawTransitionPass(GraphBuilder, View, View.HairStrandsViewData.MacroGroupDatas);
diff --git a/Engine/Source/Runtime/Renderer/Private/HairStrands/HairStrandsTile.cpp b/Engine/Source/Runtime/Renderer/Private/HairStrands/HairStrandsTile.cpp
--- a/Engine/Source/Runtime/Renderer/Private/HairStrands/HairStrandsTile.cpp
+++ b/Engine/Source/Runtime/Renderer/Private/HairStrands/HairStrandsTile.cpp
@@ -183,6 +183,7 @@ static FHairStrandsTiles AddHairStrandsGenerateTilesPass_Internal(
 
 	TShaderMapRef<FHairStrandsTileGenerationPassCS> ComputeShader(View.ShaderMap);
 	FHairStrandsTileGenerationPassCS::FParameters* PassParameters = GraphBuilder.AllocParameters<FHairStrandsTileGenerationPassCS::FParameters>();
+	PassParameters->ViewUniformBuffer		= View.ViewUniformBuffer;
 	PassParameters->BufferResolution		= View.ViewRect.Size();
 	PassParameters->bForceOutputAllTiles	= bHasValidInput ? 0 : 1;
 	PassParameters->bUintTexture			= bUintTexture ? 1u : 0u;
diff --git a/Engine/Source/Runtime/Renderer/Private/HairStrands/HairStrandsVoxelization.cpp b/Engine/Source/Runtime/Renderer/Private/HairStrands/HairStrandsVoxelization.cpp
--- a/Engine/Source/Runtime/Renderer/Private/HairStrands/HairStrandsVoxelization.cpp
+++ b/Engine/Source/Runtime/Renderer/Private/HairStrands/HairStrandsVoxelization.cpp
@@ -841,6 +841,7 @@ static float RoundHairVoxeliSize(float In)
 static FHairStrandsVoxelResources AllocateVirtualVoxelResources(
 	FRDGBuilder& GraphBuilder,
 	const FViewInfo& View,
+	const FVector& PreViewStereoCorrection,
 	FHairStrandsMacroGroupDatas& MacroGroupDatas,
 	FHairStrandsMacroGroupResources& MacroGroupResources,
 	FRDGBufferRef& PageToPageIndexBuffer,
@@ -890,6 +891,7 @@ static FHairStrandsVoxelResources AllocateVirtualVoxelResources(
 	Out.Parameters.Common.Raytracing_SkyOcclusionThreshold		= FMath::Max(0.f, GHairVirtualVoxelRaytracing_SkyOcclusionThreshold);
 	Out.Parameters.Common.HairCoveragePixelRadiusAtDepth1		= ComputeMinStrandRadiusAtDepth1(FIntPoint(View.ViewRect.Width(), View.ViewRect.Height()), View.FOV, 1/*SampleCount*/, 1/*RasterizationScale*/).Primary;
 	Out.Parameters.Common.TranslatedWorldOffset					= FVector3f(View.ViewMatrices.GetPreViewTranslation());
+	Out.Parameters.Common.TranslatedWorldOffsetStereoCorrection = FVector3f(PreViewStereoCorrection);
 
 	// For debug purpose
 	Out.Parameters.Common.AllocationFeedbackEnable = bVoxelAllocationFeedbackEnable ? 1u : 0u;
@@ -1404,7 +1406,8 @@ void VoxelizeHairStrands(
 	FRDGBuilder& GraphBuilder, 
 	const FScene* Scene, 
 	FViewInfo& View,
-	FInstanceCullingManager& InstanceCullingManager)
+	FInstanceCullingManager& InstanceCullingManager,
+	const FVector& PreViewStereoCorrection)
 {	
 	FHairStrandsMacroGroupDatas& MacroGroupDatas = View.HairStrandsViewData.MacroGroupDatas;
 	FHairStrandsVoxelResources& VirtualVoxelResources = View.HairStrandsViewData.VirtualVoxelResources;
@@ -1441,7 +1444,7 @@ void VoxelizeHairStrands(
 	{
 		FRDGBufferRef PageToPageIndexBuffer = nullptr;
 		FHairStrandsViewStateData* HairStrandsViewStateData = View.ViewState ? &View.ViewState->HairStrandsViewStateData : nullptr;
-		VirtualVoxelResources = AllocateVirtualVoxelResources(GraphBuilder, View, MacroGroupDatas, MacroGroupResources, PageToPageIndexBuffer, HairStrandsViewStateData);
+		VirtualVoxelResources = AllocateVirtualVoxelResources(GraphBuilder, View, PreViewStereoCorrection, MacroGroupDatas, MacroGroupResources, PageToPageIndexBuffer, HairStrandsViewStateData);
 
 		FRDGBufferRef ClearIndArgsBuffer = IndirectVoxelPageClear(GraphBuilder, View, VirtualVoxelResources);
 
diff --git a/Engine/Source/Runtime/Renderer/Private/HairStrands/HairStrandsVoxelization.h b/Engine/Source/Runtime/Renderer/Private/HairStrands/HairStrandsVoxelization.h
--- a/Engine/Source/Runtime/Renderer/Private/HairStrands/HairStrandsVoxelization.h
+++ b/Engine/Source/Runtime/Renderer/Private/HairStrands/HairStrandsVoxelization.h
@@ -20,4 +20,5 @@ void VoxelizeHairStrands(
 	FRDGBuilder& GraphBuilder,
 	const class FScene* Scene,
 	FViewInfo& View,
-	FInstanceCullingManager& InstanceCullingManager);
+	FInstanceCullingManager& InstanceCullingManager, 
+	const FVector& PreViewStereoCorrection);
