Dump the full material name in the debug output when we fail to find a material parameter collection uniform buffer.

- Factor out common logging code into a function

diff --git a/Engine/Source/Runtime/Renderer/Private/ShaderBaseClasses.cpp b/Engine/Source/Runtime/Renderer/Private/ShaderBaseClasses.cpp
--- a/Engine/Source/Runtime/Renderer/Private/ShaderBaseClasses.cpp
+++ b/Engine/Source/Runtime/Renderer/Private/ShaderBaseClasses.cpp
@@ -232,6 +232,28 @@ void FMaterialShader::SetViewParameters(FRHIBatchedShaderParameters& BatchedPara
 	}
 }
 
+static void DumpRegisteredMaterialParameterCollections(const FMaterial& Material, const FGuid& ParameterCollectionGuid)
+{
+	// Dump the currently registered parameter collections and the ID we failed to find.
+	// In a cooked project these numbers are persistent so we can track back to the original
+	// parameter collection that was being referenced and no longer exists
+	FString InstancesString;
+	TMultiMap<FGuid, FMaterialParameterCollectionInstanceResource*>::TIterator Iter = GDefaultMaterialParameterCollectionInstances.CreateIterator();
+	while (Iter)
+	{
+		FMaterialParameterCollectionInstanceResource* Instance = Iter.Value();
+		InstancesString += FString::Printf(TEXT("\n0x%p: %s: %s"),
+			Instance, Instance ? *Instance->GetOwnerName().ToString() : TEXT("None"), *Iter.Key().ToString());
+		++Iter;
+	}
+
+	UE_LOG(LogRenderer, Fatal, TEXT("Failed to find parameter collection buffer with GUID '%s' for %s.\n")
+		TEXT("Currently %i listed default instances: %s"),
+		*ParameterCollectionGuid.ToString(),
+		*Material.GetFullPath(),
+		GDefaultMaterialParameterCollectionInstances.Num(), *InstancesString);
+}
+
 void FMaterialShader::SetParameters(
 	FRHIBatchedShaderParameters& BatchedParameters,
 	const FMaterialRenderProxy* MaterialRenderProxy,
@@ -314,23 +336,7 @@ void FMaterialShader::SetParameters(
 
 			if (!UniformBuffer)
 			{
-				// Dump the currently registered parameter collections and the ID we failed to find.
-				// In a cooked project these numbers are persistent so we can track back to the original
-				// parameter collection that was being referenced and no longer exists
-				FString InstancesString;
-				TMultiMap<FGuid, FMaterialParameterCollectionInstanceResource*>::TIterator Iter = GDefaultMaterialParameterCollectionInstances.CreateIterator();
-				while (Iter)
-				{
-					FMaterialParameterCollectionInstanceResource* Instance = Iter.Value();
-					InstancesString += FString::Printf(TEXT("\n0x%p: %s: %s"),
-						Instance, Instance ? *Instance->GetOwnerName().ToString() : TEXT("None"), *Iter.Key().ToString());
-					++Iter;
-				}
-
-				UE_LOG(LogRenderer, Fatal, TEXT("Failed to find parameter collection buffer with GUID '%s'.\n")
-					TEXT("Currently %i listed default instances: %s"),
-					*ParameterCollections[CollectionIndex].ToString(),
-					GDefaultMaterialParameterCollectionInstances.Num(), *InstancesString);
+				DumpRegisteredMaterialParameterCollections(Material, ParameterCollections[CollectionIndex]);
 			}
 
 			SetUniformBufferParameter(BatchedParameters, ParameterCollectionUniformBuffers[CollectionIndex], UniformBuffer);
@@ -391,23 +397,7 @@ void FMaterialShader::GetShaderBindings(
 
 			if (!UniformBuffer)
 			{
-				// Dump the currently registered parameter collections and the ID we failed to find.
-				// In a cooked project these numbers are persistent so we can track back to the original
-				// parameter collection that was being referenced and no longer exists
-				FString InstancesString;
-				TMultiMap<FGuid, FMaterialParameterCollectionInstanceResource*>::TIterator Iter = GDefaultMaterialParameterCollectionInstances.CreateIterator();
-				while (Iter)
-				{
-					FMaterialParameterCollectionInstanceResource* Instance = Iter.Value();
-					InstancesString += FString::Printf(TEXT("\n0x%p: %s: %s"),
-						Instance, Instance ? *Instance->GetOwnerName().ToString() : TEXT("None"), *Iter.Key().ToString());
-					++Iter;
-				}
-
-				UE_LOG(LogRenderer, Fatal, TEXT("Failed to find parameter collection buffer with GUID '%s'.\n")
-					TEXT("Currently %i listed default instances: %s"),
-					*ParameterCollections[CollectionIndex].ToString(),
-					GDefaultMaterialParameterCollectionInstances.Num(), *InstancesString);
+				DumpRegisteredMaterialParameterCollections(Material, ParameterCollections[CollectionIndex]);
 			}
 
 			ShaderBindings.Add(ParameterCollectionUniformBuffers[CollectionIndex], UniformBuffer);		
