Make sure we flush occlusion queries out of memory if we've turned them off.

diff --git a/Engine/Source/Runtime/Renderer/Private/SceneOcclusion.cpp b/Engine/Source/Runtime/Renderer/Private/SceneOcclusion.cpp
--- a/Engine/Source/Runtime/Renderer/Private/SceneOcclusion.cpp
+++ b/Engine/Source/Runtime/Renderer/Private/SceneOcclusion.cpp
@@ -1055,6 +1055,18 @@ void FHZBOcclusionTester::Submit(FRDGBuilder& GraphBuilder, const FViewInfo& Vie
 	AddEnqueueCopyPass(GraphBuilder, ResultsReadback.Get(), ResultsTextureGPU);
 }
 
+static void TrimAllOcclusionHistory(TArrayView<FViewInfo> Views)
+{
+	for (const FViewInfo& View : Views)
+	{
+		if (View.Family && View.ViewState)
+		{
+			const double RealTimeSeconds = View.Family->Time.GetRealTimeSeconds();
+			View.ViewState->TrimOcclusionHistory(RealTimeSeconds, RealTimeSeconds - GEngine->PrimitiveProbablyVisibleTime, RealTimeSeconds, View.ViewState->OcclusionFrameCounter);
+		}
+	}
+}
+
 static FViewOcclusionQueriesPerView AllocateOcclusionTests(const FScene* Scene, TArrayView<const FVisibleLightInfo> VisibleLightInfos, TArrayView<FViewInfo> Views)
 {
 	SCOPED_NAMED_EVENT(FSceneRenderer_AllocateOcclusionTestsOcclusionTests, FColor::Emerald);
@@ -1448,6 +1460,11 @@ void FDeferredShadingSceneRenderer::RenderOcclusion(
 			});
 		}
 	}
+	else
+	{
+		// Make sure the views are freeing their memory if we toggled occlusion off at any point
+		TrimAllOcclusionHistory(Views);
+	}
 
 	const bool bUseHzbOcclusion = RenderHzb(GraphBuilder, SceneTextures.Depth.Resolve, BuildHZBAsyncComputeParams);
 
