SceneRendererInterface dependency cleanup
- Removing Renderer includes (SceneRendererInterface.h) from Engine header files
- Removing public accessors for FSceneUniformBuffer which is private to Renderer
- Replacing public uses of FSceneUniformBuffer with a method on ISceneRenderer
- Adding a quick-access utility function to get the Scene uniform buffer from a SceneView
- Removing includes of a private Renderer header (SceneUniformBuffer.h) from other modules
- Restoring removal of Renderer/Private from Niagara include paths

diff --git a/Engine/Plugins/FX/Niagara/Source/Niagara/Niagara.Build.cs b/Engine/Plugins/FX/Niagara/Source/Niagara/Niagara.Build.cs
--- a/Engine/Plugins/FX/Niagara/Source/Niagara/Niagara.Build.cs
+++ b/Engine/Plugins/FX/Niagara/Source/Niagara/Niagara.Build.cs
@@ -68,7 +68,6 @@ public class Niagara : ModuleRules
         PrivateIncludePaths.AddRange(
             new string[] {
 				System.IO.Path.Combine(GetModuleDirectory("Engine"), "Private"),
-				System.IO.Path.Combine(GetModuleDirectory("Renderer"), "Private"),
 			});
 
         if (Target.bBuildEditor == true)
diff --git a/Engine/Plugins/FX/Niagara/Source/Niagara/Private/NiagaraRendererMeshes.cpp b/Engine/Plugins/FX/Niagara/Source/Niagara/Private/NiagaraRendererMeshes.cpp
--- a/Engine/Plugins/FX/Niagara/Source/Niagara/Private/NiagaraRendererMeshes.cpp
+++ b/Engine/Plugins/FX/Niagara/Source/Niagara/Private/NiagaraRendererMeshes.cpp
@@ -28,7 +28,7 @@
 #include "RenderGraphBuilder.h"
 #include "RenderGraphUtils.h"
 #include "SceneInterface.h"
-#include "SceneUniformBuffer.h"
+#include "SceneRendererInterface.h"
 
 DECLARE_CYCLE_STAT(TEXT("Generate Mesh Vertex Data [GT]"), STAT_NiagaraGenMeshVertexData, STATGROUP_Niagara);
 
@@ -977,14 +977,13 @@ void FNiagaraRendererMeshes::SetupElementForGPUScene(
 
 		// Force it to preserve instance order if we are sorted so that GPU Scene's instance culling doesn't scramble them
 		OutMeshBatchElement.bPreserveInstanceOrder = ParticleMeshRenderData.bNeedsSort;
-		FSceneUniformBuffer& SceneUniforms = View.GetSceneUniforms();
 
 		GPUSceneRes.DynamicPrimitiveData.DataWriterGPU = FGPUSceneWriteDelegate::CreateLambda(				
 			[&GPUSceneRes](FRDGBuilder& GraphBuilder, const FGPUSceneWriteDelegateParams& Params)
 			{
 				GPUSceneRes.GPUWriteParams.GPUSceneWriterParameters	= Params.GPUWriteParams;
 				GPUSceneRes.GPUWriteParams.View						= Params.View->ViewUniformBuffer; // NOTE: Set here, not outside lambda
-				GPUSceneRes.GPUWriteParams.Scene					= Params.View->GetSceneUniforms().GetBuffer(GraphBuilder);
+				GPUSceneRes.GPUWriteParams.Scene					= GetSceneUniformBufferRef(GraphBuilder, *Params.View);
 				GPUSceneRes.GPUWriteParams.PrimitiveId 				= Params.PrimitiveId;
 		
 				FNiagaraGPUSceneUtils::AddUpdateMeshParticleInstancesPass(
diff --git a/Engine/Plugins/Runtime/nDisplay/Source/DisplayClusterLightCardEditorShaders/Private/DisplayClusterMeshProjectionRenderer.cpp b/Engine/Plugins/Runtime/nDisplay/Source/DisplayClusterLightCardEditorShaders/Private/DisplayClusterMeshProjectionRenderer.cpp
--- a/Engine/Plugins/Runtime/nDisplay/Source/DisplayClusterLightCardEditorShaders/Private/DisplayClusterMeshProjectionRenderer.cpp
+++ b/Engine/Plugins/Runtime/nDisplay/Source/DisplayClusterLightCardEditorShaders/Private/DisplayClusterMeshProjectionRenderer.cpp
@@ -11,7 +11,7 @@
 #include "CanvasTypes.h"
 #include "EngineModule.h"
 #include "SceneManagement.h"
-#include "SceneUniformBuffer.h"
+#include "SceneRendererInterface.h"
 #include "SceneViewExtension.h"
 #include "ScreenPass.h"
 #include "Components/PrimitiveComponent.h"
@@ -278,7 +278,7 @@ FMeshProjectionPassParameters* CreateProjectionPassParameters(FRDGBuilder& Graph
 {
 	FMeshProjectionPassParameters* PassParameters = GraphBuilder.AllocParameters<FMeshProjectionPassParameters>();
 	PassParameters->View = View->ViewUniformBuffer;
-	PassParameters->Scene = View->GetSceneUniforms().GetBuffer(GraphBuilder);
+	PassParameters->Scene = GetSceneUniformBufferRef(GraphBuilder, *View);
 	PassParameters->InstanceCulling = FInstanceCullingContext::CreateDummyInstanceCullingUniformBuffer(GraphBuilder);
 
 	return PassParameters;
diff --git a/Engine/Plugins/Runtime/nDisplay/Source/DisplayClusterShaders/Private/Shaders/DisplayClusterShadersPreprocess_UVLightCards.cpp b/Engine/Plugins/Runtime/nDisplay/Source/DisplayClusterShaders/Private/Shaders/DisplayClusterShadersPreprocess_UVLightCards.cpp
--- a/Engine/Plugins/Runtime/nDisplay/Source/DisplayClusterShaders/Private/Shaders/DisplayClusterShadersPreprocess_UVLightCards.cpp
+++ b/Engine/Plugins/Runtime/nDisplay/Source/DisplayClusterShaders/Private/Shaders/DisplayClusterShadersPreprocess_UVLightCards.cpp
@@ -14,8 +14,8 @@
 #include "RenderGraphBuilder.h"
 #include "RenderGraphResources.h"
 #include "ScenePrivate.h"
+#include "SceneRendererInterface.h"
 #include "SceneRenderTargetParameters.h"
-#include "SceneUniformBuffer.h"
 #include "MeshPassProcessor.inl"
 #include "UnrealClient.h"
 
@@ -278,7 +278,7 @@ bool FDisplayClusterShadersPreprocess_UVLightCards::RenderPreprocess_UVLightCard
 
 	FUVLightCardPassParameters* PassParameters = GraphBuilder.AllocParameters<FUVLightCardPassParameters>();
 	PassParameters->View = View->ViewUniformBuffer;
-	PassParameters->Scene = View->GetSceneUniforms().GetBuffer(GraphBuilder);
+	PassParameters->Scene = GetSceneUniformBufferRef(GraphBuilder, *View);
 	PassParameters->InstanceCulling = FInstanceCullingContext::CreateDummyInstanceCullingUniformBuffer(GraphBuilder);
 	PassParameters->RenderTargets[0] = OutputRenderTargetBinding;
 
diff --git a/Engine/Source/Runtime/Engine/Private/SceneView.cpp b/Engine/Source/Runtime/Engine/Private/SceneView.cpp
--- a/Engine/Source/Runtime/Engine/Private/SceneView.cpp
+++ b/Engine/Source/Runtime/Engine/Private/SceneView.cpp
@@ -2815,12 +2815,6 @@ TArray<const FSceneView*> FSceneView::GetSecondaryViews() const
 	return Views;
 }
 
-FSceneUniformBuffer& FSceneView::GetSceneUniforms() const
-{
-	check(Family);
-	return Family->GetSceneUniforms();
-}
-
 FSceneViewFamily::ConstructionValues::ConstructionValues(
 	const FRenderTarget* InRenderTarget,
 	FSceneInterface* InScene,
diff --git a/Engine/Source/Runtime/Engine/Public/SceneView.h b/Engine/Source/Runtime/Engine/Public/SceneView.h
--- a/Engine/Source/Runtime/Engine/Public/SceneView.h
+++ b/Engine/Source/Runtime/Engine/Public/SceneView.h
@@ -16,7 +16,6 @@
 #include "RenderResource.h"
 #include "ShowFlags.h"
 #include "StereoRendering.h"
-#include "SceneRendererInterface.h"
 
 #if UE_ENABLE_INCLUDE_ORDER_DEPRECATED_IN_5_2
 #include "CoreMinimal.h"
@@ -40,6 +39,7 @@ class FViewElementDrawer;
 class ISceneViewExtension;
 class FSceneViewFamily;
 class FVolumetricFogViewResources;
+class ISceneRenderer;
 class ISpatialUpscaler;
 class ITemporalUpscaler;
 
@@ -1661,8 +1661,6 @@ public:
 		return InstancedViewUniformBuffer;
 	}
 
-	FSceneUniformBuffer& GetSceneUniforms() const;
-
 protected:
 	FSceneViewStateInterface* EyeAdaptationViewState = nullptr;
 
@@ -2143,16 +2141,8 @@ public:
 	inline bool GetIsInFocus() const				{ return bIsInFocus; }
 	inline void SetIsInFocus(bool bInIsInFocus)		{ bIsInFocus = bInIsInFocus; }
 
-	FSceneUniformBuffer& GetSceneUniforms() const
-	{
-		check(SceneRenderer);
-		return SceneRenderer->GetSceneUniforms();
-	}
-
-	void SetSceneRenderer(ISceneRenderer* NewSceneRenderer)
-	{
-		SceneRenderer = NewSceneRenderer;
-	}
+	void SetSceneRenderer(ISceneRenderer* NewSceneRenderer) { SceneRenderer = NewSceneRenderer; }
+	ISceneRenderer* GetSceneRenderer() const { check(SceneRenderer); return SceneRenderer; }
 
 private:
 	/** The scene renderer that is rendering this view family. This is only initialized in the rendering thread's copies of the FSceneViewFamily. */
diff --git a/Engine/Source/Runtime/Landscape/Private/LandscapeGrass.cpp b/Engine/Source/Runtime/Landscape/Private/LandscapeGrass.cpp
--- a/Engine/Source/Runtime/Landscape/Private/LandscapeGrass.cpp
+++ b/Engine/Source/Runtime/Landscape/Private/LandscapeGrass.cpp
@@ -68,7 +68,7 @@
 #include "MaterialCachedData.h"
 #include "SceneRenderTargetParameters.h"
 #include "TextureResource.h"
-#include "SceneUniformBuffer.h"
+#include "SceneRendererInterface.h"
 
 #define LOCTEXT_NAMESPACE "Landscape"
 
@@ -630,7 +630,7 @@ public:
 
 		auto* PassParameters = GraphBuilder.AllocParameters<FLandscapeGrassPassParameters>();
 		PassParameters->View = View->ViewUniformBuffer;
-		PassParameters->Scene = View->GetSceneUniforms().GetBuffer(GraphBuilder);
+		PassParameters->Scene = GetSceneUniformBufferRef(GraphBuilder, *View);
 		PassParameters->RenderTargets[0] = FRenderTargetBinding(OutputTexture, ERenderTargetLoadAction::EClear);
 
 		AddSimpleMeshPass(GraphBuilder, PassParameters, SceneInterface->GetRenderScene(), *View, nullptr, RDG_EVENT_NAME("LandscapeGrass"), View->UnscaledViewRect,
diff --git a/Engine/Source/Runtime/Landscape/Private/LandscapePhysicalMaterial.cpp b/Engine/Source/Runtime/Landscape/Private/LandscapePhysicalMaterial.cpp
--- a/Engine/Source/Runtime/Landscape/Private/LandscapePhysicalMaterial.cpp
+++ b/Engine/Source/Runtime/Landscape/Private/LandscapePhysicalMaterial.cpp
@@ -21,7 +21,7 @@
 #include "RHIGPUReadback.h"
 #include "RenderGraphUtils.h"
 #include "RenderCaptureInterface.h"
-#include "SceneUniformBuffer.h"
+#include "SceneRendererInterface.h"
 
 DECLARE_GPU_STAT_NAMED(LandscapePhysicalMaterial_Draw, TEXT("LandscapePhysicalMaterial"));
 
@@ -242,7 +242,7 @@ namespace
 
 		auto* PassParameters = GraphBuilder.AllocParameters<FLandscapePhysicalMaterialPassParameters>();
 		PassParameters->View = View->ViewUniformBuffer;
-		PassParameters->Scene = View->GetSceneUniforms().GetBuffer(GraphBuilder);
+		PassParameters->Scene = GetSceneUniformBufferRef(GraphBuilder, *View);
 		PassParameters->RenderTargets[0] = FRenderTargetBinding(OutputTexture, ERenderTargetLoadAction::EClear);
 
 		AddSimpleMeshPass(GraphBuilder, PassParameters, SceneInterface->GetRenderScene(), *View, nullptr, RDG_EVENT_NAME("LandscapePhysicalMaterial"), View->UnscaledViewRect,
diff --git a/Engine/Source/Runtime/Renderer/Private/InstanceCulling/InstanceCullingManager.h b/Engine/Source/Runtime/Renderer/Private/InstanceCulling/InstanceCullingManager.h
--- a/Engine/Source/Runtime/Renderer/Private/InstanceCulling/InstanceCullingManager.h
+++ b/Engine/Source/Runtime/Renderer/Private/InstanceCulling/InstanceCullingManager.h
@@ -7,6 +7,7 @@
 #include "RenderGraphResources.h"
 
 class FGPUScene;
+class FSceneUniformBuffer;
 
 namespace Nanite
 {
diff --git a/Engine/Source/Runtime/Renderer/Private/SceneRendering.cpp b/Engine/Source/Runtime/Renderer/Private/SceneRendering.cpp
--- a/Engine/Source/Runtime/Renderer/Private/SceneRendering.cpp
+++ b/Engine/Source/Runtime/Renderer/Private/SceneRendering.cpp
@@ -2390,7 +2390,6 @@ FViewFamilyInfo::~FViewFamilyInfo()
 FSceneRenderer::FSceneRenderer(const FSceneViewFamily* InViewFamily, FHitProxyConsumer* HitProxyConsumer)
 :	Scene(CheckPointer(InViewFamily->Scene)->GetRenderScene())
 ,	ViewFamily(*CheckPointer(InViewFamily))
-,	SceneUniforms()
 ,	MeshCollector(InViewFamily->GetFeatureLevel(), Allocator)
 ,	RayTracingCollector(InViewFamily->GetFeatureLevel(), Allocator)
 ,	VirtualShadowMapArray(*CheckPointer(Scene))
@@ -4810,19 +4809,6 @@ void FRendererModule::RenderPostResolvedSceneColorExtension(FRDGBuilder& GraphBu
 	}
 }
 
-class FScenePrimitiveRenderer : public ISceneRenderer
-{
-public:
-	FScenePrimitiveRenderer()
-		: SceneUniforms()
-	{}
-
-	const FSceneUniformBuffer& GetSceneUniforms() const final override { return SceneUniforms; }
-	FSceneUniformBuffer& GetSceneUniforms() final override { return SceneUniforms; }
-private:
-	FSceneUniformBuffer SceneUniforms;
-};
-
 class FScenePrimitiveRenderingContext : public IScenePrimitiveRenderingContext
 {
 public:
@@ -4848,7 +4834,7 @@ public:
 		ViewFamily.SetSceneRenderer(nullptr);
 	}
 
-	FScenePrimitiveRenderer Renderer;
+	FSceneRendererBase Renderer;
 	FGPUScene& GPUScene;
 	FGPUSceneDynamicContext GPUSceneDynamicContext;
 	FSceneViewFamily& ViewFamily;
diff --git a/Engine/Source/Runtime/Renderer/Private/SceneRendering.h b/Engine/Source/Runtime/Renderer/Private/SceneRendering.h
--- a/Engine/Source/Runtime/Renderer/Private/SceneRendering.h
+++ b/Engine/Source/Runtime/Renderer/Private/SceneRendering.h
@@ -1831,6 +1831,8 @@ public:
 	const FSceneTextures& GetSceneTextures() const;
 	const FSceneTextures* GetSceneTexturesChecked() const;
 
+	FSceneUniformBuffer& GetSceneUniforms() const;
+
 	RENDERER_API FRDGTextureRef GetVolumetricCloudTexture(FRDGBuilder& GraphBuilder) const;
 
 	/**
@@ -2028,12 +2030,28 @@ struct FComputeLightGridOutput
 	FRDGPassRef CompactLinksPass = {};
 };
 
+class FSceneRendererBase : public ISceneRenderer
+{
+public:
+	// ISceneRenderer interface
+	const FSceneUniformBuffer& GetSceneUniforms() const final override { return SceneUniforms; }
+	FSceneUniformBuffer& GetSceneUniforms() final override { return SceneUniforms; }
+
+	TRDGUniformBufferRef<FSceneUniformParameters> GetSceneUniformBufferRef(FRDGBuilder& GraphBuilder) final
+	{
+		return GetSceneUniforms().GetBuffer(GraphBuilder);
+	}
+
+private:
+	FSceneUniformBuffer SceneUniforms;
+};
+
 /**
  * Used as the scope for scene rendering functions.
  * It is initialized in the game thread by FSceneViewFamily::BeginRender, and then passed to the rendering thread.
  * The rendering thread calls Render(), and deletes the scene renderer when it returns.
  */
-class FSceneRenderer : public ISceneRenderer
+class FSceneRenderer : public FSceneRendererBase
 {
 public:
 	/** Linear bulk allocator with a lifetime tied to the scene renderer. */
@@ -2051,12 +2069,6 @@ public:
 	/** Views across all view families (may contain additional views if rendering multiple families together). */
 	TArray<const FSceneView*> AllFamilyViews;
 
-private:
-	FSceneUniformBuffer SceneUniforms;
-public:
-	const FSceneUniformBuffer& GetSceneUniforms() const final override { return SceneUniforms; }
-	FSceneUniformBuffer& GetSceneUniforms() final override { return SceneUniforms; }
-
 	/** All the dynamic scaling informations */
 	DynamicRenderScaling::TMap<float> DynamicResolutionFractions;
 	DynamicRenderScaling::TMap<float> DynamicResolutionUpperBounds;
@@ -2831,6 +2843,11 @@ inline const FSceneTextures* FViewInfo::GetSceneTexturesChecked() const
 	return ((FViewFamilyInfo*)Family)->GetSceneTexturesChecked();
 }
 
+inline FSceneUniformBuffer& FViewInfo::GetSceneUniforms() const
+{
+	return Family->GetSceneRenderer()->GetSceneUniforms();
+}
+
 /**
   * Returns a family from an array of views, with the assumption that all point to the same view family, which will be
   * true for the "Views" array in the scene renderer.  There are some utility functions that receive the Views array,
diff --git a/Engine/Source/Runtime/Renderer/Private/VirtualShadowMaps/VirtualShadowMapArray.h b/Engine/Source/Runtime/Renderer/Private/VirtualShadowMaps/VirtualShadowMapArray.h
--- a/Engine/Source/Runtime/Renderer/Private/VirtualShadowMaps/VirtualShadowMapArray.h
+++ b/Engine/Source/Runtime/Renderer/Private/VirtualShadowMaps/VirtualShadowMapArray.h
@@ -19,6 +19,7 @@ struct FScreenPassTexture;
 struct FSingleLayerWaterPrePassResult;
 class FNaniteVisibilityResults;
 class FSceneRenderer;
+class FSceneUniformBuffer;
 struct FShaderCompilerEnvironment;
 
 namespace Nanite
diff --git a/Engine/Source/Runtime/Renderer/Public/SceneRendererInterface.h b/Engine/Source/Runtime/Renderer/Public/SceneRendererInterface.h
--- a/Engine/Source/Runtime/Renderer/Public/SceneRendererInterface.h
+++ b/Engine/Source/Runtime/Renderer/Public/SceneRendererInterface.h
@@ -2,8 +2,13 @@
 
 #pragma once
 
+#include "SceneView.h"
+#include "ShaderParameterStructDeclaration.h"
+
 class FSceneUniformBuffer;
 
+DECLARE_UNIFORM_BUFFER_STRUCT(FSceneUniformParameters, RENDERER_API)
+
 /**
  * Public interface into FSceneRenderer. Used as the scope for scene rendering functions.
  */
@@ -12,4 +17,12 @@ class ISceneRenderer
 public:
 	virtual const FSceneUniformBuffer& GetSceneUniforms() const = 0;
 	virtual FSceneUniformBuffer& GetSceneUniforms() = 0;
+
+	virtual TRDGUniformBufferRef<FSceneUniformParameters> GetSceneUniformBufferRef(FRDGBuilder& GraphBuilder) = 0;
 };
+
+inline TRDGUniformBufferRef<FSceneUniformParameters> GetSceneUniformBufferRef(FRDGBuilder& GraphBuilder, const FSceneView& View)
+{
+	ISceneRenderer* SceneRender = View.Family ? View.Family->GetSceneRenderer() : nullptr;
+	return SceneRender ? SceneRender->GetSceneUniformBufferRef(GraphBuilder) : nullptr;
+}
