Remove unused internal OIT function.

diff --git a/Engine/Source/Runtime/Renderer/Private/OIT/OIT.cpp b/Engine/Source/Runtime/Renderer/Private/OIT/OIT.cpp
--- a/Engine/Source/Runtime/Renderer/Private/OIT/OIT.cpp
+++ b/Engine/Source/Runtime/Renderer/Private/OIT/OIT.cpp
@@ -296,11 +296,6 @@ public:
 	FUnorderedAccessViewRHIRef SortedIndexUAV = nullptr;
 };
 
-static bool IsOITSupported(EShaderPlatform InShaderPlatform)
-{
-	return !IsMobilePlatform(InShaderPlatform) && !FDataDrivenShaderPlatformInfo::GetIsHlslcc(InShaderPlatform);
-}
-
 static void RemoveAllocation(FSortedIndexBuffer* InBuffer)
 {
 	InBuffer->ReleaseResource();
