Changing shadow and light rendering passes/types to use the new BatchedShaderParameters types.

diff --git a/Engine/Source/Runtime/Renderer/Private/LightFunctionRendering.cpp b/Engine/Source/Runtime/Renderer/Private/LightFunctionRendering.cpp
--- a/Engine/Source/Runtime/Renderer/Private/LightFunctionRendering.cpp
+++ b/Engine/Source/Runtime/Renderer/Private/LightFunctionRendering.cpp
@@ -54,9 +54,9 @@ public:
 		StencilingGeometryParameters.Bind(Initializer.ParameterMap);
 	}
 
-	void SetParameters(FRHICommandList& RHICmdList, const FSceneView& View, const FLightSceneInfo* LightSceneInfo )
+	void SetParameters(FRHIBatchedShaderParameters& BatchedParameters, const FSceneView& View, const FLightSceneInfo* LightSceneInfo )
 	{
-		FMaterialShader::SetViewParameters(RHICmdList, RHICmdList.GetBoundVertexShader(), View, View.ViewUniformBuffer);
+		FMaterialShader::SetViewParameters(BatchedParameters, View, View.ViewUniformBuffer);
 		
 		// Light functions are projected using a bounding sphere.
 		// Calculate transform for bounding stencil sphere.
@@ -68,7 +68,7 @@ public:
 
 		FVector4f StencilingSpherePosAndScale;
 		StencilingGeometry::GStencilSphereVertexBuffer.CalcTransform(StencilingSpherePosAndScale, LightBounds, View.ViewMatrices.GetPreViewTranslation());
-		StencilingGeometryParameters.Set(RHICmdList, this, (FVector4f)StencilingSpherePosAndScale); // LWC_TODO: precision loss
+		StencilingGeometryParameters.Set(BatchedParameters, (FVector4f)StencilingSpherePosAndScale); // LWC_TODO: precision loss
 	}
 
 private:
@@ -111,12 +111,10 @@ public:
 		HairOnlyDepthTexture.Bind(Initializer.ParameterMap, TEXT("HairOnlyDepthTexture"));
 	}
 
-	void SetParameters(FRHICommandList& RHICmdList, const FViewInfo& View, const FLightSceneInfo* LightSceneInfo, const FMaterialRenderProxy* MaterialProxy, const FMaterial& Material, bool bRenderingPreviewShadowIndicator, float ShadowFadeFraction, bool bUseHairStrands, FRHITexture* InHairOnlyDepthTexture)
+	void SetParameters(FRHIBatchedShaderParameters& BatchedParameters, const FViewInfo& View, const FLightSceneInfo* LightSceneInfo, const FMaterialRenderProxy* MaterialProxy, const FMaterial& Material, bool bRenderingPreviewShadowIndicator, float ShadowFadeFraction, bool bUseHairStrands, FRHITexture* InHairOnlyDepthTexture)
 	{
-		FRHIPixelShader* ShaderRHI = RHICmdList.GetBoundPixelShader();
-
-		FMaterialShader::SetViewParameters(RHICmdList, ShaderRHI, View, View.ViewUniformBuffer);
-		FMaterialShader::SetParameters(RHICmdList, ShaderRHI, MaterialProxy, Material, View);
+		FMaterialShader::SetViewParameters(BatchedParameters, View, View.ViewUniformBuffer);
+		FMaterialShader::SetParameters(BatchedParameters, MaterialProxy, Material, View);
 
 		// Set the transform from screen space to light space.
 		if ( SvPositionToLight.IsBound() )
@@ -145,12 +143,12 @@ public:
 					FPlane(Ax, Ay,   0,  1)
 				) * View.ViewMatrices.GetInvViewProjectionMatrix() * WorldToLight;
 
-			SetShaderValue(RHICmdList, ShaderRHI, SvPositionToLight, FMatrix44f(SvPositionToLightValue) );
+			SetShaderValue(BatchedParameters, SvPositionToLight, FMatrix44f(SvPositionToLightValue) );
 		}
 
-		LightFunctionParameters.Set(RHICmdList, ShaderRHI, LightSceneInfo, ShadowFadeFraction);
+		LightFunctionParameters.Set(BatchedParameters, LightSceneInfo, ShadowFadeFraction);
 
-		SetShaderValue(RHICmdList, ShaderRHI, LightFunctionParameters2, FVector4f(
+		SetShaderValue(BatchedParameters, LightFunctionParameters2, FVector4f(
 			LightSceneInfo->Proxy->GetLightFunctionFadeDistance(), 
 			LightSceneInfo->Proxy->GetLightFunctionDisabledBrightness(),
 			bRenderingPreviewShadowIndicator ? 1.0f : 0.0f,
@@ -158,13 +156,13 @@ public:
 
 		if (HairOnlyDepthTexture.IsBound() && InHairOnlyDepthTexture)
 		{
-			SetTextureParameter(RHICmdList, ShaderRHI, HairOnlyDepthTexture, InHairOnlyDepthTexture);
+			SetTextureParameter(BatchedParameters, HairOnlyDepthTexture, InHairOnlyDepthTexture);
 		}
 
 		auto DeferredLightParameter = GetUniformBufferParameter<FDeferredLightUniformStruct>();
 		if (DeferredLightParameter.IsBound())
 		{
-			SetDeferredLightParameters(RHICmdList, ShaderRHI, DeferredLightParameter, LightSceneInfo, View);
+			SetDeferredLightParameters(BatchedParameters, DeferredLightParameter, LightSceneInfo, View);
 		}
 	}
 
@@ -435,8 +433,9 @@ bool FDeferredShadingSceneRenderer::RenderLightFunctionForMaterial(
 
 					// Render a bounding light sphere.
 					SetGraphicsPipelineState(RHICmdList, GraphicsPSOInit, 0);
-					VertexShader->SetParameters(RHICmdList, View, LightSceneInfo);
-					PixelShader->SetParameters(RHICmdList, View, LightSceneInfo, MaterialProxyForRendering, *MaterialForRendering, bRenderingPreviewShadowsIndicator, FadeAlpha, bUseHairStrands, PassParameters->HairOnlyDepthTexture->GetRHI());
+
+					SetShaderParametersLegacyVS(RHICmdList, VertexShader, View, LightSceneInfo);
+					SetShaderParametersLegacyPS(RHICmdList, PixelShader, View, LightSceneInfo, MaterialProxyForRendering, *MaterialForRendering, bRenderingPreviewShadowsIndicator, FadeAlpha, bUseHairStrands, PassParameters->HairOnlyDepthTexture->GetRHI());
 
 					// Project the light function using a sphere around the light
 					//@todo - could use a cone for spotlights
diff --git a/Engine/Source/Runtime/Renderer/Private/LightRendering.cpp b/Engine/Source/Runtime/Renderer/Private/LightRendering.cpp
--- a/Engine/Source/Runtime/Renderer/Private/LightRendering.cpp
+++ b/Engine/Source/Runtime/Renderer/Private/LightRendering.cpp
@@ -320,30 +320,26 @@ void FStencilingGeometryShaderParameters::Bind(const FShaderParameterMap& Parame
 	StencilConeTransform.Bind(ParameterMap, TEXT("StencilingConeTransform"));
 }
 
-void FStencilingGeometryShaderParameters::Set(FRHICommandList& RHICmdList, FShader* Shader, const FVector4f& InStencilingGeometryPosAndScale) const
+void FStencilingGeometryShaderParameters::Set(FRHIBatchedShaderParameters& BatchedParameters, const FVector4f& InStencilingGeometryPosAndScale) const
 {
 	const FParameters P = GetParameters(InStencilingGeometryPosAndScale);
-	SetShaderValue(RHICmdList, RHICmdList.GetBoundVertexShader(), StencilGeometryPosAndScale, P.StencilingGeometryPosAndScale);
-	SetShaderValue(RHICmdList, RHICmdList.GetBoundVertexShader(), StencilConeParameters, P.StencilingConeParameters);
+	SetShaderValue(BatchedParameters, StencilGeometryPosAndScale, P.StencilingGeometryPosAndScale);
+	SetShaderValue(BatchedParameters, StencilConeParameters, P.StencilingConeParameters);
 }
 
-void FStencilingGeometryShaderParameters::Set(FRHICommandList& RHICmdList, FShader* Shader, const FSceneView& View, const FLightSceneInfo* LightSceneInfo) const
+void FStencilingGeometryShaderParameters::Set(FRHIBatchedShaderParameters& BatchedParameters, const FSceneView& View, const FLightSceneInfo* LightSceneInfo) const
 {
 	const FParameters P = GetParameters(View, LightSceneInfo);
 	if (LightSceneInfo->Proxy->GetLightType() == LightType_Point ||
 		LightSceneInfo->Proxy->GetLightType() == LightType_Rect)
 	{
-		SetShaderValue(RHICmdList, RHICmdList.GetBoundVertexShader(), StencilGeometryPosAndScale, P.StencilingGeometryPosAndScale);
-		SetShaderValue(RHICmdList, RHICmdList.GetBoundVertexShader(), StencilConeParameters, P.StencilingConeParameters);
+		SetShaderValue(BatchedParameters, StencilGeometryPosAndScale, P.StencilingGeometryPosAndScale);
+		SetShaderValue(BatchedParameters, StencilConeParameters, P.StencilingConeParameters);
 	}
 	else if (LightSceneInfo->Proxy->GetLightType() == LightType_Spot)
 	{
-		SetShaderValue(RHICmdList, RHICmdList.GetBoundVertexShader(), StencilConeTransform, P.StencilingConeTransform);
-		SetShaderValue(
-			RHICmdList,
-			RHICmdList.GetBoundVertexShader(),
-			StencilConeParameters,
-			P.StencilingConeParameters);
+		SetShaderValue(BatchedParameters, StencilConeTransform, P.StencilingConeTransform);
+		SetShaderValue(BatchedParameters, StencilConeParameters, P.StencilingConeParameters);
 	}
 }
 
diff --git a/Engine/Source/Runtime/Renderer/Private/LightRendering.h b/Engine/Source/Runtime/Renderer/Private/LightRendering.h
--- a/Engine/Source/Runtime/Renderer/Private/LightRendering.h
+++ b/Engine/Source/Runtime/Renderer/Private/LightRendering.h
@@ -37,15 +37,13 @@ extern float GetLightFadeFactor(const FSceneView& View, const FLightSceneProxy*
 
 extern FDeferredLightUniformStruct GetDeferredLightParameters(const FSceneView& View, const FLightSceneInfo& LightSceneInfo);
 
-template<typename ShaderRHIParamRef>
-void SetDeferredLightParameters(
-	FRHICommandList& RHICmdList, 
-	const ShaderRHIParamRef ShaderRHI, 
+inline void SetDeferredLightParameters(
+	FRHIBatchedShaderParameters& BatchedParameters,
 	const TShaderUniformBufferParameter<FDeferredLightUniformStruct>& DeferredLightUniformBufferParameter, 
 	const FLightSceneInfo* LightSceneInfo,
 	const FSceneView& View)
 {
-	SetUniformBufferParameterImmediate(RHICmdList, ShaderRHI, DeferredLightUniformBufferParameter, GetDeferredLightParameters(View, *LightSceneInfo));
+	SetUniformBufferParameterImmediate(BatchedParameters, DeferredLightUniformBufferParameter, GetDeferredLightParameters(View, *LightSceneInfo));
 }
 
 extern FDeferredLightUniformStruct GetSimpleDeferredLightParameters(
@@ -53,17 +51,15 @@ extern FDeferredLightUniformStruct GetSimpleDeferredLightParameters(
 	const FSimpleLightEntry& SimpleLight,
 	const FSimpleLightPerViewEntry &SimpleLightPerViewData);
 
-template<typename ShaderRHIParamRef>
-void SetSimpleDeferredLightParameters(
-	FRHICommandList& RHICmdList, 
-	const ShaderRHIParamRef ShaderRHI, 
+inline void SetSimpleDeferredLightParameters(
+	FRHIBatchedShaderParameters& BatchedParameters,
 	const TShaderUniformBufferParameter<FDeferredLightUniformStruct>& DeferredLightUniformBufferParameter, 
 	const FSimpleLightEntry& SimpleLight,
 	const FSimpleLightPerViewEntry &SimpleLightPerViewData,
 	const FSceneView& View)
 {
 	FDeferredLightUniformStruct DeferredLightUniformsValue = GetSimpleDeferredLightParameters(View, SimpleLight, SimpleLightPerViewData);
-	SetUniformBufferParameterImmediate(RHICmdList, ShaderRHI, DeferredLightUniformBufferParameter, DeferredLightUniformsValue);
+	SetUniformBufferParameterImmediate(BatchedParameters, DeferredLightUniformBufferParameter, DeferredLightUniformsValue);
 }
 
 /** Shader parameters needed to render a light function. */
@@ -75,14 +71,9 @@ public:
 
 	static FVector4f GetLightFunctionSharedParameters(const FLightSceneInfo* LightSceneInfo, float ShadowFadeFraction);
 
-	template<typename ShaderRHIParamRef>
-	void Set(FRHICommandList& RHICmdList, const ShaderRHIParamRef ShaderRHI, const FLightSceneInfo* LightSceneInfo, float ShadowFadeFraction) const
+	void Set(FRHIBatchedShaderParameters& BatchedParameters, const FLightSceneInfo* LightSceneInfo, float ShadowFadeFraction) const
 	{
-		SetShaderValue( 
-			RHICmdList, 
-			ShaderRHI, 
-			LightFunctionParameters, 
-			GetLightFunctionSharedParameters(LightSceneInfo, ShadowFadeFraction));
+		SetShaderValue(BatchedParameters, LightFunctionParameters, GetLightFunctionSharedParameters(LightSceneInfo, ShadowFadeFraction));
 	}
 
 	/** Serializer. */ 
@@ -270,8 +261,8 @@ public:
 	END_SHADER_PARAMETER_STRUCT()
 
 	void Bind(const FShaderParameterMap& ParameterMap);
-	void Set(FRHICommandList& RHICmdList, FShader* Shader, const FVector4f& InStencilingGeometryPosAndScale) const;
-	void Set(FRHICommandList& RHICmdList, FShader* Shader, const FSceneView& View, const FLightSceneInfo* LightSceneInfo) const;
+	void Set(FRHIBatchedShaderParameters& BatchedParameters, const FVector4f& InStencilingGeometryPosAndScale) const;
+	void Set(FRHIBatchedShaderParameters& BatchedParameters, const FSceneView& View, const FLightSceneInfo* LightSceneInfo) const;
 
 	static FParameters GetParameters(const FVector4f& InStencilingGeometryPosAndScale);
 	static FParameters GetParameters(const FSceneView& View, const FLightSceneInfo* LightSceneInfo);
diff --git a/Engine/Source/Runtime/Renderer/Private/SceneOcclusion.cpp b/Engine/Source/Runtime/Renderer/Private/SceneOcclusion.cpp
--- a/Engine/Source/Runtime/Renderer/Private/SceneOcclusion.cpp
+++ b/Engine/Source/Runtime/Renderer/Private/SceneOcclusion.cpp
@@ -602,7 +602,8 @@ static void ExecutePointLightShadowOcclusionQuery(FRHICommandList& RHICmdList, F
 	RHICmdList.BeginRenderQuery(ShadowOcclusionQuery);
 	
 	// Draw bounding sphere
-	VertexShader->SetParametersWithBoundingSphere(RHICmdList, View, LightBounds);
+	SetShaderParametersLegacyVS(RHICmdList, VertexShader, View, LightBounds);
+
 	StencilingGeometry::DrawVectorSphere(RHICmdList);
 		
 	RHICmdList.EndRenderQuery(ShadowOcclusionQuery);
@@ -1317,7 +1318,7 @@ static void BeginOcclusionTests(
 
 			{
 				SCOPED_DRAW_EVENT(RHICmdList, ShadowFrustumQueries);
-				VertexShader->SetParameters(RHICmdList, View);
+				SetShaderParametersLegacyVS(RHICmdList, VertexShader, View);
 				RHICmdList.SetStreamSource(0, VertexBufferRHI, 0);
 				BaseVertexOffset = 0;
 
@@ -1349,7 +1350,7 @@ static void BeginOcclusionTests(
 
 		if (ViewQuery.bFlushQueries)
 		{
-			VertexShader->SetParameters(RHICmdList, View);
+			SetShaderParametersLegacyVS(RHICmdList, VertexShader, View);
 
 			if (GOcclusionQueryDispatchOrder == 0)
 			{
diff --git a/Engine/Source/Runtime/Renderer/Private/SceneOcclusion.h b/Engine/Source/Runtime/Renderer/Private/SceneOcclusion.h
--- a/Engine/Source/Runtime/Renderer/Private/SceneOcclusion.h
+++ b/Engine/Source/Runtime/Renderer/Private/SceneOcclusion.h
@@ -65,30 +65,29 @@ public:
 
 	FOcclusionQueryVS() {}
 
-	void SetParametersWithBoundingSphere(FRHICommandList& RHICmdList, const FViewInfo& View, const FSphere& BoundingSphere)
+	void SetParameters(FRHIBatchedShaderParameters& BatchedParameters, const FViewInfo& View, const FSphere& BoundingSphere)
 	{
-		FGlobalShader::SetParameters<FViewUniformShaderParameters>(RHICmdList, RHICmdList.GetBoundVertexShader(), View.ViewUniformBuffer);
-
 		FVector4f StencilingSpherePosAndScale;
 		StencilingGeometry::GStencilSphereVertexBuffer.CalcTransform(StencilingSpherePosAndScale, BoundingSphere, View.ViewMatrices.GetPreViewTranslation());
-		StencilingGeometryParameters.Set(RHICmdList, this, StencilingSpherePosAndScale);
+		SetParametersInternal(BatchedParameters, View, StencilingSpherePosAndScale);
+	}
 
-		if (ViewId.IsBound())
-		{
-			SetShaderValue(RHICmdList, RHICmdList.GetBoundVertexShader(), ViewId, (View.StereoPass == EStereoscopicPass::eSSP_FULL) ? 0 : View.StereoViewIndex);
-		}
+	void SetParameters(FRHIBatchedShaderParameters& BatchedParameters, const FViewInfo& View)
+	{
+		SetParametersInternal(BatchedParameters, View, FVector4f(0, 0, 0, 1));
 	}
 
-	void SetParameters(FRHICommandList& RHICmdList, const FViewInfo& View)
+private:
+	void SetParametersInternal(FRHIBatchedShaderParameters& BatchedParameters, const FViewInfo& View, const FVector4f& StencilingSpherePosAndScale)
 	{
-		FGlobalShader::SetParameters<FViewUniformShaderParameters>(RHICmdList, RHICmdList.GetBoundVertexShader(),View.ViewUniformBuffer);
+		FGlobalShader::SetParameters<FViewUniformShaderParameters>(BatchedParameters, View.ViewUniformBuffer);
 
 		// Don't transform if rendering frustum
-		StencilingGeometryParameters.Set(RHICmdList, this, FVector4f(0,0,0,1));
+		StencilingGeometryParameters.Set(BatchedParameters, StencilingSpherePosAndScale);
 
 		if (ViewId.IsBound())
 		{
-			SetShaderValue(RHICmdList, RHICmdList.GetBoundVertexShader(), ViewId, (View.StereoPass == EStereoscopicPass::eSSP_FULL) ? 0 : View.StereoViewIndex);
+			SetShaderValue(BatchedParameters, ViewId, (View.StereoPass == EStereoscopicPass::eSSP_FULL) ? 0 : View.StereoViewIndex);
 		}
 	}
 
diff --git a/Engine/Source/Runtime/Renderer/Private/ShadowRendering.cpp b/Engine/Source/Runtime/Renderer/Private/ShadowRendering.cpp
--- a/Engine/Source/Runtime/Renderer/Private/ShadowRendering.cpp
+++ b/Engine/Source/Runtime/Renderer/Private/ShadowRendering.cpp
@@ -302,33 +302,31 @@ void GetOnePassPointShadowProjectionParameters(FRDGBuilder& GraphBuilder, const
 -----------------------------------------------------------------------------*/
 
 void FShadowVolumeBoundProjectionVS::SetParameters(
-	FRHICommandList& RHICmdList,
+	FRHIBatchedShaderParameters& BatchedParameters,
 	const FSceneView& View,
 	const FProjectedShadowInfo* ShadowInfo,
 	EShadowProjectionVertexShaderFlags Flags)
 {
-	FRHIVertexShader* ShaderRHI = RHICmdList.GetBoundVertexShader();
-
 	if(ShadowInfo->IsWholeScenePointLightShadow())
 	{
 		// Handle stenciling sphere for point light.
-		StencilingGeometryParameters.Set(RHICmdList, this, View, ShadowInfo->LightSceneInfo);
+		StencilingGeometryParameters.Set(BatchedParameters, View, ShadowInfo->LightSceneInfo);
 	}
 	else
 	{
-		StencilingGeometryParameters.Set(RHICmdList, this, FVector4f(0,0,0,1));
+		StencilingGeometryParameters.Set(BatchedParameters, FVector4f(0,0,0,1));
 	}
 
 	if ((Flags & EShadowProjectionVertexShaderFlags::DrawingFrustum) != EShadowProjectionVertexShaderFlags::None)
 	{
 		const FVector PreShadowToPreView(View.ViewMatrices.GetPreViewTranslation() - ShadowInfo->PreShadowTranslation);
-		SetShaderValue(RHICmdList, ShaderRHI, InvReceiverInnerMatrix, ShadowInfo->InvReceiverInnerMatrix);
-		SetShaderValue(RHICmdList, ShaderRHI, PreShadowToPreViewTranslation, FVector4f((FVector3f)PreShadowToPreView, 0));
+		SetShaderValue(BatchedParameters, InvReceiverInnerMatrix, ShadowInfo->InvReceiverInnerMatrix);
+		SetShaderValue(BatchedParameters, PreShadowToPreViewTranslation, FVector4f((FVector3f)PreShadowToPreView, 0));
 	}
 	else
 	{
-		SetShaderValue(RHICmdList, ShaderRHI, InvReceiverInnerMatrix, FMatrix44f::Identity);
-		SetShaderValue(RHICmdList, ShaderRHI, PreShadowToPreViewTranslation, FVector4f(0, 0, 0, 0));
+		SetShaderValue(BatchedParameters, InvReceiverInnerMatrix, FMatrix44f::Identity);
+		SetShaderValue(BatchedParameters, PreShadowToPreViewTranslation, FVector4f(0, 0, 0, 0));
 	}
 }
 
@@ -463,21 +461,24 @@ static void BindShaderShaders(FRHICommandList& RHICmdList, FGraphicsPipelineStat
 	GraphicsPSOInit.BoundShaderState.PixelShaderRHI = PixelShader.GetPixelShader();
 	SetGraphicsPipelineState(RHICmdList, GraphicsPSOInit, StencilRef);
 
-	VertexShader->SetParameters(RHICmdList, View, ShadowInfo, EShadowProjectionVertexShaderFlags::DrawingFrustum);
-	PixelShader->SetParameters(RHICmdList, ViewIndex, View, ShadowInfo);
+	SetShaderParametersLegacyVS(RHICmdList, VertexShader, View, ShadowInfo, EShadowProjectionVertexShaderFlags::DrawingFrustum);
+
+	FRHIBatchedShaderParameters& BatchedParameters = RHICmdList.GetScratchShaderParameters();
+
+	PixelShader->SetParameters(BatchedParameters, ViewIndex, View, ShadowInfo);
 
 	if (Strata::IsStrataEnabled())
 	{
-		FRHIPixelShader* ShaderRHI = PixelShader.GetPixelShader();
 		TRDGUniformBufferRef<FStrataGlobalUniformParameters> StrataUniformBuffer = Strata::BindStrataGlobalUniformParameters(View);
-		PixelShader->FGlobalShader::template SetParameters<FStrataGlobalUniformParameters>(RHICmdList, ShaderRHI, StrataUniformBuffer->GetRHIRef());
+		PixelShader->FGlobalShader::template SetParameters<FStrataGlobalUniformParameters>(BatchedParameters, StrataUniformBuffer->GetRHIRef());
 	}
 
 	if (HairStrandsUniformBuffer)
 	{
-		FRHIPixelShader* ShaderRHI = PixelShader.GetPixelShader();
-		PixelShader->FGlobalShader::template SetParameters<FHairStrandsViewUniformParameters>(RHICmdList, ShaderRHI, HairStrandsUniformBuffer);
+		PixelShader->FGlobalShader::template SetParameters<FHairStrandsViewUniformParameters>(BatchedParameters, HairStrandsUniformBuffer);
 	}
+
+	RHICmdList.SetBatchedShaderParameters(PixelShader.GetPixelShader(), BatchedParameters);
 }
 
 
@@ -1060,7 +1061,7 @@ void FProjectedShadowInfo::SetupProjectionStencilMask(
 		SetGraphicsPipelineState(RHICmdList, GraphicsPSOInit, 0);
 
 		// Set the projection vertex shader parameters
-		VertexShader->SetParameters(RHICmdList, *View, this, EShadowProjectionVertexShaderFlags::DrawingFrustum);
+		SetShaderParametersLegacyVS(RHICmdList, VertexShader, *View, this, EShadowProjectionVertexShaderFlags::DrawingFrustum);
 
 		RHICmdList.SetStreamSource(0, GFrustumVertexBuffer.VertexBufferRHI, 0);
 
@@ -1104,7 +1105,7 @@ void FProjectedShadowInfo::SetupProjectionStencilMaskForHair(FRHICommandList& RH
 	GraphicsPSOInit.BoundShaderState.VertexShaderRHI = VertexShader.GetVertexShader();
 	GraphicsPSOInit.BoundShaderState.PixelShaderRHI = nullptr;
 	SetGraphicsPipelineState(RHICmdList, GraphicsPSOInit, 0);
-	VertexShader->SetParameters(RHICmdList, *View, this, EShadowProjectionVertexShaderFlags::DrawingFrustum);
+	SetShaderParametersLegacyVS(RHICmdList, VertexShader, *View, this, EShadowProjectionVertexShaderFlags::DrawingFrustum);
 
 	RHICmdList.SetStreamSource(0, GFrustumVertexBuffer.VertexBufferRHI, 0);
 	RHICmdList.DrawIndexedPrimitive(GCubeIndexBuffer.IndexBufferRHI, 0, 0, 8, 0, 12, 1);
@@ -1480,9 +1481,9 @@ static void SetPointLightShaderTempl(FRHICommandList& RHICmdList, FGraphicsPipel
 	GraphicsPSOInit.BoundShaderState.PixelShaderRHI = PixelShader.GetPixelShader();
 
 	SetGraphicsPipelineState(RHICmdList, GraphicsPSOInit, 0);
-	
-	VertexShader->SetParameters(RHICmdList, View, ShadowInfo, EShadowProjectionVertexShaderFlags::None);
-	PixelShader->SetParameters(RHICmdList, ViewIndex, View, ShadowInfo, HairStrandsUniformBuffer);
+
+	SetShaderParametersLegacyVS(RHICmdList, VertexShader, View, ShadowInfo, EShadowProjectionVertexShaderFlags::None);
+	SetShaderParametersLegacyPS(RHICmdList, PixelShader, ViewIndex, View, ShadowInfo, HairStrandsUniformBuffer);
 }
 
 void FProjectedShadowInfo::RenderOnePassPointLightProjection(
diff --git a/Engine/Source/Runtime/Renderer/Private/ShadowRendering.h b/Engine/Source/Runtime/Renderer/Private/ShadowRendering.h
--- a/Engine/Source/Runtime/Renderer/Private/ShadowRendering.h
+++ b/Engine/Source/Runtime/Renderer/Private/ShadowRendering.h
@@ -977,7 +977,7 @@ public:
 		PreShadowToPreViewTranslation.Bind(ParameterMap, TEXT("PreShadowToPreViewTranslation"));
 	}
 
-	void SetParameters(FRHICommandList& RHICmdList, const FSceneView& View,	const FProjectedShadowInfo* ShadowInfo, EShadowProjectionVertexShaderFlags Flags);
+	void SetParameters(FRHIBatchedShaderParameters& BatchedParameters, const FSceneView& View,	const FProjectedShadowInfo* ShadowInfo, EShadowProjectionVertexShaderFlags Flags);
 
 private:
 	LAYOUT_FIELD(FStencilingGeometryShaderParameters, StencilingGeometryParameters);
@@ -995,7 +995,7 @@ public:
 	{
 	}
 
-	void SetParameters(FRHICommandList& RHICmdList, const FSceneView& View, const FProjectedShadowInfo* ShadowInfo, EShadowProjectionVertexShaderFlags Flags) {}
+	void SetParameters(FRHIBatchedShaderParameters& BatchedParameters, const FSceneView& View, const FProjectedShadowInfo* ShadowInfo, EShadowProjectionVertexShaderFlags Flags) {}
 };
 
 /**
@@ -1044,10 +1044,8 @@ public:
 		bCascadeUseFadePlane.Bind(ParameterMap, TEXT("bCascadeUseFadePlane"));
 	}
 
-	void Set(FRHICommandList& RHICmdList, FShader* Shader, const FSceneView& View, const FProjectedShadowInfo* ShadowInfo, bool bModulatedShadows, bool bUseFadePlane, bool SubPixelShadow)
+	void Set(FRHIBatchedShaderParameters& BatchedParameters, const FSceneView& View, const FProjectedShadowInfo* ShadowInfo, bool bModulatedShadows, bool bUseFadePlane, bool SubPixelShadow)
 	{
-		FRHIPixelShader* ShaderRHI = RHICmdList.GetBoundPixelShader();
-
 		const FVector PreViewTranslation = View.ViewMatrices.GetPreViewTranslation();
 		const FIntPoint ShadowBufferResolution = ShadowInfo->GetShadowBufferResolution();
 
@@ -1059,7 +1057,7 @@ public:
 				(ShadowInfo->BorderSize + ShadowInfo->Y) * InverseShadowBufferResolution.Y,
 				ShadowInfo->ResolutionX * InverseShadowBufferResolution.X,
 				ShadowInfo->ResolutionY * InverseShadowBufferResolution.Y);
-			SetShaderValue(RHICmdList, ShaderRHI, ShadowTileOffsetAndSizeParam, ShadowTileOffsetAndSize);
+			SetShaderValue(BatchedParameters, ShadowTileOffsetAndSizeParam, ShadowTileOffsetAndSize);
 		}
 
 		// Set the transform from screen coordinates to shadow depth texture coordinates.
@@ -1067,26 +1065,26 @@ public:
 		{
 			// UE-29083 : work around precision issues with ScreenToShadowMatrix on low end devices.
 			const FMatrix44f ScreenToShadow = FMatrix44f(ShadowInfo->GetScreenToShadowMatrix(View, 0, 0, ShadowBufferResolution.X, ShadowBufferResolution.Y));	// LWC_TODO: Precision loss
-			SetShaderValue(RHICmdList, ShaderRHI, ScreenToShadowMatrix, ScreenToShadow);
+			SetShaderValue(BatchedParameters, ScreenToShadowMatrix, ScreenToShadow);
 		}
 		else
 		{
 			const FMatrix44f ScreenToShadow = FMatrix44f(ShadowInfo->GetScreenToShadowMatrix(View));		// LWC_TODO: Precision loss
-			SetShaderValue(RHICmdList, ShaderRHI, ScreenToShadowMatrix, ScreenToShadow);
+			SetShaderValue(BatchedParameters, ScreenToShadowMatrix, ScreenToShadow);
 		}
 
 		if (SoftTransitionScale.IsBound())
 		{
 			const float TransitionSize = ShadowInfo->ComputeTransitionSize();
 
-			SetShaderValue(RHICmdList, ShaderRHI, SoftTransitionScale, FVector3f(0, 0, 1.0f / TransitionSize));
+			SetShaderValue(BatchedParameters, SoftTransitionScale, FVector3f(0, 0, 1.0f / TransitionSize));
 		}
 
 		if (ShadowBufferSize.IsBound())
 		{
 			FVector2D ShadowBufferSizeValue(ShadowBufferResolution.X, ShadowBufferResolution.Y);
 
-			SetShaderValue(RHICmdList, ShaderRHI, ShadowBufferSize,
+			SetShaderValue(BatchedParameters, ShadowBufferSize,
 				FVector4f(ShadowBufferSizeValue.X, ShadowBufferSizeValue.Y, 1.0f / ShadowBufferSizeValue.X, 1.0f / ShadowBufferSizeValue.Y));
 		}
 
@@ -1104,24 +1102,20 @@ public:
 			
 		FRHISamplerState* DepthSamplerState = TStaticSamplerState<SF_Point,AM_Clamp,AM_Clamp,AM_Clamp>::GetRHI();
 
-		SetTextureParameter(RHICmdList, ShaderRHI, ShadowDepthTexture, ShadowDepthTextureSampler, DepthSamplerState, ShadowDepthTextureValue);		
+		SetTextureParameter(BatchedParameters, ShadowDepthTexture, ShadowDepthTextureSampler, DepthSamplerState, ShadowDepthTextureValue);
 
 		if (ShadowDepthTextureSampler.IsBound())
 		{
-			RHICmdList.SetShaderSampler(
-				ShaderRHI, 
-				ShadowDepthTextureSampler.GetBaseIndex(),
-				DepthSamplerState
-				);
+			SetSamplerParameter(BatchedParameters, ShadowDepthTextureSampler, DepthSamplerState);
 		}
 
-		SetShaderValue(RHICmdList, ShaderRHI, ProjectionDepthBias, FVector4f(ShadowInfo->GetShaderDepthBias(), ShadowInfo->GetShaderSlopeDepthBias(), ShadowInfo->GetShaderReceiverDepthBias(), ShadowInfo->MaxSubjectZ - ShadowInfo->MinSubjectZ));
-		SetShaderValue(RHICmdList, ShaderRHI, FadePlaneOffset, ShadowInfo->CascadeSettings.FadePlaneOffset);
+		SetShaderValue(BatchedParameters, ProjectionDepthBias, FVector4f(ShadowInfo->GetShaderDepthBias(), ShadowInfo->GetShaderSlopeDepthBias(), ShadowInfo->GetShaderReceiverDepthBias(), ShadowInfo->MaxSubjectZ - ShadowInfo->MinSubjectZ));
+		SetShaderValue(BatchedParameters, FadePlaneOffset, ShadowInfo->CascadeSettings.FadePlaneOffset);
 
 		if(InvFadePlaneLength.IsBound() && bUseFadePlane)
 		{
 			check(ShadowInfo->CascadeSettings.FadePlaneLength > 0);
-			SetShaderValue(RHICmdList, ShaderRHI, InvFadePlaneLength, 1.0f / ShadowInfo->CascadeSettings.FadePlaneLength);
+			SetShaderValue(BatchedParameters, InvFadePlaneLength, 1.0f / ShadowInfo->CascadeSettings.FadePlaneLength);
 		}
 
 		if (LightPositionOrDirection.IsBound())
@@ -1129,7 +1123,7 @@ public:
 			const FVector LightDirection = ShadowInfo->GetLightSceneInfo().Proxy->GetDirection();
 			const FVector LightPosition = ShadowInfo->GetLightSceneInfo().Proxy->GetPosition() + PreViewTranslation;
 			const bool bIsDirectional = ShadowInfo->GetLightSceneInfo().Proxy->GetLightType() == LightType_Directional;
-			SetShaderValue(RHICmdList, ShaderRHI, LightPositionOrDirection, bIsDirectional ? FVector4f((FVector3f)LightDirection,0) : FVector4f((FVector3f)LightPosition,1));
+			SetShaderValue(BatchedParameters, LightPositionOrDirection, bIsDirectional ? FVector4f((FVector3f)LightDirection,0) : FVector4f((FVector3f)LightPosition,1));
 		}
 
 		if (SubPixelShadow)
@@ -1149,12 +1143,12 @@ public:
 			FVector2f SliceNearAndFarDepth;
 			SliceNearAndFarDepth.X = DeviceZNear;
 			SliceNearAndFarDepth.Y = DeviceZFar;
-			SetShaderValue(RHICmdList, ShaderRHI, ShadowNearAndFarDepth, SliceNearAndFarDepth);
-			SetShaderValue(RHICmdList, ShaderRHI, bCascadeUseFadePlane, bUseFadePlane ? 1 : 0);
+			SetShaderValue(BatchedParameters, ShadowNearAndFarDepth, SliceNearAndFarDepth);
+			SetShaderValue(BatchedParameters, bCascadeUseFadePlane, bUseFadePlane ? 1 : 0);
 		}
 
-		SetShaderValue(RHICmdList, ShaderRHI, PerObjectShadowFadeStart, ShadowInfo->PerObjectShadowFadeStart);
-		SetShaderValue(RHICmdList, ShaderRHI, InvPerObjectShadowFadeLength, ShadowInfo->InvPerObjectShadowFadeLength);
+		SetShaderValue(BatchedParameters, PerObjectShadowFadeStart, ShadowInfo->PerObjectShadowFadeStart);
+		SetShaderValue(BatchedParameters, InvPerObjectShadowFadeLength, ShadowInfo->InvPerObjectShadowFadeLength);
 	}
 
 private:
@@ -1231,28 +1225,26 @@ public:
 	 * @param ShadowInfo - projected shadow info for a single light
 	 */
 	void SetParameters(
-		FRHICommandList& RHICmdList, 
+		FRHIBatchedShaderParameters& BatchedParameters,
 		int32 ViewIndex,
 		const FSceneView& View,
 		const FProjectedShadowInfo* ShadowInfo)
 	{
-		FRHIPixelShader* ShaderRHI = RHICmdList.GetBoundPixelShader();
-
 		const FVector PreViewTranslation = View.ViewMatrices.GetPreViewTranslation();
 		const bool bUseFadePlaneEnable = ShadowInfo->CascadeSettings.FadePlaneLength > 0;
 
-		ProjectionParameters.Set(RHICmdList, this, View, ShadowInfo, bModulatedShadows, bUseFadePlaneEnable, SubPixelShadow);
+		ProjectionParameters.Set(BatchedParameters, View, ShadowInfo, bModulatedShadows, bUseFadePlaneEnable, SubPixelShadow);
 		const FLightSceneProxy& LightProxy = *(ShadowInfo->GetLightSceneInfo().Proxy);
 
-		SetShaderValue(RHICmdList, ShaderRHI, ShadowFadeFraction, ShadowInfo->FadeAlphas[ViewIndex] );
-		SetShaderValue(RHICmdList, ShaderRHI, ShadowSharpen, LightProxy.GetShadowSharpen() * 7.0f + 1.0f );
-		SetShaderValue(RHICmdList, ShaderRHI, LightPosition, FVector4f(FVector3f((FVector)LightProxy.GetPosition() + PreViewTranslation), 1.0f / LightProxy.GetRadius()));
+		SetShaderValue(BatchedParameters, ShadowFadeFraction, ShadowInfo->FadeAlphas[ViewIndex] );
+		SetShaderValue(BatchedParameters, ShadowSharpen, LightProxy.GetShadowSharpen() * 7.0f + 1.0f );
+		SetShaderValue(BatchedParameters, LightPosition, FVector4f(FVector3f((FVector)LightProxy.GetPosition() + PreViewTranslation), 1.0f / LightProxy.GetRadius()));
 
 		auto DeferredLightParameter = GetUniformBufferParameter<FDeferredLightUniformStruct>();
 
 		if (DeferredLightParameter.IsBound())
 		{
-			SetDeferredLightParameters(RHICmdList, ShaderRHI, DeferredLightParameter, &ShadowInfo->GetLightSceneInfo(), View);
+			SetDeferredLightParameters(BatchedParameters, DeferredLightParameter, &ShadowInfo->GetLightSceneInfo(), View);
 		}
 
 
@@ -1295,14 +1287,13 @@ public:
 	TModulatedShadowProjection(const ShaderMetaType::CompiledShaderInitializerType& Initializer);
 
 	void SetParameters(
-		FRHICommandList& RHICmdList,
+		FRHIBatchedShaderParameters& BatchedParameters,
 		int32 ViewIndex,
 		const FSceneView& View,
 		const FProjectedShadowInfo* ShadowInfo)
 	{
-		TShadowProjectionPS<Quality, false, true>::SetParameters(RHICmdList, ViewIndex, View, ShadowInfo);
-		FRHIPixelShader* ShaderRHI = RHICmdList.GetBoundPixelShader();
-		SetShaderValue(RHICmdList, ShaderRHI, ModulatedShadowColorParameter, ShadowInfo->GetLightSceneInfo().Proxy->GetModulatedShadowColor());
+		TShadowProjectionPS<Quality, false, true>::SetParameters(BatchedParameters, ViewIndex, View, ShadowInfo);
+		SetShaderValue(BatchedParameters, ModulatedShadowColorParameter, ShadowInfo->GetLightSceneInfo().Proxy->GetModulatedShadowColor());
 	}
 
 protected:
@@ -1363,16 +1354,16 @@ public:
 	}
 
 	void SetParameters(
-		FRHICommandList& RHICmdList, 
+		FRHIBatchedShaderParameters& BatchedParameters,
 		int32 ViewIndex,
 		const FSceneView& View,
 		const FProjectedShadowInfo* ShadowInfo)
 	{
-		TShadowProjectionPS<Quality>::SetParameters(RHICmdList, ViewIndex, View, ShadowInfo);
+		TShadowProjectionPS<Quality>::SetParameters(BatchedParameters, ViewIndex, View, ShadowInfo);
 
 		FTranslucentSelfShadowUniformParameters TranslucentSelfShadowUniformParameters;
 		SetupTranslucentSelfShadowUniformParameters(ShadowInfo, TranslucentSelfShadowUniformParameters);
-		SetUniformBufferParameterImmediate(RHICmdList, RHICmdList.GetBoundPixelShader(), this->template GetUniformBufferParameter<FTranslucentSelfShadowUniformParameters>(), TranslucentSelfShadowUniformParameters);
+		SetUniformBufferParameterImmediate(BatchedParameters, this->template GetUniformBufferParameter<FTranslucentSelfShadowUniformParameters>(), TranslucentSelfShadowUniformParameters);
 	}
 };
 
@@ -1404,8 +1395,7 @@ public:
 		LightPositionOrDirection.Bind(ParameterMap, TEXT("LightPositionOrDirection"));
 	}
 
-	template<typename ShaderRHIParamRef>
-	void Set(FRHICommandList& RHICmdList, const ShaderRHIParamRef ShaderRHI, const FSceneView& View, const FProjectedShadowInfo* ShadowInfo) const
+	inline void Set(FRHIBatchedShaderParameters& BatchedParameters, const FSceneView& View, const FProjectedShadowInfo* ShadowInfo) const
 	{
 		const FVector PreViewTranslation = View.ViewMatrices.GetPreViewTranslation();
 
@@ -1419,34 +1409,19 @@ public:
 			}
 		}
 
-		SetTextureParameter(
-			RHICmdList, 
-			ShaderRHI, 
-			ShadowDepthTexture, 
-			ShadowDepthTextureValue
-			);
-
-		SetTextureParameter(
-			RHICmdList,
-			ShaderRHI,
-			ShadowDepthTexture2,
-			ShadowDepthTextureValue
-		);
+		SetTextureParameter(BatchedParameters, ShadowDepthTexture, ShadowDepthTextureValue);
+		SetTextureParameter(BatchedParameters, ShadowDepthTexture2, ShadowDepthTextureValue);
 
 		if (LightPositionOrDirection.IsBound())
 		{
 			const FVector LightPosition = ShadowInfo ? FVector(ShadowInfo->GetLightSceneInfo().Proxy->GetPosition()) : FVector::ZeroVector;
-			SetShaderValue(RHICmdList, ShaderRHI, LightPositionOrDirection, FVector4f(FVector3f(LightPosition + PreViewTranslation), 1));
+			SetShaderValue(BatchedParameters, LightPositionOrDirection, FVector4f(FVector3f(LightPosition + PreViewTranslation), 1));
 		}
 		
 		if (ShadowDepthCubeComparisonSampler.IsBound())
 		{
-			RHICmdList.SetShaderSampler(
-				ShaderRHI, 
-				ShadowDepthCubeComparisonSampler.GetBaseIndex(), 
-				// Use a comparison sampler to do hardware PCF
-				TStaticSamplerState<SF_Bilinear, AM_Clamp, AM_Clamp, AM_Clamp, 0, 0, 0, SCF_Less>::GetRHI()
-				);
+			// Use a comparison sampler to do hardware PCF
+			SetSamplerParameter(BatchedParameters, ShadowDepthCubeComparisonSampler, TStaticSamplerState<SF_Bilinear, AM_Clamp, AM_Clamp, AM_Clamp, 0, 0, 0, SCF_Less>::GetRHI());
 		}
 
 		const int32 NumShaderMatrices = FMath::DivideAndRoundUp<int32>(ShadowViewProjectionMatrices.GetNumBytes(), sizeof(FMatrix44f));
@@ -1464,30 +1439,28 @@ public:
 			{
 				TypeCastedMatrices.AddZeroed(NumShaderMatrices - NumUsableMatrices);
 			}
-			SetShaderValueArray<ShaderRHIParamRef, FMatrix44f>(
-				RHICmdList,
-				ShaderRHI,
+			SetShaderValueArray<FMatrix44f>(
+				BatchedParameters,
 				ShadowViewProjectionMatrices,
 				TypeCastedMatrices.GetData(),
 				NumShaderMatrices
 				);
 
-			SetShaderValue(RHICmdList, ShaderRHI,InvShadowmapResolution,1.0f / ShadowInfo->ResolutionX);
+			SetShaderValue(BatchedParameters,InvShadowmapResolution,1.0f / ShadowInfo->ResolutionX);
 		}
 		else
 		{
 			TArray<FMatrix44f, SceneRenderingAllocator> ZeroMatrices;
 			ZeroMatrices.AddZeroed(NumShaderMatrices);
 
-			SetShaderValueArray<ShaderRHIParamRef, FMatrix44f>(
-				RHICmdList, 
-				ShaderRHI,
+			SetShaderValueArray<FMatrix44f>(
+				BatchedParameters,
 				ShadowViewProjectionMatrices,
 				ZeroMatrices.GetData(),
 				NumShaderMatrices
 				);
 
-			SetShaderValue(RHICmdList, ShaderRHI,InvShadowmapResolution,0);
+			SetShaderValue(BatchedParameters,InvShadowmapResolution,0);
 		}
 	}
 
@@ -1552,38 +1525,37 @@ public:
 	}
 
 	void SetParameters(
-		FRHICommandList& RHICmdList, 
+		FRHIBatchedShaderParameters& BatchedParameters,
 		int32 ViewIndex,
 		const FViewInfo& View,
 		const FProjectedShadowInfo* ShadowInfo,
 		FRHIUniformBuffer* HairStrandsUniformBuffer)
 	{
-		FRHIPixelShader* ShaderRHI = RHICmdList.GetBoundPixelShader();
 		const FVector PreViewTranslation = View.ViewMatrices.GetPreViewTranslation();
 
-		OnePassShadowParameters.Set(RHICmdList, ShaderRHI, View, ShadowInfo);
+		OnePassShadowParameters.Set(BatchedParameters, View, ShadowInfo);
 
 		const FLightSceneProxy& LightProxy = *(ShadowInfo->GetLightSceneInfo().Proxy);
 
-		SetShaderValue(RHICmdList, ShaderRHI, LightPosition, FVector4f(FVector3f(FVector(LightProxy.GetPosition()) + PreViewTranslation), 1.0f / LightProxy.GetRadius()));
+		SetShaderValue(BatchedParameters, LightPosition, FVector4f(FVector3f(FVector(LightProxy.GetPosition()) + PreViewTranslation), 1.0f / LightProxy.GetRadius()));
 
-		SetShaderValue(RHICmdList, ShaderRHI, ShadowFadeFraction, ShadowInfo->FadeAlphas[ViewIndex]);
-		SetShaderValue(RHICmdList, ShaderRHI, ShadowSharpen, LightProxy.GetShadowSharpen() * 7.0f + 1.0f);
+		SetShaderValue(BatchedParameters, ShadowFadeFraction, ShadowInfo->FadeAlphas[ViewIndex]);
+		SetShaderValue(BatchedParameters, ShadowSharpen, LightProxy.GetShadowSharpen() * 7.0f + 1.0f);
 
 		FVector2d ProjectionParams = FVector2d(ShadowInfo->OnePassShadowFaceProjectionMatrix.M[2][2], ShadowInfo->OnePassShadowFaceProjectionMatrix.M[3][2]);
 		FVector2f InverseProjParams = FVector2f(1.0 / ProjectionParams.Y, ProjectionParams.X / ProjectionParams.Y);
-		SetShaderValue(RHICmdList, ShaderRHI, PointLightDepthBias, FVector3f(ShadowInfo->GetShaderDepthBias(), ShadowInfo->GetShaderSlopeDepthBias(), ShadowInfo->GetShaderMaxSlopeDepthBias()));
-		SetShaderValue(RHICmdList, ShaderRHI, PointLightProjParameters, InverseProjParams);
+		SetShaderValue(BatchedParameters, PointLightDepthBias, FVector3f(ShadowInfo->GetShaderDepthBias(), ShadowInfo->GetShaderSlopeDepthBias(), ShadowInfo->GetShaderMaxSlopeDepthBias()));
+		SetShaderValue(BatchedParameters, PointLightProjParameters, InverseProjParams);
 
 		if (HairStrandsParameters.IsBound())
 		{
-			SetUniformBufferParameter(RHICmdList, ShaderRHI, HairStrandsParameters, HairStrandsUniformBuffer);
+			SetUniformBufferParameter(BatchedParameters, HairStrandsParameters, HairStrandsUniformBuffer);
 		}
 
 		if (StrataGlobalParameters.IsBound())
 		{
 			TRDGUniformBufferRef<FStrataGlobalUniformParameters> StrataUniformBuffer = Strata::BindStrataGlobalUniformParameters(View);
-			SetUniformBufferParameter(RHICmdList, ShaderRHI, StrataGlobalParameters, StrataUniformBuffer->GetRHIRef());
+			SetUniformBufferParameter(BatchedParameters, StrataGlobalParameters, StrataUniformBuffer->GetRHIRef());
 		}
 
 		FScene* Scene = nullptr;
@@ -1593,13 +1565,13 @@ public:
 			Scene = View.Family->Scene->GetRenderScene();
 		}
 
-		SetSamplerParameter(RHICmdList, ShaderRHI, ShadowDepthTextureSampler, TStaticSamplerState<SF_Point, AM_Clamp, AM_Clamp, AM_Clamp>::GetRHI());
+		SetSamplerParameter(BatchedParameters, ShadowDepthTextureSampler, TStaticSamplerState<SF_Point, AM_Clamp, AM_Clamp, AM_Clamp>::GetRHI());
 
 		auto DeferredLightParameter = GetUniformBufferParameter<FDeferredLightUniformStruct>();
 
 		if (DeferredLightParameter.IsBound())
 		{
-			SetDeferredLightParameters(RHICmdList, ShaderRHI, DeferredLightParameter, &ShadowInfo->GetLightSceneInfo(), View);
+			SetDeferredLightParameters(BatchedParameters, DeferredLightParameter, &ShadowInfo->GetLightSceneInfo(), View);
 		}
 	}
 
@@ -1674,14 +1646,12 @@ public:
 	}
 
 	void SetParameters(
-		FRHICommandList& RHICmdList,
+		FRHIBatchedShaderParameters& BatchedParameters,
 		int32 ViewIndex,
 		const FSceneView& View,
 		const FProjectedShadowInfo* ShadowInfo)
 	{
-		TShadowProjectionPS<Quality, bUseFadePlane>::SetParameters(RHICmdList, ViewIndex, View, ShadowInfo);
-
-		FRHIPixelShader* ShaderRHI = RHICmdList.GetBoundPixelShader();
+		TShadowProjectionPS<Quality, bUseFadePlane>::SetParameters(BatchedParameters, ViewIndex, View, ShadowInfo);
 
 		// GetLightSourceAngle returns the full angle.
 		float TanLightSourceAngle = FMath::Tan(0.5 * FMath::DegreesToRadians(ShadowInfo->GetLightSceneInfo().Proxy->GetLightSourceAngle()));
@@ -1694,7 +1664,7 @@ public:
 		float SZ = ShadowInfo->MaxSubjectZ - ShadowInfo->MinSubjectZ;
 
 		FVector4f PCSSParameterValues = FVector4f(TanLightSourceAngle * SZ / SW, MaxKernelSize / float(ShadowInfo->ResolutionX), 0, 0);
-		SetShaderValue(RHICmdList, ShaderRHI, PCSSParameters, PCSSParameterValues);
+		SetShaderValue(BatchedParameters, PCSSParameters, PCSSParameterValues);
 	}
 
 protected:
@@ -1730,23 +1700,21 @@ public:
 	}
 
 	void SetParameters(
-		FRHICommandList& RHICmdList,
+		FRHIBatchedShaderParameters& BatchedParameters,
 		int32 ViewIndex,
 		const FSceneView& View,
 		const FProjectedShadowInfo* ShadowInfo)
 	{
 		check(ShadowInfo->GetLightSceneInfo().Proxy->GetLightType() == LightType_Spot);
 
-		TShadowProjectionPS<Quality, bUseFadePlane>::SetParameters(RHICmdList, ViewIndex, View, ShadowInfo);
-
-		FRHIPixelShader* ShaderRHI = RHICmdList.GetBoundPixelShader();
+		TShadowProjectionPS<Quality, bUseFadePlane>::SetParameters(BatchedParameters, ViewIndex, View, ShadowInfo);
 
 		static IConsoleVariable* CVarMaxSoftShadowKernelSize = IConsoleManager::Get().FindConsoleVariable(TEXT("r.Shadow.MaxSoftKernelSize"));
 		check(CVarMaxSoftShadowKernelSize);
 		int32 MaxKernelSize = CVarMaxSoftShadowKernelSize->GetInt();
 
 		FVector4f PCSSParameterValues = FVector4f(0, MaxKernelSize / float(ShadowInfo->ResolutionX), 0, 0);
-		SetShaderValue(RHICmdList, ShaderRHI, PCSSParameters, PCSSParameterValues);
+		SetShaderValue(BatchedParameters, PCSSParameters, PCSSParameterValues);
 	}
 
 protected:
