Coalesced UpdateAllPrimitiveSceneInfo calls into one call prior to scene rendering. This includes all {Add, Remove}Level instances and a spot on the game thread where all scenes are updated. The former case is handled by deferring level adds / removes until UpdateAllScenePrimitiveInfos is called. The latter processes scenes at the end of the frame only if they weren't already processed by a render event.

This coalescing has a few benefits. First, it allows async LPI and mesh command creation jobs to handle level streaming events. Second, batches more work together which opens more doors to parallelization in follow-up CL's.

Finally, this change splits the mesh caching jobs for nanite and ray tracing off into their own tasks, which improves parallelism.

diff --git a/Engine/Source/Runtime/Engine/Public/SceneInterface.h b/Engine/Source/Runtime/Engine/Public/SceneInterface.h
--- a/Engine/Source/Runtime/Engine/Public/SceneInterface.h
+++ b/Engine/Source/Runtime/Engine/Public/SceneInterface.h
@@ -626,6 +626,7 @@ public:
 	virtual TConstArrayView<FPrimitiveComponentId> GetScenePrimitiveComponentIds() const = 0;
 
 	virtual void StartFrame() {}
+	virtual void EndFrame(FRHICommandListImmediate& RHICmdList) {}
 	virtual uint32 GetFrameNumber() const { return 0; }
 	virtual void IncrementFrameNumber() {}
 
diff --git a/Engine/Source/Runtime/Launch/Private/LaunchEngineLoop.cpp b/Engine/Source/Runtime/Launch/Private/LaunchEngineLoop.cpp
--- a/Engine/Source/Runtime/Launch/Private/LaunchEngineLoop.cpp
+++ b/Engine/Source/Runtime/Launch/Private/LaunchEngineLoop.cpp
@@ -5648,38 +5648,18 @@ void FEngineLoop::Tick()
 			GEngine->SetSimulationLatencyMarkerStart(CurrentFrameCounter);
 		}
 
-		for (const FWorldContext& Context : GEngine->GetWorldContexts())
-		{
-			UWorld* CurrentWorld = Context.World();
-			if (CurrentWorld)
-			{
-				FSceneInterface* Scene = CurrentWorld->Scene;
-				ENQUEUE_RENDER_COMMAND(UpdateScenePrimitives)(
-					[Scene](FRHICommandListImmediate& RHICmdList)
-				{
-					Scene->UpdateAllPrimitiveSceneInfos(RHICmdList);
-				});
-			}
-		}
-
 		// beginning of RHI frame
 		ENQUEUE_RENDER_COMMAND(BeginFrame)([CurrentFrameCounter](FRHICommandListImmediate& RHICmdList)
 		{
 			BeginFrameRenderThread(RHICmdList, CurrentFrameCounter);
 		});
 
-		for (const FWorldContext& Context : GEngine->GetWorldContexts())
+		for (FSceneInterface* Scene : GetRendererModule().GetAllocatedScenes())
 		{
-			UWorld* CurrentWorld = Context.World();
-			if (CurrentWorld)
+			ENQUEUE_RENDER_COMMAND(FScene_StartFrame)([Scene](FRHICommandListImmediate& RHICmdList)
 			{
-				FSceneInterface* Scene = CurrentWorld->Scene;
-
-				ENQUEUE_RENDER_COMMAND(SceneStartFrame)([Scene](FRHICommandListImmediate& RHICmdList)
-				{
-					Scene->StartFrame();
-				});
-			}
+				Scene->StartFrame();
+			});
 		}
 
 #if !UE_SERVER && WITH_ENGINE
@@ -6004,6 +5984,14 @@ void FEngineLoop::Tick()
 		}
 #endif
 
+		for (FSceneInterface* Scene : GetRendererModule().GetAllocatedScenes())
+		{
+			ENQUEUE_RENDER_COMMAND(FScene_EndFrame)([Scene](FRHICommandListImmediate& RHICmdList)
+			{
+				Scene->EndFrame(RHICmdList);
+			});
+		}
+
 		// tick render hardware interface
 		{			
 			SCOPE_CYCLE_COUNTER(STAT_RHITickTime);
diff --git a/Engine/Source/Runtime/RenderCore/Private/ShaderResource.cpp b/Engine/Source/Runtime/RenderCore/Private/ShaderResource.cpp
--- a/Engine/Source/Runtime/RenderCore/Private/ShaderResource.cpp
+++ b/Engine/Source/Runtime/RenderCore/Private/ShaderResource.cpp
@@ -600,7 +600,7 @@ FRHIShader* FShaderMapResource_InlineCode::CreateRHIShaderOrCrash(int32 ShaderIn
 	double TimeFunctionEntered = FPlatformTime::Seconds();
 	ON_SCOPE_EXIT
 	{
-		if (IsInRenderingThread())
+		if (IsInParallelRenderingThread())
 		{
 			double ShaderCreationTime = FPlatformTime::Seconds() - TimeFunctionEntered;
 			INC_FLOAT_STAT_BY(STAT_Shaders_TotalRTShaderInitForRenderingTime, ShaderCreationTime);
diff --git a/Engine/Source/Runtime/Renderer/Private/PrimitiveSceneInfo.cpp b/Engine/Source/Runtime/Renderer/Private/PrimitiveSceneInfo.cpp
--- a/Engine/Source/Runtime/Renderer/Private/PrimitiveSceneInfo.cpp
+++ b/Engine/Source/Runtime/Renderer/Private/PrimitiveSceneInfo.cpp
@@ -200,6 +200,7 @@ FPrimitiveSceneInfo::FPrimitiveSceneInfo(UPrimitiveComponent* InComponent, FScen
 	bCacheShadowAsStatic(InComponent->Mobility != EComponentMobility::Movable),
 	bNaniteRasterBinsRenderCustomDepth(false),
 	bPendingAddToScene(false),
+	bPendingAddStaticMeshes(false),
 	bPendingFlushVirtualTexture(false),
 	LevelUpdateNotificationIndex(INDEX_NONE),
 	InstanceSceneDataOffset(INDEX_NONE),
@@ -451,7 +452,7 @@ void FPrimitiveSceneInfo::CacheMeshDrawCommands(FScene* Scene, TArrayView<FPrimi
 				FOptionalTaskTagScope Scope(ETaskTag::EParallelRenderingThread);
 				DoWorkLambda(DrawListContexts[Index], Index);
 			},
-			EParallelForFlags::PumpRenderingThread | EParallelForFlags::Unbalanced
+			EParallelForFlags::Unbalanced
 		);
 
 		if (NumBatches > 0)
@@ -797,6 +798,7 @@ void FPrimitiveSceneInfo::UpdateCachedRayTracingInstances(FScene* Scene, const T
 {
 	if (IsRayTracingEnabled())
 	{
+		SCOPED_NAMED_EVENT(FPrimitiveSceneInfo_UpdateCachedRayTracingInstances, FColor::Turquoise);
 		checkf(GRHISupportsMultithreadedShaderCreation, TEXT("Raytracing code needs the ability to create shaders from task threads."));
 
 		for (FPrimitiveSceneInfo* SceneInfo : SceneInfos)
@@ -1130,6 +1132,7 @@ void FPrimitiveSceneInfo::AddStaticMeshes(FScene* Scene, TArrayView<FPrimitiveSc
 			SceneInfo->Proxy->DrawStaticElements(&BatchingSPDI);
 			SceneInfo->StaticMeshes.Shrink();
 			SceneInfo->StaticMeshRelevances.Shrink();
+			SceneInfo->bPendingAddStaticMeshes = false;
 
 			check(SceneInfo->StaticMeshRelevances.Num() == SceneInfo->StaticMeshes.Num());
 		});
@@ -1362,9 +1365,10 @@ void FPrimitiveSceneInfo::FreeGPUSceneInstances()
 	}
 }
 
-void FPrimitiveSceneInfo::AddToScene(FScene* Scene, TArrayView<FPrimitiveSceneInfo*> SceneInfos, EPrimitiveAddToSceneOps Ops)
+void FPrimitiveSceneInfo::AddToScene(FScene* Scene, TArrayView<FPrimitiveSceneInfo*> SceneInfos)
 {
 	check(IsInRenderingThread());
+	SCOPED_NAMED_EVENT(FPrimitiveSceneInfo_AddToScene, FColor::Turquoise);
 
 	{
 		SCOPED_NAMED_EVENT(FPrimitiveSceneInfo_AddToScene_IndirectLightingCacheUniformBuffer, FColor::Turquoise);
@@ -1455,14 +1459,6 @@ void FPrimitiveSceneInfo::AddToScene(FScene* Scene, TArrayView<FPrimitiveSceneIn
 		}
 	}
 
-	{
-		SCOPED_NAMED_EVENT(FPrimitiveSceneInfo_AddToScene_AddStaticMeshes, FColor::Magenta);
-		if (EnumHasAnyFlags(Ops, EPrimitiveAddToSceneOps::AddStaticMeshes))
-		{
-			AddStaticMeshes(Scene, SceneInfos, EnumHasAnyFlags(Ops, EPrimitiveAddToSceneOps::CacheMeshDrawCommands));
-		}
-	}
-
 	{
 		SCOPED_NAMED_EVENT(FPrimitiveSceneInfo_AddToScene_AddToPrimitiveOctree, FColor::Red);
 		for (FPrimitiveSceneInfo* SceneInfo : SceneInfos)
@@ -1559,49 +1555,41 @@ void FPrimitiveSceneInfo::AddToScene(FScene* Scene, TArrayView<FPrimitiveSceneIn
 	}
 
 	{
-		SCOPED_NAMED_EVENT(FPrimitiveSceneInfo_AddToScene_UpdateVirtualTexture, FColor::Emerald);
+		SCOPED_NAMED_EVENT(FPrimitiveSceneInfo_AddToScene_LevelNotifyPrimitives, FColor::Blue);
 		for (FPrimitiveSceneInfo* SceneInfo : SceneInfos)
 		{
-			FPrimitiveSceneProxy* Proxy = SceneInfo->Proxy;
-			// Store the runtime virtual texture flags.
-			SceneInfo->UpdateRuntimeVirtualTextureFlags();
-			Scene->PrimitiveVirtualTextureFlags[SceneInfo->PackedIndex] = SceneInfo->RuntimeVirtualTextureFlags;
-
-			// Store the runtime virtual texture Lod info.
-			if (SceneInfo->RuntimeVirtualTextureFlags.bRenderToVirtualTexture)
+			if (SceneInfo->Proxy->ShouldNotifyOnWorldAddRemove())
 			{
-				int8 MinLod, MaxLod;
-				GetRuntimeVirtualTextureLODRange(SceneInfo->StaticMeshRelevances, MinLod, MaxLod);
-
-				FPrimitiveVirtualTextureLodInfo& LodInfo = Scene->PrimitiveVirtualTextureLod[SceneInfo->PackedIndex];
-				LodInfo.MinLod = FMath::Clamp((int32)MinLod, 0, 15);
-				LodInfo.MaxLod = FMath::Clamp((int32)MaxLod, 0, 15);
-				LodInfo.LodBias = FMath::Clamp(Proxy->GetVirtualTextureLodBias() + FPrimitiveVirtualTextureLodInfo::LodBiasOffset, 0, 15);
-				LodInfo.CullMethod = Proxy->GetVirtualTextureMinCoverage() == 0 ? 0 : 1;
-				LodInfo.CullValue = LodInfo.CullMethod == 0 ? Proxy->GetVirtualTextureCullMips() : Proxy->GetVirtualTextureMinCoverage();
+				TArray<FPrimitiveSceneInfo*>& LevelNotifyPrimitives = Scene->PrimitivesNeedingLevelUpdateNotification.FindOrAdd(SceneInfo->Proxy->GetLevelName());
+				SceneInfo->LevelUpdateNotificationIndex = LevelNotifyPrimitives.Num();
+				LevelNotifyPrimitives.Add(SceneInfo);
 			}
 		}
 	}
+}
 
-	// Find lights that affect the primitive in the light octree.
-	if (EnumHasAnyFlags(Ops, EPrimitiveAddToSceneOps::CreateLightPrimitiveInteractions))
+void FPrimitiveSceneInfo::UpdateVirtualTextures(FScene* Scene, TArrayView<FPrimitiveSceneInfo*> SceneInfos)
+{
+	SCOPED_NAMED_EVENT(FPrimitiveSceneInfo_AddToScene_UpdateVirtualTexture, FColor::Emerald);
+	for (FPrimitiveSceneInfo* SceneInfo : SceneInfos)
 	{
-		for (FPrimitiveSceneInfo* SceneInfo : SceneInfos)
-		{
-			Scene->CreateLightPrimitiveInteractionsForPrimitive(SceneInfo);
-		}
-	}
+		FPrimitiveSceneProxy* Proxy = SceneInfo->Proxy;
+		// Store the runtime virtual texture flags.
+		SceneInfo->UpdateRuntimeVirtualTextureFlags();
+		Scene->PrimitiveVirtualTextureFlags[SceneInfo->PackedIndex] = SceneInfo->RuntimeVirtualTextureFlags;
 
-	{
-		SCOPED_NAMED_EVENT(FPrimitiveSceneInfo_AddToScene_LevelNotifyPrimitives, FColor::Blue);
-		for (FPrimitiveSceneInfo* SceneInfo : SceneInfos)
+		// Store the runtime virtual texture Lod info.
+		if (SceneInfo->RuntimeVirtualTextureFlags.bRenderToVirtualTexture)
 		{
-			if (SceneInfo->Proxy->ShouldNotifyOnWorldAddRemove())
-			{
-				TArray<FPrimitiveSceneInfo*>& LevelNotifyPrimitives = Scene->PrimitivesNeedingLevelUpdateNotification.FindOrAdd(SceneInfo->Proxy->GetLevelName());
-				SceneInfo->LevelUpdateNotificationIndex = LevelNotifyPrimitives.Num();
-				LevelNotifyPrimitives.Add(SceneInfo);
-			}
+			int8 MinLod, MaxLod;
+			GetRuntimeVirtualTextureLODRange(SceneInfo->StaticMeshRelevances, MinLod, MaxLod);
+
+			FPrimitiveVirtualTextureLodInfo& LodInfo = Scene->PrimitiveVirtualTextureLod[SceneInfo->PackedIndex];
+			LodInfo.MinLod = FMath::Clamp((int32)MinLod, 0, 15);
+			LodInfo.MaxLod = FMath::Clamp((int32)MaxLod, 0, 15);
+			LodInfo.LodBias = FMath::Clamp(Proxy->GetVirtualTextureLodBias() + FPrimitiveVirtualTextureLodInfo::LodBiasOffset, 0, 15);
+			LodInfo.CullMethod = Proxy->GetVirtualTextureMinCoverage() == 0 ? 0 : 1;
+			LodInfo.CullValue = LodInfo.CullMethod == 0 ? Proxy->GetVirtualTextureCullMips() : Proxy->GetVirtualTextureMinCoverage();
 		}
 	}
 }
diff --git a/Engine/Source/Runtime/Renderer/Private/RendererScene.cpp b/Engine/Source/Runtime/Renderer/Private/RendererScene.cpp
--- a/Engine/Source/Runtime/Renderer/Private/RendererScene.cpp
+++ b/Engine/Source/Runtime/Renderer/Private/RendererScene.cpp
@@ -4944,41 +4944,13 @@ void FScene::OnLevelAddedToWorld(const FName& InLevelAddedName, UWorld* InWorld,
 	ENQUEUE_RENDER_COMMAND(FLevelAddedToWorld)(
 		[Scene, LevelAddedName](FRHICommandListImmediate& RHICmdList)
 		{
-			Scene->UpdateAllPrimitiveSceneInfos(RHICmdList);
-			Scene->OnLevelAddedToWorld_RenderThread(LevelAddedName);
+			FLevelCommand Cmd;
+			Cmd.Name = LevelAddedName;
+			Cmd.Op = FLevelCommand::EOp::Add;
+			Scene->LevelCommands.Add(Cmd);
 		});
 }
 
-void FScene::OnLevelAddedToWorld_RenderThread(const FName& InLevelName)
-{
-	TRACE_CPUPROFILER_EVENT_SCOPE(FScene::OnLevelAddedToWorld_RenderThread);
-
-	// Mark level primitives
-	TArray<FPrimitiveSceneInfo*> PrimitivesToAdd;
-
-	if (const TArray<FPrimitiveSceneInfo*>* LevelPrimitives = PrimitivesNeedingLevelUpdateNotification.Find(InLevelName))
-	{
-		for (FPrimitiveSceneInfo* Primitive : *LevelPrimitives)
-		{
-			// If the primitive proxy returns true, it needs it's static meshes added to the scene
-			if (Primitive->Proxy->OnLevelAddedToWorld_RenderThread())
-			{
-				// Remove static meshes and cached commands for any primitives that need to be added
-				Primitive->RemoveStaticMeshes();
-				PrimitivesToAdd.Add(Primitive);
-			}
-			// Invalidate primitive proxy entry in GPU Scene. 
-			// This is necessary for Nanite::FSceneProxy to be uploaded to GPU scene (see GetPrimitiveID in GPUScene.cpp)
-			if (Primitive->Proxy->IsNaniteMesh())
-			{
-				Primitive->RequestGPUSceneUpdate();
-			}
-		}
-	}
-
-	FPrimitiveSceneInfo::AddStaticMeshes(this, PrimitivesToAdd);
-}
-
 void FScene::OnLevelRemovedFromWorld(const FName& InLevelRemovedName, UWorld* InWorld, bool bIsLightingScenario)
 {
 	if (bIsLightingScenario)
@@ -4991,32 +4963,13 @@ void FScene::OnLevelRemovedFromWorld(const FName& InLevelRemovedName, UWorld* In
 	ENQUEUE_RENDER_COMMAND(FLevelRemovedFromWorld)(
 		[Scene, LevelRemovedName](FRHICommandListImmediate& RHICmdList)
 		{
-			Scene->UpdateAllPrimitiveSceneInfos(RHICmdList);
-			Scene->OnLevelRemovedFromWorld_RenderThread(LevelRemovedName);
+			FLevelCommand Cmd;
+			Cmd.Name = LevelRemovedName;
+			Cmd.Op = FLevelCommand::EOp::Remove;
+			Scene->LevelCommands.Add(Cmd);
 		});
 }
 
-
-void FScene::OnLevelRemovedFromWorld_RenderThread(const FName& InLevelName)
-{
-	TRACE_CPUPROFILER_EVENT_SCOPE(FScene::OnLevelRemovedFromWorld_RenderThread);
-
-	if (TArray<FPrimitiveSceneInfo*>* LevelPrimitives = PrimitivesNeedingLevelUpdateNotification.Find(InLevelName))
-	{
-		for (FPrimitiveSceneInfo* Primitive : *LevelPrimitives)
-		{
-			Primitive->Proxy->OnLevelRemovedFromWorld_RenderThread();
-			// Invalidate primitive proxy entry in GPU Scene.
-			// This is necessary for Nanite::FSceneProxy to be uploaded to GPU scene (see GetPrimitiveID in GPUScene.cpp)
-			if (Primitive->Proxy->IsNaniteMesh())
-			{
-				Primitive->RequestGPUSceneUpdate();
-			}
-		}
-	}
-}
-
-
 void FScene::ProcessAtmosphereLightAddition_RenderThread(FLightSceneInfo* LightSceneInfo)
 {
 	if (LightSceneInfo->Proxy->IsUsedAsAtmosphereSunLight())
@@ -5237,6 +5190,7 @@ static inline bool IsPrimitiveRelevantToPathTracing(FPrimitiveSceneInfo* Primiti
 
 FLightSceneChangeSet FScene::UpdateAllLightSceneInfos(FRDGBuilder& GraphBuilder)
 {
+	SCOPED_NAMED_EVENT(FScene_UpdateAllLightSceneInfos, FColor::Orange);
 	struct FFLightSceneChangeSetAllocation
 	{
 		TArray<int32, SceneRenderingAllocator> RemovedLightIds;
@@ -5384,30 +5338,35 @@ void FScene::UpdateAllPrimitiveSceneInfos(FRDGBuilder& GraphBuilder, EUpdateAllP
 	RemovedLocalPrimitiveSceneInfos.Sort(FPrimitiveArraySortKey());
 
 	TArray<FVirtualShadowMapArrayCacheManager*, SceneRenderingAllocator> VirtualShadowCacheManagers;
-	GetAllVirtualShadowMapCacheManagers(VirtualShadowCacheManagers);
 
-	for (FVirtualShadowMapArrayCacheManager* CacheManager : VirtualShadowCacheManagers)
 	{
-		FVirtualShadowMapArrayCacheManager::FInvalidatingPrimitiveCollector InvalidatingPrimitiveCollector(CacheManager);
+		SCOPED_NAMED_EVENT(FScene_VirtualShadowCacheUpdate, FColor::Orange);
+		GetAllVirtualShadowMapCacheManagers(VirtualShadowCacheManagers);
 
-		// All removed primitives must invalidate their footprints in the VSM before leaving
-		for (FPrimitiveSceneInfo* PrimitiveSceneInfo : RemovedLocalPrimitiveSceneInfos)
-		{
-			InvalidatingPrimitiveCollector.Removed(PrimitiveSceneInfo);
-		}
-		// All updated instances must also before moving or re-allocating (TODO: filter out only those actually updated)
-		for (const auto& Instance : UpdatedInstances)
-		{
-			InvalidatingPrimitiveCollector.UpdatedInstances(Instance.Key->GetPrimitiveSceneInfo());
-		}
-		// As must all primitive updates, 
-		for (const auto& Transform : UpdatedTransforms)
+		for (FVirtualShadowMapArrayCacheManager* CacheManager : VirtualShadowCacheManagers)
 		{
-			InvalidatingPrimitiveCollector.UpdatedTransform(Transform.Key->GetPrimitiveSceneInfo());
-		}
+			FVirtualShadowMapArrayCacheManager::FInvalidatingPrimitiveCollector InvalidatingPrimitiveCollector(CacheManager);
+
+			// All removed primitives must invalidate their footprints in the VSM before leaving
+			for (FPrimitiveSceneInfo* PrimitiveSceneInfo : RemovedLocalPrimitiveSceneInfos)
+			{
+				InvalidatingPrimitiveCollector.Removed(PrimitiveSceneInfo);
+			}
+			// All updated instances must also before moving or re-allocating (TODO: filter out only those actually updated)
+			for (const auto& Instance : UpdatedInstances)
+			{
+				InvalidatingPrimitiveCollector.UpdatedInstances(Instance.Key->GetPrimitiveSceneInfo());
+			}
+			// As must all primitive updates, 
+			for (const auto& Transform : UpdatedTransforms)
+			{
+				InvalidatingPrimitiveCollector.UpdatedTransform(Transform.Key->GetPrimitiveSceneInfo());
+			}
 
-		CacheManager->ProcessRemovedOrUpdatedPrimitives(GraphBuilder, GPUScene, InvalidatingPrimitiveCollector);
+			CacheManager->ProcessRemovedOrUpdatedPrimitives(GraphBuilder, GPUScene, InvalidatingPrimitiveCollector);
+		}
 	}
+
 	TArray<FPrimitiveSceneInfo*> AddedLocalPrimitiveSceneInfos;
 	AddedLocalPrimitiveSceneInfos.Reserve(AddedPrimitiveSceneInfos.Num());
 	for (FPrimitiveSceneInfo* SceneInfo : AddedPrimitiveSceneInfos)
@@ -5690,20 +5649,44 @@ void FScene::UpdateAllPrimitiveSceneInfos(FRDGBuilder& GraphBuilder, EUpdateAllP
 
 	const int32 SceneInfosContainerReservedSize = AddedPrimitiveSceneInfos.Num() + UpdatedTransforms.Num() + UpdatedInstances.Num();
 
-	TArray<FPrimitiveSceneInfo*, SceneRenderingAllocator> SceneInfosToFlushVirtualTexture;
+	TArray<FPrimitiveSceneInfo*, SceneRenderingAllocator>& SceneInfosWithAddToScene = *GraphBuilder.AllocObject<TArray<FPrimitiveSceneInfo*, SceneRenderingAllocator>>();
+	TArray<FPrimitiveSceneInfo*, SceneRenderingAllocator> SceneInfosWithFlushVirtualTexture;
 	TArray<FPrimitiveSceneInfo*, SceneRenderingAllocator>& SceneInfosWithStaticDrawListUpdate = *GraphBuilder.AllocObject<TArray<FPrimitiveSceneInfo*, SceneRenderingAllocator>>();
-	TArray<FPrimitiveSceneInfo*, SceneRenderingAllocator> SceneInfosWithoutStaticDrawListUpdate;
-	SceneInfosToFlushVirtualTexture.Reserve(SceneInfosContainerReservedSize);
+	SceneInfosWithAddToScene.Reserve(SceneInfosContainerReservedSize);
+	SceneInfosWithFlushVirtualTexture.Reserve(SceneInfosContainerReservedSize);
 	SceneInfosWithStaticDrawListUpdate.Reserve(SceneInfosContainerReservedSize);
-	SceneInfosWithoutStaticDrawListUpdate.Reserve(SceneInfosContainerReservedSize);
 
-	const auto QueueSceneInfoToFlushVirtualTexture = [&](FPrimitiveSceneInfo* SceneInfo)
+	const auto QueueFlushVirtualTexture = [&](FPrimitiveSceneInfo* SceneInfo) -> bool
 	{
 		if (!SceneInfo->bPendingFlushVirtualTexture)
 		{
 			SceneInfo->bPendingFlushVirtualTexture = true;
-			SceneInfosToFlushVirtualTexture.Push(SceneInfo);
+			SceneInfosWithFlushVirtualTexture.Push(SceneInfo);
+			return true;
+		}
+		return false;
+	};
+
+	const auto QueueAddToScene = [&](FPrimitiveSceneInfo* SceneInfo) -> bool
+	{
+		if (!SceneInfo->bPendingAddToScene)
+		{
+			SceneInfo->bPendingAddToScene = true;
+			SceneInfosWithAddToScene.Push(SceneInfo);
+			return true;
 		}
+		return false;
+	};
+
+	const auto QueueAddStaticMeshes = [&](FPrimitiveSceneInfo* SceneInfo)
+	{
+		if (!SceneInfo->bPendingAddStaticMeshes)
+		{
+			SceneInfo->bPendingAddStaticMeshes = 1;
+			SceneInfosWithStaticDrawListUpdate.Push(SceneInfo);
+			return true;
+		}
+		return false;
 	};
 
 	{
@@ -5718,6 +5701,7 @@ void FScene::UpdateAllPrimitiveSceneInfos(FRDGBuilder& GraphBuilder, EUpdateAllP
 
 		if (AddedLocalPrimitiveSceneInfos.Num())
 		{
+			SCOPED_NAMED_EVENT(FScene_AddPrimitiveSceneInfos_ReserveMemory, FColor::Green);
 			Primitives.Reserve(Primitives.Num() + AddedLocalPrimitiveSceneInfos.Num());
 			PrimitiveTransforms.Reserve(PrimitiveTransforms.Num() + AddedLocalPrimitiveSceneInfos.Num());
 			PrimitiveSceneProxies.Reserve(PrimitiveSceneProxies.Num() + AddedLocalPrimitiveSceneInfos.Num());
@@ -5762,57 +5746,61 @@ void FScene::UpdateAllPrimitiveSceneInfos(FRDGBuilder& GraphBuilder, EUpdateAllP
 				StartIndex--;
 			}
 
-			for (int32 AddIndex = StartIndex; AddIndex < AddedLocalPrimitiveSceneInfos.Num(); AddIndex++)
 			{
-				FPrimitiveSceneInfo* PrimitiveSceneInfo = AddedLocalPrimitiveSceneInfos[AddIndex];
-				Primitives.Add(PrimitiveSceneInfo);
-				const FMatrix LocalToWorld = PrimitiveSceneInfo->Proxy->GetLocalToWorld();
-				PrimitiveTransforms.Add(LocalToWorld);
-				PrimitiveSceneProxies.Add(PrimitiveSceneInfo->Proxy);
-				PrimitiveBounds.AddUninitialized();
-				PrimitiveFlagsCompact.AddUninitialized();
-				PrimitiveVisibilityIds.AddUninitialized();
-				PrimitiveOctreeIndex.Add(0);
-				PrimitiveOcclusionFlags.AddUninitialized();
-				PrimitiveComponentIds.AddUninitialized();
-				PrimitiveVirtualTextureFlags.AddUninitialized();
-				PrimitiveVirtualTextureLod.AddUninitialized();
-				PrimitiveOcclusionBounds.AddUninitialized();
-			#if WITH_EDITOR
-				PrimitivesSelected.Add(PrimitiveSceneInfo->Proxy->IsSelected());
-			#endif
-			#if RHI_RAYTRACING
-				PrimitiveRayTracingFlags.AddZeroed();
-				PrimitiveRayTracingGroupIds.Add(Experimental::FHashElementId());
-			#endif
-				PrimitivesNeedingStaticMeshUpdate.Add(false);
+				SCOPED_NAMED_EVENT(FScene_AddPrimitiveSceneInfos, FColor::Turquoise);
 
-				for (TMap<int32, TArray<FCachedShadowMapData>>::TIterator CachedShadowMapIt(CachedShadowMaps); CachedShadowMapIt; ++CachedShadowMapIt)
+				for (int32 AddIndex = StartIndex; AddIndex < AddedLocalPrimitiveSceneInfos.Num(); AddIndex++)
 				{
-					TArray<FCachedShadowMapData>& ShadowMapDatas = CachedShadowMapIt.Value();
+					FPrimitiveSceneInfo* PrimitiveSceneInfo = AddedLocalPrimitiveSceneInfos[AddIndex];
+					Primitives.Add(PrimitiveSceneInfo);
+					const FMatrix LocalToWorld = PrimitiveSceneInfo->Proxy->GetLocalToWorld();
+					PrimitiveTransforms.Add(LocalToWorld);
+					PrimitiveSceneProxies.Add(PrimitiveSceneInfo->Proxy);
+					PrimitiveBounds.AddUninitialized();
+					PrimitiveFlagsCompact.AddUninitialized();
+					PrimitiveVisibilityIds.AddUninitialized();
+					PrimitiveOctreeIndex.Add(0);
+					PrimitiveOcclusionFlags.AddUninitialized();
+					PrimitiveComponentIds.AddUninitialized();
+					PrimitiveVirtualTextureFlags.AddUninitialized();
+					PrimitiveVirtualTextureLod.AddUninitialized();
+					PrimitiveOcclusionBounds.AddUninitialized();
+#if WITH_EDITOR
+					PrimitivesSelected.Add(PrimitiveSceneInfo->Proxy->IsSelected());
+#endif
+#if RHI_RAYTRACING
+					PrimitiveRayTracingFlags.AddZeroed();
+					PrimitiveRayTracingGroupIds.Add(Experimental::FHashElementId());
+#endif
+					PrimitivesNeedingStaticMeshUpdate.Add(false);
 
-					for (auto& ShadowMapData : ShadowMapDatas)
+					for (TMap<int32, TArray<FCachedShadowMapData>>::TIterator CachedShadowMapIt(CachedShadowMaps); CachedShadowMapIt; ++CachedShadowMapIt)
 					{
-						if (ShadowMapData.StaticShadowSubjectMap.Num() > 0)
+						TArray<FCachedShadowMapData>& ShadowMapDatas = CachedShadowMapIt.Value();
+
+						for (auto& ShadowMapData : ShadowMapDatas)
 						{
-							ShadowMapData.StaticShadowSubjectMap.Add(false);
+							if (ShadowMapData.StaticShadowSubjectMap.Num() > 0)
+							{
+								ShadowMapData.StaticShadowSubjectMap.Add(false);
+							}
 						}
 					}
-				}
 
-				const int32 SourceIndex = PrimitiveSceneProxies.Num() - 1;
-				PrimitiveSceneInfo->PackedIndex = SourceIndex;
-				checkSlow(PrimitiveSceneInfo->PersistentIndex.Index == INDEX_NONE);
-				FPersistentPrimitiveIndex PersistentPrimitiveIndex{ PersistentPrimitiveIdAllocator.Allocate() };
-				PrimitiveSceneInfo->PersistentIndex = PersistentPrimitiveIndex;
-				// Ensure map is large enough
-				if (PersistentPrimitiveIndex.Index >= PersistentPrimitiveIdToIndexMap.Num())
-				{
-					PersistentPrimitiveIdToIndexMap.SetNumUninitialized(PersistentPrimitiveIndex.Index + 1);
-				}
-				PersistentPrimitiveIdToIndexMap[PersistentPrimitiveIndex.Index] = SourceIndex;
+					const int32 SourceIndex = PrimitiveSceneProxies.Num() - 1;
+					PrimitiveSceneInfo->PackedIndex = SourceIndex;
+					checkSlow(PrimitiveSceneInfo->PersistentIndex.Index == INDEX_NONE);
+					FPersistentPrimitiveIndex PersistentPrimitiveIndex{ PersistentPrimitiveIdAllocator.Allocate() };
+					PrimitiveSceneInfo->PersistentIndex = PersistentPrimitiveIndex;
+					// Ensure map is large enough
+					if (PersistentPrimitiveIndex.Index >= PersistentPrimitiveIdToIndexMap.Num())
+					{
+						PersistentPrimitiveIdToIndexMap.SetNumUninitialized(PersistentPrimitiveIndex.Index + 1);
+					}
+					PersistentPrimitiveIdToIndexMap[PersistentPrimitiveIndex.Index] = SourceIndex;
 
-				GPUScene.AddPrimitiveToUpdate(SourceIndex, EPrimitiveDirtyState::AddedMask);
+					GPUScene.AddPrimitiveToUpdate(SourceIndex, EPrimitiveDirtyState::AddedMask);
+				}
 			}
 
 			bool EntryFound = false;
@@ -5954,9 +5942,6 @@ void FScene::UpdateAllPrimitiveSceneInfos(FRDGBuilder& GraphBuilder, EUpdateAllP
 				int32 PrimitiveIndex = PrimitiveSceneInfo->PackedIndex;
 				FPrimitiveSceneProxy* SceneProxy = PrimitiveSceneInfo->Proxy;
 
-				SceneInfosWithStaticDrawListUpdate.Push(PrimitiveSceneInfo);
-				PrimitiveSceneInfo->bPendingAddToScene = true;
-
 				if (ShouldPrimitiveOutputVelocity(SceneProxy, GetShaderPlatform()))
 				{
 					PrimitiveSceneInfo->bRegisteredWithVelocityData = true;
@@ -5975,8 +5960,9 @@ void FScene::UpdateAllPrimitiveSceneInfos(FRDGBuilder& GraphBuilder, EUpdateAllP
 				}
 #endif
 
-				// Flush virtual textures touched by primitive
-				QueueSceneInfoToFlushVirtualTexture(PrimitiveSceneInfo);
+				QueueFlushVirtualTexture(PrimitiveSceneInfo);
+				QueueAddToScene(PrimitiveSceneInfo);
+				QueueAddStaticMeshes(PrimitiveSceneInfo);
 
 				bNeedPathTracedInvalidation = bNeedPathTracedInvalidation || IsPrimitiveRelevantToPathTracing(PrimitiveSceneInfo);
 			}
@@ -6006,25 +5992,19 @@ void FScene::UpdateAllPrimitiveSceneInfos(FRDGBuilder& GraphBuilder, EUpdateAllP
 			FPrimitiveSceneInfo* PrimitiveSceneInfo = PrimitiveSceneProxy->GetPrimitiveSceneInfo();
 			const bool bUpdateStaticDrawLists = !PrimitiveSceneProxy->StaticElementsAlwaysUseProxyPrimitiveUniformBuffer();
 
-			if (!PrimitiveSceneInfo->bPendingAddToScene)
+			if (QueueAddToScene(PrimitiveSceneInfo))
 			{
-				PrimitiveSceneInfo->bPendingAddToScene = true;
+				// Remove the primitive from the scene at its old location
+				// (note that the octree update relies on the bounds not being modified yet).
+				PrimitiveSceneInfo->RemoveFromScene(bUpdateStaticDrawLists);
 
 				if (bUpdateStaticDrawLists)
 				{
-					SceneInfosWithStaticDrawListUpdate.Push(PrimitiveSceneInfo);
+					QueueAddStaticMeshes(PrimitiveSceneInfo);
 				}
-				else
-				{
-					SceneInfosWithoutStaticDrawListUpdate.Push(PrimitiveSceneInfo);
-				}
-
-				// Remove the primitive from the scene at its old location
-				// (note that the octree update relies on the bounds not being modified yet).
-				PrimitiveSceneInfo->RemoveFromScene(bUpdateStaticDrawLists);
 			}
 
-			QueueSceneInfoToFlushVirtualTexture(PrimitiveSceneInfo);
+			QueueFlushVirtualTexture(PrimitiveSceneInfo);
 
 			if (ShouldPrimitiveOutputVelocity(PrimitiveSceneInfo->Proxy, GetShaderPlatform()))
 			{
@@ -6099,7 +6079,7 @@ void FScene::UpdateAllPrimitiveSceneInfos(FRDGBuilder& GraphBuilder, EUpdateAllP
 			FScopeCycleCounter Context(PrimitiveSceneProxy->GetStatId());
 			FPrimitiveSceneInfo* PrimitiveSceneInfo = PrimitiveSceneProxy->GetPrimitiveSceneInfo();
 			
-			QueueSceneInfoToFlushVirtualTexture(PrimitiveSceneInfo);
+			QueueFlushVirtualTexture(PrimitiveSceneInfo);
 
 			// If we recorded no adds or removes the instance count has stayed the same.  Therefore the cached mesh draw commands do not
 			// need to be updated.  In situations where the instance count changes the mesh draw command stores the instance count which
@@ -6108,13 +6088,13 @@ void FScene::UpdateAllPrimitiveSceneInfos(FRDGBuilder& GraphBuilder, EUpdateAllP
 
 			const bool bUpdateStaticDrawLists = !PrimitiveSceneProxy->StaticElementsAlwaysUseProxyPrimitiveUniformBuffer() || bInstanceCountChanged;
 
-			if (!PrimitiveSceneInfo->bPendingAddToScene)
+			if (QueueAddToScene(PrimitiveSceneInfo))
 			{
-				PrimitiveSceneInfo->bPendingAddToScene = true;
+				PrimitiveSceneInfo->RemoveFromScene(bUpdateStaticDrawLists);
 
 				if (bUpdateStaticDrawLists)
 				{
-					SceneInfosWithStaticDrawListUpdate.Push(PrimitiveSceneInfo);
+					QueueAddStaticMeshes(PrimitiveSceneInfo);
 				}
 				else
 				{
@@ -6122,10 +6102,7 @@ void FScene::UpdateAllPrimitiveSceneInfos(FRDGBuilder& GraphBuilder, EUpdateAllP
 					RayTracingPrimitivesToUpdate.Add(PrimitiveSceneInfo);
 					bUpdateCachedRayTracingInstances = true;
 #endif
-					SceneInfosWithoutStaticDrawListUpdate.Push(PrimitiveSceneInfo);
 				}
-
-				PrimitiveSceneInfo->RemoveFromScene(bUpdateStaticDrawLists);
 			}
 
 			// Update the Proxy's data.
@@ -6166,72 +6143,104 @@ void FScene::UpdateAllPrimitiveSceneInfos(FRDGBuilder& GraphBuilder, EUpdateAllP
 #endif
 	}
 
-	if (SceneInfosWithoutStaticDrawListUpdate.Num() > 0)
+	if (SceneInfosWithAddToScene.Num() > 0)
 	{
-		const EPrimitiveAddToSceneOps AddToSceneOps = EnumHasAnyFlags(AsyncOps, EUpdateAllPrimitiveSceneInfosAsyncOps::CreateLightPrimitiveInteractions)
-			? EPrimitiveAddToSceneOps::None
-			: EPrimitiveAddToSceneOps::CreateLightPrimitiveInteractions;
-
-		FPrimitiveSceneInfo::AddToScene(this, SceneInfosWithoutStaticDrawListUpdate, AddToSceneOps);
+		FPrimitiveSceneInfo::AddToScene(this, SceneInfosWithAddToScene);
+	}
 
 #if RHI_RAYTRACING
-		if (bUpdateCachedRayTracingInstances)
-		{
-			FPrimitiveSceneInfo::UpdateCachedRayTracingInstances(this, RayTracingPrimitivesToUpdate);
-		}
-#endif
+	if (bUpdateCachedRayTracingInstances)
+	{
+		FPrimitiveSceneInfo::UpdateCachedRayTracingInstances(this, RayTracingPrimitivesToUpdate);
 	}
+#endif
 
-	// Re-add the primitive to the scene with the new transform.
-	if (SceneInfosWithStaticDrawListUpdate.Num() > 0)
+	for (const FLevelCommand& Command : LevelCommands)
 	{
-		EPrimitiveAddToSceneOps AddToSceneOps = EPrimitiveAddToSceneOps::All;
-
-		if (EnumHasAnyFlags(AsyncOps, EUpdateAllPrimitiveSceneInfosAsyncOps::CreateLightPrimitiveInteractions))
+		const auto InvalidateNaniteGPUScene = [](FPrimitiveSceneInfo* Primitive)
 		{
-			EnumRemoveFlags(AddToSceneOps, EPrimitiveAddToSceneOps::CreateLightPrimitiveInteractions);
-		}
+			// Invalidate primitive proxy entry in GPU Scene. This is necessary for Nanite::FSceneProxy to be uploaded to GPU scene (see GetPrimitiveID in GPUScene.cpp)
+			if (Primitive->Proxy->IsNaniteMesh())
+			{
+				Primitive->RequestGPUSceneUpdate();
+			}
+		};
 
-		if (EnumHasAnyFlags(AsyncOps, EUpdateAllPrimitiveSceneInfosAsyncOps::CacheMeshDrawCommands))
+		if (Command.Op == FLevelCommand::EOp::Add)
 		{
-			EnumRemoveFlags(AddToSceneOps, EPrimitiveAddToSceneOps::CacheMeshDrawCommands);
-		}
-
-		FPrimitiveSceneInfo::AddToScene(this, SceneInfosWithStaticDrawListUpdate, AddToSceneOps);
+			if (const TArray<FPrimitiveSceneInfo*>* LevelPrimitives = PrimitivesNeedingLevelUpdateNotification.Find(Command.Name))
+			{
+				for (FPrimitiveSceneInfo* Primitive : *LevelPrimitives)
+				{
+					// If the primitive proxy returns true, it needs it's static meshes added to the scene
+					if (Primitive->Proxy->OnLevelAddedToWorld_RenderThread() && QueueAddStaticMeshes(Primitive))
+					{
+						Primitive->RemoveStaticMeshes();
+					}
 
-		if (EnumHasAnyFlags(AsyncOps, EUpdateAllPrimitiveSceneInfosAsyncOps::CacheMeshDrawCommands))
+					InvalidateNaniteGPUScene(Primitive);
+				}
+			}
+		}
+		else
 		{
-			CacheMeshDrawCommandsTask = GraphBuilder.AddSetupTask([this, &SceneInfosWithStaticDrawListUpdate] () mutable
+			if (TArray<FPrimitiveSceneInfo*>* LevelPrimitives = PrimitivesNeedingLevelUpdateNotification.Find(Command.Name))
 			{
-				SCOPED_NAMED_EVENT(AsyncCacheMeshDrawCommands, FColor::Emerald);
-				FPrimitiveSceneInfo::CacheMeshDrawCommands(this, SceneInfosWithStaticDrawListUpdate);
-				FPrimitiveSceneInfo::CacheNaniteDrawCommands(this, SceneInfosWithStaticDrawListUpdate);
-			#if RHI_RAYTRACING
-				FPrimitiveSceneInfo::CacheRayTracingPrimitives(this, SceneInfosWithStaticDrawListUpdate);
-			#endif
-			});
+				for (FPrimitiveSceneInfo* Primitive : *LevelPrimitives)
+				{
+					Primitive->Proxy->OnLevelRemovedFromWorld_RenderThread();
+					InvalidateNaniteGPUScene(Primitive);
+				}
+			}
 		}
 	}
 
-	if (EnumHasAnyFlags(AsyncOps, EUpdateAllPrimitiveSceneInfosAsyncOps::CreateLightPrimitiveInteractions))
+	if (SceneInfosWithStaticDrawListUpdate.Num() > 0)
 	{
-		CreateLightPrimitiveInteractionsTask = GraphBuilder.AddSetupTask(
-			[this, &SceneInfosWithStaticDrawListUpdate, SceneInfosWithoutStaticDrawListUpdate = MoveTemp(SceneInfosWithoutStaticDrawListUpdate)]
+		FPrimitiveSceneInfo::AddStaticMeshes(this, SceneInfosWithStaticDrawListUpdate, false);
+	}
+
+	FPrimitiveSceneInfo::UpdateVirtualTextures(this, SceneInfosWithAddToScene);
+
+	if (SceneInfosWithStaticDrawListUpdate.Num() > 0)
+	{
+		const bool bLaunchAsyncTask = EnumHasAnyFlags(AsyncOps, EUpdateAllPrimitiveSceneInfosAsyncOps::CacheMeshDrawCommands);
+
+		CacheMeshDrawCommandsTask = LaunchSceneRenderTask(TEXT("CacheMeshDrawCommandsTask"), [this, &SceneInfosWithStaticDrawListUpdate, bLaunchAsyncTask]()
 		{
-			SCOPED_NAMED_EVENT(AsyncCreateLightPrimitiveInteractions, FColor::Emerald);
+			SCOPED_NAMED_EVENT(CacheMeshDrawCommands, FColor::Emerald);
 
-			for (FPrimitiveSceneInfo* SceneInfo : SceneInfosWithStaticDrawListUpdate)
+			UE::Tasks::AddNested(LaunchSceneRenderTask(TEXT("CacheMeshDrawCommandsTask"), [this, &SceneInfosWithStaticDrawListUpdate, bLaunchAsyncTask]()
 			{
-				CreateLightPrimitiveInteractionsForPrimitive(SceneInfo);
-			}
+				FPrimitiveSceneInfo::CacheMeshDrawCommands(this, SceneInfosWithStaticDrawListUpdate);
+			}, bLaunchAsyncTask));
 
-			for (FPrimitiveSceneInfo* SceneInfo : SceneInfosWithoutStaticDrawListUpdate)
+			UE::Tasks::AddNested(LaunchSceneRenderTask(TEXT("CacheNaniteDrawCommands"), [this, &SceneInfosWithStaticDrawListUpdate, bLaunchAsyncTask]()
 			{
-				CreateLightPrimitiveInteractionsForPrimitive(SceneInfo);
-			}
-		});
+				FPrimitiveSceneInfo::CacheNaniteDrawCommands(this, SceneInfosWithStaticDrawListUpdate);
+			}, bLaunchAsyncTask));
+
+		}, bLaunchAsyncTask);
+
+#if RHI_RAYTRACING
+		CacheRayTracingPrimitivesTask = LaunchSceneRenderTask(TEXT("CacheRayTracingPrimitivesTask"), [this, &SceneInfosWithStaticDrawListUpdate]()
+		{
+			FPrimitiveSceneInfo::CacheRayTracingPrimitives(this, SceneInfosWithStaticDrawListUpdate);
+		}, bLaunchAsyncTask);
+#endif
 	}
 
+	CreateLightPrimitiveInteractionsTask = LaunchSceneRenderTask(TEXT("CreateLightPrimitiveInteractionsTask"), [this, &SceneInfosWithAddToScene]
+	{
+		SCOPED_NAMED_EVENT(CreateLightPrimitiveInteractions, FColor::Emerald);
+
+		for (FPrimitiveSceneInfo* SceneInfo : SceneInfosWithAddToScene)
+		{
+			CreateLightPrimitiveInteractionsForPrimitive(SceneInfo);
+		}
+
+	}, EnumHasAnyFlags(AsyncOps, EUpdateAllPrimitiveSceneInfosAsyncOps::CreateLightPrimitiveInteractions));
+
 	for (FPrimitiveSceneInfo* PrimitiveSceneInfo : AddedPrimitiveSceneInfos)
 	{
 		// Set LOD parent information if valid
@@ -6241,7 +6250,7 @@ void FScene::UpdateAllPrimitiveSceneInfos(FRDGBuilder& GraphBuilder, EUpdateAllP
 		SceneLODHierarchy.UpdateNodeSceneInfo(PrimitiveSceneInfo->PrimitiveComponentId, PrimitiveSceneInfo);
 	}
 
-	for (FPrimitiveSceneInfo* PrimitiveSceneInfo : SceneInfosToFlushVirtualTexture)
+	for (FPrimitiveSceneInfo* PrimitiveSceneInfo : SceneInfosWithFlushVirtualTexture)
 	{
 		PrimitiveSceneInfo->FlushRuntimeVirtualTexture();
 		PrimitiveSceneInfo->bPendingFlushVirtualTexture = false;
@@ -6341,6 +6350,9 @@ void FScene::UpdateAllPrimitiveSceneInfos(FRDGBuilder& GraphBuilder, EUpdateAllP
 	UpdatedOcclusionBoundsSlacks.Empty();
 	DistanceFieldSceneDataUpdates.Empty();
 	AddedPrimitiveSceneInfos.Empty();
+	LevelCommands.Empty();
+
+	LastUpdateFrameCounter = GFrameCounterRenderThread;
 
 #if DO_GUARD_SLOW
 	for (const FPrimitiveSceneInfo* PrimitiveSceneInfo : Primitives)
diff --git a/Engine/Source/Runtime/Renderer/Private/ScenePrivate.h b/Engine/Source/Runtime/Renderer/Private/ScenePrivate.h
--- a/Engine/Source/Runtime/Renderer/Private/ScenePrivate.h
+++ b/Engine/Source/Runtime/Renderer/Private/ScenePrivate.h
@@ -3370,6 +3370,15 @@ public:
 		VelocityData.StartFrame(this);
 	}
 
+	virtual void EndFrame(FRHICommandListImmediate& RHICmdList) override
+	{
+		// Performs a final update of any queued scene primitives in the case where the scene wasn't rendered to avoid a build-up of queued data that is never flushed.
+		if (LastUpdateFrameCounter != GFrameCounterRenderThread)
+		{
+			UpdateAllPrimitiveSceneInfos(RHICmdList);
+		}
+	}
+
 	/**
 	 * Returns the current "FrameNumber" where frame corresponds to how many times FRendererModule::BeginRenderingViewFamilies has been called.
 	 * Thread safe, and returns a different copy for game/render thread. GetFrameNumberRenderThread can only be called from the render thread. 
@@ -3450,8 +3459,19 @@ public:
 	void WaitForCacheMeshDrawCommandsTask()
 	{
 		CacheMeshDrawCommandsTask.Wait();
+
+#if RHI_RAYTRACING
+		CacheRayTracingPrimitivesTask.Wait();
+#endif
 	}
 
+#if RHI_RAYTRACING
+	void WaitForCacheRayTracingPrimitivesTask()
+	{
+		CacheRayTracingPrimitivesTask.Wait();
+	}
+#endif
+
 	void LumenAddPrimitive(FPrimitiveSceneInfo* InPrimitive);
 	void LumenUpdatePrimitive(FPrimitiveSceneInfo* InPrimitive);
 	void LumenInvalidateSurfaceCacheForPrimitive(FPrimitiveSceneInfo* InPrimitive);
@@ -3569,20 +3589,6 @@ private:
 	 */
 	void ApplyWorldOffset_RenderThread(const FVector& InOffset);
 
-	/**
-	 * Notification from game thread that level was added to a world
-	 *
-	 * @param	InLevelName		Level name
-	 */
-	void OnLevelAddedToWorld_RenderThread(const FName& InLevelName);
-
-	/**
-	 * Notification from game thread that level was removed from a world
-	 *
-	 * @param	InLevelName		Level name
-	 */
-	void OnLevelRemovedFromWorld_RenderThread(const FName& InLevelName);
-
 	void ProcessAtmosphereLightRemoval_RenderThread(FLightSceneInfo* LightSceneInfo);
 	void ProcessAtmosphereLightAddition_RenderThread(FLightSceneInfo* LightSceneInfo);
 
@@ -3616,6 +3622,18 @@ private:
 		FBoxSphereBounds StaticMeshBounds;
 	};
 
+	struct FLevelCommand
+	{
+		enum class EOp
+		{
+			Add,
+			Remove
+		};
+
+		FName Name;
+		EOp Op;
+	};
+
 	Experimental::TRobinHoodHashMap<FPrimitiveSceneInfo*, FPrimitiveComponentId> UpdatedAttachmentRoots;
 	Experimental::TRobinHoodHashMap<FPrimitiveSceneProxy*, FCustomPrimitiveData> UpdatedCustomPrimitiveParams;
 	Experimental::TRobinHoodHashMap<FPrimitiveSceneProxy*, FUpdateTransformCommand> UpdatedTransforms;
@@ -3625,9 +3643,13 @@ private:
 	Experimental::TRobinHoodHashSet<FPrimitiveSceneInfo*> AddedPrimitiveSceneInfos;
 	Experimental::TRobinHoodHashSet<FPrimitiveSceneInfo*> RemovedPrimitiveSceneInfos;
 	Experimental::TRobinHoodHashSet<FPrimitiveSceneInfo*> DistanceFieldSceneDataUpdates;
+	TArray<FLevelCommand> LevelCommands;
 
 	UE::Tasks::FTask CreateLightPrimitiveInteractionsTask;
 	UE::Tasks::FTask CacheMeshDrawCommandsTask;
+#if RHI_RAYTRACING
+	UE::Tasks::FTask CacheRayTracingPrimitivesTask;
+#endif
 
 	FSceneLightInfoUpdates *SceneLightInfoUpdates;
 
@@ -3647,6 +3669,8 @@ private:
 	uint32 SceneFrameNumber;
 	uint32 SceneFrameNumberRenderThread;
 
+	uint32 LastUpdateFrameCounter = UINT32_MAX;
+
 	/** Whether world settings has bForceNoPrecomputedLighting set */
 	bool bForceNoPrecomputedLighting;
 
diff --git a/Engine/Source/Runtime/Renderer/Private/SceneVisibility.cpp b/Engine/Source/Runtime/Renderer/Private/SceneVisibility.cpp
--- a/Engine/Source/Runtime/Renderer/Private/SceneVisibility.cpp
+++ b/Engine/Source/Runtime/Renderer/Private/SceneVisibility.cpp
@@ -917,7 +917,7 @@ static void PrimitiveCullTask(FThreadSafeCounter& NumCulledPrimitives, const FSc
 	STAT(NumCulledPrimitives.Add(NumPrimitivesCulledForTask));
 }
 
-static int32 PrimitiveCull(const FScene* RESTRICT Scene, FViewInfo& View, bool bShouldVisibilityCull)
+static int32 PrimitiveCull(FScene* RESTRICT Scene, FViewInfo& View, bool bShouldVisibilityCull)
 {
 	FPrimitiveCullingFlags Flags;
 	Flags.bShouldVisibilityCull = bShouldVisibilityCull;
@@ -931,6 +931,9 @@ static int32 PrimitiveCull(const FScene* RESTRICT Scene, FViewInfo& View, bool b
 
 #if RHI_RAYTRACING
 	View.RayTracingCullingParameters.Init(View);
+
+	// Sync cached raytracing tasks ShouldCullForRayTracing.
+	Scene->WaitForCacheRayTracingPrimitivesTask();
 #endif
 
 	SCOPE_CYCLE_COUNTER(STAT_PrimitiveCull);
@@ -4793,14 +4796,6 @@ void FSceneRenderer::ComputeViewVisibility(
 
 			{
 				TRACE_CPUPROFILER_EVENT_SCOPE(FSceneRenderer_Cull);
-#if RHI_RAYTRACING
-				if (bAnyRayTracingPassEnabled)
-				{
-					// The logic inside PrimitiveCull makes use of ShouldCullForRayTracing to decide if the primitive should be considered for raytracing
-					// Therefore we must be sure we are done with caching the mesh draw commands (which include raytracing caches) before it runs if raytracing is being used.
-					Scene->WaitForCacheMeshDrawCommandsTask();
-				}
-#endif
 				int32 NumCulledPrimitivesForView = PrimitiveCull(Scene, View, bNeedsFrustumCulling);
 				STAT(NumCulledPrimitives += NumCulledPrimitivesForView);
 			}
diff --git a/Engine/Source/Runtime/Renderer/Public/PrimitiveSceneInfo.h b/Engine/Source/Runtime/Renderer/Public/PrimitiveSceneInfo.h
--- a/Engine/Source/Runtime/Renderer/Public/PrimitiveSceneInfo.h
+++ b/Engine/Source/Runtime/Renderer/Public/PrimitiveSceneInfo.h
@@ -269,16 +269,6 @@ struct FPersistentPrimitiveIndex
 	int32 Index;
 };
 
-enum class EPrimitiveAddToSceneOps
-{
-	None = 0,
-	AddStaticMeshes = 1 << 0,
-	CacheMeshDrawCommands = 1 << 1,
-	CreateLightPrimitiveInteractions = 1 << 2,
-	All = AddStaticMeshes | CacheMeshDrawCommands | CreateLightPrimitiveInteractions
-};
-ENUM_CLASS_FLAGS(EPrimitiveAddToSceneOps);
-
 /**
  * The renderer's internal state for a single UPrimitiveComponent.  This has a one to one mapping with FPrimitiveSceneProxy, which is in the engine module.
  */
@@ -438,7 +428,7 @@ public:
 	~FPrimitiveSceneInfo();
 
 	/** Adds the primitive to the scene. */
-	static void AddToScene(FScene* Scene, TArrayView<FPrimitiveSceneInfo*> SceneInfos, EPrimitiveAddToSceneOps Ops = EPrimitiveAddToSceneOps::All);
+	static void AddToScene(FScene* Scene, TArrayView<FPrimitiveSceneInfo*> SceneInfos);
 
 	/** Removes the primitive from the scene. */
 	void RemoveFromScene(bool bUpdateStaticDrawLists);
@@ -686,7 +676,10 @@ private:
 
 	/** True if the primitive is queued for add. */
 	bool bPendingAddToScene : 1;
-	
+
+	/** True if the primitive is queued to have static meshes built. */
+	bool bPendingAddStaticMeshes : 1;
+
 	/** True if the primitive is queued to have its virtual texture flushed. */
 	bool bPendingFlushVirtualTexture : 1;
 
@@ -721,6 +714,9 @@ private:
 	/** Creates cached mesh draw commands for all meshes. */
 	static void CacheMeshDrawCommands(FScene* Scene, TArrayView<FPrimitiveSceneInfo*> SceneInfos);
 
+	/** Updates virtual texture state on the scene primitives. */
+	static void UpdateVirtualTextures(FScene* Scene, TArrayView<FPrimitiveSceneInfo*> SceneInfos);
+
 	/** Removes cached mesh draw commands for all meshes. */
 	void RemoveCachedMeshDrawCommands();
 
