Changing uses of TStaticArray<FTextureRenderTargetBinding,MaxSimultaneousRenderTargets>& to TArrayView<FTextureRenderTargetBinding> so we can forward declare FTextureRenderTargetBinding.

diff --git a/Engine/Source/Runtime/Renderer/Private/SceneTextures.cpp b/Engine/Source/Runtime/Renderer/Private/SceneTextures.cpp
--- a/Engine/Source/Runtime/Renderer/Private/SceneTextures.cpp
+++ b/Engine/Source/Runtime/Renderer/Private/SceneTextures.cpp
@@ -630,8 +630,7 @@ void FSceneTextures::InitializeViewFamily(FRDGBuilder& GraphBuilder, FViewFamily
 }
 
 uint32 FSceneTextures::GetGBufferRenderTargets(
-	TStaticArray<FTextureRenderTargetBinding,
-	MaxSimultaneousRenderTargets>& RenderTargets,
+	TArrayView<FTextureRenderTargetBinding> RenderTargets,
 	EGBufferLayout Layout) const
 {
 	uint32 RenderTargetCount = 0;
diff --git a/Engine/Source/Runtime/Renderer/Private/SceneTextures.h b/Engine/Source/Runtime/Renderer/Private/SceneTextures.h
--- a/Engine/Source/Runtime/Renderer/Private/SceneTextures.h
+++ b/Engine/Source/Runtime/Renderer/Private/SceneTextures.h
@@ -55,8 +55,7 @@ struct RENDERER_API FSceneTextures : public FMinimalSceneTextures
 
 	// Configures an array of render targets for the GBuffer pass.
 	uint32 GetGBufferRenderTargets(
-		TStaticArray<FTextureRenderTargetBinding,
-		MaxSimultaneousRenderTargets>& RenderTargets,
+		TArrayView<FTextureRenderTargetBinding> RenderTargets,
 		EGBufferLayout Layout = GBL_Default) const;
 	uint32 GetGBufferRenderTargets(
 		ERenderTargetLoadAction LoadAction,
diff --git a/Engine/Source/Runtime/Renderer/Private/Strata/Strata.cpp b/Engine/Source/Runtime/Renderer/Private/Strata/Strata.cpp
--- a/Engine/Source/Runtime/Renderer/Private/Strata/Strata.cpp
+++ b/Engine/Source/Runtime/Renderer/Private/Strata/Strata.cpp
@@ -1097,7 +1097,7 @@ void AddStrataStencilPass(
 
 ////////////////////////////////////////////////////////////////////////////////////////////////////////////////////////////////////////////////
 
-void AppendStrataMRTs(const FSceneRenderer& SceneRenderer, uint32& RenderTargetCount, TStaticArray<FTextureRenderTargetBinding, MaxSimultaneousRenderTargets>& RenderTargets)
+void AppendStrataMRTs(const FSceneRenderer& SceneRenderer, uint32& RenderTargetCount, TArrayView<FTextureRenderTargetBinding> RenderTargets)
 {
 	if (Strata::IsStrataEnabled() && SceneRenderer.Scene)
 	{
diff --git a/Engine/Source/Runtime/Renderer/Private/Strata/Strata.h b/Engine/Source/Runtime/Renderer/Private/Strata/Strata.h
--- a/Engine/Source/Runtime/Renderer/Private/Strata/Strata.h
+++ b/Engine/Source/Runtime/Renderer/Private/Strata/Strata.h
@@ -2,24 +2,23 @@
 
 #pragma once
 
-#include "CoreMinimal.h"
 #include "ShaderParameterMacros.h"
 #include "RenderGraphResources.h"
 #include "MeshPassProcessor.h"
 #include "UnifiedBuffer.h"
 #include "RHIUtilities.h"
 #include "StrataDefinitions.h"
-#include "RenderGraphUtils.h"
 #include "GBufferInfo.h"
-#include "RendererUtils.h"
+#include "ShaderParameterStruct.h"
 
 // Forward declarations.
 class FScene;
 class FSceneRenderer;
 class FRDGBuilder;
+struct FDBufferTextures;
 struct FMinimalSceneTextures;
 struct FScreenPassTexture;
-struct FDBufferTextures;
+struct FTextureRenderTargetBinding;
 
 BEGIN_SHADER_PARAMETER_STRUCT(FStrataBasePassUniformParameters, )
 	SHADER_PARAMETER(uint32, MaxBytesPerPixel)
@@ -172,7 +171,7 @@ void BindStrataForwardPasslUniformParameters(FRDGBuilder& GraphBuilder, const FV
 void BindStrataMobileForwardPasslUniformParameters(FRDGBuilder& GraphBuilder, const FViewInfo& View, FStrataMobileForwardPassUniformParameters& OutStrataUniformParameters);
 TRDGUniformBufferRef<FStrataGlobalUniformParameters> BindStrataGlobalUniformParameters(const FViewInfo& View);
 
-void AppendStrataMRTs(const FSceneRenderer& SceneRenderer, uint32& BasePassTextureCount, TStaticArray<FTextureRenderTargetBinding, MaxSimultaneousRenderTargets>& BasePassTextures);
+void AppendStrataMRTs(const FSceneRenderer& SceneRenderer, uint32& BasePassTextureCount, TArrayView<FTextureRenderTargetBinding> BasePassTextures);
 void SetBasePassRenderTargetOutputFormat(const EShaderPlatform Platform, const FMaterialShaderParameters& MaterialParameters, FShaderCompilerEnvironment& OutEnvironment, EGBufferLayout GBufferLayout);
 
 TRDGUniformBufferRef<FStrataPublicGlobalUniformParameters> CreatePublicGlobalUniformBuffer(FRDGBuilder& GraphBuilder, FStrataSceneData* StrataScene);
