Fixing bad merge.

diff --git a/Engine/Source/Runtime/Renderer/Private/Strata/Strata.h b/Engine/Source/Runtime/Renderer/Private/Strata/Strata.h
--- a/Engine/Source/Runtime/Renderer/Private/Strata/Strata.h
+++ b/Engine/Source/Runtime/Renderer/Private/Strata/Strata.h
@@ -9,6 +9,7 @@
 #include "RHIUtilities.h"
 #include "StrataDefinitions.h"
 #include "GBufferInfo.h"
+#include "RendererUtils.h"
 #include "ShaderParameterStruct.h"
 
 // Forward declarations.
