Making PostProcess/PostProcessHMD.h Public since its types are used by other modules

diff --git a/Engine/Source/Runtime/Renderer/Private/PostProcess/PostProcessHMD.h b/Engine/Source/Runtime/Renderer/Private/PostProcess/PostProcessHMD.h
deleted file mode 100644
diff --git a/Engine/Source/Runtime/Renderer/Public/PostProcess/PostProcessHMD.h b/Engine/Source/Runtime/Renderer/Public/PostProcess/PostProcessHMD.h
new file mode 100644
--- /dev/null
+++ b/Engine/Source/Runtime/Renderer/Public/PostProcess/PostProcessHMD.h
@@ -0,0 +1,29 @@
+// Copyright Epic Games, Inc. All Rights Reserved.
+
+#pragma once
+
+#include "ScreenPass.h"
+
+/** The vertex data used to filter a texture. */
+struct FDistortionVertex
+{
+	FVector2f	Position;
+	FVector2f	TexR;
+	FVector2f	TexG;
+	FVector2f	TexB;
+	float		VignetteFactor;
+	float		TimewarpFactor;
+};
+
+struct FHMDDistortionInputs
+{
+	// [Optional] Render to the specified output. If invalid, a new texture is created and returned.
+	FScreenPassRenderTarget OverrideOutput;
+
+	// [Required] The input scene color and view rect.
+	FScreenPassTexture SceneColor;
+};
+
+FScreenPassTexture AddDefaultHMDDistortionPass(FRDGBuilder& GraphBuilder, const FViewInfo& View, const FHMDDistortionInputs& Inputs);
+
+FScreenPassTexture AddHMDDistortionPass(FRDGBuilder& GraphBuilder, const FViewInfo& View, const FHMDDistortionInputs& Inputs);
\ No newline at end of file
