Moving SceneCaptureRendering.h to Renderer/Public

diff --git a/Engine/Source/Runtime/Renderer/Private/SceneCaptureRendering.h b/Engine/Source/Runtime/Renderer/Private/SceneCaptureRendering.h
deleted file mode 100644
diff --git a/Engine/Source/Runtime/Renderer/Public/SceneCaptureRendering.h b/Engine/Source/Runtime/Renderer/Public/SceneCaptureRendering.h
new file mode 100644
--- /dev/null
+++ b/Engine/Source/Runtime/Renderer/Public/SceneCaptureRendering.h
@@ -0,0 +1,18 @@
+// Copyright Epic Games, Inc. All Rights Reserved.
+
+/*=============================================================================
+	SceneCaptureRendering.h: SceneCaptureRendering definitions.
+=============================================================================*/
+
+#pragma once
+
+#include "CoreMinimal.h"
+#include "RenderGraphDefinitions.h"
+#include "SceneTextures.h"
+
+RENDERER_API void CopySceneCaptureComponentToTarget(
+	FRDGBuilder& GraphBuilder,
+	const FMinimalSceneTextures& SceneTextures,
+	FRDGTextureRef ViewFamilyTexture,
+	const FSceneViewFamily& ViewFamily,
+	const TArray<FViewInfo>& Views);
