Add support for allocating VSM page pool as a reserved resource to reduce performance degradation when video memory is over-committed

* Add CVar r.Shadow.Virtual.AllocatePagePoolAsReservedResource to control the behavior (default: off)

Reserved resource path will internally split the physical memory allocation into N small chunks, making it easier for Windows video memory manager to retain more of the allocations in local video memory.

diff --git a/Engine/Source/Runtime/Renderer/Private/VirtualShadowMaps/VirtualShadowMapCacheManager.cpp b/Engine/Source/Runtime/Renderer/Private/VirtualShadowMaps/VirtualShadowMapCacheManager.cpp
--- a/Engine/Source/Runtime/Renderer/Private/VirtualShadowMaps/VirtualShadowMapCacheManager.cpp
+++ b/Engine/Source/Runtime/Renderer/Private/VirtualShadowMaps/VirtualShadowMapCacheManager.cpp
@@ -71,6 +71,13 @@ static FAutoConsoleVariableRef  CVarForceInvalidateDirectionalVSM(
 	ECVF_RenderThreadSafe);
 
 
+static TAutoConsoleVariable<int32> CVarVSMReservedResource(
+	TEXT("r.Shadow.Virtual.AllocatePagePoolAsReservedResource"),
+	0,
+	TEXT("Allocate VSM page pool as a reserved/virtual texture, backed by N small physical memory allocations to reduce fragmentation."),
+	ECVF_RenderThreadSafe
+);
+
 void FVirtualShadowMapCacheEntry::UpdateClipmap(
 	int32 VirtualShadowMapId,
 	FInt64Point PageSpaceLocation,
@@ -606,13 +613,22 @@ FVirtualShadowMapArrayCacheManager::~FVirtualShadowMapArrayCacheManager()
 
 TRefCountPtr<IPooledRenderTarget> FVirtualShadowMapArrayCacheManager::SetPhysicalPoolSize(FRDGBuilder& GraphBuilder, FIntPoint RequestedSize, int RequestedArraySize)
 {
-	if (!PhysicalPagePool || PhysicalPagePool->GetDesc().Extent != RequestedSize || PhysicalPagePool->GetDesc().ArraySize != RequestedArraySize)
+	// Using ReservedResource|ImmediateCommit flags hint to the RHI that the resource can be allocated using N small physical memory allocations,
+	// instead of a single large contighous allocation. This helps Windows video memory manager page allocations in and out of local memory more efficiently.
+	ETextureCreateFlags TextureCreateFlags = CVarVSMReservedResource.GetValueOnRenderThread() && GRHISupportsReservedResources
+		? (TexCreate_ReservedResource | TexCreate_ImmediateCommit)
+		: TexCreate_None;
+
+	if (!PhysicalPagePool 
+		|| PhysicalPagePool->GetDesc().Extent != RequestedSize 
+		|| PhysicalPagePool->GetDesc().ArraySize != RequestedArraySize
+		|| PhysicalPagePool->GetDesc().Flags != TextureCreateFlags)
 	{
 		FPooledRenderTargetDesc Desc2D = FPooledRenderTargetDesc::Create2DArrayDesc(
 			RequestedSize,
 			PF_R32_UINT,
 			FClearValueBinding::None,
-			TexCreate_None,
+			TextureCreateFlags,
 			TexCreate_ShaderResource | TexCreate_UAV,
 			false,
 			RequestedArraySize
