Quad Complexity working for mobile preview

diff --git a/Engine/Shaders/Private/QuadOverdraw.ush b/Engine/Shaders/Private/QuadOverdraw.ush
--- a/Engine/Shaders/Private/QuadOverdraw.ush
+++ b/Engine/Shaders/Private/QuadOverdraw.ush
@@ -45,7 +45,11 @@ float ComplexityToFloat(uint Complexity)
 	return (float)Complexity / (float)QO_FRACTIONNAL_COMPLEXITY;
 }
 
+#if SHADING_PATH_MOBILE
+#define RWQuadBuffer MobileBasePass.QuadOverdraw
+#else
 #define RWQuadBuffer DebugViewModePass.QuadOverdraw
+#endif
 
 #if OUTPUT_QUAD_OVERDRAW
 
diff --git a/Engine/Source/Runtime/Engine/Private/DebugViewModeHelpers.cpp b/Engine/Source/Runtime/Engine/Private/DebugViewModeHelpers.cpp
--- a/Engine/Source/Runtime/Engine/Private/DebugViewModeHelpers.cpp
+++ b/Engine/Source/Runtime/Engine/Private/DebugViewModeHelpers.cpp
@@ -87,7 +87,7 @@ bool AllowDebugViewShaderMode(EDebugViewShaderMode ShaderMode, EShaderPlatform P
 	case DVSM_ShaderComplexityContainedQuadOverhead:
 	case DVSM_ShaderComplexityBleedingQuadOverhead:
 	case DVSM_QuadComplexity:
-		return FeatureLevel >= ERHIFeatureLevel::SM5 && (bForceQuadOverdraw || (PlatformSupportsDebugViewShaders(Platform) && !IsMetalPlatform(Platform))); // Last one to fix for Metal then remove this Metal check.
+		return (bForceQuadOverdraw || (PlatformSupportsDebugViewShaders(Platform) && !IsMetalPlatform(Platform))); // Last one to fix for Metal then remove this Metal check.
 	case DVSM_PrimitiveDistanceAccuracy:
 	case DVSM_MeshUVDensityAccuracy:
 		return FeatureLevel >= ERHIFeatureLevel::SM5 && (bForceStreamingAccuracy || PlatformSupportsDebugViewShaders(Platform));
diff --git a/Engine/Source/Runtime/Renderer/Private/DebugViewModeRendering.cpp b/Engine/Source/Runtime/Renderer/Private/DebugViewModeRendering.cpp
--- a/Engine/Source/Runtime/Renderer/Private/DebugViewModeRendering.cpp
+++ b/Engine/Source/Runtime/Renderer/Private/DebugViewModeRendering.cpp
@@ -272,8 +272,7 @@ void FDebugViewModePS::ModifyCompilationEnvironment(const FMaterialShaderPermuta
 	TCHAR BufferRegister[] = { 'u', '0', 0 };
 	BufferRegister[1] += GetQuadOverdrawUAVIndex(Parameters.Platform);
 	OutEnvironment.SetDefine(TEXT("QUAD_BUFFER_REGISTER"), BufferRegister);
-	const bool bUsingMobileRenderer = FSceneInterface::GetShadingPath(GetMaxSupportedFeatureLevel(Parameters.Platform)) == EShadingPath::Mobile;
-	OutEnvironment.SetDefine(TEXT("OUTPUT_QUAD_OVERDRAW"), !bUsingMobileRenderer);
+	OutEnvironment.SetDefine(TEXT("OUTPUT_QUAD_OVERDRAW"), !FDataDrivenShaderPlatformInfo::GetIsLanguageMetal(Parameters.Platform));
 
 
 	for (int i = 0; i < DVSM_MAX; ++i)
diff --git a/Engine/Source/Runtime/Renderer/Private/MobileBasePassRendering.cpp b/Engine/Source/Runtime/Renderer/Private/MobileBasePassRendering.cpp
--- a/Engine/Source/Runtime/Renderer/Private/MobileBasePassRendering.cpp
+++ b/Engine/Source/Runtime/Renderer/Private/MobileBasePassRendering.cpp
@@ -258,6 +258,15 @@ TRDGUniformBufferRef<FMobileBasePassUniformParameters> CreateMobileBasePassUnifo
 	}
 #endif
 
+	// QuadOverdraw is a UAV so it needs to be initialized even if not used
+	const FSceneTextures* SceneTextures = View.GetSceneTexturesChecked();
+	FRDGTextureRef QuadOverdrawTexture = SceneTextures->QuadOverdraw;
+	if (!QuadOverdrawTexture)
+	{
+		QuadOverdrawTexture = GraphBuilder.CreateTexture(FRDGTextureDesc::Create2D(FIntPoint(1, 1), PF_R32_UINT, FClearValueBinding::None, TexCreate_UAV), TEXT("DummyOverdrawUAV"));
+	}
+	BasePassParameters->QuadOverdraw = GraphBuilder.CreateUAV(QuadOverdrawTexture);
+	
 	return GraphBuilder.CreateUniformBuffer(BasePassParameters);
 }
 
diff --git a/Engine/Source/Runtime/Renderer/Private/MobileBasePassRendering.h b/Engine/Source/Runtime/Renderer/Private/MobileBasePassRendering.h
--- a/Engine/Source/Runtime/Renderer/Private/MobileBasePassRendering.h
+++ b/Engine/Source/Runtime/Renderer/Private/MobileBasePassRendering.h
@@ -39,6 +39,7 @@ BEGIN_GLOBAL_SHADER_PARAMETER_STRUCT(FMobileBasePassUniformParameters, )
 	SHADER_PARAMETER_STRUCT(FMobileSceneTextureUniformParameters, SceneTextures)
 	SHADER_PARAMETER_STRUCT(FStrataMobileForwardPassUniformParameters, Strata)
 	SHADER_PARAMETER_STRUCT(FDebugViewModeUniformParameters, DebugViewMode)
+	SHADER_PARAMETER_RDG_TEXTURE_UAV(RWTexture2D<uint>, QuadOverdraw)
 	SHADER_PARAMETER_STRUCT(FReflectionUniformParameters, ReflectionsParameters)
 	SHADER_PARAMETER_TEXTURE(Texture2D, PreIntegratedGFTexture)
 	SHADER_PARAMETER_SAMPLER(SamplerState, PreIntegratedGFSampler)
diff --git a/Engine/Source/Runtime/Renderer/Private/MobileShadingRenderer.cpp b/Engine/Source/Runtime/Renderer/Private/MobileShadingRenderer.cpp
--- a/Engine/Source/Runtime/Renderer/Private/MobileShadingRenderer.cpp
+++ b/Engine/Source/Runtime/Renderer/Private/MobileShadingRenderer.cpp
@@ -933,6 +933,13 @@ void FMobileSceneRenderer::Render(FRDGBuilder& GraphBuilder)
 	FSceneTextures::InitializeViewFamily(GraphBuilder, ViewFamily);
 	FSceneTextures& SceneTextures = GetActiveSceneTextures();
 
+#if WITH_DEBUG_VIEW_MODES
+	if (ViewFamily.UseDebugViewPS() && ViewFamily.EngineShowFlags.ShaderComplexity && SceneTextures.QuadOverdraw)
+	{
+		AddClearUAVPass(GraphBuilder, GraphBuilder.CreateUAV(SceneTextures.QuadOverdraw), FUintVector4(0, 0, 0, 0));
+	}
+#endif
+
 	if (bUseVirtualTexturing)
 	{
 		FVirtualTextureSystem::Get().EndUpdate(GraphBuilder, MoveTemp(VirtualTextureUpdater), FeatureLevel);
@@ -1605,6 +1612,7 @@ void FMobileSceneRenderer::RenderDeferredSinglePass(FRDGBuilder& GraphBuilder, c
 		// Opaque and masked
 		RHICmdList.SetCurrentStat(GET_STATID(STAT_CLMM_Opaque));
 		RenderMobileBasePass(RHICmdList, View);
+		RenderMobileDebugView(RHICmdList, View);
 		RHICmdList.PollOcclusionQueries();
 		PostRenderBasePass(RHICmdList, View);
 		// SceneColor + GBuffer write, SceneDepth is read only
@@ -1647,6 +1655,7 @@ void FMobileSceneRenderer::RenderDeferredMultiPass(FRDGBuilder& GraphBuilder, cl
 		// Opaque and masked
 		RHICmdList.SetCurrentStat(GET_STATID(STAT_CLMM_Opaque));
 		RenderMobileBasePass(RHICmdList, View);
+		RenderMobileDebugView(RHICmdList, View);
 		RHICmdList.PollOcclusionQueries();
 		PostRenderBasePass(RHICmdList, View);
 	});
