Various Renderer dependency cleanup prereqs
- Moving GetOverwriteLoadAction from FViewInfo to FSceneView.
- Moving FPostProcessMaterialInputs into its own header.
- Adding CreateSceneTextureUniformBuffer variants that only need a FSceneView.

diff --git a/Engine/Source/Runtime/Engine/Public/SceneInterface.h b/Engine/Source/Runtime/Engine/Public/SceneInterface.h
--- a/Engine/Source/Runtime/Engine/Public/SceneInterface.h
+++ b/Engine/Source/Runtime/Engine/Public/SceneInterface.h
@@ -620,6 +620,8 @@ public:
 	}
 #endif //WITH_EDITOR
 
+	virtual TConstArrayView<FPrimitiveSceneProxy*> GetPrimitiveSceneProxies() const = 0;
+
 	/**
 	 * Returns the FPrimitiveComponentId for all primitives in the scene
 	 */
diff --git a/Engine/Source/Runtime/Engine/Public/SceneTexturesConfig.h b/Engine/Source/Runtime/Engine/Public/SceneTexturesConfig.h
--- a/Engine/Source/Runtime/Engine/Public/SceneTexturesConfig.h
+++ b/Engine/Source/Runtime/Engine/Public/SceneTexturesConfig.h
@@ -8,6 +8,7 @@
 #include "SceneUtils.h"
 
 class FSceneViewFamily;
+struct FMinimalSceneTextures;
 
 /** A uniform buffer containing common scene textures used by materials or global shaders. */
 BEGIN_GLOBAL_SHADER_PARAMETER_STRUCT(FSceneTextureUniformParameters, ENGINE_API)
diff --git a/Engine/Source/Runtime/Engine/Public/SceneView.h b/Engine/Source/Runtime/Engine/Public/SceneView.h
--- a/Engine/Source/Runtime/Engine/Public/SceneView.h
+++ b/Engine/Source/Runtime/Engine/Public/SceneView.h
@@ -1661,6 +1661,12 @@ public:
 		return InstancedViewUniformBuffer;
 	}
 
+	/** Returns the load action to use when overwriting all pixels of a target that you intend to read from. Takes into account the HMD hidden area mesh. */
+	inline ERenderTargetLoadAction GetOverwriteLoadAction() const
+	{
+		return bHMDHiddenAreaMaskActive ? ERenderTargetLoadAction::EClear : ERenderTargetLoadAction::ENoAction;
+	}
+
 protected:
 	FSceneViewStateInterface* EyeAdaptationViewState = nullptr;
 
diff --git a/Engine/Source/Runtime/Renderer/Private/PostProcess/PostProcessMaterial.h b/Engine/Source/Runtime/Renderer/Private/PostProcess/PostProcessMaterial.h
--- a/Engine/Source/Runtime/Renderer/Private/PostProcess/PostProcessMaterial.h
+++ b/Engine/Source/Runtime/Renderer/Private/PostProcess/PostProcessMaterial.h
@@ -6,37 +6,14 @@
 #include "OverridePassSequence.h"
 #include "Strata/Strata.h"
 
-class UMaterialInterface;
-struct FSceneWithoutWaterTextures;
+#include "PostProcess/PostProcessMaterialInputs.h"
 
-const uint32 kPostProcessMaterialInputCountMax = 5;
+class UMaterialInterface;
 
 using FPostProcessMaterialChain = TArray<const UMaterialInterface*, TInlineAllocator<10>>;
 
 FPostProcessMaterialChain GetPostProcessMaterialChain(const FViewInfo& View, EBlendableLocation Location);
 
-/** Named post process material slots. Inputs are aliased and have different semantics
- *  based on the post process material blend point, which is documented with the input.
- */
-enum class EPostProcessMaterialInput : uint32
-{
-	// Always Active. Color from the previous stage of the post process chain.
-	SceneColor = 0,
-
-	// Always Active.
-	SeparateTranslucency = 1,
-
-	// Replace Tonemap Only. Half resolution combined bloom input.
-	CombinedBloom = 2,
-
-	// Buffer Visualization Only.
-	PreTonemapHDRColor = 2,
-	PostTonemapHDRColor = 3,
-
-	// Active if separate velocity pass is used--i.e. not part of base pass; Not active during Replace Tonemap.
-	Velocity = 4
-};
-
 BEGIN_SHADER_PARAMETER_STRUCT(FPostProcessMaterialParameters, )
 	SHADER_PARAMETER_STRUCT_REF(FViewUniformShaderParameters, View)
 	SHADER_PARAMETER_STRUCT_INCLUDE(FSceneTextureShaderParameters, SceneTextures)
@@ -57,74 +34,7 @@ BEGIN_SHADER_PARAMETER_STRUCT(FPostProcessMaterialParameters, )
 	RENDER_TARGET_BINDING_SLOTS()
 END_SHADER_PARAMETER_STRUCT()
 
-struct FPostProcessMaterialInputs
-{
-	inline void SetInput(EPostProcessMaterialInput Input, FScreenPassTexture Texture)
-	{
-		Textures[(uint32)Input] = Texture;
-	}
-
-	inline FScreenPassTexture GetInput(EPostProcessMaterialInput Input) const
-	{
-		return Textures[(uint32)Input];
-	}
-
-	inline void Validate() const
-	{
-		ValidateInputExists(EPostProcessMaterialInput::SceneColor);
-		ValidateInputExists(EPostProcessMaterialInput::SeparateTranslucency);
-
-		// Either override output format is valid or the override output texture is; not both.
-		if (OutputFormat != PF_Unknown)
-		{
-			check(OverrideOutput.Texture == nullptr);
-		}
-		if (OverrideOutput.Texture)
-		{
-			check(OutputFormat == PF_Unknown);
-		}
-
-		check(SceneTextures.SceneTextures || SceneTextures.MobileSceneTextures);
-	}
-
-	inline void ValidateInputExists(EPostProcessMaterialInput Input) const
-	{
-		const FScreenPassTexture Texture = GetInput(EPostProcessMaterialInput::SceneColor);
-		check(Texture.IsValid());
-	}
-
-	// [Optional] Render to the specified output. If invalid, a new texture is created and returned.
-	FScreenPassRenderTarget OverrideOutput;
-
-	/** Array of input textures bound to the material. The first element represents the output from
-	 *  the previous post process and is required. All other inputs are optional.
-	 */
-	TStaticArray<FScreenPassTexture, kPostProcessMaterialInputCountMax> Textures;
-
-	/** The output texture format to use if a new texture is created. Uses the input format if left unknown. */
-	EPixelFormat OutputFormat = PF_Unknown;
-
-	/** Whether or not the stencil test must be done in the pixel shader rather than rasterizer state. */
-	bool bManualStencilTest = false;
-
-	/** Custom depth/stencil used for stencil operations. */
-	FRDGTextureRef CustomDepthTexture = nullptr;
 
-	/** The uniform buffer containing all scene textures. */
-	FSceneTextureShaderParameters SceneTextures;
-
-	/** Depth and color textures of the scene without single layer water. May be nullptr if not available. */
-	const FSceneWithoutWaterTextures* SceneWithoutWaterTextures = nullptr;
-
-	/** Allows (but doesn't guarantee) an optimization where, if possible, the scene color input is reused as
-	 *  the output. This can elide a copy in certain circumstances; for example, when the scene color input isn't
-	 *  actually used by the post process material and no special depth-stencil / blend composition is required.
-	 *  Set this to false when you need to guarantee creation of a dedicated output texture.
-	 */
-	bool bAllowSceneColorInputAsOutput = true;
-
-	bool bMetalMSAAHDRDecode = false;
-};
 
 FScreenPassTexture AddPostProcessMaterialPass(
 	FRDGBuilder& GraphBuilder,
diff --git a/Engine/Source/Runtime/Renderer/Private/RendererScene.cpp b/Engine/Source/Runtime/Renderer/Private/RendererScene.cpp
--- a/Engine/Source/Runtime/Renderer/Private/RendererScene.cpp
+++ b/Engine/Source/Runtime/Renderer/Private/RendererScene.cpp
@@ -6593,7 +6593,12 @@ public:
 
 	virtual bool HasAnyLights() const override { return false; }
 
-	virtual TConstArrayView<FPrimitiveComponentId> GetScenePrimitiveComponentIds() const override
+	virtual TConstArrayView<FPrimitiveSceneProxy*> GetPrimitiveSceneProxies() const override
+	{
+		return TConstArrayView<FPrimitiveSceneProxy*>();
+	}
+
+	virtual TConstArrayView<FPrimitiveComponentId> GetScenePrimitiveComponentIds() const final
 	{
 		return TConstArrayView<FPrimitiveComponentId>();
 	}
diff --git a/Engine/Source/Runtime/Renderer/Private/SceneCaptureRendering.cpp b/Engine/Source/Runtime/Renderer/Private/SceneCaptureRendering.cpp
--- a/Engine/Source/Runtime/Renderer/Private/SceneCaptureRendering.cpp
+++ b/Engine/Source/Runtime/Renderer/Private/SceneCaptureRendering.cpp
@@ -179,7 +179,7 @@ void CopySceneCaptureComponentToTarget(
 	const FMinimalSceneTextures& SceneTextures,
 	FRDGTextureRef ViewFamilyTexture,
 	const FSceneViewFamily& ViewFamily,
-	const TArray<FViewInfo>& Views)
+	TConstArrayView<FViewInfo> Views)
 {
 	ESceneCaptureSource SceneCaptureSource = ViewFamily.SceneCaptureSource;
 
@@ -259,6 +259,22 @@ void CopySceneCaptureComponentToTarget(
 	}
 }
 
+void CopySceneCaptureComponentToTarget(
+	FRDGBuilder& GraphBuilder,
+	FRDGTextureRef ViewFamilyTexture,
+	const FSceneViewFamily& ViewFamily,
+	TConstStridedView<FSceneView> Views)
+{
+	const FSceneView& View = Views[0];
+
+	check(View.bIsViewInfo);
+	const FMinimalSceneTextures& SceneTextures = static_cast<const FViewInfo&>(View).GetSceneTextures();
+
+	TConstArrayView<FViewInfo> ViewInfos = MakeArrayView(static_cast<const FViewInfo*>(&Views[0]), Views.Num());
+
+	CopySceneCaptureComponentToTarget(GraphBuilder, SceneTextures, ViewFamilyTexture, ViewFamily, ViewInfos);
+}
+
 static void UpdateSceneCaptureContentDeferred_RenderThread(
 	FRHICommandListImmediate& RHICmdList, 
 	FSceneRenderer* SceneRenderer, 
diff --git a/Engine/Source/Runtime/Renderer/Private/ScenePrivate.h b/Engine/Source/Runtime/Renderer/Private/ScenePrivate.h
--- a/Engine/Source/Runtime/Renderer/Private/ScenePrivate.h
+++ b/Engine/Source/Runtime/Renderer/Private/ScenePrivate.h
@@ -3348,7 +3348,12 @@ public:
 
 	bool ShouldRenderSkylightInBasePass(bool bIsTranslucent) const;
 
-	virtual TConstArrayView<FPrimitiveComponentId> GetScenePrimitiveComponentIds() const override
+	virtual TConstArrayView<FPrimitiveSceneProxy*> GetPrimitiveSceneProxies() const final
+	{
+		return PrimitiveSceneProxies;
+	}
+
+	virtual TConstArrayView<FPrimitiveComponentId> GetScenePrimitiveComponentIds() const final
 	{
 		return PrimitiveComponentIds;
 	}
diff --git a/Engine/Source/Runtime/Renderer/Private/SceneRendering.cpp b/Engine/Source/Runtime/Renderer/Private/SceneRendering.cpp
--- a/Engine/Source/Runtime/Renderer/Private/SceneRendering.cpp
+++ b/Engine/Source/Runtime/Renderer/Private/SceneRendering.cpp
@@ -2249,11 +2249,6 @@ float FViewInfo::GetLastAverageSceneLuminance() const
 	return 0.0f; // Invalid scene luminance
 }
 
-ERenderTargetLoadAction FViewInfo::GetOverwriteLoadAction() const
-{
-	return bHMDHiddenAreaMaskActive ? ERenderTargetLoadAction::EClear : ERenderTargetLoadAction::ENoAction;
-}
-
 void FViewInfo::SetValidTonemappingLUT() const
 {
 	if (FSceneViewState* EffectiveViewState = GetEyeAdaptationViewState())
diff --git a/Engine/Source/Runtime/Renderer/Private/SceneRendering.h b/Engine/Source/Runtime/Renderer/Private/SceneRendering.h
--- a/Engine/Source/Runtime/Renderer/Private/SceneRendering.h
+++ b/Engine/Source/Runtime/Renderer/Private/SceneRendering.h
@@ -1750,9 +1750,6 @@ public:
 	/** Enqueue a pass to readback current exposure */
 	void EnqueueEyeAdaptationExposureBufferReadback(FRDGBuilder& GraphBuilder) const;
 	
-	/** Returns the load action to use when overwriting all pixels of a target that you intend to read from. Takes into account the HMD hidden area mesh. */
-	RENDERER_API ERenderTargetLoadAction GetOverwriteLoadAction() const;
-
 	/** Informs sceneinfo that tonemapping LUT has queued commands to compute it at least once */
 	void SetValidTonemappingLUT() const;
 
diff --git a/Engine/Source/Runtime/Renderer/Private/SceneTextures.cpp b/Engine/Source/Runtime/Renderer/Private/SceneTextures.cpp
--- a/Engine/Source/Runtime/Renderer/Private/SceneTextures.cpp
+++ b/Engine/Source/Runtime/Renderer/Private/SceneTextures.cpp
@@ -871,6 +871,19 @@ TRDGUniformBufferRef<FSceneTextureUniformParameters> CreateSceneTextureUniformBu
 	return GraphBuilder.CreateUniformBuffer(SceneTexturesParameters);
 }
 
+TRDGUniformBufferRef<FSceneTextureUniformParameters> CreateSceneTextureUniformBuffer(
+	FRDGBuilder& GraphBuilder,
+	const FSceneView& View,
+	ESceneTextureSetupMode SetupMode)
+{
+	if (View.bIsViewInfo)
+	{
+		return CreateSceneTextureUniformBuffer(GraphBuilder, ((const FViewInfo&)View).GetSceneTexturesChecked(), View.GetFeatureLevel(), SetupMode);
+	}
+
+	return nullptr;
+}
+
 EMobileSceneTextureSetupMode Translate(ESceneTextureSetupMode InSetupMode)
 {
 	EMobileSceneTextureSetupMode OutSetupMode = EMobileSceneTextureSetupMode::None;
@@ -989,6 +1002,16 @@ TRDGUniformBufferRef<FMobileSceneTextureUniformParameters> CreateMobileSceneText
 	return GraphBuilder.CreateUniformBuffer(SceneTexturesParameters);
 }
 
+TRDGUniformBufferRef<FMobileSceneTextureUniformParameters> CreateMobileSceneTextureUniformBuffer(FRDGBuilder& GraphBuilder, const FSceneView& View, EMobileSceneTextureSetupMode SetupMode)
+{
+	if (View.bIsViewInfo)
+	{
+		return CreateMobileSceneTextureUniformBuffer(GraphBuilder, ((const FViewInfo&)View).GetSceneTexturesChecked(), SetupMode);
+	}
+
+	return nullptr;
+}
+
 FSceneTextureShaderParameters CreateSceneTextureShaderParameters(
 	FRDGBuilder& GraphBuilder,
 	const FSceneTextures* SceneTextures,
@@ -1007,6 +1030,20 @@ FSceneTextureShaderParameters CreateSceneTextureShaderParameters(
 	return Parameters;
 }
 
+FSceneTextureShaderParameters CreateSceneTextureShaderParameters(FRDGBuilder& GraphBuilder, const FSceneView& View, ESceneTextureSetupMode SetupMode)
+{
+	FSceneTextureShaderParameters Parameters;
+	if (FSceneInterface::GetShadingPath(View.FeatureLevel) == EShadingPath::Deferred)
+	{
+		Parameters.SceneTextures = CreateSceneTextureUniformBuffer(GraphBuilder, View, SetupMode);
+	}
+	else if (FSceneInterface::GetShadingPath(View.FeatureLevel) == EShadingPath::Mobile)
+	{
+		Parameters.MobileSceneTextures = CreateMobileSceneTextureUniformBuffer(GraphBuilder, View, Translate(SetupMode));
+	}
+	return Parameters;
+}
+
 TRDGUniformBufferRef<FSceneTextureUniformParameters> GetSceneTextureUnformBuffer(const FSceneView& View)
 {
 	if (const FSceneTextures* SceneTextures = static_cast<const FViewFamilyInfo*>(View.Family)->GetSceneTexturesChecked())
diff --git a/Engine/Source/Runtime/Renderer/Public/PostProcess/PostProcessMaterialInputs.h b/Engine/Source/Runtime/Renderer/Public/PostProcess/PostProcessMaterialInputs.h
new file mode 100644
--- /dev/null
+++ b/Engine/Source/Runtime/Renderer/Public/PostProcess/PostProcessMaterialInputs.h
@@ -0,0 +1,100 @@
+// Copyright Epic Games, Inc. All Rights Reserved.
+
+#pragma once
+
+#include "ScreenPass.h"
+
+struct FSceneWithoutWaterTextures;
+
+const uint32 kPostProcessMaterialInputCountMax = 5;
+
+/** Named post process material slots. Inputs are aliased and have different semantics
+ *  based on the post process material blend point, which is documented with the input.
+ */
+enum class EPostProcessMaterialInput : uint32
+{
+	// Always Active. Color from the previous stage of the post process chain.
+	SceneColor = 0,
+
+	// Always Active.
+	SeparateTranslucency = 1,
+
+	// Replace Tonemap Only. Half resolution combined bloom input.
+	CombinedBloom = 2,
+
+	// Buffer Visualization Only.
+	PreTonemapHDRColor = 2,
+	PostTonemapHDRColor = 3,
+
+	// Active if separate velocity pass is used--i.e. not part of base pass; Not active during Replace Tonemap.
+	Velocity = 4
+};
+
+struct FPostProcessMaterialInputs
+{
+	inline void SetInput(EPostProcessMaterialInput Input, FScreenPassTexture Texture)
+	{
+		Textures[(uint32)Input] = Texture;
+	}
+
+	inline FScreenPassTexture GetInput(EPostProcessMaterialInput Input) const
+	{
+		return Textures[(uint32)Input];
+	}
+
+	inline void Validate() const
+	{
+		ValidateInputExists(EPostProcessMaterialInput::SceneColor);
+		ValidateInputExists(EPostProcessMaterialInput::SeparateTranslucency);
+
+		// Either override output format is valid or the override output texture is; not both.
+		if (OutputFormat != PF_Unknown)
+		{
+			check(OverrideOutput.Texture == nullptr);
+		}
+		if (OverrideOutput.Texture)
+		{
+			check(OutputFormat == PF_Unknown);
+		}
+
+		check(SceneTextures.SceneTextures || SceneTextures.MobileSceneTextures);
+	}
+
+	inline void ValidateInputExists(EPostProcessMaterialInput Input) const
+	{
+		const FScreenPassTexture Texture = GetInput(EPostProcessMaterialInput::SceneColor);
+		check(Texture.IsValid());
+	}
+
+	// [Optional] Render to the specified output. If invalid, a new texture is created and returned.
+	FScreenPassRenderTarget OverrideOutput;
+
+	/** Array of input textures bound to the material. The first element represents the output from
+	 *  the previous post process and is required. All other inputs are optional.
+	 */
+	TStaticArray<FScreenPassTexture, kPostProcessMaterialInputCountMax> Textures;
+
+	/** The output texture format to use if a new texture is created. Uses the input format if left unknown. */
+	EPixelFormat OutputFormat = PF_Unknown;
+
+	/** Whether or not the stencil test must be done in the pixel shader rather than rasterizer state. */
+	bool bManualStencilTest = false;
+
+	/** Custom depth/stencil used for stencil operations. */
+	FRDGTextureRef CustomDepthTexture = nullptr;
+
+	/** The uniform buffer containing all scene textures. */
+	FSceneTextureShaderParameters SceneTextures;
+
+	/** Depth and color textures of the scene without single layer water. May be nullptr if not available. */
+	const FSceneWithoutWaterTextures* SceneWithoutWaterTextures = nullptr;
+
+	/** Allows (but doesn't guarantee) an optimization where, if possible, the scene color input is reused as
+	 *  the output. This can elide a copy in certain circumstances; for example, when the scene color input isn't
+	 *  actually used by the post process material and no special depth-stencil / blend composition is required.
+	 *  Set this to false when you need to guarantee creation of a dedicated output texture.
+	 */
+	bool bAllowSceneColorInputAsOutput = true;
+
+	bool bMetalMSAAHDRDecode = false;
+};
diff --git a/Engine/Source/Runtime/Renderer/Public/SceneCaptureRendering.h b/Engine/Source/Runtime/Renderer/Public/SceneCaptureRendering.h
--- a/Engine/Source/Runtime/Renderer/Public/SceneCaptureRendering.h
+++ b/Engine/Source/Runtime/Renderer/Public/SceneCaptureRendering.h
@@ -6,13 +6,23 @@
 
 #pragma once
 
-#include "CoreMinimal.h"
-#include "RenderGraphDefinitions.h"
-#include "SceneTextures.h"
+#include "Containers/ContainersFwd.h"
+#include "RenderGraphFwd.h"
+
+class FSceneView;
+class FSceneViewFamily;
+class FViewInfo;
+struct FMinimalSceneTextures;
 
 RENDERER_API void CopySceneCaptureComponentToTarget(
 	FRDGBuilder& GraphBuilder,
 	const FMinimalSceneTextures& SceneTextures,
 	FRDGTextureRef ViewFamilyTexture,
 	const FSceneViewFamily& ViewFamily,
-	const TArray<FViewInfo>& Views);
+	TConstArrayView<FViewInfo> Views);
+
+RENDERER_API void CopySceneCaptureComponentToTarget(
+	FRDGBuilder& GraphBuilder,
+	FRDGTextureRef ViewFamilyTexture,
+	const FSceneViewFamily& ViewFamily,
+	TConstStridedView<FSceneView> Views);
diff --git a/Engine/Source/Runtime/Renderer/Public/SceneRenderTargetParameters.h b/Engine/Source/Runtime/Renderer/Public/SceneRenderTargetParameters.h
--- a/Engine/Source/Runtime/Renderer/Public/SceneRenderTargetParameters.h
+++ b/Engine/Source/Runtime/Renderer/Public/SceneRenderTargetParameters.h
@@ -45,6 +45,12 @@ extern RENDERER_API TRDGUniformBufferRef<FSceneTextureUniformParameters> CreateS
 	ERHIFeatureLevel::Type FeatureLevel,
 	ESceneTextureSetupMode SetupMode = ESceneTextureSetupMode::All);
 
+/** Returns RDG scene texture uniform buffer for a specified View. */
+extern RENDERER_API TRDGUniformBufferRef<FSceneTextureUniformParameters> CreateSceneTextureUniformBuffer(
+	FRDGBuilder& GraphBuilder,
+	const FSceneView& View,
+	ESceneTextureSetupMode SetupMode = ESceneTextureSetupMode::All);
+
 enum class EMobileSceneTextureSetupMode : uint32
 {
 	None			= 0,
@@ -75,6 +81,12 @@ extern RENDERER_API TRDGUniformBufferRef<FMobileSceneTextureUniformParameters> C
 	const FSceneTextures* SceneTextures,
 	EMobileSceneTextureSetupMode SetupMode = EMobileSceneTextureSetupMode::All);
 
+/** Creates the RDG mobile scene texture uniform buffer. */
+extern RENDERER_API TRDGUniformBufferRef<FMobileSceneTextureUniformParameters> CreateMobileSceneTextureUniformBuffer(
+	FRDGBuilder& GraphBuilder,
+	const FSceneView& View,
+	EMobileSceneTextureSetupMode SetupMode = EMobileSceneTextureSetupMode::All);
+
 /** Returns scene texture shader parameters containing the RDG uniform buffer for either mobile or deferred shading. */
 extern RENDERER_API FSceneTextureShaderParameters CreateSceneTextureShaderParameters(
 	FRDGBuilder& GraphBuilder,
@@ -82,6 +94,12 @@ extern RENDERER_API FSceneTextureShaderParameters CreateSceneTextureShaderParame
 	ERHIFeatureLevel::Type FeatureLevel,
 	ESceneTextureSetupMode SetupMode = ESceneTextureSetupMode::All);
 
+/** Returns scene texture shader parameters containing the RDG uniform buffer for either mobile or deferred shading. */
+extern RENDERER_API FSceneTextureShaderParameters CreateSceneTextureShaderParameters(
+	FRDGBuilder& GraphBuilder,
+	const FSceneView& View,
+	ESceneTextureSetupMode SetupMode = ESceneTextureSetupMode::All);
+
 /** Struct containing references to extracted RHI resources after RDG execution. All textures are
  *  left in an SRV read state, so they can safely be used for read without being re-imported into
  *  RDG. Likewise, the uniform buffer is non-RDG and can be used as is.
