Fixes TSR color history drift with r.TSR.History.R11G11B10=0 and r.TSR.History.GrandReprojection=1

diff --git a/Engine/Source/Runtime/Renderer/Private/PostProcess/TemporalSuperResolution.cpp b/Engine/Source/Runtime/Renderer/Private/PostProcess/TemporalSuperResolution.cpp
--- a/Engine/Source/Runtime/Renderer/Private/PostProcess/TemporalSuperResolution.cpp
+++ b/Engine/Source/Runtime/Renderer/Private/PostProcess/TemporalSuperResolution.cpp
@@ -896,6 +896,7 @@ ITemporalUpscaler::FOutputs AddTemporalSuperResolutionPasses(
 	const bool bRejectSeparateTranslucency = false; // bIsSeperateTranslucyTexturesValid && CVarTSRTranslucencyPreviousFrameRejection.GetValueOnRenderThread() != 0;
 
 	EPixelFormat ColorFormat = bSupportsAlpha ? PF_FloatRGBA : PF_FloatR11G11B10;
+	EPixelFormat HistoryColorFormat = (CVarTSRR11G11B10History.GetValueOnRenderThread() != 0 && !bSupportsAlpha) ? PF_FloatR11G11B10 : PF_FloatRGBA;
 
 	int32 RejectionAntiAliasingQuality = FMath::Clamp(CVarTSRRejectionAntiAliasingQuality.GetValueOnRenderThread(), 1, 2);
 	if (UpdateHistoryQuality == FTSRUpdateHistoryCS::EQuality::Low)
@@ -1202,7 +1203,7 @@ ITemporalUpscaler::FOutputs AddTemporalSuperResolutionPasses(
 		check(!(PassInputs.bGenerateOutputMip1 && (PassInputs.bGenerateSceneColorHalfRes || PassInputs.bGenerateSceneColorQuarterRes)));
 		FRDGTextureDesc Desc = FRDGTextureDesc::Create2D(
 			HistoryExtent,
-			ColorFormat,
+			bIsOutputDifferentThanHighFrequency ? ColorFormat : HistoryColorFormat,
 			FClearValueBinding::None,
 			/* InFlags = */ TexCreate_ShaderResource | TexCreate_UAV,
 			/* NumMips = */ PassInputs.bGenerateOutputMip1 ? 2 : 1);
@@ -1210,6 +1211,8 @@ ITemporalUpscaler::FOutputs AddTemporalSuperResolutionPasses(
 		UpdateHistoryOutputTexture = GraphBuilder.CreateTexture(Desc, TEXT("TSR.Output"));
 		History.Output = UpdateHistoryOutputTexture;
 
+		Desc.Format = ColorFormat;
+
 		if (OutputRect.Size() != HistorySize)
 		{
 			Desc.Extent = OutputExtent;
@@ -1242,7 +1245,7 @@ ITemporalUpscaler::FOutputs AddTemporalSuperResolutionPasses(
 		{
 			FRDGTextureDesc ArrayDesc = FRDGTextureDesc::Create2DArray(
 				HistoryExtent,
-				(CVarTSRR11G11B10History.GetValueOnRenderThread() != 0 && !bSupportsAlpha) ? PF_FloatR11G11B10 : PF_FloatRGBA,
+				HistoryColorFormat,
 				FClearValueBinding::None,
 				TexCreate_ShaderResource | TexCreate_UAV,
 				History.ArrayIndices.Size);
@@ -1253,7 +1256,7 @@ ITemporalUpscaler::FOutputs AddTemporalSuperResolutionPasses(
 	{
 		FRDGTextureDesc Desc = FRDGTextureDesc::Create2D(
 			HistoryExtent,
-			(CVarTSRR11G11B10History.GetValueOnRenderThread() != 0 && !bSupportsAlpha) ? PF_FloatR11G11B10 : PF_FloatRGBA,
+			HistoryColorFormat,
 			FClearValueBinding::None,
 			TexCreate_ShaderResource | TexCreate_UAV);
 
