Reset size of scene textures after drawing a material to a render target from Blueprint. This prevents a user from keeping too much memory allocated for a large, one-off render target and this BP function shouldn't be altering this state anyway.

diff --git a/Engine/Source/Runtime/Engine/Private/KismetRenderingLibrary.cpp b/Engine/Source/Runtime/Engine/Private/KismetRenderingLibrary.cpp
--- a/Engine/Source/Runtime/Engine/Private/KismetRenderingLibrary.cpp
+++ b/Engine/Source/Runtime/Engine/Private/KismetRenderingLibrary.cpp
@@ -23,6 +23,7 @@
 #include "ClearQuad.h"
 #include "Engine/Texture2D.h"
 #include "UObject/Package.h"
+#include "EngineModule.h"
 
 #include UE_INLINE_GENERATED_CPP_BY_NAME(KismetRenderingLibrary)
 
@@ -215,6 +216,12 @@ void UKismetRenderingLibrary::DrawMaterialToRenderTarget(UObject* WorldContextOb
 
 			//UpdateResourceImmediate must be called here to ensure mips are generated.
 			TextureRenderTarget->UpdateResourceImmediate(false);
+
+			ENQUEUE_RENDER_COMMAND(ResetSceneTextureExtentHistory)(
+				[RenderTargetResource](FRHICommandListImmediate& RHICmdList)
+				{
+					RenderTargetResource->ResetSceneTextureExtentsHistory();
+				});
 		}
 	}
 }
diff --git a/Engine/Source/Runtime/Engine/Private/TextureRenderTarget.cpp b/Engine/Source/Runtime/Engine/Private/TextureRenderTarget.cpp
--- a/Engine/Source/Runtime/Engine/Private/TextureRenderTarget.cpp
+++ b/Engine/Source/Runtime/Engine/Private/TextureRenderTarget.cpp
@@ -7,6 +7,8 @@
 #include "Engine/TextureRenderTarget.h"
 #include "ProfilingDebugging/CsvProfiler.h"
 #include "TextureResource.h"
+#include "SceneRenderTargetParameters.h"
+#include "EngineModule.h"
 
 #include UE_INLINE_GENERATED_CPP_BY_NAME(TextureRenderTarget)
 
@@ -195,3 +197,7 @@ void FDeferredUpdateResource::RemoveFromDeferredUpdateList()
 	UpdateListLink.Unlink();
 }
 
+void FDeferredUpdateResource::ResetSceneTextureExtentsHistory()
+{
+	GetRendererModule().ResetSceneTextureExtentHistory();
+}
diff --git a/Engine/Source/Runtime/Engine/Public/TextureResource.h b/Engine/Source/Runtime/Engine/Public/TextureResource.h
--- a/Engine/Source/Runtime/Engine/Public/TextureResource.h
+++ b/Engine/Source/Runtime/Engine/Public/TextureResource.h
@@ -289,6 +289,7 @@ public:
 		return bNeedsUpdate;
 	}
 
+	ENGINE_API void ResetSceneTextureExtentsHistory();
 protected:
 
 	/**
diff --git a/Engine/Source/Runtime/RenderCore/Public/RendererInterface.h b/Engine/Source/Runtime/RenderCore/Public/RendererInterface.h
--- a/Engine/Source/Runtime/RenderCore/Public/RendererInterface.h
+++ b/Engine/Source/Runtime/RenderCore/Public/RendererInterface.h
@@ -835,5 +835,10 @@ public:
 
 	/** Experimental:  Render multiple view families in a single scene render call.  All families must reference the same FScene.  Scene Capture not yet supported. */
 	virtual void BeginRenderingViewFamilies(FCanvas* Canvas, TArrayView<FSceneViewFamily*> ViewFamilies) = 0;
+
+	/** Resets the scene texture extent history. Call this method after rendering with very large render
+	 *  targets. The next scene render will create them at the requested size.
+	 */
+	virtual void ResetSceneTextureExtentHistory() = 0;
 };
 
diff --git a/Engine/Source/Runtime/Renderer/Private/Renderer.cpp b/Engine/Source/Runtime/Renderer/Private/Renderer.cpp
--- a/Engine/Source/Runtime/Renderer/Private/Renderer.cpp
+++ b/Engine/Source/Runtime/Renderer/Private/Renderer.cpp
@@ -39,6 +39,9 @@
 #include "Lumen/Lumen.h"
 #include "ScenePrivate.h"
 #include "SceneUniformBuffer.h"
+#include "SceneRenderTargetParameters.h"
+#include "EngineModule.h"
+#include "RendererInterface.h"
 
 DEFINE_LOG_CATEGORY(LogRenderer);
 
@@ -115,7 +118,6 @@ void FRendererModule::InitializeSystemTextures(FRHICommandListImmediate& RHICmdL
 
 BEGIN_SHADER_PARAMETER_STRUCT(FDrawTileMeshPassParameters, )
 	SHADER_PARAMETER_STRUCT_INCLUDE(FViewShaderParameters, View)
-	SHADER_PARAMETER_RDG_UNIFORM_BUFFER(FSceneUniformParameters, Scene)
 	SHADER_PARAMETER_RDG_UNIFORM_BUFFER(FInstanceCullingGlobalUniforms, InstanceCulling)
 	SHADER_PARAMETER_STRUCT_REF(FReflectionCaptureShaderData, ReflectionCapture)
 	SHADER_PARAMETER_RDG_UNIFORM_BUFFER(FDebugViewModePassUniformParameters, DebugViewMode)
@@ -125,80 +127,6 @@ BEGIN_SHADER_PARAMETER_STRUCT(FDrawTileMeshPassParameters, )
 	RENDER_TARGET_BINDING_SLOTS()
 END_SHADER_PARAMETER_STRUCT()
 
-FSceneUniformBuffer* FRendererModule::CreateSinglePrimitiveSceneUniformBuffer(FRDGBuilder& GraphBuilder, const FViewInfo& SceneView, FMeshBatch& Mesh)
-{
-	const auto FeatureLevel = SceneView.GetFeatureLevel();
-	
-	FSceneUniformBuffer& SceneUniforms = *GraphBuilder.AllocObject<FSceneUniformBuffer>();
-
-	if (Mesh.VertexFactory->GetPrimitiveIdStreamIndex(FeatureLevel, EVertexInputStreamType::PositionOnly) >= 0)
-	{
-		FMeshBatchElement& MeshElement = Mesh.Elements[0];
-
-		checkf(Mesh.Elements.Num() == 1, TEXT("Only 1 batch element currently supported by CreateSinglePrimitiveSceneUniformBuffer"));
-		checkf(MeshElement.PrimitiveUniformBuffer == nullptr, TEXT("CreateSinglePrimitiveSceneUniformBuffer does not currently support an explicit primitive uniform buffer on vertex factories which manually fetch primitive data.  Use PrimitiveUniformBufferResource instead."));
-
-		if (MeshElement.PrimitiveUniformBufferResource)
-		{
-			checkf(MeshElement.NumInstances == 1, TEXT("CreateSinglePrimitiveSceneUniformBuffer does not currently support instancing"));
-			// Force PrimitiveId to be 0 in the shader
-			MeshElement.PrimitiveIdMode = PrimID_ForceZero;
-
-			// Set the LightmapID to 0, since that's where our light map data resides for this primitive
-			FPrimitiveUniformShaderParameters PrimitiveParams = *(const FPrimitiveUniformShaderParameters*)MeshElement.PrimitiveUniformBufferResource->GetContents();
-			PrimitiveParams.LightmapDataIndex = 0;
-			PrimitiveParams.LightmapUVIndex = 0;
-
-			// Set up reference to the single-instance 
-			PrimitiveParams.InstanceSceneDataOffset = 0;
-			PrimitiveParams.NumInstanceSceneDataEntries = 1;
-			PrimitiveParams.InstancePayloadDataOffset = INDEX_NONE;
-			PrimitiveParams.InstancePayloadDataStride = 0;
-			
-			// Now we just need to fill out the first entry of primitive data in a buffer and bind it
-			FPrimitiveSceneShaderData PrimitiveSceneData(PrimitiveParams);
-
-			// Also fill out correct single-primitive instance data, derived from the primitive.
-			FInstanceSceneShaderData InstanceSceneData{};
-			InstanceSceneData.BuildInternal
-			(
-				0 /* Primitive Id */,
-				0 /* Relative Instance Id */,
-				0 /* Payload Data Flags */,
-				INVALID_LAST_UPDATE_FRAME,
-				0 /* Custom Data Count */,
-				0.0f /* Random ID */,
-				PrimitiveParams.LocalToRelativeWorld
-			);
-
-			// Set up the parameters for the LightmapSceneData from the given LCI data 
-			FPrecomputedLightingUniformParameters LightmapParams;
-			GetPrecomputedLightingParameters(FeatureLevel, LightmapParams, Mesh.LCI);
-			FLightmapSceneShaderData LightmapSceneData(LightmapParams);
-
-			FRDGBufferRef PrimitiveSceneDataBuffer = CreateStructuredBuffer(GraphBuilder, TEXT("PrimitiveSceneDataBuffer"), TConstArrayView<FVector4f>(PrimitiveSceneData.Data));
-			FRDGBufferRef LightmapSceneDataBuffer = CreateStructuredBuffer(GraphBuilder, TEXT("LightmapSceneDataBuffer"), TConstArrayView<FVector4f>(LightmapSceneData.Data));
-			FRDGBufferRef InstanceSceneDataBuffer = CreateStructuredBuffer(GraphBuilder, TEXT("InstanceSceneDataBuffer"), TConstArrayView<FVector4f>(InstanceSceneData.Data));
-			FRDGBufferRef InstancePayloadDataBuffer = GSystemTextures.GetDefaultStructuredBuffer(GraphBuilder, sizeof(FVector4f));
-			FRDGBufferRef DummyBufferLight = GSystemTextures.GetDefaultStructuredBuffer(GraphBuilder, sizeof(FLightSceneData));
-
-			FGPUSceneResourceParameters ShaderParameters;
-			ShaderParameters.GPUScenePrimitiveSceneData = GraphBuilder.CreateSRV(PrimitiveSceneDataBuffer);
-			ShaderParameters.GPUSceneInstanceSceneData = GraphBuilder.CreateSRV(InstanceSceneDataBuffer);
-			ShaderParameters.GPUSceneInstancePayloadData = GraphBuilder.CreateSRV(InstancePayloadDataBuffer);
-			ShaderParameters.GPUSceneLightmapData = GraphBuilder.CreateSRV(LightmapSceneDataBuffer);
-			ShaderParameters.GPUSceneLightData = GraphBuilder.CreateSRV(DummyBufferLight);
-			ShaderParameters.InstanceDataSOAStride = 1;
-			ShaderParameters.NumInstances = 1;
-			ShaderParameters.NumScenePrimitives = 1;
-			
-			SceneUniforms.Set(ShaderParameters);
-		}
-	}
-
-	return &SceneUniforms;
-}
-
 void FRendererModule::DrawTileMesh(FCanvasRenderContext& RenderContext, FMeshPassProcessorRenderState& DrawRenderState, const FSceneView& SceneView, FMeshBatch& Mesh, bool bIsHitTesting, const FHitProxyId& HitProxyId, bool bUse128bitRT)
 {
 	if (!GUsingNullRHI)
@@ -220,6 +148,7 @@ void FRendererModule::DrawTileMesh(FCanvasRenderContext& RenderContext, FMeshPas
 		const EShadingPath ShadingPath = FSceneInterface::GetShadingPath(FeatureLevel);
 
 		FScene* Scene = nullptr;
+
 		if (ViewFamily->Scene)
 		{
 			Scene = ViewFamily->Scene->GetRenderScene();
@@ -228,10 +157,65 @@ void FRendererModule::DrawTileMesh(FCanvasRenderContext& RenderContext, FMeshPas
 		Mesh.MaterialRenderProxy->UpdateUniformExpressionCacheIfNeeded(FeatureLevel);
 		FMaterialRenderProxy::UpdateDeferredCachedUniformExpressions();
 
+		FSinglePrimitiveStructured& SinglePrimitiveStructured = GTilePrimitiveBuffer;
+
+		if (Mesh.VertexFactory->GetPrimitiveIdStreamIndex(FeatureLevel, EVertexInputStreamType::PositionOnly) >= 0)
+		{
+			FMeshBatchElement& MeshElement = Mesh.Elements[0];
+
+			checkf(Mesh.Elements.Num() == 1, TEXT("Only 1 batch element currently supported by DrawTileMesh"));
+			checkf(MeshElement.PrimitiveUniformBuffer == nullptr, TEXT("DrawTileMesh does not currently support an explicit primitive uniform buffer on vertex factories which manually fetch primitive data.  Use PrimitiveUniformBufferResource instead."));
+
+			if (MeshElement.PrimitiveUniformBufferResource)
+			{
+				checkf(MeshElement.NumInstances == 1, TEXT("DrawTileMesh does not currently support instancing"));
+				// Force PrimitiveId to be 0 in the shader
+				MeshElement.PrimitiveIdMode = PrimID_ForceZero;
+				
+				// Set the LightmapID to 0, since that's where our light map data resides for this primitive
+				FPrimitiveUniformShaderParameters PrimitiveParams = *(const FPrimitiveUniformShaderParameters*)MeshElement.PrimitiveUniformBufferResource->GetContents();
+				PrimitiveParams.LightmapDataIndex = 0;
+				PrimitiveParams.LightmapUVIndex = 0;
+
+				// Set up reference to the single-instance 
+				PrimitiveParams.InstanceSceneDataOffset = 0;
+				PrimitiveParams.NumInstanceSceneDataEntries = 1;
+				PrimitiveParams.InstancePayloadDataOffset = INDEX_NONE;
+				PrimitiveParams.InstancePayloadDataStride = 0;
+
+				// Now we just need to fill out the first entry of primitive data in a buffer and bind it
+				SinglePrimitiveStructured.PrimitiveSceneData = FPrimitiveSceneShaderData(PrimitiveParams);
+				SinglePrimitiveStructured.ShaderPlatform = View.GetShaderPlatform();
+
+				// Also fill out correct single-primitive instance data, derived from the primitive.
+				SinglePrimitiveStructured.InstanceSceneData.BuildInternal
+				(
+					0 /* Primitive Id */,
+					0 /* Relative Instance Id */,
+					0 /* Payload Data Flags */,
+					INVALID_LAST_UPDATE_FRAME,
+					0 /* Custom Data Count */,
+					0.0f /* Random ID */,
+					PrimitiveParams.LocalToRelativeWorld
+				);
+
+				// TODO: Payload dummy?
+
+				// Set up the parameters for the LightmapSceneData from the given LCI data 
+				FPrecomputedLightingUniformParameters LightmapParams;
+				GetPrecomputedLightingParameters(FeatureLevel, LightmapParams, Mesh.LCI);
+				SinglePrimitiveStructured.LightmapSceneData = FLightmapSceneShaderData(LightmapParams);
+
+				SinglePrimitiveStructured.UploadToGPU();
+
+				View.PrimitiveSceneDataOverrideSRV = SinglePrimitiveStructured.PrimitiveSceneDataBufferSRV;
+				View.InstanceSceneDataOverrideSRV  = SinglePrimitiveStructured.InstanceSceneDataBufferSRV;
+				View.InstancePayloadDataOverrideSRV = SinglePrimitiveStructured.InstancePayloadDataBufferSRV;
+				View.LightmapSceneDataOverrideSRV = SinglePrimitiveStructured.LightmapSceneDataBufferSRV;
+			}
+		}
 
 		FRDGBuilder& GraphBuilder = RenderContext.GraphBuilder;
-		
-		FSceneUniformBuffer& SceneUniforms = *CreateSinglePrimitiveSceneUniformBuffer(GraphBuilder, View, Mesh);
 
 		if (!FRDGSystemTextures::IsValid(GraphBuilder))
 		{
@@ -251,7 +235,7 @@ void FRendererModule::DrawTileMesh(FCanvasRenderContext& RenderContext, FMeshPas
 		}
 
 		View.InitRHIResources();
-		View.ForwardLightingResources.SetUniformBuffer(CreateDummyForwardLightUniformBuffer(GraphBuilder, View.GetShaderPlatform()));
+		View.ForwardLightingResources.SetUniformBuffer(CreateDummyForwardLightUniformBuffer(GraphBuilder));
 
 		TUniformBufferRef<FReflectionCaptureShaderData> EmptyReflectionCaptureUniformBuffer;
 
@@ -265,7 +249,6 @@ void FRendererModule::DrawTileMesh(FCanvasRenderContext& RenderContext, FMeshPas
 		auto* PassParameters = GraphBuilder.AllocParameters<FDrawTileMeshPassParameters>();
 		PassParameters->RenderTargets[0] = FRenderTargetBinding(RenderContext.GetRenderTarget(), ERenderTargetLoadAction::ELoad);
 		PassParameters->View = View.GetShaderParameters();
-		PassParameters->Scene = SceneUniforms.GetBuffer(GraphBuilder);
 		PassParameters->ReflectionCapture = EmptyReflectionCaptureUniformBuffer;
 		PassParameters->InstanceCulling = FInstanceCullingContext::CreateDummyInstanceCullingUniformBuffer(GraphBuilder);
 
diff --git a/Engine/Source/Runtime/Renderer/Private/RendererModule.h b/Engine/Source/Runtime/Renderer/Private/RendererModule.h
--- a/Engine/Source/Runtime/Renderer/Private/RendererModule.h
+++ b/Engine/Source/Runtime/Renderer/Private/RendererModule.h
@@ -136,6 +136,8 @@ public:
 
 	virtual void BeginRenderingViewFamilies(FCanvas* Canvas, TArrayView<FSceneViewFamily*> ViewFamilies) override;
 
+	virtual void ResetSceneTextureExtentHistory() override;
+
 private:
 	TSet<FSceneInterface*> AllocatedScenes;
 	FOnPostOpaqueRender PostOpaqueRenderDelegate;
