Silence static analysis warning related to potentially wrong && and ? operator precedence

* No actual behavior difference

diff --git a/Engine/Source/Runtime/Renderer/Private/VirtualShadowMaps/VirtualShadowMapCacheManager.cpp b/Engine/Source/Runtime/Renderer/Private/VirtualShadowMaps/VirtualShadowMapCacheManager.cpp
--- a/Engine/Source/Runtime/Renderer/Private/VirtualShadowMaps/VirtualShadowMapCacheManager.cpp
+++ b/Engine/Source/Runtime/Renderer/Private/VirtualShadowMaps/VirtualShadowMapCacheManager.cpp
@@ -615,7 +615,7 @@ TRefCountPtr<IPooledRenderTarget> FVirtualShadowMapArrayCacheManager::SetPhysica
 {
 	// Using ReservedResource|ImmediateCommit flags hint to the RHI that the resource can be allocated using N small physical memory allocations,
 	// instead of a single large contighous allocation. This helps Windows video memory manager page allocations in and out of local memory more efficiently.
-	ETextureCreateFlags TextureCreateFlags = CVarVSMReservedResource.GetValueOnRenderThread() && GRHISupportsReservedResources
+	ETextureCreateFlags TextureCreateFlags = (CVarVSMReservedResource.GetValueOnRenderThread() && GRHISupportsReservedResources)
 		? (TexCreate_ReservedResource | TexCreate_ImmediateCommit)
 		: TexCreate_None;
 
