Enable high quality refelction by default on some platforms.

diff --git a/Engine/Config/BaseScalability.ini b/Engine/Config/BaseScalability.ini
--- a/Engine/Config/BaseScalability.ini
+++ b/Engine/Config/BaseScalability.ini
@@ -705,6 +705,7 @@ r.SkyAtmosphere.TransmittanceLUT.SampleCount=10.0
 r.SkyAtmosphere.MultiScatteringLUT.SampleCount=15.0
 r.SkyLight.RealTimeReflectionCapture=1
 fx.Niagara.QualityLevel=3
+r.Refraction.OffsetQuality=1
 
 [EffectsQuality@Cine]
 r.TranslucencyLightingVolumeDim=64
@@ -733,6 +734,7 @@ r.SkyAtmosphere.TransmittanceLUT.SampleCount=30.0
 r.SkyAtmosphere.MultiScatteringLUT.SampleCount=20.0
 r.SkyLight.RealTimeReflectionCapture=1
 fx.Niagara.QualityLevel=4
+r.Refraction.OffsetQuality=1
 
 ;-----------------------------------------------------------------------------------------------------------------
 
diff --git a/Engine/Source/Runtime/Renderer/Private/DistortionRendering.cpp b/Engine/Source/Runtime/Renderer/Private/DistortionRendering.cpp
--- a/Engine/Source/Runtime/Renderer/Private/DistortionRendering.cpp
+++ b/Engine/Source/Runtime/Renderer/Private/DistortionRendering.cpp
@@ -17,6 +17,7 @@
 #include "Strata/Strata.h"
 #include "ScreenRendering.h"
 #include "PostProcess/TemporalAA.h"
+#include "ShaderPlatformCachedIniValue.h"
 
 DECLARE_GPU_DRAWCALL_STAT(Distortion);
 
@@ -63,6 +64,12 @@ END_GLOBAL_SHADER_PARAMETER_STRUCT()
 
 IMPLEMENT_STATIC_UNIFORM_BUFFER_STRUCT(FDistortionPassUniformParameters, "DistortionPass", SceneTextures);
 
+uint32 GetRefractionOffsetQuality(EShaderPlatform InPlatform)
+{
+	static FShaderPlatformCachedIniValue<int32> CVarRefractionOffsetQualityPlatform(TEXT("r.Refraction.OffsetQuality"));
+	return CVarRefractionOffsetQualityPlatform.Get(InPlatform);
+}
+
 int32 FSceneRenderer::GetRefractionQuality(const FSceneViewFamily& ViewFamily)
 {
 	static const auto ICVar = IConsoleManager::Get().FindTConsoleVariableDataInt(TEXT("r.RefractionQuality"));
@@ -686,7 +693,7 @@ void FDeferredShadingSceneRenderer::RenderDistortion(
 		DistortionTexture = GraphBuilder.CreateTexture(
 			FRDGTextureDesc::Create2D(
 				SceneDepthTexture->Desc.Extent,
-				CVarRefractionOffsetQuality.GetValueOnRenderThread() > 0 ? PF_FloatRGBA : PF_B8G8R8A8,
+				GetRefractionOffsetQuality(ShaderPlatform) > 0 ? PF_FloatRGBA : PF_B8G8R8A8,
 				FClearValueBinding::Transparent,
 				GFastVRamConfig.Distortion | TexCreate_RenderTargetable | TexCreate_ShaderResource,
 				1,
@@ -982,7 +989,7 @@ void FDistortionMeshProcessor::CollectPSOInitializers(const FSceneTexturesConfig
 		
 	FGraphicsPipelineRenderTargetsInfo RenderTargetsInfo;
 	RenderTargetsInfo.NumSamples = 1;
-	AddRenderTargetInfo(CVarRefractionOffsetQuality.GetValueOnAnyThread() > 0 ? PF_FloatRGBA : PF_B8G8R8A8, GFastVRamConfig.Distortion | TexCreate_RenderTargetable | TexCreate_ShaderResource, RenderTargetsInfo);
+	AddRenderTargetInfo(GetRefractionOffsetQuality(SceneTexturesConfig.ShaderPlatform) > 0 ? PF_FloatRGBA : PF_B8G8R8A8, GFastVRamConfig.Distortion | TexCreate_RenderTargetable | TexCreate_ShaderResource, RenderTargetsInfo);
 	if (GetUseRoughRefraction())
 	{
 		AddRenderTargetInfo(PF_R16F, GFastVRamConfig.Distortion | TexCreate_RenderTargetable | TexCreate_ShaderResource, RenderTargetsInfo);
