Removed unused HZB from clouds.

diff --git a/Engine/Shaders/Private/VolumetricCloud.usf b/Engine/Shaders/Private/VolumetricCloud.usf
--- a/Engine/Shaders/Private/VolumetricCloud.usf
+++ b/Engine/Shaders/Private/VolumetricCloud.usf
@@ -180,53 +180,6 @@ void ConvertCloudPixelMaterialInputsToWorkingColorSpace(inout FPixelMaterialInpu
 
 
 
-////////////////////////////////////////////////////////////
-// When we get HZB.hlsl from Reverb, remove that and use common code. 
-// Keeping duplicated code for now to avoid merge pain.
-////////////////////////////////////////////////////////////
-
-// Adapted from HZBOcclusion.usf.
-// Rect contains screen coordinates as {minx, miny, maxx, maxy}. Those coordinates are not clamped to screen edges so they mnight be outside the [0,1]^2.
-bool ViewHierarchicalZBufferMinDepth(float4 ScreenUVMinMax, const float MaxZDepth)
-{
-	float4 Rect = saturate(ScreenUVMinMax);
-	float4 RectPixels = Rect * RenderVolumetricCloudParameters.HZBSize.xyxy;
-
-	float2 RectSize = (RectPixels.zw - RectPixels.xy);
-	float Level = max(ceil(log2(max(RectSize.x, RectSize.y))), RenderVolumetricCloudParameters.HZBUvFactor.z);
-
-	// Check if we can drop one level lower
-	float LevelLower = max(Level - 1, 0);
-	float4 LowerRect = RectPixels * exp2(-LevelLower);
-	float2 LowerRectSize = ceil(LowerRect.zw) - floor(LowerRect.xy);
-	if (all(LowerRectSize <= 4))
-	{
-		Level = LevelLower;
-	}
-
-	// 4x4 samples
-	float2 Scale = RenderVolumetricCloudParameters.HZBUvFactor.xy * (Rect.zw - Rect.xy) / 3;
-	float2 Bias = RenderVolumetricCloudParameters.HZBUvFactor.xy * Rect.xy;
-
-	float4 MinDepth = 1;
-	/*UNROLL*/ for (int i = 0; i < 4; i++)
-	{
-		// TODO could vectorize this
-		float4 Depth;
-		Depth.x = RenderVolumetricCloudParameters.HZBTexture.SampleLevel(RenderVolumetricCloudParameters.HZBSampler, float2(i, 0) * Scale + Bias, Level).r;
-		Depth.y = RenderVolumetricCloudParameters.HZBTexture.SampleLevel(RenderVolumetricCloudParameters.HZBSampler, float2(i, 1) * Scale + Bias, Level).r;
-		Depth.z = RenderVolumetricCloudParameters.HZBTexture.SampleLevel(RenderVolumetricCloudParameters.HZBSampler, float2(i, 2) * Scale + Bias, Level).r;
-		Depth.w = RenderVolumetricCloudParameters.HZBTexture.SampleLevel(RenderVolumetricCloudParameters.HZBSampler, float2(i, 3) * Scale + Bias, Level).r;
-		MinDepth = min(MinDepth, Depth);
-	}
-	MinDepth.x = min(min(MinDepth.x, MinDepth.y), min(MinDepth.z, MinDepth.w));
-
-	// Inverted Z buffer
-	//return RectMax.z >= MinDepth.x ? true : false;
-	return MaxZDepth >= MinDepth.x ? true : false;
-}
-
-
 ////////////////////////////////////////////////////////////
 // Single full screen triangle vertex shader
 ////////////////////////////////////////////////////////////
diff --git a/Engine/Source/Runtime/Renderer/Private/VolumetricCloudRendering.cpp b/Engine/Source/Runtime/Renderer/Private/VolumetricCloudRendering.cpp
--- a/Engine/Source/Runtime/Renderer/Private/VolumetricCloudRendering.cpp
+++ b/Engine/Source/Runtime/Renderer/Private/VolumetricCloudRendering.cpp
@@ -67,11 +67,6 @@ static TAutoConsoleVariable<int32> CVarVolumetricCloudHighQualityAerialPerspecti
 	TEXT("Enable/disable a second pass to trace the aerial perspective per pixel on clouds instead of using the aerial persepctive texture. Only usable when r.VolumetricCloud.EnableAerialPerspectiveSampling=1 and only needed for extra quality when r.VolumetricRenderTarget=1."),
 	ECVF_RenderThreadSafe | ECVF_Scalability);
 
-static TAutoConsoleVariable<int32> CVarVolumetricCloudHzbCulling(
-	TEXT("r.VolumetricCloud.HzbCulling"), 1,
-	TEXT("Enable/disable the use of the HZB in order to not trace behind opaque surfaces. Should be disabled when r.VolumetricRenderTarget.Mode is 2."),
-	ECVF_Scalability);
-
 static TAutoConsoleVariable<int32> CVarVolumetricCloudDisableCompute(
 	TEXT("r.VolumetricCloud.DisableCompute"), 0,
 	TEXT("Do not use compute shader for cloud tracing."),
@@ -596,14 +591,8 @@ BEGIN_GLOBAL_SHADER_PARAMETER_STRUCT(FRenderVolumetricCloudGlobalParameters, )
 	SHADER_PARAMETER(uint32, VolumetricRenderTargetMode)
 	SHADER_PARAMETER(uint32, CloudDebugViewMode)
 	SHADER_PARAMETER(uint32, IsReflectionRendering)
-	SHADER_PARAMETER(uint32, HasValidHZB)
-	SHADER_PARAMETER(uint32, ClampRayTToDepthBufferPostHZB)
 	SHADER_PARAMETER(uint32, TraceShadowmap)
 	SHADER_PARAMETER(float, LocalLightsShadowSampleCount)
-	SHADER_PARAMETER(FVector3f, HZBUvFactor)
-	SHADER_PARAMETER(FVector4f, HZBSize)
-	SHADER_PARAMETER_RDG_TEXTURE(Texture2D<float>, HZBTexture)
-	SHADER_PARAMETER_SAMPLER(SamplerState, HZBSampler)
 	SHADER_PARAMETER(FVector4f, OutputSizeInvSize)
 	SHADER_PARAMETER(int32, StepSizeOnZeroConservativeDensity)
 	SHADER_PARAMETER(int32, EnableAerialPerspectiveSampling)
@@ -657,11 +646,6 @@ void SetupDefaultRenderVolumetricCloudGlobalParameters(FRDGBuilder& GraphBuilder
 	VolumetricCloudParams.VolumetricRenderTargetMode = ViewInfo.ViewState ? ViewInfo.ViewState->VolumetricCloudRenderTarget.GetMode() : 0;
 	VolumetricCloudParams.CloudDebugViewMode = GetVolumetricCloudDebugViewMode(ViewInfo.Family->EngineShowFlags);
 
-	VolumetricCloudParams.HasValidHZB = 0;
-	VolumetricCloudParams.ClampRayTToDepthBufferPostHZB = 0;
-	VolumetricCloudParams.HZBTexture = BlackDummy;
-	VolumetricCloudParams.HZBSampler = TStaticSamplerState<SF_Point, AM_Clamp, AM_Clamp, AM_Clamp>::GetRHI();
-
 	VolumetricCloudParams.StepSizeOnZeroConservativeDensity = FMath::Max(CVarVolumetricCloudStepSizeOnZeroConservativeDensity.GetValueOnRenderThread(), 1);
 
 	VolumetricCloudParams.EmptySpaceSkippingSliceDepth = GetEmptySpaceSkippingSliceDepth();
@@ -682,30 +666,6 @@ void SetupDefaultRenderVolumetricCloudGlobalParameters(FRDGBuilder& GraphBuilder
 	SetupSceneTextureUniformParameters(GraphBuilder, ViewInfo.GetSceneTexturesChecked(), ViewInfo.FeatureLevel, ESceneTextureSetupMode::CustomDepth, VolumetricCloudParams.SceneTextures);
 }
 
-static void SetupRenderVolumetricCloudGlobalParametersHZB(FRDGBuilder& GraphBuilder, const FViewInfo& ViewInfo, FRenderVolumetricCloudGlobalParameters& ShaderParameters)
-{
-	ShaderParameters.HasValidHZB = (ViewInfo.HZB && CVarVolumetricCloudHzbCulling.GetValueOnAnyThread() > 0) ? 1 : 0;
-
-	ShaderParameters.HZBTexture = ShaderParameters.HasValidHZB ? ViewInfo.HZB : GSystemTextures.GetBlackDummy(GraphBuilder);
-	ShaderParameters.HZBSampler = TStaticSamplerState<SF_Point, AM_Clamp, AM_Clamp, AM_Clamp>::GetRHI();
-
-	const float kHZBTestMaxMipmap = 9.0f;
-	const float HZBMipmapCounts = FMath::Log2(FMath::Max<float>(ViewInfo.HZBMipmap0Size.X, ViewInfo.HZBMipmap0Size.Y));
-	const FVector3f HZBUvFactor(
-		float(ViewInfo.ViewRect.Width()) / float(2 * ViewInfo.HZBMipmap0Size.X),
-		float(ViewInfo.ViewRect.Height()) / float(2 * ViewInfo.HZBMipmap0Size.Y),
-		FMath::Max(HZBMipmapCounts - kHZBTestMaxMipmap, 0.0f)
-	);
-	const FVector4f HZBSize(
-		ViewInfo.HZBMipmap0Size.X,
-		ViewInfo.HZBMipmap0Size.Y,
-		1.0f / float(ViewInfo.HZBMipmap0Size.X),
-		1.0f / float(ViewInfo.HZBMipmap0Size.Y)
-	);
-	ShaderParameters.HZBUvFactor = HZBUvFactor;
-	ShaderParameters.HZBSize = HZBSize;
-}
-
 //////////////////////////////////////////////////////////////////////////
 
 class FRenderVolumetricCloudVS : public FMeshMaterialShader
@@ -2207,17 +2167,12 @@ static TRDGUniformBufferRef<FRenderVolumetricCloudGlobalParameters> CreateCloudP
 	const FRDGTextureDesc& Desc = CloudRC.RenderTargets.Output[0].GetTexture()->Desc;
 	VolumetricCloudParams.OutputSizeInvSize = FVector4f(Desc.Extent.X, Desc.Extent.Y, 1.0f / Desc.Extent.X, 1.0f / Desc.Extent.Y);
 
-	SetupRenderVolumetricCloudGlobalParametersHZB(GraphBuilder, MainView, VolumetricCloudParams);
-
 	if (CloudRC.bIsSkyRealTimeReflectionRendering)
 	{
 		VolumetricCloudParams.FogStruct.ApplyVolumetricFog = 0;		// No valid camera froxel volume available.
 		VolumetricCloudParams.OpaqueIntersectionMode = 0;			// No depth buffer is available
-		VolumetricCloudParams.HasValidHZB = 0;						// No valid HZB is available
 	}
 
-	VolumetricCloudParams.ClampRayTToDepthBufferPostHZB = CloudRC.bShouldViewRenderVolumetricRenderTarget ? 0 : 1;
-
 	return GraphBuilder.CreateUniformBuffer(&VolumetricCloudParams);
 }
 
