Adding Renderer interface for drawing screen rectangles.
Changing plugins that used the Renderer/Private DrawRectangle types to use the new interface.

diff --git a/Engine/Plugins/Experimental/RemoteSession/Source/RemoteSession/Private/Channels/RemoteSessionARCameraChannel.cpp b/Engine/Plugins/Experimental/RemoteSession/Source/RemoteSession/Private/Channels/RemoteSessionARCameraChannel.cpp
--- a/Engine/Plugins/Experimental/RemoteSession/Source/RemoteSession/Private/Channels/RemoteSessionARCameraChannel.cpp
+++ b/Engine/Plugins/Experimental/RemoteSession/Source/RemoteSession/Private/Channels/RemoteSessionARCameraChannel.cpp
@@ -13,7 +13,6 @@
 #include "Materials/MaterialRenderProxy.h"
 #include "MaterialDomain.h"
 #include "DataDrivenShaderPlatformInfo.h"
-#include "PostProcess/PostProcessMaterial.h"
 
 #include "ARTextures.h"
 #include "ARBlueprintLibrary.h"
@@ -22,7 +21,6 @@
 #include "IImageWrapper.h"
 #include "IImageWrapperModule.h"
 #include "UObject/Package.h"
-#include "PostProcess/SceneFilterRendering.h"
 
 #include "MaterialShader.h"
 #include "Containers/DynamicRHIResourceArray.h"
@@ -30,6 +28,11 @@
 #include "ScreenPass.h"
 #include "SceneRenderTargetParameters.h"
 
+#include "PostProcess/DrawRectangle.h"
+
+//#include "PostProcess/PostProcessMaterial.h"
+//#include "PostProcess/SceneFilterRendering.h"
+
 #define CAMERA_MESSAGE_ADDRESS TEXT("/ARCamera")
 
 TAutoConsoleVariable<int32> CVarJPEGQuality(
@@ -53,8 +56,11 @@ TAutoConsoleVariable<int32> CVarJPEGGpu(
 class FPostProcessMaterialShader : public FMaterialShader
 {
 public:
-	using FParameters = FPostProcessMaterialParameters;
-	SHADER_USE_PARAMETER_STRUCT_WITH_LEGACY_BASE(FPostProcessMaterialShader, FMaterialShader);
+	FPostProcessMaterialShader() = default;
+	FPostProcessMaterialShader(const ShaderMetaType::CompiledShaderInitializerType& Initializer)
+		: FMaterialShader(Initializer)
+	{
+	}
 
 	static bool ShouldCompilePermutation(const FMaterialShaderPermutationParameters& Parameters)
 	{
@@ -289,15 +295,8 @@ void FARCameraSceneViewExtension::RenderARCamera_RenderThread(FRDGBuilder& Graph
 
 		SetGraphicsPipelineState(RHICmdList, GraphicsPSOInit, 0);
 
-		const FIntPoint ViewSize = InView.UnconstrainedViewRect.Size();
-		FDrawRectangleParameters Parameters;
-		Parameters.PosScaleBias = FVector4f(ViewSize.X, ViewSize.Y, 0, 0);
-		Parameters.UVScaleBias = FVector4f(1.0f, 1.0f, 0.0f, 0.0f);
-		Parameters.InvTargetSizeAndTextureSize = FVector4f(
-			1.0f / ViewSize.X, 1.0f / ViewSize.Y,
-			1.0f, 1.0f);
+		UE::Renderer::PostProcess::SetDrawRectangleParameters(RHICmdList, VertexShader, InView);
 
-		SetUniformBufferParameterImmediate(RHICmdList, VertexShader.GetVertexShader(), VertexShader->GetUniformBufferParameter<FDrawRectangleParameters>(), Parameters);
 		VertexShader->SetParameters(RHICmdList, InView);
 		PixelShader->SetParameters(RHICmdList, InView, MaterialProxy);
 
diff --git a/Engine/Plugins/Experimental/RemoteSession/Source/RemoteSession/RemoteSession.Build.cs b/Engine/Plugins/Experimental/RemoteSession/Source/RemoteSession/RemoteSession.Build.cs
--- a/Engine/Plugins/Experimental/RemoteSession/Source/RemoteSession/RemoteSession.Build.cs
+++ b/Engine/Plugins/Experimental/RemoteSession/Source/RemoteSession/RemoteSession.Build.cs
@@ -10,7 +10,6 @@ public class RemoteSession : ModuleRules
 
 		PrivateIncludePaths.AddRange(
 			new string[] {
-				System.IO.Path.Combine(GetModuleDirectory("Renderer"), "Private"),
 			}
 		);
 			
diff --git a/Engine/Plugins/Runtime/AR/AppleAR/AppleARKit/Source/AppleARKit/AppleARKit.Build.cs b/Engine/Plugins/Runtime/AR/AppleAR/AppleARKit/Source/AppleARKit/AppleARKit.Build.cs
--- a/Engine/Plugins/Runtime/AR/AppleAR/AppleARKit/Source/AppleARKit/AppleARKit.Build.cs
+++ b/Engine/Plugins/Runtime/AR/AppleAR/AppleARKit/Source/AppleARKit/AppleARKit.Build.cs
@@ -9,7 +9,6 @@ public class AppleARKit : ModuleRules
 	{
 		PrivateIncludePaths.AddRange(new string[]
 		{
-			System.IO.Path.Combine(GetModuleDirectory("Renderer"), "Private"),
 		});
 			
 		
diff --git a/Engine/Plugins/Runtime/AR/AppleAR/AppleARKit/Source/AppleARKit/Private/AppleARKitVideoOverlay.cpp b/Engine/Plugins/Runtime/AR/AppleAR/AppleARKit/Source/AppleARKit/Private/AppleARKitVideoOverlay.cpp
--- a/Engine/Plugins/Runtime/AR/AppleAR/AppleARKit/Source/AppleARKit/Private/AppleARKitVideoOverlay.cpp
+++ b/Engine/Plugins/Runtime/AR/AppleAR/AppleARKit/Source/AppleARKit/Private/AppleARKitVideoOverlay.cpp
@@ -23,11 +23,10 @@
 #include "RendererInterface.h"
 #include "ScreenRendering.h"
 #include "Containers/DynamicRHIResourceArray.h"
-#include "PostProcess/SceneFilterRendering.h"
-#include "PostProcess/PostProcessMaterial.h"
 #include "CommonRenderResources.h"
 #include "ARUtilitiesFunctionLibrary.h"
 #include "UObject/Package.h"
+#include "PostProcess/DrawRectangle.h"
 
 #if SUPPORTS_ARKIT_1_0
 	#include "IOSAppDelegate.h"
@@ -125,8 +124,11 @@ template <bool bIsMobileRenderer>
 class TPostProcessMaterialShader : public FMaterialShader
 {
 public:
-	using FParameters = FPostProcessMaterialParameters;
-	SHADER_USE_PARAMETER_STRUCT_WITH_LEGACY_BASE(TPostProcessMaterialShader, FMaterialShader);
+	TPostProcessMaterialShader() = default;
+	TPostProcessMaterialShader(const ShaderMetaType::CompiledShaderInitializerType& Initializer)
+		: FMaterialShader(Initializer)
+	{
+	}
 
 	static bool ShouldCompilePermutation(const FMaterialShaderPermutationParameters& Parameters)
 	{
@@ -392,32 +394,22 @@ void FAppleARKitVideoOverlay::RenderVideoOverlayWithMaterial(FRHICommandListImme
 
 	SetGraphicsPipelineState(RHICmdList, GraphicsPSOInit, 0);
 
-	const FIntPoint ViewSize = InView.UnconstrainedViewRect.Size();
-	FDrawRectangleParameters Parameters;
-	Parameters.PosScaleBias = FVector4f(ViewSize.X, ViewSize.Y, 0, 0);
-	Parameters.UVScaleBias = FVector4f(1.0f, 1.0f, 0.0f, 0.0f);
-	Parameters.InvTargetSizeAndTextureSize = FVector4f(
-													  1.0f / ViewSize.X, 1.0f / ViewSize.Y,
-													  1.0f, 1.0f);
+	UE::Renderer::PostProcess::SetDrawRectangleParameters(RHICmdList, VertexShader, InView);
 
 	if (bIsMobileRenderer)
 	{
-		FARKitCameraOverlayMobileVS* const VertexShaderPtr = static_cast<FARKitCameraOverlayMobileVS*>(VertexShader.GetShader());
-		check(VertexShaderPtr != nullptr);
-		SetUniformBufferParameterImmediate(RHICmdList, VertexShader.GetVertexShader(), VertexShaderPtr->GetUniformBufferParameter<FDrawRectangleParameters>(), Parameters);
+		TShaderRef<FARKitCameraOverlayMobileVS> VertexShaderPtr(static_cast<FARKitCameraOverlayMobileVS*>(VertexShader.GetShader()), *MaterialShaderMap);
+		TShaderRef<FARKitCameraOverlayMobilePS> PixelShaderPtr(static_cast<FARKitCameraOverlayMobilePS*>(PixelShader.GetShader()), *MaterialShaderMap);
+
 		VertexShaderPtr->SetParameters(RHICmdList, InView);
-		FARKitCameraOverlayMobilePS* PixelShaderPtr = static_cast<FARKitCameraOverlayMobilePS*>(PixelShader.GetShader());
-		check(PixelShaderPtr != nullptr);
 		PixelShaderPtr->SetParameters(RHICmdList, InView, MaterialProxy, CameraMaterial);
 	}
 	else
 	{
-		FARKitCameraOverlayVS* const VertexShaderPtr = static_cast<FARKitCameraOverlayVS*>(VertexShader.GetShader());
-		check(VertexShaderPtr != nullptr);
-		SetUniformBufferParameterImmediate(RHICmdList, VertexShader.GetVertexShader(), VertexShaderPtr->GetUniformBufferParameter<FDrawRectangleParameters>(), Parameters);
+		TShaderRef<FARKitCameraOverlayVS> VertexShaderPtr(static_cast<FARKitCameraOverlayVS*>(VertexShader.GetShader()), *MaterialShaderMap);
+		TShaderRef<FARKitCameraOverlayPS> PixelShaderPtr(static_cast<FARKitCameraOverlayPS*>(PixelShader.GetShader()), *MaterialShaderMap);
+
 		VertexShaderPtr->SetParameters(RHICmdList, InView);
-		FARKitCameraOverlayPS* PixelShaderPtr = static_cast<FARKitCameraOverlayPS*>(PixelShader.GetShader());
-		check(PixelShaderPtr != nullptr);
 		PixelShaderPtr->SetParameters(RHICmdList, InView, MaterialProxy, CameraMaterial);
 	}
 	
diff --git a/Engine/Plugins/Runtime/AR/Google/GoogleARCore/Source/GoogleARCoreRendering/GoogleARCoreRendering.Build.cs b/Engine/Plugins/Runtime/AR/Google/GoogleARCore/Source/GoogleARCoreRendering/GoogleARCoreRendering.Build.cs
--- a/Engine/Plugins/Runtime/AR/Google/GoogleARCore/Source/GoogleARCoreRendering/GoogleARCoreRendering.Build.cs
+++ b/Engine/Plugins/Runtime/AR/Google/GoogleARCore/Source/GoogleARCoreRendering/GoogleARCoreRendering.Build.cs
@@ -10,7 +10,6 @@ public class GoogleARCoreRendering : ModuleRules
 		{
 			"GoogleARCoreRendering/Private",
 			"GoogleARCoreBase/Private",
-			System.IO.Path.Combine(GetModuleDirectory("Renderer"), "Private"),
 		});
 			
 		PrivateDependencyModuleNames.AddRange(new string[]
diff --git a/Engine/Plugins/Runtime/AR/Google/GoogleARCore/Source/GoogleARCoreRendering/Private/GoogleARCorePassthroughCameraRenderer.cpp b/Engine/Plugins/Runtime/AR/Google/GoogleARCore/Source/GoogleARCoreRendering/Private/GoogleARCorePassthroughCameraRenderer.cpp
--- a/Engine/Plugins/Runtime/AR/Google/GoogleARCore/Source/GoogleARCoreRendering/Private/GoogleARCorePassthroughCameraRenderer.cpp
+++ b/Engine/Plugins/Runtime/AR/Google/GoogleARCore/Source/GoogleARCoreRendering/Private/GoogleARCorePassthroughCameraRenderer.cpp
@@ -12,8 +12,6 @@
 #include "Engine/Engine.h"
 #include "SceneUtils.h"
 #include "Engine/TextureRenderTarget2D.h"
-#include "PostProcess/SceneFilterRendering.h"
-#include "PostProcess/PostProcessMaterial.h"
 #include "Materials/MaterialRenderProxy.h"
 #include "MaterialDomain.h"
 #include "MaterialShader.h"
@@ -101,8 +99,11 @@ void FGoogleARCorePassthroughCameraRenderer::InitializeRenderer_RenderThread(FSc
 class FPostProcessMaterialShader : public FMaterialShader
 {
 public:
-	using FParameters = FPostProcessMaterialParameters;
-	SHADER_USE_PARAMETER_STRUCT_WITH_LEGACY_BASE(FPostProcessMaterialShader, FMaterialShader);
+	FPostProcessMaterialShader() = default;
+	FPostProcessMaterialShader(const ShaderMetaType::CompiledShaderInitializerType& Initializer)
+		: FMaterialShader(Initializer)
+	{
+	}
 
 	static bool ShouldCompilePermutation(const FMaterialShaderPermutationParameters& Parameters)
 	{
@@ -229,17 +230,7 @@ void FGoogleARCorePassthroughCameraRenderer::RenderVideoOverlayWithMaterial(FRHI
 		VertexShader->SetParameters(RHICmdList, InView);
 		PixelShader->SetParameters(RHICmdList, InView, MaterialProxy, CameraMaterial);
 
-		FIntPoint ViewSize = InView.UnscaledViewRect.Size();
-
-		FDrawRectangleParameters Parameters;
-		Parameters.PosScaleBias = FVector4f(ViewSize.X, ViewSize.Y, 0, 0);
-		Parameters.UVScaleBias = FVector4f(1.0f, 1.0f, 0.0f, 0.0f);
-
-		Parameters.InvTargetSizeAndTextureSize = FVector4f(
-			1.0f / ViewSize.X, 1.0f / ViewSize.Y,
-			1.0f, 1.0f);
-
-		SetUniformBufferParameterImmediate(RHICmdList, VertexShader.GetVertexShader(), VertexShader->GetUniformBufferParameter<FDrawRectangleParameters>(), Parameters);
+		UE::Renderer::PostProcess::SetDrawRectangleParameters(RHICmdList, VertexShader, InView);
 
 		if (OverlayVertexBufferRHI.IsValid() && OverlayIndexBufferRHI.IsValid())
 		{
diff --git a/Engine/Plugins/Runtime/nDisplay/Source/DisplayCluster/Private/Render/Viewport/DisplayClusterViewportProxy.cpp b/Engine/Plugins/Runtime/nDisplay/Source/DisplayCluster/Private/Render/Viewport/DisplayClusterViewportProxy.cpp
--- a/Engine/Plugins/Runtime/nDisplay/Source/DisplayCluster/Private/Render/Viewport/DisplayClusterViewportProxy.cpp
+++ b/Engine/Plugins/Runtime/nDisplay/Source/DisplayCluster/Private/Render/Viewport/DisplayClusterViewportProxy.cpp
@@ -47,6 +47,7 @@
 #include "SceneRendering.h"
 #include "RenderGraphUtils.h"
 #include "RenderGraphResources.h"
+#include "PostProcess/DrawRectangle.h"
 
 #include "ScreenPass.h"
 #include "PostProcess/PostProcessing.h"
@@ -825,16 +826,15 @@ void ResampleCopyTextureImpl_RenderThread(FRHICommandListImmediate& RHICmdList,
 
 		// Set up vertex uniform parameters for scaling and biasing the rectangle.
 		// Note: Use DrawRectangle in the vertex shader to calculate the correct vertex position and uv.
-		FDrawRectangleParameters Parameters;
-		{
-			Parameters.PosScaleBias = FVector4f(DstRect.Size().X, DstRect.Size().Y, DstRect.Min.X, DstRect.Min.Y);
-			Parameters.UVScaleBias = FVector4f(SrcRect.Size().X, SrcRect.Size().Y, SrcRect.Min.X, SrcRect.Min.Y);
-			Parameters.InvTargetSizeAndTextureSize = FVector4f(1.0f / DstSize.X, 1.0f / DstSize.Y, 1.0f / SrcSize.X, 1.0f / SrcSize.Y);
-
-			SetUniformBufferParameterImmediate(RHICmdList, VertexShader.GetVertexShader(), VertexShader->GetUniformBufferParameter<FDrawRectangleParameters>(), Parameters);
-		}
 
-		FPixelShaderUtils::DrawFullscreenQuad(RHICmdList, 1);
+		UE::Renderer::PostProcess::DrawRectangle(
+			RHICmdList, VertexShader,
+			DstRect.Min.X, DstRect.Min.Y,
+			DstRect.Size().X, DstRect.Size().Y,
+			SrcRect.Min.X, SrcRect.Min.Y,
+			SrcRect.Size().X, SrcRect.Size().Y,
+			DstSize, SrcSize
+			);
 	}
 	RHICmdList.EndRenderPass();
 	RHICmdList.Transition(FRHITransitionInfo(DstTexture, ERHIAccess::Unknown, ERHIAccess::SRVMask));
diff --git a/Engine/Plugins/Runtime/nDisplay/Source/DisplayClusterMedia/Private/DisplayClusterMediaHelpers.cpp b/Engine/Plugins/Runtime/nDisplay/Source/DisplayClusterMedia/Private/DisplayClusterMediaHelpers.cpp
--- a/Engine/Plugins/Runtime/nDisplay/Source/DisplayClusterMedia/Private/DisplayClusterMediaHelpers.cpp
+++ b/Engine/Plugins/Runtime/nDisplay/Source/DisplayClusterMedia/Private/DisplayClusterMediaHelpers.cpp
@@ -5,7 +5,7 @@
 #include "CommonRenderResources.h"
 #include "PixelShaderUtils.h"
 #include "ScreenRendering.h"
-#include "PostProcess/SceneFilterRendering.h"
+#include "PostProcess/DrawRectangle.h"
 
 
 namespace DisplayClusterMediaHelpers
@@ -62,14 +62,15 @@ namespace DisplayClusterMediaHelpers
 
 			// Set up vertex uniform parameters for scaling and biasing the rectangle.
 			// Note: Use DrawRectangle in the vertex shader to calculate the correct vertex position and uv.
-			FDrawRectangleParameters Parameters;
-			Parameters.PosScaleBias = FVector4f(DstRect.Size().X, DstRect.Size().Y, DstRect.Min.X, DstRect.Min.Y);
-			Parameters.UVScaleBias = FVector4f(SrcRect.Size().X, SrcRect.Size().Y, SrcRect.Min.X, SrcRect.Min.Y);
-			Parameters.InvTargetSizeAndTextureSize = FVector4f(1.0f / DstSize.X, 1.0f / DstSize.Y, 1.0f / SrcSize.X, 1.0f / SrcSize.Y);
 
-			SetUniformBufferParameterImmediate(RHICmdList, VertexShader.GetVertexShader(), VertexShader->GetUniformBufferParameter<FDrawRectangleParameters>(), Parameters);
-
-			FPixelShaderUtils::DrawFullscreenQuad(RHICmdList, 1);
+			UE::Renderer::PostProcess::DrawRectangle(
+				RHICmdList, VertexShader,
+				DstRect.Min.X, DstRect.Min.Y,
+				DstRect.Size().X, DstRect.Size().Y,
+				SrcRect.Min.X, SrcRect.Min.Y,
+				SrcRect.Size().X, SrcRect.Size().Y,
+				DstSize, SrcSize
+			);
 		}
 
 		RHICmdList.EndRenderPass();
diff --git a/Engine/Plugins/Runtime/nDisplay/Source/DisplayClusterShaders/Private/Shaders/DisplayClusterShadersPreprocess_UVLightCards.cpp b/Engine/Plugins/Runtime/nDisplay/Source/DisplayClusterShaders/Private/Shaders/DisplayClusterShadersPreprocess_UVLightCards.cpp
--- a/Engine/Plugins/Runtime/nDisplay/Source/DisplayClusterShaders/Private/Shaders/DisplayClusterShadersPreprocess_UVLightCards.cpp
+++ b/Engine/Plugins/Runtime/nDisplay/Source/DisplayClusterShaders/Private/Shaders/DisplayClusterShadersPreprocess_UVLightCards.cpp
@@ -13,7 +13,6 @@
 #include "PrimitiveSceneProxy.h"
 #include "RenderGraphBuilder.h"
 #include "RenderGraphResources.h"
-//#include "ScenePrivate.h"
 #include "SceneInterface.h"
 #include "SceneRendererInterface.h"
 #include "SceneRenderTargetParameters.h"
diff --git a/Engine/Plugins/VirtualProduction/TextureShare/Source/TextureShare/Private/Resources/TextureShareResourcesProxy.cpp b/Engine/Plugins/VirtualProduction/TextureShare/Source/TextureShare/Private/Resources/TextureShareResourcesProxy.cpp
--- a/Engine/Plugins/VirtualProduction/TextureShare/Source/TextureShare/Private/Resources/TextureShareResourcesProxy.cpp
+++ b/Engine/Plugins/VirtualProduction/TextureShare/Source/TextureShare/Private/Resources/TextureShareResourcesProxy.cpp
@@ -29,7 +29,7 @@
 #include "ShaderParameterUtils.h"
 
 #include "ScreenRendering.h"
-#include "PostProcess/SceneFilterRendering.h"
+#include "PostProcess/DrawRectangle.h"
 
 #include "RenderTargetPool.h"
 
@@ -116,16 +116,15 @@ namespace TextureShareResourcesHelpers
 
 			// Set up vertex uniform parameters for scaling and biasing the rectangle.
 			// Note: Use DrawRectangle in the vertex shader to calculate the correct vertex position and uv.
-			FDrawRectangleParameters Parameters;
-			{
-				Parameters.PosScaleBias = FVector4f(DstRect.Size().X, DstRect.Size().Y, DstRect.Min.X, DstRect.Min.Y);
-				Parameters.UVScaleBias = FVector4f(SrcRect.Size().X, SrcRect.Size().Y, SrcRect.Min.X, SrcRect.Min.Y);
-				Parameters.InvTargetSizeAndTextureSize = FVector4f(1.0f / DstSize.X, 1.0f / DstSize.Y, 1.0f / SrcSize.X, 1.0f / SrcSize.Y);
-
-				SetUniformBufferParameterImmediate(RHICmdList, VertexShader.GetVertexShader(), VertexShader->GetUniformBufferParameter<FDrawRectangleParameters>(), Parameters);
-			}
 
-			FPixelShaderUtils::DrawFullscreenQuad(RHICmdList, 1);
+			UE::Renderer::PostProcess::DrawRectangle(
+				RHICmdList, VertexShader,
+				DstRect.Min.X, DstRect.Min.Y,
+				DstRect.Size().X, DstRect.Size().Y,
+				SrcRect.Min.X, SrcRect.Min.Y,
+				SrcRect.Size().X, SrcRect.Size().Y,
+				DstSize, SrcSize
+			);
 		}
 		RHICmdList.EndRenderPass();
 	}
diff --git a/Engine/Source/Runtime/Renderer/Private/PostProcess/SceneFilterRendering.cpp b/Engine/Source/Runtime/Renderer/Private/PostProcess/SceneFilterRendering.cpp
--- a/Engine/Source/Runtime/Renderer/Private/PostProcess/SceneFilterRendering.cpp
+++ b/Engine/Source/Runtime/Renderer/Private/PostProcess/SceneFilterRendering.cpp
@@ -12,6 +12,7 @@
 #include "IXRTrackingSystem.h"
 #include "PixelShaderUtils.h"
 #include "CommonRenderResources.h"
+#include "PostProcess/DrawRectangle.h"
 
 IMPLEMENT_TYPE_LAYOUT(FGammaShaderParameters);
 
@@ -67,9 +68,6 @@ static TGlobalResource<FTesselatedScreenRectangleIndexBuffer> GTesselatedScreenR
 
 IMPLEMENT_GLOBAL_SHADER_PARAMETER_STRUCT(FDrawRectangleParameters, "DrawRectangleParameters");
 
-typedef TUniformBufferRef<FDrawRectangleParameters> FDrawRectangleBufferRef;
-
-
 static TAutoConsoleVariable<int32> CVarDrawRectangleOptimization(
 	TEXT("r.DrawRectangleOptimization"),
 	1,
@@ -93,9 +91,58 @@ static void DoDrawRectangleFlagOverride(EDrawRectangleFlags& Flags)
 #endif
 }
 
-template <typename TRHICommandList>
-static inline void InternalDrawRectangle(
-	TRHICommandList& RHICmdList,
+namespace UE::Renderer::PostProcess
+{
+	void SetDrawRectangleParameters(
+		FRHICommandList& RHICmdList,
+		const TShaderRef<FShader>& VertexShader,
+		float X,
+		float Y,
+		float SizeX,
+		float SizeY,
+		float U,
+		float V,
+		float SizeU,
+		float SizeV,
+		FIntPoint TargetSize,
+		FIntPoint TextureSize)
+	{
+		// Set up vertex uniform parameters for scaling and biasing the rectangle.
+		// Note: Use DrawRectangle in the vertex shader to calculate the correct vertex position and uv.
+
+		FDrawRectangleParameters Parameters;
+		Parameters.PosScaleBias = FVector4f(SizeX, SizeY, X, Y);
+		Parameters.UVScaleBias = FVector4f(SizeU, SizeV, U, V);
+
+		Parameters.InvTargetSizeAndTextureSize = FVector4f(
+			1.0f / TargetSize.X, 1.0f / TargetSize.Y,
+			1.0f / TextureSize.X, 1.0f / TextureSize.Y);
+
+		SetUniformBufferParameterImmediate(RHICmdList, VertexShader.GetVertexShader(), VertexShader->GetUniformBufferParameter<FDrawRectangleParameters>(), Parameters);
+	}
+
+	void SetDrawRectangleParameters(FRHICommandList& RHICmdList, const TShaderRef<FShader>& VertexShader, const FIntPoint& ViewSize)
+	{
+		SetDrawRectangleParameters(
+			RHICmdList,
+			VertexShader,
+			0.0f, 0.0f,
+			ViewSize.X, ViewSize.Y,
+			0.0f, 0.0f,
+			1.0f, 1.0f,
+			ViewSize,
+			FIntPoint(1.0f, 1.0f)
+		);
+	}
+
+	void SetDrawRectangleParameters(FRHICommandList& RHICmdList, const TShaderRef<FShader>& VertexShader, const FSceneView& View)
+	{
+		SetDrawRectangleParameters(RHICmdList, VertexShader, View.UnconstrainedViewRect.Size());
+	}
+}
+
+inline void InternalDrawRectangle(
+	FRHICommandList& RHICmdList,
 	float X,
 	float Y,
 	float SizeX,
@@ -111,8 +158,6 @@ static inline void InternalDrawRectangle(
 	uint32 InstanceCount
 	)
 {
-	float ClipSpaceQuadZ = 0.0f;
-
 	DoDrawRectangleFlagOverride(Flags);
 
 	// triangle if extending to left and top of the given rectangle, if it's not left top of the viewport it can cause artifacts
@@ -122,18 +167,7 @@ static inline void InternalDrawRectangle(
 		Flags = EDRF_Default;
 	}
 
-	// Set up vertex uniform parameters for scaling and biasing the rectangle.
-	// Note: Use DrawRectangle in the vertex shader to calculate the correct vertex position and uv.
-
-	FDrawRectangleParameters Parameters;
-	Parameters.PosScaleBias = FVector4f(SizeX, SizeY, X, Y);
-	Parameters.UVScaleBias = FVector4f(SizeU, SizeV, U, V);
-
-	Parameters.InvTargetSizeAndTextureSize = FVector4f(
-		1.0f / TargetSize.X, 1.0f / TargetSize.Y,
-		1.0f / TextureSize.X, 1.0f / TextureSize.Y);
-
-	SetUniformBufferParameterImmediate(RHICmdList, VertexShader.GetVertexShader(), VertexShader->GetUniformBufferParameter<FDrawRectangleParameters>(), Parameters);
+	UE::Renderer::PostProcess::SetDrawRectangleParameters(RHICmdList, VertexShader, X, Y, SizeX, SizeY, U, V, SizeU, SizeV, TargetSize, TextureSize);
 
 	if(Flags == EDRF_UseTesselatedIndexBuffer)
 	{
@@ -150,16 +184,13 @@ static inline void InternalDrawRectangle(
 			/*NumInstances=*/ InstanceCount
 			);
 	}
+	else if (Flags == EDRF_UseTriangleOptimization)
+	{
+		FPixelShaderUtils::DrawFullscreenTriangle(RHICmdList, InstanceCount);
+	}
 	else
 	{
-		if (Flags == EDRF_UseTriangleOptimization)
-		{
-			FPixelShaderUtils::DrawFullscreenTriangle(RHICmdList, InstanceCount);
-		}
-		else
-		{
-			FPixelShaderUtils::DrawFullscreenQuad(RHICmdList, InstanceCount);
-		}
+		FPixelShaderUtils::DrawFullscreenQuad(RHICmdList, InstanceCount);
 	}
 }
 
@@ -251,15 +282,11 @@ void DrawHmdMesh(
 	const TShaderRef<FShader>& VertexShader
 	)
 {
-	FDrawRectangleParameters Parameters;
-	Parameters.PosScaleBias = FVector4f(SizeX, SizeY, X, Y);
-	Parameters.UVScaleBias = FVector4f(SizeU, SizeV, U, V);
-
-	Parameters.InvTargetSizeAndTextureSize = FVector4f(
-		1.0f / TargetSize.X, 1.0f / TargetSize.Y,
-		1.0f / TextureSize.X, 1.0f / TextureSize.Y);
-
-	SetUniformBufferParameterImmediate(RHICmdList, VertexShader.GetVertexShader(), VertexShader->GetUniformBufferParameter<FDrawRectangleParameters>(), Parameters);
+	UE::Renderer::PostProcess::SetDrawRectangleParameters(
+		RHICmdList, VertexShader,
+		X, Y, SizeX, SizeY, U, V, SizeU, SizeV,
+		TargetSize, TextureSize
+	);
 
 	if (GEngine->XRSystem->GetHMDDevice())
 	{
@@ -293,3 +320,47 @@ void DrawPostProcessPass(
 		DrawRectangle(RHICmdList, X, Y, SizeX, SizeY, U, V, SizeU, SizeV, TargetSize, TextureSize, VertexShader, Flags);
 	}
 }
+
+namespace UE::Renderer::PostProcess
+{
+	void DrawRectangle(
+		FRHICommandList& RHICmdList,
+		const TShaderRef<FShader>& VertexShader,
+		float X,
+		float Y,
+		float SizeX,
+		float SizeY,
+		float U,
+		float V,
+		float SizeU,
+		float SizeV,
+		FIntPoint TargetSize,
+		FIntPoint TextureSize,
+		EDrawRectangleFlags Flags,
+		uint32 InstanceCount
+	)
+	{
+		::DrawRectangle(RHICmdList, X, Y, SizeX, SizeY, U, V, SizeU, SizeV, TargetSize, TextureSize, VertexShader, Flags, InstanceCount);
+	}
+
+	void DrawPostProcessPass(
+		FRHICommandList& RHICmdList,
+		const TShaderRef<FShader>& VertexShader,
+		float X,
+		float Y,
+		float SizeX,
+		float SizeY,
+		float U,
+		float V,
+		float SizeU,
+		float SizeV,
+		FIntPoint TargetSize,
+		FIntPoint TextureSize,
+		int32 StereoViewIndex,
+		bool bHasCustomMesh,
+		EDrawRectangleFlags Flags
+	)
+	{
+		::DrawPostProcessPass(RHICmdList, X, Y, SizeX, SizeY, U, V, SizeU, SizeV, TargetSize, TextureSize, VertexShader, StereoViewIndex, bHasCustomMesh, Flags);
+	}
+}
diff --git a/Engine/Source/Runtime/Renderer/Private/PostProcess/SceneFilterRendering.h b/Engine/Source/Runtime/Renderer/Private/PostProcess/SceneFilterRendering.h
--- a/Engine/Source/Runtime/Renderer/Private/PostProcess/SceneFilterRendering.h
+++ b/Engine/Source/Runtime/Renderer/Private/PostProcess/SceneFilterRendering.h
@@ -18,7 +18,7 @@
 enum class EStereoscopicPass;
 
 /** Uniform buffer for computing the vertex positional and UV adjustments in the vertex shader. */
-BEGIN_GLOBAL_SHADER_PARAMETER_STRUCT( FDrawRectangleParameters, RENDERER_API)
+BEGIN_GLOBAL_SHADER_PARAMETER_STRUCT( FDrawRectangleParameters, )
 	SHADER_PARAMETER( FVector4f, PosScaleBias )
 	SHADER_PARAMETER( FVector4f, UVScaleBias )
 	SHADER_PARAMETER( FVector4f, InvTargetSizeAndTextureSize )
diff --git a/Engine/Source/Runtime/Renderer/Public/MaterialShader.h b/Engine/Source/Runtime/Renderer/Public/MaterialShader.h
--- a/Engine/Source/Runtime/Renderer/Public/MaterialShader.h
+++ b/Engine/Source/Runtime/Renderer/Public/MaterialShader.h
@@ -10,6 +10,7 @@
 #include "MaterialShared.h"
 #include "SceneView.h"
 #include "Shader.h"
+#include "ShaderParameterUtils.h"
 #include "RHIFwd.h"
 #include "Serialization/MemoryLayout.h"
 
diff --git a/Engine/Source/Runtime/Renderer/Public/PostProcess/DrawRectangle.h b/Engine/Source/Runtime/Renderer/Public/PostProcess/DrawRectangle.h
new file mode 100644
--- /dev/null
+++ b/Engine/Source/Runtime/Renderer/Public/PostProcess/DrawRectangle.h
@@ -0,0 +1,86 @@
+// Copyright Epic Games, Inc. All Rights Reserved.
+
+#pragma once
+
+#include "Math/MathFwd.h"
+#include "RendererInterface.h"
+
+class FRHICommandList;
+class FSceneView;
+class FShaderMapPointerTable;
+
+template<typename ShaderType, typename PointerTableType> class TShaderRefBase;
+template<typename ShaderType> using TShaderRef = TShaderRefBase<ShaderType, FShaderMapPointerTable>;
+
+namespace UE::Renderer::PostProcess
+{
+	RENDERER_API void SetDrawRectangleParameters(
+		FRHICommandList& RHICmdList,
+		const TShaderRef<FShader>& VertexShader,
+		float X,
+		float Y,
+		float SizeX,
+		float SizeY,
+		float U,
+		float V,
+		float SizeU,
+		float SizeV,
+		FIntPoint TargetSize,
+		FIntPoint TextureSize);
+
+	RENDERER_API void SetDrawRectangleParameters(FRHICommandList& RHICmdList, const TShaderRef<FShader>& VertexShader, const FIntPoint& ViewSize);
+	RENDERER_API void SetDrawRectangleParameters(FRHICommandList& RHICmdList, const TShaderRef<FShader>& VertexShader, const FSceneView& View);
+
+	/**
+	 * Draws a quad with the given vertex positions and UVs in denormalized pixel/texel coordinates.
+	 * The platform-dependent mapping from pixels to texels is done automatically.
+	 * Note that the positions are affected by the current viewport.
+	 * NOTE: DrawRectangle should be used in the vertex shader to calculate the correct position and uv for vertices.
+	 * NOTE2: Assumes previously set PSO has PrimitiveType = PT_TriangleList
+	 *
+	 * X, Y							Position in screen pixels of the top left corner of the quad
+	 * SizeX, SizeY					Size in screen pixels of the quad
+	 * U, V							Position in texels of the top left corner of the quad's UV's
+	 * SizeU, SizeV					Size in texels of the quad's UV's
+	 * TargetSizeX, TargetSizeY		Size in screen pixels of the target surface
+	 * TextureSize                  Size in texels of the source texture
+	 * VertexShader					The vertex shader used for rendering
+	 * Flags						see EDrawRectangleFlags
+	 * InstanceCount				Number of instances of rectangle
+	 */
+	RENDERER_API void DrawRectangle(
+		FRHICommandList& RHICmdList,
+		const TShaderRef<FShader>& VertexShader,
+		float X,
+		float Y,
+		float SizeX,
+		float SizeY,
+		float U,
+		float V,
+		float SizeU,
+		float SizeV,
+		FIntPoint TargetSize,
+		FIntPoint TextureSize,
+		EDrawRectangleFlags Flags = EDRF_Default,
+		uint32 InstanceCount = 1
+	);
+
+	// NOTE: Assumes previously set PSO has PrimitiveType = PT_TriangleList
+	RENDERER_API void DrawPostProcessPass(
+		FRHICommandList& RHICmdList,
+		const TShaderRef<FShader>& VertexShader,
+		float X,
+		float Y,
+		float SizeX,
+		float SizeY,
+		float U,
+		float V,
+		float SizeU,
+		float SizeV,
+		FIntPoint TargetSize,
+		FIntPoint TextureSize,
+		int32 StereoViewIndex,
+		bool bHasCustomMesh,
+		EDrawRectangleFlags Flags = EDRF_Default
+	);
+}
