Fixing NonUnity compile errors.

diff --git a/Engine/Source/Runtime/Renderer/Private/PostProcess/SceneFilterRendering.cpp b/Engine/Source/Runtime/Renderer/Private/PostProcess/SceneFilterRendering.cpp
--- a/Engine/Source/Runtime/Renderer/Private/PostProcess/SceneFilterRendering.cpp
+++ b/Engine/Source/Runtime/Renderer/Private/PostProcess/SceneFilterRendering.cpp
@@ -13,6 +13,7 @@
 #include "PixelShaderUtils.h"
 #include "CommonRenderResources.h"
 #include "PostProcess/DrawRectangle.h"
+#include "ScenePrivate.h"
 
 IMPLEMENT_TYPE_LAYOUT(FGammaShaderParameters);
 
