Add a separate permutation for TONEMAPPER_OUTPUT_sRGB which is the most common case in a PostProcessCombineLUT

diff --git a/Engine/Shaders/Private/PostProcessCombineLUTs.usf b/Engine/Shaders/Private/PostProcessCombineLUTs.usf
--- a/Engine/Shaders/Private/PostProcessCombineLUTs.usf
+++ b/Engine/Shaders/Private/PostProcessCombineLUTs.usf
@@ -135,7 +135,11 @@ float OutputMaxLuminance;
 
 uint GetOutputDevice()
 {
+#if OUTPUT_DEVICE_SRGB
+	return TONEMAPPER_OUTPUT_sRGB;
+#else
 	return OutputDevice;
+#endif
 }
 
 float BlueCorrection;
diff --git a/Engine/Source/Runtime/Renderer/Private/PostProcess/PostProcessCombineLUTs.cpp b/Engine/Source/Runtime/Renderer/Private/PostProcess/PostProcessCombineLUTs.cpp
--- a/Engine/Source/Runtime/Renderer/Private/PostProcess/PostProcessCombineLUTs.cpp
+++ b/Engine/Source/Runtime/Renderer/Private/PostProcess/PostProcessCombineLUTs.cpp
@@ -248,7 +248,8 @@ public:
 
 	class FBlendCount : SHADER_PERMUTATION_RANGE_INT("BLENDCOUNT", 1, 5);
 	class FSkipTemperature : SHADER_PERMUTATION_BOOL("SKIP_TEMPERATURE");
-	using FPermutationDomain = TShaderPermutationDomain<FBlendCount, FSkipTemperature>;
+	class FOutputDeviceSRGB : SHADER_PERMUTATION_BOOL("OUTPUT_DEVICE_SRGB");
+	using FPermutationDomain = TShaderPermutationDomain<FBlendCount, FSkipTemperature, FOutputDeviceSRGB>;
 
 	static bool ShouldCompilePermutation(const FGlobalShaderPermutationParameters& Parameters)
 	{
@@ -464,8 +465,10 @@ FRDGTextureRef AddCombineLUTPass(FRDGBuilder& GraphBuilder, const FViewInfo& Vie
 		PassParameters->RWOutputTexture = GraphBuilder.CreateUAV(OutputTexture);
 
 		const bool ShouldSkipTemperature = FMath::IsNearlyEqual(PassParameters->CombineLUT.WhiteTemp, DefaultTemperature) && FMath::IsNearlyEqual(PassParameters->CombineLUT.WhiteTint, DefaultTint);
+		const bool bOutputDeviceSRGB = (PassParameters->CombineLUT.OutputDevice.OutputDevice == (uint32)EDisplayOutputFormat::SDR_sRGB);
 
 		PermutationVector.Set<FLUTBlenderShader::FSkipTemperature>(ShouldSkipTemperature);
+		PermutationVector.Set<FLUTBlenderShader::FOutputDeviceSRGB>(bOutputDeviceSRGB);
 		TShaderMapRef<FLUTBlenderCS> ComputeShader(View.ShaderMap, PermutationVector);
 
 		const uint32 GroupSizeXY = FMath::DivideAndRoundUp(OutputViewSize.X, FLUTBlenderCS::GroupSize);
@@ -485,8 +488,10 @@ FRDGTextureRef AddCombineLUTPass(FRDGBuilder& GraphBuilder, const FViewInfo& Vie
 		PassParameters->RenderTargets[0] = FRenderTargetBinding(OutputTexture, ERenderTargetLoadAction::ENoAction);
 
 		const bool ShouldSkipTemperature = FMath::IsNearlyEqual(PassParameters->CombineLUT.WhiteTemp, DefaultTemperature) && FMath::IsNearlyEqual(PassParameters->CombineLUT.WhiteTint, DefaultTint);
+		const bool bOutputDeviceSRGB = (PassParameters->CombineLUT.OutputDevice.OutputDevice == (uint32)EDisplayOutputFormat::SDR_sRGB);
 
 		PermutationVector.Set<FLUTBlenderShader::FSkipTemperature>(ShouldSkipTemperature);
+		PermutationVector.Set<FLUTBlenderShader::FOutputDeviceSRGB>(bOutputDeviceSRGB);
 		TShaderMapRef<FLUTBlenderPS> PixelShader(View.ShaderMap, PermutationVector);
 
 		GraphBuilder.AddPass(
