Fix bug with VSM page pool (re)creation that was causing cache to get wiped every frame even when enabled.
Add a warning for the specific case.

diff --git a/Engine/Source/Runtime/Renderer/Private/VirtualShadowMaps/VirtualShadowMapCacheManager.cpp b/Engine/Source/Runtime/Renderer/Private/VirtualShadowMaps/VirtualShadowMapCacheManager.cpp
--- a/Engine/Source/Runtime/Renderer/Private/VirtualShadowMaps/VirtualShadowMapCacheManager.cpp
+++ b/Engine/Source/Runtime/Renderer/Private/VirtualShadowMaps/VirtualShadowMapCacheManager.cpp
@@ -615,20 +615,28 @@ TRefCountPtr<IPooledRenderTarget> FVirtualShadowMapArrayCacheManager::SetPhysica
 {
 	// Using ReservedResource|ImmediateCommit flags hint to the RHI that the resource can be allocated using N small physical memory allocations,
 	// instead of a single large contighous allocation. This helps Windows video memory manager page allocations in and out of local memory more efficiently.
-	ETextureCreateFlags TextureCreateFlags = (CVarVSMReservedResource.GetValueOnRenderThread() && GRHISupportsReservedResources)
+	ETextureCreateFlags RequestedCreateFlags = (CVarVSMReservedResource.GetValueOnRenderThread() && GRHISupportsReservedResources)
 		? (TexCreate_ReservedResource | TexCreate_ImmediateCommit)
 		: TexCreate_None;
 
 	if (!PhysicalPagePool 
 		|| PhysicalPagePool->GetDesc().Extent != RequestedSize 
 		|| PhysicalPagePool->GetDesc().ArraySize != RequestedArraySize
-		|| PhysicalPagePool->GetDesc().Flags != TextureCreateFlags)
+		|| PhysicalPagePoolCreateFlags != RequestedCreateFlags)
 	{
+		if (PhysicalPagePool)
+		{
+			UE_LOG(LogRenderer, Display, TEXT("Recreating Shadow.Virtual.PhysicalPagePool due to size or flags change. This will also drop any cached pages."));
+		}
+
+		// Track changes to these ourselves instead of from the GetDesc() since that may get manipulated internally
+		PhysicalPagePoolCreateFlags = RequestedCreateFlags;
+
 		FPooledRenderTargetDesc Desc2D = FPooledRenderTargetDesc::Create2DArrayDesc(
 			RequestedSize,
 			PF_R32_UINT,
 			FClearValueBinding::None,
-			TextureCreateFlags,
+			PhysicalPagePoolCreateFlags,
 			TexCreate_ShaderResource | TexCreate_UAV,
 			false,
 			RequestedArraySize
@@ -636,7 +644,6 @@ TRefCountPtr<IPooledRenderTarget> FVirtualShadowMapArrayCacheManager::SetPhysica
 		GRenderTargetPool.FindFreeElement(GraphBuilder.RHICmdList, Desc2D, PhysicalPagePool, TEXT("Shadow.Virtual.PhysicalPagePool"));
 
 		Invalidate();
-		//UE_LOG(LogRenderer, Display, TEXT("Recreating Shadow.Virtual.PhysicalPagePool. This will also drop any cached pages."));
 	}
 
 	return PhysicalPagePool;
diff --git a/Engine/Source/Runtime/Renderer/Private/VirtualShadowMaps/VirtualShadowMapCacheManager.h b/Engine/Source/Runtime/Renderer/Private/VirtualShadowMaps/VirtualShadowMapCacheManager.h
--- a/Engine/Source/Runtime/Renderer/Private/VirtualShadowMaps/VirtualShadowMapCacheManager.h
+++ b/Engine/Source/Runtime/Renderer/Private/VirtualShadowMaps/VirtualShadowMapCacheManager.h
@@ -330,6 +330,7 @@ private:
 	// we store the physical pool here.
 	TRefCountPtr<IPooledRenderTarget> PhysicalPagePool;
 	TRefCountPtr<IPooledRenderTarget> HZBPhysicalPagePool;
+	ETextureCreateFlags PhysicalPagePoolCreateFlags = TexCreate_None;
 
 	// Index the Cache entries by the light ID
 	TMap< uint64, TSharedPtr<FVirtualShadowMapPerLightCacheEntry> > CacheEntries;
