Implements r.TSR.Subpixel.Method=2 to be able to reproject static geometric detail thiner than a rendering pixel

diff --git a/Engine/Shaders/Private/TemporalSuperResolution/TSRClearPrevTextures.usf b/Engine/Shaders/Private/TemporalSuperResolution/TSRClearPrevTextures.usf
--- a/Engine/Shaders/Private/TemporalSuperResolution/TSRClearPrevTextures.usf
+++ b/Engine/Shaders/Private/TemporalSuperResolution/TSRClearPrevTextures.usf
@@ -14,6 +14,7 @@
 
 RWTexture2D<uint> PrevUseCountOutput;
 RWTexture2D<uint> PrevClosestDepthOutput;
+RWTexture2D<uint> SubpixelDepthOutput;
 
 
 //------------------------------------------------------- ENTRY POINT
@@ -34,11 +35,20 @@ void MainCS(
 		{
 			uint2 PixelOffset = uint2(x, y) * TILE_SIZE;
 			uint2 InputPixelPos = OutputPixelOrigin + PixelOffset;
-			
+			uint2 PrevUseCountPixelPos = InputPixelPos;
+
 			InputPixelPos.x = select(all(InputPixelPos < InputInfo_ViewportMax), InputPixelPos.x, uint(~0));
 
-			PrevUseCountOutput[InputPixelPos] = 0;
+			#if COMPRESS_PREV_USE_COUNT
+				const uint2 PrevUseCountMaxPixelPos = InputInfo_ViewportMin + uint2(InputInfo_ViewportMax - InputInfo_ViewportMin + 1) / 2 + 8;
+			#else
+				const uint2 PrevUseCountMaxPixelPos = InputInfo_ViewportMax;
+			#endif
+			PrevUseCountPixelPos.x = select(all(PrevUseCountPixelPos < PrevUseCountMaxPixelPos), PrevUseCountPixelPos.x, uint(~0));
+
+			PrevUseCountOutput[PrevUseCountPixelPos] = 0;
 			PrevClosestDepthOutput[InputPixelPos] = 0;
+			SubpixelDepthOutput[InputPixelPos] = 0;
 		}
 	}
 }
diff --git a/Engine/Shaders/Private/TemporalSuperResolution/TSRCommon.ush b/Engine/Shaders/Private/TemporalSuperResolution/TSRCommon.ush
--- a/Engine/Shaders/Private/TemporalSuperResolution/TSRCommon.ush
+++ b/Engine/Shaders/Private/TemporalSuperResolution/TSRCommon.ush
@@ -1203,3 +1203,36 @@ tsr_half2 RejectionFactorToBlendFactor(tsr_half2 RejectionFactor)
 {
 	return tsr_half(1.0) - RejectionFactor;
 }
+
+/** Payload of a subpixel depth information in the history */
+struct FTSRSubpixelDepth
+{
+	/** Depth of the pixel. */
+	float DeviceZ;
+	
+	/** Offset of the depth from the center of pixel. */
+	float2 Offset;
+	
+	/** Depth of the pixel in number of frames. */
+	uint Age;
+};
+
+uint CompressSubpixelDepth(FTSRSubpixelDepth SubpixelDepth)
+{
+	uint CompressedSubpixelDepth = 0u;
+	CompressedSubpixelDepth |= f32tof16(SubpixelDepth.DeviceZ) << 17u;
+	CompressedSubpixelDepth |= uint(clamp(int(round(SubpixelDepth.Offset.x * 16.00 + 16.0)), 0, 31)) << 12u;
+	CompressedSubpixelDepth |= uint(clamp(int(round(SubpixelDepth.Offset.y * 16.00 + 16.0)), 0, 31)) <<  7u;
+	CompressedSubpixelDepth |= (SubpixelDepth.Age) << 0u;
+	return CompressedSubpixelDepth;
+}
+
+FTSRSubpixelDepth UncompressSubpixelDepth(uint CompressedSubpixelDepth)
+{
+	FTSRSubpixelDepth SubpixelDepth;
+	SubpixelDepth.DeviceZ  = f16tof32(CompressedSubpixelDepth >> 17u);
+	SubpixelDepth.Offset.x = float((CompressedSubpixelDepth >> 12u) & 0x1Fu) * rcp(16.0) - 1.0;
+	SubpixelDepth.Offset.y = float((CompressedSubpixelDepth >>  7u) & 0x1Fu) * rcp(16.0) - 1.0;
+	SubpixelDepth.Age      = CompressedSubpixelDepth & 0x7Fu;
+	return SubpixelDepth;
+}
diff --git a/Engine/Shaders/Private/TemporalSuperResolution/TSRDecimateHistory.usf b/Engine/Shaders/Private/TemporalSuperResolution/TSRDecimateHistory.usf
--- a/Engine/Shaders/Private/TemporalSuperResolution/TSRDecimateHistory.usf
+++ b/Engine/Shaders/Private/TemporalSuperResolution/TSRDecimateHistory.usf
@@ -31,6 +31,8 @@
 /** Whether to compile the moire reprojection. */
 #define CONFIG_MOIRE_REPROJECTION (DIM_MOIRE_REPROJECTION)
 
+#define CONFIG_SUBPIXEL_PARALLAX_FACTOR (DIM_SUBPIXEL_PARALLAX_FACTOR)
+
 
 //------------------------------------------------------- PARAMETERS
 
@@ -41,9 +43,11 @@ Texture2D<float2> DilatedVelocityTexture;
 Texture2D<float> ClosestDepthTexture;
 Texture2D<uint> PrevUseCountTexture;
 Texture2D<uint> PrevClosestDepthTexture;
-Texture2D<tsr_subpixel_payload> ParallaxFactorTexture;
 
+#if CONFIG_SUBPIXEL_PARALLAX_FACTOR
+Texture2D<tsr_subpixel_payload> ParallaxFactorTexture;
 Texture2D<tsr_subpixel_details> PrevHistorySubpixelDetails;
+#endif
 
 Texture2D<tsr_halfC> PrevHistoryGuide;
 Texture2D<tsr_half4> PrevHistoryMoire;
@@ -68,7 +72,9 @@ RWTexture2D<float2> HoleFilledVelocityOutput;
 RWTexture2D<float2> GrandHoleFilledVelocityOutput;
 RWTexture2D<tsr_half2> HoleFilledVelocityMaskOutput;
 RWTexture2D<tsr_half> ParallaxRejectionMaskOutput;
-RWTexture2D<tsr_subpixel_payload> HistorySubpixelDetailsOutput;
+#if CONFIG_SUBPIXEL_PARALLAX_FACTOR
+	RWTexture2D<tsr_subpixel_payload> HistorySubpixelDetailsOutput;
+#endif
 
 
 //------------------------------------------------------- LDS
@@ -91,8 +97,10 @@ void MainCS(
 	// Sample current frame informations.
 	float2 EncodedScreenVelocity;
 	float DeviceZ;
-	tsr_subpixel_payload ParallaxFactorBits;
-	tsr_half ParallaxFactor;
+	#if CONFIG_SUBPIXEL_PARALLAX_FACTOR
+		tsr_subpixel_payload ParallaxFactorBits;
+		tsr_half ParallaxFactor;
+	#endif
 	tsr_half RejectHighFrequency = tsr_half(0.0);
 	ISOLATE
 	{
@@ -107,11 +115,16 @@ void MainCS(
 		// Issue texture fetechs
 		ISOLATE
 		{
-			tsr_subpixel_payload RawParallaxFactor = ParallaxFactorTexture[InputPixelPos];
+			#if CONFIG_SUBPIXEL_PARALLAX_FACTOR
+				tsr_subpixel_payload RawParallaxFactor = ParallaxFactorTexture[InputPixelPos];
+			#endif
 
 			DeviceZ = ClosestDepthTexture[InputPixelPos];
-			ParallaxFactor = UncompressParallaxFactor(RawParallaxFactor & tsr_subpixel_payload(SUB_PIXEL_PARALLAX_FACTOR_BIT_MASK));
-			ParallaxFactorBits = RawParallaxFactor & tsr_subpixel_payload(SUB_PIXEL_PARALLAX_FACTOR_BIT_MASK);
+
+			#if CONFIG_SUBPIXEL_PARALLAX_FACTOR
+				ParallaxFactor = UncompressParallaxFactor(RawParallaxFactor & tsr_subpixel_payload(SUB_PIXEL_PARALLAX_FACTOR_BIT_MASK));
+				ParallaxFactorBits = RawParallaxFactor & tsr_subpixel_payload(SUB_PIXEL_PARALLAX_FACTOR_BIT_MASK);
+			#endif
 
 			#if CONFIG_VELOCITY_SAMPLES == 1
 				EncodedScreenVelocity = DilatedVelocityTexture[InputPixelPos];
@@ -275,7 +288,9 @@ void MainCS(
 	#if CONFIG_MOIRE_REPROJECTION
 		tsr_half4 ReprojectedHistoryMoireColor;
 	#endif
-	FSubpixelNeighborhood SubpixelNeigborhood;
+	#if CONFIG_SUBPIXEL_PARALLAX_FACTOR
+		FSubpixelNeighborhood SubpixelNeigborhood;
+	#endif
 	bool bIsOffScreen;
 
 	// Issues all the history texture fetches
@@ -361,6 +376,7 @@ void MainCS(
 			#endif
 		}
 
+		#if CONFIG_SUBPIXEL_PARALLAX_FACTOR
 		{
 			// Pixel coordinate of the input pixel in the predictor viewport.
 			float2 InputPPCk = float2(InputPixelPos - InputInfo_ViewportMin) + 0.5f - InputJitter;
@@ -387,13 +403,16 @@ void MainCS(
 
 			SubpixelNeigborhood = GatherPrevSubpixelNeighborhood(PrevHistorySubpixelDetails, SubpixelBufferUV);
 		}
+		#endif
 	}
 	
 
 	// Process the history samples
 	tsr_half ParallaxRejectionMask = 0.0;
 	
-	tsr_subpixel_details FinalSubpixelDetails;
+	#if CONFIG_SUBPIXEL_PARALLAX_FACTOR
+		tsr_subpixel_details FinalSubpixelDetails;
+	#endif
 
 	ISOLATE
 	{
@@ -413,6 +432,7 @@ void MainCS(
 		float2 ScreenVelocity = DecodeVelocityFromTexture(float4(EncodedScreenVelocity, 0.0, 0.0)).xy;
 		float2 PrevScreenPos = ScreenPos - ScreenVelocity;
 	
+		#if CONFIG_SUBPIXEL_PARALLAX_FACTOR
 		// Compute HistoryMaxRejectionFactor
 		tsr_half HistoryMaxRejectionFactor;
 		tsr_subpixel_details PrevSubpixelDetails;
@@ -434,6 +454,7 @@ void MainCS(
 
 			HistoryMaxRejectionFactor = UncompressParallaxFactor(EncodedHistoryMaxRejectionFactor);
 		}
+		#endif
 		
 		// Compute the parralax rejection mask
 		bool bIsValidHoleFillingPixelVelocity = false;
@@ -489,6 +510,7 @@ void MainCS(
 		}
 
 		// Reject subpixel details that might be in history but not input just based on their parrallax
+		#if CONFIG_SUBPIXEL_PARALLAX_FACTOR
 		{
 			tsr_half HistorySubdetailRejectionFactor = max(HistoryMaxRejectionFactor / ParallaxFactor, tsr_half(1.0));
 
@@ -509,6 +531,7 @@ void MainCS(
 
 			ParallaxRejectionMask = min(ParallaxRejectionMask, HistoryParallaxRejectionMask);
 		}
+		#endif
 
 		// Finish off the ParallaxRejectionMask
 		ParallaxRejectionMask = select(bIsOffScreen, tsr_half(0.0), ParallaxRejectionMask);
@@ -540,6 +563,7 @@ void MainCS(
 		}
 		
 		// Update subpixel detail
+		#if CONFIG_SUBPIXEL_PARALLAX_FACTOR
 		{
 			FinalSubpixelDetails = PrevSubpixelDetails;
 
@@ -567,6 +591,7 @@ void MainCS(
 
 			FinalSubpixelDetails = UpdatedFinalSubpixelDetails;
 		}
+		#endif
 	}
 	
 	#if CONFIG_GRAND_REPROJECTION
@@ -642,7 +667,9 @@ void MainCS(
 		#endif
 		HoleFilledVelocityMaskOutput[InputPixelPos] = tsr_half2(0.0, RejectHighFrequency);
 		ParallaxRejectionMaskOutput[OutputPixelPos] = ParallaxRejectionMask;
-		HistorySubpixelDetailsOutput[OutputPixelPos] = FinalSubpixelDetails;
+		#if CONFIG_SUBPIXEL_PARALLAX_FACTOR
+			HistorySubpixelDetailsOutput[OutputPixelPos] = FinalSubpixelDetails;
+		#endif
 
 		#if DEBUG_OUTPUT
 		{
diff --git a/Engine/Shaders/Private/TemporalSuperResolution/TSRDilateVelocity.usf b/Engine/Shaders/Private/TemporalSuperResolution/TSRDilateVelocity.usf
--- a/Engine/Shaders/Private/TemporalSuperResolution/TSRDilateVelocity.usf
+++ b/Engine/Shaders/Private/TemporalSuperResolution/TSRDilateVelocity.usf
@@ -16,13 +16,27 @@
 	#define TILE_SIZE 8
 #endif
 
-#define AA_CROSS 1
+#define CONFIG_AA_CROSS 1
+
+#define CONFIG_AA_CROSS_OPTIMISATION (PLATFORM_SUPPORTS_WAVE_BROADCAST)
 
 #define CONFIG_USE_SCENE_VELOCITY 1
 
 #define CONFIG_DEPTH_EDGE_DETECTION 1
 
-#define CONFIG_VELOCITY_EXTRAPOLATION 1
+#define CONFIG_SUBPIXEL_DEPTH (DIM_SUBPIXEL_DEPTH)
+
+#define CONFIG_VELOCITY_EXTRAPOLATION 0
+
+
+//------------------------------------------------------- CONSTANTS
+
+static const int2 kOffsets[4] = {
+	int2( 1,  0),
+	int2( 1,  1),
+	int2( 0,  1),
+	int2(-1,  1),
+};
 
 
 //------------------------------------------------------- PARAMETERS
@@ -34,15 +48,18 @@ float WorldDepthToDepthError;
 float VelocityExtrapolationMultiplier;
 float InvFlickeringMaxParralaxVelocity;
 uint bOutputIsMovingTexture;
+uint bIncludeDynamicDepthInSubpixelDetails;
 
 Texture2D<float> SceneDepthTexture;
 Texture2D SceneVelocityTexture;
+Texture2D<uint> PrevScatteredSubpixelDepthTexture;
 
 RWTexture2D<float2> DilatedVelocityOutput;
 RWTexture2D<float> ClosestDepthOutput;
 RWTexture2D<uint> PrevUseCountOutput;
 RWTexture2D<uint> PrevClosestDepthOutput;
 RWTexture2DArray<uint> R8Output;
+RWTexture2D<uint> SubpixelDepthOutput;
 
 RWTexture2D<float3>	VelocityFlattenOutput;
 RWTexture2D<float4>	VelocityTileOutput0;
@@ -63,6 +80,10 @@ void UnpackDepthAndOffset(uint Packed, out int2 Offset, out float DeviceZ)
 	Offset.y = int((Packed >> 2) & 0x3) - 1;
 }
 
+float MaxDeviceZWithSubpixel(float RenderedDeviceZ, float SubpixelDeviceZ)
+{
+	return max(RenderedDeviceZ, SubpixelDeviceZ);
+}
 
 //------------------------------------------------------- ENTRY POINT
 
@@ -77,38 +98,91 @@ void MainCS(
 		ZOrder2D(GroupThreadIndex, uint(log2(float(TILE_SIZE)))) +
 		GroupId * uint2(TILE_SIZE, TILE_SIZE));
 
+	#if CONFIG_AA_CROSS_OPTIMISATION
+		int2 ThreadOffsetSign = -1 + int2((DispatchThreadId % uint(2)) * uint(2));
+	#endif
+
 	uint2 InputPixelPos = InputInfo_ViewportMin + DispatchThreadId;
 	
 	float2 ViewportUV = (float2(DispatchThreadId) + 0.5f) * InputInfo_ViewportSizeInverse;
 	float2 ScreenPos = ViewportUVToScreenPos(ViewportUV);
 	float2 InputBufferUV = (InputPixelPos + 0.5) * InputInfo_ExtentInverse;
 
-	float DeviceZ = SceneDepthTexture[InputPixelPos];
+	// Do all depth memory fetches.
+	float DeviceZ;
+	#if CONFIG_SUBPIXEL_DEPTH
+		uint PrevCompressedSubpixelDepth;
+	#endif
+	#if CONFIG_AA_CROSS_OPTIMISATION
+		float NeighborDeviceZArray[3];
+		#if CONFIG_SUBPIXEL_DEPTH
+			uint NeighborPrevCompressedSubpixelDepthArray[3];
+		#endif
+	#endif
+	ISOLATE
+	{
+		DeviceZ = SceneDepthTexture[InputPixelPos];
+		#if CONFIG_SUBPIXEL_DEPTH
+			PrevCompressedSubpixelDepth = PrevScatteredSubpixelDepthTexture[InputPixelPos];
+		#endif
+		#if CONFIG_AA_CROSS_OPTIMISATION
+			UNROLL_N(3)
+			for (uint i = 0; i < 3; i++)
+			{
+				int2 NeighborPixelPos = int2(InputPixelPos) + kOffsets[i] * ThreadOffsetSign;
+				NeighborPixelPos = clamp(NeighborPixelPos, InputPixelPosMin, InputPixelPosMax);
 
+				NeighborDeviceZArray[i] = SceneDepthTexture[NeighborPixelPos];
+				#if CONFIG_SUBPIXEL_DEPTH
+					NeighborPrevCompressedSubpixelDepthArray[i] = PrevScatteredSubpixelDepthTexture[NeighborPixelPos];
+				#endif
+			}
+		#endif
+	}
+
+	#if CONFIG_SUBPIXEL_DEPTH
+		float PrevScatteredDeviceZ = UncompressSubpixelDepth(PrevCompressedSubpixelDepth).DeviceZ;
+	#endif
+	
 	float ClosestDeviceZ = DeviceZ;
 	int2 VelocityPixelOffset = int2(0, 0);
 	float ParallaxFactorForHistory = 1.0;
 	float ParallaxFactorForHistoryRejection = 1.0;
-	#if AA_CROSS
+	#if CONFIG_AA_CROSS
+	ISOLATE
 	{
-		const int2 kOffsets[4] = {
-			int2( 1,  0),
-			int2( 1,  1),
-			int2( 0,  1),
-			int2(-1,  1),
-		};
-
-		uint PackedDepthOffset = PackDepthAndOffset(/* Offset = */ int2(0, 0), DeviceZ);
-		
+		#if CONFIG_AA_CROSS_OPTIMISATION
+			const FWaveBroadcastSettings Horizontal = InitWaveXorButterfly(0x1);
+			const FWaveBroadcastSettings Vertical   = InitWaveXorButterfly(0x2);
+			const FWaveBroadcastSettings Diagonal   = InitWaveXorButterfly(0x3);
+
+			const FWaveBroadcastSettings CenterBroadcastArray[4] = {
+				Horizontal,
+				Diagonal,
+				Vertical,
+				InitNopBroadcast(), // unused
+			};
+		#endif
+
+
+		float CenterSubpixelDeviceZ = DeviceZ;
+		#if CONFIG_SUBPIXEL_DEPTH
+			CenterSubpixelDeviceZ = MaxDeviceZWithSubpixel(DeviceZ, PrevScatteredDeviceZ);
+		#endif
+
+		uint PackedDepthOffset = PackDepthAndOffset(/* Offset = */ int2(0, 0), CenterSubpixelDeviceZ);
+
 		float FurthestPairDeviceZ = DeviceZ;
 		float FurthestDeviceZ = DeviceZ;
 
+		#if CONFIG_DEPTH_EDGE_DETECTION
 		float DeviceZError;
 		{
-			float Depth = ConvertFromDeviceZ(DeviceZ);
+			float Depth = ConvertFromDeviceZ(CenterSubpixelDeviceZ);
 			float ErrorCorrectedDeviceZ = ConvertToDeviceZ(Depth + Depth * WorldDepthToDepthError);
-			DeviceZError = abs(ErrorCorrectedDeviceZ - DeviceZ);
+			DeviceZError = abs(ErrorCorrectedDeviceZ - CenterSubpixelDeviceZ);
 		}
+		#endif
 
 		#if 1
 			UNROLL_N(4)
@@ -118,27 +192,60 @@ void MainCS(
 			for (uint i = 1; i < 4; i+= 2)
 		#endif
 		{
-			const int2 Offset = kOffsets[i];
-
-			float SampleDeviceZ0 = SceneDepthTexture[int2(InputPixelPos) + Offset];
-			float SampleDeviceZ1 = SceneDepthTexture[int2(InputPixelPos) - Offset];
+			#if CONFIG_AA_CROSS_OPTIMISATION
+				int2 Offset = kOffsets[i] * ThreadOffsetSign;
+
+				float SampleDeviceZ0 = i == 3 ? WaveBroadcast(Horizontal, NeighborDeviceZArray[2]) : NeighborDeviceZArray[i];
+				float SampleDeviceZ1 = i == 3 ? WaveBroadcast(Vertical  , NeighborDeviceZArray[0]) : WaveBroadcast(CenterBroadcastArray[i], DeviceZ);
+				
+				#if CONFIG_SUBPIXEL_DEPTH
+					uint SamplePrevCompressedSubpixel0 = i == 3 ? WaveBroadcast(Horizontal, NeighborPrevCompressedSubpixelDepthArray[2]) : NeighborPrevCompressedSubpixelDepthArray[i];
+					uint SamplePrevCompressedSubpixel1 = i == 3 ? WaveBroadcast(Vertical  , NeighborPrevCompressedSubpixelDepthArray[0]) : WaveBroadcast(CenterBroadcastArray[i], PrevCompressedSubpixelDepth);
+					
+					float SampleSubpixelDeviceZ0 = UncompressSubpixelDepth(SamplePrevCompressedSubpixel0).DeviceZ;
+					float SampleSubpixelDeviceZ1 = UncompressSubpixelDepth(SamplePrevCompressedSubpixel1).DeviceZ;
+				#endif
+			#else // !CONFIG_AA_CROSS_OPTIMISATION
+				const int2 Offset = kOffsets[i];
+				
+				int2 NeighborPixelPos0 = int2(InputPixelPos) + Offset;
+				int2 NeighborPixelPos1 = int2(InputPixelPos) - Offset;
+				NeighborPixelPos0 = clamp(NeighborPixelPos0, InputPixelPosMin, InputPixelPosMax);
+				NeighborPixelPos1 = clamp(NeighborPixelPos1, InputPixelPosMin, InputPixelPosMax);
+
+				float SampleDeviceZ0 = SceneDepthTexture[NeighborPixelPos0];
+				float SampleDeviceZ1 = SceneDepthTexture[NeighborPixelPos1];
 			
+				#if CONFIG_SUBPIXEL_DEPTH
+					float SampleSubpixelDeviceZ0 = UncompressSubpixelDepth(PrevScatteredSubpixelDepthTexture[NeighborPixelPos0]).DeviceZ;
+					float SampleSubpixelDeviceZ1 = UncompressSubpixelDepth(PrevScatteredSubpixelDepthTexture[NeighborPixelPos1]).DeviceZ;
+				#endif
+			#endif // !CONFIG_AA_CROSS_OPTIMISATION
+
+			#if CONFIG_SUBPIXEL_DEPTH
+				float CombinedSampleSubpixelDeviceZ0 = MaxDeviceZWithSubpixel(SampleDeviceZ0, SampleSubpixelDeviceZ0);
+				float CombinedSampleSubpixelDeviceZ1 = MaxDeviceZWithSubpixel(SampleDeviceZ1, SampleSubpixelDeviceZ1);
+			#else
+				float CombinedSampleSubpixelDeviceZ0 = SampleDeviceZ0;
+				float CombinedSampleSubpixelDeviceZ1 = SampleDeviceZ1;
+			#endif
+
 			#if CONFIG_DEPTH_EDGE_DETECTION
 			{
-				float DepthDiff = abs(SampleDeviceZ1 - SampleDeviceZ0);
-				float DepthVariation = (SampleDeviceZ0 + SampleDeviceZ1) * 0.5 - DeviceZ;
+				float DepthDiff = abs(CombinedSampleSubpixelDeviceZ1 - CombinedSampleSubpixelDeviceZ0);
+				float DepthVariation = (CombinedSampleSubpixelDeviceZ0 + CombinedSampleSubpixelDeviceZ1) * 0.5 - CenterSubpixelDeviceZ;
 
 				FLATTEN
 				if (DepthVariation > max(DepthDiff * 0.25, DeviceZError))
 				{
-					PackedDepthOffset = max(PackedDepthOffset, PackDepthAndOffset(+Offset, SampleDeviceZ0));
-					PackedDepthOffset = max(PackedDepthOffset, PackDepthAndOffset(-Offset, SampleDeviceZ1));
+					PackedDepthOffset = max(PackedDepthOffset, PackDepthAndOffset(+Offset, CombinedSampleSubpixelDeviceZ0));
+					PackedDepthOffset = max(PackedDepthOffset, PackDepthAndOffset(-Offset, CombinedSampleSubpixelDeviceZ1));
 				}
 			}
 			#else
 			{
-				PackedDepthOffset = max(PackedDepthOffset, PackDepthAndOffset(+Offset, SampleDeviceZ0));
-				PackedDepthOffset = max(PackedDepthOffset, PackDepthAndOffset(-Offset, SampleDeviceZ1));
+				PackedDepthOffset = max(PackedDepthOffset, PackDepthAndOffset(+Offset, CombinedSampleSubpixelDeviceZ0));
+				PackedDepthOffset = max(PackedDepthOffset, PackDepthAndOffset(-Offset, CombinedSampleSubpixelDeviceZ1));
 			}
 			#endif
 
@@ -156,6 +263,9 @@ void MainCS(
 	// Camera motion for pixel.
 	float3 FinalScreenVelocity = 0.0;
 
+	float4 CenterEncodedVelocity = SceneVelocityTexture[int2(InputBufferUV * InputInfo_Extent)];
+	bool bIsMovingPixel = CenterEncodedVelocity.x > 0.0;
+
 	// Compute final screen vecocity.
 	#if CONFIG_VELOCITY_EXTRAPOLATION
 	BRANCH
@@ -168,6 +278,11 @@ void MainCS(
 
 		float ClosestDeviceZ1 = ClosestDeviceZ;
 		float ClosestDeviceZ2 = SceneDepthTexture[PixelCoord2];
+		#if CONFIG_SUBPIXEL_DEPTH
+			float ClosestSubpixelDeviceZ2 = UncompressSubpixelDepth(PrevScatteredSubpixelDepthTexture[PixelCoord2]).DeviceZ;
+			
+			ClosestDeviceZ2 = MaxDeviceZWithSubpixel(ClosestDeviceZ2, ClosestSubpixelDeviceZ2, SampleDeviceZError);
+		#endif
 
 		float4 EncodedVelocity1 = SceneVelocityTexture[PixelCoord1];
 		float4 EncodedVelocity2 = SceneVelocityTexture[PixelCoord2];
@@ -242,6 +357,31 @@ void MainCS(
 			saturate(PixelParallaxVelocity * tsr_half(InvFlickeringMaxParralaxVelocity) - tsr_half(0.5)));
 	}
 
+	#if CONFIG_SUBPIXEL_DEPTH
+	uint CompressedSubpixelDepth = PrevCompressedSubpixelDepth;
+	{
+		FTSRSubpixelDepth PrevSubpixelDepth = UncompressSubpixelDepth(PrevCompressedSubpixelDepth);
+
+		FLATTEN
+		if (DeviceZ > PrevSubpixelDepth.DeviceZ || (bIsMovingPixel && !bIncludeDynamicDepthInSubpixelDetails))
+		{
+			CompressedSubpixelDepth = 0;
+		}
+
+		if (ParallaxFactorForHistory > 1.005 && (!bIsMovingPixel || bIncludeDynamicDepthInSubpixelDetails))
+		{
+			FTSRSubpixelDepth FinalSubpixelDepth;
+			FinalSubpixelDepth.DeviceZ = DeviceZ;
+			FinalSubpixelDepth.Offset = 0.0;
+			FinalSubpixelDepth.Age = 0;
+
+			CompressedSubpixelDepth = CompressSubpixelDepth(FinalSubpixelDepth);
+
+			ParallaxFactorForHistory = 0.0;
+		}
+	}
+	#endif
+
 	// Do motion blur velocity flatten.
 	#if DIM_MOTION_BLUR_DIRECTIONS > 0
 		FVelocityRange VelocityPolarRange;
@@ -302,7 +442,17 @@ void MainCS(
 		//DilatedVelocityOutput[OutputPixelPos] = EncodedVelocity.xy; //EncodeVelocityToTexture(float3(FinalScreenVelocity.xy, 0.0)).xy;
 		DilatedVelocityOutput[OutputPixelPos] = EncodeVelocityToTexture(float3(FinalScreenVelocity.xy, 0.0)).xy;
 		ClosestDepthOutput[OutputPixelPos] = ClosestDeviceZ;
-		R8Output[uint3(OutputPixelPos, 0)] = CompressParallaxFactor(tsr_half(ParallaxFactorForHistory)) | (CompressParallaxFactor(tsr_half(ParallaxFactorForHistoryRejection)) << 4);
+		#if CONFIG_SUBPIXEL_DEPTH
+		{
+			SubpixelDepthOutput[OutputPixelPos] = CompressedSubpixelDepth;
+			R8Output[uint3(OutputPixelPos, 0)] = 0; // TODO save some memory bandwidth.
+		}
+		#else
+		{
+			// TODO save some memory bandwidth.
+			R8Output[uint3(OutputPixelPos, 0)] = CompressParallaxFactor(tsr_half(ParallaxFactorForHistory)) | (CompressParallaxFactor(tsr_half(ParallaxFactorForHistoryRejection)) << 4);
+		}
+		#endif
 
 		BRANCH
 		if (bOutputIsMovingTexture)
diff --git a/Engine/Shaders/Private/TemporalSuperResolution/TSRForwardScatterDepth.usf b/Engine/Shaders/Private/TemporalSuperResolution/TSRForwardScatterDepth.usf
new file mode 100644
--- /dev/null
+++ b/Engine/Shaders/Private/TemporalSuperResolution/TSRForwardScatterDepth.usf
@@ -0,0 +1,104 @@
+// Copyright Epic Games, Inc. All Rights Reserved.
+
+#include "TSRCommon.ush"
+
+
+//------------------------------------------------------- CONFIG
+
+#define TILE_SIZE 8
+
+#define DEBUG_INTERLOCK_ALL 0
+
+#define CONFIG_EARLY_RETURN 1
+
+
+//------------------------------------------------------- PARAMETERS
+
+uint2 PrevSubpixelDepthViewportMin;
+uint2 PrevSubpixelDepthViewportMax;
+FScreenTransform InputPixelPosToPrevScreenPosition;
+FScreenTransform ScreenPosToOutputPixelPos;
+uint MaxDepthAge;
+
+Texture2D<uint> PrevSubpixelDepthTexture;
+
+RWTexture2D<uint> SubpixelDepthOutput;
+
+
+//------------------------------------------------------- ENTRY POINT
+
+[numthreads(TILE_SIZE * TILE_SIZE, 1, 1)]
+void MainCS(
+	uint2 GroupId : SV_GroupID,
+	uint GroupThreadIndex : SV_GroupIndex)
+{
+	float4 Debug = 0.0;
+
+	uint2 DispatchThreadId = (
+		ZOrder2D(GroupThreadIndex, uint(log2(float(TILE_SIZE)))) +
+		GroupId * uint2(TILE_SIZE, TILE_SIZE));
+
+	uint2 InputPixelPos = PrevSubpixelDepthViewportMin + DispatchThreadId;
+	
+	uint CompressedPrevSubpixelDepth = PrevSubpixelDepthTexture[InputPixelPos];
+	FTSRSubpixelDepth PrevScatteredSubpixelDepth = UncompressSubpixelDepth(CompressedPrevSubpixelDepth);
+		
+	const bool bEarlyReturn = PrevScatteredSubpixelDepth.DeviceZ == 0.0 || PrevScatteredSubpixelDepth.Age >= MaxDepthAge || any(InputPixelPos > PrevSubpixelDepthViewportMax);
+#if CONFIG_EARLY_RETURN && !DEBUG_INTERLOCK_ALL
+	BRANCH
+	if (bEarlyReturn)
+	{
+		return;
+	}
+#endif
+	
+	float3 PreViewTranslationOffset = LWCToFloat(LWCSubtract(PrimaryView.PreViewTranslation, PrimaryView.PrevPreViewTranslation));
+ 
+	// Compute position of the history in world space.
+	float PrevSceneDepth = ConvertFromDeviceZ(PrevScatteredSubpixelDepth.DeviceZ);
+	float2 PrevScreenPosition = ApplyScreenTransform(float2(InputPixelPos), InputPixelPosToPrevScreenPosition);
+	float3 TranslatedWorldPosition = mul(float4(PrevScreenPosition * PrevSceneDepth, PrevSceneDepth, 1), View.PrevScreenToTranslatedWorld).xyz + PreViewTranslationOffset;
+	
+	// Reproject the world position into new viewport.
+	float4 ClipPos = mul(float4(TranslatedWorldPosition, 1.0), View.TranslatedWorldToClip);
+	float2 ScreenPos = ClipPos.xy / ClipPos.w;
+	float2 DesiredPixelPos = ApplyScreenTransform(ScreenPos, ScreenPosToOutputPixelPos);
+
+	float2 ActualPixelPos = floor(DesiredPixelPos);
+	float2 NewOffset = DesiredPixelPos - (ActualPixelPos + 0.5);
+
+	int2 OutputPixelPos = int2(ActualPixelPos);
+	bool bIsValidOutputPixelPos = all(and(OutputPixelPos >= InputInfo_ViewportMin, OutputPixelPos < InputInfo_ViewportMax)) && ClipPos.w > 0.0;
+
+#if !CONFIG_EARLY_RETURN
+	const bool bScatter = !bEarlyReturn && bIsValidOutputPixelPos;
+#else
+	const bool bScatter = bIsValidOutputPixelPos;
+#endif
+
+#if !DEBUG_INTERLOCK_ALL
+	BRANCH
+	if (bScatter)
+#endif
+	{
+		FTSRSubpixelDepth FinalSubpixelDepth;
+		FinalSubpixelDepth.DeviceZ = ConvertToDeviceZ(ClipPos.w);
+		FinalSubpixelDepth.Offset = NewOffset;
+		FinalSubpixelDepth.Age = PrevScatteredSubpixelDepth.Age + 1;
+
+		uint CompressedSubpixelDepth = CompressSubpixelDepth(FinalSubpixelDepth);
+	
+		#if DEBUG_INTERLOCK_ALL
+			OutputPixelPos.x = select(bScatter, OutputPixelPos.x, -1);
+		#endif
+
+		InterlockedMax(SubpixelDepthOutput[OutputPixelPos], CompressedSubpixelDepth);
+		//SubpixelDepthOutput[OutputPixelPos] = CompressedSubpixelDepth;
+	}
+
+	#if DEBUG_OUTPUT
+	{
+		DebugOutput[tsr_short3(InputPixelPos, 0)] = Debug;
+	}
+	#endif
+}
diff --git a/Engine/Source/Runtime/Renderer/Private/PostProcess/TemporalSuperResolution.cpp b/Engine/Source/Runtime/Renderer/Private/PostProcess/TemporalSuperResolution.cpp
--- a/Engine/Source/Runtime/Renderer/Private/PostProcess/TemporalSuperResolution.cpp
+++ b/Engine/Source/Runtime/Renderer/Private/PostProcess/TemporalSuperResolution.cpp
@@ -127,6 +127,24 @@
 	TEXT("Defines how much the velocity should be extrapolated on geometric discontinuities (Default = 1.0f)."),
 	ECVF_Scalability | ECVF_RenderThreadSafe);
 
+TAutoConsoleVariable<int32> CVarTSRSubpixelMethod(
+	TEXT("r.TSR.Subpixel.Method"), 2,
+	TEXT("Selects the method to reproject and/or discard the subpixel details.\n")
+	TEXT(" 0: disable subpixel detail accumulation\n")
+	TEXT(" 1: accumulate how much subpixel detail can paralax under movement for history rejection\n")
+	TEXT(" 2: accumulate subpixel details' closest depth to be able to reproject them even when not drawing in depth/velocity buffer (default)\n"),
+	ECVF_RenderThreadSafe);
+
+TAutoConsoleVariable<int32> CVarTSRSubpixelDepthMaxAge(
+	TEXT("r.TSR.Subpixel.DephtMaxAge"), 3,
+	TEXT("Maximum age in frames of subpixel's depth kept in history for their self reprojection (default to 3 frames)."),
+	ECVF_RenderThreadSafe);
+
+TAutoConsoleVariable<int32> CVarTSRSubpixelIncludeMovingDepth(
+	TEXT("r.TSR.Subpixel.IncludeMovingDepth"), 0,
+	TEXT("Whether the depth of moving subpixel detail should also be included in the subpixel depth history for their reprojection (disabled by default)."),
+	ECVF_RenderThreadSafe);
+
 #if COMPILE_TSR_DEBUG_PASSES
 
 TAutoConsoleVariable<int32> CVarTSRDebugArraySize(
@@ -174,6 +192,7 @@
 	SHADER_PARAMETER_RDG_TEXTURE(Texture2DArray, ColorArray)
 	SHADER_PARAMETER_RDG_TEXTURE(Texture2D, Metadata)
 	SHADER_PARAMETER_RDG_TEXTURE(Texture2D, SubpixelDetails)
+	SHADER_PARAMETER_RDG_TEXTURE(Texture2D, SubpixelDepth)
 	SHADER_PARAMETER_RDG_TEXTURE(Texture2D, TranslucencyAlpha)
 	SHADER_PARAMETER_RDG_TEXTURE(Texture2D, Guide)
 	SHADER_PARAMETER_RDG_TEXTURE(Texture2D, Moire)
@@ -187,7 +206,6 @@
 	SHADER_PARAMETER_RDG_TEXTURE_SRV(Texture2D, PrevHighFrequencyResultant)
 	SHADER_PARAMETER_RDG_TEXTURE_SRV(Texture2D, HighFrequencyOverblur)
 
-	SHADER_PARAMETER_RDG_TEXTURE(Texture2D, SubpixelDetails)
 	SHADER_PARAMETER_RDG_TEXTURE_SRV(Texture2D, Translucency)
 	SHADER_PARAMETER_RDG_TEXTURE_SRV(Texture2D, TranslucencyAlpha)
 	SHADER_PARAMETER_RDG_TEXTURE(Texture2D, Guide)
@@ -198,7 +216,6 @@
 	SHADER_PARAMETER_STRUCT(FTSRHistoryArrayIndices, ArrayIndices)
 	SHADER_PARAMETER_RDG_TEXTURE_UAV(RWTexture2DArray, ColorArray)
 	SHADER_PARAMETER_RDG_TEXTURE_UAV(RWTexture2D, Metadata)
-	SHADER_PARAMETER_RDG_TEXTURE_UAV(RWTexture2D, SubpixelDetails)
 	SHADER_PARAMETER_RDG_TEXTURE_UAV(RWTexture2D, TranslucencyAlpha)
 END_SHADER_PARAMETER_STRUCT()
 
@@ -220,6 +237,13 @@
 };
 ENUM_CLASS_FLAGS(ETSRHistoryFormatBits);
 
+enum class ETSRSubpixelMethod
+{
+	Disabled,
+	ParallaxFactor,
+	ClosestDepth,
+};
+
 bool IsOutputDifferentThanHighFrequency(ETSRHistoryFormatBits HistoryFormatBits)
 {
 	return EnumHasAnyFlags(HistoryFormatBits, ETSRHistoryFormatBits::Translucency);
@@ -267,7 +291,6 @@
 		SRVs.HighFrequency = GraphBuilder.CreateSRV(Textures.Output);
 	}
 	SRVs.Metadata = GraphBuilder.CreateSRV(FRDGTextureSRVDesc::Create(Textures.Metadata));
-	SRVs.SubpixelDetails = Textures.SubpixelDetails;
 	if (Textures.ArrayIndices.Translucency >= 0)
 	{
 		SRVs.Translucency = GraphBuilder.CreateSRV(FRDGTextureSRVDesc::CreateForSlice(Textures.ColorArray, Textures.ArrayIndices.Translucency));
@@ -294,7 +317,6 @@
 		UAVs.ColorArray = GraphBuilder.CreateUAV(Textures.ColorArray);
 	}
 	UAVs.Metadata = GraphBuilder.CreateUAV(Textures.Metadata);
-	UAVs.SubpixelDetails = GraphBuilder.CreateUAV(Textures.SubpixelDetails);
 	if (Textures.TranslucencyAlpha)
 	{
 		UAVs.TranslucencyAlpha = GraphBuilder.CreateUAV(Textures.TranslucencyAlpha);
@@ -383,16 +405,38 @@
 
 		SHADER_PARAMETER_RDG_TEXTURE_UAV(RWTexture2D, PrevUseCountOutput)
 		SHADER_PARAMETER_RDG_TEXTURE_UAV(RWTexture2D, PrevClosestDepthOutput)
+		SHADER_PARAMETER_RDG_TEXTURE_UAV(RWTexture2D, SubpixelDepthOutput)
 	END_SHADER_PARAMETER_STRUCT()
 }; // class FTSRClearPrevTexturesCS
 
+class FTSRForwardScatterDepthCS : public FTSRShader
+{
+	DECLARE_GLOBAL_SHADER(FTSRForwardScatterDepthCS);
+	SHADER_USE_PARAMETER_STRUCT(FTSRForwardScatterDepthCS, FTSRShader);
+
+	BEGIN_SHADER_PARAMETER_STRUCT(FParameters, )
+		SHADER_PARAMETER_STRUCT_INCLUDE(FTSRCommonParameters, CommonParameters)
+		SHADER_PARAMETER(FIntPoint, PrevSubpixelDepthViewportMin)
+		SHADER_PARAMETER(FIntPoint, PrevSubpixelDepthViewportMax)
+		SHADER_PARAMETER(FScreenTransform, InputPixelPosToPrevScreenPosition)
+		SHADER_PARAMETER(FScreenTransform, ScreenPosToOutputPixelPos)
+		SHADER_PARAMETER(int32, MaxDepthAge)
+
+		SHADER_PARAMETER_RDG_TEXTURE(RWTexture2D, PrevSubpixelDepthTexture)
+
+		SHADER_PARAMETER_RDG_TEXTURE_UAV(RWTexture2D, SubpixelDepthOutput)
+		SHADER_PARAMETER_RDG_TEXTURE_UAV(RWTexture2DArray, DebugOutput)
+	END_SHADER_PARAMETER_STRUCT()
+}; // class FTSRForwardScatterDepthCS
+
 class FTSRDilateVelocityCS : public FTSRShader
 {
 	DECLARE_GLOBAL_SHADER(FTSRDilateVelocityCS);
 	SHADER_USE_PARAMETER_STRUCT(FTSRDilateVelocityCS, FTSRShader);
 
 	class FMotionBlurDirectionsDim : SHADER_PERMUTATION_INT("DIM_MOTION_BLUR_DIRECTIONS", 3);
-	using FPermutationDomain = TShaderPermutationDomain<FMotionBlurDirectionsDim>;
+	class FSubpixelDepthDim : SHADER_PERMUTATION_BOOL("DIM_SUBPIXEL_DEPTH");
+	using FPermutationDomain = TShaderPermutationDomain<FMotionBlurDirectionsDim, FSubpixelDepthDim>;
 
 	BEGIN_SHADER_PARAMETER_STRUCT(FParameters, )
 		SHADER_PARAMETER_STRUCT_INCLUDE(FTSRCommonParameters, CommonParameters)
@@ -405,15 +449,18 @@
 		SHADER_PARAMETER(float, VelocityExtrapolationMultiplier)
 		SHADER_PARAMETER(float, InvFlickeringMaxParralaxVelocity)
 		SHADER_PARAMETER(int32, bOutputIsMovingTexture)
+		SHADER_PARAMETER(int32, bIncludeDynamicDepthInSubpixelDetails)
 
 		SHADER_PARAMETER_RDG_TEXTURE(Texture2D, SceneDepthTexture)
 		SHADER_PARAMETER_RDG_TEXTURE(Texture2D, SceneVelocityTexture)
+		SHADER_PARAMETER_RDG_TEXTURE(Texture2D, PrevScatteredSubpixelDepthTexture)
 
 		SHADER_PARAMETER_RDG_TEXTURE_UAV(RWTexture2D, DilatedVelocityOutput)
 		SHADER_PARAMETER_RDG_TEXTURE_UAV(RWTexture2D, ClosestDepthOutput)
 		SHADER_PARAMETER_RDG_TEXTURE_UAV(RWTexture2D, PrevUseCountOutput)
 		SHADER_PARAMETER_RDG_TEXTURE_UAV(RWTexture2D, PrevClosestDepthOutput)
 		SHADER_PARAMETER_RDG_TEXTURE_UAV(RWTexture2DArray, R8Output)
+		SHADER_PARAMETER_RDG_TEXTURE_UAV(RWTexture2DArray, SubpixelDepthOutput)
 
 		SHADER_PARAMETER_RDG_TEXTURE_UAV(RWTexture2D, VelocityFlattenOutput)
 		SHADER_PARAMETER_RDG_TEXTURE_UAV(RWTexture2D, VelocityTileOutput0)
@@ -421,6 +468,28 @@
 
 		SHADER_PARAMETER_RDG_TEXTURE_UAV(RWTexture2DArray, DebugOutput)
 	END_SHADER_PARAMETER_STRUCT()
+
+	static FPermutationDomain RemapPermutation(FPermutationDomain PermutationVector)
+	{
+		// There isn't enough bindable UAV for outputing subpixel depth and motion blur render targets.
+		if (PermutationVector.Get<FSubpixelDepthDim>())
+		{
+			PermutationVector.Set<FMotionBlurDirectionsDim>(0);
+		}
+
+		return PermutationVector;
+	}
+
+	static bool ShouldCompilePermutation(const FGlobalShaderPermutationParameters& Parameters)
+	{
+		FPermutationDomain PermutationVector(Parameters.PermutationId);
+		if (PermutationVector != RemapPermutation(PermutationVector))
+		{
+			return false;
+		}
+
+		return FTSRShader::ShouldCompilePermutation(Parameters);
+	}
 }; // class FTSRDilateVelocityCS
 
 class FTSRDecimateHistoryCS : public FTSRShader
@@ -430,7 +499,8 @@
 
 	class FMoireReprojectionDim : SHADER_PERMUTATION_BOOL("DIM_MOIRE_REPROJECTION");
 	class FGrandReprojectionDim : SHADER_PERMUTATION_BOOL("DIM_GRAND_REPROJECTION");
-	using FPermutationDomain = TShaderPermutationDomain<FMoireReprojectionDim, FGrandReprojectionDim>;
+	class FSubpixelParallaxFactorDim : SHADER_PERMUTATION_BOOL("DIM_SUBPIXEL_PARALLAX_FACTOR");
+	using FPermutationDomain = TShaderPermutationDomain<FMoireReprojectionDim, FGrandReprojectionDim, FSubpixelParallaxFactorDim>;
 
 	BEGIN_SHADER_PARAMETER_STRUCT(FParameters, )
 		SHADER_PARAMETER_STRUCT_INCLUDE(FTSRCommonParameters, CommonParameters)
@@ -761,6 +831,7 @@
 
 IMPLEMENT_GLOBAL_SHADER(FTSRComputeMoireLumaCS,      "/Engine/Private/TemporalSuperResolution/TSRComputeMoireLuma.usf",      "MainCS", SF_Compute);
 IMPLEMENT_GLOBAL_SHADER(FTSRClearPrevTexturesCS,     "/Engine/Private/TemporalSuperResolution/TSRClearPrevTextures.usf",     "MainCS", SF_Compute);
+IMPLEMENT_GLOBAL_SHADER(FTSRForwardScatterDepthCS,   "/Engine/Private/TemporalSuperResolution/TSRForwardScatterDepth.usf",   "MainCS", SF_Compute);
 IMPLEMENT_GLOBAL_SHADER(FTSRDilateVelocityCS,        "/Engine/Private/TemporalSuperResolution/TSRDilateVelocity.usf",        "MainCS", SF_Compute);
 IMPLEMENT_GLOBAL_SHADER(FTSRDecimateHistoryCS,       "/Engine/Private/TemporalSuperResolution/TSRDecimateHistory.usf",       "MainCS", SF_Compute);
 IMPLEMENT_GLOBAL_SHADER(FTSRCompareTranslucencyCS,   "/Engine/Private/TemporalSuperResolution/TSRCompareTranslucency.usf",   "MainCS", SF_Compute);
@@ -846,6 +917,9 @@
 	const bool bSupportsLDS = FTSRShader::SupportsLDS(View.GetShaderPlatform());
 	const bool bUseWaveOps = (CVarTSRWaveOps.GetValueOnRenderThread() != 0 && GRHISupportsWaveOperations && bSupportsLDS && (WaveOpsSupport == ERHIFeatureSupport::RuntimeDependent || WaveOpsSupport == ERHIFeatureSupport::RuntimeGuaranteed)) || !bSupportsLDS;
 
+	// Whether should accumulate sub pixel depth.
+	const ETSRSubpixelMethod SubpixelMethod = ETSRSubpixelMethod(FMath::Clamp(CVarTSRSubpixelMethod.GetValueOnRenderThread(), 0, 2));
+
 	// Whether alpha channel is supported.
 	const bool bSupportsAlpha = IsPostProcessingWithAlphaChannelSupported();
 
@@ -1039,160 +1113,6 @@
 		return GraphBuilder.CreateUAV(DebugTexture);
 	};
 
-	// Clear atomic scattered texture.
-	FRDGTextureRef PrevUseCountTexture;
-	FRDGTextureRef PrevClosestDepthTexture;
-	{
-		{
-			FRDGTextureDesc Desc = FRDGTextureDesc::Create2D(
-				InputExtent,
-				PF_R32_UINT,
-				FClearValueBinding::None,
-				/* InFlags = */ TexCreate_ShaderResource | TexCreate_UAV | TexCreate_AtomicCompatible);
-
-			PrevUseCountTexture = GraphBuilder.CreateTexture(Desc, TEXT("TSR.PrevUseCountTexture"));
-			PrevClosestDepthTexture = GraphBuilder.CreateTexture(Desc, TEXT("TSR.PrevClosestDepthTexture"));
-		}
-
-		FTSRClearPrevTexturesCS::FParameters* PassParameters = GraphBuilder.AllocParameters<FTSRClearPrevTexturesCS::FParameters>();
-		PassParameters->CommonParameters = CommonParameters;
-		PassParameters->PrevUseCountOutput = GraphBuilder.CreateUAV(PrevUseCountTexture);
-		PassParameters->PrevClosestDepthOutput = GraphBuilder.CreateUAV(PrevClosestDepthTexture);
-
-		TShaderMapRef<FTSRClearPrevTexturesCS> ComputeShader(View.ShaderMap);
-		FComputeShaderUtils::AddPass(
-			GraphBuilder,
-			RDG_EVENT_NAME("TSR ClearPrevTextures %dx%d", InputRect.Width(), InputRect.Height()),
-			AsyncComputePasses >= 1 ? ERDGPassFlags::AsyncCompute : ERDGPassFlags::Compute,
-			ComputeShader,
-			PassParameters,
-			FComputeShaderUtils::GetGroupCount(InputRect.Size(), 8 * 2));
-	}
-
-	// Dilate the velocity texture & scatter reprojection into previous frame
-	FRDGTextureRef DilatedVelocityTexture;
-	FRDGTextureRef ClosestDepthTexture;
-	FRDGTextureSRVRef ParallaxFactorTexture;
-	FRDGTextureSRVRef IsMovingMaskTexture = nullptr;
-	FVelocityFlattenTextures VelocityFlattenTextures;
-	{
-		const bool bOutputIsMovingTexture = FlickeringFramePeriod > 0.0f;
-
-		FRDGTextureRef R8OutputTexture;
-
-		{
-			FRDGTextureDesc Desc = FRDGTextureDesc::Create2D(
-				InputExtent,
-				PF_G16R16,
-				FClearValueBinding::None,
-				/* InFlags = */ TexCreate_ShaderResource | TexCreate_UAV);
-
-			DilatedVelocityTexture = GraphBuilder.CreateTexture(Desc, TEXT("TSR.Velocity.Dilated"));
-
-			Desc.Format = PF_R16F;
-			ClosestDepthTexture = GraphBuilder.CreateTexture(Desc, TEXT("TSR.ClosestDepthTexture"));
-		}
-		{
-			FRDGTextureDesc Desc = FRDGTextureDesc::Create2DArray(
-				InputExtent,
-				PF_R8_UINT,
-				FClearValueBinding::None,
-				/* InFlags = */ TexCreate_ShaderResource | TexCreate_UAV,
-				bOutputIsMovingTexture ? 2 : 1);
-
-			R8OutputTexture = GraphBuilder.CreateTexture(Desc, TEXT("TSR.ParallaxFactor"));
-			ParallaxFactorTexture = GraphBuilder.CreateSRV(FRDGTextureSRVDesc::CreateForSlice(R8OutputTexture, 0));
-			if (bOutputIsMovingTexture)
-			{
-				IsMovingMaskTexture = GraphBuilder.CreateSRV(FRDGTextureSRVDesc::CreateForSlice(R8OutputTexture, 1));
-			}
-		}
-
-		int32 TileSize = 8;
-		FTSRDilateVelocityCS::FPermutationDomain PermutationVector;
-		FTSRDilateVelocityCS::FParameters* PassParameters = GraphBuilder.AllocParameters<FTSRDilateVelocityCS::FParameters>();
-		PassParameters->CommonParameters = CommonParameters;
-		PassParameters->RotationalClipToPrevClip = RotationalClipToPrevClip;
-		PassParameters->PrevOutputBufferUVMin = CommonParameters.InputInfo.UVViewportBilinearMin - CommonParameters.InputInfo.ExtentInverse;
-		PassParameters->PrevOutputBufferUVMax = CommonParameters.InputInfo.UVViewportBilinearMax + CommonParameters.InputInfo.ExtentInverse;
-		{
-			float TanHalfFieldOfView = View.ViewMatrices.GetInvProjectionMatrix().M[0][0];
-
-			// Should be multiplied 0.5* for the diameter to radius, and by 2.0 because GetTanHalfFieldOfView() cover only half of the pixels.
-			float WorldDepthToPixelWorldRadius = TanHalfFieldOfView / float(View.ViewRect.Width());
-
-			PassParameters->WorldDepthToDepthError = WorldDepthToPixelWorldRadius * 2.0f;
-		}
-		PassParameters->VelocityExtrapolationMultiplier = FMath::Clamp(CVarTSRVelocityExtrapolation.GetValueOnRenderThread(), 0.0f, 1.0f);
-		{
-			float FlickeringMaxParralaxVelocity = RefreshRateToFrameRateCap * CVarTSRFlickeringMaxParralaxVelocity.GetValueOnRenderThread() * float(View.ViewRect.Width()) / 1920.0f;
-			PassParameters->InvFlickeringMaxParralaxVelocity = 1.0f / FlickeringMaxParralaxVelocity;
-		}
-		PassParameters->bOutputIsMovingTexture = bOutputIsMovingTexture;
-
-		PassParameters->SceneDepthTexture = PassInputs.SceneDepthTexture;
-		PassParameters->SceneVelocityTexture = PassInputs.SceneVelocityTexture;
-
-		PassParameters->DilatedVelocityOutput = GraphBuilder.CreateUAV(DilatedVelocityTexture);
-		PassParameters->ClosestDepthOutput = GraphBuilder.CreateUAV(ClosestDepthTexture);
-		PassParameters->PrevUseCountOutput = GraphBuilder.CreateUAV(PrevUseCountTexture);
-		PassParameters->PrevClosestDepthOutput = GraphBuilder.CreateUAV(PrevClosestDepthTexture);
-		PassParameters->R8Output = GraphBuilder.CreateUAV(R8OutputTexture);
-
-		// Setup up the motion blur's velocity flatten pass.
-		if (PassInputs.bGenerateVelocityFlattenTextures)
-		{
-			const int32 MotionBlurDirections = GetMotionBlurDirections();
-			PermutationVector.Set<FTSRDilateVelocityCS::FMotionBlurDirectionsDim>(MotionBlurDirections);
-			TileSize = FVelocityFlattenTextures::kTileSize;
-
-			{
-				FRDGTextureDesc Desc = FRDGTextureDesc::Create2D(
-					InputExtent,
-					PF_FloatR11G11B10,
-					FClearValueBinding::None,
-					GFastVRamConfig.VelocityFlat | TexCreate_ShaderResource | TexCreate_UAV);
-
-				VelocityFlattenTextures.VelocityFlatten.Texture = GraphBuilder.CreateTexture(Desc, TEXT("MotionBlur.VelocityTile"));
-				VelocityFlattenTextures.VelocityFlatten.ViewRect = InputRect;
-			}
-
-			{
-				FRDGTextureDesc Desc = FRDGTextureDesc::Create2D(
-					FIntPoint::DivideAndRoundUp(InputRect.Size(), FVelocityFlattenTextures::kTileSize),
-					PF_FloatRGBA,
-					FClearValueBinding::None,
-					GFastVRamConfig.MotionBlur | TexCreate_ShaderResource | TexCreate_UAV);
-
-				VelocityFlattenTextures.VelocityTile[0].Texture = GraphBuilder.CreateTexture(Desc, TEXT("MotionBlur.VelocityTile"));
-				VelocityFlattenTextures.VelocityTile[0].ViewRect = FIntRect(FIntPoint::ZeroValue, Desc.Extent);
-
-				Desc.Format = PF_G16R16F;
-				VelocityFlattenTextures.VelocityTile[1].Texture = GraphBuilder.CreateTexture(Desc, TEXT("MotionBlur.VelocityTile"));
-				VelocityFlattenTextures.VelocityTile[1].ViewRect = FIntRect(FIntPoint::ZeroValue, Desc.Extent);
-			}
-
-			PassParameters->VelocityFlattenParameters = GetVelocityFlattenParameters(View);
-			PassParameters->VelocityFlattenOutput = GraphBuilder.CreateUAV(VelocityFlattenTextures.VelocityFlatten.Texture);
-			PassParameters->VelocityTileOutput0 = GraphBuilder.CreateUAV(VelocityFlattenTextures.VelocityTile[0].Texture);
-			PassParameters->VelocityTileOutput1 = GraphBuilder.CreateUAV(VelocityFlattenTextures.VelocityTile[1].Texture);
-		}
-
-		PassParameters->DebugOutput = CreateDebugUAV(InputExtent, TEXT("Debug.TSR.DilateVelocity"));
-
-		TShaderMapRef<FTSRDilateVelocityCS> ComputeShader(View.ShaderMap, PermutationVector);
-		FComputeShaderUtils::AddPass(
-			GraphBuilder,
-			RDG_EVENT_NAME("TSR DilateVelocity(MotionBlurDirections=%d%s) %dx%d",
-				int32(PermutationVector.Get<FTSRDilateVelocityCS::FMotionBlurDirectionsDim>()),
-				bOutputIsMovingTexture ? TEXT(" OutputIsMoving") : TEXT(""),
-				InputRect.Width(), InputRect.Height()),
-			AsyncComputePasses >= 2 ? ERDGPassFlags::AsyncCompute : ERDGPassFlags::Compute,
-			ComputeShader,
-			PassParameters,
-			FComputeShaderUtils::GetGroupCount(InputRect.Size(), TileSize));
-	}
-
 	// Create new history.
 	FRDGTextureRef UpdateHistoryOutputTexture;
 	FRDGTextureRef SceneColorOutputTexture;
@@ -1222,7 +1142,7 @@
 		{
 			SceneColorOutputTexture = UpdateHistoryOutputTexture;
 		}
-		
+
 		// Generate quarter res output if only one needed, otherwise only output half res and let caller downscale to quarter res. 
 		// This is motivated to saves UAV slots on FTSRUpdateHistoryCS
 		if (PassInputs.bGenerateSceneColorQuarterRes && !PassInputs.bGenerateSceneColorHalfRes && OutputRect.Size() == HistorySize)
@@ -1276,6 +1196,8 @@
 			History.TranslucencyAlpha = GraphBuilder.CreateTexture(Desc, TEXT("TSR.History.TranslucencyAlpha"));
 		}
 	}
+
+	if (SubpixelMethod == ETSRSubpixelMethod::ParallaxFactor)
 	{
 		FRDGTextureDesc SubpixelDetailsDesc = FRDGTextureDesc::Create2D(
 			InputExtent,
@@ -1284,6 +1206,16 @@
 			TexCreate_ShaderResource | TexCreate_UAV);
 		History.SubpixelDetails = GraphBuilder.CreateTexture(SubpixelDetailsDesc, TEXT("TSR.History.SubpixelInfo"));
 	}
+	else if (SubpixelMethod == ETSRSubpixelMethod::ClosestDepth)
 	{
+		FRDGTextureDesc SubpixelDepthDesc = FRDGTextureDesc::Create2D(
+			InputExtent,
+			PF_R32_UINT,
+			FClearValueBinding::None,
+			TexCreate_ShaderResource | TexCreate_UAV);
+		History.SubpixelDepth = GraphBuilder.CreateTexture(SubpixelDepthDesc, TEXT("TSR.History.SubpixelDepth"));
+	}
+
+	{
 		FRDGTextureDesc Desc = FRDGTextureDesc::Create2D(
 			InputExtent,
@@ -1308,8 +1240,6 @@
 		DummyHistorySRVs.HighFrequencyOverblur = GraphBuilder.CreateSRV(FRDGTextureSRVDesc::Create(BlackDummy));
 		DummyHistorySRVs.Guide = BlackDummy;
 		DummyHistorySRVs.Moire = BlackDummy;
-
-		DummyHistorySRVs.SubpixelDetails = BlackUintDummy;
 	}
 
 	// Setup the previous frame history
@@ -1330,7 +1260,8 @@
 		PrevHistory.Moire = InputHistory.Moire.IsValid() ? GraphBuilder.RegisterExternalTexture(InputHistory.Moire, TEXT("TSR.PrevHistory.Moire")) : DummyHistorySRVs.Moire;
 
 		// Register non-filterable history
-		PrevHistory.SubpixelDetails = GraphBuilder.RegisterExternalTexture(InputHistory.SubpixelDetails, TEXT("TSR.PrevHistory.SubpixelInfo"));
+		PrevHistory.SubpixelDetails = InputHistory.SubpixelDetails.IsValid() ? GraphBuilder.RegisterExternalTexture(InputHistory.SubpixelDetails, TEXT("TSR.PrevHistory.SubpixelInfo")) : nullptr;
+		PrevHistory.SubpixelDepth = InputHistory.SubpixelDepth.IsValid() ? GraphBuilder.RegisterExternalTexture(InputHistory.SubpixelDepth, TEXT("TSR.PrevHistory.SubpixelDepth")) : nullptr;
 
 		PrevHistorySRVs = CreateSRVs(GraphBuilder, PrevHistory);
 
@@ -1357,7 +1288,7 @@
 	{
 		// Setup prev history parameters.
 		FScreenPassTextureViewport PrevHistoryViewport(PrevHistorySRVs.HighFrequency->Desc.Texture->Desc.Extent, InputHistory.OutputViewportRect);
-		FScreenPassTextureViewport PrevSubpixelDetailsViewport(PrevHistorySRVs.SubpixelDetails->Desc.Extent, InputHistory.InputViewportRect);
+		FScreenPassTextureViewport PrevSubpixelDetailsViewport(PrevHistorySRVs.Guide->Desc.Extent, InputHistory.InputViewportRect);
 		FScreenPassTextureViewport GrandPrevSubpixelDetailsViewport(GrandPrevColorTexture->Desc.Texture->Desc.Extent, InputHistory.PrevOutputViewportRect);
 
 		if (bCameraCut)
@@ -1381,6 +1312,218 @@
 		PrevHistoryParameters.HistoryPreExposureCorrection = View.PreExposure / View.PrevViewInfo.SceneColorPreExposure;
 	}
 
+	// Clear atomic scattered texture.
+	FRDGTextureRef PrevUseCountTexture;
+	FRDGTextureRef PrevClosestDepthTexture;
+	FRDGTextureRef PrevScatteredSubpixelDepthTexture = nullptr;
+	{
+		{
+			FRDGTextureDesc Desc = FRDGTextureDesc::Create2D(
+				InputExtent,
+				PF_R32_UINT,
+				FClearValueBinding::None,
+				/* InFlags = */ TexCreate_ShaderResource | TexCreate_UAV | TexCreate_AtomicCompatible);
+
+			PrevUseCountTexture = GraphBuilder.CreateTexture(Desc, TEXT("TSR.PrevUseCountTexture"));
+			PrevClosestDepthTexture = GraphBuilder.CreateTexture(Desc, TEXT("TSR.PrevClosestDepthTexture"));
+		}
+
+		FTSRClearPrevTexturesCS::FParameters* PassParameters = GraphBuilder.AllocParameters<FTSRClearPrevTexturesCS::FParameters>();
+		PassParameters->CommonParameters = CommonParameters;
+		PassParameters->PrevUseCountOutput = GraphBuilder.CreateUAV(PrevUseCountTexture);
+		PassParameters->PrevClosestDepthOutput = GraphBuilder.CreateUAV(PrevClosestDepthTexture);
+		if (History.SubpixelDepth && PrevHistory.SubpixelDepth)
+		{
+			PrevScatteredSubpixelDepthTexture = GraphBuilder.CreateTexture(History.SubpixelDepth->Desc, TEXT("TSR.PrevScatteredSubpixelDepth"));
+			PassParameters->SubpixelDepthOutput = GraphBuilder.CreateUAV(PrevScatteredSubpixelDepthTexture);
+		}
+		else
+		{
+			PassParameters->SubpixelDepthOutput = CreateDummyUAV(GraphBuilder, PF_R32_UINT);
+		}
+
+		TShaderMapRef<FTSRClearPrevTexturesCS> ComputeShader(View.ShaderMap);
+		FComputeShaderUtils::AddPass(
+			GraphBuilder,
+			RDG_EVENT_NAME("TSR ClearPrevTextures %dx%d", InputRect.Width(), InputRect.Height()),
+			AsyncComputePasses >= 1 ? ERDGPassFlags::AsyncCompute : ERDGPassFlags::Compute,
+			ComputeShader,
+			PassParameters,
+			FComputeShaderUtils::GetGroupCount(InputRect.Size(), 8 * 2));
+	}
+
+	// Forward scatter previous frame sub pixel depths.
+	if (PrevScatteredSubpixelDepthTexture)
+	{
+		FIntRect ScatterDepthViewport = InputHistory.InputViewportRect;
+
+		FTSRForwardScatterDepthCS::FParameters* PassParameters = GraphBuilder.AllocParameters<FTSRForwardScatterDepthCS::FParameters>();
+		PassParameters->CommonParameters = CommonParameters;
+		PassParameters->PrevSubpixelDepthViewportMin = ScatterDepthViewport.Min;
+		PassParameters->PrevSubpixelDepthViewportMax = ScatterDepthViewport.Max - 1;
+		PassParameters->InputPixelPosToPrevScreenPosition = (FScreenTransform::Identity + 0.5f) * FScreenTransform::ChangeTextureBasisFromTo(
+			PrevHistory.SubpixelDepth->Desc.Extent, ScatterDepthViewport,
+			FScreenTransform::ETextureBasis::TexelPosition, FScreenTransform::ETextureBasis::ScreenPosition);
+		PassParameters->ScreenPosToOutputPixelPos = FScreenTransform::ChangeTextureBasisFromTo(
+			InputExtent, InputRect,
+			FScreenTransform::ETextureBasis::ScreenPosition, FScreenTransform::ETextureBasis::TexelPosition);
+		PassParameters->MaxDepthAge = FMath::Clamp(CVarTSRSubpixelDepthMaxAge.GetValueOnRenderThread(), 1, (1 << 7) - 1);
+
+		PassParameters->PrevSubpixelDepthTexture = PrevHistory.SubpixelDepth;
+		PassParameters->SubpixelDepthOutput = GraphBuilder.CreateUAV(PrevScatteredSubpixelDepthTexture);
+		PassParameters->DebugOutput = CreateDebugUAV(InputExtent, TEXT("Debug.TSR.ForwardScatterDepth"));
+
+		TShaderMapRef<FTSRForwardScatterDepthCS> ComputeShader(View.ShaderMap);
+		FComputeShaderUtils::AddPass(
+			GraphBuilder,
+			RDG_EVENT_NAME("TSR ForwardScatterDepth %dx%d", ScatterDepthViewport.Width(), ScatterDepthViewport.Height()),
+			AsyncComputePasses >= 1 ? ERDGPassFlags::AsyncCompute : ERDGPassFlags::Compute,
+			ComputeShader,
+			PassParameters,
+			FComputeShaderUtils::GetGroupCount(ScatterDepthViewport.Size(), 8));
+	}
+
+	// Dilate the velocity texture & scatter reprojection into previous frame
+	FRDGTextureRef DilatedVelocityTexture;
+	FRDGTextureRef ClosestDepthTexture;
+	FRDGTextureSRVRef ParallaxFactorTexture;
+	FRDGTextureSRVRef IsMovingMaskTexture = nullptr;
+	FVelocityFlattenTextures VelocityFlattenTextures;
+	{
+		const bool bOutputIsMovingTexture = FlickeringFramePeriod > 0.0f;
+
+		FRDGTextureRef R8OutputTexture;
+
+		{
+			FRDGTextureDesc Desc = FRDGTextureDesc::Create2D(
+				InputExtent,
+				PF_G16R16,
+				FClearValueBinding::None,
+				/* InFlags = */ TexCreate_ShaderResource | TexCreate_UAV);
+
+			DilatedVelocityTexture = GraphBuilder.CreateTexture(Desc, TEXT("TSR.Velocity.Dilated"));
+
+			Desc.Format = PF_R16F;
+			ClosestDepthTexture = GraphBuilder.CreateTexture(Desc, TEXT("TSR.ClosestDepthTexture"));
+		}
+		{
+			FRDGTextureDesc Desc = FRDGTextureDesc::Create2DArray(
+				InputExtent,
+				PF_R8_UINT,
+				FClearValueBinding::None,
+				/* InFlags = */ TexCreate_ShaderResource | TexCreate_UAV,
+				bOutputIsMovingTexture ? 2 : 1);
+
+			R8OutputTexture = GraphBuilder.CreateTexture(Desc, TEXT("TSR.ParallaxFactor"));
+			ParallaxFactorTexture = GraphBuilder.CreateSRV(FRDGTextureSRVDesc::CreateForSlice(R8OutputTexture, 0));
+			if (bOutputIsMovingTexture)
+			{
+				IsMovingMaskTexture = GraphBuilder.CreateSRV(FRDGTextureSRVDesc::CreateForSlice(R8OutputTexture, 1));
+			}
+		}
+
+		int32 TileSize = 8;
+		FTSRDilateVelocityCS::FPermutationDomain PermutationVector;
+		PermutationVector.Set<FTSRDilateVelocityCS::FSubpixelDepthDim>(SubpixelMethod == ETSRSubpixelMethod::ClosestDepth);
+
+		FTSRDilateVelocityCS::FParameters* PassParameters = GraphBuilder.AllocParameters<FTSRDilateVelocityCS::FParameters>();
+		PassParameters->CommonParameters = CommonParameters;
+		PassParameters->RotationalClipToPrevClip = RotationalClipToPrevClip;
+		PassParameters->PrevOutputBufferUVMin = CommonParameters.InputInfo.UVViewportBilinearMin - CommonParameters.InputInfo.ExtentInverse;
+		PassParameters->PrevOutputBufferUVMax = CommonParameters.InputInfo.UVViewportBilinearMax + CommonParameters.InputInfo.ExtentInverse;
+		{
+			float TanHalfFieldOfView = View.ViewMatrices.GetInvProjectionMatrix().M[0][0];
+
+			// Should be multiplied 0.5* for the diameter to radius, and by 2.0 because GetTanHalfFieldOfView() cover only half of the pixels.
+			float WorldDepthToPixelWorldRadius = TanHalfFieldOfView / float(View.ViewRect.Width());
+
+			PassParameters->WorldDepthToDepthError = WorldDepthToPixelWorldRadius * 2.0f;
+		}
+		PassParameters->VelocityExtrapolationMultiplier = FMath::Clamp(CVarTSRVelocityExtrapolation.GetValueOnRenderThread(), 0.0f, 1.0f);
+		{
+			float FlickeringMaxParralaxVelocity = RefreshRateToFrameRateCap * CVarTSRFlickeringMaxParralaxVelocity.GetValueOnRenderThread() * float(View.ViewRect.Width()) / 1920.0f;
+			PassParameters->InvFlickeringMaxParralaxVelocity = 1.0f / FlickeringMaxParralaxVelocity;
+		}
+		PassParameters->bOutputIsMovingTexture = bOutputIsMovingTexture;
+		PassParameters->bIncludeDynamicDepthInSubpixelDetails = SubpixelMethod == ETSRSubpixelMethod::ClosestDepth && CVarTSRSubpixelIncludeMovingDepth.GetValueOnRenderThread();
+
+		PassParameters->SceneDepthTexture = PassInputs.SceneDepthTexture;
+		PassParameters->SceneVelocityTexture = PassInputs.SceneVelocityTexture;
+		if (PrevScatteredSubpixelDepthTexture)
+		{
+			PassParameters->PrevScatteredSubpixelDepthTexture = PrevScatteredSubpixelDepthTexture;
+		}
+		else
+		{
+			PassParameters->PrevScatteredSubpixelDepthTexture = BlackUintDummy;
+		}
+
+		PassParameters->DilatedVelocityOutput = GraphBuilder.CreateUAV(DilatedVelocityTexture);
+		PassParameters->ClosestDepthOutput = GraphBuilder.CreateUAV(ClosestDepthTexture);
+		PassParameters->PrevUseCountOutput = GraphBuilder.CreateUAV(PrevUseCountTexture);
+		PassParameters->PrevClosestDepthOutput = GraphBuilder.CreateUAV(PrevClosestDepthTexture);
+		PassParameters->R8Output = GraphBuilder.CreateUAV(R8OutputTexture);
+		if (SubpixelMethod == ETSRSubpixelMethod::ClosestDepth)
+		{
+			PassParameters->SubpixelDepthOutput = GraphBuilder.CreateUAV(History.SubpixelDepth);
+		}
+
+		// Setup up the motion blur's velocity flatten pass.
+		if (PassInputs.bGenerateVelocityFlattenTextures && SubpixelMethod != ETSRSubpixelMethod::ClosestDepth)
+		{
+			const int32 MotionBlurDirections = GetMotionBlurDirections();
+			PermutationVector.Set<FTSRDilateVelocityCS::FMotionBlurDirectionsDim>(MotionBlurDirections);
+			TileSize = FVelocityFlattenTextures::kTileSize;
+
+			{
+				FRDGTextureDesc Desc = FRDGTextureDesc::Create2D(
+					InputExtent,
+					PF_FloatR11G11B10,
+					FClearValueBinding::None,
+					GFastVRamConfig.VelocityFlat | TexCreate_ShaderResource | TexCreate_UAV);
+
+				VelocityFlattenTextures.VelocityFlatten.Texture = GraphBuilder.CreateTexture(Desc, TEXT("MotionBlur.VelocityTile"));
+				VelocityFlattenTextures.VelocityFlatten.ViewRect = InputRect;
+			}
+
+			{
+				FRDGTextureDesc Desc = FRDGTextureDesc::Create2D(
+					FIntPoint::DivideAndRoundUp(InputRect.Size(), FVelocityFlattenTextures::kTileSize),
+					PF_FloatRGBA,
+					FClearValueBinding::None,
+					GFastVRamConfig.MotionBlur | TexCreate_ShaderResource | TexCreate_UAV);
+
+				VelocityFlattenTextures.VelocityTile[0].Texture = GraphBuilder.CreateTexture(Desc, TEXT("MotionBlur.VelocityTile"));
+				VelocityFlattenTextures.VelocityTile[0].ViewRect = FIntRect(FIntPoint::ZeroValue, Desc.Extent);
+
+				Desc.Format = PF_G16R16F;
+				VelocityFlattenTextures.VelocityTile[1].Texture = GraphBuilder.CreateTexture(Desc, TEXT("MotionBlur.VelocityTile"));
+				VelocityFlattenTextures.VelocityTile[1].ViewRect = FIntRect(FIntPoint::ZeroValue, Desc.Extent);
+			}
+
+			PassParameters->VelocityFlattenParameters = GetVelocityFlattenParameters(View);
+			PassParameters->VelocityFlattenOutput = GraphBuilder.CreateUAV(VelocityFlattenTextures.VelocityFlatten.Texture);
+			PassParameters->VelocityTileOutput0 = GraphBuilder.CreateUAV(VelocityFlattenTextures.VelocityTile[0].Texture);
+			PassParameters->VelocityTileOutput1 = GraphBuilder.CreateUAV(VelocityFlattenTextures.VelocityTile[1].Texture);
+		}
+
+		PassParameters->DebugOutput = CreateDebugUAV(InputExtent, TEXT("Debug.TSR.DilateVelocity"));
+
+		check(PermutationVector == FTSRDilateVelocityCS::RemapPermutation(PermutationVector));
+		TShaderMapRef<FTSRDilateVelocityCS> ComputeShader(View.ShaderMap, PermutationVector);
+		FComputeShaderUtils::AddPass(
+			GraphBuilder,
+			RDG_EVENT_NAME("TSR DilateVelocity(MotionBlurDirections=%d%s%s) %dx%d",
+				int32(PermutationVector.Get<FTSRDilateVelocityCS::FMotionBlurDirectionsDim>()),
+				bOutputIsMovingTexture ? TEXT(" OutputIsMoving") : TEXT(""),
+				PermutationVector.Get<FTSRDilateVelocityCS::FSubpixelDepthDim>() ? TEXT(" SubpixelDepth") : TEXT(""),
+				InputRect.Width(), InputRect.Height()),
+			AsyncComputePasses >= 2 ? ERDGPassFlags::AsyncCompute : ERDGPassFlags::Compute,
+			ComputeShader,
+			PassParameters,
+			FComputeShaderUtils::GetGroupCount(InputRect.Size(), TileSize));
+	}
+
 	// Decimate input to flicker at same frequency as input.
 	FRDGTextureRef ReprojectedHistoryGuideTexture = nullptr;
 	FRDGTextureRef ReprojectedHistoryMoireTexture = nullptr;
@@ -1445,7 +1588,7 @@
 		PassParameters->ParallaxFactorTexture = ParallaxFactorTexture;
 
 		PassParameters->PrevHistoryParameters = PrevHistoryParameters;
-		PassParameters->PrevHistorySubpixelDetails = PrevHistorySRVs.SubpixelDetails;
+		PassParameters->PrevHistorySubpixelDetails = PrevHistory.SubpixelDetails ? PrevHistory.SubpixelDetails : BlackUintDummy;
 		PassParameters->PrevVelocityTexture = PrevVelocity;
 
 		{
@@ -1467,18 +1610,28 @@
 		PassParameters->HoleFilledVelocityMaskOutput = GraphBuilder.CreateUAV(HoleFilledVelocityMaskTexture);
 		PassParameters->GrandHoleFilledVelocityOutput = GraphBuilder.CreateUAV(GrandHoleFilledVelocityTexture);
 		PassParameters->ParallaxRejectionMaskOutput = GraphBuilder.CreateUAV(ParallaxRejectionMaskTexture);
-		PassParameters->HistorySubpixelDetailsOutput = GraphBuilder.CreateUAV(History.SubpixelDetails);
+		if (History.SubpixelDetails)
+		{
+			PassParameters->HistorySubpixelDetailsOutput = GraphBuilder.CreateUAV(History.SubpixelDetails);
+		}
+		else
+		{
+			PassParameters->HistorySubpixelDetailsOutput = CreateDummyUAV(GraphBuilder, PF_R16_UINT);
+		}
 		PassParameters->DebugOutput = CreateDebugUAV(InputExtent, TEXT("Debug.TSR.DecimateHistory"));
 
 		FTSRDecimateHistoryCS::FPermutationDomain PermutationVector;
 		PermutationVector.Set<FTSRDecimateHistoryCS::FMoireReprojectionDim>(FlickeringFramePeriod > 0.0f);
 		PermutationVector.Set<FTSRDecimateHistoryCS::FGrandReprojectionDim>(bGrandReprojection);
+		PermutationVector.Set<FTSRDecimateHistoryCS::FSubpixelParallaxFactorDim>(SubpixelMethod == ETSRSubpixelMethod::ParallaxFactor);
 
 		TShaderMapRef<FTSRDecimateHistoryCS> ComputeShader(View.ShaderMap, PermutationVector);
 		FComputeShaderUtils::AddPass(
 			GraphBuilder,
-			RDG_EVENT_NAME("TSR DecimateHistory(%s) %dx%d",
+			RDG_EVENT_NAME("TSR DecimateHistory(%s%s%s) %dx%d",
 				PermutationVector.Get<FTSRDecimateHistoryCS::FMoireReprojectionDim>() ? TEXT("ReprojectMoire") : TEXT(""),
+				PermutationVector.Get<FTSRDecimateHistoryCS::FGrandReprojectionDim>() ? TEXT(" GrandReprojection") : TEXT(""),
+				PermutationVector.Get<FTSRDecimateHistoryCS::FSubpixelParallaxFactorDim>() ? TEXT(" SubpixelParallaxFactor") : TEXT(""),
 				InputRect.Width(), InputRect.Height()),
 			AsyncComputePasses >= 2 ? ERDGPassFlags::AsyncCompute : ERDGPassFlags::Compute,
 			ComputeShader,
@@ -1899,7 +2052,14 @@
 		}
 
 		// Extract non-filterable history
-		GraphBuilder.QueueTextureExtraction(History.SubpixelDetails, &OutputHistory.SubpixelDetails);
+		if (History.SubpixelDetails)
+		{
+			GraphBuilder.QueueTextureExtraction(History.SubpixelDetails, &OutputHistory.SubpixelDetails);
+		}
+		else if (History.SubpixelDepth)
+		{
+			GraphBuilder.QueueTextureExtraction(History.SubpixelDepth, &OutputHistory.SubpixelDepth);
+		}
 
 		// Extract history guide
 		GraphBuilder.QueueTextureExtraction(History.Guide, &OutputHistory.Guide);
diff --git a/Engine/Source/Runtime/Renderer/Private/SceneRendering.h b/Engine/Source/Runtime/Renderer/Private/SceneRendering.h
--- a/Engine/Source/Runtime/Renderer/Private/SceneRendering.h
+++ b/Engine/Source/Runtime/Renderer/Private/SceneRendering.h
@@ -984,6 +984,7 @@ struct FTSRHistory
 
 	// Input resolution representation of the output
 	TRefCountPtr<IPooledRenderTarget> SubpixelDetails;
+	TRefCountPtr<IPooledRenderTarget> SubpixelDepth;
 	TRefCountPtr<IPooledRenderTarget> Guide;
 	TRefCountPtr<IPooledRenderTarget> Moire;
 	TRefCountPtr<IPooledRenderTarget> Velocity;
