Fixes warning about TSR's PrevSubpixelDepthTexture

diff --git a/Engine/Source/Runtime/Renderer/Private/PostProcess/TemporalSuperResolution.cpp b/Engine/Source/Runtime/Renderer/Private/PostProcess/TemporalSuperResolution.cpp
--- a/Engine/Source/Runtime/Renderer/Private/PostProcess/TemporalSuperResolution.cpp
+++ b/Engine/Source/Runtime/Renderer/Private/PostProcess/TemporalSuperResolution.cpp
@@ -422,7 +422,7 @@ class FTSRForwardScatterDepthCS : public FTSRShader
 		SHADER_PARAMETER(FScreenTransform, ScreenPosToOutputPixelPos)
 		SHADER_PARAMETER(int32, MaxDepthAge)
 
-		SHADER_PARAMETER_RDG_TEXTURE(RWTexture2D, PrevSubpixelDepthTexture)
+		SHADER_PARAMETER_RDG_TEXTURE(Texture2D, PrevSubpixelDepthTexture)
 
 		SHADER_PARAMETER_RDG_TEXTURE_UAV(RWTexture2D, SubpixelDepthOutput)
 		SHADER_PARAMETER_RDG_TEXTURE_UAV(RWTexture2DArray, DebugOutput)
