Null check for SceneTextures

diff --git a/Engine/Source/Runtime/Renderer/Private/MobileBasePassRendering.cpp b/Engine/Source/Runtime/Renderer/Private/MobileBasePassRendering.cpp
--- a/Engine/Source/Runtime/Renderer/Private/MobileBasePassRendering.cpp
+++ b/Engine/Source/Runtime/Renderer/Private/MobileBasePassRendering.cpp
@@ -260,7 +260,7 @@ TRDGUniformBufferRef<FMobileBasePassUniformParameters> CreateMobileBasePassUnifo
 
 	// QuadOverdraw is a UAV so it needs to be initialized even if not used
 	const FSceneTextures* SceneTextures = View.GetSceneTexturesChecked();
-	FRDGTextureRef QuadOverdrawTexture = SceneTextures->QuadOverdraw;
+	FRDGTextureRef QuadOverdrawTexture = SceneTextures ? SceneTextures->QuadOverdraw : nullptr;
 	if (!QuadOverdrawTexture)
 	{
 		QuadOverdrawTexture = GraphBuilder.CreateTexture(FRDGTextureDesc::Create2D(FIntPoint(1, 1), PF_R32_UINT, FClearValueBinding::None, TexCreate_UAV), TEXT("DummyOverdrawUAV"));
