Adding ScreenPass variants that don't need Views.

diff --git a/Engine/Plugins/Media/MediaIOFramework/Source/MediaIOCore/MediaIOCore.Build.cs b/Engine/Plugins/Media/MediaIOFramework/Source/MediaIOCore/MediaIOCore.Build.cs
--- a/Engine/Plugins/Media/MediaIOFramework/Source/MediaIOCore/MediaIOCore.Build.cs
+++ b/Engine/Plugins/Media/MediaIOFramework/Source/MediaIOCore/MediaIOCore.Build.cs
@@ -29,14 +29,6 @@ namespace UnrealBuildTool.Rules
 					"TimeManagement"
 				});
 
-			var EngineDir = Path.GetFullPath(Target.RelativeEnginePath);
-
-			PrivateIncludePaths.AddRange(
-				new string[] {
-				//required for FScreenPassVS, AddDrawScreenPass, and for scene view extensions related headers
-				Path.Combine(EngineDir, "Source", "Runtime", "Renderer", "Private")
-			});
-
 			PrivateDependencyModuleNames.AddRange(
                 new string[]
                 {
diff --git a/Engine/Plugins/Media/MediaIOFramework/Source/MediaIOCore/Private/MediaCapture.cpp b/Engine/Plugins/Media/MediaIOFramework/Source/MediaIOCore/Private/MediaCapture.cpp
--- a/Engine/Plugins/Media/MediaIOFramework/Source/MediaIOCore/Private/MediaCapture.cpp
+++ b/Engine/Plugins/Media/MediaIOFramework/Source/MediaIOCore/Private/MediaCapture.cpp
@@ -26,7 +26,6 @@
 #include "ProfilingDebugging/ScopedTimers.h"
 #include "UObject/WeakObjectPtrTemplates.h"
 #include "RHIGPUReadback.h"
-#include "ScenePrivate.h"
 
 #include UE_INLINE_GENERATED_CPP_BY_NAME(MediaCapture)
 
@@ -498,18 +497,6 @@ namespace UE::MediaCaptureData
 			// Rectangle area to use from source
 			const FIntRect ViewRect(ConversionPassArgs.CopyInfo.GetSourceRect());
 
-			//Dummy ViewFamily/ViewInfo created to use built in Draw Screen/Texture Pass
-			FSceneViewFamily ViewFamily(FSceneViewFamily::ConstructionValues(nullptr, nullptr, FEngineShowFlags(ESFIM_Game))
-				.SetTime(FGameTime())
-				.SetGammaCorrection(1.0f));
-			FSceneViewInitOptions ViewInitOptions;
-			ViewInitOptions.ViewFamily = &ViewFamily;
-			ViewInitOptions.SetViewRectangle(ViewRect);
-			ViewInitOptions.ViewOrigin = FVector::ZeroVector;
-			ViewInitOptions.ViewRotationMatrix = FMatrix::Identity;
-			ViewInitOptions.ProjectionMatrix = FMatrix::Identity;
-			FViewInfo ViewInfo = FViewInfo(ViewInitOptions);
-
 			// If no conversion was required, go through a simple copy
 			if (Args.MediaCapture->ConversionOperation == EMediaCaptureConversionOperation::NONE && !bRequiresFormatConversion)
 			{
@@ -532,7 +519,7 @@ namespace UE::MediaCaptureData
 					FScreenPassTextureViewport OutputViewport(OutputTexture);
 					TShaderMapRef<FRGB8toUYVY8ConvertPS> PixelShader(GlobalShaderMap);
 					FRGB8toUYVY8ConvertPS::FParameters* Parameters = PixelShader->AllocateAndSetParameters(Args.GraphBuilder, ConversionPassArgs.SourceRGBTexture, MediaShaders::RgbToYuvRec709Scaled, MediaShaders::YUVOffset8bits, bDoLinearToSRGB, OutputTexture);
-					AddDrawScreenPass(Args.GraphBuilder, RDG_EVENT_NAME("RGBToUYVY 8 bit"), ViewInfo, OutputViewport, InputViewport, VertexShader, PixelShader, Parameters);
+					AddDrawScreenPass(Args.GraphBuilder, RDG_EVENT_NAME("RGBToUYVY 8 bit"), FScreenPassViewInfo(), OutputViewport, InputViewport, VertexShader, PixelShader, Parameters);
 				}
 				break;
 				case EMediaCaptureConversionOperation::RGB10_TO_YUVv210_10BIT:
@@ -543,7 +530,7 @@ namespace UE::MediaCaptureData
 					FScreenPassTextureViewport OutputViewport(OutputTexture);
 					TShaderMapRef<FRGB10toYUVv210ConvertPS> PixelShader(GlobalShaderMap);
 					FRGB10toYUVv210ConvertPS::FParameters* Parameters = PixelShader->AllocateAndSetParameters(Args.GraphBuilder, ConversionPassArgs.SourceRGBTexture, MediaShaders::RgbToYuvRec709Scaled, MediaShaders::YUVOffset10bits, bDoLinearToSRGB, OutputTexture);
-					AddDrawScreenPass(Args.GraphBuilder, RDG_EVENT_NAME("RGBToYUVv210"), ViewInfo, OutputViewport, InputViewport, VertexShader, PixelShader, Parameters);
+					AddDrawScreenPass(Args.GraphBuilder, RDG_EVENT_NAME("RGBToYUVv210"), FScreenPassViewInfo(), OutputViewport, InputViewport, VertexShader, PixelShader, Parameters);
 				}
 				break;
 				case EMediaCaptureConversionOperation::INVERT_ALPHA:
@@ -566,7 +553,7 @@ namespace UE::MediaCaptureData
 
 						TShaderMapRef<FModifyAlphaSwizzleRgbaPS> PixelShader(GlobalShaderMap, PermutationVector);
 						FModifyAlphaSwizzleRgbaPS::FParameters* Parameters = PixelShader->AllocateAndSetParameters(Args.GraphBuilder, ConversionPassArgs.SourceRGBTexture, OutputTexture);
-						AddDrawScreenPass(Args.GraphBuilder, RDG_EVENT_NAME("MediaCaptureSwizzle"), ViewInfo, OutputViewport, InputViewport, VertexShader, PixelShader, Parameters);
+						AddDrawScreenPass(Args.GraphBuilder, RDG_EVENT_NAME("MediaCaptureSwizzle"), FScreenPassViewInfo(), OutputViewport, InputViewport, VertexShader, PixelShader, Parameters);
 					}
 					break;
 				}
diff --git a/Engine/Plugins/Media/PixelCapture/Source/PixelCapture/PixelCapture.build.cs b/Engine/Plugins/Media/PixelCapture/Source/PixelCapture/PixelCapture.build.cs
--- a/Engine/Plugins/Media/PixelCapture/Source/PixelCapture/PixelCapture.build.cs
+++ b/Engine/Plugins/Media/PixelCapture/Source/PixelCapture/PixelCapture.build.cs
@@ -23,7 +23,6 @@ namespace UnrealBuildTool.Rules
 			string WebRtcIncludeDir = Path.Combine(EngineDir, "Source/ThirdParty/WebRTC/4664/Include");
 			PrivateIncludePaths.AddRange(new string[] {
 				Path.Combine(WebRtcIncludeDir, "third_party/libyuv/include"),
-				Path.Combine(EngineDir, "Source/Runtime/Renderer/Private"),
 			});
 
 			PrivateDependencyModuleNames.AddRange(new string[] {
diff --git a/Engine/Plugins/Media/PixelCapture/Source/PixelCapture/Private/PixelCaptureUtils.h b/Engine/Plugins/Media/PixelCapture/Source/PixelCapture/Private/PixelCaptureUtils.h
--- a/Engine/Plugins/Media/PixelCapture/Source/PixelCapture/Private/PixelCaptureUtils.h
+++ b/Engine/Plugins/Media/PixelCapture/Source/PixelCapture/Private/PixelCaptureUtils.h
@@ -8,7 +8,6 @@
 #include "ScreenPass.h"
 #include "ScreenRendering.h"
 #include "RenderGraphUtils.h"
-#include "ScenePrivate.h"
 
 /*
     * Copy from one texture to another.
@@ -132,39 +131,21 @@ inline void CopyTextureRDG(FRHICommandListImmediate& RHICmdList, FTextureRHIRef
 		// Rectangle area to use from source
 		const FIntRect ViewRect(FIntPoint(0, 0), InputTexture->Desc.Extent);
 
-		//Dummy ViewFamily/ViewInfo created to use built in Draw Screen/Texture Pass
-		FSceneViewFamily ViewFamily(FSceneViewFamily::ConstructionValues(nullptr, nullptr, FEngineShowFlags(ESFIM_Game))
-			.SetTime(FGameTime())
-			.SetGammaCorrection(1.0f));
-		FSceneViewInitOptions ViewInitOptions;
-		ViewInitOptions.ViewFamily = &ViewFamily;
-		ViewInitOptions.SetViewRectangle(ViewRect);
-		ViewInitOptions.ViewOrigin = FVector::ZeroVector;
-		ViewInitOptions.ViewRotationMatrix = FMatrix::Identity;
-		ViewInitOptions.ProjectionMatrix = FMatrix::Identity;
-		TSharedPtr<FViewInfo> View = MakeShared<FViewInfo>(ViewInitOptions);
-
 		TShaderMapRef<FModifyAlphaSwizzleRgbaPS> PixelShader(GlobalShaderMap, PermutationVector);
 		FModifyAlphaSwizzleRgbaPS::FParameters* PixelShaderParameters = PixelShader->AllocateAndSetParameters(GraphBuilder, InputTexture, OutputTexture);
 		
 		FRHIBlendState* BlendState = FScreenPassPipelineState::FDefaultBlendState::GetRHI();
 		FRHIDepthStencilState* DepthStencilState = FScreenPassPipelineState::FDefaultDepthStencilState::GetRHI();
 
-		ClearUnusedGraphResources(PixelShader, PixelShaderParameters);
-
-		const FScreenPassPipelineState PipelineState(VertexShader, PixelShader, BlendState, DepthStencilState);
-
-		GraphBuilder.AddPass(
+		AddDrawScreenPass(
+			GraphBuilder,
 			RDG_EVENT_NAME("PixelCapturerSwizzle"),
-			PixelShaderParameters,
-			ERDGPassFlags::Raster,
-			[View, OutputViewport, InputViewport, PipelineState, PixelShader, PixelShaderParameters](FRHICommandList& RHICmdList)
-		{
-			DrawScreenPass(RHICmdList, *View, OutputViewport, InputViewport, PipelineState, EScreenPassDrawFlags::None, [&](FRHICommandList&)
-			{
-				SetShaderParameters(RHICmdList, PixelShader, PixelShader.GetPixelShader(), *PixelShaderParameters);
-			});
-		});
+			FScreenPassViewInfo(),
+			OutputViewport,
+			InputViewport,
+			VertexShader,
+			PixelShader,
+			PixelShaderParameters);
 	}
 	GraphBuilder.Execute();
 }
diff --git a/Engine/Source/Runtime/Renderer/Private/ScreenPass.cpp b/Engine/Source/Runtime/Renderer/Private/ScreenPass.cpp
--- a/Engine/Source/Runtime/Renderer/Private/ScreenPass.cpp
+++ b/Engine/Source/Runtime/Renderer/Private/ScreenPass.cpp
@@ -154,7 +154,7 @@ void SetScreenPassPipelineState(FRHICommandList& RHICmdList, const FScreenPassPi
 
 void DrawScreenPass_PostSetup(
 	FRHICommandList& RHICmdList,
-	const FSceneView& View,
+	const FScreenPassViewInfo& ViewInfo,
 	const FScreenPassTextureViewport& OutputViewport,
 	const FScreenPassTextureViewport& InputViewport,
 	const FScreenPassPipelineState& PipelineState,
@@ -169,7 +169,7 @@ void DrawScreenPass_PostSetup(
 	FIntPoint LocalOutputSize(OutputSize);
 	EDrawRectangleFlags DrawRectangleFlags = EDRF_UseTriangleOptimization;
 
-	const bool bUseHMDHiddenAreaMask = EnumHasAllFlags(Flags, EScreenPassDrawFlags::AllowHMDHiddenAreaMask) && View.bHMDHiddenAreaMaskActive;
+	const bool bUseHMDHiddenAreaMask = EnumHasAllFlags(Flags, EScreenPassDrawFlags::AllowHMDHiddenAreaMask) && ViewInfo.bHMDHiddenAreaMaskActive;
 
 	DrawPostProcessPass(
 		RHICmdList,
@@ -178,7 +178,7 @@ void DrawScreenPass_PostSetup(
 		OutputSize,
 		InputSize,
 		PipelineState.VertexShader,
-		View.StereoViewIndex,
+		ViewInfo.StereoViewIndex,
 		bUseHMDHiddenAreaMask,
 		DrawRectangleFlags);
 }
diff --git a/Engine/Source/Runtime/Renderer/Public/ScreenPass.h b/Engine/Source/Runtime/Renderer/Public/ScreenPass.h
--- a/Engine/Source/Runtime/Renderer/Public/ScreenPass.h
+++ b/Engine/Source/Runtime/Renderer/Public/ScreenPass.h
@@ -12,7 +12,7 @@
 #include "SceneView.h"
 
 // Returns whether a HMD hidden area mask is being used for VR.
-bool IsHMDHiddenAreaMaskActive();
+RENDERER_API bool IsHMDHiddenAreaMaskActive();
 
 // Returns the global engine mini font texture.
 FRHITexture* GetMiniFontTexture();
@@ -392,9 +392,28 @@ enum class EScreenPassDrawFlags : uint8
 };
 ENUM_CLASS_FLAGS(EScreenPassDrawFlags);
 
+/** Type used to carry the limited amount of data we need from a FSceneView. */
+struct FScreenPassViewInfo
+{
+	const int32 StereoViewIndex;
+	const bool bHMDHiddenAreaMaskActive;
+
+	FScreenPassViewInfo()
+		: StereoViewIndex(INDEX_NONE)
+		, bHMDHiddenAreaMaskActive(IsHMDHiddenAreaMaskActive())
+	{
+	}
+
+	FScreenPassViewInfo(const FSceneView& View)
+		: StereoViewIndex(View.StereoViewIndex)
+		, bHMDHiddenAreaMaskActive(View.bHMDHiddenAreaMaskActive)
+	{
+	}
+};
+
 RENDERER_API void DrawScreenPass_PostSetup(
 	FRHICommandList& RHICmdList,
-	const FSceneView& View,
+	const FScreenPassViewInfo& ViewInfo,
 	const FScreenPassTextureViewport& OutputViewport,
 	const FScreenPassTextureViewport& InputViewport,
 	const FScreenPassPipelineState& PipelineState,
@@ -408,7 +427,7 @@ RENDERER_API void DrawScreenPass_PostSetup(
 template<typename TSetupFunction>
 void DrawScreenPass(
 	FRHICommandList& RHICmdList,
-	const FSceneView& View,
+	const FScreenPassViewInfo ViewInfo,
 	const FScreenPassTextureViewport& OutputViewport,
 	const FScreenPassTextureViewport& InputViewport,
 	const FScreenPassPipelineState& PipelineState,
@@ -425,7 +444,7 @@ void DrawScreenPass(
 
 	SetupFunction(RHICmdList);
 
-	DrawScreenPass_PostSetup(RHICmdList, View, OutputViewport, InputViewport, PipelineState, Flags);
+	DrawScreenPass_PostSetup(RHICmdList, ViewInfo, OutputViewport, InputViewport, PipelineState, Flags);
 }
 
 /** Render graph variant of simpler DrawScreenPass function. Clears graph resources unused by the
@@ -435,7 +454,7 @@ template <typename PixelShaderType>
 FORCEINLINE void AddDrawScreenPass(
 	FRDGBuilder& GraphBuilder,
 	FRDGEventName&& PassName,
-	const FSceneView& View,
+	const FScreenPassViewInfo ViewInfo,
 	const FScreenPassTextureViewport& OutputViewport,
 	const FScreenPassTextureViewport& InputViewport,
 	const TShaderRef<FShader>& VertexShader,
@@ -457,9 +476,9 @@ FORCEINLINE void AddDrawScreenPass(
 		Forward<FRDGEventName&&>(PassName),
 		PixelShaderParameters,
 		ERDGPassFlags::Raster,
-		[&View, OutputViewport, InputViewport, PipelineState, PixelShader, PixelShaderParameters, Flags](FRHICommandList& RHICmdList)
+		[ViewInfo, OutputViewport, InputViewport, PipelineState, PixelShader, PixelShaderParameters, Flags](FRHICommandList& RHICmdList)
 	{
-		DrawScreenPass(RHICmdList, View, OutputViewport, InputViewport, PipelineState, Flags, [&](FRHICommandList&)
+		DrawScreenPass(RHICmdList, ViewInfo, OutputViewport, InputViewport, PipelineState, Flags, [&](FRHICommandList&)
 		{
 			SetShaderParameters(RHICmdList, PixelShader, PixelShader.GetPixelShader(), *PixelShaderParameters);
 		});
@@ -470,7 +489,7 @@ template <typename PixelShaderType>
 FORCEINLINE void AddDrawScreenPass(
 	FRDGBuilder& GraphBuilder,
 	FRDGEventName&& PassName,
-	const FSceneView& View,
+	const FScreenPassViewInfo ViewInfo,
 	const FScreenPassTextureViewport& OutputViewport,
 	const FScreenPassTextureViewport& InputViewport,
 	const TShaderRef<FShader>& VertexShader,
@@ -480,14 +499,14 @@ FORCEINLINE void AddDrawScreenPass(
 {
 	FRHIBlendState* BlendState = FScreenPassPipelineState::FDefaultBlendState::GetRHI();
 	FRHIDepthStencilState* DepthStencilState = FScreenPassPipelineState::FDefaultDepthStencilState::GetRHI();
-	AddDrawScreenPass(GraphBuilder, Forward<FRDGEventName&&>(PassName), View, OutputViewport, InputViewport, VertexShader, PixelShader, BlendState, DepthStencilState, PixelShaderParameters, Flags);
+	AddDrawScreenPass(GraphBuilder, Forward<FRDGEventName&&>(PassName), ViewInfo, OutputViewport, InputViewport, VertexShader, PixelShader, BlendState, DepthStencilState, PixelShaderParameters, Flags);
 }
 
 template <typename PixelShaderType>
 FORCEINLINE void AddDrawScreenPass(
 	FRDGBuilder& GraphBuilder,
 	FRDGEventName&& PassName,
-	const FSceneView& View,
+	const FScreenPassViewInfo ViewInfo,
 	const FScreenPassTextureViewport& OutputViewport,
 	const FScreenPassTextureViewport& InputViewport,
 	const TShaderRef<FShader>& VertexShader,
@@ -497,14 +516,14 @@ FORCEINLINE void AddDrawScreenPass(
 	EScreenPassDrawFlags Flags = EScreenPassDrawFlags::None)
 {
 	FRHIDepthStencilState* DepthStencilState = FScreenPassPipelineState::FDefaultDepthStencilState::GetRHI();
-	AddDrawScreenPass(GraphBuilder, Forward<FRDGEventName&&>(PassName), View, OutputViewport, InputViewport, VertexShader, PixelShader, BlendState, DepthStencilState, PixelShaderParameters, Flags);
+	AddDrawScreenPass(GraphBuilder, Forward<FRDGEventName&&>(PassName), ViewInfo, OutputViewport, InputViewport, VertexShader, PixelShader, BlendState, DepthStencilState, PixelShaderParameters, Flags);
 }
 
 template <typename PixelShaderType>
 FORCEINLINE void AddDrawScreenPass(
 	FRDGBuilder& GraphBuilder,
 	FRDGEventName&& PassName,
-	const FSceneView& View,
+	const FScreenPassViewInfo ViewInfo,
 	const FScreenPassTextureViewport& OutputViewport,
 	const FScreenPassTextureViewport& InputViewport,
 	const TShaderRef<FShader>& VertexShader,
@@ -514,7 +533,7 @@ FORCEINLINE void AddDrawScreenPass(
 	EScreenPassDrawFlags Flags = EScreenPassDrawFlags::None)
 {
 	FRHIBlendState* BlendState = FScreenPassPipelineState::FDefaultBlendState::GetRHI();
-	AddDrawScreenPass(GraphBuilder, Forward<FRDGEventName&&>(PassName), View, OutputViewport, InputViewport, VertexShader, PixelShader, BlendState, DepthStencilState, PixelShaderParameters, Flags);
+	AddDrawScreenPass(GraphBuilder, Forward<FRDGEventName&&>(PassName), ViewInfo, OutputViewport, InputViewport, VertexShader, PixelShader, BlendState, DepthStencilState, PixelShaderParameters, Flags);
 }
 
 template <typename PixelShaderType>
@@ -535,6 +554,24 @@ FORCEINLINE void AddDrawScreenPass(
 	AddDrawScreenPass(GraphBuilder, Forward<FRDGEventName&&>(PassName), View, OutputViewport, InputViewport, VertexShader, PixelShader, BlendState, DepthStencilState, PixelShaderParameters, Flags);
 }
 
+template <typename PixelShaderType>
+FORCEINLINE void AddDrawScreenPass(
+	FRDGBuilder& GraphBuilder,
+	FRDGEventName&& PassName,
+	ERHIFeatureLevel::Type FeatureLevel,
+	const FScreenPassTextureViewport& OutputViewport,
+	const FScreenPassTextureViewport& InputViewport,
+	const TShaderRef<PixelShaderType>& PixelShader,
+	typename PixelShaderType::FParameters* PixelShaderParameters,
+	EScreenPassDrawFlags Flags = EScreenPassDrawFlags::None)
+{
+	FGlobalShaderMap* ShaderMap = GetGlobalShaderMap(FeatureLevel);
+	TShaderMapRef<FScreenPassVS> VertexShader(ShaderMap);
+	FRHIBlendState* BlendState = FScreenPassPipelineState::FDefaultBlendState::GetRHI();
+	FRHIDepthStencilState* DepthStencilState = FScreenPassPipelineState::FDefaultDepthStencilState::GetRHI();
+	AddDrawScreenPass(GraphBuilder, Forward<FRDGEventName&&>(PassName), FScreenPassViewInfo(), OutputViewport, InputViewport, VertexShader, PixelShader, BlendState, DepthStencilState, PixelShaderParameters, Flags);
+}
+
 /** Render graph variant of more advanced DrawScreenPass function. Does *not* clear unused graph
  *  resources, since the parameters might be shared between the vertex and pixel shaders.
  */
@@ -542,7 +579,7 @@ template <typename TSetupFunction, typename TPassParameterStruct>
 FORCEINLINE void AddDrawScreenPass(
 	FRDGBuilder& GraphBuilder,
 	FRDGEventName&& PassName,
-	const FSceneView& View,
+	const FScreenPassViewInfo ViewInfo,
 	const FScreenPassTextureViewport& OutputViewport,
 	const FScreenPassTextureViewport& InputViewport,
 	const FScreenPassPipelineState& PipelineState,
@@ -557,9 +594,9 @@ FORCEINLINE void AddDrawScreenPass(
 		Forward<FRDGEventName&&>(PassName),
 		PassParameterStruct,
 		ERDGPassFlags::Raster,
-		[&View, OutputViewport, InputViewport, PipelineState, SetupFunction, Flags] (FRHICommandList& RHICmdList)
+		[ViewInfo, OutputViewport, InputViewport, PipelineState, SetupFunction, Flags] (FRHICommandList& RHICmdList)
 	{
-		DrawScreenPass(RHICmdList, View, OutputViewport, InputViewport, PipelineState, Flags, SetupFunction);
+		DrawScreenPass(RHICmdList, ViewInfo, OutputViewport, InputViewport, PipelineState, Flags, SetupFunction);
 	});
 }
 
@@ -567,14 +604,14 @@ template <typename TSetupFunction, typename TPassParameterStruct>
 FORCEINLINE void AddDrawScreenPass(
 	FRDGBuilder& GraphBuilder,
 	FRDGEventName&& PassName,
-	const FSceneView& View,
+	const FScreenPassViewInfo ViewInfo,
 	const FScreenPassTextureViewport& OutputViewport,
 	const FScreenPassTextureViewport& InputViewport,
 	const FScreenPassPipelineState& PipelineState,
 	TPassParameterStruct* PassParameterStruct,
 	TSetupFunction SetupFunction)
 {
-	AddDrawScreenPass(GraphBuilder, Forward<FRDGEventName&&>(PassName), View, OutputViewport, InputViewport, PipelineState, PassParameterStruct, EScreenPassDrawFlags::None, SetupFunction);
+	AddDrawScreenPass(GraphBuilder, Forward<FRDGEventName&&>(PassName), ViewInfo, OutputViewport, InputViewport, PipelineState, PassParameterStruct, EScreenPassDrawFlags::None, SetupFunction);
 }
 
 /** Helper function which copies a region of an input texture to a region of the output texture,
