Fixed FViewInfo snapshot allocator and SetupUniformBufferParameters to be threadsafe.

diff --git a/Engine/Source/Runtime/Renderer/Private/SceneRendering.cpp b/Engine/Source/Runtime/Renderer/Private/SceneRendering.cpp
--- a/Engine/Source/Runtime/Renderer/Private/SceneRendering.cpp
+++ b/Engine/Source/Runtime/Renderer/Private/SceneRendering.cpp
@@ -1492,6 +1492,10 @@ void FViewInfo::SetupUniformBufferParameters(
 			FinalMaterialTextureMipBias += MaterialTextureMipBias;
 		}
 
+		// Protect access to the view state sampler caches when called from multiple tasks.
+		static FCriticalSection CS;
+		FScopeLock Lock(&CS);
+
 		FSamplerStateRHIRef WrappedSampler = nullptr;
 		FSamplerStateRHIRef ClampedSampler = nullptr;
 
@@ -1527,6 +1531,30 @@ void FViewInfo::SetupUniformBufferParameters(
 			ViewState->MaterialTextureBilinearWrapedSamplerCache = WrappedSampler;
 			ViewState->MaterialTextureBilinearClampedSamplerCache = ClampedSampler;
 		}
+
+		// Landscape global resources
+		{
+			FSamplerStateRHIRef WeightmapSampler = nullptr;
+			if (ViewState && FMath::Abs(ViewState->LandscapeCachedMipBias - FinalMaterialTextureMipBias) < KINDA_SMALL_NUMBER)
+			{
+				// use cached sampler
+				WeightmapSampler = ViewState->LandscapeWeightmapSamplerCache;
+			}
+			else
+			{
+				// create a new one
+				ESamplerFilter Filter = bIsValidTextureGroupSamplerFilters ? TerrainWeightmapTextureGroupSamplerFilter : SF_AnisotropicPoint;
+				WeightmapSampler = RHICreateSamplerState(FSamplerStateInitializerRHI(Filter, AM_Clamp, AM_Clamp, AM_Clamp, FinalMaterialTextureMipBias));
+			}
+			check(WeightmapSampler.IsValid());
+			ViewUniformShaderParameters.LandscapeWeightmapSampler = WeightmapSampler;
+
+			if (ViewState)
+			{
+				ViewState->LandscapeCachedMipBias = FinalMaterialTextureMipBias;
+				ViewState->LandscapeWeightmapSamplerCache = WeightmapSampler;
+			}
+		}
 	}
 
 	{
@@ -1916,30 +1944,6 @@ void FViewInfo::SetupUniformBufferParameters(
 		ViewUniformShaderParameters.WaterData = GWhiteVertexBufferWithSRV->ShaderResourceViewRHI;
 	}
 
-	// Landscape global resources
-	{
-		FSamplerStateRHIRef WeightmapSampler = nullptr;
-		if (ViewState && FMath::Abs(ViewState->LandscapeCachedMipBias - FinalMaterialTextureMipBias) < KINDA_SMALL_NUMBER)
-		{
-			// use cached sampler
-			WeightmapSampler = ViewState->LandscapeWeightmapSamplerCache;
-		}
-		else
-		{
-			// create a new one
-			ESamplerFilter Filter = bIsValidTextureGroupSamplerFilters ? TerrainWeightmapTextureGroupSamplerFilter : SF_AnisotropicPoint;
-			WeightmapSampler = RHICreateSamplerState(FSamplerStateInitializerRHI(Filter, AM_Clamp, AM_Clamp, AM_Clamp, FinalMaterialTextureMipBias));
-		}
-		check(WeightmapSampler.IsValid());
-		ViewUniformShaderParameters.LandscapeWeightmapSampler = WeightmapSampler;
-
-		if (ViewState)
-		{
-			ViewState->LandscapeCachedMipBias = FinalMaterialTextureMipBias;
-			ViewState->LandscapeWeightmapSamplerCache = WeightmapSampler;
-		}
-	}
-
 	if (LandscapePerComponentDataBuffer.IsValid() && LandscapeIndirectionBuffer.IsValid())
 	{
 		ViewUniformShaderParameters.LandscapeIndirection = LandscapeIndirectionBuffer.GetReference();
@@ -2067,10 +2071,18 @@ const FViewInfo* FViewInfo::GetInstancedView() const
 	return nullptr;
 }
 
-// These are not real view infos, just dumb memory blocks
-static TArray<FViewInfo*> ViewInfoSnapshots;
-// these are never freed, even at program shutdown
-static TArray<FViewInfo*> FreeViewInfoSnapshots;
+struct FViewInfoSnapshotCache
+{
+	// These are not real view infos, just dumb memory blocks
+	TLockFreePointerListLIFOPad<FViewInfo, PLATFORM_CACHE_LINE_SIZE> Snapshots;
+	// these are never freed, even at program shutdown
+	TLockFreePointerListLIFOPad<FViewInfo, PLATFORM_CACHE_LINE_SIZE> FreeSnapshots;
+
+	std::atomic_int32_t NumSnapshots{ 0 };
+	std::atomic_int32_t NumFreeSnapshots{ 0 };
+};
+
+static FViewInfoSnapshotCache ViewInfoSnapshotCache;
 
 extern TUniformBufferRef<FMobileDirectionalLightShaderParameters>& GetNullMobileDirectionalLightShaderParameters();
 
@@ -2078,16 +2090,20 @@ FViewInfo* FViewInfo::CreateSnapshot() const
 {
 	QUICK_SCOPE_CYCLE_COUNTER(STAT_FViewInfo_CreateSnapshot);
 
-	check(IsInRenderingThread()); // we do not want this popped before the end of the scene and it better be the scene allocator
-	FViewInfo* Result;
-	if (FreeViewInfoSnapshots.Num())
+	check(IsInParallelRenderingThread()); // we do not want this popped before the end of the scene and it better be the scene allocator
+	FViewInfo* Result = ViewInfoSnapshotCache.FreeSnapshots.Pop();
+	if (Result)
 	{
-		Result = FreeViewInfoSnapshots.Pop(false);
+		ViewInfoSnapshotCache.NumFreeSnapshots--;
 	}
 	else
 	{
 		Result = (FViewInfo*)FMemory::Malloc(sizeof(FViewInfo), alignof(FViewInfo));
 	}
+
+	ViewInfoSnapshotCache.Snapshots.Push(Result);
+	ViewInfoSnapshotCache.NumSnapshots++;
+
 	FMemory::Memcpy(*Result, *this);
 
 	// we want these to start null without a reference count, since we clear a ref later
@@ -2113,7 +2129,6 @@ FViewInfo* FViewInfo::CreateSnapshot() const
 	Result->DynamicPrimitiveCollector = FGPUScenePrimitiveCollector(DynamicPrimitiveCollector);
 
 	Result->bIsSnapshot = true;
-	ViewInfoSnapshots.Add(Result);
 	return Result;
 }
 
@@ -2122,16 +2137,16 @@ void FViewInfo::DestroyAllSnapshots(FParallelMeshDrawCommandPass::EWaitThread Wa
 	QUICK_SCOPE_CYCLE_COUNTER(STAT_FViewInfo_DestroyAllSnapshots);
 
 	// we will only keep double the number actually used, plus a few
-	int32 NumToRemove = FreeViewInfoSnapshots.Num() - (ViewInfoSnapshots.Num() + 2);
-	if (NumToRemove > 0)
+	const int32 NumToRemove = FMath::Max(ViewInfoSnapshotCache.NumFreeSnapshots - (ViewInfoSnapshotCache.NumSnapshots + 2), 0);
+
+	for (int32 Index = 0; Index < NumToRemove; Index++)
 	{
-		for (int32 Index = 0; Index < NumToRemove; Index++)
-		{
-			FMemory::Free(FreeViewInfoSnapshots[Index]);
-		}
-		FreeViewInfoSnapshots.RemoveAt(0, NumToRemove, false);
+		FViewInfo* ViewInfo = ViewInfoSnapshotCache.FreeSnapshots.Pop();
+		check(ViewInfo);
+		FMemory::Free(ViewInfo);
 	}
-	for (FViewInfo* Snapshot : ViewInfoSnapshots)
+
+	while (FViewInfo* Snapshot = ViewInfoSnapshotCache.Snapshots.Pop())
 	{
 		Snapshot->ViewUniformBuffer.SafeRelease();
 		Snapshot->InstancedViewUniformBuffer.SafeRelease();
@@ -2147,9 +2162,11 @@ void FViewInfo::DestroyAllSnapshots(FParallelMeshDrawCommandPass::EWaitThread Wa
 			Snapshot->ParallelMeshDrawCommandPasses[i].FreeCreateSnapshot();
 		}
 
-		FreeViewInfoSnapshots.Add(Snapshot);
+		ViewInfoSnapshotCache.FreeSnapshots.Push(Snapshot);
 	}
-	ViewInfoSnapshots.Reset();
+
+	ViewInfoSnapshotCache.NumFreeSnapshots = ViewInfoSnapshotCache.NumFreeSnapshots + ViewInfoSnapshotCache.NumSnapshots - NumToRemove;
+	ViewInfoSnapshotCache.NumSnapshots = 0;
 }
 
 FInt32Range FViewInfo::GetDynamicMeshElementRange(uint32 PrimitiveIndex) const
