Moved CopyPersistentLightState to an RDG setup task.

diff --git a/Engine/Source/Runtime/Renderer/Private/ShadowDepthRendering.cpp b/Engine/Source/Runtime/Renderer/Private/ShadowDepthRendering.cpp
--- a/Engine/Source/Runtime/Renderer/Private/ShadowDepthRendering.cpp
+++ b/Engine/Source/Runtime/Renderer/Private/ShadowDepthRendering.cpp
@@ -1242,7 +1242,7 @@ FViewInfo* FProjectedShadowInfo::FindViewForShadow(FSceneRenderer* SceneRenderer
 void FProjectedShadowInfo::SetupShadowDepthView(FSceneRenderer* SceneRenderer)
 {
 	FViewInfo* FoundView = FindViewForShadow(SceneRenderer);
-	check(FoundView && IsInRenderingThread());
+	check(FoundView && IsInParallelRenderingThread());
 	FViewInfo* DepthPassView = FoundView->CreateSnapshot();
 	// We are starting a new collection of dynamic primitives for the shadow views.
 	DepthPassView->DynamicPrimitiveCollector = FGPUScenePrimitiveCollector(&SceneRenderer->GetGPUSceneDynamicContext());
@@ -1871,6 +1871,7 @@ void FSceneRenderer::RenderShadowDepthMaps(FRDGBuilder& GraphBuilder, FInstanceC
 	}
 
 	// Move current persistent shadow state to previous and clear current.
+	GraphBuilder.AddSetupTask([Scene = Scene]
 	{
 		TRACE_CPUPROFILER_EVENT_SCOPE(CopyPersistentLightState);
 		// TODO: This could be very slow.
@@ -1879,7 +1880,7 @@ void FSceneRenderer::RenderShadowDepthMaps(FRDGBuilder& GraphBuilder, FInstanceC
 			Light.LightSceneInfo->PrevPersistentShadows = Light.LightSceneInfo->PersistentShadows;
 			Light.LightSceneInfo->PersistentShadows.Empty();
 		}
-	}
+	});
 
 	bShadowDepthRenderCompleted = true;
 }
