Ray Tracing: Remove the redundant TWO_SIDED_MATERIAL define in favor of the existing MATERIAL_TWOSIDED

diff --git a/Engine/Shaders/Private/RayTracing/RayTracingMaterialHitShaders.usf b/Engine/Shaders/Private/RayTracing/RayTracingMaterialHitShaders.usf
--- a/Engine/Shaders/Private/RayTracing/RayTracingMaterialHitShaders.usf
+++ b/Engine/Shaders/Private/RayTracing/RayTracingMaterialHitShaders.usf
@@ -570,7 +570,7 @@ RAY_TRACING_ENTRY_CLOSEST_HIT(MaterialCHS,
 		Payload.SetFrontFace();
 	}
 
-#if TWO_SIDED_MATERIAL
+#if MATERIAL_TWOSIDED
 	Payload.SetTwoSided();
 #endif
 
diff --git a/Engine/Source/Runtime/Renderer/Private/RayTracing/RayTracingMaterialHitShaders.cpp b/Engine/Source/Runtime/Renderer/Private/RayTracing/RayTracingMaterialHitShaders.cpp
--- a/Engine/Source/Runtime/Renderer/Private/RayTracing/RayTracingMaterialHitShaders.cpp
+++ b/Engine/Source/Runtime/Renderer/Private/RayTracing/RayTracingMaterialHitShaders.cpp
@@ -154,7 +154,6 @@ public:
 		OutEnvironment.SetDefine(TEXT("USE_MATERIAL_INTERSECTION_SHADER"), UseIntersectionShader ? 1 : 0);
 		OutEnvironment.SetDefine(TEXT("USE_RAYTRACED_TEXTURE_RAYCONE_LOD"), UseRayConeTextureLod ? 1 : 0);
 		OutEnvironment.SetDefine(TEXT("SCENE_TEXTURES_DISABLED"), 1);
-		OutEnvironment.SetDefine(TEXT("TWO_SIDED_MATERIAL"), Parameters.MaterialParameters.bIsTwoSided ? 1 : 0);
 		LightMapPolicyType::ModifyCompilationEnvironment(Parameters, OutEnvironment);
 		FMeshMaterialShader::ModifyCompilationEnvironment(Parameters, OutEnvironment);
 		static const auto CVar = IConsoleManager::Get().FindTConsoleVariableDataInt(TEXT("r.VirtualTexturedLightmaps"));
