Moved Nanite Visibility earlier in the frame to reduce waiting on it.

diff --git a/Engine/Source/Runtime/Renderer/Private/DeferredShadingRenderer.cpp b/Engine/Source/Runtime/Renderer/Private/DeferredShadingRenderer.cpp
--- a/Engine/Source/Runtime/Renderer/Private/DeferredShadingRenderer.cpp
+++ b/Engine/Source/Runtime/Renderer/Private/DeferredShadingRenderer.cpp
@@ -2383,6 +2383,31 @@ void FDeferredShadingSceneRenderer::PreGatherDynamicMeshElements(FInitViewTaskDa
 	}
 
 	BeginGatherLumenLights(TaskDatas.LumenDirectLighting);
+
+	if (IsNaniteEnabled() && Views.Num() > 0)
+	{
+		TArray<FConvexVolume, TInlineAllocator<2>> NaniteCullingViews;
+
+		// For now we'll share the same visibility results across all views
+		for (int32 ViewIndex = 0; ViewIndex < Views.Num(); ViewIndex++)
+		{
+			FViewInfo& View = Views[ViewIndex];
+			NaniteCullingViews.Add(View.ViewFrustum);
+		}
+
+		FNaniteVisibility& NaniteVisibility = Scene->NaniteVisibility[ENaniteMeshPass::BasePass];
+		const FNaniteMaterialCommands& NaniteMaterials = Scene->NaniteMaterials[ENaniteMeshPass::BasePass];
+		const FNaniteRasterPipelines& NaniteRasterPipeline = Scene->NaniteRasterPipelines[ENaniteMeshPass::BasePass];
+
+		NaniteVisibility.BeginVisibilityFrame(NaniteRasterPipeline.GetBinIndexTranslator());
+
+		NaniteBasePassVisibility.Visibility = &NaniteVisibility;
+		NaniteBasePassVisibility.Query = NaniteVisibility.BeginVisibilityQuery(
+			NaniteCullingViews,
+			&NaniteRasterPipeline,
+			&NaniteMaterials
+		);
+	}
 }
 
 static TAutoConsoleVariable<float> CVarStallInitViews(
@@ -2685,34 +2710,6 @@ void FDeferredShadingSceneRenderer::Render(FRDGBuilder& GraphBuilder)
 		BeginInitViews(GraphBuilder, SceneTexturesConfig, BasePassDepthStencilAccess, InstanceCullingManager, VirtualTextureUpdater.Get(), InitViewTaskDatas);
 	}
 
-	// GetBinIndexTranslator cannot be called before UpdateAllPrimitiveSceneInfos which can change the number of raster bins
-	FNaniteScopedVisibilityFrame NaniteVisibility(
-		bNaniteEnabled,
-		Scene->NaniteVisibility[ENaniteMeshPass::BasePass],
-		Scene->NaniteRasterPipelines[ENaniteMeshPass::BasePass].GetBinIndexTranslator());
-
-	FNaniteVisibilityQuery* NaniteVisibilityQuery = nullptr;
-	if (bNaniteEnabled)
-	{
-		if (Views.Num() > 0)
-		{
-			TArray<FConvexVolume, TInlineAllocator<2>> NaniteCullingViews;
-
-			// For now we'll share the same visibility results across all views
-			for (int32 ViewIndex = 0; ViewIndex < Views.Num(); ViewIndex++)
-			{
-				FViewInfo& View = Views[ViewIndex];
-				NaniteCullingViews.Add(View.ViewFrustum);
-			}
-
-			NaniteVisibilityQuery = NaniteVisibility.Get().BeginVisibilityQuery(
-				NaniteCullingViews,
-				&Scene->NaniteRasterPipelines[ENaniteMeshPass::BasePass],
-				&Scene->NaniteMaterials[ENaniteMeshPass::BasePass]
-			);
-		}
-	}
-
 #if !UE_BUILD_SHIPPING
 	if (CVarStallInitViews.GetValueOnRenderThread() > 0.0f)
 	{
@@ -3050,9 +3047,9 @@ void FDeferredShadingSceneRenderer::Render(FRDGBuilder& GraphBuilder)
 			TRACE_CPUPROFILER_EVENT_SCOPE(InitNaniteRaster);
 
 			NaniteRasterResults.AddDefaulted(Views.Num());
-			if (NaniteVisibilityQuery != nullptr)
+			if (NaniteBasePassVisibility.Query != nullptr)
 			{
-				NaniteVisibility.Get().FinishVisibilityQuery(NaniteVisibilityQuery, NaniteRasterResults[0].VisibilityResults);
+				NaniteBasePassVisibility.Visibility->FinishVisibilityQuery(NaniteBasePassVisibility.Query, NaniteRasterResults[0].VisibilityResults);
 
 				// For now we'll share the same visibility results across all views
 				for (int32 ViewIndex = 1; ViewIndex < NaniteRasterResults.Num(); ++ViewIndex)
@@ -3220,6 +3217,11 @@ void FDeferredShadingSceneRenderer::Render(FRDGBuilder& GraphBuilder)
 				}
 			}
 		}
+
+		if (NaniteBasePassVisibility.Visibility)
+		{
+			NaniteBasePassVisibility.Visibility->FinishVisibilityFrame();
+		}
 	}
 
 	SceneTextures.SetupMode = ESceneTextureSetupMode::SceneDepth;
diff --git a/Engine/Source/Runtime/Renderer/Private/DeferredShadingRenderer.h b/Engine/Source/Runtime/Renderer/Private/DeferredShadingRenderer.h
--- a/Engine/Source/Runtime/Renderer/Private/DeferredShadingRenderer.h
+++ b/Engine/Source/Runtime/Renderer/Private/DeferredShadingRenderer.h
@@ -42,6 +42,8 @@ namespace LumenRadianceCache
 }
 class FRenderLightParameters;
 class FRayTracingScene;
+class FNaniteVisibility;
+struct FNaniteVisibilityQuery;
 
 struct FSceneWithoutWaterTextures;
 struct FRayTracingReflectionOptions;
@@ -1205,6 +1207,14 @@ private:
 	static void PrepareLumenHardwareRayTracingDirectLightingLumenMaterial(const FViewInfo& View, TArray<FRHIRayTracingShader*>& OutRayGenShaders);
 #endif // RHI_RAYTRACING
 
+
+	struct FNaniteBasePassVisibility
+	{
+		FNaniteVisibilityQuery* Query = nullptr;
+		FNaniteVisibility* Visibility = nullptr;
+
+	} NaniteBasePassVisibility;
+
 	/** Set to true if lights were injected into the light grid (this controlled by somewhat complex logic, this flag is used to cross-check). */
 	bool bAreLightsInLightGrid;
 };
diff --git a/Engine/Source/Runtime/Renderer/Private/SceneHitProxyRendering.cpp b/Engine/Source/Runtime/Renderer/Private/SceneHitProxyRendering.cpp
--- a/Engine/Source/Runtime/Renderer/Private/SceneHitProxyRendering.cpp
+++ b/Engine/Source/Runtime/Renderer/Private/SceneHitProxyRendering.cpp
@@ -749,6 +749,11 @@ void FDeferredShadingSceneRenderer::RenderHitProxies(FRDGBuilder& GraphBuilder)
 
 	::DoRenderHitProxies(GraphBuilder, this, HitProxyTexture, HitProxyDepthTexture, NaniteRasterResults, InstanceCullingManager);
 
+	if (NaniteBasePassVisibility.Visibility)
+	{
+		NaniteBasePassVisibility.Visibility->FinishVisibilityFrame();
+	}
+
 	ShaderPrint::EndViews(Views);
 
 	GEngine->GetPostRenderDelegateEx().Broadcast(GraphBuilder);
