Moved shadow filtering and view snapshot creation off the render thread.

diff --git a/Engine/Source/Runtime/Renderer/Private/DeferredShadingRenderer.cpp b/Engine/Source/Runtime/Renderer/Private/DeferredShadingRenderer.cpp
--- a/Engine/Source/Runtime/Renderer/Private/DeferredShadingRenderer.cpp
+++ b/Engine/Source/Runtime/Renderer/Private/DeferredShadingRenderer.cpp
@@ -2972,7 +2972,7 @@ void FDeferredShadingSceneRenderer::Render(FRDGBuilder& GraphBuilder)
 		}
 	}
 
-	PrepareDistanceFieldScene(GraphBuilder, ExternalAccessQueue, false);
+	PrepareDistanceFieldScene(GraphBuilder, InitViewTaskDatas.DynamicShadows, ExternalAccessQueue);
 
 	FLumenSceneFrameTemporaries LumenFrameTemporaries;
 	{
@@ -3291,7 +3291,7 @@ void FDeferredShadingSceneRenderer::Render(FRDGBuilder& GraphBuilder)
 
 	// End early occlusion queries
 
-	BeginAsyncDistanceFieldShadowProjections(GraphBuilder, SceneTextures);
+	BeginAsyncDistanceFieldShadowProjections(GraphBuilder, SceneTextures, InitViewTaskDatas.DynamicShadows);
 
 	if (bShouldRenderVolumetricCloudBase)
 	{
diff --git a/Engine/Source/Runtime/Renderer/Private/DeferredShadingRenderer.h b/Engine/Source/Runtime/Renderer/Private/DeferredShadingRenderer.h
--- a/Engine/Source/Runtime/Renderer/Private/DeferredShadingRenderer.h
+++ b/Engine/Source/Runtime/Renderer/Private/DeferredShadingRenderer.h
@@ -563,7 +563,7 @@ private:
 		FLumenSceneFrameTemporaries& FrameTemporaries,
 		FInstanceCullingManager& InstanceCullingManager,
 		FRDGExternalAccessQueue& ExternalAccessQueue,
-		const FInitViewTaskDatas& TaskDatas);
+		FInitViewTaskDatas& TaskDatas);
 
 	void BeginUpdateLumenSceneTasks(FRDGBuilder& GraphBuilder, FLumenSceneFrameTemporaries& FrameTemporaries);
 	void UpdateLumenScene(FRDGBuilder& GraphBuilder, FLumenSceneFrameTemporaries& FrameTemporaries);
diff --git a/Engine/Source/Runtime/Renderer/Private/DistanceFieldAmbientOcclusion.cpp b/Engine/Source/Runtime/Renderer/Private/DistanceFieldAmbientOcclusion.cpp
--- a/Engine/Source/Runtime/Renderer/Private/DistanceFieldAmbientOcclusion.cpp
+++ b/Engine/Source/Runtime/Renderer/Private/DistanceFieldAmbientOcclusion.cpp
@@ -711,7 +711,7 @@ bool FSceneRenderer::ShouldPrepareForDistanceFieldAO() const
 			|| (GDistanceFieldAOApplyToStaticIndirect && bAnyViewHasGIMethodSupportingDFAO && ViewFamily.EngineShowFlags.DistanceFieldAO));
 }
 
-bool FSceneRenderer::ShouldPrepareDistanceFieldScene() const
+bool FSceneRenderer::ShouldPrepareDistanceFieldScene(FDynamicShadowsTaskData* TaskData) const
 {
 	if (!ensure(Scene != nullptr))
 	{
@@ -734,7 +734,7 @@ bool FSceneRenderer::ShouldPrepareDistanceFieldScene() const
 	bool bShouldPrepareForDFInsetIndirectShadow = ShouldPrepareForDFInsetIndirectShadow();
 
 	// Prepare the distance field scene (object buffers and distance field atlas) if any feature needs it
-	return bShouldPrepareGlobalDistanceField || bShouldPrepareForAO || ShouldPrepareForDistanceFieldShadows() || bShouldPrepareForDFInsetIndirectShadow;
+	return bShouldPrepareGlobalDistanceField || bShouldPrepareForAO || ShouldPrepareForDistanceFieldShadows(TaskData) || bShouldPrepareForDFInsetIndirectShadow;
 }
 
 bool FSceneRenderer::ShouldPrepareGlobalDistanceField() const
diff --git a/Engine/Source/Runtime/Renderer/Private/DistanceFieldObjectManagement.cpp b/Engine/Source/Runtime/Renderer/Private/DistanceFieldObjectManagement.cpp
--- a/Engine/Source/Runtime/Renderer/Private/DistanceFieldObjectManagement.cpp
+++ b/Engine/Source/Runtime/Renderer/Private/DistanceFieldObjectManagement.cpp
@@ -849,7 +849,7 @@ void FSceneRenderer::UpdateGlobalHeightFieldObjectBuffers(FRDGBuilder& GraphBuil
 	}
 }
 
-void FSceneRenderer::PrepareDistanceFieldScene(FRDGBuilder& GraphBuilder, FRDGExternalAccessQueue& ExternalAccessQueue, bool bSplitDispatch)
+void FSceneRenderer::PrepareDistanceFieldScene(FRDGBuilder& GraphBuilder, FDynamicShadowsTaskData* TaskData, FRDGExternalAccessQueue& ExternalAccessQueue)
 {
 	RDG_CSV_STAT_EXCLUSIVE_SCOPE(GraphBuilder, PrepareDistanceFieldScene);
 	TRACE_CPUPROFILER_EVENT_SCOPE(FSceneRenderer::PrepareDistanceFieldScene);
@@ -859,7 +859,7 @@ void FSceneRenderer::PrepareDistanceFieldScene(FRDGBuilder& GraphBuilder, FRDGEx
 	RDG_RHI_GPU_STAT_SCOPE(GraphBuilder, DistanceFields);
 
 	const bool bShouldPrepareHeightFieldScene = ShouldPrepareHeightFieldScene();
-	const bool bShouldPrepareDistanceFieldScene = ShouldPrepareDistanceFieldScene();
+	const bool bShouldPrepareDistanceFieldScene = ShouldPrepareDistanceFieldScene(TaskData);
 
 	if (!bShouldPrepareDistanceFieldScene && !bShouldPrepareHeightFieldScene)
 	{
@@ -900,11 +900,6 @@ void FSceneRenderer::PrepareDistanceFieldScene(FRDGBuilder& GraphBuilder, FRDGEx
 
 		DistanceFieldSceneData.UpdateDistanceFieldAtlas(GraphBuilder, ExternalAccessQueue, Views[0], Scene, IsLumenEnabled(Views[0]), Views[0].ShaderMap, DistanceFieldAssetAdds, DistanceFieldAssetRemoves);
 
-		if (bSplitDispatch)
-		{
-			GraphBuilder.AddDispatchHint();
-		}
-
 		if (ShouldPrepareGlobalDistanceField())
 		{
 			for (int32 ViewIndex = 0; ViewIndex < Views.Num(); ViewIndex++)
@@ -924,10 +919,6 @@ void FSceneRenderer::PrepareDistanceFieldScene(FRDGBuilder& GraphBuilder, FRDGEx
 				UpdateGlobalDistanceFieldVolume(GraphBuilder, ExternalAccessQueue, View, Scene, OcclusionMaxDistance, IsLumenEnabled(View), View.GlobalDistanceFieldInfo);
 			}
 		}
-		if (!bSplitDispatch)
-		{
-			GraphBuilder.AddDispatchHint();
-		}
 	}
 }
 
diff --git a/Engine/Source/Runtime/Renderer/Private/DistanceFieldShadowing.cpp b/Engine/Source/Runtime/Renderer/Private/DistanceFieldShadowing.cpp
--- a/Engine/Source/Runtime/Renderer/Private/DistanceFieldShadowing.cpp
+++ b/Engine/Source/Runtime/Renderer/Private/DistanceFieldShadowing.cpp
@@ -698,35 +698,14 @@ bool SupportsHeightFieldShadows(ERHIFeatureLevel::Type FeatureLevel, EShaderPlat
 		&& DoesPlatformSupportDistanceFieldShadowing(ShaderPlatform);
 }
 
-bool FSceneRenderer::ShouldPrepareForDistanceFieldShadows() const
+bool FSceneRenderer::ShouldPrepareForDistanceFieldShadows(FDynamicShadowsTaskData* TaskData) const
 {
-	bool bSceneHasRayTracedDFShadows = false;
-
-	for (auto LightIt = Scene->Lights.CreateConstIterator(); LightIt; ++LightIt)
+	if (!ViewFamily.EngineShowFlags.DynamicShadows || !SupportsDistanceFieldShadows(Scene->GetFeatureLevel(), Scene->GetShaderPlatform()))
 	{
-		const FLightSceneInfoCompact& LightSceneInfoCompact = *LightIt;
-		const FLightSceneInfo* const LightSceneInfo = LightSceneInfoCompact.LightSceneInfo;
-
-		if (LightSceneInfo->ShouldRenderLightViewIndependent())
-		{
-			const FVisibleLightInfo& VisibleLightInfo = VisibleLightInfos[LightSceneInfo->Id];
-
-			for (int32 ShadowIndex = 0; ShadowIndex < VisibleLightInfo.AllProjectedShadows.Num(); ShadowIndex++)
-			{
-				const FProjectedShadowInfo* ProjectedShadowInfo = VisibleLightInfo.AllProjectedShadows[ShadowIndex];
-
-				if (ProjectedShadowInfo->bRayTracedDistanceField)
-				{
-					bSceneHasRayTracedDFShadows = true;
-					break;
-				}
-			}
-		}
+		return false;
 	}
 
-	return ViewFamily.EngineShowFlags.DynamicShadows 
-		&& bSceneHasRayTracedDFShadows
-		&& SupportsDistanceFieldShadows(Scene->GetFeatureLevel(), Scene->GetShaderPlatform());
+	return HasRayTracedDistanceFieldShadows(TaskData);
 }
 
 bool FSceneRenderer::ShouldPrepareHeightFieldScene() const
diff --git a/Engine/Source/Runtime/Renderer/Private/MobileShadingRenderer.cpp b/Engine/Source/Runtime/Renderer/Private/MobileShadingRenderer.cpp
--- a/Engine/Source/Runtime/Renderer/Private/MobileShadingRenderer.cpp
+++ b/Engine/Source/Runtime/Renderer/Private/MobileShadingRenderer.cpp
@@ -656,11 +656,13 @@ void FMobileSceneRenderer::InitViews(FRDGBuilder& GraphBuilder, FSceneTexturesCo
 		}
 	}
 
+	FDynamicShadowsTaskData* DynamicShadowsTaskData = nullptr;
+
 	const bool bDynamicShadows = ViewFamily.EngineShowFlags.DynamicShadows;
 	if (bDynamicShadows)
 	{
 		// Setup dynamic shadows.
-		InitDynamicShadows(GraphBuilder, InstanceCullingManager, ExternalAccessQueue);
+		DynamicShadowsTaskData = InitDynamicShadows(GraphBuilder, InstanceCullingManager, ExternalAccessQueue);
 	}
 	else
 	{
@@ -678,7 +680,7 @@ void FMobileSceneRenderer::InitViews(FRDGBuilder& GraphBuilder, FSceneTexturesCo
 
 	if (bRequiresDistanceField)
 	{
-		PrepareDistanceFieldScene(GraphBuilder, ExternalAccessQueue, false);
+		PrepareDistanceFieldScene(GraphBuilder, DynamicShadowsTaskData, ExternalAccessQueue);
 	}
 
 	ExternalAccessQueue.Submit(GraphBuilder);
diff --git a/Engine/Source/Runtime/Renderer/Private/SceneRendering.h b/Engine/Source/Runtime/Renderer/Private/SceneRendering.h
--- a/Engine/Source/Runtime/Renderer/Private/SceneRendering.h
+++ b/Engine/Source/Runtime/Renderer/Private/SceneRendering.h
@@ -2226,19 +2226,19 @@ public:
 	static void SetStereoViewport(FRHICommandList& RHICmdList, const FViewInfo& View, float ViewportScale = 1.0f);
 
 	/** Whether distance field global data structures should be prepared for features that use it. */
-	bool ShouldPrepareForDistanceFieldShadows() const;
+	bool ShouldPrepareForDistanceFieldShadows(FDynamicShadowsTaskData* TaskData) const;
 	bool ShouldPrepareForDistanceFieldAO() const;
 	bool ShouldPrepareForDFInsetIndirectShadow() const;
 	FGPUSceneDynamicContext& GetGPUSceneDynamicContext() { return GPUSceneDynamicContext; }
 
-	bool ShouldPrepareDistanceFieldScene() const;
+	bool ShouldPrepareDistanceFieldScene(FDynamicShadowsTaskData* TaskData) const;
 	bool ShouldPrepareGlobalDistanceField() const;
 	bool ShouldPrepareHeightFieldScene() const;
 
 	void UpdateGlobalDistanceFieldObjectBuffers(FRDGBuilder& GraphBuilder);
 	void UpdateGlobalHeightFieldObjectBuffers(FRDGBuilder& GraphBuilder, const TArray<uint32>& IndicesToUpdateInHeightFieldObjectBuffers);
 	void ProcessPendingHeightFieldPrimitiveAddAndRemoveOps(TArray<uint32>& IndicesToUpdateInHeightFieldObjectBuffers);
-	void PrepareDistanceFieldScene(FRDGBuilder& GraphBuilder, FRDGExternalAccessQueue& ExternalAccessQueue, bool bSplitDispatch);
+	void PrepareDistanceFieldScene(FRDGBuilder& GraphBuilder, FDynamicShadowsTaskData* TaskData, FRDGExternalAccessQueue& ExternalAccessQueue);
 
 	void DrawGPUSkinCacheVisualizationInfoText();
 
@@ -2297,10 +2297,13 @@ protected:
 
 	FDynamicShadowsTaskData* BeginInitDynamicShadows(bool bRunningEarly);
 	void FinishInitDynamicShadows(FRDGBuilder& GraphBuilder, FDynamicShadowsTaskData* TaskData, FGlobalDynamicIndexBuffer& DynamicIndexBuffer, FGlobalDynamicVertexBuffer& DynamicVertexBuffer, FGlobalDynamicReadBuffer& DynamicReadBuffer, FInstanceCullingManager& InstanceCullingManager, FRDGExternalAccessQueue& ExternalAccessQueue);
-	void InitDynamicShadows(FRDGBuilder& GraphBuilder, FGlobalDynamicIndexBuffer& DynamicIndexBuffer, FGlobalDynamicVertexBuffer& DynamicVertexBuffer, FGlobalDynamicReadBuffer& DynamicReadBuffer, FInstanceCullingManager& InstanceCullingManager, FRDGExternalAccessQueue& ExternalAccessQueue);
+	FDynamicShadowsTaskData* InitDynamicShadows(FRDGBuilder& GraphBuilder, FGlobalDynamicIndexBuffer& DynamicIndexBuffer, FGlobalDynamicVertexBuffer& DynamicVertexBuffer, FGlobalDynamicReadBuffer& DynamicReadBuffer, FInstanceCullingManager& InstanceCullingManager, FRDGExternalAccessQueue& ExternalAccessQueue);
 
-	static void WaitForPrepareDynamicShadowsTask(FDynamicShadowsTaskData* TaskData);
-	void PrepareDynamicShadows(FDynamicShadowsTaskData& TaskData);
+	static bool HasRayTracedDistanceFieldShadows(const FDynamicShadowsTaskData* TaskData);
+	static TConstArrayView<FProjectedShadowInfo*> GetProjectedDistanceFieldShadows(const FDynamicShadowsTaskData* TaskData);
+
+	void CreateDynamicShadows(FDynamicShadowsTaskData& TaskData);
+	void FilterDynamicShadows(FDynamicShadowsTaskData& TaskData);
 	friend struct FGatherShadowPrimitivesPrepareTask;
 
 	void SetupMeshPass(FViewInfo& View, FExclusiveDepthStencil::Type BasePassDepthStencilAccess, FViewCommands& ViewCommands, FInstanceCullingManager& InstanceCullingManager);
@@ -2326,7 +2329,8 @@ protected:
 
 	void BeginAsyncDistanceFieldShadowProjections(
 		FRDGBuilder& GraphBuilder,
-		const FMinimalSceneTextures& SceneTextures) const;
+		const FMinimalSceneTextures& SceneTextures,
+		const FDynamicShadowsTaskData* TaskData) const;
 
 	/** Finds a matching cached preshadow, if one exists. */
 	TRefCountPtr<FProjectedShadowInfo> GetCachedPreshadow(
@@ -2361,19 +2365,19 @@ protected:
 		FLightSceneInfo& LightSceneInfo,
 		uint32& NumCSMCachesUpdatedThisFrame);
 
-	void AllocateShadowDepthTargets(FRHICommandListImmediate& RHICmdList);
+	void AllocateShadowDepthTargets(FRHICommandListImmediate& RHICmdList, const FDynamicShadowsTaskData& TaskData);
 	
-	void AllocateAtlasedShadowDepthTargets(FRHICommandListImmediate& RHICmdList, TArray<FProjectedShadowInfo*, SceneRenderingAllocator>& Shadows, TArray<FSortedShadowMapAtlas,SceneRenderingAllocator>& OutAtlases);
+	void AllocateAtlasedShadowDepthTargets(FRHICommandListImmediate& RHICmdList, TConstArrayView<FProjectedShadowInfo*> Shadows, TArray<FSortedShadowMapAtlas,SceneRenderingAllocator>& OutAtlases);
 
-	void AllocateCachedShadowDepthTargets(FRHICommandListImmediate& RHICmdList, TArray<FProjectedShadowInfo*, SceneRenderingAllocator>& CachedShadows);
+	void AllocateCachedShadowDepthTargets(FRHICommandListImmediate& RHICmdList, TConstArrayView<FProjectedShadowInfo*> CachedShadows);
 
-	void AllocateCSMDepthTargets(FRHICommandListImmediate& RHICmdList, const TArray<FProjectedShadowInfo*, SceneRenderingAllocator>& WholeSceneDirectionalShadows, TArray<FSortedShadowMapAtlas,SceneRenderingAllocator>& OutAtlases);
+	void AllocateCSMDepthTargets(FRHICommandListImmediate& RHICmdList, TConstArrayView<FProjectedShadowInfo*> WholeSceneDirectionalShadows, TArray<FSortedShadowMapAtlas, SceneRenderingAllocator>& OutAtlases);
 
-	void AllocateOnePassPointLightDepthTargets(FRHICommandListImmediate& RHICmdList, const TArray<FProjectedShadowInfo*, SceneRenderingAllocator>& WholeScenePointShadows);
+	void AllocateOnePassPointLightDepthTargets(FRHICommandListImmediate& RHICmdList, TConstArrayView<FProjectedShadowInfo*> WholeScenePointShadows);
 
-	void AllocateTranslucentShadowDepthTargets(FRHICommandListImmediate& RHICmdList, TArray<FProjectedShadowInfo*, SceneRenderingAllocator>& TranslucentShadows);
+	void AllocateTranslucentShadowDepthTargets(FRHICommandListImmediate& RHICmdList, TConstArrayView<FProjectedShadowInfo*> TranslucentShadows);
 
-	void AllocateMobileCSMAndSpotLightShadowDepthTargets(FRHICommandListImmediate& RHICmdList, const TArray<FProjectedShadowInfo*, SceneRenderingAllocator>& MobileCSMAndSpotLightShadows);
+	void AllocateMobileCSMAndSpotLightShadowDepthTargets(FRHICommandListImmediate& RHICmdList, TConstArrayView<FProjectedShadowInfo*> MobileCSMAndSpotLightShadows);
 	/**
 	* Used by RenderLights to figure out if projected shadows need to be rendered to the attenuation buffer.
 	* Or to render a given shadowdepth map for forward rendering.
@@ -2550,9 +2554,6 @@ private:
 #endif
 	bool bShadowDepthRenderCompleted;
 
-	/** Distance field shadows to project. Used to avoid iterating through the scene lights array. */
-	TArray<FProjectedShadowInfo*, TInlineAllocator<2, SceneRenderingAllocator>> ProjectedDistanceFieldShadows;
-
 	friend class FRendererModule;
 };
 
@@ -2609,7 +2610,7 @@ public:
 
 protected:
 	/** Finds the visible dynamic shadows for each view. */
-	void InitDynamicShadows(FRDGBuilder& GraphBuilder, FInstanceCullingManager& FInstanceCullingManager, FRDGExternalAccessQueue& ExternalAccessQueue);
+	FDynamicShadowsTaskData* InitDynamicShadows(FRDGBuilder& GraphBuilder, FInstanceCullingManager& FInstanceCullingManager, FRDGExternalAccessQueue& ExternalAccessQueue);
 
 	void PrepareViewVisibilityLists();
 
diff --git a/Engine/Source/Runtime/Renderer/Private/SceneVisibility.cpp b/Engine/Source/Runtime/Renderer/Private/SceneVisibility.cpp
--- a/Engine/Source/Runtime/Renderer/Private/SceneVisibility.cpp
+++ b/Engine/Source/Runtime/Renderer/Private/SceneVisibility.cpp
@@ -5352,8 +5352,6 @@ void FDeferredShadingSceneRenderer::BeginInitViews(
 	{
 		RHICmdList.ImmediateFlush(EImmediateFlushType::DispatchToRHIThread);
 	}
-
-	WaitForPrepareDynamicShadowsTask(TaskDatas.DynamicShadows);
 }
 
 template<class T>
@@ -5432,35 +5430,13 @@ void FDeferredShadingSceneRenderer::EndInitViews(
 	FLumenSceneFrameTemporaries& FrameTemporaries,
 	FInstanceCullingManager& InstanceCullingManager,
 	FRDGExternalAccessQueue& ExternalAccessQueue,
-	const FInitViewTaskDatas& TaskDatas)
+	FInitViewTaskDatas& TaskDatas)
 {
 	SCOPED_NAMED_EVENT(FDeferredShadingSceneRenderer_InitViewsAfterPrepass, FColor::Emerald);
 	SCOPE_CYCLE_COUNTER(STAT_InitViewsPossiblyAfterPrepass);
 
 	FRHICommandListImmediate& RHICmdList = GraphBuilder.RHICmdList;
 
-	const bool bHasRayTracedOverlay = HasRayTracedOverlay(ViewFamily);
-
-	if (ViewFamily.EngineShowFlags.DynamicShadows 
-		&& !ViewFamily.EngineShowFlags.HitProxies
-		&& !bHasRayTracedOverlay)
-	{
-		// Setup dynamic shadows.
-		if (TaskDatas.DynamicShadows)
-		{
-			FinishInitDynamicShadows(GraphBuilder, TaskDatas.DynamicShadows, DynamicIndexBufferForInitShadows, DynamicVertexBufferForInitShadows, DynamicReadBufferForInitShadows, InstanceCullingManager, ExternalAccessQueue);
-		}
-		else
-		{
-			InitDynamicShadows(GraphBuilder, DynamicIndexBufferForInitShadows, DynamicVertexBufferForInitShadows, DynamicReadBufferForInitShadows, InstanceCullingManager, ExternalAccessQueue);
-		}
-
-		if (GDynamicRHI->RHIIncludeOptionalFlushes())
-		{
-			RHICmdList.ImmediateFlush(EImmediateFlushType::DispatchToRHIThread);
-		}
-	}
-
 	// If parallel ILC update is disabled, then process it in place.
 	if (ViewFamily.EngineShowFlags.HitProxies == 0
 		&& Scene->PrecomputedLightVolumes.Num() > 0
@@ -5488,6 +5464,23 @@ void FDeferredShadingSceneRenderer::EndInitViews(
 	SetupSceneReflectionCaptureBuffer(RHICmdList);
 
 	BeginUpdateLumenSceneTasks(GraphBuilder, FrameTemporaries);
+
+	const bool bHasRayTracedOverlay = HasRayTracedOverlay(ViewFamily);
+
+	if (ViewFamily.EngineShowFlags.DynamicShadows
+		&& !ViewFamily.EngineShowFlags.HitProxies
+		&& !bHasRayTracedOverlay)
+	{
+		// Setup dynamic shadows.
+		if (TaskDatas.DynamicShadows)
+		{
+			FinishInitDynamicShadows(GraphBuilder, TaskDatas.DynamicShadows, DynamicIndexBufferForInitShadows, DynamicVertexBufferForInitShadows, DynamicReadBufferForInitShadows, InstanceCullingManager, ExternalAccessQueue);
+		}
+		else
+		{
+			TaskDatas.DynamicShadows = InitDynamicShadows(GraphBuilder, DynamicIndexBufferForInitShadows, DynamicVertexBufferForInitShadows, DynamicReadBufferForInitShadows, InstanceCullingManager, ExternalAccessQueue);
+		}
+	}
 }
 
 /*------------------------------------------------------------------------------
diff --git a/Engine/Source/Runtime/Renderer/Private/ShadowRendering.cpp b/Engine/Source/Runtime/Renderer/Private/ShadowRendering.cpp
--- a/Engine/Source/Runtime/Renderer/Private/ShadowRendering.cpp
+++ b/Engine/Source/Runtime/Renderer/Private/ShadowRendering.cpp
@@ -2196,10 +2196,12 @@ void FSceneRenderer::RenderShadowProjections(
 	}
 }
 
-void FSceneRenderer::BeginAsyncDistanceFieldShadowProjections(FRDGBuilder& GraphBuilder, const FMinimalSceneTextures& SceneTextures) const
+void FSceneRenderer::BeginAsyncDistanceFieldShadowProjections(FRDGBuilder& GraphBuilder, const FMinimalSceneTextures& SceneTextures, const FDynamicShadowsTaskData* TaskData) const
 {
 	extern int32 GDFShadowAsyncCompute;
 
+	TConstArrayView<FProjectedShadowInfo*> ProjectedDistanceFieldShadows = GetProjectedDistanceFieldShadows(TaskData);
+
 	if (!!GDFShadowAsyncCompute && ViewFamily.EngineShowFlags.DynamicShadows && GetShadowQuality() > 0 && ProjectedDistanceFieldShadows.Num() > 0)
 	{
 		RDG_EVENT_SCOPE(GraphBuilder, "DistanceFieldShadows");
diff --git a/Engine/Source/Runtime/Renderer/Private/ShadowSetup.cpp b/Engine/Source/Runtime/Renderer/Private/ShadowSetup.cpp
--- a/Engine/Source/Runtime/Renderer/Private/ShadowSetup.cpp
+++ b/Engine/Source/Runtime/Renderer/Private/ShadowSetup.cpp
@@ -255,11 +255,11 @@
 	ECVF_RenderThreadSafe
 	);
 
-int32 GParallelPrepareDynamicShadows = 1;
-static FAutoConsoleVariableRef CVarParallelPrepareDynamicShadows(
-	TEXT("r.ParallelPrepareDynamicShadows"),
-	GParallelPrepareDynamicShadows,
-	TEXT("Toggles parallel prepare dynamic shadows. 0 = off; 1 = on"),
+int32 GParallelInitDynamicShadows = 1;
+static FAutoConsoleVariableRef CVarParallelInitDynamicShadows(
+	TEXT("r.ParallelInitDynamicShadows"),
+	GParallelInitDynamicShadows,
+	TEXT("Toggles parallel dynamic shadow initialization. 0 = off; 1 = on"),
 	ECVF_RenderThreadSafe
 );
 
@@ -1182,8 +1182,36 @@
 	FProjectedShadowInfo* ProjectedShadowInfo = nullptr;
 };
 
+struct FFilteredShadowArrays
+{
+	FFilteredShadowArrays() = default;
+
+	// Sort visible shadows based on their allocation needs
+	// 2d shadowmaps for this frame only that can be atlased across lights
+	TArray<FProjectedShadowInfo*, SceneRenderingAllocator> Shadows;
+	// 2d shadowmaps that will persist across frames, can't be atlased
+	TArray<FProjectedShadowInfo*, SceneRenderingAllocator> CachedSpotlightShadows;
+	TArray<FProjectedShadowInfo*, SceneRenderingAllocator> TranslucentShadows;
+	// 2d shadowmaps that persist across frames
+	TArray<FProjectedShadowInfo*, SceneRenderingAllocator> CachedPreShadows;
+	// Cubemaps, can't be atlased
+	TArray<FProjectedShadowInfo*, SceneRenderingAllocator> WholeScenePointShadows;
+
+	TArray<FProjectedShadowInfo*, SceneRenderingAllocator> WholeSceneDirectionalShadows;
+	TArray<FProjectedShadowInfo*, SceneRenderingAllocator> CachedWholeSceneDirectionalShadows;
+
+	/** Distance field shadows to project. Used to avoid iterating through the scene lights array. */
+	TArray<FProjectedShadowInfo*, TInlineAllocator<2, SceneRenderingAllocator>> ProjectedDistanceFieldShadows;
+
+	TArray<FProjectedShadowInfo*, SceneRenderingAllocator> ShadowsToSetupViews;
+
+#if !(UE_BUILD_SHIPPING || UE_BUILD_TEST)
+	TBitArray<SceneRenderingAllocator> OnePassShadowUnsupportedLights;
+#endif
+};
+
 // Common setup and working data for all GatherShadowPrimitives Tasks
-struct FDynamicShadowsTaskData : public FSceneRenderingAllocatorObject<FDynamicShadowsTaskData>
+struct FDynamicShadowsTaskData
 {
 	// Common data read from all points
 	FSceneRenderer* SceneRenderer;
@@ -1194,7 +1222,10 @@
 	const bool bStaticSceneOnly;
 	const bool bRunningEarly;
 	const bool bMultithreaded;
-	const bool bMultithreadedPrepareAndFinalize;
+	const bool bMultithreadedCreateAndFilterShadows;
 
+	// Whether any light in the scene has a ray traced distance field shadow.
+	bool bHasRayTracedDistanceFieldShadows = false;
+
 	// Generated from prepare task
 	TArray<FProjectedShadowInfo*, SceneRenderingAllocator> PreShadows;
@@ -1208,10 +1239,21 @@
 	// Written from task
 	TArray<struct FGatherShadowPrimitivesPacket*, SceneRenderingAllocator> Packets;
 	FPerShadowGatherStats GatherStats;
+	FFilteredShadowArrays ShadowArrays;
 
 	// Used by RenderThread
 	FGraphEventRef TaskEvent;
-	FGraphEventRef PrepareTaskEvent;
+	mutable FGraphEventRef CreateDynamicShadowsTaskEvent;
 
+	void WaitForCreateDynamicShadowsTask() const
+	{
+		if (CreateDynamicShadowsTaskEvent)
+		{
+			check(IsInRenderingThread());
+			CreateDynamicShadowsTaskEvent->Wait(ENamedThreads::GetRenderThread_Local());
+			CreateDynamicShadowsTaskEvent = nullptr;
+		}
+	}
+
 	FDynamicShadowsTaskData(const FDynamicShadowsTaskData&) = delete;
 
@@ -1224,15 +1266,35 @@
 		, bStaticSceneOnly(AreAnyViewsStaticSceneOnly(Views))
 		, bRunningEarly(bInRunningEarly)
 		, bMultithreaded((FApp::ShouldUseThreadingForPerformance() || FForkProcessHelper::IsForkedMultithreadInstance()) && CVarParallelGatherShadowPrimitives.GetValueOnRenderThread() > 0)
-		, bMultithreadedPrepareAndFinalize(bRunningEarly && bMultithreaded && GRHISupportsAsyncGetRenderQueryResult && GParallelPrepareDynamicShadows > 0)
+		, bMultithreadedCreateAndFilterShadows(bRunningEarly&& bMultithreaded&& GRHISupportsAsyncGetRenderQueryResult&& GParallelInitDynamicShadows > 0)
 	{
-		if (bMultithreadedPrepareAndFinalize)
+		if (bMultithreadedCreateAndFilterShadows)
 		{
-			PrepareTaskEvent = FGraphEvent::CreateGraphEvent();
+			CreateDynamicShadowsTaskEvent = FGraphEvent::CreateGraphEvent();
 		}
 	}
 };
 
+TConstArrayView<FProjectedShadowInfo*> FSceneRenderer::GetProjectedDistanceFieldShadows(const FDynamicShadowsTaskData* TaskData)
+{
+	if (TaskData)
+	{
+		TaskData->WaitForCreateDynamicShadowsTask();
+		return TaskData->ShadowArrays.ProjectedDistanceFieldShadows;
+	}
+	return TConstArrayView<FProjectedShadowInfo*>();
+}
+
+bool FSceneRenderer::HasRayTracedDistanceFieldShadows(const FDynamicShadowsTaskData* TaskData)
+{
+	if (TaskData)
+	{
+		TaskData->WaitForCreateDynamicShadowsTask();
+		return TaskData->bHasRayTracedDistanceFieldShadows;
+	}
+	return false;
+}
+
 struct FAddSubjectPrimitiveOverflowedIndices
 {
 	TArray<uint16> MDCIndices;
@@ -4544,10 +4606,10 @@
 		SCOPED_NAMED_EVENT_TEXT("FGatherShadowPrimitivesPrepareTask", FColor::Green);
 		FOptionalTaskTagScope Scope(ETaskTag::EParallelRenderingThread);
 
-		if (TaskData.bMultithreadedPrepareAndFinalize)
+		if (TaskData.bMultithreadedCreateAndFilterShadows)
 		{
-			TaskData.SceneRenderer->PrepareDynamicShadows(TaskData);
-			TaskData.PrepareTaskEvent->DispatchSubsequents();
+			TaskData.SceneRenderer->CreateDynamicShadows(TaskData);
+			TaskData.CreateDynamicShadowsTaskEvent->DispatchSubsequents();
 		}
 
 		AnyThreadTask();
@@ -4589,7 +4651,7 @@
 			}
 		}
 
-		if (TaskData.bMultithreadedPrepareAndFinalize && Prereqs.Num() > 0)
+		if (TaskData.bMultithreadedCreateAndFilterShadows && Prereqs.Num() > 0)
 		{
 			FGraphEventRef FinalizeTask = FFunctionGraphTask::CreateAndDispatchWhenReady([TaskData = &TaskData]()
 				{
@@ -4601,6 +4663,8 @@
 						Packet->AnyThreadFinalize(*TaskData);
 						delete Packet;
 					}
+					TaskData->Packets.Empty();
+					TaskData->SceneRenderer->FilterDynamicShadows(*TaskData);
 				},
 				TStatId(),
 				&Prereqs,
@@ -4692,9 +4756,9 @@
 
 	check(TaskData);
 
-	if (!TaskData->bMultithreadedPrepareAndFinalize)
+	if (!TaskData->bMultithreadedCreateAndFilterShadows)
 	{
-		PrepareDynamicShadows(*TaskData);
+		CreateDynamicShadows(*TaskData);
 	}
 
 	if (!TaskData->bMultithreaded || !TaskData->bRunningEarly)
@@ -4739,7 +4803,7 @@
 		FTaskGraphInterface::Get().WaitUntilTaskCompletes(TaskData->TaskEvent, ENamedThreads::GetRenderThread_Local());
 	}
 
-	if (!TaskData->bMultithreadedPrepareAndFinalize)
+	if (!TaskData->bMultithreadedCreateAndFilterShadows)
 	{
 		SCOPED_NAMED_EVENT_TEXT("FGatherShadowPrimitivesTask", FColor::Green);
 		for (FGatherShadowPrimitivesPacket* Packet : TaskData->Packets)
@@ -4747,6 +4811,8 @@
 			Packet->AnyThreadFinalize(*TaskData);
 			delete Packet;
 		}
+		TaskData->Packets.Empty();
+		FilterDynamicShadows(*TaskData);
 	}
 
 	for (FPrimitiveSceneInfo* PrimitiveSceneInfo : TaskData->NeedsUniformBufferUpdatePrimitives)
@@ -4764,7 +4830,24 @@
 		DrawDebugShadowFrustum(*Op.View, *Op.ProjectedShadowInfo);
 	}
 
-	delete TaskData;
+#if !(UE_BUILD_SHIPPING || UE_BUILD_TEST)
+	if (!TaskData->ShadowArrays.OnePassShadowUnsupportedLights.IsEmpty())
+	{
+		OnGetOnScreenMessages.AddLambda([Scene = Scene, OnePassShadowUnsupportedLights = MoveTemp(TaskData->ShadowArrays.OnePassShadowUnsupportedLights)](FScreenMessageWriter& ScreenMessageWriter)->void
+		{
+			static const FText Message = NSLOCTEXT("Renderer", "PointLightVsLayer", "RUNTIME DOES NOT SUPPORT WHOLE SCENE POINT LIGHT SHADOWS (Missing Vertexshader Layer support): ");
+			ScreenMessageWriter.DrawLine(Message);
+			for (TConstSetBitIterator<SceneRenderingAllocator> It(OnePassShadowUnsupportedLights); It; ++It)
+			{
+				int32 LightId = It.GetIndex();
+				if (Scene->Lights.IsValidIndex(LightId))
+				{
+					ScreenMessageWriter.DrawLine(FText::FromString(Scene->Lights[LightId].LightSceneInfo->Proxy->GetOwnerNameOrLabel()), 35);
+				}
+			}
+		});
+	}
+#endif
 }
 
 static bool NeedsUnatlasedCSMDepthsWorkaround(ERHIFeatureLevel::Type FeatureLevel)
@@ -5031,260 +5114,16 @@
 	}
 }
 
-void FSceneRenderer::AllocateShadowDepthTargets(FRHICommandListImmediate& RHICmdList)
+void FSceneRenderer::AllocateShadowDepthTargets(FRHICommandListImmediate& RHICmdList, const FDynamicShadowsTaskData& TaskData)
 {
-#if !(UE_BUILD_SHIPPING || UE_BUILD_TEST)
-	TBitArray<SceneRenderingAllocator> OnePassShadowUnsupportedLights;
-	auto MarkUnsupportedOnePassShadow = [&OnePassShadowUnsupportedLights](int32 LightId)
-	{
-		if (OnePassShadowUnsupportedLights.Num() <= LightId)
-		{
-			OnePassShadowUnsupportedLights.Add(false, 1 + LightId - OnePassShadowUnsupportedLights.Num());
-		}
-		OnePassShadowUnsupportedLights[LightId] = true;
-	};
-#else 
-	auto MarkUnsupportedOnePassShadow = [](int32 LightId) {};
-#endif
+	SCOPED_NAMED_EVENT_TEXT("FSceneRenderer::AllocateShadowDepthTargets", FColor::Magenta);
 
-	// Sort visible shadows based on their allocation needs
-	// 2d shadowmaps for this frame only that can be atlased across lights
-	TArray<FProjectedShadowInfo*, SceneRenderingAllocator> Shadows;
-	// 2d shadowmaps that will persist across frames, can't be atlased
-	TArray<FProjectedShadowInfo*, SceneRenderingAllocator> CachedSpotlightShadows;
-	TArray<FProjectedShadowInfo*, SceneRenderingAllocator> TranslucentShadows;
-	// 2d shadowmaps that persist across frames
-	TArray<FProjectedShadowInfo*, SceneRenderingAllocator> CachedPreShadows;
-	TArray<FProjectedShadowInfo*, SceneRenderingAllocator> RSMShadows;
-	// Cubemaps, can't be atlased
-	TArray<FProjectedShadowInfo*, SceneRenderingAllocator> WholeScenePointShadows;
-	TArray<FProjectedShadowInfo*, SceneRenderingAllocator> MobileWholeSceneDirectionalShadows;
-	TArray<FProjectedShadowInfo*, SceneRenderingAllocator> MobileDynamicSpotlightShadows;
-
-	TArray<FProjectedShadowInfo*, SceneRenderingAllocator> CachedWholeSceneDirectionalShadows;
-
 	const bool bMobile = FeatureLevel < ERHIFeatureLevel::SM5;
 
-	for (auto LightIt = Scene->Lights.CreateConstIterator(); LightIt; ++LightIt)
+	const FFilteredShadowArrays& ShadowArrays = TaskData.ShadowArrays;
+
+	if (ShadowArrays.CachedPreShadows.Num() > 0)
 	{
-		const FLightSceneInfoCompact& LightSceneInfoCompact = *LightIt;
-		FLightSceneInfo* LightSceneInfo = LightSceneInfoCompact.LightSceneInfo;
-		FVisibleLightInfo& VisibleLightInfo = VisibleLightInfos[LightSceneInfo->Id];
-
-		// All cascades for a light need to be in the same texture
-		TArray<FProjectedShadowInfo*, SceneRenderingAllocator> WholeSceneDirectionalShadows;
-		TArray<FProjectedShadowInfo*, SceneRenderingAllocator> WholeSceneCompleteDirectionalShadows;
-
-		for (int32 ShadowIndex = 0; ShadowIndex < VisibleLightInfo.AllProjectedShadows.Num(); ShadowIndex++)
-		{
-			FProjectedShadowInfo* ProjectedShadowInfo = VisibleLightInfo.AllProjectedShadows[ShadowIndex];
-
-			// Check that the shadow is visible in at least one view before rendering it.
-			bool bShadowIsVisible = false;
-
-			for (int32 ViewIndex = 0; ViewIndex < Views.Num(); ViewIndex++)
-			{
-				FViewInfo& View = Views[ViewIndex];
-
-				if (ProjectedShadowInfo->DependentView && ProjectedShadowInfo->DependentView != &View)
-				{
-					continue;
-				}
-
-				const FVisibleLightViewInfo& VisibleLightViewInfo = View.VisibleLightInfos[LightSceneInfo->Id];
-				const FPrimitiveViewRelevance ViewRelevance = VisibleLightViewInfo.ProjectedShadowViewRelevanceMap[ShadowIndex];
-				const bool bHasViewRelevance = (ProjectedShadowInfo->bTranslucentShadow && ViewRelevance.HasTranslucency()) 
-					|| (!ProjectedShadowInfo->bTranslucentShadow && ViewRelevance.bOpaque);
-
-				bShadowIsVisible |= bHasViewRelevance && VisibleLightViewInfo.ProjectedShadowVisibilityMap[ShadowIndex];
-			}
-
-			// Fully cached VSM should not be skipped, since they need to be queued for projection.
-			if (!ProjectedShadowInfo->bShouldRenderVSM)
-			{
-			}
-			// Skip the shadow rendering if there is no primitive in the cached shadow map and the current shadow map either
-			else if ((ProjectedShadowInfo->CacheMode == SDCM_StaticPrimitivesOnly || (ProjectedShadowInfo->CacheMode == SDCM_MovablePrimitivesOnly && !ProjectedShadowInfo->IsWholeSceneDirectionalShadow())) && !ProjectedShadowInfo->HasSubjectPrims())
-			{
-				const FCachedShadowMapData& CachedShadowMapData = Scene->GetCachedShadowMapDataRef(ProjectedShadowInfo->GetLightSceneInfo().Id, FMath::Max(ProjectedShadowInfo->CascadeSettings.ShadowSplitIndex, 0));
-
-				// A shadowmap for movable primitives when there are no movable primitives would normally read directly from the cached shadowmap
-				// However if the cached shadowmap also had no primitives then we need to skip rendering the shadow entirely
-				if (!CachedShadowMapData.bCachedShadowMapHasPrimitives)
-				{
-					bShadowIsVisible = false;
-				}
-			}
-			// If uncached and no primitives, skip allocations etc
-			// Note: cached PreShadows, for some reason, has CacheMode == SDCM_Uncached so need to skip them here.
-			// Ray traced shadows use the GPU managed distance field object buffers, no CPU culling should be used
-			else if (
-				(ProjectedShadowInfo->CacheMode == SDCM_Uncached && !ProjectedShadowInfo->bPreShadow) 
-				&& (!ProjectedShadowInfo->HasSubjectPrims() && !ProjectedShadowInfo->bRayTracedDistanceField) 
-				&& !ProjectedShadowInfo->IsWholeSceneDirectionalShadow())
-			{
-				bShadowIsVisible = false;
-			}
-
-			if ((IsForwardShadingEnabled(ShaderPlatform) || (IsMobilePlatform(ShaderPlatform) && MobileUsesShadowMaskTexture(ShaderPlatform)))
-				&& ProjectedShadowInfo->GetLightSceneInfo().GetDynamicShadowMapChannel() == -1)
-			{
-				// With forward shading, dynamic shadows are projected into channels of the light attenuation texture based on their assigned DynamicShadowMapChannel
-				bShadowIsVisible = false;
-			}
-
-			// Skip one-pass point light shadows if rendering to them is not suported, unless VSM or DF
-			if (bShadowIsVisible
-				&& !ProjectedShadowInfo->HasVirtualShadowMap()
-				&& !ProjectedShadowInfo->bRayTracedDistanceField
-				&& ProjectedShadowInfo->bOnePassPointLightShadow 
-				&& !DoesRuntimeSupportOnePassPointLightShadows(GShaderPlatformForFeatureLevel[FeatureLevel]))
-			{ 
-				// Also note this to produce a log message & on-screen warning
-				MarkUnsupportedOnePassShadow(ProjectedShadowInfo->GetLightSceneInfo().Id);
-				bShadowIsVisible = false;
-			}
-
-			if (bShadowIsVisible)
-			{
-				// Visible shadow stats
-				if (ProjectedShadowInfo->bWholeSceneShadow)
-				{
-					INC_DWORD_STAT(STAT_WholeSceneShadows);
-
-					if (ProjectedShadowInfo->CacheMode == SDCM_MovablePrimitivesOnly)
-					{
-						INC_DWORD_STAT(STAT_CachedWholeSceneShadows);
-					}
-				}
-				else if (ProjectedShadowInfo->bPreShadow)
-				{
-					INC_DWORD_STAT(STAT_PreShadows);
-				}
-				else
-				{
-					INC_DWORD_STAT(STAT_PerObjectShadows);
-				}
-
-				bool bNeedsProjection = ProjectedShadowInfo->CacheMode != SDCM_StaticPrimitivesOnly
-					//// Filter out everything but PerObjectOpaqueShadows & Distance field shadows for ES31
-					&& (!bMobile || ProjectedShadowInfo->bPerObjectOpaqueShadow || ProjectedShadowInfo->bRayTracedDistanceField || ProjectedShadowInfo->bWholeSceneShadow);
-
-				if (bNeedsProjection)
-				{
-					if (ProjectedShadowInfo->bCapsuleShadow)
-					{
-						VisibleLightInfo.CapsuleShadowsToProject.Add(ProjectedShadowInfo);
-					}
-					else
-					{
-						if (ProjectedShadowInfo->bRayTracedDistanceField)
-						{
-							ProjectedDistanceFieldShadows.Add(ProjectedShadowInfo);
-						}
-						VisibleLightInfo.ShadowsToProject.Add(ProjectedShadowInfo);
-					}
-				}
-
-				const bool bNeedsShadowmapSetup = !ProjectedShadowInfo->bCapsuleShadow && !ProjectedShadowInfo->bRayTracedDistanceField;
-
-				if (bNeedsShadowmapSetup)
-				{
-					if (ProjectedShadowInfo->HasVirtualShadowMap() || ProjectedShadowInfo->VirtualShadowMapClipmap.IsValid())
-					{
-						ProjectedShadowInfo->SetupShadowDepthView(this);
-						SortedShadowsForShadowDepthPass.VirtualShadowMapShadows.Add(ProjectedShadowInfo);
-					}
-					else if (ProjectedShadowInfo->bPreShadow && ProjectedShadowInfo->bAllocatedInPreshadowCache)
-					{
-						CachedPreShadows.Add(ProjectedShadowInfo);
-					}
-					else if (ProjectedShadowInfo->bDirectionalLight && ProjectedShadowInfo->bWholeSceneShadow)
-					{
-						if (ProjectedShadowInfo->CacheMode == SDCM_StaticPrimitivesOnly)
-						{
-							CachedWholeSceneDirectionalShadows.Add(ProjectedShadowInfo);
-						}
-						else
-						{
-							WholeSceneDirectionalShadows.Add(ProjectedShadowInfo);
-						}
-					}
-					else if (ProjectedShadowInfo->bOnePassPointLightShadow)
-					{
-						WholeScenePointShadows.Add(ProjectedShadowInfo);
-					}
-					else if (ProjectedShadowInfo->bTranslucentShadow)
-					{
-						TranslucentShadows.Add(ProjectedShadowInfo);
-					}
-					else if (ProjectedShadowInfo->CacheMode == SDCM_StaticPrimitivesOnly)
-					{
-						check(ProjectedShadowInfo->bWholeSceneShadow);
-						CachedSpotlightShadows.Add(ProjectedShadowInfo);
-					}
-					else if (bMobile && ProjectedShadowInfo->bWholeSceneShadow)
-					{
-						MobileDynamicSpotlightShadows.Add(ProjectedShadowInfo);
-					}
-					else
-					{
-						Shadows.Add(ProjectedShadowInfo);
-					}
-				}
-			}
-		}
-
-		// Sort cascades, this is needed for blending between cascades to work, and also to enable fetching the farthest cascade as index [0]
-		VisibleLightInfo.ShadowsToProject.Sort(FCompareFProjectedShadowInfoBySplitIndex());
-
-		if (!bMobile)
-		{
-			AllocateCachedShadowDepthTargets(RHICmdList, CachedWholeSceneDirectionalShadows);
-
-			// Only allocate CSM targets if at least one of the SMs in the CSM has any primitives, but then always allocate all of them as some algorithms depend on this (notably the forward shading structures).
-			// Don't perform this if VSM are not enabled.
-			bool bAnyHasSubjectPrims = !VirtualShadowMapArray.IsEnabled();
-			for (const FProjectedShadowInfo* ProjectedShadowInfo : WholeSceneDirectionalShadows)
-			{
-				bAnyHasSubjectPrims = bAnyHasSubjectPrims || ProjectedShadowInfo->HasSubjectPrims();
-			}
-			if (bAnyHasSubjectPrims)
-			{
-				AllocateCSMDepthTargets(RHICmdList, WholeSceneDirectionalShadows, SortedShadowsForShadowDepthPass.ShadowMapAtlases);
-			}
-			CachedWholeSceneDirectionalShadows.Empty();
-		}
-		else
-		{
-			//Only one directional light could cast csm on mobile, so we could delay allocation for it and see if we could combine any spotlight shadow with it.
-			if (WholeSceneDirectionalShadows.Num() > 0)
-			{
-				MobileWholeSceneDirectionalShadows.Append(WholeSceneDirectionalShadows);
-			}
-		}
-	}
-
-	if (bMobile)
-	{
-		// AllocateMobileCSMAndSpotLightShadowDepthTargets would only allocate a single large render target for all shadows, so if the requirement exceeds the MaxTextureSize, the rest of the shadows will not get space for rendering
-		// So we sort spotlight shadows and append them at the last to make sure csm will get space in any case.
-		MobileDynamicSpotlightShadows.Sort(FCompareFProjectedShadowInfoByResolution());
-
-		//Limit the number of spotlights shadow for performance reason
-		static const auto MobileMaxVisibleMovableSpotLightShadowsCVar = IConsoleManager::Get().FindTConsoleVariableDataInt(TEXT("r.Mobile.MaxVisibleMovableSpotLightShadows"));
-		if (MobileMaxVisibleMovableSpotLightShadowsCVar)
-		{
-			int32 MobileMaxVisibleMovableSpotLightShadows = MobileMaxVisibleMovableSpotLightShadowsCVar->GetValueOnRenderThread();
-			MobileDynamicSpotlightShadows.RemoveAt(MobileMaxVisibleMovableSpotLightShadows, FMath::Max(MobileDynamicSpotlightShadows.Num() - MobileMaxVisibleMovableSpotLightShadows, 0), false);
-		}
-
-		MobileWholeSceneDirectionalShadows.Append(MobileDynamicSpotlightShadows);
-		MobileDynamicSpotlightShadows.Empty();
-	}
-
-	if (CachedPreShadows.Num() > 0)
-	{
 		if (!Scene->PreShadowCacheDepthZ)
 		{
 			FPooledRenderTargetDesc Desc(FPooledRenderTargetDesc::Create2DDesc(GetPreShadowCacheTextureResolution(FeatureLevel), PF_ShadowDepth, FClearValueBinding::None, TexCreate_None, TexCreate_DepthStencilTargetable | TexCreate_ShaderResource, false));
@@ -5292,33 +5131,31 @@
 		}
 
 		SortedShadowsForShadowDepthPass.PreshadowCache.RenderTargets.DepthTarget = Scene->PreShadowCacheDepthZ;
-
-		for (int32 ShadowIndex = 0; ShadowIndex < CachedPreShadows.Num(); ShadowIndex++)
+		for (int32 ShadowIndex = 0; ShadowIndex < ShadowArrays.CachedPreShadows.Num(); ShadowIndex++)
 		{
-			FProjectedShadowInfo* ProjectedShadowInfo = CachedPreShadows[ShadowIndex];
+			FProjectedShadowInfo* ProjectedShadowInfo = ShadowArrays.CachedPreShadows[ShadowIndex];
 			ProjectedShadowInfo->RenderTargets.DepthTarget = Scene->PreShadowCacheDepthZ.GetReference();
-
-			// Note: adding preshadows whose depths are cached so that GatherDynamicMeshElements
-			// will still happen, which is necessary for preshadow receiver stenciling
-			ProjectedShadowInfo->SetupShadowDepthView(this);
-			SortedShadowsForShadowDepthPass.PreshadowCache.Shadows.Add(ProjectedShadowInfo);
 		}
 	}
 
-	AllocateOnePassPointLightDepthTargets(RHICmdList, WholeScenePointShadows);
-	AllocateCachedShadowDepthTargets(RHICmdList, CachedSpotlightShadows);
+	AllocateOnePassPointLightDepthTargets(RHICmdList, ShadowArrays.WholeScenePointShadows);
+	AllocateCachedShadowDepthTargets(RHICmdList, ShadowArrays.CachedSpotlightShadows);
+	AllocateCachedShadowDepthTargets(RHICmdList, ShadowArrays.CachedWholeSceneDirectionalShadows);
 	if (bMobile)
 	{
-		AllocateCachedShadowDepthTargets(RHICmdList, CachedWholeSceneDirectionalShadows);
+		AllocateMobileCSMAndSpotLightShadowDepthTargets(RHICmdList, ShadowArrays.WholeSceneDirectionalShadows);
 	}
-	AllocateAtlasedShadowDepthTargets(RHICmdList, Shadows, SortedShadowsForShadowDepthPass.ShadowMapAtlases);
-	AllocateTranslucentShadowDepthTargets(RHICmdList, TranslucentShadows);
-	AllocateMobileCSMAndSpotLightShadowDepthTargets(RHICmdList, MobileWholeSceneDirectionalShadows);
+	else
+	{
+		AllocateCSMDepthTargets(RHICmdList, ShadowArrays.WholeSceneDirectionalShadows, SortedShadowsForShadowDepthPass.ShadowMapAtlases);
+	}
+	AllocateAtlasedShadowDepthTargets(RHICmdList, ShadowArrays.Shadows, SortedShadowsForShadowDepthPass.ShadowMapAtlases);
+	AllocateTranslucentShadowDepthTargets(RHICmdList, ShadowArrays.TranslucentShadows);
 
 	// Update translucent shadow map uniform buffers.
-	for (int32 TranslucentShadowIndex = 0; TranslucentShadowIndex < TranslucentShadows.Num(); ++TranslucentShadowIndex)
+	for (int32 TranslucentShadowIndex = 0; TranslucentShadowIndex < ShadowArrays.TranslucentShadows.Num(); ++TranslucentShadowIndex)
 	{
-		FProjectedShadowInfo* ShadowInfo = TranslucentShadows[TranslucentShadowIndex];
+		FProjectedShadowInfo* ShadowInfo = ShadowArrays.TranslucentShadows[TranslucentShadowIndex];
 		const int32 PrimitiveIndex = ShadowInfo->GetParentSceneInfo()->GetIndex();
 
 		for (int32 ViewIndex = 0; ViewIndex < Views.Num(); ++ViewIndex)
@@ -5330,43 +5167,11 @@
 			{
 				FTranslucentSelfShadowUniformParameters Parameters;
 				SetupTranslucentSelfShadowUniformParameters(ShadowInfo, Parameters);
-				RHIUpdateUniformBuffer(*UniformBufferPtr, &Parameters);
+				RHICmdList.UpdateUniformBuffer(*UniformBufferPtr, &Parameters);
 			}
 		}
 	}
 
-	// Remove cache entries that haven't been used in a while
-	for (TMap<int32, TArray<FCachedShadowMapData>>::TIterator CachedShadowMapIt(Scene->CachedShadowMaps); CachedShadowMapIt; ++CachedShadowMapIt)
-	{
-		TArray<FCachedShadowMapData>& ShadowMapDatas = CachedShadowMapIt.Value();
-
-		for (auto& ShadowMapData : ShadowMapDatas)
-		{
-			if (ShadowMapData.ShadowMap.IsValid() && ViewFamily.Time.GetRealTimeSeconds() - ShadowMapData.LastUsedTime > 2.0f)
-			{
-				ShadowMapData.InvalidateCachedShadow();
-			}
-		}
-	}
-
-#if !(UE_BUILD_SHIPPING || UE_BUILD_TEST)
-	if (!OnePassShadowUnsupportedLights.IsEmpty())
-	{
-		OnGetOnScreenMessages.AddLambda([Scene = Scene, OnePassShadowUnsupportedLights = MoveTemp(OnePassShadowUnsupportedLights)](FScreenMessageWriter& ScreenMessageWriter)->void
-		{
-			static const FText Message = NSLOCTEXT("Renderer", "PointLightVsLayer", "RUNTIME DOES NOT SUPPORT WHOLE SCENE POINT LIGHT SHADOWS (Missing Vertexshader Layer support): ");
-			ScreenMessageWriter.DrawLine(Message);
-			for (TConstSetBitIterator<SceneRenderingAllocator> It(OnePassShadowUnsupportedLights); It; ++It)
-			{
-				int32 LightId = It.GetIndex();
-				if (Scene->Lights.IsValidIndex(LightId))
-				{
-					ScreenMessageWriter.DrawLine(FText::FromString(Scene->Lights[LightId].LightSceneInfo->Proxy->GetOwnerNameOrLabel()), 35);
-				}
-			}
-		});
-	}
-#endif
 	SET_MEMORY_STAT(STAT_CachedShadowmapMemory, Scene->GetCachedWholeSceneShadowMapsSize());
 	SET_MEMORY_STAT(STAT_ShadowmapAtlasMemory, SortedShadowsForShadowDepthPass.ComputeMemorySize());
 }
@@ -5387,7 +5192,7 @@
 
 void FSceneRenderer::AllocateAtlasedShadowDepthTargets(
 	FRHICommandListImmediate& RHICmdList,
-	TArray<FProjectedShadowInfo*, SceneRenderingAllocator>& Shadows,
+	TConstArrayView<FProjectedShadowInfo*> Shadows,
 	TArray<FSortedShadowMapAtlas,SceneRenderingAllocator>& OutAtlases)
 {
 	if (Shadows.Num() == 0)
@@ -5470,7 +5275,6 @@
 			if (ProjectedShadowInfo->bAllocated)
 			{
 				ProjectedShadowInfo->RenderTargets.CopyReferencesFromRenderTargets(ShadowMapAtlas.RenderTargets);
-				ProjectedShadowInfo->SetupShadowDepthView(this);
 				ShadowMapAtlas.Shadows.Add(ProjectedShadowInfo);
 			}
 		}
@@ -5509,7 +5313,7 @@
 	}
 };
 
-void FSceneRenderer::AllocateCachedShadowDepthTargets(FRHICommandListImmediate& RHICmdList, TArray<FProjectedShadowInfo*, SceneRenderingAllocator>& CachedShadows)
+void FSceneRenderer::AllocateCachedShadowDepthTargets(FRHICommandListImmediate& RHICmdList, TConstArrayView<FProjectedShadowInfo*> CachedShadows)
 {
 	for (int32 ShadowIndex = 0; ShadowIndex < CachedShadows.Num(); ShadowIndex++)
 	{
@@ -5534,14 +5338,13 @@
 		ProjectedShadowInfo->bAllocated = true;
 		ProjectedShadowInfo->RenderTargets.CopyReferencesFromRenderTargets(ShadowMap.RenderTargets);
 
-		ProjectedShadowInfo->SetupShadowDepthView(this);
 		ShadowMap.Shadows.Add(ProjectedShadowInfo);
 	}
 }
 
 void FSceneRenderer::AllocateCSMDepthTargets(
 	FRHICommandListImmediate& RHICmdList,
-	const TArray<FProjectedShadowInfo*,	SceneRenderingAllocator>& WholeSceneDirectionalShadows,
+	TConstArrayView<FProjectedShadowInfo*> WholeSceneDirectionalShadows,
 	TArray<FSortedShadowMapAtlas, SceneRenderingAllocator>& OutAtlases
 	)
 {
@@ -5594,14 +5397,13 @@
 			{
 				check(ProjectedShadowInfo->bAllocated);
 				ProjectedShadowInfo->RenderTargets.CopyReferencesFromRenderTargets(ShadowMapAtlas.RenderTargets);
-				ProjectedShadowInfo->SetupShadowDepthView(this);
 				ShadowMapAtlas.Shadows.Add(ProjectedShadowInfo);
 			}
 		}
 	}
 }
 
-void FSceneRenderer::AllocateOnePassPointLightDepthTargets(FRHICommandListImmediate& RHICmdList, const TArray<FProjectedShadowInfo*, SceneRenderingAllocator>& WholeScenePointShadows)
+void FSceneRenderer::AllocateOnePassPointLightDepthTargets(FRHICommandListImmediate& RHICmdList, TConstArrayView<FProjectedShadowInfo*> WholeScenePointShadows)
 {
 	if (FeatureLevel >= ERHIFeatureLevel::SM5)
 	{
@@ -5637,8 +5439,7 @@
 				ProjectedShadowInfo->X = ProjectedShadowInfo->Y = 0;
 				ProjectedShadowInfo->bAllocated = true;
 				ProjectedShadowInfo->RenderTargets.CopyReferencesFromRenderTargets(ShadowMapCubemap.RenderTargets);
-				
-				ProjectedShadowInfo->SetupShadowDepthView(this);
+
 				ShadowMapCubemap.Shadows.Add(ProjectedShadowInfo);
 			}
 		}
@@ -5660,7 +5461,7 @@
 	return (TCHAR*)TEXT("InvalidName");
 }
 
-void FSceneRenderer::AllocateTranslucentShadowDepthTargets(FRHICommandListImmediate& RHICmdList, TArray<FProjectedShadowInfo*, SceneRenderingAllocator>& TranslucentShadows)
+void FSceneRenderer::AllocateTranslucentShadowDepthTargets(FRHICommandListImmediate& RHICmdList, TConstArrayView<FProjectedShadowInfo*> TranslucentShadows)
 {
 	if (TranslucentShadows.Num() > 0 && FeatureLevel >= ERHIFeatureLevel::SM5)
 	{
@@ -5671,9 +5472,6 @@
 
 		FTextureLayout CurrentShadowLayout(1, 1, TranslucentShadowBufferResolution.X, TranslucentShadowBufferResolution.Y, false, ETextureLayoutAspectRatio::None, false);
 
-		// Sort the projected shadows by resolution.
-		TranslucentShadows.Sort(FCompareFProjectedShadowInfoByResolution());
-
 		for (int32 ShadowIndex = 0; ShadowIndex < TranslucentShadows.Num(); ShadowIndex++)
 		{
 			FProjectedShadowInfo* ProjectedShadowInfo = TranslucentShadows[ShadowIndex];
@@ -5724,14 +5522,15 @@
 			}
 
 			ProjectedShadowInfo->RenderTargets.CopyReferencesFromRenderTargets(ShadowMapAtlas.RenderTargets);
-			ProjectedShadowInfo->SetupShadowDepthView(this);
 			ShadowMapAtlas.Shadows.Add(ProjectedShadowInfo);
 		}
 	}
 }
 
-void FSceneRenderer::PrepareDynamicShadows(FDynamicShadowsTaskData& TaskData)
+void FSceneRenderer::CreateDynamicShadows(FDynamicShadowsTaskData& TaskData)
 {
+	TRACE_CPUPROFILER_EVENT_SCOPE(CreateDynamicShadows);
+
 	const bool bMobile = FeatureLevel < ERHIFeatureLevel::SM5;
 	const bool bHairStrands = HairStrands::HasHairInstanceInScene(*Scene);
 
@@ -5902,11 +5701,39 @@
 								}
 							}
 						}
+
+						if (!TaskData.bHasRayTracedDistanceFieldShadows)
+						{
+							for (int32 ShadowIndex = 0; ShadowIndex < VisibleLightInfo.AllProjectedShadows.Num(); ShadowIndex++)
+							{
+								const FProjectedShadowInfo* ProjectedShadowInfo = VisibleLightInfo.AllProjectedShadows[ShadowIndex];
+
+								if (ProjectedShadowInfo->bRayTracedDistanceField)
+								{
+									TaskData.bHasRayTracedDistanceFieldShadows = true;
+									break;
+								}
+							}
+						}
 					}
 				}
 			}
 		}
 
+		// Remove cache entries that haven't been used in a while
+		for (TMap<int32, TArray<FCachedShadowMapData>>::TIterator CachedShadowMapIt(Scene->CachedShadowMaps); CachedShadowMapIt; ++CachedShadowMapIt)
+		{
+			TArray<FCachedShadowMapData>& ShadowMapDatas = CachedShadowMapIt.Value();
+
+			for (auto& ShadowMapData : ShadowMapDatas)
+			{
+				if (ShadowMapData.ShadowMap.IsValid() && ViewFamily.Time.GetRealTimeSeconds() - ShadowMapData.LastUsedTime > 2.0f)
+				{
+					ShadowMapData.InvalidateCachedShadow();
+				}
+			}
+		}
+
 		CSV_CUSTOM_STAT(LightCount, UpdatedShadowMaps, float(NumPointShadowCachesUpdatedThisFrame + NumSpotShadowCachesUpdatedThisFrame + NumCSMCachesUpdatedThisFrame), ECsvCustomStatOp::Set);
 		CSV_CUSTOM_STAT_GLOBAL(ShadowCacheUsageMB, (float(Scene->GetCachedWholeSceneShadowMapsSize()) / 1024) / 1024, ECsvCustomStatOp::Set);
 
@@ -5923,12 +5750,269 @@
 	}
 }
 
-void FSceneRenderer::WaitForPrepareDynamicShadowsTask(FDynamicShadowsTaskData* DynamicShadowTaskData)
+void FSceneRenderer::FilterDynamicShadows(FDynamicShadowsTaskData& TaskData)
 {
-	if (DynamicShadowTaskData && DynamicShadowTaskData->PrepareTaskEvent)
+	TRACE_CPUPROFILER_EVENT_SCOPE(CreateFilterShadowPrimitives);
+
+	FFilteredShadowArrays& ShadowArrays = TaskData.ShadowArrays;
+
+	const bool bMobile = Scene->GetFeatureLevel() < ERHIFeatureLevel::SM5;
+
+	TArray<FProjectedShadowInfo*, SceneRenderingAllocator> WholeSceneDirectionalShadowsForLight;
+	TArray<FProjectedShadowInfo*, SceneRenderingAllocator> MobileDynamicSpotlightShadows;
+
+#if !(UE_BUILD_SHIPPING || UE_BUILD_TEST)
+	auto MarkUnsupportedOnePassShadow = [&ShadowArrays](int32 LightId)
 	{
-		DynamicShadowTaskData->PrepareTaskEvent->Wait(ENamedThreads::GetRenderThread_Local());
+		if (ShadowArrays.OnePassShadowUnsupportedLights.Num() <= LightId)
+		{
+			ShadowArrays.OnePassShadowUnsupportedLights.Add(false, 1 + LightId - ShadowArrays.OnePassShadowUnsupportedLights.Num());
+		}
+		ShadowArrays.OnePassShadowUnsupportedLights[LightId] = true;
+	};
+#else 
+	auto MarkUnsupportedOnePassShadow = [](int32 LightId) {};
+#endif
+
+	for (auto LightIt = Scene->Lights.CreateConstIterator(); LightIt; ++LightIt)
+	{
+		const FLightSceneInfoCompact& LightSceneInfoCompact = *LightIt;
+		const FLightSceneInfo* LightSceneInfo = LightSceneInfoCompact.LightSceneInfo;
+		FVisibleLightInfo& VisibleLightInfo = VisibleLightInfos[LightSceneInfo->Id];
+
+		for (int32 ShadowIndex = 0; ShadowIndex < VisibleLightInfo.AllProjectedShadows.Num(); ShadowIndex++)
+		{
+			FProjectedShadowInfo* ProjectedShadowInfo = VisibleLightInfo.AllProjectedShadows[ShadowIndex];
+
+			// Check that the shadow is visible in at least one view before rendering it.
+			bool bShadowIsVisible = false;
+
+			for (int32 ViewIndex = 0; ViewIndex < Views.Num(); ViewIndex++)
+			{
+				const FViewInfo& View = Views[ViewIndex];
+
+				if (ProjectedShadowInfo->DependentView && ProjectedShadowInfo->DependentView != &View)
+				{
+					continue;
+				}
+
+				const FVisibleLightViewInfo& VisibleLightViewInfo = View.VisibleLightInfos[LightSceneInfo->Id];
+				const FPrimitiveViewRelevance ViewRelevance = VisibleLightViewInfo.ProjectedShadowViewRelevanceMap[ShadowIndex];
+				const bool bHasViewRelevance = (ProjectedShadowInfo->bTranslucentShadow && ViewRelevance.HasTranslucency())
+					|| (!ProjectedShadowInfo->bTranslucentShadow && ViewRelevance.bOpaque);
+
+				bShadowIsVisible |= bHasViewRelevance && VisibleLightViewInfo.ProjectedShadowVisibilityMap[ShadowIndex];
+			}
+
+			// Fully cached VSM should not be skipped, since they need to be queued for projection.
+			if (!ProjectedShadowInfo->bShouldRenderVSM)
+			{
+			}
+			// Skip the shadow rendering if there is no primitive in the cached shadow map and the current shadow map either
+			else if ((ProjectedShadowInfo->CacheMode == SDCM_StaticPrimitivesOnly || (ProjectedShadowInfo->CacheMode == SDCM_MovablePrimitivesOnly && !ProjectedShadowInfo->IsWholeSceneDirectionalShadow())) && !ProjectedShadowInfo->HasSubjectPrims())
+			{
+				const FCachedShadowMapData& CachedShadowMapData = Scene->GetCachedShadowMapDataRef(ProjectedShadowInfo->GetLightSceneInfo().Id, FMath::Max(ProjectedShadowInfo->CascadeSettings.ShadowSplitIndex, 0));
+
+				// A shadowmap for movable primitives when there are no movable primitives would normally read directly from the cached shadowmap
+				// However if the cached shadowmap also had no primitives then we need to skip rendering the shadow entirely
+				if (!CachedShadowMapData.bCachedShadowMapHasPrimitives)
+				{
+					bShadowIsVisible = false;
+				}
+			}
+			// If uncached and no primitives, skip allocations etc
+			// Note: cached PreShadows, for some reason, has CacheMode == SDCM_Uncached so need to skip them here.
+			// Ray traced shadows use the GPU managed distance field object buffers, no CPU culling should be used
+			else if (
+				(ProjectedShadowInfo->CacheMode == SDCM_Uncached && !ProjectedShadowInfo->bPreShadow)
+				&& (!ProjectedShadowInfo->HasSubjectPrims() && !ProjectedShadowInfo->bRayTracedDistanceField)
+				&& !ProjectedShadowInfo->IsWholeSceneDirectionalShadow())
+			{
+				bShadowIsVisible = false;
+			}
+
+			if ((IsForwardShadingEnabled(ShaderPlatform) || (IsMobilePlatform(ShaderPlatform) && MobileUsesShadowMaskTexture(ShaderPlatform)))
+				&& ProjectedShadowInfo->GetLightSceneInfo().GetDynamicShadowMapChannel() == -1)
+			{
+				// With forward shading, dynamic shadows are projected into channels of the light attenuation texture based on their assigned DynamicShadowMapChannel
+				bShadowIsVisible = false;
+			}
+
+			// Skip one-pass point light shadows if rendering to them is not suported, unless VSM or DF
+			if (bShadowIsVisible
+				&& !ProjectedShadowInfo->HasVirtualShadowMap()
+				&& !ProjectedShadowInfo->bRayTracedDistanceField
+				&& ProjectedShadowInfo->bOnePassPointLightShadow
+				&& !DoesRuntimeSupportOnePassPointLightShadows(ShaderPlatform))
+			{
+				// Also note this to produce a log message & on-screen warning
+				MarkUnsupportedOnePassShadow(ProjectedShadowInfo->GetLightSceneInfo().Id);
+				bShadowIsVisible = false;
+			}
+
+			if (bShadowIsVisible)
+			{
+				// Visible shadow stats
+				if (ProjectedShadowInfo->bWholeSceneShadow)
+				{
+					INC_DWORD_STAT(STAT_WholeSceneShadows);
+
+					if (ProjectedShadowInfo->CacheMode == SDCM_MovablePrimitivesOnly)
+					{
+						INC_DWORD_STAT(STAT_CachedWholeSceneShadows);
+					}
+				}
+				else if (ProjectedShadowInfo->bPreShadow)
+				{
+					INC_DWORD_STAT(STAT_PreShadows);
+				}
+				else
+				{
+					INC_DWORD_STAT(STAT_PerObjectShadows);
+				}
+
+				bool bNeedsProjection = ProjectedShadowInfo->CacheMode != SDCM_StaticPrimitivesOnly
+					//// Filter out everything but PerObjectOpaqueShadows & Distance field shadows for ES31
+					&& (!bMobile || ProjectedShadowInfo->bPerObjectOpaqueShadow || ProjectedShadowInfo->bRayTracedDistanceField || ProjectedShadowInfo->bWholeSceneShadow);
+
+				if (bNeedsProjection)
+				{
+					if (ProjectedShadowInfo->bCapsuleShadow)
+					{
+						VisibleLightInfo.CapsuleShadowsToProject.Add(ProjectedShadowInfo);
+					}
+					else
+					{
+						if (ProjectedShadowInfo->bRayTracedDistanceField)
+						{
+							ShadowArrays.ProjectedDistanceFieldShadows.Add(ProjectedShadowInfo);
+						}
+						VisibleLightInfo.ShadowsToProject.Add(ProjectedShadowInfo);
+					}
+				}
+
+				const bool bNeedsShadowmapSetup = !ProjectedShadowInfo->bCapsuleShadow && !ProjectedShadowInfo->bRayTracedDistanceField;
+
+				if (bNeedsShadowmapSetup)
+				{
+					// Certain shadow types skip the depth pass when there are no movable primitives to composite.
+					bool bAlwaysHasDepthPass = true;
+
+					if (ProjectedShadowInfo->HasVirtualShadowMap() || ProjectedShadowInfo->VirtualShadowMapClipmap.IsValid())
+					{
+						SortedShadowsForShadowDepthPass.VirtualShadowMapShadows.Add(ProjectedShadowInfo);
+					}
+					else if (ProjectedShadowInfo->bPreShadow && ProjectedShadowInfo->bAllocatedInPreshadowCache)
+					{
+						ShadowArrays.CachedPreShadows.Add(ProjectedShadowInfo);
+					}
+					else if (ProjectedShadowInfo->bDirectionalLight && ProjectedShadowInfo->bWholeSceneShadow)
+					{
+						if (ProjectedShadowInfo->CacheMode == SDCM_StaticPrimitivesOnly)
+						{
+							ShadowArrays.CachedWholeSceneDirectionalShadows.Add(ProjectedShadowInfo);
+						}
+						else
+						{
+							WholeSceneDirectionalShadowsForLight.Add(ProjectedShadowInfo);
+						}
+					}
+					else if (ProjectedShadowInfo->bOnePassPointLightShadow)
+					{
+						ShadowArrays.WholeScenePointShadows.Add(ProjectedShadowInfo);
+						bAlwaysHasDepthPass = false;
+					}
+					else if (ProjectedShadowInfo->bTranslucentShadow)
+					{
+						ShadowArrays.TranslucentShadows.Add(ProjectedShadowInfo);
+					}
+					else if (ProjectedShadowInfo->CacheMode == SDCM_StaticPrimitivesOnly)
+					{
+						check(ProjectedShadowInfo->bWholeSceneShadow);
+						ShadowArrays.CachedSpotlightShadows.Add(ProjectedShadowInfo);
+					}
+					else if (bMobile && ProjectedShadowInfo->bWholeSceneShadow)
+					{
+						MobileDynamicSpotlightShadows.Add(ProjectedShadowInfo);
+					}
+					else
+					{
+						ShadowArrays.Shadows.Add(ProjectedShadowInfo);
+						bAlwaysHasDepthPass = false;
+					}
+
+					if (bAlwaysHasDepthPass || ProjectedShadowInfo->CacheMode != SDCM_MovablePrimitivesOnly || ProjectedShadowInfo->HasSubjectPrims())
+					{
+						ShadowArrays.ShadowsToSetupViews.Add(ProjectedShadowInfo);
+					}
+				}
+			}
+		}
+
+		// Sort cascades, this is needed for blending between cascades to work, and also to enable fetching the farthest cascade as index [0]
+		VisibleLightInfo.ShadowsToProject.Sort(FCompareFProjectedShadowInfoBySplitIndex());
+
+		if (!bMobile)
+		{
+			// Only allocate CSM targets if at least one of the SMs in the CSM has any primitives, but then always allocate all of them as some algorithms depend on this (notably the forward shading structures).
+			// Don't perform this if VSM are not enabled.
+			bool bAnyHasSubjectPrims = !VirtualShadowMapArray.IsEnabled();
+			for (const FProjectedShadowInfo* ProjectedShadowInfo : WholeSceneDirectionalShadowsForLight)
+			{
+				if (bAnyHasSubjectPrims)
+				{
+					break;
+				}
+
+				bAnyHasSubjectPrims |= ProjectedShadowInfo->HasSubjectPrims();
+			}
+			if (bAnyHasSubjectPrims)
+			{
+				ShadowArrays.WholeSceneDirectionalShadows.Append(WholeSceneDirectionalShadowsForLight);
+			}
+		}
+		else
+		{
+			//Only one directional light could cast csm on mobile, so we could delay allocation for it and see if we could combine any spotlight shadow with it.
+			ShadowArrays.WholeSceneDirectionalShadows.Append(WholeSceneDirectionalShadowsForLight);
+		}
+
+		WholeSceneDirectionalShadowsForLight.Reset();
 	}
+
+	if (bMobile)
+	{
+		// AllocateMobileCSMAndSpotLightShadowDepthTargets would only allocate a single large render target for all shadows, so if the requirement exceeds the MaxTextureSize, the rest of the shadows will not get space for rendering
+		// So we sort spotlight shadows and append them at the last to make sure csm will get space in any case.
+		MobileDynamicSpotlightShadows.Sort(FCompareFProjectedShadowInfoByResolution());
+
+		//Limit the number of spotlights shadow for performance reason
+		static const auto MobileMaxVisibleMovableSpotLightShadowsCVar = IConsoleManager::Get().FindTConsoleVariableDataInt(TEXT("r.Mobile.MaxVisibleMovableSpotLightShadows"));
+		if (MobileMaxVisibleMovableSpotLightShadowsCVar)
+		{
+			int32 MobileMaxVisibleMovableSpotLightShadows = MobileMaxVisibleMovableSpotLightShadowsCVar->GetValueOnRenderThread();
+			MobileDynamicSpotlightShadows.RemoveAt(MobileMaxVisibleMovableSpotLightShadows, FMath::Max(MobileDynamicSpotlightShadows.Num() - MobileMaxVisibleMovableSpotLightShadows, 0), false);
+		}
+
+		ShadowArrays.WholeSceneDirectionalShadows.Append(MobileDynamicSpotlightShadows);
+		MobileDynamicSpotlightShadows.Reset();
+	}
+
+	// Sort the projected shadows by resolution.
+	ShadowArrays.TranslucentShadows.Sort(FCompareFProjectedShadowInfoByResolution());
+	
+	TConstArrayView<FProjectedShadowInfo*> ShadowsToSetupViews = TaskData.ShadowArrays.ShadowsToSetupViews;
+	const int32 SetupShadowDepthViewBatchSize = 8;
+
+	ParallelForTemplate(
+		TEXT("SetupShadowDepthView"), ShadowsToSetupViews.Num(), SetupShadowDepthViewBatchSize,
+		[this, ShadowsToSetupViews](int32 Index)
+		{
+			FOptionalTaskTagScope Scope(ETaskTag::EParallelRenderingThread);
+			ShadowsToSetupViews[Index]->SetupShadowDepthView(this);
+		},
+		TaskData.bMultithreadedCreateAndFilterShadows ? EParallelForFlags::None : EParallelForFlags::ForceSingleThread
+	);
 }
 
 FDynamicShadowsTaskData* FSceneRenderer::BeginInitDynamicShadows(bool bRunningEarly)
@@ -5937,7 +6021,7 @@
 	CSV_SCOPED_TIMING_STAT_EXCLUSIVE(InitViews_Shadows);
 	SCOPED_NAMED_EVENT_TEXT("FSceneRenderer::BeginInitDynamicShadows", FColor::Magenta);
 
-	FDynamicShadowsTaskData* DynamicShadowTaskData = new FDynamicShadowsTaskData(this, bRunningEarly);
+	FDynamicShadowsTaskData* DynamicShadowTaskData = Allocator.Create<FDynamicShadowsTaskData>(this, bRunningEarly);
 
 	// Gathers the list of primitives used to draw various shadow types
 	BeginGatherShadowPrimitives(DynamicShadowTaskData);
@@ -5953,6 +6037,8 @@
 	SCOPE_CYCLE_COUNTER(STAT_DynamicShadowSetupTime);
 	CSV_SCOPED_TIMING_STAT_EXCLUSIVE(ShadowInitDynamic);
 
+	check(TaskData);
+
 	FinishGatherShadowPrimitives(TaskData);
 
 	if (ShadowSceneRenderer)
@@ -5960,7 +6046,7 @@
 		ShadowSceneRenderer->PostSetupDebugRender();
 	}
 
-	AllocateShadowDepthTargets(GraphBuilder.RHICmdList);
+	AllocateShadowDepthTargets(GraphBuilder.RHICmdList, *TaskData);
 
 	// Generate mesh element arrays from shadow primitive arrays
 	GatherShadowDynamicMeshElements(DynamicIndexBuffer, DynamicVertexBuffer, DynamicReadBuffer, InstanceCullingManager);
@@ -6004,13 +6090,14 @@
 	}
 }
 
-void FSceneRenderer::InitDynamicShadows(FRDGBuilder& GraphBuilder, FGlobalDynamicIndexBuffer& DynamicIndexBuffer, FGlobalDynamicVertexBuffer& DynamicVertexBuffer, FGlobalDynamicReadBuffer& DynamicReadBuffer, FInstanceCullingManager& InstanceCullingManager, FRDGExternalAccessQueue& ExternalAccessQueue)
+FDynamicShadowsTaskData* FSceneRenderer::InitDynamicShadows(FRDGBuilder& GraphBuilder, FGlobalDynamicIndexBuffer& DynamicIndexBuffer, FGlobalDynamicVertexBuffer& DynamicVertexBuffer, FGlobalDynamicReadBuffer& DynamicReadBuffer, FInstanceCullingManager& InstanceCullingManager, FRDGExternalAccessQueue& ExternalAccessQueue)
 {
 	FDynamicShadowsTaskData* TaskData = BeginInitDynamicShadows(false);
 	FinishInitDynamicShadows(GraphBuilder, TaskData, DynamicIndexBuffer, DynamicVertexBuffer, DynamicReadBuffer, InstanceCullingManager, ExternalAccessQueue);
+	return TaskData;
 }
 
-void FSceneRenderer::AllocateMobileCSMAndSpotLightShadowDepthTargets(FRHICommandListImmediate& RHICmdList, const TArray<FProjectedShadowInfo*, SceneRenderingAllocator>& MobileCSMAndSpotLightShadows)
+void FSceneRenderer::AllocateMobileCSMAndSpotLightShadowDepthTargets(FRHICommandListImmediate& RHICmdList, TConstArrayView<FProjectedShadowInfo*> MobileCSMAndSpotLightShadows)
 {
 	if (MobileCSMAndSpotLightShadows.Num() > 0)
 	{
@@ -6061,7 +6148,6 @@
 				if (ProjectedShadowInfo->bAllocated)
 				{
 					ProjectedShadowInfo->RenderTargets.CopyReferencesFromRenderTargets(ShadowMapAtlas.RenderTargets);
-					ProjectedShadowInfo->SetupShadowDepthView(this);
 					ShadowMapAtlas.Shadows.Add(ProjectedShadowInfo);
 				}
 			}
diff --git a/Engine/Source/Runtime/Renderer/Private/ShadowSetupMobile.cpp b/Engine/Source/Runtime/Renderer/Private/ShadowSetupMobile.cpp
--- a/Engine/Source/Runtime/Renderer/Private/ShadowSetupMobile.cpp
+++ b/Engine/Source/Runtime/Renderer/Private/ShadowSetupMobile.cpp
@@ -283,7 +283,7 @@ static void VisualizeMobileDynamicCSMSubjectCapsules(FViewInfo& View, FLightScen
 }
 
 /** Finds the visible dynamic shadows for each view. */
-void FMobileSceneRenderer::InitDynamicShadows(FRDGBuilder& GraphBuilder, FInstanceCullingManager& InstanceCullingManager, FRDGExternalAccessQueue& ExternalAccessQueue)
+FDynamicShadowsTaskData* FMobileSceneRenderer::InitDynamicShadows(FRDGBuilder& GraphBuilder, FInstanceCullingManager& InstanceCullingManager, FRDGExternalAccessQueue& ExternalAccessQueue)
 {
 	static auto* MyCVarMobileEnableStaticAndCSMShadowReceivers = IConsoleManager::Get().FindTConsoleVariableDataInt(TEXT("r.Mobile.EnableStaticAndCSMShadowReceivers"));
 	const bool bCombinedStaticAndCSMEnabled = MyCVarMobileEnableStaticAndCSMShadowReceivers->GetValueOnRenderThread()!=0;
@@ -307,7 +307,7 @@ void FMobileSceneRenderer::InitDynamicShadows(FRDGBuilder& GraphBuilder, FInstan
 		}
 	}
 
-	FSceneRenderer::InitDynamicShadows(GraphBuilder, DynamicIndexBuffer, DynamicVertexBuffer, DynamicReadBuffer, InstanceCullingManager, ExternalAccessQueue);
+	FDynamicShadowsTaskData* TaskData = FSceneRenderer::InitDynamicShadows(GraphBuilder, DynamicIndexBuffer, DynamicVertexBuffer, DynamicReadBuffer, InstanceCullingManager, ExternalAccessQueue);
 
 	bool bAlwaysUseCSM = false;
 	const bool bSkipCSMShaderCulling = MobileBasePassAlwaysUsesCSM(Scene->GetShaderPlatform());
@@ -361,6 +361,8 @@ void FMobileSceneRenderer::InitDynamicShadows(FRDGBuilder& GraphBuilder, FInstan
 			bModulatedShadowsInUse = VisibleLightInfo.ShadowsToProject.Num() > 0;
 		}
 	}
+
+	return TaskData;
 }
 
 // generate a single FProjectedShadowInfo to encompass LightSceneInfo.
