Compute optimized scissor rect size to exclude portions of the CSM slices outside the view frustum (disable by default).

diff --git a/Engine/Source/Runtime/Renderer/Private/ShadowDepthRendering.cpp b/Engine/Source/Runtime/Renderer/Private/ShadowDepthRendering.cpp
--- a/Engine/Source/Runtime/Renderer/Private/ShadowDepthRendering.cpp
+++ b/Engine/Source/Runtime/Renderer/Private/ShadowDepthRendering.cpp
@@ -668,6 +668,17 @@ void FProjectedShadowInfo::SetStateForView(FRHICommandList& RHICmdList) const
 		Y + ResolutionY + 2 * BorderSize,
 		1.0f
 	);
+
+	if (ShouldUseCSMScissorOptim())
+	{
+		RHICmdList.SetScissorRect(
+			true,
+			ScissorRectOptim.Min.X + X,
+			ScissorRectOptim.Min.Y + Y,
+			ScissorRectOptim.Max.X + X,
+			ScissorRectOptim.Max.Y + Y
+		);
+	}
 }
 
 void SetStateForShadowDepth(bool bOnePassPointLightShadow, bool bDirectionalLight, FMeshPassProcessorRenderState& DrawRenderState, EMeshPass::Type InMeshPassTargetType)
diff --git a/Engine/Source/Runtime/Renderer/Private/ShadowRendering.h b/Engine/Source/Runtime/Renderer/Private/ShadowRendering.h
--- a/Engine/Source/Runtime/Renderer/Private/ShadowRendering.h
+++ b/Engine/Source/Runtime/Renderer/Private/ShadowRendering.h
@@ -367,6 +367,12 @@ public:
 	/** Size of the border, if any, used to allow filtering without clamping for shadows stored in an atlas. */
 	uint32 BorderSize;
 
+	/** Scissor rect size to exclude portion of the CSM slices outside the view frustum */
+	FIntRect ScissorRectOptim;
+
+	/** Compute the scissors coords to exclude portion of the CSM slices outside the view frustum*/
+	void ComputeScissorRectOptim();
+
 	/** The largest percent of either the width or height of any view. */
 	float MaxScreenPercent;
 
@@ -783,6 +789,9 @@ public:
 		return ShadowDepthPass.BuildRenderingCommands(GraphBuilder, GPUScene, InstanceCullingDrawParams);
 	}
 
+	/** Check if we need to set the scissor rect to exclude portion of the CSM slices outside the view frustum */
+	bool ShouldUseCSMScissorOptim() const;
+
 private:
 	// 0 if Setup...() wasn't called yet
 	FLightSceneInfo* LightSceneInfo;
diff --git a/Engine/Source/Runtime/Renderer/Private/ShadowSetup.cpp b/Engine/Source/Runtime/Renderer/Private/ShadowSetup.cpp
--- a/Engine/Source/Runtime/Renderer/Private/ShadowSetup.cpp
+++ b/Engine/Source/Runtime/Renderer/Private/ShadowSetup.cpp
@@ -39,6 +39,9 @@
 #include "UnrealEngine.h"
 #include "CanvasTypes.h"
 #include "Shadows/ShadowScene.h"
+#include "LineTypes.h"
+
+using namespace UE::Geometry;
 
 /** Number of cube map shadow depth surfaces that will be created and used for rendering one pass point light shadows. */
 static const int32 NumCubeShadowDepthSurfaces = 5;
@@ -359,6 +362,13 @@ static TAutoConsoleVariable<int32> CVarCSMScrollingMaxExtraStaticShadowSubjects(
 	TEXT("The maximum number of extra static shadow subjects need to be drawed when scrolling CSM."),
 	ECVF_RenderThreadSafe);
 
+static TAutoConsoleVariable<bool> CVarCSMScissorOptim(
+	TEXT("r.Shadow.CSMScissorOptim"),
+	false,
+	TEXT("Compute optimized scissor rect size to exclude portions of the CSM slices outside the view frustum"),
+	ECVF_RenderThreadSafe
+);
+
 /**
  * Helper function to maintain the common login of caching the whole scene shadow and be able to handle the specific-purpose through the CacheLambda and UnCacheLambda
  */
@@ -1058,6 +1068,11 @@ void FProjectedShadowInfo::SetupWholeSceneProjection(
 		}
 		CasterOuterFrustum.Init();
 	}
+
+	if (ShouldUseCSMScissorOptim())
+	{
+		ComputeScissorRectOptim();
+	}
 }
 
 
@@ -1588,6 +1603,11 @@ FLODMask FProjectedShadowInfo::CalcAndUpdateLODToRender(FViewInfo& CurrentView,
 	return ShadowLODToRender;
 }
 
+bool FProjectedShadowInfo::ShouldUseCSMScissorOptim() const
+{
+	return CVarCSMScissorOptim.GetValueOnRenderThread() != 0 && bWholeSceneShadow && bDirectionalLight && !bOnePassPointLightShadow && !bRayTracedDistanceField;
+}
+
 FORCEINLINE bool FProjectedShadowInfo::ShouldDrawStaticMesh(const FStaticMeshBatchRelevance& StaticMeshRelevance, const FLODMask& ShadowLODToRender, bool& bOutDrawingStaticMeshes) const
 {
 	if ((StaticMeshRelevance.CastShadow || (bSelfShadowOnly && StaticMeshRelevance.bUseForDepthPass)) && ShadowLODToRender.ContainsLOD(StaticMeshRelevance.LODIndex))
@@ -2641,6 +2661,149 @@ void FProjectedShadowInfo::ClearTransientArrays()
 	GraphicsMinimalPipelineStateSet.Empty();
 }
 
+void FProjectedShadowInfo::ComputeScissorRectOptim()
+{
+	ScissorRectOptim.Min = FIntPoint(0, 0);
+	ScissorRectOptim.Max = FIntPoint(ResolutionX, ResolutionY);
+
+	// get the far plane corners in world space
+	const FMatrix& ViewMatrix = DependentView->ViewMatrices.GetViewMatrix();
+	const FMatrix& ProjectionMatrix = DependentView->ViewMatrices.GetProjectionMatrix();
+	const FVector CameraDirection = ViewMatrix.GetColumn(2);
+	const FVector ViewOrigin = DependentView->ViewMatrices.GetViewOrigin();
+
+	// Support asymmetric projection
+	// Get FOV and AspectRatio from the view's projection matrix.
+	float AspectRatio = ProjectionMatrix.M[1][1] / ProjectionMatrix.M[0][0];
+	bool bIsPerspectiveProjection = true;
+
+	// Build the camera frustum for this cascade
+	float HalfHorizontalFOV = bIsPerspectiveProjection ? FMath::Atan(1.0f / ProjectionMatrix.M[0][0]) : UE_PI / 4.0f;
+	float HalfVerticalFOV = bIsPerspectiveProjection ? FMath::Atan(1.0f / ProjectionMatrix.M[1][1]) : FMath::Atan((FMath::Tan(UE_PI / 4.0f) / AspectRatio));
+	float AsymmetricFOVScaleX = ProjectionMatrix.M[2][0];
+	float AsymmetricFOVScaleY = ProjectionMatrix.M[2][1];
+
+	// Far plane
+	const float EndHorizontalTotalLength = CascadeSettings.SplitFar * FMath::Tan(HalfHorizontalFOV);
+	const float EndVerticalTotalLength = CascadeSettings.SplitFar * FMath::Tan(HalfVerticalFOV);
+	const FVector EndCameraLeftOffset = ViewMatrix.GetColumn(0) * -EndHorizontalTotalLength * (1 + AsymmetricFOVScaleX);
+	const FVector EndCameraRightOffset = ViewMatrix.GetColumn(0) * EndHorizontalTotalLength * (1 - AsymmetricFOVScaleX);
+	const FVector EndCameraBottomOffset = ViewMatrix.GetColumn(1) * -EndVerticalTotalLength * (1 + AsymmetricFOVScaleY);
+	const FVector EndCameraTopOffset = ViewMatrix.GetColumn(1) * EndVerticalTotalLength * (1 - AsymmetricFOVScaleY);
+
+	FVector FrustumCorners[5];
+	FrustumCorners[0] = ViewOrigin + CameraDirection * CascadeSettings.SplitFar + EndCameraRightOffset + EndCameraTopOffset;//Far  Top    Right
+	FrustumCorners[1] = ViewOrigin + CameraDirection * CascadeSettings.SplitFar + EndCameraRightOffset + EndCameraBottomOffset;//Far  Bottom Right
+	FrustumCorners[2] = ViewOrigin + CameraDirection * CascadeSettings.SplitFar + EndCameraLeftOffset + EndCameraTopOffset;//Far  Top    Left
+	FrustumCorners[3] = ViewOrigin + CameraDirection * CascadeSettings.SplitFar + EndCameraLeftOffset + EndCameraBottomOffset;//Far  Bottom Left
+	FrustumCorners[4] = ViewOrigin + CameraDirection * CascadeSettings.SplitFar;//view direction
+
+	FMatrix WorldToShadow = FTranslationMatrix(PreShadowTranslation) *
+		FMatrix(TranslatedWorldToClipInnerMatrix) *
+		FMatrix(
+			FPlane(0.5f, 0, 0, 0),
+			FPlane(0, -0.5f, 0, 0),
+			FPlane(0, 0, InvMaxSubjectDepth, 0),
+			FPlane(
+				1.0f / (ResolutionX + BorderSize * 2) + 0.5f,
+				1.0f / (ResolutionY + BorderSize * 2) + 0.5f,
+				0,
+				1
+			)
+		);
+
+	int32 FullResX = ResolutionX + BorderSize * 2;
+	int32 FullResY = ResolutionY + BorderSize * 2;
+	FVector4 FrustumClipShadowCorners[6];
+	FVector4 ScaleView = FVector4(FullResX, FullResY, 1.0, 1.0);
+
+	// Transform each corners of the view frustum in Shadow clip space
+	for (int Index = 0; Index < 5; Index++)
+	{
+		FrustumClipShadowCorners[Index] = WorldToShadow.TransformPosition(FrustumCorners[Index]);
+		FrustumClipShadowCorners[Index] /= FrustumClipShadowCorners[Index].W;
+		FrustumClipShadowCorners[Index] *= ScaleView;
+	}
+
+	// Transform the origin of the view frustum in Shadow Clip Space
+	FrustumClipShadowCorners[5] = WorldToShadow.TransformPosition(ViewOrigin);
+	FrustumClipShadowCorners[5] /= FrustumClipShadowCorners[4].W;
+	FrustumClipShadowCorners[5] *= ScaleView;
+
+	auto ComputeScissorIntersection = [FullResX, FullResY](const FLine2d& Line)
+	{
+		FLine2d Borders[4] = { FLine2d(FVector2d(0.0, 0.0), FVector2d(1.0, 0.0)),
+			FLine2d(FVector2d(FullResX, 0.0), FVector2d(0.0, 1.0)),
+			FLine2d(FVector2d(FullResX, FullResY), FVector2d(-1.0, 0.0)),
+			FLine2d(FVector2d(0.0, FullResY), FVector2d(0.0, -1.0)) };
+
+		FVector2D IntPoint;
+		FVector2D Result;
+		float MinDistance = 100000.0f;
+		bool bIntersectionvalid = false;
+
+		for (int Index = 0; Index < 4; Index++)
+		{
+			if (Line.IntersectionPoint(Borders[Index], IntPoint))
+			{
+				//check if the point is behind
+				FVector2D IntVector = (IntPoint - Line.Origin);
+				if (FVector2D::DotProduct(IntVector.GetSafeNormal(), Line.Direction) > 0.0)
+				{
+					float IntersectionDistance = FVector2D::Distance(IntPoint, Line.Origin);
+					if (IntersectionDistance < MinDistance)
+					{
+						Result = IntPoint;
+						MinDistance = IntersectionDistance;
+						bIntersectionvalid = true;
+					}
+				}
+			}
+		}
+		return FVector4(Result.X, Result.Y, 0.0, 0.0);
+	};
+
+	// compute the intersection with shadow render target borders
+	FVector2D FrustrumOrigin2D = FVector2D(FrustumClipShadowCorners[5].X, FrustumClipShadowCorners[5].Y);
+
+	if (FrustrumOrigin2D.X >= 0.0 && FrustrumOrigin2D.X <= FullResX &&
+		FrustrumOrigin2D.Y >= 0.0 && FrustrumOrigin2D.Y <= FullResY)
+	{
+		// Compute the intersection of the view frustum corners in shadow clip space
+		for (int Index = 0; Index < 5; Index++)
+		{			
+			FLine2d FrustumLine = FLine2d(FrustrumOrigin2D, (FVector2D(FrustumClipShadowCorners[Index].X, FrustumClipShadowCorners[Index].Y) - FrustrumOrigin2D).GetSafeNormal());
+			FrustumClipShadowCorners[Index] = ComputeScissorIntersection(FrustumLine);
+		}
+
+		// Get the scissor rect min and max values from the intersection result
+		FIntRect ScissorMinMax(FIntPoint(FullResX, FullResX), FIntPoint(0, 0));
+
+		for (int Index = 0; Index < 6; Index++)
+		{
+			if (FrustumClipShadowCorners[Index].X < ScissorMinMax.Min.X)
+			{
+				ScissorMinMax.Min.X = FMath::Clamp(FrustumClipShadowCorners[Index].X, 0, FullResX);
+			}
+			if (FrustumClipShadowCorners[Index].Y < ScissorMinMax.Min.Y)
+			{
+				ScissorMinMax.Min.Y = FMath::Clamp(FrustumClipShadowCorners[Index].Y, 0, FullResY);
+			}
+			if (FrustumClipShadowCorners[Index].X > ScissorMinMax.Max.X)
+			{
+				ScissorMinMax.Max.X = FMath::Clamp(FrustumClipShadowCorners[Index].X, 0, FullResX);
+			}
+			if (FrustumClipShadowCorners[Index].Y > ScissorMinMax.Max.Y)
+			{
+				ScissorMinMax.Max.Y = FMath::Clamp(FrustumClipShadowCorners[Index].Y, 0, FullResY);
+			}
+		}
+		
+		ScissorRectOptim = ScissorMinMax;
+		ScissorRectOptim += FIntPoint(X, Y);
+	}
+}
+
 FSceneViewState::FProjectedShadowKey::FProjectedShadowKey(const FProjectedShadowInfo& ProjectedShadowInfo)
 	: PrimitiveId(ProjectedShadowInfo.GetParentSceneInfo() ? ProjectedShadowInfo.GetParentSceneInfo()->PrimitiveComponentId : FPrimitiveComponentId())
 	, Light(ProjectedShadowInfo.GetLightSceneInfo().Proxy->GetLightComponent())
