Changed r.Nanite and r.Shadow.Virtual.Enable from FAutoConsoleVariableRef to TAutoConsoleVariable to work around FAutoConsoleVariableRef::GetInt not being safe to access from the rendering thread.
Without this fix, the toggling of these cvars from the game thread can end up being observable from the rendering thread mid-frame, which can cause rare and hard to debug crashes.

diff --git a/Engine/Source/Runtime/Engine/Private/StaticMeshRender.cpp b/Engine/Source/Runtime/Engine/Private/StaticMeshRender.cpp
--- a/Engine/Source/Runtime/Engine/Private/StaticMeshRender.cpp
+++ b/Engine/Source/Runtime/Engine/Private/StaticMeshRender.cpp
@@ -93,10 +93,9 @@ static FAutoConsoleCommand GToggleReversedIndexBuffersCmd(
 	);
 
 // TODO: Should move this outside of SM, since Nanite can be used for multiple primitive types
-int32 GRenderNaniteMeshes = 1;
-FAutoConsoleVariableRef CVarRenderNaniteMeshes(
+TAutoConsoleVariable<int32> CVarRenderNaniteMeshes(
 	TEXT("r.Nanite"),
-	GRenderNaniteMeshes,
+	1,
 	TEXT("Render static meshes using Nanite."),
 	FConsoleVariableDelegate::CreateLambda([](IConsoleVariable* InVariable)
 	{
diff --git a/Engine/Source/Runtime/Renderer/Private/VirtualShadowMaps/VirtualShadowMapArray.cpp b/Engine/Source/Runtime/Renderer/Private/VirtualShadowMaps/VirtualShadowMapArray.cpp
--- a/Engine/Source/Runtime/Renderer/Private/VirtualShadowMaps/VirtualShadowMapArray.cpp
+++ b/Engine/Source/Runtime/Renderer/Private/VirtualShadowMaps/VirtualShadowMapArray.cpp
@@ -65,10 +65,9 @@ FAutoConsoleVariableRef CVarVSMShowLightDrawEvents(
 	ECVF_RenderThreadSafe
 );
 
-int32 GEnableVirtualShadowMaps = 0;
-FAutoConsoleVariableRef CVarEnableVirtualShadowMaps(
+TAutoConsoleVariable<int32> CVarEnableVirtualShadowMaps(
 	TEXT("r.Shadow.Virtual.Enable"),
-	GEnableVirtualShadowMaps,
+	0,
 	TEXT("Enable Virtual Shadow Maps."),
 	FConsoleVariableDelegate::CreateLambda([](IConsoleVariable* InVariable)
 	{
