PR #10117 : Fix RT dynamic geometry update buffers being twice larger than they need to be.

Also removed some unnecessary copying of FRayTracingDynamicGeometryUpdateParams.

#rnx

diff --git a/Engine/Source/Runtime/Engine/Private/InstancedStaticMesh.cpp b/Engine/Source/Runtime/Engine/Private/InstancedStaticMesh.cpp
--- a/Engine/Source/Runtime/Engine/Private/InstancedStaticMesh.cpp
+++ b/Engine/Source/Runtime/Engine/Private/InstancedStaticMesh.cpp
@@ -2095,7 +2095,7 @@ void FInstancedStaticMeshSceneProxy::GetDynamicRayTracingInstances(struct FRayTr
 						RayTracingWPODynamicTemplate.Materials,
 						false,
 						(uint32)LODModel.GetNumVertices(),
-						uint32((SIZE_T)LODModel.GetNumVertices() * sizeof(FVector)),
+						uint32((SIZE_T)LODModel.GetNumVertices() * sizeof(FVector3f)),
 						DynamicData.DynamicGeometry.Initializer.TotalPrimitiveCount,
 						&DynamicData.DynamicGeometry,
 						nullptr,
diff --git a/Engine/Source/Runtime/Renderer/Private/DeferredShadingRenderer.cpp b/Engine/Source/Runtime/Renderer/Private/DeferredShadingRenderer.cpp
--- a/Engine/Source/Runtime/Renderer/Private/DeferredShadingRenderer.cpp
+++ b/Engine/Source/Runtime/Renderer/Private/DeferredShadingRenderer.cpp
@@ -1205,7 +1205,7 @@ bool FDeferredShadingSceneRenderer::GatherRayTracingWorldInstancesForView(FRDGBu
 
 			SceneProxy->GetDynamicRayTracingInstances(MaterialGatheringContext, TempRayTracingInstances);
 
-			for (auto DynamicRayTracingGeometryUpdate : MaterialGatheringContext.DynamicRayTracingGeometriesToUpdate)
+			for (const FRayTracingDynamicGeometryUpdateParams& DynamicRayTracingGeometryUpdate : MaterialGatheringContext.DynamicRayTracingGeometriesToUpdate)
 			{
 				Scene->GetRayTracingDynamicGeometryCollection()->AddDynamicMeshBatchForGeometryUpdate(
 					Scene,
diff --git a/Engine/Source/Runtime/Renderer/Private/RayTracing/RayTracingDynamicGeometry.cpp b/Engine/Source/Runtime/Renderer/Private/RayTracing/RayTracingDynamicGeometry.cpp
--- a/Engine/Source/Runtime/Renderer/Private/RayTracing/RayTracingDynamicGeometry.cpp
+++ b/Engine/Source/Runtime/Renderer/Private/RayTracing/RayTracingDynamicGeometry.cpp
@@ -159,7 +159,7 @@ void FRayTracingDynamicGeometryCollection::AddDynamicMeshBatchForGeometryUpdate(
 	const FScene* Scene,
 	const FSceneView* View,
 	const FPrimitiveSceneProxy* PrimitiveSceneProxy,
-	FRayTracingDynamicGeometryUpdateParams UpdateParams,
+	const FRayTracingDynamicGeometryUpdateParams& UpdateParams,
 	uint32 PrimitiveId
 )
 {
diff --git a/Engine/Source/Runtime/Renderer/Public/RayTracingDynamicGeometryCollection.h b/Engine/Source/Runtime/Renderer/Public/RayTracingDynamicGeometryCollection.h
--- a/Engine/Source/Runtime/Renderer/Public/RayTracingDynamicGeometryCollection.h
+++ b/Engine/Source/Runtime/Renderer/Public/RayTracingDynamicGeometryCollection.h
@@ -22,7 +22,7 @@ public:
 		const FScene* Scene, 
 		const FSceneView* View, 
 		const FPrimitiveSceneProxy* PrimitiveSceneProxy, 
-		FRayTracingDynamicGeometryUpdateParams Params,
+		const FRayTracingDynamicGeometryUpdateParams& Params,
 		uint32 PrimitiveId 
 	);
 
