Fix incorrect tile resolution rounding, causing tile issue at the bottom of the screen with certain screen percentage value.

diff --git a/Engine/Source/Runtime/Renderer/Private/Strata/Strata.cpp b/Engine/Source/Runtime/Renderer/Private/Strata/Strata.cpp
--- a/Engine/Source/Runtime/Renderer/Private/Strata/Strata.cpp
+++ b/Engine/Source/Runtime/Renderer/Private/Strata/Strata.cpp
@@ -130,7 +130,7 @@ static FIntPoint GetStrataTextureTileResolution(const FViewInfo& View, const FIn
 	if (InSpace & EStrataTileSpace::StrataTileSpace_Overflow)
 	{
 		const float OverflowRatio = GetStrataTileOverflowRatio(View);
-		Out.Y += FMath::DivideAndRoundUp(FMath::CeilToInt(InResolution.Y * OverflowRatio), STRATA_TILE_SIZE);
+		Out.Y += FMath::CeilToInt(FMath::DivideAndRoundUp(InResolution.Y, STRATA_TILE_SIZE) * OverflowRatio);
 	}
 	return Out;
 }
