Fix parallel shadow setup that fails to call FilterDynamicShadows if no Gather shadow primitive tasks were issued (bug introduced in CL 24292130).

diff --git a/Engine/Source/Runtime/Renderer/Private/ShadowSetup.cpp b/Engine/Source/Runtime/Renderer/Private/ShadowSetup.cpp
--- a/Engine/Source/Runtime/Renderer/Private/ShadowSetup.cpp
+++ b/Engine/Source/Runtime/Renderer/Private/ShadowSetup.cpp
@@ -4814,7 +4814,7 @@ struct FGatherShadowPrimitivesPrepareTask
 			}
 		}
 
-		if (TaskData.bMultithreadedCreateAndFilterShadows && Prereqs.Num() > 0)
+		if (TaskData.bMultithreadedCreateAndFilterShadows)
 		{
 			FGraphEventRef FinalizeTask = FFunctionGraphTask::CreateAndDispatchWhenReady([TaskData = &TaskData]()
 				{
