Fixed sky atmosphere SkyLightLut not accounting for overriden light direction.

diff --git a/Engine/Source/Runtime/Renderer/Private/SkyAtmosphereRendering.cpp b/Engine/Source/Runtime/Renderer/Private/SkyAtmosphereRendering.cpp
--- a/Engine/Source/Runtime/Renderer/Private/SkyAtmosphereRendering.cpp
+++ b/Engine/Source/Runtime/Renderer/Private/SkyAtmosphereRendering.cpp
@@ -1391,7 +1391,7 @@ void FSceneRenderer::RenderSkyAtmosphereLookUpTables(FRDGBuilder& GraphBuilder,
 		FLightSceneInfo* Light1 = Scene->AtmosphereLights[1];
 		if (Light0)
 		{
-			PassParameters->AtmosphereLightDirection0 = FVector3f(-Light0->Proxy->GetDirection());
+			PassParameters->AtmosphereLightDirection0 = FVector3f(SkyAtmosphereSceneProxy.GetAtmosphereLightDirection(0, -Light0->Proxy->GetDirection()));
 			PassParameters->AtmosphereLightIlluminanceOuterSpace0 = Light0->Proxy->GetOuterSpaceIlluminance();
 		}
 		else
@@ -1401,7 +1401,7 @@ void FSceneRenderer::RenderSkyAtmosphereLookUpTables(FRDGBuilder& GraphBuilder,
 		}
 		if (Light1)
 		{
-			PassParameters->AtmosphereLightDirection1 = FVector3f(-Light1->Proxy->GetDirection());
+			PassParameters->AtmosphereLightDirection1 = FVector3f(SkyAtmosphereSceneProxy.GetAtmosphereLightDirection(1, -Light1->Proxy->GetDirection()));
 			PassParameters->AtmosphereLightIlluminanceOuterSpace1 = Light1->Proxy->GetOuterSpaceIlluminance();
 		}
 		else
