Fixed VSM binding a dummy structured buffer where a byte address buffer was expected.

We didn't have GSystemTextures.GetDefaultByteAddressBuffer(), so now we do.

#rnx

diff --git a/Engine/Source/Runtime/Renderer/Private/SystemTextures.cpp b/Engine/Source/Runtime/Renderer/Private/SystemTextures.cpp
--- a/Engine/Source/Runtime/Renderer/Private/SystemTextures.cpp
+++ b/Engine/Source/Runtime/Renderer/Private/SystemTextures.cpp
@@ -1461,7 +1461,7 @@ FRDGTextureRef GetInternalDefaultTexture(
 // Default Buffers
 
 template<typename T>
-static FDefaultBufferKey GetDefaultBufferKey(uint32 NumBytePerElement, bool bIsStructuredBuffer, const T* In)
+static FDefaultBufferKey GetDefaultBufferKey(uint32 NumBytePerElement, EDefaultBufferType BufferType, const T* In)
 {
 	FDefaultBufferKey Out;
 	if (In)
@@ -1475,14 +1475,14 @@ static FDefaultBufferKey GetDefaultBufferKey(uint32 NumBytePerElement, bool bIsS
 	}
 
 	Out.NumBytePerElement	= NumBytePerElement;
-	Out.bIsStructuredBuffer = bIsStructuredBuffer;
+	Out.BufferType			= BufferType;
 	return Out;
 }
 
 bool operator !=(const FDefaultBufferKey& A, const FDefaultBufferKey& B)
 {
 	return A.NumBytePerElement != B.NumBytePerElement ||
-		A.bIsStructuredBuffer != B.bIsStructuredBuffer ||
+		A.BufferType != B.BufferType ||
 		A.ValueAsUInt[0] != B.ValueAsUInt[0] ||
 		A.ValueAsUInt[1] != B.ValueAsUInt[1] ||
 		A.ValueAsUInt[2] != B.ValueAsUInt[2] ||
@@ -1495,15 +1495,15 @@ FRDGBufferRef GetInternalDefaultBuffer(
 	TArray<FDefaultBuffer>& DefaultBuffers,
 	FHashTable& HashDefaultBuffers,
 	uint32 NumBytePerElement,
-	bool bIsStructuredBuffer,
+	EDefaultBufferType BufferType,
 	const TClearValue* Value)
 {
 	// Buffer key
 	const uint32 NumElements = 1;
-	const FDefaultBufferKey Key = GetDefaultBufferKey(NumBytePerElement, bIsStructuredBuffer, Value);
-	const uint32 Hash = Murmur32({uint32(Key.bIsStructuredBuffer ? 0x20000000u : 0x10000000u) | Key.NumBytePerElement, Key.ValueAsUInt[0], Key.ValueAsUInt[1], Key.ValueAsUInt[2], Key.ValueAsUInt[3] });
+	const FDefaultBufferKey Key = GetDefaultBufferKey(NumBytePerElement, BufferType, Value);
+	const uint32 Hash = Murmur32({(uint32)BufferType, Key.NumBytePerElement, Key.ValueAsUInt[0], Key.ValueAsUInt[1], Key.ValueAsUInt[2], Key.ValueAsUInt[3] });
 
-	// Find exsting buffer ("fast" path)
+	// Find existing buffer ("fast" path)
 	uint32 Index = HashDefaultBuffers.First(Hash);
 	while (HashDefaultBuffers.IsValid(Index) && DefaultBuffers[Index].Key != Key)
 	{
@@ -1523,22 +1523,26 @@ FRDGBufferRef GetInternalDefaultBuffer(
 
 	// Adding new buffer if there is no fit (slow path)
 	TRefCountPtr<FRHIBuffer> RHIBuffer;
-	if (bIsStructuredBuffer)
+	switch (BufferType)
 	{
-		CreateInfo.DebugName = TEXT("DefaultStructuredBuffer");
+	case EDefaultBufferType::VertexBuffer:
+		CreateInfo.DebugName = TEXT("DefaultBuffer");
+		BufferDesc = FRDGBufferDesc::CreateUploadDesc(NumBytePerElement, NumElements);
+		RHIBuffer = RHICreateVertexBuffer(BufferSize, BUF_Static | BUF_ShaderResource, CreateInfo);
+		break;
 
+	case EDefaultBufferType::StructuredBuffer:
+		CreateInfo.DebugName = TEXT("DefaultStructuredBuffer");
 		BufferDesc = FRDGBufferDesc::CreateStructuredDesc(NumBytePerElement, NumElements);
-		BufferDesc.Usage = BUF_Static | BUF_ShaderResource;
-
 		RHIBuffer = RHICreateStructuredBuffer(NumBytePerElement, BufferSize, BufferDesc.Usage, CreateInfo);
-	}
-	else
-	{
-		CreateInfo.DebugName = TEXT("DefaultBuffer");
+		break;
 
-		BufferDesc = FRDGBufferDesc::CreateUploadDesc(NumBytePerElement, NumElements);
+	case EDefaultBufferType::ByteAddressBuffer:
+		CreateInfo.DebugName = TEXT("DefaultByteAddressBuffer");
+		BufferDesc = FRDGBufferDesc::CreateByteAddressDesc(BufferSize);
+		RHIBuffer = RHICreateStructuredBuffer(NumBytePerElement, BufferSize, BufferDesc.Usage, CreateInfo);
+		break;
 
-		RHIBuffer = RHICreateVertexBuffer(BufferSize, BUF_Static | BUF_ShaderResource, CreateInfo);
 	}
 
 	uint8* DestPtr = static_cast<uint8*>(GraphBuilder.RHICmdList.LockBuffer(RHIBuffer, 0, BufferSize, RLM_WriteOnly));
@@ -1607,18 +1611,19 @@ FRDGTextureRef FSystemTextures::GetDefaultTexture(FRDGBuilder& GraphBuilder, ETe
 // Buffers
 
 // Default init to 0
-FRDGBufferRef FSystemTextures::GetDefaultBuffer(FRDGBuilder& GraphBuilder, uint32 NumBytePerElement)										{ return GetInternalDefaultBuffer(GraphBuilder, DefaultBuffers, HashDefaultBuffers, NumBytePerElement, false, (uint32*)nullptr); }
-FRDGBufferRef FSystemTextures::GetDefaultStructuredBuffer(FRDGBuilder& GraphBuilder, uint32 NumBytePerElement)								{ return GetInternalDefaultBuffer(GraphBuilder, DefaultBuffers, HashDefaultBuffers, NumBytePerElement, true,  (uint32*)nullptr); }
+FRDGBufferRef FSystemTextures::GetDefaultBuffer(FRDGBuilder& GraphBuilder, uint32 NumBytePerElement)										{ return GetInternalDefaultBuffer(GraphBuilder, DefaultBuffers, HashDefaultBuffers, NumBytePerElement, EDefaultBufferType::VertexBuffer, (uint32*)nullptr); }
+FRDGBufferRef FSystemTextures::GetDefaultStructuredBuffer(FRDGBuilder& GraphBuilder, uint32 NumBytePerElement)								{ return GetInternalDefaultBuffer(GraphBuilder, DefaultBuffers, HashDefaultBuffers, NumBytePerElement, EDefaultBufferType::StructuredBuffer, (uint32*)nullptr); }
+FRDGBufferRef FSystemTextures::GetDefaultByteAddressBuffer(FRDGBuilder& GraphBuilder, uint32 NumBytePerElement)								{ return GetInternalDefaultBuffer(GraphBuilder, DefaultBuffers, HashDefaultBuffers, NumBytePerElement, EDefaultBufferType::ByteAddressBuffer, (uint32*)nullptr); }
 
 // Default value of an element
-FRDGBufferRef FSystemTextures::GetDefaultBuffer(FRDGBuilder& GraphBuilder, uint32 NumBytePerElement, float Value)							{ return GetInternalDefaultBuffer(GraphBuilder, DefaultBuffers, HashDefaultBuffers, NumBytePerElement, false/* Vertex */, &Value); }
-FRDGBufferRef FSystemTextures::GetDefaultBuffer(FRDGBuilder& GraphBuilder, uint32 NumBytePerElement, uint32 Value)							{ return GetInternalDefaultBuffer(GraphBuilder, DefaultBuffers, HashDefaultBuffers, NumBytePerElement, false/* Vertex */, &Value); }
-FRDGBufferRef FSystemTextures::GetDefaultBuffer(FRDGBuilder& GraphBuilder, uint32 NumBytePerElement, const FVector3f& Value)					{ return GetInternalDefaultBuffer(GraphBuilder, DefaultBuffers, HashDefaultBuffers, NumBytePerElement, false/* Vertex */, &Value); }
-FRDGBufferRef FSystemTextures::GetDefaultBuffer(FRDGBuilder& GraphBuilder, uint32 NumBytePerElement, const FVector4f& Value)					{ return GetInternalDefaultBuffer(GraphBuilder, DefaultBuffers, HashDefaultBuffers, NumBytePerElement, false/* Vertex */, &Value); }
-FRDGBufferRef FSystemTextures::GetDefaultBuffer(FRDGBuilder& GraphBuilder, uint32 NumBytePerElement, const FUintVector4& Value)				{ return GetInternalDefaultBuffer(GraphBuilder, DefaultBuffers, HashDefaultBuffers, NumBytePerElement, false/* Vertex */, &Value); }
-
-FRDGBufferRef FSystemTextures::GetDefaultStructuredBuffer(FRDGBuilder& GraphBuilder, uint32 NumBytePerElement, float Value)					{ return GetInternalDefaultBuffer(GraphBuilder, DefaultBuffers, HashDefaultBuffers, NumBytePerElement, true /* Structured */, &Value); }
-FRDGBufferRef FSystemTextures::GetDefaultStructuredBuffer(FRDGBuilder& GraphBuilder, uint32 NumBytePerElement, uint32 Value)				{ return GetInternalDefaultBuffer(GraphBuilder, DefaultBuffers, HashDefaultBuffers, NumBytePerElement, true /* Structured */, &Value); }
-FRDGBufferRef FSystemTextures::GetDefaultStructuredBuffer(FRDGBuilder& GraphBuilder, uint32 NumBytePerElement, const FVector3f& Value)		{ return GetInternalDefaultBuffer(GraphBuilder, DefaultBuffers, HashDefaultBuffers, NumBytePerElement, true /* Structured */, &Value); }
-FRDGBufferRef FSystemTextures::GetDefaultStructuredBuffer(FRDGBuilder& GraphBuilder, uint32 NumBytePerElement, const FVector4f& Value)		{ return GetInternalDefaultBuffer(GraphBuilder, DefaultBuffers, HashDefaultBuffers, NumBytePerElement, true /* Structured */, &Value); }
-FRDGBufferRef FSystemTextures::GetDefaultStructuredBuffer(FRDGBuilder& GraphBuilder, uint32 NumBytePerElement, const FUintVector4& Value)	{ return GetInternalDefaultBuffer(GraphBuilder, DefaultBuffers, HashDefaultBuffers, NumBytePerElement, true /* Structured */, &Value); }
+FRDGBufferRef FSystemTextures::GetDefaultBuffer(FRDGBuilder& GraphBuilder, uint32 NumBytePerElement, float Value)							{ return GetInternalDefaultBuffer(GraphBuilder, DefaultBuffers, HashDefaultBuffers, NumBytePerElement, EDefaultBufferType::VertexBuffer, &Value); }
+FRDGBufferRef FSystemTextures::GetDefaultBuffer(FRDGBuilder& GraphBuilder, uint32 NumBytePerElement, uint32 Value)							{ return GetInternalDefaultBuffer(GraphBuilder, DefaultBuffers, HashDefaultBuffers, NumBytePerElement, EDefaultBufferType::VertexBuffer, &Value); }
+FRDGBufferRef FSystemTextures::GetDefaultBuffer(FRDGBuilder& GraphBuilder, uint32 NumBytePerElement, const FVector3f& Value)				{ return GetInternalDefaultBuffer(GraphBuilder, DefaultBuffers, HashDefaultBuffers, NumBytePerElement, EDefaultBufferType::VertexBuffer, &Value); }
+FRDGBufferRef FSystemTextures::GetDefaultBuffer(FRDGBuilder& GraphBuilder, uint32 NumBytePerElement, const FVector4f& Value)				{ return GetInternalDefaultBuffer(GraphBuilder, DefaultBuffers, HashDefaultBuffers, NumBytePerElement, EDefaultBufferType::VertexBuffer, &Value); }
+FRDGBufferRef FSystemTextures::GetDefaultBuffer(FRDGBuilder& GraphBuilder, uint32 NumBytePerElement, const FUintVector4& Value)				{ return GetInternalDefaultBuffer(GraphBuilder, DefaultBuffers, HashDefaultBuffers, NumBytePerElement, EDefaultBufferType::VertexBuffer, &Value); }
+
+FRDGBufferRef FSystemTextures::GetDefaultStructuredBuffer(FRDGBuilder& GraphBuilder, uint32 NumBytePerElement, float Value)					{ return GetInternalDefaultBuffer(GraphBuilder, DefaultBuffers, HashDefaultBuffers, NumBytePerElement, EDefaultBufferType::StructuredBuffer, &Value); }
+FRDGBufferRef FSystemTextures::GetDefaultStructuredBuffer(FRDGBuilder& GraphBuilder, uint32 NumBytePerElement, uint32 Value)				{ return GetInternalDefaultBuffer(GraphBuilder, DefaultBuffers, HashDefaultBuffers, NumBytePerElement, EDefaultBufferType::StructuredBuffer, &Value); }
+FRDGBufferRef FSystemTextures::GetDefaultStructuredBuffer(FRDGBuilder& GraphBuilder, uint32 NumBytePerElement, const FVector3f& Value)		{ return GetInternalDefaultBuffer(GraphBuilder, DefaultBuffers, HashDefaultBuffers, NumBytePerElement, EDefaultBufferType::StructuredBuffer, &Value); }
+FRDGBufferRef FSystemTextures::GetDefaultStructuredBuffer(FRDGBuilder& GraphBuilder, uint32 NumBytePerElement, const FVector4f& Value)		{ return GetInternalDefaultBuffer(GraphBuilder, DefaultBuffers, HashDefaultBuffers, NumBytePerElement, EDefaultBufferType::StructuredBuffer, &Value); }
+FRDGBufferRef FSystemTextures::GetDefaultStructuredBuffer(FRDGBuilder& GraphBuilder, uint32 NumBytePerElement, const FUintVector4& Value)	{ return GetInternalDefaultBuffer(GraphBuilder, DefaultBuffers, HashDefaultBuffers, NumBytePerElement, EDefaultBufferType::StructuredBuffer, &Value); }
diff --git a/Engine/Source/Runtime/Renderer/Private/VirtualShadowMaps/VirtualShadowMapArray.cpp b/Engine/Source/Runtime/Renderer/Private/VirtualShadowMaps/VirtualShadowMapArray.cpp
--- a/Engine/Source/Runtime/Renderer/Private/VirtualShadowMaps/VirtualShadowMapArray.cpp
+++ b/Engine/Source/Runtime/Renderer/Private/VirtualShadowMaps/VirtualShadowMapArray.cpp
@@ -463,7 +463,7 @@ void FVirtualShadowMapArray::Initialize(FRDGBuilder& GraphBuilder, FVirtualShado
 	// Reference dummy data in the UB initially
 	const uint32 DummyPageTableElement = 0xFFFFFFFF;
 	UniformParameters.PageTable = GraphBuilder.CreateSRV(GSystemTextures.GetDefaultStructuredBuffer(GraphBuilder, sizeof(DummyPageTableElement), DummyPageTableElement));
-	UniformParameters.ProjectionData = GraphBuilder.CreateSRV(GSystemTextures.GetDefaultStructuredBuffer(GraphBuilder, sizeof(FVirtualShadowMapProjectionShaderData)));
+	UniformParameters.ProjectionData = GraphBuilder.CreateSRV(GSystemTextures.GetDefaultByteAddressBuffer(GraphBuilder, sizeof(FVirtualShadowMapProjectionShaderData)));
 	UniformParameters.PageFlags = GraphBuilder.CreateSRV(GSystemTextures.GetDefaultStructuredBuffer(GraphBuilder, sizeof(uint32)));
 	UniformParameters.PageRectBounds = GraphBuilder.CreateSRV(GSystemTextures.GetDefaultStructuredBuffer(GraphBuilder, sizeof(FIntVector4)));
 
diff --git a/Engine/Source/Runtime/Renderer/Public/SystemTextures.h b/Engine/Source/Runtime/Renderer/Public/SystemTextures.h
--- a/Engine/Source/Runtime/Renderer/Public/SystemTextures.h
+++ b/Engine/Source/Runtime/Renderer/Public/SystemTextures.h
@@ -54,11 +54,18 @@ struct FDefaultTexture
 	TRefCountPtr<IPooledRenderTarget> Texture;
 };
 
+enum class EDefaultBufferType
+{
+	VertexBuffer,
+	StructuredBuffer,
+	ByteAddressBuffer
+};
+
 struct FDefaultBufferKey
 {
 	uint32 ValueAsUInt[4] = { 0u,0u,0u,0u };
 	uint32 NumBytePerElement = 0;
-	bool bIsStructuredBuffer = false;
+	EDefaultBufferType BufferType = EDefaultBufferType::VertexBuffer;
 };
 struct FDefaultBuffer
 {
@@ -200,6 +207,7 @@ public:
 	// Create default buffer initialize to zero.
 	FRDGBufferRef RENDERER_API GetDefaultBuffer(FRDGBuilder& GraphBuilder, uint32 NumBytePerElement);
 	FRDGBufferRef RENDERER_API GetDefaultStructuredBuffer(FRDGBuilder& GraphBuilder, uint32 NumBytePerElement);
+	FRDGBufferRef RENDERER_API GetDefaultByteAddressBuffer(FRDGBuilder& GraphBuilder, uint32 NumBytePerElement);
 
 	template <typename T>
 	FRDGBufferRef GetDefaultBuffer(FRDGBuilder& GraphBuilder)
