Enable "r.Shadow.Virtual.AllocatePagePoolAsReservedResource" by default

diff --git a/Engine/Source/Runtime/Renderer/Private/VirtualShadowMaps/VirtualShadowMapCacheManager.cpp b/Engine/Source/Runtime/Renderer/Private/VirtualShadowMaps/VirtualShadowMapCacheManager.cpp
--- a/Engine/Source/Runtime/Renderer/Private/VirtualShadowMaps/VirtualShadowMapCacheManager.cpp
+++ b/Engine/Source/Runtime/Renderer/Private/VirtualShadowMaps/VirtualShadowMapCacheManager.cpp
@@ -73,7 +73,7 @@ static FAutoConsoleVariableRef  CVarForceInvalidateDirectionalVSM(
 
 static TAutoConsoleVariable<int32> CVarVSMReservedResource(
 	TEXT("r.Shadow.Virtual.AllocatePagePoolAsReservedResource"),
-	0,
+	1,
 	TEXT("Allocate VSM page pool as a reserved/virtual texture, backed by N small physical memory allocations to reduce fragmentation."),
 	ECVF_RenderThreadSafe
 );
