MGPU:  Fix for atmosphere related scan line artifacts for Path Tracer.  Needed GPU mask to generate atmosphere lookup texture on all GPUs.

#rnx

diff --git a/Engine/Source/Runtime/Renderer/Private/PathTracing.cpp b/Engine/Source/Runtime/Renderer/Private/PathTracing.cpp
--- a/Engine/Source/Runtime/Renderer/Private/PathTracing.cpp
+++ b/Engine/Source/Runtime/Renderer/Private/PathTracing.cpp
@@ -2187,6 +2187,7 @@ void FDeferredShadingSceneRenderer::RenderPathTracing(
 		FAtmosphereConfig AtmoConfig(*Scene->GetSkyAtmosphereSceneInfo()->GetAtmosphereShaderParameters());
 		if (!PathTracingState->AtmosphereOpticalDepthLUT.IsValid() || PathTracingState->LastAtmosphereConfig.IsDifferent(AtmoConfig))
 		{
+			RDG_GPU_MASK_SCOPE(GraphBuilder, FRHIGPUMask::All());
 			PathTracingState->LastAtmosphereConfig = AtmoConfig;
 			// need to create a new LUT
 			FRDGTextureDesc Desc = FRDGTextureDesc::Create2D(
