Fix hair scene scattering not running when the scene does not have a sky component.

diff --git a/Engine/Source/Runtime/Renderer/Private/HairStrands/HairStrandsEnvironment.cpp b/Engine/Source/Runtime/Renderer/Private/HairStrands/HairStrandsEnvironment.cpp
--- a/Engine/Source/Runtime/Renderer/Private/HairStrands/HairStrandsEnvironment.cpp
+++ b/Engine/Source/Runtime/Renderer/Private/HairStrands/HairStrandsEnvironment.cpp
@@ -356,7 +356,7 @@ static void AddHairStrandsEnvironmentLightingPassPS(
 	if (bHasStaticLighting) { check(LightingType != EHairLightingSourceType::Lumen); }
 
 	// Early out if there is no static lighting, no sky lighting, nor reflection probes, nor Lumen
-	if (!bSkyLight && !bDynamicSkyLight && !bHasStaticLighting && !bReflectionEnv && !bLumenActive)
+	if (!bSkyLight && !bDynamicSkyLight && !bHasStaticLighting && !bReflectionEnv && !bLumenActive && IntegrationType != EHairLightingIntegrationType::SceneColor)
 		return;
 
 	FSceneTextureParameters SceneTextures = GetSceneTextureParameters(GraphBuilder, View);
