Tonemapping LUT pass only when needed

diff --git a/Engine/Source/Runtime/Renderer/Private/PostProcess/PostProcessCombineLUTs.cpp b/Engine/Source/Runtime/Renderer/Private/PostProcess/PostProcessCombineLUTs.cpp
--- a/Engine/Source/Runtime/Renderer/Private/PostProcess/PostProcessCombineLUTs.cpp
+++ b/Engine/Source/Runtime/Renderer/Private/PostProcess/PostProcessCombineLUTs.cpp
@@ -44,6 +44,11 @@ TAutoConsoleVariable<int32> CVarColorGrading(
 	TEXT("Controls whether post process settings's color grading settings should be applied."),
 	ECVF_RenderThreadSafe);
 
+TAutoConsoleVariable<int32> CVarLUTUpdateEveryFrame(
+	TEXT("r.LUT.UpdateEveryFrame"), 0,
+	TEXT("Controls whether the tonemapping LUT pass is executed every frame."),
+	ECVF_RenderThreadSafe);
+
 // Including the neutral one at index 0
 const uint32 GMaxLUTBlendCount = 5;
 
@@ -76,7 +81,7 @@ bool PipelineVolumeTextureLUTSupportGuaranteedAtRuntime(EShaderPlatform Platform
 	return RHIVolumeTextureRenderingSupportGuaranteed(Platform) && (RHISupportsGeometryShaders(Platform) || RHISupportsVertexShaderLayer(Platform));
 }
 
-FColorRemapParameters GetColorRemapParameters()
+FVector3f GetMappingPolynomial()
 {
 	FColorTransform ColorTransform;
 	ColorTransform.MinValue = FMath::Clamp(CVarColorMin.GetValueOnRenderThread(), -10.0f, 10.0f);
@@ -89,11 +94,17 @@ FColorRemapParameters GetColorRemapParameters()
 	float b = 4 * ColorTransform.MidValue - 3 * ColorTransform.MinValue - ColorTransform.MaxValue;
 	float a = ColorTransform.MaxValue - ColorTransform.MinValue - b;
 
+	return FVector3f(a, b, c);
+}
+
+FColorRemapParameters GetColorRemapParameters()
+{
 	FColorRemapParameters Parameters;
-	Parameters.MappingPolynomial = FVector3f(a, b, c);
+	Parameters.MappingPolynomial = GetMappingPolynomial();
 	return Parameters;
 }
 
+
 BEGIN_SHADER_PARAMETER_STRUCT(FACESTonemapShaderParameters, )
 	SHADER_PARAMETER(FVector4f, ACESMinMaxData) // xy = min ACES/luminance, zw = max ACES/luminance
 	SHADER_PARAMETER(FVector4f, ACESMidData) // x = mid ACES, y = mid luminance, z = mid slope
@@ -153,93 +164,142 @@ BEGIN_SHADER_PARAMETER_STRUCT(FCombineLUTParameters, )
 	SHADER_PARAMETER_STRUCT_INCLUDE(FTonemapperOutputDeviceParameters, OutputDevice)
 END_SHADER_PARAMETER_STRUCT()
 
-void GetCombineLUTParameters(
-	FCombineLUTParameters& Parameters,
-	const FViewInfo& View,
-	const FTexture* const* Textures,
-	const float* Weights,
-	uint32 BlendCount)
+#define UPDATE_CACHE_SETTINGS(DestParameters, ParamValue, bOutHasChanged) \
+if(DestParameters != (ParamValue)) \
+{ \
+	DestParameters = (ParamValue); \
+	bOutHasChanged = true; \
+}
+
+struct FCachedLUTSettings
 {
-	check(Textures);
-	check(Weights);
+	uint32 UniqueID = 0;
+	EShaderPlatform ShaderPlatform = GMaxRHIShaderPlatform;
+	FCombineLUTParameters Parameters;
+	FWorkingColorSpaceShaderParameters WorkingColorSpaceShaderParameters;
 
-	static const FPostProcessSettings DefaultSettings;
+	bool UpdateCachedValues(const FViewInfo& View, const FTexture* const* Textures, const float* Weights, uint32 BlendCount)
+	{
+		bool bHasChanged = false;
+		GetCombineLUTParameters(View, Textures, Weights, BlendCount, bHasChanged);
+		UPDATE_CACHE_SETTINGS(UniqueID, View.State ? View.State->GetViewKey() : 0, bHasChanged);
+		UPDATE_CACHE_SETTINGS(ShaderPlatform, View.GetShaderPlatform(), bHasChanged);
 
-	const FSceneViewFamily& ViewFamily = *(View.Family);
-	const FPostProcessSettings& Settings = (
-		ViewFamily.EngineShowFlags.ColorGrading && CVarColorGrading.GetValueOnRenderThread() != 0)
-		? View.FinalPostProcessSettings
-		: DefaultSettings;
+		const FWorkingColorSpaceShaderParameters* InWorkingColorSpaceShaderParameters = reinterpret_cast<const FWorkingColorSpaceShaderParameters*>(GDefaultWorkingColorSpaceUniformBuffer.GetContents());
+		if (InWorkingColorSpaceShaderParameters)
+		{
+			UPDATE_CACHE_SETTINGS(WorkingColorSpaceShaderParameters.ToXYZ, InWorkingColorSpaceShaderParameters->ToXYZ, bHasChanged);
+			UPDATE_CACHE_SETTINGS(WorkingColorSpaceShaderParameters.FromXYZ, InWorkingColorSpaceShaderParameters->FromXYZ, bHasChanged);
+			UPDATE_CACHE_SETTINGS(WorkingColorSpaceShaderParameters.ToAP1, InWorkingColorSpaceShaderParameters->ToAP1, bHasChanged);
+			UPDATE_CACHE_SETTINGS(WorkingColorSpaceShaderParameters.FromAP1, InWorkingColorSpaceShaderParameters->FromAP1, bHasChanged);
+			UPDATE_CACHE_SETTINGS(WorkingColorSpaceShaderParameters.ToAP0, InWorkingColorSpaceShaderParameters->ToAP0, bHasChanged);
+			UPDATE_CACHE_SETTINGS(WorkingColorSpaceShaderParameters.bIsSRGB, InWorkingColorSpaceShaderParameters->bIsSRGB, bHasChanged);
+		}
 
-	for (uint32 BlendIndex = 0; BlendIndex < BlendCount; ++BlendIndex)
+		return bHasChanged;
+	}
+	
+	void GetCombineLUTParameters(
+		const FViewInfo& View,
+		const FTexture* const* Textures,
+		const float* Weights,
+		uint32 BlendCount,
+		bool& bHasChanged)
 	{
-		// Neutral texture occupies the first slot and doesn't actually need to be set.
-		if (BlendIndex != 0)
+		check(Textures);
+		check(Weights);
+
+		static const FPostProcessSettings DefaultSettings;
+
+		const FSceneViewFamily& ViewFamily = *(View.Family);
+		const FPostProcessSettings& Settings = (
+			ViewFamily.EngineShowFlags.ColorGrading && CVarColorGrading.GetValueOnRenderThread() != 0)
+			? View.FinalPostProcessSettings
+			: DefaultSettings;
+
+		for (uint32 BlendIndex = 0; BlendIndex < BlendCount; ++BlendIndex)
 		{
-			check(Textures[BlendIndex]);
+			// Neutral texture occupies the first slot and doesn't actually need to be set.
+			if (BlendIndex != 0)
+			{
+				check(Textures[BlendIndex]);
+
+				// Don't use texture asset sampler as it might have anisotropic filtering enabled
+				UPDATE_CACHE_SETTINGS(Parameters.Textures[BlendIndex], Textures[BlendIndex]->TextureRHI, bHasChanged);
+				Parameters.Samplers[BlendIndex] = TStaticSamplerState<SF_Bilinear, AM_Clamp, AM_Clamp, AM_Clamp, 0, 1>::GetRHI();
+			}
 
-			// Don't use texture asset sampler as it might have anisotropic filtering enabled
-			Parameters.Textures[BlendIndex] = Textures[BlendIndex]->TextureRHI;
-			Parameters.Samplers[BlendIndex] = TStaticSamplerState<SF_Bilinear, AM_Clamp, AM_Clamp, AM_Clamp, 0, 1>::GetRHI();
+			UPDATE_CACHE_SETTINGS(GET_SCALAR_ARRAY_ELEMENT(Parameters.LUTWeights, BlendIndex), Weights[BlendIndex], bHasChanged);
 		}
 
-		GET_SCALAR_ARRAY_ELEMENT(Parameters.LUTWeights, BlendIndex) = Weights[BlendIndex];
+		Parameters.WorkingColorSpace = GDefaultWorkingColorSpaceUniformBuffer.GetUniformBufferRef();
+
+		FACESTonemapParams TonemapperParams;
+		GetACESTonemapParameters(TonemapperParams);
+		UPDATE_CACHE_SETTINGS(Parameters.ACESTonemapParameters.ACESMinMaxData, TonemapperParams.ACESMinMaxData, bHasChanged);
+		UPDATE_CACHE_SETTINGS(Parameters.ACESTonemapParameters.ACESMidData, TonemapperParams.ACESMidData, bHasChanged);
+		UPDATE_CACHE_SETTINGS(Parameters.ACESTonemapParameters.ACESCoefsLow_0, TonemapperParams.ACESCoefsLow_0, bHasChanged);
+		UPDATE_CACHE_SETTINGS(Parameters.ACESTonemapParameters.ACESCoefsHigh_0, TonemapperParams.ACESCoefsHigh_0, bHasChanged);
+		UPDATE_CACHE_SETTINGS(Parameters.ACESTonemapParameters.ACESCoefsLow_4, TonemapperParams.ACESCoefsLow_4, bHasChanged);
+		UPDATE_CACHE_SETTINGS(Parameters.ACESTonemapParameters.ACESCoefsHigh_4, TonemapperParams.ACESCoefsHigh_4, bHasChanged);
+		UPDATE_CACHE_SETTINGS(Parameters.ACESTonemapParameters.ACESSceneColorMultiplier, TonemapperParams.ACESSceneColorMultiplier, bHasChanged);
+
+		UPDATE_CACHE_SETTINGS(Parameters.ColorScale, FVector4f(View.ColorScale), bHasChanged);
+		UPDATE_CACHE_SETTINGS(Parameters.OverlayColor, (View.OverlayColor), bHasChanged);
+		UPDATE_CACHE_SETTINGS(Parameters.ColorRemap.MappingPolynomial, GetMappingPolynomial(), bHasChanged);
+
+		// White balance
+		UPDATE_CACHE_SETTINGS(Parameters.bIsTemperatureWhiteBalance, uint32(Settings.TemperatureType == ETemperatureMethod::TEMP_WhiteBalance), bHasChanged);
+		UPDATE_CACHE_SETTINGS(Parameters.LUTSize, (float)CVarLUTSize->GetInt(), bHasChanged);
+		UPDATE_CACHE_SETTINGS(Parameters.WhiteTemp, Settings.WhiteTemp, bHasChanged);
+		UPDATE_CACHE_SETTINGS(Parameters.WhiteTint, Settings.WhiteTint, bHasChanged);
+
+		// Color grade
+		UPDATE_CACHE_SETTINGS(Parameters.ColorSaturation, FVector4f(Settings.ColorSaturation), bHasChanged);
+		UPDATE_CACHE_SETTINGS(Parameters.ColorContrast, FVector4f(Settings.ColorContrast), bHasChanged);
+		UPDATE_CACHE_SETTINGS(Parameters.ColorGamma, FVector4f(Settings.ColorGamma), bHasChanged);
+		UPDATE_CACHE_SETTINGS(Parameters.ColorGain, FVector4f(Settings.ColorGain), bHasChanged);
+		UPDATE_CACHE_SETTINGS(Parameters.ColorOffset, FVector4f(Settings.ColorOffset), bHasChanged);
+
+		UPDATE_CACHE_SETTINGS(Parameters.ColorSaturationShadows, FVector4f(Settings.ColorSaturationShadows), bHasChanged);
+		UPDATE_CACHE_SETTINGS(Parameters.ColorContrastShadows, FVector4f(Settings.ColorContrastShadows), bHasChanged);
+		UPDATE_CACHE_SETTINGS(Parameters.ColorGammaShadows, FVector4f(Settings.ColorGammaShadows), bHasChanged);
+		UPDATE_CACHE_SETTINGS(Parameters.ColorGainShadows, FVector4f(Settings.ColorGainShadows), bHasChanged);
+		UPDATE_CACHE_SETTINGS(Parameters.ColorOffsetShadows, FVector4f(Settings.ColorOffsetShadows), bHasChanged);
+
+		UPDATE_CACHE_SETTINGS(Parameters.ColorSaturationMidtones, FVector4f(Settings.ColorSaturationMidtones), bHasChanged);
+		UPDATE_CACHE_SETTINGS(Parameters.ColorContrastMidtones, FVector4f(Settings.ColorContrastMidtones), bHasChanged);
+		UPDATE_CACHE_SETTINGS(Parameters.ColorGammaMidtones, FVector4f(Settings.ColorGammaMidtones), bHasChanged);
+		UPDATE_CACHE_SETTINGS(Parameters.ColorGainMidtones, FVector4f(Settings.ColorGainMidtones), bHasChanged);
+		UPDATE_CACHE_SETTINGS(Parameters.ColorOffsetMidtones, FVector4f(Settings.ColorOffsetMidtones), bHasChanged);
+
+		UPDATE_CACHE_SETTINGS(Parameters.ColorSaturationHighlights, FVector4f(Settings.ColorSaturationHighlights), bHasChanged);
+		UPDATE_CACHE_SETTINGS(Parameters.ColorContrastHighlights, FVector4f(Settings.ColorContrastHighlights), bHasChanged);
+		UPDATE_CACHE_SETTINGS(Parameters.ColorGammaHighlights, FVector4f(Settings.ColorGammaHighlights), bHasChanged);
+		UPDATE_CACHE_SETTINGS(Parameters.ColorGainHighlights, FVector4f(Settings.ColorGainHighlights), bHasChanged);
+		UPDATE_CACHE_SETTINGS(Parameters.ColorOffsetHighlights, FVector4f(Settings.ColorOffsetHighlights), bHasChanged);
+
+		UPDATE_CACHE_SETTINGS(Parameters.ColorCorrectionShadowsMax, Settings.ColorCorrectionShadowsMax, bHasChanged);
+		UPDATE_CACHE_SETTINGS(Parameters.ColorCorrectionHighlightsMin, Settings.ColorCorrectionHighlightsMin, bHasChanged);
+		UPDATE_CACHE_SETTINGS(Parameters.ColorCorrectionHighlightsMax, Settings.ColorCorrectionHighlightsMax, bHasChanged);
+
+		UPDATE_CACHE_SETTINGS(Parameters.BlueCorrection, Settings.BlueCorrection, bHasChanged);
+		UPDATE_CACHE_SETTINGS(Parameters.ExpandGamut, Settings.ExpandGamut, bHasChanged);
+		UPDATE_CACHE_SETTINGS(Parameters.ToneCurveAmount, Settings.ToneCurveAmount, bHasChanged);
+
+		UPDATE_CACHE_SETTINGS(Parameters.FilmSlope, Settings.FilmSlope, bHasChanged);
+		UPDATE_CACHE_SETTINGS(Parameters.FilmToe, Settings.FilmToe, bHasChanged);
+		UPDATE_CACHE_SETTINGS(Parameters.FilmShoulder, Settings.FilmShoulder, bHasChanged);
+		UPDATE_CACHE_SETTINGS(Parameters.FilmBlackClip, Settings.FilmBlackClip, bHasChanged);
+		UPDATE_CACHE_SETTINGS(Parameters.FilmWhiteClip, Settings.FilmWhiteClip, bHasChanged);
+
+		FTonemapperOutputDeviceParameters TonemapperOutputDeviceParameters = GetTonemapperOutputDeviceParameters(ViewFamily);
+		UPDATE_CACHE_SETTINGS(Parameters.OutputDevice.InverseGamma, TonemapperOutputDeviceParameters.InverseGamma, bHasChanged);
+		UPDATE_CACHE_SETTINGS(Parameters.OutputDevice.OutputDevice, TonemapperOutputDeviceParameters.OutputDevice, bHasChanged);
+		UPDATE_CACHE_SETTINGS(Parameters.OutputDevice.OutputGamut, TonemapperOutputDeviceParameters.OutputGamut, bHasChanged);
+		UPDATE_CACHE_SETTINGS(Parameters.OutputDevice.OutputMaxLuminance, TonemapperOutputDeviceParameters.OutputMaxLuminance, bHasChanged);
 	}
-
-	Parameters.WorkingColorSpace = GDefaultWorkingColorSpaceUniformBuffer.GetUniformBufferRef();
-	GetACESTonemapParameters(Parameters.ACESTonemapParameters);
-
-	Parameters.ColorScale = FVector4f(View.ColorScale);
-	Parameters.OverlayColor = View.OverlayColor;
-	Parameters.ColorRemap = GetColorRemapParameters();
-
-	// White balance
-	Parameters.bIsTemperatureWhiteBalance = Settings.TemperatureType == ETemperatureMethod::TEMP_WhiteBalance;
-	Parameters.LUTSize = (float)CVarLUTSize->GetInt();
-	Parameters.WhiteTemp = Settings.WhiteTemp;
-	Parameters.WhiteTint = Settings.WhiteTint;
-
-	// Color grade
-	Parameters.ColorSaturation = (FVector4f)Settings.ColorSaturation;
-	Parameters.ColorContrast = (FVector4f)Settings.ColorContrast;
-	Parameters.ColorGamma = (FVector4f)Settings.ColorGamma;
-	Parameters.ColorGain = (FVector4f)Settings.ColorGain;
-	Parameters.ColorOffset = (FVector4f)Settings.ColorOffset;
-
-	Parameters.ColorSaturationShadows = (FVector4f)Settings.ColorSaturationShadows;
-	Parameters.ColorContrastShadows = (FVector4f)Settings.ColorContrastShadows;
-	Parameters.ColorGammaShadows = (FVector4f)Settings.ColorGammaShadows;
-	Parameters.ColorGainShadows = (FVector4f)Settings.ColorGainShadows;
-	Parameters.ColorOffsetShadows = (FVector4f)Settings.ColorOffsetShadows;
-
-	Parameters.ColorSaturationMidtones = (FVector4f)Settings.ColorSaturationMidtones;
-	Parameters.ColorContrastMidtones = (FVector4f)Settings.ColorContrastMidtones;
-	Parameters.ColorGammaMidtones = (FVector4f)Settings.ColorGammaMidtones;
-	Parameters.ColorGainMidtones = (FVector4f)Settings.ColorGainMidtones;
-	Parameters.ColorOffsetMidtones = (FVector4f)Settings.ColorOffsetMidtones;
-
-	Parameters.ColorSaturationHighlights = (FVector4f)Settings.ColorSaturationHighlights;
-	Parameters.ColorContrastHighlights = (FVector4f)Settings.ColorContrastHighlights;
-	Parameters.ColorGammaHighlights = (FVector4f)Settings.ColorGammaHighlights;
-	Parameters.ColorGainHighlights = (FVector4f)Settings.ColorGainHighlights;
-	Parameters.ColorOffsetHighlights = (FVector4f)Settings.ColorOffsetHighlights;
-
-	Parameters.ColorCorrectionShadowsMax = Settings.ColorCorrectionShadowsMax;
-	Parameters.ColorCorrectionHighlightsMin = Settings.ColorCorrectionHighlightsMin;
-	Parameters.ColorCorrectionHighlightsMax = Settings.ColorCorrectionHighlightsMax;
-
-	Parameters.BlueCorrection = Settings.BlueCorrection;
-	Parameters.ExpandGamut = Settings.ExpandGamut;
-	Parameters.ToneCurveAmount = Settings.ToneCurveAmount;
-
-	Parameters.FilmSlope = Settings.FilmSlope;
-	Parameters.FilmToe = Settings.FilmToe;
-	Parameters.FilmShoulder = Settings.FilmShoulder;
-	Parameters.FilmBlackClip = Settings.FilmBlackClip;
-	Parameters.FilmWhiteClip = Settings.FilmWhiteClip;
-
-	Parameters.OutputDevice = GetTonemapperOutputDeviceParameters(ViewFamily);
-}
+};
 
 class FLUTBlenderShader : public FGlobalShader
 {
@@ -411,6 +471,8 @@ uint32 GenerateFinalTable(const FFinalPostProcessSettings& Settings, const FText
 
 FRDGTextureRef AddCombineLUTPass(FRDGBuilder& GraphBuilder, const FViewInfo& View)
 {
+	static FCachedLUTSettings CachedLUTSettings;
+
 	const FSceneViewFamily& ViewFamily = *(View.Family);
 	const FTexture* LocalTextures[GMaxLUTBlendCount];
 	float LocalWeights[GMaxLUTBlendCount];
@@ -438,6 +500,12 @@ FRDGTextureRef AddCombineLUTPass(FRDGBuilder& GraphBuilder, const FViewInfo& Vie
 
 	View.SetValidTonemappingLUT();
 
+	const bool bHasChanged = CachedLUTSettings.UpdateCachedValues(View, LocalTextures, LocalWeights, LocalCount);
+	if (!bHasChanged && OutputTexture && CVarLUTUpdateEveryFrame.GetValueOnRenderThread() == 0)
+	{
+		return OutputTexture;
+	}
+
 	// View doesn't support a persistent LUT, so create a temporary one.
 	if (!OutputTexture)
 	{
@@ -460,7 +528,7 @@ FRDGTextureRef AddCombineLUTPass(FRDGBuilder& GraphBuilder, const FViewInfo& Vie
 	if (bUseComputePass)
 	{
 		FLUTBlenderCS::FParameters* PassParameters = GraphBuilder.AllocParameters<FLUTBlenderCS::FParameters>();
-		GetCombineLUTParameters(PassParameters->CombineLUT, View, LocalTextures, LocalWeights, LocalCount);
+		PassParameters->CombineLUT = CachedLUTSettings.Parameters;
 		PassParameters->OutputExtentInverse = FVector2f(1.0f, 1.0f) / FVector2f(OutputViewSize);
 		PassParameters->RWOutputTexture = GraphBuilder.CreateUAV(OutputTexture);
 
@@ -484,7 +552,7 @@ FRDGTextureRef AddCombineLUTPass(FRDGBuilder& GraphBuilder, const FViewInfo& Vie
 	else
 	{
 		FLUTBlenderPS::FParameters* PassParameters = GraphBuilder.AllocParameters<FLUTBlenderPS::FParameters>();
-		GetCombineLUTParameters(PassParameters->CombineLUT, View, LocalTextures, LocalWeights, LocalCount);
+		PassParameters->CombineLUT = CachedLUTSettings.Parameters;
 		PassParameters->RenderTargets[0] = FRenderTargetBinding(OutputTexture, ERenderTargetLoadAction::ENoAction);
 
 		const bool ShouldSkipTemperature = FMath::IsNearlyEqual(PassParameters->CombineLUT.WhiteTemp, DefaultTemperature) && FMath::IsNearlyEqual(PassParameters->CombineLUT.WhiteTint, DefaultTint);
