Fix broken decals in the forward renderer when MSAA is enabled.

The ResolveSceneDepth pass was leaving the stencil plane untouched. Since the resolve target was never cleared (it came from the transient allocator), the resolved stencil buffer contained garbage. Made that pass clear the stencil plane to 0.

#rnx

diff --git a/Engine/Source/Runtime/Renderer/Private/SceneRendering.cpp b/Engine/Source/Runtime/Renderer/Private/SceneRendering.cpp
--- a/Engine/Source/Runtime/Renderer/Private/SceneRendering.cpp
+++ b/Engine/Source/Runtime/Renderer/Private/SceneRendering.cpp
@@ -5446,7 +5446,7 @@ void AddResolveSceneDepthPass(FRDGBuilder& GraphBuilder, const FViewInfo& View,
 		RHICmdList.ApplyCachedRenderTargets(GraphicsPSOInit);
 		GraphicsPSOInit.BlendState = TStaticBlendState<>::GetRHI();
 		GraphicsPSOInit.RasterizerState = TStaticRasterizerState<FM_Solid, CM_None>::GetRHI();
-		GraphicsPSOInit.DepthStencilState = TStaticDepthStencilState<true, CF_Always>::GetRHI();
+		GraphicsPSOInit.DepthStencilState = TStaticDepthStencilState<true, CF_Always, true, CF_Always, SO_Zero, SO_Zero, SO_Zero, true, CF_Always, SO_Zero, SO_Zero, SO_Zero>::GetRHI();
 
 		RHICmdList.SetViewport(0.0f, 0.0f, 0.0f, DepthExtent.X, DepthExtent.Y, 1.0f);
 
