Removing deprecated code from Renderer.

diff --git a/Engine/Source/Runtime/Renderer/Private/SceneTextures.h b/Engine/Source/Runtime/Renderer/Private/SceneTextures.h
--- a/Engine/Source/Runtime/Renderer/Private/SceneTextures.h
+++ b/Engine/Source/Runtime/Renderer/Private/SceneTextures.h
@@ -97,35 +97,3 @@ struct RENDERER_API FSceneTextures : public FMinimalSceneTextures
 
 /** Extracts scene textures into the global extraction instance. */
 void QueueSceneTextureExtractions(FRDGBuilder& GraphBuilder, const FSceneTextures& SceneTextures);
-
-/** Utility functions for accessing common global scene texture configuration state. Reads a bit less awkwardly than the singleton access. */
-
-UE_DEPRECATED(5.1, "Single pass multiple view family rendering makes this obsolete.  Use ViewFamily.SceneTexturesConfig.NumSamples instead.")
-inline uint32 GetSceneTextureNumSamples()
-{
-	return FSceneTexturesConfig::Get().NumSamples;
-}
-
-UE_DEPRECATED(5.1, "Single pass multiple view family rendering makes this obsolete.  Use ViewFamily.SceneTexturesConfig.EditorPrimitiveNumSamples instead.")
-inline uint32 GetEditorPrimitiveNumSamples()
-{
-	return FSceneTexturesConfig::Get().EditorPrimitiveNumSamples;
-}
-
-UE_DEPRECATED(5.1, "Single pass multiple view family rendering makes this obsolete.  Use ViewFamily.SceneTexturesConfig.DepthClearValue instead.")
-inline FClearValueBinding GetSceneDepthClearValue()
-{
-	return FSceneTexturesConfig::Get().DepthClearValue;
-}
-
-UE_DEPRECATED(5.1, "Single pass multiple view family rendering makes this obsolete.  Use ViewFamily.SceneTexturesConfig.ColorClearValue instead.")
-inline FClearValueBinding GetSceneColorClearValue()
-{
-	return FSceneTexturesConfig::Get().ColorClearValue;
-}
-
-UE_DEPRECATED(5.1, "Single pass multiple view family rendering makes this obsolete.  Use ViewFamily.SceneTexturesConfig.ColorFormat instead.")
-inline EPixelFormat GetSceneColorFormat()
-{
-	return FSceneTexturesConfig::Get().ColorFormat;
-}
diff --git a/Engine/Source/Runtime/Renderer/Public/MeshPassProcessor.h b/Engine/Source/Runtime/Renderer/Public/MeshPassProcessor.h
--- a/Engine/Source/Runtime/Renderer/Public/MeshPassProcessor.h
+++ b/Engine/Source/Runtime/Renderer/Public/MeshPassProcessor.h
@@ -1331,9 +1331,6 @@ enum class EFVisibleMeshDrawCommandFlags : uint8
 	/** If set, the FMaterial::MaterialUsesWorldPositionOffset_RenderThread() indicates that WPO is active for the given material. */
 	MaterialUsesWorldPositionOffset = 1U << 0U,
 
-	/** If set, the FMaterial::MaterialModifiesMeshPosition_RenderThread() indicates that WPO or something similar is active for the given material. */
-	MaterialMayModifyPosition UE_DEPRECATED(5.1, "Use MaterialUsesWorldPositionOffset, MaterialMayModifyPosition is now an alias for this and no longer represents MaterialModifiesMeshPosition_RenderThread().")  =  MaterialUsesWorldPositionOffset,
-
 	/** If set, the mesh draw command supports primitive ID steam (required for dynamic instancing and GPU-Scene instance culling). */
 	HasPrimitiveIdStreamIndex = 1U << 1U,
 
@@ -1826,32 +1823,9 @@ ENUM_CLASS_FLAGS(EMeshPassFeatures);
 struct FMeshPassProcessorRenderState
 {
 	FMeshPassProcessorRenderState() = default;
-
 	FMeshPassProcessorRenderState(const FMeshPassProcessorRenderState& DrawRenderState) = default;
-
 	~FMeshPassProcessorRenderState() = default;
 
-	UE_DEPRECATED(5.1, "FMeshPassProcessorRenderState with pass / view uniform buffers is deprecated.")
-	FMeshPassProcessorRenderState(
-		const TUniformBufferRef<FViewUniformShaderParameters>& InViewUniformBuffer, 
-		FRHIUniformBuffer* InPassUniformBuffer = nullptr
-		)
-		: ViewUniformBuffer(InViewUniformBuffer)
-		, PassUniformBuffer(InPassUniformBuffer)
-	{
-	}
-
-	UE_DEPRECATED(5.1, "FMeshPassProcessorRenderState with pass / view uniform buffers is deprecated.")
-	FMeshPassProcessorRenderState(
-		const FSceneView& SceneView, 
-		FRHIUniformBuffer* InPassUniformBuffer = nullptr
-		)
-	PRAGMA_DISABLE_DEPRECATION_WARNINGS
-		: FMeshPassProcessorRenderState(SceneView.ViewUniformBuffer, InPassUniformBuffer)
-	PRAGMA_ENABLE_DEPRECATION_WARNINGS
-	{
-	}
-
 public:
 	FORCEINLINE_DEBUGGABLE void SetBlendState(FRHIBlendState* InBlendState)
 	{
@@ -1889,30 +1863,6 @@ public:
 		return DepthStencilAccess;
 	}
 
-	UE_DEPRECATED(5.1, "SetViewUniformBuffer is deprecated. Use View.ViewUniformBuffer and bind on an RDG pass instead.")
-	FORCEINLINE_DEBUGGABLE void SetViewUniformBuffer(const TUniformBufferRef<FViewUniformShaderParameters>& InViewUniformBuffer)
-	{
-		ViewUniformBuffer = InViewUniformBuffer;
-	}
-
-	UE_DEPRECATED(5.1, "GetViewUniformBuffer is deprecated. Use View.ViewUniformBuffer and bind on an RDG pass instead.")
-	FORCEINLINE_DEBUGGABLE const FRHIUniformBuffer* GetViewUniformBuffer() const
-	{
-		return ViewUniformBuffer;
-	}
-
-	UE_DEPRECATED(5.1, "GetNaniteUniformBuffer is deprecated. Use a static uniform buffer and bind on an RDG pass instead.")
-	FORCEINLINE_DEBUGGABLE void SetNaniteUniformBuffer(FRHIUniformBuffer* InNaniteUniformBuffer)
-	{
-		NaniteUniformBuffer = InNaniteUniformBuffer;
-	}
-
-	UE_DEPRECATED(5.1, "GetNaniteUniformBuffer is deprecated. Use a static uniform buffer and bind on an RDG pass instead.")
-	FORCEINLINE_DEBUGGABLE FRHIUniformBuffer* GetNaniteUniformBuffer() const
-	{
-		return NaniteUniformBuffer;
-	}
-
 	FORCEINLINE_DEBUGGABLE uint32 GetStencilRef() const
 	{
 		return StencilRef;
@@ -1994,18 +1944,6 @@ public:
 	RENDERER_API static FMeshDrawingPolicyOverrideSettings ComputeMeshOverrideSettings(const FPSOPrecacheParams& PrecachePSOParams);
 	RENDERER_API static FMeshDrawingPolicyOverrideSettings ComputeMeshOverrideSettings(const FMeshBatch& Mesh);
 
-	UE_DEPRECATED(5.1, "ComputeMeshFillMode with FMeshBatch is deprecated.")
-	static ERasterizerFillMode ComputeMeshFillMode(const FMeshBatch& Mesh, const FMaterial& InMaterialResource, const FMeshDrawingPolicyOverrideSettings& InOverrideSettings)
-	{
-		return ComputeMeshFillMode(InMaterialResource, InOverrideSettings);
-	}
-
-	UE_DEPRECATED(5.1, "ComputeMeshCullMode with FMeshBatch is deprecated.")
-	RENDERER_API static ERasterizerCullMode ComputeMeshCullMode(const FMeshBatch& Mesh, const FMaterial& InMaterialResource, const FMeshDrawingPolicyOverrideSettings& InOverrideSettings)
-	{
-		return ComputeMeshCullMode(InMaterialResource, InOverrideSettings);
-	}
-
 	RENDERER_API static ERasterizerFillMode ComputeMeshFillMode(const FMaterial& InMaterialResource, const FMeshDrawingPolicyOverrideSettings& InOverrideSettings);
 	RENDERER_API static ERasterizerCullMode ComputeMeshCullMode(const FMaterial& InMaterialResource, const FMeshDrawingPolicyOverrideSettings& InOverrideSettings);
 
@@ -2090,15 +2028,6 @@ public:
 		}
 	}
 
-	UE_DEPRECATED(5.1, "GetCreateFunction is deprecated. Use CreateMeshPassProcessor above.")
-	static DeprecatedPassProcessorCreateFunction GetCreateFunction(EShadingPath ShadingPath, EMeshPass::Type PassType)
-	{
-		check(ShadingPath < EShadingPath::Num && PassType < EMeshPass::Num);
-		uint32 ShadingPathIdx = (uint32)ShadingPath;
-		checkf(DeprecatedJumpTable[ShadingPathIdx][PassType], TEXT("Pass type %u create function was never registered for shading path %u.  Use a FRegisterPassProcessorCreateFunction to register a create function for this enum value."), (uint32)PassType, ShadingPathIdx);
-		return DeprecatedJumpTable[ShadingPathIdx][PassType];
-	}
-
 	static EMeshPassFlags GetPassFlags(EShadingPath ShadingPath, EMeshPass::Type PassType)
 	{
 		check(ShadingPath < EShadingPath::Num && PassType < EMeshPass::Num);
@@ -2128,16 +2057,6 @@ public:
 		FPassProcessorManager::Flags[ShadingPathIdx][PassType] = PassFlags;
 	}
 
-	UE_DEPRECATED(5.1, "FRegisterPassProcessorCreateFunction with DeprecatedPassProcessorCreateFunction is deprecated. Use new PassProcessorCreateFunction with extra ERHIFeatureLevel::Type argument.")
-	FRegisterPassProcessorCreateFunction(DeprecatedPassProcessorCreateFunction CreateFunction, EShadingPath InShadingPath, EMeshPass::Type InPassType, EMeshPassFlags PassFlags)
-		: ShadingPath(InShadingPath)
-		, PassType(InPassType)
-	{
-		uint32 ShadingPathIdx = (uint32)ShadingPath;
-		FPassProcessorManager::DeprecatedJumpTable[ShadingPathIdx][PassType] = CreateFunction;
-		FPassProcessorManager::Flags[ShadingPathIdx][PassType] = PassFlags;
-	}
-
 	~FRegisterPassProcessorCreateFunction()
 	{
 		uint32 ShadingPathIdx = (uint32)ShadingPath;
