Change ForwardLocalLightBuffer's stride to use sizeof(float4) because some platforms won't use the stride declared in the shader

#rnx

diff --git a/Engine/Source/Runtime/Renderer/Private/LightGridInjection.cpp b/Engine/Source/Runtime/Renderer/Private/LightGridInjection.cpp
--- a/Engine/Source/Runtime/Renderer/Private/LightGridInjection.cpp
+++ b/Engine/Source/Runtime/Renderer/Private/LightGridInjection.cpp
@@ -700,7 +700,14 @@ FComputeLightGridOutput FSceneRenderer::ComputeLightGrid(FRDGBuilder& GraphBuild
 		// Store off the number of lights before we add a fake entry
 		const int32 NumLocalLightsFinal = ForwardLocalLightData.Num();
 
-		FRDGBufferRef ForwardLocalLightBuffer = CreateStructuredBuffer(GraphBuilder, TEXT("ForwardLocalLightBuffer"), TConstArrayView<FForwardLocalLightData>(ForwardLocalLightData));
+		// Some platforms index the StructuredBuffer in the shader based on the stride specified at buffer creation time, not from the stride specified in the shader.
+		// ForwardLocalLightBuffer is a StructuredBuffer<float4> in the shader, so create the buffer with a stride of sizeof(float4)
+		static_assert(sizeof(FForwardLocalLightData) % sizeof(FVector4f) == 0, "ForwardLocalLightBuffer is used as a StructuredBuffer<float4> in the shader");
+		const uint32 ForwardLocalLightDataSizeNumFloat4 = (NumLocalLightsFinal * sizeof(FForwardLocalLightData)) / sizeof(FVector4f);
+
+		FRDGBufferRef ForwardLocalLightBuffer = CreateStructuredBuffer(GraphBuilder, TEXT("ForwardLocalLightBuffer"),
+			TConstArrayView<FVector4f>(reinterpret_cast<const FVector4f*>(ForwardLocalLightData.GetData()), ForwardLocalLightDataSizeNumFloat4));
+
 		View.ForwardLightingResources.LocalLightVisibleLightInfosIndex = LocalLightVisibleLightInfosIndex;
 
 		const FIntPoint LightGridSizeXY = FIntPoint::DivideAndRoundUp(View.ViewRect.Size(), GLightGridPixelSize);
