Fixed link error when building with non-inlining.

diff --git a/Engine/Source/Runtime/Renderer/Public/PrimitiveSceneInfo.h b/Engine/Source/Runtime/Renderer/Public/PrimitiveSceneInfo.h
--- a/Engine/Source/Runtime/Renderer/Public/PrimitiveSceneInfo.h
+++ b/Engine/Source/Runtime/Renderer/Public/PrimitiveSceneInfo.h
@@ -526,7 +526,7 @@ public:
 	 * This index is only valid until a primitive is added to or removed from
 	 * the scene!
 	 */
-	RENDERER_API FORCEINLINE int32 GetIndex() const { return PackedIndex; }
+	FORCEINLINE int32 GetIndex() const { return PackedIndex; }
 	/** 
 	 * Retrieves the address of the primitives index into in the scene's primitives array.
 	 * This address is only for reference purposes
