Fix crash on opengl due to dither and fade global UBs being created outside of the render thread.

diff --git a/Engine/Source/Runtime/Renderer/Private/HeterogeneousVolumes/HeterogeneousVolumesLiveShadingPipeline.cpp b/Engine/Source/Runtime/Renderer/Private/HeterogeneousVolumes/HeterogeneousVolumesLiveShadingPipeline.cpp
--- a/Engine/Source/Runtime/Renderer/Private/HeterogeneousVolumes/HeterogeneousVolumesLiveShadingPipeline.cpp
+++ b/Engine/Source/Runtime/Renderer/Private/HeterogeneousVolumes/HeterogeneousVolumesLiveShadingPipeline.cpp
@@ -286,8 +286,8 @@ void AddComputePass(
 			FMeshPassProcessorRenderState DrawRenderState;
 
 			FMeshMaterialShaderElementData ShaderElementData;
-			ShaderElementData.FadeUniformBuffer = GGetDistanceCullFadedInUniformBuffer();
-			ShaderElementData.DitherUniformBuffer = GGetDitherFadedInUniformBuffer();
+			ShaderElementData.FadeUniformBuffer = GDistanceCullFadedInUniformBuffer.GetUniformBufferRHI();
+			ShaderElementData.DitherUniformBuffer = GDitherFadedInUniformBuffer.GetUniformBufferRHI();
 
 			FMeshProcessorShaders PassShaders;
 			PassShaders.ComputeShader = ComputeShader;
diff --git a/Engine/Source/Runtime/Renderer/Private/HeterogeneousVolumes/HeterogeneousVolumesMaterialBakingPipeline.cpp b/Engine/Source/Runtime/Renderer/Private/HeterogeneousVolumes/HeterogeneousVolumesMaterialBakingPipeline.cpp
--- a/Engine/Source/Runtime/Renderer/Private/HeterogeneousVolumes/HeterogeneousVolumesMaterialBakingPipeline.cpp
+++ b/Engine/Source/Runtime/Renderer/Private/HeterogeneousVolumes/HeterogeneousVolumesMaterialBakingPipeline.cpp
@@ -183,8 +183,8 @@ void ComputeHeterogeneousVolumeBakeMaterial(
 				FMeshPassProcessorRenderState DrawRenderState;
 
 				FMeshMaterialShaderElementData ShaderElementData;
-				ShaderElementData.FadeUniformBuffer = GGetDistanceCullFadedInUniformBuffer();
-				ShaderElementData.DitherUniformBuffer = GGetDitherFadedInUniformBuffer();
+				ShaderElementData.FadeUniformBuffer = GDistanceCullFadedInUniformBuffer.GetUniformBufferRHI();
+				ShaderElementData.DitherUniformBuffer = GDitherFadedInUniformBuffer.GetUniformBufferRHI();
 
 				FMeshProcessorShaders PassShaders;
 				PassShaders.ComputeShader = ComputeShader;
diff --git a/Engine/Source/Runtime/Renderer/Private/Renderer.cpp b/Engine/Source/Runtime/Renderer/Private/Renderer.cpp
--- a/Engine/Source/Runtime/Renderer/Private/Renderer.cpp
+++ b/Engine/Source/Runtime/Renderer/Private/Renderer.cpp
@@ -77,6 +77,8 @@ void FRendererModule::StartupModule()
 
 	// Needs to run on startup, after static init. Ideally at some point replaced with static accessor similar to GGetDitherFadedInUniformBuffer().
 	GIdentityPrimitiveUniformBuffer.InitContents();
+	GDistanceCullFadedInUniformBuffer.InitContents();
+	GDitherFadedInUniformBuffer.InitContents();
 }
 
 void FRendererModule::ShutdownModule()
diff --git a/Engine/Source/Runtime/Renderer/Private/RendererScene.cpp b/Engine/Source/Runtime/Renderer/Private/RendererScene.cpp
--- a/Engine/Source/Runtime/Renderer/Private/RendererScene.cpp
+++ b/Engine/Source/Runtime/Renderer/Private/RendererScene.cpp
@@ -143,28 +143,11 @@ IMPLEMENT_GLOBAL_SHADER_PARAMETER_STRUCT(FDistanceCullFadeUniformShaderParameter
 
 IMPLEMENT_GLOBAL_SHADER_PARAMETER_STRUCT(FDitherUniformShaderParameters, "PrimitiveDither");
 
-/** Global primitive distance cull fading uniform buffer containing fully faded in. */
-FRHIUniformBuffer* GGetDistanceCullFadedInUniformBuffer()
-{
-	FDistanceCullFadeUniformShaderParameters Parameters;
-	Parameters.FadeTimeScaleBias.X = 0.0f;
-	Parameters.FadeTimeScaleBias.Y = 1.0f;
-	static TUniformBufferRef<FDistanceCullFadeUniformShaderParameters> ResourceRef = 
-		TUniformBufferRef<FDistanceCullFadeUniformShaderParameters>::CreateUniformBufferImmediate(Parameters, UniformBuffer_MultiFrame, EUniformBufferValidation::None);
-
-	return ResourceRef.GetReference();
-}
+/** Global primitive uniform buffer resource containing distance cull faded in */
+TGlobalResource< FGlobalDistanceCullFadeUniformBuffer > GDistanceCullFadedInUniformBuffer;
 
-/** Global primitive dither uniform buffer containing fully faded in. */
-FRHIUniformBuffer* GGetDitherFadedInUniformBuffer()
-{
-	FDitherUniformShaderParameters Parameters;
-	Parameters.LODFactor = 0.0f;
-	static TUniformBufferRef<FDitherUniformShaderParameters> ResourceRef =
-		TUniformBufferRef<FDitherUniformShaderParameters>::CreateUniformBufferImmediate(Parameters, UniformBuffer_MultiFrame, EUniformBufferValidation::None);
-
-	return ResourceRef.GetReference();
-}
+/** Global primitive uniform buffer resource containing dither faded in */
+TGlobalResource< FGlobalDitherUniformBuffer > GDitherFadedInUniformBuffer;
 
 static FThreadSafeCounter FSceneViewState_UniqueID;
 
diff --git a/Engine/Source/Runtime/Renderer/Private/ScenePrivate.h b/Engine/Source/Runtime/Renderer/Private/ScenePrivate.h
--- a/Engine/Source/Runtime/Renderer/Private/ScenePrivate.h
+++ b/Engine/Source/Runtime/Renderer/Private/ScenePrivate.h
@@ -369,11 +369,40 @@ public:
 	}
 };
 
-/** Global primitive distance cull fading uniform buffer containing fully faded in. */
-FRHIUniformBuffer* GGetDistanceCullFadedInUniformBuffer();
+/**
+ * Distance cull fading uniform buffer containing fully faded in.
+ */
+class FGlobalDistanceCullFadeUniformBuffer : public TUniformBuffer< FDistanceCullFadeUniformShaderParameters >
+{
+public:
+	void InitContents()
+	{
+		FDistanceCullFadeUniformShaderParameters Parameters;
+		Parameters.FadeTimeScaleBias.X = 0.0f;
+		Parameters.FadeTimeScaleBias.Y = 1.0f;
+		SetContents(Parameters);
+	}
+};
+
+/** Global primitive uniform buffer resource containing faded in */
+extern TGlobalResource< FGlobalDistanceCullFadeUniformBuffer > GDistanceCullFadedInUniformBuffer;
+
+/**
+ * Dither uniform buffer containing fully faded in.
+ */
+class FGlobalDitherUniformBuffer : public TUniformBuffer< FDitherUniformShaderParameters >
+{
+public:
+	void InitContents()
+	{
+		FDitherUniformShaderParameters Parameters;
+		Parameters.LODFactor = 0.0f;
+		SetContents(Parameters);
+	}
+};
 
-/** Global primitive dither uniform buffer containing fully faded in. */
-FRHIUniformBuffer* GGetDitherFadedInUniformBuffer();
+/** Global primitive uniform buffer resource containing faded in */
+extern TGlobalResource< FGlobalDitherUniformBuffer > GDitherFadedInUniformBuffer;
 
 /**
  * Stores fading state for a single primitive in a single view
diff --git a/Engine/Source/Runtime/Renderer/Private/ShaderBaseClasses.cpp b/Engine/Source/Runtime/Renderer/Private/ShaderBaseClasses.cpp
--- a/Engine/Source/Runtime/Renderer/Private/ShaderBaseClasses.cpp
+++ b/Engine/Source/Runtime/Renderer/Private/ShaderBaseClasses.cpp
@@ -30,8 +30,8 @@ FAutoConsoleVariableRef FMaterialShader::CVarAllowCachedUniformExpressions(
 
 void FMeshMaterialShaderElementData::InitializeMeshMaterialData(const FSceneView* SceneView, const FPrimitiveSceneProxy* RESTRICT PrimitiveSceneProxy, int32 StaticMeshId, bool bDitheredLODTransition, bool bAllowStencilDither)
 {
-	FadeUniformBuffer = GGetDistanceCullFadedInUniformBuffer();
-	DitherUniformBuffer = GGetDitherFadedInUniformBuffer();
+	FadeUniformBuffer = GDistanceCullFadedInUniformBuffer.GetUniformBufferRHI();
+	DitherUniformBuffer = GDitherFadedInUniformBuffer.GetUniformBufferRHI();
 
 	if (SceneView)
 	{
diff --git a/Engine/Source/Runtime/Renderer/Private/VolumetricCloudRendering.cpp b/Engine/Source/Runtime/Renderer/Private/VolumetricCloudRendering.cpp
--- a/Engine/Source/Runtime/Renderer/Private/VolumetricCloudRendering.cpp
+++ b/Engine/Source/Runtime/Renderer/Private/VolumetricCloudRendering.cpp
@@ -2347,8 +2347,8 @@ void FSceneRenderer::RenderVolumetricCloudsInternal(FRDGBuilder& GraphBuilder, F
 					FMeshPassProcessorRenderState DrawRenderState;
 
 					FMeshMaterialShaderElementData ShaderElementData;
-					ShaderElementData.FadeUniformBuffer = GGetDistanceCullFadedInUniformBuffer();
-					ShaderElementData.DitherUniformBuffer = GGetDitherFadedInUniformBuffer();
+					ShaderElementData.FadeUniformBuffer = GDistanceCullFadedInUniformBuffer.GetUniformBufferRHI();
+					ShaderElementData.DitherUniformBuffer = GDitherFadedInUniformBuffer.GetUniformBufferRHI();
 
 					FMeshProcessorShaders PassShaders;
 					PassShaders.ComputeShader = ComputeShader;
@@ -2445,8 +2445,8 @@ void FSceneRenderer::RenderVolumetricCloudsInternal(FRDGBuilder& GraphBuilder, F
 				FMeshPassProcessorRenderState DrawRenderState;
 
 				FMeshMaterialShaderElementData ShaderElementData;
-				ShaderElementData.FadeUniformBuffer = GGetDistanceCullFadedInUniformBuffer();
-				ShaderElementData.DitherUniformBuffer = GGetDitherFadedInUniformBuffer();
+				ShaderElementData.FadeUniformBuffer = GDistanceCullFadedInUniformBuffer.GetUniformBufferRHI();
+				ShaderElementData.DitherUniformBuffer = GDitherFadedInUniformBuffer.GetUniformBufferRHI();
 
 				FMeshProcessorShaders PassShaders;
 				PassShaders.ComputeShader = ComputeShader;
