Add postprocess material and MRQ support for path tracing buffers.
1. Postprocess can access path tracing buffers (raw/denoised radiance, albedo, normal, variance) with unique material node `PathTracingBufferTexture`.
2. Output in MRQ is also supported.

diff --git a/Engine/Shaders/Private/MaterialTemplate.ush b/Engine/Shaders/Private/MaterialTemplate.ush
--- a/Engine/Shaders/Private/MaterialTemplate.ush
+++ b/Engine/Shaders/Private/MaterialTemplate.ush
@@ -2586,6 +2586,39 @@ float4 MaterialExpressionDBufferTextureLookup(FMaterialPixelParameters Parameter
 	return float4(0, 0, 0, 1);
 }
 
+
+MaterialFloat2 GetDefaultPathTracingBufferTextureUV(FMaterialPixelParameters Parameters, const uint PathTracingBufferTextureIndex)
+{
+	float2 ViewportUV = GetViewportUV(Parameters);
+#if POST_PROCESS_MATERIAL && SHADING_PATH_DEFERRED && MATERIAL_PATH_TRACING_BUFFER_READ
+	
+	switch (PathTracingBufferTextureIndex)
+	{
+	case 0: return ViewportUV * PathTracingPostProcessInput_0_UVViewportSize + PathTracingPostProcessInput_0_UVViewportMin;
+	case 1: return ViewportUV * PathTracingPostProcessInput_1_UVViewportSize + PathTracingPostProcessInput_1_UVViewportMin;
+	case 2: return ViewportUV * PathTracingPostProcessInput_2_UVViewportSize + PathTracingPostProcessInput_2_UVViewportMin;
+	case 3: return ViewportUV * PathTracingPostProcessInput_3_UVViewportSize + PathTracingPostProcessInput_3_UVViewportMin;
+	case 4: return ViewportUV * PathTracingPostProcessInput_4_UVViewportSize + PathTracingPostProcessInput_4_UVViewportMin;
+	}
+#endif
+	return ViewportUV;
+}
+
+float4 MaterialExpressionPathTracingBufferTextureLookup(FMaterialPixelParameters Parameters, float2 BufferUV, int PathTracingBufferTextureIndex)
+{
+#if POST_PROCESS_MATERIAL && SHADING_PATH_DEFERRED && MATERIAL_PATH_TRACING_BUFFER_READ
+	switch (PathTracingBufferTextureIndex)
+	{
+	case 0:	return Texture2DSample(PathTracingPostProcessInput_0_Texture, PathTracingPostProcessInput_0_Sampler, BufferUV);
+	case 1:	return Texture2DSample(PathTracingPostProcessInput_1_Texture, PathTracingPostProcessInput_1_Sampler, BufferUV);
+	case 2:	return Texture2DSample(PathTracingPostProcessInput_2_Texture, PathTracingPostProcessInput_2_Sampler, BufferUV);
+	case 3:	return Texture2DSample(PathTracingPostProcessInput_3_Texture, PathTracingPostProcessInput_3_Sampler, BufferUV);
+	case 4:	return Texture2DSample(PathTracingPostProcessInput_4_Texture, PathTracingPostProcessInput_4_Sampler, BufferUV);
+	}
+#endif
+	return float4(1, 0, 0, 1);
+}
+
 // DERIV_BASE_VALUE() is a disgusting macro to manage backwards compatibility while changing the generated materials. The existing nodes are not "aware"
 // of partial derivatives, and will use sprintf() to put the same line in both CalcPixelMaterialInputs() and CalcPixelMaterialInputsAnalyticDerivatives().
 // If we have a line like this in CalcPixelMaterialInputs():
diff --git a/Engine/Shaders/Private/PathTracing/PathTracingSpatialTemporalDenoising.usf b/Engine/Shaders/Private/PathTracing/PathTracingSpatialTemporalDenoising.usf
--- a/Engine/Shaders/Private/PathTracing/PathTracingSpatialTemporalDenoising.usf
+++ b/Engine/Shaders/Private/PathTracing/PathTracingSpatialTemporalDenoising.usf
@@ -40,6 +40,7 @@ PathTracingSpatialTemporalDenoising.usf: Spatial temporal denoising for path tra
 
 #ifndef TEMPORAL_PREPASS
 #define TEMPORAL_PREPASS				0
+#define VARIANCE_TYPE					0
 #endif
 
 #define PREPASS_PHASE_INIT				0
@@ -49,6 +50,10 @@ PathTracingSpatialTemporalDenoising.usf: Spatial temporal denoising for path tra
 #define SPATIAL_DENOISING				0
 #endif
 
+#ifndef PREPASS_GENERATE_TEXTURE
+#define PREPASS_GENERATE_TEXTURE		0
+#endif
+
 #ifndef VISUALIZE_MOTIONVECTOR
 #define VISUALIZE_MOTIONVECTOR			0
 #endif
@@ -778,9 +783,42 @@ float3 GetTextureValueWithOffset(Texture2D Tex, float2 Position)
 #if TEMPORAL_PREPASS
 
 Texture2D<float4> InputTexture;
+Texture2D<float4> AlbedoTexture;
+Texture2D<float4> NormalTexture;
 RWStructuredBuffer<FPixelMaterialLightingFingerprint> RWVarianceMap;
 int Iteration;
 
+#define RADIANCE_MULTI_CHANNEL					0
+#define RADIANCE_ALBEDO_NORMAL_SINGLECHANNEL	1
+
+// Note that the radiance, albedo and normal are the mean value.
+float4 GetMeanFromTextures(int2 Position)
+{
+	float4 Value = 0;
+#if VARIANCE_TYPE == RADIANCE_MULTI_CHANNEL
+	Value  = InputTexture[Position];
+#elif VARIANCE_TYPE == RADIANCE_ALBEDO_NORMAL_SINGLECHANNEL
+	Value.x = Luminance(InputTexture[Position].rgb);
+	Value.y = Luminance(AlbedoTexture[Position].rgb);
+	
+	//The length of the average normal indicates the variance
+	//1: All pointing to the same direction
+	//0: Uniform in all directions
+	Value.z = length2(NormalTexture[Position].rgb);
+	
+	// Special case handling:
+	// 1. Regions without normal has a zero length, like sky
+	// 2. The first frame
+	if (Value.z == 0.0f || Iteration == 0)
+	{
+		Value.z = 1.0f;
+	}
+#else
+	#error "Not implemented" 
+#endif
+	return Value;
+}
+
 [numthreads(THREAD_SIZE_X, THREAD_SIZE_Y, 1)]
 void TemporalPrepassCS(uint3 DT_ID : SV_DispatchThreadID)
 {
@@ -789,12 +827,14 @@ void TemporalPrepassCS(uint3 DT_ID : SV_DispatchThreadID)
 	if (all(DT_ID.xy < MaxPosition))
 	{
 		int index = Position.x + Position.y * TargetViewport_Extent.x;
-		float4 Radiance = InputTexture[Position];
+		
+		float4 Mean = GetMeanFromTextures(Position);
 
 #if PREPASS_PHASE == PREPASS_PHASE_INIT
 
 		FPixelMaterialLightingFingerprint VarianceInfo = (FPixelMaterialLightingFingerprint)0;
-		VarianceInfo.Mean = Radiance;
+		VarianceInfo.Mean = Mean;
+		VarianceInfo.Mean.w = 0;
 		RWVarianceMap[index] = VarianceInfo;
 
 #elif PREPASS_PHASE == PREPASS_PHASE_UPDATE
@@ -802,19 +842,45 @@ void TemporalPrepassCS(uint3 DT_ID : SV_DispatchThreadID)
 		FPixelMaterialLightingFingerprint VarianceInfo = RWVarianceMap[index];
 		
 		float4 PreviousMean = VarianceInfo.Mean;
-		float4 CurrentValue = (Radiance - PreviousMean) * Iteration + Radiance;
+		float4 CurrentValue = (Mean - PreviousMean) * Iteration + Mean;
 		
-		VarianceInfo.Mean = Radiance; 
+		VarianceInfo.Mean = Mean;
 		VarianceInfo.Mean.w = Iteration;
 		VarianceInfo.Var += (CurrentValue - VarianceInfo.Mean)*(CurrentValue - PreviousMean);
 
 		RWVarianceMap[index] = VarianceInfo;
 #else
+	#error "Not implemented" 
 #endif
 	}
 }
 #endif // TEMPORAL_PREPASS
 
+#if PREPASS_GENERATE_TEXTURE
+RWTexture2D<float4> OutputTexture;
+StructuredBuffer<FPixelMaterialLightingFingerprint> VarianceMap;
+
+[numthreads(THREAD_SIZE_X, THREAD_SIZE_Y, 1)]
+void PrepassGenerateTextureCS(uint3 DT_ID : SV_DispatchThreadID)
+{
+	int2 Position = (int2)DT_ID.xy;
+	int2 MaxPosition = int2(TargetViewport_Extent);
+	if (all(DT_ID.xy < MaxPosition))
+	{
+		int index = Position.x + Position.y * TargetViewport_Extent.x;
+		FPixelMaterialLightingFingerprint VarianceInfo = VarianceMap[index];
+
+		float4 StandDerivation= sqrt((VarianceInfo.Var) / (VarianceInfo.Mean.w + 1));
+
+#if VARIANCE_TYPE == RADIANCE_ALBEDO_NORMAL_SINGLECHANNEL
+		// The normal variance is indicated by 1 - sqrt for the lengh2(normal).
+		StandDerivation.z = 1 - sqrt(VarianceInfo.Mean.z);
+#endif
+		OutputTexture[Position] = float4(StandDerivation.xyz, 1.0f);
+	}
+}
+#endif // PREPASS_GENERATE_TEXTURE
+
 #if TEMPORAL_FEATURE_FUSION
 //Use Albedo, Normal, and Radiance Texture to create another feature vector
 // such that the feature vector leads to the best motion vector estimation.
diff --git a/Engine/Shaders/Private/PostProcessMaterialShaders.usf b/Engine/Shaders/Private/PostProcessMaterialShaders.usf
--- a/Engine/Shaders/Private/PostProcessMaterialShaders.usf
+++ b/Engine/Shaders/Private/PostProcessMaterialShaders.usf
@@ -56,6 +56,26 @@ SamplerState PostProcessInput_BilinearSampler;
 	#define PostProcessInput_4_SharedSampler PostProcessInput_4_Sampler
 #endif
 
+#if MATERIAL_PATH_TRACING_BUFFER_READ
+SCREEN_PASS_TEXTURE_VIEWPORT(PathTracingPostProcessInput_0)
+SCREEN_PASS_TEXTURE_VIEWPORT(PathTracingPostProcessInput_1)
+SCREEN_PASS_TEXTURE_VIEWPORT(PathTracingPostProcessInput_2)
+SCREEN_PASS_TEXTURE_VIEWPORT(PathTracingPostProcessInput_3)
+SCREEN_PASS_TEXTURE_VIEWPORT(PathTracingPostProcessInput_4)
+
+Texture2D PathTracingPostProcessInput_0_Texture;
+Texture2D PathTracingPostProcessInput_1_Texture;
+Texture2D PathTracingPostProcessInput_2_Texture;
+Texture2D PathTracingPostProcessInput_3_Texture;
+Texture2D PathTracingPostProcessInput_4_Texture;
+
+SamplerState PathTracingPostProcessInput_0_Sampler;
+SamplerState PathTracingPostProcessInput_1_Sampler;
+SamplerState PathTracingPostProcessInput_2_Sampler;
+SamplerState PathTracingPostProcessInput_3_Sampler;
+SamplerState PathTracingPostProcessInput_4_Sampler;
+#endif
+
 #define EYE_ADAPTATION_LOOSE_PARAMETERS 1
 
 #include "/Engine/Generated/Material.ush"
diff --git a/Engine/Source/Runtime/Engine/Classes/Materials/MaterialExpressionPathTracingBufferTexture.h b/Engine/Source/Runtime/Engine/Classes/Materials/MaterialExpressionPathTracingBufferTexture.h
new file mode 100644
--- /dev/null
+++ b/Engine/Source/Runtime/Engine/Classes/Materials/MaterialExpressionPathTracingBufferTexture.h
@@ -0,0 +1,48 @@
+// Copyright Epic Games, Inc. All Rights Reserved.
+
+#pragma once
+
+#include "CoreMinimal.h"
+#include "UObject/ObjectMacros.h"
+#include "Materials/MaterialExpression.h"
+#include "MaterialExpressionPathTracingBufferTexture.generated.h"
+
+UENUM()
+enum EPathTracingBufferTextureId : int
+{
+	/** Radiance (Path Tracing). The raw radiance. */
+	PTBT_Radiance UMETA(DisplayName = "Radiance (Path Tracing)"),
+	/** Normal (Path Tracing). Stores the denoised radiance if denoising is turned on and complete for the current frame, otherwise, black. */
+	PTBT_DenoisedRadiance UMETA(DisplayName = "Denoised radiance (Path Tracing)"),
+	/** Albedo (Path Tracing). Average albedo at the current sample count. */
+	PTBT_Albedo UMETA(DisplayName = "Albedo (Path Tracing)"),
+	/** Normal (Path Tracing). Average normal at the current sample count. */
+	PTBT_Normal UMETA(DisplayName = "Normal (Path Tracing)"),
+	/** Variance (Path Tracing). Path tracing variance stored as standard derivation. Variance can be per channel variance or 
+		variance of luminance, albedo, and normal based on the path tracing configuration. Hooking up this buffer can increase additional
+		cost.
+	*/
+	PTBT_Variance UMETA(DisplayName = "Variance (Path Tracing)"),
+};
+
+/** Path tracing buffer is only accessable in postprocess material. */
+UCLASS(collapsecategories, hidecategories = Object)
+class UMaterialExpressionPathTracingBufferTexture : public UMaterialExpression
+{
+	GENERATED_UCLASS_BODY()
+
+	/** UV in 0..1 range */
+	UPROPERTY(meta = (RequiredInput = "false", ToolTip = "Ignored if not specified"))
+	FExpressionInput Coordinates;
+
+	/** Which path tracing buffer texture we want to make a lookup into.*/
+	UPROPERTY(EditAnywhere, Category = UMaterialExpressionPathTracingBufferTexture, meta = (DisplayName = "Buffer Texture (Path Tracing)"))
+	TEnumAsByte<EPathTracingBufferTextureId> PathTracingBufferTextureId;
+
+	//~ Begin UMaterialExpression Interface
+#if WITH_EDITOR
+	virtual int32 Compile(class FMaterialCompiler* Compiler, int32 OutputIndex) override;
+	virtual void GetCaption(TArray<FString>& OutCaptions) const override;
+#endif
+	//~ End UMaterialExpression Interface
+};
diff --git a/Engine/Source/Runtime/Engine/Private/Materials/HLSLMaterialTranslator.cpp b/Engine/Source/Runtime/Engine/Private/Materials/HLSLMaterialTranslator.cpp
--- a/Engine/Source/Runtime/Engine/Private/Materials/HLSLMaterialTranslator.cpp
+++ b/Engine/Source/Runtime/Engine/Private/Materials/HLSLMaterialTranslator.cpp
@@ -1799,6 +1799,7 @@ void FHLSLMaterialTranslator::GetMaterialEnvironment(EShaderPlatform InPlatform,
 	OutEnvironment.SetDefine(TEXT("MATERIAL_USES_ANISOTROPY"), bUsesAnisotropy && FDataDrivenShaderPlatformInfo::GetSupportsAnisotropicMaterials(InPlatform));
 
 	OutEnvironment.SetDefine(TEXT("MATERIAL_DECAL_READ_MASK"), MaterialCompilationOutput.UsedDBufferTextures);
+	OutEnvironment.SetDefine(TEXT("MATERIAL_PATH_TRACING_BUFFER_READ"), MaterialCompilationOutput.UsedPathTracingBufferTextures);
 
 	// Count the number of VTStacks (each stack will allocate a feedback slot)
 	OutEnvironment.SetDefine(TEXT("NUM_VIRTUALTEXTURE_SAMPLES"), VTStacks.Num());
@@ -7049,6 +7050,29 @@ int32 FHLSLMaterialTranslator::DBufferTextureLookup(int32 ViewportUV, uint32 DBu
 	return AddCodeChunk(MCT_Float4, TEXT("MaterialExpressionDBufferTextureLookup(Parameters, %s, %d)"), *CoerceParameter(BufferUV, MCT_Float2), (int)DBufferTextureIndex);
 }
 
+int32 FHLSLMaterialTranslator::PathTracingBufferTextureLookup(int32 ViewportUV, uint32 PathTracingBufferTextureIndex)
+{
+	if (Material->GetMaterialDomain() != MD_PostProcess)
+	{
+		Errorf(TEXT("Path tracing buffer textures are only available on post process material."));
+	}
+
+	int32 BufferUV = INDEX_NONE;
+	if (ViewportUV != INDEX_NONE)
+	{
+		BufferUV = ViewportUV;
+	}
+	else
+	{
+		BufferUV = AddInlinedCodeChunk(MCT_Float2, TEXT("GetDefaultPathTracingBufferTextureUV(Parameters, 0)"));
+	}
+
+	MaterialCompilationOutput.SetIsPathTracingBufferTextureUsed(PathTracingBufferTextureIndex);
+	AddEstimatedTextureSample();
+
+	return AddCodeChunk(MCT_Float4, TEXT("MaterialExpressionPathTracingBufferTextureLookup(Parameters, %s, %d)"), *CoerceParameter(BufferUV, MCT_Float2), (int)PathTracingBufferTextureIndex);
+}
+
 int32 FHLSLMaterialTranslator::Switch(int32 SwitchValueInput, int32 DefaultInput, TArray<int32>& CompiledInputs)
 {
 	if (SwitchValueInput == INDEX_NONE)
diff --git a/Engine/Source/Runtime/Engine/Private/Materials/HLSLMaterialTranslator.h b/Engine/Source/Runtime/Engine/Private/Materials/HLSLMaterialTranslator.h
--- a/Engine/Source/Runtime/Engine/Private/Materials/HLSLMaterialTranslator.h
+++ b/Engine/Source/Runtime/Engine/Private/Materials/HLSLMaterialTranslator.h
@@ -898,6 +898,8 @@ protected:
 
 	virtual int32 DBufferTextureLookup(int32 ViewportUV, uint32 DBufferTextureIndex) override;
 
+	virtual int32 PathTracingBufferTextureLookup(int32 ViewportUV, uint32 PathTracingBufferTextureIndex) override;
+
 	// @param bTextureLookup true: texture, false:no texture lookup, usually to get the size
 	void UseSceneTextureId(ESceneTextureId SceneTextureId, bool bTextureLookup);
 
diff --git a/Engine/Source/Runtime/Engine/Private/Materials/MaterialExpressions.cpp b/Engine/Source/Runtime/Engine/Private/Materials/MaterialExpressions.cpp
--- a/Engine/Source/Runtime/Engine/Private/Materials/MaterialExpressions.cpp
+++ b/Engine/Source/Runtime/Engine/Private/Materials/MaterialExpressions.cpp
@@ -111,6 +111,7 @@
 #include "Materials/MaterialExpressionRayTracingQualitySwitch.h"
 #include "Materials/MaterialExpressionPathTracingQualitySwitch.h"
 #include "Materials/MaterialExpressionPathTracingRayTypeSwitch.h"
+#include "Materials/MaterialExpressionPathTracingBufferTexture.h"
 #include "Materials/MaterialExpressionGetMaterialAttributes.h"
 #include "Materials/MaterialExpressionHairAttributes.h"
 #include "Materials/MaterialExpressionHairColor.h"
@@ -17987,6 +17988,60 @@ void UMaterialExpressionPathTracingRayTypeSwitch::GetCaption(TArray<FString>& Ou
 
 #endif // WITH_EDITOR
 
+///////////////////////////////////////////////////////////////////////////////
+// UMaterialExpressionPathTracingBufferTexture
+///////////////////////////////////////////////////////////////////////////////
+UMaterialExpressionPathTracingBufferTexture::UMaterialExpressionPathTracingBufferTexture(const FObjectInitializer& ObjectInitializer)
+	: Super(ObjectInitializer)
+{
+#if WITH_EDITORONLY_DATA
+	// Structure to hold one-time initialization
+	struct FConstructorStatics
+	{
+		FText NAME_Decals;
+		FConstructorStatics()
+			: NAME_Decals(LOCTEXT("PathTracing", "PathTracing"))
+		{
+		}
+	};
+	static FConstructorStatics ConstructorStatics;
+
+	MenuCategories.Add(ConstructorStatics.NAME_Decals);
+
+	bShaderInputData = true;
+	bShowOutputNameOnPin = true;
+#endif
+
+#if WITH_EDITORONLY_DATA
+	Outputs.Reset();
+	Outputs.Add(FExpressionOutput(TEXT("RGBA"), 1, 1, 1, 1, 1));
+	Outputs.Add(FExpressionOutput(TEXT("RGB"), 1, 1, 1, 1, 0));
+	Outputs.Add(FExpressionOutput(TEXT("A"), 1, 0, 0, 0, 1));
+#endif
+}
+
+#if WITH_EDITOR
+int32 UMaterialExpressionPathTracingBufferTexture::Compile(class FMaterialCompiler* Compiler, int32 OutputIndex)
+{
+	int32 ViewportUV = INDEX_NONE;
+
+	if (Coordinates.GetTracedInput().Expression)
+	{
+		ViewportUV = Coordinates.Compile(Compiler);
+	}
+
+	return Compiler->PathTracingBufferTextureLookup(ViewportUV, PathTracingBufferTextureId);
+}
+
+void UMaterialExpressionPathTracingBufferTexture::GetCaption(TArray<FString>& OutCaptions) const
+{
+	UEnum* Enum = StaticEnum<EPathTracingBufferTextureId>();
+	check(Enum);
+
+	FString Name = Enum->GetDisplayNameTextByValue(PathTracingBufferTextureId).ToString();
+	OutCaptions.Add(Name);
+}
+#endif // WITH_EDITOR
 
 UMaterialExpressionObjectOrientation::UMaterialExpressionObjectOrientation(const FObjectInitializer& ObjectInitializer)
 	: Super(ObjectInitializer)
diff --git a/Engine/Source/Runtime/Engine/Public/MaterialCompiler.h b/Engine/Source/Runtime/Engine/Public/MaterialCompiler.h
--- a/Engine/Source/Runtime/Engine/Public/MaterialCompiler.h
+++ b/Engine/Source/Runtime/Engine/Public/MaterialCompiler.h
@@ -489,6 +489,7 @@ public:
 	virtual int32 SceneTextureLookup(int32 ViewportUV, uint32 SceneTextureId, bool bFiltered) = 0;
 	virtual int32 GetSceneTextureViewSize(int32 SceneTextureId, bool InvProperty) = 0;
 	virtual int32 DBufferTextureLookup(int32 ViewportUV, uint32 DBufferTextureIndex) = 0;
+	virtual int32 PathTracingBufferTextureLookup(int32 ViewportUV, uint32 PathTracingBufferTextureIndex) = 0;
 
 	virtual int32 StaticBool(bool Value) = 0;
 	virtual int32 StaticBoolParameter(FName ParameterName,bool bDefaultValue) = 0;
@@ -933,6 +934,7 @@ public:
 	virtual int32 SceneTextureLookup(int32 ViewportUV, uint32 InSceneTextureId, bool bFiltered) override { return Compiler->SceneTextureLookup(ViewportUV, InSceneTextureId, bFiltered); }
 	virtual int32 GetSceneTextureViewSize(int32 SceneTextureId, bool InvProperty) override { return Compiler->GetSceneTextureViewSize(SceneTextureId, InvProperty); }
 	virtual int32 DBufferTextureLookup(int32 ViewportUV, uint32 DBufferTextureIndex) override { return Compiler->DBufferTextureLookup(ViewportUV, DBufferTextureIndex); }
+	virtual int32 PathTracingBufferTextureLookup(int32 ViewportUV, uint32 PathTracingBufferTextureIndex) override { return Compiler->PathTracingBufferTextureLookup(ViewportUV, PathTracingBufferTextureIndex); }
 
 	virtual int32 StaticBool(bool Value) override { return Compiler->StaticBool(Value); }
 	virtual int32 StaticBoolParameter(FName ParameterName, bool bDefaultValue) override { return Compiler->StaticBoolParameter(ParameterName, bDefaultValue); }
diff --git a/Engine/Source/Runtime/Engine/Public/MaterialShared.h b/Engine/Source/Runtime/Engine/Public/MaterialShared.h
--- a/Engine/Source/Runtime/Engine/Public/MaterialShared.h
+++ b/Engine/Source/Runtime/Engine/Public/MaterialShared.h
@@ -666,6 +666,7 @@ class FMaterialCompilationOutput
 public:
 	FMaterialCompilationOutput() :
 		UsedSceneTextures(0),
+		UsedPathTracingBufferTextures(0),
 #if WITH_EDITOR
 		EstimatedNumTextureSamplesVS(0),
 		EstimatedNumTextureSamplesPS(0),
@@ -697,6 +698,8 @@ public:
 
 	ENGINE_API void SetIsDBufferTextureUsed(int32 TextureIndex) { UsedDBufferTextures |= (1 << TextureIndex); }
 
+	ENGINE_API void SetIsPathTracingBufferTextureUsed(int32 TexId) { UsedPathTracingBufferTextures |= (1 << TexId); }
+
 	/** Indicates whether the material uses scene color. */
 	ENGINE_API bool RequiresSceneColorCopy() const { return IsSceneTextureUsed(PPI_SceneColor); }
 
@@ -735,6 +738,9 @@ public:
 	/** Bitfield of the ESceneTextures used */
 	LAYOUT_FIELD(uint32, UsedSceneTextures);
 
+	/** Bitfield of the EPathTracingBufferTextureId used */
+	LAYOUT_FIELD(uint8, UsedPathTracingBufferTextures);
+
 	/** Number of times SampleTexture is called, excludes custom nodes. */
 	LAYOUT_FIELD_EDITORONLY(uint16, EstimatedNumTextureSamplesVS);
 	LAYOUT_FIELD_EDITORONLY(uint16, EstimatedNumTextureSamplesPS);
@@ -1370,6 +1376,7 @@ public:
 	uint32 GetNumVirtualTextureStacks() const { return GetContent()->MaterialCompilationOutput.UniformExpressionSet.VTStacks.Num(); }
 	uint8 GetRuntimeVirtualTextureOutputAttributeMask() const { return GetContent()->MaterialCompilationOutput.RuntimeVirtualTextureOutputAttributeMask; }
 	bool UsesSceneTexture(uint32 TexId) const { return (GetContent()->MaterialCompilationOutput.UsedSceneTextures & (1ull << TexId)) != 0; }
+	bool UsesPathTracingBufferTexture(uint32 TexId) const { return (GetContent()->MaterialCompilationOutput.UsedPathTracingBufferTextures & (1ull << TexId)) != 0;}
 
 	bool IsValidForRendering(bool bFailOnInvalid = false) const
 	{
diff --git a/Engine/Source/Runtime/Renderer/Private/DeferredShadingRenderer.cpp b/Engine/Source/Runtime/Renderer/Private/DeferredShadingRenderer.cpp
--- a/Engine/Source/Runtime/Renderer/Private/DeferredShadingRenderer.cpp
+++ b/Engine/Source/Runtime/Renderer/Private/DeferredShadingRenderer.cpp
@@ -4122,6 +4122,8 @@ void FDeferredShadingSceneRenderer::Render(FRDGBuilder& GraphBuilder)
 		VirtualTextureFeedbackEnd(GraphBuilder);
 	}
 
+	FPathTracingResources PathTracingResources;
+
 #if RHI_RAYTRACING
 	if (IsRayTracingEnabled())
 	{
@@ -4133,7 +4135,7 @@ void FDeferredShadingSceneRenderer::Render(FRDGBuilder& GraphBuilder)
 		{
 			for (const FViewInfo& View : Views)
 			{
-				RenderPathTracing(GraphBuilder, View, SceneTextures.UniformBuffer, SceneTextures.Color.Target);
+				RenderPathTracing(GraphBuilder, View, SceneTextures.UniformBuffer, SceneTextures.Color.Target, PathTracingResources);
 			}
 		}
 		else if (ViewFamily.EngineShowFlags.RayTracingDebug)
@@ -4252,6 +4254,7 @@ void FDeferredShadingSceneRenderer::Render(FRDGBuilder& GraphBuilder)
 		PostProcessingInputs.ExposureIlluminance = ExposureIlluminance;
 		PostProcessingInputs.SceneTextures = SceneTextures.UniformBuffer;
 		PostProcessingInputs.bSeparateCustomStencil = SceneTextures.CustomDepth.bSeparateStencilBuffer;
+		PostProcessingInputs.PathTracingResources = PathTracingResources;
 
 		GraphBuilder.FlushSetupQueue();
 
diff --git a/Engine/Source/Runtime/Renderer/Private/DeferredShadingRenderer.h b/Engine/Source/Runtime/Renderer/Private/DeferredShadingRenderer.h
--- a/Engine/Source/Runtime/Renderer/Private/DeferredShadingRenderer.h
+++ b/Engine/Source/Runtime/Renderer/Private/DeferredShadingRenderer.h
@@ -1142,7 +1142,8 @@ private:
 		FRDGBuilder& GraphBuilder,
 		const FViewInfo& View,
 		TRDGUniformBufferRef<FSceneTextureUniformParameters> SceneTexturesUniformBuffer,
-		FRDGTextureRef SceneColorOutputTexture);
+		FRDGTextureRef SceneColorOutputTexture,
+		struct FPathTracingResources& PathTracingResources);
 
 	void ComputePathCompaction(FRHICommandListImmediate& RHICmdList, const FViewInfo& View, FRHITexture* RadianceTexture, FRHITexture* SampleCountTexture, FRHITexture* PixelPositionTexture,
 		FRHIUnorderedAccessView* RadianceSortedRedUAV, FRHIUnorderedAccessView* RadianceSortedGreenUAV, FRHIUnorderedAccessView* RadianceSortedBlueUAV, FRHIUnorderedAccessView* RadianceSortedAlphaUAV, FRHIUnorderedAccessView* SampleCountSortedUAV);
diff --git a/Engine/Source/Runtime/Renderer/Private/PathTracing.cpp b/Engine/Source/Runtime/Renderer/Private/PathTracing.cpp
--- a/Engine/Source/Runtime/Renderer/Private/PathTracing.cpp
+++ b/Engine/Source/Runtime/Renderer/Private/PathTracing.cpp
@@ -313,6 +313,15 @@ TAutoConsoleVariable<int32> CVarPathTracingCameraMediumTracking(
 	ECVF_RenderThreadSafe
 );
 
+TAutoConsoleVariable<int32> CVarPathTracingOutputPostProcessResources(
+	TEXT("r.PathTracing.OutputPostProcessResources"),
+	1,
+	TEXT("Output the pathtracing resources to the postprocess passes\n")
+	TEXT("0: off\n")
+	TEXT("1: on (Buffers including, raw/denoised radiance, albedo, normal, and variance)\n"),
+	ECVF_RenderThreadSafe
+);
+
 BEGIN_SHADER_PARAMETER_STRUCT(FPathTracingData, )
 	SHADER_PARAMETER(float, BlendFactor)
 	SHADER_PARAMETER(uint32, Iteration)
@@ -2032,7 +2041,8 @@ void FDeferredShadingSceneRenderer::RenderPathTracing(
 	FRDGBuilder& GraphBuilder,
 	const FViewInfo& View,
 	TRDGUniformBufferRef<FSceneTextureUniformParameters> SceneTexturesUniformBuffer,
-	FRDGTextureRef SceneColorOutputTexture)
+	FRDGTextureRef SceneColorOutputTexture,
+	FPathTracingResources& PathTracingResources)
 {
 	RDG_EVENT_SCOPE(GraphBuilder, "Path Tracing");
 
@@ -2682,6 +2692,9 @@ void FDeferredShadingSceneRenderer::RenderPathTracing(
 		FRDGBuffer* CurrentVarianceBufer = nullptr;
 		{
 			DenoisingContext.RadianceTexture = RadianceTexture;
+			DenoisingContext.AlbedoTexture = AlbedoTexture;
+			DenoisingContext.NormalTexture = NormalTexture;
+
 			if (PathTracingState->VarianceBuffer)
 			{
 				DenoisingContext.VarianceBuffer = GraphBuilder.RegisterExternalBuffer(PathTracingState->VarianceBuffer, TEXT("PathTracing.VarianceBuffer"));
@@ -2695,8 +2708,6 @@ void FDeferredShadingSceneRenderer::RenderPathTracing(
 		// 2. Denoising pass
 		if (NeedsDenoise || EnablePathTracingDenoiserRealtimeDebug)
 		{
-			DenoisingContext.AlbedoTexture = AlbedoTexture;
-			DenoisingContext.NormalTexture = NormalTexture;
 			DenoisingContext.RadianceTexture = RadianceTexture;
 			DenoisingContext.FrameIndex = PathTracingState->FrameIndex;
 			DenoisingContext.VarianceBuffer = CurrentVarianceBufer;
@@ -2785,6 +2796,17 @@ void FDeferredShadingSceneRenderer::RenderPathTracing(
 		DisplayParameters
 	);
 
+	// Setup the path tracing resources to be used by post process pass.
+	if (CVarPathTracingOutputPostProcessResources.GetValueOnRenderThread() != 0)
+	{
+		PathTracingResources.bPostProcessEnabled = true;
+		PathTracingResources.DenoisedRadiance = DenoisedRadianceTexture ? DenoisedRadianceTexture : RadianceTexture;
+		PathTracingResources.Radiance = RadianceTexture;
+		PathTracingResources.Albedo = AlbedoTexture;
+		PathTracingResources.Normal = NormalTexture;
+		PathTracingResources.Variance = DenoisingContext.VarianceTexture;
+	}
+
 	// Add a visualization path for denoising
 	if (NeedsDenoise || EnablePathTracingDenoiserRealtimeDebug)
 	{
diff --git a/Engine/Source/Runtime/Renderer/Private/PathTracing.h b/Engine/Source/Runtime/Renderer/Private/PathTracing.h
--- a/Engine/Source/Runtime/Renderer/Private/PathTracing.h
+++ b/Engine/Source/Runtime/Renderer/Private/PathTracing.h
@@ -25,6 +25,19 @@ RENDERER_API FRHIRayTracingShader* GetPathTracingDefaultMissShader(const FGlobal
 RENDERER_API FRHIRayTracingShader* GetPathTracingDefaultOpaqueHitShader(const FGlobalShaderMap* ShaderMap);
 RENDERER_API FRHIRayTracingShader* GetPathTracingDefaultHiddenHitShader(const FGlobalShaderMap* ShaderMap);
 
+class FRDGTexture;
+
+struct FPathTracingResources
+{
+	FRDGTexture* DenoisedRadiance = nullptr;
+	FRDGTexture* Radiance = nullptr;
+	FRDGTexture* Albedo = nullptr;
+	FRDGTexture* Normal = nullptr;
+	FRDGTexture* Variance = nullptr;
+	
+	bool bPostProcessEnabled = false;
+};
+
 namespace PathTracing
 {
 	bool UsesDecals(const FSceneViewFamily& ViewFamily);
diff --git a/Engine/Source/Runtime/Renderer/Private/PathTracingSpatialTemporalDenoising.cpp b/Engine/Source/Runtime/Renderer/Private/PathTracingSpatialTemporalDenoising.cpp
--- a/Engine/Source/Runtime/Renderer/Private/PathTracingSpatialTemporalDenoising.cpp
+++ b/Engine/Source/Runtime/Renderer/Private/PathTracingSpatialTemporalDenoising.cpp
@@ -21,6 +21,7 @@ PathTracingMotionVectorFunction* GPathTracingMotionVectorFunc = nullptr;
 #include "RayTracingTypes.h"
 #include "RenderGraphDefinitions.h"
 #include "RendererPrivate.h"
+#include "PostProcess/PostProcessMaterial.h"
 #include <limits>
 
 DEFINE_LOG_CATEGORY_STATIC(LogPathTracingDenoising, Log, All);
@@ -68,6 +69,22 @@ namespace {
 		ECVF_RenderThreadSafe
 	);
 
+	TAutoConsoleVariable<int32> CVarPathTracingDenoiserPrepassVarianceType(
+		TEXT("r.PathTracing.Denoiser.Prepass.VarianceType"),
+		1,
+		TEXT("Select the per-pixel variance type:")
+		TEXT("0: Multiple channel (RGB) variance for radiance")
+		TEXT("1: Combined single channel variance for radiance, albedo and normal"),
+		ECVF_RenderThreadSafe
+	);
+
+	TAutoConsoleVariable<int32> CVarPathTracingDenoiserPrepassOutputVarianceTexture(
+		TEXT("r.PathTracing.Denoiser.Prepass.OutputVarianceTexture"),
+		1,
+		TEXT("0: Variance is used only in the denoiser")
+		TEXT("1: Output to the postprocess material, usually used by MRQ")
+	);
+
 	TAutoConsoleVariable<int32> CVarPathTracingSpatialDenoiser(
 		TEXT("r.PathTracing.SpatialDenoiser"),
 		1,
@@ -519,6 +536,18 @@ static bool ShouldEnableSubpixelOffset(uint32 MipLevel)
 	return bShouldEnableSubpixelOffset;
 }
 
+static bool ShouldPrepassOutputVarianceTexture(const FViewInfo& View)
+{
+	static const auto CVarOutputPostProcessResources = 
+		IConsoleManager::Get().FindTConsoleVariableDataInt(TEXT("r.PathTracing.OutputPostProcessResources"));
+	const bool bOutputPostProcessResources = CVarOutputPostProcessResources ?
+		(CVarOutputPostProcessResources->GetValueOnRenderThread() != 0) : false;
+
+	return CVarPathTracingDenoiserPrepassOutputVarianceTexture.GetValueOnRenderThread() != 0 &&
+		bOutputPostProcessResources && 
+		IsPathTracingVarianceTextureRequiredInPostProcessMaterial(View);
+}
+
 static constexpr uint32 kMipDiffDelta = 2;
 
 static constexpr uint32 kNumberOfPixelShifts = 25;
@@ -1635,13 +1664,41 @@ public:
 
 	BEGIN_SHADER_PARAMETER_STRUCT(FParameters, )
 		SHADER_PARAMETER_RDG_TEXTURE_SRV(Texture2D, InputTexture)
+		SHADER_PARAMETER_RDG_TEXTURE_SRV(Texture2D, AlbedoTexture)
+		SHADER_PARAMETER_RDG_TEXTURE_SRV(Texture2D, NormalTexture)
 		SHADER_PARAMETER_RDG_BUFFER_UAV(RWStructuredBuffer<FPixelMaterialLightingFingerprint>, RWVarianceMap)
 		SHADER_PARAMETER_STRUCT(FScreenPassTextureViewportParameters, TargetViewport)
 		SHADER_PARAMETER(int, Iteration)
 	END_SHADER_PARAMETER_STRUCT()
 	
+	enum class EVarianceType : uint32
+	{
+		RadianceMultiChannel = 0,
+		RadianceAlbedoNormalSingleChannel,
+		MAX
+	};
+
 	class FPrepassPhase : SHADER_PERMUTATION_BOOL("PREPASS_PHASE");	// 0: initialize, 1: update
-	using FPermutationDomain = TShaderPermutationDomain<FPrepassPhase>;
+	class FVarianceType : SHADER_PERMUTATION_ENUM_CLASS("VARIANCE_TYPE", EVarianceType);
+	using FPermutationDomain = TShaderPermutationDomain<FPrepassPhase,FVarianceType>;
+
+	static EVarianceType GetVarianceType()
+	{
+		return static_cast<FTemporalPrepassCS::EVarianceType>(
+			FMath::Clamp(CVarPathTracingDenoiserPrepassVarianceType.GetValueOnRenderThread(),
+				0,
+				static_cast<int32>(EVarianceType::MAX) - 1));
+	}
+
+	static const TCHAR* GetEventName(EVarianceType VarianceType)
+	{
+		static const TCHAR* const kEventName[] = {
+			TEXT("Radiance"),
+			TEXT("Rad,Albedo,Norm")
+		};
+		static_assert(UE_ARRAY_COUNT(kEventName) == int32(EVarianceType::MAX), "Fix me");
+		return kEventName[int32(VarianceType)];
+	}
 
 	static bool ShouldCompilePermutation(const FGlobalShaderPermutationParameters& Parameters)
 	{
@@ -1656,11 +1713,40 @@ public:
 
 IMPLEMENT_GLOBAL_SHADER(FTemporalPrepassCS, "/Engine/Private/PathTracing/PathTracingSpatialTemporalDenoising.usf", "TemporalPrepassCS", SF_Compute);
 
+class FPrepassGenerateTextureCS : public FGlobalShader
+{
+public:
+	DECLARE_GLOBAL_SHADER(FPrepassGenerateTextureCS);
+	SHADER_USE_PARAMETER_STRUCT(FPrepassGenerateTextureCS, FGlobalShader);
+
+	BEGIN_SHADER_PARAMETER_STRUCT(FParameters, )
+		SHADER_PARAMETER_RDG_BUFFER_SRV(StructuredBuffer<FPixelMaterialLightingFingerprint>, VarianceMap)
+		SHADER_PARAMETER_RDG_TEXTURE_UAV(RWTexture2D, OutputTexture)
+		SHADER_PARAMETER_STRUCT(FScreenPassTextureViewportParameters, TargetViewport)
+	END_SHADER_PARAMETER_STRUCT()
+
+	class FVarianceType : SHADER_PERMUTATION_ENUM_CLASS("VARIANCE_TYPE", FTemporalPrepassCS::EVarianceType);
+	using FPermutationDomain = TShaderPermutationDomain<FVarianceType>;
+
+	static bool ShouldCompilePermutation(const FGlobalShaderPermutationParameters& Parameters)
+	{
+		return ShouldCompilePathTracingDenoiserShadersForProject(Parameters.Platform);
+	}
+
+	static void ModifyCompilationEnvironment(const FShaderPermutationParameters&, FShaderCompilerEnvironment& OutEnvironment)
+	{
+		OutEnvironment.SetDefine(TEXT("PREPASS_GENERATE_TEXTURE"), 1);
+	}
+};
+
+IMPLEMENT_GLOBAL_SHADER(FPrepassGenerateTextureCS, "/Engine/Private/PathTracing/PathTracingSpatialTemporalDenoising.usf", "PrepassGenerateTextureCS", SF_Compute);
+
 void PathTracingSpatialTemporalDenoisingPrePass(FRDGBuilder& GraphBuilder, const FViewInfo& View,
 	int IterationNumber,
 	FPathTracingSpatialTemporalDenoisingContext& SpatialTemporalDenoisingContext)
 {
-	bool bShouldGenerateVarianceMap = ShouldGenerateVarianceMap();
+	bool bShouldPrepassOutputVarianceTexture = ShouldPrepassOutputVarianceTexture(View);
+	bool bShouldGenerateVarianceMap = ShouldGenerateVarianceMap() || bShouldPrepassOutputVarianceTexture;
 	if (bShouldGenerateVarianceMap)
 	{
 		bool bUpdateVarianceMap = (IterationNumber > 0);
@@ -1678,6 +1764,8 @@ void PathTracingSpatialTemporalDenoisingPrePass(FRDGBuilder& GraphBuilder, const
 			SHADER::FParameters* PassParameters = GraphBuilder.AllocParameters<SHADER::FParameters>();
 			{
 				PassParameters->InputTexture = GraphBuilder.CreateSRV(FRDGTextureSRVDesc::Create(SpatialTemporalDenoisingContext.RadianceTexture));
+				PassParameters->AlbedoTexture = GraphBuilder.CreateSRV(FRDGTextureSRVDesc::Create(SpatialTemporalDenoisingContext.AlbedoTexture));
+				PassParameters->NormalTexture = GraphBuilder.CreateSRV(FRDGTextureSRVDesc::Create(SpatialTemporalDenoisingContext.NormalTexture));
 				PassParameters->RWVarianceMap = GraphBuilder.CreateUAV(SpatialTemporalDenoisingContext.VarianceBuffer, EPixelFormat::PF_R32_FLOAT);
 				PassParameters->TargetViewport = TargetViewportParameters;
 				PassParameters->Iteration = IterationNumber;
@@ -1685,11 +1773,49 @@ void PathTracingSpatialTemporalDenoisingPrePass(FRDGBuilder& GraphBuilder, const
 
 			SHADER::FPermutationDomain ComputeShaderPermutationVector;
 			ComputeShaderPermutationVector.Set<SHADER::FPrepassPhase>(bUpdateVarianceMap);
+			ComputeShaderPermutationVector.Set<SHADER::FVarianceType>(SHADER::GetVarianceType());
+
+			TShaderMapRef<SHADER> ComputeShader(View.ShaderMap, ComputeShaderPermutationVector);
+			FComputeShaderUtils::AddPass(
+				GraphBuilder,
+				RDG_EVENT_NAME("PathTracing::Denoising::Prepass(%s %dx%d)",
+					SHADER::GetEventName(SHADER::GetVarianceType()),
+					TargetViewport.Extent.X,
+					TargetViewport.Extent.Y),
+				ComputeShader,
+				PassParameters,
+				FComputeShaderUtils::GetGroupCount(TargetViewport.Extent, 8));
+		}
+
+		if (bShouldPrepassOutputVarianceTexture)
+		{
+			const FRDGTextureDesc TextureDescriptor = FRDGTextureDesc::Create2D(
+				TargetViewport.Extent,
+				PF_A32B32G32R32F,
+				FClearValueBinding(),
+				TexCreate_RenderTargetable | TexCreate_ShaderResource | TexCreate_UAV | GetExtraTextureCreateFlagsForDenoiser());
+
+			if (!SpatialTemporalDenoisingContext.VarianceTexture)
+			{
+				SpatialTemporalDenoisingContext.VarianceTexture = GraphBuilder.CreateTexture(TextureDescriptor, TEXT("PathTracing.VarianceTexture"));
+			}
+
+			typedef FPrepassGenerateTextureCS SHADER;
+			SHADER::FParameters* PassParameters = GraphBuilder.AllocParameters<SHADER::FParameters>();
+			{
+				PassParameters->OutputTexture = GraphBuilder.CreateUAV(SpatialTemporalDenoisingContext.VarianceTexture);
+				PassParameters->VarianceMap	  = GraphBuilder.CreateSRV(SpatialTemporalDenoisingContext.VarianceBuffer, EPixelFormat::PF_R32_FLOAT);
+				PassParameters->TargetViewport = TargetViewportParameters;
+			}
+
+			SHADER::FPermutationDomain ComputeShaderPermutationVector;
+			ComputeShaderPermutationVector.Set<SHADER::FVarianceType>(FTemporalPrepassCS::GetVarianceType());
 
 			TShaderMapRef<SHADER> ComputeShader(View.ShaderMap, ComputeShaderPermutationVector);
 			FComputeShaderUtils::AddPass(
 				GraphBuilder,
-				RDG_EVENT_NAME("PathTracing::Denoising::Prepass %dx%d",
+				RDG_EVENT_NAME("PathTracing::Denoising::Prepass::Texture(Var[%s] %dx%d)",
+					FTemporalPrepassCS::GetEventName(FTemporalPrepassCS::GetVarianceType()),
 					TargetViewport.Extent.X,
 					TargetViewport.Extent.Y),
 				ComputeShader,
diff --git a/Engine/Source/Runtime/Renderer/Private/PathTracingSpatialTemporalDenoising.h b/Engine/Source/Runtime/Renderer/Private/PathTracingSpatialTemporalDenoising.h
--- a/Engine/Source/Runtime/Renderer/Private/PathTracingSpatialTemporalDenoising.h
+++ b/Engine/Source/Runtime/Renderer/Private/PathTracingSpatialTemporalDenoising.h
@@ -39,6 +39,9 @@ struct FPathTracingSpatialTemporalDenoisingContext
 	// Debug information
 	FRDGTexture*	MotionVector = nullptr;
 	FRDGTexture*	WarpedSource = nullptr;
+
+	// Utility texture
+	FRDGTexture* VarianceTexture = nullptr;
 };
 
 
@@ -61,7 +64,6 @@ void PathTracingSpatialTemporalDenoising(FRDGBuilder& GraphBuilder,
 	FRDGTexture*& SpatialTemporalDenoisedTexture,
 	FPathTracingSpatialTemporalDenoisingContext& SpatialTemporalDenoisingContext);
 
-
 struct FVisualizePathTracingDenoisingInputs
 {
 	// [Optional] Render to the specified output. If invalid, a new texture is created and returned.
diff --git a/Engine/Source/Runtime/Renderer/Private/PostProcess/PostProcessMaterial.cpp b/Engine/Source/Runtime/Renderer/Private/PostProcess/PostProcessMaterial.cpp
--- a/Engine/Source/Runtime/Renderer/Private/PostProcess/PostProcessMaterial.cpp
+++ b/Engine/Source/Runtime/Renderer/Private/PostProcess/PostProcessMaterial.cpp
@@ -593,6 +593,19 @@ FScreenPassTexture AddPostProcessMaterialPass(
 		PostProcessMaterialParameters->PostProcessInput[InputIndex] = GetScreenPassTextureInput(Input, PointClampSampler);
 	}
 
+	// Path tracing buffer textures
+	for (uint32 InputIndex = 0; InputIndex < kPathTracingPostProcessMaterialInputCountMax; ++InputIndex)
+	{
+		FScreenPassTexture Input = Inputs.GetPathTracingInput((EPathTracingPostProcessMaterialInput)InputIndex);
+
+		if (!Input.Texture || !MaterialShaderMap->UsesPathTracingBufferTexture(InputIndex))
+		{
+			Input = BlackDummy;
+		}
+
+		PostProcessMaterialParameters->PathTracingPostProcessInput[InputIndex] = GetScreenPassTextureInput(Input, PointClampSampler);
+	}
+
 	PostProcessMaterialParameters->Strata = Strata::BindStrataGlobalUniformParameters(View);
 
 	// SceneDepthWithoutWater
@@ -724,7 +737,7 @@ static FPostProcessMaterialNode* IteratePostProcessMaterialNodes(const FFinalPos
 	{
 		FPostProcessMaterialNode* DataPtr = Dest.BlendableManager.IterateBlendables<FPostProcessMaterialNode>(Iterator);
 
-		if (!DataPtr || DataPtr->GetLocation() == Location)
+		if (!DataPtr || DataPtr->GetLocation() == Location || Location == EBlendableLocation::BL_MAX)
 		{
 			return DataPtr;
 		}
@@ -748,9 +761,9 @@ FPostProcessMaterialChain GetPostProcessMaterialChain(const FViewInfo& View, EBl
 		UMaterialInterface* VisMaterial = GetBufferVisualizationData().GetMaterial(View.CurrentBufferVisualizationMode);
 		UMaterial* Material = VisMaterial ? VisMaterial->GetMaterial() : nullptr;
 
-		if (Material && Material->BlendableLocation == Location)
+		if (Material && (Material->BlendableLocation == Location || Location == EBlendableLocation::BL_MAX))
 		{
-			Nodes.Add(FPostProcessMaterialNode(Material, Location, Material->BlendablePriority, Material->bIsBlendable));
+			Nodes.Add(FPostProcessMaterialNode(Material, Material->BlendableLocation, Material->BlendablePriority, Material->bIsBlendable));
 		}
 	}
 
@@ -815,6 +828,55 @@ bool IsHighResolutionScreenshotMaskEnabled(const FViewInfo& View)
 	return View.Family->EngineShowFlags.HighResScreenshotMask || View.FinalPostProcessSettings.HighResScreenshotCaptureRegionMaterial;
 }
 
+bool IsPathTracingVarianceTextureRequiredInPostProcessMaterial(const FViewInfo& View)
+{
+	// query the post process material to check if any variance texture has been used
+	bool bIsPathTracingVarianceTextureRequired = false;
+
+	auto CheckIfPathTracingVarianceTextureIsRequried = [&](const UMaterialInterface* MaterialInterface) {
+		// Get the RenderProxy of the material.
+		const FMaterialRenderProxy* MaterialProxy = MaterialInterface->GetRenderProxy();
+		check(MaterialProxy);
+
+		// Get the Shadermap for the view's feature level
+		const FMaterial* Material = MaterialProxy->GetMaterialNoFallback(View.FeatureLevel);
+		if (Material && Material->GetMaterialDomain() == MD_PostProcess)
+		{
+			const FMaterialShaderMap* MaterialShaderMap = Material->GetRenderingThreadShaderMap();
+			check(MaterialShaderMap);
+			if (MaterialShaderMap->UsesPathTracingBufferTexture(static_cast<uint32>(EPathTracingPostProcessMaterialInput::Variance)))
+			{
+				return true;
+			}
+		}
+
+		return false;
+	};
+
+	FPostProcessMaterialChain PostProcessMaterialChain = GetPostProcessMaterialChain(View, EBlendableLocation::BL_MAX);
+	for (const UMaterialInterface* MaterialInterface : PostProcessMaterialChain)
+	{
+		if (CheckIfPathTracingVarianceTextureIsRequried(MaterialInterface))
+		{
+			bIsPathTracingVarianceTextureRequired = true;
+			break;
+		}
+	}
+
+	// Check buffer visualization pipes
+	const FFinalPostProcessSettings& PostProcessSettings = View.FinalPostProcessSettings;
+	for (const UMaterialInterface* MaterialInterface : PostProcessSettings.BufferVisualizationOverviewMaterials)
+	{
+		if (CheckIfPathTracingVarianceTextureIsRequried(MaterialInterface))
+		{
+			bIsPathTracingVarianceTextureRequired = true;
+			break;
+		}
+	}
+
+	return bIsPathTracingVarianceTextureRequired;
+}
+
 FScreenPassTexture AddHighResolutionScreenshotMaskPass(
 	FRDGBuilder& GraphBuilder,
 	const FViewInfo& View,
diff --git a/Engine/Source/Runtime/Renderer/Private/PostProcess/PostProcessMaterial.h b/Engine/Source/Runtime/Renderer/Private/PostProcess/PostProcessMaterial.h
--- a/Engine/Source/Runtime/Renderer/Private/PostProcess/PostProcessMaterial.h
+++ b/Engine/Source/Runtime/Renderer/Private/PostProcess/PostProcessMaterial.h
@@ -20,6 +20,7 @@ BEGIN_SHADER_PARAMETER_STRUCT(FPostProcessMaterialParameters, )
 	SHADER_PARAMETER_RDG_UNIFORM_BUFFER(FStrataGlobalUniformParameters, Strata)
 	SHADER_PARAMETER_STRUCT(FScreenPassTextureViewportParameters, PostProcessOutput)
 	SHADER_PARAMETER_STRUCT_ARRAY(FScreenPassTextureInput, PostProcessInput, [kPostProcessMaterialInputCountMax])
+	SHADER_PARAMETER_STRUCT_ARRAY(FScreenPassTextureInput, PathTracingPostProcessInput, [kPathTracingPostProcessMaterialInputCountMax])
 	SHADER_PARAMETER_SAMPLER(SamplerState, PostProcessInput_BilinearSampler)
 	SHADER_PARAMETER_RDG_BUFFER_SRV(StructuredBuffer<float4>, EyeAdaptationBuffer)
 	SHADER_PARAMETER(uint32, bMetalMSAAHDRDecode)
@@ -63,6 +64,7 @@ struct FHighResolutionScreenshotMaskInputs
 };
 
 bool IsHighResolutionScreenshotMaskEnabled(const FViewInfo& View);
+bool IsPathTracingVarianceTextureRequiredInPostProcessMaterial(const FViewInfo& View);
 
 FScreenPassTexture AddHighResolutionScreenshotMaskPass(
 	FRDGBuilder& GraphBuilder,
diff --git a/Engine/Source/Runtime/Renderer/Private/PostProcess/PostProcessVisualizeBuffer.cpp b/Engine/Source/Runtime/Renderer/Private/PostProcess/PostProcessVisualizeBuffer.cpp
--- a/Engine/Source/Runtime/Renderer/Private/PostProcess/PostProcessVisualizeBuffer.cpp
+++ b/Engine/Source/Runtime/Renderer/Private/PostProcess/PostProcessVisualizeBuffer.cpp
@@ -12,6 +12,7 @@
 #include "BufferVisualizationData.h"
 #include "SceneRendering.h"
 #include "UnrealEngine.h"
+#include "PathTracing.h"
 
 class FVisualizeBufferPS : public FGlobalShader
 {
@@ -374,6 +375,24 @@ FScreenPassTexture AddVisualizeGBufferOverviewPass(
 		PostProcessMaterialInputs.SetInput(EPostProcessMaterialInput::PreTonemapHDRColor, Inputs.SceneColorBeforeTonemap);
 		PostProcessMaterialInputs.SetInput(EPostProcessMaterialInput::PostTonemapHDRColor, Inputs.SceneColorAfterTonemap);
 		PostProcessMaterialInputs.SetInput(EPostProcessMaterialInput::Velocity, Inputs.Velocity);
+
+		if (View.Family->EngineShowFlags.PathTracing && Inputs.PathTracingResources->bPostProcessEnabled)
+		{
+			const FPathTracingResources& PathTracingResources = *Inputs.PathTracingResources;
+			
+			FIntRect ViewRect = Inputs.SceneColor.ViewRect;
+			PostProcessMaterialInputs.SetPathTracingInput(
+				EPathTracingPostProcessMaterialInput::DenoisedRadiance, FScreenPassTexture(PathTracingResources.DenoisedRadiance, ViewRect));
+			PostProcessMaterialInputs.SetPathTracingInput(
+				EPathTracingPostProcessMaterialInput::Albedo, FScreenPassTexture(PathTracingResources.Albedo, ViewRect));
+			PostProcessMaterialInputs.SetPathTracingInput(
+				EPathTracingPostProcessMaterialInput::Normal, FScreenPassTexture(PathTracingResources.Normal, ViewRect));
+			PostProcessMaterialInputs.SetPathTracingInput(
+				EPathTracingPostProcessMaterialInput::Variance, FScreenPassTexture(PathTracingResources.Variance, ViewRect));
+			PostProcessMaterialInputs.SetPathTracingInput(
+				EPathTracingPostProcessMaterialInput::Radiance, FScreenPassTexture(PathTracingResources.Radiance, ViewRect));
+		}
+
 		PostProcessMaterialInputs.SceneTextures = Inputs.SceneTextures;
 		PostProcessMaterialInputs.OutputFormat = OutputFormat;
 
diff --git a/Engine/Source/Runtime/Renderer/Private/PostProcess/PostProcessVisualizeBuffer.h b/Engine/Source/Runtime/Renderer/Private/PostProcess/PostProcessVisualizeBuffer.h
--- a/Engine/Source/Runtime/Renderer/Private/PostProcess/PostProcessVisualizeBuffer.h
+++ b/Engine/Source/Runtime/Renderer/Private/PostProcess/PostProcessVisualizeBuffer.h
@@ -38,6 +38,9 @@ struct FVisualizeGBufferOverviewInputs
 	/** The uniform buffer containing all scene textures. */
 	FSceneTextureShaderParameters SceneTextures;
 
+	/** The path tracing intermediate textures. */
+	const struct FPathTracingResources* PathTracingResources = nullptr;
+
 	// Dump targets to files on disk.
 	bool bDumpToFile = false;
 
diff --git a/Engine/Source/Runtime/Renderer/Private/PostProcess/PostProcessing.cpp b/Engine/Source/Runtime/Renderer/Private/PostProcess/PostProcessing.cpp
--- a/Engine/Source/Runtime/Renderer/Private/PostProcess/PostProcessing.cpp
+++ b/Engine/Source/Runtime/Renderer/Private/PostProcess/PostProcessing.cpp
@@ -401,6 +401,19 @@ void AddPostProcessingPasses(
 		PostProcessMaterialInputs.SetInput(EPostProcessMaterialInput::SceneColor, InSceneColor);
 
 		FIntRect ViewRect{ 0, 0, 1, 1 };
+
+		if (Inputs.PathTracingResources.bPostProcessEnabled)
+		{
+			const FPathTracingResources& PathTracingResources = Inputs.PathTracingResources;
+
+			ViewRect = InSceneColor.ViewRect;
+			PostProcessMaterialInputs.SetPathTracingInput(EPathTracingPostProcessMaterialInput::Radiance, FScreenPassTexture(PathTracingResources.Radiance, ViewRect));
+			PostProcessMaterialInputs.SetPathTracingInput(EPathTracingPostProcessMaterialInput::DenoisedRadiance, FScreenPassTexture(PathTracingResources.DenoisedRadiance, ViewRect));
+			PostProcessMaterialInputs.SetPathTracingInput(EPathTracingPostProcessMaterialInput::Albedo, FScreenPassTexture(PathTracingResources.Albedo, ViewRect));
+			PostProcessMaterialInputs.SetPathTracingInput(EPathTracingPostProcessMaterialInput::Normal, FScreenPassTexture(PathTracingResources.Normal, ViewRect));
+			PostProcessMaterialInputs.SetPathTracingInput(EPathTracingPostProcessMaterialInput::Variance, FScreenPassTexture(PathTracingResources.Variance, ViewRect));
+		}
+		
 		if (PostDOFTranslucencyResources.IsValid())
 		{
 			ViewRect = PostDOFTranslucencyResources.ViewRect;
@@ -1278,6 +1291,7 @@ void AddPostProcessingPasses(
 		PassInputs.bOverview = bVisualizeGBufferOverview;
 		PassInputs.bDumpToFile = bVisualizeGBufferDumpToFile;
 		PassInputs.bOutputInHDR = bOutputInHDR;
+		PassInputs.PathTracingResources = &Inputs.PathTracingResources;
 
 		SceneColor = AddVisualizeGBufferOverviewPass(GraphBuilder, View, PassInputs);
 	}
diff --git a/Engine/Source/Runtime/Renderer/Private/PostProcess/PostProcessing.h b/Engine/Source/Runtime/Renderer/Private/PostProcess/PostProcessing.h
--- a/Engine/Source/Runtime/Renderer/Private/PostProcess/PostProcessing.h
+++ b/Engine/Source/Runtime/Renderer/Private/PostProcess/PostProcessing.h
@@ -5,6 +5,7 @@
 #include "RenderGraphFwd.h"
 #include "RHIFwd.h"
 #include "TranslucentRendering.h"
+#include "PathTracing.h"
 
 enum class EReflectionsMethod;
 
@@ -37,6 +38,8 @@ struct FPostProcessingInputs
 	FRDGTextureRef CustomDepthTexture = nullptr;
 	FRDGTextureRef ExposureIlluminance = nullptr;
 	FTranslucencyViewResourcesMap TranslucencyViewResourcesMap;
+	FPathTracingResources PathTracingResources;
+
 	bool bSeparateCustomStencil = false;
 
 	void Validate() const
diff --git a/Engine/Source/Runtime/Renderer/Public/PostProcess/PostProcessMaterialInputs.h b/Engine/Source/Runtime/Renderer/Public/PostProcess/PostProcessMaterialInputs.h
--- a/Engine/Source/Runtime/Renderer/Public/PostProcess/PostProcessMaterialInputs.h
+++ b/Engine/Source/Runtime/Renderer/Public/PostProcess/PostProcessMaterialInputs.h
@@ -8,6 +8,7 @@
 struct FSceneWithoutWaterTextures;
 
 const uint32 kPostProcessMaterialInputCountMax = 5;
+const uint32 kPathTracingPostProcessMaterialInputCountMax = 5;
 
 /** Named post process material slots. Inputs are aliased and have different semantics
  *  based on the post process material blend point, which is documented with the input.
@@ -31,6 +32,15 @@ enum class EPostProcessMaterialInput : uint32
 	Velocity = 4
 };
 
+enum class EPathTracingPostProcessMaterialInput : uint32
+{
+	Radiance = 0,
+	DenoisedRadiance = 1,
+	Albedo = 2,
+	Normal = 3,
+	Variance = 4
+};
+
 struct FPostProcessMaterialInputs
 {
 	inline void SetInput(EPostProcessMaterialInput Input, FScreenPassTexture Texture)
@@ -43,6 +53,16 @@ struct FPostProcessMaterialInputs
 		return Textures[(uint32)Input];
 	}
 
+	inline void SetPathTracingInput(EPathTracingPostProcessMaterialInput Input, FScreenPassTexture Texture)
+	{
+		PathTracingTextures[(uint32)Input] = Texture;
+	}
+
+	inline FScreenPassTexture GetPathTracingInput(EPathTracingPostProcessMaterialInput Input) const
+	{
+		return PathTracingTextures[(uint32)Input];
+	}
+
 	inline void Validate() const
 	{
 		ValidateInputExists(EPostProcessMaterialInput::SceneColor);
@@ -75,6 +95,11 @@ struct FPostProcessMaterialInputs
 	 */
 	TStaticArray<FScreenPassTexture, kPostProcessMaterialInputCountMax> Textures;
 
+	/**
+	*	Array of input textures bound to the material from path tracing. All inputs are optional
+	*/
+	TStaticArray<FScreenPassTexture, kPathTracingPostProcessMaterialInputCountMax> PathTracingTextures;
+
 	/** The output texture format to use if a new texture is created. Uses the input format if left unknown. */
 	EPixelFormat OutputFormat = PF_Unknown;
 
