Implement a modular member system for the scene UB.

diff --git a/Engine/Plugins/Experimental/GPULightmass/Source/GPULightmass/Private/LightmapRenderer.cpp b/Engine/Plugins/Experimental/GPULightmass/Source/GPULightmass/Private/LightmapRenderer.cpp
--- a/Engine/Plugins/Experimental/GPULightmass/Source/GPULightmass/Private/LightmapRenderer.cpp
+++ b/Engine/Plugins/Experimental/GPULightmass/Source/GPULightmass/Private/LightmapRenderer.cpp
@@ -609,7 +609,7 @@ void FCachedRayTracingSceneData::SetupViewAndSceneUniformBufferFromSceneRenderSt
 			GPUScene.GPUSceneLightData = GraphBuilder.CreateSRV(RDGLightDataBuffer);
 		}
 
-		Scene.SceneUniforms.Set(GPUScene);
+		Scene.SceneUniforms.Set(SceneUB::GPUScene, GPUScene);
 	}
 }
 
@@ -780,7 +780,7 @@ void FCachedRayTracingSceneData::RestoreCachedBuffers(FRDGBuilder& GraphBuilder,
 {
 	TRACE_CPUPROFILER_EVENT_SCOPE(RestoreCachedBuffers);
 
-	FGPUSceneResourceParameters GPUScene = RenderState.SceneUniforms.GetParameters().GPUScene;
+	FGPUSceneResourceParameters GPUScene = RenderState.SceneUniforms.Get(SceneUB::GPUScene);
 	
 	check(GPUScenePrimitiveDataBuffer.IsValid())
 	GPUScene.GPUScenePrimitiveSceneData = GraphBuilder.CreateSRV(GraphBuilder.RegisterExternalBuffer(GPUScenePrimitiveDataBuffer));
@@ -800,8 +800,8 @@ void FCachedRayTracingSceneData::RestoreCachedBuffers(FRDGBuilder& GraphBuilder,
 
 	// Clear cached parameters so that bGPUSceneIsDirty will be true
 	// Required as the uniform buffer in SceneUniforms has only 1 frame lifetime and needs to be recreated
-	RenderState.SceneUniforms.Set({});
-	RenderState.SceneUniforms.Set(GPUScene);
+	RenderState.SceneUniforms.Set(SceneUB::GPUScene, {});
+	RenderState.SceneUniforms.Set(SceneUB::GPUScene, GPUScene);
 }
 
 FCachedRayTracingSceneData::~FCachedRayTracingSceneData()
@@ -2071,7 +2071,7 @@ void FLightmapRenderer::Finalize(FRDGBuilder& GraphBuilder)
 			TRACE_CPUPROFILER_EVENT_SCOPE(InstanceIdsIdentityBuffer);
 
 			TArray<uint32, SceneRenderingAllocator> InstanceIdsIdentity;
-			for (uint32 Index = 0U; Index < FMath::Max(1U, Scene->SceneUniforms.GetParameters().GPUScene.InstanceDataSOAStride); ++Index)
+			for (uint32 Index = 0U; Index < FMath::Max(1U, Scene->SceneUniforms.Get(SceneUB::GPUScene).InstanceDataSOAStride); ++Index)
 			{
 				InstanceIdsIdentity.Add(Index);
 			}
diff --git a/Engine/Source/Runtime/RenderCore/Private/ShaderParameterMetadataBuilder.cpp b/Engine/Source/Runtime/RenderCore/Private/ShaderParameterMetadataBuilder.cpp
--- a/Engine/Source/Runtime/RenderCore/Private/ShaderParameterMetadataBuilder.cpp
+++ b/Engine/Source/Runtime/RenderCore/Private/ShaderParameterMetadataBuilder.cpp
@@ -193,3 +193,38 @@ FShaderParametersMetadata* FShaderParametersMetadataBuilder::Build(
 
 	return ShaderParameterMetadata;
 }
+
+FShaderParametersMetadata* FShaderParametersMetadataBuilder::Build(
+	FShaderParametersMetadata::EUseCase InUseCase,
+	EUniformBufferBindingFlags InBindingFlags,
+	const TCHAR* InLayoutName,
+	const TCHAR* InStructTypeName,
+	const TCHAR* InShaderVariableName,
+	const TCHAR* InStaticSlotName,
+	const ANSICHAR* InFileName,
+	const int32 InFileLine,
+	bool bForceCompleteInitialization,
+	FRHIUniformBufferLayoutInitializer* OutLayoutInitializer,
+	uint32 InUsageFlags
+)
+{
+	const uint32 StructSize = Align(NextMemberOffset, SHADER_PARAMETER_STRUCT_ALIGNMENT);
+	
+	FShaderParametersMetadata* ShaderParameterMetadata = new FShaderParametersMetadata(
+		InUseCase,
+		InBindingFlags,
+		InLayoutName,
+		InStructTypeName,
+		InShaderVariableName,
+		InStaticSlotName,
+		InFileName,
+		InFileLine,
+		StructSize,
+		Members,
+		bForceCompleteInitialization,
+		OutLayoutInitializer,
+		InUsageFlags
+	);
+
+	return ShaderParameterMetadata;
+}
diff --git a/Engine/Source/Runtime/RenderCore/Public/ShaderParameterMetadataBuilder.h b/Engine/Source/Runtime/RenderCore/Public/ShaderParameterMetadataBuilder.h
--- a/Engine/Source/Runtime/RenderCore/Public/ShaderParameterMetadataBuilder.h
+++ b/Engine/Source/Runtime/RenderCore/Public/ShaderParameterMetadataBuilder.h
@@ -176,6 +176,20 @@ public:
 		const TCHAR* ShaderParameterName
 		);
 
+	FShaderParametersMetadata* Build(
+		FShaderParametersMetadata::EUseCase InUseCase,
+		EUniformBufferBindingFlags InBindingFlags,
+		const TCHAR* InLayoutName,
+		const TCHAR* InStructTypeName,
+		const TCHAR* InShaderVariableName,
+		const TCHAR* InStaticSlotName,
+		const ANSICHAR* InFileName,
+		const int32 InFileLine,
+		bool bForceCompleteInitialization = false,
+		FRHIUniformBufferLayoutInitializer* OutLayoutInitializer = nullptr,
+		uint32 InUsageFlags = 0
+		);
+
 private:
 	TArray<FShaderParametersMetadata::FMember> Members;
 	uint32 NextMemberOffset = 0;
diff --git a/Engine/Source/Runtime/Renderer/Private/GPUScene.cpp b/Engine/Source/Runtime/Renderer/Private/GPUScene.cpp
--- a/Engine/Source/Runtime/Renderer/Private/GPUScene.cpp
+++ b/Engine/Source/Runtime/Renderer/Private/GPUScene.cpp
@@ -26,9 +26,26 @@
 #include "RenderCore.h"
 #include "LightSceneData.h"
 #include "LightSceneProxy.h"
+#include "SystemTextures.h"
 
 #define LOG_INSTANCE_ALLOCATIONS 0
 
+static void ConstructDefault(FGPUSceneResourceParameters& GPUScene, FRDGBuilder& GraphBuilder)
+{
+	FRDGBufferRef DummyBufferVec4 = GSystemTextures.GetDefaultStructuredBuffer(GraphBuilder, sizeof(FVector4f));
+	GPUScene.GPUSceneInstanceSceneData = GraphBuilder.CreateSRV(DummyBufferVec4);
+	GPUScene.GPUSceneInstancePayloadData = GraphBuilder.CreateSRV(DummyBufferVec4);
+	GPUScene.GPUScenePrimitiveSceneData = GraphBuilder.CreateSRV(DummyBufferVec4);
+	GPUScene.GPUSceneLightmapData = GraphBuilder.CreateSRV(DummyBufferVec4);
+	FRDGBufferRef DummyBufferLight = GSystemTextures.GetDefaultStructuredBuffer(GraphBuilder, sizeof(FLightSceneData));
+	GPUScene.GPUSceneLightData = GraphBuilder.CreateSRV(DummyBufferLight);
+	GPUScene.InstanceDataSOAStride = 0;
+	GPUScene.GPUSceneFrameNumber = 0;
+	GPUScene.NumInstances = 0;
+	GPUScene.NumScenePrimitives = 0;
+}
+IMPLEMENT_SCENE_UB_STRUCT(FGPUSceneResourceParameters, GPUScene, ConstructDefault);
+
 int32 GGPUSceneUploadEveryFrame = 0;
 FAutoConsoleVariableRef CVarGPUSceneUploadEveryFrame(
 	TEXT("r.GPUScene.UploadEveryFrame"),
@@ -799,7 +816,7 @@ void FGPUScene::UpdateBufferState(FRDGBuilder& GraphBuilder, FSceneUniformBuffer
 	ShaderParameters.NumInstances = InstanceSceneDataAllocator.GetMaxSize();
 	ShaderParameters.GPUSceneFrameNumber = GetSceneFrameNumber();
 
-	SceneUB.Set(ShaderParameters);
+	SceneUB.Set(SceneUB::GPUScene, ShaderParameters);
 }
 
 /**
@@ -1585,7 +1602,7 @@ bool FGPUScene::FillSceneUniformBuffer(FRDGBuilder& GraphBuilder, FSceneUniformB
 	{
 		return false;
 	}
-	return SceneUB.Set(ShaderParameters);
+	return SceneUB.Set(SceneUB::GPUScene, ShaderParameters);
 }
 
 void FGPUScene::AddPrimitiveToUpdate(int32 PrimitiveId, EPrimitiveDirtyState DirtyState)
diff --git a/Engine/Source/Runtime/Renderer/Private/GPUScene.h b/Engine/Source/Runtime/Renderer/Private/GPUScene.h
--- a/Engine/Source/Runtime/Renderer/Private/GPUScene.h
+++ b/Engine/Source/Runtime/Renderer/Private/GPUScene.h
@@ -21,6 +21,20 @@ class FGPUScene;
 class FGPUSceneDynamicContext;
 class FViewUniformShaderParameters;
 
+BEGIN_SHADER_PARAMETER_STRUCT(FGPUSceneResourceParameters, RENDERER_API)
+	SHADER_PARAMETER_RDG_BUFFER_SRV(StructuredBuffer<float4>, GPUSceneInstanceSceneData)
+	SHADER_PARAMETER_RDG_BUFFER_SRV(StructuredBuffer<float4>, GPUSceneInstancePayloadData)
+	SHADER_PARAMETER_RDG_BUFFER_SRV(StructuredBuffer<float4>, GPUScenePrimitiveSceneData)
+	SHADER_PARAMETER_RDG_BUFFER_SRV(StructuredBuffer<float4>, GPUSceneLightmapData)
+	SHADER_PARAMETER_RDG_BUFFER_SRV(StructuredBuffer<FLightSceneData>, GPUSceneLightData)
+	SHADER_PARAMETER(uint32, InstanceDataSOAStride)
+	SHADER_PARAMETER(uint32, GPUSceneFrameNumber)
+	SHADER_PARAMETER(int32, NumInstances)
+	SHADER_PARAMETER(int32, NumScenePrimitives)
+END_SHADER_PARAMETER_STRUCT()
+
+DECLARE_SCENE_UB_STRUCT(FGPUSceneResourceParameters, GPUScene, RENDERER_API)
+
 /**
  * Used to manage dynamic primitives for a given view, during InitViews the data is collected and then can be committed to the GPU-Scene. 
  * Once committed the range of indices are valid and can be used to calculate the PrimitiveIds.
diff --git a/Engine/Source/Runtime/Renderer/Private/Renderer.cpp b/Engine/Source/Runtime/Renderer/Private/Renderer.cpp
--- a/Engine/Source/Runtime/Renderer/Private/Renderer.cpp
+++ b/Engine/Source/Runtime/Renderer/Private/Renderer.cpp
@@ -75,7 +75,7 @@ void FRendererModule::StartupModule()
 		});
 	}));
 
-	// Needs to run on startup, after static init. Ideally at some point replaced with static accessor similar to GGetDitherFadedInUniformBuffer().
+	// Needs to run on startup, after static init.
 	GIdentityPrimitiveUniformBuffer.InitContents();
 	GDistanceCullFadedInUniformBuffer.InitContents();
 	GDitherFadedInUniformBuffer.InitContents();
@@ -200,7 +200,7 @@ FSceneUniformBuffer* FRendererModule::CreateSinglePrimitiveSceneUniformBuffer(FR
 			ShaderParameters.NumInstances = 1;
 			ShaderParameters.NumScenePrimitives = 1;
 			
-			SceneUniforms.Set(ShaderParameters);
+			SceneUniforms.Set(SceneUB::GPUScene, ShaderParameters);
 		}
 	}
 
diff --git a/Engine/Source/Runtime/Renderer/Private/SceneUniformBuffer.cpp b/Engine/Source/Runtime/Renderer/Private/SceneUniformBuffer.cpp
--- a/Engine/Source/Runtime/Renderer/Private/SceneUniformBuffer.cpp
+++ b/Engine/Source/Runtime/Renderer/Private/SceneUniformBuffer.cpp
@@ -2,48 +2,215 @@
 
 #include "SceneUniformBuffer.h"
 #include "RenderGraphBuilder.h"
-#include "SystemTextures.h"
 
-#include "LightSceneData.h"
+/******************************/
+/* Struct layout registration */
+/******************************/
+
+struct FSceneUniformBufferTypeRegistry::FMemberInfo
+{
+	int32 Id;
+	int32 Offset;
+	FString Name;
+	const FShaderParametersMetadata* Metadata;
+	FShaderParameterStructConstructor DefaultValueFactory;
+};
+
+class FSceneUniformBufferTypeRegistry::FImpl
+{
+public:
+	static inline TArray<FMemberInfo> MemberInfos{};
+	static inline bool bStructMetadataWasBuilt = false;
+
+	static FShaderParametersMetadata* BuildStructMetadata();
+};
+
+FSceneUniformBuffer::FMemberId FSceneUniformBufferTypeRegistry::RegisterMember(
+	const TCHAR* Name, const FShaderParametersMetadata& Metadata, FShaderParameterStructConstructor Factory)
+{
+	check(!FImpl::bStructMetadataWasBuilt);
+
+	int NewMemberInfoIndex = FImpl::MemberInfos.Emplace();
+	FMemberInfo& NewMember = FImpl::MemberInfos[NewMemberInfoIndex];
+	NewMember.Id = NewMemberInfoIndex;
+	NewMember.Name = Name;
+	NewMember.Offset = -1;
+	NewMember.Metadata = &Metadata;
+	NewMember.DefaultValueFactory = Factory;
+	return NewMemberInfoIndex;
+}
+
+void FSceneUniformBufferMemberRegistration::CommitAll()
+{
+	for (FSceneUniformBufferMemberRegistration* Instance : GetInstances())
+	{
+		Instance->Commit();
+	}
+	GetInstances().Empty();
+}
+
+static TArray<FSceneUniformBufferMemberRegistration*>* GSceneUniformBufferMemberRegistrationInstances = nullptr;
+TArray<FSceneUniformBufferMemberRegistration*>& FSceneUniformBufferMemberRegistration::GetInstances()
+{
+	// This is not a static or static local to ensure this works correctly during static initialization
+	if (GSceneUniformBufferMemberRegistrationInstances == nullptr)
+	{
+		GSceneUniformBufferMemberRegistrationInstances = new TArray<FSceneUniformBufferMemberRegistration*>();
+	}
+	return *GSceneUniformBufferMemberRegistrationInstances;
+}
+
+FShaderParametersMetadata* FSceneUniformBufferTypeRegistry::FImpl::BuildStructMetadata()
+{
+	FSceneUniformBufferMemberRegistration::CommitAll();
+
+	// TODO: Could be nice to have more intelligent packing here, rather than just appending with padding
+	TArray<FMemberInfo*> SortedMembers;
+	SortedMembers.Reserve(MemberInfos.Num());
+	for (FMemberInfo& Member : MemberInfos)
+	{
+		SortedMembers.Emplace(&Member);
+	}
+	Algo::SortBy(SortedMembers, &FMemberInfo::Name);
+
+	FShaderParametersMetadataBuilder Builder{};
+
+	for (FMemberInfo* Member : SortedMembers)
+	{
+		Member->Offset = Builder.AddNestedStruct(*Member->Name, Member->Metadata);
+	}
+
+	bStructMetadataWasBuilt = true;
+	FShaderParametersMetadata* Metadata = Builder.Build(
+		FShaderParametersMetadata::EUseCase::UniformBuffer,
+		EUniformBufferBindingFlags::StaticAndShader,
+		TEXT("FSceneUniformParameters"),
+		TEXT("FSceneUniformParameters"),
+		TEXT("Scene"),
+		TEXT("Scene"),
+		FSceneUniformParameters::FTypeInfo::FileName,
+		FSceneUniformParameters::FTypeInfo::FileLine
+	);
+	checkf(Metadata->GetSize() <= FSceneUniformParameters::MaxSize,
+		TEXT("The total size of the scene uniform buffer members exceeds the storage capacity in FSceneUniformParameters::Data. Increase the capacity."));
+	return Metadata;
+}
+
+/*****************/
+/* Buffer values */
+/*****************/
 
 IMPLEMENT_STATIC_UNIFORM_BUFFER_SLOT(Scene);
-IMPLEMENT_STATIC_AND_SHADER_UNIFORM_BUFFER_STRUCT(FSceneUniformParameters, "Scene", Scene);
 
-FSceneUniformBuffer::FSceneUniformBuffer()
-: Buffer(nullptr)
-, RHIBuffer(nullptr)
-, bGPUSceneIsDirty(false)
+RENDERER_API const FShaderParametersMetadata* GetForwardDeclaredShaderParametersStructMetadata(const FSceneUniformParameters* DummyPtr)
 {
-	FMemory::Memzero(CachedParameters);
+	return FSceneUniformParameters::FTypeInfo::GetStructMetadata();
+}
+
+FShaderParametersMetadataRegistration SceneUniformParametersMetadataRegistration
+{ TFunctionRef<const ::FShaderParametersMetadata* ()>{ FSceneUniformParameters::FTypeInfo::GetStructMetadata } };
+
+RENDERER_API const FShaderParametersMetadata* FSceneUniformParameters::FTypeInfo::GetStructMetadata()
+{
+	static FShaderParametersMetadata* Metadata = FSceneUniformBuffer::FRegistry::FImpl::BuildStructMetadata();
+	return Metadata;
+}
+
+FSceneUniformParameters::FSceneUniformParameters()
+{ }
+
+
+
+FSceneUniformBuffer::FSceneUniformBuffer() :
+	Buffer(nullptr),
+	RHIBuffer(nullptr),
+	bAnyMemberDirty(false),
+	MemberHasBeenSet(false, FRegistry::FImpl::MemberInfos.Num())
+{
+	CachedData.Init(0, FSceneUniformParameters::FTypeInfo::GetStructMetadata()->GetSize());
+	check(Align(CachedData.GetData(), SHADER_PARAMETER_STRUCT_ALIGNMENT) == CachedData.GetData());
+}
+
+RENDERER_API bool FSceneUniformBuffer::Set(const FMemberId MemberId, const void* Value, const int32 ValueSize)
+{
+	check(MemberId >= 0 && MemberId < FRegistry::FImpl::MemberInfos.Num());
+	const FRegistry::FMemberInfo& MemberInfo = FRegistry::FImpl::MemberInfos[MemberId];
+	check(MemberInfo.Metadata->GetSize() == ValueSize);
+	check((MemberInfo.Offset + ValueSize) <= CachedData.Num());
+
+	MemberHasBeenSet[MemberId] = true;
+
+	uint8_t* MemberDataPtr = &CachedData[MemberInfo.Offset];
+	if (FPlatformMemory::Memcmp(MemberDataPtr, Value, ValueSize) != 0)
+	{
+		bAnyMemberDirty = true;
+		FPlatformMemory::Memcpy(MemberDataPtr, Value, ValueSize);
+		return true;
+	}
+	return false;
+}
+
+RENDERER_API const void* FSceneUniformBuffer::Get(const FMemberId MemberId, const int32 ValueSize) const
+{
+	check(MemberId >= 0 && MemberId < FRegistry::FImpl::MemberInfos.Num());
+	const FRegistry::FMemberInfo& MemberInfo = FRegistry::FImpl::MemberInfos[MemberId];
+	check(MemberInfo.Metadata->GetSize() == ValueSize);
+	check((MemberInfo.Offset + ValueSize) <= CachedData.Num());
+	
+	if (!MemberHasBeenSet[MemberId])
+	{
+		return nullptr;
+	}
+
+	return &CachedData[MemberInfo.Offset];
+}
+
+RENDERER_API const void* FSceneUniformBuffer::GetOrDefault(const FMemberId MemberId, const int32 ValueSize, FRDGBuilder& GraphBuilder)
+{
+	check(MemberId >= 0 && MemberId < FRegistry::FImpl::MemberInfos.Num());
+	const FRegistry::FMemberInfo& MemberInfo = FRegistry::FImpl::MemberInfos[MemberId];
+	check(MemberInfo.Metadata->GetSize() == ValueSize);
+	check((MemberInfo.Offset + ValueSize) <= CachedData.Num());
+
+	uint8* Offset = &CachedData[MemberInfo.Offset];
+	if (!MemberHasBeenSet[MemberId])
+	{
+		check((MemberInfo.Offset + MemberInfo.Metadata->GetSize()) <= static_cast<uint32>(CachedData.Num()));
+		MemberInfo.DefaultValueFactory(Offset, *MemberInfo.Metadata, GraphBuilder);
+		MemberHasBeenSet[MemberId] = true;
+	}
+
+	return Offset;
 }
 
 RENDERER_API TRDGUniformBufferRef<FSceneUniformParameters> FSceneUniformBuffer::GetBuffer(FRDGBuilder& GraphBuilder)
 {
 	check(IsInRenderingThread());
-
-	if (Buffer == nullptr || bGPUSceneIsDirty)
+	if (Buffer == nullptr || bAnyMemberDirty)
 	{
-		// Never been populated, set up defaults, we defer this since we don't want to create redundant SRVs for the common case where it is not needed.
-		if (!bGPUSceneIsDirty)
+		for (int MemberIndex = 0; MemberIndex < FRegistry::FImpl::MemberInfos.Num(); ++MemberIndex)
 		{
-			FGPUSceneResourceParameters GPUScene{};
-			FRDGBufferRef DummyBufferVec4 = GSystemTextures.GetDefaultStructuredBuffer(GraphBuilder, sizeof(FVector4f));
-			CachedParameters.GPUScene.GPUSceneInstanceSceneData = GraphBuilder.CreateSRV(DummyBufferVec4);
-			CachedParameters.GPUScene.GPUSceneInstancePayloadData = GraphBuilder.CreateSRV(DummyBufferVec4);
-			CachedParameters.GPUScene.GPUScenePrimitiveSceneData = GraphBuilder.CreateSRV(DummyBufferVec4);
-			CachedParameters.GPUScene.GPUSceneLightmapData = GraphBuilder.CreateSRV(DummyBufferVec4);
-			FRDGBufferRef DummyBufferLight = GSystemTextures.GetDefaultStructuredBuffer(GraphBuilder, sizeof(FLightSceneData));
-			CachedParameters.GPUScene.GPUSceneLightData = GraphBuilder.CreateSRV(DummyBufferLight);
-			CachedParameters.GPUScene.InstanceDataSOAStride = 0;
-			CachedParameters.GPUScene.GPUSceneFrameNumber = 0;
-			CachedParameters.GPUScene.NumInstances = 0;
-			CachedParameters.GPUScene.NumScenePrimitives = 0;
+			const FRegistry::FMemberInfo& MemberInfo = FRegistry::FImpl::MemberInfos[MemberIndex];
+			if (!MemberHasBeenSet[MemberIndex])
+			{
+				// Never been populated, set up defaults, we defer this since we don't want to create redundant SRVs for the common case where it is not needed.
+				uint8* Offset = &CachedData[MemberInfo.Offset];
+
+				check((MemberInfo.Offset + MemberInfo.Metadata->GetSize()) <= static_cast<uint32>(CachedData.Num()));
+				MemberInfo.DefaultValueFactory(Offset, *MemberInfo.Metadata, GraphBuilder);
+				MemberHasBeenSet[MemberIndex] = true;
+			}
 		}
 		// Create and copy cached parameters into the RDG-lifetime struct
-		Buffer = GraphBuilder.CreateUniformBuffer(GraphBuilder.AllocObject<FSceneUniformParameters>(CachedParameters));
+		//auto DataCopy = GraphBuilder.Alloc(CachedData.Num(), SHADER_PARAMETER_STRUCT_ALIGNMENT);
+		//FPlatformMemory::Memcpy(DataCopy, CachedData.GetData(), CachedData.Num());
+		//Buffer = GraphBuilder.CreateUniformBuffer(GraphBuilder.AllocObject<FSceneUniformParameters>(DataCopy));
+		FSceneUniformParameters* ParameterStruct = GraphBuilder.AllocObject<FSceneUniformParameters>();
+		FPlatformMemory::Memcpy(ParameterStruct->Data, CachedData.GetData(), CachedData.Num());
+		Buffer = GraphBuilder.CreateUniformBuffer(ParameterStruct);
 		RHIBuffer = GraphBuilder.ConvertToExternalUniformBuffer(Buffer); //RT pipeline can't bind RDG UBs, so we must convert to external
 
-		bGPUSceneIsDirty = false;
+		bAnyMemberDirty = false;
 	}
 	return Buffer;
 }
@@ -55,9 +222,18 @@ RENDERER_API FRHIUniformBuffer* FSceneUniformBuffer::GetBufferRHI(FRDGBuilder& G
 	return RHIBuffer;
 }
 
-bool FSceneUniformBuffer::Set(const FGPUSceneResourceParameters& GPUScene)
+#if !UE_BUILD_SHIPPING
+TArray<FSceneUniformBuffer::FDebugMemberInfo> FSceneUniformBuffer::GetDebugInfo() const
 {
-	bGPUSceneIsDirty = bGPUSceneIsDirty || FMemory::Memcmp(&CachedParameters.GPUScene, &GPUScene, sizeof(GPUScene)) != 0;
-	CachedParameters.GPUScene = GPUScene;
-	return bGPUSceneIsDirty;
+	TArray<FSceneUniformBuffer::FDebugMemberInfo> DebugInfo;
+	for (const auto& Member : FSceneUniformParameters::FTypeInfo::GetStructMetadata()->GetMembers())
+	{
+		FSceneUniformBuffer::FDebugMemberInfo Info;
+		Info.ShaderType = Member.GetShaderType();
+		Info.Name = Member.GetName();
+		Info.Data = &CachedData[Member.GetOffset()];
+		DebugInfo.Add(Info);
+	}
+	return DebugInfo;
 }
+#endif
diff --git a/Engine/Source/Runtime/Renderer/Private/SceneUniformBuffer.h b/Engine/Source/Runtime/Renderer/Private/SceneUniformBuffer.h
--- a/Engine/Source/Runtime/Renderer/Private/SceneUniformBuffer.h
+++ b/Engine/Source/Runtime/Renderer/Private/SceneUniformBuffer.h
@@ -4,25 +4,128 @@
 
 #include "RenderGraphDefinitions.h"
 #include "ShaderParameterMacros.h"
+#include "ShaderParameterMetadata.h"
+#include "ShaderParameterMetadataBuilder.h"
 
-BEGIN_SHADER_PARAMETER_STRUCT(FGPUSceneResourceParameters, RENDERER_API)
-	SHADER_PARAMETER_RDG_BUFFER_SRV(StructuredBuffer<float4>, GPUSceneInstanceSceneData)
-	SHADER_PARAMETER_RDG_BUFFER_SRV(StructuredBuffer<float4>, GPUSceneInstancePayloadData)
-	SHADER_PARAMETER_RDG_BUFFER_SRV(StructuredBuffer<float4>, GPUScenePrimitiveSceneData)
-	SHADER_PARAMETER_RDG_BUFFER_SRV(StructuredBuffer<float4>, GPUSceneLightmapData)
-	SHADER_PARAMETER_RDG_BUFFER_SRV(StructuredBuffer<FLightSceneData>, GPUSceneLightData)
-	SHADER_PARAMETER(uint32, InstanceDataSOAStride)
-	SHADER_PARAMETER(uint32, GPUSceneFrameNumber)
-	SHADER_PARAMETER(int32, NumInstances)
-	SHADER_PARAMETER(int32, NumScenePrimitives)
-END_SHADER_PARAMETER_STRUCT()
+// Construct a shader parameter struct with the specified type at Obj.
+// The resulting struct must be valid for usage (i.e. all members initialized)
+using FShaderParameterStructConstructor = TFunction<void(void* Obj, const FShaderParametersMetadata& Metadata, FRDGBuilder& GraphBuilder)>;
+template<typename TParameterStruct>
+using TShaderParameterStructConstructor = TFunction<void(TParameterStruct& Obj, FRDGBuilder& GraphBuilder)>;
 
 /**
- * The RDG data struct that is used in FSceneUniformBuffer.
+ * Collects scene UB members during static initialization,
+ * and uses these to construct the buffer layout during Main.
  */
-BEGIN_UNIFORM_BUFFER_STRUCT(FSceneUniformParameters, RENDERER_API)
-	SHADER_PARAMETER_STRUCT(FGPUSceneResourceParameters, GPUScene)
-END_UNIFORM_BUFFER_STRUCT()
+class FSceneUniformBufferTypeRegistry
+{
+public:
+	using FMemberId = int;
+	struct FMemberInfo;
+	class FImpl;
+
+	// Must be called during static initialization
+	static FMemberId RegisterMember(
+		const TCHAR* Name,
+		const FShaderParametersMetadata& StructMetadata,
+		// Used to set default values for this member if no value is set explicitly
+		FShaderParameterStructConstructor DefaultValueConstructor);
+};
+
+class FSceneUniformBufferMemberRegistration
+{
+public:
+	using FMemberId = FSceneUniformBufferTypeRegistry::FMemberId;
+
+	// Called at runtime
+	virtual void Commit() = 0;
+	static void CommitAll();
+
+	FString Name;
+	FMemberId MemberId = -1;
+	FShaderParameterStructConstructor DefaultValueConstructor;
+
+protected:
+	static TArray<FSceneUniformBufferMemberRegistration*>& GetInstances();
+
+	// Called during static initialization
+	FSceneUniformBufferMemberRegistration(const TCHAR* Name)
+		: Name(Name)
+	{
+		GetInstances().Add(this);
+	}
+	virtual ~FSceneUniformBufferMemberRegistration() = default;
+};
+
+template<typename TMember>
+class TSceneUniformBufferMemberRegistration : public FSceneUniformBufferMemberRegistration
+{
+public:
+	TSceneUniformBufferMemberRegistration(const TCHAR* Name, TShaderParameterStructConstructor<TMember> TemplatedDefaultValueConstructor)
+		: FSceneUniformBufferMemberRegistration(Name)
+	{
+		DefaultValueConstructor = [=](void* Obj, const FShaderParametersMetadata& Metadata, FRDGBuilder& GraphBuilder)
+		{
+			check(&Metadata == TMember::FTypeInfo::GetStructMetadata());
+			TMember& Member = *(new(Obj) TMember());
+			TemplatedDefaultValueConstructor(Member, GraphBuilder);
+		};
+	}
+
+	void Commit() override
+	{
+		MemberId = FSceneUniformBufferTypeRegistry::RegisterMember(
+			*Name,
+			*TShaderParameterStructTypeInfo<TMember>::GetStructMetadata(),
+			DefaultValueConstructor);
+	}
+};
+
+#define DECLARE_SCENE_UB_STRUCT(StructType, FieldName, PrefixKeywords) \
+	namespace SceneUB { \
+		PrefixKeywords extern TSceneUniformBufferMemberRegistration<StructType> FieldName; \
+	}
+
+#define IMPLEMENT_SCENE_UB_STRUCT(StructType, FieldName, DefaultValueFactoryType) \
+	TSceneUniformBufferMemberRegistration<StructType> SceneUB::FieldName { TEXT(#FieldName), DefaultValueFactoryType } 
+
+
+DECLARE_UNIFORM_BUFFER_STRUCT(FSceneUniformParameters, RENDERER_API)
+/**
+ * RDG shader parameter struct containing data for FSceneUniformBuffer.
+ */
+MS_ALIGN(SHADER_PARAMETER_STRUCT_ALIGNMENT) class RENDERER_API FSceneUniformParameters final
+{
+public:
+	struct RENDERER_API FTypeInfo {
+		static constexpr int32 NumRows = 1;
+		static constexpr int32 NumColumns = 1;
+		static constexpr int32 NumElements = 0;
+		static constexpr int32 Alignment = SHADER_PARAMETER_STRUCT_ALIGNMENT;
+		static constexpr bool bIsStoredInConstantBuffer = true;
+		static constexpr const ANSICHAR* const FileName = UE_LOG_SOURCE_FILE(__FILE__);
+		static constexpr int32 FileLine = __LINE__;
+		using TAlignedType = FSceneUniformParameters;
+		static const FShaderParametersMetadata* GetStructMetadata();
+	};
+
+	const void* GetContents() const
+	{
+		return Data;
+	}
+
+	FSceneUniformParameters();
+
+private:
+	friend class FSceneUniformBuffer;
+	friend class FSceneUniformBufferTypeRegistry::FImpl;
+
+	// TODO: Decouple parameter struct and actual data, so the data can be allocated dynamically (variable size)
+	// This requires some work in buffer binding/parameter macros/RHI, as these all assume the data to be inline.
+	const static size_t MaxSize = 4096;
+	uint8_t Data[MaxSize];
+};
+
 
 /**
  * Holds scene-scoped parameters and stores these in uniform (constant) buffers for access on GPU.
@@ -30,28 +133,87 @@ END_UNIFORM_BUFFER_STRUCT()
 class FSceneUniformBuffer final
 {
 public:
+	using FMemberId = FSceneUniformBufferTypeRegistry::FMemberId;
+	using FRegistry = FSceneUniformBufferTypeRegistry;
+
 	RENDERER_API FSceneUniformBuffer();
 
+	/**
+	 * Set a field in the parameter struct.
+	 * The change will be reflected in any buffer that GetBuffer() returns after this call.
+	 * Returns true if anything actually changed.
+	 */
+	template<typename TMember>
+	bool Set(const TSceneUniformBufferMemberRegistration<TMember>& Registration, const TMember& Value)
+	{
+		return Set(Registration.MemberId, &Value, TMember::FTypeInfo::GetStructMetadata()->GetSize());
+	}
+
+	/**
+	 * Retrieve a field in the parameter struct.
+	 * If the field has not been set, this will crash.
+	 */
+	template<typename TMember>
+	const TMember& Get(const TSceneUniformBufferMemberRegistration<TMember>& Registration) const
+	{
+		const void* Ptr = Get(Registration.MemberId, TMember::FTypeInfo::GetStructMetadata()->GetSize());
+		checkf(Ptr, TEXT("SceneUB::%s has not been set yet"), *Registration.Name);
+		return *reinterpret_cast<const TMember*>(Ptr);
+	}
+
+	/**
+	 * Retrieve a field in the parameter struct.
+	 * If the field has not been set, this will set a default value.
+	 */
+	template<typename TMember>
+	const TMember& GetOrDefault(const TSceneUniformBufferMemberRegistration<TMember>& Registration, FRDGBuilder& GraphBuilder) const
+	{
+		return *reinterpret_cast<const TMember*>(GetOrDefault(Registration.MemberId, TMember::FTypeInfo::GetStructMetadata()->GetSize(), GraphBuilder));
+	}
+
 	// Get and re-create the UB if the cached parameters are modified.
 	RENDERER_API TRDGUniformBufferRef<FSceneUniformParameters> GetBuffer(FRDGBuilder& GraphBuilder);
 	RENDERER_API FRHIUniformBuffer* GetBufferRHI(FRDGBuilder& GraphBuilder);
 
-	// Read-only view into the parameters. Use Set() to make changes.
-	const FSceneUniformParameters& GetParameters() const
+#if !UE_BUILD_SHIPPING
+	struct FDebugMemberInfo
 	{
-		return CachedParameters;
-	}
+		FString ShaderType;
+		FString Name;
+		const void* Data;
+	};
+	TArray<FDebugMemberInfo> GetDebugInfo() const;
+#endif
+
+private:
+	template<typename TMember>
+	friend class TSceneUniformBufferMemberRegistration;
+	friend class FSceneUniformParameters;
 
 	/**
 	 * Set a field in the parameter struct.
 	 * The change will be reflected in any buffer that GetBuffer() returns after this call.
 	 * Returns true if anything actually changed.
 	 */
-	RENDERER_API bool Set(const FGPUSceneResourceParameters& GPUScene);
+	RENDERER_API bool Set(const FMemberId MemberId, const void* Value, const int32 ValueSize);
+
+	/**
+	 * Retrieve a field in the parameter struct.
+	 * If the field has not been set, this will return null;
+	 */
+	RENDERER_API const void* Get(const FMemberId MemberId, const int32 ExpectedValueSize) const;
+
+	/**
+	 * Retrieve a field in the parameter struct.
+	 * If the field has not been set, this will set a default value;
+	 */
+	RENDERER_API const void* GetOrDefault(const FMemberId MemberId, const int32 ExpectedValueSize, FRDGBuilder& GraphBuilder);
 
-private:
-	FSceneUniformParameters CachedParameters;
 	TRDGUniformBufferRef<FSceneUniformParameters> Buffer;
 	FRHIUniformBuffer* RHIBuffer;
-	bool bGPUSceneIsDirty;
+
+	TBitArray<> MemberIsDirty;
+	bool bAnyMemberDirty;
+	TBitArray<> MemberHasBeenSet;
+	TArray<uint8_t, TSizedHeapAllocator<32>> CachedData;
 };
diff --git a/Engine/Source/Runtime/Renderer/Public/SceneRendererInterface.h b/Engine/Source/Runtime/Renderer/Public/SceneRendererInterface.h
--- a/Engine/Source/Runtime/Renderer/Public/SceneRendererInterface.h
+++ b/Engine/Source/Runtime/Renderer/Public/SceneRendererInterface.h
@@ -15,6 +15,8 @@ DECLARE_UNIFORM_BUFFER_STRUCT(FSceneUniformParameters, RENDERER_API)
 class ISceneRenderer
 {
 public:
+	virtual ~ISceneRenderer() = default;
+
 	virtual const FSceneUniformBuffer& GetSceneUniforms() const = 0;
 	virtual FSceneUniformBuffer& GetSceneUniforms() = 0;
 
