Use Scene UB in CullPerPageDrawCommands instead of GLOBAL_GPU_SCENE_DATA

diff --git a/Engine/Shaders/Private/VirtualShadowMaps/VirtualShadowMapBuildPerPageDrawCommands.usf b/Engine/Shaders/Private/VirtualShadowMaps/VirtualShadowMapBuildPerPageDrawCommands.usf
--- a/Engine/Shaders/Private/VirtualShadowMaps/VirtualShadowMapBuildPerPageDrawCommands.usf
+++ b/Engine/Shaders/Private/VirtualShadowMaps/VirtualShadowMapBuildPerPageDrawCommands.usf
@@ -24,7 +24,6 @@ struct FVisibleInstanceCmd
 
 RWStructuredBuffer<FVisibleInstanceCmd> VisibleInstancesOut;
 RWStructuredBuffer<uint> VisibleInstanceCountBufferOut;
-uint InstanceSceneDataSOAStride;
 uint TotalPrimaryViews;
 uint VisibleInstancesBufferNum;
 
@@ -253,7 +252,7 @@ void CullPerPageDrawCommandsCs(uint3 GroupId : SV_GroupID, int GroupThreadIndex
 		FDrawCommandDesc DrawCommandDesc = UnpackDrawCommandDesc(DrawCommandDescs[Payload.IndirectArgIndex]);
 
 		// Load relevant instance data
-		FInstanceSceneData InstanceData = GetInstanceSceneData(InstanceId, InstanceSceneDataSOAStride);
+		FInstanceSceneData InstanceData = GetInstanceSceneData(InstanceId, Scene.GPUScene.InstanceDataSOAStride);
 
 		FPrimitiveSceneData PrimitiveData = GetPrimitiveData(InstanceData.PrimitiveId);
 		bool bMaterialInvalidates = DrawCommandDesc.bMaterialUsesWorldPositionOffset && (PrimitiveData.Flags & PRIMITIVE_SCENE_DATA_FLAG_EVALUATE_WORLD_POSITION_OFFSET) != 0u;;
diff --git a/Engine/Source/Runtime/Renderer/Private/SceneUniformBuffer.h b/Engine/Source/Runtime/Renderer/Private/SceneUniformBuffer.h
--- a/Engine/Source/Runtime/Renderer/Private/SceneUniformBuffer.h
+++ b/Engine/Source/Runtime/Renderer/Private/SceneUniformBuffer.h
@@ -122,7 +122,7 @@ private:
 
 	// TODO: Decouple parameter struct and actual data, so the data can be allocated dynamically (variable size)
 	// This requires some work in buffer binding/parameter macros/RHI, as these all assume the data to be inline.
-	const static size_t MaxSize = 4096;
+	static constexpr size_t MaxSize = 4096;
 	uint8_t Data[MaxSize];
 };
 
diff --git a/Engine/Source/Runtime/Renderer/Private/VirtualShadowMaps/VirtualShadowMapArray.cpp b/Engine/Source/Runtime/Renderer/Private/VirtualShadowMaps/VirtualShadowMapArray.cpp
--- a/Engine/Source/Runtime/Renderer/Private/VirtualShadowMaps/VirtualShadowMapArray.cpp
+++ b/Engine/Source/Runtime/Renderer/Private/VirtualShadowMaps/VirtualShadowMapArray.cpp
@@ -2137,7 +2137,6 @@ public:
 		OutEnvironment.SetDefine(TEXT("NANITE_MULTI_VIEW"), 1);
 		OutEnvironment.SetDefine(TEXT("INDIRECT_ARGS_NUM_WORDS"), FInstanceCullingContext::IndirectArgsNumWords);
 		OutEnvironment.SetDefine(TEXT("VF_SUPPORTS_PRIMITIVE_SCENE_DATA"), 1);
-		OutEnvironment.SetDefine(TEXT("USE_GLOBAL_GPU_SCENE_DATA"), 1);
 	}
 
 
@@ -2156,11 +2155,7 @@ public:
 	BEGIN_SHADER_PARAMETER_STRUCT(FParameters, )
 		SHADER_PARAMETER_RDG_UNIFORM_BUFFER(FVirtualShadowMapUniformParameters, VirtualShadowMap)
 
-		SHADER_PARAMETER_RDG_BUFFER_SRV(StructuredBuffer<float4>, GPUSceneInstanceSceneData)
-		SHADER_PARAMETER_RDG_BUFFER_SRV(StructuredBuffer<float4>, GPUSceneInstancePayloadData)
-		SHADER_PARAMETER_RDG_BUFFER_SRV(StructuredBuffer<float4>, GPUScenePrimitiveSceneData)
-		SHADER_PARAMETER(uint32, InstanceSceneDataSOAStride)
-		SHADER_PARAMETER(uint32, GPUSceneFrameNumber)
+		SHADER_PARAMETER_RDG_UNIFORM_BUFFER(FSceneUniformParameters, Scene)
 
 		SHADER_PARAMETER_RDG_BUFFER_SRV(StructuredBuffer<uint>, PrimitiveRevealedMask)
 		SHADER_PARAMETER(uint32, PrimitiveRevealedNum)
@@ -2318,7 +2313,8 @@ static FCullingResult AddCullingPasses(FRDGBuilder& GraphBuilder,
 	FRDGBufferRef VirtualShadowViewsRDG,
 	const FCullPerPageDrawCommandsCs::FHZBShaderParameters &HZBShaderParameters,
 	FVirtualShadowMapArray &VirtualShadowMapArray,
-	FGPUScene& GPUScene,
+	FSceneUniformBuffer& SceneUniformBuffer,
+	ERHIFeatureLevel::Type FeatureLevel,
 	const TConstArrayView<uint32> PrimitiveRevealedMask)
 {
 	const bool bUseBatchMode = !BatchInds.IsEmpty();
@@ -2345,7 +2341,7 @@ static FCullingResult AddCullingPasses(FRDGBuilder& GraphBuilder,
 	CullingResult.DrawIndirectArgsRDG = GraphBuilder.CreateBuffer(IndirectArgsDesc, TEXT("Shadow.Virtual.DrawIndirectArgsBuffer"));
 	GraphBuilder.QueueBufferUpload(CullingResult.DrawIndirectArgsRDG, IndirectArgs.GetData(), IndirectArgs.GetTypeSize() * IndirectArgs.Num());
 
-	FGlobalShaderMap* ShaderMap = GetGlobalShaderMap(GPUScene.GetFeatureLevel());
+	FGlobalShaderMap* ShaderMap = GetGlobalShaderMap(FeatureLevel);
 
 	// Note: we redundantly clear the instance counts here as there is some issue with replays on certain consoles.
 	FInstanceCullingContext::AddClearIndirectArgInstanceCountPass(GraphBuilder, ShaderMap, CullingResult.DrawIndirectArgsRDG);
@@ -2366,14 +2362,7 @@ static FCullingResult AddCullingPasses(FRDGBuilder& GraphBuilder,
 		FCullPerPageDrawCommandsCs::FParameters* PassParameters = GraphBuilder.AllocParameters<FCullPerPageDrawCommandsCs::FParameters>();
 
 		PassParameters->VirtualShadowMap = VirtualShadowMapArray.GetUniformBuffer();
-
-		const FGPUSceneResourceParameters GPUSceneParameters = GPUScene.GetShaderParameters();
-
-		PassParameters->GPUSceneInstanceSceneData = GPUSceneParameters.GPUSceneInstanceSceneData;
-		PassParameters->GPUSceneInstancePayloadData = GPUSceneParameters.GPUSceneInstancePayloadData;
-		PassParameters->GPUScenePrimitiveSceneData = GPUSceneParameters.GPUScenePrimitiveSceneData;
-		PassParameters->GPUSceneFrameNumber = GPUSceneParameters.GPUSceneFrameNumber;
-		PassParameters->InstanceSceneDataSOAStride = GPUSceneParameters.InstanceDataSOAStride;
+		PassParameters->Scene = SceneUniformBuffer.GetBuffer(GraphBuilder);
 
 		// Make sure there is enough space in the buffer for all the primitive IDs that might be used to index, at least in the first batch...
 		check(PrimitiveRevealedMaskRdg->Desc.NumElements * 32u >= uint32(VSMCullingBatchInfos[0].PrimitiveRevealedNum));
@@ -2467,7 +2456,7 @@ static FCullingResult AddCullingPasses(FRDGBuilder& GraphBuilder,
 	CullingResult.InstanceIdsBuffer = GraphBuilder.CreateBuffer(FRDGBufferDesc::CreateStructuredDesc(sizeof(uint32), CullingResult.MaxNumInstancesPerPass), TEXT("Shadow.Virtual.InstanceIdsBuffer"));
 	CullingResult.PageInfoBuffer = GraphBuilder.CreateBuffer(FRDGBufferDesc::CreateStructuredDesc(sizeof(uint32), CullingResult.MaxNumInstancesPerPass), TEXT("Shadow.Virtual.PageInfoBuffer"));
 
-	FRDGBufferRef OutputPassIndirectArgs = FComputeShaderUtils::AddIndirectArgsSetupCsPass1D(GraphBuilder, GPUScene.GetFeatureLevel(), VisibleInstanceWriteOffsetRDG, TEXT("Shadow.Virtual.IndirectArgs"), FOutputCommandInstanceListsCs::NumThreadsPerGroup);
+	FRDGBufferRef OutputPassIndirectArgs = FComputeShaderUtils::AddIndirectArgsSetupCsPass1D(GraphBuilder, FeatureLevel, VisibleInstanceWriteOffsetRDG, TEXT("Shadow.Virtual.IndirectArgs"), FOutputCommandInstanceListsCs::NumThreadsPerGroup);
 	{
 
 		FOutputCommandInstanceListsCs::FParameters* PassParameters = GraphBuilder.AllocParameters<FOutputCommandInstanceListsCs::FParameters>();
@@ -2939,7 +2928,8 @@ void FVirtualShadowMapArray::RenderVirtualShadowMapsNonNanite(FRDGBuilder& Graph
 			VirtualShadowViewsRDG,
 			HZBShaderParameters,
 			*this,
-			GPUScene,
+			SceneUniformBuffer,
+			GPUScene.GetFeatureLevel(),
 			PrimitiveRevealedMask
 		);
 		GraphBuilder.EndEventScope();
@@ -3068,7 +3058,8 @@ void FVirtualShadowMapArray::RenderVirtualShadowMapsNonNanite(FRDGBuilder& Graph
 				VirtualShadowViewsRDG,
 				HZBShaderParameters,
 				*this,
-				GPUScene,
+				SceneUniformBuffer,
+				GPUScene.GetFeatureLevel(),
 				Clipmap->GetRevealedPrimitivesMask()
 			);
 
