GPU Lightmass: Use a dedicated payload for GPU lightmass

Split the logic for GPULM shading out from the path tracing shader to make it easier to follow. This is also in preparation of Strata support which would otherwise introduce yet another dimension into the matrix of possibilities.

Given that the payload is now different, we can also afford to make GPULM use a much smaller payload which should improve performance.

diff --git a/Engine/Plugins/Experimental/GPULightmass/Shaders/Private/LightmapPathTracing.usf b/Engine/Plugins/Experimental/GPULightmass/Shaders/Private/LightmapPathTracing.usf
--- a/Engine/Plugins/Experimental/GPULightmass/Shaders/Private/LightmapPathTracing.usf
+++ b/Engine/Plugins/Experimental/GPULightmass/Shaders/Private/LightmapPathTracing.usf
@@ -396,15 +396,9 @@ void PathTracingKernel(
 			Payload.TranslatedWorldPos = TranslatedWorldPosition;
 			Payload.WorldNormal = ShadingNormal;
 			Payload.WorldGeoNormal = ShadingNormal;
-			Payload.WorldSmoothNormal = ShadingNormal;
 			Payload.Radiance = 0;
 			Payload.BaseColor = 1;
-			Payload.DiffuseColor = 1;
-			Payload.SpecularColor = 0;
-			Payload.Specular = 0;
-			Payload.Roughness = 1;
-			Payload.Metallic = 0;
-			Payload.Ior = 1;
+			Payload.SubsurfaceColor = 0;
 			Payload.BSDFOpacity = 1;
 			Payload.TransparencyColor = 0;
 			Payload.ShadingModelID = SHADINGMODELID_DEFAULT_LIT;
diff --git a/Engine/Plugins/Experimental/GPULightmass/Source/GPULightmass/Private/LightmapRayTracing.h b/Engine/Plugins/Experimental/GPULightmass/Source/GPULightmass/Private/LightmapRayTracing.h
--- a/Engine/Plugins/Experimental/GPULightmass/Source/GPULightmass/Private/LightmapRayTracing.h
+++ b/Engine/Plugins/Experimental/GPULightmass/Source/GPULightmass/Private/LightmapRayTracing.h
@@ -69,7 +69,7 @@ class FLightmapPathTracingRGS : public FGlobalShader
 
 	static ERayTracingPayloadType GetRayTracingPayloadType(const int32 PermutationId)
 	{
-		return ERayTracingPayloadType::PathTracingMaterial;
+		return ERayTracingPayloadType::GPULightmass;
 	}
 
 	BEGIN_SHADER_PARAMETER_STRUCT(FParameters, )
@@ -135,7 +135,7 @@ class FVolumetricLightmapPathTracingRGS : public FGlobalShader
 
 	static ERayTracingPayloadType GetRayTracingPayloadType(const int32 PermutationId)
 	{
-		return ERayTracingPayloadType::PathTracingMaterial;
+		return ERayTracingPayloadType::GPULightmass;
 	}
 
 	BEGIN_SHADER_PARAMETER_STRUCT(FParameters, )
@@ -196,7 +196,7 @@ class FStationaryLightShadowTracingRGS : public FGlobalShader
 
 	static ERayTracingPayloadType GetRayTracingPayloadType(const int32 PermutationId)
 	{
-		return ERayTracingPayloadType::PathTracingMaterial;
+		return ERayTracingPayloadType::GPULightmass;
 	}
 
 	BEGIN_SHADER_PARAMETER_STRUCT(FParameters, )
@@ -269,7 +269,7 @@ class FStaticShadowDepthMapTracingRGS : public FGlobalShader
 
 	static ERayTracingPayloadType GetRayTracingPayloadType(const int32 PermutationId)
 	{
-		return ERayTracingPayloadType::PathTracingMaterial;
+		return ERayTracingPayloadType::GPULightmass;
 	}
 
 	BEGIN_SHADER_PARAMETER_STRUCT(FParameters, )
diff --git a/Engine/Plugins/Experimental/GPULightmass/Source/GPULightmass/Private/LightmapRenderer.cpp b/Engine/Plugins/Experimental/GPULightmass/Source/GPULightmass/Private/LightmapRenderer.cpp
--- a/Engine/Plugins/Experimental/GPULightmass/Source/GPULightmass/Private/LightmapRenderer.cpp
+++ b/Engine/Plugins/Experimental/GPULightmass/Source/GPULightmass/Private/LightmapRenderer.cpp
@@ -1235,7 +1235,7 @@ bool FSceneRenderState::SetupRayTracingScene(FRDGBuilder& GraphBuilder, int32 LO
 			
 			FRayTracingPipelineStateInitializer PSOInitializer;
 
-			PSOInitializer.MaxPayloadSizeInBytes = GetRayTracingPayloadTypeMaxSize(ERayTracingPayloadType::PathTracingMaterial);
+			PSOInitializer.MaxPayloadSizeInBytes = GetRayTracingPayloadTypeMaxSize(ERayTracingPayloadType::GPULightmass);
 			PSOInitializer.bAllowHitGroupIndexing = true;
 
 			FGlobalShaderMap* GlobalShaderMap = GetGlobalShaderMap(FeatureLevel);
@@ -1266,13 +1266,13 @@ bool FSceneRenderState::SetupRayTracingScene(FRDGBuilder& GraphBuilder, int32 LO
 			}
 			PSOInitializer.SetRayGenShaderTable(RayGenShaderTable);
 
-			auto DefaultClosestHitShader = GetPathTracingDefaultOpaqueHitShader(GlobalShaderMap);
+			auto DefaultClosestHitShader = GetGPULightmassDefaultOpaqueHitShader(GlobalShaderMap);
 			TArray<FRHIRayTracingShader*> RayTracingHitGroupLibrary;
 			FShaderMapResource::GetRayTracingHitGroupLibrary(RayTracingHitGroupLibrary, DefaultClosestHitShader);
 
 			PSOInitializer.SetHitGroupTable(RayTracingHitGroupLibrary);
 
-			FRHIRayTracingShader* MissTable[] = { GetPathTracingDefaultMissShader(GlobalShaderMap) };
+			FRHIRayTracingShader* MissTable[] = { GetGPULightmassDefaultMissShader(GlobalShaderMap) };
 			PSOInitializer.SetMissShaderTable(MissTable);
 
 			RayTracingPipelineState = PipelineStateCache::GetAndOrCreateRayTracingPipelineState(RHICmdList, PSOInitializer);
diff --git a/Engine/Shaders/Private/PathTracing/PathTracingCommon.ush b/Engine/Shaders/Private/PathTracing/PathTracingCommon.ush
--- a/Engine/Shaders/Private/PathTracing/PathTracingCommon.ush
+++ b/Engine/Shaders/Private/PathTracing/PathTracingCommon.ush
@@ -119,12 +119,21 @@ float3 AdjustShadingNormal(float3 ShadingNormal, float3 GeoNormal, float3 RayDir
 }
 
 
+#if IS_PAYLOAD_ENABLED(RT_PAYLOAD_PATH_TRACING_MATERIAL) && IS_PAYLOAD_ENABLED(RT_PAYLOAD_TYPE_GPULIGHTMASS)
+#error "Path tracing and GPULightmass payloads are mutually exclusive. They should not be both enabled at once."
+#endif
+
 // This payload structure is what we transport between RGS/CHS/AHS programs
 struct FPackedPathTracingPayload : FMinimalPayload
 {
 	// float FMinimalPayload::HitT                       // 4  bytes
+#if IS_PAYLOAD_ENABLED(RT_PAYLOAD_TYPE_GPULIGHTMASS)
+	uint PackedData[7];                                  // 28 bytes (encoded data, depends on shading model and ray type)
+														 // 32 bytes total
+#else
 	uint PackedData[15];								 // 60 bytes (encoded data, depends on shading model and ray type)
 														 // 64 bytes total
+#endif
 
 	uint GetFlags() { return (PackedData[0] >> 20) & 0x1FF; }
 	void SetFlag(uint Flag) { PackedData[0] |= Flag << 20; }
@@ -167,6 +176,9 @@ struct FPackedPathTracingPayload : FMinimalPayload
 		PackedData[4] = asuint(PathRoughness);
 	}
 
+#if IS_PAYLOAD_ENABLED(RT_PAYLOAD_TYPE_GPULIGHTMASS)
+	// skip some unused methods in GPULM case
+#else
 	// tau: optical depth along shadow ray, this is the integral of the extinction coefficient
 	float3 GetTau()
 	{
@@ -183,10 +195,17 @@ struct FPackedPathTracingPayload : FMinimalPayload
 		PackedData[6] = asuint(Tau.y);
 		PackedData[7] = asuint(Tau.z);
 	}
+#endif
 };
+
 #if IS_PAYLOAD_ENABLED(RT_PAYLOAD_PATH_TRACING_MATERIAL)
 CHECK_RT_PAYLOAD_SIZE(FPackedPathTracingPayload)
 #endif
+#if IS_PAYLOAD_ENABLED(RT_PAYLOAD_TYPE_GPULIGHTMASS)
+CHECK_RT_PAYLOAD_SIZE(FPackedPathTracingPayload)
+#endif
+
+
 
 // This payload structure is the expanded version of the above which is more convenient to work with
 struct FPathTracingPayload : FMinimalPayload
@@ -194,25 +213,33 @@ struct FPathTracingPayload : FMinimalPayload
 	float3 Radiance;
 	float3 WorldGeoNormal;           // normal of the actual triangle (faceted)
 	float3 WorldNormal;              // normal output of the material (includes normal/bump)
+#if IS_PAYLOAD_ENABLED(RT_PAYLOAD_TYPE_GPULIGHTMASS)
+	// skip smooth normal for GPULightmass
+#else
 	float3 WorldSmoothNormal;        // smooth normal before normal/bump is applied
+#endif
 	float3 BaseColor;
+	float  BSDFOpacity;              // how much should we weigh down the bsdf by? (for glass, controls the blend between solid and glass lobes)
+	float3 TransparencyColor;        // how much do we see straight through the surface?
+	uint   ShadingModelID;
+	uint   Flags;
+	uint   PrimitiveLightingChannelMask;
+#if IS_PAYLOAD_ENABLED(RT_PAYLOAD_TYPE_GPULIGHTMASS)
+	float3 SubsurfaceColor;
+#else
 	float3 DiffuseColor;
 	float3 SpecularColor;
-	float3 TransparencyColor;        // how much do we see straight through the surface?
-	float  BSDFOpacity;              // how much should we weigh down the bsdf by? (for glass, controls the blend between solid and glass lobes)
 	float  Metallic;
 	float  Specular;
 	float  Roughness;
 	float  Anisotropy;
 	float  Ior;
-	uint   ShadingModelID;
-	uint   Flags;
-	uint   PrimitiveLightingChannelMask;
 	float4 CustomData0;
 	float4 CustomData1;
 	float3 WorldTangent;
 	float3 SubsurfaceColor;
 	float3 SubsurfaceRadius;
+#endif
 	float3 TranslatedWorldPos;
 
 	void SetFrontFace() { Flags |= PATH_TRACING_PAYLOAD_OUTPUT_FLAG_FRONT_FACE; }
@@ -227,6 +254,13 @@ struct FPathTracingPayload : FMinimalPayload
 	void SetMaterialTwoSided() { Flags |= PATH_TRACING_PAYLOAD_OUTPUT_FLAG_TWO_SIDED_MATERIAL; }
 	bool IsMaterialTwoSided() { return (Flags & PATH_TRACING_PAYLOAD_OUTPUT_FLAG_TWO_SIDED_MATERIAL) != 0; }
 
+#if IS_PAYLOAD_ENABLED(RT_PAYLOAD_TYPE_GPULIGHTMASS)
+	float3 GetSubsurfaceColor() { return SubsurfaceColor; }
+	bool IsMaterialTransmissive()
+	{
+		return ShadingModelID == SHADINGMODELID_TWOSIDED_FOLIAGE;
+	}
+#else
 	// Various ways to interpret CustomData (depending on ShadingModelID)
 	// NOTE: This is not always following the same conventions as GetMaterialCustomData0,1()
 
@@ -314,6 +348,7 @@ struct FPathTracingPayload : FMinimalPayload
 			   ShadingModelID == SHADINGMODELID_SUBSURFACE_PROFILE ||
 			   ShadingModelID == SHADINGMODELID_EYE;
 	}
+#endif
 };
 
 FPackedPathTracingPayload InitPathTracingPayload(uint ScatterType, float PathRoughness)
@@ -356,6 +391,14 @@ FPackedPathTracingPayload PackPathTracingPayload(FPathTracingPayload Input)
 	Output.PackedData[0] |= (Input.Flags & 0x1FF) << 20;                      // 9 bits
 	Output.PackedData[0] |= (Input.PrimitiveLightingChannelMask & 0x7) << 29; // 3 bits
 	                                                                   // total 32 bits
+#if IS_PAYLOAD_ENABLED(RT_PAYLOAD_TYPE_GPULIGHTMASS)
+	Output.PackedData[1] = PayloadEncodeUnitVector(Input.WorldGeoNormal);
+	Output.PackedData[2] = PayloadEncodeUnitVector(Input.WorldNormal);
+	Output.PackedData[3] = PayloadEncodeHDRColor(Input.Radiance);
+	Output.PackedData[4] = PayloadEncodeLDRColor(Input.BaseColor);
+	Output.PackedData[5] = PayloadEncodeLDRColor(Input.TransparencyColor);
+	Output.PackedData[6] = PayloadEncodeLDRColor(Input.SubsurfaceColor);
+#else
 	Output.PackedData[1] = PayloadEncodeUnitVector(Input.WorldGeoNormal);
 	Output.PackedData[2] = PayloadEncodeUnitVector(Input.WorldSmoothNormal);
 	Output.PackedData[3] = PayloadEncodeUnitVector(Input.WorldNormal);
@@ -383,7 +426,7 @@ FPackedPathTracingPayload PackPathTracingPayload(FPathTracingPayload Input)
 
 	Output.PackedData[14] = f32tof16(Input.Ior);
 	// NOTE: 16 bits left in PackedData[14]
-
+#endif
 	return Output;
 }
 
@@ -396,7 +439,14 @@ FPathTracingPayload UnpackPathTracingPayload(FPackedPathTracingPayload Input, FR
 	Output.ShadingModelID = (Input.PackedData[0] >> 16) & 0xF;
 	Output.Flags = Input.GetFlags();
 	Output.PrimitiveLightingChannelMask = Input.GetPrimitiveLightingChannelMask();
-
+#if IS_PAYLOAD_ENABLED(RT_PAYLOAD_TYPE_GPULIGHTMASS)
+	Output.WorldGeoNormal = PayloadDecodeUnitVector(Input.PackedData[1]);
+	Output.WorldNormal = PayloadDecodeUnitVector(Input.PackedData[2]);
+	Output.Radiance = PayloadDecodeHDRColor(Input.PackedData[3]);
+	Output.BaseColor = PayloadDecodeLDRColor(Input.PackedData[4]);
+	Output.TransparencyColor = PayloadDecodeLDRColor(Input.PackedData[5]);
+	Output.SubsurfaceColor = PayloadDecodeLDRColor(Input.PackedData[6]);
+#else
 	Output.WorldGeoNormal = PayloadDecodeUnitVector(Input.PackedData[1]);
 	Output.WorldSmoothNormal = PayloadDecodeUnitVector(Input.PackedData[2]);
 	Output.WorldNormal = PayloadDecodeUnitVector(Input.PackedData[3]);
@@ -424,10 +474,9 @@ FPathTracingPayload UnpackPathTracingPayload(FPackedPathTracingPayload Input, FR
 
 	Output.Ior = f16tof32(Input.PackedData[14]);
 
-	Output.TranslatedWorldPos = Ray.Origin + Output.HitT * Ray.Direction;
-
 	Output.DiffuseColor = Output.BaseColor - Output.BaseColor * Output.Metallic;
 	Output.SpecularColor = ComputeF0(Output.Specular, Output.BaseColor, Output.Metallic);
-
+#endif
+	Output.TranslatedWorldPos = Ray.Origin + Output.HitT * Ray.Direction;
 	return Output;
 }
diff --git a/Engine/Shaders/Private/PathTracing/PathTracingGPULightmassMaterialHitShader.usf b/Engine/Shaders/Private/PathTracing/PathTracingGPULightmassMaterialHitShader.usf
new file mode 100644
--- /dev/null
+++ b/Engine/Shaders/Private/PathTracing/PathTracingGPULightmassMaterialHitShader.usf
@@ -0,0 +1,335 @@
+// Copyright Epic Games, Inc. All Rights Reserved.
+
+#define PATH_TRACING 1
+#define ENABLE_SKY_LIGHT 1      
+#define NEEDS_LIGHTMAP_COORDINATE 0
+#ifdef NEEDS_VERTEX_FACTORY_INTERPOLATION
+#undef NEEDS_VERTEX_FACTORY_INTERPOLATION
+#endif
+// Needed for VertexFactoryInterpolate to interpolate attributes from vertices to hit point
+#define NEEDS_VERTEX_FACTORY_INTERPOLATION 1
+
+// This should be good enough for path tracing and avoids having to bind an extra buffer
+#define EYE_ADAPTATION_PREV_FRAME_EXPOSURE 1
+
+#include "/Engine/Private/Common.ush"
+#include "/Engine/Private/RayTracing/RayTracingCommon.ush"
+#include "/Engine/Private/RayTracing/RayTracingHitGroupCommon.ush"
+#include "/Engine/Private/PathTracing/PathTracingShaderUtils.ush"
+#include "/Engine/Generated/Material.ush"   
+#include "/Engine/Generated/VertexFactory.ush"
+
+#include "/Engine/Private/RayTracing/RayTracingCalcInterpolants.ush"
+#include "/Engine/Private/ShadingCommon.ush"
+#include "/Engine/Private/DeferredShadingCommon.ush"
+#include "/Engine/Private/SubsurfaceProfileCommon.ush"
+#include "/Engine/Private/BurleyNormalizedSSSCommon.ush"
+#include "/Engine/Private/PathTracing/Material/PathTracingFresnel.ush"
+#include "/Engine/Private/PathTracing/Material/PathTracingThinGlass.ush"
+
+
+RAY_TRACING_ENTRY_CLOSEST_HIT(GPULightmassMaterialCHS,
+	FPackedPathTracingPayload, PackedPayload,
+	FRayTracingIntersectionAttributes, Attributes)
+{
+	PackedPayload.HitT = RayTCurrent();
+#if MATERIALBLENDING_SOLID || MATERIALBLENDING_ALPHAHOLDOUT || (MATERIALBLENDING_MASKED && !MATERIAL_DITHER_OPACITY_MASK)
+	if (PackedPayload.IsVisibilityRay())
+	{
+		// Shadow ray case -- all these blend modes get treated as fully opaque, so exit early.
+		// In the case of Masked materials without dithered opacity, the masking logic was already handled by AHS.
+		PackedPayload.SetRayThroughput(0.0);
+		return;
+	}
+#endif
+#if MATERIALBLENDING_ADDITIVE
+	if (PackedPayload.IsVisibilityRay())
+	{
+		// should be fully invisible to shadow rays
+		// this case should be caught by making the mesh invisible to shadows, but it doesn't cost too much to have this extra check here just in case that fails
+		return;
+	}
+#endif
+	ResolvedView = ResolveView();
+
+	const float3 TranslatedWorldPosition = TranslatedWorldRayOrigin() + RayTCurrent() * WorldRayDirection();
+	const float4 SvPosition = mul(float4(TranslatedWorldPosition, 1.0f), ResolvedView.TranslatedWorldToClip);
+
+	CurrentPayloadInputFlags = PackedPayload.GetFlags();
+
+#if VF_SUPPORTS_RAYTRACING_PREPARE_MATERIAL_PIXEL_PARAMETERS
+	// this is a newer codepath that is both more flexible and allows for more direct calculation compared to the other codepath
+	// TODO: implement such a method for all vertex factories
+	float3 GeoNormal = 0;
+	FMaterialPixelParameters MaterialParameters = GetMaterialPixelParameters(TranslatedWorldRayOrigin(), WorldRayDirection(), RayTCurrent(), PrimitiveIndex(), Attributes, HitKind(), GeoNormal);
+#else
+	FVertexFactoryInterpolantsVSToPS Interpolants;
+	float3 GeoNormal = 0;
+	CalcInterpolants((FRayCone)0, Attributes, Interpolants, GeoNormal);
+
+	FMaterialPixelParameters MaterialParameters = GetMaterialPixelParameters(Interpolants, SvPosition);
+#endif
+
+	FPixelMaterialInputs PixelMaterialInputs;
+
+	const bool bIsFrontFace = HitKind() == HIT_KIND_TRIANGLE_FRONT_FACE;
+	{
+		const float4 ScreenPosition = SvPositionToResolvedScreenPosition(SvPosition);
+		CalcMaterialParametersEx(MaterialParameters, PixelMaterialInputs, SvPosition, ScreenPosition, bIsFrontFace, TranslatedWorldPosition, TranslatedWorldPosition);
+	}
+
+	// the following blend modes need to process shadows after having executed the material
+#if MATERIALBLENDING_MODULATE || (MATERIALBLENDING_MASKED && MATERIAL_DITHER_OPACITY_MASK) || MATERIALBLENDING_ALPHACOMPOSITE || MATERIALBLENDING_TRANSLUCENT
+	if (PackedPayload.IsVisibilityRay())
+	{
+#if MATERIALBLENDING_MODULATE
+		const float3 Transparency = GetMaterialEmissive(PixelMaterialInputs);
+#elif MATERIALBLENDING_ALPHACOMPOSITE
+		const float Opacity = GetMaterialOpacity(PixelMaterialInputs);
+		const float Transparency = 1 - Opacity;
+#elif MATERIALBLENDING_MASKED && MATERIAL_DITHER_OPACITY_MASK
+		// See MATERIAL_DITHER_OPACITY_MASK comment below
+		const float Opacity = saturate(GetMaterialMaskInputRaw(PixelMaterialInputs));
+		const float Transparency = 1 - Opacity;
+#elif MATERIALBLENDING_TRANSLUCENT
+		const float Opacity = GetMaterialOpacity(PixelMaterialInputs);
+		// GPULM has backwards compatibility with some cpu lightmass tricks
+		float3 Transparency = 1 - Opacity;
+		uint ShadingModelID = GetMaterialShadingModel(PixelMaterialInputs);
+
+#if MATERIAL_SHADINGMODEL_THIN_TRANSLUCENT
+		if (ShadingModelID == SHADINGMODELID_THIN_TRANSLUCENT)
+		{
+			float3 Transmission = GetThinTranslucentMaterialOutput0(MaterialParameters);
+			float Ior = 0.0;
+			float3 V = WorldRayDirection();
+			float3 N = normalize(MaterialParameters.WorldNormal);
+			float VoN = abs(dot(V, N));
+			// simplified logic with no bending at the interface and no fresnel
+			Transparency *= pow(Transmission, 1.0 / VoN);
+		}
+		else
+#endif // MATERIAL_SHADINGMODEL_THIN_TRANSLUCENT
+#if MATERIAL_SHADINGMODEL_DEFAULT_LIT
+		if (ShadingModelID == SHADINGMODELID_DEFAULT_LIT)
+		{
+			// CPU Lightmass translucency behavior emulation
+			Transparency = lerp(float3(1, 1, 1), GetMaterialBaseColor(PixelMaterialInputs), Opacity);
+		}
+		else
+#endif// MATERIAL_SHADINGMODEL_DEFAULT_LIT
+		{
+			// base case for shadingmodel if/else
+		}
+#else // MATERIALBLEINDG_*
+#error Unhandled blending mode!
+#endif
+
+		// Update the ray throughput (it is packed simply into the payload since we don't need to carry any other information across hits)
+		float3 RayThroughput = PackedPayload.GetRayThroughput();
+		RayThroughput *= Transparency;
+		PackedPayload.SetRayThroughput(RayThroughput);
+		return;
+	}
+#endif
+
+	FPathTracingPayload Payload = (FPathTracingPayload)0;
+
+	/**
+	 * Set common material attributes for both full and simplified materials
+	 **/
+	Payload.ShadingModelID = GetMaterialShadingModel(PixelMaterialInputs);
+
+#if MATERIALBLENDING_ALPHACOMPOSITE
+	Payload.BSDFOpacity = 1.0;
+	Payload.TransparencyColor = 1.0 - GetMaterialOpacity(PixelMaterialInputs);
+#elif MATERIALBLENDING_ALPHAHOLDOUT
+	Payload.BSDFOpacity = GetMaterialOpacity(PixelMaterialInputs);
+	Payload.TransparencyColor = 1.0 - GetMaterialOpacity(PixelMaterialInputs);
+	Payload.SetHoldout();
+	HLSL_STATIC_ASSERT(MATERIAL_SHADINGMODEL_UNLIT == 1, "Alpha holdout blend mode requires unlit shading model");
+	Payload.ShadingModelID = SHADINGMODELID_UNLIT;
+#elif MATERIALBLENDING_TRANSLUCENT
+	Payload.BSDFOpacity = GetMaterialOpacity(PixelMaterialInputs);
+	Payload.TransparencyColor = 1.0 - Payload.BSDFOpacity;
+#elif MATERIALBLENDING_ADDITIVE
+	Payload.BSDFOpacity = GetMaterialOpacity(PixelMaterialInputs);
+	Payload.TransparencyColor = 1.0;
+#elif MATERIALBLENDING_MODULATE
+	Payload.BSDFOpacity = 0.0;
+	Payload.TransparencyColor = GetMaterialEmissive(PixelMaterialInputs);
+	HLSL_STATIC_ASSERT(MATERIAL_SHADINGMODEL_UNLIT == 1, "Modulate blend mode requires unlit shading model");
+	Payload.ShadingModelID = SHADINGMODELID_UNLIT;
+#elif MATERIALBLENDING_MASKED && MATERIAL_DITHER_OPACITY_MASK
+	// dithering emulates real transparency, so switch to translucent
+	// NOTE: the raster path technically takes into account the opacity mask clip value, so the effective transparency should be:
+	//        saturate(MaskRaw - ClipValue + 0.5)
+	// (See derivation in DitheredOpacityMaskToOpacity)
+	// However this behavior is surprising to most users and does not exactly match the rasterizer anyway due to how the realtime AA
+	// code performs blending.
+	// Since the goal of dithered opacity is to emulate ordinary transparency, just use the mask input as opacity directly and
+	// ignore the configured clip value.
+	Payload.BSDFOpacity = saturate(GetMaterialMaskInputRaw(PixelMaterialInputs));
+	Payload.TransparencyColor = 1.0 - Payload.BSDFOpacity;
+#elif MATERIALBLENDING_SOLID || MATERIALBLENDING_MASKED
+	Payload.BSDFOpacity = 1.0;
+	Payload.TransparencyColor = 0.0;
+#else
+#error Unknown material blending mode
+#endif
+
+	// fetch primitive flags only once
+	// TODO: would be nice to keep this inside MaterialParameters as it is also needed there as well
+	const uint PrimitiveFlags = GetPrimitiveData(MaterialParameters.PrimitiveId).Flags;
+
+	Payload.PrimitiveLightingChannelMask = GetPrimitive_LightingChannelMask_FromFlags(PrimitiveFlags);
+
+	Payload.HitT = RayTCurrent();
+	if (HitKind() == HIT_KIND_TRIANGLE_FRONT_FACE)
+	{
+		Payload.SetFrontFace();
+	}
+
+#if MATERIAL_IS_SKY
+	if (!PackedPayload.IsCameraRay())
+	{
+		// avoid double counting what was captured by the skylight
+		// also avoid noise from hot spots (they can be properly
+		// importance sampled if a capturing skylight is added)
+		PackedPayload = PackPathTracingPayload(Payload);
+		return;
+	}
+#endif
+
+	// Store the results in local variables and reuse instead of calling the functions multiple times.
+	half3 BaseColor = GetMaterialBaseColor(PixelMaterialInputs);
+	half  Metallic = GetMaterialMetallic(PixelMaterialInputs);
+	half  Specular = GetMaterialSpecular(PixelMaterialInputs);
+	half Roughness = GetMaterialRoughness(PixelMaterialInputs);
+	float Ior = 0.0;
+
+	Payload.TranslatedWorldPos = LWCToFloat(LWCAdd(MaterialParameters.AbsoluteWorldPosition, ResolvedView.PreViewTranslation));
+	float GeoNormalSign = MaterialParameters.TwoSidedSign;
+#if !VF_SUPPORTS_RAYTRACING_PREPARE_MATERIAL_PIXEL_PARAMETERS
+	// Because the geometric normal is computed directly in world space
+	// it doesn't reflect the sign flip from the object transform, so apply it here
+	GeoNormalSign *= GetPrimitive_DeterminantSign(MaterialParameters.PrimitiveId);
+#endif
+	Payload.WorldGeoNormal = GeoNormalSign * GeoNormal;
+	Payload.WorldNormal = MaterialParameters.WorldNormal;
+
+	/**
+	 * Set material attributes for simplified materials
+	 **/
+#if MATERIALBLENDING_TRANSLUCENT && !MATERIAL_SHADINGMODEL_THIN_TRANSLUCENT
+	 // Force opacity to 0 so it is ignored for path hits and only evaluated in AHS as a modulation to throughput
+	Payload.BSDFOpacity = 0;
+	Payload.TransparencyColor = 1;
+#else
+	Payload.BSDFOpacity = GetMaterialOpacity(PixelMaterialInputs);
+	Payload.TransparencyColor = 1 - Payload.BSDFOpacity;
+#endif
+
+	// Anything unsupported will be forced to default lit
+#if MATERIAL_SHADINGMODEL_TWOSIDED_FOLIAGE
+	if (Payload.ShadingModelID == SHADINGMODELID_TWOSIDED_FOLIAGE)
+	{
+		Payload.SubsurfaceColor = GetMaterialSubsurfaceData(PixelMaterialInputs).rgb;
+	}
+	else
+#endif
+#if MATERIAL_SHADINGMODEL_THIN_TRANSLUCENT
+		if (Payload.ShadingModelID == SHADINGMODELID_THIN_TRANSLUCENT)
+		{
+			const float3 Transmission = GetThinTranslucentMaterialOutput0(MaterialParameters);
+			const float3 V = WorldRayDirection();
+			const float3 N = normalize(MaterialParameters.WorldNormal);
+			const float VoN = abs(dot(V, N));
+			// simplified logic with no bending at the interface and no fresnel
+			Payload.TransparencyColor *= pow(Transmission, 1.0 / VoN);
+		}
+		else
+#endif
+#if MATERIAL_SHADINGMODEL_UNLIT
+			if (Payload.ShadingModelID == SHADINGMODELID_UNLIT)
+			{
+			}
+			else
+#endif
+			{
+				// Redirect translucent blending mode to thin translucent
+#if MATERIALBLENDING_TRANSLUCENT
+				Payload.ShadingModelID = SHADINGMODELID_THIN_TRANSLUCENT;
+#else
+				Payload.ShadingModelID = SHADINGMODELID_DEFAULT_LIT;
+#endif
+			}
+	Payload.Radiance = GetMaterialEmissive(PixelMaterialInputs);
+	Payload.Radiance *= Payload.BSDFOpacity; // premultiply
+#if MATERIAL_TWOSIDED
+	Payload.SetMaterialTwoSided();
+#else
+	if (MaterialParameters.TwoSidedSign < 0)
+	{
+		// when viewing the surface from "inside", don't include emission
+		Payload.Radiance = 0;
+	}
+#endif
+	Payload.BaseColor = BaseColor;
+
+	PackedPayload = PackPathTracingPayload(Payload);
+}
+
+#if USE_MATERIAL_ANY_HIT_SHADER
+
+RAY_TRACING_ENTRY_ANY_HIT(GPULightmassMaterialAHS,
+	FPackedPathTracingPayload, PackedPayload,
+	FRayTracingIntersectionAttributes, Attributes)
+{
+#if MATERIALBLENDING_MASKED && !MATERIAL_DITHER_OPACITY_MASK
+	// This is the only case which actually needs to run the full material
+	ResolvedView = ResolveView();
+
+	const float3 TranslatedWorldPosition = TranslatedWorldRayOrigin() + RayTCurrent() * WorldRayDirection();
+	const float4 SvPosition = mul(float4(TranslatedWorldPosition, 1.0f), ResolvedView.TranslatedWorldToClip);
+
+	CurrentPayloadInputFlags = PackedPayload.GetFlags();
+
+#if VF_SUPPORTS_RAYTRACING_PREPARE_MATERIAL_PIXEL_PARAMETERS
+	// this is a newer codepath that is both more flexible and allows for more direct calculation compared to the other codepath
+	// TODO: implement such a method for all vertex factories
+	float3 GeoNormal = 0;
+	FMaterialPixelParameters MaterialParameters = GetMaterialPixelParameters(TranslatedWorldRayOrigin(), WorldRayDirection(), RayTCurrent(), PrimitiveIndex(), Attributes, HitKind(), GeoNormal);
+#else
+	FVertexFactoryInterpolantsVSToPS Interpolants;
+	float3 GeoNormal = 0;
+	CalcInterpolants((FRayCone)0, Attributes, Interpolants, GeoNormal);
+
+	FMaterialPixelParameters MaterialParameters = GetMaterialPixelParameters(Interpolants, SvPosition);
+#endif
+
+	FPixelMaterialInputs PixelMaterialInputs;
+
+	const bool bIsFrontFace = HitKind() == HIT_KIND_TRIANGLE_FRONT_FACE;
+
+	{
+		const float4 ScreenPosition = SvPositionToResolvedScreenPosition(SvPosition);
+
+		MaterialParameters.CameraVector = -WorldRayDirection();
+
+		CalcMaterialParametersEx(MaterialParameters, PixelMaterialInputs, SvPosition, ScreenPosition, bIsFrontFace, TranslatedWorldPosition, TranslatedWorldPosition);
+	}
+
+	// Regardless of payload flags -- we always apply this 
+	if (GetMaterialMask(PixelMaterialInputs) < 0)
+	{
+		IgnoreHit();
+	}
+#else
+	// All other blending modes need to run CHS to decide what to do, or are handled via the default opaque/translucent CHS/AHS materials.
+	// in these cases, the AHS should be disabled, so emit an error if we get here somehow
+#error Unexpected blend mode encountered in AHS!
+#endif
+}
+
+#endif // USE_MATERIAL_ANY_HIT_SHADER
diff --git a/Engine/Shaders/Private/PathTracing/PathTracingMaterialHitShader.usf b/Engine/Shaders/Private/PathTracing/PathTracingMaterialHitShader.usf
--- a/Engine/Shaders/Private/PathTracing/PathTracingMaterialHitShader.usf
+++ b/Engine/Shaders/Private/PathTracing/PathTracingMaterialHitShader.usf
@@ -138,8 +138,8 @@ RAY_TRACING_ENTRY_CLOSEST_HIT(PathTracingMaterialCHS,
 		const float Transparency = 1 - Opacity;
 #elif MATERIALBLENDING_TRANSLUCENT
 		const float Opacity = GetMaterialOpacity(PixelMaterialInputs);
-#if MATERIAL_SHADINGMODEL_THIN_TRANSLUCENT || SIMPLIFIED_MATERIAL_SHADER
-		// We can get colored shadows either for thin translucent, or for GPULM which has backwards compatibility with some cpu lightmass tricks
+#if MATERIAL_SHADINGMODEL_THIN_TRANSLUCENT
+		// We can only get colored shadows for thin translucent (solid translucent case happens via absorption only)
 		float3 Transparency = 1 - Opacity;
 #else
 		float Transparency = 1 - Opacity;
@@ -150,7 +150,7 @@ RAY_TRACING_ENTRY_CLOSEST_HIT(PathTracingMaterialCHS,
 		if (ShadingModelID == SHADINGMODELID_THIN_TRANSLUCENT)
 		{
 			float3 Transmission = GetThinTranslucentMaterialOutput0(MaterialParameters);
-#if REFRACTION_USE_INDEX_OF_REFRACTION && !SIMPLIFIED_MATERIAL_SHADER
+#if REFRACTION_USE_INDEX_OF_REFRACTION
 			float Ior = GetRefractionIor(PixelMaterialInputs);
 #else
 			float Ior = 0.0;
@@ -158,10 +158,6 @@ RAY_TRACING_ENTRY_CLOSEST_HIT(PathTracingMaterialCHS,
 			float3 V = WorldRayDirection();
 			float3 N = normalize(MaterialParameters.WorldNormal);
 			float VoN = abs(dot(V, N));
-#if SIMPLIFIED_MATERIAL_SHADER
-			// simplified logic with no bending at the interface and no fresnel
-			Transparency *= pow(Transmission, 1.0 / VoN);
-#else
 			if (Opacity < 1.0)
 			{
 				float PathRoughness = PackedPayload.GetPathRoughness();
@@ -185,17 +181,12 @@ RAY_TRACING_ENTRY_CLOSEST_HIT(PathTracingMaterialCHS,
 					Transparency *= (1 - Roughness * Roughness) * saturate(PathRoughness - Roughness);
 				}
 			}
-#endif
 		}
 		else
 #endif // MATERIAL_SHADINGMODEL_THIN_TRANSLUCENT
-#if MATERIAL_SHADINGMODEL_DEFAULT_LIT
+#if MATERIAL_SHADINGMODEL_DEFAULT_LIT && REFRACTION_USE_INDEX_OF_REFRACTION
 		if (ShadingModelID == SHADINGMODELID_DEFAULT_LIT)
 		{
-#if SIMPLIFIED_MATERIAL_SHADER
-			// CPU Lightmass translucency behavior emulation
-			Transparency = lerp(float3(1, 1, 1), GetMaterialBaseColor(PixelMaterialInputs), Opacity);
-#elif REFRACTION_USE_INDEX_OF_REFRACTION
 			// Is refraction enabled?
 			float Ior = GetRefractionIor(PixelMaterialInputs);
 			if (Transparency > 0 && Ior > 0.0)
@@ -241,14 +232,12 @@ RAY_TRACING_ENTRY_CLOSEST_HIT(PathTracingMaterialCHS,
 				PackedPayload.SetTau(PackedPayload.GetTau() + LocalSigmaT * Distance);
 #endif // HAVE_GetAbsorptionMediumMaterialOutput0
 			}
-#endif // REFRACTION_USE_INDEX_OF_REFRACTION
 		}
-#else // MATERIAL_SHADINGMODEL_DEFAULT_LIT
+		else
+#endif // MATERIAL_SHADINGMODEL_DEFAULT_LIT
 		{
 			// base case for shadingmodel if/else
 		}
-#endif // MATERIAL_SHADINGMODEL_DEFAULT_LIT
-
 #else // MATERIALBLEINDG_*
 #error Unhandled blending mode!
 #endif
@@ -352,12 +341,6 @@ RAY_TRACING_ENTRY_CLOSEST_HIT(PathTracingMaterialCHS,
 	Payload.WorldNormal = MaterialParameters.WorldNormal;
 	Payload.WorldSmoothNormal = MaterialParameters.TwoSidedSign * TransformTangentNormalToWorld(MaterialParameters.TangentToWorld, float3(0, 0, 1));
 
-
-#if !SIMPLIFIED_MATERIAL_SHADER
-	/**
-	 * Set material attributes for full materials
-	 **/
-
 	// fix dark edges in reflections from shading normals
 	Payload.WorldNormal       = AdjustShadingNormal(Payload.WorldNormal      , Payload.WorldGeoNormal, WorldRayDirection());
 	Payload.WorldSmoothNormal = AdjustShadingNormal(Payload.WorldSmoothNormal, Payload.WorldGeoNormal, WorldRayDirection());
@@ -709,74 +692,6 @@ RAY_TRACING_ENTRY_CLOSEST_HIT(PathTracingMaterialCHS,
 	// adjust after everything else (because SSS profile case needs to decode the dual spec info from the unmodified roughness)
 	Payload.Roughness = AdjustMaterialRoughness(Payload.Roughness, PackedPayload.GetPathRoughness());
 
-#else // SIMPLIFIED_MATERIAL_SHADER
-
-	/**
-	 * Set material attributes for simplified materials
-	 **/
-
-#if MATERIALBLENDING_TRANSLUCENT && !MATERIAL_SHADINGMODEL_THIN_TRANSLUCENT
-	 // Force opacity to 0 so it is ignored for path hits and only evaluated in AHS as a modulation to throughput
-	Payload.BSDFOpacity = 0;
-	Payload.TransparencyColor = 1;
-#else
-	Payload.BSDFOpacity = GetMaterialOpacity(PixelMaterialInputs);
-	Payload.TransparencyColor = 1 - Payload.BSDFOpacity;
-#endif
-
-	 // Anything unsupported will be forced to default lit
-#if MATERIAL_SHADINGMODEL_TWOSIDED_FOLIAGE
-	if (Payload.ShadingModelID == SHADINGMODELID_TWOSIDED_FOLIAGE)
-	{
-		Payload.SetSubsurfaceColor(GetMaterialSubsurfaceData(PixelMaterialInputs).rgb);
-	}
-	else
-#endif
-#if MATERIAL_SHADINGMODEL_THIN_TRANSLUCENT
-	if (Payload.ShadingModelID == SHADINGMODELID_THIN_TRANSLUCENT)
-	{
-		const float3 Transmission = GetThinTranslucentMaterialOutput0(MaterialParameters);
-		const float3 V = WorldRayDirection();
-		const float3 N = normalize(MaterialParameters.WorldNormal);
-		const float VoN = abs(dot(V, N));
-		// simplified logic with no bending at the interface and no fresnel
-		Payload.TransparencyColor *= pow(Transmission, 1.0 / VoN);
-	}
-	else
-#endif
-#if MATERIAL_SHADINGMODEL_UNLIT
-	if (Payload.ShadingModelID == SHADINGMODELID_UNLIT)
-	{
-	}
-	else
-#endif
-	{
-		// Redirect translucent blending mode to thin translucent
-#if MATERIALBLENDING_TRANSLUCENT
-		Payload.ShadingModelID = SHADINGMODELID_THIN_TRANSLUCENT;
-#else
-		Payload.ShadingModelID = SHADINGMODELID_DEFAULT_LIT;
-#endif
-	}
-	Payload.Radiance = GetMaterialEmissive(PixelMaterialInputs);
-	Payload.Radiance *= Payload.BSDFOpacity; // premultiply
-#if MATERIAL_TWOSIDED
-	Payload.SetMaterialTwoSided();
-#else
-	if (MaterialParameters.TwoSidedSign < 0)
-	{
-		// when viewing the surface from "inside", don't include emission
-		Payload.Radiance = 0;
-	}
-#endif
-	Payload.BaseColor = BaseColor;
-	Payload.Specular = 0;
-	Payload.Roughness = 1;
-	Payload.Metallic = 0;
-	Payload.Ior = 0;
-
-#endif
-
 	PackedPayload = PackPathTracingPayload(Payload);
 }
 
diff --git a/Engine/Shaders/Shared/RayTracingPayloadType.h b/Engine/Shaders/Shared/RayTracingPayloadType.h
--- a/Engine/Shaders/Shared/RayTracingPayloadType.h
+++ b/Engine/Shaders/Shared/RayTracingPayloadType.h
@@ -14,6 +14,7 @@
 #define RT_PAYLOAD_TYPE_VFX   					(1 << 7)	// FVFXTracePayload
 #define RT_PAYLOAD_TYPE_DECALS					(1 << 8)	// FDecalShaderPayload 
 #define RT_PAYLOAD_TYPE_SPARSE_VOXEL			(1 << 9)	// FSparseVoxelPayload
+#define RT_PAYLOAD_TYPE_GPULIGHTMASS            (1 << 10)   // FPackedPathTracingPayload
 
 #ifdef __cplusplus
 /**
@@ -39,6 +40,7 @@ enum ERayTracingPayloadType : uint // HLSL
 	VFX = RT_PAYLOAD_TYPE_VFX,
 	Decals = RT_PAYLOAD_TYPE_DECALS,
 	SparseVoxel = RT_PAYLOAD_TYPE_SPARSE_VOXEL,
+	GPULightmass = RT_PAYLOAD_TYPE_GPULIGHTMASS,
 };
 
 
diff --git a/Engine/Source/Runtime/Renderer/Private/PathTracing.cpp b/Engine/Source/Runtime/Renderer/Private/PathTracing.cpp
--- a/Engine/Source/Runtime/Renderer/Private/PathTracing.cpp
+++ b/Engine/Source/Runtime/Renderer/Private/PathTracing.cpp
@@ -586,6 +586,29 @@ static bool ShouldCompilePathTracingShadersForProject(EShaderPlatform ShaderPlat
 			CVarPathTracing.GetValueOnAnyThread() != 0;
 }
 
+static bool ShouldCompileGPULightmassShadersForProject(EShaderPlatform ShaderPlatform)
+{
+#if WITH_EDITOR
+	if (!ShouldCompileRayTracingShadersForProject(ShaderPlatform))
+	{
+		return false;
+	}
+	// NOTE: cache on first use as this won't change
+	static const bool bIsGPULightmassLoaded = FModuleManager::Get().IsModuleLoaded(TEXT("GPULightmass"));
+	return bIsGPULightmassLoaded;
+#else
+	// GPULightmass is an editor only plugin, so don't compile any of its permutations otherwise
+	return false;
+#endif
+}
+
+static bool ShouldCompileGPULightmassShadersForProject(const FMeshMaterialShaderPermutationParameters& Parameters)
+{
+	return ShouldCompileGPULightmassShadersForProject(Parameters.Platform) &&
+		EnumHasAllFlags(Parameters.Flags, EShaderPermutationFlags::HasEditorOnlyData) &&
+		Parameters.VertexFactoryType->SupportsLightmapBaking();
+}
+
 class FPathTracingSkylightPrepareCS : public FGlobalShader
 {
 	DECLARE_GLOBAL_SHADER(FPathTracingSkylightPrepareCS)
@@ -665,8 +688,8 @@ class FPathTracingBuildLightGridCS : public FGlobalShader
 
 	static bool ShouldCompilePermutation(const FGlobalShaderPermutationParameters& Parameters)
 	{
-		// shared with GPULightmass
-		return ShouldCompileRayTracingShadersForProject(Parameters.Platform);
+		return ShouldCompilePathTracingShadersForProject(Parameters.Platform) ||
+			   ShouldCompileGPULightmassShadersForProject(Parameters.Platform);
 	}
 
 	static void ModifyCompilationEnvironment(const FGlobalShaderPermutationParameters& Parameters, FShaderCompilerEnvironment& OutEnvironment)
@@ -737,6 +760,7 @@ static FPathTracingFogParameters PrepareFogParameters(const FViewInfo& View, con
 }
 
 IMPLEMENT_RT_PAYLOAD_TYPE(ERayTracingPayloadType::PathTracingMaterial, 64); // TODO: Expand this when strata is enabled (use GetStrataRaytracingMaterialPayload())
+IMPLEMENT_RT_PAYLOAD_TYPE(ERayTracingPayloadType::GPULightmass, 32);
 
 class FPathTracingRG : public FGlobalShader
 {
@@ -849,7 +873,7 @@ class FPathTracingSwizzleScanlinesCS : public FGlobalShader
 
 	static bool ShouldCompilePermutation(const FGlobalShaderPermutationParameters& Parameters)
 	{
-		return ShouldCompileRayTracingShadersForProject(Parameters.Platform);
+		return ShouldCompilePathTracingShadersForProject(Parameters.Platform);
 	}
 
 	static void ModifyCompilationEnvironment(const FGlobalShaderPermutationParameters& Parameters, FShaderCompilerEnvironment& OutEnvironment)
@@ -877,7 +901,7 @@ class FPathTracingBuildAtmosphereOpticalDepthLUTCS : public FGlobalShader
 
 	static bool ShouldCompilePermutation(const FGlobalShaderPermutationParameters& Parameters)
 	{
-		return ShouldCompileRayTracingShadersForProject(Parameters.Platform);
+		return ShouldCompilePathTracingShadersForProject(Parameters.Platform);
 	}
 
 	static void ModifyCompilationEnvironment(const FGlobalShaderPermutationParameters& Parameters, FShaderCompilerEnvironment& OutEnvironment)
@@ -897,19 +921,27 @@ class FPathTracingBuildAtmosphereOpticalDepthLUTCS : public FGlobalShader
 IMPLEMENT_SHADER_TYPE(, FPathTracingBuildAtmosphereOpticalDepthLUTCS, TEXT("/Engine/Private/PathTracing/PathTracingBuildAtmosphereLUT.usf"), TEXT("PathTracingBuildAtmosphereOpticalDepthLUTCS"), SF_Compute);
 
 // Default miss shader (using the path tracing payload)
-class FPathTracingDefaultMS : public FGlobalShader
+template <bool IsGPULightmass>
+class TPathTracingDefaultMS : public FGlobalShader
 {
-	DECLARE_SHADER_TYPE(FPathTracingDefaultMS, Global, );
+	DECLARE_SHADER_TYPE(TPathTracingDefaultMS, Global, );
 public:
 
-	FPathTracingDefaultMS() = default;
-	FPathTracingDefaultMS(const ShaderMetaType::CompiledShaderInitializerType& Initializer)
+	TPathTracingDefaultMS() = default;
+	TPathTracingDefaultMS(const ShaderMetaType::CompiledShaderInitializerType& Initializer)
 		: FGlobalShader(Initializer)
 	{}
 
 	static bool ShouldCompilePermutation(const FGlobalShaderPermutationParameters& Parameters)
 	{
-		return ShouldCompilePathTracingShadersForProject(Parameters.Platform);
+		if (IsGPULightmass)
+		{
+			return ShouldCompileGPULightmassShadersForProject(Parameters.Platform);
+		}
+		else
+		{
+			return ShouldCompilePathTracingShadersForProject(Parameters.Platform);
+		}
 	}
 
 	static void ModifyCompilationEnvironment(const FGlobalShaderPermutationParameters& Parameters, FShaderCompilerEnvironment& OutEnvironment)
@@ -918,16 +950,32 @@ public:
 
 	static ERayTracingPayloadType GetRayTracingPayloadType(const int32 PermutationId)
 	{
-		return ERayTracingPayloadType::PathTracingMaterial;
+		if (IsGPULightmass)
+		{
+			return ERayTracingPayloadType::GPULightmass;
+		}
+		else
+		{
+			return ERayTracingPayloadType::PathTracingMaterial;
+		}
 	}
 };
-IMPLEMENT_SHADER_TYPE(, FPathTracingDefaultMS, TEXT("/Engine/Private/PathTracing/PathTracingMissShader.usf"), TEXT("PathTracingDefaultMS"), SF_RayMiss);
+
+using FPathTracingDefaultMS  = TPathTracingDefaultMS<false>;
+using FGPULightmassDefaultMS = TPathTracingDefaultMS<true>;
+IMPLEMENT_SHADER_TYPE(template<>, FPathTracingDefaultMS , TEXT("/Engine/Private/PathTracing/PathTracingMissShader.usf"), TEXT("PathTracingDefaultMS"), SF_RayMiss);
+IMPLEMENT_SHADER_TYPE(template<>, FGPULightmassDefaultMS, TEXT("/Engine/Private/PathTracing/PathTracingMissShader.usf"), TEXT("PathTracingDefaultMS"), SF_RayMiss);
 
 FRHIRayTracingShader* GetPathTracingDefaultMissShader(const FGlobalShaderMap* ShaderMap)
 {
 	return ShaderMap->GetShader<FPathTracingDefaultMS>().GetRayTracingShader();
 }
 
+FRHIRayTracingShader* GetGPULightmassDefaultMissShader(const FGlobalShaderMap* ShaderMap)
+{
+	return ShaderMap->GetShader<FGPULightmassDefaultMS>().GetRayTracingShader();
+}
+
 void FDeferredShadingSceneRenderer::SetupPathTracingDefaultMissShader(FRHICommandListImmediate& RHICmdList, const FViewInfo& View)
 {
 	int32 MissShaderPipelineIndex = FindRayTracingMissShaderIndex(View.RayTracingMaterialPipeline, GetPathTracingDefaultMissShader(View.ShaderMap), true);
@@ -1141,7 +1189,7 @@ static bool NeedsAnyHitShader(const FMaterial& RESTRICT MaterialResource)
 	return NeedsAnyHitShader(MaterialResource.IsMasked(), MaterialResource.IsDitherMasked());
 }
 
-template<bool UseAnyHitShader, bool UseIntersectionShader, bool UseSimplifiedShader>
+template<bool UseAnyHitShader, bool UseIntersectionShader, bool IsGPULightmass>
 class TPathTracingMaterial : public FMeshMaterialShader
 {
 	DECLARE_SHADER_TYPE(TPathTracingMaterial, MeshMaterial);
@@ -1164,9 +1212,9 @@ public:
 			// does the VF support ray tracing at all?
 			return false;
 		}
-		if (Parameters.MaterialParameters.MaterialDomain == MD_DeferredDecal)
+		if (Parameters.MaterialParameters.MaterialDomain != MD_Surface)
 		{
-			// decals are handled elsewhere
+			// This material is only for surfaces at the moment
 			return false;
 		}
 		if (NeedsAnyHitShader(Parameters.MaterialParameters.bIsMasked, Parameters.MaterialParameters.bIsDitherMasked) != UseAnyHitShader)
@@ -1179,20 +1227,12 @@ public:
 			// only need to compile the intersection shader permutation if the VF actually requires it
 			return false;
 		}
-		if (UseSimplifiedShader)
+		if (IsGPULightmass)
 		{
-#if WITH_EDITOR
-			// this is only used by GPULightmass
-			return EnumHasAllFlags(Parameters.Flags, EShaderPermutationFlags::HasEditorOnlyData) &&
-				   Parameters.VertexFactoryType->SupportsLightmapBaking() &&
-				   FModuleManager::Get().IsModuleLoaded(TEXT("GPULightmass"));
-#else
-			return false;
-#endif
+			return ShouldCompileGPULightmassShadersForProject(Parameters);
 		}
 		else
 		{
-			// this is only used by the Path Tracer
 			return ShouldCompilePathTracingShadersForProject(Parameters.Platform);
 		}
 	}
@@ -1204,7 +1244,7 @@ public:
 		OutEnvironment.SetDefine(TEXT("USE_MATERIAL_INTERSECTION_SHADER"), UseIntersectionShader ? 1 : 0);
 		OutEnvironment.SetDefine(TEXT("USE_RAYTRACED_TEXTURE_RAYCONE_LOD"), 0);
 		OutEnvironment.SetDefine(TEXT("SCENE_TEXTURES_DISABLED"), 1);
-		OutEnvironment.SetDefine(TEXT("SIMPLIFIED_MATERIAL_SHADER"), UseSimplifiedShader);
+		OutEnvironment.SetDefine(TEXT("SIMPLIFIED_MATERIAL_SHADER"), IsGPULightmass);
 		FMeshMaterialShader::ModifyCompilationEnvironment(Parameters, OutEnvironment);
 	}
 
@@ -1232,7 +1272,14 @@ public:
 
 	static ERayTracingPayloadType GetRayTracingPayloadType(const int32 PermutationId)
 	{
-		return ERayTracingPayloadType::PathTracingMaterial;
+		if (IsGPULightmass)
+		{
+			return ERayTracingPayloadType::GPULightmass;
+		}
+		else
+		{
+			return ERayTracingPayloadType::PathTracingMaterial;
+		}
 	}
 };
 
@@ -1251,62 +1298,72 @@ IMPLEMENT_MATERIAL_SHADER_TYPE(template <>, FPathTracingMaterialCHS       , TEXT
 IMPLEMENT_MATERIAL_SHADER_TYPE(template <>, FPathTracingMaterialCHS_AHS   , TEXT("/Engine/Private/PathTracing/PathTracingMaterialHitShader.usf"), TEXT("closesthit=PathTracingMaterialCHS anyhit=PathTracingMaterialAHS"), SF_RayHitGroup);
 IMPLEMENT_MATERIAL_SHADER_TYPE(template <>, FPathTracingMaterialCHS_IS    , TEXT("/Engine/Private/PathTracing/PathTracingMaterialHitShader.usf"), TEXT("closesthit=PathTracingMaterialCHS intersection=MaterialIS"), SF_RayHitGroup);
 IMPLEMENT_MATERIAL_SHADER_TYPE(template <>, FPathTracingMaterialCHS_AHS_IS, TEXT("/Engine/Private/PathTracing/PathTracingMaterialHitShader.usf"), TEXT("closesthit=PathTracingMaterialCHS anyhit=PathTracingMaterialAHS intersection=MaterialIS"), SF_RayHitGroup);
-IMPLEMENT_MATERIAL_SHADER_TYPE(template <>, FGPULightmassCHS              , TEXT("/Engine/Private/PathTracing/PathTracingMaterialHitShader.usf"), TEXT("closesthit=PathTracingMaterialCHS"), SF_RayHitGroup);
-IMPLEMENT_MATERIAL_SHADER_TYPE(template <>, FGPULightmassCHS_AHS          , TEXT("/Engine/Private/PathTracing/PathTracingMaterialHitShader.usf"), TEXT("closesthit=PathTracingMaterialCHS anyhit=PathTracingMaterialAHS"), SF_RayHitGroup);
+IMPLEMENT_MATERIAL_SHADER_TYPE(template <>, FGPULightmassCHS              , TEXT("/Engine/Private/PathTracing/PathTracingGPULightmassMaterialHitShader.usf"), TEXT("closesthit=GPULightmassMaterialCHS"), SF_RayHitGroup);
+IMPLEMENT_MATERIAL_SHADER_TYPE(template <>, FGPULightmassCHS_AHS          , TEXT("/Engine/Private/PathTracing/PathTracingGPULightmassMaterialHitShader.usf"), TEXT("closesthit=GPULightmassMaterialCHS anyhit=GPULightmassMaterialAHS"), SF_RayHitGroup);
 
-class FPathTracingDefaultOpaqueHitGroup : public FGlobalShader
+template <bool IsGPULightmass, bool IsOpaque>
+class TPathTracingDefaultHitGroup : public FGlobalShader
 {
-	DECLARE_GLOBAL_SHADER(FPathTracingDefaultOpaqueHitGroup)
-	SHADER_USE_ROOT_PARAMETER_STRUCT(FPathTracingDefaultOpaqueHitGroup, FGlobalShader)
+	DECLARE_GLOBAL_SHADER(TPathTracingDefaultHitGroup)
+	SHADER_USE_ROOT_PARAMETER_STRUCT(TPathTracingDefaultHitGroup, FGlobalShader)
 
 	static bool ShouldCompilePermutation(const FGlobalShaderPermutationParameters& Parameters)
 	{
-		// Technically should be if either PT or GPULM are enabled.
-		// Make a utility function for this?
-		return ShouldCompileRayTracingShadersForProject(Parameters.Platform);
+		if (IsGPULightmass)
+		{
+			return ShouldCompileGPULightmassShadersForProject(Parameters.Platform);
+		}
+		else
+		{
+			return ShouldCompilePathTracingShadersForProject(Parameters.Platform);
+		}
 	}
 
 	static ERayTracingPayloadType GetRayTracingPayloadType(const int32 PermutationId)
 	{
-		return ERayTracingPayloadType::PathTracingMaterial;
+		if (IsGPULightmass)
+		{
+			return ERayTracingPayloadType::GPULightmass;
+		}
+		else
+		{
+			return ERayTracingPayloadType::PathTracingMaterial;
+		}
 	}
 
 	using FParameters = FEmptyShaderParameters;
 };
 
-class FPathTracingDefaultHiddenHitGroup : public FGlobalShader
-{
-	DECLARE_GLOBAL_SHADER(FPathTracingDefaultHiddenHitGroup)
-	SHADER_USE_ROOT_PARAMETER_STRUCT(FPathTracingDefaultHiddenHitGroup, FGlobalShader)
-
-	static bool ShouldCompilePermutation(const FGlobalShaderPermutationParameters& Parameters)
-	{
-		// Technically should be if either PT or GPULM are enabled.
-		// Make a utility function for this?
-		return ShouldCompileRayTracingShadersForProject(Parameters.Platform);
-	}
+using FPathTracingDefaultOpaqueHitGroup  = TPathTracingDefaultHitGroup<false, true >;
+using FPathTracingDefaultHiddenHitGroup  = TPathTracingDefaultHitGroup<false, false>;
+using FGPULightmassDefaultOpaqueHitGroup = TPathTracingDefaultHitGroup<true , true >;
+using FGPULightmassDefaultHiddenHitGroup = TPathTracingDefaultHitGroup<true , false>;
 
-	static ERayTracingPayloadType GetRayTracingPayloadType(const int32 PermutationId)
-	{
-		return ERayTracingPayloadType::PathTracingMaterial;
-	}
-
-	using FParameters = FEmptyShaderParameters;
-};
-
-IMPLEMENT_SHADER_TYPE(, FPathTracingDefaultOpaqueHitGroup, TEXT("/Engine/Private/PathTracing/PathTracingDefaultHitShader.usf"), TEXT("closesthit=PathTracingDefaultOpaqueCHS"), SF_RayHitGroup);
-IMPLEMENT_SHADER_TYPE(, FPathTracingDefaultHiddenHitGroup, TEXT("/Engine/Private/PathTracing/PathTracingDefaultHitShader.usf"), TEXT("closesthit=PathTracingDefaultHiddenCHS anyhit=PathTracingDefaultHiddenAHS"), SF_RayHitGroup);
+IMPLEMENT_SHADER_TYPE(template<>, FPathTracingDefaultOpaqueHitGroup , TEXT("/Engine/Private/PathTracing/PathTracingDefaultHitShader.usf"), TEXT("closesthit=PathTracingDefaultOpaqueCHS"), SF_RayHitGroup);
+IMPLEMENT_SHADER_TYPE(template<>, FGPULightmassDefaultOpaqueHitGroup, TEXT("/Engine/Private/PathTracing/PathTracingDefaultHitShader.usf"), TEXT("closesthit=PathTracingDefaultOpaqueCHS"), SF_RayHitGroup);
+IMPLEMENT_SHADER_TYPE(template<>, FPathTracingDefaultHiddenHitGroup , TEXT("/Engine/Private/PathTracing/PathTracingDefaultHitShader.usf"), TEXT("closesthit=PathTracingDefaultHiddenCHS anyhit=PathTracingDefaultHiddenAHS"), SF_RayHitGroup);
+IMPLEMENT_SHADER_TYPE(template<>, FGPULightmassDefaultHiddenHitGroup, TEXT("/Engine/Private/PathTracing/PathTracingDefaultHitShader.usf"), TEXT("closesthit=PathTracingDefaultHiddenCHS anyhit=PathTracingDefaultHiddenAHS"), SF_RayHitGroup);
 
 FRHIRayTracingShader* GetPathTracingDefaultOpaqueHitShader(const FGlobalShaderMap* ShaderMap)
 {
 	return ShaderMap->GetShader<FPathTracingDefaultOpaqueHitGroup>().GetRayTracingShader();
 }
 
+FRHIRayTracingShader* GetGPULightmassDefaultOpaqueHitShader(const FGlobalShaderMap* ShaderMap)
+{
+	return ShaderMap->GetShader<FGPULightmassDefaultOpaqueHitGroup>().GetRayTracingShader();
+}
+
 FRHIRayTracingShader* GetPathTracingDefaultHiddenHitShader(const FGlobalShaderMap* ShaderMap)
 {
 	return ShaderMap->GetShader<FPathTracingDefaultHiddenHitGroup>().GetRayTracingShader();
 }
 
+FRHIRayTracingShader* GetGPULightmassDefaultHiddenHitShader(const FGlobalShaderMap* ShaderMap)
+{
+	return ShaderMap->GetShader<FGPULightmassDefaultHiddenHitGroup>().GetRayTracingShader();
+}
+
 bool FRayTracingMeshProcessor::ProcessPathTracing(
 	const FMeshBatch& RESTRICT MeshBatch,
 	uint64 BatchElementMask,
diff --git a/Engine/Source/Runtime/Renderer/Private/PathTracing.h b/Engine/Source/Runtime/Renderer/Private/PathTracing.h
--- a/Engine/Source/Runtime/Renderer/Private/PathTracing.h
+++ b/Engine/Source/Runtime/Renderer/Private/PathTracing.h
@@ -25,6 +25,10 @@ RENDERER_API FRHIRayTracingShader* GetPathTracingDefaultMissShader(const FGlobal
 RENDERER_API FRHIRayTracingShader* GetPathTracingDefaultOpaqueHitShader(const FGlobalShaderMap* ShaderMap);
 RENDERER_API FRHIRayTracingShader* GetPathTracingDefaultHiddenHitShader(const FGlobalShaderMap* ShaderMap);
 
+RENDERER_API FRHIRayTracingShader* GetGPULightmassDefaultMissShader(const FGlobalShaderMap* ShaderMap);
+RENDERER_API FRHIRayTracingShader* GetGPULightmassDefaultOpaqueHitShader(const FGlobalShaderMap* ShaderMap);
+RENDERER_API FRHIRayTracingShader* GetGPULightmassDefaultHiddenHitShader(const FGlobalShaderMap* ShaderMap);
+
 class FRDGTexture;
 
 struct FPathTracingResources
