Rename PathTracing GPU Stat to match convetion used by other systems.

diff --git a/Engine/Source/Runtime/Renderer/Private/PathTracing.cpp b/Engine/Source/Runtime/Renderer/Private/PathTracing.cpp
--- a/Engine/Source/Runtime/Renderer/Private/PathTracing.cpp
+++ b/Engine/Source/Runtime/Renderer/Private/PathTracing.cpp
@@ -2088,10 +2088,10 @@ BEGIN_SHADER_PARAMETER_STRUCT(FMGPUTransferParameters, )
 END_SHADER_PARAMETER_STRUCT()
 #endif
 
-DECLARE_GPU_STAT_NAMED(Stat_GPU_PathTracing, TEXT("Path Tracing"));
-DECLARE_GPU_STAT_NAMED(Stat_GPU_PathTracingPost, TEXT("Path Tracing Post"));
+DECLARE_GPU_STAT_NAMED(PathTracing, TEXT("Path Tracing"));
+DECLARE_GPU_STAT_NAMED(PathTracingPost, TEXT("Path Tracing Post"));
 #if WITH_MGPU
-DECLARE_GPU_STAT_NAMED(Stat_GPU_PathTracingCopy, TEXT("Path Tracing Copy"));
+DECLARE_GPU_STAT_NAMED(PathTracingCopy, TEXT("Path Tracing Copy"));
 #endif
 
 void FDeferredShadingSceneRenderer::RenderPathTracing(
@@ -2414,7 +2414,7 @@ void FDeferredShadingSceneRenderer::RenderPathTracing(
 				RDG_GPU_MASK_SCOPE(GraphBuilder, FRHIGPUMask::FromIndex(GPUIndex));
 				RDG_EVENT_SCOPE_CONDITIONAL(GraphBuilder, NumGPUs > 1, "Path Tracing GPU%d", GPUIndex);
 #if WITH_MGPU
-				RDG_GPU_STAT_SCOPE(GraphBuilder, Stat_GPU_PathTracing);
+				RDG_GPU_STAT_SCOPE(GraphBuilder, PathTracing);
 #endif
 				for (int32 TileY = 0; TileY < DispatchResY; TileY += DispatchSize)
 				{
@@ -2614,7 +2614,7 @@ void FDeferredShadingSceneRenderer::RenderPathTracing(
 			// pay attention to the mask, so it has no effect on behavior.  Technically the work of the copy is done on the second GPU,
 			// and the first GPU stalls waiting on that, so it's useful to show this interval on both GPUs.
 			RDG_GPU_MASK_SCOPE(GraphBuilder, GPUMask);
-			RDG_GPU_STAT_SCOPE(GraphBuilder, Stat_GPU_PathTracingCopy);
+			RDG_GPU_STAT_SCOPE(GraphBuilder, PathTracingCopy);
 
 			FMGPUTransferParameters* Parameters = GraphBuilder.AllocParameters<FMGPUTransferParameters>();
 			Parameters->InputTexture = RadianceTexture;
@@ -2679,7 +2679,7 @@ void FDeferredShadingSceneRenderer::RenderPathTracing(
 	}
 
 	RDG_GPU_MASK_SCOPE(GraphBuilder, View.GPUMask);
-	RDG_GPU_STAT_SCOPE(GraphBuilder, Stat_GPU_PathTracingPost);
+	RDG_GPU_STAT_SCOPE(GraphBuilder, PathTracingPost);
 
 	// Figure out if the denoiser is enabled and needs to run
 	FRDGTexture* DenoisedRadianceTexture = nullptr;
