Disable SSAO pass VRS by default, since it's currently broken and the fix is targeted for 5.3.

diff --git a/Engine/Source/Runtime/Renderer/Private/VariableRateShading/ContrastAdaptiveImageGenerator.cpp b/Engine/Source/Runtime/Renderer/Private/VariableRateShading/ContrastAdaptiveImageGenerator.cpp
--- a/Engine/Source/Runtime/Renderer/Private/VariableRateShading/ContrastAdaptiveImageGenerator.cpp
+++ b/Engine/Source/Runtime/Renderer/Private/VariableRateShading/ContrastAdaptiveImageGenerator.cpp
@@ -82,7 +82,7 @@ TAutoConsoleVariable<int32> CVarVRSNaniteEmitGBuffer(
 	ECVF_RenderThreadSafe);
 TAutoConsoleVariable<int32> CVarVRS_SSAO(
 	TEXT("r.VRS.SSAO"),
-	2,
+	0,
 	TEXT("Enable VRS with SSAO rendering.\n")
 	TEXT("0: Disabled")
 	TEXT("1: Full")
