CitySample optimizations.
 1. Don't resize octree elements buffer when removing elements. When fluctuating between ~30k - 50k primitives in the scene, the reallocation + memcopies become very expensive.
 2. Added scene profiling counters to track scene primitive counts throughout the run.
 3. Use chunked arrays for large scene primitive arrays like bounds and matrices. With the switch to doubles and the number of primitives in the scene, resizing can take >5ms on a PS5 each frame.
 4. Reduced streaming time budget to a reasonable level.
 5. Enabled parallel GPU scene updates in CitySample.

diff --git a/Engine/Source/Runtime/Core/Public/Math/GenericOctree.h b/Engine/Source/Runtime/Core/Public/Math/GenericOctree.h
--- a/Engine/Source/Runtime/Core/Public/Math/GenericOctree.h
+++ b/Engine/Source/Runtime/Core/Public/Math/GenericOctree.h
@@ -822,7 +822,7 @@ public:
 		checkSlow(ElementId.IsValidId());
 
 		// Remove the element from the node's element list.
-		TreeElements[ElementId.NodeIndex].RemoveAtSwap(ElementId.ElementIndex);
+		TreeElements[ElementId.NodeIndex].RemoveAtSwap(ElementId.ElementIndex, 1, false);
 
 		if (ElementId.ElementIndex < TreeElements[ElementId.NodeIndex].Num())
 		{
diff --git a/Engine/Source/Runtime/Renderer/Private/DistanceFieldObjectManagement.cpp b/Engine/Source/Runtime/Renderer/Private/DistanceFieldObjectManagement.cpp
--- a/Engine/Source/Runtime/Renderer/Private/DistanceFieldObjectManagement.cpp
+++ b/Engine/Source/Runtime/Renderer/Private/DistanceFieldObjectManagement.cpp
@@ -543,7 +543,7 @@ void FDistanceFieldSceneData::UpdateDistanceFieldObjectBuffers(
 					UploadDistanceFieldDataBuffer.Init(GraphBuilder, NumDFObjectUploads, GDistanceFieldObjectDataStride * sizeof(FVector4f), true, TEXT("DistanceFields.DFObjectDataUploadBuffer"));
 					UploadDistanceFieldBoundsBuffer.Init(GraphBuilder, NumDFObjectUploads, GDistanceFieldObjectBoundsStride * sizeof(FVector4f), true, TEXT("DistanceFields.DFObjectBoundsUploadBuffer"));
 
-					const TArray<FPrimitiveBounds>& PrimitiveBounds = Scene->PrimitiveBounds;
+					const TScenePrimitiveArray<FPrimitiveBounds>& PrimitiveBounds = Scene->PrimitiveBounds;
 
 					FParallelUpdateRangesDFO ParallelRanges;
 
diff --git a/Engine/Source/Runtime/Renderer/Private/MeshDrawCommands.cpp b/Engine/Source/Runtime/Renderer/Private/MeshDrawCommands.cpp
--- a/Engine/Source/Runtime/Renderer/Private/MeshDrawCommands.cpp
+++ b/Engine/Source/Runtime/Renderer/Private/MeshDrawCommands.cpp
@@ -141,7 +141,7 @@ void UpdateTranslucentMeshSortKeys(
 	const FVector& TranslucentSortAxis,
 	const FVector& ViewOrigin,
 	const FMatrix& ViewMatrix,
-	const TArray<struct FPrimitiveBounds>& PrimitiveBounds,
+	const TScenePrimitiveArray<FPrimitiveBounds>& PrimitiveBounds,
 	ETranslucencyPass::Type TranslucencyPass, 
 	bool bInverseSorting,
 	FMeshCommandOneFrameArray& VisibleMeshCommands
@@ -320,7 +320,7 @@ static uint64 GetMobileBasePassSortKey_ByState(bool bMasked, bool bBackground, u
 */
 void UpdateMobileBasePassMeshSortKeys(
 	const FVector& ViewOrigin,
-	const TArray<struct FPrimitiveBounds>& ScenePrimitiveBounds,
+	const TScenePrimitiveArray<FPrimitiveBounds>& ScenePrimitiveBounds,
 	FMeshCommandOneFrameArray& VisibleMeshCommands
 )
 {
diff --git a/Engine/Source/Runtime/Renderer/Private/MeshDrawCommands.h b/Engine/Source/Runtime/Renderer/Private/MeshDrawCommands.h
--- a/Engine/Source/Runtime/Renderer/Private/MeshDrawCommands.h
+++ b/Engine/Source/Runtime/Renderer/Private/MeshDrawCommands.h
@@ -10,6 +10,7 @@ MeshDrawCommands.h: Mesh draw commands.
 #include "TranslucencyPass.h"
 #include "InstanceCulling/InstanceCullingContext.h"
 #include "InstanceCulling/InstanceCullingManager.h"
+#include "ScenePrivateBase.h"
 
 struct FMeshBatchAndRelevance;
 class FStaticMeshBatch;
@@ -96,7 +97,7 @@ public:
 	FVector TranslucentSortAxis;
 	FVector ViewOrigin;
 	FMatrix ViewMatrix;
-	const TArray<struct FPrimitiveBounds>* PrimitiveBounds;
+	const TScenePrimitiveArray<struct FPrimitiveBounds>* PrimitiveBounds;
 
 	// For logging instancing stats.
 	int32 VisibleMeshDrawCommandsNum;
diff --git a/Engine/Source/Runtime/Renderer/Private/RendererScene.cpp b/Engine/Source/Runtime/Renderer/Private/RendererScene.cpp
--- a/Engine/Source/Runtime/Renderer/Private/RendererScene.cpp
+++ b/Engine/Source/Runtime/Renderer/Private/RendererScene.cpp
@@ -73,6 +73,7 @@
 #include "RenderCore.h"
 #include "IESTextureManager.h"
 #include "Materials/MaterialRenderProxy.h"
+#include "ProfilingDebugging/CountersTrace.h"
 
 #if RHI_RAYTRACING
 #include "Nanite/NaniteRayTracing.h"
@@ -143,6 +144,10 @@ IMPLEMENT_GLOBAL_SHADER_PARAMETER_STRUCT(FDistanceCullFadeUniformShaderParameter
 
 IMPLEMENT_GLOBAL_SHADER_PARAMETER_STRUCT(FDitherUniformShaderParameters, "PrimitiveDither");
 
+TRACE_DECLARE_INT_COUNTER(ScenePrimitives, TEXT("Scene/Primitives"));
+TRACE_DECLARE_INT_COUNTER(ScenePrimitivesArrayMax, TEXT("Scene/PrimitiveArrayMax"));
+TRACE_DECLARE_INT_COUNTER(SceneLights, TEXT("Scene/Lights"));
+
 /** Global primitive uniform buffer resource containing distance cull faded in */
 TGlobalResource< FGlobalDistanceCullFadeUniformBuffer > GDistanceCullFadedInUniformBuffer;
 
@@ -1746,6 +1751,14 @@ static void TArraySwapElements(TArray<T>& Array, int i1, int i2)
 	Array[i2] = tmp;
 }
 
+template<typename T>
+static void TArraySwapElements(TScenePrimitiveArray<T>& Array, int i1, int i2)
+{
+	T tmp = Array[i1];
+	Array[i1] = Array[i2];
+	Array[i2] = tmp;
+}
+
 static void TBitArraySwapElements(TBitArray<>& Array, int32 i1, int32 i2)
 {
 	FBitReference BitRef1 = Array[i1];
@@ -5539,9 +5552,9 @@ void FScene::UpdateAllPrimitiveSceneInfos(FRDGBuilder& GraphBuilder, EUpdateAllP
 			int SourceIndex = Primitives.Num() - RemoveCount;
 
 			Primitives.RemoveAt(SourceIndex, RemoveCount, false);
-			PrimitiveTransforms.RemoveAt(SourceIndex, RemoveCount, false);
+			PrimitiveTransforms.Remove(RemoveCount, false);
 			PrimitiveSceneProxies.RemoveAt(SourceIndex, RemoveCount, false);
-			PrimitiveBounds.RemoveAt(SourceIndex, RemoveCount, false);
+			PrimitiveBounds.Remove(RemoveCount, false);
 			PrimitiveFlagsCompact.RemoveAt(SourceIndex, RemoveCount, false);
 			PrimitiveVisibilityIds.RemoveAt(SourceIndex, RemoveCount, false);
 			PrimitiveOctreeIndex.RemoveAt(SourceIndex, RemoveCount, false);
@@ -5549,7 +5562,7 @@ void FScene::UpdateAllPrimitiveSceneInfos(FRDGBuilder& GraphBuilder, EUpdateAllP
 			PrimitiveComponentIds.RemoveAt(SourceIndex, RemoveCount, false);
 			PrimitiveVirtualTextureFlags.RemoveAt(SourceIndex, RemoveCount, false);
 			PrimitiveVirtualTextureLod.RemoveAt(SourceIndex, RemoveCount, false);
-			PrimitiveOcclusionBounds.RemoveAt(SourceIndex, RemoveCount, false);
+			PrimitiveOcclusionBounds.Remove(RemoveCount, false);
 
 			#if WITH_EDITOR
 			PrimitivesSelected.RemoveAt(SourceIndex, RemoveCount);
@@ -6355,6 +6368,10 @@ void FScene::UpdateAllPrimitiveSceneInfos(FRDGBuilder& GraphBuilder, EUpdateAllP
 
 	LastUpdateFrameCounter = GFrameCounterRenderThread;
 
+	TRACE_COUNTER_SET(SceneLights, Lights.Num());
+	TRACE_COUNTER_SET(ScenePrimitives, Primitives.Num());
+	TRACE_COUNTER_SET(ScenePrimitivesArrayMax, Primitives.Max());
+
 #if DO_GUARD_SLOW
 	for (const FPrimitiveSceneInfo* PrimitiveSceneInfo : Primitives)
 	{
diff --git a/Engine/Source/Runtime/Renderer/Private/ScenePrivate.h b/Engine/Source/Runtime/Renderer/Private/ScenePrivate.h
--- a/Engine/Source/Runtime/Renderer/Private/ScenePrivate.h
+++ b/Engine/Source/Runtime/Renderer/Private/ScenePrivate.h
@@ -2704,11 +2704,11 @@ public:
 	/** Packed array of primitives in the scene. */
 	TArray<FPrimitiveSceneInfo*> Primitives;
 	/** Packed array of all transforms in the scene. */
-	TArray<FMatrix> PrimitiveTransforms;
+	TScenePrimitiveArray<FMatrix> PrimitiveTransforms;
 	/** Packed array of primitive scene proxies in the scene. */
 	TArray<FPrimitiveSceneProxy*> PrimitiveSceneProxies;
 	/** Packed array of primitive bounds. */
-	TArray<FPrimitiveBounds> PrimitiveBounds;
+	TScenePrimitiveArray<FPrimitiveBounds> PrimitiveBounds;
 	/** Packed array of primitive flags. */
 	TArray<FPrimitiveFlagsCompact> PrimitiveFlagsCompact;
 	/** Packed array of precomputed primitive visibility IDs. */
@@ -2718,7 +2718,7 @@ public:
 	/** Packed array of primitive occlusion flags. See EOcclusionFlags. */
 	TArray<uint8> PrimitiveOcclusionFlags;
 	/** Packed array of primitive occlusion bounds. */
-	TArray<FBoxSphereBounds> PrimitiveOcclusionBounds;
+	TScenePrimitiveArray<FBoxSphereBounds> PrimitiveOcclusionBounds;
 	/** Packed array of primitive components associated with the primitive. */
 	TArray<FPrimitiveComponentId> PrimitiveComponentIds;
 	/** Packed array of runtime virtual texture flags. */
diff --git a/Engine/Source/Runtime/Renderer/Private/ScenePrivateBase.h b/Engine/Source/Runtime/Renderer/Private/ScenePrivateBase.h
--- a/Engine/Source/Runtime/Renderer/Private/ScenePrivateBase.h
+++ b/Engine/Source/Runtime/Renderer/Private/ScenePrivateBase.h
@@ -48,3 +48,73 @@ public:
 		return QueryIndex;
 	}
 };
+
+/** A simple chunked array representation for scene primitives data arrays. */
+template <typename T>
+class TScenePrimitiveArray
+{
+	static const int32 NumElementsPerChunk = 1024;
+public:
+	TScenePrimitiveArray() = default;
+
+	T& Add(const T& Element)
+	{
+		return *(new (&AddUninitialized()) T(Element));
+	}
+
+	T& AddUninitialized()
+	{
+		if (NumElements % NumElementsPerChunk == 0)
+		{
+			Chunks.Emplace();
+		}
+		return Get(NumElements++);
+	}
+
+	void Remove(uint32 Count, bool bAllowShrinking)
+	{
+		check(Count <= NumElements);
+		const uint32 NumElementsNew = NumElements - Count;
+		while (NumElements != NumElementsNew)
+		{
+			Get(--NumElements).~T();
+		}
+		const uint32 NumChunksToRemove = Chunks.Num() - NumChunks(NumElementsNew);
+		Chunks.RemoveAt(Chunks.Num() - NumChunksToRemove, NumChunksToRemove, bAllowShrinking);
+	}
+
+	void Reserve(int32 Count)
+	{
+		Chunks.Reserve(NumChunks(Count));
+	}
+
+	T& Get(int32 ElementIndex)
+	{
+		const uint32 ChunkIndex        = ElementIndex / NumElementsPerChunk;
+		const uint32 ChunkElementIndex = ElementIndex % NumElementsPerChunk;
+		return Chunks[ChunkIndex][ChunkElementIndex];
+	}
+
+	const T& Get(int32 ElementIndex) const
+	{
+		const uint32 ChunkIndex        = ElementIndex / NumElementsPerChunk;
+		const uint32 ChunkElementIndex = ElementIndex % NumElementsPerChunk;
+		return Chunks[ChunkIndex][ChunkElementIndex];
+	}
+
+	FORCEINLINE T& operator[] (int32 Index) { return Get(Index); }
+	FORCEINLINE const T& operator[] (int32 Index) const { return Get(Index); }
+
+	bool IsValidIndex(int32 Index) const { return static_cast<uint32>(Index) < NumElements; }
+
+	int32 Num() const { return NumElements; }
+
+private:
+	static constexpr uint32 NumChunks(uint32 NumElements)
+	{
+		return (NumElements + NumElementsPerChunk - 1u) / NumElementsPerChunk;
+	}
+
+	TArray<TStaticArray<T, NumElementsPerChunk>> Chunks;
+	uint32 NumElements = 0;
+};
\ No newline at end of file
