Make UpdateAllPrimitiveSceneInfos sync with the renderer clean-up task. This is necessary before removing cached mesh commands.

diff --git a/Engine/Source/Runtime/Renderer/Private/RendererScene.cpp b/Engine/Source/Runtime/Renderer/Private/RendererScene.cpp
--- a/Engine/Source/Runtime/Renderer/Private/RendererScene.cpp
+++ b/Engine/Source/Runtime/Renderer/Private/RendererScene.cpp
@@ -5327,7 +5327,7 @@ void FScene::UpdateAllPrimitiveSceneInfos(FRDGBuilder& GraphBuilder, EUpdateAllP
 	SCOPE_CYCLE_COUNTER(STAT_UpdateScenePrimitiveRenderThreadTime);
 
 	check(IsInRenderingThread());
-
+	FSceneRenderer::WaitForCleanUpTasks(GraphBuilder.RHICmdList);
 	FMaterialRenderProxy::UpdateDeferredCachedUniformExpressions();
 
 	RDG_EVENT_SCOPE(GraphBuilder, "UpdateAllPrimitiveSceneInfos");
