Fixed: mobile DOF setting is not applied when play on standalone mode

diff --git a/Engine/Source/Runtime/Renderer/Private/PostProcess/PostProcessing.cpp b/Engine/Source/Runtime/Renderer/Private/PostProcess/PostProcessing.cpp
--- a/Engine/Source/Runtime/Renderer/Private/PostProcess/PostProcessing.cpp
+++ b/Engine/Source/Runtime/Renderer/Private/PostProcess/PostProcessing.cpp
@@ -2035,7 +2035,7 @@ void AddMobilePostProcessingPasses(FRDGBuilder& GraphBuilder, FScene* Scene, con
 		if (PassSequence.IsEnabled(EPass::SunMask))
 		{
 			PassSequence.AcceptPass(EPass::SunMask);
-			bool bUseDepthTexture = SceneColor.Texture->Desc.Format == PF_FloatR11G11B10 && !SceneDepthAux.IsValid();
+			bool bUseDepthTexture = !MobileRequiresSceneDepthAux(View.GetShaderPlatform());
 
 			FMobileSunMaskInputs SunMaskInputs;
 			SunMaskInputs.bUseDepthTexture = bUseDepthTexture;
