Converting HDRCustomResolve and WideCustomResolve shaders to use batched parameters.

diff --git a/Engine/Source/Runtime/Renderer/Private/SceneRendering.cpp b/Engine/Source/Runtime/Renderer/Private/SceneRendering.cpp
--- a/Engine/Source/Runtime/Renderer/Private/SceneRendering.cpp
+++ b/Engine/Source/Runtime/Renderer/Private/SceneRendering.cpp
@@ -5256,7 +5256,7 @@ namespace
 		check(Shader);
 		GraphicsPSOInit.BoundShaderState.PixelShaderRHI = Shader;
 		SetGraphicsPipelineState(RHICmdList, GraphicsPSOInit, 0);
-		ShaderRef->SetParameters(RHICmdList, SceneColorTargetableRHI);
+		SetShaderParametersLegacyPS(RHICmdList, ShaderRef, SceneColorTargetableRHI);
 	}
 
 	template<typename ShaderType, typename ShaderArrayType>
@@ -5365,7 +5365,7 @@ void AddResolveSceneColorPass(FRDGBuilder& GraphBuilder, const FViewInfo& View,
 					GraphicsPSOInit.BoundShaderState.PixelShaderRHI = PixelShader.GetPixelShader();
 
 					SetGraphicsPipelineState(RHICmdList, GraphicsPSOInit, 0);
-					PixelShader->SetParameters(RHICmdList, SceneColorTargetableRHI, SceneColorFMaskRHI);
+					SetShaderParametersLegacyPS(RHICmdList, PixelShader, SceneColorTargetableRHI, SceneColorFMaskRHI);
 				}
 				else if (NumSamples == 4)
 				{
@@ -5373,7 +5373,7 @@ void AddResolveSceneColorPass(FRDGBuilder& GraphBuilder, const FViewInfo& View,
 					GraphicsPSOInit.BoundShaderState.PixelShaderRHI = PixelShader.GetPixelShader();
 
 					SetGraphicsPipelineState(RHICmdList, GraphicsPSOInit, 0);
-					PixelShader->SetParameters(RHICmdList, SceneColorTargetableRHI, SceneColorFMaskRHI);
+					SetShaderParametersLegacyPS(RHICmdList, PixelShader, SceneColorTargetableRHI, SceneColorFMaskRHI);
 				}
 				else if (NumSamples == 8)
 				{
@@ -5381,7 +5381,7 @@ void AddResolveSceneColorPass(FRDGBuilder& GraphBuilder, const FViewInfo& View,
 					GraphicsPSOInit.BoundShaderState.PixelShaderRHI = PixelShader.GetPixelShader();
 
 					SetGraphicsPipelineState(RHICmdList, GraphicsPSOInit, 0);
-					PixelShader->SetParameters(RHICmdList, SceneColorTargetableRHI, SceneColorFMaskRHI);
+					SetShaderParametersLegacyPS(RHICmdList, PixelShader, SceneColorTargetableRHI, SceneColorFMaskRHI);
 				}
 				else
 				{
diff --git a/Engine/Source/Runtime/Renderer/Private/WideCustomResolveShaders.cpp b/Engine/Source/Runtime/Renderer/Private/WideCustomResolveShaders.cpp
--- a/Engine/Source/Runtime/Renderer/Private/WideCustomResolveShaders.cpp
+++ b/Engine/Source/Runtime/Renderer/Private/WideCustomResolveShaders.cpp
@@ -61,8 +61,8 @@ static void ResolveColorWideInternal2(
 	GraphicsPSOInit.PrimitiveType = PT_TriangleList;
 
 	SetGraphicsPipelineState(RHICmdList, GraphicsPSOInit, 0);
-	
-	PixelShader->SetParameters(RHICmdList, SrcTexture, FmaskSRV, SrcOrigin);
+
+	SetShaderParametersLegacyPS(RHICmdList, PixelShader, SrcTexture, FmaskSRV, SrcOrigin);
 
 	RHICmdList.SetStreamSource(0, DummyVB, 0);
 	RHICmdList.DrawPrimitive(0, 1, 1);
diff --git a/Engine/Source/Runtime/Renderer/Private/WideCustomResolveShaders.h b/Engine/Source/Runtime/Renderer/Private/WideCustomResolveShaders.h
--- a/Engine/Source/Runtime/Renderer/Private/WideCustomResolveShaders.h
+++ b/Engine/Source/Runtime/Renderer/Private/WideCustomResolveShaders.h
@@ -50,15 +50,14 @@ public:
 		return IsFeatureLevelSupported(Parameters.Platform, ERHIFeatureLevel::SM5);
 	}
 
-	void SetParameters(FRHICommandList& RHICmdList, FRHITexture* Texture2DMS, FRHIShaderResourceView* FmaskSRV, FIntPoint Origin)
+	void SetParameters(FRHIBatchedShaderParameters& BatchedParameters, FRHITexture* Texture2DMS, FRHIShaderResourceView* FmaskSRV, FIntPoint Origin)
 	{
-		FRHIPixelShader* PixelShaderRHI = RHICmdList.GetBoundPixelShader();
-		SetTextureParameter(RHICmdList, PixelShaderRHI, Tex, Texture2DMS);
+		SetTextureParameter(BatchedParameters, Tex, Texture2DMS);
 		if (MSAASampleCount > 0)
 		{
-			SetSRVParameter(RHICmdList, PixelShaderRHI, FMaskTex, FmaskSRV);
+			SetSRVParameter(BatchedParameters, FMaskTex, FmaskSRV);
 		}
-		SetShaderValue(RHICmdList, PixelShaderRHI, ResolveOrigin, FVector2f(Origin));
+		SetShaderValue(BatchedParameters, ResolveOrigin, FVector2f(Origin));
 	}
 
 	static void ModifyCompilationEnvironment(const FGlobalShaderPermutationParameters& Parameters, FShaderCompilerEnvironment& OutEnvironment)
diff --git a/Engine/Source/Runtime/Renderer/Public/HdrCustomResolveShaders.h b/Engine/Source/Runtime/Renderer/Public/HdrCustomResolveShaders.h
--- a/Engine/Source/Runtime/Renderer/Public/HdrCustomResolveShaders.h
+++ b/Engine/Source/Runtime/Renderer/Public/HdrCustomResolveShaders.h
@@ -70,10 +70,17 @@ public:
 		Tex.Bind(Initializer.ParameterMap, TEXT("Tex"), SPF_Mandatory);
 	}
 
+	void SetParameters(FRHIBatchedShaderParameters& BatchedParameters, FRHITexture* Texture2DMS)
+	{
+		SetTextureParameter(BatchedParameters, Tex, Texture2DMS);
+	}
+
+	UE_DEPRECATED(5.3, "SetParameters with FRHIBatchedShaderParameters should be used.")
 	void SetParameters(FRHICommandList& RHICmdList, FRHITexture* Texture2DMS)
 	{
-		FRHIPixelShader* PixelShaderRHI = RHICmdList.GetBoundPixelShader();
-		SetTextureParameter(RHICmdList, PixelShaderRHI, Tex, Texture2DMS);
+		FRHIBatchedShaderParameters& BatchedParameters = RHICmdList.GetScratchShaderParameters();
+		SetParameters(BatchedParameters, Texture2DMS);
+		RHICmdList.SetBatchedShaderParameters(RHICmdList.GetBoundPixelShader(), BatchedParameters);
 	}
 
 	static void ModifyCompilationEnvironment(const FGlobalShaderPermutationParameters& Parameters, FShaderCompilerEnvironment& OutEnvironment)
@@ -96,16 +103,6 @@ public:
 	{
 	}
 
-	static bool ShouldCompilePermutation(const FGlobalShaderPermutationParameters& Parameters)
-	{
-		return FHdrCustomResolve2xPS::ShouldCompilePermutation(Parameters);
-	}
-
-	void SetParameters(FRHICommandList& RHICmdList, FRHITexture* Texture2DMS)
-	{
-		return FHdrCustomResolve2xPS::SetParameters(RHICmdList, Texture2DMS);
-	}
-
 	static void ModifyCompilationEnvironment(const FGlobalShaderPermutationParameters& Parameters, FShaderCompilerEnvironment& OutEnvironment)
 	{
 		// skip parent because it sets 2X macro
@@ -125,16 +122,6 @@ public:
 	{
 	}
 
-	static bool ShouldCompilePermutation(const FGlobalShaderPermutationParameters& Parameters)
-	{
-		return FHdrCustomResolve2xPS::ShouldCompilePermutation(Parameters);
-	}
-
-	void SetParameters(FRHICommandList& RHICmdList, FRHITexture* Texture2DMS)
-	{
-		return FHdrCustomResolve2xPS::SetParameters(RHICmdList, Texture2DMS);
-	}
-
 	static void ModifyCompilationEnvironment(const FGlobalShaderPermutationParameters& Parameters, FShaderCompilerEnvironment& OutEnvironment)
 	{
 		// skip parent because it sets 2X macro
@@ -165,11 +152,6 @@ public:
 		return false;
 	}
 
-	void SetParameters(FRHICommandList& RHICmdList, FRHITexture* Texture2DMS)
-	{
-		return FHdrCustomResolve2xPS::SetParameters(RHICmdList, Texture2DMS);
-	}
-
 	static void ModifyCompilationEnvironment(const FGlobalShaderPermutationParameters& Parameters, FShaderCompilerEnvironment& OutEnvironment)
 	{
 		FHdrCustomResolve2xPS::ModifyCompilationEnvironment(Parameters, OutEnvironment);
@@ -187,16 +169,6 @@ public:
 	{
 	}
 
-	static bool ShouldCompilePermutation(const FGlobalShaderPermutationParameters& Parameters)
-	{
-		return FHdrCustomResolveArray2xPS::ShouldCompilePermutation(Parameters);
-	}
-
-	void SetParameters(FRHICommandList& RHICmdList, FRHITexture* Texture2DMS)
-	{
-		return FHdrCustomResolveArray2xPS::SetParameters(RHICmdList, Texture2DMS);
-	}
-
 	static void ModifyCompilationEnvironment(const FGlobalShaderPermutationParameters& Parameters, FShaderCompilerEnvironment& OutEnvironment)
 	{
 		// skip parent because it sets 2X macro
@@ -217,16 +189,6 @@ public:
 	{
 	}
 
-	static bool ShouldCompilePermutation(const FGlobalShaderPermutationParameters& Parameters)
-	{
-		return FHdrCustomResolveArray2xPS::ShouldCompilePermutation(Parameters);
-	}
-
-	void SetParameters(FRHICommandList& RHICmdList, FRHITexture* Texture2DMS)
-	{
-		return FHdrCustomResolveArray2xPS::SetParameters(RHICmdList, Texture2DMS);
-	}
-
 	static void ModifyCompilationEnvironment(const FGlobalShaderPermutationParameters& Parameters, FShaderCompilerEnvironment& OutEnvironment)
 	{
 		// skip parent because it sets 2X macro
@@ -250,11 +212,18 @@ public:
 		FMaskTex.Bind(Initializer.ParameterMap, TEXT("FMaskTex"), SPF_Optional);
 	}
 
+	void SetParameters(FRHIBatchedShaderParameters& BatchedParameters, FRHITexture* Texture2DMS, FRHIShaderResourceView* FMaskSRV)
+	{
+		SetTextureParameter(BatchedParameters, Tex, Texture2DMS);
+		SetSRVParameter(BatchedParameters, FMaskTex, FMaskSRV);
+	}
+
+	UE_DEPRECATED(5.3, "SetParameters with FRHIBatchedShaderParameters should be used.")
 	void SetParameters(FRHICommandList& RHICmdList, FRHITexture* Texture2DMS, FRHIShaderResourceView* FMaskSRV)
 	{
-		FRHIPixelShader* PixelShaderRHI = RHICmdList.GetBoundPixelShader();
-		SetTextureParameter(RHICmdList, PixelShaderRHI, Tex, Texture2DMS);
-		SetSRVParameter(RHICmdList, PixelShaderRHI, FMaskTex, FMaskSRV);
+		FRHIBatchedShaderParameters& BatchedParameters = RHICmdList.GetScratchShaderParameters();
+		SetParameters(BatchedParameters, Texture2DMS, FMaskSRV);
+		RHICmdList.SetBatchedShaderParameters(RHICmdList.GetBoundPixelShader(), BatchedParameters);
 	}
 
 	static void ModifyCompilationEnvironment(const FGlobalShaderPermutationParameters& Parameters, FShaderCompilerEnvironment& OutEnvironment)
@@ -281,11 +250,18 @@ public:
 		FMaskTex.Bind(Initializer.ParameterMap, TEXT("FMaskTex"), SPF_Optional);
 	}
 
+	void SetParameters(FRHIBatchedShaderParameters& BatchedParameters, FRHITexture* Texture2DMS, FRHIShaderResourceView* FMaskSRV)
+	{
+		SetTextureParameter(BatchedParameters, Tex, Texture2DMS);
+		SetSRVParameter(BatchedParameters, FMaskTex, FMaskSRV);
+	}
+
+	UE_DEPRECATED(5.3, "SetParameters with FRHIBatchedShaderParameters should be used.")
 	void SetParameters(FRHICommandList& RHICmdList, FRHITexture* Texture2DMS, FRHIShaderResourceView* FMaskSRV)
 	{
-		FRHIPixelShader* PixelShaderRHI = RHICmdList.GetBoundPixelShader();
-		SetTextureParameter(RHICmdList, PixelShaderRHI, Tex, Texture2DMS);
-		SetSRVParameter(RHICmdList, PixelShaderRHI, FMaskTex, FMaskSRV);
+		FRHIBatchedShaderParameters& BatchedParameters = RHICmdList.GetScratchShaderParameters();
+		SetParameters(BatchedParameters, Texture2DMS, FMaskSRV);
+		RHICmdList.SetBatchedShaderParameters(RHICmdList.GetBoundPixelShader(), BatchedParameters);
 	}
 
 	static void ModifyCompilationEnvironment(const FGlobalShaderPermutationParameters& Parameters, FShaderCompilerEnvironment& OutEnvironment)
@@ -313,11 +289,18 @@ public:
 		FMaskTex.Bind(Initializer.ParameterMap, TEXT("FMaskTex"), SPF_Optional);
 	}
 
+	void SetParameters(FRHIBatchedShaderParameters& BatchedParameters, FRHITexture* Texture2DMS, FRHIShaderResourceView* FMaskSRV)
+	{
+		SetTextureParameter(BatchedParameters, Tex, Texture2DMS);
+		SetSRVParameter(BatchedParameters, FMaskTex, FMaskSRV);
+	}
+
+	UE_DEPRECATED(5.3, "SetParameters with FRHIBatchedShaderParameters should be used.")
 	void SetParameters(FRHICommandList& RHICmdList, FRHITexture* Texture2DMS, FRHIShaderResourceView* FMaskSRV)
 	{
-		FRHIPixelShader* PixelShaderRHI = RHICmdList.GetBoundPixelShader();
-		SetTextureParameter(RHICmdList, PixelShaderRHI, Tex, Texture2DMS);
-		SetSRVParameter(RHICmdList, PixelShaderRHI, FMaskTex, FMaskSRV);
+		FRHIBatchedShaderParameters& BatchedParameters = RHICmdList.GetScratchShaderParameters();
+		SetParameters(BatchedParameters, Texture2DMS, FMaskSRV);
+		RHICmdList.SetBatchedShaderParameters(RHICmdList.GetBoundPixelShader(), BatchedParameters);
 	}
 
 	static void ModifyCompilationEnvironment(const FGlobalShaderPermutationParameters& Parameters, FShaderCompilerEnvironment& OutEnvironment)
