Fix Nanite Visualization modes not rendering in Shipping Editor builds.

diff --git a/Engine/Source/Runtime/Renderer/Private/PostProcess/PostProcessing.cpp b/Engine/Source/Runtime/Renderer/Private/PostProcess/PostProcessing.cpp
--- a/Engine/Source/Runtime/Renderer/Private/PostProcess/PostProcessing.cpp
+++ b/Engine/Source/Runtime/Renderer/Private/PostProcess/PostProcessing.cpp
@@ -1448,9 +1448,10 @@ void AddPostProcessingPasses(
 		}
 	}
 
-	// Draw debug stuff directly onto the back buffer
-	#if !UE_BUILD_SHIPPING
+	#if WITH_EDITOR || !UE_BUILD_SHIPPING
 	{
+		// Draw debug stuff directly onto the back buffer
+
 		if (EngineShowFlags.TestImage)
 		{
 			AddTestImagePass(GraphBuilder, View, SceneColor);
