Add coverage threshold for hit proxy, in order to control at which coverage hair can be selected.

diff --git a/Engine/Shaders/Private/HairStrands/HairStrandsHitProxy.usf b/Engine/Shaders/Private/HairStrands/HairStrandsHitProxy.usf
--- a/Engine/Shaders/Private/HairStrands/HairStrandsHitProxy.usf
+++ b/Engine/Shaders/Private/HairStrands/HairStrandsHitProxy.usf
@@ -8,10 +8,12 @@ float3 DecodeVelocityFromTexture(float4 EncodedV);
 
 #if SHADER_HITPROXY_ID
 
+float 							 CoverageThreshold;
 uint							 MaxMaterialCount;
 Buffer<uint>					 MaterialIdToHitProxyIdBuffer;
 Texture2D<uint>					 VisNodeIndex;
 StructuredBuffer<FPackedHairVis> VisNodeData;
+Texture2D<float>				 CoverageTexture;
 
 void EmitPS(
 	in float4 SvPosition : SV_Position,
@@ -22,7 +24,8 @@ void EmitPS(
 
 	// Take the first sample
 	const FNodeDesc NodeDesc = DecodeNodeDesc(VisNodeIndex.Load(uint3(PixelCoord,0)));
-	if (NodeDesc.Count > 0)
+	const float Coverage = CoverageTexture.Load(uint3(PixelCoord, 0));
+	if (NodeDesc.Count > 0 && Coverage > CoverageThreshold)
 	{
 		const uint SampleIndex = NodeDesc.Offset + 0;
 		const FHairVis InNode  = UnpackHairVis(VisNodeData[SampleIndex]);
diff --git a/Engine/Source/Runtime/Renderer/Private/HairStrands/HairStrandsVisibility.cpp b/Engine/Source/Runtime/Renderer/Private/HairStrands/HairStrandsVisibility.cpp
--- a/Engine/Source/Runtime/Renderer/Private/HairStrands/HairStrandsVisibility.cpp
+++ b/Engine/Source/Runtime/Renderer/Private/HairStrands/HairStrandsVisibility.cpp
@@ -3790,11 +3790,13 @@ class FHairStrandsEmitHitProxyIdPS : public FGlobalShader
 	SHADER_USE_PARAMETER_STRUCT(FHairStrandsEmitHitProxyIdPS, FGlobalShader);
 
 	BEGIN_SHADER_PARAMETER_STRUCT(FParameters, )
+		SHADER_PARAMETER(float, CoverageThreshold)
 		SHADER_PARAMETER(uint32, MaxMaterialCount)
 		SHADER_PARAMETER_STRUCT_REF(FViewUniformShaderParameters, View)
 		SHADER_PARAMETER_RDG_TEXTURE(Texture2D, VisNodeIndex)
 		SHADER_PARAMETER_RDG_BUFFER_SRV(StructuredBuffer<FPackedHairVis>, VisNodeData)
 		SHADER_PARAMETER_RDG_BUFFER_SRV(Buffer<uint>, MaterialIdToHitProxyIdBuffer)
+		SHADER_PARAMETER_RDG_TEXTURE(Texture2D<float>, CoverageTexture)
 		RENDER_TARGET_BINDING_SLOTS()
 	END_SHADER_PARAMETER_STRUCT()
 
@@ -3817,6 +3819,7 @@ void AddHairStrandsHitProxyIdPass(
 	const FViewInfo& View,
 	FRDGTextureRef VisNodeIndex,
 	FRDGBufferRef VisNodeData,
+	FRDGTextureRef CoverageTexture,
 	FRDGTextureRef HitProxyTexture,
 	FRDGTextureRef HitProxyDepthTexture)
 {
@@ -3838,10 +3841,12 @@ void AddHairStrandsHitProxyIdPass(
 	FRDGBufferRef MaterialIdToHitProxyIdBuffer = CreateUploadBuffer(GraphBuilder, TEXT("Hair.MaterialIdToHitProxyIdBuffer"), sizeof(uint32), MaterialIdToHitProxyId.Num(), MaterialIdToHitProxyId.GetData(), sizeof(uint32) * MaterialIdToHitProxyId.Num());
 	auto* PassParameters = GraphBuilder.AllocParameters<FHairStrandsEmitHitProxyIdPS::FParameters>();
 	PassParameters->View = View.ViewUniformBuffer;
+	PassParameters->CoverageThreshold = FMath::Clamp(GHairStrands_Selection_CoverageThreshold, 0.f, 1.f);
 	PassParameters->MaxMaterialCount = MaterialIdToHitProxyId.Num();
 	PassParameters->VisNodeIndex = VisNodeIndex;
 	PassParameters->VisNodeData = GraphBuilder.CreateSRV(VisNodeData);
 	PassParameters->MaterialIdToHitProxyIdBuffer = GraphBuilder.CreateSRV(MaterialIdToHitProxyIdBuffer, PF_R32_UINT);
+	PassParameters->CoverageTexture = CoverageTexture;
 	PassParameters->RenderTargets[0] = FRenderTargetBinding(HitProxyTexture, ERenderTargetLoadAction::ELoad);
 	PassParameters->RenderTargets.DepthStencil = FDepthStencilBinding(HitProxyDepthTexture, ERenderTargetLoadAction::ELoad, FExclusiveDepthStencil::DepthWrite_StencilWrite);
 
@@ -3982,7 +3987,6 @@ void DrawHitProxies(
 
 	FHairStrandsTiles TileData;
 	const FHairStrandsMacroGroupDatas& MacroGroupDatas = HairStrandsViewData.MacroGroupDatas;
-	FHairPrimaryTransmittance ViewTransmittance;
 
 	FRDGTextureRef SceneDepthTexture = HitProxyDepthTexture;
 
@@ -4029,8 +4033,19 @@ void DrawHitProxies(
 	}
 	else if (RenderMode == HairVisibilityRenderMode_MSAA_Visibility)
 	{
+		FHairPrimaryTransmittance ViewTransmittance = AddHairViewTransmittancePass(
+			GraphBuilder,
+			&Scene,
+			&View,
+			MacroGroupDatas,
+			Resolution,
+			false /*bOutputHairCount*/,
+			SceneDepthTexture,
+			InstanceCullingManager);
+
 		// Generate Tile data
-		TileData = AddHairStrandsGenerateTilesPass(GraphBuilder, View, Resolution);
+		check (ViewTransmittance.TransmittanceTexture);
+		TileData = AddHairStrandsGenerateTilesPass(GraphBuilder, View, ViewTransmittance.TransmittanceTexture);
 
 		FRDGTextureRef VisDepthTexture = AddHairVisibilityFillOpaqueDepth(
 			GraphBuilder,
@@ -4055,7 +4070,7 @@ void DrawHitProxies(
 		FHairVisibilityPrimitiveIdCompactionCS::FParameters* PassParameters = GraphBuilder.AllocParameters<FHairVisibilityPrimitiveIdCompactionCS::FParameters>();
 		PassParameters->MSAA_DepthTexture = VisDepthTexture;
 		PassParameters->MSAA_IDTexture = VisIdTexture;
-		PassParameters->ViewTransmittanceTexture = DummyTransmittanceTexture;
+		PassParameters->ViewTransmittanceTexture = ViewTransmittance.TransmittanceTexture;
 
 		AddHairVisibilityPrimitiveIdCompactionPass(
 			false, // bUsePPLL
@@ -4075,7 +4090,7 @@ void DrawHitProxies(
 			OutMaxNodeCount);
 	}
 
-	AddHairStrandsHitProxyIdPass(GraphBuilder, Scene, View, VisNodeIndex, VisNodeData, HitProxyTexture, HitProxyDepthTexture);
+	AddHairStrandsHitProxyIdPass(GraphBuilder, Scene, View, VisNodeIndex, VisNodeData, CoverageTexture, HitProxyTexture, HitProxyDepthTexture);
 }
 
 // Check if any simulated/skinned-bound groom has its positions updated (e.g. for invalidating the path-tracer accumulation)
