Use IHeterogeneousVolumeInterface to assign per-volume parameters when rendering.

diff --git a/Engine/Source/Runtime/Renderer/Private/HeterogeneousVolumes/HeterogeneousVolumes.cpp b/Engine/Source/Runtime/Renderer/Private/HeterogeneousVolumes/HeterogeneousVolumes.cpp
--- a/Engine/Source/Runtime/Renderer/Private/HeterogeneousVolumes/HeterogeneousVolumes.cpp
+++ b/Engine/Source/Runtime/Renderer/Private/HeterogeneousVolumes/HeterogeneousVolumes.cpp
@@ -2,6 +2,7 @@
 
 #include "HeterogeneousVolumes.h"
 
+#include "HeterogeneousVolumeInterface.h"
 #include "PixelShaderUtils.h"
 #include "RayTracingDefinitions.h"
 #include "RayTracingInstance.h"
@@ -81,23 +82,29 @@ static TAutoConsoleVariable<int32> CVarHeterogeneousVolumesPreshadingMipLevel(
 );
 
 static TAutoConsoleVariable<int32> CVarHeterogeneousVolumesPreshadingVolumeResolutionX(
-	TEXT("r.HeterogeneousVolumes.Preshading.VolumeResolution.X"),
-	256,
-	TEXT("Determines the preshading volume resolution in X (Default = 256)"),
+	TEXT("r.HeterogeneousVolumes.VolumeResolution.X"),
+	0,
+	TEXT("Overrides the preshading and lighting volume resolution in X (Default = 0)")
+	TEXT("0: Disabled, uses per-volume attribute\n")
+	TEXT(">0: Overrides resolution in X\n"),
 	ECVF_RenderThreadSafe
 );
 
 static TAutoConsoleVariable<int32> CVarHeterogeneousVolumesPreshadingVolumeResolutionY(
-	TEXT("r.HeterogeneousVolumes.Preshading.VolumeResolution.Y"),
-	256,
-	TEXT("Determines the preshading volume resolution in Y (Default = 256)"),
+	TEXT("r.HeterogeneousVolumes.VolumeResolution.Y"),
+	0,
+	TEXT("Overrides the preshading and lighting volume resolution in X (Default = 0)")
+	TEXT("0: Disabled, uses per-volume attribute\n")
+	TEXT(">0: Overrides resolution in Y\n"),
 	ECVF_RenderThreadSafe
 );
 
 static TAutoConsoleVariable<int32> CVarHeterogeneousVolumesPreshadingVolumeResolutionZ(
-	TEXT("r.HeterogeneousVolumes.Preshading.VolumeResolution.Z"),
-	256,
-	TEXT("Determines the preshading volume resolution in Z (Default = 256)"),
+	TEXT("r.HeterogeneousVolumes.VolumeResolution.Z"),
+	0,
+	TEXT("Overrides the preshading and lighting volume resolution in X (Default = 0)")
+	TEXT("0: Disabled, uses per-volume attribute\n")
+	TEXT(">0: Overrides resolution in Z\n"),
 	ECVF_RenderThreadSafe
 );
 
@@ -155,8 +162,10 @@ static TAutoConsoleVariable<int32> CVarHeterogeneousVolumesLightingCache(
 
 static TAutoConsoleVariable<int32> CVarHeterogeneousVolumesLightingCacheDownsampleFactor(
 	TEXT("r.HeterogeneousVolumes.LightingCache.DownsampleFactor"),
-	1,
-	TEXT("Determines the downsample factor, relative to the preshading volume resolution (Default = 1)"),
+	0,
+	TEXT("Overrides the lighting-cache downsample factor, relative to the preshading volume resolution (Default = 0)\n")
+	TEXT("0: Disabled, uses per-volume attribute\n")
+	TEXT(">0: Overrides the lighting-cache downsample factor"),
 	ECVF_RenderThreadSafe
 );
 
@@ -197,15 +206,13 @@ bool DoesPlatformSupportHeterogeneousVolumes(EShaderPlatform Platform)
 bool DoesMaterialShaderSupportHeterogeneousVolumes(const FMaterialShaderParameters& MaterialShaderParameters)
 {
 	return (MaterialShaderParameters.MaterialDomain == MD_Volume)
-		// Restricting compilation to materials bound to Niagara meshes
-		&& MaterialShaderParameters.bIsUsedWithNiagaraMeshParticles;
+		&& MaterialShaderParameters.bIsUsedWithHeterogeneousVolumes;
 }
 
 bool DoesMaterialShaderSupportHeterogeneousVolumes(const FMaterial& Material)
 {
 	return (Material.GetMaterialDomain() == MD_Volume)
-		// Restricting compilation to materials bound to Niagara meshes
-		&& Material.IsUsedWithNiagaraMeshParticles();
+		&& Material.IsUsedWithHeterogeneousVolumes();
 }
 
 bool ShouldRenderMeshBatchWithHeterogeneousVolumes(
@@ -228,13 +235,23 @@ bool ShouldRenderMeshBatchWithHeterogeneousVolumes(
 namespace HeterogeneousVolumes
 {
 	// CVars
-
-	FIntVector GetVolumeResolution()
+	FIntVector GetVolumeResolution(const IHeterogeneousVolumeInterface* Interface)
 	{
-		FIntVector VolumeResolution;
-		VolumeResolution.X = FMath::Clamp(CVarHeterogeneousVolumesPreshadingVolumeResolutionX.GetValueOnRenderThread(), 64, 1024);
-		VolumeResolution.Y = FMath::Clamp(CVarHeterogeneousVolumesPreshadingVolumeResolutionY.GetValueOnRenderThread(), 64, 1024);
-		VolumeResolution.Z = FMath::Clamp(CVarHeterogeneousVolumesPreshadingVolumeResolutionZ.GetValueOnRenderThread(), 64, 1024);
+		FIntVector VolumeResolution = Interface->GetVoxelResolution();
+
+		FIntVector OverrideVolumeResolution = FIntVector(
+			CVarHeterogeneousVolumesPreshadingVolumeResolutionX.GetValueOnRenderThread(),
+			CVarHeterogeneousVolumesPreshadingVolumeResolutionY.GetValueOnRenderThread(),
+			CVarHeterogeneousVolumesPreshadingVolumeResolutionZ.GetValueOnRenderThread());
+
+		VolumeResolution.X = OverrideVolumeResolution.X > 0 ? OverrideVolumeResolution.X : VolumeResolution.X;
+		VolumeResolution.Y = OverrideVolumeResolution.Y > 0 ? OverrideVolumeResolution.Y : VolumeResolution.Y;
+		VolumeResolution.Z = OverrideVolumeResolution.Z > 0 ? OverrideVolumeResolution.Z : VolumeResolution.Z;
+
+		// Clamp each dimension to [1, 1024]
+		VolumeResolution.X = FMath::Clamp(VolumeResolution.X, 1, 1024);
+		VolumeResolution.Y = FMath::Clamp(VolumeResolution.Y, 1, 1024);
+		VolumeResolution.Z = FMath::Clamp(VolumeResolution.Z, 1, 1024);
 		return VolumeResolution;
 	}
 
@@ -327,6 +344,15 @@ namespace HeterogeneousVolumes
 
 	// Convenience Utils
 
+	const IHeterogeneousVolumeInterface* GetInterface(const FPrimitiveSceneProxy* PrimitiveSceneProxy)
+	{
+		check(PrimitiveSceneProxy->IsHeterogeneousVolume());
+		const IHeterogeneousVolumeInterface* Interface = PrimitiveSceneProxy->GetHeterogeneousVolumeInterface();
+		check(Interface);
+
+		return Interface;
+	}
+
 	int GetVoxelCount(FIntVector VolumeResolution)
 	{
 		return VolumeResolution.X * VolumeResolution.Y * VolumeResolution.Z;
@@ -345,12 +371,13 @@ namespace HeterogeneousVolumes
 			FMath::Max(VolumeResolution.Z >> MipLevel, 1)
 		);
 	}
-
-	FIntVector GetLightingCacheResolution()
+	FIntVector GetLightingCacheResolution(const IHeterogeneousVolumeInterface* RenderInterface)
 	{
-		float DownsampleFactor = FMath::Max(CVarHeterogeneousVolumesLightingCacheDownsampleFactor.GetValueOnRenderThread(), 1);
-		FIntVector LightingCacheResolution = GetVolumeResolution() / DownsampleFactor;
+		float OverrideDownsampleFactor = CVarHeterogeneousVolumesLightingCacheDownsampleFactor.GetValueOnRenderThread();
+		float DownsampleFactor = OverrideDownsampleFactor > 0.0 ? OverrideDownsampleFactor : RenderInterface->GetLightingDownsampleFactor();
+		DownsampleFactor = FMath::Max(DownsampleFactor, 1.0);
 
+		FIntVector LightingCacheResolution = GetVolumeResolution(RenderInterface) / DownsampleFactor;
 		return LightingCacheResolution;
 	}
 }
@@ -394,7 +421,7 @@ void FDeferredShadingSceneRenderer::RenderHeterogeneousVolumes(
 
 				// Allocate transmittance volume
 				// TODO: Allow option for scalar transmittance to conserve bandwidth
-				FIntVector LightingCacheResolution = HeterogeneousVolumes::GetLightingCacheResolution();
+				FIntVector LightingCacheResolution = HeterogeneousVolumes::GetLightingCacheResolution(HeterogeneousVolumes::GetInterface(PrimitiveSceneProxy));
 				uint32 NumMips = FMath::Log2(float(FMath::Min(FMath::Min(LightingCacheResolution.X, LightingCacheResolution.Y), LightingCacheResolution.Z))) + 1;
 				FRDGTextureDesc LightingCacheDesc = FRDGTextureDesc::Create3D(
 					LightingCacheResolution,
diff --git a/Engine/Source/Runtime/Renderer/Private/HeterogeneousVolumes/HeterogeneousVolumes.h b/Engine/Source/Runtime/Renderer/Private/HeterogeneousVolumes/HeterogeneousVolumes.h
--- a/Engine/Source/Runtime/Renderer/Private/HeterogeneousVolumes/HeterogeneousVolumes.h
+++ b/Engine/Source/Runtime/Renderer/Private/HeterogeneousVolumes/HeterogeneousVolumes.h
@@ -15,6 +15,7 @@ class FSceneView;
 class FViewInfo;
 class FVirtualShadowMapArray;
 class FVisibleLightInfo;
+class IHeterogeneousVolumeInterface;
 
 struct FMaterialShaderParameters;
 struct FRDGTextureDesc;
@@ -28,6 +29,7 @@ bool ShouldRenderHeterogeneousVolumes(const FScene* Scene);
 bool ShouldRenderHeterogeneousVolumesForView(const FViewInfo& View);
 bool DoesPlatformSupportHeterogeneousVolumes(EShaderPlatform Platform);
 bool DoesMaterialShaderSupportHeterogeneousVolumes(const FMaterialShaderParameters& Parameters);
+bool DoesMaterialShaderSupportHeterogeneousVolumes(const FMaterial& Material);
 bool ShouldRenderMeshBatchWithHeterogeneousVolumes(
 	const FMeshBatch* Mesh,
 	const FPrimitiveSceneProxy* Proxy,
@@ -41,8 +43,8 @@ bool ShouldRenderMeshBatchWithHeterogeneousVolumes(
 namespace HeterogeneousVolumes
 {
 	// CVars
-	FIntVector GetVolumeResolution();
-	FIntVector GetLightingCacheResolution();
+	FIntVector GetVolumeResolution(const IHeterogeneousVolumeInterface*);
+	FIntVector GetLightingCacheResolution(const IHeterogeneousVolumeInterface*);
 
 	float GetShadowStepSize();
 	float GetMaxTraceDistance();
@@ -65,6 +67,8 @@ namespace HeterogeneousVolumes
 	bool UseLightingCacheForTransmittance();
 
 	// Convenience Utils
+	const IHeterogeneousVolumeInterface* GetInterface(const FPrimitiveSceneProxy* PrimitiveSceneProxy);
+
 	int GetVoxelCount(FIntVector VolumeResolution);
 	int GetVoxelCount(const FRDGTextureDesc& TextureDesc);
 	FIntVector GetMipVolumeResolution(FIntVector VolumeResolution, uint32 MipLevel);
diff --git a/Engine/Source/Runtime/Renderer/Private/HeterogeneousVolumes/HeterogeneousVolumesHardwareRayTracing.cpp b/Engine/Source/Runtime/Renderer/Private/HeterogeneousVolumes/HeterogeneousVolumesHardwareRayTracing.cpp
--- a/Engine/Source/Runtime/Renderer/Private/HeterogeneousVolumes/HeterogeneousVolumesHardwareRayTracing.cpp
+++ b/Engine/Source/Runtime/Renderer/Private/HeterogeneousVolumes/HeterogeneousVolumesHardwareRayTracing.cpp
@@ -590,6 +590,8 @@ void RenderLightingCacheWithPreshadingHardwareRayTracing(
 	FRDGTextureRef& LightingCacheTexture
 )
 {
+	const IHeterogeneousVolumeInterface* Interface = HeterogeneousVolumes::GetInterface(PrimitiveSceneProxy);
+
 	FRenderLightingCacheWithPreshadingRGS::FParameters* PassParameters = GraphBuilder.AllocParameters<FRenderLightingCacheWithPreshadingRGS::FParameters>();
 	{
 		// Scene
@@ -611,7 +613,7 @@ void RenderLightingCacheWithPreshadingHardwareRayTracing(
 		PassParameters->SparseVoxelUniformBuffer = SparseVoxelUniformBuffer;
 
 		// Transmittance volume
-		PassParameters->LightingCache.LightingCacheResolution = HeterogeneousVolumes::GetLightingCacheResolution();
+		PassParameters->LightingCache.LightingCacheResolution = HeterogeneousVolumes::GetLightingCacheResolution(Interface);
 		//PassParameters->LightingCache.LightingCacheTexture = GraphBuilder.CreateSRV(LightingCacheTexture);
 
 		// Ray data
@@ -664,7 +666,7 @@ void RenderLightingCacheWithPreshadingHardwareRayTracing(
 	FRenderLightingCacheWithPreshadingRGS::FPermutationDomain PermutationVector;
 	PermutationVector.Set<FRenderLightingCacheWithPreshadingRGS::FLightingCacheMode>(HeterogeneousVolumes::GetLightingCacheMode() - 1);
 	TShaderRef<FRenderLightingCacheWithPreshadingRGS> RayGenerationShader = View.ShaderMap->GetShader<FRenderLightingCacheWithPreshadingRGS>(PermutationVector);
-	FIntVector VolumeResolution = HeterogeneousVolumes::GetLightingCacheResolution();
+	FIntVector VolumeResolution = HeterogeneousVolumes::GetLightingCacheResolution(Interface);
 	FIntPoint DispatchResolution = FIntPoint(VolumeResolution.X, VolumeResolution.Y * VolumeResolution.Z);
 
 	GraphBuilder.AddPass(
@@ -730,6 +732,8 @@ void RenderSingleScatteringWithPreshadingHardwareRayTracing(
 	FRDGTextureRef& HeterogeneousVolumeTexture
 )
 {
+	const IHeterogeneousVolumeInterface* Interface = HeterogeneousVolumes::GetInterface(PrimitiveSceneProxy);
+
 	FRenderSingleScatteringWithPreshadingRGS::FParameters* PassParameters = GraphBuilder.AllocParameters<FRenderSingleScatteringWithPreshadingRGS::FParameters>();
 	{
 		// Scene
@@ -786,7 +790,7 @@ void RenderSingleScatteringWithPreshadingHardwareRayTracing(
 		// Transmittance volume
 		if ((HeterogeneousVolumes::UseLightingCacheForTransmittance() && bApplyShadowTransmittance) || HeterogeneousVolumes::UseLightingCacheForInscattering())
 		{
-			PassParameters->LightingCache.LightingCacheResolution = HeterogeneousVolumes::GetLightingCacheResolution();
+			PassParameters->LightingCache.LightingCacheResolution = HeterogeneousVolumes::GetLightingCacheResolution(Interface);
 			PassParameters->LightingCache.LightingCacheTexture = LightingCacheTexture;
 		}
 
diff --git a/Engine/Source/Runtime/Renderer/Private/HeterogeneousVolumes/HeterogeneousVolumesLiveShadingPipeline.cpp b/Engine/Source/Runtime/Renderer/Private/HeterogeneousVolumes/HeterogeneousVolumesLiveShadingPipeline.cpp
--- a/Engine/Source/Runtime/Renderer/Private/HeterogeneousVolumes/HeterogeneousVolumesLiveShadingPipeline.cpp
+++ b/Engine/Source/Runtime/Renderer/Private/HeterogeneousVolumes/HeterogeneousVolumesLiveShadingPipeline.cpp
@@ -333,6 +333,8 @@ void RenderLightingCacheWithLiveShading(
 	FRDGTextureRef LightingCacheTexture
 )
 {
+	const IHeterogeneousVolumeInterface* Interface = HeterogeneousVolumes::GetInterface(PrimitiveSceneProxy);
+
 	const FMaterialRenderProxy* MaterialRenderProxy = nullptr;
 	const FMaterial& Material = DefaultMaterialRenderProxy->GetMaterialWithFallback(View.GetFeatureLevel(), MaterialRenderProxy);
 	MaterialRenderProxy = MaterialRenderProxy ? MaterialRenderProxy : DefaultMaterialRenderProxy;
@@ -364,7 +366,7 @@ void RenderLightingCacheWithLiveShading(
 		PassParameters->PrimitiveId = PrimitiveId;
 
 		// Transmittance volume
-		PassParameters->LightingCache.LightingCacheResolution = HeterogeneousVolumes::GetLightingCacheResolution();
+		PassParameters->LightingCache.LightingCacheResolution = HeterogeneousVolumes::GetLightingCacheResolution(Interface);
 		//PassParameters->LightingCache.LightingCacheTexture = GraphBuilder.CreateSRV(LightingCacheTexture);
 
 		// Ray data
@@ -416,7 +418,7 @@ void RenderLightingCacheWithLiveShading(
 	}
 #endif // WANTS_DRAW_MESH_EVENTS
 
-	FIntVector GroupCount = HeterogeneousVolumes::GetLightingCacheResolution();
+	FIntVector GroupCount = HeterogeneousVolumes::GetLightingCacheResolution(Interface);
 	GroupCount.X = FMath::DivideAndRoundUp(GroupCount.X, FRenderLightingCacheWithLiveShadingCS::GetThreadGroupSize3D());
 	GroupCount.Y = FMath::DivideAndRoundUp(GroupCount.Y, FRenderLightingCacheWithLiveShadingCS::GetThreadGroupSize3D());
 	GroupCount.Z = FMath::DivideAndRoundUp(GroupCount.Z, FRenderLightingCacheWithLiveShadingCS::GetThreadGroupSize3D());
@@ -456,6 +458,8 @@ void RenderSingleScatteringWithLiveShading(
 	FRDGTextureRef& HeterogeneousVolumeTexture
 )
 {
+	const IHeterogeneousVolumeInterface* Interface = HeterogeneousVolumes::GetInterface(PrimitiveSceneProxy);
+
 	const FMaterialRenderProxy* MaterialRenderProxy = nullptr;
 	const FMaterial& Material = DefaultMaterialRenderProxy->GetMaterialWithFallback(View.GetFeatureLevel(), MaterialRenderProxy);
 	MaterialRenderProxy = MaterialRenderProxy ? MaterialRenderProxy : DefaultMaterialRenderProxy;
@@ -536,7 +540,7 @@ void RenderSingleScatteringWithLiveShading(
 		// Volume data
 		if ((HeterogeneousVolumes::UseLightingCacheForTransmittance() && bApplyShadowTransmittance) || HeterogeneousVolumes::UseLightingCacheForInscattering())
 		{
-			PassParameters->LightingCache.LightingCacheResolution = HeterogeneousVolumes::GetLightingCacheResolution();
+			PassParameters->LightingCache.LightingCacheResolution = HeterogeneousVolumes::GetLightingCacheResolution(Interface);
 			PassParameters->LightingCache.LightingCacheTexture = LightingCacheTexture;
 		}
 
diff --git a/Engine/Source/Runtime/Renderer/Private/HeterogeneousVolumes/HeterogeneousVolumesPreshadingPipeline.cpp b/Engine/Source/Runtime/Renderer/Private/HeterogeneousVolumes/HeterogeneousVolumesPreshadingPipeline.cpp
--- a/Engine/Source/Runtime/Renderer/Private/HeterogeneousVolumes/HeterogeneousVolumesPreshadingPipeline.cpp
+++ b/Engine/Source/Runtime/Renderer/Private/HeterogeneousVolumes/HeterogeneousVolumesPreshadingPipeline.cpp
@@ -406,6 +406,8 @@ void RenderLightingCacheWithPreshadingCompute(
 	FRDGTextureRef& LightingCacheTexture
 )
 {
+	const IHeterogeneousVolumeInterface* Interface = HeterogeneousVolumes::GetInterface(PrimitiveSceneProxy);
+
 	FRenderLightingCacheWithPreshadingCS::FParameters* PassParameters = GraphBuilder.AllocParameters<FRenderLightingCacheWithPreshadingCS::FParameters>();
 	{
 		// Scene data
@@ -426,7 +428,7 @@ void RenderLightingCacheWithPreshadingCompute(
 		PassParameters->SparseVoxelUniformBuffer = SparseVoxelUniformBuffer;
 
 		// Transmittance volume
-		PassParameters->LightingCache.LightingCacheResolution = HeterogeneousVolumes::GetLightingCacheResolution();
+		PassParameters->LightingCache.LightingCacheResolution = HeterogeneousVolumes::GetLightingCacheResolution(Interface);
 		PassParameters->LightingCache.LightingCacheTexture = LightingCacheTexture;
 
 		// Ray data
@@ -480,7 +482,7 @@ void RenderLightingCacheWithPreshadingCompute(
 	PermutationVector.Set<FRenderLightingCacheWithPreshadingCS::FLightingCacheMode>(HeterogeneousVolumes::GetLightingCacheMode() - 1);
 	TShaderRef<FRenderLightingCacheWithPreshadingCS> ComputeShader = View.ShaderMap->GetShader<FRenderLightingCacheWithPreshadingCS>(PermutationVector);
 
-	FIntVector GroupCount = HeterogeneousVolumes::GetLightingCacheResolution();
+	FIntVector GroupCount = HeterogeneousVolumes::GetLightingCacheResolution(Interface);
 	GroupCount.X = FMath::DivideAndRoundUp(GroupCount.X, FRenderLightingCacheWithPreshadingCS::GetThreadGroupSize3D());
 	GroupCount.Y = FMath::DivideAndRoundUp(GroupCount.Y, FRenderLightingCacheWithPreshadingCS::GetThreadGroupSize3D());
 	GroupCount.Z = FMath::DivideAndRoundUp(GroupCount.Z, FRenderLightingCacheWithPreshadingCS::GetThreadGroupSize3D());
@@ -522,6 +524,8 @@ void RenderSingleScatteringWithPreshadingCompute(
 	FRDGTextureRef& HeterogeneousVolumeTexture
 )
 {
+	const IHeterogeneousVolumeInterface* Interface = HeterogeneousVolumes::GetInterface(PrimitiveSceneProxy);
+
 	FRDGBufferRef VoxelOutputBuffer = GraphBuilder.CreateBuffer(
 		FRDGBufferDesc::CreateStructuredDesc(sizeof(FVoxelDataPacked), HeterogeneousVolumes::GetVoxelCount(SparseVoxelUniformBuffer->GetParameters()->VolumeResolution)),
 		TEXT("HeterogeneousVolumes.VoxelOutputBuffer")
@@ -582,7 +586,7 @@ void RenderSingleScatteringWithPreshadingCompute(
 		// Transmittance volume
 		if ((HeterogeneousVolumes::UseLightingCacheForTransmittance() && bApplyShadowTransmittance) || HeterogeneousVolumes::UseLightingCacheForInscattering())
 		{
-			PassParameters->LightingCache.LightingCacheResolution = HeterogeneousVolumes::GetLightingCacheResolution();
+			PassParameters->LightingCache.LightingCacheResolution = HeterogeneousVolumes::GetLightingCacheResolution(Interface);
 			PassParameters->LightingCache.LightingCacheTexture = LightingCacheTexture;
 		}
 
@@ -1484,8 +1488,10 @@ void RenderWithPreshading(
 	FRDGTextureRef& HeterogeneousVolumeRadiance
 )
 {
+	const IHeterogeneousVolumeInterface* Interface = HeterogeneousVolumes::GetInterface(PrimitiveSceneProxy);
+
 	// Determine baking voxel resolution
-	FIntVector VolumeResolution = HeterogeneousVolumes::GetVolumeResolution();
+	FIntVector VolumeResolution = HeterogeneousVolumes::GetVolumeResolution(Interface);
 
 	// Create baked material grids
 	uint32 NumMips = FMath::Log2(float(FMath::Min(FMath::Min(VolumeResolution.X, VolumeResolution.Y), VolumeResolution.Z))) + 1;
