Add UE_LIFETIMEBOUND to scene UB field getters.

diff --git a/Engine/Source/Runtime/Renderer/Private/SceneUniformBuffer.h b/Engine/Source/Runtime/Renderer/Private/SceneUniformBuffer.h
--- a/Engine/Source/Runtime/Renderer/Private/SceneUniformBuffer.h
+++ b/Engine/Source/Runtime/Renderer/Private/SceneUniformBuffer.h
@@ -154,7 +154,7 @@ public:
 	 * If the field has not been set, this will crash.
 	 */
 	template<typename TMember>
-	const TMember& Get(const TSceneUniformBufferMemberRegistration<TMember>& Registration) const
+	const TMember& Get(const TSceneUniformBufferMemberRegistration<TMember>& Registration) const UE_LIFETIMEBOUND
 	{
 		const void* Ptr = Get(Registration.MemberId, TMember::FTypeInfo::GetStructMetadata()->GetSize());
 		checkf(Ptr, TEXT("SceneUB::%s has not been set yet"), *Registration.Name);
@@ -166,7 +166,7 @@ public:
 	 * If the field has not been set, this will set a default value.
 	 */
 	template<typename TMember>
-	const TMember& GetOrDefault(const TSceneUniformBufferMemberRegistration<TMember>& Registration, FRDGBuilder& GraphBuilder) const
+	const TMember& GetOrDefault(const TSceneUniformBufferMemberRegistration<TMember>& Registration, FRDGBuilder& GraphBuilder) const UE_LIFETIMEBOUND
 	{
 		return *reinterpret_cast<const TMember*>(GetOrDefault(Registration.MemberId, TMember::FTypeInfo::GetStructMetadata()->GetSize(), GraphBuilder));
 	}
@@ -201,18 +201,17 @@ private:
 	 * Retrieve a field in the parameter struct.
 	 * If the field has not been set, this will return null;
 	 */
-	RENDERER_API const void* Get(const FMemberId MemberId, const int32 ExpectedValueSize) const;
+	RENDERER_API const void* Get(const FMemberId MemberId, const int32 ExpectedValueSize) const UE_LIFETIMEBOUND;
 
 	/**
 	 * Retrieve a field in the parameter struct.
 	 * If the field has not been set, this will set a default value;
 	 */
-	RENDERER_API const void* GetOrDefault(const FMemberId MemberId, const int32 ExpectedValueSize, FRDGBuilder& GraphBuilder);
+	RENDERER_API const void* GetOrDefault(const FMemberId MemberId, const int32 ExpectedValueSize, FRDGBuilder& GraphBuilder) UE_LIFETIMEBOUND;
 
 	TRDGUniformBufferRef<FSceneUniformParameters> Buffer;
 	FRHIUniformBuffer* RHIBuffer;
 
-	TBitArray<> MemberIsDirty;
 	bool bAnyMemberDirty;
 	TBitArray<> MemberHasBeenSet;
 	TArray<uint8_t, TSizedHeapAllocator<32>> CachedData;
