Fix a bug of rendering custom depth on OpenGL.

diff --git a/Engine/Source/Runtime/Renderer/Private/CustomDepthRendering.cpp b/Engine/Source/Runtime/Renderer/Private/CustomDepthRendering.cpp
--- a/Engine/Source/Runtime/Renderer/Private/CustomDepthRendering.cpp
+++ b/Engine/Source/Runtime/Renderer/Private/CustomDepthRendering.cpp
@@ -366,17 +366,21 @@ bool FSceneRenderer::RenderCustomDepthPass(
 	else
 	{
 		const FSceneTexturesConfig& Config = FSceneTexturesConfig::Get();
-		FRDGTextureRef CustomDepth = CustomDepthTextures.Depth;
-
 		// TextureView is not supported in GLES, so we can't lookup CustomDepth and CustomStencil from a single texture
 		// Do a copy of the CustomDepthStencil texture if CustomStencil is sampled in a shader.
-		if (IsOpenGLPlatform(ShaderPlatform) && Config.bSamplesCustomStencil)
+		if (IsOpenGLPlatform(ShaderPlatform))
+		{
+			if (Config.bSamplesCustomStencil)
+			{
+				FRDGTextureRef CustomStencil = GraphBuilder.CreateTexture(CustomDepthTextures.Depth->Desc, TEXT("CustomStencil"));
+				AddCopyTexturePass(GraphBuilder, CustomDepthTextures.Depth, CustomStencil);
+				CustomDepthTextures.Stencil = GraphBuilder.CreateSRV(FRDGTextureSRVDesc::CreateWithPixelFormat(CustomStencil, PF_X24_G8));
+			}
+		}
+		else
 		{
-			CustomDepth = GraphBuilder.CreateTexture(CustomDepthTextures.Depth->Desc, TEXT("CustomDepthCopy"));
-			AddCopyTexturePass(GraphBuilder, CustomDepthTextures.Depth, CustomDepth);
+			CustomDepthTextures.Stencil = GraphBuilder.CreateSRV(FRDGTextureSRVDesc::CreateWithPixelFormat(CustomDepthTextures.Depth, PF_X24_G8));
 		}
-
-		CustomDepthTextures.Stencil = GraphBuilder.CreateSRV(FRDGTextureSRVDesc::CreateWithPixelFormat(CustomDepth, PF_X24_G8));
 		CustomDepthTextures.bSeparateStencilBuffer = false;
 	}
 
