Fix wrong initial state for default buffers.

I accidentally removed the usage flags override for these buffers in CL 24336925, so their initial state changed to UAV because they still had EBufferUsageFlags::UnorderedAccess set.

#rnx

diff --git a/Engine/Source/Runtime/Renderer/Private/SystemTextures.cpp b/Engine/Source/Runtime/Renderer/Private/SystemTextures.cpp
--- a/Engine/Source/Runtime/Renderer/Private/SystemTextures.cpp
+++ b/Engine/Source/Runtime/Renderer/Private/SystemTextures.cpp
@@ -1534,12 +1534,16 @@ FRDGBufferRef GetInternalDefaultBuffer(
 	case EDefaultBufferType::StructuredBuffer:
 		CreateInfo.DebugName = TEXT("DefaultStructuredBuffer");
 		BufferDesc = FRDGBufferDesc::CreateStructuredDesc(NumBytePerElement, NumElements);
+		// Remove the UAV flag, as default resources are supposed to be read-only.
+		EnumRemoveFlags(BufferDesc.Usage, EBufferUsageFlags::UnorderedAccess);
 		RHIBuffer = RHICreateStructuredBuffer(NumBytePerElement, BufferSize, BufferDesc.Usage, CreateInfo);
 		break;
 
 	case EDefaultBufferType::ByteAddressBuffer:
 		CreateInfo.DebugName = TEXT("DefaultByteAddressBuffer");
 		BufferDesc = FRDGBufferDesc::CreateByteAddressDesc(BufferSize);
+		// Same as above.
+		EnumRemoveFlags(BufferDesc.Usage, EBufferUsageFlags::UnorderedAccess);
 		RHIBuffer = RHICreateStructuredBuffer(NumBytePerElement, BufferSize, BufferDesc.Usage, CreateInfo);
 		break;
 
