Add support for Lumen Diffuse Indirect without Dual Source Blending

diff --git a/Engine/Shaders/Private/DiffuseIndirectComposite.usf b/Engine/Shaders/Private/DiffuseIndirectComposite.usf
--- a/Engine/Shaders/Private/DiffuseIndirectComposite.usf
+++ b/Engine/Shaders/Private/DiffuseIndirectComposite.usf
@@ -114,6 +114,11 @@ Texture2D DiffuseIndirect_Textures_3;
 #define DiffuseIndirect_Textures_3 DiffuseIndirect_Textures_0
 #endif
 
+#if !ENABLE_DUAL_SRC_BLENDING
+Texture2D SceneColorTexture;
+SamplerState SceneColorSampler;
+#endif
+
 RWTexture2D<float4> PassDebugOutput;
 
 uint bLumenSupportBackfaceDiffuse;
@@ -302,8 +307,12 @@ FStrataDeferredLighting StrataIndirectLighting(
 void MainPS(
 	float4 SvPosition : SV_POSITION
 #if DIM_APPLY_DIFFUSE_INDIRECT
+#if ENABLE_DUAL_SRC_BLENDING
 	, out float4 OutAddColor      DUAL_SOURCE_BLENDING_SLOT(0) : SV_Target0
 	, out float4 OutMultiplyColor DUAL_SOURCE_BLENDING_SLOT(1) : SV_Target1
+#else
+	, out float4 OutColor 		  : SV_Target0
+#endif
 #else
 	, out float4 OutMultiplyColor : SV_Target0
 #endif
@@ -313,6 +322,13 @@ void MainPS(
 	const float2 SceneBufferUV = SvPositionToBufferUV(SvPosition);
 	const float2 ScreenPosition = SvPositionToScreenPosition(SvPosition).xy;
 
+	#if !ENABLE_DUAL_SRC_BLENDING && DIM_APPLY_DIFFUSE_INDIRECT
+	float4 OutAddColor = float4(0.0, 0.0, 0.0, 0.0);
+	float4 OutMultiplyColor;
+
+	const float4 SceneColor = SceneColorTexture.SampleLevel(SceneColorSampler, SceneBufferUV, 0);
+	#endif
+
 	// Sample scene textures.
 	const FLumenMaterialData Material = ReadMaterialData(PixelPos, SceneBufferUV);
 
@@ -541,4 +557,8 @@ void MainPS(
 	#endif
 
 	OutMultiplyColor = FinalAmbientOcclusion;
+
+	#if !ENABLE_DUAL_SRC_BLENDING && DIM_APPLY_DIFFUSE_INDIRECT
+	OutColor = SceneColor * OutMultiplyColor + OutAddColor;
+	#endif
 }
diff --git a/Engine/Source/Runtime/Renderer/Private/IndirectLightRendering.cpp b/Engine/Source/Runtime/Renderer/Private/IndirectLightRendering.cpp
--- a/Engine/Source/Runtime/Renderer/Private/IndirectLightRendering.cpp
+++ b/Engine/Source/Runtime/Renderer/Private/IndirectLightRendering.cpp
@@ -95,6 +95,10 @@ static TAutoConsoleVariable<bool> CVarDiffuseIndirectOffUseDepthBoundsAO(
 	TEXT("Use depth bounds when we apply the AO when DiffuseIndirect is disabled."),
 	ECVF_RenderThreadSafe);
 
+static TAutoConsoleVariable<bool> CVarDiffuseIndirectForceCopyPass(
+	TEXT("r.DiffuseIndirectForceCopyPass"), false,
+	TEXT("Forces use of copy pass instead of dual source blend. (for debugging)"),
+	ECVF_RenderThreadSafe);
 
 DECLARE_GPU_STAT_NAMED(ReflectionEnvironment, TEXT("Reflection Environment"));
 DECLARE_GPU_STAT_NAMED(RayTracingReflections, TEXT("Ray Tracing Reflections"));
@@ -117,8 +121,9 @@ class FDiffuseIndirectCompositePS : public FGlobalShader
 	class FUpscaleDiffuseIndirectDim : SHADER_PERMUTATION_BOOL("DIM_UPSCALE_DIFFUSE_INDIRECT");
 	class FScreenBentNormal : SHADER_PERMUTATION_BOOL("DIM_SCREEN_BENT_NORMAL");
 	class FStrataTileType : SHADER_PERMUTATION_INT("STRATA_TILETYPE", 3);
+	class FEnableDualSrcBlending : SHADER_PERMUTATION_BOOL("ENABLE_DUAL_SRC_BLENDING");
 
-	using FPermutationDomain = TShaderPermutationDomain<FApplyDiffuseIndirectDim, FUpscaleDiffuseIndirectDim, FScreenBentNormal, FStrataTileType>;
+	using FPermutationDomain = TShaderPermutationDomain<FApplyDiffuseIndirectDim, FUpscaleDiffuseIndirectDim, FScreenBentNormal, FStrataTileType, FEnableDualSrcBlending>;
 
 	static bool ShouldCompilePermutation(const FGlobalShaderPermutationParameters& Parameters)
 	{
@@ -166,6 +171,9 @@ class FDiffuseIndirectCompositePS : public FGlobalShader
 		SHADER_PARAMETER_RDG_TEXTURE(Texture2D, AmbientOcclusionTexture)
 		SHADER_PARAMETER_SAMPLER(SamplerState,  AmbientOcclusionSampler)
 
+		SHADER_PARAMETER_RDG_TEXTURE(Texture2D, SceneColorTexture)
+		SHADER_PARAMETER_SAMPLER(SamplerState, SceneColorSampler)
+
 		SHADER_PARAMETER_TEXTURE(Texture2D, PreIntegratedGF)
 		SHADER_PARAMETER_SAMPLER(SamplerState, PreIntegratedGFSampler)
 
@@ -1191,8 +1199,32 @@ void FDeferredShadingSceneRenderer::RenderDiffuseIndirectAndAmbientOcclusion(
 			AddClearRenderTargetPass(GraphBuilder, SceneColorTexture, FLinearColor::Black);
 		}
 
+		bool bEnableCopyPass = (!RHISupportsDualSourceBlending(ShaderPlatform) || CVarDiffuseIndirectForceCopyPass.GetValueOnRenderThread()) 
+								&& DenoiserOutputs.Textures[0];
+
+		FRDGTextureRef SceneColorCopyTexture;
+
+		if (bEnableCopyPass)
+		{
+			FRDGTextureDesc SceneColorCopyTextureDesc = FRDGTextureDesc::Create2D(
+				SceneColorTexture->Desc.Extent,
+				SceneColorTexture->Desc.Format,
+				FClearValueBinding::None,
+				TexCreate_ShaderResource | TexCreate_RenderTargetable);
+
+			
+			SceneColorCopyTexture = GraphBuilder.CreateTexture(SceneColorCopyTextureDesc, TEXT("SceneCopyTextureCopy0"));
+		}
+
 		auto ApplyDiffuseIndirect = [&](EStrataTileType TileType)
 		{
+			if (bEnableCopyPass)
+			{
+				FRHICopyTextureInfo CopyInfo;
+				CopyInfo.Size = SceneColorTexture->Desc.GetSize();
+				AddCopyTexturePass(GraphBuilder, SceneColorTexture, SceneColorCopyTexture, CopyInfo);
+			}
+
 			FDiffuseIndirectCompositePS::FParameters* PassParameters = GraphBuilder.AllocParameters<FDiffuseIndirectCompositePS::FParameters>();
 			PassParameters->Strata = Strata::BindStrataGlobalUniformParameters(View);
 			PassParameters->SceneTexturesStruct = SceneTextures.UniformBuffer;
@@ -1234,6 +1266,12 @@ void FDeferredShadingSceneRenderer::RenderDiffuseIndirectAndAmbientOcclusion(
 
 			PassParameters->AmbientOcclusionTexture = AmbientOcclusionMask;
 			PassParameters->AmbientOcclusionSampler = TStaticSamplerState<SF_Point>::GetRHI();
+
+			if (bEnableCopyPass)
+			{
+				PassParameters->SceneColorTexture = SceneColorCopyTexture;
+				PassParameters->SceneColorSampler = TStaticSamplerState<SF_Point>::GetRHI();
+			}
 			
 			if (!PassParameters->AmbientOcclusionTexture || bIsVisualizePass)
 			{
@@ -1295,11 +1333,21 @@ void FDeferredShadingSceneRenderer::RenderDiffuseIndirectAndAmbientOcclusion(
 				PermutationVector.Set<FDiffuseIndirectCompositePS::FUpscaleDiffuseIndirectDim>(bUpscale);
 			}
 
+			PermutationVector.Set<FDiffuseIndirectCompositePS::FEnableDualSrcBlending>(!bEnableCopyPass);
+			
 			TShaderMapRef<FDiffuseIndirectCompositePS> PixelShader(View.ShaderMap, PermutationVector);
 
-			FRHIBlendState* BlendState = PermutationVector.Get<FDiffuseIndirectCompositePS::FApplyDiffuseIndirectDim>() > 0 ?
-				TStaticBlendState<CW_RGBA, BO_Add, BF_One, BF_Source1Color, BO_Add, BF_One, BF_Source1Alpha>::GetRHI() :
-				TStaticBlendState<CW_RGBA, BO_Add, BF_Zero, BF_SourceColor, BO_Add, BF_Zero, BF_SourceAlpha>::GetRHI();
+			FRHIBlendState* BlendState;
+			if (!bEnableCopyPass)
+			{
+				BlendState = PermutationVector.Get<FDiffuseIndirectCompositePS::FApplyDiffuseIndirectDim>() > 0 ?
+					TStaticBlendState<CW_RGBA, BO_Add, BF_One, BF_Source1Color, BO_Add, BF_One, BF_Source1Alpha>::GetRHI() :
+					TStaticBlendState<CW_RGBA, BO_Add, BF_Zero, BF_SourceColor, BO_Add, BF_Zero, BF_SourceAlpha>::GetRHI();
+			}
+			else
+			{
+				BlendState = TStaticBlendState<>::GetRHI();
+			}
 
 			if (bIsVisualizePass)
 			{
