Fix missing BlueNoise resource binding/remapping for SingleLayerWater.

diff --git a/Engine/Shaders/Private/BasePassPixelShader.usf b/Engine/Shaders/Private/BasePassPixelShader.usf
--- a/Engine/Shaders/Private/BasePassPixelShader.usf
+++ b/Engine/Shaders/Private/BasePassPixelShader.usf
@@ -39,6 +39,10 @@
 	#endif
 #endif
 
+#if MATERIAL_SHADINGMODEL_SINGLELAYERWATER
+#define BlueNoise SingleLayerWater.BlueNoise
+#endif
+
 // Enable Strata. This define & include need to be defined before certains includes (i.e., DBufferDecalShared which uses them internally)
 #if !MATERIAL_IS_STRATA && STRATA_ENABLED
 #undef STRATA_ENABLED
diff --git a/Engine/Source/Runtime/Renderer/Private/SingleLayerWaterRendering.cpp b/Engine/Source/Runtime/Renderer/Private/SingleLayerWaterRendering.cpp
--- a/Engine/Source/Runtime/Renderer/Private/SingleLayerWaterRendering.cpp
+++ b/Engine/Source/Runtime/Renderer/Private/SingleLayerWaterRendering.cpp
@@ -1293,6 +1293,7 @@ BEGIN_UNIFORM_BUFFER_STRUCT(FSingleLayerWaterPassUniformParameters,)
 	SHADER_PARAMETER(uint32, bMainDirectionalLightVSMFiltering)
 	SHADER_PARAMETER(uint32, bSeparateMainDirLightLuminance)
 	SHADER_PARAMETER_STRUCT(FLightCloudTransmittanceParameters, ForwardDirLightCloudShadow)
+	SHADER_PARAMETER_STRUCT(FBlueNoiseParameters, BlueNoise)
 END_UNIFORM_BUFFER_STRUCT()
 
 // At the moment we reuse the DeferredDecals static uniform buffer slot because it is currently unused in this pass.
@@ -1379,6 +1380,7 @@ void FDeferredShadingSceneRenderer::RenderSingleLayerWaterInner(
 			SLWUniformParameters.SceneWithoutSingleLayerWaterInvTextureSize = FVector2f(1.0f / DepthTextureSize.X, 1.0f / DepthTextureSize.Y);
 			SLWUniformParameters.bMainDirectionalLightVSMFiltering = IsWaterVirtualShadowMapFilteringEnabled_Runtime(View.GetShaderPlatform());
 			SLWUniformParameters.bSeparateMainDirLightLuminance = NeedsSeparatedMainDirectionalLightTexture_Runtime(View.GetShaderPlatform());
+			SLWUniformParameters.BlueNoise = GetBlueNoiseParameters();
 
 			const FLightSceneProxy* SelectedForwardDirectionalLightProxy = View.ForwardLightingResources.SelectedForwardDirectionalLightProxy;
 			SetupLightCloudTransmittanceParameters(GraphBuilder, Scene, View, SelectedForwardDirectionalLightProxy ? SelectedForwardDirectionalLightProxy->GetLightSceneInfo() : nullptr, SLWUniformParameters.ForwardDirLightCloudShadow);
