Adding missing GCC_ALIGN

diff --git a/Engine/Source/Runtime/Renderer/Private/SceneUniformBuffer.h b/Engine/Source/Runtime/Renderer/Private/SceneUniformBuffer.h
--- a/Engine/Source/Runtime/Renderer/Private/SceneUniformBuffer.h
+++ b/Engine/Source/Runtime/Renderer/Private/SceneUniformBuffer.h
@@ -124,7 +124,7 @@ private:
 	// This requires some work in buffer binding/parameter macros/RHI, as these all assume the data to be inline.
 	static constexpr size_t MaxSize = 4096;
 	uint8_t Data[MaxSize];
-};
+} GCC_ALIGN(SHADER_PARAMETER_STRUCT_ALIGNMENT);
 
 
 /**
