Setting the correct fallback for QuadOverdraw. The black system texture is a different format, changing this fallback to use ZeroUInt which is PF_R32_UINT.

diff --git a/Engine/Source/Runtime/Renderer/Private/PostProcess/PostProcessVisualizeComplexity.cpp b/Engine/Source/Runtime/Renderer/Private/PostProcess/PostProcessVisualizeComplexity.cpp
--- a/Engine/Source/Runtime/Renderer/Private/PostProcess/PostProcessVisualizeComplexity.cpp
+++ b/Engine/Source/Runtime/Renderer/Private/PostProcess/PostProcessVisualizeComplexity.cpp
@@ -129,7 +129,7 @@ FScreenPassTexture AddVisualizeComplexityPass(FRDGBuilder& GraphBuilder, const F
 	if (SceneTextures.QuadOverdraw)
 	{
 		const FRDGSystemTextures& SystemTextures = FRDGSystemTextures::Get(GraphBuilder);
-		PassParameters->QuadOverdrawTexture = GetIfProduced(SceneTextures.QuadOverdraw, SystemTextures.Black);
+		PassParameters->QuadOverdrawTexture = HasBeenProduced(SceneTextures.QuadOverdraw) ? SceneTextures.QuadOverdraw : GSystemTextures.GetZeroUIntDummy(GraphBuilder);
 		PassParameters->DebugViewShaderMode = DebugViewShaderMode;
 		QuadOverdrawEnum = FVisualizeComplexityApplyPS::EQuadOverdraw::Enable;
 	}
