Fix MID leaks when unloading a GFP level.  Prevents [redacted] from crashing.

diff --git a/Engine/Source/Runtime/Engine/Classes/Engine/LocalPlayer.h b/Engine/Source/Runtime/Engine/Classes/Engine/LocalPlayer.h
--- a/Engine/Source/Runtime/Engine/Classes/Engine/LocalPlayer.h
+++ b/Engine/Source/Runtime/Engine/Classes/Engine/LocalPlayer.h
@@ -587,7 +587,7 @@ public:
 	/**
 	 * Clear cached view state.  Suitable for calling when cleaning up the world but the view state has some references objects (usually mids) owned by the world (thus preventing GC) 
 	 */
-	virtual void CleanupViewState();
+	virtual void CleanupViewState(FStringView MidParentRootPath = {});
 
 	/** Locked view state needs access to GetViewPoint. */
 	friend class FLockedViewState;
diff --git a/Engine/Source/Runtime/Engine/Private/LocalPlayer.cpp b/Engine/Source/Runtime/Engine/Private/LocalPlayer.cpp
--- a/Engine/Source/Runtime/Engine/Private/LocalPlayer.cpp
+++ b/Engine/Source/Runtime/Engine/Private/LocalPlayer.cpp
@@ -1834,14 +1834,14 @@ bool ULocalPlayer::IsPrimaryPlayer() const
 	return GetLocalPlayerIndex() == 0;
 }
 
-void ULocalPlayer::CleanupViewState()
+void ULocalPlayer::CleanupViewState(FStringView MidParentRootPath /*= {}*/)
 {
 	for (FSceneViewStateReference& State : ViewStates)
 	{
 		FSceneViewStateInterface* Ref = State.GetReference();
 		if (Ref)
 		{
-			Ref->ClearMIDPool();
+			Ref->ClearMIDPool(MidParentRootPath);
 		}
 	}
 }
diff --git a/Engine/Source/Runtime/Engine/Public/SceneManagement.h b/Engine/Source/Runtime/Engine/Public/SceneManagement.h
--- a/Engine/Source/Runtime/Engine/Public/SceneManagement.h
+++ b/Engine/Source/Runtime/Engine/Public/SceneManagement.h
@@ -204,7 +204,7 @@ public:
 	/**
 	 * Clears the pool of mids being referenced by this view state 
 	 */
-	virtual void ClearMIDPool() = 0;
+	virtual void ClearMIDPool(FStringView MidParentRootPath = {}) = 0;
 
 #if !(UE_BUILD_SHIPPING || UE_BUILD_TEST)
 	/** If frozen view matrices are available, return a pointer to them */
diff --git a/Engine/Source/Runtime/Renderer/Private/ScenePrivate.h b/Engine/Source/Runtime/Renderer/Private/ScenePrivate.h
--- a/Engine/Source/Runtime/Renderer/Private/ScenePrivate.h
+++ b/Engine/Source/Runtime/Renderer/Private/ScenePrivate.h
@@ -52,6 +52,8 @@
 #include "Experimental/Containers/RobinHoodHashTable.h"
 #include "SpanAllocator.h"
 #include "GlobalDistanceField.h"
+#include "Algo/RemoveIf.h"
+#include "UObject/Package.h"
 
 /** Factor by which to grow occlusion tests **/
 #define OCCLUSION_SLOP (1.0f)
@@ -1422,10 +1424,21 @@ PRAGMA_ENABLE_DEPRECATION_WARNINGS
 		return NewMID;
 	}
 
-	virtual void ClearMIDPool() override
+	virtual void ClearMIDPool(FStringView MidParentRootPath = {}) override
 	{
 		check(IsInGameThread());
-		MIDPool.Empty();
+		if (MidParentRootPath.IsEmpty())
+		{
+			MIDPool.Empty();
+			return;
+		}
+
+		const int32 RemoveNum = Algo::RemoveIf(MIDPool, [MidParentRootPath](UMaterialInstanceDynamic* MID) -> bool
+		{
+			return MID->Parent && MID->Parent->GetPackage()->GetName().StartsWith(MidParentRootPath);
+		});
+
+		MIDPool.SetNum(RemoveNum);
 	}
 
 #if !(UE_BUILD_SHIPPING || UE_BUILD_TEST)
