Removing the unused FGammaShaderParameters

diff --git a/Engine/Source/Runtime/Renderer/Private/PostProcess/SceneFilterRendering.cpp b/Engine/Source/Runtime/Renderer/Private/PostProcess/SceneFilterRendering.cpp
--- a/Engine/Source/Runtime/Renderer/Private/PostProcess/SceneFilterRendering.cpp
+++ b/Engine/Source/Runtime/Renderer/Private/PostProcess/SceneFilterRendering.cpp
@@ -15,8 +15,6 @@
 #include "PostProcess/DrawRectangle.h"
 #include "ScenePrivate.h"
 
-IMPLEMENT_TYPE_LAYOUT(FGammaShaderParameters);
-
 void FTesselatedScreenRectangleIndexBuffer::InitRHI()
 {
 	TResourceArray<uint16, INDEXBUFFER_ALIGNMENT> IndexBuffer;
diff --git a/Engine/Source/Runtime/Renderer/Private/PostProcess/SceneFilterRendering.h b/Engine/Source/Runtime/Renderer/Private/PostProcess/SceneFilterRendering.h
--- a/Engine/Source/Runtime/Renderer/Private/PostProcess/SceneFilterRendering.h
+++ b/Engine/Source/Runtime/Renderer/Private/PostProcess/SceneFilterRendering.h
@@ -111,102 +111,6 @@ extern RENDERER_API void DrawPostProcessPass(
 	EDrawRectangleFlags Flags = EDRF_Default
 	);
 
-
-
-/*-----------------------------------------------------------------------------
-FMGammaShaderParameters
------------------------------------------------------------------------------*/
-
-/** Encapsulates the gamma correction parameters. */
-class FGammaShaderParameters
-{
-	DECLARE_TYPE_LAYOUT(FGammaShaderParameters, NonVirtual);
-public:
-
-	/** Default constructor. */
-	FGammaShaderParameters() {}
-
-	/** Initialization constructor. */
-	FGammaShaderParameters(const FShaderParameterMap& ParameterMap)
-	{
-		RenderTargetExtent.Bind(ParameterMap,TEXT("RenderTargetExtent"));
-		GammaColorScaleAndInverse.Bind(ParameterMap,TEXT("GammaColorScaleAndInverse"));
-		GammaOverlayColor.Bind(ParameterMap,TEXT("GammaOverlayColor"));
-	}
-
-	/** Set the material shader parameter values. */
-	void Set(FRHICommandList& RHICmdList, FRHIPixelShader* RHIShader, const FViewInfo& View, float DisplayGamma, FLinearColor const& ColorScale, FLinearColor const& ColorOverlay)
-	{
-		// GammaColorScaleAndInverse
-
-		float InvDisplayGamma = 1.f / FMath::Max<float>(DisplayGamma,KINDA_SMALL_NUMBER);
-		float OneMinusOverlayBlend = 1.f - ColorOverlay.A;
-
-		FVector4f ColorScaleAndInverse;
-
-		ColorScaleAndInverse.X = ColorScale.R * OneMinusOverlayBlend;
-		ColorScaleAndInverse.Y = ColorScale.G * OneMinusOverlayBlend;
-		ColorScaleAndInverse.Z = ColorScale.B * OneMinusOverlayBlend;
-		ColorScaleAndInverse.W = InvDisplayGamma;
-
-		SetShaderValue(
-			RHICmdList,
-			RHIShader,
-			GammaColorScaleAndInverse,
-			ColorScaleAndInverse
-			);
-
-		// GammaOverlayColor
-
-		FVector4f OverlayColor;
-
-		OverlayColor.X = ColorOverlay.R * ColorOverlay.A;
-		OverlayColor.Y = ColorOverlay.G * ColorOverlay.A;
-		OverlayColor.Z = ColorOverlay.B * ColorOverlay.A;
-		OverlayColor.W = 0.f; // Unused
-
-		SetShaderValue(
-			RHICmdList,
-			RHIShader,
-			GammaOverlayColor,
-			OverlayColor
-			);
-
-		FIntPoint BufferSize = GetSceneTextureExtentFromView(View);
-		float BufferSizeX = (float)BufferSize.X;
-		float BufferSizeY = (float)BufferSize.Y;
-		float InvBufferSizeX = 1.0f / BufferSizeX;
-		float InvBufferSizeY = 1.0f / BufferSizeY;
-
-		const FVector4f vRenderTargetExtent(BufferSizeX, BufferSizeY,  InvBufferSizeX, InvBufferSizeY);
-
-		SetShaderValue(
-			RHICmdList,
-			RHIShader,
-			RenderTargetExtent,
-			vRenderTargetExtent);
-	}
-
-	/** Serializer. */
-	friend FArchive& operator<<(FArchive& Ar,FGammaShaderParameters& P)
-	{
-		Ar << P.GammaColorScaleAndInverse;
-		Ar << P.GammaOverlayColor;
-		Ar << P.RenderTargetExtent;
-
-		return Ar;
-	}
-
-
-private:
-	
-		LAYOUT_FIELD(FShaderParameter, GammaColorScaleAndInverse)
-		LAYOUT_FIELD(FShaderParameter, GammaOverlayColor)
-		LAYOUT_FIELD(FShaderParameter, RenderTargetExtent)
-	
-};
-
-
 class FTesselatedScreenRectangleIndexBuffer : public FIndexBuffer
 {
 public:
