Re-add restrictions for shadowprojection

diff --git a/Engine/Source/Runtime/Renderer/Private/MobileShadingRenderer.cpp b/Engine/Source/Runtime/Renderer/Private/MobileShadingRenderer.cpp
--- a/Engine/Source/Runtime/Renderer/Private/MobileShadingRenderer.cpp
+++ b/Engine/Source/Runtime/Renderer/Private/MobileShadingRenderer.cpp
@@ -537,7 +537,13 @@ void FMobileSceneRenderer::InitViews(
 		&& !ViewFamily.UseDebugViewPS()
 		&& (CVarDistanceFieldShadowQuality != nullptr && CVarDistanceFieldShadowQuality->GetInt() > 0);
 
-	bRequiresShadowProjections = MobileUsesShadowMaskTexture(ShaderPlatform);
+	bRequiresShadowProjections = MobileUsesShadowMaskTexture(ShaderPlatform)
+		&& ViewFamily.EngineShowFlags.Lighting
+		&& !Views[0].bIsReflectionCapture
+		&& !Views[0].bIsPlanarReflection
+		&& !ViewFamily.EngineShowFlags.HitProxies
+		&& !ViewFamily.EngineShowFlags.VisualizeLightCulling
+		&& !ViewFamily.UseDebugViewPS();
 
 	bShouldRenderHZB = ShouldRenderHZB();
 
