Fix crash on platform not loading blue noise resources.
Bind dummy BlueNoise resources when VSM filter quality is not set to high.

diff --git a/Engine/Source/Runtime/Renderer/Private/BlueNoise.cpp b/Engine/Source/Runtime/Renderer/Private/BlueNoise.cpp
--- a/Engine/Source/Runtime/Renderer/Private/BlueNoise.cpp
+++ b/Engine/Source/Runtime/Renderer/Private/BlueNoise.cpp
@@ -8,9 +8,20 @@
 #include "Engine/Engine.h"
 #include "Engine/Texture2D.h"
 #include "TextureResource.h"
+#include "SystemTextures.h"
 
 IMPLEMENT_GLOBAL_SHADER_PARAMETER_STRUCT(FBlueNoise, "BlueNoise");
 
+FBlueNoiseParameters GetBlueNoiseDummyParameters()
+{
+	FBlueNoiseParameters Out;
+	Out.Dimensions = FIntVector(1,1,1);
+	Out.ModuloMasks = FIntVector(1,1,1);
+	Out.ScalarTexture = GSystemTextures.BlackDummy->GetRHI();
+	Out.Vec2Texture = GSystemTextures.BlackDummy->GetRHI();
+	return Out;
+}
+
 FBlueNoiseParameters GetBlueNoiseParameters()
 {
 	FBlueNoiseParameters Out;
diff --git a/Engine/Source/Runtime/Renderer/Private/BlueNoise.h b/Engine/Source/Runtime/Renderer/Private/BlueNoise.h
--- a/Engine/Source/Runtime/Renderer/Private/BlueNoise.h
+++ b/Engine/Source/Runtime/Renderer/Private/BlueNoise.h
@@ -22,3 +22,4 @@ END_GLOBAL_SHADER_PARAMETER_STRUCT()
 
 extern RENDERER_API FBlueNoise GetBlueNoiseGlobalParameters();
 extern RENDERER_API FBlueNoiseParameters GetBlueNoiseParameters();
+extern RENDERER_API FBlueNoiseParameters GetBlueNoiseDummyParameters();
diff --git a/Engine/Source/Runtime/Renderer/Private/TranslucentRendering.cpp b/Engine/Source/Runtime/Renderer/Private/TranslucentRendering.cpp
--- a/Engine/Source/Runtime/Renderer/Private/TranslucentRendering.cpp
+++ b/Engine/Source/Runtime/Renderer/Private/TranslucentRendering.cpp
@@ -32,6 +32,7 @@ DEFINE_GPU_DRAWCALL_STAT(Translucency);
 
 // Forward declarations
 bool ShouldRenderVolumetricCloud(const FScene* Scene, const FEngineShowFlags& EngineShowFlags);
+bool IsVSMTranslucentHighQualityEnabled();
 
 static TAutoConsoleVariable<float> CVarSeparateTranslucencyScreenPercentage(
 	TEXT("r.SeparateTranslucencyScreenPercentage"),
@@ -980,7 +981,16 @@ TRDGUniformBufferRef<FTranslucentBasePassUniformParameters> CreateTranslucentBas
 	BasePassParameters.PreIntegratedGFSampler = TStaticSamplerState<SF_Bilinear, AM_Clamp, AM_Clamp, AM_Clamp>::GetRHI();
 
 	OIT::SetOITParameters(GraphBuilder, View, BasePassParameters.OIT, OITData);
-	BasePassParameters.BlueNoise = GetBlueNoiseParameters();
+
+	// Only use blue noise resources if VSM quality is set to high
+	if (IsVSMTranslucentHighQualityEnabled())
+	{
+		BasePassParameters.BlueNoise = GetBlueNoiseParameters();
+	}
+	else
+	{
+		BasePassParameters.BlueNoise = GetBlueNoiseDummyParameters();
+	}
 
 	return GraphBuilder.CreateUniformBuffer(&BasePassParameters);
 }
diff --git a/Engine/Source/Runtime/Renderer/Private/VirtualShadowMaps/VirtualShadowMapProjection.cpp b/Engine/Source/Runtime/Renderer/Private/VirtualShadowMaps/VirtualShadowMapProjection.cpp
--- a/Engine/Source/Runtime/Renderer/Private/VirtualShadowMaps/VirtualShadowMapProjection.cpp
+++ b/Engine/Source/Runtime/Renderer/Private/VirtualShadowMaps/VirtualShadowMapProjection.cpp
@@ -173,6 +173,11 @@ static TAutoConsoleVariable<int32> CVarTestPermutation(
 // Is also used as the workgroup size for the CS without tile list.
 static constexpr int32 VSMProjectionWorkTileSize = 8;
 
+bool IsVSMTranslucentHighQualityEnabled()
+{
+	return CVarVSMTranslucentQuality.GetValueOnRenderThread() > 0;
+}
+
 FVirtualShadowMapSMRTSettings GetVirtualShadowMapSMRTSettings(bool bDirectionalLight)
 {
 	FVirtualShadowMapSMRTSettings Out;
