Workaround Z-fighting on Switch by disabling DXC for base pass vertex shaders.

We currently enable DXC only for user created materials for certain platforms.
This causes Z-fighting due to mismatch between HLSLCC and DXC compiled shaders.
This workaround disables CVar "r.Shaders.ForceDXC" in the respective platform compiler backend when "WORKAROUND_DISABLE_rShadersForceDXC" is defined; Currently only in TMobileBasePassVSPolicyParamType::ModifyCompilationEnvironment.

diff --git a/Engine/Source/Runtime/Renderer/Private/MobileBasePassRendering.h b/Engine/Source/Runtime/Renderer/Private/MobileBasePassRendering.h
--- a/Engine/Source/Runtime/Renderer/Private/MobileBasePassRendering.h
+++ b/Engine/Source/Runtime/Renderer/Private/MobileBasePassRendering.h
@@ -138,6 +138,11 @@ public:
 	static void ModifyCompilationEnvironment(const FMaterialShaderPermutationParameters& Parameters, FShaderCompilerEnvironment& OutEnvironment)
 	{
 		FMeshMaterialShader::ModifyCompilationEnvironment(Parameters, OutEnvironment);
+
+		// @lh-todo:
+		//		This is a workaround to prevent Z-fighting due to a mismatch between HLSLCC compiled shaders (from cooked materials) and DXC compiled shaders (from user materials).
+		//		It effectively disables CVar "r.Shaders.ForceDXC" which we partially enabled during the HLSLCC to DXC transition for certain platform backends.
+		OutEnvironment.SetDefine(TEXT("WORKAROUND_DISABLE_rShadersForceDXC"), 1);
 	}
 
 	void GetShaderBindings(
@@ -230,6 +235,9 @@ public:
 		FMeshMaterialShader::ModifyCompilationEnvironment(Parameters, OutEnvironment);
 		// Modify compilation environment depending upon material shader quality level settings.
 		ModifyCompilationEnvironmentForQualityLevel(Parameters.Platform, Parameters.MaterialParameters.QualityLevel, OutEnvironment);
+
+		// @lh-todo: Same workaround as for the VS of MobileBasePass. See TMobileBasePassVSPolicyParamType::ModifyCompilationEnvironment for details.
+		OutEnvironment.SetDefine(TEXT("WORKAROUND_DISABLE_rShadersForceDXC"), 1);
 	}
 
 	/** Initialization constructor. */
