Fix a raytracing shader to respect the proper switches.

#rnx

diff --git a/Engine/Source/Runtime/Renderer/Private/RayTracing/RayTracingInstanceBufferUtil.cpp b/Engine/Source/Runtime/Renderer/Private/RayTracing/RayTracingInstanceBufferUtil.cpp
--- a/Engine/Source/Runtime/Renderer/Private/RayTracing/RayTracingInstanceBufferUtil.cpp
+++ b/Engine/Source/Runtime/Renderer/Private/RayTracing/RayTracingInstanceBufferUtil.cpp
@@ -321,7 +321,7 @@ struct FRayTracingBuildInstanceBufferCS : public FGlobalShader
 
 	static bool ShouldCompilePermutation(const FGlobalShaderPermutationParameters& Parameters)
 	{
-		return RHISupportsRayTracing(Parameters.Platform);
+		return IsRayTracingEnabledForProject(Parameters.Platform);
 	}
 };
 
