Fixes missing TexCreate_AtomicCompatible on TSR.PrevScatteredSubpixelDepth for TSR to work correctly on metal

diff --git a/Engine/Source/Runtime/Renderer/Private/PostProcess/TemporalSuperResolution.cpp b/Engine/Source/Runtime/Renderer/Private/PostProcess/TemporalSuperResolution.cpp
--- a/Engine/Source/Runtime/Renderer/Private/PostProcess/TemporalSuperResolution.cpp
+++ b/Engine/Source/Runtime/Renderer/Private/PostProcess/TemporalSuperResolution.cpp
@@ -1334,7 +1334,10 @@ ITemporalUpscaler::FOutputs AddTemporalSuperResolutionPasses(
 		PassParameters->PrevClosestDepthOutput = GraphBuilder.CreateUAV(PrevClosestDepthTexture);
 		if (History.SubpixelDepth && PrevHistory.SubpixelDepth)
 		{
-			PrevScatteredSubpixelDepthTexture = GraphBuilder.CreateTexture(History.SubpixelDepth->Desc, TEXT("TSR.PrevScatteredSubpixelDepth"));
+			FRDGTextureDesc Desc = History.SubpixelDepth->Desc;
+			Desc.Flags |= TexCreate_AtomicCompatible;
+
+			PrevScatteredSubpixelDepthTexture = GraphBuilder.CreateTexture(Desc, TEXT("TSR.PrevScatteredSubpixelDepth"));
 			PassParameters->SubpixelDepthOutput = GraphBuilder.CreateUAV(PrevScatteredSubpixelDepthTexture);
 		}
 		else
