Fix missing switch case for RVT blend state.

diff --git a/Engine/Source/Runtime/Renderer/Private/VT/RuntimeVirtualTextureRender.cpp b/Engine/Source/Runtime/Renderer/Private/VT/RuntimeVirtualTextureRender.cpp
--- a/Engine/Source/Runtime/Renderer/Private/VT/RuntimeVirtualTextureRender.cpp
+++ b/Engine/Source/Runtime/Renderer/Private/VT/RuntimeVirtualTextureRender.cpp
@@ -233,6 +233,7 @@ namespace RuntimeVirtualTexture
 			case 17: return TGetBlendStateFromAttributeMask<17>();
 			case 18: return TGetBlendStateFromAttributeMask<18>();
 			case 19: return TGetBlendStateFromAttributeMask<19>();
+			case 20: return TGetBlendStateFromAttributeMask<20>();
 			case 21: return TGetBlendStateFromAttributeMask<21>();
 			case 22: return TGetBlendStateFromAttributeMask<22>();
 			case 23: return TGetBlendStateFromAttributeMask<23>();
@@ -320,6 +321,7 @@ namespace RuntimeVirtualTexture
 			case 17: return TGetBlendStateFromAttributeMask<17>();
 			case 18: return TGetBlendStateFromAttributeMask<18>();
 			case 19: return TGetBlendStateFromAttributeMask<19>();
+			case 20: return TGetBlendStateFromAttributeMask<20>();
 			case 21: return TGetBlendStateFromAttributeMask<21>();
 			case 22: return TGetBlendStateFromAttributeMask<22>();
 			case 23: return TGetBlendStateFromAttributeMask<23>();
