Fix crash with blue noise resources in SLW on lower platforms.

diff --git a/Engine/Source/Runtime/Renderer/Private/SingleLayerWaterRendering.cpp b/Engine/Source/Runtime/Renderer/Private/SingleLayerWaterRendering.cpp
--- a/Engine/Source/Runtime/Renderer/Private/SingleLayerWaterRendering.cpp
+++ b/Engine/Source/Runtime/Renderer/Private/SingleLayerWaterRendering.cpp
@@ -1310,6 +1310,8 @@ BEGIN_SHADER_PARAMETER_STRUCT(FSingleLayerWaterPassParameters, )
 	RENDER_TARGET_BINDING_SLOTS()
 END_SHADER_PARAMETER_STRUCT()
 
+bool IsVSMTranslucentHighQualityEnabled();
+
 void FDeferredShadingSceneRenderer::RenderSingleLayerWaterInner(
 	FRDGBuilder& GraphBuilder,
 	const FSceneTextures& SceneTextures,
@@ -1380,7 +1382,15 @@ void FDeferredShadingSceneRenderer::RenderSingleLayerWaterInner(
 			SLWUniformParameters.SceneWithoutSingleLayerWaterInvTextureSize = FVector2f(1.0f / DepthTextureSize.X, 1.0f / DepthTextureSize.Y);
 			SLWUniformParameters.bMainDirectionalLightVSMFiltering = IsWaterVirtualShadowMapFilteringEnabled_Runtime(View.GetShaderPlatform());
 			SLWUniformParameters.bSeparateMainDirLightLuminance = NeedsSeparatedMainDirectionalLightTexture_Runtime(View.GetShaderPlatform());
-			SLWUniformParameters.BlueNoise = GetBlueNoiseParameters();
+			// Only use blue noise resources if VSM quality is set to high
+			if (IsVSMTranslucentHighQualityEnabled())
+			{
+				SLWUniformParameters.BlueNoise = GetBlueNoiseParameters();
+			}
+			else
+			{
+				SLWUniformParameters.BlueNoise = GetBlueNoiseDummyParameters();
+			}
 
 			const FLightSceneProxy* SelectedForwardDirectionalLightProxy = View.ForwardLightingResources.SelectedForwardDirectionalLightProxy;
 			SetupLightCloudTransmittanceParameters(GraphBuilder, Scene, View, SelectedForwardDirectionalLightProxy ? SelectedForwardDirectionalLightProxy->GetLightSceneInfo() : nullptr, SLWUniformParameters.ForwardDirLightCloudShadow);
