Fix `GPUDump` crash when path tracing buffers are used by post process materials.

diff --git a/Engine/Source/Runtime/Renderer/Private/PostProcess/PostProcessMaterial.cpp b/Engine/Source/Runtime/Renderer/Private/PostProcess/PostProcessMaterial.cpp
--- a/Engine/Source/Runtime/Renderer/Private/PostProcess/PostProcessMaterial.cpp
+++ b/Engine/Source/Runtime/Renderer/Private/PostProcess/PostProcessMaterial.cpp
@@ -815,19 +815,27 @@ bool IsPathTracingVarianceTextureRequiredInPostProcessMaterial(const FViewInfo&
 	bool bIsPathTracingVarianceTextureRequired = false;
 
 	auto CheckIfPathTracingVarianceTextureIsRequried = [&](const UMaterialInterface* MaterialInterface) {
-		// Get the RenderProxy of the material.
-		const FMaterialRenderProxy* MaterialProxy = MaterialInterface->GetRenderProxy();
-		check(MaterialProxy);
-
-		// Get the Shadermap for the view's feature level
-		const FMaterial* Material = MaterialProxy->GetMaterialNoFallback(View.FeatureLevel);
-		if (Material && Material->GetMaterialDomain() == MD_PostProcess)
+		
+		if (MaterialInterface)
 		{
-			const FMaterialShaderMap* MaterialShaderMap = Material->GetRenderingThreadShaderMap();
-			check(MaterialShaderMap);
-			if (MaterialShaderMap->UsesPathTracingBufferTexture(static_cast<uint32>(EPathTracingPostProcessMaterialInput::Variance)))
+			// Get the RenderProxy of the material.
+			const FMaterialRenderProxy* MaterialProxy = MaterialInterface->GetRenderProxy();
+
+			if (MaterialProxy)
 			{
-				return true;
+
+				// Get the Shadermap for the view's feature level
+				const FMaterial* Material = MaterialProxy->GetMaterialNoFallback(View.FeatureLevel);
+				if (Material && Material->GetMaterialDomain() == MD_PostProcess)
+				{
+					const FMaterialShaderMap* MaterialShaderMap = Material->GetRenderingThreadShaderMap();
+
+					if (MaterialShaderMap &&
+						MaterialShaderMap->UsesPathTracingBufferTexture(static_cast<uint32>(EPathTracingPostProcessMaterialInput::Variance)))
+					{
+						return true;
+					}
+				}
 			}
 		}
 
