Move mobile specific GPU instancing buffer to use ByteAddressBuffer instead of StructuredBuffer as mixing Vertex usage and StructuredBuffer usage is not allowed in D3D11

diff --git a/Engine/Shaders/Private/InstanceCulling/BuildInstanceDrawCommands.usf b/Engine/Shaders/Private/InstanceCulling/BuildInstanceDrawCommands.usf
--- a/Engine/Shaders/Private/InstanceCulling/BuildInstanceDrawCommands.usf
+++ b/Engine/Shaders/Private/InstanceCulling/BuildInstanceDrawCommands.usf
@@ -50,7 +50,7 @@ int bDrawOnlyVSMInvalidatingGeometry;
 
 
 RWStructuredBuffer<uint> InstanceIdsBufferOut;
-RWStructuredBuffer<float4> InstanceIdsBufferOutMobile;
+RWByteAddressBuffer InstanceIdsBufferOutMobile;
 RWBuffer<uint> DrawIndirectArgsBufferOut;
 
 // Used for draw commands that need to use compaction to preserve instance order
@@ -123,8 +123,9 @@ void WritePackedInstanceData(uint Offset, uint InstanceId, FInstanceSceneData In
 	InstanceAuxData4 = PackLightmapData(LightMapUVScaleBias, ShadowMapUVScaleBias);
 	InstanceAuxData1 = asfloat(PrimitiveData.LightmapDataIndex);
 #else
-	InstanceAuxData4.xy = PackLocalBoundsCenter(InstanceData.LocalBoundsCenter);
-	InstanceAuxData4.zw = PackLocalBoundsExtent(InstanceData.LocalBoundsExtent);
+	float2 PackedCenter = PackLocalBoundsCenter(InstanceData.LocalBoundsCenter);
+	float2 PackedExtent = PackLocalBoundsExtent(InstanceData.LocalBoundsExtent);
+	InstanceAuxData4 = float4(PackedCenter, PackedExtent);
 	InstanceAuxData1 = InstanceData.RandomID;
 #endif // ALLOW_STATIC_LIGHTING
 
@@ -142,16 +143,16 @@ void WritePackedInstanceData(uint Offset, uint InstanceId, FInstanceSceneData In
 	InstanceTransform1.w = asfloat(PackedPrimitiveFlags); 
 	InstanceTransform2.w = InstanceAuxData1; 
 
-	uint CustomDataCount;
+	uint CustomDataCount = 0u;
 	LoadInstanceRelativeIdAndCustomDataCount(InstanceId, InstanceSceneDataSOAStride, InstanceData.RelativeId, CustomDataCount);
 	FInstancePayloadDataOffsets Offsets = GetInstancePayloadDataOffsets(InstanceData.PrimitiveId, InstanceData.Flags, InstanceData.RelativeId);
 	InstanceTransform3.w = asfloat((CustomDataCount & ((1 << INSTANCE_CUSTOM_DATA_COUNT_NUM_BITS) - 1)) | (Offsets.CustomData << INSTANCE_CUSTOM_DATA_COUNT_NUM_BITS));
 		
-	InstanceIdsBufferOutMobile[Offset*5 + 0] = InstanceOriginAndId;
-	InstanceIdsBufferOutMobile[Offset*5 + 1] = InstanceTransform1;
-	InstanceIdsBufferOutMobile[Offset*5 + 2] = InstanceTransform2;
-	InstanceIdsBufferOutMobile[Offset*5 + 3] = InstanceTransform3;
-	InstanceIdsBufferOutMobile[Offset*5 + 4] = InstanceAuxData4;
+	InstanceIdsBufferOutMobile.Store4((Offset*5 + 0) * 16, asuint(InstanceOriginAndId));
+	InstanceIdsBufferOutMobile.Store4((Offset*5 + 1) * 16, asuint(InstanceTransform1));
+	InstanceIdsBufferOutMobile.Store4((Offset*5 + 2) * 16, asuint(InstanceTransform2));
+	InstanceIdsBufferOutMobile.Store4((Offset*5 + 3) * 16, asuint(InstanceTransform3));
+	InstanceIdsBufferOutMobile.Store4((Offset*5 + 4) * 16, asuint(InstanceAuxData4));
 }
 #endif
 
diff --git a/Engine/Shaders/Private/InstanceCulling/CompactVisibleInstances.usf b/Engine/Shaders/Private/InstanceCulling/CompactVisibleInstances.usf
--- a/Engine/Shaders/Private/InstanceCulling/CompactVisibleInstances.usf
+++ b/Engine/Shaders/Private/InstanceCulling/CompactVisibleInstances.usf
@@ -102,7 +102,7 @@ StructuredBuffer<uint> BlockDestInstanceOffsets;
 /** The compacted output buffer of InstanceIds (non-mobile) */
 RWStructuredBuffer<uint> InstanceIdsBufferOut;
 /** The compacted output buffer of instance data consumed by vertex shaders on mobile */
-RWStructuredBuffer<float4> InstanceIdsBufferOutMobile;
+RWByteAddressBuffer InstanceIdsBufferOutMobile;
 
 /** Entry point */
 [numthreads(NUM_THREADS_PER_GROUP, 1, 1)]
diff --git a/Engine/Source/Runtime/Renderer/Private/InstanceCulling/InstanceCullingContext.cpp b/Engine/Source/Runtime/Renderer/Private/InstanceCulling/InstanceCullingContext.cpp
--- a/Engine/Source/Runtime/Renderer/Private/InstanceCulling/InstanceCullingContext.cpp
+++ b/Engine/Source/Runtime/Renderer/Private/InstanceCulling/InstanceCullingContext.cpp
@@ -306,7 +306,7 @@ public:
 		SHADER_PARAMETER_RDG_BUFFER_SRV(StructuredBuffer<uint32>, BlockDestInstanceOffsets)
 
 		SHADER_PARAMETER_RDG_BUFFER_UAV(RWStructuredBuffer<uint32>, InstanceIdsBufferOut)
-		SHADER_PARAMETER_RDG_BUFFER_UAV(RWStructuredBuffer<FVector4f>, InstanceIdsBufferOutMobile)		
+		SHADER_PARAMETER_RDG_BUFFER_UAV(RWByteAddressBuffer, InstanceIdsBufferOutMobile)		
 	END_SHADER_PARAMETER_STRUCT()
 };
 IMPLEMENT_GLOBAL_SHADER(FCompactVisibleInstancesCs, "/Engine/Private/InstanceCulling/CompactVisibleInstances.usf", "CompactVisibleInstances", SF_Compute);
@@ -394,7 +394,7 @@ public:
 		SHADER_PARAMETER_RDG_BUFFER_SRV(Buffer<uint>, InstanceIdOffsetBuffer)
 
 		SHADER_PARAMETER_RDG_BUFFER_UAV(RWStructuredBuffer<uint>, InstanceIdsBufferOut)
-		SHADER_PARAMETER_RDG_BUFFER_UAV(RWStructuredBuffer<float4>, InstanceIdsBufferOutMobile)
+		SHADER_PARAMETER_RDG_BUFFER_UAV(RWByteAddressBuffer, InstanceIdsBufferOutMobile)
 		SHADER_PARAMETER_RDG_BUFFER_UAV(RWBuffer<uint>, DrawIndirectArgsBufferOut)
 
 		SHADER_PARAMETER_RDG_BUFFER_SRV(StructuredBuffer<FInstanceCullingContext::FCompactionData>, DrawCommandCompactionData)
@@ -447,16 +447,34 @@ public:
 	void ProcessBatched(TStaticArray<FBuildInstanceIdBufferAndCommandsFromPrimitiveIdsCs::FParameters*, static_cast<uint32>(EBatchProcessingMode::Num)> PassParameters);
 };
 
-static FRDGBufferDesc CreateInstanceIdsBufferDesc(ERHIFeatureLevel::Type FeatureLevel, uint32 BufferSize)
+static uint32 GetInstanceIdsNumElements(ERHIFeatureLevel::Type FeatureLevel, uint32 NumInstances)
 {
-	FRDGBufferDesc BufferDesc = FRDGBufferDesc::CreateStructuredDesc(FInstanceCullingContext::GetInstanceIdBufferStride(FeatureLevel), BufferSize);
+	if (FeatureLevel == ERHIFeatureLevel::ES3_1)
+	{
+		// Mobile uses ByteAddressBuffer which is 4 bytes per element
+		const uint32 ByteBufferStride = FInstanceCullingContext::GetInstanceIdBufferStride(ERHIFeatureLevel::ES3_1) / 4u;
+		return (ByteBufferStride * NumInstances);
+	}
+	else
+	{
+		// Desktop uses StructuredBuffer<uint> NumElements==NumInstances
+		return NumInstances;
+	}
+}
+
+static FRDGBufferDesc CreateInstanceIdsBufferDesc(ERHIFeatureLevel::Type FeatureLevel, uint32 NumInstances)
+{
+	const uint32 BufferNumElements = GetInstanceIdsNumElements(FeatureLevel, NumInstances);
 	if (FeatureLevel == ERHIFeatureLevel::ES3_1)
 	{
 		// Mobile writes to this buffer from compute and then uses as a vertex input
-		// D3D does not allow a storage buffer with vertex buffer usage
-		BufferDesc.Usage |= (IsD3DPlatform(GMaxRHIShaderPlatform) ? BUF_None : BUF_VertexBuffer);
+		// We can't expose this as typed buffer to compute, because Android has a 64K limit for texel buffers which is not enough 
+		// Use ByteAddress buffer here since this is the only way D3D allows it to mix with vertex buffer usage, and it translates to a storage buffer on GL and VK
+		FRDGBufferDesc BufferDesc = FRDGBufferDesc::CreateByteAddressDesc(4u * BufferNumElements);
+		BufferDesc.Usage |= BUF_VertexBuffer;
+		return BufferDesc;
 	}
-	return BufferDesc;
+	return FRDGBufferDesc::CreateStructuredDesc(sizeof(uint32), BufferNumElements);
 }
 
 void FInstanceCullingContext::BuildRenderingCommands(
@@ -867,7 +885,11 @@ FInstanceCullingDeferredContext *FInstanceCullingContext::CreateDeferredContext(
 	FRDGBufferRef InstanceIdOffsetBuffer = GraphBuilder.CreateBuffer(FRDGBufferDesc::CreateBufferDesc(sizeof(uint32), 1), TEXT("InstanceCulling.InstanceIdOffsetBuffer"), INST_CULL_CALLBACK(DeferredContext->InstanceIdOffsets.Num()));
 	GraphBuilder.QueueBufferUpload(InstanceIdOffsetBuffer, INST_CULL_CALLBACK(DeferredContext->InstanceIdOffsets.GetData()), INST_CULL_CALLBACK(DeferredContext->InstanceIdOffsets.GetTypeSize() * DeferredContext->InstanceIdOffsets.Num()));
 
-	FRDGBufferRef InstanceIdsBuffer = GraphBuilder.CreateBuffer(CreateInstanceIdsBufferDesc(FeatureLevel, 1), TEXT("InstanceCulling.InstanceIdsBuffer"), INST_CULL_CALLBACK(DeferredContext->InstanceIdBufferSize));
+	FRDGBufferRef InstanceIdsBuffer = GraphBuilder.CreateBuffer(
+			CreateInstanceIdsBufferDesc(FeatureLevel, 1), 
+			TEXT("InstanceCulling.InstanceIdsBuffer"), 
+			INST_CULL_CALLBACK(GetInstanceIdsNumElements(DeferredContext->FeatureLevel, DeferredContext->InstanceIdBufferSize))
+	);
 	FRDGBufferUAVRef InstanceIdsBufferUAV = GraphBuilder.CreateUAV(InstanceIdsBuffer, ERDGUnorderedAccessViewFlags::SkipBarrier);
 	if (FeatureLevel == ERHIFeatureLevel::ES3_1)
 	{
