Fix issue where the shadowmask texture was incorrectly marked as used due to the preview shadows indicator, when the preview indicator was not actually rendered because of a missing material (e.g. in cooked builds).

diff --git a/Engine/Source/Runtime/Renderer/Private/LightRendering.cpp b/Engine/Source/Runtime/Renderer/Private/LightRendering.cpp
--- a/Engine/Source/Runtime/Renderer/Private/LightRendering.cpp
+++ b/Engine/Source/Runtime/Renderer/Private/LightRendering.cpp
@@ -1924,8 +1924,7 @@ void FDeferredShadingSceneRenderer::RenderLights(
 
 					if (bDrawPreviewIndicator)
 					{
-						RenderPreviewShadowsIndicator(GraphBuilder, SceneTextures, &LightSceneInfo, ScreenShadowMaskTexture, bUsedShadowMaskTexture, false);
-						bUsedShadowMaskTexture = true;
+						bUsedShadowMaskTexture |= RenderPreviewShadowsIndicator(GraphBuilder, SceneTextures, &LightSceneInfo, ScreenShadowMaskTexture, bUsedShadowMaskTexture, false);
 					}
 
 					if (!bDrawShadows)
