Fixed preshadows not receiving the call to SetupMeshDrawCommandsForProjectionStenciling which left them with an empty array of projection stenciling passes.

The empty array also made them render their projection for both views in stereo even when ISR was active, hitting an ensure in FProjectedShadowInfo::RenderProjection.

diff --git a/Engine/Source/Runtime/Renderer/Private/ShadowSetup.cpp b/Engine/Source/Runtime/Renderer/Private/ShadowSetup.cpp
--- a/Engine/Source/Runtime/Renderer/Private/ShadowSetup.cpp
+++ b/Engine/Source/Runtime/Renderer/Private/ShadowSetup.cpp
@@ -5299,6 +5299,7 @@ void FSceneRenderer::AllocateShadowDepthTargets(FRHICommandListImmediate& RHICmd
 		{
 			FProjectedShadowInfo* ProjectedShadowInfo = ShadowArrays.CachedPreShadows[ShadowIndex];
 			ProjectedShadowInfo->RenderTargets.DepthTarget = Scene->PreShadowCacheDepthZ.GetReference();
+			SortedShadowsForShadowDepthPass.PreshadowCache.Shadows.Add(ProjectedShadowInfo);
 		}
 	}
 
