Removing unused 'ViewParent' feature and other dead code.

diff --git a/Engine/Source/Editor/UnrealEd/Classes/Editor/EditorEngine.h b/Engine/Source/Editor/UnrealEd/Classes/Editor/EditorEngine.h
--- a/Engine/Source/Editor/UnrealEd/Classes/Editor/EditorEngine.h
+++ b/Engine/Source/Editor/UnrealEd/Classes/Editor/EditorEngine.h
@@ -1203,13 +1203,8 @@ public:
 	 */
 	void RedrawAllViewports(bool bInvalidateHitProxies=true);
 
-	/**
-	 * Invalidates all viewports parented to the specified view.
-	 *
-	 * @param	InParentView				The parent view whose child views should be invalidated.
-	 * @param	bInvalidateHitProxies		[opt] If true (the default), invalidates cached hit proxies too.
-	 */
-	void InvalidateChildViewports(FSceneViewStateInterface* InParentView, bool bInvalidateHitProxies=true);
+	UE_DEPRECATED(5.3, "InvalidateChildViewports is deprecated")
+	void InvalidateChildViewports(FSceneViewStateInterface* InParentView, bool bInvalidateHitProxies=true) {}
 
 	/**
 	 * Looks for an appropriate actor factory for the specified UClass.
diff --git a/Engine/Source/Editor/UnrealEd/Private/EditorEngine.cpp b/Engine/Source/Editor/UnrealEd/Private/EditorEngine.cpp
--- a/Engine/Source/Editor/UnrealEd/Private/EditorEngine.cpp
+++ b/Engine/Source/Editor/UnrealEd/Private/EditorEngine.cpp
@@ -2194,9 +2194,7 @@ void UEditorEngine::Tick( float DeltaSeconds, bool bIdleMode )
 					// Only update ortho viewports if that mode is turned on, the viewport client we are about to update is orthographic and the current editing viewport is orthographic and tracking mouse movement.
 					bUpdateLinkedOrthoViewports = GetDefault<ULevelEditorViewportSettings>()->bUseLinkedOrthographicViewports && ViewportClient->IsOrtho() && GCurrentLevelEditingViewportClient && GCurrentLevelEditingViewportClient->IsOrtho() && GCurrentLevelEditingViewportClient->IsTracking();
 
-					const bool bIsViewParent = ViewportClient->ViewState.GetReference()->IsViewParent();
-					if ((bRenderingChildren && !bIsViewParent) ||
-						(!bRenderingChildren && bIsViewParent) || bUpdateLinkedOrthoViewports)
+					if (bRenderingChildren || bUpdateLinkedOrthoViewports)
 					{
 						//if we haven't drawn a non-realtime viewport OR not one of the main viewports
 						bool bAllowNonRealtimeViewports = (!bEditorFrameNonRealtimeViewportDrawn) || !(ViewportClient->IsLevelEditorClient());
diff --git a/Engine/Source/Editor/UnrealEd/Private/EditorServer.cpp b/Engine/Source/Editor/UnrealEd/Private/EditorServer.cpp
--- a/Engine/Source/Editor/UnrealEd/Private/EditorServer.cpp
+++ b/Engine/Source/Editor/UnrealEd/Private/EditorServer.cpp
@@ -301,29 +301,6 @@ void UEditorEngine::RedrawAllViewports(bool bInvalidateHitProxies)
 	}
 }
 
-
-void UEditorEngine::InvalidateChildViewports(FSceneViewStateInterface* InParentView, bool bInvalidateHitProxies)
-{
-	if ( InParentView )
-	{
-		// Iterate over viewports and redraw those that have the specified view as a parent.
-		for (FEditorViewportClient* ViewportClient : AllViewportClients)
-		{
-			if ( ViewportClient && ViewportClient->ViewState.GetReference() )
-			{
-				if ( ViewportClient->ViewState.GetReference()->HasViewParent() &&
-					ViewportClient->ViewState.GetReference()->GetViewParent() == InParentView &&
-					!ViewportClient->ViewState.GetReference()->IsViewParent() )
-				{
-					constexpr bool bForceChildViewportRedraw = false;
-					ViewportClient->Invalidate(bForceChildViewportRedraw, bInvalidateHitProxies);
-				}
-			}
-		}
-	}
-}
-
-
 bool UEditorEngine::SafeExec( UWorld* InWorld, const TCHAR* InStr, FOutputDevice& Ar )
 {
 	const TCHAR* Str = InStr;
diff --git a/Engine/Source/Editor/UnrealEd/Private/EditorViewportClient.cpp b/Engine/Source/Editor/UnrealEd/Private/EditorViewportClient.cpp
--- a/Engine/Source/Editor/UnrealEd/Private/EditorViewportClient.cpp
+++ b/Engine/Source/Editor/UnrealEd/Private/EditorViewportClient.cpp
@@ -5583,13 +5583,6 @@ void FEditorViewportClient::Invalidate(bool bInvalidateChildViews, bool bInvalid
 			// Invalidate only display pixels.
 			Viewport->InvalidateDisplay();
 		}
-
-		// If this viewport is a view parent . . .
-		if ( bInvalidateChildViews &&
-			ViewState.GetReference()->IsViewParent() )
-		{
-			GEditor->InvalidateChildViewports( ViewState.GetReference(), bInvalidateHitProxies );
-		}
 	}
 }
 
diff --git a/Engine/Source/Editor/UnrealEd/Private/LevelEditorViewport.cpp b/Engine/Source/Editor/UnrealEd/Private/LevelEditorViewport.cpp
--- a/Engine/Source/Editor/UnrealEd/Private/LevelEditorViewport.cpp
+++ b/Engine/Source/Editor/UnrealEd/Private/LevelEditorViewport.cpp
@@ -2496,8 +2496,8 @@ void FLevelEditorViewportClient::Tick(float DeltaTime)
 	// Update the preview mesh for the preview mesh mode. 
 	GEditor->UpdatePreviewMesh();
 
-	// Copy perspective views to the global if this viewport is a view parent or has streaming volume previs enabled
-	if ( ViewState.GetReference()->IsViewParent() || (IsPerspective() && GetDefault<ULevelEditorViewportSettings>()->bLevelStreamingVolumePrevis && Viewport->GetSizeXY().X > 0) )
+	// Copy perspective views to the global if this viewport has streaming volume previs enabled
+	if ( (IsPerspective() && GetDefault<ULevelEditorViewportSettings>()->bLevelStreamingVolumePrevis && Viewport->GetSizeXY().X > 0) )
 	{
 		GPerspFrustumAngle=ViewFOV;
 		GPerspFrustumAspectRatio=AspectRatio;
@@ -4381,36 +4381,26 @@ void FLevelEditorViewportClient::Draw(const FSceneView* View,FPrimitiveDrawInter
 		DrawTextureStreamingBounds(View, PDI);
 	}
 
-	// Determine if a view frustum should be rendered in the viewport.
-	// The frustum should definitely be rendered if the viewport has a view parent.
-	bool bRenderViewFrustum = ViewState.GetReference()->HasViewParent();
-
-	// If the viewport doesn't have a view parent, a frustum still should be drawn anyway if the viewport is ortho and level streaming
-	// volume previs is enabled in some viewport
-	if ( !bRenderViewFrustum && IsOrtho() )
+	// A frustum should be drawn if the viewport is ortho and level streaming volume previs is enabled in some viewport
+	if ( IsOrtho() )
 	{
 		for (FLevelEditorViewportClient* CurViewportClient : GEditor->GetLevelViewportClients())
 		{
 			if ( CurViewportClient && IsPerspective() && GetDefault<ULevelEditorViewportSettings>()->bLevelStreamingVolumePrevis )
 			{
-				bRenderViewFrustum = true;
+				// Draw the view frustum of the level streaming volume previs viewport.
+				RenderViewFrustum(PDI, FLinearColor(1.0, 0.0, 1.0, 1.0),
+					GPerspFrustumAngle,
+					GPerspFrustumAspectRatio,
+					GPerspFrustumStartDist,
+					GPerspFrustumEndDist,
+					GPerspViewMatrix);
+
 				break;
 			}
 		}
 	}
 
-	// Draw the view frustum of the view parent or level streaming volume previs viewport, if necessary
-	if ( bRenderViewFrustum )
-	{
-		RenderViewFrustum( PDI, FLinearColor(1.0,0.0,1.0,1.0),
-			GPerspFrustumAngle,
-			GPerspFrustumAspectRatio,
-			GPerspFrustumStartDist,
-			GPerspFrustumEndDist,
-			GPerspViewMatrix);
-	}
-
-
 	if (IsPerspective())
 	{
 		DrawStaticLightingDebugInfo( View, PDI );
diff --git a/Engine/Source/Runtime/Engine/Public/SceneManagement.h b/Engine/Source/Runtime/Engine/Public/SceneManagement.h
--- a/Engine/Source/Runtime/Engine/Public/SceneManagement.h
+++ b/Engine/Source/Runtime/Engine/Public/SceneManagement.h
@@ -131,59 +131,12 @@ public:
 	FSceneViewStateInterface()
 		:	bValidEyeAdaptationTexture(0)
 		,	bValidEyeAdaptationBuffer(0)
-		,	ViewParent(NULL)
-		,	NumChildren( 0 )
 	{}
 	
 	/** Called in the game thread to destroy the view state. */
 	virtual void Destroy() = 0;
 
 public:
-	/** Sets the view state's scene parent. */
-	void SetViewParent(FSceneViewStateInterface* InViewParent)
-	{
-		if ( ViewParent )
-		{
-			// Assert that the existing parent does not have a parent.
-			check( !ViewParent->HasViewParent() );
-			// Decrement ref ctr of existing parent.
-			--ViewParent->NumChildren;
-		}
-
-		if ( InViewParent && InViewParent != this )
-		{
-			// Assert that the incoming parent does not have a parent.
-			check( !InViewParent->HasViewParent() );
-			ViewParent = InViewParent;
-			// Increment ref ctr of new parent.
-			InViewParent->NumChildren++;
-		}
-		else
-		{
-			ViewParent = NULL;
-		}
-	}
-	/** @return			The view state's scene parent, or NULL if none present. */
-	FSceneViewStateInterface* GetViewParent()
-	{
-		return ViewParent;
-	}
-	/** @return			The view state's scene parent, or NULL if none present. */
-	const FSceneViewStateInterface* GetViewParent() const
-	{
-		return ViewParent;
-	}
-	/** @return			true if the scene state has a parent, false otherwise. */
-	bool HasViewParent() const
-	{
-		return GetViewParent() != NULL;
-	}
-	/** @return			true if this scene state is a parent, false otherwise. */
-	bool IsViewParent() const
-	{
-		return NumChildren > 0;
-	}
-	
 	/** @return	the derived view state object */
 	virtual FSceneViewState* GetConcreteViewState () = 0;
 
@@ -303,7 +256,22 @@ PRAGMA_ENABLE_DEPRECATION_WARNINGS
 	virtual void AddLumenSceneData(FSceneInterface* InScene, float SurfaceCacheResolution = 1.0f) {}
 	virtual void RemoveLumenSceneData(FSceneInterface* InScene) {}
 	virtual bool HasLumenSceneData() const = 0;
+	
+	UE_DEPRECATED(5.3, "SetViewParent is deprecated")
+	void SetViewParent(FSceneViewStateInterface*) {}
+
+	UE_DEPRECATED(5.3, "GetViewParent is deprecated")
+	FSceneViewStateInterface* GetViewParent() { return nullptr; }
+
+	UE_DEPRECATED(5.3, "GetViewParent is deprecated")
+	const FSceneViewStateInterface* GetViewParent() const { return nullptr; }
 
+	UE_DEPRECATED(5.3, "HasViewParent is deprecated")
+	bool HasViewParent() const { return false; }
+
+	UE_DEPRECATED(5.3, "IsViewParent is deprecated")
+	bool IsViewParent() const { return false; }
+	
 protected:
 	// Don't allow direct deletion of the view state, Destroy should be called instead.
 	virtual ~FSceneViewStateInterface() {}
@@ -313,11 +281,6 @@ protected:
 	uint8 bValidEyeAdaptationBuffer : 1;
 
 private:
-	/** This scene state's view parent; NULL if no parent present. */
-	FSceneViewStateInterface*	ViewParent;
-	/** Reference counts the number of children parented to this state. */
-	int32							NumChildren;
-
 	virtual FVirtualShadowMapArrayCacheManager* GetVirtualShadowMapCache(const FScene* InScene) const { return nullptr; }
 	friend class FScene;
 };
diff --git a/Engine/Source/Runtime/Renderer/Private/DeferredShadingRenderer.cpp b/Engine/Source/Runtime/Renderer/Private/DeferredShadingRenderer.cpp
--- a/Engine/Source/Runtime/Renderer/Private/DeferredShadingRenderer.cpp
+++ b/Engine/Source/Runtime/Renderer/Private/DeferredShadingRenderer.cpp
@@ -2930,15 +2930,13 @@ void FDeferredShadingSceneRenderer::Render(FRDGBuilder& GraphBuilder)
 	check(Views.Num());
 	#if (UE_BUILD_SHIPPING || UE_BUILD_TEST)
 		const bool bIsViewFrozen = false;
-		const bool bHasViewParent = false;
 	#else
 		const bool bIsViewFrozen = Views[0].State && ((FSceneViewState*)Views[0].State)->bIsFrozen;
-		const bool bHasViewParent = Views[0].State && ((FSceneViewState*)Views[0].State)->HasViewParent();
 	#endif
 
 	
 	const bool bIsOcclusionTesting = DoOcclusionQueries() && !ViewFamily.EngineShowFlags.DisableOcclusionQueries
-		&& (!ViewFamily.EngineShowFlags.Wireframe || bIsViewFrozen || bHasViewParent);
+		&& (!ViewFamily.EngineShowFlags.Wireframe || bIsViewFrozen);
 	const bool bNeedsPrePass = ShouldRenderPrePass();
 
 	// Sanity check - Note: Nanite forces a Z prepass in ShouldForceFullDepthPass()
diff --git a/Engine/Source/Runtime/Renderer/Private/MobileShadingRenderer.cpp b/Engine/Source/Runtime/Renderer/Private/MobileShadingRenderer.cpp
--- a/Engine/Source/Runtime/Renderer/Private/MobileShadingRenderer.cpp
+++ b/Engine/Source/Runtime/Renderer/Private/MobileShadingRenderer.cpp
@@ -1774,7 +1774,7 @@ int32 FMobileSceneRenderer::ComputeNumOcclusionQueriesToBatch() const
 		const FViewInfo& View = Views[ViewIndex];
 		const FSceneViewState* ViewState = (FSceneViewState*)View.State;
 #if !(UE_BUILD_SHIPPING || UE_BUILD_TEST)
-		if (!ViewState || (!ViewState->HasViewParent() && !ViewState->bIsFrozen))
+		if (!ViewState || (!ViewState->bIsFrozen))
 #endif
 		{
 			NumQueriesForBatch += View.IndividualOcclusionQueries.GetNumBatchOcclusionQueries();
diff --git a/Engine/Source/Runtime/Renderer/Private/PostProcess/PostProcessCompositeEditorPrimitives.cpp b/Engine/Source/Runtime/Renderer/Private/PostProcess/PostProcessCompositeEditorPrimitives.cpp
--- a/Engine/Source/Runtime/Renderer/Private/PostProcess/PostProcessCompositeEditorPrimitives.cpp
+++ b/Engine/Source/Runtime/Renderer/Private/PostProcess/PostProcessCompositeEditorPrimitives.cpp
@@ -284,9 +284,6 @@ const FViewInfo* CreateEditorPrimitiveView(const FViewInfo& ParentView, FIntRect
 	// Kills material texture mipbias because after TAA.
 	EditorView->MaterialTextureMipBias = 0.0f;
 
-	// Disable decals so that we don't do a SetDepthStencilState() in TMobileBasePassDrawingPolicy::SetupPipelineState()
-	EditorView->bSceneHasDecals = false;
-
 	if (IsTemporalAccumulationBasedMethod(EditorView->AntiAliasingMethod))
 	{
 		EditorView->ViewMatrices.HackRemoveTemporalAAProjectionJitter();
diff --git a/Engine/Source/Runtime/Renderer/Private/SceneOcclusion.cpp b/Engine/Source/Runtime/Renderer/Private/SceneOcclusion.cpp
--- a/Engine/Source/Runtime/Renderer/Private/SceneOcclusion.cpp
+++ b/Engine/Source/Runtime/Renderer/Private/SceneOcclusion.cpp
@@ -325,7 +325,6 @@ SIZE_T FSceneViewState::GetSizeBytes() const
 
 	return sizeof(*this) 
 		+ ShadowOcclusionQuerySize
-		+ ParentPrimitives.GetAllocatedSize() 
 		+ PrimitiveFadingStates.GetAllocatedSize()
 		+ PrimitiveOcclusionHistorySet.GetAllocatedSize();
 }
@@ -1209,7 +1208,7 @@ static FViewOcclusionQueriesPerView AllocateOcclusionTests(const FScene* Scene,
 
 			// Don't do primitive occlusion if we have a view parent or are frozen - only applicable to Debug & Development.
 #if !(UE_BUILD_SHIPPING || UE_BUILD_TEST)
-			ViewQuery.bFlushQueries &= (!ViewState->HasViewParent() && !ViewState->bIsFrozen);
+			ViewQuery.bFlushQueries &= (!ViewState->bIsFrozen);
 #endif
 
 			bBatchedQueries |= (View.IndividualOcclusionQueries.HasBatches() || View.GroupedOcclusionQueries.HasBatches() || ViewQuery.bFlushQueries);
@@ -1437,7 +1436,7 @@ void FDeferredShadingSceneRenderer::RenderOcclusion(
 					FViewInfo& View = Views[ViewIndex];
 					FSceneViewState* ViewState = (FSceneViewState*)View.State;
 #if !(UE_BUILD_SHIPPING || UE_BUILD_TEST)
-					if (!ViewState->HasViewParent() && !ViewState->bIsFrozen)
+					if (!ViewState->bIsFrozen)
 #endif
 					{
 						NumQueriesForBatch += View.IndividualOcclusionQueries.GetNumBatchOcclusionQueries();
diff --git a/Engine/Source/Runtime/Renderer/Private/ScenePrivate.h b/Engine/Source/Runtime/Renderer/Private/ScenePrivate.h
--- a/Engine/Source/Runtime/Renderer/Private/ScenePrivate.h
+++ b/Engine/Source/Runtime/Renderer/Private/ScenePrivate.h
@@ -757,9 +757,6 @@ public:
 	// A counter incremented once each time this view is rendered.
 	uint32 OcclusionFrameCounter;
 
-	/** Used by states that have IsViewParent() == true to store primitives for child states. */
-	TSet<FPrimitiveComponentId> ParentPrimitives;
-
 	/** For this view, the set of primitives that are currently fading, either in or out. */
 	FPrimitiveFadingStateMap PrimitiveFadingStates;
 
diff --git a/Engine/Source/Runtime/Renderer/Private/SceneRendering.cpp b/Engine/Source/Runtime/Renderer/Private/SceneRendering.cpp
--- a/Engine/Source/Runtime/Renderer/Private/SceneRendering.cpp
+++ b/Engine/Source/Runtime/Renderer/Private/SceneRendering.cpp
@@ -862,7 +862,6 @@ void FViewInfo::Init()
 
 	NumVisibleStaticMeshElements = 0;
 	PrecomputedVisibilityData = 0;
-	bSceneHasDecals = 0;
 
 	bIsViewInfo = true;
 	
@@ -3395,7 +3394,7 @@ void FSceneRenderer::RenderFinish(FRDGBuilder& GraphBuilder, FRDGTextureRef View
 
 				// display a message saying we're frozen
 				FSceneViewState* ViewState = (FSceneViewState*)View.State;
-				bool bViewParentOrFrozen = ViewState && (ViewState->HasViewParent() || ViewState->bIsFrozen);
+				bool bIsFrozen = ViewState && (ViewState->bIsFrozen);
 				bool bLocked = View.bIsLocked;
 				const bool bStereoView = IStereoRendering::IsStereoEyeView(View);
 
@@ -3415,7 +3414,7 @@ void FSceneRenderer::RenderFinish(FRDGBuilder& GraphBuilder, FRDGTextureRef View
 					}
 				}
 #endif
-				if ((GAreScreenMessagesEnabled && !GEngine->bSuppressMapWarnings) && (bViewParentOrFrozen || bLocked || bStereoView || bShowAmbientCubemapMipGenSettingsWarning || bAnyWarning))
+				if ((GAreScreenMessagesEnabled && !GEngine->bSuppressMapWarnings) && (bIsFrozen || bLocked || bStereoView || bShowAmbientCubemapMipGenSettingsWarning || bAnyWarning))
 				{
 					RDG_EVENT_SCOPE_CONDITIONAL(GraphBuilder, Views.Num() > 1, "View%d", ViewIndex);
 
@@ -3427,7 +3426,7 @@ void FSceneRenderer::RenderFinish(FRDGBuilder& GraphBuilder, FRDGTextureRef View
 					AddDrawCanvasPass(GraphBuilder, {}, View, Output,
 						[this, &ReadOnlyCVARCache, ViewState, GPUSkinCacheExtraRequiredMemory,
 						bLocked, bShowPrecomputedVisibilityWarning, bShowDemotedLocalMemoryWarning, bShowGlobalClipPlaneWarning, bShowDFAODisabledWarning, bShowDFDisabledWarning,
-						bViewParentOrFrozen, bShowSkylightWarning, bShowPointLightWarning, bShowShadowedLightOverflowWarning,
+						bIsFrozen, bShowSkylightWarning, bShowPointLightWarning, bShowShadowedLightOverflowWarning,
 						bShowMobileLowQualityLightmapWarning, bShowMobileMovableDirectionalLightWarning, bShowMobileDynamicCSMWarning, bMobileMissingSkyMaterial, 
 						bShowSkinCacheOOM, bSingleLayerWaterWarning, bShowNoSkyAtmosphereComponentWarning, bFxDebugDraw, FXInterface, bShowLocalExposureDisabledWarning,
 						bLumenEnabledButHasNoDataForTracing, bLumenEnabledButDisabledForTheProject, bNaniteEnabledButNoAtomics, bNaniteEnabledButDisabledInProject, bRealTimeSkyCaptureButNothingToCapture, bShowWaitingSkylight, bShowAmbientCubemapMipGenSettingsWarning,
@@ -3437,13 +3436,9 @@ void FSceneRenderer::RenderFinish(FRDGBuilder& GraphBuilder, FRDGTextureRef View
 						// so it can get the screen size
 						FScreenMessageWriter Writer(Canvas, 130);
 
-						if (bViewParentOrFrozen)
+						if (bIsFrozen)
 						{
-							const FText StateText =
-								ViewState->bIsFrozen ?
-								NSLOCTEXT("SceneRendering", "RenderingFrozen", "Rendering frozen...")
-								:
-								NSLOCTEXT("SceneRendering", "OcclusionChild", "Occlusion Child");
+							static const FText StateText = NSLOCTEXT("SceneRendering", "RenderingFrozen", "Rendering frozen...");
 							Writer.DrawLine(StateText, 10, FLinearColor(0.8, 1.0, 0.2, 1.0));
 						}
 						if (bShowPrecomputedVisibilityWarning)
diff --git a/Engine/Source/Runtime/Renderer/Private/SceneRendering.h b/Engine/Source/Runtime/Renderer/Private/SceneRendering.h
--- a/Engine/Source/Runtime/Renderer/Private/SceneRendering.h
+++ b/Engine/Source/Runtime/Renderer/Private/SceneRendering.h
@@ -1491,11 +1491,6 @@ public:
 	/** Whether fog should only be computed on rendered opaque pixels or not. */
 	uint32 bFogOnlyOnRenderedOpaque : 1;
 
-	/** 
-	 * true if the scene has at least one decal. Used to disable stencil operations in the mobile base pass when the scene has no decals.
-	 * TODO: Right now decal visibility is computed right before rendering them. Ideally it should be done in InitViews and this flag should be replaced with list of visible decals  
-	 */
-	uint32 bSceneHasDecals : 1;
 	/**
 	 * true if the scene has at least one mesh with a material tagged as sky. 
 	 * This is used to skip the sky rendering part during the SkyAtmosphere pass on non mobile platforms.
diff --git a/Engine/Source/Runtime/Renderer/Private/SceneVisibility.cpp b/Engine/Source/Runtime/Renderer/Private/SceneVisibility.cpp
--- a/Engine/Source/Runtime/Renderer/Private/SceneVisibility.cpp
+++ b/Engine/Source/Runtime/Renderer/Private/SceneVisibility.cpp
@@ -2054,7 +2054,7 @@ static int32 OcclusionCull(FRHICommandListImmediate& RHICmdList, const FScene* S
 	{
 		bool bSubmitQueries = !View.bDisableQuerySubmissions;
 #if !(UE_BUILD_SHIPPING || UE_BUILD_TEST)
-		bSubmitQueries = bSubmitQueries && !ViewState->HasViewParent() && !ViewState->bIsFrozen;
+		bSubmitQueries = bSubmitQueries && !ViewState->bIsFrozen;
 #endif
 
 		if (ViewState->OcclusionFeedback.IsInitialized())
@@ -4547,15 +4547,6 @@ void FSceneRenderer::ComputeViewVisibility(
 			View.PrimitiveViewRelevanceMap.Reset(Scene->Primitives.Num());
 			View.PrimitiveViewRelevanceMap.AddZeroed(Scene->Primitives.Num());
 
-			// If this is the visibility-parent of other views, reset its ParentPrimitives list.
-			const bool bIsParent = ViewState && ViewState->IsViewParent();
-			if (bIsParent)
-			{
-				// PVS-Studio does not understand the validation of ViewState above, so we're disabling
-				// its warning that ViewState may be null:
-				ViewState->ParentPrimitives.Reset(); //-V595
-			}
-
 			if (ViewState)
 			{
 				SCOPE_CYCLE_COUNTER(STAT_DecompressPrecomputedOcclusion);
@@ -4577,22 +4568,6 @@ void FSceneRenderer::ComputeViewVisibility(
 #if !(UE_BUILD_SHIPPING || UE_BUILD_TEST)
 			if (ViewState)
 			{
-#if WITH_EDITOR
-				// For visibility child views, check if the primitive was visible in the parent view.
-				const FSceneViewState* const ViewParent = (FSceneViewState*)ViewState->GetViewParent();
-				if (ViewParent)
-				{
-					bNeedsFrustumCulling = false;
-					for (int32 Index = 0; Index < View.PrimitiveVisibilityMap.Num(); ++Index)
-					{
-						if (ViewParent->ParentPrimitives.Contains(Scene->PrimitiveComponentIds[Index]) || 
-							IsAlwaysVisible(Scene, Index))
-						{
-							View.PrimitiveVisibilityMap[Index] = true;
-						}
-					}
-				}
-#endif
 				// For views with frozen visibility, check if the primitive is in the frozen visibility set.
 				if (ViewState->bIsFrozen)
 				{
@@ -4683,17 +4658,6 @@ void FSceneRenderer::ComputeViewVisibility(
 				ComputeAndMarkRelevanceForViewParallel(RHICmdList, Scene, View, ViewCommands, ViewBit, HasDynamicMeshElementsMasks, HasDynamicEditorMeshElementsMasks);
 			}
 
-#if !(UE_BUILD_SHIPPING || UE_BUILD_TEST)
-			// Store the primitive for parent occlusion rendering.
-			if (FPlatformProperties::SupportsWindowedMode() && ViewState && ViewState->IsViewParent())
-			{
-				for (FSceneDualSetBitIterator BitIt(View.PrimitiveVisibilityMap, View.PrimitiveDefinitelyUnoccludedMap); BitIt; ++BitIt)
-				{
-					ViewState->ParentPrimitives.Add(Scene->PrimitiveComponentIds[BitIt.GetIndex()]);
-				}
-			}
-#endif
-
 #if !(UE_BUILD_SHIPPING || UE_BUILD_TEST)
 			// if we are freezing the scene, then remember the primitives that are rendered.
 			if (ViewState && ViewState->bIsFreezing)
@@ -4705,10 +4669,6 @@ void FSceneRenderer::ComputeViewVisibility(
 			}
 #endif
 
-			// TODO: right now decals visibility computed right before rendering them, ideally it should be done in InitViews and this flag should be replaced with list of visible decals  
-			// Currently used to disable stencil operations in forward base pass when scene has no any decals
-			View.bSceneHasDecals = (Scene->Decals.Num() > 0) || (GForceSceneHasDecals != 0);
-
 			if (bIsSinglePassStereo && IStereoRendering::IsASecondaryView(View) && Views.IsValidIndex(View.PrimaryViewIndex))
 			{
 				// Ensure primitives from the secondary view are visible in the primary view
