Cleanup string allocations when clearing GFP mids

diff --git a/Engine/Source/Runtime/Renderer/Private/ScenePrivate.h b/Engine/Source/Runtime/Renderer/Private/ScenePrivate.h
--- a/Engine/Source/Runtime/Renderer/Private/ScenePrivate.h
+++ b/Engine/Source/Runtime/Renderer/Private/ScenePrivate.h
@@ -1430,9 +1430,15 @@ PRAGMA_ENABLE_DEPRECATION_WARNINGS
 			return;
 		}
 
-		const int32 RemoveNum = Algo::RemoveIf(MIDPool, [MidParentRootPath](UMaterialInstanceDynamic* MID) -> bool
+		FNameBuilder PackagePath;
+		const int32 RemoveNum = Algo::RemoveIf(MIDPool, [&PackagePath, MidParentRootPath](UMaterialInstanceDynamic* MID) -> bool
 		{
-			return MID->Parent && MID->Parent->GetPackage()->GetName().StartsWith(MidParentRootPath);
+			if (MID->Parent)
+			{
+				MID->Parent->GetPackage()->GetFName().ToString(PackagePath);
+				return FStringView(PackagePath).StartsWith(MidParentRootPath);
+			}
+			return false;
 		});
 
 		MIDPool.SetNum(RemoveNum);
