Defer syncing of mesh pass setup tasks until RDG execute (when deferred instance culling commands are processed), increases available overlap for RT.

diff --git a/Engine/Source/Runtime/Renderer/Private/InstanceCulling/InstanceCullingContext.cpp b/Engine/Source/Runtime/Renderer/Private/InstanceCulling/InstanceCullingContext.cpp
--- a/Engine/Source/Runtime/Renderer/Private/InstanceCulling/InstanceCullingContext.cpp
+++ b/Engine/Source/Runtime/Renderer/Private/InstanceCulling/InstanceCullingContext.cpp
@@ -491,7 +491,8 @@ void FInstanceCullingContext::BuildRenderingCommands(
 	int32 DynamicInstanceIdOffset,
 	int32 DynamicInstanceIdNum,
 	FInstanceCullingResult& Results,
-	FInstanceCullingDrawParams* InstanceCullingDrawParams) const
+	FInstanceCullingDrawParams* InstanceCullingDrawParams,
+	TFunction<void ()>&& SyncPrerequisitesFunc) const
 {
 	Results = FInstanceCullingResult();
 
@@ -509,8 +510,6 @@ void FInstanceCullingContext::BuildRenderingCommands(
 		return;
 	}
 
-	const bool bOcclusionCullInstances = PrevHZB.IsValid() && IsOcclusionCullingEnabled();
-	const uint32 InstanceIdBufferSize = TotalInstances * ViewIds.Num();
 	if (InstanceCullingDrawParams && InstanceCullingManager && InstanceCullingManager->IsDeferredCullingActive() && (InstanceCullingMode == EInstanceCullingMode::Normal))
 	{
 		FInstanceCullingDeferredContext *DeferredContext = InstanceCullingManager->DeferredContext;
@@ -521,10 +520,15 @@ void FInstanceCullingContext::BuildRenderingCommands(
 			Results.DrawIndirectArgsBuffer = DeferredContext->DrawIndirectArgsBuffer;
 			Results.InstanceDataBuffer = DeferredContext->InstanceDataBuffer;
 			Results.UniformBuffer = DeferredContext->UniformBuffer;
-			DeferredContext->AddBatch(GraphBuilder, this, DynamicInstanceIdOffset, DynamicInstanceIdNum, InstanceCullingDrawParams);
+			DeferredContext->AddBatch(GraphBuilder, this, DynamicInstanceIdOffset, DynamicInstanceIdNum, InstanceCullingDrawParams, MoveTemp(SyncPrerequisitesFunc));
 		}
 		return;
 	}
+	// If not running with a deferred context, ensure sync happens at once
+	if (SyncPrerequisitesFunc)
+	{
+		SyncPrerequisitesFunc();
+	}
 
 	ensure(InstanceCullingMode == EInstanceCullingMode::Normal || ViewIds.Num() == 2);
 
@@ -577,6 +581,7 @@ void FInstanceCullingContext::BuildRenderingCommands(
 
 	FRDGBufferRef ViewIdsBuffer = CreateStructuredBuffer(GraphBuilder, TEXT("InstanceCulling.ViewIds"), ViewIds);
 
+	const uint32 InstanceIdBufferSize = TotalInstances * ViewIds.Num();
 	FRDGBufferRef InstanceIdsBuffer = GraphBuilder.CreateBuffer(CreateInstanceIdsBufferDesc(FeatureLevel, InstanceIdBufferSize), TEXT("InstanceCulling.InstanceIdsBuffer"));
 	FRDGBufferUAVRef InstanceIdsBufferUAV = GraphBuilder.CreateUAV(InstanceIdsBuffer, ERDGUnorderedAccessViewFlags::SkipBarrier);
 
@@ -644,6 +649,7 @@ void FInstanceCullingContext::BuildRenderingCommands(
 	PassParametersTmp.DrawIndirectArgsBufferOut = GraphBuilder.CreateUAV(DrawIndirectArgsRDG, PF_R32_UINT, ERDGUnorderedAccessViewFlags::SkipBarrier);
 	PassParametersTmp.InstanceIdOffsetBuffer = GraphBuilder.CreateSRV(InstanceIdOffsetBufferRDG, PF_R32_UINT);
 
+	const bool bOcclusionCullInstances = PrevHZB.IsValid() && IsOcclusionCullingEnabled();
 	if (bOcclusionCullInstances)
 	{
 		PassParametersTmp.HZBTexture = GraphBuilder.RegisterExternalTexture(PrevHZB);
diff --git a/Engine/Source/Runtime/Renderer/Private/InstanceCulling/InstanceCullingMergedContext.cpp b/Engine/Source/Runtime/Renderer/Private/InstanceCulling/InstanceCullingMergedContext.cpp
--- a/Engine/Source/Runtime/Renderer/Private/InstanceCulling/InstanceCullingMergedContext.cpp
+++ b/Engine/Source/Runtime/Renderer/Private/InstanceCulling/InstanceCullingMergedContext.cpp
@@ -6,6 +6,12 @@
 
 void FInstanceCullingMergedContext::MergeBatches()
 {
+	for (const FAsyncBatchItem& AsyncBatchItem : AsyncBatches)
+	{
+		AsyncBatchItem.SyncPrerequisitesFunc();
+		AddBatchItem(AsyncBatchItem.BatchItem);
+	}
+
 	for (uint32 Mode = 0U; Mode < uint32(EBatchProcessingMode::Num); ++Mode)
 	{
 		LoadBalancers[Mode].ReserveStorage(TotalBatches[Mode], TotalItems[Mode]);
@@ -111,10 +117,9 @@ void FInstanceCullingMergedContext::MergeBatches()
 	}
 }
 
-void FInstanceCullingMergedContext::AddBatch(FRDGBuilder& GraphBuilder, const FInstanceCullingContext* Context, int32 DynamicInstanceIdOffset,	int32 DynamicInstanceIdNum, FInstanceCullingDrawParams* InstanceCullingDrawParams)
+void FInstanceCullingMergedContext::AddBatch(FRDGBuilder& GraphBuilder, const FInstanceCullingContext* Context, int32 DynamicInstanceIdOffset,	int32 DynamicInstanceIdNum, FInstanceCullingDrawParams* InstanceCullingDrawParams, TFunction<void()>&& SyncPrerequisitesFunc)
 {
 	checkfSlow(Batches.FindByPredicate([InstanceCullingDrawParams](const FBatchItem& Item) { return Item.Result == InstanceCullingDrawParams; }) == nullptr, TEXT("Output draw paramters registered twice."));
-	Batches.Add(FBatchItem{ Context, InstanceCullingDrawParams, DynamicInstanceIdOffset, DynamicInstanceIdNum });
 
 	const bool bOcclusionCullInstances = Context->PrevHZB.IsValid() && FInstanceCullingContext::IsOcclusionCullingEnabled();
 
@@ -131,6 +136,23 @@ void FInstanceCullingMergedContext::AddBatch(FRDGBuilder& GraphBuilder, const FI
 		}
 	}
 
+	if (SyncPrerequisitesFunc)
+	{
+		AsyncBatches.Add(FAsyncBatchItem{ FBatchItem{ Context, InstanceCullingDrawParams, DynamicInstanceIdOffset, DynamicInstanceIdNum }, MoveTemp(SyncPrerequisitesFunc) });
+
+	}
+	else
+	{
+		AddBatchItem(FBatchItem{ Context, InstanceCullingDrawParams, DynamicInstanceIdOffset, DynamicInstanceIdNum });
+	}
+
+}
+
+void FInstanceCullingMergedContext::AddBatchItem(const FBatchItem& BatchItem)
+{
+	Batches.Add(BatchItem);
+
+	const FInstanceCullingContext* Context = BatchItem.Context;
 	// Accumulate the totals so the deferred processing can pre-size the arrays
 	for (uint32 Mode = 0U; Mode < uint32(EBatchProcessingMode::Num); ++Mode)
 	{
diff --git a/Engine/Source/Runtime/Renderer/Private/InstanceCulling/InstanceCullingMergedContext.h b/Engine/Source/Runtime/Renderer/Private/InstanceCulling/InstanceCullingMergedContext.h
--- a/Engine/Source/Runtime/Renderer/Private/InstanceCulling/InstanceCullingMergedContext.h
+++ b/Engine/Source/Runtime/Renderer/Private/InstanceCulling/InstanceCullingMergedContext.h
@@ -19,6 +19,12 @@ public:
 		FInstanceCullingDrawParams* Result = nullptr;
 		int32 DynamicInstanceIdOffset = 0;
 		int32 DynamicInstanceIdNum = 0;
+		TFunction<void()> SyncPrerequisitesFunc;
+	};
+	struct FAsyncBatchItem
+	{
+		FBatchItem BatchItem;
+		TFunction<void()> SyncPrerequisitesFunc;
 	};
 
 	// Info about a batch of culling work produced by a context, when part of a batched job
@@ -39,6 +45,8 @@ public:
 	/** Batches of GPU instance culling input data. */
 	TArray<FBatchItem, SceneRenderingAllocator> Batches;
 
+	/** Async (and thus added as to the above as late as possible) Batches of GPU instance culling input data. */
+	TArray<FAsyncBatchItem, SceneRenderingAllocator> AsyncBatches;
 
 	/** 
 	 * Merged data, derived in MergeBatches(), follows.
@@ -77,5 +85,9 @@ public:
 	void MergeBatches();
 
 
-	void AddBatch(FRDGBuilder& GraphBuilder, const FInstanceCullingContext* Context, int32 DynamicInstanceIdOffset, int32 DynamicInstanceIdNum, FInstanceCullingDrawParams* InstanceCullingDrawParams = nullptr);
+	void AddBatch(FRDGBuilder& GraphBuilder, const FInstanceCullingContext* Context, int32 DynamicInstanceIdOffset, int32 DynamicInstanceIdNum, FInstanceCullingDrawParams* InstanceCullingDrawParams, TFunction<void()>&& SyncPrerequisitesFunc = TFunction<void()>());
+
+private:
+	void AddBatchItem(const FBatchItem& BatchItem);
+
 };
\ No newline at end of file
diff --git a/Engine/Source/Runtime/Renderer/Private/MeshDrawCommands.cpp b/Engine/Source/Runtime/Renderer/Private/MeshDrawCommands.cpp
--- a/Engine/Source/Runtime/Renderer/Private/MeshDrawCommands.cpp
+++ b/Engine/Source/Runtime/Renderer/Private/MeshDrawCommands.cpp
@@ -1421,21 +1421,23 @@ void FParallelMeshDrawCommandPass::BuildRenderingCommands(
 	if (TaskContext.InstanceCullingContext.IsEnabled())
 	{
 		check(!bHasInstanceCullingDrawParameters);
-		WaitForMeshPassSetupTask();
-
-#if DO_CHECK
-		for (const FVisibleMeshDrawCommand& VisibleMeshDrawCommand : TaskContext.MeshDrawCommands)
+		// 2. Run or queue finalize culling commands pass
+		TaskContext.InstanceCullingContext.BuildRenderingCommands(GraphBuilder, GPUScene, TaskContext.View->DynamicPrimitiveCollector.GetInstanceSceneDataOffset(), TaskContext.View->DynamicPrimitiveCollector.NumInstances(), TaskContext.InstanceCullingResult, &OutInstanceCullingDrawParams,
+		[this, &GPUScene]()
 		{
-			if (VisibleMeshDrawCommand.PrimitiveIdInfo.bIsDynamicPrimitive)
+			WaitForMeshPassSetupTask();
+
+	#if DO_CHECK
+			for (const FVisibleMeshDrawCommand& VisibleMeshDrawCommand : TaskContext.MeshDrawCommands)
 			{
-				uint32 PrimitiveIndex = VisibleMeshDrawCommand.PrimitiveIdInfo.DrawPrimitiveId & ~GPrimIDDynamicFlag;
-				TaskContext.View->DynamicPrimitiveCollector.CheckPrimitiveProcessed(PrimitiveIndex, GPUScene);
+				if (VisibleMeshDrawCommand.PrimitiveIdInfo.bIsDynamicPrimitive)
+				{
+					uint32 PrimitiveIndex = VisibleMeshDrawCommand.PrimitiveIdInfo.DrawPrimitiveId & ~GPrimIDDynamicFlag;
+					TaskContext.View->DynamicPrimitiveCollector.CheckPrimitiveProcessed(PrimitiveIndex, GPUScene);
+				}
 			}
-		}
-#endif
-
-		// 2. Run or queue finalize culling commands pass
-		TaskContext.InstanceCullingContext.BuildRenderingCommands(GraphBuilder, GPUScene, TaskContext.View->DynamicPrimitiveCollector.GetInstanceSceneDataOffset(), TaskContext.View->DynamicPrimitiveCollector.NumInstances(), TaskContext.InstanceCullingResult, &OutInstanceCullingDrawParams);
+	#endif
+		});
 		TaskContext.InstanceCullingResult.GetDrawParameters(OutInstanceCullingDrawParams);
 		bHasInstanceCullingDrawParameters = true;
 		check(!TaskContext.InstanceCullingContext.HasCullingCommands() || OutInstanceCullingDrawParams.DrawIndirectArgsBuffer && OutInstanceCullingDrawParams.InstanceIdOffsetBuffer);
diff --git a/Engine/Source/Runtime/Renderer/Public/InstanceCulling/InstanceCullingContext.h b/Engine/Source/Runtime/Renderer/Public/InstanceCulling/InstanceCullingContext.h
--- a/Engine/Source/Runtime/Renderer/Public/InstanceCulling/InstanceCullingContext.h
+++ b/Engine/Source/Runtime/Renderer/Public/InstanceCulling/InstanceCullingContext.h
@@ -139,7 +139,8 @@ public:
 		int32 DynamicInstanceIdOffset,
 		int32 DynamicInstanceIdNum,
 		FInstanceCullingResult& Results,
-		FInstanceCullingDrawParams* InstanceCullingDrawParams = nullptr) const;
+		FInstanceCullingDrawParams* InstanceCullingDrawParams = nullptr,
+		TFunction<void()>&& SyncPrerequisitesFunc = TFunction<void()>()) const;
 
 	inline bool HasCullingCommands() const { return TotalInstances > 0; 	}
 
