Fix for stack-use-after-free when iterating bit array with the wrong allocator.

diff --git a/Engine/Source/Runtime/Renderer/Private/RendererScene.cpp b/Engine/Source/Runtime/Renderer/Private/RendererScene.cpp
--- a/Engine/Source/Runtime/Renderer/Private/RendererScene.cpp
+++ b/Engine/Source/Runtime/Renderer/Private/RendererScene.cpp
@@ -6380,7 +6380,7 @@ void FScene::UpdateAllPrimitiveSceneInfos(FRDGBuilder& GraphBuilder, EUpdateAllP
 		SCOPED_NAMED_EVENT(UpdateUniformBuffers, FColor::Emerald);
 		TArray<FPrimitiveSceneProxy*, SceneRenderingAllocator> ProxiesToUpdate;
 
-		for (TConstSetBitIterator<SceneRenderingBitArrayAllocator> BitIt(PrimitivesNeedingUniformBufferUpdate); BitIt; ++BitIt)
+		for (TConstSetBitIterator<FDefaultBitArrayAllocator> BitIt(PrimitivesNeedingUniformBufferUpdate); BitIt; ++BitIt)
 		{
 			const int32 Index = BitIt.GetIndex();
 			FPrimitiveSceneInfo* Primitive = Primitives[Index];
