Fix a bug where the culling views SRV was not getting set for deferred instance culling.

diff --git a/Engine/Source/Runtime/Renderer/Private/InstanceCulling/InstanceCullingContext.cpp b/Engine/Source/Runtime/Renderer/Private/InstanceCulling/InstanceCullingContext.cpp
--- a/Engine/Source/Runtime/Renderer/Private/InstanceCulling/InstanceCullingContext.cpp
+++ b/Engine/Source/Runtime/Renderer/Private/InstanceCulling/InstanceCullingContext.cpp
@@ -942,8 +942,8 @@ FInstanceCullingDeferredContext *FInstanceCullingContext::CreateDeferredContext(
 	PassParametersTmp.InstanceIdsBufferOutMobile = InstanceIdsBufferUAV;
 
 	PassParametersTmp.DrawIndirectArgsBufferOut = GraphBuilder.CreateUAV(DeferredContext->DrawIndirectArgsBuffer, PF_R32_UINT, ERDGUnorderedAccessViewFlags::SkipBarrier);
-	PassParametersTmp.InstanceIdOffsetBuffer = GraphBuilder.CreateSRV(InstanceIdOffsetBuffer, PF_R32_UINT);
-	if (bCullInstances)
+	PassParametersTmp.InstanceIdOffsetBuffer = GraphBuilder.CreateSRV(InstanceIdOffsetBuffer, PF_R32_UINT);	
+	if (bCullInstances || bAllowWPODisable)
 	{
 		PassParametersTmp.InViews = GraphBuilder.CreateSRV(InstanceCullingManager->CullingIntermediate.CullingViews);
 		PassParametersTmp.NumCullingViews = InstanceCullingManager->CullingIntermediate.NumViews;
