Fix missing method declaration due to incorrect merge.

diff --git a/Engine/Source/Runtime/Renderer/Private/ScenePrivate.h b/Engine/Source/Runtime/Renderer/Private/ScenePrivate.h
--- a/Engine/Source/Runtime/Renderer/Private/ScenePrivate.h
+++ b/Engine/Source/Runtime/Renderer/Private/ScenePrivate.h
@@ -3090,6 +3090,7 @@ public:
 	virtual void UpdatePrimitiveDistanceFieldSceneData_GameThread(UPrimitiveComponent* Primitive) override;
 	virtual FPrimitiveSceneInfo* GetPrimitiveSceneInfo(int32 PrimitiveIndex) const final;
 	virtual FPrimitiveSceneInfo* GetPrimitiveSceneInfo(FPrimitiveComponentId PrimitiveId) const final;
+	virtual FPrimitiveSceneInfo* GetPrimitiveSceneInfo(const FPersistentPrimitiveIndex& PersistentPrimitiveIndex) const final;
 	virtual bool GetPreviousLocalToWorld(const FPrimitiveSceneInfo* PrimitiveSceneInfo, FMatrix& OutPreviousLocalToWorld) const override;
 	virtual void AddLight(ULightComponent* Light) override;
 	virtual void RemoveLight(ULightComponent* Light) override;
