- Fix crash on some RHIs due to  FMeshDrawSingleShaderBindings binding a null resource which is invalid
- Add missing call to SetComputePipelineState
- Ensure UB is in scope until the graph pass is added

#rnx

diff --git a/Engine/Source/Runtime/Renderer/Private/HeterogeneousVolumes/HeterogeneousVolumesHardwareRayTracing.cpp b/Engine/Source/Runtime/Renderer/Private/HeterogeneousVolumes/HeterogeneousVolumesHardwareRayTracing.cpp
--- a/Engine/Source/Runtime/Renderer/Private/HeterogeneousVolumes/HeterogeneousVolumesHardwareRayTracing.cpp
+++ b/Engine/Source/Runtime/Renderer/Private/HeterogeneousVolumes/HeterogeneousVolumesHardwareRayTracing.cpp
@@ -592,6 +592,10 @@ void RenderLightingCacheWithPreshadingHardwareRayTracing(
 {
 	const IHeterogeneousVolumeInterface* Interface = HeterogeneousVolumes::GetInterface(PrimitiveSceneProxy);
 
+	// Note must be done in the same scope as we add the pass otherwise the UB lifetime will not be guaranteed
+	FDeferredLightUniformStruct DeferredLightUniform = GetDeferredLightParameters(View, *LightSceneInfo);
+	TUniformBufferRef<FDeferredLightUniformStruct> DeferredLightUB = CreateUniformBufferImmediate(DeferredLightUniform, UniformBuffer_SingleDraw);
+
 	FRenderLightingCacheWithPreshadingRGS::FParameters* PassParameters = GraphBuilder.AllocParameters<FRenderLightingCacheWithPreshadingRGS::FParameters>();
 	{
 		// Scene
@@ -604,8 +608,7 @@ void RenderLightingCacheWithPreshadingHardwareRayTracing(
 		PassParameters->bApplyEmissionAndTransmittance = bApplyEmissionAndTransmittance;
 		PassParameters->bApplyDirectLighting = bApplyDirectLighting;
 		PassParameters->bApplyShadowTransmittance = bApplyShadowTransmittance;
-		FDeferredLightUniformStruct DeferredLightUniform = GetDeferredLightParameters(View, *LightSceneInfo);
-		PassParameters->DeferredLight = CreateUniformBufferImmediate(DeferredLightUniform, UniformBuffer_SingleDraw);
+		PassParameters->DeferredLight = DeferredLightUB;
 		PassParameters->LightType = LightType;
 		PassParameters->VolumetricScatteringIntensity = LightSceneInfo->Proxy->GetVolumetricScatteringIntensity();
 
@@ -734,6 +737,14 @@ void RenderSingleScatteringWithPreshadingHardwareRayTracing(
 {
 	const IHeterogeneousVolumeInterface* Interface = HeterogeneousVolumes::GetInterface(PrimitiveSceneProxy);
 
+	// Note must be done in the same scope as we add the pass otherwise the UB lifetime will not be guaranteed
+	FDeferredLightUniformStruct DeferredLightUniform;
+	if (bApplyDirectLighting && (LightSceneInfo != nullptr))
+	{
+		DeferredLightUniform = GetDeferredLightParameters(View, *LightSceneInfo);
+	}
+	TUniformBufferRef<FDeferredLightUniformStruct> DeferredLightUB = CreateUniformBufferImmediate(DeferredLightUniform, UniformBuffer_SingleDraw);
+
 	FRenderSingleScatteringWithPreshadingRGS::FParameters* PassParameters = GraphBuilder.AllocParameters<FRenderSingleScatteringWithPreshadingRGS::FParameters>();
 	{
 		// Scene
@@ -742,15 +753,13 @@ void RenderSingleScatteringWithPreshadingHardwareRayTracing(
 		PassParameters->SceneTextures = GetSceneTextureParameters(GraphBuilder, SceneTextures);
 
 		// Light data
-		FDeferredLightUniformStruct DeferredLightUniform;
 		PassParameters->bApplyEmissionAndTransmittance = bApplyEmissionAndTransmittance;
 		PassParameters->bApplyDirectLighting = bApplyDirectLighting;
 		if (PassParameters->bApplyDirectLighting && (LightSceneInfo != nullptr))
 		{
-			DeferredLightUniform = GetDeferredLightParameters(View, *LightSceneInfo);
 			PassParameters->VolumetricScatteringIntensity = LightSceneInfo->Proxy->GetVolumetricScatteringIntensity();
 		}
-		PassParameters->DeferredLight = CreateUniformBufferImmediate(DeferredLightUniform, UniformBuffer_SingleDraw);
+		PassParameters->DeferredLight = DeferredLightUB;
 		PassParameters->LightType = LightType;
 
 		// Shadow data
diff --git a/Engine/Source/Runtime/Renderer/Private/HeterogeneousVolumes/HeterogeneousVolumesLiveShadingPipeline.cpp b/Engine/Source/Runtime/Renderer/Private/HeterogeneousVolumes/HeterogeneousVolumesLiveShadingPipeline.cpp
--- a/Engine/Source/Runtime/Renderer/Private/HeterogeneousVolumes/HeterogeneousVolumesLiveShadingPipeline.cpp
+++ b/Engine/Source/Runtime/Renderer/Private/HeterogeneousVolumes/HeterogeneousVolumesLiveShadingPipeline.cpp
@@ -14,6 +14,80 @@
 #include "VolumeLighting.h"
 #include "VolumetricFog.h"
 
+//-OPT: Remove duplicate bindings
+// At the moment we need to bind the mesh draw parameters as they will be applied and on some RHIs this will crash if the texture is nullptr
+// We have the same parameters in the loose FParameters shader structure that are applied after the mesh draw.
+class FRenderLightingCacheLooseBindings
+{
+	DECLARE_TYPE_LAYOUT(FRenderLightingCacheLooseBindings, NonVirtual);
+
+public:
+	void Bind(const FShaderParameterMap& ParameterMap)
+	{
+		SceneDepthTextureBinding.Bind(ParameterMap, TEXT("SceneDepthTexture"));
+		ShadowDepthTextureBinding.Bind(ParameterMap, TEXT("ShadowDepthTexture"));
+		ShadowDepthTextureSamplerBinding.Bind(ParameterMap, TEXT("ShadowDepthTextureSampler"));
+		StaticShadowDepthTextureBinding.Bind(ParameterMap, TEXT("StaticShadowDepthTexture"));
+		StaticShadowDepthTextureSamplerBinding.Bind(ParameterMap, TEXT("StaticShadowDepthTextureSampler"));
+		ShadowDepthCubeTextureBinding.Bind(ParameterMap, TEXT("ShadowDepthCubeTexture"));
+		ShadowDepthCubeTexture2Binding.Bind(ParameterMap, TEXT("ShadowDepthCubeTexture2"));
+		ShadowDepthCubeTextureSamplerBinding.Bind(ParameterMap, TEXT("ShadowDepthCubeTextureSampler"));
+		LightingCacheTextureBinding.Bind(ParameterMap, TEXT("LightingCacheTexture"));
+	}
+
+	template<typename TPassParameters>	
+	void SetParameters(FMeshDrawSingleShaderBindings& ShaderBindings, const TPassParameters* PassParameters)
+	{
+		ShaderBindings.AddTexture(
+			SceneDepthTextureBinding,
+			FShaderResourceParameter(),
+			TStaticSamplerState<SF_Point>::GetRHI(),
+			PassParameters->SceneTextures.SceneDepthTexture->GetRHI()
+		);
+		ShaderBindings.AddTexture(
+			ShadowDepthTextureBinding,
+			ShadowDepthTextureSamplerBinding,
+			PassParameters->VolumeShadowingShaderParameters.ShadowDepthTextureSampler,
+			PassParameters->VolumeShadowingShaderParameters.ShadowDepthTexture->GetRHI()
+		);
+		ShaderBindings.AddTexture(
+			StaticShadowDepthTextureBinding,
+			StaticShadowDepthTextureSamplerBinding,
+			PassParameters->VolumeShadowingShaderParameters.StaticShadowDepthTextureSampler,
+			PassParameters->VolumeShadowingShaderParameters.StaticShadowDepthTexture
+		);
+		ShaderBindings.AddTexture(
+			ShadowDepthCubeTextureBinding,
+			ShadowDepthCubeTextureSamplerBinding,
+			PassParameters->VolumeShadowingShaderParameters.OnePassPointShadowProjection.ShadowDepthCubeTextureSampler,
+			PassParameters->VolumeShadowingShaderParameters.OnePassPointShadowProjection.ShadowDepthCubeTexture->GetRHI()
+		);
+		ShaderBindings.AddTexture(
+			ShadowDepthCubeTexture2Binding,
+			ShadowDepthCubeTextureSamplerBinding,
+			PassParameters->VolumeShadowingShaderParameters.OnePassPointShadowProjection.ShadowDepthCubeTextureSampler,
+			PassParameters->VolumeShadowingShaderParameters.OnePassPointShadowProjection.ShadowDepthCubeTexture->GetRHI()
+		);
+		ShaderBindings.AddTexture(
+			LightingCacheTextureBinding,
+			FShaderResourceParameter(),
+			TStaticSamplerState<SF_Point>::GetRHI(),
+			PassParameters->LightingCache.LightingCacheTexture->GetRHI()
+		);
+	}
+
+	LAYOUT_FIELD(FShaderResourceParameter, SceneDepthTextureBinding);
+	LAYOUT_FIELD(FShaderResourceParameter, ShadowDepthTextureBinding);
+	LAYOUT_FIELD(FShaderResourceParameter, ShadowDepthTextureSamplerBinding);
+	LAYOUT_FIELD(FShaderResourceParameter, StaticShadowDepthTextureBinding);
+	LAYOUT_FIELD(FShaderResourceParameter, StaticShadowDepthTextureSamplerBinding);
+	LAYOUT_FIELD(FShaderResourceParameter, ShadowDepthCubeTextureBinding);
+	LAYOUT_FIELD(FShaderResourceParameter, ShadowDepthCubeTexture2Binding);
+	LAYOUT_FIELD(FShaderResourceParameter, ShadowDepthCubeTextureSamplerBinding);
+	LAYOUT_FIELD(FShaderResourceParameter, LightingCacheTextureBinding);
+};
+IMPLEMENT_TYPE_LAYOUT(FRenderLightingCacheLooseBindings);
+
 class FRenderLightingCacheWithLiveShadingCS : public FMeshMaterialShader
 {
 	DECLARE_SHADER_TYPE(FRenderLightingCacheWithLiveShadingCS, MeshMaterial);
@@ -77,6 +151,7 @@ class FRenderLightingCacheWithLiveShadingCS : public FMeshMaterialShader
 			// Don't require full bindings, we use FMaterialShader::SetParameters
 			false
 		);
+		ShaderLooseBindings.Bind(Initializer.ParameterMap);
 	}
 
 	static bool ShouldCompilePermutation(
@@ -119,6 +194,8 @@ class FRenderLightingCacheWithLiveShadingCS : public FMeshMaterialShader
 	static int32 GetThreadGroupSize1D() { return GetThreadGroupSize2D() * GetThreadGroupSize2D(); }
 	static int32 GetThreadGroupSize2D() { return 8; }
 	static int32 GetThreadGroupSize3D() { return 4; }
+
+	LAYOUT_FIELD(FRenderLightingCacheLooseBindings, ShaderLooseBindings);
 };
 
 IMPLEMENT_MATERIAL_SHADER_TYPE(, FRenderLightingCacheWithLiveShadingCS, TEXT("/Engine/Private/HeterogeneousVolumes/HeterogeneousVolumesLiveShadingPipeline.usf"), TEXT("RenderLightingCacheWithLiveShadingCS"), SF_Compute);
@@ -194,6 +271,7 @@ class FRenderSingleScatteringWithLiveShadingCS : public FMeshMaterialShader
 			// Don't require full bindings, we use FMaterialShader::SetParameters
 			false
 		);
+		ShaderLooseBindings.Bind(Initializer.ParameterMap);
 	}
 
 	static bool ShouldCompilePermutation(
@@ -234,24 +312,25 @@ class FRenderSingleScatteringWithLiveShadingCS : public FMeshMaterialShader
 
 	static int32 GetThreadGroupSize1D() { return GetThreadGroupSize2D() * GetThreadGroupSize2D(); }
 	static int32 GetThreadGroupSize2D() { return 8; }
+
+	LAYOUT_FIELD(FRenderLightingCacheLooseBindings, ShaderLooseBindings);
 };
 
 IMPLEMENT_MATERIAL_SHADER_TYPE(, FRenderSingleScatteringWithLiveShadingCS, TEXT("/Engine/Private/HeterogeneousVolumes/HeterogeneousVolumesLiveShadingPipeline.usf"), TEXT("RenderSingleScatteringWithLiveShadingCS"), SF_Compute);
 
-template <typename ComputeShaderType>
+template<bool bWithLumen, typename ComputeShaderType>
 void AddComputePass(
 	FRDGBuilder& GraphBuilder,
 	TShaderRef<ComputeShaderType>& ComputeShader,
 	typename ComputeShaderType::FParameters* PassParameters,
 	const FScene* Scene,
-	const FViewInfo& View,
 	const FMaterialRenderProxy* MaterialRenderProxy,
 	const FMaterial& Material,
 	const FString& PassName,
 	FIntVector GroupCount
 )
 {
-	ClearUnusedGraphResources(ComputeShader, PassParameters);
+	//ClearUnusedGraphResources(ComputeShader, PassParameters);
 
 	GraphBuilder.AddPass(
 		RDG_EVENT_NAME("%s", *PassName),
@@ -259,8 +338,6 @@ void AddComputePass(
 		ERDGPassFlags::Compute,
 		[ComputeShader, PassParameters, Scene, MaterialRenderProxy, &Material, GroupCount](FRHIComputeCommandList& RHICmdList)
 		{
-			FMeshPassProcessorRenderState DrawRenderState;
-
 			FMeshMaterialShaderElementData ShaderElementData;
 			ShaderElementData.FadeUniformBuffer = GDistanceCullFadedInUniformBuffer.GetUniformBufferRHI();
 			ShaderElementData.DitherUniformBuffer = GDitherFadedInUniformBuffer.GetUniformBufferRHI();
@@ -268,16 +345,24 @@ void AddComputePass(
 			FMeshProcessorShaders PassShaders;
 			PassShaders.ComputeShader = ComputeShader;
 
+			FMeshPassProcessorRenderState DrawRenderState;
 			FMeshDrawShaderBindings ShaderBindings;
+			ShaderBindings.Initialize(PassShaders);
 			{
-				ShaderBindings.Initialize(PassShaders);
-
 				int32 DataOffset = 0;
 				FMeshDrawSingleShaderBindings SingleShaderBindings = ShaderBindings.GetSingleShaderBindings(SF_Compute, DataOffset);
 				ComputeShader->GetShaderBindings(Scene, Scene->GetFeatureLevel(), nullptr, *MaterialRenderProxy, Material, DrawRenderState, ShaderElementData, SingleShaderBindings);
-
+				SingleShaderBindings.Add(ComputeShader->template GetUniformBufferParameter<FDeferredLightUniformStruct>(), PassParameters->DeferredLight.GetUniformBuffer());
+				SingleShaderBindings.Add(ComputeShader->template GetUniformBufferParameter<FForwardLightData>(), PassParameters->ForwardLightData.GetUniformBuffer()->GetRHIRef());
+				SingleShaderBindings.Add(ComputeShader->template GetUniformBufferParameter<FVirtualShadowMapUniformParameters>(), PassParameters->VirtualShadowMapSamplingParameters.VirtualShadowMap.GetUniformBuffer()->GetRHIRef());
+				if constexpr (bWithLumen)
+				{
+					SingleShaderBindings.Add(ComputeShader->template GetUniformBufferParameter<FLumenTranslucencyLightingUniforms>(), PassParameters->LumenGIVolumeStruct.GetUniformBuffer()->GetRHIRef());
+				}
+				ComputeShader->ShaderLooseBindings.SetParameters(SingleShaderBindings, PassParameters);
 				ShaderBindings.Finalize(&PassShaders);
 			}
+			SetComputePipelineState(RHICmdList, ComputeShader.GetComputeShader());
 			ShaderBindings.SetOnCommandList(RHICmdList, ComputeShader.GetComputeShader());
 
 			FComputeShaderUtils::Dispatch(RHICmdList, ComputeShader, *PassParameters, GroupCount);
@@ -317,6 +402,10 @@ void RenderLightingCacheWithLiveShading(
 
 	check(Material.GetMaterialDomain() == MD_Volume);
 
+	// Note must be done in the same scope as we add the pass otherwise the UB lifetime will not be guaranteed
+	FDeferredLightUniformStruct DeferredLightUniform = GetDeferredLightParameters(View, *LightSceneInfo);
+	TUniformBufferRef<FDeferredLightUniformStruct> DeferredLightUB = CreateUniformBufferImmediate(DeferredLightUniform, UniformBuffer_SingleDraw);
+
 	FRenderLightingCacheWithLiveShadingCS::FParameters* PassParameters = GraphBuilder.AllocParameters<FRenderLightingCacheWithLiveShadingCS::FParameters>();
 	{
 		// Scene data
@@ -328,8 +417,7 @@ void RenderLightingCacheWithLiveShading(
 		PassParameters->bApplyEmissionAndTransmittance = bApplyEmissionAndTransmittance;
 		PassParameters->bApplyDirectLighting = bApplyDirectLighting;
 		PassParameters->bApplyShadowTransmittance = bApplyShadowTransmittance;
-		FDeferredLightUniformStruct DeferredLightUniform = GetDeferredLightParameters(View, *LightSceneInfo);
-		PassParameters->DeferredLight = CreateUniformBufferImmediate(DeferredLightUniform, UniformBuffer_SingleDraw);
+		PassParameters->DeferredLight = DeferredLightUB;
 		PassParameters->LightType = LightType;
 		PassParameters->VolumetricScatteringIntensity = LightSceneInfo->Proxy->GetVolumetricScatteringIntensity();
 
@@ -344,6 +432,7 @@ void RenderLightingCacheWithLiveShading(
 		// Transmittance volume
 		PassParameters->LightingCache.LightingCacheResolution = HeterogeneousVolumes::GetLightingCacheResolution(Interface);
 		//PassParameters->LightingCache.LightingCacheTexture = GraphBuilder.CreateSRV(LightingCacheTexture);
+		PassParameters->LightingCache.LightingCacheTexture = FRDGSystemTextures::Get(GraphBuilder).VolumetricBlack;
 
 		// Ray data
 		PassParameters->MaxTraceDistance = HeterogeneousVolumes::GetMaxTraceDistance();
@@ -404,7 +493,7 @@ void RenderLightingCacheWithLiveShading(
 	TShaderRef<FRenderLightingCacheWithLiveShadingCS> ComputeShader = Material.GetShader<FRenderLightingCacheWithLiveShadingCS>(&FLocalVertexFactory::StaticType, PermutationVector, false);
 	if (!ComputeShader.IsNull())
 	{
-		AddComputePass(GraphBuilder, ComputeShader, PassParameters, Scene, View, MaterialRenderProxy, Material, PassName, GroupCount);
+		AddComputePass<false>(GraphBuilder, ComputeShader, PassParameters, Scene, MaterialRenderProxy, Material, PassName, GroupCount);
 	}
 }
 
@@ -452,6 +541,14 @@ void RenderSingleScatteringWithLiveShading(
 	AddClearUAVPass(GraphBuilder, GraphBuilder.CreateUAV(DebugOutputBuffer), 0);
 #endif
 
+	// Note must be done in the same scope as we add the pass otherwise the UB lifetime will not be guaranteed
+	FDeferredLightUniformStruct DeferredLightUniform;
+	if (bApplyDirectLighting && (LightSceneInfo != nullptr))
+	{
+		DeferredLightUniform = GetDeferredLightParameters(View, *LightSceneInfo);
+	}
+	TUniformBufferRef<FDeferredLightUniformStruct> DeferredLightUB = CreateUniformBufferImmediate(DeferredLightUniform, UniformBuffer_SingleDraw);
+
 	FRenderSingleScatteringWithLiveShadingCS::FParameters* PassParameters = GraphBuilder.AllocParameters<FRenderSingleScatteringWithLiveShadingCS::FParameters>();
 	{
 		// Scene data
@@ -459,16 +556,14 @@ void RenderSingleScatteringWithLiveShading(
 		PassParameters->SceneTextures = GetSceneTextureParameters(GraphBuilder, SceneTextures);
 
 		// Light data
-		FDeferredLightUniformStruct DeferredLightUniform;
 		PassParameters->bApplyEmissionAndTransmittance = bApplyEmissionAndTransmittance;
 		PassParameters->bApplyDirectLighting = bApplyDirectLighting;
 		PassParameters->bApplyShadowTransmittance = bApplyShadowTransmittance;
-		if (PassParameters->bApplyDirectLighting && (LightSceneInfo != nullptr))
+		if (bApplyDirectLighting && (LightSceneInfo != nullptr))
 		{
-			DeferredLightUniform = GetDeferredLightParameters(View, *LightSceneInfo);
 			PassParameters->VolumetricScatteringIntensity = LightSceneInfo->Proxy->GetVolumetricScatteringIntensity();
 		}
-		PassParameters->DeferredLight = CreateUniformBufferImmediate(DeferredLightUniform, UniformBuffer_SingleDraw);
+		PassParameters->DeferredLight = DeferredLightUB;
 		PassParameters->LightType = LightType;
 		PassParameters->ShadowStepSize = HeterogeneousVolumes::GetShadowStepSize();
 
@@ -519,6 +614,10 @@ void RenderSingleScatteringWithLiveShading(
 			PassParameters->LightingCache.LightingCacheResolution = HeterogeneousVolumes::GetLightingCacheResolution(Interface);
 			PassParameters->LightingCache.LightingCacheTexture = LightingCacheTexture;
 		}
+		else
+		{
+			PassParameters->LightingCache.LightingCacheTexture = FRDGSystemTextures::Get(GraphBuilder).VolumetricBlack;
+		}
 
 		// Dispatch data
 		PassParameters->GroupCount = GroupCount;
@@ -548,7 +647,7 @@ void RenderSingleScatteringWithLiveShading(
 	TShaderRef<FRenderSingleScatteringWithLiveShadingCS> ComputeShader = Material.GetShader<FRenderSingleScatteringWithLiveShadingCS>(&FLocalVertexFactory::StaticType, PermutationVector, false);
 	if (!ComputeShader.IsNull())
 	{
-		AddComputePass(GraphBuilder, ComputeShader, PassParameters, Scene, View, MaterialRenderProxy, Material, PassName, GroupCount);
+		AddComputePass<true>(GraphBuilder, ComputeShader, PassParameters, Scene, MaterialRenderProxy, Material, PassName, GroupCount);
 	}
 }
 
diff --git a/Engine/Source/Runtime/Renderer/Private/HeterogeneousVolumes/HeterogeneousVolumesMaterialBakingPipeline.cpp b/Engine/Source/Runtime/Renderer/Private/HeterogeneousVolumes/HeterogeneousVolumesMaterialBakingPipeline.cpp
--- a/Engine/Source/Runtime/Renderer/Private/HeterogeneousVolumes/HeterogeneousVolumesMaterialBakingPipeline.cpp
+++ b/Engine/Source/Runtime/Renderer/Private/HeterogeneousVolumes/HeterogeneousVolumesMaterialBakingPipeline.cpp
@@ -187,6 +187,7 @@ void ComputeHeterogeneousVolumeBakeMaterial(
 
 					ShaderBindings.Finalize(&PassShaders);
 				}
+				SetComputePipelineState(RHICmdList, ComputeShader.GetComputeShader());
 				ShaderBindings.SetOnCommandList(RHICmdList, ComputeShader.GetComputeShader());
 
 				FComputeShaderUtils::Dispatch(RHICmdList, ComputeShader, *PassParameters, GroupCount);
diff --git a/Engine/Source/Runtime/Renderer/Private/HeterogeneousVolumes/HeterogeneousVolumesPreshadingPipeline.cpp b/Engine/Source/Runtime/Renderer/Private/HeterogeneousVolumes/HeterogeneousVolumesPreshadingPipeline.cpp
--- a/Engine/Source/Runtime/Renderer/Private/HeterogeneousVolumes/HeterogeneousVolumesPreshadingPipeline.cpp
+++ b/Engine/Source/Runtime/Renderer/Private/HeterogeneousVolumes/HeterogeneousVolumesPreshadingPipeline.cpp
@@ -408,6 +408,10 @@ void RenderLightingCacheWithPreshadingCompute(
 {
 	const IHeterogeneousVolumeInterface* Interface = HeterogeneousVolumes::GetInterface(PrimitiveSceneProxy);
 
+	// Note must be done in the same scope as we add the pass otherwise the UB lifetime will not be guaranteed
+	FDeferredLightUniformStruct DeferredLightUniform = GetDeferredLightParameters(View, *LightSceneInfo);
+	TUniformBufferRef<FDeferredLightUniformStruct> DeferredLightUB = CreateUniformBufferImmediate(DeferredLightUniform, UniformBuffer_SingleDraw);
+
 	FRenderLightingCacheWithPreshadingCS::FParameters* PassParameters = GraphBuilder.AllocParameters<FRenderLightingCacheWithPreshadingCS::FParameters>();
 	{
 		// Scene data
@@ -419,8 +423,7 @@ void RenderLightingCacheWithPreshadingCompute(
 		PassParameters->bApplyEmissionAndTransmittance = bApplyEmissionAndTransmittance;
 		PassParameters->bApplyDirectLighting = bApplyDirectLighting;
 		PassParameters->bApplyShadowTransmittance = bApplyShadowTransmittance;
-		FDeferredLightUniformStruct DeferredLightUniform = GetDeferredLightParameters(View, *LightSceneInfo);
-		PassParameters->DeferredLight = CreateUniformBufferImmediate(DeferredLightUniform, UniformBuffer_SingleDraw);
+		PassParameters->DeferredLight = DeferredLightUB;
 		PassParameters->LightType = LightType;
 		PassParameters->VolumetricScatteringIntensity = LightSceneInfo->Proxy->GetVolumetricScatteringIntensity();
 
@@ -531,6 +534,14 @@ void RenderSingleScatteringWithPreshadingCompute(
 		TEXT("HeterogeneousVolumes.VoxelOutputBuffer")
 	);
 
+	// Note must be done in the same scope as we add the pass otherwise the UB lifetime will not be guaranteed
+	FDeferredLightUniformStruct DeferredLightUniform;
+	if (bApplyDirectLighting && (LightSceneInfo != nullptr))
+	{
+		DeferredLightUniform = GetDeferredLightParameters(View, *LightSceneInfo);
+	}
+	TUniformBufferRef<FDeferredLightUniformStruct> DeferredLightUB = CreateUniformBufferImmediate(DeferredLightUniform, UniformBuffer_SingleDraw);
+
 	FRenderSingleScatteringWithPreshadingCS::FParameters* PassParameters = GraphBuilder.AllocParameters<FRenderSingleScatteringWithPreshadingCS::FParameters>();
 	{
 		// Scene data
@@ -538,15 +549,13 @@ void RenderSingleScatteringWithPreshadingCompute(
 		PassParameters->SceneTextures = GetSceneTextureParameters(GraphBuilder, SceneTextures);
 
 		// Light data
-		FDeferredLightUniformStruct DeferredLightUniform;
 		PassParameters->bApplyEmissionAndTransmittance = bApplyEmissionAndTransmittance;
 		PassParameters->bApplyDirectLighting = bApplyDirectLighting;
 		if (PassParameters->bApplyDirectLighting && (LightSceneInfo != nullptr))
 		{
-			DeferredLightUniform = GetDeferredLightParameters(View, *LightSceneInfo);
 			PassParameters->VolumetricScatteringIntensity = LightSceneInfo->Proxy->GetVolumetricScatteringIntensity();
 		}
-		PassParameters->DeferredLight = CreateUniformBufferImmediate(DeferredLightUniform, UniformBuffer_SingleDraw);
+		PassParameters->DeferredLight = DeferredLightUB;
 		PassParameters->LightType = LightType;
 
 		// Shadow data
