Fix merging typo.

diff --git a/Engine/Source/Runtime/Engine/Public/SceneInterface.h b/Engine/Source/Runtime/Engine/Public/SceneInterface.h
--- a/Engine/Source/Runtime/Engine/Public/SceneInterface.h
+++ b/Engine/Source/Runtime/Engine/Public/SceneInterface.h
@@ -158,7 +158,7 @@ public:
 
 	/** Finds the  primitive with the associated component id. */
 	virtual FPrimitiveSceneInfo* GetPrimitiveSceneInfo(FPrimitiveComponentId PrimitiveId) const = 0;
-	virtual FPrimitiveSceneInfo* GetPrimitiveSceneInfo(const FPersistentPrimitiveIndex& PersistentPrimitiveIndex) = 0;
+	virtual FPrimitiveSceneInfo* GetPrimitiveSceneInfo(const FPersistentPrimitiveIndex& PersistentPrimitiveIndex) const = 0;
 
 	/** Get the primitive previous local to world (used for motion blur). Returns true if the matrix was set. */
 	virtual bool GetPreviousLocalToWorld(const FPrimitiveSceneInfo* PrimitiveSceneInfo, FMatrix& OutPreviousLocalToWorld) const { return false; }
diff --git a/Engine/Source/Runtime/Renderer/Private/RendererScene.cpp b/Engine/Source/Runtime/Renderer/Private/RendererScene.cpp
--- a/Engine/Source/Runtime/Renderer/Private/RendererScene.cpp
+++ b/Engine/Source/Runtime/Renderer/Private/RendererScene.cpp
@@ -2527,7 +2527,7 @@ FPrimitiveSceneInfo* FScene::GetPrimitiveSceneInfo(FPrimitiveComponentId Primiti
 	return GetPrimitiveSceneInfo(PrimitiveComponentIds.Find(PrimitiveId));
 }
 
-FPrimitiveSceneInfo* FScene::GetPrimitiveSceneInfo(const FPersistentPrimitiveIndex& PersistentPrimitiveIndex)
+FPrimitiveSceneInfo* FScene::GetPrimitiveSceneInfo(const FPersistentPrimitiveIndex& PersistentPrimitiveIndex) const
 {
 	int32 PrimitiveIndex = GetPrimitiveIndex(PersistentPrimitiveIndex);
 	return GetPrimitiveSceneInfo(PrimitiveIndex);
@@ -6679,7 +6679,7 @@ public:
 	virtual void UpdateAllPrimitiveSceneInfos(FRDGBuilder& GraphBuilder, EUpdateAllPrimitiveSceneInfosAsyncOps = EUpdateAllPrimitiveSceneInfosAsyncOps::None) override {}
 	virtual FPrimitiveSceneInfo* GetPrimitiveSceneInfo(int32 PrimiteIndex) const final { return nullptr; }
 	virtual FPrimitiveSceneInfo* GetPrimitiveSceneInfo(FPrimitiveComponentId PrimitiveId) const final { return nullptr; }
-	virtual FPrimitiveSceneInfo* GetPrimitiveSceneInfo(const FPersistentPrimitiveIndex& PersistentPrimitiveIndex) override { return nullptr; }
+	virtual FPrimitiveSceneInfo* GetPrimitiveSceneInfo(const FPersistentPrimitiveIndex& PersistentPrimitiveIndex) const final { return nullptr; }
 
 	/** Updates the transform of a primitive which has already been added to the scene. */
 	virtual void UpdatePrimitiveTransform(UPrimitiveComponent* Primitive) override {}
