Fix VSM page marking shader being compiled for non SM5 target.

diff --git a/Engine/Source/Runtime/Renderer/Private/VirtualShadowMaps/VirtualShadowMapArray.cpp b/Engine/Source/Runtime/Renderer/Private/VirtualShadowMaps/VirtualShadowMapArray.cpp
--- a/Engine/Source/Runtime/Renderer/Private/VirtualShadowMaps/VirtualShadowMapArray.cpp
+++ b/Engine/Source/Runtime/Renderer/Private/VirtualShadowMaps/VirtualShadowMapArray.cpp
@@ -714,13 +714,12 @@ class FGeneratePageFlagsFromPixelsCS : public FVirtualPageManagementShader
 	
 	static bool ShouldCompilePermutation(const FGlobalShaderPermutationParameters& Parameters) 
 	{ 
-		FGlobalShader::ShouldCompilePermutation(Parameters);
 		FPermutationDomain PermutationVector(Parameters.PermutationId);
 		if (PermutationVector.Get<FInputType>() != 0 && (PermutationVector.Get<FWaterDepth>() || PermutationVector.Get<FTranslucencyDepth>()))
 		{
 			return false;
 		}
-		return true;
+		return FVirtualPageManagementShader::ShouldCompilePermutation(Parameters);
 	}
 
 	BEGIN_SHADER_PARAMETER_STRUCT(FParameters, )
