Add cvar to disable deferred mesh pass setup sync (r.DeferredMeshPassSetupTaskSync).

diff --git a/Engine/Source/Runtime/Renderer/Private/MeshDrawCommands.cpp b/Engine/Source/Runtime/Renderer/Private/MeshDrawCommands.cpp
--- a/Engine/Source/Runtime/Renderer/Private/MeshDrawCommands.cpp
+++ b/Engine/Source/Runtime/Renderer/Private/MeshDrawCommands.cpp
@@ -37,6 +37,12 @@ static FAutoConsoleVariableRef CVarAllowOnDemandShaderCreation(
 	TEXT("\t1: If RHI supports multi-threaded shader creation, create them on demand on tasks threads, at the time of submitting the draws.\n"),
 	ECVF_RenderThreadSafe);
 
+static TAutoConsoleVariable<int32> CVarDeferredMeshPassSetupTaskSync(
+	TEXT("r.DeferredMeshPassSetupTaskSync"),
+	0,
+	TEXT("If enabled, the sync point of the mesh pass setup task is deferred until RDG execute (from during RDG setup) significantly increasing the overlap possible."),
+	ECVF_RenderThreadSafe);
+
 FPrimitiveIdVertexBufferPool::FPrimitiveIdVertexBufferPool()
 	: DiscardId(0)
 {
@@ -1421,13 +1427,12 @@ void FParallelMeshDrawCommandPass::BuildRenderingCommands(
 	if (TaskContext.InstanceCullingContext.IsEnabled())
 	{
 		check(!bHasInstanceCullingDrawParameters);
-		// 2. Run or queue finalize culling commands pass
-		TaskContext.InstanceCullingContext.BuildRenderingCommands(GraphBuilder, GPUScene, TaskContext.View->DynamicPrimitiveCollector.GetInstanceSceneDataOffset(), TaskContext.View->DynamicPrimitiveCollector.NumInstances(), TaskContext.InstanceCullingResult, &OutInstanceCullingDrawParams,
-		[this, &GPUScene]()
+
+		auto DoSetupTaskSync = [this, &GPUScene]()
 		{
 			WaitForMeshPassSetupTask();
 
-	#if DO_CHECK
+#if DO_CHECK
 			for (const FVisibleMeshDrawCommand& VisibleMeshDrawCommand : TaskContext.MeshDrawCommands)
 			{
 				if (VisibleMeshDrawCommand.PrimitiveIdInfo.bIsDynamicPrimitive)
@@ -1436,8 +1441,20 @@ void FParallelMeshDrawCommandPass::BuildRenderingCommands(
 					TaskContext.View->DynamicPrimitiveCollector.CheckPrimitiveProcessed(PrimitiveIndex, GPUScene);
 				}
 			}
-	#endif
-		});
+#endif
+		};
+
+		if (CVarDeferredMeshPassSetupTaskSync.GetValueOnRenderThread() != 0)
+		{
+			TaskContext.InstanceCullingContext.BuildRenderingCommands(GraphBuilder, GPUScene, TaskContext.View->DynamicPrimitiveCollector.GetInstanceSceneDataOffset(), TaskContext.View->DynamicPrimitiveCollector.NumInstances(), TaskContext.InstanceCullingResult, &OutInstanceCullingDrawParams, DoSetupTaskSync);
+		}
+		else
+		{
+			DoSetupTaskSync();
+			TaskContext.InstanceCullingContext.BuildRenderingCommands(GraphBuilder, GPUScene, TaskContext.View->DynamicPrimitiveCollector.GetInstanceSceneDataOffset(), TaskContext.View->DynamicPrimitiveCollector.NumInstances(), TaskContext.InstanceCullingResult, &OutInstanceCullingDrawParams);
+		}
+
+		// 2. Run or queue finalize culling commands pass
 		TaskContext.InstanceCullingResult.GetDrawParameters(OutInstanceCullingDrawParams);
 		bHasInstanceCullingDrawParameters = true;
 		check(!TaskContext.InstanceCullingContext.HasCullingCommands() || OutInstanceCullingDrawParams.DrawIndirectArgsBuffer && OutInstanceCullingDrawParams.InstanceIdOffsetBuffer);
