Fixed incorrect Reserve of occlusion sub-query results.

#rnx

diff --git a/Engine/Source/Runtime/Renderer/Private/SceneVisibility.cpp b/Engine/Source/Runtime/Renderer/Private/SceneVisibility.cpp
--- a/Engine/Source/Runtime/Renderer/Private/SceneVisibility.cpp
+++ b/Engine/Source/Runtime/Renderer/Private/SceneVisibility.cpp
@@ -1341,7 +1341,7 @@ static void FetchVisibilityForPrimitives_Range(FVisForPrimParams& Params, FGloba
 					View.PrimitiveVisibilityMap.AccessCorrespondingBit(BitIt) = false;
 					continue;
 				}
-				SubIsOccluded.Reserve(NumSubQueries);
+				SubIsOccluded.Reserve(SubIsOccluded.Num() + NumSubQueries);
 			}
 
 			bool bAllSubOcclusionStateIsDefinite = true;
