Insight asset meta data tracking on DX12 textures

diff --git a/Engine/Source/Runtime/Core/Public/ProfilingDebugging/AssetMetadataTrace.h b/Engine/Source/Runtime/Core/Public/ProfilingDebugging/AssetMetadataTrace.h
--- a/Engine/Source/Runtime/Core/Public/ProfilingDebugging/AssetMetadataTrace.h
+++ b/Engine/Source/Runtime/Core/Public/ProfilingDebugging/AssetMetadataTrace.h
@@ -35,8 +35,11 @@ UE_TRACE_METADATA_EVENT_END()
  *		 is added.
  */
 #define UE_TRACE_METADATA_SCOPE_ASSET(Object, ObjClass) \
-	auto MetaNameRef = UE_TRACE_CHANNELEXPR_IS_ENABLED(MetadataChannel) ? FStringTrace::GetNameRef(Object->GetFName()) : UE::Trace::FEventRef32(0,0);\
-	auto ClassNameRef = UE_TRACE_CHANNELEXPR_IS_ENABLED(MetadataChannel) ? FStringTrace::GetNameRef(ObjClass->GetFName()) : UE::Trace::FEventRef32(0,0);\
+	UE_TRACE_METADATA_SCOPE_ASSET_FNAME(Object->GetFName(), ObjClass->GetFName())
+
+#define UE_TRACE_METADATA_SCOPE_ASSET_FNAME(ObjectName, ObjClassName) \
+	auto MetaNameRef = UE_TRACE_CHANNELEXPR_IS_ENABLED(MetadataChannel) ? FStringTrace::GetNameRef(ObjectName) : UE::Trace::FEventRef32(0,0);\
+	auto ClassNameRef = UE_TRACE_CHANNELEXPR_IS_ENABLED(MetadataChannel) ? FStringTrace::GetNameRef(ObjClassName) : UE::Trace::FEventRef32(0,0);\
 	UE_TRACE_METADATA_SCOPE(Asset, AssetMetadataChannel)\
 		<< Asset.Name(MetaNameRef)\
 		<< Asset.Class(ClassNameRef);
@@ -44,5 +47,6 @@ UE_TRACE_METADATA_EVENT_END()
 #else
 
 #define UE_TRACE_METADATA_SCOPE_ASSET(...)
+#define UE_TRACE_METADATA_SCOPE_ASSET_FNAME(...)
 
 #endif // UE_TRACE_ASSET_METADATA_ENABLED
diff --git a/Engine/Source/Runtime/D3D12RHI/D3D12RHI.Build.cs b/Engine/Source/Runtime/D3D12RHI/D3D12RHI.Build.cs
--- a/Engine/Source/Runtime/D3D12RHI/D3D12RHI.Build.cs
+++ b/Engine/Source/Runtime/D3D12RHI/D3D12RHI.Build.cs
@@ -18,6 +18,7 @@ public class D3D12RHI : ModuleRules
 				"RHI",
 				"RHICore",
 				"RenderCore",
+				"TraceLog",
 				}
 			);
 
diff --git a/Engine/Source/Runtime/D3D12RHI/Private/D3D12Texture.cpp b/Engine/Source/Runtime/D3D12RHI/Private/D3D12Texture.cpp
--- a/Engine/Source/Runtime/D3D12RHI/Private/D3D12Texture.cpp
+++ b/Engine/Source/Runtime/D3D12RHI/Private/D3D12Texture.cpp
@@ -8,6 +8,7 @@
 #include "TextureProfiler.h"
 #include "ProfilingDebugging/MemoryTrace.h"
 #include "RHIUtilities.h"
+#include "ProfilingDebugging/AssetMetadataTrace.h"
 
 int64 FD3D12GlobalStats::GDedicatedVideoMemory = 0;
 int64 FD3D12GlobalStats::GDedicatedSystemMemory = 0;
@@ -879,6 +880,7 @@ FD3D12Texture* FD3D12DynamicRHI::CreateD3D12Texture(const FRHITextureCreateDesc&
 #if PLATFORM_WINDOWS || PLATFORM_HOLOLENS
 	
 	TRACE_CPUPROFILER_EVENT_SCOPE(D3D12RHI::CreateD3D12Texture);
+	UE_TRACE_METADATA_SCOPE_ASSET_FNAME(InCreateDesc.GetTraceAssetName(), InCreateDesc.GetTraceClassName());
 
 	// Make local copy of create desc because certain flags might be removed before actual texture creation
 	FRHITextureCreateDesc CreateDesc = InCreateDesc;
@@ -1066,6 +1068,8 @@ FTextureRHIRef FD3D12DynamicRHI::RHIAsyncCreateTexture2D(uint32 SizeX, uint32 Si
 {	
 	check(GRHISupportsAsyncTextureCreation);
 	
+	UE_TRACE_METADATA_SCOPE_ASSET_FNAME(FName(TEXT("RHIAsyncCreateTexture2D")), FName(TEXT("FRHITexture")));
+
 	const ETextureCreateFlags InvalidFlags = TexCreate_RenderTargetable | TexCreate_ResolveTargetable | TexCreate_DepthStencilTargetable | TexCreate_GenerateMipCapable | TexCreate_UAV | TexCreate_Presentable | TexCreate_CPUReadback;
 	check(!EnumHasAnyFlags(Flags, InvalidFlags));
 
@@ -1787,6 +1791,9 @@ void* FD3D12Texture::Lock(class FRHICommandListImmediate* RHICmdList, uint32 Mip
 {
 	SCOPE_CYCLE_COUNTER(STAT_D3D12LockTextureTime);
 
+	FName AssetName = GetOwnerName() == NAME_None ? GetName() : GetOwnerName();
+	UE_TRACE_METADATA_SCOPE_ASSET_FNAME(AssetName, FName(TEXT("FRHITexture Lock")));
+
 	FD3D12Device* Device = GetParentDevice();
 	FD3D12Adapter* Adapter = Device->GetParentAdapter();
 
diff --git a/Engine/Source/Runtime/Engine/Private/Components/LightComponent.cpp b/Engine/Source/Runtime/Engine/Private/Components/LightComponent.cpp
--- a/Engine/Source/Runtime/Engine/Private/Components/LightComponent.cpp
+++ b/Engine/Source/Runtime/Engine/Private/Components/LightComponent.cpp
@@ -40,7 +40,8 @@ void FStaticShadowDepthMap::InitRHI()
 	if (FApp::CanEverRender() && Data && Data->ShadowMapSizeX > 0 && Data->ShadowMapSizeY > 0 && GMaxRHIFeatureLevel >= ERHIFeatureLevel::SM5)
 	{
 		const FRHITextureCreateDesc Desc =
-			FRHITextureCreateDesc::Create2D(TEXT("FStaticShadowDepthMap"), Data->ShadowMapSizeX, Data->ShadowMapSizeY, PF_R16F);
+			FRHITextureCreateDesc::Create2D(TEXT("FStaticShadowDepthMap"), Data->ShadowMapSizeX, Data->ShadowMapSizeY, PF_R16F)
+			.SetClassName(TEXT("FStaticShadowDepthMap"));
 
 		TextureRHI = RHICreateTexture(Desc);
 
diff --git a/Engine/Source/Runtime/Engine/Private/Particles/ParticleCurveTexture.cpp b/Engine/Source/Runtime/Engine/Private/Particles/ParticleCurveTexture.cpp
--- a/Engine/Source/Runtime/Engine/Private/Particles/ParticleCurveTexture.cpp
+++ b/Engine/Source/Runtime/Engine/Private/Particles/ParticleCurveTexture.cpp
@@ -514,7 +514,8 @@ void FParticleCurveTexture::InitRHI()
 		.SetFlags(ETextureCreateFlags::RenderTargetable | ETextureCreateFlags::ShaderResource | ETextureCreateFlags::NoFastClear)
 		.SetFormat(PF_B8G8R8A8)
 		.SetClearValue(FClearValueBinding(FLinearColor::Blue))
-		.SetInitialState(ERHIAccess::SRVMask);
+		.SetInitialState(ERHIAccess::SRVMask)
+		.SetClassName(TEXT("FParticleCurveTexture"));
 
 	CurveTextureRHI = RHICreateTexture(Desc);
 }
diff --git a/Engine/Source/Runtime/Engine/Private/Particles/ParticleGpuSimulation.cpp b/Engine/Source/Runtime/Engine/Private/Particles/ParticleGpuSimulation.cpp
--- a/Engine/Source/Runtime/Engine/Private/Particles/ParticleGpuSimulation.cpp
+++ b/Engine/Source/Runtime/Engine/Private/Particles/ParticleGpuSimulation.cpp
@@ -235,7 +235,8 @@ public:
 				.SetFormat(PF_A32B32G32R32F)
 				.SetClearValue(FClearValueBinding::Transparent)
 				.SetFlags(ETextureCreateFlags::RenderTargetable | ETextureCreateFlags::ShaderResource)
-				.SetInitialState(ERHIAccess::SRVMask);
+				.SetInitialState(ERHIAccess::SRVMask)
+				.SetClassName(TEXT("FParticleStateTextures"));
 
 			PositionTextureRHI = RHICreateTexture(Desc);
 		}
@@ -250,7 +251,8 @@ public:
 				.SetFormat(PF_FloatRGBA)
 				.SetClearValue(FClearValueBinding::Transparent)
 				.SetFlags(ETextureCreateFlags::RenderTargetable | ETextureCreateFlags::ShaderResource)
-				.SetInitialState(ERHIAccess::SRVMask);
+				.SetInitialState(ERHIAccess::SRVMask)
+				.SetClassName(TEXT("FParticleStateTextures"));
 
 			VelocityTextureRHI = RHICreateTexture(Desc);
 		}
@@ -304,7 +306,8 @@ public:
 			.SetFormat(PF_B8G8R8A8)
 			.SetClearValue(FClearValueBinding::Transparent)
 			.SetFlags(ETextureCreateFlags::RenderTargetable | ETextureCreateFlags::ShaderResource | ETextureCreateFlags::NoFastClear)
-			.SetInitialState(ERHIAccess::RTV);
+			.SetInitialState(ERHIAccess::RTV)
+			.SetClassName(TEXT("FParticleAttributesTexture"));
 
 		TextureRHI = RHICreateTexture(Desc);
 
diff --git a/Engine/Source/Runtime/Engine/Private/PrecomputedVolumetricLightmap.cpp b/Engine/Source/Runtime/Engine/Private/PrecomputedVolumetricLightmap.cpp
--- a/Engine/Source/Runtime/Engine/Private/PrecomputedVolumetricLightmap.cpp
+++ b/Engine/Source/Runtime/Engine/Private/PrecomputedVolumetricLightmap.cpp
@@ -25,7 +25,8 @@ void FVolumetricLightmapDataLayer::CreateTexture(FIntVector Dimensions)
 	const FRHITextureCreateDesc Desc =
 		FRHITextureCreateDesc::Create3D(TEXT("VolumetricLightmap"), Dimensions, Format)
 		.SetFlags(ETextureCreateFlags::ShaderResource | ETextureCreateFlags::UAV)
-		.SetBulkData(this);
+		.SetBulkData(this)
+		.SetClassName(TEXT("FVolumetricLightmapDataLayer"));
 
 	Texture = RHICreateTexture(Desc);
 }
@@ -34,7 +35,8 @@ void FVolumetricLightmapDataLayer::CreateTargetTexture(FIntVector Dimensions)
 {
 	const FRHITextureCreateDesc Desc =
 		FRHITextureCreateDesc::Create3D(TEXT("VolumetricLightmap"), Dimensions, Format)
-		.SetFlags(ETextureCreateFlags::ShaderResource | ETextureCreateFlags::UAV);
+		.SetFlags(ETextureCreateFlags::ShaderResource | ETextureCreateFlags::UAV)
+		.SetClassName(TEXT("FVolumetricLightmapDataLayer"));
 
 	Texture = RHICreateTexture(Desc);
 }
diff --git a/Engine/Source/Runtime/Engine/Private/PrimitiveUniformShaderParameters.cpp b/Engine/Source/Runtime/Engine/Private/PrimitiveUniformShaderParameters.cpp
--- a/Engine/Source/Runtime/Engine/Private/PrimitiveUniformShaderParameters.cpp
+++ b/Engine/Source/Runtime/Engine/Private/PrimitiveUniformShaderParameters.cpp
@@ -42,7 +42,8 @@ void FSinglePrimitiveStructured::InitRHI()
 	{
 		const FRHITextureCreateDesc Desc =
 			FRHITextureCreateDesc::Create2D(TEXT("PrimitiveSceneDataTexture"), FPrimitiveSceneShaderData::DataStrideInFloat4s, 1, PF_A32B32G32R32F)
-			.SetFlags(ETextureCreateFlags::ShaderResource | ETextureCreateFlags::UAV);
+			.SetFlags(ETextureCreateFlags::ShaderResource | ETextureCreateFlags::UAV)
+			.SetClassName(TEXT("FSinglePrimitiveStructured"));
 
 		PrimitiveSceneDataTextureRHI = RHICreateTexture(Desc);
 		PrimitiveSceneDataTextureSRV = RHICreateShaderResourceView(PrimitiveSceneDataTextureRHI, 0);
diff --git a/Engine/Source/Runtime/Engine/Private/Rendering/Texture2DResource.cpp b/Engine/Source/Runtime/Engine/Private/Rendering/Texture2DResource.cpp
--- a/Engine/Source/Runtime/Engine/Private/Rendering/Texture2DResource.cpp
+++ b/Engine/Source/Runtime/Engine/Private/Rendering/Texture2DResource.cpp
@@ -112,7 +112,9 @@ void FTexture2DResource::CreateTexture()
 		.SetNumMips(State.NumRequestedLODs)
 		.SetFlags(CreationFlags)
 		.SetExtData(PlatformData->GetExtData())
-		.SetBulkData(ResourceMem);
+		.SetBulkData(ResourceMem)
+		.SetClassName(TEXT("FTexture2DResource"))
+		.SetAssetName(GetOwnerName());
 
 	TextureRHI = RHICreateTexture(Desc);
 
@@ -155,7 +157,9 @@ void FTexture2DResource::CreatePartiallyResidentTexture()
 		.SetNumMips(State.MaxNumLODs)
 		.SetFlags(CreationFlags | ETextureCreateFlags::Virtual)
 		.SetExtData(PlatformData->GetExtData())
-		.SetBulkData(ResourceMem);
+		.SetBulkData(ResourceMem)
+		.SetClassName(TEXT("FTexture2DResource"))
+		.SetAssetName(GetOwnerName());
 
 	TextureRHI = RHICreateTexture(Desc);
 
diff --git a/Engine/Source/Runtime/Engine/Private/Rendering/Texture3DResource.cpp b/Engine/Source/Runtime/Engine/Private/Rendering/Texture3DResource.cpp
--- a/Engine/Source/Runtime/Engine/Private/Rendering/Texture3DResource.cpp
+++ b/Engine/Source/Runtime/Engine/Private/Rendering/Texture3DResource.cpp
@@ -124,7 +124,9 @@ void FTexture3DResource::CreateTexture()
 			FRHITextureCreateDesc::Create3D(TEXT("FTexture3DResource"), FirstMip.SizeX, FirstMip.SizeY, FirstMip.SizeZ, PixelFormat)
 			.SetNumMips(State.NumRequestedLODs)
 			.SetFlags(CreationFlags)
-			.SetExtData(PlatformData->GetExtData());
+			.SetExtData(PlatformData->GetExtData())
+			.SetClassName(TEXT("FTexture3DResource"))
+			.SetAssetName(GetOwnerName());
 
 		if (GUseTexture3DBulkDataRHI)
 		{
diff --git a/Engine/Source/Runtime/Engine/Private/Slate/SlateTextures.cpp b/Engine/Source/Runtime/Engine/Private/Slate/SlateTextures.cpp
--- a/Engine/Source/Runtime/Engine/Private/Slate/SlateTextures.cpp
+++ b/Engine/Source/Runtime/Engine/Private/Slate/SlateTextures.cpp
@@ -53,7 +53,8 @@ void FSlateTexture2DRHIRef::InitDynamicRHI()
 
 			const FRHITextureCreateDesc Desc =
 				FRHITextureCreateDesc::Create2D(TEXT("FSlateTexture2DRHIRef"), Width, Height, PixelFormat)
-				.SetFlags(TexCreateFlags);
+				.SetFlags(TexCreateFlags)
+				.SetClassName(TEXT("FSlateTexture2DRHIRef"));
 
 			ShaderResource = RHICreateTexture(Desc);
 			check( IsValidRef( ShaderResource ) );
@@ -329,7 +330,8 @@ void FSlateTextureRenderTarget2DResource::InitDynamicRHI()
 			.SetFormat((EPixelFormat)Format)
 			.SetClearValue(FClearValueBinding(ClearColor))
 			.SetFlags(ETextureCreateFlags::RenderTargetable | ETextureCreateFlags::ShaderResource)
-			.SetInitialState(ERHIAccess::SRVMask);
+			.SetInitialState(ERHIAccess::SRVMask)
+			.SetClassName(TEXT("FSlateTextureRenderTarget2DResource"));
 
 		RenderTargetTextureRHI = TextureRHI = RHICreateTexture(Desc);
 	}
diff --git a/Engine/Source/Runtime/Engine/Private/TextureCube.cpp b/Engine/Source/Runtime/Engine/Private/TextureCube.cpp
--- a/Engine/Source/Runtime/Engine/Private/TextureCube.cpp
+++ b/Engine/Source/Runtime/Engine/Private/TextureCube.cpp
@@ -454,7 +454,9 @@ public:
 			.SetFormat(Owner->GetPixelFormat())
 			.SetNumMips(Owner->GetNumMips())
 			.SetFlags(TexCreateFlags)
-			.SetExtData(Owner->GetPlatformData() ? Owner->GetPlatformData()->GetExtData() : 0);
+			.SetExtData(Owner->GetPlatformData() ? Owner->GetPlatformData()->GetExtData() : 0)
+			.SetClassName(TEXT("FTextureCubeResource"))
+			.SetAssetName(GetOwnerName());
 
 		TextureCubeRHI = RHICreateTexture(Desc);
 
diff --git a/Engine/Source/Runtime/Engine/Private/TextureCubeArray.cpp b/Engine/Source/Runtime/Engine/Private/TextureCubeArray.cpp
--- a/Engine/Source/Runtime/Engine/Private/TextureCubeArray.cpp
+++ b/Engine/Source/Runtime/Engine/Private/TextureCubeArray.cpp
@@ -135,7 +135,9 @@ public:
 			.SetFormat(Owner->GetPixelFormat())
 			.SetNumMips(Owner->GetNumMips())
 			.SetFlags(TexCreateFlags)
-			.SetExtData(Owner->GetPlatformData()->GetExtData());
+			.SetExtData(Owner->GetPlatformData()->GetExtData())
+			.SetClassName(TEXT("FTextureCubeArrayResource"))
+			.SetAssetName(GetOwnerName());
 
 		TextureCubeRHI = RHICreateTexture(Desc);
 		TextureRHI = TextureCubeRHI;
diff --git a/Engine/Source/Runtime/Engine/Private/TextureRenderTarget2D.cpp b/Engine/Source/Runtime/Engine/Private/TextureRenderTarget2D.cpp
--- a/Engine/Source/Runtime/Engine/Private/TextureRenderTarget2D.cpp
+++ b/Engine/Source/Runtime/Engine/Private/TextureRenderTarget2D.cpp
@@ -629,7 +629,9 @@ void FTextureRenderTarget2DResource::InitDynamicRHI()
 			.SetNumMips(Owner->GetNumMips())
 			.SetFlags(TexCreateFlags)
 			.SetInitialState(ERHIAccess::SRVMask)
-			.SetClearValue(FClearValueBinding(ClearColor));
+			.SetClearValue(FClearValueBinding(ClearColor))
+			.SetClassName(TEXT("FTextureRenderTarget2DResource"))
+			.SetAssetName(GetOwnerName());
 
 		TextureRHI = RenderTargetTextureRHI = RHICreateTexture(Desc);
 
@@ -648,6 +650,8 @@ void FTextureRenderTarget2DResource::InitDynamicRHI()
 		SetGPUMask(FRHIGPUMask::All());
 		RHIUpdateTextureReference(Owner->TextureReference.TextureReferenceRHI, TextureRHI);
 
+		TextureRHI->SetOwnerName(GetOwnerName());
+
 		AddToDeferredUpdateList(true);
 	}
 
diff --git a/Engine/Source/Runtime/MediaAssets/Private/Misc/MediaTextureResource.cpp b/Engine/Source/Runtime/MediaAssets/Private/Misc/MediaTextureResource.cpp
--- a/Engine/Source/Runtime/MediaAssets/Private/Misc/MediaTextureResource.cpp
+++ b/Engine/Source/Runtime/MediaAssets/Private/Misc/MediaTextureResource.cpp
@@ -1319,11 +1319,14 @@ void FMediaTextureResource::CreateOutputRenderTarget(const FIntPoint & InDim, EP
 			.SetNumMips(InNumMips)
 			.SetFlags(OutputCreateFlags | ETextureCreateFlags::RenderTargetable | ETextureCreateFlags::ShaderResource)
 			.SetInitialState(ERHIAccess::SRVMask)
-			.SetClearValue(FClearValueBinding(InClearColor));
+			.SetClearValue(FClearValueBinding(InClearColor))
+			.SetClassName(TEXT("FMediaTextureResource"))
+			.SetAssetName(GetOwnerName());
 
 		OutputTarget = RHICreateTexture(Desc);
 
 		OutputTarget->SetName(TEXT("MediaTextureResourceOutput"));
+		OutputTarget->SetOwnerName(GetOwnerName());
 
 		CurrentClearColor = InClearColor;
 		CurrentNumMips = InNumMips;
diff --git a/Engine/Source/Runtime/RHI/Public/RHIResources.h b/Engine/Source/Runtime/RHI/Public/RHIResources.h
--- a/Engine/Source/Runtime/RHI/Public/RHIResources.h
+++ b/Engine/Source/Runtime/RHI/Public/RHIResources.h
@@ -1750,6 +1750,10 @@ struct FRHITextureCreateDesc : public FRHITextureDesc
 	FRHITextureCreateDesc& SetBulkData(FResourceBulkDataInterface* InBulkData) { BulkData = InBulkData;                    return *this; }
 	FRHITextureCreateDesc& DetermineInititialState()                           { if (InitialState == ERHIAccess::Unknown) InitialState = RHIGetDefaultResourceState(Flags, BulkData != nullptr); return *this; }
 	FRHITextureCreateDesc& SetFastVRAMPercentage(float In)                     { FastVRAMPercentage = uint8(FMath::Clamp(In, 0.f, 1.0f) * 0xFF); return *this; }
+	FRHITextureCreateDesc& SetClassName(const TCHAR* InClassName)			   { ClassName = FName(InClassName);           return *this; }
+	FRHITextureCreateDesc& SetAssetName(const FName& InAssetName)			   { AssetName = InAssetName;                  return *this; }
+	FName GetTraceClassName() const											   { return ClassName == NAME_None ? FName(TEXT("FRHITexture")) : ClassName; }
+	FName GetTraceAssetName() const											   { return AssetName == NAME_None ? FName(DebugName) : AssetName; }
 
 	/* The RHI access state that the resource will be created in. */
 	ERHIAccess InitialState = ERHIAccess::Unknown;
@@ -1759,6 +1763,10 @@ struct FRHITextureCreateDesc : public FRHITextureDesc
 
 	/* Optional initial data to fill the resource with. */
 	FResourceBulkDataInterface* BulkData = nullptr;
+
+private:
+	FName ClassName = NAME_None;	// The owner class of FRHITexture used for Insight asset metadata tracing
+	FName AssetName = NAME_None;	// The owner asset name used for Insight asset metadata tracing
 };
 
 class FRHITexture : public FRHIViewableResource
diff --git a/Engine/Source/Runtime/RenderCore/Private/GlobalRenderResources.cpp b/Engine/Source/Runtime/RenderCore/Private/GlobalRenderResources.cpp
--- a/Engine/Source/Runtime/RenderCore/Private/GlobalRenderResources.cpp
+++ b/Engine/Source/Runtime/RenderCore/Private/GlobalRenderResources.cpp
@@ -54,7 +54,8 @@ public:
 		const FRHITextureCreateDesc Desc =
 			FRHITextureCreateDesc::Create2D(TEXT("ColoredTexture"), 1, 1, PF_R8G8B8A8)
 			.SetFlags(ETextureCreateFlags::ShaderResource)
-			.SetBulkData(&BulkData);
+			.SetBulkData(&BulkData)
+			.SetClassName(TEXT("FColoredTexture"));
 
 		TextureRHI = RHICreateTexture(Desc);
 
@@ -186,7 +187,8 @@ public:
 			const FRHITextureCreateDesc Desc =
 				FRHITextureCreateDesc::Create3D(TEXT("BlackVolumeTexture3D"), 1, 1, 1, PixelFormat)
 				.SetFlags(ETextureCreateFlags::ShaderResource)
-				.SetBulkData(&BulkData);
+				.SetBulkData(&BulkData)
+				.SetClassName(TEXT("FBlackVolumeTexture"));
 
 			TextureRHI = RHICreateTexture(Desc);
 		}
@@ -230,7 +232,8 @@ public:
 		const FRHITextureCreateDesc Desc =
 			FRHITextureCreateDesc::Create2DArray(TEXT("BlackArrayTexture"), 1, 1, 1, PF_B8G8R8A8)
 			.SetFlags(ETextureCreateFlags::ShaderResource)
-			.SetBulkData(&BulkData);
+			.SetBulkData(&BulkData)
+			.SetClassName(TEXT("FBlackArrayTexture"));
 
 		TextureRHI = RHICreateTexture(Desc);
 
@@ -270,7 +273,8 @@ public:
 		const FRHITextureCreateDesc Desc =
 			FRHITextureCreateDesc::Create2D(TEXT("FMipColorTexture"), TextureSize, TextureSize, PF_B8G8R8A8)
 			.SetNumMips(NumMips)
-			.SetFlags(ETextureCreateFlags::ShaderResource);
+			.SetFlags(ETextureCreateFlags::ShaderResource)
+			.SetClassName(TEXT("FMipColorTexture"));
 
 		TextureRHI = RHICreateTexture(Desc);
 
@@ -358,7 +362,8 @@ public:
 		// Create the texture RHI.
 		const FRHITextureCreateDesc Desc =
 			FRHITextureCreateDesc::CreateCube(TEXT("SolidColorCube"), 1, PixelFormat)
-			.SetFlags(ETextureCreateFlags::ShaderResource);
+			.SetFlags(ETextureCreateFlags::ShaderResource)
+			.SetClassName(TEXT("FSolidColorTextureCube"));
 
 		FTextureCubeRHIRef TextureCube = RHICreateTexture(Desc);
 		TextureRHI = TextureCube;
@@ -414,7 +419,8 @@ public:
 
 			const FRHITextureCreateDesc Desc =
 				FRHITextureCreateDesc::CreateCubeArray(TEXT("BlackCubeArray"), 1, 1, PF_B8G8R8A8)
-				.SetFlags(ETextureCreateFlags::ShaderResource);
+				.SetFlags(ETextureCreateFlags::ShaderResource)
+				.SetClassName(TEXT("FBlackCubeArrayTexture"));
 
 			// Create the texture RHI.
 			TextureRHI = RHICreateTexture(Desc);
@@ -451,7 +457,8 @@ public:
 	{
 		const FRHITextureCreateDesc Desc =
 			FRHITextureCreateDesc::Create2D(TEXT("UintTexture"), 1, 1, Format)
-			.SetFlags(ETextureCreateFlags::ShaderResource);
+			.SetFlags(ETextureCreateFlags::ShaderResource)
+			.SetClassName(TEXT("FUintTexture"));
 
 		TextureRHI = RHICreateTexture(Desc);
 
diff --git a/Engine/Source/Runtime/RenderCore/Private/RenderTargetPool.cpp b/Engine/Source/Runtime/RenderCore/Private/RenderTargetPool.cpp
--- a/Engine/Source/Runtime/RenderCore/Private/RenderTargetPool.cpp
+++ b/Engine/Source/Runtime/RenderCore/Private/RenderTargetPool.cpp
@@ -111,6 +111,7 @@ TRefCountPtr<IPooledRenderTarget> FRenderTargetPool::FindFreeElement(FRHITexture
 
 		const ERHIAccess AccessInitial = ERHIAccess::SRVMask;
 		FRHITextureCreateDesc CreateDesc(Desc, AccessInitial, Name);
+		CreateDesc.SetClassName(TEXT("FPooledRenderTarget"));
 
 		Found = new FPooledRenderTarget(
 			RHICreateTexture(CreateDesc),
diff --git a/Engine/Source/Runtime/Renderer/Private/SystemTextures.cpp b/Engine/Source/Runtime/Renderer/Private/SystemTextures.cpp
--- a/Engine/Source/Runtime/Renderer/Private/SystemTextures.cpp
+++ b/Engine/Source/Runtime/Renderer/Private/SystemTextures.cpp
@@ -116,7 +116,8 @@ void FSystemTextures::InitializeCommonTextures(FRHICommandListImmediate& RHICmdL
 	{
 		const FRHITextureCreateDesc Desc =
 			FRHITextureCreateDesc::Create2D(TEXT("BlackDummy"), 1, 1, PF_B8G8R8A8)
-			.SetFlags(ETextureCreateFlags::ShaderResource);
+			.SetFlags(ETextureCreateFlags::ShaderResource)
+			.SetClassName(TEXT("FSystemTextures"));
 
 		FTextureRHIRef Texture = RHICreateTexture(Desc);
 		SetDummyTextureData<FColor>(Texture, FColor(0, 0, 0, 0));
@@ -127,7 +128,8 @@ void FSystemTextures::InitializeCommonTextures(FRHICommandListImmediate& RHICmdL
 	{
 		const FRHITextureCreateDesc Desc =
 			FRHITextureCreateDesc::Create2DArray(TEXT("BlackArrayDummy"), 1, 1, 1, PF_B8G8R8A8)
-			.SetFlags(ETextureCreateFlags::ShaderResource);
+			.SetFlags(ETextureCreateFlags::ShaderResource)
+			.SetClassName(TEXT("FSystemTextures"));
 
 		FTextureRHIRef Texture = RHICreateTexture(Desc);
 		SetDummyTextureArrayData<FColor>(Texture, FColor(0, 0, 0, 0));
@@ -138,7 +140,8 @@ void FSystemTextures::InitializeCommonTextures(FRHICommandListImmediate& RHICmdL
 	{
 		const FRHITextureCreateDesc Desc =
 			FRHITextureCreateDesc::Create2D(TEXT("ZeroUIntDummy"), 1, 1, PF_R32_UINT)
-			.SetFlags(ETextureCreateFlags::ShaderResource);
+			.SetFlags(ETextureCreateFlags::ShaderResource)
+			.SetClassName(TEXT("FSystemTextures"));
 
 		FTextureRHIRef Texture = RHICreateTexture(Desc);
 		SetDummyTextureData<uint32>(Texture, 0u);
@@ -149,7 +152,8 @@ void FSystemTextures::InitializeCommonTextures(FRHICommandListImmediate& RHICmdL
 	{
 		const FRHITextureCreateDesc Desc =
 			FRHITextureCreateDesc::Create2DArray(TEXT("ZeroUIntArrayDummy"), 1, 1, 1, PF_R32_UINT)
-			.SetFlags(ETextureCreateFlags::ShaderResource);
+			.SetFlags(ETextureCreateFlags::ShaderResource)
+			.SetClassName(TEXT("FSystemTextures"));
 
 		FTextureRHIRef Texture = RHICreateTexture(Desc);
 		SetDummyTextureArrayData<uint32>(Texture, 0u);
@@ -160,7 +164,8 @@ void FSystemTextures::InitializeCommonTextures(FRHICommandListImmediate& RHICmdL
 	{
 		const FRHITextureCreateDesc Desc =
 			FRHITextureCreateDesc::Create2D(TEXT("BlackAlphaOneDummy"), 1, 1, PF_B8G8R8A8)
-			.SetFlags(ETextureCreateFlags::ShaderResource);
+			.SetFlags(ETextureCreateFlags::ShaderResource)
+			.SetClassName(TEXT("FSystemTextures"));
 
 		FTextureRHIRef Texture = RHICreateTexture(Desc);
 		SetDummyTextureData<FColor>(Texture, FColor(0, 0, 0, 255));
@@ -171,7 +176,8 @@ void FSystemTextures::InitializeCommonTextures(FRHICommandListImmediate& RHICmdL
 	{
 		const FRHITextureCreateDesc Desc =
 			FRHITextureCreateDesc::Create2DArray(TEXT("GreenDummy"), 1, 1, PF_B8G8R8A8)
-			.SetFlags(ETextureCreateFlags::ShaderResource);
+			.SetFlags(ETextureCreateFlags::ShaderResource)
+			.SetClassName(TEXT("FSystemTextures"));
 
 		FTextureRHIRef Texture = RHICreateTexture(Desc);
 		SetDummyTextureData<FColor>(Texture, FColor(0, 255, 0, 255));
@@ -182,7 +188,8 @@ void FSystemTextures::InitializeCommonTextures(FRHICommandListImmediate& RHICmdL
 	{
 		const FRHITextureCreateDesc Desc =
 			FRHITextureCreateDesc::Create2D(TEXT("DefaultNormal8Bit"), 1, 1, PF_B8G8R8A8)
-			.SetFlags(ETextureCreateFlags::ShaderResource);
+			.SetFlags(ETextureCreateFlags::ShaderResource)
+			.SetClassName(TEXT("FSystemTextures"));
 
 		FTextureRHIRef Texture = RHICreateTexture(Desc);
 		SetDummyTextureData<FColor>(Texture, FColor(128, 128, 128, 255));
@@ -195,7 +202,8 @@ void FSystemTextures::InitializeCommonTextures(FRHICommandListImmediate& RHICmdL
 
 		const FRHITextureCreateDesc Desc =
 			FRHITextureCreateDesc::Create2D(TEXT("PerlinNoiseGradient"), Extent, Extent, PF_B8G8R8A8)
-			.SetFlags(ETextureCreateFlags::ShaderResource);
+			.SetFlags(ETextureCreateFlags::ShaderResource)
+			.SetClassName(TEXT("FSystemTextures"));
 	
 		FTextureRHIRef Texture = RHICreateTexture(Desc);
 
@@ -232,7 +240,8 @@ void FSystemTextures::InitializeCommonTextures(FRHICommandListImmediate& RHICmdL
 	{
 		const FRHITextureCreateDesc Desc =
 			FRHITextureCreateDesc::Create2D(TEXT("MaxFP16Depth"), 1, 1, PF_FloatRGBA)
-			.SetFlags(ETextureCreateFlags::ShaderResource);
+			.SetFlags(ETextureCreateFlags::ShaderResource)
+			.SetClassName(TEXT("FSystemTextures"));
 
 		FTextureRHIRef Texture = RHICreateTexture(Desc);
 		SetDummyTextureData<FFloat16Color>(Texture, FFloat16Color(FLinearColor(65500.0f, 65500.0f, 65500.0f, 65500.0f)));
@@ -249,7 +258,8 @@ void FSystemTextures::InitializeCommonTextures(FRHICommandListImmediate& RHICmdL
 	{
 		const FRHITextureCreateDesc Desc =
 			FRHITextureCreateDesc::Create2D(TEXT("StencilDummy"), 1, 1, PF_R8G8B8A8_UINT)
-			.SetFlags(ETextureCreateFlags::ShaderResource);
+			.SetFlags(ETextureCreateFlags::ShaderResource)
+			.SetClassName(TEXT("FSystemTextures"));
 
 		FTextureRHIRef Texture = RHICreateTexture(Desc);
 		SetDummyTextureData<FColor>(Texture, FColor::White);
@@ -261,7 +271,8 @@ void FSystemTextures::InitializeCommonTextures(FRHICommandListImmediate& RHICmdL
 	{
 		const FRHITextureCreateDesc Desc =
 			FRHITextureCreateDesc::Create2D(TEXT("MidGreyDummy"), 1, 1, PF_FloatRGBA)
-			.SetFlags(ETextureCreateFlags::ShaderResource);
+			.SetFlags(ETextureCreateFlags::ShaderResource)
+			.SetClassName(TEXT("FSystemTextures"));
 
 		FTextureRHIRef Texture = RHICreateTexture(Desc);
 		SetDummyTextureData<FFloat16Color>(Texture, FFloat16Color(FLinearColor(0.5f, 0.5f, 0.5f, 0.5f)));
@@ -298,7 +309,8 @@ void FSystemTextures::InitializeFeatureLevelDependentTextures(FRHICommandListImm
 	{
 		const FRHITextureCreateDesc Desc =
 			FRHITextureCreateDesc::Create2D(TEXT("SobolSampling"), 32, 16, PF_R16_UINT)
-			.SetFlags(ETextureCreateFlags::ShaderResource);
+			.SetFlags(ETextureCreateFlags::ShaderResource)
+			.SetClassName(TEXT("FSystemTextures"));
 
 		FTextureRHIRef Texture = RHICreateTexture(Desc);
 		// Write the contents of the texture.
@@ -425,7 +437,8 @@ void FSystemTextures::InitializeFeatureLevelDependentTextures(FRHICommandListImm
 
 		const FRHITextureCreateDesc Desc =
 			FRHITextureCreateDesc::Create2D(TEXT("PreintegratedGF"), Extent, Format)
-			.SetFlags(ETextureCreateFlags::ShaderResource);
+			.SetFlags(ETextureCreateFlags::ShaderResource)
+			.SetClassName(TEXT("FSystemTextures"));
 
 		FTextureRHIRef Texture = RHICreateTexture(Desc);
 
@@ -548,7 +561,8 @@ void FSystemTextures::InitializeFeatureLevelDependentTextures(FRHICommandListImm
 
 			const FRHITextureCreateDesc Desc =
 				FRHITextureCreateDesc::Create3D(TEXT("PerlinNoise3D"), Extent, Extent, Extent, PF_B8G8R8A8)
-				.SetFlags(ETextureCreateFlags::ShaderResource | ETextureCreateFlags::NoTiling);
+				.SetFlags(ETextureCreateFlags::ShaderResource | ETextureCreateFlags::NoTiling)
+				.SetClassName(TEXT("FSystemTextures"));
 
 			FTextureRHIRef Texture3D = RHICreateTexture(Desc);
 
@@ -662,7 +676,8 @@ void FSystemTextures::InitializeFeatureLevelDependentTextures(FRHICommandListImm
 			{
 				const FRHITextureCreateDesc Desc =
 					FRHITextureCreateDesc::Create2D(TEXT("LTCMat"), LTC_Size, LTC_Size, PF_FloatRGBA)
-					.SetFlags(ETextureCreateFlags::ShaderResource | ETextureCreateFlags::FastVRAM);
+					.SetFlags(ETextureCreateFlags::ShaderResource | ETextureCreateFlags::FastVRAM)
+					.SetClassName(TEXT("FSystemTextures"));
 
 				FTextureRHIRef Texture = RHICreateTexture(Desc);
 
@@ -690,7 +705,8 @@ void FSystemTextures::InitializeFeatureLevelDependentTextures(FRHICommandListImm
 			{
 				const FRHITextureCreateDesc Desc =
 					FRHITextureCreateDesc::Create2D(TEXT("LTCAmp"), LTC_Size, LTC_Size, PF_G16R16F)
-					.SetFlags(ETextureCreateFlags::ShaderResource | ETextureCreateFlags::FastVRAM);
+					.SetFlags(ETextureCreateFlags::ShaderResource | ETextureCreateFlags::FastVRAM)
+					.SetClassName(TEXT("FSystemTextures"));
 
 				FTexture2DRHIRef Texture = RHICreateTexture(Desc);
 
@@ -752,7 +768,8 @@ void FSystemTextures::InitializeFeatureLevelDependentTextures(FRHICommandListImm
 
 				const FRHITextureCreateDesc Desc =
 					FRHITextureCreateDesc::Create2D(TEXT("SSAORandomization"), Extent, Extent, PF_R8G8)
-					.SetFlags(ETextureCreateFlags::ShaderResource);
+					.SetFlags(ETextureCreateFlags::ShaderResource)
+					.SetClassName(TEXT("FSystemTextures"));
 
 				FTextureRHIRef Texture = RHICreateTexture(Desc);
 
@@ -795,7 +812,8 @@ void FSystemTextures::InitializeFeatureLevelDependentTextures(FRHICommandListImm
 			.SetExtent(Extent)
 			.SetDepth(bGTAOPreIngegratedUsingVolumeLUT ? Extent : 1)
 			.SetFormat(PF_R16F)
-			.SetFlags(ETextureCreateFlags::ShaderResource);
+			.SetFlags(ETextureCreateFlags::ShaderResource)
+			.SetClassName(TEXT("FSystemTextures"));
 
 		TRefCountPtr<FRHITexture> Texture = RHICreateTexture(Desc);
 
diff --git a/Engine/Source/Runtime/SlateRHIRenderer/Private/SlateRHIFontTexture.cpp b/Engine/Source/Runtime/SlateRHIRenderer/Private/SlateRHIFontTexture.cpp
--- a/Engine/Source/Runtime/SlateRHIRenderer/Private/SlateRHIFontTexture.cpp
+++ b/Engine/Source/Runtime/SlateRHIRenderer/Private/SlateRHIFontTexture.cpp
@@ -25,7 +25,8 @@ void FSlateFontTextureRHIResource::InitDynamicRHI()
 
 		FRHITextureCreateDesc Desc =
 			FRHITextureCreateDesc::Create2D(TEXT("FSlateFontTextureRHIResource"), Width, Height, PixelFormat)
-			.SetFlags(ETextureCreateFlags::Dynamic);
+			.SetFlags(ETextureCreateFlags::Dynamic)
+			.SetClassName(TEXT("FSlateFontTextureRHIResource"));
 
 		if (!bIsGrayscale)
 		{
