Change VSM front layer page marking enabled, only if VSM high quality is enabled.

diff --git a/Engine/Source/Runtime/Renderer/Private/VirtualShadowMaps/VirtualShadowMapArray.cpp b/Engine/Source/Runtime/Renderer/Private/VirtualShadowMaps/VirtualShadowMapArray.cpp
--- a/Engine/Source/Runtime/Renderer/Private/VirtualShadowMaps/VirtualShadowMapArray.cpp
+++ b/Engine/Source/Runtime/Renderer/Private/VirtualShadowMaps/VirtualShadowMapArray.cpp
@@ -350,6 +350,7 @@ namespace Nanite
 	extern bool IsStatFilterActive(const FString& FilterName);
 }
 
+bool IsVSMTranslucentHighQualityEnabled();
 extern bool LightGridUses16BitBuffers(EShaderPlatform Platform);
 
 FMatrix CalcTranslatedWorldToShadowUVMatrix(
@@ -1510,6 +1511,8 @@ void FVirtualShadowMapArray::BuildPageAllocations(
 					const FIntPoint PixelStride(
 						FMath::Clamp(CVarVirtualShadowMapPageMarkingPixelStrideX.GetValueOnRenderThread(), 1, 128),
 						FMath::Clamp(CVarVirtualShadowMapPageMarkingPixelStrideY.GetValueOnRenderThread(), 1, 128));
+					
+					const bool bFrontLayerEnabled = IsVSMTranslucentHighQualityEnabled() && FrontLayerTranslucencyData.IsValid();
 
 					PassParameters->HairStrands = HairStrands::BindHairStrandsViewUniformParameters(View);
 					PassParameters->View = View.ViewUniformBuffer;
@@ -1526,7 +1529,7 @@ void FVirtualShadowMapArray::BuildPageAllocations(
 							GSystemTextures.GetDefaultStructuredBuffer(GraphBuilder, sizeof(uint32), 0xFFFFFFFF));
 						PassParameters->SingleLayerWaterTileViewRes = SingleLayerWaterPrePassResult->ViewTileClassification[ViewIndex].TiledViewRes;
 					}
-					if (FrontLayerTranslucencyData.IsValid())
+					if (bFrontLayerEnabled)
 					{
 						PassParameters->FrontLayerTranslucencyDepthTexture = FrontLayerTranslucencyData.SceneDepth;
 					}
@@ -1562,7 +1565,7 @@ void FVirtualShadowMapArray::BuildPageAllocations(
 						FGeneratePageFlagsFromPixelsCS::FPermutationDomain PermutationVector;
 						PermutationVector.Set<FGeneratePageFlagsFromPixelsCS::FInputType>(static_cast<uint32>(InputType));
 						PermutationVector.Set<FGeneratePageFlagsFromPixelsCS::FWaterDepth>(SingleLayerWaterPrePassResult != nullptr);
-						PermutationVector.Set<FGeneratePageFlagsFromPixelsCS::FTranslucencyDepth>(FrontLayerTranslucencyData.IsValid());
+						PermutationVector.Set<FGeneratePageFlagsFromPixelsCS::FTranslucencyDepth>(bFrontLayerEnabled);
 						auto ComputeShader = View.ShaderMap->GetShader<FGeneratePageFlagsFromPixelsCS>(PermutationVector);
 						FComputeShaderUtils::AddPass(
 							GraphBuilder,
