Fix cross-compilation errors with BuildInstanceDrawCommands.usf on OpenGL platforms.

diff --git a/Engine/Source/Runtime/Renderer/Private/InstanceCulling/InstanceCullingContext.cpp b/Engine/Source/Runtime/Renderer/Private/InstanceCulling/InstanceCullingContext.cpp
--- a/Engine/Source/Runtime/Renderer/Private/InstanceCulling/InstanceCullingContext.cpp
+++ b/Engine/Source/Runtime/Renderer/Private/InstanceCulling/InstanceCullingContext.cpp
@@ -365,6 +365,9 @@ public:
 		FGlobalShader::ModifyCompilationEnvironment(Parameters, OutEnvironment);
 		FInstanceProcessingGPULoadBalancer::SetShaderDefines(OutEnvironment);
 
+		// This shader makes use of HLSL 2021 syntax
+		OutEnvironment.CompilerFlags.Add(CFLAG_HLSL2021);
+
 		OutEnvironment.SetDefine(TEXT("INDIRECT_ARGS_NUM_WORDS"), FInstanceCullingContext::IndirectArgsNumWords);
 		OutEnvironment.SetDefine(TEXT("VF_SUPPORTS_PRIMITIVE_SCENE_DATA"), 1);
 		OutEnvironment.SetDefine(TEXT("USE_GLOBAL_GPU_SCENE_DATA"), 1);
@@ -1110,6 +1113,9 @@ public:
 		FGlobalShader::ModifyCompilationEnvironment(Parameters, OutEnvironment);
 		FInstanceProcessingGPULoadBalancer::SetShaderDefines(OutEnvironment);
 
+		// This shader makes use of HLSL 2021 syntax
+		OutEnvironment.CompilerFlags.Add(CFLAG_HLSL2021);
+
 		OutEnvironment.SetDefine(TEXT("INDIRECT_ARGS_NUM_WORDS"), FInstanceCullingContext::IndirectArgsNumWords);
 		OutEnvironment.SetDefine(TEXT("NUM_THREADS_PER_GROUP"), NumThreadsPerGroup);
 		OutEnvironment.SetDefine(TEXT("COMPACTION_BLOCK_NUM_INSTANCES"), FCompactVisibleInstancesBaseCs::CompactionBlockNumInstances);
