Fix VSM cache being invalidated every frame when a planar reflection is in view.

diff --git a/Engine/Source/Runtime/Renderer/Private/PlanarReflectionRendering.cpp b/Engine/Source/Runtime/Renderer/Private/PlanarReflectionRendering.cpp
--- a/Engine/Source/Runtime/Renderer/Private/PlanarReflectionRendering.cpp
+++ b/Engine/Source/Runtime/Renderer/Private/PlanarReflectionRendering.cpp
@@ -655,6 +655,17 @@ void FScene::UpdatePlanarReflectionContents(UPlanarReflectionComponent* CaptureC
 			Extension->SetupViewFamily(ViewFamily);
 		}
 
+		FSceneViewStateInterface* ViewStateInterface = CaptureComponent->GetViewState(0);
+
+		if (UseVirtualShadowMaps(SceneRenderer->ShaderPlatform, FeatureLevel) && ViewStateInterface)
+		{
+			ViewStateInterface->AddVirtualShadowMapCache(this);
+		}
+		else if (ViewStateInterface)
+		{
+			ViewStateInterface->RemoveVirtualShadowMapCache(this);
+		}
+
 		for (int32 ViewIndex = 0; ViewIndex < SceneCaptureViewInfo.Num(); ++ViewIndex)
 		{
 			FViewInfo& ViewInfo = SceneRenderer->Views[ViewIndex];
