Fixed assert in OpenGL due to creating uniform buffer off the render thread. Also removed explicit template type from TConstSetBitIterator calls.

diff --git a/Engine/Source/Runtime/Engine/Private/PrimitiveSceneProxy.cpp b/Engine/Source/Runtime/Engine/Private/PrimitiveSceneProxy.cpp
--- a/Engine/Source/Runtime/Engine/Private/PrimitiveSceneProxy.cpp
+++ b/Engine/Source/Runtime/Engine/Private/PrimitiveSceneProxy.cpp
@@ -543,7 +543,7 @@ FPrimitiveViewRelevance FPrimitiveSceneProxy::GetViewRelevance(const FSceneView*
 
 void FPrimitiveSceneProxy::CreateUniformBuffer()
 {
-	if (DoesVFRequirePrimitiveUniformBuffer())
+	if (!UniformBuffer && DoesVFRequirePrimitiveUniformBuffer())
 	{
 		UniformBuffer = TUniformBufferRef<FPrimitiveUniformShaderParameters>::CreateEmptyUniformBufferImmediate(UniformBuffer_MultiFrame);
 	}
diff --git a/Engine/Source/Runtime/Renderer/Private/RendererScene.cpp b/Engine/Source/Runtime/Renderer/Private/RendererScene.cpp
--- a/Engine/Source/Runtime/Renderer/Private/RendererScene.cpp
+++ b/Engine/Source/Runtime/Renderer/Private/RendererScene.cpp
@@ -2110,8 +2110,6 @@ void FScene::BatchAddPrimitives(TArrayView<UPrimitiveComponent*> InPrimitives)
 			continue;
 		}
 
-		PrimitiveSceneProxy->CreateUniformBuffer();
-
 		// Create the primitive scene info.
 		FPrimitiveSceneInfo* PrimitiveSceneInfo = new FPrimitiveSceneInfo(Primitive, this);
 		PrimitiveSceneProxy->PrimitiveSceneInfo = PrimitiveSceneInfo;
@@ -6367,7 +6365,7 @@ void FScene::UpdateAllPrimitiveSceneInfos(FRDGBuilder& GraphBuilder, EUpdateAllP
 			PrimitivesNeedingStaticMeshUpdate[PrimitiveSceneInfo->PackedIndex] = false;
 		}
 
-		for (TConstSetBitIterator<FDefaultBitArrayAllocator> BitIt(PrimitivesNeedingStaticMeshUpdate); BitIt; ++BitIt)
+		for (TConstSetBitIterator<> BitIt(PrimitivesNeedingStaticMeshUpdate); BitIt; ++BitIt)
 		{
 			const int32 Index = BitIt.GetIndex();
 			FPrimitiveSceneInfo* Primitive = Primitives[Index];
@@ -6386,10 +6384,11 @@ void FScene::UpdateAllPrimitiveSceneInfos(FRDGBuilder& GraphBuilder, EUpdateAllP
 		SCOPED_NAMED_EVENT(UpdateUniformBuffers, FColor::Emerald);
 		TArray<FPrimitiveSceneProxy*, SceneRenderingAllocator> ProxiesToUpdate;
 
-		for (TConstSetBitIterator<FDefaultBitArrayAllocator> BitIt(PrimitivesNeedingUniformBufferUpdate); BitIt; ++BitIt)
+		for (TConstSetBitIterator<> BitIt(PrimitivesNeedingUniformBufferUpdate); BitIt; ++BitIt)
 		{
 			const int32 Index = BitIt.GetIndex();
 			FPrimitiveSceneInfo* Primitive = Primitives[Index];
+			Primitive->Proxy->CreateUniformBuffer();
 			PrimitivesNeedingUniformBufferUpdate[Index] = false;
 			ProxiesToUpdate.Emplace(Primitive->Proxy);
 			GPUScene.AddPrimitiveToUpdate(Index, EPrimitiveDirtyState::ChangedAll);
