Moved distance field updates to happen during update scene. This will introduce more overlap with async visibility tasks.

diff --git a/Engine/Shaders/Private/DistanceFieldStreaming.usf b/Engine/Shaders/Private/DistanceFieldStreaming.usf
--- a/Engine/Shaders/Private/DistanceFieldStreaming.usf
+++ b/Engine/Shaders/Private/DistanceFieldStreaming.usf
@@ -94,6 +94,8 @@ void CopyDistanceFieldIndirectionAtlasCS(
 RWStructuredBuffer<uint> RWDistanceFieldAssetWantedNumMips;
 RWStructuredBuffer<uint> RWDistanceFieldAssetStreamingRequests;
 
+float3 PreViewTranslationTile;
+float3 PreViewTranslationOffset;
 float3 Mip1WorldTranslatedCenter;
 float3 Mip1WorldExtent;
 float3 Mip2WorldTranslatedCenter;
@@ -126,7 +128,7 @@ void ComputeDistanceFieldAssetWantedMipsCS(
 		if (DebugForceNumMips == 0)
 		{
 			FDFObjectBounds DFObjectBounds = LoadDFObjectBounds(ObjectIndex);
-			const float3 TranslatedCenter = LWCToFloat(LWCAdd(DFObjectBounds.Center, PrimaryView.PreViewTranslation));
+			const float3 TranslatedCenter = LWCToFloat(LWCAdd(DFObjectBounds.Center, MakeLWCVector3(PreViewTranslationTile, PreViewTranslationOffset)));
 
 			float OuterDistanceSq = ComputeSquaredDistanceBetweenAABBs(Mip1WorldTranslatedCenter, Mip1WorldExtent, TranslatedCenter, DFObjectBounds.BoxExtent);
  
diff --git a/Engine/Source/Runtime/Renderer/Private/DeferredShadingRenderer.cpp b/Engine/Source/Runtime/Renderer/Private/DeferredShadingRenderer.cpp
--- a/Engine/Source/Runtime/Renderer/Private/DeferredShadingRenderer.cpp
+++ b/Engine/Source/Runtime/Renderer/Private/DeferredShadingRenderer.cpp
@@ -2518,6 +2518,8 @@ void FDeferredShadingSceneRenderer::Render(FRDGBuilder& GraphBuilder)
 		ShadingEnergyConservation::Init(GraphBuilder, View);
 	}
 
+	FRDGExternalAccessQueue ExternalAccessQueue;
+
 	{
 		EUpdateAllPrimitiveSceneInfosAsyncOps AsyncOps = EUpdateAllPrimitiveSceneInfosAsyncOps::None;
 
@@ -2532,6 +2534,8 @@ void FDeferredShadingSceneRenderer::Render(FRDGBuilder& GraphBuilder)
 		}
 
 		UpdateScene(GraphBuilder, AsyncOps);
+
+		PrepareDistanceFieldScene(GraphBuilder, ExternalAccessQueue);
 	}
 	// kick off dependent scene updates 
 	Scene->ShadowScene->UpdateForRenderedFrame(GraphBuilder);
@@ -2843,8 +2847,6 @@ void FDeferredShadingSceneRenderer::Render(FRDGBuilder& GraphBuilder)
 	Scene->DebugRender(Views);
 #endif
 
-	FRDGExternalAccessQueue ExternalAccessQueue;
-
 	{
 		RDG_CSV_STAT_EXCLUSIVE_SCOPE(GraphBuilder, UpdateGPUScene);
 		RDG_GPU_STAT_SCOPE(GraphBuilder, GPUSceneUpdate);
@@ -2998,8 +3000,6 @@ void FDeferredShadingSceneRenderer::Render(FRDGBuilder& GraphBuilder)
 		}
 	}
 
-	PrepareDistanceFieldScene(GraphBuilder, InitViewTaskDatas.DynamicShadows, ExternalAccessQueue);
-
 	{
 		{
 			RDG_RHI_GPU_STAT_SCOPE(GraphBuilder, VisibilityCommands);
diff --git a/Engine/Source/Runtime/Renderer/Private/DistanceFieldAmbientOcclusion.cpp b/Engine/Source/Runtime/Renderer/Private/DistanceFieldAmbientOcclusion.cpp
--- a/Engine/Source/Runtime/Renderer/Private/DistanceFieldAmbientOcclusion.cpp
+++ b/Engine/Source/Runtime/Renderer/Private/DistanceFieldAmbientOcclusion.cpp
@@ -711,7 +711,7 @@ bool FSceneRenderer::ShouldPrepareForDistanceFieldAO() const
 			|| (GDistanceFieldAOApplyToStaticIndirect && bAnyViewHasGIMethodSupportingDFAO && ViewFamily.EngineShowFlags.DistanceFieldAO));
 }
 
-bool FSceneRenderer::ShouldPrepareDistanceFieldScene(FDynamicShadowsTaskData* TaskData) const
+bool FSceneRenderer::ShouldPrepareDistanceFieldScene() const
 {
 	if (!ensure(Scene != nullptr))
 	{
@@ -734,7 +734,7 @@ bool FSceneRenderer::ShouldPrepareDistanceFieldScene(FDynamicShadowsTaskData* Ta
 	bool bShouldPrepareForDFInsetIndirectShadow = ShouldPrepareForDFInsetIndirectShadow();
 
 	// Prepare the distance field scene (object buffers and distance field atlas) if any feature needs it
-	return bShouldPrepareGlobalDistanceField || bShouldPrepareForAO || ShouldPrepareForDistanceFieldShadows(TaskData) || bShouldPrepareForDFInsetIndirectShadow;
+	return bShouldPrepareGlobalDistanceField || bShouldPrepareForAO || ShouldPrepareForDistanceFieldShadows() || bShouldPrepareForDFInsetIndirectShadow;
 }
 
 bool FSceneRenderer::ShouldPrepareGlobalDistanceField() const
diff --git a/Engine/Source/Runtime/Renderer/Private/DistanceFieldObjectManagement.cpp b/Engine/Source/Runtime/Renderer/Private/DistanceFieldObjectManagement.cpp
--- a/Engine/Source/Runtime/Renderer/Private/DistanceFieldObjectManagement.cpp
+++ b/Engine/Source/Runtime/Renderer/Private/DistanceFieldObjectManagement.cpp
@@ -849,7 +849,7 @@ void FSceneRenderer::UpdateGlobalHeightFieldObjectBuffers(FRDGBuilder& GraphBuil
 	}
 }
 
-void FSceneRenderer::PrepareDistanceFieldScene(FRDGBuilder& GraphBuilder, FDynamicShadowsTaskData* TaskData, FRDGExternalAccessQueue& ExternalAccessQueue)
+void FSceneRenderer::PrepareDistanceFieldScene(FRDGBuilder& GraphBuilder, FRDGExternalAccessQueue& ExternalAccessQueue)
 {
 	RDG_CSV_STAT_EXCLUSIVE_SCOPE(GraphBuilder, PrepareDistanceFieldScene);
 	TRACE_CPUPROFILER_EVENT_SCOPE(FSceneRenderer::PrepareDistanceFieldScene);
@@ -859,7 +859,7 @@ void FSceneRenderer::PrepareDistanceFieldScene(FRDGBuilder& GraphBuilder, FDynam
 	RDG_RHI_GPU_STAT_SCOPE(GraphBuilder, DistanceFields);
 
 	const bool bShouldPrepareHeightFieldScene = ShouldPrepareHeightFieldScene();
-	const bool bShouldPrepareDistanceFieldScene = ShouldPrepareDistanceFieldScene(TaskData);
+	const bool bShouldPrepareDistanceFieldScene = ShouldPrepareDistanceFieldScene();
 
 	if (!bShouldPrepareDistanceFieldScene && !bShouldPrepareHeightFieldScene)
 	{
diff --git a/Engine/Source/Runtime/Renderer/Private/DistanceFieldShadowing.cpp b/Engine/Source/Runtime/Renderer/Private/DistanceFieldShadowing.cpp
--- a/Engine/Source/Runtime/Renderer/Private/DistanceFieldShadowing.cpp
+++ b/Engine/Source/Runtime/Renderer/Private/DistanceFieldShadowing.cpp
@@ -698,14 +698,14 @@ bool SupportsHeightFieldShadows(ERHIFeatureLevel::Type FeatureLevel, EShaderPlat
 		&& DoesPlatformSupportDistanceFieldShadowing(ShaderPlatform);
 }
 
-bool FSceneRenderer::ShouldPrepareForDistanceFieldShadows(FDynamicShadowsTaskData* TaskData) const
+bool FSceneRenderer::ShouldPrepareForDistanceFieldShadows() const
 {
 	if (!ViewFamily.EngineShowFlags.DynamicShadows || !SupportsDistanceFieldShadows(Scene->GetFeatureLevel(), Scene->GetShaderPlatform()))
 	{
 		return false;
 	}
 
-	return HasRayTracedDistanceFieldShadows(TaskData);
+	return true;
 }
 
 bool FSceneRenderer::ShouldPrepareHeightFieldScene() const
diff --git a/Engine/Source/Runtime/Renderer/Private/DistanceFieldStreaming.cpp b/Engine/Source/Runtime/Renderer/Private/DistanceFieldStreaming.cpp
--- a/Engine/Source/Runtime/Renderer/Private/DistanceFieldStreaming.cpp
+++ b/Engine/Source/Runtime/Renderer/Private/DistanceFieldStreaming.cpp
@@ -252,6 +252,8 @@ class FComputeDistanceFieldAssetWantedMipsCS : public FGlobalShader
 		SHADER_PARAMETER(FVector3f, Mip1WorldExtent)
 		SHADER_PARAMETER(FVector3f, Mip2WorldTranslatedCenter)
 		SHADER_PARAMETER(FVector3f, Mip2WorldExtent)
+		SHADER_PARAMETER(FVector3f, PreViewTranslationTile)
+		SHADER_PARAMETER(FVector3f, PreViewTranslationOffset)
 	END_SHADER_PARAMETER_STRUCT()
 
 	using FPermutationDomain = TShaderPermutationDomain<>;
@@ -1125,6 +1127,8 @@ void FDistanceFieldSceneData::GenerateStreamingRequests(
 		FRDGBufferRef StreamingRequestsBuffer = GraphBuilder.CreateBuffer(StreamingRequestsDesc, TEXT("DistanceFields.DistanceFieldStreamingRequests"));
 
 		{
+			TLargeWorldRenderPosition<float> PreViewTranslation(View.ViewMatrices.GetPreViewTranslation());
+
 			FComputeDistanceFieldAssetWantedMipsCS::FParameters* PassParameters = GraphBuilder.AllocParameters<FComputeDistanceFieldAssetWantedMipsCS::FParameters>();
 			PassParameters->View = View.ViewUniformBuffer;
 			checkf(DistanceField::NumMips == 3, TEXT("Shader needs to be updated"));
@@ -1138,6 +1142,8 @@ void FDistanceFieldSceneData::GenerateStreamingRequests(
 			PassParameters->Mip1WorldExtent = FVector3f(GlobalDistanceField::GetClipmapExtent(GAOGlobalDistanceFieldNumClipmaps - 1, Scene, bLumenEnabled));
 			PassParameters->Mip2WorldTranslatedCenter = FVector3f(View.ViewMatrices.GetViewOrigin() + View.ViewMatrices.GetPreViewTranslation());
 			PassParameters->Mip2WorldExtent = FVector3f(GlobalDistanceField::GetClipmapExtent(FMath::Max<int32>(GAOGlobalDistanceFieldNumClipmaps / 2 - 1, 0), Scene, bLumenEnabled));
+			PassParameters->PreViewTranslationTile = PreViewTranslation.GetTile();
+			PassParameters->PreViewTranslationOffset = PreViewTranslation.GetOffset();
 
 			auto ComputeShader = GlobalShaderMap->GetShader<FComputeDistanceFieldAssetWantedMipsCS>();
 
diff --git a/Engine/Source/Runtime/Renderer/Private/GlobalDistanceField.cpp b/Engine/Source/Runtime/Renderer/Private/GlobalDistanceField.cpp
--- a/Engine/Source/Runtime/Renderer/Private/GlobalDistanceField.cpp
+++ b/Engine/Source/Runtime/Renderer/Private/GlobalDistanceField.cpp
@@ -1302,6 +1302,11 @@ static void ComputeUpdateRegionsAndUpdateViewState(
 
 void FViewInfo::SetupDefaultGlobalDistanceFieldUniformBufferParameters(FViewUniformShaderParameters& ViewUniformShaderParameters) const
 {
+	if (GlobalDistanceFieldInfo.bInitialized)
+	{
+		return;
+	}
+
 	// Initialize global distance field members to defaults, because View.GlobalDistanceFieldInfo is not valid yet
 	for (int32 Index = 0; Index < GlobalDistanceField::MaxClipmaps; Index++)
 	{
@@ -2082,9 +2087,13 @@ void UpdateGlobalDistanceFieldVolume(
 		const int32 NumClipmaps = FMath::Clamp<int32>(GlobalDistanceField::GetNumGlobalDistanceFieldClipmaps(bLumenEnabled, View.FinalPostProcessSettings.LumenSceneViewDistance), 0, GlobalDistanceField::MaxClipmaps);
 		ComputeUpdateRegionsAndUpdateViewState(GraphBuilder.RHICmdList, View, Scene, GlobalDistanceFieldInfo, NumClipmaps, MaxOcclusionDistance, bLumenEnabled);
 
+		if (!View.CachedViewUniformShaderParameters)
+		{
+			View.CachedViewUniformShaderParameters = MakeUnique<FViewUniformShaderParameters>();
+		}
+
 		// Recreate the view uniform buffer now that we have updated GlobalDistanceFieldInfo
 		View.SetupGlobalDistanceFieldUniformBufferParameters(*View.CachedViewUniformShaderParameters);
-		View.ViewUniformBuffer = TUniformBufferRef<FViewUniformShaderParameters>::CreateUniformBufferImmediate(*View.CachedViewUniformShaderParameters, UniformBuffer_SingleFrame);
 
 		bool bHasUpdateBounds = false;
 
diff --git a/Engine/Source/Runtime/Renderer/Private/MobileShadingRenderer.cpp b/Engine/Source/Runtime/Renderer/Private/MobileShadingRenderer.cpp
--- a/Engine/Source/Runtime/Renderer/Private/MobileShadingRenderer.cpp
+++ b/Engine/Source/Runtime/Renderer/Private/MobileShadingRenderer.cpp
@@ -687,11 +687,6 @@ void FMobileSceneRenderer::InitViews(
 		Scene->IndirectLightingCache.FinalizeCacheUpdates(Scene, *this, *ILCTaskData);
 	}
 
-	if (bRequiresDistanceField)
-	{
-		PrepareDistanceFieldScene(GraphBuilder, TaskDatas.DynamicShadows, ExternalAccessQueue);
-	}
-
 	ExternalAccessQueue.Submit(GraphBuilder);
 
 	InstanceCullingManager.BeginDeferredCulling(GraphBuilder, Scene->GPUScene);
@@ -842,6 +837,15 @@ void FMobileSceneRenderer::Render(FRDGBuilder& GraphBuilder)
 
 	UpdateScene(GraphBuilder);
 
+	FRDGExternalAccessQueue ExternalAccessQueue;
+
+	if (bRequiresDistanceField)
+	{
+		PrepareDistanceFieldScene(GraphBuilder, ExternalAccessQueue);
+	}
+
+	ExternalAccessQueue.Submit(GraphBuilder);
+
 	GPU_MESSAGE_SCOPE(GraphBuilder);
 
 	// Establish scene primitive count (must be done after UpdateAllPrimitiveSceneInfos)
@@ -986,8 +990,6 @@ void FMobileSceneRenderer::Render(FRDGBuilder& GraphBuilder)
 		}
 	}
 
-	FRDGExternalAccessQueue ExternalAccessQueue;
-
 	// Generate the Sky/Atmosphere look up tables
 	const bool bShouldRenderSkyAtmosphere = ShouldRenderSkyAtmosphere(Scene, ViewFamily.EngineShowFlags);
 	if (bShouldRenderSkyAtmosphere)
diff --git a/Engine/Source/Runtime/Renderer/Private/SceneRendering.cpp b/Engine/Source/Runtime/Renderer/Private/SceneRendering.cpp
--- a/Engine/Source/Runtime/Renderer/Private/SceneRendering.cpp
+++ b/Engine/Source/Runtime/Renderer/Private/SceneRendering.cpp
@@ -1976,7 +1976,10 @@ void FViewInfo::InitRHIResources(uint32 OverrideNumMSAASamples)
 
 	check(IsInRenderingThread());
 
-	CachedViewUniformShaderParameters = MakeUnique<FViewUniformShaderParameters>();
+	if (!CachedViewUniformShaderParameters)
+	{
+		CachedViewUniformShaderParameters = MakeUnique<FViewUniformShaderParameters>();
+	}
 
 	SetupUniformBufferParameters(
 		VolumeBounds,
diff --git a/Engine/Source/Runtime/Renderer/Private/SceneRendering.h b/Engine/Source/Runtime/Renderer/Private/SceneRendering.h
--- a/Engine/Source/Runtime/Renderer/Private/SceneRendering.h
+++ b/Engine/Source/Runtime/Renderer/Private/SceneRendering.h
@@ -2222,19 +2222,19 @@ public:
 	static void SetStereoViewport(FRHICommandList& RHICmdList, const FViewInfo& View, float ViewportScale = 1.0f);
 
 	/** Whether distance field global data structures should be prepared for features that use it. */
-	bool ShouldPrepareForDistanceFieldShadows(FDynamicShadowsTaskData* TaskData) const;
+	bool ShouldPrepareForDistanceFieldShadows() const;
 	bool ShouldPrepareForDistanceFieldAO() const;
 	bool ShouldPrepareForDFInsetIndirectShadow() const;
 	FGPUSceneDynamicContext& GetGPUSceneDynamicContext() { return GPUSceneDynamicContext; }
 
-	bool ShouldPrepareDistanceFieldScene(FDynamicShadowsTaskData* TaskData) const;
+	bool ShouldPrepareDistanceFieldScene() const;
 	bool ShouldPrepareGlobalDistanceField() const;
 	bool ShouldPrepareHeightFieldScene() const;
 
 	void UpdateGlobalDistanceFieldObjectBuffers(FRDGBuilder& GraphBuilder);
 	void UpdateGlobalHeightFieldObjectBuffers(FRDGBuilder& GraphBuilder, const TArray<uint32>& IndicesToUpdateInHeightFieldObjectBuffers);
 	void ProcessPendingHeightFieldPrimitiveAddAndRemoveOps(TArray<uint32>& IndicesToUpdateInHeightFieldObjectBuffers);
-	void PrepareDistanceFieldScene(FRDGBuilder& GraphBuilder, FDynamicShadowsTaskData* TaskData, FRDGExternalAccessQueue& ExternalAccessQueue);
+	void PrepareDistanceFieldScene(FRDGBuilder& GraphBuilder, FRDGExternalAccessQueue& ExternalAccessQueue);
 
 	void DrawGPUSkinCacheVisualizationInfoText();
 
