Moved primitive culling to an async task that overlaps with other render thread work.

diff --git a/Engine/Source/Runtime/Renderer/Private/DeferredShadingRenderer.cpp b/Engine/Source/Runtime/Renderer/Private/DeferredShadingRenderer.cpp
--- a/Engine/Source/Runtime/Renderer/Private/DeferredShadingRenderer.cpp
+++ b/Engine/Source/Runtime/Renderer/Private/DeferredShadingRenderer.cpp
@@ -2361,6 +2361,8 @@ void FDeferredShadingSceneRenderer::PreGatherDynamicMeshElements(FRDGBuilder& Gr
 		const int32 ReferenceViewIndex = 0;
 		FViewInfo& ReferenceView = Views[ReferenceViewIndex];
 
+		Scene->WaitForCacheRayTracingPrimitivesTask();
+
 		TaskDatas.RayTracingRelevantPrimitives = Allocator.Create<FRayTracingRelevantPrimitiveTaskData>();
 		TaskDatas.RayTracingRelevantPrimitives->Task = FFunctionGraphTask::CreateAndDispatchWhenReady(
 			[Scene = Scene, &ReferenceView, &RayTracingRelevantPrimitiveList = TaskDatas.RayTracingRelevantPrimitives->List]()
@@ -2502,23 +2504,9 @@ void FDeferredShadingSceneRenderer::Render(FRDGBuilder& GraphBuilder)
 
 	// Compute & commit the final state of the entire dependency topology of the renderer.
 	CommitFinalPipelineState();
-
-	ShaderPrint::BeginViews(GraphBuilder, Views);
-
-	ON_SCOPE_EXIT
-	{
-		ShaderPrint::EndViews(Views);
-	};
-
-	for (int32 ViewIndex = 0; ViewIndex < Views.Num(); ViewIndex++)
-	{
-		FViewInfo& View = Views[ViewIndex];
-		RDG_GPU_MASK_SCOPE(GraphBuilder, View.GPUMask);
-
-		ShadingEnergyConservation::Init(GraphBuilder, View);
-	}
-
+	
 	FRDGExternalAccessQueue ExternalAccessQueue;
+	FVisibilityTaskData* VisibilityTaskData = nullptr;
 
 	{
 		EUpdateAllPrimitiveSceneInfosAsyncOps AsyncOps = EUpdateAllPrimitiveSceneInfosAsyncOps::None;
@@ -2533,10 +2521,26 @@ void FDeferredShadingSceneRenderer::Render(FRDGBuilder& GraphBuilder)
 			AsyncOps |= EUpdateAllPrimitiveSceneInfosAsyncOps::CacheMeshDrawCommands;
 		}
 
-		UpdateScene(GraphBuilder, AsyncOps);
+		VisibilityTaskData = UpdateScene(GraphBuilder, AsyncOps);
 
 		PrepareDistanceFieldScene(GraphBuilder, ExternalAccessQueue);
 	}
+
+	ShaderPrint::BeginViews(GraphBuilder, Views);
+
+	ON_SCOPE_EXIT
+	{
+		ShaderPrint::EndViews(Views);
+	};
+
+	for (int32 ViewIndex = 0; ViewIndex < Views.Num(); ViewIndex++)
+	{
+		FViewInfo& View = Views[ViewIndex];
+		RDG_GPU_MASK_SCOPE(GraphBuilder, View.GPUMask);
+
+		ShadingEnergyConservation::Init(GraphBuilder, View);
+	}
+
 	// kick off dependent scene updates 
 	Scene->ShadowScene->UpdateForRenderedFrame(GraphBuilder);
 
@@ -2562,7 +2566,7 @@ void FDeferredShadingSceneRenderer::Render(FRDGBuilder& GraphBuilder)
 
 		if (CurrentMode != Scene->CachedRayTracingMeshCommandsMode || bNaniteCoarseMeshStreamingModeChanged || bNaniteRayTracingModeChanged || bHasRayTracingEnableChanged)
 		{
-			Scene->WaitForCacheMeshDrawCommandsTask();
+			Scene->WaitForCacheRayTracingPrimitivesTask();
 
 			// In some situations, we need to refresh the cached ray tracing mesh commands because they contain data about the currently bound shader. 
 			// This operation is a bit expensive but only happens once as we transition between modes which should be rare.
@@ -2573,7 +2577,7 @@ void FDeferredShadingSceneRenderer::Render(FRDGBuilder& GraphBuilder)
 
 		if (bRefreshRayTracingInstances)
 		{
-			Scene->WaitForCacheMeshDrawCommandsTask();
+			Scene->WaitForCacheRayTracingPrimitivesTask();
 
 			// In some situations, we need to refresh the cached ray tracing instance.
 			// eg: Need to update PrimitiveRayTracingFlags
@@ -2628,9 +2632,6 @@ void FDeferredShadingSceneRenderer::Render(FRDGBuilder& GraphBuilder)
 
 	CSV_SCOPED_TIMING_STAT_EXCLUSIVE(RenderOther);
 
-	// Setups the final FViewInfo::ViewRect.
-	PrepareViewRectsForRendering(GraphBuilder.RHICmdList);
-
 	const bool bPathTracedAtmosphere = ViewFamily.EngineShowFlags.PathTracing && Views.Num() > 0 && Views[0].FinalPostProcessSettings.PathTracingEnableReferenceAtmosphere;
 	if (ShouldRenderSkyAtmosphere(Scene, ViewFamily.EngineShowFlags) && !bPathTracedAtmosphere)
 	{
@@ -2681,10 +2682,7 @@ void FDeferredShadingSceneRenderer::Render(FRDGBuilder& GraphBuilder)
 		IESAtlas::UpdateAtlasTexture(GraphBuilder, FeatureLevel);
 	}
 
-	InitializeSceneTexturesConfig(ViewFamily.SceneTexturesConfig, ViewFamily);
 	FSceneTexturesConfig& SceneTexturesConfig = GetActiveSceneTexturesConfig();
-	FSceneTexturesConfig::Set(SceneTexturesConfig);
-
 	const FRDGSystemTextures& SystemTextures = FRDGSystemTextures::Create(GraphBuilder);
 
 	const bool bHasRayTracedOverlay = HasRayTracedOverlay(ViewFamily);
@@ -2732,7 +2730,7 @@ void FDeferredShadingSceneRenderer::Render(FRDGBuilder& GraphBuilder)
 
 	FLumenSceneFrameTemporaries LumenFrameTemporaries;
 
-	FInitViewTaskDatas InitViewTaskDatas;
+	FInitViewTaskDatas InitViewTaskDatas(VisibilityTaskData);
 	InitViewTaskDatas.LumenFrameTemporaries = &LumenFrameTemporaries;
 
 	::Strata::PreInitViews(*Scene);
diff --git a/Engine/Source/Runtime/Renderer/Private/DeferredShadingRenderer.h b/Engine/Source/Runtime/Renderer/Private/DeferredShadingRenderer.h
--- a/Engine/Source/Runtime/Renderer/Private/DeferredShadingRenderer.h
+++ b/Engine/Source/Runtime/Renderer/Private/DeferredShadingRenderer.h
@@ -60,6 +60,8 @@ struct FDBufferTextures;
 struct FILCUpdatePrimTaskData;
 struct FLumenDirectLightingTaskData;
 
+struct FVisibilityTaskData;
+
 #if RHI_RAYTRACING
 struct FRayTracingRelevantPrimitiveTaskData;
 #endif
@@ -535,6 +537,11 @@ private:
 
 	struct FInitViewTaskDatas
 	{
+		FInitViewTaskDatas(FVisibilityTaskData* InVisibilityTaskData)
+			: VisibilityTaskData(InVisibilityTaskData)
+		{}
+
+		FVisibilityTaskData* VisibilityTaskData;
 		FILCUpdatePrimTaskData* ILCUpdatePrim = nullptr;
 	#if RHI_RAYTRACING
 		FRayTracingRelevantPrimitiveTaskData* RayTracingRelevantPrimitives = nullptr;
@@ -546,7 +553,7 @@ private:
 
 	void PreGatherDynamicMeshElements(FRDGBuilder& GraphBuilder, FInitViewTaskDatas& TaskDatas);
 
-	void PreVisibilityFrameSetup(FRDGBuilder& GraphBuilder, const FSceneTexturesConfig& SceneTexturesConfig);
+	void PreVisibilityFrameSetup(FRDGBuilder& GraphBuilder);
 
 	void FinishInitDynamicShadows(FRDGBuilder& GraphBuilder, FDynamicShadowsTaskData*& TaskData, FInstanceCullingManager& InstanceCullingManager, FRDGExternalAccessQueue& ExternalAccessQueue);
 
diff --git a/Engine/Source/Runtime/Renderer/Private/MobileShadingRenderer.cpp b/Engine/Source/Runtime/Renderer/Private/MobileShadingRenderer.cpp
--- a/Engine/Source/Runtime/Renderer/Private/MobileShadingRenderer.cpp
+++ b/Engine/Source/Runtime/Renderer/Private/MobileShadingRenderer.cpp
@@ -381,7 +381,7 @@ void FMobileSceneRenderer::PrepareViewVisibilityLists()
 	}
 }
 
-void FMobileSceneRenderer::SetupMobileBasePassAfterShadowInit(FExclusiveDepthStencil::Type BasePassDepthStencilAccess, FViewVisibleCommandsPerView& ViewCommandsPerView, FInstanceCullingManager& InstanceCullingManager)
+void FMobileSceneRenderer::SetupMobileBasePassAfterShadowInit(FExclusiveDepthStencil::Type BasePassDepthStencilAccess, TArrayView<FViewCommands> ViewCommandsPerView, FInstanceCullingManager& InstanceCullingManager)
 {
 	// Sort front to back on all platforms, even HSR benefits from it
 	//const bool bWantsFrontToBackSorting = (GHardwareHiddenSurfaceRemoval == false);
@@ -459,12 +459,10 @@ void FMobileSceneRenderer::InitViews(
 
 	FILCUpdatePrimTaskData* ILCTaskData = nullptr;
 	FComputeViewVisibilityCallbacks ComputeViewVisibilityCallbacks;
-	FViewVisibleCommandsPerView ViewCommandsPerView;
-	ViewCommandsPerView.SetNum(Views.Num());
 
 	const FExclusiveDepthStencil::Type BasePassDepthStencilAccess = FExclusiveDepthStencil::DepthWrite_StencilWrite;
 
-	PreVisibilityFrameSetup(GraphBuilder, SceneTexturesConfig);
+	PreVisibilityFrameSetup(GraphBuilder);
 	if (FXSystem && FXSystem->RequiresEarlyViewUniformBuffer() && Views.IsValidIndex(0))
 	{
 		// This is to init the ViewUniformBuffer before rendering for the Niagara compute shader.
@@ -472,7 +470,7 @@ void FMobileSceneRenderer::InitViews(
 		Views[0].InitRHIResources();
 		FXSystem->PostInitViews(GraphBuilder, GetSceneViews(), !ViewFamily.EngineShowFlags.HitProxies);
 	}
-	ComputeViewVisibility(RHICmdList, BasePassDepthStencilAccess, ViewCommandsPerView, DynamicIndexBuffer, DynamicVertexBuffer, DynamicReadBuffer, InstanceCullingManager, VirtualTextureUpdater, ComputeViewVisibilityCallbacks);
+	ComputeViewVisibility(RHICmdList, TaskDatas.VisibilityTaskData, BasePassDepthStencilAccess, DynamicIndexBuffer, DynamicVertexBuffer, DynamicReadBuffer, InstanceCullingManager, VirtualTextureUpdater, ComputeViewVisibilityCallbacks);
 	PostVisibilityFrameSetup(ILCTaskData);
 
 	FIntPoint RenderTargetSize = ViewFamily.RenderTarget->GetSizeXY();
@@ -679,7 +677,7 @@ void FMobileSceneRenderer::InitViews(
 		PrepareViewVisibilityLists();
 	}
 
-	SetupMobileBasePassAfterShadowInit(BasePassDepthStencilAccess, ViewCommandsPerView, InstanceCullingManager);
+	SetupMobileBasePassAfterShadowInit(BasePassDepthStencilAccess, GetViewCommandsPerView(TaskDatas.VisibilityTaskData), InstanceCullingManager);
 
 	// if we kicked off ILC update via task, wait and finalize.
 	if (ILCTaskData)
@@ -835,7 +833,7 @@ void FMobileSceneRenderer::Render(FRDGBuilder& GraphBuilder)
 	RDG_RHI_EVENT_SCOPE(GraphBuilder, MobileSceneRender);
 	RDG_RHI_GPU_STAT_SCOPE(GraphBuilder, MobileSceneRender);
 
-	UpdateScene(GraphBuilder);
+	FVisibilityTaskData* VisibilityTaskData = UpdateScene(GraphBuilder);
 
 	FRDGExternalAccessQueue ExternalAccessQueue;
 
@@ -851,8 +849,6 @@ void FMobileSceneRenderer::Render(FRDGBuilder& GraphBuilder)
 	// Establish scene primitive count (must be done after UpdateAllPrimitiveSceneInfos)
 	FGPUSceneScopeBeginEndHelper GPUSceneScopeBeginEndHelper(Scene->GPUScene, GPUSceneDynamicContext, Scene);
 
-	PrepareViewRectsForRendering(GraphBuilder.RHICmdList);
-
 	if (ShouldRenderSkyAtmosphere(Scene, ViewFamily.EngineShowFlags))
 	{
 		for (int32 LightIndex = 0; LightIndex < NUM_ATMOSPHERE_LIGHTS; ++LightIndex)
@@ -878,9 +874,7 @@ void FMobileSceneRenderer::Render(FRDGBuilder& GraphBuilder)
 
 	WaitOcclusionTests(GraphBuilder.RHICmdList);
 
-	InitializeSceneTexturesConfig(ViewFamily.SceneTexturesConfig, ViewFamily);
 	FSceneTexturesConfig& SceneTexturesConfig = GetActiveSceneTexturesConfig();
-	FSceneTexturesConfig::Set(SceneTexturesConfig);
 
 	// Initialize global system textures (pass-through if already initialized).
 	GSystemTextures.InitializeTextures(GraphBuilder.RHICmdList, FeatureLevel);
@@ -924,7 +918,7 @@ void FMobileSceneRenderer::Render(FRDGBuilder& GraphBuilder)
 	// TODO: This doesn't take into account the potential for split screen views with separate shadow caches
 	VirtualShadowMapArray.Initialize(GraphBuilder, Scene->GetVirtualShadowMapCache(Views[0]), UseVirtualShadowMaps(ShaderPlatform, FeatureLevel), Views[0].bIsSceneCapture);
 
-	FInitViewTaskDatas InitViewTaskDatas;
+	FInitViewTaskDatas InitViewTaskDatas(VisibilityTaskData);
 
 	GraphBuilder.SetCommandListStat(GET_STATID(STAT_CLMM_InitViews));
 
diff --git a/Engine/Source/Runtime/Renderer/Private/RendererScene.cpp b/Engine/Source/Runtime/Renderer/Private/RendererScene.cpp
--- a/Engine/Source/Runtime/Renderer/Private/RendererScene.cpp
+++ b/Engine/Source/Runtime/Renderer/Private/RendererScene.cpp
@@ -6411,18 +6411,12 @@ void FScene::UpdateAllPrimitiveSceneInfos(FRDGBuilder& GraphBuilder, EUpdateAllP
 
 		CacheMeshDrawCommandsTask = LaunchSceneRenderTask(TEXT("CacheMeshDrawCommandsTask"), [this, &SceneInfosWithStaticDrawListUpdate, bLaunchAsyncTask]()
 		{
-			SCOPED_NAMED_EVENT(CacheMeshDrawCommands, FColor::Emerald);
-
-			UE::Tasks::AddNested(LaunchSceneRenderTask(TEXT("CacheMeshDrawCommandsTask"), [this, &SceneInfosWithStaticDrawListUpdate, bLaunchAsyncTask]()
-			{
-				FPrimitiveSceneInfo::CacheMeshDrawCommands(this, SceneInfosWithStaticDrawListUpdate);
-			}, bLaunchAsyncTask));
-
-			UE::Tasks::AddNested(LaunchSceneRenderTask(TEXT("CacheNaniteDrawCommands"), [this, &SceneInfosWithStaticDrawListUpdate, bLaunchAsyncTask]()
-			{
-				FPrimitiveSceneInfo::CacheNaniteDrawCommands(this, SceneInfosWithStaticDrawListUpdate);
-			}, bLaunchAsyncTask));
+			FPrimitiveSceneInfo::CacheMeshDrawCommands(this, SceneInfosWithStaticDrawListUpdate);
+		}, bLaunchAsyncTask);
 
+		CacheNaniteDrawCommandsTask = LaunchSceneRenderTask(TEXT("CacheNaniteDrawCommands"), [this, &SceneInfosWithStaticDrawListUpdate, bLaunchAsyncTask]()
+		{
+			FPrimitiveSceneInfo::CacheNaniteDrawCommands(this, SceneInfosWithStaticDrawListUpdate);
 		}, bLaunchAsyncTask);
 
 #if RHI_RAYTRACING
diff --git a/Engine/Source/Runtime/Renderer/Private/SceneHitProxyRendering.cpp b/Engine/Source/Runtime/Renderer/Private/SceneHitProxyRendering.cpp
--- a/Engine/Source/Runtime/Renderer/Private/SceneHitProxyRendering.cpp
+++ b/Engine/Source/Runtime/Renderer/Private/SceneHitProxyRendering.cpp
@@ -565,26 +565,21 @@ static void DoRenderHitProxies(
 
 void FMobileSceneRenderer::RenderHitProxies(FRDGBuilder& GraphBuilder)
 {
-	Scene->UpdateAllPrimitiveSceneInfos(GraphBuilder);
+	FVisibilityTaskData* VisibilityTaskData = UpdateScene(GraphBuilder);
 
 	GPU_MESSAGE_SCOPE(GraphBuilder);
 
 	FGPUSceneScopeBeginEndHelper GPUSceneScopeBeginEndHelper(Scene->GPUScene, GPUSceneDynamicContext, Scene);
 
-	PrepareViewRectsForRendering(GraphBuilder.RHICmdList);
-
 #if WITH_EDITOR
-	InitializeSceneTexturesConfig(ViewFamily.SceneTexturesConfig, ViewFamily);
 	FSceneTexturesConfig& SceneTexturesConfig = GetActiveSceneTexturesConfig();
-	FSceneTexturesConfig::Set(SceneTexturesConfig);
-
 	FRDGTextureRef HitProxyTexture = nullptr;
 	FRDGTextureRef HitProxyDepthTexture = nullptr;
 	InitHitProxyRender(GraphBuilder, this, HitProxyTexture, HitProxyDepthTexture);
 
 	FInstanceCullingManager& InstanceCullingManager = *GraphBuilder.AllocObject<FInstanceCullingManager>(GetSceneUniforms(), Scene->GPUScene.IsEnabled(), GraphBuilder);
 
-	FInitViewTaskDatas InitViewTaskDatas;
+	FInitViewTaskDatas InitViewTaskDatas(VisibilityTaskData);
 
 	// Find the visible primitives.
 	InitViews(GraphBuilder, SceneTexturesConfig, InstanceCullingManager, nullptr, InitViewTaskDatas);
@@ -611,19 +606,15 @@ void FDeferredShadingSceneRenderer::RenderHitProxies(FRDGBuilder& GraphBuilder)
 
 	CommitFinalPipelineState();
 
-	Scene->UpdateAllPrimitiveSceneInfos(GraphBuilder);
+	FVisibilityTaskData* VisibilityTaskData = UpdateScene(GraphBuilder);
 
 	GPU_MESSAGE_SCOPE(GraphBuilder);
 
 	FGPUSceneScopeBeginEndHelper GPUSceneScopeBeginEndHelper(Scene->GPUScene, GPUSceneDynamicContext, Scene);
 
-	PrepareViewRectsForRendering(GraphBuilder.RHICmdList);
 
 #if WITH_EDITOR
-	InitializeSceneTexturesConfig(ViewFamily.SceneTexturesConfig, ViewFamily);
 	FSceneTexturesConfig& SceneTexturesConfig = GetActiveSceneTexturesConfig();
-	FSceneTexturesConfig::Set(SceneTexturesConfig);
-
 	FRDGTextureRef HitProxyTexture = nullptr;
 	FRDGTextureRef HitProxyDepthTexture = nullptr;
 
@@ -635,7 +626,7 @@ void FDeferredShadingSceneRenderer::RenderHitProxies(FRDGBuilder& GraphBuilder)
 
 	// Find the visible primitives.
 	FLumenSceneFrameTemporaries LumenFrameTemporaries;
-	FInitViewTaskDatas InitViewTaskDatas;
+	FInitViewTaskDatas InitViewTaskDatas(VisibilityTaskData);
 	FRDGExternalAccessQueue ExternalAccessQueue;
 	BeginInitViews(GraphBuilder, SceneTexturesConfig, FExclusiveDepthStencil::DepthWrite_StencilWrite, InstanceCullingManager, nullptr, InitViewTaskDatas);
 
diff --git a/Engine/Source/Runtime/Renderer/Private/ScenePrivate.h b/Engine/Source/Runtime/Renderer/Private/ScenePrivate.h
--- a/Engine/Source/Runtime/Renderer/Private/ScenePrivate.h
+++ b/Engine/Source/Runtime/Renderer/Private/ScenePrivate.h
@@ -3477,10 +3477,16 @@ public:
 	void WaitForCacheMeshDrawCommandsTask()
 	{
 		CacheMeshDrawCommandsTask.Wait();
+	}
 
-#if RHI_RAYTRACING
-		CacheRayTracingPrimitivesTask.Wait();
-#endif
+	UE::Tasks::FTask GetCacheMeshDrawCommandsTask() const
+	{
+		return CacheMeshDrawCommandsTask;
+	}
+
+	void WaitForCacheNaniteDrawCommandsTask()
+	{
+		CacheNaniteDrawCommandsTask.Wait();
 	}
 
 #if RHI_RAYTRACING
@@ -3488,6 +3494,11 @@ public:
 	{
 		CacheRayTracingPrimitivesTask.Wait();
 	}
+
+	UE::Tasks::FTask GetCacheRayTracingPrimitivesTask()
+	{
+		return CacheRayTracingPrimitivesTask;
+	}
 #endif
 
 	void LumenAddPrimitive(FPrimitiveSceneInfo* InPrimitive);
@@ -3665,6 +3676,7 @@ private:
 
 	UE::Tasks::FTask CreateLightPrimitiveInteractionsTask;
 	UE::Tasks::FTask CacheMeshDrawCommandsTask;
+	UE::Tasks::FTask CacheNaniteDrawCommandsTask;
 #if RHI_RAYTRACING
 	UE::Tasks::FTask CacheRayTracingPrimitivesTask;
 #endif
diff --git a/Engine/Source/Runtime/Renderer/Private/SceneRendering.cpp b/Engine/Source/Runtime/Renderer/Private/SceneRendering.cpp
--- a/Engine/Source/Runtime/Renderer/Private/SceneRendering.cpp
+++ b/Engine/Source/Runtime/Renderer/Private/SceneRendering.cpp
@@ -2718,7 +2718,7 @@ FIntPoint FSceneRenderer::GetDesiredInternalBufferSize(const FSceneViewFamily& V
 	return DesiredBufferSize;
 }
 
-void FSceneRenderer::UpdateScene(FRDGBuilder& GraphBuilder, EUpdateAllPrimitiveSceneInfosAsyncOps AsyncOps)
+FVisibilityTaskData* FSceneRenderer::UpdateScene(FRDGBuilder& GraphBuilder, EUpdateAllPrimitiveSceneInfosAsyncOps AsyncOps)
 {
 	/**
 	  * UpdateStaticMeshes removes and re-creates cached FMeshDrawCommands.  If there are multiple scene renderers being run together,
@@ -2742,6 +2742,15 @@ void FSceneRenderer::UpdateScene(FRDGBuilder& GraphBuilder, EUpdateAllPrimitiveS
 	{
 		FGraphicsMinimalPipelineStateId::FreezeIdTable(false);
 	}
+
+	// Setups the final FViewInfo::ViewRect.
+	PrepareViewRectsForRendering(GraphBuilder.RHICmdList);
+
+	InitializeSceneTexturesConfig(ViewFamily.SceneTexturesConfig, ViewFamily);
+	FSceneTexturesConfig& SceneTexturesConfig = GetActiveSceneTexturesConfig();
+	FSceneTexturesConfig::Set(SceneTexturesConfig);
+
+	return BeginInitVisibility(GraphBuilder, SceneTexturesConfig);
 }
 
 void FSceneRenderer::PrepareViewRectsForRendering(FRHICommandListImmediate& RHICmdList)
diff --git a/Engine/Source/Runtime/Renderer/Private/SceneRendering.h b/Engine/Source/Runtime/Renderer/Private/SceneRendering.h
--- a/Engine/Source/Runtime/Renderer/Private/SceneRendering.h
+++ b/Engine/Source/Runtime/Renderer/Private/SceneRendering.h
@@ -83,6 +83,7 @@ class FTexture2DResource;
 class FSimpleLightArray;
 struct FNaniteMaterialPassCommand;
 struct FScreenMessageWriter;
+struct FVisibilityTaskData;
 
 DECLARE_GPU_DRAWCALL_STAT_EXTERN(VirtualTextureUpdate);
 
@@ -2291,7 +2292,7 @@ protected:
 
 	// Shared functionality between all scene renderers
 
-	void UpdateScene(FRDGBuilder& GraphBuilder, EUpdateAllPrimitiveSceneInfosAsyncOps AsyncOps = EUpdateAllPrimitiveSceneInfosAsyncOps::None);
+	FVisibilityTaskData* UpdateScene(FRDGBuilder& GraphBuilder, EUpdateAllPrimitiveSceneInfosAsyncOps AsyncOps = EUpdateAllPrimitiveSceneInfosAsyncOps::None);
 
 	FDynamicShadowsTaskData* BeginInitDynamicShadows(bool bRunningEarly);
 	void FinishInitDynamicShadows(FRDGBuilder& GraphBuilder, FDynamicShadowsTaskData* TaskData, FGlobalDynamicIndexBuffer& DynamicIndexBuffer, FGlobalDynamicVertexBuffer& DynamicVertexBuffer, FGlobalDynamicReadBuffer& DynamicReadBuffer, FInstanceCullingManager& InstanceCullingManager, FRDGExternalAccessQueue& ExternalAccessQueue);
@@ -2413,7 +2414,7 @@ protected:
 	void GatherShadowDynamicMeshElements(FGlobalDynamicIndexBuffer& DynamicIndexBuffer, FGlobalDynamicVertexBuffer& DynamicVertexBuffer, FGlobalDynamicReadBuffer& DynamicReadBuffer, FInstanceCullingManager& InstanceCullingManager);
 
 	/** Performs once per frame setup prior to visibility determination. */
-	void PreVisibilityFrameSetup(FRDGBuilder& GraphBuilder, const FSceneTexturesConfig& SceneTexturesConfig);
+	void PreVisibilityFrameSetup(FRDGBuilder& GraphBuilder);
 
 	struct FComputeViewVisibilityCallbacks
 	{
@@ -2423,8 +2424,8 @@ protected:
 	/** Computes which primitives are visible and relevant for each view. */
 	void ComputeViewVisibility(
 		FRHICommandListImmediate& RHICmdList,
+		FVisibilityTaskData* TaskData,
 		FExclusiveDepthStencil::Type BasePassDepthStencilAccess,
-		FViewVisibleCommandsPerView& ViewCommandsPerView, 
 		FGlobalDynamicIndexBuffer& DynamicIndexBuffer,
 		FGlobalDynamicVertexBuffer& DynamicVertexBuffer,
 		FGlobalDynamicReadBuffer& DynamicReadBuffer,
@@ -2542,6 +2543,10 @@ protected:
 private:
 	void ComputeFamilySize();
 
+	FVisibilityTaskData* BeginInitVisibility(FRDGBuilder& GraphBuilder, const FSceneTexturesConfig& SceneTexturesConfig);
+
+	void PrepareViewStateForVisibility(const FSceneTexturesConfig& SceneTexturesConfig);
+
 #if !UE_BUILD_SHIPPING
 	/** Collect the draw data of all visible UPrimitiveComponents in the Scene */
 	void ProcessPrimitives(const FViewInfo& View, const FViewCommands& ViewCommands) const;
@@ -2553,17 +2558,35 @@ private:
 	friend class FRendererModule;
 };
 
+template <typename LambdaType, typename PrerequisiteTaskCollectionType>
+UE::Tasks::FTask LaunchSceneRenderTask(const TCHAR* DebugName, LambdaType&& Lambda, PrerequisiteTaskCollectionType&& Prerequisites, bool bExecuteInParallelCondition = true, UE::Tasks::ETaskPriority TaskPriority = UE::Tasks::ETaskPriority::High)
+{
+	const bool bExecuteInParallel = bExecuteInParallelCondition && FApp::ShouldUseThreadingForPerformance() && GIsThreadedRendering;
+
+	return UE::Tasks::Launch(DebugName, [Lambda = Forward<LambdaType&&>(Lambda)] () mutable
+		{
+			FOptionalTaskTagScope Scope(ETaskTag::EParallelRenderingThread);
+			Lambda();
+		}, 
+		Forward<PrerequisiteTaskCollectionType&&>(Prerequisites),
+		TaskPriority,
+		bExecuteInParallel ? UE::Tasks::EExtendedTaskPriority::None : UE::Tasks::EExtendedTaskPriority::Inline
+	);
+}
+
 template <typename LambdaType>
-UE::Tasks::FTask LaunchSceneRenderTask(const TCHAR* DebugName, LambdaType&& Lambda, bool bExecuteInParallelCondition = true, LowLevelTasks::ETaskPriority TaskPriority = LowLevelTasks::ETaskPriority::Normal)
+UE::Tasks::FTask LaunchSceneRenderTask(const TCHAR* DebugName, LambdaType&& Lambda, bool bExecuteInParallelCondition = true, UE::Tasks::ETaskPriority TaskPriority = UE::Tasks::ETaskPriority::High)
 {
 	const bool bExecuteInParallel = bExecuteInParallelCondition && FApp::ShouldUseThreadingForPerformance() && GIsThreadedRendering;
 
-	return UE::Tasks::Launch(DebugName, [Lambda = MoveTemp(Lambda)]
-	{
-		FOptionalTaskTagScope Scope(ETaskTag::EParallelRenderingThread);
-		Lambda();
-	},
-	TaskPriority, bExecuteInParallel ? UE::Tasks::EExtendedTaskPriority::None : UE::Tasks::EExtendedTaskPriority::Inline);
+	return UE::Tasks::Launch(DebugName, [Lambda = Forward<LambdaType&&>(Lambda)] () mutable
+		{
+			FOptionalTaskTagScope Scope(ETaskTag::EParallelRenderingThread);
+			Lambda();
+		},
+		TaskPriority,
+		bExecuteInParallel ? UE::Tasks::EExtendedTaskPriority::None : UE::Tasks::EExtendedTaskPriority::Inline
+	);
 }
 
 struct FForwardScreenSpaceShadowMaskTextureMobileOutputs
@@ -2615,6 +2638,11 @@ protected:
 
 	struct FInitViewTaskDatas
 	{
+		FInitViewTaskDatas(FVisibilityTaskData* InVisibilityTaskData)
+			: VisibilityTaskData(InVisibilityTaskData)
+		{}
+
+		FVisibilityTaskData* VisibilityTaskData;
 		FDynamicShadowsTaskData* DynamicShadows = nullptr;
 	};
 
@@ -2670,8 +2698,7 @@ protected:
 	/** On chip pre-tonemap before scene color MSAA resolve (iOS only) */
 	void PreTonemapMSAA(FRHICommandListImmediate& RHICmdList, const FMinimalSceneTextures& SceneTextures);
 
-	void SortMobileBasePassAfterShadowInit(FExclusiveDepthStencil::Type BasePassDepthStencilAccess, FViewVisibleCommandsPerView& ViewCommandsPerView);
-	void SetupMobileBasePassAfterShadowInit(FExclusiveDepthStencil::Type BasePassDepthStencilAccess, FViewVisibleCommandsPerView& ViewCommandsPerView, FInstanceCullingManager& InstanceCullingManager);
+	void SetupMobileBasePassAfterShadowInit(FExclusiveDepthStencil::Type BasePassDepthStencilAccess, TArrayView<FViewCommands> ViewCommandsPerView, FInstanceCullingManager& InstanceCullingManager);
 
 	void UpdateDirectionalLightUniformBuffers(FRDGBuilder& GraphBuilder, const FViewInfo& View);
 	void UpdateSkyReflectionUniformBuffer();
@@ -2807,6 +2834,9 @@ extern FFastVramConfig GFastVRamConfig;
  */
 extern bool HasRayTracedDistanceFieldShadows(const FDynamicShadowsTaskData* TaskData);
 
+/** Returns the array of view commands associated with the view visibility tasks. */
+extern TArrayView<FViewCommands> GetViewCommandsPerView(FVisibilityTaskData* TaskData);
+
 /** Returns the array of shadows with distance fields. Call only after finishing shadow initialization. */
 extern TConstArrayView<FProjectedShadowInfo*> GetProjectedDistanceFieldShadows(const FDynamicShadowsTaskData* TaskData);
 
diff --git a/Engine/Source/Runtime/Renderer/Private/SceneVisibility.cpp b/Engine/Source/Runtime/Renderer/Private/SceneVisibility.cpp
--- a/Engine/Source/Runtime/Renderer/Private/SceneVisibility.cpp
+++ b/Engine/Source/Runtime/Renderer/Private/SceneVisibility.cpp
@@ -45,6 +45,7 @@
 #include "RectLightSceneProxy.h"
 #include "Math/Halton.h"
 #include "ProfilingDebugging/DiagnosticTable.h"
+#include "ProfilingDebugging/CountersTrace.h"
 #include "Algo/Unique.h"
 #include "InstanceCulling/InstanceCullingManager.h"
 #include "PostProcess/TemporalAA.h"
@@ -342,6 +343,48 @@
 DECLARE_CYCLE_STAT(TEXT("Occlusion Readback"), STAT_CLMM_OcclusionReadback, STATGROUP_CommandListMarkers);
 DECLARE_CYCLE_STAT(TEXT("After Occlusion Readback"), STAT_CLMM_AfterOcclusionReadback, STATGROUP_CommandListMarkers);
 
+TRACE_DECLARE_INT_COUNTER(Scene_Visibility_NumProcessedPrimitives, TEXT("Scene/Visibility/NumProcessedPrimitives"));
+TRACE_DECLARE_INT_COUNTER(Scene_Visibility_NumCulledPrimitives, TEXT("Scene/Visibility/NumCulledPrimitives"));
+TRACE_DECLARE_INT_COUNTER(Scene_Visibility_NumOccludedPrimitives, TEXT("Scene/Visibility/NumOccludedPrimitives"));
+
+struct FVisibilityTaskData
+{
+	FVisibilityTaskData()
+		: TaskEvent(UE_SOURCE_LOCATION)
+		, bExecuteInParallel(CVarParallelInitViews.GetValueOnRenderThread() > 0)
+	{}
+
+	FSceneRenderingBulkObjectAllocator Allocator;
+	FViewVisibleCommandsPerView ViewCommandsPerView;
+	UE::Tasks::FTaskEvent TaskEvent;
+	const bool bExecuteInParallel;
+
+	std::atomic_int32_t NumProcessedPrimitives{ 0 };
+	std::atomic_int32_t NumCulledPrimitives{ 0 };
+	std::atomic_int32_t NumOccludedPrimitives{ 0 };
+
+	void Wait()
+	{
+		TaskEvent.Wait();
+	}
+
+	void Finish()
+	{
+		TaskEvent.Wait();
+		Allocator.BulkDelete();
+	}
+};
+
+TArrayView<FViewCommands> GetViewCommandsPerView(FVisibilityTaskData* TaskData)
+{
+	if (TaskData)
+	{
+		TaskData->Wait();
+		return TaskData->ViewCommandsPerView;
+	}
+	return {};
+}
+
 /*------------------------------------------------------------------------------
 	Visibility determination.
 ------------------------------------------------------------------------------*/
@@ -714,13 +757,13 @@
 		});
 }
 
-static void PrimitiveCullTask(FThreadSafeCounter& NumCulledPrimitives, const FScene* RESTRICT Scene, FViewInfo& View, FPrimitiveCullingFlags Flags, float MaxDrawDistanceScale, const FHLODVisibilityState* const HLODState, const FSceneBitArray& VisibleNodes, int32 TaskIndex)
+static void PrimitiveCullTask(FVisibilityTaskData& TaskData, const FScene* RESTRICT Scene, FViewInfo& View, FPrimitiveCullingFlags Flags, float MaxDrawDistanceScale, const FHLODVisibilityState* const HLODState, const FSceneBitArray& VisibleNodes, int32 TaskIndex)
 {
 	TRACE_CPUPROFILER_EVENT_SCOPE(SceneVisibility_PrimitiveCull);
 	SCOPED_NAMED_EVENT(SceneVisibility_PrimitiveCull, FColor::Red);
 	QUICK_SCOPE_CYCLE_COUNTER(STAT_PrimitiveCull_Loop);
 
-	FTaskTagScope TaskTagScope(ETaskTag::EParallelRenderingThread);
+	FOptionalTaskTagScope TaskTagScope(ETaskTag::EParallelRenderingThread);
 
 	bool bDisableLODFade = GDisableLODFade || View.bDisableDistanceBasedFadeTransitions;
 	const FPlane* PermutedPlanePtr = View.ViewFrustum.PermutedPlanes.GetData();
@@ -882,7 +925,7 @@
 			}
 			else
 			{
-				STAT(++NumPrimitivesCulledForTask);
+				++NumPrimitivesCulledForTask;
 			}
 
 #if RHI_RAYTRACING
@@ -914,10 +957,10 @@
 #endif
 	}
 
-	STAT(NumCulledPrimitives.Add(NumPrimitivesCulledForTask));
+	TaskData.NumCulledPrimitives += NumPrimitivesCulledForTask;
 }
 
-static int32 PrimitiveCull(FScene* RESTRICT Scene, FViewInfo& View, bool bShouldVisibilityCull)
+static void PrimitiveCull(FVisibilityTaskData& TaskData, FScene* RESTRICT Scene, FViewInfo& View, bool bShouldVisibilityCull)
 {
 	FPrimitiveCullingFlags Flags;
 	Flags.bShouldVisibilityCull = bShouldVisibilityCull;
@@ -929,16 +972,18 @@
 	Flags.bHasHiddenPrimitives = View.HiddenPrimitives.Num() > 0;
 	Flags.bHasShowOnlyPrimitives = View.ShowOnlyPrimitives.IsSet();
 
+	UE::Tasks::FTask PrerequisiteTask;
+
 #if RHI_RAYTRACING
 	View.RayTracingCullingParameters.Init(View);
 
 	// Sync cached raytracing tasks ShouldCullForRayTracing.
-	Scene->WaitForCacheRayTracingPrimitivesTask();
+	PrerequisiteTask = Scene->GetCacheRayTracingPrimitivesTask();
 #endif
 
 	SCOPE_CYCLE_COUNTER(STAT_PrimitiveCull);
 
-	FSceneBitArray VisibleNodes;
+	FSceneBitArray& VisibleNodes = *TaskData.Allocator.Create<FSceneBitArray>();
 
 	if (bShouldVisibilityCull && Flags.bUseVisibilityOctree)
 	{
@@ -950,7 +995,6 @@
 	//Performance varies on total primitive count and tasks scheduled. Check the mentioned link above for some measurements.
 	//There have been some changes as compared to the code measured in the link
 
-	FThreadSafeCounter NumCulledPrimitives;
 	FSceneViewState* ViewState = (FSceneViewState*)View.State;
 	const bool bHLODActive = Scene->SceneLODHierarchy.IsActive();
 	const FHLODVisibilityState* const HLODState = bHLODActive && ViewState ? &ViewState->HLODVisibilityState : nullptr;
@@ -960,16 +1004,16 @@
 	const int32 BitArrayWords = FMath::DivideAndRoundUp(BitArrayNum, (int32)NumBitsPerDWORD);
 	const int32 NumTasks = FMath::DivideAndRoundUp(BitArrayWords, FrustumCullNumWordsPerTask);
 
-	ParallelFor(NumTasks,
-		[&NumCulledPrimitives, Scene, &View, MaxDrawDistanceScale, HLODState, &VisibleNodes, &Flags](int32 TaskIndex)
+	const bool bExecuteInParallel = TaskData.bExecuteInParallel && (!Flags.bUseCustomCulling || View.CustomVisibilityQuery->IsThreadsafe());
+
+	for (int32 TaskIndex = 0; TaskIndex < NumTasks; ++TaskIndex)
+	{
+		TaskData.TaskEvent.AddPrerequisites(LaunchSceneRenderTask(UE_SOURCE_LOCATION, [&TaskData, Scene, &View, MaxDrawDistanceScale, HLODState, &VisibleNodes, Flags, TaskIndex]
 		{
-			PrimitiveCullTask(NumCulledPrimitives, Scene, View, Flags, MaxDrawDistanceScale, HLODState, VisibleNodes, TaskIndex);
-		},
-		!FApp::ShouldUseThreadingForPerformance() || (Flags.bUseCustomCulling && !View.CustomVisibilityQuery->IsThreadsafe()) || CVarParallelInitViews.GetValueOnRenderThread() == 0 || !IsInActualRenderingThread()
-		);
+			PrimitiveCullTask(TaskData, Scene, View, Flags, MaxDrawDistanceScale, HLODState, VisibleNodes, TaskIndex);
 
-
-	return NumCulledPrimitives.GetValue();
+		}, PrerequisiteTask, bExecuteInParallel));
+	}
 }
 
 /**
@@ -1673,7 +1717,7 @@
 					if (bIsOccluded)
 					{
 						View.PrimitiveVisibilityMap.AccessCorrespondingBit(BitIt) = false;
-						STAT(NumOccludedPrimitives++);
+						NumOccludedPrimitives++;
 					}
 					else if (bOcclusionStateIsDefinite)
 					{
@@ -1693,7 +1737,7 @@
 				if (bAllSubOccluded)
 				{
 					View.PrimitiveVisibilityMap.AccessCorrespondingBit(BitIt) = false;
-					STAT(NumOccludedPrimitives++);
+					NumOccludedPrimitives++;
 				}
 				else if (bAllSubOcclusionStateIsDefinite)
 				{
@@ -2037,7 +2081,7 @@
 				{
 					View.PrimitiveVisibilityMap.AccessCorrespondingBit(BitIt) = false;
 					INC_DWORD_STAT_BY(STAT_StaticallyOccludedPrimitives,1);
-					STAT(NumOccludedPrimitives++);
+					NumOccludedPrimitives++;
 
 					if (GVisualizeOccludedPrimitives)
 					{
@@ -2401,6 +2445,7 @@
 		bTranslucentSurfaceLighting = false;
 		const EShadingPath ShadingPath = Scene->GetShadingPath();
 		const bool bAddLightmapDensityCommands = View.Family->EngineShowFlags.LightMapDensity && AllowDebugViewmodes();
+		const bool bHairStrandsEnabled = IsHairStrandsEnabled(EHairStrandsShaderType::All, Scene->GetShaderPlatform());
 
 		SCOPE_CYCLE_COUNTER(STAT_ComputeViewRelevance);
 		for (int32 Index = 0; Index < Input.NumPrims; Index++)
@@ -2408,7 +2453,7 @@
 			int32 BitIndex = Input.Prims[Index];
 			FPrimitiveSceneInfo* PrimitiveSceneInfo = Scene->Primitives[BitIndex];
 			FPrimitiveViewRelevance& ViewRelevance = const_cast<FPrimitiveViewRelevance&>(View.PrimitiveViewRelevanceMap[BitIndex]);
-			ViewRelevance = PrimitiveSceneInfo->Proxy->GetViewRelevance(&View);
+			ViewRelevance = Scene->PrimitiveSceneProxies[BitIndex]->GetViewRelevance(&View);
 			ViewRelevance.bInitializedThisFrame = true;
 
 			const bool bStaticRelevance = ViewRelevance.bStaticRelevance;
@@ -2419,8 +2464,7 @@
 			const bool bEditorVisualizeLevelInstanceRelevance = ViewRelevance.bEditorVisualizeLevelInstanceRelevance;
 			const bool bEditorSelectionRelevance = ViewRelevance.bEditorStaticSelectionRelevance;
 			const bool bTranslucentRelevance = ViewRelevance.HasTranslucency();
+			const bool bHairStrandsRelevance = bHairStrandsEnabled && ViewRelevance.bHairStrands;
 
-			const bool bHairStrandsEnabled = ViewRelevance.bHairStrands && IsHairStrandsEnabled(EHairStrandsShaderType::All, Scene->GetShaderPlatform());
-
 			if (View.bIsReflectionCapture && !PrimitiveSceneInfo->Proxy->IsVisibleInReflectionCaptures())
 			{
@@ -2471,7 +2515,7 @@
 					VisibleDynamicPrimitivesWithSimpleLights.AddPrim(PrimitiveSceneInfo);
 				}
 			}
-			else if (bHairStrandsEnabled)
+			else if (bHairStrandsRelevance)
 			{
 				// Strands MeshElement
 				++NumVisibleDynamicPrimitives;
@@ -3591,7 +3635,7 @@
 		Distance.SizeSquared() > CameraTranslationThreshold * CameraTranslationThreshold;
 }
 
-void FSceneRenderer::PreVisibilityFrameSetup(FRDGBuilder& GraphBuilder, const FSceneTexturesConfig& SceneTexturesConfig)
+void FSceneRenderer::PreVisibilityFrameSetup(FRDGBuilder& GraphBuilder)
 {
 	FRHICommandListImmediate& RHICmdList = GraphBuilder.RHICmdList;
 
@@ -3659,6 +3703,38 @@
 	}
 #endif
 
+#if RHI_RAYTRACING
+	if (Scene && Views.Num())
+	{
+		const int32 ReferenceViewIndex = 0;
+		const FViewInfo& ReferenceView = Views[ReferenceViewIndex];
+
+		Scene->RayTracingScene.InitPreViewTranslation(ReferenceView.ViewMatrices);
+		Scene->RayTracingScene.bNeedsDebugInstanceGPUSceneIndexBuffer = IsRayTracingInstanceOverlapEnabled(ReferenceView);
+	}
+#endif
+
+	// Setup global dither fade in and fade out uniform buffers.
+	for (int32 ViewIndex = 0; ViewIndex < Views.Num(); ViewIndex++)
+	{
+		FViewInfo& View = Views[ViewIndex];
+
+		FDitherUniformShaderParameters DitherUniformShaderParameters;
+		DitherUniformShaderParameters.LODFactor = View.GetTemporalLODTransition();
+		View.DitherFadeOutUniformBuffer = FDitherUniformBufferRef::CreateUniformBufferImmediate(DitherUniformShaderParameters, UniformBuffer_SingleFrame);
+
+		DitherUniformShaderParameters.LODFactor = View.GetTemporalLODTransition() - 1.0f;
+		View.DitherFadeInUniformBuffer = FDitherUniformBufferRef::CreateUniformBufferImmediate(DitherUniformShaderParameters, UniformBuffer_SingleFrame);
+	}
+
+	for (const auto& ViewExtension : ViewFamily.ViewExtensions)
+	{
+		ViewExtension->PreInitViews_RenderThread(GraphBuilder);
+	}
+}
+
+void FSceneRenderer::PrepareViewStateForVisibility(const FSceneTexturesConfig& SceneTexturesConfig)
+{
 #if UE_BUILD_SHIPPING
 	const bool bFreezeTemporalHistories = false;
 	const bool bFreezeTemporalSequences = false;
@@ -4075,35 +4151,6 @@
 			View.PrevViewInfo = NewPrevViewInfo;
 		}
 	}
-
-#if RHI_RAYTRACING
-	if (Scene && Views.Num())
-	{
-		const int32 ReferenceViewIndex = 0;
-		const FViewInfo& ReferenceView = Views[ReferenceViewIndex];
-
-		Scene->RayTracingScene.InitPreViewTranslation(ReferenceView.ViewMatrices);
-		Scene->RayTracingScene.bNeedsDebugInstanceGPUSceneIndexBuffer = IsRayTracingInstanceOverlapEnabled(ReferenceView);
-	}
-#endif
-
-	// Setup global dither fade in and fade out uniform buffers.
-	for (int32 ViewIndex = 0; ViewIndex < Views.Num(); ViewIndex++)
-	{
-		FViewInfo& View = Views[ViewIndex];
-
-		FDitherUniformShaderParameters DitherUniformShaderParameters;
-		DitherUniformShaderParameters.LODFactor = View.GetTemporalLODTransition();
-		View.DitherFadeOutUniformBuffer = FDitherUniformBufferRef::CreateUniformBufferImmediate(DitherUniformShaderParameters, UniformBuffer_SingleFrame);
-
-		DitherUniformShaderParameters.LODFactor = View.GetTemporalLODTransition() - 1.0f;
-		View.DitherFadeInUniformBuffer = FDitherUniformBufferRef::CreateUniformBufferImmediate(DitherUniformShaderParameters, UniformBuffer_SingleFrame);
-	}
-
-	for (const auto& ViewExtension : ViewFamily.ViewExtensions)
-	{
-		ViewExtension->PreInitViews_RenderThread(GraphBuilder);
-	}
 }
 
 void FSceneViewState::UpdateMotionBlurTimeScale(const FViewInfo& View)
@@ -4475,55 +4522,44 @@
 
 #endif
 
-void FSceneRenderer::ComputeViewVisibility(
-	FRHICommandListImmediate& RHICmdList,
-	FExclusiveDepthStencil::Type BasePassDepthStencilAccess,
-	FViewVisibleCommandsPerView& ViewCommandsPerView, 
-	FGlobalDynamicIndexBuffer& DynamicIndexBuffer,
-	FGlobalDynamicVertexBuffer& DynamicVertexBuffer,
-	FGlobalDynamicReadBuffer& DynamicReadBuffer, 
-	FInstanceCullingManager& InstanceCullingManager,
-	FVirtualTextureUpdater* VirtualTextureUpdater,
-	const FComputeViewVisibilityCallbacks& Callbacks)
+FVisibilityTaskData* FSceneRenderer::BeginInitVisibility(FRDGBuilder& GraphBuilder, const FSceneTexturesConfig& SceneTexturesConfig)
 {
-	SCOPE_CYCLE_COUNTER(STAT_ViewVisibilityTime);
-	SCOPED_NAMED_EVENT(FSceneRenderer_ComputeViewVisibility, FColor::Magenta);
+	TRACE_CPUPROFILER_EVENT_SCOPE(FSceneRenderer_BeginInitVisibility);
 
-	STAT(int32 NumProcessedPrimitives = 0);
-	STAT(int32 NumCulledPrimitives = 0);
-	STAT(int32 NumOccludedPrimitives = 0);
+	FVisibilityTaskData* TaskData = Allocator.Create<FVisibilityTaskData>();
+	TaskData->ViewCommandsPerView.SetNum(Views.Num());
 
-	UE::Tasks::FTask ComputeLightVisibilityTask = LaunchSceneRenderTask(UE_SOURCE_LOCATION, [this]
+	TaskData->TaskEvent.AddPrerequisites(LaunchSceneRenderTask(UE_SOURCE_LOCATION, [this]
 	{
 		ComputeLightVisibility();
-	});
+	}));
 
-	int32 NumPrimitives = Scene->Primitives.Num();
-	float CurrentRealTime = ViewFamily.Time.GetRealTimeSeconds();
+	PrepareViewStateForVisibility(SceneTexturesConfig);
 
-	FPrimitiveViewMasks HasDynamicMeshElementsMasks;
-	HasDynamicMeshElementsMasks.AddZeroed(NumPrimitives);
+	for (int32 ViewIndex = 0; ViewIndex < Views.Num(); ++ViewIndex)
+	{
+		FViewInfo& View = Views[ViewIndex];
+		FSceneViewState* ViewState = (FSceneViewState*)View.State;
 
-	FPrimitiveViewMasks HasDynamicEditorMeshElementsMasks;
+		if (ViewState)
+		{
+			SCOPE_CYCLE_COUNTER(STAT_DecompressPrecomputedOcclusion);
+			View.PrecomputedVisibilityData = ViewState->GetPrecomputedVisibilityData(View, Scene);
 
-	if (GIsEditor)
-	{
-		HasDynamicEditorMeshElementsMasks.AddZeroed(NumPrimitives);
-	}
+			if (View.PrecomputedVisibilityData)
+			{
+				bUsedPrecomputedVisibility = true;
+			}
+		}
+		else
+		{
+			View.PrecomputedVisibilityData = nullptr;
+		}
 
-	uint8 ViewBit = 0x1;
-	{
-		TRACE_CPUPROFILER_EVENT_SCOPE(FSceneRenderer_Views);
-		for (int32 ViewIndex = 0; ViewIndex < Views.Num(); ++ViewIndex, ViewBit <<= 1)
+		TaskData->TaskEvent.AddPrerequisites(LaunchSceneRenderTask(UE_SOURCE_LOCATION, [this, TaskData, &View, ViewState, ViewIndex]
 		{
 			TRACE_CPUPROFILER_EVENT_SCOPE_TEXT(*FString::Printf(TEXT("View %d"), ViewIndex));
-			STAT(NumProcessedPrimitives += NumPrimitives);
+			TaskData->NumProcessedPrimitives += Scene->Primitives.Num();
 
-			FViewInfo& View = Views[ViewIndex];
-			FViewCommands& ViewCommands = ViewCommandsPerView[ViewIndex];
-			FSceneViewState* ViewState = (FSceneViewState*)View.State;
-
-			const bool bIsSinglePassStereo = View.bIsInstancedStereoEnabled || View.bIsMobileMultiViewEnabled;
-
 			// Allocate the view's visibility maps.
 			View.PrimitiveVisibilityMap.Init(false, Scene->Primitives.Num());
@@ -4547,21 +4583,6 @@
 			View.PrimitiveViewRelevanceMap.Reset(Scene->Primitives.Num());
 			View.PrimitiveViewRelevanceMap.AddZeroed(Scene->Primitives.Num());
 
-			if (ViewState)
-			{
-				SCOPE_CYCLE_COUNTER(STAT_DecompressPrecomputedOcclusion);
-				View.PrecomputedVisibilityData = ViewState->GetPrecomputedVisibilityData(View, Scene);
-			}
-			else
-			{
-				View.PrecomputedVisibilityData = nullptr;
-			}
-
-			if (View.PrecomputedVisibilityData)
-			{
-				bUsedPrecomputedVisibility = true;
-			}
-
 			bool bNeedsFrustumCulling = CVarEnableFrustumCull.GetValueOnRenderThread();
 
 			// Development builds sometimes override frustum culling, e.g. dependent views in the editor.
@@ -4585,7 +4606,7 @@
 #endif
 
 			// Most views use standard frustum culling.
-			if(bNeedsFrustumCulling)
+			if (bNeedsFrustumCulling)
 			{
 				// Update HLOD transition/visibility states to allow use during distance culling
 				FLODSceneTree& HLODTree = Scene->SceneLODHierarchy;
@@ -4602,9 +4623,55 @@
 
 			{
 				TRACE_CPUPROFILER_EVENT_SCOPE(FSceneRenderer_Cull);
-				int32 NumCulledPrimitivesForView = PrimitiveCull(Scene, View, bNeedsFrustumCulling);
-				STAT(NumCulledPrimitives += NumCulledPrimitivesForView);
+				PrimitiveCull(*TaskData, Scene, View, bNeedsFrustumCulling);
 			}
+		}));
+	}
+
+	TaskData->TaskEvent.Trigger();
+	return TaskData;
+}
+
+void FSceneRenderer::ComputeViewVisibility(
+	FRHICommandListImmediate& RHICmdList,
+	FVisibilityTaskData* TaskData,
+	FExclusiveDepthStencil::Type BasePassDepthStencilAccess,
+	FGlobalDynamicIndexBuffer& DynamicIndexBuffer,
+	FGlobalDynamicVertexBuffer& DynamicVertexBuffer,
+	FGlobalDynamicReadBuffer& DynamicReadBuffer, 
+	FInstanceCullingManager& InstanceCullingManager,
+	FVirtualTextureUpdater* VirtualTextureUpdater,
+	const FComputeViewVisibilityCallbacks& Callbacks)
+{
+	SCOPE_CYCLE_COUNTER(STAT_ViewVisibilityTime);
+	SCOPED_NAMED_EVENT(FSceneRenderer_ComputeViewVisibility, FColor::Magenta);
+
+	int32 NumPrimitives = Scene->Primitives.Num();
+
+	FPrimitiveViewMasks HasDynamicMeshElementsMasks;
+	HasDynamicMeshElementsMasks.AddZeroed(NumPrimitives);
+
+	FPrimitiveViewMasks HasDynamicEditorMeshElementsMasks;
+
+	if (GIsEditor)
+	{
+		HasDynamicEditorMeshElementsMasks.AddZeroed(NumPrimitives);
+	}
+
+	TaskData->Finish();
+
+	uint8 ViewBit = 0x1;
+	{
+		TRACE_CPUPROFILER_EVENT_SCOPE(FSceneRenderer_Views);
+		for (int32 ViewIndex = 0; ViewIndex < Views.Num(); ++ViewIndex, ViewBit <<= 1)
+		{
+			TRACE_CPUPROFILER_EVENT_SCOPE_TEXT(*FString::Printf(TEXT("View %d"), ViewIndex));
+
+			FViewInfo& View = Views[ViewIndex];
+			FViewCommands& ViewCommands = TaskData->ViewCommandsPerView[ViewIndex];
+			FSceneViewState* ViewState = (FSceneViewState*)View.State;
+
+			const bool bIsSinglePassStereo = View.bIsInstancedStereoEnabled || View.bIsMobileMultiViewEnabled;
 
 			{
 				TRACE_CPUPROFILER_EVENT_SCOPE(FSceneRenderer_UpdatePrimitiveFading);
@@ -4645,8 +4712,7 @@
 			// Occlusion cull for all primitives in the view frustum, but not in wireframe.
 			if (!View.Family->EngineShowFlags.Wireframe)
 			{
-				int32 NumOccludedPrimitivesInView = OcclusionCull(RHICmdList, Scene, View, DynamicVertexBuffer);
-				STAT(NumOccludedPrimitives += NumOccludedPrimitivesInView);
+				TaskData->NumOccludedPrimitives += OcclusionCull(RHICmdList, Scene, View, DynamicVertexBuffer);
 			}
 
 			Scene->WaitForCacheMeshDrawCommandsTask();
@@ -4693,7 +4759,7 @@
 	for (int32 ViewIndex = 0; ViewIndex < Views.Num(); ++ViewIndex, ViewBit <<= 1)
 	{
 		FViewInfo& View = Views[ViewIndex];
-		FViewCommands& ViewCommands = ViewCommandsPerView[ViewIndex];
+		FViewCommands& ViewCommands = TaskData->ViewCommandsPerView[ViewIndex];
 		
 		if (View.bIsInstancedStereoEnabled || View.bIsMobileMultiViewEnabled)
 		{
@@ -4702,7 +4768,7 @@
 		}
 	}
 
-	ComputeLightVisibilityTask.Wait();
+	Scene->WaitForCacheNaniteDrawCommandsTask();
 
 	if (Callbacks.PreGatherDynamicMeshElements)
 	{
@@ -4735,7 +4801,7 @@
 			continue;
 		}
 
-		FViewCommands& ViewCommands = ViewCommandsPerView[ViewIndex];
+		FViewCommands& ViewCommands = TaskData->ViewCommandsPerView[ViewIndex];
 
 #if !UE_BUILD_SHIPPING
 		DumpPrimitives(ViewCommands);
@@ -4745,9 +4811,13 @@
 		SetupMeshPass(View, BasePassDepthStencilAccess, ViewCommands, InstanceCullingManager);
 	}
 
-	INC_DWORD_STAT_BY(STAT_ProcessedPrimitives,NumProcessedPrimitives);
-	INC_DWORD_STAT_BY(STAT_CulledPrimitives,NumCulledPrimitives);
-	INC_DWORD_STAT_BY(STAT_OccludedPrimitives,NumOccludedPrimitives);
+	INC_DWORD_STAT_BY(STAT_ProcessedPrimitives, TaskData->NumProcessedPrimitives);
+	INC_DWORD_STAT_BY(STAT_CulledPrimitives, TaskData->NumCulledPrimitives);
+	INC_DWORD_STAT_BY(STAT_OccludedPrimitives, TaskData->NumOccludedPrimitives);
+
+	TRACE_COUNTER_SET(Scene_Visibility_NumProcessedPrimitives, TaskData->NumProcessedPrimitives);
+	TRACE_COUNTER_SET(Scene_Visibility_NumCulledPrimitives, TaskData->NumCulledPrimitives);
+	TRACE_COUNTER_SET(Scene_Visibility_NumOccludedPrimitives, TaskData->NumOccludedPrimitives);
 }
 
 void FDeferredShadingSceneRenderer::ComputeLightVisibility()
@@ -4991,7 +5061,7 @@
 /** 
 * Performs once per frame setup prior to visibility determination.
 */
-void FDeferredShadingSceneRenderer::PreVisibilityFrameSetup(FRDGBuilder& GraphBuilder, const FSceneTexturesConfig& SceneTexturesConfig)
+void FDeferredShadingSceneRenderer::PreVisibilityFrameSetup(FRDGBuilder& GraphBuilder)
 {
 	// Possible stencil dither optimization approach
 	for (int32 ViewIndex = 0; ViewIndex < Views.Num(); ViewIndex++)
@@ -5000,7 +5070,7 @@
 		View.bAllowStencilDither = DepthPass.bDitheredLODTransitionsUseStencil;
 	}
 
-	FSceneRenderer::PreVisibilityFrameSetup(GraphBuilder, SceneTexturesConfig);
+	FSceneRenderer::PreVisibilityFrameSetup(GraphBuilder);
 }
 
 /**
@@ -5019,7 +5089,7 @@
 	SCOPE_CYCLE_COUNTER(STAT_InitViewsTime);
 	RDG_CSV_STAT_EXCLUSIVE_SCOPE(GraphBuilder, InitViews_Scene);
 
-	PreVisibilityFrameSetup(GraphBuilder, SceneTexturesConfig);
+	PreVisibilityFrameSetup(GraphBuilder);
 
 	FRHICommandListImmediate& RHICmdList = GraphBuilder.RHICmdList;
 
@@ -5046,9 +5116,6 @@
 
 	LumenScenePDIVisualization();
 	
-	FViewVisibleCommandsPerView ViewCommandsPerView;
-	ViewCommandsPerView.SetNum(Views.Num());
-
 	{
 		FComputeViewVisibilityCallbacks ComputeViewVisibilityCallbacks;
 
@@ -5057,7 +5124,7 @@
 			PreGatherDynamicMeshElements(GraphBuilder, TaskDatas);
 		};
 
-		ComputeViewVisibility(RHICmdList, BasePassDepthStencilAccess, ViewCommandsPerView, DynamicIndexBufferForInitViews, DynamicVertexBufferForInitViews, DynamicReadBufferForInitViews, InstanceCullingManager, VirtualTextureUpdater, ComputeViewVisibilityCallbacks);
+		ComputeViewVisibility(RHICmdList, TaskDatas.VisibilityTaskData, BasePassDepthStencilAccess, DynamicIndexBufferForInitViews, DynamicVertexBufferForInitViews, DynamicReadBufferForInitViews, InstanceCullingManager, VirtualTextureUpdater, ComputeViewVisibilityCallbacks);
 	}
 
 	// This must happen before we start initialising and using views.
