Support OpenXR gaze-tracking in the foveated VRS generator. Foveated image is re-centered on user's gaze point if gaze data is available and confident.
Refactored some outdated portions of the foveated generator from before the VRS manager system was updated, including removing remnants of a combiner system within the shader

diff --git a/Engine/Config/BaseDeviceProfiles.ini b/Engine/Config/BaseDeviceProfiles.ini
--- a/Engine/Config/BaseDeviceProfiles.ini
+++ b/Engine/Config/BaseDeviceProfiles.ini
@@ -968,7 +968,7 @@ BaseProfileName=Android_High
 [Oculus_Quest DeviceProfile]
 DeviceType=Android
 BaseProfileName=Android_Mid
-+CVars=vr.VRS.HMDFixedFoveationDynamic=1
++CVars=r.VRS.HMDDynamicFoveation=1
 +CVars=r.Android.DisableVulkanSupport=0
 
 [Oculus_Quest2 DeviceProfile]
diff --git a/Engine/Shaders/Private/VariableRateShading/FixedFoveationVariableRateShading.usf b/Engine/Shaders/Private/VariableRateShading/FixedFoveationVariableRateShading.usf
deleted file mode 100644
diff --git a/Engine/Shaders/Private/VariableRateShading/VRSShadingRateFoveated.usf b/Engine/Shaders/Private/VariableRateShading/VRSShadingRateFoveated.usf
new file mode 100644
--- /dev/null
+++ b/Engine/Shaders/Private/VariableRateShading/VRSShadingRateFoveated.usf
@@ -0,0 +1,99 @@
+// Copyright Epic Games, Inc. All Rights Reserved.
+
+/*=============================================================================
+	FixedFoveationVariableRateShading.usf
+=============================================================================*/
+
+#include "../Common.ush"
+
+// Supplied defines:
+// SHADING_RATE_1x1			- Shading rate of 1x1
+// SHADING_RATE_1x2			- Shading rate of 1x2
+// SHADING_RATE_2x1			- Shading rate of 2x1
+// SHADING_RATE_2x2			- Shading rate of 2x2
+// SHADING_RATE_2x4			- Shading rate of 2x4
+// SHADING_RATE_4x2			- Shading rate of 4x2
+// SHADING_RATE_4x4			- Shading rate of 4x4
+
+// MAX_COMBINED_SOURCES_IN	- Maximum 	
+// SHADING_RATE_TILE_WIDTH	- Shading rate tile width
+// SHADING_RATE_TILE_HEIGHT - Shading rate tile height
+// THREADGROUP_SIZEX		- Threadgroup size (X)
+// THREADGROUP_SIZEY		- Threadgroup size (Y)
+
+// Generated output shading rate attachment.
+RWTexture2D<uint> RWOutputTexture;
+
+// Horizontal and vertical field of view for the current view.
+float2	FieldOfView;
+
+// Center pixel coordinate for the left eye.
+float2	LeftEyeCenterPixelXY;
+
+// Center pixel coordinate for the right eye.
+float2	RightEyeCenterPixelXY;
+
+// Diagonal of the view in pixels, squared.
+float	ViewDiagonalSquaredInPixels;
+// Cutoff percentage for full-rate shading, foveation (0.0 to 1.0), squared. Outside this range will be shaded at half-rate in x and y.
+float	FoveationFullRateCutoffSquared;
+// Cutoff percentage for half-rate shading, foveation (0.0 to 1.0), squared. Outside this range will be shaded at quarter-rate in x and y.
+float	FoveationHalfRateCutoffSquared;
+
+// Bitfield indicating optional supporting things (defined in the EVRSGenerationFlags enum:)
+// 	EVRSGenerationFlags::None = 0x0,
+//	EVRSGenerationFlags::StereoRendering = 0x1,
+//	EVRSGenerationFlags::SideBySideStereo = 0x2,
+uint ShadingRateAttachmentGenerationFlags;
+
+////////////////////////////////////////////////////////////////////////////////////////////////////
+// Get the fractional offset of the pixel coordinate from the center of the eye-view target 
+// 0.0 = center of view.
+// 1.0 = diagonal far edge of view.
+////////////////////////////////////////////////////////////////////////////////////////////////////
+
+float GetFractionalOffsetFromEyeOrigin(float2 PixelPosition)
+{
+	if ((ShadingRateAttachmentGenerationFlags & SIDE_BY_SIDE_STEREO) != 0)
+	{
+		const float2 VecToLeft = PixelPosition - LeftEyeCenterPixelXY;
+		const float2 VecToRight = PixelPosition - RightEyeCenterPixelXY;
+		return min(dot(VecToLeft, VecToLeft), dot(VecToRight, VecToRight)) / ViewDiagonalSquaredInPixels;
+	}
+	else
+	{
+		// Single eye, just return this case.
+		const float2 Vec = PixelPosition - LeftEyeCenterPixelXY;
+		return dot(Vec, Vec) / ViewDiagonalSquaredInPixels;
+	}
+}
+
+////////////////////////////////////////////////////////////////////////////////////////////////////
+// Get the shading rate at the specified pixel position, for HMD-based fixed or gaze-tracked foveation.
+////////////////////////////////////////////////////////////////////////////////////////////////////
+
+uint GetFoveatedShadingRate(uint2 PixelPositionIn)
+{
+	const float2 PixelPosition = float2((float)PixelPositionIn.x, (float)PixelPositionIn.y);
+	const float FractionalOffset = GetFractionalOffsetFromEyeOrigin(PixelPosition);
+	
+	if (FractionalOffset > FoveationHalfRateCutoffSquared)
+	{
+		return SHADING_RATE_4x4;
+	}
+
+	if (FractionalOffset > FoveationFullRateCutoffSquared)
+	{
+		return SHADING_RATE_2x2;
+	}
+
+	return SHADING_RATE_1x1;
+}
+
+
+[numthreads(THREADGROUP_SIZEX, THREADGROUP_SIZEY, 1)]
+void GenerateShadingRateTexture(uint3 DispatchThreadId : SV_DispatchThreadID)
+{
+	const uint2 TexelCoord = DispatchThreadId.xy;
+	RWOutputTexture[TexelCoord] = GetFoveatedShadingRate(TexelCoord);
+}
diff --git a/Engine/Source/Runtime/Engine/Classes/Engine/RendererSettings.h b/Engine/Source/Runtime/Engine/Classes/Engine/RendererSettings.h
--- a/Engine/Source/Runtime/Engine/Classes/Engine/RendererSettings.h
+++ b/Engine/Source/Runtime/Engine/Classes/Engine/RendererSettings.h
@@ -605,14 +605,14 @@ class ENGINE_API URendererSettings : public UDeveloperSettings
 	FVector TranslucentSortAxis;
 
 	UPROPERTY(config, EditAnywhere, Category = VR, meta = (
-		ConsoleVariable = "vr.VRS.HMDFixedFoveationLevel", DisplayName = "HMD Fixed Foveation Level (Experimental)",
-		ToolTip = "Set the level of fixed-foveation to apply when generating the Variable Rate Shading attachment. This feature is currently experimental.\nThis can yield some fairly significant performance benefits on GPUs that support Tier 2 VRS.\nLower settings will result in almost no discernible artifacting on most HMDs; higher settings will show some artifacts towards the edges of the view."))
-	TEnumAsByte<EFixedFoveationLevels::Type> HMDFixedFoveationLevel;
+		ConsoleVariable = "r.VRS.HMDFoveationLevel", DisplayName = "HMD Foveation Level (Experimental)",
+		ToolTip = "Set the level of foveation to apply when generating the Variable Rate Shading attachment. This feature is currently experimental.\nThis can yield some fairly significant performance benefits on GPUs that support Tier 2 VRS.\nLower settings will result in almost no discernible artifacting on most HMDs; higher settings will show some artifacts towards the edges of the view."))
+	TEnumAsByte<EFixedFoveationLevels::Type> HMDFoveationLevel;
 
 	UPROPERTY(config, EditAnywhere, Category = VR, meta = (
-		ConsoleVariable = "vr.VRS.HMDFixedFoveationDynamic", DisplayName = "Dynamic Fixed Foveation (Experimental)",
-		ToolTip = "Allows fixed foveation level to adjust dynamically based on GPU utilization.\nLevel will range between none at the minimum, and the currently selected foveation level at the maximum."))
-	uint32 bHMDFixedFoveationDynamic:1;
+		ConsoleVariable = "r.VRS.HMDDynamicFoveation", DisplayName = "Dynamic Foveation (Experimental)",
+		ToolTip = "Allows foveation level to adjust dynamically based on GPU utilization.\nLevel will range between none at the minimum, and the currently selected foveation level at the maximum."))
+	uint32 bHMDDynamicFoveation:1;
 
 	UPROPERTY(config, EditAnywhere, Category=Postprocessing, meta=(
 		ConsoleVariable="r.CustomDepth",DisplayName="Custom Depth-Stencil Pass",
diff --git a/Engine/Source/Runtime/Renderer/Private/VariableRateShading/FixedFoveationImageGenerator.cpp b/Engine/Source/Runtime/Renderer/Private/VariableRateShading/FixedFoveationImageGenerator.cpp
deleted file mode 100644
diff --git a/Engine/Source/Runtime/Renderer/Private/VariableRateShading/FixedFoveationImageGenerator.h b/Engine/Source/Runtime/Renderer/Private/VariableRateShading/FixedFoveationImageGenerator.h
deleted file mode 100644
diff --git a/Engine/Source/Runtime/Renderer/Private/VariableRateShading/FoveatedImageGenerator.cpp b/Engine/Source/Runtime/Renderer/Private/VariableRateShading/FoveatedImageGenerator.cpp
new file mode 100644
--- /dev/null
+++ b/Engine/Source/Runtime/Renderer/Private/VariableRateShading/FoveatedImageGenerator.cpp
@@ -0,0 +1,324 @@
+// Copyright Epic Games, Inc. All Rights Reserved.
+
+#include "FoveatedImageGenerator.h"
+#include "SystemTextures.h"
+#include "GlobalShader.h"
+#include "IXRTrackingSystem.h"
+#include "StereoRendering.h"
+#include "SceneView.h"
+#include "UnrealClient.h"
+#include "DataDrivenShaderPlatformInfo.h"
+#include "SceneTexturesConfig.h"
+#include "IEyeTracker.h"
+#include "IHeadMountedDisplay.h"
+#include "HeadMountedDisplayFunctionLibrary.h"
+#include "SceneRendering.h"
+
+/* CVar values used to control generator behavior */
+
+static TAutoConsoleVariable<int> CVarFoveationLevel(
+	TEXT("r.VRS.HMDFoveationLevel"),
+	0,
+	TEXT("Level of foveated VRS to apply (when Variable Rate Shading is available)\n")
+	TEXT(" 0: Disabled (default);\n")
+	TEXT(" 1: Low;\n")
+	TEXT(" 2: Medium;\n")
+	TEXT(" 3: High;\n")
+	TEXT(" 4: High Top;\n"),
+	ECVF_RenderThreadSafe);
+
+static TAutoConsoleVariable<int32> CVarDynamicFoveation(
+	TEXT("r.VRS.HMDDynamicFoveation"),
+	0,
+	TEXT("Whether foveation level should adjust based on GPU utilization\n")
+	TEXT(" 0: Disabled (default);\n")
+	TEXT(" 1: Enabled\n"),
+	ECVF_RenderThreadSafe);
+
+static TAutoConsoleVariable<int32> CVarFoveationPreview(
+	TEXT("r.VRS.HMDFoveationPreview"),
+	1,
+	TEXT("Include foveated VRS in the VRS debug overlay.")
+	TEXT(" 0: Disabled;\n")
+	TEXT(" 1: Enabled (default)\n"),
+	ECVF_RenderThreadSafe);
+
+static TAutoConsoleVariable<int32> CVarGazeTrackedFoveation(
+	TEXT("r.VRS.HMDGazeTrackedFoveation"),
+	0,
+	TEXT("Enable gaze-tracking for foveated VRS\n")
+	TEXT(" 0: Disabled (default);\n")
+	TEXT(" 1: Enabled\n"),
+	ECVF_RenderThreadSafe);
+
+
+/* Image generation parameters, set up by Prepare() */
+FVector2f HMDFieldOfView = FVector2f(90.0f, 90.0f);
+
+
+/* Shader parameters and parameter structs */
+
+constexpr int32 kComputeGroupSize = FComputeShaderUtils::kGolden2DGroupSize;
+constexpr int32 kMaxCombinedSources = 4;
+
+enum class EVRSGenerationFlags : uint32
+{
+	None = 0x0,
+	StereoRendering = 0x1,
+	SideBySideStereo = 0x2,
+};
+
+ENUM_CLASS_FLAGS(EVRSGenerationFlags);
+
+
+/* The shader itself, which generates the foveated shading rate image */
+
+class FComputeVariableRateShadingImageGeneration : public FGlobalShader
+{
+public:
+	DECLARE_GLOBAL_SHADER(FComputeVariableRateShadingImageGeneration);
+	SHADER_USE_PARAMETER_STRUCT(FComputeVariableRateShadingImageGeneration, FGlobalShader);
+
+	BEGIN_SHADER_PARAMETER_STRUCT(FParameters, )
+		SHADER_PARAMETER_RDG_TEXTURE_UAV(RWTexture2D<uint>, RWOutputTexture)
+		SHADER_PARAMETER(FVector2f, HMDFieldOfView)
+		SHADER_PARAMETER(FVector2f, LeftEyeCenterPixelXY)
+		SHADER_PARAMETER(FVector2f, RightEyeCenterPixelXY)
+		SHADER_PARAMETER(float, ViewDiagonalSquaredInPixels)
+		SHADER_PARAMETER(float, FoveationFullRateCutoffSquared)
+		SHADER_PARAMETER(float, FoveationHalfRateCutoffSquared)
+		SHADER_PARAMETER(uint32, ShadingRateAttachmentGenerationFlags)
+	END_SHADER_PARAMETER_STRUCT()
+
+	static bool ShouldCompilePermutation(const FGlobalShaderPermutationParameters& Parameters)
+	{
+		return FDataDrivenShaderPlatformInfo::GetSupportsVariableRateShading(Parameters.Platform);
+	}
+
+	static void ModifyCompilationEnvironment(const FGlobalShaderPermutationParameters& Parameters, FShaderCompilerEnvironment& OutEnvironment)
+	{
+		FGlobalShader::ModifyCompilationEnvironment(Parameters, OutEnvironment);
+
+		// Shading rates:
+		OutEnvironment.SetDefine(TEXT("SHADING_RATE_1x1"), VRSSR_1x1);
+		OutEnvironment.SetDefine(TEXT("SHADING_RATE_1x2"), VRSSR_1x2);
+		OutEnvironment.SetDefine(TEXT("SHADING_RATE_2x1"), VRSSR_2x1);
+		OutEnvironment.SetDefine(TEXT("SHADING_RATE_2x2"), VRSSR_2x2);
+		OutEnvironment.SetDefine(TEXT("SHADING_RATE_2x4"), VRSSR_2x4);
+		OutEnvironment.SetDefine(TEXT("SHADING_RATE_4x2"), VRSSR_4x2);
+		OutEnvironment.SetDefine(TEXT("SHADING_RATE_4x4"), VRSSR_4x4);
+
+		OutEnvironment.SetDefine(TEXT("MAX_COMBINED_SOURCES_IN"), kMaxCombinedSources);
+		OutEnvironment.SetDefine(TEXT("SHADING_RATE_TILE_WIDTH"), GRHIVariableRateShadingImageTileMinWidth);
+		OutEnvironment.SetDefine(TEXT("SHADING_RATE_TILE_HEIGHT"), GRHIVariableRateShadingImageTileMinHeight);
+		OutEnvironment.SetDefine(TEXT("THREADGROUP_SIZEX"), kComputeGroupSize);
+		OutEnvironment.SetDefine(TEXT("THREADGROUP_SIZEY"), kComputeGroupSize);
+
+		OutEnvironment.SetDefine(TEXT("STEREO_RENDERING"), (uint32)EVRSGenerationFlags::StereoRendering);
+		OutEnvironment.SetDefine(TEXT("SIDE_BY_SIDE_STEREO"), (uint32)EVRSGenerationFlags::SideBySideStereo);
+	}
+};
+
+IMPLEMENT_GLOBAL_SHADER(FComputeVariableRateShadingImageGeneration, "/Engine/Private/VariableRateShading/VRSShadingRateFoveated.usf", "GenerateShadingRateTexture", SF_Compute);
+
+FRDGTextureRef FFoveatedImageGenerator::GetImage(FRDGBuilder& GraphBuilder, const FViewInfo& ViewInfo, FVariableRateShadingImageManager::EVRSPassType PassType)
+{
+	// Generator only supports up to two side-by-side views
+	if (ViewInfo.StereoViewIndex > 1)
+	{
+		return nullptr;
+	}
+
+	return CachedImage;
+}
+
+void FFoveatedImageGenerator::PrepareImages(FRDGBuilder& GraphBuilder, const FSceneViewFamily& ViewFamily, const FMinimalSceneTextures& SceneTextures)
+{
+
+	// VRS level parameters - pretty arbitrary right now, later should depend on device characteristics
+	static const float kFoveationFullRateCutoffs[] = { 1.0f, 0.7f, 0.50f, 0.35f, 0.35f };
+	static const float kFoveationHalfRateCutofffs[] = { 1.0f, 0.9f, 0.75f, 0.55f, 0.55f };
+	static const float kFixedFoveationCenterX[] = { 0.5f, 0.5f, 0.5f, 0.5f, 0.5f };
+	static const float kFixedFoveationCenterY[] = { 0.5f, 0.5f, 0.5f, 0.5f, 0.42f };
+
+	const int VRSMaxLevel = FMath::Clamp(CVarFoveationLevel.GetValueOnAnyThread(), 0, 4);
+
+	// Set up dynamic VRS amount based on target framerate, default to 1.0
+	float VRSAmount = 1.0f;
+	if (CVarDynamicFoveation.GetValueOnAnyThread() && VRSMaxLevel > 0)
+	{
+		VRSAmount = UpdateDynamicVRSAmount();
+	}
+
+	if (!VRSMaxLevel || !VRSAmount)
+	{
+		CachedImage = nullptr;
+		return;
+	}
+
+	const float FoveationFullRateCutoff = FMath::Lerp(kFoveationFullRateCutoffs[0], kFoveationFullRateCutoffs[VRSMaxLevel], VRSAmount);
+	const float FoveationHalfRateCutoff = FMath::Lerp(kFoveationHalfRateCutofffs[0], kFoveationHalfRateCutofffs[VRSMaxLevel], VRSAmount);
+
+	// Default to fixed center point
+	float FoveationCenterX = kFixedFoveationCenterX[VRSMaxLevel];
+	float FoveationCenterY = kFixedFoveationCenterY[VRSMaxLevel];
+
+	// If gaze data is available and gaze-tracking is enabled, adjust foveation center point
+	if (IsGazeTrackingEnabled())
+	{
+		TSharedPtr<IEyeTracker> EyeTracker = GEngine->EyeTrackingDevice;
+		FEyeTrackerGazeData GazeData = FEyeTrackerGazeData();
+
+		// Only use gaze if we have confident gaze data, otherwise stick with fixed
+		if (EyeTracker->GetEyeTrackerGazeData(GazeData) && GazeData.ConfidenceValue > 0.5f)
+		{
+			// Get gaze orientation relative to HMD view
+			FQuat ViewOrientation;
+			FVector ViewPosition;
+			GEngine->XRSystem->GetCurrentPose(0, ViewOrientation, ViewPosition);
+			const FVector RelativeGazeDirection = ViewOrientation.UnrotateVector(GazeData.GazeDirection);
+
+			// Switch from Unreal coordinate system to X right, Y up, Z forward before projecting
+			const FMatrix StereoProjectionMatrix = GEngine->StereoRenderingDevice->GetStereoProjectionMatrix(0);
+			const FVector4 GazePoint = StereoProjectionMatrix.GetTransposed().TransformVector(FVector(RelativeGazeDirection.Y, RelativeGazeDirection.Z, RelativeGazeDirection.X));
+
+			// Transform into 0-1 screen coordinates. 0,0 is top left.  
+			FoveationCenterX = 0.5f + GazePoint.X / 2.0f;
+			FoveationCenterY = 0.5f - GazePoint.Y / 2.0f;
+
+		}
+	}
+
+	// Sanity check VRS tile size.
+	check(GRHIVariableRateShadingImageTileMinWidth >= 8 && GRHIVariableRateShadingImageTileMinWidth <= 64 && GRHIVariableRateShadingImageTileMinHeight >= 8 && GRHIVariableRateShadingImageTileMaxHeight <= 64);
+
+	// Calculate texture size
+	const FIntPoint Size = FSceneTexturesConfig::Get().Extent;
+	const FIntPoint TextureSize(Size.X / GRHIVariableRateShadingImageTileMinWidth, Size.Y / GRHIVariableRateShadingImageTileMinHeight);
+
+	// Create texture to hold shading rate image
+	FRDGTextureDesc Desc = FRDGTextureDesc::Create2D(
+		TextureSize,
+		GRHIVariableRateShadingImageFormat,
+		FClearValueBinding::None,
+		TexCreate_Foveation | TexCreate_UAV);
+
+	FRDGTextureRef ShadingRateTexture = GraphBuilder.CreateTexture(Desc, TEXT("FoveatedShadingRateTexture"));
+
+	// Setup shader parameters and flags
+	FComputeVariableRateShadingImageGeneration::FParameters* PassParameters = GraphBuilder.AllocParameters<FComputeVariableRateShadingImageGeneration::FParameters>();
+	PassParameters->RWOutputTexture = GraphBuilder.CreateUAV(ShadingRateTexture);
+
+	PassParameters->HMDFieldOfView = HMDFieldOfView;
+
+	PassParameters->FoveationFullRateCutoffSquared = FoveationFullRateCutoff * FoveationFullRateCutoff;
+	PassParameters->FoveationHalfRateCutoffSquared = FoveationHalfRateCutoff * FoveationHalfRateCutoff;
+
+	PassParameters->LeftEyeCenterPixelXY = FVector2f(TextureSize.X * FoveationCenterX, TextureSize.Y * FoveationCenterY);
+	PassParameters->RightEyeCenterPixelXY = PassParameters->LeftEyeCenterPixelXY;
+
+	EVRSGenerationFlags GenFlags = EVRSGenerationFlags::None;
+
+	// If stereo is enabled, side-by-side is assumed for now, may need to change this if we add support for mobile multi-view
+	if (IStereoRendering::IsStereoEyeView(*ViewFamily.Views[0]) && GEngine->XRSystem.IsValid())
+	{
+		EnumAddFlags(GenFlags, EVRSGenerationFlags::StereoRendering);
+		EnumAddFlags(GenFlags, EVRSGenerationFlags::SideBySideStereo); 
+
+		// Adjust eyes for side-by-side stereo and/or quadview
+		PassParameters->LeftEyeCenterPixelXY.X /= ViewFamily.Views.Num();;
+		PassParameters->RightEyeCenterPixelXY.X = PassParameters->LeftEyeCenterPixelXY.X + TextureSize.X / ViewFamily.Views.Num();
+	}
+
+	// Set up remaining parameters
+	const FVector2f ViewCenterPoint = FVector2f(TextureSize.X / ViewFamily.Views.Num() * 0.5f, TextureSize.Y * 0.5f);
+	PassParameters->ViewDiagonalSquaredInPixels = FVector2f::DotProduct(ViewCenterPoint, ViewCenterPoint);
+	PassParameters->ShadingRateAttachmentGenerationFlags = (uint32)GenFlags;
+
+	TShaderMapRef<FComputeVariableRateShadingImageGeneration> ComputeShader(GetGlobalShaderMap(GMaxRHIFeatureLevel)); // If in stereo, shader sets up a single side-by-side image for both eyes at once
+	FIntVector GroupCount = FComputeShaderUtils::GetGroupCount(TextureSize, FComputeShaderUtils::kGolden2DGroupSize);
+	GraphBuilder.AddPass(
+		RDG_EVENT_NAME("GenerateFoveatedVRSImage"),
+		PassParameters,
+		ERDGPassFlags::Compute,
+		[PassParameters, ComputeShader, GroupCount](FRHIComputeCommandList& RHICmdList)
+		{
+			FComputeShaderUtils::Dispatch(RHICmdList, ComputeShader, *PassParameters, GroupCount);
+		});
+
+	CachedImage = ShadingRateTexture;
+}
+
+bool FFoveatedImageGenerator::IsEnabledForView(const FSceneView& View) const
+{
+	// Only enabled if we are in XR
+	bool bStereoRendering = IStereoRendering::IsStereoEyeView(View) && GEngine->XRSystem.IsValid();
+	return bStereoRendering && CVarFoveationLevel.GetValueOnRenderThread() > 0;
+}
+
+FVariableRateShadingImageManager::EVRSSourceType FFoveatedImageGenerator::GetType() const
+{
+	return IsGazeTrackingEnabled() ? FVariableRateShadingImageManager::EVRSSourceType::FixedFoveation : FVariableRateShadingImageManager::EVRSSourceType::EyeTrackedFoveation;
+}
+
+FRDGTextureRef FFoveatedImageGenerator::GetDebugImage(FRDGBuilder& GraphBuilder, const FViewInfo& ViewInfo)
+{
+	return CachedImage;
+}
+
+float FFoveatedImageGenerator::UpdateDynamicVRSAmount()
+{
+	const float		kUpdateIncrement = 0.1f;
+	const int		kNumFramesToAverage = 10;
+	const double	kRoundingErrorMargin = 0.5; // Add a small margin (.5 ms) to avoid oscillation
+	const double	kDesktopMaxAllowedFrameTime = 12.50; // 80 fps
+	const double	kMobileMaxAllowedFrameTime = 16.66; // 60 fps
+
+	if (GFrameNumber != DynamicVRSData.LastUpdateFrame)
+	{
+		DynamicVRSData.LastUpdateFrame = GFrameNumber;
+
+		const double GPUFrameTime = FPlatformTime::ToMilliseconds(RHIGetGPUFrameCycles());
+
+		// Update sum for rolling average
+		if (DynamicVRSData.NumFramesStored < kNumFramesToAverage)
+		{
+			DynamicVRSData.SumBusyTime += GPUFrameTime;
+			DynamicVRSData.NumFramesStored++;
+		}
+		else if (DynamicVRSData.NumFramesStored == kNumFramesToAverage)
+		{
+			DynamicVRSData.SumBusyTime -= DynamicVRSData.SumBusyTime / kNumFramesToAverage;
+			DynamicVRSData.SumBusyTime += GPUFrameTime;
+		}
+
+		// If rolling average has sufficient samples, check if update is necessary
+		if (DynamicVRSData.NumFramesStored == kNumFramesToAverage)
+		{
+			const double AverageBusyTime = DynamicVRSData.SumBusyTime / kNumFramesToAverage;
+			const double TargetBusyTime = IsMobilePlatform(GShaderPlatformForFeatureLevel[GMaxRHIFeatureLevel]) ? kMobileMaxAllowedFrameTime : kDesktopMaxAllowedFrameTime;
+
+			if (AverageBusyTime > TargetBusyTime + kRoundingErrorMargin && DynamicVRSData.VRSAmount < 1.0f)
+			{
+				DynamicVRSData.VRSAmount = FMath::Clamp(DynamicVRSData.VRSAmount + kUpdateIncrement, 0.0f, 1.0f);
+				DynamicVRSData.SumBusyTime = 0.0;
+				DynamicVRSData.NumFramesStored = 0;
+			}
+			else if (AverageBusyTime < TargetBusyTime - kRoundingErrorMargin && DynamicVRSData.VRSAmount > 0.0f)
+			{
+				DynamicVRSData.VRSAmount = FMath::Clamp(DynamicVRSData.VRSAmount - kUpdateIncrement, 0.0f, 1.0f);
+				DynamicVRSData.SumBusyTime = 0.0;
+				DynamicVRSData.NumFramesStored = 0;
+			}
+		}
+	}
+
+	return DynamicVRSData.VRSAmount;
+}
+
+bool FFoveatedImageGenerator::IsGazeTrackingEnabled() const
+{
+	return CVarGazeTrackedFoveation.GetValueOnAnyThread()
+		&& GEngine->EyeTrackingDevice.IsValid()
+		&& GEngine->StereoRenderingDevice.IsValid();
+}
\ No newline at end of file
diff --git a/Engine/Source/Runtime/Renderer/Private/VariableRateShading/FoveatedImageGenerator.h b/Engine/Source/Runtime/Renderer/Private/VariableRateShading/FoveatedImageGenerator.h
new file mode 100644
--- /dev/null
+++ b/Engine/Source/Runtime/Renderer/Private/VariableRateShading/FoveatedImageGenerator.h
@@ -0,0 +1,32 @@
+// Copyright Epic Games, Inc. All Rights Reserved.
+
+#pragma once
+
+#include "CoreMinimal.h"
+#include "RHI.h"
+#include "RendererInterface.h"
+#include "RenderGraphDefinitions.h"
+#include "VariableRateShadingImageManager.h"
+#include "Engine/Engine.h"
+
+class FFoveatedImageGenerator : public IVariableRateShadingImageGenerator
+{
+public:
+	virtual ~FFoveatedImageGenerator() override {};
+	virtual FRDGTextureRef GetImage(FRDGBuilder& GraphBuilder, const FViewInfo& ViewInfo, FVariableRateShadingImageManager::EVRSPassType PassType) override;
+	virtual void PrepareImages(FRDGBuilder& GraphBuilder, const FSceneViewFamily& ViewFamily, const FMinimalSceneTextures& SceneTextures) override;
+	virtual bool IsEnabledForView(const FSceneView& View) const override;
+	virtual FVariableRateShadingImageManager::EVRSSourceType GetType() const override;
+	virtual FRDGTextureRef GetDebugImage(FRDGBuilder& GraphBuilder, const FViewInfo& ViewInfo) override;
+private:
+	FRDGTextureRef CachedImage = nullptr;
+	struct FDynamicVRSData
+	{
+		float	VRSAmount = 1.0f;
+		double	SumBusyTime = 0.0;
+		int		NumFramesStored = 0;
+		uint32	LastUpdateFrame = 0;
+	} DynamicVRSData;
+	float UpdateDynamicVRSAmount();
+	bool IsGazeTrackingEnabled() const;
+};
\ No newline at end of file
diff --git a/Engine/Source/Runtime/Renderer/Private/VariableRateShading/VariableRateShadingImageManager.cpp b/Engine/Source/Runtime/Renderer/Private/VariableRateShading/VariableRateShadingImageManager.cpp
--- a/Engine/Source/Runtime/Renderer/Private/VariableRateShading/VariableRateShadingImageManager.cpp
+++ b/Engine/Source/Runtime/Renderer/Private/VariableRateShading/VariableRateShadingImageManager.cpp
@@ -1,7 +1,7 @@
 // Copyright Epic Games, Inc. All Rights Reserved.
 
 #include "VariableRateShadingImageManager.h"
-#include "FixedFoveationImageGenerator.h"
+#include "FoveatedImageGenerator.h"
 #include "ContrastAdaptiveImageGenerator.h"
 #include "StereoRenderTargetManager.h"
 #include "GlobalShader.h"
@@ -163,7 +163,7 @@ IMPLEMENT_GLOBAL_SHADER(FDebugVariableRateShadingPS, "/Engine/Private/VariableRa
 FVariableRateShadingImageManager::FVariableRateShadingImageManager()
 	: FRenderResource()
 {
-	ImageGenerators.Add(MakeUnique<FFixedFoveationImageGenerator>());
+	ImageGenerators.Add(MakeUnique<FFoveatedImageGenerator>());
 	ImageGenerators.Add(MakeUnique<FContrastAdaptiveImageGenerator>());
 
 	// TODO: Add more generators or allow registration from plugins
