Prevent zero size buffer in Non-Nanite Virtual Shadow Map culling passes.

diff --git a/Engine/Source/Runtime/Renderer/Private/VirtualShadowMaps/VirtualShadowMapArray.cpp b/Engine/Source/Runtime/Renderer/Private/VirtualShadowMaps/VirtualShadowMapArray.cpp
--- a/Engine/Source/Runtime/Renderer/Private/VirtualShadowMaps/VirtualShadowMapArray.cpp
+++ b/Engine/Source/Runtime/Renderer/Private/VirtualShadowMaps/VirtualShadowMapArray.cpp
@@ -2318,7 +2318,7 @@ static FCullingResult AddCullingPasses(FRDGBuilder& GraphBuilder,
 
 	// TODO: This is both not right, and also over conservative when running with the atomic path
 	FCullingResult CullingResult;
-	CullingResult.MaxNumInstancesPerPass = TotalInstances * 64u;
+	CullingResult.MaxNumInstancesPerPass = FMath::Max(1u, TotalInstances * 64u);
 	FRDGBufferRef VisibleInstancesRdg = CreateStructuredBuffer(GraphBuilder, TEXT("Shadow.Virtual.VisibleInstances"), sizeof(FVSMVisibleInstanceCmd), CullingResult.MaxNumInstancesPerPass, nullptr, 0);
 
 	FRDGBufferRef VisibleInstanceWriteOffsetRDG = CreateStructuredBuffer(GraphBuilder, TEXT("Shadow.Virtual.VisibleInstanceWriteOffset"), sizeof(uint32), 1, nullptr, 0);
