Remove todos and update comments for forcing optimizations on for a few SSD shaders

- shader has a static branch wrapping a unrolled loop
  - unoptimized compilation will not strip the code in the branch and hence will generate a error due to a non-compile time specified loop iteration count.
- remove the todo and comment the force optimization flags so its clear why they are required

diff --git a/Engine/Source/Runtime/Renderer/Private/ScreenSpaceDenoise.cpp b/Engine/Source/Runtime/Renderer/Private/ScreenSpaceDenoise.cpp
--- a/Engine/Source/Runtime/Renderer/Private/ScreenSpaceDenoise.cpp
+++ b/Engine/Source/Runtime/Renderer/Private/ScreenSpaceDenoise.cpp
@@ -1114,7 +1114,7 @@ class FSSDSpatialAccumulationCS : public FGlobalShader
 
 	static void ModifyCompilationEnvironment(const FGlobalShaderPermutationParameters& Parameters, FShaderCompilerEnvironment& OutEnvironment)
 	{
-		// TODO: UECON-464 - force optimizations to workaround shader compiler issue on DXC until fixed by MS
+		// Dead code stripping required or else we have an unrolled loop with a non-compile time specified iteration count.
 		OutEnvironment.CompilerFlags.Add(CFLAG_ForceOptimization);
 	}
 
@@ -1167,7 +1167,7 @@ class FSSDTemporalAccumulationCS : public FGlobalShader
 
 	static void ModifyCompilationEnvironment(const FGlobalShaderPermutationParameters& Parameters, FShaderCompilerEnvironment& OutEnvironment)
 	{
-		// TODO: UECON-464 - force optimizations to workaround shader compiler issue on DXC until fixed by MS
+		// Dead code stripping required or else we have an unrolled loop with a non-compile time specified iteration count.
 		OutEnvironment.CompilerFlags.Add(CFLAG_ForceOptimization);
 	}
 
