Move TAA back to the end of PP Chain to save bandwidth.
Filter the bloom input to get rid of the flickering white spots.

diff --git a/Engine/Shaders/Private/PostProcessMobile.usf b/Engine/Shaders/Private/PostProcessMobile.usf
--- a/Engine/Shaders/Private/PostProcessMobile.usf
+++ b/Engine/Shaders/Private/PostProcessMobile.usf
@@ -233,8 +233,33 @@ void BloomVS_Mobile(
 	OutTexCoords[3] = TransformedUV + BufferSizeAndInvSize.zw * float2( 1,  1);
 }
 
+// Temporal jitter at the pixel scale.
+float2 TemporalJitterPixels;
+
+static const float2 kOffsets[5] =
+{
+	int2( 0, -1),
+	int2(-1,  0),
+	int2( 1,  0),
+	int2( 0,  1),
+	int2( 0,  0), // K
+};
+
+DECLARE_SCALAR_ARRAY(float, PlusWeights, 5);
+
+half Luma4(half3 Color)
+{
+	return (Color.g * 2.0) + (Color.r + Color.b);
+}
+
+half HdrWeight4(half3 Color)
+{
+	return rcp(Luma4(Color) + 4.0);
+}
+
 void BloomPS_Mobile(
-	float2 InUVs[4] : TEXCOORD0
+	float2 InUVs[4] : TEXCOORD0,
+	in float4 SvPosition : SV_POSITION
 #if MOBILE_USEBLOOM
 
 	, out HALF4_TYPE OutColor : SV_Target0
@@ -280,13 +305,27 @@ void BloomPS_Mobile(
 	half3 AverageColor = 0.0f;
 
 #if MOBILE_USEBLOOM || MOBILE_USESUN || MOBILE_USEEYEADAPTATION
-	half4 C0 = SceneColorTexture.Sample(SceneColorSampler, InUVs[0]);
-	half4 C1 = SceneColorTexture.Sample(SceneColorSampler, InUVs[1]);
-	half4 C2 = SceneColorTexture.Sample(SceneColorSampler, InUVs[2]);
-	half4 C3 = SceneColorTexture.Sample(SceneColorSampler, InUVs[3]);
+
+	half NeighborsFinalWeight = 0;
+	float3 NeighborsColor = 0;
+
+	// Filter input pixels. Referenced from the FilterCurrentFrameInputSamples in TemporalAA.usf
+	for (uint i = 0; i < 5; i++)
+	{
+ 		float2 PPCk = SvPosition.xy + TemporalJitterPixels;
+ 
+		float2 TexCoords = BufferSizeAndInvSize.zw * (kOffsets[i] + PPCk);
+		half3 SampleColor = SceneColorTexture.Sample(SceneColorSampler, TexCoords).rgb;
+		half SampleHdrWeight = HdrWeight4(SampleColor);
+		half SampleSpatialWeight = GET_SCALAR_ARRAY_ELEMENT(PlusWeights, i);
+		half SampleFinalWeight = SampleSpatialWeight * SampleHdrWeight;
+
+		NeighborsColor += SampleFinalWeight * SampleColor;
+		NeighborsFinalWeight += SampleFinalWeight;
+	}
 
 	// Output color is average.
-	AverageColor.rgb = (C0.rgb * 0.25) + (C1.rgb * 0.25) + (C2.rgb * 0.25) + (C3.rgb * 0.25);
+	AverageColor.rgb = NeighborsColor * rcp(NeighborsFinalWeight);
 
 	#if METAL_MSAA_HDR_DECODE
 		// This should really happen before the average, instead doing after average as optimization.
@@ -1527,4 +1566,4 @@ void HistogramEyeAdaptationCS(uint2 DispatchThreadId : SV_DispatchThreadID)
 
 	OutputBuffer[0] = OutColor;
 }
-#endif
\ No newline at end of file
+#endif
diff --git a/Engine/Source/Runtime/Renderer/Private/PostProcess/PostProcessMobile.cpp b/Engine/Source/Runtime/Renderer/Private/PostProcess/PostProcessMobile.cpp
--- a/Engine/Source/Runtime/Renderer/Private/PostProcess/PostProcessMobile.cpp
+++ b/Engine/Source/Runtime/Renderer/Private/PostProcess/PostProcessMobile.cpp
@@ -217,6 +217,9 @@ public:
 
 	BEGIN_SHADER_PARAMETER_STRUCT(FParameters, )
 		SHADER_PARAMETER(float, BloomThreshold)
+		SHADER_PARAMETER(FVector2f, TemporalJitterPixels)
+		SHADER_PARAMETER(FVector4f, BufferSizeAndInvSize)
+		SHADER_PARAMETER_SCALAR_ARRAY(float, PlusWeights, [5])
 		SHADER_PARAMETER_STRUCT_REF(FViewUniformShaderParameters, View)
 		SHADER_PARAMETER_STRUCT(FEyeAdaptationParameters, EyeAdaptation)
 		SHADER_PARAMETER_RDG_TEXTURE(Texture2D, SceneColorTexture)
@@ -320,7 +323,8 @@ FMobileBloomSetupOutputs AddMobileBloomSetupPass(FRDGBuilder& GraphBuilder, cons
 	FMobileBloomSetupVS::FParameters VSShaderParameters;
 
 	VSShaderParameters.View = View.ViewUniformBuffer;
-	VSShaderParameters.BufferSizeAndInvSize = FVector4f(BufferSize.X, BufferSize.Y, 1.0f / BufferSize.X, 1.0f / BufferSize.Y);
+	FVector4f BufferSizeAndInvSize = FVector4f(BufferSize.X, BufferSize.Y, 1.0f / BufferSize.X, 1.0f / BufferSize.Y);
+	VSShaderParameters.BufferSizeAndInvSize = BufferSizeAndInvSize;
 
 	auto ShaderPermutationVector = FMobileBloomSetupPS::BuildPermutationVector(Inputs.bUseBloom, Inputs.bUseSun, Inputs.bUseDof, Inputs.bUseEyeAdaptation, Inputs.bUseMetalMSAAHDRDecode);
 
@@ -333,16 +337,49 @@ FMobileBloomSetupOutputs AddMobileBloomSetupPass(FRDGBuilder& GraphBuilder, cons
 		PSShaderParameters->RenderTargets[i] = DestRenderTargets[i].GetRenderTargetBinding();
 	}
 
+	PSShaderParameters->BufferSizeAndInvSize = BufferSizeAndInvSize;
 	PSShaderParameters->EyeAdaptation = EyeAdaptationParameters;
 	PSShaderParameters->BloomThreshold = View.FinalPostProcessSettings.BloomThreshold;
 	PSShaderParameters->View = View.ViewUniformBuffer;
 
 	PSShaderParameters->SceneColorTexture = Inputs.SceneColor.Texture;
-	PSShaderParameters->SceneColorSampler = TStaticSamplerState<SF_Bilinear, AM_Clamp, AM_Clamp, AM_Clamp>::GetRHI();
+	PSShaderParameters->SceneColorSampler = TStaticSamplerState<SF_Point, AM_Clamp, AM_Clamp, AM_Clamp>::GetRHI();
 
 	PSShaderParameters->SunShaftAndDofTexture = Inputs.SunShaftAndDof.Texture;
 	PSShaderParameters->SunShaftAndDofSampler = TStaticSamplerState<SF_Bilinear, AM_Clamp, AM_Clamp, AM_Clamp>::GetRHI();
 
+	// Compute PlusWeights. Referenced from SetupSampleWeightParameters in TemporalAA.cpp
+	{
+		static const float SampleOffsets[5][2] =
+		{
+			{  0.0f, -1.0f },
+			{ -1.0f,  0.0f },
+			{  1.0f,  0.0f },
+			{  0.0f,  1.0f },
+			{  0.0f,  0.0f },
+		};
+
+		FVector2f TemporalJitterPixels = FVector2f(View.TemporalJitterPixels);
+		float TotalWeight = 0.0f;
+		for (int32 i = 0; i < 5; i++)
+		{
+			float PixelOffsetX = SampleOffsets[i][0] - TemporalJitterPixels.X;
+			float PixelOffsetY = SampleOffsets[i][1] - TemporalJitterPixels.Y;
+
+			// Normal distribution, Sigma = 0.47
+			const float CurrWeight = FMath::Exp(-2.29f * (PixelOffsetX * PixelOffsetX + PixelOffsetY * PixelOffsetY));
+			GET_SCALAR_ARRAY_ELEMENT(PSShaderParameters->PlusWeights, i) = CurrWeight;
+			TotalWeight += CurrWeight;
+		}
+
+		for (int32 i = 0; i < 5; i++)
+		{
+			GET_SCALAR_ARRAY_ELEMENT(PSShaderParameters->PlusWeights, i) /= TotalWeight;
+		}
+
+		PSShaderParameters->TemporalJitterPixels = TemporalJitterPixels;
+	}
+
 	const FScreenPassTextureViewport InputViewport(Inputs.SceneColor);
 	const FScreenPassTextureViewport OutputViewport(DestRenderTargets[0]);
 
