Fix for FRayTracingMeshProcessor::BuildRayTracingMeshCommands using explicit SM5 feature level

- this caused assertion on SM6 platforms

#rnx

diff --git a/Engine/Source/Runtime/Renderer/Private/RayTracing/RayTracingMaterialHitShaders.h b/Engine/Source/Runtime/Renderer/Private/RayTracing/RayTracingMaterialHitShaders.h
--- a/Engine/Source/Runtime/Renderer/Private/RayTracing/RayTracingMaterialHitShaders.h
+++ b/Engine/Source/Runtime/Renderer/Private/RayTracing/RayTracingMaterialHitShaders.h
@@ -71,7 +71,7 @@ protected:
 		}
 
 		FVertexInputStreamArray VertexStreams;
-		VertexFactory->GetStreams(ERHIFeatureLevel::SM5, EVertexInputStreamType::Default, VertexStreams);
+		VertexFactory->GetStreams(FeatureLevel, EVertexInputStreamType::Default, VertexStreams);
 
 		int32 DataOffset = 0;
 		if (PassShaders.RayTracingShader.IsValid())
