Create proxy uniform buffers when added to the scene which happens before static mesh gathering.

diff --git a/Engine/Source/Runtime/Renderer/Private/RendererScene.cpp b/Engine/Source/Runtime/Renderer/Private/RendererScene.cpp
--- a/Engine/Source/Runtime/Renderer/Private/RendererScene.cpp
+++ b/Engine/Source/Runtime/Renderer/Private/RendererScene.cpp
@@ -6078,6 +6078,7 @@ void FScene::UpdateAllPrimitiveSceneInfos(FRDGBuilder& GraphBuilder, EUpdateAllP
 				FPrimitiveSceneInfo* PrimitiveSceneInfo = AddedLocalPrimitiveSceneInfos[AddIndex];
 				int32 PrimitiveIndex = PrimitiveSceneInfo->PackedIndex;
 				FPrimitiveSceneProxy* SceneProxy = PrimitiveSceneInfo->Proxy;
+				SceneProxy->CreateUniformBuffer();
 
 				if (ShouldPrimitiveOutputVelocity(SceneProxy, GetShaderPlatform()))
 				{
@@ -6388,7 +6389,6 @@ void FScene::UpdateAllPrimitiveSceneInfos(FRDGBuilder& GraphBuilder, EUpdateAllP
 		{
 			const int32 Index = BitIt.GetIndex();
 			FPrimitiveSceneInfo* Primitive = Primitives[Index];
-			Primitive->Proxy->CreateUniformBuffer();
 			PrimitivesNeedingUniformBufferUpdate[Index] = false;
 			ProxiesToUpdate.Emplace(Primitive->Proxy);
 			GPUScene.AddPrimitiveToUpdate(Index, EPrimitiveDirtyState::ChangedAll);
