Fixed RHI validation error in virtual textures due to underlying physical textures not being transitioned to SRV.

#rnx

diff --git a/Engine/Source/Runtime/Renderer/Private/VT/VirtualTexturePhysicalSpace.cpp b/Engine/Source/Runtime/Renderer/Private/VT/VirtualTexturePhysicalSpace.cpp
--- a/Engine/Source/Runtime/Renderer/Private/VT/VirtualTexturePhysicalSpace.cpp
+++ b/Engine/Source/Runtime/Renderer/Private/VT/VirtualTexturePhysicalSpace.cpp
@@ -175,6 +175,18 @@ void FVirtualTexturePhysicalSpace::ReleaseRHI()
 	}
 }
 
+void FVirtualTexturePhysicalSpace::FinalizeTextures(FRDGBuilder& GraphBuilder)
+{
+	for (int32 Layer = 0; Layer < Description.NumLayers; ++Layer)
+	{
+		// It's only necessary to enable external access mode on textures modified by RDG this frame.
+		if (FRDGTexture* Texture = GraphBuilder.FindExternalTexture(PooledRenderTarget[Layer]))
+		{
+			GraphBuilder.UseExternalAccessMode(Texture, ERHIAccess::SRVMask);
+		}
+	}
+}
+
 uint32 FVirtualTexturePhysicalSpace::GetSizeInBytes() const
 {
 	SIZE_T TileSizeBytes = 0;
diff --git a/Engine/Source/Runtime/Renderer/Private/VT/VirtualTexturePhysicalSpace.h b/Engine/Source/Runtime/Renderer/Private/VT/VirtualTexturePhysicalSpace.h
--- a/Engine/Source/Runtime/Renderer/Private/VT/VirtualTexturePhysicalSpace.h
+++ b/Engine/Source/Runtime/Renderer/Private/VT/VirtualTexturePhysicalSpace.h
@@ -111,6 +111,8 @@ public:
 		return PooledRenderTarget[Layer];
 	}
 
+	void FinalizeTextures(FRDGBuilder& GraphBuilder);
+
 	/** Update internal tracking of residency. This is used to update stats and to calculate a mip bias to keep within the pool budget. */
 	void UpdateResidencyTracking(uint32 Frame);
 	/** Get dynamic mip bias used to keep within residency budget. */
diff --git a/Engine/Source/Runtime/Renderer/Private/VT/VirtualTextureSystem.cpp b/Engine/Source/Runtime/Renderer/Private/VT/VirtualTextureSystem.cpp
--- a/Engine/Source/Runtime/Renderer/Private/VT/VirtualTextureSystem.cpp
+++ b/Engine/Source/Runtime/Renderer/Private/VT/VirtualTextureSystem.cpp
@@ -2413,6 +2413,11 @@ void FVirtualTextureSystem::FinalizeRequests(FRDGBuilder& GraphBuilder)
 		}
 	}
 
+	for (FVirtualTexturePhysicalSpace* PhysicalSpace : PhysicalSpaces)
+	{
+		PhysicalSpace->FinalizeTextures(GraphBuilder);
+	}
+
 	Frame++;
 }
 
