Fix false positive build health static analysis warning caused by the use of static const float arrays.

diff --git a/Engine/Source/Runtime/Renderer/Private/VariableRateShading/FoveatedImageGenerator.cpp b/Engine/Source/Runtime/Renderer/Private/VariableRateShading/FoveatedImageGenerator.cpp
--- a/Engine/Source/Runtime/Renderer/Private/VariableRateShading/FoveatedImageGenerator.cpp
+++ b/Engine/Source/Runtime/Renderer/Private/VariableRateShading/FoveatedImageGenerator.cpp
@@ -136,28 +136,28 @@ void FFoveatedImageGenerator::PrepareImages(FRDGBuilder& GraphBuilder, const FSc
 {
 
 	// VRS level parameters - pretty arbitrary right now, later should depend on device characteristics
-	static const float kFoveationFullRateCutoffs[] = { 1.0f, 0.7f, 0.50f, 0.35f, 0.35f };
-	static const float kFoveationHalfRateCutofffs[] = { 1.0f, 0.9f, 0.75f, 0.55f, 0.55f };
-	static const float kFixedFoveationCenterX[] = { 0.5f, 0.5f, 0.5f, 0.5f, 0.5f };
-	static const float kFixedFoveationCenterY[] = { 0.5f, 0.5f, 0.5f, 0.5f, 0.42f };
+	static const TArray<float> kFoveationFullRateCutoffs = { 1.0f, 0.7f, 0.50f, 0.35f, 0.35f };
+	static const TArray<float> kFoveationHalfRateCutofffs = { 1.0f, 0.9f, 0.75f, 0.55f, 0.55f };
+	static const TArray<float> kFixedFoveationCenterX = { 0.5f, 0.5f, 0.5f, 0.5f, 0.5f };
+	static const TArray<float> kFixedFoveationCenterY = { 0.5f, 0.5f, 0.5f, 0.5f, 0.42f };
 
 	const int VRSMaxLevel = FMath::Clamp(CVarFoveationLevel.GetValueOnAnyThread(), 0, 4);
 
-	// Set up dynamic VRS amount based on target framerate, default to 1.0
-	float VRSAmount = 1.0f;
+	// Set up dynamic VRS amount based on target framerate and use it to interpolate cutoffs (default to max level)
+	float VRSDynamicAmount = 1.0f;
 	if (CVarDynamicFoveation.GetValueOnAnyThread() && VRSMaxLevel > 0)
 	{
-		VRSAmount = UpdateDynamicVRSAmount();
+		VRSDynamicAmount = UpdateDynamicVRSAmount();
 	}
 
-	if (!VRSMaxLevel || !VRSAmount)
+	if (VRSMaxLevel <= 0 || VRSDynamicAmount <= 0)
 	{
-		CachedImage = nullptr;
+		CachedImage = nullptr; // Early out if no VRS requested
 		return;
 	}
 
-	const float FoveationFullRateCutoff = FMath::Lerp(kFoveationFullRateCutoffs[0], kFoveationFullRateCutoffs[VRSMaxLevel], VRSAmount);
-	const float FoveationHalfRateCutoff = FMath::Lerp(kFoveationHalfRateCutofffs[0], kFoveationHalfRateCutofffs[VRSMaxLevel], VRSAmount);
+	const float FoveationFullRateCutoff = FMath::Lerp(kFoveationFullRateCutoffs[0], kFoveationFullRateCutoffs[VRSMaxLevel], VRSDynamicAmount);
+	const float FoveationHalfRateCutoff = FMath::Lerp(kFoveationHalfRateCutofffs[0], kFoveationFullRateCutoffs[VRSMaxLevel], VRSDynamicAmount);
 
 	// Default to fixed center point
 	float FoveationCenterX = kFixedFoveationCenterX[VRSMaxLevel];
