Fixed crash in distance fields checking for capsule shadows.

diff --git a/Engine/Source/Runtime/Renderer/Private/CapsuleShadowRendering.cpp b/Engine/Source/Runtime/Renderer/Private/CapsuleShadowRendering.cpp
--- a/Engine/Source/Runtime/Renderer/Private/CapsuleShadowRendering.cpp
+++ b/Engine/Source/Runtime/Renderer/Private/CapsuleShadowRendering.cpp
@@ -1182,31 +1182,7 @@ void FDeferredShadingSceneRenderer::RenderIndirectCapsuleShadows(FRDGBuilder& Gr
 
 bool FSceneRenderer::ShouldPrepareForDFInsetIndirectShadow() const
 {
-	bool bSceneHasInsetDFPrimitives = false;
-
-	for (int32 ViewIndex = 0; ViewIndex < Views.Num(); ViewIndex++)
-	{
-		const FViewInfo& View = Views[ViewIndex];
-
-		for (int32 PrimitiveIndex = 0; PrimitiveIndex < View.IndirectShadowPrimitives.Num(); PrimitiveIndex++)
-		{
-			FPrimitiveSceneInfo* PrimitiveSceneInfo = View.IndirectShadowPrimitives[PrimitiveIndex];
-			TArray<const FPrimitiveSceneInfo*, SceneRenderingAllocator> ShadowGroupPrimitives;
-			PrimitiveSceneInfo->GatherLightingAttachmentGroupPrimitives(ShadowGroupPrimitives);
-
-			for (int32 ChildIndex = 0; ChildIndex < ShadowGroupPrimitives.Num(); ChildIndex++)
-			{
-				const FPrimitiveSceneInfo* GroupPrimitiveSceneInfo = ShadowGroupPrimitives[ChildIndex];
-
-				if (GroupPrimitiveSceneInfo->Proxy->CastsDynamicShadow() && GroupPrimitiveSceneInfo->Proxy->HasDistanceFieldRepresentation())
-				{
-					bSceneHasInsetDFPrimitives = true;
-				}
-			}
-		}
-	}
-
-	return bSceneHasInsetDFPrimitives && SupportsCapsuleIndirectShadows(ShaderPlatform) && ViewFamily.EngineShowFlags.CapsuleShadows;
+	return SupportsCapsuleIndirectShadows(ShaderPlatform) && ViewFamily.EngineShowFlags.CapsuleShadows;
 }
 
 void FDeferredShadingSceneRenderer::RenderCapsuleShadowsForMovableSkylight(
