Fix hair crash with nDisplay.

diff --git a/Engine/Source/Runtime/Renderer/Private/HairStrands/HairStrandsRendering.cpp b/Engine/Source/Runtime/Renderer/Private/HairStrands/HairStrandsRendering.cpp
--- a/Engine/Source/Runtime/Renderer/Private/HairStrands/HairStrandsRendering.cpp
+++ b/Engine/Source/Runtime/Renderer/Private/HairStrands/HairStrandsRendering.cpp
@@ -100,7 +100,11 @@ void RenderHairPrePass(
 		FVector PreViewStereoCorrection = FVector::ZeroVector;
 		if (IStereoRendering::IsStereoEyeView(View))
 		{
-			PreViewStereoCorrection = Views[0].ViewMatrices.GetPreViewTranslation() - Views[1].ViewMatrices.GetPreViewTranslation();
+			// nDisplay uses StereoRendering code path with a mono-view
+			if (Views.Num() >= 2)
+			{
+				PreViewStereoCorrection = Views[0].ViewMatrices.GetPreViewTranslation() - Views[1].ViewMatrices.GetPreViewTranslation();
+			}
 			if (IStereoRendering::IsASecondaryView(View))
 			{
 				// No need to copy the view state (i.e., HairStrandsViewStateData) as it is only used for 
