Volumetric Support For Mobile - disabled by default in DP
Removed bSupportsVolumetricFog from DDPI

diff --git a/Engine/Config/Android/AndroidScalability.ini b/Engine/Config/Android/AndroidScalability.ini
--- a/Engine/Config/Android/AndroidScalability.ini
+++ b/Engine/Config/Android/AndroidScalability.ini
@@ -51,6 +51,7 @@ r.Shadow.MaxCSMResolution=1024
 r.Shadow.DistanceScale=0.5
 r.Shadow.CSM.MaxCascades=1
 r.Shadow.RadiusThreshold=0.06
+r.VolumetricFog=0
 
 [ShadowQuality@1]
 r.ShadowQuality=1
@@ -58,6 +59,7 @@ r.Shadow.MaxCSMResolution=1024
 r.Shadow.DistanceScale=0.7
 r.Shadow.CSM.MaxCascades=1
 r.Shadow.RadiusThreshold=0.06
+r.VolumetricFog=0
 
 [ShadowQuality@2]
 r.ShadowQuality=2
@@ -65,6 +67,7 @@ r.Shadow.MaxCSMResolution=2048
 r.Shadow.DistanceScale=1.0
 r.Shadow.CSM.MaxCascades=2
 r.Shadow.RadiusThreshold=0.06
+r.VolumetricFog=0
 
 [ShadowQuality@3]
 r.ShadowQuality=2
@@ -72,6 +75,7 @@ r.Shadow.MaxCSMResolution=2048
 r.Shadow.DistanceScale=1.0
 r.Shadow.CSM.MaxCascades=3
 r.Shadow.RadiusThreshold=0.06
+r.VolumetricFog=0
 
 ; ------------------------------------------------------------------------------
 ; Postprocess settings
diff --git a/Engine/Config/IOS/IOSScalability.ini b/Engine/Config/IOS/IOSScalability.ini
--- a/Engine/Config/IOS/IOSScalability.ini
+++ b/Engine/Config/IOS/IOSScalability.ini
@@ -51,6 +51,7 @@ r.Shadow.MaxCSMResolution=1024
 r.Shadow.DistanceScale=0.5
 r.Shadow.CSM.MaxCascades=1
 r.Shadow.RadiusThreshold=0.06
+r.VolumetricFog=0
 
 [ShadowQuality@1]
 r.ShadowQuality=1
@@ -58,6 +59,7 @@ r.Shadow.MaxCSMResolution=1024
 r.Shadow.DistanceScale=0.7
 r.Shadow.CSM.MaxCascades=1
 r.Shadow.RadiusThreshold=0.06
+r.VolumetricFog=0
 
 [ShadowQuality@2]
 r.ShadowQuality=2
@@ -65,6 +67,7 @@ r.Shadow.MaxCSMResolution=2048
 r.Shadow.DistanceScale=1.0
 r.Shadow.CSM.MaxCascades=2
 r.Shadow.RadiusThreshold=0.06
+r.VolumetricFog=0
 
 [ShadowQuality@3]
 r.ShadowQuality=2
@@ -72,6 +75,7 @@ r.Shadow.MaxCSMResolution=2048
 r.Shadow.DistanceScale=1.0
 r.Shadow.CSM.MaxCascades=3
 r.Shadow.RadiusThreshold=0.06
+r.VolumetricFog=0
 
 ; ------------------------------------------------------------------------------
 ; Postprocess settings
diff --git a/Engine/Config/Mac/DataDrivenPlatformInfo.ini b/Engine/Config/Mac/DataDrivenPlatformInfo.ini
--- a/Engine/Config/Mac/DataDrivenPlatformInfo.ini
+++ b/Engine/Config/Mac/DataDrivenPlatformInfo.ini
@@ -46,7 +46,6 @@ bSupportsDxc=true
 bIsSPIRV=true
 bSupportsInstancedStereo=true
 SupportsMultiViewport=RuntimeGuaranteed
-bSupportsVolumetricFog=true
 bSupportsPercentageCloserShadows=true
 bSupportsCapsuleShadows=true
 bSupports4ComponentUAVReadWrite=true
diff --git a/Engine/Config/VulkanPC/DataDrivenPlatformInfo.ini b/Engine/Config/VulkanPC/DataDrivenPlatformInfo.ini
--- a/Engine/Config/VulkanPC/DataDrivenPlatformInfo.ini
+++ b/Engine/Config/VulkanPC/DataDrivenPlatformInfo.ini
@@ -33,7 +33,6 @@ bSupportsDiaphragmDOF = true
 bSupportsRGBColorBuffer = true
 bSupportsCapsuleShadows = true
 bSupportsPercentageCloserShadows=true
-bSupportsVolumetricFog = true
 bSupportsIndexBufferUAVs = true
 
 bSupportsInstancedStereo = true
diff --git a/Engine/Config/Windows/DataDrivenPlatformInfo.ini b/Engine/Config/Windows/DataDrivenPlatformInfo.ini
--- a/Engine/Config/Windows/DataDrivenPlatformInfo.ini
+++ b/Engine/Config/Windows/DataDrivenPlatformInfo.ini
@@ -69,7 +69,6 @@ bSupportsVertexShaderLayer=false
 bSupportsInstancedStereo=true
 bSupportsCapsuleShadows=true
 bSupportsPercentageCloserShadows=true
-bSupportsVolumetricFog=true
 bSupportsIndexBufferUAVs=true
 bSupportsROV=true
 bSupportsOIT=true
@@ -127,7 +126,6 @@ bSupportsVertexShaderLayer=true
 bSupportsInstancedStereo=true
 bSupportsCapsuleShadows=true
 bSupportsPercentageCloserShadows=true
-bSupportsVolumetricFog=true
 bSupportsIndexBufferUAVs=true
 BindlessSupport=AllShaderTypes
 bSupportsROV=true
diff --git a/Engine/Shaders/Private/BasePassPixelShader.usf b/Engine/Shaders/Private/BasePassPixelShader.usf
--- a/Engine/Shaders/Private/BasePassPixelShader.usf
+++ b/Engine/Shaders/Private/BasePassPixelShader.usf
@@ -1302,7 +1302,7 @@ void FPixelShaderInOut_MainPS(
 
 	float4 Fogging = HeightFogging;
 
-#if NEEDS_BASEPASS_PIXEL_VOLUMETRIC_FOGGING && COMPILE_BASEPASS_PIXEL_VOLUMETRIC_FOGGING
+#if NEEDS_BASEPASS_PIXEL_VOLUMETRIC_FOGGING
 	if (FogStruct.ApplyVolumetricFog > 0) 
 	{
 		float3 VolumeUV = ComputeVolumeUV(MaterialParameters.AbsoluteWorldPosition, ResolvedView.WorldToClip);
diff --git a/Engine/Shaders/Private/HeightFogCommon.ush b/Engine/Shaders/Private/HeightFogCommon.ush
--- a/Engine/Shaders/Private/HeightFogCommon.ush
+++ b/Engine/Shaders/Private/HeightFogCommon.ush
@@ -5,7 +5,7 @@
 =============================================================================*/
 
 #ifndef PERMUTATION_SUPPORT_FOG_INSCATTERING_TEXTURE
-#define PERMUTATION_SUPPORT_FOG_INSCATTERING_TEXTURE (FEATURE_LEVEL >= FEATURE_LEVEL_SM4)
+#define PERMUTATION_SUPPORT_FOG_INSCATTERING_TEXTURE 1
 #endif
 
 #ifndef PERMUTATION_SUPPORT_FOG_DIRECTIONAL_LIGHT_INSCATTERING
@@ -13,7 +13,7 @@
 #endif
 
 #ifndef PERMUTATION_SUPPORT_VOLUMETRIC_FOG
-#define PERMUTATION_SUPPORT_VOLUMETRIC_FOG (FEATURE_LEVEL >= FEATURE_LEVEL_SM4)
+#define PERMUTATION_SUPPORT_VOLUMETRIC_FOG 1
 #endif
 
 #ifndef PERMUTATION_SUPPORT_FOG_START_DISTANCE
diff --git a/Engine/Shaders/Private/MobileFog.usf b/Engine/Shaders/Private/MobileFog.usf
--- a/Engine/Shaders/Private/MobileFog.usf
+++ b/Engine/Shaders/Private/MobileFog.usf
@@ -39,16 +39,26 @@ void MobileFogPS(
 	half4 FinalFog = half4(0,0,0,1);
 	float SceneDepth = CalcSceneDepth(TexCoord);
 	float3 WorldPositionRelativeToCamera = ScreenVector.xyz * SceneDepth;
+	float3 WorldPosition = WorldPositionRelativeToCamera + ResolvedView.TranslatedWorldCameraOrigin - LWCHackToFloat(ResolvedView.PreViewTranslation);
 
 // at least one of these is active
 #if PERMUTATION_SUPPORT_HEIGHT_FOG
 	FinalFog = CalculateHeightFog(WorldPositionRelativeToCamera);
 #endif
 
+#if PERMUTATION_SUPPORT_VOLUMETRIC_FOG
+	if (FogStruct.ApplyVolumetricFog > 0) 
+		{
+			const uint EyeIndex = 0;
+			float3 VolumeUV = ComputeVolumeUV(LWCPromote(WorldPosition), ResolvedView.WorldToClip);
+			FinalFog = CombineVolumetricFog(FinalFog, VolumeUV, EyeIndex, SceneDepth);
+		}
+#endif
+
 #if PERMUTATION_SUPPORT_AERIAL_PERSPECTIVE
 	if (ResolvedView.SkyAtmosphereApplyCameraAerialPerspectiveVolume > 0.0f)
 	{
-		float3 WorldPosition = WorldPositionRelativeToCamera + ResolvedView.TranslatedWorldCameraOrigin - LWCHackToFloat(ResolvedView.PreViewTranslation);
+		
 
 	//	float4 NDCPosition = LWCMultiply(MakeLWCVector4(WorldPosition, 1.0f), PrimaryView.WorldToClip);		// LWC_TODO: MakeLWCVector4 seems broken
 		float4 NDCPosition = LWCMultiply(LWCPromote(float4(WorldPosition, 1.0f)), PrimaryView.WorldToClip);	// Using LWCPromote make the math work as expected.
diff --git a/Engine/Shaders/Private/VolumetricFog.usf b/Engine/Shaders/Private/VolumetricFog.usf
--- a/Engine/Shaders/Private/VolumetricFog.usf
+++ b/Engine/Shaders/Private/VolumetricFog.usf
@@ -9,15 +9,19 @@
 #include "LightGridCommon.ush"
 #include "HeightFogCommon.ush"
 #include "SHCommon.ush"
+#if DISTANCE_FIELD_SKY_OCCLUSION
 #include "DistanceFieldAOShared.ush"
 #include "DistanceField/GlobalDistanceFieldShared.ush"
+#endif
 #include "VolumeLightingCommon.ush"
 #include "VolumetricLightmapShared.ush"
 #include "ForwardShadowingCommon.ush"
 #include "ParticipatingMediaCommon.ush"
 #include "LightDataUniforms.ush"
+#if LUMEN_GI
 #define FrontLayerTranslucencyReflectionsStruct LumenGIVolumeStruct
 #include "Lumen/LumenTranslucencyVolumeShared.ush"
+#endif
 #include "Random.ush"
 
 #if VIRTUAL_SHADOW_MAP
@@ -385,6 +389,8 @@ Texture3D<float4> LocalShadowedLightScattering;
 
 RWTexture3D<float4> RWLightScattering;
 
+#if DISTANCE_FIELD_SKY_OCCLUSION
+
 float HemisphereConeTraceAgainstGlobalDistanceFieldClipmap(
 	uniform uint ClipmapIndex,
 	float3 TranslatedWorldShadingPosition,
@@ -443,6 +449,8 @@ float HemisphereConeTraceAgainstGlobalDistanceField(float3 TranslatedWorldShadin
 	return MinVisibility;
 }
 
+#endif
+
 float SkyLightUseStaticShadowing;
 
 float ComputeSkyVisibility(float3 TranslatedWorldPosition, float3 BrickTextureUVs)
diff --git a/Engine/Source/Runtime/RHI/Private/DataDrivenShaderPlatformInfo.cpp b/Engine/Source/Runtime/RHI/Private/DataDrivenShaderPlatformInfo.cpp
--- a/Engine/Source/Runtime/RHI/Private/DataDrivenShaderPlatformInfo.cpp
+++ b/Engine/Source/Runtime/RHI/Private/DataDrivenShaderPlatformInfo.cpp
@@ -134,6 +134,7 @@ void FGenericDataDrivenShaderPlatformInfo::SetDefaultValues()
 	bSupportsManualVertexFetch = true;
 	bSupportsVolumeTextureAtomics = true;
 	bSupportsClipDistance = true;
+	bSupportsDistanceFields = true;
 	PreviewShaderPlatformParent = EShaderPlatform::SP_NumPlatforms;
 }
 
@@ -195,7 +196,6 @@ void FGenericDataDrivenShaderPlatformInfo::ParseDataDrivenShaderInfo(const FConf
 	GET_SECTION_BOOL_HELPER(bSupportsRGBColorBuffer);
 	GET_SECTION_BOOL_HELPER(bSupportsCapsuleShadows);
 	GET_SECTION_BOOL_HELPER(bSupportsPercentageCloserShadows);
-	GET_SECTION_BOOL_HELPER(bSupportsVolumetricFog);
 	GET_SECTION_BOOL_HELPER(bSupportsIndexBufferUAVs);
 	GET_SECTION_BOOL_HELPER(bSupportsInstancedStereo);
 	GET_SECTION_SUPPORT_HELPER(SupportsMultiViewport);
diff --git a/Engine/Source/Runtime/RHI/Public/DataDrivenShaderPlatformInfo.h b/Engine/Source/Runtime/RHI/Public/DataDrivenShaderPlatformInfo.h
--- a/Engine/Source/Runtime/RHI/Public/DataDrivenShaderPlatformInfo.h
+++ b/Engine/Source/Runtime/RHI/Public/DataDrivenShaderPlatformInfo.h
@@ -36,7 +36,6 @@ class FGenericDataDrivenShaderPlatformInfo
 	uint32 bSupportsRGBColorBuffer : 1;
 	uint32 bSupportsCapsuleShadows : 1;
 	uint32 bSupportsPercentageCloserShadows : 1;
-	uint32 bSupportsVolumetricFog : 1; // also used for FVVoxelization
 	uint32 bSupportsIndexBufferUAVs : 1;
 	uint32 bSupportsInstancedStereo : 1;
 	uint32 SupportsMultiViewport : int32(ERHIFeatureSupport::NumBits);
@@ -279,12 +278,6 @@ public:
 		return Infos[Platform].bSupportsPercentageCloserShadows;
 	}
 
-	static FORCEINLINE_DEBUGGABLE const bool GetSupportsVolumetricFog(const FStaticShaderPlatform Platform)
-	{
-		check(IsValid(Platform));
-		return Infos[Platform].bSupportsVolumetricFog;
-	}
-
 	static FORCEINLINE_DEBUGGABLE const bool GetSupportsIndexBufferUAVs(const FStaticShaderPlatform Platform)
 	{
 		check(IsValid(Platform));
diff --git a/Engine/Source/Runtime/Renderer/Private/BasePassRendering.cpp b/Engine/Source/Runtime/Renderer/Private/BasePassRendering.cpp
--- a/Engine/Source/Runtime/Renderer/Private/BasePassRendering.cpp
+++ b/Engine/Source/Runtime/Renderer/Private/BasePassRendering.cpp
@@ -885,7 +885,6 @@ void ModifyBasePassCSPSCompilationEnvironment(const FMeshMaterialShaderPermutati
 {
 	OutEnvironment.SetDefine(TEXT("SCENE_TEXTURES_DISABLED"), Parameters.MaterialParameters.MaterialDomain != MD_Surface);
 	OutEnvironment.SetDefine(TEXT("ENABLE_DBUFFER_TEXTURES"), Parameters.MaterialParameters.MaterialDomain == MD_Surface);
-	OutEnvironment.SetDefine(TEXT("COMPILE_BASEPASS_PIXEL_VOLUMETRIC_FOGGING"), DoesPlatformSupportVolumetricFog(Parameters.Platform));
 	OutEnvironment.SetDefine(TEXT("ENABLE_SKY_LIGHT"), bEnableSkyLight);
 	OutEnvironment.SetDefine(TEXT("PLATFORM_FORCE_SIMPLE_SKY_DIFFUSE"), ForceSimpleSkyDiffuse(Parameters.Platform));
 	OutEnvironment.SetDefine(TEXT("GBUFFER_LAYOUT"), GBufferLayout);
diff --git a/Engine/Source/Runtime/Renderer/Private/DeferredShadingRenderer.h b/Engine/Source/Runtime/Renderer/Private/DeferredShadingRenderer.h
--- a/Engine/Source/Runtime/Renderer/Private/DeferredShadingRenderer.h
+++ b/Engine/Source/Runtime/Renderer/Private/DeferredShadingRenderer.h
@@ -968,44 +968,6 @@ private:
 		FRDGTextureSRVRef SceneStencilTexture,
 		const TArrayView<FRDGTextureRef> NaniteResolveTextures);
 
-	bool ShouldRenderVolumetricFog() const;
-
-	void SetupVolumetricFog();
-
-	void RenderLocalLightsForVolumetricFog(
-		FRDGBuilder& GraphBuilder,
-		FViewInfo& View,
-		bool bUseTemporalReprojection,
-		const struct FVolumetricFogIntegrationParameterData& IntegrationData,
-		const FExponentialHeightFogSceneInfo& FogInfo,
-		FIntVector VolumetricFogGridSize,
-		FVector GridZParams,
-		const FRDGTextureDesc& VolumeDesc,
-		FRDGTexture*& OutLocalShadowedLightScattering,
-		FRDGTextureRef ConservativeDepthTexture);
-
-	void RenderLightFunctionForVolumetricFog(
-		FRDGBuilder& GraphBuilder,
-		FViewInfo& View,
-		const FSceneTextures& SceneTextures,
-		FIntVector VolumetricFogGridSize,
-		float VolumetricFogMaxDistance,
-		FMatrix44f& OutLightFunctionTranslatedWorldToShadow,
-		FRDGTexture*& OutLightFunctionTexture,
-		bool& bOutUseDirectionalLightShadowing);
-
-	void VoxelizeFogVolumePrimitives(
-		FRDGBuilder& GraphBuilder,
-		const FViewInfo& View,
-		const FVolumetricFogIntegrationParameterData& IntegrationData,
-		FIntVector VolumetricFogGridSize,
-		FVector GridZParams,
-		float VolumetricFogDistance,
-		bool bVoxelizeEmissive);
-
-	void ComputeVolumetricFog(FRDGBuilder& GraphBuilder,
-		const FSceneTextures& SceneTextures);
-
 	void RenderHeterogeneousVolumes(FRDGBuilder& GraphBuilder, const FSceneTextures& SceneTextures);
 	void CompositeHeterogeneousVolumes(FRDGBuilder& GraphBuilder, const FSceneTextures& SceneTextures);
 
@@ -1016,8 +978,7 @@ private:
 
 	bool HasDeferredPlanarReflections(const FViewInfo& View) const;
 	void RenderDeferredPlanarReflections(FRDGBuilder& GraphBuilder, const FSceneTextureParameters& SceneTextures, const FViewInfo& View, FRDGTextureRef& ReflectionsOutput);
-	
-	bool ShouldRenderDistanceFieldAO() const;
+
 
 	bool IsNaniteEnabled() const;
 
diff --git a/Engine/Source/Runtime/Renderer/Private/DistanceFieldAmbientOcclusion.cpp b/Engine/Source/Runtime/Renderer/Private/DistanceFieldAmbientOcclusion.cpp
--- a/Engine/Source/Runtime/Renderer/Private/DistanceFieldAmbientOcclusion.cpp
+++ b/Engine/Source/Runtime/Renderer/Private/DistanceFieldAmbientOcclusion.cpp
@@ -879,7 +879,7 @@ void FDeferredShadingSceneRenderer::RenderDistanceFieldLighting(
 	OutDynamicBentNormalAO = BentNormalOutput;
 }
 
-bool FDeferredShadingSceneRenderer::ShouldRenderDistanceFieldAO() const
+bool FSceneRenderer::ShouldRenderDistanceFieldAO() const
 {
 	bool bShouldRenderRTAO = false;
 	for (int ViewIndex = 0; ViewIndex < Views.Num(); ++ViewIndex)
diff --git a/Engine/Source/Runtime/Renderer/Private/MobileFogRendering.cpp b/Engine/Source/Runtime/Renderer/Private/MobileFogRendering.cpp
--- a/Engine/Source/Runtime/Renderer/Private/MobileFogRendering.cpp
+++ b/Engine/Source/Runtime/Renderer/Private/MobileFogRendering.cpp
@@ -62,6 +62,7 @@ class FMobileFogPS : public FGlobalShader
 	class FSupportFogSecondTerm						: SHADER_PERMUTATION_BOOL("PERMUTATION_SUPPORT_FOG_SECOND_TERM");
 	class FSupportFogDirectionalLightInScattering	: SHADER_PERMUTATION_BOOL("PERMUTATION_SUPPORT_FOG_DIRECTIONAL_LIGHT_INSCATTERING");
 	class FSupportAerialPerspective					: SHADER_PERMUTATION_BOOL("PERMUTATION_SUPPORT_AERIAL_PERSPECTIVE");
+	class FSupportVolumetricFog						: SHADER_PERMUTATION_BOOL("PERMUTATION_SUPPORT_VOLUMETRIC_FOG");
 	
 	using FPermutationDomain = TShaderPermutationDomain< 
 		FSupportHeightFog, 
@@ -69,7 +70,8 @@ class FMobileFogPS : public FGlobalShader
 		FSupportFogInScatteringTexture, 
 		FSupportFogSecondTerm,
 		FSupportFogDirectionalLightInScattering,
-		FSupportAerialPerspective 
+		FSupportAerialPerspective,
+		FSupportVolumetricFog
 	>;
 
 	static FPermutationDomain RemapPermutation(FPermutationDomain PermutationVector)
@@ -179,6 +181,7 @@ void FMobileSceneRenderer::RenderFog(FRHICommandListImmediate& RHICmdList, const
 	PsPermutationVector.Set<FMobileFogPS::FSupportFogSecondTerm>(bUseFogSecondTerm);
 	PsPermutationVector.Set<FMobileFogPS::FSupportAerialPerspective>(bUseAerialPerspective);
 	PsPermutationVector.Set<FMobileFogPS::FSupportFogDirectionalLightInScattering>(bUseFogDirectionalInscatering);
+	PsPermutationVector.Set<FMobileFogPS::FSupportVolumetricFog>(ShouldRenderVolumetricFog());
 	
 	TShaderMapRef<FMobileFogPS> PixelShader(View.ShaderMap, PsPermutationVector);
 		
diff --git a/Engine/Source/Runtime/Renderer/Private/MobileShadingRenderer.cpp b/Engine/Source/Runtime/Renderer/Private/MobileShadingRenderer.cpp
--- a/Engine/Source/Runtime/Renderer/Private/MobileShadingRenderer.cpp
+++ b/Engine/Source/Runtime/Renderer/Private/MobileShadingRenderer.cpp
@@ -471,6 +471,11 @@ void FMobileSceneRenderer::InitViews(
 		FXSystem->PostInitViews(GraphBuilder, GetSceneViews(), !ViewFamily.EngineShowFlags.HitProxies);
 	}
 	ComputeViewVisibility(RHICmdList, TaskDatas.VisibilityTaskData, BasePassDepthStencilAccess, DynamicIndexBuffer, DynamicVertexBuffer, DynamicReadBuffer, InstanceCullingManager, VirtualTextureUpdater, ComputeViewVisibilityCallbacks);
+	
+	if(ShouldRenderVolumetricFog())
+	{
+		SetupVolumetricFog();
+	}
 	PostVisibilityFrameSetup(ILCTaskData);
 
 	FIntPoint RenderTargetSize = ViewFamily.RenderTarget->GetSizeXY();
@@ -993,6 +998,11 @@ void FMobileSceneRenderer::Render(FRDGBuilder& GraphBuilder)
 
 	GraphBuilder.SetCommandListStat(GET_STATID(STAT_CLMM_Shadows));
 	RenderShadowDepthMaps(GraphBuilder, InstanceCullingManager, ExternalAccessQueue);
+
+	if(ShouldRenderVolumetricFog())
+	{
+		ComputeVolumetricFog(GraphBuilder, SceneTextures);
+	}
 	ExternalAccessQueue.Submit(GraphBuilder);
 	
 	PollOcclusionQueriesPass(GraphBuilder);
diff --git a/Engine/Source/Runtime/Renderer/Private/SceneRendering.h b/Engine/Source/Runtime/Renderer/Private/SceneRendering.h
--- a/Engine/Source/Runtime/Renderer/Private/SceneRendering.h
+++ b/Engine/Source/Runtime/Renderer/Private/SceneRendering.h
@@ -64,7 +64,7 @@ class FLumenSceneData;
 class FShadowSceneRenderer;
 class FGlobalShaderMap;
 class FVirtualTextureUpdater;
-
+class FExponentialHeightFogSceneInfo;
 struct FCloudRenderContext;
 struct FSceneWithoutWaterTextures;
 struct FHairStrandsVisibilityViews;
@@ -84,6 +84,7 @@ class FSimpleLightArray;
 struct FNaniteMaterialPassCommand;
 struct FScreenMessageWriter;
 struct FVisibilityTaskData;
+struct FVolumetricFogIntegrationParameterData;
 
 DECLARE_GPU_DRAWCALL_STAT_EXTERN(VirtualTextureUpdate);
 
@@ -2232,11 +2233,48 @@ public:
 	bool ShouldPrepareGlobalDistanceField() const;
 	bool ShouldPrepareHeightFieldScene() const;
 
+	bool ShouldRenderDistanceFieldAO() const;
 	void UpdateGlobalDistanceFieldObjectBuffers(FRDGBuilder& GraphBuilder);
 	void UpdateGlobalHeightFieldObjectBuffers(FRDGBuilder& GraphBuilder, const TArray<uint32>& IndicesToUpdateInHeightFieldObjectBuffers);
 	void ProcessPendingHeightFieldPrimitiveAddAndRemoveOps(TArray<uint32>& IndicesToUpdateInHeightFieldObjectBuffers);
 	void PrepareDistanceFieldScene(FRDGBuilder& GraphBuilder, FRDGExternalAccessQueue& ExternalAccessQueue);
 
+	void VoxelizeFogVolumePrimitives(
+		FRDGBuilder& GraphBuilder,
+		const FViewInfo& View,
+		const FVolumetricFogIntegrationParameterData& IntegrationData,
+		FIntVector VolumetricFogGridSize,
+		FVector GridZParams,
+		float VolumetricFogDistance,
+		bool bVoxelizeEmissive);
+
+	void RenderLightFunctionForVolumetricFog(
+		FRDGBuilder& GraphBuilder,
+		FViewInfo& View,
+		const FSceneTextures& SceneTextures,
+		FIntVector VolumetricFogGridSize,
+		float VolumetricFogMaxDistance,
+		FMatrix44f& OutLightFunctionTranslatedWorldToShadow,
+		FRDGTexture*& OutLightFunctionTexture,
+		bool& bOutUseDirectionalLightShadowing);
+
+	void RenderLocalLightsForVolumetricFog(
+		FRDGBuilder& GraphBuilder,
+		FViewInfo& View,
+		bool bUseTemporalReprojection,
+		const struct FVolumetricFogIntegrationParameterData& IntegrationData,
+		const FExponentialHeightFogSceneInfo& FogInfo,
+		FIntVector VolumetricFogGridSize,
+		FVector GridZParams,
+		const FRDGTextureDesc& VolumeDesc,
+		FRDGTexture*& OutLocalShadowedLightScattering,
+		FRDGTextureRef ConservativeDepthTexture);
+
+	void SetupVolumetricFog();
+
+	bool ShouldRenderVolumetricFog() const;
+	void ComputeVolumetricFog(FRDGBuilder& GraphBuilder, const FSceneTextures& SceneTextures);
+
 	void DrawGPUSkinCacheVisualizationInfoText();
 
 	virtual bool IsLumenEnabled(const FViewInfo& View) const { return false; }
diff --git a/Engine/Source/Runtime/Renderer/Private/ShadowDepthRendering.cpp b/Engine/Source/Runtime/Renderer/Private/ShadowDepthRendering.cpp
--- a/Engine/Source/Runtime/Renderer/Private/ShadowDepthRendering.cpp
+++ b/Engine/Source/Runtime/Renderer/Private/ShadowDepthRendering.cpp
@@ -1645,7 +1645,7 @@ void FSceneRenderer::RenderShadowDepthMapAtlases(FRDGBuilder& GraphBuilder)
 
 		// Make readable because AtlasDepthTexture is not tracked via RDG yet
 		// On mobile CSM atlas sampled only in pixel shaders
-		ERHIAccess AtlasDepthTextureAccessFinal = (FeatureLevel == ERHIFeatureLevel::ES3_1 ? ERHIAccess::SRVGraphics : ERHIAccess::SRVMask);
+		ERHIAccess AtlasDepthTextureAccessFinal = (IsMobilePlatform(ShaderPlatform) && !ShouldRenderVolumetricFog()? ERHIAccess::SRVGraphics : ERHIAccess::SRVMask);
 		ShadowMapAtlas.RenderTargets.DepthTarget = ConvertToExternalAccessTexture(GraphBuilder, ExternalAccessQueue, AtlasDepthTexture, AtlasDepthTextureAccessFinal);
 	}
 
diff --git a/Engine/Source/Runtime/Renderer/Private/VolumetricFog.cpp b/Engine/Source/Runtime/Renderer/Private/VolumetricFog.cpp
--- a/Engine/Source/Runtime/Renderer/Private/VolumetricFog.cpp
+++ b/Engine/Source/Runtime/Renderer/Private/VolumetricFog.cpp
@@ -228,11 +228,6 @@ class FVolumetricFogMaterialSetupCS : public FGlobalShader
 
 public:
 
-	static bool ShouldCompilePermutation(const FGlobalShaderPermutationParameters& Parameters)
-	{
-		return DoesPlatformSupportVolumetricFog(Parameters.Platform);
-	}
-
 	static void ModifyCompilationEnvironment(const FGlobalShaderPermutationParameters& Parameters, FShaderCompilerEnvironment& OutEnvironment)
 	{
 		FGlobalShader::ModifyCompilationEnvironment(Parameters, OutEnvironment);
@@ -256,11 +251,6 @@ class FWriteToBoundingSphereVS : public FGlobalShader
 		SHADER_PARAMETER(int32, MinZ)
 	END_SHADER_PARAMETER_STRUCT()
 
-	static bool ShouldCompilePermutation(const FGlobalShaderPermutationParameters& Parameters)
-	{
-		return DoesPlatformSupportVolumetricFog(Parameters.Platform);
-	}
-
 	static void ModifyCompilationEnvironment(const FGlobalShaderPermutationParameters& Parameters, FShaderCompilerEnvironment& OutEnvironment)
 	{
 		FGlobalShader::ModifyCompilationEnvironment(Parameters, OutEnvironment);
@@ -314,11 +304,6 @@ class FInjectShadowedLocalLightPS : public FGlobalShader
 		FGlobalShader::ModifyCompilationEnvironment(Parameters, OutEnvironment);
 		FVirtualShadowMapArray::SetShaderDefines(OutEnvironment);
 	}
-
-	static bool ShouldCompilePermutation(const FGlobalShaderPermutationParameters& Parameters)
-	{
-		return DoesPlatformSupportVolumetricFog(Parameters.Platform);
-	}
 };
 
 IMPLEMENT_GLOBAL_SHADER(FInjectShadowedLocalLightPS, "/Engine/Private/VolumetricFog.usf", "InjectShadowedLocalLightPS", SF_Pixel);
@@ -459,7 +444,7 @@ public:
 
 TGlobalResource<FCircleRasterizeIndexBuffer> GCircleRasterizeIndexBuffer;
 
-void FDeferredShadingSceneRenderer::RenderLocalLightsForVolumetricFog(
+void FSceneRenderer::RenderLocalLightsForVolumetricFog(
 	FRDGBuilder& GraphBuilder,
 	FViewInfo& View,
 	bool bUseTemporalReprojection,
@@ -695,11 +680,6 @@ class FVolumetricFogLightScatteringCS : public FGlobalShader
 		SHADER_PARAMETER(uint32, UseEmissive)
 	END_SHADER_PARAMETER_STRUCT()
 
-	static bool ShouldCompilePermutation(const FGlobalShaderPermutationParameters& Parameters)
-	{
-		return DoesPlatformSupportVolumetricFog(Parameters.Platform);
-	}
-
 	static FIntVector GetGroupSize()
 	{
 		return FIntVector(4, 4, 4);
@@ -723,6 +703,33 @@ class FVolumetricFogLightScatteringCS : public FGlobalShader
 		return 16;
 	}
 
+	static FPermutationDomain RemapPermutation(FPermutationDomain PermutationVector, EShaderPlatform ShaderPlatform)
+	{
+		if (IsMobilePlatform(ShaderPlatform))
+		{
+			PermutationVector.Set<FDistanceFieldSkyOcclusion>(false);
+			PermutationVector.Set<FCloudTransmittance>(false);
+			PermutationVector.Set<FTemporalReprojection>(false);
+		}
+
+		if (!FDataDrivenShaderPlatformInfo::GetSupportsLumenGI(ShaderPlatform))
+		{
+			PermutationVector.Set<FLumenGI>(false);
+		}
+		return PermutationVector;
+	}
+
+	static bool ShouldCompilePermutation(const FGlobalShaderPermutationParameters& Parameters)
+	{
+		FPermutationDomain PermutationVector(Parameters.PermutationId);
+		if (RemapPermutation(PermutationVector, Parameters.Platform) != PermutationVector)
+		{
+			return false;
+		}
+
+		return true;
+	}
+
 	static void ModifyCompilationEnvironment(const FGlobalShaderPermutationParameters& Parameters, FShaderCompilerEnvironment& OutEnvironment)
 	{
 		FGlobalShader::ModifyCompilationEnvironment(Parameters, OutEnvironment);
@@ -751,11 +758,6 @@ class FVolumetricFogFinalIntegrationCS : public FGlobalShader
 		SHADER_PARAMETER_STRUCT_INCLUDE(FVolumetricFogIntegrationParameters, VolumetricFogParameters)
 	END_SHADER_PARAMETER_STRUCT()
 
-	static bool ShouldCompilePermutation(const FGlobalShaderPermutationParameters& Parameters)
-	{
-		return DoesPlatformSupportVolumetricFog(Parameters.Platform);
-	}
-
 	static void ModifyCompilationEnvironment(const FGlobalShaderPermutationParameters& Parameters, FShaderCompilerEnvironment& OutEnvironment)
 	{
 		FGlobalShader::ModifyCompilationEnvironment(Parameters, OutEnvironment);
@@ -765,22 +767,14 @@ class FVolumetricFogFinalIntegrationCS : public FGlobalShader
 
 IMPLEMENT_GLOBAL_SHADER(FVolumetricFogFinalIntegrationCS, "/Engine/Private/VolumetricFog.usf", "FinalIntegrationCS", SF_Compute);
 
-bool DoesPlatformSupportVolumetricFog(const FStaticShaderPlatform Platform)
-{
-	return FDataDrivenShaderPlatformInfo::GetSupportsVolumetricFog(Platform);
-}
-
 bool DoesPlatformSupportVolumetricFogVoxelization(const FStaticShaderPlatform Platform)
 {
-	return DoesPlatformSupportVolumetricFog(Platform);
+	return !IsMobilePlatform(Platform);
 }
-
 bool ShouldRenderVolumetricFog(const FScene* Scene, const FSceneViewFamily& ViewFamily)
 {
 	return ShouldRenderFog(ViewFamily)
 		&& Scene
-		&& Scene->GetFeatureLevel() >= ERHIFeatureLevel::SM5
-		&& DoesPlatformSupportVolumetricFog(Scene->GetShaderPlatform())
 		&& GVolumetricFog
 		&& ViewFamily.EngineShowFlags.VolumetricFog
 		&& Scene->ExponentialFogs.Num() > 0
@@ -920,12 +914,12 @@ void FViewInfo::SetupVolumetricFogUniformBufferParameters(FViewUniformShaderPara
 	}
 }
 
-bool FDeferredShadingSceneRenderer::ShouldRenderVolumetricFog() const
+bool FSceneRenderer::ShouldRenderVolumetricFog() const
 {
 	return ::ShouldRenderVolumetricFog(Scene, ViewFamily);
 }
 
-void FDeferredShadingSceneRenderer::SetupVolumetricFog()
+void FSceneRenderer::SetupVolumetricFog()
 {
 	if (ShouldRenderVolumetricFog())
 	{
@@ -959,7 +953,7 @@ void FDeferredShadingSceneRenderer::SetupVolumetricFog()
 	}
 }
 
-void FDeferredShadingSceneRenderer::ComputeVolumetricFog(FRDGBuilder& GraphBuilder,
+void FSceneRenderer::ComputeVolumetricFog(FRDGBuilder& GraphBuilder,
 	const FSceneTextures& SceneTextures)
 {
 	if (!ShouldRenderVolumetricFog())
@@ -969,7 +963,7 @@ void FDeferredShadingSceneRenderer::ComputeVolumetricFog(FRDGBuilder& GraphBuild
 
 	const FExponentialHeightFogSceneInfo& FogInfo = Scene->ExponentialFogs[0];
 
-	TRACE_CPUPROFILER_EVENT_SCOPE(FDeferredShadingSceneRenderer::ComputeVolumetricFog);
+	TRACE_CPUPROFILER_EVENT_SCOPE(FSceneRenderer::ComputeVolumetricFog);
 	QUICK_SCOPE_CYCLE_COUNTER(STAT_VolumetricFog);
 	RDG_CSV_STAT_EXCLUSIVE_SCOPE(GraphBuilder, VolumetricFog);
 	RDG_GPU_STAT_SCOPE(GraphBuilder, VolumetricFog);
@@ -995,7 +989,8 @@ void FDeferredShadingSceneRenderer::ComputeVolumetricFog(FRDGBuilder& GraphBuild
 
 		const bool bUseTemporalReprojection =
 			GVolumetricFogTemporalReprojection
-			&& View.ViewState;
+			&& View.ViewState
+			&& !IsMobilePlatform(View.GetShaderPlatform());
 
 		IntegrationData.bTemporalHistoryIsValid =
 			bUseTemporalReprojection
@@ -1252,7 +1247,7 @@ void FDeferredShadingSceneRenderer::ComputeVolumetricFog(FRDGBuilder& GraphBuild
 			PassParameters->CloudShadowmapStrength = CloudShadowmap_Strength;
 			PassParameters->CloudShadowmapTranslatedWorldToLightClipMatrix = CloudWorldToLightClipShadowMatrix;
 
-			const bool bUseLumenGI = View.LumenTranslucencyGIVolume.Texture0 != nullptr;
+			const bool bUseLumenGI = View.LumenTranslucencyGIVolume.Texture0 != nullptr && FDataDrivenShaderPlatformInfo::GetSupportsLumenGI(View.GetShaderPlatform());
 			const bool bUseGlobalDistanceField = UseGlobalDistanceField() && Scene->DistanceFieldSceneData.NumObjectsInBuffer > 0;
 
 			const bool bUseDistanceFieldSkyOcclusion =
@@ -1265,7 +1260,8 @@ void FDeferredShadingSceneRenderer::ComputeVolumetricFog(FRDGBuilder& GraphBuild
 				&& SupportsDistanceFieldAO(View.GetFeatureLevel(), View.GetShaderPlatform())
 				&& bUseGlobalDistanceField
 				&& Views.Num() == 1
-				&& View.IsPerspectiveProjection();
+				&& View.IsPerspectiveProjection()
+				&& !IsMobilePlatform(View.GetShaderPlatform());
 
 			const int32 SuperSampleCount = FVolumetricFogLightScatteringCS::GetSuperSampleCount(GVolumetricFogHistoryMissSupersampleCount);
 
diff --git a/Engine/Source/Runtime/Renderer/Private/VolumetricFog.h b/Engine/Source/Runtime/Renderer/Private/VolumetricFog.h
--- a/Engine/Source/Runtime/Renderer/Private/VolumetricFog.h
+++ b/Engine/Source/Runtime/Renderer/Private/VolumetricFog.h
@@ -50,9 +50,7 @@ private:
 extern FIntVector GetVolumetricFogGridSize(FIntPoint ViewRectSize, int32& OutVolumetricFogGridPixelSize);
 extern FVector2f GetVolumetricFogFroxelToScreenSVPosRatio(FIntPoint ViewRectSize);
 
-extern bool DoesPlatformSupportVolumetricFog(const FStaticShaderPlatform Platform);
 extern bool DoesPlatformSupportVolumetricFogVoxelization(const FStaticShaderPlatform Platform);
-
 extern bool ShouldRenderVolumetricFog(const FScene* Scene, const FSceneViewFamily& ViewFamily);
 extern const FProjectedShadowInfo* GetShadowForInjectionIntoVolumetricFog(const FVisibleLightInfo& VisibleLightInfo);
 
diff --git a/Engine/Source/Runtime/Renderer/Private/VolumetricFogLightFunction.cpp b/Engine/Source/Runtime/Renderer/Private/VolumetricFogLightFunction.cpp
--- a/Engine/Source/Runtime/Renderer/Private/VolumetricFogLightFunction.cpp
+++ b/Engine/Source/Runtime/Renderer/Private/VolumetricFogLightFunction.cpp
@@ -81,7 +81,7 @@ class FVolumetricFogLightFunctionPS : public FMaterialShader
 
 	static bool ShouldCompilePermutation(const FMaterialShaderPermutationParameters& Parameters)
 	{
-		return Parameters.MaterialParameters.MaterialDomain == MD_LightFunction && DoesPlatformSupportVolumetricFog(Parameters.Platform);
+		return Parameters.MaterialParameters.MaterialDomain == MD_LightFunction;
 	}
 
 	void SetParameters(FRHIBatchedShaderParameters& BatchedParameters, const FSceneView& View, const FMaterialRenderProxy* MaterialProxy)
@@ -132,7 +132,7 @@ END_SHADER_PARAMETER_STRUCT()
 
 
 
-void FDeferredShadingSceneRenderer::RenderLightFunctionForVolumetricFog(
+void FSceneRenderer::RenderLightFunctionForVolumetricFog(
 	FRDGBuilder& GraphBuilder,
 	FViewInfo& View,
 	const FSceneTextures& SceneTextures,
diff --git a/Engine/Source/Runtime/Renderer/Private/VolumetricFogVoxelization.cpp b/Engine/Source/Runtime/Renderer/Private/VolumetricFogVoxelization.cpp
--- a/Engine/Source/Runtime/Renderer/Private/VolumetricFogVoxelization.cpp
+++ b/Engine/Source/Runtime/Renderer/Private/VolumetricFogVoxelization.cpp
@@ -186,8 +186,7 @@ public:
 
 	static bool ShouldCompilePermutation(const FMeshMaterialShaderPermutationParameters& Parameters)
 	{
-		return IsFeatureLevelSupported(Parameters.Platform, ERHIFeatureLevel::SM5) 
-			&& DoesPlatformSupportVolumetricFogVoxelization(Parameters.Platform)
+		return DoesPlatformSupportVolumetricFogVoxelization(Parameters.Platform)
 			&& Parameters.MaterialParameters.MaterialDomain == MD_Volume;
 	}
 
@@ -271,8 +270,8 @@ public:
 
 	static bool ShouldCompilePermutation(const FMeshMaterialShaderPermutationParameters& Parameters)
 	{
-		return IsFeatureLevelSupported(Parameters.Platform, ERHIFeatureLevel::SM5) 
-			&& RHISupportsGeometryShaders(Parameters.Platform)
+		return
+			RHISupportsGeometryShaders(Parameters.Platform)
 			&& DoesPlatformSupportVolumetricFogVoxelization(Parameters.Platform)
 			&& Parameters.MaterialParameters.MaterialDomain == MD_Volume;
 	}
@@ -344,8 +343,7 @@ public:
 
 	static bool ShouldCompilePermutation(const FMeshMaterialShaderPermutationParameters& Parameters)
 	{
-		return IsFeatureLevelSupported(Parameters.Platform, ERHIFeatureLevel::SM5) 
-			&& DoesPlatformSupportVolumetricFogVoxelization(Parameters.Platform)
+		return DoesPlatformSupportVolumetricFogVoxelization(Parameters.Platform)
 			&& Parameters.MaterialParameters.MaterialDomain == MD_Volume;
 	}
 };
@@ -599,7 +597,7 @@ BEGIN_SHADER_PARAMETER_STRUCT(FVoxelizeVolumePassParameters, )
 	RENDER_TARGET_BINDING_SLOTS()
 END_SHADER_PARAMETER_STRUCT()
 
-void FDeferredShadingSceneRenderer::VoxelizeFogVolumePrimitives(
+void FSceneRenderer::VoxelizeFogVolumePrimitives(
 	FRDGBuilder& GraphBuilder,
 	const FViewInfo& View,
 	const FVolumetricFogIntegrationParameterData& IntegrationData,
