Change Substrate to use the visible material 'max byte per pixel' (statically known at shader compiler time) to allocate Substrate material buffer, instead of using the constant project settings.

diff --git a/Engine/Source/Runtime/Engine/Classes/Materials/MaterialRelevance.h b/Engine/Source/Runtime/Engine/Classes/Materials/MaterialRelevance.h
--- a/Engine/Source/Runtime/Engine/Classes/Materials/MaterialRelevance.h
+++ b/Engine/Source/Runtime/Engine/Classes/Materials/MaterialRelevance.h
@@ -17,6 +17,7 @@ struct ENGINE_API FMaterialRelevance
 		struct
 		{
 			uint16 ShadingModelMask;
+			uint8 StrataUintPerPixel;
 			uint8 StrataBSDFCountMask;
 			uint8 bOpaque : 1;
 			uint8 bMasked : 1;
@@ -55,7 +56,9 @@ struct ENGINE_API FMaterialRelevance
 	/** Bitwise OR operator.  Sets any relevance bits which are present in either. */
 	FMaterialRelevance& operator|=(const FMaterialRelevance& B)
 	{
+		const uint8 LocalStrataUintPerPixel = StrataUintPerPixel;
 		Raw |= B.Raw;
+		StrataUintPerPixel = FMath::Max(LocalStrataUintPerPixel, B.StrataUintPerPixel);
 		return *this;
 	}
 
diff --git a/Engine/Source/Runtime/Engine/Private/Materials/HLSLMaterialTranslator.cpp b/Engine/Source/Runtime/Engine/Private/Materials/HLSLMaterialTranslator.cpp
--- a/Engine/Source/Runtime/Engine/Private/Materials/HLSLMaterialTranslator.cpp
+++ b/Engine/Source/Runtime/Engine/Private/Materials/HLSLMaterialTranslator.cpp
@@ -10971,8 +10971,12 @@ bool FHLSLMaterialTranslator::StrataGenerateDerivedMaterialOperatorData()
 			}
 			StrataSimplificationStatus.bRunFullSimplification = !StrataSimplificationStatus.bMaterialFitsInMemoryBudget;
 
+			const uint32 RequestedSizeInUint = FMath::DivideAndRoundUp(StrataMaterialRequestedSizeByte, 4u);
+			check(RequestedSizeInUint < 256u);
+
 			MaterialCompilationOutput.StrataMaterialType = bStrataMaterialIsSimple ? 0 : bStrataMaterialIsSingle? 1 : 2;
 			MaterialCompilationOutput.StrataBSDFCount = StrataMaterialBSDFCount;
+			MaterialCompilationOutput.StrataUintPerPixel = uint8(FMath::Clamp(RequestedSizeInUint, 0u, 0xFF));
 		}
 	}
 	while (!StrataSimplificationStatus.bMaterialFitsInMemoryBudget);
diff --git a/Engine/Source/Runtime/Engine/Private/Materials/MaterialInterface.cpp b/Engine/Source/Runtime/Engine/Private/Materials/MaterialInterface.cpp
--- a/Engine/Source/Runtime/Engine/Private/Materials/MaterialInterface.cpp
+++ b/Engine/Source/Runtime/Engine/Private/Materials/MaterialInterface.cpp
@@ -460,6 +460,7 @@ FMaterialRelevance UMaterialInterface::GetRelevance_Internal(const UMaterial* Ma
 		// Encode Strata BSDF into a mask where each bit correspond to a number of BSDF (1-8)
 		const uint8 StrataBSDFCount = FMath::Max(MaterialResource->MaterialGetStrataBSDFCount_GameThread(), uint8(1u));
 		const uint8 StrataBSDFCountMask = 1u << uint8(FMath::Min(StrataBSDFCount - 1, 8));
+		const uint8 StrataUintPerPixel = FMath::Max(MaterialResource->MaterialGetStrataUintPerPixel_GameThread(), uint8(1u));
 
 		MaterialRelevance.bOpaque = !bIsTranslucent;
 		MaterialRelevance.bMasked = IsMasked();
@@ -485,6 +486,7 @@ FMaterialRelevance UMaterialInterface::GetRelevance_Internal(const UMaterial* Ma
 		MaterialRelevance.bUsesSingleLayerWaterMaterial = bUsesSingleLayerWaterMaterial;
 		MaterialRelevance.bUsesAnisotropy = bUsesAnisotropy;
 		MaterialRelevance.StrataBSDFCountMask = StrataBSDFCountMask;
+		MaterialRelevance.StrataUintPerPixel = StrataUintPerPixel;
 		return MaterialRelevance;
 	}
 	else
diff --git a/Engine/Source/Runtime/Engine/Private/Materials/MaterialShared.cpp b/Engine/Source/Runtime/Engine/Private/Materials/MaterialShared.cpp
--- a/Engine/Source/Runtime/Engine/Private/Materials/MaterialShared.cpp
+++ b/Engine/Source/Runtime/Engine/Private/Materials/MaterialShared.cpp
@@ -1117,6 +1117,17 @@ uint8 FMaterial::MaterialGetStrataBSDFCount_RenderThread() const
 	return RenderingThreadShaderMap ? RenderingThreadShaderMap->GetStrataBSDFCount() : false;
 }
 
+uint8 FMaterial::MaterialGetStrataUintPerPixel_GameThread() const
+{
+	return GameThreadShaderMap ? GameThreadShaderMap->GetStrataUintPerPixel() : false;
+}
+
+uint8 FMaterial::MaterialGetStrataUintPerPixel_RenderThread() const
+{
+	check(IsInParallelRenderingThread());
+	return RenderingThreadShaderMap ? RenderingThreadShaderMap->GetStrataUintPerPixel() : false;
+}
+
 void FMaterial::SetGameThreadShaderMap(FMaterialShaderMap* InMaterialShaderMap)
 {
 	checkSlow(IsInGameThread() || IsInAsyncLoadingThread());
diff --git a/Engine/Source/Runtime/Engine/Public/MaterialShared.h b/Engine/Source/Runtime/Engine/Public/MaterialShared.h
--- a/Engine/Source/Runtime/Engine/Public/MaterialShared.h
+++ b/Engine/Source/Runtime/Engine/Public/MaterialShared.h
@@ -691,7 +691,8 @@ public:
 		bHasRuntimeVirtualTextureOutputNode(false),
 		bUsesAnisotropy(false),
 		StrataMaterialType(0),
-		StrataBSDFCount(0)
+		StrataBSDFCount(0),
+		StrataUintPerPixel(0)
 	{}
 
 	ENGINE_API bool IsSceneTextureUsed(ESceneTextureId TexId) const { return (UsedSceneTextures & (1 << TexId)) != 0; }
@@ -804,11 +805,15 @@ public:
 	/** true if the material uses non 0 anisotropy value */
 	LAYOUT_BITFIELD(uint8, bUsesAnisotropy, 1);
 
-	/** Strata material type, at compile type (0:simple, 1:single, 2: complex) */
+	/** Strata material type, at compile time (0:simple, 1:single, 2: complex) */
 	LAYOUT_BITFIELD(uint8, StrataMaterialType, 2);
 
-	/** Strata BSDF count, at compile type (0-7) */
+	/** Strata BSDF count, at compile time (0-7) */
 	LAYOUT_BITFIELD(uint8, StrataBSDFCount, 3);
+
+	/** Strata uint per pixel, at compile time (0-7) */
+	LAYOUT_BITFIELD(uint8, StrataUintPerPixel, 8);
+	
 };
 
 struct FDebugShaderPipelineInfo
@@ -1370,6 +1375,7 @@ public:
 	bool UsesAnisotropy() const { return GetContent()->MaterialCompilationOutput.bUsesAnisotropy; }
 	uint8 GetStrataMaterialType() const { return GetContent()->MaterialCompilationOutput.StrataMaterialType; }
 	uint8 GetStrataBSDFCount() const { return GetContent()->MaterialCompilationOutput.StrataBSDFCount; }
+	uint8 GetStrataUintPerPixel() const { return GetContent()->MaterialCompilationOutput.StrataUintPerPixel; }
 	
 #if WITH_EDITOR
 	uint32 GetNumUsedUVScalars() const { return GetContent()->MaterialCompilationOutput.NumUsedUVScalars; }
@@ -2118,6 +2124,10 @@ public:
 	ENGINE_API uint8 MaterialGetStrataBSDFCount_GameThread() const;
 	ENGINE_API uint8 MaterialGetStrataBSDFCount_RenderThread() const;
 
+	/** Get Strata material uint count oer pixel. */
+	ENGINE_API uint8 MaterialGetStrataUintPerPixel_GameThread() const;
+	ENGINE_API uint8 MaterialGetStrataUintPerPixel_RenderThread() const;
+
 	class FMaterialShaderMap* GetGameThreadShaderMap() const 
 	{ 
 		checkSlow(IsInGameThread() || IsInParallelGameThread() || IsInAsyncLoadingThread());
diff --git a/Engine/Source/Runtime/Renderer/Private/SceneVisibility.cpp b/Engine/Source/Runtime/Renderer/Private/SceneVisibility.cpp
--- a/Engine/Source/Runtime/Renderer/Private/SceneVisibility.cpp
+++ b/Engine/Source/Runtime/Renderer/Private/SceneVisibility.cpp
@@ -2372,6 +2372,7 @@ struct FRelevancePacket : public FSceneRenderingAllocatorObject<FRelevancePacket
 	FRelevancePrimSet<FPrimitiveLODMask> PrimitivesLODMask; // group both lod mask with primitive index to be able to properly merge them in the view
 
 	uint16 CombinedShadingModelMask;
+	uint8 StrataUintPerPixel;
 	uint8 StrataBSDFCountMask;
 	bool bUsesGlobalDistanceField;
 	bool bUsesLightingChannels;
@@ -2411,6 +2412,7 @@ struct FRelevancePacket : public FSceneRenderingAllocatorObject<FRelevancePacket
 		, bHasDistortionPrimitives(false)
 		, bHasCustomDepthPrimitives(false)
 		, CombinedShadingModelMask(0)
+		, StrataUintPerPixel(0)
 		, StrataBSDFCountMask(0)
 		, bUsesGlobalDistanceField(false)
 		, bUsesLightingChannels(false)
@@ -2435,6 +2437,7 @@ struct FRelevancePacket : public FSceneRenderingAllocatorObject<FRelevancePacket
 	void ComputeRelevance()
 	{
 		CombinedShadingModelMask = 0;
+		StrataUintPerPixel = 0;
 		StrataBSDFCountMask = 0;
 		bSceneHasSkyMaterial = 0;
 		bHasSingleLayerWaterMaterial = 0;
@@ -2564,6 +2567,7 @@ struct FRelevancePacket : public FSceneRenderingAllocatorObject<FRelevancePacket
 			}
 			
 			CombinedShadingModelMask |= ViewRelevance.ShadingModelMask;
+			StrataUintPerPixel = FMath::Max(StrataUintPerPixel, ViewRelevance.StrataUintPerPixel);
 			StrataBSDFCountMask |= ViewRelevance.StrataBSDFCountMask;
 			bUsesGlobalDistanceField |= ViewRelevance.bUsesGlobalDistanceField;
 			bUsesLightingChannels |= ViewRelevance.bUsesLightingChannels;
@@ -3054,6 +3058,7 @@ struct FRelevancePacket : public FSceneRenderingAllocatorObject<FRelevancePacket
 		WriteView.bUsesCustomDepth |= bUsesCustomDepth;
 		WriteView.bUsesCustomStencil |= bUsesCustomStencil;
 		WriteView.StrataViewData.MaxBSDFCount = FMath::Max(WriteView.StrataViewData.MaxBSDFCount, 8u - FMath::CountLeadingZeros8(StrataBSDFCountMask));
+		WriteView.StrataViewData.MaxBytePerPixel = FMath::Max(WriteView.StrataViewData.MaxBytePerPixel, StrataUintPerPixel * 4u);
 		DirtyIndirectLightingCacheBufferPrimitives.AppendTo(WriteView.DirtyIndirectLightingCacheBufferPrimitives);
 
 		WriteView.MeshDecalBatches.Append(MeshDecalBatches);
diff --git a/Engine/Source/Runtime/Renderer/Private/Strata/Strata.cpp b/Engine/Source/Runtime/Renderer/Private/Strata/Strata.cpp
--- a/Engine/Source/Runtime/Renderer/Private/Strata/Strata.cpp
+++ b/Engine/Source/Runtime/Renderer/Private/Strata/Strata.cpp
@@ -69,8 +69,10 @@ void FStrataViewData::Reset()
 {
 	// Preserve old max BSDF count which is set prior to the reset operation
 	const uint32 OldMaxBSDFCount = MaxBSDFCount;
+	const uint32 OldMaxBytePerPixel = MaxBytePerPixel;
 	*this = FStrataViewData();
 	MaxBSDFCount = OldMaxBSDFCount;
+	MaxBytePerPixel = OldMaxBytePerPixel;
 	for (uint32 i = 0; i < EStrataTileType::ECount; ++i)
 	{
 		ClassificationTileListBuffer[i] = nullptr;
@@ -334,11 +336,14 @@ void InitialiseStrataFrameSceneData(FRDGBuilder& GraphBuilder, FSceneRenderer& S
 
 	bool bNeedBSDFOffsets = false;
 	bool bNeedUAV = false;
+	uint32 MaxBytePerPixel = 0;
 	for (const FViewInfo& View : SceneRenderer.Views)
 	{
 		bNeedBSDFOffsets = bNeedBSDFOffsets || NeedBSDFOffsets(SceneRenderer.Scene, View);
 		bNeedUAV = bNeedUAV || IsDBufferPassEnabled(View.GetShaderPlatform());
+		MaxBytePerPixel = FMath::Max(MaxBytePerPixel, View.StrataViewData.MaxBytePerPixel);
 	}
+	MaxBytePerPixel = FMath::Clamp(MaxBytePerPixel, 4u * STRATA_BASE_PASS_MRT_OUTPUT_COUNT, GetBytePerPixel(SceneRenderer.ShaderPlatform));
 
 	FIntPoint MaterialBufferSizeXY;
 	UpdateMaterialBufferToTiledResolution(FIntPoint(1, 1), MaterialBufferSizeXY);
@@ -358,7 +363,7 @@ void InitialiseStrataFrameSceneData(FRDGBuilder& GraphBuilder, FSceneRenderer& S
 		UpdateMaterialBufferToTiledResolution(SceneTextureExtent, MaterialBufferSizeXY);
 
 		const uint32 RoundToValue = 4u;
-		Out.MaxBytesPerPixel = FMath::DivideAndRoundUp(GetBytePerPixel(SceneRenderer.ShaderPlatform), RoundToValue) * RoundToValue;
+		Out.MaxBytesPerPixel = FMath::DivideAndRoundUp(MaxBytePerPixel, RoundToValue) * RoundToValue;
 
 		// Top layer texture
 		{
diff --git a/Engine/Source/Runtime/Renderer/Private/Strata/Strata.h b/Engine/Source/Runtime/Renderer/Private/Strata/Strata.h
--- a/Engine/Source/Runtime/Renderer/Private/Strata/Strata.h
+++ b/Engine/Source/Runtime/Renderer/Private/Strata/Strata.h
@@ -134,6 +134,7 @@ struct FStrataViewData
 {
 	// Max BSDF count among all visible materials
 	uint32 MaxBSDFCount = 0;
+	uint32 MaxBytePerPixel = 0;
 
 	FIntPoint TileCount  = FIntPoint(0, 0);
 	FIntPoint TileOffset = FIntPoint(0, 0);
