Move bUpdating=false earlier in vt system update to where we complete the async update.
Fixes potential false positive check coming from RVT GPU scene update later in function.

diff --git a/Engine/Source/Runtime/Renderer/Private/VT/VirtualTextureSystem.cpp b/Engine/Source/Runtime/Renderer/Private/VT/VirtualTextureSystem.cpp
--- a/Engine/Source/Runtime/Renderer/Private/VT/VirtualTextureSystem.cpp
+++ b/Engine/Source/Runtime/Renderer/Private/VT/VirtualTextureSystem.cpp
@@ -2768,12 +2768,12 @@ void FVirtualTextureSystem::EndUpdate(FRDGBuilder& GraphBuilder, TUniquePtr<FVir
 	if (Updater->bAsyncTaskAllowed)
 	{
 		Updater->AsyncTask.Wait();
-		bUpdating = true;
 	}
 	else
 	{
 		BeginUpdate(GraphBuilder, Updater.Get());
 	}
+	bUpdating = false;
 
 	if (Updater->FeedbackMapResult.Data)
 	{
@@ -2807,7 +2807,6 @@ void FVirtualTextureSystem::EndUpdate(FRDGBuilder& GraphBuilder, TUniquePtr<FVir
 	ReleasePendingSpaces();
 
 	TilesToLock = MoveTemp(TilesToLockForNextFrame);
-	bUpdating = false;
 }
 
 void FVirtualTextureSystem::Update(FRDGBuilder& GraphBuilder, ERHIFeatureLevel::Type FeatureLevel, FScene* Scene, const FVirtualTextureUpdateSettings& InSettings)
