Virtual Texture System: Fix crash when changing levels

The PhysicalSpaces array can contain null pointers, so the call to FinalizeTextures must be guarded.

diff --git a/Engine/Source/Runtime/Renderer/Private/VT/VirtualTextureSystem.cpp b/Engine/Source/Runtime/Renderer/Private/VT/VirtualTextureSystem.cpp
--- a/Engine/Source/Runtime/Renderer/Private/VT/VirtualTextureSystem.cpp
+++ b/Engine/Source/Runtime/Renderer/Private/VT/VirtualTextureSystem.cpp
@@ -2415,7 +2415,10 @@ void FVirtualTextureSystem::FinalizeRequests(FRDGBuilder& GraphBuilder)
 
 	for (FVirtualTexturePhysicalSpace* PhysicalSpace : PhysicalSpaces)
 	{
-		PhysicalSpace->FinalizeTextures(GraphBuilder);
+		if (PhysicalSpace != nullptr)
+		{
+			PhysicalSpace->FinalizeTextures(GraphBuilder);
+		}
 	}
 
 	Frame++;
