* Add byte per pixel view mode for Substrate
* Simplifty substrate view mode to ease adding them to the level menu
* Add project  vs. on screen max byte per pixel on screen

diff --git a/Engine/Shaders/Private/Strata/StrataVisualize.usf b/Engine/Shaders/Private/Strata/StrataVisualize.usf
--- a/Engine/Shaders/Private/Strata/StrataVisualize.usf
+++ b/Engine/Shaders/Private/Strata/StrataVisualize.usf
@@ -590,6 +590,28 @@ void MaterialPrintInfoCS(uint3 DispatchThreadId : SV_DispatchThreadID)
 
 uint ViewMode;
 
+float3 BytePerPixelToColor(uint In)
+{
+	float3 Out = 0;
+	// 16->32  : Blue -> green
+	// 32->64  : Green -> Orange
+	// 64->128 : Orange -> Red
+	if (In < 32)
+	{
+		Out = ColorMapTurbo(saturate((In-16.f) / 32.f) * 0.25f + 0.25f);
+	}
+	else if (In < 64)
+	{
+		Out = ColorMapTurbo(saturate((In-32.f) / 32.f) * 0.25f + 0.5f);
+	}
+	else
+	{
+		Out = ColorMapTurbo(saturate((In-64.f) / 64.f) * 0.25f + 0.75f);
+	}
+
+	return Out;
+}
+
 void VisualizeMaterialPS(
 	float4 SVPos : SV_POSITION,
 	out float4 OutColor : SV_Target0)
@@ -605,28 +627,67 @@ void VisualizeMaterialPS(
 	const float3 V = -normalize(WorldPosition - PrimaryView.TranslatedWorldCameraOrigin);
 
 	FStrataAddressing StrataAddressing = GetStrataPixelDataByteOffset(PixelPosDynRes, uint2(View.BufferSizeAndInvSize.xy), Strata.MaxBytesPerPixel);
+	const uint FootPrint_Start = StrataAddressing.ReadBytes;
 	FStrataPixelHeader Header = UnpackStrataHeaderIn(Strata.MaterialTextureArray, StrataAddressing, Strata.TopLayerTexture);
 	const bool bStrataMaterial = Header.BSDFCount > 0;
 
 	OutColor = 0;
 	OutColor.w = 0.1f;
-
-	// BSDF count
-	const float MaxBSDFCount = 4.f;
-	if (Header.BSDFCount > 0)
+	if (ViewMode == 2)
 	{
-		const float BSDFCost = saturate(float(Header.BSDFCount) / MaxBSDFCount);
-		OutColor.xyz = ColorMapTurbo(BSDFCost);
+		// BSDF count
+		const float MaxBSDFCount = 4.f;
+		if (Header.BSDFCount > 0)
+		{
+			const float BSDFCost = saturate(float(Header.BSDFCount) / MaxBSDFCount);
+			OutColor.xyz = ColorMapTurbo(BSDFCost);
+		}
+	
+		// Print Legend
+		if (all(uint2(SVPos.xy) == uint2(0,0)))
+		{
+			FShaderPrintContext Context = InitShaderPrintContext(true, uint2(50, 50));
+			Print(Context, TEXT("1 BSDF"), InitFontColor(ColorMapTurbo(1 / MaxBSDFCount))); Newline(Context);
+			Print(Context, TEXT("2 BSDF"), InitFontColor(ColorMapTurbo(2 / MaxBSDFCount))); Newline(Context);
+			Print(Context, TEXT("3 BSDF"), InitFontColor(ColorMapTurbo(3 / MaxBSDFCount))); Newline(Context);
+			Print(Context, TEXT("4 BSDF"), InitFontColor(ColorMapTurbo(4 / MaxBSDFCount))); Newline(Context);
+		}
 	}
-
-	// Print Legend
-	if (all(uint2(SVPos.xy) == uint2(0,0)))
+	else if (ViewMode == 3)
 	{
-		FShaderPrintContext Context = InitShaderPrintContext(true, uint2(50, 50));
-		Print(Context, TEXT("1 BSDF"), InitFontColor(ColorMapTurbo(1 / MaxBSDFCount))); Newline(Context);
-		Print(Context, TEXT("2 BSDF"), InitFontColor(ColorMapTurbo(2 / MaxBSDFCount))); Newline(Context);
-		Print(Context, TEXT("3 BSDF"), InitFontColor(ColorMapTurbo(3 / MaxBSDFCount))); Newline(Context);
-		Print(Context, TEXT("4 BSDF"), InitFontColor(ColorMapTurbo(4 / MaxBSDFCount))); Newline(Context);
+		FStrataAddressing StrataAddressing = GetStrataPixelDataByteOffset(PixelPosDynRes, uint2(View.BufferSizeAndInvSize.xy), Strata.MaxBytesPerPixel);
+		FStrataPixelHeader Header = UnpackStrataHeaderIn(Strata.MaterialTextureArray, StrataAddressing, Strata.TopLayerTexture);
+
+		uint MaterialBytePerPixel = 0;
+		Strata_for (uint BSDFIndex = 0, BSDFIndex < Header.BSDFCount, ++BSDFIndex)
+		{
+			// Unpack BSDF data
+			FStrataBSDF BSDF = UnpackStrataBSDF(Strata.MaterialTextureArray, StrataAddressing, Header);
+			MaterialBytePerPixel = StrataAddressing.ReadBytes;
+		}
+			
+		OutColor.xyz = BytePerPixelToColor(MaterialBytePerPixel);
+	
+		// Print Legend
+		if (all(uint2(SVPos.xy) == uint2(0,0)))
+		{
+			FShaderPrintContext Context = InitShaderPrintContext(true, uint2(50, 50));
+			Print(Context, TEXT("16  Bytes"), InitFontColor(BytePerPixelToColor(16))); Newline(Context);	
+			Print(Context, TEXT("32  Bytes"), InitFontColor(BytePerPixelToColor(32))); Newline(Context);
+			Print(Context, TEXT("64  Bytes"), InitFontColor(BytePerPixelToColor(64))); Newline(Context);
+			Print(Context, TEXT("96  Bytes"), InitFontColor(BytePerPixelToColor(96))); Newline(Context);
+			Print(Context, TEXT("128 Bytes"), InitFontColor(BytePerPixelToColor(128))); Newline(Context);
+		}
+
+		// Cursor
+		if (all(PixelPosDynRes == ShaderPrintData.CursorCoord))
+		{
+			FShaderPrintContext Context = InitShaderPrintContext(true, uint2(250, 50));
+			Print(Context, MaterialBytePerPixel, InitFontColor(BytePerPixelToColor(MaterialBytePerPixel)), 3, 0); 
+			Print(Context, TEXT("/"), FontWhite); 
+			Print(Context, Strata.MaxBytesPerPixel, FontYellow, 3, 0); 
+			Print(Context, TEXT(" Bytes"), FontWhite); 
+		}
 	}
 }
 #endif // SHADER_MATERIALCOUNT
@@ -1415,6 +1476,7 @@ uint ClassificationAsync;
 uint Classification8bits;
 uint bRoughRefraction;
 uint bTileOverflowUseMaterialData;
+uint ProjectMaxBytePerPixel;
 float TileOverflowRatio;
 Buffer<uint> ClassificationTileDrawIndirectBuffer;
 
@@ -1439,7 +1501,9 @@ void MainCS(uint3 DispatchThreadId : SV_DispatchThreadID)
 		Newline(Context);
 
 		Print(Context, TEXT(" Max byte per pixel : "), LegendColor); 
-		Print(Context, Strata.MaxBytesPerPixel, ValueColor);
+		Print(Context, Strata.MaxBytesPerPixel, ValueColor, 3, 0);
+		Print(Context, TEXT(" / "), FontWhite); 
+		Print(Context, ProjectMaxBytePerPixel, ValueColor, 3, 0);
 		Newline(Context);
 
 		Print(Context, TEXT(" Shading quality    : "), LegendColor);
diff --git a/Engine/Source/Editor/CommonMenuExtensions/Private/StrataVisualizationMenuCommands.cpp b/Engine/Source/Editor/CommonMenuExtensions/Private/StrataVisualizationMenuCommands.cpp
--- a/Engine/Source/Editor/CommonMenuExtensions/Private/StrataVisualizationMenuCommands.cpp
+++ b/Engine/Source/Editor/CommonMenuExtensions/Private/StrataVisualizationMenuCommands.cpp
@@ -49,22 +49,7 @@ void FStrataVisualizationMenuCommands::BuildCommandMap()
 			.UserInterfaceType(EUserInterfaceActionType::RadioButton)
 			.DefaultChord(FInputChord()
 		);
-
-		switch (Entry.Mode)
-		{
-		default:
-		case FStrataVisualizationData::FViewMode::MaterialProperties:
-			Record.Type = FStrataVisualizationType::MaterialProperties;
-			break;
-
-		case FStrataVisualizationData::FViewMode::MaterialCount:
-			Record.Type = FStrataVisualizationType::MaterialCount;
-			break;
-
-		case FStrataVisualizationData::FViewMode::AdvancedMaterialProperties:
-			Record.Type = FStrataVisualizationType::AdvancedMaterialProperties;
-			break;
-		}
+		Record.ViewMode = Entry.ViewMode;
 	}
 }
 
@@ -73,17 +58,29 @@ void FStrataVisualizationMenuCommands::BuildVisualisationSubMenu(FMenuBuilder& M
 	const FStrataVisualizationMenuCommands& Commands = FStrataVisualizationMenuCommands::Get();
 	if (Commands.IsPopulated())
 	{
-		Menu.BeginSection("LevelViewportSubstrateVisualizationMode", LOCTEXT("SubstrateVisualizationHeader", "Substrate Visualization Mode"));
-
-		Commands.AddCommandTypeToMenu(Menu, FStrataVisualizationType::MaterialProperties);
-		Commands.AddCommandTypeToMenu(Menu, FStrataVisualizationType::MaterialCount);
-		Commands.AddCommandTypeToMenu(Menu, FStrataVisualizationType::AdvancedMaterialProperties);
+		// General
+		{
+			Menu.BeginSection("LevelViewportSubstrateVisualizationMode", LOCTEXT("SubstrateVisualizationHeader", "Substrate General View Mode"));
+			Commands.AddCommandTypeToMenu(Menu, FStrataViewMode::MaterialProperties);
+			Commands.AddCommandTypeToMenu(Menu, FStrataViewMode::MaterialCount);
+			Commands.AddCommandTypeToMenu(Menu, FStrataViewMode::MaterialByteCount);
+			Commands.AddCommandTypeToMenu(Menu, FStrataViewMode::StrataInfo);
+			Menu.EndSection();
+		}
 
-		Menu.EndSection();
+		// Advanced
+		{
+			Menu.BeginSection("LevelViewportSubstrateVisualizationMode", LOCTEXT("SubstrateVisualizationHeader", "Substrate Advanced View Mode"));
+			Commands.AddCommandTypeToMenu(Menu, FStrataViewMode::AdvancedMaterialProperties);
+			Commands.AddCommandTypeToMenu(Menu, FStrataViewMode::MaterialClassification);
+			Commands.AddCommandTypeToMenu(Menu, FStrataViewMode::RoughRefractionClassification);
+			Commands.AddCommandTypeToMenu(Menu, FStrataViewMode::DecalClassification);
+			Menu.EndSection();
+		}
 	}
 }
 
-bool FStrataVisualizationMenuCommands::AddCommandTypeToMenu(FMenuBuilder& Menu, const FStrataVisualizationType Type) const
+bool FStrataVisualizationMenuCommands::AddCommandTypeToMenu(FMenuBuilder& Menu, const FStrataViewMode ViewMode) const
 {
 	bool bAddedCommands = false;
 
@@ -91,7 +88,7 @@ bool FStrataVisualizationMenuCommands::AddCommandTypeToMenu(FMenuBuilder& Menu,
 	for (TCommandConstIterator It = CreateCommandConstIterator(); It; ++It)
 	{
 		const FStrataVisualizationRecord& Record = It.Value();
-		if (Record.Type == Type)
+		if (Record.ViewMode == ViewMode)
 		{
 			Menu.AddMenuEntry(Record.Command, NAME_None, Record.Command->GetLabel());
 			bAddedCommands = true;
diff --git a/Engine/Source/Editor/CommonMenuExtensions/Public/StrataVisualizationMenuCommands.h b/Engine/Source/Editor/CommonMenuExtensions/Public/StrataVisualizationMenuCommands.h
--- a/Engine/Source/Editor/CommonMenuExtensions/Public/StrataVisualizationMenuCommands.h
+++ b/Engine/Source/Editor/CommonMenuExtensions/Public/StrataVisualizationMenuCommands.h
@@ -9,6 +9,7 @@
 #include "HAL/Platform.h"
 #include "Templates/SharedPointer.h"
 #include "UObject/NameTypes.h"
+#include "StrataVisualizationData.h"
 
 class FEditorViewportClient;
 class FMenuBuilder;
@@ -18,18 +19,11 @@ class FUICommandList;
 class COMMONMENUEXTENSIONS_API FStrataVisualizationMenuCommands : public TCommands<FStrataVisualizationMenuCommands>
 {
 public:
-	enum class FStrataVisualizationType : uint8
-	{
-		MaterialProperties,
-		MaterialCount,
-		AdvancedMaterialProperties,
-	};
-
 	struct FStrataVisualizationRecord
 	{
 		FName Name;
 		TSharedPtr<FUICommandInfo> Command;
-		FStrataVisualizationType Type = FStrataVisualizationType::MaterialProperties;
+		FStrataViewMode ViewMode = FStrataViewMode::MaterialProperties;
 	};
 
 	typedef TMultiMap<FName, FStrataVisualizationRecord> TStrataVisualizationModeCommandMap;
@@ -52,7 +46,7 @@ public:
 
 private:
 	void BuildCommandMap();
-	bool AddCommandTypeToMenu(FMenuBuilder& Menu, const FStrataVisualizationType Type) const;
+	bool AddCommandTypeToMenu(FMenuBuilder& Menu, const FStrataViewMode ViewMode) const;
 
 	static void ChangeStrataVisualizationMode(TWeakPtr<FEditorViewportClient> WeakClient, FName InName);
 	static bool IsStrataVisualizationModeSelected(TWeakPtr<FEditorViewportClient> WeakClient, FName InName);
diff --git a/Engine/Source/Runtime/Engine/Private/StrataVisualizationData.cpp b/Engine/Source/Runtime/Engine/Private/StrataVisualizationData.cpp
--- a/Engine/Source/Runtime/Engine/Private/StrataVisualizationData.cpp
+++ b/Engine/Source/Runtime/Engine/Private/StrataVisualizationData.cpp
@@ -17,7 +17,7 @@ static FString ConfigureConsoleCommand(FStrataVisualizationData::TModeMap& ModeM
 	{
 		const FStrataVisualizationData::FModeRecord& Record = It.Value();
 		AvailableVisualizationModes += FString(TEXT("\n  Value="));
-		AvailableVisualizationModes += FString::Printf(TEXT("%d: "), uint8(Record.Mode));
+		AvailableVisualizationModes += FString::Printf(TEXT("%d: "), uint8(Record.ViewMode));
 		AvailableVisualizationModes += Record.ModeString;
 		if (!Record.bAvailableCommand)
 		{
@@ -43,7 +43,7 @@ static void AddVisualizationMode(
 	const TCHAR* ModeString,
 	const FText& ModeText,
 	const FText& ModeDesc,
-	const FStrataVisualizationData::FViewMode Mode,
+	const FStrataViewMode ViewMode,
 	bool bDefaultComposited,
 	bool bAvailableCommand,
 	const FText& UnavailableReason
@@ -56,7 +56,7 @@ static void AddVisualizationMode(
 	Record.ModeName				= ModeName;
 	Record.ModeText				= ModeText;
 	Record.ModeDesc				= ModeDesc;
-	Record.Mode					= Mode;
+	Record.ViewMode				= ViewMode;
 	Record.bDefaultComposited	= bDefaultComposited;
 	Record.bAvailableCommand	= bAvailableCommand;
 	Record.UnavailableReason	= UnavailableReason;
@@ -73,7 +73,7 @@ void FStrataVisualizationData::Initialize()
 			TEXT("MaterialProperties"),
 			LOCTEXT("MaterialProperties", "Material Properties"),
 			LOCTEXT("MaterialPropertiesDesc", "Visualizes Substrate material properties under mouse cursor"),
-			FViewMode::MaterialProperties,
+			FStrataViewMode::MaterialProperties,
 			true,
 			true,
 			LOCTEXT("None", "None"));
@@ -83,7 +83,7 @@ void FStrataVisualizationData::Initialize()
 			TEXT("MaterialCount"),
 			LOCTEXT("MaterialCount", "Material Count"),
 			LOCTEXT("MaterialCountDesc", "Visualizes Substrate material count per pixel"),
-			FViewMode::MaterialCount,
+			FStrataViewMode::MaterialCount,
 			true,
 			true,
 			LOCTEXT("None", "None"));
@@ -93,7 +93,7 @@ void FStrataVisualizationData::Initialize()
 			TEXT("AdvancedMaterialProperties"),
 			LOCTEXT("AdvancedMaterialProperties", "Advanced Material Properties"),
 			LOCTEXT("AdvancedMaterialPropertiesDesc", "Visualizes Substrate advanced material properties"),
-			FViewMode::AdvancedMaterialProperties,
+			FStrataViewMode::AdvancedMaterialProperties,
 			true,
 			Strata::IsAdvancedVisualizationEnabled(),
 			LOCTEXT("IsSubstrateAdvancedDebugShaderEnabled", "Substrate advanced debugging r.Substrate.Debug.AdvancedVisualizationShaders is disabled"));
@@ -103,7 +103,7 @@ void FStrataVisualizationData::Initialize()
 			TEXT("MaterialClassification"),
 			LOCTEXT("MaterialClassification", "Material Classification"),
 			LOCTEXT("MaterialClassificationDesc", "Visualizes Substrate material classification"),
-			FViewMode::MaterialClassification,
+			FStrataViewMode::MaterialClassification,
 			true,
 			true,
 			LOCTEXT("None", "None"));
@@ -113,7 +113,7 @@ void FStrataVisualizationData::Initialize()
 			TEXT("DecalClassification"),
 			LOCTEXT("DecalClassification", "Decal classification"),
 			LOCTEXT("DecalClassificationDesc", "Visualizes Substrate decal classification"),
-			FViewMode::DecalClassification,
+			FStrataViewMode::DecalClassification,
 			true,
 			false, // Disable for now, as it is not important, and is mainly used for debugging
 			LOCTEXT("IsSubstrateBufferPassEnabled", "Substrate tiled DBuffer pass (r.Substrate.DBufferPass and r.Substrate.DBufferPass.DedicatedTiles) is disabled"));
@@ -123,7 +123,7 @@ void FStrataVisualizationData::Initialize()
 			TEXT("RoughRefractionClassification"),
 			LOCTEXT("RoughRefractionClassification", "Rough Refraction Classification"),
 			LOCTEXT("RoughRefractionClassificationDesc", "Visualizes Substrate rough refraction classification"),
-			FViewMode::RoughRefractionClassification,
+			FStrataViewMode::RoughRefractionClassification,
 			true,
 			Strata::IsOpaqueRoughRefractionEnabled(),
 			LOCTEXT("IsSubstrateRoughRefractionEnabled", "Substrate rough refraction r.Substrate.OpaqueMaterialRoughRefraction is disabled"));
@@ -133,11 +133,22 @@ void FStrataVisualizationData::Initialize()
 			TEXT("SubstrateInfo"),
 			LOCTEXT("SubstrateInfo", "Substrate Info"),
 			LOCTEXT("SubstrateInfoDesc", "Visualizes Substrate info"),
-			FViewMode::StrataInfo,
+			FStrataViewMode::StrataInfo,
 			true,
 			true,
 			LOCTEXT("None", "None"));
 
+		AddVisualizationMode(
+			AllModeMap,
+			TEXT("MaterialByteCount"),
+			LOCTEXT("MaterialByteCount", "Material Bytes Count"),
+			LOCTEXT("MaterialByteCountDesc", "Visualizes Substrate material footprint per pixel"),
+			FStrataViewMode::MaterialByteCount,
+			true,
+			true,
+			LOCTEXT("None", "None"));
+		
+
 		ConsoleDocumentationVisualizationMode = ConfigureConsoleCommand(AllModeMap);
 
 		// Now only copy the available modes for the menu to not overload it with useless entries.
@@ -164,15 +175,15 @@ FText FStrataVisualizationData::GetModeDisplayName(const FName& InModeName) cons
 	}
 }
 
-FStrataVisualizationData::FViewMode FStrataVisualizationData::GetViewMode(const FName& InModeName) const
+FStrataViewMode FStrataVisualizationData::GetViewMode(const FName& InModeName) const
 {
 	if (const FModeRecord* Record = ModeMap.Find(InModeName))
 	{
-		return Record->Mode;
+		return Record->ViewMode;
 	}
 	else
 	{
-		return FStrataVisualizationData::FViewMode::None;
+		return FStrataViewMode::None;
 	}
 }
 
diff --git a/Engine/Source/Runtime/Engine/Public/StrataVisualizationData.h b/Engine/Source/Runtime/Engine/Public/StrataVisualizationData.h
--- a/Engine/Source/Runtime/Engine/Public/StrataVisualizationData.h
+++ b/Engine/Source/Runtime/Engine/Public/StrataVisualizationData.h
@@ -4,20 +4,22 @@
 
 #include "CoreMinimal.h"
 
+enum class FStrataViewMode : uint8
+{
+	None,
+	MaterialProperties,
+	MaterialCount,
+	MaterialByteCount,
+	StrataInfo,
+	AdvancedMaterialProperties,
+	MaterialClassification,
+	RoughRefractionClassification,
+	DecalClassification,
+};
+
 class FStrataVisualizationData
 {
 public:
-	enum class FViewMode : uint8
-	{
-		None,
-		MaterialProperties,
-		MaterialCount,
-		AdvancedMaterialProperties,
-		MaterialClassification,
-		DecalClassification,
-		RoughRefractionClassification,
-		StrataInfo,
-	};
 
 	/** Describes a single available visualization mode. */
 	struct FModeRecord
@@ -26,7 +28,7 @@ public:
 		FName     ModeName;
 		FText     ModeText;
 		FText     ModeDesc;
-		FViewMode Mode;
+		FStrataViewMode ViewMode;
 
 		// Whether or not this mode (by default) composites with regular scene depth.
 		bool      bDefaultComposited;
@@ -53,7 +55,7 @@ public:
 	/** Get the display name of a named mode from the available mode map. **/
 	ENGINE_API FText GetModeDisplayName(const FName& InModeName) const;
 
-	ENGINE_API FViewMode GetViewMode(const FName& InModeName) const;
+	ENGINE_API FStrataViewMode GetViewMode(const FName& InModeName) const;
 
 	ENGINE_API bool GetModeDefaultComposited(const FName& InModeName) const;
 
diff --git a/Engine/Source/Runtime/Renderer/Private/Strata/StrataVisualize.cpp b/Engine/Source/Runtime/Renderer/Private/Strata/StrataVisualize.cpp
--- a/Engine/Source/Runtime/Renderer/Private/Strata/StrataVisualize.cpp
+++ b/Engine/Source/Runtime/Renderer/Private/Strata/StrataVisualize.cpp
@@ -120,6 +120,7 @@ class FStrataSystemInfoCS : public FGlobalShader
 		SHADER_PARAMETER(uint32, Classification8bits)
 		SHADER_PARAMETER(uint32, bRoughRefraction)
 		SHADER_PARAMETER(uint32, bTileOverflowUseMaterialData)
+		SHADER_PARAMETER(uint32, ProjectMaxBytePerPixel)
 		SHADER_PARAMETER(float, TileOverflowRatio)
 		SHADER_PARAMETER_RDG_BUFFER_SRV(Buffer<uint>, ClassificationTileDrawIndirectBuffer)
 		SHADER_PARAMETER_STRUCT_REF(FViewUniformShaderParameters, ViewUniformBuffer)
@@ -249,7 +250,7 @@ static void AddVisualizeMaterialPropertiesPasses(FRDGBuilder& GraphBuilder, cons
 	}
 }
 
-static void AddVisualizeMaterialCountPasses(FRDGBuilder & GraphBuilder, const FViewInfo & View, FScreenPassTexture & ScreenPassSceneColor, EShaderPlatform Platform)
+static void AddVisualizeMaterialCountPasses(FRDGBuilder & GraphBuilder, const FViewInfo & View, FScreenPassTexture & ScreenPassSceneColor, EShaderPlatform Platform, uint32 ViewMode)
 {
 	ShaderPrint::SetEnabled(true);
 	ShaderPrint::RequestSpaceForLines(1024);
@@ -260,7 +261,7 @@ static void AddVisualizeMaterialCountPasses(FRDGBuilder & GraphBuilder, const FV
 
 	FVisualizeMaterialCountPS::FParameters* PassParameters = GraphBuilder.AllocParameters<FVisualizeMaterialCountPS::FParameters>();
 	PassParameters->ViewUniformBuffer = View.ViewUniformBuffer;
-	PassParameters->ViewMode = 2;
+	PassParameters->ViewMode = FMath::Clamp(ViewMode, 2, 3);
 	PassParameters->Strata = Strata::BindStrataGlobalUniformParameters(View);
 	PassParameters->SceneTextures = GetSceneTextureParameters(GraphBuilder, View);
 	PassParameters->RenderTargets[0] = FRenderTargetBinding(SceneColorTexture, ERenderTargetLoadAction::ELoad);
@@ -299,6 +300,7 @@ static void AddVisualizeSystemInfoPasses(FRDGBuilder& GraphBuilder, const FViewI
 	PassParameters->ViewUniformBuffer = View.ViewUniformBuffer;
 	PassParameters->Strata = Strata::BindStrataGlobalUniformParameters(View);
 	PassParameters->SceneTextures = GetSceneTextureParameters(GraphBuilder, View);
+	PassParameters->ProjectMaxBytePerPixel = GetBytePerPixel(View.GetShaderPlatform());
 	ShaderPrint::SetParameters(GraphBuilder, View.ShaderPrintData, PassParameters->ShaderPrintParameters);
 
 	TShaderMapRef<FStrataSystemInfoCS> ComputeShader(View.ShaderMap);
@@ -352,9 +354,9 @@ static void AddVisualizeAdvancedMaterialPasses(FRDGBuilder& GraphBuilder, const
 bool ShouldRenderStrataRoughRefractionRnD();
 void StrataRoughRefractionRnD(FRDGBuilder& GraphBuilder, const FViewInfo& View, FScreenPassTexture& ScreenPassSceneColor);
 
-static FStrataVisualizationData::FViewMode GetStrataVisualizeMode(const FViewInfo & View)
+static FStrataViewMode GetStrataVisualizeMode(const FViewInfo & View)
 {
-	FStrataVisualizationData::FViewMode Out = FStrataVisualizationData::FViewMode::None;
+	FStrataViewMode Out = FStrataViewMode::None;
 	if (IsStrataEnabled() && StrataDebugVisualizationCanRunOnPlatform(View.GetShaderPlatform()))
 	{
 		// Variable defined in StrataVisualizationData.h/.cpp
@@ -362,13 +364,14 @@ static FStrataVisualizationData::FViewMode GetStrataVisualizeMode(const FViewInf
 		const uint32 ViewMode = CVar && CVar->AsVariableInt() ? CVar->AsVariableInt()->GetValueOnRenderThread() : 0;
 		switch (ViewMode)
 		{
-			case 1: return FStrataVisualizationData::FViewMode::MaterialProperties;
-			case 2: return FStrataVisualizationData::FViewMode::MaterialCount;
-			case 3: return FStrataVisualizationData::FViewMode::AdvancedMaterialProperties;
-			case 4: return FStrataVisualizationData::FViewMode::MaterialClassification;
-			case 5: return FStrataVisualizationData::FViewMode::DecalClassification;
-			case 6: return FStrataVisualizationData::FViewMode::RoughRefractionClassification;
-			case 7: return FStrataVisualizationData::FViewMode::StrataInfo;
+			case 1: return FStrataViewMode::MaterialProperties;
+			case 2: return FStrataViewMode::MaterialCount;
+			case 3: return FStrataViewMode::AdvancedMaterialProperties;
+			case 4: return FStrataViewMode::MaterialClassification;
+			case 5: return FStrataViewMode::DecalClassification;
+			case 6: return FStrataViewMode::RoughRefractionClassification;
+			case 7: return FStrataViewMode::StrataInfo;
+			case 8: return FStrataViewMode::MaterialByteCount;
 		}
 
 		const FStrataVisualizationData& VisualizationData = GetStrataVisualizationData();
@@ -382,36 +385,40 @@ static FStrataVisualizationData::FViewMode GetStrataVisualizeMode(const FViewInf
 
 bool ShouldRenderStrataDebugPasses(const FViewInfo& View)
 {
-	return GetStrataVisualizeMode(View) != FStrataVisualizationData::FViewMode::None || ShouldRenderStrataRoughRefractionRnD();
+	return GetStrataVisualizeMode(View) != FStrataViewMode::None || ShouldRenderStrataRoughRefractionRnD();
 }
 
 FScreenPassTexture AddStrataDebugPasses(FRDGBuilder& GraphBuilder, const FViewInfo& View, FScreenPassTexture& ScreenPassSceneColor)
 {
 	check(IsStrataEnabled());
 
-	const FStrataVisualizationData::FViewMode DebugMode = GetStrataVisualizeMode(View);
-	if (DebugMode != FStrataVisualizationData::FViewMode::None)
+	const FStrataViewMode DebugMode = GetStrataVisualizeMode(View);
+	if (DebugMode != FStrataViewMode::None)
 	{
 		RDG_EVENT_SCOPE(GraphBuilder, "Substrate::VisualizeMaterial");
 
 		const bool bDebugPass = true;
-		if (DebugMode == FStrataVisualizationData::FViewMode::MaterialProperties)
+		if (DebugMode == FStrataViewMode::MaterialProperties)
 		{
 			AddVisualizeMaterialPropertiesPasses(GraphBuilder, View, ScreenPassSceneColor, View.GetShaderPlatform());
 		}
-		if (DebugMode == FStrataVisualizationData::FViewMode::MaterialCount)
+		if (DebugMode == FStrataViewMode::MaterialCount)
 		{
-			AddVisualizeMaterialCountPasses(GraphBuilder, View, ScreenPassSceneColor, View.GetShaderPlatform());
+			AddVisualizeMaterialCountPasses(GraphBuilder, View, ScreenPassSceneColor, View.GetShaderPlatform(), 2);
 		}
-		if (DebugMode == FStrataVisualizationData::FViewMode::AdvancedMaterialProperties)
+		if (DebugMode == FStrataViewMode::MaterialByteCount)
+		{
+			AddVisualizeMaterialCountPasses(GraphBuilder, View, ScreenPassSceneColor, View.GetShaderPlatform(), 3);
+		}
+		if (DebugMode == FStrataViewMode::AdvancedMaterialProperties)
 		{
 			AddVisualizeAdvancedMaterialPasses(GraphBuilder, View, ScreenPassSceneColor, View.GetShaderPlatform());
 		}
-		else if (DebugMode == FStrataVisualizationData::FViewMode::StrataInfo)
+		else if (DebugMode == FStrataViewMode::StrataInfo)
 		{
 			AddVisualizeSystemInfoPasses(GraphBuilder, View, ScreenPassSceneColor, View.GetShaderPlatform());
 		}
-		else if (DebugMode == FStrataVisualizationData::FViewMode::DecalClassification)
+		else if (DebugMode == FStrataViewMode::DecalClassification)
 		{
 			if (IsDBufferPassEnabled(View.GetShaderPlatform()))
 			{
@@ -420,7 +427,7 @@ FScreenPassTexture AddStrataDebugPasses(FRDGBuilder& GraphBuilder, const FViewIn
 				AddStrataInternalClassificationTilePass(GraphBuilder, View, nullptr, &ScreenPassSceneColor.Texture, EStrataTileType::EDecalComplex, bDebugPass);
 			}
 		}
-		else if (DebugMode == FStrataVisualizationData::FViewMode::RoughRefractionClassification)
+		else if (DebugMode == FStrataViewMode::RoughRefractionClassification)
 		{
 			if (IsOpaqueRoughRefractionEnabled())
 			{
@@ -428,7 +435,7 @@ FScreenPassTexture AddStrataDebugPasses(FRDGBuilder& GraphBuilder, const FViewIn
 				AddStrataInternalClassificationTilePass(GraphBuilder, View, nullptr, &ScreenPassSceneColor.Texture, EStrataTileType::EOpaqueRoughRefractionSSSWithout, bDebugPass);
 			}
 		}
-		else if (DebugMode == FStrataVisualizationData::FViewMode::MaterialClassification)
+		else if (DebugMode == FStrataViewMode::MaterialClassification)
 		{
 			AddStrataInternalClassificationTilePass(GraphBuilder, View, nullptr, &ScreenPassSceneColor.Texture, EStrataTileType::EComplex, bDebugPass);
 			AddStrataInternalClassificationTilePass(GraphBuilder, View, nullptr, &ScreenPassSceneColor.Texture, EStrataTileType::ESingle, bDebugPass);
