* Change Substrate to use the visible material 'max byte per pixel' (statically known at shader compiler time) to allocate Substrate material buffer, instead of using the constant project settings.
* Add Substrate material buffer allocation visualization

diff --git a/Engine/Shaders/Private/Strata/StrataVisualize.usf b/Engine/Shaders/Private/Strata/StrataVisualize.usf
--- a/Engine/Shaders/Private/Strata/StrataVisualize.usf
+++ b/Engine/Shaders/Private/Strata/StrataVisualize.usf
@@ -1477,6 +1477,8 @@ uint Classification8bits;
 uint bRoughRefraction;
 uint bTileOverflowUseMaterialData;
 uint ProjectMaxBytePerPixel;
+uint MaterialBufferAllocationInBytes;
+
 float TileOverflowRatio;
 Buffer<uint> ClassificationTileDrawIndirectBuffer;
 
@@ -1504,6 +1506,9 @@ void MainCS(uint3 DispatchThreadId : SV_DispatchThreadID)
 		Print(Context, Strata.MaxBytesPerPixel, ValueColor, 3, 0);
 		Print(Context, TEXT(" / "), FontWhite); 
 		Print(Context, ProjectMaxBytePerPixel, ValueColor, 3, 0);
+		Print(Context, TEXT(" ["), FontWhite); 
+		Print(Context, uint(DivideAndRoundUp(MaterialBufferAllocationInBytes, 1000000.f)), ValueColor, 4, 0);
+		Print(Context, TEXT("Mb]"), FontWhite);
 		Newline(Context);
 
 		Print(Context, TEXT(" Shading quality    : "), LegendColor);
diff --git a/Engine/Source/Runtime/Renderer/Private/Strata/StrataVisualize.cpp b/Engine/Source/Runtime/Renderer/Private/Strata/StrataVisualize.cpp
--- a/Engine/Source/Runtime/Renderer/Private/Strata/StrataVisualize.cpp
+++ b/Engine/Source/Runtime/Renderer/Private/Strata/StrataVisualize.cpp
@@ -121,6 +121,7 @@ class FStrataSystemInfoCS : public FGlobalShader
 		SHADER_PARAMETER(uint32, bRoughRefraction)
 		SHADER_PARAMETER(uint32, bTileOverflowUseMaterialData)
 		SHADER_PARAMETER(uint32, ProjectMaxBytePerPixel)
+		SHADER_PARAMETER(uint32, MaterialBufferAllocationInBytes)
 		SHADER_PARAMETER(float, TileOverflowRatio)
 		SHADER_PARAMETER_RDG_BUFFER_SRV(Buffer<uint>, ClassificationTileDrawIndirectBuffer)
 		SHADER_PARAMETER_STRUCT_REF(FViewUniformShaderParameters, ViewUniformBuffer)
@@ -285,6 +286,8 @@ static void AddVisualizeSystemInfoPasses(FRDGBuilder& GraphBuilder, const FViewI
 	ShaderPrint::RequestSpaceForLines(1024);
 	ShaderPrint::RequestSpaceForCharacters(1024);
 
+	const FRDGTextureDesc MaterialBufferDesc = View.StrataViewData.SceneData->MaterialTextureArray->Desc;
+
 	FStrataSystemInfoCS::FParameters* PassParameters = GraphBuilder.AllocParameters<FStrataSystemInfoCS::FParameters>();
 	PassParameters->bAdvancedDebugEnabled = IsAdvancedVisualizationEnabled() ? 1u : 0u;
 	PassParameters->bEnergyConservation = View.ViewState ? View.ViewState->ShadingEnergyConservationData.bEnergyConservation : false;;
@@ -301,6 +304,7 @@ static void AddVisualizeSystemInfoPasses(FRDGBuilder& GraphBuilder, const FViewI
 	PassParameters->Strata = Strata::BindStrataGlobalUniformParameters(View);
 	PassParameters->SceneTextures = GetSceneTextureParameters(GraphBuilder, View);
 	PassParameters->ProjectMaxBytePerPixel = GetBytePerPixel(View.GetShaderPlatform());
+	PassParameters->MaterialBufferAllocationInBytes = MaterialBufferDesc.Extent.X * MaterialBufferDesc.Extent.Y * MaterialBufferDesc.ArraySize * sizeof(uint32);
 	ShaderPrint::SetParameters(GraphBuilder, View.ShaderPrintData, PassParameters->ShaderPrintParameters);
 
 	TShaderMapRef<FStrataSystemInfoCS> ComputeShader(View.ShaderMap);
