OpenXR: request DepthStencilResolveTarget for depth swapchains that will be resolve targets

diff --git a/Engine/Source/Runtime/Renderer/Private/SceneTextures.cpp b/Engine/Source/Runtime/Renderer/Private/SceneTextures.cpp
--- a/Engine/Source/Runtime/Renderer/Private/SceneTextures.cpp
+++ b/Engine/Source/Runtime/Renderer/Private/SceneTextures.cpp
@@ -461,7 +461,7 @@ void FMinimalSceneTextures::InitializeViewFamily(FRDGBuilder& GraphBuilder, FVie
 		{
 			Desc.NumSamples = 1;
 
-			if ((StereoDepthRHI = FindStereoDepthTexture(Config.bSupportsXRTargetManagerDepthAlloc, Config.Extent, ETextureCreateFlags::ResolveTargetable)) != nullptr)
+			if ((StereoDepthRHI = FindStereoDepthTexture(Config.bSupportsXRTargetManagerDepthAlloc, Config.Extent, ETextureCreateFlags::DepthStencilResolveTarget)) != nullptr)
 			{
 				ensureMsgf(Desc.ArraySize == StereoDepthRHI->GetDesc().ArraySize, TEXT("Resolve texture does not agree in dimensionality with Target (Resolve.ArraySize=%d, Target.ArraySize=%d)"),
 					Desc.ArraySize, StereoDepthRHI->GetDesc().ArraySize);
