Fixed non-unity build break.

diff --git a/Engine/Source/Runtime/Renderer/Private/VT/VirtualTexturePhysicalSpace.cpp b/Engine/Source/Runtime/Renderer/Private/VT/VirtualTexturePhysicalSpace.cpp
--- a/Engine/Source/Runtime/Renderer/Private/VT/VirtualTexturePhysicalSpace.cpp
+++ b/Engine/Source/Runtime/Renderer/Private/VT/VirtualTexturePhysicalSpace.cpp
@@ -12,6 +12,7 @@
 #include "VT/VirtualTextureScalability.h"
 #include "VT/VirtualTextureSystem.h"
 #include "RHIUtilities.h"
+#include "RenderGraphBuilder.h"
 
 CSV_DECLARE_CATEGORY_EXTERN(VirtualTexturing);
 
