Define FName as static before passing to UE_TRACE_METADATA_SCOPE_ASSET_FNAME as FName lookup has performance impact.
Also rename AssetName to OwnerName to better describe what is stored.

diff --git a/Engine/Plugins/Runtime/HairStrands/Source/HairStrandsCore/Private/HairStrandsVertexFactory.cpp b/Engine/Plugins/Runtime/HairStrands/Source/HairStrandsCore/Private/HairStrandsVertexFactory.cpp
--- a/Engine/Plugins/Runtime/HairStrands/Source/HairStrandsCore/Private/HairStrandsVertexFactory.cpp
+++ b/Engine/Plugins/Runtime/HairStrands/Source/HairStrandsCore/Private/HairStrandsVertexFactory.cpp
@@ -43,9 +43,10 @@ public:
 
 	virtual void InitRHI() override
 	{
+		const static FLazyName ClassName(TEXT("FDummyCulledDispatchVertexIdsBuffer"));
 		{
 			FRHIResourceCreateInfo CreateInfo(TEXT("FDummyCulledDispatchVertexIdsBuffer"));
-			CreateInfo.ClassName = TEXT("FDummyCulledDispatchVertexIdsBuffer");
+			CreateInfo.ClassName = ClassName;
 			uint32 NumBytes = sizeof(uint32) * 4;
 			VertexBufferRHI = RHICreateBuffer(NumBytes, BUF_Static | BUF_VertexBuffer | BUF_ShaderResource, 0, ERHIAccess::VertexOrIndexBuffer | ERHIAccess::SRVMask, CreateInfo);
 			uint32* DummyContents = (uint32*)RHILockBuffer(VertexBufferRHI, 0, NumBytes, RLM_WriteOnly);
@@ -55,7 +56,7 @@ public:
 
 		{
 			FRHIResourceCreateInfo CreateInfo(TEXT("FDummyByteAddressBuffer"));
-			CreateInfo.ClassName = TEXT("FDummyCulledDispatchVertexIdsBuffer");
+			CreateInfo.ClassName = ClassName;
 			uint32 NumBytes = sizeof(uint32) * 4;
 			ByteAddressBufferRHI = RHICreateBuffer(NumBytes, BUF_Static | BUF_ShaderResource |BUF_ByteAddressBuffer, 0, ERHIAccess::SRVMask, CreateInfo);
 			uint32* DummyContents = (uint32*)RHILockBuffer(ByteAddressBufferRHI, 0, NumBytes, RLM_WriteOnly);
diff --git a/Engine/Source/Runtime/D3D12RHI/Private/D3D12Adapter.h b/Engine/Source/Runtime/D3D12RHI/Private/D3D12Adapter.h
--- a/Engine/Source/Runtime/D3D12RHI/Private/D3D12Adapter.h
+++ b/Engine/Source/Runtime/D3D12RHI/Private/D3D12Adapter.h
@@ -354,7 +354,7 @@ public:
 		const FRHIGPUMask& InGPUMask,
 		ID3D12ResourceAllocator* ResourceAllocator,
 		const TCHAR* InDebugName,
-		const FName& InAssetName = NAME_None,
+		const FName& InOwnerName = NAME_None,
 		const FName& InClassName = NAME_None);
 
 	void CreateUAVAliasResource(
diff --git a/Engine/Source/Runtime/D3D12RHI/Private/D3D12Buffer.cpp b/Engine/Source/Runtime/D3D12RHI/Private/D3D12Buffer.cpp
--- a/Engine/Source/Runtime/D3D12RHI/Private/D3D12Buffer.cpp
+++ b/Engine/Source/Runtime/D3D12RHI/Private/D3D12Buffer.cpp
@@ -7,9 +7,10 @@ D3D12Buffer.cpp: D3D Common code for buffers.
 #include "D3D12RHIPrivate.h"
 #include "ProfilingDebugging/AssetMetadataTrace.h"
 
-FName ChooseMetadataName(const FName& FirstChoice, const FName& SecondChoice)
+FName GetRHIBufferClassName(const FName& ClassName)
 {
-	return FirstChoice == NAME_None ? SecondChoice : FirstChoice;
+	const static FLazyName RHIBufferName(TEXT("FRHIBuffer"));
+	return ClassName == NAME_None ? RHIBufferName : ClassName;
 }
 
 extern int32 GD3D12BindResourceLabels;
@@ -74,7 +75,8 @@ struct FRHICommandRenameUploadBuffer final : public FRHICommand<FRHICommandRenam
 
 	void Execute(FRHICommandListBase& CmdList)
 	{
-		UE_TRACE_METADATA_SCOPE_ASSET_FNAME(Resource->GetName(), FName(TEXT("FRHICommandRenameUploadBuffer::Execute")), Resource->GetOwnerName());
+		const static FLazyName ExecuteName(TEXT("FRHICommandRenameUploadBuffer::Execute"));
+		UE_TRACE_METADATA_SCOPE_ASSET_FNAME(Resource->GetName(), ExecuteName, Resource->GetOwnerName());
 
 		// Make sure we have an active pipeline when running at the top of the pipe, so the lambda below can obtain a context.
 		TOptional<ERHIPipeline> PreviousPipeline;
@@ -202,7 +204,7 @@ struct FD3D12RHICommandInitializeBuffer final : public FRHICommand<FD3D12RHIComm
 };
 
 
-void FD3D12Buffer::UploadResourceData(FRHICommandListBase& RHICmdList, FResourceArrayInterface* InResourceArray, D3D12_RESOURCE_STATES InDestinationState, const FName& AssetName, const FName& ClassName, const FName& PackageName)
+void FD3D12Buffer::UploadResourceData(FRHICommandListBase& RHICmdList, FResourceArrayInterface* InResourceArray, D3D12_RESOURCE_STATES InDestinationState, const TCHAR* AssetName, const FName& ClassName, const FName& PackageName)
 {
 	UE_TRACE_METADATA_SCOPE_ASSET_FNAME(AssetName, ClassName, PackageName);
 
@@ -350,13 +352,13 @@ FD3D12Buffer* FD3D12Adapter::CreateRHIBuffer(
 	const FRHIGPUMask& InGPUMask,
 	ID3D12ResourceAllocator* ResourceAllocator,
 	const TCHAR* InDebugName,
-	const FName& AssetName,
+	const FName& OwnerName,
 	const FName& ClassName)
 {
 	TRACE_CPUPROFILER_EVENT_SCOPE(D3D12RHI::CreateRHIBuffer);
 	SCOPE_CYCLE_COUNTER(STAT_D3D12CreateBufferTime);
 
-	UE_TRACE_METADATA_SCOPE_ASSET_FNAME(FName(InDebugName), ChooseMetadataName(ClassName, TEXT("FRHIBuffer")), AssetName);
+	UE_TRACE_METADATA_SCOPE_ASSET_FNAME(FName(InDebugName), GetRHIBufferClassName(ClassName), OwnerName);
 
 	check(InDesc.Width == Size);
 
@@ -541,9 +543,7 @@ FBufferRHIRef FD3D12DynamicRHI::CreateBuffer(FRHICommandListBase& RHICmdList, ui
 
 FD3D12Buffer* FD3D12DynamicRHI::CreateD3D12Buffer(class FRHICommandListBase* RHICmdList, uint32 Size, EBufferUsageFlags Usage, uint32 Stride, ERHIAccess InResourceState, FRHIResourceCreateInfo& CreateInfo, ID3D12ResourceAllocator* ResourceAllocator)
 {
-	FName TraceAssetName = CreateInfo.DebugName;
-	FName TraceClassName = ChooseMetadataName(CreateInfo.GetTraceClassName(), TEXT("FRHIBuffer"));
-	FName TracePackageName = CreateInfo.GetTraceAssetName();
+	FName TraceClassName = GetRHIBufferClassName(CreateInfo.GetTraceClassName());
 
 	D3D12_RESOURCE_DESC Desc;
 	uint32 Alignment;
@@ -568,14 +568,14 @@ FD3D12Buffer* FD3D12DynamicRHI::CreateD3D12Buffer(class FRHICommandListBase* RHI
 	D3D12_RESOURCE_STATES CreateState = (CreateInfo.ResourceArray && bSupportResourceStateTracking) ? D3D12_RESOURCE_STATE_COPY_DEST : DesiredState;
 	bool bHasInitialData = CreateInfo.ResourceArray != nullptr;
 
-	FD3D12Buffer* Buffer = GetAdapter().CreateRHIBuffer(Desc, Alignment, Stride, Size, Usage, StateMode, CreateState, bHasInitialData, CreateInfo.GPUMask, ResourceAllocator, CreateInfo.DebugName, TraceAssetName, TraceClassName);
+	FD3D12Buffer* Buffer = GetAdapter().CreateRHIBuffer(Desc, Alignment, Stride, Size, Usage, StateMode, CreateState, bHasInitialData, CreateInfo.GPUMask, ResourceAllocator, CreateInfo.DebugName, CreateInfo.OwnerName, TraceClassName);
 	check(Buffer->ResourceLocation.IsValid());
 
 	// Copy the resource data if available 
 	if (bHasInitialData)
 	{
 		check(RHICmdList);
-		Buffer->UploadResourceData(*RHICmdList, CreateInfo.ResourceArray, DesiredState, TraceAssetName, TraceClassName, TracePackageName);
+		Buffer->UploadResourceData(*RHICmdList, CreateInfo.ResourceArray, DesiredState, CreateInfo.DebugName, TraceClassName, CreateInfo.OwnerName);
 	}
 
 	return Buffer;
diff --git a/Engine/Source/Runtime/D3D12RHI/Private/D3D12Commands.cpp b/Engine/Source/Runtime/D3D12RHI/Private/D3D12Commands.cpp
--- a/Engine/Source/Runtime/D3D12RHI/Private/D3D12Commands.cpp
+++ b/Engine/Source/Runtime/D3D12RHI/Private/D3D12Commands.cpp
@@ -676,7 +676,8 @@ void FD3D12CommandContext::RHICopyToStagingBuffer(FRHIBuffer* SourceBufferRHI, F
 {
 	SCOPE_CYCLE_COUNTER(STAT_D3D12CopyToStagingBufferTime);
 
-	UE_TRACE_METADATA_SCOPE_ASSET_FNAME(SourceBufferRHI->GetName(), FName(TEXT("FRHIStagingBuffer")), SourceBufferRHI->GetOwnerName());
+	const static FLazyName RHIStagingBufferName(TEXT("FRHIStagingBuffer"));
+	UE_TRACE_METADATA_SCOPE_ASSET_FNAME(SourceBufferRHI->GetName(), RHIStagingBufferName, SourceBufferRHI->GetOwnerName());
 
 	FD3D12StagingBuffer* StagingBuffer = FD3D12DynamicRHI::ResourceCast(StagingBufferRHI);
 	check(StagingBuffer);
diff --git a/Engine/Source/Runtime/D3D12RHI/Private/D3D12Query.cpp b/Engine/Source/Runtime/D3D12RHI/Private/D3D12Query.cpp
--- a/Engine/Source/Runtime/D3D12RHI/Private/D3D12Query.cpp
+++ b/Engine/Source/Runtime/D3D12RHI/Private/D3D12Query.cpp
@@ -75,7 +75,8 @@ FD3D12QueryHeap::FD3D12QueryHeap(FD3D12Device* Device, D3D12_QUERY_TYPE QueryTyp
 		break;
 	}
 
-	UE_TRACE_METADATA_SCOPE_ASSET_FNAME(FName(ResultBufferName), FName(TEXT("FD3D12QueryHeap")), NAME_None);
+	const static FLazyName D3D12QueryHeapName(TEXT("FD3D12QueryHeap"));
+	UE_TRACE_METADATA_SCOPE_ASSET_FNAME(FName(ResultBufferName), D3D12QueryHeapName, NAME_None);
 
 	// Create the query heap
 #if D3D12RHI_PLATFORM_USES_TIMESTAMP_QUERIES == 0
diff --git a/Engine/Source/Runtime/D3D12RHI/Private/D3D12Texture.cpp b/Engine/Source/Runtime/D3D12RHI/Private/D3D12Texture.cpp
--- a/Engine/Source/Runtime/D3D12RHI/Private/D3D12Texture.cpp
+++ b/Engine/Source/Runtime/D3D12RHI/Private/D3D12Texture.cpp
@@ -887,7 +887,7 @@ FD3D12Texture* FD3D12DynamicRHI::CreateD3D12Texture(const FRHITextureCreateDesc&
 #if PLATFORM_WINDOWS || PLATFORM_HOLOLENS
 	
 	TRACE_CPUPROFILER_EVENT_SCOPE(D3D12RHI::CreateD3D12Texture);
-	UE_TRACE_METADATA_SCOPE_ASSET_FNAME(FName(InCreateDesc.DebugName), InCreateDesc.GetTraceClassName(), InCreateDesc.GetTraceAssetName());
+	UE_TRACE_METADATA_SCOPE_ASSET_FNAME(FName(InCreateDesc.DebugName), InCreateDesc.GetTraceClassName(), InCreateDesc.OwnerName);
 
 	// Make local copy of create desc because certain flags might be removed before actual texture creation
 	FRHITextureCreateDesc CreateDesc = InCreateDesc;
@@ -1075,7 +1075,9 @@ FTextureRHIRef FD3D12DynamicRHI::RHIAsyncCreateTexture2D(uint32 SizeX, uint32 Si
 {	
 	check(GRHISupportsAsyncTextureCreation);
 	
-	UE_TRACE_METADATA_SCOPE_ASSET_FNAME(FName(TEXT("RHIAsyncCreateTexture2D")), FName(TEXT("FRHITexture")), NAME_None);
+	const static FName RHIAsyncCreateTexture2DName(TEXT("RHIAsyncCreateTexture2D"));
+	const static FName RHITextureName(TEXT("FRHITexture"));
+	UE_TRACE_METADATA_SCOPE_ASSET_FNAME(RHIAsyncCreateTexture2DName, RHITextureName, NAME_None);
 
 	const ETextureCreateFlags InvalidFlags = TexCreate_RenderTargetable | TexCreate_ResolveTargetable | TexCreate_DepthStencilTargetable | TexCreate_GenerateMipCapable | TexCreate_UAV | TexCreate_Presentable | TexCreate_CPUReadback;
 	check(!EnumHasAnyFlags(Flags, InvalidFlags));
@@ -1798,7 +1800,8 @@ void* FD3D12Texture::Lock(class FRHICommandListImmediate* RHICmdList, uint32 Mip
 {
 	SCOPE_CYCLE_COUNTER(STAT_D3D12LockTextureTime);
 
-	UE_TRACE_METADATA_SCOPE_ASSET_FNAME(GetName(), FName(TEXT("FRHITexture Lock")), GetOwnerName());
+	const static FName RHITextureLockName(TEXT("FRHITexture Lock"));
+	UE_TRACE_METADATA_SCOPE_ASSET_FNAME(GetName(), RHITextureLockName, GetOwnerName());
 
 	FD3D12Device* Device = GetParentDevice();
 	FD3D12Adapter* Adapter = Device->GetParentAdapter();
diff --git a/Engine/Source/Runtime/D3D12RHI/Private/D3D12TransientResourceAllocator.cpp b/Engine/Source/Runtime/D3D12RHI/Private/D3D12TransientResourceAllocator.cpp
--- a/Engine/Source/Runtime/D3D12RHI/Private/D3D12TransientResourceAllocator.cpp
+++ b/Engine/Source/Runtime/D3D12RHI/Private/D3D12TransientResourceAllocator.cpp
@@ -29,7 +29,8 @@ D3D12_RESOURCE_STATES GetInitialResourceState(const D3D12_RESOURCE_DESC& InDesc)
 FD3D12TransientHeap::FD3D12TransientHeap(const FInitializer& Initializer, FD3D12Adapter* Adapter, FD3D12Device* Device, FRHIGPUMask VisibleNodeMask)
 	: FRHITransientHeap(Initializer)
 {
-	UE_TRACE_METADATA_SCOPE_ASSET_FNAME(FName(TEXT("FD3D12TransientHeap")), FName(TEXT("FD3D12TransientHeap")), NAME_None);
+	const static FLazyName D3D12TransientHeapName(TEXT("FD3D12TransientHeap"));
+	UE_TRACE_METADATA_SCOPE_ASSET_FNAME(D3D12TransientHeapName, D3D12TransientHeapName, NAME_None);
 
 	D3D12_HEAP_FLAGS HeapFlags = D3D12_HEAP_FLAG_ALLOW_ALL_BUFFERS_AND_TEXTURES;
 
diff --git a/Engine/Source/Runtime/D3D12RHI/Private/Windows/WindowsD3D12Device.cpp b/Engine/Source/Runtime/D3D12RHI/Private/Windows/WindowsD3D12Device.cpp
--- a/Engine/Source/Runtime/D3D12RHI/Private/Windows/WindowsD3D12Device.cpp
+++ b/Engine/Source/Runtime/D3D12RHI/Private/Windows/WindowsD3D12Device.cpp
@@ -1969,7 +1969,9 @@ void FWindowsD3D12Adapter::CreateCommandSignatures()
 
 TUniquePtr<FD3D12DiagnosticBuffer> FD3D12Device::CreateDiagnosticBuffer(const D3D12_RESOURCE_DESC& Desc, const TCHAR* Name)
 {
-	UE_TRACE_METADATA_SCOPE_ASSET_FNAME(FName(TEXT("FD3D12DiagnosticBuffer")), FName(TEXT("FD3D12Device::CreateDiagnosticBuffer")), NAME_None);
+	const static FLazyName D3D12DiagnosticBufferName(TEXT("FD3D12DiagnosticBuffer"));
+	const static FLazyName CreateDiagnosticBufferName(TEXT("FD3D12Device::CreateDiagnosticBuffer"));
+	UE_TRACE_METADATA_SCOPE_ASSET_FNAME(D3D12DiagnosticBufferName, CreateDiagnosticBufferName, NAME_None);
 
 	TRefCountPtr<ID3D12Device3> D3D12Device3;
 	HRESULT hr = GetDevice()->QueryInterface(IID_PPV_ARGS(D3D12Device3.GetInitReference()));
diff --git a/Engine/Source/Runtime/D3D12RHI/Public/D3D12Resources.h b/Engine/Source/Runtime/D3D12RHI/Public/D3D12Resources.h
--- a/Engine/Source/Runtime/D3D12RHI/Public/D3D12Resources.h
+++ b/Engine/Source/Runtime/D3D12RHI/Public/D3D12Resources.h
@@ -879,7 +879,7 @@ public:
 
 	virtual uint32 GetParentGPUIndex() const override;
 
-	void UploadResourceData(FRHICommandListBase& InRHICmdList, FResourceArrayInterface* InResourceArray, D3D12_RESOURCE_STATES InDestinationState, const FName& AssetName = NAME_None, const FName& ClassName = NAME_None, const FName& PackageName = NAME_None);
+	void UploadResourceData(FRHICommandListBase& InRHICmdList, FResourceArrayInterface* InResourceArray, D3D12_RESOURCE_STATES InDestinationState, const TCHAR* AssetName = nullptr, const FName& ClassName = NAME_None, const FName& PackageName = NAME_None);
 	FD3D12SyncPointRef UploadResourceDataViaCopyQueue(FResourceArrayInterface* InResourceArray);
 
 	// FRHIResource overrides
diff --git a/Engine/Source/Runtime/Engine/Private/Components/LightComponent.cpp b/Engine/Source/Runtime/Engine/Private/Components/LightComponent.cpp
--- a/Engine/Source/Runtime/Engine/Private/Components/LightComponent.cpp
+++ b/Engine/Source/Runtime/Engine/Private/Components/LightComponent.cpp
@@ -39,9 +39,10 @@ void FStaticShadowDepthMap::InitRHI()
 {
 	if (FApp::CanEverRender() && Data && Data->ShadowMapSizeX > 0 && Data->ShadowMapSizeY > 0 && GMaxRHIFeatureLevel >= ERHIFeatureLevel::SM5)
 	{
+		const static FLazyName ClassName(TEXT("FStaticShadowDepthMap"));
 		const FRHITextureCreateDesc Desc =
 			FRHITextureCreateDesc::Create2D(TEXT("FStaticShadowDepthMap"), Data->ShadowMapSizeX, Data->ShadowMapSizeY, PF_R16F)
-			.SetClassName(TEXT("FStaticShadowDepthMap"));
+			.SetClassName(ClassName);
 
 		TextureRHI = RHICreateTexture(Desc);
 
diff --git a/Engine/Source/Runtime/Engine/Private/Particles/ParticleCurveTexture.cpp b/Engine/Source/Runtime/Engine/Private/Particles/ParticleCurveTexture.cpp
--- a/Engine/Source/Runtime/Engine/Private/Particles/ParticleCurveTexture.cpp
+++ b/Engine/Source/Runtime/Engine/Private/Particles/ParticleCurveTexture.cpp
@@ -507,6 +507,8 @@ FParticleCurveTexture::FParticleCurveTexture()
 */
 void FParticleCurveTexture::InitRHI()
 {
+	const static FLazyName ClassName(TEXT("FParticleCurveTexture"));
+
 	// 8-bit per channel RGBA texture for curves.
 	const FRHITextureCreateDesc Desc =
 		FRHITextureCreateDesc::Create2D(TEXT("ParticleCurveTexture"))
@@ -515,7 +517,7 @@ void FParticleCurveTexture::InitRHI()
 		.SetFormat(PF_B8G8R8A8)
 		.SetClearValue(FClearValueBinding(FLinearColor::Blue))
 		.SetInitialState(ERHIAccess::SRVMask)
-		.SetClassName(TEXT("FParticleCurveTexture"));
+		.SetClassName(ClassName);
 
 	CurveTextureRHI = RHICreateTexture(Desc);
 }
diff --git a/Engine/Source/Runtime/Engine/Private/Particles/ParticleGpuSimulation.cpp b/Engine/Source/Runtime/Engine/Private/Particles/ParticleGpuSimulation.cpp
--- a/Engine/Source/Runtime/Engine/Private/Particles/ParticleGpuSimulation.cpp
+++ b/Engine/Source/Runtime/Engine/Private/Particles/ParticleGpuSimulation.cpp
@@ -228,6 +228,7 @@ public:
 #define PARTICLE_STATE_POSITION_TEXTURE_NAME	TEXT("ParticleStatePosition")
 #define PARTICLE_STATE_VELOCITY_TEXTURE_NAME	TEXT("ParticleStateVelocity")
 
+		const static FLazyName ClassName(TEXT("FParticleStateTextures"));
 		{
 			const FRHITextureCreateDesc Desc =
 				FRHITextureCreateDesc::Create2D(PARTICLE_STATE_POSITION_TEXTURE_NAME)
@@ -236,7 +237,7 @@ public:
 				.SetClearValue(FClearValueBinding::Transparent)
 				.SetFlags(ETextureCreateFlags::RenderTargetable | ETextureCreateFlags::ShaderResource)
 				.SetInitialState(ERHIAccess::SRVMask)
-				.SetClassName(TEXT("FParticleStateTextures"));
+				.SetClassName(ClassName);
 
 			PositionTextureRHI = RHICreateTexture(Desc);
 		}
@@ -252,7 +253,7 @@ public:
 				.SetClearValue(FClearValueBinding::Transparent)
 				.SetFlags(ETextureCreateFlags::RenderTargetable | ETextureCreateFlags::ShaderResource)
 				.SetInitialState(ERHIAccess::SRVMask)
-				.SetClassName(TEXT("FParticleStateTextures"));
+				.SetClassName(ClassName);
 
 			VelocityTextureRHI = RHICreateTexture(Desc);
 		}
@@ -299,6 +300,7 @@ public:
 		const int32 SizeY = GParticleSimulationTextureSizeY;
 
 #define ATTRIBUTES_TEXTURE_NAME	TEXT("ParticleAttributes")
+		const static FLazyName ClassName(TEXT("FParticleAttributesTexture"));
 
 		const FRHITextureCreateDesc Desc =
 			FRHITextureCreateDesc::Create2D(ATTRIBUTES_TEXTURE_NAME)
@@ -307,7 +309,7 @@ public:
 			.SetClearValue(FClearValueBinding::Transparent)
 			.SetFlags(ETextureCreateFlags::RenderTargetable | ETextureCreateFlags::ShaderResource | ETextureCreateFlags::NoFastClear)
 			.SetInitialState(ERHIAccess::RTV)
-			.SetClassName(TEXT("FParticleAttributesTexture"));
+			.SetClassName(ClassName);
 
 		TextureRHI = RHICreateTexture(Desc);
 
diff --git a/Engine/Source/Runtime/Engine/Private/PrecomputedVolumetricLightmap.cpp b/Engine/Source/Runtime/Engine/Private/PrecomputedVolumetricLightmap.cpp
--- a/Engine/Source/Runtime/Engine/Private/PrecomputedVolumetricLightmap.cpp
+++ b/Engine/Source/Runtime/Engine/Private/PrecomputedVolumetricLightmap.cpp
@@ -20,13 +20,15 @@
 
 DECLARE_MEMORY_STAT(TEXT("Volumetric Lightmap"),STAT_VolumetricLightmapBuildData,STATGROUP_MapBuildData);
 
+const static FLazyName VolumetricLightmapDataLayerName(TEXT("VolumetricLightmapDataLayer"));
+
 void FVolumetricLightmapDataLayer::CreateTexture(FIntVector Dimensions)
 {
 	const FRHITextureCreateDesc Desc =
 		FRHITextureCreateDesc::Create3D(TEXT("VolumetricLightmap"), Dimensions, Format)
 		.SetFlags(ETextureCreateFlags::ShaderResource | ETextureCreateFlags::UAV)
 		.SetBulkData(this)
-		.SetClassName(TEXT("FVolumetricLightmapDataLayer"));
+		.SetClassName(VolumetricLightmapDataLayerName);
 
 	Texture = RHICreateTexture(Desc);
 }
@@ -36,7 +38,7 @@ void FVolumetricLightmapDataLayer::CreateTargetTexture(FIntVector Dimensions)
 	const FRHITextureCreateDesc Desc =
 		FRHITextureCreateDesc::Create3D(TEXT("VolumetricLightmap"), Dimensions, Format)
 		.SetFlags(ETextureCreateFlags::ShaderResource | ETextureCreateFlags::UAV)
-		.SetClassName(TEXT("FVolumetricLightmapDataLayer"));
+		.SetClassName(VolumetricLightmapDataLayerName);
 
 	Texture = RHICreateTexture(Desc);
 }
diff --git a/Engine/Source/Runtime/Engine/Private/PrimitiveUniformShaderParameters.cpp b/Engine/Source/Runtime/Engine/Private/PrimitiveUniformShaderParameters.cpp
--- a/Engine/Source/Runtime/Engine/Private/PrimitiveUniformShaderParameters.cpp
+++ b/Engine/Source/Runtime/Engine/Private/PrimitiveUniformShaderParameters.cpp
@@ -40,10 +40,11 @@ void FSinglePrimitiveStructured::InitRHI()
 	}
 
 	{
+		const static FLazyName ClassName(TEXT("FSinglePrimitiveStructured"));
 		const FRHITextureCreateDesc Desc =
 			FRHITextureCreateDesc::Create2D(TEXT("PrimitiveSceneDataTexture"), FPrimitiveSceneShaderData::DataStrideInFloat4s, 1, PF_A32B32G32R32F)
 			.SetFlags(ETextureCreateFlags::ShaderResource | ETextureCreateFlags::UAV)
-			.SetClassName(TEXT("FSinglePrimitiveStructured"));
+			.SetClassName(ClassName);
 
 		PrimitiveSceneDataTextureRHI = RHICreateTexture(Desc);
 		PrimitiveSceneDataTextureSRV = RHICreateShaderResourceView(PrimitiveSceneDataTextureRHI, 0);
diff --git a/Engine/Source/Runtime/Engine/Private/RawIndexBuffer.cpp b/Engine/Source/Runtime/Engine/Private/RawIndexBuffer.cpp
--- a/Engine/Source/Runtime/Engine/Private/RawIndexBuffer.cpp
+++ b/Engine/Source/Runtime/Engine/Private/RawIndexBuffer.cpp
@@ -338,11 +338,14 @@ FBufferRHIRef FRawStaticIndexBuffer::CreateRHIBuffer_Internal()
 
 		const EBufferUsageFlags BufferFlags = EBufferUsageFlags::Static | (bSRV ? EBufferUsageFlags::ShaderResource : EBufferUsageFlags::None);
 
+		const static FLazyName ClassName32(TEXT("FRawStaticIndexBuffer32"));
+		const static FLazyName ClassName16(TEXT("FRawStaticIndexBuffer16"));
+
 		// Create the index buffer.
 		FBufferRHIRef Ret;
 		FRHIResourceCreateInfo CreateInfo(Is32Bit() ? TEXT("FRawStaticIndexBuffer32") : TEXT("FRawStaticIndexBuffer16"), &IndexStorage);
-		CreateInfo.ClassName = CreateInfo.DebugName;
-		CreateInfo.AssetName = GetOwnerName();
+		CreateInfo.ClassName = Is32Bit() ? ClassName32 : ClassName16;
+		CreateInfo.OwnerName = GetOwnerName();
 		CreateInfo.bWithoutNativeResource = !SizeInBytes;
 		if (bRenderThread)
 		{
@@ -513,7 +516,7 @@ FBufferRHIRef FRawStaticIndexBuffer16or32Interface::CreateRHIIndexBufferInternal
 	// Create the index buffer.
 	FRHIResourceCreateInfo CreateInfo(InDebugName, ResourceArray);
 	CreateInfo.ClassName = InDebugName;
-	CreateInfo.AssetName = InOwnerName;
+	CreateInfo.OwnerName = InOwnerName;
 	EBufferUsageFlags Flags = EBufferUsageFlags::Static;
 
 	if (bNeedSRV)
diff --git a/Engine/Source/Runtime/Engine/Private/Rendering/MorphTargetVertexInfoBuffers.cpp b/Engine/Source/Runtime/Engine/Private/Rendering/MorphTargetVertexInfoBuffers.cpp
--- a/Engine/Source/Runtime/Engine/Private/Rendering/MorphTargetVertexInfoBuffers.cpp
+++ b/Engine/Source/Runtime/Engine/Private/Rendering/MorphTargetVertexInfoBuffers.cpp
@@ -18,10 +18,12 @@ void FMorphTargetVertexInfoBuffers::InitRHI()
 	check(NumTotalBatches > 0);
 	check(!bRHIIntialized);
 
+	const static FLazyName ClassName(TEXT("FMorphTargetVertexInfoBuffers"));
+
 	const uint32 BufferSize = MorphData.Num() * sizeof(uint32);
 	FRHIResourceCreateInfo CreateInfo(TEXT("MorphData"));
-	CreateInfo.ClassName = TEXT("FMorphTargetVertexInfoBuffers");
-	CreateInfo.AssetName = GetOwnerName();
+	CreateInfo.ClassName = ClassName;
+	CreateInfo.OwnerName = GetOwnerName();
 	MorphDataBuffer = RHICreateStructuredBuffer(sizeof(uint32), BufferSize, BUF_Static | BUF_ByteAddressBuffer | BUF_ShaderResource, ERHIAccess::SRVMask, CreateInfo);
 	MorphDataBuffer->SetOwnerName(GetOwnerName());
 	
diff --git a/Engine/Source/Runtime/Engine/Private/Rendering/Texture2DResource.cpp b/Engine/Source/Runtime/Engine/Private/Rendering/Texture2DResource.cpp
--- a/Engine/Source/Runtime/Engine/Private/Rendering/Texture2DResource.cpp
+++ b/Engine/Source/Runtime/Engine/Private/Rendering/Texture2DResource.cpp
@@ -11,6 +11,8 @@
 #include "RenderingThread.h"
 #include "Containers/ResourceArray.h"
 
+const static FLazyName Texture2DResourceName(TEXT("FTexture2DResource"));
+
 // TODO Only adding this setting to allow backwards compatibility to be forced.  The default
 // behavior is to NOT do this.  This variable should be removed in the future.  #ADDED 4.13
 static TAutoConsoleVariable<int32> CVarForceHighestMipOnUITexturesEnabled(
@@ -113,8 +115,8 @@ void FTexture2DResource::CreateTexture()
 		.SetFlags(CreationFlags)
 		.SetExtData(PlatformData->GetExtData())
 		.SetBulkData(ResourceMem)
-		.SetClassName(TEXT("FTexture2DResource"))
-		.SetAssetName(GetOwnerName());
+		.SetClassName(Texture2DResourceName)
+		.SetOwnerName(GetOwnerName());
 
 	TextureRHI = RHICreateTexture(Desc);
 
@@ -158,8 +160,8 @@ void FTexture2DResource::CreatePartiallyResidentTexture()
 		.SetFlags(CreationFlags | ETextureCreateFlags::Virtual)
 		.SetExtData(PlatformData->GetExtData())
 		.SetBulkData(ResourceMem)
-		.SetClassName(TEXT("FTexture2DResource"))
-		.SetAssetName(GetOwnerName());
+		.SetClassName(Texture2DResourceName)
+		.SetOwnerName(GetOwnerName());
 
 	TextureRHI = RHICreateTexture(Desc);
 
diff --git a/Engine/Source/Runtime/Engine/Private/Rendering/Texture3DResource.cpp b/Engine/Source/Runtime/Engine/Private/Rendering/Texture3DResource.cpp
--- a/Engine/Source/Runtime/Engine/Private/Rendering/Texture3DResource.cpp
+++ b/Engine/Source/Runtime/Engine/Private/Rendering/Texture3DResource.cpp
@@ -119,14 +119,15 @@ void FTexture3DResource::CreateTexture()
 	// Create the RHI texture.
 	{
 		const FTexture2DMipMap& FirstMip = *MipsView[FirstMipIdx];
+		const static FLazyName ClassName(TEXT("FTexture3DResource"));
 
 		FRHITextureCreateDesc Desc =
 			FRHITextureCreateDesc::Create3D(TEXT("FTexture3DResource"), FirstMip.SizeX, FirstMip.SizeY, FirstMip.SizeZ, PixelFormat)
 			.SetNumMips(State.NumRequestedLODs)
 			.SetFlags(CreationFlags)
 			.SetExtData(PlatformData->GetExtData())
-			.SetClassName(TEXT("FTexture3DResource"))
-			.SetAssetName(GetOwnerName());
+			.SetClassName(ClassName)
+			.SetOwnerName(GetOwnerName());
 
 		if (GUseTexture3DBulkDataRHI)
 		{
diff --git a/Engine/Source/Runtime/Engine/Private/Slate/SlateTextures.cpp b/Engine/Source/Runtime/Engine/Private/Slate/SlateTextures.cpp
--- a/Engine/Source/Runtime/Engine/Private/Slate/SlateTextures.cpp
+++ b/Engine/Source/Runtime/Engine/Private/Slate/SlateTextures.cpp
@@ -323,6 +323,8 @@ void FSlateTextureRenderTarget2DResource::InitDynamicRHI()
 
 	if( TargetSizeX > 0 && TargetSizeY > 0 )
 	{
+		const static FLazyName ClassName(TEXT("FSlateTextureRenderTarget2DResource"));
+
 		// Create the RHI texture. Only one mip is used and the texture is targetable for resolve.
 		const FRHITextureCreateDesc Desc =
 			FRHITextureCreateDesc::Create2D(TEXT("FSlateTextureRenderTarget2DResource"))
@@ -331,7 +333,7 @@ void FSlateTextureRenderTarget2DResource::InitDynamicRHI()
 			.SetClearValue(FClearValueBinding(ClearColor))
 			.SetFlags(ETextureCreateFlags::RenderTargetable | ETextureCreateFlags::ShaderResource)
 			.SetInitialState(ERHIAccess::SRVMask)
-			.SetClassName(TEXT("FSlateTextureRenderTarget2DResource"));
+			.SetClassName(ClassName);
 
 		RenderTargetTextureRHI = TextureRHI = RHICreateTexture(Desc);
 	}
diff --git a/Engine/Source/Runtime/Engine/Private/TextureCube.cpp b/Engine/Source/Runtime/Engine/Private/TextureCube.cpp
--- a/Engine/Source/Runtime/Engine/Private/TextureCube.cpp
+++ b/Engine/Source/Runtime/Engine/Private/TextureCube.cpp
@@ -447,6 +447,7 @@ public:
 		// Create the RHI texture.
 		const ETextureCreateFlags TexCreateFlags = (Owner->SRGB ? TexCreate_SRGB : TexCreate_None)  | (Owner->bNotOfflineProcessed ? TexCreate_None : TexCreate_OfflineProcessed);
 		const FString Name = Owner->GetPathName();
+		const static FLazyName ClassName(TEXT("FTextureCubeResource"));
 
 		const FRHITextureCreateDesc Desc =
 			FRHITextureCreateDesc::CreateCube(*Name)
@@ -455,8 +456,8 @@ public:
 			.SetNumMips(Owner->GetNumMips())
 			.SetFlags(TexCreateFlags)
 			.SetExtData(Owner->GetPlatformData() ? Owner->GetPlatformData()->GetExtData() : 0)
-			.SetClassName(TEXT("FTextureCubeResource"))
-			.SetAssetName(GetOwnerName());
+			.SetClassName(ClassName)
+			.SetOwnerName(GetOwnerName());
 
 		TextureCubeRHI = RHICreateTexture(Desc);
 
diff --git a/Engine/Source/Runtime/Engine/Private/TextureCubeArray.cpp b/Engine/Source/Runtime/Engine/Private/TextureCubeArray.cpp
--- a/Engine/Source/Runtime/Engine/Private/TextureCubeArray.cpp
+++ b/Engine/Source/Runtime/Engine/Private/TextureCubeArray.cpp
@@ -127,6 +127,7 @@ public:
 		// Create the RHI texture.
 		const ETextureCreateFlags TexCreateFlags = (Owner->SRGB ? ETextureCreateFlags::SRGB : ETextureCreateFlags::None) | (Owner->bNotOfflineProcessed ? ETextureCreateFlags::None : ETextureCreateFlags::OfflineProcessed);
 		const FString Name = Owner->GetPathName();
+		const static FLazyName ClassName(TEXT("FTextureCubeArrayResource"));
 
 		const FRHITextureCreateDesc Desc =
 			FRHITextureCreateDesc::CreateCubeArray(*Name)
@@ -136,8 +137,8 @@ public:
 			.SetNumMips(Owner->GetNumMips())
 			.SetFlags(TexCreateFlags)
 			.SetExtData(Owner->GetPlatformData()->GetExtData())
-			.SetClassName(TEXT("FTextureCubeArrayResource"))
-			.SetAssetName(GetOwnerName());
+			.SetClassName(ClassName)
+			.SetOwnerName(GetOwnerName());
 
 		TextureCubeRHI = RHICreateTexture(Desc);
 		TextureRHI = TextureCubeRHI;
diff --git a/Engine/Source/Runtime/Engine/Private/TextureRenderTarget2D.cpp b/Engine/Source/Runtime/Engine/Private/TextureRenderTarget2D.cpp
--- a/Engine/Source/Runtime/Engine/Private/TextureRenderTarget2D.cpp
+++ b/Engine/Source/Runtime/Engine/Private/TextureRenderTarget2D.cpp
@@ -619,6 +619,8 @@ void FTextureRenderTarget2DResource::InitDynamicRHI()
 
 	if( TargetSizeX > 0 && TargetSizeY > 0 )
 	{
+		const static FLazyName ClassName(TEXT("FTextureRenderTarget2DResource"));
+
 		FString ResourceName = Owner->GetName();
 		ETextureCreateFlags TexCreateFlags = GetCreateFlags();
 
@@ -630,8 +632,8 @@ void FTextureRenderTarget2DResource::InitDynamicRHI()
 			.SetFlags(TexCreateFlags)
 			.SetInitialState(ERHIAccess::SRVMask)
 			.SetClearValue(FClearValueBinding(ClearColor))
-			.SetClassName(TEXT("FTextureRenderTarget2DResource"))
-			.SetAssetName(GetOwnerName());
+			.SetClassName(ClassName)
+			.SetOwnerName(GetOwnerName());
 
 		TextureRHI = RenderTargetTextureRHI = RHICreateTexture(Desc);
 
diff --git a/Engine/Source/Runtime/Engine/Public/GPUSkinCache.h b/Engine/Source/Runtime/Engine/Public/GPUSkinCache.h
--- a/Engine/Source/Runtime/Engine/Public/GPUSkinCache.h
+++ b/Engine/Source/Runtime/Engine/Public/GPUSkinCache.h
@@ -197,10 +197,15 @@ public:
 		FRWBuffersAllocation(uint32 InNumVertices, bool InWithTangents, bool InUseIntermediateTangents, uint32 InIntermediateAccumulatedTangentsSize, FRHICommandListImmediate& RHICmdList, const FName& OwnerName)
 			: NumVertices(InNumVertices), WithTangents(InWithTangents), UseIntermediateTangents(InUseIntermediateTangents), IntermediateAccumulatedTangentsSize(InIntermediateAccumulatedTangentsSize)
 		{
+			const static FLazyName PositionsName(TEXT("SkinCachePositions"));
+			const static FLazyName TangentsName(TEXT("SkinCacheTangents"));
+			const static FLazyName IntermediateTangentsName(TEXT("SkinCacheIntermediateTangents"));
+			const static FLazyName IntermediateAccumulatedTangentsName(TEXT("SkinCacheIntermediateAccumulatedTangents"));
+
 			for (int32 Index = 0; Index < NUM_BUFFERS; ++Index)
 			{
-				PositionBuffers[Index].Buffer.ClassName = TEXT("SkinCachePositions");
-				PositionBuffers[Index].Buffer.AssetName = OwnerName;
+				PositionBuffers[Index].Buffer.ClassName = PositionsName;
+				PositionBuffers[Index].Buffer.OwnerName = OwnerName;
 				PositionBuffers[Index].Buffer.Initialize(TEXT("SkinCachePositions"), PosBufferBytesPerElement, NumVertices * 3, PF_R32_FLOAT, BUF_Static);
 				PositionBuffers[Index].Buffer.Buffer->SetOwnerName(OwnerName);
 				PositionBuffers[Index].AccessState = ERHIAccess::Unknown;
@@ -210,15 +215,15 @@ public:
 				// OpenGL ES does not support writing to RGBA16_SNORM images, instead pack data into SINT in the shader
 				const EPixelFormat TangentsFormat = IsOpenGLPlatform(GMaxRHIShaderPlatform) ? PF_R16G16B16A16_SINT : PF_R16G16B16A16_SNORM;
 				
-				Tangents.Buffer.ClassName = TEXT("SkinCacheTangents");
-				Tangents.Buffer.AssetName = OwnerName;
+				Tangents.Buffer.ClassName = TangentsName;
+				Tangents.Buffer.OwnerName = OwnerName;
 				Tangents.Buffer.Initialize(TEXT("SkinCacheTangents"), TangentBufferBytesPerElement, NumVertices * 2, TangentsFormat, BUF_Static);
 				Tangents.Buffer.Buffer->SetOwnerName(OwnerName);
 				Tangents.AccessState = ERHIAccess::Unknown;
 				if (UseIntermediateTangents)
 				{
-					IntermediateTangents.Buffer.ClassName = TEXT("SkinCacheIntermediateTangents");
-					IntermediateTangents.Buffer.AssetName = OwnerName;
+					IntermediateTangents.Buffer.ClassName = IntermediateTangentsName;
+					IntermediateTangents.Buffer.OwnerName = OwnerName;
 					IntermediateTangents.Buffer.Initialize(TEXT("SkinCacheIntermediateTangents"), TangentBufferBytesPerElement, NumVertices * 2, TangentsFormat, BUF_Static);
 					IntermediateTangents.Buffer.Buffer->SetOwnerName(OwnerName);
 					IntermediateTangents.AccessState = ERHIAccess::Unknown;
@@ -226,8 +231,8 @@ public:
 			}
 			if (IntermediateAccumulatedTangentsSize > 0)
 			{
-				IntermediateAccumulatedTangents.Buffer.ClassName = TEXT("SkinCacheIntermediateAccumulatedTangents");
-				IntermediateAccumulatedTangents.Buffer.AssetName = OwnerName;
+				IntermediateAccumulatedTangents.Buffer.ClassName = IntermediateAccumulatedTangentsName;
+				IntermediateAccumulatedTangents.Buffer.OwnerName = OwnerName;
 				IntermediateAccumulatedTangents.Buffer.Initialize(TEXT("SkinCacheIntermediateAccumulatedTangents"), sizeof(int32), IntermediateAccumulatedTangentsSize * FGPUSkinCache::IntermediateAccumBufferNumInts, PF_R32_SINT, BUF_UnorderedAccess);
 				IntermediateAccumulatedTangents.Buffer.Buffer->SetOwnerName(OwnerName);
 				IntermediateAccumulatedTangents.AccessState = ERHIAccess::Unknown;
diff --git a/Engine/Source/Runtime/Engine/Public/Rendering/SkeletalMeshDuplicatedVerticesBuffer.h b/Engine/Source/Runtime/Engine/Public/Rendering/SkeletalMeshDuplicatedVerticesBuffer.h
--- a/Engine/Source/Runtime/Engine/Public/Rendering/SkeletalMeshDuplicatedVerticesBuffer.h
+++ b/Engine/Source/Runtime/Engine/Public/Rendering/SkeletalMeshDuplicatedVerticesBuffer.h
@@ -107,13 +107,14 @@ public:
 
     virtual void InitRHI() override
     {
+		const static FLazyName DuplicatedVerticesBuffeName(TEXT("FDuplicatedVerticesBuffer"));
         {
             FResourceArrayInterface* ResourceArray = DupVertData.GetResourceArray();
             check(ResourceArray->GetResourceDataSize() > 0);
 
             FRHIResourceCreateInfo CreateInfo(TEXT("DuplicatedVerticesIndexBuffer"), ResourceArray);
-			CreateInfo.ClassName = TEXT("FDuplicatedVerticesBuffer DuplicatedVerticesIndexBuffer");
-			CreateInfo.AssetName = GetOwnerName();
+			CreateInfo.ClassName = DuplicatedVerticesBuffeName;
+			CreateInfo.OwnerName = GetOwnerName();
             DuplicatedVerticesIndexBuffer.VertexBufferRHI = RHICreateVertexBuffer(ResourceArray->GetResourceDataSize(), BUF_Static | BUF_ShaderResource, CreateInfo);
 			DuplicatedVerticesIndexBuffer.VertexBufferRHI->SetOwnerName(GetOwnerName());
             DuplicatedVerticesIndexBuffer.VertexBufferSRV = RHICreateShaderResourceView(DuplicatedVerticesIndexBuffer.VertexBufferRHI, sizeof(uint32), PF_R32_UINT);
@@ -124,8 +125,8 @@ public:
             check(ResourceArray->GetResourceDataSize() > 0);
 
             FRHIResourceCreateInfo CreateInfo(TEXT("LengthAndIndexDuplicatedVerticesIndexBuffer"), ResourceArray);
-			CreateInfo.ClassName = TEXT("FDuplicatedVerticesBuffer LengthAndIndexDuplicatedVerticesIndexBuffer");
-			CreateInfo.AssetName = GetOwnerName();
+			CreateInfo.ClassName = DuplicatedVerticesBuffeName;
+			CreateInfo.OwnerName = GetOwnerName();
             LengthAndIndexDuplicatedVerticesIndexBuffer.VertexBufferRHI = RHICreateVertexBuffer(ResourceArray->GetResourceDataSize(), BUF_Static | BUF_ShaderResource, CreateInfo);
 			LengthAndIndexDuplicatedVerticesIndexBuffer.VertexBufferRHI->SetOwnerName(GetOwnerName());
             LengthAndIndexDuplicatedVerticesIndexBuffer.VertexBufferSRV = RHICreateShaderResourceView(LengthAndIndexDuplicatedVerticesIndexBuffer.VertexBufferRHI, sizeof(uint32), PF_R32_UINT);
diff --git a/Engine/Source/Runtime/Landscape/Private/LandscapeRender.cpp b/Engine/Source/Runtime/Landscape/Private/LandscapeRender.cpp
--- a/Engine/Source/Runtime/Landscape/Private/LandscapeRender.cpp
+++ b/Engine/Source/Runtime/Landscape/Private/LandscapeRender.cpp
@@ -673,7 +673,8 @@ void FLandscapeRenderSystem::UpdateBuffers()
 		if (!SectionLODBiasBuffer.IsValid())
 		{
 			FRHIResourceCreateInfo CreateInfo(TEXT("SectionLODBiasBuffer"), &SectionLODBiases);
-			CreateInfo.ClassName = TEXT("FLandscapeRenderSystem");
+			const static FLazyName ClassName(TEXT("FLandscapeRenderSystem"));
+			CreateInfo.ClassName = ClassName;
 			SectionLODBiasBuffer = RHICreateVertexBuffer(SectionLODBiases.GetResourceDataSize(), BUF_ShaderResource | BUF_Dynamic, CreateInfo);
 			SectionLODBiasSRV = RHICreateShaderResourceView(SectionLODBiasBuffer, sizeof(float), PF_R32_FLOAT);
 			bUpdateUB = true;
@@ -2593,8 +2594,9 @@ void FLandscapeVertexBuffer::InitRHI()
 
 	// create a static vertex buffer
 	FRHIResourceCreateInfo CreateInfo(TEXT("FLandscapeVertexBuffer"));
-	CreateInfo.ClassName = FName(TEXT("FLandscapeVertexBuffer"));
-	CreateInfo.AssetName = GetOwnerName();
+	const static FLazyName ClassName(TEXT("FLandscapeVertexBuffer"));
+	CreateInfo.ClassName = ClassName;
+	CreateInfo.OwnerName = GetOwnerName();
 	VertexBufferRHI = RHICreateBuffer(NumVertices * sizeof(FLandscapeVertex), BUF_Static | BUF_VertexBuffer, 0, ERHIAccess::VertexOrIndexBuffer, CreateInfo);
 	FLandscapeVertex* Vertex = (FLandscapeVertex*)RHILockBuffer(VertexBufferRHI, 0, NumVertices * sizeof(FLandscapeVertex), RLM_WriteOnly);
 	int32 VertexIndex = 0;
diff --git a/Engine/Source/Runtime/MediaAssets/Private/Misc/MediaTextureResource.cpp b/Engine/Source/Runtime/MediaAssets/Private/Misc/MediaTextureResource.cpp
--- a/Engine/Source/Runtime/MediaAssets/Private/Misc/MediaTextureResource.cpp
+++ b/Engine/Source/Runtime/MediaAssets/Private/Misc/MediaTextureResource.cpp
@@ -1319,6 +1319,7 @@ void FMediaTextureResource::CreateOutputRenderTarget(const FIntPoint & InDim, EP
 	{
 		MipGenerationCache.SafeRelease();
 
+		const static FLazyName ClassName(TEXT("FMediaTextureResource"));
 		const FRHITextureCreateDesc Desc =
 			FRHITextureCreateDesc::Create2D(TEXT("MediaTextureResourceOutput"))
 			.SetExtent(InDim)
@@ -1327,8 +1328,8 @@ void FMediaTextureResource::CreateOutputRenderTarget(const FIntPoint & InDim, EP
 			.SetFlags(OutputCreateFlags | ETextureCreateFlags::RenderTargetable | ETextureCreateFlags::ShaderResource)
 			.SetInitialState(ERHIAccess::SRVMask)
 			.SetClearValue(FClearValueBinding(InClearColor))
-			.SetClassName(TEXT("FMediaTextureResource"))
-			.SetAssetName(GetOwnerName());
+			.SetClassName(ClassName)
+			.SetOwnerName(GetOwnerName());
 
 		OutputTarget = RHICreateTexture(Desc);
 
diff --git a/Engine/Source/Runtime/RHI/Public/RHIResources.h b/Engine/Source/Runtime/RHI/Public/RHIResources.h
--- a/Engine/Source/Runtime/RHI/Public/RHIResources.h
+++ b/Engine/Source/Runtime/RHI/Public/RHIResources.h
@@ -391,8 +391,7 @@ struct FRHIResourceCreateInfo
 		ExtData = InExtData;
 	}
 
-	FName GetTraceClassName() const									{ return ClassName == NAME_None ? FName(TEXT("FRHIBuffer")) : ClassName; }
-	FName GetTraceAssetName() const									{ return AssetName == NAME_None ? FName(DebugName) : AssetName; }
+	FName GetTraceClassName() const									{ const static FLazyName FRHIBufferName(TEXT("FRHIBuffer")); return ClassName == NAME_None ? FRHIBufferName : ClassName; }
 
 	// for CreateTexture calls
 	FResourceBulkDataInterface* BulkData;
@@ -414,7 +413,7 @@ struct FRHIResourceCreateInfo
 	uint32 ExtData;
 
 	FName ClassName = NAME_None;	// The owner class of FRHIBuffer used for Insight asset metadata tracing
-	FName AssetName = NAME_None;	// The owner asset name used for Insight asset metadata tracing
+	FName OwnerName = NAME_None;	// The owner name used for Insight asset metadata tracing
 };
 
 class FExclusiveDepthStencil
@@ -1756,10 +1755,9 @@ struct FRHITextureCreateDesc : public FRHITextureDesc
 	FRHITextureCreateDesc& SetBulkData(FResourceBulkDataInterface* InBulkData) { BulkData = InBulkData;                    return *this; }
 	FRHITextureCreateDesc& DetermineInititialState()                           { if (InitialState == ERHIAccess::Unknown) InitialState = RHIGetDefaultResourceState(Flags, BulkData != nullptr); return *this; }
 	FRHITextureCreateDesc& SetFastVRAMPercentage(float In)                     { FastVRAMPercentage = uint8(FMath::Clamp(In, 0.f, 1.0f) * 0xFF); return *this; }
-	FRHITextureCreateDesc& SetClassName(const TCHAR* InClassName)			   { ClassName = FName(InClassName);           return *this; }
-	FRHITextureCreateDesc& SetAssetName(const FName& InAssetName)			   { AssetName = InAssetName;                  return *this; }
-	FName GetTraceClassName() const											   { return ClassName == NAME_None ? FName(TEXT("FRHITexture")) : ClassName; }
-	FName GetTraceAssetName() const											   { return AssetName == NAME_None ? FName(DebugName) : AssetName; }
+	FRHITextureCreateDesc& SetClassName(const FName& InClassName)			   { ClassName = InClassName;				   return *this; }
+	FRHITextureCreateDesc& SetOwnerName(const FName& InOwnerName)			   { OwnerName = InOwnerName;                  return *this; }
+	FName GetTraceClassName() const											   { const static FLazyName FRHITextureName(TEXT("FRHITexture")); return ClassName == NAME_None ? FRHITextureName : ClassName; }
 
 	/* The RHI access state that the resource will be created in. */
 	ERHIAccess InitialState = ERHIAccess::Unknown;
@@ -1770,9 +1768,8 @@ struct FRHITextureCreateDesc : public FRHITextureDesc
 	/* Optional initial data to fill the resource with. */
 	FResourceBulkDataInterface* BulkData = nullptr;
 
-private:
 	FName ClassName = NAME_None;	// The owner class of FRHITexture used for Insight asset metadata tracing
-	FName AssetName = NAME_None;	// The owner asset name used for Insight asset metadata tracing
+	FName OwnerName = NAME_None;	// The owner name used for Insight asset metadata tracing
 };
 
 class FRHITexture : public FRHIViewableResource
diff --git a/Engine/Source/Runtime/RHI/Public/RHIUtilities.h b/Engine/Source/Runtime/RHI/Public/RHIUtilities.h
--- a/Engine/Source/Runtime/RHI/Public/RHIUtilities.h
+++ b/Engine/Source/Runtime/RHI/Public/RHIUtilities.h
@@ -163,7 +163,7 @@ struct FRWBuffer
 	FShaderResourceViewRHIRef SRV;
 	uint32 NumBytes;
 	FName ClassName = NAME_None;	// The owner class of Buffer used for Insight asset metadata tracing, set it before calling Initialize()
-	FName AssetName = NAME_None;	// The owner asset name used for Insight asset metadata tracing, set it before calling Initialize()
+	FName OwnerName = NAME_None;	// The owner name used for Insight asset metadata tracing, set it before calling Initialize()
 
 	FRWBuffer()
 		: NumBytes(0)
@@ -220,7 +220,7 @@ struct FRWBuffer
 		NumBytes = BytesPerElement * NumElements;
 		FRHIResourceCreateInfo CreateInfo(InDebugName);
 		CreateInfo.ClassName = ClassName;
-		CreateInfo.AssetName = AssetName;
+		CreateInfo.OwnerName = OwnerName;
 		CreateInfo.ResourceArray = InResourceArray;
 		Buffer = RHICreateVertexBuffer(NumBytes, BUF_UnorderedAccess | BUF_ShaderResource | AdditionalUsage, InResourceState, CreateInfo);
 		UAV = RHICreateUnorderedAccessView(Buffer, UE_PIXELFORMAT_TO_UINT8(Format));
diff --git a/Engine/Source/Runtime/RenderCore/Private/RenderResource.cpp b/Engine/Source/Runtime/RenderCore/Private/RenderResource.cpp
--- a/Engine/Source/Runtime/RenderCore/Private/RenderResource.cpp
+++ b/Engine/Source/Runtime/RenderCore/Private/RenderResource.cpp
@@ -335,7 +335,7 @@ FBufferRHIRef FRenderResource::CreateRHIBufferInternal(
 	const uint32 SizeInBytes = ResourceArray ? ResourceArray->GetResourceDataSize() : 0;
 	FRHIResourceCreateInfo CreateInfo(InDebugName, ResourceArray);
 	CreateInfo.ClassName = FName(InDebugName);
-	CreateInfo.AssetName = InOwnerName;
+	CreateInfo.OwnerName = InOwnerName;
 	CreateInfo.bWithoutNativeResource = bWithoutNativeResource;
 
 	FBufferRHIRef Buffer;
diff --git a/Engine/Source/Runtime/RenderCore/Private/RenderTargetPool.cpp b/Engine/Source/Runtime/RenderCore/Private/RenderTargetPool.cpp
--- a/Engine/Source/Runtime/RenderCore/Private/RenderTargetPool.cpp
+++ b/Engine/Source/Runtime/RenderCore/Private/RenderTargetPool.cpp
@@ -111,7 +111,8 @@ TRefCountPtr<IPooledRenderTarget> FRenderTargetPool::FindFreeElement(FRHITexture
 
 		const ERHIAccess AccessInitial = ERHIAccess::SRVMask;
 		FRHITextureCreateDesc CreateDesc(Desc, AccessInitial, Name);
-		CreateDesc.SetClassName(TEXT("FPooledRenderTarget"));
+		const static FLazyName ClassName(TEXT("FPooledRenderTarget"));
+		CreateDesc.SetClassName(ClassName);
 
 		Found = new FPooledRenderTarget(
 			RHICreateTexture(CreateDesc),
diff --git a/Engine/Source/Runtime/Renderer/Private/SystemTextures.cpp b/Engine/Source/Runtime/Renderer/Private/SystemTextures.cpp
--- a/Engine/Source/Runtime/Renderer/Private/SystemTextures.cpp
+++ b/Engine/Source/Runtime/Renderer/Private/SystemTextures.cpp
@@ -101,6 +101,8 @@ void SetDummyTextureArrayData(FRHITexture* Texture, const DataType& DummyData)
 	RHIUnlockTexture2DArray(Texture, 0, 0, false);
 }
 
+const static FLazyName SystemTexturesName(TEXT("FSystemTextures"));
+
 void FSystemTextures::InitializeCommonTextures(FRHICommandListImmediate& RHICmdList)
 {
 	// First initialize textures that are common to all feature levels. This is always done the first time we come into this function, as doesn't care about the
@@ -117,7 +119,7 @@ void FSystemTextures::InitializeCommonTextures(FRHICommandListImmediate& RHICmdL
 		const FRHITextureCreateDesc Desc =
 			FRHITextureCreateDesc::Create2D(TEXT("BlackDummy"), 1, 1, PF_B8G8R8A8)
 			.SetFlags(ETextureCreateFlags::ShaderResource)
-			.SetClassName(TEXT("FSystemTextures"));
+			.SetClassName(SystemTexturesName);
 
 		FTextureRHIRef Texture = RHICreateTexture(Desc);
 		SetDummyTextureData<FColor>(Texture, FColor(0, 0, 0, 0));
@@ -129,7 +131,7 @@ void FSystemTextures::InitializeCommonTextures(FRHICommandListImmediate& RHICmdL
 		const FRHITextureCreateDesc Desc =
 			FRHITextureCreateDesc::Create2DArray(TEXT("BlackArrayDummy"), 1, 1, 1, PF_B8G8R8A8)
 			.SetFlags(ETextureCreateFlags::ShaderResource)
-			.SetClassName(TEXT("FSystemTextures"));
+			.SetClassName(SystemTexturesName);
 
 		FTextureRHIRef Texture = RHICreateTexture(Desc);
 		SetDummyTextureArrayData<FColor>(Texture, FColor(0, 0, 0, 0));
@@ -141,7 +143,7 @@ void FSystemTextures::InitializeCommonTextures(FRHICommandListImmediate& RHICmdL
 		const FRHITextureCreateDesc Desc =
 			FRHITextureCreateDesc::Create2D(TEXT("ZeroUIntDummy"), 1, 1, PF_R32_UINT)
 			.SetFlags(ETextureCreateFlags::ShaderResource)
-			.SetClassName(TEXT("FSystemTextures"));
+			.SetClassName(SystemTexturesName);
 
 		FTextureRHIRef Texture = RHICreateTexture(Desc);
 		SetDummyTextureData<uint32>(Texture, 0u);
@@ -153,7 +155,7 @@ void FSystemTextures::InitializeCommonTextures(FRHICommandListImmediate& RHICmdL
 		const FRHITextureCreateDesc Desc =
 			FRHITextureCreateDesc::Create2DArray(TEXT("ZeroUIntArrayDummy"), 1, 1, 1, PF_R32_UINT)
 			.SetFlags(ETextureCreateFlags::ShaderResource)
-			.SetClassName(TEXT("FSystemTextures"));
+			.SetClassName(SystemTexturesName);
 
 		FTextureRHIRef Texture = RHICreateTexture(Desc);
 		SetDummyTextureArrayData<uint32>(Texture, 0u);
@@ -165,7 +167,7 @@ void FSystemTextures::InitializeCommonTextures(FRHICommandListImmediate& RHICmdL
 		const FRHITextureCreateDesc Desc =
 			FRHITextureCreateDesc::Create2D(TEXT("BlackAlphaOneDummy"), 1, 1, PF_B8G8R8A8)
 			.SetFlags(ETextureCreateFlags::ShaderResource)
-			.SetClassName(TEXT("FSystemTextures"));
+			.SetClassName(SystemTexturesName);
 
 		FTextureRHIRef Texture = RHICreateTexture(Desc);
 		SetDummyTextureData<FColor>(Texture, FColor(0, 0, 0, 255));
@@ -177,7 +179,7 @@ void FSystemTextures::InitializeCommonTextures(FRHICommandListImmediate& RHICmdL
 		const FRHITextureCreateDesc Desc =
 			FRHITextureCreateDesc::Create2DArray(TEXT("GreenDummy"), 1, 1, PF_B8G8R8A8)
 			.SetFlags(ETextureCreateFlags::ShaderResource)
-			.SetClassName(TEXT("FSystemTextures"));
+			.SetClassName(SystemTexturesName);
 
 		FTextureRHIRef Texture = RHICreateTexture(Desc);
 		SetDummyTextureData<FColor>(Texture, FColor(0, 255, 0, 255));
@@ -189,7 +191,7 @@ void FSystemTextures::InitializeCommonTextures(FRHICommandListImmediate& RHICmdL
 		const FRHITextureCreateDesc Desc =
 			FRHITextureCreateDesc::Create2D(TEXT("DefaultNormal8Bit"), 1, 1, PF_B8G8R8A8)
 			.SetFlags(ETextureCreateFlags::ShaderResource)
-			.SetClassName(TEXT("FSystemTextures"));
+			.SetClassName(SystemTexturesName);
 
 		FTextureRHIRef Texture = RHICreateTexture(Desc);
 		SetDummyTextureData<FColor>(Texture, FColor(128, 128, 128, 255));
@@ -203,7 +205,7 @@ void FSystemTextures::InitializeCommonTextures(FRHICommandListImmediate& RHICmdL
 		const FRHITextureCreateDesc Desc =
 			FRHITextureCreateDesc::Create2D(TEXT("PerlinNoiseGradient"), Extent, Extent, PF_B8G8R8A8)
 			.SetFlags(ETextureCreateFlags::ShaderResource)
-			.SetClassName(TEXT("FSystemTextures"));
+			.SetClassName(SystemTexturesName);
 	
 		FTextureRHIRef Texture = RHICreateTexture(Desc);
 
@@ -241,7 +243,7 @@ void FSystemTextures::InitializeCommonTextures(FRHICommandListImmediate& RHICmdL
 		const FRHITextureCreateDesc Desc =
 			FRHITextureCreateDesc::Create2D(TEXT("MaxFP16Depth"), 1, 1, PF_FloatRGBA)
 			.SetFlags(ETextureCreateFlags::ShaderResource)
-			.SetClassName(TEXT("FSystemTextures"));
+			.SetClassName(SystemTexturesName);
 
 		FTextureRHIRef Texture = RHICreateTexture(Desc);
 		SetDummyTextureData<FFloat16Color>(Texture, FFloat16Color(FLinearColor(65500.0f, 65500.0f, 65500.0f, 65500.0f)));
@@ -259,7 +261,7 @@ void FSystemTextures::InitializeCommonTextures(FRHICommandListImmediate& RHICmdL
 		const FRHITextureCreateDesc Desc =
 			FRHITextureCreateDesc::Create2D(TEXT("StencilDummy"), 1, 1, PF_R8G8B8A8_UINT)
 			.SetFlags(ETextureCreateFlags::ShaderResource)
-			.SetClassName(TEXT("FSystemTextures"));
+			.SetClassName(SystemTexturesName);
 
 		FTextureRHIRef Texture = RHICreateTexture(Desc);
 		SetDummyTextureData<FColor>(Texture, FColor::White);
@@ -272,7 +274,7 @@ void FSystemTextures::InitializeCommonTextures(FRHICommandListImmediate& RHICmdL
 		const FRHITextureCreateDesc Desc =
 			FRHITextureCreateDesc::Create2D(TEXT("MidGreyDummy"), 1, 1, PF_FloatRGBA)
 			.SetFlags(ETextureCreateFlags::ShaderResource)
-			.SetClassName(TEXT("FSystemTextures"));
+			.SetClassName(SystemTexturesName);
 
 		FTextureRHIRef Texture = RHICreateTexture(Desc);
 		SetDummyTextureData<FFloat16Color>(Texture, FFloat16Color(FLinearColor(0.5f, 0.5f, 0.5f, 0.5f)));
@@ -310,7 +312,7 @@ void FSystemTextures::InitializeFeatureLevelDependentTextures(FRHICommandListImm
 		const FRHITextureCreateDesc Desc =
 			FRHITextureCreateDesc::Create2D(TEXT("SobolSampling"), 32, 16, PF_R16_UINT)
 			.SetFlags(ETextureCreateFlags::ShaderResource)
-			.SetClassName(TEXT("FSystemTextures"));
+			.SetClassName(SystemTexturesName);
 
 		FTextureRHIRef Texture = RHICreateTexture(Desc);
 		// Write the contents of the texture.
@@ -438,7 +440,7 @@ void FSystemTextures::InitializeFeatureLevelDependentTextures(FRHICommandListImm
 		const FRHITextureCreateDesc Desc =
 			FRHITextureCreateDesc::Create2D(TEXT("PreintegratedGF"), Extent, Format)
 			.SetFlags(ETextureCreateFlags::ShaderResource)
-			.SetClassName(TEXT("FSystemTextures"));
+			.SetClassName(SystemTexturesName);
 
 		FTextureRHIRef Texture = RHICreateTexture(Desc);
 
@@ -562,7 +564,7 @@ void FSystemTextures::InitializeFeatureLevelDependentTextures(FRHICommandListImm
 			const FRHITextureCreateDesc Desc =
 				FRHITextureCreateDesc::Create3D(TEXT("PerlinNoise3D"), Extent, Extent, Extent, PF_B8G8R8A8)
 				.SetFlags(ETextureCreateFlags::ShaderResource | ETextureCreateFlags::NoTiling)
-				.SetClassName(TEXT("FSystemTextures"));
+				.SetClassName(SystemTexturesName);
 
 			FTextureRHIRef Texture3D = RHICreateTexture(Desc);
 
@@ -677,7 +679,7 @@ void FSystemTextures::InitializeFeatureLevelDependentTextures(FRHICommandListImm
 				const FRHITextureCreateDesc Desc =
 					FRHITextureCreateDesc::Create2D(TEXT("LTCMat"), LTC_Size, LTC_Size, PF_FloatRGBA)
 					.SetFlags(ETextureCreateFlags::ShaderResource | ETextureCreateFlags::FastVRAM)
-					.SetClassName(TEXT("FSystemTextures"));
+					.SetClassName(SystemTexturesName);
 
 				FTextureRHIRef Texture = RHICreateTexture(Desc);
 
@@ -706,7 +708,7 @@ void FSystemTextures::InitializeFeatureLevelDependentTextures(FRHICommandListImm
 				const FRHITextureCreateDesc Desc =
 					FRHITextureCreateDesc::Create2D(TEXT("LTCAmp"), LTC_Size, LTC_Size, PF_G16R16F)
 					.SetFlags(ETextureCreateFlags::ShaderResource | ETextureCreateFlags::FastVRAM)
-					.SetClassName(TEXT("FSystemTextures"));
+					.SetClassName(SystemTexturesName);
 
 				FTexture2DRHIRef Texture = RHICreateTexture(Desc);
 
@@ -769,7 +771,7 @@ void FSystemTextures::InitializeFeatureLevelDependentTextures(FRHICommandListImm
 				const FRHITextureCreateDesc Desc =
 					FRHITextureCreateDesc::Create2D(TEXT("SSAORandomization"), Extent, Extent, PF_R8G8)
 					.SetFlags(ETextureCreateFlags::ShaderResource)
-					.SetClassName(TEXT("FSystemTextures"));
+					.SetClassName(SystemTexturesName);
 
 				FTextureRHIRef Texture = RHICreateTexture(Desc);
 
@@ -813,7 +815,7 @@ void FSystemTextures::InitializeFeatureLevelDependentTextures(FRHICommandListImm
 			.SetDepth(bGTAOPreIngegratedUsingVolumeLUT ? Extent : 1)
 			.SetFormat(PF_R16F)
 			.SetFlags(ETextureCreateFlags::ShaderResource)
-			.SetClassName(TEXT("FSystemTextures"));
+			.SetClassName(SystemTexturesName);
 
 		TRefCountPtr<FRHITexture> Texture = RHICreateTexture(Desc);
 
diff --git a/Engine/Source/Runtime/SlateRHIRenderer/Private/SlateRHIFontTexture.cpp b/Engine/Source/Runtime/SlateRHIRenderer/Private/SlateRHIFontTexture.cpp
--- a/Engine/Source/Runtime/SlateRHIRenderer/Private/SlateRHIFontTexture.cpp
+++ b/Engine/Source/Runtime/SlateRHIRenderer/Private/SlateRHIFontTexture.cpp
@@ -23,10 +23,11 @@ void FSlateFontTextureRHIResource::InitDynamicRHI()
 
 		check( !IsValidRef( ShaderResource) );
 
+		const static FLazyName ClassName(TEXT("FSlateFontTextureRHIResource"));
 		FRHITextureCreateDesc Desc =
 			FRHITextureCreateDesc::Create2D(TEXT("FSlateFontTextureRHIResource"), Width, Height, PixelFormat)
 			.SetFlags(ETextureCreateFlags::Dynamic)
-			.SetClassName(TEXT("FSlateFontTextureRHIResource"));
+			.SetClassName(ClassName);
 
 		if (!bIsGrayscale)
 		{
