Fixed: Crash occurs opening OpenGL Deferred Android project
Added a raster path for updating IES Atlas to avoid OpenGL limitations with R16F images

diff --git a/Engine/Shaders/Private/IESAtlas.usf b/Engine/Shaders/Private/IESAtlas.usf
--- a/Engine/Shaders/Private/IESAtlas.usf
+++ b/Engine/Shaders/Private/IESAtlas.usf
@@ -6,7 +6,7 @@
 
 ////////////////////////////////////////////////////////////////////////////////////////////////
 // Add texture
-#if SHADER_ADD_TEXTURE
+#if SHADER_ADD_TEXTURE_CS
 
 Texture2D<float4> InTexture_0;
 Texture2D<float4> InTexture_1;
@@ -69,7 +69,24 @@ void MainCS(uint3 DispatchThreadId : SV_DispatchThreadID)
 	}
 }
 
-#endif // SHADER_ADD_TEXTURE
+#endif // SHADER_ADD_TEXTURE_CS
+
+#if SHADER_ADD_TEXTURE_PS
+
+Texture2D<float4> InTexture;
+SamplerState InSampler;
+int2 AtlasResolution;
+
+void MainPS(float4 SvPosition : SV_POSITION,
+			out float OutColor : SV_Target0)
+{
+	const float2 SrcUV = SvPosition.xy / float2(AtlasResolution);
+	float Color = InTexture.SampleLevel(InSampler, SrcUV, 0).x;
+	// Ensure there is no NaN value
+	OutColor = -min(-Color, 0);
+}
+
+#endif // SHADER_ADD_TEXTURE_PS
 
 ////////////////////////////////////////////////////////////////////////////////////////////////
 // Debug
diff --git a/Engine/Source/Runtime/Renderer/Private/IESTextureManager.cpp b/Engine/Source/Runtime/Renderer/Private/IESTextureManager.cpp
--- a/Engine/Source/Runtime/Renderer/Private/IESTextureManager.cpp
+++ b/Engine/Source/Runtime/Renderer/Private/IESTextureManager.cpp
@@ -19,6 +19,8 @@
 #include "CommonRenderResources.h"
 #include "ScreenPass.h"
 #include "Engine/TextureLightProfile.h"
+#include "PixelShaderUtils.h"
+#include "DataDrivenShaderPlatformInfo.h"
 
 ///////////////////////////////////////////////////////////////////////////////////////////////////
 // Config
@@ -79,6 +81,12 @@ struct FIESTextureManager : public FRenderResource
 	bool bHasPendingAdds = false;
 	bool bHasPendingRefreshes = false;
 
+	FIESTextureManager()
+		// Initialize feature level to make sure this resource is reset on a feature level changes
+		: FRenderResource(GMaxRHIFeatureLevel) 
+	{
+	}
+
 	virtual void ReleaseDynamicRHI()
 	{
 		AtlasTexture.SafeRelease();
@@ -182,6 +190,12 @@ static FRDGTextureRef AddIESDebugPass(
 ///////////////////////////////////////////////////////////////////////////////////////////////////
 // Add pass
 
+static bool UseComputeForAtlasUpdate(ERHIFeatureLevel::Type FeatureLevel)
+{
+	// Prefer raster version on mobile, as OpenGL ES does not support writing to R16F images from compute
+	return (FeatureLevel > ERHIFeatureLevel::ES3_1);
+}
+
 class FIESAtlasAddTextureCS : public FGlobalShader
 {
 	DECLARE_GLOBAL_SHADER(FIESAtlasAddTextureCS);
@@ -198,18 +212,21 @@ class FIESAtlasAddTextureCS : public FGlobalShader
 	END_SHADER_PARAMETER_STRUCT()
 
 public:
-	static bool ShouldCompilePermutation(const FGlobalShaderPermutationParameters& Parameters) { return true; }
+	static bool ShouldCompilePermutation(const FGlobalShaderPermutationParameters& Parameters) 
+	{ 
+		return UseComputeForAtlasUpdate(GetMaxSupportedFeatureLevel(Parameters.Platform)); 
+	}
 	static void ModifyCompilationEnvironment(const FGlobalShaderPermutationParameters& Parameters, FShaderCompilerEnvironment& OutEnvironment)
 	{
 		FGlobalShader::ModifyCompilationEnvironment(Parameters, OutEnvironment);
-		OutEnvironment.SetDefine(TEXT("SHADER_ADD_TEXTURE"), 1);
+		OutEnvironment.SetDefine(TEXT("SHADER_ADD_TEXTURE_CS"), 1);
 	}
 };
 
 IMPLEMENT_GLOBAL_SHADER(FIESAtlasAddTextureCS, "/Engine/Private/IESAtlas.usf", "MainCS", SF_Compute);
 
-// Insert a texture into the atlas
-static void AddSlotsPass(
+// Insert a texture into the atlas. Compute version
+static void AddSlotsPassCS(
 	FRDGBuilder& GraphBuilder,
 	FGlobalShaderMap* ShaderMap,
 	const TArray<FAtlasSlot>& Slots,
@@ -252,6 +269,80 @@ static void AddSlotsPass(
 	GraphBuilder.UseExternalAccessMode(OutAtlas, ERHIAccess::SRVMask);
 }
 
+class FIESAtlasAddTexturePS : public FGlobalShader
+{
+	DECLARE_GLOBAL_SHADER(FIESAtlasAddTexturePS);
+	SHADER_USE_PARAMETER_STRUCT(FIESAtlasAddTexturePS, FGlobalShader);
+
+	BEGIN_SHADER_PARAMETER_STRUCT(FParameters, )
+		SHADER_PARAMETER_TEXTURE(Texture2D, InTexture)
+		SHADER_PARAMETER_SAMPLER(SamplerState, InSampler)
+		SHADER_PARAMETER(FIntPoint, AtlasResolution)
+		RENDER_TARGET_BINDING_SLOTS()
+	END_SHADER_PARAMETER_STRUCT()
+
+public:
+	static bool ShouldCompilePermutation(const FGlobalShaderPermutationParameters& Parameters) 
+	{ 
+		return !UseComputeForAtlasUpdate(GetMaxSupportedFeatureLevel(Parameters.Platform)); 
+	}
+	static void ModifyCompilationEnvironment(const FGlobalShaderPermutationParameters& Parameters, FShaderCompilerEnvironment& OutEnvironment)
+	{
+		FGlobalShader::ModifyCompilationEnvironment(Parameters, OutEnvironment);
+		OutEnvironment.SetDefine(TEXT("SHADER_ADD_TEXTURE_PS"), 1);
+	}
+};
+
+IMPLEMENT_GLOBAL_SHADER(FIESAtlasAddTexturePS, "/Engine/Private/IESAtlas.usf", "MainPS", SF_Pixel);
+
+// Insert a texture into the atlas. Raster version
+static void AddSlotsPassPS(
+	FRDGBuilder& GraphBuilder,
+	FGlobalShaderMap* ShaderMap,
+	const TArray<FAtlasSlot>& Slots,
+	FRDGTextureRef& OutAtlas)
+{
+	TShaderMapRef<FIESAtlasAddTexturePS> PixelShader(ShaderMap);
+
+	const uint32 PassCount = Slots.Num();
+	for (uint32 PassIt = 0; PassIt < PassCount; ++PassIt)
+	{
+		const FAtlasSlot& Slot = Slots[PassIt];
+		const FIntRect ViewRect(0, 0, OutAtlas->Desc.Extent.X, OutAtlas->Desc.Extent.Y);
+
+		auto* Parameters = GraphBuilder.AllocParameters<FIESAtlasAddTexturePS::FParameters>();
+		Parameters->InTexture = Slot.GetTextureRHI();
+		Parameters->InSampler = TStaticSamplerState<SF_Bilinear, AM_Clamp, AM_Clamp, AM_Clamp>::GetRHI();
+		Parameters->AtlasResolution = OutAtlas->Desc.Extent;
+		Parameters->RenderTargets[0] = FRenderTargetBinding(OutAtlas, ERenderTargetLoadAction::ENoAction, 0, Slot.SliceIndex);
+
+		FPixelShaderUtils::AddFullscreenPass(
+			GraphBuilder,
+			ShaderMap,
+			RDG_EVENT_NAME("IESAtlas::AddTexturePS (Slot: %d)", PassIt),
+			PixelShader,
+			Parameters,
+			ViewRect);
+	}
+}
+
+static void AddSlotsPass(
+	FRDGBuilder& GraphBuilder,
+	ERHIFeatureLevel::Type FeatureLevel,
+	FGlobalShaderMap* ShaderMap,
+	const TArray<FAtlasSlot>& Slots,
+	FRDGTextureRef& OutAtlas)
+{
+	if (UseComputeForAtlasUpdate(FeatureLevel))
+	{
+		AddSlotsPassCS(GraphBuilder, ShaderMap, Slots, OutAtlas);
+	}
+	else
+	{
+		AddSlotsPassPS(GraphBuilder, ShaderMap, Slots, OutAtlas);
+	}
+}
+
 ///////////////////////////////////////////////////////////////////////////////////////////////////
 // API
 
@@ -333,13 +424,15 @@ void RemoveTexture(uint32 InSlotIndex)
 	}
 }
 
-static FRDGTextureRef CreateAtlasTexture(FRDGBuilder& GraphBuilder, const FIntPoint& Resolution, uint32 SliceCount)
+static FRDGTextureRef CreateAtlasTexture(FRDGBuilder& GraphBuilder, ERHIFeatureLevel::Type FeatureLevel, const FIntPoint& Resolution, uint32 SliceCount)
 {
+	ETextureCreateFlags CreateAddFlags = UseComputeForAtlasUpdate(FeatureLevel) ? ETextureCreateFlags::UAV : ETextureCreateFlags::TargetArraySlicesIndependently;
+	
 	return GraphBuilder.CreateTexture(FRDGTextureDesc::Create2DArray(
 		Resolution,
 		PF_R16F,
 		FClearValueBinding::Transparent,
-		ETextureCreateFlags::UAV | ETextureCreateFlags::ShaderResource | ETextureCreateFlags::RenderTargetable,
+		ETextureCreateFlags::ShaderResource | ETextureCreateFlags::RenderTargetable | CreateAddFlags,
 		FMath::Max(1u, SliceCount),
 		1),
 		TEXT("IES.AtlasTexture"),
@@ -464,7 +557,7 @@ void UpdateAtlasTexture(FRDGBuilder& GraphBuilder, const ERHIFeatureLevel::Type
 			bool bNeedExtraction = false;
 			if (bForceRecreate)
 			{
-				AtlasTexture = CreateAtlasTexture(GraphBuilder, CachedAtlasResolution, RequestedSliceCount);
+				AtlasTexture = CreateAtlasTexture(GraphBuilder, FeatureLevel, CachedAtlasResolution, RequestedSliceCount);
 				bNeedExtraction = true;
 			}
 			else
@@ -475,7 +568,7 @@ void UpdateAtlasTexture(FRDGBuilder& GraphBuilder, const ERHIFeatureLevel::Type
 			// 3.2.2 Insert new slots into the atlas texture
 			if (NewSlots.Num() > 0)
 			{
-				AddSlotsPass(GraphBuilder, ShaderMap, NewSlots, AtlasTexture);
+				AddSlotsPass(GraphBuilder, FeatureLevel, ShaderMap, NewSlots, AtlasTexture);
 			}
 
 			if (bNeedExtraction)
@@ -502,7 +595,7 @@ void UpdateAtlasTexture(FRDGBuilder& GraphBuilder, const ERHIFeatureLevel::Type
 
 		FGlobalShaderMap* ShaderMap = GetGlobalShaderMap(FeatureLevel);
 		FRDGTextureRef AtlasTexture = GraphBuilder.RegisterExternalTexture(GIESTextureManager.AtlasTexture);
-		AddSlotsPass(GraphBuilder, ShaderMap, RefreshSlots, AtlasTexture);
+		AddSlotsPass(GraphBuilder, FeatureLevel, ShaderMap, RefreshSlots, AtlasTexture);
 		GIESTextureManager.bHasPendingRefreshes = false;
 	}
 }
