VolumetricFog is buffer allocations are now indepdent of dynamic resolution scaling.
Tested local light shadow, voxelisation and history reprojection.
Tested functional test as much as I could because some are currently broken.

diff --git a/Engine/Shaders/Private/HeightFogCommon.ush b/Engine/Shaders/Private/HeightFogCommon.ush
--- a/Engine/Shaders/Private/HeightFogCommon.ush
+++ b/Engine/Shaders/Private/HeightFogCommon.ush
@@ -274,7 +274,8 @@ float ComputeNormalizedZSliceFromDepth(float SceneDepth)
 float3 ComputeVolumeUVFromNDC(float4 NDCPosition)
 {
 	NDCPosition.xy /= NDCPosition.w;
-	return float3(NDCPosition.xy * float2(.5f, -.5f) + .5f, ComputeNormalizedZSliceFromDepth(NDCPosition.w));
+	float3 VolumeUV = float3(NDCPosition.xy * float2(.5f, -.5f) + .5f, ComputeNormalizedZSliceFromDepth(NDCPosition.w));
+	return min(VolumeUV * float3(View.VolumetricFogScreenToResourceUV.xy, 1.0), float3(View.VolumetricFogUVMax, 1.0));
 }
 float3 ComputeVolumeUV_DEPRECATED(float3 WorldPosition, float4x4 WorldToClip)
 {
@@ -286,3 +287,20 @@ float3 ComputeVolumeUV(FLWCVector3 WorldPosition, FLWCInverseMatrix WorldToClip)
 	float4 NDCPosition = LWCMultiply(MakeLWCVector4(WorldPosition, 1.0f), WorldToClip);
 	return ComputeVolumeUVFromNDC(NDCPosition);
 }
+
+float3 ComputeHistoryVolumeUVFromNDC(float4 NDCPosition)
+{
+	NDCPosition.xy /= NDCPosition.w;
+	float3 VolumeUV = float3(NDCPosition.xy * float2(.5f, -.5f) + .5f, ComputeNormalizedZSliceFromDepth(NDCPosition.w));
+	return min(VolumeUV * float3(View.VolumetricFogViewGridUVToPrevViewRectUV.xy * float3(View.VolumetricFogPrevViewGridRectUVToResourceUV.xy, 1.0), 1.0), float3(View.VolumetricFogPrevUVMax, 1.0));
+}
+float3 ComputeHistoryVolumeUV_DEPRECATED(float3 WorldPosition, float4x4 WorldToClip)
+{
+	float4 NDCPosition = mul(float4(WorldPosition, 1), WorldToClip);
+	return ComputeHistoryVolumeUVFromNDC(NDCPosition);
+}
+float3 ComputeHistoryVolumeUV(FLWCVector3 WorldPosition, FLWCInverseMatrix WorldToClip)
+{
+	float4 NDCPosition = LWCMultiply(MakeLWCVector4(WorldPosition, 1.0f), WorldToClip);
+	return ComputeHistoryVolumeUVFromNDC(NDCPosition);
+}
diff --git a/Engine/Shaders/Private/VolumetricFog.usf b/Engine/Shaders/Private/VolumetricFog.usf
--- a/Engine/Shaders/Private/VolumetricFog.usf
+++ b/Engine/Shaders/Private/VolumetricFog.usf
@@ -55,7 +55,7 @@ float4x4 UnjitteredPrevTranslatedWorldToClip;
 
 float3 ComputeCellTranslatedWorldPosition(uint3 GridCoordinate, float3 CellOffset, out float SceneDepth)
 {
-	float2 VolumeUV = (GridCoordinate.xy + CellOffset.xy) / VolumetricFog.GridSize.xy;
+	float2 VolumeUV = (GridCoordinate.xy + CellOffset.xy) / VolumetricFog.ViewGridSize.xy;
 	float2 VolumeNDC = (VolumeUV * 2 - 1) * float2(1, -1);
 
 	SceneDepth = ComputeDepthFromZSlice(GridCoordinate.z + CellOffset.z);
@@ -145,7 +145,7 @@ void MaterialSetupCS(
 	float3 Scattering = Albedo * Extinction;
 	float Absorption = max(Extinction - Luminance(Scattering), 0.0f);
 
-	if (all((int3)GridCoordinate < VolumetricFog.GridSizeInt))
+	if (all((int3)GridCoordinate < VolumetricFog.ResourceGridSizeInt))
 	{
 		RWVBufferA[GridCoordinate] = float4(Scattering, Absorption);
 #if USE_EMISSIVE
@@ -286,7 +286,7 @@ void InjectShadowedLocalLightPS(
 
 	// Somehow pixels are being rasterized outside of the viewport on a 970 GTX, perhaps due to use of a geometry shader bypassing the viewport scissor.
 	// This triggers the HistoryMissSuperSampleCount path causing significant overhead for shading off-screen pixels.
-	if (all(int3(GridCoordinate) < VolumetricFog.GridSizeInt))
+	if (all(int3(GridCoordinate) < VolumetricFog.ResourceGridSizeInt))
 	{
 		FDeferredLightData LightData = InitDeferredLightFromUniforms();
 
@@ -298,8 +298,7 @@ void InjectShadowedLocalLightPS(
 		uint NumSuperSamples = 1;
 
 #if USE_TEMPORAL_REPROJECTION
-
-		float3 HistoryUV = ComputeVolumeUV_DEPRECATED(ComputeCellTranslatedWorldPosition(GridCoordinate, .5f), UnjitteredPrevTranslatedWorldToClip) * float3(View.VolumetricFogScreenUVToHistoryVolumeUV.xy, 1.0);
+		float3 HistoryUV = ComputeHistoryVolumeUV_DEPRECATED(ComputeCellTranslatedWorldPosition(GridCoordinate, .5f), UnjitteredPrevTranslatedWorldToClip);
 		float HistoryAlpha = HistoryWeight;
 
 		FLATTEN
@@ -535,7 +534,7 @@ void LightScatteringCS(
 	}
 
 #if USE_TEMPORAL_REPROJECTION
-	float3 HistoryUV = ComputeVolumeUV_DEPRECATED(ComputeCellTranslatedWorldPosition(GridCoordinate, .5f), UnjitteredPrevTranslatedWorldToClip) * float3(View.VolumetricFogScreenUVToHistoryVolumeUV.xy, 1.0);
+	float3 HistoryUV = ComputeHistoryVolumeUV_DEPRECATED(ComputeCellTranslatedWorldPosition(GridCoordinate, .5f), UnjitteredPrevTranslatedWorldToClip);
 	float HistoryAlpha = HistoryWeight;
 
 	FLATTEN
@@ -546,7 +545,7 @@ void LightScatteringCS(
 
 	// Supersample if the history was outside the camera frustum
 	// The compute shader is dispatched with extra threads, make sure those don't supersample
-	NumSuperSamples = HistoryAlpha < .001f && all(int3(GridCoordinate) < VolumetricFog.GridSizeInt) ? HISTORY_MISS_SUPER_SAMPLE_COUNT : 1;
+	NumSuperSamples = HistoryAlpha < .001f && all(int3(GridCoordinate) < VolumetricFog.ViewGridSizeInt) ? HISTORY_MISS_SUPER_SAMPLE_COUNT : 1;
 #endif
 
 	for (uint SampleIndex = 0; SampleIndex < NumSuperSamples; SampleIndex++)
@@ -737,7 +736,7 @@ void LightScatteringCS(
 	
 #endif
 
-	if (all(int3(GridCoordinate) < VolumetricFog.GridSizeInt))
+	if (all(int3(GridCoordinate) < VolumetricFog.ResourceGridSizeInt))
 	{
 		PreExposedScatteringAndExtinction = MakePositiveFinite(PreExposedScatteringAndExtinction);
 		RWLightScattering[GridCoordinate] = PreExposedScatteringAndExtinction;
@@ -761,7 +760,7 @@ void FinalIntegrationCS(
 	float AccumulatedTransmittance = 1.0f;
 	float3 PreviousSliceTranslatedWorldPosition = ComputeCellTranslatedWorldPosition(uint3(GridCoordinate.xy, 0), float3(0.5f, 0.5f, 0.0f));
 	float AccumulatedDepth = 0.0;
-	for (int LayerIndex = 0; LayerIndex < VolumetricFog.GridSizeInt.z; LayerIndex++)
+	for (int LayerIndex = 0; LayerIndex < VolumetricFog.ViewGridSizeInt.z; LayerIndex++)
 	{
 		uint3 LayerCoordinate = uint3(GridCoordinate.xy, LayerIndex);
 		float4 PreExposedScatteringAndExtinction = LightScattering[LayerCoordinate];
diff --git a/Engine/Shaders/Private/VolumetricFogVoxelization.usf b/Engine/Shaders/Private/VolumetricFogVoxelization.usf
--- a/Engine/Shaders/Private/VolumetricFogVoxelization.usf
+++ b/Engine/Shaders/Private/VolumetricFogVoxelization.usf
@@ -110,7 +110,7 @@ void VoxelizeGS(triangle FVoxelizeVolumePrimitiveVSToGS Inputs[3], inout Triangl
 	ResolvedView = ResolveView();
 #endif
 
-	int GridSizeZ = VolumetricFog.GridSize.z;
+	int GridSizeZ = VolumetricFog.ViewGridSize.z;
 	float4 PrimitiveVolumeBounds = VertexFactoryGetTranslatedPrimitiveVolumeBounds(Inputs[0].FactoryInterpolants);
 	float3 ViewSpacePrimitiveVolumeOrigin = mul(float4(PrimitiveVolumeBounds.xyz, 1), ResolvedView.TranslatedWorldToView).xyz;
 
diff --git a/Engine/Source/Runtime/Engine/Public/SceneView.h b/Engine/Source/Runtime/Engine/Public/SceneView.h
--- a/Engine/Source/Runtime/Engine/Public/SceneView.h
+++ b/Engine/Source/Runtime/Engine/Public/SceneView.h
@@ -840,7 +840,11 @@ enum ETranslucencyVolumeCascade
 	VIEW_UNIFORM_BUFFER_MEMBER(FVector3f, VolumetricFogInvGridSize) \
 	VIEW_UNIFORM_BUFFER_MEMBER(FVector3f, VolumetricFogGridZParams) \
 	VIEW_UNIFORM_BUFFER_MEMBER(FVector2f, VolumetricFogSVPosToVolumeUV) \
-	VIEW_UNIFORM_BUFFER_MEMBER(FVector2f, VolumetricFogScreenUVToHistoryVolumeUV) \
+	VIEW_UNIFORM_BUFFER_MEMBER(FVector2f, VolumetricFogViewGridUVToPrevViewRectUV) \
+	VIEW_UNIFORM_BUFFER_MEMBER(FVector2f, VolumetricFogPrevViewGridRectUVToResourceUV) \
+	VIEW_UNIFORM_BUFFER_MEMBER(FVector2f, VolumetricFogPrevUVMax) \
+	VIEW_UNIFORM_BUFFER_MEMBER(FVector2f, VolumetricFogScreenToResourceUV) \
+	VIEW_UNIFORM_BUFFER_MEMBER(FVector2f, VolumetricFogUVMax) \
 	VIEW_UNIFORM_BUFFER_MEMBER(float, VolumetricFogMaxDistance) \
 	VIEW_UNIFORM_BUFFER_MEMBER(FVector3f, VolumetricLightmapWorldToUVScale) \
 	VIEW_UNIFORM_BUFFER_MEMBER(FVector3f, VolumetricLightmapWorldToUVAdd) \
diff --git a/Engine/Source/Runtime/Renderer/Private/ScenePrivate.h b/Engine/Source/Runtime/Renderer/Private/ScenePrivate.h
--- a/Engine/Source/Runtime/Renderer/Private/ScenePrivate.h
+++ b/Engine/Source/Runtime/Renderer/Private/ScenePrivate.h
@@ -1017,7 +1017,9 @@ public:
 	TUniquePtr<FForwardLightingViewResources> ForwardLightingResources;
 
 	float LightScatteringHistoryPreExposure;
-	FVector2f LightScatteringScreenUVToHistoryVolumeUV;
+	FVector2f PrevLightScatteringViewGridUVToViewRectVolumeUV;
+	FVector2f VolumetricFogPrevViewGridRectUVToResourceUV;
+	FVector2f VolumetricFogPrevUVMax;
 	TRefCountPtr<IPooledRenderTarget> LightScatteringHistory;
 	TRefCountPtr<IPooledRenderTarget> PrevLightScatteringConservativeDepthTexture;
 
diff --git a/Engine/Source/Runtime/Renderer/Private/SceneRendering.cpp b/Engine/Source/Runtime/Renderer/Private/SceneRendering.cpp
--- a/Engine/Source/Runtime/Renderer/Private/SceneRendering.cpp
+++ b/Engine/Source/Runtime/Renderer/Private/SceneRendering.cpp
@@ -1459,18 +1459,28 @@ void FViewInfo::SetupUniformBufferParameters(
 	SetupDefaultGlobalDistanceFieldUniformBufferParameters(ViewUniformShaderParameters);
 
 	SetupVolumetricFogUniformBufferParameters(ViewUniformShaderParameters);
-	ViewUniformShaderParameters.VolumetricFogScreenUVToHistoryVolumeUV = FVector2f::One();
+	ViewUniformShaderParameters.VolumetricFogViewGridUVToPrevViewRectUV = FVector2f::One();
+	ViewUniformShaderParameters.VolumetricFogPrevViewGridRectUVToResourceUV = FVector2f::One();
+	ViewUniformShaderParameters.VolumetricFogPrevUVMax = FVector2f::One();
+	int32 VolumetricFogViewGridPixelSize;
+	int32 VolumetricFogResourceGridPixelSize;
+	const FIntVector VolumetricFogResourceGridSize = GetVolumetricFogResourceGridSize(*this, VolumetricFogViewGridPixelSize);
+	const FIntVector VolumetricFogViewGridSize = GetVolumetricFogViewGridSize(*this, VolumetricFogResourceGridPixelSize);
+	const FVector2f ViewRectSize = FVector2f(ViewRect.Size());
 	if (ViewState)
 	{
-		// Compute LightScatteringScreenUVToVolumeUV, for the current frame resolution and volume texture resolutoni according to grid size.
-		int32 VolumetricFogGridPixelSize;
-		const FIntVector VolumetricFogGridSize = GetVolumetricFogGridSize(ViewRect.Size(), VolumetricFogGridPixelSize);
-		FVector2f LightScatteringScreenUVToVolumeUV =  FVector2f(ViewRect.Size()) / (FVector2f(VolumetricFogGridSize.X, VolumetricFogGridSize.Y) * VolumetricFogGridPixelSize);
+		// Compute LightScatteringViewGridUVToViewRectVolumeUV, for the current frame resolution and volume texture resolution according to grid size.
+		FVector2f LightScatteringViewGridUVToViewRectVolumeUV = ViewRectSize / (FVector2f(VolumetricFogViewGridSize.X, VolumetricFogViewGridSize.Y) * VolumetricFogViewGridPixelSize);
 
 		// Due to dynamic resolution scaling, the previous frame might have had a different screen to volume UV due to padding not being aligned on resolution changes.
 		// This effectively correct history samples to account for the change as a ratio of current volume UV to history volume UV.
-		ViewUniformShaderParameters.VolumetricFogScreenUVToHistoryVolumeUV = ViewState->LightScatteringScreenUVToHistoryVolumeUV / LightScatteringScreenUVToVolumeUV;
+		ViewUniformShaderParameters.VolumetricFogViewGridUVToPrevViewRectUV = ViewState->PrevLightScatteringViewGridUVToViewRectVolumeUV / LightScatteringViewGridUVToViewRectVolumeUV;
+
+		ViewUniformShaderParameters.VolumetricFogPrevViewGridRectUVToResourceUV = ViewState->VolumetricFogPrevViewGridRectUVToResourceUV;
+		ViewUniformShaderParameters.VolumetricFogPrevUVMax = ViewState->VolumetricFogPrevUVMax;
 	}
+	ViewUniformShaderParameters.VolumetricFogScreenToResourceUV = ViewRectSize / (FVector2f(VolumetricFogResourceGridSize.X, VolumetricFogResourceGridSize.Y) * VolumetricFogResourceGridPixelSize);
+	ViewUniformShaderParameters.VolumetricFogUVMax = FVector2f(ViewRectSize.X - 0.51, ViewRectSize.Y - 0.51) / (FVector2f(VolumetricFogResourceGridSize.X, VolumetricFogResourceGridSize.Y) * VolumetricFogResourceGridPixelSize);
 
 	SetupPrecomputedVolumetricLightmapUniformBufferParameters(Scene, Family->EngineShowFlags, ViewUniformShaderParameters);
 
diff --git a/Engine/Source/Runtime/Renderer/Private/SceneRendering.h b/Engine/Source/Runtime/Renderer/Private/SceneRendering.h
--- a/Engine/Source/Runtime/Renderer/Private/SceneRendering.h
+++ b/Engine/Source/Runtime/Renderer/Private/SceneRendering.h
@@ -833,8 +833,10 @@ public:
 #define ENABLE_LIGHT_CULLING_VIEW_SPACE_BUILD_DATA 1
 
 BEGIN_GLOBAL_SHADER_PARAMETER_STRUCT_WITH_CONSTRUCTOR(FVolumetricFogGlobalData,) 
-	SHADER_PARAMETER(FIntVector, GridSizeInt)
-	SHADER_PARAMETER(FVector3f, GridSize)
+	SHADER_PARAMETER(FIntVector,ViewGridSizeInt)
+	SHADER_PARAMETER(FVector3f, ViewGridSize)
+	SHADER_PARAMETER(FIntVector,ResourceGridSizeInt)
+	SHADER_PARAMETER(FVector3f, ResourceGridSize)
 	SHADER_PARAMETER(FVector3f, GridZParams)
 	SHADER_PARAMETER(FVector2f, SVPosToVolumeUV)
 	SHADER_PARAMETER(float, MaxDistance)
diff --git a/Engine/Source/Runtime/Renderer/Private/VolumetricFog.cpp b/Engine/Source/Runtime/Renderer/Private/VolumetricFog.cpp
--- a/Engine/Source/Runtime/Renderer/Private/VolumetricFog.cpp
+++ b/Engine/Source/Runtime/Renderer/Private/VolumetricFog.cpp
@@ -151,6 +151,13 @@ static int32 GetVolumetricFogGridSizeZ()
 	return FMath::Max(1, GVolumetricFogGridSizeZ);
 }
 
+static FIntPoint GetVolumetricFogTextureResourceRes(const FViewInfo& View)
+{
+	// Allocate texture using scene render targets size so we do not reallocate every frame when dynamic resolution is used in order to avoid resources allocation hitches.
+	const FIntPoint BufferSize = View.GetSceneTexturesConfig().Extent;
+	return BufferSize;
+}
+
 IMPLEMENT_GLOBAL_SHADER_PARAMETER_STRUCT(FVolumetricFogGlobalData, "VolumetricFog");
 
 DECLARE_GPU_STAT(VolumetricFog);
@@ -450,7 +457,7 @@ void FSceneRenderer::RenderLocalLightsForVolumetricFog(
 	bool bUseTemporalReprojection,
 	const FVolumetricFogIntegrationParameterData& IntegrationData,
 	const FExponentialHeightFogSceneInfo& FogInfo,
-	FIntVector VolumetricFogGridSize,
+	FIntVector VolumetricFogViewGridSize,
 	FVector GridZParams,
 	const FRDGTextureDesc& VolumeDesc,
 	FRDGTexture*& OutLocalShadowedLightScattering,
@@ -493,7 +500,7 @@ void FSceneRenderer::RenderLocalLightsForVolumetricFog(
 			const FVisibleLightInfo& VisibleLightInfo = VisibleLightInfos[LightSceneInfo->Id];
 
 			const FSphere LightBounds = LightSceneInfo->Proxy->GetBoundingSphere();
-			const FIntPoint VolumeZBounds = CalculateVolumetricFogBoundsForLight(LightBounds, View, VolumetricFogGridSize, GridZParams);
+			const FIntPoint VolumeZBounds = CalculateVolumetricFogBoundsForLight(LightBounds, View, VolumetricFogViewGridSize, GridZParams);
 			if (VolumeZBounds.X < VolumeZBounds.Y)
 			{
 				bool bIsShadowed = LightNeedsSeparateInjectionIntoVolumetricFogForOpaqueShadow(View, LightSceneInfo, VisibleLightInfo);
@@ -573,7 +580,7 @@ void FSceneRenderer::RenderLocalLightsForVolumetricFog(
 					RDG_EVENT_NAME("ShadowedLights"),
 					PassParameters,
 					ERDGPassFlags::Raster,
-					[PassParameters, &View, this, VertexShader, GeometryShader, PixelShader, VolumeZBounds, LightBounds](FRHICommandList& RHICmdList)
+					[PassParameters, &View, this, VertexShader, GeometryShader, PixelShader, VolumeZBounds, LightBounds, VolumetricFogViewGridSize](FRHICommandList& RHICmdList)
 				{
 						FGraphicsPipelineStateInitializer GraphicsPSOInit;
 						RHICmdList.ApplyCachedRenderTargets(GraphicsPSOInit);
@@ -597,7 +604,7 @@ void FSceneRenderer::RenderLocalLightsForVolumetricFog(
 						VSPassParameters.ViewSpaceBoundingSphere = FVector4f(FVector4f(View.ViewMatrices.GetViewMatrix().TransformPosition(LightBounds.Center)), LightBounds.W); // LWC_TODO: precision loss
 						VSPassParameters.ViewToVolumeClip = FMatrix44f(View.ViewMatrices.ComputeProjectionNoAAMatrix());	// LWC_TODO: Precision loss?
 
-						VSPassParameters.ClipRatio = GetVolumetricFogFroxelToScreenSVPosRatio(View.ViewRect.Size());
+						VSPassParameters.ClipRatio = GetVolumetricFogFroxelToScreenSVPosRatio(View);
 
 						VSPassParameters.VolumetricFogParameters = PassParameters->VolumetricFogParameters;
 						SetShaderParameters(RHICmdList, VertexShader, VertexShader.GetVertexShader(), VSPassParameters);
@@ -607,6 +614,9 @@ void FSceneRenderer::RenderLocalLightsForVolumetricFog(
 							SetShaderParametersLegacyGS(RHICmdList, GeometryShader, VolumeZBounds.X);
 						}
 
+						// Set the sub region of the texture according to the current dynamic resolution scale.
+						RHICmdList.SetViewport(0.0f, 0.0f, 0.0f, VolumetricFogViewGridSize.X, VolumetricFogViewGridSize.Y, 1.0f);
+
 						RHICmdList.SetStreamSource(0, GCircleRasterizeVertexBuffer.VertexBufferRHI, 0);
 						const int32 NumInstances = VolumeZBounds.Y - VolumeZBounds.X;
 						const int32 NumTriangles = FCircleRasterizeVertexBuffer::NumVertices - 2;
@@ -830,29 +840,41 @@ FVector GetVolumetricFogGridZParams(float VolumetricFogStartDistance, float Near
 	return FVector(B, O, S);
 }
 
-FIntVector GetVolumetricFogGridSize(FIntPoint ViewRectSize, int32& OutVolumetricFogGridPixelSize)
+static FIntVector GetVolumetricFogGridSize(const FIntPoint& TargetResolution, int32& OutVolumetricFogGridPixelSize)
 {
 	extern int32 GLightGridSizeZ;
 	FIntPoint VolumetricFogGridSizeXY;
 	int32 VolumetricFogGridPixelSize = GetVolumetricFogGridPixelSize();
-	VolumetricFogGridSizeXY = FIntPoint::DivideAndRoundUp(ViewRectSize, VolumetricFogGridPixelSize);
+	VolumetricFogGridSizeXY = FIntPoint::DivideAndRoundUp(TargetResolution, VolumetricFogGridPixelSize);
 	if(VolumetricFogGridSizeXY.X > GMaxVolumeTextureDimensions || VolumetricFogGridSizeXY.Y > GMaxVolumeTextureDimensions) //clamp to max volume texture dimensions. only happens for extreme resolutions (~8x2k)
 	{
-		float PixelSizeX = (float)ViewRectSize.X / GMaxVolumeTextureDimensions;
-		float PixelSizeY = (float)ViewRectSize.Y / GMaxVolumeTextureDimensions;
+		float PixelSizeX = (float)TargetResolution.X / GMaxVolumeTextureDimensions;
+		float PixelSizeY = (float)TargetResolution.Y / GMaxVolumeTextureDimensions;
 		VolumetricFogGridPixelSize = FMath::Max(FMath::CeilToInt(PixelSizeX), FMath::CeilToInt(PixelSizeY));
-		VolumetricFogGridSizeXY = FIntPoint::DivideAndRoundUp(ViewRectSize, VolumetricFogGridPixelSize);
+		VolumetricFogGridSizeXY = FIntPoint::DivideAndRoundUp(TargetResolution, VolumetricFogGridPixelSize);
 	}
 	OutVolumetricFogGridPixelSize = VolumetricFogGridPixelSize;
 	return FIntVector(VolumetricFogGridSizeXY.X, VolumetricFogGridSizeXY.Y, GetVolumetricFogGridSizeZ());
 }
 
-FVector2f GetVolumetricFogFroxelToScreenSVPosRatio(FIntPoint ViewRectSize)
+FIntVector GetVolumetricFogResourceGridSize(const FViewInfo& View, int32& OutVolumetricFogGridPixelSize)
+{
+	return GetVolumetricFogGridSize(GetVolumetricFogTextureResourceRes(View), OutVolumetricFogGridPixelSize);
+}
+
+FIntVector GetVolumetricFogViewGridSize(const FViewInfo& View, int32& OutVolumetricFogGridPixelSize)
+{
+	return GetVolumetricFogGridSize(View.ViewRect.Size(), OutVolumetricFogGridPixelSize);
+}
+
+FVector2f GetVolumetricFogFroxelToScreenSVPosRatio(const FViewInfo& View)
 {
+	const FIntPoint ViewRectSize = View.ViewRect.Size();
+
 	// Calculate how much the Fog froxel volume "overhangs" the actual view frustum to the right and bottom.
 	// This needs to be applied on SVPos because froxel pixel size (see r.VolumetricFog.GridPixelSize) does not align perfectly with view rect.
 	int32 VolumetricFogGridPixelSize;
-	const FIntVector VolumetricFogGridSize = GetVolumetricFogGridSize(ViewRectSize, VolumetricFogGridPixelSize);
+	const FIntVector VolumetricFogGridSize = GetVolumetricFogViewGridSize(View, VolumetricFogGridPixelSize);
 	const FVector2f FogPhysicalSize = FVector2f(VolumetricFogGridSize.X, VolumetricFogGridSize.Y) * VolumetricFogGridPixelSize;
 	const FVector2f ClipRatio = FogPhysicalSize / FVector2f(ViewRectSize);
 	return ClipRatio;
@@ -864,15 +886,18 @@ void SetupVolumetricFogGlobalData(const FViewInfo& View, FVolumetricFogGlobalDat
 	const FExponentialHeightFogSceneInfo& FogInfo = Scene->ExponentialFogs[0];
 
 	int32 VolumetricFogGridPixelSize;
-	const FIntVector VolumetricFogGridSize = GetVolumetricFogGridSize(View.ViewRect.Size(), VolumetricFogGridPixelSize);
+	const FIntVector VolumetricFogViewGridSize = GetVolumetricFogViewGridSize(View, VolumetricFogGridPixelSize);
+	const FIntVector VolumetricFogResourceGridSize = GetVolumetricFogResourceGridSize(View, VolumetricFogGridPixelSize);  
 
-	Parameters.GridSizeInt = VolumetricFogGridSize;
-	Parameters.GridSize = FVector3f(VolumetricFogGridSize);
+	Parameters.ViewGridSizeInt = VolumetricFogViewGridSize;
+	Parameters.ViewGridSize = FVector3f(VolumetricFogViewGridSize);
+	Parameters.ResourceGridSizeInt = VolumetricFogResourceGridSize;
+	Parameters.ResourceGridSize = FVector3f(VolumetricFogResourceGridSize);
 
-	FVector ZParams = GetVolumetricFogGridZParams(FogInfo.VolumetricFogStartDistance, View.NearClippingDistance, FogInfo.VolumetricFogDistance, VolumetricFogGridSize.Z);
+	FVector ZParams = GetVolumetricFogGridZParams(FogInfo.VolumetricFogStartDistance, View.NearClippingDistance, FogInfo.VolumetricFogDistance, VolumetricFogResourceGridSize.Z);
 	Parameters.GridZParams = (FVector3f)ZParams;
 
-	Parameters.SVPosToVolumeUV = FVector2f::UnitVector / (FVector2f(VolumetricFogGridSize.X, VolumetricFogGridSize.Y) * VolumetricFogGridPixelSize);
+	Parameters.SVPosToVolumeUV = FVector2f::UnitVector / (FVector2f(VolumetricFogResourceGridSize.X, VolumetricFogResourceGridSize.Y) * VolumetricFogGridPixelSize);
 	Parameters.FogGridToPixelXY = FIntPoint(VolumetricFogGridPixelSize, VolumetricFogGridPixelSize);
 	Parameters.MaxDistance = FogInfo.VolumetricFogDistance;
 
@@ -893,15 +918,17 @@ void FViewInfo::SetupVolumetricFogUniformBufferParameters(FViewUniformShaderPara
 	{
 		const FExponentialHeightFogSceneInfo& FogInfo = Scene->ExponentialFogs[0];
 
-		int32 VolumetricFogGridPixelSize;
-		const FIntVector VolumetricFogGridSize = GetVolumetricFogGridSize(ViewRect.Size(), VolumetricFogGridPixelSize);
+		int32 VolumetricFogResourceGridPixelSize;
+		int32 VolumetricFogViewGridPixelSize;
+		const FIntVector VolumetricFogResourceGridSize = GetVolumetricFogResourceGridSize(*this, VolumetricFogResourceGridPixelSize);
+		const FIntVector VolumetricFogViewGridSize = GetVolumetricFogViewGridSize(*this, VolumetricFogViewGridPixelSize);
 
-		ViewUniformShaderParameters.VolumetricFogInvGridSize = FVector3f(1.0f / VolumetricFogGridSize.X, 1.0f / VolumetricFogGridSize.Y, 1.0f / VolumetricFogGridSize.Z);
+		ViewUniformShaderParameters.VolumetricFogInvGridSize = FVector3f(1.0f / VolumetricFogResourceGridSize.X, 1.0f / VolumetricFogResourceGridSize.Y, 1.0f / VolumetricFogResourceGridSize.Z);
 
-		const FVector ZParams = GetVolumetricFogGridZParams(FogInfo.VolumetricFogStartDistance, NearClippingDistance, FogInfo.VolumetricFogDistance, VolumetricFogGridSize.Z);
+		const FVector ZParams = GetVolumetricFogGridZParams(FogInfo.VolumetricFogStartDistance, NearClippingDistance, FogInfo.VolumetricFogDistance, VolumetricFogResourceGridSize.Z);
 		ViewUniformShaderParameters.VolumetricFogGridZParams = (FVector3f)ZParams;
 
-		ViewUniformShaderParameters.VolumetricFogSVPosToVolumeUV = FVector2f::UnitVector / (FVector2f(VolumetricFogGridSize.X, VolumetricFogGridSize.Y) * VolumetricFogGridPixelSize);
+		ViewUniformShaderParameters.VolumetricFogSVPosToVolumeUV = FVector2f::UnitVector / (FVector2f(VolumetricFogResourceGridSize.X, VolumetricFogResourceGridSize.Y) * VolumetricFogResourceGridPixelSize);
 		ViewUniformShaderParameters.VolumetricFogMaxDistance = FogInfo.VolumetricFogDistance;
 	}
 	else
@@ -909,7 +936,7 @@ void FViewInfo::SetupVolumetricFogUniformBufferParameters(FViewUniformShaderPara
 		ViewUniformShaderParameters.VolumetricFogInvGridSize = FVector3f::ZeroVector;
 		ViewUniformShaderParameters.VolumetricFogGridZParams = FVector3f::ZeroVector;
 		ViewUniformShaderParameters.VolumetricFogSVPosToVolumeUV = FVector2f::ZeroVector;
-		ViewUniformShaderParameters.VolumetricFogScreenUVToHistoryVolumeUV = FVector2f::ZeroVector;
+		ViewUniformShaderParameters.VolumetricFogViewGridUVToPrevViewRectUV = FVector2f::ZeroVector;
 		ViewUniformShaderParameters.VolumetricFogMaxDistance = 0;
 	}
 }
@@ -928,9 +955,6 @@ void FSceneRenderer::SetupVolumetricFog()
 		for (int32 ViewIndex = 0; ViewIndex < Views.Num(); ViewIndex++)
 		{
 			FViewInfo& View = Views[ViewIndex];
-			
-			int32 VolumetricFogGridPixelSize;
-			const FIntVector VolumetricFogGridSize = GetVolumetricFogGridSize(View.ViewRect.Size(), VolumetricFogGridPixelSize);
 
 			FVolumetricFogGlobalData GlobalData;
 			SetupVolumetricFogGlobalData(View, GlobalData);
@@ -947,7 +971,9 @@ void FSceneRenderer::SetupVolumetricFog()
 			{
 				View.ViewState->LightScatteringHistory = NULL;
 				View.ViewState->LightScatteringHistoryPreExposure = 1.0f;
-				View.ViewState->LightScatteringScreenUVToHistoryVolumeUV = FVector2f::One();
+				View.ViewState->PrevLightScatteringViewGridUVToViewRectVolumeUV = FVector2f::One();
+				View.ViewState->VolumetricFogPrevViewGridRectUVToResourceUV = FVector2f::One();
+				View.ViewState->VolumetricFogPrevUVMax = FVector2f::One();
 			}
 		}
 	}
@@ -973,9 +999,13 @@ void FSceneRenderer::ComputeVolumetricFog(FRDGBuilder& GraphBuilder,
 		FViewInfo& View = Views[ViewIndex];
 		RDG_GPU_MASK_SCOPE(GraphBuilder, View.GPUMask);
 
+		// Allocate texture using scene render targets size so we do not reallocate every frame when dynamic resolution is used in order to avoid resources allocation hitches.
+		const FIntPoint BufferSize = View.GetSceneTexturesConfig().Extent;
+
 		int32 VolumetricFogGridPixelSize;
-		const FIntVector VolumetricFogGridSize = GetVolumetricFogGridSize(View.ViewRect.Size(), VolumetricFogGridPixelSize);
-		const FVector GridZParams = GetVolumetricFogGridZParams(FogInfo.VolumetricFogStartDistance, View.NearClippingDistance, FogInfo.VolumetricFogDistance, VolumetricFogGridSize.Z);
+		const FIntVector VolumetricFogResourceGridSize = GetVolumetricFogResourceGridSize(View, VolumetricFogGridPixelSize);
+		const FIntVector VolumetricFogViewGridSize = GetVolumetricFogViewGridSize(View, VolumetricFogGridPixelSize); 
+		const FVector GridZParams = GetVolumetricFogGridZParams(FogInfo.VolumetricFogStartDistance, View.NearClippingDistance, FogInfo.VolumetricFogDistance, VolumetricFogViewGridSize.Z);
 
 		FVolumetricFogIntegrationParameterData IntegrationData;
 		IntegrationData.FrameJitterOffsetValues.Empty(16);
@@ -1007,7 +1037,7 @@ void FSceneRenderer::ComputeVolumetricFog(FRDGBuilder& GraphBuilder,
 		// To use a depth target format, and depth tests, we will have to render depth from a PS depth output. Keeping it simple for now with all the tests happening in shader.
 		if (GVolumetricFogConservativeDepth > 0)
 		{
-			FIntPoint ConservativeDepthTextureSize = FIntPoint(VolumetricFogGridSize.X, VolumetricFogGridSize.Y);
+			FIntPoint ConservativeDepthTextureSize = FIntPoint(VolumetricFogViewGridSize.X, VolumetricFogViewGridSize.Y);
 			ConservativeDepthTexture = GraphBuilder.CreateTexture(FRDGTextureDesc::Create2D(ConservativeDepthTextureSize, PF_R16F,
 				FClearValueBinding::Black, TexCreate_RenderTargetable | TexCreate_ShaderResource | TexCreate_UAV), TEXT("VolumetricFog.ConservativeDepthTexture"));
 			AddGenerateConservativeDepthBufferPass(View, GraphBuilder, ConservativeDepthTexture, GetVolumetricFogGridPixelSize());
@@ -1054,7 +1084,7 @@ void FSceneRenderer::ComputeVolumetricFog(FRDGBuilder& GraphBuilder,
 			GraphBuilder,
 			View,
 			SceneTextures,
-			VolumetricFogGridSize,
+			VolumetricFogViewGridSize,
 			FogInfo.VolumetricFogDistance,
 			DirectionalLightFunctionTranslatedWorldToShadow,
 			DirectionalLightFunctionTexture,
@@ -1064,7 +1094,7 @@ void FSceneRenderer::ComputeVolumetricFog(FRDGBuilder& GraphBuilder,
 		TRDGUniformBufferRef<FFogUniformParameters> FogUniformBuffer = CreateFogUniformBuffer(GraphBuilder, View);
 
 		FRDGTextureDesc VolumeDesc(FRDGTextureDesc::Create3D(
-			VolumetricFogGridSize,
+			VolumetricFogResourceGridSize,
 			PF_FloatRGBA, 
 			FClearValueBinding::Black, 
 			TexCreate_ShaderResource | TexCreate_RenderTargetable | TexCreate_UAV | TexCreate_ReduceMemoryWithTilingMode | TexCreate_3DTiling));
@@ -1083,8 +1113,8 @@ void FSceneRenderer::ComputeVolumetricFog(FRDGBuilder& GraphBuilder,
 		}
 
 		FRDGTexture* LocalShadowedLightScattering = GraphBuilder.RegisterExternalTexture(GSystemTextures.VolumetricBlackDummy);
-		RenderLocalLightsForVolumetricFog(GraphBuilder, View, bUseTemporalReprojection, IntegrationData, FogInfo, 
-			VolumetricFogGridSize, GridZParams, VolumeDescFastVRAM, LocalShadowedLightScattering, ConservativeDepthTexture);
+		RenderLocalLightsForVolumetricFog(GraphBuilder, View, bUseTemporalReprojection, IntegrationData, FogInfo,
+			VolumetricFogViewGridSize, GridZParams, VolumeDescFastVRAM, LocalShadowedLightScattering, ConservativeDepthTexture);
 
 		{
 			FVolumetricFogMaterialSetupCS::FParameters* PassParameters = GraphBuilder.AllocParameters<FVolumetricFogMaterialSetupCS::FParameters>();
@@ -1108,9 +1138,9 @@ void FSceneRenderer::ComputeVolumetricFog(FRDGBuilder& GraphBuilder,
 				RDG_EVENT_NAME("InitializeVolumeAttributes"),
 				PassParameters,
 				ERDGPassFlags::Compute,
-				[PassParameters, &View, VolumetricFogGridSize, IntegrationData, ComputeShader](FRHICommandList& RHICmdList)
+				[PassParameters, &View, VolumetricFogViewGridSize, IntegrationData, ComputeShader](FRHICommandList& RHICmdList)
 			{
-				const FIntVector NumGroups = FIntVector::DivideAndRoundUp(VolumetricFogGridSize, VolumetricFogGridInjectionGroupSize);
+				const FIntVector NumGroups = FIntVector::DivideAndRoundUp(VolumetricFogViewGridSize, VolumetricFogGridInjectionGroupSize);
 
 				SetComputePipelineState(RHICmdList, ComputeShader.GetComputeShader());
 
@@ -1123,7 +1153,7 @@ void FSceneRenderer::ComputeVolumetricFog(FRDGBuilder& GraphBuilder,
 				GraphBuilder,
 				View,
 				IntegrationData,
-				VolumetricFogGridSize,
+				VolumetricFogViewGridSize,
 				GridZParams,
 				FogInfo.VolumetricFogDistance,
 				bUseEmissive);
@@ -1278,18 +1308,18 @@ void FSceneRenderer::ComputeVolumetricFog(FRDGBuilder& GraphBuilder,
 
 			GraphBuilder.AddPass(
 				RDG_EVENT_NAME("LightScattering %dx%dx%d SS:%d %s %s %s",
-					VolumetricFogGridSize.X,
-					VolumetricFogGridSize.Y,
-					VolumetricFogGridSize.Z,
+					VolumetricFogViewGridSize.X,
+					VolumetricFogViewGridSize.Y,
+					VolumetricFogViewGridSize.Z,
 					SuperSampleCount,
 					bUseDistanceFieldSkyOcclusion ? TEXT("DFAO") : TEXT(""),
 					PassParameters->LightFunctionTexture ? TEXT("LF") : TEXT(""),
 					bUseLumenGI ? TEXT("Lumen") : TEXT("")),
 				PassParameters,
 				ERDGPassFlags::Compute,
-				[PassParameters, ComputeShader, &View, this, VolumetricFogGridSize](FRHICommandList& RHICmdList)
+				[PassParameters, ComputeShader, &View, this, VolumetricFogViewGridSize](FRHICommandList& RHICmdList)
 			{
-				const FIntVector NumGroups = FComputeShaderUtils::GetGroupCount(VolumetricFogGridSize, FVolumetricFogLightScatteringCS::GetGroupSize());
+				const FIntVector NumGroups = FComputeShaderUtils::GetGroupCount(VolumetricFogViewGridSize, FVolumetricFogLightScatteringCS::GetGroupSize());
 
 				SetComputePipelineState(RHICmdList, ComputeShader.GetComputeShader());
 
@@ -1314,9 +1344,9 @@ void FSceneRenderer::ComputeVolumetricFog(FRDGBuilder& GraphBuilder,
 				RDG_EVENT_NAME("FinalIntegration"),
 				PassParameters,
 				ERDGPassFlags::Compute,
-				[PassParameters, &View, VolumetricFogGridSize, IntegrationData, this](FRHICommandList& RHICmdList)
+				[PassParameters, &View, VolumetricFogViewGridSize, IntegrationData, this](FRHICommandList& RHICmdList)
 			{
-				const FIntVector NumGroups = FIntVector::DivideAndRoundUp(VolumetricFogGridSize, VolumetricFogIntegrationGroupSize);
+				const FIntVector NumGroups = FIntVector::DivideAndRoundUp(VolumetricFogViewGridSize, VolumetricFogIntegrationGroupSize);
 
 				auto ComputeShader = View.ShaderMap->GetShader< FVolumetricFogFinalIntegrationCS >();
 				SetComputePipelineState(RHICmdList, ComputeShader.GetComputeShader());
@@ -1331,15 +1361,22 @@ void FSceneRenderer::ComputeVolumetricFog(FRDGBuilder& GraphBuilder,
 
 		if (bUseTemporalReprojection)
 		{
+			const FVector2f ViewRectSize = FVector2f(View.ViewRect.Size());
+
 			GraphBuilder.QueueTextureExtraction(IntegrationData.LightScattering, &View.ViewState->LightScatteringHistory);
 			View.ViewState->LightScatteringHistoryPreExposure = View.CachedViewUniformShaderParameters->PreExposure;
-			View.ViewState->LightScatteringScreenUVToHistoryVolumeUV = FVector2f(View.ViewRect.Size()) / (FVector2f(VolumetricFogGridSize.X, VolumetricFogGridSize.Y) * VolumetricFogGridPixelSize);
+			View.ViewState->PrevLightScatteringViewGridUVToViewRectVolumeUV = ViewRectSize / (FVector2f(VolumetricFogViewGridSize.X, VolumetricFogViewGridSize.Y) * VolumetricFogGridPixelSize);
+
+			View.ViewState->VolumetricFogPrevViewGridRectUVToResourceUV = FVector2f(VolumetricFogViewGridSize.X, VolumetricFogViewGridSize.Y) / FVector2f(VolumetricFogResourceGridSize.X, VolumetricFogResourceGridSize.Y);
+			View.ViewState->VolumetricFogPrevUVMax = FVector2f(ViewRectSize.X - 0.51f, ViewRectSize.Y - 0.51f) / (FVector2f(VolumetricFogResourceGridSize.X, VolumetricFogResourceGridSize.Y) * VolumetricFogGridPixelSize);
 		}
 		else if (View.ViewState)
 		{
 			View.ViewState->LightScatteringHistory = nullptr;
 			View.ViewState->LightScatteringHistoryPreExposure = 1.0f;
-			View.ViewState->LightScatteringScreenUVToHistoryVolumeUV = FVector2f::One();
+			View.ViewState->PrevLightScatteringViewGridUVToViewRectVolumeUV = FVector2f::One();
+			View.ViewState->VolumetricFogPrevViewGridRectUVToResourceUV = FVector2f::One();
+			View.ViewState->VolumetricFogPrevUVMax = FVector2f::One();
 		}
 
 		if (bUseTemporalReprojection && GVolumetricFogConservativeDepth > 0)
diff --git a/Engine/Source/Runtime/Renderer/Private/VolumetricFog.h b/Engine/Source/Runtime/Renderer/Private/VolumetricFog.h
--- a/Engine/Source/Runtime/Renderer/Private/VolumetricFog.h
+++ b/Engine/Source/Runtime/Renderer/Private/VolumetricFog.h
@@ -47,8 +47,12 @@ private:
 	FRDGTextureRef DefaultLightFunctionAtlasItemTexture;
 };
 
-extern FIntVector GetVolumetricFogGridSize(FIntPoint ViewRectSize, int32& OutVolumetricFogGridPixelSize);
-extern FVector2f GetVolumetricFogFroxelToScreenSVPosRatio(FIntPoint ViewRectSize);
+// Grid size for resource allocation to be independent of dynamic resolution
+extern FIntVector GetVolumetricFogResourceGridSize(const FViewInfo& View, int32& OutVolumetricFogGridPixelSize);
+// Grid size for the view rectangle within the allocated resource
+extern FIntVector GetVolumetricFogViewGridSize(const FViewInfo& View, int32& OutVolumetricFogGridPixelSize);
+
+extern FVector2f GetVolumetricFogFroxelToScreenSVPosRatio(const FViewInfo& View);
 
 extern bool DoesPlatformSupportVolumetricFogVoxelization(const FStaticShaderPlatform Platform);
 extern bool ShouldRenderVolumetricFog(const FScene* Scene, const FSceneViewFamily& ViewFamily);
diff --git a/Engine/Source/Runtime/Renderer/Private/VolumetricFogVoxelization.cpp b/Engine/Source/Runtime/Renderer/Private/VolumetricFogVoxelization.cpp
--- a/Engine/Source/Runtime/Renderer/Private/VolumetricFogVoxelization.cpp
+++ b/Engine/Source/Runtime/Renderer/Private/VolumetricFogVoxelization.cpp
@@ -51,7 +51,7 @@ TRDGUniformBufferRef<FVoxelizeVolumePassUniformParameters> CreateVoxelizeVolumeP
 	auto* Parameters = GraphBuilder.AllocParameters<FVoxelizeVolumePassUniformParameters>();
 	SetupSceneTextureUniformParameters(GraphBuilder, &View.GetSceneTextures(), View.FeatureLevel, ESceneTextureSetupMode::None, Parameters->SceneTextures);
 
-	Parameters->ClipRatio = GetVolumetricFogFroxelToScreenSVPosRatio(View.ViewRect.Size());
+	Parameters->ClipRatio = GetVolumetricFogFroxelToScreenSVPosRatio(View);
 
 	Parameters->ViewToVolumeClip = FMatrix44f(View.ViewMatrices.ComputeProjectionNoAAMatrix());		// LWC_TODO: Precision loss?
 	Parameters->ViewToVolumeClip.M[2][0] += Jitter.X;
@@ -542,7 +542,7 @@ bool FVoxelizeVolumeMeshProcessor::Process(
 void VoxelizeVolumePrimitive(FVoxelizeVolumeMeshProcessor& PassMeshProcessor,
 	FRHICommandListImmediate& RHICmdList,
 	const FViewInfo& View,
-	FIntVector VolumetricFogGridSize,
+	FIntVector VolumetricFogViewGridSize,
 	FVector GridZParams,
 	const FPrimitiveSceneProxy* PrimitiveSceneProxy,
 	const FMeshBatch& OriginalMesh)
@@ -578,8 +578,8 @@ void VoxelizeVolumePrimitive(FVoxelizeVolumeMeshProcessor& PassMeshProcessor,
 		int32 NearSlice = ComputeZSliceFromDepth(BoundsCenterDepth - Bounds.SphereRadius, GridZParams);
 		int32 FarSlice = ComputeZSliceFromDepth(BoundsCenterDepth + Bounds.SphereRadius, GridZParams);
 
-		NearSlice = FMath::Clamp(NearSlice, 0, VolumetricFogGridSize.Z - 1);
-		FarSlice = FMath::Clamp(FarSlice, 0, VolumetricFogGridSize.Z - 1);
+		NearSlice = FMath::Clamp(NearSlice, 0, VolumetricFogViewGridSize.Z - 1);
+		FarSlice = FMath::Clamp(FarSlice, 0, VolumetricFogViewGridSize.Z - 1);
 
 		const int32 NumSlices = FarSlice - NearSlice + 1;
 		const int32 NumVoxelizationPasses = FMath::DivideAndRoundUp(NumSlices, GetVoxelizationSlicesPerPass(View.GetShaderPlatform()));
@@ -601,7 +601,7 @@ void FSceneRenderer::VoxelizeFogVolumePrimitives(
 	FRDGBuilder& GraphBuilder,
 	const FViewInfo& View,
 	const FVolumetricFogIntegrationParameterData& IntegrationData,
-	FIntVector VolumetricFogGridSize,
+	FIntVector VolumetricFogViewGridSize,
 	FVector GridZParams,
 	float VolumetricFogDistance,
 	bool bVoxelizeEmissive)
@@ -609,8 +609,8 @@ void FSceneRenderer::VoxelizeFogVolumePrimitives(
 	if (View.VolumetricMeshBatches.Num() > 0 && DoesPlatformSupportVolumetricFogVoxelization(View.GetShaderPlatform()))
 	{
 		const FVector2D Jitter(
-			IntegrationData.FrameJitterOffsetValues[0].X / VolumetricFogGridSize.X,
-			IntegrationData.FrameJitterOffsetValues[0].Y / VolumetricFogGridSize.Y);
+			IntegrationData.FrameJitterOffsetValues[0].X / VolumetricFogViewGridSize.X,
+			IntegrationData.FrameJitterOffsetValues[0].Y / VolumetricFogViewGridSize.Y);
 
 		auto* PassParameters = GraphBuilder.AllocParameters<FVoxelizeVolumePassParameters>();
 		PassParameters->Pass = CreateVoxelizeVolumePassUniformBuffer(GraphBuilder, View, IntegrationData, Jitter, Scene->GetVolumetricCloudSceneInfo());
@@ -627,8 +627,8 @@ void FSceneRenderer::VoxelizeFogVolumePrimitives(
 			// Update the parts of VoxelizeParameters which are dependent on the buffer size and view rect
 			View.SetupViewRectUniformBufferParameters(
 				ViewVoxelizeParameters,
-				FIntPoint(VolumetricFogGridSize.X, VolumetricFogGridSize.Y),
-				FIntRect(0, 0, VolumetricFogGridSize.X, VolumetricFogGridSize.Y),
+				FIntPoint(VolumetricFogViewGridSize.X, VolumetricFogViewGridSize.Y),
+				FIntRect(0, 0, VolumetricFogViewGridSize.X, VolumetricFogViewGridSize.Y),
 				View.ViewMatrices,
 				View.PrevViewInfo.ViewMatrices
 			);
@@ -653,12 +653,11 @@ void FSceneRenderer::VoxelizeFogVolumePrimitives(
 			RDG_EVENT_NAME("VoxelizeVolumePrimitives"),
 			PassParameters,
 			ERDGPassFlags::Raster,
-			[PassParameters, Scene = Scene, &View, VolumetricFogGridSize, IntegrationData, VolumetricFogDistance, GridZParams](FRHICommandListImmediate& RHICmdList)
-		{
-			FMeshPassProcessorRenderState DrawRenderState;
+			[PassParameters, Scene = Scene, &View, VolumetricFogViewGridSize, IntegrationData, VolumetricFogDistance, GridZParams](FRHICommandListImmediate& RHICmdList)
+			{
 
 			DrawDynamicMeshPass(View, RHICmdList,
-				[&View, VolumetricFogDistance, &RHICmdList, &VolumetricFogGridSize, &GridZParams](FDynamicPassMeshDrawListContext* DynamicMeshPassContext)
+				[&View, VolumetricFogDistance, &RHICmdList, &VolumetricFogViewGridSize, &GridZParams](FDynamicPassMeshDrawListContext* DynamicMeshPassContext)
 				{
 					FVoxelizeVolumeMeshProcessor PassMeshProcessor(
 						View.Family->Scene->GetRenderScene(),
@@ -668,6 +667,9 @@ void FSceneRenderer::VoxelizeFogVolumePrimitives(
 
 					const bool bShouldRenderHeterogeneousVolumes = ShouldRenderHeterogeneousVolumesForView(View);
 
+					// Set the sub region of the texture according to the current dynamic resolution scale.
+					RHICmdList.SetViewport(0.0f, 0.0f, 0.0f, VolumetricFogViewGridSize.X, VolumetricFogViewGridSize.Y, 1.0f);
+
 					for (int32 MeshBatchIndex = 0; MeshBatchIndex < View.VolumetricMeshBatches.Num(); ++MeshBatchIndex)
 					{
 						// Skip volumes flagged as rendered with HeterogenousVolumes
@@ -683,7 +685,7 @@ void FSceneRenderer::VoxelizeFogVolumePrimitives(
 
 						if ((View.ViewMatrices.GetViewOrigin() - Bounds.Origin).SizeSquared() < (VolumetricFogDistance + Bounds.SphereRadius) * (VolumetricFogDistance + Bounds.SphereRadius))
 						{
-							VoxelizeVolumePrimitive(PassMeshProcessor, RHICmdList, View, VolumetricFogGridSize, GridZParams, PrimitiveSceneProxy, *Mesh);
+							VoxelizeVolumePrimitive(PassMeshProcessor, RHICmdList, View, VolumetricFogViewGridSize, GridZParams, PrimitiveSceneProxy, *Mesh);
 						}
 					}
 				},
