Fix C++20 compile error. With the new generated comparison functions in C++20, a couple of lines in FCachedLUTSettings::GetCombineLUTParameters were triggering errors relating to ambiguity.

diff --git a/Engine/Source/Runtime/Renderer/Private/PostProcess/PostProcessCombineLUTs.cpp b/Engine/Source/Runtime/Renderer/Private/PostProcess/PostProcessCombineLUTs.cpp
--- a/Engine/Source/Runtime/Renderer/Private/PostProcess/PostProcessCombineLUTs.cpp
+++ b/Engine/Source/Runtime/Renderer/Private/PostProcess/PostProcessCombineLUTs.cpp
@@ -244,8 +244,8 @@ struct FCachedLUTSettings
 		UPDATE_CACHE_SETTINGS(Parameters.ACESTonemapParameters.ACESCoefsHigh_4, TonemapperParams.ACESCoefsHigh_4, bHasChanged);
 		UPDATE_CACHE_SETTINGS(Parameters.ACESTonemapParameters.ACESSceneColorMultiplier, TonemapperParams.ACESSceneColorMultiplier, bHasChanged);
 
-		UPDATE_CACHE_SETTINGS(Parameters.ColorScale, FVector4f(View.ColorScale), bHasChanged);
-		UPDATE_CACHE_SETTINGS(Parameters.OverlayColor, (View.OverlayColor), bHasChanged);
+		UPDATE_CACHE_SETTINGS(Parameters.ColorScale, FVector3f(View.ColorScale), bHasChanged);
+		UPDATE_CACHE_SETTINGS(Parameters.OverlayColor, FVector4f(View.OverlayColor), bHasChanged);
 		UPDATE_CACHE_SETTINGS(Parameters.ColorRemap.MappingPolynomial, GetMappingPolynomial(), bHasChanged);
 
 		// White balance
