Fix build issue in Nanite culling shaders by forcing use of DXC for mobile platforms.

#rnx

diff --git a/Engine/Source/Runtime/Renderer/Private/InstanceCulling/InstanceCullingContext.cpp b/Engine/Source/Runtime/Renderer/Private/InstanceCulling/InstanceCullingContext.cpp
--- a/Engine/Source/Runtime/Renderer/Private/InstanceCulling/InstanceCullingContext.cpp
+++ b/Engine/Source/Runtime/Renderer/Private/InstanceCulling/InstanceCullingContext.cpp
@@ -386,6 +386,12 @@ public:
 		FGlobalShader::ModifyCompilationEnvironment(Parameters, OutEnvironment);
 		FInstanceProcessingGPULoadBalancer::SetShaderDefines(OutEnvironment);
 
+		// Force use of DXC for mobile platforms as there are some remaining platforms that still use HLSLCC due to perf. regressions with DXC
+		if (IsMobilePlatform(Parameters.Platform))
+		{
+			OutEnvironment.CompilerFlags.Add(CFLAG_ForceDXC);
+		}
+
 		OutEnvironment.SetDefine(TEXT("INDIRECT_ARGS_NUM_WORDS"), FInstanceCullingContext::IndirectArgsNumWords);
 		OutEnvironment.SetDefine(TEXT("VF_SUPPORTS_PRIMITIVE_SCENE_DATA"), 1);
 		OutEnvironment.SetDefine(TEXT("USE_GLOBAL_GPU_SCENE_DATA"), 1);
