VFX opaque collision with partial depth buffer (avoid particles collisions with themselves).
Currently that pass does not run on mobile.
Follow up: (1) make FCopyViewDepthCS common with Lumen and (2) make some code common in depthrendering.

diff --git a/Engine/Plugins/FX/Niagara/Shaders/Private/NiagaraDataInterfaceCollisionQuery.ush b/Engine/Plugins/FX/Niagara/Shaders/Private/NiagaraDataInterfaceCollisionQuery.ush
--- a/Engine/Plugins/FX/Niagara/Shaders/Private/NiagaraDataInterfaceCollisionQuery.ush
+++ b/Engine/Plugins/FX/Niagara/Shaders/Private/NiagaraDataInterfaceCollisionQuery.ush
@@ -16,6 +16,15 @@ float NDICollisionQuery_GetCustomSceneDepth(float2 ScreenUV)
 #endif
 }
 
+float NDICollisionQuery_GetScenePartialDepth(float2 ScreenUV)
+{
+#if (FEATURE_LEVEL <= FEATURE_LEVEL_ES3_1)
+	return ConvertFromDeviceZ(Texture2DSampleLevel(MobileSceneTextures.ScenePartialDepthTexture, MobileSceneTextures.ScenePartialDepthTextureSampler, ScreenUV, 0).r);
+#else
+	return ConvertFromDeviceZ(Texture2DSampleLevel(SceneTexturesStruct.ScenePartialDepthTexture, SceneTexturesStruct_ScenePartialDepthTextureSampler, ScreenUV, 0).r);
+#endif
+}
+
 void NDICollisionQuery_QuerySceneDepthGPU(in float3 In_SamplePos, in float3 In_LWCTile, out float Out_SceneDepth, out float3 Out_CameraPosWorld, out bool Out_IsInsideView, out float3 Out_WorldPos, out float3 Out_WorldNormal)
 {
 	Out_SceneDepth = -1;
@@ -78,6 +87,70 @@ void NDICollisionQuery_QuerySceneDepthGPU(in float3 In_SamplePos, in float3 In_L
 	}
 }
 
+void NDICollisionQuery_QueryScenePartialDepthGPU(in float3 In_SamplePos, in float3 In_LWCTile, out float Out_SceneDepth, out float3 Out_CameraPosWorld, out bool Out_IsInsideView, out float3 Out_WorldPos, out float3 Out_WorldNormal)
+{
+	Out_SceneDepth = -1;
+	Out_WorldPos = float3(0.0, 0.0, 0.0);
+	Out_WorldNormal = float3(0.0, 0.0, 1.0);
+	Out_IsInsideView = true;
+	FLWCVector3 CameraPos = PrimaryView.WorldCameraOrigin;
+	LWCSetTile(CameraPos, LWCGetTile(CameraPos) - In_LWCTile); // convert to simulation space
+	Out_CameraPosWorld = LWCToFloat(CameraPos);
+
+	FLWCVector3 LwcSamplePos = MakeLWCVector3(In_LWCTile, In_SamplePos);
+	float4 SamplePosition = float4(LWCToFloat(LWCAdd(LwcSamplePos, PrimaryView.PreViewTranslation)), 1); // TODO[mg]: LWCToFloat here?
+	float4 ClipPosition = mul(SamplePosition, View.TranslatedWorldToClip);
+	float2 ScreenPosition = ClipPosition.xy / ClipPosition.w;
+
+	// Check if the sample is inside the view.
+	if (all(abs(ScreenPosition.xy) <= float2(1, 1)))
+	{
+		// Sample the depth buffer to get a world position near the sample position.
+		float2 ScreenUV = ScreenPosition * View.ScreenPositionScaleBias.xy + View.ScreenPositionScaleBias.wz;
+
+		float SceneDepth = NDICollisionQuery_GetScenePartialDepth(ScreenUV);
+
+		FLWCVector3 WorldPosition = WorldPositionFromSceneDepth(ScreenPosition.xy, SceneDepth);
+
+#if NDICOLLISIONQUERY_USE_GBUFFER_NORMAL
+#if STRATA_ENABLED
+		uint2 ScreenPos = uint2(ScreenUV * View.BufferSizeAndInvSize.xy);
+		float3 WorldNormal = StrataPublic_GetWorldNormal(ScreenPos);
+#else
+		float3 WorldNormal = Texture2DSampleLevel(SceneTexturesStruct.GBufferATexture, SceneTexturesStruct_GBufferATextureSampler, ScreenUV, 0).xyz * 2.0 - 1.0;
+#endif
+#else
+		float CollisionDepthBounds = 500.0f;
+		float SceneDepth0 = NDICollisionQuery_GetScenePartialDepth(ScreenUV + float2(View.BufferSizeAndInvSize.z, 0.0));
+		float SceneDepth1 = NDICollisionQuery_GetScenePartialDepth(ScreenUV + float2(0.0, View.BufferSizeAndInvSize.w));
+		// When using the forward shading, the normal of the pixel is approximated by the derivative of the world position
+		// of the pixel. But in on the visible edge this derivative can become very high, making CollisionPlane almost
+		// perpendicular to the view plane. In these case the particle may collide the visible edges of the diferent meshes
+		// in the view frustum. To avoid this, we disable the collision test if one of the derivate is above a threshold.
+		if (max(abs(SceneDepth - SceneDepth0), abs(SceneDepth - SceneDepth1)) > CollisionDepthBounds)
+		{
+			return;
+		}
+
+		FLWCVector3 WorldPosition0 = WorldPositionFromSceneDepth(ScreenPosition.xy + float2(2 * View.ViewSizeAndInvSize.z, 0.0), SceneDepth0);
+		FLWCVector3 WorldPosition1 = WorldPositionFromSceneDepth(ScreenPosition.xy - float2(0.0, 2 * View.ViewSizeAndInvSize.w), SceneDepth1);
+		float3 WorldNormal = normalize(cross(LWCSubtract(WorldPosition0, WorldPosition).Offset, LWCSubtract(WorldPosition1, WorldPosition).Offset));
+#endif
+
+		// convert LWC position back to simulation space
+		LWCSetTile(WorldPosition, LWCGetTile(WorldPosition) - In_LWCTile);
+
+		// Set outputs
+		Out_SceneDepth = SceneDepth;
+		Out_WorldPos = LWCToFloat(WorldPosition);
+		Out_WorldNormal = WorldNormal;
+	}
+	else
+	{
+		Out_IsInsideView = false;
+	}
+}
+
 void NDICollisionQuery_QueryCustomDepthGPU(in float3 In_SamplePos, in float3 In_LWCTile, out float Out_SceneDepth, out float3 Out_CameraPosWorld, out bool Out_IsInsideView, out float3 Out_WorldPos, out float3 Out_WorldNormal)
 {
 	Out_SceneDepth = -1;
diff --git a/Engine/Plugins/FX/Niagara/Shaders/Private/NiagaraDataInterfaceCollisionQueryTemplate.ush b/Engine/Plugins/FX/Niagara/Shaders/Private/NiagaraDataInterfaceCollisionQueryTemplate.ush
--- a/Engine/Plugins/FX/Niagara/Shaders/Private/NiagaraDataInterfaceCollisionQueryTemplate.ush
+++ b/Engine/Plugins/FX/Niagara/Shaders/Private/NiagaraDataInterfaceCollisionQueryTemplate.ush
@@ -7,6 +7,11 @@ void QuerySceneDepthGPU_{ParameterName}(in float3 In_SamplePos, out float Out_Sc
 	NDICollisionQuery_QuerySceneDepthGPU(In_SamplePos, {ParameterName}_SystemLWCTile, Out_SceneDepth, Out_CameraPosWorld, Out_IsInsideView, Out_WorldPos, Out_WorldNormal);
 }
 
+void QueryScenePartialDepthGPU_{ParameterName}(in float3 In_SamplePos, out float Out_SceneDepth, out float3 Out_CameraPosWorld, out bool Out_IsInsideView, out float3 Out_WorldPos, out float3 Out_WorldNormal)
+{
+	NDICollisionQuery_QueryScenePartialDepthGPU(In_SamplePos, {ParameterName}_SystemLWCTile, Out_SceneDepth, Out_CameraPosWorld, Out_IsInsideView, Out_WorldPos, Out_WorldNormal);
+}
+
 void QueryCustomDepthGPU_{ParameterName}(in float3 In_SamplePos, out float Out_SceneDepth, out float3 Out_CameraPosWorld, out bool Out_IsInsideView, out float3 Out_WorldPos, out float3 Out_WorldNormal)
 {
 	NDICollisionQuery_QueryCustomDepthGPU(In_SamplePos, {ParameterName}_SystemLWCTile, Out_SceneDepth, Out_CameraPosWorld, Out_IsInsideView, Out_WorldPos, Out_WorldNormal);
diff --git a/Engine/Plugins/FX/Niagara/Source/Niagara/Private/NiagaraDataInterfaceCollisionQuery.cpp b/Engine/Plugins/FX/Niagara/Source/Niagara/Private/NiagaraDataInterfaceCollisionQuery.cpp
--- a/Engine/Plugins/FX/Niagara/Source/Niagara/Private/NiagaraDataInterfaceCollisionQuery.cpp
+++ b/Engine/Plugins/FX/Niagara/Source/Niagara/Private/NiagaraDataInterfaceCollisionQuery.cpp
@@ -26,6 +26,7 @@ namespace NDICollisionQueryLocal
 	static const TCHAR* TemplateShaderFile = TEXT("/Plugin/FX/Niagara/Private/NiagaraDataInterfaceCollisionQueryTemplate.ush");
 
 	static const FName SceneDepthName(TEXT("QuerySceneDepthGPU"));
+	static const FName ScenePartialDepthName(TEXT("QueryScenePartialDepthGPU"));
 	static const FName CustomDepthName(TEXT("QueryCustomDepthGPU"));
 	static const FName DistanceFieldName(TEXT("QueryMeshDistanceFieldGPU"));
 	static const FName SyncTraceName(TEXT("PerformCollisionQuerySyncCPU"));
@@ -225,9 +226,27 @@ void UNiagaraDataInterfaceCollisionQuery::GetFunctions(TArray<FNiagaraFunctionSi
 		SigDepth.AddOutput(FNiagaraVariable(FNiagaraTypeDefinition::GetBoolDef(), TEXT("IsInsideView")), IsInsideViewDescription);
 		SigDepth.AddOutput(FNiagaraVariable(FNiagaraTypeDefinition::GetPositionDef(), TEXT("SamplePosWorld")), SamplePosWorldDescription);
 		SigDepth.AddOutput(FNiagaraVariable(FNiagaraTypeDefinition::GetVec3Def(), TEXT("SampleWorldNormal")), SampleWorldNormalDescription);
-	
 		OutFunctions.Add(SigDepth);
 
+
+		FNiagaraFunctionSignature SigPartialDepth;
+		SigPartialDepth.Name = NDICollisionQueryLocal::ScenePartialDepthName;
+		SigPartialDepth.bMemberFunction = true;
+		SigPartialDepth.bSupportsCPU = false;
+		SigPartialDepth.MiscUsageBitMask = (uint16)ENiagaraScriptMiscUsageMask::UsesPartialDepthCollisionQuery;
+#if WITH_EDITORONLY_DATA
+		SigPartialDepth.FunctionVersion = FNiagaraCollisionDIFunctionVersion::LatestVersion;
+		SigPartialDepth.Description = LOCTEXT("ScenePartialDepthSignatureDescription", "Projects a given world position to view space and then queries the partial depth buffer (opaque emitter using this function are not in this depth buffer) with that position.");
+#endif
+		SigPartialDepth.Inputs.Add(FNiagaraVariable(FNiagaraTypeDefinition(GetClass()), TEXT("CollisionQuery")));
+		SigPartialDepth.AddInput(FNiagaraVariable(FNiagaraTypeDefinition::GetPositionDef(), TEXT("DepthSamplePosWorld")), DepthSamplePosWorldDescription);
+		SigPartialDepth.AddOutput(FNiagaraVariable(FNiagaraTypeDefinition::GetFloatDef(), TEXT("SceneDepth")), SceneDepthDescription);
+		SigPartialDepth.AddOutput(FNiagaraVariable(FNiagaraTypeDefinition::GetPositionDef(), TEXT("CameraPosWorld")), CameraPosWorldDescription);
+		SigPartialDepth.AddOutput(FNiagaraVariable(FNiagaraTypeDefinition::GetBoolDef(), TEXT("IsInsideView")), IsInsideViewDescription);
+		SigPartialDepth.AddOutput(FNiagaraVariable(FNiagaraTypeDefinition::GetPositionDef(), TEXT("SamplePosWorld")), SamplePosWorldDescription);
+		SigPartialDepth.AddOutput(FNiagaraVariable(FNiagaraTypeDefinition::GetVec3Def(), TEXT("SampleWorldNormal")), SampleWorldNormalDescription);
+		OutFunctions.Add(SigPartialDepth);
+
 		
 		FNiagaraFunctionSignature SigCustomDepth;
 		SigCustomDepth.Name = NDICollisionQueryLocal::CustomDepthName;
@@ -420,6 +439,7 @@ void UNiagaraDataInterfaceCollisionQuery::GetFunctions(TArray<FNiagaraFunctionSi
 bool UNiagaraDataInterfaceCollisionQuery::GetFunctionHLSL(const FNiagaraDataInterfaceGPUParamInfo& ParamInfo, const FNiagaraDataInterfaceGeneratedFunction& FunctionInfo, int FunctionInstanceIndex, FString& OutHLSL)
 {
 	if ( (FunctionInfo.DefinitionName == NDICollisionQueryLocal::SceneDepthName) ||
+		 (FunctionInfo.DefinitionName == NDICollisionQueryLocal::ScenePartialDepthName) ||
 		 (FunctionInfo.DefinitionName == NDICollisionQueryLocal::CustomDepthName) ||
 		 (FunctionInfo.DefinitionName == NDICollisionQueryLocal::DistanceFieldName) ||
 		// DEPRECATE_BEGIN
diff --git a/Engine/Plugins/FX/Niagara/Source/Niagara/Private/NiagaraRenderer.cpp b/Engine/Plugins/FX/Niagara/Source/Niagara/Private/NiagaraRenderer.cpp
--- a/Engine/Plugins/FX/Niagara/Source/Niagara/Private/NiagaraRenderer.cpp
+++ b/Engine/Plugins/FX/Niagara/Source/Niagara/Private/NiagaraRenderer.cpp
@@ -378,6 +378,28 @@ void FNiagaraRenderer::Initialize(const UNiagaraRendererProperties* InProps, con
 	InProps->GetUsedMaterials(Emitter, BaseMaterials_GT);
 	bool bCreateMidsForUsedMaterials = InProps->NeedsMIDsForMaterials();
 
+	bRendersInSecondaryDepthPass = false;
+
+	// Let check if the GPU simulation shader script needed to use the partial depth texture for depth queries.
+	bool bNeedsPartialDepthTexture = false;
+	if (Emitter)
+	{
+		FVersionedNiagaraEmitterData* CachedEmitterData = Emitter->GetCachedEmitterData();
+		if (CachedEmitterData)
+		{
+			UNiagaraScript* NiagaraScriptGPU = CachedEmitterData->GetGPUComputeScript();
+			if (NiagaraScriptGPU)
+			{
+				FNiagaraShaderScript* NiagaraShaderScript = NiagaraScriptGPU->GetRenderThreadScript();
+				if (NiagaraShaderScript)
+				{
+					FNiagaraShaderRef NiagaraShader = NiagaraShaderScript->GetShaderGameThread(0);
+					bNeedsPartialDepthTexture = (NiagaraShader->MiscUsageBitMask & uint16(ENiagaraScriptMiscUsageMask::UsesPartialDepthCollisionQuery)) > 0;
+				}
+			}
+		}
+	}
+
 	uint32 Index = 0;
 	for (UMaterialInterface*& Mat : BaseMaterials_GT)
 	{
@@ -395,7 +417,13 @@ void FNiagaraRenderer::Initialize(const UNiagaraRendererProperties* InProps, con
 
 		Index ++;
 		if (Mat)
+		{
 			BaseMaterialRelevance_GT |= Mat->GetRelevance_Concurrent(FeatureLevel);
+
+			// If partial depth was sampled, we need to make sure this emitter does not end up in it but only in this full depth texture.
+			// For exemple, this is to avoid self collision.
+			bRendersInSecondaryDepthPass |= IsOpaqueOrMaskedBlendMode(Mat->GetBlendMode()) && bNeedsPartialDepthTexture;
+		}
 	}
 }
 
@@ -419,6 +447,7 @@ FPrimitiveViewRelevance FNiagaraRenderer::GetViewRelevance(const FSceneView* Vie
 		Result.bOpaque = View->Family->EngineShowFlags.Bounds;
 		DynamicDataRender->GetMaterialRelevance().SetPrimitiveViewRelevance(Result);
 	}
+	Result.bRenderInSecondStageDepthPass = bRendersInSecondaryDepthPass;
 
 	return Result;
 }
diff --git a/Engine/Plugins/FX/Niagara/Source/Niagara/Public/NiagaraCommon.h b/Engine/Plugins/FX/Niagara/Source/Niagara/Public/NiagaraCommon.h
--- a/Engine/Plugins/FX/Niagara/Source/Niagara/Public/NiagaraCommon.h
+++ b/Engine/Plugins/FX/Niagara/Source/Niagara/Public/NiagaraCommon.h
@@ -410,6 +410,10 @@ struct NIAGARA_API FNiagaraFunctionSignature
 	UPROPERTY(meta = (Bitmask, BitmaskEnum = "/Script/Niagara.ENiagaraScriptUsage"))
 	int32 ModuleUsageBitmask;
 
+	/** Bit mask for misc usage. */
+	UPROPERTY(meta = (Bitmask))
+	uint16 MiscUsageBitMask;
+
 	/** When using simulation stages and bRequiresContext is true this will be the index of the stage that is associated with the function. */
 	UPROPERTY()
 	int32 ContextStageIndex;
@@ -451,6 +455,7 @@ struct NIAGARA_API FNiagaraFunctionSignature
 		, bIsCompileTagGenerator(false)
 		, bHidden(false)
 		, ModuleUsageBitmask(0)
+		, MiscUsageBitMask(0)
 		, ContextStageIndex(INDEX_NONE)
 		, RequiredInputs(INDEX_NONE)
 		, RequiredOutputs(INDEX_NONE)
@@ -471,6 +476,7 @@ struct NIAGARA_API FNiagaraFunctionSignature
 		, bIsCompileTagGenerator(false)
 		, bHidden(false)
 		, ModuleUsageBitmask(0)
+		, MiscUsageBitMask(0)
 		, ContextStageIndex(INDEX_NONE)
 		, RequiredInputs(INDEX_NONE)
 		, RequiredOutputs(INDEX_NONE)
@@ -940,6 +946,13 @@ namespace ENiagaraScriptUsageMask
 	};
 }
 
+/** Misc usage flag for script to transmit informations. Limited to 16 bits, see MiscUsageBitMaskType. */
+enum class ENiagaraScriptMiscUsageMask : uint16
+{
+	/** This Niagara script is sampling partial depth scene texture. So it should be rendered in the secondary depth pass. */
+	UsesPartialDepthCollisionQuery = 1<<0,
+};
+
 /** Defines different usages for a niagara script. */
 UENUM()
 enum class ENiagaraCompileUsageStaticSwitch : uint8
diff --git a/Engine/Plugins/FX/Niagara/Source/Niagara/Public/NiagaraRenderer.h b/Engine/Plugins/FX/Niagara/Source/Niagara/Public/NiagaraRenderer.h
--- a/Engine/Plugins/FX/Niagara/Source/Niagara/Public/NiagaraRenderer.h
+++ b/Engine/Plugins/FX/Niagara/Source/Niagara/Public/NiagaraRenderer.h
@@ -163,6 +163,7 @@ protected:
 	uint32 bLocalSpace : 1;
 	uint32 bHasLights : 1;
 	uint32 bMotionBlurEnabled : 1;
+	uint32 bRendersInSecondaryDepthPass : 1;
 	const ENiagaraSimTarget SimTarget;
 
 	ERHIFeatureLevel::Type FeatureLevel;
diff --git a/Engine/Plugins/FX/Niagara/Source/NiagaraEditor/Private/NiagaraHlslTranslator.cpp b/Engine/Plugins/FX/Niagara/Source/NiagaraEditor/Private/NiagaraHlslTranslator.cpp
--- a/Engine/Plugins/FX/Niagara/Source/NiagaraEditor/Private/NiagaraHlslTranslator.cpp
+++ b/Engine/Plugins/FX/Niagara/Source/NiagaraEditor/Private/NiagaraHlslTranslator.cpp
@@ -2872,6 +2872,7 @@ void FHlslNiagaraTranslator::ConvertCompileInfoToParamInfo(const FNiagaraScriptD
 		FNiagaraDataInterfaceGeneratedFunction& DIFunc = DIInstanceInfo.GeneratedFunctions.AddDefaulted_GetRef();
 		DIFunc.DefinitionName = Sig.Name;
 		DIFunc.InstanceName = GetFunctionSignatureSymbol(Sig);
+		DIFunc.MiscUsageBitMask = Sig.MiscUsageBitMask;
 		DIFunc.Specifiers.Empty(Sig.FunctionSpecifiers.Num());
 		for (const TTuple<FName, FName>& Specifier : Sig.FunctionSpecifiers)
 		{
diff --git a/Engine/Plugins/FX/Niagara/Source/NiagaraShader/Private/NiagaraShader.cpp b/Engine/Plugins/FX/Niagara/Source/NiagaraShader/Private/NiagaraShader.cpp
--- a/Engine/Plugins/FX/Niagara/Source/NiagaraShader/Private/NiagaraShader.cpp
+++ b/Engine/Plugins/FX/Niagara/Source/NiagaraShader/Private/NiagaraShader.cpp
@@ -1052,6 +1052,7 @@ FNiagaraShader::FNiagaraShader(const FNiagaraShaderType::CompiledShaderInitializ
 	);
 
 	// Gather data interface bindings
+	MiscUsageBitMask = 0;
 	DataInterfaceParameters.Empty(InitializerParameters->ScriptParametersMetadata->DataInterfaceParamInfo.Num());
 	for (const FNiagaraDataInterfaceGPUParamInfo& DataInterfaceParamInfo : InitializerParameters->ScriptParametersMetadata->DataInterfaceParamInfo)
 	{
@@ -1064,6 +1065,11 @@ FNiagaraShader::FNiagaraShader(const FNiagaraShaderType::CompiledShaderInitializ
 			continue;
 		}
 
+		for (auto& GeneratedFunction : DataInterfaceParamInfo.GeneratedFunctions)
+		{
+			MiscUsageBitMask |= GeneratedFunction.MiscUsageBitMask;
+		}
+
 		DataInterfaceParamRef.DIType = TIndexedPtr<UNiagaraDataInterfaceBase>(CDODataInterface);
 		DataInterfaceParamRef.Parameters = CDODataInterface->CreateShaderStorage(DataInterfaceParamInfo, Initializer.ParameterMap);
 		checkf(DataInterfaceParamRef.Parameters == nullptr || CDODataInterface->GetShaderStorageType() != nullptr, TEXT("DataInterface(%s) provides shader storage but did not implement GetShaderStorageType"), *GetNameSafe(CDODataInterface->GetClass()));
diff --git a/Engine/Plugins/FX/Niagara/Source/NiagaraShader/Public/NiagaraShader.h b/Engine/Plugins/FX/Niagara/Source/NiagaraShader/Public/NiagaraShader.h
--- a/Engine/Plugins/FX/Niagara/Source/NiagaraShader/Public/NiagaraShader.h
+++ b/Engine/Plugins/FX/Niagara/Source/NiagaraShader/Public/NiagaraShader.h
@@ -238,6 +238,7 @@ public:
 	}
 
 	LAYOUT_FIELD(bool, bNeedsViewUniformBuffer);
+	LAYOUT_FIELD(uint16, MiscUsageBitMask);
 	LAYOUT_ARRAY(FShaderUniformBufferParameter, ExternalConstantBufferParam, 2);
 
 private:
diff --git a/Engine/Plugins/FX/Niagara/Source/NiagaraShader/Public/NiagaraShared.h b/Engine/Plugins/FX/Niagara/Source/NiagaraShader/Public/NiagaraShared.h
--- a/Engine/Plugins/FX/Niagara/Source/NiagaraShader/Public/NiagaraShared.h
+++ b/Engine/Plugins/FX/Niagara/Source/NiagaraShader/Public/NiagaraShared.h
@@ -153,6 +153,8 @@ struct NIAGARASHADER_API FNiagaraDataInterfaceGeneratedFunction
 	using FunctionSpecifier = TTuple<FName, FName>;
 	TArray<FunctionSpecifier> Specifiers;
 
+	uint16 MiscUsageBitMask;
+
 	UPROPERTY()
 	TArray<FNiagaraVariableCommonReference> VariadicInputs;
 	UPROPERTY()
@@ -175,7 +177,7 @@ struct NIAGARASHADER_API FNiagaraDataInterfaceGeneratedFunction
 
 	bool operator==(const FNiagaraDataInterfaceGeneratedFunction& Other) const
 	{
-		if (DefinitionName != Other.DefinitionName || InstanceName != Other.InstanceName)
+		if (DefinitionName != Other.DefinitionName || InstanceName != Other.InstanceName || MiscUsageBitMask != Other.MiscUsageBitMask)
 			return false;
 		if (Specifiers.Num() != Other.Specifiers.Num())
 			return false;
diff --git a/Engine/Shaders/Private/SceneTexturesCommon.ush b/Engine/Shaders/Private/SceneTexturesCommon.ush
--- a/Engine/Shaders/Private/SceneTexturesCommon.ush
+++ b/Engine/Shaders/Private/SceneTexturesCommon.ush
@@ -20,6 +20,7 @@
 
 #define SceneTexturesStruct_SceneColorTextureSampler SceneTexturesStruct.PointClampSampler
 #define SceneTexturesStruct_SceneDepthTextureSampler SceneTexturesStruct.PointClampSampler
+#define SceneTexturesStruct_ScenePartialDepthTextureSampler SceneTexturesStruct.PointClampSampler
 #define SceneTexturesStruct_CustomDepthTextureSampler SceneTexturesStruct.PointClampSampler
 #define SceneTexturesStruct_GBufferATextureSampler SceneTexturesStruct.PointClampSampler
 #define SceneTexturesStruct_GBufferBTextureSampler SceneTexturesStruct.PointClampSampler
diff --git a/Engine/Source/Runtime/Core/Private/UObject/UE5MainStreamObjectVersion.cpp b/Engine/Source/Runtime/Core/Private/UObject/UE5MainStreamObjectVersion.cpp
--- a/Engine/Source/Runtime/Core/Private/UObject/UE5MainStreamObjectVersion.cpp
+++ b/Engine/Source/Runtime/Core/Private/UObject/UE5MainStreamObjectVersion.cpp
@@ -14,7 +14,7 @@ TMap<FGuid, FGuid> FUE5MainStreamObjectVersion::GetSystemGuids()
 	SystemGuids.Add(DevGuids.MATERIALSHADERMAP_DERIVEDDATA_VER, FGuid("65C01D817C9A4EEFAE9E988D41A1F3DD"));
 	SystemGuids.Add(DevGuids.NANITE_DERIVEDDATA_VER, FGuid("C283995F5F1346C3B6F7FB1DDF01ED7F"));
 	SystemGuids.Add(DevGuids.NIAGARASHADERMAP_DERIVEDDATA_VER, FGuid("8A37C45D24F2423CBE5F8F371DE33575"));
-	SystemGuids.Add(DevGuids.Niagara_LatestScriptCompileVersion, FGuid("D2A4314F3211254A9531B32E5E33AD66"));
+	SystemGuids.Add(DevGuids.Niagara_LatestScriptCompileVersion, FGuid("4E3E29C28655E162A05986724B41711B"));
 	SystemGuids.Add(DevGuids.SkeletalMeshDerivedDataVersion, FGuid("4084BE94509A8C3FA7CF410DDF8ABF04"));
 	SystemGuids.Add(DevGuids.STATICMESH_DERIVEDDATA_VER, FGuid("98F7D79A5811013825E75BBCC41ED3E9"));
 	SystemGuids.Add(DevGuids.POSESEARCHDB_DERIVEDDATA_VER, FGuid("4E595C2AC5E947D6BA9ABC874353E5BC"));
diff --git a/Engine/Source/Runtime/Engine/Public/PSOPrecache.h b/Engine/Source/Runtime/Engine/Public/PSOPrecache.h
--- a/Engine/Source/Runtime/Engine/Public/PSOPrecache.h
+++ b/Engine/Source/Runtime/Engine/Public/PSOPrecache.h
@@ -257,7 +257,7 @@ class ENGINE_API FPSOCollectorCreateManager
 {
 public:
 
-	constexpr static uint32 MaxPSOCollectorCount = 33;
+	constexpr static uint32 MaxPSOCollectorCount = 34;
 
 	static PSOCollectorCreateFunction GetCreateFunction(EShadingPath ShadingPath, uint32 Index)
 	{
diff --git a/Engine/Source/Runtime/Engine/Public/PrimitiveViewRelevance.h b/Engine/Source/Runtime/Engine/Public/PrimitiveViewRelevance.h
--- a/Engine/Source/Runtime/Engine/Public/PrimitiveViewRelevance.h
+++ b/Engine/Source/Runtime/Engine/Public/PrimitiveViewRelevance.h
@@ -66,6 +66,8 @@ struct FPrimitiveViewRelevance : public FMaterialRelevance
 	uint32 bUsesLightingChannels : 1;
 	/** Whether the primitive has materials that use volumetric translucent self shadow. */
 	uint32 bTranslucentSelfShadow : 1;
+	/** Whether the primitive should be rendered in the second stage depth only pass. */
+	uint32 bRenderInSecondStageDepthPass : 1;
 
 	/** 
 	 * Whether this primitive view relevance has been initialized this frame.  
diff --git a/Engine/Source/Runtime/Engine/Public/SceneTexturesConfig.h b/Engine/Source/Runtime/Engine/Public/SceneTexturesConfig.h
--- a/Engine/Source/Runtime/Engine/Public/SceneTexturesConfig.h
+++ b/Engine/Source/Runtime/Engine/Public/SceneTexturesConfig.h
@@ -12,9 +12,10 @@ struct FMinimalSceneTextures;
 
 /** A uniform buffer containing common scene textures used by materials or global shaders. */
 BEGIN_GLOBAL_SHADER_PARAMETER_STRUCT(FSceneTextureUniformParameters, ENGINE_API)
-	// Scene Color / Depth
+	// Scene Color / Depth / Partial Depth
 	SHADER_PARAMETER_RDG_TEXTURE(Texture2D, SceneColorTexture)
 	SHADER_PARAMETER_RDG_TEXTURE(Texture2D, SceneDepthTexture)
+	SHADER_PARAMETER_RDG_TEXTURE(Texture2D, ScenePartialDepthTexture)
 
 	// GBuffer
 	SHADER_PARAMETER_RDG_TEXTURE(Texture2D, GBufferATexture)
@@ -41,6 +42,8 @@ BEGIN_GLOBAL_SHADER_PARAMETER_STRUCT(FMobileSceneTextureUniformParameters, ENGIN
 	SHADER_PARAMETER_SAMPLER(SamplerState, SceneColorTextureSampler)
 	SHADER_PARAMETER_RDG_TEXTURE(Texture2D, SceneDepthTexture)
 	SHADER_PARAMETER_SAMPLER(SamplerState, SceneDepthTextureSampler)
+	SHADER_PARAMETER_RDG_TEXTURE(Texture2D, ScenePartialDepthTexture)
+	SHADER_PARAMETER_SAMPLER(SamplerState, ScenePartialDepthTextureSampler)
 	SHADER_PARAMETER_RDG_TEXTURE(Texture2D, CustomDepthTexture)
 	SHADER_PARAMETER_SAMPLER(SamplerState, CustomDepthTextureSampler)
 	SHADER_PARAMETER_RDG_TEXTURE_SRV(Texture2D<uint2>, CustomStencilTexture)
diff --git a/Engine/Source/Runtime/Renderer/Private/DeferredShadingRenderer.cpp b/Engine/Source/Runtime/Renderer/Private/DeferredShadingRenderer.cpp
--- a/Engine/Source/Runtime/Renderer/Private/DeferredShadingRenderer.cpp
+++ b/Engine/Source/Runtime/Renderer/Private/DeferredShadingRenderer.cpp
@@ -3020,6 +3020,7 @@ void FDeferredShadingSceneRenderer::Render(FRDGBuilder& GraphBuilder)
 	const bool bBasePassCanOutputVelocity = FVelocityRendering::BasePassCanOutputVelocity(Platform);
 	const bool bHairStrandsEnable = HairStrandsBookmarkParameters.HasInstances() && Views.Num() > 0 && IsHairStrandsEnabled(EHairStrandsShaderType::Strands, Platform);
 
+	FRDGTextureRef FirstStageDepthBuffer = nullptr;
 	{
 		GraphBuilder.SetCommandListStat(GET_STATID(STAT_CLM_PrePass));
 
@@ -3034,7 +3035,7 @@ void FDeferredShadingSceneRenderer::Render(FRDGBuilder& GraphBuilder)
 		// Draw the scene pre-pass / early z pass, populating the scene depth buffer and HiZ
 		if (bNeedsPrePass)
 		{
-			RenderPrePass(GraphBuilder, SceneTextures.Depth.Target, InstanceCullingManager);
+			RenderPrePass(GraphBuilder, SceneTextures.Depth.Target, InstanceCullingManager, &FirstStageDepthBuffer);
 		}
 		else
 		{
@@ -3288,6 +3289,15 @@ void FDeferredShadingSceneRenderer::Render(FRDGBuilder& GraphBuilder)
 		}
 	}
 
+	if (FirstStageDepthBuffer)
+	{
+		SceneTextures.PartialDepth = FirstStageDepthBuffer;
+		AddResolveSceneDepthPass(GraphBuilder, Views, SceneTextures.PartialDepth);
+	}
+	else
+	{
+		SceneTextures.PartialDepth = SystemTextures.DepthDummy;
+	}
 	SceneTextures.SetupMode = ESceneTextureSetupMode::SceneDepth;
 	SceneTextures.UniformBuffer = CreateSceneTextureUniformBuffer(GraphBuilder, &SceneTextures, FeatureLevel, SceneTextures.SetupMode);
 
diff --git a/Engine/Source/Runtime/Renderer/Private/DeferredShadingRenderer.h b/Engine/Source/Runtime/Renderer/Private/DeferredShadingRenderer.h
--- a/Engine/Source/Runtime/Renderer/Private/DeferredShadingRenderer.h
+++ b/Engine/Source/Runtime/Renderer/Private/DeferredShadingRenderer.h
@@ -625,7 +625,7 @@ private:
 
 	void CreateIndirectCapsuleShadows();
 
-	void RenderPrePass(FRDGBuilder& GraphBuilder, FRDGTextureRef SceneDepthTexture, FInstanceCullingManager& InstanceCullingManager);
+	void RenderPrePass(FRDGBuilder& GraphBuilder, FRDGTextureRef SceneDepthTexture, FInstanceCullingManager& InstanceCullingManager, FRDGTextureRef* FirstStageDepthBuffer);
 	void RenderPrePassHMD(FRDGBuilder& GraphBuilder, FRDGTextureRef SceneDepthTexture);
 
 	void RenderFog(
diff --git a/Engine/Source/Runtime/Renderer/Private/DepthRendering.cpp b/Engine/Source/Runtime/Renderer/Private/DepthRendering.cpp
--- a/Engine/Source/Runtime/Renderer/Private/DepthRendering.cpp
+++ b/Engine/Source/Runtime/Renderer/Private/DepthRendering.cpp
@@ -114,6 +114,38 @@ BEGIN_SHADER_PARAMETER_STRUCT(FDepthPassParameters, )
 	RENDER_TARGET_BINDING_SLOTS()
 END_SHADER_PARAMETER_STRUCT()
 
+class FCopyViewDepthCS : public FGlobalShader
+{
+	DECLARE_GLOBAL_SHADER(FCopyViewDepthCS)
+	SHADER_USE_PARAMETER_STRUCT(FCopyViewDepthCS, FGlobalShader)
+
+	BEGIN_SHADER_PARAMETER_STRUCT(FParameters, )
+		SHADER_PARAMETER_RDG_TEXTURE_UAV(RWTexture2D<float3>, RWDepth)
+		SHADER_PARAMETER_STRUCT_REF(FViewUniformShaderParameters, View)
+		SHADER_PARAMETER_RDG_TEXTURE(Texture2D, SceneDepthTexture)
+	END_SHADER_PARAMETER_STRUCT()
+
+	using FPermutationDomain = TShaderPermutationDomain<>;
+
+	static bool ShouldCompilePermutation(const FGlobalShaderPermutationParameters& Parameters)
+	{
+		return IsFeatureLevelSupported(Parameters.Platform, ERHIFeatureLevel::SM5);
+	}
+
+	static int32 GetGroupSize()
+	{
+		return 8;
+	}
+
+	static void ModifyCompilationEnvironment(const FGlobalShaderPermutationParameters& Parameters, FShaderCompilerEnvironment& OutEnvironment)
+	{
+		FGlobalShader::ModifyCompilationEnvironment(Parameters, OutEnvironment);
+		OutEnvironment.SetDefine(TEXT("THREADGROUP_SIZE"), GetGroupSize());
+	}
+};
+IMPLEMENT_GLOBAL_SHADER(FCopyViewDepthCS, "/Engine/Private/Lumen/LumenScreenProbeGather.usf", "CopyDepthCS", SF_Compute);
+// TODO make FCopyViewDepthCS common between lumen and other parts within the depth render file.
+
 FDepthPassParameters* GetDepthPassParameters(FRDGBuilder& GraphBuilder, const FViewInfo& View, FRDGTextureRef DepthTexture)
 {
 	auto* PassParameters = GraphBuilder.AllocParameters<FDepthPassParameters>();
@@ -485,7 +517,7 @@ void SetupDepthPassState(FMeshPassProcessorRenderState& DrawRenderState)
 
 extern const TCHAR* GetDepthPassReason(bool bDitheredLODTransitionsUseStencil, EShaderPlatform ShaderPlatform);
 
-void FDeferredShadingSceneRenderer::RenderPrePass(FRDGBuilder& GraphBuilder, FRDGTextureRef SceneDepthTexture, FInstanceCullingManager& InstanceCullingManager)
+void FDeferredShadingSceneRenderer::RenderPrePass(FRDGBuilder& GraphBuilder, FRDGTextureRef SceneDepthTexture, FInstanceCullingManager& InstanceCullingManager, FRDGTextureRef* FirstStageDepthBuffer)
 {
 	RDG_EVENT_SCOPE(GraphBuilder, "PrePass %s %s", GetDepthDrawingModeString(DepthPass.EarlyZPassMode), GetDepthPassReason(DepthPass.bDitheredLODTransitionsUseStencil, ShaderPlatform));
 	RDG_CSV_STAT_EXCLUSIVE_SCOPE(GraphBuilder, RenderPrePass);
@@ -576,6 +608,116 @@ void FDeferredShadingSceneRenderer::RenderPrePass(FRDGBuilder& GraphBuilder, FRD
 				}
 			}
 		}
+
+		// Evaluate if any second stage depth buffer processing is required
+		bool bUsesSecondStageDepthPass = false;
+		for (int32 ViewIndex = 0; ViewIndex < Views.Num(); ++ViewIndex)
+		{
+			FViewInfo& View = Views[ViewIndex];
+			bUsesSecondStageDepthPass |= View.bUsesSecondStageDepthPass;
+		}
+
+		// Copy depth buffer
+		if(bUsesSecondStageDepthPass)
+		{
+			// We cannot use AddCopyTexturePass which do not support depth htile copy on some platforms. So we do a custom copy using compute through UAV.
+			// On some platforms, this will however trigger a HTILE decompress when transitioning the depth surface as SRV (because the depth format cannot be read as is ).
+			// TODO: A proper solution would be a proper implementation of AddCopyTexturePass that also maintain the source texture HTile in place.
+			FRDGTextureDesc FirstStageDepthBufferDesc = FRDGTextureDesc::Create2D(SceneDepthTexture->Desc.Extent, PF_R32_FLOAT, FClearValueBinding::Black, TexCreate_ShaderResource | TexCreate_UAV);
+			*FirstStageDepthBuffer = GraphBuilder.CreateTexture(FirstStageDepthBufferDesc, TEXT("FirstStageDepthBuffer"));
+			for (int32 ViewIndex = 0; ViewIndex < Views.Num(); ++ViewIndex)
+			{
+				FViewInfo& View = Views[ViewIndex];
+				if (View.ViewState && !View.bStatePrevViewInfoIsReadOnly && View.bUsesSecondStageDepthPass)
+				{
+					FCopyViewDepthCS::FPermutationDomain PermutationVector;
+					auto ComputeShader = View.ShaderMap->GetShader<FCopyViewDepthCS>(PermutationVector);
+
+					FCopyViewDepthCS::FParameters* PassParameters = GraphBuilder.AllocParameters<FCopyViewDepthCS::FParameters>();
+					PassParameters->RWDepth = GraphBuilder.CreateUAV(FRDGTextureUAVDesc(*FirstStageDepthBuffer));
+					PassParameters->View = View.ViewUniformBuffer;
+					PassParameters->SceneDepthTexture = SceneDepthTexture;
+
+					FComputeShaderUtils::AddPass(
+						GraphBuilder,
+						RDG_EVENT_NAME("CopyViewDepth"),
+						ComputeShader,
+						PassParameters,
+						FComputeShaderUtils::GetGroupCount(View.ViewRect.Size(), FCopyViewDepthCS::GetGroupSize()));
+				}
+			}
+		}
+
+		if (bParallelDepthPass)
+		{
+			RDG_WAIT_FOR_TASKS_CONDITIONAL(GraphBuilder, IsDepthPassWaitForTasksEnabled());
+
+			for (int32 ViewIndex = 0; ViewIndex < Views.Num(); ++ViewIndex)
+			{
+				FViewInfo& View = Views[ViewIndex];
+				RDG_GPU_MASK_SCOPE(GraphBuilder, View.GPUMask);
+				RDG_EVENT_SCOPE_CONDITIONAL(GraphBuilder, Views.Num() > 1, "View%d", ViewIndex);
+
+				FMeshPassProcessorRenderState DrawRenderState;
+				SetupDepthPassState(DrawRenderState);
+
+				const bool bShouldRenderView = View.ShouldRenderView() && View.bUsesSecondStageDepthPass;
+				if (bShouldRenderView)
+				{
+					View.BeginRenderView();
+
+					FDepthPassParameters* PassParameters = GetDepthPassParameters(GraphBuilder, View, SceneDepthTexture);
+					View.ParallelMeshDrawCommandPasses[EMeshPass::SecondStageDepthPass].BuildRenderingCommands(GraphBuilder, Scene->GPUScene, PassParameters->InstanceCullingDrawParams);
+
+					GraphBuilder.AddPass(
+						RDG_EVENT_NAME("SecondStageDepthPassParallel"),
+						PassParameters,
+						ERDGPassFlags::Raster | ERDGPassFlags::SkipRenderPass,
+						[this, &View, PassParameters](const FRDGPass* InPass, FRHICommandListImmediate& RHICmdList)
+					{
+						FRDGParallelCommandListSet ParallelCommandListSet(InPass, RHICmdList, GET_STATID(STAT_CLP_Prepass), View, FParallelCommandListBindings(PassParameters));
+						ParallelCommandListSet.SetHighPriority();
+
+						View.ParallelMeshDrawCommandPasses[EMeshPass::SecondStageDepthPass].DispatchDraw(&ParallelCommandListSet, RHICmdList, &PassParameters->InstanceCullingDrawParams);
+					});
+
+					RenderPrePassEditorPrimitives(GraphBuilder, View, PassParameters, DrawRenderState, DepthPass.EarlyZPassMode, InstanceCullingManager);
+				}
+			}
+		}
+		else
+		{
+			for (int32 ViewIndex = 0; ViewIndex < Views.Num(); ++ViewIndex)
+			{
+				FViewInfo& View = Views[ViewIndex];
+				RDG_GPU_MASK_SCOPE(GraphBuilder, View.GPUMask);
+				RDG_EVENT_SCOPE_CONDITIONAL(GraphBuilder, Views.Num() > 1, "View%d", ViewIndex);
+
+				FMeshPassProcessorRenderState DrawRenderState;
+				SetupDepthPassState(DrawRenderState);
+
+				const bool bShouldRenderView = View.ShouldRenderView() && View.bUsesSecondStageDepthPass;
+				if (bShouldRenderView)
+				{
+					View.BeginRenderView();
+
+					FDepthPassParameters* PassParameters = GetDepthPassParameters(GraphBuilder, View, SceneDepthTexture);
+					View.ParallelMeshDrawCommandPasses[EMeshPass::SecondStageDepthPass].BuildRenderingCommands(GraphBuilder, Scene->GPUScene, PassParameters->InstanceCullingDrawParams);
+
+					GraphBuilder.AddPass(
+						RDG_EVENT_NAME("SecondStageDepthPass"),
+						PassParameters,
+						ERDGPassFlags::Raster,
+						[this, &View, PassParameters](FRHICommandList& RHICmdList)
+						{
+							SetStereoViewport(RHICmdList, View, 1.0f);
+					View.ParallelMeshDrawCommandPasses[EMeshPass::SecondStageDepthPass].DispatchDraw(nullptr, RHICmdList, &PassParameters->InstanceCullingDrawParams);
+						});
+
+					RenderPrePassEditorPrimitives(GraphBuilder, View, PassParameters, DrawRenderState, DepthPass.EarlyZPassMode, InstanceCullingManager);
+				}
+			}
+		}
 	}
 
 	// Dithered transition stencil mask clear, accounting for all active viewports
@@ -1128,13 +1270,15 @@ FDepthPassMeshProcessor::FDepthPassMeshProcessor(
 	const bool InbEarlyZPassMovable,
 	const bool bDitheredLODFadingOutMaskPass,
 	FMeshPassDrawListContext* InDrawListContext,
-	const bool bInShadowProjection)
+	const bool bInShadowProjection,
+	const bool bInSecondStageDepthPass)
 	: FMeshPassProcessor(InMeshPassType, Scene, FeatureLevel, InViewIfDynamicMeshCommand, InDrawListContext)
 	, bRespectUseAsOccluderFlag(InbRespectUseAsOccluderFlag)
 	, EarlyZPassMode(InEarlyZPassMode)
 	, bEarlyZPassMovable(InbEarlyZPassMovable)
 	, bDitheredLODFadingOutMaskPass(bDitheredLODFadingOutMaskPass)
 	, bShadowProjection(bInShadowProjection)
+	, bSecondStageDepthPass(bInSecondStageDepthPass)
 {
 	PassDrawRenderState = InPassDrawRenderState;
 }
@@ -1154,6 +1298,22 @@ FMeshPassProcessor* CreateDepthPassProcessor(ERHIFeatureLevel::Type FeatureLevel
 REGISTER_MESHPASSPROCESSOR_AND_PSOCOLLECTOR(DepthPass, CreateDepthPassProcessor, EShadingPath::Deferred, EMeshPass::DepthPass, EMeshPassFlags::CachedMeshCommands | EMeshPassFlags::MainView);
 FRegisterPassProcessorCreateFunction RegisterMobileDepthPass(&CreateDepthPassProcessor, EShadingPath::Mobile, EMeshPass::DepthPass, EMeshPassFlags::CachedMeshCommands | EMeshPassFlags::MainView);
 
+FMeshPassProcessor* CreateSecondStageDepthPassProcessor(ERHIFeatureLevel::Type FeatureLevel, const FScene* Scene, const FSceneView* InViewIfDynamicMeshCommand, FMeshPassDrawListContext* InDrawListContext)
+{
+	EDepthDrawingMode EarlyZPassMode;
+	bool bEarlyZPassMovable;
+	FScene::GetEarlyZPassMode(FeatureLevel, EarlyZPassMode, bEarlyZPassMovable);
+
+	FMeshPassProcessorRenderState DepthPassState;
+	SetupDepthPassState(DepthPassState);
+
+	return new FDepthPassMeshProcessor(EMeshPass::SecondStageDepthPass, Scene, FeatureLevel, InViewIfDynamicMeshCommand, DepthPassState, true, EarlyZPassMode, bEarlyZPassMovable, false, InDrawListContext, false , true);
+}
+
+REGISTER_MESHPASSPROCESSOR_AND_PSOCOLLECTOR(SecondStageDepthPass, CreateSecondStageDepthPassProcessor, EShadingPath::Deferred, EMeshPass::SecondStageDepthPass, EMeshPassFlags::CachedMeshCommands | EMeshPassFlags::MainView);
+//Secondary depth pass is not implemented on mobile so far (see SceneVisibility.cpp)
+//FRegisterPassProcessorCreateFunction RegisterMobileDepthPass(&CreateSecondStageDepthPassProcessor, EShadingPath::Mobile, EMeshPass::SecondStageDepthPass, EMeshPassFlags::CachedMeshCommands | EMeshPassFlags::MainView);
+
 FMeshPassProcessor* CreateDitheredLODFadingOutMaskPassProcessor(ERHIFeatureLevel::Type FeatureLevel, const FScene* Scene, const FSceneView* InViewIfDynamicMeshCommand, FMeshPassDrawListContext* InDrawListContext)
 {
 	EDepthDrawingMode EarlyZPassMode;
diff --git a/Engine/Source/Runtime/Renderer/Private/DepthRendering.h b/Engine/Source/Runtime/Renderer/Private/DepthRendering.h
--- a/Engine/Source/Runtime/Renderer/Private/DepthRendering.h
+++ b/Engine/Source/Runtime/Renderer/Private/DepthRendering.h
@@ -184,7 +184,8 @@ public:
 		/** Whether this mesh processor is being reused for rendering a pass that marks all fading out pixels on the screen */
 		const bool bDitheredLODFadingOutMaskPass,
 		FMeshPassDrawListContext* InDrawListContext,
-		const bool bShadowProjection = false);
+		const bool bShadowProjection = false,
+		const bool bSecondStageDepthPass = false);
 
 	virtual void AddMeshBatch(const FMeshBatch& RESTRICT MeshBatch, uint64 BatchElementMask, const FPrimitiveSceneProxy* RESTRICT PrimitiveSceneProxy, int32 StaticMeshId = -1) override final;
 	virtual void CollectPSOInitializers(const FSceneTexturesConfig& SceneTexturesConfig, const FMaterial& Material, const FPSOPrecacheVertexFactoryData& VertexFactoryData, const FPSOPrecacheParams& PreCacheParams, TArray<FPSOPrecacheData>& PSOInitializers) override final;
@@ -230,6 +231,7 @@ private:
 	const bool bEarlyZPassMovable;
 	const bool bDitheredLODFadingOutMaskPass;
 	const bool bShadowProjection;
+	const bool bSecondStageDepthPass;
 };
 
 extern void SetupDepthPassState(FMeshPassProcessorRenderState& DrawRenderState);
diff --git a/Engine/Source/Runtime/Renderer/Private/MeshPassProcessor.cpp b/Engine/Source/Runtime/Renderer/Private/MeshPassProcessor.cpp
--- a/Engine/Source/Runtime/Renderer/Private/MeshPassProcessor.cpp
+++ b/Engine/Source/Runtime/Renderer/Private/MeshPassProcessor.cpp
@@ -1988,7 +1988,7 @@ void FCachedPassMeshDrawListContext::FinalizeCommandCommon(
 	{
 		ensureMsgf(MeshDrawCommand.VertexStreams.GetAllocatedSize() == 0, TEXT("Cached Mesh Draw command overflows VertexStreams. VertexStream inline size should be tweaked."));
 
-		if (CurrMeshPass == EMeshPass::BasePass || CurrMeshPass == EMeshPass::DepthPass || CurrMeshPass == EMeshPass::CSMShadowDepth || CurrMeshPass == EMeshPass::VSMShadowDepth)
+		if (CurrMeshPass == EMeshPass::BasePass || CurrMeshPass == EMeshPass::DepthPass || CurrMeshPass == EMeshPass::SecondStageDepthPass || CurrMeshPass == EMeshPass::CSMShadowDepth || CurrMeshPass == EMeshPass::VSMShadowDepth)
 		{
 			TArray<EShaderFrequency, TInlineAllocator<SF_NumFrequencies>> ShaderFrequencies;
 			MeshDrawCommand.ShaderBindings.GetShaderFrequencies(ShaderFrequencies);
diff --git a/Engine/Source/Runtime/Renderer/Private/SceneRendering.cpp b/Engine/Source/Runtime/Renderer/Private/SceneRendering.cpp
--- a/Engine/Source/Runtime/Renderer/Private/SceneRendering.cpp
+++ b/Engine/Source/Runtime/Renderer/Private/SceneRendering.cpp
@@ -859,6 +859,7 @@ void FViewInfo::Init()
 	bHasTranslucencySeparateModulation = 0;
 	bHasStandardTranslucencyModulation = 0;
 	AutoBeforeDOFTranslucencyBoundary = 0.0f;
+	bUsesSecondStageDepthPass = 0;
 
 	NumVisibleStaticMeshElements = 0;
 	PrecomputedVisibilityData = 0;
diff --git a/Engine/Source/Runtime/Renderer/Private/SceneRendering.h b/Engine/Source/Runtime/Renderer/Private/SceneRendering.h
--- a/Engine/Source/Runtime/Renderer/Private/SceneRendering.h
+++ b/Engine/Source/Runtime/Renderer/Private/SceneRendering.h
@@ -1514,6 +1514,12 @@ public:
 	 * If true, that means we need to run the post-dof separate modulation render pass.
 	 */
 	uint32 bHasTranslucencySeparateModulation : 1;
+	/**
+	 * True if the scene has at least one mesh that needs to sample form the first stage depth buffer.
+	 * And as such will need to render in the second stage depth buffer after the first stage depth buffer is copied.
+	 * The first stage depth buffer is usually used for depth buffer collision and projection of Niagara's particles.
+	 */
+	uint32 bUsesSecondStageDepthPass : 1;
 
 	/**
 	 * true if the scene has at least one mesh with a material that needs dual blending AND is applied before DOF. 
diff --git a/Engine/Source/Runtime/Renderer/Private/SceneTextures.cpp b/Engine/Source/Runtime/Renderer/Private/SceneTextures.cpp
--- a/Engine/Source/Runtime/Renderer/Private/SceneTextures.cpp
+++ b/Engine/Source/Runtime/Renderer/Private/SceneTextures.cpp
@@ -717,6 +717,7 @@ void FSceneTextureExtracts::QueueExtractions(FRDGBuilder& GraphBuilder, const FS
 	{
 		SetupMode |= ESceneTextureSetupMode::SceneDepth;
 		ExtractIfProduced(SceneTextures.Depth.Resolve, Depth);
+		ExtractIfProduced(SceneTextures.PartialDepth.Resolve, PartialDepth);
 	}
 
 	if (EnumHasAnyFlags(SceneTextures.Config.Extracts, ESceneTextureExtracts::CustomDepth))
@@ -802,6 +803,7 @@ void SetupSceneTextureUniformParameters(
 	SceneTextureParameters.PointClampSampler = TStaticSamplerState<SF_Point>::GetRHI();
 	SceneTextureParameters.SceneColorTexture = SystemTextures.Black;
 	SceneTextureParameters.SceneDepthTexture = SystemTextures.DepthDummy;
+	SceneTextureParameters.ScenePartialDepthTexture = SystemTextures.DepthDummy;
 	SceneTextureParameters.GBufferATexture = SystemTextures.Black;
 	SceneTextureParameters.GBufferBTexture = SystemTextures.Black;
 	SceneTextureParameters.GBufferCTexture = SystemTextures.Black;
@@ -825,6 +827,7 @@ void SetupSceneTextureUniformParameters(
 		if (EnumHasAnyFlags(SetupMode, ESceneTextureSetupMode::SceneDepth))
 		{
 			SceneTextureParameters.SceneDepthTexture = SceneTextures->Depth.Resolve;
+			SceneTextureParameters.ScenePartialDepthTexture = SceneTextures->PartialDepth.Resolve;
 		}
 
 		if (IsUsingGBuffers(ShaderPlatform))
@@ -933,6 +936,8 @@ void SetupMobileSceneTextureUniformParameters(
 	SceneTextureParameters.SceneColorTextureSampler = TStaticSamplerState<>::GetRHI();
 	SceneTextureParameters.SceneDepthTexture = SystemTextures.DepthDummy;
 	SceneTextureParameters.SceneDepthTextureSampler = TStaticSamplerState<>::GetRHI();
+	SceneTextureParameters.ScenePartialDepthTexture = SystemTextures.DepthDummy;
+	SceneTextureParameters.ScenePartialDepthTextureSampler = TStaticSamplerState<>::GetRHI();
 	// CustomDepthTexture is a color texture on mobile, with DeviceZ values
 	SceneTextureParameters.CustomDepthTexture = SystemTextures.Black;
 	SceneTextureParameters.CustomDepthTextureSampler = TStaticSamplerState<>::GetRHI();
@@ -965,6 +970,13 @@ void SetupMobileSceneTextureUniformParameters(
 			SceneTextureParameters.SceneDepthTexture = SceneTextures->Depth.Resolve;
 		}
 
+		if (EnumHasAnyFlags(SetupMode, EMobileSceneTextureSetupMode::SceneDepth) &&
+			HasBeenProduced(SceneTextures->PartialDepth.Resolve) &&
+			!EnumHasAnyFlags(SceneTextures->PartialDepth.Resolve->Desc.Flags, TexCreate_Memoryless))
+		{
+			SceneTextureParameters.ScenePartialDepthTexture = SceneTextures->PartialDepth.Resolve;
+		}
+
 		if (SceneTextures->Config.bIsUsingGBuffers)
 		{
 			if (HasBeenProduced(SceneTextures->GBufferA))
diff --git a/Engine/Source/Runtime/Renderer/Private/SceneTextures.h b/Engine/Source/Runtime/Renderer/Private/SceneTextures.h
--- a/Engine/Source/Runtime/Renderer/Private/SceneTextures.h
+++ b/Engine/Source/Runtime/Renderer/Private/SceneTextures.h
@@ -41,6 +41,9 @@ struct RENDERER_API FMinimalSceneTextures
 	// Texture containing a stencil view of the resolved (if MSAA) scene depth. 
 	FRDGTextureSRVRef Stencil{};
 
+	// Textures containing primary depth buffer copied before other meshes are rendered in the secondary depth pass.
+	FRDGTextureMSAA PartialDepth{};
+
 	// Textures containing depth / stencil information from the custom depth pass.
 	FCustomDepthTextures CustomDepth{};
 
diff --git a/Engine/Source/Runtime/Renderer/Private/SceneVisibility.cpp b/Engine/Source/Runtime/Renderer/Private/SceneVisibility.cpp
--- a/Engine/Source/Runtime/Renderer/Private/SceneVisibility.cpp
+++ b/Engine/Source/Runtime/Renderer/Private/SceneVisibility.cpp
@@ -2384,6 +2384,7 @@ struct FRelevancePacket : public FSceneRenderingAllocatorObject<FRelevancePacket
 	bool bHasSingleLayerWaterMaterial;
 	bool bHasTranslucencySeparateModulation;
 	bool bHasStandardTranslucencyModulation;
+	bool bUsesSecondStageDepthPass;
 
 	FRelevancePacket(
 		FRHICommandListImmediate& InRHICmdList,
@@ -2424,6 +2425,7 @@ struct FRelevancePacket : public FSceneRenderingAllocatorObject<FRelevancePacket
 		, bHasSingleLayerWaterMaterial(false)
 		, bHasTranslucencySeparateModulation(false)
 		, bHasStandardTranslucencyModulation(false)
+		, bUsesSecondStageDepthPass(false)
 	{
 	}
 
@@ -2443,6 +2445,7 @@ struct FRelevancePacket : public FSceneRenderingAllocatorObject<FRelevancePacket
 		bHasSingleLayerWaterMaterial = 0;
 		bHasTranslucencySeparateModulation = 0;
 		bHasStandardTranslucencyModulation = 0;
+		bUsesSecondStageDepthPass = 0;
 		bUsesGlobalDistanceField = false;
 		bUsesLightingChannels = false;
 		bTranslucentSurfaceLighting = false;
@@ -2579,6 +2582,7 @@ struct FRelevancePacket : public FSceneRenderingAllocatorObject<FRelevancePacket
 			bHasSingleLayerWaterMaterial |= ViewRelevance.bUsesSingleLayerWaterMaterial;
 			bHasStandardTranslucencyModulation |= ViewRelevance.bNormalTranslucency && ViewRelevance.bTranslucencyModulate && View.Family->AllowStandardTranslucencySeparated();
 			bHasTranslucencySeparateModulation |= ViewRelevance.bSeparateTranslucency && ViewRelevance.bTranslucencyModulate;
+			bUsesSecondStageDepthPass |= ViewRelevance.bRenderInSecondStageDepthPass && ShadingPath!=EShadingPath::Mobile;
 
 			if (ViewRelevance.bRenderCustomDepth)
 			{
@@ -2781,7 +2785,14 @@ struct FRelevancePacket : public FSceneRenderingAllocatorObject<FRelevancePacket
 							{
 								if (!(bIsMeshInVelocityPass && bVelocityPassWritesDepth))
 								{
-									DrawCommandPacket.AddCommandsForMesh(PrimitiveIndex, PrimitiveSceneInfo, StaticMeshRelevance, StaticMesh, Scene, bCanCache, EMeshPass::DepthPass);
+									if (ViewRelevance.bRenderInSecondStageDepthPass && ShadingPath != EShadingPath::Mobile)
+									{
+										DrawCommandPacket.AddCommandsForMesh(PrimitiveIndex, PrimitiveSceneInfo, StaticMeshRelevance, StaticMesh, Scene, bCanCache, EMeshPass::SecondStageDepthPass);
+									}
+									else
+									{
+										DrawCommandPacket.AddCommandsForMesh(PrimitiveIndex, PrimitiveSceneInfo, StaticMeshRelevance, StaticMesh, Scene, bCanCache, EMeshPass::DepthPass);
+									}
 								}
 #if RHI_RAYTRACING
 								if (IsRayTracingEnabled())
@@ -2999,6 +3010,7 @@ struct FRelevancePacket : public FSceneRenderingAllocatorObject<FRelevancePacket
 	{
 		FViewInfo& WriteView = const_cast<FViewInfo&>(View);
 		FViewCommands& WriteViewCommands = const_cast<FViewCommands&>(ViewCommands);
+		const EShadingPath ShadingPath = Scene->GetShadingPath();
 		
 		for (int32 Index = 0; Index < NotDrawRelevant.NumPrims; Index++)
 		{
@@ -3046,6 +3058,7 @@ struct FRelevancePacket : public FSceneRenderingAllocatorObject<FRelevancePacket
 		WriteView.bSceneHasSkyMaterial |= bSceneHasSkyMaterial;
 		WriteView.bHasSingleLayerWaterMaterial |= bHasSingleLayerWaterMaterial;
 		WriteView.bHasTranslucencySeparateModulation |= bHasTranslucencySeparateModulation;
+		WriteView.bUsesSecondStageDepthPass |= bUsesSecondStageDepthPass && ShadingPath != EShadingPath::Mobile;
 		WriteView.bHasStandardTranslucencyModulation |= bHasStandardTranslucencyModulation;
 		VisibleDynamicPrimitivesWithSimpleLights.AppendTo(WriteView.VisibleDynamicPrimitivesWithSimpleLights);
 		WriteView.NumVisibleDynamicPrimitives += NumVisibleDynamicPrimitives;
@@ -3278,8 +3291,16 @@ void ComputeDynamicMeshRelevance(EShadingPath ShadingPath, bool bAddLightmapDens
 
 	if (ViewRelevance.bDrawRelevance && (ViewRelevance.bRenderInMainPass || ViewRelevance.bRenderCustomDepth || ViewRelevance.bRenderInDepthPass))
 	{
-		PassMask.Set(EMeshPass::DepthPass);
-		View.NumVisibleDynamicMeshElements[EMeshPass::DepthPass] += NumElements;
+		if (ViewRelevance.bRenderInSecondStageDepthPass && ShadingPath != EShadingPath::Mobile)
+		{
+			PassMask.Set(EMeshPass::SecondStageDepthPass);
+			View.NumVisibleDynamicMeshElements[EMeshPass::SecondStageDepthPass] += NumElements;
+		}
+		else
+		{
+			PassMask.Set(EMeshPass::DepthPass);
+			View.NumVisibleDynamicMeshElements[EMeshPass::DepthPass] += NumElements;
+		}
 
 		if (ViewRelevance.bRenderInMainPass || ViewRelevance.bRenderCustomDepth)
 		{
diff --git a/Engine/Source/Runtime/Renderer/Public/MeshPassProcessor.h b/Engine/Source/Runtime/Renderer/Public/MeshPassProcessor.h
--- a/Engine/Source/Runtime/Renderer/Public/MeshPassProcessor.h
+++ b/Engine/Source/Runtime/Renderer/Public/MeshPassProcessor.h
@@ -32,6 +32,7 @@ namespace EMeshPass
 	enum Type : uint8
 	{
 		DepthPass,
+		SecondStageDepthPass,
 		BasePass,
 		AnisotropyPass,
 		SkyPass,
@@ -80,6 +81,7 @@ inline const TCHAR* GetMeshPassName(EMeshPass::Type MeshPass)
 	switch (MeshPass)
 	{
 	case EMeshPass::DepthPass: return TEXT("DepthPass");
+	case EMeshPass::SecondStageDepthPass: return TEXT("SecondStageDepthPass");
 	case EMeshPass::BasePass: return TEXT("BasePass");
 	case EMeshPass::AnisotropyPass: return TEXT("AnisotropyPass");
 	case EMeshPass::SkyPass: return TEXT("SkyPass");
@@ -117,9 +119,9 @@ inline const TCHAR* GetMeshPassName(EMeshPass::Type MeshPass)
 	}
 
 #if WITH_EDITOR
-	static_assert(EMeshPass::Num == 29 + 4, "Need to update switch(MeshPass) after changing EMeshPass"); // GUID to prevent incorrect auto-resolves, please change when changing the expression: {A6E82589-44B3-4DAD-AC57-8AF6BD50DF43}
+	static_assert(EMeshPass::Num == 30 + 4, "Need to update switch(MeshPass) after changing EMeshPass"); // GUID to prevent incorrect auto-resolves, please change when changing the expression: {A6E82589-44B3-4DAD-AC57-8AF6BD50DF43}
 #else
-	static_assert(EMeshPass::Num == 29, "Need to update switch(MeshPass) after changing EMeshPass"); // GUID to prevent incorrect auto-resolves, please change when changing the expression: {A6E82589-44B3-4DAD-AC57-8AF6BD50DF43}
+	static_assert(EMeshPass::Num == 30, "Need to update switch(MeshPass) after changing EMeshPass"); // GUID to prevent incorrect auto-resolves, please change when changing the expression: {A6E82589-44B3-4DAD-AC57-8AF6BD50DF43}
 #endif
 
 	checkf(0, TEXT("Missing case for EMeshPass %u"), (uint32)MeshPass);
diff --git a/Engine/Source/Runtime/Renderer/Public/SceneRenderTargetParameters.h b/Engine/Source/Runtime/Renderer/Public/SceneRenderTargetParameters.h
--- a/Engine/Source/Runtime/Renderer/Public/SceneRenderTargetParameters.h
+++ b/Engine/Source/Runtime/Renderer/Public/SceneRenderTargetParameters.h
@@ -156,6 +156,9 @@ private:
 	// Contains the resolved scene depth target.
 	TRefCountPtr<IPooledRenderTarget> Depth;
 
+	// Contains the resolved scene depth target.
+	TRefCountPtr<IPooledRenderTarget> PartialDepth;
+
 	// Contains the custom depth targets.
 	TRefCountPtr<IPooledRenderTarget> CustomDepth;
 
