Support pixel shader version of TAA upsampling on mobile.

diff --git a/Engine/Shaders/Private/TemporalAA.usf b/Engine/Shaders/Private/TemporalAA.usf
--- a/Engine/Shaders/Private/TemporalAA.usf
+++ b/Engine/Shaders/Private/TemporalAA.usf
@@ -171,22 +171,20 @@
 		#error Unknown TAA quality
 	#endif
 
-	// Do not use LDS caching for screen percentage > 100% or < 50%.
-	#if TAA_SCREEN_PERCENTAGE_RANGE == 2 || TAA_SCREEN_PERCENTAGE_RANGE == 3
-		#define AA_SAMPLE_CACHE_METHOD (AA_SAMPLE_CACHE_METHOD_VGPR_3X3)
+	#if COMPUTESHADER
+		// Do not use LDS caching for screen percentage > 100% or < 50%.
+		#if TAA_SCREEN_PERCENTAGE_RANGE == 2 || TAA_SCREEN_PERCENTAGE_RANGE == 3
+			#define AA_SAMPLE_CACHE_METHOD (AA_SAMPLE_CACHE_METHOD_VGPR_3X3)
 
 
-	#else
-		#if AA_MOBILE_CONFIG
-			#define AA_SAMPLE_CACHE_METHOD (AA_SAMPLE_CACHE_METHOD_DISABLE) // the shared memory is not as efficient as expected on mobile devices, try not to use it on mobile devices.
 		#else
-			#define AA_SAMPLE_CACHE_METHOD (AA_SAMPLE_CACHE_METHOD_LDS)
+			#if AA_MOBILE_CONFIG
+				#define AA_SAMPLE_CACHE_METHOD (AA_SAMPLE_CACHE_METHOD_DISABLE) // the shared memory is not as efficient as expected on mobile devices, try not to use it on mobile devices.
+			#else
+				#define AA_SAMPLE_CACHE_METHOD (AA_SAMPLE_CACHE_METHOD_LDS)
+			#endif
 		#endif
 	#endif
-
-	#if !COMPUTESHADER
-		#error "Upsampling is only available for CS because of the resceponsive stencil bit."
-	#endif
 	
 #elif TAA_PASS_CONFIG == 5 || TAA_PASS_CONFIG == 6 // Used for diaphragm DOF pre-filtering.
 	#if POST_PROCESS_ALPHA
@@ -355,7 +353,11 @@
 
 // By default, always cache neighbooring samples into VGPR.
 #ifndef AA_SAMPLE_CACHE_METHOD
-	#define AA_SAMPLE_CACHE_METHOD (AA_SAMPLE_CACHE_METHOD_VGPR_3X3)
+	#if COMPUTESHADER
+		#define AA_SAMPLE_CACHE_METHOD (AA_SAMPLE_CACHE_METHOD_VGPR_3X3)
+	#else
+		#define AA_SAMPLE_CACHE_METHOD (AA_SAMPLE_CACHE_METHOD_DISABLE)
+	#endif
 #endif
 
 // By default, enable stocastic quantization of the output.
@@ -2377,11 +2379,7 @@ void MainCS(
 	#endif //TAA_DOWNSAMPLE
 }
 
-#elif PIXELSHADER // Mobile Only.
-
-#if AA_UPSAMPLE
-#error Up sample only available to CS.
-#endif
+#elif PIXELSHADER // Mobile Only, supports Main and MainUpsampling
 
 void MainPS(
 	float4 SvPosition : SV_POSITION,
diff --git a/Engine/Source/Runtime/Renderer/Private/PostProcess/TemporalAA.cpp b/Engine/Source/Runtime/Renderer/Private/PostProcess/TemporalAA.cpp
--- a/Engine/Source/Runtime/Renderer/Private/PostProcess/TemporalAA.cpp
+++ b/Engine/Source/Runtime/Renderer/Private/PostProcess/TemporalAA.cpp
@@ -100,6 +100,7 @@ class FTemporalAA : public FGlobalShader
 public:
 	class FTAAPassConfigDim : SHADER_PERMUTATION_ENUM_CLASS("TAA_PASS_CONFIG", ETAAPassConfig);
 	class FTAAQualityDim : SHADER_PERMUTATION_ENUM_CLASS("TAA_QUALITY", ETAAQuality);
+	class FTAAScreenPercentageDim : SHADER_PERMUTATION_INT("TAA_SCREEN_PERCENTAGE_RANGE", 4);
 
 	BEGIN_SHADER_PARAMETER_STRUCT(FParameters, )
 		SHADER_PARAMETER(FVector4f, ViewportUVToInputBufferUV)
@@ -124,6 +125,13 @@ public:
 		SHADER_PARAMETER(FVector4f, HistoryBufferUVMinMax)
 		SHADER_PARAMETER(FVector4f, ScreenPosToHistoryBufferUV)
 
+		// Temporal upsample specific parameters.
+		SHADER_PARAMETER(FVector4f, InputViewSize)
+		SHADER_PARAMETER(FVector2f, InputViewMin)
+		SHADER_PARAMETER(FVector2f, TemporalJitterPixels)
+		SHADER_PARAMETER(float, ScreenPercentage)
+		SHADER_PARAMETER(float, UpscaleFactor)
+
 		SHADER_PARAMETER_RDG_BUFFER_SRV(StructuredBuffer<float4>, EyeAdaptationBuffer)
 
 		// Inputs
@@ -171,7 +179,8 @@ class FTemporalAAPS : public FTemporalAA
 
 	using FPermutationDomain = TShaderPermutationDomain<
 		FTemporalAA::FTAAPassConfigDim,
-		FTemporalAA::FTAAQualityDim>;
+		FTemporalAA::FTAAQualityDim,
+		FTemporalAA::FTAAScreenPercentageDim>;
 
 	BEGIN_SHADER_PARAMETER_STRUCT(FParameters, )
 		SHADER_PARAMETER_STRUCT_INCLUDE(FTemporalAA::FParameters, Common)
@@ -188,8 +197,32 @@ class FTemporalAAPS : public FTemporalAA
 			return false;
 		}
 
-		// Only Main config is supported on pixel shader.
-		if (PermutationVector.Get<FTemporalAA::FTAAPassConfigDim>() != ETAAPassConfig::Main)
+		// Screen percentage dimension is only for upsampling permutation.
+		if (!IsTAAUpsamplingConfig(PermutationVector.Get<FTemporalAA::FTAAPassConfigDim>()) &&
+			PermutationVector.Get<FTemporalAA::FTAAScreenPercentageDim>() != 0)
+		{
+			return false;
+		}
+
+		if (PermutationVector.Get<FTAAPassConfigDim>() == ETAAPassConfig::MainSuperSampling)
+		{
+			// Super sampling is only available in certain configurations.
+			if (!DoesPlatformSupportTemporalHistoryUpscale(Parameters.Platform))
+			{
+				return false;
+			}
+		}
+
+		// Screen percentage range 3 is only for super sampling.
+		if (PermutationVector.Get<FTemporalAA::FTAAPassConfigDim>() != ETAAPassConfig::MainSuperSampling &&
+			PermutationVector.Get<FTemporalAA::FTAAScreenPercentageDim>() == 3)
+		{
+			return false;
+		}
+
+		// Only Main and MainUpsampling config are supported on pixel shader.
+		if (PermutationVector.Get<FTemporalAA::FTAAPassConfigDim>() != ETAAPassConfig::Main &&
+			PermutationVector.Get<FTemporalAA::FTAAPassConfigDim>() != ETAAPassConfig::MainUpsampling)
 		{
 			return false;
 		}
@@ -208,25 +241,17 @@ class FTemporalAACS : public FTemporalAA
 	DECLARE_GLOBAL_SHADER(FTemporalAACS);
 	SHADER_USE_PARAMETER_STRUCT(FTemporalAACS, FTemporalAA);
 
-	class FTAAScreenPercentageDim : SHADER_PERMUTATION_INT("TAA_SCREEN_PERCENTAGE_RANGE", 4);
 	class FTAADownsampleDim : SHADER_PERMUTATION_BOOL("TAA_DOWNSAMPLE");
 
 	using FPermutationDomain = TShaderPermutationDomain<
 		FTemporalAA::FTAAPassConfigDim,
 		FTemporalAA::FTAAQualityDim,
-		FTAAScreenPercentageDim,
+		FTemporalAA::FTAAScreenPercentageDim,
 		FTAADownsampleDim>;
 
 	BEGIN_SHADER_PARAMETER_STRUCT(FParameters, )
 		SHADER_PARAMETER_STRUCT_INCLUDE(FTemporalAA::FParameters, Common)
 
-		// Temporal upsample specific parameters.
-		SHADER_PARAMETER(FVector4f, InputViewSize)
-		SHADER_PARAMETER(FVector2f, InputViewMin)
-		SHADER_PARAMETER(FVector2f, TemporalJitterPixels)
-		SHADER_PARAMETER(float, ScreenPercentage)
-		SHADER_PARAMETER(float, UpscaleFactor)
-
 		SHADER_PARAMETER_RDG_TEXTURE_UAV_ARRAY(RWTexture2D, OutComputeTex, [FTemporalAAHistory::kRenderTargetCount])
 		SHADER_PARAMETER_RDG_TEXTURE_UAV(RWTexture2D, OutComputeTexDownsampled)
 
@@ -278,7 +303,7 @@ class FTemporalAACS : public FTemporalAA
 
 		// Screen percentage dimension is only for upsampling permutation.
 		if (!IsTAAUpsamplingConfig(PermutationVector.Get<FTemporalAA::FTAAPassConfigDim>()) &&
-			PermutationVector.Get<FTAAScreenPercentageDim>() != 0)
+			PermutationVector.Get<FTemporalAA::FTAAScreenPercentageDim>() != 0)
 		{
 			return false;
 		}
@@ -294,7 +319,7 @@ class FTemporalAACS : public FTemporalAA
 
 		// Screen percentage range 3 is only for super sampling.
 		if (PermutationVector.Get<FTemporalAA::FTAAPassConfigDim>() != ETAAPassConfig::MainSuperSampling &&
-			PermutationVector.Get<FTAAScreenPercentageDim>() == 3)
+			PermutationVector.Get<FTemporalAA::FTAAScreenPercentageDim>() == 3)
 		{
 			return false;
 		}
@@ -694,6 +719,21 @@ FTAAOutputs AddTemporalAAPass(
 			PassParameters->InputSceneMetadataSampler = TStaticSamplerState<SF_Point>::GetRHI();
 		}
 
+		// Temporal upsample specific shader parameters.
+		{
+			// Temporal AA upscale specific params.
+			float InputViewSizeInvScale = Inputs.ResolutionDivisor;
+			float InputViewSizeScale = 1.0f / InputViewSizeInvScale;
+
+			PassParameters->TemporalJitterPixels = InputViewSizeScale * FVector2f(View.TemporalJitterPixels);
+			PassParameters->ScreenPercentage = float(InputViewRect.Width()) / float(OutputViewRect.Width());
+			PassParameters->UpscaleFactor = float(OutputViewRect.Width()) / float(InputViewRect.Width());
+			PassParameters->InputViewMin = InputViewSizeScale * FVector2f(InputViewRect.Min.X, InputViewRect.Min.Y);
+			PassParameters->InputViewSize = FVector4f(
+				InputViewSizeScale * InputViewRect.Width(), InputViewSizeScale * InputViewRect.Height(),
+				InputViewSizeInvScale / InputViewRect.Width(), InputViewSizeInvScale / InputViewRect.Height());
+		}
+
 		PassParameters->OutputViewportSize = FVector4f(
 			PracticableDestRect.Width(), PracticableDestRect.Height(), 1.0f / float(PracticableDestRect.Width()), 1.0f / float(PracticableDestRect.Height()));
 		PassParameters->OutputViewportRect = FVector4f(PracticableDestRect.Min.X, PracticableDestRect.Min.Y, PracticableDestRect.Max.X, PracticableDestRect.Max.Y);
@@ -805,20 +845,20 @@ FTAAOutputs AddTemporalAAPass(
 			if (SrcRect.Width() > DestRect.Width() ||
 				SrcRect.Height() > DestRect.Height())
 			{
-				PermutationVector.Set<FTemporalAACS::FTAAScreenPercentageDim>(2);
+				PermutationVector.Set<FTemporalAA::FTAAScreenPercentageDim>(2);
 			}
 			// If screen percentage < 50% on X and Y axes, then use screen percentage range = 3 shader permutation.
 			else if (SrcRect.Width() * 100 < 50 * DestRect.Width() &&
 				SrcRect.Height() * 100 < 50 * DestRect.Height() &&
 				Inputs.Pass == ETAAPassConfig::MainSuperSampling)
 			{
-				PermutationVector.Set<FTemporalAACS::FTAAScreenPercentageDim>(3);
+				PermutationVector.Set<FTemporalAA::FTAAScreenPercentageDim>(3);
 			}
 			// If screen percentage < 71% on X and Y axes, then use screen percentage range = 1 shader permutation to have smaller LDS caching.
 			else if (SrcRect.Width() * 100 < 71 * DestRect.Width() &&
 				SrcRect.Height() * 100 < 71 * DestRect.Height())
 			{
-				PermutationVector.Set<FTemporalAACS::FTAAScreenPercentageDim>(1);
+				PermutationVector.Set<FTemporalAA::FTAAScreenPercentageDim>(1);
 			}
 		}
 
@@ -828,21 +868,6 @@ FTAAOutputs AddTemporalAAPass(
 
 		SetupTemporalAACommonPassParameters(&PassParameters->Common);
 
-		// Temporal upsample specific shader parameters.
-		{
-			// Temporal AA upscale specific params.
-			float InputViewSizeInvScale = Inputs.ResolutionDivisor;
-			float InputViewSizeScale = 1.0f / InputViewSizeInvScale;
-
-			PassParameters->TemporalJitterPixels = InputViewSizeScale * FVector2f(View.TemporalJitterPixels);
-			PassParameters->ScreenPercentage = float(InputViewRect.Width()) / float(OutputViewRect.Width());
-			PassParameters->UpscaleFactor = float(OutputViewRect.Width()) / float(InputViewRect.Width());
-			PassParameters->InputViewMin = InputViewSizeScale * FVector2f(InputViewRect.Min.X, InputViewRect.Min.Y);
-			PassParameters->InputViewSize = FVector4f(
-				InputViewSizeScale * InputViewRect.Width(), InputViewSizeScale * InputViewRect.Height(),
-				InputViewSizeInvScale / InputViewRect.Width(), InputViewSizeInvScale / InputViewRect.Height());
-		}
-
 		// UAVs
 		{
 			for (int32 i = 0; i < FTemporalAAHistory::kRenderTargetCount; i++)
@@ -889,12 +914,35 @@ FTAAOutputs AddTemporalAAPass(
 	}
 	else
 	{
-		check(IsMobilePlatform(View.GetShaderPlatform()) && !IsTAAUpsamplingConfig(Inputs.Pass));
+		check(IsMobilePlatform(View.GetShaderPlatform()));
 
 		FTemporalAAPS::FPermutationDomain PermutationVector;
 		PermutationVector.Set<FTemporalAA::FTAAPassConfigDim>(Inputs.Pass);
 		PermutationVector.Set<FTemporalAA::FTAAQualityDim>(Inputs.Quality);
 
+		if (IsTAAUpsamplingConfig(Inputs.Pass))
+		{
+			// If screen percentage > 100% on X or Y axes, then use screen percentage range = 2 shader permutation to disable LDS caching.
+			if (SrcRect.Width() > DestRect.Width() ||
+				SrcRect.Height() > DestRect.Height())
+			{
+				PermutationVector.Set<FTemporalAA::FTAAScreenPercentageDim>(2);
+			}
+			// If screen percentage < 50% on X and Y axes, then use screen percentage range = 3 shader permutation.
+			else if (SrcRect.Width() * 100 < 50 * DestRect.Width() &&
+				SrcRect.Height() * 100 < 50 * DestRect.Height() &&
+				Inputs.Pass == ETAAPassConfig::MainSuperSampling)
+			{
+				PermutationVector.Set<FTemporalAA::FTAAScreenPercentageDim>(3);
+			}
+			// If screen percentage < 71% on X and Y axes, then use screen percentage range = 1 shader permutation to have smaller LDS caching.
+			else if (SrcRect.Width() * 100 < 71 * DestRect.Width() &&
+				SrcRect.Height() * 100 < 71 * DestRect.Height())
+			{
+				PermutationVector.Set<FTemporalAA::FTAAScreenPercentageDim>(1);
+			}
+		}
+
 		FTemporalAAPS::FParameters* PassParameters = GraphBuilder.AllocParameters<FTemporalAAPS::FParameters>();
 
 		SetupTemporalAACommonPassParameters(&PassParameters->Common);
@@ -908,9 +956,10 @@ FTAAOutputs AddTemporalAAPass(
 		FPixelShaderUtils::AddFullscreenPass(
 			GraphBuilder,
 			View.ShaderMap,
-			RDG_EVENT_NAME("TAA(%s Quality=%s) %dx%d",
+			RDG_EVENT_NAME("TAA(%s Quality=%s) %dx%d -> %dx%d",
 				PassName,
 				kTAAQualityNames[int32(PermutationVector.Get<FTemporalAA::FTAAQualityDim>())],
+				PracticableSrcRect.Width(), PracticableSrcRect.Height(),
 				PracticableDestRect.Width(), PracticableDestRect.Height()),
 			PixelShader,
 			PassParameters,
diff --git a/Engine/Source/Runtime/Renderer/Private/SceneRendering.cpp b/Engine/Source/Runtime/Renderer/Private/SceneRendering.cpp
--- a/Engine/Source/Runtime/Renderer/Private/SceneRendering.cpp
+++ b/Engine/Source/Runtime/Renderer/Private/SceneRendering.cpp
@@ -2890,13 +2890,11 @@ void FSceneRenderer::PrepareViewRectsForRendering(FRHICommandListImmediate& RHIC
 
 		// Automatic screen percentage fallback.
 		{
-			extern bool DoesTemporalAAUseComputeShader(EShaderPlatform Platform);
 			// Tenmporal upsample is supported only if TAA is turned on.
 			if (View.PrimaryScreenPercentageMethod == EPrimaryScreenPercentageMethod::TemporalUpscale &&
 				(!IsTemporalAccumulationBasedMethod(View.AntiAliasingMethod) ||
 				 ViewFamily.EngineShowFlags.VisualizeBuffer || 
-				 ViewFamily.EngineShowFlags.VisualizeSSS ||
-				 !DoesTemporalAAUseComputeShader(View.GetShaderPlatform())))
+				 ViewFamily.EngineShowFlags.VisualizeSSS))
 			{
 				View.PrimaryScreenPercentageMethod = EPrimaryScreenPercentageMethod::SpatialUpscale;
 			}
