Made primary and secondary depth dispatch into a common lambda RenderDepthPass.

diff --git a/Engine/Source/Runtime/Renderer/Private/DepthRendering.cpp b/Engine/Source/Runtime/Renderer/Private/DepthRendering.cpp
--- a/Engine/Source/Runtime/Renderer/Private/DepthRendering.cpp
+++ b/Engine/Source/Runtime/Renderer/Private/DepthRendering.cpp
@@ -535,9 +535,10 @@ void FDeferredShadingSceneRenderer::RenderPrePass(FRDGBuilder& GraphBuilder, FRD
 		AddDitheredStencilFillPass(GraphBuilder, Views, SceneDepthTexture, DepthPass);
 	}
 
-	// Draw a depth pass to avoid overdraw in the other passes.
-	if (DepthPass.EarlyZPassMode != DDM_None)
+	auto RenderDepthPass = [&](uint8 DepthMeshPass)
 	{
+		check(DepthMeshPass == EMeshPass::DepthPass || DepthMeshPass == EMeshPass::SecondStageDepthPass);
+
 		if (bParallelDepthPass)
 		{
 			RDG_WAIT_FOR_TASKS_CONDITIONAL(GraphBuilder, IsDepthPassWaitForTasksEnabled());
@@ -557,18 +558,17 @@ void FDeferredShadingSceneRenderer::RenderPrePass(FRDGBuilder& GraphBuilder, FRD
 					View.BeginRenderView();
 
 					FDepthPassParameters* PassParameters = GetDepthPassParameters(GraphBuilder, View, SceneDepthTexture);
-					View.ParallelMeshDrawCommandPasses[EMeshPass::DepthPass].BuildRenderingCommands(GraphBuilder, Scene->GPUScene, PassParameters->InstanceCullingDrawParams);
+					View.ParallelMeshDrawCommandPasses[DepthMeshPass].BuildRenderingCommands(GraphBuilder, Scene->GPUScene, PassParameters->InstanceCullingDrawParams);
 
 					GraphBuilder.AddPass(
-						RDG_EVENT_NAME("DepthPassParallel"),
+						DepthMeshPass == EMeshPass::DepthPass ? RDG_EVENT_NAME("DepthPassParallel") : RDG_EVENT_NAME("SecondStageDepthPassParallel"),
 						PassParameters,
 						ERDGPassFlags::Raster | ERDGPassFlags::SkipRenderPass,
-						[this, &View, PassParameters](const FRDGPass* InPass, FRHICommandListImmediate& RHICmdList)
+						[this, &View, PassParameters, DepthMeshPass](const FRDGPass* InPass, FRHICommandListImmediate& RHICmdList)
 					{
 						FRDGParallelCommandListSet ParallelCommandListSet(InPass, RHICmdList, GET_STATID(STAT_CLP_Prepass), View, FParallelCommandListBindings(PassParameters));
 						ParallelCommandListSet.SetHighPriority();
-
-						View.ParallelMeshDrawCommandPasses[EMeshPass::DepthPass].DispatchDraw(&ParallelCommandListSet, RHICmdList, &PassParameters->InstanceCullingDrawParams);
+						View.ParallelMeshDrawCommandPasses[DepthMeshPass].DispatchDraw(&ParallelCommandListSet, RHICmdList, &PassParameters->InstanceCullingDrawParams);
 					});
 
 					RenderPrePassEditorPrimitives(GraphBuilder, View, PassParameters, DrawRenderState, DepthPass.EarlyZPassMode, InstanceCullingManager);
@@ -592,22 +592,29 @@ void FDeferredShadingSceneRenderer::RenderPrePass(FRDGBuilder& GraphBuilder, FRD
 					View.BeginRenderView();
 
 					FDepthPassParameters* PassParameters = GetDepthPassParameters(GraphBuilder, View, SceneDepthTexture);
-					View.ParallelMeshDrawCommandPasses[EMeshPass::DepthPass].BuildRenderingCommands(GraphBuilder, Scene->GPUScene, PassParameters->InstanceCullingDrawParams);
+					View.ParallelMeshDrawCommandPasses[DepthMeshPass].BuildRenderingCommands(GraphBuilder, Scene->GPUScene, PassParameters->InstanceCullingDrawParams);
 
 					GraphBuilder.AddPass(
-						RDG_EVENT_NAME("DepthPass"),
+						DepthMeshPass == EMeshPass::DepthPass ? RDG_EVENT_NAME("DepthPass") : RDG_EVENT_NAME("SecondStageDepthPass"),
 						PassParameters,
 						ERDGPassFlags::Raster,
-						[this, &View, PassParameters](FRHICommandList& RHICmdList)
+						[this, &View, PassParameters, DepthMeshPass](FRHICommandList& RHICmdList)
 					{
 						SetStereoViewport(RHICmdList, View, 1.0f);
-						View.ParallelMeshDrawCommandPasses[EMeshPass::DepthPass].DispatchDraw(nullptr, RHICmdList, &PassParameters->InstanceCullingDrawParams);
+						View.ParallelMeshDrawCommandPasses[DepthMeshPass].DispatchDraw(nullptr, RHICmdList, &PassParameters->InstanceCullingDrawParams);
 					});
 
 					RenderPrePassEditorPrimitives(GraphBuilder, View, PassParameters, DrawRenderState, DepthPass.EarlyZPassMode, InstanceCullingManager);
 				}
 			}
 		}
+	};
+
+	// Draw a depth pass to avoid overdraw in the other passes.
+	if (DepthPass.EarlyZPassMode != DDM_None)
+	{
+		// Render primary depth pass.
+		RenderDepthPass(EMeshPass::DepthPass);
 
 		// Evaluate if any second stage depth buffer processing is required
 		bool bUsesSecondStageDepthPass = false;
@@ -617,7 +624,7 @@ void FDeferredShadingSceneRenderer::RenderPrePass(FRDGBuilder& GraphBuilder, FRD
 			bUsesSecondStageDepthPass |= View.bUsesSecondStageDepthPass;
 		}
 
-		// Copy depth buffer
+		// Copy depth buffer and render secondary depth pass if needed.
 		if(bUsesSecondStageDepthPass)
 		{
 			// We cannot use AddCopyTexturePass which do not support depth htile copy on some platforms. So we do a custom copy using compute through UAV.
@@ -646,77 +653,9 @@ void FDeferredShadingSceneRenderer::RenderPrePass(FRDGBuilder& GraphBuilder, FRD
 						FComputeShaderUtils::GetGroupCount(View.ViewRect.Size(), FCopyViewDepthCS::GetGroupSize()));
 				}
 			}
-		}
 
-		if (bParallelDepthPass)
-		{
-			RDG_WAIT_FOR_TASKS_CONDITIONAL(GraphBuilder, IsDepthPassWaitForTasksEnabled());
-
-			for (int32 ViewIndex = 0; ViewIndex < Views.Num(); ++ViewIndex)
-			{
-				FViewInfo& View = Views[ViewIndex];
-				RDG_GPU_MASK_SCOPE(GraphBuilder, View.GPUMask);
-				RDG_EVENT_SCOPE_CONDITIONAL(GraphBuilder, Views.Num() > 1, "View%d", ViewIndex);
-
-				FMeshPassProcessorRenderState DrawRenderState;
-				SetupDepthPassState(DrawRenderState);
-
-				const bool bShouldRenderView = View.ShouldRenderView() && View.bUsesSecondStageDepthPass;
-				if (bShouldRenderView)
-				{
-					View.BeginRenderView();
-
-					FDepthPassParameters* PassParameters = GetDepthPassParameters(GraphBuilder, View, SceneDepthTexture);
-					View.ParallelMeshDrawCommandPasses[EMeshPass::SecondStageDepthPass].BuildRenderingCommands(GraphBuilder, Scene->GPUScene, PassParameters->InstanceCullingDrawParams);
-
-					GraphBuilder.AddPass(
-						RDG_EVENT_NAME("SecondStageDepthPassParallel"),
-						PassParameters,
-						ERDGPassFlags::Raster | ERDGPassFlags::SkipRenderPass,
-						[this, &View, PassParameters](const FRDGPass* InPass, FRHICommandListImmediate& RHICmdList)
-					{
-						FRDGParallelCommandListSet ParallelCommandListSet(InPass, RHICmdList, GET_STATID(STAT_CLP_Prepass), View, FParallelCommandListBindings(PassParameters));
-						ParallelCommandListSet.SetHighPriority();
-
-						View.ParallelMeshDrawCommandPasses[EMeshPass::SecondStageDepthPass].DispatchDraw(&ParallelCommandListSet, RHICmdList, &PassParameters->InstanceCullingDrawParams);
-					});
-
-					RenderPrePassEditorPrimitives(GraphBuilder, View, PassParameters, DrawRenderState, DepthPass.EarlyZPassMode, InstanceCullingManager);
-				}
-			}
-		}
-		else
-		{
-			for (int32 ViewIndex = 0; ViewIndex < Views.Num(); ++ViewIndex)
-			{
-				FViewInfo& View = Views[ViewIndex];
-				RDG_GPU_MASK_SCOPE(GraphBuilder, View.GPUMask);
-				RDG_EVENT_SCOPE_CONDITIONAL(GraphBuilder, Views.Num() > 1, "View%d", ViewIndex);
-
-				FMeshPassProcessorRenderState DrawRenderState;
-				SetupDepthPassState(DrawRenderState);
-
-				const bool bShouldRenderView = View.ShouldRenderView() && View.bUsesSecondStageDepthPass;
-				if (bShouldRenderView)
-				{
-					View.BeginRenderView();
-
-					FDepthPassParameters* PassParameters = GetDepthPassParameters(GraphBuilder, View, SceneDepthTexture);
-					View.ParallelMeshDrawCommandPasses[EMeshPass::SecondStageDepthPass].BuildRenderingCommands(GraphBuilder, Scene->GPUScene, PassParameters->InstanceCullingDrawParams);
-
-					GraphBuilder.AddPass(
-						RDG_EVENT_NAME("SecondStageDepthPass"),
-						PassParameters,
-						ERDGPassFlags::Raster,
-						[this, &View, PassParameters](FRHICommandList& RHICmdList)
-						{
-							SetStereoViewport(RHICmdList, View, 1.0f);
-					View.ParallelMeshDrawCommandPasses[EMeshPass::SecondStageDepthPass].DispatchDraw(nullptr, RHICmdList, &PassParameters->InstanceCullingDrawParams);
-						});
-
-					RenderPrePassEditorPrimitives(GraphBuilder, View, PassParameters, DrawRenderState, DepthPass.EarlyZPassMode, InstanceCullingManager);
-				}
-			}
+			// Dispatch and render the meshes
+			RenderDepthPass(EMeshPass::SecondStageDepthPass);
 		}
 	}
 
