Moved rendering the debug draw helper execution to after post process pass to avoid ghosting of TAA/TSR/MB.

The EditorPrimitives compositing code was the only PostProcess pass that composited appropriately, but was not setup to accept non-editor draws. DrawDebug also needs to be available to Debug and Development builds, so moving the debug draw calls to this pass directly was not straight forward, and there was no other existing option.

SimpleElement draws used by debug were also non-trivial to setup with velocity to compensate for Temporal/Blur passes, so fixing this in the base pass was also not straight forward.

Then an additional issue was that the color and depth textures used in post process after temporal passes were mismatched in size (as depth is not necessarily automatically included in the temporal upscale pass). However EditorPrimitives had hard coded temporal upsample depth code in it's main AddEditorPrimitivesPass function.

Refactored PostProcessEditorPrimitives code to move the hard coded temporal upscale depth pass to a common file which is protected by UE_ENABLE_DEBUG_DRAWING (which is blocked from compiling in shipping, and works with editor). Also ensured the relevant shader compiles remain protected by these macros. The temporal and populate functions have been slightly rewritten so that a generic temporal history is created rather than an editor only one, and also populate depth can now also output the background colour (to reduce unnecessary draw calls as, depending on the setup, the previous scene may still needs to be drawn to the output color texture).

After that refactor, there is a new set of draw debug composite files which now also use the temporal upsample depth code (when color and depth are mismatched, ignored if not), now that it has been isolated into it's own re-usable structure. This means that we can use a simple draw elements call, stripped of all shader variables except the view, to overlay the draw debug primitives on the scene color along with upscaled depth.

Also moved draw debug to it's own SimpleElement array in the View to avoid conflating it with other simple element draws (e.g. non-opaque, or ones setup in plugins). This code can be extended for those at a later date if needed though.

diff --git a/Engine/Shaders/Private/PostProcessCompositeEditorPrimitives.usf b/Engine/Shaders/Private/PostProcessCompositeEditorPrimitives.usf
deleted file mode 100644
diff --git a/Engine/Shaders/Private/PostProcessCompositePrimitives.usf b/Engine/Shaders/Private/PostProcessCompositePrimitives.usf
new file mode 100644
--- /dev/null
+++ b/Engine/Shaders/Private/PostProcessCompositePrimitives.usf
@@ -0,0 +1,396 @@
+// Copyright Epic Games, Inc. All Rights Reserved.
+
+#include "Common.ush"
+#include "ScreenPass.ush"
+#include "PostProcessCommon.ush"
+
+#ifndef MSAA_SAMPLE_COUNT
+#define MSAA_SAMPLE_COUNT 1
+#endif
+
+#define CONFIG_SAMPLES 6
+
+// K = Center of the nearest input pixel.
+// O = Center of the output pixel.
+//
+//          |           |
+//    0     |     1     |     2
+//          |           |
+//          |           |
+//  --------+-----------+--------
+//          |           |
+//          | O         |
+//    3     |     K     |     5
+//          |           |
+//          |           |
+//  --------+-----------+--------
+//          |           |
+//          |           |
+//    6     |     7     |     8
+//          |           |
+//
+static const int2 kOffsets3x3[9] =
+{
+	int2(-1, -1),
+	int2(0, -1),
+	int2(1, -1),
+	int2(-1,  0),
+	int2(0,  0), // K
+	int2(1,  0),
+	int2(-1,  1),
+	int2(0,  1),
+	int2(1,  1),
+};
+
+// T = Center of the nearest top left pixel input pixel.
+// O = Center of the output pixel.
+//
+//          | 
+//    T     |     .
+//          | 
+//       O  | 
+//  --------+--------
+//          | 
+//          | 
+//    .     |     .
+//          | 
+static const int2 Offsets2x2[4] =
+{
+	int2( 0,  0), // T
+	int2( 1,  0),
+	int2( 0,  1),
+	int2( 1,  1),
+};
+
+// Indexes of the 3x3 square.
+static const uint kSquareIndexes3x3[9] = { 4u, 0u, 1u, 2u, 3u, 5u, 6u, 7u, 8u };
+
+// Indexes of the offsets to have plus + shape.
+static const uint kPlusIndexes3x3[5] = { 4u, 1u, 3u, 5u, 7u };
+
+#if MSAA_SAMPLE_COUNT > 1
+Texture2DMS<float4, MSAA_SAMPLE_COUNT> EditorPrimitivesColor;
+Texture2DMS<float, MSAA_SAMPLE_COUNT> EditorPrimitivesDepth;
+#else
+Texture2D EditorPrimitivesColor;
+Texture2D<float> EditorPrimitivesDepth;
+#endif
+
+Texture2D ColorTexture;
+SamplerState ColorSampler;
+
+Texture2D DepthTexture;
+SamplerState DepthSampler;
+
+Texture2D VelocityTexture;
+SamplerState VelocitySampler;
+
+Texture2D PrevHistoryTexture;
+SamplerState PrevHistorySampler;
+
+SCREEN_PASS_TEXTURE_VIEWPORT(Color)
+SCREEN_PASS_TEXTURE_VIEWPORT(Depth)
+SCREEN_PASS_TEXTURE_VIEWPORT(PrevHistory)
+SCREEN_PASS_TEXTURE_VIEWPORT(History)
+SCREEN_PASS_TEXTURE_VIEWPORT(Output)
+FScreenTransform ColorToDepth;
+
+float2 DepthTextureJitter;
+uint bCameraCut;
+
+float4 SampleOffsetArray[MSAA_SAMPLE_COUNT];
+
+uint bOpaqueEditorGizmo;
+uint bCompositeAnyNonNullDepth;
+
+void MainTemporalUpsampleDepthPS(
+	float4 SvPosition : SV_POSITION,
+	out float OutDeviceZ : SV_Target0)
+{
+	float2 ViewportUV = (SvPosition.xy - History_ViewportMin) * History_ViewportSizeInverse;
+	float2 ScreenPos = ViewportUVToScreenPos(ViewportUV);
+	
+	// Pixel coordinate of the center of output pixel O in the input viewport.
+	float2 PPCo = ViewportUV * Depth_ViewportSize + DepthTextureJitter;
+
+	// Pixel coordinate of the center of the nearest input pixel K.
+	float2 PPCk = floor(PPCo) + 0.5;
+	
+	// Sample nearest depth
+	float InputDeviceZ;
+	{
+		float2 SampleInputBufferUV = (Depth_ViewportMin + PPCk) * Depth_ExtentInverse;
+		SampleInputBufferUV = clamp(SampleInputBufferUV, Depth_UVViewportBilinearMin, Depth_UVViewportBilinearMax);
+			
+		InputDeviceZ = Texture2DSampleLevel(DepthTexture, DepthSampler, SampleInputBufferUV, 0).r;
+	}
+
+	// Compute the screen position to sample in history
+	float2 PrevScreenPos = ScreenPos;
+	#if 1
+	{
+		float4 ThisClip = float4(ScreenPos, InputDeviceZ, 1);
+		float4 PrevClip = mul(ThisClip, View.ClipToPrevClip);
+
+		PrevScreenPos = PrevClip.xy / PrevClip.w;
+
+		#if 0
+		{
+			float2 SampleInputBufferUV = (Depth_ViewportMin + PPCk) * Depth_ExtentInverse;
+			SampleInputBufferUV = clamp(SampleInputBufferUV, Depth_UVViewportBilinearMin, Depth_UVViewportBilinearMax);
+			
+			float4 EncodedVelocity = VelocityTexture.SampleLevel(VelocitySampler, SampleInputBufferUV, 0);
+			bool DynamicN = EncodedVelocity.x > 0.0;
+			if(DynamicN)
+			{
+				PrevScreenPos = ScreenPos - DecodeVelocityFromTexture(EncodedVelocity).xy;
+			}
+		}
+		#endif
+	}
+	#endif
+
+	// Sample the history.
+	float PrevHistoryDeviceZ;
+	{
+		float2 PrevHistoryBufferUV = (PrevHistory_ScreenPosToViewportScale * PrevScreenPos + PrevHistory_ScreenPosToViewportBias) * PrevHistory_ExtentInverse;
+		PrevHistoryBufferUV = clamp(PrevHistoryBufferUV, PrevHistory_UVViewportBilinearMin, PrevHistory_UVViewportBilinearMax);
+
+		PrevHistoryDeviceZ = Texture2DSampleLevel(PrevHistoryTexture, PrevHistorySampler, PrevHistoryBufferUV, 0).r;
+	}
+
+	// Correct the DeviceZ from the previous depth to the current depth.
+	float CorrectedPrevHistoryDeviceZ;
+	#if 1
+	{
+		float PrevHistoryDepth = ConvertFromDeviceZ(PrevHistoryDeviceZ);
+
+		const float3 PreViewTranslationOffset = LWCToFloat(LWCSubtract(PrimaryView.PreViewTranslation, PrimaryView.PrevPreViewTranslation));
+		const float3 PrevHistoryTranslatedWorldPosition = mul(float4(PrevScreenPos * PrevHistoryDepth, PrevHistoryDepth, 1), View.PrevScreenToTranslatedWorld).xyz + PreViewTranslationOffset;
+
+		float CorrectedPrevHistoryDepth = mul(float4(PrevHistoryTranslatedWorldPosition, 1.0), View.TranslatedWorldToView).z;
+		
+		CorrectedPrevHistoryDeviceZ = ConvertToDeviceZ(CorrectedPrevHistoryDepth);
+	}
+	#else
+	{
+		CorrectedPrevHistoryDeviceZ = PrevHistoryDeviceZ;
+	}
+	#endif
+
+	// Replaces NaN from history with 0.0
+	CorrectedPrevHistoryDeviceZ = -min(-CorrectedPrevHistoryDeviceZ, 0.0);
+
+	// Sample current frame
+	bool bUpdateHistory;
+	float InputMinDeviceZ;
+	float InputMaxDeviceZ;
+	{
+		// Vector in pixel between pixel K -> O.
+		float2 dKO = float2(PPCo - PPCk);
+		
+		bUpdateHistory = all(abs(dKO) < 0.5 * (Depth_ViewportSize.x * History_ViewportSizeInverse.x)) || bCameraCut != 0;
+
+		InputMinDeviceZ = InputDeviceZ;
+		InputMaxDeviceZ = InputDeviceZ;
+
+		UNROLL_N(CONFIG_SAMPLES - 1)
+		for (uint SampleId = 1; SampleId < CONFIG_SAMPLES; SampleId++)
+		{
+			float2 PixelOffset;
+			
+			#if CONFIG_SAMPLES == 9
+			{
+				const uint SampleIdx = kSquareIndexes3x3[SampleId];
+				PixelOffset = kOffsets3x3[SampleIdx];
+			}
+			#elif CONFIG_SAMPLES == 5 || CONFIG_SAMPLES == 6
+			{
+				if (SampleId == 5)
+				{
+					PixelOffset = SignFastInt(float2(dKO));
+				}
+				else
+				{
+					const uint SampleIdx = kPlusIndexes3x3[SampleId];
+					PixelOffset = kOffsets3x3[SampleIdx];
+				}
+			}
+			#else
+				#error Unknown sample count
+			#endif
+			
+			float2 SampleInputBufferUV = (Depth_ViewportMin + PPCk + PixelOffset) * Depth_ExtentInverse;
+			SampleInputBufferUV = clamp(SampleInputBufferUV, Depth_UVViewportBilinearMin, Depth_UVViewportBilinearMax);
+
+			float SampleDeviceZ = Texture2DSampleLevel(DepthTexture, DepthSampler, SampleInputBufferUV, 0).r;
+			
+			InputMinDeviceZ = min(InputMinDeviceZ, SampleDeviceZ);
+			InputMaxDeviceZ = max(InputMaxDeviceZ, SampleDeviceZ);
+		}
+	}
+
+	// Reject history based on neighborhood
+	float ClampedHistoryDeviceZ = clamp(CorrectedPrevHistoryDeviceZ, InputMinDeviceZ, InputMaxDeviceZ);
+
+	// Output history
+	float FinalHistoryDeviceZ = bUpdateHistory ? InputDeviceZ : ClampedHistoryDeviceZ;
+	
+	OutDeviceZ = -min(-FinalHistoryDeviceZ, 0.0);
+}
+
+void MainPopulateSceneDepthPS(
+	#if USE_MSAA && !COMPILER_HLSLCC
+		sample noperspective float4 SvPosition : SV_POSITION,
+	#else
+		noperspective float4 SvPosition : SV_POSITION,
+	#endif
+	out float4 OutColor : SV_Target0,
+	out float OutDepth : SV_DEPTH)
+{
+	float2 ViewportUV = (SvPosition.xy - Color_ViewportMin) * Color_ViewportSizeInverse;
+	
+	float2 DepthUV = clamp(
+		(Depth_ViewportMin + ViewportUV * Depth_ViewportSize - DepthTextureJitter) * Depth_ExtentInverse,
+		Depth_UVViewportBilinearMin, Depth_UVViewportBilinearMax);
+
+	OutDepth = Texture2DSampleLevel(DepthTexture, DepthSampler, DepthUV, 0).r;
+#if FORCE_DRAW_COLOR
+	OutColor = Texture2DSampleLevel(ColorTexture, ColorSampler, ViewportUV, 0).rgba;
+#else
+	OutColor = 0.0;
+#endif
+}
+
+float ComputeDepthMask(float SceneDeviceZ, float EditorDeviceZ)
+{
+	// Soft Bias with SceneDeviceZ for best quality
+	const float DeviceDepthFade = 0.00005f;
+
+	return saturate(1.0f - (SceneDeviceZ - EditorDeviceZ) / DeviceDepthFade);
+}
+
+void ResolveEditorPrimitiveColor(int2 PixelPos, out float4 OutColor, out float OutDeviceZ)
+{
+	// Furthest device Z or 0 if there is none.
+	OutDeviceZ = 0;
+	OutColor = 0;
+	
+
+	// Bring out of premultiplied.
+	OutColor.rgb /= max(OutColor.a, 0.0001f);
+	// Fix gamma.
+	OutColor.rgb = pow(OutColor.rgb, 1.0f / 2.2f);
+	// Bring back to premultiplied
+	OutColor.rgb *= OutColor.a;
+}
+
+void MainCompositeEditorPrimitivesPS(
+	float4 SvPosition : SV_POSITION,
+	out float4 OutColor : SV_Target0)
+{
+	float2 ViewportUV = (SvPosition.xy - Output_ViewportMin) * Output_ViewportSizeInverse;
+	
+	const float2 ColorUV = (ViewportUV * Color_ViewportSize + Color_ViewportMin) * Color_ExtentInverse;
+	const int2 ColorPixelPos = int2(ColorUV * Color_Extent);
+	
+	float4 SceneColor = Texture2DSample(ColorTexture, ColorSampler, ColorUV);
+
+	// Resolve editor primitive scene color and depth.
+	float4 EditorPrimitiveColor;
+	float DepthMask;
+
+	#if MSAA_SAMPLE_COUNT > 1
+	{
+		EditorPrimitiveColor = 0.0;
+		DepthMask = 0.0;
+		float DepthMaskWeight = 0.0;
+
+		UNROLL_N(MSAA_SAMPLE_COUNT)
+		for (uint SampleIndex = 0; SampleIndex < MSAA_SAMPLE_COUNT; ++SampleIndex)
+		{
+			const float2 SampleOffset = SampleOffsetArray[SampleIndex].xy;
+
+			float2 SampleDepthUV = (ViewportUV * Depth_ViewportSize + Depth_ViewportMin + SampleOffset - DepthTextureJitter) * Depth_ExtentInverse;
+			SampleDepthUV = clamp(SampleDepthUV, Depth_UVViewportBilinearMin, Depth_UVViewportBilinearMax);
+			float SampleSceneDeviceZ = Texture2DSampleLevel(DepthTexture, DepthSampler, SampleDepthUV, 0).r;
+
+			float4 Sample = EditorPrimitivesColor.Load(ColorPixelPos, SampleIndex);
+			float SampleEditorDeviceZ = EditorPrimitivesDepth.Load(ColorPixelPos, SampleIndex).r;
+
+			// Check if any color was applied to this pixel. Note: This prevents actual black pixels from being visible.
+			float Weight = Sample.a;
+
+			float SampleDepthMask = ComputeDepthMask(SampleSceneDeviceZ, SampleEditorDeviceZ);
+
+			if (SampleEditorDeviceZ > 0.0 && bCompositeAnyNonNullDepth)
+			{
+				Weight = 1;
+			}
+
+			FLATTEN
+			if (Weight)
+			{
+				EditorPrimitiveColor += float4(Sample.rgb, Weight);
+
+				DepthMask += SampleDepthMask * Weight;
+				DepthMaskWeight += Weight;
+			}
+		}
+
+		EditorPrimitiveColor.rgb /= MSAA_SAMPLE_COUNT;
+		EditorPrimitiveColor.a /= MSAA_SAMPLE_COUNT;
+		
+		DepthMask *= DepthMaskWeight > 0 ? rcp(DepthMaskWeight) : 0.0;
+	}
+	#else
+	{
+		// De-jitter the sample position and make a filtered lookup - for planes this allows to reconstruct a much less jittery depth comparison function. It doesn't fix silhouettes, however.
+		float2 SampleDepthUV = (ViewportUV * Depth_ViewportSize + Depth_ViewportMin - DepthTextureJitter) * Depth_ExtentInverse;
+		SampleDepthUV = clamp(SampleDepthUV, Depth_UVViewportBilinearMin, Depth_UVViewportBilinearMax);
+
+		float SceneDeviceZ = Texture2DSampleLevel(DepthTexture, DepthSampler, SampleDepthUV, 0).r;
+		float EditorDeviceZ = EditorPrimitivesDepth.Load(int3(ColorPixelPos, 0)).r;
+		
+		EditorPrimitiveColor = EditorPrimitivesColor.Load(int3(ColorPixelPos, 0));
+		DepthMask = ComputeDepthMask(SceneDeviceZ, EditorDeviceZ);
+
+		if (EditorDeviceZ > 0.0 && bCompositeAnyNonNullDepth)
+		{
+			EditorPrimitiveColor.a = 1;
+		}
+	}
+	#endif
+	
+	// Fixes the gama of editor primitives
+	{
+		// Bring out of premultiplied.
+		EditorPrimitiveColor.rgb /= max(EditorPrimitiveColor.a, 0.0001f);
+		// Fix gamma.
+		EditorPrimitiveColor.rgb = pow(EditorPrimitiveColor.rgb, 1.0f / 2.2f);
+		// Bring back to premultiplied
+		EditorPrimitiveColor.rgb *= EditorPrimitiveColor.a;
+	}
+
+	float2 DepthUV = ApplyScreenTransform(ColorUV, ColorToDepth);
+
+	// Generate 2x2 square tiling pattern for foreground primitive that end up behind scene opaque primitives.
+	float PatternMask = ((ColorPixelPos.x/2 + ColorPixelPos.y/2) % 2) * 0.7f;
+
+	// the contants express the two opacity values we see when the primitive is hidden
+	float LowContrastPatternMask = lerp(0.2, 1, PatternMask);
+
+	LowContrastPatternMask = saturate(lerp(LowContrastPatternMask, 1, bOpaqueEditorGizmo));
+
+	float HiddenMask = lerp(0.7f, 1.0f, DepthMask);
+	float DarkenMask = lerp(LowContrastPatternMask, 1.0f, DepthMask);
+
+	// Blend editor ptimitives with scene color.
+	OutColor.rgb = SceneColor.rgb * (1 - EditorPrimitiveColor.a) + EditorPrimitiveColor.rgb * (DarkenMask * HiddenMask);
+
+	#if POST_PROCESS_ALPHA
+		OutColor.a = lerp(SceneColor.a, 1, EditorPrimitiveColor.a * DarkenMask * HiddenMask);
+	#else
+		OutColor.a = 0;
+	#endif
+}
diff --git a/Engine/Source/Runtime/Engine/Private/Components/LineBatchComponent.cpp b/Engine/Source/Runtime/Engine/Private/Components/LineBatchComponent.cpp
--- a/Engine/Source/Runtime/Engine/Private/Components/LineBatchComponent.cpp
+++ b/Engine/Source/Runtime/Engine/Private/Components/LineBatchComponent.cpp
@@ -64,8 +64,13 @@ void FLineBatcherSceneProxy::GetDynamicMeshElements(const TArray<const FSceneVie
 		if (VisibilityMap & (1 << ViewIndex))
 		{
 			const FSceneView* View = Views[ViewIndex];
-			FPrimitiveDrawInterface* PDI = Collector.GetPDI(ViewIndex);
 
+#if UE_ENABLE_DEBUG_DRAWING
+			//If we are in a debug draw build, this line batcher should write out to the debug only PDI
+			FPrimitiveDrawInterface* PDI = Collector.GetDebugPDI(ViewIndex);
+#else
+			FPrimitiveDrawInterface* PDI = Collector.GetPDI(ViewIndex);
+#endif
 			for (int32 i = 0; i < Lines.Num(); i++)
 			{
 				PDI->DrawLine(Lines[i].Start, Lines[i].End, Lines[i].Color, Lines[i].DepthPriority, Lines[i].Thickness);
diff --git a/Engine/Source/Runtime/Engine/Private/SceneManagement.cpp b/Engine/Source/Runtime/Engine/Private/SceneManagement.cpp
--- a/Engine/Source/Runtime/Engine/Private/SceneManagement.cpp
+++ b/Engine/Source/Runtime/Engine/Private/SceneManagement.cpp
@@ -360,6 +360,13 @@ void FMeshElementCollector::SetPrimitive(const FPrimitiveSceneProxy* InPrimitive
 	{
 		MeshIdInPrimitivePerView[ViewIndex] = 0;
 	}
+#if UE_ENABLE_DEBUG_DRAWING
+	for (int32 ViewIndex = 0; ViewIndex < DebugSimpleElementCollectors.Num(); ViewIndex++)
+	{
+		DebugSimpleElementCollectors[ViewIndex]->HitProxyId = DefaultHitProxyId;
+		DebugSimpleElementCollectors[ViewIndex]->PrimitiveMeshId = 0;
+	}
+#endif
 }
 
 void FMeshElementCollector::ClearViewMeshArrays()
@@ -367,6 +374,9 @@ void FMeshElementCollector::ClearViewMeshArrays()
 	Views.Empty();
 	MeshBatches.Empty();
 	SimpleElementCollectors.Empty();
+#if UE_ENABLE_DEBUG_DRAWING
+	DebugSimpleElementCollectors.Empty();
+#endif
 	MeshIdInPrimitivePerView.Empty();
 	DynamicPrimitiveCollectorPerView.Empty();
 	NumMeshBatchElementsPerView.Empty();
@@ -383,7 +393,11 @@ void FMeshElementCollector::AddViewMeshArrays(
 	ERHIFeatureLevel::Type InFeatureLevel,
 	FGlobalDynamicIndexBuffer* InDynamicIndexBuffer,
 	FGlobalDynamicVertexBuffer* InDynamicVertexBuffer,
-	FGlobalDynamicReadBuffer* InDynamicReadBuffer)
+	FGlobalDynamicReadBuffer* InDynamicReadBuffer
+#if UE_ENABLE_DEBUG_DRAWING
+	,FSimpleElementCollector* InDebugSimpleElementCollector
+#endif
+)
 {
 	Views.Add(InView);
 	MeshIdInPrimitivePerView.Add(0);
@@ -396,6 +410,14 @@ void FMeshElementCollector::AddViewMeshArrays(
 	DynamicIndexBuffer = InDynamicIndexBuffer;
 	DynamicVertexBuffer = InDynamicVertexBuffer;
 	DynamicReadBuffer = InDynamicReadBuffer;
+
+#if UE_ENABLE_DEBUG_DRAWING
+	//Assign the debug draw only simple element collector per view	
+	if (InDebugSimpleElementCollector)
+	{
+		DebugSimpleElementCollectors.Add(InDebugSimpleElementCollector);
+	}
+#endif
 }
 
 void FMeshElementCollector::ProcessTasks()
diff --git a/Engine/Source/Runtime/Engine/Public/SceneManagement.h b/Engine/Source/Runtime/Engine/Public/SceneManagement.h
--- a/Engine/Source/Runtime/Engine/Public/SceneManagement.h
+++ b/Engine/Source/Runtime/Engine/Public/SceneManagement.h
@@ -1779,6 +1779,11 @@ public:
 
 	void DrawBatchedElements(FRHICommandList& RHICmdList, const FMeshPassProcessorRenderState& DrawRenderState, const FSceneView& InView, EBlendModeFilter::Type Filter, ESceneDepthPriorityGroup DPG) const;
 
+	bool HasAnyPrimitives() const
+	{
+		return BatchedElements.HasPrimsToDraw() || TopBatchedElements.HasPrimsToDraw();
+	}
+
 	bool HasPrimitives(ESceneDepthPriorityGroup DPG) const
 	{
 		if (DPG == SDPG_World)
@@ -1859,6 +1864,13 @@ public:
 		return SimpleElementCollectors[ViewIndex];
 	}
 
+#if UE_ENABLE_DEBUG_DRAWING
+	inline FPrimitiveDrawInterface* GetDebugPDI(int32 ViewIndex)
+	{
+		return DebugSimpleElementCollectors[ViewIndex];
+	}
+#endif
+
 	/** 
 	 * Allocates an FMeshBatch that can be safely referenced by the collector (lifetime will be long enough).
 	 * Returns a reference that will not be invalidated due to further AllocateMesh() calls.
@@ -1962,7 +1974,11 @@ protected:
 		ERHIFeatureLevel::Type InFeatureLevel,
 		FGlobalDynamicIndexBuffer* InDynamicIndexBuffer,
 		FGlobalDynamicVertexBuffer* InDynamicVertexBuffer,
-		FGlobalDynamicReadBuffer* InDynamicReadBuffer);
+		FGlobalDynamicReadBuffer* InDynamicReadBuffer
+#if UE_ENABLE_DEBUG_DRAWING
+		,FSimpleElementCollector* InDebugSimpleElementCollector = nullptr
+#endif
+	);
 
 	/** 
 	 * Using TChunkedArray which will never realloc as new elements are added
@@ -1979,6 +1995,10 @@ protected:
 	/** PDIs */
 	TArray<FSimpleElementCollector*, TInlineAllocator<2, SceneRenderingAllocator> > SimpleElementCollectors;
 
+#if UE_ENABLE_DEBUG_DRAWING
+	TArray<FSimpleElementCollector*, TInlineAllocator<2, SceneRenderingAllocator> > DebugSimpleElementCollectors;
+#endif
+
 	/** Views being collected for */
 	TArray<FSceneView*, TInlineAllocator<2, SceneRenderingAllocator>> Views;
 
diff --git a/Engine/Source/Runtime/Renderer/Private/PostProcess/PostProcessCompositeDebugPrimitives.cpp b/Engine/Source/Runtime/Renderer/Private/PostProcess/PostProcessCompositeDebugPrimitives.cpp
new file mode 100644
--- /dev/null
+++ b/Engine/Source/Runtime/Renderer/Private/PostProcess/PostProcessCompositeDebugPrimitives.cpp
@@ -0,0 +1,95 @@
+// Copyright Epic Games, Inc. All Rights Reserved.
+
+#include "PostProcess/PostProcessCompositeDebugPrimitives.h"
+#if UE_ENABLE_DEBUG_DRAWING
+#include "ScenePrivate.h"
+
+BEGIN_SHADER_PARAMETER_STRUCT(FDebugPrimitivesPassParameters, )
+	SHADER_PARAMETER_STRUCT_INCLUDE(FViewShaderParameters, View)
+	RENDER_TARGET_BINDING_SLOTS()
+END_SHADER_PARAMETER_STRUCT()
+
+FScreenPassTexture AddDebugPrimitivePass(
+	FRDGBuilder& GraphBuilder,
+	const FViewInfo& View,
+	const FCompositePrimitiveInputs& Inputs)
+{
+	check(Inputs.SceneColor.IsValid());
+	check(Inputs.SceneDepth.IsValid());
+
+	RDG_EVENT_SCOPE(GraphBuilder, "CompositeDebugPrimitives");
+	
+	const uint32 NumMSAASamples = View.GetSceneTextures().Config.NumSamples;
+	//Set sizing values to match the SceneColor size
+	const FIntPoint Extent = Inputs.SceneColor.Texture->Desc.Extent;
+	const FIntRect ViewRect = Inputs.SceneColor.ViewRect;
+	const FViewInfo* DebugView = CreateCompositePrimitiveView(View, ViewRect, NumMSAASamples);
+
+	FScreenPassRenderTarget Output = Inputs.OverrideOutput;
+	FRDGTextureRef DebugPrimitiveColor = Output.Texture;
+	if (!Output.IsValid())
+	{
+		// Create a new output for the scene color
+		FRDGTextureDesc ColorDesc = Inputs.SceneColor.Texture->Desc;
+		ColorDesc.Flags |= TexCreate_ShaderResource | TexCreate_RenderTargetable;
+		DebugPrimitiveColor = GraphBuilder.CreateTexture(ColorDesc, TEXT("Debug.DrawPrimitivesColor"));
+		//Create a new depth output for the final draw
+		Output = FScreenPassRenderTarget(DebugPrimitiveColor, ViewRect, ERenderTargetLoadAction::ENoAction);
+	}
+
+	//Prepare output textures for composite draw
+	const FScreenPassTextureViewport OutputViewport(Output);
+	FRDGTextureRef DebugPrimitiveDepth = CreateCompositeDepthTexture(GraphBuilder, Extent, NumMSAASamples);
+	FScreenPassTexture SceneDepth = Inputs.SceneDepth;//Inputs is const so create a over-ridable texture reference
+	FVector2f SceneDepthJitter = FVector2f(View.TemporalJitterPixels);
+
+	if (IsTemporalAccumulationBasedMethod(View.AntiAliasingMethod))
+	{
+		TemporalUpscaleDepthPass(GraphBuilder,
+			*DebugView,
+			Inputs.SceneColor,
+			SceneDepth,
+			SceneDepthJitter,
+			NumMSAASamples);
+	}
+
+	//Simple element pixel shaders do not output background color for composite, 
+	//so this allows the background to be drawn to the RT at the same time as depth without adding extra draw calls
+	PopulateDepthPass(GraphBuilder,
+		*DebugView,
+		Inputs.SceneColor,
+		SceneDepth,
+		DebugPrimitiveColor,
+		DebugPrimitiveDepth,
+		SceneDepthJitter,
+		NumMSAASamples,
+		true);	// - Draw Background Colour
+
+	//Composite the debug draw elements into the scene
+	FDebugPrimitivesPassParameters* PassParameters = GraphBuilder.AllocParameters<FDebugPrimitivesPassParameters>();
+	PassParameters->View = DebugView->GetShaderParameters();
+	PassParameters->RenderTargets[0] = Output.GetRenderTargetBinding();
+	PassParameters->RenderTargets[0].SetLoadAction(ERenderTargetLoadAction::ELoad);
+	PassParameters->RenderTargets.DepthStencil = FDepthStencilBinding(DebugPrimitiveDepth, ERenderTargetLoadAction::ELoad, ERenderTargetLoadAction::ELoad, FExclusiveDepthStencil::DepthWrite_StencilWrite);
+
+	GraphBuilder.AddPass(
+		RDG_EVENT_NAME("DrawDebugPrimitives"),
+		PassParameters,
+		ERDGPassFlags::Raster,
+		[&View, PassParameters, DebugView, OutputViewport](FRHICommandList& RHICmdList)
+	{
+		RHICmdList.SetViewport(OutputViewport.Rect.Min.X, OutputViewport.Rect.Min.Y, 0.0f, OutputViewport.Rect.Max.X, OutputViewport.Rect.Max.Y, 1.0f);
+
+		FMeshPassProcessorRenderState DrawRenderState;
+		DrawRenderState.SetDepthStencilAccess(FExclusiveDepthStencil::DepthWrite_StencilWrite);
+		DrawRenderState.SetBlendState(TStaticBlendState<CW_RGBA, BO_Add, BF_One, BF_One, BO_Add, BF_One, BF_One>::GetRHI());
+		DrawRenderState.SetDepthStencilState(TStaticDepthStencilState<true, CF_DepthNearOrEqual>::GetRHI());
+		
+		DebugView->DebugSimpleElementCollector.DrawBatchedElements(RHICmdList, DrawRenderState, *DebugView, EBlendModeFilter::OpaqueAndMasked, SDPG_World);
+		DebugView->DebugSimpleElementCollector.DrawBatchedElements(RHICmdList, DrawRenderState, *DebugView, EBlendModeFilter::OpaqueAndMasked, SDPG_Foreground);		
+	});
+
+	return MoveTemp(Output);
+}
+
+#endif
\ No newline at end of file
diff --git a/Engine/Source/Runtime/Renderer/Private/PostProcess/PostProcessCompositeDebugPrimitives.h b/Engine/Source/Runtime/Renderer/Private/PostProcess/PostProcessCompositeDebugPrimitives.h
new file mode 100644
--- /dev/null
+++ b/Engine/Source/Runtime/Renderer/Private/PostProcess/PostProcessCompositeDebugPrimitives.h
@@ -0,0 +1,8 @@
+// Copyright Epic Games, Inc. All Rights Reserved.
+
+#pragma once
+#include "PostProcess/PostProcessCompositePrimitivesCommon.h"
+
+#if UE_ENABLE_DEBUG_DRAWING
+FScreenPassTexture AddDebugPrimitivePass(FRDGBuilder& GraphBuilder, const FViewInfo& View, const FCompositePrimitiveInputs& Inputs);
+#endif
diff --git a/Engine/Source/Runtime/Renderer/Private/PostProcess/PostProcessCompositeEditorPrimitives.cpp b/Engine/Source/Runtime/Renderer/Private/PostProcess/PostProcessCompositeEditorPrimitives.cpp
--- a/Engine/Source/Runtime/Renderer/Private/PostProcess/PostProcessCompositeEditorPrimitives.cpp
+++ b/Engine/Source/Runtime/Renderer/Private/PostProcess/PostProcessCompositeEditorPrimitives.cpp
@@ -1,122 +1,47 @@
 // Copyright Epic Games, Inc. All Rights Reserved.
 
-#if WITH_EDITOR
-
 #include "PostProcess/PostProcessCompositeEditorPrimitives.h"
+
+#if WITH_EDITOR
 #include "EditorPrimitivesRendering.h"
 #include "MeshPassProcessor.inl"
 #include "BasePassRendering.h"
 #include "MobileBasePassRendering.h"
-#include "SceneRenderingUtils.h"
 #include "PixelShaderUtils.h"
-#include "ScenePrivate.h"
 
 namespace
 {
-TAutoConsoleVariable<int32> CVarEditorTemporalUpsampleDepth(
-	TEXT("r.Editor.TemporalUpsampleDepth"), 2,
-	TEXT("Temporal upsample factor of the depth buffer for depth testing editor primitives against."),
-	ECVF_RenderThreadSafe);
-
-class FTemporalUpsampleEditorDepthPS : public FGlobalShader
-{
-public:
-	DECLARE_GLOBAL_SHADER(FTemporalUpsampleEditorDepthPS);
-	SHADER_USE_PARAMETER_STRUCT(FTemporalUpsampleEditorDepthPS, FGlobalShader);
-
-	BEGIN_SHADER_PARAMETER_STRUCT(FParameters, )
-		SHADER_PARAMETER_STRUCT_REF(FViewUniformShaderParameters, View)
-		SHADER_PARAMETER_STRUCT(FScreenPassTextureViewportParameters, Depth)
-		SHADER_PARAMETER_STRUCT(FScreenPassTextureViewportParameters, PrevHistory)
-		SHADER_PARAMETER_STRUCT(FScreenPassTextureViewportParameters, History)
-		SHADER_PARAMETER(FVector2f, DepthTextureJitter)
-		SHADER_PARAMETER(int32, bCameraCut)
-
-		SHADER_PARAMETER_RDG_TEXTURE(Texture2D, DepthTexture)
-		SHADER_PARAMETER_SAMPLER(SamplerState, DepthSampler)
-		SHADER_PARAMETER_RDG_TEXTURE(Texture2D, VelocityTexture)
-		SHADER_PARAMETER_SAMPLER(SamplerState, VelocitySampler)
-		SHADER_PARAMETER_RDG_TEXTURE(Texture2D, PrevHistoryTexture)
-		SHADER_PARAMETER_SAMPLER(SamplerState, PrevHistorySampler)
-
-		RENDER_TARGET_BINDING_SLOTS()
-	END_SHADER_PARAMETER_STRUCT()
-
-	static bool ShouldCompilePermutation(const FGlobalShaderPermutationParameters& Parameters)
-	{
-		// Only PC platforms render editor primitives.
-		return IsPCPlatform(Parameters.Platform);
-	}
-};
-
-class FPopulateEditorDepthPS : public FGlobalShader
-{
-public:
-	DECLARE_GLOBAL_SHADER(FPopulateEditorDepthPS);
-	SHADER_USE_PARAMETER_STRUCT(FPopulateEditorDepthPS, FGlobalShader);
-
-	class FUseMSAADimension : SHADER_PERMUTATION_BOOL("USE_MSAA");
-	using FPermutationDomain = TShaderPermutationDomain<FUseMSAADimension>;
-
-	BEGIN_SHADER_PARAMETER_STRUCT(FParameters, )
-		SHADER_PARAMETER_STRUCT_REF(FViewUniformShaderParameters, View)
-		SHADER_PARAMETER_STRUCT(FScreenPassTextureViewportParameters, Color)
-		SHADER_PARAMETER_STRUCT(FScreenPassTextureViewportParameters, Depth)
-		SHADER_PARAMETER(FVector2f, DepthTextureJitter)
-		SHADER_PARAMETER_RDG_TEXTURE(Texture2D, DepthTexture)
-		SHADER_PARAMETER_SAMPLER(SamplerState, DepthSampler)
-		RENDER_TARGET_BINDING_SLOTS()
-	END_SHADER_PARAMETER_STRUCT()
-
-	static bool ShouldCompilePermutation(const FGlobalShaderPermutationParameters& Parameters)
+	class FCompositeEditorPrimitivesPS : public FCompositePrimitiveShaderBase
 	{
-		const FPermutationDomain PermutationVector(Parameters.PermutationId);
-		const bool bUseMSAA = PermutationVector.Get<FUseMSAADimension>();
-
-		// Only SM5+ platforms supports MSAA.
-		if (!IsFeatureLevelSupported(Parameters.Platform, ERHIFeatureLevel::SM5) && bUseMSAA)
+	public:
+		DECLARE_GLOBAL_SHADER(FCompositeEditorPrimitivesPS);
+		SHADER_USE_PARAMETER_STRUCT(FCompositeEditorPrimitivesPS, FCompositePrimitiveShaderBase);
+
+		BEGIN_SHADER_PARAMETER_STRUCT(FParameters, )
+			SHADER_PARAMETER_STRUCT_REF(FViewUniformShaderParameters, View)
+			SHADER_PARAMETER_STRUCT(FScreenPassTextureViewportParameters, Color)
+			SHADER_PARAMETER_STRUCT(FScreenPassTextureViewportParameters, Depth)
+			SHADER_PARAMETER_STRUCT(FScreenPassTextureViewportParameters, Output)
+			SHADER_PARAMETER_ARRAY(FVector4f, SampleOffsetArray, [FCompositePrimitiveShaderBase::kMSAASampleCountMax])
+			SHADER_PARAMETER_RDG_TEXTURE(Texture2D, EditorPrimitivesDepth)
+			SHADER_PARAMETER_RDG_TEXTURE(Texture2D, EditorPrimitivesColor)
+			SHADER_PARAMETER_RDG_TEXTURE(Texture2D, ColorTexture)
+			SHADER_PARAMETER_SAMPLER(SamplerState, ColorSampler)
+			SHADER_PARAMETER_RDG_TEXTURE(Texture2D, DepthTexture)
+			SHADER_PARAMETER_SAMPLER(SamplerState, DepthSampler)
+			SHADER_PARAMETER(uint32, bOpaqueEditorGizmo)
+			SHADER_PARAMETER(uint32, bCompositeAnyNonNullDepth)
+			SHADER_PARAMETER(FVector2f, DepthTextureJitter)
+			RENDER_TARGET_BINDING_SLOTS()
+			END_SHADER_PARAMETER_STRUCT()
+
+			static void ModifyCompilationEnvironment(const FGlobalShaderPermutationParameters& Parameters, FShaderCompilerEnvironment& OutEnvironment)
 		{
-			return false;
+			FGlobalShader::ModifyCompilationEnvironment(Parameters, OutEnvironment);
 		}
+	};
+	IMPLEMENT_GLOBAL_SHADER(FCompositeEditorPrimitivesPS, "/Engine/Private/PostProcessCompositePrimitives.usf", "MainCompositeEditorPrimitivesPS", SF_Pixel);
 
-		// Only PC platforms render editor primitives.
-		return IsPCPlatform(Parameters.Platform);
-	}
-};
-
-class FCompositeEditorPrimitivesPS : public FEditorPrimitiveShader
-{
-public:
-	DECLARE_GLOBAL_SHADER(FCompositeEditorPrimitivesPS);
-	SHADER_USE_PARAMETER_STRUCT(FCompositeEditorPrimitivesPS, FEditorPrimitiveShader);
-
-	BEGIN_SHADER_PARAMETER_STRUCT(FParameters, )
-		SHADER_PARAMETER_STRUCT_REF(FViewUniformShaderParameters, View)
-		SHADER_PARAMETER_STRUCT(FScreenPassTextureViewportParameters, Color)
-		SHADER_PARAMETER_STRUCT(FScreenPassTextureViewportParameters, Depth)
-		SHADER_PARAMETER_STRUCT(FScreenPassTextureViewportParameters, Output)
-		SHADER_PARAMETER_ARRAY(FVector4f, SampleOffsetArray, [FEditorPrimitiveShader::kEditorMSAASampleCountMax])
-		SHADER_PARAMETER_RDG_TEXTURE(Texture2D, EditorPrimitivesDepth)
-		SHADER_PARAMETER_RDG_TEXTURE(Texture2D, EditorPrimitivesColor)
-		SHADER_PARAMETER_RDG_TEXTURE(Texture2D, ColorTexture)
-		SHADER_PARAMETER_SAMPLER(SamplerState, ColorSampler)
-		SHADER_PARAMETER_RDG_TEXTURE(Texture2D, DepthTexture)
-		SHADER_PARAMETER_SAMPLER(SamplerState, DepthSampler)
-		SHADER_PARAMETER(uint32, bOpaqueEditorGizmo)
-		SHADER_PARAMETER(uint32, bCompositeAnyNonNullDepth)
-		SHADER_PARAMETER(FVector2f, DepthTextureJitter)
-		RENDER_TARGET_BINDING_SLOTS()
-	END_SHADER_PARAMETER_STRUCT()
-
-	static void ModifyCompilationEnvironment(const FGlobalShaderPermutationParameters& Parameters, FShaderCompilerEnvironment& OutEnvironment)
-	{
-		FGlobalShader::ModifyCompilationEnvironment(Parameters, OutEnvironment);
-	}
-};
-
-IMPLEMENT_GLOBAL_SHADER(FTemporalUpsampleEditorDepthPS, "/Engine/Private/PostProcessCompositeEditorPrimitives.usf", "MainTemporalUpsampleEditorDepthPS", SF_Pixel);
-IMPLEMENT_GLOBAL_SHADER(FPopulateEditorDepthPS, "/Engine/Private/PostProcessCompositeEditorPrimitives.usf", "MainPopulateSceneDepthPS", SF_Pixel);
-IMPLEMENT_GLOBAL_SHADER(FCompositeEditorPrimitivesPS, "/Engine/Private/PostProcessCompositeEditorPrimitives.usf", "MainCompositeEditorPrimitivesPS", SF_Pixel);
 
 void RenderEditorPrimitives(FRHICommandList& RHICmdList, const FViewInfo& View, FMeshPassProcessorRenderState& DrawRenderState, FInstanceCullingManager& InstanceCullingManager)
 {
@@ -271,29 +196,6 @@ void RenderForegroundEditorPrimitives(FRHICommandList& RHICmdList, const FViewIn
 
 } //! namespace
 
-const FViewInfo* CreateEditorPrimitiveView(const FViewInfo& ParentView, FIntRect ViewRect, uint32 NumSamples)
-{
-	FViewInfo* EditorView = ParentView.CreateSnapshot();
-
-	// Patch view rect.
-	EditorView->ViewRect = ViewRect;
-
-	// Override pre exposure to 1.0f, because rendering after tonemapper. 
-	EditorView->PreExposure = 1.0f;
-
-	// Kills material texture mipbias because after TAA.
-	EditorView->MaterialTextureMipBias = 0.0f;
-
-	if (IsTemporalAccumulationBasedMethod(EditorView->AntiAliasingMethod))
-	{
-		EditorView->ViewMatrices.HackRemoveTemporalAAProjectionJitter();
-	}
-
-	EditorView->InitRHIResources(NumSamples);
-
-	return EditorView;
-}
-
 BEGIN_SHADER_PARAMETER_STRUCT(FEditorPrimitivesPassParameters, )
 	SHADER_PARAMETER_STRUCT_INCLUDE(FViewShaderParameters, View)
 	SHADER_PARAMETER_RDG_UNIFORM_BUFFER(FSceneUniformParameters, Scene)
@@ -309,16 +211,16 @@ END_SHADER_PARAMETER_STRUCT()
 FScreenPassTexture AddEditorPrimitivePass(
 	FRDGBuilder& GraphBuilder,
 	const FViewInfo& View,
-	const FEditorPrimitiveInputs& Inputs,
+	const FCompositePrimitiveInputs& Inputs,
 	FInstanceCullingManager& InstanceCullingManager)
 {
 	check(Inputs.SceneColor.IsValid());
 	check(Inputs.SceneDepth.IsValid());
-	check(Inputs.BasePassType != FEditorPrimitiveInputs::EBasePassType::MAX);
+	check(Inputs.BasePassType != FCompositePrimitiveInputs::EBasePassType::MAX);
 
 	const FSceneTextures& SceneTextures = View.GetSceneTextures();
-	const uint32 NumSamples = SceneTextures.Config.EditorPrimitiveNumSamples;
-	const FViewInfo* EditorView = CreateEditorPrimitiveView(View, Inputs.SceneColor.ViewRect, NumSamples);
+	const uint32 NumMSAASamples = SceneTextures.Config.EditorPrimitiveNumSamples;
+	const FViewInfo* EditorView = CreateCompositePrimitiveView(View, Inputs.SceneColor.ViewRect, NumMSAASamples);
 
 	// Load the color target if it already exists.
 	const bool bProducedByPriorPass = HasBeenProduced(SceneTextures.EditorPrimitiveColor);
@@ -335,8 +237,6 @@ FScreenPassTexture AddEditorPrimitivePass(
 	}
 	else
 	{
-		const FSceneTexturesConfig& Config = SceneTextures.Config;
-
 		FIntPoint Extent = Inputs.SceneColor.Texture->Desc.Extent;
 
 		const FRDGTextureDesc ColorDesc = FRDGTextureDesc::Create2D(
@@ -345,18 +245,10 @@ FScreenPassTexture AddEditorPrimitivePass(
 			FClearValueBinding::Transparent,
 			TexCreate_ShaderResource | TexCreate_RenderTargetable,
 			1,
-			Config.EditorPrimitiveNumSamples);
-
-		const FRDGTextureDesc DepthDesc = FRDGTextureDesc::Create2D(
-			Extent,
-			PF_DepthStencil,
-			FClearValueBinding::DepthFar,
-			TexCreate_ShaderResource | TexCreate_DepthStencilTargetable,
-			1,
-			Config.EditorPrimitiveNumSamples);
+			NumMSAASamples);
 
 		EditorPrimitiveColor = GraphBuilder.CreateTexture(ColorDesc, TEXT("Editor.PrimitivesColor"));
-		EditorPrimitiveDepth = GraphBuilder.CreateTexture(DepthDesc, TEXT("Editor.PrimitivesDepth"));
+		EditorPrimitiveDepth = CreateCompositeDepthTexture(GraphBuilder, Extent, NumMSAASamples);
 	}
 
 	// Load the color target if it already exists.
@@ -366,8 +258,9 @@ FScreenPassTexture AddEditorPrimitivePass(
 	RDG_EVENT_SCOPE(GraphBuilder, "CompositeEditorPrimitives %dx%d MSAA=%d",
 		EditorPrimitivesViewport.Rect.Width(),
 		EditorPrimitivesViewport.Rect.Height(),
-		NumSamples);
+		NumMSAASamples);
 
+	//Inputs is const so create a over-ridable texture reference
 	FScreenPassTexture SceneDepth = Inputs.SceneDepth;
 	FVector2f SceneDepthJitter = FVector2f(View.TemporalJitterPixels);
 
@@ -378,121 +271,25 @@ FScreenPassTexture AddEditorPrimitivePass(
 		// Populate depth if a prior pass did not already do it.
 		if (!bProducedByPriorPass)
 		{
-			// Upscale factor of the depth buffer that might be needed.
-			const float UpscaleFactor = float(EditorPrimitivesViewport.Rect.Width()) / float(Inputs.SceneDepth.ViewRect.Width());
-
-			// Upscale factor shouldn't be higher than there is TAA samples, or that means there will be unrendered pixels.
-			const int32 ComputeMaxUpsampleFactorDueToTAA = FMath::FloorToInt(FMath::Sqrt(float(View.TemporalJitterSequenceLength) / (UpscaleFactor * UpscaleFactor)));
-			
-			const int32 ComputeMaxUpsampleFactor = FMath::Clamp(ComputeMaxUpsampleFactorDueToTAA, 0, 4);
-
-			const int32 DepthUpsampleFactor = FMath::Clamp(CVarEditorTemporalUpsampleDepth.GetValueOnRenderThread(), 0, ComputeMaxUpsampleFactor);
-
-			// Upsample the depth at higher resolution to reduce depth intersection instability of editor primitives.
-			if (DepthUpsampleFactor > 0 && View.ViewState && IsTemporalAccumulationBasedMethod(View.AntiAliasingMethod))
+			if (IsTemporalAccumulationBasedMethod(View.AntiAliasingMethod))
 			{
-				FScreenPassTexture History;
-				{
-					const FRDGTextureDesc Desc = FRDGTextureDesc::Create2D(
-						EditorPrimitivesViewport.Extent * DepthUpsampleFactor,
-						PF_R32_FLOAT,
-						FClearValueBinding::None,
-						TexCreate_ShaderResource | TexCreate_RenderTargetable);
-
-					History.Texture = GraphBuilder.CreateTexture(Desc, TEXT("Editor.PrimitivesDepthHistory"));
-					History.ViewRect = EditorPrimitivesViewport.Rect * DepthUpsampleFactor;
-				}
-
-				FTemporalUpsampleEditorDepthPS::FParameters* PassParameters = GraphBuilder.AllocParameters<FTemporalUpsampleEditorDepthPS::FParameters>();
-				PassParameters->View = View.ViewUniformBuffer;
-				PassParameters->Depth = GetScreenPassTextureViewportParameters(FScreenPassTextureViewport(Inputs.SceneDepth));
-				PassParameters->History = GetScreenPassTextureViewportParameters(FScreenPassTextureViewport(History));
-				PassParameters->DepthTextureJitter = FVector2f(View.TemporalJitterPixels);
-
-				PassParameters->DepthTexture = Inputs.SceneDepth.Texture;
-				PassParameters->DepthSampler = TStaticSamplerState<SF_Point, AM_Clamp, AM_Clamp, AM_Clamp>::GetRHI();
-
-				// TODO
-				PassParameters->VelocityTexture = GraphBuilder.RegisterExternalTexture(GSystemTextures.BlackDummy);
-				PassParameters->VelocitySampler = TStaticSamplerState<SF_Point, AM_Clamp, AM_Clamp, AM_Clamp>::GetRHI();
-
-				if (View.PrevViewInfo.EditorPrimtiveDepthHistory.IsValid())
-				{
-					PassParameters->PrevHistory = GetScreenPassTextureViewportParameters(FScreenPassTextureViewport(
-						View.PrevViewInfo.EditorPrimtiveDepthHistory.RT[0]->GetDesc().Extent, View.PrevViewInfo.EditorPrimtiveDepthHistory.ViewportRect));
-					PassParameters->PrevHistoryTexture = GraphBuilder.RegisterExternalTexture(View.PrevViewInfo.EditorPrimtiveDepthHistory.RT[0]);
-					PassParameters->PrevHistorySampler = TStaticSamplerState<SF_Bilinear, AM_Clamp, AM_Clamp, AM_Clamp>::GetRHI();
-					PassParameters->bCameraCut = false;
-				}
-				else
-				{
-					PassParameters->PrevHistory = GetScreenPassTextureViewportParameters(FScreenPassTextureViewport(
-						FIntPoint(1, 1), FIntRect(FIntPoint(0, 0), FIntPoint(1, 1))));
-					PassParameters->PrevHistoryTexture = GraphBuilder.RegisterExternalTexture(GSystemTextures.BlackDummy);
-					PassParameters->PrevHistorySampler = TStaticSamplerState<SF_Point, AM_Clamp, AM_Clamp, AM_Clamp>::GetRHI();
-					PassParameters->bCameraCut = true;
-				}
-
-				PassParameters->RenderTargets[0] = FRenderTargetBinding(History.Texture, ERenderTargetLoadAction::ENoAction);
-
-				TShaderMapRef<FTemporalUpsampleEditorDepthPS> PixelShader(View.ShaderMap);
-				FPixelShaderUtils::AddFullscreenPass(
-					GraphBuilder,
-					View.ShaderMap,
-					RDG_EVENT_NAME("TemporalUpsampleDepth %dx%d -> %dx%d",
-						Inputs.SceneDepth.ViewRect.Width(),
-						Inputs.SceneDepth.ViewRect.Height(),
-						History.ViewRect.Width(),
-						History.ViewRect.Height()),
-					PixelShader,
-					PassParameters,
-					History.ViewRect);
-
-				if (!View.bStatePrevViewInfoIsReadOnly)
-				{
-					FTemporalAAHistory* OutputHistory = &View.ViewState->PrevFrameViewInfo.EditorPrimtiveDepthHistory;
-					OutputHistory->SafeRelease();
-
-					GraphBuilder.QueueTextureExtraction(History.Texture, &OutputHistory->RT[0]);
-					OutputHistory->ViewportRect = History.ViewRect;
-					OutputHistory->ReferenceBufferSize = History.Texture->Desc.Extent;
-				}
-
-				SceneDepth = History;
-				SceneDepthJitter = FVector2f::ZeroVector;
+				TemporalUpscaleDepthPass(GraphBuilder,
+					*EditorView,
+					Inputs.SceneColor,
+					SceneDepth,
+					SceneDepthJitter,
+					NumMSAASamples);
 			}
 
-			// Populate the MSAA depth buffer from depth buffer or temporally upscaled depth buffer
-			{
-				FPopulateEditorDepthPS::FParameters* PassParameters = GraphBuilder.AllocParameters<FPopulateEditorDepthPS::FParameters>();
-				PassParameters->View = View.ViewUniformBuffer;
-				PassParameters->Color = GetScreenPassTextureViewportParameters(FScreenPassTextureViewport(Inputs.SceneColor));
-				PassParameters->Depth = GetScreenPassTextureViewportParameters(FScreenPassTextureViewport(SceneDepth));
-				PassParameters->DepthTextureJitter = SceneDepthJitter;
-				PassParameters->DepthTexture = SceneDepth.Texture;
-				PassParameters->DepthSampler = TStaticSamplerState<SF_Point, AM_Clamp, AM_Clamp, AM_Clamp>::GetRHI();
-				PassParameters->RenderTargets[0] = FRenderTargetBinding(EditorPrimitiveColor, ERenderTargetLoadAction::EClear);
-				PassParameters->RenderTargets.DepthStencil = FDepthStencilBinding(EditorPrimitiveDepth, ERenderTargetLoadAction::EClear, ERenderTargetLoadAction::EClear, FExclusiveDepthStencil::DepthWrite_StencilWrite);
-
-				FPopulateEditorDepthPS::FPermutationDomain PermutationVector;
-				PermutationVector.Set<FPopulateEditorDepthPS::FUseMSAADimension>(NumSamples > 1);
-				TShaderMapRef<FPopulateEditorDepthPS> PixelShader(View.ShaderMap, PermutationVector);
-
-				FPixelShaderUtils::AddFullscreenPass(
-					GraphBuilder,
-					View.ShaderMap,
-					RDG_EVENT_NAME("PopulateDepth %dx%d%s",
-						EditorPrimitivesViewport.Rect.Width(),
-						EditorPrimitivesViewport.Rect.Height(),
-						NumSamples > 1 ? TEXT(" MSAA") : TEXT("")),
-					PixelShader,
-					PassParameters,
-					EditorPrimitivesViewport.Rect,
-					/* BlendState = */ nullptr,
-					/* RasterizerState = */ nullptr,
-					TStaticDepthStencilState<true, CF_Always>::GetRHI());
-			}
-		}
+			PopulateDepthPass(GraphBuilder,
+				*EditorView,
+				Inputs.SceneColor,
+				SceneDepth,
+				EditorPrimitiveColor,
+				EditorPrimitiveDepth,
+				SceneDepthJitter,
+				NumMSAASamples);
+		}	
 
 		// Draws the editors opaque primitives
 		{
@@ -505,9 +302,9 @@ FScreenPassTexture AddEditorPrimitivePass(
 			PassParameters->RenderTargets[0] = FRenderTargetBinding(EditorPrimitiveColor, ERenderTargetLoadAction::ELoad);
 			PassParameters->RenderTargets.DepthStencil = FDepthStencilBinding(EditorPrimitiveDepth, ERenderTargetLoadAction::ELoad, ERenderTargetLoadAction::ELoad, FExclusiveDepthStencil::DepthWrite_StencilWrite);
 
-			const FEditorPrimitiveInputs::EBasePassType BasePassType = Inputs.BasePassType;
+			const FCompositePrimitiveInputs::EBasePassType BasePassType = Inputs.BasePassType;
 
-			if (BasePassType == FEditorPrimitiveInputs::EBasePassType::Deferred)
+			if (BasePassType == FCompositePrimitiveInputs::EBasePassType::Deferred)
 			{
 				PassParameters->BasePass = CreateOpaqueBasePassUniformBuffer(GraphBuilder, *EditorView, 0);
 			}
@@ -520,10 +317,10 @@ FScreenPassTexture AddEditorPrimitivePass(
 				RDG_EVENT_NAME("EditorPrimitives %dx%d MSAA=%d",
 					EditorPrimitivesViewport.Rect.Width(),
 					EditorPrimitivesViewport.Rect.Height(),
-					NumSamples),
+					NumMSAASamples),
 				PassParameters,
 				ERDGPassFlags::Raster,
-				[&View, &InstanceCullingManager, PassParameters, EditorView, EditorPrimitivesViewport, BasePassType, NumSamples](FRHICommandList& RHICmdList)
+				[&View, &InstanceCullingManager, PassParameters, EditorView, EditorPrimitivesViewport, BasePassType, NumMSAASamples](FRHICommandList& RHICmdList)
 			{
 				RHICmdList.SetViewport(EditorPrimitivesViewport.Rect.Min.X, EditorPrimitivesViewport.Rect.Min.Y, 0.0f, EditorPrimitivesViewport.Rect.Max.X, EditorPrimitivesViewport.Rect.Max.Y, 1.0f);
 
@@ -535,7 +332,7 @@ FScreenPassTexture AddEditorPrimitivePass(
 				{
 					SCOPED_DRAW_EVENTF(RHICmdList, EditorPrimitives,
 						TEXT("RenderViewEditorPrimitives %dx%d msaa=%d"),
-						EditorPrimitivesViewport.Rect.Width(), EditorPrimitivesViewport.Rect.Height(), NumSamples);
+						EditorPrimitivesViewport.Rect.Width(), EditorPrimitivesViewport.Rect.Height(), NumMSAASamples);
 
 					RenderEditorPrimitives(RHICmdList, *EditorView, DrawRenderState, InstanceCullingManager);
 				}
@@ -544,7 +341,7 @@ FScreenPassTexture AddEditorPrimitivePass(
 				{
 					SCOPED_DRAW_EVENTF(RHICmdList, EditorPrimitives,
 						TEXT("RenderViewEditorForegroundPrimitives %dx%d msaa=%d"),
-						EditorPrimitivesViewport.Rect.Width(), EditorPrimitivesViewport.Rect.Height(), NumSamples);
+						EditorPrimitivesViewport.Rect.Width(), EditorPrimitivesViewport.Rect.Height(), NumMSAASamples);
 
 					RenderForegroundEditorPrimitives(RHICmdList, *EditorView, DrawRenderState, InstanceCullingManager);
 				}
@@ -580,20 +377,20 @@ FScreenPassTexture AddEditorPrimitivePass(
 		PassParameters->bCompositeAnyNonNullDepth = bProducedByPriorPass;
 		PassParameters->DepthTextureJitter = SceneDepthJitter;
 
-		for (int32 i = 0; i < int32(NumSamples); i++)
+		for (int32 i = 0; i < int32(NumMSAASamples); i++)
 		{
-			PassParameters->SampleOffsetArray[i].X = GetMSAASampleOffsets(NumSamples, i).X;
-			PassParameters->SampleOffsetArray[i].Y = GetMSAASampleOffsets(NumSamples, i).Y;
+			PassParameters->SampleOffsetArray[i].X = GetMSAASampleOffsets(NumMSAASamples, i).X;
+			PassParameters->SampleOffsetArray[i].Y = GetMSAASampleOffsets(NumMSAASamples, i).Y;
 		}
 
 		FCompositeEditorPrimitivesPS::FPermutationDomain PermutationVector;
-		PermutationVector.Set<FCompositeEditorPrimitivesPS::FSampleCountDimension>(NumSamples);
+		PermutationVector.Set<FCompositeEditorPrimitivesPS::FSampleCountDimension>(NumMSAASamples);
 
 		TShaderMapRef<FCompositeEditorPrimitivesPS> PixelShader(View.ShaderMap, PermutationVector);
 		FPixelShaderUtils::AddFullscreenPass(
 			GraphBuilder,
 			View.ShaderMap,
-			RDG_EVENT_NAME("Composite %dx%d MSAA=%d", Output.ViewRect.Width(), Output.ViewRect.Height(), NumSamples),
+			RDG_EVENT_NAME("Composite %dx%d MSAA=%d", Output.ViewRect.Width(), Output.ViewRect.Height(), NumMSAASamples),
 			PixelShader,
 			PassParameters,
 			Output.ViewRect);
@@ -611,9 +408,9 @@ FScreenPassTexture AddEditorPrimitivePass(
 		PassParameters->RenderTargets[0] = Output.GetRenderTargetBinding();
 		PassParameters->RenderTargets[0].SetLoadAction(ERenderTargetLoadAction::ELoad);
 
-		const FEditorPrimitiveInputs::EBasePassType BasePassType = Inputs.BasePassType;
+		const FCompositePrimitiveInputs::EBasePassType BasePassType = Inputs.BasePassType;
 
-		if (BasePassType == FEditorPrimitiveInputs::EBasePassType::Deferred)
+		if (BasePassType == FCompositePrimitiveInputs::EBasePassType::Deferred)
 		{
 			PassParameters->TranslucentBasePass = CreateTranslucentBasePassUniformBuffer(GraphBuilder, nullptr, *EditorView, 0);
 		}
@@ -626,10 +423,10 @@ FScreenPassTexture AddEditorPrimitivePass(
 			RDG_EVENT_NAME("EditorPrimitives Translucent %dx%d MSAA=%d",
 				EditorPrimitivesViewport.Rect.Width(),
 				EditorPrimitivesViewport.Rect.Height(),
-				NumSamples),
+				NumMSAASamples),
 			PassParameters,
 			ERDGPassFlags::Raster,
-			[&View, &InstanceCullingManager, PassParameters, EditorView, EditorPrimitivesViewport, OutputViewportRect = Output.ViewRect, BasePassType, NumSamples](FRHICommandListImmediate& RHICmdList)
+			[&View, &InstanceCullingManager, PassParameters, EditorView, EditorPrimitivesViewport, OutputViewportRect = Output.ViewRect, BasePassType, NumMSAASamples](FRHICommandListImmediate& RHICmdList)
 			{
 				RHICmdList.SetViewport(OutputViewportRect.Min.X, OutputViewportRect.Min.Y, 0.0f, OutputViewportRect.Max.X, OutputViewportRect.Max.Y, 1.0f);
 
@@ -641,7 +438,7 @@ FScreenPassTexture AddEditorPrimitivePass(
 				{
 					SCOPED_DRAW_EVENTF(RHICmdList, EditorPrimitives,
 						TEXT("RenderViewEditorForegroundTranslucentPrimitives %dx%d msaa=%d"),
-						EditorPrimitivesViewport.Rect.Width(), EditorPrimitivesViewport.Rect.Height(), NumSamples);
+						EditorPrimitivesViewport.Rect.Width(), EditorPrimitivesViewport.Rect.Height(), NumMSAASamples);
 
 					RenderForegroundTranslucentEditorPrimitives(RHICmdList, *EditorView, DrawRenderState, InstanceCullingManager);
 				}
diff --git a/Engine/Source/Runtime/Renderer/Private/PostProcess/PostProcessCompositeEditorPrimitives.h b/Engine/Source/Runtime/Renderer/Private/PostProcess/PostProcessCompositeEditorPrimitives.h
--- a/Engine/Source/Runtime/Renderer/Private/PostProcess/PostProcessCompositeEditorPrimitives.h
+++ b/Engine/Source/Runtime/Renderer/Private/PostProcess/PostProcessCompositeEditorPrimitives.h
@@ -1,73 +1,8 @@
 // Copyright Epic Games, Inc. All Rights Reserved.
 
 #pragma once
-
-#include "ScreenPass.h"
-#include "OverridePassSequence.h"
-#include "DataDrivenShaderPlatformInfo.h"
+#include "PostProcess/PostProcessCompositePrimitivesCommon.h"
 
 #if WITH_EDITOR
-
-/** Base class for a global pixel shader which renders editor primitives (outlines, helpers, etc). */
-class FEditorPrimitiveShader : public FGlobalShader
-{
-public:
-	static const uint32 kEditorMSAASampleCountMax = 8;
-
-	class FSampleCountDimension : SHADER_PERMUTATION_RANGE_INT("MSAA_SAMPLE_COUNT", 1, kEditorMSAASampleCountMax + 1);
-	using FPermutationDomain = TShaderPermutationDomain<FSampleCountDimension>;
-
-	static bool ShouldCompilePermutation(const FPermutationDomain& PermutationVector, const EShaderPlatform Platform)
-	{
-		const int32 SampleCount = PermutationVector.Get<FSampleCountDimension>();
-
-		// Only use permutations with valid MSAA sample counts.
-		if (!FMath::IsPowerOfTwo(SampleCount))
-		{
-			return false;
-		}
-
-		// Only PC platforms render editor primitives.
-		return IsPCPlatform(Platform);
-	}
-
-	static bool ShouldCompilePermutation(const FGlobalShaderPermutationParameters& Parameters)
-	{
-		const FPermutationDomain PermutationVector(Parameters.PermutationId);
-		return ShouldCompilePermutation(PermutationVector, Parameters.Platform);
-	}
-
-	FEditorPrimitiveShader() = default;
-	FEditorPrimitiveShader(const ShaderMetaType::CompiledShaderInitializerType& Initializer)
-		: FGlobalShader(Initializer)
-	{}
-};
-
-// Constructs a new view suitable for rendering editor primitives (outlines, helpers, etc).
-const FViewInfo* CreateEditorPrimitiveView(const FViewInfo& ParentView, FIntRect ViewRect, uint32 NumSamples);
-
-struct FEditorPrimitiveInputs
-{
-	enum class EBasePassType : uint32
-	{
-		Deferred,
-		Mobile,
-		MAX
-	};
-
-	// [Optional] Render to the specified output. If invalid, a new texture is created and returned.
-	FScreenPassRenderTarget OverrideOutput;
-
-	// [Required] The scene color to composite with editor primitives.
-	FScreenPassTexture SceneColor;
-
-	// [Required] The scene depth to composite with editor primitives.
-	FScreenPassTexture SceneDepth;
-
-	// [Required] The type of base pass to use for rendering editor primitives.
-	EBasePassType BasePassType = EBasePassType::MAX;
-};
-
-FScreenPassTexture AddEditorPrimitivePass(FRDGBuilder& GraphBuilder, const FViewInfo& View, const FEditorPrimitiveInputs& Inputs, FInstanceCullingManager& InstanceCullingManager);
-
+FScreenPassTexture AddEditorPrimitivePass(FRDGBuilder& GraphBuilder, const FViewInfo& View, const FCompositePrimitiveInputs& Inputs, FInstanceCullingManager& InstanceCullingManager);
 #endif
diff --git a/Engine/Source/Runtime/Renderer/Private/PostProcess/PostProcessCompositePrimitivesCommon.cpp b/Engine/Source/Runtime/Renderer/Private/PostProcess/PostProcessCompositePrimitivesCommon.cpp
new file mode 100644
--- /dev/null
+++ b/Engine/Source/Runtime/Renderer/Private/PostProcess/PostProcessCompositePrimitivesCommon.cpp
@@ -0,0 +1,277 @@
+// Copyright Epic Games, Inc. All Rights Reserved.
+
+#include "PostProcess/PostProcessCompositePrimitivesCommon.h"
+
+#if UE_ENABLE_DEBUG_DRAWING 
+#include "BasePassRendering.h"
+#include "PixelShaderUtils.h"
+#include "ScenePrivate.h"
+
+TAutoConsoleVariable<int32> CVarCompositeTemporalUpsampleDepth(
+	TEXT("r.Composite.TemporalUpsampleDepth"), 2,
+	TEXT("Temporal upsample factor of the depth buffer for depth testing editor primitives against."),
+	ECVF_RenderThreadSafe);
+
+const FViewInfo* CreateCompositePrimitiveView(const FViewInfo& ParentView, FIntRect ViewRect, uint32 NumMSAASamples)
+{
+	FViewInfo* View = ParentView.CreateSnapshot();
+
+	// Patch view rect.
+	View->ViewRect = ViewRect;
+
+	// Override pre exposure to 1.0f, because rendering after tonemapper. 
+	View->PreExposure = 1.0f;
+
+	// Kills material texture mipbias because after TAA.
+	View->MaterialTextureMipBias = 0.0f;
+
+	if (IsTemporalAccumulationBasedMethod(View->AntiAliasingMethod))
+	{
+		View->ViewMatrices.HackRemoveTemporalAAProjectionJitter();
+	}
+
+	View->InitRHIResources(NumMSAASamples);
+
+	return View;
+}
+
+FRDGTextureRef CreateCompositeDepthTexture(FRDGBuilder& GraphBuilder, FIntPoint Extent, uint32 NumMSAASamples)
+{
+	const FRDGTextureDesc DepthDesc = FRDGTextureDesc::Create2D(
+		Extent,
+		PF_DepthStencil,
+		FClearValueBinding::DepthFar,
+		TexCreate_ShaderResource | TexCreate_DepthStencilTargetable,
+		1,
+		NumMSAASamples);
+
+	return GraphBuilder.CreateTexture(DepthDesc, TEXT("Composite.PrimitivesDepth"));
+}
+
+class FTemporalUpsampleDepthPS : public FGlobalShader
+{
+public:
+	DECLARE_GLOBAL_SHADER(FTemporalUpsampleDepthPS);
+	SHADER_USE_PARAMETER_STRUCT(FTemporalUpsampleDepthPS, FGlobalShader);
+
+	BEGIN_SHADER_PARAMETER_STRUCT(FParameters, )
+		SHADER_PARAMETER_STRUCT_REF(FViewUniformShaderParameters, View)
+		SHADER_PARAMETER_STRUCT(FScreenPassTextureViewportParameters, Depth)
+		SHADER_PARAMETER_STRUCT(FScreenPassTextureViewportParameters, PrevHistory)
+		SHADER_PARAMETER_STRUCT(FScreenPassTextureViewportParameters, History)
+		SHADER_PARAMETER(FVector2f, DepthTextureJitter)
+		SHADER_PARAMETER(int32, bCameraCut)
+
+		SHADER_PARAMETER_RDG_TEXTURE(Texture2D, DepthTexture)
+		SHADER_PARAMETER_SAMPLER(SamplerState, DepthSampler)
+		SHADER_PARAMETER_RDG_TEXTURE(Texture2D, VelocityTexture)
+		SHADER_PARAMETER_SAMPLER(SamplerState, VelocitySampler)
+		SHADER_PARAMETER_RDG_TEXTURE(Texture2D, PrevHistoryTexture)
+		SHADER_PARAMETER_SAMPLER(SamplerState, PrevHistorySampler)
+
+		RENDER_TARGET_BINDING_SLOTS()
+		END_SHADER_PARAMETER_STRUCT()
+
+		static bool ShouldCompilePermutation(const FGlobalShaderPermutationParameters& Parameters)
+	{
+#if WITH_EDITOR
+		// Only PC platforms render editor primitives.
+		return IsPCPlatform(Parameters.Platform);
+#else
+		return true;
+#endif
+	}
+};
+
+class FPopulateCompositeDepthPS : public FGlobalShader
+{
+public:
+	DECLARE_GLOBAL_SHADER(FPopulateCompositeDepthPS);
+	SHADER_USE_PARAMETER_STRUCT(FPopulateCompositeDepthPS, FGlobalShader);
+
+	class FUseMSAADimension : SHADER_PERMUTATION_BOOL("USE_MSAA");
+	class FForceDrawColorOutput : SHADER_PERMUTATION_BOOL("FORCE_DRAW_COLOR"); //Allow the option to draw out the scene color too to lower draw calls
+	using FPermutationDomain = TShaderPermutationDomain<FUseMSAADimension, FForceDrawColorOutput>;
+
+	BEGIN_SHADER_PARAMETER_STRUCT(FParameters, )
+		SHADER_PARAMETER_STRUCT_REF(FViewUniformShaderParameters, View)
+		SHADER_PARAMETER_STRUCT(FScreenPassTextureViewportParameters, Color)
+		SHADER_PARAMETER_STRUCT(FScreenPassTextureViewportParameters, Depth)
+		SHADER_PARAMETER_RDG_TEXTURE(Texture2D, ColorTexture)
+		SHADER_PARAMETER_SAMPLER(SamplerState, ColorSampler)
+		SHADER_PARAMETER_RDG_TEXTURE(Texture2D, DepthTexture)
+		SHADER_PARAMETER_SAMPLER(SamplerState, DepthSampler)
+		SHADER_PARAMETER(FVector2f, DepthTextureJitter)
+		RENDER_TARGET_BINDING_SLOTS()
+		END_SHADER_PARAMETER_STRUCT()
+
+		static bool ShouldCompilePermutation(const FGlobalShaderPermutationParameters& Parameters)
+	{
+		const FPermutationDomain PermutationVector(Parameters.PermutationId);
+		const bool bUseMSAA = PermutationVector.Get<FUseMSAADimension>();
+
+		// Only SM5+ platforms supports MSAA.
+		if (!IsFeatureLevelSupported(Parameters.Platform, ERHIFeatureLevel::SM5) && bUseMSAA)
+		{
+			return false;
+		}
+
+#if WITH_EDITOR
+		// Only PC platforms render editor primitives.
+		return IsPCPlatform(Parameters.Platform);
+#else
+		return true;
+#endif
+	}
+};
+
+IMPLEMENT_GLOBAL_SHADER(FTemporalUpsampleDepthPS, "/Engine/Private/PostProcessCompositePrimitives.usf", "MainTemporalUpsampleDepthPS", SF_Pixel);
+IMPLEMENT_GLOBAL_SHADER(FPopulateCompositeDepthPS, "/Engine/Private/PostProcessCompositePrimitives.usf", "MainPopulateSceneDepthPS", SF_Pixel);
+
+void TemporalUpscaleDepthPass(
+	FRDGBuilder& GraphBuilder,
+	const FViewInfo& View,
+	const FScreenPassTexture& InSceneColor,
+	FScreenPassTexture& InOutSceneDepth,
+	FVector2f& SceneDepthJitter,
+	uint32 NumMSAASamples)
+{
+	if (InSceneColor.ViewRect != InOutSceneDepth.ViewRect)
+	{
+		FIntPoint Extent = InSceneColor.Texture->Desc.Extent;
+		FIntRect ViewRect = InSceneColor.ViewRect;
+
+		// Upscale factor of the depth buffer that might be needed.
+		const float UpscaleFactor = float(View.ViewRect.Width()) / float(InOutSceneDepth.ViewRect.Width());
+
+		// Upscale factor shouldn't be higher than there is TAA samples, or that means there will be unrendered pixels.
+		const int32 ComputeMaxUpsampleFactorDueToTAA = FMath::FloorToInt(FMath::Sqrt(float(View.TemporalJitterSequenceLength) / (UpscaleFactor * UpscaleFactor)));
+
+		const int32 ComputeMaxUpsampleFactor = FMath::Clamp(ComputeMaxUpsampleFactorDueToTAA, 0, 4);
+
+		const int32 DepthUpsampleFactor = FMath::Clamp(CVarCompositeTemporalUpsampleDepth.GetValueOnRenderThread(), 0, ComputeMaxUpsampleFactor);
+
+		// Upsample the depth at higher resolution to reduce depth intersection instability of editor primitives.
+		if (DepthUpsampleFactor > 0 && View.ViewState && IsTemporalAccumulationBasedMethod(View.AntiAliasingMethod))
+		{
+			FScreenPassTexture History;
+			const FRDGTextureDesc Desc = FRDGTextureDesc::Create2D(
+				Extent * DepthUpsampleFactor,
+				PF_R32_FLOAT,
+				FClearValueBinding::None,
+				TexCreate_ShaderResource | TexCreate_RenderTargetable);
+
+			History.Texture = GraphBuilder.CreateTexture(Desc, TEXT("Composite.PrimitivesDepthHistory"));
+			History.ViewRect = ViewRect * DepthUpsampleFactor;
+
+
+			FTemporalUpsampleDepthPS::FParameters* PassParameters = GraphBuilder.AllocParameters<FTemporalUpsampleDepthPS::FParameters>();
+			PassParameters->View = View.ViewUniformBuffer;
+			PassParameters->Depth = GetScreenPassTextureViewportParameters(FScreenPassTextureViewport(InOutSceneDepth));
+			PassParameters->History = GetScreenPassTextureViewportParameters(FScreenPassTextureViewport(History));
+			PassParameters->DepthTextureJitter = SceneDepthJitter;
+
+			PassParameters->DepthTexture = InOutSceneDepth.Texture;
+			PassParameters->DepthSampler = TStaticSamplerState<SF_Point, AM_Clamp, AM_Clamp, AM_Clamp>::GetRHI();
+
+			// TODO
+			PassParameters->VelocityTexture = GraphBuilder.RegisterExternalTexture(GSystemTextures.BlackDummy);
+			PassParameters->VelocitySampler = TStaticSamplerState<SF_Point, AM_Clamp, AM_Clamp, AM_Clamp>::GetRHI();
+
+			if (View.PrevViewInfo.CompositePrimitiveDepthHistory.IsValid())
+			{
+				PassParameters->PrevHistory = GetScreenPassTextureViewportParameters(FScreenPassTextureViewport(
+					View.PrevViewInfo.CompositePrimitiveDepthHistory.RT[0]->GetDesc().Extent, View.PrevViewInfo.CompositePrimitiveDepthHistory.ViewportRect));
+				PassParameters->PrevHistoryTexture = GraphBuilder.RegisterExternalTexture(View.PrevViewInfo.CompositePrimitiveDepthHistory.RT[0]);
+				PassParameters->PrevHistorySampler = TStaticSamplerState<SF_Bilinear, AM_Clamp, AM_Clamp, AM_Clamp>::GetRHI();
+				PassParameters->bCameraCut = false;
+			}
+			else
+			{
+				PassParameters->PrevHistory = GetScreenPassTextureViewportParameters(FScreenPassTextureViewport(
+					FIntPoint(1, 1), FIntRect(FIntPoint(0, 0), FIntPoint(1, 1))));
+				PassParameters->PrevHistoryTexture = GraphBuilder.RegisterExternalTexture(GSystemTextures.BlackDummy);
+				PassParameters->PrevHistorySampler = TStaticSamplerState<SF_Point, AM_Clamp, AM_Clamp, AM_Clamp>::GetRHI();
+				PassParameters->bCameraCut = true;
+			}
+
+			PassParameters->RenderTargets[0] = FRenderTargetBinding(History.Texture, ERenderTargetLoadAction::ENoAction);
+
+			TShaderMapRef<FTemporalUpsampleDepthPS> PixelShader(View.ShaderMap);
+			FPixelShaderUtils::AddFullscreenPass(
+				GraphBuilder,
+				View.ShaderMap,
+				RDG_EVENT_NAME("TemporalUpsampleDepth %dx%d -> %dx%d",
+					InOutSceneDepth.ViewRect.Width(),
+					InOutSceneDepth.ViewRect.Height(),
+					History.ViewRect.Width(),
+					History.ViewRect.Height()),
+				PixelShader,
+				PassParameters,
+				History.ViewRect);
+
+			if (!View.bStatePrevViewInfoIsReadOnly)
+			{
+				FTemporalAAHistory* OutputHistory = &View.ViewState->PrevFrameViewInfo.CompositePrimitiveDepthHistory;
+				OutputHistory->SafeRelease();
+
+				GraphBuilder.QueueTextureExtraction(History.Texture, &OutputHistory->RT[0]);
+				OutputHistory->ViewportRect = History.ViewRect;
+				OutputHistory->ReferenceBufferSize = History.Texture->Desc.Extent;
+			}
+
+			InOutSceneDepth = History;
+			SceneDepthJitter = FVector2f::ZeroVector;
+		}
+	}
+}
+
+// Populate the MSAA depth buffer from depth buffer or temporally upscaled depth buffer
+void PopulateDepthPass(FRDGBuilder& GraphBuilder,
+	const FViewInfo& View,
+	const FScreenPassTexture& InSceneColor,
+	const FScreenPassTexture& InSceneDepth,
+	FRDGTextureRef OutPopColor,
+	FRDGTextureRef OutPopDepth,
+	const FVector2f& SceneDepthJitter,
+	uint32 NumMSAASamples,
+	bool bForceDrawColor)
+{
+	FRHISamplerState* PointClampSampler = TStaticSamplerState<SF_Point, AM_Clamp, AM_Clamp, AM_Clamp>::GetRHI();
+
+	FPopulateCompositeDepthPS::FParameters* PassParameters = GraphBuilder.AllocParameters<FPopulateCompositeDepthPS::FParameters>();
+	PassParameters->View = View.ViewUniformBuffer;
+	PassParameters->Color = GetScreenPassTextureViewportParameters(FScreenPassTextureViewport(InSceneColor));
+	PassParameters->Depth = GetScreenPassTextureViewportParameters(FScreenPassTextureViewport(InSceneDepth));
+	if (bForceDrawColor)
+	{
+		PassParameters->ColorTexture = InSceneColor.Texture;
+		PassParameters->ColorSampler = PointClampSampler;
+	}
+	PassParameters->DepthTexture = InSceneDepth.Texture;
+	PassParameters->DepthSampler = PointClampSampler;
+	PassParameters->DepthTextureJitter = SceneDepthJitter;
+	PassParameters->RenderTargets[0] = FRenderTargetBinding(OutPopColor, ERenderTargetLoadAction::EClear);
+	PassParameters->RenderTargets.DepthStencil = FDepthStencilBinding(OutPopDepth, ERenderTargetLoadAction::EClear, ERenderTargetLoadAction::EClear, FExclusiveDepthStencil::DepthWrite_StencilWrite);
+
+	FPopulateCompositeDepthPS::FPermutationDomain PermutationVector;
+	PermutationVector.Set<FPopulateCompositeDepthPS::FUseMSAADimension>(NumMSAASamples > 1);
+	PermutationVector.Set< FPopulateCompositeDepthPS::FForceDrawColorOutput>(bForceDrawColor);
+	TShaderMapRef<FPopulateCompositeDepthPS> PixelShader(View.ShaderMap, PermutationVector);
+
+	FPixelShaderUtils::AddFullscreenPass(
+		GraphBuilder,
+		View.ShaderMap,
+		RDG_EVENT_NAME("PopulateDepth %dx%d%s",
+			View.ViewRect.Width(),
+			View.ViewRect.Height(),
+			NumMSAASamples > 1 ? TEXT(" MSAA") : TEXT("")),
+		PixelShader,
+		PassParameters,
+		View.ViewRect,
+		/* BlendState = */ nullptr,
+		/* RasterizerState = */ nullptr,
+		TStaticDepthStencilState<true, CF_Always>::GetRHI());
+}
+
+#endif //UE_ENABLE_DEBUG_DRAWING
\ No newline at end of file
diff --git a/Engine/Source/Runtime/Renderer/Private/PostProcess/PostProcessCompositePrimitivesCommon.h b/Engine/Source/Runtime/Renderer/Private/PostProcess/PostProcessCompositePrimitivesCommon.h
new file mode 100644
--- /dev/null
+++ b/Engine/Source/Runtime/Renderer/Private/PostProcess/PostProcessCompositePrimitivesCommon.h
@@ -0,0 +1,98 @@
+// Copyright Epic Games, Inc. All Rights Reserved.
+
+#pragma once
+
+#include "ScreenPass.h"
+#include "OverridePassSequence.h"
+#include "DataDrivenShaderPlatformInfo.h"
+
+//UE_ENABLE_DEBUG_DRAWING i.e. !(UE_BUILD_SHIPPING || UE_BUILD_TEST) || WITH_EDITOR
+//Only available in Debug/Development/Editor builds due to current use cases, but can be extended in future
+#if UE_ENABLE_DEBUG_DRAWING
+
+/** Base class for a global pixel shader which renders primitives (outlines, helpers, etc). */
+class FCompositePrimitiveShaderBase : public FGlobalShader
+{
+public:
+	static const uint32 kMSAASampleCountMax = 8;
+
+	class FSampleCountDimension : SHADER_PERMUTATION_RANGE_INT("MSAA_SAMPLE_COUNT", 1, kMSAASampleCountMax + 1);
+	using FPermutationDomain = TShaderPermutationDomain<FSampleCountDimension>;
+
+	static bool ShouldCompilePermutation(const FPermutationDomain& PermutationVector, const EShaderPlatform Platform)
+	{
+		const int32 SampleCount = PermutationVector.Get<FSampleCountDimension>();
+
+		// Only use permutations with valid MSAA sample counts.
+		if (!FMath::IsPowerOfTwo(SampleCount))
+		{
+			return false;
+		}
+
+#if WITH_EDITOR
+		// Only PC platforms render editor primitives.
+		return IsPCPlatform(Platform);
+#else
+		//Not currently used for any other passes but place 
+		return false;
+#endif
+	}
+
+	static bool ShouldCompilePermutation(const FGlobalShaderPermutationParameters& Parameters)
+	{
+		const FPermutationDomain PermutationVector(Parameters.PermutationId);
+		return ShouldCompilePermutation(PermutationVector, Parameters.Platform);
+	}
+
+	FCompositePrimitiveShaderBase() = default;
+	FCompositePrimitiveShaderBase(const ShaderMetaType::CompiledShaderInitializerType& Initializer)
+		: FGlobalShader(Initializer)
+	{}
+};
+
+struct FCompositePrimitiveInputs
+{
+	enum class EBasePassType : uint32
+	{
+		Deferred,
+		Mobile,
+		MAX
+	};
+	// [Required] The type of base pass to use for rendering editor primitives.
+	EBasePassType BasePassType = EBasePassType::MAX;
+
+	// [Optional] Render to the specified output. If invalid, a new texture is created and returned.
+	FScreenPassRenderTarget OverrideOutput;
+
+	// [Required] The scene color to composite with editor primitives.
+	FScreenPassTexture SceneColor;
+
+	// [Required] The scene depth to composite with editor primitives.
+	FScreenPassTexture SceneDepth;
+};
+
+// Constructs a new view suitable for rendering debug primitives.
+const FViewInfo* CreateCompositePrimitiveView(const FViewInfo& ParentView, FIntRect ViewRect, uint32 NumMSAASamples);
+
+FRDGTextureRef CreateCompositeDepthTexture(FRDGBuilder& GraphBuilder, FIntPoint Extent, uint32 NumMSAASamples);
+
+
+void TemporalUpscaleDepthPass(
+	FRDGBuilder& GraphBuilder,
+	const FViewInfo& View,
+	const FScreenPassTexture& InSceneColor,
+	FScreenPassTexture& InOutSceneDepth,
+	FVector2f& SceneDepthJitter,
+	uint32 NumMSAASamples);
+
+void PopulateDepthPass(FRDGBuilder& GraphBuilder,
+	const FViewInfo& View,
+	const FScreenPassTexture& InSceneColor, 
+	const FScreenPassTexture& InSceneDepth, 
+	FRDGTextureRef OutPopColor, 
+	FRDGTextureRef OutPopDepth, 
+	const FVector2f& SceneDepthJitter,
+	uint32 NumMSAASamples,
+	bool bForceDrawColor = false);
+
+#endif //#if UE_ENABLE_DEBUG_DRAWING
\ No newline at end of file
diff --git a/Engine/Source/Runtime/Renderer/Private/PostProcess/PostProcessSelectionOutline.cpp b/Engine/Source/Runtime/Renderer/Private/PostProcess/PostProcessSelectionOutline.cpp
--- a/Engine/Source/Runtime/Renderer/Private/PostProcess/PostProcessSelectionOutline.cpp
+++ b/Engine/Source/Runtime/Renderer/Private/PostProcess/PostProcessSelectionOutline.cpp
@@ -14,13 +14,13 @@
 
 namespace
 {
-class FSelectionOutlinePS : public FEditorPrimitiveShader
+class FSelectionOutlinePS : public FCompositePrimitiveShaderBase
 {
 public:
 	DECLARE_GLOBAL_SHADER(FSelectionOutlinePS);
-	SHADER_USE_PARAMETER_STRUCT(FSelectionOutlinePS, FEditorPrimitiveShader);
+	SHADER_USE_PARAMETER_STRUCT(FSelectionOutlinePS, FCompositePrimitiveShaderBase);
 	class FSelectionOutlineHDRDim : SHADER_PERMUTATION_BOOL("SELECTION_OUTLINE_HDR");
-	using FPermutationDomain = TShaderPermutationDomain< FEditorPrimitiveShader::FPermutationDomain, FSelectionOutlineHDRDim>;
+	using FPermutationDomain = TShaderPermutationDomain< FCompositePrimitiveShaderBase::FPermutationDomain, FSelectionOutlineHDRDim>;
 
 	BEGIN_SHADER_PARAMETER_STRUCT(FParameters, )
 		SHADER_PARAMETER_STRUCT_REF(FViewUniformShaderParameters, View)
@@ -49,13 +49,13 @@ public:
 	static bool ShouldCompilePermutation(const FGlobalShaderPermutationParameters& Parameters)
 	{
 		FPermutationDomain CurrentPermutation(Parameters.PermutationId);
-		FEditorPrimitiveShader::FPermutationDomain EditorPrimitiveShaderDomain = CurrentPermutation.Get<FEditorPrimitiveShader::FPermutationDomain>();
+		FCompositePrimitiveShaderBase::FPermutationDomain EditorPrimitiveShaderDomain = CurrentPermutation.Get<FCompositePrimitiveShaderBase::FPermutationDomain>();
 		bool bIsHDR = CurrentPermutation.Get<FSelectionOutlineHDRDim>();
 		if (bIsHDR && !SupportHDR(Parameters.Platform))
 		{
 			return false;
 		}
-		return FEditorPrimitiveShader::ShouldCompilePermutation(EditorPrimitiveShaderDomain, Parameters.Platform);
+		return FCompositePrimitiveShaderBase::ShouldCompilePermutation(EditorPrimitiveShaderDomain, Parameters.Platform);
 	}
 };
 
@@ -83,7 +83,7 @@ FScreenPassTexture AddSelectionOutlinePass(FRDGBuilder& GraphBuilder, const FVie
 	const uint32 NumSamples = View.GetSceneTexturesConfig().NumSamples;
 
 	// Patch uniform buffers with updated state for rendering the outline mesh draw commands.
-	const FViewInfo* EditorView = CreateEditorPrimitiveView(View, Inputs.SceneColor.ViewRect, NumSamples);
+	const FViewInfo* EditorView = CreateCompositePrimitiveView(View, Inputs.SceneColor.ViewRect, NumSamples);
 
 	FRDGTextureRef DepthStencilTexture = nullptr;
 
@@ -242,9 +242,9 @@ FScreenPassTexture AddSelectionOutlinePass(FRDGBuilder& GraphBuilder, const FVie
 		PassParameters->UILuminanceAndIsSCRGB = bIsSCRGB ? UILuminance : -UILuminance;
 
 		FSelectionOutlinePS::FPermutationDomain PermutationVector;
-		FEditorPrimitiveShader::FPermutationDomain PermutationVectorBase;
+		FCompositePrimitiveShaderBase::FPermutationDomain PermutationVectorBase;
 		PermutationVectorBase.Set<FSelectionOutlinePS::FSampleCountDimension>(NumSamples);
-		PermutationVector.Set<FEditorPrimitiveShader::FPermutationDomain>(PermutationVectorBase);
+		PermutationVector.Set<FCompositePrimitiveShaderBase::FPermutationDomain>(PermutationVectorBase);
 		PermutationVector.Set<FSelectionOutlinePS::FSelectionOutlineHDRDim>(bIsHDR);
 
 		TShaderMapRef<FSelectionOutlinePS> PixelShader(View.ShaderMap, PermutationVector);
diff --git a/Engine/Source/Runtime/Renderer/Private/PostProcess/PostProcessVisualizeLevelInstance.cpp b/Engine/Source/Runtime/Renderer/Private/PostProcess/PostProcessVisualizeLevelInstance.cpp
--- a/Engine/Source/Runtime/Renderer/Private/PostProcess/PostProcessVisualizeLevelInstance.cpp
+++ b/Engine/Source/Runtime/Renderer/Private/PostProcess/PostProcessVisualizeLevelInstance.cpp
@@ -3,7 +3,7 @@
 #if WITH_EDITOR
 
 #include "PostProcess/PostProcessVisualizeLevelInstance.h"
-#include "PostProcess/PostProcessCompositeEditorPrimitives.h"
+#include "PostProcess/PostProcessCompositePrimitivesCommon.h"
 #include "SceneTextureParameters.h"
 #include "CanvasTypes.h"
 #include "RenderTargetTemp.h"
@@ -12,11 +12,11 @@
 
 namespace
 {
-class FVisualizeLevelInstancePS : public FEditorPrimitiveShader
+class FVisualizeLevelInstancePS : public FCompositePrimitiveShaderBase
 {
 public:
 	DECLARE_GLOBAL_SHADER(FVisualizeLevelInstancePS);
-	SHADER_USE_PARAMETER_STRUCT(FVisualizeLevelInstancePS, FEditorPrimitiveShader);
+	SHADER_USE_PARAMETER_STRUCT(FVisualizeLevelInstancePS, FCompositePrimitiveShaderBase);
 
 	BEGIN_SHADER_PARAMETER_STRUCT(FParameters, )
 		SHADER_PARAMETER_STRUCT_REF(FViewUniformShaderParameters, View)
@@ -56,7 +56,7 @@ FScreenPassTexture AddVisualizeLevelInstancePass(FRDGBuilder& GraphBuilder, cons
 	const uint32 NumSamples = View.GetSceneTexturesConfig().NumSamples;
 
 	// Patch uniform buffers with updated state for rendering the outline mesh draw commands.
-	const FViewInfo* EditorView = CreateEditorPrimitiveView(View, Inputs.SceneColor.ViewRect, NumSamples);
+	const FViewInfo* EditorView = CreateCompositePrimitiveView(View, Inputs.SceneColor.ViewRect, NumSamples);
 
 	FRDGTextureRef DepthStencilTexture = nullptr;
 
diff --git a/Engine/Source/Runtime/Renderer/Private/PostProcess/PostProcessing.cpp b/Engine/Source/Runtime/Renderer/Private/PostProcess/PostProcessing.cpp
--- a/Engine/Source/Runtime/Renderer/Private/PostProcess/PostProcessing.cpp
+++ b/Engine/Source/Runtime/Renderer/Private/PostProcess/PostProcessing.cpp
@@ -33,7 +33,12 @@
 #include "PostProcess/PostProcessUpscale.h"
 #include "PostProcess/PostProcessHMD.h"
 #include "PostProcess/PostProcessVisualizeComplexity.h"
-#include "PostProcess/PostProcessCompositeEditorPrimitives.h"
+#if UE_ENABLE_DEBUG_DRAWING
+	#include "PostProcess/PostProcessCompositeDebugPrimitives.h"
+#endif
+#if WITH_EDITOR
+	#include "PostProcess/PostProcessCompositeEditorPrimitives.h"
+#endif
 #include "PostProcess/PostProcessTestImage.h"
 #include "PostProcess/PostProcessVisualizeCalibrationMaterial.h"
 #include "PostProcess/PostProcessFFTBloom.h"
@@ -42,6 +47,7 @@
 #include "PostProcess/VisualizeMotionVectors.h"
 #include "Rendering/MotionVectorSimulation.h"
 #include "ShaderPrint.h"
+#include "DataDrivenShaderPlatformInfo.h"
 #include "HairStrands/HairStrandsComposition.h"
 #include "HairStrands/HairStrandsUtils.h"
 #include "HighResScreenshot.h"
@@ -299,6 +305,9 @@ void AddPostProcessingPasses(
 		PixelInspector,
 		HMDDistortion,
 		HighResolutionScreenshotMask,
+#if UE_ENABLE_DEBUG_DRAWING
+		DebugPrimitive,
+#endif
 		PrimaryUpscale,
 		SecondaryUpscale,
 		MAX
@@ -347,6 +356,9 @@ void AddPostProcessingPasses(
 		TEXT("PixelInspector"),
 		TEXT("HMDDistortion"),
 		TEXT("HighResolutionScreenshotMask"),
+#if UE_ENABLE_DEBUG_DRAWING
+		TEXT("DebugPrimitive"),
+#endif
 		TEXT("PrimaryUpscale"),
 		TEXT("SecondaryUpscale")
 	};
@@ -391,6 +403,9 @@ void AddPostProcessingPasses(
 #endif
 	PassSequence.SetEnabled(EPass::HMDDistortion, EngineShowFlags.StereoRendering && EngineShowFlags.HMDDistortion);
 	PassSequence.SetEnabled(EPass::HighResolutionScreenshotMask, IsHighResolutionScreenshotMaskEnabled(View));
+#if UE_ENABLE_DEBUG_DRAWING
+	PassSequence.SetEnabled(EPass::DebugPrimitive, FSceneRenderer::ShouldCompositeDebugPrimitivesInPostProcess(View));
+#endif
 	PassSequence.SetEnabled(EPass::PrimaryUpscale, PaniniConfig.IsEnabled() || (View.PrimaryScreenPercentageMethod == EPrimaryScreenPercentageMethod::SpatialUpscale && PrimaryViewRect.Size() != View.GetSecondaryViewRectSize()));
 	PassSequence.SetEnabled(EPass::SecondaryUpscale, View.RequiresSecondaryUpscale() || View.Family->GetSecondarySpatialUpscalerInterface() != nullptr);
 
@@ -1057,6 +1072,18 @@ void AddPostProcessingPasses(
 			SceneColor = AddPostProcessMaterialChain(GraphBuilder, View, PassInputs, PostProcessMaterialAfterTonemappingChain);
 		}
 
+#if UE_ENABLE_DEBUG_DRAWING
+		if (PassSequence.IsEnabled(EPass::DebugPrimitive)) //Create new debug pass sequence
+		{
+			FCompositePrimitiveInputs PassInputs;
+			PassSequence.AcceptOverrideIfLastPass(EPass::DebugPrimitive, PassInputs.OverrideOutput);
+			PassInputs.SceneColor = SceneColor;
+			PassInputs.SceneDepth = SceneDepth;
+
+			SceneColor = AddDebugPrimitivePass(GraphBuilder, View, PassInputs);
+		}
+#endif
+
 		if (PassSequence.IsEnabled(EPass::VisualizeLumenScene))
 		{
 			FVisualizeLumenSceneInputs PassInputs;
@@ -1224,11 +1251,11 @@ void AddPostProcessingPasses(
 
 	if (PassSequence.IsEnabled(EPass::EditorPrimitive))
 	{
-		FEditorPrimitiveInputs PassInputs;
+		FCompositePrimitiveInputs PassInputs;
 		PassSequence.AcceptOverrideIfLastPass(EPass::EditorPrimitive, PassInputs.OverrideOutput);
 		PassInputs.SceneColor = SceneColor;
 		PassInputs.SceneDepth = SceneDepth;
-		PassInputs.BasePassType = FEditorPrimitiveInputs::EBasePassType::Deferred;
+		PassInputs.BasePassType = FCompositePrimitiveInputs::EBasePassType::Deferred;
 
 		SceneColor = AddEditorPrimitivePass(GraphBuilder, View, PassInputs, InstanceCullingManager);
 	}
@@ -1831,6 +1858,9 @@ void AddMobilePostProcessingPasses(FRDGBuilder& GraphBuilder, FScene* Scene, con
 		HighResolutionScreenshotMask,
 		SelectionOutline,
 		EditorPrimitive,
+#if UE_ENABLE_DEBUG_DRAWING
+		DebugPrimitive,
+#endif
 		PrimaryUpscale,
 		SecondaryUpscale,
 		Visualize,
@@ -1855,6 +1885,9 @@ void AddMobilePostProcessingPasses(FRDGBuilder& GraphBuilder, FScene* Scene, con
 		TEXT("HighResolutionScreenshotMask"),
 		TEXT("SelectionOutline"),
 		TEXT("EditorPrimitive"),
+#if UE_ENABLE_DEBUG_DRAWING
+		TEXT("DebugPrimitive"),
+#endif
 		TEXT("PrimaryUpscale"),
 		TEXT("SecondaryUpscale"),
 		TEXT("Visualize"),
@@ -1908,6 +1941,10 @@ void AddMobilePostProcessingPasses(FRDGBuilder& GraphBuilder, FScene* Scene, con
 	PassSequence.SetEnabled(EPass::SelectionOutline, false);
 	PassSequence.SetEnabled(EPass::EditorPrimitive, false);
 #endif
+
+#if UE_ENABLE_DEBUG_DRAWING
+	PassSequence.SetEnabled(EPass::DebugPrimitive, FSceneRenderer::ShouldCompositeDebugPrimitivesInPostProcess(View));
+#endif
 	PassSequence.SetEnabled(EPass::PrimaryUpscale, bShouldPrimaryUpscale);
 	PassSequence.SetEnabled(EPass::SecondaryUpscale, View.Family->GetSecondarySpatialUpscalerInterface() != nullptr);
 
@@ -2457,17 +2494,30 @@ void AddMobilePostProcessingPasses(FRDGBuilder& GraphBuilder, FScene* Scene, con
 
 	if (PassSequence.IsEnabled(EPass::EditorPrimitive))
 	{
-		FEditorPrimitiveInputs PassInputs;
+		FCompositePrimitiveInputs PassInputs;
 		PassSequence.AcceptOverrideIfLastPass(EPass::EditorPrimitive, PassInputs.OverrideOutput);
 		PassInputs.SceneColor = SceneColor;
 		PassInputs.SceneDepth = SceneDepth;
-		PassInputs.BasePassType = FEditorPrimitiveInputs::EBasePassType::Mobile;
+		PassInputs.BasePassType = FCompositePrimitiveInputs::EBasePassType::Mobile;
 		PassInputs.OverrideOutput.LoadAction = View.IsFirstInFamily() ? ERenderTargetLoadAction::EClear : ERenderTargetLoadAction::ELoad;
 
 		SceneColor = AddEditorPrimitivePass(GraphBuilder, View, PassInputs, InstanceCullingManager);
 	}
 #endif
 
+
+#if UE_ENABLE_DEBUG_DRAWING
+	if (PassSequence.IsEnabled(EPass::DebugPrimitive)) //Create new debug pass sequence
+	{
+		FCompositePrimitiveInputs PassInputs;
+		PassSequence.AcceptOverrideIfLastPass(EPass::DebugPrimitive, PassInputs.OverrideOutput);
+		PassInputs.SceneColor = SceneColor;
+		PassInputs.SceneDepth = SceneDepth;
+
+		SceneColor = AddDebugPrimitivePass(GraphBuilder, View, PassInputs);
+	}
+#endif
+
 	// Apply ScreenPercentage
 	if (PassSequence.IsEnabled(EPass::PrimaryUpscale))
 	{
diff --git a/Engine/Source/Runtime/Renderer/Private/RendererScene.cpp b/Engine/Source/Runtime/Renderer/Private/RendererScene.cpp
--- a/Engine/Source/Runtime/Renderer/Private/RendererScene.cpp
+++ b/Engine/Source/Runtime/Renderer/Private/RendererScene.cpp
@@ -1025,7 +1025,9 @@ uint64 FPreviousViewInfo::GetGPUSizeBytes(bool bLogSizes) const
 		WaterSSRHistory.GetGPUSizeBytes(bLogSizes) +
 		RoughRefractionHistory.GetGPUSizeBytes(bLogSizes) +
 		HairHistory.GetGPUSizeBytes(bLogSizes) +
-		EditorPrimtiveDepthHistory.GetGPUSizeBytes(bLogSizes) +
+#if UE_ENABLE_DEBUG_DRAWING
+		CompositePrimitiveDepthHistory.GetGPUSizeBytes(bLogSizes) +
+#endif
 		CustomSSRInput.GetGPUSizeBytes(bLogSizes) +
 		ReflectionsHistory.GetGPUSizeBytes(bLogSizes) +
 		WaterReflectionsHistory.GetGPUSizeBytes(bLogSizes) +
diff --git a/Engine/Source/Runtime/Renderer/Private/SceneRendering.cpp b/Engine/Source/Runtime/Renderer/Private/SceneRendering.cpp
--- a/Engine/Source/Runtime/Renderer/Private/SceneRendering.cpp
+++ b/Engine/Source/Runtime/Renderer/Private/SceneRendering.cpp
@@ -4000,6 +4000,13 @@ bool FSceneRenderer::ShouldCompositeEditorPrimitives(const FViewInfo& View)
 	return false;
 }
 
+#if UE_ENABLE_DEBUG_DRAWING
+bool FSceneRenderer::ShouldCompositeDebugPrimitivesInPostProcess(const FViewInfo& View)
+{
+	return View.DebugSimpleElementCollector.HasAnyPrimitives();
+}
+#endif
+
 void FSceneRenderer::UpdatePrimitiveIndirectLightingCacheBuffers()
 {
 	// Use a bit array to prevent primitives from being updated more than once.
diff --git a/Engine/Source/Runtime/Renderer/Private/SceneRendering.h b/Engine/Source/Runtime/Renderer/Private/SceneRendering.h
--- a/Engine/Source/Runtime/Renderer/Private/SceneRendering.h
+++ b/Engine/Source/Runtime/Renderer/Private/SceneRendering.h
@@ -1157,8 +1157,10 @@ struct FPreviousViewInfo
 	// Temporal AA history for Hair
 	FTemporalAAHistory HairHistory;
 
+#if UE_ENABLE_DEBUG_DRAWING
 	// Temporal AA history for the editor primitive depth upsampling
-	FTemporalAAHistory EditorPrimtiveDepthHistory;
+	FTemporalAAHistory CompositePrimitiveDepthHistory;
+#endif
 
 	// Scene color input for SSR, that can be different from TemporalAAHistory.RT[0] if there is a SSR
 	// input post process material.
@@ -1410,6 +1412,11 @@ public:
 
 	FSimpleElementCollector EditorSimpleElementCollector;
 
+#if UE_ENABLE_DEBUG_DRAWING
+	//Separate DebugSimpleElementCollector to not conflate any other simple element collector draws which may have been added from non-debug draw passes (e.g. non-opaque draws)
+	FSimpleElementCollector DebugSimpleElementCollector;
+#endif
+
 	TStaticArray<FParallelMeshDrawCommandPass, EMeshPass::Num> ParallelMeshDrawCommandPasses;
 	
 #if RHI_RAYTRACING // #yuriy_todo: Move to FRayTracingScene, as this is only valid for View[0]
@@ -2192,6 +2199,9 @@ public:
 	* @return true if compositing is needed
 	*/
 	static bool ShouldCompositeEditorPrimitives(const FViewInfo& View);
+#if UE_ENABLE_DEBUG_DRAWING
+	static bool ShouldCompositeDebugPrimitivesInPostProcess(const FViewInfo& View);
+#endif
 
 	/**
 	* Helper function performing actual work in render thread.
diff --git a/Engine/Source/Runtime/Renderer/Private/SceneVisibility.cpp b/Engine/Source/Runtime/Renderer/Private/SceneVisibility.cpp
--- a/Engine/Source/Runtime/Renderer/Private/SceneVisibility.cpp
+++ b/Engine/Source/Runtime/Renderer/Private/SceneVisibility.cpp
@@ -3551,7 +3551,11 @@ void FSceneRenderer::GatherDynamicMeshElements(
 				InViewFamily.GetFeatureLevel(),
 				&DynamicIndexBuffer,
 				&DynamicVertexBuffer,
-				&DynamicReadBuffer);
+				&DynamicReadBuffer
+#if UE_ENABLE_DEBUG_DRAWING
+				,&InViews[ViewIndex].DebugSimpleElementCollector
+#endif
+			);
 		}
 
 		const EShadingPath ShadingPath = Scene->GetShadingPath();
