Adds shader permutation in TSR for using 16bit floats on platforms that have bSupportsRealTypes=RuntimeDependent

diff --git a/Engine/Shaders/Private/TemporalSuperResolution/TSRCommon.ush b/Engine/Shaders/Private/TemporalSuperResolution/TSRCommon.ush
--- a/Engine/Shaders/Private/TemporalSuperResolution/TSRCommon.ush
+++ b/Engine/Shaders/Private/TemporalSuperResolution/TSRCommon.ush
@@ -24,7 +24,9 @@
 
 #define DEBUG_OUTPUT 0
 
-#if PLATFORM_SUPPORTS_REAL_TYPES && 1
+#if defined(DIM_16BIT_VALU)
+	#define CONFIG_COMPILE_FP16 DIM_16BIT_VALU
+#elif PLATFORM_SUPPORTS_REAL_TYPES
 	#define CONFIG_COMPILE_FP16 1
 #else
 	#define CONFIG_COMPILE_FP16 0
@@ -598,17 +600,29 @@ float2 SafeRcp(float2 x)
 
 CALL_SITE_DEBUGLOC
 half SafeRcp(half x)
+#if PLATFORM_GPU_ARCH >= PLATFORM_GPU_ARCH_AMD_RDNA_2 && PLATFORM_GPU_ARCH <= PLATFORM_GPU_ARCH_AMD_LATTEST
 {
 	// If x=0.0, then MaxHalfFloat * 0.0 = 0.0
 	return min(rcp(x), half(MaxHalfFloat)) * saturate(x * asfloat16(uint16_t(0x7C00)));
 }
+#else
+{
+	return select(x > half(0.0), min(rcp(x), half(MaxHalfFloat)), half(0.0));
+}
+#endif
 
 CALL_SITE_DEBUGLOC
 half2 SafeRcp(half2 x)
+#if PLATFORM_GPU_ARCH >= PLATFORM_GPU_ARCH_AMD_RDNA_2 && PLATFORM_GPU_ARCH <= PLATFORM_GPU_ARCH_AMD_LATTEST
 {
 	// If x=0.0, then MaxHalfFloat * 0.0 = 0.0
 	return min(rcp(x), half(MaxHalfFloat)) * saturate(x * asfloat16(uint16_t(0x7C00)));
 }
+#else
+{
+	return select(x > half(0.0), min(rcp(x), half(MaxHalfFloat)), half(0.0));
+}
+#endif
 
 #endif
 
diff --git a/Engine/Shaders/Public/WaveBroadcastIntrinsics.ush b/Engine/Shaders/Public/WaveBroadcastIntrinsics.ush
--- a/Engine/Shaders/Public/WaveBroadcastIntrinsics.ush
+++ b/Engine/Shaders/Public/WaveBroadcastIntrinsics.ush
@@ -280,6 +280,12 @@ uint16_t2 WaveBroadcast(const FWaveBroadcastSettings Settings, uint16_t2 v)
 	uint bvu = WaveBroadcast(Settings, vu);
 	return bit_cast_uint16_t2(bvu);
 } 
+#elif 1
+{
+	uint vu = uint(v.x) | (uint(v.y) << 16u);
+	uint bvu = WaveBroadcast(Settings, vu);
+	return uint16_t2(bvu, bvu >> 16u);
+}
 #else
 {
 	return uint16_t2(
diff --git a/Engine/Source/Runtime/Renderer/Private/PostProcess/TemporalSuperResolution.cpp b/Engine/Source/Runtime/Renderer/Private/PostProcess/TemporalSuperResolution.cpp
--- a/Engine/Source/Runtime/Renderer/Private/PostProcess/TemporalSuperResolution.cpp
+++ b/Engine/Source/Runtime/Renderer/Private/PostProcess/TemporalSuperResolution.cpp
@@ -52,6 +52,11 @@ TAutoConsoleVariable<int32> CVarTSRWaveOps(
 	TEXT("Whether to use wave ops in the shading rejection heuristics"),
 	ECVF_RenderThreadSafe);
 
+TAutoConsoleVariable<int32> CVarTSR16BitVALU(
+	TEXT("r.TSR.16BitVALU"), 1,
+	TEXT("Whether to use 16bit VALU on platform that have bSupportsRealTypes=RuntimeDependent"),
+	ECVF_RenderThreadSafe);
+
 TAutoConsoleVariable<float> CVarTSRHistoryRejectionSampleCount(
 	TEXT("r.TSR.ShadingRejection.SampleCount"), 2.0f,
 	TEXT("Maximum number of sample in each output pixel of the history after total shading rejection."),
@@ -329,7 +334,9 @@ class FTSRShader : public FGlobalShader
 public:
 	static constexpr int32 kSupportMinWaveSize = 32;
 	static constexpr int32 kSupportMaxWaveSize = 64;
-	
+
+	class F16BitVALUDim : SHADER_PERMUTATION_BOOL("DIM_16BIT_VALU");
+
 	FTSRShader(const ShaderMetaType::CompiledShaderInitializerType& Initializer)
 		: FGlobalShader(Initializer)
 	{ }
@@ -337,6 +344,31 @@ public:
 	FTSRShader()
 	{ }
 
+	static ERHIFeatureSupport Supports16BitVALU(EShaderPlatform Platform)
+	{
+		return FDataDrivenShaderPlatformInfo::GetSupportsRealTypes(Platform);
+	}
+
+	static bool ShouldCompile32or16BitPermutation(EShaderPlatform Platform, bool bIs16BitVALUPermutation)
+	{
+		const ERHIFeatureSupport Support = FTSRShader::Supports16BitVALU(Platform);
+
+		if (Support == ERHIFeatureSupport::RuntimeGuaranteed)
+		{
+			// Only compile the 16bit permutation
+			return bIs16BitVALUPermutation;
+		}
+		else if (Support == ERHIFeatureSupport::RuntimeDependent)
+		{
+			// Compile both the 32bit and 16bit permutations
+			return true;
+		}
+		else
+		{
+			return !bIs16BitVALUPermutation;
+		}
+	}
+
 	static ERHIFeatureSupport SupportsWaveOps(EShaderPlatform Platform)
 	{
 		// UE-161125: SPIRV ends up with 5min+ compilation times, and D3D11 needs SPIRV backend for HLSL2021 features on FXC.
@@ -373,7 +405,7 @@ public:
 
 	static void ModifyCompilationEnvironment(const FGlobalShaderPermutationParameters& Parameters, FShaderCompilerEnvironment& OutEnvironment)
 	{
-		if (FDataDrivenShaderPlatformInfo::GetSupportsRealTypes(Parameters.Platform) == ERHIFeatureSupport::RuntimeGuaranteed)
+		if (FTSRShader::Supports16BitVALU(Parameters.Platform) == ERHIFeatureSupport::RuntimeGuaranteed)
 		{
 			OutEnvironment.CompilerFlags.Add(CFLAG_AllowRealTypes);
 		}
@@ -571,7 +603,7 @@ class FTSRRejectShadingCS : public FTSRShader
 	class FWaveSizeOps : SHADER_PERMUTATION_SPARSE_INT("DIM_WAVE_SIZE", 0, 32, 64);
 	class FFlickeringDetectionDim : SHADER_PERMUTATION_BOOL("DIM_FLICKERING_DETECTION");
 
-	using FPermutationDomain = TShaderPermutationDomain<FWaveSizeOps, FFlickeringDetectionDim>;
+	using FPermutationDomain = TShaderPermutationDomain<FWaveSizeOps, FFlickeringDetectionDim, FTSRShader::F16BitVALUDim>;
 
 	BEGIN_SHADER_PARAMETER_STRUCT(FParameters, )
 		SHADER_PARAMETER_STRUCT_INCLUDE(FTSRCommonParameters, CommonParameters)
@@ -597,6 +629,11 @@ class FTSRRejectShadingCS : public FTSRShader
 		SHADER_PARAMETER_RDG_TEXTURE_UAV(RWTexture2D, DebugOutput)
 	END_SHADER_PARAMETER_STRUCT()
 
+	static FPermutationDomain RemapPermutation(FPermutationDomain PermutationVector)
+	{
+		return PermutationVector;
+	}
+
 	static bool ShouldCompilePermutation(const FGlobalShaderPermutationParameters& Parameters)
 	{
 		if (!FTSRShader::ShouldCompilePermutation(Parameters))
@@ -605,6 +642,11 @@ class FTSRRejectShadingCS : public FTSRShader
 		}
 
 		FPermutationDomain PermutationVector(Parameters.PermutationId);
+		if (PermutationVector != RemapPermutation(PermutationVector))
+		{
+			return false;
+		}
+
 		int32 WaveSize = PermutationVector.Get<FWaveSizeOps>();
 
 		ERHIFeatureSupport WaveOpsSupport = FTSRShader::SupportsWaveOps(Parameters.Platform);
@@ -623,7 +665,15 @@ class FTSRRejectShadingCS : public FTSRShader
 		}
 		else // if (WaveSize == LDS fallback)
 		{
-			return FTSRShader::SupportsLDS(Parameters.Platform);
+			if (!FTSRShader::SupportsLDS(Parameters.Platform))
+			{
+				return false;
+			}
+		}
+
+		if (!FTSRShader::ShouldCompile32or16BitPermutation(Parameters.Platform, PermutationVector.Get<FTSRShader::F16BitVALUDim>()))
+		{
+			return false;
 		}
 
 		return true;
@@ -643,6 +693,11 @@ class FTSRRejectShadingCS : public FTSRShader
 			}
 			OutEnvironment.CompilerFlags.Add(CFLAG_WaveOperations);
 		}
+
+		if (PermutationVector.Get<FTSRShader::F16BitVALUDim>())
+		{
+			OutEnvironment.CompilerFlags.Add(CFLAG_AllowRealTypes);
+		}
 	}
 }; // class FTSRRejectShadingCS
 
@@ -715,7 +770,7 @@ class FTSRUpdateHistoryCS : public FTSRShader
 	class FSeparateTranslucencyDim : SHADER_PERMUTATION_BOOL("DIM_SEPARATE_TRANSLUCENCY");
 	class FGrandReprojectionDim : SHADER_PERMUTATION_BOOL("DIM_GRAND_REPROJECTION");
 
-	using FPermutationDomain = TShaderPermutationDomain<FQualityDim, FSeparateTranslucencyDim, FGrandReprojectionDim>;
+	using FPermutationDomain = TShaderPermutationDomain<FQualityDim, FSeparateTranslucencyDim, FGrandReprojectionDim, FTSRShader::F16BitVALUDim>;
 
 	BEGIN_SHADER_PARAMETER_STRUCT(FParameters, )
 		SHADER_PARAMETER_STRUCT_INCLUDE(FTSRCommonParameters, CommonParameters)
@@ -777,6 +832,11 @@ class FTSRUpdateHistoryCS : public FTSRShader
 			return false;
 		}
 
+		if (!FTSRShader::ShouldCompile32or16BitPermutation(Parameters.Platform, PermutationVector.Get<FTSRShader::F16BitVALUDim>()))
+		{
+			return false;
+		}
+
 		return FTSRShader::ShouldCompilePermutation(Parameters);
 	}
 
@@ -784,6 +844,12 @@ class FTSRUpdateHistoryCS : public FTSRShader
 	{
 		FTSRShader::ModifyCompilationEnvironment(Parameters, OutEnvironment);
 		OutEnvironment.CompilerFlags.Add(CFLAG_Wave32);
+
+		FPermutationDomain PermutationVector(Parameters.PermutationId);
+		if (PermutationVector.Get<FTSRShader::F16BitVALUDim>())
+		{
+			OutEnvironment.CompilerFlags.Add(CFLAG_AllowRealTypes);
+		}
 	}
 }; // class FTSRUpdateHistoryCS
 
@@ -916,6 +982,10 @@ ITemporalUpscaler::FOutputs AddTemporalSuperResolutionPasses(
 	const ERHIFeatureSupport WaveOpsSupport = FTSRShader::SupportsWaveOps(View.GetShaderPlatform());
 	const bool bSupportsLDS = FTSRShader::SupportsLDS(View.GetShaderPlatform());
 	const bool bUseWaveOps = (CVarTSRWaveOps.GetValueOnRenderThread() != 0 && GRHISupportsWaveOperations && bSupportsLDS && (WaveOpsSupport == ERHIFeatureSupport::RuntimeDependent || WaveOpsSupport == ERHIFeatureSupport::RuntimeGuaranteed)) || !bSupportsLDS;
+	
+	// Whether to use 16bit VALU
+	const ERHIFeatureSupport VALU16BitSupport = FTSRShader::Supports16BitVALU(View.GetShaderPlatform());
+	const bool bUse16BitVALU = (CVarTSR16BitVALU.GetValueOnRenderThread() != 0 && GRHIGlobals.SupportsNative16BitOps && VALU16BitSupport == ERHIFeatureSupport::RuntimeDependent) || VALU16BitSupport == ERHIFeatureSupport::RuntimeGuaranteed;
 
 	// Whether should accumulate sub pixel depth.
 	const ETSRSubpixelMethod SubpixelMethod = ETSRSubpixelMethod(FMath::Clamp(CVarTSRSubpixelMethod.GetValueOnRenderThread(), 0, 2));
@@ -1772,13 +1842,16 @@ ITemporalUpscaler::FOutputs AddTemporalSuperResolutionPasses(
 		FTSRRejectShadingCS::FPermutationDomain PermutationVector;
 		PermutationVector.Set<FTSRRejectShadingCS::FWaveSizeOps>(bUseWaveOps && GRHIMinimumWaveSize >= 32 && GRHIMinimumWaveSize <= 64 ? GRHIMinimumWaveSize : 0);
 		PermutationVector.Set<FTSRRejectShadingCS::FFlickeringDetectionDim>(FlickeringFramePeriod > 0.0f);
+		PermutationVector.Set<FTSRShader::F16BitVALUDim>(bUse16BitVALU);
+		PermutationVector = FTSRRejectShadingCS::RemapPermutation(PermutationVector);
 
 		TShaderMapRef<FTSRRejectShadingCS> ComputeShader(View.ShaderMap, PermutationVector);
 		FComputeShaderUtils::AddPass(
 			GraphBuilder,
-			RDG_EVENT_NAME("TSR RejectShading(WaveSize=%d FlickeringFramePeriod=%f %s) %dx%d",
+			RDG_EVENT_NAME("TSR RejectShading(WaveSize=%d FlickeringFramePeriod=%f %s%s) %dx%d",
 				int32(PermutationVector.Get<FTSRRejectShadingCS::FWaveSizeOps>()),
 				FlickeringFramePeriod,
+				PermutationVector.Get<FTSRShader::F16BitVALUDim>() ? TEXT(" 16bit") : TEXT(""),
 				!bAccumulateTranslucencySeparately ? TEXT(" ComposeTranslucency") : TEXT(""),
 				InputRect.Width(), InputRect.Height()),
 			AsyncComputePasses >= 3 ? ERDGPassFlags::AsyncCompute : ERDGPassFlags::Compute,
@@ -1942,12 +2015,14 @@ ITemporalUpscaler::FOutputs AddTemporalSuperResolutionPasses(
 		PermutationVector.Set<FTSRUpdateHistoryCS::FQualityDim>(UpdateHistoryQuality);
 		PermutationVector.Set<FTSRUpdateHistoryCS::FSeparateTranslucencyDim>(bAccumulateTranslucencySeparately);
 		PermutationVector.Set<FTSRUpdateHistoryCS::FGrandReprojectionDim>(bGrandReprojection);
+		PermutationVector.Set<FTSRShader::F16BitVALUDim>(bUse16BitVALU);
 
 		TShaderMapRef<FTSRUpdateHistoryCS> ComputeShader(View.ShaderMap, PermutationVector);
 		FComputeShaderUtils::AddPass(
 			GraphBuilder,
-			RDG_EVENT_NAME("TSR UpdateHistory(Quality=%s%s%s%s%s) %dx%d",
+			RDG_EVENT_NAME("TSR UpdateHistory(Quality=%s%s%s%s%s%s) %dx%d",
 				kUpdateQualityNames[int32(PermutationVector.Get<FTSRUpdateHistoryCS::FQualityDim>())],
+				PermutationVector.Get<FTSRShader::F16BitVALUDim>() ? TEXT(" 16bit") : TEXT(""),
 				PermutationVector.Get<FTSRUpdateHistoryCS::FSeparateTranslucencyDim>() ? TEXT(" SeparateTranslucency") : TEXT(""),
 				PermutationVector.Get<FTSRUpdateHistoryCS::FGrandReprojectionDim>() ? TEXT(" GrandReprojection") : TEXT(""),
 				(History.ColorArray ? History.ColorArray : History.Output)->Desc.Format == PF_FloatR11G11B10 ? TEXT(" R11G11B10") : TEXT(""),
