Use the new FBoxSphereBounds::Builder to clean other initialization of FBoxSphereBounds
* Those where correct, but did manual job to ensure the origin wasn't automatically included in the bounds

diff --git a/Engine/Plugins/2D/Paper2D/Source/Paper2D/Private/PaperGroupedSpriteComponent.cpp b/Engine/Plugins/2D/Paper2D/Source/Paper2D/Private/PaperGroupedSpriteComponent.cpp
--- a/Engine/Plugins/2D/Paper2D/Source/Paper2D/Private/PaperGroupedSpriteComponent.cpp
+++ b/Engine/Plugins/2D/Paper2D/Source/Paper2D/Private/PaperGroupedSpriteComponent.cpp
@@ -271,8 +271,7 @@ bool UPaperGroupedSpriteComponent::CanEditSimulatePhysics()
 
 FBoxSphereBounds UPaperGroupedSpriteComponent::CalcBounds(const FTransform& BoundTransform) const
 {
-	bool bHadAnyBounds = false;
-	FBoxSphereBounds NewBounds(ForceInit);
+	FBoxSphereBounds::Builder BoundsBuilder;
 
 	if (PerInstanceSpriteData.Num() > 0)
 	{
@@ -284,21 +283,12 @@ FBoxSphereBounds UPaperGroupedSpriteComponent::CalcBounds(const FTransform& Boun
 			{
 				const FBoxSphereBounds RenderBounds = InstanceData.SourceSprite->GetRenderBounds();
 				const FBoxSphereBounds InstanceBounds = RenderBounds.TransformBy(InstanceData.Transform * BoundTransformMatrix);
-
-				if (bHadAnyBounds)
-				{
-					NewBounds = NewBounds + InstanceBounds;
-				}
-				else
-				{
-					NewBounds = InstanceBounds;
-					bHadAnyBounds = true;
-				}
+				BoundsBuilder += InstanceBounds;
 			}
 		}
 	}
 
-	return bHadAnyBounds ? NewBounds : FBoxSphereBounds(BoundTransform.GetLocation(), FVector::ZeroVector, 0.f);
+	return BoundsBuilder.IsValid() ? BoundsBuilder : FBoxSphereBounds(BoundTransform.GetLocation(), FVector::ZeroVector, 0.f);
 }
 
 #if WITH_EDITOR
diff --git a/Engine/Plugins/Editor/WorldPartitionHLODUtilities/Source/Private/WorldPartition/HLOD/Builders/HLODBuilderMeshSimplify.cpp b/Engine/Plugins/Editor/WorldPartitionHLODUtilities/Source/Private/WorldPartition/HLOD/Builders/HLODBuilderMeshSimplify.cpp
--- a/Engine/Plugins/Editor/WorldPartitionHLODUtilities/Source/Private/WorldPartition/HLOD/Builders/HLODBuilderMeshSimplify.cpp
+++ b/Engine/Plugins/Editor/WorldPartitionHLODUtilities/Source/Private/WorldPartition/HLOD/Builders/HLODBuilderMeshSimplify.cpp
@@ -100,20 +100,17 @@ TArray<UActorComponent*> UHLODBuilderMeshSimplify::Build(const FHLODBuildContext
 	// Gather bounds of the input components
 	auto GetComponentsBounds = [&]() -> FBoxSphereBounds
 	{
-		FBoxSphereBounds Bounds;
-		bool bFirst = true;
+		FBoxSphereBounds::Builder BoundsBuilder;
 
 		for (UActorComponent* Component : InSourceComponents)
 		{
 			if (USceneComponent* SceneComponent = Cast<USceneComponent>(Component))
 			{
-				FBoxSphereBounds ComponentBounds = SceneComponent->Bounds;
-				Bounds = bFirst ? ComponentBounds : Bounds + ComponentBounds;
-				bFirst = false;
+				BoundsBuilder += SceneComponent->Bounds;
 			}
 		}
 
-		return Bounds;
+		return BoundsBuilder;
 	};
 
 	float ScreenSizePercent = ComputeBoundsScreenSize(FVector::ZeroVector, static_cast<float>(GetComponentsBounds().SphereRadius), FVector(0.0f, 0.0f, InHLODBuildContext.MinVisibleDistance), ProjectionMatrix);
diff --git a/Engine/Plugins/Importers/AlembicImporter/Source/AlembicLibrary/Private/AbcImportUtilities.cpp b/Engine/Plugins/Importers/AlembicImporter/Source/AlembicLibrary/Private/AbcImportUtilities.cpp
--- a/Engine/Plugins/Importers/AlembicImporter/Source/AlembicLibrary/Private/AbcImportUtilities.cpp
+++ b/Engine/Plugins/Importers/AlembicImporter/Source/AlembicLibrary/Private/AbcImportUtilities.cpp
@@ -1460,8 +1460,9 @@ bool AbcImporterUtilities::IsObjectVisibilityConstant(const Alembic::Abc::IObjec
 
 FBoxSphereBounds AbcImporterUtilities::ExtractBounds(Alembic::Abc::IBox3dProperty InBoxBoundsProperty)
 {
-	FBoxSphereBounds Bounds(ForceInitToZero);
-        // Extract data only if the property is found	
+	FBoxSphereBounds::Builder BoundsBuilder;
+
+	// Extract data only if the property is found	
 	if (InBoxBoundsProperty.valid())
 	{
 		const int32 NumSamples = InBoxBoundsProperty.getNumSamples();
@@ -1469,15 +1470,17 @@ FBoxSphereBounds AbcImporterUtilities::ExtractBounds(Alembic::Abc::IBox3dPropert
 		{
 			Alembic::Abc::Box3d BoundsSample;
 			InBoxBoundsProperty.get(BoundsSample, SampleIndex);
-                        // Set up bounds from Alembic data format
+
+			// Set up bounds from Alembic data format
 			const Imath::V3d BoundSize = BoundsSample.size();
 			const Imath::V3d BoundCenter = BoundsSample.center();
 			const FBoxSphereBounds ConvertedBounds(FVector(BoundCenter.x, BoundCenter.y, BoundCenter.z), FVector(BoundSize.x  * 0.5f, BoundSize.y * 0.5f, BoundSize.z * 0.5f), (const float)BoundSize.length() * 0.5f);
-			Bounds = ( SampleIndex == 0 ) ? ConvertedBounds : Bounds + ConvertedBounds;
+
+			BoundsBuilder += ConvertedBounds;
 		}
 	}
 
-	return Bounds;
+	return BoundsBuilder;
 }
 
 void AbcImporterUtilities::ApplyConversion(FMatrix& InOutMatrix, const FAbcConversionSettings& InConversionSettings)
diff --git a/Engine/Plugins/MeshPainting/Source/MeshPaintingToolset/Private/MeshVertexPaintingTool.cpp b/Engine/Plugins/MeshPainting/Source/MeshPaintingToolset/Private/MeshVertexPaintingTool.cpp
--- a/Engine/Plugins/MeshPainting/Source/MeshPaintingToolset/Private/MeshVertexPaintingTool.cpp
+++ b/Engine/Plugins/MeshPainting/Source/MeshPaintingToolset/Private/MeshVertexPaintingTool.cpp
@@ -191,30 +191,17 @@ void UMeshVertexPaintingTool::OnPropertyModified(UObject* PropertySet, FProperty
 
 double UMeshVertexPaintingTool::EstimateMaximumTargetDimension()
 {
-	bool bFoundComponentToUse = false;
-	FBoxSphereBounds Bounds = FBoxSphereBounds(0.0);
 	if (UMeshPaintingSubsystem* MeshPaintingSubsystem = GEngine->GetEngineSubsystem<UMeshPaintingSubsystem>())
 	{
-		bool bFirstItem = true;
-
-		FBoxSphereBounds Extents;
+		FBoxSphereBounds::Builder ExtentsBuilder;
 		for (UMeshComponent* SelectedComponent : MeshPaintingSubsystem->GetSelectedMeshComponents())
 		{
-			if (bFirstItem)
-			{
-				Extents = SelectedComponent->Bounds;
-			}
-			else
-			{
-				Extents = Extents + SelectedComponent->Bounds;
-			}
-
-			bFirstItem = false;
-			bFoundComponentToUse = true;
+			ExtentsBuilder += SelectedComponent->Bounds;
 		}
-		if (bFoundComponentToUse)
+
+		if (ExtentsBuilder.IsValid())
 		{
-			return Extents.BoxExtent.GetAbsMax();
+			return FBoxSphereBounds(ExtentsBuilder).BoxExtent.GetAbsMax();
 		}
 	}
 
diff --git a/Engine/Plugins/Runtime/MeshModelingToolset/Source/ModelingComponents/Private/Scene/WorldRenderCapture.cpp b/Engine/Plugins/Runtime/MeshModelingToolset/Source/ModelingComponents/Private/Scene/WorldRenderCapture.cpp
--- a/Engine/Plugins/Runtime/MeshModelingToolset/Source/ModelingComponents/Private/Scene/WorldRenderCapture.cpp
+++ b/Engine/Plugins/Runtime/MeshModelingToolset/Source/ModelingComponents/Private/Scene/WorldRenderCapture.cpp
@@ -189,8 +189,7 @@ void FWorldRenderCapture::SetVisibleActors(const TArray<AActor*>& Actors)
 
 	VisiblePrimitives.Reset();
 
-	VisibleBounds = FBoxSphereBounds();
-	bool bFirst = true;
+	FBoxSphereBounds::Builder BoundsBuilder;
 
 	// Find all components that need to be included in rendering.
 	// This also descends into any ChildActorComponents
@@ -212,9 +211,7 @@ void FWorldRenderCapture::SetVisibleActors(const TArray<AActor*>& Actors)
 				VisiblePrimitives.Add(PrimitiveComponent->ComponentId);
 
 				// Append bounds of visible components only
-				FBoxSphereBounds ComponentBounds = PrimitiveComponent->Bounds;
-				VisibleBounds = (bFirst) ? ComponentBounds : (VisibleBounds + ComponentBounds);
-				bFirst = false;
+				BoundsBuilder += PrimitiveComponent->Bounds;
 			}
 			else if (Cast<UChildActorComponent>(Component) != nullptr)
 			{
@@ -229,6 +226,8 @@ void FWorldRenderCapture::SetVisibleActors(const TArray<AActor*>& Actors)
 			}
 		}
 	}
+
+	VisibleBounds = BoundsBuilder;
 }
 
 
diff --git a/Engine/Source/Editor/Blutility/Private/EditorUtilityLibrary.cpp b/Engine/Source/Editor/Blutility/Private/EditorUtilityLibrary.cpp
--- a/Engine/Source/Editor/Blutility/Private/EditorUtilityLibrary.cpp
+++ b/Engine/Source/Editor/Blutility/Private/EditorUtilityLibrary.cpp
@@ -212,26 +212,16 @@ TArray<AActor*> UEditorUtilityLibrary::GetSelectionSet()
 
 void UEditorUtilityLibrary::GetSelectionBounds(FVector& Origin, FVector& BoxExtent, float& SphereRadius)
 {
-	bool bFirstItem = true;
-
-	FBoxSphereBounds Extents;
+	FBoxSphereBounds::Builder BoundsBuilder;
 	for (FSelectionIterator It(GEditor->GetSelectedActorIterator()); It; ++It)
 	{
 		if (AActor* Actor = Cast<AActor>(*It))
 		{
-			if (bFirstItem)
-			{
-				Extents = Actor->GetRootComponent()->Bounds;
-			}
-			else
-			{
-				Extents = Extents + Actor->GetRootComponent()->Bounds;
-			}
-
-			bFirstItem = false;
+			BoundsBuilder += Actor->GetRootComponent()->Bounds;
 		}
 	}
 
+	FBoxSphereBounds Extents(BoundsBuilder);
 	Origin = Extents.Origin;
 	BoxExtent = Extents.BoxExtent;
 	SphereRadius = (float)Extents.SphereRadius; // TODO: LWC: should be double, but need to deprecate function and replace for old C++ references to continue working.
diff --git a/Engine/Source/Editor/Blutility/Private/GlobalEditorUtilityBase.cpp b/Engine/Source/Editor/Blutility/Private/GlobalEditorUtilityBase.cpp
--- a/Engine/Source/Editor/Blutility/Private/GlobalEditorUtilityBase.cpp
+++ b/Engine/Source/Editor/Blutility/Private/GlobalEditorUtilityBase.cpp
@@ -42,26 +42,16 @@ TArray<AActor*> UDEPRECATED_GlobalEditorUtilityBase::GetSelectionSet()
 
 void UDEPRECATED_GlobalEditorUtilityBase::GetSelectionBounds(FVector& Origin, FVector& BoxExtent, float& SphereRadius)
 {
-	bool bFirstItem = true;
-
-	FBoxSphereBounds Extents;
+	FBoxSphereBounds::Builder BoundsBuilder;
 	for (FSelectionIterator It(GEditor->GetSelectedActorIterator()); It; ++It)
 	{
 		if (AActor* Actor = Cast<AActor>(*It))
 		{
-			if (bFirstItem)
-			{
-				Extents = Actor->GetRootComponent()->Bounds;
-			}
-			else
-			{
-				Extents = Extents + Actor->GetRootComponent()->Bounds;
-			}
-
-			bFirstItem = false;
+			BoundsBuilder += Actor->GetRootComponent()->Bounds;
 		}
 	}
 
+	FBoxSphereBounds Extents(BoundsBuilder);
 	Origin = Extents.Origin;
 	BoxExtent = Extents.BoxExtent;
 	SphereRadius = (float)Extents.SphereRadius;
diff --git a/Engine/Source/Runtime/AIModule/Private/EnvironmentQuery/EQSRenderingComponent.cpp b/Engine/Source/Runtime/AIModule/Private/EnvironmentQuery/EQSRenderingComponent.cpp
--- a/Engine/Source/Runtime/AIModule/Private/EnvironmentQuery/EQSRenderingComponent.cpp
+++ b/Engine/Source/Runtime/AIModule/Private/EnvironmentQuery/EQSRenderingComponent.cpp
@@ -333,13 +333,13 @@ FBoxSphereBounds UEQSRenderingComponent::CalcBounds(const FTransform& LocalToWor
 
 	if (Spheres.Num())
 	{
-		FBoxSphereBounds DebugBounds(FSphere(Spheres[0].Location, Spheres[0].Radius));
-		for (int32 Index = 1; Index < Spheres.Num(); ++Index)
+		FBoxSphereBounds::Builder DebugBoundsBuilder;
+		for (int32 Index = 0; Index < Spheres.Num(); ++Index)
 		{
-			DebugBounds = DebugBounds + FSphere(Spheres[Index].Location, Spheres[Index].Radius);
+			DebugBoundsBuilder += FSphere(Spheres[Index].Location, Spheres[Index].Radius);
 		}
 
-		return DebugBounds;
+		return DebugBoundsBuilder;
 	}
 #endif
 
diff --git a/Engine/Source/Runtime/Engine/Private/InstancedStaticMesh.cpp b/Engine/Source/Runtime/Engine/Private/InstancedStaticMesh.cpp
--- a/Engine/Source/Runtime/Engine/Private/InstancedStaticMesh.cpp
+++ b/Engine/Source/Runtime/Engine/Private/InstancedStaticMesh.cpp
@@ -2788,14 +2788,14 @@ FBoxSphereBounds UInstancedStaticMeshComponent::CalcBounds(const FTransform& Bou
 		FMatrix BoundTransformMatrix = BoundTransform.ToMatrixWithScale();
 
 		FBoxSphereBounds RenderBounds = GetStaticMesh()->GetBounds();
-		FBoxSphereBounds NewBounds = RenderBounds.TransformBy(PerInstanceSMData[0].Transform * BoundTransformMatrix);
 
-		for (int32 InstanceIndex = 1; InstanceIndex < PerInstanceSMData.Num(); InstanceIndex++)
+		FBoxSphereBounds::Builder BoundsBuilder;
+		for (int32 InstanceIndex = 0; InstanceIndex < PerInstanceSMData.Num(); InstanceIndex++)
 		{
-			NewBounds = NewBounds + RenderBounds.TransformBy(PerInstanceSMData[InstanceIndex].Transform * BoundTransformMatrix);
+			BoundsBuilder += RenderBounds.TransformBy(PerInstanceSMData[InstanceIndex].Transform * BoundTransformMatrix);
 		}
 
-		return NewBounds;
+		return BoundsBuilder;
 	}
 	else
 	{
diff --git a/Engine/Source/Runtime/Engine/Private/MaterialMerging.cpp b/Engine/Source/Runtime/Engine/Private/MaterialMerging.cpp
--- a/Engine/Source/Runtime/Engine/Private/MaterialMerging.cpp
+++ b/Engine/Source/Runtime/Engine/Private/MaterialMerging.cpp
@@ -191,20 +191,17 @@ bool FMaterialProxySettings::ResolveTexelDensity(const TArray<UPrimitiveComponen
 	// Gather bounds of the input components
 	auto GetActorsBounds = [&]() -> FBoxSphereBounds
 	{
-		FBoxSphereBounds Bounds;
-		bool bFirst = true;
+		FBoxSphereBounds::Builder BoundsBuilder;
 
 		for (UPrimitiveComponent* Component : InComponents)
 		{
 			if (Component)
 			{
-				FBoxSphereBounds ComponentBounds = Component->Bounds;
-				Bounds = bFirst ? ComponentBounds : Bounds + ComponentBounds;
-				bFirst = false;
+				BoundsBuilder += Component->Bounds;
 			}
 		}
 
-		return Bounds;
+		return BoundsBuilder;
 	};
 
 	if (TextureSizingType == ETextureSizingType::TextureSizingType_AutomaticFromMeshScreenSize)
diff --git a/Engine/Source/Runtime/Renderer/Private/ShadowSetupMobile.cpp b/Engine/Source/Runtime/Renderer/Private/ShadowSetupMobile.cpp
--- a/Engine/Source/Runtime/Renderer/Private/ShadowSetupMobile.cpp
+++ b/Engine/Source/Runtime/Renderer/Private/ShadowSetupMobile.cpp
@@ -245,11 +245,12 @@ static void VisualizeMobileDynamicCSMSubjectCapsules(FViewInfo& View, FLightScen
 			// Combined bounds
 			FVector CombinedCasterStart;
 			FVector CombinedCasterEnd;
-			FBoxSphereBounds CombinedBounds(ForceInitToZero);
+			FBoxSphereBounds::Builder CombinedBoundsBuilder;
 			for (auto& Caster : MobileCSMSubjectPrimitives.GetShadowSubjectPrimitives())
 			{
-				CombinedBounds = (CombinedBounds.SphereRadius > 0.0f) ? CombinedBounds + Caster->Proxy->GetBounds() : Caster->Proxy->GetBounds();
+				CombinedBoundsBuilder += Caster->Proxy->GetBounds();
 			}
+			FBoxSphereBounds CombinedBounds(CombinedBoundsBuilder);
 			CombinedCasterStart = CombinedBounds.Origin;
 			CombinedCasterEnd = CombinedBounds.Origin + (LightDir * ShadowCastLength);
 
@@ -514,17 +515,16 @@ void FMobileSceneRenderer::BuildCSMVisibilityState(FLightSceneInfo* LightSceneIn
 					}
 					case 2: // combined casters:
 					{
-						FVector CombinedCasterStart;
-						FVector CombinedCasterEnd;
-						FBoxSphereBounds CombinedBounds(ForceInitToZero);
-
 						// Calculate combined bounds
+						FBoxSphereBounds::Builder CombinedBoundsBuilder;
 						for (auto& Caster : ShadowSubjectPrimitives)
 						{
-							CombinedBounds = (CombinedBounds.SphereRadius > 0.0f) ? CombinedBounds + Caster->Proxy->GetBounds() : Caster->Proxy->GetBounds();
+							CombinedBoundsBuilder += Caster->Proxy->GetBounds();
 						}
-						CombinedCasterStart = CombinedBounds.Origin;
-						CombinedCasterEnd = CombinedBounds.Origin + (LightDir * ShadowCastLength);
+
+						FBoxSphereBounds CombinedBounds(CombinedBoundsBuilder);
+						FVector CombinedCasterStart = CombinedBounds.Origin;
+						FVector CombinedCasterEnd = CombinedBounds.Origin + (LightDir * ShadowCastLength);
 
 						if (bSphereTest)
 						{
