Fix crash due to a missing FPopulateCompositeDepthPS shader on non-PC platforms.

diff --git a/Engine/Source/Runtime/Renderer/Private/PostProcess/PostProcessCompositePrimitivesCommon.cpp b/Engine/Source/Runtime/Renderer/Private/PostProcess/PostProcessCompositePrimitivesCommon.cpp
--- a/Engine/Source/Runtime/Renderer/Private/PostProcess/PostProcessCompositePrimitivesCommon.cpp
+++ b/Engine/Source/Runtime/Renderer/Private/PostProcess/PostProcessCompositePrimitivesCommon.cpp
@@ -105,7 +105,7 @@ public:
 		RENDER_TARGET_BINDING_SLOTS()
 		END_SHADER_PARAMETER_STRUCT()
 
-		static bool ShouldCompilePermutation(const FGlobalShaderPermutationParameters& Parameters)
+	static bool ShouldCompilePermutation(const FGlobalShaderPermutationParameters& Parameters)
 	{
 		const FPermutationDomain PermutationVector(Parameters.PermutationId);
 		const bool bUseMSAA = PermutationVector.Get<FUseMSAADimension>();
@@ -116,12 +116,8 @@ public:
 			return false;
 		}
 
-#if WITH_EDITOR
 		// Only PC platforms render editor primitives.
 		return IsPCPlatform(Parameters.Platform);
-#else
-		return true;
-#endif
 	}
 };
 
