Fix crash due to a missing FTemporalUpsampleDepthPS + FCompositePrimitiveShaderBaseshader on non-PC platforms. Error caused by there being no config discrepancy for shaders. Allow debug shader compile CVar to be used to compensate on non-PC platforms instead.

diff --git a/Engine/Source/Runtime/Renderer/Private/PostProcess/PostProcessCompositePrimitivesCommon.cpp b/Engine/Source/Runtime/Renderer/Private/PostProcess/PostProcessCompositePrimitivesCommon.cpp
--- a/Engine/Source/Runtime/Renderer/Private/PostProcess/PostProcessCompositePrimitivesCommon.cpp
+++ b/Engine/Source/Runtime/Renderer/Private/PostProcess/PostProcessCompositePrimitivesCommon.cpp
@@ -74,12 +74,8 @@ public:
 
 		static bool ShouldCompilePermutation(const FGlobalShaderPermutationParameters& Parameters)
 	{
-#if WITH_EDITOR
 		// Only PC platforms render editor primitives.
 		return IsPCPlatform(Parameters.Platform);
-#else
-		return true;
-#endif
 	}
 };
 
diff --git a/Engine/Source/Runtime/Renderer/Private/PostProcess/PostProcessCompositePrimitivesCommon.h b/Engine/Source/Runtime/Renderer/Private/PostProcess/PostProcessCompositePrimitivesCommon.h
--- a/Engine/Source/Runtime/Renderer/Private/PostProcess/PostProcessCompositePrimitivesCommon.h
+++ b/Engine/Source/Runtime/Renderer/Private/PostProcess/PostProcessCompositePrimitivesCommon.h
@@ -29,13 +29,7 @@ public:
 			return false;
 		}
 
-#if WITH_EDITOR
-		// Only PC platforms render editor primitives.
 		return IsPCPlatform(Platform);
-#else
-		//Not currently used for any other passes but place 
-		return false;
-#endif
 	}
 
 	static bool ShouldCompilePermutation(const FGlobalShaderPermutationParameters& Parameters)
