Added cvar "r.CaptureNextDeferredShadingRendererFrame" to trigger capture a following deferred scene render with e.g., RenderDoc or PIX.

diff --git a/Engine/Source/Runtime/Renderer/Private/DeferredShadingRenderer.cpp b/Engine/Source/Runtime/Renderer/Private/DeferredShadingRenderer.cpp
--- a/Engine/Source/Runtime/Renderer/Private/DeferredShadingRenderer.cpp
+++ b/Engine/Source/Runtime/Renderer/Private/DeferredShadingRenderer.cpp
@@ -88,6 +88,10 @@
 #include "VariableRateShadingImageManager.h"
 #include "Shadows/ShadowScene.h"
 
+#if !UE_BUILD_SHIPPING
+#include "RenderCaptureInterface.h"
+#endif
+
 extern int32 GNaniteShowStats;
 extern int32 GNanitePickingDomain;
 
@@ -117,10 +121,22 @@ static FAutoConsoleVariableRef CVarEnableAsyncComputeTranslucencyLightingVolumeC
 	ECVF_RenderThreadSafe | ECVF_Scalability
 );
 
-static int32 GRayTracing = 0;
+#if !UE_BUILD_SHIPPING
+
+static int32 GCaptureNextDeferredShadingRendererFrame = -1;
+static FAutoConsoleVariableRef CVarCaptureNextRenderFrame(
+	TEXT("r.CaptureNextDeferredShadingRendererFrame"),
+	GCaptureNextDeferredShadingRendererFrame,
+	TEXT("0 to capture the immideately next frame using e.g. RenderDoc or PIX.\n")
+	TEXT(" > 0: N frames delay\n")
+	TEXT(" < 0: disabled"),
+	ECVF_RenderThreadSafe);
+
+#endif
+
 static TAutoConsoleVariable<int32> CVarRayTracing(
 	TEXT("r.RayTracing"),
-	GRayTracing,
+	0,
 	TEXT("0 to disable ray tracing.\n")
 	TEXT(" 0: off\n")
 	TEXT(" 1: on"),
@@ -2399,7 +2415,7 @@ void FDeferredShadingSceneRenderer::PreGatherDynamicMeshElements(FRDGBuilder& Gr
 		{
 			FViewInfo& View = Views[ViewIndex];
 			NaniteCullingViews.Add(View.ViewFrustum);
-		}
+}
 
 		FNaniteVisibility& NaniteVisibility = Scene->NaniteVisibility[ENaniteMeshPass::BasePass];
 		const FNaniteMaterialCommands& NaniteMaterials = Scene->NaniteMaterials[ENaniteMeshPass::BasePass];
@@ -2500,6 +2516,12 @@ void FDeferredShadingSceneRenderer::Render(FRDGBuilder& GraphBuilder)
 {
 	const bool bNaniteEnabled = IsNaniteEnabled();
 
+#if !UE_BUILD_SHIPPING
+	RenderCaptureInterface::FScopedCapture RenderCapture(GCaptureNextDeferredShadingRendererFrame-- == 0, GraphBuilder, TEXT("DeferredShadingSceneRenderer"));
+	// Prevent overflow every 2B frames.
+	GCaptureNextDeferredShadingRendererFrame = FMath::Max(-1, GCaptureNextDeferredShadingRendererFrame);
+#endif
+
 	GPU_MESSAGE_SCOPE(GraphBuilder);
 
 	// Compute & commit the final state of the entire dependency topology of the renderer.
@@ -3011,7 +3033,7 @@ void FDeferredShadingSceneRenderer::Render(FRDGBuilder& GraphBuilder)
 			DynamicVertexBufferForInitViews.Commit();
 			DynamicReadBufferForInitViews.Commit();
 		}
-	}
+		}
 
 	ExternalAccessQueue.Submit(GraphBuilder);
 
@@ -3286,7 +3308,7 @@ void FDeferredShadingSceneRenderer::Render(FRDGBuilder& GraphBuilder)
 		if (NaniteBasePassVisibility.Visibility)
 		{
 			NaniteBasePassVisibility.Visibility->FinishVisibilityFrame();
-		}
+	}
 	}
 
 	if (FirstStageDepthBuffer)
@@ -3702,7 +3724,7 @@ void FDeferredShadingSceneRenderer::Render(FRDGBuilder& GraphBuilder)
 			bNeedToWaitForRayTracingScene = false;
 		}
 #endif // RHI_RAYTRACING
-	}
+		}
 
 	ExternalAccessQueue.Submit(GraphBuilder);
 
@@ -3895,7 +3917,7 @@ void FDeferredShadingSceneRenderer::Render(FRDGBuilder& GraphBuilder)
 		{
 			// Now remove all the Strata tile stencil tags used by deferred tiled light passes. Make later marks such as responssive AA works.
 			AddClearStencilPass(GraphBuilder, SceneTextures.Depth.Target);
-		}
+	}
 	}
 	else if (HairStrands::HasViewHairStrandsData(Views) && ViewFamily.EngineShowFlags.Lighting)
 	{
