Shader error with new debug draw post process passes, oversight of the shader compilation and client shader checks not matching up correctly.
Swapped the compiles and runtime settings to rely on r.ForceDebugViewModes/AllowDebugViewModes instead so this works cross platform.

diff --git a/Engine/Source/Runtime/Renderer/Private/PostProcess/PostProcessCompositePrimitivesCommon.cpp b/Engine/Source/Runtime/Renderer/Private/PostProcess/PostProcessCompositePrimitivesCommon.cpp
--- a/Engine/Source/Runtime/Renderer/Private/PostProcess/PostProcessCompositePrimitivesCommon.cpp
+++ b/Engine/Source/Runtime/Renderer/Private/PostProcess/PostProcessCompositePrimitivesCommon.cpp
@@ -74,8 +74,8 @@ public:
 
 		static bool ShouldCompilePermutation(const FGlobalShaderPermutationParameters& Parameters)
 	{
-		// Only PC platforms render editor primitives.
-		return IsPCPlatform(Parameters.Platform);
+		// Only render if debug viewmodes are allowed.
+		return AllowDebugViewmodes(Parameters.Platform);
 	}
 };
 
@@ -112,8 +112,8 @@ public:
 			return false;
 		}
 
-		// Only PC platforms render editor primitives.
-		return IsPCPlatform(Parameters.Platform);
+		// Only render if debug viewmodes are allowed.
+		return AllowDebugViewmodes(Parameters.Platform);
 	}
 };
 
diff --git a/Engine/Source/Runtime/Renderer/Private/SceneRendering.cpp b/Engine/Source/Runtime/Renderer/Private/SceneRendering.cpp
--- a/Engine/Source/Runtime/Renderer/Private/SceneRendering.cpp
+++ b/Engine/Source/Runtime/Renderer/Private/SceneRendering.cpp
@@ -4003,7 +4003,7 @@ bool FSceneRenderer::ShouldCompositeEditorPrimitives(const FViewInfo& View)
 #if UE_ENABLE_DEBUG_DRAWING
 bool FSceneRenderer::ShouldCompositeDebugPrimitivesInPostProcess(const FViewInfo& View)
 {
-	return View.DebugSimpleElementCollector.HasAnyPrimitives();
+	return AllowDebugViewmodes() && View.DebugSimpleElementCollector.HasAnyPrimitives();
 }
 #endif
 
