Ensure that OnPreRender is called with the MobileShadringRenderer even if particles aren't being shown to mirror what is being done with OnPostRenderOpaque.  Resolves crash when GPU sorting of particles during cubemap generation.

diff --git a/Engine/Source/Runtime/Renderer/Private/MobileShadingRenderer.cpp b/Engine/Source/Runtime/Renderer/Private/MobileShadingRenderer.cpp
--- a/Engine/Source/Runtime/Renderer/Private/MobileShadingRenderer.cpp
+++ b/Engine/Source/Runtime/Renderer/Private/MobileShadingRenderer.cpp
@@ -980,11 +980,13 @@ void FMobileSceneRenderer::Render(FRDGBuilder& GraphBuilder)
 	}
 
 	// Notify the FX system that the scene is about to be rendered.
-	if (FXSystem && ViewFamily.EngineShowFlags.Particles)
+	if (FXSystem)
 	{
 		FXSystem->PreRender(GraphBuilder, GetSceneViews(), true /*bAllowGPUParticleUpdate*/);
 		if (FGPUSortManager* GPUSortManager = FXSystem->GetGPUSortManager())
 		{
+			// if GPUSortManager::OnPostRenderOpaque is called below (from RenderOpaqueFX) we must also call OnPreRender (as it sets up
+			// the internal state of the GPUSortManager).  Any optimization to skip this block needs to take that into consideration.
 			GPUSortManager->OnPreRender(GraphBuilder);
 		}
 	}
