Fixes check() failure in VisualizeDOF

diff --git a/Engine/Shaders/Private/PostProcessVisualizeDOF.usf b/Engine/Shaders/Private/PostProcessVisualizeDOF.usf
--- a/Engine/Shaders/Private/PostProcessVisualizeDOF.usf
+++ b/Engine/Shaders/Private/PostProcessVisualizeDOF.usf
@@ -10,6 +10,9 @@ SamplerState ColorSampler;
 Texture2D DepthTexture;
 SamplerState DepthSampler;
 
+FScreenTransform SvPositionToColorUV;
+FScreenTransform SvPositionToDepthUV;
+
 float4 NearColor;
 float4 FarColor;
 
@@ -21,11 +24,9 @@ void VisualizeDOFPS(
 {
 	OutColor = 0;
 
-	float2 UV = UVAndScreenPos.xy; 
-	int2 PixelPos = (int2)SvPosition.xy;
-	float2 ViewLocalUV = float2(UVAndScreenPos.z * 0.5f + 0.5f, 0.5f - 0.5f * UVAndScreenPos.w);
+	float SceneDeviceZ = DepthTexture.SampleLevel(DepthSampler, ApplyScreenTransform(SvPosition, SvPositionToColorUV), 0).r;
+	half3 SceneColor = Texture2DSample(ColorTexture, ColorSampler, ApplyScreenTransform(SvPosition, SvPositionToDepthUV)).rgb;
 
-	float SceneDeviceZ = DepthTexture.SampleLevel(DepthSampler, UV, 0).r;
 	float SceneDepth = ConvertFromDeviceZ(SceneDeviceZ);
 
 	// 0..1
@@ -33,11 +34,10 @@ void VisualizeDOFPS(
 	// 0..1
 	half Far = CalcUnfocusedPercentCustomBound(SceneDepth, 0, 1);
 
+	OutColor = 0;
 	OutColor.rgb = lerp(OutColor.rgb, NearColor.rgb, Near);
 	OutColor.rgb = lerp(OutColor.rgb, FarColor.rgb, Far);
 
-	half3 SceneColor = Texture2DSample(ColorTexture, ColorSampler, UV).rgb;
-
 	// blend in a bit of the scene color to make navigation easier
 	OutColor.rgb = lerp(saturate(OutColor.rgb), Luminance(SceneColor), 0.1f);
 }
\ No newline at end of file
diff --git a/Engine/Source/Runtime/Renderer/Private/PostProcess/PostProcessBokehDOF.cpp b/Engine/Source/Runtime/Renderer/Private/PostProcess/PostProcessBokehDOF.cpp
--- a/Engine/Source/Runtime/Renderer/Private/PostProcess/PostProcessBokehDOF.cpp
+++ b/Engine/Source/Runtime/Renderer/Private/PostProcess/PostProcessBokehDOF.cpp
@@ -26,6 +26,8 @@ public:
 		SHADER_PARAMETER_RDG_TEXTURE(Texture2D, DepthTexture)
 		SHADER_PARAMETER_SAMPLER(SamplerState, ColorSampler)
 		SHADER_PARAMETER_SAMPLER(SamplerState, DepthSampler)
+		SHADER_PARAMETER(FScreenTransform, SvPositionToColorUV)
+		SHADER_PARAMETER(FScreenTransform, SvPositionToDepthUV)
 		SHADER_PARAMETER(FVector4f, DepthOfFieldParams)
 		SHADER_PARAMETER(FVector4f, NearColor)
 		SHADER_PARAMETER(FVector4f, FarColor)
@@ -39,7 +41,6 @@ FScreenPassTexture AddVisualizeDOFPass(FRDGBuilder& GraphBuilder, const FViewInf
 {
 	check(Inputs.SceneColor.IsValid());
 	check(Inputs.SceneDepth.IsValid());
-	checkf(Inputs.SceneColor.ViewRect == Inputs.SceneDepth.ViewRect, TEXT("VisualizeDOF requires that the scene and depth view rects match."));
 
 	FScreenPassRenderTarget Output = Inputs.OverrideOutput;
 
@@ -59,6 +60,17 @@ FScreenPassTexture AddVisualizeDOFPass(FRDGBuilder& GraphBuilder, const FViewInf
 	PassParameters->DepthTexture = Inputs.SceneDepth.Texture;
 	PassParameters->ColorSampler = PointClampSampler;
 	PassParameters->DepthSampler = PointClampSampler;
+	{
+		FScreenTransform SvPositionToViewportUV = FScreenTransform::SvPositionToViewportUV(Output.ViewRect);
+		PassParameters->SvPositionToColorUV = 
+			SvPositionToViewportUV *
+			FScreenTransform::ChangeTextureBasisFromTo(
+				InputViewport, FScreenTransform::ETextureBasis::ViewportUV, FScreenTransform::ETextureBasis::TextureUV);
+		PassParameters->SvPositionToDepthUV =
+			SvPositionToViewportUV *
+			FScreenTransform::ChangeTextureBasisFromTo(
+				FScreenPassTextureViewport(Inputs.SceneDepth), FScreenTransform::ETextureBasis::ViewportUV, FScreenTransform::ETextureBasis::TextureUV);
+	}
 	PassParameters->DepthOfFieldParams = GetDepthOfFieldParameters(View.FinalPostProcessSettings);
 	PassParameters->NearColor = FLinearColor(0, 0.8f, 0, 0);
 	PassParameters->FarColor = FLinearColor(0, 0, 0.8f, 0);
