VulkanRHI: Deprecate IsVulkanSM5Platform() and prefer the use of IsVulkanPlatform() with IsFeatureLevelSupported().  This ensures SM6 can use these code paths as well (as was inteded where it's used).

#rnx

diff --git a/Engine/Source/Runtime/RHI/Public/DataDrivenShaderPlatformInfo.h b/Engine/Source/Runtime/RHI/Public/DataDrivenShaderPlatformInfo.h
--- a/Engine/Source/Runtime/RHI/Public/DataDrivenShaderPlatformInfo.h
+++ b/Engine/Source/Runtime/RHI/Public/DataDrivenShaderPlatformInfo.h
@@ -831,6 +831,7 @@ inline bool IsVulkanPlatform(const FStaticShaderPlatform Platform)
 	return FDataDrivenShaderPlatformInfo::GetIsLanguageVulkan(Platform);
 }
 
+UE_DEPRECATED(5.3, "Use a combination of IsVulkanPlatform and IsFeatureLevelSupported (or GetMaxFeatureLevel) instead.")
 inline bool IsVulkanSM5Platform(const FStaticShaderPlatform Platform)
 {
 	return FDataDrivenShaderPlatformInfo::GetIsLanguageVulkan(Platform)
diff --git a/Engine/Source/Runtime/Renderer/Private/HairStrands/HairStrandsInterface.cpp b/Engine/Source/Runtime/Renderer/Private/HairStrands/HairStrandsInterface.cpp
--- a/Engine/Source/Runtime/Renderer/Private/HairStrands/HairStrandsInterface.cpp
+++ b/Engine/Source/Runtime/Renderer/Private/HairStrands/HairStrandsInterface.cpp
@@ -200,7 +200,7 @@ bool IsHairStrandsSupported(EHairStrandsShaderType Type, EShaderPlatform Platfor
 	case EHairStrandsShaderType::Strands: return IsHairStrandsGeometrySupported(Platform);
 	case EHairStrandsShaderType::Cards:	  return Cards_Meshes_All;
 	case EHairStrandsShaderType::Meshes:  return Cards_Meshes_All;
-	case EHairStrandsShaderType::Tool:	  return (IsD3DPlatform(Platform) || IsVulkanSM5Platform(Platform)) && IsPCPlatform(Platform) && IsFeatureLevelSupported(Platform, ERHIFeatureLevel::SM5);
+	case EHairStrandsShaderType::Tool:	  return (IsD3DPlatform(Platform) || IsVulkanPlatform(Platform)) && IsPCPlatform(Platform) && IsFeatureLevelSupported(Platform, ERHIFeatureLevel::SM5);
 	case EHairStrandsShaderType::All:	  return Cards_Meshes_All && !bIsMobile;
 	}
 	return false;
