Pack fields in FPrimitiveSceneInfo to reduce memory footprint

diff --git a/Engine/Source/Runtime/Renderer/Private/PrimitiveSceneInfo.cpp b/Engine/Source/Runtime/Renderer/Private/PrimitiveSceneInfo.cpp
--- a/Engine/Source/Runtime/Renderer/Private/PrimitiveSceneInfo.cpp
+++ b/Engine/Source/Runtime/Renderer/Private/PrimitiveSceneInfo.cpp
@@ -170,12 +170,24 @@ FPrimitiveSceneInfo::FPrimitiveSceneInfo(UPrimitiveComponent* InComponent, FScen
 	IndirectLightingCacheAllocation(NULL),
 	CachedPlanarReflectionProxy(NULL),
 	CachedReflectionCaptureProxy(NULL),
-	bNeedsCachedReflectionCaptureUpdate(true),
 	DefaultDynamicHitProxy(NULL),
-	LightList(NULL),
 	LastRenderTime(-FLT_MAX),
+	LightList(NULL),
 	Scene(InScene),
 	NumMobileDynamicLocalLights(0),
+	PackedIndex(INDEX_NONE),
+	PersistentIndex(FPersistentPrimitiveIndex{ INDEX_NONE }),
+	ComponentForDebuggingOnly(InComponent),
+	bNeedsUniformBufferUpdate(false),
+	bIndirectLightingCacheBufferDirty(false),
+	bRegisteredVirtualTextureProducerCallback(false),
+	bRegisteredWithVelocityData(false),
+	bCacheShadowAsStatic(InComponent->Mobility != EComponentMobility::Movable),
+	bNaniteRasterBinsRenderCustomDepth(false),
+	bPendingAddToScene(false),
+	bPendingAddStaticMeshes(false),
+	bPendingFlushVirtualTexture(false),
+	bNeedsCachedReflectionCaptureUpdate(true),
 	bShouldRenderInMainPass(InComponent->SceneProxy->ShouldRenderInMainPass()),
 	bVisibleInRealTimeSkyCapture(InComponent->SceneProxy->IsVisibleInRealTimeSkyCaptures()),
 #if RHI_RAYTRACING
@@ -189,18 +201,6 @@ FPrimitiveSceneInfo::FPrimitiveSceneInfo(UPrimitiveComponent* InComponent, FScen
 	bCachedRaytracingDataDirty(false),
 	CoarseMeshStreamingHandle(INDEX_NONE),
 #endif
-	PackedIndex(INDEX_NONE),
-	PersistentIndex(FPersistentPrimitiveIndex{ INDEX_NONE }),
-	ComponentForDebuggingOnly(InComponent),
-	bNeedsUniformBufferUpdate(false),
-	bIndirectLightingCacheBufferDirty(false),
-	bRegisteredVirtualTextureProducerCallback(false),
-	bRegisteredWithVelocityData(false),
-	bCacheShadowAsStatic(InComponent->Mobility != EComponentMobility::Movable),
-	bNaniteRasterBinsRenderCustomDepth(false),
-	bPendingAddToScene(false),
-	bPendingAddStaticMeshes(false),
-	bPendingFlushVirtualTexture(false),
 	LevelUpdateNotificationIndex(INDEX_NONE),
 	InstanceSceneDataOffset(INDEX_NONE),
 	NumInstanceSceneDataEntries(0),
diff --git a/Engine/Source/Runtime/Renderer/Public/PrimitiveSceneInfo.h b/Engine/Source/Runtime/Renderer/Public/PrimitiveSceneInfo.h
--- a/Engine/Source/Runtime/Renderer/Public/PrimitiveSceneInfo.h
+++ b/Engine/Source/Runtime/Renderer/Public/PrimitiveSceneInfo.h
@@ -367,9 +367,6 @@
 	/** Mapping from instance index in this primitive to index in the LumenPrimitiveGroup array. */
 	TArray<int32, TInlineAllocator<1>> LumenPrimitiveGroupIndices;
 
-	/** Whether the primitive is newly registered or moved and CachedReflectionCaptureProxy needs to be updated on the next render. */
-	uint32 bNeedsCachedReflectionCaptureUpdate : 1;
-
 	static const uint32 MaxCachedReflectionCaptureProxies = 3;
 	const FReflectionCaptureProxy* CachedReflectionCaptureProxies[MaxCachedReflectionCaptureProxies];
 	
@@ -382,46 +379,18 @@
 	/** The ID of the hit proxy which is used to represent the primitive's dynamic elements. */
 	FHitProxyId DefaultDynamicHitProxyId;
 
-	/** The list of lights affecting this primitive. */
-	class FLightPrimitiveInteraction* LightList;
-
 	/** Last render time in seconds since level started play. */
 	float LastRenderTime;
 
+	/** The list of lights affecting this primitive. */
+	class FLightPrimitiveInteraction* LightList;
+
 	/** The scene the primitive is in. */
 	FScene* Scene;
 
 	/** The number of local lights with dynamic lighting for mobile */
 	int32 NumMobileDynamicLocalLights;
 
-	/** Set to true for the primitive to be rendered in the main pass to be visible in a view. */
-	bool bShouldRenderInMainPass : 1;
-
-	/** Set to true for the primitive to be rendered into the real-time sky light reflection capture. */
-	bool bVisibleInRealTimeSkyCapture : 1;
-
-#if RHI_RAYTRACING
-	bool bDrawInGame : 1;
-	bool bRayTracingFarField : 1;
-	bool bIsVisibleInSceneCaptures : 1;
-	bool bIsVisibleInSceneCapturesOnly : 1;
-	bool bIsRayTracingRelevant : 1;
-	bool bIsRayTracingStaticRelevant : 1;
-	bool bIsVisibleInRayTracing : 1;
-	bool bCachedRaytracingDataDirty : 1;
-	Nanite::CoarseMeshStreamingHandle CoarseMeshStreamingHandle;
-
-	TArray<TArray<int32, TInlineAllocator<2>>> CachedRayTracingMeshCommandIndicesPerLOD;
-
-	TArray<uint64> CachedRayTracingMeshCommandsHashPerLOD;
-	// TODO: this should be placed in FRayTracingScene and we have a pointer/handle here. It's here for now for PoC
-	FRayTracingGeometryInstance CachedRayTracingInstance;
-	bool bCachedRayTracingInstanceAnySegmentsDecal : 1;
-	bool bCachedRayTracingInstanceAllSegmentsDecal : 1;
-	TArray<FBoxSphereBounds> CachedRayTracingInstanceWorldBounds;
-	int32 SmallestRayTracingInstanceWorldBoundsIndex;
-#endif
-
 	/** Initialization constructor. */
 	FPrimitiveSceneInfo(UPrimitiveComponent* InPrimitive,FScene* InScene);
 
@@ -650,6 +619,9 @@
 	 */
 	const UPrimitiveComponent* ComponentForDebuggingOnly;
 
+	/** These flags carry information about which runtime virtual textures are bound to this primitive. */
+	FPrimitiveVirtualTextureFlags RuntimeVirtualTextureFlags;
+
 	/** If this is TRUE, this primitive's uniform buffer needs to be updated before it can be rendered. */
 	bool bNeedsUniformBufferUpdate : 1;
 
@@ -677,6 +649,39 @@
 	/** True if the primitive is queued to have its virtual texture flushed. */
 	bool bPendingFlushVirtualTexture : 1;
 
+public:
+	/** Whether the primitive is newly registered or moved and CachedReflectionCaptureProxy needs to be updated on the next render. */
+	bool bNeedsCachedReflectionCaptureUpdate : 1;
+
+	/** Set to true for the primitive to be rendered in the main pass to be visible in a view. */
+	bool bShouldRenderInMainPass : 1;
+
+	/** Set to true for the primitive to be rendered into the real-time sky light reflection capture. */
+	bool bVisibleInRealTimeSkyCapture : 1;
+
+#if RHI_RAYTRACING
+	bool bDrawInGame : 1;
+	bool bRayTracingFarField : 1;
+	bool bIsVisibleInSceneCaptures : 1;
+	bool bIsVisibleInSceneCapturesOnly : 1;
+	bool bIsRayTracingRelevant : 1;
+	bool bIsRayTracingStaticRelevant : 1;
+	bool bIsVisibleInRayTracing : 1;
+	bool bCachedRaytracingDataDirty : 1;
+	bool bCachedRayTracingInstanceAnySegmentsDecal : 1;
+	bool bCachedRayTracingInstanceAllSegmentsDecal : 1;
+	Nanite::CoarseMeshStreamingHandle CoarseMeshStreamingHandle;
+
+	TArray<TArray<int32, TInlineAllocator<2>>> CachedRayTracingMeshCommandIndicesPerLOD;
+
+	TArray<uint64> CachedRayTracingMeshCommandsHashPerLOD;
+	// TODO: this should be placed in FRayTracingScene and we have a pointer/handle here. It's here for now for PoC
+	FRayTracingGeometryInstance CachedRayTracingInstance;
+	TArray<FBoxSphereBounds> CachedRayTracingInstanceWorldBounds;
+	int32 SmallestRayTracingInstanceWorldBoundsIndex;
+#endif
+
+private:
 	/** Index into the scene's PrimitivesNeedingLevelUpdateNotification array for this primitive scene info level. */
 	int32 LevelUpdateNotificationIndex;
 
@@ -714,9 +719,6 @@
 	/** Removes cached mesh draw commands for all meshes. */
 	void RemoveCachedMeshDrawCommands();
 
-	/** These flags carry information about which runtime virtual textures are bound to this primitive. */
-	FPrimitiveVirtualTextureFlags RuntimeVirtualTextureFlags;
-
 	/** Creates or add ref's cached draw commands for each unique material instance found within the scene. */
 	static void CacheNaniteDrawCommands(FScene* Scene, const TArrayView<FPrimitiveSceneInfo*>& SceneInfos);
 
