Fixed crash in Nanite visibility due to shadow query being waited on at the wrong time.

diff --git a/Engine/Source/Runtime/Renderer/Private/DeferredShadingRenderer.cpp b/Engine/Source/Runtime/Renderer/Private/DeferredShadingRenderer.cpp
--- a/Engine/Source/Runtime/Renderer/Private/DeferredShadingRenderer.cpp
+++ b/Engine/Source/Runtime/Renderer/Private/DeferredShadingRenderer.cpp
@@ -2423,7 +2423,7 @@ void FDeferredShadingSceneRenderer::PreGatherDynamicMeshElements(FRDGBuilder& Gr
 		{
 			FViewInfo& View = Views[ViewIndex];
 			NaniteCullingViews.Add(View.ViewFrustum);
-}
+		}
 
 		FNaniteVisibility& NaniteVisibility = Scene->NaniteVisibility[ENaniteMeshPass::BasePass];
 		const FNaniteMaterialCommands& NaniteMaterials = Scene->NaniteMaterials[ENaniteMeshPass::BasePass];
@@ -3041,7 +3041,7 @@ void FDeferredShadingSceneRenderer::Render(FRDGBuilder& GraphBuilder)
 			DynamicVertexBufferForInitViews.Commit();
 			DynamicReadBufferForInitViews.Commit();
 		}
-		}
+	}
 
 	ExternalAccessQueue.Submit(GraphBuilder);
 
@@ -3312,11 +3312,6 @@ void FDeferredShadingSceneRenderer::Render(FRDGBuilder& GraphBuilder)
 				}
 			}
 		}
-
-		if (NaniteBasePassVisibility.Visibility)
-		{
-			NaniteBasePassVisibility.Visibility->FinishVisibilityFrame();
-	}
 	}
 
 	if (FirstStageDepthBuffer)
@@ -3732,7 +3727,7 @@ void FDeferredShadingSceneRenderer::Render(FRDGBuilder& GraphBuilder)
 			bNeedToWaitForRayTracingScene = false;
 		}
 #endif // RHI_RAYTRACING
-		}
+	}
 
 	ExternalAccessQueue.Submit(GraphBuilder);
 
@@ -3925,7 +3920,7 @@ void FDeferredShadingSceneRenderer::Render(FRDGBuilder& GraphBuilder)
 		{
 			// Now remove all the Strata tile stencil tags used by deferred tiled light passes. Make later marks such as responssive AA works.
 			AddClearStencilPass(GraphBuilder, SceneTextures.Depth.Target);
-	}
+		}
 	}
 	else if (HairStrands::HasViewHairStrandsData(Views) && ViewFamily.EngineShowFlags.Lighting)
 	{
@@ -4445,6 +4440,12 @@ void FDeferredShadingSceneRenderer::Render(FRDGBuilder& GraphBuilder)
 	{
 		Views[ViewIndex].PrevViewInfo = FPreviousViewInfo();
 	}
+
+	if (NaniteBasePassVisibility.Visibility)
+	{
+		NaniteBasePassVisibility.Visibility->FinishVisibilityFrame();
+		NaniteBasePassVisibility.Visibility = nullptr;
+	}
 }
 
 #if RHI_RAYTRACING
