Change AddDebugPrimitivePass function to create upscaled depth texture from output texture extent rather than scene color texture extent to avoid mismatched color and depth size when a resized texture output override is used. Also moved a comment.

diff --git a/Engine/Source/Runtime/Renderer/Private/PostProcess/PostProcessCompositeDebugPrimitives.cpp b/Engine/Source/Runtime/Renderer/Private/PostProcess/PostProcessCompositeDebugPrimitives.cpp
--- a/Engine/Source/Runtime/Renderer/Private/PostProcess/PostProcessCompositeDebugPrimitives.cpp
+++ b/Engine/Source/Runtime/Renderer/Private/PostProcess/PostProcessCompositeDebugPrimitives.cpp
@@ -21,7 +21,6 @@ FScreenPassTexture AddDebugPrimitivePass(
 	
 	const uint32 NumMSAASamples = View.GetSceneTextures().Config.NumSamples;
 	//Set sizing values to match the SceneColor size
-	const FIntPoint Extent = Inputs.SceneColor.Texture->Desc.Extent;
 	const FIntRect ViewRect = Inputs.SceneColor.ViewRect;
 	const FViewInfo* DebugView = CreateCompositePrimitiveView(View, ViewRect, NumMSAASamples);
 
@@ -36,11 +35,13 @@ FScreenPassTexture AddDebugPrimitivePass(
 		//Create a new depth output for the final draw
 		Output = FScreenPassRenderTarget(DebugPrimitiveColor, ViewRect, ERenderTargetLoadAction::ENoAction);
 	}
+	FIntPoint Extent = Output.Texture->Desc.Extent;
 
 	//Prepare output textures for composite draw
 	const FScreenPassTextureViewport OutputViewport(Output);
 	FRDGTextureRef DebugPrimitiveDepth = CreateCompositeDepthTexture(GraphBuilder, Extent, NumMSAASamples);
-	FScreenPassTexture SceneDepth = Inputs.SceneDepth;//Inputs is const so create a over-ridable texture reference
+	//Inputs is const so create a over-ridable texture reference
+	FScreenPassTexture SceneDepth = Inputs.SceneDepth;
 	FVector2f SceneDepthJitter = FVector2f(View.TemporalJitterPixels);
 
 	if (IsTemporalAccumulationBasedMethod(View.AntiAliasingMethod))
