Turns on bloom and TSR async compute by default on all capable platforms

diff --git a/Engine/Source/Runtime/Renderer/Private/PostProcess/PostProcessFFTBloom.cpp b/Engine/Source/Runtime/Renderer/Private/PostProcess/PostProcessFFTBloom.cpp
--- a/Engine/Source/Runtime/Renderer/Private/PostProcess/PostProcessFFTBloom.cpp
+++ b/Engine/Source/Runtime/Renderer/Private/PostProcess/PostProcessFFTBloom.cpp
@@ -24,7 +24,7 @@ TAutoConsoleVariable<int32> CVarBloomCacheKernel(
 	ECVF_RenderThreadSafe);
 
 TAutoConsoleVariable<int32> CVarAsynComputeFFTBloom(
-	TEXT("r.Bloom.AsyncCompute"), 0,
+	TEXT("r.Bloom.AsyncCompute"), 1,
 	TEXT("Whether to run FFT bloom on async compute.\n"),
 	ECVF_RenderThreadSafe);
 
diff --git a/Engine/Source/Runtime/Renderer/Private/PostProcess/TemporalSuperResolution.cpp b/Engine/Source/Runtime/Renderer/Private/PostProcess/TemporalSuperResolution.cpp
--- a/Engine/Source/Runtime/Renderer/Private/PostProcess/TemporalSuperResolution.cpp
+++ b/Engine/Source/Runtime/Renderer/Private/PostProcess/TemporalSuperResolution.cpp
@@ -93,7 +93,7 @@ TAutoConsoleVariable<int32> CVarTSRRejectionAntiAliasingQuality(
 	ECVF_Scalability | ECVF_RenderThreadSafe);
 
 TAutoConsoleVariable<int32> CVarTSRAsyncCompute(
-	TEXT("r.TSR.AsyncCompute"), 0,
+	TEXT("r.TSR.AsyncCompute"), 2,
 	TEXT("Whether to run TSR on async compute. Some TSR passes can overlap with previous passe.\n")
 	TEXT(" 0: Disabled (default);\n")
 	TEXT(" 1: Only ClearPrevTextures pass;\n")
