Fixes crash with r.TSR.History.ScreenPercentage>100 with r.PostProcessing.QuarterResolutionDownsample=1 and motion blur disabled

diff --git a/Engine/Source/Runtime/Renderer/Private/PostProcess/TemporalSuperResolution.cpp b/Engine/Source/Runtime/Renderer/Private/PostProcess/TemporalSuperResolution.cpp
--- a/Engine/Source/Runtime/Renderer/Private/PostProcess/TemporalSuperResolution.cpp
+++ b/Engine/Source/Runtime/Renderer/Private/PostProcess/TemporalSuperResolution.cpp
@@ -2114,7 +2114,7 @@ ITemporalUpscaler::FOutputs AddTemporalSuperResolutionPasses(
 				FScreenTransform::ETextureBasis::ViewportUV, FScreenTransform::ETextureBasis::TexelPosition));
 		PassParameters->OutputViewRectMin = OutputRect.Min;
 		PassParameters->OutputViewRectMax = OutputRect.Max;
-		PassParameters->bGenerateOutputMip1 = (PassInputs.bGenerateOutputMip1 || PassInputs.bGenerateSceneColorHalfRes) ? 1 : 0;
+		PassParameters->bGenerateOutputMip1 = false;
 		PassParameters->HistoryValidityMultiply = float(HistorySize.X * HistorySize.Y) / float(OutputRect.Width() * OutputRect.Height());
 
 		PassParameters->UpdateHistoryOutputTexture = UpdateHistoryOutputTexture;
@@ -2122,10 +2122,12 @@ ITemporalUpscaler::FOutputs AddTemporalSuperResolutionPasses(
 		PassParameters->SceneColorOutputMip0 = GraphBuilder.CreateUAV(FRDGTextureUAVDesc(SceneColorOutputTexture, /* InMipLevel = */ 0));
 		if (PassInputs.bGenerateOutputMip1)
 		{
+			PassParameters->bGenerateOutputMip1 = true;
 			PassParameters->SceneColorOutputMip1 = GraphBuilder.CreateUAV(FRDGTextureUAVDesc(SceneColorOutputTexture, /* InMipLevel = */ 1));
 		}
 		else if (SceneColorOutputHalfResTexture)
 		{
+			PassParameters->bGenerateOutputMip1 = true;
 			PassParameters->SceneColorOutputMip1 = GraphBuilder.CreateUAV(FRDGTextureUAVDesc(SceneColorOutputHalfResTexture));
 		}
 		else
