* Fix incorrect atlas packing when texture are streamed in/out and require a full atlas refit.
* Fix crash when rect light atlas force update is enabled and light get toggle on and off.

The incorrect packing was causing rect light textures to become blurrier as the light get toggled on/off.

diff --git a/Engine/Source/Runtime/Renderer/Private/RectLightTextureManager.cpp b/Engine/Source/Runtime/Renderer/Private/RectLightTextureManager.cpp
--- a/Engine/Source/Runtime/Renderer/Private/RectLightTextureManager.cpp
+++ b/Engine/Source/Runtime/Renderer/Private/RectLightTextureManager.cpp
@@ -323,11 +323,13 @@ static FRDGTextureRef AddRectLightDebugInfoPass(
 	
 	uint32 ValidSlotCount = 0;
 	uint32 OccupiedPixels = 0;
+	TArray<FAtlasSlot> ValidSlots;
+	ValidSlots.Reserve(GRectLightTextureManager.AtlasSlots.Num());
 	for (const FAtlasSlot& Slot : GRectLightTextureManager.AtlasSlots)
 	{
 		if (Slot.IsValid())
 		{
-			ValidSlotCount++;
+			ValidSlots.Add(Slot);
 			OccupiedPixels += Slot.Rect.Resolution.X * Slot.Rect.Resolution.Y;
 		}
 	}
@@ -341,7 +343,7 @@ static FRDGTextureRef AddRectLightDebugInfoPass(
 #endif
 
 	// Create a buffer with all the valid slot to highlight them on the debug view
-	FRDGBufferRef SlotBuffer = CreateSlotBuffer(GraphBuilder, GRectLightTextureManager.AtlasSlots, TEXT("RectLight.AtlasSlotBuffer"));
+	FRDGBufferRef SlotBuffer = CreateSlotBuffer(GraphBuilder, ValidSlots, TEXT("RectLight.AtlasSlotBuffer"));
 	FRDGBufferRef HorizonBuffer = CreateSlotBuffer(GraphBuilder, Horizons, TEXT("RectLight.HorizonBuffer"));
 	FRDGBufferRef FreeBuffer = CreateSlotBuffer(GraphBuilder, FreeRects, TEXT("RectLight.FreeBuffer"));
 
@@ -351,7 +353,7 @@ static FRDGTextureRef AddRectLightDebugInfoPass(
 	Parameters->AtlasMaxMipLevel = AtlasTexture->Desc.NumMips;
 	Parameters->AtlasSourceTextureMIPBias = GRectLightTextureManager.AtlasLayout.SourceTextureMIPBias;
 	Parameters->Occupancy = OccupiedPixels / float(AtlasTexture->Desc.Extent.X * AtlasTexture->Desc.Extent.Y);
-	Parameters->SlotCount = ValidSlotCount;
+	Parameters->SlotCount = ValidSlots.Num();
 	Parameters->HorizonCount = Horizons.Num();
 	Parameters->FreeCount = FreeRects.Num();
 	Parameters->OutputResolution = OutputResolution;
@@ -1342,7 +1344,7 @@ static void PackAtlas(
 		NewSlots.SetNum(0, false);
 		
 		FIntPoint CurrentAtlasResolution = Layout.AtlasResolution;
-		int32 CurrentSourceTextureMIPBias = Layout.SourceTextureMIPBias;
+		int32 CurrentSourceTextureMIPBias = 0; // When a refit is done, restart with a MIP bias of 0, to ensure we use the full potential space of the atlas
 		bool bIsPackingValid = false;
 		while (!bIsPackingValid)
 		{
