VulkanRHI: Fix cases of ENoAction/StencilNop (i.e.discard) being used in cases where Stencil content needs to be preserved.

#rnx

diff --git a/Engine/Source/Runtime/Renderer/Private/HairStrands/HairStrandsComposition.cpp b/Engine/Source/Runtime/Renderer/Private/HairStrands/HairStrandsComposition.cpp
--- a/Engine/Source/Runtime/Renderer/Private/HairStrands/HairStrandsComposition.cpp
+++ b/Engine/Source/Runtime/Renderer/Private/HairStrands/HairStrandsComposition.cpp
@@ -212,7 +212,7 @@ static void AddHairVisibilityComposeSamplePass(
 	Parameters->TemporalLayerCount = 1;
 	Parameters->HairTemporalAccumulationTexture = bTemporal ? HairTemporalAccumulationTexture : GSystemTextures.GetBlackDummy(GraphBuilder);
 	Parameters->RenderTargets[0] = FRenderTargetBinding(OutColorTexture, bTemporal ? ERenderTargetLoadAction::ELoad : (bTemporalLayeringEnabled ? ERenderTargetLoadAction::EClear : ERenderTargetLoadAction::ELoad) );
-	Parameters->RenderTargets.DepthStencil = FDepthStencilBinding(OutDepthTexture, ERenderTargetLoadAction::ELoad, ERenderTargetLoadAction::ENoAction, FExclusiveDepthStencil::DepthWrite_StencilNop);	
+	Parameters->RenderTargets.DepthStencil = FDepthStencilBinding(OutDepthTexture, ERenderTargetLoadAction::ELoad, ERenderTargetLoadAction::ELoad, FExclusiveDepthStencil::DepthWrite_StencilRead);
 
 	const bool bDebugComposition = View.Family->EngineShowFlags.LODColoration;
 	FHairVisibilityComposeSamplePS::FPermutationDomain PermutationVector;
diff --git a/Engine/Source/Runtime/Renderer/Private/HairStrands/HairStrandsVisibility.cpp b/Engine/Source/Runtime/Renderer/Private/HairStrands/HairStrandsVisibility.cpp
--- a/Engine/Source/Runtime/Renderer/Private/HairStrands/HairStrandsVisibility.cpp
+++ b/Engine/Source/Runtime/Renderer/Private/HairStrands/HairStrandsVisibility.cpp
@@ -2245,8 +2245,8 @@ static FHairPrimaryTransmittance AddHairViewTransmittancePass(
 	PassParameters->RenderTargets.DepthStencil = FDepthStencilBinding(
 		SceneDepthTexture,
 		ERenderTargetLoadAction::ELoad,
-		ERenderTargetLoadAction::ENoAction,
-		FExclusiveDepthStencil::DepthRead_StencilNop);
+		ERenderTargetLoadAction::ELoad,
+		FExclusiveDepthStencil::DepthRead_StencilRead);
 	AddHairVisibilityCommonPass(GraphBuilder, Scene, ViewInfo, MacroGroupDatas, RenderMode, PassParameters, InstanceCullingManager);
 
 	return Out;
diff --git a/Engine/Source/Runtime/Renderer/Private/IndirectLightRendering.cpp b/Engine/Source/Runtime/Renderer/Private/IndirectLightRendering.cpp
--- a/Engine/Source/Runtime/Renderer/Private/IndirectLightRendering.cpp
+++ b/Engine/Source/Runtime/Renderer/Private/IndirectLightRendering.cpp
@@ -1386,7 +1386,8 @@ void FDeferredShadingSceneRenderer::RenderDiffuseIndirectAndAmbientOcclusion(
 
 				if (bUseDepthBounds)
 				{
-					PassParameters->RenderTargets.DepthStencil = FDepthStencilBinding(SceneTextures.Depth.Target, ERenderTargetLoadAction::ELoad, FExclusiveDepthStencil::DepthRead_StencilNop);
+					PassParameters->RenderTargets.DepthStencil = FDepthStencilBinding(SceneTextures.Depth.Target, 
+						ERenderTargetLoadAction::ELoad, ERenderTargetLoadAction::ELoad, FExclusiveDepthStencil::DepthRead_StencilRead);
 				}
 
 				GraphBuilder.AddPass(
diff --git a/Engine/Source/Runtime/Renderer/Private/SkyAtmosphereRendering.cpp b/Engine/Source/Runtime/Renderer/Private/SkyAtmosphereRendering.cpp
--- a/Engine/Source/Runtime/Renderer/Private/SkyAtmosphereRendering.cpp
+++ b/Engine/Source/Runtime/Renderer/Private/SkyAtmosphereRendering.cpp
@@ -1930,7 +1930,7 @@ void FSceneRenderer::RenderSkyAtmosphere(FRDGBuilder& GraphBuilder, const FMinim
 	SkyRC.MultiScatteredLuminanceLut = GraphBuilder.RegisterExternalTexture(SkyInfo.GetMultiScatteredLuminanceLutTexture());
 
 	SkyRC.RenderTargets[0] = FRenderTargetBinding(SceneTextures.Color.Target, ERenderTargetLoadAction::ELoad);
-	SkyRC.RenderTargets.DepthStencil = FDepthStencilBinding(SceneTextures.Depth.Target, ERenderTargetLoadAction::ELoad, ERenderTargetLoadAction::ENoAction, FExclusiveDepthStencil::DepthRead_StencilNop);
+	SkyRC.RenderTargets.DepthStencil = FDepthStencilBinding(SceneTextures.Depth.Target, ERenderTargetLoadAction::ELoad, ERenderTargetLoadAction::ELoad, FExclusiveDepthStencil::DepthRead_StencilRead);
 
 	SkyRC.MSAASampleCount = SceneTextures.Depth.Target->Desc.NumSamples;
 	SkyRC.MSAADepthTexture = SceneTextures.Depth.Target;
