Enabled r.Water.SingleLayer.DepthPrepass by default.

#rnx

diff --git a/Engine/Source/Runtime/Renderer/Private/SingleLayerWaterRendering.cpp b/Engine/Source/Runtime/Renderer/Private/SingleLayerWaterRendering.cpp
--- a/Engine/Source/Runtime/Renderer/Private/SingleLayerWaterRendering.cpp
+++ b/Engine/Source/Runtime/Renderer/Private/SingleLayerWaterRendering.cpp
@@ -114,7 +114,7 @@ static TAutoConsoleVariable<int32> CVarRHICmdFlushRenderThreadTasksSingleLayerWa
 	TEXT("Wait for completion of parallel render thread tasks at the end of Single layer water. A more granular version of r.RHICmdFlushRenderThreadTasks. If either r.RHICmdFlushRenderThreadTasks or r.RHICmdFlushRenderThreadTasksSingleLayerWater is > 0 we will flush."));
 
 static TAutoConsoleVariable<int32> CVarWaterSingleLayerDepthPrepass(
-	TEXT("r.Water.SingleLayer.DepthPrepass"), 0,
+	TEXT("r.Water.SingleLayer.DepthPrepass"), 1,
 	TEXT("Enable a depth prepass for single layer water. Necessary for proper Virtual Shadow Maps support."),
 	ECVF_ReadOnly | ECVF_RenderThreadSafe);
 
