Fix skin cache editor debug visualization modes looks different from the cvar equivalent
*) Fixed shading difference
*) Fixed debug texts not showing up

diff --git a/Engine/Source/Runtime/Engine/Private/ShowFlags.cpp b/Engine/Source/Runtime/Engine/Private/ShowFlags.cpp
--- a/Engine/Source/Runtime/Engine/Private/ShowFlags.cpp
+++ b/Engine/Source/Runtime/Engine/Private/ShowFlags.cpp
@@ -356,7 +356,7 @@ void ApplyViewMode(EViewModeIndex ViewModeIndex, bool bPerspective, FEngineShowF
 			bPostProcessing = true;
 			break;
 		case VMI_VisualizeGPUSkinCache:
-			bPostProcessing = true;
+			bPostProcessing = false;
 			break;
 		case VMI_ReflectionOverride:
 			bPostProcessing = true;
diff --git a/Engine/Source/Runtime/Renderer/Private/SceneRendering.cpp b/Engine/Source/Runtime/Renderer/Private/SceneRendering.cpp
--- a/Engine/Source/Runtime/Renderer/Private/SceneRendering.cpp
+++ b/Engine/Source/Runtime/Renderer/Private/SceneRendering.cpp
@@ -4419,7 +4419,8 @@ static void RenderViewFamilies_RenderThread(FRHICommandListImmediate& RHICmdList
 			FGPUSkinCacheVisualizationData& VisualizationData = GetGPUSkinCacheVisualizationData();
 
 			// Only run visualization update once, but set debug flags for all view families if the mode is active
-			if (bUpdatedGPUSkinCacheVisualization || VisualizationData.Update(View.CurrentGPUSkinCacheVisualizationMode))
+			// Note VisualizationData.Update needs to be called per frame, as || lazy evaluation is used, so need to do it before evaluating VisualizeGPUSkinCache flag
+			if (bUpdatedGPUSkinCacheVisualization || VisualizationData.Update(View.CurrentGPUSkinCacheVisualizationMode) || ViewFamily.EngineShowFlags.VisualizeGPUSkinCache)
 			{
 				// When activating visualization from the command line, enable VisualizeGPUSkinCache.
 				ViewFamily.EngineShowFlags.SetVisualizeGPUSkinCache(true);
