Flush the immediate command list at the end of FDeferredShadingSceneRenderer::RenderHitProxies to fix GPU crash due to resources being released while they're still being used.

Without this, the FlushPendingDeleteRHIResources call made after RenderHitProxies returns can release some resources which are still referenced by pending commands, so on the next flush we try to use them and get a GPU crash.

#rnx

diff --git a/Engine/Source/Runtime/Renderer/Private/SceneHitProxyRendering.cpp b/Engine/Source/Runtime/Renderer/Private/SceneHitProxyRendering.cpp
--- a/Engine/Source/Runtime/Renderer/Private/SceneHitProxyRendering.cpp
+++ b/Engine/Source/Runtime/Renderer/Private/SceneHitProxyRendering.cpp
@@ -751,6 +751,7 @@ void FDeferredShadingSceneRenderer::RenderHitProxies(FRDGBuilder& GraphBuilder)
 
 	GEngine->GetPostRenderDelegateEx().Broadcast(GraphBuilder);
 
+	AddDispatchToRHIThreadPass(GraphBuilder);
 #endif
 }
 
