Add Postprocessing csv gpu tag to match between consoles.

diff --git a/Engine/Source/Runtime/Renderer/Private/DeferredShadingRenderer.cpp b/Engine/Source/Runtime/Renderer/Private/DeferredShadingRenderer.cpp
--- a/Engine/Source/Runtime/Renderer/Private/DeferredShadingRenderer.cpp
+++ b/Engine/Source/Runtime/Renderer/Private/DeferredShadingRenderer.cpp
@@ -426,7 +426,7 @@ DECLARE_GPU_STAT(RayTracingUpdate);
 DECLARE_GPU_STAT(RayTracingScene);
 DECLARE_GPU_STAT(RayTracingGeometry);
 
-DECLARE_GPU_DRAWCALL_STAT(Postprocessing);
+DEFINE_GPU_STAT(Postprocessing);
 DECLARE_GPU_STAT(VisibilityCommands);
 DECLARE_GPU_STAT(RenderDeferredLighting);
 DECLARE_GPU_STAT(AllocateRendertargets);
diff --git a/Engine/Source/Runtime/Renderer/Private/MobileShadingRenderer.cpp b/Engine/Source/Runtime/Renderer/Private/MobileShadingRenderer.cpp
--- a/Engine/Source/Runtime/Renderer/Private/MobileShadingRenderer.cpp
+++ b/Engine/Source/Runtime/Renderer/Private/MobileShadingRenderer.cpp
@@ -1125,6 +1125,7 @@ void FMobileSceneRenderer::Render(FRDGBuilder& GraphBuilder)
 			// Finish rendering for each view, or the full stereo buffer if enabled
 			{
 				RDG_EVENT_SCOPE(GraphBuilder, "PostProcessing");
+				RDG_GPU_STAT_SCOPE(GraphBuilder, Postprocessing);
 				SCOPE_CYCLE_COUNTER(STAT_FinishRenderViewTargetTime);
 
 				FMobilePostProcessingInputs PostProcessingInputs;
diff --git a/Engine/Source/Runtime/Renderer/Private/SceneRendering.h b/Engine/Source/Runtime/Renderer/Private/SceneRendering.h
--- a/Engine/Source/Runtime/Renderer/Private/SceneRendering.h
+++ b/Engine/Source/Runtime/Renderer/Private/SceneRendering.h
@@ -88,7 +88,7 @@ struct FVolumetricFogIntegrationParameterData;
 class FLumenHardwareRayTracingUniformBufferParameters;
 
 DECLARE_GPU_DRAWCALL_STAT_EXTERN(VirtualTextureUpdate);
-
+DECLARE_GPU_STAT_NAMED_EXTERN(Postprocessing, TEXT("Postprocessing"))
 
 /** Mobile only. Information used to determine whether static meshes will be rendered with CSM shaders or not. */
 class FMobileCSMVisibilityInfo
