Fix bloom perf regression (0.1ms on some console):

-Keep only the luminance filtering
-Use the square filter (uvcoords from vs interpolators)

diff --git a/Engine/Shaders/Private/PostProcessMobile.usf b/Engine/Shaders/Private/PostProcessMobile.usf
--- a/Engine/Shaders/Private/PostProcessMobile.usf
+++ b/Engine/Shaders/Private/PostProcessMobile.usf
@@ -233,30 +233,25 @@ void BloomVS_Mobile(
 	OutTexCoords[3] = TransformedUV + BufferSizeAndInvSize.zw * float2( 1,  1);
 }
 
-static const float2 kOffsets[5] =
-{
-	int2( 0, -1),
-	int2(-1,  0),
-	int2( 1,  0),
-	int2( 0,  1),
-	int2( 0,  0), // K
-};
 
-DECLARE_SCALAR_ARRAY(float, PlusWeights, 5);
 
-half Luma4(half3 Color)
+float Luma4(float3 Color)
 {
 	return (Color.g * 2.0) + (Color.r + Color.b);
 }
 
-half HdrWeight4(half3 Color)
+float HdrWeight4(float3 Color)
 {
 	return rcp(Luma4(Color) + 4.0);
 }
 
+float HdrWeightInv4(float3 Color)
+{
+	return 4.0 * rcp(1.0 - Luma4(Color));
+}
+
 void BloomPS_Mobile(
-	float2 InUVs[4] : TEXCOORD0,
-	in float4 SvPosition : SV_POSITION
+	float2 InUVs[4] : TEXCOORD0
 #if MOBILE_USEBLOOM
 
 	, out HALF4_TYPE OutColor : SV_Target0
@@ -303,24 +298,19 @@ void BloomPS_Mobile(
 
 #if MOBILE_USEBLOOM || MOBILE_USESUN || MOBILE_USEEYEADAPTATION
 
-	half NeighborsFinalWeight = 0;
-	float3 NeighborsColor = 0;
-	float2 PPCk = floor(SvPosition.xy) * 4 + 0.5f;
-	// Filter input pixels. Referenced from the FilterCurrentFrameInputSamples in TemporalAA.usf
-	for (uint i = 0; i < 5; i++)
-	{
-		float2 TexCoords = BufferSizeAndInvSize.zw * (kOffsets[i] + PPCk) - View.TemporalAAJitter.xy * float2(0.5f, -0.5f);
-		half3 SampleColor = SceneColorTexture.Sample(SceneColorSampler, TexCoords).rgb;
-		half SampleHdrWeight = HdrWeight4(SampleColor);
-		half SampleSpatialWeight = GET_SCALAR_ARRAY_ELEMENT(PlusWeights, i);
-		half SampleFinalWeight = SampleSpatialWeight * SampleHdrWeight;
-
-		NeighborsColor += SampleFinalWeight * SampleColor;
-		NeighborsFinalWeight += SampleFinalWeight;
-	}
+	float4 C0 = SceneColorTexture.Sample(SceneColorSampler, InUVs[0]);
+	float4 C1 = SceneColorTexture.Sample(SceneColorSampler, InUVs[1]);
+	float4 C2 = SceneColorTexture.Sample(SceneColorSampler, InUVs[2]);
+	float4 C3 = SceneColorTexture.Sample(SceneColorSampler, InUVs[3]);
+
+	C0 *= HdrWeight4(C0);
+	C1 *= HdrWeight4(C1);
+	C2 *= HdrWeight4(C2);
+	C3 *= HdrWeight4(C3);
 
 	// Output color is average.
-	AverageColor.rgb = NeighborsColor * rcp(NeighborsFinalWeight);
+	AverageColor.rgb = (C0.rgb * 0.25) + (C1.rgb * 0.25) + (C2.rgb * 0.25) + (C3.rgb * 0.25);
+	AverageColor.rgb *= HdrWeightInv4(AverageColor);
 
 	#if METAL_MSAA_HDR_DECODE
 		// This should really happen before the average, instead doing after average as optimization.
diff --git a/Engine/Source/Runtime/Renderer/Private/PostProcess/PostProcessMobile.cpp b/Engine/Source/Runtime/Renderer/Private/PostProcess/PostProcessMobile.cpp
--- a/Engine/Source/Runtime/Renderer/Private/PostProcess/PostProcessMobile.cpp
+++ b/Engine/Source/Runtime/Renderer/Private/PostProcess/PostProcessMobile.cpp
@@ -217,8 +217,6 @@ public:
 
 	BEGIN_SHADER_PARAMETER_STRUCT(FParameters, )
 		SHADER_PARAMETER(float, BloomThreshold)
-		SHADER_PARAMETER(FVector4f, BufferSizeAndInvSize)
-		SHADER_PARAMETER_SCALAR_ARRAY(float, PlusWeights, [5])
 		SHADER_PARAMETER_STRUCT_REF(FViewUniformShaderParameters, View)
 		SHADER_PARAMETER_STRUCT(FEyeAdaptationParameters, EyeAdaptation)
 		SHADER_PARAMETER_RDG_TEXTURE(Texture2D, SceneColorTexture)
@@ -322,8 +320,7 @@ FMobileBloomSetupOutputs AddMobileBloomSetupPass(FRDGBuilder& GraphBuilder, cons
 	FMobileBloomSetupVS::FParameters VSShaderParameters;
 
 	VSShaderParameters.View = View.ViewUniformBuffer;
-	FVector4f BufferSizeAndInvSize = FVector4f(BufferSize.X, BufferSize.Y, 1.0f / BufferSize.X, 1.0f / BufferSize.Y);
-	VSShaderParameters.BufferSizeAndInvSize = BufferSizeAndInvSize;
+	VSShaderParameters.BufferSizeAndInvSize = FVector4f(BufferSize.X, BufferSize.Y, 1.0f / BufferSize.X, 1.0f / BufferSize.Y);
 
 	auto ShaderPermutationVector = FMobileBloomSetupPS::BuildPermutationVector(Inputs.bUseBloom, Inputs.bUseSun, Inputs.bUseDof, Inputs.bUseEyeAdaptation, Inputs.bUseMetalMSAAHDRDecode);
 
@@ -336,47 +333,16 @@ FMobileBloomSetupOutputs AddMobileBloomSetupPass(FRDGBuilder& GraphBuilder, cons
 		PSShaderParameters->RenderTargets[i] = DestRenderTargets[i].GetRenderTargetBinding();
 	}
 
-	PSShaderParameters->BufferSizeAndInvSize = BufferSizeAndInvSize;
 	PSShaderParameters->EyeAdaptation = EyeAdaptationParameters;
 	PSShaderParameters->BloomThreshold = View.FinalPostProcessSettings.BloomThreshold;
 	PSShaderParameters->View = View.ViewUniformBuffer;
 
 	PSShaderParameters->SceneColorTexture = Inputs.SceneColor.Texture;
-	PSShaderParameters->SceneColorSampler = TStaticSamplerState<SF_Point, AM_Clamp, AM_Clamp, AM_Clamp>::GetRHI();
+	PSShaderParameters->SceneColorSampler = TStaticSamplerState<SF_Bilinear, AM_Clamp, AM_Clamp, AM_Clamp>::GetRHI();
 
 	PSShaderParameters->SunShaftAndDofTexture = Inputs.SunShaftAndDof.Texture;
 	PSShaderParameters->SunShaftAndDofSampler = TStaticSamplerState<SF_Bilinear, AM_Clamp, AM_Clamp, AM_Clamp>::GetRHI();
 
-	// Compute PlusWeights. Referenced from SetupSampleWeightParameters in TemporalAA.cpp
-	{
-		static const float SampleOffsets[5][2] =
-		{
-			{  0.0f, -1.0f },
-			{ -1.0f,  0.0f },
-			{  1.0f,  0.0f },
-			{  0.0f,  1.0f },
-			{  0.0f,  0.0f },
-		};
-
-		FVector2f TemporalJitterPixels = FVector2f(View.TemporalJitterPixels);
-		float TotalWeight = 0.0f;
-		for (int32 i = 0; i < 5; i++)
-		{
-			float PixelOffsetX = SampleOffsets[i][0] - TemporalJitterPixels.X;
-			float PixelOffsetY = SampleOffsets[i][1] - TemporalJitterPixels.Y;
-
-			// Normal distribution, Sigma = 0.47
-			const float CurrWeight = FMath::Exp(-2.29f * (PixelOffsetX * PixelOffsetX + PixelOffsetY * PixelOffsetY));
-			GET_SCALAR_ARRAY_ELEMENT(PSShaderParameters->PlusWeights, i) = CurrWeight;
-			TotalWeight += CurrWeight;
-		}
-
-		for (int32 i = 0; i < 5; i++)
-		{
-			GET_SCALAR_ARRAY_ELEMENT(PSShaderParameters->PlusWeights, i) /= TotalWeight;
-		}
-	}
-
 	const FScreenPassTextureViewport InputViewport(Inputs.SceneColor);
 	const FScreenPassTextureViewport OutputViewport(DestRenderTargets[0]);
 
