Sampling SceneDepth instead of fetching SceneDepthAux when full depth prepass is enabled.

diff --git a/Engine/Plugins/Runtime/AR/AppleAR/AppleARKit/Source/AppleARKit/Private/AppleARKitVideoOverlay.cpp b/Engine/Plugins/Runtime/AR/AppleAR/AppleARKit/Source/AppleARKit/Private/AppleARKitVideoOverlay.cpp
--- a/Engine/Plugins/Runtime/AR/AppleAR/AppleARKit/Source/AppleARKit/Private/AppleARKitVideoOverlay.cpp
+++ b/Engine/Plugins/Runtime/AR/AppleAR/AppleARKit/Source/AppleARKit/Private/AppleARKitVideoOverlay.cpp
@@ -155,6 +155,9 @@ public:
 		OutEnvironment.SetDefine(TEXT("POST_PROCESS_MATERIAL"), 1);
 		OutEnvironment.SetDefine(TEXT("POST_PROCESS_MATERIAL_BEFORE_TONEMAP"), (Parameters.MaterialParameters.BlendableLocation != BL_AfterTonemapping) ? 1 : 0);
 		OutEnvironment.SetDefine(TEXT("POST_PROCESS_AR_PASSTHROUGH"), 1);
+
+		const bool bMobileForceDepthRead = MobileUsesFullDepthPrepass(Parameters.Platform);
+		OutEnvironment.SetDefine(TEXT("FORCE_DEPTH_TEXTURE_READS"), bMobileForceDepthRead ? 1u : 0u);
 	}
 };
 
@@ -170,6 +173,9 @@ public:
 	{
 		TPostProcessMaterialShader<bIsMobileRenderer>::ModifyCompilationEnvironment(Parameters, OutEnvironment);
 		OutEnvironment.SetDefine(TEXT("POST_PROCESS_AR_PASSTHROUGH"), 1);
+
+		const bool bMobileForceDepthRead = MobileUsesFullDepthPrepass(Parameters.Platform);
+		OutEnvironment.SetDefine(TEXT("FORCE_DEPTH_TEXTURE_READS"), bMobileForceDepthRead ? 1u : 0u);
 	}
 
 	TARKitCameraOverlayVS() = default;
diff --git a/Engine/Shaders/Private/DeferredShadingCommon.ush b/Engine/Shaders/Private/DeferredShadingCommon.ush
--- a/Engine/Shaders/Private/DeferredShadingCommon.ush
+++ b/Engine/Shaders/Private/DeferredShadingCommon.ush
@@ -790,6 +790,10 @@ void MobileFetchGBuffer(in float2 UV, out half4 GBufferA, out half4 GBufferB, ou
 #endif
 	SceneDepth = ConvertFromDeviceZ(Texture2DSampleLevel(MobileSceneTextures.SceneDepthTexture, MobileSceneTextures.SceneDepthTextureSampler, UV, 0).r);
 #endif
+
+#if FORCE_DEPTH_TEXTURE_READS || PLATFORM_NEEDS_DEPTH_TEXTURE_READS
+	SceneDepth = ConvertFromDeviceZ(Texture2DSampleLevel(MobileSceneTextures.SceneDepthTexture, MobileSceneTextures.SceneDepthTextureSampler, UV, 0).r);
+#endif
 }
 #endif //MOBILE_DEFERRED_SHADING
 
diff --git a/Engine/Shaders/Public/ShaderVersion.ush b/Engine/Shaders/Public/ShaderVersion.ush
--- a/Engine/Shaders/Public/ShaderVersion.ush
+++ b/Engine/Shaders/Public/ShaderVersion.ush
@@ -2,4 +2,4 @@
 // Each time the console command is executed it generates a new GUID. As this file is included
 // in Platform.ush (which should be included in any shader) it allows to invalidate the shader DDC.
 // 
-// GUID = D9C7A5A5-4311-05E0-9E08-3691866F9F78
+// GUID = F11BFC66-558C-4F14-AE31-780055080BB9
diff --git a/Engine/Source/Runtime/RenderCore/Private/RenderUtils.cpp b/Engine/Source/Runtime/RenderCore/Private/RenderUtils.cpp
--- a/Engine/Source/Runtime/RenderCore/Private/RenderUtils.cpp
+++ b/Engine/Source/Runtime/RenderCore/Private/RenderUtils.cpp
@@ -505,15 +505,19 @@ RENDERCORE_API bool MobileRequiresSceneDepthAux(const FStaticShaderPlatform Plat
 	static const auto CVarMobileHDR = IConsoleManager::Get().FindTConsoleVariableDataInt(TEXT("r.MobileHDR"));
 	const bool bMobileHDR = (CVarMobileHDR && CVarMobileHDR->GetValueOnAnyThread() != 0);
 
-	// SceneDepthAux is used on tiled GPUs (iOS, Android) with forward shading and on iOS only with deferred
-	if (IsMetalMobilePlatform(Platform))
+	if (!MobileUsesFullDepthPrepass(Platform))
 	{
-		return true;
-	}
-	else if (bMobileHDR && !IsMobileDeferredShadingEnabled(Platform))
-	{
-		return (IsAndroidOpenGLESPlatform(Platform) || IsVulkanMobilePlatform(Platform));
+		// SceneDepthAux is used on tiled GPUs (iOS, Android) with forward shading and on iOS only with deferred
+		if (IsMetalMobilePlatform(Platform))
+		{
+			return true;
+		}
+		else if (bMobileHDR && !IsMobileDeferredShadingEnabled(Platform))
+		{
+			return (IsAndroidOpenGLESPlatform(Platform) || IsVulkanMobilePlatform(Platform));
+		}
 	}
+
 	return false;
 }
 
@@ -642,6 +646,13 @@ RENDERCORE_API bool MobileBasePassAlwaysUsesCSM(const FStaticShaderPlatform Plat
 	}
 }
 
+RENDERCORE_API bool MobileUsesFullDepthPrepass(const FStaticShaderPlatform Platform)
+{
+	static TConsoleVariableData<int32>* CVar = IConsoleManager::Get().FindTConsoleVariableDataInt(TEXT("r.Mobile.EarlyZPassOnlyMaterialMasking"));
+
+	return MobileUsesShadowMaskTexture(Platform) || IsMobileAmbientOcclusionEnabled(Platform) || (CVar && CVar->GetValueOnAnyThread() == 1);
+}
+
 RENDERCORE_API bool SupportsGen4TAA(const FStaticShaderPlatform Platform)
 {
 	if (IsMobilePlatform(Platform))
diff --git a/Engine/Source/Runtime/RenderCore/Private/Shader.cpp b/Engine/Source/Runtime/RenderCore/Private/Shader.cpp
--- a/Engine/Source/Runtime/RenderCore/Private/Shader.cpp
+++ b/Engine/Source/Runtime/RenderCore/Private/Shader.cpp
@@ -1730,6 +1730,9 @@ void ShaderMapAppendKeyString(EShaderPlatform Platform, FString& KeyString)
 			KeyString += IsMobileDistanceFieldEnabled(Platform) ? TEXT("_MobSDF") : TEXT("");
 		}
 
+		{
+			KeyString += MobileUsesFullDepthPrepass(Platform) ? TEXT("_MobFDP") : TEXT("");
+		}
 	}
 	else
 	{
diff --git a/Engine/Source/Runtime/RenderCore/Public/RenderUtils.h b/Engine/Source/Runtime/RenderCore/Public/RenderUtils.h
--- a/Engine/Source/Runtime/RenderCore/Public/RenderUtils.h
+++ b/Engine/Source/Runtime/RenderCore/Public/RenderUtils.h
@@ -321,6 +321,8 @@ RENDERCORE_API bool MobileUsesGBufferCustomData(const FStaticShaderPlatform Plat
 
 RENDERCORE_API bool MobileBasePassAlwaysUsesCSM(const FStaticShaderPlatform Platform);
 
+RENDERCORE_API bool MobileUsesFullDepthPrepass(const FStaticShaderPlatform Platform);
+
 RENDERCORE_API bool SupportsGen4TAA(const FStaticShaderPlatform Platform);
 
 RENDERCORE_API bool SupportsTSR(const FStaticShaderPlatform Platform);
diff --git a/Engine/Source/Runtime/Renderer/Private/DecalRenderingCommon.cpp b/Engine/Source/Runtime/Renderer/Private/DecalRenderingCommon.cpp
--- a/Engine/Source/Runtime/Renderer/Private/DecalRenderingCommon.cpp
+++ b/Engine/Source/Runtime/Renderer/Private/DecalRenderingCommon.cpp
@@ -925,7 +925,8 @@ namespace DecalRendering
 		if (IsMobilePlatform(Platform))
 		{
 			// On mobile decals are rendered in a "depth read" sub-pass
-			OutEnvironment.SetDefine(TEXT("IS_MOBILE_DEPTHREAD_SUBPASS"), 1u);
+			const bool bMobileForceDepthRead = MobileUsesFullDepthPrepass(Platform);
+			OutEnvironment.SetDefine(TEXT("IS_MOBILE_DEPTHREAD_SUBPASS"), bMobileForceDepthRead ? 0u : 1u);
 		}
 	}
 }
diff --git a/Engine/Source/Runtime/Renderer/Private/MobileBasePassRendering.cpp b/Engine/Source/Runtime/Renderer/Private/MobileBasePassRendering.cpp
--- a/Engine/Source/Runtime/Renderer/Private/MobileBasePassRendering.cpp
+++ b/Engine/Source/Runtime/Renderer/Private/MobileBasePassRendering.cpp
@@ -136,8 +136,9 @@ void MobileBasePassModifyCompilationEnvironment(const FMaterialShaderPermutation
 	{
 		OutEnvironment.SetDefine(TEXT("ENABLE_SHADINGMODEL_SUPPORT_MOBILE_DEFERRED"), MobileUsesGBufferCustomData(Parameters.Platform));
 	}
-	
-	OutEnvironment.SetDefine(TEXT("IS_MOBILE_DEPTHREAD_SUBPASS"), bTranslucentMaterial ? 1u : 0u);
+
+	const bool bMobileForceDepthRead = MobileUsesFullDepthPrepass(Parameters.Platform);
+	OutEnvironment.SetDefine(TEXT("IS_MOBILE_DEPTHREAD_SUBPASS"), bTranslucentMaterial && !bMobileForceDepthRead ? 1u : 0u);
 	// translucency is in the same subpass with deferred shading shaders, so it has access to GBuffer
 	const bool bDeferredShadingSubpass = (bDeferredShadingEnabled && bTranslucentMaterial && !Parameters.MaterialParameters.bIsMobileSeparateTranslucencyEnabled);
 	OutEnvironment.SetDefine(TEXT("IS_MOBILE_DEFERREDSHADING_SUBPASS"), bDeferredShadingSubpass ? 1u : 0u);
diff --git a/Engine/Source/Runtime/Renderer/Private/MobileDeferredShadingPass.cpp b/Engine/Source/Runtime/Renderer/Private/MobileDeferredShadingPass.cpp
--- a/Engine/Source/Runtime/Renderer/Private/MobileDeferredShadingPass.cpp
+++ b/Engine/Source/Runtime/Renderer/Private/MobileDeferredShadingPass.cpp
@@ -81,6 +81,9 @@ class FMobileDirectionalLightFunctionPS : public FMaterialShader
 		OutEnvironment.SetDefine(TEXT("USE_SHADOWMASKTEXTURE"), MobileUsesShadowMaskTexture(Parameters.Platform) ? 1u : 0u);
 		OutEnvironment.SetDefine(TEXT("MATERIAL_SHADER"), 1);
 		OutEnvironment.SetDefine(TEXT("IS_MOBILE_DEFERREDSHADING_SUBPASS"), 1u);
+
+		const bool bMobileForceDepthRead = MobileUsesFullDepthPrepass(Parameters.Platform);
+		OutEnvironment.SetDefine(TEXT("FORCE_DEPTH_TEXTURE_READS"), bMobileForceDepthRead ? 1u : 0u);
 	}
 
 	static FPermutationDomain RemapPermutationVector(FPermutationDomain PermutationVector, EShaderPlatform Platform)
@@ -221,6 +224,9 @@ public:
 		OutEnvironment.SetDefine(TEXT("USE_SHADOWMASKTEXTURE"), 0);
 		OutEnvironment.SetDefine(TEXT("ENABLE_CLUSTERED_LIGHTS"), 0);
 		OutEnvironment.SetDefine(TEXT("IS_MOBILE_DEFERREDSHADING_SUBPASS"), 1u);
+
+		const bool bMobileForceDepthRead = MobileUsesFullDepthPrepass(Parameters.Platform);
+		OutEnvironment.SetDefine(TEXT("FORCE_DEPTH_TEXTURE_READS"), bMobileForceDepthRead ? 1u : 0u);
 	}
 
 	void SetParameters(FRHIBatchedShaderParameters& BatchedParameters, const FViewInfo& View, const FMaterialRenderProxy* Proxy, const FMaterial& Material)
@@ -288,6 +294,9 @@ public:
 		FGlobalShader::ModifyCompilationEnvironment(Parameters, OutEnvironment);
 		FForwardLightingParameters::ModifyCompilationEnvironment(Parameters.Platform, OutEnvironment);
 		OutEnvironment.SetDefine(TEXT("IS_MOBILE_DEFERREDSHADING_SUBPASS"), 1u);
+
+		const bool bMobileForceDepthRead = MobileUsesFullDepthPrepass(Parameters.Platform);
+		OutEnvironment.SetDefine(TEXT("FORCE_DEPTH_TEXTURE_READS"), bMobileForceDepthRead ? 1u : 0u);
 	}
 };
 
diff --git a/Engine/Source/Runtime/Renderer/Private/MobileFogRendering.cpp b/Engine/Source/Runtime/Renderer/Private/MobileFogRendering.cpp
--- a/Engine/Source/Runtime/Renderer/Private/MobileFogRendering.cpp
+++ b/Engine/Source/Runtime/Renderer/Private/MobileFogRendering.cpp
@@ -105,7 +105,8 @@ class FMobileFogPS : public FGlobalShader
 	static void ModifyCompilationEnvironment(const FGlobalShaderPermutationParameters& Parameters, FShaderCompilerEnvironment& OutEnvironment)
 	{
 		FGlobalShader::ModifyCompilationEnvironment(Parameters, OutEnvironment);
-		OutEnvironment.SetDefine(TEXT("IS_MOBILE_DEPTHREAD_SUBPASS"), 1u);
+		const bool bMobileForceDepthRead = MobileUsesFullDepthPrepass(Parameters.Platform);
+		OutEnvironment.SetDefine(TEXT("IS_MOBILE_DEPTHREAD_SUBPASS"), bMobileForceDepthRead ? 0u : 1u);
 	}
 
 	BEGIN_SHADER_PARAMETER_STRUCT(FParameters, )
diff --git a/Engine/Source/Runtime/Renderer/Private/MobileShadingRenderer.cpp b/Engine/Source/Runtime/Renderer/Private/MobileShadingRenderer.cpp
--- a/Engine/Source/Runtime/Renderer/Private/MobileShadingRenderer.cpp
+++ b/Engine/Source/Runtime/Renderer/Private/MobileShadingRenderer.cpp
@@ -1051,7 +1051,7 @@ void FMobileSceneRenderer::Render(FRDGBuilder& GraphBuilder)
 		{
 			RDG_CSV_STAT_EXCLUSIVE_SCOPE(GraphBuilder, RenderMobileShadowProjections);
 			RDG_GPU_STAT_SCOPE(GraphBuilder, ShadowProjection);
-			RenderMobileShadowProjections(GraphBuilder, SceneTextures.Depth.Resolve);
+			RenderMobileShadowProjections(GraphBuilder);
 		}
 
 		if (bShouldRenderHZB)
@@ -1263,7 +1263,7 @@ void FMobileSceneRenderer::RenderForward(FRDGBuilder& GraphBuilder, FRDGTextureR
 		FMobileBasePassTextures MobileBasePassTextures{};
 		MobileBasePassTextures.ScreenSpaceAO = bRequiresAmbientOcclusionPass ? SceneTextures.ScreenSpaceAO : SystemTextures.White;
 
-		EMobileSceneTextureSetupMode SetupMode = EMobileSceneTextureSetupMode::CustomDepth;
+		EMobileSceneTextureSetupMode SetupMode = (bIsFullDepthPrepassEnabled ? EMobileSceneTextureSetupMode::SceneDepth : EMobileSceneTextureSetupMode::None) | EMobileSceneTextureSetupMode::CustomDepth;
 		FMobileRenderPassParameters* PassParameters = GraphBuilder.AllocParameters<FMobileRenderPassParameters>();
 		PassParameters->View = View.GetShaderParameters();
 		PassParameters->Scene = GetSceneUniforms().GetBuffer(GraphBuilder);
@@ -1591,7 +1591,7 @@ void FMobileSceneRenderer::RenderDeferred(FRDGBuilder& GraphBuilder, const FSort
 		FMobileBasePassTextures MobileBasePassTextures{};
 		MobileBasePassTextures.ScreenSpaceAO = bRequiresAmbientOcclusionPass ? SceneTextures.ScreenSpaceAO : SystemTextures.White;
 
-		EMobileSceneTextureSetupMode SetupMode = EMobileSceneTextureSetupMode::CustomDepth;
+		EMobileSceneTextureSetupMode SetupMode = (bIsFullDepthPrepassEnabled ? EMobileSceneTextureSetupMode::SceneDepth : EMobileSceneTextureSetupMode::None) | EMobileSceneTextureSetupMode::CustomDepth;
 		auto* PassParameters = GraphBuilder.AllocParameters<FMobileRenderPassParameters>();
 		PassParameters->View = View.GetShaderParameters();
 		PassParameters->Scene = GetSceneUniforms().GetBuffer(GraphBuilder);
diff --git a/Engine/Source/Runtime/Renderer/Private/RendererScene.cpp b/Engine/Source/Runtime/Renderer/Private/RendererScene.cpp
--- a/Engine/Source/Runtime/Renderer/Private/RendererScene.cpp
+++ b/Engine/Source/Runtime/Renderer/Private/RendererScene.cpp
@@ -4729,15 +4729,7 @@ void FScene::GetEarlyZPassMode(ERHIFeatureLevel::Type InFeatureLevel, EDepthDraw
 			OutZPassMode = DDM_MaskedOnly;
 		}
 
-		const bool bIsMobileAmbientOcclusionEnabled = IsMobileAmbientOcclusionEnabled(ShaderPlatform);
-		const bool bMobileUsesShadowMaskTexture = MobileUsesShadowMaskTexture(ShaderPlatform);
-		if (bIsMobileAmbientOcclusionEnabled || bMobileUsesShadowMaskTexture)
-		{
-			OutZPassMode = DDM_AllOpaque;
-		}
-
-		bool bMobileForceFullDepthPrepass = CVarMobileEarlyZPass.GetValueOnAnyThread() == 1;
-		if (bMobileForceFullDepthPrepass)
+		if (MobileUsesFullDepthPrepass(ShaderPlatform))
 		{
 			OutZPassMode = DDM_AllOpaque;
 		}
diff --git a/Engine/Source/Runtime/Renderer/Private/SceneRendering.h b/Engine/Source/Runtime/Renderer/Private/SceneRendering.h
--- a/Engine/Source/Runtime/Renderer/Private/SceneRendering.h
+++ b/Engine/Source/Runtime/Renderer/Private/SceneRendering.h
@@ -2379,8 +2379,7 @@ protected:
 		const FLightSceneProxy* LightSceneProxy,
 		TArrayView<const FProjectedShadowInfo* const> Shadows,
 		bool bSubPixelShadow,
-		bool bProjectingForForwardShading,
-		bool bMobileModulatedProjections);
+		bool bProjectingForForwardShading);
 
 
 	void RenderShadowProjections(
@@ -2785,7 +2784,7 @@ protected:
 
 	void RenderPixelProjectedReflection(FRDGBuilder& GraphBuilder, FRDGTextureRef SceneColorTexture, FRDGTextureRef SceneDepthTexture, FRDGTextureRef PixelProjectedReflectionTexture, const FPlanarReflectionSceneProxy* PlanarReflectionSceneProxy);
 
-	void RenderMobileShadowProjections(FRDGBuilder& GraphBuilder, FRDGTextureRef SceneDepthTexture);
+	void RenderMobileShadowProjections(FRDGBuilder& GraphBuilder);
 private:
 	const bool bGammaSpace;
 	const bool bDeferredShading;
diff --git a/Engine/Source/Runtime/Renderer/Private/ShadowRendering.cpp b/Engine/Source/Runtime/Renderer/Private/ShadowRendering.cpp
--- a/Engine/Source/Runtime/Renderer/Private/ShadowRendering.cpp
+++ b/Engine/Source/Runtime/Renderer/Private/ShadowRendering.cpp
@@ -1130,7 +1130,6 @@ void FProjectedShadowInfo::RenderProjection(
 	const FLightSceneProxy* LightSceneProxy,
 	const FSceneRenderer* SceneRender,
 	bool bProjectingForForwardShading,
-	bool bMobileModulatedProjections,
 	bool bSubPixelShadow) const
 {
 	// Find the shadow's view relevance.
@@ -1255,9 +1254,9 @@ void FProjectedShadowInfo::RenderProjection(
 		RDG_EVENT_NAME("%s", *EventName),
 		PassParameters,
 		ERDGPassFlags::Raster | PassFlags,
-		[this, SceneRender, View, ViewIndex, LightSceneProxy, bProjectingForForwardShading, bMobileModulatedProjections, &InstanceCullingDrawParams, bSubPixelShadow, PassParameters](FRHICommandList& RHICmdList)
+		[this, SceneRender, View, ViewIndex, LightSceneProxy, bProjectingForForwardShading, &InstanceCullingDrawParams, bSubPixelShadow, PassParameters](FRHICommandList& RHICmdList)
 	{
-		RenderProjectionInternal(RHICmdList, ViewIndex, View, LightSceneProxy, SceneRender, bProjectingForForwardShading, bMobileModulatedProjections, InstanceCullingDrawParams, bSubPixelShadow && PassParameters->HairStrands ? PassParameters->HairStrands.GetUniformBuffer()->GetRHI() : nullptr);
+		RenderProjectionInternal(RHICmdList, ViewIndex, View, LightSceneProxy, SceneRender, bProjectingForForwardShading, false, InstanceCullingDrawParams, bSubPixelShadow && PassParameters->HairStrands ? PassParameters->HairStrands.GetUniformBuffer()->GetRHI() : nullptr);
 	});
 }
 
@@ -2018,8 +2017,7 @@ void FSceneRenderer::RenderShadowProjections(
 	const FLightSceneProxy* LightSceneProxy,
 	TArrayView<const FProjectedShadowInfo* const> Shadows,
 	bool bSubPixelShadow,
-	bool bProjectingForForwardShading,
-	bool bMobileModulatedProjections)
+	bool bProjectingForForwardShading)
 {
 	CheckShadowDepthRenderCompleted();
 
@@ -2073,7 +2071,7 @@ void FSceneRenderer::RenderShadowProjections(
 					}
 					else
 					{
-						ProjectedShadowInfo->RenderProjection(GraphBuilder, CommonPassParameters, ViewIndex, &View, LightSceneProxy, this, bProjectingForForwardShading, bMobileModulatedProjections, bSubPixelShadow);
+						ProjectedShadowInfo->RenderProjection(GraphBuilder, CommonPassParameters, ViewIndex, &View, LightSceneProxy, this, bProjectingForForwardShading, bSubPixelShadow);
 					}
 				}
 			}
@@ -2129,7 +2127,6 @@ void FSceneRenderer::RenderShadowProjections(
 			GetLightNameForDrawEvent(LightSceneProxy, LightNameWithLevel);
 			RDG_EVENT_SCOPE(GraphBuilder, "%s", *LightNameWithLevel);
 
-			const bool bMobileModulatedProjections = false;
 			FSceneRenderer::RenderShadowProjections(
 				GraphBuilder,
 				OutputTexture,
@@ -2137,8 +2134,7 @@ void FSceneRenderer::RenderShadowProjections(
 				LightSceneProxy,
 				NormalShadows,
 				bSubPixel,
-				bProjectingForForwardShading,
-				bMobileModulatedProjections);
+				bProjectingForForwardShading);
 		};
 
 		{
@@ -2410,8 +2406,7 @@ void InitMobileShadowProjectionOutputs(FRHICommandListImmediate& RHICmdList, con
 }
 
 void FMobileSceneRenderer::RenderMobileShadowProjections(
-	FRDGBuilder& GraphBuilder, 
-	FRDGTextureRef SceneDepthTexture)
+	FRDGBuilder& GraphBuilder)
 {
 	RDG_RHI_EVENT_SCOPE(GraphBuilder, RenderMobileShadowProjections);
 
diff --git a/Engine/Source/Runtime/Renderer/Private/ShadowRendering.h b/Engine/Source/Runtime/Renderer/Private/ShadowRendering.h
--- a/Engine/Source/Runtime/Renderer/Private/ShadowRendering.h
+++ b/Engine/Source/Runtime/Renderer/Private/ShadowRendering.h
@@ -582,7 +582,6 @@ public:
 		const FLightSceneProxy* LightSceneProxy,
 		const FSceneRenderer* SceneRender,
 		bool bProjectingForForwardShading,
-		bool bMobileModulatedProjections,
 		bool bSubPixelShadow) const;
 
 	void RenderProjectionInternal(
@@ -1223,8 +1222,8 @@ public:
 		OutEnvironment.SetDefine(TEXT("USE_FADE_PLANE"), (uint32)(bUseFadePlane ? 1 : 0));
 		OutEnvironment.SetDefine(TEXT("USE_TRANSMISSION"), (uint32)(bUseTransmission ? 1 : 0));
 
-		const bool bMobileVulkanForceDepthRead = IsVulkanMobilePlatform(Parameters.Platform) && MobileUsesShadowMaskTexture(Parameters.Platform);
-		OutEnvironment.SetDefine(TEXT("FORCE_DEPTH_TEXTURE_READS"), (uint32)(bMobileVulkanForceDepthRead ? 1 : 0));
+		const bool bMobileForceDepthRead = MobileUsesFullDepthPrepass(Parameters.Platform);
+		OutEnvironment.SetDefine(TEXT("FORCE_DEPTH_TEXTURE_READS"), (uint32)(bMobileForceDepthRead ? 1 : 0));
 	}
 
 	/**
@@ -1282,7 +1281,8 @@ public:
 	{
 		TShadowProjectionPS<Quality, false, true>::ModifyCompilationEnvironment(Parameters, OutEnvironment);
 		OutEnvironment.SetDefine(TEXT("MODULATED_SHADOWS"), 1);
-		OutEnvironment.SetDefine(TEXT("IS_MOBILE_DEPTHREAD_SUBPASS"), 1);
+		const bool bMobileForceDepthRead = MobileUsesFullDepthPrepass(Parameters.Platform);
+		OutEnvironment.SetDefine(TEXT("IS_MOBILE_DEPTHREAD_SUBPASS"), bMobileForceDepthRead ? 0u : 1u);
 	}
 
 	static bool ShouldCompilePermutation(const FGlobalShaderPermutationParameters& Parameters)
