Replace RHICommandListImmediate to RHICommandList to enable RDG pass parallelism on mobile platforms.

diff --git a/Engine/Plugins/Runtime/AR/AppleAR/AppleARKit/Source/AppleARKit/Private/AppleARKitSystem.cpp b/Engine/Plugins/Runtime/AR/AppleAR/AppleARKit/Source/AppleARKit/Private/AppleARKitSystem.cpp
--- a/Engine/Plugins/Runtime/AR/AppleAR/AppleARKit/Source/AppleARKit/Private/AppleARKitSystem.cpp
+++ b/Engine/Plugins/Runtime/AR/AppleAR/AppleARKit/Source/AppleARKit/Private/AppleARKitSystem.cpp
@@ -198,7 +198,7 @@ private:
 		}
 	}
 
-	virtual void PostRenderBasePassMobile_RenderThread(FRHICommandListImmediate& RHICmdList, FSceneView& InView) override
+	virtual void PostRenderBasePassMobile_RenderThread(FRHICommandList& RHICmdList, FSceneView& InView) override
 	{
 		if (ARKitSystem.RenderThreadFrame.IsValid())
 		{
diff --git a/Engine/Plugins/Runtime/AR/AppleAR/AppleARKit/Source/AppleARKit/Private/AppleARKitVideoOverlay.cpp b/Engine/Plugins/Runtime/AR/AppleAR/AppleARKit/Source/AppleARKit/Private/AppleARKitVideoOverlay.cpp
--- a/Engine/Plugins/Runtime/AR/AppleAR/AppleARKit/Source/AppleARKit/Private/AppleARKitVideoOverlay.cpp
+++ b/Engine/Plugins/Runtime/AR/AppleAR/AppleARKit/Source/AppleARKit/Private/AppleARKitVideoOverlay.cpp
@@ -282,7 +282,7 @@ void FAppleARKitVideoOverlay::UpdateOcclusionTextures(const FAppleARKitFrame& Fr
 #endif
 }
 
-void FAppleARKitVideoOverlay::RenderVideoOverlayWithMaterial(FRHICommandListImmediate& RHICmdList, const FSceneView& InView, struct FAppleARKitFrame& Frame, UMaterialInstanceDynamic* RenderingOverlayMaterial, const bool bRenderingOcclusion)
+void FAppleARKitVideoOverlay::RenderVideoOverlayWithMaterial(FRHICommandList& RHICmdList, const FSceneView& InView, struct FAppleARKitFrame& Frame, UMaterialInstanceDynamic* RenderingOverlayMaterial, const bool bRenderingOcclusion)
 {
 	if (RenderingOverlayMaterial == nullptr || !RenderingOverlayMaterial->IsValidLowLevel())
 	{
@@ -438,7 +438,7 @@ void FAppleARKitVideoOverlay::RenderVideoOverlayWithMaterial(FRHICommandListImme
 	}
 }
 
-void FAppleARKitVideoOverlay::RenderVideoOverlay_RenderThread(FRHICommandListImmediate& RHICmdList, const FSceneView& InView, FAppleARKitFrame& Frame, const float WorldToMeterScale)
+void FAppleARKitVideoOverlay::RenderVideoOverlay_RenderThread(FRHICommandList& RHICmdList, const FSceneView& InView, FAppleARKitFrame& Frame, const float WorldToMeterScale)
 {
 	UpdateVideoTextures(Frame);
 	UpdateOcclusionTextures(Frame);
diff --git a/Engine/Plugins/Runtime/AR/AppleAR/AppleARKit/Source/AppleARKit/Private/AppleARKitVideoOverlay.h b/Engine/Plugins/Runtime/AR/AppleAR/AppleARKit/Source/AppleARKit/Private/AppleARKitVideoOverlay.h
--- a/Engine/Plugins/Runtime/AR/AppleAR/AppleARKit/Source/AppleARKit/Private/AppleARKitVideoOverlay.h
+++ b/Engine/Plugins/Runtime/AR/AppleAR/AppleARKit/Source/AppleARKit/Private/AppleARKitVideoOverlay.h
@@ -82,7 +82,7 @@ public:
 
 	void SetCameraTexture(UARTextureCameraImage* InCameraImage);
 
-	void RenderVideoOverlay_RenderThread(FRHICommandListImmediate& RHICmdList, const FSceneView& InView, struct FAppleARKitFrame& Frame, const float WorldToMeterScale);
+	void RenderVideoOverlay_RenderThread(FRHICommandList& RHICmdList, const FSceneView& InView, struct FAppleARKitFrame& Frame, const float WorldToMeterScale);
 	bool GetPassthroughCameraUVs_RenderThread(TArray<FVector2D>& OutUVs, const EDeviceScreenOrientation DeviceOrientation);
 
 	void SetOverlayTexture(UARTextureCameraImage* InCameraImage);
@@ -102,7 +102,7 @@ private:
 	}
 	//~ FGCObject
 	
-	void RenderVideoOverlayWithMaterial(FRHICommandListImmediate& RHICmdList, const FSceneView& InView, struct FAppleARKitFrame& Frame, UMaterialInstanceDynamic* RenderingOverlayMaterial, const bool bRenderingOcclusion);
+	void RenderVideoOverlayWithMaterial(FRHICommandList& RHICmdList, const FSceneView& InView, struct FAppleARKitFrame& Frame, UMaterialInstanceDynamic* RenderingOverlayMaterial, const bool bRenderingOcclusion);
 	void UpdateOcclusionTextures(const FAppleARKitFrame& Frame);
 	void UpdateVideoTextures(const FAppleARKitFrame& Frame);
 	
diff --git a/Engine/Plugins/Runtime/AR/Google/GoogleARCore/Source/GoogleARCoreBase/Private/GoogleARCoreXRCamera.cpp b/Engine/Plugins/Runtime/AR/Google/GoogleARCore/Source/GoogleARCoreBase/Private/GoogleARCoreXRCamera.cpp
--- a/Engine/Plugins/Runtime/AR/Google/GoogleARCore/Source/GoogleARCoreBase/Private/GoogleARCoreXRCamera.cpp
+++ b/Engine/Plugins/Runtime/AR/Google/GoogleARCore/Source/GoogleARCoreBase/Private/GoogleARCoreXRCamera.cpp
@@ -79,7 +79,7 @@ void FGoogleARCoreXRCamera::PostRenderBasePassDeferred_RenderThread(FRDGBuilder&
 	});
 }
 
-void FGoogleARCoreXRCamera::PostRenderBasePassMobile_RenderThread(FRHICommandListImmediate& RHICmdList, FSceneView& InView)
+void FGoogleARCoreXRCamera::PostRenderBasePassMobile_RenderThread(FRHICommandList& RHICmdList, FSceneView& InView)
 {
 	PassthroughRenderer->RenderVideoOverlay_RenderThread(RHICmdList, InView);
 }
diff --git a/Engine/Plugins/Runtime/AR/Google/GoogleARCore/Source/GoogleARCoreBase/Private/GoogleARCoreXRCamera.h b/Engine/Plugins/Runtime/AR/Google/GoogleARCore/Source/GoogleARCoreBase/Private/GoogleARCoreXRCamera.h
--- a/Engine/Plugins/Runtime/AR/Google/GoogleARCore/Source/GoogleARCoreBase/Private/GoogleARCoreXRCamera.h
+++ b/Engine/Plugins/Runtime/AR/Google/GoogleARCore/Source/GoogleARCoreBase/Private/GoogleARCoreXRCamera.h
@@ -19,7 +19,7 @@ public:
 	virtual void BeginRenderViewFamily(FSceneViewFamily& InViewFamily) override;
 	virtual void PreRenderViewFamily_RenderThread(FRDGBuilder& GraphBuilder, FSceneViewFamily& InViewFamily) override;
 	virtual void PostRenderBasePassDeferred_RenderThread(FRDGBuilder& GraphBuilder, FSceneView& InView, const FRenderTargetBindingSlots& RenderTargets, TRDGUniformBufferRef<FSceneTextureUniformParameters> SceneTextures) override;
-	virtual void PostRenderBasePassMobile_RenderThread(FRHICommandListImmediate& RHICmdList, FSceneView& InView) override;
+	virtual void PostRenderBasePassMobile_RenderThread(FRHICommandList& RHICmdList, FSceneView& InView) override;
 	virtual bool IsActiveThisFrame_Internal(const FSceneViewExtensionContext& Context) const override;
 	virtual bool GetPassthroughCameraUVs_RenderThread(TArray<FVector2D>& OutUVs) override;
 	//~ FDefaultXRCamera
diff --git a/Engine/Plugins/Runtime/AR/Google/GoogleARCore/Source/GoogleARCoreRendering/Private/GoogleARCorePassthroughCameraRenderer.cpp b/Engine/Plugins/Runtime/AR/Google/GoogleARCore/Source/GoogleARCoreRendering/Private/GoogleARCorePassthroughCameraRenderer.cpp
--- a/Engine/Plugins/Runtime/AR/Google/GoogleARCore/Source/GoogleARCoreRendering/Private/GoogleARCorePassthroughCameraRenderer.cpp
+++ b/Engine/Plugins/Runtime/AR/Google/GoogleARCore/Source/GoogleARCoreRendering/Private/GoogleARCorePassthroughCameraRenderer.cpp
@@ -171,7 +171,7 @@ public:
 
 IMPLEMENT_GLOBAL_SHADER(FGoogleARCoreCameraOverlayPS, "/Engine/Private/PostProcessMaterialShaders.usf", "MainPS", SF_Pixel);
 
-void FGoogleARCorePassthroughCameraRenderer::RenderVideoOverlayWithMaterial(FRHICommandListImmediate& RHICmdList, FSceneView& InView, UMaterialInstanceDynamic* OverlayMaterialToUse, bool bRenderingOcclusion)
+void FGoogleARCorePassthroughCameraRenderer::RenderVideoOverlayWithMaterial(FRHICommandList& RHICmdList, FSceneView& InView, UMaterialInstanceDynamic* OverlayMaterialToUse, bool bRenderingOcclusion)
 {
 #if PLATFORM_ANDROID
 	if (FAndroidMisc::ShouldUseVulkan() && IsMobileHDR() && !RHICmdList.IsInsideRenderPass())
@@ -247,7 +247,7 @@ void FGoogleARCorePassthroughCameraRenderer::RenderVideoOverlayWithMaterial(FRHI
 #endif
 }
 
-void FGoogleARCorePassthroughCameraRenderer::RenderVideoOverlay_RenderThread(FRHICommandListImmediate& RHICmdList, FSceneView& InView)
+void FGoogleARCorePassthroughCameraRenderer::RenderVideoOverlay_RenderThread(FRHICommandList& RHICmdList, FSceneView& InView)
 {
 	auto OverlayMaterialToUse = RegularOverlayMaterial;
 	
diff --git a/Engine/Plugins/Runtime/AR/Google/GoogleARCore/Source/GoogleARCoreRendering/Public/GoogleARCorePassthroughCameraRenderer.h b/Engine/Plugins/Runtime/AR/Google/GoogleARCore/Source/GoogleARCoreRendering/Public/GoogleARCorePassthroughCameraRenderer.h
--- a/Engine/Plugins/Runtime/AR/Google/GoogleARCore/Source/GoogleARCoreRendering/Public/GoogleARCorePassthroughCameraRenderer.h
+++ b/Engine/Plugins/Runtime/AR/Google/GoogleARCore/Source/GoogleARCoreRendering/Public/GoogleARCorePassthroughCameraRenderer.h
@@ -68,7 +68,7 @@ public:
 
 	void InitializeRenderer_RenderThread(FSceneViewFamily& InViewFamily);
 
-	void RenderVideoOverlay_RenderThread(FRHICommandListImmediate& RHICmdList, FSceneView& InView);
+	void RenderVideoOverlay_RenderThread(FRHICommandList& RHICmdList, FSceneView& InView);
 	
 	void UpdateCameraTextures(UTexture* NewCameraTexture, UTexture* DepthTexture, bool bEnableOcclusion);
 	
@@ -79,7 +79,7 @@ public:
 	}
 
 private:
-	void RenderVideoOverlayWithMaterial(FRHICommandListImmediate& RHICmdList, FSceneView& InView, UMaterialInstanceDynamic* OverlayMaterialToUse, bool bRenderingOcclusion);
+	void RenderVideoOverlayWithMaterial(FRHICommandList& RHICmdList, FSceneView& InView, UMaterialInstanceDynamic* OverlayMaterialToUse, bool bRenderingOcclusion);
 
 private:
 	FBufferRHIRef OverlayIndexBufferRHI;
diff --git a/Engine/Source/Runtime/Engine/Public/SceneViewExtension.h b/Engine/Source/Runtime/Engine/Public/SceneViewExtension.h
--- a/Engine/Source/Runtime/Engine/Public/SceneViewExtension.h
+++ b/Engine/Source/Runtime/Engine/Public/SceneViewExtension.h
@@ -163,7 +163,7 @@ public:
 	/**
 	 * Called right after Base Pass rendering finished when using the mobile renderer.
 	 */
-	virtual void PostRenderBasePassMobile_RenderThread(FRHICommandListImmediate& RHICmdList, FSceneView& InView) {}
+	virtual void PostRenderBasePassMobile_RenderThread(FRHICommandList& RHICmdList, FSceneView& InView) {}
 
 	/**
 	 * Called right before Post Processing rendering begins
diff --git a/Engine/Source/Runtime/RenderCore/Private/RenderGraphPrivate.cpp b/Engine/Source/Runtime/RenderCore/Private/RenderGraphPrivate.cpp
--- a/Engine/Source/Runtime/RenderCore/Private/RenderGraphPrivate.cpp
+++ b/Engine/Source/Runtime/RenderCore/Private/RenderGraphPrivate.cpp
@@ -367,7 +367,7 @@ FAutoConsoleVariableRef CVarRDGParallelExecute(
 				GRDGParallelExecutePassMax = 1;
 			}
 
-			if (GRDGParallelExecutePassMin < GRDGParallelExecutePassMax)
+			if (GRDGParallelExecutePassMax < GRDGParallelExecutePassMin)
 			{
 				GRDGParallelExecutePassMin = GRDGParallelExecutePassMax;
 			}
@@ -634,7 +634,7 @@ bool IsParallelExecuteEnabled()
 		&& !IsImmediateMode()
 		&& !GRDGDebug
 		&& !GRDGTransitionLog
-		&& !IsMobilePlatform(GMaxRHIShaderPlatform)
+		&& !IsOpenGLPlatform(GMaxRHIShaderPlatform)
 		&& GRHISupportsMultithreadedShaderCreation
 #if WITH_DUMPGPU
 		&& !IsDumpingRDGResources()
@@ -651,7 +651,7 @@ bool IsParallelSetupEnabled()
 		&& !IsImmediateMode()
 		&& !GRDGDebug
 		&& !GRDGTransitionLog
-		&& !IsMobilePlatform(GMaxRHIShaderPlatform)
+		&& !IsOpenGLPlatform(GMaxRHIShaderPlatform)
 		&& GRHISupportsMultithreadedShaderCreation
 #if WITH_DUMPGPU
 		&& !IsDumpingRDGResources()
diff --git a/Engine/Source/Runtime/Renderer/Private/CompositionLighting/PostProcessMeshDecals.cpp b/Engine/Source/Runtime/Renderer/Private/CompositionLighting/PostProcessMeshDecals.cpp
--- a/Engine/Source/Runtime/Renderer/Private/CompositionLighting/PostProcessMeshDecals.cpp
+++ b/Engine/Source/Runtime/Renderer/Private/CompositionLighting/PostProcessMeshDecals.cpp
@@ -492,7 +492,7 @@ void RenderMeshDecals(
 	}
 }
 
-void RenderMeshDecalsMobile(FRHICommandListImmediate& RHICmdList, const FViewInfo& View, EDecalRenderStage DecalRenderStage, EDecalRenderTargetMode RenderTargetMode)
+void RenderMeshDecalsMobile(FRHICommandList& RHICmdList, const FViewInfo& View, EDecalRenderStage DecalRenderStage, EDecalRenderTargetMode RenderTargetMode)
 {
 	FGraphicsPipelineStateInitializer GraphicsPSOInit;
 	RHICmdList.SetViewport(View.ViewRect.Min.X, View.ViewRect.Min.Y, 0, View.ViewRect.Max.X, View.ViewRect.Max.Y, 1);
diff --git a/Engine/Source/Runtime/Renderer/Private/DepthRendering.cpp b/Engine/Source/Runtime/Renderer/Private/DepthRendering.cpp
--- a/Engine/Source/Runtime/Renderer/Private/DepthRendering.cpp
+++ b/Engine/Source/Runtime/Renderer/Private/DepthRendering.cpp
@@ -657,7 +657,7 @@ bool FMobileSceneRenderer::ShouldRenderPrePass() const
 	return Scene->EarlyZPassMode == DDM_MaskedOnly || Scene->EarlyZPassMode == DDM_AllOpaque;
 }
 
-void FMobileSceneRenderer::RenderPrePass(FRHICommandListImmediate& RHICmdList, const FViewInfo& View)
+void FMobileSceneRenderer::RenderPrePass(FRHICommandList& RHICmdList, const FViewInfo& View)
 {
 	if (!ShouldRenderPrePass())
 	{
diff --git a/Engine/Source/Runtime/Renderer/Private/MobileBasePassRendering.cpp b/Engine/Source/Runtime/Renderer/Private/MobileBasePassRendering.cpp
--- a/Engine/Source/Runtime/Renderer/Private/MobileBasePassRendering.cpp
+++ b/Engine/Source/Runtime/Renderer/Private/MobileBasePassRendering.cpp
@@ -377,7 +377,7 @@ void SetupMobileSkyReflectionUniformParameters(FSkyLightSceneProxy* SkyLight, FM
 	Parameters.TextureSampler = CaptureTexture->SamplerStateRHI;
 }
 
-void FMobileSceneRenderer::RenderMobileBasePass(FRHICommandListImmediate& RHICmdList, const FViewInfo& View)
+void FMobileSceneRenderer::RenderMobileBasePass(FRHICommandList& RHICmdList, const FViewInfo& View)
 {
 	CSV_SCOPED_TIMING_STAT_EXCLUSIVE(RenderBasePass);
 	SCOPED_DRAW_EVENT(RHICmdList, MobileBasePass);
diff --git a/Engine/Source/Runtime/Renderer/Private/MobileDecalRendering.cpp b/Engine/Source/Runtime/Renderer/Private/MobileDecalRendering.cpp
--- a/Engine/Source/Runtime/Renderer/Private/MobileDecalRendering.cpp
+++ b/Engine/Source/Runtime/Renderer/Private/MobileDecalRendering.cpp
@@ -17,10 +17,10 @@
 #include "DecalRenderingShared.h"
 #include "RenderCore.h"
 
-void RenderMeshDecalsMobile(FRHICommandListImmediate& RHICmdList, const FViewInfo& View, EDecalRenderStage DecalRenderStage, EDecalRenderTargetMode RenderTargetMode);
-extern void RenderDeferredDecalsMobile(FRHICommandListImmediate& RHICmdList, const FScene& Scene, const FViewInfo& View, EDecalRenderStage DecalRenderStage, EDecalRenderTargetMode RenderTargetMode);
+void RenderMeshDecalsMobile(FRHICommandList& RHICmdList, const FViewInfo& View, EDecalRenderStage DecalRenderStage, EDecalRenderTargetMode RenderTargetMode);
+extern void RenderDeferredDecalsMobile(FRHICommandList& RHICmdList, const FScene& Scene, const FViewInfo& View, EDecalRenderStage DecalRenderStage, EDecalRenderTargetMode RenderTargetMode);
 
-void FMobileSceneRenderer::RenderDecals(FRHICommandListImmediate& RHICmdList, const FViewInfo& View)
+void FMobileSceneRenderer::RenderDecals(FRHICommandList& RHICmdList, const FViewInfo& View)
 {
 	if (!IsMobileHDR() || !ViewFamily.EngineShowFlags.Decals || View.bIsPlanarReflection)
 	{
@@ -49,7 +49,7 @@ void FMobileSceneRenderer::RenderDecals(FRHICommandListImmediate& RHICmdList, co
 	}
 }
 
-void RenderDeferredDecalsMobile(FRHICommandListImmediate& RHICmdList, const FScene& Scene, const FViewInfo& View, EDecalRenderStage DecalRenderStage, EDecalRenderTargetMode RenderTargetMode)
+void RenderDeferredDecalsMobile(FRHICommandList& RHICmdList, const FScene& Scene, const FViewInfo& View, EDecalRenderStage DecalRenderStage, EDecalRenderTargetMode RenderTargetMode)
 {
 	const uint32 DecalCount = Scene.Decals.Num();
 	int32 SortedDecalCount = 0;
diff --git a/Engine/Source/Runtime/Renderer/Private/MobileDeferredShadingPass.cpp b/Engine/Source/Runtime/Renderer/Private/MobileDeferredShadingPass.cpp
--- a/Engine/Source/Runtime/Renderer/Private/MobileDeferredShadingPass.cpp
+++ b/Engine/Source/Runtime/Renderer/Private/MobileDeferredShadingPass.cpp
@@ -357,7 +357,7 @@ static void GetLightMaterial(const FCachedLightMaterial& DefaultLightMaterial, c
 	OutShader = MaterialShaderMap->GetShader<ShaderType>(PermutationId);
 }
 
-void RenderReflectionEnvironmentSkyLighting(FRHICommandListImmediate& RHICmdList, const FScene& Scene, const FViewInfo& View)
+void RenderReflectionEnvironmentSkyLighting(FRHICommandList& RHICmdList, const FScene& Scene, const FViewInfo& View)
 {
 	// Skylights with static lighting already had their diffuse contribution baked into lightmaps
 	const bool bSkyLight = Scene.SkyLight && !Scene.SkyLight->bHasStaticLighting && View.Family->EngineShowFlags.SkyLighting;
@@ -442,7 +442,7 @@ static void SetDirectionalLightDepthStencilState(FGraphicsPipelineStateInitializ
 		GET_STENCIL_MOBILE_SM_MASK(0xff) | STENCIL_LIGHTING_CHANNELS_MASK(1u << LightingChannelIdx), 0x00>::GetRHI();
 }
 
-static void RenderDirectionalLight(FRHICommandListImmediate& RHICmdList, const FScene& Scene, const FViewInfo& View, const FCachedLightMaterial& DefaultLightMaterial, const FLightSceneInfo& DirectionalLight, uint32 LightingChannel, bool bInlineReflectionAndSky)
+static void RenderDirectionalLight(FRHICommandList& RHICmdList, const FScene& Scene, const FViewInfo& View, const FCachedLightMaterial& DefaultLightMaterial, const FLightSceneInfo& DirectionalLight, uint32 LightingChannel, bool bInlineReflectionAndSky)
 {
 	FString LightNameWithLevel;
 	FSceneRenderer::GetLightNameForDrawEvent(DirectionalLight.Proxy, LightNameWithLevel);
@@ -561,7 +561,7 @@ static void RenderDirectionalLight(FRHICommandListImmediate& RHICmdList, const F
 	}
 }
 
-static void RenderDirectionalLights(FRHICommandListImmediate& RHICmdList, const FScene& Scene, const FViewInfo& View, const FCachedLightMaterial& DefaultLightMaterial)
+static void RenderDirectionalLights(FRHICommandList& RHICmdList, const FScene& Scene, const FViewInfo& View, const FCachedLightMaterial& DefaultLightMaterial)
 {
 	uint32 NumLights = 0;
 	for (uint32 ChannelIdx = 0; ChannelIdx < UE_ARRAY_COUNT(Scene.MobileDirectionalLights); ChannelIdx++)
@@ -646,7 +646,7 @@ static void SetLocalLightRasterizerAndDepthState(FGraphicsPipelineStateInitializ
 	}
 }
 
-static void RenderLocalLight_StencilMask(FRHICommandListImmediate& RHICmdList, const FScene& Scene, const FViewInfo& View, const FLightSceneInfo& LightSceneInfo)
+static void RenderLocalLight_StencilMask(FRHICommandList& RHICmdList, const FScene& Scene, const FViewInfo& View, const FLightSceneInfo& LightSceneInfo)
 {
 	const uint8 LightType = LightSceneInfo.Proxy->GetLightType();
 
@@ -685,7 +685,7 @@ static void RenderLocalLight_StencilMask(FRHICommandListImmediate& RHICmdList, c
 }
 
 static void RenderLocalLight(
-	FRHICommandListImmediate& RHICmdList, 
+	FRHICommandList& RHICmdList, 
 	const FScene& Scene, 
 	const FViewInfo& View, 
 	const FLightSceneInfo& LightSceneInfo, 
@@ -813,7 +813,7 @@ static void RenderLocalLight(
 }
 
 static void RenderSimpleLights(
-	FRHICommandListImmediate& RHICmdList, 
+	FRHICommandList& RHICmdList, 
 	const FScene& Scene, 
 	int32 ViewIndex,
 	int32 NumViews,
@@ -925,7 +925,7 @@ static void RenderSimpleLights(
 }
 
 void MobileDeferredShadingPass(
-	FRHICommandListImmediate& RHICmdList,
+	FRHICommandList& RHICmdList,
 	int32 ViewIndex,
 	int32 NumViews,
 	const FViewInfo& View,
diff --git a/Engine/Source/Runtime/Renderer/Private/MobileDeferredShadingPass.h b/Engine/Source/Runtime/Renderer/Private/MobileDeferredShadingPass.h
--- a/Engine/Source/Runtime/Renderer/Private/MobileDeferredShadingPass.h
+++ b/Engine/Source/Runtime/Renderer/Private/MobileDeferredShadingPass.h
@@ -17,7 +17,7 @@ struct FSortedLightSetSceneInfo;
 extern int32 GMobileUseClusteredDeferredShading;
 
 void MobileDeferredShadingPass(
-	FRHICommandListImmediate& RHICmdList,
+	FRHICommandList& RHICmdList,
 	int32 ViewIndex,
 	int32 NumViews,
 	const FViewInfo& View,
diff --git a/Engine/Source/Runtime/Renderer/Private/MobileFogRendering.cpp b/Engine/Source/Runtime/Renderer/Private/MobileFogRendering.cpp
--- a/Engine/Source/Runtime/Renderer/Private/MobileFogRendering.cpp
+++ b/Engine/Source/Runtime/Renderer/Private/MobileFogRendering.cpp
@@ -117,7 +117,7 @@ class FMobileFogPS : public FGlobalShader
 
 IMPLEMENT_SHADER_TYPE(, FMobileFogPS, TEXT("/Engine/Private/MobileFog.usf"), TEXT("MobileFogPS"), SF_Pixel);
 
-void FMobileSceneRenderer::RenderFog(FRHICommandListImmediate& RHICmdList, const FViewInfo& View)
+void FMobileSceneRenderer::RenderFog(FRHICommandList& RHICmdList, const FViewInfo& View)
 {	
 	static const auto* CVarDisableVertexFog = IConsoleManager::Get().FindTConsoleVariableDataInt(TEXT("r.Mobile.DisableVertexFog"));
 	if (CVarDisableVertexFog && CVarDisableVertexFog->GetValueOnRenderThread() == 0)
diff --git a/Engine/Source/Runtime/Renderer/Private/MobileShadingRenderer.cpp b/Engine/Source/Runtime/Renderer/Private/MobileShadingRenderer.cpp
--- a/Engine/Source/Runtime/Renderer/Private/MobileShadingRenderer.cpp
+++ b/Engine/Source/Runtime/Renderer/Private/MobileShadingRenderer.cpp
@@ -556,7 +556,7 @@ void FMobileSceneRenderer::InitViews(
 	const bool bSeparateTranslucencyActive = IsMobileSeparateTranslucencyActive(Views.GetData(), Views.Num()); 
 	const bool bPostProcessUsesSceneDepth = PostProcessUsesSceneDepth(Views[0]) || IsMobileDistortionActive(Views[0]);
 	const bool bRequireSeparateViewPass = Views.Num() > 1 && !Views[0].bIsMobileMultiViewEnabled;
-	bRequiresMultiPass = RequiresMultiPass(RHICmdList, Views[0]);
+	bRequiresMultiPass = RequiresMultiPass(Views[0]);
 	bKeepDepthContent =
 		bRequiresMultiPass ||
 		bForceDepthResolve ||
@@ -805,7 +805,7 @@ void FMobileSceneRenderer::RenderFullDepthPrepass(FRDGBuilder& GraphBuilder, FSc
 			RDG_EVENT_NAME("FullDepthPrepass"),
 			PassParameters,
 			ERDGPassFlags::Raster,
-			[this, PassParameters, &View, bDoOcclusionQueries](FRHICommandListImmediate& RHICmdList)
+			[this, PassParameters, &View, bDoOcclusionQueries](FRHICommandList& RHICmdList)
 			{
 				RenderPrePass(RHICmdList, View);
 
@@ -820,7 +820,7 @@ void FMobileSceneRenderer::RenderFullDepthPrepass(FRDGBuilder& GraphBuilder, FSc
 	FenceOcclusionTests(GraphBuilder);
 }
 
-void FMobileSceneRenderer::RenderMaskedPrePass(FRHICommandListImmediate& RHICmdList, const FViewInfo& View)
+void FMobileSceneRenderer::RenderMaskedPrePass(FRHICommandList& RHICmdList, const FViewInfo& View)
 {
 	if (bIsMaskedOnlyDepthPrepassEnabled)
 	{
@@ -1297,7 +1297,7 @@ void FMobileSceneRenderer::RenderForwardSinglePass(FRDGBuilder& GraphBuilder, FM
 		PassParameters,
 		// the second view pass should not be merged with the first view pass on mobile since the subpass would not work properly.
 		ERDGPassFlags::Raster | ERDGPassFlags::NeverMerge,
-		[this, PassParameters, ViewContext, bDoOcclusionQueries, &SceneTextures](FRHICommandListImmediate& RHICmdList)
+		[this, PassParameters, ViewContext, bDoOcclusionQueries, &SceneTextures](FRHICommandList& RHICmdList)
 	{
 		FViewInfo& View = *ViewContext.ViewInfo;
 			
@@ -1354,7 +1354,7 @@ void FMobileSceneRenderer::RenderForwardMultiPass(FRDGBuilder& GraphBuilder, FMo
 		RDG_EVENT_NAME("SceneColorRendering"),
 		PassParameters,
 		ERDGPassFlags::Raster,
-		[this, PassParameters, ViewContext, &SceneTextures](FRHICommandListImmediate& RHICmdList)
+		[this, PassParameters, ViewContext, &SceneTextures](FRHICommandList& RHICmdList)
 	{
 		FViewInfo& View = *ViewContext.ViewInfo;
 			
@@ -1411,7 +1411,7 @@ void FMobileSceneRenderer::RenderForwardMultiPass(FRDGBuilder& GraphBuilder, FMo
 		RDG_EVENT_NAME("DecalsAndTranslucency"),
 		SecondPassParameters,
 		ERDGPassFlags::Raster,
-		[this, SecondPassParameters, ViewContext, bDoOcclusionQueries, &SceneTextures](FRHICommandListImmediate& RHICmdList)
+		[this, SecondPassParameters, ViewContext, bDoOcclusionQueries, &SceneTextures](FRHICommandList& RHICmdList)
 	{
 		FViewInfo& View = *ViewContext.ViewInfo;
 			
@@ -1490,7 +1490,7 @@ class FMobileDeferredCopyDepthPS : public FGlobalShader
 IMPLEMENT_SHADER_TYPE(, FMobileDeferredCopyDepthPS, TEXT("/Engine/Private/MobileDeferredUtils.usf"), TEXT("MobileDeferredCopyDepthPS"), SF_Pixel);
 
 template<class T>
-void MobileDeferredCopyBuffer(FRHICommandListImmediate& RHICmdList, const FViewInfo& View)
+void MobileDeferredCopyBuffer(FRHICommandList& RHICmdList, const FViewInfo& View)
 {
 	FGraphicsPipelineStateInitializer GraphicsPSOInit;
 	RHICmdList.ApplyCachedRenderTargets(GraphicsPSOInit);
@@ -1624,7 +1624,7 @@ void FMobileSceneRenderer::RenderDeferredSinglePass(FRDGBuilder& GraphBuilder, c
 		PassParameters,
 		// the second view pass should not be merged with the first view pass on mobile since the subpass would not work properly.
 		ERDGPassFlags::Raster | ERDGPassFlags::NeverMerge,
-		[this, PassParameters, ViewContext, bDoOcclusionQueires, &SceneTextures, &SortedLightSet, bUsingPixelLocalStorage](FRHICommandListImmediate& RHICmdList)
+		[this, PassParameters, ViewContext, bDoOcclusionQueires, &SceneTextures, &SortedLightSet, bUsingPixelLocalStorage](FRHICommandList& RHICmdList)
 	{
 		FViewInfo& View = *ViewContext.ViewInfo;
 			
@@ -1667,7 +1667,7 @@ void FMobileSceneRenderer::RenderDeferredMultiPass(FRDGBuilder& GraphBuilder, cl
 		RDG_EVENT_NAME("SceneColorRendering"),
 		PassParameters,
 		ERDGPassFlags::Raster,
-		[this, PassParameters, ViewContext, &SceneTextures](FRHICommandListImmediate& RHICmdList)
+		[this, PassParameters, ViewContext, &SceneTextures](FRHICommandList& RHICmdList)
 	{
 		FViewInfo& View = *ViewContext.ViewInfo;
 		
@@ -1702,7 +1702,7 @@ void FMobileSceneRenderer::RenderDeferredMultiPass(FRDGBuilder& GraphBuilder, cl
 		RDG_EVENT_NAME("Decals"),
 		SecondPassParameters,
 		ERDGPassFlags::Raster,
-		[this, SecondPassParameters, ViewContext](FRHICommandListImmediate& RHICmdList)
+		[this, SecondPassParameters, ViewContext](FRHICommandList& RHICmdList)
 	{
 		FViewInfo& View = *ViewContext.ViewInfo;
 		RenderDecals(RHICmdList, View);
@@ -1724,7 +1724,7 @@ void FMobileSceneRenderer::RenderDeferredMultiPass(FRDGBuilder& GraphBuilder, cl
 		RDG_EVENT_NAME("LightingAndTranslucency"),
 		ThirdPassParameters,
 		ERDGPassFlags::Raster,
-		[this, ThirdPassParameters, ViewContext, bDoOcclusionQueires, &SceneTextures, &SortedLightSet](FRHICommandListImmediate& RHICmdList)
+		[this, ThirdPassParameters, ViewContext, bDoOcclusionQueires, &SceneTextures, &SortedLightSet](FRHICommandList& RHICmdList)
 	{
 		FViewInfo& View = *ViewContext.ViewInfo;
 			
@@ -1743,7 +1743,7 @@ void FMobileSceneRenderer::RenderDeferredMultiPass(FRDGBuilder& GraphBuilder, cl
 	});
 }
 
-void FMobileSceneRenderer::PostRenderBasePass(FRHICommandListImmediate& RHICmdList, FViewInfo& View)
+void FMobileSceneRenderer::PostRenderBasePass(FRHICommandList& RHICmdList, FViewInfo& View)
 {
 	if (ViewFamily.ViewExtensions.Num() > 1)
 	{
@@ -1756,7 +1756,7 @@ void FMobileSceneRenderer::PostRenderBasePass(FRHICommandListImmediate& RHICmdLi
 	}
 }
 
-void FMobileSceneRenderer::RenderMobileDebugView(FRHICommandListImmediate& RHICmdList, const FViewInfo& View)
+void FMobileSceneRenderer::RenderMobileDebugView(FRHICommandList& RHICmdList, const FViewInfo& View)
 {
 #if WITH_DEBUG_VIEW_MODES
 	if (ViewFamily.UseDebugViewPS())
@@ -1795,7 +1795,7 @@ int32 FMobileSceneRenderer::ComputeNumOcclusionQueriesToBatch() const
 }
 
 // Whether we need a separate render-passes for translucency, decals etc
-bool FMobileSceneRenderer::RequiresMultiPass(FRHICommandListImmediate& RHICmdList, const FViewInfo& View) const
+bool FMobileSceneRenderer::RequiresMultiPass(const FViewInfo& View) const
 {
 	// Vulkan uses subpasses
 	if (IsVulkanPlatform(ShaderPlatform))
@@ -1894,7 +1894,7 @@ public:
 
 IMPLEMENT_SHADER_TYPE(, FPreTonemapMSAA_Mobile,TEXT("/Engine/Private/PostProcessMobile.usf"),TEXT("PreTonemapMSAA_Mobile"),SF_Pixel);
 
-void FMobileSceneRenderer::PreTonemapMSAA(FRHICommandListImmediate& RHICmdList, const FMinimalSceneTextures& SceneTextures)
+void FMobileSceneRenderer::PreTonemapMSAA(FRHICommandList& RHICmdList, const FMinimalSceneTextures& SceneTextures)
 {
 	// iOS only
 	bool bOnChipPP = GSupportsRenderTargetFormat_PF_FloatRGBA && GSupportsShaderFramebufferFetch &&	ViewFamily.EngineShowFlags.PostProcessing;
diff --git a/Engine/Source/Runtime/Renderer/Private/MobileTranslucentRendering.cpp b/Engine/Source/Runtime/Renderer/Private/MobileTranslucentRendering.cpp
--- a/Engine/Source/Runtime/Renderer/Private/MobileTranslucentRendering.cpp
+++ b/Engine/Source/Runtime/Renderer/Private/MobileTranslucentRendering.cpp
@@ -4,7 +4,7 @@
 #include "RenderCore.h"
 #include "TranslucentRendering.h"
 
-void FMobileSceneRenderer::RenderTranslucency(FRHICommandListImmediate& RHICmdList, const FViewInfo& View)
+void FMobileSceneRenderer::RenderTranslucency(FRHICommandList& RHICmdList, const FViewInfo& View)
 {	
 	const bool bShouldRenderTranslucency = ShouldRenderTranslucency(StandardTranslucencyPass) && ViewFamily.EngineShowFlags.Translucency;
 	if (bShouldRenderTranslucency)
diff --git a/Engine/Source/Runtime/Renderer/Private/SceneOcclusion.cpp b/Engine/Source/Runtime/Renderer/Private/SceneOcclusion.cpp
--- a/Engine/Source/Runtime/Renderer/Private/SceneOcclusion.cpp
+++ b/Engine/Source/Runtime/Renderer/Private/SceneOcclusion.cpp
@@ -1484,7 +1484,7 @@ void FDeferredShadingSceneRenderer::RenderOcclusion(
 	}
 }
 
-void FMobileSceneRenderer::RenderOcclusion(FRHICommandListImmediate& RHICmdList)
+void FMobileSceneRenderer::RenderOcclusion(FRHICommandList& RHICmdList)
 {
 	if (!DoOcclusionQueries())
 	{
diff --git a/Engine/Source/Runtime/Renderer/Private/SceneRendering.h b/Engine/Source/Runtime/Renderer/Private/SceneRendering.h
--- a/Engine/Source/Runtime/Renderer/Private/SceneRendering.h
+++ b/Engine/Source/Runtime/Renderer/Private/SceneRendering.h
@@ -2317,9 +2317,7 @@ public:
 	/** Get the ParallelExecuteFlags depending on FeatureLevel*/
 	static ERDGBuilderFlags GetRDGParalelExecuteFlags(ERHIFeatureLevel::Type FeatureLevel)
 	{
-		return  FeatureLevel == ERHIFeatureLevel::ES3_1
-			? ERDGBuilderFlags::None
-			: ERDGBuilderFlags::AllowParallelExecute;
+		return ERDGBuilderFlags::AllowParallelExecute;
 	}
 
 	FORCEINLINE FSceneTextures& GetActiveSceneTextures() { return ViewFamily.GetSceneTextures(); }
@@ -2720,25 +2718,25 @@ protected:
 		FVirtualTextureUpdater* VirtualTextureUpdater,
 		FInitViewTaskDatas& TaskDatas);
 
-	void RenderPrePass(FRHICommandListImmediate& RHICmdList, const FViewInfo& View);
-	void RenderMaskedPrePass(FRHICommandListImmediate& RHICmdList, const FViewInfo& View);
+	void RenderPrePass(FRHICommandList& RHICmdList, const FViewInfo& View);
+	void RenderMaskedPrePass(FRHICommandList& RHICmdList, const FViewInfo& View);
 	void RenderFullDepthPrepass(FRDGBuilder& GraphBuilder, FSceneTextures& SceneTextures);
 
 	/** Renders the opaque base pass for mobile. */
-	void RenderMobileBasePass(FRHICommandListImmediate& RHICmdList, const FViewInfo& View);
+	void RenderMobileBasePass(FRHICommandList& RHICmdList, const FViewInfo& View);
 
-	void PostRenderBasePass(FRHICommandListImmediate& RHICmdList, FViewInfo& View);
+	void PostRenderBasePass(FRHICommandList& RHICmdList, FViewInfo& View);
 
 	void RenderMobileEditorPrimitives(FRHICommandList& RHICmdList, const FViewInfo& View, const FMeshPassProcessorRenderState& DrawRenderState);
 
 	/** Renders the debug view pass for mobile. */
-	void RenderMobileDebugView(FRHICommandListImmediate& RHICmdList, const FViewInfo& View);
+	void RenderMobileDebugView(FRHICommandList& RHICmdList, const FViewInfo& View);
 
 	/** Render modulated shadow projections in to the scene, loops over any unrendered shadows until all are processed.*/
-	void RenderModulatedShadowProjections(FRHICommandListImmediate& RHICmdList, int32 ViewIndex, const FViewInfo& View);
+	void RenderModulatedShadowProjections(FRHICommandList& RHICmdList, int32 ViewIndex, const FViewInfo& View);
 
 	/** Issues occlusion queries */
-	void RenderOcclusion(FRHICommandListImmediate& RHICmdList);
+	void RenderOcclusion(FRHICommandList& RHICmdList);
 	
 	bool ShouldRenderHZB();
 
@@ -2751,19 +2749,19 @@ protected:
 	int32 ComputeNumOcclusionQueriesToBatch() const;
 
 	/** Whether platform requires multiple render-passes for SceneColor rendering */
-	bool RequiresMultiPass(FRHICommandListImmediate& RHICmdList, const FViewInfo& View) const;
+	bool RequiresMultiPass(const FViewInfo& View) const;
 
 	/** Renders decals. */
-	void RenderDecals(FRHICommandListImmediate& RHICmdList, const FViewInfo& View);
+	void RenderDecals(FRHICommandList& RHICmdList, const FViewInfo& View);
 
 	/** Renders the atmospheric and height fog */
-	void RenderFog(FRHICommandListImmediate& RHICmdList, const FViewInfo& View);
+	void RenderFog(FRHICommandList& RHICmdList, const FViewInfo& View);
 
 	/** Renders the base pass for translucency. */
-	void RenderTranslucency(FRHICommandListImmediate& RHICmdList, const FViewInfo& View);
+	void RenderTranslucency(FRHICommandList& RHICmdList, const FViewInfo& View);
 
 	/** On chip pre-tonemap before scene color MSAA resolve (iOS only) */
-	void PreTonemapMSAA(FRHICommandListImmediate& RHICmdList, const FMinimalSceneTextures& SceneTextures);
+	void PreTonemapMSAA(FRHICommandList& RHICmdList, const FMinimalSceneTextures& SceneTextures);
 
 	void SetupMobileBasePassAfterShadowInit(FExclusiveDepthStencil::Type BasePassDepthStencilAccess, TArrayView<FViewCommands> ViewCommandsPerView, FInstanceCullingManager& InstanceCullingManager);
 
diff --git a/Engine/Source/Runtime/Renderer/Private/ShadowRendering.cpp b/Engine/Source/Runtime/Renderer/Private/ShadowRendering.cpp
--- a/Engine/Source/Runtime/Renderer/Private/ShadowRendering.cpp
+++ b/Engine/Source/Runtime/Renderer/Private/ShadowRendering.cpp
@@ -1434,7 +1434,7 @@ void FProjectedShadowInfo::RenderProjectionInternal(
 }
 
 void FProjectedShadowInfo::RenderMobileModulatedShadowProjection(
-	FRHICommandListImmediate& RHICmdList,
+	FRHICommandList& RHICmdList,
 	int32 ViewIndex,
 	const FViewInfo* View,
 	const FLightSceneProxy* LightSceneProxy,
@@ -2343,7 +2343,7 @@ void FDeferredShadingSceneRenderer::RenderDeferredShadowProjections(
 	}
 }
 
-void FMobileSceneRenderer::RenderModulatedShadowProjections(FRHICommandListImmediate& RHICmdList, int32 ViewIndex, const FViewInfo& View)
+void FMobileSceneRenderer::RenderModulatedShadowProjections(FRHICommandList& RHICmdList, int32 ViewIndex, const FViewInfo& View)
 {
 	if (!ViewFamily.EngineShowFlags.DynamicShadows || View.bIsPlanarReflection || bRequiresShadowProjections)
 	{
diff --git a/Engine/Source/Runtime/Renderer/Private/ShadowRendering.h b/Engine/Source/Runtime/Renderer/Private/ShadowRendering.h
--- a/Engine/Source/Runtime/Renderer/Private/ShadowRendering.h
+++ b/Engine/Source/Runtime/Renderer/Private/ShadowRendering.h
@@ -599,7 +599,7 @@ public:
 	 * Projects the mobile modulated shadow onto the scene for a particular view.
 	 */
 	void RenderMobileModulatedShadowProjection(
-		FRHICommandListImmediate& RHICmdList,
+		FRHICommandList& RHICmdList,
 		int32 ViewIndex,
 		const FViewInfo* View,
 		const FLightSceneProxy* LightSceneProxy,
diff --git a/Engine/Source/Runtime/VulkanRHI/Private/VulkanUniformBuffer.cpp b/Engine/Source/Runtime/VulkanRHI/Private/VulkanUniformBuffer.cpp
--- a/Engine/Source/Runtime/VulkanRHI/Private/VulkanUniformBuffer.cpp
+++ b/Engine/Source/Runtime/VulkanRHI/Private/VulkanUniformBuffer.cpp
@@ -226,8 +226,9 @@ FVulkanUniformBuffer::FVulkanUniformBuffer(FVulkanDevice& InDevice, const FRHIUn
 	{
 		const bool bInRenderingThread = IsInRenderingThread();
 		const bool bInRHIThread = IsInRHIThread();
-
-		if (UseRingBuffer(InUsage) && (bInRenderingThread || bInRHIThread))
+PRAGMA_DISABLE_DEPRECATION_WARNINGS
+		if (UseRingBuffer(InUsage) && (bInRenderingThread || bInRHIThread) && !UE::Tasks::Private::IsThreadRetractingTask())
+PRAGMA_ENABLE_DEPRECATION_WARNINGS
 		{
 			if (Contents)
 			{
@@ -268,6 +269,7 @@ FVulkanUniformBuffer::FVulkanUniformBuffer(FVulkanDevice& InDevice, const FRHIUn
 				Allocation.FlushMappedMemory(Device);
 			}
 		}
+
 	}
 }
 
diff --git a/Engine/Source/Runtime/Windows/D3D11RHI/Private/D3D11UniformBuffer.cpp b/Engine/Source/Runtime/Windows/D3D11RHI/Private/D3D11UniformBuffer.cpp
--- a/Engine/Source/Runtime/Windows/D3D11RHI/Private/D3D11UniformBuffer.cpp
+++ b/Engine/Source/Runtime/Windows/D3D11RHI/Private/D3D11UniformBuffer.cpp
@@ -186,7 +186,9 @@ FUniformBufferRHIRef FD3D11DynamicRHI::RHICreateUniformBuffer(const void* Conten
 
 		SCOPE_CYCLE_COUNTER(STAT_D3D11UpdateUniformBufferTime);
 
-		if (IsPoolingEnabled() && (IsInRenderingThread() || IsInRHIThread()))
+PRAGMA_DISABLE_DEPRECATION_WARNINGS
+		if (IsPoolingEnabled() && (IsInRenderingThread() || IsInRHIThread()) && !UE::Tasks::Private::IsThreadRetractingTask())
+PRAGMA_ENABLE_DEPRECATION_WARNINGS
 		{
 			const bool bAllocatedFromPool = true;
 
