Substrata - SM5 does not support waveops so fix that situation to avoid artefacts.

diff --git a/Engine/Source/Runtime/Renderer/Private/Strata/Strata.cpp b/Engine/Source/Runtime/Renderer/Private/Strata/Strata.cpp
--- a/Engine/Source/Runtime/Renderer/Private/Strata/Strata.cpp
+++ b/Engine/Source/Runtime/Renderer/Private/Strata/Strata.cpp
@@ -607,6 +607,17 @@ static void BindStrataPublicGlobalUniformParameters(FRDGBuilder& GraphBuilder, F
 	//TODO: Other Strata scene textures or other globals.
 }
 
+static ERHIFeatureSupport SubstrateSupportsWaveOps(EShaderPlatform Platform)
+{
+	// D3D11 / SM5 do not support wave-ops by default, that fixes classification and black/wrong tiling.
+	if (Platform == SP_PCD3D_SM5)
+	{
+		return ERHIFeatureSupport::Unsupported;
+	}
+
+	return FDataDrivenShaderPlatformInfo::GetSupportsWaveOperations(Platform);
+}
+
 TRDGUniformBufferRef<FStrataPublicGlobalUniformParameters> CreatePublicGlobalUniformBuffer(FRDGBuilder& GraphBuilder, FStrataSceneData* StrataScene)
 {
 	FStrataPublicGlobalUniformParameters* StrataPublicUniformParameters = GraphBuilder.AllocParameters<FStrataPublicGlobalUniformParameters>();
@@ -648,7 +659,7 @@ class FStrataBSDFTilePassCS : public FGlobalShader
 
 	static bool ShouldCompilePermutation(const FGlobalShaderPermutationParameters& Parameters)
 	{
-		const bool bUseWaveIntrinsics = FDataDrivenShaderPlatformInfo::GetSupportsWaveOperations(Parameters.Platform) != ERHIFeatureSupport::Unsupported;
+		const bool bUseWaveIntrinsics = SubstrateSupportsWaveOps(Parameters.Platform) != ERHIFeatureSupport::Unsupported;
 		FPermutationDomain PermutationVector(Parameters.PermutationId);
 		if (PermutationVector.Get<FWaveOps>() && !bUseWaveIntrinsics)
 		{
@@ -708,7 +719,7 @@ class FStrataMaterialTileClassificationPassCS : public FGlobalShader
 
 	static bool ShouldCompilePermutation(const FGlobalShaderPermutationParameters& Parameters)
 	{
-		const bool bUseWaveIntrinsics = FDataDrivenShaderPlatformInfo::GetSupportsWaveOperations(Parameters.Platform) != ERHIFeatureSupport::Unsupported;
+		const bool bUseWaveIntrinsics = SubstrateSupportsWaveOps(Parameters.Platform) != ERHIFeatureSupport::Unsupported;
 		FPermutationDomain PermutationVector(Parameters.PermutationId);
 		if (PermutationVector.Get<FWaveOps>() && !bUseWaveIntrinsics)
 		{
@@ -1190,7 +1201,9 @@ void AddStrataMaterialClassificationPass(FRDGBuilder& GraphBuilder, const FMinim
 		RDG_EVENT_SCOPE_CONDITIONAL(GraphBuilder, Views.Num() > 1, "View%d", i);
 
 		const FViewInfo& View = Views[i];
-		bool bWaveOps = GRHISupportsWaveOperations && GRHIMaximumWaveSize >= 64 && FDataDrivenShaderPlatformInfo::GetSupportsWaveOperations(View.GetShaderPlatform()) != ERHIFeatureSupport::Unsupported;
+		EShaderPlatform Platform = View.GetShaderPlatform();
+
+		const bool bWaveOps = GRHISupportsWaveOperations&& GRHIMaximumWaveSize >= 64 && SubstrateSupportsWaveOps(Platform) != ERHIFeatureSupport::Unsupported;
 		
 		const FStrataViewData* StrataViewData = &View.StrataViewData;
 		const FStrataSceneData* StrataSceneData = View.StrataViewData.SceneData;
@@ -1199,7 +1212,7 @@ void AddStrataMaterialClassificationPass(FRDGBuilder& GraphBuilder, const FMinim
 		{
 			// When the platform support explicit CMask texture, we disable material data bufferclear. Material buffer buffer clear (the header part) is done during the classification pass.  
 			// To reduce the reading bandwidth, we rely on TopLayerData CMask to 'drive' the clearing process. This allows to clear quickly empty tiles.
-			const bool bSupportCMask = SupportsCMask(View.GetShaderPlatform());
+			const bool bSupportCMask = SupportsCMask(Platform);
 			FRDGTextureRef TopLayerCmaskTexture = StrataSceneData->TopLayerTexture;			
 			if (bSupportCMask)
 			{
@@ -1209,7 +1222,7 @@ void AddStrataMaterialClassificationPass(FRDGBuilder& GraphBuilder, const FMinim
 			}
 
 			// If Dbuffer pass (i.e. apply DBuffer data after the base-pass) is enabled, run special classification for outputing tile with/without tiles
-			const bool bDBufferTiles = IsDBufferPassEnabled(View.GetShaderPlatform()) && CVarSubstrateDBufferPassDedicatedTiles.GetValueOnRenderThread() > 0 && DBufferTextures.IsValid() && IsConsolePlatform(View.GetShaderPlatform());
+			const bool bDBufferTiles = IsDBufferPassEnabled(Platform) && CVarSubstrateDBufferPassDedicatedTiles.GetValueOnRenderThread() > 0 && DBufferTextures.IsValid() && IsConsolePlatform(View.GetShaderPlatform());
 
 			FStrataMaterialTileClassificationPassCS::FPermutationDomain PermutationVector;
 			PermutationVector.Set< FStrataMaterialTileClassificationPassCS::FCmask >(bSupportCMask);
@@ -1227,7 +1240,7 @@ void AddStrataMaterialClassificationPass(FRDGBuilder& GraphBuilder, const FMinim
 			PassParameters->MaterialTextureArrayUAV = StrataSceneData->MaterialTextureArrayUAV;
 			PassParameters->OpaqueRoughRefractionTexture = StrataSceneData->OpaqueRoughRefractionTexture;
 			PassParameters->TileDrawIndirectDataBuffer = StrataViewData->ClassificationTileDrawIndirectBufferUAV;
-			PassParameters->DBuffer = GetDBufferParameters(GraphBuilder, DBufferTextures, View.GetShaderPlatform());
+			PassParameters->DBuffer = GetDBufferParameters(GraphBuilder, DBufferTextures, Platform);
 			PassParameters->SceneStencilTexture = SceneTextures.Stencil;
 			// Regular tiles
 			PassParameters->SimpleTileListDataBuffer = StrataViewData->ClassificationTileListBufferUAV[EStrataTileType::ESimple];
