Changes velocity tile texture to a Texture2DArray to save one UAV in FTSRDilateVelocityCS

diff --git a/Engine/Shaders/Private/MotionBlur/MotionBlurCommon.ush b/Engine/Shaders/Private/MotionBlur/MotionBlurCommon.ush
--- a/Engine/Shaders/Private/MotionBlur/MotionBlurCommon.ush
+++ b/Engine/Shaders/Private/MotionBlur/MotionBlurCommon.ush
@@ -230,14 +230,13 @@ FVelocityRange PolarToCartesian(FVelocityRange PolarRange)
 }
 
 void StoreVelocityRange(
-	RWTexture2D<float4> VelocityFlatTexture_0,
-	RWTexture2D<float4>	VelocityTileTexture_1,
+	RWTexture2DArray<float4> VelocityFlatTextureArray,
 	uint2 PixelCoord,
 	FVelocityRange CartesianRange)
 {
-	VelocityFlatTexture_0[PixelCoord] = float4(CartesianRange.Min, CartesianRange.Max[0]);
+	VelocityFlatTextureArray[uint3(PixelCoord, 0)] = float4(CartesianRange.Min, CartesianRange.Max[0]);
 	#if CONFIG_MAX_RANGE_SIZE > 1
-		VelocityTileTexture_1[PixelCoord] = float4(CartesianRange.Max[1], 0.0, 0.0);
+		VelocityFlatTextureArray[uint3(PixelCoord, 1)] = float4(CartesianRange.Max[1], 0.0, 0.0);
 	#endif
 }
 
diff --git a/Engine/Shaders/Private/MotionBlur/MotionBlurTileGather.usf b/Engine/Shaders/Private/MotionBlur/MotionBlurTileGather.usf
--- a/Engine/Shaders/Private/MotionBlur/MotionBlurTileGather.usf
+++ b/Engine/Shaders/Private/MotionBlur/MotionBlurTileGather.usf
@@ -16,8 +16,7 @@ Texture2D CenterVelocityTileTextures_Textures_1;
 Texture2D VelocityTileTextures_Textures_0;
 Texture2D VelocityTileTextures_Textures_1;
 
-RWTexture2D<float4>	OutVelocityTile_Textures_0;
-RWTexture2D<float4>	OutVelocityTile_Textures_1;
+RWTexture2DArray<float4>	OutVelocityTileArray;
 
 
 //------------------------------------------------------- ENTRY POINT
@@ -87,6 +86,6 @@ void MainCS(
 	BRANCH
 	if (all(PixelPos.xy < VelocityTile_ViewportMax))
 	{
-		StoreVelocityRange(OutVelocityTile_Textures_0, OutVelocityTile_Textures_1, PixelPos, PolarToCartesian(VelocityPolarRange));
+		StoreVelocityRange(OutVelocityTileArray, PixelPos, PolarToCartesian(VelocityPolarRange));
 	}
 }
diff --git a/Engine/Shaders/Private/MotionBlur/MotionBlurVelocityFlatten.usf b/Engine/Shaders/Private/MotionBlur/MotionBlurVelocityFlatten.usf
--- a/Engine/Shaders/Private/MotionBlur/MotionBlurVelocityFlatten.usf
+++ b/Engine/Shaders/Private/MotionBlur/MotionBlurVelocityFlatten.usf
@@ -29,8 +29,7 @@ Texture2D VelocityTexture;
 Texture2D DepthTexture;
 
 RWTexture2D<float3>	OutVelocityFlatTexture;
-RWTexture2D<float4>	OutVelocityTile_Textures_0;
-RWTexture2D<float4>	OutVelocityTile_Textures_1;
+RWTexture2DArray<float4>	OutVelocityTileArray;
 
 RWTexture2D<float4>	DebugOutput;
 
@@ -138,7 +137,7 @@ void MainCS(
 	// Output tile
 	{
 		uint2 OutputPos = GroupIndex == 0 ? GroupId : uint(~0).xx;
-		StoreVelocityRange(OutVelocityTile_Textures_0, OutVelocityTile_Textures_1, OutputPos, PolarToCartesian(VelocityPolarRange));
+		StoreVelocityRange(OutVelocityTileArray, OutputPos, PolarToCartesian(VelocityPolarRange));
 	}
 	
 #if DEBUG_MOTION_BLUR_OUTPUT
diff --git a/Engine/Shaders/Private/TemporalSuperResolution/TSRDilateVelocity.usf b/Engine/Shaders/Private/TemporalSuperResolution/TSRDilateVelocity.usf
--- a/Engine/Shaders/Private/TemporalSuperResolution/TSRDilateVelocity.usf
+++ b/Engine/Shaders/Private/TemporalSuperResolution/TSRDilateVelocity.usf
@@ -62,8 +62,7 @@ RWTexture2DArray<uint> R8Output;
 RWTexture2D<uint> SubpixelDepthOutput;
 
 RWTexture2D<float3>	VelocityFlattenOutput;
-RWTexture2D<float4>	VelocityTileOutput0;
-RWTexture2D<float4>	VelocityTileOutput1;
+RWTexture2DArray<float4> VelocityTileArrayOutput;
 
 
 //------------------------------------------------------- ENTRY POINT
@@ -470,7 +469,7 @@ void MainCS(
 			VelocityFlattenOutput[OutputPixelPos] = EncodeVelocityFlatten(VelocityPolar, ClosestDeviceZ);
 			
 			uint2 OutputTilePos = GroupThreadIndex == 0 ? GroupId : uint(~0).xx;
-			StoreVelocityRange(VelocityTileOutput0, VelocityTileOutput1, OutputTilePos, PolarToCartesian(VelocityPolarRange));
+			StoreVelocityRange(VelocityTileArrayOutput, OutputTilePos, PolarToCartesian(VelocityPolarRange));
 		}
 		#endif
 
diff --git a/Engine/Source/Runtime/Renderer/Private/PostProcess/PostProcessMotionBlur.cpp b/Engine/Source/Runtime/Renderer/Private/PostProcess/PostProcessMotionBlur.cpp
--- a/Engine/Source/Runtime/Renderer/Private/PostProcess/PostProcessMotionBlur.cpp
+++ b/Engine/Source/Runtime/Renderer/Private/PostProcess/PostProcessMotionBlur.cpp
@@ -251,41 +251,30 @@ FRDGTextureUAVRef CreateDebugUAV(FRDGBuilder& GraphBuilder, const FIntPoint& Ext
 #endif
 }
 
-BEGIN_SHADER_PARAMETER_STRUCT(FVecocityTileTextures, )
-	SHADER_PARAMETER_RDG_TEXTURE_ARRAY(Texture2D, Textures, [FVelocityFlattenTextures::kMaxVelocityTileTextureCount])
+BEGIN_SHADER_PARAMETER_STRUCT(FVecocityTileTextureSRVs, )
+	SHADER_PARAMETER_RDG_TEXTURE_SRV_ARRAY(Texture2D, Textures, [FVelocityFlattenTextures::kMaxVelocityTileTextureCount])
 END_SHADER_PARAMETER_STRUCT()
 
-BEGIN_SHADER_PARAMETER_STRUCT(FVecocityTileUAVs, )
-	SHADER_PARAMETER_RDG_TEXTURE_UAV_ARRAY(RWTexture2D, Textures, [FVelocityFlattenTextures::kMaxVelocityTileTextureCount])
-END_SHADER_PARAMETER_STRUCT()
-
-FVecocityTileTextures CreateVecocityTileTextures(FRDGBuilder& GraphBuilder, FIntPoint VelocityTileCount, const TCHAR* DebugName, bool ScatterDilatation = false)
+FRDGTextureRef CreateVelocityTileTexture(FRDGBuilder& GraphBuilder, FIntPoint VelocityTileCount, int32 BlurDirections, const TCHAR* DebugName, bool ScatterDilatation = false)
 {
-	FRDGTextureDesc Desc = FRDGTextureDesc::Create2D(
+	FRDGTextureDesc Desc = FRDGTextureDesc::Create2DArray(
 		VelocityTileCount,
 		PF_FloatRGBA,
 		FClearValueBinding::None,
-		GFastVRamConfig.MotionBlur | TexCreate_ShaderResource | (ScatterDilatation ? TexCreate_RenderTargetable : TexCreate_UAV));
+		GFastVRamConfig.MotionBlur | TexCreate_ShaderResource | (ScatterDilatation ? TexCreate_RenderTargetable : TexCreate_UAV),
+		/* ArraySize = */ BlurDirections);
 
-	FVecocityTileTextures Textures;
-	Textures.Textures[0] = GraphBuilder.CreateTexture(Desc, DebugName);
-
-	Desc.Format = PF_G16R16F;
-	Textures.Textures[1] = GraphBuilder.CreateTexture(Desc, DebugName);
-	return Textures;
+	return GraphBuilder.CreateTexture(Desc, DebugName);
 }
 
-FVecocityTileUAVs CreateUAVs(FRDGBuilder& GraphBuilder, FVecocityTileTextures VecocityTileTextures)
+FVecocityTileTextureSRVs CreateVelocityTileTextureSRVs(FRDGBuilder& GraphBuilder, FRDGTextureRef VecocityTileArray)
 {
-	FVecocityTileUAVs UAVs;
-	for (int32 i = 0; i < UAVs.Textures.Num(); i++)
+	FVecocityTileTextureSRVs SRVs;
+	for (int32 i = 0; i < VecocityTileArray->Desc.ArraySize; i++)
 	{
-		if (VecocityTileTextures.Textures[i])
-		{
-			UAVs.Textures[i] = GraphBuilder.CreateUAV(VecocityTileTextures.Textures[i]);
-		}
+		SRVs.Textures[i] = GraphBuilder.CreateSRV(FRDGTextureSRVDesc::CreateForSlice(VecocityTileArray, i));
 	}
-	return UAVs;
+	return SRVs;
 }
 
 // Base class for a motion blur / velocity shader.
@@ -327,14 +316,14 @@ public:
 		SHADER_PARAMETER_RDG_TEXTURE(Texture2D, DepthTexture)
 
 		SHADER_PARAMETER_RDG_TEXTURE_UAV(RWTexture2D<float4>, OutVelocityFlatTexture)
-		SHADER_PARAMETER_STRUCT(FVecocityTileUAVs, OutVelocityTile)
+		SHADER_PARAMETER_RDG_TEXTURE_UAV(RWTexture2DArray<float4>, OutVelocityTileArray)
 		SHADER_PARAMETER_RDG_TEXTURE_UAV(RWTexture2D<float4>, DebugOutput)
 	END_SHADER_PARAMETER_STRUCT()
 };
 
 BEGIN_SHADER_PARAMETER_STRUCT(FMotionBlurVelocityDilateParameters, )
 	SHADER_PARAMETER_STRUCT(FScreenPassTextureViewportParameters, VelocityTile)
-	SHADER_PARAMETER_STRUCT(FVecocityTileTextures, VelocityTileTextures)
+	SHADER_PARAMETER_STRUCT(FVecocityTileTextureSRVs, VelocityTileTextures)
 	SHADER_PARAMETER(float, VelocityScaleForFlattenTiles)
 END_SHADER_PARAMETER_STRUCT()
 
@@ -348,8 +337,8 @@ public:
 
 	BEGIN_SHADER_PARAMETER_STRUCT(FParameters, )
 		SHADER_PARAMETER_STRUCT_INCLUDE(FMotionBlurVelocityDilateParameters, Dilate)
-		SHADER_PARAMETER_STRUCT(FVecocityTileTextures, CenterVelocityTileTextures)
-		SHADER_PARAMETER_STRUCT(FVecocityTileUAVs, OutVelocityTile)
+		SHADER_PARAMETER_STRUCT(FVecocityTileTextureSRVs, CenterVelocityTileTextures)
+		SHADER_PARAMETER_RDG_TEXTURE_UAV(RWTexture2DArray<float4>, OutVelocityTileArray)
 		SHADER_PARAMETER_RDG_TEXTURE_UAV(RWTexture2D<float4>, DebugOutput)
 	END_SHADER_PARAMETER_STRUCT()
 
@@ -397,7 +386,7 @@ class FMotionBlurFilterTileClassifyCS : public FMotionBlurShader
 		SHADER_PARAMETER(int32, bAllowHalfResGather)
 		SHADER_PARAMETER(int32, TileListMaxSize)
 		SHADER_PARAMETER(float, HalfResPixelVelocityThreshold)
-		SHADER_PARAMETER_STRUCT(FVecocityTileTextures, VelocityTileTextures)
+		SHADER_PARAMETER_STRUCT(FVecocityTileTextureSRVs, VelocityTileTextures)
 		SHADER_PARAMETER_RDG_BUFFER_UAV(RWStructuredBuffer<uint>, TileListsOutput)
 		SHADER_PARAMETER_RDG_BUFFER_UAV(RWStructuredBuffer<uint>, TileListsSizeOutput)
 		SHADER_PARAMETER_RDG_TEXTURE_UAV(RWTexture2D<float4>, DebugOutput)
@@ -457,7 +446,7 @@ class FMotionBlurFilterCS : public FMotionBlurShader
 
 		SHADER_PARAMETER_RDG_TEXTURE(Texture2D, ColorTexture)
 		SHADER_PARAMETER_RDG_TEXTURE(Texture2D, VelocityFlatTexture)
-		SHADER_PARAMETER_STRUCT(FVecocityTileTextures, VelocityTileTextures)
+		SHADER_PARAMETER_STRUCT(FVecocityTileTextureSRVs, VelocityTileTextures)
 
 		SHADER_PARAMETER_SAMPLER(SamplerState, ColorSampler)
 		SHADER_PARAMETER_SAMPLER(SamplerState, VelocitySampler)
@@ -608,7 +597,7 @@ void AddMotionBlurVelocityPass(
 	const FMotionBlurViewports& Viewports,
 	const FMotionBlurInputs& Inputs,
 	FRDGTextureRef* VelocityFlatTextureOutput,
-	FVecocityTileTextures* VelocityTileTexturesOutput)
+	FVecocityTileTextureSRVs* VelocityTileTexturesOutput)
 {
 	check(VelocityFlatTextureOutput);
 	check(VelocityTileTexturesOutput);
@@ -619,17 +608,14 @@ void AddMotionBlurVelocityPass(
 
 	// Velocity flatten pass: combines depth / velocity into a single target for sampling efficiency.
 	FRDGTextureRef VelocityFlatTexture = nullptr;
-	FVecocityTileTextures VelocityTileTexturesSetup;
+	FRDGTextureRef VelocityTileTextureSetup = nullptr;
 	if (Inputs.VelocityFlattenTextures.IsValid())
 	{
 		ensure(Inputs.VelocityFlattenTextures.VelocityFlatten.ViewRect == View.ViewRect);
-		ensure(Inputs.VelocityFlattenTextures.VelocityTile[0].ViewRect == FIntRect(FIntPoint::ZeroValue, VelocityTileCount));
+		ensure(Inputs.VelocityFlattenTextures.VelocityTileArray.ViewRect == FIntRect(FIntPoint::ZeroValue, VelocityTileCount));
 
 		VelocityFlatTexture = Inputs.VelocityFlattenTextures.VelocityFlatten.Texture;
-		for (int32 i = 0; i < FVelocityFlattenTextures::kMaxVelocityTileTextureCount; i++)
-		{
-			VelocityTileTexturesSetup.Textures[i] = Inputs.VelocityFlattenTextures.VelocityTile[i].Texture;
-		}
+		VelocityTileTextureSetup = Inputs.VelocityFlattenTextures.VelocityTileArray.Texture;
 	}
 	else
 	{
@@ -644,8 +630,8 @@ void AddMotionBlurVelocityPass(
 			VelocityFlatTexture = GraphBuilder.CreateTexture(Desc, TEXT("MotionBlur.VelocityFlatten"));
 		}
 
-		VelocityTileTexturesSetup = CreateVecocityTileTextures(
-			GraphBuilder, VelocityTileCount, TEXT("MotionBlur.VelocityTile"));
+		VelocityTileTextureSetup = CreateVelocityTileTexture(
+			GraphBuilder, VelocityTileCount, BlurDirections, TEXT("MotionBlur.VelocityTile"));
 
 		const bool bEnableCameraMotionBlur = View.bCameraMotionBlur.Get(true);
 		const bool bOverrideCameraMotionBlur = View.ClipToPrevClipOverride.IsSet();
@@ -665,7 +651,7 @@ void AddMotionBlurVelocityPass(
 		PassParameters->VelocityTexture = Inputs.SceneVelocity.Texture;
 
 		PassParameters->OutVelocityFlatTexture = GraphBuilder.CreateUAV(VelocityFlatTexture);
-		PassParameters->OutVelocityTile = CreateUAVs(GraphBuilder, VelocityTileTexturesSetup);
+		PassParameters->OutVelocityTileArray = GraphBuilder.CreateUAV(VelocityTileTextureSetup);
 		PassParameters->DebugOutput = CreateDebugUAV(GraphBuilder, VelocityFlatTexture->Desc.Extent, TEXT("Debug.MotionBlur.Flatten"));
 
 		FMotionBlurVelocityFlattenCS::FPermutationDomain PermutationVector;
@@ -686,14 +672,16 @@ void AddMotionBlurVelocityPass(
 
 	FMotionBlurVelocityDilateParameters VelocityDilateParameters;
 	VelocityDilateParameters.VelocityTile = Viewports.VelocityTileParameters;
-	VelocityDilateParameters.VelocityTileTextures = VelocityTileTexturesSetup;
+	VelocityDilateParameters.VelocityTileTextures = CreateVelocityTileTextureSRVs(GraphBuilder, VelocityTileTextureSetup);
 	VelocityDilateParameters.VelocityScaleForFlattenTiles = (1.0f / float(kMotionBlurFlattenTileSize)) * (float(Viewports.Velocity.Rect.Width()) / float(Viewports.Color.Rect.Width()));
 
-	FVecocityTileTextures VelocityTileTextures;
+	FVecocityTileTextureSRVs VelocityTileTextures;
 	if (ScatterDilatation)
 	{
-		VelocityTileTextures = CreateVecocityTileTextures(
-			GraphBuilder, VelocityTileCount, TEXT("MotionBlur.ScatteredVelocityTile"), /* ScatterDilatation = */ true);
+		FRDGTextureRef VelocityTileArrayTexture = CreateVelocityTileTexture(
+			GraphBuilder, VelocityTileCount, /* BlurDirections = */ 1, TEXT("MotionBlur.ScatteredVelocityTile"), /* ScatterDilatation = */ true);
+
+		VelocityTileTextures.Textures[0] = GraphBuilder.CreateSRV(FRDGTextureSRVDesc::CreateForSlice(VelocityTileArrayTexture, 0));
 
 		FRDGTextureRef VelocityTileDepthTexture =
 			GraphBuilder.CreateTexture(
@@ -716,8 +704,10 @@ void AddMotionBlurVelocityPass(
 
 		PassParameters->RenderTargets[0] =
 			FRenderTargetBinding(
-				VelocityTileTextures.Textures[0],
-				ERenderTargetLoadAction::ENoAction);
+				VelocityTileArrayTexture,
+				ERenderTargetLoadAction::ENoAction,
+				/* MipIndex = */ 0,
+				/* ArraySlice = */ 0);
 
 		TShaderMapRef<FMotionBlurVelocityDilateScatterVS> VertexShader(View.ShaderMap);
 		TShaderMapRef<FMotionBlurVelocityDilateScatterPS> PixelShader(View.ShaderMap);
@@ -777,8 +767,8 @@ void AddMotionBlurVelocityPass(
 	// ScatterGather the dilatation
 	if (!ScatterDilatation || BlurDirections > 1)
 	{
-		FVecocityTileTextures DilatedTileTextures = CreateVecocityTileTextures(
-			GraphBuilder, VelocityTileCount, TEXT("MotionBlur.GatheredVelocityTile"));
+		FRDGTextureRef DilatedTileTexture = CreateVelocityTileTexture(
+			GraphBuilder, VelocityTileCount, BlurDirections, TEXT("MotionBlur.GatheredVelocityTile"));
 
 		FMotionBlurVelocityDilateGatherCS::FParameters* PassParameters = GraphBuilder.AllocParameters<FMotionBlurVelocityDilateGatherCS::FParameters>();
 		PassParameters->Dilate = VelocityDilateParameters;
@@ -787,13 +777,13 @@ void AddMotionBlurVelocityPass(
 			check(BlurDirections > 1);
 			// Feed the scattered min max to make sure ReducePolarVelocityRange() is aware the highest velocity might be coming from tiles further away
 			PassParameters->CenterVelocityTileTextures.Textures[0] = VelocityTileTextures.Textures[0];
-			PassParameters->CenterVelocityTileTextures.Textures[1] = VelocityTileTexturesSetup.Textures[1];
+			PassParameters->CenterVelocityTileTextures.Textures[1] = VelocityDilateParameters.VelocityTileTextures.Textures[1];
 		}
 		else
 		{
-			PassParameters->CenterVelocityTileTextures = VelocityTileTexturesSetup;
+			PassParameters->CenterVelocityTileTextures = VelocityDilateParameters.VelocityTileTextures;
 		}
-		PassParameters->OutVelocityTile = CreateUAVs(GraphBuilder, DilatedTileTextures);
+		PassParameters->OutVelocityTileArray = GraphBuilder.CreateUAV(DilatedTileTexture);
 
 		FMotionBlurVelocityDilateGatherCS::FPermutationDomain PermutationVector;
 		PermutationVector.Set<FMotionBlurDirections>(BlurDirections);
@@ -808,13 +798,11 @@ void AddMotionBlurVelocityPass(
 
 		if (!ScatterDilatation)
 		{
-			VelocityTileTextures = DilatedTileTextures;
+			VelocityTileTextures = CreateVelocityTileTextureSRVs(GraphBuilder, DilatedTileTexture);
 		}
 		else if (BlurDirections > 1)
 		{
-			// VelocityTileTextures.Textures[0] is actually scattered.
-			GraphBuilder.RemoveUnusedTextureWarning(DilatedTileTextures.Textures[0]);
-			VelocityTileTextures.Textures[1] = DilatedTileTextures.Textures[1];
+			VelocityTileTextures.Textures[1] = CreateVelocityTileTextureSRVs(GraphBuilder, DilatedTileTexture).Textures[1];
 		}
 	}
 
@@ -829,7 +817,7 @@ FMotionBlurOutputs AddMotionBlurFilterPass(
 	const FMotionBlurViewports& Viewports,
 	FRDGTextureRef ColorTexture,
 	FRDGTextureRef VelocityFlatTexture,
-	FVecocityTileTextures VelocityTileTextures,
+	FVecocityTileTextureSRVs VelocityTileTextures,
 	FRDGTextureRef PostMotionBlurTranslucency,
 	const FIntPoint& PostMotionBlurTranslucencySize,
 	EMotionBlurFilterPass MotionBlurFilterPass,
@@ -1199,7 +1187,7 @@ FMotionBlurOutputs AddMotionBlurPass(FRDGBuilder& GraphBuilder, const FViewInfo&
 	RDG_GPU_STAT_SCOPE(GraphBuilder, MotionBlur);
 
 	FRDGTextureRef VelocityFlatTexture = nullptr;
-	FVecocityTileTextures VelocityTileTextures;
+	FVecocityTileTextureSRVs VelocityTileTextures;
 	AddMotionBlurVelocityPass(
 		GraphBuilder,
 		View,
diff --git a/Engine/Source/Runtime/Renderer/Private/PostProcess/PostProcessMotionBlur.h b/Engine/Source/Runtime/Renderer/Private/PostProcess/PostProcessMotionBlur.h
--- a/Engine/Source/Runtime/Renderer/Private/PostProcess/PostProcessMotionBlur.h
+++ b/Engine/Source/Runtime/Renderer/Private/PostProcess/PostProcessMotionBlur.h
@@ -59,11 +59,11 @@ struct FVelocityFlattenTextures
 	static constexpr int32 kTileSize = 16;
 
 	FScreenPassTexture VelocityFlatten;
-	FScreenPassTexture VelocityTile[kMaxVelocityTileTextureCount];
+	FScreenPassTexture VelocityTileArray;
 
 	bool IsValid() const
 	{
-		return VelocityFlatten.IsValid() && VelocityTile[0].IsValid();
+		return VelocityFlatten.IsValid() && VelocityTileArray.IsValid();
 	}
 
 	// returns whether FVelocityFlattenTextures can be generated in external system.
diff --git a/Engine/Source/Runtime/Renderer/Private/PostProcess/TemporalSuperResolution.cpp b/Engine/Source/Runtime/Renderer/Private/PostProcess/TemporalSuperResolution.cpp
--- a/Engine/Source/Runtime/Renderer/Private/PostProcess/TemporalSuperResolution.cpp
+++ b/Engine/Source/Runtime/Renderer/Private/PostProcess/TemporalSuperResolution.cpp
@@ -495,20 +495,13 @@ class FTSRDilateVelocityCS : public FTSRShader
 		SHADER_PARAMETER_RDG_TEXTURE_UAV(RWTexture2DArray, SubpixelDepthOutput)
 
 		SHADER_PARAMETER_RDG_TEXTURE_UAV(RWTexture2D, VelocityFlattenOutput)
-		SHADER_PARAMETER_RDG_TEXTURE_UAV(RWTexture2D, VelocityTileOutput0)
-		SHADER_PARAMETER_RDG_TEXTURE_UAV(RWTexture2D, VelocityTileOutput1)
+		SHADER_PARAMETER_RDG_TEXTURE_UAV(RWTexture2DArray, VelocityTileArrayOutput)
 
 		SHADER_PARAMETER_RDG_TEXTURE_UAV(RWTexture2DArray, DebugOutput)
 	END_SHADER_PARAMETER_STRUCT()
 
 	static FPermutationDomain RemapPermutation(FPermutationDomain PermutationVector)
 	{
-		// There isn't enough bindable UAV for outputing subpixel depth and motion blur render targets.
-		if (PermutationVector.Get<FSubpixelDepthDim>())
-		{
-			PermutationVector.Set<FMotionBlurDirectionsDim>(0);
-		}
-
 		return PermutationVector;
 	}
 
@@ -1590,7 +1583,7 @@ ITemporalUpscaler::FOutputs AddTemporalSuperResolutionPasses(
 		}
 
 		// Setup up the motion blur's velocity flatten pass.
-		if (PassInputs.bGenerateVelocityFlattenTextures && SubpixelMethod != ETSRSubpixelMethod::ClosestDepth)
+		if (PassInputs.bGenerateVelocityFlattenTextures)
 		{
 			const int32 MotionBlurDirections = GetMotionBlurDirections();
 			PermutationVector.Set<FTSRDilateVelocityCS::FMotionBlurDirectionsDim>(MotionBlurDirections);
@@ -1603,29 +1596,25 @@ ITemporalUpscaler::FOutputs AddTemporalSuperResolutionPasses(
 					FClearValueBinding::None,
 					GFastVRamConfig.VelocityFlat | TexCreate_ShaderResource | TexCreate_UAV);
 
-				VelocityFlattenTextures.VelocityFlatten.Texture = GraphBuilder.CreateTexture(Desc, TEXT("MotionBlur.VelocityTile"));
+				VelocityFlattenTextures.VelocityFlatten.Texture = GraphBuilder.CreateTexture(Desc, TEXT("MotionBlur.VelocityFlatten"));
 				VelocityFlattenTextures.VelocityFlatten.ViewRect = InputRect;
 			}
 
 			{
-				FRDGTextureDesc Desc = FRDGTextureDesc::Create2D(
+				FRDGTextureDesc Desc = FRDGTextureDesc::Create2DArray(
 					FIntPoint::DivideAndRoundUp(InputRect.Size(), FVelocityFlattenTextures::kTileSize),
 					PF_FloatRGBA,
 					FClearValueBinding::None,
-					GFastVRamConfig.MotionBlur | TexCreate_ShaderResource | TexCreate_UAV);
-
-				VelocityFlattenTextures.VelocityTile[0].Texture = GraphBuilder.CreateTexture(Desc, TEXT("MotionBlur.VelocityTile"));
-				VelocityFlattenTextures.VelocityTile[0].ViewRect = FIntRect(FIntPoint::ZeroValue, Desc.Extent);
+					GFastVRamConfig.MotionBlur | TexCreate_ShaderResource | TexCreate_UAV,
+					/* ArraySize = */ MotionBlurDirections);
 
-				Desc.Format = PF_G16R16F;
-				VelocityFlattenTextures.VelocityTile[1].Texture = GraphBuilder.CreateTexture(Desc, TEXT("MotionBlur.VelocityTile"));
-				VelocityFlattenTextures.VelocityTile[1].ViewRect = FIntRect(FIntPoint::ZeroValue, Desc.Extent);
+				VelocityFlattenTextures.VelocityTileArray.Texture = GraphBuilder.CreateTexture(Desc, TEXT("MotionBlur.VelocityTile"));
+				VelocityFlattenTextures.VelocityTileArray.ViewRect = FIntRect(FIntPoint::ZeroValue, Desc.Extent);
 			}
 
 			PassParameters->VelocityFlattenParameters = GetVelocityFlattenParameters(View);
 			PassParameters->VelocityFlattenOutput = GraphBuilder.CreateUAV(VelocityFlattenTextures.VelocityFlatten.Texture);
-			PassParameters->VelocityTileOutput0 = GraphBuilder.CreateUAV(VelocityFlattenTextures.VelocityTile[0].Texture);
-			PassParameters->VelocityTileOutput1 = GraphBuilder.CreateUAV(VelocityFlattenTextures.VelocityTile[1].Texture);
+			PassParameters->VelocityTileArrayOutput = GraphBuilder.CreateUAV(VelocityFlattenTextures.VelocityTileArray.Texture);
 		}
 
 		PassParameters->DebugOutput = CreateDebugUAV(InputExtent, TEXT("Debug.TSR.DilateVelocity"));
