Fix DXC workaround for platform shader backends when Strata is enabled.

#rnx

diff --git a/Engine/Source/Runtime/Renderer/Private/MobileBasePassRendering.h b/Engine/Source/Runtime/Renderer/Private/MobileBasePassRendering.h
--- a/Engine/Source/Runtime/Renderer/Private/MobileBasePassRendering.h
+++ b/Engine/Source/Runtime/Renderer/Private/MobileBasePassRendering.h
@@ -142,7 +142,10 @@ public:
 		// @lh-todo:
 		//		This is a workaround to prevent Z-fighting due to a mismatch between HLSLCC compiled shaders (from cooked materials) and DXC compiled shaders (from user materials).
 		//		It effectively disables CVar "r.Shaders.ForceDXC" which we partially enabled during the HLSLCC to DXC transition for certain platform backends.
-		OutEnvironment.SetCompileArgument(TEXT("WORKAROUND_DISABLE_rShadersForceDXC"), true);
+		if (!Strata::IsStrataEnabled())
+		{
+			OutEnvironment.SetCompileArgument(TEXT("WORKAROUND_DISABLE_rShadersForceDXC"), true);
+		}
 	}
 
 	void GetShaderBindings(
@@ -237,7 +240,10 @@ public:
 		ModifyCompilationEnvironmentForQualityLevel(Parameters.Platform, Parameters.MaterialParameters.QualityLevel, OutEnvironment);
 
 		// @lh-todo: Same workaround as for the VS of MobileBasePass. See TMobileBasePassVSPolicyParamType::ModifyCompilationEnvironment for details.
-		OutEnvironment.SetCompileArgument(TEXT("WORKAROUND_DISABLE_rShadersForceDXC"), true);
+		if (!Strata::IsStrataEnabled())
+		{
+			OutEnvironment.SetCompileArgument(TEXT("WORKAROUND_DISABLE_rShadersForceDXC"), true);
+		}
 	}
 
 	/** Initialization constructor. */
