Change Substrate material buffer adaptive allocatation only increase to avoid recurrent hitch.

diff --git a/Engine/Shaders/Private/Strata/StrataVisualize.usf b/Engine/Shaders/Private/Strata/StrataVisualize.usf
--- a/Engine/Shaders/Private/Strata/StrataVisualize.usf
+++ b/Engine/Shaders/Private/Strata/StrataVisualize.usf
@@ -1476,8 +1476,10 @@ uint ClassificationAsync;
 uint Classification8bits;
 uint bRoughRefraction;
 uint bTileOverflowUseMaterialData;
+uint ViewsMaxBytePerPixel;
 uint ProjectMaxBytePerPixel;
 uint MaterialBufferAllocationInBytes;
+uint MaterialBufferAllocationMode;
 
 float TileOverflowRatio;
 Buffer<uint> ClassificationTileDrawIndirectBuffer;
@@ -1511,6 +1513,22 @@ void MainCS(uint3 DispatchThreadId : SV_DispatchThreadID)
 		Print(Context, TEXT("Mb]"), FontWhite);
 		Newline(Context);
 
+		Print(Context, TEXT(" Views max. byte    : "), LegendColor); 
+		Print(Context, ViewsMaxBytePerPixel, ValueColor, 3, 0);
+		Print(Context, TEXT(" / "), FontWhite); 
+		Print(Context, ProjectMaxBytePerPixel, ValueColor, 3, 0);
+		Newline(Context);
+
+		Print(Context, TEXT(" Allocation Mode    : "), LegendColor);
+		switch (MaterialBufferAllocationMode)
+		{
+			case 0: Print(Context, TEXT("View based"), ValueColor); break;
+			case 1: Print(Context, TEXT("View based | Growing-only"), ValueColor); break;
+			case 2: Print(Context, TEXT("Setting based"), ValueColor); break;
+			default: Print(Context, TEXT("Unknown: "), ValueColor); break;
+		}
+		Newline(Context);
+
 		Print(Context, TEXT(" Shading quality    : "), LegendColor);
 		Print(Context, uint(STRATA_SHADING_QUALITY), ValueColor);
 		Newline(Context);
diff --git a/Engine/Source/Runtime/Renderer/Private/Strata/Strata.cpp b/Engine/Source/Runtime/Renderer/Private/Strata/Strata.cpp
--- a/Engine/Source/Runtime/Renderer/Private/Strata/Strata.cpp
+++ b/Engine/Source/Runtime/Renderer/Private/Strata/Strata.cpp
@@ -63,6 +63,12 @@ static TAutoConsoleVariable<int32> CVarSubstrateDBufferPassDedicatedTiles(
 	TEXT("Use dedicated tile for DBuffer application when DBuffer pass is enabled."),
 	ECVF_RenderThreadSafe);
 
+static TAutoConsoleVariable<int32> CVarSubstrateBytePerPixelMode(
+	TEXT("r.Substrate.BytesPerPixel.Mode"),
+	0,
+	TEXT("Substrate material allocation mode. \n 0: Allocate material buffer based on view requirement \n 1: Allocate material buffer based on platform settings. \n 2: Allocate material buffer based on view requirement, but can only grow over frame to minimize buffer reallocation. "),
+	ECVF_RenderThreadSafe);
+
 IMPLEMENT_GLOBAL_SHADER_PARAMETER_STRUCT(FStrataGlobalUniformParameters, "Strata");
 
 void FStrataViewData::Reset()
@@ -102,6 +108,11 @@ const TCHAR* ToString(EStrataTileType Type)
 namespace Strata
 {
 
+uint32 GetMaterialBufferAllocationMode()
+{
+	return FMath::Clamp(CVarSubstrateBytePerPixelMode.GetValueOnAnyThread(), 0, 2);
+}
+
 enum EStrataTileSpace
 {
 	StrataTileSpace_Primary = 1u,
@@ -325,7 +336,13 @@ static EPixelFormat GetTopLayerTextureFormat()
 void InitialiseStrataFrameSceneData(FRDGBuilder& GraphBuilder, FSceneRenderer& SceneRenderer)
 {
 	FStrataSceneData& Out = SceneRenderer.Scene->StrataSceneData;
-	Out = FStrataSceneData();
+
+	// Reset Substrate scene data
+	{
+		const uint32 MinBytesPerPixel = Out.MinBytesPerPixel;
+		Out = FStrataSceneData();
+		Out.MinBytesPerPixel = MinBytesPerPixel;
+	}
 
 	auto UpdateMaterialBufferToTiledResolution = [](FIntPoint InBufferSizeXY, FIntPoint& OutMaterialBufferSizeXY)
 	{
@@ -334,16 +351,9 @@ void InitialiseStrataFrameSceneData(FRDGBuilder& GraphBuilder, FSceneRenderer& S
 		OutMaterialBufferSizeXY.Y = FMath::DivideAndRoundUp(InBufferSizeXY.Y, STRATA_TILE_SIZE) * STRATA_TILE_SIZE;
 	};
 
+	// Compute the max byte per pixels required by the views
 	bool bNeedBSDFOffsets = false;
 	bool bNeedUAV = false;
-	uint32 MaxBytePerPixel = 0;
-	for (const FViewInfo& View : SceneRenderer.Views)
-	{
-		bNeedBSDFOffsets = bNeedBSDFOffsets || NeedBSDFOffsets(SceneRenderer.Scene, View);
-		bNeedUAV = bNeedUAV || IsDBufferPassEnabled(View.GetShaderPlatform());
-		MaxBytePerPixel = FMath::Max(MaxBytePerPixel, View.StrataViewData.MaxBytePerPixel);
-	}
-	MaxBytePerPixel = FMath::Clamp(MaxBytePerPixel, 4u * STRATA_BASE_PASS_MRT_OUTPUT_COUNT, GetBytePerPixel(SceneRenderer.ShaderPlatform));
 
 	FIntPoint MaterialBufferSizeXY;
 	UpdateMaterialBufferToTiledResolution(FIntPoint(1, 1), MaterialBufferSizeXY);
@@ -357,14 +367,41 @@ void InitialiseStrataFrameSceneData(FRDGBuilder& GraphBuilder, FSceneRenderer& S
 			bAnalyticsInitialized = true;
 		}
 
+		// Gather views' requirements
+		Out.ViewsMaxBytePerPixel = 0;
+		for (const FViewInfo& View : SceneRenderer.Views)
+		{
+			bNeedBSDFOffsets = bNeedBSDFOffsets || NeedBSDFOffsets(SceneRenderer.Scene, View);
+			bNeedUAV = bNeedUAV || IsDBufferPassEnabled(View.GetShaderPlatform());
+			Out.ViewsMaxBytePerPixel = FMath::Max(Out.ViewsMaxBytePerPixel, View.StrataViewData.MaxBytePerPixel);
+
+			// Only use primary views max. byte per pixel as reflection/capture views can bias allocation requirement when using growing-only mode
+			if (!View.bIsPlanarReflection && !View.bIsReflectionCapture && !View.bIsSceneCapture)
+			{
+				Out.MinBytesPerPixel = FMath::Max(Out.MinBytesPerPixel, View.StrataViewData.MaxBytePerPixel);
+			}
+		}
+
+		// Material buffer allocation can use different modes:
+		// * 0: Allocate material buffer based on view requirement 
+		// * 2: Allocate material buffer based on view requirement, but can only grow over frame to minimize buffer reallocation and hitches
+		// * 1: Allocate material buffer based on platform settings. 
+		uint32 MaxBytePerPixel = 0;
+		switch (GetMaterialBufferAllocationMode())
+		{
+			case 0: MaxBytePerPixel = Out.ViewsMaxBytePerPixel; break;
+			case 1: MaxBytePerPixel = FMath::Max(Out.ViewsMaxBytePerPixel, Out.MinBytesPerPixel); break;
+			case 2: MaxBytePerPixel = GetBytePerPixel(SceneRenderer.ShaderPlatform); break;
+		}
+		const uint32 RoundToValue = 4u;
+		MaxBytePerPixel = FMath::Clamp(MaxBytePerPixel, 4u * STRATA_BASE_PASS_MRT_OUTPUT_COUNT, GetBytePerPixel(SceneRenderer.ShaderPlatform));
+		Out.MaxBytesPerPixel = FMath::DivideAndRoundUp(MaxBytePerPixel, RoundToValue) * RoundToValue;
+
 		FIntPoint SceneTextureExtent = SceneRenderer.GetActiveSceneTexturesConfig().Extent;
 		
 		// We need to allocate enough for the tiled memory addressing of material data to always work
 		UpdateMaterialBufferToTiledResolution(SceneTextureExtent, MaterialBufferSizeXY);
 
-		const uint32 RoundToValue = 4u;
-		Out.MaxBytesPerPixel = FMath::DivideAndRoundUp(MaxBytePerPixel, RoundToValue) * RoundToValue;
-
 		// Top layer texture
 		{
 			Out.TopLayerTexture = GraphBuilder.CreateTexture(FRDGTextureDesc::Create2D(SceneTextureExtent, GetTopLayerTextureFormat(), FClearValueBinding::Black, TexCreate_RenderTargetable | TexCreate_ShaderResource | TexCreate_FastVRAM | (bNeedUAV ? TexCreate_UAV : TexCreate_None)), TEXT("Substrate.TopLayerTexture"));
diff --git a/Engine/Source/Runtime/Renderer/Private/Strata/Strata.h b/Engine/Source/Runtime/Renderer/Private/Strata/Strata.h
--- a/Engine/Source/Runtime/Renderer/Private/Strata/Strata.h
+++ b/Engine/Source/Runtime/Renderer/Private/Strata/Strata.h
@@ -95,15 +95,17 @@ const TCHAR* ToString(EStrataTileType Type);
 
 struct FStrataSceneData
 {
-	uint32 MaxBytesPerPixel;
-	bool bRoughDiffuse;
-	int32 PeelLayersAboveDepth;
-	bool bRoughnessTracking;
-
-	int32 SliceStoringDebugStrataTreeDataWithoutMRT;
-	int32 SliceStoringDebugStrataTreeData;
-	int32 FirstSliceStoringStrataSSSDataWithoutMRT;
-	int32 FirstSliceStoringStrataSSSData;
+	uint32 ViewsMaxBytePerPixel = 0;
+	uint32 MinBytesPerPixel = 0;
+	uint32 MaxBytesPerPixel = 0;
+	int32 PeelLayersAboveDepth = -1;
+	bool bRoughDiffuse = false;
+	bool bRoughnessTracking = false;
+
+	int32 SliceStoringDebugStrataTreeDataWithoutMRT = -1;
+	int32 SliceStoringDebugStrataTreeData = -1;
+	int32 FirstSliceStoringStrataSSSDataWithoutMRT = -1;
+	int32 FirstSliceStoringStrataSSSData = -1;
 
 	// Resources allocated and updated each frame
 
diff --git a/Engine/Source/Runtime/Renderer/Private/Strata/StrataVisualize.cpp b/Engine/Source/Runtime/Renderer/Private/Strata/StrataVisualize.cpp
--- a/Engine/Source/Runtime/Renderer/Private/Strata/StrataVisualize.cpp
+++ b/Engine/Source/Runtime/Renderer/Private/Strata/StrataVisualize.cpp
@@ -121,7 +121,9 @@ class FStrataSystemInfoCS : public FGlobalShader
 		SHADER_PARAMETER(uint32, bRoughRefraction)
 		SHADER_PARAMETER(uint32, bTileOverflowUseMaterialData)
 		SHADER_PARAMETER(uint32, ProjectMaxBytePerPixel)
+		SHADER_PARAMETER(uint32, ViewsMaxBytePerPixel)
 		SHADER_PARAMETER(uint32, MaterialBufferAllocationInBytes)
+		SHADER_PARAMETER(uint32, MaterialBufferAllocationMode)
 		SHADER_PARAMETER(float, TileOverflowRatio)
 		SHADER_PARAMETER_RDG_BUFFER_SRV(Buffer<uint>, ClassificationTileDrawIndirectBuffer)
 		SHADER_PARAMETER_STRUCT_REF(FViewUniformShaderParameters, ViewUniformBuffer)
@@ -278,6 +280,7 @@ float GetStrataTileOverflowRatio(const FViewInfo& View);
 bool IsClassificationAsync();
 bool SupportsCMask(const FStaticShaderPlatform InPlatform);
 bool DoesStrataTileOverflowUseMaterialData();
+uint32 GetMaterialBufferAllocationMode();
 
 static void AddVisualizeSystemInfoPasses(FRDGBuilder& GraphBuilder, const FViewInfo& View, FScreenPassTexture& ScreenPassSceneColor, EShaderPlatform Platform)
 {
@@ -304,6 +307,8 @@ static void AddVisualizeSystemInfoPasses(FRDGBuilder& GraphBuilder, const FViewI
 	PassParameters->Strata = Strata::BindStrataGlobalUniformParameters(View);
 	PassParameters->SceneTextures = GetSceneTextureParameters(GraphBuilder, View);
 	PassParameters->ProjectMaxBytePerPixel = GetBytePerPixel(View.GetShaderPlatform());
+	PassParameters->ViewsMaxBytePerPixel = View.StrataViewData.SceneData->ViewsMaxBytePerPixel;
+	PassParameters->MaterialBufferAllocationMode = GetMaterialBufferAllocationMode();
 	PassParameters->MaterialBufferAllocationInBytes = MaterialBufferDesc.Extent.X * MaterialBufferDesc.Extent.Y * MaterialBufferDesc.ArraySize * sizeof(uint32);
 	ShaderPrint::SetParameters(GraphBuilder, View.ShaderPrintData, PassParameters->ShaderPrintParameters);
 
