Rename MaxBytePerPixel into MaxByte*s*PerPixel for consistency across Substrate code.

diff --git a/Engine/Shaders/Private/Strata/StrataVisualize.usf b/Engine/Shaders/Private/Strata/StrataVisualize.usf
--- a/Engine/Shaders/Private/Strata/StrataVisualize.usf
+++ b/Engine/Shaders/Private/Strata/StrataVisualize.usf
@@ -1476,8 +1476,8 @@ uint ClassificationAsync;
 uint Classification8bits;
 uint bRoughRefraction;
 uint bTileOverflowUseMaterialData;
-uint ViewsMaxBytePerPixel;
-uint ProjectMaxBytePerPixel;
+uint ViewsMaxBytesPerPixel;
+uint ProjectMaxBytesPerPixel;
 uint MaterialBufferAllocationInBytes;
 uint MaterialBufferAllocationMode;
 
@@ -1507,16 +1507,16 @@ void MainCS(uint3 DispatchThreadId : SV_DispatchThreadID)
 		Print(Context, TEXT(" Max byte per pixel : "), LegendColor); 
 		Print(Context, Strata.MaxBytesPerPixel, ValueColor, 3, 0);
 		Print(Context, TEXT(" / "), FontWhite); 
-		Print(Context, ProjectMaxBytePerPixel, ValueColor, 3, 0);
+		Print(Context, ProjectMaxBytesPerPixel, ValueColor, 3, 0);
 		Print(Context, TEXT(" ["), FontWhite); 
 		Print(Context, uint(DivideAndRoundUp(MaterialBufferAllocationInBytes, 1000000.f)), ValueColor, 4, 0);
 		Print(Context, TEXT("Mb]"), FontWhite);
 		Newline(Context);
 
 		Print(Context, TEXT(" Views max. byte    : "), LegendColor); 
-		Print(Context, ViewsMaxBytePerPixel, ValueColor, 3, 0);
+		Print(Context, ViewsMaxBytesPerPixel, ValueColor, 3, 0);
 		Print(Context, TEXT(" / "), FontWhite); 
-		Print(Context, ProjectMaxBytePerPixel, ValueColor, 3, 0);
+		Print(Context, ProjectMaxBytesPerPixel, ValueColor, 3, 0);
 		Newline(Context);
 
 		Print(Context, TEXT(" Allocation Mode    : "), LegendColor);
diff --git a/Engine/Source/Runtime/Renderer/Private/SceneVisibility.cpp b/Engine/Source/Runtime/Renderer/Private/SceneVisibility.cpp
--- a/Engine/Source/Runtime/Renderer/Private/SceneVisibility.cpp
+++ b/Engine/Source/Runtime/Renderer/Private/SceneVisibility.cpp
@@ -3071,7 +3071,7 @@ struct FRelevancePacket : public FSceneRenderingAllocatorObject<FRelevancePacket
 		WriteView.bUsesCustomDepth |= bUsesCustomDepth;
 		WriteView.bUsesCustomStencil |= bUsesCustomStencil;
 		WriteView.StrataViewData.MaxBSDFCount = FMath::Max(WriteView.StrataViewData.MaxBSDFCount, 8u - FMath::CountLeadingZeros8(StrataBSDFCountMask));
-		WriteView.StrataViewData.MaxBytePerPixel = FMath::Max(WriteView.StrataViewData.MaxBytePerPixel, StrataUintPerPixel * 4u);
+		WriteView.StrataViewData.MaxBytesPerPixel = FMath::Max(WriteView.StrataViewData.MaxBytesPerPixel, StrataUintPerPixel * 4u);
 		DirtyIndirectLightingCacheBufferPrimitives.AppendTo(WriteView.DirtyIndirectLightingCacheBufferPrimitives);
 
 		WriteView.MeshDecalBatches.Append(MeshDecalBatches);
diff --git a/Engine/Source/Runtime/Renderer/Private/Strata/Strata.cpp b/Engine/Source/Runtime/Renderer/Private/Strata/Strata.cpp
--- a/Engine/Source/Runtime/Renderer/Private/Strata/Strata.cpp
+++ b/Engine/Source/Runtime/Renderer/Private/Strata/Strata.cpp
@@ -75,10 +75,10 @@ void FStrataViewData::Reset()
 {
 	// Preserve old max BSDF count which is set prior to the reset operation
 	const uint32 OldMaxBSDFCount = MaxBSDFCount;
-	const uint32 OldMaxBytePerPixel = MaxBytePerPixel;
+	const uint32 OldMaxBytesPerPixel = MaxBytesPerPixel;
 	*this = FStrataViewData();
 	MaxBSDFCount = OldMaxBSDFCount;
-	MaxBytePerPixel = OldMaxBytePerPixel;
+	MaxBytesPerPixel = OldMaxBytesPerPixel;
 	for (uint32 i = 0; i < EStrataTileType::ECount; ++i)
 	{
 		ClassificationTileListBuffer[i] = nullptr;
@@ -368,17 +368,17 @@ void InitialiseStrataFrameSceneData(FRDGBuilder& GraphBuilder, FSceneRenderer& S
 		}
 
 		// Gather views' requirements
-		Out.ViewsMaxBytePerPixel = 0;
+		Out.ViewsMaxBytesPerPixel = 0;
 		for (const FViewInfo& View : SceneRenderer.Views)
 		{
 			bNeedBSDFOffsets = bNeedBSDFOffsets || NeedBSDFOffsets(SceneRenderer.Scene, View);
 			bNeedUAV = bNeedUAV || IsDBufferPassEnabled(View.GetShaderPlatform());
-			Out.ViewsMaxBytePerPixel = FMath::Max(Out.ViewsMaxBytePerPixel, View.StrataViewData.MaxBytePerPixel);
+			Out.ViewsMaxBytesPerPixel = FMath::Max(Out.ViewsMaxBytesPerPixel, View.StrataViewData.MaxBytesPerPixel);
 
 			// Only use primary views max. byte per pixel as reflection/capture views can bias allocation requirement when using growing-only mode
 			if (!View.bIsPlanarReflection && !View.bIsReflectionCapture && !View.bIsSceneCapture)
 			{
-				Out.MinBytesPerPixel = FMath::Max(Out.MinBytesPerPixel, View.StrataViewData.MaxBytePerPixel);
+				Out.MinBytesPerPixel = FMath::Max(Out.MinBytesPerPixel, View.StrataViewData.MaxBytesPerPixel);
 			}
 		}
 
@@ -386,16 +386,16 @@ void InitialiseStrataFrameSceneData(FRDGBuilder& GraphBuilder, FSceneRenderer& S
 		// * 0: Allocate material buffer based on view requirement 
 		// * 2: Allocate material buffer based on view requirement, but can only grow over frame to minimize buffer reallocation and hitches
 		// * 1: Allocate material buffer based on platform settings. 
-		uint32 MaxBytePerPixel = 0;
+		uint32 MaxBytesPerPixel = 0;
 		switch (GetMaterialBufferAllocationMode())
 		{
-			case 0: MaxBytePerPixel = Out.ViewsMaxBytePerPixel; break;
-			case 1: MaxBytePerPixel = FMath::Max(Out.ViewsMaxBytePerPixel, Out.MinBytesPerPixel); break;
-			case 2: MaxBytePerPixel = GetBytePerPixel(SceneRenderer.ShaderPlatform); break;
+			case 0: MaxBytesPerPixel = Out.ViewsMaxBytesPerPixel; break;
+			case 1: MaxBytesPerPixel = FMath::Max(Out.ViewsMaxBytesPerPixel, Out.MinBytesPerPixel); break;
+			case 2: MaxBytesPerPixel = GetBytePerPixel(SceneRenderer.ShaderPlatform); break;
 		}
 		const uint32 RoundToValue = 4u;
-		MaxBytePerPixel = FMath::Clamp(MaxBytePerPixel, 4u * STRATA_BASE_PASS_MRT_OUTPUT_COUNT, GetBytePerPixel(SceneRenderer.ShaderPlatform));
-		Out.MaxBytesPerPixel = FMath::DivideAndRoundUp(MaxBytePerPixel, RoundToValue) * RoundToValue;
+		MaxBytesPerPixel = FMath::Clamp(MaxBytesPerPixel, 4u * STRATA_BASE_PASS_MRT_OUTPUT_COUNT, GetBytePerPixel(SceneRenderer.ShaderPlatform));
+		Out.MaxBytesPerPixel = FMath::DivideAndRoundUp(MaxBytesPerPixel, RoundToValue) * RoundToValue;
 
 		FIntPoint SceneTextureExtent = SceneRenderer.GetActiveSceneTexturesConfig().Extent;
 		
diff --git a/Engine/Source/Runtime/Renderer/Private/Strata/Strata.h b/Engine/Source/Runtime/Renderer/Private/Strata/Strata.h
--- a/Engine/Source/Runtime/Renderer/Private/Strata/Strata.h
+++ b/Engine/Source/Runtime/Renderer/Private/Strata/Strata.h
@@ -95,7 +95,7 @@ const TCHAR* ToString(EStrataTileType Type);
 
 struct FStrataSceneData
 {
-	uint32 ViewsMaxBytePerPixel = 0;
+	uint32 ViewsMaxBytesPerPixel = 0;
 	uint32 MinBytesPerPixel = 0;
 	uint32 MaxBytesPerPixel = 0;
 	int32 PeelLayersAboveDepth = -1;
@@ -136,7 +136,7 @@ struct FStrataViewData
 {
 	// Max BSDF count among all visible materials
 	uint32 MaxBSDFCount = 0;
-	uint32 MaxBytePerPixel = 0;
+	uint32 MaxBytesPerPixel = 0;
 
 	FIntPoint TileCount  = FIntPoint(0, 0);
 	FIntPoint TileOffset = FIntPoint(0, 0);
diff --git a/Engine/Source/Runtime/Renderer/Private/Strata/StrataVisualize.cpp b/Engine/Source/Runtime/Renderer/Private/Strata/StrataVisualize.cpp
--- a/Engine/Source/Runtime/Renderer/Private/Strata/StrataVisualize.cpp
+++ b/Engine/Source/Runtime/Renderer/Private/Strata/StrataVisualize.cpp
@@ -120,8 +120,8 @@ class FStrataSystemInfoCS : public FGlobalShader
 		SHADER_PARAMETER(uint32, Classification8bits)
 		SHADER_PARAMETER(uint32, bRoughRefraction)
 		SHADER_PARAMETER(uint32, bTileOverflowUseMaterialData)
-		SHADER_PARAMETER(uint32, ProjectMaxBytePerPixel)
-		SHADER_PARAMETER(uint32, ViewsMaxBytePerPixel)
+		SHADER_PARAMETER(uint32, ProjectMaxBytesPerPixel)
+		SHADER_PARAMETER(uint32, ViewsMaxBytesPerPixel)
 		SHADER_PARAMETER(uint32, MaterialBufferAllocationInBytes)
 		SHADER_PARAMETER(uint32, MaterialBufferAllocationMode)
 		SHADER_PARAMETER(float, TileOverflowRatio)
@@ -306,8 +306,8 @@ static void AddVisualizeSystemInfoPasses(FRDGBuilder& GraphBuilder, const FViewI
 	PassParameters->ViewUniformBuffer = View.ViewUniformBuffer;
 	PassParameters->Strata = Strata::BindStrataGlobalUniformParameters(View);
 	PassParameters->SceneTextures = GetSceneTextureParameters(GraphBuilder, View);
-	PassParameters->ProjectMaxBytePerPixel = GetBytePerPixel(View.GetShaderPlatform());
-	PassParameters->ViewsMaxBytePerPixel = View.StrataViewData.SceneData->ViewsMaxBytePerPixel;
+	PassParameters->ProjectMaxBytesPerPixel = GetBytePerPixel(View.GetShaderPlatform());
+	PassParameters->ViewsMaxBytesPerPixel = View.StrataViewData.SceneData->ViewsMaxBytesPerPixel;
 	PassParameters->MaterialBufferAllocationMode = GetMaterialBufferAllocationMode();
 	PassParameters->MaterialBufferAllocationInBytes = MaterialBufferDesc.Extent.X * MaterialBufferDesc.Extent.Y * MaterialBufferDesc.ArraySize * sizeof(uint32);
 	ShaderPrint::SetParameters(GraphBuilder, View.ShaderPrintData, PassParameters->ShaderPrintParameters);
