Compiles groupshared GPU FFT on consoles

diff --git a/Engine/Shaders/Private/GPUFastFourierTransformCore.ush b/Engine/Shaders/Private/GPUFastFourierTransformCore.ush
--- a/Engine/Shaders/Private/GPUFastFourierTransformCore.ush
+++ b/Engine/Shaders/Private/GPUFastFourierTransformCore.ush
@@ -4,6 +4,10 @@
 GPUFastFourierTransformCore.usf: Core Fast Fourier Transform Code
 =============================================================================*/
 
+//------------------------------------------------------- RECOMPILE HASH
+// 2E0547A4-F927-4785-941A-93B2B8A0E1B4
+
+
 
 // This file requires the two values to be defined.
 // SCAN_LINE_LENGTH must be a power of RADIX and RADIX in turn must be a power of two 
diff --git a/Engine/Source/Runtime/Renderer/Private/GPUFastFourierTransform.cpp b/Engine/Source/Runtime/Renderer/Private/GPUFastFourierTransform.cpp
--- a/Engine/Source/Runtime/Renderer/Private/GPUFastFourierTransform.cpp
+++ b/Engine/Source/Runtime/Renderer/Private/GPUFastFourierTransform.cpp
@@ -8,7 +8,12 @@
 #include "SystemTextures.h"
 #include "DataDrivenShaderPlatformInfo.h"
 
-bool ShouldCompileGPUFFT(EShaderPlatform Platform)
+bool ShouldCompileFullGPUFFT(EShaderPlatform Platform)
+{
+	return FDataDrivenShaderPlatformInfo::GetSupportsFFTBloom(Platform) && !FDataDrivenShaderPlatformInfo::GetIsConsole(Platform);
+}
+
+bool ShouldCompileGroupSharedGPUFFT(EShaderPlatform Platform)
 {
 	return FDataDrivenShaderPlatformInfo::GetSupportsFFTBloom(Platform);
 }
@@ -152,7 +157,7 @@ class FFFTShader : public FGlobalShader
 public:
 	static bool ShouldCompilePermutation(const FGlobalShaderPermutationParameters& Parameters)
 	{
-		return ShouldCompileGPUFFT(Parameters.Platform);
+		return ShouldCompileFullGPUFFT(Parameters.Platform);
 	}
 
 	FFFTShader() = default;
@@ -322,6 +327,11 @@ class FGSComplexTransformCS : public FFFTShader
 		SHADER_PARAMETER_RDG_TEXTURE_UAV(RWTexture2D, DstTexture)
 	END_SHADER_PARAMETER_STRUCT()
 
+	static bool ShouldCompilePermutation(const FGlobalShaderPermutationParameters& Parameters)
+	{
+		return ShouldCompileGroupSharedGPUFFT(Parameters.Platform);
+	}
+
 	static void ModifyCompilationEnvironment(const FGlobalShaderPermutationParameters& Parameters, FShaderCompilerEnvironment& OutEnvironment)
 	{
 		FPermutationDomain PermutationVector(Parameters.PermutationId);
@@ -356,6 +366,11 @@ class FGSTwoForOneTransformCS : public FFFTShader
 		SHADER_PARAMETER_RDG_TEXTURE_UAV(RWTexture2D, DstTexture)
 	END_SHADER_PARAMETER_STRUCT()
 
+	static bool ShouldCompilePermutation(const FGlobalShaderPermutationParameters& Parameters)
+	{
+		return ShouldCompileGroupSharedGPUFFT(Parameters.Platform);
+	}
+
 	static void ModifyCompilationEnvironment(const FGlobalShaderPermutationParameters& Parameters, FShaderCompilerEnvironment& OutEnvironment)
 	{
 		FPermutationDomain PermutationVector(Parameters.PermutationId);
@@ -388,6 +403,11 @@ class FGSConvolutionWithTextureCS : public FFFTShader
 		SHADER_PARAMETER_RDG_TEXTURE_UAV(RWTexture2D, DstTexture)
 	END_SHADER_PARAMETER_STRUCT()
 
+	static bool ShouldCompilePermutation(const FGlobalShaderPermutationParameters& Parameters)
+	{
+		return ShouldCompileGroupSharedGPUFFT(Parameters.Platform);
+	}
+
 	static void ModifyCompilationEnvironment(const FGlobalShaderPermutationParameters& Parameters, FShaderCompilerEnvironment& OutEnvironment)
 	{
 		FPermutationDomain PermutationVector(Parameters.PermutationId);
@@ -1297,9 +1317,11 @@ void GPUFFT::ConvolutionWithTextureImage2D(
 	FRDGTextureRef SrcTexture, const FIntRect& ROIRect,
 	FRDGTextureRef DstTexture, const FIntRect& DstRect,
 	const FPreFilter& PreFilter,
-	FRDGBufferRef PostFilterParameters)
+	FRDGBufferRef PostFilterParameters,
+	ETextureCreateFlags AdditionalTextureCreateFlags)
 {
 	FRDGTextureDesc SpectralDesc = CreateFrequencyDesc(FrequencySize, bHorizontalFirst, ROIRect.Size());
+	SpectralDesc.Flags |= AdditionalTextureCreateFlags;
 
 	// This will be for the actual output.
 	const FIntRect FFTResultRect = ROIRect;
diff --git a/Engine/Source/Runtime/Renderer/Private/GPUFastFourierTransform.h b/Engine/Source/Runtime/Renderer/Private/GPUFastFourierTransform.h
--- a/Engine/Source/Runtime/Renderer/Private/GPUFastFourierTransform.h
+++ b/Engine/Source/Runtime/Renderer/Private/GPUFastFourierTransform.h
@@ -291,7 +291,8 @@ namespace GPUFFT
 		FRDGTextureRef SrcTexture, const FIntRect& ROIRect,
 		FRDGTextureRef DstTexture, const FIntRect& DstRect,
 		const FPreFilter& PreFilter,
-		FRDGBufferRef PostFilterParameters);
+		FRDGBufferRef PostFilterParameters,
+		ETextureCreateFlags AdditionalTextureCreateFlags);
 
 	FIntPoint Convolution2DBufferSize(const FIntPoint& FrequencySize, const bool bHorizontalFirst, const FIntPoint& SrcExtent);
 
diff --git a/Engine/Source/Runtime/Renderer/Private/PostProcess/PostProcessFFTBloom.cpp b/Engine/Source/Runtime/Renderer/Private/PostProcess/PostProcessFFTBloom.cpp
--- a/Engine/Source/Runtime/Renderer/Private/PostProcess/PostProcessFFTBloom.cpp
+++ b/Engine/Source/Runtime/Renderer/Private/PostProcess/PostProcessFFTBloom.cpp
@@ -795,7 +795,7 @@ FFFTBloomOutput AddFFTBloomPass(
 		FFTInputSceneColor.Texture->Desc.Extent,
 		FFTInputSceneColor.Texture->Desc.Format,
 		FClearValueBinding::None,
-		TexCreate_ShaderResource | TexCreate_UAV);
+		TexCreate_ShaderResource | TexCreate_UAV | GFastVRamConfig.Bloom);
 		
 	BloomOutput.BloomTexture.Texture = GraphBuilder.CreateTexture(OutputSceneColorDesc, TEXT("Bloom.FFT.ScatterDispersion"));
 	BloomOutput.BloomTexture.ViewRect = FFTInputSceneColor.ViewRect;
@@ -810,7 +810,8 @@ FFFTBloomOutput AddFFTBloomPass(
 		FFTInputSceneColor.Texture, FFTInputSceneColor.ViewRect,
 		BloomOutput.BloomTexture.Texture, BloomOutput.BloomTexture.ViewRect,
 		Intermediates.PreFilter,
-		FFTMulitplyParameters);
+		FFTMulitplyParameters,
+		GFastVRamConfig.Bloom);
 
 	return BloomOutput;
 }
diff --git a/Engine/Source/Runtime/Renderer/Private/SceneRendering.cpp b/Engine/Source/Runtime/Renderer/Private/SceneRendering.cpp
--- a/Engine/Source/Runtime/Renderer/Private/SceneRendering.cpp
+++ b/Engine/Source/Runtime/Renderer/Private/SceneRendering.cpp
@@ -400,6 +400,7 @@ FASTVRAM_CVAR(GBufferVelocity, 0);
 FASTVRAM_CVAR(HZB, 1);
 FASTVRAM_CVAR(SceneDepth, 1);
 FASTVRAM_CVAR(SceneColor, 1);
+FASTVRAM_CVAR(Bloom, 1);
 FASTVRAM_CVAR(BokehDOF, 1);
 FASTVRAM_CVAR(CircleDOF, 1);
 FASTVRAM_CVAR(CombineLUTs, 1);
@@ -606,6 +607,7 @@ void FFastVramConfig::Update()
 	bDirty |= UpdateTextureFlagFromCVar(CVarFastVRam_HZB, HZB);
 	bDirty |= UpdateTextureFlagFromCVar(CVarFastVRam_SceneDepth, SceneDepth);
 	bDirty |= UpdateTextureFlagFromCVar(CVarFastVRam_SceneColor, SceneColor);
+	bDirty |= UpdateTextureFlagFromCVar(CVarFastVRam_Bloom, Bloom);
 	bDirty |= UpdateTextureFlagFromCVar(CVarFastVRam_BokehDOF, BokehDOF);
 	bDirty |= UpdateTextureFlagFromCVar(CVarFastVRam_CircleDOF, CircleDOF);
 	bDirty |= UpdateTextureFlagFromCVar(CVarFastVRam_CombineLUTs, CombineLUTs);
diff --git a/Engine/Source/Runtime/Renderer/Private/SceneRendering.h b/Engine/Source/Runtime/Renderer/Private/SceneRendering.h
--- a/Engine/Source/Runtime/Renderer/Private/SceneRendering.h
+++ b/Engine/Source/Runtime/Renderer/Private/SceneRendering.h
@@ -2838,6 +2838,7 @@ struct FFastVramConfig
 	ETextureCreateFlags HZB;
 	ETextureCreateFlags SceneDepth;
 	ETextureCreateFlags SceneColor;
+	ETextureCreateFlags Bloom;
 	ETextureCreateFlags BokehDOF;
 	ETextureCreateFlags CircleDOF;
 	ETextureCreateFlags CombineLUTs;
