Substrate - fixed mismatch between cvar comment and code when computing per pixel byte count. Also added an assert.

diff --git a/Engine/Source/Runtime/Renderer/Private/Strata/Strata.cpp b/Engine/Source/Runtime/Renderer/Private/Strata/Strata.cpp
--- a/Engine/Source/Runtime/Renderer/Private/Strata/Strata.cpp
+++ b/Engine/Source/Runtime/Renderer/Private/Strata/Strata.cpp
@@ -65,8 +65,8 @@ static TAutoConsoleVariable<int32> CVarSubstrateDBufferPassDedicatedTiles(
 
 static TAutoConsoleVariable<int32> CVarSubstrateBytePerPixelMode(
 	TEXT("r.Substrate.BytesPerPixel.Mode"),
-	0,
-	TEXT("Substrate material allocation mode. \n 0: Allocate material buffer based on view requirement \n 1: Allocate material buffer based on platform settings. \n 2: Allocate material buffer based on view requirement, but can only grow over frame to minimize buffer reallocation. "),
+	1,
+	TEXT("Substrate material allocation mode. \n 0: Allocate material buffer based on view requirement, \n 1: Allocate material buffer based on view requirement, but can only grow over frame to minimize buffer reallocation and hitches, \n 2: Allocate material buffer based on platform settings."),
 	ECVF_RenderThreadSafe);
 
 IMPLEMENT_GLOBAL_SHADER_PARAMETER_STRUCT(FStrataGlobalUniformParameters, "Strata");
@@ -383,18 +383,23 @@ void InitialiseStrataFrameSceneData(FRDGBuilder& GraphBuilder, FSceneRenderer& S
 		}
 
 		// Material buffer allocation can use different modes:
-		// * 0: Allocate material buffer based on view requirement 
-		// * 2: Allocate material buffer based on view requirement, but can only grow over frame to minimize buffer reallocation and hitches
-		// * 1: Allocate material buffer based on platform settings. 
+		// * 0: Allocate material buffer based on view requirement,
+		// * 1: Allocate material buffer based on view requirement, but can only grow over frame to minimize buffer reallocation and hitches,
+		// * 2: Allocate material buffer based on platform settings.
+		const uint32 PlatformSettingsBytesPerPixel = GetBytePerPixel(SceneRenderer.ShaderPlatform);
 		uint32 MaxBytesPerPixel = 0;
 		switch (GetMaterialBufferAllocationMode())
 		{
 			case 0: MaxBytesPerPixel = Out.ViewsMaxBytesPerPixel; break;
 			case 1: MaxBytesPerPixel = FMath::Max(Out.ViewsMaxBytesPerPixel, Out.MinBytesPerPixel); break;
-			case 2: MaxBytesPerPixel = GetBytePerPixel(SceneRenderer.ShaderPlatform); break;
+			case 2: MaxBytesPerPixel = PlatformSettingsBytesPerPixel; break;
 		}
+
+		// If this happens, it means there is probably a shader compilation mismatch issue (the compiler has not correctly accounted for the byte per pixel limitation for the platform).
+		check(MaxBytesPerPixel <= PlatformSettingsBytesPerPixel);
+
 		const uint32 RoundToValue = 4u;
-		MaxBytesPerPixel = FMath::Clamp(MaxBytesPerPixel, 4u * STRATA_BASE_PASS_MRT_OUTPUT_COUNT, GetBytePerPixel(SceneRenderer.ShaderPlatform));
+		MaxBytesPerPixel = FMath::Clamp(MaxBytesPerPixel, 4u * STRATA_BASE_PASS_MRT_OUTPUT_COUNT, PlatformSettingsBytesPerPixel);
 		Out.MaxBytesPerPixel = FMath::DivideAndRoundUp(MaxBytesPerPixel, RoundToValue) * RoundToValue;
 
 		FIntPoint SceneTextureExtent = SceneRenderer.GetActiveSceneTexturesConfig().Extent;
