VulkanRHI: Fix a handful of leaked staging buffers on shutdown.

#rnx

diff --git a/Engine/Source/Runtime/Renderer/Private/GlobalDistanceField.cpp b/Engine/Source/Runtime/Renderer/Private/GlobalDistanceField.cpp
--- a/Engine/Source/Runtime/Renderer/Private/GlobalDistanceField.cpp
+++ b/Engine/Source/Runtime/Renderer/Private/GlobalDistanceField.cpp
@@ -778,7 +778,7 @@ void RecaptureClipmapForMeshSDFStreamingIfNeeded(
 		if (ClipmapViewState.HasPendingStreamingReadbackBuffers[Index]->IsReady())
 		{
 			ClipmapViewState.ReadbackBuffersNumPending--;
-			LatestReadbackBuffer = ClipmapViewState.HasPendingStreamingReadbackBuffers[Index];
+			LatestReadbackBuffer = ClipmapViewState.HasPendingStreamingReadbackBuffers[Index].Get();
 		}
 		else
 		{
@@ -2377,13 +2377,13 @@ void UpdateGlobalDistanceFieldVolume(
 							{
 								FGlobalDistanceFieldClipmapState& ClipmapViewState = View.ViewState->GlobalDistanceFieldData->ClipmapState[ClipmapIndex];
 
-								if (!ClipmapViewState.HasPendingStreamingReadbackBuffers[ClipmapViewState.ReadbackBuffersWriteIndex])
+								if (!ClipmapViewState.HasPendingStreamingReadbackBuffers[ClipmapViewState.ReadbackBuffersWriteIndex].IsValid())
 								{
-									FRHIGPUBufferReadback* GPUBufferReadback = new FRHIGPUBufferReadback(TEXT("GlobalDistanceField.HasPendingStreamingReadback"));
-									ClipmapViewState.HasPendingStreamingReadbackBuffers[ClipmapViewState.ReadbackBuffersWriteIndex] = GPUBufferReadback;
+									ClipmapViewState.HasPendingStreamingReadbackBuffers[ClipmapViewState.ReadbackBuffersWriteIndex] = 
+										MakeUnique<FRHIGPUBufferReadback>(TEXT("GlobalDistanceField.HasPendingStreamingReadback"));
 								}
 
-								FRHIGPUBufferReadback* ReadbackBuffer = ClipmapViewState.HasPendingStreamingReadbackBuffers[ClipmapViewState.ReadbackBuffersWriteIndex];
+								FRHIGPUBufferReadback* ReadbackBuffer = ClipmapViewState.HasPendingStreamingReadbackBuffers[ClipmapViewState.ReadbackBuffersWriteIndex].Get();
 
 								AddReadbackBufferPass(GraphBuilder, RDG_EVENT_NAME("GlobalDistanceField.HasPendingStreamingReadback"), HasPendingStreamingReadbackBuffer,
 									[ReadbackBuffer, HasPendingStreamingReadbackBuffer](FRHICommandList& RHICmdList)
diff --git a/Engine/Source/Runtime/Renderer/Private/ScenePrivate.h b/Engine/Source/Runtime/Renderer/Private/ScenePrivate.h
--- a/Engine/Source/Runtime/Renderer/Private/ScenePrivate.h
+++ b/Engine/Source/Runtime/Renderer/Private/ScenePrivate.h
@@ -460,7 +460,7 @@ public:
 		CacheMostlyStaticSeparately = 1;
 		LastUsedSceneDataForFullUpdate = nullptr;
 
-		HasPendingStreamingReadbackBuffers.AddZeroed(MaxPendingStreamingReadbackBuffers);
+		HasPendingStreamingReadbackBuffers.SetNum(MaxPendingStreamingReadbackBuffers);
 	}
 
 	FInt64Vector FullUpdateOriginInPages;
@@ -471,7 +471,7 @@ public:
 	float CachedClipmapExtent;
 	float CacheClipmapInfluenceRadius;
 
-	TArray<FRHIGPUBufferReadback*> HasPendingStreamingReadbackBuffers;
+	TArray<TUniquePtr<FRHIGPUBufferReadback>> HasPendingStreamingReadbackBuffers;
 	uint32 MaxPendingStreamingReadbackBuffers = 4;
 	uint32 ReadbackBuffersWriteIndex = 0;
 	uint32 ReadbackBuffersNumPending = 0;
