Substrate - Fixed translucent material luminance not being visible when strata is enabled.

diff --git a/Engine/Source/Runtime/Renderer/Private/PostProcess/TemporalSuperResolution.cpp b/Engine/Source/Runtime/Renderer/Private/PostProcess/TemporalSuperResolution.cpp
--- a/Engine/Source/Runtime/Renderer/Private/PostProcess/TemporalSuperResolution.cpp
+++ b/Engine/Source/Runtime/Renderer/Private/PostProcess/TemporalSuperResolution.cpp
@@ -1165,7 +1165,6 @@ ITemporalUpscaler::FOutputs AddTemporalSuperResolutionPasses(
 		   (!View.Family->EngineShowFlags.VisualizeBuffer    || (View.Family->EngineShowFlags.VisualizeBuffer    && View.CurrentBufferVisualizationMode == OverviewName))
 		&& (!View.Family->EngineShowFlags.VisualizeNanite    || (View.Family->EngineShowFlags.VisualizeNanite    && View.CurrentNaniteVisualizationMode == OverviewName))
 		&& (!View.Family->EngineShowFlags.VisualizeLumen     || (View.Family->EngineShowFlags.VisualizeLumen     && View.CurrentLumenVisualizationMode  == OverviewName))
-		&& (!View.Family->EngineShowFlags.VisualizeSubstrate || (View.Family->EngineShowFlags.VisualizeSubstrate && View.CurrentStrataVisualizationMode == OverviewName))
 		&& (!View.Family->EngineShowFlags.VisualizeGroom     || (View.Family->EngineShowFlags.VisualizeGroom     && View.CurrentGroomVisualizationMode  == OverviewName));
 #endif
 	if (bHasSeparateTranslucency)
