Share common function in Strata shader parameters.

diff --git a/Engine/Source/Runtime/Renderer/Private/Strata/Strata.cpp b/Engine/Source/Runtime/Renderer/Private/Strata/Strata.cpp
--- a/Engine/Source/Runtime/Renderer/Private/Strata/Strata.cpp
+++ b/Engine/Source/Runtime/Renderer/Private/Strata/Strata.cpp
@@ -292,14 +292,13 @@ static void InitialiseStrataViewData(FRDGBuilder& GraphBuilder, FViewInfo& View,
 			Out.BSDFTileDispatchIndirectBuffer = nullptr;
 			Out.BSDFTileCountBuffer = nullptr;
 		}
-	}
 
-	// Create the readable uniform buffers
-	if (IsStrataEnabled())
-	{
-		FStrataGlobalUniformParameters* StrataUniformParameters = GraphBuilder.AllocParameters<FStrataGlobalUniformParameters>();
-		BindStrataGlobalUniformParameters(GraphBuilder, &Out, *StrataUniformParameters);
-		Out.StrataGlobalUniformParameters = GraphBuilder.CreateUniformBuffer(StrataUniformParameters);
+		// Create the readable uniform buffers
+		{
+			FStrataGlobalUniformParameters* StrataUniformParameters = GraphBuilder.AllocParameters<FStrataGlobalUniformParameters>();
+			BindStrataGlobalUniformParameters(GraphBuilder, &Out, *StrataUniformParameters);
+			Out.StrataGlobalUniformParameters = GraphBuilder.CreateUniformBuffer(StrataUniformParameters);
+		}
 	}
 }
 
@@ -489,15 +488,31 @@ void InitialiseStrataFrameSceneData(FRDGBuilder& GraphBuilder, FSceneRenderer& S
 	}
 }
 
+static FStrataCommonParameters GetStrataCommonParameter()
+{
+	FStrataCommonParameters Out;
+	Out.bRoughDiffuse 		= 0u;
+	Out.MaxBytesPerPixel 	= 0u;
+	Out.PeelLayersAboveDepth= 0u;
+	Out.bRoughnessTracking 	= 0u;
+	return Out;
+}
+static FStrataCommonParameters GetStrataCommonParameter(const FStrataSceneData& In)
+{
+	FStrataCommonParameters Out;
+	Out.bRoughDiffuse 		= In.bRoughDiffuse ? 1u : 0u;
+	Out.MaxBytesPerPixel 	= In.MaxBytesPerPixel;
+	Out.PeelLayersAboveDepth= In.PeelLayersAboveDepth;
+	Out.bRoughnessTracking 	= In.bRoughnessTracking ? 1u : 0u;
+	return Out;
+}
+
 void BindStrataBasePassUniformParameters(FRDGBuilder& GraphBuilder, const FViewInfo& View, FStrataBasePassUniformParameters& OutStrataUniformParameters)
 {
 	const FStrataSceneData* StrataSceneData = View.StrataViewData.SceneData;
 	if (IsStrataEnabled() && StrataSceneData)
 	{
-		OutStrataUniformParameters.bRoughDiffuse = StrataSceneData->bRoughDiffuse ? 1u : 0u;
-		OutStrataUniformParameters.MaxBytesPerPixel = StrataSceneData->MaxBytesPerPixel;
-		OutStrataUniformParameters.PeelLayersAboveDepth = StrataSceneData->PeelLayersAboveDepth;
-		OutStrataUniformParameters.bRoughnessTracking = StrataSceneData->bRoughnessTracking ? 1u : 0u;
+		OutStrataUniformParameters.Common = GetStrataCommonParameter(*StrataSceneData);
 		OutStrataUniformParameters.SliceStoringDebugStrataTreeDataWithoutMRT = StrataSceneData->SliceStoringDebugStrataTreeDataWithoutMRT;
 		OutStrataUniformParameters.FirstSliceStoringStrataSSSDataWithoutMRT = StrataSceneData->FirstSliceStoringStrataSSSDataWithoutMRT;
 		OutStrataUniformParameters.MaterialTextureArrayUAVWithoutRTs = StrataSceneData->MaterialTextureArrayUAVWithoutRTs;
@@ -512,10 +527,7 @@ void BindStrataBasePassUniformParameters(FRDGBuilder& GraphBuilder, const FViewI
 		FRDGTextureUAVRef DummyWritableTextureArrayUAV = GraphBuilder.CreateUAV(FRDGTextureUAVDesc(DummyWritableTextureArray));
 
 		const FRDGSystemTextures& SystemTextures = FRDGSystemTextures::Get(GraphBuilder);
-		OutStrataUniformParameters.bRoughDiffuse = 0u;
-		OutStrataUniformParameters.MaxBytesPerPixel = 0;
-		OutStrataUniformParameters.PeelLayersAboveDepth = 0;
-		OutStrataUniformParameters.bRoughnessTracking = 0;
+		OutStrataUniformParameters.Common = GetStrataCommonParameter();
 		OutStrataUniformParameters.SliceStoringDebugStrataTreeDataWithoutMRT = -1;
 		OutStrataUniformParameters.FirstSliceStoringStrataSSSDataWithoutMRT = -1;
 		OutStrataUniformParameters.MaterialTextureArrayUAVWithoutRTs = DummyWritableTextureArrayUAV;
@@ -528,10 +540,7 @@ static void BindStrataGlobalUniformParameters(FRDGBuilder& GraphBuilder, FStrata
 	FStrataSceneData* StrataSceneData = StrataViewData->SceneData;
 	if (IsStrataEnabled() && StrataSceneData)
 	{
-		OutStrataUniformParameters.bRoughDiffuse = StrataSceneData->bRoughDiffuse ? 1u : 0u;
-		OutStrataUniformParameters.MaxBytesPerPixel = StrataSceneData->MaxBytesPerPixel;
-		OutStrataUniformParameters.PeelLayersAboveDepth = StrataSceneData->PeelLayersAboveDepth;
-		OutStrataUniformParameters.bRoughnessTracking = StrataSceneData->bRoughnessTracking ? 1u : 0u;
+		OutStrataUniformParameters.Common = GetStrataCommonParameter(*StrataSceneData);
 		OutStrataUniformParameters.SliceStoringDebugStrataTreeData = StrataSceneData->SliceStoringDebugStrataTreeData;
 		OutStrataUniformParameters.FirstSliceStoringStrataSSSData = StrataSceneData->FirstSliceStoringStrataSSSData;
 		OutStrataUniformParameters.TileSize = STRATA_TILE_SIZE;
@@ -561,10 +570,7 @@ static void BindStrataGlobalUniformParameters(FRDGBuilder& GraphBuilder, FStrata
 		const FRDGSystemTextures& SystemTextures = FRDGSystemTextures::Get(GraphBuilder);
 		FRDGTextureRef DefaultTextureArray = GSystemTextures.GetDefaultTexture(GraphBuilder, ETextureDimension::Texture2DArray, EPixelFormat::PF_R32_UINT, FClearValueBinding::Transparent);
 		FRDGBufferSRVRef DefaultBuffer = GraphBuilder.CreateSRV(GSystemTextures.GetDefaultBuffer(GraphBuilder, 4, 0u), PF_R32_UINT);
-		OutStrataUniformParameters.bRoughDiffuse = 0;
-		OutStrataUniformParameters.MaxBytesPerPixel = 0;
-		OutStrataUniformParameters.PeelLayersAboveDepth = 0;
-		OutStrataUniformParameters.bRoughnessTracking = 0;
+		OutStrataUniformParameters.Common = GetStrataCommonParameter();
 		OutStrataUniformParameters.SliceStoringDebugStrataTreeData = -1;
 		OutStrataUniformParameters.FirstSliceStoringStrataSSSData = -1;
 		OutStrataUniformParameters.TileSize = 0;
@@ -587,10 +593,7 @@ void BindStrataForwardPasslUniformParameters(FRDGBuilder& GraphBuilder, const FV
 	FStrataSceneData* StrataSceneData = View.StrataViewData.SceneData;
 	if (IsStrataEnabled() && StrataSceneData)
 	{
-		OutStrataUniformParameters.MaxBytesPerPixel = StrataSceneData->MaxBytesPerPixel;
-		OutStrataUniformParameters.bRoughDiffuse = StrataSceneData->bRoughDiffuse ? 1u : 0u;
-		OutStrataUniformParameters.PeelLayersAboveDepth = StrataSceneData->PeelLayersAboveDepth;
-		OutStrataUniformParameters.bRoughnessTracking = StrataSceneData->bRoughnessTracking ? 1u : 0u;
+		OutStrataUniformParameters.Common = GetStrataCommonParameter(*StrataSceneData);
 		OutStrataUniformParameters.FirstSliceStoringStrataSSSData = StrataSceneData->FirstSliceStoringStrataSSSData;
 		OutStrataUniformParameters.MaterialTextureArray = StrataSceneData->MaterialTextureArray;
 		OutStrataUniformParameters.TopLayerTexture = StrataSceneData->TopLayerTexture;
@@ -599,10 +602,7 @@ void BindStrataForwardPasslUniformParameters(FRDGBuilder& GraphBuilder, const FV
 	{
 		const FRDGSystemTextures& SystemTextures = FRDGSystemTextures::Get(GraphBuilder);
 		FRDGTextureRef DefaultTextureArray = GSystemTextures.GetDefaultTexture(GraphBuilder, ETextureDimension::Texture2DArray, EPixelFormat::PF_R32_UINT, FClearValueBinding::Transparent);
-		OutStrataUniformParameters.MaxBytesPerPixel = 0;
-		OutStrataUniformParameters.bRoughDiffuse = 0;
-		OutStrataUniformParameters.PeelLayersAboveDepth = 0;
-		OutStrataUniformParameters.bRoughnessTracking = 0;
+		OutStrataUniformParameters.Common = GetStrataCommonParameter();
 		OutStrataUniformParameters.FirstSliceStoringStrataSSSData = -1;
 		OutStrataUniformParameters.MaterialTextureArray = DefaultTextureArray;
 		OutStrataUniformParameters.TopLayerTexture = SystemTextures.DefaultNormal8Bit;
@@ -614,17 +614,11 @@ void BindStrataMobileForwardPasslUniformParameters(FRDGBuilder& GraphBuilder, co
 	FStrataSceneData* StrataSceneData = View.StrataViewData.SceneData;
 	if (IsStrataEnabled() && StrataSceneData)
 	{
-		OutStrataUniformParameters.MaxBytesPerPixel = StrataSceneData->MaxBytesPerPixel;
-		OutStrataUniformParameters.bRoughDiffuse = StrataSceneData->bRoughDiffuse ? 1u : 0u;
-		OutStrataUniformParameters.PeelLayersAboveDepth = StrataSceneData->PeelLayersAboveDepth;
-		OutStrataUniformParameters.bRoughnessTracking = StrataSceneData->bRoughnessTracking ? 1u : 0u;
+		OutStrataUniformParameters.Common = GetStrataCommonParameter(*StrataSceneData);
 	}
 	else
 	{
-		OutStrataUniformParameters.MaxBytesPerPixel = 0;
-		OutStrataUniformParameters.bRoughDiffuse = 0;
-		OutStrataUniformParameters.PeelLayersAboveDepth = 0;
-		OutStrataUniformParameters.bRoughnessTracking = 0;
+		OutStrataUniformParameters.Common = GetStrataCommonParameter();
 	}
 }
 
diff --git a/Engine/Source/Runtime/Renderer/Private/Strata/Strata.h b/Engine/Source/Runtime/Renderer/Private/Strata/Strata.h
--- a/Engine/Source/Runtime/Renderer/Private/Strata/Strata.h
+++ b/Engine/Source/Runtime/Renderer/Private/Strata/Strata.h
@@ -21,11 +21,15 @@ struct FMinimalSceneTextures;
 struct FScreenPassTexture;
 struct FTextureRenderTargetBinding;
 
-BEGIN_SHADER_PARAMETER_STRUCT(FStrataBasePassUniformParameters, )
+BEGIN_SHADER_PARAMETER_STRUCT(FStrataCommonParameters, )
 	SHADER_PARAMETER(uint32, MaxBytesPerPixel)
 	SHADER_PARAMETER(uint32, bRoughDiffuse)
 	SHADER_PARAMETER(uint32, PeelLayersAboveDepth)
 	SHADER_PARAMETER(uint32, bRoughnessTracking)
+END_SHADER_PARAMETER_STRUCT()
+
+BEGIN_SHADER_PARAMETER_STRUCT(FStrataBasePassUniformParameters, )
+	SHADER_PARAMETER_STRUCT_INCLUDE(FStrataCommonParameters, Common)
 	SHADER_PARAMETER(int32, SliceStoringDebugStrataTreeDataWithoutMRT)
 	SHADER_PARAMETER(int32, FirstSliceStoringStrataSSSDataWithoutMRT)
 	SHADER_PARAMETER_RDG_TEXTURE_UAV(RWTexture2DArray<uint>, MaterialTextureArrayUAVWithoutRTs)
@@ -33,20 +37,14 @@ BEGIN_SHADER_PARAMETER_STRUCT(FStrataBasePassUniformParameters, )
 END_SHADER_PARAMETER_STRUCT()
 
 BEGIN_SHADER_PARAMETER_STRUCT(FStrataForwardPassUniformParameters, )
-	SHADER_PARAMETER(uint32, MaxBytesPerPixel)
-	SHADER_PARAMETER(uint32, bRoughDiffuse)
-	SHADER_PARAMETER(uint32, PeelLayersAboveDepth)
-	SHADER_PARAMETER(uint32, bRoughnessTracking)
+	SHADER_PARAMETER_STRUCT_INCLUDE(FStrataCommonParameters, Common)
 	SHADER_PARAMETER(int32, FirstSliceStoringStrataSSSData)
 	SHADER_PARAMETER_RDG_TEXTURE(Texture2DArray<uint>, MaterialTextureArray)
 	SHADER_PARAMETER_RDG_TEXTURE(Texture2D<uint2>, TopLayerTexture)
 END_SHADER_PARAMETER_STRUCT()
 
 BEGIN_SHADER_PARAMETER_STRUCT(FStrataMobileForwardPassUniformParameters, )
-	SHADER_PARAMETER(uint32, MaxBytesPerPixel)
-	SHADER_PARAMETER(uint32, bRoughDiffuse)
-	SHADER_PARAMETER(uint32, PeelLayersAboveDepth)
-	SHADER_PARAMETER(uint32, bRoughnessTracking)
+	SHADER_PARAMETER_STRUCT_INCLUDE(FStrataCommonParameters, Common)
 END_SHADER_PARAMETER_STRUCT()
 
 BEGIN_SHADER_PARAMETER_STRUCT(FStrataTileParameter, )
@@ -57,10 +55,7 @@ END_SHADER_PARAMETER_STRUCT()
 // This paramater struct is declared with RENDERER_API even though it is not public. This is
 // to workaround other modules doing 'private include' of the Renderer module
 BEGIN_GLOBAL_SHADER_PARAMETER_STRUCT(FStrataGlobalUniformParameters, RENDERER_API)
-	SHADER_PARAMETER(uint32, MaxBytesPerPixel)
-	SHADER_PARAMETER(uint32, bRoughDiffuse)
-	SHADER_PARAMETER(uint32, PeelLayersAboveDepth)
-	SHADER_PARAMETER(uint32, bRoughnessTracking)
+	SHADER_PARAMETER_STRUCT_INCLUDE(FStrataCommonParameters, Common)
 	SHADER_PARAMETER(int32,  SliceStoringDebugStrataTreeData)
 	SHADER_PARAMETER(int32,  FirstSliceStoringStrataSSSData)
 	SHADER_PARAMETER(uint32, TileSize)
