Adds a downsample passes to bloom's kernel preprocessing to properly anti-alias it

This also adds a warning on console to reduce the resolution of the kernel's asset resolution to save memory

diff --git a/Engine/Shaders/Private/Bloom/BloomDownsampleKernel.usf b/Engine/Shaders/Private/Bloom/BloomDownsampleKernel.usf
new file mode 100644
--- /dev/null
+++ b/Engine/Shaders/Private/Bloom/BloomDownsampleKernel.usf
@@ -0,0 +1,39 @@
+// Copyright Epic Games, Inc. All Rights Reserved.
+
+#include "BloomCommon.ush"
+
+
+//------------------------------------------------------- CONFIG
+
+#define TILE_SIZE 8
+
+
+//------------------------------------------------------- PARAMETERS
+
+uint2 KernelSpatialTextureSize;
+Texture2D KernelSpatialTexture;
+
+RWTexture2D<float4> DownsampleKernelSpatialOutput;
+
+
+//------------------------------------------------------- ENTRY POINT
+
+[numthreads(TILE_SIZE, TILE_SIZE, 1)]
+void MainCS(
+	uint2 GroupId : SV_GroupID,
+	uint GroupThreadIndex : SV_GroupIndex)
+{
+	uint2 OutputPixelCoord = TILE_SIZE * GroupId + uint2(GroupThreadIndex % TILE_SIZE, GroupThreadIndex / TILE_SIZE);
+
+	float2 TextureUV = (float2(OutputPixelCoord) * 2.0 + 1.0) / float2(KernelSpatialTextureSize);
+
+	float2 Offset = 0.7 / float2(KernelSpatialTextureSize);
+
+	float4 KernelColor = 0.0;
+	KernelColor += 0.25 * KernelSpatialTexture.SampleLevel(GlobalBilinearClampedSampler, TextureUV + float2(+1, +1) * Offset, 0);
+	KernelColor += 0.25 * KernelSpatialTexture.SampleLevel(GlobalBilinearClampedSampler, TextureUV + float2(+1, -1) * Offset, 0);
+	KernelColor += 0.25 * KernelSpatialTexture.SampleLevel(GlobalBilinearClampedSampler, TextureUV + float2(-1, -1) * Offset, 0);
+	KernelColor += 0.25 * KernelSpatialTexture.SampleLevel(GlobalBilinearClampedSampler, TextureUV + float2(-1, +1) * Offset, 0);
+
+	DownsampleKernelSpatialOutput[OutputPixelCoord] = KernelColor;
+}
diff --git a/Engine/Source/Runtime/Renderer/Private/PostProcess/PostProcessFFTBloom.cpp b/Engine/Source/Runtime/Renderer/Private/PostProcess/PostProcessFFTBloom.cpp
--- a/Engine/Source/Runtime/Renderer/Private/PostProcess/PostProcessFFTBloom.cpp
+++ b/Engine/Source/Runtime/Renderer/Private/PostProcess/PostProcessFFTBloom.cpp
@@ -28,6 +28,14 @@ TAutoConsoleVariable<int32> CVarAsynComputeFFTBloom(
 	TEXT("Whether to run FFT bloom on async compute.\n"),
 	ECVF_RenderThreadSafe);
 
+TAutoConsoleVariable<int32> CVarBloomWarnKernelResolution(
+	TEXT("r.Bloom.WarnKernelResolution"), 1,
+	TEXT("Whether to emit a warning when the resolution of the kernel is unecessary to high.\n")
+	TEXT(" 0: Disabled;\n")
+	TEXT(" 1: Emit the warning on consoles (default);\n")
+	TEXT(" 2: Emit the warning on all platforms;\n"),
+	ECVF_RenderThreadSafe);
+
 static bool DoesPlatformSupportFFTBloom(EShaderPlatform Platform)
 {
 	return FDataDrivenShaderPlatformInfo::GetSupportsFFTBloom(Platform);
@@ -151,6 +159,19 @@ public:
 	END_SHADER_PARAMETER_STRUCT()
 };
 
+class FBloomDownsampleKernelCS : public FFFTBloomShader
+{
+public:
+	DECLARE_GLOBAL_SHADER(FBloomDownsampleKernelCS);
+	SHADER_USE_PARAMETER_STRUCT(FBloomDownsampleKernelCS, FFFTBloomShader);
+
+	BEGIN_SHADER_PARAMETER_STRUCT(FParameters, )
+		SHADER_PARAMETER(FIntPoint, KernelSpatialTextureSize)
+		SHADER_PARAMETER_RDG_TEXTURE(Texture2D, KernelSpatialTexture)
+		SHADER_PARAMETER_RDG_TEXTURE_UAV(RWTexture2D<float4>, DownsampleKernelSpatialOutput)
+	END_SHADER_PARAMETER_STRUCT()
+};
+
 class FBloomResizeKernelCS : public FFFTBloomShader
 {
 public:
@@ -194,6 +215,7 @@ IMPLEMENT_GLOBAL_SHADER(FBloomReduceKernelSurveyCS,         "/Engine/Private/Blo
 IMPLEMENT_GLOBAL_SHADER(FBloomSumScatterDispersionEnergyCS, "/Engine/Private/Bloom/BloomSumScatterDispersionEnergy.usf", "MainCS", SF_Compute);
 IMPLEMENT_GLOBAL_SHADER(FBloomPackKernelConstantsCS,        "/Engine/Private/Bloom/BloomPackKernelConstants.usf", "MainCS", SF_Compute);
 IMPLEMENT_GLOBAL_SHADER(FBloomClampKernelCS,                "/Engine/Private/Bloom/BloomClampKernel.usf", "MainCS", SF_Compute);
+IMPLEMENT_GLOBAL_SHADER(FBloomDownsampleKernelCS,           "/Engine/Private/Bloom/BloomDownsampleKernel.usf", "MainCS", SF_Compute);
 IMPLEMENT_GLOBAL_SHADER(FBloomResizeKernelCS,               "/Engine/Private/Bloom/BloomResizeKernel.usf", "MainCS", SF_Compute);
 IMPLEMENT_GLOBAL_SHADER(FBloomFinalizeApplyConstantsCS,     "/Engine/Private/Bloom/BloomFinalizeApplyConstants.usf", "MainCS", SF_Compute);
 
@@ -575,6 +597,54 @@ void InitDomainAndGetKernel(
 				FComputeShaderUtils::GetGroupCount(SpatialKernelTexture->Desc.Extent, 8));
 		}
 
+		// Downscale the clamped kernel succesively to have nice anti-aliased
+		{
+			int32 MajorAxisPixelCount = FMath::Max(Intermediates.ImageSize.X, Intermediates.ImageSize.Y);
+			float KernelSizeInDstPixels = FMath::Max(float(MajorAxisPixelCount) * Intermediates.KernelSupportScale, 1.f);
+			float DownscaleFactor = float(SpatialKernelTexture->Desc.Extent.X) / KernelSizeInDstPixels;
+			int32 RecommendedKernelDownsscale = 1;
+
+			for (float RemainingDownscaleFactor = DownscaleFactor; RemainingDownscaleFactor > 2.0f; RemainingDownscaleFactor *= 0.5f)
+			{
+				FRDGTextureRef NewClampedKernelTexture;
+				{
+					FRDGTextureDesc Desc = FRDGTextureDesc::Create2D(
+						ClampedKernelTexture->Desc.Extent / 2,
+						PF_FloatRGBA,
+						FClearValueBinding::None,
+						TexCreate_ShaderResource | TexCreate_UAV);
+
+					NewClampedKernelTexture = GraphBuilder.CreateTexture(Desc, TEXT("Bloom.FFT.DownsampleKernel"));
+				}
+
+				FBloomDownsampleKernelCS::FParameters* PassParameters = GraphBuilder.AllocParameters<FBloomDownsampleKernelCS::FParameters>();
+				PassParameters->KernelSpatialTextureSize = ClampedKernelTexture->Desc.Extent;
+				PassParameters->KernelSpatialTexture = ClampedKernelTexture;
+				PassParameters->DownsampleKernelSpatialOutput = GraphBuilder.CreateUAV(NewClampedKernelTexture);
+
+				TShaderMapRef<FBloomDownsampleKernelCS> ComputeShader(View.ShaderMap);
+				FComputeShaderUtils::AddPass(
+					GraphBuilder,
+					RDG_EVENT_NAME("FFTBloom DownsampleKernel %dx%d -> %dx%d",
+						ClampedKernelTexture->Desc.Extent.X,
+						ClampedKernelTexture->Desc.Extent.Y,
+						NewClampedKernelTexture->Desc.Extent.X,
+						NewClampedKernelTexture->Desc.Extent.Y),
+					Intermediates.ComputePassFlags,
+					ComputeShader,
+					PassParameters,
+					FComputeShaderUtils::GetGroupCount(NewClampedKernelTexture->Desc.Extent, 8));
+
+				ClampedKernelTexture = NewClampedKernelTexture;
+				RecommendedKernelDownsscale *= 2;
+			}
+
+			if (RecommendedKernelDownsscale > 1 && (CVarBloomWarnKernelResolution.GetValueOnRenderThread() > 1 || (CVarBloomWarnKernelResolution.GetValueOnRenderThread() == 1 && FDataDrivenShaderPlatformInfo::GetIsConsole(View.GetShaderPlatform()))))
+			{
+				UE_LOG(LogRenderer, Warning, TEXT("The FPostProcessSettings::BloomConvolutionTexture could have it's resolution lowered by a factor of %d to save memory."), RecommendedKernelDownsscale);
+			}
+		}
+
 		// Final resize the kernel for Fourier transformation
 		{
 			ResizedKernel = GraphBuilder.CreateTexture(TransformDesc, TEXT("Bloom.FFT.ResizedKernel"));
