Remove DXC precompile step in FRenderLightingCacheWithLiveShadingCS shader to workaround cross-compile bug in SPIRV-Cross.

#rnx

diff --git a/Engine/Source/Runtime/Renderer/Private/HeterogeneousVolumes/HeterogeneousVolumesLiveShadingPipeline.cpp b/Engine/Source/Runtime/Renderer/Private/HeterogeneousVolumes/HeterogeneousVolumesLiveShadingPipeline.cpp
--- a/Engine/Source/Runtime/Renderer/Private/HeterogeneousVolumes/HeterogeneousVolumesLiveShadingPipeline.cpp
+++ b/Engine/Source/Runtime/Renderer/Private/HeterogeneousVolumes/HeterogeneousVolumesLiveShadingPipeline.cpp
@@ -185,7 +185,7 @@ class FRenderLightingCacheWithLiveShadingCS : public FMeshMaterialShader
 		}
 
 		// This shader takes a very long time to compile with FXC, so we pre-compile it with DXC first and then forward the optimized HLSL to FXC.
-		OutEnvironment.CompilerFlags.Add(CFLAG_PrecompileWithDXC);
+		//OutEnvironment.CompilerFlags.Add(CFLAG_PrecompileWithDXC); // @lh-todo - Disabled to workaround SPIRV-Cross bug: StructuredBuffer<uint> is translated to ByteAddressBuffer in HLSL backend
 		OutEnvironment.CompilerFlags.Add(CFLAG_AllowTypedUAVLoads);
 
 		OutEnvironment.SetDefine(TEXT("GET_PRIMITIVE_DATA_OVERRIDE"), 1);
