Remove distance field AO visibility check. Distance field updates run prior to visibility and this check is no longer valid.

diff --git a/Engine/Source/Runtime/Renderer/Private/DistanceFieldAmbientOcclusion.cpp b/Engine/Source/Runtime/Renderer/Private/DistanceFieldAmbientOcclusion.cpp
--- a/Engine/Source/Runtime/Renderer/Private/DistanceFieldAmbientOcclusion.cpp
+++ b/Engine/Source/Runtime/Renderer/Private/DistanceFieldAmbientOcclusion.cpp
@@ -749,10 +749,7 @@ bool FSceneRenderer::ShouldPrepareGlobalDistanceField() const
 		return false;
 	}
 
-	bool bShouldPrepareForAO = SupportsDistanceFieldAO(Scene->GetFeatureLevel(), Scene->GetShaderPlatform())
-		&& (ShouldPrepareForDistanceFieldAO()
-			|| ((Views.Num() > 0) && Views[0].bUsesGlobalDistanceField)
-			|| ((FXSystem != nullptr) && FXSystem->UsesGlobalDistanceField()));
+	bool bShouldPrepareForAO = SupportsDistanceFieldAO(Scene->GetFeatureLevel(), Scene->GetShaderPlatform());
 
 	bShouldPrepareForAO = bShouldPrepareForAO || (IsLumenEnabled(Views[0]) && Lumen::UseGlobalSDFObjectGrid(*Views[0].Family));
 
