Undo forcing of DXC for BuildInstanceDrawCommands shader (UE-180840).

diff --git a/Engine/Shaders/Private/InstanceCulling/BuildInstanceDrawCommands.usf b/Engine/Shaders/Private/InstanceCulling/BuildInstanceDrawCommands.usf
--- a/Engine/Shaders/Private/InstanceCulling/BuildInstanceDrawCommands.usf
+++ b/Engine/Shaders/Private/InstanceCulling/BuildInstanceDrawCommands.usf
@@ -11,7 +11,7 @@
 #endif
 
 // Disable the declaration of any templated types so we can use this shader without HLSL 2021. This is necessary because
-// cross-compilation creates bugs with ENABLE_INSTANCE_COMPACTION
+// cross-compilation creates bugs with ENABLE_INSTANCE_COMPACTION.
 #define ALLOW_TEMPLATES 0
 
 #include "../Common.ush"
diff --git a/Engine/Source/Runtime/Renderer/Private/InstanceCulling/InstanceCullingContext.cpp b/Engine/Source/Runtime/Renderer/Private/InstanceCulling/InstanceCullingContext.cpp
--- a/Engine/Source/Runtime/Renderer/Private/InstanceCulling/InstanceCullingContext.cpp
+++ b/Engine/Source/Runtime/Renderer/Private/InstanceCulling/InstanceCullingContext.cpp
@@ -386,12 +386,6 @@ public:
 		FGlobalShader::ModifyCompilationEnvironment(Parameters, OutEnvironment);
 		FInstanceProcessingGPULoadBalancer::SetShaderDefines(OutEnvironment);
 
-		// Force use of DXC for mobile platforms as there are some remaining platforms that still use HLSLCC due to perf. regressions with DXC
-		if (IsMobilePlatform(Parameters.Platform))
-		{
-			OutEnvironment.CompilerFlags.Add(CFLAG_ForceDXC);
-		}
-
 		OutEnvironment.SetDefine(TEXT("INDIRECT_ARGS_NUM_WORDS"), FInstanceCullingContext::IndirectArgsNumWords);
 		OutEnvironment.SetDefine(TEXT("VF_SUPPORTS_PRIMITIVE_SCENE_DATA"), 1);
 		OutEnvironment.SetDefine(TEXT("USE_GLOBAL_GPU_SCENE_DATA"), 1);
