Fixes crash that could have Editor.PrimitivesDepthHistory allocate a RDG texture too large

diff --git a/Engine/Source/Runtime/Renderer/Private/PostProcess/PostProcessCompositePrimitivesCommon.cpp b/Engine/Source/Runtime/Renderer/Private/PostProcess/PostProcessCompositePrimitivesCommon.cpp
--- a/Engine/Source/Runtime/Renderer/Private/PostProcess/PostProcessCompositePrimitivesCommon.cpp
+++ b/Engine/Source/Runtime/Renderer/Private/PostProcess/PostProcessCompositePrimitivesCommon.cpp
@@ -139,7 +139,9 @@ void TemporalUpscaleDepthPass(
 		// Upscale factor shouldn't be higher than there is TAA samples, or that means there will be unrendered pixels.
 		const int32 ComputeMaxUpsampleFactorDueToTAA = FMath::FloorToInt(FMath::Sqrt(float(View.TemporalJitterSequenceLength) / (UpscaleFactor * UpscaleFactor)));
 
-		const int32 ComputeMaxUpsampleFactor = FMath::Clamp(ComputeMaxUpsampleFactorDueToTAA, 0, 4);
+		const int32 MaxRHIUpsampleFactor = FMath::Clamp(FMath::FloorToInt(float(GetMax2DTextureDimension()) / float(Extent.GetMax())), 1, 4);
+
+		const int32 ComputeMaxUpsampleFactor = FMath::Clamp(ComputeMaxUpsampleFactorDueToTAA, 0, MaxRHIUpsampleFactor);
 
 		const int32 DepthUpsampleFactor = FMath::Clamp(CVarCompositeTemporalUpsampleDepth.GetValueOnRenderThread(), 0, ComputeMaxUpsampleFactor);
 
