Implements a VisualizeTemporalUpscaler show flag to more quickly diagnose common problems

* All buffers are anoted with the VisualizeTexture command to see it;
* Velocity buffer is visualized with the VisualizeMotionBlur to make it obvious motion vector are important and need to be coherent, and VisualizeReprojection to make sure the reprojection lines up perfectly;
* SceneDepth is shown with the same pixel shader as the VisualizeTexture to not confuse the user;
* PostDOF translucency is shown along with it's alpha channel given how benefical it is to minimise ghosting on translucency that don't draw velocity;
* SceneColor and PostDOF translucency is shown in linear color space to make it obvious when pre-exposure arises;
* When r.PostProcessing.PropagateAlpha is enabled, alpha channel of SceneColor and temporal upscaler's output are displayed;
* The luminance used for flickering detection is also displayed for user to identify what is subject to detection of the flickering;

Summary is displayed with:
* Which temporal upscaler is being used and what anti-aliasing scalability group
* Input resolution and pixel format
* Output resolution and pixel format
* Current preexposure of the color pixel format to diagnose pre-exposure problems and also specify if it is overriden with r.EyeAdaptation.PreExposureOverride

diff --git a/Engine/Source/Runtime/Engine/Public/ShowFlags.h b/Engine/Source/Runtime/Engine/Public/ShowFlags.h
--- a/Engine/Source/Runtime/Engine/Public/ShowFlags.h
+++ b/Engine/Source/Runtime/Engine/Public/ShowFlags.h
@@ -391,6 +391,7 @@ private:
 		SetVisualizeMotionBlur(false);
 		SetVisualizeMotionVectors(false);
 		SetVisualizeReprojection(false);
+		SetVisualizeTemporalUpscaler(false);
 		SetEditingLevelInstance(false);
 		SetSelectionOutline(false);
 		SetDebugAI(false);
diff --git a/Engine/Source/Runtime/Engine/Public/ShowFlagsValues.inl b/Engine/Source/Runtime/Engine/Public/ShowFlagsValues.inl
--- a/Engine/Source/Runtime/Engine/Public/ShowFlagsValues.inl
+++ b/Engine/Source/Runtime/Engine/Public/ShowFlagsValues.inl
@@ -122,6 +122,8 @@ SHOWFLAG_FIXED_IN_SHIPPING(0, VisualizeMotionBlur, SFG_Visualize, NSLOCTEXT("Unr
 SHOWFLAG_FIXED_IN_SHIPPING(0, VisualizeMotionVectors, SFG_Hidden, NSLOCTEXT("UnrealEd", "VisualizeMotionVectorsSF", "Motion Vectors"))
 /** Helper to ensure previous frame's temporal history reprojections using motion vectors are working correctly. */
 SHOWFLAG_FIXED_IN_SHIPPING(0, VisualizeReprojection, SFG_Visualize, NSLOCTEXT("UnrealEd", "VisualizeReprojectionSF", "Previous frame's reprojection"))
+/** Helper to diagnose why the temporal upscaler may run into issues. */
+SHOWFLAG_FIXED_IN_SHIPPING(0, VisualizeTemporalUpscaler, SFG_Visualize, NSLOCTEXT("UnrealEd", "VisualizeTemporalUpscalerSF", "Temporal Upscaler (TSR, TAAU or third party plugins)"))
 /** Whether to display the Reflection Environment feature, which has local reflections from Reflection Capture actors, for now SHOWFLAG_ALWAYS_ACCESSIBLE because it's exposed in SceneCapture */
 SHOWFLAG_ALWAYS_ACCESSIBLE(ReflectionEnvironment, SFG_LightingFeatures, NSLOCTEXT("UnrealEd", "ReflectionEnvironmentSF", "Reflection Environment"))
 /** Visualize pixels that are outside of their object's bounding box (content error). */
diff --git a/Engine/Source/Runtime/RenderCore/Private/VisualizeTexture.cpp b/Engine/Source/Runtime/RenderCore/Private/VisualizeTexture.cpp
--- a/Engine/Source/Runtime/RenderCore/Private/VisualizeTexture.cpp
+++ b/Engine/Source/Runtime/RenderCore/Private/VisualizeTexture.cpp
@@ -631,21 +631,21 @@ void FVisualizeTexture::ReleaseDynamicRHI()
 	Captured = {};
 }
 
-void FVisualizeTexture::CreateContentCapturePass(FRDGBuilder& GraphBuilder, const FRDGTextureRef InputTexture, uint32 CaptureId)
+// static
+FRDGTextureRef FVisualizeTexture::AddVisualizeTexturePass(
+	FRDGBuilder& GraphBuilder,
+	FGlobalShaderMap* ShaderMap,
+	const FRDGTextureRef InputTexture,
+	const FConfig& VisualizeConfig,
+	EInputValueMapping InputValueMapping,
+	uint32 CaptureId)
 {
-	if (!InputTexture)
-	{
-		return;
-	}
+	check(InputTexture);
+	check(!EnumHasAnyFlags(InputTexture->Desc.Flags, TexCreate_CPUReadback));
 
 	const FRDGTextureDesc& InputDesc = InputTexture->Desc;
 	const FIntPoint InputExtent = InputDesc.Extent;
 
-	if (EnumHasAnyFlags(InputDesc.Flags, TexCreate_CPUReadback))
-	{
-		return;
-	}
-
 	FIntPoint OutputExtent = InputExtent;
 
 	// Clamp to reasonable value to prevent crash
@@ -654,18 +654,7 @@ void FVisualizeTexture::CreateContentCapturePass(FRDGBuilder& GraphBuilder, cons
 
 	FRDGTextureRef OutputTexture = GraphBuilder.CreateTexture(FRDGTextureDesc::Create2D(OutputExtent, PF_B8G8R8A8, FClearValueBinding(FLinearColor(1, 1, 0, 1)), TexCreate_RenderTargetable | TexCreate_ShaderResource), TEXT("VisualizeTexture"));
 
-	EInputValueMapping InputValueMapping = EInputValueMapping::Color;
-
 	{
-		if (InputDesc.Format == PF_ShadowDepth)
-		{
-			InputValueMapping = EInputValueMapping::Shadow;
-		}
-		else if (EnumHasAnyFlags(InputDesc.Flags, TexCreate_DepthStencilTargetable))
-		{
-			InputValueMapping = EInputValueMapping::Depth;
-		}
-
 		const EVisualisePSType VisualizeType = GetVisualizePSType(InputDesc);
 
 		FVisualizeTexturePS::FParameters* PassParameters = GraphBuilder.AllocParameters<FVisualizeTexturePS::FParameters>();
@@ -685,16 +674,16 @@ void FVisualizeTexture::CreateContentCapturePass(FRDGBuilder& GraphBuilder, cons
 				float FracScale = 1.0f;
 
 				// w * almost_1 to avoid frac(1) => 0
-				PassParameters->VisualizeParam[0] = FVector4f(Config.RGBMul, Config.SingleChannelMul, Add, FracScale * 0.9999f);
-				PassParameters->VisualizeParam[1] = FVector4f(CVarAllowBlinking.GetValueOnRenderThread() ? BlinkState : 1.0f, (Config.ShaderOp == EShaderOp::Saturate) ? 1.0f : 0.0f, Config.ArrayIndex, Config.MipIndex);
-				PassParameters->VisualizeParam[2] = FVector4f((float)InputValueMapping, 0.0f, Config.SingleChannel);
+				PassParameters->VisualizeParam[0] = FVector4f(VisualizeConfig.RGBMul, VisualizeConfig.SingleChannelMul, Add, FracScale * 0.9999f);
+				PassParameters->VisualizeParam[1] = FVector4f(CVarAllowBlinking.GetValueOnRenderThread() ? BlinkState : 1.0f, (VisualizeConfig.ShaderOp == EShaderOp::Saturate) ? 1.0f : 0.0f, VisualizeConfig.ArrayIndex, VisualizeConfig.MipIndex);
+				PassParameters->VisualizeParam[2] = FVector4f((float)InputValueMapping, 0.0f, VisualizeConfig.SingleChannel);
 			}
 
 			FRDGTextureSRV* InputSRV = nullptr;
 			FRHISamplerState* PointSampler = TStaticSamplerState<SF_Point, AM_Clamp, AM_Clamp, AM_Clamp>::GetRHI();
 			if (InputTexture->Desc.Dimension == ETextureDimension::Texture2DArray)
 			{
-				InputSRV = GraphBuilder.CreateSRV(FRDGTextureSRVDesc::CreateForSlice(InputTexture, FMath::Clamp(int32(Config.ArrayIndex), 0, int32(InputDesc.ArraySize) - 1)));
+				InputSRV = GraphBuilder.CreateSRV(FRDGTextureSRVDesc::CreateForSlice(InputTexture, FMath::Clamp(int32(VisualizeConfig.ArrayIndex), 0, int32(InputDesc.ArraySize) - 1)));
 			}
 			else
 			{
@@ -722,7 +711,6 @@ void FVisualizeTexture::CreateContentCapturePass(FRDGBuilder& GraphBuilder, cons
 			PassParameters->RenderTargets[0] = FRenderTargetBinding(OutputTexture, ERenderTargetLoadAction::EClear);
 		}
 
-		auto ShaderMap = GetGlobalShaderMap(FeatureLevel);
 		FVisualizeTexturePS::FPermutationDomain PermutationVector;
 		PermutationVector.Set<FVisualizeTexturePS::FVisualisePSTypeDim>(VisualizeType);
 
@@ -733,18 +721,18 @@ void FVisualizeTexture::CreateContentCapturePass(FRDGBuilder& GraphBuilder, cons
 		{
 			if (InputDesc.IsTexture3D())
 			{
-				ExtendedDrawEvent += FString::Printf(TEXT("x%d CapturedSlice=%d"), InputDesc.Depth, Config.ArrayIndex);
+				ExtendedDrawEvent += FString::Printf(TEXT("x%d CapturedSlice=%d"), InputDesc.Depth, VisualizeConfig.ArrayIndex);
 			}
 
 			if (InputDesc.IsTextureArray())
 			{
-				ExtendedDrawEvent += FString::Printf(TEXT(" ArraySize=%d CapturedSlice=%d"), InputDesc.ArraySize, Config.ArrayIndex);
+				ExtendedDrawEvent += FString::Printf(TEXT(" ArraySize=%d CapturedSlice=%d"), InputDesc.ArraySize, VisualizeConfig.ArrayIndex);
 			}
 
 			// Precise the mip level being captured in the mip level when there is a mip chain.
 			if (InputDesc.IsMipChain())
 			{
-				ExtendedDrawEvent += FString::Printf(TEXT(" Mips=%d CapturedMip=%d"), InputDesc.NumMips, Config.MipIndex);
+				ExtendedDrawEvent += FString::Printf(TEXT(" Mips=%d CapturedMip=%d"), InputDesc.NumMips, VisualizeConfig.MipIndex);
 			}
 		}
 
@@ -761,6 +749,43 @@ void FVisualizeTexture::CreateContentCapturePass(FRDGBuilder& GraphBuilder, cons
 			FIntRect(0, 0, OutputExtent.X, OutputExtent.Y));
 	}
 
+	return OutputTexture;
+}
+
+void FVisualizeTexture::CreateContentCapturePass(FRDGBuilder& GraphBuilder, const FRDGTextureRef InputTexture, uint32 CaptureId)
+{
+	if (!InputTexture)
+	{
+		return;
+	}
+
+	const FRDGTextureDesc& InputDesc = InputTexture->Desc;
+	const FIntPoint InputExtent = InputDesc.Extent;
+
+	if (EnumHasAnyFlags(InputDesc.Flags, TexCreate_CPUReadback))
+	{
+		return;
+	}
+
+	EInputValueMapping InputValueMapping = EInputValueMapping::Color;
+	{
+		if (InputDesc.Format == PF_ShadowDepth)
+		{
+			InputValueMapping = EInputValueMapping::Shadow;
+		}
+		else if (EnumHasAnyFlags(InputDesc.Flags, TexCreate_DepthStencilTargetable))
+		{
+			InputValueMapping = EInputValueMapping::Depth;
+		}
+	}
+
+	FGlobalShaderMap* ShaderMap = GetGlobalShaderMap(FeatureLevel);
+	FRDGTextureRef OutputTexture = AddVisualizeTexturePass(GraphBuilder, ShaderMap, InputTexture, Config, InputValueMapping, CaptureId);
+
+	FIntPoint OutputExtent = InputExtent;
+	OutputExtent.X = FMath::Max(OutputExtent.X, 1);
+	OutputExtent.Y = FMath::Max(OutputExtent.Y, 1);
+
 	{
 		Captured.Desc = Translate(InputDesc);
 		Captured.Desc.DebugName = InputTexture->Name;
@@ -877,3 +902,54 @@ void FVisualizeTexture::Visualize(const FString& InName, TOptional<uint32> InVer
 }
 
 #endif // SUPPORTS_VISUALIZE_TEXTURE
+
+// static
+FRDGTextureRef FVisualizeTexture::AddVisualizeTexturePass(
+	FRDGBuilder& GraphBuilder,
+	class FGlobalShaderMap* ShaderMap,
+	const FRDGTextureRef InputTexture)
+#if SUPPORTS_VISUALIZE_TEXTURE
+{
+	check(InputTexture);
+	EInputValueMapping InputValueMapping = EInputValueMapping::Color;
+	{
+		if (InputTexture->Desc.Format == PF_ShadowDepth)
+		{
+			InputValueMapping = EInputValueMapping::Shadow;
+		}
+		else if (EnumHasAnyFlags(InputTexture->Desc.Flags, TexCreate_DepthStencilTargetable))
+		{
+			InputValueMapping = EInputValueMapping::Depth;
+		}
+	}
+
+	FConfig VisualizeConfig;
+
+	return AddVisualizeTexturePass(GraphBuilder, ShaderMap, InputTexture, VisualizeConfig, InputValueMapping, /* CaptureId = */ 0);
+}
+#else
+{
+	return InputTexture;
+}
+#endif
+
+// static
+FRDGTextureRef FVisualizeTexture::AddVisualizeTextureAlphaPass(
+	FRDGBuilder& GraphBuilder,
+	class FGlobalShaderMap* ShaderMap,
+	const FRDGTextureRef InputTexture)
+#if SUPPORTS_VISUALIZE_TEXTURE
+{
+	check(InputTexture);
+	FConfig VisualizeConfig;
+	VisualizeConfig.SingleChannel = 3;
+	VisualizeConfig.SingleChannelMul = 1.0f;
+	VisualizeConfig.RGBMul = 0.0f;
+
+	return AddVisualizeTexturePass(GraphBuilder, ShaderMap, InputTexture, VisualizeConfig, EInputValueMapping::Color, /* CaptureId = */ 0);
+}
+#else
+{
+	return InputTexture;
+}
+#endif
diff --git a/Engine/Source/Runtime/RenderCore/Public/VisualizeTexture.h b/Engine/Source/Runtime/RenderCore/Public/VisualizeTexture.h
--- a/Engine/Source/Runtime/RenderCore/Public/VisualizeTexture.h
+++ b/Engine/Source/Runtime/RenderCore/Public/VisualizeTexture.h
@@ -48,6 +48,16 @@ public:
 	inline void SetCheckPoint(FRHICommandListImmediate& RHICmdList, IPooledRenderTarget* PooledRenderTarget) {}
 #endif
 
+	static FRDGTextureRef AddVisualizeTexturePass(
+		FRDGBuilder& GraphBuilder,
+		class FGlobalShaderMap* ShaderMap,
+		const FRDGTextureRef InputTexture);
+
+	static FRDGTextureRef AddVisualizeTextureAlphaPass(
+		FRDGBuilder& GraphBuilder,
+		class FGlobalShaderMap* ShaderMap,
+		const FRDGTextureRef InputTexture);
+
 private:
 	enum class EFlags
 	{
@@ -109,15 +119,6 @@ private:
 	/** Determine whether a texture should be captured for debugging purposes and return the capture id if needed. */
 	TOptional<uint32> ShouldCapture(const TCHAR* DebugName, uint32 MipIndex);
 
-	/** Create a pass capturing a texture. */
-	void CreateContentCapturePass(FRDGBuilder& GraphBuilder, FRDGTextureRef Texture, uint32 CaptureId);
-
-	void ReleaseDynamicRHI() override;
-
-	void Visualize(const FString& InName, TOptional<uint32> InVersion = {});
-
-	uint32 GetVersionCount(const TCHAR* InName) const;
-
 	struct FConfig
 	{
 		float RGBMul = 1.0f;
@@ -132,7 +133,27 @@ private:
 		EShaderOp ShaderOp = EShaderOp::Frac;
 		uint32 MipIndex = 0;
 		uint32 ArrayIndex = 0;
-	} Config;
+	};
+
+	/** Adds a pass to visualize a texture. */
+	static FRDGTextureRef AddVisualizeTexturePass(
+		FRDGBuilder& GraphBuilder,
+		class FGlobalShaderMap* ShaderMap,
+		const FRDGTextureRef InputTexture,
+		const FConfig& Config,
+		EInputValueMapping InputValueMapping,
+		uint32 CaptureId);
+
+	/** Create a pass capturing a texture. */
+	void CreateContentCapturePass(FRDGBuilder& GraphBuilder, FRDGTextureRef Texture, uint32 CaptureId);
+
+	void ReleaseDynamicRHI() override;
+
+	void Visualize(const FString& InName, TOptional<uint32> InVersion = {});
+
+	uint32 GetVersionCount(const TCHAR* InName) const;
+
+	FConfig Config;
 
 	struct FRequested
 	{
diff --git a/Engine/Source/Runtime/Renderer/Private/PostProcess/PostProcessVisualizeBuffer.cpp b/Engine/Source/Runtime/Renderer/Private/PostProcess/PostProcessVisualizeBuffer.cpp
--- a/Engine/Source/Runtime/Renderer/Private/PostProcess/PostProcessVisualizeBuffer.cpp
+++ b/Engine/Source/Runtime/Renderer/Private/PostProcess/PostProcessVisualizeBuffer.cpp
@@ -31,29 +31,6 @@ public:
 
 IMPLEMENT_GLOBAL_SHADER(FVisualizeBufferPS, "/Engine/Private/PostProcessVisualizeBuffer.usf", "MainPS", SF_Pixel);
 
-struct FVisualizeBufferTile
-{
-	// The input texture to visualize.
-	FScreenPassTexture Input;
-
-	// The label of the tile shown on the visualizer.
-	FString Label;
-
-	// Whether the tile is shown as selected.
-	bool bSelected = false;
-};
-
-struct FVisualizeBufferInputs
-{
-	FScreenPassRenderTarget OverrideOutput;
-
-	// The scene color input to propagate.
-	FScreenPassTexture SceneColor;
-
-	// The array of tiles to render onto the scene color texture.
-	TArrayView<const FVisualizeBufferTile> Tiles;
-};
-
 FScreenPassTexture AddVisualizeBufferPass(FRDGBuilder& GraphBuilder, const FViewInfo& View, const FVisualizeBufferInputs& Inputs)
 {
 	check(Inputs.SceneColor.IsValid());
diff --git a/Engine/Source/Runtime/Renderer/Private/PostProcess/PostProcessVisualizeBuffer.h b/Engine/Source/Runtime/Renderer/Private/PostProcess/PostProcessVisualizeBuffer.h
--- a/Engine/Source/Runtime/Renderer/Private/PostProcess/PostProcessVisualizeBuffer.h
+++ b/Engine/Source/Runtime/Renderer/Private/PostProcess/PostProcessVisualizeBuffer.h
@@ -51,4 +51,29 @@ struct FVisualizeGBufferOverviewInputs
 	bool bOutputInHDR = false;
 };
 
-FScreenPassTexture AddVisualizeGBufferOverviewPass(FRDGBuilder& GraphBuilder, const FViewInfo& View, const FVisualizeGBufferOverviewInputs& Inputs);
\ No newline at end of file
+FScreenPassTexture AddVisualizeGBufferOverviewPass(FRDGBuilder& GraphBuilder, const FViewInfo& View, const FVisualizeGBufferOverviewInputs& Inputs);
+
+struct FVisualizeBufferTile
+{
+	// The input texture to visualize.
+	FScreenPassTexture Input;
+
+	// The label of the tile shown on the visualizer.
+	FString Label;
+
+	// Whether the tile is shown as selected.
+	bool bSelected = false;
+};
+
+struct FVisualizeBufferInputs
+{
+	FScreenPassRenderTarget OverrideOutput;
+
+	// The scene color input to propagate.
+	FScreenPassTexture SceneColor;
+
+	// The array of tiles to render onto the scene color texture.
+	TArrayView<const FVisualizeBufferTile> Tiles;
+};
+
+FScreenPassTexture AddVisualizeBufferPass(FRDGBuilder& GraphBuilder, const FViewInfo& View, const FVisualizeBufferInputs& Inputs);
diff --git a/Engine/Source/Runtime/Renderer/Private/PostProcess/PostProcessing.cpp b/Engine/Source/Runtime/Renderer/Private/PostProcess/PostProcessing.cpp
--- a/Engine/Source/Runtime/Renderer/Private/PostProcess/PostProcessing.cpp
+++ b/Engine/Source/Runtime/Renderer/Private/PostProcess/PostProcessing.cpp
@@ -266,6 +266,8 @@ void AddPostProcessingPasses(
 	FRDGTextureRef LocalExposureTexture = nullptr;
 	FRDGTextureRef LocalExposureBlurredLogLumTexture = BlackDummy.Texture;
 
+	FVisualizeTemporalUpscalerInputs VisualizeTemporalUpscalerInputs;
+
 	const FEngineShowFlags& EngineShowFlags = View.Family->EngineShowFlags;
 	const bool bVisualizeHDR = EngineShowFlags.VisualizeHDR;
 	const bool bViewFamilyOutputInHDR = View.Family->RenderTarget->GetSceneHDREnabled();
@@ -303,6 +305,7 @@ void AddPostProcessingPasses(
 		VisualizeHDR,
 		VisualizeLocalExposure,
 		VisualizeMotionVectors,
+		VisualizeTemporalUpscaler,
 		PixelInspector,
 		HMDDistortion,
 		HighResolutionScreenshotMask,
@@ -354,6 +357,7 @@ void AddPostProcessingPasses(
 		TEXT("VisualizeHDR"),
 		TEXT("VisualizeLocalExposure"),
 		TEXT("VisualizeMotionVectors"),
+		TEXT("VisualizeTemporalUpscaler"),
 		TEXT("PixelInspector"),
 		TEXT("HMDDistortion"),
 		TEXT("HighResolutionScreenshotMask"),
@@ -397,6 +401,7 @@ void AddPostProcessingPasses(
 	PassSequence.SetEnabled(EPass::VisualizeLumenSceneOverview, LumenVisualizeMode == VISUALIZE_MODE_OVERVIEW && bPostProcessingEnabled);
 	PassSequence.SetEnabled(EPass::VisualizeHDR, EngineShowFlags.VisualizeHDR);
 	PassSequence.SetEnabled(EPass::VisualizeMotionVectors, EngineShowFlags.VisualizeMotionVectors || EngineShowFlags.VisualizeReprojection);
+	PassSequence.SetEnabled(EPass::VisualizeTemporalUpscaler, EngineShowFlags.VisualizeTemporalUpscaler);
 #if WITH_EDITOR
 	PassSequence.SetEnabled(EPass::PixelInspector, View.bUsePixelInspector);
 #else
@@ -688,6 +693,14 @@ void AddPostProcessingPasses(
 			HalfResSceneColor = Outputs.HalfRes;
 			QuarterResSceneColor = Outputs.QuarterRes;
 			VelocityFlattenTextures = Outputs.VelocityFlattenTextures;
+
+			if (PassSequence.IsEnabled(EPass::VisualizeTemporalUpscaler))
+			{
+				VisualizeTemporalUpscalerInputs.TAAConfig = TAAConfig;
+				VisualizeTemporalUpscalerInputs.UpscalerUsed = UpscalerToUse;
+				VisualizeTemporalUpscalerInputs.Inputs = UpscalerPassInputs;
+				VisualizeTemporalUpscalerInputs.Outputs = Outputs;
+			}
 		}
 		else if (ReflectionsMethod == EReflectionsMethod::SSR)
 		{
@@ -1363,6 +1376,14 @@ void AddPostProcessingPasses(
 		SceneColor = AddVisualizeMotionVectorsPass(GraphBuilder, View, PassInputs);
 	}
 
+	if (PassSequence.IsEnabled(EPass::VisualizeTemporalUpscaler))
+	{
+		PassSequence.AcceptOverrideIfLastPass(EPass::VisualizeTemporalUpscaler, VisualizeTemporalUpscalerInputs.OverrideOutput);
+		VisualizeTemporalUpscalerInputs.SceneColor = SceneColor;
+
+		SceneColor = AddVisualizeTemporalUpscalerPass(GraphBuilder, View, VisualizeTemporalUpscalerInputs);
+	}
+
 #if WITH_EDITOR
 	if (PassSequence.IsEnabled(EPass::PixelInspector))
 	{
diff --git a/Engine/Source/Runtime/Renderer/Private/PostProcess/TemporalAA.h b/Engine/Source/Runtime/Renderer/Private/PostProcess/TemporalAA.h
--- a/Engine/Source/Runtime/Renderer/Private/PostProcess/TemporalAA.h
+++ b/Engine/Source/Runtime/Renderer/Private/PostProcess/TemporalAA.h
@@ -225,3 +225,21 @@ public:
 
 	static EMainTAAPassConfig GetMainTAAPassConfig(const FViewInfo& View);
 }; 
+
+/** Interface for the VisualizeTSR showflag. */
+struct FVisualizeTemporalUpscalerInputs
+{
+	// [Optional] Render to the specified output. If invalid, a new texture is created and returned.
+	FScreenPassRenderTarget OverrideOutput;
+
+	// Scene color.
+	FScreenPassTexture SceneColor;
+
+	// Temporal upscaler used and its inputs and outputs.
+	EMainTAAPassConfig TAAConfig = EMainTAAPassConfig::Disabled;
+	const ITemporalUpscaler* UpscalerUsed = nullptr;
+	ITemporalUpscaler::FPassInputs Inputs;
+	ITemporalUpscaler::FOutputs Outputs;
+};
+
+FScreenPassTexture AddVisualizeTemporalUpscalerPass(FRDGBuilder& GraphBuilder, const FViewInfo& View, const FVisualizeTemporalUpscalerInputs& Inputs);
diff --git a/Engine/Source/Runtime/Renderer/Private/PostProcess/VisualizeTemporalUpscaler.cpp b/Engine/Source/Runtime/Renderer/Private/PostProcess/VisualizeTemporalUpscaler.cpp
new file mode 100644
--- /dev/null
+++ b/Engine/Source/Runtime/Renderer/Private/PostProcess/VisualizeTemporalUpscaler.cpp
@@ -0,0 +1,264 @@
+// Copyright Epic Games, Inc. All Rights Reserved.
+
+#include "PostProcess/TemporalAA.h"
+#include "PostProcess/PostProcessing.h"
+#include "ScenePrivate.h"
+#include "ScreenPass.h"
+#include "PixelShaderUtils.h"
+#include "UnrealEngine.h"
+#include "PostProcess/PostProcessMotionBlur.h"
+#include "PostProcess/PostProcessVisualizeBuffer.h"
+#include "PostProcess/VisualizeMotionVectors.h"
+#include "VisualizeTexture.h"
+
+
+FScreenPassTexture AddVisualizeTemporalUpscalerPass(FRDGBuilder& GraphBuilder, const FViewInfo& View, const FVisualizeTemporalUpscalerInputs& Inputs)
+{
+	check(Inputs.SceneColor.IsValid());
+
+	RDG_EVENT_SCOPE(GraphBuilder, "VisualizeTemporalUpscaler %dx%d", Inputs.SceneColor.ViewRect.Width(), Inputs.SceneColor.ViewRect.Height());
+
+	const bool bSupportsAlpha = IsPostProcessingWithAlphaChannelSupported();
+
+	// Allocate output and copy
+	FScreenPassRenderTarget Output = Inputs.OverrideOutput;
+	if (!Output.IsValid())
+	{
+		Output = FScreenPassRenderTarget::CreateFromInput(GraphBuilder, Inputs.SceneColor, View.GetOverwriteLoadAction(), TEXT("MotionVectors.Visualize"));
+	}
+
+	// Early return if not using a temporal upscaler.
+	if (!Inputs.UpscalerUsed)
+	{
+		check(Inputs.TAAConfig == EMainTAAPassConfig::Disabled);
+		FRHICopyTextureInfo CopyInfo;
+		AddCopyTexturePass(
+			GraphBuilder,
+			Inputs.SceneColor.Texture,
+			Output.Texture,
+			CopyInfo);
+
+		AddDrawCanvasPass(GraphBuilder, RDG_EVENT_NAME("VisualizeTemporalUpscaler Text"), View, FScreenPassRenderTarget(Output, ERenderTargetLoadAction::ELoad),
+			[&ViewRect = Output.ViewRect](FCanvas& Canvas)
+		{
+			const float DPIScale = Canvas.GetDPIScale();
+			Canvas.SetBaseTransform(FMatrix(FScaleMatrix(DPIScale) * Canvas.CalcBaseTransform2D(Canvas.GetViewRect().Width(), Canvas.GetViewRect().Height())));
+
+			FIntPoint LabelLocation(60, 60);
+			Canvas.DrawShadowedString(LabelLocation.X / DPIScale, LabelLocation.Y / DPIScale, TEXT("No temporal upscaler used"), GetStatsFont(), FLinearColor::Red);
+		});
+
+		return MoveTemp(Output);
+	}
+
+	// Populate all the tiles
+	TArray<FVisualizeBufferTile> Tiles;
+	Tiles.SetNum(16);
+	{
+		auto VisualizeTextureLabel = [](FRDGTextureRef Texture, const TCHAR* Suffix = TEXT(""))
+		{
+			return FString::Printf(TEXT("vis %s%s"), Texture->Name, Suffix);
+		};
+
+		auto CropViewRectToCenter = [](FIntRect ViewRect)
+		{
+			return FIntRect(ViewRect.Min + ViewRect.Size() / 4, ViewRect.Min + (ViewRect.Size() * 3) / 4);
+		};
+
+		// Depth buffer
+		{
+			FVisualizeBufferTile& Tile = Tiles[4 * 0 + 0];
+			Tile.Input.Texture = FVisualizeTexture::AddVisualizeTexturePass(GraphBuilder, View.ShaderMap, Inputs.Inputs.SceneDepthTexture);
+			Tile.Input.ViewRect = CropViewRectToCenter(View.ViewRect);
+			Tile.Label = VisualizeTextureLabel(Inputs.Inputs.SceneDepthTexture);
+		}
+
+		// Scene Color
+		{
+			FVisualizeBufferTile& Tile = Tiles[4 * 0 + 1];
+			Tile.Input = FScreenPassTexture(Inputs.Inputs.SceneColorTexture, CropViewRectToCenter(View.ViewRect));
+			Tile.Label = VisualizeTextureLabel(Inputs.Inputs.SceneColorTexture);
+		}
+
+		// Scene Color alpha
+		if (bSupportsAlpha)
+		{
+			FVisualizeBufferTile& Tile = Tiles[4 * 3 + 1];
+			//Tile.Input = FScreenPassTexture(Inputs.Inputs.SceneColorTexture, CropViewRectToCenter(View.ViewRect));
+			Tile.Input.Texture = FVisualizeTexture::AddVisualizeTextureAlphaPass(GraphBuilder, View.ShaderMap, Inputs.Inputs.SceneColorTexture);
+			Tile.Input.ViewRect = CropViewRectToCenter(View.ViewRect);
+			Tile.Label = VisualizeTextureLabel(Inputs.Inputs.SceneColorTexture, TEXT(" A"));
+		}
+
+		// Translucency
+		{
+			FVisualizeBufferTile& Tile = Tiles[4 * 0 + 2];
+			if (Inputs.Inputs.PostDOFTranslucencyResources.IsValid())
+			{
+				Tile.Input.Texture = Inputs.Inputs.PostDOFTranslucencyResources.ColorTexture.Resolve;
+				Tile.Input.ViewRect = CropViewRectToCenter(Inputs.Inputs.PostDOFTranslucencyResources.ViewRect);
+				Tile.Label = VisualizeTextureLabel(Inputs.Inputs.PostDOFTranslucencyResources.ColorTexture.Resolve);
+			}
+			else
+			{
+				Tile.Input = FScreenPassTexture(GSystemTextures.GetBlackDummy(GraphBuilder));
+				Tile.Label = TEXT("No PostDOF Translucency!");
+			}
+		}
+
+		// Translucency alpha
+		{
+			FVisualizeBufferTile& Tile = Tiles[4 * 0 + 3];
+			if (Inputs.Inputs.PostDOFTranslucencyResources.IsValid())
+			{
+				Tile.Input.Texture = FVisualizeTexture::AddVisualizeTextureAlphaPass(GraphBuilder, View.ShaderMap, Inputs.Inputs.PostDOFTranslucencyResources.ColorTexture.Resolve);
+				Tile.Input.ViewRect = CropViewRectToCenter(Inputs.Inputs.PostDOFTranslucencyResources.ViewRect);
+				Tile.Label = VisualizeTextureLabel(Inputs.Inputs.PostDOFTranslucencyResources.ColorTexture.Resolve, TEXT(" A"));
+			}
+			else
+			{
+				Tile.Input = FScreenPassTexture(GSystemTextures.GetBlackDummy(GraphBuilder));
+				Tile.Label = TEXT("No PostDOF Translucency!");
+			}
+		}
+
+		// Motion blur
+		{
+			FMotionBlurInputs PassInputs;
+			PassInputs.SceneColor = Inputs.Outputs.FullRes;
+			PassInputs.SceneDepth = FScreenPassTexture(Inputs.Inputs.SceneDepthTexture, View.ViewRect);
+			PassInputs.SceneVelocity = FScreenPassTexture(Inputs.Inputs.SceneVelocityTexture, View.ViewRect);
+
+			FVisualizeBufferTile& Tile = Tiles[4 * 1 + 0];
+			Tile.Input = AddVisualizeMotionBlurPass(GraphBuilder, View, PassInputs);
+			Tile.Input.ViewRect = CropViewRectToCenter(Tile.Input.ViewRect);
+			Tile.Label = TEXT("show VisualizeMotionBlur");
+		}
+
+		// Reprojection
+		{
+			FVisualizeMotionVectorsInputs PassInputs;
+			PassInputs.SceneColor = Inputs.SceneColor;
+			PassInputs.SceneDepth = FScreenPassTexture(Inputs.Inputs.SceneDepthTexture, View.ViewRect);
+			PassInputs.SceneVelocity = FScreenPassTexture(Inputs.Inputs.SceneVelocityTexture, View.ViewRect);
+
+			FVisualizeBufferTile& Tile = Tiles[4 * 2 + 0];
+			Tile.Input = AddVisualizeMotionVectorsPass(GraphBuilder, View, PassInputs);
+			Tile.Input.ViewRect = CropViewRectToCenter(Tile.Input.ViewRect);
+			Tile.Label = TEXT("show VisualizeReprojection");
+		}
+
+		// Display TSR's luminance used for stability.
+		if (Inputs.TAAConfig == EMainTAAPassConfig::TSR)
+		{
+			FVisualizeBufferTile& Tile = Tiles[4 * 1 + 3];
+			Tile.Input = Inputs.Inputs.MoireInputTexture;
+			Tile.Input.ViewRect = CropViewRectToCenter(View.ViewRect);
+			Tile.Label = VisualizeTextureLabel(Inputs.Inputs.MoireInputTexture.Texture);
+		}
+
+		// Output
+		{
+			FVisualizeBufferTile& Tile = Tiles[4 * 3 + 3];
+			Tile.Input = Inputs.Outputs.FullRes;
+			Tile.Input.ViewRect = CropViewRectToCenter(Tile.Input.ViewRect);
+			Tile.Label = VisualizeTextureLabel(Inputs.Outputs.FullRes.Texture);
+		}
+
+		// Output alpha
+		if (bSupportsAlpha)
+		{
+			FVisualizeBufferTile& Tile = Tiles[4 * 3 + 2];
+			Tile.Input.Texture = FVisualizeTexture::AddVisualizeTextureAlphaPass(GraphBuilder, View.ShaderMap, Inputs.Outputs.FullRes.Texture);
+			Tile.Input.ViewRect = CropViewRectToCenter(Inputs.Outputs.FullRes.ViewRect);
+			Tile.Label = VisualizeTextureLabel(Inputs.Outputs.FullRes.Texture, TEXT(" A"));
+		}
+
+		// Black bottom left corner
+		{
+			FVisualizeBufferTile& Tile = Tiles[4 * 3 + 0];
+			Tile.Input = FScreenPassTexture(GSystemTextures.GetBlackDummy(GraphBuilder));
+			Tile.Label = TEXT("Summary");
+		}
+	}
+
+	// Draws all overview tiles
+	{
+		FVisualizeBufferInputs PassInputs;
+		PassInputs.OverrideOutput = Output;
+		PassInputs.SceneColor = Inputs.SceneColor;
+		PassInputs.Tiles = Tiles;
+
+		AddVisualizeBufferPass(GraphBuilder, View, PassInputs);
+	}
+
+	// Draw additional text
+	{
+		AddDrawCanvasPass(GraphBuilder, RDG_EVENT_NAME("VisualizeTemporalUpscaler Text"), View, FScreenPassRenderTarget(Output, ERenderTargetLoadAction::ELoad),
+			[&View, &ViewRect = Output.ViewRect, &Inputs, bSupportsAlpha](FCanvas& Canvas)
+		{
+			const float DPIScale = Canvas.GetDPIScale();
+			Canvas.SetBaseTransform(FMatrix(FScaleMatrix(DPIScale) * Canvas.CalcBaseTransform2D(Canvas.GetViewRect().Width(), Canvas.GetViewRect().Height())));
+
+			auto QuickDrawSummary = [&](int32 Location, const FString& Text)
+			{
+				FIntPoint LabelLocation(20, 20 + 22 * Location + (ViewRect.Height() * 3) / 4);
+				Canvas.DrawShadowedString(LabelLocation.X / DPIScale, LabelLocation.Y / DPIScale, *Text, GetStatsFont(), FLinearColor::White);
+			};
+
+			// Display which temporal upscaler is being used.
+			{
+				static auto CVarAntiAliasingQuality = IConsoleManager::Get().FindConsoleVariable(TEXT("sg.AntiAliasingQuality"));
+				check(CVarAntiAliasingQuality);
+
+				FString Text;
+				if (Inputs.TAAConfig == EMainTAAPassConfig::TAA)
+				{
+					Text = FString::Printf(TEXT("TemporalUpscaler: TAA (sg.AntiAliasingQuality=%d)"), CVarAntiAliasingQuality->GetInt());
+				}
+				else if (Inputs.TAAConfig == EMainTAAPassConfig::TSR)
+				{
+					Text = FString::Printf(TEXT("TemporalUpscaler: TSR (sg.AntiAliasingQuality=%d)"), CVarAntiAliasingQuality->GetInt());
+				}
+				else
+				{
+					Text = FString::Printf(TEXT("ThirdParty TemporalUpscaler: %s"), Inputs.UpscalerUsed->GetDebugName());
+				}
+				QuickDrawSummary(/* Location = */ 0, Text);
+			}
+
+			// Display the input/output resolutions
+			{
+				QuickDrawSummary(/* Location = */ 1, FString::Printf(TEXT("Input: %dx%d %s"), View.ViewRect.Width(), View.ViewRect.Height(), GPixelFormats[Inputs.Inputs.SceneColorTexture->Desc.Format].Name));
+				QuickDrawSummary(/* Location = */ 2, FString::Printf(TEXT("Output: %dx%d %s"), Inputs.Outputs.FullRes.ViewRect.Width(), Inputs.Outputs.FullRes.ViewRect.Height(), GPixelFormats[Inputs.Outputs.FullRes.Texture->Desc.Format].Name));
+			}
+
+			// Display the pre-exposure being used
+			{
+				static auto CVarPreExposureOverride = IConsoleManager::Get().FindConsoleVariable(TEXT("r.EyeAdaptation.PreExposureOverride"));
+				check(CVarPreExposureOverride);
+
+				FString Text = FString::Printf(TEXT("PreExposure: %f"), View.PreExposure);
+				if (CVarPreExposureOverride->GetFloat() > 0.0f)
+				{
+					Text += TEXT(" (r.EyeAdaptation.PreExposureOverride)");
+				}
+
+				QuickDrawSummary(/* Location = */ 3, Text);
+			}
+
+			// Display alpha support
+			{
+				FString Text = bSupportsAlpha ? TEXT("r.PostProcessing.PropagateAlpha=true") : TEXT("r.PostProcessing.PropagateAlpha=false");
+				if (Inputs.TAAConfig == EMainTAAPassConfig::ThirdParty)
+				{
+					Text = TEXT("Unknown");
+				}
+				QuickDrawSummary(/* Location = */ 4, TEXT("Support Alpha: ") + Text);
+			}
+		});
+
+	}
+
+	return MoveTemp(Output);
+}
diff --git a/Engine/Source/Runtime/Renderer/Private/VelocityRendering.cpp b/Engine/Source/Runtime/Renderer/Private/VelocityRendering.cpp
--- a/Engine/Source/Runtime/Renderer/Private/VelocityRendering.cpp
+++ b/Engine/Source/Runtime/Renderer/Private/VelocityRendering.cpp
@@ -197,7 +197,7 @@ bool FDeferredShadingSceneRenderer::ShouldRenderVelocities() const
 
 		bool bTemporalAA = IsTemporalAccumulationBasedMethod(View.AntiAliasingMethod) && !View.bCameraCut;
 		bool bMotionBlur = IsMotionBlurEnabled(View);
-		bool bVisualizeMotionblur = View.Family->EngineShowFlags.VisualizeMotionBlur;
+		bool bVisualizeMotionblur = View.Family->EngineShowFlags.VisualizeMotionBlur || View.Family->EngineShowFlags.VisualizeTemporalUpscaler;
 		bool bDistanceFieldAO = ShouldPrepareForDistanceFieldAO();
 
 		bool bSceneSSREnabled = ViewPipelineState.ReflectionsMethod == EReflectionsMethod::SSR && ScreenSpaceRayTracing::ShouldRenderScreenSpaceReflections(View);
