Separate name for multisampled textures.

- A minor QoL improvement for captures.

diff --git a/Engine/Source/Runtime/RenderCore/Private/RenderGraphUtils.cpp b/Engine/Source/Runtime/RenderCore/Private/RenderGraphUtils.cpp
--- a/Engine/Source/Runtime/RenderCore/Private/RenderGraphUtils.cpp
+++ b/Engine/Source/Runtime/RenderCore/Private/RenderGraphUtils.cpp
@@ -190,20 +190,15 @@ FRDGTextureRef RegisterExternalTextureWithFallback(
 RENDERCORE_API FRDGTextureMSAA CreateTextureMSAA(
 	FRDGBuilder& GraphBuilder,
 	FRDGTextureDesc Desc,
-	const TCHAR* Name,
+	const TCHAR* NameMultisampled, const TCHAR* NameResolved,
 	ETextureCreateFlags ResolveFlagsToAdd)
 {
-	bool bForceSeparateTargetAndShaderResource = Desc.NumSamples > 1 && RHISupportsSeparateMSAAAndResolveTextures(GMaxRHIShaderPlatform);
-	
-	if (!bForceSeparateTargetAndShaderResource)
-	{
-		Desc.Flags |= TexCreate_ShaderResource;
-	}
+	const bool bForceSeparateTargetAndShaderResource = Desc.NumSamples > 1 && RHISupportsSeparateMSAAAndResolveTextures(GMaxRHIShaderPlatform);
 
-	FRDGTextureMSAA Texture(GraphBuilder.CreateTexture(Desc, Name));
-
-	if (bForceSeparateTargetAndShaderResource)
+	if (LIKELY(bForceSeparateTargetAndShaderResource))
 	{
+		FRDGTextureMSAA Texture(GraphBuilder.CreateTexture(Desc, NameMultisampled));
+
 		Desc.NumSamples = 1;
 		ETextureCreateFlags ResolveFlags = TexCreate_ShaderResource;
 		if (EnumHasAnyFlags(Desc.Flags, TexCreate_DepthStencilTargetable))
@@ -217,10 +212,13 @@ RENDERCORE_API FRDGTextureMSAA CreateTextureMSAA(
         ResolveFlags &= ~(TexCreate_Memoryless);
         
 		Desc.Flags = ResolveFlags | ResolveFlagsToAdd;
-		Texture.Resolve = GraphBuilder.CreateTexture(Desc, Name);
+		Texture.Resolve = GraphBuilder.CreateTexture(Desc, NameResolved);
+
+		return Texture;
 	}
 
-	return Texture;
+	Desc.Flags |= TexCreate_ShaderResource;
+	return FRDGTextureMSAA(GraphBuilder.CreateTexture(Desc, NameResolved));
 }
 
 BEGIN_SHADER_PARAMETER_STRUCT(FCopyTextureParameters, )
diff --git a/Engine/Source/Runtime/RenderCore/Public/RenderGraphUtils.h b/Engine/Source/Runtime/RenderCore/Public/RenderGraphUtils.h
--- a/Engine/Source/Runtime/RenderCore/Public/RenderGraphUtils.h
+++ b/Engine/Source/Runtime/RenderCore/Public/RenderGraphUtils.h
@@ -320,9 +320,19 @@ struct FRDGTextureMSAA
 RENDERCORE_API FRDGTextureMSAA CreateTextureMSAA(
 	FRDGBuilder& GraphBuilder,
 	FRDGTextureDesc Desc,
-	const TCHAR* Name,
+	const TCHAR* NameMultisampled, const TCHAR* NameResolved,
 	ETextureCreateFlags ResolveFlagsToAdd = TexCreate_None);
 
+UE_DEPRECATED(5.3, "CreateTextureMSAA with one name is deprecated, prease provide separate name for the multisampled texture")
+inline FRDGTextureMSAA CreateTextureMSAA(
+	FRDGBuilder& GraphBuilder,
+	FRDGTextureDesc Desc,
+	const TCHAR* Name,
+	ETextureCreateFlags ResolveFlagsToAdd = TexCreate_None)
+{
+	return CreateTextureMSAA(GraphBuilder, Desc, Name, Name, ResolveFlagsToAdd);
+}
+
 /** All utils for compute shaders.
  */
 struct RENDERCORE_API FComputeShaderUtils
diff --git a/Engine/Source/Runtime/Renderer/Private/LightGridInjection.cpp b/Engine/Source/Runtime/Renderer/Private/LightGridInjection.cpp
--- a/Engine/Source/Runtime/Renderer/Private/LightGridInjection.cpp
+++ b/Engine/Source/Runtime/Renderer/Private/LightGridInjection.cpp
@@ -980,11 +980,11 @@ void FDeferredShadingSceneRenderer::RenderForwardShadowProjections(
 		{
 			FRDGTextureDesc Desc(FRDGTextureDesc::Create2D(SceneTextures.Config.Extent, PF_B8G8R8A8, FClearValueBinding::White, TexCreate_RenderTargetable | TexCreate_ShaderResource));
 			Desc.NumSamples = SceneDepthTexture->Desc.NumSamples;
-			ForwardScreenSpaceShadowMask = CreateTextureMSAA(GraphBuilder, Desc, TEXT("ShadowMaskTexture"), GFastVRamConfig.ScreenSpaceShadowMask);
+			ForwardScreenSpaceShadowMask = CreateTextureMSAA(GraphBuilder, Desc, TEXT("ShadowMaskTextureMS"), TEXT("ShadowMaskTexture"), GFastVRamConfig.ScreenSpaceShadowMask);
 			if (bIsHairEnable)
 			{
 				Desc.NumSamples = 1;
-				ForwardScreenSpaceShadowMaskSubPixel = CreateTextureMSAA(GraphBuilder, Desc, TEXT("ShadowMaskSubPixelTexture"), GFastVRamConfig.ScreenSpaceShadowMask);
+				ForwardScreenSpaceShadowMaskSubPixel = CreateTextureMSAA(GraphBuilder, Desc, TEXT("ShadowMaskSubPixelTextureMS"), TEXT("ShadowMaskSubPixelTexture"), GFastVRamConfig.ScreenSpaceShadowMask);
 			}
 		}
 
diff --git a/Engine/Source/Runtime/Renderer/Private/SceneTextures.cpp b/Engine/Source/Runtime/Renderer/Private/SceneTextures.cpp
--- a/Engine/Source/Runtime/Renderer/Private/SceneTextures.cpp
+++ b/Engine/Source/Runtime/Renderer/Private/SceneTextures.cpp
@@ -481,15 +481,13 @@ void FMinimalSceneTextures::InitializeViewFamily(FRDGBuilder& GraphBuilder, FVie
 		const bool bIsMobilePlatform = Config.ShadingPath == EShadingPath::Mobile;
 		const ETextureCreateFlags sRGBFlag = (bIsMobilePlatform && IsMobileColorsRGB()) ? TexCreate_SRGB : TexCreate_None;
 
-		const TCHAR* SceneColorName = TEXT("SceneColor");
-
 		// Create the scene color.
 		// TODO: Array-size could be values > 2, on upcoming XR devices. This should be part of the config.
 		FRDGTextureDesc Desc(Config.bRequireMultiView ?
 							 FRDGTextureDesc::Create2DArray(Config.Extent, Config.ColorFormat, Config.ColorClearValue, Config.ColorCreateFlags, 2) :
 							 FRDGTextureDesc::Create2D(Config.Extent, Config.ColorFormat, Config.ColorClearValue, Config.ColorCreateFlags));
 		Desc.NumSamples = Config.NumSamples;
-		SceneTextures.Color = CreateTextureMSAA(GraphBuilder, Desc, SceneColorName, GFastVRamConfig.SceneColor | sRGBFlag);
+		SceneTextures.Color = CreateTextureMSAA(GraphBuilder, Desc, TEXT("SceneColorMS"), TEXT("SceneColor"), GFastVRamConfig.SceneColor | sRGBFlag);
 	}
 
 	// Custom Depth
@@ -605,7 +603,7 @@ void FSceneTextures::InitializeViewFamily(FRDGBuilder& GraphBuilder, FViewFamily
 			FRDGTextureDesc::Create2DArray(Config.Extent, DepthAuxFormat, FClearValueBinding(FarDepthColor), TexCreate_RenderTargetable | TexCreate_ShaderResource | TexCreate_InputAttachmentRead | MemorylessFlag, 2) :
 			FRDGTextureDesc::Create2D(Config.Extent, DepthAuxFormat, FClearValueBinding(FarDepthColor), TexCreate_RenderTargetable | TexCreate_ShaderResource | TexCreate_InputAttachmentRead| MemorylessFlag);
 		Desc.NumSamples = Config.NumSamples;
-		SceneTextures.DepthAux = CreateTextureMSAA(GraphBuilder, Desc, TEXT("SceneDepthAux"));
+		SceneTextures.DepthAux = CreateTextureMSAA(GraphBuilder, Desc, TEXT("SceneDepthAuxMS"), TEXT("SceneDepthAux"));
 	}
 #if WITH_EDITOR
 	{
diff --git a/Engine/Source/Runtime/Renderer/Private/TranslucentRendering.cpp b/Engine/Source/Runtime/Renderer/Private/TranslucentRendering.cpp
--- a/Engine/Source/Runtime/Renderer/Private/TranslucentRendering.cpp
+++ b/Engine/Source/Runtime/Renderer/Private/TranslucentRendering.cpp
@@ -148,6 +148,16 @@ static const TCHAR* kTranslucencyColorTextureName[] = {
 };
 static_assert(UE_ARRAY_COUNT(kTranslucencyColorTextureName) == int32(ETranslucencyPass::TPT_MAX), "Fix me");
 
+static const TCHAR* kTranslucencyColorTextureMultisampledName[] = {
+	TEXT("Translucency.BeforeDistortion.ColorMS"),
+	TEXT("Translucency.BeforeDistortion.ModulateMS"),
+	TEXT("Translucency.AfterDOF.ColorMS"),
+	TEXT("Translucency.AfterDOF.ModulateMS"),
+	TEXT("Translucency.AfterMotionBlur.ColorMS"),
+	TEXT("Translucency.All.ColorMS"),
+};
+static_assert(UE_ARRAY_COUNT(kTranslucencyColorTextureMultisampledName) == UE_ARRAY_COUNT(kTranslucencyColorTextureName), "Fix me");
+
 static const TCHAR* TranslucencyPassToString(ETranslucencyPass::Type TranslucencyPass)
 {
 	return kTranslucencyPassName[TranslucencyPass];
@@ -756,6 +766,7 @@ FRDGTextureMSAA CreatePostDOFTranslucentTexture(
 	const FRDGTextureDesc Desc = GetPostDOFTranslucentTextureDesc(TranslucencyPass, SeparateTranslucencyDimensions, bIsModulate, ShaderPlatform);
 	return CreateTextureMSAA(
 		GraphBuilder, Desc,
+		kTranslucencyColorTextureMultisampledName[int32(TranslucencyPass)],
 		kTranslucencyColorTextureName[int32(TranslucencyPass)],
 		bIsModulate ? GFastVRamConfig.SeparateTranslucencyModulate : GFastVRamConfig.SeparateTranslucency);
 }
@@ -1530,6 +1541,7 @@ void FDeferredShadingSceneRenderer::RenderTranslucency(
 
 		SharedDepthTexture = CreateTextureMSAA(
 			GraphBuilder, Desc,
+			TEXT("Translucency.DepthMS"),
 			TEXT("Translucency.Depth"),
 			GFastVRamConfig.SeparateTranslucencyModulate); // TODO: this should be SeparateTranslucency, but is what the code was doing
 
