Fix a few minor issues and clean up some future landmines
Remove bIncludeInScreenSpaceShadowMask as it is no longer used

diff --git a/Engine/Source/Runtime/Renderer/Private/SceneRendering.cpp b/Engine/Source/Runtime/Renderer/Private/SceneRendering.cpp
--- a/Engine/Source/Runtime/Renderer/Private/SceneRendering.cpp
+++ b/Engine/Source/Runtime/Renderer/Private/SceneRendering.cpp
@@ -490,8 +490,7 @@ bool FVisibleLightInfo::ContainsOnlyVirtualShadowMaps() const
 	{
 		// Simple test for now, but sufficient
 		const FProjectedShadowInfo* ProjectedShadowInfo = AllProjectedShadows[ShadowIndex];
-		if (ProjectedShadowInfo->bIncludeInScreenSpaceShadowMask && ProjectedShadowInfo->bAllocated &&
-			!(ProjectedShadowInfo->HasVirtualShadowMap() || ProjectedShadowInfo->VirtualShadowMapClipmap))
+		if (ProjectedShadowInfo->bAllocated && !ProjectedShadowInfo->HasVirtualShadowMap())
 		{
 			return false;
 		}
diff --git a/Engine/Source/Runtime/Renderer/Private/ShadowRendering.cpp b/Engine/Source/Runtime/Renderer/Private/ShadowRendering.cpp
--- a/Engine/Source/Runtime/Renderer/Private/ShadowRendering.cpp
+++ b/Engine/Source/Runtime/Renderer/Private/ShadowRendering.cpp
@@ -2100,11 +2100,6 @@ void FSceneRenderer::RenderShadowProjections(
 	for (int32 ShadowIndex = 0; ShadowIndex < VisibleLightInfo.ShadowsToProject.Num(); ShadowIndex++)
 	{
 		FProjectedShadowInfo* ProjectedShadowInfo = VisibleLightInfo.ShadowsToProject[ShadowIndex];
-		if (!ProjectedShadowInfo->bIncludeInScreenSpaceShadowMask)
-		{
-			continue;
-		}
-
 		if (ProjectedShadowInfo->bRayTracedDistanceField)
 		{
 			DistanceFieldShadows.Add(ProjectedShadowInfo);
@@ -2208,8 +2203,6 @@ void FSceneRenderer::BeginAsyncDistanceFieldShadowProjections(FRDGBuilder& Graph
 			{
 				FProjectedShadowInfo* ProjectedShadowInfo = ProjectedDistanceFieldShadows[DFShadowIndex];
 
-				check(ProjectedShadowInfo->bIncludeInScreenSpaceShadowMask);
-
 				FIntRect ScissorRect;
 				if (ProjectedShadowInfo->bDirectionalLight || !ProjectedShadowInfo->GetLightSceneInfo().Proxy->GetScissorRect(ScissorRect, View, View.ViewRect))
 				{
@@ -2320,10 +2313,6 @@ void FDeferredShadingSceneRenderer::RenderDeferredShadowProjections(
 			for (int32 ShadowIndex = 0; ShadowIndex < VisibleLightInfo.ShadowsToProject.Num(); ShadowIndex++)
 			{
 				FProjectedShadowInfo* ProjectedShadowInfo = VisibleLightInfo.ShadowsToProject[ShadowIndex];
-				if (!ProjectedShadowInfo->bIncludeInScreenSpaceShadowMask)
-				{
-					continue;
-				}
 				if (ProjectedShadowInfo->HasVirtualShadowMap())
 				{
 					bNeedHairShadowMaskPass = false;
diff --git a/Engine/Source/Runtime/Renderer/Private/ShadowRendering.h b/Engine/Source/Runtime/Renderer/Private/ShadowRendering.h
--- a/Engine/Source/Runtime/Renderer/Private/ShadowRendering.h
+++ b/Engine/Source/Runtime/Renderer/Private/ShadowRendering.h
@@ -427,11 +427,6 @@ public:
 	/** Whether to render Nanite geometry into this shadow map. */
 	uint32 bNaniteGeometry : 1;
 
-	/** Whether to include this shadow in the direct lighting shadow mask.
-	 * Currently this is used to disable projection of some low-resolution proxy shadow maps when using other high resolution virtual shadow maps
-	 * for direct lighting visibility. */
-	uint32 bIncludeInScreenSpaceShadowMask : 1;
-
 	/** Whether the the shadow overlaps any nanite primitives */
 	uint32 bContainsNaniteSubjects : 1;
 	uint32 bShouldRenderVSM : 1;
@@ -773,7 +768,7 @@ public:
 		return FShadowDepthType(bDirectionalLight, bOnePassPointLightShadow);
 	}
 
-	bool HasVirtualShadowMap() const { return VirtualShadowMaps.Num() > 0; }
+	bool HasVirtualShadowMap() const { return VirtualShadowMaps.Num() > 0 || VirtualShadowMapClipmap.IsValid(); }
 
 	FParallelMeshDrawCommandPass& GetShadowDepthPass() { return ShadowDepthPass; }
 
diff --git a/Engine/Source/Runtime/Renderer/Private/ShadowSetup.cpp b/Engine/Source/Runtime/Renderer/Private/ShadowSetup.cpp
--- a/Engine/Source/Runtime/Renderer/Private/ShadowSetup.cpp
+++ b/Engine/Source/Runtime/Renderer/Private/ShadowSetup.cpp
@@ -731,7 +731,6 @@ FProjectedShadowInfo::FProjectedShadowInfo()
 	, bTransmission(false)
 	, bHairStrandsDeepShadow(false)
 	, bNaniteGeometry(true)
-	, bIncludeInScreenSpaceShadowMask(true)
 	, bContainsNaniteSubjects(false)
 	, bShouldRenderVSM(true)
 	, PerObjectShadowFadeStart(UE_OLD_WORLD_MAX)		// LWC_TODO: Upgrade to double? Should have been HALF_WORLD_MAX to match precision.
@@ -1087,7 +1086,6 @@ void FProjectedShadowInfo::SetupClipmapProjection(FLightSceneInfo* InLightSceneI
 	bDirectionalLight = true;
 	bWholeSceneShadow = true;
 	BorderSize = 0;
-	bIncludeInScreenSpaceShadowMask = false;
 	MaxNonFarCascadeDistance = InMaxNonFarCascadeDistance;
 	MeshPassTargetType = EMeshPass::VSMShadowDepth;
 	MeshSelectionMask = EShadowMeshSelection::VSM;
@@ -2317,7 +2315,7 @@ void FProjectedShadowInfo::SetupMeshDrawCommandsForShadowDepth(FSceneRenderer& R
 	const uint32 InstanceFactor = bMayUseHostCubeFaceReplication ? 6 : 1;
 
 	// Ensure all work goes down the one path to simplify processing
-	EBatchProcessingMode SingleInstanceProcessingMode = (HasVirtualShadowMap() || VirtualShadowMapClipmap.IsValid()) ? EBatchProcessingMode::Generic : EBatchProcessingMode::UnCulled;
+	EBatchProcessingMode SingleInstanceProcessingMode = HasVirtualShadowMap() ? EBatchProcessingMode::Generic : EBatchProcessingMode::UnCulled;
 
 	ShadowDepthPass.DispatchPassSetup(
 		Renderer.Scene,
@@ -6010,7 +6008,7 @@ void FSceneRenderer::FilterDynamicShadows(FDynamicShadowsTaskData& TaskData)
 					// Certain shadow types skip the depth pass when there are no movable primitives to composite.
 					bool bAlwaysHasDepthPass = true;
 
-					if (ProjectedShadowInfo->HasVirtualShadowMap() || ProjectedShadowInfo->VirtualShadowMapClipmap.IsValid())
+					if (ProjectedShadowInfo->HasVirtualShadowMap())
 					{
 						SortedShadowsForShadowDepthPass.VirtualShadowMapShadows.Add(ProjectedShadowInfo);
 					}
diff --git a/Engine/Source/Runtime/Renderer/Private/VirtualShadowMaps/VirtualShadowMapArray.cpp b/Engine/Source/Runtime/Renderer/Private/VirtualShadowMaps/VirtualShadowMapArray.cpp
--- a/Engine/Source/Runtime/Renderer/Private/VirtualShadowMaps/VirtualShadowMapArray.cpp
+++ b/Engine/Source/Runtime/Renderer/Private/VirtualShadowMaps/VirtualShadowMapArray.cpp
@@ -1025,22 +1025,6 @@ void FVirtualShadowMapArray::MergeStaticPhysicalPages(FRDGBuilder& GraphBuilder)
 	}
 }
 
-/**
- * Helper to get hold of / check for associated virtual shadow map
- */
-FORCEINLINE FProjectedShadowInfo* GetVirtualShadowMapInfo(const FVisibleLightInfo &LightInfo)
-{
-	for (FProjectedShadowInfo *ProjectedShadowInfo : LightInfo.AllProjectedShadows)
-	{
-		if (ProjectedShadowInfo->HasVirtualShadowMap())
-		{
-			return ProjectedShadowInfo;
-		}
-	}
-
-	return nullptr;
-}
-
 
 class FInitPageRectBoundsCS : public FVirtualPageManagementShader
 {
@@ -1261,7 +1245,7 @@ void FVirtualShadowMapArray::BuildPageAllocations(
 
 		for (FProjectedShadowInfo* ProjectedShadowInfo : VisibleLightInfo.AllProjectedShadows)
 		{
-			if (ProjectedShadowInfo->HasVirtualShadowMap())
+			if (ProjectedShadowInfo->VirtualShadowMaps.Num() > 0)
 			{
 				const float ResolutionLodBiasLocal = GetResolutionLODBiasLocal(GetLightMobilityFactor(LightId));
 				check(ProjectedShadowInfo->CascadeSettings.ShadowSplitIndex == INDEX_NONE);		// We use clipmaps for virtual shadow maps, not cascades
