Fixexs spelling mistakes in TSR cvar

diff --git a/Engine/Source/Runtime/Renderer/Private/PostProcess/TemporalSuperResolution.cpp b/Engine/Source/Runtime/Renderer/Private/PostProcess/TemporalSuperResolution.cpp
--- a/Engine/Source/Runtime/Renderer/Private/PostProcess/TemporalSuperResolution.cpp
+++ b/Engine/Source/Runtime/Renderer/Private/PostProcess/TemporalSuperResolution.cpp
@@ -139,11 +139,11 @@ TAutoConsoleVariable<float> CVarTSRFlickeringPeriod(
 	ECVF_RenderThreadSafe);
 
 TAutoConsoleVariable<float> CVarTSRFlickeringMaxParralaxVelocity(
-	TEXT("r.TSR.ShadingRejection.Flickering.MaxParralaxVelocity"), 10.0,
+	TEXT("r.TSR.ShadingRejection.Flickering.MaxParallaxVelocity"), 10.0,
 	TEXT("Some material might for instance might do something like parallax occlusion mapping such as CitySample's buildings' window's interiors. ")
 	TEXT("This often can not render accurately a motion vector of this fake interior geometry and therefore make the heuristic believe it is in fact flickering.\n")
 	TEXT("\n")
-	TEXT("This variable define the parralax velocity in 1080p pixel at frame rate defined by r.TSR.ShadingRejection.Flickering.FrameRateCap at which point the ")
+	TEXT("This variable define the parallax velocity in 1080p pixel at frame rate defined by r.TSR.ShadingRejection.Flickering.FrameRateCap at which point the ")
 	TEXT("heuristic should be disabled to not ghost. ")
 	TEXT("\n")
 	TEXT("(Default to 10 pixels 1080p).\n"),
@@ -229,7 +229,7 @@ TAutoConsoleVariable<int32> CVarTSRSubpixelMethod(
 	ECVF_RenderThreadSafe);
 
 TAutoConsoleVariable<int32> CVarTSRSubpixelDepthMaxAge(
-	TEXT("r.TSR.Subpixel.DephtMaxAge"), 3,
+	TEXT("r.TSR.Subpixel.DepthMaxAge"), 3,
 	TEXT("Maximum age in frames of subpixel's depth kept in history for their self reprojection (default to 3 frames)."),
 	ECVF_RenderThreadSafe);
 
