Fixes a bug where MotionBlur.Visualize could end up with a mip chain causing corruptions in VisualizeTemporalUpscaler

diff --git a/Engine/Source/Runtime/Renderer/Private/PostProcess/PostProcessMotionBlur.cpp b/Engine/Source/Runtime/Renderer/Private/PostProcess/PostProcessMotionBlur.cpp
--- a/Engine/Source/Runtime/Renderer/Private/PostProcess/PostProcessMotionBlur.cpp
+++ b/Engine/Source/Runtime/Renderer/Private/PostProcess/PostProcessMotionBlur.cpp
@@ -1107,7 +1107,11 @@ FScreenPassTexture AddVisualizeMotionBlurPass(FRDGBuilder& GraphBuilder, const F
 
 	if (!Output.IsValid())
 	{
-		Output = FScreenPassRenderTarget::CreateFromInput(GraphBuilder, Inputs.SceneColor, View.GetOverwriteLoadAction(), TEXT("MotionBlur.Visualize"));
+		FRDGTextureDesc OutputDesc = Inputs.SceneColor.Texture->Desc;
+		OutputDesc.Reset();
+		OutputDesc.NumMips = 1;
+
+		Output = FScreenPassRenderTarget(GraphBuilder.CreateTexture(OutputDesc, TEXT("MotionBlur.Visualize")), Inputs.SceneColor.ViewRect, View.GetOverwriteLoadAction());
 	}
 
 	// NOTE: Scene depth is used as the velocity viewport because velocity can actually be a 1x1 black texture.
@@ -1128,11 +1132,11 @@ FScreenPassTexture AddVisualizeMotionBlurPass(FRDGBuilder& GraphBuilder, const F
 
 	TShaderMapRef<FMotionBlurVisualizePS> PixelShader(View.ShaderMap);
 
-	AddDrawScreenPass(GraphBuilder, RDG_EVENT_NAME("Visualizer"), View, Viewports.Color, Viewports.Color, PixelShader, PassParameters);
+	AddDrawScreenPass(GraphBuilder, RDG_EVENT_NAME("VisualizeMotionBlur"), View, Viewports.Color, Viewports.Color, PixelShader, PassParameters);
 
 	Output.LoadAction = ERenderTargetLoadAction::ELoad;
 
-	AddDrawCanvasPass(GraphBuilder, RDG_EVENT_NAME("Overlay"), View, Output,
+	AddDrawCanvasPass(GraphBuilder, RDG_EVENT_NAME("VisualizeMotionBlurOverlay"), View, Output,
 		[&View](FCanvas& Canvas)
 	{
 		float X = 20;
