Improve VRS debug functionality in a variety of ways:
- Debug overlay no longer disappears when no generator is enabled, instead displaying all red (1x1 rate)
- Debug overlay now correctly displays conservative CAS rate when `r.VRS.ContrastAdaptiveShadingRate.PreviewType` is set to 2, and unreprojected rate when set to 3
- Force rate image generation now resides in the manager, and overrides alternative image generation
- Force rate now supports coarse shading rates (2x4, 4x2, 4x4)
- Force rate rates are now indexed as 0 through 6 rather than using the shading rate actual bit mask
- Force rate console command renamed from `r.VRS.ContrastAdaptiveShading.Debug.ForceRate` to `r.VRS.DebugForceRate`
- Removed unneeded force rate functionality from the CAS VRSShadingRateCalculate.usf shader
- Fixed off-by-one error in SRI texture size calculation and moved all calculation into one helper function in the manager

diff --git a/Engine/Shaders/Private/VariableRateShading/VRSShadingRateCalculate.usf b/Engine/Shaders/Private/VariableRateShading/VRSShadingRateCalculate.usf
--- a/Engine/Shaders/Private/VariableRateShading/VRSShadingRateCalculate.usf
+++ b/Engine/Shaders/Private/VariableRateShading/VRSShadingRateCalculate.usf
@@ -160,28 +160,7 @@ void CalculateShadingRateImage(
 	uint GroupIndex : SV_GroupIndex)
 {
 	int2 TextureCoord = int2(int(ViewRect.x + DispatchThreadId.x), int(ViewRect.y + DispatchThreadId.y));
-#if FORCE_RATE != -1
-	const float2 VRSTileSize = float2(THREADGROUP_SIZEX, THREADGROUP_SIZEY);
-	uint2 MaxSize = ceil(ViewRect.zw / VRSTileSize);
 
-	if (all(DispatchThreadId.xy < MaxSize))
-	{
-		uint DebugRate =
-	#if FORCE_RATE == 0
-			D3D12_SHADING_RATE_1X1;
-	#elif FORCE_RATE == 1
-			D3D12_SHADING_RATE_1X2;
-	#elif FORCE_RATE == 4
-			D3D12_SHADING_RATE_2X1;
-	#else // FORCE_RATE == 5
-			D3D12_SHADING_RATE_2X2;
-	#endif
-
-		VariableRateShadingTexture[TextureCoord] = 
-			(DebugRate << CONSERVATIVE_SHADING_RATE_SHIFT) |
-			DebugRate;
-	}
-#else
 	CalculateShadingRateImageInner(
 		DispatchThreadId.xy, 
 		GroupThreadId.xy, 
@@ -189,7 +168,6 @@ void CalculateShadingRateImage(
 		GroupId, 
 		GroupIndex,
 		LuminanceTexture[TextureCoord]);
-#endif
 }
 
 #endif /* COMPUTE_SHADER */
\ No newline at end of file
diff --git a/Engine/Source/Runtime/Renderer/Private/VariableRateShading/ContrastAdaptiveImageGenerator.cpp b/Engine/Source/Runtime/Renderer/Private/VariableRateShading/ContrastAdaptiveImageGenerator.cpp
--- a/Engine/Source/Runtime/Renderer/Private/VariableRateShading/ContrastAdaptiveImageGenerator.cpp
+++ b/Engine/Source/Runtime/Renderer/Private/VariableRateShading/ContrastAdaptiveImageGenerator.cpp
@@ -125,17 +125,11 @@ TAutoConsoleVariable<int32> CVarVRSDecals(
  * Debug Settings 
  */
 
-int GVRSDebugForceRate = -1;
-FAutoConsoleVariableRef CVarVRSDebugForceRate(
-	TEXT("r.VRS.ContrastAdaptiveShading.Debug.ForceRate"),
-	GVRSDebugForceRate,
-	TEXT("-1 : None, 0 : Force 1x1, 1 : Force 1x2, 4 : Force 2x1, 5: Force 2x2"));
-
 TAutoConsoleVariable<int32> CVarCASPreviewType(
 	TEXT("r.VRS.ContrastAdaptiveShading.PreviewType"),
 	1,
 	TEXT("Include CAS in the VRS debug overlay.")
-	TEXT("0 - off, 1 - the SRI texture, 2- the conservative SRI texture"),
+	TEXT("0 - off, 1 - the SRI texture, 2- the conservative SRI texture, 3 - the pre-reprojection SRI texture"),
 	ECVF_RenderThreadSafe);
 
 
@@ -152,9 +146,8 @@ class FCalculateShadingRateImageCS : public FGlobalShader
 
 	class FThreadGroupX : SHADER_PERMUTATION_SPARSE_INT("THREADGROUP_SIZEX", 8, 16);
 	class FThreadGroupY : SHADER_PERMUTATION_SPARSE_INT("THREADGROUP_SIZEY", 8, 16);
-	class FForceRate : SHADER_PERMUTATION_SPARSE_INT("FORCE_RATE", -1, 0, 1, 4, 5);
 
-	using FPermutationDomain = TShaderPermutationDomain<FThreadGroupX, FThreadGroupY, FForceRate>;
+	using FPermutationDomain = TShaderPermutationDomain<FThreadGroupX, FThreadGroupY>;
 
 	BEGIN_SHADER_PARAMETER_STRUCT(FParameters, )
 		SHADER_PARAMETER_RDG_TEXTURE(Texture2D<float>, LuminanceTexture)
@@ -380,29 +373,27 @@ struct RENDERER_API FVRSTextures
 private:
 	static FRDGTextureDesc CreateSRIDesc(const FSceneViewFamily& ViewFamily, bool bIsForDynResScaled)
 	{
-		FIntPoint TileSize = FVariableRateShadingImageManager::GetSRITileSize();
-
-		FIntPoint ViewTargetExtents = FIntPoint::ZeroValue;
 		if (bIsForDynResScaled)
 		{
-			ViewTargetExtents = FSceneTexturesConfig::Get().Extent;
+			// Use SceneTextures
+			return FVariableRateShadingImageManager::GetSRIDesc();
 		}
 		else
 		{
 			// Get initial size based on luminance texture from previous frame
 			check(ViewFamily.Views[0]->bIsViewInfo);
 			const FViewInfo* ViewInfo = static_cast<const FViewInfo*>(ViewFamily.Views[0]);
-			ViewTargetExtents = ViewInfo->PrevViewInfo.LuminanceHistory->GetDesc().Extent;
+			const FIntPoint ViewTargetExtents = ViewInfo->PrevViewInfo.LuminanceHistory->GetDesc().Extent;
+
+			const FIntPoint SRIDimensions = FMath::DivideAndRoundUp(ViewTargetExtents, FVariableRateShadingImageManager::GetSRITileSize());
+			return FRDGTextureDesc::Create2D(
+				SRIDimensions,
+				GRHIVariableRateShadingImageFormat,
+				EClearBinding::ENoneBound,
+				ETextureCreateFlags::DisableDCC |
+				ETextureCreateFlags::ShaderResource |
+				ETextureCreateFlags::UAV);
 		}
-
-		FIntPoint SRIDimensions = FMath::DivideAndRoundUp(ViewTargetExtents, TileSize);
-		return FRDGTextureDesc::Create2D(
-			SRIDimensions,
-			GRHIVariableRateShadingImageFormat,
-			EClearBinding::ENoneBound,
-			ETextureCreateFlags::DisableDCC |
-			ETextureCreateFlags::ShaderResource |
-			ETextureCreateFlags::UAV);
 	}
 };
 RDG_REGISTER_BLACKBOARD_STRUCT(FVRSTextures);
@@ -454,18 +445,6 @@ bool AddCreateShadingRateImagePass(
 		PermutationVector.Set<FCalculateShadingRateImageCS::FThreadGroupX>(TileSize.X);
 		PermutationVector.Set<FCalculateShadingRateImageCS::FThreadGroupY>(TileSize.Y);
 
-		// Set an override rate if we're in a debug mode
-		int32 ForceRate = GVRSDebugForceRate;
-		if (ForceRate != 0 &&
-			ForceRate != 1 &&
-			ForceRate != 4 &&
-			ForceRate != 5)
-		{
-			ForceRate = -1;
-		}
-
-		PermutationVector.Set<FCalculateShadingRateImageCS::FForceRate>(ForceRate);
-
 		TShaderMapRef<FCalculateShadingRateImageCS> ComputeShader(View.ShaderMap, PermutationVector);
 		auto* PassParameters = GraphBuilder.AllocParameters<FCalculateShadingRateImageCS::FParameters>();
 		EDisplayOutputFormat OutputDisplayFormat = GetDisplayOutputFormat(View);
@@ -640,12 +619,15 @@ FRDGTextureRef FContrastAdaptiveImageGenerator::GetDebugImage(FRDGBuilder& Graph
 
 	switch (PreviewType)
 	{
-	case ESRIPreviewType::BeforeReprojection:
-		return VRSTextures.ConstructedSRI;
-		break;
 	case ESRIPreviewType::Projected:
 		return VRSTextures.ScaledSRI;
 		break;
+	case ESRIPreviewType::ProjectedConservative:
+		return VRSTextures.ScaledConservativeSRI;
+		break;
+	case ESRIPreviewType::BeforeReprojection:
+		return VRSTextures.ConstructedSRI;
+		break;
 	}
 
 	return nullptr;
diff --git a/Engine/Source/Runtime/Renderer/Private/VariableRateShading/FoveatedImageGenerator.cpp b/Engine/Source/Runtime/Renderer/Private/VariableRateShading/FoveatedImageGenerator.cpp
--- a/Engine/Source/Runtime/Renderer/Private/VariableRateShading/FoveatedImageGenerator.cpp
+++ b/Engine/Source/Runtime/Renderer/Private/VariableRateShading/FoveatedImageGenerator.cpp
@@ -191,17 +191,8 @@ void FFoveatedImageGenerator::PrepareImages(FRDGBuilder& GraphBuilder, const FSc
 	// Sanity check VRS tile size.
 	check(GRHIVariableRateShadingImageTileMinWidth >= 8 && GRHIVariableRateShadingImageTileMinWidth <= 64 && GRHIVariableRateShadingImageTileMinHeight >= 8 && GRHIVariableRateShadingImageTileMaxHeight <= 64);
 
-	// Calculate texture size
-	const FIntPoint Size = FSceneTexturesConfig::Get().Extent;
-	const FIntPoint TextureSize(Size.X / GRHIVariableRateShadingImageTileMinWidth, Size.Y / GRHIVariableRateShadingImageTileMinHeight);
-
 	// Create texture to hold shading rate image
-	FRDGTextureDesc Desc = FRDGTextureDesc::Create2D(
-		TextureSize,
-		GRHIVariableRateShadingImageFormat,
-		FClearValueBinding::None,
-		TexCreate_Foveation | TexCreate_UAV);
-
+	FRDGTextureDesc Desc = FVariableRateShadingImageManager::GetSRIDesc();
 	FRDGTextureRef ShadingRateTexture = GraphBuilder.CreateTexture(Desc, TEXT("FoveatedShadingRateTexture"));
 
 	// Setup shader parameters and flags
@@ -213,7 +204,7 @@ void FFoveatedImageGenerator::PrepareImages(FRDGBuilder& GraphBuilder, const FSc
 	PassParameters->FoveationFullRateCutoffSquared = FoveationFullRateCutoff * FoveationFullRateCutoff;
 	PassParameters->FoveationHalfRateCutoffSquared = FoveationHalfRateCutoff * FoveationHalfRateCutoff;
 
-	PassParameters->LeftEyeCenterPixelXY = FVector2f(TextureSize.X * FoveationCenterX, TextureSize.Y * FoveationCenterY);
+	PassParameters->LeftEyeCenterPixelXY = FVector2f(Desc.Extent.X * FoveationCenterX, Desc.Extent.Y * FoveationCenterY);
 	PassParameters->RightEyeCenterPixelXY = PassParameters->LeftEyeCenterPixelXY;
 
 	EVRSGenerationFlags GenFlags = EVRSGenerationFlags::None;
@@ -226,16 +217,16 @@ void FFoveatedImageGenerator::PrepareImages(FRDGBuilder& GraphBuilder, const FSc
 
 		// Adjust eyes for side-by-side stereo and/or quadview
 		PassParameters->LeftEyeCenterPixelXY.X /= ViewFamily.Views.Num();;
-		PassParameters->RightEyeCenterPixelXY.X = PassParameters->LeftEyeCenterPixelXY.X + TextureSize.X / ViewFamily.Views.Num();
+		PassParameters->RightEyeCenterPixelXY.X = PassParameters->LeftEyeCenterPixelXY.X + Desc.Extent.X / ViewFamily.Views.Num();
 	}
 
 	// Set up remaining parameters
-	const FVector2f ViewCenterPoint = FVector2f(TextureSize.X / ViewFamily.Views.Num() * 0.5f, TextureSize.Y * 0.5f);
+	const FVector2f ViewCenterPoint = FVector2f(Desc.Extent.X / ViewFamily.Views.Num() * 0.5f, Desc.Extent.Y * 0.5f);
 	PassParameters->ViewDiagonalSquaredInPixels = FVector2f::DotProduct(ViewCenterPoint, ViewCenterPoint);
 	PassParameters->ShadingRateAttachmentGenerationFlags = (uint32)GenFlags;
 
 	TShaderMapRef<FComputeVariableRateShadingImageGeneration> ComputeShader(GetGlobalShaderMap(GMaxRHIFeatureLevel)); // If in stereo, shader sets up a single side-by-side image for both eyes at once
-	FIntVector GroupCount = FComputeShaderUtils::GetGroupCount(TextureSize, FComputeShaderUtils::kGolden2DGroupSize);
+	FIntVector GroupCount = FComputeShaderUtils::GetGroupCount(Desc.Extent, FComputeShaderUtils::kGolden2DGroupSize);
 	GraphBuilder.AddPass(
 		RDG_EVENT_NAME("GenerateFoveatedVRSImage"),
 		PassParameters,
diff --git a/Engine/Source/Runtime/Renderer/Private/VariableRateShading/VariableRateShadingImageManager.cpp b/Engine/Source/Runtime/Renderer/Private/VariableRateShading/VariableRateShadingImageManager.cpp
--- a/Engine/Source/Runtime/Renderer/Private/VariableRateShading/VariableRateShadingImageManager.cpp
+++ b/Engine/Source/Runtime/Renderer/Private/VariableRateShading/VariableRateShadingImageManager.cpp
@@ -34,6 +34,12 @@ TAutoConsoleVariable<int32> CVarVRSPreview(
 	TEXT("0 - off, 1 - on"),
 	ECVF_RenderThreadSafe);
 
+int GVRSDebugForceRate = -1;
+FAutoConsoleVariableRef CVarVRSDebugForceRate(
+	TEXT("r.VRS.DebugForceRate"),
+	GVRSDebugForceRate,
+	TEXT("-1 : None, 0 : Force 1x1, 1 : Force 1x2, 2 : Force 2x1, 3: Force 2x2, 4 : Force 2x4, 5 : Force 4x2, 6 : Force 4x4"));
+
 /**
  * Shaders
  */
@@ -209,49 +215,17 @@ FIntPoint FVariableRateShadingImageManager::GetSRITileSize()
 	return FIntPoint(GRHIVariableRateShadingImageTileMinWidth, GRHIVariableRateShadingImageTileMinHeight);
 }
 
-FRDGTextureRef FVariableRateShadingImageManager::CombineShadingRateImages(FRDGBuilder& GraphBuilder, const FViewInfo& ViewInfo, TArray<FRDGTextureRef> Sources)
+FRDGTextureDesc FVariableRateShadingImageManager::GetSRIDesc()
 {
-	// If we have more than one source, combine the first available two
-	// TODO: Support combining more textures
-	if (Sources.Num() < 1)
-	{
-		return nullptr;
-	}
-	else if (Sources.Num() == 1)
-	{
-		return Sources[0];
-	}
-	else
-	{
-		FIntPoint ViewSize = FSceneTexturesConfig::Get().Extent;
-		const FIntPoint TileSize = GetSRITileSize();
-
-		// Create texture to hold shading rate image
-		FRDGTextureDesc Desc = FRDGTextureDesc::Create2D(
-			ViewSize / TileSize,
-			GRHIVariableRateShadingImageFormat,
-			FClearValueBinding::None,
-			TexCreate_Foveation | TexCreate_UAV);
-
-		FRDGTextureRef CombinedShadingRateTexture = GraphBuilder.CreateTexture(Desc, TEXT("CombinedShadingRateTexture"));
-
-		FCombineShadingRateTexturesCS::FParameters* PassParameters = GraphBuilder.AllocParameters<FCombineShadingRateTexturesCS::FParameters>();
-		PassParameters->SourceTexture0 = Sources[0];
-		PassParameters->SourceTexture1 = Sources[1];
-		PassParameters->RWOutputTexture = GraphBuilder.CreateUAV(CombinedShadingRateTexture);
-
-		TShaderMapRef<FCombineShadingRateTexturesCS> ComputeShader(GetGlobalShaderMap(GMaxRHIFeatureLevel));
-
-		FComputeShaderUtils::AddPass(
-			GraphBuilder,
-			RDG_EVENT_NAME("CombineShadingRateImages"),
-			ERDGPassFlags::AsyncCompute | ERDGPassFlags::NeverCull,
-			ComputeShader,
-			PassParameters,
-			FComputeShaderUtils::GetGroupCount(ViewSize, TileSize));
+	const FIntPoint Size = FSceneTexturesConfig::Get().Extent;
+	const FIntPoint TileSize = GetSRITileSize();
+	const FIntPoint SRISize = FMath::DivideAndRoundUp(Size, TileSize);
 
-		return CombinedShadingRateTexture;
-	}
+	return FRDGTextureDesc::Create2D(
+		SRISize,
+		GRHIVariableRateShadingImageFormat,
+		FClearValueBinding::None,
+		TexCreate_Foveation | TexCreate_UAV | TexCreate_ShaderResource | TexCreate_DisableDCC);
 
 }
 
@@ -264,6 +238,12 @@ FRDGTextureRef FVariableRateShadingImageManager::GetVariableRateShadingImage(FRD
 		return nullptr;
 	}
 
+	// Use debug rate if provided
+	if (CVarVRSDebugForceRate->GetInt() >= 0)
+	{
+		return GetForceRateImage(GraphBuilder, CVarVRSDebugForceRate->GetInt());
+	}
+
 	// Otherwise collate all internal sources
 	TArray<FRDGTextureRef> InternalVRSSources;
 
@@ -293,9 +273,11 @@ FRDGTextureRef FVariableRateShadingImageManager::GetVariableRateShadingImage(FRD
 	{
 		return GraphBuilder.RegisterExternalTexture((*ExternalVRSSources)[0]);
 	}
+
+	// Default to 1x1 shading rate if no sources are available
 	else
 	{
-		return nullptr;
+		return GetForceRateImage(GraphBuilder);
 	}
 }
 
@@ -390,28 +372,37 @@ void FVariableRateShadingImageManager::DrawDebugPreview(FRDGBuilder& GraphBuilde
 		auto ViewInfo = static_cast<const FViewInfo*>(View);
 		if (IsVRSCompatibleWithView(*ViewInfo))
 		{
-			// Collate debug images
-			TArray<FRDGTextureRef> InternalVRSSources;
+			FRDGTextureRef PreviewTexture;
+			
+			// Use debug rate if provided
+			if (CVarVRSDebugForceRate->GetInt() >= 0)
+			{
+				PreviewTexture = GetForceRateImage(GraphBuilder, CVarVRSDebugForceRate->GetInt());
+			}
 
-			for (TUniquePtr<IVariableRateShadingImageGenerator>& Generator : ImageGenerators)
+			// Otherwise collate debug images
+			else
 			{
-				FRDGTextureRef Image = nullptr;
-				if (Generator->IsEnabledForView(*View))
-				{
-					Image = Generator->GetDebugImage(GraphBuilder, *ViewInfo);
-				}
+				
+				TArray<FRDGTextureRef> InternalVRSSources;
 
-				if (Image)
+				for (TUniquePtr<IVariableRateShadingImageGenerator>& Generator : ImageGenerators)
 				{
-					InternalVRSSources.Add(Image);
+					FRDGTextureRef Image = nullptr;
+					if (Generator->IsEnabledForView(*View))
+					{
+						Image = Generator->GetDebugImage(GraphBuilder, *ViewInfo);
+					}
+
+					if (Image)
+					{
+						InternalVRSSources.Add(Image);
+					}
 				}
-			}
 
-			FRDGTextureRef PreviewTexture = CombineShadingRateImages(GraphBuilder, *ViewInfo, InternalVRSSources);
-			if (!PreviewTexture)
-			{
-				return;
+				PreviewTexture = CombineShadingRateImages(GraphBuilder, *ViewInfo, InternalVRSSources);
 			}
+			
 
 			// If we have an active debug image, render it as a preview overlay
 			auto& RHICmdList = GraphBuilder.RHICmdList;
@@ -458,7 +449,63 @@ void FVariableRateShadingImageManager::DrawDebugPreview(FRDGBuilder& GraphBuilde
 		}
 	}
 
-	
+}
 
-	
+
+/**
+ * Private functions
+ */
+
+FRDGTextureRef FVariableRateShadingImageManager::CombineShadingRateImages(FRDGBuilder& GraphBuilder, const FViewInfo& ViewInfo, TArray<FRDGTextureRef> Sources)
+{
+	// If we have more than one source, combine the first available two
+	// TODO: Support combining more textures
+	if (Sources.Num() < 1)
+	{
+		return GetForceRateImage(GraphBuilder); // Fall back to uniform 1x1 shading rate if no images provided
+	}
+	else if (Sources.Num() == 1)
+	{
+		return Sources[0];
+	}
+	else
+	{
+		// Create texture to hold shading rate image
+		FRDGTextureRef CombinedShadingRateTexture = GraphBuilder.CreateTexture(GetSRIDesc(), TEXT("CombinedShadingRateTexture"));
+
+		FCombineShadingRateTexturesCS::FParameters* PassParameters = GraphBuilder.AllocParameters<FCombineShadingRateTexturesCS::FParameters>();
+		PassParameters->SourceTexture0 = Sources[0];
+		PassParameters->SourceTexture1 = Sources[1];
+		PassParameters->RWOutputTexture = GraphBuilder.CreateUAV(CombinedShadingRateTexture);
+
+		TShaderMapRef<FCombineShadingRateTexturesCS> ComputeShader(GetGlobalShaderMap(GMaxRHIFeatureLevel));
+
+		FComputeShaderUtils::AddPass(
+			GraphBuilder,
+			RDG_EVENT_NAME("CombineShadingRateImages"),
+			ERDGPassFlags::AsyncCompute | ERDGPassFlags::NeverCull,
+			ComputeShader,
+			PassParameters,
+			FComputeShaderUtils::GetGroupCount(FSceneTexturesConfig::Get().Extent, GetSRITileSize()));
+
+		return CombinedShadingRateTexture;
+	}
+
+}
+
+FRDGTextureRef FVariableRateShadingImageManager::GetForceRateImage(FRDGBuilder& GraphBuilder, int RateIndex /* = 0*/)
+{
+	static const TArray<uint32> ValidShadingRates = { VRSSR_1x1, VRSSR_1x2, VRSSR_2x1, VRSSR_2x2, VRSSR_2x4, VRSSR_4x2, VRSSR_4x4 };
+
+	if (RateIndex >= ValidShadingRates.Num())
+	{
+		// TODO: Log warning
+		RateIndex = ValidShadingRates.Num() - 1; // Default to maximum shading rate if value exceeds valid rates
+	}
+
+	FRDGTextureRef ForceShadingRateTexture = GraphBuilder.CreateTexture(GetSRIDesc(), TEXT("ForceShadingRateTexture"));
+	FRDGTextureUAVRef ForceShadingRateUAV = GraphBuilder.CreateUAV(ForceShadingRateTexture);
+	AddClearUAVPass(GraphBuilder, ForceShadingRateUAV, ValidShadingRates[RateIndex]);
+
+	return ForceShadingRateTexture;
 }
\ No newline at end of file
diff --git a/Engine/Source/Runtime/Renderer/Public/VariableRateShadingImageManager.h b/Engine/Source/Runtime/Renderer/Public/VariableRateShadingImageManager.h
--- a/Engine/Source/Runtime/Renderer/Public/VariableRateShadingImageManager.h
+++ b/Engine/Source/Runtime/Renderer/Public/VariableRateShadingImageManager.h
@@ -96,6 +96,7 @@ public:
 	static bool IsVRSCompatibleWithOutputType(const EDisplayOutputFormat& DisplayOutputFormat);
 
 	static FIntPoint GetSRITileSize();
+	static FRDGTextureDesc GetSRIDesc();
 
 	void DrawDebugPreview(FRDGBuilder& GraphBuilder, const FSceneViewFamily& ViewFamily, FRDGTextureRef OutputSceneColor);
 
@@ -104,6 +105,7 @@ private:
 	TArray<TUniquePtr<IVariableRateShadingImageGenerator>> ImageGenerators;
 
 	FRDGTextureRef CombineShadingRateImages(FRDGBuilder& GraphBuilder, const FViewInfo& ViewInfo, TArray<FRDGTextureRef> Sources);
+	FRDGTextureRef GetForceRateImage(FRDGBuilder& GraphBuilder, int RateIndex = 0);
 };
 
 ENUM_CLASS_FLAGS(FVariableRateShadingImageManager::EVRSSourceType);
