Change name of Substrate analytics event.

diff --git a/Engine/Source/Runtime/Renderer/Private/Strata/Strata.cpp b/Engine/Source/Runtime/Renderer/Private/Strata/Strata.cpp
--- a/Engine/Source/Runtime/Renderer/Private/Strata/Strata.cpp
+++ b/Engine/Source/Runtime/Renderer/Private/Strata/Strata.cpp
@@ -315,7 +315,7 @@ static void RecordStrataAnalytics(EShaderPlatform InPlatform)
 		EventAttributes.Add(FAnalyticsEventAttribute(TEXT("Enabled"), 1));
 		EventAttributes.Add(FAnalyticsEventAttribute(TEXT("BytesPerPixel"), GetBytePerPixel(InPlatform)));
 
-		FString OutStr(TEXT("Substrate"));
+		FString OutStr(TEXT("Substrate.Usage.ProjectSettings"));
 		FEngineAnalytics::GetProvider().RecordEvent(OutStr, EventAttributes);
 	}
 }
