Non-unity fixes
#rnx

diff --git a/Engine/Plugins/Experimental/PCG/Source/PCG/Public/Grid/PCGComponentOctree.h b/Engine/Plugins/Experimental/PCG/Source/PCG/Public/Grid/PCGComponentOctree.h
--- a/Engine/Plugins/Experimental/PCG/Source/PCG/Public/Grid/PCGComponentOctree.h
+++ b/Engine/Plugins/Experimental/PCG/Source/PCG/Public/Grid/PCGComponentOctree.h
@@ -4,6 +4,7 @@
 
 #include "Math/BoxSphereBounds.h"
 #include "Math/GenericOctree.h"
+#include "Misc/ScopeRWLock.h"
 #include "UObject/ObjectPtr.h"
 #include "Templates/SharedPointer.h"
 #include "Templates/SharedPointerFwd.h"
diff --git a/Engine/Plugins/Experimental/PCG/Source/PCG/Public/PCGActorAndComponentMapping.h b/Engine/Plugins/Experimental/PCG/Source/PCG/Public/PCGActorAndComponentMapping.h
--- a/Engine/Plugins/Experimental/PCG/Source/PCG/Public/PCGActorAndComponentMapping.h
+++ b/Engine/Plugins/Experimental/PCG/Source/PCG/Public/PCGActorAndComponentMapping.h
@@ -10,7 +10,7 @@
 #include "Containers/Set.h"
 #include "Templates/Function.h"
 #include "UObject/ObjectPtr.h"
-#include "UObject/WeakObjectPtrTemplates.h"
+#include "UObject/WeakObjectPtr.h"
 
 class AActor;
 class APCGPartitionActor;
diff --git a/Engine/Source/Runtime/Renderer/Private/VariableRateShading/VariableRateShadingImageManager.cpp b/Engine/Source/Runtime/Renderer/Private/VariableRateShading/VariableRateShadingImageManager.cpp
--- a/Engine/Source/Runtime/Renderer/Private/VariableRateShading/VariableRateShadingImageManager.cpp
+++ b/Engine/Source/Runtime/Renderer/Private/VariableRateShading/VariableRateShadingImageManager.cpp
@@ -12,6 +12,7 @@
 #include "RenderTargetPool.h"
 #include "SceneRendering.h"
 #include "SystemTextures.h"
+#include "SceneRendering.h"
 #include "SceneView.h"
 #include "IEyeTracker.h"
 #include "IHeadMountedDisplay.h"
