- Remove FRHIResource::IsValid() checkSlow from uniform buffer / SRV mesh bindings
- The resource can be marked for delete but still valid so this doesn't provide any meaningful information

#rnx

diff --git a/Engine/Source/Runtime/Renderer/Public/MeshDrawShaderBindings.h b/Engine/Source/Runtime/Renderer/Public/MeshDrawShaderBindings.h
--- a/Engine/Source/Runtime/Renderer/Public/MeshDrawShaderBindings.h
+++ b/Engine/Source/Runtime/Renderer/Public/MeshDrawShaderBindings.h
@@ -117,7 +117,6 @@ public:
 		if (Parameter.IsBound())
 		{
 			checkf(Value.GetReference(), TEXT("Attempted to set null uniform buffer for type %s"), UniformBufferStructType::FTypeInfo::GetStructMetadata()->GetStructTypeName());
-			checkfSlow(Value.GetReference()->IsValid(), TEXT("Attempted to set already deleted uniform buffer for type %s"), UniformBufferStructType::FTypeInfo::GetStructMetadata()->GetStructTypeName());
 			WriteBindingUniformBuffer(Value.GetReference(), Parameter.GetBaseIndex());
 		}
 	}
@@ -130,7 +129,6 @@ public:
 		if (Parameter.IsBound())
 		{
 			checkf(Value.GetUniformBufferRHI(), TEXT("Attempted to set null uniform buffer for type %s"), UniformBufferStructType::FTypeInfo::GetStructMetadata()->GetStructTypeName());
-			checkfSlow(Value.GetUniformBufferRHI()->IsValid(), TEXT("Attempted to set already deleted uniform buffer for type %s"), UniformBufferStructType::FTypeInfo::GetStructMetadata()->GetStructTypeName());
 			WriteBindingUniformBuffer(Value.GetUniformBufferRHI(), Parameter.GetBaseIndex());
 		}
 	}
@@ -142,7 +140,6 @@ public:
 		if (Parameter.IsBound())
 		{
 			checkf(Value, TEXT("Attempted to set null uniform buffer"));
-			checkfSlow(Value->IsValid(), TEXT("Attempted to set already deleted uniform buffer of type %s"), *Value->GetLayout().GetDebugName());
 			WriteBindingUniformBuffer(Value, Parameter.GetBaseIndex());
 		}
 	}
@@ -154,7 +151,6 @@ public:
 		if (Parameter.IsBound())
 		{
 			checkf(Value, TEXT("Attempted to set null SRV on slot %u"), Parameter.GetBaseIndex());
-			checkfSlow(Value->IsValid(), TEXT("Attempted to set already deleted SRV on slot %u"), Parameter.GetBaseIndex());
 			WriteBindingSRV(Value, Parameter.GetBaseIndex());
 		}
 	}
