Prevent initializing the VSM arays if dynamic shadows are disabled via show flags.

diff --git a/Engine/Source/Runtime/Renderer/Private/DeferredShadingRenderer.cpp b/Engine/Source/Runtime/Renderer/Private/DeferredShadingRenderer.cpp
--- a/Engine/Source/Runtime/Renderer/Private/DeferredShadingRenderer.cpp
+++ b/Engine/Source/Runtime/Renderer/Private/DeferredShadingRenderer.cpp
@@ -2733,7 +2733,8 @@ void FDeferredShadingSceneRenderer::Render(FRDGBuilder& GraphBuilder)
 	// Important that this uses consistent logic throughout the frame, so evaluate once and pass in the flag from here
 	// NOTE: Must be done after  system texture initialization
 	// TODO: This doesn't take into account the potential for split screen views with separate shadow caches
-	VirtualShadowMapArray.Initialize(GraphBuilder, Scene->GetVirtualShadowMapCache(Views[0]), UseVirtualShadowMaps(ShaderPlatform, FeatureLevel), Views[0].bIsSceneCapture);
+	const bool bEnableVirtualShadowMaps = UseVirtualShadowMaps(ShaderPlatform, FeatureLevel) && ViewFamily.EngineShowFlags.DynamicShadows;
+	VirtualShadowMapArray.Initialize(GraphBuilder, Scene->GetVirtualShadowMapCache(Views[0]), bEnableVirtualShadowMaps, Views[0].bIsSceneCapture);
 
 	// if DDM_AllOpaqueNoVelocity was used, then velocity should have already been rendered as well
 	const bool bIsEarlyDepthComplete = (DepthPass.EarlyZPassMode == DDM_AllOpaque || DepthPass.EarlyZPassMode == DDM_AllOpaqueNoVelocity);
