Add hair attribute value picking when attribute view mode is active.

diff --git a/Engine/Plugins/Runtime/HairStrands/Source/HairStrandsCore/Private/HairStrandsInterpolation.cpp b/Engine/Plugins/Runtime/HairStrands/Source/HairStrandsCore/Private/HairStrandsInterpolation.cpp
--- a/Engine/Plugins/Runtime/HairStrands/Source/HairStrandsCore/Private/HairStrandsInterpolation.cpp
+++ b/Engine/Plugins/Runtime/HairStrands/Source/HairStrandsCore/Private/HairStrandsInterpolation.cpp
@@ -1667,6 +1667,8 @@ static FHairGroupPublicData::FVertexFactoryInput InternalComputeHairStrandsVerte
 	OutVFInput.Strands.Common.bStableRasterization = Instance->Strands.Modifier.bUseStableRasterization;
 	OutVFInput.Strands.Common.bRaytracingGeometry = false;
 	OutVFInput.Strands.Common.RaytracingProceduralSplits = GetHairRaytracingProceduralSplits();
+	OutVFInput.Strands.Common.GroupIndex = Instance->Debug.GroupIndex;
+	OutVFInput.Strands.Common.GroupCount = Instance->Debug.GroupCount;
 
 	// Add local transform assignement here ...
 
diff --git a/Engine/Shaders/Private/HairStrands/HairStrandsAttributeCommon.ush b/Engine/Shaders/Private/HairStrands/HairStrandsAttributeCommon.ush
--- a/Engine/Shaders/Private/HairStrands/HairStrandsAttributeCommon.ush
+++ b/Engine/Shaders/Private/HairStrands/HairStrandsAttributeCommon.ush
@@ -93,138 +93,6 @@ bool IsAttributeValid(uint InOffset)
 #include "/Engine/Private/HairStrands/HairStrandsVertexFactoryCommon.ush"
 
 #if HAIR_STRAND_MESH_FACTORY
-
-float InternalGetHairStrandsWorldRadius(uint HairPrimitiveId, float2 VertexUV)
-{
-	const float Radius0 = InternalGetHairStrandsWorldRadius(HairStrandsVF.PositionBuffer, HairPrimitiveId,   HairStrandsVF.Radius, HairStrandsVF.RootScale, HairStrandsVF.TipScale);
-	const float Radius1 = InternalGetHairStrandsWorldRadius(HairStrandsVF.PositionBuffer, HairPrimitiveId+1, HairStrandsVF.Radius, HairStrandsVF.RootScale, HairStrandsVF.TipScale);
-	return lerp(Radius0, Radius1, VertexUV.x);
-}
-
-float InternalGetHairStrandsCoordU(uint HairPrimitiveId, float2 VertexUV)
-{
-	const float UCoord0 = InternalGetHairStrandsCoordU(HairStrandsVF.PositionBuffer, HairPrimitiveId,   HairStrandsVF.Radius, HairStrandsVF.RootScale, HairStrandsVF.TipScale);
-	const float UCoord1 = InternalGetHairStrandsCoordU(HairStrandsVF.PositionBuffer, HairPrimitiveId+1, HairStrandsVF.Radius, HairStrandsVF.RootScale, HairStrandsVF.TipScale);
-	return lerp(UCoord0, UCoord1, VertexUV.x);
-}
-
-float GetHairStrandsDepth(float2 InUV)
-{
-	return 0; // TODO: what should we do for the strands cards as we don't want to override the depth output, but we still want to share the same shader than the non cards one
-}
-
-float GetHairStrandsCoverage(uint HairPrimitiveId, float2 HairPrimitiveUV)
-{
-	return 1;
-}
-
-float2 GetHairStrandsUV(uint HairPrimitiveId, float2 HairPrimitiveUV)
-{
-	float2 Out;
-	Out.x = InternalGetHairStrandsCoordU(HairPrimitiveId, HairPrimitiveUV);
-	Out.y = HairPrimitiveUV.y;
-	return Out;
-}
-
-float2 GetHairStrandsDimensions(uint HairPrimitiveId, float2 HairPrimitiveUV)
-{
-	const float CurveLength = UnpackHairLength(InternalGetHairStrandsAttributesPerCurve(HairPrimitiveId, HAIR_CURVE_ATTRIBUTE_OFFSET_LENGTH(HairStrandsVF.CurveAttributeOffsets), HAIR_CURVE_ATTRIBUTE_STRIDE_LENGTH, HairStrandsVF.CurveAttributeBuffer, HairStrandsVF.PointToCurveBuffer));
-	const float CoordU = InternalGetHairStrandsCoordU(HairPrimitiveId, HairPrimitiveUV);
-	const float Radius = InternalGetHairStrandsWorldRadius(HairPrimitiveId, HairPrimitiveUV);
-	return float2(CurveLength * CoordU, Radius);
-}
-
-float2 GetHairStrandsRootUV(uint HairPrimitiveId, float2 HairPrimitiveUV)
-{
-	const uint RootUVOffset = HAIR_CURVE_ATTRIBUTE_OFFSET_ROOTUV(HairStrandsVF.CurveAttributeOffsets);
-	if (IsAttributeValid(RootUVOffset))
-	{
-		return UnpackHairRootUV(InternalGetHairStrandsAttributesPerCurve(HairPrimitiveId, RootUVOffset, HAIR_CURVE_ATTRIBUTE_STRIDE_ROOTUV, HairStrandsVF.CurveAttributeBuffer, HairStrandsVF.PointToCurveBuffer));
-	}
-	else
-	{
-		return 0;
-	}
-}
-
-float GetHairStrandsSeed(uint HairPrimitiveId, float2 HairPrimitiveUV)
-{
-	return UnpackHairSeed(InternalGetHairStrandsAttributesPerCurve(HairPrimitiveId, HAIR_CURVE_ATTRIBUTE_OFFSET_SEED(HairStrandsVF.CurveAttributeOffsets), HAIR_CURVE_ATTRIBUTE_STRIDE_SEED, HairStrandsVF.CurveAttributeBuffer, HairStrandsVF.PointToCurveBuffer));
-}
-
-uint3 GetHairStrandsClumpID(uint HairPrimitiveId, float2 HairPrimitiveUV)
-{
-	const uint ClumpIdOffset  = HAIR_CURVE_ATTRIBUTE_OFFSET_CLUMPID(HairStrandsVF.CurveAttributeOffsets);
-	const uint ClumpId3Offset = HAIR_CURVE_ATTRIBUTE_OFFSET_CLUMPID3(HairStrandsVF.CurveAttributeOffsets);
-	if (IsAttributeValid(ClumpIdOffset))
-	{
-		return UnpackHairClumpID(InternalGetHairStrandsAttributesPerCurve(HairPrimitiveId, ClumpIdOffset, HAIR_CURVE_ATTRIBUTE_STRIDE_CLUMPID, HairStrandsVF.CurveAttributeBuffer, HairStrandsVF.PointToCurveBuffer));
-	}
-	else if (IsAttributeValid(ClumpId3Offset))
-	{
-		return UnpackHairClumpID(InternalGetHairStrandsAttributesPerCurve64(HairPrimitiveId, ClumpId3Offset, HAIR_CURVE_ATTRIBUTE_STRIDE_CLUMPID3, HairStrandsVF.CurveAttributeBuffer, HairStrandsVF.PointToCurveBuffer));
-	}
-	else
-	{
-		return 0;
-	}
-}
-
-float3 GetHairStrandsBaseColor(uint HairPrimitiveId, float2 HairPrimitiveUV)
-{
-	const uint ColorOffset = HAIR_POINT_ATTRIBUTE_OFFSET_COLOR(HairStrandsVF.PointAttributeOffsets);
-	if (IsAttributeValid(ColorOffset))
-	{
-		const float3 Color0 = UnpackHairColor(InternalGetHairStrandsAttributesPerVertex(HairPrimitiveId,   ColorOffset, HAIR_POINT_ATTRIBUTE_STRIDE_COLOR, HairStrandsVF.PointAttributeBuffer));
-		const float3 Color1 = UnpackHairColor(InternalGetHairStrandsAttributesPerVertex(HairPrimitiveId+1, ColorOffset, HAIR_POINT_ATTRIBUTE_STRIDE_COLOR, HairStrandsVF.PointAttributeBuffer));
-		return lerp(Color0, Color1, HairPrimitiveUV.x);
-	}
-	else
-	{
-		return 0;
-	}
-}
-
-float GetHairStrandsRoughness(uint HairPrimitiveId, float2 HairPrimitiveUV)
-{
-	const uint RoughnessOffset = HAIR_POINT_ATTRIBUTE_OFFSET_ROUGHNESS(HairStrandsVF.PointAttributeOffsets);
-	if (IsAttributeValid(RoughnessOffset))
-	{
-		const float Roughness0 = UnpackHairRoughness(InternalGetHairStrandsAttributesPerVertex(HairPrimitiveId,   RoughnessOffset, HAIR_POINT_ATTRIBUTE_STRIDE_ROUGHNESS, HairStrandsVF.PointAttributeBuffer));
-		const float Roughness1 = UnpackHairRoughness(InternalGetHairStrandsAttributesPerVertex(HairPrimitiveId+1, RoughnessOffset, HAIR_POINT_ATTRIBUTE_STRIDE_ROUGHNESS, HairStrandsVF.PointAttributeBuffer));
-		return lerp(Roughness0, Roughness1, HairPrimitiveUV.x);
-	}
-	else
-	{
-		return 0;
-	}
-}
-
-float GetHairStrandsAO(uint HairPrimitiveId, float2 HairPrimitiveUV)
-{
-	const uint AoOffset = HAIR_POINT_ATTRIBUTE_OFFSET_AO(HairStrandsVF.PointAttributeOffsets);
-	if (IsAttributeValid(AoOffset))
-	{
-		const float AO0 = UnpackHairRoughness(InternalGetHairStrandsAttributesPerVertex(HairPrimitiveId, AoOffset, HAIR_POINT_ATTRIBUTE_STRIDE_AO, HairStrandsVF.PointAttributeBuffer));
-		const float AO1 = UnpackHairRoughness(InternalGetHairStrandsAttributesPerVertex(HairPrimitiveId+1, AoOffset, HAIR_POINT_ATTRIBUTE_STRIDE_AO, HairStrandsVF.PointAttributeBuffer));
-		return lerp(AO0, AO1, HairPrimitiveUV.x);
-	}
-	else
-	{
-		return 1.f;
-	}
-}
-
-float GetHairStrandsDepth(uint HairPrimitiveId, float2 InUV, float InDeviceZ)
-{
-	return 0.f;
-}
-
-float GetHairStrandsGroupIndex(float2 InAtlasUV)
-{
-	return HairStrandsVF.GroupIndex;
-}
-
-////////////////////////////////////////////////////////////////////////////////
-
-#endif // HAIR_STRAND_MESH_FACTORY
\ No newline at end of file
+#define HAIR_STRANDS_ATTRIBUTE_ACCESSORS(Name) HairStrandsVF_##Name
+#include "/Engine/Private/HairStrands/HairStrandsAttributeTemplate.ush"
+#endif
diff --git a/Engine/Shaders/Private/HairStrands/HairStrandsAttributeTemplate.ush b/Engine/Shaders/Private/HairStrands/HairStrandsAttributeTemplate.ush
new file mode 100644
--- /dev/null
+++ b/Engine/Shaders/Private/HairStrands/HairStrandsAttributeTemplate.ush
@@ -0,0 +1,135 @@
+// Copyright Epic Games, Inc. All Rights Reserved.
+
+#pragma once 
+
+// Check parameters/resources prefixing macro is well defined
+#ifndef HAIR_STRANDS_ATTRIBUTE_ACCESSORS
+#error HAIR_STRANDS_ATTRIBUTE_ACCESSORS needs to be defined
+#endif
+
+// Redefine macro for easing reability
+#define PREFIX(Name) HAIR_STRANDS_ATTRIBUTE_ACCESSORS(Name)
+
+float InternalGetHairStrandsWorldRadius(uint HairPrimitiveId, float2 VertexUV)
+{
+	const float Radius0 = InternalGetHairStrandsWorldRadius(PREFIX(PositionBuffer), HairPrimitiveId,   PREFIX(Radius), PREFIX(RootScale), PREFIX(TipScale));
+	const float Radius1 = InternalGetHairStrandsWorldRadius(PREFIX(PositionBuffer), HairPrimitiveId+1, PREFIX(Radius), PREFIX(RootScale), PREFIX(TipScale));
+	return lerp(Radius0, Radius1, VertexUV.x);
+}
+float InternalGetHairStrandsCoordU(uint HairPrimitiveId, float2 VertexUV)
+{
+	const float UCoord0 = InternalGetHairStrandsCoordU(PREFIX(PositionBuffer), HairPrimitiveId,   PREFIX(Radius), PREFIX(RootScale), PREFIX(TipScale));
+	const float UCoord1 = InternalGetHairStrandsCoordU(PREFIX(PositionBuffer), HairPrimitiveId+1, PREFIX(Radius), PREFIX(RootScale), PREFIX(TipScale));
+	return lerp(UCoord0, UCoord1, VertexUV.x);
+}
+float GetHairStrandsDepth(float2 InUV)
+{
+	return 0;
+}
+float GetHairStrandsCoverage(uint HairPrimitiveId, float2 HairPrimitiveUV)
+{
+	return 1;
+}
+float2 GetHairStrandsUV(uint HairPrimitiveId, float2 HairPrimitiveUV)
+{
+	float2 Out;
+	Out.x = InternalGetHairStrandsCoordU(HairPrimitiveId, HairPrimitiveUV);
+	Out.y = HairPrimitiveUV.y;
+	return Out;
+}
+float2 GetHairStrandsDimensions(uint HairPrimitiveId, float2 HairPrimitiveUV)
+{
+	const float CurveLength = UnpackHairLength(InternalGetHairStrandsAttributesPerCurve(HairPrimitiveId, HAIR_CURVE_ATTRIBUTE_OFFSET_LENGTH(PREFIX(CurveAttributeOffsets)), HAIR_CURVE_ATTRIBUTE_STRIDE_LENGTH, PREFIX(CurveAttributeBuffer), PREFIX(PointToCurveBuffer)));
+	const float CoordU = InternalGetHairStrandsCoordU(HairPrimitiveId, HairPrimitiveUV);
+	const float Radius = InternalGetHairStrandsWorldRadius(HairPrimitiveId, HairPrimitiveUV);
+	return float2(CurveLength * CoordU, Radius);
+}\
+float2 GetHairStrandsRootUV(uint HairPrimitiveId, float2 HairPrimitiveUV)
+{
+	const uint RootUVOffset = HAIR_CURVE_ATTRIBUTE_OFFSET_ROOTUV(PREFIX(CurveAttributeOffsets));
+	if (IsAttributeValid(RootUVOffset))
+	{
+		return UnpackHairRootUV(InternalGetHairStrandsAttributesPerCurve(HairPrimitiveId, RootUVOffset, HAIR_CURVE_ATTRIBUTE_STRIDE_ROOTUV, PREFIX(CurveAttributeBuffer), PREFIX(PointToCurveBuffer)));
+	}
+	else
+	{
+		return 0;
+	}
+}
+float GetHairStrandsSeed(uint HairPrimitiveId, float2 HairPrimitiveUV)
+{
+	return UnpackHairSeed(InternalGetHairStrandsAttributesPerCurve(HairPrimitiveId, HAIR_CURVE_ATTRIBUTE_OFFSET_SEED(PREFIX(CurveAttributeOffsets)), HAIR_CURVE_ATTRIBUTE_STRIDE_SEED, PREFIX(CurveAttributeBuffer), PREFIX(PointToCurveBuffer)));
+}
+uint3 GetHairStrandsClumpID(uint HairPrimitiveId, float2 HairPrimitiveUV)
+{
+	const uint ClumpIdOffset  = HAIR_CURVE_ATTRIBUTE_OFFSET_CLUMPID(PREFIX(CurveAttributeOffsets));
+	const uint ClumpId3Offset = HAIR_CURVE_ATTRIBUTE_OFFSET_CLUMPID3(PREFIX(CurveAttributeOffsets));
+	if (IsAttributeValid(ClumpIdOffset))
+	{
+		return UnpackHairClumpID(InternalGetHairStrandsAttributesPerCurve(HairPrimitiveId, ClumpIdOffset, HAIR_CURVE_ATTRIBUTE_STRIDE_CLUMPID, PREFIX(CurveAttributeBuffer), PREFIX(PointToCurveBuffer)));
+	}
+	else if (IsAttributeValid(ClumpId3Offset))
+	{
+		return UnpackHairClumpID(InternalGetHairStrandsAttributesPerCurve64(HairPrimitiveId, ClumpId3Offset, HAIR_CURVE_ATTRIBUTE_STRIDE_CLUMPID3, PREFIX(CurveAttributeBuffer), PREFIX(PointToCurveBuffer)));
+	}
+	else
+	{
+		return 0;
+	}
+}
+float3 GetHairStrandsBaseColor(uint HairPrimitiveId, float2 HairPrimitiveUV)
+{
+	const uint ColorOffset = HAIR_POINT_ATTRIBUTE_OFFSET_COLOR(PREFIX(PointAttributeOffsets));
+	if (IsAttributeValid(ColorOffset))
+	{
+		const float3 Color0 = UnpackHairColor(InternalGetHairStrandsAttributesPerVertex(HairPrimitiveId,   ColorOffset, HAIR_POINT_ATTRIBUTE_STRIDE_COLOR, PREFIX(PointAttributeBuffer)));
+		const float3 Color1 = UnpackHairColor(InternalGetHairStrandsAttributesPerVertex(HairPrimitiveId+1, ColorOffset, HAIR_POINT_ATTRIBUTE_STRIDE_COLOR, PREFIX(PointAttributeBuffer)));
+		return lerp(Color0, Color1, HairPrimitiveUV.x);
+	}
+	else
+	{
+		return 0;
+	}
+}
+float GetHairStrandsRoughness(uint HairPrimitiveId, float2 HairPrimitiveUV)
+{
+	const uint RoughnessOffset = HAIR_POINT_ATTRIBUTE_OFFSET_ROUGHNESS(PREFIX(PointAttributeOffsets));
+	if (IsAttributeValid(RoughnessOffset))
+	{
+		const float Roughness0 = UnpackHairRoughness(InternalGetHairStrandsAttributesPerVertex(HairPrimitiveId,   RoughnessOffset, HAIR_POINT_ATTRIBUTE_STRIDE_ROUGHNESS, PREFIX(PointAttributeBuffer)));
+		const float Roughness1 = UnpackHairRoughness(InternalGetHairStrandsAttributesPerVertex(HairPrimitiveId+1, RoughnessOffset, HAIR_POINT_ATTRIBUTE_STRIDE_ROUGHNESS, PREFIX(PointAttributeBuffer)));
+		return lerp(Roughness0, Roughness1, HairPrimitiveUV.x);
+	}
+	else
+	{
+		return 0;
+	}
+}
+float GetHairStrandsAO(uint HairPrimitiveId, float2 HairPrimitiveUV)
+{
+	const uint AoOffset = HAIR_POINT_ATTRIBUTE_OFFSET_AO(PREFIX(PointAttributeOffsets));
+	if (IsAttributeValid(AoOffset))
+	{
+		const float AO0 = UnpackHairRoughness(InternalGetHairStrandsAttributesPerVertex(HairPrimitiveId, AoOffset, HAIR_POINT_ATTRIBUTE_STRIDE_AO, PREFIX(PointAttributeBuffer)));
+		const float AO1 = UnpackHairRoughness(InternalGetHairStrandsAttributesPerVertex(HairPrimitiveId+1, AoOffset, HAIR_POINT_ATTRIBUTE_STRIDE_AO, PREFIX(PointAttributeBuffer)));
+		return lerp(AO0, AO1, HairPrimitiveUV.x);
+	}
+	else
+	{
+		return 1.f;
+	}
+}
+float GetHairStrandsDepth(uint HairPrimitiveId, float2 InUV, float InDeviceZ) 
+{ 
+	return 0.f; 
+} 
+float GetHairStrandsGroupIndex(float2 InAtlasUV) 
+{ 
+	return PREFIX(GroupIndex); 
+}
+bool HasHairStrandsRootUV() 	{ return IsAttributeValid(HAIR_CURVE_ATTRIBUTE_OFFSET_ROOTUV(PREFIX(CurveAttributeOffsets))); }
+bool HasHairStrandsClumpId() 	{ return IsAttributeValid(HAIR_CURVE_ATTRIBUTE_OFFSET_CLUMPID(PREFIX(CurveAttributeOffsets))); }
+bool HasHairStrandsClumpId3() 	{ return IsAttributeValid(HAIR_CURVE_ATTRIBUTE_OFFSET_CLUMPID3(PREFIX(CurveAttributeOffsets))); }
+bool HasHairStrandsColor() 		{ return IsAttributeValid(HAIR_POINT_ATTRIBUTE_OFFSET_COLOR(PREFIX(PointAttributeOffsets))); }
+bool HasHairStrandsRoughness() 	{ return IsAttributeValid(HAIR_POINT_ATTRIBUTE_OFFSET_ROUGHNESS(PREFIX(PointAttributeOffsets))); }
+bool HasHairStrandsAO() 		{ return IsAttributeValid(HAIR_POINT_ATTRIBUTE_OFFSET_AO(PREFIX(PointAttributeOffsets))); }
\ No newline at end of file
diff --git a/Engine/Shaders/Private/HairStrands/HairStrandsDebug.usf b/Engine/Shaders/Private/HairStrands/HairStrandsDebug.usf
--- a/Engine/Shaders/Private/HairStrands/HairStrandsDebug.usf
+++ b/Engine/Shaders/Private/HairStrands/HairStrandsDebug.usf
@@ -2667,3 +2667,78 @@ void MainCS(uint GroupIndex : SV_GroupIndex, uint3 DispatchThreadId : SV_Dispatc
 	}
 }
 #endif // SHADER_SHADOW_CULLING
+
+//////////////////////////////////////////////////////////////////////////////////////////////////////////////////////////////
+// Display hair picking attribute
+
+#if SHADER_ATTRIBUTE_DEBUG
+
+#include "../Common.ush"
+#include "../ShaderPrint.ush"
+#include "HairStrandsVisibilityCommon.ush"
+#include "HairStrandsVertexFactoryCommon.ush"
+HAIR_STRANDS_INSTANCE_PARAMETERS(HairInstance)
+
+#include "HairStrandsAttributeCommon.ush"
+#define HAIR_STRANDS_ATTRIBUTE_ACCESSORS(Name) HairInstance_##Name
+#include "/Engine/Private/HairStrands/HairStrandsAttributeTemplate.ush"
+
+uint MaterialId;
+Texture2D<uint> NodeIndexTexture;
+StructuredBuffer<FPackedHairVis> NodeVisBuffer;
+
+[numthreads(1, 1, 1)]
+void CSMain(
+	uint DispatchThreadID : SV_DispatchThreadID, 
+	uint GroupThreadID : SV_GroupThreadID, 
+	uint GroupID : SV_GroupID)
+{
+	const int2 PixelCoord = ShaderPrintData.CursorCoord;
+	if (all(PixelCoord < 0))
+	{
+		return;
+	}
+
+	const FNodeDesc NodeDesc = DecodeNodeDesc(NodeIndexTexture.Load(uint3(PixelCoord, 0)));
+	if (NodeDesc.Count == 0)
+	{
+		return;
+	}
+
+	if (NodeDesc.Count > 0)
+	{
+		// Only show the first sample properties for simplicity
+		const uint SampleIt = 0;
+		const uint SampleIndex = NodeDesc.Offset + SampleIt;
+		const FHairVis VisNode = UnpackHairVis(NodeVisBuffer[SampleIndex]);
+		if (VisNode.MaterialId == MaterialId)
+		{
+			FHairControlPoint Out = UnpackHairControlPoint(
+				HairInstance_PositionBuffer[VisNode.PrimitiveId],
+				HairInstance_PositionOffsetBuffer[0].xyz,
+				HairInstance_Radius,
+				HairInstance_RootScale,
+				HairInstance_TipScale);
+
+			FFontColor FontValue = FontYellow;
+			FFontColor FontName = FontWhite;
+
+			FShaderPrintContext Ctx = InitShaderPrintContext(true, uint2(PixelCoord) + uint2(20, 0));
+			Print(Ctx, TEXT("Group    : "), FontName); Print(Ctx, GetHairStrandsGroupIndex(0), FontValue); Newline(Ctx);
+			Print(Ctx, TEXT("Radius   : "), FontName); Print(Ctx, GetHairStrandsDimensions(VisNode.PrimitiveId, 0).y, FontValue); Newline(Ctx);
+			Print(Ctx, TEXT("Seed     : "), FontName); Print(Ctx, GetHairStrandsSeed(VisNode.PrimitiveId, 0), FontValue); Newline(Ctx);
+			Print(Ctx, TEXT("CoordU   : "), FontName); Print(Ctx, GetHairStrandsUV(VisNode.PrimitiveId, 0).x, FontValue); Newline(Ctx);
+			Newline(Ctx);
+
+			if (HasHairStrandsRootUV()) 	{ Print(Ctx, TEXT("RootUV   : "), FontName); Print(Ctx, GetHairStrandsRootUV(VisNode.PrimitiveId, 0),  FontValue); Newline(Ctx); }
+			if (HasHairStrandsClumpId()) 	{ Print(Ctx, TEXT("Clump ID : "), FontName); Print(Ctx, GetHairStrandsClumpID(VisNode.PrimitiveId, 0).x, FontValue); Newline(Ctx); }
+			if (HasHairStrandsClumpId3()) 	{ Print(Ctx, TEXT("Clump IDs: "), FontName); Print(Ctx, GetHairStrandsClumpID(VisNode.PrimitiveId, 0), FontValue); Newline(Ctx); }
+			if (HasHairStrandsColor()) 		{ Print(Ctx, TEXT("Color    : "), FontName); Print(Ctx, GetHairStrandsBaseColor(VisNode.PrimitiveId, 0), FontValue); Newline(Ctx); }
+			if (HasHairStrandsRoughness()) 	{ Print(Ctx, TEXT("Roughness: "), FontName); Print(Ctx, GetHairStrandsRoughness(VisNode.PrimitiveId, 0), FontValue); Newline(Ctx); }
+			if (HasHairStrandsAO()) 		{ Print(Ctx, TEXT("AO       : "), FontName); Print(Ctx, GetHairStrandsAO(VisNode.PrimitiveId, 0), FontValue); Newline(Ctx); }
+			Newline(Ctx);
+		}
+	}
+}
+
+#endif // SHADER_ATTRIBUTE_DEBUG
diff --git a/Engine/Shaders/Private/HairStrands/HairStrandsVertexFactoryCommon.ush b/Engine/Shaders/Private/HairStrands/HairStrandsVertexFactoryCommon.ush
--- a/Engine/Shaders/Private/HairStrands/HairStrandsVertexFactoryCommon.ush
+++ b/Engine/Shaders/Private/HairStrands/HairStrandsVertexFactoryCommon.ush
@@ -17,6 +17,7 @@
 	uint				Prefix##_bRaytracingGeometry; \
 	uint				Prefix##_RaytracingProceduralSplits; \
 	uint				Prefix##_GroupIndex; \
+	uint				Prefix##_GroupCount; \
 	float				Prefix##_Radius; \
 	float				Prefix##_RootScale; \
 	float				Prefix##_TipScale; \
diff --git a/Engine/Shaders/Shared/HairStrandsDefinitions.h b/Engine/Shaders/Shared/HairStrandsDefinitions.h
--- a/Engine/Shaders/Shared/HairStrandsDefinitions.h
+++ b/Engine/Shaders/Shared/HairStrandsDefinitions.h
@@ -71,11 +71,11 @@
 
 // Pack all offset into uint4
 #define PACK_HAIR_ATTRIBUTE_OFFSETS(Out, In, Offset, Count) \
-	for (uint32 AttributeIt4 = 0; AttributeIt4 < Count; ++AttributeIt4) \
+	for (uint32 AttributeIt4 = 0; AttributeIt4 < Offset; ++AttributeIt4) \
 	{ \
 		Out[AttributeIt4] = FUintVector4(HAIR_ATTRIBUTE_INVALID_OFFSET, HAIR_ATTRIBUTE_INVALID_OFFSET, HAIR_ATTRIBUTE_INVALID_OFFSET, HAIR_ATTRIBUTE_INVALID_OFFSET); \
 	} \
-	for (uint32 AttributeIt = 0; AttributeIt < Offset; ++AttributeIt) \
+	for (uint32 AttributeIt = 0; AttributeIt < Count; ++AttributeIt) \
 	{ \
 		const uint32 Index4 = AttributeIt & (~0x3); \
 		const uint32 SubIndex = AttributeIt - Index4; \
diff --git a/Engine/Source/Runtime/Renderer/Private/HairStrands/HairStrandsData.h b/Engine/Source/Runtime/Renderer/Private/HairStrands/HairStrandsData.h
--- a/Engine/Source/Runtime/Renderer/Private/HairStrands/HairStrandsData.h
+++ b/Engine/Source/Runtime/Renderer/Private/HairStrands/HairStrandsData.h
@@ -47,6 +47,10 @@ BEGIN_GLOBAL_SHADER_PARAMETER_STRUCT(FHairStrandsViewUniformParameters, )
 	SHADER_PARAMETER_RDG_BUFFER_SRV(Buffer<uint>,  HairTileCount)						// Tile total count (actual number of tiles)
 END_GLOBAL_SHADER_PARAMETER_STRUCT()
 
+
+void GetHairStrandsInstanceCommon(FRDGBuilder& GraphBuilder, const FViewInfo& ViewInfo, const FHairGroupPublicData* HairGroupPublicData, FHairStrandsInstanceCommonParameters& OutCommon);
+void GetHairStrandsInstanceResources(FRDGBuilder& GraphBuilder, const FViewInfo& ViewInfo, const FHairGroupPublicData* HairGroupPublicData, bool bForceRegister, FHairStrandsInstanceResourceParameters& OutResources);
+void GetHairStrandsInstanceCulling(FRDGBuilder& GraphBuilder, const FViewInfo& ViewInfo, const FHairGroupPublicData* HairGroupPublicData, bool bCullingEnable, FHairStrandsInstanceCullingParameters& OutCulling);
 FHairStrandsInstanceParameters GetHairStrandsInstanceParameters(FRDGBuilder& GraphBuilder, const FViewInfo& ViewInfo, const FHairGroupPublicData* HairGroupPublicData, bool bCullingEnable, bool bForceRegister);
 
 ////////////////////////////////////////////////////////////////////////////////////
diff --git a/Engine/Source/Runtime/Renderer/Private/HairStrands/HairStrandsVisibility.cpp b/Engine/Source/Runtime/Renderer/Private/HairStrands/HairStrandsVisibility.cpp
--- a/Engine/Source/Runtime/Renderer/Private/HairStrands/HairStrandsVisibility.cpp
+++ b/Engine/Source/Runtime/Renderer/Private/HairStrands/HairStrandsVisibility.cpp
@@ -829,6 +829,8 @@ struct FMaterialPassOutput
 	FRDGTextureRef SampleLightingTexture = nullptr;
 };
 
+static void AddAttributeDebugPass(FRDGBuilder& GraphBuilder, const FViewInfo& View, const FHairStrandsMacroGroupDatas& MacroGroupDatas, const FRDGTextureRef& NodeIndexTexture, const FRDGBufferSRVRef& NodeVisBuffer);
+
 static FMaterialPassOutput AddHairMaterialPass(
 	FRDGBuilder& GraphBuilder,
 	const FScene* Scene,
@@ -970,7 +972,6 @@ static FMaterialPassOutput AddHairMaterialPass(
 		});
 	};
 
-
 	const bool bIsPlatformRequireRenderTarget = IsPlatformRequiringRenderTargetForMaterialPass(Scene->GetShaderPlatform()) || GRHIRequiresRenderTargetForPixelShaderUAVs;
 
 	// Output:
@@ -991,9 +992,112 @@ static FMaterialPassOutput AddHairMaterialPass(
 		MaterialPass(nullptr, EHairMaterialPassFilter::NonEmissiveOnly);
 	}
 
+	// Add debug attribute pass
+	const EGroomViewMode ViewMode = GetGroomViewMode(*ViewInfo);
+	const bool bDebugEnabled = 
+		ViewMode == EGroomViewMode::UV
+		|| ViewMode == EGroomViewMode::RootUV
+		|| ViewMode == EGroomViewMode::RootUDIM
+		|| ViewMode == EGroomViewMode::Seed
+		|| ViewMode == EGroomViewMode::Dimension
+		|| ViewMode == EGroomViewMode::RadiusVariation
+		|| ViewMode == EGroomViewMode::Tangent
+		|| ViewMode == EGroomViewMode::Color
+		|| ViewMode == EGroomViewMode::Roughness
+		|| ViewMode == EGroomViewMode::AO
+		|| ViewMode == EGroomViewMode::ClumpID;
+	if (bDebugEnabled)
+	{	
+		AddAttributeDebugPass(GraphBuilder, *ViewInfo, MacroGroupDatas, CompactNodeIndex, GraphBuilder.CreateSRV(CompactNodeVis));
+	}
+	
 	return Output;
 }
 
+///////////////////////////////////////////////////////////////////////////////////////////////////
+
+class FHairAttributeDebugCS : public FGlobalShader
+{
+	DECLARE_GLOBAL_SHADER(FHairAttributeDebugCS);
+	SHADER_USE_PARAMETER_STRUCT(FHairAttributeDebugCS, FGlobalShader);
+
+	using FPermutationDomain = TShaderPermutationDomain<>;
+
+	BEGIN_SHADER_PARAMETER_STRUCT(FHairInstance, )
+		SHADER_PARAMETER_STRUCT_INCLUDE(FHairStrandsInstanceCommonParameters, Common)
+		SHADER_PARAMETER_STRUCT_INCLUDE(FHairStrandsInstanceResourceParameters, Resources)
+	END_SHADER_PARAMETER_STRUCT()
+
+	BEGIN_SHADER_PARAMETER_STRUCT(FParameters, )
+		SHADER_PARAMETER(uint32, MaterialId)
+		SHADER_PARAMETER_RDG_TEXTURE(Texture2D<uint>, NodeIndexTexture)
+		SHADER_PARAMETER_RDG_BUFFER_SRV(StructuredBuffer<FPackedHairVis>, NodeVisBuffer)
+		SHADER_PARAMETER_STRUCT(FHairInstance, HairInstance)
+		SHADER_PARAMETER_STRUCT_REF(FViewUniformShaderParameters, ViewUniformBuffer)
+		SHADER_PARAMETER_STRUCT_INCLUDE(ShaderPrint::FShaderParameters, ShaderPrintUniformBuffer)
+	END_SHADER_PARAMETER_STRUCT()
+
+
+public:
+	static bool ShouldCompilePermutation(const FGlobalShaderPermutationParameters &Parameters) { return IsHairStrandsSupported(EHairStrandsShaderType::Strands, Parameters.Platform); }
+	static void ModifyCompilationEnvironment(const FGlobalShaderPermutationParameters &Parameters, FShaderCompilerEnvironment &OutEnvironment)
+	{
+		FGlobalShader::ModifyCompilationEnvironment(Parameters, OutEnvironment);
+		OutEnvironment.SetDefine(TEXT("SHADER_ATTRIBUTE_DEBUG"), 1);
+	}
+};
+
+IMPLEMENT_GLOBAL_SHADER(FHairAttributeDebugCS, "/Engine/Private/HairStrands/HairStrandsDebug.usf", "CSMain", SF_Compute);
+
+static void AddAttributeDebugPass(FRDGBuilder& GraphBuilder, const FViewInfo& View, const FHairStrandsMacroGroupDatas& MacroGroupDatas, const FRDGTextureRef& NodeIndexTexture, const FRDGBufferSRVRef& NodeVisBuffer)
+{
+	// Force ShaderPrint on.
+	ShaderPrint::SetEnabled(true);
+	ShaderPrint::RequestSpaceForCharacters(2048);
+	ShaderPrint::RequestSpaceForLines(256);
+	if (!ShaderPrint::IsEnabled(View.ShaderPrintData))
+	{
+		return;
+	}
+
+	TShaderMapRef<FHairAttributeDebugCS> ComputeShader(View.ShaderMap);
+
+	for (const FHairStrandsMacroGroupData& MacroGroup : MacroGroupDatas)
+	{
+		const FHairStrandsMacroGroupData::TPrimitiveInfos& PrimitiveSceneInfos = MacroGroup.PrimitivesInfos;
+		for (const FHairStrandsMacroGroupData::PrimitiveInfo& PrimitiveInfo : PrimitiveSceneInfos)
+		{
+			// If a groom is not visible in primary view, but visible in shadow view, its PrimitiveInfo.Mesh will be null.
+			if (PrimitiveInfo.Mesh == nullptr || PrimitiveInfo.Mesh->Elements.Num() == 0)
+			{
+				continue;
+			}
+
+			const FHairGroupPublicData* HairGroupPublicData = reinterpret_cast<const FHairGroupPublicData*>(PrimitiveInfo.Mesh->Elements[0].VertexFactoryUserData);
+			check(HairGroupPublicData);
+
+			const bool bCullingEnable = HairGroupPublicData->GetCullingResultAvailable();
+
+			FHairAttributeDebugCS::FParameters* Parameters = GraphBuilder.AllocParameters<FHairAttributeDebugCS::FParameters>();
+			Parameters->MaterialId = PrimitiveInfo.MaterialId;
+			Parameters->NodeIndexTexture = NodeIndexTexture;
+			Parameters->NodeVisBuffer = NodeVisBuffer;
+			GetHairStrandsInstanceCommon(GraphBuilder, View, HairGroupPublicData, Parameters->HairInstance.Common);
+			GetHairStrandsInstanceResources(GraphBuilder, View, HairGroupPublicData, true/*bForceRegister*/, Parameters->HairInstance.Resources);
+//			Parameters->HairInstance = GetHairStrandsInstanceParameters(GraphBuilder, View, HairGroupPublicData, bCullingEnable, true );
+			Parameters->ViewUniformBuffer = View.ViewUniformBuffer;
+			ShaderPrint::SetParameters(GraphBuilder, View.ShaderPrintData, Parameters->ShaderPrintUniformBuffer);
+
+			FComputeShaderUtils::AddPass(
+				GraphBuilder,
+				RDG_EVENT_NAME("HairStrands::AttributeDebug"),
+				ComputeShader,
+				Parameters,
+				FIntVector(1, 1, 1));
+		}
+	}
+}
+
 ///////////////////////////////////////////////////////////////////////////////////////////////////
 class FHairVelocityCS: public FGlobalShader
 {
@@ -2831,54 +2935,71 @@ IMPLEMENT_GLOBAL_SHADER(FVisiblityRasterComputePrepareDepthGridCS, "/Engine/Priv
 
 /////////////////////////////////////////////////////////////////////////////////////////
 
-FHairStrandsInstanceParameters GetHairStrandsInstanceParameters(FRDGBuilder& GraphBuilder, const FViewInfo& ViewInfo, const FHairGroupPublicData* HairGroupPublicData, bool bCullingEnable, bool bForceRegister)
+void GetHairStrandsInstanceCommon(FRDGBuilder& GraphBuilder, const FViewInfo& ViewInfo, const FHairGroupPublicData* HairGroupPublicData, FHairStrandsInstanceCommonParameters& OutCommon)
 {
 	check(HairGroupPublicData);
+	const FHairGroupPublicData::FVertexFactoryInput& VFInput = HairGroupPublicData->VFInput;
+
+	OutCommon = VFInput.Strands.Common;
+
+	// Absolute local to world
+	OutCommon.LocalToWorldPrimitiveTransform	= FMatrix44f(VFInput.LocalToWorldTransform.ToMatrixWithScale()); // LWC_TODO: Precision loss, remove usage for this
+
+	// Translated local to world
+	const FVector& TranslatedWorldOffset = ViewInfo.ViewMatrices.GetPreViewTranslation();
+	FTransform LocalToTranslatedWorldTransform = VFInput.LocalToWorldTransform;
+	LocalToTranslatedWorldTransform.AddToTranslation(TranslatedWorldOffset);
+	OutCommon.LocalToTranslatedWorldPrimitiveTransform = FMatrix44f(LocalToTranslatedWorldTransform.ToMatrixWithScale());
+}
 
+void GetHairStrandsInstanceResources(FRDGBuilder& GraphBuilder, const FViewInfo& ViewInfo, const FHairGroupPublicData* HairGroupPublicData, bool bForceRegister, FHairStrandsInstanceResourceParameters& OutResources)
+{
+	check(HairGroupPublicData);
 	const FHairGroupPublicData::FVertexFactoryInput& VFInput = HairGroupPublicData->VFInput;
 
-	FHairStrandsInstanceParameters Output;
-	FHairStrandsInstanceIntermediateParameters& Out = Output.HairStrandsVF;
 	if (bForceRegister)
 	{
-		Out.Resources.PositionBuffer		= Register(GraphBuilder, VFInput.Strands.PositionBufferExternal, ERDGImportedBufferFlags::CreateSRV).SRV;
-		Out.Resources.PositionOffsetBuffer	= Register(GraphBuilder, VFInput.Strands.PositionOffsetBufferExternal, ERDGImportedBufferFlags::CreateSRV).SRV;
-		Out.Resources.CurveBuffer			= Register(GraphBuilder, VFInput.Strands.CurveBufferExternal, ERDGImportedBufferFlags::CreateSRV).SRV;
-		Out.Resources.PointToCurveBuffer	= Register(GraphBuilder, VFInput.Strands.PointToCurveBufferExternal, ERDGImportedBufferFlags::CreateSRV).SRV;
-		Out.Resources.CurveAttributeBuffer	= Register(GraphBuilder, VFInput.Strands.CurveAttributeBufferExternal, ERDGImportedBufferFlags::CreateSRV).SRV;
-		Out.Resources.PointAttributeBuffer	= Register(GraphBuilder, VFInput.Strands.PointAttributeBufferExternal, ERDGImportedBufferFlags::CreateSRV).SRV;
+		OutResources.PositionBuffer			= Register(GraphBuilder, VFInput.Strands.PositionBufferExternal, ERDGImportedBufferFlags::CreateSRV).SRV;
+		OutResources.PositionOffsetBuffer	= Register(GraphBuilder, VFInput.Strands.PositionOffsetBufferExternal, ERDGImportedBufferFlags::CreateSRV).SRV;
+		OutResources.CurveBuffer			= Register(GraphBuilder, VFInput.Strands.CurveBufferExternal, ERDGImportedBufferFlags::CreateSRV).SRV;
+		OutResources.PointToCurveBuffer		= Register(GraphBuilder, VFInput.Strands.PointToCurveBufferExternal, ERDGImportedBufferFlags::CreateSRV).SRV;
+		OutResources.CurveAttributeBuffer	= Register(GraphBuilder, VFInput.Strands.CurveAttributeBufferExternal, ERDGImportedBufferFlags::CreateSRV).SRV;
+		OutResources.PointAttributeBuffer	= Register(GraphBuilder, VFInput.Strands.PointAttributeBufferExternal, ERDGImportedBufferFlags::CreateSRV).SRV;
 	}
 	else
 	{
-		Out.Resources.PositionBuffer		= VFInput.Strands.PositionBuffer.SRV;
-		Out.Resources.PositionOffsetBuffer	= VFInput.Strands.PositionOffsetBuffer.SRV;
-		Out.Resources.CurveBuffer			= VFInput.Strands.CurveBuffer.SRV;
-		Out.Resources.PointToCurveBuffer	= VFInput.Strands.PointToCurveBuffer.SRV;
-		Out.Resources.CurveAttributeBuffer	= VFInput.Strands.CurveAttributeBuffer.SRV;
-		Out.Resources.PointAttributeBuffer	= VFInput.Strands.PointAttributeBuffer.SRV;
+		OutResources.PositionBuffer			= VFInput.Strands.PositionBuffer.SRV;
+		OutResources.PositionOffsetBuffer	= VFInput.Strands.PositionOffsetBuffer.SRV;
+		OutResources.CurveBuffer			= VFInput.Strands.CurveBuffer.SRV;
+		OutResources.PointToCurveBuffer		= VFInput.Strands.PointToCurveBuffer.SRV;
+		OutResources.CurveAttributeBuffer	= VFInput.Strands.CurveAttributeBuffer.SRV;
+		OutResources.PointAttributeBuffer	= VFInput.Strands.PointAttributeBuffer.SRV;
 	}
+}
 
-	Out.Common = VFInput.Strands.Common;
-	Out.Culling.bCullingEnable = 0;
-
-	// Absolute local to world
-	Out.Common.LocalToWorldPrimitiveTransform	= FMatrix44f(VFInput.LocalToWorldTransform.ToMatrixWithScale()); // LWC_TODO: Precision loss, remove usage for this
-
-	// Translated local to world
-	const FVector& TranslatedWorldOffset = ViewInfo.ViewMatrices.GetPreViewTranslation();
-	FTransform LocalToTranslatedWorldTransform = VFInput.LocalToWorldTransform;
-	LocalToTranslatedWorldTransform.AddToTranslation(TranslatedWorldOffset);
-	Out.Common.LocalToTranslatedWorldPrimitiveTransform = FMatrix44f(LocalToTranslatedWorldTransform.ToMatrixWithScale());
+void GetHairStrandsInstanceCulling(FRDGBuilder& GraphBuilder, const FViewInfo& ViewInfo, const FHairGroupPublicData* HairGroupPublicData, bool bCullingEnable, FHairStrandsInstanceCullingParameters& OutCulling)
+{
+	check(HairGroupPublicData);
+	const FHairGroupPublicData::FVertexFactoryInput& VFInput = HairGroupPublicData->VFInput;
 
+	OutCulling.bCullingEnable = 0;
 	if (bCullingEnable)
 	{
 		FRDGImportedBuffer CullingIndirectBuffer = Register(GraphBuilder, HairGroupPublicData->GetDrawIndirectRasterComputeBuffer(), ERDGImportedBufferFlags::CreateSRV);
-		Out.Culling.bCullingEnable = 1;
-		Out.Culling.CullingIndirectBuffer = CullingIndirectBuffer.SRV;
-		Out.Culling.CullingIndexBuffer = RegisterAsSRV(GraphBuilder, HairGroupPublicData->GetCulledVertexIdBuffer());
-		Out.Culling.CullingRadiusScaleBuffer = RegisterAsSRV(GraphBuilder, HairGroupPublicData->GetCulledVertexRadiusScaleBuffer());
-		Out.Culling.CullingIndirectBufferArgs = CullingIndirectBuffer.Buffer;
+		OutCulling.bCullingEnable = 1;
+		OutCulling.CullingIndirectBuffer = CullingIndirectBuffer.SRV;
+		OutCulling.CullingIndexBuffer = RegisterAsSRV(GraphBuilder, HairGroupPublicData->GetCulledVertexIdBuffer());
+		OutCulling.CullingRadiusScaleBuffer = RegisterAsSRV(GraphBuilder, HairGroupPublicData->GetCulledVertexRadiusScaleBuffer());
+		OutCulling.CullingIndirectBufferArgs = CullingIndirectBuffer.Buffer;
 	}
+}
+
+FHairStrandsInstanceParameters GetHairStrandsInstanceParameters(FRDGBuilder& GraphBuilder, const FViewInfo& ViewInfo, const FHairGroupPublicData* HairGroupPublicData, bool bCullingEnable, bool bForceRegister)
+{
+	FHairStrandsInstanceParameters Output;
+	GetHairStrandsInstanceCommon(GraphBuilder, ViewInfo, HairGroupPublicData, Output.HairStrandsVF.Common);
+	GetHairStrandsInstanceResources(GraphBuilder, ViewInfo, HairGroupPublicData, bForceRegister, Output.HairStrandsVF.Resources);
+	GetHairStrandsInstanceCulling(GraphBuilder, ViewInfo, HairGroupPublicData, bCullingEnable, Output.HairStrandsVF.Culling);
 	return Output;
 }
 
diff --git a/Engine/Source/Runtime/Renderer/Public/HairStrandsInterface.h b/Engine/Source/Runtime/Renderer/Public/HairStrandsInterface.h
--- a/Engine/Source/Runtime/Renderer/Public/HairStrandsInterface.h
+++ b/Engine/Source/Runtime/Renderer/Public/HairStrandsInterface.h
@@ -33,6 +33,7 @@ BEGIN_SHADER_PARAMETER_STRUCT(FHairStrandsInstanceCommonParameters, RENDERER_API
 	SHADER_PARAMETER(float,  LengthScale)
 	SHADER_PARAMETER(float,  RaytracingRadiusScale)
 	SHADER_PARAMETER(uint32, GroupIndex)
+	SHADER_PARAMETER(uint32, GroupCount)
 	SHADER_PARAMETER(uint32, PointCount)
 	SHADER_PARAMETER(uint32, CurveCount)
 	SHADER_PARAMETER(uint32, RaytracingProceduralSplits)
