Fix for editor viewports not drawing post process scene color and depth correctly when both editor primitives/selection outline is enabled at the same time as debug draws.

Debug primitive draws will only output the debug color and depth values onto the viewport location given due to SimpleElementShaders being locked into a specific draw setup. As a result, if the viewport MinX/Y value is not 0, then the color and depth textures must be appropriately drawn out in an alternative way to the target render targets for the debug primitive composite to match the location in the viewport. This can likely be achieved with multiple calls to AddDrawTexturePass, but instead this change alters the PopulateDepthPass to allow for an alternative output format which draws the color and depth to the required viewport location.

For future work, it is likely the best way forward would be to rethink how debug primitives are drawn, as the SimpleElementShaders and their setup are locked in and difficult to alter. For now this is likely the best way to compensate until debug draws can be refactored entirely.

Also removed a redundant argument in TemporalUpscaleDepthPass.

diff --git a/Engine/Shaders/Private/PostProcessCompositePrimitives.usf b/Engine/Shaders/Private/PostProcessCompositePrimitives.usf
--- a/Engine/Shaders/Private/PostProcessCompositePrimitives.usf
+++ b/Engine/Shaders/Private/PostProcessCompositePrimitives.usf
@@ -254,12 +254,28 @@ void MainPopulateSceneDepthPS(
 		(Depth_ViewportMin + ViewportUV * Depth_ViewportSize - DepthTextureJitter) * Depth_ExtentInverse,
 		Depth_UVViewportBilinearMin, Depth_UVViewportBilinearMax);
 
-	OutDepth = Texture2DSampleLevel(DepthTexture, DepthSampler, DepthUV, 0).r;
-#if FORCE_DRAW_COLOR
-	OutColor = Texture2DSampleLevel(ColorTexture, ColorSampler, DepthUV, 0).rgba;
+#if FORCE_DEBUG_DRAW_COLOR
+	//This extended pass is primarily to compensate for the lack of scene color output when compositing the debug draw primitives with SimpleElementShaders
+	//Draw out the color whilst compensating for the viewport not necessarily starting at 0,0. 
+	//Does not use jitter for Scene Color as it is not upscaled for this pass
+	float2 ColorUV = clamp(ViewportUV * Color_UVViewportSize, 
+						Color_UVViewportBilinearMin - Color_UVViewportMin,
+						Color_UVViewportBilinearMax - Color_UVViewportMin);
+	ColorUV = clamp(ColorUV, float2(0, 0), float2(1, 1));
+	OutColor = Texture2DSampleLevel(ColorTexture, ColorSampler, ColorUV, 0).rgba;
+
+	//As this is also outputting color and positions need to match, compensate the depth for the viewport position not necessarily starting at 0,0
+	DepthUV = DepthUV - Depth_UVViewportMin;
+	//At larger viewport sizes, the temporal depth upscale can increase the depth viewport size significantly, 
+	//which isn't necessarily compensated for when outputting the color, so scale the DepthUV output down to match the ColorUV viewport size
+	//Clamped to the Min/Max temporal upscale values defined in PostProcessCompositePrimitivesCommon.cpp
+	DepthUV *= clamp(Depth_Extent / Color_Extent, MIN_DEPTHTEX_UPSCALE_FACTOR, MAX_DEPTHTEX_UPSCALE_FACTOR);
+	DepthUV = clamp(DepthUV, float2(0, 0), float2(1, 1));
 #else
 	OutColor = 0.0;
 #endif
+
+	OutDepth = Texture2DSampleLevel(DepthTexture, DepthSampler, DepthUV, 0).r;
 }
 
 float ComputeDepthMask(float SceneDeviceZ, float EditorDeviceZ)
diff --git a/Engine/Source/Runtime/Renderer/Private/PostProcess/PostProcessCompositeDebugPrimitives.cpp b/Engine/Source/Runtime/Renderer/Private/PostProcess/PostProcessCompositeDebugPrimitives.cpp
--- a/Engine/Source/Runtime/Renderer/Private/PostProcess/PostProcessCompositeDebugPrimitives.cpp
+++ b/Engine/Source/Runtime/Renderer/Private/PostProcess/PostProcessCompositeDebugPrimitives.cpp
@@ -18,24 +18,17 @@ FScreenPassTexture AddDebugPrimitivePass(
 	check(Inputs.SceneDepth.IsValid());
 
 	RDG_EVENT_SCOPE(GraphBuilder, "CompositeDebugPrimitives");
-	
-	const uint32 NumMSAASamples = View.GetSceneTextures().Config.NumSamples;
-	//Set sizing values to match the SceneColor size
-	const FIntRect ViewRect = Inputs.SceneColor.ViewRect;
-	const FViewInfo* DebugView = CreateCompositePrimitiveView(View, ViewRect, NumMSAASamples);
 
 	FScreenPassRenderTarget Output = Inputs.OverrideOutput;
-	FRDGTextureRef DebugPrimitiveColor = Output.Texture;
 	if (!Output.IsValid())
 	{
-		// Create a new output for the scene color
-		FRDGTextureDesc ColorDesc = Inputs.SceneColor.Texture->Desc;
-		ColorDesc.Flags |= TexCreate_ShaderResource | TexCreate_RenderTargetable;
-		DebugPrimitiveColor = GraphBuilder.CreateTexture(ColorDesc, TEXT("Debug.DrawPrimitivesColor"));
-		//Create a new depth output for the final draw
-		Output = FScreenPassRenderTarget(DebugPrimitiveColor, ViewRect, ERenderTargetLoadAction::ENoAction);
+		Output = FScreenPassRenderTarget::CreateFromInput(GraphBuilder, Inputs.SceneColor, View.GetOverwriteLoadAction(), TEXT("Debug.DrawPrimitivesColor"));
 	}
+	const uint32 NumMSAASamples = View.GetSceneTextures().Config.NumSamples;
+	//Set sizing values to match the SceneColor size
+	const FIntRect ViewRect = Output.ViewRect;
 	FIntPoint Extent = Output.Texture->Desc.Extent;
+	const FViewInfo* DebugView = CreateCompositePrimitiveView(View, ViewRect, NumMSAASamples);	
 
 	//Prepare output textures for composite draw
 	const FScreenPassTextureViewport OutputViewport(Output);
@@ -50,8 +43,7 @@ FScreenPassTexture AddDebugPrimitivePass(
 			*DebugView,
 			Inputs.SceneColor,
 			SceneDepth,
-			SceneDepthJitter,
-			NumMSAASamples);
+			SceneDepthJitter);
 	}
 
 	//Simple element pixel shaders do not output background color for composite, 
@@ -60,7 +52,7 @@ FScreenPassTexture AddDebugPrimitivePass(
 		*DebugView,
 		Inputs.SceneColor,
 		SceneDepth,
-		DebugPrimitiveColor,
+		Output.Texture,
 		DebugPrimitiveDepth,
 		SceneDepthJitter,
 		NumMSAASamples,
diff --git a/Engine/Source/Runtime/Renderer/Private/PostProcess/PostProcessCompositeEditorPrimitives.cpp b/Engine/Source/Runtime/Renderer/Private/PostProcess/PostProcessCompositeEditorPrimitives.cpp
--- a/Engine/Source/Runtime/Renderer/Private/PostProcess/PostProcessCompositeEditorPrimitives.cpp
+++ b/Engine/Source/Runtime/Renderer/Private/PostProcess/PostProcessCompositeEditorPrimitives.cpp
@@ -277,8 +277,7 @@ FScreenPassTexture AddEditorPrimitivePass(
 					*EditorView,
 					Inputs.SceneColor,
 					SceneDepth,
-					SceneDepthJitter,
-					NumMSAASamples);
+					SceneDepthJitter);
 			}
 
 			PopulateDepthPass(GraphBuilder,
diff --git a/Engine/Source/Runtime/Renderer/Private/PostProcess/PostProcessCompositePrimitivesCommon.cpp b/Engine/Source/Runtime/Renderer/Private/PostProcess/PostProcessCompositePrimitivesCommon.cpp
--- a/Engine/Source/Runtime/Renderer/Private/PostProcess/PostProcessCompositePrimitivesCommon.cpp
+++ b/Engine/Source/Runtime/Renderer/Private/PostProcess/PostProcessCompositePrimitivesCommon.cpp
@@ -7,6 +7,8 @@
 #include "PixelShaderUtils.h"
 #include "ScenePrivate.h"
 
+#define MIN_DEPTHTEX_UPSCALE_FACTOR 1
+#define MAX_DEPTHTEX_UPSCALE_FACTOR 4
 TAutoConsoleVariable<int32> CVarCompositeTemporalUpsampleDepth(
 	TEXT("r.Composite.TemporalUpsampleDepth"), 2,
 	TEXT("Temporal upsample factor of the depth buffer for depth testing editor primitives against."),
@@ -86,8 +88,8 @@ public:
 	SHADER_USE_PARAMETER_STRUCT(FPopulateCompositeDepthPS, FGlobalShader);
 
 	class FUseMSAADimension : SHADER_PERMUTATION_BOOL("USE_MSAA");
-	class FForceDrawColorOutput : SHADER_PERMUTATION_BOOL("FORCE_DRAW_COLOR"); //Allow the option to draw out the scene color too to lower draw calls
-	using FPermutationDomain = TShaderPermutationDomain<FUseMSAADimension, FForceDrawColorOutput>;
+	class FForceDebugDrawColorOutput : SHADER_PERMUTATION_BOOL("FORCE_DEBUG_DRAW_COLOR"); //Allow the option to draw out the scene color too to lower draw calls
+	using FPermutationDomain = TShaderPermutationDomain<FUseMSAADimension, FForceDebugDrawColorOutput>;
 
 	BEGIN_SHADER_PARAMETER_STRUCT(FParameters, )
 		SHADER_PARAMETER_STRUCT_REF(FViewUniformShaderParameters, View)
@@ -101,6 +103,18 @@ public:
 		RENDER_TARGET_BINDING_SLOTS()
 		END_SHADER_PARAMETER_STRUCT()
 
+	static void ModifyCompilationEnvironment(const FGlobalShaderPermutationParameters& InParameters, FShaderCompilerEnvironment& OutEnvironment)
+	{
+		FGlobalShader::ModifyCompilationEnvironment(InParameters, OutEnvironment);
+
+		FPermutationDomain PermutationVector(InParameters.PermutationId);
+		if (PermutationVector.Get<FForceDebugDrawColorOutput>() != 0)
+		{
+			OutEnvironment.SetDefine(TEXT("MIN_DEPTHTEX_UPSCALE_FACTOR"), MIN_DEPTHTEX_UPSCALE_FACTOR);
+			OutEnvironment.SetDefine(TEXT("MAX_DEPTHTEX_UPSCALE_FACTOR"), MAX_DEPTHTEX_UPSCALE_FACTOR);
+		}
+	}
+
 	static bool ShouldCompilePermutation(const FGlobalShaderPermutationParameters& Parameters)
 	{
 		const FPermutationDomain PermutationVector(Parameters.PermutationId);
@@ -125,8 +139,7 @@ void TemporalUpscaleDepthPass(
 	const FViewInfo& View,
 	const FScreenPassTexture& InSceneColor,
 	FScreenPassTexture& InOutSceneDepth,
-	FVector2f& SceneDepthJitter,
-	uint32 NumMSAASamples)
+	FVector2f& SceneDepthJitter)
 {
 	if (InSceneColor.ViewRect != InOutSceneDepth.ViewRect)
 	{
@@ -139,7 +152,7 @@ void TemporalUpscaleDepthPass(
 		// Upscale factor shouldn't be higher than there is TAA samples, or that means there will be unrendered pixels.
 		const int32 ComputeMaxUpsampleFactorDueToTAA = FMath::FloorToInt(FMath::Sqrt(float(View.TemporalJitterSequenceLength) / (UpscaleFactor * UpscaleFactor)));
 
-		const int32 MaxRHIUpsampleFactor = FMath::Clamp(FMath::FloorToInt(float(GetMax2DTextureDimension()) / float(Extent.GetMax())), 1, 4);
+		const int32 MaxRHIUpsampleFactor = FMath::Clamp(FMath::FloorToInt(float(GetMax2DTextureDimension()) / float(Extent.GetMax())), MIN_DEPTHTEX_UPSCALE_FACTOR, MAX_DEPTHTEX_UPSCALE_FACTOR);
 
 		const int32 ComputeMaxUpsampleFactor = FMath::Clamp(ComputeMaxUpsampleFactorDueToTAA, 0, MaxRHIUpsampleFactor);
 
@@ -235,8 +248,8 @@ void PopulateDepthPass(FRDGBuilder& GraphBuilder,
 
 	FPopulateCompositeDepthPS::FParameters* PassParameters = GraphBuilder.AllocParameters<FPopulateCompositeDepthPS::FParameters>();
 	PassParameters->View = View.ViewUniformBuffer;
-	PassParameters->Color = GetScreenPassTextureViewportParameters(FScreenPassTextureViewport(InSceneColor));
-	PassParameters->Depth = GetScreenPassTextureViewportParameters(FScreenPassTextureViewport(InSceneDepth));
+	PassParameters->Color = GetScreenPassTextureViewportParameters(FScreenPassTextureViewport(InSceneColor.Texture->Desc.Extent, View.ViewRect));
+	PassParameters->Depth = GetScreenPassTextureViewportParameters(FScreenPassTextureViewport(InSceneDepth.Texture->Desc.Extent, View.ViewRect));
 	if (bForceDrawColor)
 	{
 		PassParameters->ColorTexture = InSceneColor.Texture;
@@ -250,7 +263,7 @@ void PopulateDepthPass(FRDGBuilder& GraphBuilder,
 
 	FPopulateCompositeDepthPS::FPermutationDomain PermutationVector;
 	PermutationVector.Set<FPopulateCompositeDepthPS::FUseMSAADimension>(NumMSAASamples > 1);
-	PermutationVector.Set< FPopulateCompositeDepthPS::FForceDrawColorOutput>(bForceDrawColor);
+	PermutationVector.Set< FPopulateCompositeDepthPS::FForceDebugDrawColorOutput>(bForceDrawColor);
 	TShaderMapRef<FPopulateCompositeDepthPS> PixelShader(View.ShaderMap, PermutationVector);
 
 	FPixelShaderUtils::AddFullscreenPass(
diff --git a/Engine/Source/Runtime/Renderer/Private/PostProcess/PostProcessCompositePrimitivesCommon.h b/Engine/Source/Runtime/Renderer/Private/PostProcess/PostProcessCompositePrimitivesCommon.h
--- a/Engine/Source/Runtime/Renderer/Private/PostProcess/PostProcessCompositePrimitivesCommon.h
+++ b/Engine/Source/Runtime/Renderer/Private/PostProcess/PostProcessCompositePrimitivesCommon.h
@@ -76,8 +76,7 @@ void TemporalUpscaleDepthPass(
 	const FViewInfo& View,
 	const FScreenPassTexture& InSceneColor,
 	FScreenPassTexture& InOutSceneDepth,
-	FVector2f& SceneDepthJitter,
-	uint32 NumMSAASamples);
+	FVector2f& SceneDepthJitter);
 
 void PopulateDepthPass(FRDGBuilder& GraphBuilder,
 	const FViewInfo& View,
