Implements a r.TSR.ShadingRejection.TileOverscan to reducing tiling artifact on high, epic and cinematic scalability

diff --git a/Engine/Config/BaseScalability.ini b/Engine/Config/BaseScalability.ini
--- a/Engine/Config/BaseScalability.ini
+++ b/Engine/Config/BaseScalability.ini
@@ -44,6 +44,7 @@ r.TSR.History.ScreenPercentage=100
 r.TSR.History.UpdateQuality=0
 r.TSR.History.GrandReprojection=0
 r.TSR.ShadingRejection.Flickering=0
+r.TSR.ShadingRejection.TileOverscan=2
 r.TSR.Velocity.Extrapolation=0
 r.TSR.RejectionAntiAliasingQuality=0
 
@@ -55,6 +56,7 @@ r.TSR.History.ScreenPercentage=100
 r.TSR.History.UpdateQuality=1
 r.TSR.History.GrandReprojection=0
 r.TSR.ShadingRejection.Flickering=0
+r.TSR.ShadingRejection.TileOverscan=2
 r.TSR.Velocity.Extrapolation=0
 r.TSR.RejectionAntiAliasingQuality=1
 
@@ -66,6 +68,7 @@ r.TSR.History.ScreenPercentage=100
 r.TSR.History.UpdateQuality=2
 r.TSR.History.GrandReprojection=0
 r.TSR.ShadingRejection.Flickering=1
+r.TSR.ShadingRejection.TileOverscan=3
 r.TSR.Velocity.Extrapolation=0
 r.TSR.RejectionAntiAliasingQuality=1
 
@@ -77,6 +80,7 @@ r.TSR.History.ScreenPercentage=200
 r.TSR.History.UpdateQuality=3
 r.TSR.History.GrandReprojection=0
 r.TSR.ShadingRejection.Flickering=1
+r.TSR.ShadingRejection.TileOverscan=3
 r.TSR.Velocity.Extrapolation=1
 r.TSR.RejectionAntiAliasingQuality=2
 
@@ -88,6 +92,7 @@ r.TSR.History.ScreenPercentage=200
 r.TSR.History.UpdateQuality=3
 r.TSR.History.GrandReprojection=1
 r.TSR.ShadingRejection.Flickering=1
+r.TSR.ShadingRejection.TileOverscan=3
 r.TSR.Velocity.Extrapolation=1
 r.TSR.RejectionAntiAliasingQuality=2
 
diff --git a/Engine/Shaders/Private/TemporalSuperResolution/TSRRejectShading.usf b/Engine/Shaders/Private/TemporalSuperResolution/TSRRejectShading.usf
--- a/Engine/Shaders/Private/TemporalSuperResolution/TSRRejectShading.usf
+++ b/Engine/Shaders/Private/TemporalSuperResolution/TSRRejectShading.usf
@@ -20,12 +20,6 @@
 /** Size of the entire group loaded into memory including overscan. */
 #define GROUP_TILE_SIZE 16
 
-/** Size of the overscan. */
-#define TILE_OVERSCAN_PADDING 2
-
-/** Size of of the output tile. */
-#define TILE_SIZE (GROUP_TILE_SIZE - 2 * TILE_OVERSCAN_PADDING)
-
 /** Whether should remove firefliers from input and history before comparison to improve stability. */
 #define CONFIG_FIREFLY_REMOVAL 1
 
@@ -103,6 +97,7 @@ float2 TranslucencyTextureUVMax;
 float3 HistoryGuideQuantizationError;
 float FlickeringFramePeriod;
 float TheoricBlendFactor;
+uint TileOverscan;
 
 
 Texture2D<tsr_halfC> InputTexture;
@@ -160,8 +155,8 @@ void MainCS(
 	const uint LaneIndex = GroupThreadIndex;
 	const tsr_short2 GroupPixelOffset = (
 		tsr_short2(InputPixelPosMin) +
-		tsr_short2(GroupId) * tsr_short2(TILE_SIZE, TILE_SIZE) +
-		tsr_short(-TILE_OVERSCAN_PADDING).xx);
+		tsr_short2(GroupId) * tsr_short(GROUP_TILE_SIZE - 2 * TileOverscan).xx +
+		tsr_short(-TileOverscan).xx);
 		
 	// Load
 	tsr_simd_halfC OriginalOpaqueInput;
@@ -561,7 +556,7 @@ void MainCS(
 		for (uint SimdIndex = 0; SimdIndex < SIMD_SIZE; SimdIndex++)
 		{
 			tsr_short2 LaneSimdPixelOffset = GetLaneSimdPixelOffset<LANE_STRIDE_X, LANE_STRIDE_Y>(LaneIndex, SimdIndex);
-			bool bIsNonOverscanPixel = all(and(LaneSimdPixelOffset >= TILE_OVERSCAN_PADDING, LaneSimdPixelOffset < TILE_SIZE + TILE_OVERSCAN_PADDING));
+			bool bIsNonOverscanPixel = all(and(LaneSimdPixelOffset >= TileOverscan, LaneSimdPixelOffset < (GROUP_TILE_SIZE - TileOverscan)));
 
 			tsr_short2 OutputPixelPos = InvalidateOutputPixelPos(GroupPixelOffset + LaneSimdPixelOffset, InputInfo_ViewportMax);
 			OutputPixelPos.x = select(bIsNonOverscanPixel, OutputPixelPos.x, ~tsr_short(0));
diff --git a/Engine/Source/Runtime/Renderer/Private/PostProcess/TemporalSuperResolution.cpp b/Engine/Source/Runtime/Renderer/Private/PostProcess/TemporalSuperResolution.cpp
--- a/Engine/Source/Runtime/Renderer/Private/PostProcess/TemporalSuperResolution.cpp
+++ b/Engine/Source/Runtime/Renderer/Private/PostProcess/TemporalSuperResolution.cpp
@@ -149,6 +149,14 @@ TAutoConsoleVariable<float> CVarTSRFlickeringMaxParralaxVelocity(
 	TEXT("(Default to 10 pixels 1080p).\n"),
 	ECVF_RenderThreadSafe);
 
+TAutoConsoleVariable<int32> CVarTSRShadingTileOverscan(
+	TEXT("r.TSR.ShadingRejection.TileOverscan"), 3,
+	TEXT("The shading rejection run a network of convolutions on the GPU all in single 16x16 without roundtrip to main memory. ")
+	TEXT("However chaining many convlutions in this tiles means that some convolutions on the edge arround are becoming corrupted ")
+	TEXT("and therefor need to overlap the tile by couple of padding to hide it. This controls is controled by default in the ")
+	TEXT(" anti-aliasing scalability settings."),
+	ECVF_Scalability | ECVF_RenderThreadSafe);
+
 TAutoConsoleVariable<int32> CVarTSRRejectionAntiAliasingQuality(
 	TEXT("r.TSR.RejectionAntiAliasingQuality"), 3,
 	TEXT("Controls the quality of TSR's built-in spatial anti-aliasing technology when the history needs to be rejected. ")
@@ -695,6 +703,7 @@ class FTSRRejectShadingCS : public FTSRShader
 		SHADER_PARAMETER(FVector3f, HistoryGuideQuantizationError)
 		SHADER_PARAMETER(float, FlickeringFramePeriod)
 		SHADER_PARAMETER(float, TheoricBlendFactor)
+		SHADER_PARAMETER(int32, TileOverscan)
 
 		SHADER_PARAMETER_RDG_TEXTURE(Texture2D, InputTexture)
 		SHADER_PARAMETER_RDG_TEXTURE(Texture2D, InputMoireLumaTexture)
@@ -1906,6 +1915,8 @@ ITemporalUpscaler::FOutputs AddTemporalSuperResolutionPasses(
 	FRDGTextureRef HistoryRejectionTexture;
 	FRDGTextureRef InputSceneColorLdrLumaTexture = nullptr;
 	{
+		const int32 GroupTileSize = 16;
+
 		{
 			FRDGTextureDesc Desc = FRDGTextureDesc::Create2D(
 				InputExtent,
@@ -1939,6 +1950,7 @@ ITemporalUpscaler::FOutputs AddTemporalSuperResolutionPasses(
 		PassParameters->HistoryGuideQuantizationError = ComputePixelFormatQuantizationError(History.Guide->Desc.Format);
 		PassParameters->FlickeringFramePeriod = FlickeringFramePeriod;
 		PassParameters->TheoricBlendFactor = 1.0f / (1.0f + MaxHistorySampleCount / OutputToInputResolutionFractionSquare);
+		PassParameters->TileOverscan = FMath::Clamp(CVarTSRShadingTileOverscan.GetValueOnRenderThread(), 2, GroupTileSize / 2 - 1);
 
 		PassParameters->InputTexture = PassInputs.SceneColorTexture;
 		if (PassInputs.MoireInputTexture.IsValid())
@@ -1969,6 +1981,8 @@ ITemporalUpscaler::FOutputs AddTemporalSuperResolutionPasses(
 		PassParameters->InputSceneColorLdrLumaOutput = GraphBuilder.CreateUAV(InputSceneColorLdrLumaTexture);
 		PassParameters->DebugOutput = CreateDebugUAV(InputExtent, TEXT("Debug.TSR.RejectShading"));
 
+		int32 TileSize = GroupTileSize - 2 * PassParameters->TileOverscan;
+
 		FTSRRejectShadingCS::FPermutationDomain PermutationVector;
 		PermutationVector.Set<FTSRRejectShadingCS::FWaveSizeOps>(bUseWaveOps && GRHIMinimumWaveSize >= 32 && GRHIMinimumWaveSize <= 64 ? GRHIMinimumWaveSize : 0);
 		PermutationVector.Set<FTSRRejectShadingCS::FFlickeringDetectionDim>(FlickeringFramePeriod > 0.0f);
@@ -1978,7 +1992,9 @@ ITemporalUpscaler::FOutputs AddTemporalSuperResolutionPasses(
 		TShaderMapRef<FTSRRejectShadingCS> ComputeShader(View.ShaderMap, PermutationVector);
 		FComputeShaderUtils::AddPass(
 			GraphBuilder,
-			RDG_EVENT_NAME("TSR RejectShading(WaveSize=%d FlickeringFramePeriod=%f %s%s) %dx%d",
+			RDG_EVENT_NAME("TSR RejectShading(TileSize=%d PaddingCostMultiplier=%1.1f WaveSize=%d FlickeringFramePeriod=%f %s%s) %dx%d",
+				TileSize,
+				FMath::Pow(float(GroupTileSize) / float(TileSize), 2),
 				int32(PermutationVector.Get<FTSRRejectShadingCS::FWaveSizeOps>()),
 				FlickeringFramePeriod,
 				PermutationVector.Get<FTSRShader::F16BitVALUDim>() ? TEXT(" 16bit") : TEXT(""),
@@ -1987,7 +2003,7 @@ ITemporalUpscaler::FOutputs AddTemporalSuperResolutionPasses(
 			AsyncComputePasses >= 3 ? ERDGPassFlags::AsyncCompute : ERDGPassFlags::Compute,
 			ComputeShader,
 			PassParameters,
-			FComputeShaderUtils::GetGroupCount(InputRect.Size(), 12));
+			FComputeShaderUtils::GetGroupCount(InputRect.Size(), TileSize));
 	}
 
 	// Spatial anti-aliasing when doing history rejection.
@@ -2112,7 +2128,7 @@ ITemporalUpscaler::FOutputs AddTemporalSuperResolutionPasses(
 		PassParameters->GrandPrevColorTexture = GrandPrevColorTexture;
 
 		PassParameters->HistoryOutput = CreateUAVs(GraphBuilder, History);
-		PassParameters->SceneColorOutputMip0 = GraphBuilder.CreateUAV(FRDGTextureUAVDesc(UpdateHistoryOutputTexture, /* InMipLevel = */ 0));
+			PassParameters->SceneColorOutputMip0 = GraphBuilder.CreateUAV(FRDGTextureUAVDesc(UpdateHistoryOutputTexture, /* InMipLevel = */ 0));
 		
 		if (PassInputs.bGenerateOutputMip1 && HistorySize == OutputRect.Size())
 		{
