Fixes crash in TSR when doing SM5 previsualization on -d3d12 -sm6

diff --git a/Engine/Source/Runtime/Renderer/Private/PostProcess/TemporalSuperResolution.cpp b/Engine/Source/Runtime/Renderer/Private/PostProcess/TemporalSuperResolution.cpp
--- a/Engine/Source/Runtime/Renderer/Private/PostProcess/TemporalSuperResolution.cpp
+++ b/Engine/Source/Runtime/Renderer/Private/PostProcess/TemporalSuperResolution.cpp
@@ -11,6 +11,7 @@
 #include "PixelShaderUtils.h"
 #include "ScenePrivate.h"
 #include "RendererModule.h"
+#include "ShaderPlatformCachedIniValue.h"
 
 #define COMPILE_TSR_DEBUG_PASSES (!UE_BUILD_SHIPPING)
 
@@ -469,7 +470,7 @@ public:
 	static ERHIFeatureSupport SupportsWaveOps(EShaderPlatform Platform)
 	{
 		// UE-161125: SPIRV ends up with 5min+ compilation times, and D3D11 needs SPIRV backend for HLSL2021 features on FXC.
-		if (FDataDrivenShaderPlatformInfo::GetIsSPIRV(Platform) || Platform == SP_PCD3D_SM5)
+		if (FDataDrivenShaderPlatformInfo::GetIsSPIRV(Platform) || GetEditorShaderPlatform(Platform) == SP_PCD3D_SM5)
 		{
 			return ERHIFeatureSupport::Unsupported;
 		}
@@ -989,7 +990,7 @@ class FTSRResolveHistoryCS : public FTSRShader
 
 		if (PermutationVector.Get<FNyquistDim>())
 		{
-			if (FTSRShader::SupportsWaveOps(Parameters.Platform) != ERHIFeatureSupport::RuntimeGuaranteed)
+			if (FTSRShader::SupportsWaveOps(Parameters.Platform) == ERHIFeatureSupport::Unsupported)
 			{
 				return false;
 			}
