Uses the same scene color format on mobile distortion pass.

diff --git a/Engine/Source/Runtime/Renderer/Private/MobileDistortionPass.cpp b/Engine/Source/Runtime/Renderer/Private/MobileDistortionPass.cpp
--- a/Engine/Source/Runtime/Renderer/Private/MobileDistortionPass.cpp
+++ b/Engine/Source/Runtime/Renderer/Private/MobileDistortionPass.cpp
@@ -128,7 +128,7 @@ IMPLEMENT_GLOBAL_SHADER(FMobileDistortionMergePS, "/Engine/Private/DistortApplyS
 
 FScreenPassTexture AddMobileDistortionMergePass(FRDGBuilder& GraphBuilder, const FViewInfo& View, const FMobileDistortionMergeInputs& Inputs)
 {
-	FRDGTextureDesc DistortionMergeDesc = FRDGTextureDesc::Create2D(Inputs.DistortionAccumulate.Texture->Desc.Extent, PF_FloatRGBA, FClearValueBinding::Black, TexCreate_ShaderResource | TexCreate_RenderTargetable);
+	FRDGTextureDesc DistortionMergeDesc = FRDGTextureDesc::Create2D(Inputs.DistortionAccumulate.Texture->Desc.Extent, Inputs.SceneColor.Texture->Desc.Format, FClearValueBinding::Black, TexCreate_ShaderResource | TexCreate_RenderTargetable);
 
 	FScreenPassRenderTarget DistortionMergeOutput = FScreenPassRenderTarget(GraphBuilder.CreateTexture(DistortionMergeDesc, TEXT("DistortionMergePass")), Inputs.DistortionAccumulate.ViewRect, ERenderTargetLoadAction::EClear);
 
