Update directional light uniform buffers in InitViews on mobile.

diff --git a/Engine/Source/Runtime/Renderer/Private/MobileShadingRenderer.cpp b/Engine/Source/Runtime/Renderer/Private/MobileShadingRenderer.cpp
--- a/Engine/Source/Runtime/Renderer/Private/MobileShadingRenderer.cpp
+++ b/Engine/Source/Runtime/Renderer/Private/MobileShadingRenderer.cpp
@@ -717,6 +717,8 @@ void FMobileSceneRenderer::InitViews(
 
 	// Now that the indirect lighting cache is updated, we can update the uniform buffers.
 	UpdatePrimitiveIndirectLightingCacheBuffers();
+
+	UpdateDirectionalLightUniformBuffers(GraphBuilder, Views[0]);
 	
 	OnStartRender(RHICmdList);
 }
