Fix a d3d and vulkan validation failures with VisualizeTemporalUpscaler on the CopyCommandList hitting RHI edge cases with the backbuffer

diff --git a/Engine/Source/Runtime/Renderer/Private/PostProcess/VisualizeTemporalUpscaler.cpp b/Engine/Source/Runtime/Renderer/Private/PostProcess/VisualizeTemporalUpscaler.cpp
--- a/Engine/Source/Runtime/Renderer/Private/PostProcess/VisualizeTemporalUpscaler.cpp
+++ b/Engine/Source/Runtime/Renderer/Private/PostProcess/VisualizeTemporalUpscaler.cpp
@@ -27,33 +27,11 @@ FScreenPassTexture AddVisualizeTemporalUpscalerPass(FRDGBuilder& GraphBuilder, c
 		Output = FScreenPassRenderTarget::CreateFromInput(GraphBuilder, Inputs.SceneColor, View.GetOverwriteLoadAction(), TEXT("MotionVectors.Visualize"));
 	}
 
-	// Early return if not using a temporal upscaler.
-	if (!Inputs.UpscalerUsed)
-	{
-		check(Inputs.TAAConfig == EMainTAAPassConfig::Disabled);
-		FRHICopyTextureInfo CopyInfo;
-		AddCopyTexturePass(
-			GraphBuilder,
-			Inputs.SceneColor.Texture,
-			Output.Texture,
-			CopyInfo);
-
-		AddDrawCanvasPass(GraphBuilder, RDG_EVENT_NAME("VisualizeTemporalUpscaler Text"), View, FScreenPassRenderTarget(Output, ERenderTargetLoadAction::ELoad),
-			[&ViewRect = Output.ViewRect](FCanvas& Canvas)
-		{
-			const float DPIScale = Canvas.GetDPIScale();
-			Canvas.SetBaseTransform(FMatrix(FScaleMatrix(DPIScale) * Canvas.CalcBaseTransform2D(Canvas.GetViewRect().Width(), Canvas.GetViewRect().Height())));
-
-			FIntPoint LabelLocation(60, 60);
-			Canvas.DrawShadowedString(LabelLocation.X / DPIScale, LabelLocation.Y / DPIScale, TEXT("No temporal upscaler used"), GetStatsFont(), FLinearColor::Red);
-		});
-
-		return MoveTemp(Output);
-	}
-
 	// Populate all the tiles
 	TArray<FVisualizeBufferTile> Tiles;
 	Tiles.SetNum(16);
+
+	if (Inputs.UpscalerUsed)
 	{
 		auto VisualizeTextureLabel = [](FRDGTextureRef Texture, const TCHAR* Suffix = TEXT(""))
 		{
@@ -193,6 +171,8 @@ FScreenPassTexture AddVisualizeTemporalUpscalerPass(FRDGBuilder& GraphBuilder, c
 	}
 
 	// Draw additional text
+	// Early return if not using a temporal upscaler.
+	if (Inputs.UpscalerUsed)
 	{
 		AddDrawCanvasPass(GraphBuilder, RDG_EVENT_NAME("VisualizeTemporalUpscaler Text"), View, FScreenPassRenderTarget(Output, ERenderTargetLoadAction::ELoad),
 			[&View, &ViewRect = Output.ViewRect, &Inputs, bSupportsAlpha](FCanvas& Canvas)
@@ -259,6 +239,19 @@ FScreenPassTexture AddVisualizeTemporalUpscalerPass(FRDGBuilder& GraphBuilder, c
 		});
 
 	}
+	else
+	{
+		check(Inputs.TAAConfig == EMainTAAPassConfig::Disabled);
+		AddDrawCanvasPass(GraphBuilder, RDG_EVENT_NAME("VisualizeTemporalUpscaler Text"), View, FScreenPassRenderTarget(Output, ERenderTargetLoadAction::ELoad),
+			[&ViewRect = Output.ViewRect](FCanvas& Canvas)
+			{
+				const float DPIScale = Canvas.GetDPIScale();
+				Canvas.SetBaseTransform(FMatrix(FScaleMatrix(DPIScale) * Canvas.CalcBaseTransform2D(Canvas.GetViewRect().Width(), Canvas.GetViewRect().Height())));
+
+				FIntPoint LabelLocation(60, 60);
+				Canvas.DrawShadowedString(LabelLocation.X / DPIScale, LabelLocation.Y / DPIScale, TEXT("No temporal upscaler used"), GetStatsFont(), FLinearColor::Red);
+			});
+	}
 
 	return MoveTemp(Output);
 }
