- Disable low latency on translucent materials that write to custom depth

diff --git a/Engine/Plugins/FX/Niagara/Source/Niagara/Private/NiagaraRendererSprites.cpp b/Engine/Plugins/FX/Niagara/Source/Niagara/Private/NiagaraRendererSprites.cpp
--- a/Engine/Plugins/FX/Niagara/Source/Niagara/Private/NiagaraRendererSprites.cpp
+++ b/Engine/Plugins/FX/Niagara/Source/Niagara/Private/NiagaraRendererSprites.cpp
@@ -251,6 +251,7 @@ void FNiagaraRendererSprites::PrepareParticleSpriteRenderData(FParticleSpriteRen
 		ParticleSpriteRenderData.bHasTranslucentMaterials &&
 		bGpuLowLatencyTranslucency &&
 		GpuReadyTickStage >= CurrentParticleData->GetGPUDataReadyStage() &&
+		Material.IsTranslucencyWritingCustomDepth() &&
 		!SceneProxy->CastsVolumetricTranslucentShadow() &&
 		ViewFamilySupportLowLatencyTranslucency(ViewFamily);
 
diff --git a/Engine/Source/Runtime/Renderer/Private/FXRenderingUtils.cpp b/Engine/Source/Runtime/Renderer/Private/FXRenderingUtils.cpp
--- a/Engine/Source/Runtime/Renderer/Private/FXRenderingUtils.cpp
+++ b/Engine/Source/Runtime/Renderer/Private/FXRenderingUtils.cpp
@@ -26,11 +26,9 @@ bool UE::FXRenderingUtils::CanMaterialRenderBeforeFXPostOpaque(
 	const FPrimitiveSceneProxy& SceneProxy,
 	const FMaterial& Material)
 {
-	const bool bTranslucent = IsTranslucentBlendMode(Material);
-
-	if (!bTranslucent)
+	// Opaque materials & translucent that write custom depth will always need to render before FFXSystemInterface::PostOpaqueRender
+	if (!IsTranslucentBlendMode(Material) || Material.IsTranslucencyWritingCustomDepth())
 	{
-		// opaque materials always render before FFXSystemInterface::PostOpaqueRender
 		return true;
 	}
 
