Add depth pass for path tracing. The output is in DepthZ. The user can use the same depth post process material to output in MRQ. The user can use `r.PathTracing.Override.Depth 0` to disable path traced depth buffer and use the depth from the rasterizer instead.

diff --git a/Engine/Shaders/Private/PathTracing/Material/PathTracingMaterialSampling.ush b/Engine/Shaders/Private/PathTracing/Material/PathTracingMaterialSampling.ush
--- a/Engine/Shaders/Private/PathTracing/Material/PathTracingMaterialSampling.ush
+++ b/Engine/Shaders/Private/PathTracing/Material/PathTracingMaterialSampling.ush
@@ -243,6 +243,65 @@ void AccumulateAlbedo(float3 SigmaS, float3 PathThroughput, float PathRoughness,
 	Albedo += SigmaS * PathThroughput;
 }
 
+#define DEBUG_DEPTH_SAMPLING 0
+
+struct FDepthContext
+{
+	float3 TranslatedWorldPos;
+	float Transparency;
+	uint HitCount;
+
+#if DEBUG_DEPTH_SAMPLING
+	// visualize the depth status for the ith ray hit.
+	int DebugHitCount;
+#endif
+
+	void AddHit(FPathTracingPayload Payload)
+	{
+		if (false
+#if DEBUG_DEPTH_SAMPLING
+			|| DebugHitCount == 0
+#endif
+			)
+		{
+			return;
+		}
+
+		
+		// The smaller the transparency color, the more opaque the material is.
+		// Ensure Transparency <= 0.5 specified in the material editor contributes to the depth.
+		//const float HalfTransparency = Luminance(PayloadDecodeLDRColor(PayloadEncodeLDRColor(0.5f)));
+		//  asfloat(0x3f0011c5); e.g., opacity >=0.4996 in the material editor for translucent material contributes to depth.
+		const float HalfTransparency = 0.500271142f;
+		const float HitTransparency = Luminance(Payload.TransparencyColor);
+
+		// ensure that 0.5 in the Material editor has a depth contribution.
+		const bool bShouldAcceptNewDepth =
+			Transparency > HalfTransparency && // current depth is not valid
+			HitTransparency < Transparency;    // since it's not valid, we should accept hit with smaller transparency color
+
+		if (bShouldAcceptNewDepth || HitCount == 0)
+		{
+			TranslatedWorldPos = Payload.TranslatedWorldPos;
+			Transparency = HitTransparency;
+			HitCount++;
+		}
+
+#if DEBUG_DEPTH_SAMPLING
+		DebugHitCount--;
+#endif
+	}
+};
+
+FDepthContext InitDepthContext()
+{
+	FDepthContext Context = (FDepthContext)0;
+#if DEBUG_DEPTH_SAMPLING
+	Context.DebugHitCount = View.GeneralPurposeTweak;
+#endif
+	return Context;
+}
+
 // returns an average roughness for the whole material (for the path roughness of shadow rays from this material)
 float GetAverageRoughness(FPathTracingPayload Payload)
 {
diff --git a/Engine/Shaders/Private/PathTracing/PathTracing.usf b/Engine/Shaders/Private/PathTracing/PathTracing.usf
--- a/Engine/Shaders/Private/PathTracing/PathTracing.usf
+++ b/Engine/Shaders/Private/PathTracing/PathTracing.usf
@@ -111,7 +111,10 @@ RAY_TRACING_ENTRY_RAYGEN(PathTracingMainRG)
 	FPathState PathState;
 	if (Bounce == 0)
 	{
-		PathState = CreatePathState(ComputePixelIndex(DispatchIdx));
+		PathState = CreatePathState(ComputePixelIndex(DispatchIdx), 
+									ComputeTextureIndex(uint2(LinearPixelIndex % DispatchDim.x,
+															  LinearPixelIndex / DispatchDim.x))
+								   );
 	}
 	else
 	{
@@ -147,8 +150,9 @@ RAY_TRACING_ENTRY_RAYGEN(PathTracingMainRG)
 {
 	const uint2 DispatchIdx = DispatchRaysIndex().xy;
 	const uint2 DispatchDim = DispatchRaysDimensions().xy;
+	int2 TextureIndex = ComputeTextureIndex(DispatchIdx);
 
-	FPathState PathState = CreatePathState(ComputePixelIndex(DispatchIdx));
+	FPathState PathState = CreatePathState(ComputePixelIndex(DispatchIdx), TextureIndex);
 
 	for (int Bounce = 0; Bounce <= MaxBounces; Bounce++)
 	{
@@ -160,7 +164,7 @@ RAY_TRACING_ENTRY_RAYGEN(PathTracingMainRG)
 	}
 
 	// Accumulate radiance and update pixel variance
-	PathState.WritePixel(ComputeTextureIndex(DispatchIdx));
+	PathState.WritePixel();
 }
 
 #endif
\ No newline at end of file
diff --git a/Engine/Shaders/Private/PathTracing/PathTracingCompositingPixelShader.usf b/Engine/Shaders/Private/PathTracing/PathTracingCompositingPixelShader.usf
--- a/Engine/Shaders/Private/PathTracing/PathTracingCompositingPixelShader.usf
+++ b/Engine/Shaders/Private/PathTracing/PathTracingCompositingPixelShader.usf
@@ -10,7 +10,7 @@ uint Iteration;
 uint MaxSamples;
 int ProgressDisplayEnabled;
 Texture2D<float4> RadianceTexture;
-
+Texture2D<float4> NormalDepthTexture;
 
 float TonemapGammaAndHDR(float Value)
 {
@@ -37,7 +37,8 @@ float Tonemap(float Value)
 
 void CompositeMain(
 	in noperspective float2 UV : TEXCOORD0,
-	out float4 OutColor : SV_Target0
+	out float4 OutColor : SV_Target0,
+	out float OutDepth : SV_DEPTH
 )
 {
 	float2 BufferSize = View.BufferSizeAndInvSize.xy;
@@ -103,4 +104,6 @@ void CompositeMain(
 		OutColor = 1 - Radiance.a; // visualize alpha channel
 	}
 #endif
+	
+	OutDepth = NormalDepthTexture.Load(TexCoord).w;
 }
diff --git a/Engine/Shaders/Private/PathTracing/PathTracingCore.ush b/Engine/Shaders/Private/PathTracing/PathTracingCore.ush
--- a/Engine/Shaders/Private/PathTracing/PathTracingCore.ush
+++ b/Engine/Shaders/Private/PathTracing/PathTracingCore.ush
@@ -87,6 +87,10 @@ struct FPathState
 	uint FirstScatterType;          //  3 bits - packed in sign bits of SigmaT
 									// 80 bytes total (see FPathTracingPackedPathState)
 
+	// Temporary parameters not stored				
+	bool DepthResamplingWeight;		
+	int2 PrimaryRayTextureIndex;		// texture index cache to write to the frame buffer for primary ray operation like depth.	
+
 	bool HasMadeContributionScatter()
 	{
 		// The first scattering type that matters for the PATHTRACER_CONTRIBUTION_* flags is the first non-refractive one
@@ -173,7 +177,27 @@ struct FPathState
 		const float4 PixelValue = float4(Radiance, saturate(BackgroundVisibility));
 		RadianceTexture[TextureIndex] = lerp(OldPixel, PixelValue, BlendFactor);
 		AlbedoTexture[TextureIndex] = lerp(OldAlbedo, float4(Albedo, 0), BlendFactor);
-		NormalTexture[TextureIndex] = lerp(OldNormal, float4(Normal, 0), BlendFactor);
+		NormalTexture[TextureIndex].xyz = lerp(OldNormal.xyz,Normal.xyz, BlendFactor);
+	}
+
+	void WritePixel()
+	{
+		WritePixel(PrimaryRayTextureIndex);
+	}
+
+	void WriteDepth(FDepthContext DepthContext)
+	{
+		float OldDepth = Iteration > 0 ? NormalTexture[PrimaryRayTextureIndex].w : 0;
+		
+		float Depth = 0;
+		
+		if (DepthContext.HitCount > 0)
+		{
+			const float3 ViewZ = View.ViewToTranslatedWorld[2].xyz;
+			Depth = dot(DepthContext.TranslatedWorldPos, ViewZ);
+		}
+
+		NormalTexture[PrimaryRayTextureIndex].w = lerp(OldDepth, max(OldDepth, ConvertToDeviceZ(Depth)), DepthResamplingWeight);
 	}
 };
 
@@ -340,6 +364,7 @@ FPathTracingPayload TraceTransparentRay(inout FPathState PathState, bool bIsCame
 	float2 RISRandSample = RandomSequence_GenerateSample2D(PathState.RandSequence);
 	FRISContext HitSample = InitRISContext(RISRandSample.x);
 	FRISContext VolSample = InitRISContext(RISRandSample.y);
+	FDepthContext DepthContext = InitDepthContext();
 
 	float3 PayloadThroughput = PathState.PathThroughput;
 	FPathTracingPayload Payload;
@@ -632,6 +657,11 @@ FPathTracingPayload TraceTransparentRay(inout FPathState PathState, bool bIsCame
 			break;
 		}
 
+		if (bIsCameraRay)
+		{
+			DepthContext.AddHit(HitPayload);
+		}
+
 		if (!bLastBounce)
 		{
 			float3 Contrib = PathState.PathThroughput * EstimateMaterialAlbedo(HitPayload);
@@ -674,6 +704,12 @@ FPathTracingPayload TraceTransparentRay(inout FPathState PathState, bool bIsCame
 		PathState.PathThroughput = 0;
 		Payload.SetMiss();
 	}
+
+	if (bIsCameraRay)
+	{
+		PathState.WriteDepth(DepthContext);
+	}
+
 	return Payload;
 }
 
@@ -1266,7 +1302,7 @@ bool ProcessSubsurfaceRandomWalk(inout FPathTracingPayload Payload, inout float3
 	return false;
 }
 
-FPathState CreatePathState(int2 PixelIndex)
+FPathState CreatePathState(int2 PixelIndex, int2 PrimaryRayTextureIndex)
 {
 	FPathState Output = (FPathState)0;
 
@@ -1296,6 +1332,9 @@ FPathState CreatePathState(int2 PixelIndex)
 
 		float2 ViewportUV = (LaunchIndex + AAJitter) * View.BufferSizeAndInvSize.zw;
 		Output.Ray = CreatePrimaryRay(ViewportUV);
+
+		// only record the depth within a square.
+		Output.DepthResamplingWeight = all(AAJitter >= 0.0) && all(AAJitter <= 1.0) ? 1 : 0;
 	}
 
 	if (CameraLensRadius.y > 0)
@@ -1325,6 +1364,7 @@ FPathState CreatePathState(int2 PixelIndex)
 	Output.PathRoughness = 0;
 	Output.SigmaT = float3(StartingExtinctionCoefficient[0], StartingExtinctionCoefficient[1], StartingExtinctionCoefficient[2]);
 	Output.FirstScatterType = PATHTRACER_SCATTER_CAMERA;
+	Output.PrimaryRayTextureIndex = PrimaryRayTextureIndex;
 
 	return Output;
 }
diff --git a/Engine/Source/Runtime/Renderer/Private/DeferredShadingRenderer.cpp b/Engine/Source/Runtime/Renderer/Private/DeferredShadingRenderer.cpp
--- a/Engine/Source/Runtime/Renderer/Private/DeferredShadingRenderer.cpp
+++ b/Engine/Source/Runtime/Renderer/Private/DeferredShadingRenderer.cpp
@@ -4202,7 +4202,7 @@ void FDeferredShadingSceneRenderer::Render(FRDGBuilder& GraphBuilder)
 		{
 			for (const FViewInfo& View : Views)
 			{
-				RenderPathTracing(GraphBuilder, View, SceneTextures.UniformBuffer, SceneTextures.Color.Target, PathTracingResources);
+				RenderPathTracing(GraphBuilder, View, SceneTextures.UniformBuffer, SceneTextures.Color.Target, SceneTextures.Depth.Target,PathTracingResources);
 			}
 		}
 		else if (ViewFamily.EngineShowFlags.RayTracingDebug)
diff --git a/Engine/Source/Runtime/Renderer/Private/DeferredShadingRenderer.h b/Engine/Source/Runtime/Renderer/Private/DeferredShadingRenderer.h
--- a/Engine/Source/Runtime/Renderer/Private/DeferredShadingRenderer.h
+++ b/Engine/Source/Runtime/Renderer/Private/DeferredShadingRenderer.h
@@ -1120,6 +1120,7 @@ private:
 		const FViewInfo& View,
 		TRDGUniformBufferRef<FSceneTextureUniformParameters> SceneTexturesUniformBuffer,
 		FRDGTextureRef SceneColorOutputTexture,
+		FRDGTextureRef SceneDepthOutputTexture,
 		struct FPathTracingResources& PathTracingResources);
 
 	void ComputePathCompaction(FRHICommandListImmediate& RHICmdList, const FViewInfo& View, FRHITexture* RadianceTexture, FRHITexture* SampleCountTexture, FRHITexture* PixelPositionTexture,
diff --git a/Engine/Source/Runtime/Renderer/Private/PathTracing.cpp b/Engine/Source/Runtime/Renderer/Private/PathTracing.cpp
--- a/Engine/Source/Runtime/Renderer/Private/PathTracing.cpp
+++ b/Engine/Source/Runtime/Renderer/Private/PathTracing.cpp
@@ -355,6 +355,14 @@ TAutoConsoleVariable<int32> CVarPathTracingSubstrateCompileSimplifiedMaterial(
 	ECVF_RenderThreadSafe | ECVF_ReadOnly
 );
 
+TAutoConsoleVariable<int32> CVarpathTracingOverrideDepth(
+	TEXT("r.PathTracing.Override.Depth"),
+	1,
+	TEXT("Override the scene depth z by the path tracing depth z")
+	TEXT("0: off\n")
+	TEXT("1: On (Default, translucent materials have better DOF with the post-process DOF.)\n"),
+	ECVF_RenderThreadSafe
+);
 
 
 BEGIN_SHADER_PARAMETER_STRUCT(FPathTracingData, )
@@ -2096,6 +2104,7 @@ class FPathTracingCompositorPS : public FGlobalShader
 
 	BEGIN_SHADER_PARAMETER_STRUCT(FParameters, )
 		SHADER_PARAMETER_RDG_TEXTURE_SRV(Texture2D<float4>, RadianceTexture)
+		SHADER_PARAMETER_RDG_TEXTURE_SRV(Texture2D<float4>, NormalDepthTexture)
 		SHADER_PARAMETER_STRUCT_REF(FViewUniformShaderParameters, ViewUniformBuffer)
 		SHADER_PARAMETER(uint32, Iteration)
 		SHADER_PARAMETER(uint32, MaxSamples)
@@ -2187,6 +2196,7 @@ void FDeferredShadingSceneRenderer::RenderPathTracing(
 	const FViewInfo& View,
 	TRDGUniformBufferRef<FSceneTextureUniformParameters> SceneTexturesUniformBuffer,
 	FRDGTextureRef SceneColorOutputTexture,
+	FRDGTextureRef SceneDepthOutputTexture,
 	FPathTracingResources& PathTracingResources)
 {
 	RDG_EVENT_SCOPE(GraphBuilder, "Path Tracing");
@@ -2906,7 +2916,7 @@ void FDeferredShadingSceneRenderer::RenderPathTracing(
 	}
 	PathTracingState->LastConfig.DenoiserMode = DenoiserMode;
 
-	// now add a pixel shader pass to display our Radiance buffer
+	// now add a pixel shader pass to display our Radiance buffer and write to the depth buffer
 
 	FPathTracingCompositorPS::FParameters* DisplayParameters = GraphBuilder.AllocParameters<FPathTracingCompositorPS::FParameters>();
 	DisplayParameters->Iteration = Config.PathTracingData.Iteration;
@@ -2914,7 +2924,9 @@ void FDeferredShadingSceneRenderer::RenderPathTracing(
 	DisplayParameters->ProgressDisplayEnabled = CVarPathTracingProgressDisplay.GetValueOnRenderThread();
 	DisplayParameters->ViewUniformBuffer = View.ViewUniformBuffer;
 	DisplayParameters->RadianceTexture = GraphBuilder.CreateSRV(FRDGTextureSRVDesc::Create(DenoisedRadianceTexture ? DenoisedRadianceTexture : RadianceTexture));
+	DisplayParameters->NormalDepthTexture = GraphBuilder.CreateSRV(FRDGTextureSRVDesc::Create(NormalTexture));
 	DisplayParameters->RenderTargets[0] = FRenderTargetBinding(SceneColorOutputTexture, ERenderTargetLoadAction::ELoad);
+	DisplayParameters->RenderTargets.DepthStencil = FDepthStencilBinding(SceneDepthOutputTexture,  ERenderTargetLoadAction::ELoad, ERenderTargetLoadAction::ENoAction, FExclusiveDepthStencil::DepthWrite_StencilNop);
 
 	FScreenPassTextureViewport Viewport(SceneColorOutputTexture, View.ViewRect);
 
@@ -2936,15 +2948,30 @@ void FDeferredShadingSceneRenderer::RenderPathTracing(
 	}
 
 	TShaderMapRef<FPathTracingCompositorPS> PixelShader(View.ShaderMap);
+	TShaderMapRef<FScreenPassVS> VertexShader(View.ShaderMap);
+	FRHIBlendState* BlendState = FScreenPassPipelineState::FDefaultBlendState::GetRHI();
+	FRHIDepthStencilState* DepthStencilState = nullptr;
+
+	if (CVarpathTracingOverrideDepth.GetValueOnRenderThread() != 0)
+	{
+		DepthStencilState = TStaticDepthStencilState<true, CF_Always>::GetRHI();
+	}
+	else
+	{
+		DepthStencilState = TStaticDepthStencilState<false, CF_Always>::GetRHI();
+	}
+
 	AddDrawScreenPass(
 		GraphBuilder,
 		RDG_EVENT_NAME("Path Tracer Display (%d x %d)", View.ViewRect.Size().X, View.ViewRect.Size().Y),
 		View,
 		Viewport,
 		Viewport,
+		VertexShader,
 		PixelShader,
-		DisplayParameters
-	);
+		BlendState,
+		DepthStencilState,
+		DisplayParameters);
 
 	// Setup the path tracing resources to be used by post process pass.
 	if (CVarPathTracingOutputPostProcessResources.GetValueOnRenderThread() != 0)
