Fix for partial depth collision quesry to work the same way as refular collision query when it has not been generated or when used on translucent emitters.

diff --git a/Engine/Source/Runtime/Renderer/Private/DeferredShadingRenderer.cpp b/Engine/Source/Runtime/Renderer/Private/DeferredShadingRenderer.cpp
--- a/Engine/Source/Runtime/Renderer/Private/DeferredShadingRenderer.cpp
+++ b/Engine/Source/Runtime/Renderer/Private/DeferredShadingRenderer.cpp
@@ -3313,7 +3313,8 @@ void FDeferredShadingSceneRenderer::Render(FRDGBuilder& GraphBuilder)
 	}
 	else
 	{
-		SceneTextures.PartialDepth = SystemTextures.DepthDummy;
+		// Setup default partial depth to be scene depth so that it also works on transparent emitter when partial depth has not been generated.
+		SceneTextures.PartialDepth = SceneTextures.Depth;
 	}
 	SceneTextures.SetupMode = ESceneTextureSetupMode::SceneDepth;
 	SceneTextures.UniformBuffer = CreateSceneTextureUniformBuffer(GraphBuilder, &SceneTextures, FeatureLevel, SceneTextures.SetupMode);
