Support for Inline RT Shadows in Vulkan

diff --git a/Engine/Shaders/Private/RayTracing/RayTracingCommon.ush b/Engine/Shaders/Private/RayTracing/RayTracingCommon.ush
--- a/Engine/Shaders/Private/RayTracing/RayTracingCommon.ush
+++ b/Engine/Shaders/Private/RayTracing/RayTracingCommon.ush
@@ -109,6 +109,31 @@
 [shader("callable")] void name(inout params_type params_name)
 #endif //RAY_TRACING_ENTRY_CALLABLE
 
+// Helper macro to set common entry point for passes that have both inline (CS) and non-inline (RGS) version.
+// Given "name" it will create "nameCS" entry point for compute and "nameRGS" for raygen shader.
+#ifndef RAY_TRACING_ENTRY_RAYGEN_OR_INLINE
+
+#if COMPUTESHADER && PLATFORM_SUPPORTS_INLINE_RAY_TRACING
+
+#define RAY_TRACING_ENTRY_RAYGEN_OR_INLINE(name)\
+void name##_INTERNAL(uint3 DispatchThreadIndex);\
+[numthreads(INLINE_RAY_TRACING_THREAD_GROUP_SIZE_X, INLINE_RAY_TRACING_THREAD_GROUP_SIZE_Y, 1)]\
+void name##CS(uint3 DispatchThreadIndex : SV_DispatchThreadID) {\
+	name##_INTERNAL(DispatchThreadIndex);}\
+void name##_INTERNAL(uint3 DispatchThreadIndex)
+
+#else // COMPUTESHADER && PLATFORM_SUPPORTS_INLINE_RAY_TRACING
+
+#define RAY_TRACING_ENTRY_RAYGEN_OR_INLINE(name)\
+void name##_INTERNAL(uint3 DispatchThreadIndex);\
+RAY_TRACING_ENTRY_RAYGEN(name##RGS){\
+name##_INTERNAL(DispatchRaysIndex());}\
+void name##_INTERNAL(uint3 DispatchThreadIndex)
+
+#endif // COMPUTESHADER && PLATFORM_SUPPORTS_INLINE_RAY_TRACING
+
+#endif // RAY_TRACING_ENTRY_RAYGEN_OR_INLINE
+
 #ifndef OVERRIDE_PLATFORM_RAYTRACING_INSTANCE_DESCRITOR
 
 // Must match D3D12_RAYTRACING_INSTANCE_FLAGS / VkGeometryInstanceFlagBitsKHR
@@ -437,10 +462,6 @@ struct FRayDesc
 	}
 };
 
-#endif // !COMPUTESHADER || PLATFORM_SUPPORTS_INLINE_RAY_TRACING
-
-#if !COMPUTESHADER
-
 float3 TranslatedWorldRayOrigin()
 {
 	return TLASWorldToTranslatedWorld(WorldRayOrigin());
@@ -456,10 +477,6 @@ void ApplyCameraRelativeDepthBias(inout FRayDesc Ray, uint2 PixelCoord, float De
 	ApplyCameraRelativeDepthBias(Ray.Origin, Ray.Direction, PixelCoord, DeviceZ, WorldNormal, AbsoluteNormalBias);
 }
 
-#endif // !COMPUTESHADER
-
-#if !COMPUTESHADER || PLATFORM_SUPPORTS_INLINE_RAY_TRACING
-
 FRayDesc CreatePrimaryRay(float2 UV)
 {
 	float2 ScreenPosition = (UV - View.ScreenPositionScaleBias.wz) / View.ScreenPositionScaleBias.xy;
@@ -511,7 +528,7 @@ uint TranslateRayTracingInstanceFlags(uint InFlags)
 ///////////////////////////////////////////////////////////////////////////////////////////////////
 // Material payload
 
-#if !COMPUTESHADER
+#if !COMPUTESHADER || PLATFORM_SUPPORTS_INLINE_RAY_TRACING
 
 #if STRATA_ENABLED
 
@@ -934,7 +951,7 @@ FMaterialClosestHitPayload UnpackRayTracingPayload(FPackedMaterialClosestHitPayl
 CHECK_RT_PAYLOAD_SIZE(FPackedMaterialClosestHitPayload)
 #endif
 
-#endif // !COMPUTESHADER
+#endif // !COMPUTESHADER || PLATFORM_SUPPORTS_INLINE_RAY_TRACING
 
 ///////////////////////////////////////////////////////////////////////////////////////////////////
 // Decal parameters
@@ -1005,7 +1022,11 @@ CHECK_RT_PAYLOAD_SIZE(FDecalShaderPayload)
 ///////////////////////////////////////////////////////////////////////////////////////////////////
 // Tracing util functions
 
-#if !COMPUTESHADER
+#if !COMPUTESHADER || PLATFORM_SUPPORTS_INLINE_RAY_TRACING
+
+#if COMPUTESHADER
+#include "TraceRayInline.ush"
+#endif // COMPUTESHADER
 
 void TraceVisibilityRayPacked(
 	inout FPackedMaterialClosestHitPayload PackedPayload,
@@ -1024,6 +1045,21 @@ void TraceVisibilityRayPacked(
 
 	// Trace the ray
 
+#if COMPUTESHADER
+
+	FTraceRayInlineContext TraceRayInlineContext = CreateTraceRayInlineContext();
+
+	FTraceRayInlineResult TraceRayResult = TraceRayInline(
+		TLAS,
+		RayFlags,
+		InstanceInclusionMask,
+		Ray.GetNativeDesc(),
+		TraceRayInlineContext
+	);
+	PackedPayload.HitT = TraceRayResult.HitT;
+
+#else
+
 	TraceRay(
 		TLAS,
 		RayFlags,
@@ -1033,6 +1069,8 @@ void TraceVisibilityRayPacked(
 		MissShaderIndex,
 		Ray.GetNativeDesc(),
 		PackedPayload);
+
+#endif // COMPUTESHADER
 }
 
 FMinimalPayload TraceVisibilityRay(
@@ -1077,6 +1115,8 @@ FMinimalPayload TraceVisibilityRay(
 		Ray);
 }
 
+#if !COMPUTESHADER
+
 void TraceMaterialRayPacked(
 	inout FPackedMaterialClosestHitPayload Payload,
 	in RaytracingAccelerationStructure TLAS,
@@ -1172,4 +1212,6 @@ FMaterialClosestHitPayload TraceMaterialRay(
 
 #endif // !COMPUTESHADER
 
+#endif // !COMPUTESHADER || PLATFORM_SUPPORTS_INLINE_RAY_TRACING
+
 #endif // RAYTRACINGCOMMON_USH_INCLUDED // Workarround for UE-66460
diff --git a/Engine/Shaders/Private/RayTracing/RayTracingOcclusionRGS.usf b/Engine/Shaders/Private/RayTracing/RayTracingOcclusionRGS.usf
--- a/Engine/Shaders/Private/RayTracing/RayTracingOcclusionRGS.usf
+++ b/Engine/Shaders/Private/RayTracing/RayTracingOcclusionRGS.usf
@@ -634,9 +634,9 @@ float DecodeSamplingSSSProfileRadius(FGBufferData GBufferData)
 }
 #endif
 
-RAY_TRACING_ENTRY_RAYGEN(OcclusionRGS)
+RAY_TRACING_ENTRY_RAYGEN_OR_INLINE(Occlusion)
 {
-	const uint2 PixelCoord = DispatchRaysIndex().xy + View.ViewRectMin.xy + PixelOffset;
+	uint2 PixelCoord = DispatchThreadIndex.xy + View.ViewRectMin.xy + PixelOffset;
 
 	FOcclusionResult Occlusion = InitOcclusionResult();
 	FOcclusionResult HairOcclusion = InitOcclusionResult();
diff --git a/Engine/Source/Runtime/Renderer/Private/LightRendering.cpp b/Engine/Source/Runtime/Renderer/Private/LightRendering.cpp
--- a/Engine/Source/Runtime/Renderer/Private/LightRendering.cpp
+++ b/Engine/Source/Runtime/Renderer/Private/LightRendering.cpp
@@ -173,12 +173,12 @@ bool ShouldRenderRayTracingShadows()
 	const bool bIsStereo = GEngine->StereoRenderingDevice.IsValid() && GEngine->StereoRenderingDevice->IsStereoEnabled();
 	const bool bHairStrands = IsHairStrandsEnabled(EHairStrandsShaderType::Strands);
 
-	return ShouldRenderRayTracingEffect((CVarRayTracingOcclusion.GetValueOnRenderThread() > 0) && !(bIsStereo && bHairStrands), ERayTracingPipelineCompatibilityFlags::FullPipeline, nullptr);
+	return ShouldRenderRayTracingEffect((CVarRayTracingOcclusion.GetValueOnRenderThread() > 0) && !(bIsStereo && bHairStrands), ERayTracingPipelineCompatibilityFlags::FullPipeline | ERayTracingPipelineCompatibilityFlags::Inline, nullptr);
 }
 
 bool ShouldRenderRayTracingShadowsForLight(const FLightSceneProxy& LightProxy)
 {
-	const bool bShadowRayTracingAllowed = ShouldRenderRayTracingEffect(true, ERayTracingPipelineCompatibilityFlags::FullPipeline, nullptr);
+	const bool bShadowRayTracingAllowed = ShouldRenderRayTracingEffect(true, ERayTracingPipelineCompatibilityFlags::FullPipeline | ERayTracingPipelineCompatibilityFlags::Inline, nullptr);
 	return (LightProxy.CastsRaytracedShadow() == ECastRayTracedShadow::Enabled || (ShouldRenderRayTracingShadows() && LightProxy.CastsRaytracedShadow() == ECastRayTracedShadow::UseProjectSetting))
 		&& ShouldRenderRayTracingShadowsForLightType((ELightComponentType)LightProxy.GetLightType())
 		&& bShadowRayTracingAllowed;
@@ -186,7 +186,7 @@ bool ShouldRenderRayTracingShadowsForLight(const FLightSceneProxy& LightProxy)
 
 bool ShouldRenderRayTracingShadowsForLight(const FLightSceneInfoCompact& LightInfo)
 {
-	const bool bShadowRayTracingAllowed = ShouldRenderRayTracingEffect(true, ERayTracingPipelineCompatibilityFlags::FullPipeline, nullptr);
+	const bool bShadowRayTracingAllowed = ShouldRenderRayTracingEffect(true, ERayTracingPipelineCompatibilityFlags::FullPipeline | ERayTracingPipelineCompatibilityFlags::Inline, nullptr);
 	return (LightInfo.CastRaytracedShadow == ECastRayTracedShadow::Enabled || (ShouldRenderRayTracingShadows() && LightInfo.CastRaytracedShadow == ECastRayTracedShadow::UseProjectSetting))
 		&& ShouldRenderRayTracingShadowsForLightType((ELightComponentType)LightInfo.LightType)
 		&& bShadowRayTracingAllowed;
diff --git a/Engine/Source/Runtime/Renderer/Private/RayTracing/RayTracingShadows.cpp b/Engine/Source/Runtime/Renderer/Private/RayTracing/RayTracingShadows.cpp
--- a/Engine/Source/Runtime/Renderer/Private/RayTracing/RayTracingShadows.cpp
+++ b/Engine/Source/Runtime/Renderer/Private/RayTracing/RayTracingShadows.cpp
@@ -208,11 +208,76 @@ class FOcclusionRGS : public FGlobalShader
 
 IMPLEMENT_GLOBAL_SHADER(FOcclusionRGS, "/Engine/Private/RayTracing/RayTracingOcclusionRGS.usf", "OcclusionRGS", SF_RayGen);
 
+class FOcclusionCS : public FGlobalShader
+{
+	DECLARE_GLOBAL_SHADER(FOcclusionCS)
+	SHADER_USE_PARAMETER_STRUCT(FOcclusionCS, FGlobalShader)
+
+	class FLightTypeDim : SHADER_PERMUTATION_INT("LIGHT_TYPE", LightType_MAX);
+	class FDenoiserOutputDim : SHADER_PERMUTATION_INT("DIM_DENOISER_OUTPUT", 3);
+	class FEnableMultipleSamplesPerPixel : SHADER_PERMUTATION_BOOL("ENABLE_MULTIPLE_SAMPLES_PER_PIXEL");
+	class FHairLighting : SHADER_PERMUTATION_INT("USE_HAIR_LIGHTING", 2);
+	class FEnableTransmissionDim : SHADER_PERMUTATION_INT("ENABLE_TRANSMISSION", 2);
+
+	using FPermutationDomain = TShaderPermutationDomain<FLightTypeDim, FDenoiserOutputDim, FHairLighting, FEnableMultipleSamplesPerPixel, FEnableTransmissionDim>;
+
+	static bool ShouldCompilePermutation(const FGlobalShaderPermutationParameters& Parameters)
+	{
+		return IsRayTracingEnabledForProject(Parameters.Platform) && FDataDrivenShaderPlatformInfo::GetSupportsInlineRayTracing(Parameters.Platform);
+	}
+
+	static void ModifyCompilationEnvironment(const FGlobalShaderPermutationParameters& Parameters, FShaderCompilerEnvironment& OutEnvironment)
+	{
+		FGlobalShader::ModifyCompilationEnvironment(Parameters, OutEnvironment);
+
+		OutEnvironment.CompilerFlags.Add(CFLAG_Wave32);
+		OutEnvironment.CompilerFlags.Add(CFLAG_InlineRayTracing);
+
+		OutEnvironment.SetDefine(TEXT("INLINE_RAY_TRACING_THREAD_GROUP_SIZE_X"), ThreadGroupSizeX);
+		OutEnvironment.SetDefine(TEXT("INLINE_RAY_TRACING_THREAD_GROUP_SIZE_Y"), ThreadGroupSizeY);
+
+		FPermutationDomain PermutationVector(Parameters.PermutationId);
+		if (PermutationVector.Get<FLightTypeDim>() == LightType_Directional &&
+			PermutationVector.Get<FHairLighting>() == 0)
+		{
+			OutEnvironment.SetDefine(TEXT("UE_RAY_TRACING_COHERENT_RAYS"), 1);
+		}
+		else
+		{
+			OutEnvironment.SetDefine(TEXT("UE_RAY_TRACING_COHERENT_RAYS"), 0);
+		}
+	}
+
+	BEGIN_SHADER_PARAMETER_STRUCT(FParameters, )
+		SHADER_PARAMETER_STRUCT_INCLUDE(FOcclusionRGS::FParameters, CommonParameters)
+	END_SHADER_PARAMETER_STRUCT()
+
+	// Current inline ray tracing implementation requires 1:1 mapping between thread groups and waves and only supports wave32 mode.
+	static constexpr uint32 ThreadGroupSizeX = 8;
+	static constexpr uint32 ThreadGroupSizeY = 4;
+
+};
+
+IMPLEMENT_GLOBAL_SHADER(FOcclusionCS, "/Engine/Private/RayTracing/RayTracingOcclusionRGS.usf", "OcclusionCS", SF_Compute);
+
 float GetRaytracingMaxNormalBias()
 {
 	return FMath::Max(0.01f, GRayTracingMaxNormalBias);
 }
 
+FOcclusionCS::FPermutationDomain ToComputePermutationVector(const FOcclusionRGS::FPermutationDomain& RGSPermutationVector)
+{
+	FOcclusionCS::FPermutationDomain PermutationVector;
+
+	PermutationVector.Set<FOcclusionCS::FLightTypeDim>(RGSPermutationVector.Get<FOcclusionRGS::FLightTypeDim>());
+	PermutationVector.Set<FOcclusionCS::FDenoiserOutputDim>(RGSPermutationVector.Get<FOcclusionRGS::FDenoiserOutputDim>());
+	PermutationVector.Set<FOcclusionCS::FEnableMultipleSamplesPerPixel>(RGSPermutationVector.Get<FOcclusionRGS::FEnableMultipleSamplesPerPixel>());
+	PermutationVector.Set<FOcclusionCS::FHairLighting>(RGSPermutationVector.Get<FOcclusionRGS::FHairLighting>());
+	PermutationVector.Set<FOcclusionCS::FEnableTransmissionDim>(RGSPermutationVector.Get<FOcclusionRGS::FEnableTransmissionDim>());
+
+	return PermutationVector;
+}
+
 void FDeferredShadingSceneRenderer::PrepareRayTracingShadows(const FViewInfo& View, const FScene& Scene, TArray<FRHIRayTracingShader*>& OutRayGenShaders)
 {
 	// Ray tracing shadows shaders should be properly configured even if r.RayTracing.Shadows is 0 because lights can have raytracing shadows enabled independently of that CVar
@@ -279,6 +344,8 @@ void FDeferredShadingSceneRenderer::RenderRayTracingShadows(
 	FLightSceneProxy* LightSceneProxy = LightSceneInfo.Proxy;
 	check(LightSceneProxy);
 
+	const bool bInlineRayTracing = !GRHISupportsRayTracingShaders && GRHISupportsInlineRayTracing;
+
 	FIntRect ScissorRect = View.ViewRect;
 	FIntPoint PixelOffset = { 0, 0 };
 
@@ -306,52 +373,55 @@ void FDeferredShadingSceneRenderer::RenderRayTracingShadows(
 		const bool bUseHairLighting   = HairStrands::HasViewHairStrandsData(View) && HairStrands::HasViewHairStrandsVoxelData(View);
 		const bool bUseHairDeepShadow = HairStrands::HasViewHairStrandsData(View) && LightSceneProxy->CastsHairStrandsDeepShadow();
 
-		FOcclusionRGS::FParameters* PassParameters = GraphBuilder.AllocParameters<FOcclusionRGS::FParameters>();
-		PassParameters->RWOcclusionMaskUAV = OutShadowMaskUAV;
-		PassParameters->RWRayDistanceUAV = OutRayHitDistanceUAV;
-		PassParameters->RWSubPixelOcclusionMaskUAV = SubPixelRayTracingShadowMaskUAV;
-		PassParameters->SamplesPerPixel = RayTracingConfig.RayCountPerPixel;
-		PassParameters->NormalBias = GetRaytracingMaxNormalBias();
-		PassParameters->LightingChannelMask = LightSceneProxy->GetLightingChannelMask();
+		FOcclusionRGS::FParameters* CommonPassParameters = GraphBuilder.AllocParameters<FOcclusionRGS::FParameters>();
+
+		CommonPassParameters->RWOcclusionMaskUAV = OutShadowMaskUAV;
+		CommonPassParameters->RWRayDistanceUAV = OutRayHitDistanceUAV;
+		CommonPassParameters->RWSubPixelOcclusionMaskUAV = SubPixelRayTracingShadowMaskUAV;
+		CommonPassParameters->SamplesPerPixel = RayTracingConfig.RayCountPerPixel;
+		CommonPassParameters->NormalBias = GetRaytracingMaxNormalBias();
+		CommonPassParameters->LightingChannelMask = LightSceneProxy->GetLightingChannelMask();
 		
 		{
 			FLightRenderParameters LightParameters;
 			LightSceneProxy->GetLightShaderParameters(LightParameters);
-			LightParameters.MakeShaderParameters(View.ViewMatrices, View.GetLastEyeAdaptationExposure(), PassParameters->Light);
-			PassParameters->Light.SourceRadius *= LightSceneProxy->GetShadowSourceAngleFactor();
+			LightParameters.MakeShaderParameters(View.ViewMatrices, View.GetLastEyeAdaptationExposure(), CommonPassParameters->Light);
+			CommonPassParameters->Light.SourceRadius *= LightSceneProxy->GetShadowSourceAngleFactor();
 		}
 
-		PassParameters->TraceDistance = LightSceneProxy->GetTraceDistance();
-		PassParameters->LODTransitionStart = CVarRayTracingShadowsLODTransitionStart.GetValueOnRenderThread();
-		PassParameters->LODTransitionEnd = CVarRayTracingShadowsLODTransitionEnd.GetValueOnRenderThread();
-		PassParameters->AvoidSelfIntersectionTraceDistance = GRayTracingShadowsAvoidSelfIntersectionTraceDistance;
-		PassParameters->bAcceptFirstHit = CVarRayTracingShadowsAcceptFirstHit.GetValueOnRenderThread();
-		PassParameters->bTwoSidedGeometry = EnableRayTracingShadowTwoSidedGeometry() ? 1 : 0;
-		PassParameters->TLAS = View.GetRayTracingSceneLayerViewChecked(ERayTracingSceneLayer::Base);
-		PassParameters->ViewUniformBuffer = View.ViewUniformBuffer;
-		PassParameters->SceneTextures = SceneTextures;
-		PassParameters->SceneLightingChannels = GetSceneLightingChannelParameters(GraphBuilder, LightingChannelsTexture);
-		PassParameters->LightScissor = ScissorRect;
-		PassParameters->PixelOffset = PixelOffset;
-		PassParameters->SSProfilesTexture = View.RayTracingSubSurfaceProfileTexture;
-		PassParameters->bTransmissionSamplingDistanceCulling = CVarRayTracingTransmissionSamplingDistanceCulling.GetValueOnRenderThread();
-		PassParameters->TransmissionSamplingTechnique = CVarRayTracingTransmissionSamplingTechnique.GetValueOnRenderThread();
-		PassParameters->TransmissionMeanFreePathType = GetRayTracingTransmissionMeanFreePathType();
-		PassParameters->RejectionSamplingTrials = CVarRayTracingTransmissionRejectionSamplingTrials.GetValueOnRenderThread();
-		PassParameters->Strata = Strata::BindStrataGlobalUniformParameters(View);
+		CommonPassParameters->TraceDistance = LightSceneProxy->GetTraceDistance();
+		CommonPassParameters->LODTransitionStart = CVarRayTracingShadowsLODTransitionStart.GetValueOnRenderThread();
+		CommonPassParameters->LODTransitionEnd = CVarRayTracingShadowsLODTransitionEnd.GetValueOnRenderThread();
+		CommonPassParameters->AvoidSelfIntersectionTraceDistance = GRayTracingShadowsAvoidSelfIntersectionTraceDistance;
+		CommonPassParameters->bAcceptFirstHit = CVarRayTracingShadowsAcceptFirstHit.GetValueOnRenderThread();
+		CommonPassParameters->bTwoSidedGeometry = EnableRayTracingShadowTwoSidedGeometry() ? 1 : 0;
+		CommonPassParameters->TLAS = View.GetRayTracingSceneLayerViewChecked(ERayTracingSceneLayer::Base);
+		CommonPassParameters->ViewUniformBuffer = View.ViewUniformBuffer;
+		CommonPassParameters->SceneTextures = SceneTextures;
+		CommonPassParameters->SceneLightingChannels = GetSceneLightingChannelParameters(GraphBuilder, LightingChannelsTexture);
+		CommonPassParameters->LightScissor = ScissorRect;
+		CommonPassParameters->PixelOffset = PixelOffset;
+		CommonPassParameters->SSProfilesTexture = View.RayTracingSubSurfaceProfileTexture;
+		CommonPassParameters->bTransmissionSamplingDistanceCulling = CVarRayTracingTransmissionSamplingDistanceCulling.GetValueOnRenderThread();
+		CommonPassParameters->TransmissionSamplingTechnique = CVarRayTracingTransmissionSamplingTechnique.GetValueOnRenderThread();
+		CommonPassParameters->TransmissionMeanFreePathType = GetRayTracingTransmissionMeanFreePathType();
+		CommonPassParameters->RejectionSamplingTrials = CVarRayTracingTransmissionRejectionSamplingTrials.GetValueOnRenderThread();
+		CommonPassParameters->Strata = Strata::BindStrataGlobalUniformParameters(View);
+
 		if (bUseHairLighting)
 		{
 			const bool bUseHairVoxel = CVarRayTracingShadowsEnableHairVoxel.GetValueOnRenderThread() > 0;
-			PassParameters->bUseHairVoxel = !bUseHairDeepShadow && bUseHairVoxel ? 1 : 0;
-			PassParameters->HairLightChannelMaskTexture = View.HairStrandsViewData.VisibilityData.LightChannelMaskTexture;
-			PassParameters->HairStrands = HairStrands::BindHairStrandsViewUniformParameters(View);
-			PassParameters->VirtualVoxel = HairStrands::BindHairStrandsVoxelUniformParameters(View);
+			CommonPassParameters->bUseHairVoxel = !bUseHairDeepShadow && bUseHairVoxel ? 1 : 0;
+			CommonPassParameters->HairLightChannelMaskTexture = View.HairStrandsViewData.VisibilityData.LightChannelMaskTexture;
+			CommonPassParameters->HairStrands = HairStrands::BindHairStrandsViewUniformParameters(View);
+			CommonPassParameters->VirtualVoxel = HairStrands::BindHairStrandsVoxelUniformParameters(View);
 
 			if (ShaderPrint::IsValid(View.ShaderPrintData))
 			{
-				ShaderPrint::SetParameters(GraphBuilder, View.ShaderPrintData, PassParameters->ShaderPrintParameters);
+				ShaderPrint::SetParameters(GraphBuilder, View.ShaderPrintData, CommonPassParameters->ShaderPrintParameters);
 			}
 		}
+
 		FOcclusionRGS::FPermutationDomain PermutationVector;
 		PermutationVector.Set<FOcclusionRGS::FLightTypeDim>(LightSceneProxy->GetLightType());
 		PermutationVector.Set<FOcclusionRGS::FAvoidSelfIntersectionTraceDim>(GRayTracingShadowsAvoidSelfIntersectionTraceDistance > 0.0f);
@@ -372,53 +442,84 @@ void FDeferredShadingSceneRenderer::RenderRayTracingShadows(
 		PermutationVector.Set<FOcclusionRGS::FEnableMultipleSamplesPerPixel>(RayTracingConfig.RayCountPerPixel > 1);
 		PermutationVector.Set<FOcclusionRGS::FEnableTransmissionDim>(LightSceneProxy->Transmission() ? 1 : 0);
 
-		TShaderMapRef<FOcclusionRGS> RayGenerationShader(GetGlobalShaderMap(FeatureLevel), PermutationVector);
+		if (bInlineRayTracing)
+		{
+			FOcclusionCS::FParameters* InlinePassParameters = GraphBuilder.AllocParameters<FOcclusionCS::FParameters>();
 
-		ClearUnusedGraphResources(RayGenerationShader, PassParameters);
+			InlinePassParameters->CommonParameters = *CommonPassParameters;
 
-		FIntPoint Resolution(View.ViewRect.Width(), View.ViewRect.Height());
+			TShaderRef<FOcclusionCS> ComputeShader = View.ShaderMap->GetShader<FOcclusionCS>(ToComputePermutationVector(PermutationVector));
 
-		if (bClipDispatch)
-		{
-			Resolution = ScissorRect.Size();
-		}
+			FIntPoint Resolution(View.ViewRect.Width(), View.ViewRect.Height());
+
+			if (bClipDispatch)
+			{
+				Resolution = ScissorRect.Size();
+			}
 
-		GraphBuilder.AddPass(
-			RDG_EVENT_NAME("RayTracedShadow (spp=%d) %dx%d", RayTracingConfig.RayCountPerPixel, Resolution.X, Resolution.Y),
-			PassParameters,
-			ERDGPassFlags::Compute,
-			[this, &View, RayGenerationShader, PassParameters, Resolution](FRHIRayTracingCommandList& RHICmdList)
+			const FIntPoint GroupSize(FOcclusionCS::ThreadGroupSizeX, FOcclusionCS::ThreadGroupSizeY);
+			const FIntVector GroupCount = FComputeShaderUtils::GetGroupCount(View.ViewRect.Size(), GroupSize);
+
+			GraphBuilder.AddPass(
+				RDG_EVENT_NAME("RayTracedShadow (INLINE) (spp=%d) %dx%d", RayTracingConfig.RayCountPerPixel, Resolution.X, Resolution.Y),
+				InlinePassParameters,
+				ERDGPassFlags::Compute,
+				[InlinePassParameters, ComputeShader, GroupCount](FRHICommandListImmediate& RHICmdList)
+				{
+					
+					FComputeShaderUtils::Dispatch(RHICmdList, ComputeShader, *InlinePassParameters, GroupCount);	
+				});
+
+			//FComputeShaderUtils::AddPass(GraphBuilder, RDG_EVENT_NAME("RayTracedShadow (INLINE) (spp=%d) %dx%d", RayTracingConfig.RayCountPerPixel, Resolution.X, Resolution.Y), ComputeShader, InlinePassParameters, GroupCount);
+		}
+		else
 		{
-			FRayTracingShaderBindingsWriter GlobalResources;
-			SetShaderParameters(GlobalResources, RayGenerationShader, *PassParameters);
+			TShaderMapRef<FOcclusionRGS> RayGenerationShader(GetGlobalShaderMap(FeatureLevel), PermutationVector);
 
-			FRHIRayTracingScene* RayTracingSceneRHI = View.GetRayTracingSceneChecked();
+			ClearUnusedGraphResources(RayGenerationShader, CommonPassParameters);
 
-			if (GRayTracingShadowsEnableMaterials)
+			FIntPoint Resolution(View.ViewRect.Width(), View.ViewRect.Height());
+
+			if (bClipDispatch)
 			{
-				RHICmdList.RayTraceDispatch(View.RayTracingMaterialPipeline, RayGenerationShader.GetRayTracingShader(), RayTracingSceneRHI, GlobalResources, Resolution.X, Resolution.Y);
+				Resolution = ScissorRect.Size();
 			}
-			else
-			{
-				FRayTracingPipelineStateInitializer Initializer;
 
-				Initializer.MaxPayloadSizeInBytes = GetRayTracingPayloadTypeMaxSize(ERayTracingPayloadType::RayTracingMaterial);
+			GraphBuilder.AddPass(
+				RDG_EVENT_NAME("RayTracedShadow (spp=%d) %dx%d", RayTracingConfig.RayCountPerPixel, Resolution.X, Resolution.Y),
+				CommonPassParameters,
+				ERDGPassFlags::Compute,
+				[this, &View, RayGenerationShader, CommonPassParameters, Resolution](FRHIRayTracingCommandList& RHICmdList)
+				{
+					FRayTracingShaderBindingsWriter GlobalResources;
+					SetShaderParameters(GlobalResources, RayGenerationShader, *CommonPassParameters);
 
-				FRHIRayTracingShader* RayGenShaderTable[] = { RayGenerationShader.GetRayTracingShader() };
-				Initializer.SetRayGenShaderTable(RayGenShaderTable);
+					FRHIRayTracingScene* RayTracingSceneRHI = View.GetRayTracingSceneChecked();
 
-				FRHIRayTracingShader* HitGroupTable[] = { GetRayTracingDefaultOpaqueShader(View.ShaderMap) };
-				Initializer.SetHitGroupTable(HitGroupTable);
-				Initializer.bAllowHitGroupIndexing = false; // Use the same hit shader for all geometry in the scene by disabling SBT indexing.
+					if (GRayTracingShadowsEnableMaterials)
+					{
+						RHICmdList.RayTraceDispatch(View.RayTracingMaterialPipeline, RayGenerationShader.GetRayTracingShader(), RayTracingSceneRHI, GlobalResources, Resolution.X, Resolution.Y);
+					}
+					else
+					{
+						FRayTracingPipelineStateInitializer Initializer;
 
-				FRHIRayTracingShader* MissGroupTable[] = { GetRayTracingDefaultMissShader(View.ShaderMap) };
-				Initializer.SetMissShaderTable(MissGroupTable);
+						Initializer.MaxPayloadSizeInBytes = GetRayTracingPayloadTypeMaxSize(ERayTracingPayloadType::RayTracingMaterial);
 
-				FRayTracingPipelineState* Pipeline = PipelineStateCache::GetAndOrCreateRayTracingPipelineState(RHICmdList, Initializer);
-				RHICmdList.SetRayTracingMissShader(RayTracingSceneRHI, 0, Pipeline, 0 /* ShaderIndexInPipeline */, 0, nullptr, 0);
-				RHICmdList.RayTraceDispatch(Pipeline, RayGenerationShader.GetRayTracingShader(), RayTracingSceneRHI, GlobalResources, Resolution.X, Resolution.Y);
-			}
-		});
+						FRHIRayTracingShader* HitGroupTable[] = { GetRayTracingDefaultOpaqueShader(View.ShaderMap) };
+						Initializer.SetHitGroupTable(HitGroupTable);
+						Initializer.bAllowHitGroupIndexing = false; // Use the same hit shader for all geometry in the scene by disabling SBT indexing.
+
+						FRHIRayTracingShader* MissGroupTable[] = { GetRayTracingDefaultMissShader(View.ShaderMap) };
+						Initializer.SetMissShaderTable(MissGroupTable);
+
+						FRayTracingPipelineState* Pipeline = PipelineStateCache::GetAndOrCreateRayTracingPipelineState(RHICmdList, Initializer);
+						RHICmdList.SetRayTracingMissShader(RayTracingSceneRHI, 0, Pipeline, 0 /* ShaderIndexInPipeline */, 0, nullptr, 0);
+						RHICmdList.RayTraceDispatch(Pipeline, RayGenerationShader.GetRayTracingShader(), RayTracingSceneRHI, GlobalResources, Resolution.X, Resolution.Y);
+					}
+				}
+			);
+		}
 	}
 }
 #else // !RHI_RAYTRACING
diff --git a/Engine/Source/Runtime/Renderer/Private/ScreenSpaceDenoise.cpp b/Engine/Source/Runtime/Renderer/Private/ScreenSpaceDenoise.cpp
--- a/Engine/Source/Runtime/Renderer/Private/ScreenSpaceDenoise.cpp
+++ b/Engine/Source/Runtime/Renderer/Private/ScreenSpaceDenoise.cpp
@@ -775,7 +775,7 @@ static bool ShouldCompileSignalPipeline(ESignalProcessing SignalProcessing, ESha
 		SignalProcessing == ESignalProcessing::DiffuseAndAmbientOcclusion)
 	{
 		// Only for ray tracing denoising.
-		return RHISupportsRayTracingShaders(Platform);
+		return RHISupportsRayTracing(Platform);
 	}
 	else if (SignalProcessing == ESignalProcessing::PolychromaticPenumbraHarmonic)
 	{
