Separate mask only and full depth prepass on mobile.
Disable mask clip in PS if full depth prepass enabled.
Remove r.Mobile.EarlyZPassOnlyMaterialMasking and use r.Mobile.EarlyZPass=2 instead

diff --git a/Engine/Source/Runtime/Engine/Private/SceneUtils.cpp b/Engine/Source/Runtime/Engine/Private/SceneUtils.cpp
--- a/Engine/Source/Runtime/Engine/Private/SceneUtils.cpp
+++ b/Engine/Source/Runtime/Engine/Private/SceneUtils.cpp
@@ -151,16 +151,13 @@ ENGINE_API uint32 GetDefaultMSAACount(const FStaticFeatureLevel InFeatureLevel,
 			{
 				bool bMobilePixelProjectedReflection = IsUsingMobilePixelProjectedReflection(ShaderPlatform);
 				
-				static const auto MobileEarlyZPassCVar = IConsoleManager::Get().FindTConsoleVariableDataInt(TEXT("r.Mobile.EarlyZPass"));
-				const bool bMobileAmbientOcclusion = IsMobileAmbientOcclusionEnabled(ShaderPlatform);
-				const bool bMobileUsesShadowMaskTexture = MobileUsesShadowMaskTexture(ShaderPlatform);
-				bool bIsFullDepthPrepassEnabled = bMobileAmbientOcclusion || bMobileUsesShadowMaskTexture || MobileEarlyZPassCVar->GetValueOnAnyThread() > 0;
+				bool bIsFullDepthPrepassEnabled = MobileUsesFullDepthPrepass(ShaderPlatform);
 
 				bRendererSupportMSAA = bRHISupportsMSAA && !bMobilePixelProjectedReflection && !bIsFullDepthPrepassEnabled;
 
 				if (!bRendererSupportMSAA)
 				{
-					FailedReason = FString::Printf(TEXT("RHISupportsMSAA %d, MobilePixelProjectedReflection %d, MobileAmbientOcclusion %d, MobileUsesShadowMaskTexture %d, MobileEarlyZPass %d"), bRHISupportsMSAA ? 1 : 0, bMobilePixelProjectedReflection ? 1 : 0, bMobileAmbientOcclusion ? 1 : 0, bMobileUsesShadowMaskTexture? 1 : 0, MobileEarlyZPassCVar->GetValueOnAnyThread());
+					FailedReason = FString::Printf(TEXT("RHISupportsMSAA %d, MobilePixelProjectedReflection %d, MobileFullDepthPrepass %d"), bRHISupportsMSAA ? 1 : 0, bMobilePixelProjectedReflection ? 1 : 0, bIsFullDepthPrepassEnabled ? 1 : 0);
 				}
 			}
 			else
diff --git a/Engine/Source/Runtime/RenderCore/Private/RenderUtils.cpp b/Engine/Source/Runtime/RenderCore/Private/RenderUtils.cpp
--- a/Engine/Source/Runtime/RenderCore/Private/RenderUtils.cpp
+++ b/Engine/Source/Runtime/RenderCore/Private/RenderUtils.cpp
@@ -531,11 +531,11 @@ RENDERCORE_API bool SupportsTextureCubeArray(ERHIFeatureLevel::Type FeatureLevel
 
 RENDERCORE_API bool MaskedInEarlyPass(const FStaticShaderPlatform Platform)
 {
-	static FShaderPlatformCachedIniValue<bool> CVarMobileEarlyZPassOnlyMaterialMasking(TEXT("r.Mobile.EarlyZPassOnlyMaterialMasking"));
+	static FShaderPlatformCachedIniValue<int32> CVarMobileEarlyZPass(TEXT("r.Mobile.EarlyZPass"));
 	static IConsoleVariable* CVarEarlyZPassOnlyMaterialMasking = IConsoleManager::Get().FindConsoleVariable(TEXT("r.EarlyZPassOnlyMaterialMasking"));
 	if (IsMobilePlatform(Platform))
 	{
-		return CVarMobileEarlyZPassOnlyMaterialMasking.Get(Platform) == 1;
+		return CVarMobileEarlyZPass.Get(Platform) == 2 || MobileUsesFullDepthPrepass(Platform);
 	}
 	else
 	{
@@ -648,9 +648,33 @@ RENDERCORE_API bool MobileBasePassAlwaysUsesCSM(const FStaticShaderPlatform Plat
 
 RENDERCORE_API bool MobileUsesFullDepthPrepass(const FStaticShaderPlatform Platform)
 {
-	static TConsoleVariableData<int32>* CVar = IConsoleManager::Get().FindTConsoleVariableDataInt(TEXT("r.Mobile.EarlyZPassOnlyMaterialMasking"));
+	static FShaderPlatformCachedIniValue<int32> CVarMobileEarlyZPass(TEXT("r.Mobile.EarlyZPass"));
+	return MobileUsesShadowMaskTexture(Platform) || IsMobileAmbientOcclusionEnabled(Platform) || (CVarMobileEarlyZPass.Get(Platform) == 1);
+}
+
+RENDERCORE_API bool ShouldForceFullDepthPass(const FStaticShaderPlatform Platform)
+{
+	if (IsMobilePlatform(Platform))
+	{
+		return MobileUsesFullDepthPrepass(Platform);
+	}
+	else
+	{
+		const bool bNaniteEnabled = UseNanite(Platform);
+		const bool bDBufferAllowed = IsUsingDBuffers(Platform);
+		const bool bVirtualTextureEnabled = UseVirtualTexturing(Platform);
+
+		static const auto StencilLODDitherCVar = IConsoleManager::Get().FindTConsoleVariableDataInt(TEXT("r.StencilForLODDither"));
+		const bool bStencilLODDither = StencilLODDitherCVar->GetValueOnAnyThread() != 0;
 
-	return MobileUsesShadowMaskTexture(Platform) || IsMobileAmbientOcclusionEnabled(Platform) || (CVar && CVar->GetValueOnAnyThread() == 1);
+		static const auto AOComputeCVar = IConsoleManager::Get().FindTConsoleVariableDataInt(TEXT("r.AmbientOcclusion.Compute"));
+		const bool bAOCompute = AOComputeCVar->GetValueOnAnyThread() > 0;
+
+		const bool bEarlyZMaterialMasking = MaskedInEarlyPass(Platform);
+
+		// Note: ShouldForceFullDepthPass affects which static draw lists meshes go into, so nothing it depends on can change at runtime, unless you do a FGlobalComponentRecreateRenderStateContext to propagate the cvar change
+		return bNaniteEnabled || bAOCompute || bDBufferAllowed || bVirtualTextureEnabled || bStencilLODDither || bEarlyZMaterialMasking || IsForwardShadingEnabled(Platform) || IsUsingSelectiveBasePassOutputs(Platform);
+	}
 }
 
 RENDERCORE_API bool SupportsGen4TAA(const FStaticShaderPlatform Platform)
diff --git a/Engine/Source/Runtime/RenderCore/Public/RenderUtils.h b/Engine/Source/Runtime/RenderCore/Public/RenderUtils.h
--- a/Engine/Source/Runtime/RenderCore/Public/RenderUtils.h
+++ b/Engine/Source/Runtime/RenderCore/Public/RenderUtils.h
@@ -323,6 +323,8 @@ RENDERCORE_API bool MobileBasePassAlwaysUsesCSM(const FStaticShaderPlatform Plat
 
 RENDERCORE_API bool MobileUsesFullDepthPrepass(const FStaticShaderPlatform Platform);
 
+RENDERCORE_API bool ShouldForceFullDepthPass(const FStaticShaderPlatform Platform);
+
 RENDERCORE_API bool SupportsGen4TAA(const FStaticShaderPlatform Platform);
 
 RENDERCORE_API bool SupportsTSR(const FStaticShaderPlatform Platform);
diff --git a/Engine/Source/Runtime/Renderer/Private/RendererScene.cpp b/Engine/Source/Runtime/Renderer/Private/RendererScene.cpp
--- a/Engine/Source/Runtime/Renderer/Private/RendererScene.cpp
+++ b/Engine/Source/Runtime/Renderer/Private/RendererScene.cpp
@@ -102,16 +102,6 @@ static TAutoConsoleVariable<int32> CVarEarlyZPassOnlyMaterialMasking(
 	ECVF_RenderThreadSafe | ECVF_ReadOnly
 );
 
-/** Affects MobileBasePassPixelShader.usf so must relaunch editor to recompile shaders. */
-static TAutoConsoleVariable<int32> CVarMobileEarlyZPassOnlyMaterialMasking(
-	TEXT("r.Mobile.EarlyZPassOnlyMaterialMasking"),
-	0,
-	TEXT("Whether to compute materials' mask opacity only in early Z pass for Mobile platform. Changing this setting requires restarting the editor.\n")
-	TEXT("<=0: off\n")
-	TEXT("  1: on\n"),
-	ECVF_RenderThreadSafe | ECVF_ReadOnly
-);
-
 TAutoConsoleVariable<int32> CVarEarlyZPass(
 	TEXT("r.EarlyZPass"),
 	3,
@@ -126,10 +116,11 @@ TAutoConsoleVariable<int32> CVarEarlyZPass(
 static TAutoConsoleVariable<int32> CVarMobileEarlyZPass(
 	TEXT("r.Mobile.EarlyZPass"),
 	0,
-	TEXT("Whether to use a depth only pass to initialize Z culling for the mobile base pass.\n")
+	TEXT("Whether to use a depth only pass to initialize Z culling for the mobile base pass. Changing this setting requires restarting the editor.\n")
 	TEXT("  0: off\n")
-	TEXT("  1: all opaque \n"),
-	ECVF_Scalability
+	TEXT("  1: all opaque \n")
+	TEXT("  2: masked primitives only \n"),
+	ECVF_ReadOnly
 );
 
 static TAutoConsoleVariable<int32> CVarBasePassWriteDepthEvenWithFullPrepass(
@@ -4646,24 +4637,6 @@ void FScene::Release()
 		});
 }
 
-bool ShouldForceFullDepthPass(EShaderPlatform ShaderPlatform)
-{
-	const bool bNaniteEnabled = UseNanite(ShaderPlatform);
-	const bool bDBufferAllowed = IsUsingDBuffers(ShaderPlatform);
-	const bool bVirtualTextureEnabled = UseVirtualTexturing(ShaderPlatform);
-
-	static const auto StencilLODDitherCVar = IConsoleManager::Get().FindTConsoleVariableDataInt(TEXT("r.StencilForLODDither"));
-	const bool bStencilLODDither = StencilLODDitherCVar->GetValueOnAnyThread() != 0;
-
-	static const auto AOComputeCVar = IConsoleManager::Get().FindTConsoleVariableDataInt(TEXT("r.AmbientOcclusion.Compute"));
-	const bool bAOCompute = AOComputeCVar->GetValueOnAnyThread() > 0;
-
-	const bool bEarlyZMaterialMasking = MaskedInEarlyPass(ShaderPlatform);
-
-	// Note: ShouldForceFullDepthPass affects which static draw lists meshes go into, so nothing it depends on can change at runtime, unless you do a FGlobalComponentRecreateRenderStateContext to propagate the cvar change
-	return bNaniteEnabled || bAOCompute || bDBufferAllowed || bVirtualTextureEnabled || bStencilLODDither || bEarlyZMaterialMasking || IsForwardShadingEnabled(ShaderPlatform) || IsUsingSelectiveBasePassOutputs(ShaderPlatform);
-}
-
 void FScene::UpdateEarlyZPassMode()
 {
 	checkSlow(IsInGameThread());
@@ -4722,8 +4695,8 @@ void FScene::GetEarlyZPassMode(ERHIFeatureLevel::Type InFeatureLevel, EDepthDraw
 	{
 		OutZPassMode = DDM_None;
 				 
-		const bool bMaskedInEarlyPass = MaskedInEarlyPass(ShaderPlatform);
-		if (bMaskedInEarlyPass)
+		const bool bMaskedOnlyPrePass = CVarMobileEarlyZPass.GetValueOnAnyThread() == 2;
+		if (bMaskedOnlyPrePass)
 		{
 			OutZPassMode = DDM_MaskedOnly;
 		}
diff --git a/Engine/Source/Runtime/Renderer/Private/SceneVisibility.cpp b/Engine/Source/Runtime/Renderer/Private/SceneVisibility.cpp
--- a/Engine/Source/Runtime/Renderer/Private/SceneVisibility.cpp
+++ b/Engine/Source/Runtime/Renderer/Private/SceneVisibility.cpp
@@ -2220,17 +2220,8 @@ struct FMarkRelevantStaticMeshesForViewData
 		MinScreenRadiusForCSMDepthSquared = GMinScreenRadiusForCSMDepth * GMinScreenRadiusForCSMDepth;
 		MinScreenRadiusForDepthPrepassSquared = GMinScreenRadiusForDepthPrepass * GMinScreenRadiusForDepthPrepass;
 
-		extern bool ShouldForceFullDepthPass(EShaderPlatform ShaderPlatform);
 		EShaderPlatform ShaderPlatform = View.GetShaderPlatform();
-		if (IsMobilePlatform(ShaderPlatform))
-		{
-			FScene* Scene = View.Family->Scene->GetRenderScene();
-			bFullEarlyZPass = (Scene && Scene->EarlyZPassMode == DDM_AllOpaque);
-		}
-		else
-		{
-			bFullEarlyZPass = ShouldForceFullDepthPass(ShaderPlatform);
-		}
+		bFullEarlyZPass = ShouldForceFullDepthPass(ShaderPlatform);
 	}
 };
 
