Fixed TPrimitiveArray to actually be a chunked array.

diff --git a/Engine/Source/Runtime/Renderer/Private/ScenePrivateBase.h b/Engine/Source/Runtime/Renderer/Private/ScenePrivateBase.h
--- a/Engine/Source/Runtime/Renderer/Private/ScenePrivateBase.h
+++ b/Engine/Source/Runtime/Renderer/Private/ScenePrivateBase.h
@@ -66,9 +66,15 @@ public:
 	{
 		if (NumElements % NumElementsPerChunk == 0)
 		{
-			Chunks.Emplace();
+			ChunkType* Chunk = new ChunkType;
+			Chunk->Reserve(NumElementsPerChunk);
+			Chunks.Emplace(Chunk);
 		}
-		return Get(NumElements++);
+
+		NumElements++;
+		ChunkType& Chunk = *Chunks.Last();
+		Chunk.AddUninitialized();
+		return Chunk.Last();
 	}
 
 	void Remove(uint32 Count, bool bAllowShrinking)
@@ -77,10 +83,16 @@ public:
 		const uint32 NumElementsNew = NumElements - Count;
 		while (NumElements != NumElementsNew)
 		{
-			Get(--NumElements).~T();
+			--NumElements;
+
+			ChunkType& Chunk = *Chunks.Last();
+			Chunk.Pop(false);
+
+			if (Chunk.IsEmpty())
+			{
+				Chunks.Pop(false);
+			}
 		}
-		const uint32 NumChunksToRemove = Chunks.Num() - NumChunks(NumElementsNew);
-		Chunks.RemoveAt(Chunks.Num() - NumChunksToRemove, NumChunksToRemove, bAllowShrinking);
 	}
 
 	void Reserve(int32 Count)
@@ -92,14 +104,14 @@ public:
 	{
 		const uint32 ChunkIndex        = ElementIndex / NumElementsPerChunk;
 		const uint32 ChunkElementIndex = ElementIndex % NumElementsPerChunk;
-		return Chunks[ChunkIndex][ChunkElementIndex];
+		return (*Chunks[ChunkIndex])[ChunkElementIndex];
 	}
 
 	const T& Get(int32 ElementIndex) const
 	{
 		const uint32 ChunkIndex        = ElementIndex / NumElementsPerChunk;
 		const uint32 ChunkElementIndex = ElementIndex % NumElementsPerChunk;
-		return Chunks[ChunkIndex][ChunkElementIndex];
+		return (*Chunks[ChunkIndex])[ChunkElementIndex];
 	}
 
 	FORCEINLINE T& operator[] (int32 Index) { return Get(Index); }
@@ -115,6 +127,7 @@ private:
 		return (NumElements + NumElementsPerChunk - 1u) / NumElementsPerChunk;
 	}
 
-	TArray<TStaticArray<T, NumElementsPerChunk>> Chunks;
+	using ChunkType = TArray<T>;
+	TArray<TUniquePtr<ChunkType>> Chunks;
 	uint32 NumElements = 0;
 };
\ No newline at end of file
