In PopulateDepthPass, select the Extent/Viewport parameters depending on the type of composite draw we are doing.

If outputting colour and scaling the depth back down, both Color and Depth textures need to be targeting the same output viewport. This assumes further composite passes are not overcompensating for a larger depth texture size.

For other compositing methods, which scale down the depth in subsequent passes, we need their existing Extent/Viewport scaling so their sizes remain for further compositing passes as to not overcompensate.

diff --git a/Engine/Source/Runtime/Renderer/Private/PostProcess/PostProcessCompositePrimitivesCommon.cpp b/Engine/Source/Runtime/Renderer/Private/PostProcess/PostProcessCompositePrimitivesCommon.cpp
--- a/Engine/Source/Runtime/Renderer/Private/PostProcess/PostProcessCompositePrimitivesCommon.cpp
+++ b/Engine/Source/Runtime/Renderer/Private/PostProcess/PostProcessCompositePrimitivesCommon.cpp
@@ -248,12 +248,17 @@ void PopulateDepthPass(FRDGBuilder& GraphBuilder,
 
 	FPopulateCompositeDepthPS::FParameters* PassParameters = GraphBuilder.AllocParameters<FPopulateCompositeDepthPS::FParameters>();
 	PassParameters->View = View.ViewUniformBuffer;
-	PassParameters->Color = GetScreenPassTextureViewportParameters(FScreenPassTextureViewport(InSceneColor.Texture->Desc.Extent, View.ViewRect));
-	PassParameters->Depth = GetScreenPassTextureViewportParameters(FScreenPassTextureViewport(InSceneDepth.Texture->Desc.Extent, View.ViewRect));
 	if (bForceDrawColor)
 	{
 		PassParameters->ColorTexture = InSceneColor.Texture;
 		PassParameters->ColorSampler = PointClampSampler;
+		PassParameters->Color = GetScreenPassTextureViewportParameters(FScreenPassTextureViewport(InSceneColor.Texture->Desc.Extent, View.ViewRect));
+		PassParameters->Depth = GetScreenPassTextureViewportParameters(FScreenPassTextureViewport(InSceneDepth.Texture->Desc.Extent, View.ViewRect));
+	}
+	else
+	{
+		PassParameters->Color = GetScreenPassTextureViewportParameters(FScreenPassTextureViewport(InSceneColor));
+		PassParameters->Depth = GetScreenPassTextureViewportParameters(FScreenPassTextureViewport(InSceneDepth));
 	}
 	PassParameters->DepthTexture = InSceneDepth.Texture;
 	PassParameters->DepthSampler = PointClampSampler;
