Fixed mismatched copy constructor/assignment operator in FScreenTransform.

diff --git a/Engine/Source/Runtime/Renderer/Public/ScreenPass.h b/Engine/Source/Runtime/Renderer/Public/ScreenPass.h
--- a/Engine/Source/Runtime/Renderer/Public/ScreenPass.h
+++ b/Engine/Source/Runtime/Renderer/Public/ScreenPass.h
@@ -255,11 +255,6 @@ struct FScreenTransform
 		, Bias(InBias)
 	{ }
 
-	inline FScreenTransform(const FScreenTransform& AToB)
-		: Scale(AToB.Scale)
-		, Bias(AToB.Bias)
-	{ }
-
 
 	// A * FScreenTransform::Identity = A
 	static RENDERER_API const FScreenTransform Identity;
