Fixed mismatched copy constructor/assignment operator in FScreenPassTextureViewport.

diff --git a/Engine/Source/Runtime/Renderer/Public/ScreenPass.h b/Engine/Source/Runtime/Renderer/Public/ScreenPass.h
--- a/Engine/Source/Runtime/Renderer/Public/ScreenPass.h
+++ b/Engine/Source/Runtime/Renderer/Public/ScreenPass.h
@@ -126,8 +126,6 @@ public:
 		: FScreenPassTextureViewport(FScreenPassTexture(InTexture, InRect))
 	{}
 
-	FScreenPassTextureViewport(const FScreenPassTextureViewport&) = default;
-
 	bool operator==(const FScreenPassTextureViewport& Other) const;
 	bool operator!=(const FScreenPassTextureViewport& Other) const;
 
