Disabled hair 'CastShadowsWhenNonVisible' as it has a race condition.

diff --git a/Engine/Source/Runtime/Renderer/Private/HairStrands/HairStrandsInterface.cpp b/Engine/Source/Runtime/Renderer/Private/HairStrands/HairStrandsInterface.cpp
--- a/Engine/Source/Runtime/Renderer/Private/HairStrands/HairStrandsInterface.cpp
+++ b/Engine/Source/Runtime/Renderer/Private/HairStrands/HairStrandsInterface.cpp
@@ -54,7 +54,7 @@ static TAutoConsoleVariable<int32> CVarHairStrandsSimulation(
 	ECVF_RenderThreadSafe | ECVF_Scalability);
 
 static TAutoConsoleVariable<int32> CVarHairStrandsNonVisibleShadowCasting(
-	TEXT("r.HairStrands.Shadow.CastShadowWhenNonVisible"), 1,
+	TEXT("r.HairStrands.Shadow.CastShadowWhenNonVisible"), 0,
 	TEXT("Enable shadow casting for hair strands even when culled out from the primary view"),
 	ECVF_RenderThreadSafe);
 
