Make GPUScene update wait on nanite cached draw command task.

diff --git a/Engine/Source/Runtime/Renderer/Private/GPUScene.cpp b/Engine/Source/Runtime/Renderer/Private/GPUScene.cpp
--- a/Engine/Source/Runtime/Renderer/Private/GPUScene.cpp
+++ b/Engine/Source/Runtime/Renderer/Private/GPUScene.cpp
@@ -790,6 +790,9 @@ void FGPUScene::UpdateBufferState(FRDGBuilder& GraphBuilder, FSceneUniformBuffer
 	const bool bNaniteEnabled = DoesPlatformSupportNanite(GMaxRHIShaderPlatform);
 	if (UploadDataSourceAdapter.bUpdateNaniteMaterialTables && bNaniteEnabled)
 	{
+		// Nanite draw commands build raster material tables.
+		Scene.WaitForCacheNaniteDrawCommandsTask();
+
 		for (int32 NaniteMeshPassIndex = 0; NaniteMeshPassIndex < ENaniteMeshPass::Num; ++NaniteMeshPassIndex)
 		{
 			Scene.NaniteMaterials[NaniteMeshPassIndex].UpdateBufferState(GraphBuilder, Scene.Primitives.Num());
