Fix Planar reflection memory due to culling (out of frustum or view) without releasing the memory.

diff --git a/Engine/Source/Runtime/Renderer/Private/PlanarReflectionRendering.cpp b/Engine/Source/Runtime/Renderer/Private/PlanarReflectionRendering.cpp
--- a/Engine/Source/Runtime/Renderer/Private/PlanarReflectionRendering.cpp
+++ b/Engine/Source/Runtime/Renderer/Private/PlanarReflectionRendering.cpp
@@ -297,6 +297,7 @@ static void UpdatePlanarReflectionContents_RenderThread(
 
 		if (!bIsInAnyFrustum)
 		{
+			delete SceneRenderer;
 			return;
 		}
 
@@ -336,6 +337,7 @@ static void UpdatePlanarReflectionContents_RenderThread(
 
 		if (!bIsVisibleInAnyView)
 		{
+			delete SceneRenderer;
 			return;
 		}
 	}
