Remove truncation warnings in FMeshPassProcessor and AddSimpleMeshPass

diff --git a/Engine/Source/Runtime/Renderer/Public/MeshPassProcessor.inl b/Engine/Source/Runtime/Renderer/Public/MeshPassProcessor.inl
--- a/Engine/Source/Runtime/Renderer/Public/MeshPassProcessor.inl
+++ b/Engine/Source/Runtime/Renderer/Public/MeshPassProcessor.inl
@@ -111,7 +111,7 @@ void FMeshPassProcessor::BuildMeshDrawCommands(
 	PSOCollectorStats::CheckMinimalPipelineStateInCache(PipelineState, (uint32)MeshPassType, VertexFactory->GetType());
 #endif // PSO_PRECACHING_VALIDATE
 
-	SharedMeshDrawCommand.PrimitiveIdStreamIndex = VertexFactory->GetPrimitiveIdStreamIndex(FeatureLevel, InputStreamType);
+	SharedMeshDrawCommand.PrimitiveIdStreamIndex = static_cast<int8>(VertexFactory->GetPrimitiveIdStreamIndex(FeatureLevel, InputStreamType));
 
 	if (SharedMeshDrawCommand.PrimitiveIdStreamIndex != INDEX_NONE)
 	{
diff --git a/Engine/Source/Runtime/Renderer/Public/SimpleMeshDrawCommandPass.h b/Engine/Source/Runtime/Renderer/Public/SimpleMeshDrawCommandPass.h
--- a/Engine/Source/Runtime/Renderer/Public/SimpleMeshDrawCommandPass.h
+++ b/Engine/Source/Runtime/Renderer/Public/SimpleMeshDrawCommandPass.h
@@ -107,7 +107,8 @@ void AddSimpleMeshPass(FRDGBuilder& GraphBuilder, PassParametersType* PassParame
 	AddSimpleMeshPass(GraphBuilder, PassParameters, Scene, View, InstanceCullingManager, MoveTemp(PassName), ERDGPassFlags::Raster, AddMeshBatchesCallback,
 		[ViewPortRect](FRHICommandList& RHICmdList)
 		{
-			RHICmdList.SetViewport(ViewPortRect.Min.X, ViewPortRect.Min.Y, 0.0f, ViewPortRect.Max.X, ViewPortRect.Max.Y, 1.0f);
+			RHICmdList.SetViewport(static_cast<float>(ViewPortRect.Min.X), static_cast<float>(ViewPortRect.Min.Y), 0.0f,
+			                       static_cast<float>(ViewPortRect.Max.X), static_cast<float>(ViewPortRect.Max.Y), 1.0f);
 		}
 	);
 }
