Make LPI / CachedMeshDrawCommand tasks dependent on each other in mobile configurations.

diff --git a/Engine/Source/Runtime/Renderer/Private/RendererScene.cpp b/Engine/Source/Runtime/Renderer/Private/RendererScene.cpp
--- a/Engine/Source/Runtime/Renderer/Private/RendererScene.cpp
+++ b/Engine/Source/Runtime/Renderer/Private/RendererScene.cpp
@@ -6392,7 +6392,8 @@ void FScene::UpdateAllPrimitiveSceneInfos(FRDGBuilder& GraphBuilder, EUpdateAllP
 		CacheMeshDrawCommandsTask = LaunchSceneRenderTask(TEXT("CacheMeshDrawCommandsTask"), [this, &SceneInfosWithStaticDrawListUpdate, bLaunchAsyncTask]()
 		{
 			FPrimitiveSceneInfo::CacheMeshDrawCommands(this, SceneInfosWithStaticDrawListUpdate);
-		}, bLaunchAsyncTask);
+
+		}, IsMobilePlatform(GetShaderPlatform()) ? CreateLightPrimitiveInteractionsTask : UE::Tasks::FTask(), bLaunchAsyncTask);
 
 		CacheNaniteDrawCommandsTask = LaunchSceneRenderTask(TEXT("CacheNaniteDrawCommands"), [this, &SceneInfosWithStaticDrawListUpdate, bLaunchAsyncTask]()
 		{
