Implements a r.Translucency.Velocity to control whether translucencies can draw velocity

diff --git a/Engine/Source/Runtime/Renderer/Private/DeferredShadingRenderer.cpp b/Engine/Source/Runtime/Renderer/Private/DeferredShadingRenderer.cpp
--- a/Engine/Source/Runtime/Renderer/Private/DeferredShadingRenderer.cpp
+++ b/Engine/Source/Runtime/Renderer/Private/DeferredShadingRenderer.cpp
@@ -315,6 +315,11 @@ static TAutoConsoleVariable<int32> CVarShadowsUseSharedExternalAccessQueue(
 	TEXT("If enabled, shadows will use the shared external access queue, minimizing unnecessary transitions"),
 	ECVF_RenderThreadSafe);
 
+static TAutoConsoleVariable<int32> CVarTranslucencyVelocity(
+	TEXT("r.Translucency.Velocity"), 1,
+	TEXT("Whether translucency can draws depth/velocity (enabled by default)"),
+	ECVF_RenderThreadSafe);
+
 
 #if RHI_RAYTRACING
 
@@ -4126,7 +4131,7 @@ void FDeferredShadingSceneRenderer::Render(FRDGBuilder& GraphBuilder)
 			RenderDistortion(GraphBuilder, SceneTextures.Color.Target, SceneTextures.Depth.Target, SceneTextures.Velocity, TranslucencyResourceMap);
 		}
 
-		if (bShouldRenderVelocities)
+		if (bShouldRenderVelocities && CVarTranslucencyVelocity.GetValueOnRenderThread() != 0)
 		{
 			const bool bRecreateSceneTextures = !HasBeenProduced(SceneTextures.Velocity);
 
