FBuildInstanceIdBufferAndCommandsFromPrimitiveIdsCs: Bring back forcing of DXC, but only for platforms normally compiled with hlslcc

diff --git a/Engine/Shaders/Private/InstanceCulling/BuildInstanceDrawCommands.usf b/Engine/Shaders/Private/InstanceCulling/BuildInstanceDrawCommands.usf
--- a/Engine/Shaders/Private/InstanceCulling/BuildInstanceDrawCommands.usf
+++ b/Engine/Shaders/Private/InstanceCulling/BuildInstanceDrawCommands.usf
@@ -6,7 +6,7 @@
 #endif
 
 #if ENABLE_INSTANCE_COMPACTION
-	// Currently, instance compaction is the only reason to need extended payload data, so tie it to the permutation for now
+	// Currently, instance compaction is the only reason to need extended payload data, so tie it to the permutation for now.
 	#define ENABLE_EXTENDED_INSTANCE_CULLING_PAYLOADS 1
 #endif
 
diff --git a/Engine/Source/Runtime/Renderer/Private/InstanceCulling/InstanceCullingContext.cpp b/Engine/Source/Runtime/Renderer/Private/InstanceCulling/InstanceCullingContext.cpp
--- a/Engine/Source/Runtime/Renderer/Private/InstanceCulling/InstanceCullingContext.cpp
+++ b/Engine/Source/Runtime/Renderer/Private/InstanceCulling/InstanceCullingContext.cpp
@@ -386,6 +386,12 @@ public:
 		FGlobalShader::ModifyCompilationEnvironment(Parameters, OutEnvironment);
 		FInstanceProcessingGPULoadBalancer::SetShaderDefines(OutEnvironment);
 
+		// Force use of DXC for platforms compiled with hlslcc due to hlslcc's inability to handle member functions in structs
+		if (FDataDrivenShaderPlatformInfo::GetIsHlslcc(Parameters.Platform))
+		{
+			OutEnvironment.CompilerFlags.Add(CFLAG_ForceDXC);
+		}
+
 		OutEnvironment.SetDefine(TEXT("INDIRECT_ARGS_NUM_WORDS"), FInstanceCullingContext::IndirectArgsNumWords);
 		OutEnvironment.SetDefine(TEXT("VF_SUPPORTS_PRIMITIVE_SCENE_DATA"), 1);
 		OutEnvironment.SetDefine(TEXT("USE_GLOBAL_GPU_SCENE_DATA"), 1);
