Change shadow CSM caching to use persistent primitive index to simplify and make more efficient.
- Also adds an async-usable scene change-set.

diff --git a/Engine/Source/Runtime/Renderer/Private/RendererScene.cpp b/Engine/Source/Runtime/Renderer/Private/RendererScene.cpp
--- a/Engine/Source/Runtime/Renderer/Private/RendererScene.cpp
+++ b/Engine/Source/Runtime/Renderer/Private/RendererScene.cpp
@@ -5403,6 +5403,43 @@ void UpdateReflectionSceneData(FScene* Scene)
 	}
 }
 
+/**
+ * Container for scene change set, that can live on beyond the update function. Note that the values are not safe to interpret at all points 
+ * (e.g., a persistent ID of a removed item is not valid after the remove phase), but the arrays are valid as long as the RDG lives.
+ */
+struct FSceneUpdateChangeSetStorage
+{
+	TArray<FPersistentPrimitiveIndex, SceneRenderingAllocator> RemovedPrimitiveIds;
+	TArray<FPrimitiveSceneInfo*,SceneRenderingAllocator> RemovedPrimitiveSceneInfos;
+
+	TArray<FPersistentPrimitiveIndex, SceneRenderingAllocator> UpdatedPrimitiveIds;
+	TArray<FPrimitiveSceneInfo*, SceneRenderingAllocator> UpdatedPrimitiveSceneInfos;
+
+	TArray<FPersistentPrimitiveIndex, SceneRenderingAllocator> AddedPrimitiveIds;
+	TArray<FPrimitiveSceneInfo*, SceneRenderingAllocator> AddedPrimitiveSceneInfos;
+
+	FScenePreUpdateChangeSet GetPreUpdateSet() const 
+	{
+		return FScenePreUpdateChangeSet {
+			TConstArrayView<FPersistentPrimitiveIndex>(RemovedPrimitiveIds),
+			TConstArrayView<FPrimitiveSceneInfo*>(RemovedPrimitiveSceneInfos),
+			TConstArrayView<FPersistentPrimitiveIndex>(UpdatedPrimitiveIds),
+			TConstArrayView<FPrimitiveSceneInfo*>(UpdatedPrimitiveSceneInfos)
+		};
+	}
+
+	FScenePostUpdateChangeSet GetPostUpdateSet() const 
+	{
+		return FScenePostUpdateChangeSet {
+			TConstArrayView<FPersistentPrimitiveIndex>(AddedPrimitiveIds),
+			TConstArrayView<FPrimitiveSceneInfo*>(AddedPrimitiveSceneInfos),
+			TConstArrayView<FPersistentPrimitiveIndex>(UpdatedPrimitiveIds),
+			TConstArrayView<FPrimitiveSceneInfo*>(UpdatedPrimitiveSceneInfos)
+		};
+	}
+};
+
+
 void FScene::UpdateAllPrimitiveSceneInfos(FRDGBuilder& GraphBuilder, EUpdateAllPrimitiveSceneInfosAsyncOps AsyncOps)
 {
 	TRACE_CPUPROFILER_EVENT_SCOPE(Scene::UpdateAllPrimitiveSceneInfos);
@@ -5422,14 +5459,65 @@ void FScene::UpdateAllPrimitiveSceneInfos(FRDGBuilder& GraphBuilder, EUpdateAllP
 	UpdateRayTracingGroupBounds_AddPrimitives(AddedPrimitiveSceneInfos);
 #endif
 
+	// Allocated with render graph lifetime, safe to reference from RDG tasks.
+	FSceneUpdateChangeSetStorage &SceneUpdateChangeSetStorage = *GraphBuilder.AllocObject<FSceneUpdateChangeSetStorage>();
+	SceneUpdateChangeSetStorage.RemovedPrimitiveIds.Reserve(RemovedPrimitiveSceneInfos.Num());
+	SceneUpdateChangeSetStorage.RemovedPrimitiveSceneInfos.Reserve(RemovedPrimitiveSceneInfos.Num());
+
 	TArray<FPrimitiveSceneInfo*> RemovedLocalPrimitiveSceneInfos;
 	RemovedLocalPrimitiveSceneInfos.Reserve(RemovedPrimitiveSceneInfos.Num());
-	for (FPrimitiveSceneInfo* SceneInfo : RemovedPrimitiveSceneInfos)
+	for (FPrimitiveSceneInfo* PrimitiveSceneInfo : RemovedPrimitiveSceneInfos)
+	{
+		RemovedLocalPrimitiveSceneInfos.Add(PrimitiveSceneInfo);
+
+		SceneUpdateChangeSetStorage.RemovedPrimitiveIds.Add(PrimitiveSceneInfo->GetPersistentIndex());
+		SceneUpdateChangeSetStorage.RemovedPrimitiveSceneInfos.Add(PrimitiveSceneInfo);
+	}
+
 	{
-		RemovedLocalPrimitiveSceneInfos.Add(SceneInfo);
+		SceneUpdateChangeSetStorage.UpdatedPrimitiveIds.Reserve(UpdatedInstances.Num() + UpdatedTransforms.Num());
+		SceneUpdateChangeSetStorage.UpdatedPrimitiveSceneInfos.Reserve(UpdatedInstances.Num() + UpdatedTransforms.Num());
+		// All updated instances must also before moving or re-allocating (TODO: filter out only those actually updated)
+		for (const auto& Instance : UpdatedInstances)
+		{
+			FPrimitiveSceneInfo* PrimitiveSceneInfo = Instance.Key->GetPrimitiveSceneInfo();
+
+			if (RemovedPrimitiveSceneInfos.Find(PrimitiveSceneInfo) != nullptr)
+			{
+				continue;
+			}
+
+			FPersistentPrimitiveIndex Id = Instance.Key->GetPrimitiveSceneInfo()->GetPersistentIndex();
+			// Note: may not be valid if it is also being added this in this update, in which case it will get processed in the post-update anyway.
+			if (Id.IsValid())
+			{
+				SceneUpdateChangeSetStorage.UpdatedPrimitiveIds.Add(Id);
+				SceneUpdateChangeSetStorage.UpdatedPrimitiveSceneInfos.Add(PrimitiveSceneInfo);
+			}
+		}
+		// As must all primitive updates, 
+		for (const auto& Transform : UpdatedTransforms)
+		{
+			FPrimitiveSceneInfo* PrimitiveSceneInfo = Transform.Key->GetPrimitiveSceneInfo();
+
+			if (RemovedPrimitiveSceneInfos.Find(PrimitiveSceneInfo) != nullptr)
+			{
+				continue;
+			}
+
+			FPersistentPrimitiveIndex Id = Transform.Key->GetPrimitiveSceneInfo()->GetPersistentIndex();
+			// Note: may not be valid if it is also being added this in this update, in which case it will get processed in the post-update anyway.
+			if (Id.IsValid())
+			{
+				SceneUpdateChangeSetStorage.UpdatedPrimitiveIds.Add(Id);
+				SceneUpdateChangeSetStorage.UpdatedPrimitiveSceneInfos.Add(PrimitiveSceneInfo);
+			}
+		}
 	}
+	
 	// NOTE: We clear this early because IsPrimitiveBeingRemoved gets called from the CreateLightPrimitiveInteraction (to make sure that old primitives are not accessed) 
 	// we cannot safely kick off the AsyncCreateLightPrimitiveInteractionsTask before the RemovedPrimitiveSceneInfos has been cleared.
+	// TODO: this is probably not true anymore!
 	RemovedPrimitiveSceneInfos.Empty();
 
 	RemovedLocalPrimitiveSceneInfos.Sort(FPrimitiveArraySortKey());
@@ -5580,18 +5668,6 @@ void FScene::UpdateAllPrimitiveSceneInfos(FRDGBuilder& GraphBuilder, EUpdateAllP
 							TBitArraySwapElements(PrimitivesNeedingStaticMeshUpdate, DestIndex, SourceIndex);
 							TBitArraySwapElements(PrimitivesNeedingUniformBufferUpdate, DestIndex, SourceIndex);
 
-							for (TMap<int32, TArray<FCachedShadowMapData>>::TIterator CachedShadowMapIt(CachedShadowMaps); CachedShadowMapIt; ++CachedShadowMapIt)
-							{
-								TArray<FCachedShadowMapData>& ShadowMapDatas = CachedShadowMapIt.Value();
-
-								for (auto& ShadowMapData : ShadowMapDatas)
-								{
-									if (ShadowMapData.StaticShadowSubjectMap.Num() > 0)
-									{
-										TBitArraySwapElements(ShadowMapData.StaticShadowSubjectMap, DestIndex, SourceIndex);
-									}
-								}
-							}
 							GPUScene.RecordPrimitiveIdSwap(DestIndex, SourceIndex);
 
 						#if RHI_RAYTRACING
@@ -5707,40 +5783,11 @@ void FScene::UpdateAllPrimitiveSceneInfos(FRDGBuilder& GraphBuilder, EUpdateAllP
 
 				DeletedSceneInfos.Add(PrimitiveSceneInfo);
 
-				for (const FLightSceneInfo* LightSceneInfo : DirectionalLights)
-				{
-					TArray<FCachedShadowMapData>* CachedShadowMapDatas = GetCachedShadowMapDatas(LightSceneInfo->Id);
-
-					if (CachedShadowMapDatas)
-					{
-						for (auto& CachedShadowMapData : *CachedShadowMapDatas)
-						{
-							if (CachedShadowMapData.StaticShadowSubjectMap[PrimitiveIndex] == true)
-							{
-								CachedShadowMapData.InvalidateCachedShadow();
-							}
-						}
-					}
-				}
-
 				const int32 PersistentIndex = PrimitiveSceneInfo->PersistentIndex.Index;
 				PersistentPrimitiveIdAllocator.Free(PersistentIndex);
 				PersistentPrimitiveIdToIndexMap[PersistentIndex] = INDEX_NONE;
 			}
 
-			for (TMap<int32, TArray<FCachedShadowMapData>>::TIterator CachedShadowMapIt(CachedShadowMaps); CachedShadowMapIt; ++CachedShadowMapIt)
-			{
-				TArray<FCachedShadowMapData>& ShadowMapDatas = CachedShadowMapIt.Value();
-
-				for (auto& ShadowMapData : ShadowMapDatas)
-				{
-					if (ShadowMapData.StaticShadowSubjectMap.Num() > 0)
-					{
-						ShadowMapData.StaticShadowSubjectMap.RemoveAt(SourceIndex, RemoveCount);
-						check(Primitives.Num() == ShadowMapData.StaticShadowSubjectMap.Num());
-					}
-				}
-			}
 			RemovedLocalPrimitiveSceneInfos.RemoveAt(StartIndex, RemovedLocalPrimitiveSceneInfos.Num() - StartIndex, false);
 		}
 	
@@ -5825,19 +5872,6 @@ void FScene::UpdateAllPrimitiveSceneInfos(FRDGBuilder& GraphBuilder, EUpdateAllP
 		#endif
 			PrimitivesNeedingStaticMeshUpdate.Reserve(PrimitivesNeedingStaticMeshUpdate.Num() + AddedLocalPrimitiveSceneInfos.Num());
 			PrimitivesNeedingUniformBufferUpdate.Reserve(PrimitivesNeedingUniformBufferUpdate.Num() + AddedLocalPrimitiveSceneInfos.Num());
-
-			for (TMap<int32, TArray<FCachedShadowMapData>>::TIterator CachedShadowMapIt(CachedShadowMaps); CachedShadowMapIt; ++CachedShadowMapIt)
-			{
-				TArray<FCachedShadowMapData>& ShadowMapDatas = CachedShadowMapIt.Value();
-
-				for (auto& ShadowMapData : ShadowMapDatas)
-				{
-					if (ShadowMapData.StaticShadowSubjectMap.Num() > 0)
-					{
-						ShadowMapData.StaticShadowSubjectMap.Reserve(ShadowMapData.StaticShadowSubjectMap.Num() + AddedLocalPrimitiveSceneInfos.Num());
-					}
-				}
-			}
 		}
 
 		while (AddedLocalPrimitiveSceneInfos.Num())
@@ -5879,19 +5913,6 @@ void FScene::UpdateAllPrimitiveSceneInfos(FRDGBuilder& GraphBuilder, EUpdateAllP
 					PrimitivesNeedingStaticMeshUpdate.Add(false);
 					PrimitivesNeedingUniformBufferUpdate.Add(true);
 
-					for (TMap<int32, TArray<FCachedShadowMapData>>::TIterator CachedShadowMapIt(CachedShadowMaps); CachedShadowMapIt; ++CachedShadowMapIt)
-					{
-						TArray<FCachedShadowMapData>& ShadowMapDatas = CachedShadowMapIt.Value();
-
-						for (auto& ShadowMapData : ShadowMapDatas)
-						{
-							if (ShadowMapData.StaticShadowSubjectMap.Num() > 0)
-							{
-								ShadowMapData.StaticShadowSubjectMap.Add(false);
-							}
-						}
-					}
-
 					const int32 SourceIndex = PrimitiveSceneProxies.Num() - 1;
 					PrimitiveSceneInfo->PackedIndex = SourceIndex;
 					checkSlow(PrimitiveSceneInfo->PersistentIndex.Index == INDEX_NONE);
@@ -5998,19 +6019,6 @@ void FScene::UpdateAllPrimitiveSceneInfos(FRDGBuilder& GraphBuilder, EUpdateAllP
 							TBitArraySwapElements(PrimitivesNeedingStaticMeshUpdate, DestIndex, SourceIndex);
 							TBitArraySwapElements(PrimitivesNeedingUniformBufferUpdate, DestIndex, SourceIndex);
 
-							for (TMap<int32, TArray<FCachedShadowMapData>>::TIterator CachedShadowMapIt(CachedShadowMaps); CachedShadowMapIt; ++CachedShadowMapIt)
-							{
-								TArray<FCachedShadowMapData>& ShadowMapDatas = CachedShadowMapIt.Value();
-
-								for (auto& ShadowMapData : ShadowMapDatas)
-								{
-									if (ShadowMapData.StaticShadowSubjectMap.Num() > 0)
-									{
-										TBitArraySwapElements(ShadowMapData.StaticShadowSubjectMap, DestIndex, SourceIndex);
-									}
-								}
-							}
-
 							GPUScene.RecordPrimitiveIdSwap(DestIndex, SourceIndex);
 
 						#if RHI_RAYTRACING
@@ -6301,6 +6309,16 @@ void FScene::UpdateAllPrimitiveSceneInfos(FRDGBuilder& GraphBuilder, EUpdateAllP
 			}
 		}
 	}
+	
+	SceneUpdateChangeSetStorage.AddedPrimitiveIds.Reserve(AddedPrimitiveSceneInfos.Num());
+	SceneUpdateChangeSetStorage.AddedPrimitiveSceneInfos.Reserve(AddedPrimitiveSceneInfos.Num());
+	for (FPrimitiveSceneInfo* PrimitiveSceneInfo : AddedPrimitiveSceneInfos)
+	{
+		SceneUpdateChangeSetStorage.AddedPrimitiveIds.Add(PrimitiveSceneInfo->GetPersistentIndex());
+		SceneUpdateChangeSetStorage.AddedPrimitiveSceneInfos.Add(PrimitiveSceneInfo);
+	}
+
+	UpdateCachedShadowState(SceneUpdateChangeSetStorage.GetPreUpdateSet(), SceneUpdateChangeSetStorage.GetPostUpdateSet());
 
 	if (SceneInfosWithStaticDrawListUpdate.Num() > 0)
 	{
@@ -6941,3 +6959,28 @@ void FScene::DebugRender(TArrayView<FViewInfo> Views)
 	ShadowScene->DebugRender(Views);
 }
 #endif
+
+void FScene::UpdateCachedShadowState(const FScenePreUpdateChangeSet &ScenePreUpdateChangeSet, const FScenePostUpdateChangeSet &ScenePostUpdateChangeSet)
+{
+	for (const FLightSceneInfo* LightSceneInfo : DirectionalLights)
+	{
+		TArray<FCachedShadowMapData>* CachedShadowMapDatas = GetCachedShadowMapDatas(LightSceneInfo->Id);
+
+		if (CachedShadowMapDatas)
+		{
+			for (auto& CachedShadowMapData : *CachedShadowMapDatas)
+			{
+				for (FPersistentPrimitiveIndex PersistentPrimitiveIndex : ScenePreUpdateChangeSet.RemovedPrimitiveIds)
+				{
+					if (CachedShadowMapData.StaticShadowSubjectPersistentPrimitiveIdMap[PersistentPrimitiveIndex.Index] == true)
+					{
+						CachedShadowMapData.InvalidateCachedShadow();
+						break;
+					}
+				}
+				CachedShadowMapData.StaticShadowSubjectPersistentPrimitiveIdMap.SetNum(GetMaxPersistentPrimitiveIndex(), false);
+			}
+		}
+	}
+}
+
diff --git a/Engine/Source/Runtime/Renderer/Private/ScenePrivate.h b/Engine/Source/Runtime/Renderer/Private/ScenePrivate.h
--- a/Engine/Source/Runtime/Renderer/Private/ScenePrivate.h
+++ b/Engine/Source/Runtime/Renderer/Private/ScenePrivate.h
@@ -107,6 +107,34 @@ struct FLightSceneChangeSet
 	TConstArrayView<int32> ColorUpdatedLightIds;
 };
 
+/**
+ * Change set that is valid before removes are processed and the scene data modified.
+ * The referenced arrays have RDG life-time and can be safely used in RDG tasks.
+ * However, the referenced data (primitive/proxy) and meaning of the persistent ID is not generally valid past the call in which this is passed. 
+ * Thus, care need to be excercised.
+ */
+class FScenePreUpdateChangeSet
+{
+public:
+	TConstArrayView<FPersistentPrimitiveIndex> RemovedPrimitiveIds;
+	TConstArrayView<FPrimitiveSceneInfo*> RemovedPrimitiveSceneInfos;
+	TConstArrayView<FPersistentPrimitiveIndex> UpdatedPrimitiveIds;
+	TConstArrayView<FPrimitiveSceneInfo*> UpdatedPrimitiveSceneInfos;
+};
+
+/**
+ * Change set that is valid before after adds are processed and the scene data is modified.
+ * The referenced arrays have RDG life-time and can be safely used in RDG tasks.
+ */
+class FScenePostUpdateChangeSet
+{
+public:
+	TConstArrayView<FPersistentPrimitiveIndex> AddedPrimitiveIds;
+	TConstArrayView<FPrimitiveSceneInfo*> AddedPrimitiveSceneInfos;
+	TConstArrayView<FPersistentPrimitiveIndex> UpdatedPrimitiveIds;
+	TConstArrayView<FPrimitiveSceneInfo*> UpdatedPrimitiveSceneInfos;
+};
+
 /** Holds information about a single primitive's occlusion. */
 class FPrimitiveOcclusionHistory
 {
@@ -2544,9 +2572,9 @@ public:
 	bool bCachedShadowMapHasNaniteGeometry;
 
 	/**
-	* The static meshes cast shadow on this cached csm
-	*/
-	TBitArray<> StaticShadowSubjectMap;
+	 * The static meshes cast shadow on this cached csm
+	 */
+	TBitArray<> StaticShadowSubjectPersistentPrimitiveIdMap;
 
 	FIntPoint ShadowBufferResolution;
 	FVector PreShadowTranslation;
@@ -2562,7 +2590,7 @@ public:
 	{
 		ShadowMap.Release();
 
-		StaticShadowSubjectMap.SetRange(0, StaticShadowSubjectMap.Num(), false);
+		StaticShadowSubjectPersistentPrimitiveIdMap.SetRange(0, StaticShadowSubjectPersistentPrimitiveIdMap.Num(), false);
 	}
 
 	FCachedShadowMapData(const FWholeSceneProjectedShadowInitializer& InInitializer, float InLastUsedTime) :
@@ -3629,6 +3657,12 @@ private:
 	FLightSceneChangeSet UpdateAllLightSceneInfos(FRDGBuilder& GraphBuilder);
 
 private:
+
+	/**
+	 * Update tracked scene state for cached CSM shadows
+	 */
+	void UpdateCachedShadowState(const FScenePreUpdateChangeSet &ScenePreUpdateChangeSet, const FScenePostUpdateChangeSet &ScenePostUpdateChangeSet);
+
 #if RHI_RAYTRACING
 	void UpdateRayTracingGroupBounds_AddPrimitives(const Experimental::TRobinHoodHashSet<FPrimitiveSceneInfo*>& PrimitiveSceneInfos);
 	void UpdateRayTracingGroupBounds_RemovePrimitives(const Experimental::TRobinHoodHashSet<FPrimitiveSceneInfo*>& PrimitiveSceneInfos);
diff --git a/Engine/Source/Runtime/Renderer/Private/ShadowSetup.cpp b/Engine/Source/Runtime/Renderer/Private/ShadowSetup.cpp
--- a/Engine/Source/Runtime/Renderer/Private/ShadowSetup.cpp
+++ b/Engine/Source/Runtime/Renderer/Private/ShadowSetup.cpp
@@ -2131,7 +2131,7 @@ void FProjectedShadowInfo::FinalizeAddSubjectPrimitive(
 		if (CacheMode == SDCM_StaticPrimitivesOnly)
 		{
 			// Record all static meshes in the cached shadow map.
-			CachedShadowMapData.StaticShadowSubjectMap[PrimitiveSceneInfo->GetIndex()] = true;
+			CachedShadowMapData.StaticShadowSubjectPersistentPrimitiveIdMap[PrimitiveSceneInfo->GetPersistentIndex().Index] = true;
 		}
 		else if (!PrimitiveSceneInfo->Proxy->IsMeshShapeOftenMoving())
 		{
@@ -3608,7 +3608,7 @@ void ComputeViewDependentWholeSceneShadowCacheModes(
 
 							CachedShadowMapData->Initializer = ProjectedShadowInitializer;
 							CachedShadowMapData->InvalidateCachedShadow();
-							checkSlow(CachedShadowMapData->StaticShadowSubjectMap.Num() == Scene->Primitives.Num());
+							checkSlow(CachedShadowMapData->StaticShadowSubjectPersistentPrimitiveIdMap.Num() == Scene->GetMaxPersistentPrimitiveIndex());
 						}
 					}
 				}
@@ -3619,7 +3619,7 @@ void ComputeViewDependentWholeSceneShadowCacheModes(
 					{
 						CachedShadowMapData->ShadowBufferResolution = ShadowMapSize;
 						CachedShadowMapData->Initializer = ProjectedShadowInitializer;
-						checkSlow(CachedShadowMapData->StaticShadowSubjectMap.Num() == Scene->Primitives.Num());
+						checkSlow(CachedShadowMapData->StaticShadowSubjectPersistentPrimitiveIdMap.Num() == Scene->GetMaxPersistentPrimitiveIndex());
 					},
 					[&]()
 					{
@@ -3650,7 +3650,7 @@ void ComputeViewDependentWholeSceneShadowCacheModes(
 				CachedShadowMapDatas->Add(FCachedShadowMapData(ProjectedShadowInitializer, RealTime));
 				CachedShadowMapData = &CachedShadowMapDatas->Last();
 				CachedShadowMapData->ShadowBufferResolution = ShadowMapSize;
-				CachedShadowMapData->StaticShadowSubjectMap.Init(false, Scene->Primitives.Num());
+				CachedShadowMapData->StaticShadowSubjectPersistentPrimitiveIdMap.Init(false, Scene->GetMaxPersistentPrimitiveIndex());
 			},
 			[&]()
 			{
@@ -4575,7 +4575,7 @@ struct FGatherShadowPrimitivesPacket
 					const FCachedShadowMapData& CachedShadowMapData = TaskData.Scene->GetCachedShadowMapDataRef(LightSceneInfo.Id, ProjectedShadowInfo->CascadeSettings.ShadowSplitIndex);
 
 					// if the mesh is dynamic or a new added mesh, it should cast csm.
-					bShadowRelevance = PrimitiveProxy->IsMeshShapeOftenMoving() || !CachedShadowMapData.StaticShadowSubjectMap[PrimitiveSceneInfo->GetIndex()];
+					bShadowRelevance = PrimitiveProxy->IsMeshShapeOftenMoving() || !CachedShadowMapData.StaticShadowSubjectPersistentPrimitiveIdMap[PrimitiveSceneInfo->GetPersistentIndex().Index];
 				}
 				else if (ProjectedShadowInfo->CacheMode == SDCM_StaticPrimitivesOnly)
 				{
@@ -4586,7 +4586,7 @@ struct FGatherShadowPrimitivesPacket
 					const FCachedShadowMapData& CachedShadowMapData = TaskData.Scene->GetCachedShadowMapDataRef(LightSceneInfo.Id, ProjectedShadowInfo->CascadeSettings.ShadowSplitIndex);
 
 					// if the mesh is dynamic or a new added mesh, it should cast csm.
-					bShadowRelevance = PrimitiveProxy->IsMeshShapeOftenMoving() || !CachedShadowMapData.StaticShadowSubjectMap[PrimitiveSceneInfo->GetIndex()];
+					bShadowRelevance = PrimitiveProxy->IsMeshShapeOftenMoving() || !CachedShadowMapData.StaticShadowSubjectPersistentPrimitiveIdMap[PrimitiveSceneInfo->GetPersistentIndex().Index];
 
 					// if the static mesh is rendered in the cached shadow map, check if it intersects with the non-overlapped area
 					if (!bShadowRelevance)
