Increase ShaderComplexity cost of Masked Materials on Mobile ShaderPlatforms that do not use MaskedInEarlyZ

diff --git a/Engine/Source/Runtime/Renderer/Private/DebugViewModeRendering.cpp b/Engine/Source/Runtime/Renderer/Private/DebugViewModeRendering.cpp
--- a/Engine/Source/Runtime/Renderer/Private/DebugViewModeRendering.cpp
+++ b/Engine/Source/Runtime/Renderer/Private/DebugViewModeRendering.cpp
@@ -83,6 +83,14 @@ static FAutoConsoleVariableRef CVarShaderComplexityBaselineDeferredUnlitPS(
 	ECVF_Default
 );
 
+float GShaderComplexityMobileMaskedCostMultiplier = 1.5f;
+static FAutoConsoleVariableRef CVarShaderComplexityMobileMaskedCostMultiplier(
+	TEXT("r.ShaderComplexity.MobileMaskedCostMultiplier"),
+	GShaderComplexityMobileMaskedCostMultiplier,
+	TEXT("Extra cost of masked materials if we do not use masked in early Z optimization"),
+	ECVF_Default
+);
+
 IMPLEMENT_MATERIAL_SHADER_TYPE(, FDebugViewModePS, TEXT("/Engine/Private/DebugViewModePixelShader.usf"), TEXT("Main"), SF_Pixel);
 
 int32 GetQuadOverdrawUAVIndex(EShaderPlatform Platform)
@@ -469,6 +477,12 @@ void FDebugViewModeMeshProcessor::UpdateInstructionCount(FDebugViewModeShaderEle
 				OutShaderElementData.NumVSInstructions = MobileVS.IsValid() ? MobileVS->GetNumInstructions() : 0;
 				OutShaderElementData.NumPSInstructions = MobilePS.IsValid() ? MobilePS->GetNumInstructions() : 0;
 			}
+
+			const EShaderPlatform ShaderPlatform = GetFeatureLevelShaderPlatform(InBatchMaterial->GetFeatureLevel());
+			if (IsMaskedBlendMode(*InBatchMaterial) && !MaskedInEarlyPass(ShaderPlatform))
+			{
+				OutShaderElementData.NumPSInstructions *= GShaderComplexityMobileMaskedCostMultiplier;
+			}
 		}
 	}
 }
